Copying //UE4/Orion-Staging to //UE4/Main (Source: //Orion/Dev-General @ 3606378)

Change 3604978 on 2017/08/23 by [REDACTED_PERSON]

	Fix for OR-42722 from 4.17 branch

	#!tests compiled
	#!rb max.chen

Change 3604960 on 2017/08/23 by [REDACTED_PERSON]

	Proper fix for OR-43001, removed hack-around.

	#!tests compiled
	#!rb max.chen

Change 3604881 on 2017/08/23 by robomerge

	#!ROBOMERGE-AUTHOR: andrew.grant
	Changed LostReservation message in tests to info from warning.

	#!tests compiled
	#!rb none

	#!ROBOMERGE-SOURCE: CL 3604871 in //Orion/Release-42.3/... via CL 3604878 via CL 3604880
	#!ROBOMERGE-BOT: ORION (Main -> Dev-General)

Change 3604566 on 2017/08/23 by robomerge

	#!ROBOMERGE-AUTHOR: andrew.grant
	Gauntlet - return -1 for a test if there's a fatal error. Removed network errors failing tests (should be down to test to determine)

	Made SoakTest better able to detect failed drafts and incomplete matches

	#!tests ran multiple soaks and SoloAllHeroes
	#!rb none

	#!ROBOMERGE-SOURCE: CL 3604560 in //Orion/Release-42.3/... via CL 3604563 via CL 3604565
	#!ROBOMERGE-BOT: ORION (Main -> Dev-General)

Change 3604181 on 2017/08/23 by [REDACTED_PERSON]

	Set ULandscapeHeightfieldCollisionComponent to be NetAddressable so that it can be replicated as a movement base without error.

	OR-42615
	#!rb none
	#!tests pie

Change 3603647 on 2017/08/22 by [REDACTED_PERSON]

	AnimProxy: initialize Actor/Component transforms.

	#!rb none
	#!tests lane minion test map

Change 3603343 on 2017/08/22 by robomerge

	#!ROBOMERGE-AUTHOR: don.eubanks
	Added several more stats to the Card Shop stat panel.
	  + Basic Attack Damage
	  + Max Move Speed
	  + Armor Penetration / Percent
	  + Ability Penetration / Percent

	Added utility function to AbilitySystemComponent to calculate an ability's value with additional Required/Ignore tags and Target tag container (functionality moved from AttributeView)

	#!rb matt.schembari
	#!tests Compile DebugGame Editor Win64 / Shipping Client PS4 / PIE

	#!ROBOMERGE-SOURCE: CL 3603172 in //Orion/Main/...
	#!ROBOMERGE-BOT: ORION (Main -> Dev-General)

Change 3603121 on 2017/08/22 by [REDACTED_PERSON]

	Added FAnimationRuntime::LerpPosesPerBone

	#!rb none
	#!tests lane minions split body anims.

Change 3603010 on 2017/08/22 by [REDACTED_PERSON]

	Draw box around Actor being debugged by 'ShowDebug' command, to help identify what we're viewing.

	#!rb none
	#!tests lane minions

Change 3602574 on 2017/08/22 by robomerge

	#!ROBOMERGE-AUTHOR: andrew.grant
	Fixed linux warning

	#!tests compiled
	#!rb none

	#!ROBOMERGE-SOURCE: CL 3602571 in //Orion/Release-42.3/... via CL 3602572
	#!ROBOMERGE-BOT: ORION (Main -> Dev-General)

Change 3602396 on 2017/08/22 by andrew.grant

	#!CodeReview: andrew.grant, jason.bestimt, jeff.williams

	Unresolved conflicts. andrew.grant, please merge this change by hand.
	//ROBOMERGE_ORION_Dev_General/Engine/Source/Programs/AutomationTool/NotForLicensees/Gauntlet/Gauntlet.Automation.csproj

	--------------------------------------
	Updated Gauntlet log parser for new callstack format
	Added offline and runtime Gauntlet tests to verify error generation and parsing

	#!tests ran lots of tests locally
	#!rb codereviewed

	#!ROBOMERGE-SOURCE: CL 3601948 in //Orion/Release-42.3/... via CL 3601950
	#!ROBOMERGE-BOT: ORION (Main -> Dev-General)

Change 3602373 on 2017/08/22 by [REDACTED_PERSON]

	Fixed compile error

	#!tests compiled
	#!rb none

Change 3602321 on 2017/08/22 by robomerge

	#!ROBOMERGE-AUTHOR: andrew.grant
	LogAssertFailedMessage
	- Removed callstack dumping from LogAssertFailedMessage (now in StaticFailDebug)
	- Moved script dumping inside GIsCritial check

	StaticFailDebug
	- Now dumps the callstack if supported for non-ensures. This results in PS4 (and other platforms) now displaying a callstack on a Fatal log.
	- Removed log flush, was redundant
	- Removed LowLevelOutputDebugStringf. If a platform needs this it can be done later in its error device

	OutputMultiLineCallstack
	- Renamed to FDebug::LogFormattedMessageWithCallstack and exposed in headers (calling locations have been updated to remove their extern declarations)
	- No longer writes into the buffer to format it (!)
	- LogName is now the first param. If NAME_None it writes using LowLevelOutputDebugString
	- While writing out error information all callstack likes are prefixed with [Callstack]

	Added brief documentation about the order of ops for ensures/asserts/fatal logs

	#!tests ran lots of tests locally
	#!rb codereviewed

	#!ROBOMERGE-SOURCE: CL 3601943 in //Orion/Release-42.3/... via CL 3601944
	#!ROBOMERGE-BOT: ORION (Main -> Dev-General)

Change 3602316 on 2017/08/22 by robomerge

	#!ROBOMERGE-AUTHOR: andrew.grant
	Changed various platform-specific implementations of ProgramCounterToHumanReadableString to return info in a canonical format of address module!func [file:line].

	E.g.

	0x00416A5F OrionClient.self!FEngineLoop::Tick() [D:\Epic\Orion\Release-Next\Engine\Source\Runtime\Launch\Private\LaunchEngineLoop.cpp:3295]

	#!tests ran lots of tests locally
	#!rb codereviewd

	#!ROBOMERGE-SOURCE: CL 3601940 in //Orion/Release-42.3/... via CL 3601941
	#!ROBOMERGE-BOT: ORION (Main -> Dev-General)

Change 3602311 on 2017/08/22 by robomerge

	#!ROBOMERGE-AUTHOR: andrew.grant
	Sony don't allow exception handling on PS4 so for sometime we've had an issue where crashes have no presence in logs*

	This change adds postmortem analysis to PS4DevkitUtil. If the OS terminates the process then it process the minidump and writes the cause and the callstack to stdout. As a bonus PS4DevkitUtil can now display cause & callstacks from minidumps via "PS4DevkitUtil postmortem -dump=path\to\crash.orbisdmp"

	 (*if you're lucky the kit may have been setup correctly to submit dumps to crashreporter, it may have actually worked, and you may be able to find it).

	#!tests run many times on Orion tests
	#!rb CR

	#!ROBOMERGE-SOURCE: CL 3601929 in //Orion/Release-42.3/... via CL 3601930
	#!ROBOMERGE-BOT: ORION (Main -> Dev-General)

Change 3599823 on 2017/08/21 by [REDACTED_PERSON]

	Allow MaxCascades cvar be set to 0 (should perobject shadow be forced on?)
	#!jira UE-48468,OR-42749
	#!rb Daniel.Wright
	#!tests editor

Change 3598765 on 2017/08/19 by [REDACTED_PERSON]

	Disable regeneration of missing cubemap data in cooked build.

	Need some way of handling this, but currently this breaks PS4 anytime someone makes a map change and doesn't build lighting...

	#!review-3598766 [REDACTED_HANDLE]
	#!tests ran PS4 successfully
	#!rb none

Change 3597800 on 2017/08/18 by [REDACTED_PERSON]

	Added FAnimationRuntime::LerpPoses. Blends two poses together, but first pose is also storing results. To save on copying poses when not necessary.
	Removed individual use of ZERO_ANIMWEIGHT_THRESH, instead use FAnimWeight functions.
	FAnimationRuntime weight functions use FAnimWeight for consistency. (IsFullWeight was different).

	#!rb none
	#!codereview martin.wilson, lina.halper
	#!tests minion test lane map

Change 3597332 on 2017/08/18 by [REDACTED_PERSON]

	SkelMeshComponent LOD update only refreshes transforms when rendered.
	If AnimGraph eval is done, make sure Graph has been updated at least once.

	Fixes:
	- Significance Manager setting MinLOD on non recently rendered minions, causing them to refresh bones.
	-  Minions calling eval with graph not updated due to having bUseRefPoseOnInitAnim set. It means it's possible to refresh bones on a graph that has never been updated, causing a crash.

	#!rb martin.wilson, lina.halper
	#!codereview martin.wilson, lina.halper
	#!tests minions test lane map. placing skelmeshes in editor and making sure LOD update works as expected.

Change 3597042 on 2017/08/18 by [REDACTED_PERSON]

	Added onlinesubsystemmcp to the cooksettings blacklist.
	#!rb Trivial
	#!test none

Change 3596575 on 2017/08/18 by [REDACTED_PERSON]

	Fixing nonunity build issues

	#!rb none
	#!tests compiled OrionGame with unity builds disabled

Change 3595475 on 2017/08/17 by [REDACTED_PERSON]

	Fixed issue with PS4 asserts not being detected during test shutdown

	#!tests ran locally
	#!rb none

Change 3595415 on 2017/08/17 by robomerge

	#!ROBOMERGE-AUTHOR: andrew.grant
	Added check for network failure to tests

	#!tests ran locally
	#!rb none

	#!ROBOMERGE-SOURCE: CL 3595412 in //Orion/Release-42.3/... via CL 3595414
	#!ROBOMERGE-BOT: ORION (Main -> Dev-General)

Change 3594725 on 2017/08/17 by [REDACTED_PERSON]

	Add "Worst offenders" table and runtime to FX perf reports.
	#!rb adric.worley
	#!tests Ran a few FX Perf tests and generated reports.

Change 3594195 on 2017/08/17 by [REDACTED_PERSON]

	Integration from Dev-Niagara to Dev-General. Note that this may cause assets to need to be recooked. Please be prepared for longer than normal cook times once syncing past this changelist.

	#!rb none
	#!tests preflight'ed change, QA ran multiple soak tests

Change 3594177 on 2017/08/17 by [REDACTED_PERSON]

	Fixed issue that was causing exceptions not to be recognized

	#!tests ran locally
	#!rb none

Change 3594090 on 2017/08/17 by robomerge

	#!ROBOMERGE-AUTHOR: laurent.delayen
	Fix for mesh desyncing during online play.
	https://jira.it.epicgames.net/browse/UE-45947

	#!rb none
	#!tests none

	#!ROBOMERGE-SOURCE: CL 3590625 in //Orion/Release-42.3/... via CL 3590626
	#!ROBOMERGE-BOT: ORION (Main -> Dev-General)

Change 3594022 on 2017/08/17 by [REDACTED_PERSON]

	SkeletalMeshComponent::InitAnim
	- Do not call RefreshBoneTransforms without calling TickAnim first. AnimGraph could get in a bad state.
	- Do not call RefreshBoneTransforms if AnimInstance has not been initialized.
	- Do not call RefreshBoneTransforms if bUseRefPoseOnInitAnim is set. Which it could if bForceReinit was false.

	#!rb martin.wilson
	#!tests minion test lane map.

Change 3593972 on 2017/08/17 by [REDACTED_PERSON]

	Duplicating instanced static mesh fixes from Dev-Editor - 3502581, 3570934, 3593597.
	#!rb None
	#!tests Editor, -game, PC
	#!jira UE-48521, OR-42612

Change 3590611 on 2017/08/16 by [REDACTED_PERSON]

	Spot edigrate CL 3584203 to fix show collision crash in mono 2
	#!rb none
	#!tests compile

Change 3590452 on 2017/08/16 by [REDACTED_PERSON]

	Fix for mesh desyncing during online play.
	https://jira.it.epicgames.net/browse/UE-45947

	#!rb none
	#!tests none
	#!codereview zak.middleton

Change 3590378 on 2017/08/16 by [REDACTED_PERSON]

	Integrated CL #!3585145 from Main.

	>>
	Fix for https://jira.it.epicgames.net/browse/OR-42337 and https://jira.it.epicgames.net/browse/OR-42338
	Don't call UpdateMontageSyncGroup() when doing EMeshComponentUpdateFlag::OnlyTickMontagesWhenNotRendered.
	Fixed auto ranged iterator.

	#!rb none
	#!tests bot match with Twinblast
	<<

Change 3590263 on 2017/08/16 by [REDACTED_PERSON]

	Added the ability to remove vertex colors from static meshes from the content browser

	#!rb none
	#!tests none
	#!fyi jordan.walker

Change 3590026 on 2017/08/16 by Jurre.deBaare

	HLOD: When mesh has auto LOD generation disabled it prevents user from dragging meshes to make a cluster
	#!fix changed the tooltip and error handling for cluster creation, if there is any valid mesh now it will show up as a warning
	#!jira OR-41584
	#!rb none
	#!test have tested several 'error' meshes which would before prevent the user from creating a cluster,

Change 3588580 on 2017/08/15 by [REDACTED_PERSON]

	Fix for https://jira.it.epicgames.net/browse/OR-42755

	#!rb none
	#!tests bot match

Change 3588360 on 2017/08/15 by [REDACTED_PERSON]

	Phat
	- Fixed up the materials to turn off the Render Before DOF so that we can see what were doing in Phat again.

Change 3587983 on 2017/08/15 by robomerge

	#!ROBOMERGE-AUTHOR: daniel.lamb
	Fixed up audio issue on the new dawn intro movie
	#!jira OR-42144
	#!rb Max.Preussner
	#!test Paragon ps4
	#!lockdown Andrew.Grant

	#!ROBOMERGE-SOURCE: CL 3581466 in //Orion/Release-42.1/... via CL 3581468 via CL 3581469
	#!ROBOMERGE-BOT: ORION (Main -> Dev-General)

Change 3587746 on 2017/08/15 by [REDACTED_PERSON]

	Build launcher now has option for custom test.
	Dumps generated buildcooktest commandline to log when using -interactive for reference.
	#!rb trivial
	#!test BuildCookTest with and without test options.

Change 3587733 on 2017/08/15 by [REDACTED_PERSON]

	Reenabled binned 2 stats in not shipping.
	#!rb Andrew.Grant
	#!test Paragon editor.

Change 3587727 on 2017/08/15 by [REDACTED_PERSON]

	Removed all the checks to help track down Texture GC issue.
	#!rb Trivial
	#!test Paragon editor

Change 3584487 on 2017/08/14 by [REDACTED_PERSON]

	https://jira.it.epicgames.net/browse/OR-42754 making sure this is not firing because of false positives.

	#!rb none
	#!tests compiles

Change 3582074 on 2017/08/11 by [REDACTED_PERSON]

	TimeStretchCurve system for AnimMontages.

	Allows adjusting playback time of montages in a non constant way. This optional curve allows the user to determine how animation frames are influenced by play back time scaling, instead of using a constant play rate scaling.
	For example, with an attack animation, the strike could have zero influence (no play rate), and the recovery time could have 100% influence (most play rated).

	This system is entirely optional and compatible with current Montage API. You still determine desired play back time via the Play Rate parameter. Based on desired playback time, if there is a curve named 'MontageTimeStretchCurve' then, it will use the curve to determine animation frames to play based on playback time and TimeStretchCurve. Otherwise, if no curve exists, it is done using the constant supplied PlayRate.

	#!rb martin.wilson
	#!codereview lina.halper, james.golding
	#!tests wukong primary attacks

Change 3582063 on 2017/08/11 by [REDACTED_PERSON]

	Removing Monolith02_LowTest from automated builds until errors can be cleared
	#!codereview: daniel.lamb
	#!rb - none
	#!tests - none

Change 3581229 on 2017/08/10 by [REDACTED_PERSON]

	- Back out revision 21 from //Orion/Dev-General/Engine/Source/Runtime/SlateCore/Private/Rendering/DrawElements.cpp
	- Comment out issue that causes crash in the draw box

	#!jira: UE-48222
	#!code review: Jurre.DeBaare, Nick.Darnell, Daniel.Lamb
	#!rb: Jurre.DeBaare
	#!tests: building HLOD as specified in the ticket, and anim blueprint graph

Change 3577839 on 2017/08/08 by [REDACTED_PERSON]

	Revert CL 3576931. To fix crash when generating HLOD in Monolith 2.
	#!rb Andrew.Grant
	#!test Rebuild HLOD paragon editor
	#!fyi Lina.Halper

Change 3577684 on 2017/08/08 by [REDACTED_PERSON]

	Removed ensure and merged proper fix for crash from UE4/Main (3572777)

	#!tests #!rb none

Change 3577562 on 2017/08/08 by robomerge

	#!ROBOMERGE-AUTHOR: laurent.delayen
	Potential fix for https://jira.it.epicgames.net/browse/OR-42383
	Clear NotifyQueue prior to ticking montages. Also dispatch events right away, since ticking ends here, and no rendering is happening. In the event TickPose() is called directly and Component does not get ticked.

	#!rb lina.halper
	#!tests bot match

	#!ROBOMERGE-SOURCE: CL 3577071 in //Orion/Main/...
	#!ROBOMERGE-BOT: ORION (Main -> Dev-General)

Change 3577358 on 2017/08/08 by [REDACTED_PERSON]

	Updating Colorchecker

Change 3576931 on 2017/08/08 by [REDACTED_PERSON]

	Fix for missing arrowhead of the transition on state machine graph

	#!rb: Nick.Darnell
	#!tests: editor

Change 3576847 on 2017/08/08 by [REDACTED_PERSON]

	#!ORION_DG - OR-42361 - Removing assert that was causing crash report client to break (during shipping builds of Paragon due to directories that didn't exist)

	#!RB: Guillaume.Abadie
	#!Tests:none

Change 3576794 on 2017/08/08 by [REDACTED_PERSON]

	Removed call to DebugCanvas->Flush_GameThread(); to fix 'showdebug' commands not rendering anymore.

	#!rb Matt.Kuhlenschmidt
	#!codereview Matt.Kuhlenschmidt
	#!tests Ghost in PIE, showdebug animation works.

Change 3576302 on 2017/08/08 by Jurre.deBaare

	Fix for Materials get switched up while building HLODs in case the section order is different than the material order
	#!rb none
	#!tests Rebuild problematic clusters in Monolith2
	#!fix materials remapping was done using the section index instead of material index in the hlod path (static mesh merge one was fine)

Change 3575221 on 2017/08/07 by robomerge

	#!ROBOMERGE-AUTHOR: andrew.grant
	Fix for PS4 shipping

	#!tests compiled PS4 shipping
	#!rb none

	#!ROBOMERGE-SOURCE: CL 3575112 in //Orion/Release-42.1/... via CL 3575114
	#!ROBOMERGE-BOT: ORION (Main -> Dev-General)

Change 3575165 on 2017/08/07 by andrew.grant

	#!CodeReview: andrew.grant, jason.bestimt, jeff.williams

	Unresolved conflicts. andrew.grant, please merge this change by hand.
	//ROBOMERGE_ORION_Dev_General/Engine/Source/Runtime/PS4/PS4RHI/Private/GnmManager.cpp

	--------------------------------------
	PS4 GPU time now shows correct values instead of vsync time.

	#!tests Ran on PS4
	#!rb Luke.Thatcher

	#!ROBOMERGE-SOURCE: CL 3574821 in //Orion/Release-42.1/... via CL 3574823
	#!ROBOMERGE-BOT: ORION (Main -> Dev-General)

Change 3574408 on 2017/08/06 by [REDACTED_PERSON]

	Switch nightly solo smoke/fx tests to have -unattended on their commandlines.
	#!rb none
	#!tests compiled.

Change 3574308 on 2017/08/06 by robomerge

	#!ROBOMERGE-AUTHOR: andrew.grant
	Include session time, MVP, and hitches on new health report

	#!tests ran locally
	#!rb none

	#!ROBOMERGE-SOURCE: CL 3574305 in //Orion/Release-42/... via CL 3574306 via CL 3574307
	#!ROBOMERGE-BOT: ORION (Main -> Dev-General)

Change 3574217 on 2017/08/05 by [REDACTED_PERSON]

	Fixing UAT compile error

	#!review-3574218 [REDACTED_HANDLE]
	#!rb na
	#!tests na

Change 3574139 on 2017/08/05 by robomerge

	#!ROBOMERGE-AUTHOR: andrew.grant
	[NULL MERGE]
	Locked net version to 3571982 for future v42 patches

	#!tests #!rb none

	#!ROBOMERGE-SOURCE: CL 3574136 in //Orion/Release-42/... via CL 3574137 via CL 3574138
	#!ROBOMERGE-BOT: ORION (Main -> Dev-General)

Change 3574130 on 2017/08/05 by andrew.grant

	#!CodeReview: andrew.grant, jason.bestimt, jeff.williams

	Unresolved conflicts. andrew.grant, please merge this change by hand.
	//ROBOMERGE_ORION_Dev_General/Engine/Source/Programs/AutomationTool/NotForLicensees/Gauntlet/SelfTest/Gauntlet.SelfTest.LogParserTest.cs

	--------------------------------------
	Non-shipping Gauntlet changes -
	 Log parser cleanup that provides better access to log channel info, errors, warnings, and ensures
	 Updated BaselinePerf test to record number of units travelled. This should highlight tests where a bot gets blocked :(

	#!tests ran locally on PS4
	#!rb none

	#!ROBOMERGE-SOURCE: CL 3574116 in //Orion/Release-42/... via CL 3574123 via CL 3574124
	#!ROBOMERGE-BOT: ORION (Main -> Dev-General)

Change 3573197 on 2017/08/04 by andrew.grant

	#!CodeReview: andrew.grant, jason.bestimt, jeff.williams

	Unresolved conflicts. andrew.grant, please merge this change by hand.
	//ROBOMERGE_ORION_Dev_General/Engine/Source/Programs/AutomationTool/NotForLicensees/Gauntlet/Gauntlet.Automation.csproj

	--------------------------------------
	Merging //Orion/Release-42.1 to Dev-ContentUpdate (//Orion/Dev-ContentUpdate)
	#!tests #!rb none

	#!ROBOMERGE-SOURCE: CL 3573179 in //Orion/Main/...
	#!ROBOMERGE-BOT: ORION (Main -> Dev-General)

Change 3573079 on 2017/08/04 by [REDACTED_PERSON]

	Removed some code that was preventing Wacom from working
	#!tests #!rb none

Change 3572790 on 2017/08/04 by Jurre.deBaare

	Moving over fixes from 4.17 stream related to HLOD:
	"UE-47360
	Non Uniform baking of HLOD materials causes texture stretching
	UE-47031
	Generating a HLOD cluster with multiple actors that contain lods will not bake correctly"
	#!rb none
	#!tests none

Change 3572663 on 2017/08/04 by robomerge

	#!ROBOMERGE-AUTHOR: andrew.grant
	Gauntlet - copy elf files to temp dir and launch from there

	#!tests ran gauntlet
	#!rb none

	#!ROBOMERGE-SOURCE: CL 3572657 in //Orion/Release-42/... via CL 3572659 via CL 3572662
	#!ROBOMERGE-BOT: ORION (Main -> Dev-General)

Change 3572622 on 2017/08/04 by Jurre.deBaare

	Moving over:

	"Guard against zero sized boxes being sent to the slate batcher.  This was exposed by the clipping changes since zero sized elements would have been previously clipped.

	#!rb nick.darnell
	#!jira UE-46919"

	#!tests none

Change 3572428 on 2017/08/04 by [REDACTED_PERSON]

	Added per-axis inertia and parent dominance settings to physics assets
	#!rb Thomas.Sarkanen
	#!tests Shrapnel in editor and -game (seems to be the only character using old inertia scaler for rigid body node)

Change 3572142 on 2017/08/04 by robomerge

	#!ROBOMERGE-AUTHOR: andrew.grant
	Disabling r.Cache.UpdatePrimsTaskEnabled  for PS4 due to lock-ups

	#!tests none
	#!rb none

	#!ROBOMERGE-SOURCE: CL 3572139 in //Orion/Release-42/... via CL 3572140 via CL 3572141
	#!ROBOMERGE-BOT: ORION (Main -> Dev-General)

Change 3572133 on 2017/08/04 by [REDACTED_PERSON]

	Merging

	//UE4/Main/Engine/Source/Runtime/SlateCore/Private/Input/HittestGrid.cpp

	to //Orion/Dev-General/Engine/Source/Runtime/SlateCore/Private/Input/HittestGrid.cpp

	#!tests #!rb none

Change 3572065 on 2017/08/03 by [REDACTED_PERSON]

	Fixed bug in Gauntlet parsing of perf data for tests
	Added perf-parsing test to SelfTest

	#!tests ran self test
	#!rb none
	#!ROBOMERGE: 42.1

Change 3572033 on 2017/08/03 by [REDACTED_PERSON]

	Merging //UE4/Main @ 3571062 through Orion-Staging

	#!rb none
	#!tests Engine QA pass

Change 3571262 on 2017/08/03 by [REDACTED_PERSON]

	StaticMeshComponent and LandscapeComponent now register AO material mask and sky occlusion texture in the texture streamer.
	Changing the current lighting scenario now triggers an update of the texture streamer.
	Added an option to the "liststreamingtextures" named UNKOWNREF allowing to inspect texture without references in the texture streamer.
	#!rb none
	#!tests played monolith2 on PS4 and loaded the map in the editor.

Change 3571247 on 2017/08/03 by robomerge

	#!ROBOMERGE-AUTHOR: daniel.lamb
	Integrate Dev-UI -> Release42

	Reduced the sensitivity on the slow tick timer warning
	#!rb Trivial
	#!test Cooked paragon ps4
	#!lockdown Andrew.Grant

	#!ROBOMERGE-SOURCE: CL 3571239 in //Orion/Release-42/... via CL 3571244 via CL 3571246
	#!ROBOMERGE-BOT: ORION (Main -> Dev-General)

Change 3570431 on 2017/08/03 by robomerge

	#!ROBOMERGE-AUTHOR: andrew.grant
	Removed a change that was inadvertendly included in a larger fix and seems to be having problems on windows server.

	#!tests ran windows server locally
	#!rb none

	#!ROBOMERGE-SOURCE: CL 3570428 in //Orion/Release-42/... via CL 3570429 via CL 3570430
	#!ROBOMERGE-BOT: ORION (Main -> Dev-General)

Change 3570343 on 2017/08/03 by Jurre.deBaare

	HLOD: opening a level with HLOD defaults to forcing HLOD on for all clusters
	#!fix forcing the HLODs to show up in the editor changes the MinDrawDistance of the static mesh component, this was getting saved and not restored during runtime (now does it in PostLoad)
	#!jira UE-47712
	#!rb none
	#!tests PIE in Editor while having HLOD forced on

Change 3570047 on 2017/08/02 by robomerge

	#!ROBOMERGE-AUTHOR: andrew.grant
	Fix for ensure leading to deadlock

	#!jira OR-42071
	#!tests compiled, ran ShortSoloGame test
	#!rb none

	#!ROBOMERGE-SOURCE: CL 3570044 in //Orion/Release-42/... via CL 3570045 via CL 3570046
	#!ROBOMERGE-BOT: ORION (Main -> Dev-General)

Change 3569932 on 2017/08/02 by robomerge

	#!ROBOMERGE-AUTHOR: andrew.grant
	Integration from UE4 - fixed a set of critical bugs that would cause rare crashes in the  async IO subsystems.

	#!rb gil.gribb
	#!tests ran 10x solo games, ran LoadTest, all PS4
	#!jira UE-47483

	#!ROBOMERGE-SOURCE: CL 3569929 in //Orion/Release-42/... via CL 3569930 via CL 3569931
	#!ROBOMERGE-BOT: ORION (Main -> Dev-General)

Change 3569839 on 2017/08/02 by [REDACTED_PERSON]

	Added the low quality monolith map to the build launcher and cook maps list.
	#!rb Trivial
	#!test Compile UAT

Change 3569441 on 2017/08/02 by robomerge

	#!ROBOMERGE-AUTHOR: laurent.delayen
	Re-enabled OnlyTickMontagesWhenNotRendered animation optimization, after fixing edge case where 'Evaluate' could be called on out of sync cached data.
	Only update cached data when the graph will be updated, to ensure we don't have mismatching update and evaluate data.

	#!rb none
	#!tests bot match w/ dedicated server.

	#!ROBOMERGE-SOURCE: CL 3569235 in //Orion/Main/...
	#!ROBOMERGE-BOT: ORION (Main -> Dev-General)

Change 3569397 on 2017/08/02 by robomerge

	#!ROBOMERGE-AUTHOR: ben.salem
	Merging using Dev-Gen_->_Release-42.1 Support for skins in perf reporting, also report overall test run time.
	#!rb various
	#!tests ran with new shallow test maps.

	#!ROBOMERGE-SOURCE: CL 3568892 in //Orion/Release-42.1/... via CL 3568967
	#!ROBOMERGE-BOT: ORION (Main -> Dev-General)

Change 3568639 on 2017/08/02 by [REDACTED_PERSON]

	Add individual FTest enumeration to Orion perf tests

	Moved methods to run FTests to helper for reuse.
	Effects perf tests now show individual test cases for more granular runs.
	All perf tests now do map check and draft logic as setup.

	#!tests ran updated tests in client/server
	#!rb ben.salem

Change 3568616 on 2017/08/02 by [REDACTED_PERSON]

	Switch shallow fx tests to report to their own mailing list.
	#!rb bob.ferreira
	#!tests recompiled

Change 3568607 on 2017/08/02 by [REDACTED_PERSON]

	Added staticmeshcomponent to the memreportfull command.
	#!rb Trivial
	#!test Paragon

Change 3568018 on 2017/08/01 by [REDACTED_PERSON]

	Add skin and character name to test reports, and also test runtime.
	#!rb none
	#!tests ran steel shallow test

Change 3567995 on 2017/08/01 by [REDACTED_PERSON]

	[FXTests] Updating FXTests and ShallowFXTests to call OnComplete() when doing check point validation.
	[Gauntlet] Changed InnerTestResult to be protected so that inherited test nodes can modify it.
	#!rb Ben.Salem
	#!TESTS Ran local automationtool run on a cooked main build.

Change 3567912 on 2017/08/01 by robomerge

	#!ROBOMERGE-AUTHOR: ori.cohen
	Added support for physical materials in immediate mode

	#!rb none
	#!tests none

	#!ROBOMERGE-SOURCE: CL 3567865 in //Orion/Release-42/... via CL 3567901 via CL 3567904
	#!ROBOMERGE-BOT: ORION (Main -> Dev-General)

Change 3567594 on 2017/08/01 by robomerge

	#!ROBOMERGE-AUTHOR: laurent.delayen
	Disabling animation optimization until I can fix an edge case.

	#!rb none
	#!tests arcade mode

	#!ROBOMERGE-SOURCE: CL 3567593 in //Orion/Main/...
	#!ROBOMERGE-BOT: ORION (Main -> Dev-General)

Change 3566953 on 2017/08/01 by robomerge

	#!ROBOMERGE-AUTHOR: laurent.delayen
	Animation Optimization: Added EMeshComponentUpdateFlag::OnlyTickMontagesWhenNotRendered to only update montages when not rendered, instead of doing a whole pose update. This means on dedicated servers we'll only update montages when needed, and not the AnimGraph, which we have no need for.
	This is also affects non rendered meshes on clients.

	#!rb martin.wilson
	#!tests ghost and wukong networked doing montage abilities.

	#!ROBOMERGE-SOURCE: CL 3566950 in //Orion/Main/...
	#!ROBOMERGE-BOT: ORION (Main -> Dev-General)

Change 3564610 on 2017/07/31 by [REDACTED_PERSON]

	Integrated CL 3543210 : Fixed an issue when computing material scales where the default material ends up being used instead of the required material.
	Deprecated previous material data as it was causing some waste.
	Integrated CL 3526859 : Texture mip bias is now reset whenever the streaming budget increases
	#!rb none
	#!tests played monolith2 on PS4

Change 3564509 on 2017/07/31 by robomerge

	#!ROBOMERGE-AUTHOR: laurent.delayen
	Fixed warning when opening Kwang AnimBP the first time, due to missing virtual bone name. (When Skeleton doesn't have PostLoad() called on it yet - happens only the first time the AnimBP is opened).

	#!rb none
	#!tests Kwang AnimBP opens without a warning.

	#!ROBOMERGE-SOURCE: CL 3564506 in //Orion/Release-42/... via CL 3564507
	#!ROBOMERGE-BOT: ORION (Main -> Dev-General)

diff --git a/Engine/Config/BaseEditor.ini b/Engine/Config/BaseEditor.ini
--- a/Engine/Config/BaseEditor.ini
+++ b/Engine/Config/BaseEditor.ini
@@ -395,7 +395,7 @@ bUseFullGCAssetClassNames=true
 +ConfigSettingBlacklist=*.Engine:Audio
 +ConfigSettingBlacklist=*.Engine:/Script/LinuxTargetPlatform.LinuxTargetSettings
 +ConfigSettingBlacklist=*.Engine:/Script/UnrealEd.CookerSettings:bIterativeCookingForFileCookContent
-
++ConfigSettingBlacklist=*.Engine:OnlineSubsystemMcp
 
 +VersionedIntRValues=r.AllowStaticLighting
 +VersionedIntRValues=r.GBuffer
diff --git a/Engine/Config/BaseEngine.ini b/Engine/Config/BaseEngine.ini
--- a/Engine/Config/BaseEngine.ini
+++ b/Engine/Config/BaseEngine.ini
@@ -969,6 +969,7 @@ ConfigClass=
 +Cmd="obj list class=SkeletalMesh -alphasort"
 +Cmd="obj list class=StaticMesh -alphasort"
 +Cmd="obj list class=Level -alphasort"
++Cmd="obj list class=StaticMeshComponent -alphasort"
 
 [MemoryPools]
 FLightPrimitiveInteractionInitialBlockSize=512
diff --git a/Engine/Plugins/FX/Niagara/Config/BaseNiagara.ini b/Engine/Plugins/FX/Niagara/Config/BaseNiagara.ini
new file mode 100644
--- /dev/null
+++ b/Engine/Plugins/FX/Niagara/Config/BaseNiagara.ini
@@ -0,0 +1,13 @@
+[CoreRedirects]
+
++ClassRedirects=(OldName="NiagaraEmitterProperties", NewName="/Script/Niagara.NiagaraEmitter")
++ClassRedirects=(OldName="NiagaraSimulation", NewName="/Script/Niagara.NiagaraEmitterInstance")
++ClassRedirects=(OldName="NiagaraEffect", NewName="/Script/Niagara.NiagaraSystem")
++ClassRedirects=(OldName="NiagaraEffectInstance", NewName="/Script/Niagara.NiagaraSystemInstance")
++ClassRedirects=(OldName="NiagaraEffectRenderer", NewName="/Script/Niagara.NiagaraSystemRenderer")
++ClassRedirects=(OldName="NiagaraEffectRendererRibbon", NewName="/Script/Niagara.NiagaraRendererRibbons")
++ClassRedirects=(OldName="NiagaraEffectRendererSprites", NewName="/Script/Niagara.NiagaraRendererSprites")
++ClassRedirects=(OldName="NiagaraEffectRendererMeshes", NewName="/Script/Niagara.NiagaraRendererMeshes")
++ClassRedirects=(OldName="NiagaraEffectRendererProperties", NewName="/Script/Niagara.NiagaraRendererProperties")
+
++EnumRedirects=(OldName="ENiagaraScriptUsage",ValueChanges=(("ENiagaraScriptUsage::SpawnScript","ENiagaraScriptUsage::ParticleSpawnScript"), ("ENiagaraScriptUsage::SpawnScriptInterpolated","ENiagaraScriptUsage::ParticleSpawnScriptInterpolated"),("ENiagaraScriptUsage::UpdateScript","ENiagaraScriptUsage::ParticleUpdateScript"), ("ENiagaraScriptUsage::EventScript","ENiagaraScriptUsage::ParticleEventScript"), ("ENiagaraScriptUsage::SystemScript","ENiagaraScriptUsage::SystemSpawnScript")) )
\ No newline at end of file
diff --git a/Engine/Plugins/FX/Niagara/Niagara.uplugin b/Engine/Plugins/FX/Niagara/Niagara.uplugin
new file mode 100644
--- /dev/null
+++ b/Engine/Plugins/FX/Niagara/Niagara.uplugin
@@ -0,0 +1,51 @@
+{
+	"FileVersion" : 1,
+	"Version" : 1,
+	"VersionName" : "1.0",
+	"FriendlyName" : "Niagara",
+	"Description" : "Niagara effect systems.",
+	"Category" : "FX",
+	"CreatedBy" : "Epic Games, Inc.",
+	"CreatedByURL" : "http://epicgames.com",
+	"DocsURL" : "",
+	"MarketplaceURL" : "",
+	"SupportURL" : "",
+	"EnabledByDefault" : false,
+	"CanContainContent" : true,
+	"IsBetaVersion" : true,
+	"Installed" : false,
+	"Modules" :
+	[
+		{
+		  "Name": "Niagara",
+		  "Type": "Runtime",
+		  "LoadingPhase": "PreDefault",
+		  "WhitelistPlatforms" : [ "Win64", "Win32", "PS4", "XboxOne" ]
+		},
+		{
+			"Name" : "NiagaraShader",
+			"Type" : "Runtime",
+			"LoadingPhase" : "PostConfigInit",
+			"WhitelistPlatforms" : [ "Win64", "Win32", "PS4", "XboxOne" ]
+		},
+
+		{
+			"Name" : "NiagaraVertexFactories",
+			"Type" : "Runtime",
+			"LoadingPhase" : "PostConfigInit",
+			"WhitelistPlatforms" : [ "Win64", "Win32", "PS4", "XboxOne" ]
+		},
+		{
+			"Name" : "NiagaraEditor",
+			"Type" : "Editor",
+			"LoadingPhase" : "Default",
+			"WhitelistPlatforms" : [ "Win64", "Win32", "PS4", "XboxOne" ]
+		},
+		{
+			"Name" : "NiagaraEditorWidgets",
+			"Type" : "Editor",
+			"LoadingPhase" : "Default",
+			"WhitelistPlatforms" : [ "Win64", "Win32", "PS4", "XboxOne" ]
+		}
+	]
+}
\ No newline at end of file
diff --git a/Engine/Plugins/FX/Niagara/Source/Niagara/Classes/NiagaraCollision.h b/Engine/Plugins/FX/Niagara/Source/Niagara/Classes/NiagaraCollision.h
new file mode 100644
--- /dev/null
+++ b/Engine/Plugins/FX/Niagara/Source/Niagara/Classes/NiagaraCollision.h
@@ -0,0 +1,92 @@
+#pragma once
+
+#include "NiagaraEvents.h"
+#include "WorldCollision.h"
+
+UENUM()
+enum class ENiagaraCollisionMode : uint8
+{
+	None = 0,
+	SceneGeometry,
+	DepthBuffer,
+	DistanceField
+};
+
+
+struct FNiagaraCollisionTrace
+{
+	FTraceHandle CollisionTraceHandle;
+	uint32 SourceParticleIndex;
+	FVector OriginalVelocity;
+};
+
+
+class FNiagaraCollisionBatch
+{
+public:
+	FNiagaraCollisionBatch()
+		: CollisionEventDataSet(nullptr)
+		, EmitterName(NAME_None)
+	{
+	}
+
+	~FNiagaraCollisionBatch()
+	{
+		FNiagaraEventDataSetMgr::Reset(OwnerSystemInstanceName, EmitterName);
+	}
+
+	void Tick(ENiagaraSimTarget Target)
+	{
+		if (CollisionEventDataSet)
+		{
+			CollisionEventDataSet->Tick(Target);
+		}
+	}
+
+	void Reset()
+	{
+		if (CollisionEventDataSet)
+		{
+			CollisionEventDataSet->SetNumInstances(0);
+		}
+	}
+
+	void Init(FName InOwnerSystemInstanceName, FName InEmitterName)
+	{
+		if (CollisionEventDataSet)
+		{
+			CollisionEventDataSet->Reset();
+		}
+		EmitterName = InEmitterName;
+		OwnerSystemInstanceName = InOwnerSystemInstanceName;
+		CollisionEventDataSet = FNiagaraEventDataSetMgr::CreateEventDataSet(OwnerSystemInstanceName, EmitterName, NIAGARA_BUILTIN_EVENTNAME_COLLISION);
+		// TODO: this should go away once we can use the FNiagaraCollisionEventPayload
+		// UStruct to create the data set
+		//FNiagaraVariable ValidVar(FNiagaraTypeDefinition::GetIntDef(), "Valid");
+		FNiagaraVariable PosVar(FNiagaraTypeDefinition::GetVec3Def(), "CollisionLocation");
+		FNiagaraVariable NrmVar(FNiagaraTypeDefinition::GetVec3Def(), "CollisionNormal");
+		FNiagaraVariable PhysMatIdxVar(FNiagaraTypeDefinition::GetIntDef(), "PhysicalMaterialIndex");
+		FNiagaraVariable VelVar(FNiagaraTypeDefinition::GetVec3Def(), "CollisionVelocity");
+		FNiagaraVariable ParticleIndexVar(FNiagaraTypeDefinition::GetIntDef(), "ParticleIndex");
+		//CollisionEventDataSet->AddVariable(ValidVar);
+		CollisionEventDataSet->AddVariable(PosVar);
+		CollisionEventDataSet->AddVariable(NrmVar);
+		CollisionEventDataSet->AddVariable(PhysMatIdxVar);
+		CollisionEventDataSet->AddVariable(VelVar);
+		CollisionEventDataSet->AddVariable(ParticleIndexVar);
+		CollisionEventDataSet->Finalize();
+	}
+	
+	void KickoffNewBatch(struct FNiagaraEmitterInstance *Sim, float DeltaSeconds);
+
+	void GenerateEventsFromResults(struct FNiagaraEmitterInstance *Sim);
+
+	const FNiagaraDataSet *GetDataSet() const { return CollisionEventDataSet; }
+private:
+	TArray<FTraceHandle> CollisionTraceHandles;
+	TArray<FNiagaraCollisionTrace> CollisionTraces;
+	TArray<FNiagaraCollisionEventPayload> CollisionEvents;
+	FNiagaraDataSet *CollisionEventDataSet;
+	FName EmitterName;
+	FName OwnerSystemInstanceName;
+};
\ No newline at end of file
diff --git a/Engine/Plugins/FX/Niagara/Source/Niagara/Classes/NiagaraConstants.h b/Engine/Plugins/FX/Niagara/Source/Niagara/Classes/NiagaraConstants.h
new file mode 100644
--- /dev/null
+++ b/Engine/Plugins/FX/Niagara/Source/Niagara/Classes/NiagaraConstants.h
@@ -0,0 +1,52 @@
+// Copyright 1998-2017 Epic Games, Inc. All Rights Reserved.
+
+#pragma once
+#include "NiagaraTypes.h"
+
+#define SYS_PARAM_ENGINE_DELTA_TIME					FNiagaraVariable(FNiagaraTypeDefinition::GetFloatDef(), TEXT("Engine.DeltaTime"))
+#define SYS_PARAM_ENGINE_INV_DELTA_TIME				FNiagaraVariable(FNiagaraTypeDefinition::GetFloatDef(), TEXT("Engine.InverseDeltaTime"))
+#define SYS_PARAM_ENGINE_POSITION					FNiagaraVariable(FNiagaraTypeDefinition::GetVec3Def(), TEXT("Engine.Owner.Position"))
+#define SYS_PARAM_ENGINE_VELOCITY					FNiagaraVariable(FNiagaraTypeDefinition::GetVec3Def(), TEXT("Engine.Owner.Velocity"))
+#define SYS_PARAM_ENGINE_X_AXIS						FNiagaraVariable(FNiagaraTypeDefinition::GetVec3Def(), TEXT("Engine.Owner.SystemXAxis"))
+#define SYS_PARAM_ENGINE_Y_AXIS						FNiagaraVariable(FNiagaraTypeDefinition::GetVec3Def(), TEXT("Engine.Owner.SystemYAxis"))
+#define SYS_PARAM_ENGINE_Z_AXIS						FNiagaraVariable(FNiagaraTypeDefinition::GetVec3Def(), TEXT("Engine.Owner.SystemZAxis"))
+#define SYS_PARAM_ENGINE_LOCAL_TO_WORLD				FNiagaraVariable(FNiagaraTypeDefinition::GetMatrix4Def(), TEXT("Engine.Owner.SystemLocalToWorld"))
+#define SYS_PARAM_ENGINE_WORLD_TO_LOCAL				FNiagaraVariable(FNiagaraTypeDefinition::GetMatrix4Def(), TEXT("Engine.Owner.SystemWorldToLocal"))
+#define SYS_PARAM_ENGINE_LOCAL_TO_WORLD_TRANSPOSED	FNiagaraVariable(FNiagaraTypeDefinition::GetMatrix4Def(), TEXT("Engine.Owner.SystemLocalToWorldTransposed"))
+#define SYS_PARAM_ENGINE_WORLD_TO_LOCAL_TRANSPOSED	FNiagaraVariable(FNiagaraTypeDefinition::GetMatrix4Def(), TEXT("Engine.Owner.WorldToLocalTransposed"))
+#define SYS_PARAM_ENGINE_MIN_DIST_TO_CAMERA			FNiagaraVariable(FNiagaraTypeDefinition::GetFloatDef(), TEXT("Engine.Owner.MinDistanceToCamera"))
+#define SYS_PARAM_ENGINE_EXEC_COUNT					FNiagaraVariable(FNiagaraTypeDefinition::GetIntDef(), TEXT("Engine.ExecutionCount"))
+#define SYS_PARAM_ENGINE_EMITTER_NUM_PARTICLES		FNiagaraVariable(FNiagaraTypeDefinition::GetIntDef(), TEXT("Engine.Emitter.NumParticles"))
+#define SYS_PARAM_ENGINE_SYSTEM_NUM_EMITTERS_ALIVE	FNiagaraVariable(FNiagaraTypeDefinition::GetIntDef(), TEXT("Engine.System.NumEmittersAlive"))
+#define SYS_PARAM_ENGINE_SYSTEM_NUM_EMITTERS		FNiagaraVariable(FNiagaraTypeDefinition::GetIntDef(), TEXT("Engine.System.NumEmitters"))
+
+#define SYS_PARAM_EMITTER_AGE						FNiagaraVariable(FNiagaraTypeDefinition::GetFloatDef(), TEXT("Emitter.Age"))
+#define SYS_PARAM_EMITTER_SPAWNRATE					FNiagaraVariable(FNiagaraTypeDefinition::GetFloatDef(), TEXT("Emitter.SpawnRate"))
+#define SYS_PARAM_EMITTER_SPAWN_INTERVAL			FNiagaraVariable(FNiagaraTypeDefinition::GetFloatDef(), TEXT("Emitter.SpawnInterval"))
+#define SYS_PARAM_EMITTER_INTERP_SPAWN_START_DT		FNiagaraVariable(FNiagaraTypeDefinition::GetFloatDef(), TEXT("Emitter.InterpSpawnStartDt"))
+
+#define SYS_PARAM_PARTICLES_POSITION				FNiagaraVariable(FNiagaraTypeDefinition::GetVec3Def(), TEXT("Particles.Position"))
+#define SYS_PARAM_PARTICLES_VELOCITY				FNiagaraVariable(FNiagaraTypeDefinition::GetVec3Def(), TEXT("Particles.Velocity"))
+#define SYS_PARAM_PARTICLES_COLOR					FNiagaraVariable(FNiagaraTypeDefinition::GetColorDef(), TEXT("Particles.Color"))
+#define SYS_PARAM_PARTICLES_SPRITE_ROTATION			FNiagaraVariable(FNiagaraTypeDefinition::GetFloatDef(), TEXT("Particles.SpriteRotation"))
+#define SYS_PARAM_PARTICLES_NORMALIZED_AGE			FNiagaraVariable(FNiagaraTypeDefinition::GetFloatDef(), TEXT("Particles.NormalizedAge"))
+#define SYS_PARAM_PARTICLES_SPRITE_SIZE				FNiagaraVariable(FNiagaraTypeDefinition::GetVec2Def(), TEXT("Particles.SpriteSize"))
+#define SYS_PARAM_PARTICLES_SPRITE_FACING			FNiagaraVariable(FNiagaraTypeDefinition::GetVec3Def(), TEXT("Particles.SpriteFacing"))
+#define SYS_PARAM_PARTICLES_SPRITE_ALIGNMENT		FNiagaraVariable(FNiagaraTypeDefinition::GetVec3Def(), TEXT("Particles.SpriteAlignment"))
+#define SYS_PARAM_PARTICLES_SUB_IMAGE_INDEX			FNiagaraVariable(FNiagaraTypeDefinition::GetFloatDef(), TEXT("Particles.SubImageIndex"))
+#define SYS_PARAM_PARTICLES_DYNAMIC_MATERIAL_PARAM	FNiagaraVariable(FNiagaraTypeDefinition::GetVec4Def(), TEXT("Particles.DynamicMaterialParameter"))
+#define SYS_PARAM_PARTICLES_SCALE					FNiagaraVariable(FNiagaraTypeDefinition::GetVec3Def(), TEXT("Particles.Scale"))
+#define SYS_PARAM_PARTICLES_ROTATION				FNiagaraVariable(FNiagaraTypeDefinition::GetVec3Def(), TEXT("Particles.Rotation"))
+#define SYS_PARAM_INSTANCE_ALIVE					FNiagaraVariable(FNiagaraTypeDefinition::GetBoolDef(), TEXT("Instance.Alive"))
+
+struct NIAGARA_API FNiagaraConstants
+{
+	static const TArray<FNiagaraVariable>& GetEngineConstants();
+	static FNiagaraVariable UpdateEngineConstant(const FNiagaraVariable& InVar);
+	static const FNiagaraVariable *FindEngineConstant(const FNiagaraVariable& InVar);
+	static FText GetEngineConstantDescription(const FNiagaraVariable& InVar);
+
+	static const TArray<FNiagaraVariable>& GetCommonParticleAttributes();
+	static FText GetAttributeDescription(const FNiagaraVariable& InVar);
+	static FString GetAttributeDefaultValue(const FNiagaraVariable& InVar);
+};
\ No newline at end of file
diff --git a/Engine/Plugins/FX/Niagara/Source/Niagara/Classes/NiagaraDataInterface.h b/Engine/Plugins/FX/Niagara/Source/Niagara/Classes/NiagaraDataInterface.h
new file mode 100644
--- /dev/null
+++ b/Engine/Plugins/FX/Niagara/Source/Niagara/Classes/NiagaraDataInterface.h
@@ -0,0 +1,210 @@
+// Copyright 1998-2016 Epic Games, Inc. All Rights Reserved.
+#pragma once
+
+#include "NiagaraCommon.h"
+#include "NiagaraShared.h"
+#include "VectorVM.h"
+#include "StaticMeshResources.h"
+#include "Curves/RichCurve.h"
+#include "NiagaraDataInterface.generated.h"
+
+class INiagaraCompiler;
+class UCurveVector;
+class UCurveLinearColor;
+class UCurveFloat;
+class FNiagaraSystemInstance;
+
+//////////////////////////////////////////////////////////////////////////
+// Some helper classes allowing neat, init time binding of templated vm external functions.
+USTRUCT()
+struct FNiagaraDataInterfaceBufferData
+{
+	GENERATED_BODY()
+public:
+	FNiagaraDataInterfaceBufferData()
+		:UniformName(TEXT("Undefined"))
+	{}
+
+	FNiagaraDataInterfaceBufferData(FName InName)
+		:UniformName(InName)
+	{
+	}
+	FReadBuffer Buffer;
+	FName UniformName;
+};
+
+// Adds a known type to the parameters
+template<typename DirectType, typename NextBinder>
+struct TNDIExplicitBinder
+{
+	template<typename... ParamTypes>
+	static FVMExternalFunction Bind(UNiagaraDataInterface* Interface, const FVMExternalFunctionBindingInfo& BindingInfo, void* InstanceData)
+	{
+		return NextBinder::template Bind<ParamTypes..., DirectType>(Interface, BindingInfo, InstanceData);
+	}
+};
+
+// Binder that tests the location of an operand and adds the correct handler type to the Binding parameters.
+template<int32 ParamIdx, typename DataType, typename NextBinder>
+struct TNDIParamBinder
+{
+	template<typename... ParamTypes>
+	static FVMExternalFunction Bind(UNiagaraDataInterface* Interface, const FVMExternalFunctionBindingInfo& BindingInfo, void* InstanceData)
+	{
+		if (BindingInfo.InputParamLocations[ParamIdx])
+		{
+			return NextBinder::template Bind<ParamTypes..., FConstantHandler<DataType>>(Interface, BindingInfo, InstanceData);
+		}
+		else
+		{
+			return NextBinder::template Bind<ParamTypes..., FRegisterHandler<DataType>>(Interface, BindingInfo, InstanceData);
+		}
+	}
+};
+
+//Helper macros allowing us to define the final binding structs for each vm external function function more concisely.
+#define NDI_FUNC_BINDER(ClassName, FuncName) T##ClassName##_##FuncName##Binder
+
+#define DEFINE_NDI_FUNC_BINDER(ClassName, FuncName)\
+struct NDI_FUNC_BINDER(ClassName, FuncName)\
+{\
+	template<typename ... ParamTypes>\
+	static FVMExternalFunction Bind(UNiagaraDataInterface* Interface, const FVMExternalFunctionBindingInfo& BindingInfo, void* InstanceData)\
+	{\
+		return FVMExternalFunction::CreateUObject(CastChecked<ClassName>(Interface), &ClassName::FuncName<ParamTypes...>);\
+	}\
+};
+
+//////////////////////////////////////////////////////////////////////////
+
+/** Base class for all Niagara data interfaces. */
+UCLASS(abstract, EditInlineNew)
+class NIAGARA_API UNiagaraDataInterface : public UObject
+{
+	GENERATED_UCLASS_BODY()
+		 
+public: 
+
+	/** Initializes the per instance data for this interface. Returns false if there was some error and the simulation should be disabled. */
+	virtual bool InitPerInstanceData(void* PerInstanceData, FNiagaraSystemInstance* SystemInstance) { return true; }
+
+	/** Destroys the per instence data for this interface. */
+	virtual void DestroyPerInstanceData(void* PerInstanceData, FNiagaraSystemInstance* SystemInstance) {}
+
+	/** Ticks the per instance data for this interface, if it has any. */
+	virtual bool PerInstanceTick(void* PerInstanceData, FNiagaraSystemInstance* SystemInstance, float DeltaSeconds) { return false; }
+
+	/** 
+	Returns the size of the per instance data for this interface. 0 if this interface has no per instance data. 
+	Must depend solely on the class of the interface and not on any particular member data of a individual interface.
+	*/
+	virtual int32 PerInstanceDataSize()const { return 0; }
+
+	/** Gets all the available functions for this data interface. */
+	virtual void GetFunctions(TArray<FNiagaraFunctionSignature>& OutFunctions) {}
+
+	/** Returns the delegate for the passed function signature. */
+	virtual FVMExternalFunction GetVMExternalFunction(const FVMExternalFunctionBindingInfo& BindingInfo, void* InstanceData) { return FVMExternalFunction(); };
+
+	/** Copies the contents of this DataInterface to another.*/
+	virtual bool CopyTo(UNiagaraDataInterface* Destination) const;
+
+	/** Determines if this DataInterface is the same as another.*/
+	virtual bool Equals(const UNiagaraDataInterface* Other) const;
+
+	virtual bool CanExecuteOnTarget(ENiagaraSimTarget Target)const { return false; }
+
+	/** Determines if this type definition matches to a known data interface type.*/
+	static bool IsDataInterfaceType(const FNiagaraTypeDefinition& TypeDef);
+
+	virtual bool GetFunctionHLSL(FString FunctionName, TArray<DIGPUBufferParamDescriptor> &Descriptors, FString &HLSLInterfaceID, FString &OutHLSL)
+	{
+		check("Unefined HLSL in data interface. Interfaces need to be able to return HLSL for each function they define in GetFunctions.");
+		return false;
+	}
+	virtual void GetBufferDefinitionHLSL(FString DataInterfaceID, TArray<DIGPUBufferParamDescriptor> &BufferDescriptors, FString &OutHLSL)
+	{
+		check("Unefined HLSL in data interface. Interfaces need to define HLSL for uniforms their functions access.");
+	}
+	virtual TArray<FNiagaraDataInterfaceBufferData> &GetBufferDataArray()
+	{
+		check("Undefined buffer array access.");
+		return GPUBuffers;
+	}
+
+	virtual void SetupBuffers(TArray<DIGPUBufferParamDescriptor> &BufferDescriptors)
+	{
+		check("Undefined buffer setup.");
+	}
+
+protected:
+	UPROPERTY()
+	TArray<FNiagaraDataInterfaceBufferData> GPUBuffers;  // not sure whether storing this on the base is the right idea
+};
+
+/** Base class for curve data interfaces which facilitates handling the curve data in a standardized way. */
+UCLASS(EditInlineNew, Category = "Curves", meta = (DisplayName = "Float Curve"))
+class NIAGARA_API UNiagaraDataInterfaceCurveBase : public UNiagaraDataInterface
+{
+protected:
+	GENERATED_BODY()
+	UPROPERTY()
+	bool GPUBufferDirty;
+	UPROPERTY()
+	TArray<float> ShaderLUT;
+
+public:
+	UNiagaraDataInterfaceCurveBase()
+		: GPUBufferDirty(false)
+	{
+	}
+
+	UNiagaraDataInterfaceCurveBase(FObjectInitializer const& ObjectInitializer)
+		: GPUBufferDirty(false)
+	{}
+
+	enum
+	{
+		CurveLUTWidth = 128
+	};
+	/** Structure to facilitate getting standardized curve information from a curve data interface. */
+	struct FCurveData
+	{
+		FCurveData(FRichCurve* InCurve, FName InName, FLinearColor InColor)
+			: Curve(InCurve)
+			, Name(InName)
+			, Color(InColor)
+		{
+		}
+		/** A pointer to the curve. */
+		FRichCurve* Curve;
+		/** The name of the curve, unique within the data interface, which identifies the curve in the UI. */
+		FName Name;
+		/** The color to use when displaying this curve in the UI. */
+		FLinearColor Color;
+	};
+
+	/** Gets information for all of the curves owned by this curve data interface. */
+	virtual void GetCurveData(TArray<FCurveData>& OutCurveData) { }
+
+	virtual bool GetFunctionHLSL(FString FunctionName, TArray<DIGPUBufferParamDescriptor> &Descriptors, FString &HLSLInterfaceID, FString &OutHLSL)
+	{
+		check("Undefined HLSL in data interface. All curve interfaces need to define this function and return SampleCurve code");
+		return false;
+	}
+	virtual void GetBufferDefinitionHLSL(FString DataInterfaceID, TArray<DIGPUBufferParamDescriptor> &BufferDescriptors, FString &OutHLSL)
+	{
+		check("Unefined HLSL in data interface. Interfaces need to define HLSL for uniforms their functions access.");
+	}
+	virtual TArray<FNiagaraDataInterfaceBufferData> &GetBufferDataArray()
+	{
+		check("Undefined buffer array access.");
+		return GPUBuffers;
+	}
+	virtual void SetupBuffers(TArray<DIGPUBufferParamDescriptor> &BufferDescriptors)
+	{
+		check("Undefined buffer setup.");
+	}
+
+	virtual bool CanExecuteOnTarget(ENiagaraSimTarget Target)const override { return true; }
+};
diff --git a/Engine/Plugins/FX/Niagara/Source/Niagara/Classes/NiagaraDataInterfaceColorCurve.h b/Engine/Plugins/FX/Niagara/Source/Niagara/Classes/NiagaraDataInterfaceColorCurve.h
new file mode 100644
--- /dev/null
+++ b/Engine/Plugins/FX/Niagara/Source/Niagara/Classes/NiagaraDataInterfaceColorCurve.h
@@ -0,0 +1,71 @@
+// Copyright 1998-2016 Epic Games, Inc. All Rights Reserved.
+#pragma once
+
+#include "NiagaraCommon.h"
+#include "NiagaraShared.h"
+#include "VectorVM.h"
+#include "StaticMeshResources.h"
+#include "Curves/RichCurve.h"
+#include "NiagaraDataInterface.h"
+#include "NiagaraDataInterfaceColorCurve.generated.h"
+
+class INiagaraCompiler;
+class UCurveVector;
+class UCurveLinearColor;
+class UCurveFloat;
+class FNiagaraSystemInstance;
+
+
+/** Data Interface allowing sampling of color curves. */
+UCLASS(EditInlineNew, Category="Curves", meta = (DisplayName = "Curve for Colors"))
+class NIAGARA_API UNiagaraDataInterfaceColorCurve : public UNiagaraDataInterfaceCurveBase
+{
+	GENERATED_UCLASS_BODY()
+public:
+
+#if WITH_EDITORONLY_DATA
+	UPROPERTY(EditAnywhere, Category = "Curve", meta=(AllowedClasses = CurveLinearColor))
+	FStringAssetReference CurveToCopy;
+#endif
+
+	UPROPERTY(EditAnywhere, Category = "Curve")
+	FRichCurve RedCurve;
+
+	UPROPERTY(EditAnywhere, Category = "Curve")
+	FRichCurve GreenCurve;
+
+	UPROPERTY(EditAnywhere, Category = "Curve")
+	FRichCurve BlueCurve;
+
+	UPROPERTY(EditAnywhere, Category = "Curve")
+	FRichCurve AlphaCurve;
+
+
+	//UObject Interface
+	virtual void PostInitProperties() override;
+	virtual void PostLoad() override;
+#if WITH_EDITOR
+	virtual void PostEditChangeProperty(struct FPropertyChangedEvent& PropertyChangedEvent) override;
+#endif
+	//UObject Interface End
+
+	void UpdateLUT();
+
+	virtual void GetFunctions(TArray<FNiagaraFunctionSignature>& OutFunctions)override;
+	virtual FVMExternalFunction GetVMExternalFunction(const FVMExternalFunctionBindingInfo& BindingInfo, void* InstanceData)override;
+
+	template<typename XParamType>
+	void SampleCurve(FVectorVMContext& Context);
+
+	virtual bool CopyTo(UNiagaraDataInterface* Destination) const override;
+
+	virtual bool Equals(const UNiagaraDataInterface* Other) const override;
+
+	//~ UNiagaraDataInterfaceCurveBase interface
+	virtual void GetCurveData(TArray<FCurveData>& OutCurveData) override;
+
+	virtual bool GetFunctionHLSL(FString FunctionName, TArray<DIGPUBufferParamDescriptor> &Descriptors, FString &HLSLInterfaceID, FString &OutHLSL) override;
+	virtual void GetBufferDefinitionHLSL(FString DataInterfaceID, TArray<DIGPUBufferParamDescriptor> &BufferDescriptors, FString &OutHLSL) override;
+	virtual TArray<FNiagaraDataInterfaceBufferData> &GetBufferDataArray() override;
+	virtual void SetupBuffers(TArray<DIGPUBufferParamDescriptor> &BufferDescriptors) override;
+};
diff --git a/Engine/Plugins/FX/Niagara/Source/Niagara/Classes/NiagaraDataInterfaceCurlNoise.h b/Engine/Plugins/FX/Niagara/Source/Niagara/Classes/NiagaraDataInterfaceCurlNoise.h
new file mode 100644
--- /dev/null
+++ b/Engine/Plugins/FX/Niagara/Source/Niagara/Classes/NiagaraDataInterfaceCurlNoise.h
@@ -0,0 +1,42 @@
+// Copyright 1998-2016 Epic Games, Inc. All Rights Reserved.
+#pragma once
+
+#include "NiagaraDataInterface.h"
+#include "NiagaraCommon.h"
+#include "VectorVM.h"
+#include "NiagaraDataInterfaceCurlNoise.generated.h"
+
+/** Data Interface allowing sampling of curl noise LUT. */
+UCLASS(EditInlineNew, Category = "Curl Noise LUT", meta = (DisplayName = "Curl Noise"))
+class NIAGARA_API UNiagaraDataInterfaceCurlNoise : public UNiagaraDataInterface
+{
+	GENERATED_UCLASS_BODY()
+	bool bGPUBufferDirty;
+public:
+
+	UNiagaraDataInterfaceCurlNoise()
+		: bGPUBufferDirty(true)
+	{}
+	static VectorRegister NoiseTable[17][17][17];
+
+	static void InitNoiseLUT();
+
+	//UObject Interface
+	virtual void PostInitProperties()override;
+	//UObject Interface End
+
+	//UNiagaraDataInterface Interface
+	virtual void GetFunctions(TArray<FNiagaraFunctionSignature>& OutFunctions)override;
+	virtual FVMExternalFunction GetVMExternalFunction(const FVMExternalFunctionBindingInfo& BindingInfo, void* InstanceData)override;
+	virtual bool CanExecuteOnTarget(ENiagaraSimTarget Target)const override { return true; }
+	//UNiagaraDataInterface Interface End
+
+	template<typename XType, typename YType, typename ZType>
+	void SampleNoiseField(FVectorVMContext& Context);
+
+	// GPU sim functionality
+	virtual bool GetFunctionHLSL(FString FunctionName, TArray<DIGPUBufferParamDescriptor> &Descriptors, FString &HLSLInterfaceID, FString &OutHLSL) override;
+	virtual void GetBufferDefinitionHLSL(FString DataInterfaceID, TArray<DIGPUBufferParamDescriptor> &BufferDescriptors, FString &OutHLSL) override;
+	virtual TArray<FNiagaraDataInterfaceBufferData> &GetBufferDataArray() override;
+	virtual void SetupBuffers(TArray<DIGPUBufferParamDescriptor> &BufferDescriptors) override;
+};
\ No newline at end of file
diff --git a/Engine/Plugins/FX/Niagara/Source/Niagara/Classes/NiagaraDataInterfaceCurve.h b/Engine/Plugins/FX/Niagara/Source/Niagara/Classes/NiagaraDataInterfaceCurve.h
new file mode 100644
--- /dev/null
+++ b/Engine/Plugins/FX/Niagara/Source/Niagara/Classes/NiagaraDataInterfaceCurve.h
@@ -0,0 +1,55 @@
+// Copyright 1998-2016 Epic Games, Inc. All Rights Reserved.
+#pragma once
+
+#include "NiagaraCommon.h"
+#include "NiagaraShared.h"
+#include "VectorVM.h"
+#include "StaticMeshResources.h"
+#include "Curves/RichCurve.h"
+#include "NiagaraDataInterface.h"
+#include "NiagaraDataInterfaceCurve.generated.h"
+
+
+/** Data Interface allowing sampling of float curves. */
+UCLASS(EditInlineNew, Category = "Curves", meta = (DisplayName = "Curve for Floats"))
+class NIAGARA_API UNiagaraDataInterfaceCurve : public UNiagaraDataInterfaceCurveBase
+{
+	GENERATED_UCLASS_BODY()
+
+public:
+#if WITH_EDITORONLY_DATA
+	UPROPERTY(EditAnywhere, Category = "Curve", meta = (AllowedClasses = CurveFloat))
+	FStringAssetReference CurveToCopy;
+#endif
+
+	UPROPERTY(EditAnywhere, Category = "Curve")
+	FRichCurve Curve;
+
+	void UpdateLUT();
+
+	//UObject Interface
+	virtual void PostInitProperties() override;
+
+#if WITH_EDITOR
+	virtual void PostEditChangeProperty(struct FPropertyChangedEvent& PropertyChangedEvent) override;
+#endif
+	//UObject Interface End
+
+	virtual void GetFunctions(TArray<FNiagaraFunctionSignature>& OutFunctions)override;
+	virtual FVMExternalFunction GetVMExternalFunction(const FVMExternalFunctionBindingInfo& BindingInfo, void* InstanceData)override;
+
+	template<typename XParamType>
+	void SampleCurve(FVectorVMContext& Context);
+
+	virtual bool CopyTo(UNiagaraDataInterface* Destination) const override;
+
+	virtual bool Equals(const UNiagaraDataInterface* Other) const override;
+
+	//~ UNiagaraDataInterfaceCurveBase interface
+	virtual void GetCurveData(TArray<FCurveData>& OutCurveData) override;
+
+	virtual bool GetFunctionHLSL(FString FunctionName, TArray<DIGPUBufferParamDescriptor> &Descriptors, FString &HLSLInterfaceID, FString &OutHLSL) override;
+	virtual void GetBufferDefinitionHLSL(FString DataInterfaceID, TArray<DIGPUBufferParamDescriptor> &BufferDescriptors, FString &OutHLSL) override;
+	virtual TArray<FNiagaraDataInterfaceBufferData> &GetBufferDataArray() override;
+	virtual void SetupBuffers(TArray<DIGPUBufferParamDescriptor> &BufferDescriptors) override;
+};
diff --git a/Engine/Plugins/FX/Niagara/Source/Niagara/Classes/NiagaraDataInterfaceSpline.h b/Engine/Plugins/FX/Niagara/Source/Niagara/Classes/NiagaraDataInterfaceSpline.h
new file mode 100644
--- /dev/null
+++ b/Engine/Plugins/FX/Niagara/Source/Niagara/Classes/NiagaraDataInterfaceSpline.h
@@ -0,0 +1,70 @@
+// Copyright 1998-2016 Epic Games, Inc. All Rights Reserved.
+#pragma once
+
+#include "NiagaraDataInterface.h"
+#include "NiagaraCommon.h"
+#include "VectorVM.h"
+#include "Components/SplineComponent.h"
+#include "NiagaraDataInterfaceSpline.generated.h"
+
+
+struct FNDISpline_InstanceData
+{
+	//Cached ptr to component we sample from. 
+	USplineComponent* Component;
+
+	//Cached ComponentToWorld.
+	FMatrix Transform;
+	//InverseTranspose of above for transforming normals/tangents.
+	FMatrix TransformInverseTransposed;
+};
+
+/** Data Interface allowing sampling of static meshes. */
+UCLASS(EditInlineNew, Category = "Splines", meta = (DisplayName = "Spline"))
+class NIAGARA_API UNiagaraDataInterfaceSpline : public UNiagaraDataInterface
+{
+	GENERATED_UCLASS_BODY()
+
+public:
+	
+	/** The source actor from which to sample.  */
+	UPROPERTY(EditAnywhere, Category = "Spline")
+	AActor* Source;
+	
+	//UObject Interface
+	virtual void PostInitProperties()override;
+	//UObject Interface End
+
+	//UNiagaraDataInterface Interface
+	virtual bool InitPerInstanceData(void* PerInstanceData, FNiagaraSystemInstance* SystemInstance) override;
+	virtual void DestroyPerInstanceData(void* PerInstanceData, FNiagaraSystemInstance* SystemInstance)override;
+	virtual bool PerInstanceTick(void* PerInstanceData, FNiagaraSystemInstance* SystemInstance, float DeltaSeconds) override;
+	virtual int32 PerInstanceDataSize()const override;
+	virtual void GetFunctions(TArray<FNiagaraFunctionSignature>& OutFunctions)override;
+	virtual FVMExternalFunction GetVMExternalFunction(const FVMExternalFunctionBindingInfo& BindingInfo, void* InstanceData)override;
+	virtual bool CopyTo(UNiagaraDataInterface* Destination) const override;
+	virtual bool Equals(const UNiagaraDataInterface* Other) const override;
+	virtual bool CanExecuteOnTarget(ENiagaraSimTarget Target)const override { return Target == ENiagaraSimTarget::CPUSim; }
+	//UNiagaraDataInterface Interface End
+
+	template<typename TransformHandlerType, typename SplineSampleType>
+	void SampleSplinePositionByUnitDistance(FVectorVMContext& Context);
+	template<typename TransformHandlerType, typename SplineSampleType>
+	void SampleSplineUpVectorByUnitDistance(FVectorVMContext& Context);
+	template<typename TransformHandlerType, typename SplineSampleType>
+	void SampleSplineRightVectorByUnitDistance(FVectorVMContext& Context);
+	template<typename TransformHandlerType, typename SplineSampleType>
+	void SampleSplineDirectionByUnitDistance(FVectorVMContext& Context);
+	template<typename TransformHandlerType, typename SplineSampleType>
+	void SampleSplineTangentByUnitDistance(FVectorVMContext& Context);
+	template<typename PosXType, typename PosYType, typename PosZType>
+	void FindClosestUnitDistanceFromPositionWS(FVectorVMContext& Context);
+	
+	void GetLocalToWorld(FVectorVMContext& Context);
+	void GetLocalToWorldInverseTransposed(FVectorVMContext& Context);
+
+private:
+	
+	void WriteTransform(const FMatrix& ToWrite, FVectorVMContext& Context);
+	//Cached ComponentToWorld.
+};
diff --git a/Engine/Plugins/FX/Niagara/Source/Niagara/Classes/NiagaraDataInterfaceStaticMesh.h b/Engine/Plugins/FX/Niagara/Source/Niagara/Classes/NiagaraDataInterfaceStaticMesh.h
new file mode 100644
--- /dev/null
+++ b/Engine/Plugins/FX/Niagara/Source/Niagara/Classes/NiagaraDataInterfaceStaticMesh.h
@@ -0,0 +1,230 @@
+// Copyright 1998-2016 Epic Games, Inc. All Rights Reserved.
+#pragma once
+
+#include "NiagaraDataInterface.h"
+#include "NiagaraCommon.h"
+#include "VectorVM.h"
+#include "StaticMeshResources.h"
+#include "Engine/StaticMesh.h"
+#include "NiagaraDataInterfaceStaticMesh.generated.h"
+
+//A coordinate on a mesh usable in Niagara.
+//Do not alter this struct without updating the data interfaces that use it!
+USTRUCT()
+struct FMeshTriCoordinate
+{
+	GENERATED_USTRUCT_BODY();
+	
+	UPROPERTY(EditAnywhere, Category="Coordinate")
+	int32 Tri;
+
+	UPROPERTY(EditAnywhere, Category="Coordinate")
+	FVector BaryCoord;
+};
+
+struct FNDIStaticMesh_InstanceData;
+struct FNDIStaticMeshSectionFilter;
+
+/** Allows uniform random sampling of a number of mesh sections filtered by an FNDIStaticMeshSectionFilter */
+struct FStaticMeshFilteredAreaWeightedSectionSampler : FStaticMeshAreaWeightedSectionSampler
+{
+	FStaticMeshFilteredAreaWeightedSectionSampler();
+	void Init(FStaticMeshLODResources* InRes, FNDIStaticMesh_InstanceData* InOwner);
+	virtual float GetWeights(TArray<float>& OutWeights)override;
+
+protected:
+	FStaticMeshLODResources* Res;
+	FNDIStaticMesh_InstanceData* Owner;
+};
+
+USTRUCT()
+struct FNDIStaticMeshSectionFilter
+{
+	GENERATED_USTRUCT_BODY();
+
+	/** Only allow sections these material slots. */
+	UPROPERTY(EditAnywhere, Category="Section Filter")
+	TArray<int32> AllowedMaterialSlots;
+
+	//Others?
+	//Banned material slots
+	
+	void Init(class UNiagaraDataInterfaceStaticMesh* Owner, bool bAreaWeighted);
+	FORCEINLINE bool CanEverReject()const { return AllowedMaterialSlots.Num() > 0; }
+};
+
+class UNiagaraDataInterfaceStaticMesh;
+
+struct FNDIStaticMesh_InstanceData
+{
+	FRandomStream RandStream;//Might remove this when we rework randoms. I don't like stateful randoms as they can't work on GPU and likely not threaded VM either!
+
+	 //Cached ptr to component we sample from. 
+	TWeakObjectPtr<USceneComponent> Component;
+
+	//Cached ptr to actual mesh we sample from. 
+	UStaticMesh* Mesh;
+
+	//Cached ComponentToWorld.
+	FMatrix Transform;
+	//InverseTranspose of above for transforming normals/tangents.
+	FMatrix TransformInverseTransposed;
+
+	//Cached ComponentToWorld from previous tick.
+	FMatrix PrevTransform;
+	//InverseTranspose of above for transforming normals/tangents.
+	FMatrix PrevTransformInverseTransposed;
+
+	/** Time separating Transform and PrevTransform. */
+	float DeltaSeconds;
+
+	/** True if the mesh we're using allows area weighted sampling. */
+	uint32 bIsAreaWeightedSampling : 1;
+
+	/** Cached results of this filter being applied to the owning mesh. */
+	TArray<int32> ValidSections;
+	/** Area weighted sampler for the valid sections. */
+	FStaticMeshFilteredAreaWeightedSectionSampler Sampler;
+
+	/** Allows sampling of the mesh's tris based on a dynamic color range. */
+	TSharedPtr<struct FDynamicVertexColorFilterData> DynamicVertexColorSampler;
+
+	FORCEINLINE UStaticMesh* GetActualMesh()const { return Mesh; }
+	FORCEINLINE bool UsesAreaWeighting()const { return bIsAreaWeightedSampling; }
+	FORCEINLINE bool MeshHasPositions()const { return Mesh && Mesh->RenderData->LODResources[0].PositionVertexBuffer.GetNumVertices() > 0; }
+	FORCEINLINE bool MeshHasVerts()const { return Mesh && Mesh->RenderData->LODResources[0].VertexBuffer.GetNumVertices() > 0; }
+	FORCEINLINE bool MeshHasColors()const { return Mesh && Mesh->RenderData->LODResources[0].ColorVertexBuffer.GetNumVertices() > 0; }
+
+	FORCEINLINE_DEBUGGABLE bool ResetRequired()const;
+
+	FORCEINLINE const TArray<int32>& GetValidSections()const { return ValidSections; }
+	FORCEINLINE const FStaticMeshAreaWeightedSectionSampler& GetAreaWeigtedSampler()const { return Sampler; }
+
+	void InitVertexColorFiltering();
+
+	FORCEINLINE_DEBUGGABLE bool Init(UNiagaraDataInterfaceStaticMesh* Interface, FNiagaraSystemInstance* SystemInstance);
+	FORCEINLINE_DEBUGGABLE bool Tick(UNiagaraDataInterfaceStaticMesh* Interface, FNiagaraSystemInstance* SystemInstance, float InDeltaSeconds);
+};
+
+/** Data Interface allowing sampling of static meshes. */
+UCLASS(EditInlineNew, Category = "Meshes", meta = (DisplayName = "Static Mesh"))
+class NIAGARA_API UNiagaraDataInterfaceStaticMesh : public UNiagaraDataInterface
+{
+	GENERATED_UCLASS_BODY()
+
+public:
+	
+	/** Mesh used to sample from when not overridden by a source actor from the scene. Also useful for previewing in the editor. */
+	UPROPERTY(EditAnywhere, Category = "Mesh")
+	UStaticMesh* DefaultMesh;
+
+	/** The source actor from which to sample. Takes precedence over the direct mesh. */
+	UPROPERTY(EditAnywhere, Category = "Mesh")
+	AActor* Source;
+	
+	/** Array of filters the can be used to limit sampling to certain sections of the mesh. */
+	UPROPERTY(EditAnywhere, Category = "Mesh")
+	FNDIStaticMeshSectionFilter SectionFilter;
+
+public:
+
+	//~ UObject interface
+
+#if WITH_EDITOR
+	virtual void PostInitProperties()override;
+#endif
+
+public:
+
+	//~ UNiagaraDataInterface interface
+
+	virtual bool InitPerInstanceData(void* PerInstanceData, FNiagaraSystemInstance* SystemInstance)override;
+	virtual void DestroyPerInstanceData(void* PerInstanceData, FNiagaraSystemInstance* SystemInstance)override;
+	virtual bool PerInstanceTick(void* PerInstanceData, FNiagaraSystemInstance* SystemInstance, float DeltaSeconds) override;
+	virtual int32 PerInstanceDataSize()const override { return sizeof(FNDIStaticMesh_InstanceData); }
+
+	virtual void GetFunctions(TArray<FNiagaraFunctionSignature>& OutFunctions)override;
+	virtual FVMExternalFunction GetVMExternalFunction(const FVMExternalFunctionBindingInfo& BindingInfo, void* InstanceData)override;
+	virtual bool CopyTo(UNiagaraDataInterface* Destination) const override;
+	virtual bool Equals(const UNiagaraDataInterface* Other) const override;
+	virtual bool CanExecuteOnTarget(ENiagaraSimTarget Target)const override { return Target == ENiagaraSimTarget::CPUSim; }
+public:
+
+	template<typename TAreaWeighted>
+	void RandomSection(FVectorVMContext& Context);
+
+	template<typename TAreaWeighted>
+	void RandomTriCoord(FVectorVMContext& Context);
+
+	template<typename TAreaWeighted, typename SectionIdxType>
+	void RandomTriCoordOnSection(FVectorVMContext& Context);
+
+ 	template<typename InputType0, typename InputType1>
+ 	void RandomTriCoordVertexColorFiltered(FVectorVMContext& Context);
+	
+	template<typename TransformHandlerType, typename TriType, typename BaryXType, typename BaryYType, typename BaryZType>
+	void GetTriCoordPosition(FVectorVMContext& Context);
+
+	template<typename TransformHandlerType, typename TriType, typename BaryXType, typename BaryYType, typename BaryZType>
+	void GetTriCoordNormal(FVectorVMContext& Context);
+
+	template<typename VertexAccessorType, typename TransformHandlerType, typename TriType, typename BaryXType, typename BaryYType, typename BaryZType>
+	void GetTriCoordTangents(FVectorVMContext& Context);
+
+	template<typename TriType, typename BaryXType, typename BaryYType, typename BaryZType>
+	void GetTriCoordColor(FVectorVMContext& Context);
+
+	template<typename VertexAccessorType, typename TriType, typename BaryXType, typename BaryYType, typename BaryZType, typename UVSetType>
+	void GetTriCoordUV(FVectorVMContext& Context);
+
+	template<typename TriType, typename BaryXType, typename BaryYType, typename BaryZType>
+	void GetTriCoordPositionAndVelocity(FVectorVMContext& Context);
+
+	void GetLocalToWorld(FVectorVMContext& Context);
+	void GetLocalToWorldInverseTransposed(FVectorVMContext& Context);
+	void GetWorldVelocity(FVectorVMContext& Context);
+
+	FORCEINLINE_DEBUGGABLE bool UsesSectionFilter()const { return SectionFilter.CanEverReject(); }
+
+	//TODO: Vertex color filtering requires a bit more work.
+	//FORCEINLINE bool UsesVertexColorFiltering()const { return bSupportingVertexColorSampling && bEnableVertexColorRangeSorting; }
+private:
+	
+	template<typename TAreaWeighted, bool bFiltered>
+	FORCEINLINE_DEBUGGABLE int32 RandomSection(FStaticMeshLODResources& Res, FNDIStaticMesh_InstanceData* InstData);
+
+	template<typename TAreaWeighted, bool bFiltered>
+	FORCEINLINE_DEBUGGABLE int32 RandomTriIndex(FStaticMeshLODResources& Res, FNDIStaticMesh_InstanceData* InstData);
+
+	template<typename TAreaWeighted>
+	FORCEINLINE_DEBUGGABLE int32 RandomTriIndexOnSection(FStaticMeshLODResources& Res, int32 SectionIdx, FNDIStaticMesh_InstanceData* InstData);
+
+	void WriteTransform(const FMatrix& ToWrite, FVectorVMContext& Context);
+};
+
+//TODO: IMO this should be generalized fuhrer if possible and extended to a system allowing filtering based on texture color etc too.
+struct FDynamicVertexColorFilterData
+{
+	//Cached ComponentToWorld.
+	//Cached ComponentToWorld from previous tick.
+	/** Container for the vertex colored triangles broken out by red channel values*/
+	TArray<uint32> TrianglesSortedByVertexColor;
+	/** Mapping from vertex color red value to starting entry in TrianglesSortedByVertexColor*/
+	TArray<uint32> VertexColorToTriangleStart;
+
+	bool Init(FNDIStaticMesh_InstanceData* Instance);
+};
+
+class FNDI_StaticMesh_GeneratedData
+{
+	static TMap<uint32, TSharedPtr<FDynamicVertexColorFilterData>> DynamicVertexColorFilters;
+
+	static FCriticalSection CriticalSection;
+public:
+	
+	/** Retrieves existing filter data for the passed mesh or generates a new one. */
+	static TSharedPtr<FDynamicVertexColorFilterData> GetDynamicColorFilterData(FNDIStaticMesh_InstanceData* Instance);
+
+	/** Todo: Find a place to call this on level change or somewhere. */
+	static void CleanupDynamicColorFilterData();
+};
diff --git a/Engine/Plugins/FX/Niagara/Source/Niagara/Classes/NiagaraDataInterfaceVector2DCurve.h b/Engine/Plugins/FX/Niagara/Source/Niagara/Classes/NiagaraDataInterfaceVector2DCurve.h
new file mode 100644
--- /dev/null
+++ b/Engine/Plugins/FX/Niagara/Source/Niagara/Classes/NiagaraDataInterfaceVector2DCurve.h
@@ -0,0 +1,62 @@
+// Copyright 1998-2016 Epic Games, Inc. All Rights Reserved.
+#pragma once
+
+#include "NiagaraCommon.h"
+#include "NiagaraShared.h"
+#include "VectorVM.h"
+#include "StaticMeshResources.h"
+#include "Curves/RichCurve.h"
+#include "NiagaraDataInterface.h"
+
+#include "NiagaraDataInterfaceVector2DCurve.generated.h"
+
+
+/** Data Interface allowing sampling of vector curves. */
+UCLASS(EditInlineNew, Category = "Curves", meta = (DisplayName = "Curve for Vector2D's"))
+class NIAGARA_API UNiagaraDataInterfaceVector2DCurve : public UNiagaraDataInterfaceCurveBase
+{
+	GENERATED_UCLASS_BODY()
+
+public:
+#if WITH_EDITORONLY_DATA
+	UPROPERTY(EditAnywhere, Category = "Curve", meta = (AllowedClasses = CurveVector))
+		FStringAssetReference CurveToCopy;
+#endif
+
+	UPROPERTY(EditAnywhere, Category = "Curve")
+	FRichCurve XCurve;
+
+	UPROPERTY(EditAnywhere, Category = "Curve")
+	FRichCurve YCurve;
+
+	void UpdateLUT();
+
+	//UObject Interface
+	virtual void PostInitProperties() override;
+	virtual void PostLoad() override;
+
+#if WITH_EDITOR
+	virtual void PostEditChangeProperty(struct FPropertyChangedEvent& PropertyChangedEvent) override;
+#endif
+	//UObject Interface End
+
+	virtual bool CopyTo(UNiagaraDataInterface* Destination) const override;
+
+	virtual bool Equals(const UNiagaraDataInterface* Other) const override;
+
+
+	virtual void GetFunctions(TArray<FNiagaraFunctionSignature>& OutFunctions)override;
+
+	virtual FVMExternalFunction GetVMExternalFunction(const FVMExternalFunctionBindingInfo& BindingInfo, void* InstanceData)override;
+
+	template<typename XParamType>
+	void SampleCurve(FVectorVMContext& Context);
+
+	//~ UNiagaraDataInterfaceCurveBase interface
+	virtual void GetCurveData(TArray<FCurveData>& OutCurveData) override;
+
+	virtual bool GetFunctionHLSL(FString FunctionName, TArray<DIGPUBufferParamDescriptor> &Descriptors, FString &HLSLInterfaceID, FString &OutHLSL) override;
+	virtual void GetBufferDefinitionHLSL(FString DataInterfaceID, TArray<DIGPUBufferParamDescriptor> &BufferDescriptors, FString &OutHLSL) override;
+	virtual TArray<FNiagaraDataInterfaceBufferData> &GetBufferDataArray() override;
+	virtual void SetupBuffers(TArray<DIGPUBufferParamDescriptor> &BufferDescriptors) override;
+};
\ No newline at end of file
diff --git a/Engine/Plugins/FX/Niagara/Source/Niagara/Classes/NiagaraDataInterfaceVector4Curve.h b/Engine/Plugins/FX/Niagara/Source/Niagara/Classes/NiagaraDataInterfaceVector4Curve.h
new file mode 100644
--- /dev/null
+++ b/Engine/Plugins/FX/Niagara/Source/Niagara/Classes/NiagaraDataInterfaceVector4Curve.h
@@ -0,0 +1,68 @@
+// Copyright 1998-2016 Epic Games, Inc. All Rights Reserved.
+#pragma once
+
+#include "NiagaraCommon.h"
+#include "NiagaraShared.h"
+#include "VectorVM.h"
+#include "StaticMeshResources.h"
+#include "Curves/RichCurve.h"
+#include "NiagaraDataInterface.h"
+#include "NiagaraDataInterfaceVector4Curve.generated.h"
+
+class INiagaraCompiler;
+class UCurveVector;
+class UCurveLinearColor;
+class UCurveFloat;
+class FNiagaraSystemInstance;
+
+/** Data Interface allowing sampling of vector curves. */
+UCLASS(EditInlineNew, Category = "Curves", meta = (DisplayName = "Curve for Vector4's"))
+class NIAGARA_API UNiagaraDataInterfaceVector4Curve : public UNiagaraDataInterfaceCurveBase
+{
+	GENERATED_UCLASS_BODY()
+
+public:
+#if WITH_EDITORONLY_DATA
+	UPROPERTY(EditAnywhere, Category = "Curve", meta = (AllowedClasses = CurveLinearColor))
+	FStringAssetReference CurveToCopy;
+#endif
+
+	UPROPERTY(EditAnywhere, Category = "Curve")
+	FRichCurve XCurve;
+
+	UPROPERTY(EditAnywhere, Category = "Curve")
+	FRichCurve YCurve;
+
+	UPROPERTY(EditAnywhere, Category = "Curve")
+	FRichCurve ZCurve;
+
+	UPROPERTY(EditAnywhere, Category = "Curve")
+	FRichCurve WCurve;
+
+	void UpdateLUT();
+
+	//UObject Interface
+	virtual void PostInitProperties() override;
+	virtual void PostLoad() override;
+#if WITH_EDITOR
+	virtual void PostEditChangeProperty(struct FPropertyChangedEvent& PropertyChangedEvent) override;
+#endif
+	//UObject Interface End
+
+	virtual void GetFunctions(TArray<FNiagaraFunctionSignature>& OutFunctions)override;
+	virtual FVMExternalFunction GetVMExternalFunction(const FVMExternalFunctionBindingInfo& BindingInfo, void* InstanceData)override;
+
+	template<typename XParamType>
+	void SampleCurve(FVectorVMContext& Context);
+
+	virtual bool CopyTo(UNiagaraDataInterface* Destination) const override;
+
+	virtual bool Equals(const UNiagaraDataInterface* Other) const override;
+
+	//~ UNiagaraDataInterfaceCurveBase interface
+	virtual void GetCurveData(TArray<FCurveData>& OutCurveData) override;
+	virtual bool GetFunctionHLSL(FString FunctionName, TArray<DIGPUBufferParamDescriptor> &Descriptors, FString &HLSLInterfaceID, FString &OutHLSL) override;
+	virtual void GetBufferDefinitionHLSL(FString DataInterfaceID, TArray<DIGPUBufferParamDescriptor> &BufferDescriptors, FString &OutHLSL) override;
+	virtual TArray<FNiagaraDataInterfaceBufferData> &GetBufferDataArray() override;
+	virtual void SetupBuffers(TArray<DIGPUBufferParamDescriptor> &BufferDescriptors) override;
+};
diff --git a/Engine/Plugins/FX/Niagara/Source/Niagara/Classes/NiagaraDataInterfaceVectorCurve.h b/Engine/Plugins/FX/Niagara/Source/Niagara/Classes/NiagaraDataInterfaceVectorCurve.h
new file mode 100644
--- /dev/null
+++ b/Engine/Plugins/FX/Niagara/Source/Niagara/Classes/NiagaraDataInterfaceVectorCurve.h
@@ -0,0 +1,66 @@
+// Copyright 1998-2016 Epic Games, Inc. All Rights Reserved.
+#pragma once
+
+#include "NiagaraCommon.h"
+#include "NiagaraShared.h"
+#include "VectorVM.h"
+#include "StaticMeshResources.h"
+#include "Curves/RichCurve.h"
+#include "NiagaraDataInterface.h"
+#include "NiagaraDataInterfaceVectorCurve.generated.h"
+
+class INiagaraCompiler;
+class UCurveVector;
+class UCurveLinearColor;
+class UCurveFloat;
+class FNiagaraSystemInstance;
+
+
+/** Data Interface allowing sampling of vector curves. */
+UCLASS(EditInlineNew, Category = "Curves", meta = (DisplayName = "Curve for Vector3's"))
+class NIAGARA_API UNiagaraDataInterfaceVectorCurve : public UNiagaraDataInterfaceCurveBase
+{
+	GENERATED_UCLASS_BODY()
+public:
+#if WITH_EDITORONLY_DATA
+	UPROPERTY(EditAnywhere, Category = "Curve", meta = (AllowedClasses = CurveVector))
+	FStringAssetReference CurveToCopy;
+#endif
+
+	UPROPERTY(EditAnywhere, Category = "Curve")
+	FRichCurve XCurve;
+
+	UPROPERTY(EditAnywhere, Category = "Curve")
+	FRichCurve YCurve;
+
+	UPROPERTY(EditAnywhere, Category = "Curve")
+	FRichCurve ZCurve;
+
+	void UpdateLUT();
+
+	//UObject Interface
+	virtual void PostInitProperties() override;
+	virtual void PostLoad() override;
+#if WITH_EDITOR
+	virtual void PostEditChangeProperty(struct FPropertyChangedEvent& PropertyChangedEvent) override;
+#endif
+	//UObject Interface End
+
+	virtual void GetFunctions(TArray<FNiagaraFunctionSignature>& OutFunctions)override;
+	virtual FVMExternalFunction GetVMExternalFunction(const FVMExternalFunctionBindingInfo& BindingInfo, void* InstanceData)override;
+
+	template<typename XParamType>
+	void SampleCurve(FVectorVMContext& Context);
+
+	virtual bool CopyTo(UNiagaraDataInterface* Destination) const override;
+
+	virtual bool Equals(const UNiagaraDataInterface* Other) const override;
+
+	//~ UNiagaraDataInterfaceCurveBase interface
+	virtual void GetCurveData(TArray<FCurveData>& OutCurveData) override;
+
+	virtual bool GetFunctionHLSL(FString FunctionName, TArray<DIGPUBufferParamDescriptor> &Descriptors, FString &HLSLInterfaceID, FString &OutHLSL) override;
+	virtual void GetBufferDefinitionHLSL(FString DataInterfaceID, TArray<DIGPUBufferParamDescriptor> &BufferDescriptors, FString &OutHLSL) override;
+	virtual TArray<FNiagaraDataInterfaceBufferData> &GetBufferDataArray() override;
+	virtual void SetupBuffers(TArray<DIGPUBufferParamDescriptor> &BufferDescriptors) override;
+};
diff --git a/Engine/Plugins/FX/Niagara/Source/Niagara/Classes/NiagaraDataSet.h b/Engine/Plugins/FX/Niagara/Source/Niagara/Classes/NiagaraDataSet.h
new file mode 100644
--- /dev/null
+++ b/Engine/Plugins/FX/Niagara/Source/Niagara/Classes/NiagaraDataSet.h
@@ -0,0 +1,1282 @@
+// Copyright 1998-2017 Epic Games, Inc. All Rights Reserved.
+#pragma once
+
+#include "CoreMinimal.h"
+#include "NiagaraCommon.h"
+#include "DynamicRHIResourceArray.h"
+#include "RHI.h"
+#include "VectorVM.h"
+#include "RenderingThread.h"
+
+
+/** Helper class defining the layout and location of an FNiagaraVariable in an FNiagaraDataBuffer-> */
+struct FNiagaraVariableLayoutInfo
+{
+	/** Start index for the float components in the main buffer. */
+	uint32 FloatComponentStart;
+	/** Start index for the int32 components in the main buffer. */
+	uint32 Int32ComponentStart;
+
+	uint32 GetNumFloatComponents()const { return LayoutInfo.FloatComponentByteOffsets.Num(); }
+	uint32 GetNumInt32Components()const { return LayoutInfo.Int32ComponentByteOffsets.Num(); }
+
+	/** This variable's type layout info. */
+	FNiagaraTypeLayoutInfo LayoutInfo;
+};
+
+class FNiagaraDataSet;
+
+/** Buffer containing one frame of Niagara simulation data. */
+class FNiagaraDataBuffer
+{
+public:
+	FNiagaraDataBuffer() : Owner(nullptr)
+	{
+		Reset();
+	}
+	void Init(FNiagaraDataSet* InOwner);
+	void Allocate(uint32 NumInstances, ENiagaraSimTarget Target = ENiagaraSimTarget::CPUSim, bool bMaintainExisting=false);
+	void AllocateGPU(uint32 InNumInstances, FRHICommandList &RHICmdList);
+	void InitGPUFromCPU();
+	void SwapInstances(uint32 OldIndex, uint32 NewIndex);
+	void KillInstance(uint32 InstanceIdx);
+	void CopyTo(FNiagaraDataBuffer& DestBuffer);
+
+	/** Returns a ptr to the start of the buffer for the passed component idx and base type. */
+	//template<typename T>
+	//FORCEINLINE uint8* GetComponentPtr(uint32 ComponentIdx);
+	/** Returns a ptr to specific instance data in the buffer for the passed component idx and base type. */
+
+	FORCEINLINE uint8* GetComponentPtrFloat(uint32 ComponentIdx)
+	{
+		return FloatData.GetData() + FloatStride * ComponentIdx;
+	}
+
+	FORCEINLINE uint8* GetComponentPtrInt32(uint32 ComponentIdx)
+	{
+		return Int32Data.GetData() + Int32Stride * ComponentIdx;
+	}
+
+	FORCEINLINE float* GetInstancePtrFloat(uint32 ComponentIdx, uint32 InstanceIdx)
+	{
+		return (float*)GetComponentPtrFloat(ComponentIdx) + InstanceIdx;
+	}
+
+	FORCEINLINE int32* GetInstancePtrInt32(uint32 ComponentIdx, uint32 InstanceIdx)
+	{
+		return (int32*)GetComponentPtrInt32(ComponentIdx) + InstanceIdx;
+	}
+
+	uint32 GetNumInstances()const { return NumInstances; }
+	uint32 GetNumInstancesAllocated()const { return NumInstancesAllocated; }
+
+	void SetNumInstances(uint32 InNumInstances) 
+	{ 
+		NumInstances = InNumInstances;
+	}
+
+	FORCEINLINE void Reset()
+	{
+		FloatData.Empty();
+		Int32Data.Empty();
+		FloatStride = 0;
+		Int32Stride = 0;
+		NumInstances = 0;
+		NumInstancesAllocated = 0;
+		NumChunksAllocatedForGPU = 0;
+	}
+
+	FORCEINLINE uint32 GetSizeBytes()const
+	{
+		return FloatData.Num() + Int32Data.Num();
+	}
+
+	const FRWBuffer *GetGPUBufferFloat() const
+	{
+		return &GPUBufferFloat;
+	}
+
+	const FRWBuffer *GetGPUBufferInt() const
+	{
+		return &GPUBufferInt;
+	}
+
+	int32 GetSafeComponentBufferSize() const
+	{
+		return GetSafeComponentBufferSize(GetNumInstancesAllocated());
+	}
+
+	uint32 GetFloatStride() const { return FloatStride; }
+	uint32 GetInt32Stride() const { return Int32Stride; }
+
+private:
+
+	FORCEINLINE int32 GetSafeComponentBufferSize(int32 RequiredSize) const
+	{
+		//Round up to VECTOR_WIDTH_BYTES.
+		//Both aligns the component buffers to the vector width but also ensures their ops cannot stomp over one another.		
+		return RequiredSize + VECTOR_WIDTH_BYTES - (RequiredSize % VECTOR_WIDTH_BYTES);
+	}
+
+	/** Back ptr to our owning data set. Used to access layout info for the buffer. */
+	FNiagaraDataSet* Owner;
+
+	/** Float components of simulation data. */
+	TResourceArray<uint8> FloatData;
+	/** Int32 components of simulation data. */
+	TResourceArray<uint8> Int32Data;
+
+	/** Stride between components in the float buffer. */
+	uint32 FloatStride;
+	/** Stride between components in the int32 buffer. */
+	uint32 Int32Stride;
+
+	uint32 NumChunksAllocatedForGPU;
+
+	/** Number of instances in data. */
+	uint32 NumInstances;
+	/** Number of instances the buffer has been allocated for. */
+	uint32 NumInstancesAllocated;
+
+
+	FRWBuffer GPUBufferFloat;
+	FRWBuffer GPUBufferInt;
+};
+
+//////////////////////////////////////////////////////////////////////////
+
+/**
+General storage class for all per instance simulation data in Niagara.
+*/
+class FNiagaraDataSet
+{
+	friend FNiagaraDataBuffer;
+public:
+	FNiagaraDataSet()
+	{
+		Reset();
+	}
+	
+	FNiagaraDataSet(FNiagaraDataSetID InID)
+		: ID(InID)
+	{
+		Reset();
+	}
+
+	void Reset()
+	{
+		Variables.Empty();
+		VariableLayoutMap.Empty();
+		Data[0].Reset();
+		Data[1].Reset();
+		CurrBuffer = 0;
+		CurrRenderBuffer = 0;
+		bFinalized = false;
+		TotalFloatComponents = 0;
+		TotalInt32Components = 0;
+		MaxBufferIdx = 1;
+	}
+
+	void AddVariable(FNiagaraVariable& Variable)
+	{
+		check(!bFinalized);
+		Variables.AddUnique(Variable);
+	}
+
+	void AddVariables(const TArray<FNiagaraVariable>& Vars)
+	{
+		check(!bFinalized);
+		for (const FNiagaraVariable& Var : Vars)
+		{
+			Variables.AddUnique(Var);
+		}
+	}
+
+	/** Finalize the addition of variables and other setup before this data set can be used. */
+	FORCEINLINE void Finalize()
+	{
+		check(!bFinalized);
+		bFinalized = true;
+		BuildLayout();
+	}
+
+	/** Removes a specific instance from the current frame's data buffer. */
+	FORCEINLINE void KillInstance(uint32 InstanceIdx)
+	{
+		check(bFinalized);
+		CurrData().KillInstance(InstanceIdx);
+	}
+
+	FORCEINLINE void SwapInstances(uint32 OldIndex, uint32 NewIndex)
+	{
+		check(bFinalized);
+		PrevData().SwapInstances(OldIndex, NewIndex);
+	}
+
+	/** Appends the passed variable to the set of input and output registers ready for consumption by the VectorVM. */
+	bool AppendToRegisterTable(const FNiagaraVariable& VarInfo, uint8** InputRegisters, int32& NumInputRegisters, uint8** OutputRegisters, int32& NumOutputRegisters, int32 StartInstance, bool bNoOutputRegisters = false)
+	{
+		check(bFinalized);
+		if (const FNiagaraVariableLayoutInfo* VariableLayout = VariableLayoutMap.Find(VarInfo))
+		{
+			uint32 NumComponents = VariableLayout->LayoutInfo.GetNumComponents();
+
+			for (uint32 CompIdx = 0; CompIdx < VariableLayout->GetNumFloatComponents(); ++CompIdx)
+			{
+				uint32 CompBufferOffset = VariableLayout->FloatComponentStart + CompIdx;
+				uint32 CompRegisterOffset = VariableLayout->LayoutInfo.FloatComponentRegisterOffsets[CompIdx];
+				InputRegisters[NumInputRegisters + CompRegisterOffset] = (uint8*)PrevData().GetInstancePtrFloat(CompBufferOffset, StartInstance);
+				OutputRegisters[NumOutputRegisters + CompRegisterOffset] = bNoOutputRegisters ? nullptr : (uint8*)CurrData().GetInstancePtrFloat(CompBufferOffset, StartInstance);
+			}
+			for (uint32 CompIdx = 0; CompIdx < VariableLayout->GetNumInt32Components(); ++CompIdx)
+			{
+				uint32 CompBufferOffset = VariableLayout->Int32ComponentStart + CompIdx;
+				uint32 CompRegisterOffset = VariableLayout->LayoutInfo.Int32ComponentRegisterOffsets[CompIdx];
+				InputRegisters[NumInputRegisters + CompRegisterOffset] = (uint8*)PrevData().GetInstancePtrInt32(CompBufferOffset, StartInstance);
+				OutputRegisters[NumOutputRegisters + CompRegisterOffset] = bNoOutputRegisters ? nullptr : (uint8*)CurrData().GetInstancePtrInt32(CompBufferOffset, StartInstance);
+			}
+			NumInputRegisters += NumComponents;
+			NumOutputRegisters += NumComponents;
+
+			return true;
+		}
+		return false;
+	}
+
+	void SetShaderParams(class FNiagaraShader *Shader, FRHICommandList &CommandList);
+	void UnsetShaderParams(class FNiagaraShader *Shader, FRHICommandList &CommandList);
+	FORCEINLINE void Allocate(uint32 NumInstances, ENiagaraSimTarget Target = ENiagaraSimTarget::CPUSim, bool bMaintainExisting=false)
+	{
+		check(bFinalized);
+		CurrData().Allocate(NumInstances, Target, bMaintainExisting);
+	}
+
+	void InitGPUFromCPU()
+	{
+		ensure(IsInGameThread());
+		PrevData().InitGPUFromCPU();
+	}
+
+	FORCEINLINE void Tick(ENiagaraSimTarget SimTarget = ENiagaraSimTarget::CPUSim)
+	{
+		SwapBuffers(SimTarget);
+	}
+	FORCEINLINE void TickRenderThread(ENiagaraSimTarget SimTarget = ENiagaraSimTarget::CPUSim)
+	{
+		SwapBuffersRenderThread(SimTarget);
+	}
+
+
+
+	// called before rendering to make sure we access the correct buffer
+	FORCEINLINE void ValidateBufferIndices()
+	{
+		CurrRenderBuffer = CurrBuffer;
+	}
+
+	FORCEINLINE void CopyPrevToCur()
+	{
+		PrevData().CopyTo(CurrData());
+	}
+
+	FORCEINLINE FNiagaraDataSetID GetID()const { return ID; }
+	FORCEINLINE void SetID(FNiagaraDataSetID InID) { ID = InID; }
+
+	FORCEINLINE FNiagaraDataBuffer& CurrData() 
+	{
+		//check(IsInGameThread());
+		return Data[CurrBuffer]; 
+	}
+	FORCEINLINE FNiagaraDataBuffer& PrevData() 
+	{ 
+		//check(IsInGameThread());
+		return Data[CurrBuffer > 0 ? CurrBuffer-1 : MaxBufferIdx];
+	}
+	FORCEINLINE const FNiagaraDataBuffer& CurrData()const 
+	{ 
+		//check(IsInGameThread());
+		return Data[CurrBuffer];
+	}
+	FORCEINLINE const FNiagaraDataBuffer& PrevData()const 
+	{ 
+		//check(IsInGameThread());
+		return Data[CurrBuffer > 0 ? CurrBuffer - 1 : MaxBufferIdx];
+	}
+
+	FORCEINLINE FNiagaraDataBuffer& CurrDataRender() 
+	{ 
+		check(!IsInGameThread());
+		return Data[CurrRenderBuffer];
+	}
+	FORCEINLINE FNiagaraDataBuffer& PrevDataRender() 
+	{ 
+		check(!IsInGameThread());
+		return Data[CurrRenderBuffer > 0 ? CurrRenderBuffer - 1 : MaxBufferIdx];
+	}
+	FORCEINLINE const FNiagaraDataBuffer& CurrDataRender() const 
+	{ 
+		check(!IsInGameThread());
+		return Data[CurrRenderBuffer];
+	}
+	FORCEINLINE const FNiagaraDataBuffer& PrevDataRender() const 
+	{
+		check(!IsInGameThread());
+		return Data[CurrRenderBuffer > 0 ? CurrRenderBuffer - 1 : MaxBufferIdx];
+	}
+
+	FORCEINLINE uint32 GetNumInstances()const { return CurrData().GetNumInstances(); }
+	FORCEINLINE uint32 GetNumInstancesAllocated()const { return CurrData().GetNumInstancesAllocated(); }
+	FORCEINLINE void SetNumInstances(uint32 InNumInstances) { CurrData().SetNumInstances(InNumInstances); }
+
+	FORCEINLINE void ResetNumInstances()
+	{
+		CurrData().SetNumInstances(0);
+		PrevData().SetNumInstances(0);
+	}
+
+	FORCEINLINE void ResetBuffers()
+	{
+		CurrData().Reset();
+		PrevData().Reset();
+	}
+
+	FORCEINLINE uint32 GetPrevNumInstances()const { return PrevData().GetNumInstances(); }
+
+	FORCEINLINE uint32 GetNumVariables()const { return Variables.Num(); }
+
+	FORCEINLINE uint32 GetSizeBytes()const
+	{
+		return Data[0].GetSizeBytes() + Data[1].GetSizeBytes();
+	}
+
+	FORCEINLINE bool HasVariable(const FNiagaraVariable& Var)const
+	{
+		return Variables.Contains(Var);
+	}
+
+	FORCEINLINE const FNiagaraVariableLayoutInfo* GetVariableLayout(const FNiagaraVariable& Var)const
+	{
+		return VariableLayoutMap.Find(Var);
+	}
+
+	// get the float and int component offsets of a variable; if the variable doesn't exist, returns -1
+	//
+	FORCEINLINE const bool GetVariableComponentOffsets(const FNiagaraVariable& Var, int32 &FloatStart, int32 &IntStart) const
+	{
+		const FNiagaraVariableLayoutInfo *Info = GetVariableLayout(Var);
+		if (Info)
+		{
+			FloatStart = Info->FloatComponentStart;
+			IntStart = Info->Int32ComponentStart;
+			return true;
+		}
+
+		FloatStart = -1;
+		IntStart = -1;
+		return false;
+	}
+
+	void Dump(bool bCurr, int32 StartIdx = 0, int32 NumInstances = INDEX_NONE);
+	void Dump(FNiagaraDataSet& Other, bool bCurr);
+	FORCEINLINE const TArray<FNiagaraVariable> &GetVariables() { return Variables; }
+
+	// called before dispatch from NiagaraEmitterInstanceBatcher
+	FRWBuffer &SetupDataSetIndices() 
+	{ 
+		check(IsInRenderingThread());
+		if (DataSetIndices.Buffer)
+		{
+			DataSetIndices.Release();
+		}
+		DataSetIndices.Initialize(sizeof(int32), 64 /*Context->NumDataSets*/, EPixelFormat::PF_R32_SINT);	// always allocate for up to 64 data sets
+		return DataSetIndices; 
+	}
+
+	FRWBuffer &GetDataSetIndices()
+	{
+		return DataSetIndices;
+	}
+
+private:
+
+	FORCEINLINE void SwapBuffers(ENiagaraSimTarget SimTarget)
+	{
+		if (SimTarget == ENiagaraSimTarget::CPUSim)
+		{
+			MaxBufferIdx = 2;
+			CurrBuffer = CurrBuffer < 2 ? CurrBuffer+1 : 0;
+		}
+		else
+		{
+			MaxBufferIdx = 1;
+			CurrBuffer = CurrBuffer == 0 ? 1 : 0;
+		}
+
+		ENQUEUE_UNIQUE_RENDER_COMMAND_TWOPARAMETER(SwapDataSetBuffersGPU,
+			uint32, CurrBuffer, CurrBuffer,
+			uint32 &, CurrRenderBuffer, CurrRenderBuffer,
+			{
+				CurrRenderBuffer = CurrBuffer;
+			});
+	}
+
+	FORCEINLINE void SwapBuffersRenderThread(ENiagaraSimTarget SimTarget)
+	{
+		check(IsInRenderingThread());
+		if (SimTarget == ENiagaraSimTarget::CPUSim)
+		{
+			MaxBufferIdx = 2;
+			CurrBuffer = CurrBuffer < 2 ? CurrBuffer + 1 : 0;
+		}
+		else
+		{
+			MaxBufferIdx = 1;
+			CurrBuffer = CurrBuffer == 0 ? 1 : 0;
+		}
+
+		CurrRenderBuffer = CurrBuffer;
+	}
+
+
+	void BuildLayout()
+	{
+		VariableLayoutMap.Empty();
+		TotalFloatComponents = 0;
+		TotalInt32Components = 0;
+
+		for (FNiagaraVariable& Var : Variables)
+		{
+			FNiagaraVariableLayoutInfo& VarInfo = VariableLayoutMap.Add(Var);
+			FNiagaraTypeLayoutInfo::GenerateLayoutInfo(VarInfo.LayoutInfo, Var.GetType().GetScriptStruct());
+			VarInfo.FloatComponentStart = TotalFloatComponents;
+			VarInfo.Int32ComponentStart = TotalInt32Components;
+			TotalFloatComponents += VarInfo.GetNumFloatComponents();
+			TotalInt32Components += VarInfo.GetNumInt32Components();
+		}
+		Data[0].Init(this);
+		Data[1].Init(this);
+		Data[2].Init(this);
+	}
+
+	FORCEINLINE uint32 GetNumFloatComponents() { return TotalFloatComponents; }
+	FORCEINLINE uint32 GetNumInt32Components() { return TotalInt32Components; }
+		
+	/** Unique ID for this data set. Used to allow referencing from other emitters and Systems. */
+	FNiagaraDataSetID ID;
+	/** Variables in the data set. */
+	TArray<FNiagaraVariable> Variables;
+	/** Map from the variable to some extra data describing it's layout in the data set. */
+	TMap<FNiagaraVariable, FNiagaraVariableLayoutInfo> VariableLayoutMap;
+	/** Total number of components of each type in the data set. */
+	uint32 TotalFloatComponents;
+	uint32 TotalInt32Components;
+	/** Index of current state data. */
+	uint32 CurrBuffer;
+	uint32 CurrRenderBuffer;
+	uint32 MaxBufferIdx;
+
+	/** Once finalized, the data layout etc is built and no more variables can be added. */
+	bool bFinalized;
+
+	FNiagaraDataBuffer Data[3];
+	FRWBuffer DataSetIndices; 
+};
+
+/**
+General iterator for getting and setting data in and FNiagaraDataSet.
+*/
+struct FNiagaraDataSetAccessorBase
+{
+	FNiagaraDataSetAccessorBase()
+		: DataSet(nullptr)
+		, DataBuffer(nullptr)
+		, VarLayout(nullptr)
+	{}
+
+	FNiagaraDataSetAccessorBase(FNiagaraDataSet* InDataSet, FNiagaraVariable InVar, bool bCurrBuffer = true)
+		: DataSet(InDataSet)
+		, DataBuffer(bCurrBuffer ? &DataSet->CurrData() : &DataSet->PrevData())
+	{
+		VarLayout = DataSet->GetVariableLayout(InVar);
+	}
+
+	void Create(FNiagaraDataSet* InDataSet, FNiagaraVariable InVar)
+	{
+		DataSet = InDataSet;
+		VarLayout = DataSet->GetVariableLayout(InVar);
+	}
+
+	void InitForAccess(bool bCurrBuffer = true)
+	{
+		DataBuffer = bCurrBuffer ? &DataSet->CurrData() : &DataSet->PrevData();
+	}
+
+	FORCEINLINE bool IsValid()const { return VarLayout != nullptr && DataBuffer; }
+protected:
+
+	FNiagaraDataSet* DataSet;
+	FNiagaraDataBuffer* DataBuffer;
+	const FNiagaraVariableLayoutInfo* VarLayout;
+};
+
+template<typename T>
+struct FNiagaraDataSetAccessor : public FNiagaraDataSetAccessorBase
+{
+	FNiagaraDataSetAccessor<T>() {}
+	FNiagaraDataSetAccessor(FNiagaraDataSet& InDataSet, FNiagaraVariable InVar, bool bCurrBuffer=true)
+		: FNiagaraDataSetAccessorBase(&InDataSet, InVar, bCurrBuffer)
+	{
+		check(sizeof(T) == InVar.GetType().GetSize());
+		checkf(false, TEXT("You must provide a fast runtime specialization for this type."));// Allow this slow generic version?
+	}
+
+	FORCEINLINE T operator[](int32 Index)const
+	{
+		return Get(Index);
+	}
+
+	FORCEINLINE T GetSafe(int32 Index, T Default)const
+	{
+		if (IsValid())
+		{
+			return Get(Index);
+		}
+
+		return Default;
+	}
+
+	FORCEINLINE T Get(int32 Index)const
+	{
+		T Ret;
+		Get(Index, Ret);
+		return Ret;
+	}
+
+	FORCEINLINE void Get(int32 Index, T& OutValue)const
+	{
+		uint8* ValuePtr = (uint8*)&OutValue;
+
+		for (uint32 CompIdx = 0; CompIdx < VarLayout->GetNumFloatComponents(); ++CompIdx)
+		{
+			uint32 CompBufferOffset = VarLayout->FloatComponentStart + CompIdx;
+			float* Src = DataBuffer->GetInstancePtrFloat(CompBufferOffset, Index);
+			float* Dst = (float*)(ValuePtr + VarLayout->LayoutInfo.FloatComponentByteOffsets[CompIdx]);
+			*Dst = *Src;
+		}
+
+		for (uint32 CompIdx = 0; CompIdx < VarLayout->GetNumInt32Components(); ++CompIdx)
+		{
+			uint32 CompBufferOffset = VarLayout->FloatComponentStart + CompIdx;
+			int32* Src = DataBuffer->GetInstancePtrInt32(CompBufferOffset, Index);
+			int32* Dst = (int32*)(ValuePtr + VarLayout->LayoutInfo.FloatComponentByteOffsets[CompIdx]);
+			*Dst = *Src;
+		}
+	}
+
+	FORCEINLINE void Set(int32 Index, const T& InValue)
+	{
+		uint8* ValuePtr = (uint8*)&InValue;
+
+		for (uint32 CompIdx = 0; CompIdx < VarLayout->GetNumFloatComponents(); ++CompIdx)
+		{
+			uint32 CompBufferOffset = VarLayout->FloatComponentStart + CompIdx;
+			float* Dst = DataBuffer->GetInstancePtrFloat(CompBufferOffset, Index);
+			float* Src = (float*)(ValuePtr + VarLayout->LayoutInfo.FloatComponentByteOffsets[CompIdx]);
+			*Dst = *Src;
+		}
+
+		for (uint32 CompIdx = 0; CompIdx < VarLayout->GetNumInt32Components(); ++CompIdx)
+		{
+			uint32 CompBufferOffset = VarLayout->FloatComponentStart + CompIdx;
+			int32* Dst = DataBuffer->GetInstancePtrInt32(CompBufferOffset, Index);
+			int32* Src = (int32*)(ValuePtr + VarLayout->LayoutInfo.FloatComponentByteOffsets[CompIdx]);
+			*Dst = *Src;
+		}
+	}
+};
+
+template<>
+struct FNiagaraDataSetAccessor<FNiagaraBool> : public FNiagaraDataSetAccessorBase
+{
+	FNiagaraDataSetAccessor<FNiagaraBool>() {}
+	FNiagaraDataSetAccessor<FNiagaraBool>(FNiagaraDataSet& InDataSet, FNiagaraVariable InVar, bool bCurrBuffer = true)
+		: FNiagaraDataSetAccessorBase(&InDataSet, InVar, bCurrBuffer)
+	{
+		check(sizeof(FNiagaraBool) == InVar.GetType().GetSize());
+		InitForAccess(bCurrBuffer);
+	}
+
+	void InitForAccess(bool bCurrBuffer = true)
+	{
+		DataBuffer = bCurrBuffer ? &DataSet->CurrData() : &DataSet->PrevData();
+		if (VarLayout != nullptr)
+		{
+			Base = (int32*)DataBuffer->GetComponentPtrInt32(VarLayout->Int32ComponentStart);
+		}
+		else
+		{
+			Base = nullptr;
+		}
+	}
+
+	FORCEINLINE FNiagaraBool operator[](int32 Index)const
+	{
+		return Get(Index);
+	}
+
+	FORCEINLINE FNiagaraBool Get(int32 Index)const
+	{
+		FNiagaraBool Ret;
+		Get(Index, Ret);
+		return Ret;
+	}
+
+	FORCEINLINE FNiagaraBool GetSafe(int32 Index, bool Default = true)const
+	{
+		if (IsValid())
+		{
+			return Get(Index);
+		}
+
+		return Default;
+	}
+
+	FORCEINLINE void Get(int32 Index, FNiagaraBool& OutValue)const
+	{
+		OutValue.Value = Base[Index];
+	}
+
+	FORCEINLINE void Set(int32 Index, const FNiagaraBool& InValue)
+	{
+		Base[Index] = InValue.Value;
+	}
+
+private:
+
+	int32* Base;
+};
+
+template<>
+struct FNiagaraDataSetAccessor<int32> : public FNiagaraDataSetAccessorBase
+{
+	FNiagaraDataSetAccessor<int32>() {}
+	FNiagaraDataSetAccessor<int32>(FNiagaraDataSet& InDataSet, FNiagaraVariable InVar, bool bCurrBuffer = true)
+		: FNiagaraDataSetAccessorBase(&InDataSet, InVar, bCurrBuffer)
+	{
+		check(sizeof(int32) == InVar.GetType().GetSize());
+		InitForAccess(bCurrBuffer);
+	}
+	
+	void InitForAccess(bool bCurrBuffer = true)
+	{
+		DataBuffer = bCurrBuffer ? &DataSet->CurrData() : &DataSet->PrevData();
+		if (VarLayout != nullptr)
+		{
+			Base = (int32*)DataBuffer->GetComponentPtrInt32(VarLayout->Int32ComponentStart);
+			if (DataBuffer->GetNumInstances() != 0)
+			{
+				check(Base != nullptr);
+			}
+		}
+		else
+		{
+			Base = nullptr;
+		}
+	}
+
+	FORCEINLINE int32 operator[](int32 Index)const
+	{
+		return Get(Index);
+	}
+
+	FORCEINLINE int32 Get(int32 Index)const
+	{
+		int32 Ret;
+		Get(Index, Ret);
+		return Ret;
+	}
+
+	FORCEINLINE int32 GetSafe(int32 Index, int32 Default = 0)const
+	{
+		if (IsValid())
+		{
+			return Get(Index);
+		}
+
+		return Default;
+	}
+
+	FORCEINLINE void Get(int32 Index, int32& OutValue)const
+	{
+		check(Base != nullptr);
+		OutValue = Base[Index];
+	}
+
+	FORCEINLINE void Set(int32 Index, const int32& InValue)
+	{
+		check(Base != nullptr);
+		Base[Index] = InValue;
+	}
+
+private:
+
+	int32* Base;
+};
+
+template<>
+struct FNiagaraDataSetAccessor<float> : public FNiagaraDataSetAccessorBase
+{
+	FNiagaraDataSetAccessor<float>() {}
+	FNiagaraDataSetAccessor<float>(FNiagaraDataSet& InDataSet, FNiagaraVariable InVar, bool bCurrBuffer=true)
+		: FNiagaraDataSetAccessorBase(&InDataSet, InVar, bCurrBuffer)
+	{
+		check(sizeof(float) == InVar.GetType().GetSize());
+		InitForAccess(bCurrBuffer);
+	}
+
+	void InitForAccess(bool bCurrBuffer = true)
+	{
+		DataBuffer = bCurrBuffer ? &DataSet->CurrData() : &DataSet->PrevData();
+		if (VarLayout != nullptr)
+		{
+			Base = (float*)DataBuffer->GetComponentPtrFloat(VarLayout->FloatComponentStart);
+		}
+		else
+		{
+			Base = nullptr;
+		}
+	}
+
+	FORCEINLINE float operator[](int32 Index)const
+	{
+		return Get(Index);
+	}
+
+	FORCEINLINE float GetSafe(int32 Index, float Default = 0.0f)const
+	{
+		if (IsValid())
+		{
+			return Get(Index);
+		}
+
+		return Default;
+	}
+
+	FORCEINLINE float Get(int32 Index)const
+	{
+		float Ret;
+		Get(Index, Ret);
+		return Ret;
+	}
+
+	FORCEINLINE void Get(int32 Index, float& OutValue)const
+	{
+		OutValue = Base[Index];
+	}
+
+	FORCEINLINE void Set(int32 Index, const float& InValue)
+	{
+		Base[Index] = InValue;
+	}
+
+private:
+	float* Base;
+};
+
+template<>
+struct FNiagaraDataSetAccessor<FVector2D> : public FNiagaraDataSetAccessorBase
+{
+	FNiagaraDataSetAccessor<FVector2D>() {}
+	FNiagaraDataSetAccessor<FVector2D>(FNiagaraDataSet& InDataSet, FNiagaraVariable InVar, bool bCurrBuffer=true)
+		: FNiagaraDataSetAccessorBase(&InDataSet, InVar, bCurrBuffer)
+	{
+		check(sizeof(FVector2D) == InVar.GetType().GetSize());
+		InitForAccess(bCurrBuffer);
+	}
+
+	void InitForAccess(bool bCurrBuffer = true)
+	{
+		DataBuffer = bCurrBuffer ? &DataSet->CurrData() : &DataSet->PrevData();
+		if (VarLayout != nullptr)
+		{
+			XBase = (float*)DataBuffer->GetComponentPtrFloat(VarLayout->FloatComponentStart);
+			YBase = (float*)DataBuffer->GetComponentPtrFloat(VarLayout->FloatComponentStart + 1);
+		}
+		else
+		{
+			XBase = nullptr;
+			YBase = nullptr;
+		}
+	}
+
+	FORCEINLINE FVector2D operator[](int32 Index)const
+	{
+		return Get(Index);
+	}
+
+	FORCEINLINE FVector2D GetSafe(int32 Index, FVector2D Default = FVector2D::ZeroVector)const
+	{
+		if (IsValid())
+		{
+			return Get(Index);
+		}
+
+		return Default;
+	}
+
+	FORCEINLINE FVector2D Get(int32 Index)const
+	{
+		FVector2D Ret;
+		Get(Index, Ret);
+		return Ret;
+	}
+
+	FORCEINLINE void Get(int32 Index, FVector2D& OutValue)const
+	{
+		OutValue.X = XBase[Index];
+		OutValue.Y = YBase[Index];
+	}
+
+	FORCEINLINE void Set(int32 Index, const FVector2D& InValue)
+	{
+		XBase[Index] = InValue.X;
+		YBase[Index] = InValue.Y;
+	}
+
+private:
+
+	float* XBase;
+	float* YBase;
+};
+
+template<>
+struct FNiagaraDataSetAccessor<FVector> : public FNiagaraDataSetAccessorBase
+{
+	FNiagaraDataSetAccessor<FVector>() {}
+	FNiagaraDataSetAccessor<FVector>(FNiagaraDataSet& InDataSet, FNiagaraVariable InVar, bool bCurrBuffer=true)
+		: FNiagaraDataSetAccessorBase(&InDataSet, InVar, bCurrBuffer)
+	{
+		check(sizeof(FVector) == InVar.GetType().GetSize());
+		InitForAccess(bCurrBuffer);
+	}
+
+	void InitForAccess(bool bCurrBuffer = true)
+	{
+		DataBuffer = bCurrBuffer ? &DataSet->CurrData() : &DataSet->PrevData();
+		if (VarLayout != nullptr)
+		{
+			XBase = (float*)DataBuffer->GetComponentPtrFloat(VarLayout->FloatComponentStart);
+			YBase = (float*)DataBuffer->GetComponentPtrFloat(VarLayout->FloatComponentStart + 1);
+			ZBase = (float*)DataBuffer->GetComponentPtrFloat(VarLayout->FloatComponentStart + 2);
+		}
+		else
+		{
+			XBase = nullptr;
+			YBase = nullptr;
+			ZBase = nullptr;
+		}
+	}
+
+	FORCEINLINE FVector operator[](int32 Index)const
+	{
+		return Get(Index);
+	}
+
+	FORCEINLINE FVector GetSafe(int32 Index, FVector Default = FVector::ZeroVector)const
+	{
+		if (IsValid())
+		{
+			return Get(Index);
+		}
+
+		return Default;
+	}
+
+	FORCEINLINE FVector Get(int32 Index)const
+	{
+		FVector Ret;
+		Get(Index, Ret);
+		return Ret;
+	}
+
+	FORCEINLINE void Get(int32 Index, FVector& OutValue)const
+	{
+		OutValue.X = XBase[Index];
+		OutValue.Y = YBase[Index];
+		OutValue.Z = ZBase[Index];
+	}
+
+	FORCEINLINE void Set(int32 Index, const FVector& InValue)
+	{
+		XBase[Index] = InValue.X;
+		YBase[Index] = InValue.Y;
+		ZBase[Index] = InValue.Z;
+	}
+
+private:
+
+	float* XBase;
+	float* YBase;
+	float* ZBase;
+};
+
+template<>
+struct FNiagaraDataSetAccessor<FVector4> : public FNiagaraDataSetAccessorBase
+{
+	FNiagaraDataSetAccessor<FVector4>() {}
+	FNiagaraDataSetAccessor<FVector4>(FNiagaraDataSet& InDataSet, FNiagaraVariable InVar, bool bCurrBuffer=true)
+		: FNiagaraDataSetAccessorBase(&InDataSet, InVar, bCurrBuffer)
+	{
+		check(sizeof(FVector4) == InVar.GetType().GetSize());
+		InitForAccess(bCurrBuffer);
+	}
+
+	void InitForAccess(bool bCurrBuffer = true)
+	{
+		DataBuffer = bCurrBuffer ? &DataSet->CurrData() : &DataSet->PrevData();
+		if (VarLayout)
+		{
+			XBase = (float*)DataBuffer->GetComponentPtrFloat(VarLayout->FloatComponentStart);
+			YBase = (float*)DataBuffer->GetComponentPtrFloat(VarLayout->FloatComponentStart + 1);
+			ZBase = (float*)DataBuffer->GetComponentPtrFloat(VarLayout->FloatComponentStart + 2);
+			WBase = (float*)DataBuffer->GetComponentPtrFloat(VarLayout->FloatComponentStart + 3);
+		}
+		else
+		{
+			XBase = nullptr;
+			YBase = nullptr;
+			ZBase = nullptr;
+			WBase = nullptr;
+		}
+	}
+
+	FORCEINLINE FVector4 operator[](int32 Index)const
+	{
+		return Get(Index);
+	}
+
+	FORCEINLINE FVector4 GetSafe(int32 Index, const FVector4& Default = FVector4(0.0f,0.0f,0.0f,0.0f))const
+	{
+		if (IsValid())
+		{
+			return Get(Index);
+		}
+
+		return Default;
+	}
+
+	FORCEINLINE FVector4 Get(int32 Index)const
+	{
+		FVector4 Ret;
+		Get(Index, Ret);
+		return Ret;
+	}
+
+	FORCEINLINE void Get(int32 Index, FVector4& OutValue)const
+	{
+		OutValue.X = XBase[Index];
+		OutValue.Y = YBase[Index];
+		OutValue.Z = ZBase[Index];
+		OutValue.W = WBase[Index];
+	}
+
+	FORCEINLINE void Set(int32 Index, const FVector4& InValue)
+	{
+		XBase[Index] = InValue.X;
+		YBase[Index] = InValue.Y;
+		ZBase[Index] = InValue.Z;
+		WBase[Index] = InValue.W;
+	}
+
+private:
+
+	float* XBase;
+	float* YBase;
+	float* ZBase;
+	float* WBase;
+};
+
+template<>
+struct FNiagaraDataSetAccessor<FLinearColor> : public FNiagaraDataSetAccessorBase
+{
+	FNiagaraDataSetAccessor<FLinearColor>() {}
+	FNiagaraDataSetAccessor<FLinearColor>(FNiagaraDataSet& InDataSet, FNiagaraVariable InVar, bool bCurrBuffer=true)
+		: FNiagaraDataSetAccessorBase(&InDataSet, InVar, bCurrBuffer)
+	{
+		check(sizeof(FLinearColor) == InVar.GetType().GetSize());
+		InitForAccess(bCurrBuffer);
+	}
+
+	void InitForAccess(bool bCurrBuffer = true)
+	{
+		DataBuffer = bCurrBuffer ? &DataSet->CurrData() : &DataSet->PrevData();
+		if (VarLayout)
+		{
+			RBase = (float*)DataBuffer->GetComponentPtrFloat(VarLayout->FloatComponentStart);
+			GBase = (float*)DataBuffer->GetComponentPtrFloat(VarLayout->FloatComponentStart + 1);
+			BBase = (float*)DataBuffer->GetComponentPtrFloat(VarLayout->FloatComponentStart + 2);
+			ABase = (float*)DataBuffer->GetComponentPtrFloat(VarLayout->FloatComponentStart + 3);
+		}
+		else
+		{
+			RBase = nullptr;
+			GBase = nullptr;
+			BBase = nullptr;
+			ABase = nullptr;
+		}
+	}
+
+	FORCEINLINE FLinearColor operator[](int32 Index)const
+	{
+		return Get(Index);
+	}
+
+	FORCEINLINE FLinearColor GetSafe(int32 Index, FLinearColor Default = FLinearColor::White)const
+	{
+		if (IsValid())
+		{
+			return Get(Index);
+		}
+
+		return Default;
+	}
+
+	FORCEINLINE FLinearColor Get(int32 Index)const
+	{
+		FLinearColor Ret;
+		Get(Index, Ret);
+		return Ret;
+	}
+
+	FORCEINLINE void Get(int32 Index, FLinearColor& OutValue)const
+	{
+		OutValue.R = RBase[Index];
+		OutValue.G = GBase[Index];
+		OutValue.B = BBase[Index];
+		OutValue.A = ABase[Index];
+	}
+
+	FORCEINLINE void Set(int32 Index, const FLinearColor& InValue)
+	{
+		RBase[Index] = InValue.R;
+		GBase[Index] = InValue.G;
+		BBase[Index] = InValue.B;
+		ABase[Index] = InValue.A;
+	}
+
+private:
+
+	float* RBase;
+	float* GBase;
+	float* BBase;
+	float* ABase;
+};
+
+template<>
+struct FNiagaraDataSetAccessor<FNiagaraSpawnInfo> : public FNiagaraDataSetAccessorBase
+{
+	FNiagaraDataSetAccessor<FNiagaraSpawnInfo>() {}
+	FNiagaraDataSetAccessor<FNiagaraSpawnInfo>(FNiagaraDataSet& InDataSet, FNiagaraVariable InVar, bool bCurrBuffer = true)
+		: FNiagaraDataSetAccessorBase(&InDataSet, InVar, bCurrBuffer)
+	{
+		check(sizeof(FNiagaraSpawnInfo) == InVar.GetType().GetSize());
+		InitForAccess(bCurrBuffer);
+	}
+
+	void InitForAccess(bool bCurrBuffer = true)
+	{
+		DataBuffer = bCurrBuffer ? &DataSet->CurrData() : &DataSet->PrevData();
+		if (VarLayout != nullptr)
+		{
+			CountBase = (int32*)DataBuffer->GetComponentPtrInt32(VarLayout->Int32ComponentStart);
+			InterpStartDtBase = (float*)DataBuffer->GetComponentPtrFloat(VarLayout->FloatComponentStart);
+			IntervalDtBase = (float*)DataBuffer->GetComponentPtrFloat(VarLayout->FloatComponentStart + 1);
+		}
+		else
+		{
+			CountBase = nullptr;
+			InterpStartDtBase = nullptr;
+			IntervalDtBase = nullptr;
+		}
+	}
+
+	FORCEINLINE FNiagaraSpawnInfo operator[](int32 Index)const
+	{
+		return Get(Index);
+	}
+
+	FORCEINLINE FNiagaraSpawnInfo GetSafe(int32 Index, FNiagaraSpawnInfo Default = FNiagaraSpawnInfo())const
+	{
+		if (IsValid())
+		{
+			return Get(Index);
+		}
+
+		return Default;
+	}
+
+	FORCEINLINE FNiagaraSpawnInfo Get(int32 Index)const
+	{
+		FNiagaraSpawnInfo Ret;
+		Get(Index, Ret);
+		return Ret;
+	}
+
+	FORCEINLINE void Get(int32 Index, FNiagaraSpawnInfo& OutValue)const
+	{
+		OutValue.Count = CountBase[Index];
+		OutValue.InterpStartDt = InterpStartDtBase[Index];
+		OutValue.IntervalDt = IntervalDtBase[Index];
+	}
+
+	FORCEINLINE void Set(int32 Index, const FNiagaraSpawnInfo& InValue)
+	{
+		CountBase[Index] = InValue.Count;
+		InterpStartDtBase[Index] = InValue.InterpStartDt;
+		IntervalDtBase[Index] = InValue.IntervalDt;
+	}
+
+private:
+
+	int32* CountBase;
+	float* InterpStartDtBase;
+	float* IntervalDtBase;
+};
+
+//Provide iterator to keep iterator access patterns still in use.
+template<typename T>
+struct FNiagaraDataSetIterator : public FNiagaraDataSetAccessor<T>
+{
+	FNiagaraDataSetIterator<T>() {}
+	FNiagaraDataSetIterator(FNiagaraDataSet& InDataSet, FNiagaraVariable InVar, uint32 StartIndex = 0, bool bCurrBuffer = true)
+		: FNiagaraDataSetAccessor<T>(InDataSet, InVar, bCurrBuffer)
+		, CurrIdx(StartIndex)
+	{}
+
+	FORCEINLINE T operator*() { return Get(); }
+	FORCEINLINE T Get()const
+	{
+		T Ret;
+		Get(Ret);
+		return Ret;
+	}
+	FORCEINLINE void Get(T& OutValue)const { FNiagaraDataSetAccessor<T>::Get(CurrIdx, OutValue); }
+
+	FORCEINLINE void Set(const T& InValue)
+	{
+		FNiagaraDataSetAccessor<T>::Set(CurrIdx, InValue);
+	}
+
+	FORCEINLINE void Advance() { ++CurrIdx; }
+	FORCEINLINE bool IsValid()const
+	{
+		return FNiagaraDataSetAccessorBase::VarLayout != nullptr && CurrIdx < FNiagaraDataSetAccessorBase::DataBuffer->GetNumInstances();
+	}
+
+	uint32 GetCurrIndex()const { return CurrIdx; }
+
+private:
+	uint32 CurrIdx;
+};
+
+/**
+Iterator that will pull or push data between a DataSet and some FNiagaraVariables it contains.
+Super slow. Don't use at runtime.
+*/
+struct FNiagaraDataSetVariableIterator
+{
+	FNiagaraDataSetVariableIterator(FNiagaraDataSet& InDataSet, uint32 StartIdx = 0, bool bCurrBuffer = true)
+		: DataSet(InDataSet)
+		, DataBuffer(bCurrBuffer ? DataSet.CurrData() : DataSet.PrevData())
+		, CurrIdx(StartIdx)
+	{
+	}
+
+	void Get()
+	{
+		for (int32 VarIdx = 0; VarIdx < Variables.Num(); ++VarIdx)
+		{
+			FNiagaraVariable* Var = Variables[VarIdx];
+			const FNiagaraVariableLayoutInfo* Layout = VarLayouts[VarIdx];
+			check(Var && Layout);
+			uint8* ValuePtr = Var->GetData();
+
+			for (uint32 CompIdx = 0; CompIdx < Layout->GetNumFloatComponents(); ++CompIdx)
+			{
+				uint32 CompBufferOffset = Layout->FloatComponentStart + CompIdx;
+				float* Src = DataBuffer.GetInstancePtrFloat(CompBufferOffset, CurrIdx);
+				float* Dst = (float*)(ValuePtr + Layout->LayoutInfo.FloatComponentByteOffsets[CompIdx]);
+				*Dst = *Src;
+			}
+
+			for (uint32 CompIdx = 0; CompIdx < Layout->GetNumInt32Components(); ++CompIdx)
+			{
+				uint32 CompBufferOffset = Layout->Int32ComponentStart + CompIdx;
+				int32* Src = DataBuffer.GetInstancePtrInt32(CompBufferOffset, CurrIdx);
+				int32* Dst = (int32*)(ValuePtr + Layout->LayoutInfo.Int32ComponentByteOffsets[CompIdx]);
+				*Dst = *Src;
+			}
+		}
+	}
+
+	void Set()
+	{
+		for (int32 VarIdx = 0; VarIdx < Variables.Num(); ++VarIdx)
+		{
+			FNiagaraVariable* Var = Variables[VarIdx];
+			const FNiagaraVariableLayoutInfo* Layout = VarLayouts[VarIdx];
+
+			check(Var && Layout);
+			uint8* ValuePtr = Var->GetData();
+
+			for (uint32 CompIdx = 0; CompIdx < Layout->GetNumFloatComponents(); ++CompIdx)
+			{
+				uint32 CompBufferOffset = Layout->FloatComponentStart + CompIdx;
+				float* Dst = DataBuffer.GetInstancePtrFloat(CompBufferOffset, CurrIdx);
+				float* Src = (float*)(ValuePtr + Layout->LayoutInfo.FloatComponentByteOffsets[CompIdx]);
+				*Dst = *Src;
+			}
+
+			for (uint32 CompIdx = 0; CompIdx < Layout->GetNumInt32Components(); ++CompIdx)
+			{
+				uint32 CompBufferOffset = Layout->FloatComponentStart + CompIdx;
+				int32* Dst = DataBuffer.GetInstancePtrInt32(CompBufferOffset, CurrIdx);
+				int32* Src = (int32*)(ValuePtr + Layout->LayoutInfo.FloatComponentByteOffsets[CompIdx]);
+				*Dst = *Src;
+			}
+		}
+	}
+
+	void Advance() { ++CurrIdx; }
+
+	bool IsValid()const
+	{
+		return CurrIdx < DataBuffer.GetNumInstances();
+	}
+
+	uint32 GetCurrIndex()const { return CurrIdx; }
+
+	void AddVariable(FNiagaraVariable* InVar)
+	{
+		check(InVar);
+		Variables.AddUnique(InVar);
+		VarLayouts.AddUnique(DataSet.GetVariableLayout(*InVar));
+		InVar->AllocateData();
+	}
+
+	void AddVariables(TArray<FNiagaraVariable>& Vars)
+	{
+		for (FNiagaraVariable& Var : Vars)
+		{
+			AddVariable(&Var);
+		}
+	}
+private:
+
+	FNiagaraDataSet& DataSet;
+	FNiagaraDataBuffer& DataBuffer;
+	TArray<FNiagaraVariable*> Variables;
+	TArray<const FNiagaraVariableLayoutInfo*> VarLayouts;
+
+	uint32 CurrIdx;
+};
\ No newline at end of file
diff --git a/Engine/Plugins/FX/Niagara/Source/Niagara/Classes/NiagaraEmitter.h b/Engine/Plugins/FX/Niagara/Source/Niagara/Classes/NiagaraEmitter.h
new file mode 100644
--- /dev/null
+++ b/Engine/Plugins/FX/Niagara/Source/Niagara/Classes/NiagaraEmitter.h
@@ -0,0 +1,272 @@
+// Copyright 1998-2017 Epic Games, Inc. All Rights Reserved.
+
+#pragma once
+
+#include "NiagaraCommon.h"
+#include "CoreMinimal.h"
+#include "UObject/ObjectMacros.h"
+#include "UObject/Object.h"
+#include "NiagaraScript.h"
+#include "NiagaraCollision.h"
+#include "NiagaraEmitter.generated.h"
+
+class UMaterial;
+class UNiagaraEmitter;
+class UNiagaraEventReceiverEmitterAction;
+
+//TODO: Event action that spawns other whole Systems?
+//One that calls a BP exposed delegate?
+
+USTRUCT()
+struct FNiagaraEventReceiverProperties
+{
+	GENERATED_BODY()
+
+	FNiagaraEventReceiverProperties()
+	: Name(NAME_None)
+	, SourceEventGenerator(NAME_None)
+	, SourceEmitter(NAME_None)
+	{
+
+	}
+
+	FNiagaraEventReceiverProperties(FName InName, FName InEventGenerator, FName InSourceEmitter)
+		: Name(InName)
+		, SourceEventGenerator(InEventGenerator)
+		, SourceEmitter(InSourceEmitter)
+	{
+
+	}
+
+	/** The name of this receiver. */
+	UPROPERTY(EditAnywhere, Category = "Event Receiver")
+	FName Name;
+
+	/** The name of the EventGenerator to bind to. */
+	UPROPERTY(EditAnywhere, Category = "Event Receiver")
+	FName SourceEventGenerator;
+
+	/** The name of the emitter from which the Event Generator is taken. */
+	UPROPERTY(EditAnywhere, Category = "Event Receiver")
+	FName SourceEmitter;
+
+	UPROPERTY(EditAnywhere, Instanced, Category = "Event Receiver")
+	TArray<UNiagaraEventReceiverEmitterAction*> EmitterActions;
+};
+
+USTRUCT()
+struct FNiagaraEventGeneratorProperties
+{
+	GENERATED_BODY()
+
+	FNiagaraEventGeneratorProperties()
+	: MaxEventsPerFrame(64)
+	{
+
+	}
+
+	FNiagaraEventGeneratorProperties(FNiagaraDataSetProperties &Props, FName InEventGenerator, FName InSourceEmitter)
+		: ID(Props.ID.Name)
+		, SourceEmitter(InSourceEmitter)
+		, SetProps(Props)		
+	{
+
+	}
+
+	/** Max Number of Events that can be generated per frame. */
+	UPROPERTY(EditAnywhere, Category = "Event Receiver")
+	int32 MaxEventsPerFrame; //TODO - More complex allocation so that we can grow dynamically if more space is needed ?
+
+	FName ID;
+	FName SourceEmitter;
+
+	UPROPERTY()
+	FNiagaraDataSetProperties SetProps;
+};
+
+
+UENUM()
+enum class EScriptExecutionMode : uint8
+{
+	/** The event script is run on every existing particle in the emitter.*/
+	EveryParticle = 0,
+	/** The event script is run only on the particles that were spawned in response to the current event in the emitter.*/
+	SpawnedParticles,
+	/** The event script is run only on the particle whose int32 ParticleIndex is specified in the event payload.*/
+	SingleParticle
+};
+
+UENUM()
+enum class EScriptCompileIndices : uint8
+{
+	SpawnScript = 0,
+	UpdateScript,
+	EventScript
+};
+
+USTRUCT()
+struct FNiagaraEmitterScriptProperties
+{
+	FNiagaraEmitterScriptProperties() : Script(nullptr)
+	{
+
+	}
+
+	GENERATED_BODY()
+	
+ 	UPROPERTY()
+ 	UNiagaraScript *Script;
+
+	UPROPERTY()
+	TArray<FNiagaraEventReceiverProperties> EventReceivers;
+
+	UPROPERTY()
+	TArray<FNiagaraEventGeneratorProperties> EventGenerators;
+
+	NIAGARA_API void InitDataSetAccess();
+};
+
+USTRUCT()
+struct FNiagaraEventScriptProperties : public FNiagaraEmitterScriptProperties
+{
+	GENERATED_BODY()
+			
+	FNiagaraEventScriptProperties() : FNiagaraEmitterScriptProperties()
+	{
+		ExecutionMode = EScriptExecutionMode::EveryParticle;
+		SpawnNumber = 0;
+		MaxEventsPerFrame = 0;
+	}
+	
+	/** Controls which particles have the event script run on them.*/
+	UPROPERTY(EditAnywhere, Category="Event Handler Options")
+	EScriptExecutionMode ExecutionMode;
+
+	/** Controls whether or not particles are spawned as a result of handling the event. Only valid for EScriptExecutionMode::SpawnedParticles*/
+	UPROPERTY(EditAnywhere, Category="Event Handler Options")
+	uint32 SpawnNumber;
+
+	/** Controls how many events are consumed by this event handler. If there are more events generated than this value, they will be ignored.*/
+	UPROPERTY(EditAnywhere, Category="Event Handler Options")
+	uint32 MaxEventsPerFrame;
+
+	/** Id of the Emitter Handle that generated the event. If all zeroes, the event generator is assumed to be this emitter.*/
+	UPROPERTY(EditAnywhere, Category="Event Handler Options")
+	FGuid SourceEmitterID;
+
+	/** The name of the event generated. This will be "Collision" for collision events and the Event Name field on the DataSetWrite node in the module graph for others.*/
+	UPROPERTY(EditAnywhere, Category="Event Handler Options")
+	FName SourceEventName;
+
+};
+
+/** Represents timed burst of particles in an emitter. */
+USTRUCT()
+struct FNiagaraEmitterBurst
+{
+	GENERATED_BODY()
+
+	/** The base time of the burst absolute emitter time. */
+	UPROPERTY(EditAnywhere, Category = "Burst")
+	float Time;
+
+	/** 
+	 * A range around the base time which can be used to randomize burst timing.  The resulting range is 
+	 * from Time - (TimeRange / 2) to Time + (TimeRange / 2).
+	 */
+	UPROPERTY(EditAnywhere, Category = "Burst")
+	float TimeRange;
+
+	/** The minimum number of particles to spawn. */
+	UPROPERTY(EditAnywhere, Category = "Burst")
+	uint32 SpawnMinimum;
+
+	/** The maximum number of particles to spawn. */
+	UPROPERTY(EditAnywhere, Category = "Burst")
+	uint32 SpawnMaximum;
+};
+
+/** 
+ *	UNiagaraEmitter stores the attributes of an FNiagaraEmitterInstance
+ *	that need to be serialized and are used for its initialization 
+ */
+UCLASS(MinimalAPI)
+class UNiagaraEmitter : public UObject
+{
+	GENERATED_UCLASS_BODY()
+public:
+	//Begin UObject Interface
+#if WITH_EDITOR
+	virtual void PostEditChangeProperty(struct FPropertyChangedEvent& PropertyChangedEvent) override;
+#endif
+	void Serialize(FArchive& Ar)override;
+	virtual void PostInitProperties() override;
+	virtual void PostLoad() override;
+	//End UObject Interface
+
+	UPROPERTY(EditAnywhere, Category = "Emitter")
+	bool bLocalSpace;
+
+	//UPROPERTY(EditAnywhere, Category = "Emitter")
+	//float StartTime;
+	//UPROPERTY(EditAnywhere, Category = "Emitter")
+	//float EndTime;
+	//UPROPERTY(EditAnywhere, Category = "Emitter")
+	//int32 NumLoops;
+	UPROPERTY(EditAnywhere, Category = "Emitter")
+	ENiagaraCollisionMode CollisionMode;
+	
+	UPROPERTY(EditAnywhere, Instanced, Category = "Render")
+	TArray<class UNiagaraRendererProperties*> RendererProperties;
+
+	UPROPERTY()
+	FNiagaraEmitterScriptProperties UpdateScriptProps;
+
+	UPROPERTY()
+	FNiagaraEmitterScriptProperties SpawnScriptProps;
+
+	UPROPERTY(EditAnywhere, Category="Event Handler")
+	TArray<FNiagaraEventScriptProperties> EventHandlerScriptProps;
+
+	/** When enabled, this will spawn using interpolated parameter values and perform a partial update at spawn time. This adds significant additional cost for spawning but will produce much smoother spawning for high spawn rates, erratic frame rates and fast moving emitters. */
+	UPROPERTY(EditAnywhere, Category = "Emitter")
+	uint32 bInterpolatedSpawning : 1;
+
+	UPROPERTY(EditAnywhere, Category = "Emitter")
+	ENiagaraSimTarget SimTarget;
+
+#if WITH_EDITORONLY_DATA
+	/** Adjusted every time that we compile this emitter. Lets us know that we might differ from any cached versions.*/
+	UPROPERTY()
+	FGuid ChangeId;
+
+	/** 'Source' data/graphs for the scripts used by this emitter. */
+	UPROPERTY()
+	class UNiagaraScriptSourceBase*	GraphSource;
+
+	void NIAGARA_API CompileScripts(TArray<ENiagaraScriptCompileStatus>& OutScriptStatuses, TArray<FString>& OutGraphLevelErrorMessages, TArray<FString>& PathNames, TArray<UNiagaraScript*>& Scripts);
+	ENiagaraScriptCompileStatus NIAGARA_API CompileScript(EScriptCompileIndices InScriptToCompile, FString& OutGraphLevelErrorMessages, int32 InSubScriptIdx = 0);
+
+	UNiagaraEmitter* UNiagaraEmitter::MakeRecursiveDeepCopy(UObject* DestOuter) const;
+	UNiagaraEmitter* MakeRecursiveDeepCopy(UObject* DestOuter, TMap<const UObject*, UObject*>& ExistingConversions) const;
+
+	void NIAGARA_API GetScripts(TArray<UNiagaraScript*>& OutScripts);
+
+	/** Data used by the editor to maintain UI state etc.. */
+	UPROPERTY()
+	UObject* EditorData;
+#endif
+
+	bool IsValid()const;
+
+	bool UsesScript(const UNiagaraScript* Script)const;
+	//bool UsesDataInterface(UNiagaraDataInterface* Interface);
+	bool UsesCollection(const class UNiagaraParameterCollection* Collection)const;
+
+	FString NIAGARA_API GetUniqueEmitterName()const;
+
+	/** Converts an emitter paramter "Emitter.XXXX" into it's real parameter name. */
+	FNiagaraVariable GetEmitterParameter(const FNiagaraVariable& EmitterVar)const;
+};
+
+
diff --git a/Engine/Plugins/FX/Niagara/Source/Niagara/Classes/NiagaraEmitterHandle.h b/Engine/Plugins/FX/Niagara/Source/Niagara/Classes/NiagaraEmitterHandle.h
new file mode 100644
--- /dev/null
+++ b/Engine/Plugins/FX/Niagara/Source/Niagara/Classes/NiagaraEmitterHandle.h
@@ -0,0 +1,119 @@
+// Copyright 1998-2016 Epic Games, Inc. All Rights Reserved.
+
+#pragma once
+
+#include "CoreMinimal.h"
+#include "UObject/ObjectMacros.h"
+#include "UObject/Object.h"
+#include "Guid.h"
+#include "NiagaraEmitterHandle.generated.h"
+
+class UNiagaraSystem;
+class UNiagaraEmitter;
+
+/** 
+ * Stores a references to a source emitter asset, and a copy of that emitter for editing within an System.  Also
+ * stores whether or not this emitter is enabled, and it's name within the editor.
+ */
+USTRUCT()
+struct NIAGARA_API FNiagaraEmitterHandle
+{
+	GENERATED_USTRUCT_BODY()
+public:
+	/** Creates a new invalid emitter handle. */
+	FNiagaraEmitterHandle();
+
+	/** Create a new emitter handle from an emitter, but does NOT make a copy, any changes made to the "Instance" will modify
+		the original asset.  This version should only be used in the emitter toolkit. */
+	FNiagaraEmitterHandle(UNiagaraEmitter& Emitter);
+
+#if WITH_EDITORONLY_DATA
+	/** Creates a new emitter handle from an emitter and an owning System. */
+	FNiagaraEmitterHandle(const UNiagaraEmitter& InSourceEmitter, FName InName, UNiagaraSystem& InOuterSystem);
+
+	/** Creates a new emitter handle by duplicating an existing handle.  The new emitter handle will reference the same source emitter
+		but will have it's own copy of the emitter made from the one in the supplied handle and will have it's own Id. */
+	FNiagaraEmitterHandle(const FNiagaraEmitterHandle& InHandleToDuplicate, FName InDuplicateName, UNiagaraSystem& InDuplicateOwnerSystem);
+#endif
+	/** Whether or not this is a valid emitter handle. */
+	bool IsValid() const;
+
+	/** Gets the unique id for this handle. */
+	FGuid GetId() const;
+
+	// HACK!  Data sets used to use the emitter name, but this isn't guaranteed to be unique.  This is a temporary hack
+	// to allow the data sets to continue work with using names, but that code needs to be refactored to use the id defined here.
+	FName GetIdName() const;
+
+	/** Gets the display name for this emitter in the System. */
+	FName GetName() const;
+
+	/** Sets the display name for this emitter in the System. */
+	void SetName(FName InName);
+
+	/** Gets whether or not this emitter is enabled within the System.  Disabled emitters aren't simulated. */
+	bool GetIsEnabled() const;
+
+	/** Gets whether or not this emitter is enabled within the System.  Disabled emitters aren't simulated. */
+	void SetIsEnabled(bool bInIsEnabled);
+
+#if WITH_EDITORONLY_DATA
+	/** Gets the source emitter this emitter handle was built from. */
+	const UNiagaraEmitter* GetSource() const;
+#endif
+
+	/** Gets the copied instance of the emitter this handle references. */
+	UNiagaraEmitter* GetInstance() const;
+
+	/** Update the instance this handle references*/
+	void SetInstance(UNiagaraEmitter* InInstance);
+
+	/** Gets a unique name for this emitter instance for use in scripts and parameter stores etc.*/
+	FString GetUniqueInstanceName()const;
+
+#if WITH_EDITORONLY_DATA
+	/** Return this emitter handle instance to its initial state, exactly matching the Source.*/
+	void ResetToSource();
+
+	/** Keep any existing settings, but include any new changes from the source emitter.*/
+	bool RefreshFromSource();
+
+	/** Determine whether or not the Source and Instance refer to the same Emitter ChangeId.*/
+	bool IsSynchronizedWithSource() const;
+
+	/** Determine whether or not the Instance script is in synch with its graph.*/
+	bool NeedsRecompile() const;
+
+#endif
+public:
+	/** A static const invalid handle. */
+	static const FNiagaraEmitterHandle InvalidHandle;
+
+private:
+	/** The id of this emitter handle. */
+	UPROPERTY(VisibleAnywhere, Category="Emitter ID")
+	FGuid Id;
+
+	// HACK!  Data sets used to use the emitter name, but this isn't guaranteed to be unique.  This is a temporary hack
+	// to allow the data sets to continue work with using names, but that code needs to be refactored to use the id defined here.
+	UPROPERTY(VisibleAnywhere, Category = "Emitter ID")
+	FName IdName;
+
+	/** Whether or not this emitter is enabled within the System.  Disabled emitters aren't simulated. */
+	UPROPERTY()
+	bool bIsEnabled;
+	
+	/** The display name for this emitter in the System. */
+	UPROPERTY()
+	FName Name;
+
+#if WITH_EDITORONLY_DATA
+	/** The source emitter this emitter handle was built from. */
+	UPROPERTY()
+	const UNiagaraEmitter* Source;
+#endif
+
+	/** The copied instance of the emitter this handle references. */
+	UPROPERTY()
+	UNiagaraEmitter* Instance;
+};
\ No newline at end of file
