Copying //UE4/Orion-Staging to //UE4/Main (Source: //Orion/Dev-General @ 3606378)

diff --git a/Engine/Plugins/FX/Niagara/Source/Niagara/Classes/NiagaraEmitterInstance.h b/Engine/Plugins/FX/Niagara/Source/Niagara/Classes/NiagaraEmitterInstance.h
new file mode 100644
--- /dev/null
+++ b/Engine/Plugins/FX/Niagara/Source/Niagara/Classes/NiagaraEmitterInstance.h
@@ -0,0 +1,165 @@
+// Copyright 1998-2017 Epic Games, Inc. All Rights Reserved.
+
+/*==============================================================================
+NiagaraEmitterInstance.h: Niagara emitter simulation class
+==============================================================================*/
+#pragma once
+
+#include "CoreMinimal.h"
+#include "UObject/WeakObjectPtr.h"
+#include "NiagaraCommon.h"
+#include "NiagaraDataSet.h"
+#include "NiagaraEvents.h"
+#include "NiagaraCollision.h"
+#include "NiagaraEmitterHandle.h"
+#include "NiagaraEmitter.h"
+#include "NiagaraScriptExecutionContext.h"
+
+class FNiagaraSystemInstance;
+class NiagaraRenderer;
+struct FNiagaraEmitterHandle;
+class UNiagaraParameterCollection;
+class UNiagaraParameterCollectionInstance;
+
+/**
+* A Niagara particle simulation.
+*/
+struct FNiagaraEmitterInstance
+{
+
+private:
+	struct FNiagaraBurstInstance
+	{
+		float Time;
+		uint32 NumberToSpawn;
+	};
+
+public:
+	explicit FNiagaraEmitterInstance(FNiagaraSystemInstance* InParentSystemInstance);
+	bool bDumpAfterEvent;
+	virtual ~FNiagaraEmitterInstance();
+
+	void Init(int32 InEmitterIdx, FName SystemInstanceName);
+	void ResetSimulation();
+	void ReInitSimulation();
+	void DirtyDataInterfaces();
+
+	/** Replaces the binding for a single parameter colleciton instance. If for example the component begins to override the global instance. */
+	//void RebindParameterCollection(UNiagaraParameterCollectionInstance* OldInstance, UNiagaraParameterCollectionInstance* NewInstance);
+	void BindParameters();
+	void UnbindParameters();
+	
+	/** Called after all emitters in an System have been initialized, allows emitters to access information from one another. */
+	void PostResetSimulation();
+
+	void PreTick();
+	void Tick(float DeltaSeconds);
+
+	uint32 CalculateEventSpawnCount(const FNiagaraEventScriptProperties &EventHandlerProps, TArray<int32> &EventSpawnCounts, FNiagaraDataSet *EventSet);
+	void PostProcessParticles();
+
+	FBox GetBounds() const { return CachedBounds; }
+
+	FNiagaraDataSet &GetData()	{ return Data; }
+
+	int32 GetEmitterRendererNum() const {
+		return EmitterRenderer.Num();
+	}
+
+	NiagaraRenderer *GetEmitterRenderer(int32 Index)	{ return EmitterRenderer[Index]; }
+	
+	bool IsEnabled()const 	{ return bIsEnabled;  }
+	
+	/** Set whether or not this simulation is enabled*/
+	void SetEnabled(bool bEnabled) { bIsEnabled = bEnabled; }
+
+	/** Sets the error state. */
+	void SetError(bool bInError) { bError = bInError; }
+		
+	/** Create a new NiagaraRenderer. The old renderer is not immediately deleted, but instead put in the ToBeRemoved list.*/
+	void NIAGARA_API UpdateEmitterRenderer(ERHIFeatureLevel::Type FeatureLevel, TArray<NiagaraRenderer*>& ToBeAddedList, TArray<NiagaraRenderer*>& ToBeRemovedList);
+
+	int32 GetNumParticles()	{ return Data.GetNumInstances(); }
+
+	NIAGARA_API const FNiagaraEmitterHandle& GetEmitterHandle() const;
+
+	FNiagaraSystemInstance* GetParentSystemInstance()	{ return ParentSystemInstance; }
+
+	float NIAGARA_API GetTotalCPUTime();
+	int	NIAGARA_API GetTotalBytesUsed();
+	
+	ENiagaraExecutionState NIAGARA_API GetExecutionState()	{ return ExecutionState; }
+	void NIAGARA_API SetExecutionState(ENiagaraExecutionState InState);
+
+	FNiagaraDataSet* GetDataSet(FNiagaraDataSetID SetID);
+
+	/** Tell the renderer thread that we're done with the Niagara renderer on this simulation.*/
+	void ClearRenderer();
+
+	FBox GetCachedBounds() { return CachedBounds; }
+
+	FNiagaraScriptExecutionContext& GetSpawnExecutionContext() { return SpawnExecContext; }
+	FNiagaraScriptExecutionContext& GetUpdateExecutionContext() { return UpdateExecContext; }
+	TArray<FNiagaraScriptExecutionContext>& GetEventExecutionContexts() { return EventExecContexts; }
+
+	TArray<FNiagaraSpawnInfo>& GetSpawnInfo() { return SpawnInfos; }
+private:
+
+	/** The index of our emitter in our parent system instance. */
+	int32 EmitterIdx;
+
+	/* The age of the emitter*/
+	float Age;
+	/* how many loops this emitter has gone through */
+	uint32 Loops;
+	/* If false, don't tick or render*/
+	bool bIsEnabled;
+	
+	/* Some error has occurred so stop ticking or rendering. Better to just kill the emitter? */
+	bool bError;
+
+	/* Seconds taken to process everything (including rendering) */
+	float CPUTimeMS;
+	/* Emitter tick state */
+	ENiagaraExecutionState ExecutionState;
+	/* Emitter bounds */
+	FBox CachedBounds;
+
+	/** Array of all spawn info driven by our owning emitter script. */
+	TArray<FNiagaraSpawnInfo> SpawnInfos;
+
+	FNiagaraScriptExecutionContext SpawnExecContext;
+	FNiagaraScriptExecutionContext UpdateExecContext;
+	TArray<FNiagaraScriptExecutionContext> EventExecContexts;
+
+	FNiagaraParameterDirectBinding<float> SpawnIntervalBinding;
+	FNiagaraParameterDirectBinding<float> InterpSpawnStartBinding;
+
+	FNiagaraParameterDirectBinding<float> SpawnEmitterAgeBinding;
+	FNiagaraParameterDirectBinding<float> UpdateEmitterAgeBinding;
+	TArray<FNiagaraParameterDirectBinding<float>> EventEmitterAgeBindings;
+
+	FNiagaraParameterDirectBinding<int32> SpawnExecCountBinding;
+	FNiagaraParameterDirectBinding<int32> UpdateExecCountBinding;
+	TArray<FNiagaraParameterDirectBinding<int32>> EventExecCountBindings;
+
+	/** particle simulation data */
+	FNiagaraDataSet Data;
+	/** The cached GetComponentTransform() transform. */
+	FTransform CachedComponentToWorld;
+
+	TArray<NiagaraRenderer*> EmitterRenderer;
+	FNiagaraSystemInstance *ParentSystemInstance;
+
+	TArray<FNiagaraDataSet*> UpdateScriptEventDataSets;
+	TArray<FNiagaraDataSet*> SpawnScriptEventDataSets;
+	TMap<FNiagaraDataSetID, FNiagaraDataSet*> DataSetMap;
+	
+	FNiagaraCollisionBatch CollisionBatch;
+
+	FName OwnerSystemInstanceName;
+
+#if WITH_EDITORONLY_DATA
+	bool CheckAttributesForRenderer(int32 Index);
+#endif
+};
diff --git a/Engine/Plugins/FX/Niagara/Source/Niagara/Classes/NiagaraEmitterInstanceBatcher.h b/Engine/Plugins/FX/Niagara/Source/Niagara/Classes/NiagaraEmitterInstanceBatcher.h
new file mode 100644
--- /dev/null
+++ b/Engine/Plugins/FX/Niagara/Source/Niagara/Classes/NiagaraEmitterInstanceBatcher.h
@@ -0,0 +1,134 @@
+// Copyright 1998-2017 Epic Games, Inc. All Rights Reserved.
+
+/*==============================================================================
+NiagaraEmitterInstanceBatcher.h: Queueing and batching for Niagara simulation;
+use to reduce per-simulation overhead by batching together simulations using
+the same VectorVM byte code / compute shader code
+==============================================================================*/
+#pragma once
+
+#include "CoreMinimal.h"
+#include "NiagaraCommon.h"
+#include "RendererInterface.h"
+#include "NiagaraParameters.h"
+#include "NiagaraEmitter.h"
+#include "Tickable.h"
+#include "ModuleManager.h"
+#include "ModuleManager.h"
+
+struct FNiagaraScriptExecutionContext;
+
+struct FNiagaraComputeExecutionContext
+{
+	FNiagaraComputeExecutionContext()
+		: MainDataSet(nullptr)
+		, SpawnRateInstances(0)
+		, BurstInstances(0)
+		, EventSpawnTotal(0)
+		, RTUpdateScript(0)
+		, RTSpawnScript(0)
+	{
+		TickCounter++;
+	}
+	class FNiagaraDataSet *MainDataSet;
+	TArray<FNiagaraDataSet*>UpdateEventWriteDataSets;
+	TArray<FNiagaraEventScriptProperties> EventHandlerScriptProps;
+	TArray<FNiagaraDataSet*> EventSets;
+	uint32 SpawnRateInstances;
+	uint32 BurstInstances;
+
+	TArray<int32> EventSpawnCounts;
+	uint32 EventSpawnTotal;
+
+	class FNiagaraScript *RTUpdateScript;
+	class FNiagaraScript *RTSpawnScript;
+	TArray<uint8, TAlignedHeapAllocator<16>> UpdateParams;		// RT side copy of the parameter data
+	TArray<FNiagaraScriptDataInterfaceInfo> UpdateInterfaces;
+	TArray<uint8, TAlignedHeapAllocator<16>> SpawnParams;		// RT side copy of the parameter data
+	static uint32 TickCounter;
+};
+
+
+class NiagaraEmitterInstanceBatcher : public FTickableGameObject, public FComputeDispatcher
+{
+public:
+	NiagaraEmitterInstanceBatcher()
+		: CurQueueIndex(0)
+	{
+		IRendererModule *RendererModule = FModuleManager::GetModulePtr<IRendererModule>("Renderer");
+		if (RendererModule)
+		{
+			RendererModule->RegisterPostOpaqueComputeDispatcher(this);
+		}
+	}
+
+	~NiagaraEmitterInstanceBatcher()
+	{
+		IRendererModule *RendererModule = FModuleManager::GetModulePtr<IRendererModule>("Renderer");
+		if (RendererModule)
+		{
+			RendererModule->UnRegisterPostOpaqueComputeDispatcher(this);
+		}
+	}
+
+	static NiagaraEmitterInstanceBatcher *Get()
+	{
+		if (BatcherSingleton == nullptr)
+		{
+			BatcherSingleton = new NiagaraEmitterInstanceBatcher();
+		}
+		return BatcherSingleton;
+	}
+
+	void Queue(FNiagaraComputeExecutionContext *InContext);
+
+	virtual bool IsTickable() const
+	{
+		return true;
+	}
+
+	virtual TStatId GetStatId() const
+	{
+		RETURN_QUICK_DECLARE_CYCLE_STAT(NiagaraEmitterInstanceBatcher, STATGROUP_Tickables);
+	}
+
+
+	virtual void Tick(float DeltaTime) override
+	{
+		BuildBatches();
+	}
+
+	// TODO: process queue, build batches from context with the same script
+	//  also need to figure out how to handle multiple sets of parameters across a batch
+	//	for now this executes every single sim in the queue individually, which is terrible 
+	//	in terms of overhead
+	void BuildBatches()
+	{
+	}
+
+	// from FComputeDispatcher; called once per frame by the render thread, swaps buffers and works down 
+	// the queue submitted by the game thread; means we're one frame behind with this; we need a mechanism to determine execution order here.
+	virtual void Execute(FRHICommandList &RHICmdList)
+	{
+		CurQueueIndex ^= 0x1;
+		ExecuteAll(RHICmdList);
+	}
+
+	void ExecuteAll(FRHICommandList &RHICmdList);
+	void TickSingle(const FNiagaraComputeExecutionContext *Context, FRHICommandList &RHICmdList) const;
+
+	void SetPrevDataStrideParams(const FNiagaraDataSet *Set, FNiagaraShader *Shader, FRHICommandList &RHICmdList) const;
+
+	void SetupEventUAVs(const FNiagaraComputeExecutionContext *Context, uint32 NumInstances, FRHICommandList &RHICmdList) const;
+	void UnsetEventUAVs(const FNiagaraComputeExecutionContext *Context, FRHICommandList &RHICmdList) const;
+	void SetupDataInterfaceBuffers(const TArray<FNiagaraScriptDataInterfaceInfo> &DIInfos, FNiagaraShader *Shader, FRHICommandList &RHICmdList) const;
+
+	void Run(FNiagaraDataSet *DataSet, const uint32 StartInstance, const uint32 NumInstances, class FNiagaraShader *Shader, const TArray<uint8, TAlignedHeapAllocator<16>>  &Params, FRHICommandList &RhiCmdList, bool bCopyBeforeStart = false) const;
+	void RunEventHandlers(const FNiagaraComputeExecutionContext *Context, uint32 NumInstancesAfterSim, uint32 NumInstancesAfterSpawn, uint32 NumInstancesAfterNonEventSpawn, FRHICommandList &RhiCmdList) const;
+	void ResolveDatasetWrites(uint32 *OutArray, const FNiagaraComputeExecutionContext *Context) const;
+private:
+	static NiagaraEmitterInstanceBatcher* BatcherSingleton;
+
+	uint32 CurQueueIndex;
+	TArray<FNiagaraComputeExecutionContext*> SimulationQueue[2];
+};
diff --git a/Engine/Plugins/FX/Niagara/Source/Niagara/Classes/NiagaraEvents.h b/Engine/Plugins/FX/Niagara/Source/Niagara/Classes/NiagaraEvents.h
new file mode 100644
--- /dev/null
+++ b/Engine/Plugins/FX/Niagara/Source/Niagara/Classes/NiagaraEvents.h
@@ -0,0 +1,192 @@
+// Copyright 1998-2017 Epic Games, Inc. All Rights Reserved.
+
+/*==============================================================================
+NiagaraEmitterInstance.h: Niagara emitter simulation class
+==============================================================================*/
+
+#pragma once
+
+#include "CoreMinimal.h"
+#include "UObject/ObjectMacros.h"
+#include "UObject/Object.h"
+#include "NiagaraCommon.h"
+#include "NiagaraDataSet.h"
+#include "NiagaraEvents.generated.h"
+
+struct FNiagaraEventReceiverProperties;
+
+#define NIAGARA_BUILTIN_EVENTNAME_COLLISION FName("NiagaraSystem_Collision")
+#define NIAGARA_BUILTIN_EVENTNAME_SPAWN FName("Spawn")
+#define NIAGARA_BUILTIN_EVENTNAME_DEATH FName("Death")
+
+/**
+ *	Type struct for collision event payloads; collision event data set is based on this
+ *  TODO: figure out how we can pipe attributes from the colliding particle in here
+ */
+USTRUCT()
+struct FNiagaraCollisionEventPayload
+{
+	GENERATED_USTRUCT_BODY()
+
+	UPROPERTY()
+	FVector CollisionPos;
+	UPROPERTY()
+	FVector CollisionNormal;
+	UPROPERTY()
+	FVector CollisionVelocity;
+	UPROPERTY()
+	int32 ParticleIndex;
+	UPROPERTY()
+	int32 PhysicalMaterialIndex;
+};
+
+
+struct FNiagaraEmitterInstance;
+
+/**
+Base class for actions that an event receiver will perform at the emitter level.
+*/
+UCLASS(abstract)
+class UNiagaraEventReceiverEmitterAction : public UObject
+{
+	GENERATED_BODY()
+public:
+	virtual void PerformAction(FNiagaraEmitterInstance& OwningSim, const FNiagaraEventReceiverProperties& OwningEventReceiver){}
+};
+
+UCLASS(EditInlineNew)
+class UNiagaraEventReceiverEmitterAction_SpawnParticles : public UNiagaraEventReceiverEmitterAction
+{
+	GENERATED_BODY()
+
+public:
+
+	/** Number of particles to spawn per event received. */
+	UPROPERTY(EditAnywhere, Category = "Spawn")
+	uint32 NumParticles;
+
+	virtual void PerformAction(FNiagaraEmitterInstance& OwningSim, const FNiagaraEventReceiverProperties& OwningEventReceiver)override;
+};
+
+
+
+typedef TMap<FName, FNiagaraDataSet*> PerEmitterEventDataSetMap;
+typedef TMap<FName, PerEmitterEventDataSetMap> PerSystemInstanceDataSetMap;
+
+class FNiagaraEventDataSetMgr
+{
+public:
+	FNiagaraEventDataSetMgr()
+	{
+	}
+
+	// call this to allocate a data set for a specific emitter and event
+	// if the data set already exists, returns the existing one
+	static FNiagaraDataSet *CreateEventDataSet(FName OwnerSystemInstanceName, FName EmitterName, FName EventName)
+	{
+		PerSystemInstanceDataSetMap *SystemInstanceMap = EmitterEventDataSets.Find(OwnerSystemInstanceName);
+		if (!SystemInstanceMap)
+		{
+			SystemInstanceMap = &EmitterEventDataSets.Add(OwnerSystemInstanceName);
+		}
+
+		// find the emitter's map
+		PerEmitterEventDataSetMap *Map = SystemInstanceMap->Find(EmitterName);
+		if (!Map)
+		{
+			Map = &SystemInstanceMap->Add(EmitterName);
+		}
+
+		// TODO: find a better way of multiple events trying to write to the same data set; 
+		// for example, if two analytical collision primitives want to send collision events, they need to push to the same data set
+
+		FNiagaraDataSet **DataSetPtr = Map->Find(EventName);
+		FNiagaraDataSet *OutSet = nullptr;
+		if (DataSetPtr == nullptr)
+		{
+			OutSet = new FNiagaraDataSet();
+			Map->Add(EventName) = OutSet;
+		}
+		else
+		{
+			OutSet = *DataSetPtr;
+		}
+		return OutSet;
+	}
+
+	// deletes and cleans up all event data sets for an emitter
+	//  should be called when the emitter is destroyed
+	static void Reset(FName OwnerSystemInstanceName, FName EmitterName)
+	{
+		PerSystemInstanceDataSetMap *SystemInstanceMap = EmitterEventDataSets.Find(OwnerSystemInstanceName);
+		if (SystemInstanceMap)
+		{
+			PerEmitterEventDataSetMap *Map = SystemInstanceMap->Find(EmitterName);
+			if (Map)
+			{
+				for (TPair<FName, FNiagaraDataSet *> Pair : *Map)
+				{
+					delete Pair.Value;
+				}
+				Map->Empty();
+				SystemInstanceMap->Remove(EmitterName);
+			}
+			if (SystemInstanceMap->Num() == 0)
+			{
+				EmitterEventDataSets.Remove(OwnerSystemInstanceName);
+			}
+		}
+	}
+
+	/*
+	static void AddEventSet(FName EmitterName, FName EventName, const FNiagaraDataSet *DataSet)
+	{
+		// find the emitter's map
+		PerEmitterEventDataSetMap *Map = EmitterEventDataSets.Find(EmitterName);
+		if(!Map)
+		{
+			Map = &EmitterEventDataSets.Add(EmitterName);
+		}
+
+		// TODO: find a better way of multiple events trying to write to the same data set; 
+		// for example, if two analytical collision primitives want to send collision events, they need to push to the same data set
+		ensure(Map->Find(EventName) == nullptr);
+
+		if (Map->Find(EventName) == nullptr)
+		{
+			Map->Add(EventName) = DataSet;
+		}
+	}
+	*/
+
+	static PerEmitterEventDataSetMap * GetEmitterMap(FName OwnerSystemInstanceName, FName EmitterName)
+	{
+		TMap<FName, PerEmitterEventDataSetMap> *SystemMap = EmitterEventDataSets.Find(OwnerSystemInstanceName);
+		if (SystemMap)
+		{
+			return SystemMap->Find(EmitterName);
+		}
+
+		return nullptr;
+	}
+
+	static FNiagaraDataSet * GetEventDataSet(FName OwnerSystemInstanceName, FName EmitterName, FName EventName)
+	{
+		PerEmitterEventDataSetMap *SetMap = GetEmitterMap(OwnerSystemInstanceName, EmitterName);
+		if (!SetMap)
+		{
+			return nullptr;
+		}
+
+		FNiagaraDataSet **FoundSet = SetMap->Find(EventName);
+		if (!FoundSet)
+		{
+			return nullptr;
+		}
+
+		return *FoundSet;
+	}
+
+private:
+	static TMap<FName, TMap<FName, PerEmitterEventDataSetMap>> EmitterEventDataSets;
+};
\ No newline at end of file
diff --git a/Engine/Plugins/FX/Niagara/Source/Niagara/Classes/NiagaraParameterCollection.h b/Engine/Plugins/FX/Niagara/Source/Niagara/Classes/NiagaraParameterCollection.h
new file mode 100644
--- /dev/null
+++ b/Engine/Plugins/FX/Niagara/Source/Niagara/Classes/NiagaraParameterCollection.h
@@ -0,0 +1,160 @@
+// Copyright 1998-2017 Epic Games, Inc. All Rights Reserved.
+
+#pragma once
+
+#include "NiagaraParameterStore.h"
+#include "NiagaraCommon.h"
+#include "NiagaraParameterCollection.generated.h"
+
+class UNiagaraParameterCollection;
+
+UCLASS()
+class NIAGARA_API UNiagaraParameterCollectionInstance : public UObject
+{
+	GENERATED_UCLASS_BODY()
+
+public:
+	virtual ~UNiagaraParameterCollectionInstance();
+
+	//~UObject interface
+	virtual void PostLoad()override;
+	//~UObject interface
+	
+	bool IsDefaultInstance()const;
+	
+	void SetParent(UNiagaraParameterCollection* InParent);
+	UNiagaraParameterCollection* GetParent()const { return Collection; }
+
+	FNiagaraParameterStore& GetParameterStore() { return ParameterStorage; }
+
+	bool AddParameter(const FNiagaraVariable& Parameter);
+	bool RemoveParameter(const FNiagaraVariable& Parameter);
+	void RenameParameter(const FNiagaraVariable& Parameter, FName NewName);
+	void Empty();
+	void GetParameters(TArray<FNiagaraVariable>& OutParameters);
+
+	void Tick();
+
+	//TODO: Abstract to some interface to allow a hierarchy like UMaterialInstance?
+	UPROPERTY(EditAnywhere, Category=Instance)
+	UNiagaraParameterCollection* Collection;
+
+	/**
+	When editing instances, we must track which parameters are overridden so we can pull in any changes to the default.
+	*/
+	UPROPERTY()
+	TArray<FNiagaraVariable> OverridenParameters;
+	bool OverridesParameter(const FNiagaraVariable& Parameter)const;
+	void SetOverridesParameter(const FNiagaraVariable& Parameter, bool bOverrides);
+
+	/** Synchronizes this instance with any changes with it's parent collection. */
+	void SyncWithCollection();
+	
+private:
+
+	UPROPERTY()
+	FNiagaraParameterStore ParameterStorage;
+
+	//TODO: These overrides should be settable per platform.
+	//UPROPERTY()
+	//TMap<FString, FNiagaraParameterStore>
+
+public:
+	//Accessors from Blueprint. For now just exposing common types but ideally we can expose any somehow in future.
+	
+	UFUNCTION(BlueprintCallable, Category = Niagara, meta = (DisplayName = "Get Bool Parameter"))
+	bool GetBoolParameter(const FString& InVariableName);
+
+	UFUNCTION(BlueprintCallable, Category = Niagara, meta = (DisplayName = "Get Float Parameter"))
+	float GetFloatParameter(const FString& InVariableName);
+
+	UFUNCTION(BlueprintCallable, Category = Niagara, meta = (DisplayName = "Get Int Parameter"))
+	int32 GetIntParameter(const FString& InVariableName);
+
+	UFUNCTION(BlueprintCallable, Category = Niagara, meta = (DisplayName = "Get Vector2D Parameter"))
+	FVector2D GetVector2DParameter(const FString& InVariableName);
+
+	UFUNCTION(BlueprintCallable, Category = Niagara, meta = (DisplayName = "Get Vector Parameter"))
+	FVector GetVectorParameter(const FString& InVariableName);
+
+	UFUNCTION(BlueprintCallable, Category = Niagara, meta = (DisplayName = "Get Vector4 Parameter"))
+	FVector4 GetVector4Parameter(const FString& InVariableName);
+
+	UFUNCTION(BlueprintCallable, Category = Niagara, meta = (DisplayName = "Get Color Parameter"))
+	FLinearColor GetColorParameter(const FString& InVariableName);
+
+	UFUNCTION(BlueprintCallable, Category = Niagara, meta = (DisplayName = "Set Bool Parameter"))
+	void SetBoolParameter(const FString& InVariableName, bool InValue);
+
+	UFUNCTION(BlueprintCallable, Category = Niagara, meta = (DisplayName = "Set Float Parameter"))
+	void SetFloatParameter(const FString& InVariableName, float InValue);
+
+	UFUNCTION(BlueprintCallable, Category = Niagara, meta = (DisplayName = "Set Int Parameter"))
+	void SetIntParameter(const FString& InVariableName, int32 InValue);
+
+	UFUNCTION(BlueprintCallable, Category = Niagara, meta = (DisplayName = "Set Vector2D Parameter"))
+	void SetVector2DParameter(const FString& InVariableName, FVector2D InValue);
+
+	UFUNCTION(BlueprintCallable, Category = Niagara, meta = (DisplayName = "Set Vector Parameter"))
+	void SetVectorParameter(const FString& InVariableName, FVector InValue);
+
+	UFUNCTION(BlueprintCallable, Category = Niagara, meta = (DisplayName = "Set Vector4 Parameter"))
+	void SetVector4Parameter(const FString& InVariableName, const FVector4& InValue);
+
+	UFUNCTION(BlueprintCallable, Category = Niagara, meta = (DisplayName = "Set Color Parameter"))
+	void SetColorParameter(const FString& InVariableName, FLinearColor InValue);
+};
+
+/** Asset containing a collection of global parameters usable by Niagara. */
+UCLASS()
+class NIAGARA_API UNiagaraParameterCollection : public UObject
+{
+	GENERATED_UCLASS_BODY()
+public:
+
+	//~UObject interface
+	virtual void PostInitProperties()override;
+	virtual void PostLoad()override;
+	//~UObject interface
+
+	int32 IndexOfParameter(const FNiagaraVariable& Var);
+
+	int32 AddParameter(FName Name, FNiagaraTypeDefinition Type);
+	void RemoveParameter(const FNiagaraVariable& Parameter);
+	void RenameParameter(FNiagaraVariable& Parameter, FName NewName);
+
+	TArray<FNiagaraVariable>& GetParameters() { return Parameters; }
+
+	//TODO: Optional per platform overrides of the above.
+	//TMap<FString, UNiagaraParameterCollectionOverride> PerPlatformOverrides;
+
+	FORCEINLINE UNiagaraParameterCollectionInstance* GetDefaultInstance() { return DefaultInstance; }
+	
+	/**
+	Takes the friendly name presented to the UI and converts to the real parameter name used under the hood.
+	Converts from "ParameterName" to "CollectionUniqueName_ParameterName".
+	*/
+	FString ParameterNameFromFriendlyName(FString FriendlyName)const;
+	/**
+	Takes the real parameter name used under the hood and converts to the friendly name for use in the UI.
+	Converts from "CollectionUniqueName_ParameterName" to "ParameterName".
+	*/
+
+	FNiagaraVariable CollectionParameterFromFriendlyParameter(const FNiagaraVariable& FriendlyParameter)const;
+	FNiagaraVariable FriendlyParameterFromCollectionParameter(const FNiagaraVariable& CollectionParameter)const;
+
+	FString FriendlyNameFromParameterName(FString ParameterName)const;
+protected:
+
+	FString GetUniqueName()const;
+	
+	UPROPERTY()
+	TArray<FNiagaraVariable> Parameters;
+	
+	UPROPERTY()
+	UNiagaraParameterCollectionInstance* DefaultInstance;
+
+	/** Unique name used by parameters in this collection and the scripts referencing them to link them with this collection. */
+	UPROPERTY()
+	FString UniqueName;
+};
\ No newline at end of file
diff --git a/Engine/Plugins/FX/Niagara/Source/Niagara/Classes/NiagaraScript.h b/Engine/Plugins/FX/Niagara/Source/Niagara/Classes/NiagaraScript.h
new file mode 100644
--- /dev/null
+++ b/Engine/Plugins/FX/Niagara/Source/Niagara/Classes/NiagaraScript.h
@@ -0,0 +1,295 @@
+// Copyright 1998-2017 Epic Games, Inc. All Rights Reserved.
+
+#pragma once
+
+#include "CoreMinimal.h"
+#include "UObject/ObjectMacros.h"
+#include "UObject/Object.h"
+#include "NiagaraCommon.h"
+#include "NiagaraShared.h"
+#include "NiagaraShader.h"
+#include "NiagaraParameters.h"
+#include "NiagaraDataSet.h"
+
+#include "NiagaraScript.generated.h"
+
+class UNiagaraDataInterface;
+
+void SerializeNiagaraShaderMaps(const TMap<const ITargetPlatform*, FNiagaraScript*>* PlatformScriptResourcesToSavePtr, FArchive& Ar, FNiagaraScript& OutLoadedResources);
+void ProcessSerializedShaderMaps(UNiagaraScript* Owner, FNiagaraScript LoadedResource, FNiagaraScript* (&OutScriptResourcesLoaded)[ERHIFeatureLevel::Num]);
+
+#define NIAGARA_INVALID_MEMORY (0xBA)
+
+//#define NIAGARA_SCRIPT_COMPILE_LOGGING_MEDIUM
+
+/** Defines what will happen to unused attributes when a script is run. */
+UENUM()
+enum class EUnusedAttributeBehaviour : uint8
+{
+	/** The previous value of the attribute is copied across. */
+	Copy,
+	/** The attribute is set to zero. */
+	Zero,
+	/** The attribute is untouched. */
+	None,
+	/** The memory for the attribute is set to NIAGARA_INVALID_MEMORY. */
+	MarkInvalid, 
+	/** The attribute is passed through without double buffering */
+	PassThrough,
+};
+
+struct FNiagaraScriptDebuggerInfo
+{
+	FNiagaraScriptDebuggerInfo();
+
+	bool bRequestDebugFrame;
+	
+	int32 DebugFrameLastWriteId;
+
+	FNiagaraDataSet DebugFrame;
+};
+
+
+/** Runtime script for a Niagara system */
+UCLASS(MinimalAPI)
+class UNiagaraScript : public UObject
+{
+	GENERATED_UCLASS_BODY()
+
+	// how this script is to be used. cannot be private due to use of GET_MEMBER_NAME_CHECKED
+	UPROPERTY(AssetRegistrySearchable, EditAnywhere, Category = Script)
+	ENiagaraScriptUsage Usage;
+
+	/** Which instance of the usage in the graph to use.*/
+	UPROPERTY()
+	int32 UsageIndex;
+
+	/** When used as a module, what are the appropriate script types for referencing this module?*/
+	UPROPERTY(AssetRegistrySearchable, EditAnywhere, Category = Script, meta = (Bitmask, BitmaskEnum = ENiagaraScriptUsage))
+	int32 ModuleUsageBitmask;
+
+	/** Used to break up scripts of the same Usage type in UI display.*/
+	UPROPERTY(AssetRegistrySearchable, EditAnywhere, Category = Script)
+	FText Category;
+
+	/** Number of user pointers we must pass to the VM. */
+	UPROPERTY()
+	int32 NumUserPtrs;
+
+	/** Byte code to execute for this system */
+	UPROPERTY()
+	TArray<uint8> ByteCode;
+
+	/** All the data for using constants in the script. */
+	UPROPERTY()
+	FNiagaraParameters Parameters;
+
+	UPROPERTY()
+	FNiagaraParameters InternalParameters;
+
+	UPROPERTY()
+	TMap<FName, FNiagaraParameters> DataSetToParameters;
+
+	/** Attributes used by this script. */
+	UPROPERTY()
+ 	TArray<FNiagaraVariable> Attributes;
+
+	/** Contains various usage information for this script. */
+	UPROPERTY()
+	FNiagaraScriptDataUsageInfo DataUsage;
+
+	/** Information about all data interfaces used by this script. */
+	UPROPERTY()
+	TArray<FNiagaraScriptDataInterfaceInfo> DataInterfaceInfo;
+
+	/** Array of ordered vm external functions to place in the function table. */
+	UPROPERTY()
+	TArray<FVMExternalFunctionBindingInfo> CalledVMExternalFunctions;
+
+	/** The mode to use when deducing the type of numeric output pins from the types of the input pins. */
+	UPROPERTY(EditAnywhere, Category=Script)
+	ENiagaraNumericOutputTypeSelectionMode NumericOutputTypeSelectionMode;
+
+	TArray<FNiagaraDataSetID> ReadDataSets;
+	TArray<FNiagaraDataSetProperties> WriteDataSets;
+
+	/** Scopes we'll track with stats.*/
+	UPROPERTY()
+	TArray<FNiagaraStatScope> StatScopes;
+	
+	/** The parameter collections used by this script. */
+	UPROPERTY()
+	TArray<class UNiagaraParameterCollection*> ParameterCollections;
+
+	UPROPERTY(AssetRegistrySearchable, EditAnywhere, Category = Script)
+	FText Description;
+
+	UPROPERTY()
+	FString LastHlslTranslation;
+	
+	/** Returns true if this script is valid and can be executed. */
+	bool IsValid()const { return ByteCode.Num() > 0; }//More? Differentiate by CPU/GPU?
+
+	void SetUsage(ENiagaraScriptUsage InUsage) { Usage = InUsage; }
+	ENiagaraScriptUsage GetUsage() const { return Usage; }
+	void SetUsageIndex(int32 InUsageIndex) { UsageIndex = InUsageIndex; }
+	int32 GetUsageIndex() const { return UsageIndex; }
+
+	bool IsEquivalentUsage(ENiagaraScriptUsage InUsage) {return (InUsage == Usage) || (Usage == ENiagaraScriptUsage::ParticleSpawnScript && InUsage == ENiagaraScriptUsage::ParticleSpawnScriptInterpolated) || (Usage == ENiagaraScriptUsage::ParticleSpawnScriptInterpolated && InUsage == ENiagaraScriptUsage::ParticleSpawnScript);}
+	static bool IsEquivalentUsage(ENiagaraScriptUsage InUsageA, ENiagaraScriptUsage InUsageB) { return (InUsageA == InUsageB) || (InUsageB == ENiagaraScriptUsage::ParticleSpawnScript && InUsageA == ENiagaraScriptUsage::ParticleSpawnScriptInterpolated) || (InUsageB == ENiagaraScriptUsage::ParticleSpawnScriptInterpolated && InUsageA == ENiagaraScriptUsage::ParticleSpawnScript); }
+
+	bool IsParticleSpawnScript() const { return Usage == ENiagaraScriptUsage::ParticleSpawnScript || Usage == ENiagaraScriptUsage::ParticleSpawnScriptInterpolated; }
+	bool IsInterpolatedParticleSpawnScript()const { return Usage == ENiagaraScriptUsage::ParticleSpawnScriptInterpolated; }
+	bool IsParticleUpdateScript()const { return Usage == ENiagaraScriptUsage::ParticleUpdateScript; }
+	bool IsModuleScript()const { return Usage == ENiagaraScriptUsage::Module; }
+	bool IsFunctionScript()	const { return Usage == ENiagaraScriptUsage::Function; }
+	bool IsDynamicInputScript()	const { return Usage == ENiagaraScriptUsage::DynamicInput; }
+	bool IsParticleEventScript()const { return Usage == ENiagaraScriptUsage::ParticleEventScript; }
+
+	bool IsNonParticleScript()const { return Usage >= ENiagaraScriptUsage::EmitterSpawnScript; }
+	
+	bool IsSystemSpawnScript()const { return Usage == ENiagaraScriptUsage::SystemSpawnScript; }
+	bool IsSystemUpdateScript()const { return Usage == ENiagaraScriptUsage::SystemUpdateScript; }
+	bool IsEmitterSpawnScript()const { return Usage == ENiagaraScriptUsage::EmitterSpawnScript; }
+	bool IsEmitterUpdateScript()const { return Usage == ENiagaraScriptUsage::EmitterUpdateScript; }
+	bool IsStandaloneScript() const { return IsDynamicInputScript() || IsFunctionScript() || IsModuleScript(); }
+
+	bool IsSpawnScript()const { return IsParticleSpawnScript() || IsEmitterSpawnScript() || IsSystemSpawnScript(); }
+
+	NIAGARA_API TArray<ENiagaraScriptUsage> GetSupportedUsageContexts() const;
+
+	bool CanBeRunOnGpu() const { return IsParticleSpawnScript() || IsParticleUpdateScript() || IsParticleEventScript(); }
+
+#if WITH_EDITORONLY_DATA
+	class UNiagaraScriptSourceBase *GetSource() { return Source; }
+	void SetSource(class UNiagaraScriptSourceBase *InSource) { Source = InSource; }
+#endif
+
+	//~ Begin UObject interface
+	void Serialize(FArchive& Ar)override;
+	virtual void PostLoad() override;
+#if WITH_EDITOR
+	virtual void PostEditChangeProperty(FPropertyChangedEvent& PropertyChangedEvent) override;
+#endif
+	//~ End UObject interface
+
+	// Infrastructure for GPU compute Shaders
+	NIAGARA_API void CacheResourceShadersForCooking(EShaderPlatform ShaderPlatform, FNiagaraScript* &OutCachedResource);
+	NIAGARA_API void CacheResourceShadersForRendering(bool bRegenerateId, bool bForceRecompile=false);
+	void BeginCacheForCookedPlatformData(const ITargetPlatform *TargetPlatform);
+	void CacheShadersForResources(EShaderPlatform ShaderPlatform, FNiagaraScript *ResourceToCache, bool bApplyCompletedShaderMapForRendering, bool bForceRecompile = false, bool bCooking=false);
+	FNiagaraScript* AllocateResource();
+	FNiagaraScript *GetRenderThreadScript()
+	{
+		return &ScriptResource;
+	}
+
+	FComputeShaderRHIRef GetScriptShader() 
+	{
+		if (!ScriptShader)
+		{
+			ScriptShader = ScriptResource.GetShader()->GetComputeShader();	// NIAGARATODO: need to put this caching somewhere else, as it wont' know when we update the resource
+		}
+		return ScriptShader; 
+	}
+
+	FComputeShaderRHIRef GetScriptShaderGameThread()
+	{
+		if (!ScriptShader)
+		{
+			ScriptShader = ScriptResource.GetShaderGameThread()->GetComputeShader();	// NIAGARATODO: need to put this caching somewhere else, as it wont' know when we update the resource
+		}
+		return ScriptShader;
+	}
+
+	void SetFeatureLevel(ERHIFeatureLevel::Type InFeatureLevel)		{ FeatureLevel = InFeatureLevel; }
+
+
+	// SourceChangeID is updated with the source ID everytime the source graph changes
+	// UniqueID is regenerated in that case, so we have separate Ids for scripts coming from the 
+	// same source (e.g. spawn and update of the same emitter)
+	void InvalidateChangeID()	{ ChangeId.Invalidate(); }
+
+	const FGuid &GetChangeID()	{ return ChangeId; }
+
+	void SetChangeID(FGuid InGuid)
+	{
+		ChangeId = InGuid;
+		
+		// whenever source changes, regenerate so we know to recompile; need this because we may have 
+		//	multiple scripts with the same source, for example spawn and update scripts from one graph
+		UniqueID = FGuid::NewGuid();
+	}
+
+	NIAGARA_API void GenerateStatScopeIDs();
+
+#if WITH_EDITORONLY_DATA
+	FText GetDescription() { return Description.IsEmpty() ? FText::FromString(GetName()) : Description; }
+	ENiagaraScriptCompileStatus GetLastCompileStatus() { return LastCompileStatus; }
+	void SetLastCompileStatus(ENiagaraScriptCompileStatus InStatus) { LastCompileStatus = InStatus; }
+
+	/** Makes a deep copy of any script dependencies, including itself.*/
+	NIAGARA_API virtual UNiagaraScript* MakeRecursiveDeepCopy(UObject* DestOuter, TMap<const UObject*, UObject*>& ExistingConversions) const;
+
+	/** Determine if there are any external dependencies with respect to scripts and ensure that those dependencies are sucked into the existing package.*/
+	NIAGARA_API virtual void SubsumeExternalDependencies(TMap<const UObject*, UObject*>& ExistingConversions);
+
+	/** Determine if the Script and its source graph are in sync.*/
+	NIAGARA_API bool AreScriptAndSourceSynchronized() const;
+
+	/** Ensure that the Script and its source graph are marked out of sync.*/
+	NIAGARA_API void MarkScriptAndSourceDesynchronized();
+
+	NIAGARA_API ENiagaraScriptCompileStatus Compile(FString& OutGraphLevelErrorMessages);
+
+	NIAGARA_API FNiagaraScriptDebuggerInfo& GetDebuggerInfo() { return DebuggerInfo; }
+
+	const TArray<TStatId>& GetStatScopeIDs()const { return StatScopesIDs; }
+#endif
+
+	bool UsesCollection(const class UNiagaraParameterCollection* Collection)const;
+	
+	virtual ~UNiagaraScript();
+
+	virtual void PostInitProperties() override;
+private:
+#if WITH_EDITORONLY_DATA
+	/** 'Source' data/graphs for this script */
+	UPROPERTY()
+	class UNiagaraScriptSourceBase*	Source;
+	
+	/** Last known compile status. Lets us determine the latest state of the script byte buffer.*/
+	UPROPERTY()
+	ENiagaraScriptCompileStatus LastCompileStatus;
+
+	FNiagaraScriptDebuggerInfo DebuggerInfo;
+#endif
+
+	/** Adjusted every time that we compile this script. Lets us know that we might differ from any cached versions.*/
+	UPROPERTY()
+	FGuid ChangeId;
+
+	/** unique ID for this script */
+	UPROPERTY()
+	FGuid UniqueID;
+
+	FNiagaraScript ScriptResource;
+	FNiagaraScript* ScriptResourcesByFeatureLevel[ERHIFeatureLevel::Num];
+
+	/** Feature level that the shader map is going to be compiled for. */
+	ERHIFeatureLevel::Type FeatureLevel;
+
+	/** Compute shader compiled for this script */
+	FComputeShaderRHIRef ScriptShader;
+
+	/** Runtime stat IDs generated from StatScopes. */
+#if STATS
+	TArray<TStatId> StatScopesIDs;
+#endif
+
+#if WITH_EDITOR
+	/* script resources being cached for cooking. */
+	TMap<const class ITargetPlatform*, FNiagaraScript*> CachedScriptResourcesForCooking;
+#endif
+
+};
diff --git a/Engine/Plugins/FX/Niagara/Source/Niagara/Classes/NiagaraScriptExecutionContext.h b/Engine/Plugins/FX/Niagara/Source/Niagara/Classes/NiagaraScriptExecutionContext.h
new file mode 100644
--- /dev/null
+++ b/Engine/Plugins/FX/Niagara/Source/Niagara/Classes/NiagaraScriptExecutionContext.h
@@ -0,0 +1,112 @@
+// Copyright 1998-2017 Epic Games, Inc. All Rights Reserved.
+
+/*==============================================================================
+NiagaraEmitterInstance.h: Niagara emitter simulation class
+==============================================================================*/
+#pragma once
+
+#include "CoreMinimal.h"
+#include "UObject/WeakObjectPtr.h"
+#include "NiagaraCommon.h"
+#include "NiagaraDataSet.h"
+#include "NiagaraEvents.h"
+#include "NiagaraCollision.h"
+#include "NiagaraEmitterHandle.h"
+#include "NiagaraEmitter.h"
+#include "NiagaraParameterStore.h"
+
+/** 
+Storage class containing actual runtime buffers to be used by the VM and the GPU. 
+Is not the actual source for any parameter data, rather just the final place it's gathered from various other places ready for execution.
+*/
+class FNiagaraScriptExecutionParameterStore : public FNiagaraParameterStore
+{
+public:
+	FNiagaraScriptExecutionParameterStore();
+	FNiagaraScriptExecutionParameterStore(const FNiagaraParameterStore& Other);
+	FNiagaraScriptExecutionParameterStore& operator=(const FNiagaraParameterStore& Other);
+
+	//TODO: This function can probably go away entirely when we replace the FNiagaraParameters and DataInterface info in the script with an FNiagaraParameterStore.
+	//Special care with prev params and internal params will have to be taken there.
+	void Init(UNiagaraScript* Script);
+
+	void CopyCurrToPrev();
+
+	bool AddParameter(const FNiagaraVariable& Param, bool bInitInterfaces=true)
+	{
+		return FNiagaraParameterStore::AddParameter(Param, bInitInterfaces);
+	}
+
+	bool RemoveParameter(FNiagaraVariable& Param)
+	{
+		check(0);//Not allowed to remove parameters from an execution store as it will adjust the table layout mess up the 
+		return false;
+	}
+
+	void RenameParameter(FNiagaraVariable& Param, FName NewName)
+	{
+		check(0);//Can't rename parameters for an execution store.
+	}
+
+	uint32 GetExternalParameterSize() { return ParameterSize; }
+private:
+
+	/** Size of the parameter data not including prev frame values or internal constants. Allows copying into previous parameter values for interpolated spawn scripts. */
+	int32 ParameterSize;
+};
+
+struct FNiagaraDataSetExecutionInfo
+{
+	FNiagaraDataSetExecutionInfo()
+		: DataSet(nullptr)
+		, StartInstance(0)
+		, bAllocate(false)
+		, bUpdateInstanceCount(false)
+	{
+	}
+
+	FNiagaraDataSetExecutionInfo(FNiagaraDataSet* InDataSet, int32 InStartInstance, bool bInAllocate, bool bInUpdateInstanceCount)
+		: DataSet(InDataSet)
+		, StartInstance(InStartInstance)
+		, bAllocate(bInAllocate)
+		, bUpdateInstanceCount(bInUpdateInstanceCount)
+	{}
+
+	FNiagaraDataSet* DataSet;
+	int32 StartInstance;
+	bool bAllocate;
+	bool bUpdateInstanceCount;
+};
+
+struct FNiagaraScriptExecutionContext
+{
+	UNiagaraScript* Script;
+
+	/** Table of external function delegates called from the VM. */
+	TArray<FVMExternalFunction> FunctionTable;
+
+	/** Table of instance data for data interfaces that require it. */
+	TArray<void*> DataInterfaceInstDataTable;
+
+	/** Parameter store. Contains all data interfaces and a parameter buffer that can be used directly by the VM or GPU. */
+	FNiagaraScriptExecutionParameterStore Parameters;
+
+	TArray<FDataSetMeta> DataSetMetaTable;
+
+	
+	static uint32 TickCounter;
+
+	FNiagaraScriptExecutionContext();
+
+	bool Init(UNiagaraScript* InScript, ENiagaraSimTarget InTarget);
+	bool Tick(class FNiagaraSystemInstance* Instance);
+	void PostTick();
+
+	bool Execute(uint32 NumInstances, TArray<FNiagaraDataSetExecutionInfo>& DataSetInfos);
+
+	const TArray<UNiagaraDataInterface*>& GetDataInterfaces()const { return Parameters.GetDataInterfaces(); }
+
+	void DirtyDataInterfaces();
+
+	bool CanExecute()const;
+};
diff --git a/Engine/Plugins/FX/Niagara/Source/Niagara/Classes/NiagaraScriptSourceBase.h b/Engine/Plugins/FX/Niagara/Source/Niagara/Classes/NiagaraScriptSourceBase.h
new file mode 100644
--- /dev/null
+++ b/Engine/Plugins/FX/Niagara/Source/Niagara/Classes/NiagaraScriptSourceBase.h
@@ -0,0 +1,59 @@
+// Copyright 1998-2017 Epic Games, Inc. All Rights Reserved.
+
+#pragma once
+
+#include "CoreMinimal.h"
+#include "UObject/ObjectMacros.h"
+#include "UObject/Object.h"
+#include "NiagaraCommon.h"
+#include "NiagaraScriptSourceBase.generated.h"
+
+struct EditorExposedVectorConstant
+{
+	FName ConstName;
+	FVector4 Value;
+};
+
+struct EditorExposedVectorCurveConstant
+{
+	FName ConstName;
+	class UCurveVector *Value;
+};
+
+/** Runtime data for a Niagara system */
+UCLASS(MinimalAPI)
+class UNiagaraScriptSourceBase : public UObject
+{
+	GENERATED_UCLASS_BODY()
+
+	TArray<TSharedPtr<EditorExposedVectorConstant> > ExposedVectorConstants;
+	TArray<TSharedPtr<EditorExposedVectorCurveConstant> > ExposedVectorCurveConstants;
+
+	/** Determines if the input change id is equal to the current source graph's change id.*/
+	virtual bool IsSynchronized(const FGuid& InChangeId) { return true; }
+
+	virtual UNiagaraScriptSourceBase* MakeRecursiveDeepCopy(UObject* DestOuter, TMap<const UObject*, UObject*>& ExistingConversions) const { return nullptr; }
+
+	/** Determine if there are any external dependencies wrt to scripts and ensure that those dependencies are sucked into the existing package.*/
+	virtual void SubsumeExternalDependencies(TMap<const UObject*, UObject*>& ExistingConversions) {}
+
+	/** Enforce that the source graph is now out of sync with the script.*/
+	virtual void MarkNotSynchronized() {}
+
+	virtual FGuid GetChangeID() { return FGuid(); };
+
+	/** Have we called PreCompile on this source previously?*/
+	virtual bool IsPreCompiled() const { return false; }
+	
+	/** Cause the source to build up any internal variables that will be useful in the compilation process.*/
+	virtual void PreCompile(UNiagaraEmitter* Emitter, bool bClearErrors = true) {}
+	
+	/** If the user previously pre-compiled the source, dig through the data to find any variables defined. */
+	virtual bool GatherPreCompiledVariables(const FString& InNamespaceFilter, TArray<FNiagaraVariable>& OutVars) { return false; }
+
+	/** Implements compilation of a Niagara script.*/
+	virtual ENiagaraScriptCompileStatus Compile(UNiagaraScript* ScriptOwner, FString& OutGraphLevelErrorMessages) { return ENiagaraScriptCompileStatus::NCS_Unknown; }
+
+	/** Do any cleanup to return this data source to its original state before precompilation. Note that this must be called after the "compile" operation.*/
+	virtual void PostCompile() {}
+};
diff --git a/Engine/Plugins/FX/Niagara/Source/Niagara/Classes/NiagaraSystem.h b/Engine/Plugins/FX/Niagara/Source/Niagara/Classes/NiagaraSystem.h
new file mode 100644
--- /dev/null
+++ b/Engine/Plugins/FX/Niagara/Source/Niagara/Classes/NiagaraSystem.h
@@ -0,0 +1,190 @@
+// Copyright 1998-2017 Epic Games, Inc. All Rights Reserved.
+
+#pragma once
+
+#include "CoreMinimal.h"
+#include "UObject/ObjectMacros.h"
+#include "UObject/Object.h"
+#include "NiagaraCommon.h"
+#include "NiagaraDataSet.h"
+#include "NiagaraEmitterInstance.h"
+#include "NiagaraEmitterHandle.h"
+#include "NiagaraSystem.generated.h"
+
+//TODO:
+/** Defines a number of values and defaults shared among all Niagara Systems that use this category. */
+// UCLASS()
+// class NIAGARA_API UNiagaraSystemCategory : public UObject
+// {
+// 	GENERATED_UCLASS_BODY()
+// 
+// 	/** Name of this System Category. */
+// 	UPROPERTY(EditAnywhere, Category=Category)
+// 	FName Name;
+// 
+// 	/** Category specific to the global instance of a particular Niagara Parameter Collection. */
+// 	UPROPERTY(EditAnywhere, Category=Category)
+// 	TArray<UNiagaraParameterCollectionInstance*> ParameterCollectionOverrides;
+// 
+// 	/** TODO: How to do an icon? */
+// 	//UPROPERTY(EditAnywhere, Category = Category)
+// 	//UTexture2D* Icon
+// };
+
+USTRUCT()
+struct FNiagaraEmitterSpawnAttributes
+{
+	GENERATED_USTRUCT_BODY()
+
+	UPROPERTY()
+	TArray<FName> SpawnAttributes;
+};
+
+/** Container for multiple emitters that combine together to create a particle system effect.*/
+UCLASS()
+class NIAGARA_API UNiagaraSystem : public UObject
+{
+	GENERATED_UCLASS_BODY()
+
+public:
+	//~ UObject interface
+	void PostInitProperties();
+	void Serialize(FArchive& Ar)override;
+	virtual void PostLoad() override;
+
+	/** Gets an array of the emitter handles. */
+	const TArray<FNiagaraEmitterHandle>& GetEmitterHandles();
+
+	/** Returns true if this system is valid and can be instanced. False otherwise. */
+	bool IsValid()const;
+
+#if WITH_EDITORONLY_DATA
+	/** Adds a new emitter handle to this System.  The new handle exposes an Instance value which is a copy of the
+		original asset. */
+	FNiagaraEmitterHandle AddEmitterHandle(const UNiagaraEmitter& SourceEmitter, FName EmitterName);
+
+	/** Adds a new emitter handle to this System.  The new handle will not copy the emitter and any changes made to it's
+		Instance value will modify the original asset.  This should only be used in the emitter toolkit for simulation
+		purposes. */
+	FNiagaraEmitterHandle AddEmitterHandleWithoutCopying(UNiagaraEmitter& Emitter);
+
+	/** Duplicates an existing emitter handle and adds it to the System.  The new handle will reference the same source asset,
+		but will have a copy of the duplicated Instance value. */
+	FNiagaraEmitterHandle DuplicateEmitterHandle(const FNiagaraEmitterHandle& EmitterHandleToDuplicate, FName EmitterName);
+#endif
+
+	/** Removes the provided emitter handle. */
+	void RemoveEmitterHandle(const FNiagaraEmitterHandle& EmitterHandleToDelete);
+
+	/** Removes the emitter handles which have an Id in the supplied set. */
+	void RemoveEmitterHandlesById(const TSet<FGuid>& HandlesToRemove);
+
+	FNiagaraEmitterHandle& GetEmitterHandle(int Idx)
+	{
+		check(Idx < EmitterHandles.Num());
+		return EmitterHandles[Idx];
+	};
+
+	const FNiagaraEmitterHandle& GetEmitterHandle(int Idx) const
+	{
+		check(Idx < EmitterHandles.Num());
+		return EmitterHandles[Idx];
+	};
+
+	int GetNumEmitters()
+	{
+		return EmitterHandles.Num();
+	}
+
+	/** From the last compile, what are the variables that were exported out of the system for external use?*/
+	const FNiagaraParameterStore& GetExposedParameters() const {	return ExposedParameters; }
+	FNiagaraParameterStore& GetExposedParameters()  { return ExposedParameters; }
+
+
+	/** Gets the System script which is used to populate the System parameters and parameter bindings. */
+	UNiagaraScript* GetSystemSpawnScript(bool bSolo=false);
+	UNiagaraScript* GetSystemUpdateScript(bool bSolo=false);
+
+#if WITH_EDITORONLY_DATA
+	/** Get whether or not we auto-import applied changes. See bAutoImportChangedEmitters member documentation for details.*/
+	bool GetAutoImportChangedEmitters() const;
+	/** Set whether or not we auto-import applied changes. See bAutoImportChangedEmitters member documentation for details.*/
+	void SetAutoImportChangedEmitters(bool bShouldImport);
+		
+	/** Called to query whether or not this emitter is referenced as the source to any emitter handles for this System.*/
+	bool ReferencesSourceEmitter(UNiagaraEmitter* Emitter);
+
+	/** Goes through all handles on the System and compares them to the latest on their source. Refreshes from source if needed. Returns true if any were refreshed.*/
+	bool ResynchronizeAllHandles();
+
+	/** Recompile all emitters and the System script associated with this System.*/
+	void Compile();
+
+	/** Compile just the System script associated with this System.*/
+	void CompileScripts(TArray<ENiagaraScriptCompileStatus>& OutScriptStatuses, TArray<FString>& OutGraphLevelErrorMessages, TArray<FString>& PathNames, TArray<UNiagaraScript*>& Scripts);
+
+	/** Uses whether or not the GetAutoImportChangedEmitters is set to ResynchronizeAllHandles.*/
+	bool CheckForUpdates();
+
+	/** Gets editor specific data stored with this system. */
+	UObject* GetEditorData();
+
+	/** Gets editor specific data stored with this system. */
+	const UObject* GetEditorData() const;
+
+	/** Sets editor specific data stored with this system. */
+	void SetEditorData(UObject* InEditorData);
+
+#endif
+
+	const TArray<FNiagaraEmitterSpawnAttributes>& GetEmitterSpawnAttributes()const {	return EmitterSpawnAttributes;	};
+
+protected:
+	/** Handles to the emitter this System will simulate. */
+	UPROPERTY(VisibleAnywhere, Category = "Emitters")
+	TArray<FNiagaraEmitterHandle> EmitterHandles;
+
+// 	/** Category of this system. */
+// 	UPROPERTY(EditAnywhere, Category = System)
+// 	UNiagaraSystemCategory* Category;
+
+	/** The script which defines the System parameters, and which generates the bindings from System
+		parameter to emitter parameter. */
+	UPROPERTY()
+	UNiagaraScript* SystemSpawnScript;
+
+	/** The script which defines the System parameters, and which generates the bindings from System
+	parameter to emitter parameter. */
+	UPROPERTY()
+	UNiagaraScript* SystemUpdateScript;
+
+	/** Spawn script compiled to be run individually on a single instance of the system rather than batched as the main spawn script */
+	UPROPERTY()
+	UNiagaraScript* SystemSpawnScriptSolo;
+
+	/** Update script compiled to be run individually on a single instance of the system rather than batched as the main spawn script */
+	UPROPERTY()
+	UNiagaraScript* SystemUpdateScriptSolo;
+
+	/** Attribute names in the data set that are driving each emitter's spawning. */
+	UPROPERTY()
+	TArray<FNiagaraEmitterSpawnAttributes> EmitterSpawnAttributes;
+
+	/** Variables exposed to the outside work for tweaking*/
+	UPROPERTY()
+	FNiagaraParameterStore ExposedParameters;
+
+#if WITH_EDITORONLY_DATA
+	/** Systems are the final step in the process of creating a Niagara system. Artists may wish to lock an System so that it only uses
+	the handle's cached version of the scripts, rather than the external assets that may be subject to changes. If this flag is set, we will only update
+	the emitters if told to do so explicitly by the user.*/
+	UPROPERTY()
+	uint32 bAutoImportChangedEmitters : 1;	
+
+	/** Data used by the editor to maintain UI state etc.. */
+	UPROPERTY()
+	UObject* EditorData;
+#endif
+
+	void InitEmitterSpawnAttributes();
+};
diff --git a/Engine/Plugins/FX/Niagara/Source/Niagara/Niagara.Build.cs b/Engine/Plugins/FX/Niagara/Source/Niagara/Niagara.Build.cs
new file mode 100644
--- /dev/null
+++ b/Engine/Plugins/FX/Niagara/Source/Niagara/Niagara.Build.cs
@@ -0,0 +1,59 @@
+// Copyright 1998-2017 Epic Games, Inc. All Rights Reserved.
+
+using UnrealBuildTool;
+
+public class Niagara : ModuleRules
+{
+    public Niagara(ReadOnlyTargetRules Target) : base(Target)
+    {
+        PrivateDependencyModuleNames.AddRange(
+            new string[] {
+                "NiagaraShader",
+                "Core",
+                "Engine",
+				"RenderCore",
+                "UtilityShaders",
+                "ShaderCore"
+            }
+        );
+
+
+        PublicDependencyModuleNames.AddRange(
+            new string[] {
+                "NiagaraShader",
+                "MovieScene",
+				"MovieSceneTracks",
+				"CoreUObject",
+                "VectorVM",
+                "RHI",
+                "UtilityShaders",
+                "NiagaraVertexFactories",
+                "ShaderCore"
+            }
+        );
+
+
+        PrivateIncludePaths.AddRange(
+            new string[] {
+            })
+        ;
+
+        // If we're compiling with the engine, then add Core's engine dependencies
+        if (Target.bCompileAgainstEngine == true)
+        {
+            if (!Target.bBuildRequiresCookedData)
+            {
+                DynamicallyLoadedModuleNames.AddRange(new string[] { "DerivedDataCache" });
+            }
+        }
+
+
+        if (Target.bBuildEditor == true)
+        {
+            PrivateDependencyModuleNames.AddRange(
+                new string[] {
+                "TargetPlatform",
+            });
+        }
+    }
+}
diff --git a/Engine/Plugins/FX/Niagara/Source/Niagara/Private/NiagaraCollision.cpp b/Engine/Plugins/FX/Niagara/Source/Niagara/Private/NiagaraCollision.cpp
new file mode 100644
--- /dev/null
+++ b/Engine/Plugins/FX/Niagara/Source/Niagara/Private/NiagaraCollision.cpp
@@ -0,0 +1,161 @@
+// Copyright 1998-2016 Epic Games, Inc. All Rights Reserved.
+
+#include "NiagaraCollision.h"
+#include "NiagaraEmitterInstance.h"
+#include "NiagaraSystemInstance.h"
+#include "NiagaraStats.h"
+#include "NiagaraComponent.h"
+
+DECLARE_CYCLE_STAT(TEXT("Collision"), STAT_NiagaraCollision, STATGROUP_Niagara);
+DECLARE_CYCLE_STAT(TEXT("Event Emission"), STAT_NiagaraEventWrite, STATGROUP_Niagara);
+
+void FNiagaraCollisionBatch::KickoffNewBatch(FNiagaraEmitterInstance *Sim, float DeltaSeconds)
+{
+	SCOPE_CYCLE_COUNTER(STAT_NiagaraCollision);
+	FNiagaraVariable PosVar(FNiagaraTypeDefinition::GetVec3Def(), "Position");
+	FNiagaraVariable VelVar(FNiagaraTypeDefinition::GetVec3Def(), "Velocity");
+	FNiagaraVariable TstVar(FNiagaraTypeDefinition::GetBoolDef(), "PerformCollision");
+	FNiagaraDataSetIterator<FVector> PosIt(Sim->GetData(), PosVar, 0, false);
+	FNiagaraDataSetIterator<FVector> VelIt(Sim->GetData(), VelVar, 0, false);
+	//FNiagaraDataSetIterator<int32> TstIt(Sim->GetData(), TstVar, 0, false);
+
+	if (!PosIt.IsValid() || !VelIt.IsValid() /*|| !TstIt.IsValid()*/)
+	{
+		return;
+	}
+
+	UWorld *SystemWorld = Sim->GetParentSystemInstance()->GetComponent()->GetWorld();
+	if (SystemWorld)
+	{
+		CollisionTraces.Empty();
+
+		for (uint32 i = 0; i < Sim->GetData().GetPrevNumInstances(); i++)
+		{
+			//int32 TestCollision = *TstIt;
+			//if (TestCollision)
+			{
+				check(PosIt.IsValid() && VelIt.IsValid());
+				FVector Position;
+				FVector Velocity;
+				PosIt.Get(Position);
+				VelIt.Get(Velocity);
+
+				FCollisionQueryParams QueryParams(SCENE_QUERY_STAT(NiagraAsync));
+				QueryParams.OwnerTag = "Niagara";
+				FTraceHandle Handle = SystemWorld->AsyncLineTraceByChannel(EAsyncTraceType::Single, Position, Position + Velocity*DeltaSeconds, ECollisionChannel::ECC_WorldStatic, QueryParams, FCollisionResponseParams::DefaultResponseParam, nullptr, i);
+				FNiagaraCollisionTrace Trace;
+				Trace.CollisionTraceHandle = Handle;
+				Trace.SourceParticleIndex = i;
+				Trace.OriginalVelocity = Velocity;
+				CollisionTraces.Add(Trace);
+			}
+
+			//TstIt.Advance();
+			PosIt.Advance();
+			VelIt.Advance();
+		}
+	}
+}
+
+void FNiagaraCollisionBatch::GenerateEventsFromResults(FNiagaraEmitterInstance *Sim)
+{
+	SCOPE_CYCLE_COUNTER(STAT_NiagaraCollision);
+	//CollisionEventDataSet.Allocate(CollisionTraceHandles.Num());
+
+	UWorld *SystemWorld = Sim->GetParentSystemInstance()->GetComponent()->GetWorld();
+	if (SystemWorld)
+	{
+		TArray<FNiagaraCollisionEventPayload> Payloads;
+
+		// generate events for last frame's collisions
+		//
+		for (FNiagaraCollisionTrace CurCheck: CollisionTraces)
+		{
+			FTraceHandle Handle = CurCheck.CollisionTraceHandle;
+			FTraceDatum CurTrace;
+			// wait for trace handles; this should block rarely to never
+			bool bReady = false;
+			while (!bReady)
+			{
+				bReady = SystemWorld->QueryTraceData(Handle, CurTrace);
+				if (!bReady)
+				{
+					// if the query came back false, it's possible that the hanle is invalid for some reason; skip in that case
+					// TODO: handle this more gracefully
+					if (!SystemWorld->IsTraceHandleValid(Handle, false))
+					{
+						break;
+					}
+					break;
+				}
+			}
+
+			if (bReady && CurTrace.OutHits.Num())
+			{
+				// grab the first hit that blocks
+				FHitResult *Hit = FHitResult::GetFirstBlockingHit(CurTrace.OutHits);
+				if (Hit && Hit->IsValidBlockingHit())
+				{
+					FNiagaraCollisionEventPayload Event;
+					Event.CollisionNormal = Hit->ImpactNormal;
+					Event.CollisionPos = Hit->ImpactPoint;
+					Event.CollisionVelocity = CurCheck.OriginalVelocity;
+					Event.ParticleIndex = CurTrace.UserData;
+					check(!Event.CollisionNormal.ContainsNaN());
+					check(Event.CollisionNormal.IsNormalized());
+					check(!Event.CollisionPos.ContainsNaN());
+					check(!Event.CollisionVelocity.ContainsNaN());
+
+					// TODO add to unique list of physical materials for Blueprint
+					Event.PhysicalMaterialIndex = 0;// Hit->PhysMaterial->GetUniqueID();
+
+					Payloads.Add(Event);
+				}
+			}
+		}
+
+		if (Payloads.Num())
+		{
+			// now allocate the data set and write all the event structs
+			//
+			CollisionEventDataSet->Allocate(Payloads.Num());
+			CollisionEventDataSet->SetNumInstances(Payloads.Num());
+			//FNiagaraVariable ValidVar(FNiagaraTypeDefinition::GetIntDef(), "Valid");
+			FNiagaraVariable PosVar(FNiagaraTypeDefinition::GetVec3Def(), "CollisionLocation");
+			FNiagaraVariable VelVar(FNiagaraTypeDefinition::GetVec3Def(), "CollisionVelocity");
+			FNiagaraVariable NormVar(FNiagaraTypeDefinition::GetVec3Def(), "CollisionNormal");
+			FNiagaraVariable PhysMatIdxVar(FNiagaraTypeDefinition::GetIntDef(), "PhysicalMaterialIndex");
+			FNiagaraVariable ParticleIndexVar(FNiagaraTypeDefinition::GetIntDef(), "ParticleIndex");
+			//FNiagaraDataSetIterator<int32> ValidItr(*CollisionEventDataSet, ValidVar, 0, true);
+			FNiagaraDataSetIterator<FVector> PosItr(*CollisionEventDataSet, PosVar, 0, true);
+			FNiagaraDataSetIterator<FVector> NormItr(*CollisionEventDataSet, NormVar, 0, true);
+			FNiagaraDataSetIterator<FVector> VelItr(*CollisionEventDataSet, VelVar, 0, true);
+			FNiagaraDataSetIterator<int32> PhysMatItr(*CollisionEventDataSet, PhysMatIdxVar, 0, true);
+			FNiagaraDataSetIterator<int32> ParticleIndexItr(*CollisionEventDataSet, ParticleIndexVar, 0, true);
+
+			for (FNiagaraCollisionEventPayload &Payload : Payloads)
+			{
+				SCOPE_CYCLE_COUNTER(STAT_NiagaraEventWrite);
+
+				check(/*ValidItr.IsValid() && */PosItr.IsValid() && VelItr.IsValid() && NormItr.IsValid() && PhysMatItr.IsValid());
+				//ValidItr.Set(1);
+				PosItr.Set(Payload.CollisionPos);
+				VelItr.Set(Payload.CollisionVelocity);
+				NormItr.Set(Payload.CollisionNormal);
+				ParticleIndexItr.Set(Payload.ParticleIndex);
+				PhysMatItr.Set(0);
+				//ValidItr.Advance();
+				PosItr.Advance();
+				VelItr.Advance();
+				NormItr.Advance();
+				PhysMatItr.Advance();
+				ParticleIndexItr.Advance();
+			}
+		}
+		else
+		{
+			CollisionEventDataSet->SetNumInstances(0);
+		}
+
+	}
+}
\ No newline at end of file
diff --git a/Engine/Plugins/FX/Niagara/Source/Niagara/Private/NiagaraCommon.cpp b/Engine/Plugins/FX/Niagara/Source/Niagara/Private/NiagaraCommon.cpp
new file mode 100644
--- /dev/null
+++ b/Engine/Plugins/FX/Niagara/Source/Niagara/Private/NiagaraCommon.cpp
@@ -0,0 +1,161 @@
+// Copyright 1998-2016 Epic Games, Inc. All Rights Reserved.
+
+#include "NiagaraCommon.h"
+#include "NiagaraDataSet.h"
+#include "NiagaraComponent.h"
+#include "NiagaraSystemInstance.h"
+#include "NiagaraParameterCollection.h"
+#include "NiagaraComponent.h"
+#include "Linker.h"
+#include "Class.h"
+#include "Package.h"
+
+//////////////////////////////////////////////////////////////////////////
+
+FString FNiagaraTypeHelper::ToString(const uint8* ValueData, const UScriptStruct* Struct)
+{
+	FString Ret;
+	if (Struct == FNiagaraTypeDefinition::GetFloatStruct())
+	{
+		Ret += FString::Printf(TEXT("%g "), *(float*)ValueData);
+	}
+	else if (Struct == FNiagaraTypeDefinition::GetIntStruct())
+	{
+		Ret += FString::Printf(TEXT("%d "), *(int32*)ValueData);
+	}
+	else if (Struct == FNiagaraTypeDefinition::GetBoolStruct())
+	{
+		int32 Val = *(int32*)ValueData;
+		Ret += Val == 0xFFFFFFFF ? (TEXT("True")) : ( Val == 0x0 ? TEXT("False") : TEXT("Invalid"));		
+	}
+	else
+	{
+		for (TFieldIterator<UProperty> PropertyIt(Struct, EFieldIteratorFlags::IncludeSuper); PropertyIt; ++PropertyIt)
+		{
+			const UProperty* Property = *PropertyIt;
+			const uint8* PropPtr = ValueData + PropertyIt->GetOffset_ForInternal();
+			if (Property->IsA(UFloatProperty::StaticClass()))
+			{
+				Ret += FString::Printf(TEXT("%s: %g "), *Property->GetNameCPP(), *(float*)PropPtr);
+			}
+			else if (Property->IsA(UIntProperty::StaticClass()))
+			{
+				Ret += FString::Printf(TEXT("%s: %d "), *Property->GetNameCPP(), *(int32*)PropPtr);
+			}
+			else if (Property->IsA(UBoolProperty::StaticClass()))
+			{
+				int32 Val = *(int32*)ValueData;
+				FString BoolStr = Val == 0xFFFFFFFF ? (TEXT("True")) : (Val == 0x0 ? TEXT("False") : TEXT("Invalid"));
+				Ret += FString::Printf(TEXT("%s: %d "), *Property->GetNameCPP(), *BoolStr);
+			}
+			else if (const UStructProperty* StructProp = CastChecked<UStructProperty>(Property))
+			{
+				Ret += FString::Printf(TEXT("%s: (%s) "), *Property->GetNameCPP(), *FNiagaraTypeHelper::ToString(PropPtr, StructProp->Struct));
+			}
+			else
+			{
+				check(false);
+				Ret += TEXT("Unknown Type");
+			}
+		}
+	}
+	return Ret;
+}
+//////////////////////////////////////////////////////////////////////////
+
+FNiagaraSystemUpdateContext::~FNiagaraSystemUpdateContext()
+{
+	for (UNiagaraComponent* Comp : ComponentsToReInit)
+	{
+		Comp->ReinitializeSystem();
+	}
+	for (UNiagaraComponent* Comp : ComponentsToReset)
+	{
+		Comp->ResetSystem();
+	}
+}
+
+void FNiagaraSystemUpdateContext::Add(const UNiagaraSystem* System, bool bReInit)
+{
+	for (TObjectIterator<UNiagaraComponent> It; It; ++It)
+	{
+		UNiagaraComponent* Comp = *It;
+		check(Comp);
+		if (Comp->GetAsset() == System)
+		{
+			AddInternal(Comp, bReInit);
+		}
+	}
+}
+#if WITH_EDITORONLY_DATA
+
+void FNiagaraSystemUpdateContext::Add(const UNiagaraEmitter* Emitter, bool bReInit)
+{
+	for (TObjectIterator<UNiagaraComponent> It; It; ++It)
+	{
+		UNiagaraComponent* Comp = *It;
+		check(Comp);
+		FNiagaraSystemInstance* SystemInst = Comp->GetSystemInstance();
+		if (SystemInst && SystemInst->UsesEmitter(Emitter))
+		{
+			AddInternal(Comp, bReInit);
+		}		
+	}
+}
+
+void FNiagaraSystemUpdateContext::Add(const UNiagaraScript* Script, bool bReInit)
+{
+	for (TObjectIterator<UNiagaraComponent> It; It; ++It)
+	{
+		UNiagaraComponent* Comp = *It;
+		check(Comp);
+		FNiagaraSystemInstance* SystemInst = Comp->GetSystemInstance();
+		if (SystemInst && SystemInst->UsesScript(Script))
+		{
+			AddInternal(Comp, bReInit);
+		}
+	}
+}
+
+// void FNiagaraSystemUpdateContext::Add(UNiagaraDataInterface* Interface, bool bReInit)
+// {
+// 	for (TObjectIterator<UNiagaraComponent> It; It; ++It)
+// 	{
+// 		UNiagaraComponent* Comp = *It;
+// 		check(Comp);		
+// 		if (FNiagaraSystemInstance* SystemInst = Comp->GetSystemInstance())
+// 		{
+// 			if (SystemInst->ContainsDataInterface(Interface))
+// 			{
+// 				AddInternal(SystemInst, bReInit);
+// 			}
+// 		}
+// 	}
+// }
+
+void FNiagaraSystemUpdateContext::Add(const UNiagaraParameterCollection* Collection, bool bReInit)
+{
+	for (TObjectIterator<UNiagaraComponent> It; It; ++It)
+	{
+		UNiagaraComponent* Comp = *It;
+		check(Comp);
+		FNiagaraSystemInstance* SystemInst = Comp->GetSystemInstance();
+		if (SystemInst && SystemInst->UsesCollection(Collection))
+		{
+			AddInternal(Comp, bReInit);
+		}
+	}
+}
+#endif
+
+void FNiagaraSystemUpdateContext::AddInternal(UNiagaraComponent* Comp, bool bReInit)
+{
+	if (bReInit)
+	{
+		ComponentsToReInit.AddUnique(Comp);
+	}
+	else
+	{
+		ComponentsToReset.AddUnique(Comp);
+	}
+}
\ No newline at end of file
diff --git a/Engine/Plugins/FX/Niagara/Source/Niagara/Private/NiagaraComponent.cpp b/Engine/Plugins/FX/Niagara/Source/Niagara/Private/NiagaraComponent.cpp
new file mode 100644
--- /dev/null
+++ b/Engine/Plugins/FX/Niagara/Source/Niagara/Private/NiagaraComponent.cpp
@@ -0,0 +1,922 @@
+// Copyright 1998-2017 Epic Games, Inc. All Rights Reserved.
+
+#include "NiagaraComponent.h"
+#include "VectorVM.h"
+#include "NiagaraRenderer.h"
+#include "NiagaraSystem.h"
+#include "NiagaraSystemInstance.h"
+#include "NiagaraEmitterInstance.h"
+#include "MeshBatch.h"
+#include "SceneUtils.h"
+#include "ComponentReregisterContext.h"
+#include "NiagaraConstants.h"
+#include "NiagaraStats.h"
+#include "NiagaraCommon.h"
+#include "NiagaraEmitterInstance.h"
+#include "NiagaraDataInterface.h"
+#include "NiagaraDataInterfaceStaticMesh.h"
+#include "NameTypes.h"
+#include "NiagaraParameterCollection.h"
+#include "NiagaraWorldManager.h"
+
+DECLARE_CYCLE_STAT(TEXT("Gen Verts"),STAT_NiagaraGenerateVertices,STATGROUP_Niagara);
+
+DEFINE_LOG_CATEGORY(LogNiagara);
+
+FNiagaraSceneProxy::FNiagaraSceneProxy(const UNiagaraComponent* InComponent)
+		: FPrimitiveSceneProxy(InComponent)
+		, bRenderingEnabled(true)
+{
+	// In this case only, update the System renderers on the game thread.
+	check(IsInGameThread());
+	FNiagaraSystemInstance* SystemInst = InComponent->GetSystemInstance();
+	if (SystemInst)
+	{
+		TArray<TSharedRef<FNiagaraEmitterInstance> > Sims = SystemInst->GetEmitters();
+		TArray<NiagaraRenderer*> RenderersFromSims;
+		for (TSharedRef<FNiagaraEmitterInstance> Sim : Sims)
+		{
+			for (int32 i = 0; i < Sim->GetEmitterRendererNum(); i++)
+			{
+				RenderersFromSims.Add(Sim->GetEmitterRenderer(i));
+			}
+		}
+		UpdateEmitterRenderers(RenderersFromSims);
+		//UE_LOG(LogNiagara, Warning, TEXT("FNiagaraSceneProxy %p"), this);
+
+		bAlwaysHasVelocity = true;
+	}
+}
+
+void FNiagaraSceneProxy::AddEmitterRenderer(NiagaraRenderer* Renderer)
+{
+	EmitterRenderers.Add(Renderer); 
+	//UE_LOG(LogNiagara, Warning, TEXT("FNiagaraSceneProxy::AddEmitterRenderer %p"), Renderer);
+}
+
+void FNiagaraSceneProxy::UpdateEmitterRenderers(TArray<NiagaraRenderer*>& InRenderers)
+{
+	EmitterRenderers.Empty();
+	for (NiagaraRenderer* EmitterRenderer : InRenderers)
+	{
+		AddEmitterRenderer(EmitterRenderer);
+	}
+}
+
+FNiagaraSceneProxy::~FNiagaraSceneProxy()
+{
+	//UE_LOG(LogNiagara, Warning, TEXT("~FNiagaraSceneProxy %p"), this);
+	ReleaseRenderThreadResources();
+}
+
+/** Called on render thread to assign new dynamic data */
+void FNiagaraSceneProxy::SetDynamicData_RenderThread(FNiagaraDynamicDataBase* NewDynamicData)
+{
+	for (NiagaraRenderer* Renderer : EmitterRenderers)
+	{
+		if (Renderer)
+		{
+			Renderer->SetDynamicData_RenderThread(NewDynamicData);
+		}
+	}
+	return;
+}
+
+
+void FNiagaraSceneProxy::ReleaseRenderThreadResources()
+{
+	for (NiagaraRenderer* Renderer : EmitterRenderers)
+	{
+		if (Renderer)
+		{
+			Renderer->ReleaseRenderThreadResources();
+		}
+	}
+	return;
+}
+
+// FPrimitiveSceneProxy interface.
+void FNiagaraSceneProxy::CreateRenderThreadResources()
+{
+	for (NiagaraRenderer* Renderer : EmitterRenderers)
+	{
+		if (Renderer)
+		{
+			Renderer->CreateRenderThreadResources();
+		}
+	}
+	return;
+}
+
+void FNiagaraSceneProxy::OnTransformChanged()
+{
+	//WorldSpacePrimitiveUniformBuffer.ReleaseResource();
+}
+
+FPrimitiveViewRelevance FNiagaraSceneProxy::GetViewRelevance(const FSceneView* View) const
+{
+	FPrimitiveViewRelevance Relevance;
+
+	if (bRenderingEnabled == false)
+	{
+		return Relevance;
+	}
+	Relevance.bDynamicRelevance = true;
+
+	for (NiagaraRenderer* Renderer : EmitterRenderers)
+	{
+		if (Renderer)
+		{
+			Relevance |= Renderer->GetViewRelevance(View, this);
+		}
+	}
+	return Relevance;
+}
+
+
+uint32 FNiagaraSceneProxy::GetMemoryFootprint() const
+{ 
+	return (sizeof(*this) + GetAllocatedSize()); 
+}
+
+uint32 FNiagaraSceneProxy::GetAllocatedSize() const
+{ 
+	uint32 DynamicDataSize = 0;
+	for (NiagaraRenderer* Renderer : EmitterRenderers)
+	{
+		if (Renderer)
+		{
+			DynamicDataSize += Renderer->GetDynamicDataSize();
+		}
+	}
+	return FPrimitiveSceneProxy::GetAllocatedSize() + DynamicDataSize;
+}
+
+bool FNiagaraSceneProxy::GetRenderingEnabled() const
+{
+	return bRenderingEnabled;
+}
+
+void FNiagaraSceneProxy::SetRenderingEnabled(bool bInRenderingEnabled)
+{
+	bRenderingEnabled = bInRenderingEnabled;
+}
+
+void FNiagaraSceneProxy::GetDynamicMeshElements(const TArray<const FSceneView*>& Views, const FSceneViewFamily& ViewFamily, uint32 VisibilityMap, FMeshElementCollector& Collector) const
+{
+	for (NiagaraRenderer* Renderer : EmitterRenderers)
+	{
+		if (Renderer)
+		{
+			Renderer->GetDynamicMeshElements(Views, ViewFamily, VisibilityMap, Collector, this);
+		}
+	}
+}
+
+
+
+void FNiagaraSceneProxy::GatherSimpleLights(const FSceneViewFamily& ViewFamily, FSimpleLightArray& OutParticleLights) const
+{
+	NiagaraRendererLights *LightRenderer = nullptr;
+	FNiagaraDynamicDataLights *DynamicData = nullptr;
+	for (int32 Idx = 0; Idx < EmitterRenderers.Num(); Idx++)
+	{
+		NiagaraRenderer *Renderer = EmitterRenderers[Idx];
+		if (Renderer && Renderer->GetPropertiesClass() == UNiagaraLightRendererProperties::StaticClass())
+		{
+			LightRenderer = static_cast<NiagaraRendererLights*>(Renderer);
+			DynamicData = static_cast<FNiagaraDynamicDataLights*>(Renderer->GetDynamicData());
+			break;
+		}
+	}
+
+
+	if (DynamicData)
+	{
+		int32 LightCount = DynamicData->LightArray.Num();
+		
+		OutParticleLights.InstanceData.Reserve(LightCount);
+		OutParticleLights.PerViewData.Reserve(LightCount);
+
+		for (NiagaraRendererLights::SimpleLightData &LightData : DynamicData->LightArray)
+		{
+			// When not using camera-offset, output one position for all views to share. 
+			OutParticleLights.PerViewData.Add(LightData.PerViewEntry);
+
+			// Add an entry for the light instance.
+			OutParticleLights.InstanceData.Add(LightData.LightEntry);
+		}
+	}
+
+}
+
+
+//////////////////////////////////////////////////////////////////////////
+
+UNiagaraComponent::UNiagaraComponent(const FObjectInitializer& ObjectInitializer)
+	: Super(ObjectInitializer)
+	, bForceSolo(false)
+	, AgeUpdateMode(EAgeUpdateMode::TickDeltaTime)
+	, DesiredAge(0.0f)
+	, SeekDelta(1 / 30.0f)
+{
+	PrimaryComponentTick.bCanEverTick = true;
+	PrimaryComponentTick.TickGroup = TG_DuringPhysics;
+	PrimaryComponentTick.bStartWithTickEnabled = false;
+	PrimaryComponentTick.SetTickFunctionEnable(false);
+	bTickInEditor = true;
+	bAutoActivate = true;
+	bRenderingEnabled = true;
+}
+
+
+void UNiagaraComponent::TickComponent(float DeltaSeconds, enum ELevelTick TickType, FActorComponentTickFunction* ThisTickFunction)
+{
+	check(SystemInstance->IsSolo());
+	if (bIsActive && SystemInstance.Get())
+	{
+		if (SceneProxy != nullptr)
+		{
+			FNiagaraSceneProxy* NiagaraProxy = static_cast<FNiagaraSceneProxy*>(SceneProxy);
+			NiagaraProxy->SetRenderingEnabled(bRenderingEnabled);
+		}
+
+		if (AgeUpdateMode == EAgeUpdateMode::TickDeltaTime)
+		{
+			SystemInstance->ComponentTick(DeltaSeconds);
+		}
+		else
+		{
+			float AgeDiff = DesiredAge - SystemInstance->GetAge();
+			if (FMath::Abs(AgeDiff) < KINDA_SMALL_NUMBER)
+			{
+				AgeDiff = 0.0f;
+			}
+
+			if (AgeDiff < 0.0f)
+			{
+				SystemInstance->Reset(FNiagaraSystemInstance::EResetMode::ImmediateReset);
+				AgeDiff = DesiredAge - SystemInstance->GetAge();
+			}
+
+			if (AgeDiff > 0.0f)
+			{
+				// TODO: This maximum should be configurable, and we should not draw the System until we are caught up.
+				static const float MaxSimTime = 1;
+				//Take SeekDelta to be the max delta time we'll accept.
+				FNiagaraSystemSimulation* SystemSim = GetSystemSimulation();
+
+				float SimTime = FMath::Min(AgeDiff, MaxSimTime);
+				int32 NumSteps = FMath::CeilToInt(SimTime / SeekDelta);
+				//Num steps must be odd. This is because when doing > 1 we have to flip the solo data set for in the system sim. This has to end up in the same flip state it was in before our ticks.
+				NumSteps = NumSteps % 2 == 0 ? NumSteps + 1 : NumSteps;
+				float RealSeekDelta = SimTime / (float)NumSteps;
+
+				for (int32 i = NumSteps - 1; i >= 0; --i)
+				{
+					SystemInstance->ComponentTick(RealSeekDelta);
+					if (i > 0)
+					{
+						SystemSim->TickSoloDataSet();
+					}
+				}
+			}
+		}
+	}
+
+	UpdateComponentToWorld();
+	MarkRenderDynamicDataDirty();
+}
+
+const UObject* UNiagaraComponent::AdditionalStatObject() const
+{
+	return Asset;
+}
+
+void UNiagaraComponent::ResetSystem()
+{
+	Activate(true);
+}
+
+void UNiagaraComponent::ReinitializeSystem()
+{
+	DestroyInstance();
+	Activate();
+
+// 	if (IsRegistered() && SystemInstance.IsValid())
+// 	{
+// 		SystemInstance->Init(Asset, true, bForceSolo);
+// 	}
+}
+
+bool UNiagaraComponent::GetRenderingEnabled() const
+{
+	return bRenderingEnabled;
+}
+
+void UNiagaraComponent::SetRenderingEnabled(bool bInRenderingEnabled)
+{
+	bRenderingEnabled = bInRenderingEnabled;
+}
+
+void UNiagaraComponent::Activate(bool bReset /* = false */)
+{
+	if (Asset == nullptr || Asset->IsValid() == false)
+	{
+		DestroyInstance();
+
+		UE_LOG(LogNiagara, Warning, TEXT("Failed to activate Niagara Component!"));
+		SetComponentTickEnabled(false);
+		return;
+	}
+
+	if(!IsRegistered())
+	{
+		return;
+	}
+	
+	Super::Activate(bReset);
+
+	if (SystemInstance.IsValid() == false)
+	{
+		if (FNiagaraSystemSimulation* SystemSim = GetSystemSimulation())
+		{
+			SystemInstance = MakeUnique<FNiagaraSystemInstance>(this);
+#if WITH_EDITORONLY_DATA
+			OnSystemInstanceChangedDelegate.Broadcast();
+#endif
+			SystemInstance->Init(SystemSim, true, bForceSolo);
+			bReset = false;
+		}
+		else
+		{
+			return;
+		}
+	}
+
+	SystemInstance->Activate(bReset);
+
+	/** We only need to tick the component if we require solo mode. */
+	SetComponentTickEnabled(SystemInstance->IsSolo());
+}
+
+void UNiagaraComponent::Deactivate()
+{
+	Super::Deactivate();
+
+	SetComponentTickEnabled(false);
+	bIsActive = false;
+
+	if (SystemInstance)
+	{
+		SystemInstance->Deactivate();
+	}
+}
+
+void UNiagaraComponent::DestroyInstance()
+{
+	SystemInstance = nullptr;
+#if WITH_EDITORONLY_DATA
+	OnSystemInstanceChangedDelegate.Broadcast();
+#endif
+}
+
+void UNiagaraComponent::OnRegister()
+{
+	Super::OnRegister();
+}
+
+void UNiagaraComponent::OnComponentDestroyed(bool bDestroyingHierarchy)
+{
+	//UE_LOG(LogNiagara, Log, TEXT("OnComponentDestroyed %p %p"), this, SystemInstance.Get());
+	DestroyInstance();
+
+	Super::OnComponentDestroyed(bDestroyingHierarchy);
+}
+
+void UNiagaraComponent::OnUnregister()
+{
+	DestroyInstance();
+	Super::OnUnregister();
+}
+
+void UNiagaraComponent::BeginDestroy()
+{
+	DestroyInstance();
+
+	Super::BeginDestroy();
+}
+
+class FNiagaraSystemSimulation* UNiagaraComponent::GetSystemSimulation()
+{
+	UWorld* World = GetWorld();
+	if (World)
+	{
+		return FNiagaraWorldManager::Get(World)->GetSystemSimulation(Asset);
+	}
+	return nullptr;
+}
+
+void UNiagaraComponent::SendRenderDynamicData_Concurrent()
+{
+	if (SystemInstance.IsValid() && SceneProxy)
+	{
+		SystemInstance->GetSystemBounds().Init();
+		
+		FNiagaraSceneProxy* NiagaraProxy = static_cast<FNiagaraSceneProxy*>(SceneProxy);
+
+		for (int32 i = 0; i < SystemInstance->GetEmitters().Num(); i++)
+		{
+			TSharedPtr<FNiagaraEmitterInstance> Emitter = SystemInstance->GetEmitters()[i];
+			for (int32 EmitterIdx = 0; EmitterIdx < Emitter->GetEmitterRendererNum(); EmitterIdx++)
+			{
+				NiagaraRenderer* Renderer = Emitter->GetEmitterRenderer(EmitterIdx);
+				if (Renderer)
+				{
+					if (Emitter->IsEnabled())
+					{
+						FNiagaraDynamicDataBase* DynamicData = Renderer->GenerateVertexData(NiagaraProxy, Emitter->GetData(), Emitter->GetEmitterHandle().GetInstance()->SimTarget);
+
+						ENQUEUE_UNIQUE_RENDER_COMMAND_TWOPARAMETER(
+							FSendNiagaraDynamicData,
+							NiagaraRenderer*, EmitterRenderer, Emitter->GetEmitterRenderer(EmitterIdx),
+							FNiagaraDynamicDataBase*, DynamicData, DynamicData,
+							{
+								EmitterRenderer->SetDynamicData_RenderThread(DynamicData);
+							});
+					}
+					else
+					{
+						ENQUEUE_UNIQUE_RENDER_COMMAND_ONEPARAMETER(
+							FSendNiagaraDynamicData,
+							NiagaraRenderer*, EmitterRenderer, Emitter->GetEmitterRenderer(EmitterIdx),
+							{
+								EmitterRenderer->SetDynamicData_RenderThread(nullptr);
+							});
+					}
+				}
+			}
+		}
+	}
+
+}
+
+int32 UNiagaraComponent::GetNumMaterials() const
+{
+	return 0;
+}
+
+
+FBoxSphereBounds UNiagaraComponent::CalcBounds(const FTransform& LocalToWorld) const
+{
+	FBox SimBounds(ForceInit);
+	if (SystemInstance.IsValid())
+	{
+		SystemInstance->GetSystemBounds().Init();
+		for (int32 i = 0; i < SystemInstance->GetEmitters().Num(); i++)
+		{
+			FNiagaraEmitterInstance &Sim = *(SystemInstance->GetEmitters()[i]);
+			SystemInstance->GetSystemBounds() += Sim.GetCachedBounds();
+		}
+		FBoxSphereBounds BSBounds(SystemInstance->GetSystemBounds());
+
+		return BSBounds;
+	}
+	return FBoxSphereBounds(SimBounds);
+}
+
+FPrimitiveSceneProxy* UNiagaraComponent::CreateSceneProxy()
+{
+	// The constructor will set up the System renderers from the component.
+	FNiagaraSceneProxy* Proxy = new FNiagaraSceneProxy(this);
+	return Proxy;
+}
+
+void UNiagaraComponent::GetUsedMaterials(TArray<UMaterialInterface*>& OutMaterials, bool bGetDebugMaterials) const
+{
+	if (!SystemInstance.IsValid())
+	{
+		return;
+	}
+
+	for (TSharedRef<FNiagaraEmitterInstance> Sim : SystemInstance->GetEmitters())	
+	{
+		if (UNiagaraEmitter* Props = Sim->GetEmitterHandle().GetInstance())
+		{	
+			for (int32 i = 0; i < Props->RendererProperties.Num(); i++)
+			{
+				if (UNiagaraRendererProperties* Renderer = Props->RendererProperties[i])
+				{
+					Renderer->GetUsedMaterials(OutMaterials);
+				}
+			}
+		}
+	}
+}
+
+FNiagaraSystemInstance* UNiagaraComponent::GetSystemInstance() const
+{
+	return SystemInstance.Get();
+}
+
+void UNiagaraComponent::OnSystemDisabled()
+{
+	//TODO: Add KillOnComplete option to asset to kill the component.
+
+}
+
+void UNiagaraComponent::SetNiagaraVariableVec4(const FString& InVariableName, const FVector4& InValue)
+{
+	FName VarName = FName(*InVariableName);
+
+	SystemInstance->GetParameters().SetParameterValue(InValue, FNiagaraVariable(FNiagaraTypeDefinition::GetVec4Def(), VarName));
+}
+
+void UNiagaraComponent::SetNiagaraVariableVec3(const FString& InVariableName, FVector InValue)
+{
+	FName VarName = FName(*InVariableName);
+	
+	SystemInstance->GetParameters().SetParameterValue(InValue, FNiagaraVariable(FNiagaraTypeDefinition::GetVec3Def(), VarName));
+}
+
+void UNiagaraComponent::SetNiagaraVariableVec2(const FString& InVariableName, FVector2D InValue)
+{
+	FName VarName = FName(*InVariableName);
+
+	SystemInstance->GetParameters().SetParameterValue(InValue, FNiagaraVariable(FNiagaraTypeDefinition::GetVec2Def(),VarName));
+}
+
+void UNiagaraComponent::SetNiagaraVariableFloat(const FString& InVariableName, float InValue)
+{
+	FName VarName = FName(*InVariableName);
+
+	SystemInstance->GetParameters().SetParameterValue(InValue, FNiagaraVariable(FNiagaraTypeDefinition::GetFloatDef(), VarName));
+}
+
+void UNiagaraComponent::SetNiagaraVariableBool(const FString& InVariableName, bool InValue)
+{
+	FName VarName = FName(*InVariableName);
+
+	SystemInstance->GetParameters().SetParameterValue(InValue ? FNiagaraBool::True : FNiagaraBool::False, FNiagaraVariable(FNiagaraTypeDefinition::GetBoolDef(), VarName));
+}
+
+TArray<FVector> UNiagaraComponent::GetNiagaraParticlePositions_DebugOnly(const FString& InEmitterName)
+{
+	return GetNiagaraParticleValueVec3_DebugOnly(InEmitterName, TEXT("Position"));
+}
+
+TArray<FVector> UNiagaraComponent::GetNiagaraParticleValueVec3_DebugOnly(const FString& InEmitterName, const FString& InValueName)
+{
+	TArray<FVector> Positions;
+	FName EmitterName = FName(*InEmitterName);
+	if (SystemInstance.IsValid())
+	{
+		for (TSharedRef<FNiagaraEmitterInstance>& Sim : SystemInstance->GetEmitters())
+		{
+			if (Sim->GetEmitterHandle().GetName() == EmitterName)
+			{
+				int32 NumParticles = Sim->GetData().GetNumInstances();
+				Positions.SetNum(NumParticles);
+				FNiagaraDataSetIterator<FVector> PosItr(Sim->GetData(), FNiagaraVariable(FNiagaraTypeDefinition::GetVec3Def(), *InValueName));
+				int32 i = 0;
+				while (PosItr.IsValid())
+				{
+					FVector Position;
+					PosItr.Get(Position);
+					Positions[i] = Position;
+					i++;
+					PosItr.Advance();
+				}
+			}
+		}
+	}
+	return Positions;
+
+}
+
+TArray<float> UNiagaraComponent::GetNiagaraParticleValues_DebugOnly(const FString& InEmitterName, const FString& InValueName)
+{
+
+	TArray<float> Values;
+	FName EmitterName = FName(*InEmitterName);
+	if (SystemInstance.IsValid())
+	{
+		for (TSharedRef<FNiagaraEmitterInstance>& Sim : SystemInstance->GetEmitters())
+		{
+			if (Sim->GetEmitterHandle().GetName() == EmitterName)
+			{
+				int32 NumParticles = Sim->GetData().GetNumInstances();
+				Values.SetNum(NumParticles);
+				FNiagaraDataSetIterator<float> ValueItr(Sim->GetData(), FNiagaraVariable(FNiagaraTypeDefinition::GetFloatDef(), *InValueName));
+				int32 i = 0;
+				while (ValueItr.IsValid())
+				{
+					float Value;
+					ValueItr.Get(Value);
+					Values[i] = Value;
+					i++;
+					ValueItr.Advance();
+				}
+			}
+		}
+	}
+	return Values;
+}
+
+void UNiagaraComponent::PostLoad()
+{
+	Super::PostLoad();
+	if (Asset)
+	{
+		Asset->ConditionalPostLoad();
+#if WITH_EDITOR
+		SynchronizeWithSourceSystem();
+#endif
+	}
+}
+
+#if WITH_EDITOR
+
+void UNiagaraComponent::PostEditChangeProperty(FPropertyChangedEvent& PropertyChangedEvent)
+{
+	Super::PostEditChangeProperty(PropertyChangedEvent);
+
+	FName PropertyName;
+	if (PropertyChangedEvent.Property)
+	{
+		PropertyName = PropertyChangedEvent.Property->GetFName();
+	}
+
+	if (PropertyName == GET_MEMBER_NAME_CHECKED(UNiagaraComponent, Asset) ||
+		PropertyName == GET_MEMBER_NAME_CHECKED(UNiagaraComponent, InitialParameters))
+	{
+		SynchronizeWithSourceSystem();
+	}
+
+	ReinitializeSystem();
+}
+
+
+bool UNiagaraComponent::SynchronizeWithSourceSystem()
+{
+	//TODO: Look through params in system in "Owner" namespace and add to our parameters.
+	if (Asset == nullptr)
+	{
+		return false;
+	}
+
+	TArray<FNiagaraVariable> SourceVars;
+	Asset->GetExposedParameters().GetParameters(SourceVars);
+
+	bool bEditsMade = false;
+	for (FNiagaraVariable& Param : SourceVars)
+	{
+		bEditsMade |= InitialParameters.AddParameter(Param, true);
+	}
+
+	TArray<FNiagaraVariable> ExistingVars;
+	InitialParameters.GetParameters(ExistingVars);
+
+	for (FNiagaraVariable ExistingVar : ExistingVars)
+	{
+		if (!SourceVars.Contains(ExistingVar))
+		{
+			InitialParameters.RemoveParameter(ExistingVar);
+			EditorOverridesValue.Remove(ExistingVar.GetName());
+		}
+	}
+
+	for (FNiagaraVariable ExistingVar : ExistingVars)
+	{
+		bool* FoundVar = EditorOverridesValue.Find(ExistingVar.GetName());
+
+		if (!IsParameterValueOverriddenLocally(ExistingVar.GetName()))
+		{
+			Asset->GetExposedParameters().CopyParameterData(InitialParameters, ExistingVar);
+		}
+	}
+
+	InitialParameters.Rebind();
+
+	//////////////////////////////////////////////////////////////////////////
+	//OLD IMPL
+// 	if (Asset == nullptr)
+// 	{
+// 		return false; 
+// 	}
+// 
+// 	UNiagaraScript* Script = Asset->GetSystemSpawnScript();
+// 	if (Script == nullptr)
+// 	{
+// 		return false;
+// 	}
+// 
+// 	bool bEditsMade = false;
+// 	{
+// 		TArray<int32> VarIndicesToRemove;
+// 
+// 		// Check over all of our overrides to see if they still match up with the source System.
+// 		for (int32 i = 0; i < SystemParameterLocalOverrides.Num(); i++)
+// 		{
+// 			FNiagaraVariable* Variable = Script->Parameters.FindParameter(SystemParameterLocalOverrides[i].GetId());
+// 			// If the variable still exists, keep digging... otherwise remove it.
+// 			if (Variable)
+// 			{
+// 				if (Variable->GetType() != SystemParameterLocalOverrides[i].GetType())
+// 				{
+// 					UE_LOG(LogNiagara, Warning, TEXT("Variable '%s' types changed.. possibly losing override value."), *Variable->GetName().ToString());
+// 					VarIndicesToRemove.Add(i);
+// 				}
+// 				else if (Variable->GetName() != SystemParameterLocalOverrides[i].GetName())
+// 				{
+// 					if (bEditsMade == false)
+// 					{
+// 						bEditsMade = true;
+// 						Modify();
+// 					}
+// 					SystemParameterLocalOverrides[i].SetName(Variable->GetName());
+// 				}
+// 			}
+// 			else
+// 			{
+// 				VarIndicesToRemove.Add(i);
+// 			}
+// 		}
+// 
+// 		// Handle the variables that we want to remove.
+// 		if (VarIndicesToRemove.Num() != 0)
+// 		{
+// 			if (bEditsMade == false)
+// 			{
+// 				bEditsMade = true;
+// 				Modify();
+// 			}
+// 
+// 			for (int32 removeIdx = VarIndicesToRemove.Num() - 1; removeIdx >= 0; removeIdx--)
+// 			{
+// 				UE_LOG(LogNiagara, Warning, TEXT("Removing local override '%s'."), *SystemParameterLocalOverrides[VarIndicesToRemove[removeIdx]].GetName().ToString());
+// 
+// 				SystemParameterLocalOverrides.RemoveAt(VarIndicesToRemove[removeIdx]);
+// 			}
+// 		}
+// 	}
+// 
+// 	{
+// 		TArray<int32> VarIndicesToRemove;
+// 
+// 		// Check over all of our overrides to see if they still match up with the source System.
+// 		for (int32 i = 0; i < SystemDataInterfaceLocalOverrides.Num(); i++)
+// 		{
+// 			FNiagaraScriptDataInterfaceInfo* Variable = Script->DataInterfaceInfo.FindByPredicate([&](const FNiagaraScriptDataInterfaceInfo& Var)
+// 			{
+// 				return Var.Id == SystemDataInterfaceLocalOverrides[i].Id;
+// 			}); 
+// 			// If the variable still exists, keep digging... otherwise remove it.
+// 			if (Variable)
+// 			{
+// 				if (Variable->DataInterface->GetClass() != SystemDataInterfaceLocalOverrides[i].DataInterface->GetClass())
+// 				{
+// 					UE_LOG(LogNiagara, Warning, TEXT("Variable '%s' types changed.. losing override value."), *Variable->Name.ToString());
+// 					VarIndicesToRemove.Add(i);
+// 				}
+// 				else if (Variable->Name != SystemDataInterfaceLocalOverrides[i].Name)
+// 				{
+// 					if (bEditsMade == false)
+// 					{
+// 						bEditsMade = true;
+// 						Modify();
+// 					}
+// 					SystemDataInterfaceLocalOverrides[i].Name = Variable->Name;
+// 				}
+// 			}
+// 			else
+// 			{
+// 				VarIndicesToRemove.Add(i);
+// 			}
+// 		}
+// 
+// 		// Handle the variables that we want to remove.
+// 		if (VarIndicesToRemove.Num() != 0)
+// 		{
+// 			if (bEditsMade == false)
+// 			{
+// 				bEditsMade = true;
+// 				Modify();
+// 			}
+// 
+// 			for (int32 removeIdx = VarIndicesToRemove.Num() - 1; removeIdx >= 0; removeIdx--)
+// 			{
+// 				UE_LOG(LogNiagara, Warning, TEXT("Removing local override '%s'."), *SystemDataInterfaceLocalOverrides[VarIndicesToRemove[removeIdx]].Name.ToString());
+// 
+// 				SystemDataInterfaceLocalOverrides.RemoveAt(VarIndicesToRemove[removeIdx]);
+// 			}
+// 		}
+// 	}
+// 
+// 	return bEditsMade;
+	return false;
+}
+
+UNiagaraComponent::EAgeUpdateMode UNiagaraComponent::GetAgeUpdateMode() const
+{
+	return AgeUpdateMode;
+}
+
+void UNiagaraComponent::SetAgeUpdateMode(EAgeUpdateMode InAgeUpdateMode)
+{
+	AgeUpdateMode = InAgeUpdateMode;
+}
+
+float UNiagaraComponent::GetDesiredAge() const
+{
+	return DesiredAge;
+}
+
+void UNiagaraComponent::SetDesiredAge(float InDesiredAge)
+{
+	DesiredAge = InDesiredAge;
+}
+
+float UNiagaraComponent::GetSeekDelta() const
+{
+	return SeekDelta;
+}
+
+void UNiagaraComponent::SetSeekDelta(float InSeekDelta)
+{
+	SeekDelta = InSeekDelta;
+}
+
+
+bool UNiagaraComponent::IsParameterValueOverriddenLocally(const FName& InParamName)
+{
+	bool* FoundVar = EditorOverridesValue.Find(InParamName);
+
+	if (FoundVar != nullptr && *(FoundVar))
+	{
+		return true;
+	}
+	return false;
+}
+
+void UNiagaraComponent::SetParameterValueOverriddenLocally(const FName& InParamName, bool bInOverriden)
+{
+	bool* FoundVar = EditorOverridesValue.Find(InParamName);
+
+	if (FoundVar != nullptr && bInOverriden) 
+	{
+		*(FoundVar) = bInOverriden;
+	}
+	else if (FoundVar == nullptr && bInOverriden)			
+	{
+		EditorOverridesValue.Add(InParamName, true);
+	}
+	else
+	{
+		EditorOverridesValue.Remove(InParamName);
+
+		UNiagaraScript* ScriptSpawn = Asset->GetSystemSpawnScript(true);
+		UNiagaraScript* ScriptUpdate = Asset->GetSystemUpdateScript(true);
+		if (ScriptSpawn != nullptr)
+		{
+			for (FNiagaraVariable Var : ScriptSpawn->Parameters.Parameters)
+			{
+				if (Var.GetName() == InParamName)
+				{
+					InitialParameters.SetParameter(Var);
+					return;
+				}
+			}			
+		}
+		if (ScriptUpdate != nullptr)
+		{
+			for (FNiagaraVariable Var : ScriptUpdate->Parameters.Parameters)
+			{
+				if (Var.GetName() == InParamName)
+				{
+					InitialParameters.SetParameter(Var);
+					return;
+				}
+			}
+		}
+	}
+}
+
+
+#endif // WITH_EDITOR
+
+
+
+void UNiagaraComponent::SetAsset(UNiagaraSystem* InAsset)
+{
+	Asset = InAsset;
+
+#if WITH_EDITOR
+	SynchronizeWithSourceSystem();
+#endif
+
+	ReinitializeSystem();
+}
+
+
+
