Copying //UE4/Orion-Staging to //UE4/Main (Source: //Orion/Dev-General @ 3606378)

diff --git a/Engine/Plugins/FX/Niagara/Source/Niagara/Private/NiagaraConstants.cpp b/Engine/Plugins/FX/Niagara/Source/Niagara/Private/NiagaraConstants.cpp
new file mode 100644
--- /dev/null
+++ b/Engine/Plugins/FX/Niagara/Source/Niagara/Private/NiagaraConstants.cpp
@@ -0,0 +1,213 @@
+// Copyright 1998-2017 Epic Games, Inc. All Rights Reserved.
+
+#include "NiagaraConstants.h"
+
+#include "Internationalization.h"
+
+#define LOCTEXT_NAMESPACE "FNiagaraConstants"
+
+const TArray<FNiagaraVariable>& FNiagaraConstants::GetEngineConstants()
+{
+	static TArray<FNiagaraVariable> SystemParameters;
+	if (SystemParameters.Num() == 0)
+	{
+		SystemParameters.Add(SYS_PARAM_ENGINE_DELTA_TIME);
+		SystemParameters.Add(SYS_PARAM_ENGINE_INV_DELTA_TIME);
+		SystemParameters.Add(SYS_PARAM_ENGINE_POSITION);
+		SystemParameters.Add(SYS_PARAM_ENGINE_VELOCITY);
+		SystemParameters.Add(SYS_PARAM_ENGINE_X_AXIS);
+		SystemParameters.Add(SYS_PARAM_ENGINE_Y_AXIS);
+		SystemParameters.Add(SYS_PARAM_ENGINE_Z_AXIS);
+		SystemParameters.Add(SYS_PARAM_ENGINE_LOCAL_TO_WORLD);
+		SystemParameters.Add(SYS_PARAM_ENGINE_WORLD_TO_LOCAL);
+		SystemParameters.Add(SYS_PARAM_ENGINE_LOCAL_TO_WORLD_TRANSPOSED);
+		SystemParameters.Add(SYS_PARAM_ENGINE_WORLD_TO_LOCAL_TRANSPOSED);
+		SystemParameters.Add(SYS_PARAM_ENGINE_MIN_DIST_TO_CAMERA);
+		SystemParameters.Add(SYS_PARAM_ENGINE_EXEC_COUNT);
+		SystemParameters.Add(SYS_PARAM_ENGINE_EMITTER_NUM_PARTICLES);
+		SystemParameters.Add(SYS_PARAM_ENGINE_SYSTEM_NUM_EMITTERS_ALIVE);
+		SystemParameters.Add(SYS_PARAM_ENGINE_SYSTEM_NUM_EMITTERS);
+	}
+	return SystemParameters;
+}
+
+FNiagaraVariable FNiagaraConstants::UpdateEngineConstant(const FNiagaraVariable& InVar)
+{
+	const FNiagaraVariable* FoundSystemVar = FindEngineConstant(InVar);
+	if (nullptr != FoundSystemVar)
+	{
+		return *FoundSystemVar;
+	}
+	else
+	{
+		static TMap<FName, FNiagaraVariable> UpdatedSystemParameters;
+		if (UpdatedSystemParameters.Num() == 0)
+		{
+			UpdatedSystemParameters.Add(FName(TEXT("System Delta Time")), SYS_PARAM_ENGINE_DELTA_TIME);
+			UpdatedSystemParameters.Add(FName(TEXT("System Inv Delta Time")), SYS_PARAM_ENGINE_INV_DELTA_TIME);
+			UpdatedSystemParameters.Add(FName(TEXT("System Position")), SYS_PARAM_ENGINE_POSITION);
+			UpdatedSystemParameters.Add(FName(TEXT("System Velocity")), SYS_PARAM_ENGINE_VELOCITY);
+			UpdatedSystemParameters.Add(FName(TEXT("System X Axis")), SYS_PARAM_ENGINE_X_AXIS);
+			UpdatedSystemParameters.Add(FName(TEXT("System Y Axis")), SYS_PARAM_ENGINE_Y_AXIS);
+			UpdatedSystemParameters.Add(FName(TEXT("System Z Axis")), SYS_PARAM_ENGINE_Z_AXIS);
+			UpdatedSystemParameters.Add(FName(TEXT("System Local To World")), SYS_PARAM_ENGINE_LOCAL_TO_WORLD);
+			UpdatedSystemParameters.Add(FName(TEXT("System World To Local")), SYS_PARAM_ENGINE_WORLD_TO_LOCAL);
+			UpdatedSystemParameters.Add(FName(TEXT("System Local To World Transposed")), SYS_PARAM_ENGINE_LOCAL_TO_WORLD_TRANSPOSED);
+			UpdatedSystemParameters.Add(FName(TEXT("System World To Local Transposed")), SYS_PARAM_ENGINE_WORLD_TO_LOCAL_TRANSPOSED);
+
+			UpdatedSystemParameters.Add(FName(TEXT("Emitter Execution Count")), SYS_PARAM_ENGINE_EXEC_COUNT);
+			UpdatedSystemParameters.Add(FName(TEXT("Emitter Age")), SYS_PARAM_EMITTER_AGE);
+			UpdatedSystemParameters.Add(FName(TEXT("Emitter Spawn Rate")), SYS_PARAM_EMITTER_SPAWNRATE);
+			UpdatedSystemParameters.Add(FName(TEXT("Emitter Spawn Interval")), SYS_PARAM_EMITTER_SPAWN_INTERVAL);
+			UpdatedSystemParameters.Add(FName(TEXT("Emitter Interp Spawn Start Dt")), SYS_PARAM_EMITTER_INTERP_SPAWN_START_DT);
+
+			UpdatedSystemParameters.Add(FName(TEXT("Delta Time")), SYS_PARAM_ENGINE_DELTA_TIME);
+			UpdatedSystemParameters.Add(FName(TEXT("Emitter Age")), SYS_PARAM_EMITTER_AGE);
+			UpdatedSystemParameters.Add(FName(TEXT("Effect Position")), SYS_PARAM_ENGINE_POSITION);
+			UpdatedSystemParameters.Add(FName(TEXT("Effect Velocity")), SYS_PARAM_ENGINE_VELOCITY);
+			UpdatedSystemParameters.Add(FName(TEXT("Effect X Axis")), SYS_PARAM_ENGINE_X_AXIS);
+			UpdatedSystemParameters.Add(FName(TEXT("Effect Y Axis")), SYS_PARAM_ENGINE_Y_AXIS);
+			UpdatedSystemParameters.Add(FName(TEXT("Effect Z Axis")), SYS_PARAM_ENGINE_Z_AXIS);
+
+			UpdatedSystemParameters.Add(FName(TEXT("Effect Local To World")), SYS_PARAM_ENGINE_LOCAL_TO_WORLD);
+			UpdatedSystemParameters.Add(FName(TEXT("Effect World To Local")), SYS_PARAM_ENGINE_WORLD_TO_LOCAL);
+			UpdatedSystemParameters.Add(FName(TEXT("Effect Local To World Transposed")), SYS_PARAM_ENGINE_LOCAL_TO_WORLD_TRANSPOSED);
+			UpdatedSystemParameters.Add(FName(TEXT("Effect World To Local Transposed")), SYS_PARAM_ENGINE_WORLD_TO_LOCAL_TRANSPOSED);
+			UpdatedSystemParameters.Add(FName(TEXT("Execution Count")), SYS_PARAM_ENGINE_EXEC_COUNT);
+			UpdatedSystemParameters.Add(FName(TEXT("Spawn Rate")), SYS_PARAM_EMITTER_SPAWNRATE);
+			UpdatedSystemParameters.Add(FName(TEXT("Spawn Interval")), SYS_PARAM_EMITTER_SPAWN_INTERVAL);
+			UpdatedSystemParameters.Add(FName(TEXT("Interp Spawn Start Dt")), SYS_PARAM_EMITTER_INTERP_SPAWN_START_DT);
+			UpdatedSystemParameters.Add(FName(TEXT("Inv Delta Time")), SYS_PARAM_ENGINE_INV_DELTA_TIME);
+		}
+
+		const FNiagaraVariable* FoundSystemVarUpdate = UpdatedSystemParameters.Find(InVar.GetName());
+		if (FoundSystemVarUpdate != nullptr)
+		{
+			return *FoundSystemVarUpdate;
+		}
+	}
+	return InVar;
+
+}
+
+
+const FNiagaraVariable* FNiagaraConstants::FindEngineConstant(const FNiagaraVariable& InVar)
+{
+	const TArray<FNiagaraVariable>& SystemParameters = GetEngineConstants();
+	const FNiagaraVariable* FoundSystemVar = SystemParameters.FindByPredicate([&](const FNiagaraVariable& Var)
+	{
+		return Var.GetName() == InVar.GetName();
+	});
+	return FoundSystemVar;
+}
+
+FText FNiagaraConstants::GetEngineConstantDescription(const FNiagaraVariable& InAttribute)
+{
+	static TMap<FNiagaraVariable, FText> AttrStrMap;
+	if (AttrStrMap.Num() == 0)
+	{
+		AttrStrMap.Add(SYS_PARAM_ENGINE_DELTA_TIME, LOCTEXT("EngineDeltaTimeDesc", "Time in seconds since the last tick."));
+		AttrStrMap.Add(SYS_PARAM_ENGINE_INV_DELTA_TIME, LOCTEXT("EngineInvDeltaTimeDesc", "One over Engine.DeltaTime"));
+		AttrStrMap.Add(SYS_PARAM_ENGINE_POSITION, LOCTEXT("EnginePositionDesc", "The owning component's position in world space."));
+		AttrStrMap.Add(SYS_PARAM_ENGINE_VELOCITY, LOCTEXT("EngineVelocityDesc", "The owning component's velocity in world space."));
+		AttrStrMap.Add(SYS_PARAM_ENGINE_X_AXIS, LOCTEXT("XAxisDesc", "The X-axis of the owning component."));
+		AttrStrMap.Add(SYS_PARAM_ENGINE_Y_AXIS, LOCTEXT("YAxisDesc", "The Y-axis of the owning component."));
+		AttrStrMap.Add(SYS_PARAM_ENGINE_Z_AXIS, LOCTEXT("ZAxisDesc", "The Z-axis of the owning component."));
+		AttrStrMap.Add(SYS_PARAM_ENGINE_LOCAL_TO_WORLD, LOCTEXT("LocalToWorldDesc", "Conversion of the owning component's local to world space."));
+		AttrStrMap.Add(SYS_PARAM_ENGINE_WORLD_TO_LOCAL, LOCTEXT("WorldToLocalDesc", "Conversion of the owning component's world to local space."));
+		AttrStrMap.Add(SYS_PARAM_ENGINE_LOCAL_TO_WORLD_TRANSPOSED, LOCTEXT("LocalToWorldTransposeDesc", "Conversion of the owning component's local to world space transposed."));
+		AttrStrMap.Add(SYS_PARAM_ENGINE_WORLD_TO_LOCAL_TRANSPOSED, LOCTEXT("WorldToLocalTransposeDesc", "Conversion of the owning component's world to local space transposed."));
+		AttrStrMap.Add(SYS_PARAM_ENGINE_MIN_DIST_TO_CAMERA, LOCTEXT("MinDistanceToCamera", "The distance from the owner component to the nearest local player viewpoint."));
+		AttrStrMap.Add(SYS_PARAM_ENGINE_EXEC_COUNT, LOCTEXT("ExecCountDesc", "The index of this particle in the read buffer."));
+		AttrStrMap.Add(SYS_PARAM_ENGINE_EMITTER_NUM_PARTICLES, LOCTEXT("EmitterNumParticles", "The number of particles for this emitter at the beginning of simulation. Should only be used in Emitter scripts."));
+		AttrStrMap.Add(SYS_PARAM_ENGINE_SYSTEM_NUM_EMITTERS_ALIVE, LOCTEXT("SystemNumEmittersAlive", "The number of emitters still alive attached to this system. Should only be used in System scripts."));
+		AttrStrMap.Add(SYS_PARAM_ENGINE_SYSTEM_NUM_EMITTERS, LOCTEXT("SystemNumEmitters", "The number of emitters attached to this system. Should only be used in System scripts."));
+	}
+
+	FText* FoundStr = AttrStrMap.Find(InAttribute);
+	if (FoundStr != nullptr && !FoundStr->IsEmpty())
+	{
+		return *FoundStr;
+	}
+	return FText();
+}
+
+const TArray<FNiagaraVariable>& FNiagaraConstants::GetCommonParticleAttributes()
+{
+	static TArray<FNiagaraVariable> Attributes;
+	if (Attributes.Num() == 0)
+	{
+		Attributes.Add(SYS_PARAM_PARTICLES_POSITION);
+		Attributes.Add(SYS_PARAM_PARTICLES_VELOCITY);
+		Attributes.Add(SYS_PARAM_PARTICLES_COLOR);
+		Attributes.Add(SYS_PARAM_PARTICLES_SPRITE_ROTATION);
+		Attributes.Add(SYS_PARAM_PARTICLES_NORMALIZED_AGE);
+		Attributes.Add(SYS_PARAM_PARTICLES_SPRITE_SIZE);
+		Attributes.Add(SYS_PARAM_PARTICLES_SPRITE_FACING);
+		Attributes.Add(SYS_PARAM_PARTICLES_SPRITE_ALIGNMENT);
+		Attributes.Add(SYS_PARAM_PARTICLES_SUB_IMAGE_INDEX);
+		Attributes.Add(SYS_PARAM_PARTICLES_DYNAMIC_MATERIAL_PARAM);
+		Attributes.Add(SYS_PARAM_PARTICLES_SCALE);
+		Attributes.Add(SYS_PARAM_PARTICLES_ROTATION);
+
+		Attributes.Add(SYS_PARAM_INSTANCE_ALIVE);
+	}
+	return Attributes;
+}
+
+FString FNiagaraConstants::GetAttributeDefaultValue(const FNiagaraVariable& InAttribute)
+{
+	static TMap<FNiagaraVariable, FString> AttrStrMap;
+	if (AttrStrMap.Num() == 0)
+	{
+		AttrStrMap.Add(SYS_PARAM_PARTICLES_POSITION, SYS_PARAM_ENGINE_POSITION.GetName().ToString());
+		AttrStrMap.Add(SYS_PARAM_PARTICLES_VELOCITY, TEXT("0.0,0.0,100.0"));
+		AttrStrMap.Add(SYS_PARAM_PARTICLES_COLOR, FLinearColor(1.0f, 1.0f, 1.0f, 1.0f).ToString());
+		AttrStrMap.Add(SYS_PARAM_PARTICLES_SPRITE_ROTATION, TEXT("0.0"));
+		AttrStrMap.Add(SYS_PARAM_PARTICLES_NORMALIZED_AGE, TEXT("0.0"));
+		AttrStrMap.Add(SYS_PARAM_PARTICLES_SPRITE_SIZE, TEXT("X=50.0 Y=50.0"));
+		AttrStrMap.Add(SYS_PARAM_PARTICLES_SPRITE_FACING, TEXT("1.0,0.0,0.0"));
+		AttrStrMap.Add(SYS_PARAM_PARTICLES_SPRITE_ALIGNMENT, TEXT("1.0,0.0,0.0"));
+		AttrStrMap.Add(SYS_PARAM_PARTICLES_SUB_IMAGE_INDEX, TEXT("0.0"));
+		AttrStrMap.Add(SYS_PARAM_PARTICLES_DYNAMIC_MATERIAL_PARAM, TEXT("1.0,1.0,1.0,1.0"));
+		AttrStrMap.Add(SYS_PARAM_PARTICLES_SCALE, TEXT("1.0,1.0,1.0"));
+		AttrStrMap.Add(SYS_PARAM_PARTICLES_ROTATION, TEXT("0.0,0.0,0.0"));
+	}
+
+	FString* FoundStr = AttrStrMap.Find(InAttribute);
+	if (FoundStr != nullptr && !FoundStr->IsEmpty())
+	{
+		return *FoundStr;
+	}
+	return FString();
+}
+
+FText FNiagaraConstants::GetAttributeDescription(const FNiagaraVariable& InAttribute)
+{
+	static TMap<FNiagaraVariable, FText> AttrStrMap;
+	if (AttrStrMap.Num() == 0)
+	{
+		AttrStrMap.Add(SYS_PARAM_PARTICLES_POSITION, LOCTEXT("PositionDesc", "The position of the particle."));
+		AttrStrMap.Add(SYS_PARAM_PARTICLES_VELOCITY, LOCTEXT("VelocityDesc", "The velocity in cm/s of the particle."));
+		AttrStrMap.Add(SYS_PARAM_PARTICLES_COLOR, LOCTEXT("ColorDesc", "The color of the particle."));
+		AttrStrMap.Add(SYS_PARAM_PARTICLES_SPRITE_ROTATION, LOCTEXT("SpriteRotDesc", "The screen aligned roll of the particle."));
+		AttrStrMap.Add(SYS_PARAM_PARTICLES_NORMALIZED_AGE, LOCTEXT("NormalizedAgeDesc", "The age in seconds divided by lifetime in seconds. Useful for animation as the value is between 0 and 1."));
+		AttrStrMap.Add(SYS_PARAM_PARTICLES_SPRITE_SIZE, LOCTEXT("SpriteSizeDesc", "The size of the sprite quad."));
+		AttrStrMap.Add(SYS_PARAM_PARTICLES_SPRITE_FACING, LOCTEXT("FacingDesc", "Makes the surface of the sprite face towards a custom vector."));
+		AttrStrMap.Add(SYS_PARAM_PARTICLES_SPRITE_ALIGNMENT, LOCTEXT("AlignmentDesc", "Imagine the texture having an arrow to the right, this attribute makes the arrow point towards the alignment axis."));
+		AttrStrMap.Add(SYS_PARAM_PARTICLES_SUB_IMAGE_INDEX, LOCTEXT("SubImageIndexDesc", "The position in the UV lookup table divided by total number of UV slots."));
+		AttrStrMap.Add(SYS_PARAM_PARTICLES_DYNAMIC_MATERIAL_PARAM, LOCTEXT("DynamicMaterialParameterDesc", "The 4-float vector used to send custom data to renderer."));
+		AttrStrMap.Add(SYS_PARAM_PARTICLES_SCALE, LOCTEXT("ScaleParamDesc", "The XYZ scale of the non-sprite based particle."));
+		AttrStrMap.Add(SYS_PARAM_PARTICLES_ROTATION, LOCTEXT("RotationParamDesc", "The Vector3 euler angle rotation"));
+		AttrStrMap.Add(SYS_PARAM_INSTANCE_ALIVE, LOCTEXT("AliveParamDesc", "Used to determine whether or not this particle instance is still valid or if it can be deleted."));
+	}
+
+	FText* FoundStr = AttrStrMap.Find(InAttribute);
+	if (FoundStr != nullptr && !FoundStr->IsEmpty())
+	{
+		return *FoundStr;
+	}
+	return FText();
+}
+
+#undef LOCTEXT_NAMESPACE
\ No newline at end of file
diff --git a/Engine/Plugins/FX/Niagara/Source/Niagara/Private/NiagaraDataInterfaceBindingInstance.cpp b/Engine/Plugins/FX/Niagara/Source/Niagara/Private/NiagaraDataInterfaceBindingInstance.cpp
new file mode 100644
--- /dev/null
+++ b/Engine/Plugins/FX/Niagara/Source/Niagara/Private/NiagaraDataInterfaceBindingInstance.cpp
@@ -0,0 +1,6 @@
+// Copyright 1998-2016 Epic Games, Inc. All Rights Reserved.
+
+#include "NiagaraDataInterfaceBindingInstance.h"
+#include "NiagaraTypes.h"
+#include "NiagaraCommon.h"
+#include "NiagaraDataInterface.h"
diff --git a/Engine/Plugins/FX/Niagara/Source/Niagara/Private/NiagaraDataInterfaceColorCurve.cpp b/Engine/Plugins/FX/Niagara/Source/Niagara/Private/NiagaraDataInterfaceColorCurve.cpp
new file mode 100644
--- /dev/null
+++ b/Engine/Plugins/FX/Niagara/Source/Niagara/Private/NiagaraDataInterfaceColorCurve.cpp
@@ -0,0 +1,235 @@
+// Copyright 1998-2016 Epic Games, Inc. All Rights Reserved.
+
+#include "NiagaraDataInterfaceColorCurve.h"
+#include "Curves/CurveVector.h"
+#include "Curves/CurveLinearColor.h"
+#include "Curves/CurveFloat.h"
+#include "NiagaraTypes.h"
+
+//////////////////////////////////////////////////////////////////////////
+//Color Curve
+
+UNiagaraDataInterfaceColorCurve::UNiagaraDataInterfaceColorCurve(FObjectInitializer const& ObjectInitializer)
+	: Super(ObjectInitializer)
+{
+	UpdateLUT();
+}
+
+void UNiagaraDataInterfaceColorCurve::PostInitProperties()
+{
+	Super::PostInitProperties();
+
+	//Can we regitser data interfaces as regular types and fold them into the FNiagaraVariable framework for UI and function calls etc?
+	if (HasAnyFlags(RF_ClassDefaultObject))
+	{
+		FNiagaraTypeRegistry::Register(FNiagaraTypeDefinition(GetClass()), true, false, false);
+	}
+
+	UpdateLUT();
+}
+
+void UNiagaraDataInterfaceColorCurve::PostLoad()
+{
+	Super::PostLoad();
+	UpdateLUT();
+}
+
+#if WITH_EDITOR
+
+void UNiagaraDataInterfaceColorCurve::PostEditChangeProperty(struct FPropertyChangedEvent& PropertyChangedEvent)
+{
+	Super::PostEditChangeProperty(PropertyChangedEvent);
+
+	if (PropertyChangedEvent.Property && PropertyChangedEvent.Property->GetFName() == GET_MEMBER_NAME_CHECKED(UNiagaraDataInterfaceColorCurve, CurveToCopy))
+	{
+		UCurveLinearColor* ColorCurveAsset = Cast<UCurveLinearColor>(CurveToCopy.TryLoad());
+		if(ColorCurveAsset)
+		{
+			Modify();
+			RedCurve = ColorCurveAsset->FloatCurves[0];
+			BlueCurve = ColorCurveAsset->FloatCurves[1];
+			GreenCurve = ColorCurveAsset->FloatCurves[2];
+			AlphaCurve = ColorCurveAsset->FloatCurves[3];
+		}
+		UpdateLUT();
+	}
+}
+
+#endif
+
+void UNiagaraDataInterfaceColorCurve::UpdateLUT()
+{
+	ShaderLUT.Empty();
+	for (uint32 i = 0; i < CurveLUTWidth; i++)
+	{
+		float X = i / (float)CurveLUTWidth;
+		FLinearColor C(RedCurve.Eval(X), GreenCurve.Eval(X), BlueCurve.Eval(X), AlphaCurve.Eval(X));
+		ShaderLUT.Add(C.R);
+		ShaderLUT.Add(C.G);
+		ShaderLUT.Add(C.B);
+		ShaderLUT.Add(C.A);
+	}
+	GPUBufferDirty = true;
+}
+
+bool UNiagaraDataInterfaceColorCurve::CopyTo(UNiagaraDataInterface* Destination) const 
+{
+	if (!Super::CopyTo(Destination))
+	{
+		return false;
+	}
+	UNiagaraDataInterfaceColorCurve* DestinationColorCurve = CastChecked<UNiagaraDataInterfaceColorCurve>(Destination);
+	DestinationColorCurve->RedCurve = RedCurve;
+	DestinationColorCurve->GreenCurve = GreenCurve;
+	DestinationColorCurve->BlueCurve = BlueCurve;
+	DestinationColorCurve->AlphaCurve = AlphaCurve;
+	DestinationColorCurve->UpdateLUT();
+
+	return true;
+}
+
+bool UNiagaraDataInterfaceColorCurve::Equals(const UNiagaraDataInterface* Other) const
+{
+	if (!Super::Equals(Other))
+	{
+		return false;
+	}
+	const UNiagaraDataInterfaceColorCurve* OtherColoRedCurve = CastChecked<const UNiagaraDataInterfaceColorCurve>(Other);
+	return OtherColoRedCurve->RedCurve == RedCurve &&
+		OtherColoRedCurve->GreenCurve == GreenCurve &&
+		OtherColoRedCurve->BlueCurve == BlueCurve &&
+		OtherColoRedCurve->AlphaCurve == AlphaCurve;
+}
+
+void UNiagaraDataInterfaceColorCurve::GetCurveData(TArray<FCurveData>& OutCurveData)
+{
+	OutCurveData.Add(FCurveData(&RedCurve, TEXT("Red"), FLinearColor::Red));
+	OutCurveData.Add(FCurveData(&GreenCurve, TEXT("Green"), FLinearColor::Green));
+	OutCurveData.Add(FCurveData(&BlueCurve, TEXT("Blue"), FLinearColor::Blue));
+	OutCurveData.Add(FCurveData(&AlphaCurve, TEXT("Alpha"), FLinearColor::White));
+}
+
+void UNiagaraDataInterfaceColorCurve::GetFunctions(TArray<FNiagaraFunctionSignature>& OutFunctions)
+{
+	FNiagaraFunctionSignature Sig;
+	Sig.Name = TEXT("SampleColorCurve");
+	Sig.bMemberFunction = true;
+	Sig.bRequiresContext = false;
+	Sig.Inputs.Add(FNiagaraVariable(FNiagaraTypeDefinition(GetClass()), TEXT("ColorCurve")));
+	Sig.Inputs.Add(FNiagaraVariable(FNiagaraTypeDefinition::GetFloatDef(), TEXT("X")));
+	Sig.Outputs.Add(FNiagaraVariable(FNiagaraTypeDefinition::GetColorDef(), TEXT("Value")));
+	//Sig.Owner = *GetName();
+
+	OutFunctions.Add(Sig);
+}
+
+// build the shader function HLSL; function name is passed in, as it's defined per-DI; that way, configuration could change
+// the HLSL in the spirit of a static switch
+// TODO: need a way to identify each specific function here
+// 
+bool UNiagaraDataInterfaceColorCurve::GetFunctionHLSL(FString FunctionName, TArray<DIGPUBufferParamDescriptor> &Descriptors, FString &HLSLInterfaceID, FString &OutHLSL)
+{
+	FString BufferName = Descriptors[0].BufferParamName;
+	OutHLSL += TEXT("void ") + FunctionName + TEXT("(in float In_X, out float4 Out_Value) \n{\n");
+	OutHLSL += TEXT("\t Out_Value.x = ") + BufferName + TEXT("[(int)(In_X *") + FString::FromInt(CurveLUTWidth) + TEXT(")* 4 ];");
+	OutHLSL += TEXT("\t Out_Value.y = ") + BufferName + TEXT("[1+ (int)(In_X *") + FString::FromInt(CurveLUTWidth) + TEXT(")* 4 ];");
+	OutHLSL += TEXT("\t Out_Value.z = ") + BufferName + TEXT("[2+ (int)(In_X *") + FString::FromInt(CurveLUTWidth) + TEXT(")* 4 ];");
+	OutHLSL += TEXT("\t Out_Value.w = ") + BufferName + TEXT("[3+ (int)(In_X *") + FString::FromInt(CurveLUTWidth) + TEXT(")* 4 ];");
+	OutHLSL += TEXT("\n}\n");
+	return true;
+}
+
+// build buffer definition hlsl
+// 1. Choose a buffer name, add the data interface ID (important!)
+// 2. add a DIGPUBufferParamDescriptor to the array argument; that'll be passed on to the FNiagaraShader for binding to a shader param, that can
+// then later be found by name via FindDIBufferParam for setting; 
+// 3. store buffer declaration hlsl in OutHLSL
+// multiple buffers can be defined at once here
+//
+void UNiagaraDataInterfaceColorCurve::GetBufferDefinitionHLSL(FString DataInterfaceID, TArray<DIGPUBufferParamDescriptor> &BufferDescriptors, FString &OutHLSL)
+{
+	FString BufferName = "CurveLUT" + DataInterfaceID;
+	OutHLSL += TEXT("Buffer<float> ") + BufferName + TEXT(";\n");
+
+	BufferDescriptors.Add( DIGPUBufferParamDescriptor(BufferName, 0) );		// add a descriptor for shader parameter binding
+
+}
+
+// called after translate, to setup buffers matching the buffer descriptors generated during hlsl translation
+// need to do this because the script used during translate is a clone, including its DIs
+//
+void UNiagaraDataInterfaceColorCurve::SetupBuffers(TArray<DIGPUBufferParamDescriptor> &BufferDescriptors)
+{
+	for (DIGPUBufferParamDescriptor &Desc : BufferDescriptors)
+	{
+		FNiagaraDataInterfaceBufferData BufferData(*Desc.BufferParamName);	// store off the data for later use
+		GPUBuffers.Add(BufferData);
+	}
+}
+
+// return the GPU buffer array (called from NiagaraInstanceBatcher to get the buffers for setting to the shader)
+// we lazily update the buffer with a new LUT here if necessary
+//
+TArray<FNiagaraDataInterfaceBufferData> &UNiagaraDataInterfaceColorCurve::GetBufferDataArray()
+{
+	check(IsInRenderingThread());
+	if (GPUBufferDirty)
+	{
+		check(GPUBuffers.Num() > 0);
+
+		FNiagaraDataInterfaceBufferData &GPUBuffer = GPUBuffers[0];
+		GPUBuffer.Buffer.Release();
+		GPUBuffer.Buffer.Initialize(sizeof(float), CurveLUTWidth*4 , EPixelFormat::PF_R32_FLOAT, BUF_Volatile);	// always allocate for up to 64 data sets
+		uint32 BufferSize = ShaderLUT.Num() * sizeof(float);
+		int32 *BufferData = static_cast<int32*>(RHILockVertexBuffer(GPUBuffer.Buffer.Buffer, 0, BufferSize, EResourceLockMode::RLM_WriteOnly));
+		FPlatformMemory::Memcpy(BufferData, ShaderLUT.GetData(), BufferSize);
+		RHIUnlockVertexBuffer(GPUBuffer.Buffer.Buffer);
+		GPUBufferDirty = false;
+	}
+
+	return GPUBuffers;
+}
+
+
+
+
+DEFINE_NDI_FUNC_BINDER(UNiagaraDataInterfaceColorCurve, SampleCurve);
+FVMExternalFunction UNiagaraDataInterfaceColorCurve::GetVMExternalFunction(const FVMExternalFunctionBindingInfo& BindingInfo, void* InstanceData)
+{
+	if (BindingInfo.Name == TEXT("SampleColorCurve") && BindingInfo.GetNumInputs() == 1 && BindingInfo.GetNumOutputs() == 4)
+	{
+		return TNDIParamBinder<0, float, NDI_FUNC_BINDER(UNiagaraDataInterfaceColorCurve, SampleCurve)>::Bind(this, BindingInfo, InstanceData);
+	}
+	else
+	{
+		UE_LOG(LogNiagara, Error, TEXT("Could not find data interface external function.\n\tExpected Name: SampleColorCurve  Actual Name: %s\n\tExpected Inputs: 1  Actual Inputs: %i\n\tExpected Outputs: 4  Actual Outputs: %i"),
+			*BindingInfo.Name.ToString(), BindingInfo.GetNumInputs(), BindingInfo.GetNumOutputs());
+		return FVMExternalFunction();
+	}
+}
+
+template<typename XParamType>
+void UNiagaraDataInterfaceColorCurve::SampleCurve(FVectorVMContext& Context)
+{
+	//TODO: Create some SIMDable optimized representation of the curve to do this faster.
+	XParamType XParam(Context);
+	FRegisterHandler<float> SamplePtrR(Context);
+	FRegisterHandler<float> SamplePtrG(Context);
+	FRegisterHandler<float> SamplePtrB(Context);
+	FRegisterHandler<float> SamplePtrA(Context);
+
+	for (int32 i = 0; i < Context.NumInstances; ++i)
+	{
+		float X = XParam.Get();
+		FLinearColor C(RedCurve.Eval(X), GreenCurve.Eval(X), BlueCurve.Eval(X), AlphaCurve.Eval(X));
+		*SamplePtrR.GetDest() = C.R;
+		*SamplePtrG.GetDest() = C.G;
+		*SamplePtrB.GetDest() = C.B;
+		*SamplePtrA.GetDest() = C.A;
+		XParam.Advance();
+		SamplePtrR.Advance();
+		SamplePtrG.Advance();
+		SamplePtrB.Advance();
+		SamplePtrA.Advance();
+	}
+}
diff --git a/Engine/Plugins/FX/Niagara/Source/Niagara/Private/NiagaraDataInterfaceCurlNoise.cpp b/Engine/Plugins/FX/Niagara/Source/Niagara/Private/NiagaraDataInterfaceCurlNoise.cpp
new file mode 100644
--- /dev/null
+++ b/Engine/Plugins/FX/Niagara/Source/Niagara/Private/NiagaraDataInterfaceCurlNoise.cpp
@@ -0,0 +1,303 @@
+#include "NiagaraDataInterfaceCurlNoise.h"
+
+VectorRegister UNiagaraDataInterfaceCurlNoise::NoiseTable[17][17][17];
+
+UNiagaraDataInterfaceCurlNoise::UNiagaraDataInterfaceCurlNoise(FObjectInitializer const& ObjectInitializer)
+	: Super(ObjectInitializer)
+{
+
+}
+
+void UNiagaraDataInterfaceCurlNoise::PostInitProperties()
+{
+	Super::PostInitProperties();
+
+	if (HasAnyFlags(RF_ClassDefaultObject))
+	{
+		FNiagaraTypeRegistry::Register(FNiagaraTypeDefinition(GetClass()), true, false, false);
+	}
+}
+
+void UNiagaraDataInterfaceCurlNoise::GetFunctions(TArray<FNiagaraFunctionSignature>& OutFunctions)
+{
+	FNiagaraFunctionSignature Sig;
+	Sig.Name = TEXT("SampleNoiseField");
+	Sig.bMemberFunction = true;
+	Sig.bRequiresContext = false;
+	Sig.Inputs.Add(FNiagaraVariable(FNiagaraTypeDefinition(GetClass()), TEXT("NoiseField")));
+	Sig.Inputs.Add(FNiagaraVariable(FNiagaraTypeDefinition::GetVec3Def(), TEXT("XYZ")));
+	Sig.Outputs.Add(FNiagaraVariable(FNiagaraTypeDefinition::GetVec3Def(), TEXT("Value")));
+	//Sig.Owner = *GetName();
+
+	OutFunctions.Add(Sig);
+}
+
+DEFINE_NDI_FUNC_BINDER(UNiagaraDataInterfaceCurlNoise, SampleNoiseField);
+FVMExternalFunction UNiagaraDataInterfaceCurlNoise::GetVMExternalFunction(const FVMExternalFunctionBindingInfo& BindingInfo, void* InstanceData)
+{
+	check(BindingInfo.Name == TEXT("SampleNoiseField"));
+	check(BindingInfo.GetNumInputs() == 3 && BindingInfo.GetNumOutputs() == 3);
+	return TNDIParamBinder<0, float, TNDIParamBinder<1, float, TNDIParamBinder<2, float, NDI_FUNC_BINDER(UNiagaraDataInterfaceCurlNoise, SampleNoiseField)>>>::Bind(this, BindingInfo, InstanceData);
+}
+
+template<typename XType, typename YType, typename ZType>
+void UNiagaraDataInterfaceCurlNoise::SampleNoiseField(FVectorVMContext& Context)
+{
+	XType XParam(Context);
+	YType YParam(Context);
+	ZType ZParam(Context);
+	FRegisterHandler<float> OutSampleX(Context);
+	FRegisterHandler<float> OutSampleY(Context);
+	FRegisterHandler<float> OutSampleZ(Context);
+	
+	for (int32 InstanceIdx = 0; InstanceIdx < Context.NumInstances; ++InstanceIdx)
+	{
+		VectorRegister InCoords = MakeVectorRegister(XParam.Get(), YParam.Get(), ZParam.Get(), 0.0f);
+
+		const VectorRegister One = MakeVectorRegister(1.0f, 1.0f, 1.0f, 1.0f);
+		const VectorRegister Zero = MakeVectorRegister(0.0f, 0.0f, 0.0f, 0.0f);
+		const VectorRegister VecSize = MakeVectorRegister(16.0f, 16.0f, 16.0f, 16.0f);
+
+		VectorRegister Dst = MakeVectorRegister(0.0f, 0.0f, 0.0f, 0.0f);
+
+		for (uint32 i = 1; i < 2; i++)
+		{
+			float Di = 0.2f * (1.0f / (1 << i));
+			VectorRegister Div = MakeVectorRegister(Di, Di, Di, Di);
+			VectorRegister Coords = VectorMod(VectorAbs(VectorMultiply(InCoords, Div)), VecSize);
+			Coords = VectorMin(Coords, VecSize);
+			Coords = VectorMax(Coords, Zero);
+			const float *CoordPtr = reinterpret_cast<float const*>(&Coords);
+			const int32 Cx = CoordPtr[0];
+			const int32 Cy = CoordPtr[1];
+			const int32 Cz = CoordPtr[2];
+
+			VectorRegister Frac = VectorFractional(Coords);
+			VectorRegister Alpha = VectorReplicate(Frac, 0);
+			VectorRegister OneMinusAlpha = VectorSubtract(One, Alpha);
+
+			VectorRegister XV1 = VectorMultiplyAdd(NoiseTable[Cx][Cy][Cz], Alpha, VectorMultiply(NoiseTable[Cx + 1][Cy][Cz], OneMinusAlpha));
+			VectorRegister XV2 = VectorMultiplyAdd(NoiseTable[Cx][Cy + 1][Cz], Alpha, VectorMultiply(NoiseTable[Cx + 1][Cy + 1][Cz], OneMinusAlpha));
+			VectorRegister XV3 = VectorMultiplyAdd(NoiseTable[Cx][Cy][Cz + 1], Alpha, VectorMultiply(NoiseTable[Cx + 1][Cy][Cz + 1], OneMinusAlpha));
+			VectorRegister XV4 = VectorMultiplyAdd(NoiseTable[Cx][Cy + 1][Cz + 1], Alpha, VectorMultiply(NoiseTable[Cx + 1][Cy + 1][Cz + 1], OneMinusAlpha));
+
+			Alpha = VectorReplicate(Frac, 1);
+			OneMinusAlpha = VectorSubtract(One, Alpha);
+			VectorRegister YV1 = VectorMultiplyAdd(XV1, Alpha, VectorMultiply(XV2, OneMinusAlpha));
+			VectorRegister YV2 = VectorMultiplyAdd(XV3, Alpha, VectorMultiply(XV4, OneMinusAlpha));
+
+			Alpha = VectorReplicate(Frac, 2);
+			OneMinusAlpha = VectorSubtract(One, Alpha);
+			VectorRegister ZV = VectorMultiplyAdd(YV1, Alpha, VectorMultiply(YV2, OneMinusAlpha));
+
+			Dst = VectorAdd(Dst, ZV);
+		}
+
+		float *RegPtr = reinterpret_cast<float*>(&Dst);
+		*OutSampleX.GetDest() = RegPtr[0];
+		*OutSampleY.GetDest() = RegPtr[1];
+		*OutSampleZ.GetDest() = RegPtr[2];
+
+		XParam.Advance();
+		YParam.Advance();
+		ZParam.Advance();
+		OutSampleX.Advance();
+		OutSampleY.Advance();
+		OutSampleZ.Advance();
+	}
+}
+
+
+
+// build the shader function HLSL; function name is passed in, as it's defined per-DI; that way, configuration could change
+// the HLSL in the spirit of a static switch
+// TODO: need a way to identify each specific function here
+// 
+bool UNiagaraDataInterfaceCurlNoise::GetFunctionHLSL(FString FunctionName, TArray<DIGPUBufferParamDescriptor> &Descriptors, FString &HLSLInterfaceID, FString &OutHLSL)
+{
+	FString BufferName = Descriptors[0].BufferParamName;
+	OutHLSL += TEXT("void ") + FunctionName + TEXT("(in float3 In_XYZ, out float3 Out_Value) \n{\n");
+	OutHLSL += TEXT("\t float3 a = trunc((In_XYZ*0.2) / 16.0);\n");
+	OutHLSL += TEXT("\t float3 ModXYZ = (In_XYZ*0.2) - a*16.0;\n");
+	OutHLSL += TEXT("\t int3 IntCoord = int3(ModXYZ.x, ModXYZ.y, ModXYZ.z);\n");
+	OutHLSL += TEXT("\t float3 frc = frac(ModXYZ);\n");
+	OutHLSL += TEXT("\t float3 V1 = ") + BufferName + TEXT("[IntCoord.x + IntCoord.y*17 + IntCoord.z*17*17].xyz;\n");
+	OutHLSL += TEXT("\t float3 V2 = ") + BufferName + TEXT("[IntCoord.x+1 + IntCoord.y*17 + IntCoord.z*17*17].xyz;\n");
+	OutHLSL += TEXT("\t float3 XV1 = lerp(V1, V2, frc.xxx);\n");
+	OutHLSL += TEXT("\t V1 = ") + BufferName + TEXT("[IntCoord.x + (IntCoord.y+1)*17 + IntCoord.z*17*17].xyz;\n");
+	OutHLSL += TEXT("\t V2 = ") + BufferName + TEXT("[IntCoord.x+1 + (IntCoord.y+1)*17 + IntCoord.z*17*17].xyz;\n");
+	OutHLSL += TEXT("\t float3 XV2 = lerp(V1, V2, frc.xxx);\n");
+	OutHLSL += TEXT("\t V1 = ") + BufferName + TEXT("[IntCoord.x + IntCoord.y*17 + (IntCoord.z+1)*17*17].xyz;\n");
+	OutHLSL += TEXT("\t V2 = ") + BufferName + TEXT("[IntCoord.x+1 + IntCoord.y*17 + (IntCoord.z+1)*17*17].xyz;\n");
+	OutHLSL += TEXT("\t float3 XV3 = lerp(V1, V2, frc.xxx);\n");
+	OutHLSL += TEXT("\t V1 = ") + BufferName + TEXT("[IntCoord.x + (IntCoord.y+1)*17 + (IntCoord.z+1)*17*17].xyz;\n");
+	OutHLSL += TEXT("\t V2 = ") + BufferName + TEXT("[IntCoord.x+1 + (IntCoord.y+1)*17 + (IntCoord.z+1)*17*17].xyz;\n");
+	OutHLSL += TEXT("\t float3 XV4 = lerp(V1, V2, frc.xxx);\n");
+	OutHLSL += TEXT("\t float3 YV1 = lerp(XV1, XV2, frc.yyy);\n");
+	OutHLSL += TEXT("\t float3 YV2 = lerp(XV3, XV4, frc.yyy);\n");
+	OutHLSL += TEXT("\t Out_Value = lerp(YV1, YV2, frc.zzz);\n");
+	OutHLSL += TEXT("\n}\n");
+	return true;
+}
+
+// build buffer definition hlsl
+// 1. Choose a buffer name, add the data interface ID (important!)
+// 2. add a DIGPUBufferParamDescriptor to the array argument; that'll be passed on to the FNiagaraShader for binding to a shader param, that can
+// then later be found by name via FindDIBufferParam for setting; 
+// 3. store buffer declaration hlsl in OutHLSL
+// multiple buffers can be defined at once here
+//
+void UNiagaraDataInterfaceCurlNoise::GetBufferDefinitionHLSL(FString DataInterfaceID, TArray<DIGPUBufferParamDescriptor> &BufferDescriptors, FString &OutHLSL)
+{
+	FString BufferName = "CurlNoiseLUT" + DataInterfaceID;
+	OutHLSL += TEXT("Buffer<float4> ") + BufferName + TEXT(";\n");
+
+	BufferDescriptors.Add(DIGPUBufferParamDescriptor(BufferName, 0));		// add a descriptor for shader parameter binding
+
+}
+
+// called after translate, to setup buffers matching the buffer descriptors generated during hlsl translation
+// need to do this because the script used during translate is a clone, including its DIs
+//
+void UNiagaraDataInterfaceCurlNoise::SetupBuffers(TArray<DIGPUBufferParamDescriptor> &BufferDescriptors)
+{
+	for (DIGPUBufferParamDescriptor &Desc : BufferDescriptors)
+	{
+		FNiagaraDataInterfaceBufferData BufferData(*Desc.BufferParamName);	// store off the data for later use
+		GPUBuffers.Add(BufferData);
+	}
+	bGPUBufferDirty = true;
+}
+
+// return the GPU buffer array (called from NiagaraInstanceBatcher to get the buffers for setting to the shader)
+// we lazily update the buffer with a new LUT here if necessary
+//
+TArray<FNiagaraDataInterfaceBufferData> &UNiagaraDataInterfaceCurlNoise::GetBufferDataArray()
+{
+	check(IsInRenderingThread());
+
+	if (bGPUBufferDirty)
+	{
+		check(GPUBuffers.Num() > 0);
+		FNiagaraDataInterfaceBufferData &GPUBuffer = GPUBuffers[0];
+		GPUBuffer.Buffer.Release();
+		GPUBuffer.Buffer.Initialize(sizeof(float) * 4, 17 * 17 * 17, EPixelFormat::PF_A32B32G32R32F);
+		uint32 BufferSize = 17 * 17 * 17 * sizeof(float) * 4;
+		int32 *BufferData = static_cast<int32*>(RHILockVertexBuffer(GPUBuffer.Buffer.Buffer, 0, BufferSize, EResourceLockMode::RLM_WriteOnly));
+		FVector4 TempTable[17 * 17 * 17];
+		for (int z = 0; z < 17; z++)
+		{
+			for (int y = 0; y < 17; y++)
+			{
+				for (int x = 0; x < 17; x++)
+				{
+					float *R = (float*)( &(NoiseTable[x][y][z]) );
+					TempTable[x + y * 17 + z * 17 * 17] = FVector4(R[0], R[1], R[2], 0.0f);
+				}
+			}
+		}
+		FPlatformMemory::Memcpy(BufferData, TempTable, BufferSize);
+		//		FPlatformMemory::Memcpy(BufferData, NoiseTable, BufferSize);
+		RHIUnlockVertexBuffer(GPUBuffer.Buffer.Buffer);
+		bGPUBufferDirty = false;
+	}
+
+
+	return GPUBuffers;
+}
+
+
+
+
+void UNiagaraDataInterfaceCurlNoise::InitNoiseLUT()
+{
+	// seed random stream
+	FRandomStream RandStream;
+
+	// random noise
+	float TempTable[17][17][17];
+	for (int z = 0; z < 17; z++)
+	{
+		for (int y = 0; y < 17; y++)
+		{
+			for (int x = 0; x < 17; x++)
+			{
+				float f1 = RandStream.FRandRange(-1.0f, 1.0f);
+				TempTable[x][y][z] = f1;
+			}
+		}
+	}
+
+	// pad
+	for (int i = 0; i < 17; i++)
+	{
+		for (int j = 0; j < 17; j++)
+		{
+			TempTable[i][j][16] = TempTable[i][j][0];
+			TempTable[i][16][j] = TempTable[i][0][j];
+			TempTable[16][j][i] = TempTable[0][j][i];
+		}
+	}
+
+	// compute gradients
+	FVector TempTable2[17][17][17];
+	for (int z = 0; z < 16; z++)
+	{
+		for (int y = 0; y < 16; y++)
+		{
+			for (int x = 0; x < 16; x++)
+			{
+				FVector XGrad = FVector(1.0f, 0.0f, TempTable[x][y][z] - TempTable[x + 1][y][z]);
+				FVector YGrad = FVector(0.0f, 1.0f, TempTable[x][y][z] - TempTable[x][y + 1][z]);
+				FVector ZGrad = FVector(0.0f, 1.0f, TempTable[x][y][z] - TempTable[x][y][z + 1]);
+
+				FVector Grad = FVector(XGrad.Z, YGrad.Z, ZGrad.Z);
+				TempTable2[x][y][z] = Grad;
+			}
+		}
+	}
+
+	// pad
+	for (int i = 0; i < 17; i++)
+	{
+		for (int j = 0; j < 17; j++)
+		{
+			TempTable2[i][j][16] = TempTable2[i][j][0];
+			TempTable2[i][16][j] = TempTable2[i][0][j];
+			TempTable2[16][j][i] = TempTable2[0][j][i];
+		}
+	}
+
+
+	// compute curl of gradient field
+	for (int z = 0; z < 16; z++)
+	{
+		for (int y = 0; y < 16; y++)
+		{
+			for (int x = 0; x < 16; x++)
+			{
+				FVector Dy = TempTable2[x][y][z] - TempTable2[x][y + 1][z];
+				FVector Sy = TempTable2[x][y][z] + TempTable2[x][y + 1][z];
+				FVector Dx = TempTable2[x][y][z] - TempTable2[x + 1][y][z];
+				FVector Sx = TempTable2[x][y][z] + TempTable2[x + 1][y][z];
+				FVector Dz = TempTable2[x][y][z] - TempTable2[x][y][z + 1];
+				FVector Sz = TempTable2[x][y][z] + TempTable2[x][y][z + 1];
+				FVector Dir = FVector(Dy.Z - Sz.Y, Dz.X - Sx.Z, Dx.Y - Sy.X);
+
+				NoiseTable[x][y][z] = MakeVectorRegister(Dir.X, Dir.Y, Dir.Z, 0.0f);
+			}
+		}
+	}
+
+	// pad
+	for (int i = 0; i < 17; i++)
+	{
+		for (int j = 0; j < 17; j++)
+		{
+			NoiseTable[i][j][16] = NoiseTable[i][j][0];
+			NoiseTable[i][16][j] = NoiseTable[i][0][j];
+			NoiseTable[16][j][i] = NoiseTable[0][j][i];
+		}
+	}
+
+}
diff --git a/Engine/Plugins/FX/Niagara/Source/Niagara/Private/NiagaraDataInterfaceCurve.cpp b/Engine/Plugins/FX/Niagara/Source/Niagara/Private/NiagaraDataInterfaceCurve.cpp
new file mode 100644
--- /dev/null
+++ b/Engine/Plugins/FX/Niagara/Source/Niagara/Private/NiagaraDataInterfaceCurve.cpp
@@ -0,0 +1,196 @@
+// Copyright 1998-2016 Epic Games, Inc. All Rights Reserved.
+
+#include "NiagaraDataInterfaceCurve.h"
+#include "Curves/CurveVector.h"
+#include "Curves/CurveLinearColor.h"
+#include "Curves/CurveFloat.h"
+#include "NiagaraTypes.h"
+
+
+UNiagaraDataInterfaceCurve::UNiagaraDataInterfaceCurve(FObjectInitializer const& ObjectInitializer)
+	: Super(ObjectInitializer)
+{
+	UpdateLUT();
+}
+
+void UNiagaraDataInterfaceCurve::PostInitProperties()
+{
+	Super::PostInitProperties();
+
+	if (HasAnyFlags(RF_ClassDefaultObject))
+	{
+		FNiagaraTypeRegistry::Register(FNiagaraTypeDefinition(GetClass()), true, false, false);
+	}
+
+	UpdateLUT();
+}
+
+#if WITH_EDITOR
+
+void UNiagaraDataInterfaceCurve::PostEditChangeProperty(struct FPropertyChangedEvent& PropertyChangedEvent)
+{
+	Super::PostEditChangeProperty(PropertyChangedEvent);
+
+	if (PropertyChangedEvent.Property && PropertyChangedEvent.Property->GetFName() == GET_MEMBER_NAME_CHECKED(UNiagaraDataInterfaceCurve, CurveToCopy))
+	{
+		UCurveFloat* CurveAsset = Cast<UCurveFloat>(CurveToCopy.TryLoad());
+		if (CurveAsset)
+		{
+			Modify();
+			Curve = CurveAsset->FloatCurve;
+		}
+		UpdateLUT();
+	}
+}
+
+#endif
+
+bool UNiagaraDataInterfaceCurve::CopyTo(UNiagaraDataInterface* Destination) const 
+{
+	if (!Super::CopyTo(Destination))
+	{
+		return false;
+	}
+	CastChecked<UNiagaraDataInterfaceCurve>(Destination)->Curve = Curve;
+	CastChecked<UNiagaraDataInterfaceCurve>(Destination)->UpdateLUT();
+
+	return true;
+}
+
+bool UNiagaraDataInterfaceCurve::Equals(const UNiagaraDataInterface* Other) const
+{
+	if (!Super::Equals(Other))
+	{
+		return false;
+	}
+
+	return CastChecked<UNiagaraDataInterfaceCurve>(Other)->Curve == Curve;
+}
+
+void UNiagaraDataInterfaceCurve::GetCurveData(TArray<FCurveData>& OutCurveData)
+{
+	OutCurveData.Add(FCurveData(&Curve, NAME_None, FLinearColor::Red));
+}
+
+void UNiagaraDataInterfaceCurve::GetFunctions(TArray<FNiagaraFunctionSignature>& OutFunctions)
+{
+	FNiagaraFunctionSignature Sig;
+	Sig.Name = TEXT("SampleCurve");
+	Sig.bMemberFunction = true;
+	Sig.bRequiresContext = false;
+	Sig.Inputs.Add(FNiagaraVariable(FNiagaraTypeDefinition(GetClass()), TEXT("Curve")));
+	Sig.Inputs.Add(FNiagaraVariable(FNiagaraTypeDefinition::GetFloatDef(), TEXT("X")));
+	Sig.Outputs.Add(FNiagaraVariable(FNiagaraTypeDefinition::GetFloatDef(), TEXT("Value")));
+//	Sig.Owner = *GetName();
+
+	OutFunctions.Add(Sig);
+}
+
+
+void UNiagaraDataInterfaceCurve::UpdateLUT()
+{
+	ShaderLUT.Empty();
+	for (uint32 i = 0; i < CurveLUTWidth; i++)
+	{
+		float X = i / (float)CurveLUTWidth;
+		float C = Curve.Eval(X);
+		ShaderLUT.Add(C);
+	}
+	GPUBufferDirty = true;
+}
+
+// build the shader function HLSL; function name is passed in, as it's defined per-DI; that way, configuration could change
+// the HLSL in the spirit of a static switch
+// TODO: need a way to identify each specific function here
+// 
+bool UNiagaraDataInterfaceCurve::GetFunctionHLSL(FString FunctionName, TArray<DIGPUBufferParamDescriptor> &Descriptors, FString &HLSLInterfaceID, FString &OutHLSL)
+{
+	FString BufferName = Descriptors[0].BufferParamName;
+	OutHLSL += TEXT("void ") + FunctionName + TEXT("(in float In_X, out float Out_Value) \n{\n");
+	OutHLSL += TEXT("\t Out_Value = ") + BufferName + TEXT("[(int)(clamp(In_X, 0.0, 1.0) *") + FString::FromInt(CurveLUTWidth) + TEXT(") ];");
+	OutHLSL += TEXT("\n}\n");
+	return true;
+}
+
+// build buffer definition hlsl
+// 1. Choose a buffer name, add the data interface ID (important!)
+// 2. add a DIGPUBufferParamDescriptor to the array argument; that'll be passed on to the FNiagaraShader for binding to a shader param, that can
+// then later be found by name via FindDIBufferParam for setting; 
+// 3. store buffer declaration hlsl in OutHLSL
+// multiple buffers can be defined at once here
+//
+void UNiagaraDataInterfaceCurve::GetBufferDefinitionHLSL(FString DataInterfaceID, TArray<DIGPUBufferParamDescriptor> &BufferDescriptors, FString &OutHLSL)
+{
+	FString BufferName = "CurveLUT" + DataInterfaceID;
+	OutHLSL += TEXT("Buffer<float> ") + BufferName + TEXT(";\n");
+
+	BufferDescriptors.Add(DIGPUBufferParamDescriptor(BufferName, 0));		// add a descriptor for shader parameter binding
+
+}
+
+// called after translate, to setup buffers matching the buffer descriptors generated during hlsl translation
+// need to do this because the script used during translate is a clone, including its DIs
+//
+void UNiagaraDataInterfaceCurve::SetupBuffers(TArray<DIGPUBufferParamDescriptor> &BufferDescriptors)
+{
+	for (DIGPUBufferParamDescriptor &Desc : BufferDescriptors)
+	{
+		FNiagaraDataInterfaceBufferData BufferData(*Desc.BufferParamName);	// store off the data for later use
+		GPUBuffers.Add(BufferData);
+	}
+}
+
+// return the GPU buffer array (called from NiagaraInstanceBatcher to get the buffers for setting to the shader)
+// we lazily update the buffer with a new LUT here if necessary
+//
+TArray<FNiagaraDataInterfaceBufferData> &UNiagaraDataInterfaceCurve::GetBufferDataArray()
+{
+	check(IsInRenderingThread());
+	if (GPUBufferDirty)
+	{
+		check(GPUBuffers.Num() > 0);
+
+		FNiagaraDataInterfaceBufferData &GPUBuffer = GPUBuffers[0];
+		GPUBuffer.Buffer.Release();
+		GPUBuffer.Buffer.Initialize(sizeof(float), CurveLUTWidth * 4, EPixelFormat::PF_R32_FLOAT);	// always allocate for up to 64 data sets
+		uint32 BufferSize = ShaderLUT.Num() * sizeof(float);
+		int32 *BufferData = static_cast<int32*>(RHILockVertexBuffer(GPUBuffer.Buffer.Buffer, 0, BufferSize, EResourceLockMode::RLM_WriteOnly));
+		FPlatformMemory::Memcpy(BufferData, ShaderLUT.GetData(), BufferSize);
+		RHIUnlockVertexBuffer(GPUBuffer.Buffer.Buffer);
+		GPUBufferDirty = false;
+	}
+
+	return GPUBuffers;
+}
+
+
+
+DEFINE_NDI_FUNC_BINDER(UNiagaraDataInterfaceCurve, SampleCurve);
+FVMExternalFunction UNiagaraDataInterfaceCurve::GetVMExternalFunction(const FVMExternalFunctionBindingInfo& BindingInfo, void* InstanceData)
+{
+	if (BindingInfo.Name == TEXT("SampleCurve") && BindingInfo.GetNumInputs() == 1 && BindingInfo.GetNumOutputs() == 1)
+	{
+		return TNDIParamBinder<0, float, NDI_FUNC_BINDER(UNiagaraDataInterfaceCurve, SampleCurve)>::Bind(this, BindingInfo, InstanceData);
+	}
+	else
+	{
+		UE_LOG(LogNiagara, Error, TEXT("Could not find data interface external function.\n\tExpected Name: SampleCurve  Actual Name: %s\n\tExpected Inputs: 1  Actual Inputs: %i\n\tExpected Outputs: 1  Actual Outputs: %i"),
+			*BindingInfo.Name.ToString(), BindingInfo.GetNumInputs(), BindingInfo.GetNumOutputs());
+		return FVMExternalFunction();
+	}
+}
+
+template<typename XParamType>
+void UNiagaraDataInterfaceCurve::SampleCurve(FVectorVMContext& Context)
+{
+	XParamType XParam(Context);
+	FRegisterHandler<float> OutSample(Context);
+
+	for (int32 i = 0; i < Context.NumInstances; ++i)
+	{
+		*OutSample.GetDest() = Curve.Eval(XParam.Get());
+		XParam.Advance();
+		OutSample.Advance();
+	}
+}
+
diff --git a/Engine/Plugins/FX/Niagara/Source/Niagara/Private/NiagaraDataInterfaceSpline.cpp b/Engine/Plugins/FX/Niagara/Source/Niagara/Private/NiagaraDataInterfaceSpline.cpp
new file mode 100644
--- /dev/null
+++ b/Engine/Plugins/FX/Niagara/Source/Niagara/Private/NiagaraDataInterfaceSpline.cpp
@@ -0,0 +1,647 @@
+// Copyright 1998-2016 Epic Games, Inc. All Rights Reserved.
+
+#include "NiagaraDataInterfaceSpline.h"
+#include "NiagaraEmitterInstance.h"
+#include "NiagaraComponent.h"
+#include "NiagaraSystemInstance.h"
+#include "Internationalization.h"
+
+#define LOCTEXT_NAMESPACE "NiagaraDataInterfaceSpline"
+
+namespace SplineUtils
+{
+	struct FTransformHandlerNoop
+	{
+		FORCEINLINE void Transform(FVector& V, FMatrix& M) {  }
+	};
+
+	struct FTransformHandlerPosition
+	{
+		FORCEINLINE void Transform(FVector& P, FMatrix& M) { P = M.TransformPosition(P); }
+	};
+
+	struct FTransformHandlerUnitVector
+	{
+		FORCEINLINE void Transform(FVector& V, FMatrix& M) { V = M.TransformVector(V).GetUnsafeNormal3(); }
+	};
+
+	struct FTransformHandlerTangent
+	{
+		FORCEINLINE void Transform(FVector& V, FMatrix& M) { V = M.TransformVector(V); }
+	};
+};
+
+UNiagaraDataInterfaceSpline::UNiagaraDataInterfaceSpline(FObjectInitializer const& ObjectInitializer)
+	: Super(ObjectInitializer)
+	, Source(nullptr)
+{
+
+}
+
+void UNiagaraDataInterfaceSpline::PostInitProperties()
+{
+	Super::PostInitProperties();
+
+	//Can we regitser data interfaces as regular types and fold them into the FNiagaraVariable framework for UI and function calls etc?
+	if (HasAnyFlags(RF_ClassDefaultObject))
+	{
+		FNiagaraTypeRegistry::Register(FNiagaraTypeDefinition(GetClass()), true, false, false);
+	}
+}
+
+static const FName SampleSplinePositionByUnitDistanceName("SampleSplinePositionByUnitDistance");
+static const FName SampleSplinePositionByUnitDistanceWSName("SampleSplinePositionByUnitDistanceWS");
+
+static const FName SampleSplineUpVectorByUnitDistanceName("SampleSplineUpVectorByUnitDistance");
+static const FName SampleSplineUpVectorByUnitDistanceWSName("SampleSplineUpVectorByUnitDistanceWS");
+
+static const FName SampleSplineDirectionByUnitDistanceName("SampleSplineDirectionByUnitDistance");
+static const FName SampleSplineDirectionByUnitDistanceWSName("SampleSplineDirectionByUnitDistanceWS");
+
+static const FName SampleSplineRightVectorByUnitDistanceName("SampleSplineRightVectorByUnitDistance");
+static const FName SampleSplineRightVectorByUnitDistanceWSName("SampleSplineRightVectorByUnitDistanceWS");
+
+static const FName SampleSplineTangentByUnitDistanceName("SampleSplineTangentByUnitDistance");
+static const FName SampleSplineTangentByUnitDistanceWSName("SampleSplineTangentByUnitDistanceWS");
+
+
+static const FName FindClosestUnitDistanceFromPositionWSName("FindClosestUnitDistanceFromPositionWS");
+
+
+/** Temporary solution for exposing the transform of a mesh. Ideally this would be done by allowing interfaces to add to the uniform set for a simulation. */
+static const FName GetSplineLocalToWorldName("GetSplineLocalToWorld");
+static const FName GetSplineLocalToWorldInverseTransposedName("GetSplineLocalToWorldInverseTransposed");
+
+void UNiagaraDataInterfaceSpline::GetFunctions(TArray<FNiagaraFunctionSignature>& OutFunctions)
+{
+	{
+		FNiagaraFunctionSignature Sig;
+		Sig.Name = SampleSplinePositionByUnitDistanceName;
+		Sig.Inputs.Add(FNiagaraVariable(FNiagaraTypeDefinition(GetClass()), TEXT("Spline")));
+		Sig.Inputs.Add(FNiagaraVariable(FNiagaraTypeDefinition::GetFloatDef(), TEXT("U")));
+		Sig.Outputs.Add(FNiagaraVariable(FNiagaraTypeDefinition::GetVec3Def(), TEXT("Position")));
+		Sig.bMemberFunction = true;
+		Sig.bRequiresContext = false;
+		//Sig.Owner = *GetName();
+		Sig.SetDescription(LOCTEXT("DataInterfaceSpline_SampleSplinePositionByUnitDistance", "Sample the spline Position where U is a 0 to 1 value representing the start and normalized length of the spline.\nThis is in the local space of the referenced USplineComponent."));
+		OutFunctions.Add(Sig);
+	}
+
+	{
+		FNiagaraFunctionSignature Sig;
+		Sig.Name = SampleSplinePositionByUnitDistanceWSName;
+		Sig.Inputs.Add(FNiagaraVariable(FNiagaraTypeDefinition(GetClass()), TEXT("Spline")));
+		Sig.Inputs.Add(FNiagaraVariable(FNiagaraTypeDefinition::GetFloatDef(), TEXT("U")));
+		Sig.Outputs.Add(FNiagaraVariable(FNiagaraTypeDefinition::GetVec3Def(), TEXT("Position")));
+		Sig.bMemberFunction = true;
+		Sig.bRequiresContext = false;
+		//Sig.Owner = *GetName();
+		Sig.SetDescription(LOCTEXT("DataInterfaceSpline_SampleSplinePositionByUnitDistanceWS", "Sample the spline Position where U is a 0 to 1 value representing the start and normalized length of the spline.\nThis is in the world space of the level."));
+		OutFunctions.Add(Sig);
+	}
+
+	{
+		FNiagaraFunctionSignature Sig;
+		Sig.Name = SampleSplineDirectionByUnitDistanceName;
+		Sig.Inputs.Add(FNiagaraVariable(FNiagaraTypeDefinition(GetClass()), TEXT("Spline")));
+		Sig.Inputs.Add(FNiagaraVariable(FNiagaraTypeDefinition::GetFloatDef(), TEXT("U")));
+		Sig.Outputs.Add(FNiagaraVariable(FNiagaraTypeDefinition::GetVec3Def(), TEXT("Direction")));
+		Sig.bMemberFunction = true;
+		Sig.bRequiresContext = false;
+		//Sig.Owner = *GetName();
+		Sig.SetDescription(LOCTEXT("DataInterfaceSpline_SampleSplineDirectionByUnitDistance", "Sample the spline direction vector where U is a 0 to 1 value representing the start and normalized length of the spline.\nThis is in the local space of the referenced USplineComponent."));
+		OutFunctions.Add(Sig);
+	}
+
+	{
+		FNiagaraFunctionSignature Sig;
+		Sig.Name = SampleSplineDirectionByUnitDistanceWSName;
+		Sig.Inputs.Add(FNiagaraVariable(FNiagaraTypeDefinition(GetClass()), TEXT("Spline")));
+		Sig.Inputs.Add(FNiagaraVariable(FNiagaraTypeDefinition::GetFloatDef(), TEXT("U")));
+		Sig.Outputs.Add(FNiagaraVariable(FNiagaraTypeDefinition::GetVec3Def(), TEXT("Direction")));
+		Sig.bMemberFunction = true;
+		Sig.bRequiresContext = false;
+		//Sig.Owner = *GetName();
+		Sig.SetDescription(LOCTEXT("DataInterfaceSpline_SampleSplineDirectionByUnitDistanceWS", "Sample the spline direction vector where U is a 0 to 1 value representing the start and normalized length of the spline.\nThis is in the world space of the level."));
+		OutFunctions.Add(Sig);
+	}
+
+	{
+		FNiagaraFunctionSignature Sig;
+		Sig.Name = SampleSplineUpVectorByUnitDistanceName;
+		Sig.Inputs.Add(FNiagaraVariable(FNiagaraTypeDefinition(GetClass()), TEXT("Spline")));
+		Sig.Inputs.Add(FNiagaraVariable(FNiagaraTypeDefinition::GetFloatDef(), TEXT("U")));
+		Sig.Outputs.Add(FNiagaraVariable(FNiagaraTypeDefinition::GetVec3Def(), TEXT("UpVector")));
+		Sig.bMemberFunction = true;
+		Sig.bRequiresContext = false;
+		//Sig.Owner = *GetName();
+		Sig.SetDescription(LOCTEXT("DataInterfaceSpline_SampleSplineUpVectorByUnitDistance", "Sample the spline up vector where U is a 0 to 1 value representing the start and normalized length of the spline.\nThis is in the local space of the referenced USplineComponent."));
+		OutFunctions.Add(Sig);
+	}
+
+	{
+		FNiagaraFunctionSignature Sig;
+		Sig.Name = SampleSplineUpVectorByUnitDistanceWSName;
+		Sig.Inputs.Add(FNiagaraVariable(FNiagaraTypeDefinition(GetClass()), TEXT("Spline")));
+		Sig.Inputs.Add(FNiagaraVariable(FNiagaraTypeDefinition::GetFloatDef(), TEXT("U")));
+		Sig.SetDescription(LOCTEXT("DataInterfaceSpline_SampleSplineUpVectorByUnitDistanceWS", "Sample the spline up vector where U is a 0 to 1 value representing the start and normalized length of the spline.\nThis is in the world space of the level."));
+		Sig.Outputs.Add(FNiagaraVariable(FNiagaraTypeDefinition::GetVec3Def(), TEXT("UpVector")));
+		Sig.bMemberFunction = true;
+		Sig.bRequiresContext = false;
+		//Sig.Owner = *GetName();
+		OutFunctions.Add(Sig);
+	}
+
+	{
+		FNiagaraFunctionSignature Sig;
+		Sig.Name = SampleSplineRightVectorByUnitDistanceName;
+		Sig.Inputs.Add(FNiagaraVariable(FNiagaraTypeDefinition(GetClass()), TEXT("Spline")));
+		Sig.Inputs.Add(FNiagaraVariable(FNiagaraTypeDefinition::GetFloatDef(), TEXT("U")));
+		Sig.SetDescription(LOCTEXT("DataInterfaceSpline_SampleSplineRightVectorByUnitDistance", "Sample the spline right vector where U is a 0 to 1 value representing the start and normalized length of the spline.\nThis is in the local space of the referenced USplineComponent."));
+		Sig.Outputs.Add(FNiagaraVariable(FNiagaraTypeDefinition::GetVec3Def(), TEXT("RightVector")));
+		Sig.bMemberFunction = true;
+		Sig.bRequiresContext = false;
+		//Sig.Owner = *GetName();
+		OutFunctions.Add(Sig);
+	}
+
+	{
+		FNiagaraFunctionSignature Sig;
+		Sig.Name = SampleSplineRightVectorByUnitDistanceWSName;
+		Sig.Inputs.Add(FNiagaraVariable(FNiagaraTypeDefinition(GetClass()), TEXT("Spline")));
+		Sig.Inputs.Add(FNiagaraVariable(FNiagaraTypeDefinition::GetFloatDef(), TEXT("U")));
+		Sig.SetDescription(LOCTEXT("DataInterfaceSpline_SampleSplineRightVectorByUnitDistanceWS", "Sample the spline right vector where U is a 0 to 1 value representing the start and normalized length of the spline.\nThis is in the world space of the level."));
+		Sig.Outputs.Add(FNiagaraVariable(FNiagaraTypeDefinition::GetVec3Def(), TEXT("RightVector")));
+		Sig.bMemberFunction = true;
+		Sig.bRequiresContext = false;
+		//Sig.Owner = *GetName();
+		OutFunctions.Add(Sig);
+	}
+
+	{
+		FNiagaraFunctionSignature Sig;
+		Sig.Name = SampleSplineTangentByUnitDistanceName;
+		Sig.Inputs.Add(FNiagaraVariable(FNiagaraTypeDefinition(GetClass()), TEXT("Spline")));
+		Sig.Inputs.Add(FNiagaraVariable(FNiagaraTypeDefinition::GetFloatDef(), TEXT("U")));
+		Sig.Outputs.Add(FNiagaraVariable(FNiagaraTypeDefinition::GetVec3Def(), TEXT("Tangent")));
+		Sig.bMemberFunction = true;
+		Sig.bRequiresContext = false;
+		//Sig.Owner = *GetName();
+		Sig.SetDescription(LOCTEXT("DataInterfaceSpline_SampleSplineTangentVectorByUnitDistance", "Sample the spline tangent vector where U is a 0 to 1 value representing the start and normalized length of the spline.\nThis is in the local space of the referenced USplineComponent."));
+		OutFunctions.Add(Sig);
+	}
+
+	{
+		FNiagaraFunctionSignature Sig;
+		Sig.Name = SampleSplineTangentByUnitDistanceWSName;
+		Sig.Inputs.Add(FNiagaraVariable(FNiagaraTypeDefinition(GetClass()), TEXT("Spline")));
+		Sig.Inputs.Add(FNiagaraVariable(FNiagaraTypeDefinition::GetFloatDef(), TEXT("U")));
+		Sig.Outputs.Add(FNiagaraVariable(FNiagaraTypeDefinition::GetVec3Def(), TEXT("Tangent")));
+		Sig.bMemberFunction = true;
+		Sig.bRequiresContext = false;
+		//Sig.Owner = *GetName();
+		Sig.SetDescription(LOCTEXT("DataInterfaceSpline_SampleSplineTangentVectorByUnitDistanceWS", "Sample the spline tangent vector where U is a 0 to 1 value representing the start and normalized length of the spline.\nThis is in the world space of the level."));
+		OutFunctions.Add(Sig);
+	}
+
+	{
+		FNiagaraFunctionSignature Sig;
+		Sig.Name = GetSplineLocalToWorldName;
+		Sig.Inputs.Add(FNiagaraVariable(FNiagaraTypeDefinition(GetClass()), TEXT("Spline")));
+		Sig.Outputs.Add(FNiagaraVariable(FNiagaraTypeDefinition::GetMatrix4Def(), TEXT("Transform")));
+		Sig.bMemberFunction = true;
+		Sig.bRequiresContext = false;
+		//Sig.Owner = *GetName();
+		Sig.SetDescription(LOCTEXT("DataInterfaceSpline_GetSplineLocalToWorld", "Get the transform from the USplineComponent's local space to world space."));
+		OutFunctions.Add(Sig);
+	}
+
+	{
+		FNiagaraFunctionSignature Sig;
+		Sig.Name = GetSplineLocalToWorldInverseTransposedName;
+		Sig.Inputs.Add(FNiagaraVariable(FNiagaraTypeDefinition(GetClass()), TEXT("Spline")));
+		Sig.Outputs.Add(FNiagaraVariable(FNiagaraTypeDefinition::GetMatrix4Def(), TEXT("Transform")));
+		Sig.bMemberFunction = true;
+		Sig.bRequiresContext = false;
+		//Sig.Owner = *GetName();
+		Sig.SetDescription(LOCTEXT("DataInterfaceSpline_GetSplineLocalToWorldInverseTransposed", "Get the transform from the world space to the USplineComponent's local space."));
+		OutFunctions.Add(Sig);
+	}
+
+	{
+		FNiagaraFunctionSignature Sig;
+		Sig.Name = FindClosestUnitDistanceFromPositionWSName;
+		Sig.Inputs.Add(FNiagaraVariable(FNiagaraTypeDefinition(GetClass()), TEXT("Spline")));
+		Sig.Inputs.Add(FNiagaraVariable(FNiagaraTypeDefinition::GetVec3Def(), TEXT("PositionWS")));
+
+		Sig.Outputs.Add(FNiagaraVariable(FNiagaraTypeDefinition::GetFloatDef(), TEXT("U")));
+		Sig.bMemberFunction = true;
+		Sig.bRequiresContext = false;
+		//Sig.Owner = *GetName();
+		Sig.SetDescription(LOCTEXT("DataInterfaceSpline_FindClosestUnitDistanceFromPositionWS", "Given a world space position, find the closest value 'U' on the USplineComponent to that point."));
+		OutFunctions.Add(Sig);
+	}
+}
+
+DEFINE_NDI_FUNC_BINDER(UNiagaraDataInterfaceSpline, SampleSplinePositionByUnitDistance);
+DEFINE_NDI_FUNC_BINDER(UNiagaraDataInterfaceSpline, SampleSplineUpVectorByUnitDistance);
+DEFINE_NDI_FUNC_BINDER(UNiagaraDataInterfaceSpline, SampleSplineRightVectorByUnitDistance);
+DEFINE_NDI_FUNC_BINDER(UNiagaraDataInterfaceSpline, SampleSplineDirectionByUnitDistance);
+DEFINE_NDI_FUNC_BINDER(UNiagaraDataInterfaceSpline, SampleSplineTangentByUnitDistance);
+DEFINE_NDI_FUNC_BINDER(UNiagaraDataInterfaceSpline, FindClosestUnitDistanceFromPositionWS);
+
+FVMExternalFunction UNiagaraDataInterfaceSpline::GetVMExternalFunction(const FVMExternalFunctionBindingInfo& BindingInfo, void* InstanceData)
+{
+	using namespace SplineUtils;
+	if (BindingInfo.Name == SampleSplinePositionByUnitDistanceName)
+	{
+		check(BindingInfo.GetNumInputs() == 2 && BindingInfo.GetNumOutputs() == 3);
+		return TNDIExplicitBinder<FTransformHandlerNoop, TNDIParamBinder<0, float, NDI_FUNC_BINDER(UNiagaraDataInterfaceSpline, SampleSplinePositionByUnitDistance)>>::Bind(this, BindingInfo, InstanceData);
+	}
+	else if (BindingInfo.Name == SampleSplinePositionByUnitDistanceWSName)
+	{
+		check(BindingInfo.GetNumInputs() == 2 && BindingInfo.GetNumOutputs() == 3);
+		return TNDIExplicitBinder<FTransformHandlerPosition, TNDIParamBinder<0, float, NDI_FUNC_BINDER(UNiagaraDataInterfaceSpline, SampleSplinePositionByUnitDistance)>>::Bind(this, BindingInfo, InstanceData);
+	}
+	else if (BindingInfo.Name == SampleSplineUpVectorByUnitDistanceName)
+	{
+		check(BindingInfo.GetNumInputs() == 2 && BindingInfo.GetNumOutputs() == 3);
+		return TNDIExplicitBinder<FTransformHandlerNoop, TNDIParamBinder<0, float, NDI_FUNC_BINDER(UNiagaraDataInterfaceSpline, SampleSplineUpVectorByUnitDistance)>>::Bind(this, BindingInfo, InstanceData);
+	}
+	else if (BindingInfo.Name == SampleSplineUpVectorByUnitDistanceWSName)
+	{
+		check(BindingInfo.GetNumInputs() == 2 && BindingInfo.GetNumOutputs() == 3);
+		return TNDIExplicitBinder<FTransformHandlerUnitVector, TNDIParamBinder<0, float, NDI_FUNC_BINDER(UNiagaraDataInterfaceSpline, SampleSplineUpVectorByUnitDistance)>>::Bind(this, BindingInfo, InstanceData);
+	}
+	else if (BindingInfo.Name == SampleSplineDirectionByUnitDistanceName)
+	{
+		check(BindingInfo.GetNumInputs() == 2 && BindingInfo.GetNumOutputs() == 3);
+		return TNDIExplicitBinder<FTransformHandlerNoop, TNDIParamBinder<0, float, NDI_FUNC_BINDER(UNiagaraDataInterfaceSpline, SampleSplineDirectionByUnitDistance)>>::Bind(this, BindingInfo, InstanceData);
+	}
+	else if (BindingInfo.Name == SampleSplineDirectionByUnitDistanceWSName)
+	{
+		check(BindingInfo.GetNumInputs() == 2 && BindingInfo.GetNumOutputs() == 3);
+		return TNDIExplicitBinder<FTransformHandlerUnitVector, TNDIParamBinder<0, float, NDI_FUNC_BINDER(UNiagaraDataInterfaceSpline, SampleSplineDirectionByUnitDistance)>>::Bind(this, BindingInfo, InstanceData);
+	}
+	else if (BindingInfo.Name == SampleSplineRightVectorByUnitDistanceName)
+	{
+		check(BindingInfo.GetNumInputs() == 2 && BindingInfo.GetNumOutputs() == 3);
+		return TNDIExplicitBinder<FTransformHandlerNoop, TNDIParamBinder<0, float, NDI_FUNC_BINDER(UNiagaraDataInterfaceSpline, SampleSplineRightVectorByUnitDistance)>>::Bind(this, BindingInfo, InstanceData);
+	}
+	else if (BindingInfo.Name == SampleSplineRightVectorByUnitDistanceWSName)
+	{
+		check(BindingInfo.GetNumInputs() == 2 && BindingInfo.GetNumOutputs() == 3);
+		return TNDIExplicitBinder<FTransformHandlerUnitVector, TNDIParamBinder<0, float, NDI_FUNC_BINDER(UNiagaraDataInterfaceSpline, SampleSplineRightVectorByUnitDistance)>>::Bind(this, BindingInfo, InstanceData);
+	}
+	else if (BindingInfo.Name == SampleSplineTangentByUnitDistanceName)
+	{
+		check(BindingInfo.GetNumInputs() == 2 && BindingInfo.GetNumOutputs() == 3);
+		return TNDIExplicitBinder<FTransformHandlerNoop, TNDIParamBinder<0, float, NDI_FUNC_BINDER(UNiagaraDataInterfaceSpline, SampleSplineTangentByUnitDistance)>>::Bind(this, BindingInfo, InstanceData);
+	}
+	else if (BindingInfo.Name == SampleSplineTangentByUnitDistanceWSName)
+	{
+		check(BindingInfo.GetNumInputs() == 2 && BindingInfo.GetNumOutputs() == 3);
+		return TNDIExplicitBinder<FTransformHandlerTangent, TNDIParamBinder<0, float, NDI_FUNC_BINDER(UNiagaraDataInterfaceSpline, SampleSplineTangentByUnitDistance)>>::Bind(this, BindingInfo, InstanceData);
+	}
+	else if (BindingInfo.Name == FindClosestUnitDistanceFromPositionWSName)
+	{
+		check(BindingInfo.GetNumInputs() == 4 && BindingInfo.GetNumOutputs() == 1);
+		return TNDIParamBinder<0, float, TNDIParamBinder<1, float, TNDIParamBinder<2, float, NDI_FUNC_BINDER(UNiagaraDataInterfaceSpline, FindClosestUnitDistanceFromPositionWS)>>>::Bind(this, BindingInfo, InstanceData);
+	}
+	else if (BindingInfo.Name == GetSplineLocalToWorldName)
+	{
+		check(BindingInfo.GetNumInputs() == 1 && BindingInfo.GetNumOutputs() == 16);
+		return FVMExternalFunction::CreateUObject(this, &UNiagaraDataInterfaceSpline::GetLocalToWorld);
+	}
+	else if (BindingInfo.Name == GetSplineLocalToWorldInverseTransposedName)
+	{
+		check(BindingInfo.GetNumInputs() == 1 && BindingInfo.GetNumOutputs() == 16);
+		return FVMExternalFunction::CreateUObject(this, &UNiagaraDataInterfaceSpline::GetLocalToWorldInverseTransposed);
+	}
+
+	//check(0);
+	return FVMExternalFunction();
+}
+
+bool UNiagaraDataInterfaceSpline::CopyTo(UNiagaraDataInterface* Destination) const
+{
+	if (!Super::CopyTo(Destination))
+	{
+		return false;
+	}
+
+	UNiagaraDataInterfaceSpline* OtherTyped = CastChecked<UNiagaraDataInterfaceSpline>(Destination);
+	OtherTyped->Source = Source;
+	return true;
+}
+
+bool UNiagaraDataInterfaceSpline::Equals(const UNiagaraDataInterface* Other) const
+{
+	if (!Super::Equals(Other))
+	{
+		return false;
+	}
+	const UNiagaraDataInterfaceSpline* OtherTyped = CastChecked<const UNiagaraDataInterfaceSpline>(Other);
+	return OtherTyped->Source == Source;
+}
+
+int32 UNiagaraDataInterfaceSpline::PerInstanceDataSize()const
+{
+	return sizeof(FNDISpline_InstanceData);
+}
+
+bool UNiagaraDataInterfaceSpline::InitPerInstanceData(void* PerInstanceData, FNiagaraSystemInstance* SystemInstance)
+{
+	FNDISpline_InstanceData* InstData = new (PerInstanceData) FNDISpline_InstanceData();
+
+	InstData->Component = nullptr;
+	InstData->Transform = FMatrix::Identity;
+	InstData->TransformInverseTransposed = FMatrix::Identity;
+	if (Source)
+	{
+		USplineComponent* SourceComp = Source->FindComponentByClass<USplineComponent>();
+
+		if (SourceComp)
+		{
+			InstData->Component = SourceComp;
+		}
+	}
+	else
+	{
+		if (UNiagaraComponent* SimComp = SystemInstance->GetComponent())
+		{
+			if (AActor* Owner = SimComp->GetAttachmentRootActor())
+			{
+				USplineComponent* SourceComp = Owner->FindComponentByClass<USplineComponent>();
+
+				if (SourceComp)
+				{
+					InstData->Component = SourceComp;
+				}
+			}
+		}
+	}
+
+	//Re-evaluate source in case it's changed?
+	if (InstData->Component)
+	{
+		InstData->Transform = InstData->Component->GetComponentToWorld().ToMatrixWithScale();
+		InstData->TransformInverseTransposed = InstData->Transform.InverseFast().GetTransposed();
+	}
+
+	return true;
+}
+
+void UNiagaraDataInterfaceSpline::DestroyPerInstanceData(void* PerInstanceData, FNiagaraSystemInstance* SystemInstance)
+{
+	FNDISpline_InstanceData* InstData = (FNDISpline_InstanceData*)PerInstanceData;
+	InstData->~FNDISpline_InstanceData();
+}
+
+bool UNiagaraDataInterfaceSpline::PerInstanceTick(void* PerInstanceData, FNiagaraSystemInstance* SystemInstance, float DeltaSeconds)
+{
+	check(SystemInstance);
+	FNDISpline_InstanceData* InstData = (FNDISpline_InstanceData*)PerInstanceData;
+
+	if (InstData->Component)
+	{
+		InstData->Transform = InstData->Component->GetComponentToWorld().ToMatrixWithScale();
+		InstData->TransformInverseTransposed = InstData->Transform.InverseFast().GetTransposed();
+	}
+
+	//Any situations requiring a rebind?
+	return false;
+}
+
+template<typename TransformHandlerType, typename SplineSampleType>
+void UNiagaraDataInterfaceSpline::SampleSplinePositionByUnitDistance(FVectorVMContext& Context)
+{
+	TransformHandlerType TransformHandler;
+	SplineSampleType SplineSampleParam(Context);
+	FUserPtrHandler<FNDISpline_InstanceData> InstData(Context);
+	FRegisterHandler<float> OutPosX(Context);
+	FRegisterHandler<float> OutPosY(Context);
+	FRegisterHandler<float> OutPosZ(Context);
+
+	checkSlow(InstData->Component);
+
+	for (int32 i = 0; i < Context.NumInstances; ++i)
+	{
+		float DistanceUnitDistance = SplineSampleParam.Get();
+
+		FVector Pos = InstData->Component->GetLocationAtDistanceAlongSpline(DistanceUnitDistance * InstData->Component->GetSplineLength(), ESplineCoordinateSpace::Local);
+		TransformHandler.Transform(Pos, InstData->Transform);
+
+		*OutPosX.GetDest() = Pos.X;
+		*OutPosY.GetDest() = Pos.Y;
+		*OutPosZ.GetDest() = Pos.Z;
+		SplineSampleParam.Advance();
+		OutPosX.Advance();
+		OutPosY.Advance();
+		OutPosZ.Advance();
+	}
+}
+
+template<typename TransformHandlerType, typename SplineSampleType>
+void UNiagaraDataInterfaceSpline::SampleSplineUpVectorByUnitDistance(FVectorVMContext& Context)
+{
+	TransformHandlerType TransformHandler;
+	SplineSampleType SplineSampleParam(Context);
+	FUserPtrHandler<FNDISpline_InstanceData> InstData(Context);
+	FRegisterHandler<float> OutPosX(Context);
+	FRegisterHandler<float> OutPosY(Context);
+	FRegisterHandler<float> OutPosZ(Context);
+
+	checkSlow(InstData->Component);
+
+	for (int32 i = 0; i < Context.NumInstances; ++i)
+	{
+		float DistanceUnitDistance = SplineSampleParam.Get();
+
+		FVector Pos = InstData->Component->GetUpVectorAtDistanceAlongSpline(DistanceUnitDistance * InstData->Component->GetSplineLength(), ESplineCoordinateSpace::Local);
+		TransformHandler.Transform(Pos, InstData->Transform);
+
+		*OutPosX.GetDest() = Pos.X;
+		*OutPosY.GetDest() = Pos.Y;
+		*OutPosZ.GetDest() = Pos.Z;
+		SplineSampleParam.Advance();
+		OutPosX.Advance();
+		OutPosY.Advance();
+		OutPosZ.Advance();
+	}
+}
+
+template<typename TransformHandlerType, typename SplineSampleType>
+void UNiagaraDataInterfaceSpline::SampleSplineRightVectorByUnitDistance(FVectorVMContext& Context)
+{
+	TransformHandlerType TransformHandler;
+	SplineSampleType SplineSampleParam(Context);
+	FUserPtrHandler<FNDISpline_InstanceData> InstData(Context);
+	FRegisterHandler<float> OutPosX(Context);
+	FRegisterHandler<float> OutPosY(Context);
+	FRegisterHandler<float> OutPosZ(Context);
+	checkSlow(InstData->Component);
+
+	for (int32 i = 0; i < Context.NumInstances; ++i)
+	{
+		float DistanceUnitDistance = SplineSampleParam.Get();
+
+		FVector Pos = InstData->Component->GetRightVectorAtDistanceAlongSpline(DistanceUnitDistance * InstData->Component->GetSplineLength(), ESplineCoordinateSpace::Local);
+		TransformHandler.Transform(Pos, InstData->Transform);
+
+		*OutPosX.GetDest() = Pos.X;
+		*OutPosY.GetDest() = Pos.Y;
+		*OutPosZ.GetDest() = Pos.Z;
+		SplineSampleParam.Advance();
+		OutPosX.Advance();
+		OutPosY.Advance();
+		OutPosZ.Advance();
+	}
+}
+
+template<typename TransformHandlerType, typename SplineSampleType>
+void UNiagaraDataInterfaceSpline::SampleSplineTangentByUnitDistance(FVectorVMContext& Context)
+{
+	TransformHandlerType TransformHandler;
+	SplineSampleType SplineSampleParam(Context);
+	FUserPtrHandler<FNDISpline_InstanceData> InstData(Context);
+	FRegisterHandler<float> OutPosX(Context);
+	FRegisterHandler<float> OutPosY(Context);
+	FRegisterHandler<float> OutPosZ(Context);
+
+	checkSlow(InstData->Component);
+	for (int32 i = 0; i < Context.NumInstances; ++i)
+	{
+		float DistanceUnitDistance = SplineSampleParam.Get();
+
+		FVector Pos = InstData->Component->GetTangentAtDistanceAlongSpline(DistanceUnitDistance * InstData->Component->GetSplineLength(), ESplineCoordinateSpace::Local);
+		TransformHandler.Transform(Pos, InstData->Transform);
+
+		*OutPosX.GetDest() = Pos.X;
+		*OutPosY.GetDest() = Pos.Y;
+		*OutPosZ.GetDest() = Pos.Z;
+		SplineSampleParam.Advance();
+		OutPosX.Advance();
+		OutPosY.Advance();
+		OutPosZ.Advance();
+	}
+}
+
+template<typename TransformHandlerType, typename SplineSampleType>
+void UNiagaraDataInterfaceSpline::SampleSplineDirectionByUnitDistance(FVectorVMContext& Context)
+{
+	TransformHandlerType TransformHandler;
+	SplineSampleType SplineSampleParam(Context);
+	FUserPtrHandler<FNDISpline_InstanceData> InstData(Context);
+	FRegisterHandler<float> OutPosX(Context);
+	FRegisterHandler<float> OutPosY(Context);
+	FRegisterHandler<float> OutPosZ(Context);
+
+	checkSlow(InstData->Component);
+
+	for (int32 i = 0; i < Context.NumInstances; ++i)
+	{
+		float DistanceUnitDistance = SplineSampleParam.Get();
+
+		FVector Pos = InstData->Component->GetDirectionAtDistanceAlongSpline(DistanceUnitDistance * InstData->Component->GetSplineLength(), ESplineCoordinateSpace::Local);
+		TransformHandler.Transform(Pos, InstData->Transform);
+
+		*OutPosX.GetDest() = Pos.X;
+		*OutPosY.GetDest() = Pos.Y;
+		*OutPosZ.GetDest() = Pos.Z;
+		SplineSampleParam.Advance();
+		OutPosX.Advance();
+		OutPosY.Advance();
+		OutPosZ.Advance();
+	}
+}
+
+void UNiagaraDataInterfaceSpline::WriteTransform(const FMatrix& ToWrite, FVectorVMContext& Context)
+{
+	FRegisterHandler<float> Out00(Context);
+	FRegisterHandler<float> Out01(Context);
+	FRegisterHandler<float> Out02(Context);
+	FRegisterHandler<float> Out03(Context);
+	FRegisterHandler<float> Out04(Context);
+	FRegisterHandler<float> Out05(Context);
+	FRegisterHandler<float> Out06(Context);
+	FRegisterHandler<float> Out07(Context);
+	FRegisterHandler<float> Out08(Context);
+	FRegisterHandler<float> Out09(Context);
+	FRegisterHandler<float> Out10(Context);
+	FRegisterHandler<float> Out11(Context);
+	FRegisterHandler<float> Out12(Context);
+	FRegisterHandler<float> Out13(Context);
+	FRegisterHandler<float> Out14(Context);
+	FRegisterHandler<float> Out15(Context);
+
+	for (int32 i = 0; i < Context.NumInstances; ++i)
+	{
+		*Out00.GetDest() = ToWrite.M[0][0]; Out00.Advance();
+		*Out01.GetDest() = ToWrite.M[0][0]; Out01.Advance();
+		*Out02.GetDest() = ToWrite.M[0][0]; Out02.Advance();
+		*Out03.GetDest() = ToWrite.M[0][0]; Out03.Advance();
+		*Out04.GetDest() = ToWrite.M[0][0]; Out04.Advance();
+		*Out05.GetDest() = ToWrite.M[0][0]; Out05.Advance();
+		*Out06.GetDest() = ToWrite.M[0][0]; Out06.Advance();
+		*Out07.GetDest() = ToWrite.M[0][0]; Out07.Advance();
+		*Out08.GetDest() = ToWrite.M[0][0]; Out08.Advance();
+		*Out09.GetDest() = ToWrite.M[0][0]; Out09.Advance();
+		*Out10.GetDest() = ToWrite.M[0][0]; Out10.Advance();
+		*Out11.GetDest() = ToWrite.M[0][0]; Out11.Advance();
+		*Out12.GetDest() = ToWrite.M[0][0]; Out12.Advance();
+		*Out13.GetDest() = ToWrite.M[0][0]; Out13.Advance();
+		*Out14.GetDest() = ToWrite.M[0][0]; Out14.Advance();
+		*Out15.GetDest() = ToWrite.M[0][0]; Out15.Advance();
+	}
+}
+
+template<typename PosXType, typename PosYType, typename PosZType>
+void UNiagaraDataInterfaceSpline::FindClosestUnitDistanceFromPositionWS(FVectorVMContext& Context)
+{
+	PosXType PosXParam(Context);
+	PosYType PosYParam(Context);
+	PosZType PosZParam(Context);
+	FUserPtrHandler<FNDISpline_InstanceData> InstData(Context);
+	FRegisterHandler<float> OutUnitDistance(Context);
+
+	checkSlow(InstData->Component);
+
+	const int32 NumPoints = InstData->Component->GetSplinePointsPosition().Points.Num();
+	const float FinalKeyTime = InstData->Component->GetSplinePointsPosition().Points[NumPoints - 1].InVal;
+
+	for (int32 i = 0; i < Context.NumInstances; ++i)
+	{
+		float PosX = PosXParam.Get();
+		float PosY = PosYParam.Get();
+		float PosZ = PosZParam.Get();
+
+		FVector Pos(PosX, PosY, PosZ);
+
+		// This first call finds the key time, but this is not in 0..1 range for the spline. 
+		float KeyTime = InstData->Component->FindInputKeyClosestToWorldLocation(Pos);
+		// We need to convert into the range by dividing through by the overall duration of the spline according to the keys.
+		float UnitDistance = KeyTime / FinalKeyTime;
+
+		*OutUnitDistance.GetDest() = UnitDistance;
+
+		PosXParam.Advance();
+		PosYParam.Advance();
+		PosZParam.Advance();
+		OutUnitDistance.Advance();
+	}
+}
+
+void UNiagaraDataInterfaceSpline::GetLocalToWorld(FVectorVMContext& Context)
+{
+	FUserPtrHandler<FNDISpline_InstanceData> InstData(Context);
+	WriteTransform(InstData->Transform, Context);
+}
+
+void UNiagaraDataInterfaceSpline::GetLocalToWorldInverseTransposed(FVectorVMContext& Context)
+{
+	FUserPtrHandler<FNDISpline_InstanceData> InstData(Context);
+	WriteTransform(InstData->TransformInverseTransposed, Context);
+}
+
+#undef LOCTEXT_NAMESPACE
