Copying //UE4/Orion-Staging to //UE4/Main (Source: //Orion/Dev-General @ 3606378)

diff --git a/Engine/Plugins/FX/Niagara/Source/Niagara/Private/NiagaraDataInterfaceStaticMesh.cpp b/Engine/Plugins/FX/Niagara/Source/Niagara/Private/NiagaraDataInterfaceStaticMesh.cpp
new file mode 100644
--- /dev/null
+++ b/Engine/Plugins/FX/Niagara/Source/Niagara/Private/NiagaraDataInterfaceStaticMesh.cpp
@@ -0,0 +1,1547 @@
+// Copyright 1998-2016 Epic Games, Inc. All Rights Reserved.
+
+#include "NiagaraDataInterfaceStaticMesh.h"
+#include "NiagaraEmitterInstance.h"
+#include "Engine/StaticMeshActor.h"
+#include "NiagaraComponent.h"
+#include "NiagaraSystemInstance.h"
+#include "Internationalization.h"
+#include "NiagaraScript.h"
+
+#define LOCTEXT_NAMESPACE "NiagaraDataInterfaceStaticMesh"
+
+FStaticMeshFilteredAreaWeightedSectionSampler::FStaticMeshFilteredAreaWeightedSectionSampler()
+	: Res(nullptr)
+	, Owner(nullptr)
+{
+}
+
+void FStaticMeshFilteredAreaWeightedSectionSampler::Init(FStaticMeshLODResources* InRes, FNDIStaticMesh_InstanceData* InOwner)
+{
+	Res = InRes;
+	Owner = InOwner;
+
+	FStaticMeshAreaWeightedSectionSampler::Init(Res);
+}
+
+float FStaticMeshFilteredAreaWeightedSectionSampler::GetWeights(TArray<float>& OutWeights)
+{
+	check(Owner && Owner->Mesh);
+	float Total = 0.0f;
+	OutWeights.Empty(Owner->GetValidSections().Num());
+	FStaticMeshLODResources& LODRes = Owner->Mesh->RenderData->LODResources[0];
+	for (int32 i = 0; i < Owner->GetValidSections().Num(); ++i)
+	{
+		int32 SecIdx = Owner->GetValidSections()[i];
+		float T = LODRes.AreaWeightedSectionSamplers[SecIdx].GetTotalWeight();
+		OutWeights.Add(T);
+		Total += T;
+	}
+	return Total;
+}
+
+//////////////////////////////////////////////////////////////////////////
+//FNDIStaticMesh_InstanceData
+
+void FNDIStaticMesh_InstanceData::InitVertexColorFiltering()
+{
+	DynamicVertexColorSampler = FNDI_StaticMesh_GeneratedData::GetDynamicColorFilterData(this);
+}
+
+bool FNDIStaticMesh_InstanceData::Init(UNiagaraDataInterfaceStaticMesh* Interface, FNiagaraSystemInstance* SystemInstance)
+{
+	check(SystemInstance);
+	UStaticMesh* PrevMesh = Mesh;
+	Component = nullptr;
+	Mesh = nullptr;
+	Transform = FMatrix::Identity;
+	TransformInverseTransposed = FMatrix::Identity;
+	PrevTransform = FMatrix::Identity;
+	PrevTransformInverseTransposed = FMatrix::Identity;
+	DeltaSeconds = 0.0f;
+	if (Interface->Source)
+	{
+		AStaticMeshActor* MeshActor = Cast<AStaticMeshActor>(Interface->Source);
+		UStaticMeshComponent* SourceComp = nullptr;
+		if (MeshActor != nullptr)
+		{
+			SourceComp = MeshActor->GetStaticMeshComponent();
+		}
+		else
+		{
+			SourceComp = Interface->Source->FindComponentByClass<UStaticMeshComponent>();
+		}
+
+		if (SourceComp)
+		{
+			Mesh = SourceComp->GetStaticMesh();
+			Component = SourceComp;
+		}
+		else
+		{
+			Component = Interface->Source->GetRootComponent();
+		}
+	}
+	else
+	{
+		if (UNiagaraComponent* SimComp = SystemInstance->GetComponent())
+		{
+			if (UStaticMeshComponent* ParentComp = Cast<UStaticMeshComponent>(SimComp->GetAttachParent()))
+			{
+				Component = ParentComp;
+				Mesh = ParentComp->GetStaticMesh();
+			}
+			else if (UStaticMeshComponent* OuterComp = SimComp->GetTypedOuter<UStaticMeshComponent>())
+			{
+				Component = OuterComp;
+				Mesh = OuterComp->GetStaticMesh();
+			}
+			else if (AActor* Owner = SimComp->GetAttachmentRootActor())
+			{
+				TArray<UActorComponent*> SourceComps = Owner->GetComponentsByClass(UStaticMeshComponent::StaticClass());
+				for (UActorComponent* ActorComp : SourceComps)
+				{
+					UStaticMeshComponent* SourceComp = Cast<UStaticMeshComponent>(ActorComp);
+					if (SourceComp)
+					{
+						UStaticMesh* PossibleMesh = SourceComp->GetStaticMesh();
+						if (PossibleMesh != nullptr && PossibleMesh->bAllowCPUAccess)
+						{
+							Mesh = PossibleMesh;
+							Component = SourceComp;
+							break;
+						}
+					}
+				}
+			}
+
+			if (!Component.IsValid())
+			{
+				Component = SimComp;
+			}
+		}
+	}
+
+	check(Component.IsValid());
+
+	if (!Mesh && Interface->DefaultMesh)
+	{
+		Mesh = Interface->DefaultMesh;
+	}
+
+	if (Component.IsValid() && Mesh)
+	{
+		PrevTransform = Transform;
+		PrevTransformInverseTransposed = TransformInverseTransposed;
+		Transform = Component->GetComponentToWorld().ToMatrixWithScale();
+		TransformInverseTransposed = Transform.InverseFast().GetTransposed();
+	}
+
+	if (!Mesh)
+	{
+		UE_LOG(LogNiagara, Log, TEXT("StaticMesh data interface has no valid mesh. Failed InitPerInstanceData - %s"), *Interface->GetFullName());
+		return false;
+	}
+
+	if (!Mesh->bAllowCPUAccess)
+	{
+		UE_LOG(LogNiagara, Log, TEXT("StaticMesh data interface using a mesh that does not allow CPU access. Failed InitPerInstanceData - Mesh: %s"), *Mesh->GetFullName());
+		return false;
+	}
+
+	if (!Component.IsValid())
+	{
+		UE_LOG(LogNiagara, Log, TEXT("StaticMesh data interface has no valid component. Failed InitPerInstanceData - %s"), *Interface->GetFullName());
+		return false;
+	}
+
+	bIsAreaWeightedSampling = Mesh->bSupportUniformlyDistributedSampling;
+
+	//Init the instance filter
+	ValidSections.Empty();
+	FStaticMeshLODResources& Res = Mesh->RenderData->LODResources[0];
+	for (int32 i = 0; i < Res.Sections.Num(); ++i)
+	{
+		if (Interface->SectionFilter.AllowedMaterialSlots.Num() == 0 || Interface->SectionFilter.AllowedMaterialSlots.Contains(Res.Sections[i].MaterialIndex))
+		{
+			ValidSections.Add(i);
+		}
+	}
+
+	if (GetValidSections().Num() == 0)
+	{
+		UE_LOG(LogNiagara, Log, TEXT("StaticMesh data interface has a section filter preventing any spawning. Failed InitPerInstanceData - %s"), *Interface->GetFullName());
+		return false;
+	}
+
+	Sampler.Init(&Res, this);
+
+	return true;
+}
+
+bool FNDIStaticMesh_InstanceData::ResetRequired()const
+{
+	check(GetActualMesh());
+
+	if (!Component.IsValid())
+	{
+		//The component we were bound to is no longer valid so we have to trigger a reset.
+		return true;
+	}
+
+	bool bPrevAreaWeighted = bIsAreaWeightedSampling;
+	//bool bPrevVCSampling = bSupportingVertexColorSampling;//TODO: Vertex color filtering needs more work.
+	bool bReset = false;
+	if (Mesh)
+	{
+		//bSupportingVertexColorSampling = bEnableVertexColorRangeSorting && MeshHasColors();
+		bReset = !Mesh->bAllowCPUAccess || Mesh->bSupportUniformlyDistributedSampling != bPrevAreaWeighted/* || bSupportingVertexColorSampling != bPrevVCSampling*/;
+	}
+	return bReset;
+}
+
+bool FNDIStaticMesh_InstanceData::Tick(UNiagaraDataInterfaceStaticMesh* Interface, FNiagaraSystemInstance* SystemInstance, float InDeltaSeconds)
+{
+	if (ResetRequired())
+	{
+		return true;
+	}
+	else
+	{
+		DeltaSeconds = InDeltaSeconds;
+		if (Component.IsValid() && Mesh)
+		{
+			PrevTransform = Transform;
+			PrevTransformInverseTransposed = TransformInverseTransposed;
+			Transform = Component->GetComponentToWorld().ToMatrixWithScale();
+			TransformInverseTransposed = Transform.InverseFast().GetTransposed();
+		}
+		else
+		{
+			PrevTransform = FMatrix::Identity;
+			PrevTransformInverseTransposed = FMatrix::Identity;
+			Transform = FMatrix::Identity;
+			TransformInverseTransposed = FMatrix::Identity;
+		}
+		return false;
+	}
+}
+
+//void UNiagaraDataInterfaceStaticMesh::ResetFilter()
+//{
+//#if WITH_EDITOR	TODO: Make editor only and serialize this stuff out.
+//	SectionFilter.Init(this, UsesAreaWeighting());
+
+//TODO: Vertex color filtering needs some more work.
+// 	// If we have enabled vertex color sorting for emission, given that these are vertex colors
+// 	// and is limited to a byte, we can make lookup relatively quick by just having a bucket
+// 	// for every possible entry.
+// 	// Will want a better strategy in the long run, but for now this is trivial for GDC..
+// 	TrianglesSortedByVertexColor.Empty();
+// 	if (bEnableVertexColorRangeSorting && bSupportingVertexColorSampling)
+// 	{
+// 		VertexColorToTriangleStart.AddDefaulted(256);
+// 		if (UStaticMesh* ActualMesh = GetActualMesh())
+// 		{
+// 			FStaticMeshLODResources& Res = ActualMesh->RenderData->LODResources[0];
+// 
+// 			// Go over all triangles for each possible vertex color and add it to that bucket
+// 			for (int32 i = 0; i < VertexColorToTriangleStart.Num(); i++)
+// 			{
+// 				uint32 MinVertexColorRed = i;
+// 				uint32 MaxVertexColorRed = i + 1;
+// 				VertexColorToTriangleStart[i] = TrianglesSortedByVertexColor.Num();
+// 
+// 				FIndexArrayView IndexView = Res.IndexBuffer.GetArrayView();
+// 				for (int32 j = 0; j < SectionFilter.GetValidSections().Num(); j++)
+// 				{
+// 					int32 SectionIdx = SectionFilter.GetValidSections()[j];
+// 					int32 TriStartIdx = Res.Sections[SectionIdx].FirstIndex;
+// 					for (uint32 TriIdx = 0; TriIdx < Res.Sections[SectionIdx].NumTriangles; TriIdx++)
+// 					{
+// 						uint32 V0Idx = IndexView[TriStartIdx + TriIdx * 3 + 0];
+// 						uint32 V1Idx = IndexView[TriStartIdx + TriIdx * 3 + 1];
+// 						uint32 V2Idx = IndexView[TriStartIdx + TriIdx * 3 + 2];
+// 
+// 						uint8 MaxR = FMath::Max<uint8>(Res.ColorVertexBuffer.VertexColor(V0Idx).R,
+// 							FMath::Max<uint8>(Res.ColorVertexBuffer.VertexColor(V1Idx).R,
+// 								Res.ColorVertexBuffer.VertexColor(V2Idx).R));
+// 						if (MaxR >= MinVertexColorRed && MaxR < MaxVertexColorRed)
+// 						{
+// 							TrianglesSortedByVertexColor.Add(TriStartIdx + TriIdx * 3);
+// 						}
+// 					}
+// 				}
+// 			}
+// 		}
+// 	}
+//	bFilterInitialized = true;
+
+//#endif
+//}
+
+//////////////////////////////////////////////////////////////////////////
+
+
+struct FTransformHandlerNoop
+{
+	FORCEINLINE void Transform(FVector& V, FMatrix& M) {  }
+};
+
+struct FTransformHandlerPosition
+{
+	FORCEINLINE void Transform(FVector& P, FMatrix& M) { P = M.TransformPosition(P); }
+};
+
+struct FTransformHandlerVector
+{
+	FORCEINLINE void Transform(FVector& V, FMatrix& M) { V = M.TransformVector(V).GetUnsafeNormal3(); }
+};
+
+UNiagaraDataInterfaceStaticMesh::UNiagaraDataInterfaceStaticMesh(FObjectInitializer const& ObjectInitializer)
+	: Super(ObjectInitializer)
+	, DefaultMesh(nullptr)
+	, Source(nullptr)	
+	//, bSupportingVertexColorSampling(0)//Vertex color filtering needs some more work.
+	//, bFilterInitialized(false)
+{
+
+}
+
+
+#if WITH_EDITOR
+
+void UNiagaraDataInterfaceStaticMesh::PostInitProperties()
+{
+	Super::PostInitProperties();
+
+	//Can we register data interfaces as regular types and fold them into the FNiagaraVariable framework for UI and function calls etc?
+	if (HasAnyFlags(RF_ClassDefaultObject))
+	{
+		FNiagaraTypeRegistry::Register(FNiagaraTypeDefinition(GetClass()), true, false, false);
+
+		//Still some issues with using custom structs. Convert node for example throws a wobbler. TODO after GDC.
+		FNiagaraTypeRegistry::Register(FMeshTriCoordinate::StaticStruct(), true, true, false);
+	}
+}
+
+#endif //WITH_EDITOR
+
+
+static const FName RandomSectionName("RandomSection");
+static const FName RandomTriCoordName("RandomTriCoord");
+static const FName RandomTriCoordOnSectionName("RandomTriCoordOnSection");
+static const FName RandomTriCoordVCFilteredName("RandomTriCoordUsingVertexColorFilter");
+
+static const FName GetTriPositionName("GetTriPosition");
+static const FName GetTriNormalName("GetTriNormal");
+static const FName GetTriTangentsName("GetTriTangents");
+
+static const FName GetTriPositionWSName("GetTriPositionWS");
+static const FName GetTriNormalWSName("GetTriNormalWS");
+static const FName GetTriTangentsWSName("GetTriTangentsWS");
+
+static const FName GetTriColorName("GetTriColor");
+static const FName GetTriUVName("GetTriUV");
+
+static const FName GetTriPositionAndVelocityName("GetTriPositionAndVelocityWS");
+
+/** Temporary solution for exposing the transform of a mesh. Ideally this would be done by allowing interfaces to add to the uniform set for a simulation. */
+static const FName GetMeshLocalToWorldName("GetLocalToWorld");
+static const FName GetMeshLocalToWorldInverseTransposedName("GetMeshLocalToWorldInverseTransposed");
+static const FName GetMeshWorldVelocityName("GetWorldVelocity");
+
+void UNiagaraDataInterfaceStaticMesh::GetFunctions(TArray<FNiagaraFunctionSignature>& OutFunctions)
+{
+	{
+		FNiagaraFunctionSignature Sig;
+		Sig.Name = RandomSectionName;
+		Sig.Inputs.Add(FNiagaraVariable(FNiagaraTypeDefinition(GetClass()), TEXT("StaticMesh")));
+		Sig.Outputs.Add(FNiagaraVariable(FNiagaraTypeDefinition::GetIntDef(), TEXT("Section")));
+		Sig.bMemberFunction = true;
+		Sig.bRequiresContext = false;
+		//Sig.Owner = *GetFullName();
+		OutFunctions.Add(Sig);
+	}
+
+	{
+		FNiagaraFunctionSignature Sig;
+		Sig.Name = RandomTriCoordName;
+		Sig.Inputs.Add(FNiagaraVariable(FNiagaraTypeDefinition(GetClass()), TEXT("StaticMesh")));
+		Sig.Outputs.Add(FNiagaraVariable(FNiagaraTypeDefinition(FMeshTriCoordinate::StaticStruct()), TEXT("Coord")));
+		//Sig.Outputs.Add(FNiagaraVariable(FNiagaraTypeDefinition::GetIntDef(), TEXT("Tri Index")));
+		//Sig.Outputs.Add(FNiagaraVariable(FNiagaraTypeDefinition::GetVec3Def(), TEXT("Bary Coord")));
+		Sig.bMemberFunction = true;
+		Sig.bRequiresContext = false;
+		//Sig.Owner = *GetFullName();
+		OutFunctions.Add(Sig);
+	}
+
+	{
+		FNiagaraFunctionSignature Sig;
+		Sig.Name = RandomTriCoordVCFilteredName;
+		Sig.Inputs.Add(FNiagaraVariable(FNiagaraTypeDefinition(GetClass()), TEXT("StaticMesh")));
+		Sig.Inputs.Add(FNiagaraVariable(FNiagaraTypeDefinition(FNiagaraTypeDefinition::GetFloatDef()), TEXT("Start")));
+		Sig.Inputs.Add(FNiagaraVariable(FNiagaraTypeDefinition(FNiagaraTypeDefinition::GetFloatDef()), TEXT("Range")));
+		Sig.Outputs.Add(FNiagaraVariable(FNiagaraTypeDefinition(FMeshTriCoordinate::StaticStruct()), TEXT("Coord")));
+		//Sig.Outputs.Add(FNiagaraVariable(FNiagaraTypeDefinition::GetIntDef(), TEXT("Tri Index")));
+		//Sig.Outputs.Add(FNiagaraVariable(FNiagaraTypeDefinition::GetVec3Def(), TEXT("Bary Coord")));
+		Sig.bMemberFunction = true;
+		Sig.bRequiresContext = false;
+		//Sig.Owner = *GetFullName();
+		Sig.SetDescription(LOCTEXT("DataInterfaceSpline_RandomTriCoordVCFiltered", "If bSupportingVertexColorSampling is set on the data source, will randomly find a triangle whose red channel is within the Start to Start + Range color range."));
+		OutFunctions.Add(Sig);
+	}
+
+	{
+		FNiagaraFunctionSignature Sig;
+		Sig.Name = RandomTriCoordOnSectionName;
+		Sig.Inputs.Add(FNiagaraVariable(FNiagaraTypeDefinition(GetClass()), TEXT("StaticMesh")));
+		Sig.Inputs.Add(FNiagaraVariable(FNiagaraTypeDefinition::GetIntDef(), TEXT("Section")));
+		Sig.Outputs.Add(FNiagaraVariable(FNiagaraTypeDefinition(FMeshTriCoordinate::StaticStruct()), TEXT("Coord")));
+		//Sig.Outputs.Add(FNiagaraVariable(FNiagaraTypeDefinition::GetIntDef(), TEXT("Tri Index")));
+		//Sig.Outputs.Add(FNiagaraVariable(FNiagaraTypeDefinition::GetVec3Def(), TEXT("Bary Coord")));
+		Sig.bMemberFunction = true;
+		Sig.bRequiresContext = false;
+		//Sig.Owner = *GetFullName();
+		OutFunctions.Add(Sig);
+	}
+
+	{
+		FNiagaraFunctionSignature Sig;
+		Sig.Name = GetTriPositionName;
+		Sig.Inputs.Add(FNiagaraVariable(FNiagaraTypeDefinition(GetClass()), TEXT("StaticMesh")));
+		Sig.Inputs.Add(FNiagaraVariable(FNiagaraTypeDefinition(FMeshTriCoordinate::StaticStruct()), TEXT("Coord")));
+		//Sig.Inputs.Add(FNiagaraVariable(FNiagaraTypeDefinition::GetIntDef(), TEXT("Tri Index")));
+		//Sig.Inputs.Add(FNiagaraVariable(FNiagaraTypeDefinition::GetVec3Def(), TEXT("Bary Coord")));
+
+		Sig.Outputs.Add(FNiagaraVariable(FNiagaraTypeDefinition::GetVec3Def(), TEXT("Position")));
+		Sig.bMemberFunction = true;
+		Sig.bRequiresContext = false;
+		//Sig.Owner = *GetFullName();
+		OutFunctions.Add(Sig);
+	}
+
+	{
+		FNiagaraFunctionSignature Sig;
+		Sig.Name = GetTriPositionAndVelocityName;
+		Sig.Inputs.Add(FNiagaraVariable(FNiagaraTypeDefinition(GetClass()), TEXT("StaticMesh")));
+		Sig.Inputs.Add(FNiagaraVariable(FNiagaraTypeDefinition(FMeshTriCoordinate::StaticStruct()), TEXT("Coord")));
+		//Sig.Inputs.Add(FNiagaraVariable(FNiagaraTypeDefinition::GetIntDef(), TEXT("Tri Index")));
+		//Sig.Inputs.Add(FNiagaraVariable(FNiagaraTypeDefinition::GetVec3Def(), TEXT("Bary Coord")));
+
+		Sig.Outputs.Add(FNiagaraVariable(FNiagaraTypeDefinition::GetVec3Def(), TEXT("Position")));
+		Sig.Outputs.Add(FNiagaraVariable(FNiagaraTypeDefinition::GetVec3Def(), TEXT("Velocity")));
+		Sig.bMemberFunction = true;
+		Sig.bRequiresContext = false;
+		//Sig.Owner = *GetFullName();
+		OutFunctions.Add(Sig);
+	}
+
+	{
+		FNiagaraFunctionSignature Sig;
+		Sig.Name = GetTriPositionWSName;
+		Sig.Inputs.Add(FNiagaraVariable(FNiagaraTypeDefinition(GetClass()), TEXT("StaticMesh")));
+		Sig.Inputs.Add(FNiagaraVariable(FNiagaraTypeDefinition(FMeshTriCoordinate::StaticStruct()), TEXT("Coord")));
+		//Sig.Inputs.Add(FNiagaraVariable(FNiagaraTypeDefinition::GetIntDef(), TEXT("Tri Index")));
+		//Sig.Inputs.Add(FNiagaraVariable(FNiagaraTypeDefinition::GetVec3Def(), TEXT("Bary Coord")));
+
+		Sig.Outputs.Add(FNiagaraVariable(FNiagaraTypeDefinition::GetVec3Def(), TEXT("Position")));
+		Sig.bMemberFunction = true;
+		Sig.bRequiresContext = false;
+		//Sig.Owner = *GetFullName();
+		OutFunctions.Add(Sig);
+	}
+
+	{
+		FNiagaraFunctionSignature Sig;
+		Sig.Name = GetTriNormalName;
+		Sig.Inputs.Add(FNiagaraVariable(FNiagaraTypeDefinition(GetClass()), TEXT("StaticMesh")));
+		Sig.Inputs.Add(FNiagaraVariable(FNiagaraTypeDefinition(FMeshTriCoordinate::StaticStruct()), TEXT("Coord")));
+		//Sig.Inputs.Add(FNiagaraVariable(FNiagaraTypeDefinition::GetIntDef(), TEXT("Tri Index")));
+		//Sig.Inputs.Add(FNiagaraVariable(FNiagaraTypeDefinition::GetVec3Def(), TEXT("Bary Coord")));
+
+		Sig.Outputs.Add(FNiagaraVariable(FNiagaraTypeDefinition::GetVec3Def(), TEXT("Normal")));
+		Sig.bMemberFunction = true;
+		Sig.bRequiresContext = false;
+		//	Sig.Owner = *GetFullName();
+		OutFunctions.Add(Sig);
+	}
+
+	{
+		FNiagaraFunctionSignature Sig;
+		Sig.Name = GetTriNormalWSName;
+		Sig.Inputs.Add(FNiagaraVariable(FNiagaraTypeDefinition(GetClass()), TEXT("StaticMesh")));
+		Sig.Inputs.Add(FNiagaraVariable(FNiagaraTypeDefinition(FMeshTriCoordinate::StaticStruct()), TEXT("Coord")));
+		//Sig.Inputs.Add(FNiagaraVariable(FNiagaraTypeDefinition::GetIntDef(), TEXT("Tri Index")));
+		//Sig.Inputs.Add(FNiagaraVariable(FNiagaraTypeDefinition::GetVec3Def(), TEXT("Bary Coord")));
+
+		Sig.Outputs.Add(FNiagaraVariable(FNiagaraTypeDefinition::GetVec3Def(), TEXT("Normal")));
+		Sig.bMemberFunction = true;
+		Sig.bRequiresContext = false;
+		//Sig.Owner = *GetFullName();
+		OutFunctions.Add(Sig);
+	}
+
+	{
+		FNiagaraFunctionSignature Sig;
+		Sig.Name = GetTriTangentsName;
+		Sig.Inputs.Add(FNiagaraVariable(FNiagaraTypeDefinition(GetClass()), TEXT("StaticMesh")));
+		Sig.Inputs.Add(FNiagaraVariable(FNiagaraTypeDefinition(FMeshTriCoordinate::StaticStruct()), TEXT("Coord")));
+		//Sig.Inputs.Add(FNiagaraVariable(FNiagaraTypeDefinition::GetIntDef(), TEXT("Tri Index")));
+		//Sig.Inputs.Add(FNiagaraVariable(FNiagaraTypeDefinition::GetVec3Def(), TEXT("Bary Coord")));
+
+		Sig.Outputs.Add(FNiagaraVariable(FNiagaraTypeDefinition::GetVec3Def(), TEXT("Tangent")));
+		Sig.Outputs.Add(FNiagaraVariable(FNiagaraTypeDefinition::GetVec3Def(), TEXT("Binormal")));
+		Sig.Outputs.Add(FNiagaraVariable(FNiagaraTypeDefinition::GetVec3Def(), TEXT("Normal")));
+		Sig.bMemberFunction = true;
+		Sig.bRequiresContext = false;
+		//Sig.Owner = *GetFullName();
+		OutFunctions.Add(Sig);
+	}
+
+	{
+		FNiagaraFunctionSignature Sig;
+		Sig.Name = GetTriTangentsWSName;
+		Sig.Inputs.Add(FNiagaraVariable(FNiagaraTypeDefinition(GetClass()), TEXT("StaticMesh")));
+		Sig.Inputs.Add(FNiagaraVariable(FNiagaraTypeDefinition(FMeshTriCoordinate::StaticStruct()), TEXT("Coord")));
+		//Sig.Inputs.Add(FNiagaraVariable(FNiagaraTypeDefinition::GetIntDef(), TEXT("Tri Index")));
+		//Sig.Inputs.Add(FNiagaraVariable(FNiagaraTypeDefinition::GetVec3Def(), TEXT("Bary Coord")));
+
+		Sig.Outputs.Add(FNiagaraVariable(FNiagaraTypeDefinition::GetVec3Def(), TEXT("Tangent")));
+		Sig.Outputs.Add(FNiagaraVariable(FNiagaraTypeDefinition::GetVec3Def(), TEXT("Binormal")));
+		Sig.Outputs.Add(FNiagaraVariable(FNiagaraTypeDefinition::GetVec3Def(), TEXT("Normal")));
+		Sig.bMemberFunction = true;
+		Sig.bRequiresContext = false;
+		//Sig.Owner = *GetFullName();
+		OutFunctions.Add(Sig);
+	}
+
+	{
+		FNiagaraFunctionSignature Sig;
+		Sig.Name = GetTriColorName;
+		Sig.Inputs.Add(FNiagaraVariable(FNiagaraTypeDefinition(GetClass()), TEXT("StaticMesh")));
+		Sig.Inputs.Add(FNiagaraVariable(FNiagaraTypeDefinition(FMeshTriCoordinate::StaticStruct()), TEXT("Coord")));
+		//Sig.Inputs.Add(FNiagaraVariable(FNiagaraTypeDefinition::GetIntDef(), TEXT("Tri Index")));
+		//Sig.Inputs.Add(FNiagaraVariable(FNiagaraTypeDefinition::GetVec3Def(), TEXT("Bary Coord")));
+
+		Sig.Outputs.Add(FNiagaraVariable(FNiagaraTypeDefinition::GetColorDef(), TEXT("Color")));
+		Sig.bMemberFunction = true;
+		Sig.bRequiresContext = false;
+		//Sig.Owner = *GetFullName();
+		OutFunctions.Add(Sig);
+	}
+
+	{
+		FNiagaraFunctionSignature Sig;
+		Sig.Name = GetTriUVName;
+		Sig.Inputs.Add(FNiagaraVariable(FNiagaraTypeDefinition(GetClass()), TEXT("StaticMesh")));
+		Sig.Inputs.Add(FNiagaraVariable(FNiagaraTypeDefinition(FMeshTriCoordinate::StaticStruct()), TEXT("Coord")));
+		//Sig.Inputs.Add(FNiagaraVariable(FNiagaraTypeDefinition::GetIntDef(), TEXT("Tri Index")));
+		//Sig.Inputs.Add(FNiagaraVariable(FNiagaraTypeDefinition::GetVec3Def(), TEXT("Bary Coord")));
+		Sig.Inputs.Add(FNiagaraVariable(FNiagaraTypeDefinition::GetIntDef(), TEXT("UV Set")));
+
+		Sig.Outputs.Add(FNiagaraVariable(FNiagaraTypeDefinition::GetVec2Def(), TEXT("UV")));
+		Sig.bMemberFunction = true;
+		Sig.bRequiresContext = false;
+		//Sig.Owner = *GetFullName();
+		OutFunctions.Add(Sig);
+	}
+
+	{
+		FNiagaraFunctionSignature Sig;
+		Sig.Name = GetMeshLocalToWorldName;
+		Sig.Inputs.Add(FNiagaraVariable(FNiagaraTypeDefinition(GetClass()), TEXT("StaticMesh")));
+		//Sig.Inputs.Add(FNiagaraVariable(FNiagaraTypeDefinition(FMeshTriCoordinate::StaticStruct()), TEXT("Coord")));
+		Sig.Outputs.Add(FNiagaraVariable(FNiagaraTypeDefinition::GetMatrix4Def(), TEXT("Transform")));
+		Sig.bMemberFunction = true;
+		Sig.bRequiresContext = false;
+		//Sig.Owner = *GetFullName();
+		OutFunctions.Add(Sig);
+	}
+
+	{
+		FNiagaraFunctionSignature Sig;
+		Sig.Name = GetMeshLocalToWorldInverseTransposedName;
+		Sig.Inputs.Add(FNiagaraVariable(FNiagaraTypeDefinition(GetClass()), TEXT("StaticMesh")));
+		//Sig.Inputs.Add(FNiagaraVariable(FNiagaraTypeDefinition(FMeshTriCoordinate::StaticStruct()), TEXT("Coord")));
+		Sig.Outputs.Add(FNiagaraVariable(FNiagaraTypeDefinition::GetMatrix4Def(), TEXT("Transform")));
+		Sig.bMemberFunction = true;
+		Sig.bRequiresContext = false;
+		//Sig.Owner = *GetFullName();
+		OutFunctions.Add(Sig);
+	}
+
+	{
+		FNiagaraFunctionSignature Sig;
+		Sig.Name = GetMeshWorldVelocityName;
+		Sig.Inputs.Add(FNiagaraVariable(FNiagaraTypeDefinition(GetClass()), TEXT("StaticMesh")));
+		//Sig.Inputs.Add(FNiagaraVariable(FNiagaraTypeDefinition(FMeshTriCoordinate::StaticStruct()), TEXT("Coord")));
+		Sig.Outputs.Add(FNiagaraVariable(FNiagaraTypeDefinition::GetVec3Def(), TEXT("Velocity")));
+		Sig.bMemberFunction = true;
+		Sig.bRequiresContext = false;
+		//Sig.Owner = *GetFullName();
+		OutFunctions.Add(Sig);
+	}
+}
+
+//External function binder choosing between template specializations based on UsesAreaWeighting
+template<typename NextBinder>
+struct TUsesAreaWeightingBinder
+{
+	template<typename... ParamTypes>
+	static FVMExternalFunction Bind(UNiagaraDataInterface* Interface, const FVMExternalFunctionBindingInfo& BindingInfo, void* InstanceData)
+	{
+		FNDIStaticMesh_InstanceData* InstData = (FNDIStaticMesh_InstanceData*)InstanceData;
+		UNiagaraDataInterfaceStaticMesh* MeshInterface = CastChecked<UNiagaraDataInterfaceStaticMesh>(Interface);
+		if (InstData->UsesAreaWeighting())
+		{
+			return NextBinder::template Bind<ParamTypes..., TIntegralConstant<bool, true>>(Interface, BindingInfo, InstanceData);
+		}
+		else
+		{
+			return NextBinder::template Bind<ParamTypes..., TIntegralConstant<bool, false>>(Interface, BindingInfo, InstanceData);
+		}
+	}
+};
+
+//Helper struct for accessing typed vertex data.
+template<EStaticMeshVertexTangentBasisType TangentT, EStaticMeshVertexUVType UVTypeT>
+struct TTypedMeshVertexAccessor
+{
+	const FStaticMeshVertexBuffer& Verts;
+	TTypedMeshVertexAccessor(const FStaticMeshVertexBuffer& InVerts)
+		: Verts(InVerts)
+	{}
+
+	FORCEINLINE FVector GetTangentX(int32 Idx)const { return Verts.VertexTangentX_Typed<TangentT>(Idx); }
+	FORCEINLINE FVector GetTangentY(int32 Idx)const { return Verts.VertexTangentY_Typed<TangentT>(Idx); }
+	FORCEINLINE FVector GetTangentZ(int32 Idx)const { return Verts.VertexTangentZ_Typed<TangentT>(Idx); }
+	FORCEINLINE FVector2D GetUV(int32 Idx, int32 UVSet)const { return Verts.GetVertexUV_Typed<TangentT, UVTypeT>(Idx, UVSet); }
+};
+
+//External function binder choosing between template specializations based on the mesh's vertex type.
+template<typename NextBinder>
+struct TTypedMeshAccessorBinder
+{
+	template<typename... ParamTypes>
+	static FVMExternalFunction Bind(UNiagaraDataInterface* Interface, const FVMExternalFunctionBindingInfo& BindingInfo, void* InstanceData)
+	{
+		FNDIStaticMesh_InstanceData* InstData = (FNDIStaticMesh_InstanceData*)InstanceData;
+		UNiagaraDataInterfaceStaticMesh* MeshInterface = CastChecked<UNiagaraDataInterfaceStaticMesh>(Interface);
+		check(InstData->Mesh);
+		FStaticMeshLODResources& Res = InstData->Mesh->RenderData->LODResources[0];
+		if (Res.VertexBuffer.GetUseHighPrecisionTangentBasis())			
+		{
+			if (Res.VertexBuffer.GetUseFullPrecisionUVs())
+			{
+				return NextBinder::template Bind<ParamTypes..., TTypedMeshVertexAccessor<EStaticMeshVertexTangentBasisType::HighPrecision, EStaticMeshVertexUVType::HighPrecision>>(Interface, BindingInfo, InstanceData);
+			}
+			else
+			{
+				return NextBinder::template Bind<ParamTypes..., TTypedMeshVertexAccessor<EStaticMeshVertexTangentBasisType::HighPrecision, EStaticMeshVertexUVType::Default>>(Interface, BindingInfo, InstanceData);
+			}
+		}
+		else
+		{
+			if (Res.VertexBuffer.GetUseFullPrecisionUVs())
+			{
+				return NextBinder::template Bind<ParamTypes..., TTypedMeshVertexAccessor<EStaticMeshVertexTangentBasisType::Default, EStaticMeshVertexUVType::HighPrecision>>(Interface, BindingInfo, InstanceData);
+			}
+			else
+			{
+				return NextBinder::template Bind<ParamTypes..., TTypedMeshVertexAccessor<EStaticMeshVertexTangentBasisType::Default, EStaticMeshVertexUVType::Default>>(Interface, BindingInfo, InstanceData);
+			}
+		}
+	}
+};
+
+//Final binders for all static mesh interface functions.
+DEFINE_NDI_FUNC_BINDER(UNiagaraDataInterfaceStaticMesh, RandomSection);
+DEFINE_NDI_FUNC_BINDER(UNiagaraDataInterfaceStaticMesh, RandomTriCoord);
+DEFINE_NDI_FUNC_BINDER(UNiagaraDataInterfaceStaticMesh, RandomTriCoordVertexColorFiltered);
+DEFINE_NDI_FUNC_BINDER(UNiagaraDataInterfaceStaticMesh, RandomTriCoordOnSection);
+DEFINE_NDI_FUNC_BINDER(UNiagaraDataInterfaceStaticMesh, GetTriCoordPosition);
+DEFINE_NDI_FUNC_BINDER(UNiagaraDataInterfaceStaticMesh, GetTriCoordNormal);
+DEFINE_NDI_FUNC_BINDER(UNiagaraDataInterfaceStaticMesh, GetTriCoordTangents);
+DEFINE_NDI_FUNC_BINDER(UNiagaraDataInterfaceStaticMesh, GetTriCoordColor);
+DEFINE_NDI_FUNC_BINDER(UNiagaraDataInterfaceStaticMesh, GetTriCoordUV);
+DEFINE_NDI_FUNC_BINDER(UNiagaraDataInterfaceStaticMesh, GetTriCoordPositionAndVelocity);
+
+FVMExternalFunction UNiagaraDataInterfaceStaticMesh::GetVMExternalFunction(const FVMExternalFunctionBindingInfo& BindingInfo, void* InstanceData)
+{
+	FNDIStaticMesh_InstanceData* InstData = (FNDIStaticMesh_InstanceData*)InstanceData;
+	check(InstData && InstData->Mesh && InstData->Component.IsValid());
+
+	bool bNeedsVertexPositions = false;
+	bool bNeedsVertexColors = false;
+	bool bNeedsVertMain = true;//Assuming we always need this?
+	
+	FVMExternalFunction Function;
+	if (BindingInfo.Name == RandomSectionName)
+	{
+		check(BindingInfo.GetNumInputs() == 1 && BindingInfo.GetNumOutputs() == 1);
+		Function = TUsesAreaWeightingBinder<NDI_FUNC_BINDER(UNiagaraDataInterfaceStaticMesh, RandomSection)>::Bind(this, BindingInfo, InstanceData);
+	}
+	else if (BindingInfo.Name == RandomTriCoordName)
+	{
+		check(BindingInfo.GetNumInputs() == 1 && BindingInfo.GetNumOutputs() == 4);
+		Function = TUsesAreaWeightingBinder<NDI_FUNC_BINDER(UNiagaraDataInterfaceStaticMesh, RandomTriCoord)>::Bind(this, BindingInfo, InstanceData);
+	}
+	//TODO: Vertex color filtering needs more work.
+	else if (BindingInfo.Name == RandomTriCoordVCFilteredName)
+	{
+		InstData->InitVertexColorFiltering();
+		check(BindingInfo.GetNumInputs() == 3 && BindingInfo.GetNumOutputs() == 4);
+		Function = TNDIParamBinder<0, float, TNDIParamBinder<1, float, NDI_FUNC_BINDER(UNiagaraDataInterfaceStaticMesh, RandomTriCoordVertexColorFiltered)>>::Bind(this, BindingInfo, InstanceData);
+	}	
+	else if (BindingInfo.Name == RandomTriCoordOnSectionName)
+	{
+		check(BindingInfo.GetNumInputs() == 2 && BindingInfo.GetNumOutputs() == 4);
+		Function = TUsesAreaWeightingBinder<TNDIParamBinder<0, int32, NDI_FUNC_BINDER(UNiagaraDataInterfaceStaticMesh, RandomTriCoordOnSection)>>::Bind(this, BindingInfo, InstanceData);
+	}
+	if (BindingInfo.Name == GetTriPositionName)
+	{
+		check(BindingInfo.GetNumInputs() == 5 && BindingInfo.GetNumOutputs() == 3);
+		bNeedsVertexPositions = true;
+		Function = TNDIExplicitBinder<FTransformHandlerNoop, TNDIParamBinder<0, int32, TNDIParamBinder<1, float, TNDIParamBinder<2, float, TNDIParamBinder<3, float, NDI_FUNC_BINDER(UNiagaraDataInterfaceStaticMesh, GetTriCoordPosition)>>>>>::Bind(this, BindingInfo, InstanceData);
+	}
+	else if (BindingInfo.Name == GetTriPositionWSName)
+	{
+		check(BindingInfo.GetNumInputs() == 5 && BindingInfo.GetNumOutputs() == 3);
+		bNeedsVertexPositions = true;
+		Function = TNDIExplicitBinder<FTransformHandlerPosition, TNDIParamBinder<0, int32, TNDIParamBinder<1, float, TNDIParamBinder<2, float, TNDIParamBinder<3, float, NDI_FUNC_BINDER(UNiagaraDataInterfaceStaticMesh, GetTriCoordPosition)>>>>>::Bind(this, BindingInfo, InstanceData);
+	}
+	else if (BindingInfo.Name == GetTriNormalName)
+	{
+		check(BindingInfo.GetNumInputs() == 6 && BindingInfo.GetNumOutputs() == 3);
+		bNeedsVertMain = true;
+		Function = TNDIExplicitBinder<FTransformHandlerNoop, TNDIParamBinder<0, int32, TNDIParamBinder<1, float, TNDIParamBinder<2, float, TNDIParamBinder<3, float, NDI_FUNC_BINDER(UNiagaraDataInterfaceStaticMesh, GetTriCoordNormal)>>>>>::Bind(this, BindingInfo, InstanceData);
+	}
+	else if (BindingInfo.Name == GetTriNormalWSName)
+	{
+		check(BindingInfo.GetNumInputs() == 5 && BindingInfo.GetNumOutputs() == 3);
+		bNeedsVertMain = true;
+		Function = TNDIExplicitBinder<FTransformHandlerVector, TNDIParamBinder<0, int32, TNDIParamBinder<1, float, TNDIParamBinder<2, float, TNDIParamBinder<3, float, NDI_FUNC_BINDER(UNiagaraDataInterfaceStaticMesh, GetTriCoordNormal)>>>>>::Bind(this, BindingInfo, InstanceData);
+	}
+	else if (BindingInfo.Name == GetTriTangentsName)
+	{
+		check(BindingInfo.GetNumInputs() == 5 && BindingInfo.GetNumOutputs() == 9);
+		bNeedsVertMain = true;
+		Function = TTypedMeshAccessorBinder<TNDIExplicitBinder<FTransformHandlerNoop, TNDIParamBinder<0, int32, TNDIParamBinder<1, float, TNDIParamBinder<2, float, TNDIParamBinder<3, float, NDI_FUNC_BINDER(UNiagaraDataInterfaceStaticMesh, GetTriCoordTangents)>>>>>>::Bind(this, BindingInfo, InstanceData);
+	}
+	else if (BindingInfo.Name == GetTriTangentsWSName)
+	{
+		check(BindingInfo.GetNumInputs() == 5 && BindingInfo.GetNumOutputs() == 9);
+		bNeedsVertMain = true;
+		Function = TTypedMeshAccessorBinder<TNDIExplicitBinder<FTransformHandlerVector, TNDIParamBinder<0, int32, TNDIParamBinder<1, float, TNDIParamBinder<2, float, TNDIParamBinder<3, float, NDI_FUNC_BINDER(UNiagaraDataInterfaceStaticMesh, GetTriCoordTangents)>>>>>>::Bind(this, BindingInfo, InstanceData);
+	}
+	else if (BindingInfo.Name == GetTriColorName)
+	{
+		check(BindingInfo.GetNumInputs() == 5 && BindingInfo.GetNumOutputs() == 4);
+		bNeedsVertexColors = true;
+		Function = TNDIParamBinder<0, int32, TNDIParamBinder<1, float, TNDIParamBinder<2, float, TNDIParamBinder<3, float, NDI_FUNC_BINDER(UNiagaraDataInterfaceStaticMesh, GetTriCoordColor)>>>>::Bind(this, BindingInfo, InstanceData);
+	}
+	else if (BindingInfo.Name == GetTriUVName)
+	{
+		check(BindingInfo.GetNumInputs() == 6 && BindingInfo.GetNumOutputs() == 2);
+		bNeedsVertMain = true;
+		Function = TTypedMeshAccessorBinder<TNDIParamBinder<0, int32, TNDIParamBinder<1, float, TNDIParamBinder<2, float, TNDIParamBinder<3, float, TNDIParamBinder<4, int32, NDI_FUNC_BINDER(UNiagaraDataInterfaceStaticMesh, GetTriCoordUV)>>>>>>::Bind(this, BindingInfo, InstanceData);
+	}
+	else if (BindingInfo.Name == GetTriPositionAndVelocityName)
+	{
+		check(BindingInfo.GetNumInputs() == 5 && BindingInfo.GetNumOutputs() == 6);
+		bNeedsVertMain = true;
+		bNeedsVertexPositions = true;
+		Function = TNDIParamBinder<0, int32, TNDIParamBinder<1, float, TNDIParamBinder<2, float, TNDIParamBinder<3, float, NDI_FUNC_BINDER(UNiagaraDataInterfaceStaticMesh, GetTriCoordPositionAndVelocity)>>>>::Bind(this, BindingInfo, InstanceData);
+	}
+	else if (BindingInfo.Name == GetMeshLocalToWorldName)
+	{
+		check(BindingInfo.GetNumInputs() == 1 && BindingInfo.GetNumOutputs() == 16);
+		Function = FVMExternalFunction::CreateUObject(this, &UNiagaraDataInterfaceStaticMesh::GetLocalToWorld);
+	}
+	else if (BindingInfo.Name == GetMeshLocalToWorldInverseTransposedName)
+	{
+		check(BindingInfo.GetNumInputs() == 1 && BindingInfo.GetNumOutputs() == 16);
+		Function = FVMExternalFunction::CreateUObject(this, &UNiagaraDataInterfaceStaticMesh::GetLocalToWorldInverseTransposed);
+	}
+	else if (BindingInfo.Name == GetMeshWorldVelocityName)
+	{
+		check(BindingInfo.GetNumInputs() == 1 && BindingInfo.GetNumOutputs() == 3);
+		Function = FVMExternalFunction::CreateUObject(this, &UNiagaraDataInterfaceStaticMesh::GetWorldVelocity);
+	}
+
+	if (bNeedsVertexPositions && !InstData->MeshHasPositions())
+	{
+		UE_LOG(LogNiagara, Log, TEXT("Static Mesh data interface is cannot run as it's reading position data on a mesh that does not provide it. - Mesh:%s  "), *InstData->Mesh->GetFullName());
+		Function = FVMExternalFunction();
+	}
+	if (bNeedsVertexColors && !InstData->MeshHasColors())
+	{
+		UE_LOG(LogNiagara, Log, TEXT("Static Mesh data interface is cannot run as it's reading color data on a mesh that does not provide it. - Mesh:%s  "), *InstData->Mesh->GetFullName());
+		Function = FVMExternalFunction();
+	}
+	if (bNeedsVertMain && !InstData->MeshHasVerts())
+	{
+		UE_LOG(LogNiagara, Log, TEXT("Static Mesh data interface is cannot run as it's reading vertex data on a mesh with no vertex data. - Mesh:%s  "), *InstData->Mesh->GetFullName());
+		Function = FVMExternalFunction();
+	}
+
+	return Function;
+}
+
+bool UNiagaraDataInterfaceStaticMesh::CopyTo(UNiagaraDataInterface* Destination) const
+{
+	if (!Super::CopyTo(Destination))
+	{
+		return false;
+	}
+
+	UNiagaraDataInterfaceStaticMesh* OtherTyped = CastChecked<UNiagaraDataInterfaceStaticMesh>(Destination);
+	OtherTyped->Source = Source;
+	OtherTyped->DefaultMesh = DefaultMesh;
+	//OtherTyped->bEnableVertexColorRangeSorting = bEnableVertexColorRangeSorting;//TODO: Vertex color filtering needs more work.
+	OtherTyped->SectionFilter = SectionFilter;
+	return true;
+}
+
+bool UNiagaraDataInterfaceStaticMesh::Equals(const UNiagaraDataInterface* Other) const
+{
+	if (!Super::Equals(Other))
+	{
+		return false;
+	}
+	const UNiagaraDataInterfaceStaticMesh* OtherTyped = CastChecked<const UNiagaraDataInterfaceStaticMesh>(Other);
+	return OtherTyped->Source == Source &&
+		OtherTyped->DefaultMesh == DefaultMesh &&
+		//OtherTyped->bEnableVertexColorRangeSorting == bEnableVertexColorRangeSorting &&//TODO: Vertex color filtering needs more work.
+		OtherTyped->SectionFilter.AllowedMaterialSlots == SectionFilter.AllowedMaterialSlots;
+}
+
+bool UNiagaraDataInterfaceStaticMesh::InitPerInstanceData(void* PerInstanceData, FNiagaraSystemInstance* SystemInstance)
+{
+	FNDIStaticMesh_InstanceData* Inst = new (PerInstanceData) FNDIStaticMesh_InstanceData();
+	return Inst->Init(this, SystemInstance);
+}
+
+void UNiagaraDataInterfaceStaticMesh::DestroyPerInstanceData(void* PerInstanceData, FNiagaraSystemInstance* SystemInstance)
+{
+	FNDIStaticMesh_InstanceData* Inst = (FNDIStaticMesh_InstanceData*)PerInstanceData;
+	Inst->~FNDIStaticMesh_InstanceData();
+}
+
+bool UNiagaraDataInterfaceStaticMesh::PerInstanceTick(void* PerInstanceData, FNiagaraSystemInstance* SystemInstance, float InDeltaSeconds)
+{
+	FNDIStaticMesh_InstanceData* Inst = (FNDIStaticMesh_InstanceData*)PerInstanceData;
+	return Inst->Tick(this, SystemInstance, InDeltaSeconds);
+}
+
+FORCEINLINE FVector RandomBarycentricCoord(FRandomStream& RandStream)
+{
+	//TODO: This is gonna be slooooow. Move to an LUT possibly or find faster method.
+	//Can probably handle lower quality randoms / uniformity for a decent speed win.
+	float r0 = RandStream.GetFraction();
+	float r1 = RandStream.GetFraction();
+	float sqrt0 = FMath::Sqrt(r0);
+	float sqrt1 = FMath::Sqrt(r1);
+	return FVector(1.0f - sqrt0, sqrt0 * (1.0 - r1), r1 * sqrt0);
+}
+
+template<typename T>
+T BarycentricInterpolate(float BaryX, float BaryY, float BaryZ, T V0, T V1, T V2)
+{
+	return V0 * BaryX + V1 * BaryY + V2 * BaryZ;
+}
+
+// Overload for FVector4 to work around C2719: (formal parameter with requested alignment of 16 won't be aligned)
+FVector4 BarycentricInterpolate(float BaryX, float BaryY, float BaryZ, const FVector4& V0, const FVector4& V1, const FVector4& V2)
+{
+	return V0 * BaryX + V1 * BaryY + V2 * BaryZ;
+}
+
+//RandomSection specializations.
+//Each combination for AreaWeighted and Section filtered.
+template<>
+FORCEINLINE int32 UNiagaraDataInterfaceStaticMesh::RandomSection<TIntegralConstant<bool, true>, true>(FStaticMeshLODResources& Res, FNDIStaticMesh_InstanceData* InstData)
+{
+	checkSlow(InstData->GetValidSections().Num() > 0);
+	int32 Idx = InstData->GetAreaWeigtedSampler().GetEntryIndex(InstData->RandStream.GetFraction(), InstData->RandStream.GetFraction());
+	return InstData->GetValidSections()[Idx];
+}
+
+template<>
+FORCEINLINE int32 UNiagaraDataInterfaceStaticMesh::RandomSection<TIntegralConstant<bool, true>, false>(FStaticMeshLODResources& Res, FNDIStaticMesh_InstanceData* InstData)
+{
+	return Res.AreaWeightedSampler.GetEntryIndex(InstData->RandStream.GetFraction(), InstData->RandStream.GetFraction());
+}
+
+template<>
+FORCEINLINE int32 UNiagaraDataInterfaceStaticMesh::RandomSection<TIntegralConstant<bool, false>, true>(FStaticMeshLODResources& Res, FNDIStaticMesh_InstanceData* InstData)
+{
+	checkSlow(InstData->GetValidSections().Num() > 0);
+	int32 Idx = InstData->RandStream.RandRange(0, InstData->GetValidSections().Num() - 1);
+	return InstData->GetValidSections()[Idx];
+}
+
+template<>
+FORCEINLINE int32 UNiagaraDataInterfaceStaticMesh::RandomSection<TIntegralConstant<bool, false>, false>(FStaticMeshLODResources& Res, FNDIStaticMesh_InstanceData* InstData)
+{
+	return InstData->RandStream.RandRange(0, Res.Sections.Num() - 1);
+}
+
+template<typename TAreaWeighted>
+void UNiagaraDataInterfaceStaticMesh::RandomSection(FVectorVMContext& Context)
+{
+	FUserPtrHandler<FNDIStaticMesh_InstanceData> InstData(Context);
+	FRegisterHandler<int32> OutSection(Context);
+
+	FStaticMeshLODResources& Res = InstData->Mesh->RenderData->LODResources[0];
+	for (int32 i = 0; i < Context.NumInstances; ++i)
+	{
+		*OutSection.GetDest() = RandomSection<TAreaWeighted, true>(Res, InstData);
+		OutSection.Advance();
+	}
+}
+
+//RandomTriIndex specializations.
+//Each combination for AreaWeighted and Section filtered.
+template<>
+FORCEINLINE int32 UNiagaraDataInterfaceStaticMesh::RandomTriIndex<TIntegralConstant<bool, true>, true>(FStaticMeshLODResources& Res, FNDIStaticMesh_InstanceData* InstData)
+{
+	int32 SecIdx = RandomSection<TIntegralConstant<bool, true>, true>(Res, InstData);
+	FStaticMeshSection&  Sec = Res.Sections[SecIdx];
+	int32 Tri = Res.AreaWeightedSectionSamplers[SecIdx].GetEntryIndex(InstData->RandStream.GetFraction(), InstData->RandStream.GetFraction());
+	return Sec.FirstIndex + Tri * 3;
+}
+
+template<>
+FORCEINLINE int32 UNiagaraDataInterfaceStaticMesh::RandomTriIndex<TIntegralConstant<bool, true>, false>(FStaticMeshLODResources& Res, FNDIStaticMesh_InstanceData* InstData)
+{
+	int32 SecIdx = RandomSection<TIntegralConstant<bool, true>, false>(Res, InstData);
+	FStaticMeshSection&  Sec = Res.Sections[SecIdx];
+	int32 Tri = Res.AreaWeightedSectionSamplers[SecIdx].GetEntryIndex(InstData->RandStream.GetFraction(), InstData->RandStream.GetFraction());
+	return Sec.FirstIndex + Tri * 3;
+}
+
+template<>
+FORCEINLINE int32 UNiagaraDataInterfaceStaticMesh::RandomTriIndex<TIntegralConstant<bool, false>, true>(FStaticMeshLODResources& Res, FNDIStaticMesh_InstanceData* InstData)
+{
+	int32 SecIdx = RandomSection<TIntegralConstant<bool, false>, true>(Res, InstData);
+	FStaticMeshSection&  Sec = Res.Sections[SecIdx];
+	int32 Tri = InstData->RandStream.RandRange(0, Sec.NumTriangles - 1);
+	return Sec.FirstIndex + Tri * 3;	
+}
+
+template<>
+FORCEINLINE int32 UNiagaraDataInterfaceStaticMesh::RandomTriIndex<TIntegralConstant<bool, false>, false>(FStaticMeshLODResources& Res, FNDIStaticMesh_InstanceData* InstData)
+{
+	int32 SecIdx = RandomSection<TIntegralConstant<bool, false>, false>(Res, InstData);
+	FStaticMeshSection&  Sec = Res.Sections[SecIdx];
+	int32 Tri = InstData->RandStream.RandRange(0, Sec.NumTriangles - 1);
+	return Sec.FirstIndex + Tri * 3;
+}
+
+template<typename TAreaWeighted>
+void UNiagaraDataInterfaceStaticMesh::RandomTriCoord(FVectorVMContext& Context)
+{
+	FUserPtrHandler<FNDIStaticMesh_InstanceData> InstData(Context);
+
+	FRegisterHandler<int32> OutTri(Context);
+	FRegisterHandler<float> OutBaryX(Context);
+	FRegisterHandler<float> OutBaryY(Context);
+	FRegisterHandler<float> OutBaryZ(Context);
+
+	FStaticMeshLODResources& Res = InstData->Mesh->RenderData->LODResources[0];
+	FIndexArrayView Indices = Res.IndexBuffer.GetArrayView();
+	for (int32 i = 0; i < Context.NumInstances; ++i)
+	{
+		*OutTri.GetDest() = RandomTriIndex<TAreaWeighted, true>(Res, InstData);
+		FVector Bary = RandomBarycentricCoord(InstData->RandStream);
+		*OutBaryX.GetDest() = Bary.X;
+		*OutBaryY.GetDest() = Bary.Y;
+		*OutBaryZ.GetDest() = Bary.Z;
+		
+		OutTri.Advance();
+		OutBaryX.Advance();
+		OutBaryY.Advance();
+		OutBaryZ.Advance();
+	}
+}
+
+template<typename InputType0, typename InputType1>
+void UNiagaraDataInterfaceStaticMesh::RandomTriCoordVertexColorFiltered(FVectorVMContext& Context)
+{
+	InputType0 MinValue(Context);
+	InputType1 RangeValue(Context);
+	FUserPtrHandler<FNDIStaticMesh_InstanceData> InstData(Context);
+
+	FRegisterHandler<int32> OutTri(Context);
+	FRegisterHandler<float> OutBaryX(Context);
+	FRegisterHandler<float> OutBaryY(Context);
+	FRegisterHandler<float> OutBaryZ(Context);
+	
+	FDynamicVertexColorFilterData* VCFData = InstData->DynamicVertexColorSampler.Get();
+
+	FStaticMeshLODResources& Res = InstData->Mesh->RenderData->LODResources[0];
+	FIndexArrayView Indices = Res.IndexBuffer.GetArrayView();
+
+	for (int32 i = 0; i < Context.NumInstances; ++i)
+	{
+		uint32 StartIdx = (uint32)(MinValue.Get()*255.0f);
+		uint32 Range = (uint32)(RangeValue.Get()*255.0f + 0.5f);
+		uint32 EndIdx = StartIdx + Range;
+		// Iterate over the bucketed range and find the total number of triangles in the list.
+		uint32 NumTris = 0;
+
+		// Unfortunately, there's always the chance that the user gave us a range and value that don't have any vertex color matches.
+		// In this case (hopefully rare), we keep expanding the search space until we find a valid value.
+		while (NumTris == 0)
+		{
+			StartIdx = FMath::Clamp<uint32>(StartIdx, 0, (uint32)VCFData->VertexColorToTriangleStart.Num() - 1);
+			EndIdx = FMath::Clamp<uint32>(EndIdx, StartIdx, (uint32)VCFData->VertexColorToTriangleStart.Num() - 1);
+			NumTris = (EndIdx < (uint32)VCFData->VertexColorToTriangleStart.Num() - 1) ? (VCFData->VertexColorToTriangleStart[EndIdx + 1] - VCFData->VertexColorToTriangleStart[StartIdx]) :
+				(uint32)VCFData->TrianglesSortedByVertexColor.Num() - VCFData->VertexColorToTriangleStart[StartIdx];
+
+			if (NumTris == 0)
+			{
+				if (StartIdx > 0)
+				{
+					StartIdx -= 1;
+				}
+				Range += 1;
+				EndIdx = StartIdx + Range;
+			}
+		}
+
+		// Select a random triangle from the list.
+		uint32 RandomTri = InstData->RandStream.GetFraction()*NumTris;
+
+		// Now emit that triangle...
+		*OutTri.GetDest() = VCFData->TrianglesSortedByVertexColor[VCFData->VertexColorToTriangleStart[StartIdx] + RandomTri];
+
+		FVector Bary = RandomBarycentricCoord(InstData->RandStream);
+		*OutBaryX.GetDest() = Bary.X;
+		*OutBaryY.GetDest() = Bary.Y;
+		*OutBaryZ.GetDest() = Bary.Z;
+
+		MinValue.Advance();
+		RangeValue.Advance();
+		OutTri.Advance();
+		OutBaryX.Advance();
+		OutBaryY.Advance();
+		OutBaryZ.Advance();
+	}
+}
+
+template<>
+FORCEINLINE int32 UNiagaraDataInterfaceStaticMesh::RandomTriIndexOnSection<TIntegralConstant<bool, true>>(FStaticMeshLODResources& Res, int32 SecIdx, FNDIStaticMesh_InstanceData* InstData)
+{
+	return Res.AreaWeightedSectionSamplers[SecIdx].GetEntryIndex(InstData->RandStream.GetFraction(), InstData->RandStream.GetFraction());
+}
+
+template<>
+FORCEINLINE int32 UNiagaraDataInterfaceStaticMesh::RandomTriIndexOnSection<TIntegralConstant<bool, false>>(FStaticMeshLODResources& Res, int32 SecIdx, FNDIStaticMesh_InstanceData* InstData)
+{
+	FStaticMeshSection&  Sec = Res.Sections[SecIdx];
+	int32 Tri = InstData->RandStream.RandRange(0, Sec.NumTriangles - 1);
+	return Sec.FirstIndex + Tri * 3;
+}
+
+template<typename TAreaWeighted, typename SectionIdxType>
+void UNiagaraDataInterfaceStaticMesh::RandomTriCoordOnSection(FVectorVMContext& Context)
+{
+	SectionIdxType SectionIdxParam(Context);
+	FUserPtrHandler<FNDIStaticMesh_InstanceData> InstData(Context);
+
+	FRegisterHandler<int32> OutTri(Context);
+	FRegisterHandler<float> OutBaryX(Context);
+	FRegisterHandler<float> OutBaryY(Context);
+	FRegisterHandler<float> OutBaryZ(Context);
+
+	FStaticMeshLODResources& Res = InstData->Mesh->RenderData->LODResources[0];
+	FIndexArrayView Indices = Res.IndexBuffer.GetArrayView();
+	for (int32 i = 0; i < Context.NumInstances; ++i)
+	{
+		int32 SecIdx = SectionIdxParam.Get();
+		*OutTri.GetDest() = RandomTriIndexOnSection<TAreaWeighted>(Res, SecIdx, InstData);
+		FVector Bary = RandomBarycentricCoord(InstData->RandStream);
+		*OutBaryX.GetDest() = Bary.X;
+		*OutBaryY.GetDest() = Bary.Y;
+		*OutBaryZ.GetDest() = Bary.Z;
+
+		SectionIdxParam.Advance();
+		OutTri.Advance();
+		OutBaryX.Advance();
+		OutBaryY.Advance();
+		OutBaryZ.Advance();
+	}
+}
+
+template<typename TransformHandlerType, typename TriType, typename BaryXType, typename BaryYType, typename BaryZType>
+void UNiagaraDataInterfaceStaticMesh::GetTriCoordPosition(FVectorVMContext& Context)
+{
+	TransformHandlerType TransformHandler;
+	TriType TriParam(Context);
+	BaryXType BaryXParam(Context);
+	BaryYType BaryYParam(Context);
+	BaryZType BaryZParam(Context);
+	FUserPtrHandler<FNDIStaticMesh_InstanceData> InstData(Context);
+
+	FRegisterHandler<float> OutPosX(Context);
+	FRegisterHandler<float> OutPosY(Context);
+	FRegisterHandler<float> OutPosZ(Context);
+
+	FStaticMeshLODResources& Res = InstData->Mesh->RenderData->LODResources[0];
+	const FIndexArrayView& Indices = Res.IndexBuffer.GetArrayView();
+	const FPositionVertexBuffer& Positions = Res.PositionVertexBuffer;
+
+	for (int32 i = 0; i < Context.NumInstances; ++i)
+	{
+		int32 Tri = TriParam.Get();
+		int32 Idx0 = Indices[Tri];
+		int32 Idx1 = Indices[Tri + 1];
+		int32 Idx2 = Indices[Tri + 2];
+
+		FVector Pos = BarycentricInterpolate(BaryXParam.Get(), BaryYParam.Get(), BaryZParam.Get(), Positions.VertexPosition(Idx0), Positions.VertexPosition(Idx1), Positions.VertexPosition(Idx2));
+		TransformHandler.Transform(Pos, InstData->Transform);
+
+		*OutPosX.GetDest() = Pos.X;
+		*OutPosY.GetDest() = Pos.Y;
+		*OutPosZ.GetDest() = Pos.Z;
+
+		TriParam.Advance();
+		BaryXParam.Advance();
+		BaryYParam.Advance();
+		BaryZParam.Advance();
+		OutPosX.Advance();
+		OutPosY.Advance();
+		OutPosZ.Advance();
+	}
+}
+
+template<typename TransformHandlerType, typename TriType, typename BaryXType, typename BaryYType, typename BaryZType>
+void UNiagaraDataInterfaceStaticMesh::GetTriCoordNormal(FVectorVMContext& Context)
+{
+	TransformHandlerType TransformHandler;
+
+	TriType TriParam(Context);
+	BaryXType BaryXParam(Context);
+	BaryYType BaryYParam(Context);
+	BaryZType BaryZParam(Context);
+	FUserPtrHandler<FNDIStaticMesh_InstanceData> InstData(Context);
+
+	FRegisterHandler<float> OutNormX(Context);
+	FRegisterHandler<float> OutNormY(Context);
+	FRegisterHandler<float> OutNormZ(Context);
+
+	FStaticMeshLODResources& Res = InstData->Mesh->RenderData->LODResources[0];
+	const FIndexArrayView& Indices = Res.IndexBuffer.GetArrayView();
+	const FStaticMeshVertexBuffer& Verts = Res.VertexBuffer;
+
+	for (int32 i = 0; i < Context.NumInstances; ++i)
+	{
+		int32 Tri = TriParam.Get();
+		int32 Idx0 = Indices[Tri];
+		int32 Idx1 = Indices[Tri + 1];
+		int32 Idx2 = Indices[Tri + 2];
+
+		FVector Norm = BarycentricInterpolate(BaryXParam.Get(), BaryYParam.Get(), BaryZParam.Get(), Verts.VertexTangentZ(Idx0), Verts.VertexTangentZ(Idx1), Verts.VertexTangentZ(Idx2));
+		TransformHandler.Transform(Norm, InstData->Transform);
+
+		*OutNormX.GetDest() = Norm.X;
+		*OutNormY.GetDest() = Norm.Y;
+		*OutNormZ.GetDest() = Norm.Z;
+		TriParam.Advance();
+		BaryXParam.Advance();
+		BaryYParam.Advance();
+		BaryZParam.Advance();
+		OutNormX.Advance();
+		OutNormY.Advance();
+		OutNormZ.Advance();
+	}
+}
+
+template<typename VertexAccessorType, typename TransformHandlerType, typename TriType, typename BaryXType, typename BaryYType, typename BaryZType>
+void UNiagaraDataInterfaceStaticMesh::GetTriCoordTangents(FVectorVMContext& Context)
+{
+	TransformHandlerType TransformHandler;	
+
+	TriType TriParam(Context);
+	BaryXType BaryXParam(Context);
+	BaryYType BaryYParam(Context);
+	BaryZType BaryZParam(Context);
+	FUserPtrHandler<FNDIStaticMesh_InstanceData> InstData(Context);
+
+	FStaticMeshLODResources& Res = InstData->Mesh->RenderData->LODResources[0];
+	const FIndexArrayView& Indices = Res.IndexBuffer.GetArrayView();
+	const VertexAccessorType Verts(Res.VertexBuffer);
+
+	FRegisterHandler<float> OutTangentX(Context);
+	FRegisterHandler<float> OutTangentY(Context);
+	FRegisterHandler<float> OutTangentZ(Context);
+	FRegisterHandler<float> OutBinormX(Context);
+	FRegisterHandler<float> OutBinormY(Context);
+	FRegisterHandler<float> OutBinormZ(Context);
+	FRegisterHandler<float> OutNormX(Context);
+	FRegisterHandler<float> OutNormY(Context);
+	FRegisterHandler<float> OutNormZ(Context);
+
+	for (int32 i = 0; i < Context.NumInstances; ++i)
+	{
+		int32 Tri = TriParam.Get();
+		int32 Idx0 = Indices[Tri];
+		int32 Idx1 = Indices[Tri + 1];
+		int32 Idx2 = Indices[Tri + 2];
+		FVector Tangent = BarycentricInterpolate(BaryXParam.Get(), BaryYParam.Get(), BaryZParam.Get(), Verts.GetTangentX(Idx0), Verts.GetTangentX(Idx1), Verts.GetTangentX(Idx2));
+		FVector Binorm = BarycentricInterpolate(BaryXParam.Get(), BaryYParam.Get(), BaryZParam.Get(), Verts.GetTangentY(Idx0), Verts.GetTangentY(Idx1), Verts.GetTangentY(Idx2));
+		FVector Norm = BarycentricInterpolate(BaryXParam.Get(), BaryYParam.Get(), BaryZParam.Get(), Verts.GetTangentZ(Idx0), Verts.GetTangentZ(Idx1), Verts.GetTangentZ(Idx2));
+		TransformHandler.Transform(Tangent, InstData->TransformInverseTransposed);
+		TransformHandler.Transform(Binorm, InstData->TransformInverseTransposed);
+		TransformHandler.Transform(Norm, InstData->TransformInverseTransposed);
+		*OutTangentX.GetDest() = Tangent.X;
+		*OutTangentY.GetDest() = Tangent.Y;
+		*OutTangentZ.GetDest() = Tangent.Z;
+		*OutBinormX.GetDest() = Binorm.X;
+		*OutBinormY.GetDest() = Binorm.Y;
+		*OutBinormZ.GetDest() = Binorm.Z;
+		*OutNormX.GetDest() = Norm.X;
+		*OutNormY.GetDest() = Norm.Y;
+		*OutNormZ.GetDest() = Norm.Z;
+
+		TriParam.Advance();
+		BaryXParam.Advance();
+		BaryYParam.Advance();
+		BaryZParam.Advance();
+		OutTangentX.Advance();
+		OutTangentY.Advance();
+		OutTangentZ.Advance();
+		OutBinormX.Advance();
+		OutBinormY.Advance();
+		OutBinormZ.Advance();
+		OutNormX.Advance();
+		OutNormY.Advance();
+		OutNormZ.Advance();
+	}
+}
+
+template<typename TriType, typename BaryXType, typename BaryYType, typename BaryZType>
+void UNiagaraDataInterfaceStaticMesh::GetTriCoordColor(FVectorVMContext& Context)
+{
+	TriType TriParam(Context);
+	BaryXType BaryXParam(Context);
+	BaryYType BaryYParam(Context);
+	BaryZType BaryZParam(Context);
+	FUserPtrHandler<FNDIStaticMesh_InstanceData> InstData(Context);
+
+	FRegisterHandler<float> OutColorR(Context);
+	FRegisterHandler<float> OutColorG(Context);
+	FRegisterHandler<float> OutColorB(Context);
+	FRegisterHandler<float> OutColorA(Context);
+
+	FStaticMeshLODResources& Res = InstData->Mesh->RenderData->LODResources[0];
+	const FIndexArrayView& Indices = Res.IndexBuffer.GetArrayView();
+	const FColorVertexBuffer& Colors = Res.ColorVertexBuffer;
+
+	for (int32 i = 0; i < Context.NumInstances; ++i)
+	{
+		int32 Tri = TriParam.Get();
+		int32 Idx0 = Indices[Tri];
+		int32 Idx1 = Indices[Tri + 1];
+		int32 Idx2 = Indices[Tri + 2];
+
+		FLinearColor Color = BarycentricInterpolate(BaryXParam.Get(), BaryYParam.Get(), BaryZParam.Get(),
+			Colors.VertexColor(Idx0).ReinterpretAsLinear(), Colors.VertexColor(Idx1).ReinterpretAsLinear(), Colors.VertexColor(Idx2).ReinterpretAsLinear());
+
+		*OutColorR.GetDest() = Color.R;
+		*OutColorG.GetDest() = Color.G;
+		*OutColorB.GetDest() = Color.B;
+		*OutColorA.GetDest() = Color.A;
+		TriParam.Advance();
+		BaryXParam.Advance();
+		BaryYParam.Advance();
+		BaryZParam.Advance();
+		OutColorR.Advance();
+		OutColorG.Advance();
+		OutColorB.Advance();
+		OutColorA.Advance();
+	}
+}
+
+template<typename VertexAccessorType, typename TriType, typename BaryXType, typename BaryYType, typename BaryZType, typename UVSetType>
+void UNiagaraDataInterfaceStaticMesh::GetTriCoordUV(FVectorVMContext& Context)
+{
+	TriType TriParam(Context);
+	BaryXType BaryXParam(Context);
+	BaryYType BaryYParam(Context);
+	BaryZType BaryZParam(Context);
+	UVSetType UVSetParam(Context);
+	FUserPtrHandler<FNDIStaticMesh_InstanceData> InstData(Context);
+
+	FRegisterHandler<float> OutU(Context);
+	FRegisterHandler<float> OutV(Context);
+
+	FStaticMeshLODResources& Res = InstData->Mesh->RenderData->LODResources[0];
+	const FIndexArrayView& Indices = Res.IndexBuffer.GetArrayView();
+	const VertexAccessorType Verts(Res.VertexBuffer);
+	for (int32 i = 0; i < Context.NumInstances; ++i)
+	{
+		int32 Tri = TriParam.Get();
+		int32 Idx0 = Indices[Tri];
+		int32 Idx1 = Indices[Tri + 1];
+		int32 Idx2 = Indices[Tri + 2];
+
+		int32 UVSet = UVSetParam.Get();
+		FVector2D UV = BarycentricInterpolate(BaryXParam.Get(), BaryYParam.Get(), BaryZParam.Get(),	Verts.GetUV(Idx0, UVSet), Verts.GetUV(Idx1, UVSet),	Verts.GetUV(Idx2, UVSet));
+
+		*OutU.GetDest() = UV.X;
+		*OutV.GetDest() = UV.Y;
+
+		TriParam.Advance();
+		BaryXParam.Advance();
+		BaryYParam.Advance();
+		BaryZParam.Advance();
+		UVSetParam.Advance();
+		OutU.Advance();
+		OutV.Advance();
+	}
+}
+
+template<typename TriType, typename BaryXType, typename BaryYType, typename BaryZType>
+void UNiagaraDataInterfaceStaticMesh::GetTriCoordPositionAndVelocity(FVectorVMContext& Context)
+{
+	TriType TriParam(Context);
+	BaryXType BaryXParam(Context);
+	BaryYType BaryYParam(Context);
+	BaryZType BaryZParam(Context);
+	FUserPtrHandler<FNDIStaticMesh_InstanceData> InstData(Context);
+
+	FRegisterHandler<float> OutPosX(Context);
+	FRegisterHandler<float> OutPosY(Context);
+	FRegisterHandler<float> OutPosZ(Context);
+	FRegisterHandler<float> OutVelX(Context);
+	FRegisterHandler<float> OutVelY(Context);
+	FRegisterHandler<float> OutVelZ(Context);
+
+	FStaticMeshLODResources& Res = InstData->Mesh->RenderData->LODResources[0];
+	const FIndexArrayView& Indices = Res.IndexBuffer.GetArrayView();
+	const FPositionVertexBuffer& Positions = Res.PositionVertexBuffer;
+
+	float InvDt = 1.0f / InstData->DeltaSeconds;
+	for (int32 i = 0; i < Context.NumInstances; ++i)
+	{
+		int32 Tri = TriParam.Get();
+		int32 Idx0 = Indices[Tri];
+		int32 Idx1 = Indices[Tri + 1];
+		int32 Idx2 = Indices[Tri + 2];
+
+		FVector Pos = BarycentricInterpolate(BaryXParam.Get(), BaryYParam.Get(), BaryZParam.Get(), Positions.VertexPosition(Idx0), Positions.VertexPosition(Idx1), Positions.VertexPosition(Idx2));
+
+		FVector PrevWSPos = InstData->PrevTransform.TransformPosition(Pos);
+		FVector WSPos = InstData->Transform.TransformPosition(Pos);
+
+		FVector Vel = (WSPos - PrevWSPos) * InvDt;
+		*OutPosX.GetDest() = WSPos.X;
+		*OutPosY.GetDest() = WSPos.Y;
+		*OutPosZ.GetDest() = WSPos.Z;
+		*OutVelX.GetDest() = Vel.X;
+		*OutVelY.GetDest() = Vel.Y;
+		*OutVelZ.GetDest() = Vel.Z;
+		TriParam.Advance();
+		BaryXParam.Advance();
+		BaryYParam.Advance();
+		BaryZParam.Advance();
+		OutPosX.Advance();
+		OutPosY.Advance();
+		OutPosZ.Advance();
+		OutVelX.Advance();
+		OutVelY.Advance();
+		OutVelZ.Advance();
+	}
+}
+
+void UNiagaraDataInterfaceStaticMesh::WriteTransform(const FMatrix& ToWrite, FVectorVMContext& Context)
+{
+	FRegisterHandler<float> Out00(Context);
+	FRegisterHandler<float> Out01(Context);
+	FRegisterHandler<float> Out02(Context);
+	FRegisterHandler<float> Out03(Context);
+	FRegisterHandler<float> Out04(Context);
+	FRegisterHandler<float> Out05(Context);
+	FRegisterHandler<float> Out06(Context);
+	FRegisterHandler<float> Out07(Context);
+	FRegisterHandler<float> Out08(Context);
+	FRegisterHandler<float> Out09(Context);
+	FRegisterHandler<float> Out10(Context);
+	FRegisterHandler<float> Out11(Context);
+	FRegisterHandler<float> Out12(Context);
+	FRegisterHandler<float> Out13(Context);
+	FRegisterHandler<float> Out14(Context);
+	FRegisterHandler<float> Out15(Context);
+
+	for (int32 i = 0; i < Context.NumInstances; ++i)
+	{
+		*Out00.GetDest() = ToWrite.M[0][0]; Out00.Advance();
+		*Out01.GetDest() = ToWrite.M[0][0]; Out01.Advance();
+		*Out02.GetDest() = ToWrite.M[0][0]; Out02.Advance();
+		*Out03.GetDest() = ToWrite.M[0][0]; Out03.Advance();
+		*Out04.GetDest() = ToWrite.M[0][0]; Out04.Advance();
+		*Out05.GetDest() = ToWrite.M[0][0]; Out05.Advance();
+		*Out06.GetDest() = ToWrite.M[0][0]; Out06.Advance();
+		*Out07.GetDest() = ToWrite.M[0][0]; Out07.Advance();
+		*Out08.GetDest() = ToWrite.M[0][0]; Out08.Advance();
+		*Out09.GetDest() = ToWrite.M[0][0]; Out09.Advance();
+		*Out10.GetDest() = ToWrite.M[0][0]; Out10.Advance();
+		*Out11.GetDest() = ToWrite.M[0][0]; Out11.Advance();
+		*Out12.GetDest() = ToWrite.M[0][0]; Out12.Advance();
+		*Out13.GetDest() = ToWrite.M[0][0]; Out13.Advance();
+		*Out14.GetDest() = ToWrite.M[0][0]; Out14.Advance();
+		*Out15.GetDest() = ToWrite.M[0][0]; Out15.Advance();
+	}
+}
+
+void UNiagaraDataInterfaceStaticMesh::GetLocalToWorld(FVectorVMContext& Context)
+{
+	FUserPtrHandler<FNDIStaticMesh_InstanceData> InstData(Context);
+	WriteTransform(InstData->Transform, Context);
+}
+
+void UNiagaraDataInterfaceStaticMesh::GetLocalToWorldInverseTransposed(FVectorVMContext& Context)
+{
+	FUserPtrHandler<FNDIStaticMesh_InstanceData> InstData(Context);
+	WriteTransform(InstData->TransformInverseTransposed, Context);
+}
+
+void UNiagaraDataInterfaceStaticMesh::GetWorldVelocity(FVectorVMContext& Context)
+{
+	FUserPtrHandler<FNDIStaticMesh_InstanceData> InstData(Context);
+
+	FRegisterHandler<float> OutVelX(Context);
+	FRegisterHandler<float> OutVelY(Context);
+	FRegisterHandler<float> OutVelZ(Context);
+
+	FVector Velocity(0.0f, 0.0f, 0.0f);
+	float InvDeltaTime = 1.0f / InstData->DeltaSeconds;
+	if (InstData->DeltaSeconds > 0.0f)
+	{
+		Velocity = (FVector(InstData->Transform.M[3][0], InstData->Transform.M[3][1], InstData->Transform.M[3][2]) - 
+			FVector(InstData->PrevTransform.M[3][0], InstData->PrevTransform.M[3][1], InstData->PrevTransform.M[3][2])) * InvDeltaTime;
+	}
+
+	for (int32 i = 0; i < Context.NumInstances; i++)
+	{
+		*OutVelX.GetDest() = Velocity.X;
+		*OutVelY.GetDest() = Velocity.Y;
+		*OutVelZ.GetDest() = Velocity.Z;
+		OutVelX.Advance();
+		OutVelY.Advance();
+		OutVelZ.Advance();
+	}
+}
+
+//////////////////////////////////////////////////////////////////////////
+
+bool FDynamicVertexColorFilterData::Init(FNDIStaticMesh_InstanceData* Owner)
+{
+	TrianglesSortedByVertexColor.Empty();
+	VertexColorToTriangleStart.AddDefaulted(256);
+	check(Owner->Mesh);
+
+	FStaticMeshLODResources& Res = Owner->Mesh->RenderData->LODResources[0];
+
+	if (Res.ColorVertexBuffer.GetNumVertices() == 0)
+	{
+		UE_LOG(LogNiagara, Log, TEXT("Cannot initialize vertex color filter data for a mesh with no color data - %s"), *Owner->Mesh->GetFullName());
+		return false;
+	}
+
+	// Go over all triangles for each possible vertex color and add it to that bucket
+	for (int32 i = 0; i < VertexColorToTriangleStart.Num(); i++)
+	{
+		uint32 MinVertexColorRed = i;
+		uint32 MaxVertexColorRed = i + 1;
+		VertexColorToTriangleStart[i] = TrianglesSortedByVertexColor.Num();
+
+		FIndexArrayView IndexView = Res.IndexBuffer.GetArrayView();
+		for (int32 j = 0; j < Owner->GetValidSections().Num(); j++)
+		{
+			int32 SectionIdx = Owner->GetValidSections()[j];
+			int32 TriStartIdx = Res.Sections[SectionIdx].FirstIndex;
+			for (uint32 TriIdx = 0; TriIdx < Res.Sections[SectionIdx].NumTriangles; TriIdx++)
+			{
+				uint32 V0Idx = IndexView[TriStartIdx + TriIdx * 3 + 0];
+				uint32 V1Idx = IndexView[TriStartIdx + TriIdx * 3 + 1];
+				uint32 V2Idx = IndexView[TriStartIdx + TriIdx * 3 + 2];
+
+				uint8 MaxR = FMath::Max<uint8>(Res.ColorVertexBuffer.VertexColor(V0Idx).R,
+					FMath::Max<uint8>(Res.ColorVertexBuffer.VertexColor(V1Idx).R,
+						Res.ColorVertexBuffer.VertexColor(V2Idx).R));
+				if (MaxR >= MinVertexColorRed && MaxR < MaxVertexColorRed)
+				{
+					TrianglesSortedByVertexColor.Add(TriStartIdx + TriIdx * 3);
+				}
+			}
+		}
+	}
+	return true;
+}
+
+TMap<uint32, TSharedPtr<FDynamicVertexColorFilterData>> FNDI_StaticMesh_GeneratedData::DynamicVertexColorFilters;
+FCriticalSection FNDI_StaticMesh_GeneratedData::CriticalSection;
+
+TSharedPtr<FDynamicVertexColorFilterData> FNDI_StaticMesh_GeneratedData::GetDynamicColorFilterData(FNDIStaticMesh_InstanceData* Instance)
+{
+	FScopeLock Lock(&CriticalSection);
+
+	check(Instance);
+	check(Instance->Mesh);
+
+	TSharedPtr<FDynamicVertexColorFilterData> Ret = nullptr;
+
+	uint32 FilterDataHash = GetTypeHash(Instance->Mesh);
+	for (int32 ValidSec : Instance->GetValidSections())
+	{
+		FilterDataHash = HashCombine(GetTypeHash(ValidSec), FilterDataHash);
+	}
+
+	if (TSharedPtr<FDynamicVertexColorFilterData>* Existing = DynamicVertexColorFilters.Find(FilterDataHash))
+	{
+		check(Existing->IsValid());//We shouldn't be able to have an invalid ptr here.
+		Ret = *Existing;		
+	}
+	else
+	{
+		Ret = MakeShared<FDynamicVertexColorFilterData>();
+		if (Ret->Init(Instance))
+		{
+			DynamicVertexColorFilters.Add(FilterDataHash) = Ret;
+		}
+		else
+		{
+			Ret = nullptr;
+		}
+	}
+
+	return Ret;
+}
+
+void FNDI_StaticMesh_GeneratedData::CleanupDynamicColorFilterData()
+{
+	TArray<uint32, TInlineAllocator<64>> ToRemove;
+	for (TPair<uint32, TSharedPtr<FDynamicVertexColorFilterData>>& Pair : DynamicVertexColorFilters)
+	{
+		TSharedPtr<FDynamicVertexColorFilterData>& Ptr = Pair.Value;
+		if (Ptr.IsUnique())
+		{
+			//If we're the only ref left then destroy this data
+			ToRemove.Add(Pair.Key);
+		}
+	}
+
+	for (uint32 Key : ToRemove)
+	{
+		DynamicVertexColorFilters.Remove(Key);
+	}
+}
+
+#undef LOCTEXT_NAMESPACE
+
diff --git a/Engine/Plugins/FX/Niagara/Source/Niagara/Private/NiagaraDataInterfaceVector2DCurve.cpp b/Engine/Plugins/FX/Niagara/Source/Niagara/Private/NiagaraDataInterfaceVector2DCurve.cpp
new file mode 100644
--- /dev/null
+++ b/Engine/Plugins/FX/Niagara/Source/Niagara/Private/NiagaraDataInterfaceVector2DCurve.cpp
@@ -0,0 +1,212 @@
+// Copyright 1998-2016 Epic Games, Inc. All Rights Reserved.
+
+#include "NiagaraDataInterfaceVector2DCurve.h"
+#include "Curves/CurveVector.h"
+#include "Curves/CurveLinearColor.h"
+#include "Curves/CurveFloat.h"
+#include "NiagaraTypes.h"
+
+
+//////////////////////////////////////////////////////////////////////////
+//Vector2D Curve
+
+UNiagaraDataInterfaceVector2DCurve::UNiagaraDataInterfaceVector2DCurve(FObjectInitializer const& ObjectInitializer)
+	: Super(ObjectInitializer)
+{
+	UpdateLUT();
+}
+void UNiagaraDataInterfaceVector2DCurve::PostInitProperties()
+{
+	Super::PostInitProperties();
+
+	if (HasAnyFlags(RF_ClassDefaultObject))
+	{
+		FNiagaraTypeRegistry::Register(FNiagaraTypeDefinition(GetClass()), true, false, false);
+	}
+
+	UpdateLUT();
+}
+
+void UNiagaraDataInterfaceVector2DCurve::PostLoad()
+{
+	Super::PostLoad();
+	UpdateLUT();
+}
+
+void UNiagaraDataInterfaceVector2DCurve::UpdateLUT()
+{
+	ShaderLUT.Empty();
+	for (uint32 i = 0; i < CurveLUTWidth; i++)
+	{
+		float X = i / (float)CurveLUTWidth;
+		FVector2D C(XCurve.Eval(X), YCurve.Eval(X));
+		ShaderLUT.Add(C.X);
+		ShaderLUT.Add(C.Y);
+	}
+	GPUBufferDirty = true;
+}
+
+#if WITH_EDITOR
+void UNiagaraDataInterfaceVector2DCurve::PostEditChangeProperty(struct FPropertyChangedEvent& PropertyChangedEvent)
+{
+	Super::PostEditChangeProperty(PropertyChangedEvent);
+
+	if (PropertyChangedEvent.Property && PropertyChangedEvent.Property->GetFName() == GET_MEMBER_NAME_CHECKED(UNiagaraDataInterfaceVector2DCurve, CurveToCopy))
+	{
+		UCurveVector* VectorCurveAsset = Cast<UCurveVector>(CurveToCopy.TryLoad());
+		if (VectorCurveAsset)
+		{
+			Modify();
+			XCurve = VectorCurveAsset->FloatCurves[0];
+			YCurve = VectorCurveAsset->FloatCurves[1];
+		}
+		UpdateLUT();
+	}
+
+}
+#endif
+
+bool UNiagaraDataInterfaceVector2DCurve::CopyTo(UNiagaraDataInterface* Destination) const
+{
+	if (!Super::CopyTo(Destination))
+	{
+		return false;
+	}
+	UNiagaraDataInterfaceVector2DCurve* DestinationVector2DCurve = CastChecked<UNiagaraDataInterfaceVector2DCurve>(Destination);
+	DestinationVector2DCurve->XCurve = XCurve;
+	DestinationVector2DCurve->YCurve = YCurve;
+	DestinationVector2DCurve->UpdateLUT();
+
+	return true;
+}
+
+bool UNiagaraDataInterfaceVector2DCurve::Equals(const UNiagaraDataInterface* Other) const
+{
+	if (!Super::Equals(Other))
+	{
+		return false;
+	}
+	const UNiagaraDataInterfaceVector2DCurve* OtherVector2DCurve = CastChecked<const UNiagaraDataInterfaceVector2DCurve>(Other);
+	return OtherVector2DCurve->XCurve == XCurve &&
+		OtherVector2DCurve->YCurve == YCurve;
+}
+void UNiagaraDataInterfaceVector2DCurve::GetCurveData(TArray<FCurveData>& OutCurveData)
+{
+	OutCurveData.Add(FCurveData(&XCurve, TEXT("X"), FLinearColor::Red));
+	OutCurveData.Add(FCurveData(&YCurve, TEXT("Y"), FLinearColor::Green));
+}
+
+void UNiagaraDataInterfaceVector2DCurve::GetFunctions(TArray<FNiagaraFunctionSignature>& OutFunctions)
+{
+	FNiagaraFunctionSignature Sig;
+	Sig.Name = TEXT("SampleVector2DCurve");
+	Sig.bMemberFunction = true;
+	Sig.bRequiresContext = false;
+	Sig.Inputs.Add(FNiagaraVariable(FNiagaraTypeDefinition(GetClass()), TEXT("Vector2DCurve")));
+	Sig.Inputs.Add(FNiagaraVariable(FNiagaraTypeDefinition::GetFloatDef(), TEXT("X")));
+	Sig.Outputs.Add(FNiagaraVariable(FNiagaraTypeDefinition::GetVec2Def(), TEXT("Value")));
+	//Sig.Owner = *GetName();
+
+	OutFunctions.Add(Sig);
+}
+
+
+
+// build the shader function HLSL; function name is passed in, as it's defined per-DI; that way, configuration could change
+// the HLSL in the spirit of a static switch
+// TODO: need a way to identify each specific function here
+// 
+bool UNiagaraDataInterfaceVector2DCurve::GetFunctionHLSL(FString FunctionName, TArray<DIGPUBufferParamDescriptor> &Descriptors, FString &HLSLInterfaceID, FString &OutHLSL)
+{
+	FString BufferName = Descriptors[0].BufferParamName;
+	OutHLSL += TEXT("void ") + FunctionName + TEXT("(in float In_X, out float2 Out_Value) \n{\n");
+	OutHLSL += TEXT("\t Out_Value.x = ") + BufferName + TEXT("[(int)(In_X *") + FString::FromInt(CurveLUTWidth) + TEXT(")* 2 ];");
+	OutHLSL += TEXT("\t Out_Value.y = ") + BufferName + TEXT("[1+ (int)(In_X *") + FString::FromInt(CurveLUTWidth) + TEXT(")* 2 ];");
+	OutHLSL += TEXT("\n}\n");
+	return true;
+}
+
+// build buffer definition hlsl
+// 1. Choose a buffer name, add the data interface ID (important!)
+// 2. add a DIGPUBufferParamDescriptor to the array argument; that'll be passed on to the FNiagaraShader for binding to a shader param, that can
+// then later be found by name via FindDIBufferParam for setting; 
+// 3. store buffer declaration hlsl in OutHLSL
+// multiple buffers can be defined at once here
+//
+void UNiagaraDataInterfaceVector2DCurve::GetBufferDefinitionHLSL(FString DataInterfaceID, TArray<DIGPUBufferParamDescriptor> &BufferDescriptors, FString &OutHLSL)
+{
+	FString BufferName = "CurveLUT" + DataInterfaceID;
+	OutHLSL += TEXT("Buffer<float> ") + BufferName + TEXT(";\n");
+
+	BufferDescriptors.Add(DIGPUBufferParamDescriptor(BufferName, 0));		// add a descriptor for shader parameter binding
+
+}
+
+// called after translate, to setup buffers matching the buffer descriptors generated during hlsl translation
+// need to do this because the script used during translate is a clone, including its DIs
+//
+void UNiagaraDataInterfaceVector2DCurve::SetupBuffers(TArray<DIGPUBufferParamDescriptor> &BufferDescriptors)
+{
+	for (DIGPUBufferParamDescriptor &Desc : BufferDescriptors)
+	{
+		FNiagaraDataInterfaceBufferData BufferData(*Desc.BufferParamName);	// store off the data for later use
+		GPUBuffers.Add(BufferData);
+	}
+}
+// return the GPU buffer array (called from NiagaraInstanceBatcher to get the buffers for setting to the shader)
+// we lazily update the buffer with a new LUT here if necessary
+//
+TArray<FNiagaraDataInterfaceBufferData> &UNiagaraDataInterfaceVector2DCurve::GetBufferDataArray()
+{
+	check(IsInRenderingThread());
+	if (GPUBufferDirty)
+	{
+		check(GPUBuffers.Num() > 0);
+
+		FNiagaraDataInterfaceBufferData &GPUBuffer = GPUBuffers[0];
+		GPUBuffer.Buffer.Release();
+		GPUBuffer.Buffer.Initialize(sizeof(float), CurveLUTWidth * 2, EPixelFormat::PF_R32_FLOAT);	// always allocate for up to 64 data sets
+		uint32 BufferSize = ShaderLUT.Num() * sizeof(float);
+		int32 *BufferData = static_cast<int32*>(RHILockVertexBuffer(GPUBuffer.Buffer.Buffer, 0, BufferSize, EResourceLockMode::RLM_WriteOnly));
+		FPlatformMemory::Memcpy(BufferData, ShaderLUT.GetData(), BufferSize);
+		RHIUnlockVertexBuffer(GPUBuffer.Buffer.Buffer);
+		GPUBufferDirty = false;
+	}
+
+	return GPUBuffers;
+}
+
+DEFINE_NDI_FUNC_BINDER(UNiagaraDataInterfaceVector2DCurve, SampleCurve);
+FVMExternalFunction UNiagaraDataInterfaceVector2DCurve::GetVMExternalFunction(const FVMExternalFunctionBindingInfo& BindingInfo, void* InstanceData)
+{
+	if (BindingInfo.Name == TEXT("SampleVector2DCurve") && BindingInfo.GetNumInputs() == 1 && BindingInfo.GetNumOutputs() == 2)
+	{
+		return TNDIParamBinder<0, float, NDI_FUNC_BINDER(UNiagaraDataInterfaceVector2DCurve, SampleCurve)>::Bind(this, BindingInfo, InstanceData);
+	}
+	else
+	{
+		UE_LOG(LogNiagara, Error, TEXT("Could not find data interface external function.\n\tExpected Name: SampleVector2DCurve  Actual Name: %s\n\tExpected Inputs: 1  Actual Inputs: %i\n\tExpected Outputs: 2  Actual Outputs: %i"),
+			*BindingInfo.Name.ToString(), BindingInfo.GetNumInputs(), BindingInfo.GetNumOutputs());
+		return FVMExternalFunction();
+	}
+}
+
+template<typename XParamType>
+void UNiagaraDataInterfaceVector2DCurve::SampleCurve(FVectorVMContext& Context)
+{
+	//TODO: Create some SIMDable optimized representation of the curve to do this faster.
+	XParamType XParam(Context);
+	FRegisterHandler<float> OutSampleX(Context);
+	FRegisterHandler<float> OutSampleY(Context);
+
+	for (int32 i = 0; i < Context.NumInstances; ++i)
+	{
+		float X = XParam.Get();
+		FVector2D V(XCurve.Eval(X), YCurve.Eval(X));
+		*OutSampleX.GetDest() = V.X;
+		*OutSampleY.GetDest() = V.Y;
+		XParam.Advance();
+		OutSampleX.Advance();
+		OutSampleY.Advance();
+	}
+}
\ No newline at end of file
diff --git a/Engine/Plugins/FX/Niagara/Source/Niagara/Private/NiagaraDataInterfaceVector4Curve.cpp b/Engine/Plugins/FX/Niagara/Source/Niagara/Private/NiagaraDataInterfaceVector4Curve.cpp
new file mode 100644
--- /dev/null
+++ b/Engine/Plugins/FX/Niagara/Source/Niagara/Private/NiagaraDataInterfaceVector4Curve.cpp
@@ -0,0 +1,235 @@
+// Copyright 1998-2016 Epic Games, Inc. All Rights Reserved.
+
+#include "NiagaraDataInterfaceVector4Curve.h"
+#include "Curves/CurveVector.h"
+#include "Curves/CurveLinearColor.h"
+#include "Curves/CurveFloat.h"
+#include "NiagaraTypes.h"
+
+//////////////////////////////////////////////////////////////////////////
+//Color Curve
+
+UNiagaraDataInterfaceVector4Curve::UNiagaraDataInterfaceVector4Curve(FObjectInitializer const& ObjectInitializer)
+	: Super(ObjectInitializer)
+{
+	UpdateLUT();
+}
+
+void UNiagaraDataInterfaceVector4Curve::PostInitProperties()
+{
+	Super::PostInitProperties();
+
+	//Can we regitser data interfaces as regular types and fold them into the FNiagaraVariable framework for UI and function calls etc?
+	if (HasAnyFlags(RF_ClassDefaultObject))
+	{
+		FNiagaraTypeRegistry::Register(FNiagaraTypeDefinition(GetClass()), true, false, false);
+	}
+
+	UpdateLUT();
+}
+
+void UNiagaraDataInterfaceVector4Curve::PostLoad()
+{
+	Super::PostLoad();
+	UpdateLUT();
+}
+
+#if WITH_EDITOR
+
+void UNiagaraDataInterfaceVector4Curve::PostEditChangeProperty(struct FPropertyChangedEvent& PropertyChangedEvent)
+{
+	Super::PostEditChangeProperty(PropertyChangedEvent);
+
+	if (PropertyChangedEvent.Property && PropertyChangedEvent.Property->GetFName() == GET_MEMBER_NAME_CHECKED(UNiagaraDataInterfaceVector4Curve, CurveToCopy))
+	{
+		UCurveLinearColor* ColorCurveAsset = Cast<UCurveLinearColor>(CurveToCopy.TryLoad());
+		if(ColorCurveAsset)
+		{
+			Modify();
+			XCurve = ColorCurveAsset->FloatCurves[0];
+			YCurve = ColorCurveAsset->FloatCurves[1];
+			ZCurve = ColorCurveAsset->FloatCurves[2];
+			WCurve = ColorCurveAsset->FloatCurves[3];
+		}
+		UpdateLUT();
+	}
+}
+
+#endif
+
+void UNiagaraDataInterfaceVector4Curve::UpdateLUT()
+{
+	ShaderLUT.Empty();
+	for (uint32 i = 0; i < CurveLUTWidth; i++)
+	{
+		float X = i / (float)CurveLUTWidth;
+		FLinearColor C(XCurve.Eval(X), YCurve.Eval(X), ZCurve.Eval(X), WCurve.Eval(X));
+		ShaderLUT.Add(C.R);
+		ShaderLUT.Add(C.G);
+		ShaderLUT.Add(C.B);
+		ShaderLUT.Add(C.A);
+	}
+	GPUBufferDirty = true;
+}
+
+bool UNiagaraDataInterfaceVector4Curve::CopyTo(UNiagaraDataInterface* Destination) const 
+{
+	if (!Super::CopyTo(Destination))
+	{
+		return false;
+	}
+	UNiagaraDataInterfaceVector4Curve* DestinationColorCurve = CastChecked<UNiagaraDataInterfaceVector4Curve>(Destination);
+	DestinationColorCurve->XCurve = XCurve;
+	DestinationColorCurve->YCurve = YCurve;
+	DestinationColorCurve->ZCurve = ZCurve;
+	DestinationColorCurve->WCurve = WCurve;
+	DestinationColorCurve->UpdateLUT();
+
+	return true;
+}
+
+bool UNiagaraDataInterfaceVector4Curve::Equals(const UNiagaraDataInterface* Other) const
+{
+	if (!Super::Equals(Other))
+	{
+		return false;
+	}
+	const UNiagaraDataInterfaceVector4Curve* OtherCurve = CastChecked<const UNiagaraDataInterfaceVector4Curve>(Other);
+	return OtherCurve->XCurve == XCurve &&
+		OtherCurve->YCurve == YCurve &&
+		OtherCurve->ZCurve == ZCurve &&
+		OtherCurve->WCurve == WCurve;
+}
+
+void UNiagaraDataInterfaceVector4Curve::GetCurveData(TArray<FCurveData>& OutCurveData)
+{
+	OutCurveData.Add(FCurveData(&XCurve, TEXT("X"), FLinearColor::Red));
+	OutCurveData.Add(FCurveData(&YCurve, TEXT("Y"), FLinearColor::Green));
+	OutCurveData.Add(FCurveData(&ZCurve, TEXT("Z"), FLinearColor::Blue));
+	OutCurveData.Add(FCurveData(&WCurve, TEXT("W"), FLinearColor::White));
+}
+
+void UNiagaraDataInterfaceVector4Curve::GetFunctions(TArray<FNiagaraFunctionSignature>& OutFunctions)
+{
+	FNiagaraFunctionSignature Sig;
+	Sig.Name = TEXT("SampleColorCurve");
+	Sig.bMemberFunction = true;
+	Sig.bRequiresContext = false;
+	Sig.Inputs.Add(FNiagaraVariable(FNiagaraTypeDefinition(GetClass()), TEXT("Vector4Curve")));
+	Sig.Inputs.Add(FNiagaraVariable(FNiagaraTypeDefinition::GetFloatDef(), TEXT("X")));
+	Sig.Outputs.Add(FNiagaraVariable(FNiagaraTypeDefinition::GetColorDef(), TEXT("Value")));
+	//Sig.Owner = *GetName();
+
+	OutFunctions.Add(Sig);
+}
+
+// build the shader function HLSL; function name is passed in, as it's defined per-DI; that way, configuration could change
+// the HLSL in the spirit of a static switch
+// TODO: need a way to identify each specific function here
+// 
+bool UNiagaraDataInterfaceVector4Curve::GetFunctionHLSL(FString FunctionName, TArray<DIGPUBufferParamDescriptor> &Descriptors, FString &HLSLInterfaceID, FString &OutHLSL)
+{
+	FString BufferName = Descriptors[0].BufferParamName;
+	OutHLSL += TEXT("void ") + FunctionName + TEXT("(in float In_X, out float4 Out_Value) \n{\n");
+	OutHLSL += TEXT("\t Out_Value.x = ") + BufferName + TEXT("[(int)(In_X *") + FString::FromInt(CurveLUTWidth) + TEXT(")* 4 ];");
+	OutHLSL += TEXT("\t Out_Value.y = ") + BufferName + TEXT("[1+ (int)(In_X *") + FString::FromInt(CurveLUTWidth) + TEXT(")* 4 ];");
+	OutHLSL += TEXT("\t Out_Value.z = ") + BufferName + TEXT("[2+ (int)(In_X *") + FString::FromInt(CurveLUTWidth) + TEXT(")* 4 ];");
+	OutHLSL += TEXT("\t Out_Value.w = ") + BufferName + TEXT("[3+ (int)(In_X *") + FString::FromInt(CurveLUTWidth) + TEXT(")* 4 ];");
+	OutHLSL += TEXT("\n}\n");
+	return true;
+}
+
+// build buffer definition hlsl
+// 1. Choose a buffer name, add the data interface ID (important!)
+// 2. add a DIGPUBufferParamDescriptor to the array argument; that'll be passed on to the FNiagaraShader for binding to a shader param, that can
+// then later be found by name via FindDIBufferParam for setting; 
+// 3. store buffer declaration hlsl in OutHLSL
+// multiple buffers can be defined at once here
+//
+void UNiagaraDataInterfaceVector4Curve::GetBufferDefinitionHLSL(FString DataInterfaceID, TArray<DIGPUBufferParamDescriptor> &BufferDescriptors, FString &OutHLSL)
+{
+	FString BufferName = "CurveLUT" + DataInterfaceID;
+	OutHLSL += TEXT("Buffer<float> ") + BufferName + TEXT(";\n");
+
+	BufferDescriptors.Add( DIGPUBufferParamDescriptor(BufferName, 0) );		// add a descriptor for shader parameter binding
+
+}
+
+// called after translate, to setup buffers matching the buffer descriptors generated during hlsl translation
+// need to do this because the script used during translate is a clone, including its DIs
+//
+void UNiagaraDataInterfaceVector4Curve::SetupBuffers(TArray<DIGPUBufferParamDescriptor> &BufferDescriptors)
+{
+	for (DIGPUBufferParamDescriptor &Desc : BufferDescriptors)
+	{
+		FNiagaraDataInterfaceBufferData BufferData(*Desc.BufferParamName);	// store off the data for later use
+		GPUBuffers.Add(BufferData);
+	}
+}
+
+// return the GPU buffer array (called from NiagaraInstanceBatcher to get the buffers for setting to the shader)
+// we lazily update the buffer with a new LUT here if necessary
+//
+TArray<FNiagaraDataInterfaceBufferData> &UNiagaraDataInterfaceVector4Curve::GetBufferDataArray()
+{
+	check(IsInRenderingThread());
+	if (GPUBufferDirty)
+	{
+		check(GPUBuffers.Num() > 0);
+
+		FNiagaraDataInterfaceBufferData &GPUBuffer = GPUBuffers[0];
+		GPUBuffer.Buffer.Release();
+		GPUBuffer.Buffer.Initialize(sizeof(float), CurveLUTWidth*4 , EPixelFormat::PF_R32_FLOAT);	// always allocate for up to 64 data sets
+		uint32 BufferSize = ShaderLUT.Num() * sizeof(float);
+		int32 *BufferData = static_cast<int32*>(RHILockVertexBuffer(GPUBuffer.Buffer.Buffer, 0, BufferSize, EResourceLockMode::RLM_WriteOnly));
+		FPlatformMemory::Memcpy(BufferData, ShaderLUT.GetData(), BufferSize);
+		RHIUnlockVertexBuffer(GPUBuffer.Buffer.Buffer);
+		GPUBufferDirty = false;
+	}
+
+	return GPUBuffers;
+}
+
+
+
+
+DEFINE_NDI_FUNC_BINDER(UNiagaraDataInterfaceVector4Curve, SampleCurve);
+FVMExternalFunction UNiagaraDataInterfaceVector4Curve::GetVMExternalFunction(const FVMExternalFunctionBindingInfo& BindingInfo, void* InstanceData)
+{
+	if (BindingInfo.Name == TEXT("SampleColorCurve") && BindingInfo.GetNumInputs() == 1 && BindingInfo.GetNumOutputs() == 4)
+	{
+		return TNDIParamBinder<0, float, NDI_FUNC_BINDER(UNiagaraDataInterfaceVector4Curve, SampleCurve)>::Bind(this, BindingInfo, InstanceData);
+	}
+	else
+	{
+		UE_LOG(LogNiagara, Error, TEXT("Could not find data interface external function.\n\tExpected Name: SampleColorCurve  Actual Name: %s\n\tExpected Inputs: 1  Actual Inputs: %i\n\tExpected Outputs: 4  Actual Outputs: %i"),
+			*BindingInfo.Name.ToString(), BindingInfo.GetNumInputs(), BindingInfo.GetNumOutputs());
+		return FVMExternalFunction();
+	}
+}
+
+template<typename XParamType>
+void UNiagaraDataInterfaceVector4Curve::SampleCurve(FVectorVMContext& Context)
+{
+	//TODO: Create some SIMDable optimized representation of the curve to do this faster.
+	XParamType XParam(Context);
+	FRegisterHandler<float> SamplePtrR(Context);
+	FRegisterHandler<float> SamplePtrG(Context);
+	FRegisterHandler<float> SamplePtrB(Context);
+	FRegisterHandler<float> SamplePtrA(Context);
+
+	for (int32 i = 0; i < Context.NumInstances; ++i)
+	{
+		float X = XParam.Get();
+		FLinearColor C(XCurve.Eval(X), YCurve.Eval(X), ZCurve.Eval(X), WCurve.Eval(X));
+		*SamplePtrR.GetDest() = C.R;
+		*SamplePtrG.GetDest() = C.G;
+		*SamplePtrB.GetDest() = C.B;
+		*SamplePtrA.GetDest() = C.A;
+		XParam.Advance();
+		SamplePtrR.Advance();
+		SamplePtrG.Advance();
+		SamplePtrB.Advance();
+		SamplePtrA.Advance();
+	}
+}
diff --git a/Engine/Plugins/FX/Niagara/Source/Niagara/Private/NiagaraDataInterfaceVectorCurve.cpp b/Engine/Plugins/FX/Niagara/Source/Niagara/Private/NiagaraDataInterfaceVectorCurve.cpp
new file mode 100644
--- /dev/null
+++ b/Engine/Plugins/FX/Niagara/Source/Niagara/Private/NiagaraDataInterfaceVectorCurve.cpp
@@ -0,0 +1,228 @@
+// Copyright 1998-2016 Epic Games, Inc. All Rights Reserved.
+
+#include "NiagaraDataInterfaceVectorCurve.h"
+#include "Curves/CurveVector.h"
+#include "Curves/CurveLinearColor.h"
+#include "Curves/CurveFloat.h"
+#include "NiagaraTypes.h"
+
+//////////////////////////////////////////////////////////////////////////
+//Vector Curve
+
+UNiagaraDataInterfaceVectorCurve::UNiagaraDataInterfaceVectorCurve(FObjectInitializer const& ObjectInitializer)
+	: Super(ObjectInitializer)
+{
+	UpdateLUT();
+}
+
+void UNiagaraDataInterfaceVectorCurve::PostInitProperties()
+{
+	Super::PostInitProperties();
+
+	//Can we regitser data interfaces as regular types and fold them into the FNiagaraVariable framework for UI and function calls etc?
+	if (HasAnyFlags(RF_ClassDefaultObject))
+	{
+		FNiagaraTypeRegistry::Register(FNiagaraTypeDefinition(GetClass()), true, false, false);
+	}
+
+	UpdateLUT();
+}
+
+void UNiagaraDataInterfaceVectorCurve::PostLoad()
+{
+	Super::PostLoad();
+	UpdateLUT();
+}
+
+#if WITH_EDITOR
+
+void UNiagaraDataInterfaceVectorCurve::PostEditChangeProperty(struct FPropertyChangedEvent& PropertyChangedEvent)
+{
+	Super::PostEditChangeProperty(PropertyChangedEvent);
+
+	if (PropertyChangedEvent.Property && PropertyChangedEvent.Property->GetFName() == GET_MEMBER_NAME_CHECKED(UNiagaraDataInterfaceVectorCurve, CurveToCopy))
+	{
+		UCurveVector* VectorCurveAsset = Cast<UCurveVector>(CurveToCopy.TryLoad());
+		if (VectorCurveAsset)
+		{
+			Modify();
+			XCurve = VectorCurveAsset->FloatCurves[0];
+			YCurve = VectorCurveAsset->FloatCurves[1];
+			ZCurve = VectorCurveAsset->FloatCurves[2];
+		}
+		UpdateLUT();
+	}
+}
+
+#endif
+
+
+void UNiagaraDataInterfaceVectorCurve::UpdateLUT()
+{
+	ShaderLUT.Empty();
+	for (uint32 i = 0; i < CurveLUTWidth; i++)
+	{
+		float X = i / (float)CurveLUTWidth;
+		FVector C(XCurve.Eval(X), YCurve.Eval(X), ZCurve.Eval(X));
+		ShaderLUT.Add(C.X);
+		ShaderLUT.Add(C.Y);
+		ShaderLUT.Add(C.Z);
+		ShaderLUT.Add(1.0f);
+	}
+	GPUBufferDirty = true;
+}
+
+bool UNiagaraDataInterfaceVectorCurve::CopyTo(UNiagaraDataInterface* Destination) const 
+{
+	if (!Super::CopyTo(Destination))
+	{
+		return false;
+	}
+	UNiagaraDataInterfaceVectorCurve* DestinationVectorCurve = CastChecked<UNiagaraDataInterfaceVectorCurve>(Destination);
+	DestinationVectorCurve->XCurve = XCurve;
+	DestinationVectorCurve->YCurve = YCurve;
+	DestinationVectorCurve->ZCurve = ZCurve;
+	DestinationVectorCurve->UpdateLUT();
+
+	return true;
+}
+
+bool UNiagaraDataInterfaceVectorCurve::Equals(const UNiagaraDataInterface* Other) const
+{
+	if (!Super::Equals(Other))
+	{
+		return false;
+	}
+	const UNiagaraDataInterfaceVectorCurve* OtherVectorCurve = CastChecked<const UNiagaraDataInterfaceVectorCurve>(Other);
+	return OtherVectorCurve->XCurve == XCurve &&
+		OtherVectorCurve->YCurve == YCurve &&
+		OtherVectorCurve->ZCurve == ZCurve;
+}
+
+void UNiagaraDataInterfaceVectorCurve::GetCurveData(TArray<FCurveData>& OutCurveData)
+{
+	OutCurveData.Add(FCurveData(&XCurve, TEXT("X"), FLinearColor::Red));
+	OutCurveData.Add(FCurveData(&YCurve, TEXT("Y"), FLinearColor::Green));
+	OutCurveData.Add(FCurveData(&ZCurve, TEXT("Z"), FLinearColor::Blue));
+}
+
+void UNiagaraDataInterfaceVectorCurve::GetFunctions(TArray<FNiagaraFunctionSignature>& OutFunctions)
+{
+	FNiagaraFunctionSignature Sig;
+	Sig.Name = TEXT("SampleVectorCurve");
+	Sig.bMemberFunction = true;
+	Sig.bRequiresContext = false;
+	Sig.Inputs.Add(FNiagaraVariable(FNiagaraTypeDefinition(GetClass()), TEXT("VectorCurve")));
+	Sig.Inputs.Add(FNiagaraVariable(FNiagaraTypeDefinition::GetFloatDef(), TEXT("X")));
+	Sig.Outputs.Add(FNiagaraVariable(FNiagaraTypeDefinition::GetVec3Def(), TEXT("Value")));
+	//Sig.Owner = *GetName();
+
+	OutFunctions.Add(Sig);
+}
+
+
+
+// build the shader function HLSL; function name is passed in, as it's defined per-DI; that way, configuration could change
+// the HLSL in the spirit of a static switch
+// TODO: need a way to identify each specific function here
+// 
+bool UNiagaraDataInterfaceVectorCurve::GetFunctionHLSL(FString FunctionName, TArray<DIGPUBufferParamDescriptor> &Descriptors, FString &HLSLInterfaceID, FString &OutHLSL)
+{
+	FString BufferName = Descriptors[0].BufferParamName;
+	OutHLSL += TEXT("void ") + FunctionName + TEXT("(in float In_X, out float3 Out_Value) \n{\n");
+	OutHLSL += TEXT("\t Out_Value.x = ") + BufferName + TEXT("[(int)(In_X *") + FString::FromInt(CurveLUTWidth) + TEXT(")* 3 ];");
+	OutHLSL += TEXT("\t Out_Value.y = ") + BufferName + TEXT("[1+ (int)(In_X *") + FString::FromInt(CurveLUTWidth) + TEXT(")* 3 ];");
+	OutHLSL += TEXT("\t Out_Value.z = ") + BufferName + TEXT("[2+ (int)(In_X *") + FString::FromInt(CurveLUTWidth) + TEXT(")* 3 ];");
+	OutHLSL += TEXT("\n}\n");
+	return true;
+}
+
+// build buffer definition hlsl
+// 1. Choose a buffer name, add the data interface ID (important!)
+// 2. add a DIGPUBufferParamDescriptor to the array argument; that'll be passed on to the FNiagaraShader for binding to a shader param, that can
+// then later be found by name via FindDIBufferParam for setting; 
+// 3. store buffer declaration hlsl in OutHLSL
+// multiple buffers can be defined at once here
+//
+void UNiagaraDataInterfaceVectorCurve::GetBufferDefinitionHLSL(FString DataInterfaceID, TArray<DIGPUBufferParamDescriptor> &BufferDescriptors, FString &OutHLSL)
+{
+	FString BufferName = "CurveLUT" + DataInterfaceID;
+	OutHLSL += TEXT("Buffer<float> ") + BufferName + TEXT(";\n");
+
+	BufferDescriptors.Add(DIGPUBufferParamDescriptor(BufferName, 0));		// add a descriptor for shader parameter binding
+
+}
+
+// called after translate, to setup buffers matching the buffer descriptors generated during hlsl translation
+// need to do this because the script used during translate is a clone, including its DIs
+//
+void UNiagaraDataInterfaceVectorCurve::SetupBuffers(TArray<DIGPUBufferParamDescriptor> &BufferDescriptors)
+{
+	for (DIGPUBufferParamDescriptor &Desc : BufferDescriptors)
+	{
+		FNiagaraDataInterfaceBufferData BufferData(*Desc.BufferParamName);	// store off the data for later use
+		GPUBuffers.Add(BufferData);
+	}
+}
+
+// return the GPU buffer array (called from NiagaraInstanceBatcher to get the buffers for setting to the shader)
+// we lazily update the buffer with a new LUT here if necessary
+//
+TArray<FNiagaraDataInterfaceBufferData> &UNiagaraDataInterfaceVectorCurve::GetBufferDataArray()
+{
+	check(IsInRenderingThread());
+	if (GPUBufferDirty)
+	{
+		check(GPUBuffers.Num() > 0);
+
+		FNiagaraDataInterfaceBufferData &GPUBuffer = GPUBuffers[0];
+		GPUBuffer.Buffer.Release();
+		GPUBuffer.Buffer.Initialize(sizeof(float), CurveLUTWidth * 3, EPixelFormat::PF_R32_FLOAT);	// always allocate for up to 64 data sets
+		uint32 BufferSize = ShaderLUT.Num() * sizeof(float);
+		int32 *BufferData = static_cast<int32*>(RHILockVertexBuffer(GPUBuffer.Buffer.Buffer, 0, BufferSize, EResourceLockMode::RLM_WriteOnly));
+		FPlatformMemory::Memcpy(BufferData, ShaderLUT.GetData(), BufferSize);
+		RHIUnlockVertexBuffer(GPUBuffer.Buffer.Buffer);
+		GPUBufferDirty = false;
+	}
+
+	return GPUBuffers;
+}
+
+
+DEFINE_NDI_FUNC_BINDER(UNiagaraDataInterfaceVectorCurve, SampleCurve);
+FVMExternalFunction UNiagaraDataInterfaceVectorCurve::GetVMExternalFunction(const FVMExternalFunctionBindingInfo& BindingInfo, void* InstanceData)
+{
+	if (BindingInfo.Name == TEXT("SampleVectorCurve") && BindingInfo.GetNumInputs() == 1 && BindingInfo.GetNumOutputs() == 3)
+	{
+		return TNDIParamBinder<0, float, NDI_FUNC_BINDER(UNiagaraDataInterfaceVectorCurve, SampleCurve)>::Bind(this, BindingInfo, InstanceData);
+	}
+	else
+	{
+		UE_LOG(LogNiagara, Error, TEXT("Could not find data interface external function.\n\tExpected Name: SampleVectorCurve  Actual Name: %s\n\tExpected Inputs: 1  Actual Inputs: %i\n\tExpected Outputs: 3  Actual Outputs: %i"),
+			*BindingInfo.Name.ToString(), BindingInfo.GetNumInputs(), BindingInfo.GetNumOutputs());
+		return FVMExternalFunction();
+	}
+}
+
+template<typename XParamType>
+void UNiagaraDataInterfaceVectorCurve::SampleCurve(FVectorVMContext& Context)
+{
+	//TODO: Create some SIMDable optimized representation of the curve to do this faster.
+	XParamType XParam(Context);
+	FRegisterHandler<float> OutSampleX(Context);
+	FRegisterHandler<float> OutSampleY(Context);
+	FRegisterHandler<float> OutSampleZ(Context);
+
+	for (int32 i = 0; i < Context.NumInstances; ++i)
+	{
+		float X = XParam.Get();
+		FVector V(XCurve.Eval(X), YCurve.Eval(X), ZCurve.Eval(X));
+		*OutSampleX.GetDest() = V.X;
+		*OutSampleY.GetDest() = V.Y;
+		*OutSampleZ.GetDest() = V.Z;
+		XParam.Advance();
+		OutSampleX.Advance();
+		OutSampleY.Advance();
+		OutSampleZ.Advance();
+	}
+}
