Copying //UE4/Orion-Staging to //UE4/Main (Source: //Orion/Dev-General @ 3606378)

diff --git a/Engine/Plugins/FX/Niagara/Source/Niagara/Private/NiagaraDataSet.cpp b/Engine/Plugins/FX/Niagara/Source/Niagara/Private/NiagaraDataSet.cpp
new file mode 100644
--- /dev/null
+++ b/Engine/Plugins/FX/Niagara/Source/Niagara/Private/NiagaraDataSet.cpp
@@ -0,0 +1,293 @@
+// Copyright 1998-2016 Epic Games, Inc. All Rights Reserved.
+
+#include "NiagaraDataSet.h"
+#include "NiagaraCommon.h"
+#include "NiagaraShader.h"
+#include "GlobalShader.h"
+#include "UpdateTextureShaders.h"
+
+//////////////////////////////////////////////////////////////////////////
+void FNiagaraDataSet::SetShaderParams(FNiagaraShader *Shader, FRHICommandList &CommandList)
+{
+	check(IsInRenderingThread());
+
+	if (Shader->FloatInputBufferParam.IsBound())
+	{
+		CommandList.TransitionResource(EResourceTransitionAccess::EReadable, EResourceTransitionPipeline::EComputeToCompute, PrevDataRender().GetGPUBufferFloat()->UAV);
+		CommandList.SetShaderResourceViewParameter(Shader->GetComputeShader(), Shader->FloatInputBufferParam.GetBaseIndex(), PrevDataRender().GetGPUBufferFloat()->SRV);
+	}
+	if (Shader->IntInputBufferParam.IsBound())
+	{
+		CommandList.TransitionResource(EResourceTransitionAccess::EReadable, EResourceTransitionPipeline::EComputeToCompute, PrevDataRender().GetGPUBufferInt()->UAV);
+		CommandList.SetShaderResourceViewParameter(Shader->GetComputeShader(), Shader->IntInputBufferParam.GetBaseIndex(), PrevDataRender().GetGPUBufferInt()->SRV);
+	}
+	if (Shader->FloatOutputBufferParam.IsUAVBound())
+	{
+		CommandList.TransitionResource(EResourceTransitionAccess::EWritable, EResourceTransitionPipeline::EGfxToCompute, CurrDataRender().GetGPUBufferFloat()->UAV);
+		CommandList.SetUAVParameter(Shader->GetComputeShader(), Shader->FloatOutputBufferParam.GetUAVIndex(), CurrDataRender().GetGPUBufferFloat()->UAV);
+	}
+	if (Shader->IntOutputBufferParam.IsUAVBound())
+	{
+		CommandList.TransitionResource(EResourceTransitionAccess::EWritable, EResourceTransitionPipeline::EGfxToCompute, CurrDataRender().GetGPUBufferInt()->UAV);
+		CommandList.SetUAVParameter(Shader->GetComputeShader(), Shader->IntOutputBufferParam.GetUAVIndex(), CurrDataRender().GetGPUBufferInt()->UAV);
+	}
+
+	if (Shader->ComponentBufferSizeWriteParam.IsBound())
+	{
+		uint32 SafeBufferSize = CurrDataRender().GetFloatStride() / sizeof(float);
+		CommandList.SetShaderParameter(Shader->GetComputeShader(), 0, Shader->ComponentBufferSizeWriteParam.GetBaseIndex(), Shader->ComponentBufferSizeWriteParam.GetNumBytes(), &SafeBufferSize);
+	}
+
+	if (Shader->ComponentBufferSizeReadParam.IsBound())
+	{
+		uint32 SafeBufferSize = PrevDataRender().GetFloatStride() / sizeof(float);
+		CommandList.SetShaderParameter(Shader->GetComputeShader(), 0, Shader->ComponentBufferSizeReadParam.GetBaseIndex(), Shader->ComponentBufferSizeReadParam.GetNumBytes(), &SafeBufferSize);
+	}
+}
+
+
+
+void FNiagaraDataSet::UnsetShaderParams(FNiagaraShader *Shader, FRHICommandList &RHICmdList)
+{
+	check(IsInRenderingThread());
+
+	if (Shader->FloatOutputBufferParam.IsUAVBound())
+	{
+#if !PLATFORM_PS4
+		Shader->FloatOutputBufferParam.UnsetUAV(RHICmdList, Shader->GetComputeShader());
+#endif
+		//RHICmdList.TransitionResource(EResourceTransitionAccess::EReadable, EResourceTransitionPipeline::EGfxToGfx, CurrDataRender().GetGPUBufferFloat()->UAV);
+	}
+
+	if (Shader->IntOutputBufferParam.IsUAVBound())
+	{
+#if !PLATFORM_PS4
+		Shader->IntOutputBufferParam.UnsetUAV(RHICmdList, Shader->GetComputeShader());
+#endif
+		//RHICmdList.TransitionResource(EResourceTransitionAccess::EReadable, EResourceTransitionPipeline::EGfxToGfx, CurrDataRender().GetGPUBufferInt()->UAV);
+	}
+}
+
+void FNiagaraDataSet::Dump(FNiagaraDataSet& Other, bool bCurr)
+{
+	Other.Reset();
+	Other.Variables = Variables;
+	Other.VariableLayoutMap = VariableLayoutMap;
+
+	int32 IndexRead = CurrBuffer > 0 ? CurrBuffer - 1 : MaxBufferIdx;
+	if (bCurr)
+	{
+		IndexRead = CurrBuffer;
+	}
+
+	if (Other.Data[Other.CurrBuffer].GetNumInstancesAllocated() != Data[IndexRead].GetNumInstancesAllocated())
+	{
+		Other.Finalize();
+		Other.Data[Other.CurrBuffer].Allocate(Data[IndexRead].GetNumInstancesAllocated());
+	}
+
+	Data[CurrBuffer].CopyTo(Other.Data[Other.CurrBuffer]);
+}
+
+
+void FNiagaraDataSet::Dump(bool bCurr, int32 StartIdx, int32 NumInstances)
+{
+	TArray<FNiagaraVariable> Vars(Variables);
+
+	FNiagaraDataSetVariableIterator Itr(*this, StartIdx, bCurr);
+	Itr.AddVariables(Vars);
+
+	if (NumInstances == INDEX_NONE)
+	{
+		NumInstances = bCurr ? GetNumInstances() : GetPrevNumInstances();
+		NumInstances -= StartIdx;
+	}
+
+	int32 NumInstancesDumped = 0;
+	TArray<FString> Lines;
+	Lines.Reserve(GetNumInstances());
+	while (Itr.IsValid() && NumInstancesDumped++ < NumInstances)
+	{
+		Itr.Get();
+
+		FString Line = TEXT("| ");
+		for (FNiagaraVariable& Var : Vars)
+		{
+			Line += Var.ToString() + TEXT(" | ");
+		}
+		Lines.Add(Line);
+		Itr.Advance();
+	}
+
+	static FString Sep;
+	if (Sep.Len() == 0)
+	{
+		for (int32 i = 0; i < 50; ++i)
+		{
+			Sep.AppendChar(TEXT('='));
+		}
+	}
+
+	UE_LOG(LogNiagara, Log, TEXT("%s"), *Sep);
+	UE_LOG(LogNiagara, Log, TEXT(" Buffer: %d"), CurrBuffer);
+ 	UE_LOG(LogNiagara, Log, TEXT("%s"), *Sep);
+// 	UE_LOG(LogNiagara, Log, TEXT("%s"), *HeaderStr);
+// 	UE_LOG(LogNiagara, Log, TEXT("%s"), *Sep);
+	for (FString& Str : Lines)
+	{
+		UE_LOG(LogNiagara, Log, TEXT("%s"), *Str);
+	}
+	UE_LOG(LogNiagara, Log, TEXT("%s"), *Sep);
+}
+
+//////////////////////////////////////////////////////////////////////////
+
+void FNiagaraDataBuffer::Init(FNiagaraDataSet* InOwner)
+{
+	Owner = InOwner;
+}
+
+void FNiagaraDataBuffer::Allocate(uint32 InNumInstances, ENiagaraSimTarget Target, bool bMaintainExisting)
+{
+	check(Owner);
+	if (Target == ENiagaraSimTarget::CPUSim)
+	{
+		NumInstancesAllocated = InNumInstances;
+
+		int32 OldFloatStride = FloatStride;
+		FloatStride = GetSafeComponentBufferSize(NumInstancesAllocated * sizeof(float));
+		FloatData.SetNum(FloatStride * Owner->GetNumFloatComponents(), false);
+
+		int32 OldInt32Stride = Int32Stride;
+		Int32Stride = GetSafeComponentBufferSize(NumInstancesAllocated * sizeof(int32));
+		Int32Data.SetNum(Int32Stride * Owner->GetNumInt32Components(), false);
+
+		//In some cases we want the existing data in the buffer to be maintained which due to the data layout requires some fix up.
+		if (bMaintainExisting)
+		{
+			if (FloatStride != OldFloatStride && FloatStride > 0 && OldFloatStride > 0)
+			{
+				for (uint32 CompIdx = Owner->TotalFloatComponents-1; CompIdx > 0; --CompIdx)
+				{
+					uint8* Src = FloatData.GetData() + OldFloatStride * CompIdx;
+					uint8* Dst = FloatData.GetData() + FloatStride * CompIdx;
+					FMemory::Memcpy(Dst, Src, OldFloatStride);
+				}
+			}
+			if (Int32Stride != OldInt32Stride && Int32Stride > 0 && OldInt32Stride > 0)
+			{
+				for (uint32 CompIdx = Owner->TotalInt32Components - 1; CompIdx > 0; --CompIdx)
+				{
+					uint8* Src = Int32Data.GetData() + OldInt32Stride * CompIdx;
+					uint8* Dst = Int32Data.GetData() + Int32Stride * CompIdx;
+					FMemory::Memcpy(Dst, Src, OldFloatStride);
+				}
+			}
+		}
+	}
+}
+
+
+void FNiagaraDataBuffer::AllocateGPU(uint32 InNumInstances, FRHICommandList &RHICmdList)
+{
+	if (Owner == 0)
+	{
+		return;
+	}
+	check(IsInRenderingThread())
+	const uint32 ALLOC_CHUNKSIZE = 4096;
+
+	NumInstancesAllocated = InNumInstances;
+
+	uint32 PaddedNumInstances = ((InNumInstances + NIAGARA_COMPUTE_THREADGROUP_SIZE -1) / NIAGARA_COMPUTE_THREADGROUP_SIZE) * NIAGARA_COMPUTE_THREADGROUP_SIZE;
+	FloatStride = PaddedNumInstances * sizeof(float);
+	Int32Stride = PaddedNumInstances * sizeof(int32);
+
+	if (NumInstancesAllocated > NumChunksAllocatedForGPU * ALLOC_CHUNKSIZE)
+	{
+		uint32 PrevChunks = NumChunksAllocatedForGPU;
+		NumChunksAllocatedForGPU = ((InNumInstances + ALLOC_CHUNKSIZE - 1) / ALLOC_CHUNKSIZE);
+		uint32 NumElementsToAlloc = NumChunksAllocatedForGPU * ALLOC_CHUNKSIZE;
+
+		if (Owner->GetNumFloatComponents())
+		{
+			if (GPUBufferFloat.Buffer)
+			{
+				GPUBufferFloat.Release();
+			}
+			GPUBufferFloat.Initialize(sizeof(float), NumElementsToAlloc * Owner->GetNumFloatComponents(), EPixelFormat::PF_R32_FLOAT);
+		}
+		if (Owner->GetNumInt32Components())
+		{
+			if (GPUBufferInt.Buffer)
+			{
+				GPUBufferInt.Release();
+			}
+			GPUBufferInt.Initialize(sizeof(int32), NumElementsToAlloc * Owner->GetNumInt32Components(), EPixelFormat::PF_R32_SINT);
+		}
+	}
+}
+
+
+void FNiagaraDataBuffer::InitGPUFromCPU()
+{
+	if (Owner->GetNumFloatComponents())
+	{
+		GPUBufferFloat.Release();
+		GPUBufferFloat.Initialize(sizeof(float), FloatStride/sizeof(float) * Owner->GetNumFloatComponents(), EPixelFormat::PF_R32_FLOAT, 0, TEXT("GPUBufferFloat"), &FloatData);
+	}
+	if (Owner->GetNumInt32Components())
+	{
+		GPUBufferInt.Release();
+		GPUBufferInt.Initialize(sizeof(int32), Int32Stride/sizeof(int32) * Owner->GetNumInt32Components(), EPixelFormat::PF_R32_SINT, 0, TEXT("GPUBufferInt"), &Int32Data);
+	}
+}
+
+void FNiagaraDataBuffer::SwapInstances(uint32 OldIndex, uint32 NewIndex) 
+{
+	for (uint32 CompIdx = 0; CompIdx < Owner->TotalFloatComponents; ++CompIdx)
+	{
+		float* Src = GetInstancePtrFloat(CompIdx, OldIndex);
+		float* Dst = GetInstancePtrFloat(CompIdx, NewIndex);
+		float Temp = *Dst;
+		*Dst = *Src;
+		*Src = Temp;
+	}
+	for (uint32 CompIdx = 0; CompIdx < Owner->TotalInt32Components; ++CompIdx)
+	{
+		int32* Src = GetInstancePtrInt32(CompIdx, OldIndex);
+		int32* Dst = GetInstancePtrInt32(CompIdx, NewIndex);
+		int32 Temp = *Dst;
+		*Dst = *Src;
+		*Src = Temp;
+	}
+}
+
+void FNiagaraDataBuffer::KillInstance(uint32 InstanceIdx)
+{
+	check(InstanceIdx < NumInstances);
+	--NumInstances;
+
+	for (uint32 CompIdx = 0; CompIdx < Owner->TotalFloatComponents; ++CompIdx)
+	{
+		float* Src = GetInstancePtrFloat(CompIdx, NumInstances);
+		float* Dst = GetInstancePtrFloat(CompIdx, InstanceIdx);
+		*Dst = *Src;
+	}
+	for (uint32 CompIdx = 0; CompIdx < Owner->TotalInt32Components; ++CompIdx)
+	{
+		int32* Src = GetInstancePtrInt32(CompIdx, NumInstances);
+		int32* Dst = GetInstancePtrInt32(CompIdx, InstanceIdx);
+		*Dst = *Src;
+	}
+}
+
+void FNiagaraDataBuffer::CopyTo(FNiagaraDataBuffer& DestBuffer)
+{
+	DestBuffer.FloatStride = FloatStride;
+	DestBuffer.FloatData = FloatData;
+	DestBuffer.Int32Stride = Int32Stride;
+	DestBuffer.Int32Data = Int32Data;
+	DestBuffer.NumInstancesAllocated = NumInstancesAllocated;
+	DestBuffer.NumInstances = NumInstances;
+}
\ No newline at end of file
diff --git a/Engine/Plugins/FX/Niagara/Source/Niagara/Private/NiagaraEmitter.cpp b/Engine/Plugins/FX/Niagara/Source/Niagara/Private/NiagaraEmitter.cpp
new file mode 100644
--- /dev/null
+++ b/Engine/Plugins/FX/Niagara/Source/Niagara/Private/NiagaraEmitter.cpp
@@ -0,0 +1,517 @@
+// Copyright 1998-2017 Epic Games, Inc. All Rights Reserved.
+
+
+#include "NiagaraEmitter.h"
+#include "NiagaraScript.h"
+#include "NiagaraScriptSourceBase.h"
+#include "NiagaraSpriteRendererProperties.h"
+#include "NiagaraCustomVersion.h"
+#include "Package.h"
+#include "Linker.h"
+#include "NiagaraModule.h"
+
+static int32 GbForceNiagaraCompileOnLoad = 0;
+static FAutoConsoleVariableRef CVarForceNiagaraCompileOnLoad(
+	TEXT("fx.ForceCompileOnLoad"),
+	GbForceNiagaraCompileOnLoad,
+	TEXT("If > 0 emitters will be forced to compile on load. \n"),
+	ECVF_Default
+	);
+
+
+void FNiagaraEmitterScriptProperties::InitDataSetAccess()
+{
+	EventReceivers.Empty();
+	EventGenerators.Empty();
+
+	if (Script)
+	{
+		// TODO: add event receiver and generator lists to the script properties here
+		//
+		for (FNiagaraDataSetID &ReadID : Script->ReadDataSets)
+		{
+			EventReceivers.Add( FNiagaraEventReceiverProperties(ReadID.Name, "", "") );
+		}
+
+		for (FNiagaraDataSetProperties &WriteID : Script->WriteDataSets)
+		{
+			FNiagaraEventGeneratorProperties Props(WriteID, "", "");
+			EventGenerators.Add(Props);
+		}
+	}
+}
+
+//////////////////////////////////////////////////////////////////////////
+
+UNiagaraEmitter::UNiagaraEmitter(const FObjectInitializer& Initializer)
+: Super(Initializer)
+, CollisionMode(ENiagaraCollisionMode::None)
+, bInterpolatedSpawning(false)
+{
+}
+
+void UNiagaraEmitter::PostInitProperties()
+{
+	Super::PostInitProperties();
+	if (HasAnyFlags(RF_ClassDefaultObject | RF_NeedLoad) == false)
+	{
+		RendererProperties.Add(NewObject<UNiagaraSpriteRendererProperties>(this, "Renderer"));
+
+		SpawnScriptProps.Script = NewObject<UNiagaraScript>(this, "SpawnScript", EObjectFlags::RF_Transactional);
+		SpawnScriptProps.Script->SetUsage(ENiagaraScriptUsage::ParticleSpawnScript);
+
+		UpdateScriptProps.Script = NewObject<UNiagaraScript>(this, "UpdateScript", EObjectFlags::RF_Transactional);
+		UpdateScriptProps.Script->SetUsage(ENiagaraScriptUsage::ParticleUpdateScript);
+	}
+}
+
+void UNiagaraEmitter::Serialize(FArchive& Ar)
+{
+	Super::Serialize(Ar);
+
+	Ar.UsingCustomVersion(FNiagaraCustomVersion::GUID);
+}
+
+void UNiagaraEmitter::PostLoad()
+{
+	Super::PostLoad();
+
+	if (GIsEditor)
+	{
+		SetFlags(RF_Transactional);
+	}
+
+// 	//Have to force the usage to be correct.
+// 	if (SpawnScriptProps.Script)
+// 	{
+// 		SpawnScriptProps.Script->ConditionalPostLoad();
+// 		ENiagaraScriptUsage ActualUsage = bInterpolatedSpawning ? ENiagaraScriptUsage::ParticleSpawnScriptInterpolated : ENiagaraScriptUsage::ParticleSpawnScript;
+// 		if (SpawnScriptProps.Script->Usage != ActualUsage)
+// 		{
+// 			SpawnScriptProps.Script->Usage = ActualUsage;
+// 			SpawnScriptProps.Script->ByteCode.Empty();//Force them to recompile.
+// 			UE_LOG(LogNiagara, Warning, TEXT("Niagara Emitter had mis matched bInterpolatedSpawning values. You must recompile this emitter. %s"), *GetFullName());
+// 		}
+// 	}
+		
+	//Temporarily disabling interpolated spawn.
+	if (SpawnScriptProps.Script)
+	{
+		SpawnScriptProps.Script->ConditionalPostLoad();
+		bool bActualInterpolatedSpawning = SpawnScriptProps.Script->IsInterpolatedParticleSpawnScript();
+		if (bInterpolatedSpawning != bActualInterpolatedSpawning)
+		{
+			bInterpolatedSpawning = false;
+			if (bActualInterpolatedSpawning)
+			{
+				SpawnScriptProps.Script->ByteCode.Empty();//clear out the script as it was compiled with interpolated spawn.
+#if WITH_EDITORONLY_DATA
+				SpawnScriptProps.Script->InvalidateChangeID(); // Identify that we have no idea if we're synchronized or not.
+#endif
+				SpawnScriptProps.Script->SetUsage(ENiagaraScriptUsage::ParticleSpawnScript);
+			}
+			UE_LOG(LogNiagara, Warning, TEXT("Temporarily disabling interpolated spawn. Emitter may need recompile.. %s"), *GetFullName());
+		}
+	}
+
+	// Referenced scripts may need to be invalidated due to version mismatch or 
+	// other issues. This is determined in PostLoad for UNiagaraScript, so we need
+	// to make sure that PostLoad has happened for each of them first.
+	bool bNeedsRecompile = false;
+	if (SpawnScriptProps.Script)
+	{
+		SpawnScriptProps.Script->ConditionalPostLoad();
+		if (SpawnScriptProps.Script->ByteCode.Num() == 0 || GbForceNiagaraCompileOnLoad > 0)
+		{
+			bNeedsRecompile = true;
+		}
+	}
+	if (UpdateScriptProps.Script)
+	{
+		UpdateScriptProps.Script->ConditionalPostLoad();
+		if (UpdateScriptProps.Script->ByteCode.Num() == 0 || GbForceNiagaraCompileOnLoad > 0)
+		{
+			bNeedsRecompile = true;
+		}
+	}
+
+	for (int32 i = 0; i < EventHandlerScriptProps.Num(); i++)
+	{
+		if (EventHandlerScriptProps[i].Script)
+		{
+			EventHandlerScriptProps[i].Script->ConditionalPostLoad();
+			if (EventHandlerScriptProps[i].Script->ByteCode.Num() == 0 || GbForceNiagaraCompileOnLoad > 0)
+			{
+				bNeedsRecompile = true;
+			}
+		}
+	}
+
+	// Check ourselves against the current script rebuild version. Force our children
+	// to need to update.
+	const int32 NiagaraVer = GetLinkerCustomVersion(FNiagaraCustomVersion::GUID);
+	if (NiagaraVer < FNiagaraCustomVersion::LatestVersion)
+	{
+		bNeedsRecompile = true;
+	}
+	
+	if (bNeedsRecompile)
+	{
+		if (SpawnScriptProps.Script)
+		{
+			SpawnScriptProps.Script->ByteCode.Empty();//clear out the script.
+#if WITH_EDITORONLY_DATA
+			SpawnScriptProps.Script->InvalidateChangeID(); // Identify that we have no idea if we're synchronized or not.
+#endif
+		}
+		if (UpdateScriptProps.Script)
+		{
+			UpdateScriptProps.Script->ByteCode.Empty();//clear out the script.
+#if WITH_EDITORONLY_DATA
+			UpdateScriptProps.Script->InvalidateChangeID(); // Identify that we have no idea if we're synchronized or not.
+#endif
+		}
+
+		for (int32 i = 0; i < EventHandlerScriptProps.Num(); i++)
+		{
+			if (EventHandlerScriptProps[i].Script)
+			{
+				EventHandlerScriptProps[i].Script->ByteCode.Empty();//clear out the script.
+#if WITH_EDITORONLY_DATA
+				EventHandlerScriptProps[i].Script->InvalidateChangeID(); // Identify that we have no idea if we're synchronized or not.
+#endif
+			}
+		}
+	}
+
+#if WITH_EDITORONLY_DATA
+	// Go ahead and do a recompile for all referenced scripts. Note that the compile will 
+	// cause the emitter's change id to be regenerated.
+	if (bNeedsRecompile)
+	{
+		UObject* Outer = GetOuter();
+		if (Outer != nullptr && Outer->IsA<UPackage>())
+		{
+			TArray<ENiagaraScriptCompileStatus> ScriptStatuses;
+			TArray<FString> ScriptErrors;
+			TArray<FString> PathNames;
+			TArray<UNiagaraScript*> Scripts;
+			CompileScripts(ScriptStatuses, ScriptErrors, PathNames, Scripts);
+
+			for (int32 i = 0; i < ScriptStatuses.Num(); i++)
+			{
+				if (ScriptErrors[i].Len() != 0 || ScriptStatuses[i] != ENiagaraScriptCompileStatus::NCS_UpToDate)
+				{
+					UE_LOG(LogNiagara, Warning, TEXT("Script '%s', compile status: %d  errors: %s"), *PathNames[i], (int32)ScriptStatuses[i], *ScriptErrors[i]);
+				}
+				else
+				{
+					UE_LOG(LogNiagara, Log, TEXT("Script '%s', compile status: Success!"), *PathNames[i]);
+				}
+			}
+		}
+	}
+#endif
+}
+
+#if WITH_EDITOR
+void UNiagaraEmitter::PostEditChangeProperty(struct FPropertyChangedEvent& PropertyChangedEvent)
+{
+	Super::PostEditChangeProperty(PropertyChangedEvent);
+
+	FName PropertyName;
+	if (PropertyChangedEvent.Property)
+	{
+		PropertyName = PropertyChangedEvent.Property->GetFName();
+	}
+
+	if (PropertyName == GET_MEMBER_NAME_CHECKED(UNiagaraEmitter, bInterpolatedSpawning))
+	{
+		bool bActualInterpolatedSpawning = SpawnScriptProps.Script->IsInterpolatedParticleSpawnScript();
+		if (bInterpolatedSpawning != bActualInterpolatedSpawning)
+		{
+			//Recompile spawn script if we've altered the interpolated spawn property.
+			SpawnScriptProps.Script->SetUsage(bInterpolatedSpawning ? ENiagaraScriptUsage::ParticleSpawnScriptInterpolated : ENiagaraScriptUsage::ParticleSpawnScript);
+			UE_LOG(LogNiagara, Log, TEXT("Updating script usage: Script->IsInterpolatdSpawn %d Emitter->bInterpolatedSpawning %d"), (int32)SpawnScriptProps.Script->IsInterpolatedParticleSpawnScript(), bInterpolatedSpawning);
+			if (GraphSource != nullptr)
+			{
+				GraphSource->MarkNotSynchronized();
+			}
+		}
+	}
+}
+#endif
+
+
+bool UNiagaraEmitter::IsValid()const
+{
+	return SpawnScriptProps.Script && SpawnScriptProps.Script->IsValid() && UpdateScriptProps.Script && UpdateScriptProps.Script->IsValid();
+}
+
+#if WITH_EDITORONLY_DATA
+
+void UNiagaraEmitter::GetScripts(TArray<UNiagaraScript*>& OutScripts)
+{
+	OutScripts.Add(SpawnScriptProps.Script);
+	OutScripts.Add(UpdateScriptProps.Script);
+	for (int32 i = 0; i < EventHandlerScriptProps.Num(); i++)
+	{
+		if (EventHandlerScriptProps[i].Script)
+		{
+			OutScripts.Add(EventHandlerScriptProps[i].Script);
+		}
+	}
+}
+
+void UNiagaraEmitter::CompileScripts(TArray<ENiagaraScriptCompileStatus>& OutScriptStatuses, TArray<FString>& OutGraphLevelErrorMessages, TArray<FString>& PathNames,
+	TArray<UNiagaraScript*>& Scripts)
+{
+	PathNames.AddDefaulted(2);
+	OutScriptStatuses.Empty();
+	OutGraphLevelErrorMessages.Empty();
+	bool bPostCompile = false;
+
+	// Force the end-user to reopen the file to resolve the problems.
+	if (GraphSource == nullptr)
+	{
+		OutScriptStatuses.AddDefaulted(2);
+		OutGraphLevelErrorMessages.AddDefaulted(2);
+
+		OutGraphLevelErrorMessages[0] = OutGraphLevelErrorMessages[1] = TEXT("Please reopen asset in editor.");
+		OutScriptStatuses[0] = OutScriptStatuses[1] = ENiagaraScriptCompileStatus::NCS_Error;
+		return;
+	}
+	else if (!GraphSource->IsPreCompiled()) // If we haven't already precompiled this source, do so..
+	{
+		GraphSource->PreCompile(this);
+		bPostCompile = true;
+	}
+
+	FString ErrorMsg;
+	check((int32)EScriptCompileIndices::SpawnScript == 0);
+	OutScriptStatuses.Add(SpawnScriptProps.Script->Compile(ErrorMsg));
+	OutGraphLevelErrorMessages.Add(ErrorMsg); // EScriptCompileIndices::SpawnScript
+	Scripts.Add(SpawnScriptProps.Script);
+	SpawnScriptProps.InitDataSetAccess();
+	PathNames[(int32)EScriptCompileIndices::SpawnScript] = SpawnScriptProps.Script->GetPathName();
+
+	ErrorMsg.Empty();
+	check((int32)EScriptCompileIndices::UpdateScript == 1);
+	OutScriptStatuses.Add(UpdateScriptProps.Script->Compile(ErrorMsg));
+	Scripts.Add(UpdateScriptProps.Script);
+	OutGraphLevelErrorMessages.Add(ErrorMsg); // EScriptCompileIndices::UpdateScript
+	UpdateScriptProps.InitDataSetAccess();
+	PathNames[(int32)EScriptCompileIndices::UpdateScript] = UpdateScriptProps.Script->GetPathName();
+
+	for (int32 i = 0; i < EventHandlerScriptProps.Num(); i++)
+	{
+		if (EventHandlerScriptProps[i].Script)
+		{
+			check((int32)EScriptCompileIndices::EventScript == 2);
+			ErrorMsg.Empty();
+			OutScriptStatuses.Add(EventHandlerScriptProps[i].Script->Compile(ErrorMsg));
+			Scripts.Add(EventHandlerScriptProps[i].Script);
+			OutGraphLevelErrorMessages.Add(ErrorMsg); // EScriptCompileIndices::EventScript
+			EventHandlerScriptProps[i].InitDataSetAccess();
+			PathNames.Add(EventHandlerScriptProps[i].Script->GetPathName());
+		}
+	}
+
+	// Go through all attached renderers and make sure that the scripts work properly for those renderers.
+	for (int32 ScriptIdx = 0; ScriptIdx < Scripts.Num(); ScriptIdx++)
+	{
+		// Only care about spawn and update
+		if (ScriptIdx >= 2)
+		{
+			continue;
+		}
+
+		// Only care about scripts that successfully compiled so far.
+		if (OutScriptStatuses[ScriptIdx] != ENiagaraScriptCompileStatus::NCS_UpToDate)
+		{
+			continue;
+		}
+
+		// Ignore missing scripts.
+		if (Scripts[ScriptIdx] == nullptr)
+		{
+			continue;
+		}
+
+		for (int32 i = 0; i < RendererProperties.Num(); i++)
+		{
+			if (RendererProperties[i] != nullptr)
+			{
+				const TArray<FNiagaraVariable>& RequiredAttrs = RendererProperties[i]->GetRequiredAttributes();
+
+				for (FNiagaraVariable Attr : RequiredAttrs)
+				{
+					// TODO .. should we always be namespaced?
+					FString AttrName = Attr.GetName().ToString();
+					if (AttrName.RemoveFromStart(TEXT("Particles.")))
+					{
+						Attr.SetName(*AttrName);
+					}
+
+					if (false == Scripts[ScriptIdx]->Attributes.ContainsByPredicate([&Attr](const FNiagaraVariable& Var) { return Var.GetName() == Attr.GetName(); }))
+					{
+						OutGraphLevelErrorMessages[ScriptIdx] += FString::Printf(TEXT("\nCannot bind to renderer %s because it does not define attribute %s %s."), *RendererProperties[i]->GetName(), *Attr.GetType().GetNameText().ToString(), *Attr.GetName().ToString());
+						OutScriptStatuses[ScriptIdx] = ENiagaraScriptCompileStatus::NCS_Error;
+					}
+				}
+			}
+		}
+	}
+
+	ChangeId = FGuid::NewGuid();
+
+	if (GraphSource && bPostCompile) // Cleanup if we precompiled inside this function
+	{
+		GraphSource->PostCompile();
+	}
+}
+
+ENiagaraScriptCompileStatus UNiagaraEmitter::CompileScript(EScriptCompileIndices InScriptToCompile, FString& OutGraphLevelErrorMessages, int32 InSubScriptIdx)
+{
+	ENiagaraScriptCompileStatus CompileStatus = ENiagaraScriptCompileStatus::NCS_Unknown;
+	bool bPostCompile = false;
+	if (GraphSource && !GraphSource->IsPreCompiled()) // If we haven't already precompiled this source, do so..
+	{
+		bPostCompile = true;
+		GraphSource->PreCompile(this);
+	}
+
+	switch (InScriptToCompile)
+	{
+	case EScriptCompileIndices::SpawnScript:
+		CompileStatus = SpawnScriptProps.Script->Compile(OutGraphLevelErrorMessages);
+		SpawnScriptProps.InitDataSetAccess();
+		break;
+	case EScriptCompileIndices::UpdateScript:
+		CompileStatus = UpdateScriptProps.Script->Compile(OutGraphLevelErrorMessages);
+		UpdateScriptProps.InitDataSetAccess();
+		break;
+	case EScriptCompileIndices::EventScript:
+		if (EventHandlerScriptProps[InSubScriptIdx].Script)
+		{
+			CompileStatus = EventHandlerScriptProps[InSubScriptIdx].Script->Compile(OutGraphLevelErrorMessages);
+			EventHandlerScriptProps[InSubScriptIdx].InitDataSetAccess();
+		}
+		else
+		{
+			return CompileStatus;
+		}
+		break;
+	default:
+		return CompileStatus;
+	}
+	ChangeId = FGuid::NewGuid();
+
+	if (GraphSource && bPostCompile) // Cleanup if we precompiled inside this function
+	{
+		GraphSource->PostCompile();
+	}
+	return CompileStatus;
+}
+
+UNiagaraEmitter* UNiagaraEmitter::MakeRecursiveDeepCopy(UObject* DestOuter) const
+{
+	TMap<const UObject*, UObject*> ExistingConversions;
+	return MakeRecursiveDeepCopy(DestOuter, ExistingConversions);
+}
+
+UNiagaraEmitter* UNiagaraEmitter::MakeRecursiveDeepCopy(UObject* DestOuter, TMap<const UObject*, UObject*>& ExistingConversions) const
+{
+	ResetLoaders(GetTransientPackage());
+	GetTransientPackage()->LinkerCustomVersion.Empty();
+
+	EObjectFlags Flags = RF_AllFlags & ~RF_Standalone & ~RF_Public; // Remove Standalone and Public flags..
+	UNiagaraEmitter* Props = CastChecked<UNiagaraEmitter>(StaticDuplicateObject(this, GetTransientPackage(), *GetName(), Flags));
+	check(Props->HasAnyFlags(RF_Standalone) == false);
+	check(Props->HasAnyFlags(RF_Public) == false);
+	Props->Rename(nullptr, DestOuter, REN_DoNotDirty | REN_DontCreateRedirectors | REN_NonTransactional);
+	UE_LOG(LogNiagara, Warning, TEXT("MakeRecursiveDeepCopy %s"), *Props->GetFullName());
+	ExistingConversions.Add(this, Props);
+
+	check(GraphSource != Props->GraphSource);
+
+	Props->GraphSource->SubsumeExternalDependencies(ExistingConversions);
+	ExistingConversions.Add(GraphSource, Props->GraphSource);
+
+	// Suck in the referenced scripts into this package.
+	if (Props->SpawnScriptProps.Script)
+	{
+		Props->SpawnScriptProps.Script->SubsumeExternalDependencies(ExistingConversions);
+		check(Props->GraphSource == Props->SpawnScriptProps.Script->GetSource());
+	}
+
+	if (Props->UpdateScriptProps.Script)
+	{
+		Props->UpdateScriptProps.Script->SubsumeExternalDependencies(ExistingConversions);
+		check(Props->GraphSource == Props->UpdateScriptProps.Script->GetSource());
+	}
+
+	for (int32 i = 0; i < Props->EventHandlerScriptProps.Num(); i++)
+	{
+		if (Props->EventHandlerScriptProps[i].Script)
+		{
+			Props->EventHandlerScriptProps[i].Script->SubsumeExternalDependencies(ExistingConversions);
+			check(Props->GraphSource == Props->EventHandlerScriptProps[i].Script->GetSource());
+		}
+	}
+	return Props;
+}
+#endif
+
+bool UNiagaraEmitter::UsesScript(const UNiagaraScript* Script)const
+{
+	if (SpawnScriptProps.Script == Script || UpdateScriptProps.Script == Script)
+	{
+		return true;
+	}
+	for (int32 i = 0; i < EventHandlerScriptProps.Num(); i++)
+	{
+		if (EventHandlerScriptProps[i].Script == Script)
+		{
+			return true;
+		}
+	}
+	return false;
+}
+
+//TODO
+// bool UNiagaraEmitter::UsesDataInterface(UNiagaraDataInterface* Interface)
+//{
+//}
+
+bool UNiagaraEmitter::UsesCollection(const class UNiagaraParameterCollection* Collection)const
+{
+	if (SpawnScriptProps.Script && SpawnScriptProps.Script->UsesCollection(Collection))
+	{
+		return true;
+	}
+	if (UpdateScriptProps.Script && UpdateScriptProps.Script->UsesCollection(Collection))
+	{
+		return true;
+	}
+	for (int32 i = 0; i < EventHandlerScriptProps.Num(); i++)
+	{
+		if (EventHandlerScriptProps[i].Script && EventHandlerScriptProps[i].Script->UsesCollection(Collection))
+		{
+			return true;
+		}
+	}
+	return false;
+}
+
+FString UNiagaraEmitter::GetUniqueEmitterName()const
+{
+	return GetName() + LexicalConversion::ToString(FCrc::StrCrc32(*GetPathName()));
+}
+
+FNiagaraVariable UNiagaraEmitter::GetEmitterParameter(const FNiagaraVariable& EmitterVar)const
+{
+	FNiagaraVariable Var = EmitterVar;
+	Var.SetName(*Var.GetName().ToString().Replace(TEXT("Emitter."), *(GetUniqueEmitterName() + TEXT("."))));
+	return Var;
+}
\ No newline at end of file
diff --git a/Engine/Plugins/FX/Niagara/Source/Niagara/Private/NiagaraEmitterHandle.cpp b/Engine/Plugins/FX/Niagara/Source/Niagara/Private/NiagaraEmitterHandle.cpp
new file mode 100644
--- /dev/null
+++ b/Engine/Plugins/FX/Niagara/Source/Niagara/Private/NiagaraEmitterHandle.cpp
@@ -0,0 +1,279 @@
+// Copyright 1998-2016 Epic Games, Inc. All Rights Reserved.
+
+#include "NiagaraEmitterHandle.h"
+#include "NiagaraSystem.h"
+#include "NiagaraEmitter.h"
+#include "NiagaraScriptSourceBase.h"
+#include "NiagaraCommon.h"
+#include "NiagaraDataInterface.h"
+
+const FNiagaraEmitterHandle FNiagaraEmitterHandle::InvalidHandle;
+
+FNiagaraEmitterHandle::FNiagaraEmitterHandle() 
+	:
+#if WITH_EDITORONLY_DATA
+	Source(nullptr)
+	,
+#endif
+	Instance(nullptr)
+{
+}
+
+FNiagaraEmitterHandle::FNiagaraEmitterHandle(UNiagaraEmitter& Emitter)
+	: Id(FGuid::NewGuid())
+	, IdName(*Id.ToString())
+	, bIsEnabled(true)
+	, Name(TEXT("Emitter"))
+#if WITH_EDITORONLY_DATA
+	, Source(&Emitter)
+#endif
+	, Instance(&Emitter)
+{
+}
+
+#if WITH_EDITORONLY_DATA
+FNiagaraEmitterHandle::FNiagaraEmitterHandle(const UNiagaraEmitter& InSourceEmitter, FName InName, UNiagaraSystem& InOuterSystem)
+	: Id(FGuid::NewGuid())
+	, IdName(*Id.ToString())
+	, bIsEnabled(true)
+	, Name(InName)
+	, Source(&InSourceEmitter)
+	, Instance(InSourceEmitter.MakeRecursiveDeepCopy(&InOuterSystem))
+{
+}
+
+FNiagaraEmitterHandle::FNiagaraEmitterHandle(const FNiagaraEmitterHandle& InHandleToDuplicate, FName InDuplicateName, UNiagaraSystem& InDuplicateOwnerSystem)
+	: Id(FGuid::NewGuid())
+	, IdName(*Id.ToString())
+	, bIsEnabled(InHandleToDuplicate.bIsEnabled)
+	, Name(InDuplicateName)
+	, Source(InHandleToDuplicate.Source)
+	, Instance(InHandleToDuplicate.Instance->MakeRecursiveDeepCopy(&InDuplicateOwnerSystem))
+{
+}
+#endif
+
+bool FNiagaraEmitterHandle::IsValid() const
+{
+	return Id.IsValid();
+}
+
+FGuid FNiagaraEmitterHandle::GetId() const
+{
+	return Id;
+}
+
+FName FNiagaraEmitterHandle::GetIdName() const
+{
+	return IdName;
+}
+
+FName FNiagaraEmitterHandle::GetName() const
+{
+	return Name;
+}
+
+void FNiagaraEmitterHandle::SetName(FName InName)
+{
+	Name = InName;
+}
+
+bool FNiagaraEmitterHandle::GetIsEnabled() const
+{
+	return bIsEnabled;
+}
+
+void FNiagaraEmitterHandle::SetIsEnabled(bool bInIsEnabled)
+{
+	bIsEnabled = bInIsEnabled;
+}
+
+#if WITH_EDITORONLY_DATA
+const UNiagaraEmitter* FNiagaraEmitterHandle::GetSource() const
+{
+	return Source;
+}
+
+#endif
+
+UNiagaraEmitter* FNiagaraEmitterHandle::GetInstance() const
+{
+	return Instance;
+}
+
+FString FNiagaraEmitterHandle::GetUniqueInstanceName()const
+{
+	check(Instance);
+	return Instance->GetUniqueEmitterName();
+}
+
+void FNiagaraEmitterHandle::SetInstance(UNiagaraEmitter* InInstance)
+{
+	Instance = InInstance;
+}
+#if WITH_EDITORONLY_DATA
+
+void FNiagaraEmitterHandle::ResetToSource()
+{
+	UObject* Outer = GetInstance()->GetOuter();
+	Instance = nullptr; // Clear out our reference to this object first..
+	Instance = GetSource()->MakeRecursiveDeepCopy(Outer);
+}
+
+void CopyParameterValues(UNiagaraScript* Script, UNiagaraScript* PreviousScript)
+{
+	for (FNiagaraVariable& InputParameter : Script->Parameters.Parameters)
+	{
+		for (FNiagaraVariable& PreviousInputParameter : PreviousScript->Parameters.Parameters)
+		{
+			if (PreviousInputParameter.IsDataAllocated())
+			{
+				if (InputParameter.GetName() == PreviousInputParameter.GetName() &&
+					InputParameter.GetType() == PreviousInputParameter.GetType())
+				{
+					InputParameter.AllocateData();
+					PreviousInputParameter.CopyTo(InputParameter.GetData());
+				}
+			}
+		}
+	}
+	
+	for (FNiagaraScriptDataInterfaceInfo& InputInfo : Script->DataInterfaceInfo)
+	{
+		for (FNiagaraScriptDataInterfaceInfo&  PreviousInputInfo : PreviousScript->DataInterfaceInfo)
+		{
+			if (InputInfo.Name == PreviousInputInfo.Name &&
+				InputInfo.DataInterface->GetClass() == PreviousInputInfo.DataInterface->GetClass())
+			{
+				PreviousInputInfo.CopyTo(&InputInfo, Script->GetOuter());
+			}
+		}
+	}
+}
+
+bool FNiagaraEmitterHandle::IsSynchronizedWithSource() const
+{
+	if (Instance && Source && Source->ChangeId.IsValid() && Instance->ChangeId.IsValid())
+	{
+		return Instance->ChangeId == Source->ChangeId;		
+	}
+	return false;
+}
+
+bool FNiagaraEmitterHandle::NeedsRecompile() const
+{
+	TArray<UNiagaraScript*> Scripts;
+	Instance->GetScripts(Scripts);
+
+	for (UNiagaraScript* Script : Scripts)
+	{
+		if (!Script->AreScriptAndSourceSynchronized())
+		{
+			return true;
+		}
+	}
+	return false;
+}
+
+bool FNiagaraEmitterHandle::RefreshFromSource()
+{
+	// TODO: Update this to support events.
+	UNiagaraScript* PreviousInstanceSpawnScript = Instance->SpawnScriptProps.Script;
+	UNiagaraScript* PreviousInstanceUpdateScript = Instance->UpdateScriptProps.Script;
+	TMap<const UObject*, UObject*> ExistingConversions;
+
+	if (Source == nullptr)
+	{
+		return false;
+	}
+
+	// First we need to deep copy the graph, which is shared amongst children
+	UNiagaraScriptSourceBase* GraphSource = Source->GraphSource->MakeRecursiveDeepCopy(Instance, ExistingConversions);
+	ExistingConversions.Add(Source->GraphSource, GraphSource);
+
+	// The script graphs will look up from ExistingConversions
+	UNiagaraScript* NewSpawnScript = Source->SpawnScriptProps.Script->MakeRecursiveDeepCopy(Instance, ExistingConversions);
+	UNiagaraScript* NewUpdateScript = Source->UpdateScriptProps.Script->MakeRecursiveDeepCopy(Instance, ExistingConversions);
+
+	if (NewSpawnScript->GetLastCompileStatus() == ENiagaraScriptCompileStatus::NCS_UpToDate && NewUpdateScript->GetLastCompileStatus() == ENiagaraScriptCompileStatus::NCS_UpToDate)
+	{
+		Instance->GraphSource = GraphSource;
+		Instance->ChangeId = Source->ChangeId;
+		Instance->SpawnScriptProps.Script = NewSpawnScript;
+		Instance->UpdateScriptProps.Script = NewUpdateScript;
+		
+		// We may have overridden the interpolated spawn script value from the default. If so, this requires
+		// a recompile.
+		bool bNeedRecompile = false;
+		if ((int32)Instance->SpawnScriptProps.Script->IsInterpolatedParticleSpawnScript() != Instance->bInterpolatedSpawning)
+		{
+			Instance->SpawnScriptProps.Script->SetUsage(ENiagaraScriptUsage::ParticleSpawnScriptInterpolated);
+			bNeedRecompile = true;
+		}
+
+		int32 NumDifference = Source->EventHandlerScriptProps.Num() - Instance->EventHandlerScriptProps.Num();
+		if (NumDifference < 0)
+		{
+			Instance->EventHandlerScriptProps.SetNum(Source->EventHandlerScriptProps.Num());
+		}
+		else if (NumDifference > 0)
+		{
+			Instance->EventHandlerScriptProps.AddDefaulted(NumDifference);
+		}
+		check(Source->EventHandlerScriptProps.Num() == Instance->EventHandlerScriptProps.Num());
+
+		for (int32 i = 0; i < Source->EventHandlerScriptProps.Num(); i++)
+		{
+			if (Instance->EventHandlerScriptProps[i].Script == nullptr)
+			{
+				Instance->EventHandlerScriptProps[i] = Source->EventHandlerScriptProps[i];
+				Instance->EventHandlerScriptProps[i].Script = Source->EventHandlerScriptProps[i].Script->MakeRecursiveDeepCopy(Instance, ExistingConversions);
+				check(Instance->EventHandlerScriptProps[i].Script->GetSource() == Instance->GraphSource);
+
+			}
+			else if (Source->EventHandlerScriptProps[i].Script != nullptr && Source->GraphSource == Source->EventHandlerScriptProps[i].Script->GetSource())
+			{
+				// Make sure that we copy all required variables and update our bytecode too.
+				Instance->EventHandlerScriptProps[i].Script = Source->EventHandlerScriptProps[i].Script->MakeRecursiveDeepCopy(Instance, ExistingConversions);
+				check(Instance->EventHandlerScriptProps[i].Script->GetSource() == Instance->GraphSource);
+			}
+		}
+
+		if (bNeedRecompile)
+		{
+			TArray<ENiagaraScriptCompileStatus> OutScriptStatuses;
+			TArray<FString> OutGraphLevelErrorMessages;
+			TArray<FString> PathNames;
+			TArray<UNiagaraScript*> Scripts;
+			Instance->CompileScripts(OutScriptStatuses, OutGraphLevelErrorMessages, PathNames, Scripts);
+
+			//ensure that we are now sync'ed with source
+			Instance->ChangeId = Source->ChangeId;
+		}
+
+		// We wait until after the potential recompile above to copy the old values over because
+		// the script parameters array is also written by the CompileScripts call.
+		CopyParameterValues(Instance->SpawnScriptProps.Script, PreviousInstanceSpawnScript);
+		CopyParameterValues(Instance->UpdateScriptProps.Script, PreviousInstanceUpdateScript);
+		check(Instance->SpawnScriptProps.Script->GetSource() == Instance->GraphSource);
+		check(Instance->UpdateScriptProps.Script->GetSource() == Instance->GraphSource);
+
+		UE_LOG(LogNiagara, Log, TEXT("Successful refresh from source %s"), *Instance->GetFullName());
+		check((int32)Instance->SpawnScriptProps.Script->IsInterpolatedParticleSpawnScript() == Instance->bInterpolatedSpawning);
+
+		Instance->SpawnScriptProps.InitDataSetAccess();
+		Instance->UpdateScriptProps.InitDataSetAccess();
+		for (int32 j = 0; j < Instance->EventHandlerScriptProps.Num(); j++)
+		{
+			Instance->EventHandlerScriptProps[j].InitDataSetAccess();
+		}
+		
+		return true;
+	}
+	else
+	{
+		UE_LOG(LogNiagara, Warning, TEXT("Failed to refresh from source %s"), *Instance->GetFullName());
+	}
+	return false;
+}
+#endif
\ No newline at end of file
diff --git a/Engine/Plugins/FX/Niagara/Source/Niagara/Private/NiagaraEmitterInstance.cpp b/Engine/Plugins/FX/Niagara/Source/Niagara/Private/NiagaraEmitterInstance.cpp
new file mode 100644
--- /dev/null
+++ b/Engine/Plugins/FX/Niagara/Source/Niagara/Private/NiagaraEmitterInstance.cpp
@@ -0,0 +1,1017 @@
+// Copyright 1998-2017 Epic Games, Inc. All Rights Reserved.
+
+#include "NiagaraEmitterInstance.h"
+#include "Materials/Material.h"
+#include "VectorVM.h"
+#include "NiagaraStats.h"
+#include "NiagaraConstants.h"
+#include "NiagaraRenderer.h"
+#include "NiagaraSystemInstance.h"
+#include "NiagaraDataInterface.h"
+#include "NiagaraEmitterInstanceBatcher.h"
+#include "NiagaraScriptExecutionContext.h"
+#include "NiagaraParameterCollection.h"
+#include "NiagaraWorldManager.h"
+
+DECLARE_DWORD_COUNTER_STAT(TEXT("Num Custom Events"), STAT_NiagaraNumCustomEvents, STATGROUP_Niagara);
+
+//DECLARE_CYCLE_STAT(TEXT("Tick"), STAT_NiagaraTick, STATGROUP_Niagara);
+DECLARE_CYCLE_STAT(TEXT("Simulate"), STAT_NiagaraSimulate, STATGROUP_Niagara);
+DECLARE_CYCLE_STAT(TEXT("Spawn"), STAT_NiagaraSpawn, STATGROUP_Niagara);
+DECLARE_CYCLE_STAT(TEXT("Spawn"), STAT_NiagaraEvents, STATGROUP_Niagara);
+DECLARE_CYCLE_STAT(TEXT("Kill"), STAT_NiagaraKill, STATGROUP_Niagara);
+DECLARE_CYCLE_STAT(TEXT("Event Handling"), STAT_NiagaraEventHandle, STATGROUP_Niagara);
+
+static int32 GbDumpParticleData = 0;
+static FAutoConsoleVariableRef CVarNiagaraDumpParticleData(
+	TEXT("fx.DumpParticleData"),
+	GbDumpParticleData,
+	TEXT("If > 0 current frame particle data will be dumped after simulation. \n"),
+	ECVF_Default
+	);
+
+//////////////////////////////////////////////////////////////////////////
+
+FNiagaraEmitterInstance::FNiagaraEmitterInstance(FNiagaraSystemInstance* InParentSystemInstance)
+: bError(false)
+, CPUTimeMS(0.0f)
+, ExecutionState(ENiagaraExecutionState::Active)
+, CachedBounds(ForceInit)
+, ParentSystemInstance(InParentSystemInstance)
+{
+	bDumpAfterEvent = false;
+}
+
+FNiagaraEmitterInstance::~FNiagaraEmitterInstance()
+{
+	//UE_LOG(LogNiagara, Warning, TEXT("~Simulator %p"), this);
+	ClearRenderer();
+	CachedBounds.Init();
+	UnbindParameters();
+}
+
+void FNiagaraEmitterInstance::ClearRenderer()
+{
+	for (int32 i = 0; i < EmitterRenderer.Num(); i++)
+	{
+		if (EmitterRenderer[i])
+		{
+			//UE_LOG(LogNiagara, Warning, TEXT("ClearRenderer %p"), EmitterRenderer);
+			// This queues up the renderer for deletion on the render thread..
+			EmitterRenderer[i]->Release();
+			EmitterRenderer[i] = nullptr;
+		}
+	}
+}
+
+void FNiagaraEmitterInstance::Init(int32 InEmitterIdx, FName InSystemInstanceName)
+{
+	EmitterIdx = InEmitterIdx;
+	OwnerSystemInstanceName = InSystemInstanceName;
+	Data = FNiagaraDataSet(FNiagaraDataSetID(GetEmitterHandle().GetIdName(), ENiagaraDataSetType::ParticleData));
+
+	//Init the spawn infos to the correct number for this system.
+	const TArray<FNiagaraEmitterSpawnAttributes>& EmitterSpawnInfoAttrs = ParentSystemInstance->GetSystem()->GetEmitterSpawnAttributes();
+	if (EmitterSpawnInfoAttrs.IsValidIndex(EmitterIdx))
+	{
+		SpawnInfos.SetNum(EmitterSpawnInfoAttrs[EmitterIdx].SpawnAttributes.Num());
+	}
+}
+
+void FNiagaraEmitterInstance::ResetSimulation()
+{
+	Data.ResetNumInstances();
+	Age = 0;
+	Loops = 0;
+	CollisionBatch.Reset();
+	bError = false;
+
+	const UNiagaraEmitter* PinnedProps = GetEmitterHandle().GetInstance();
+
+	SetExecutionState(ENiagaraExecutionState::Active);
+
+	if (!PinnedProps)
+	{
+		UE_LOG(LogNiagara, Error, TEXT("Unknown Error creating Niagara Simulation. Properties were null."));
+		bError = true;
+		return;
+	}
+
+	//Check for various failure conditions and bail.
+	if (!PinnedProps->UpdateScriptProps.Script || !PinnedProps->SpawnScriptProps.Script)
+	{
+		//TODO - Arbitrary named scripts. Would need some base functionality for Spawn/Udpate to be called that can be overriden in BPs for emitters with custom scripts.
+		UE_LOG(LogNiagara, Error, TEXT("Emitter cannot be enabled because it's doesn't have both an update and spawn script."), *PinnedProps->GetFullName());
+		bError = true;
+		return;
+	}
+	if (PinnedProps->UpdateScriptProps.Script->ByteCode.Num() == 0 && PinnedProps->SpawnScriptProps.Script->ByteCode.Num() == 0 && PinnedProps->SimTarget==ENiagaraSimTarget::CPUSim)
+	{
+		UE_LOG(LogNiagara, Error, TEXT("Emitter cannot be enabled because it's spawn or update script was not compiled correctly. %s"), *PinnedProps->GetFullName());
+		bError = true;
+		return;
+	}
+
+	if (PinnedProps->SpawnScriptProps.Script->DataUsage.bReadsAttriubteData)
+	{
+		UE_LOG(LogNiagara, Error, TEXT("%s reads attribute data and so cannot be used as a spawn script. The data being read would be invalid."), *PinnedProps->SpawnScriptProps.Script->GetName());
+		bError = true;
+		return;
+	}
+	if (PinnedProps->UpdateScriptProps.Script->Attributes.Num() == 0 || PinnedProps->SpawnScriptProps.Script->Attributes.Num() == 0)
+	{
+		UE_LOG(LogNiagara, Error, TEXT("This emitter cannot be enabled because it's spawn or update script doesn't have any attriubtes.."));
+		bError = true;
+		return;
+	}
+}
+
+void FNiagaraEmitterInstance::DirtyDataInterfaces()
+{
+	// Make sure that our function tables need to be regenerated...
+	SpawnExecContext.DirtyDataInterfaces();
+	UpdateExecContext.DirtyDataInterfaces();
+	for (FNiagaraScriptExecutionContext& EventContext : EventExecContexts)
+	{
+		EventContext.DirtyDataInterfaces();
+	}
+}
+
+void FNiagaraEmitterInstance::ReInitSimulation()
+{
+	const FNiagaraEmitterHandle& EmitterHandle = GetEmitterHandle();
+	bIsEnabled = EmitterHandle.GetIsEnabled();
+
+	ResetSimulation();
+
+	Data.Reset();
+	DataSetMap.Empty();
+
+	const UNiagaraEmitter* PinnedProps = EmitterHandle.GetInstance();
+
+	//Add the particle data to the data set map.
+	//Currently just used for the tick loop but will also allow access directly to the particle data from other emitters.
+	DataSetMap.Add(Data.GetID(), &Data);
+	//Warn the user if there are any attributes used in the update script that are not initialized in the spawn script.
+	//TODO: We need some window in the System editor and possibly the graph editor for warnings and errors.
+
+	const bool bVerboseAttributeLogging = false;
+
+	if (bVerboseAttributeLogging)
+	{
+		for (FNiagaraVariable& Attr : PinnedProps->UpdateScriptProps.Script->Attributes)
+		{
+			int32 FoundIdx;
+			if (!PinnedProps->SpawnScriptProps.Script->Attributes.Find(Attr, FoundIdx))
+			{
+				UE_LOG(LogNiagara, Warning, TEXT("Attribute %s is used in the Update script for %s but it is not initialised in the Spawn script!"), *Attr.GetName().ToString(), *EmitterHandle.GetName().ToString());
+			}
+			for (int32 i = 0; i < PinnedProps->EventHandlerScriptProps.Num(); i++)
+			{
+				if (PinnedProps->EventHandlerScriptProps[i].Script && !PinnedProps->EventHandlerScriptProps[i].Script->Attributes.Find(Attr, FoundIdx))
+				{
+					UE_LOG(LogNiagara, Warning, TEXT("Attribute %s is used in the event handler script for %s but it is not initialised in the Spawn script!"), *Attr.GetName().ToString(), *EmitterHandle.GetName().ToString());
+				}
+			}
+		}
+	}
+	Data.AddVariables(PinnedProps->UpdateScriptProps.Script->Attributes);
+	Data.AddVariables(PinnedProps->SpawnScriptProps.Script->Attributes);
+	Data.Finalize();
+
+	CollisionBatch.Init(ParentSystemInstance->GetIDName(), EmitterHandle.GetIdName());
+
+	UpdateScriptEventDataSets.Empty();
+	for (const FNiagaraEventGeneratorProperties &GeneratorProps : PinnedProps->UpdateScriptProps.EventGenerators)
+	{
+		FNiagaraDataSet *Set = FNiagaraEventDataSetMgr::CreateEventDataSet(ParentSystemInstance->GetIDName(), EmitterHandle.GetIdName(), GeneratorProps.SetProps.ID.Name);
+		Set->Reset();
+		Set->AddVariables(GeneratorProps.SetProps.Variables);
+		Set->Finalize();
+		UpdateScriptEventDataSets.Add(Set);
+	}
+
+	SpawnScriptEventDataSets.Empty();
+	for (const FNiagaraEventGeneratorProperties &GeneratorProps : PinnedProps->SpawnScriptProps.EventGenerators)
+	{
+		FNiagaraDataSet *Set = FNiagaraEventDataSetMgr::CreateEventDataSet(ParentSystemInstance->GetIDName(), EmitterHandle.GetIdName(), GeneratorProps.SetProps.ID.Name);
+		Set->Reset();
+		Set->AddVariables(GeneratorProps.SetProps.Variables);
+		Set->Finalize();
+		SpawnScriptEventDataSets.Add(Set);
+	}
+
+	SpawnExecContext.Init(PinnedProps->SpawnScriptProps.Script, PinnedProps->SimTarget);
+	UpdateExecContext.Init(PinnedProps->UpdateScriptProps.Script, PinnedProps->SimTarget);
+
+	EventExecContexts.SetNum(PinnedProps->EventHandlerScriptProps.Num());
+	int32 NumEvents = PinnedProps->EventHandlerScriptProps.Num();
+	for (int32 i = 0; i < NumEvents; i++)
+	{
+		UNiagaraScript* EventScript = PinnedProps->EventHandlerScriptProps[i].Script;
+		if (EventScript)
+		{
+			if (EventScript->ByteCode.Num() == 0)
+			{
+				UE_LOG(LogNiagara, Error, TEXT("%s has an event handler script hat didn't compile correctly."), *GetEmitterHandle().GetName().ToString());
+				bError = true;
+				return;
+			}
+		}
+
+		//This is cpu explicitly? Are we doing event handlers on GPU?
+		EventExecContexts[i].Init(EventScript, ENiagaraSimTarget::CPUSim);
+	}
+
+	UNiagaraEmitter* Emitter = GetEmitterHandle().GetInstance();
+
+	//Setup direct bindings for setting parameter values.
+	SpawnIntervalBinding.Init(SpawnExecContext.Parameters, Emitter->GetEmitterParameter(SYS_PARAM_EMITTER_SPAWN_INTERVAL));
+	InterpSpawnStartBinding.Init(SpawnExecContext.Parameters, Emitter->GetEmitterParameter(SYS_PARAM_EMITTER_INTERP_SPAWN_START_DT));
+
+	FNiagaraVariable EmitterAgeParam = Emitter->GetEmitterParameter(SYS_PARAM_EMITTER_AGE);
+	SpawnEmitterAgeBinding.Init(SpawnExecContext.Parameters, EmitterAgeParam);
+	UpdateEmitterAgeBinding.Init(UpdateExecContext.Parameters, EmitterAgeParam);
+	EventEmitterAgeBindings.SetNum(NumEvents);
+	for (int32 i = 0; i < NumEvents; i++)
+	{
+		EventEmitterAgeBindings[i].Init(EventExecContexts[i].Parameters, EmitterAgeParam);
+	}
+
+	SpawnExecCountBinding.Init(SpawnExecContext.Parameters, SYS_PARAM_ENGINE_EXEC_COUNT);
+	UpdateExecCountBinding.Init(UpdateExecContext.Parameters, SYS_PARAM_ENGINE_EXEC_COUNT);
+	EventExecCountBindings.SetNum(NumEvents);
+	for (int32 i = 0; i < NumEvents; i++)
+	{
+		EventExecCountBindings[i].Init(EventExecContexts[i].Parameters, SYS_PARAM_ENGINE_EXEC_COUNT);
+	}
+
+	if (PinnedProps->SimTarget == ENiagaraSimTarget::GPUComputeSim)
+	{
+		//Just ensure we've generated the singleton here on the GT as it throws a wobbler if we do this later in parallel.
+		NiagaraEmitterInstanceBatcher::Get();
+	}
+}
+
+//Unsure on usage of this atm. Possibly useful in future.
+// void FNiagaraEmitterInstance::RebindParameterCollection(UNiagaraParameterCollectionInstance* OldInstance, UNiagaraParameterCollectionInstance* NewInstance)
+// {
+// 	OldInstance->GetParameterStore().Unbind(&SpawnExecContext.Parameters);
+// 	NewInstance->GetParameterStore().Bind(&SpawnExecContext.Parameters);
+// 
+// 	OldInstance->GetParameterStore().Unbind(&UpdateExecContext.Parameters);
+// 	NewInstance->GetParameterStore().Bind(&UpdateExecContext.Parameters);
+// 
+// 	for (FNiagaraScriptExecutionContext& EventContext : EventExecContexts)
+// 	{
+// 		OldInstance->GetParameterStore().Unbind(&EventContext.Parameters);
+// 		NewInstance->GetParameterStore().Bind(&EventContext.Parameters);
+// 	}
+// }
+
+void FNiagaraEmitterInstance::UnbindParameters()
+{
+	FNiagaraWorldManager* WorldMan = ParentSystemInstance->GetWorldManager();
+	check(WorldMan);
+
+	//Unbind our parameter collections.
+	for (UNiagaraParameterCollection* Collection : SpawnExecContext.Script->ParameterCollections)
+	{
+		WorldMan->GetParameterCollection(Collection)->GetParameterStore().Unbind(&SpawnExecContext.Parameters);
+	}
+	for (UNiagaraParameterCollection* Collection : UpdateExecContext.Script->ParameterCollections)
+	{
+		WorldMan->GetParameterCollection(Collection)->GetParameterStore().Unbind(&UpdateExecContext.Parameters);
+	}
+
+	for (int32 EventIdx = 0; EventIdx < EventExecContexts.Num(); ++EventIdx)
+	{
+		for (UNiagaraParameterCollection* Collection : EventExecContexts[EventIdx].Script->ParameterCollections)
+		{
+			WorldMan->GetParameterCollection(Collection)->GetParameterStore().Unbind(&EventExecContexts[EventIdx].Parameters);
+		}
+	}
+
+	FNiagaraParameterStore& SystemParams = ParentSystemInstance->GetParameters();
+	SystemParams.Unbind(&SpawnExecContext.Parameters);
+	SystemParams.Unbind(&UpdateExecContext.Parameters);
+	for (FNiagaraScriptExecutionContext& EventContext : EventExecContexts)
+	{
+		SystemParams.Unbind(&EventContext.Parameters);
+	}
+}
+
+void FNiagaraEmitterInstance::BindParameters()
+{
+	FNiagaraWorldManager* WorldMan = ParentSystemInstance->GetWorldManager();
+	check(WorldMan);
+
+	for (UNiagaraParameterCollection* Collection : SpawnExecContext.Script->ParameterCollections)
+	{
+		WorldMan->GetParameterCollection(Collection)->GetParameterStore().Bind(&SpawnExecContext.Parameters);
+	}
+	for (UNiagaraParameterCollection* Collection : UpdateExecContext.Script->ParameterCollections)
+	{
+		WorldMan->GetParameterCollection(Collection)->GetParameterStore().Bind(&UpdateExecContext.Parameters);
+	}
+
+	for (int32 EventIdx = 0; EventIdx < EventExecContexts.Num(); ++EventIdx)
+	{
+		for (UNiagaraParameterCollection* Collection : EventExecContexts[EventIdx].Script->ParameterCollections)
+		{
+			WorldMan->GetParameterCollection(Collection)->GetParameterStore().Bind(&EventExecContexts[EventIdx].Parameters);
+		}
+	}
+
+	//Now bind parameters from the component and system.
+	FNiagaraParameterStore& InstanceParams = ParentSystemInstance->GetParameters();
+	InstanceParams.Bind(&SpawnExecContext.Parameters);
+	InstanceParams.Bind(&UpdateExecContext.Parameters);
+	for (FNiagaraScriptExecutionContext& EventContext : EventExecContexts)
+	{
+		InstanceParams.Bind(&EventContext.Parameters);
+	}
+}
+
+void FNiagaraEmitterInstance::PostResetSimulation()
+{
+	const FNiagaraEmitterHandle& EmitterHandle = GetEmitterHandle();
+	if (!bError)
+	{
+		check(ParentSystemInstance);
+		const UNiagaraEmitter* Props = EmitterHandle.GetInstance();
+
+		//Go through all our receivers and grab their generator sets so that the source emitters can do any init work they need to do.
+		for (const FNiagaraEventReceiverProperties& Receiver : Props->SpawnScriptProps.EventReceivers)
+		{
+			//FNiagaraDataSet* ReceiverSet = ParentSystemInstance->GetDataSet(FNiagaraDataSetID(Receiver.SourceEventGenerator, ENiagaraDataSetType::Event), Receiver.SourceEmitter);
+			const FNiagaraDataSet* ReceiverSet = FNiagaraEventDataSetMgr::GetEventDataSet(ParentSystemInstance->GetIDName(), Receiver.SourceEmitter, Receiver.SourceEventGenerator);
+
+		}
+
+		for (const FNiagaraEventReceiverProperties& Receiver : Props->UpdateScriptProps.EventReceivers)
+		{
+			//FNiagaraDataSet* ReceiverSet = ParentSystemInstance->GetDataSet(FNiagaraDataSetID(Receiver.SourceEventGenerator, ENiagaraDataSetType::Event), Receiver.SourceEmitter);
+			const FNiagaraDataSet* ReceiverSet = FNiagaraEventDataSetMgr::GetEventDataSet(ParentSystemInstance->GetIDName(), Receiver.SourceEmitter, Receiver.SourceEventGenerator);
+		}
+
+		// add the collision event set
+		// TODO: add spawn and death event data sets (if enabled)
+		if (Props->CollisionMode != ENiagaraCollisionMode::None)
+		{
+			CollisionBatch.Init(ParentSystemInstance->GetIDName(), EmitterHandle.GetIdName());	// creates and sets up the data set for the events
+		}
+	}
+}
+
+FNiagaraDataSet* FNiagaraEmitterInstance::GetDataSet(FNiagaraDataSetID SetID)
+{
+	FNiagaraDataSet** SetPtr = DataSetMap.Find(SetID);
+	FNiagaraDataSet* Ret = NULL;
+	if (SetPtr)
+	{
+		Ret = *SetPtr;
+	}
+	else
+	{
+		// TODO: keep track of data sets generated by the scripts (event writers) and find here
+	}
+
+	return Ret;
+}
+
+const FNiagaraEmitterHandle& FNiagaraEmitterInstance::GetEmitterHandle() const
+{
+	return ParentSystemInstance->GetSystem()->GetEmitterHandles()[EmitterIdx];
+}
+
+float FNiagaraEmitterInstance::GetTotalCPUTime()
+{
+	float Total = CPUTimeMS;
+	for (int32 i = 0; i < EmitterRenderer.Num(); i++)
+	{
+		if (EmitterRenderer[i])
+		{
+			Total += EmitterRenderer[i]->GetCPUTimeMS();
+
+		}
+	}
+
+	return Total;
+}
+
+int FNiagaraEmitterInstance::GetTotalBytesUsed()
+{
+	int32 BytesUsed = Data.GetSizeBytes();
+	/*
+	for (FNiagaraDataSet& Set : DataSets)
+	{
+		BytesUsed += Set.GetSizeBytes();
+	}
+	*/
+	return BytesUsed;
+}
+
+
+/** Look for dead particles and move from the end of the list to the dead location, compacting in the process
+  * Also calculates bounds; Kill will be removed from this once we do conditional write
+  */
+void FNiagaraEmitterInstance::PostProcessParticles()
+{
+	SCOPE_CYCLE_COUNTER(STAT_NiagaraKill);
+	int32 OrigNumParticles = Data.GetNumInstances();
+	int32 CurNumParticles = OrigNumParticles;
+
+	CachedBounds.Init();
+
+	const FNiagaraEmitterHandle& EmitterHandle = GetEmitterHandle();
+	if (CurNumParticles == 0 || EmitterHandle.GetInstance()->SimTarget == ENiagaraSimTarget::GPUComputeSim)
+	{
+		return;
+	}
+
+	FNiagaraDataSetIterator<FVector> PosItr(Data, FNiagaraVariable(FNiagaraTypeDefinition::GetVec3Def(), TEXT("Position")));
+	FNiagaraDataSetIterator<FVector2D> SizeItr(Data, FNiagaraVariable(FNiagaraTypeDefinition::GetVec2Def(), TEXT("Size")));
+	FNiagaraDataSetIterator<FVector> MeshScaleItr(Data, FNiagaraVariable(FNiagaraTypeDefinition::GetVec3Def(), TEXT("Scale")));
+
+	FVector MaxSize(EForceInit::ForceInitToZero);
+
+	while (PosItr.IsValid())
+	{
+		// Only increment the iterators if we don't kill a particle because
+		// we swap with the last particle and it too may have aged out, so we want
+		// to keep looping on the same index as long as it gets swapped out for a 
+		// dead particle.
+
+		FVector Position;
+		PosItr.Get(Position);
+
+		// Some graphs have a tendency to divide by zero. This ContainsNaN has been added prophylactically
+		// to keep us safe during GDC. It should be removed as soon as we feel safe that scripts are appropriately warned.
+		if (Position.ContainsNaN() == false)
+		{
+			CachedBounds += Position;
+
+			// We advance the scale or size depending of if we use either.
+			if (MeshScaleItr.IsValid())
+			{
+				MaxSize = MaxSize.ComponentMax(*MeshScaleItr);
+				MeshScaleItr.Advance();
+			}
+			else if (SizeItr.IsValid())
+			{
+				MaxSize = MaxSize.ComponentMax(FVector((*SizeItr).GetMax()));
+				SizeItr.Advance();
+			}
+			// Now we advance our main iterator since we've safely handled this particle.
+			PosItr.Advance();
+		}
+		else
+		{
+			//Must always advance otherwise we'll have inf loop if we have nans in the position
+			if (MeshScaleItr.IsValid())
+			{
+				MeshScaleItr.Advance();
+			}
+			else if (SizeItr.IsValid())
+			{
+				SizeItr.Advance();
+			}
+			PosItr.Advance();
+		}
+	}
+
+	CachedBounds = CachedBounds.ExpandBy(MaxSize);
+}
+
+/** 
+  * PreTick - handles killing dead particles, emitter death, and buffer swaps
+  */
+void FNiagaraEmitterInstance::PreTick()
+{
+	const FNiagaraEmitterHandle& EmitterHandle = GetEmitterHandle();
+	const UNiagaraEmitter* PinnedProps = EmitterHandle.GetInstance();
+
+	if (!PinnedProps || !bIsEnabled || bError || ExecutionState == ENiagaraExecutionState::Dead)
+	{
+		return;
+	}
+
+	bool bOk = true;
+	bOk &= SpawnExecContext.Tick(ParentSystemInstance);
+	bOk &= UpdateExecContext.Tick(ParentSystemInstance);
+	for (FNiagaraScriptExecutionContext& EventContext : EventExecContexts)
+	{
+		bOk &= EventContext.Tick(ParentSystemInstance);
+	}
+
+	if (!bOk)
+	{
+		ResetSimulation();
+		bError = true;
+		return;
+	}
+
+	if (Data.GetNumInstances() == 0)
+	{
+		return;
+	}
+
+	check(Data.GetNumVariables() > 0);
+	check(PinnedProps->SpawnScriptProps.Script);
+	check(PinnedProps->UpdateScriptProps.Script);
+
+	// generate events from collisions
+	if (PinnedProps->CollisionMode == ENiagaraCollisionMode::SceneGeometry)
+	{
+		CollisionBatch.GenerateEventsFromResults(this);
+	}
+
+
+	//Swap all data set buffers before doing the main tick on any simulation.
+	if (PinnedProps->SimTarget == ENiagaraSimTarget::CPUSim)
+	{
+		for (TPair<FNiagaraDataSetID, FNiagaraDataSet*> SetPair : DataSetMap)
+		{
+			SetPair.Value->Tick(PinnedProps->SimTarget);
+		}
+		CollisionBatch.Tick(PinnedProps->SimTarget);
+
+		for (FNiagaraDataSet* Set : UpdateScriptEventDataSets)
+		{
+			Set->Tick(PinnedProps->SimTarget);
+		}
+
+		for (FNiagaraDataSet* Set : SpawnScriptEventDataSets)
+		{
+			Set->Tick(PinnedProps->SimTarget);
+		}
+	}
+}
+
+
+void FNiagaraEmitterInstance::Tick(float DeltaSeconds)
+{
+	SCOPE_CYCLE_COUNTER(STAT_NiagaraTick);
+	SimpleTimer TickTime;
+
+	const FNiagaraEmitterHandle& EmitterHandle = GetEmitterHandle();
+	const UNiagaraEmitter* PinnedProps = EmitterHandle.GetInstance();
+	if (!PinnedProps || !bIsEnabled || bError || ExecutionState == ENiagaraExecutionState::Dead)
+	{
+		return;
+	}
+
+	Age += DeltaSeconds;
+
+	check(Data.GetNumVariables() > 0);
+	check(PinnedProps->SpawnScriptProps.Script);
+	check(PinnedProps->UpdateScriptProps.Script);
+	
+	//TickEvents(DeltaSeconds);
+
+	// add system constants
+	{
+		SCOPE_CYCLE_COUNTER(STAT_NiagaraConstants);
+		float InvDT = 1.0f / DeltaSeconds;
+
+		//TODO: Create a binding helper object for these to avoid the search.
+		SpawnEmitterAgeBinding.SetValue(Age);
+		UpdateEmitterAgeBinding.SetValue(Age);
+		for (FNiagaraParameterDirectBinding<float>& Binding : EventEmitterAgeBindings)
+		{
+			Binding.SetValue(Age);
+		}
+	}
+	
+	// Calculate number of new particles from regular spawning 
+	int32 OrigNumParticles = Data.GetPrevNumInstances();
+
+	uint32 SpawnTotal = 0;
+	for (FNiagaraSpawnInfo& Info : SpawnInfos)
+	{
+		SpawnTotal += Info.Count;
+	}
+
+	// Calculate number of new particles from all event related spawns
+	TArray<TArray<int32>> EventSpawnCounts;
+	EventSpawnCounts.AddDefaulted(PinnedProps->EventHandlerScriptProps.Num());
+	TArray<int32> EventHandlerSpawnCounts;
+	EventHandlerSpawnCounts.AddDefaulted(PinnedProps->EventHandlerScriptProps.Num());
+	uint32 EventSpawnTotal = 0;
+	TArray<FNiagaraDataSet*> EventSet;
+	EventSet.AddZeroed(PinnedProps->EventHandlerScriptProps.Num());
+	TArray<FGuid> SourceEmitterGuid;
+	SourceEmitterGuid.AddDefaulted(PinnedProps->EventHandlerScriptProps.Num());
+	TArray<FName> SourceEmitterName;
+	SourceEmitterName.AddDefaulted(PinnedProps->EventHandlerScriptProps.Num());
+	TArray<bool> bPerformEventSpawning;
+	bPerformEventSpawning.AddDefaulted(PinnedProps->EventHandlerScriptProps.Num());
+
+	for (int32 i = 0; i < PinnedProps->EventHandlerScriptProps.Num(); i++)
+	{
+		const FNiagaraEventScriptProperties &EventHandlerProps = PinnedProps->EventHandlerScriptProps[i];
+		SourceEmitterGuid[i] = EventHandlerProps.SourceEmitterID;
+		SourceEmitterName[i] = SourceEmitterGuid[i].IsValid() ? *SourceEmitterGuid[i].ToString() : EmitterHandle.GetIdName();
+		EventSet[i] = FNiagaraEventDataSetMgr::GetEventDataSet(ParentSystemInstance->GetIDName(), SourceEmitterName[i], EventHandlerProps.SourceEventName);
+		bPerformEventSpawning[i] = (ExecutionState == ENiagaraExecutionState::Active && EventHandlerProps.Script && EventHandlerProps.ExecutionMode == EScriptExecutionMode::SpawnedParticles);
+		if (bPerformEventSpawning[i])
+		{
+			uint32 EventSpawnNum = CalculateEventSpawnCount(EventHandlerProps, EventSpawnCounts[i], EventSet[i]);
+			EventSpawnTotal += EventSpawnNum;
+			EventHandlerSpawnCounts[i] = EventSpawnNum;
+		}
+	}
+
+
+	/* GPU simulation -  we just create an FNiagaraComputeExecutionContext, queue it, and let the batcher take care of the rest
+	 */
+	if (PinnedProps->SimTarget == ENiagaraSimTarget::GPUComputeSim)	
+	{
+		FNiagaraComputeExecutionContext *ComputeContext = new FNiagaraComputeExecutionContext();
+		ComputeContext->MainDataSet = &Data;
+		ComputeContext->RTSpawnScript = PinnedProps->SpawnScriptProps.Script->GetRenderThreadScript();
+		ComputeContext->RTUpdateScript = PinnedProps->UpdateScriptProps.Script->GetRenderThreadScript();
+		ComputeContext->SpawnRateInstances = SpawnTotal;
+		ComputeContext->BurstInstances = 0;
+		ComputeContext->EventSpawnTotal = EventSpawnTotal;
+
+		ComputeContext->UpdateInterfaces = PinnedProps->UpdateScriptProps.Script->DataInterfaceInfo;
+
+		// copy over the constants for the render thread
+		//
+		int32 Size = UpdateExecContext.Parameters.GetExternalParameterSize();
+		if (Size)
+		{
+			ComputeContext->UpdateParams.AddUninitialized(FMath::DivideAndRoundUp(Size, 16) * 16);
+			FMemory::Memcpy(ComputeContext->UpdateParams.GetData(), UpdateExecContext.Parameters.GetParameterData(0), Size);
+		}
+		Size = SpawnExecContext.Parameters.GetExternalParameterSize();
+		if (Size)
+		{
+			ComputeContext->SpawnParams.AddUninitialized(FMath::DivideAndRoundUp(Size, 16) * 16);
+			FMemory::Memcpy(ComputeContext->SpawnParams.GetData(), SpawnExecContext.Parameters.GetParameterData(0), Size);
+		}
+
+		// push event data sets to the context
+		for (FNiagaraDataSet *Set : UpdateScriptEventDataSets)
+		{
+			ComputeContext->UpdateEventWriteDataSets.Add(Set);
+		}
+
+		ComputeContext->EventHandlerScriptProps = PinnedProps->EventHandlerScriptProps;
+		ComputeContext->EventSets = EventSet;
+		ComputeContext->EventSpawnCounts = EventHandlerSpawnCounts;
+		NiagaraEmitterInstanceBatcher::Get()->Queue(ComputeContext);
+
+		CachedBounds.Init();
+		CachedBounds = CachedBounds.ExpandBy(FVector(20.0, 20.0, 20.0));	// temp until GPU sims update bounds
+		return;
+	}
+
+	int32 AllocationSize = OrigNumParticles + SpawnTotal + EventSpawnTotal;
+	//Allocate space for prev frames particles and any new one's we're going to spawn.
+	Data.Allocate(AllocationSize, PinnedProps->SimTarget);
+
+
+
+	TArray<FNiagaraDataSetExecutionInfo> DataSetExecInfos;
+	DataSetExecInfos.Emplace(&Data, 0, false, true);
+
+	// Simulate existing particles forward by DeltaSeconds.
+	if ((ExecutionState != ENiagaraExecutionState::Dead && ExecutionState != ENiagaraExecutionState::Paused) /*&& OrigNumParticles > 0*/)
+	{
+		/*
+		if (bDumpAfterEvent)
+		{
+			Data.Dump(false);
+			bDumpAfterEvent = false;
+		}
+		*/
+
+		Data.SetNumInstances(OrigNumParticles);
+		SCOPE_CYCLE_COUNTER(STAT_NiagaraSimulate);
+
+		UpdateExecCountBinding.SetValue(OrigNumParticles);
+		DataSetExecInfos.SetNum(1, false);
+		DataSetExecInfos[0].StartInstance = 0;
+		for (FNiagaraDataSet* EventDataSet : UpdateScriptEventDataSets)
+		{
+			DataSetExecInfos.Emplace(EventDataSet, 0, true, false);
+		}
+		UpdateExecContext.Execute(OrigNumParticles, DataSetExecInfos);
+
+		if (GbDumpParticleData)
+		{
+			UE_LOG(LogNiagara, Log, TEXT("=== Updated %d Particles ==="), OrigNumParticles);
+			Data.Dump(true, 0, OrigNumParticles);
+		}
+	}
+
+	uint32 EventSpawnStart = Data.GetNumInstances();
+
+	//Init new particles with the spawn script.
+	if (ExecutionState==ENiagaraExecutionState::Active && SpawnTotal + EventSpawnTotal > 0)
+	{
+		SCOPE_CYCLE_COUNTER(STAT_NiagaraSpawn);
+
+		//Handle main spawn rate spawning
+		auto SpawnParticles = [&](int32 Num, FString DumpLabel)
+		{
+			if (Num > 0)
+			{
+				int32 OrigNum = Data.GetNumInstances();
+				Data.SetNumInstances(OrigNum + Num);
+
+				SpawnExecCountBinding.SetValue(Num);
+				DataSetExecInfos.SetNum(1, false);
+				DataSetExecInfos[0].StartInstance = OrigNum;
+				for (FNiagaraDataSet* EventDataSet : SpawnScriptEventDataSets)
+				{
+					DataSetExecInfos.Emplace(EventDataSet, OrigNum, true, false);
+				}
+				SpawnExecContext.Execute(Num, DataSetExecInfos);
+
+				if (GbDumpParticleData)
+				{
+					UE_LOG(LogNiagara, Log, TEXT("=== Spawned %d Particles ==="), Num);
+					Data.Dump(true, OrigNum, Num);
+				}
+			}
+		};
+
+		//Perform all our regular spawning that's driven by our emitter script.
+		for (FNiagaraSpawnInfo& Info : SpawnInfos)
+		{
+			SpawnIntervalBinding.SetValue(Info.IntervalDt);
+			InterpSpawnStartBinding.SetValue(Info.InterpStartDt);
+
+			SpawnParticles(Info.Count, TEXT("Regular Spawn"));
+		};
+
+		EventSpawnStart = Data.GetNumInstances();
+
+		for (int32 EventScriptIdx = 0; EventScriptIdx < PinnedProps->EventHandlerScriptProps.Num(); EventScriptIdx++)
+		{
+			//Spawn particles coming from events.
+			for (int32 i = 0; i < EventSpawnCounts[EventScriptIdx].Num(); i++)
+			{
+				int32 EventNumToSpawn = EventSpawnCounts[EventScriptIdx][i];
+
+				//Event spawns are instantaneous at the middle of the frame?
+				SpawnIntervalBinding.SetValue(0.0f);
+				InterpSpawnStartBinding.SetValue(DeltaSeconds * 0.5f);
+
+				SpawnParticles(EventNumToSpawn, TEXT("Event Spawn"));
+			}
+		}
+	}
+
+	// handle event based spawning
+	for (int32 EventScriptIdx = 0;  EventScriptIdx < PinnedProps->EventHandlerScriptProps.Num(); EventScriptIdx++)
+	{
+		const FNiagaraEventScriptProperties &EventHandlerProps = PinnedProps->EventHandlerScriptProps[EventScriptIdx];
+
+		if (bPerformEventSpawning[EventScriptIdx] && EventSet[EventScriptIdx] && EventSpawnCounts[EventScriptIdx].Num())
+		{
+			SCOPE_CYCLE_COUNTER(STAT_NiagaraEventHandle);
+			for (int32 i = 0; i < EventSpawnCounts[EventScriptIdx].Num(); i++)
+			{
+				int32 EventNumToSpawn = EventSpawnCounts[EventScriptIdx][i];
+				ensure(EventNumToSpawn + EventSpawnStart < Data.GetNumInstances());
+				EventExecCountBindings[EventScriptIdx].SetValue(EventNumToSpawn);
+
+				DataSetExecInfos.SetNum(1, false);
+				DataSetExecInfos[0].StartInstance = EventSpawnStart;
+				DataSetExecInfos.Emplace(EventSet[EventScriptIdx], i, false, false);
+				SpawnExecContext.Execute(EventNumToSpawn, DataSetExecInfos);
+
+				EventSpawnStart += EventNumToSpawn;
+			}
+		}
+
+
+		// handle all-particle events
+		if (EventHandlerProps.Script && EventHandlerProps.ExecutionMode == EScriptExecutionMode::EveryParticle && EventSet[EventScriptIdx])
+		{
+			if (EventSet[EventScriptIdx]->GetPrevNumInstances())
+			{
+				SCOPE_CYCLE_COUNTER(STAT_NiagaraEventHandle);
+
+				for (uint32 i = 0; i < EventSet[EventScriptIdx]->GetPrevNumInstances(); i++)
+				{
+					// If we have events, Swap buffers, to make sure we don't overwrite previous script results
+					// and copy prev to cur, because the event script isn't likely to write all attributes
+					// TODO: copy only event script's unused attributes
+					Data.Tick();
+					Data.CopyPrevToCur();
+
+					EventExecCountBindings[EventScriptIdx].SetValue(Data.GetPrevNumInstances());
+					DataSetExecInfos.SetNum(1, false);
+					DataSetExecInfos[0].StartInstance = 0;
+					DataSetExecInfos.Emplace(EventSet[EventScriptIdx], i, false, false);
+					SpawnExecContext.Execute(Data.GetPrevNumInstances(), DataSetExecInfos);
+
+					// TODO: Should we check to see if we wrote all particles here?
+				}
+			}
+		}
+
+
+		// handle single-particle events
+		// TODO: we'll need a way to either skip execution of the VM if an index comes back as invalid, or we'll have to pre-process
+		// event/particle arrays; this is currently a very naive (and comparatively slow) implementation, until full indexed reads work
+		if (EventHandlerProps.Script && EventHandlerProps.ExecutionMode == EScriptExecutionMode::SingleParticle && EventSet[EventScriptIdx])
+		{
+			SCOPE_CYCLE_COUNTER(STAT_NiagaraEventHandle);
+			FNiagaraVariable IndexVar(FNiagaraTypeDefinition::GetIntDef(), "ParticleIndex");
+			FNiagaraDataSetIterator<int32> IndexItr(*EventSet[EventScriptIdx], IndexVar, 0, false);
+			if (IndexItr.IsValid())
+			{
+				EventExecCountBindings[EventScriptIdx].SetValue(1);
+
+				for (uint32 i = 0; i < EventSet[EventScriptIdx]->GetPrevNumInstances(); i++)
+				{
+					int32 Index = *IndexItr;
+					IndexItr.Advance();
+					DataSetExecInfos.SetNum(1, false);
+					DataSetExecInfos[0].StartInstance = Index;
+					DataSetExecInfos.Emplace(EventSet[EventScriptIdx], i, false, false);
+					SpawnExecContext.Execute(1, DataSetExecInfos);
+
+					// TODO: Should we check to see if we wrote all particles here?
+				}
+			}
+		}
+	}
+
+
+	// kick off collision tests from this emitter
+	if (PinnedProps->CollisionMode == ENiagaraCollisionMode::SceneGeometry)
+	{
+		CollisionBatch.KickoffNewBatch(this, DeltaSeconds);
+	}
+
+	PostProcessParticles();
+
+	SpawnExecContext.PostTick();
+	UpdateExecContext.PostTick();
+	for (FNiagaraScriptExecutionContext& EventContext : EventExecContexts)
+	{
+		EventContext.PostTick();
+	}
+
+	CPUTimeMS = TickTime.GetElapsedMilliseconds();
+
+	INC_DWORD_STAT_BY(STAT_NiagaraNumParticles, Data.GetNumInstances());
+}
+
+
+/** Calculate total number of spawned particles from events; these all come from event handler script with the SpawnedParticles execution mode
+ *  We get the counts ahead of event processing time so we only have to allocate new particles once
+ *  TODO: augment for multiple spawning event scripts
+ */
+uint32 FNiagaraEmitterInstance::CalculateEventSpawnCount(const FNiagaraEventScriptProperties &EventHandlerProps, TArray<int32> &EventSpawnCounts, FNiagaraDataSet *EventSet)
+{
+	uint32 EventSpawnTotal = 0;
+	int32 NumEventsToProcess = 0;
+
+	if (EventSet)
+	{
+		NumEventsToProcess = EventSet->GetPrevNumInstances();
+		if(EventHandlerProps.MaxEventsPerFrame > 0)
+		{
+			NumEventsToProcess = FMath::Min<int32>(EventSet->GetPrevNumInstances(), EventHandlerProps.MaxEventsPerFrame);
+		}
+
+		for (int32 i = 0; i < NumEventsToProcess; i++)
+		{
+			if (ExecutionState == ENiagaraExecutionState::Active)
+			{
+				EventSpawnCounts.Add(EventHandlerProps.SpawnNumber);
+				EventSpawnTotal += EventHandlerProps.SpawnNumber;
+			}
+		}
+	}
+
+	return EventSpawnTotal;
+}
+
+void FNiagaraEmitterInstance::SetExecutionState(ENiagaraExecutionState InState)
+{
+	if (InState != ExecutionState)
+	{
+		const UEnum* EnumPtr = FNiagaraTypeDefinition::GetExecutionStateEnum();
+		UE_LOG(LogNiagara, Log, TEXT("Emitter \"%s\" change state: %s to %s"), *GetEmitterHandle().GetName().ToString(), *EnumPtr->GetNameStringByValue((int64)ExecutionState),
+			*EnumPtr->GetNameStringByValue((int64)InState));
+	}
+
+	if (InState == ENiagaraExecutionState::Active && ExecutionState == ENiagaraExecutionState::Inactive)
+	{
+		UE_LOG(LogNiagara, Log, TEXT("Emitter \"%s\" change state N O O O O O "), *GetEmitterHandle().GetName().ToString());
+	}
+	ExecutionState = InState;
+}
+
+
+#if WITH_EDITORONLY_DATA
+
+bool FNiagaraEmitterInstance::CheckAttributesForRenderer(int32 Index)
+{
+	const FNiagaraEmitterHandle& EmitterHandle = GetEmitterHandle();
+	if (Index > EmitterRenderer.Num())
+	{
+		return false;
+	}
+
+	bool bOk = true;
+	if (EmitterRenderer[Index])
+	{
+		const TArray<FNiagaraVariable>& RequiredAttrs = EmitterRenderer[Index]->GetRequiredAttributes();
+
+		for (FNiagaraVariable Attr : RequiredAttrs)
+		{
+			// TODO .. should we always be namespaced?
+			FString AttrName = Attr.GetName().ToString();
+			if (AttrName.RemoveFromStart(TEXT("Particles.")))
+			{
+				Attr.SetName(*AttrName);
+			}
+
+			if (!Data.HasVariable(Attr))
+			{
+				bOk = false;
+				UE_LOG(LogNiagara, Error, TEXT("Cannot render %s because it does not define attribute %s %s."), *EmitterHandle.GetName().ToString(), *Attr.GetType().GetNameText().ToString() , *Attr.GetName().ToString());
+			}
+		}
+
+		EmitterRenderer[Index]->SetEnabled(bOk);
+	}
+	return bOk;
+}
+
+#endif
+
+/** Replace the current System renderer with a new one of Type.
+Don't forget to call RenderModuleUpdate on the SceneProxy after calling this! 
+ */
+void FNiagaraEmitterInstance::UpdateEmitterRenderer(ERHIFeatureLevel::Type FeatureLevel, TArray<NiagaraRenderer*>& ToBeAddedList, TArray<NiagaraRenderer*>& ToBeRemovedList)
+{
+	const FNiagaraEmitterHandle& EmitterHandle = GetEmitterHandle();
+	const UNiagaraEmitter* EmitterProperties = EmitterHandle.GetInstance();
+
+	if (EmitterProperties != nullptr)
+	{
+		// Add all the old to be purged..
+		for (int32 SubIdx = 0; SubIdx < EmitterRenderer.Num(); SubIdx++)
+		{
+			if (EmitterRenderer[SubIdx] != nullptr)
+			{
+				ToBeRemovedList.Add(EmitterRenderer[SubIdx]);
+				EmitterRenderer[SubIdx] = nullptr;
+			}
+		}
+
+		if (bIsEnabled && !bError)
+		{
+			EmitterRenderer.Empty();
+			EmitterRenderer.AddZeroed(EmitterProperties->RendererProperties.Num());
+			for (int32 SubIdx = 0; SubIdx < EmitterProperties->RendererProperties.Num(); SubIdx++)
+			{
+				UMaterialInterface *Material = nullptr;
+
+				TArray<UMaterialInterface*> UsedMats;
+				if (EmitterProperties->RendererProperties[SubIdx] != nullptr)
+				{
+					EmitterProperties->RendererProperties[SubIdx]->GetUsedMaterials(UsedMats);
+					if (UsedMats.Num() != 0)
+					{
+						Material = UsedMats[0];
+					}
+				}
+
+				if (Material == nullptr)
+				{
+					Material = UMaterial::GetDefaultMaterial(MD_Surface);
+				}
+
+				if (EmitterProperties->RendererProperties[SubIdx] != nullptr)
+				{
+					EmitterRenderer[SubIdx] = EmitterProperties->RendererProperties[SubIdx]->CreateEmitterRenderer(FeatureLevel);
+					EmitterRenderer[SubIdx]->SetMaterial(Material, FeatureLevel);
+					EmitterRenderer[SubIdx]->SetLocalSpace(EmitterProperties->bLocalSpace);
+					ToBeAddedList.Add(EmitterRenderer[SubIdx]);
+
+					//UE_LOG(LogNiagara, Warning, TEXT("CreateRenderer %p"), EmitterRenderer);
+#if WITH_EDITORONLY_DATA
+					CheckAttributesForRenderer(SubIdx);
+#endif
+				}
+				else
+				{
+					EmitterRenderer[SubIdx] = nullptr;
+				}
+			}
+		}
+	}
+}
diff --git a/Engine/Plugins/FX/Niagara/Source/Niagara/Private/NiagaraEmitterInstanceBatcher.cpp b/Engine/Plugins/FX/Niagara/Source/Niagara/Private/NiagaraEmitterInstanceBatcher.cpp
new file mode 100644
--- /dev/null
+++ b/Engine/Plugins/FX/Niagara/Source/Niagara/Private/NiagaraEmitterInstanceBatcher.cpp
@@ -0,0 +1,410 @@
+// Copyright 1998-2017 Epic Games, Inc. All Rights Reserved.
+
+#include "NiagaraEmitterInstanceBatcher.h"
+#include "NiagaraScriptExecutionContext.h"
+#include "RHI.h"
+#include "NiagaraStats.h"
+#include "NiagaraShader.h"
+#include "ShaderParameterUtils.h"
+#include "SceneUtils.h"
+#include "ClearQuad.h"
+
+DECLARE_CYCLE_STAT(TEXT("Batching"), STAT_NiagaraBatching, STATGROUP_Niagara);
+DECLARE_FLOAT_COUNTER_STAT(TEXT("Niagara GPU Sim"), Stat_GPU_NiagaraSim, STATGROUP_GPU);
+
+
+NiagaraEmitterInstanceBatcher* NiagaraEmitterInstanceBatcher::BatcherSingleton = nullptr;
+uint32 FNiagaraComputeExecutionContext::TickCounter = 0;
+
+
+void NiagaraEmitterInstanceBatcher::Queue(FNiagaraComputeExecutionContext *InContext)
+{
+	//SimulationQueue[CurQueueIndex]->Add(InContext);
+
+	ENQUEUE_UNIQUE_RENDER_COMMAND_THREEPARAMETER(QueueNiagaraDispatch,
+		TArray<FNiagaraComputeExecutionContext*>*, Queue, &SimulationQueue[0],
+		uint32, QueueIndex, CurQueueIndex,
+		FNiagaraComputeExecutionContext*, ExecContext, InContext,
+		{
+			Queue[QueueIndex].Add(ExecContext);
+		});
+}
+
+
+void NiagaraEmitterInstanceBatcher::ExecuteAll(FRHICommandList &RHICmdList)
+{
+	TArray<FNiagaraComputeExecutionContext*> &WorkQueue = SimulationQueue[CurQueueIndex ^ 0x1];
+	for (FNiagaraComputeExecutionContext *Context : WorkQueue)
+	{
+		//ExecuteSingle(Context, RHICmdList);
+		TickSingle(Context, RHICmdList);
+	}
+	WorkQueue.Empty();
+}
+
+void NiagaraEmitterInstanceBatcher::TickSingle(const FNiagaraComputeExecutionContext *Context, FRHICommandList &RHICmdList) const
+{
+	SCOPE_CYCLE_COUNTER(STAT_NiagaraTick);
+
+	check(IsInRenderingThread());
+	Context->MainDataSet->TickRenderThread(ENiagaraSimTarget::GPUComputeSim);
+
+	FNiagaraShader *UpdateShader = Context->RTUpdateScript->GetShader();
+	FNiagaraShader *SpawnShader = Context->RTSpawnScript->GetShader();
+	if (!UpdateShader || !SpawnShader)
+	{
+		return;
+	}
+
+	uint32 PrevNumInstances = Context->MainDataSet->PrevDataRender().GetNumInstances();
+
+	uint32 EventSpawnTotal = 0;
+	for (int32 i = 0; i < Context->EventHandlerScriptProps.Num(); i++)
+	{
+		const FNiagaraEventScriptProperties &EventHandlerProps = Context->EventHandlerScriptProps[i];
+		if (EventHandlerProps.ExecutionMode == EScriptExecutionMode::SpawnedParticles && Context->EventSets[i])
+		{
+			uint32 NumEventsToProcess = Context->EventSets[i]->PrevDataRender().GetNumInstances();
+			uint32 EventSpawnNum = NumEventsToProcess * EventHandlerProps.SpawnNumber;
+			EventSpawnTotal += EventSpawnNum;
+		}
+	}
+
+
+	// allocate for additional instances spawned and set the new number in the data set, if the new number is greater (meaning if we're spawning in this run)
+	// TODO: interpolated spawning
+	//
+	uint32 NewNumInstances = Context->SpawnRateInstances + Context->BurstInstances + EventSpawnTotal + PrevNumInstances;
+	if (NewNumInstances > PrevNumInstances)
+	{
+		Context->MainDataSet->CurrDataRender().AllocateGPU(NewNumInstances, RHICmdList);
+		Context->MainDataSet->CurrDataRender().SetNumInstances(NewNumInstances);
+	}
+	// if we're not spawning, we need to make sure that the current buffer alloc size and number of instances matches the last one
+	// we may have spawned in the last tick, so the buffers may be different sizes
+	//
+	else if (Context->MainDataSet->CurrDataRender().GetNumInstances() < Context->MainDataSet->PrevDataRender().GetNumInstances())
+	{
+		Context->MainDataSet->CurrDataRender().AllocateGPU(PrevNumInstances, RHICmdList);
+		Context->MainDataSet->CurrDataRender().SetNumInstances(PrevNumInstances);
+	}
+
+	//UE_LOG(LogNiagara, Log, TEXT("Pre sim instances %i"), PrevNumInstances);
+
+	// simulation run
+	RHICmdList.SetComputeShader(UpdateShader->GetComputeShader());
+	SetupDataInterfaceBuffers(Context->UpdateInterfaces, UpdateShader, RHICmdList);
+	SetupEventUAVs(Context, PrevNumInstances, RHICmdList);
+	Run(Context->MainDataSet, 0, PrevNumInstances, UpdateShader, Context->UpdateParams, RHICmdList);
+	UnsetEventUAVs(Context, RHICmdList);
+
+	// resolve data set writes
+	// grabs the number of instances written from the index set during the simulation run
+	//
+	uint32 NumInstancesAfterSim[64];
+	NumInstancesAfterSim[0] = PrevNumInstances;
+	ResolveDatasetWrites(NumInstancesAfterSim, Context);		// this is going to cause a flush; need to place it differently
+	Context->MainDataSet->CurrDataRender().SetNumInstances(NumInstancesAfterSim[0]);
+
+	// TODO: hack - only updating event set 0 on update scripts now; need to match them to their indices and update them all
+	if (Context->UpdateEventWriteDataSets.Num())
+	{
+		Context->UpdateEventWriteDataSets[0]->CurrDataRender().SetNumInstances(NumInstancesAfterSim[1]);
+	}
+
+	UE_LOG(LogNiagara, Log, TEXT("After sim instances %i"), NumInstancesAfterSim[0]);
+
+	// spawn run
+	// bursts and regular spawn rate happen at once here
+	RHICmdList.SetComputeShader(SpawnShader->GetComputeShader());
+	uint32 SpawnInstances = Context->SpawnRateInstances + Context->BurstInstances + EventSpawnTotal;
+	uint32 NumInstancesAfterSpawn = NumInstancesAfterSim[0] + SpawnInstances;
+	uint32 NumInstancesAfterNonEventSpawn = NumInstancesAfterSim[0] + Context->SpawnRateInstances + Context->BurstInstances;
+	if (SpawnInstances)
+	{
+		Run(Context->MainDataSet, NumInstancesAfterSim[0], Context->SpawnRateInstances + Context->BurstInstances + EventSpawnTotal, SpawnShader, Context->SpawnParams, RHICmdList);
+		Context->MainDataSet->CurrDataRender().SetNumInstances(NumInstancesAfterSpawn);
+
+		//uncomment this to check number of spawn instances versus ones actually processed by the shader
+		uint32 NumInstancesSpawned[64];
+		ResolveDatasetWrites(NumInstancesSpawned, Context);
+		ensure(NumInstancesSpawned[0] == SpawnInstances);
+		UE_LOG(LogNiagara, Log, TEXT("Spawned %i to %i, ran spawn script on %i"), Context->SpawnRateInstances, NumInstancesAfterSpawn, NumInstancesSpawned[0]);
+
+		ensure(NumInstancesAfterNonEventSpawn + EventSpawnTotal == NumInstancesAfterSpawn);
+		ensure(NumInstancesAfterNonEventSpawn + EventSpawnTotal == Context->MainDataSet->CurrDataRender().GetNumInstances());
+	}
+
+
+
+	RunEventHandlers(Context, NumInstancesAfterSim[0], NumInstancesAfterSpawn, NumInstancesAfterNonEventSpawn, RHICmdList);
+
+	// the VF grabs PrevDataRender for drawing, so need to transition
+	RHICmdList.TransitionResource(EResourceTransitionAccess::EReadable, EResourceTransitionPipeline::EComputeToGfx, Context->MainDataSet->PrevDataRender().GetGPUBufferFloat()->UAV);
+	RHICmdList.TransitionResource(EResourceTransitionAccess::EReadable, EResourceTransitionPipeline::EComputeToGfx, Context->MainDataSet->PrevDataRender().GetGPUBufferInt()->UAV);
+}
+
+
+void NiagaraEmitterInstanceBatcher::ResolveDatasetWrites(uint32 *OutArray, const FNiagaraComputeExecutionContext *Context) const
+{
+	if (Context->MainDataSet->GetDataSetIndices().NumBytes > 0)
+	{
+		int32 *DataSetInstanceIndexBuffer = static_cast<int32*>(RHILockVertexBuffer(Context->MainDataSet->GetDataSetIndices().Buffer, 0, 64 * sizeof(int32), EResourceLockMode::RLM_ReadOnly));
+		for (uint32 SetIdx = 0; SetIdx < 64; SetIdx++)
+		{
+			OutArray[SetIdx] = DataSetInstanceIndexBuffer[SetIdx];
+		}
+		RHIUnlockVertexBuffer(Context->MainDataSet->GetDataSetIndices().Buffer);
+	}
+}
+
+void NiagaraEmitterInstanceBatcher::SetupDataInterfaceBuffers(const TArray<FNiagaraScriptDataInterfaceInfo> &DIInfos, FNiagaraShader *Shader, FRHICommandList &RHICmdList) const
+{
+	// set up data interface buffers, as defined by the DIs during compilation
+	//
+	uint32 InterfaceIndex = 0;
+	for (const FNiagaraScriptDataInterfaceInfo &InterfaceInfo : DIInfos)
+	{
+		TArray<FNiagaraDataInterfaceBufferData> &BufferDatas = InterfaceInfo.DataInterface->GetBufferDataArray();
+		for (FNiagaraDataInterfaceBufferData &BufferData : BufferDatas)
+		{
+			FShaderResourceParameter *Param = Shader->FindDIBufferParam(InterfaceIndex, BufferData.UniformName);
+			if (Param)
+			{
+				RHICmdList.SetShaderResourceViewParameter(Shader->GetComputeShader(), Param->GetBaseIndex(), BufferData.Buffer.SRV);
+			}
+		}
+		InterfaceIndex++;
+	}
+}
+
+void NiagaraEmitterInstanceBatcher::Run(FNiagaraDataSet *DataSet, uint32 StartInstance, const uint32 NumInstances, FNiagaraShader *Shader, const TArray<uint8, TAlignedHeapAllocator<16>>  &Params, FRHICommandList &RHICmdList, bool bCopyBeforeStart) const
+{
+	// Recreate a clear data set index buffer for the simulation shader to write number of written instances to
+	//
+	FRWBuffer &InstanceIdxBuf = DataSet->SetupDataSetIndices();
+	ClearUAV(RHICmdList, InstanceIdxBuf, 0);
+	RHICmdList.TransitionResource(EResourceTransitionAccess::EWritable, EResourceTransitionPipeline::EComputeToCompute, InstanceIdxBuf.UAV);
+
+	// set the shader and data set params 
+	//
+	//RHICmdList.SetComputeShader(Shader->GetComputeShader());
+	DataSet->SetShaderParams(Shader, RHICmdList);
+
+	// set the index buffer uav
+	//
+	if (Shader->OutputIndexBufferParam.IsBound())
+	{
+		RHICmdList.SetUAVParameter(Shader->GetComputeShader(), Shader->OutputIndexBufferParam.GetUAVIndex(), InstanceIdxBuf.UAV);
+	}
+
+	// set the execution parameters
+	//
+	if (Shader->EmitterTickCounterParam.IsBound())
+	{
+		RHICmdList.SetShaderParameter(Shader->GetComputeShader(), 0, Shader->EmitterTickCounterParam.GetBaseIndex(), Shader->EmitterTickCounterParam.GetNumBytes(), &FNiagaraComputeExecutionContext::TickCounter);
+	}
+	uint32 Copy = bCopyBeforeStart ? 1 : 0;
+	RHICmdList.SetShaderParameter(Shader->GetComputeShader(), 0, Shader->CopyInstancesBeforeStartParam.GetBaseIndex(), Shader->CopyInstancesBeforeStartParam.GetNumBytes(), &Copy);
+
+	// Now figure out how many instances each thread needs to simulate and how many thread groups we're running. We do this based on the total number of 
+	// instances we need to run and how many threads per thread group there are as well as the max. number of thread groups (defined in NiagaraCommon.h)
+	//
+	uint32 NumInstancesPerThread = 0;
+	uint32 AdjustedNumInstances = NumInstances;
+	uint32 SimulateStartInstance = StartInstance;
+
+	// if the caller requests to copy all instances before StartInstance, set up the params accordingly
+	// spawn events need to do this, for example
+	if (bCopyBeforeStart)
+	{
+		AdjustedNumInstances += StartInstance;
+		StartInstance = 0;
+	}
+
+	RHICmdList.SetShaderParameter(Shader->GetComputeShader(), 0, Shader->SimulateStartInstanceParam.GetBaseIndex(), Shader->SimulateStartInstanceParam.GetNumBytes(), &SimulateStartInstance);
+	RHICmdList.SetShaderParameter(Shader->GetComputeShader(), 0, Shader->StartInstanceParam.GetBaseIndex(), Shader->StartInstanceParam.GetNumBytes(), &StartInstance);
+
+	NumInstancesPerThread = FMath::DivideAndRoundUp(AdjustedNumInstances, NIAGARA_COMPUTE_THREADGROUP_SIZE);
+	uint32 NumThreadGroups = 1;
+	if (AdjustedNumInstances > NIAGARA_COMPUTE_THREADGROUP_SIZE)
+	{
+		NumThreadGroups = FMath::Min(NIAGARA_MAX_COMPUTE_THREADGROUPS, FMath::DivideAndRoundUp(AdjustedNumInstances, NIAGARA_COMPUTE_THREADGROUP_SIZE));
+		NumInstancesPerThread = FMath::DivideAndRoundUp(NumInstancesPerThread, NumThreadGroups);
+	}
+
+	RHICmdList.SetShaderParameter(Shader->GetComputeShader(), 0, Shader->NumInstancesPerThreadParam.GetBaseIndex(), Shader->NumInstancesPerThreadParam.GetNumBytes(), &NumInstancesPerThread);
+	RHICmdList.SetShaderParameter(Shader->GetComputeShader(), 0, Shader->NumInstancesParam.GetBaseIndex(), Shader->NumInstancesParam.GetNumBytes(), &AdjustedNumInstances);
+	RHICmdList.SetShaderParameter(Shader->GetComputeShader(), 0, Shader->NumThreadGroupsParam.GetBaseIndex(), Shader->NumThreadGroupsParam.GetNumBytes(), &NumThreadGroups);
+
+	if (StartInstance > 0)
+	{
+		//UE_LOG(LogNiagara, Log, TEXT("Startinstance %i (%i)    NumInstances %i (%i)   %i/Thread, %i groups"), StartInstance, SimulateStartInstance, NumInstances, AdjustedNumInstances, NumInstancesPerThread, NumThreadGroups);
+	}
+
+
+	// setup script parameters
+	FRHIUniformBufferLayout CBufferLayout(TEXT("Niagara Compute Sim CBuffer"));
+	CBufferLayout.ConstantBufferSize = Params.Num();
+	if (CBufferLayout.ConstantBufferSize)
+	{
+		CBufferLayout.ResourceOffset = 0;
+		check(CBufferLayout.Resources.Num() == 0);
+		const uint8* ParamData = Params.GetData();
+		FUniformBufferRHIRef CBuffer = RHICreateUniformBuffer(ParamData, CBufferLayout, EUniformBufferUsage::UniformBuffer_MultiFrame);
+		RHICmdList.SetShaderUniformBuffer(Shader->GetComputeShader(), Shader->EmitterConstantBufferParam.GetBaseIndex(), CBuffer);
+	}
+
+	// Dispatch, if anything needs to be done
+	//
+	if (NumInstancesPerThread)
+	{
+		SCOPED_DRAW_EVENT(RHICmdList, NiagaraGPUSimulationCS);
+		SCOPED_GPU_STAT(RHICmdList, Stat_GPU_NiagaraSim);
+		DispatchComputeShader(RHICmdList, Shader, NumThreadGroups, 1, 1);
+	}
+
+	// Unset UAV parameters and transition resources (TODO: resource transition should be moved to the renderer)
+	// 
+	DataSet->UnsetShaderParams(Shader, RHICmdList);
+	Shader->OutputIndexBufferParam.UnsetUAV(RHICmdList, Shader->GetComputeShader());
+}
+
+
+void NiagaraEmitterInstanceBatcher::RunEventHandlers(const FNiagaraComputeExecutionContext *Context, uint32 NumInstancesAfterSim, uint32 NumInstancesAfterSpawn, uint32 NumInstancesAfterNonEventSpawn, FRHICommandList &RHICmdList) const
+{
+	// Event handler run
+	//
+	for (int32 EventScriptIdx = 0; EventScriptIdx < Context->EventHandlerScriptProps.Num(); EventScriptIdx++)
+	{
+		const FNiagaraEventScriptProperties &EventHandlerProps = Context->EventHandlerScriptProps[EventScriptIdx];
+		FNiagaraDataSet *EventSet = Context->EventSets[EventScriptIdx];
+		if (EventSet)
+		{
+			int32 NumEvents = EventSet->PrevDataRender().GetNumInstances();
+
+			// handle all-particle events
+			if (NumEvents && EventHandlerProps.Script && EventHandlerProps.ExecutionMode == EScriptExecutionMode::EveryParticle && EventSet)
+			{
+				FNiagaraShader *EventHandlerShader = EventHandlerProps.Script->GetRenderThreadScript()->GetShader();
+				if (EventHandlerShader)
+				{
+					SetupDataInterfaceBuffers(EventHandlerProps.Script->DataInterfaceInfo, EventHandlerShader, RHICmdList);
+
+					RHICmdList.SetShaderParameter(EventHandlerShader->GetComputeShader(), 0, EventHandlerShader->NumParticlesPerEventParam.GetBaseIndex(), sizeof(int32), &NumInstancesAfterSpawn);
+					RHICmdList.SetShaderParameter(EventHandlerShader->GetComputeShader(), 0, EventHandlerShader->NumEventsPerParticleParam.GetBaseIndex(), sizeof(int32), &NumEvents);
+
+					// If we have events, Swap buffers, to make sure we don't overwrite previous script results
+					Context->MainDataSet->TickRenderThread(ENiagaraSimTarget::GPUComputeSim);
+					Context->MainDataSet->CurrDataRender().AllocateGPU(NumInstancesAfterSpawn, RHICmdList);
+					Context->MainDataSet->CurrDataRender().SetNumInstances(NumInstancesAfterSpawn);
+
+					SetPrevDataStrideParams(EventSet, EventHandlerShader, RHICmdList);
+
+					// set up event data set buffers
+					FShaderResourceParameter &FloatParam = EventHandlerShader->EventFloatSRVParams[0];
+					FShaderResourceParameter &IntParam = EventHandlerShader->EventIntSRVParams[0];
+					RHICmdList.SetShaderResourceViewParameter(EventHandlerShader->GetComputeShader(), FloatParam.GetBaseIndex(), EventSet->PrevDataRender().GetGPUBufferFloat()->SRV);
+					RHICmdList.SetShaderResourceViewParameter(EventHandlerShader->GetComputeShader(), IntParam.GetBaseIndex(), EventSet->PrevDataRender().GetGPUBufferInt()->SRV);
+
+					TArray<uint8, TAlignedHeapAllocator<16>>  BlankParams;
+					Run(Context->MainDataSet, 0, NumInstancesAfterNonEventSpawn, EventHandlerShader, BlankParams, RHICmdList);
+				}
+			}
+
+			// handle spawn events
+			if (EventHandlerProps.Script && EventHandlerProps.ExecutionMode == EScriptExecutionMode::SpawnedParticles && EventSet)
+			{
+				uint32 NumEventsToProcess = EventSet->PrevDataRender().GetNumInstances();
+				uint32 EventSpawnNum = NumEventsToProcess * EventHandlerProps.SpawnNumber;
+
+				//int32 SpawnNum = Context->EventSpawnCounts[EventScriptIdx];
+				FNiagaraShader *EventHandlerShader = EventHandlerProps.Script->GetRenderThreadScript()->GetShader();
+				if (NumEvents && EventSpawnNum && EventHandlerShader)
+				{
+					//SetupDataInterfaceBuffers(EventHandlerProps.Script->DataInterfaceInfo, EventHandlerShader, RHICmdList);
+					int32 OneEvent = 1;
+					int32 ParticlesPerEvent = EventSpawnNum / NumEvents;
+					RHICmdList.SetShaderParameter(EventHandlerShader->GetComputeShader(), 0, EventHandlerShader->NumEventsPerParticleParam.GetBaseIndex(), sizeof(int32), &OneEvent);
+					RHICmdList.SetShaderParameter(EventHandlerShader->GetComputeShader(), 0, EventHandlerShader->NumParticlesPerEventParam.GetBaseIndex(), sizeof(int32), &ParticlesPerEvent);
+
+					// If we have events, Swap buffers
+					Context->MainDataSet->TickRenderThread(ENiagaraSimTarget::GPUComputeSim);
+					Context->MainDataSet->CurrDataRender().AllocateGPU(NumInstancesAfterSpawn, RHICmdList);
+					Context->MainDataSet->CurrDataRender().SetNumInstances(NumInstancesAfterSpawn);
+
+					// set up event data set buffers
+					FShaderResourceParameter &FloatParam = EventHandlerShader->EventFloatSRVParams[0];
+					FShaderResourceParameter &IntParam = EventHandlerShader->EventIntSRVParams[0];
+					RHICmdList.SetShaderResourceViewParameter(EventHandlerShader->GetComputeShader(), FloatParam.GetBaseIndex(), EventSet->PrevDataRender().GetGPUBufferFloat()->SRV);
+					RHICmdList.SetShaderResourceViewParameter(EventHandlerShader->GetComputeShader(), IntParam.GetBaseIndex(), EventSet->PrevDataRender().GetGPUBufferInt()->SRV);
+
+					SetPrevDataStrideParams(EventSet, EventHandlerShader, RHICmdList);
+
+					// we assume event spawns are at the end of the buffer
+					check(NumInstancesAfterNonEventSpawn + EventSpawnNum == Context->MainDataSet->CurrDataRender().GetNumInstances());
+
+					TArray<uint8, TAlignedHeapAllocator<16>>  BlankParams;
+					Run(Context->MainDataSet, NumInstancesAfterNonEventSpawn, EventSpawnNum, EventHandlerShader, BlankParams, RHICmdList, true);
+				}
+			}
+		}
+	}
+}
+
+void NiagaraEmitterInstanceBatcher::SetPrevDataStrideParams(const FNiagaraDataSet *Set, FNiagaraShader *Shader, FRHICommandList &RHICmdList) const
+{
+	int32 FloatStride = Set->PrevDataRender().GetFloatStride() / sizeof(float);
+	int32 IntStride = Set->PrevDataRender().GetInt32Stride() / sizeof(int32);
+	RHICmdList.SetShaderParameter(Shader->GetComputeShader(), 0, Shader->EventReadFloatStrideParams[0].GetBaseIndex(), sizeof(int32), &FloatStride);
+	RHICmdList.SetShaderParameter(Shader->GetComputeShader(), 0, Shader->EventReadIntStrideParams[0].GetBaseIndex(), sizeof(int32), &IntStride);
+}
+
+
+void NiagaraEmitterInstanceBatcher::SetupEventUAVs(const FNiagaraComputeExecutionContext *Context, uint32 NewNumInstances, FRHICommandList &RHICmdList) const
+{
+	FNiagaraShader *UpdateShader = Context->RTUpdateScript->GetShader();
+
+	uint32 SetIndex = 0;
+	for (FNiagaraDataSet *Set : Context->UpdateEventWriteDataSets)
+	{
+		if (NewNumInstances)
+		{
+			Set->CurrDataRender().AllocateGPU(NewNumInstances, RHICmdList);
+			Set->CurrDataRender().SetNumInstances(NewNumInstances);
+			FRWShaderParameter &FloatParam = UpdateShader->EventFloatUAVParams[SetIndex];
+			FRWShaderParameter &IntParam = UpdateShader->EventIntUAVParams[SetIndex];
+			if (FloatParam.IsUAVBound())
+			{
+				RHICmdList.SetUAVParameter(UpdateShader->GetComputeShader(), FloatParam.GetUAVIndex(), Set->CurrDataRender().GetGPUBufferFloat()->UAV);
+			}
+			if (IntParam.IsUAVBound())
+			{
+				RHICmdList.SetUAVParameter(UpdateShader->GetComputeShader(), IntParam.GetUAVIndex(), Set->CurrDataRender().GetGPUBufferInt()->UAV);
+			}
+
+			uint32 FloatStride = Set->CurrDataRender().GetFloatStride() / sizeof(float);
+			uint32 IntStride = Set->CurrDataRender().GetInt32Stride() / sizeof(int32);
+			RHICmdList.SetShaderParameter(UpdateShader->GetComputeShader(), 0, UpdateShader->EventWriteFloatStrideParams[SetIndex].GetBaseIndex(), sizeof(int32), &FloatStride);
+			RHICmdList.SetShaderParameter(UpdateShader->GetComputeShader(), 0, UpdateShader->EventWriteIntStrideParams[SetIndex].GetBaseIndex(), sizeof(int32), &IntStride);
+		}
+
+		SetIndex++;
+	}
+
+
+}
+
+
+void NiagaraEmitterInstanceBatcher::UnsetEventUAVs(const FNiagaraComputeExecutionContext *Context, FRHICommandList &RHICmdList) const
+{
+	FNiagaraShader *UpdateShader = Context->RTUpdateScript->GetShader();
+
+	for (int32 SetIndex = 0; SetIndex < Context->UpdateEventWriteDataSets.Num(); SetIndex++)
+	{
+		FRWShaderParameter &FloatParam = UpdateShader->EventFloatUAVParams[SetIndex];
+		FRWShaderParameter &IntParam = UpdateShader->EventIntUAVParams[SetIndex];
+		FloatParam.UnsetUAV(RHICmdList, UpdateShader->GetComputeShader());
+		IntParam.UnsetUAV(RHICmdList, UpdateShader->GetComputeShader());
+	}
+}
\ No newline at end of file
diff --git a/Engine/Plugins/FX/Niagara/Source/Niagara/Private/NiagaraFunctionLibrary.cpp b/Engine/Plugins/FX/Niagara/Source/Niagara/Private/NiagaraFunctionLibrary.cpp
new file mode 100644
--- /dev/null
+++ b/Engine/Plugins/FX/Niagara/Source/Niagara/Private/NiagaraFunctionLibrary.cpp
@@ -0,0 +1,113 @@
+// Copyright 1998-2017 Epic Games, Inc. All Rights Reserved.
+
+#include "NiagaraFunctionLibrary.h"
+#include "EngineGlobals.h"
+#include "GameFramework/WorldSettings.h"
+#include "Engine/Engine.h"
+#include "NiagaraComponent.h"
+#include "NiagaraSystem.h"
+
+#include "NiagaraWorldManager.h"
+
+UNiagaraFunctionLibrary::UNiagaraFunctionLibrary(const FObjectInitializer& ObjectInitializer)
+: Super(ObjectInitializer)
+{
+}
+
+
+/**
+* Spawns a Niagara System at the specified world location/rotation
+* @return			The spawned UNiagaraComponent
+*/
+UNiagaraComponent* UNiagaraFunctionLibrary::SpawnSystemAtLocation(UObject* WorldContextObject, UNiagaraSystem* SystemTemplate, FVector SpawnLocation, FRotator SpawnRotation, bool bAutoDestroy)
+{
+	UNiagaraComponent* PSC = NULL;
+	if (SystemTemplate)
+	{
+		UWorld* World = GEngine->GetWorldFromContextObject(WorldContextObject, EGetWorldErrorMode::LogAndReturnNull);
+		if (World != nullptr)
+		{
+			AActor* Actor = World->GetWorldSettings();
+			PSC = NewObject<UNiagaraComponent>((Actor ? Actor : (UObject*)World));
+			PSC->SetAsset(SystemTemplate);
+			PSC->RegisterComponentWithWorld(World);
+
+			PSC->SetAbsolute(true, true, true);
+			PSC->SetWorldLocationAndRotation(SpawnLocation, SpawnRotation);
+			PSC->SetRelativeScale3D(FVector(1.f));
+			//PSC->ActivateSystem(true);
+		}
+	}
+	return PSC;
+}
+
+
+
+
+
+/**
+* Spawns a Niagara System attached to a component
+* @return			The spawned UNiagaraComponent
+*/
+UNiagaraComponent* UNiagaraFunctionLibrary::SpawnSystemAttached(UNiagaraSystem* SystemTemplate, USceneComponent* AttachToComponent, FName AttachPointName, FVector Location, FRotator Rotation, EAttachLocation::Type LocationType, bool bAutoDestroy)
+{
+	UNiagaraComponent* PSC = nullptr;
+	if (SystemTemplate)
+	{
+		if (AttachToComponent == NULL)
+		{
+			UE_LOG(LogScript, Warning, TEXT("UNiagaraFunctionLibrary::SpawnSystemAttached: NULL AttachComponent specified!"));
+		}
+		else
+		{
+			AActor* Actor = AttachToComponent->GetOwner();
+			PSC = NewObject<UNiagaraComponent>((Actor ? Actor : (UObject*)AttachToComponent->GetWorld()));
+			PSC->SetAsset(SystemTemplate);
+			PSC->RegisterComponentWithWorld(AttachToComponent->GetWorld());
+
+			PSC->AttachToComponent(AttachToComponent, FAttachmentTransformRules::KeepRelativeTransform, AttachPointName);
+			if (LocationType == EAttachLocation::KeepWorldPosition)
+			{
+				PSC->SetWorldLocationAndRotation(Location, Rotation);
+			}
+			else
+			{
+				PSC->SetRelativeLocationAndRotation(Location, Rotation);
+			}
+			PSC->SetRelativeScale3D(FVector(1.f));
+		}
+	}
+	return PSC;
+}
+
+
+/**
+* Set a constant in an emitter of a Niagara System
+void UNiagaraFunctionLibrary::SetUpdateScriptConstant(UNiagaraComponent* Component, FName EmitterName, FName ConstantName, FVector Value)
+{
+	TArray<TSharedPtr<FNiagaraEmitterInstance>> &Emitters = Component->GetSystemInstance()->GetEmitters();
+
+	for (TSharedPtr<FNiagaraEmitterInstance> &Emitter : Emitters)
+	{		
+		if(UNiagaraEmitter* PinnedProps = Emitter->GetProperties().Get())
+		{
+			FName CurName = *PinnedProps->EmitterName;
+			if (CurName == EmitterName)
+			{
+				Emitter->GetProperties()->UpdateScriptProps.ExternalConstants.SetOrAdd(FNiagaraTypeDefinition::GetVec4Def(), ConstantName, Value);
+				break;
+			}
+		}
+	}
+}
+*/
+
+UNiagaraParameterCollectionInstance* UNiagaraFunctionLibrary::GetNiagaraParameterCollection(UObject* WorldContextObject, UNiagaraParameterCollection* Collection)
+{
+	UWorld* World = GEngine->GetWorldFromContextObject(WorldContextObject, EGetWorldErrorMode::LogAndReturnNull);
+	if (World != nullptr)
+	{
+		return FNiagaraWorldManager::Get(World)->GetParameterCollection(Collection);
+	}
+	return nullptr;
+}
