Copying //UE4/Orion-Staging to //UE4/Main (Source: //Orion/Dev-General @ 3606378)

diff --git a/Engine/Plugins/FX/Niagara/Source/Niagara/Private/NiagaraLightRendererProperties.cpp b/Engine/Plugins/FX/Niagara/Source/Niagara/Private/NiagaraLightRendererProperties.cpp
new file mode 100644
--- /dev/null
+++ b/Engine/Plugins/FX/Niagara/Source/Niagara/Private/NiagaraLightRendererProperties.cpp
@@ -0,0 +1,68 @@
+// Copyright 1998-2016 Epic Games, Inc. All Rights Reserved.
+
+#include "NiagaraLightRendererProperties.h"
+#include "NiagaraRenderer.h"
+#include "NiagaraConstants.h"
+
+UNiagaraLightRendererProperties::UNiagaraLightRendererProperties()
+	: RadiusScale(1.0f), ColorAdd(FVector(0.0f, 0.0f, 0.0f))
+{
+}
+
+NiagaraRenderer* UNiagaraLightRendererProperties::CreateEmitterRenderer(ERHIFeatureLevel::Type FeatureLevel)
+{
+	return new NiagaraRendererLights(FeatureLevel, this);
+}
+
+void UNiagaraLightRendererProperties::GetUsedMaterials(TArray<UMaterialInterface*>& OutMaterials) const
+{
+	//OutMaterials.Add(Material);
+	//Material should live here.
+}
+
+#if WITH_EDITORONLY_DATA
+
+const TArray<FNiagaraVariable>& UNiagaraLightRendererProperties::GetRequiredAttributes()
+{
+	static TArray<FNiagaraVariable> Attrs;
+
+	if (Attrs.Num() == 0)
+	{
+		Attrs.Add(SYS_PARAM_PARTICLES_POSITION);
+		Attrs.Add(SYS_PARAM_PARTICLES_COLOR);
+		Attrs.Add(SYS_PARAM_PARTICLES_SCALE);
+		/*Attrs.Add(SYS_PARAM_PARTICLES_VELOCITY);
+		Attrs.Add(SYS_PARAM_PARTICLES_SPRITE_ROTATION);
+		Attrs.Add(SYS_PARAM_PARTICLES_NORMALIZED_AGE);
+		Attrs.Add(SYS_PARAM_PARTICLES_SPRITE_SIZE);*/
+	}
+
+	return Attrs;
+}
+
+const TArray<FNiagaraVariable>& UNiagaraLightRendererProperties::GetOptionalAttributes()
+{
+	static TArray<FNiagaraVariable> Attrs;
+
+	if (Attrs.Num() == 0)
+	{
+		//Attrs.Add(SYS_PARAM_PARTICLES_SPRITE_FACING);
+		//Attrs.Add(SYS_PARAM_PARTICLES_SPRITE_ALIGNMENT);
+		//Attrs.Add(SYS_PARAM_PARTICLES_SUB_IMAGE_INDEX);
+		//Attrs.Add(SYS_PARAM_PARTICLES_DYNAMIC_MATERIAL_PARAM);
+	}
+
+	return Attrs;
+}
+
+bool UNiagaraLightRendererProperties::IsMaterialValidForRenderer(UMaterial* Material, FText& InvalidMessage)
+{
+	return true;
+}
+
+void UNiagaraLightRendererProperties::FixMaterial(UMaterial* Material)
+{
+}
+
+#endif // WITH_EDITORONLY_DATA
+
diff --git a/Engine/Plugins/FX/Niagara/Source/Niagara/Private/NiagaraMeshRendererProperties.cpp b/Engine/Plugins/FX/Niagara/Source/Niagara/Private/NiagaraMeshRendererProperties.cpp
new file mode 100644
--- /dev/null
+++ b/Engine/Plugins/FX/Niagara/Source/Niagara/Private/NiagaraMeshRendererProperties.cpp
@@ -0,0 +1,98 @@
+// Copyright 1998-2016 Epic Games, Inc. All Rights Reserved.
+
+#include "NiagaraMeshRendererProperties.h"
+#include "NiagaraRendererMeshes.h"
+#include "Engine/StaticMesh.h"
+#include "NiagaraConstants.h"
+
+UNiagaraMeshRendererProperties::UNiagaraMeshRendererProperties()
+	: ParticleMesh(nullptr)
+{
+}
+
+NiagaraRenderer* UNiagaraMeshRendererProperties::CreateEmitterRenderer(ERHIFeatureLevel::Type FeatureLevel)
+{
+	return new NiagaraRendererMeshes(FeatureLevel, this);
+}
+
+void UNiagaraMeshRendererProperties::GetUsedMaterials(TArray<UMaterialInterface*>& OutMaterials) const
+{
+	if (ParticleMesh)
+	{
+		for(FStaticMaterial& Mat : ParticleMesh->StaticMaterials)
+		{
+			OutMaterials.Add(Mat.MaterialInterface);
+		}
+	}
+}
+
+#if WITH_EDITORONLY_DATA
+
+bool UNiagaraMeshRendererProperties::IsMaterialValidForRenderer(UMaterial* Material, FText& InvalidMessage)
+{
+	if (Material->bUsedWithMeshParticles == false)
+	{
+		InvalidMessage = NSLOCTEXT("NiagaraMeshRendererProperties", "InvalidMaterialMessage", "The material isn't marked as \"Used with mesh particles\"");
+		return false;
+	}
+	return true;
+}
+
+void UNiagaraMeshRendererProperties::FixMaterial(UMaterial* Material)
+{
+	Material->Modify();
+	Material->bUsedWithNiagaraMeshParticles = true;
+	Material->ForceRecompileForRendering();
+}
+
+const TArray<FNiagaraVariable>& UNiagaraMeshRendererProperties::GetRequiredAttributes()
+{
+	static TArray<FNiagaraVariable> Attrs;
+
+	if (Attrs.Num() == 0)
+	{
+		Attrs.Add(SYS_PARAM_PARTICLES_POSITION);
+		Attrs.Add(SYS_PARAM_PARTICLES_VELOCITY);
+		Attrs.Add(SYS_PARAM_PARTICLES_COLOR);
+		Attrs.Add(SYS_PARAM_PARTICLES_NORMALIZED_AGE);
+	}
+
+	return Attrs;
+}
+
+
+const TArray<FNiagaraVariable>& UNiagaraMeshRendererProperties::GetOptionalAttributes()
+{
+	static TArray<FNiagaraVariable> Attrs;
+
+	if (Attrs.Num() == 0)
+	{
+		Attrs.Add(SYS_PARAM_PARTICLES_SCALE);
+		//Attrs.Add(SYS_PARAM_PARTICLES_SPRITE_FACING);
+		//Attrs.Add(SYS_PARAM_PARTICLES_SPRITE_ALIGNMENT);
+		//Attrs.Add(SYS_PARAM_PARTICLES_SUB_IMAGE_INDEX);
+		Attrs.Add(SYS_PARAM_PARTICLES_DYNAMIC_MATERIAL_PARAM);
+	}
+
+	return Attrs;
+}
+
+#endif // WITH_EDITORONLY_DATA
+
+void UNiagaraMeshRendererProperties::PostEditChangeProperty(FPropertyChangedEvent& PropertyChangedEvent)
+{
+	if (ParticleMesh && PropertyChangedEvent.Property && PropertyChangedEvent.Property->GetName() == "ParticleMesh")
+	{
+		const FStaticMeshLODResources& LODModel = ParticleMesh->RenderData->LODResources[0];
+		for (int32 SectionIndex = 0; SectionIndex < LODModel.Sections.Num(); SectionIndex++)
+		{
+			const FStaticMeshSection& Section = LODModel.Sections[SectionIndex];
+			UMaterialInterface *Material = ParticleMesh->GetMaterial(Section.MaterialIndex);
+			if (Material)
+			{
+				FMaterialRenderProxy* MaterialProxy = Material->GetRenderProxy(false, false);
+				Material->CheckMaterialUsage(MATUSAGE_NiagaraMeshParticles);
+			}
+		}
+	}
+}
diff --git a/Engine/Plugins/FX/Niagara/Source/Niagara/Private/NiagaraModule.cpp b/Engine/Plugins/FX/Niagara/Source/Niagara/Private/NiagaraModule.cpp
new file mode 100644
--- /dev/null
+++ b/Engine/Plugins/FX/Niagara/Source/Niagara/Private/NiagaraModule.cpp
@@ -0,0 +1,441 @@
+// Copyright 1998-2017 Epic Games, Inc. All Rights Reserved.
+
+#include "NiagaraModule.h"
+#include "Modules/ModuleManager.h"
+#include "NiagaraTypes.h"
+#include "NiagaraEvents.h"
+#include "NiagaraSettings.h"
+#include "NiagaraDataInterfaceCurlNoise.h"
+#include "Class.h"
+#include "Package.h"
+#include "NiagaraWorldManager.h"
+#include "VectorVM.h"
+
+IMPLEMENT_MODULE(INiagaraModule, Niagara);
+
+void INiagaraModule::StartupModule()
+{
+	VectorVM::Init();
+	FNiagaraTypeDefinition::Init();
+
+	FWorldDelegates::OnPreWorldInitialization.AddRaw(this, &INiagaraModule::OnWorldInit);
+	FWorldDelegates::OnWorldCleanup.AddRaw(this, &INiagaraModule::OnWorldCleanup);
+	FWorldDelegates::OnPreWorldFinishDestroy.AddRaw(this, &INiagaraModule::OnPreWorldFinishDestroy);
+
+	FWorldDelegates::OnWorldPostActorTick.AddRaw(this, &INiagaraModule::TickWorld);
+#if WITH_EDITOR	
+	// This is done so that the editor classes are available to load in the cooker on editor builds even though it doesn't load the editor directly.
+	// UMG does something similar for similar reasons.
+	// @TODO We should remove this once Niagara is fully a plug-in.
+	FModuleManager::Get().LoadModule(TEXT("NiagaraEditor"));
+#endif
+	UNiagaraDataInterfaceCurlNoise::InitNoiseLUT();
+}
+
+void INiagaraModule::ShutdownModule()
+{
+	//Should have cleared up all world managers by now.
+	check(WorldManagers.Num() == 0);
+	for (TPair<UWorld*, FNiagaraWorldManager*> Pair : WorldManagers)
+	{
+		delete Pair.Value;
+		Pair.Value = nullptr;
+	}
+}
+
+FNiagaraWorldManager* INiagaraModule::GetWorldManager(UWorld* World)
+{
+	return WorldManagers.FindChecked(World);
+}
+
+void INiagaraModule::DestroyAllSystemSimulations(class UNiagaraSystem* System)
+{
+	for (TPair<UWorld*, FNiagaraWorldManager*>& Pair : WorldManagers)
+	{
+		Pair.Value->DestroySystemSimulation(System);
+	}
+}
+
+void INiagaraModule::OnWorldInit(UWorld* World, const UWorld::InitializationValues IVS)
+{
+	check(WorldManagers.Find(World) == nullptr);
+	WorldManagers.Add(World) = new FNiagaraWorldManager(World);
+}
+
+void INiagaraModule::OnWorldCleanup(UWorld* World, bool bSessionEnded, bool bCleanupResources)
+{
+	//Cleanup world manager contents but not the manager itself.
+}
+
+void INiagaraModule::OnPreWorldFinishDestroy(UWorld* World)
+{
+	FNiagaraWorldManager*& Manager = WorldManagers.FindChecked(World);
+	delete Manager;
+	WorldManagers.Remove(World);
+}
+
+void INiagaraModule::TickWorld(UWorld* World, ELevelTick TickType, float DeltaSeconds)
+{
+	GetWorldManager(World)->Tick(DeltaSeconds);
+}
+
+
+
+//////////////////////////////////////////////////////////////////////////
+
+const UScriptStruct* FNiagaraTypeDefinition::ParameterMapStruct;
+const UScriptStruct* FNiagaraTypeDefinition::NumericStruct;
+const UScriptStruct* FNiagaraTypeDefinition::FloatStruct;
+const UScriptStruct* FNiagaraTypeDefinition::BoolStruct;
+const UScriptStruct* FNiagaraTypeDefinition::IntStruct;
+const UScriptStruct* FNiagaraTypeDefinition::Matrix4Struct;
+const UScriptStruct* FNiagaraTypeDefinition::Vec4Struct;
+const UScriptStruct* FNiagaraTypeDefinition::Vec3Struct;
+const UScriptStruct* FNiagaraTypeDefinition::Vec2Struct;
+const UScriptStruct* FNiagaraTypeDefinition::ColorStruct;
+
+const UEnum* FNiagaraTypeDefinition::ExecutionStateEnum;
+
+FNiagaraTypeDefinition FNiagaraTypeDefinition::ParameterMapDef;
+FNiagaraTypeDefinition FNiagaraTypeDefinition::NumericDef;
+FNiagaraTypeDefinition FNiagaraTypeDefinition::FloatDef;
+FNiagaraTypeDefinition FNiagaraTypeDefinition::BoolDef;
+FNiagaraTypeDefinition FNiagaraTypeDefinition::IntDef;
+FNiagaraTypeDefinition FNiagaraTypeDefinition::Matrix4Def;
+FNiagaraTypeDefinition FNiagaraTypeDefinition::Vec4Def;
+FNiagaraTypeDefinition FNiagaraTypeDefinition::Vec3Def;
+FNiagaraTypeDefinition FNiagaraTypeDefinition::Vec2Def;
+FNiagaraTypeDefinition FNiagaraTypeDefinition::ColorDef;
+
+TSet<const UScriptStruct*> FNiagaraTypeDefinition::NumericStructs;
+TArray<FNiagaraTypeDefinition> FNiagaraTypeDefinition::OrderedNumericTypes;
+
+TSet<const UScriptStruct*> FNiagaraTypeDefinition::ScalarStructs;
+
+TSet<const UStruct*> FNiagaraTypeDefinition::FloatStructs;
+TSet<const UStruct*> FNiagaraTypeDefinition::IntStructs;
+TSet<const UStruct*> FNiagaraTypeDefinition::BoolStructs;
+
+FNiagaraTypeDefinition FNiagaraTypeDefinition::CollisionEventDef;
+
+
+TArray<FNiagaraTypeDefinition> FNiagaraTypeRegistry::RegisteredTypes;
+TArray<FNiagaraTypeDefinition> FNiagaraTypeRegistry::RegisteredParamTypes;
+TArray<FNiagaraTypeDefinition> FNiagaraTypeRegistry::RegisteredPayloadTypes;
+TArray<FNiagaraTypeDefinition> FNiagaraTypeRegistry::RegisteredUserDefinedTypes;
+TArray<FNiagaraTypeDefinition> FNiagaraTypeRegistry::RegisteredNumericTypes;
+
+bool FNiagaraTypeDefinition::IsDataInterface()const
+{
+	return Struct->IsChildOf(UNiagaraDataInterface::StaticClass());
+}
+
+void FNiagaraTypeDefinition::Init()
+{
+	static auto* CoreUObjectPkg = FindObjectChecked<UPackage>(nullptr, TEXT("/Script/CoreUObject"));
+	static auto* NiagaraPkg = FindObjectChecked<UPackage>(nullptr, TEXT("/Script/Niagara"));
+	FNiagaraTypeDefinition::ParameterMapStruct = FindObjectChecked<UScriptStruct>(NiagaraPkg, TEXT("NiagaraParameterMap"));
+	FNiagaraTypeDefinition::NumericStruct = FindObjectChecked<UScriptStruct>(NiagaraPkg, TEXT("NiagaraNumeric"));
+	FNiagaraTypeDefinition::FloatStruct = FindObjectChecked<UScriptStruct>(NiagaraPkg, TEXT("NiagaraFloat"));
+	FNiagaraTypeDefinition::BoolStruct = FindObjectChecked<UScriptStruct>(NiagaraPkg, TEXT("NiagaraBool"));
+	FNiagaraTypeDefinition::IntStruct = FindObjectChecked<UScriptStruct>(NiagaraPkg, TEXT("NiagaraInt32"));
+	FNiagaraTypeDefinition::Matrix4Struct = FindObjectChecked<UScriptStruct>(NiagaraPkg, TEXT("NiagaraMatrix"));
+
+	FNiagaraTypeDefinition::Vec2Struct = FindObjectChecked<UScriptStruct>(CoreUObjectPkg, TEXT("Vector2D"));
+	FNiagaraTypeDefinition::Vec3Struct = FindObjectChecked<UScriptStruct>(CoreUObjectPkg, TEXT("Vector"));
+	FNiagaraTypeDefinition::Vec4Struct = FindObjectChecked<UScriptStruct>(CoreUObjectPkg, TEXT("Vector4"));
+	FNiagaraTypeDefinition::ColorStruct = FindObjectChecked<UScriptStruct>(CoreUObjectPkg, TEXT("LinearColor"));
+	
+	ParameterMapDef = FNiagaraTypeDefinition(ParameterMapStruct);
+	NumericDef = FNiagaraTypeDefinition(NumericStruct);
+	FloatDef = FNiagaraTypeDefinition(FloatStruct);
+	BoolDef = FNiagaraTypeDefinition(BoolStruct);
+	IntDef = FNiagaraTypeDefinition(IntStruct);
+	Vec2Def = FNiagaraTypeDefinition(Vec2Struct);
+	Vec3Def = FNiagaraTypeDefinition(Vec3Struct);
+	Vec4Def = FNiagaraTypeDefinition(Vec4Struct);
+	ColorDef = FNiagaraTypeDefinition(ColorStruct);
+	Matrix4Def = FNiagaraTypeDefinition(Matrix4Struct);
+
+	CollisionEventDef = FNiagaraTypeDefinition(FNiagaraCollisionEventPayload::StaticStruct());
+	NumericStructs.Add(NumericStruct);
+	NumericStructs.Add(FloatStruct);
+	NumericStructs.Add(IntStruct);
+	NumericStructs.Add(Vec2Struct);
+	NumericStructs.Add(Vec3Struct);
+	NumericStructs.Add(Vec4Struct);
+	NumericStructs.Add(ColorStruct);
+	//Make matrix a numeric type?
+
+	FloatStructs.Add(FloatStruct);
+	FloatStructs.Add(Vec2Struct);
+	FloatStructs.Add(Vec3Struct);
+	FloatStructs.Add(Vec4Struct);
+	//FloatStructs.Add(Matrix4Struct)??
+	FloatStructs.Add(ColorStruct);
+
+	IntStructs.Add(IntStruct);
+
+	BoolStructs.Add(BoolStruct);
+
+	OrderedNumericTypes.Add(IntStruct);
+	OrderedNumericTypes.Add(FloatStruct);
+	OrderedNumericTypes.Add(Vec2Struct);
+	OrderedNumericTypes.Add(Vec3Struct);
+	OrderedNumericTypes.Add(Vec4Struct);
+	OrderedNumericTypes.Add(ColorStruct);
+
+	ScalarStructs.Add(BoolStruct);
+	ScalarStructs.Add(IntStruct);
+	ScalarStructs.Add(FloatStruct);
+
+	ExecutionStateEnum = FindObject<UEnum>(ANY_PACKAGE, TEXT("ENiagaraExecutionState"), true);
+	
+	RecreateUserDefinedTypeRegistry();
+}
+
+bool FNiagaraTypeDefinition::IsValidNumericInput(const FNiagaraTypeDefinition& TypeDef)
+{
+	if (NumericStructs.Contains(TypeDef.GetScriptStruct()))
+	{
+		return true;
+	}
+	return false;
+}
+
+void FNiagaraTypeDefinition::RecreateUserDefinedTypeRegistry()
+{
+	static auto* CoreUObjectPkg = FindObjectChecked<UPackage>(nullptr, TEXT("/Script/CoreUObject"));
+	static auto* NiagaraPkg = FindObjectChecked<UPackage>(nullptr, TEXT("/Script/Niagara"));
+
+	FNiagaraTypeRegistry::ClearUserDefinedRegistry();
+
+	FNiagaraTypeRegistry::Register(CollisionEventDef, false, true, false);
+
+	FNiagaraTypeRegistry::Register(ParameterMapDef, true, false, false);
+	FNiagaraTypeRegistry::Register(NumericDef, true, false, false);
+	FNiagaraTypeRegistry::Register(FloatDef, true, true, false);
+	FNiagaraTypeRegistry::Register(IntDef, true, true, false);
+	FNiagaraTypeRegistry::Register(BoolDef, true, true, false);
+	FNiagaraTypeRegistry::Register(Vec2Def, true, true, false);
+	FNiagaraTypeRegistry::Register(Vec3Def, true, true, false);
+	FNiagaraTypeRegistry::Register(Vec4Def, true, true, false);
+	FNiagaraTypeRegistry::Register(ColorDef, true, true, false);
+	FNiagaraTypeRegistry::Register(Matrix4Def, true, false, false);
+
+	FNiagaraTypeRegistry::Register(FNiagaraTypeDefinition(ExecutionStateEnum), true, true, false);
+
+	UScriptStruct* TestStruct = FindObjectChecked<UScriptStruct>(NiagaraPkg, TEXT("NiagaraTestStruct"));
+	FNiagaraTypeDefinition TestDefinition(TestStruct);
+	FNiagaraTypeRegistry::Register(TestDefinition, true, false, false);
+
+	UScriptStruct* SpawnInfoStruct = FindObjectChecked<UScriptStruct>(NiagaraPkg, TEXT("NiagaraSpawnInfo"));
+	FNiagaraTypeRegistry::Register(FNiagaraTypeDefinition(SpawnInfoStruct), true, false, false);
+
+	const UNiagaraSettings* Settings = GetDefault<UNiagaraSettings>();
+	check(Settings);
+	TArray<FSoftObjectPath> TotalStructAssets;
+	for (FSoftObjectPath AssetRef : Settings->AdditionalParameterTypes)
+	{
+		TotalStructAssets.AddUnique(AssetRef);
+	}
+
+	for (FSoftObjectPath AssetRef : Settings->AdditionalPayloadTypes)
+	{
+		TotalStructAssets.AddUnique(AssetRef);
+	}
+
+	for (FSoftObjectPath AssetRef : TotalStructAssets)
+	{
+		UObject* Obj = AssetRef.ResolveObject();
+		if (Obj == nullptr)
+		{
+			Obj = AssetRef.TryLoad();
+		}
+
+		if (Obj != nullptr)
+		{
+			const FSoftObjectPath* ParamRefFound = Settings->AdditionalParameterTypes.FindByPredicate([&](const FSoftObjectPath& Ref) { return Ref.ToString() == AssetRef.ToString(); });
+			const FSoftObjectPath* PayloadRefFound = Settings->AdditionalPayloadTypes.FindByPredicate([&](const FSoftObjectPath& Ref) { return Ref.ToString() == AssetRef.ToString(); });
+			UScriptStruct* ScriptStruct = Cast<UScriptStruct>(Obj);
+			if (ScriptStruct != nullptr)
+			{
+				FNiagaraTypeRegistry::Register(ScriptStruct, ParamRefFound != nullptr, PayloadRefFound != nullptr, true);
+			}
+		}
+		else
+		{
+			UE_LOG(LogNiagara, Warning, TEXT("Could not find additional parameter/payload type: %s"), *AssetRef.ToString());
+		}
+	}
+
+
+	for (FStringAssetReference AssetRef : Settings->AdditionalParameterEnums)
+	{
+		UObject* Obj = AssetRef.ResolveObject();
+		if (Obj == nullptr)
+		{
+			Obj = AssetRef.TryLoad();
+		}
+
+		if (Obj != nullptr)
+		{
+			const FStringAssetReference* ParamRefFound = Settings->AdditionalParameterEnums.FindByPredicate([&](const FStringAssetReference& Ref) { return Ref.ToString() == AssetRef.ToString(); });
+			const FStringAssetReference* PayloadRefFound = nullptr;
+			UEnum* Enum = Cast<UEnum>(Obj);
+			if (Enum != nullptr)
+			{
+				FNiagaraTypeRegistry::Register(Enum, ParamRefFound != nullptr, PayloadRefFound != nullptr, true);
+			}
+		}
+		else
+		{
+			UE_LOG(LogNiagara, Warning, TEXT("Could not find additional parameter/payload enum: %s"), *AssetRef.ToString());
+		}
+	}
+
+}
+
+bool FNiagaraTypeDefinition::IsScalarDefinition(const FNiagaraTypeDefinition& Type)
+{
+	return ScalarStructs.Contains(Type.GetScriptStruct()) || (Type.GetScriptStruct() == IntStruct && Type.GetEnum() != nullptr);
+}
+
+bool FNiagaraTypeDefinition::TypesAreAssignable(const FNiagaraTypeDefinition& TypeA, const FNiagaraTypeDefinition& TypeB)
+{
+	if (const UClass* AClass = TypeA.GetClass())
+	{
+		if (const UClass* BClass = TypeB.GetClass())
+		{
+			return AClass == BClass;
+			return true;
+		}
+	}
+	
+	if (const UClass* BClass = TypeB.GetClass())
+	{
+		return false;
+	}
+
+	if (const UClass* AClass = TypeA.GetClass())
+	{
+		return false;
+	}
+
+	// Make sure that enums are not assignable to enums of different types or just plain ints
+	if (TypeA.GetStruct() == TypeB.GetStruct() &&
+		TypeA.GetEnum() != TypeB.GetEnum())
+	{
+		return false;
+	}
+
+	if (TypeA.GetStruct() == TypeB.GetStruct())
+	{
+		return true;
+	}
+
+	bool bIsSupportedConversion = false;
+	if (IsScalarDefinition(TypeA) && IsScalarDefinition(TypeB))
+	{
+		bIsSupportedConversion = (TypeA == IntDef && TypeB == FloatDef) || (TypeB == IntDef && TypeA == FloatDef);
+	}
+	else
+	{
+		bIsSupportedConversion = (TypeA == ColorDef && TypeB == Vec4Def) || (TypeB == ColorDef && TypeA == Vec4Def);
+	}
+
+	if (bIsSupportedConversion)
+	{
+		return true;
+	}
+
+	return	(TypeA == NumericDef && NumericStructs.Contains(TypeB.GetScriptStruct())) ||
+			(TypeB == NumericDef && NumericStructs.Contains(TypeA.GetScriptStruct())) ||
+			(TypeA == NumericDef && (TypeB.GetStruct() == GetIntStruct()) && TypeB.GetEnum() != nullptr) ||
+			(TypeB == NumericDef && (TypeA.GetStruct() == GetIntStruct()) && TypeA.GetEnum() != nullptr);
+}
+
+bool FNiagaraTypeDefinition::IsLossyConversion(const FNiagaraTypeDefinition& TypeA, const FNiagaraTypeDefinition& TypeB)
+{
+	return (TypeA == IntDef && TypeB == FloatDef) || (TypeB == IntDef && TypeA == FloatDef);
+}
+
+FNiagaraTypeDefinition FNiagaraTypeDefinition::GetNumericOutputType(const TArray<FNiagaraTypeDefinition> TypeDefinintions, ENiagaraNumericOutputTypeSelectionMode SelectionMode)
+{
+	checkf(SelectionMode != ENiagaraNumericOutputTypeSelectionMode::None, TEXT("Can not get numeric output type with selection mode none."));
+
+	//This may need some work. Should work fine for now.
+	if (SelectionMode == ENiagaraNumericOutputTypeSelectionMode::Scalar)
+	{
+		bool bHasFloats = false;
+		bool bHasInts = false;
+		bool bHasBools = false;
+		for (const FNiagaraTypeDefinition& Type : TypeDefinintions)
+		{
+			bHasFloats |= FloatStructs.Contains(Type.GetStruct());
+			bHasInts |= IntStructs.Contains(Type.GetStruct());
+			bHasBools |= BoolStructs.Contains(Type.GetStruct());
+		}
+		//Not sure what to do if we have multiple different types here.
+		//Possibly pick this up ealier and throw a compile error?
+		if (bHasFloats) return FNiagaraTypeDefinition::GetFloatDef();
+		if (bHasInts) return FNiagaraTypeDefinition::GetIntDef();
+		if (bHasBools) return FNiagaraTypeDefinition::GetBoolDef();
+	}
+	// Always return the numeric type definition if it's included since this isn't a valid use case and we don't want to hide it.
+	int32 NumericTypeDefinitionIndex = TypeDefinintions.IndexOfByKey(NumericDef);
+	if (NumericTypeDefinitionIndex != INDEX_NONE)
+	{
+		// TODO: Warning here?
+		return NumericDef;
+	}
+
+	TArray<FNiagaraTypeDefinition> SortedTypeDefinitions = TypeDefinintions;
+	SortedTypeDefinitions.Sort([&](const FNiagaraTypeDefinition& TypeA, const FNiagaraTypeDefinition& TypeB)
+	{
+		int32 AIndex = OrderedNumericTypes.IndexOfByKey(TypeA);
+		int32 BIndex = OrderedNumericTypes.IndexOfByKey(TypeB);
+		return AIndex < BIndex;
+	});
+
+	if (SelectionMode == ENiagaraNumericOutputTypeSelectionMode::Largest)
+	{
+		return SortedTypeDefinitions.Last();
+	}
+	else // if (SelectionMode == ENiagaraNumericOutputTypeSelectionMode::Smallest)
+	{
+		return SortedTypeDefinitions[0];
+	}
+}
+
+//////////////////////////////////////////////////////////////////////////
+
+void FNiagaraScriptDataInterfaceInfo::CopyTo(FNiagaraScriptDataInterfaceInfo* Destination, UObject* Outer) const
+{
+	Destination->Name = Name;
+	Destination->DataInterface = Cast<UNiagaraDataInterface>(StaticDuplicateObject(DataInterface, Outer, NAME_None, ~RF_Transient));
+	Destination->UserPtrIdx = UserPtrIdx;
+}
+
+
+
+FDelegateHandle INiagaraModule::SetOnProcessShaderCompilationQueue(FOnProcessQueue InOnProcessQueue)
+{
+	checkf(OnProcessQueue.IsBound() == false, TEXT("Shader processing queue delegate already set."));
+	OnProcessQueue = InOnProcessQueue;
+	return OnProcessQueue.GetHandle();
+}
+
+void INiagaraModule::ResetOnProcessShaderCompilationQueue(FDelegateHandle DelegateHandle)
+{
+	checkf(OnProcessQueue.GetHandle() == DelegateHandle, TEXT("Can only reset the process compilation queue delegate with the handle it was created with."));
+	OnProcessQueue.Unbind();
+}
+
+void INiagaraModule::ProcessShaderCompilationQueue()
+{
+	checkf(OnProcessQueue.IsBound(), TEXT("Can not process shader queue.  Delegate was never set."));
+	return OnProcessQueue.Execute();
+}
+
+
diff --git a/Engine/Plugins/FX/Niagara/Source/Niagara/Private/NiagaraParameterBindingInstance.cpp b/Engine/Plugins/FX/Niagara/Source/Niagara/Private/NiagaraParameterBindingInstance.cpp
new file mode 100644
--- /dev/null
+++ b/Engine/Plugins/FX/Niagara/Source/Niagara/Private/NiagaraParameterBindingInstance.cpp
@@ -0,0 +1,15 @@
+// Copyright 1998-2016 Epic Games, Inc. All Rights Reserved.
+
+#include "NiagaraParameterBindingInstance.h"
+#include "NiagaraTypes.h"
+
+FNiagaraParameterBindingInstance::FNiagaraParameterBindingInstance(const FNiagaraVariable& InSource, FNiagaraVariable& InDestination)
+	: Source(InSource)
+	, Destination(InDestination)
+{
+}
+
+void FNiagaraParameterBindingInstance::Tick()
+{
+	Source.CopyTo(Destination.GetData());
+}
\ No newline at end of file
diff --git a/Engine/Plugins/FX/Niagara/Source/Niagara/Private/NiagaraParameterCollection.cpp b/Engine/Plugins/FX/Niagara/Source/Niagara/Private/NiagaraParameterCollection.cpp
new file mode 100644
--- /dev/null
+++ b/Engine/Plugins/FX/Niagara/Source/Niagara/Private/NiagaraParameterCollection.cpp
@@ -0,0 +1,289 @@
+// Copyright 1998-2017 Epic Games, Inc. All Rights Reserved.
+
+#include "NiagaraParameterCollection.h"
+#include "NiagaraDataInterface.h"
+
+//////////////////////////////////////////////////////////////////////////
+
+UNiagaraParameterCollectionInstance::UNiagaraParameterCollectionInstance(const FObjectInitializer& ObjectInitializer)
+	: Super(ObjectInitializer)
+	, ParameterStorage(this)
+{
+	//Bind(ParameterStorage);
+}
+
+UNiagaraParameterCollectionInstance::~UNiagaraParameterCollectionInstance()
+{
+	//Unbind(ParameterStorage);
+}
+
+void UNiagaraParameterCollectionInstance::PostLoad()
+{
+	Super::PostLoad();
+	//Ensure we're synced up with our collection. TODO: Do conditionally via a version number/guid?
+	SyncWithCollection();
+}
+
+void UNiagaraParameterCollectionInstance::SetParent(UNiagaraParameterCollection* InParent)
+{
+	Collection = InParent;
+	SyncWithCollection();
+}
+
+bool UNiagaraParameterCollectionInstance::IsDefaultInstance()const
+{
+	return GetParent() && GetParent()->GetDefaultInstance() == this; 
+}
+
+bool UNiagaraParameterCollectionInstance::AddParameter(const FNiagaraVariable& Parameter)
+{
+	Modify();
+	return ParameterStorage.AddParameter(Parameter);
+}
+
+bool UNiagaraParameterCollectionInstance::RemoveParameter(const FNiagaraVariable& Parameter)
+{
+	Modify();
+	return ParameterStorage.RemoveParameter(Parameter);
+}
+
+void UNiagaraParameterCollectionInstance::RenameParameter(const FNiagaraVariable& Parameter, FName NewName)
+{
+	Modify();
+	ParameterStorage.RenameParameter(Parameter, NewName);
+}
+
+void UNiagaraParameterCollectionInstance::Empty()
+{
+	Modify();
+	ParameterStorage.Empty();
+}
+
+void UNiagaraParameterCollectionInstance::GetParameters(TArray<FNiagaraVariable>& OutParameters)
+{
+	ParameterStorage.GetParameters(OutParameters);
+}
+
+void UNiagaraParameterCollectionInstance::Tick()
+{
+	//Push our parameter changes to any bound stores.
+	ParameterStorage.Tick();
+}
+
+void UNiagaraParameterCollectionInstance::SyncWithCollection()
+{
+	FNiagaraParameterStore OldStore = ParameterStorage;
+	ParameterStorage.Empty(false);
+
+	for (FNiagaraVariable& Param : Collection->GetParameters())
+	{
+		int32 Offset = OldStore.IndexOf(Param);
+		if (Offset != INDEX_NONE && OverridesParameter(Param))
+		{
+			//If this paramter is in the old store and we're overriding it, use the existing value in the store.
+			ParameterStorage.AddParameter(Param, false);
+			if (Param.IsDataInterface())
+			{
+				ParameterStorage.SetDataInterface(OldStore.GetDataInterface(Offset), Param);
+			}
+			else
+			{
+				ParameterStorage.SetParameterData(OldStore.GetParameterData(Offset), ParameterStorage.IndexOf(Param), Param.GetSizeInBytes());
+			}
+		}
+		else
+		{
+			//If the parameter did not exist in the old store or we don't override this parameter, sync it up to the parent collection.
+			FNiagaraParameterStore& DefaultStore = Collection->GetDefaultInstance()->GetParameterStore();
+			Offset = DefaultStore.IndexOf(Param);
+			check(Offset != INDEX_NONE);
+
+			ParameterStorage.AddParameter(Param, false);
+			if (Param.IsDataInterface())
+			{
+				ParameterStorage.SetDataInterface(CastChecked<UNiagaraDataInterface>(StaticDuplicateObject(DefaultStore.GetDataInterface(Offset), this)), Param);
+			}
+			else
+			{
+				ParameterStorage.SetParameterData(DefaultStore.GetParameterData(Offset), ParameterStorage.IndexOf(Param), Param.GetSizeInBytes());
+			}
+		}
+	}
+
+	ParameterStorage.Rebind();
+}
+
+bool UNiagaraParameterCollectionInstance::OverridesParameter(const FNiagaraVariable& Parameter)const
+{ 
+	return IsDefaultInstance() || OverridenParameters.Contains(Parameter); 
+}
+
+void UNiagaraParameterCollectionInstance::SetOverridesParameter(const FNiagaraVariable& Parameter, bool bOverrides)
+{
+	if (bOverrides)
+	{
+		OverridenParameters.AddUnique(Parameter);
+	}
+	else
+	{
+		OverridenParameters.Remove(Parameter);
+	}
+}
+
+//Blueprint Accessors
+bool UNiagaraParameterCollectionInstance::GetBoolParameter(const FString& InVariableName)
+{
+	return ParameterStorage.GetParameterValue<bool>(FNiagaraVariable(FNiagaraTypeDefinition::GetBoolDef(), *Collection->ParameterNameFromFriendlyName(InVariableName)));
+}
+
+float UNiagaraParameterCollectionInstance::GetFloatParameter(const FString& InVariableName)
+{
+	return ParameterStorage.GetParameterValue<float>(FNiagaraVariable(FNiagaraTypeDefinition::GetFloatDef(), *Collection->ParameterNameFromFriendlyName(InVariableName)));
+}
+
+int32 UNiagaraParameterCollectionInstance::GetIntParameter(const FString& InVariableName)
+{
+	return ParameterStorage.GetParameterValue<int32>(FNiagaraVariable(FNiagaraTypeDefinition::GetIntStruct(), *Collection->ParameterNameFromFriendlyName(InVariableName)));
+}
+
+FVector2D UNiagaraParameterCollectionInstance::GetVector2DParameter(const FString& InVariableName)
+{
+	return ParameterStorage.GetParameterValue<FVector2D>(FNiagaraVariable(FNiagaraTypeDefinition::GetVec2Def(), *Collection->ParameterNameFromFriendlyName(InVariableName)));
+}
+
+FVector UNiagaraParameterCollectionInstance::GetVectorParameter(const FString& InVariableName)
+{
+	return ParameterStorage.GetParameterValue<FVector>(FNiagaraVariable(FNiagaraTypeDefinition::GetVec3Def(), *Collection->ParameterNameFromFriendlyName(InVariableName)));
+}
+
+FVector4 UNiagaraParameterCollectionInstance::GetVector4Parameter(const FString& InVariableName)
+{
+	return ParameterStorage.GetParameterValue<FVector4>(FNiagaraVariable(FNiagaraTypeDefinition::GetVec4Def(), *Collection->ParameterNameFromFriendlyName(InVariableName)));
+}
+
+FLinearColor UNiagaraParameterCollectionInstance::GetColorParameter(const FString& InVariableName)
+{
+	return ParameterStorage.GetParameterValue<FLinearColor>(FNiagaraVariable(FNiagaraTypeDefinition::GetColorDef(), *Collection->ParameterNameFromFriendlyName(InVariableName)));
+}
+
+void UNiagaraParameterCollectionInstance::SetBoolParameter(const FString& InVariableName, bool InValue)
+{
+	ParameterStorage.SetParameterValue(InValue, FNiagaraVariable(FNiagaraTypeDefinition::GetBoolDef(), *Collection->ParameterNameFromFriendlyName(InVariableName)));
+}
+
+void UNiagaraParameterCollectionInstance::SetFloatParameter(const FString& InVariableName, float InValue)
+{
+	ParameterStorage.SetParameterValue(InValue, FNiagaraVariable(FNiagaraTypeDefinition::GetFloatDef(), *Collection->ParameterNameFromFriendlyName(InVariableName)));
+}
+
+void UNiagaraParameterCollectionInstance::SetIntParameter(const FString& InVariableName, int32 InValue)
+{
+	ParameterStorage.SetParameterValue(InValue, FNiagaraVariable(FNiagaraTypeDefinition::GetIntDef(), *Collection->ParameterNameFromFriendlyName(InVariableName)));
+}
+
+void UNiagaraParameterCollectionInstance::SetVector2DParameter(const FString& InVariableName, FVector2D InValue)
+{
+	ParameterStorage.SetParameterValue(InValue, FNiagaraVariable(FNiagaraTypeDefinition::GetVec2Def(), *Collection->ParameterNameFromFriendlyName(InVariableName)));
+}
+
+void UNiagaraParameterCollectionInstance::SetVectorParameter(const FString& InVariableName, FVector InValue)
+{
+	ParameterStorage.SetParameterValue(InValue, FNiagaraVariable(FNiagaraTypeDefinition::GetVec3Def(), *Collection->ParameterNameFromFriendlyName(InVariableName)));
+}
+
+void UNiagaraParameterCollectionInstance::SetVector4Parameter(const FString& InVariableName, const FVector4& InValue)
+{
+	ParameterStorage.SetParameterValue(InValue, FNiagaraVariable(FNiagaraTypeDefinition::GetVec4Def(), *Collection->ParameterNameFromFriendlyName(InVariableName)));
+}
+
+void UNiagaraParameterCollectionInstance::SetColorParameter(const FString& InVariableName, FLinearColor InValue)
+{
+	ParameterStorage.SetParameterValue(InValue, FNiagaraVariable(FNiagaraTypeDefinition::GetColorDef(), *Collection->ParameterNameFromFriendlyName(InVariableName)));
+}
+
+//////////////////////////////////////////////////////////////////////////
+
+UNiagaraParameterCollection::UNiagaraParameterCollection(const FObjectInitializer& ObjectInitializer)
+	: Super(ObjectInitializer)
+{
+	DefaultInstance = ObjectInitializer.CreateDefaultSubobject<UNiagaraParameterCollectionInstance>(this, TEXT("Default Instance"));
+	DefaultInstance->SetParent(this);
+}
+
+void UNiagaraParameterCollection::PostInitProperties()
+{
+	Super::PostInitProperties();
+	//Unique name is cached off just in case the hash function changes or our name does.
+	UniqueName = GetName() + LexicalConversion::ToString(GetTypeHash(GetFullName()));
+}
+
+void UNiagaraParameterCollection::PostLoad()
+{
+	Super::PostLoad();
+}
+
+int32 UNiagaraParameterCollection::IndexOfParameter(const FNiagaraVariable& Var)
+{
+	return Parameters.IndexOfByPredicate([&](const FNiagaraVariable& Other)
+	{
+		return Var.IsEquivalent(Other);
+	});
+}
+
+int32 UNiagaraParameterCollection::AddParameter(FName Name, FNiagaraTypeDefinition Type)
+{
+	Modify();
+
+	int32 Idx = Parameters.AddDefaulted(1);
+
+	Parameters[Idx].SetName(Name);
+	Parameters[Idx].SetType(Type);
+
+	DefaultInstance->AddParameter(Parameters[Idx]);
+
+	return Idx;
+}
+
+//void UNiagaraParameterCollection::RemoveParameter(int32 ParamIdx)
+void UNiagaraParameterCollection::RemoveParameter(const FNiagaraVariable& Parameter)
+{
+	Modify();
+	DefaultInstance->RemoveParameter(Parameter);
+	Parameters.Remove(Parameter);
+}
+
+void UNiagaraParameterCollection::RenameParameter(FNiagaraVariable& Parameter, FName NewName)
+{
+	Modify();
+
+	int32 ParamIdx = Parameters.IndexOfByKey(Parameter);
+	check(ParamIdx != INDEX_NONE);
+
+	Parameters[ParamIdx].SetName(NewName);
+	DefaultInstance->RenameParameter(Parameter, NewName);
+}
+
+FString UNiagaraParameterCollection::GetUniqueName()const
+{
+	return UniqueName;
+}
+
+FNiagaraVariable UNiagaraParameterCollection::CollectionParameterFromFriendlyParameter(const FNiagaraVariable& FriendlyParameter)const
+{
+	return FNiagaraVariable(FriendlyParameter.GetType(), *ParameterNameFromFriendlyName(FriendlyParameter.GetName().ToString()));
+}
+
+FNiagaraVariable UNiagaraParameterCollection::FriendlyParameterFromCollectionParameter(const FNiagaraVariable& CollectionParameter)const
+{
+	return FNiagaraVariable(CollectionParameter.GetType(), *FriendlyNameFromParameterName(CollectionParameter.GetName().ToString()));
+}
+
+FString UNiagaraParameterCollection::FriendlyNameFromParameterName(FString ParameterName)const
+{
+	return ParameterName.Replace(*(GetUniqueName() + TEXT("_")), TEXT(""), ESearchCase::CaseSensitive);
+}
+
+FString UNiagaraParameterCollection::ParameterNameFromFriendlyName(FString FriendlyName)const
+{
+	return GetUniqueName() + TEXT("_") + FriendlyName;
+}
\ No newline at end of file
diff --git a/Engine/Plugins/FX/Niagara/Source/Niagara/Private/NiagaraParameterStore.cpp b/Engine/Plugins/FX/Niagara/Source/Niagara/Private/NiagaraParameterStore.cpp
new file mode 100644
--- /dev/null
+++ b/Engine/Plugins/FX/Niagara/Source/Niagara/Private/NiagaraParameterStore.cpp
@@ -0,0 +1,253 @@
+// Copyright 1998-2016 Epic Games, Inc. All Rights Reserved.
+
+#include "NiagaraParameterStore.h"
+#include "NiagaraCommon.h"
+#include "NiagaraDataSet.h"
+#include "NiagaraComponent.h"
+#include "NiagaraSystemInstance.h"
+#include "NiagaraParameterCollection.h"
+
+
+//////////////////////////////////////////////////////////////////////////
+
+FNiagaraParameterStore::FNiagaraParameterStore()
+	: Owner(nullptr)
+	, bParametersDirty(true)
+	, bInterfacesDirty(true)
+{
+}
+
+FNiagaraParameterStore::FNiagaraParameterStore(UObject* InOwner)
+	: Owner(InOwner)
+	, bParametersDirty(true)
+	, bInterfacesDirty(true)
+	, bLayoutDirty(true)
+{
+
+}
+
+FNiagaraParameterStore::FNiagaraParameterStore(const FNiagaraParameterStore& Other)
+{
+	*this = Other;
+}
+
+FNiagaraParameterStore& FNiagaraParameterStore::operator=(const FNiagaraParameterStore& Other)
+{
+	Owner = Other.Owner;
+	ParameterOffsets = Other.ParameterOffsets;
+	ParameterData = Other.ParameterData;
+	DataInterfaces = Other.DataInterfaces;
+	//Don't copy bindings. We just want the data.
+	return *this;
+}
+
+FNiagaraParameterStore::~FNiagaraParameterStore()
+{
+}
+
+void FNiagaraParameterStore::Bind(FNiagaraParameterStore* DestStore)
+{
+	check(DestStore);
+	FNiagaraParameterStoreBinding& Binding = Bindings.Add(DestStore);	
+
+	Binding.Empty();
+	for (auto& VarInfo : DestStore->ParameterOffsets)
+	{
+		Binding.BindParameter(*DestStore, VarInfo.Value, *this, VarInfo.Key);
+	}
+
+	//Tick once immediately to ensure initial value is correct.
+	Binding.Tick(*this, *DestStore, true);
+}
+
+void FNiagaraParameterStore::Unbind(FNiagaraParameterStore* DestStore)
+{
+	Bindings.Remove(DestStore);
+}
+
+void FNiagaraParameterStore::Rebind()
+{
+	for (TPair<FNiagaraParameterStore*, FNiagaraParameterStoreBinding>& Binding : Bindings)
+	{
+		Binding.Value.Empty();
+		for (auto& VarInfo : Binding.Key->ParameterOffsets)
+		{
+			Binding.Value.BindParameter(*Binding.Key, VarInfo.Value, *this, VarInfo.Key);
+		}
+
+		//Tick once immediately to ensure initial value is correct.
+		//TODO: Can probably rework binding so that we don't need to force all to update on a rebind.
+		Binding.Value.Tick(*this, *Binding.Key, true);
+	}
+}
+
+void FNiagaraParameterStore::TransferBindings(FNiagaraParameterStore& OtherStore)
+{
+	for (TPair<FNiagaraParameterStore*, FNiagaraParameterStoreBinding>& Binding : Bindings)
+	{
+		OtherStore.Bind(Binding.Key);
+	}
+	Bindings.Empty();
+}
+
+void FNiagaraParameterStore::Tick()
+{
+	if (bLayoutDirty)
+	{
+		Rebind();
+	}
+	else
+	{
+		for (auto& Binding : Bindings)
+		{
+			Binding.Value.Tick(*this, *Binding.Key);
+		}
+	}
+
+	//We have to have ticked all our source stores before now.
+	bParametersDirty = false;
+	bInterfacesDirty = false;
+	bLayoutDirty = false;
+}
+
+/**
+Adds the passed parameter to this store.
+Does nothing if this parameter is already present.
+Returns true if we added a new parameter.
+*/
+bool FNiagaraParameterStore::AddParameter(const FNiagaraVariable& Param, bool bInitInterfaces)
+{
+	if (int32* Existing = ParameterOffsets.Find(Param))
+	{
+		return false;
+	}
+
+	if (Param.GetType().IsDataInterface())
+	{
+		int32 Offset = DataInterfaces.AddZeroed();
+		ParameterOffsets.Add(Param) = Offset;
+		DataInterfaces[Offset] = bInitInterfaces ? NewObject<UNiagaraDataInterface>(Owner, const_cast<UClass*>(Param.GetType().GetClass()), NAME_None, RF_Transactional) : nullptr;
+	}
+	else
+	{
+		int32 ParamSize = Param.GetSizeInBytes();
+		int32 ParamAlignment = Param.GetAlignment();
+		//int32 Offset = AlignArbitrary(ParameterData.Num(), ParamAlignment);//TODO: We need to handle alignment better here. Need to both satisfy CPU and GPU alignment concerns. VM doesn't care but the VM complier needs to be aware. Probably best to have everything adhere to GPU alignment rules.
+		int32 Offset = ParameterData.Num();
+		ParameterOffsets.Add(Param) = Offset;
+		ParameterData.AddUninitialized(ParamSize);
+
+		//Temporary to init param data from FNiagaraVariable storage. This will be removed when we change the UNiagaraScript to use a parameter store too.
+		if (Param.IsDataAllocated())
+		{
+			FMemory::Memcpy(GetParameterData_Internal(Offset), Param.GetData(), ParamSize);
+		}
+	}
+
+	OnLayoutChange();
+	return true;
+}
+
+bool FNiagaraParameterStore::RemoveParameter(const FNiagaraVariable& ToRemove)
+{
+	if (int32* ExistingIndex = ParameterOffsets.Find(ToRemove))
+	{
+		//TODO: Ensure direct bindings are either updated or disallowed here.
+		//We have to regenerate the store and the offsets on removal. This shouldn't happen at runtime!
+		TMap<FNiagaraVariable, int32> NewOffsets;
+		TArray<uint8> NewData;
+		TArray<UNiagaraDataInterface*> NewInterfaces;
+		for (TPair<FNiagaraVariable, int32>& Existing : ParameterOffsets)
+		{
+			FNiagaraVariable& ExistingVar = Existing.Key;
+			int32& ExistingOffset = Existing.Value;
+
+			//Add all but the one to remove to our
+			if (ExistingVar != ToRemove)
+			{
+				if (ExistingVar.GetType().IsDataInterface())
+				{
+					int32 Offset = NewInterfaces.AddZeroed();
+					NewOffsets.Add(ExistingVar) = Offset;
+					NewInterfaces[Offset] = DataInterfaces[ExistingOffset];
+				}
+				else
+				{
+					int32 Offset = NewData.Num();
+					int32 ParamSize = ExistingVar.GetSizeInBytes();
+					NewOffsets.Add(ExistingVar) = Offset;
+					NewData.AddUninitialized(ParamSize);
+					FMemory::Memcpy(NewData.GetData() + Offset, ParameterData.GetData() + ExistingOffset, ParamSize);
+				}
+			}
+		}
+
+		ParameterOffsets = NewOffsets;
+		ParameterData = NewData;
+		DataInterfaces = NewInterfaces;
+
+		OnLayoutChange();
+		return true;
+	}
+
+	return false;
+}
+
+void FNiagaraParameterStore::RenameParameter(const FNiagaraVariable& Param, FName NewName)
+{
+	int32 Idx = IndexOf(Param);
+	if(Idx != INDEX_NONE)
+	{
+		FNiagaraVariable NewParam = Param;
+		NewParam.SetName(NewName);
+		AddParameter(NewParam);
+
+		int32 NewIdx = IndexOf(NewParam);
+		if (Param.IsDataInterface())
+		{
+			SetDataInterface(GetDataInterface(Idx), NewIdx);
+		}
+		else
+		{
+			SetParameterData(GetParameterData_Internal(Idx), NewIdx, Param.GetSizeInBytes());
+		}
+		RemoveParameter(Param);
+
+		OnLayoutChange();
+	}
+}
+
+void FNiagaraParameterStore::Empty(bool bClearBindings)
+{
+	ParameterOffsets.Empty();
+	ParameterData.Empty();
+	DataInterfaces.Empty();
+	if (bClearBindings)
+	{
+		Bindings.Empty();
+	}
+	bParametersDirty = true;
+	bInterfacesDirty = true;
+	bLayoutDirty = true;
+}
+
+const FNiagaraVariable* FNiagaraParameterStore::FindVariable(UNiagaraDataInterface* Interface)const
+{
+	int32 Idx = DataInterfaces.IndexOfByKey(Interface);
+	if (Idx != INDEX_NONE)
+	{
+		for (const TPair<FNiagaraVariable, int32>& Existing : ParameterOffsets)
+		{
+			const FNiagaraVariable& ExistingVar = Existing.Key;
+			const int32& ExistingOffset = Existing.Value;
+
+			if (ExistingOffset == Idx && ExistingVar.GetType().GetClass() == Interface->GetClass())
+			{
+				return &ExistingVar;
+			}
+		}
+	}
+	return nullptr;
+}
+
+//////////////////////////////////////////////////////////////////////////
\ No newline at end of file
diff --git a/Engine/Plugins/FX/Niagara/Source/Niagara/Private/NiagaraRenderer.cpp b/Engine/Plugins/FX/Niagara/Source/Niagara/Private/NiagaraRenderer.cpp
new file mode 100644
--- /dev/null
+++ b/Engine/Plugins/FX/Niagara/Source/Niagara/Private/NiagaraRenderer.cpp
@@ -0,0 +1,158 @@
+// Copyright 1998-2017 Epic Games, Inc. All Rights Reserved.
+
+#include "NiagaraRenderer.h"
+#include "ParticleResources.h"
+#include "ParticleBeamTrailVertexFactory.h"
+#include "NiagaraDataSet.h"
+#include "NiagaraStats.h"
+
+DECLARE_CYCLE_STAT(TEXT("Generate Particle Lights"), STAT_NiagaraGenLights, STATGROUP_Niagara);
+
+
+/** Enable/disable parallelized System renderers */
+int32 GbNiagaraParallelEmitterRenderers = 1;
+static FAutoConsoleVariableRef CVarParallelEmitterRenderers(
+	TEXT("niagara.ParallelEmitterRenderers"),
+	GbNiagaraParallelEmitterRenderers,
+	TEXT("Whether to run Niagara System renderers in parallel"),
+	ECVF_Default
+	);
+
+
+
+
+NiagaraRenderer::~NiagaraRenderer() 
+{
+}
+
+void NiagaraRenderer::Release()
+{
+	check(IsInGameThread());
+	ENQUEUE_UNIQUE_RENDER_COMMAND_ONEPARAMETER(
+		NiagaraRendererDeletion,
+		NiagaraRenderer*, Renderer, this,
+		{
+			delete Renderer;
+		}
+	);
+}
+
+	
+
+
+NiagaraRendererLights::NiagaraRendererLights(ERHIFeatureLevel::Type FeatureLevel, UNiagaraRendererProperties *InProps) :
+	NiagaraRenderer()
+{
+	Properties = Cast<UNiagaraLightRendererProperties>(InProps);
+}
+
+
+void NiagaraRendererLights::ReleaseRenderThreadResources()
+{
+}
+
+void NiagaraRendererLights::CreateRenderThreadResources()
+{
+}
+
+
+
+/** Update render data buffer from attributes */
+FNiagaraDynamicDataBase *NiagaraRendererLights::GenerateVertexData(const FNiagaraSceneProxy* Proxy, FNiagaraDataSet &Data, const ENiagaraSimTarget Target)
+{
+	SCOPE_CYCLE_COUNTER(STAT_NiagaraGenLights);
+
+	SimpleTimer VertexDataTimer;
+
+	//I'm not a great fan of pulling scalar components out to a structured vert buffer like this.
+	//TODO: Experiment with a new VF that reads the data directly from the scalar layout.
+	FNiagaraDataSetIterator<FVector> PosItr(Data, FNiagaraVariable(FNiagaraTypeDefinition::GetVec3Def(), TEXT("Position")));
+	FNiagaraDataSetIterator<FLinearColor> ColItr(Data, FNiagaraVariable(FNiagaraTypeDefinition::GetColorDef(), TEXT("Color")));
+	FNiagaraDataSetIterator<FVector2D> SizeItr(Data, FNiagaraVariable(FNiagaraTypeDefinition::GetVec2Def(), TEXT("Size")));
+
+	//Bail if we don't have the required attributes to render this emitter.
+	if (!PosItr.IsValid() || !ColItr.IsValid() || !SizeItr.IsValid() || !bEnabled)
+	{
+		return nullptr;
+	}
+
+	FNiagaraDynamicDataLights *DynamicData = new FNiagaraDynamicDataLights;
+
+	DynamicData->LightArray.Empty();
+
+	for (uint32 ParticleIndex = 0; ParticleIndex < Data.GetNumInstances(); ParticleIndex++)
+	{
+		SimpleLightData LightData;
+		LightData.LightEntry.Radius = (*SizeItr).X;	//LightPayload->RadiusScale * (Size.X + Size.Y) / 2.0f;
+		LightData.LightEntry.Color = FVector((*ColItr));				//FVector(Particle.Color) * Particle.Color.A * LightPayload->ColorScale;
+		LightData.LightEntry.Exponent = 1.0;
+		LightData.LightEntry.bAffectTranslucency = true;
+		LightData.PerViewEntry.Position = (*PosItr);
+
+		DynamicData->LightArray.Add(LightData);
+
+		PosItr.Advance();
+		ColItr.Advance();
+		SizeItr.Advance();
+	}
+
+	CPUTimeMS = VertexDataTimer.GetElapsedMilliseconds();
+
+	return DynamicData;
+}
+
+
+void NiagaraRendererLights::GetDynamicMeshElements(const TArray<const FSceneView*>& Views, const FSceneViewFamily& ViewFamily, uint32 VisibilityMap, FMeshElementCollector& Collector, const FNiagaraSceneProxy *SceneProxy) const
+{
+
+}
+
+void NiagaraRendererLights::SetDynamicData_RenderThread(FNiagaraDynamicDataBase* NewDynamicData)
+{
+	if (DynamicDataRender)
+	{
+		delete static_cast<FNiagaraDynamicDataLights*>(DynamicDataRender);
+		DynamicDataRender = NULL;
+	}
+	DynamicDataRender = static_cast<FNiagaraDynamicDataLights*>(NewDynamicData);
+}
+
+int NiagaraRendererLights::GetDynamicDataSize()
+{
+	return 0;
+}
+bool NiagaraRendererLights::HasDynamicData()
+{
+	return false;
+}
+
+bool NiagaraRendererLights::SetMaterialUsage()
+{
+	return false;
+}
+
+#if WITH_EDITORONLY_DATA
+
+const TArray<FNiagaraVariable>& NiagaraRendererLights::GetRequiredAttributes()
+{
+	return Properties->GetRequiredAttributes();
+}
+
+const TArray<FNiagaraVariable>& NiagaraRendererLights::GetOptionalAttributes()
+{
+	return Properties->GetOptionalAttributes();
+}
+
+#endif
+
+
+
+
+
+
+
+
+
+
+
+
diff --git a/Engine/Plugins/FX/Niagara/Source/Niagara/Private/NiagaraRendererMeshes.cpp b/Engine/Plugins/FX/Niagara/Source/Niagara/Private/NiagaraRendererMeshes.cpp
new file mode 100644
--- /dev/null
+++ b/Engine/Plugins/FX/Niagara/Source/Niagara/Private/NiagaraRendererMeshes.cpp
@@ -0,0 +1,460 @@
+// Copyright 1998-2017 Epic Games, Inc. All Rights Reserved.
+
+#include "NiagaraRendererMeshes.h"
+#include "ParticleResources.h"
+#include "NiagaraMeshVertexFactory.h"
+#include "NiagaraDataSet.h"
+#include "NiagaraStats.h"
+#include "Async/ParallelFor.h"
+#include "Engine/StaticMesh.h"
+
+DECLARE_CYCLE_STAT(TEXT("Generate Mesh Vertex Data"), STAT_NiagaraGenMeshVertexData, STATGROUP_Niagara);
+DECLARE_CYCLE_STAT(TEXT("Render Meshes"), STAT_NiagaraRenderMeshes, STATGROUP_Niagara);
+
+
+
+extern int32 GbNiagaraParallelEmitterRenderers;
+
+
+
+
+class FNiagaraMeshCollectorResourcesMesh : public FOneFrameResource
+{
+public:
+	FNiagaraMeshVertexFactory VertexFactory;
+	FNiagaraMeshUniformBufferRef UniformBuffer;
+
+	virtual ~FNiagaraMeshCollectorResourcesMesh()
+	{
+		VertexFactory.ReleaseResource();
+	}
+};
+
+
+NiagaraRendererMeshes::NiagaraRendererMeshes(ERHIFeatureLevel::Type FeatureLevel, UNiagaraRendererProperties *InProps) :
+	NiagaraRenderer()
+{
+	//check(InProps);
+	VertexFactory = ConstructNiagaraMeshVertexFactory(NVFT_Mesh, FeatureLevel, sizeof(FNiagaraMeshInstanceVertex), 0);
+	Properties = Cast<UNiagaraMeshRendererProperties>(InProps);
+
+
+	if (Properties && Properties->ParticleMesh)
+	{
+		const FStaticMeshLODResources& LODModel = Properties->ParticleMesh->RenderData->LODResources[0];
+		for (int32 SectionIndex = 0; SectionIndex < LODModel.Sections.Num(); SectionIndex++)
+		{
+			//FMaterialRenderProxy* MaterialProxy = MeshMaterials[SectionIndex]->GetRenderProxy(bSelected);
+			const FStaticMeshSection& Section = LODModel.Sections[SectionIndex];
+			UMaterialInterface* ParticleMeshMaterial = Properties->ParticleMesh->GetMaterial(Section.MaterialIndex);
+			if (ParticleMeshMaterial)
+			{
+				ParticleMeshMaterial->CheckMaterialUsage_Concurrent(MATUSAGE_NiagaraMeshParticles);
+			}
+		}
+	}
+
+}
+
+void NiagaraRendererMeshes::SetupVertexFactory(FNiagaraMeshVertexFactory *InVertexFactory, const FStaticMeshLODResources& LODResources) const
+{
+	FNiagaraMeshVertexFactory::FDataType Data;
+
+	Data.PositionComponent = FVertexStreamComponent(
+		&LODResources.PositionVertexBuffer,
+		STRUCT_OFFSET(FPositionVertex, Position),
+		LODResources.PositionVertexBuffer.GetStride(),
+		VET_Float3
+		);
+
+	uint32 TangentXOffset = 0;
+	uint32 TangetnZOffset = 0;
+	uint32 UVsBaseOffset = 0;
+
+	SELECT_STATIC_MESH_VERTEX_TYPE(
+		LODResources.VertexBuffer.GetUseHighPrecisionTangentBasis(),
+		LODResources.VertexBuffer.GetUseFullPrecisionUVs(),
+		LODResources.VertexBuffer.GetNumTexCoords(),
+		{
+			TangentXOffset = STRUCT_OFFSET(VertexType, TangentX);
+	TangetnZOffset = STRUCT_OFFSET(VertexType, TangentZ);
+	UVsBaseOffset = STRUCT_OFFSET(VertexType, UVs);
+		});
+
+	Data.TangentBasisComponents[0] = FVertexStreamComponent(
+		&LODResources.VertexBuffer,
+		TangentXOffset,
+		LODResources.VertexBuffer.GetStride(),
+		LODResources.VertexBuffer.GetUseHighPrecisionTangentBasis() ?
+		TStaticMeshVertexTangentTypeSelector<EStaticMeshVertexTangentBasisType::HighPrecision>::VertexElementType :
+		TStaticMeshVertexTangentTypeSelector<EStaticMeshVertexTangentBasisType::Default>::VertexElementType
+		);
+
+	Data.TangentBasisComponents[1] = FVertexStreamComponent(
+		&LODResources.VertexBuffer,
+		TangetnZOffset,
+		LODResources.VertexBuffer.GetStride(),
+		LODResources.VertexBuffer.GetUseHighPrecisionTangentBasis() ?
+		TStaticMeshVertexTangentTypeSelector<EStaticMeshVertexTangentBasisType::HighPrecision>::VertexElementType :
+		TStaticMeshVertexTangentTypeSelector<EStaticMeshVertexTangentBasisType::Default>::VertexElementType
+		);
+
+	Data.TextureCoordinates.Empty();
+
+
+	uint32 UVSizeInBytes = LODResources.VertexBuffer.GetUseFullPrecisionUVs() ?
+		sizeof(TStaticMeshVertexUVsTypeSelector<EStaticMeshVertexUVType::HighPrecision>::UVsTypeT) : sizeof(TStaticMeshVertexUVsTypeSelector<EStaticMeshVertexUVType::Default>::UVsTypeT);
+
+	EVertexElementType UVVertexElementType = LODResources.VertexBuffer.GetUseFullPrecisionUVs() ?
+		VET_Float2 : VET_Half2;
+
+	uint32 NumTexCoords = FMath::Min<uint32>(LODResources.VertexBuffer.GetNumTexCoords(), MAX_TEXCOORDS);
+	for (uint32 UVIndex = 0; UVIndex < NumTexCoords; UVIndex++)
+	{
+		Data.TextureCoordinates.Add(FVertexStreamComponent(
+			&LODResources.VertexBuffer,
+			UVsBaseOffset + UVSizeInBytes * UVIndex,
+			LODResources.VertexBuffer.GetStride(),
+			UVVertexElementType
+			));
+	}
+
+	if (LODResources.ColorVertexBuffer.GetNumVertices() > 0)
+	{
+		Data.VertexColorComponent = FVertexStreamComponent(
+			&LODResources.ColorVertexBuffer,
+			0,
+			LODResources.ColorVertexBuffer.GetStride(),
+			VET_Color
+			);
+	}
+
+	// Initialize instanced data. Vertex buffer and stride are set before render.
+	// Particle color
+	Data.ParticleColorComponent = FVertexStreamComponent(
+		NULL,
+		STRUCT_OFFSET(FNiagaraMeshInstanceVertex, Color),
+		0,
+		VET_Float4,
+		true
+		);
+
+	// Particle transform matrix
+	for (int32 MatrixRow = 0; MatrixRow < 3; MatrixRow++)
+	{
+		Data.TransformComponent[MatrixRow] = FVertexStreamComponent(
+			NULL,
+			STRUCT_OFFSET(FNiagaraMeshInstanceVertex, Transform) + sizeof(FVector4)* MatrixRow,
+			0,
+			VET_Float4,
+			true
+			);
+	}
+
+	Data.VelocityComponent = FVertexStreamComponent(
+		NULL,
+		STRUCT_OFFSET(FNiagaraMeshInstanceVertex, Velocity),
+		0,
+		VET_Float4,
+		true
+		);
+	// SubUVs.
+	Data.SubUVs = FVertexStreamComponent(
+		NULL,
+		STRUCT_OFFSET(FNiagaraMeshInstanceVertex, SubUVParams),
+		0,
+		VET_Short4,
+		true
+		);
+
+	// Pack SubUV Lerp and the particle's relative time
+	Data.SubUVLerpAndRelTime = FVertexStreamComponent(
+		NULL,
+		STRUCT_OFFSET(FNiagaraMeshInstanceVertex, SubUVLerp),
+		0,
+		VET_Float2,
+		true
+		);
+
+	Data.bInitialized = true;
+	InVertexFactory->SetData(Data);
+}
+
+
+
+void NiagaraRendererMeshes::ReleaseRenderThreadResources()
+{
+	VertexFactory->ReleaseResource();
+	WorldSpacePrimitiveUniformBuffer.ReleaseResource();
+}
+
+void NiagaraRendererMeshes::CreateRenderThreadResources()
+{
+	VertexFactory->InitResource();
+}
+
+
+void NiagaraRendererMeshes::GetDynamicMeshElements(const TArray<const FSceneView*>& Views, const FSceneViewFamily& ViewFamily, uint32 VisibilityMap, FMeshElementCollector& Collector, const FNiagaraSceneProxy *SceneProxy) const
+{
+	SCOPE_CYCLE_COUNTER(STAT_NiagaraRender);
+	SCOPE_CYCLE_COUNTER(STAT_NiagaraRenderMeshes);
+
+	SimpleTimer MeshElementsTimer;
+
+	FNiagaraDynamicDataMesh *DynamicDataMesh = (static_cast<FNiagaraDynamicDataMesh*>(DynamicDataRender));
+	if (!DynamicDataMesh || DynamicDataMesh->DataSet->GetNumInstances()==0)
+	{
+		return;
+	}
+
+	{
+		// Update the primitive uniform buffer if needed.
+		if (!WorldSpacePrimitiveUniformBuffer.IsInitialized())
+		{
+			FPrimitiveUniformShaderParameters PrimitiveUniformShaderParameters = GetPrimitiveUniformShaderParameters(
+				FMatrix::Identity,
+				SceneProxy->GetActorPosition(),
+				SceneProxy->GetBounds(),
+				SceneProxy->GetLocalBounds(),
+				SceneProxy->ReceivesDecals(),
+				false,
+				false,
+				false,
+				false,
+				SceneProxy->UseEditorDepthTest(),
+				SceneProxy->GetLightingChannelMask()
+				);
+			WorldSpacePrimitiveUniformBuffer.SetContents(PrimitiveUniformShaderParameters);
+			WorldSpacePrimitiveUniformBuffer.InitResource();
+		}
+
+		// Compute the per-view uniform buffers.
+		for (int32 ViewIndex = 0; ViewIndex < Views.Num(); ViewIndex++)
+		{
+			if (VisibilityMap & (1 << ViewIndex))
+			{
+				const FSceneView* View = Views[ViewIndex];
+				const FStaticMeshLODResources& LODModel = Properties->ParticleMesh->RenderData->LODResources[0];
+
+				FNiagaraMeshCollectorResourcesMesh& CollectorResources = Collector.AllocateOneFrameResource<FNiagaraMeshCollectorResourcesMesh>();
+				SetupVertexFactory(&CollectorResources.VertexFactory, LODModel);
+				FNiagaraMeshUniformParameters PerViewUniformParameters;// = UniformParameters;
+				PerViewUniformParameters.PrevTransformAvailable = false;
+				PerViewUniformParameters.DeltaSeconds = ViewFamily.DeltaWorldTime;
+				PerViewUniformParameters.PositionDataOffset = DynamicDataMesh->PositionDataOffset;
+				PerViewUniformParameters.VelocityDataOffset = DynamicDataMesh->VelocityDataOffset;
+				PerViewUniformParameters.ColorDataOffset = DynamicDataMesh->ColorDataOffset;
+				PerViewUniformParameters.TransformDataOffset = DynamicDataMesh->TransformDataOffset;
+				PerViewUniformParameters.ScaleDataOffset = DynamicDataMesh->ScaleDataOffset;
+				PerViewUniformParameters.SizeDataOffset = DynamicDataMesh->SizeDataOffset;
+				PerViewUniformParameters.MaterialParamDataOffset = DynamicDataMesh->MaterialParamDataOffset;
+				/*
+				if (Properties)
+				{
+				PerViewUniformParameters.SubImageSize = FVector4(Properties->SubImageInfo.X, Properties->SubImageInfo.Y, 1.0f / Properties->SubImageInfo.X, 1.0f / Properties->SubImageInfo.Y);
+				}
+				*/
+				CollectorResources.VertexFactory.SetParticleData(DynamicDataMesh->DataSet);
+
+				// Collector.AllocateOneFrameResource uses default ctor, initialize the vertex factory
+				CollectorResources.VertexFactory.SetFeatureLevel(ViewFamily.GetFeatureLevel());
+				CollectorResources.VertexFactory.SetParticleFactoryType(NVFT_Mesh);
+				CollectorResources.UniformBuffer = FNiagaraMeshUniformBufferRef::CreateUniformBufferImmediate(PerViewUniformParameters, UniformBuffer_SingleFrame);
+				CollectorResources.VertexFactory.SetStrides(0, 0);
+
+				CollectorResources.VertexFactory.InitResource();
+				CollectorResources.VertexFactory.SetUniformBuffer(CollectorResources.UniformBuffer);
+			
+				const bool bIsWireframe = AllowDebugViewmodes() && View->Family->EngineShowFlags.Wireframe;
+
+				for (int32 SectionIndex = 0; SectionIndex < LODModel.Sections.Num(); SectionIndex++)
+				{
+					const FStaticMeshSection& Section = LODModel.Sections[SectionIndex];
+					UMaterialInterface* ParticleMeshMaterial = Properties->ParticleMesh->GetMaterial(Section.MaterialIndex);
+					if (!ParticleMeshMaterial)
+					{
+						continue;
+					}
+					FMaterialRenderProxy* MaterialProxy = ParticleMeshMaterial->GetRenderProxy(false, false);
+
+					if ((Section.NumTriangles == 0) || (MaterialProxy == NULL))
+					{
+						//@todo. This should never occur, but it does occasionally.
+						continue;
+					}
+
+
+					FMeshBatch& Mesh = Collector.AllocateMesh();
+					Mesh.VertexFactory = &CollectorResources.VertexFactory;
+					Mesh.DynamicVertexData = NULL;
+					Mesh.LCI = NULL;
+					Mesh.UseDynamicData = false;
+					Mesh.ReverseCulling = SceneProxy->IsLocalToWorldDeterminantNegative();
+					Mesh.CastShadow = SceneProxy->CastsDynamicShadow();
+					Mesh.DepthPriorityGroup = (ESceneDepthPriorityGroup)SceneProxy->GetDepthPriorityGroup(View);
+
+					FMeshBatchElement& BatchElement = Mesh.Elements[0];
+					BatchElement.PrimitiveUniformBufferResource = &WorldSpacePrimitiveUniformBuffer;
+					BatchElement.FirstIndex = 0;
+					BatchElement.MinVertexIndex = 0;
+					BatchElement.MaxVertexIndex = 0;
+					BatchElement.NumInstances = DynamicDataMesh->DataSet->PrevDataRender().GetNumInstances();
+
+					if (bIsWireframe)
+					{
+						if (LODModel.WireframeIndexBuffer.IsInitialized())
+						{
+							Mesh.Type = PT_LineList;
+							Mesh.MaterialRenderProxy = UMaterial::GetDefaultMaterial(MD_Surface)->GetRenderProxy(SceneProxy->IsSelected(), SceneProxy->IsHovered());
+							BatchElement.FirstIndex = 0;
+							BatchElement.IndexBuffer = &LODModel.WireframeIndexBuffer;
+							BatchElement.NumPrimitives = LODModel.WireframeIndexBuffer.GetNumIndices() / 2;
+
+						}
+						else
+						{
+							Mesh.Type = PT_TriangleList;
+							Mesh.MaterialRenderProxy = MaterialProxy;
+							Mesh.bWireframe = true;
+							BatchElement.FirstIndex = 0;
+							BatchElement.IndexBuffer = &LODModel.IndexBuffer;
+							BatchElement.NumPrimitives = LODModel.IndexBuffer.GetNumIndices() / 3;
+						}
+					}
+					else
+					{
+						Mesh.Type = PT_TriangleList;
+						Mesh.MaterialRenderProxy = MaterialProxy;
+						BatchElement.IndexBuffer = &LODModel.IndexBuffer;
+						BatchElement.FirstIndex = Section.FirstIndex;
+						BatchElement.NumPrimitives = Section.NumTriangles;
+					}
+
+					Mesh.bCanApplyViewModeOverrides = true;
+					Mesh.bUseWireframeSelectionColoring = SceneProxy->IsSelected();
+					Collector.AddMesh(ViewIndex, Mesh);
+				}
+			}
+		}
+	}
+
+	CPUTimeMS += MeshElementsTimer.GetElapsedMilliseconds();
+}
+
+
+
+bool NiagaraRendererMeshes::SetMaterialUsage()
+{
+	//Causes deadlock :S Need to look at / rework the setting of materials and render modules.
+	return Material && Material->CheckMaterialUsage_Concurrent(MATUSAGE_NiagaraMeshParticles);
+}
+
+
+
+/** Update render data buffer from attributes */
+FNiagaraDynamicDataBase *NiagaraRendererMeshes::GenerateVertexData(const FNiagaraSceneProxy* Proxy, FNiagaraDataSet &Data, const ENiagaraSimTarget Target)
+{
+	SCOPE_CYCLE_COUNTER(STAT_NiagaraGenMeshVertexData);
+
+	if (!Properties || Properties->ParticleMesh == nullptr || !bEnabled)
+	{
+		return nullptr;
+	}
+
+	SimpleTimer VertexDataTimer;
+//	TArray<FNiagaraMeshInstanceVertex>& RenderData = DynamicData->VertexData;
+	//TArray< FNiagaraMeshInstanceVertexDynamicParameter>& RenderMaterialVertexData = DynamicData->MaterialParameterVertexData;
+
+
+	const FNiagaraVariableLayoutInfo* PositionLayout = Data.GetVariableLayout(FNiagaraVariable(FNiagaraTypeDefinition::GetVec3Def(), TEXT("Position")));
+	const FNiagaraVariableLayoutInfo* VelocityLayout = Data.GetVariableLayout(FNiagaraVariable(FNiagaraTypeDefinition::GetVec3Def(), TEXT("Velocity")));
+	const FNiagaraVariableLayoutInfo* ColorLayout = Data.GetVariableLayout(FNiagaraVariable(FNiagaraTypeDefinition::GetColorDef(), TEXT("Color")));
+	const FNiagaraVariableLayoutInfo* XformLayout = Data.GetVariableLayout(FNiagaraVariable(FNiagaraTypeDefinition::GetVec4Def(), TEXT("Transform")));
+	const FNiagaraVariableLayoutInfo* ScaleLayout = Data.GetVariableLayout(FNiagaraVariable(FNiagaraTypeDefinition::GetVec3Def(), TEXT("Scale")));
+	const FNiagaraVariableLayoutInfo* SizeLayout = Data.GetVariableLayout(FNiagaraVariable(FNiagaraTypeDefinition::GetVec2Def(), TEXT("Size")));
+	const FNiagaraVariableLayoutInfo* MatParamLayout = Data.GetVariableLayout(FNiagaraVariable(FNiagaraTypeDefinition::GetVec4Def(), TEXT("DynamicMaterialParameter")));
+
+	//Bail if we don't have the required attributes to render this emitter.
+	if (!bEnabled || !PositionLayout /*|| !XformLayout*/ || !VelocityLayout)
+	{
+		return nullptr;
+	}
+
+
+	UStaticMesh *Mesh = nullptr;
+	if (Properties)
+	{
+		Mesh = Properties->ParticleMesh;
+	}
+
+	// required attributes
+	FNiagaraDynamicDataMesh *DynamicData = new FNiagaraDynamicDataMesh;
+	DynamicData->PositionDataOffset = PositionLayout->FloatComponentStart;
+	DynamicData->VelocityDataOffset = VelocityLayout->FloatComponentStart;
+
+	// optional attributes
+	int32 IntDummy;
+	Data.GetVariableComponentOffsets(FNiagaraVariable(FNiagaraTypeDefinition::GetColorDef(), TEXT("Color")), DynamicData->ColorDataOffset, IntDummy);
+	Data.GetVariableComponentOffsets(FNiagaraVariable(FNiagaraTypeDefinition::GetVec3Def(), TEXT("Scale")), DynamicData->ScaleDataOffset, IntDummy);
+	Data.GetVariableComponentOffsets(FNiagaraVariable(FNiagaraTypeDefinition::GetVec2Def(), TEXT("Size")), DynamicData->SizeDataOffset, IntDummy);
+	Data.GetVariableComponentOffsets(FNiagaraVariable(FNiagaraTypeDefinition::GetVec4Def(), TEXT("DynamicMaterialParameter")), DynamicData->MaterialParamDataOffset, IntDummy);
+	Data.GetVariableComponentOffsets(FNiagaraVariable(FNiagaraTypeDefinition::GetVec4Def(), TEXT("Transform")), DynamicData->TransformDataOffset, IntDummy);
+
+	// if we're CPU simulating, need to init the GPU buffers for the vertex factory
+	if (Data.PrevData().GetNumInstances() > 0)
+	{
+		if (Target == ENiagaraSimTarget::CPUSim)
+		{
+			Data.ValidateBufferIndices();
+			Data.InitGPUFromCPU();
+		}
+	}
+
+	DynamicData->DataSet = &Data;
+	CPUTimeMS = VertexDataTimer.GetElapsedMilliseconds();
+	return DynamicData;  
+}
+
+
+
+void NiagaraRendererMeshes::SetDynamicData_RenderThread(FNiagaraDynamicDataBase* NewDynamicData)
+{
+	check(IsInRenderingThread());
+
+	if (DynamicDataRender)
+	{
+		delete static_cast<FNiagaraDynamicDataMesh*>(DynamicDataRender);
+		DynamicDataRender = NULL;
+	}
+	DynamicDataRender = NewDynamicData;
+}
+
+int NiagaraRendererMeshes::GetDynamicDataSize()
+{
+	uint32 Size = sizeof(FNiagaraDynamicDataMesh);
+	if (DynamicDataRender && static_cast<FNiagaraDynamicDataMesh*>(DynamicDataRender)->DataSet)
+	{
+		//Size += (static_cast<FNiagaraDynamicDataMesh*>(DynamicDataRender))->DataSet->PrevDataRender().GetNumInstances() * sizeof(float);
+	}
+
+	return Size;
+}
+
+bool NiagaraRendererMeshes::HasDynamicData()
+{
+	return DynamicDataRender != nullptr;
+}
+
+#if WITH_EDITORONLY_DATA
+
+const TArray<FNiagaraVariable>& NiagaraRendererMeshes::GetRequiredAttributes()
+{
+	return Properties->GetRequiredAttributes();
+}
+
+const TArray<FNiagaraVariable>& NiagaraRendererMeshes::GetOptionalAttributes()
+{
+	return Properties->GetOptionalAttributes();
+}
+
+#endif
\ No newline at end of file
diff --git a/Engine/Plugins/FX/Niagara/Source/Niagara/Private/NiagaraRendererRibbons.cpp b/Engine/Plugins/FX/Niagara/Source/Niagara/Private/NiagaraRendererRibbons.cpp
new file mode 100644
--- /dev/null
+++ b/Engine/Plugins/FX/Niagara/Source/Niagara/Private/NiagaraRendererRibbons.cpp
@@ -0,0 +1,407 @@
+// Copyright 1998-2017 Epic Games, Inc. All Rights Reserved.
+
+#include "NiagaraRendererRibbons.h"
+#include "ParticleResources.h"
+#include "NiagaraRibbonVertexFactory.h"
+#include "NiagaraDataSet.h"
+#include "NiagaraStats.h"
+
+DECLARE_CYCLE_STAT(TEXT("Generate Ribbon Vertex Data"), STAT_NiagaraGenRibbonVertexData, STATGROUP_Niagara);
+DECLARE_CYCLE_STAT(TEXT("Render Ribbons"), STAT_NiagaraRenderRibbons, STATGROUP_Niagara);
+
+
+
+struct FNiagaraDynamicDataRibbon : public FNiagaraDynamicDataBase
+{
+	TArray<FNiagaraRibbonVertex> VertexData;
+	TArray<int16> IndexData;
+	TArray<FNiagaraRibbonVertexDynamicParameter> MaterialParameterVertexData;
+};
+
+class FNiagaraMeshCollectorResourcesRibbon : public FOneFrameResource
+{
+public:
+	FNiagaraRibbonVertexFactory VertexFactory;
+	FNiagaraRibbonUniformBufferRef UniformBuffer;
+
+	virtual ~FNiagaraMeshCollectorResourcesRibbon()
+	{
+		VertexFactory.ReleaseResource();
+	}
+};
+
+
+NiagaraRendererRibbons::NiagaraRendererRibbons(ERHIFeatureLevel::Type FeatureLevel, UNiagaraRendererProperties *InProps) :
+	NiagaraRenderer()
+{
+	VertexFactory = new FNiagaraRibbonVertexFactory(NVFT_Ribbon, FeatureLevel);
+	Properties = Cast<UNiagaraRibbonRendererProperties>(InProps);
+}
+
+
+void NiagaraRendererRibbons::ReleaseRenderThreadResources()
+{
+	VertexFactory->ReleaseResource();
+	WorldSpacePrimitiveUniformBuffer.ReleaseResource();
+}
+
+// FPrimitiveSceneProxy interface.
+void NiagaraRendererRibbons::CreateRenderThreadResources()
+{
+	VertexFactory->InitResource();
+}
+
+
+
+
+void NiagaraRendererRibbons::GetDynamicMeshElements(const TArray<const FSceneView*>& Views, const FSceneViewFamily& ViewFamily, uint32 VisibilityMap, FMeshElementCollector& Collector, const FNiagaraSceneProxy *SceneProxy) const
+{
+	SCOPE_CYCLE_COUNTER(STAT_NiagaraRender);
+	SCOPE_CYCLE_COUNTER(STAT_NiagaraRenderRibbons);
+
+	SimpleTimer MeshElementsTimer;
+	FNiagaraDynamicDataRibbon *DynamicDataRibbon = static_cast<FNiagaraDynamicDataRibbon*>(DynamicDataRender);
+	if (!DynamicDataRibbon || DynamicDataRibbon->VertexData.Num() == 0)
+	{
+		return;
+	}
+
+	const bool bIsWireframe = ViewFamily.EngineShowFlags.Wireframe;
+	FMaterialRenderProxy* MaterialRenderProxy = Material->GetRenderProxy(SceneProxy->IsSelected(), SceneProxy->IsHovered());
+
+	int32 SizeInBytes = DynamicDataRibbon->VertexData.GetTypeSize() * DynamicDataRibbon->VertexData.Num();
+	FGlobalDynamicVertexBuffer::FAllocation LocalDynamicVertexAllocation = FGlobalDynamicVertexBuffer::Get().Allocate(SizeInBytes);
+	FGlobalDynamicVertexBuffer::FAllocation LocalDynamicVertexMaterialParamsAllocation;
+	FGlobalDynamicIndexBuffer::FAllocation DynamicIndexAllocation = FGlobalDynamicIndexBuffer::Get().Allocate(DynamicDataRibbon->IndexData.Num(), sizeof(int16));
+
+	if (DynamicDataRibbon->MaterialParameterVertexData.Num() > 0)
+	{
+		int32 MatParamSizeInBytes = DynamicDataRibbon->MaterialParameterVertexData.GetTypeSize() * DynamicDataRibbon->MaterialParameterVertexData.Num();
+		LocalDynamicVertexMaterialParamsAllocation = FGlobalDynamicVertexBuffer::Get().Allocate(MatParamSizeInBytes);
+
+		if (LocalDynamicVertexMaterialParamsAllocation.IsValid())
+		{
+			// Copy the extra material vertex data over.
+			FMemory::Memcpy(LocalDynamicVertexMaterialParamsAllocation.Buffer, DynamicDataRibbon->MaterialParameterVertexData.GetData(), MatParamSizeInBytes);
+		}
+	}
+
+
+	if (LocalDynamicVertexAllocation.IsValid())
+	{
+		// Update the primitive uniform buffer if needed.
+		if (!WorldSpacePrimitiveUniformBuffer.IsInitialized())
+		{
+			FPrimitiveUniformShaderParameters PrimitiveUniformShaderParameters = GetPrimitiveUniformShaderParameters(
+				FMatrix::Identity,
+				SceneProxy->GetActorPosition(),
+				SceneProxy->GetBounds(),
+				SceneProxy->GetLocalBounds(),
+				SceneProxy->ReceivesDecals(),
+				false,
+				false,
+				SceneProxy->UseSingleSampleShadowFromStationaryLights(),
+				SceneProxy->GetScene().HasPrecomputedVolumetricLightmap_RenderThread(),
+				SceneProxy->UseEditorDepthTest(),
+				SceneProxy->GetLightingChannelMask()
+				);
+			WorldSpacePrimitiveUniformBuffer.SetContents(PrimitiveUniformShaderParameters);
+			WorldSpacePrimitiveUniformBuffer.InitResource();
+		}
+
+		// Copy the vertex data over.
+		FMemory::Memcpy(LocalDynamicVertexAllocation.Buffer, DynamicDataRibbon->VertexData.GetData(), SizeInBytes);
+		FMemory::Memcpy(DynamicIndexAllocation.Buffer, DynamicDataRibbon->IndexData.GetData(), DynamicDataRibbon->IndexData.Num() * sizeof(int16));
+
+		// Compute the per-view uniform buffers.
+		for (int32 ViewIndex = 0; ViewIndex < Views.Num(); ViewIndex++)
+		{
+			if (VisibilityMap & (1 << ViewIndex))
+			{
+				const FSceneView* View = Views[ViewIndex];
+
+				FNiagaraMeshCollectorResourcesRibbon& CollectorResources = Collector.AllocateOneFrameResource<FNiagaraMeshCollectorResourcesRibbon>();
+				FNiagaraRibbonUniformParameters PerViewUniformParameters;// = UniformParameters;
+				PerViewUniformParameters.CameraUp = View->GetViewUp(); // FVector4(0.0f, 0.0f, 1.0f, 0.0f);
+				PerViewUniformParameters.CameraRight = View->GetViewRight();//	FVector4(1.0f, 0.0f, 0.0f, 0.0f);
+				PerViewUniformParameters.ScreenAlignment = FVector4(0.0f, 0.0f, 0.0f, 0.0f);
+
+				// Collector.AllocateOneFrameResource uses default ctor, initialize the vertex factory
+				CollectorResources.VertexFactory.SetFeatureLevel(ViewFamily.GetFeatureLevel());
+				CollectorResources.VertexFactory.SetParticleFactoryType(NVFT_Ribbon);
+
+				CollectorResources.UniformBuffer = FNiagaraRibbonUniformBufferRef::CreateUniformBufferImmediate(PerViewUniformParameters, UniformBuffer_SingleFrame);
+
+				CollectorResources.VertexFactory.InitResource();
+				CollectorResources.VertexFactory.SetBeamTrailUniformBuffer(CollectorResources.UniformBuffer);
+				CollectorResources.VertexFactory.SetVertexBuffer(LocalDynamicVertexAllocation.VertexBuffer, LocalDynamicVertexAllocation.VertexOffset, sizeof(FNiagaraRibbonVertex));
+				if (DynamicDataRibbon->MaterialParameterVertexData.Num() > 0 && LocalDynamicVertexMaterialParamsAllocation.IsValid())
+				{
+					CollectorResources.VertexFactory.SetDynamicParameterBuffer(LocalDynamicVertexMaterialParamsAllocation.VertexBuffer, LocalDynamicVertexMaterialParamsAllocation.VertexOffset, sizeof(FNiagaraRibbonVertexDynamicParameter));
+				}
+				else
+				{
+					CollectorResources.VertexFactory.SetDynamicParameterBuffer(NULL, 0, 0);
+				}
+
+				FMeshBatch& MeshBatch = Collector.AllocateMesh();
+				MeshBatch.VertexFactory = &CollectorResources.VertexFactory;
+				MeshBatch.CastShadow = SceneProxy->CastsDynamicShadow();
+				MeshBatch.bUseAsOccluder = false;
+				MeshBatch.ReverseCulling = SceneProxy->IsLocalToWorldDeterminantNegative();
+				MeshBatch.bDisableBackfaceCulling = true;
+				MeshBatch.Type = PT_TriangleStrip;
+				MeshBatch.DepthPriorityGroup = SceneProxy->GetDepthPriorityGroup(View);
+				MeshBatch.bCanApplyViewModeOverrides = true;
+				MeshBatch.bUseWireframeSelectionColoring = SceneProxy->IsSelected();
+
+				if (bIsWireframe)
+				{
+					MeshBatch.MaterialRenderProxy = UMaterial::GetDefaultMaterial(MD_Surface)->GetRenderProxy(SceneProxy->IsSelected(), SceneProxy->IsHovered());
+				}
+				else
+				{
+					MeshBatch.MaterialRenderProxy = MaterialRenderProxy;
+				}
+
+				FMeshBatchElement& MeshElement = MeshBatch.Elements[0];
+				MeshElement.IndexBuffer = DynamicIndexAllocation.IndexBuffer;
+				MeshElement.FirstIndex = 0;
+				MeshElement.NumPrimitives = DynamicDataRibbon->IndexData.Num() - 2;
+				MeshElement.NumInstances = 1;
+				MeshElement.MinVertexIndex = 0;
+				MeshElement.MaxVertexIndex = DynamicDataRibbon->VertexData.Num() - 1;
+				MeshElement.PrimitiveUniformBufferResource = &WorldSpacePrimitiveUniformBuffer;
+
+				Collector.AddMesh(ViewIndex, MeshBatch);
+			}
+		}
+	}
+
+	CPUTimeMS += MeshElementsTimer.GetElapsedMilliseconds();
+}
+
+void NiagaraRendererRibbons::SetDynamicData_RenderThread(FNiagaraDynamicDataBase* NewDynamicData)
+{
+	check(IsInRenderingThread());
+
+	if (DynamicDataRender)
+	{
+		delete static_cast<FNiagaraDynamicDataRibbon*>(DynamicDataRender);
+		DynamicDataRender = NULL;
+	}
+	DynamicDataRender = NewDynamicData;
+}
+
+int NiagaraRendererRibbons::GetDynamicDataSize()
+{
+	uint32 Size = sizeof(FNiagaraDynamicDataRibbon);
+	if (DynamicDataRender)
+	{
+		Size += (static_cast<FNiagaraDynamicDataRibbon*>(DynamicDataRender))->VertexData.GetAllocatedSize();
+	}
+
+	return Size;
+}
+
+bool NiagaraRendererRibbons::HasDynamicData()
+{
+	return DynamicDataRender && (static_cast<FNiagaraDynamicDataRibbon*>(DynamicDataRender))->VertexData.Num() > 0;
+}
+
+#if WITH_EDITORONLY_DATA
+
+const TArray<FNiagaraVariable>& NiagaraRendererRibbons::GetRequiredAttributes()
+{
+	return Properties->GetRequiredAttributes();
+}
+
+const TArray<FNiagaraVariable>& NiagaraRendererRibbons::GetOptionalAttributes()
+{
+	return Properties->GetOptionalAttributes();
+}
+
+#endif
+
+
+bool NiagaraRendererRibbons::SetMaterialUsage()
+{
+	return Material && Material->CheckMaterialUsage_Concurrent(MATUSAGE_NiagaraRibbons);
+}
+
+FNiagaraDynamicDataBase *NiagaraRendererRibbons::GenerateVertexData(const FNiagaraSceneProxy* Proxy, FNiagaraDataSet &Data, const ENiagaraSimTarget Target)
+{
+	SCOPE_CYCLE_COUNTER(STAT_NiagaraGenRibbonVertexData);
+
+	SimpleTimer VertexDataTimer;
+	if (Data.GetNumInstances() < 2 || !bEnabled)
+	{
+		return nullptr;
+	}
+	FNiagaraDynamicDataRibbon *DynamicData = new FNiagaraDynamicDataRibbon;
+	TArray<FNiagaraRibbonVertex>& RenderData = DynamicData->VertexData;
+	TArray<int16>& IndexData = DynamicData->IndexData;
+
+	//I'm not a great fan of pulling scalar components out to a structured vert buffer like this.
+	//TODO: Experiment with a new VF that reads the data directly from the scalar layout.
+	FNiagaraDataSetIterator<FVector> PosItr(Data, FNiagaraVariable(FNiagaraTypeDefinition::GetVec3Def(), TEXT("Position")));
+	FNiagaraDataSetIterator<FVector> VelItr(Data, FNiagaraVariable(FNiagaraTypeDefinition::GetVec3Def(), TEXT("Velocity")));
+	FNiagaraDataSetIterator<FLinearColor> ColItr(Data, FNiagaraVariable(FNiagaraTypeDefinition::GetColorDef(), TEXT("Color")));
+	FNiagaraDataSetIterator<float> NormAgeItr(Data, FNiagaraVariable(FNiagaraTypeDefinition::GetFloatDef(), TEXT("NormalizedAge")));
+	FNiagaraDataSetIterator<float> RotItr(Data, FNiagaraVariable(FNiagaraTypeDefinition::GetFloatDef(), TEXT("Rotation")));
+	FNiagaraDataSetIterator<int32> RibbonIdItr(Data, FNiagaraVariable(FNiagaraTypeDefinition::GetIntDef(), TEXT("RibbonId")));
+	FNiagaraDataSetIterator<FVector2D> SizeItr(Data, FNiagaraVariable(FNiagaraTypeDefinition::GetVec2Def(), TEXT("Size")));
+	//FNiagaraDataSetIterator<float> SubImageIndexItr(Data, FNiagaraVariable(FNiagaraTypeDefinition::GetFloatDef(), TEXT("SubImageIndex")));
+	FNiagaraDataSetIterator<FVector4> MaterialParamItr(Data, FNiagaraVariable(FNiagaraTypeDefinition::GetVec4Def(), TEXT("DynamicMaterialParameter")));
+
+	//Bail if we don't have the required attributes to render this emitter.
+	if (!PosItr.IsValid() || !ColItr.IsValid() || !NormAgeItr.IsValid() || !SizeItr.IsValid())
+	{
+		return DynamicData;
+	}
+
+	bool bMultiRibbons = RibbonIdItr.IsValid();
+
+	// build a sorted list by age, so we always get particles in order regardless of them being moved around due to dieing and spawning
+	// TODO: need to pull the data out because we currently can't index into the data set; terrible, but works
+	TArray<int32> SortedIndices;
+	TArray<int32> RibbonIDs;
+	TArray<int32> UniqueRibbonIDs;
+	TArray<FVector> Positions;
+	TArray<FVector2D> Sizes;
+	TArray<FLinearColor> Colors;
+	TArray<float> Rotations;
+	TArray<float> Ages;
+	uint32 CurIdx = 0;
+	for (uint32 Idx = 0; Idx < Data.GetNumInstances(); Idx++)
+	{
+		float Age = NormAgeItr.Get();
+		if (Idx>0 && Age > 0.0f && Age < 1.0f)
+		{
+			SortedIndices.Add(CurIdx++);
+			Positions.Add(PosItr.Get());
+			Sizes.Add(SizeItr.Get());
+			Rotations.Add(RotItr.Get());
+			Ages.Add(Age);
+			Colors.Add(ColItr.Get());
+			if (bMultiRibbons)
+			{
+				int32 RibbonID = RibbonIdItr.Get();
+				RibbonIDs.Add(RibbonID);
+				UniqueRibbonIDs.AddUnique(RibbonID);
+			}
+		}
+		PosItr.Advance();
+		SizeItr.Advance();
+		RotItr.Advance();
+		NormAgeItr.Advance();
+		ColItr.Advance();
+		RibbonIdItr.Advance();
+	}
+
+	if (!bMultiRibbons)
+	{
+		SortedIndices.Sort(
+			[&Ages](const int32& A, const int32& B) {
+			return (Ages[A] > Ages[B]);
+		}
+		);
+		UniqueRibbonIDs.Add(0);
+	}
+	else
+	{
+		SortedIndices.Sort(
+			[&RibbonIDs](const int32& A, const int32& B) {
+			return (RibbonIDs[A] < RibbonIDs[B]);
+		}
+		);
+	}
+
+
+	if (SortedIndices.Num() == 0)
+	{
+		return DynamicData;
+	}
+
+	RenderData.Reset(SortedIndices.Num() * 2);
+	IndexData.Empty();
+
+	// TODO : deal with the dynamic vertex material parameter should the user have specified it as an output...
+
+	FVector2D UVs[4] = { FVector2D(0.0f, 0.0f), FVector2D(1.0f, 0.0f), FVector2D(1.0f, 1.0f), FVector2D(0.0f, 1.0f) };
+	int32 StartIndex = 0;
+	int32 NumIndices = SortedIndices.Num();
+	int32 *IndexPtr = &SortedIndices[0];
+	int32 NumTotalVerts = 0;
+
+	for (int32 RibbonIdx = 0; RibbonIdx < UniqueRibbonIDs.Num(); RibbonIdx++)
+	{
+		TArray<int32> SortedPartialIndices;
+		if (bMultiRibbons)
+		{
+			int32 CurID = RibbonIDs[SortedIndices[StartIndex]];
+			NumIndices = 0;
+			while ((StartIndex + NumIndices)<RibbonIDs.Num() && RibbonIDs[SortedIndices[StartIndex + NumIndices]] == CurID)
+			{
+				SortedPartialIndices.Add(SortedIndices[StartIndex + NumIndices]);
+				NumIndices++;
+			}
+			SortedPartialIndices.Sort(
+				[&Ages](const int32& A, const int32& B) {
+				return (Ages[A] > Ages[B]);
+			}
+			);
+			IndexPtr = &SortedPartialIndices[0];
+		}
+
+		FVector PrevPos, PrevPos2, PrevDir(0.0f, 0.0f, 0.1f);
+		for (int32 i = 0; i < NumIndices - 1; i++)
+		{
+			uint32 Index1 = IndexPtr[i];
+			uint32 Index2 = IndexPtr[i + 1];
+
+			const FVector ParticlePos = Positions[Index1];
+			FVector ParticleDir = Positions[Index2] - ParticlePos;
+
+			// if two ribbon particles were spawned too close together, we skip one
+			if (ParticleDir.SizeSquared() <= FMath::Square(SMALL_NUMBER))
+			{
+				continue;
+			}
+			FVector NormDir = ParticleDir.GetSafeNormal();
+			PrevDir = NormDir;
+
+			FVector ParticleRight = FVector::CrossProduct(NormDir, FVector(0.0f, 0.0f, 1.0f));
+			ParticleRight *= Sizes[Index1].X;
+			FVector ParticleRightRot = ParticleRight.RotateAngleAxis(Rotations[Index1], NormDir);
+
+
+			AddRibbonVert(RenderData, ParticlePos + ParticleRightRot, UVs[0], Colors[Index1], Ages[Index1], Rotations[Index1], Sizes[Index1]);
+			AddRibbonVert(RenderData, ParticlePos - ParticleRightRot, UVs[1], Colors[Index1], Ages[Index1], Rotations[Index1], Sizes[Index1]);
+
+			IndexData.Add(NumTotalVerts);
+			IndexData.Add(NumTotalVerts + 1);
+			NumTotalVerts += 2;
+
+			PrevPos = ParticlePos;
+		}
+
+		StartIndex += NumIndices;
+
+		// add degenerates to connect with the next ribbon
+		if (RibbonIdx < UniqueRibbonIDs.Num() - 1)
+		{
+			IndexData.Add(FMath::Max(NumTotalVerts - 1, 0));
+			IndexData.Add(FMath::Max(NumTotalVerts - 1, 0));
+			IndexData.Add(FMath::Max(NumTotalVerts, 0));
+			IndexData.Add(FMath::Max(NumTotalVerts, 0));
+		}
+	}
+
+	RenderData.SetNum(NumTotalVerts);
+	CPUTimeMS = VertexDataTimer.GetElapsedMilliseconds();
+
+	return DynamicData;
+}
+
+
diff --git a/Engine/Plugins/FX/Niagara/Source/Niagara/Private/NiagaraRendererSprites.cpp b/Engine/Plugins/FX/Niagara/Source/Niagara/Private/NiagaraRendererSprites.cpp
new file mode 100644
--- /dev/null
+++ b/Engine/Plugins/FX/Niagara/Source/Niagara/Private/NiagaraRendererSprites.cpp
@@ -0,0 +1,408 @@
+// Copyright 1998-2017 Epic Games, Inc. All Rights Reserved.
+
+#include "NiagaraRenderer.h"
+#include "ParticleResources.h"
+#include "NiagaraSpriteVertexFactory.h"
+#include "NiagaraDataSet.h"
+#include "NiagaraStats.h"
+
+DECLARE_CYCLE_STAT(TEXT("Generate Sprite Vertex Data"), STAT_NiagaraGenSpriteVertexData, STATGROUP_Niagara);
+DECLARE_CYCLE_STAT(TEXT("Render Sprites"), STAT_NiagaraRenderSprites, STATGROUP_Niagara);
+
+DECLARE_CYCLE_STAT(TEXT("Genereate GPU Buffers"), STAT_NiagaraGenGpuBuffers, STATGROUP_Niagara);
+
+
+/**
+* Per-particle data sent to the GPU.
+*/
+struct FNiagaraSpriteVertex
+{
+	/** The position of the particle. */
+	FVector Position;
+	/** The relative time of the particle. */
+	float RelativeTime;
+	/** The previous position of the particle. */
+	FVector	OldPosition;
+	/** Value that remains constant over the lifetime of a particle. */
+	float ParticleId;
+	/** The size of the particle. */
+	FVector2D Size;
+	/** The rotation of the particle. */
+	float Rotation;
+	/** The sub-image index for the particle. */
+	float SubImageIndex;
+	/** The color of the particle. */
+	FLinearColor Color;
+	/* Custom Alignment vector*/
+	FVector CustomAlignmentVector;
+	/* Custom Facing vector*/
+	FVector CustomFacingVector;
+};
+
+
+struct FNiagaraDynamicDataSprites : public FNiagaraDynamicDataBase
+{
+	TArray<FNiagaraSpriteVertex> VertexData;
+	TArray<FParticleVertexDynamicParameter> MaterialParameterVertexData;
+	const FNiagaraDataSet *DataSet;
+	//uint32 ComponentBufferSize;
+	int32 PositionDataOffset;
+	int32 VelocityDataOffset;
+	int32 SizeDataOffset;
+	int32 RotationDataOffset;
+	int32 SubimageDataOffset;
+	int32 ColorDataOffset;
+	int32 FacingOffset;
+	int32 AlignmentOffset;
+	//int32 NumInstances;
+	bool bCustomAlignmentAvailable;
+};
+
+
+
+/* Mesh collector classes */
+class FNiagaraMeshCollectorResourcesSprite : public FOneFrameResource
+{
+public:
+	FNiagaraSpriteVertexFactory VertexFactory;
+	FNiagaraSpriteUniformBufferRef UniformBuffer;
+
+	virtual ~FNiagaraMeshCollectorResourcesSprite()
+	{
+		VertexFactory.ReleaseResource();
+	}
+};
+
+
+
+NiagaraRendererSprites::NiagaraRendererSprites(ERHIFeatureLevel::Type FeatureLevel, UNiagaraRendererProperties *InProps) :
+	NiagaraRenderer()
+{
+	//check(InProps);
+	VertexFactory = new FNiagaraSpriteVertexFactory(NVFT_Sprite, FeatureLevel);
+	Properties = Cast<UNiagaraSpriteRendererProperties>(InProps);
+}
+
+
+void NiagaraRendererSprites::ReleaseRenderThreadResources()
+{
+	VertexFactory->ReleaseResource();
+	WorldSpacePrimitiveUniformBuffer.ReleaseResource();
+}
+
+void NiagaraRendererSprites::CreateRenderThreadResources()
+{
+	VertexFactory->SetNumVertsInInstanceBuffer(4);
+	VertexFactory->InitResource();
+}
+
+void NiagaraRendererSprites::GetDynamicMeshElements(const TArray<const FSceneView*>& Views, const FSceneViewFamily& ViewFamily, uint32 VisibilityMap, FMeshElementCollector& Collector, const FNiagaraSceneProxy *SceneProxy) const
+{
+	SCOPE_CYCLE_COUNTER(STAT_NiagaraRender);
+	SCOPE_CYCLE_COUNTER(STAT_NiagaraRenderSprites);
+
+	SimpleTimer MeshElementsTimer;
+
+	//check(DynamicDataRender)
+
+	FNiagaraDynamicDataSprites *DynamicDataSprites = static_cast<FNiagaraDynamicDataSprites*>(DynamicDataRender);
+	if (!DynamicDataSprites || DynamicDataSprites->DataSet->PrevDataRender().GetNumInstances()==0 || nullptr == Properties)
+	{
+		return;
+	}
+
+	int32 NumInstances = DynamicDataSprites->DataSet->PrevDataRender().GetNumInstances();
+
+	const bool bIsWireframe = ViewFamily.EngineShowFlags.Wireframe;
+	FMaterialRenderProxy* MaterialRenderProxy = Material->GetRenderProxy(SceneProxy->IsSelected(), SceneProxy->IsHovered());
+
+	int32 SizeInBytes = DynamicDataSprites->VertexData.GetTypeSize() * DynamicDataSprites->VertexData.Num();
+	FGlobalDynamicVertexBuffer::FAllocation LocalDynamicVertexMaterialParamsAllocation;
+
+	if (DynamicDataSprites->MaterialParameterVertexData.Num() > 0)
+	{
+		int32 MatParamSizeInBytes = DynamicDataSprites->MaterialParameterVertexData.GetTypeSize() * DynamicDataSprites->MaterialParameterVertexData.Num();
+		LocalDynamicVertexMaterialParamsAllocation = FGlobalDynamicVertexBuffer::Get().Allocate(MatParamSizeInBytes);
+
+		if (LocalDynamicVertexMaterialParamsAllocation.IsValid())
+		{
+			// Copy the extra material vertex data over.
+			FMemory::Memcpy(LocalDynamicVertexMaterialParamsAllocation.Buffer, DynamicDataSprites->MaterialParameterVertexData.GetData(), MatParamSizeInBytes);
+		}
+	}
+
+	// Update the primitive uniform buffer if needed.
+	if (!WorldSpacePrimitiveUniformBuffer.IsInitialized())
+	{
+		FPrimitiveUniformShaderParameters PrimitiveUniformShaderParameters = GetPrimitiveUniformShaderParameters(
+			FMatrix::Identity,
+			SceneProxy->GetActorPosition(),
+			SceneProxy->GetBounds(),
+			SceneProxy->GetLocalBounds(),
+			SceneProxy->ReceivesDecals(),
+			false,
+			false,
+			SceneProxy->UseSingleSampleShadowFromStationaryLights(),
+			SceneProxy->GetScene().HasPrecomputedVolumetricLightmap_RenderThread(),
+			SceneProxy->UseEditorDepthTest(),
+			SceneProxy->GetLightingChannelMask()
+			);
+		WorldSpacePrimitiveUniformBuffer.SetContents(PrimitiveUniformShaderParameters);
+		WorldSpacePrimitiveUniformBuffer.InitResource();
+	}
+
+
+	// Compute the per-view uniform buffers.
+	for (int32 ViewIndex = 0; ViewIndex < Views.Num(); ViewIndex++)
+	{
+		if (VisibilityMap & (1 << ViewIndex))
+		{
+			FGlobalDynamicVertexBuffer::FAllocation LocalDynamicVertexAllocation = FGlobalDynamicVertexBuffer::Get().Allocate(SizeInBytes);
+			//if (LocalDynamicVertexAllocation.IsValid())
+			{
+				const FSceneView* View = Views[ViewIndex];
+				FMatrix TransMat;
+				/*
+				if (Properties->SortMode == ENiagaraSortMode::SortViewDepth)
+				{
+					TransMat = View->ViewMatrices.GetViewProjectionMatrix();
+
+					// view depth sorting
+					DynamicDataSprites->VertexData.Sort(
+						[&TransMat](const FNiagaraSpriteVertex& A, const FNiagaraSpriteVertex& B) {
+						float W1 = TransMat.TransformPosition(A.Position).W;
+						float W2 = TransMat.TransformPosition(B.Position).W;
+						return (W1 > W2);
+					});
+				}
+				else if (Properties->SortMode == ENiagaraSortMode::SortViewDistance)
+				{
+					FVector ViewLoc = View->ViewLocation;
+					// dist sorting
+					DynamicDataSprites->VertexData.Sort(
+						[&ViewLoc](const FNiagaraSpriteVertex& A, const FNiagaraSpriteVertex& B) {
+						float D1 = (ViewLoc - A.Position).SizeSquared();
+						float D2 = (ViewLoc - B.Position).SizeSquared();
+						return (D1 > D2);
+					});
+				}
+
+				// Copy the vertex data over.
+				*/
+				//FMemory::Memcpy(LocalDynamicVertexAllocation.Buffer, DynamicDataSprites->VertexData.GetData(), SizeInBytes);
+
+				FNiagaraMeshCollectorResourcesSprite& CollectorResources = Collector.AllocateOneFrameResource<FNiagaraMeshCollectorResourcesSprite>();
+				FNiagaraSpriteUniformParameters PerViewUniformParameters;// = UniformParameters;
+				PerViewUniformParameters.RotationBias = 0.0f;
+				PerViewUniformParameters.RotationScale = 1.0f;
+				PerViewUniformParameters.TangentSelector = FVector4(0.0f, 0.0f, 0.0f, 1.0f);
+				PerViewUniformParameters.InvDeltaSeconds = 30.0f;
+				PerViewUniformParameters.NormalsType = 0.0f;
+				PerViewUniformParameters.NormalsSphereCenter = FVector4(0.0f, 0.0f, 0.0f, 1.0f);
+				PerViewUniformParameters.NormalsCylinderUnitDirection = FVector4(0.0f, 0.0f, 1.0f, 0.0f);
+				PerViewUniformParameters.PivotOffset = FVector2D(-0.5f, -0.5f);
+				PerViewUniformParameters.MacroUVParameters = FVector4(0.0f, 0.0f, 1.0f, 1.0f);
+				PerViewUniformParameters.CameraFacingBlend = FVector4(0.0f, 0.0f, 0.0f, 1.0f);
+				PerViewUniformParameters.RemoveHMDRoll = 0.0f;
+				PerViewUniformParameters.CustomFacingVectorMask = Properties->CustomFacingVectorMask;
+				PerViewUniformParameters.SubImageSize = FVector4(Properties->SubImageSize.X, Properties->SubImageSize.Y, 1.0f / Properties->SubImageSize.X, 1.0f / Properties->SubImageSize.Y);
+				PerViewUniformParameters.PositionDataOffset = DynamicDataSprites->PositionDataOffset;
+				PerViewUniformParameters.VelocityDataOffset = DynamicDataSprites->VelocityDataOffset;
+				PerViewUniformParameters.RotationDataOffset = DynamicDataSprites->RotationDataOffset;
+				PerViewUniformParameters.SizeDataOffset = DynamicDataSprites->SizeDataOffset;
+				PerViewUniformParameters.ColorDataOffset = DynamicDataSprites->ColorDataOffset;
+				PerViewUniformParameters.SubimageDataOffset = DynamicDataSprites->SubimageDataOffset;
+				PerViewUniformParameters.FacingOffset = DynamicDataSprites->FacingOffset;
+				PerViewUniformParameters.AlignmentOffset = DynamicDataSprites->AlignmentOffset;
+				if (Properties->Alignment == ENiagaraSpriteAlignment::VelocityAligned)
+				{
+					// velocity aligned
+					PerViewUniformParameters.RotationScale = 0.0f;
+					PerViewUniformParameters.TangentSelector = FVector4(0.0f, 1.0f, 0.0f, 0.0f);
+				}
+
+				// Collector.AllocateOneFrameResource uses default ctor, initialize the vertex factory
+
+				// use custom alignment if the data is available and if it's set in the props
+				bool bUseCustomAlignment = DynamicDataSprites->bCustomAlignmentAvailable && Properties->Alignment == ENiagaraSpriteAlignment::CustomAlignment;
+				bool bUseVectorAlignment = Properties->Alignment != ENiagaraSpriteAlignment::Unaligned;
+				
+				CollectorResources.VertexFactory.SetParticleData(DynamicDataSprites->DataSet);
+				CollectorResources.VertexFactory.SetCustomAlignment(bUseCustomAlignment);
+				CollectorResources.VertexFactory.SetVectorAligned(bUseVectorAlignment);
+				CollectorResources.VertexFactory.SetCameraPlaneFacing(Properties->FacingMode == ENiagaraSpriteFacingMode::FaceCameraPlane);
+
+				CollectorResources.VertexFactory.SetFeatureLevel(ViewFamily.GetFeatureLevel());
+				CollectorResources.VertexFactory.SetParticleFactoryType(NVFT_Sprite);
+
+
+				CollectorResources.UniformBuffer = FNiagaraSpriteUniformBufferRef::CreateUniformBufferImmediate(PerViewUniformParameters, UniformBuffer_SingleFrame);
+
+				CollectorResources.VertexFactory.SetNumVertsInInstanceBuffer(4);
+				CollectorResources.VertexFactory.InitResource();
+				CollectorResources.VertexFactory.SetSpriteUniformBuffer(CollectorResources.UniformBuffer);
+
+				CollectorResources.VertexFactory.SetInstanceBuffer(
+					/*LocalDynamicVertexAllocation.VertexBuffer*/nullptr,
+					LocalDynamicVertexAllocation.VertexOffset,
+					sizeof(FNiagaraSpriteVertex),
+					true
+					);
+
+
+				/*if (DynamicDataSprites->MaterialParameterVertexData.Num() > 0 && LocalDynamicVertexMaterialParamsAllocation.IsValid())
+				{
+					CollectorResources.VertexFactory.SetDynamicParameterBuffer(
+						LocalDynamicVertexMaterialParamsAllocation.VertexBuffer,
+						LocalDynamicVertexMaterialParamsAllocation.VertexOffset,
+						sizeof(FParticleVertexDynamicParameter),
+						true);
+				}
+				else*/
+				{
+					CollectorResources.VertexFactory.SetDynamicParameterBuffer(NULL, 0, 0, true);
+				}
+
+				FMeshBatch& MeshBatch = Collector.AllocateMesh();
+				MeshBatch.VertexFactory = &CollectorResources.VertexFactory;
+				MeshBatch.CastShadow = SceneProxy->CastsDynamicShadow();
+				MeshBatch.bUseAsOccluder = false;
+				MeshBatch.ReverseCulling = SceneProxy->IsLocalToWorldDeterminantNegative();
+				MeshBatch.Type = PT_TriangleList;
+				MeshBatch.DepthPriorityGroup = SceneProxy->GetDepthPriorityGroup(View);
+				MeshBatch.bCanApplyViewModeOverrides = true;
+				MeshBatch.bUseWireframeSelectionColoring = SceneProxy->IsSelected();
+
+				if (bIsWireframe)
+				{
+					MeshBatch.MaterialRenderProxy = UMaterial::GetDefaultMaterial(MD_Surface)->GetRenderProxy(SceneProxy->IsSelected(), SceneProxy->IsHovered());
+				}
+				else
+				{
+					MeshBatch.MaterialRenderProxy = MaterialRenderProxy;
+				}
+
+				FMeshBatchElement& MeshElement = MeshBatch.Elements[0];
+				MeshElement.IndexBuffer = &GParticleIndexBuffer;
+				MeshElement.FirstIndex = 0;
+				MeshElement.NumPrimitives = 2;
+				MeshElement.NumInstances = FMath::Max(0, NumInstances);	//->VertexData.Num();
+				MeshElement.MinVertexIndex = 0;
+				MeshElement.MaxVertexIndex = 0;// MeshElement.NumInstances * 4 - 1;
+				MeshElement.PrimitiveUniformBufferResource = &WorldSpacePrimitiveUniformBuffer;
+				
+				
+				Collector.AddMesh(ViewIndex, MeshBatch);
+				
+			}
+		}
+	}
+
+	CPUTimeMS += MeshElementsTimer.GetElapsedMilliseconds();
+}
+
+bool NiagaraRendererSprites::SetMaterialUsage()
+{
+	//Causes deadlock :S Need to look at / rework the setting of materials and render modules.
+	return Material && Material->CheckMaterialUsage_Concurrent(MATUSAGE_NiagaraSprites);
+}
+
+/** Update render data buffer from attributes */
+FNiagaraDynamicDataBase *NiagaraRendererSprites::GenerateVertexData(const FNiagaraSceneProxy* Proxy, FNiagaraDataSet &Data, const ENiagaraSimTarget Target)
+{
+	FNiagaraDynamicDataSprites *DynamicData = nullptr;
+
+	if (Data.PrevData().GetNumInstances() > 0)
+	{
+		SimpleTimer VertexDataTimer;
+
+		SCOPE_CYCLE_COUNTER(STAT_NiagaraGenSpriteVertexData);
+
+		const FNiagaraVariableLayoutInfo* PositionLayout = Data.GetVariableLayout(FNiagaraVariable(FNiagaraTypeDefinition::GetVec3Def(), TEXT("Position")));
+		const FNiagaraVariableLayoutInfo* VelocityLayout = Data.GetVariableLayout(FNiagaraVariable(FNiagaraTypeDefinition::GetVec3Def(), TEXT("Velocity")));
+		const FNiagaraVariableLayoutInfo* RotationLayout = Data.GetVariableLayout(FNiagaraVariable(FNiagaraTypeDefinition::GetFloatDef(), TEXT("SpriteRotation")));
+		const FNiagaraVariableLayoutInfo* SizeLayout = Data.GetVariableLayout(FNiagaraVariable(FNiagaraTypeDefinition::GetVec2Def(), TEXT("SpriteSize")));
+		const FNiagaraVariableLayoutInfo* ColorLayout = Data.GetVariableLayout(FNiagaraVariable(FNiagaraTypeDefinition::GetColorDef(), TEXT("Color")));
+
+		if (!bEnabled || !PositionLayout || !VelocityLayout || !RotationLayout || !SizeLayout || !ColorLayout)
+		{
+			return nullptr;
+		}
+
+
+		// required attributes //TODO: This does not need to be done every frame. Losing not insignificant time getting the layout ptrs.
+		DynamicData = new FNiagaraDynamicDataSprites;
+		DynamicData->PositionDataOffset = PositionLayout->FloatComponentStart;
+		DynamicData->VelocityDataOffset = VelocityLayout->FloatComponentStart;
+		DynamicData->RotationDataOffset = RotationLayout->FloatComponentStart;
+		DynamicData->SizeDataOffset = SizeLayout->FloatComponentStart;
+		DynamicData->ColorDataOffset = ColorLayout->FloatComponentStart;
+
+		// optional attributes; we pass -1 as the offset so the VF can branch
+		int32 IntDummy;
+		Data.GetVariableComponentOffsets(FNiagaraVariable(FNiagaraTypeDefinition::GetVec3Def(), TEXT("Facing")), DynamicData->FacingOffset, IntDummy);
+		Data.GetVariableComponentOffsets(FNiagaraVariable(FNiagaraTypeDefinition::GetVec3Def(), TEXT("Alignment")), DynamicData->AlignmentOffset, IntDummy);
+		Data.GetVariableComponentOffsets(FNiagaraVariable(FNiagaraTypeDefinition::GetFloatDef(), TEXT("SubImageIndex")), DynamicData->SubimageDataOffset, IntDummy);
+
+
+		// if we're CPU simulating, need to init the GPU buffers for the vertex factory
+			//This is a killer for perf. Assume its all the GPU buffer allocation.
+			//Need to allocate big shared buffer and allocate ranges out of those.
+		SCOPE_CYCLE_COUNTER(STAT_NiagaraGenGpuBuffers);
+		if (Target == ENiagaraSimTarget::CPUSim)
+		{
+			Data.ValidateBufferIndices();
+			Data.InitGPUFromCPU();
+		}
+
+		DynamicData->DataSet = &Data;
+		CPUTimeMS = VertexDataTimer.GetElapsedMilliseconds();
+	}
+
+	return DynamicData;  // for VF that can fetch from particle data directly
+}
+
+
+
+void NiagaraRendererSprites::SetDynamicData_RenderThread(FNiagaraDynamicDataBase* NewDynamicData)
+{
+	check(IsInRenderingThread());
+
+	if (DynamicDataRender)
+	{
+		delete static_cast<FNiagaraDynamicDataSprites*>(DynamicDataRender);
+		DynamicDataRender = NULL;
+	}
+	DynamicDataRender = NewDynamicData;
+}
+
+int NiagaraRendererSprites::GetDynamicDataSize()
+{
+	uint32 Size = sizeof(FNiagaraDynamicDataSprites);
+	if (DynamicDataRender)
+	{
+		Size += (static_cast<FNiagaraDynamicDataSprites*>(DynamicDataRender))->VertexData.GetAllocatedSize();
+	}
+
+	return Size;
+}
+
+bool NiagaraRendererSprites::HasDynamicData()
+{
+//	return DynamicDataRender && static_cast<FNiagaraDynamicDataSprites*>(DynamicDataRender)->VertexData.Num() > 0;
+	return DynamicDataRender!=nullptr;// && static_cast<FNiagaraDynamicDataSprites*>(DynamicDataRender)->NumInstances > 0;
+}
+
+#if WITH_EDITORONLY_DATA
+
+const TArray<FNiagaraVariable>& NiagaraRendererSprites::GetRequiredAttributes()
+{
+	return Properties->GetRequiredAttributes();
+}
+
+const TArray<FNiagaraVariable>& NiagaraRendererSprites::GetOptionalAttributes()
+{
+	return Properties->GetOptionalAttributes();
+}
+
+#endif
\ No newline at end of file
diff --git a/Engine/Plugins/FX/Niagara/Source/Niagara/Private/NiagaraRibbonRendererProperties.cpp b/Engine/Plugins/FX/Niagara/Source/Niagara/Private/NiagaraRibbonRendererProperties.cpp
new file mode 100644
--- /dev/null
+++ b/Engine/Plugins/FX/Niagara/Source/Niagara/Private/NiagaraRibbonRendererProperties.cpp
@@ -0,0 +1,72 @@
+// Copyright 1998-2016 Epic Games, Inc. All Rights Reserved.
+
+#include "NiagaraRibbonRendererProperties.h"
+#include "NiagaraRendererRibbons.h"
+#include "NiagaraConstants.h"
+
+NiagaraRenderer* UNiagaraRibbonRendererProperties::CreateEmitterRenderer(ERHIFeatureLevel::Type FeatureLevel)
+{
+	return new NiagaraRendererRibbons(FeatureLevel, this);
+}
+
+void UNiagaraRibbonRendererProperties::GetUsedMaterials(TArray<UMaterialInterface*>& OutMaterials) const
+{
+	OutMaterials.Add(Material);
+}
+
+#if WITH_EDITORONLY_DATA
+
+const TArray<FNiagaraVariable>& UNiagaraRibbonRendererProperties::GetRequiredAttributes()
+{
+	static TArray<FNiagaraVariable> Attrs;
+
+	if (Attrs.Num() == 0)
+	{
+		Attrs.Add(SYS_PARAM_PARTICLES_POSITION);
+		/*Attrs.Add(SYS_PARAM_PARTICLES_VELOCITY);
+		Attrs.Add(SYS_PARAM_PARTICLES_COLOR);
+		Attrs.Add(SYS_PARAM_PARTICLES_SPRITE_ROTATION);
+		Attrs.Add(SYS_PARAM_PARTICLES_NORMALIZED_AGE);
+		Attrs.Add(SYS_PARAM_PARTICLES_SPRITE_SIZE);*/
+	}
+
+	return Attrs;
+}
+
+
+const TArray<FNiagaraVariable>& UNiagaraRibbonRendererProperties::GetOptionalAttributes()
+{
+	static TArray<FNiagaraVariable> Attrs;
+
+	if (Attrs.Num() == 0)
+	{
+		//Attrs.Add(SYS_PARAM_PARTICLES_SPRITE_FACING);
+		//Attrs.Add(SYS_PARAM_PARTICLES_SPRITE_ALIGNMENT);
+		//Attrs.Add(SYS_PARAM_PARTICLES_SUB_IMAGE_INDEX);
+		//Attrs.Add(SYS_PARAM_PARTICLES_DYNAMIC_MATERIAL_PARAM);
+	}
+
+	return Attrs;
+}
+
+
+
+
+bool UNiagaraRibbonRendererProperties::IsMaterialValidForRenderer(UMaterial* InMaterial, FText& InvalidMessage)
+{
+	if (InMaterial->bUsedWithNiagaraRibbons == false)
+	{
+		InvalidMessage = NSLOCTEXT("NiagaraRibbonRendererProperties", "InvalidMaterialMessage", "The material isn't marked as \"Used with Niagara ribbons\"");
+		return false;
+	}
+	return true;
+}
+
+void UNiagaraRibbonRendererProperties::FixMaterial(UMaterial* InMaterial)
+{
+	InMaterial->Modify();
+	InMaterial->bUsedWithNiagaraRibbons = true;
+	InMaterial->ForceRecompileForRendering();
+}
+
+#endif // WITH_EDITORONLY_DATA
