Copying //UE4/Orion-Staging to //UE4/Main (Source: //Orion/Dev-General @ 3606378)

diff --git a/Engine/Plugins/FX/Niagara/Source/Niagara/Private/NiagaraScript.cpp b/Engine/Plugins/FX/Niagara/Source/Niagara/Private/NiagaraScript.cpp
new file mode 100644
--- /dev/null
+++ b/Engine/Plugins/FX/Niagara/Source/Niagara/Private/NiagaraScript.cpp
@@ -0,0 +1,473 @@
+// Copyright 1998-2017 Epic Games, Inc. All Rights Reserved.
+
+#include "NiagaraScript.h"
+#include "Modules/ModuleManager.h"
+#include "NiagaraScriptSourceBase.h"
+#include "NiagaraComponent.h"
+#include "NiagaraEmitter.h"
+#include "Package.h"
+#include "Linker.h"
+#include "NiagaraModule.h"
+#include "NiagaraCustomVersion.h"
+#include "NiagaraShaderCompilationManager.h"
+
+
+#include "Stats/Stats.h"
+#include "Linker.h"
+
+#if WITH_EDITOR
+	#include "Interfaces/ITargetPlatform.h"
+#endif
+
+DECLARE_STATS_GROUP(TEXT("Niagara Detailed"), STATGROUP_NiagaraDetailed, STATCAT_Advanced);
+
+FNiagaraScriptDebuggerInfo::FNiagaraScriptDebuggerInfo() : bRequestDebugFrame(false), DebugFrameLastWriteId(-1)
+{
+}
+
+UNiagaraScript::UNiagaraScript(const FObjectInitializer& ObjectInitializer)
+	: Super(ObjectInitializer)
+	, Usage(ENiagaraScriptUsage::Function)
+	, ModuleUsageBitmask( (1 << (int32)ENiagaraScriptUsage::ParticleSpawnScript) | (1 << (int32)ENiagaraScriptUsage::ParticleSpawnScriptInterpolated) | (1 << (int32)ENiagaraScriptUsage::ParticleUpdateScript) | (1 << (int32)ENiagaraScriptUsage::ParticleEventScript) )
+	, NumUserPtrs(0)
+	, NumericOutputTypeSelectionMode(ENiagaraNumericOutputTypeSelectionMode::Largest)
+#if WITH_EDITORONLY_DATA
+	, LastCompileStatus(ENiagaraScriptCompileStatus::NCS_Unknown)
+	, UniqueID(FGuid::NewGuid())
+#endif
+{
+	
+}
+
+UNiagaraScript::~UNiagaraScript()
+{
+}
+
+void UNiagaraScript::PostInitProperties()
+{
+	Super::PostInitProperties();
+	UsageIndex = 0;
+}
+
+UNiagaraScriptSourceBase::UNiagaraScriptSourceBase(const FObjectInitializer& ObjectInitializer)
+	: Super(ObjectInitializer)
+{
+
+}
+
+
+void UNiagaraScript::Serialize(FArchive& Ar)
+{
+	Super::Serialize(Ar);
+
+	Ar.UsingCustomVersion(FNiagaraCustomVersion::GUID);		// only changes version if not loading
+	const int32 NiagaraVer = Ar.CustomVer(FNiagaraCustomVersion::GUID);
+
+	bool IsValidShaderScript = Usage != ENiagaraScriptUsage::Module && Usage != ENiagaraScriptUsage::Function && Usage != ENiagaraScriptUsage::DynamicInput 
+		&& (NiagaraVer<FNiagaraCustomVersion::NiagaraShaderMapCooking2 || (Usage != ENiagaraScriptUsage::SystemSpawnScript && Usage != ENiagaraScriptUsage::SystemUpdateScript));
+	if ( (!Ar.IsLoading() && IsValidShaderScript)		// saving shader maps only for particle sim and spawn scripts
+		|| (Ar.IsLoading() && NiagaraVer >= FNiagaraCustomVersion::NiagaraShaderMaps && (NiagaraVer < FNiagaraCustomVersion::NiagaraShaderMapCooking || IsValidShaderScript))  // load only if we know shader map is presen
+		)
+	{
+#if WITH_EDITOR
+		SerializeNiagaraShaderMaps(&CachedScriptResourcesForCooking, Ar, ScriptResource);
+#else
+		SerializeNiagaraShaderMaps(NULL, Ar, ScriptResource);
+#endif
+	}
+}
+
+void UNiagaraScript::PostLoad()
+{
+	Super::PostLoad();
+	
+	bool bNeedsRecompile = false;
+	const int32 NiagaraVer = GetLinkerCustomVersion(FNiagaraCustomVersion::GUID);
+	
+	//Recompile is before the latest version that requires a script recompile. Some may not.
+	if (NiagaraVer < FNiagaraCustomVersion::RemovalOfNiagaraVariableIDs)
+	{
+		if (Usage != ENiagaraScriptUsage::Module && Usage != ENiagaraScriptUsage::Function && Usage != ENiagaraScriptUsage::DynamicInput)
+		{
+			bNeedsRecompile = true;
+			ByteCode.Empty();
+#if WITH_EDITORONLY_DATA
+			LastCompileStatus = ENiagaraScriptCompileStatus::NCS_Unknown;
+#endif
+			UE_LOG(LogNiagara, Warning, TEXT("Niagara script is out of date and requires recompile to be used! %s"), *GetFullName());
+		}
+	}
+
+
+	// Resources can be processed / registered now that we're back on the main thread
+	ProcessSerializedShaderMaps(this, ScriptResource, ScriptResourcesByFeatureLevel);
+
+
+	if (GIsEditor)
+	{
+		
+#if WITH_EDITORONLY_DATA
+		// Since we're about to check the synchronized state, we need to make sure that it has been post-loaded (which 
+		// can affect the results of that call).
+		if (Source != nullptr)
+		{	
+			Source->ConditionalPostLoad();
+		}
+
+		// If we've never computed a change id or the source graph is different than the compiled script change id,
+		// we're out of sync and should recompile.
+		if (Usage != ENiagaraScriptUsage::Module && Usage != ENiagaraScriptUsage::Function && Usage != ENiagaraScriptUsage::DynamicInput)
+		{
+			if (ChangeId.IsValid() == false || (Source != nullptr && !Source->IsSynchronized(ChangeId)))
+			{
+				bNeedsRecompile = true;
+				ByteCode.Empty();
+				LastCompileStatus = ENiagaraScriptCompileStatus::NCS_Unknown;
+				UE_LOG(LogNiagara, Warning, TEXT("Niagara script is out of date with source graph requires recompile to be used! %s"), *GetFullName());
+			}
+		}
+#endif
+	}
+
+	// for now, force recompile until we can be sure everything is working
+//	bNeedsRecompile = true;
+#if WITH_EDITORONLY_DATA
+	CacheResourceShadersForRendering(false, bNeedsRecompile);
+	GenerateStatScopeIDs();
+#endif
+
+}
+
+
+#if WITH_EDITOR
+
+void UNiagaraScript::GenerateStatScopeIDs()
+{
+	StatScopesIDs.Empty();
+#if STATS
+	for (FNiagaraStatScope& StatScope : StatScopes)
+	{
+		StatScopesIDs.Add(FDynamicStats::CreateStatId<FStatGroup_STATGROUP_NiagaraDetailed>(StatScope.FriendlyName.ToString()));
+	}
+#endif
+}
+
+void UNiagaraScript::PostEditChangeProperty(FPropertyChangedEvent& PropertyChangedEvent)
+{
+	CacheResourceShadersForRendering(true);
+}
+
+#endif
+
+#if WITH_EDITORONLY_DATA
+bool UNiagaraScript::AreScriptAndSourceSynchronized() const
+{
+	if (Source)
+	{
+		return Source->IsSynchronized(ChangeId);
+	}
+	else
+	{
+		return false;
+	}
+}
+
+void UNiagaraScript::MarkScriptAndSourceDesynchronized()
+{
+	if (Source)
+	{
+		Source->MarkNotSynchronized();
+	}
+}
+
+
+UNiagaraScript* UNiagaraScript::MakeRecursiveDeepCopy(UObject* DestOuter, TMap<const UObject*, UObject*>& ExistingConversions) const
+{
+	check(GetOuter() != DestOuter);
+
+	bool bSourceConvertedAlready = ExistingConversions.Contains(Source);
+
+	ResetLoaders(GetTransientPackage()); // Make sure that we're not going to get invalid version number linkers into the transient package. 
+	GetTransientPackage()->LinkerCustomVersion.Empty();
+
+	// For some reason, the default parameters of FObjectDuplicationParameters aren't the same as
+	// StaticDuplicateObject uses internally. These are copied from Static Duplicate Object...
+	EObjectFlags FlagMask = RF_AllFlags & ~RF_Standalone & ~RF_Public; // Remove Standalone and Public flags.
+	EDuplicateMode::Type DuplicateMode = EDuplicateMode::Normal;
+	EInternalObjectFlags InternalFlagsMask = EInternalObjectFlags::AllFlags;
+
+	FObjectDuplicationParameters ObjParameters((UObject*)this, GetTransientPackage());
+	ObjParameters.DestName = NAME_None;
+	if (this->GetOuter() != DestOuter)
+	{
+		// try to keep the object name consistent if possible
+		if (FindObjectFast<UObject>(DestOuter, GetFName()) == nullptr)
+		{
+			ObjParameters.DestName = GetFName();
+		}
+	}
+
+	ObjParameters.DestClass = GetClass();
+	ObjParameters.FlagMask = FlagMask;
+	ObjParameters.InternalFlagMask = InternalFlagsMask;
+	ObjParameters.DuplicateMode = DuplicateMode;
+	
+	// Make sure that we don't duplicate objects that we've alredy converted...
+	TMap<const UObject*, UObject*>::TConstIterator It = ExistingConversions.CreateConstIterator();
+	while (It)
+	{
+		ObjParameters.DuplicationSeed.Add(const_cast<UObject*>(It.Key()), It.Value());
+		++It;
+	}
+
+	UNiagaraScript*	Script = CastChecked<UNiagaraScript>(StaticDuplicateObjectEx(ObjParameters));
+
+	check(Script->HasAnyFlags(RF_Standalone) == false);
+	check(Script->HasAnyFlags(RF_Public) == false);
+
+	if (bSourceConvertedAlready)
+	{
+		// Confirm that we've converted these properly..
+		check(Script->Source == ExistingConversions[Source]);
+	}
+
+	if (DestOuter != nullptr)
+	{
+		Script->Rename(nullptr, DestOuter, REN_DoNotDirty | REN_DontCreateRedirectors | REN_NonTransactional);
+	}
+	UE_LOG(LogNiagara, Warning, TEXT("MakeRecursiveDeepCopy %s"), *Script->GetFullName());
+	ExistingConversions.Add(const_cast<UNiagaraScript*>(this), Script);
+
+	// Since the Source is the only thing we subsume from UNiagaraScripts, only do the subsume if 
+	// we haven't already converted it.
+	if (bSourceConvertedAlready == false)
+	{
+		Script->SubsumeExternalDependencies(ExistingConversions);
+	}
+	return Script;
+}
+
+void UNiagaraScript::SubsumeExternalDependencies(TMap<const UObject*, UObject*>& ExistingConversions)
+{
+	Source->SubsumeExternalDependencies(ExistingConversions);
+}
+
+ENiagaraScriptCompileStatus UNiagaraScript::Compile(FString& OutGraphLevelErrorMessages)
+{
+	if (Source == nullptr)
+		return ENiagaraScriptCompileStatus::NCS_Error;
+
+	ENiagaraScriptCompileStatus CPUCompileStatus = Source->Compile(this, OutGraphLevelErrorMessages);
+	CacheResourceShadersForRendering(false, true);
+
+	return CPUCompileStatus;
+}
+#endif
+
+
+
+#if WITH_EDITOR
+void UNiagaraScript::BeginCacheForCookedPlatformData(const ITargetPlatform *TargetPlatform)
+{
+	TArray<FName> DesiredShaderFormats;
+	TargetPlatform->GetAllTargetedShaderFormats(DesiredShaderFormats);
+
+	FNiagaraScript** CachedScriptResourceForPlatformPtr = CachedScriptResourcesForCooking.Find(TargetPlatform);
+
+	if (CachedScriptResourceForPlatformPtr == nullptr)
+	{
+		CachedScriptResourcesForCooking.Add(TargetPlatform);
+		FNiagaraScript* &CachedScriptResourceForPlatform = *CachedScriptResourcesForCooking.Find(TargetPlatform);
+
+		if (DesiredShaderFormats.Num())
+		{
+			// Cache for all the shader formats that the cooking target requires
+			for (int32 FormatIndex = 0; FormatIndex < DesiredShaderFormats.Num(); FormatIndex++)
+			{
+				if (DesiredShaderFormats[FormatIndex] != TEXT("PCD3D_SM4"))	// TODO: remove once we get rid of SM4 globally
+				{
+					const EShaderPlatform LegacyShaderPlatform = ShaderFormatToLegacyShaderPlatform(DesiredShaderFormats[FormatIndex]);
+
+					// Begin caching shaders for the target platform and store the FNiagaraScript being compiled into CachedScriptResourcesForCooking
+					CacheResourceShadersForCooking(LegacyShaderPlatform, CachedScriptResourceForPlatform);
+				}
+			}
+		}
+	}
+}
+
+
+void UNiagaraScript::CacheResourceShadersForCooking(EShaderPlatform ShaderPlatform, FNiagaraScript* &OutCachedResource)
+{
+	if (Usage != ENiagaraScriptUsage::Function && Usage != ENiagaraScriptUsage::Module && Usage != ENiagaraScriptUsage::DynamicInput)
+	{
+		FNiagaraScript *ResourceToCache = nullptr;
+		ERHIFeatureLevel::Type TargetFeatureLevel = GetMaxSupportedFeatureLevel(ShaderPlatform);
+
+		FNiagaraScript* NewResource = AllocateResource();
+		NewResource->SetScript(this, (ERHIFeatureLevel::Type)TargetFeatureLevel, UniqueID, GetName());
+		ResourceToCache = NewResource;
+
+		check(ResourceToCache);
+
+		CacheShadersForResources(ShaderPlatform, ResourceToCache, false, false, true);
+
+		INiagaraModule NiagaraModule = FModuleManager::GetModuleChecked<INiagaraModule>(TEXT("Niagara"));
+		NiagaraModule.ProcessShaderCompilationQueue();
+
+		OutCachedResource = ResourceToCache;
+	}
+}
+
+
+
+void UNiagaraScript::CacheShadersForResources(EShaderPlatform ShaderPlatform, FNiagaraScript *ResourceToCache, bool bApplyCompletedShaderMapForRendering, bool bForceRecompile, bool bCooking)
+{
+	const bool bSuccess = ResourceToCache->CacheShaders(ShaderPlatform, bApplyCompletedShaderMapForRendering, bForceRecompile, bCooking);
+
+#if defined(NIAGARA_SCRIPT_COMPILE_LOGGING_MEDIUM)
+	if (!bSuccess)
+	{
+		UE_LOG(LogNiagara, Warning, TEXT("Failed to compile Niagara shader %s for platform %s."),
+			*GetPathName(),
+			*LegacyShaderPlatformToShaderFormat(ShaderPlatform).ToString());
+
+		const TArray<FString>& CompileErrors = ResourceToCache->GetCompileErrors();
+		for (int32 ErrorIndex = 0; ErrorIndex < CompileErrors.Num(); ErrorIndex++)
+		{
+			UE_LOG(LogNiagara, Warning, TEXT("	%s"), *CompileErrors[ErrorIndex]);
+		}
+	}
+#endif
+}
+
+
+
+void UNiagaraScript::CacheResourceShadersForRendering(bool bRegenerateId, bool bForceRecompile)
+{
+	if (bRegenerateId)
+	{
+		// Regenerate this script's Id if requested
+		for (int32 Idx = 0; Idx < ERHIFeatureLevel::Num; Idx++)
+		{
+			if (ScriptResourcesByFeatureLevel[Idx])
+			{
+				ScriptResourcesByFeatureLevel[Idx]->ReleaseShaderMap();
+				ScriptResourcesByFeatureLevel[Idx] = nullptr;
+			}
+		}
+	}
+
+	//UpdateResourceAllocations();
+
+	if (FApp::CanEverRender() && CanBeRunOnGpu())
+	{
+		if (Source)
+		{
+			FNiagaraScript* ResourceToCache;
+			ERHIFeatureLevel::Type CacheFeatureLevel = ERHIFeatureLevel::SM5;
+			ScriptResource.SetScript(this, FeatureLevel, UniqueID, GetName());
+
+			//if (ScriptResourcesByFeatureLevel[FeatureLevel])
+			{
+				EShaderPlatform ShaderPlatform = GShaderPlatformForFeatureLevel[CacheFeatureLevel];
+				ResourceToCache = ScriptResourcesByFeatureLevel[CacheFeatureLevel];
+				CacheShadersForResources(ShaderPlatform, &ScriptResource, true, bForceRecompile);
+				ScriptResourcesByFeatureLevel[CacheFeatureLevel] = &ScriptResource;
+			}
+		}
+	}
+}
+
+#endif
+
+void SerializeNiagaraShaderMaps(const TMap<const ITargetPlatform*, FNiagaraScript*>* PlatformScriptResourcesToSavePtr, FArchive& Ar, FNiagaraScript &OutLoadedResource)
+{
+//	SCOPED_LOADTIMER(SerializeInlineShaderMaps);
+	if (Ar.IsSaving())
+	{
+		int32 NumResourcesToSave = 0;
+		FNiagaraScript *ScriptResourceToSave = nullptr;
+		FNiagaraScript* const * ScriptResourceToSavePtr = nullptr;
+		if (Ar.IsCooking())
+		{
+			check(PlatformScriptResourcesToSavePtr);
+			auto& PlatformScriptResourcesToSave = *PlatformScriptResourcesToSavePtr;
+
+			ScriptResourceToSavePtr = PlatformScriptResourcesToSave.Find(Ar.CookingTarget());
+			check(ScriptResourceToSavePtr != NULL || (Ar.GetLinker() == NULL));
+			if (ScriptResourceToSavePtr != NULL)
+			{
+				NumResourcesToSave++;
+			}
+		}
+
+		Ar << NumResourcesToSave;
+
+		if (ScriptResourceToSavePtr)
+		{
+			ScriptResourceToSave = *ScriptResourceToSavePtr;
+			ScriptResourceToSave->SerializeShaderMap(Ar);
+		}
+
+	}
+	else if (Ar.IsLoading())
+	{
+		int32 NumLoadedResources = 0;
+		Ar << NumLoadedResources;
+		if (NumLoadedResources)
+		{
+			OutLoadedResource.SerializeShaderMap(Ar);
+		}
+	}
+}
+
+
+
+void ProcessSerializedShaderMaps(UNiagaraScript* Owner, FNiagaraScript LoadedResource, FNiagaraScript* (&OutScriptResourcesLoaded)[ERHIFeatureLevel::Num])
+{
+	check(IsInGameThread());
+
+	LoadedResource.RegisterShaderMap();
+
+	FNiagaraShaderMap* LoadedShaderMap = LoadedResource.GetGameThreadShaderMap();
+	if (LoadedShaderMap && LoadedShaderMap->GetShaderPlatform() == GMaxRHIShaderPlatform)
+	{
+		ERHIFeatureLevel::Type LoadedFeatureLevel = LoadedShaderMap->GetShaderMapId().FeatureLevel;
+		if (!OutScriptResourcesLoaded[LoadedFeatureLevel])
+		{
+			OutScriptResourcesLoaded[LoadedFeatureLevel] = Owner->AllocateResource();
+		}
+
+		OutScriptResourcesLoaded[LoadedFeatureLevel]->SetShaderMap(LoadedShaderMap);
+	}
+}
+
+
+
+FNiagaraScript* UNiagaraScript::AllocateResource()
+{
+	return new FNiagaraScript();
+}
+
+TArray<ENiagaraScriptUsage> UNiagaraScript::GetSupportedUsageContexts() const
+{
+	TArray<ENiagaraScriptUsage> Supported;
+	for (int32 i = 0; i <= (int32)ENiagaraScriptUsage::SystemUpdateScript; i++)
+	{
+		int32 TargetBit = (ModuleUsageBitmask >> (int32)i) & 1;
+		if (TargetBit == 1)
+		{
+			Supported.Add((ENiagaraScriptUsage)i);
+		}
+	}
+	return Supported;
+}
+
+bool UNiagaraScript::UsesCollection(const UNiagaraParameterCollection* Collection)const
+{
+	return ParameterCollections.FindByPredicate([&](const UNiagaraParameterCollection* CheckCollection)
+	{
+		return CheckCollection == Collection;
+	}) != NULL;
+}
\ No newline at end of file
diff --git a/Engine/Plugins/FX/Niagara/Source/Niagara/Private/NiagaraScriptExecutionContext.cpp b/Engine/Plugins/FX/Niagara/Source/Niagara/Private/NiagaraScriptExecutionContext.cpp
new file mode 100644
--- /dev/null
+++ b/Engine/Plugins/FX/Niagara/Source/Niagara/Private/NiagaraScriptExecutionContext.cpp
@@ -0,0 +1,269 @@
+// Copyright 1998-2016 Epic Games, Inc. All Rights Reserved.
+
+#include "NiagaraScriptExecutionContext.h"
+#include "NiagaraStats.h"
+#include "NiagaraEmitterInstanceBatcher.h"
+#include "NiagaraDataInterface.h"
+#include "NiagaraSystemInstance.h"
+
+DECLARE_CYCLE_STAT(TEXT("Register Setup"), STAT_NiagaraSimRegisterSetup, STATGROUP_Niagara);
+
+
+//////////////////////////////////////////////////////////////////////////
+FNiagaraScriptExecutionParameterStore::FNiagaraScriptExecutionParameterStore()
+	: FNiagaraParameterStore()
+{
+
+}
+
+FNiagaraScriptExecutionParameterStore::FNiagaraScriptExecutionParameterStore(const FNiagaraParameterStore& Other)
+{
+	*this = Other;
+}
+
+FNiagaraScriptExecutionParameterStore& FNiagaraScriptExecutionParameterStore::operator=(const FNiagaraParameterStore& Other)
+{
+	FNiagaraParameterStore::operator=(Other);
+	return *this;
+}
+
+void FNiagaraScriptExecutionParameterStore::Init(UNiagaraScript* Script)
+{
+	//TEMPORARTY:
+	//We should replace the storage on the script with an FNiagaraParameterStore also so we can just copy that over here. Though that is an even bigger refactor job so this is a convenient place to break that work up.
+
+	bParametersDirty = true;
+	bInterfacesDirty = true;
+	Empty();
+
+	//Here we add the current frame parameters.
+	for (FNiagaraVariable& Param : Script->Parameters.Parameters)
+	{
+		AddParameter(Param, false);
+	}
+	
+	//Add previous frame values if we're interpolated spawn.
+	bool bIsInterpolatedSpawn = Script->IsInterpolatedParticleSpawnScript();
+	if (bIsInterpolatedSpawn)
+	{
+		for (FNiagaraVariable& Param : Script->Parameters.Parameters)
+		{
+			FNiagaraVariable PrevParam(Param.GetType(), FName(*(TEXT("PREV__") + Param.GetName().ToString())));
+			AddParameter(PrevParam, false);
+		}
+	}
+	
+	ParameterSize = ParameterData.Num();
+	if (bIsInterpolatedSpawn)
+	{
+		CopyCurrToPrev();
+	}
+
+	//Internal constants
+	for (FNiagaraVariable& InternalVar : Script->InternalParameters.Parameters)
+	{
+		AddParameter(InternalVar, false);
+	}
+
+	for (FNiagaraScriptDataInterfaceInfo& Info : Script->DataInterfaceInfo)
+	{
+		FNiagaraVariable Var(FNiagaraTypeDefinition(Info.DataInterface->GetClass()), Info.Name);
+		AddParameter(Var, false);
+		SetDataInterface(Info.DataInterface, IndexOf(Var));
+	}
+}
+
+void FNiagaraScriptExecutionParameterStore::CopyCurrToPrev()
+{
+	int32 ParamStart = ParameterSize / 2;
+	checkSlow(FMath::Frac((float)ParameterSize / 2) == 0.0f);
+
+	FMemory::Memcpy(GetParameterData_Internal(ParamStart), GetParameterData(0), ParamStart);
+}
+
+FNiagaraScriptExecutionContext::FNiagaraScriptExecutionContext()
+	: Script(nullptr)
+{
+
+}
+
+uint32 FNiagaraScriptExecutionContext::TickCounter = 0;
+
+bool FNiagaraScriptExecutionContext::Init(UNiagaraScript* InScript, ENiagaraSimTarget InTarget)
+{
+	Script = InScript;
+
+	Parameters.Init(Script);
+
+	return true;//TODO: Error cases?
+}
+
+bool FNiagaraScriptExecutionContext::Tick(FNiagaraSystemInstance* ParentSystemInstance)
+{
+	if (Script)//TODO: Remove. Script can only be null for system instances that currently don't have their script exec context set up correctly.
+	{
+		//Bind data interfaces if needed.
+		if (Parameters.bInterfacesDirty)
+		{
+			FunctionTable.Empty();
+			const TArray<UNiagaraDataInterface*> DataInterfaces = GetDataInterfaces();
+			// We must make sure that the data interfaces match up between the original script values and our overrides...
+			if (Script->DataInterfaceInfo.Num() != DataInterfaces.Num())
+			{
+				UE_LOG(LogNiagara, Warning, TEXT("Mismatch between Niagara Exectuion Context data interfaces and those in it's script!"));
+				return false;
+			}
+
+			//Fill the instance data table.
+			if (ParentSystemInstance)
+			{
+				DataInterfaceInstDataTable.SetNumZeroed(Script->NumUserPtrs);
+				for (int32 i = 0; i < Script->DataInterfaceInfo.Num(); i++)
+				{
+					int32 UserPtrIdx = Script->DataInterfaceInfo[i].UserPtrIdx;
+					if (UserPtrIdx != INDEX_NONE)
+					{
+						UNiagaraDataInterface* Interface = DataInterfaces[i];
+						void* InstData = ParentSystemInstance->FindDataInterfaceInstanceData(Interface);
+						DataInterfaceInstDataTable[UserPtrIdx] = InstData;
+					}
+				}
+			}
+			else
+			{
+				check(Script->NumUserPtrs == 0);//Can't have user ptrs if we have no parent instance.
+			}
+
+			bool bSuccessfullyMapped = true;
+			for (FVMExternalFunctionBindingInfo& BindingInfo : Script->CalledVMExternalFunctions)
+			{
+				for (int32 i = 0; i < Script->DataInterfaceInfo.Num(); i++)
+				{
+					FNiagaraScriptDataInterfaceInfo& ScriptInfo = Script->DataInterfaceInfo[i];
+					UNiagaraDataInterface* ExternalInterface = DataInterfaces[i];
+
+					if (ScriptInfo.Name == BindingInfo.OwnerName)
+					{
+						void* InstData = ScriptInfo.UserPtrIdx == INDEX_NONE ? nullptr : DataInterfaceInstDataTable[ScriptInfo.UserPtrIdx];
+						FVMExternalFunction Func = ExternalInterface->GetVMExternalFunction(BindingInfo, InstData);
+						if (Func.IsBound())
+						{
+							FunctionTable.Add(Func);
+						}
+						else
+						{
+							UE_LOG(LogNiagara, Error, TEXT("Could not Get VMExternalFunction '%s'.. emitter will not run!"), *BindingInfo.Name.ToString());
+							bSuccessfullyMapped = false;
+						}
+					}
+				}
+			}
+
+			if (!bSuccessfullyMapped)
+			{
+				UE_LOG(LogNiagara, Warning, TEXT("Error building data interface function table!"));
+				FunctionTable.Empty();
+				return false;
+			}
+		}
+	}
+
+	Parameters.Tick();
+
+	return true;
+}
+
+void FNiagaraScriptExecutionContext::PostTick()
+{
+	//If we're for interpolated spawn, copy over the previous frame's parameters into the Prev parameters.
+	if (Script && Script->IsInterpolatedParticleSpawnScript())
+	{
+		Parameters.CopyCurrToPrev();
+	}
+}
+
+bool FNiagaraScriptExecutionContext::Execute(uint32 NumInstances, TArray<FNiagaraDataSetExecutionInfo>& DataSetInfos)
+{
+	if (NumInstances == 0)
+	{
+		return true;
+	}
+
+	++TickCounter;//Should this be per execution?
+
+	int32 NumInputRegisters = 0;
+	int32 NumOutputRegisters = 0;
+	uint8* InputRegisters[VectorVM::MaxInputRegisters];
+	uint8* OutputRegisters[VectorVM::MaxOutputRegisters];
+
+	DataSetMetaTable.Reset();
+
+	{
+		SCOPE_CYCLE_COUNTER(STAT_NiagaraSimRegisterSetup);
+		for (FNiagaraDataSetExecutionInfo& DataSetInfo : DataSetInfos)
+		{
+			check(DataSetInfo.DataSet);
+			FDataSetMeta SetMeta(DataSetInfo.DataSet->GetSizeBytes(), &InputRegisters[NumInputRegisters], NumInputRegisters);
+			DataSetMetaTable.Add(SetMeta);
+			if (DataSetInfo.bAllocate)
+			{
+				DataSetInfo.DataSet->Allocate(NumInstances);
+				DataSetInfo.DataSet->SetNumInstances(NumInstances);
+			}
+			for (const FNiagaraVariable &Var : DataSetInfo.DataSet->GetVariables())
+			{
+				DataSetInfo.DataSet->AppendToRegisterTable(Var, InputRegisters, NumInputRegisters, OutputRegisters, NumOutputRegisters, DataSetInfo.StartInstance);
+			}
+		}
+	}
+
+	VectorVM::Exec(
+		Script->ByteCode.GetData(),
+		InputRegisters,
+		NumInputRegisters,
+		OutputRegisters,
+		NumOutputRegisters,
+		Parameters.GetParameterDataArray().GetData(),
+		DataSetMetaTable,
+		FunctionTable.GetData(),
+		DataInterfaceInstDataTable.GetData(),
+		NumInstances
+#if STATS
+#if WITH_EDITORONLY_DATA
+		, Script->GetStatScopeIDs()
+#else
+		, TArray<TStatId>()
+#endif
+#endif
+	);
+
+	// Tell the datasets we wrote how many instances were actually written.
+	for (int Idx = 0; Idx < DataSetInfos.Num(); Idx++)
+	{
+		FNiagaraDataSetExecutionInfo& Info = DataSetInfos[Idx];
+		if (Info.bUpdateInstanceCount)
+		{
+			Info.DataSet->SetNumInstances(Info.StartInstance + DataSetMetaTable[Idx].DataSetAccessIndex);
+		}
+	}
+
+#if WITH_EDITORONLY_DATA
+	if (Script->GetDebuggerInfo().bRequestDebugFrame)
+	{
+		DataSetInfos[0].DataSet->Dump(Script->GetDebuggerInfo().DebugFrame, true);
+		Script->GetDebuggerInfo().bRequestDebugFrame = false;
+		Script->GetDebuggerInfo().DebugFrameLastWriteId = TickCounter;
+	}
+#endif
+	return true;//TODO: Error cases?
+}
+
+void FNiagaraScriptExecutionContext::DirtyDataInterfaces()
+{
+	Parameters.bInterfacesDirty = true;
+}
+
+bool FNiagaraScriptExecutionContext::CanExecute()const
+{
+	return Script && Script->ByteCode.Num() > 0;
+}
\ No newline at end of file
diff --git a/Engine/Plugins/FX/Niagara/Source/Niagara/Private/NiagaraSpriteRendererProperties.cpp b/Engine/Plugins/FX/Niagara/Source/Niagara/Private/NiagaraSpriteRendererProperties.cpp
new file mode 100644
--- /dev/null
+++ b/Engine/Plugins/FX/Niagara/Source/Niagara/Private/NiagaraSpriteRendererProperties.cpp
@@ -0,0 +1,89 @@
+// Copyright 1998-2016 Epic Games, Inc. All Rights Reserved.
+
+#include "NiagaraSpriteRendererProperties.h"
+#include "NiagaraRenderer.h"
+#include "Internationalization.h"
+#include "NiagaraConstants.h"
+
+#define LOCTEXT_NAMESPACE "UNiagaraSpriteRendererProperties"
+
+UNiagaraSpriteRendererProperties::UNiagaraSpriteRendererProperties()
+	: SubImageSize(1.0f, 1.0f)
+	, Alignment(ENiagaraSpriteAlignment::Unaligned)
+	, FacingMode(ENiagaraSpriteFacingMode::FaceCamera)
+	, CustomFacingVectorMask(EForceInit::ForceInitToZero)
+	, SortMode(ENiagaraSortMode::SortNone)
+{
+}
+
+NiagaraRenderer* UNiagaraSpriteRendererProperties::CreateEmitterRenderer(ERHIFeatureLevel::Type FeatureLevel)
+{
+	return new NiagaraRendererSprites(FeatureLevel, this);
+}
+
+void UNiagaraSpriteRendererProperties::GetUsedMaterials(TArray<UMaterialInterface*>& OutMaterials) const
+{
+	OutMaterials.Add(Material);
+}
+
+
+#if WITH_EDITORONLY_DATA
+
+const TArray<FNiagaraVariable>& UNiagaraSpriteRendererProperties::GetRequiredAttributes()
+{
+	static TArray<FNiagaraVariable> Attrs;
+
+	if (Attrs.Num() == 0)
+	{
+		Attrs.Add(SYS_PARAM_PARTICLES_POSITION);
+		Attrs.Add(SYS_PARAM_PARTICLES_VELOCITY);
+		Attrs.Add(SYS_PARAM_PARTICLES_COLOR);
+		Attrs.Add(SYS_PARAM_PARTICLES_SPRITE_ROTATION);
+		Attrs.Add(SYS_PARAM_PARTICLES_NORMALIZED_AGE);
+		Attrs.Add(SYS_PARAM_PARTICLES_SPRITE_SIZE);
+	}
+
+	return Attrs;
+}
+
+
+const TArray<FNiagaraVariable>& UNiagaraSpriteRendererProperties::GetOptionalAttributes()
+{
+	static TArray<FNiagaraVariable> Attrs;
+
+	if (Attrs.Num() == 0)
+	{
+		Attrs.Add(SYS_PARAM_PARTICLES_SPRITE_FACING);
+		Attrs.Add(SYS_PARAM_PARTICLES_SPRITE_ALIGNMENT);
+		Attrs.Add(SYS_PARAM_PARTICLES_SUB_IMAGE_INDEX);
+		Attrs.Add(SYS_PARAM_PARTICLES_DYNAMIC_MATERIAL_PARAM);
+	}
+
+	return Attrs;
+}
+
+
+bool UNiagaraSpriteRendererProperties::IsMaterialValidForRenderer(UMaterial* InMaterial, FText& InvalidMessage)
+{
+	if (InMaterial->bUsedWithNiagaraSprites == false)
+	{
+		InvalidMessage = NSLOCTEXT("NiagaraSpriteRendererProperties", "InvalidMaterialMessage", "The material isn't marked as \"Used with particle sprites\"");
+		return false;
+	}
+	return true;
+}
+
+void UNiagaraSpriteRendererProperties::FixMaterial(UMaterial* InMaterial)
+{
+	InMaterial->Modify();
+	InMaterial->bUsedWithNiagaraSprites = true;
+	InMaterial->ForceRecompileForRendering();
+}
+
+#undef LOCTEXT_NAMESPACE
+
+#endif // WITH_EDITORONLY_DATA
+
+
+
+
diff --git a/Engine/Plugins/FX/Niagara/Source/Niagara/Private/NiagaraStats.h b/Engine/Plugins/FX/Niagara/Source/Niagara/Private/NiagaraStats.h
new file mode 100644
--- /dev/null
+++ b/Engine/Plugins/FX/Niagara/Source/Niagara/Private/NiagaraStats.h
@@ -0,0 +1,14 @@
+// Copyright 1998-2017 Epic Games, Inc. All Rights Reserved.
+
+
+#pragma once
+
+#include "CoreMinimal.h"
+#include "Stats/Stats.h"
+#include "NiagaraModule.h"
+
+DECLARE_STATS_GROUP(TEXT("Niagara"), STATGROUP_Niagara, STATCAT_Advanced);
+DECLARE_CYCLE_STAT(TEXT("Render Total"), STAT_NiagaraRender, STATGROUP_Niagara);
+DECLARE_DWORD_COUNTER_STAT(TEXT("NumParticles"), STAT_NiagaraNumParticles, STATGROUP_Niagara);
+DECLARE_CYCLE_STAT(TEXT("Constant Setup"), STAT_NiagaraConstants, STATGROUP_Niagara);
+DECLARE_CYCLE_STAT(TEXT("Tick"), STAT_NiagaraTick, STATGROUP_Niagara);
diff --git a/Engine/Plugins/FX/Niagara/Source/Niagara/Private/NiagaraSystem.cpp b/Engine/Plugins/FX/Niagara/Source/Niagara/Private/NiagaraSystem.cpp
new file mode 100644
--- /dev/null
+++ b/Engine/Plugins/FX/Niagara/Source/Niagara/Private/NiagaraSystem.cpp
@@ -0,0 +1,523 @@
+// Copyright 1998-2017 Epic Games, Inc. All Rights Reserved.
+
+
+#include "NiagaraSystem.h"
+#include "NiagaraRendererProperties.h"
+#include "NiagaraRenderer.h"
+#include "NiagaraConstants.h"
+#include "NiagaraScriptSourceBase.h"
+#include "NiagaraCustomVersion.h"
+#include "NiagaraModule.h"
+#include "ModuleManager.h"
+
+// UNiagaraSystemCategory::UNiagaraSystemCategory(const FObjectInitializer& ObjectInitializer)
+// 	: Super(ObjectInitializer)
+// {
+// }
+
+//////////////////////////////////////////////////////////////////////////
+
+UNiagaraSystem::UNiagaraSystem(const FObjectInitializer& ObjectInitializer)
+: Super(ObjectInitializer)
+{
+}
+
+void UNiagaraSystem::PostInitProperties()
+{
+	Super::PostInitProperties();
+#if WITH_EDITORONLY_DATA
+	bAutoImportChangedEmitters = true;
+#endif
+	if (HasAnyFlags(RF_ClassDefaultObject | RF_NeedLoad) == false)
+	{
+		SystemSpawnScript = NewObject<UNiagaraScript>(this, "SystemSpawnScript", RF_Transactional);
+		SystemSpawnScript->SetUsage(ENiagaraScriptUsage::SystemSpawnScript);
+
+		SystemUpdateScript = NewObject<UNiagaraScript>(this, "SystemUpdateScript", RF_Transactional);
+		SystemUpdateScript->SetUsage(ENiagaraScriptUsage::SystemUpdateScript);
+
+		SystemSpawnScriptSolo = NewObject<UNiagaraScript>(this, "SystemSpawnScriptSolo", RF_Transactional);
+		SystemSpawnScriptSolo->SetUsage(ENiagaraScriptUsage::SystemSpawnScript);
+
+		SystemUpdateScriptSolo = NewObject<UNiagaraScript>(this, "SystemUpdateScriptSolo", RF_Transactional);
+		SystemUpdateScriptSolo->SetUsage(ENiagaraScriptUsage::SystemUpdateScript);
+	}
+}
+
+void UNiagaraSystem::Serialize(FArchive& Ar)
+{
+	Super::Serialize(Ar);
+
+	Ar.UsingCustomVersion(FNiagaraCustomVersion::GUID);
+}
+
+void UNiagaraSystem::PostLoad()
+{
+	Super::PostLoad();
+
+	if (GIsEditor)
+	{
+		SetFlags(RF_Transactional);
+	}
+
+	// Check to see if our version is out of date. If so, we'll finally need to recompile.
+	bool bNeedsRecompile = false;
+	const int32 NiagaraVer = GetLinkerCustomVersion(FNiagaraCustomVersion::GUID);
+	if (NiagaraVer < FNiagaraCustomVersion::LatestVersion)
+	{
+		bNeedsRecompile = true;
+	}		
+	
+	// We will be using this later potentially, so make sure that it's postload is already up-to-date.
+	if (SystemSpawnScript)
+	{
+		SystemSpawnScript->ConditionalPostLoad();
+	}
+	if (SystemUpdateScript)
+	{
+		SystemUpdateScript->ConditionalPostLoad();
+	}
+
+ 	if (NiagaraVer < FNiagaraCustomVersion::SystemEmitterScriptSimulations)
+ 	{
+ 		bNeedsRecompile = true;
+		SystemSpawnScriptSolo = NewObject<UNiagaraScript>(this, "SystemSpawnScriptSolo", RF_Transactional);
+		SystemSpawnScriptSolo->SetUsage(ENiagaraScriptUsage::SystemSpawnScript);
+
+		SystemUpdateScriptSolo = NewObject<UNiagaraScript>(this, "SystemUpdateScriptSolo", RF_Transactional);
+		SystemUpdateScriptSolo->SetUsage(ENiagaraScriptUsage::SystemUpdateScript);
+#if WITH_EDITORONLY_DATA
+		if (SystemSpawnScript)
+		{
+			SystemSpawnScriptSolo->SetSource(SystemSpawnScript->GetSource());
+		}
+
+		if (SystemUpdateScript)
+		{
+			SystemUpdateScriptSolo->SetSource(SystemUpdateScript->GetSource());
+		}
+#endif
+ 	}
+
+	if (bNeedsRecompile)
+	{
+		for (FNiagaraEmitterHandle& Handle : EmitterHandles)
+		{
+			UNiagaraEmitter* Props = Handle.GetInstance();
+			if (Props)
+			{
+				// We will be refreshing later potentially, so make sure that it's postload is already up-to-date.
+				Props->ConditionalPostLoad();
+			}
+		}
+	}
+
+#if WITH_EDITORONLY_DATA
+	if (SystemSpawnScript == nullptr)
+	{
+		SystemSpawnScript = NewObject<UNiagaraScript>(this, "SystemSpawnScript", RF_Transactional);
+		SystemSpawnScript->SetUsage(ENiagaraScriptUsage::SystemSpawnScript);
+	}
+
+	if (SystemUpdateScript == nullptr)
+	{
+		SystemUpdateScript = NewObject<UNiagaraScript>(this, "SystemUpdateScript", RF_Transactional);
+		SystemUpdateScript->SetUsage(ENiagaraScriptUsage::SystemUpdateScript);
+	}
+
+	if (bNeedsRecompile)
+	{
+		Compile();
+	}
+
+	// Note that we do not call ResynchronizeAllHandles() or CheckForUpdates here.
+	// There are multiple reasons:
+	// 1) You can run into issues with the Linker table associated with this package when doing the deep copy during load. 
+	// 2) We'd have to make sure that all of the referenced nodes that were in the old version of the graph have already been loaded and cleaned out.
+	// Better to just CheckForUpdates anytime this System is used for the first time.
+#endif
+}
+
+#if WITH_EDITORONLY_DATA
+bool UNiagaraSystem::CheckForUpdates()
+{
+	// An emitter may have changed since last load, let's refresh if possible.
+	if (bAutoImportChangedEmitters)
+	{
+		return ResynchronizeAllHandles();
+	}
+	return false;
+}
+
+UObject* UNiagaraSystem::GetEditorData()
+{
+	return EditorData;
+}
+
+const UObject* UNiagaraSystem::GetEditorData() const
+{
+	return EditorData;
+}
+
+void UNiagaraSystem::SetEditorData(UObject* InEditorData)
+{
+	EditorData = InEditorData;
+}
+
+bool UNiagaraSystem::ResynchronizeAllHandles()
+{
+	uint32 HandlexIdx = 0;
+	bool bAnyRefreshed = false;
+	for (FNiagaraEmitterHandle& Handle : EmitterHandles)
+	{
+		const UNiagaraEmitter* SrcPrps = Handle.GetSource();
+		const UNiagaraEmitter* InstancePrps = Handle.GetInstance();
+
+		// Check to see if the source and handle are synched. If they aren't go ahead and resynch.
+		if (!Handle.IsSynchronizedWithSource())
+		{
+			// Todo.. eliminate synchronization until we have rewritten it...
+#if 0
+			if (Handle.RefreshFromSource())
+			{
+				check(SrcPrps->ChangeId == Handle.GetInstance()->ChangeId);
+				UE_LOG(LogNiagara, Log, TEXT("Refreshed System %s emitter %d from %s"), *GetPathName(), HandlexIdx, *SrcPrps->GetPathName());
+				bAnyRefreshed = true;
+			}
+#endif
+		}
+		HandlexIdx++;
+	}
+
+	if (bAnyRefreshed)
+	{
+		Compile();
+	}
+
+	return bAnyRefreshed;
+}
+
+bool UNiagaraSystem::ReferencesSourceEmitter(UNiagaraEmitter* Emitter)
+{
+	for (FNiagaraEmitterHandle& Handle : EmitterHandles)
+	{
+		if (Emitter == Handle.GetSource())
+		{
+			return true;
+		}
+	}
+	return false;
+}
+#endif
+
+
+const TArray<FNiagaraEmitterHandle>& UNiagaraSystem::GetEmitterHandles()
+{
+	return EmitterHandles;
+}
+
+
+bool UNiagaraSystem::IsValid()const
+{
+	if (!SystemSpawnScript->IsValid() || !SystemSpawnScriptSolo->IsValid() || !SystemUpdateScript->IsValid() || !SystemUpdateScriptSolo->IsValid())
+	{
+		return false;
+	}
+
+	for (const FNiagaraEmitterHandle& Handle : EmitterHandles)
+	{
+		if (!Handle.GetInstance()->IsValid())
+		{
+			return false;
+		}
+	}
+
+	return true;
+}
+#if WITH_EDITORONLY_DATA
+
+FNiagaraEmitterHandle UNiagaraSystem::AddEmitterHandle(const UNiagaraEmitter& SourceEmitter, FName EmitterName)
+{
+	FNiagaraEmitterHandle EmitterHandle(SourceEmitter, EmitterName, *this);
+	EmitterHandles.Add(EmitterHandle);
+	return EmitterHandle;
+}
+
+FNiagaraEmitterHandle UNiagaraSystem::AddEmitterHandleWithoutCopying(UNiagaraEmitter& Emitter)
+{
+	FNiagaraEmitterHandle EmitterHandle(Emitter);
+	EmitterHandles.Add(EmitterHandle);
+	return EmitterHandle;
+}
+
+FNiagaraEmitterHandle UNiagaraSystem::DuplicateEmitterHandle(const FNiagaraEmitterHandle& EmitterHandleToDuplicate, FName EmitterName)
+{
+	FNiagaraEmitterHandle EmitterHandle(EmitterHandleToDuplicate, EmitterName, *this);
+	EmitterHandles.Add(EmitterHandle);
+	return EmitterHandle;
+}
+#endif
+
+void UNiagaraSystem::RemoveEmitterHandle(const FNiagaraEmitterHandle& EmitterHandleToDelete)
+{
+	auto RemovePredicate = [&](const FNiagaraEmitterHandle& EmitterHandle) { return EmitterHandle.GetId() == EmitterHandleToDelete.GetId(); };
+	EmitterHandles.RemoveAll(RemovePredicate);
+}
+
+void UNiagaraSystem::RemoveEmitterHandlesById(const TSet<FGuid>& HandlesToRemove)
+{
+	auto RemovePredicate = [&](const FNiagaraEmitterHandle& EmitterHandle)
+	{
+		return HandlesToRemove.Contains(EmitterHandle.GetId());
+	};
+	EmitterHandles.RemoveAll(RemovePredicate);
+}
+
+UNiagaraScript* UNiagaraSystem::GetSystemSpawnScript(bool bSolo)
+{
+	return bSolo ? SystemSpawnScriptSolo : SystemSpawnScript;
+}
+
+UNiagaraScript* UNiagaraSystem::GetSystemUpdateScript(bool bSolo)
+{
+	return bSolo ? SystemUpdateScriptSolo : SystemUpdateScript;
+}
+
+#if WITH_EDITORONLY_DATA
+bool UNiagaraSystem::GetAutoImportChangedEmitters() const
+{
+	return bAutoImportChangedEmitters;
+}
+
+void UNiagaraSystem::SetAutoImportChangedEmitters(bool bShouldImport)
+{
+	bAutoImportChangedEmitters = bShouldImport;
+}
+
+void UNiagaraSystem::CompileScripts(TArray<ENiagaraScriptCompileStatus>& OutScriptStatuses, TArray<FString>& OutGraphLevelErrorMessages, TArray<FString>& PathNames, TArray<UNiagaraScript*>& Scripts)
+{
+	PathNames.Empty();
+	OutScriptStatuses.Empty();
+	OutGraphLevelErrorMessages.Empty();
+
+	TArray<FNiagaraVariable> ExposedVars;
+
+	//Compile all emitters
+	for (FNiagaraEmitterHandle Handle : EmitterHandles)
+	{
+		TArray<FString> ScriptErrors;
+		UNiagaraScriptSourceBase* GraphSource = Handle.GetInstance()->GraphSource;
+		check(!GraphSource->IsPreCompiled());
+		GraphSource->PreCompile(Handle.GetInstance());
+
+		Handle.GetInstance()->CompileScripts(OutScriptStatuses, ScriptErrors, PathNames, Scripts);
+
+		GraphSource->GatherPreCompiledVariables(TEXT("User"), ExposedVars);
+		GraphSource->PostCompile();
+		OutGraphLevelErrorMessages.Append(ScriptErrors);
+	}
+
+	check(SystemSpawnScript->GetSource() == SystemUpdateScript->GetSource());
+	check(SystemSpawnScript->GetSource()->IsPreCompiled() == false);
+
+	SystemSpawnScript->GetSource()->PreCompile(nullptr);
+	SystemSpawnScript->GetSource()->GatherPreCompiledVariables(TEXT("User"), ExposedVars);
+
+	FString ErrorMsg;
+	check((int32)EScriptCompileIndices::SpawnScript == 0);
+	OutScriptStatuses.Add(SystemSpawnScript->Compile(ErrorMsg));
+	OutGraphLevelErrorMessages.Add(ErrorMsg); // EScriptCompileIndices::SpawnScript
+	Scripts.Add(SystemSpawnScript);
+	PathNames.Add(SystemSpawnScript->GetPathName());
+
+	ErrorMsg.Empty();
+	check((int32)EScriptCompileIndices::UpdateScript == 1);
+	OutScriptStatuses.Add(SystemUpdateScript->Compile(ErrorMsg));
+	OutGraphLevelErrorMessages.Add(ErrorMsg); // EScriptCompileIndices::SpawnScript
+	Scripts.Add(SystemUpdateScript);
+	PathNames.Add(SystemUpdateScript->GetPathName());
+
+	ErrorMsg.Empty();
+	check((int32)EScriptCompileIndices::SpawnScript == 0);
+	OutScriptStatuses.Add(SystemSpawnScriptSolo->Compile(ErrorMsg));
+	OutGraphLevelErrorMessages.Add(ErrorMsg); // EScriptCompileIndices::SpawnScript
+	Scripts.Add(SystemSpawnScriptSolo);
+	PathNames.Add(SystemSpawnScriptSolo->GetPathName());
+
+	ErrorMsg.Empty();
+	check((int32)EScriptCompileIndices::UpdateScript == 1);
+	OutScriptStatuses.Add(SystemUpdateScriptSolo->Compile(ErrorMsg));
+	OutGraphLevelErrorMessages.Add(ErrorMsg); // EScriptCompileIndices::SpawnScript
+	Scripts.Add(SystemUpdateScriptSolo);
+	PathNames.Add(SystemUpdateScriptSolo->GetPathName());
+
+	SystemSpawnScript->GetSource()->PostCompile();
+
+	InitEmitterSpawnAttributes();
+	
+	// Now let's synchronize the variables that we expose to the end user.
+	TArray<FNiagaraVariable> OriginalVars;
+	ExposedParameters.GetParameters(OriginalVars);
+	for (int32 i = 0; i < ExposedVars.Num(); i++)
+	{
+		if (OriginalVars.Contains(ExposedVars[i]) == false)
+		{
+			ExposedParameters.AddParameter(ExposedVars[i]);
+		}
+	}
+	for (int32 i = 0; i < OriginalVars.Num(); i++)
+	{
+		if (ExposedVars.Contains(OriginalVars[i]) == false)
+		{
+			ExposedParameters.RemoveParameter(OriginalVars[i]);
+		}
+	}
+
+	FNiagaraSystemUpdateContext UpdateCtx(this, true);
+
+	INiagaraModule& NiagaraModule = FModuleManager::LoadModuleChecked<INiagaraModule>("Niagara");
+	NiagaraModule.DestroyAllSystemSimulations(this);
+
+	//I assume this was added to recompile emitters too but I'm doing that above. We need to rework this compilation flow.
+	//Compile();
+}
+
+void UNiagaraSystem::Compile()
+{
+	TArray<FNiagaraVariable> ExposedVars;
+	for (FNiagaraEmitterHandle Handle: EmitterHandles)
+	{
+		TArray<ENiagaraScriptCompileStatus> ScriptStatuses;
+		TArray<FString> ScriptErrors;
+		TArray<FString> PathNames;
+		TArray<UNiagaraScript*> Scripts;
+		
+		UNiagaraScriptSourceBase* GraphSource = Handle.GetInstance()->GraphSource;
+		check(!GraphSource->IsPreCompiled());
+		GraphSource->PreCompile(Handle.GetInstance());
+
+		Handle.GetInstance()->CompileScripts(ScriptStatuses, ScriptErrors, PathNames, Scripts);
+
+		GraphSource->GatherPreCompiledVariables(TEXT("User"), ExposedVars);
+
+		GraphSource->PostCompile();
+
+		for (int32 i = 0; i < ScriptStatuses.Num(); i++)
+		{
+			if (ScriptErrors[i].Len() != 0 || ScriptStatuses[i] != ENiagaraScriptCompileStatus::NCS_UpToDate)
+			{
+				UE_LOG(LogNiagara, Warning, TEXT("Script '%s', compile status: %d  errors: %s"), *PathNames[i], (int32)ScriptStatuses[i], *ScriptErrors[i]);
+			}
+			else
+			{
+				UE_LOG(LogNiagara, Log, TEXT("Script '%s', compile status: Success!"), *PathNames[i]);
+			}
+		}
+	}
+
+	check(SystemSpawnScript->GetSource() == SystemUpdateScript->GetSource());
+	check(SystemSpawnScript->GetSource()->IsPreCompiled() == false);
+
+	SystemSpawnScript->GetSource()->PreCompile(nullptr);
+	SystemSpawnScript->GetSource()->GatherPreCompiledVariables(TEXT("User"), ExposedVars);
+
+	FString CompileErrors;
+	ENiagaraScriptCompileStatus CompileStatus = SystemSpawnScript->Compile(CompileErrors);
+	if (CompileErrors.Len() != 0 || CompileStatus != ENiagaraScriptCompileStatus::NCS_UpToDate)
+	{
+		UE_LOG(LogNiagara, Warning, TEXT("System Spawn Script '%s', compile status: %d  errors: %s"), *SystemSpawnScript->GetPathName(), (int32)CompileStatus, *CompileErrors);
+	}
+	else
+	{
+		UE_LOG(LogNiagara, Log, TEXT("Script Spawn '%s', compile status: Success!"), *SystemSpawnScript->GetPathName());
+	}
+
+	CompileStatus = SystemUpdateScript->Compile(CompileErrors);
+	if (CompileErrors.Len() != 0 || CompileStatus != ENiagaraScriptCompileStatus::NCS_UpToDate)
+	{
+		UE_LOG(LogNiagara, Warning, TEXT("System Update Script '%s', compile status: %d  errors: %s"), *SystemUpdateScript->GetPathName(), (int32)CompileStatus, *CompileErrors);
+	}
+	else
+	{
+		UE_LOG(LogNiagara, Log, TEXT("Script Update '%s', compile status: Success!"), *SystemUpdateScript->GetPathName());
+	}
+
+	if(SystemSpawnScriptSolo)
+	{
+		CompileStatus = SystemSpawnScriptSolo->Compile(CompileErrors);
+		if (CompileErrors.Len() != 0 || CompileStatus != ENiagaraScriptCompileStatus::NCS_UpToDate)
+		{
+			UE_LOG(LogNiagara, Warning, TEXT("System Spawn Solo Script '%s', compile status: %d  errors: %s"), *SystemSpawnScriptSolo->GetPathName(), (int32)CompileStatus, *CompileErrors);
+		}
+		else
+		{
+			UE_LOG(LogNiagara, Log, TEXT("Script Spawn Solo '%s', compile status: Success!"), *SystemSpawnScriptSolo->GetPathName());
+		}
+	}
+
+	if (SystemUpdateScriptSolo)
+	{
+		CompileStatus = SystemUpdateScriptSolo->Compile(CompileErrors);
+		if (CompileErrors.Len() != 0 || CompileStatus != ENiagaraScriptCompileStatus::NCS_UpToDate)
+		{
+			UE_LOG(LogNiagara, Warning, TEXT("System Update Solo Script '%s', compile status: %d  errors: %s"), *SystemUpdateScriptSolo->GetPathName(), (int32)CompileStatus, *CompileErrors);
+		}
+		else
+		{
+			UE_LOG(LogNiagara, Log, TEXT("Script Update Solo '%s', compile status: Success!"), *SystemUpdateScriptSolo->GetPathName());
+		}
+	}
+	SystemSpawnScript->GetSource()->PostCompile();
+
+	InitEmitterSpawnAttributes();
+	TArray<FNiagaraVariable> OriginalVars;
+	ExposedParameters.GetParameters(OriginalVars);
+	for (int32 i = 0; i < ExposedVars.Num(); i++)
+	{
+		ExposedParameters.AddParameter(ExposedVars[i]);
+	}
+	for (int32 i = 0; i < OriginalVars.Num(); i++)
+	{
+		if (ExposedVars.Contains(OriginalVars[i]) == false)
+		{
+			ExposedParameters.RemoveParameter(OriginalVars[i]);
+		}
+	}
+
+	FNiagaraSystemUpdateContext UpdateCtx(this, true);
+
+	INiagaraModule& NiagaraModule = FModuleManager::LoadModuleChecked<INiagaraModule>("Niagara");
+	NiagaraModule.DestroyAllSystemSimulations(this);
+}
+#endif
+
+void UNiagaraSystem::InitEmitterSpawnAttributes()
+{
+	EmitterSpawnAttributes.Empty();
+	EmitterSpawnAttributes.SetNum(EmitterHandles.Num());
+	FNiagaraTypeDefinition SpawnInfoDef = FNiagaraTypeDefinition(FNiagaraSpawnInfo::StaticStruct());
+	for (FNiagaraVariable& Var : SystemSpawnScript->Attributes)
+	{
+		for (int32 EmitterIdx = 0; EmitterIdx < EmitterHandles.Num(); ++EmitterIdx)
+		{
+			UNiagaraEmitter* Emitter = EmitterHandles[EmitterIdx].GetInstance();
+			FString EmitterName = Emitter->GetUniqueEmitterName();
+			if (Var.GetType() == SpawnInfoDef && Var.GetName().ToString().StartsWith(EmitterName))
+			{
+				EmitterSpawnAttributes[EmitterIdx].SpawnAttributes.AddUnique(Var.GetName());
+			}
+		}
+	}
+	for (FNiagaraVariable& Var : SystemUpdateScript->Attributes)
+	{
+		for (int32 EmitterIdx = 0; EmitterIdx < EmitterHandles.Num(); ++EmitterIdx)
+		{
+			UNiagaraEmitter* Emitter = EmitterHandles[EmitterIdx].GetInstance();
+			FString EmitterName = Emitter->GetUniqueEmitterName();
+			if (Var.GetType() == SpawnInfoDef && Var.GetName().ToString().StartsWith(EmitterName))
+			{
+				EmitterSpawnAttributes[EmitterIdx].SpawnAttributes.AddUnique(Var.GetName());
+			}
+		}
+	}
+}
\ No newline at end of file
diff --git a/Engine/Plugins/FX/Niagara/Source/Niagara/Private/NiagaraSystemInstance.cpp b/Engine/Plugins/FX/Niagara/Source/Niagara/Private/NiagaraSystemInstance.cpp
new file mode 100644
--- /dev/null
+++ b/Engine/Plugins/FX/Niagara/Source/Niagara/Private/NiagaraSystemInstance.cpp
@@ -0,0 +1,876 @@
+// Copyright 1998-2016 Epic Games, Inc. All Rights Reserved.
+
+#include "NiagaraSystemInstance.h"
+#include "NiagaraConstants.h"
+#include "NiagaraCommon.h"
+#include "NiagaraDataInterface.h"
+#include "NiagaraStats.h"
+#include "Async/ParallelFor.h"
+#include "NiagaraParameterCollection.h"
+#include "NiagaraWorldManager.h"
+#include "NiagaraComponent.h"
+#include "NiagaraRenderer.h"
+#include "GameFramework/PlayerController.h"
+
+DECLARE_CYCLE_STAT(TEXT("Parallel Tick"), STAT_NiagaraParallelTick, STATGROUP_Niagara);
+
+
+/** Whether to allow emitters to tick in parallel */
+static int32 GbNiagaraParallelEmitterTick = 1;
+static FAutoConsoleVariableRef CVarParallelEmitterTick(
+	TEXT("niagara.ParallelEmitterTick"),
+	GbNiagaraParallelEmitterTick,
+	TEXT("Whether to tick individual emitters in an System in parallel"),
+	ECVF_Default
+	);
+ 
+
+
+FNiagaraSystemInstance::FNiagaraSystemInstance(UNiagaraComponent* InComponent)
+	: SystemInstanceIndex(INDEX_NONE)
+	, Component(InComponent)
+	, SystemSimulation(nullptr)
+	, Age(0.0f)
+	, ID(FGuid::NewGuid())
+	, IDName(*ID.ToString())
+	, bSolo(false)
+	, bForceSolo(false)
+	, bResetPending(false)
+	, bReInitPending(false)
+	, bError(false)
+	, bDataInterfacesNeedInit(true)
+	, ExecutionState(ENiagaraExecutionState::Inactive)
+{
+	SystemBounds.Init();
+}
+
+void FNiagaraSystemInstance::Init(FNiagaraSystemSimulation* InSystemSimulation, bool bForceReset, bool bInForceSolo)
+{
+	SystemSimulation = InSystemSimulation;
+// 	if (System != nullptr)NOOOOOPE!. This should be done via the editor or on post load of the system!
+// 	{
+// 		System->CheckForUpdates();
+// 	}
+
+	bError = false;
+	bForceSolo = bInForceSolo;
+	ExecutionState = ENiagaraExecutionState::Inactive;
+
+	//InstanceParameters = GetSystem()->GetInstanceParameters();
+	InitEmitters();
+	Reset(bForceReset ? EResetMode::ImmediateReInit : EResetMode::DeferredReInit);
+
+
+	OnInitializedDelegate.Broadcast();
+}
+
+void FNiagaraSystemInstance::SetExecutionState(ENiagaraExecutionState InState)
+{
+	ExecutionState = InState;
+}
+
+void FNiagaraSystemInstance::Activate(bool bReset /* = false */)
+{
+	if (GetSystem() && GetSystem()->IsValid())
+	{
+		if (ExecutionState != ENiagaraExecutionState::Active || bReset)
+		{
+			Reset(EResetMode::DeferredReset);
+		}
+
+		ExecutionState = ENiagaraExecutionState::Active;
+	}
+}
+
+void FNiagaraSystemInstance::Deactivate(bool bImmediate)
+{
+	ExecutionState = ENiagaraExecutionState::Dead;
+
+	if (bImmediate)
+	{
+		//Just kill this instance.
+		FNiagaraSystemSimulation* SystemSim = GetSystemSimulation();
+		SystemSim->RemoveInstance(this);
+	}
+	
+	//TODO: pipe into simulation?
+	//Drive a bool param at system level which systems can read? Stop spawning at all when not active so we can clean up?
+
+	//Need to work out how we're going to clean up system instances and components.
+	//When all emitters are disabled and have no remaining particles?
+}
+
+void FNiagaraSystemInstance::Reset(FNiagaraSystemInstance::EResetMode Mode)
+{
+	FNiagaraSystemSimulation* SystemSim = GetSystemSimulation();
+	SystemSim->RemoveInstance(this);
+	bPendingSpawn = true;
+	
+	//WIP: Currently tesing out having all resets be immediate. 
+	if (Mode == EResetMode::DeferredReset)
+	{
+		//bResetPending = true;
+		ResetInternal();
+	}
+	else if (Mode == EResetMode::ImmediateReset)
+	{
+		ResetInternal();
+	}
+	else if (Mode == EResetMode::DeferredReInit)
+	{
+		DestroyDataInterfaceInstanceData();
+		//bReInitPending = true;
+		ReInitInternal();
+	}
+	else if (Mode == EResetMode::ImmediateReInit)
+	{
+		DestroyDataInterfaceInstanceData();
+		ReInitInternal();
+	}
+
+	if (!bSolo && !bError)
+	{
+		SystemSim->AddInstance(this);
+	}
+	else if (bSolo && !bError)
+	{
+		SystemSim->ResetSolo(this);
+	}
+}
+
+void FNiagaraSystemInstance::ResetInternal()
+{
+	Age = 0;
+	if (GetSystem() == nullptr || Component == nullptr || bError)
+	{
+		return;
+	}
+
+
+	for (TSharedRef<FNiagaraEmitterInstance> Simulation : Emitters)
+	{
+		Simulation->ResetSimulation();
+	}
+
+	bResetPending = false;
+
+#if WITH_EDITOR
+	//UE_LOG(LogNiagara, Log, TEXT("OnResetInternal %p"), this);
+	OnResetDelegate.Broadcast();
+#endif
+}
+
+void FNiagaraSystemInstance::ReInitInternal()
+{
+	bool bPrevError = bError;
+	Age = 0;
+	if (GetSystem() == nullptr || Component == nullptr)
+	{
+		return;
+	}
+	bError = false;
+
+
+	for (TSharedRef<FNiagaraEmitterInstance> Simulation : Emitters)
+	{
+		Simulation->ReInitSimulation();
+	}
+
+	for (TSharedRef<FNiagaraEmitterInstance> Simulation : Emitters)
+	{
+		Simulation->PostResetSimulation();
+	}
+	
+	InstanceParameters.Empty();
+#if WITH_EDITOR
+	InstanceParameters = Component->GetInitialParameters();
+#endif
+
+	InstanceParameters.AddParameter(SYS_PARAM_ENGINE_POSITION);
+	InstanceParameters.AddParameter(SYS_PARAM_ENGINE_VELOCITY);
+	InstanceParameters.AddParameter(SYS_PARAM_ENGINE_X_AXIS);
+	InstanceParameters.AddParameter(SYS_PARAM_ENGINE_Y_AXIS);
+	InstanceParameters.AddParameter(SYS_PARAM_ENGINE_Z_AXIS);
+	InstanceParameters.AddParameter(SYS_PARAM_ENGINE_LOCAL_TO_WORLD);
+	InstanceParameters.AddParameter(SYS_PARAM_ENGINE_WORLD_TO_LOCAL);
+	InstanceParameters.AddParameter(SYS_PARAM_ENGINE_LOCAL_TO_WORLD_TRANSPOSED);
+	InstanceParameters.AddParameter(SYS_PARAM_ENGINE_WORLD_TO_LOCAL_TRANSPOSED);
+	InstanceParameters.AddParameter(SYS_PARAM_ENGINE_DELTA_TIME);
+	InstanceParameters.AddParameter(SYS_PARAM_ENGINE_INV_DELTA_TIME);
+	InstanceParameters.AddParameter(SYS_PARAM_ENGINE_MIN_DIST_TO_CAMERA);
+	InstanceParameters.AddParameter(SYS_PARAM_ENGINE_SYSTEM_NUM_EMITTERS);
+	InstanceParameters.AddParameter(SYS_PARAM_ENGINE_SYSTEM_NUM_EMITTERS_ALIVE);
+
+	OwnerPositionParam.Init(InstanceParameters, SYS_PARAM_ENGINE_POSITION);
+	OwnerVelocityParam.Init(InstanceParameters, SYS_PARAM_ENGINE_VELOCITY);
+	OwnerXAxisParam.Init(InstanceParameters, SYS_PARAM_ENGINE_X_AXIS);
+	OwnerYAxisParam.Init(InstanceParameters, SYS_PARAM_ENGINE_Y_AXIS);
+	OwnerZAxisParam.Init(InstanceParameters, SYS_PARAM_ENGINE_Z_AXIS);
+	OwnerTransformParam.Init(InstanceParameters, SYS_PARAM_ENGINE_LOCAL_TO_WORLD);
+	OwnerInverseParam.Init(InstanceParameters, SYS_PARAM_ENGINE_WORLD_TO_LOCAL);
+	OwnerTransposeParam.Init(InstanceParameters, SYS_PARAM_ENGINE_LOCAL_TO_WORLD_TRANSPOSED);
+	OwnerInverseTransposeParam.Init(InstanceParameters, SYS_PARAM_ENGINE_WORLD_TO_LOCAL_TRANSPOSED);
+	OwnerDeltaSecondsParam.Init(InstanceParameters, SYS_PARAM_ENGINE_DELTA_TIME);
+	OwnerInverseDeltaSecondsParam.Init(InstanceParameters, SYS_PARAM_ENGINE_INV_DELTA_TIME);
+	OwnerMinDistanceToCameraParam.Init(InstanceParameters, SYS_PARAM_ENGINE_MIN_DIST_TO_CAMERA);
+	SystemNumEmittersParam.Init(InstanceParameters, SYS_PARAM_ENGINE_SYSTEM_NUM_EMITTERS);
+	SystemNumEmittersAliveParam.Init(InstanceParameters, SYS_PARAM_ENGINE_SYSTEM_NUM_EMITTERS_ALIVE);
+
+	ParameterNumParticleBindings.SetNum(Emitters.Num());
+	for (int32 i = 0; i < Emitters.Num(); i++)
+	{
+		TSharedRef<FNiagaraEmitterInstance> Simulation = Emitters[i];
+		FString EmitterName = Simulation->GetEmitterHandle().GetInstance()->GetUniqueEmitterName();
+		FNiagaraVariable Var = SYS_PARAM_ENGINE_EMITTER_NUM_PARTICLES;
+		FString ParamName = Var.GetName().ToString().Replace(TEXT("Emitter"), *EmitterName);
+		Var.SetName(*ParamName);
+		InstanceParameters.AddParameter(Var);
+		ParameterNumParticleBindings[i].Init(InstanceParameters, Var);
+	}
+
+	TickInstanceParameters(0.01f);
+
+	//Determine if we can update normally or have to update solo.
+	bSolo = bForceSolo || InstanceParameters.GetDataInterfaces().Num() > 0;
+	//If our scripts have any interfaces that require instance data.
+	UNiagaraScript* SystemSpawn = GetSystem()->GetSystemSpawnScript();
+	for (int32 i = 0; !bSolo && i < SystemSpawn->DataInterfaceInfo.Num(); ++i)
+	{
+		FNiagaraScriptDataInterfaceInfo& Info = SystemSpawn->DataInterfaceInfo[i];
+		if (Info.DataInterface->PerInstanceDataSize() > 0)
+		{
+			bSolo = true;
+		}
+	}	
+	UNiagaraScript* SystemUpdate = GetSystem()->GetSystemUpdateScript();
+	for (int32 i = 0; !bSolo && i < SystemUpdate->DataInterfaceInfo.Num(); ++i)
+	{
+		FNiagaraScriptDataInterfaceInfo& Info = SystemUpdate->DataInterfaceInfo[i];
+		if (Info.DataInterface->PerInstanceDataSize() > 0)
+		{
+			bSolo = true;
+		}
+	}
+
+	BindParameters();
+
+	InitDataInterfaces(); // We can't wait until tick b/c the data interfaces need to be around when we update render modules.
+	FlushRenderingCommands();
+
+	// This gets a little tricky, but we want to delete any renderers that are no longer in use on the rendering thread, but
+	// first (to be safe), we want to update the proxy to point to the new renderer objects.
+
+	// Step 1: Recreate the renderers on the simulations, we keep the old and new renderers.
+	TArray<NiagaraRenderer*> NewRenderers;
+	TArray<NiagaraRenderer*> OldRenderers;
+
+	UpdateRenderModules(Component->GetWorld()->FeatureLevel, NewRenderers, OldRenderers);
+
+	// Step 2: Update the proxy with the new renderers that were created.
+	UpdateProxy(NewRenderers);
+	Component->MarkRenderStateDirty();
+
+	// Step 3: Queue up the old renderers for deletion on the render thread.
+	for (NiagaraRenderer* Renderer : OldRenderers)
+	{
+		if (Renderer != nullptr)
+		{
+			Renderer->Release();
+		}
+	}
+
+	bReInitPending = false;
+
+#if WITH_EDITOR
+	//UE_LOG(LogNiagara, Log, TEXT("OnResetInternal %p"), this);
+	OnResetDelegate.Broadcast();
+#endif
+
+}
+
+FNiagaraSystemInstance::~FNiagaraSystemInstance()
+{
+	//UE_LOG(LogNiagara, Warning, TEXT("~FNiagaraSystemInstance %p"), this);
+
+	//FlushRenderingCommands();
+
+	if (SystemInstanceIndex != INDEX_NONE)
+	{
+		FNiagaraSystemSimulation* SystemSim = GetSystemSimulation();
+		SystemSim->RemoveInstance(this);
+	}
+
+	DestroyDataInterfaceInstanceData();
+
+	// Clear out the System renderer from the proxy.
+	TArray<NiagaraRenderer*> NewRenderers;
+	UpdateProxy(NewRenderers);
+
+	// Clear out the System renderer from the simulation.
+	for (TSharedRef<FNiagaraEmitterInstance> Simulation : Emitters)
+	{
+		Simulation->ClearRenderer();
+	}
+
+	UnbindParameters();
+
+	// Clear out the emitters.
+	Emitters.Empty(0);
+
+// #if WITH_EDITOR
+// 	OnDestroyedDelegate.Broadcast();
+// #endif
+}
+
+//Unsure on usage of this atm. Possibly useful in future.
+// void FNiagaraSystemInstance::RebindParameterCollection(UNiagaraParameterCollectionInstance* OldInstance, UNiagaraParameterCollectionInstance* NewInstance)
+// {
+// 	OldInstance->GetParameterStore().Unbind(&InstanceParameters);
+// 	NewInstance->GetParameterStore().Bind(&InstanceParameters);
+// 
+// 	for (TSharedRef<FNiagaraEmitterInstance> Simulation : Emitters)
+// 	{
+// 		Simulation->RebindParameterCollection(OldInstance, NewInstance);
+// 	}
+// 
+// 	//Have to re init the instance data for data interfaces.
+// 	//This is actually lots more work than absolutely needed in some cases so we can improve it a fair bit.
+// 	InitDataInterfaces();
+// }
+
+void FNiagaraSystemInstance::BindParameters()
+{
+	Component->GetInitialParameters().Bind(&InstanceParameters);
+
+	for (TSharedRef<FNiagaraEmitterInstance> Simulation : Emitters)
+	{
+		Simulation->BindParameters();
+	}
+}
+
+void FNiagaraSystemInstance::UnbindParameters()
+{
+	Component->GetInitialParameters().Unbind(&InstanceParameters);
+
+	for (TSharedRef<FNiagaraEmitterInstance> Simulation : Emitters)
+	{
+		Simulation->UnbindParameters();
+	}
+}
+
+FNiagaraWorldManager* FNiagaraSystemInstance::GetWorldManager()const
+{
+	return Component ? FNiagaraWorldManager::Get(Component->GetWorld()) : nullptr; 
+}
+
+void FNiagaraSystemInstance::InitDataInterfaces()
+{
+	// If either the System or the component is invalid, it is possible that our cached data interfaces
+	// are now bogus and could point to invalid memory. Only the UNiagaraComponent or UNiagaraSystem
+	// can hold onto GC references to the DataInterfaces.
+	if (GetSystem() == nullptr || bError)
+	{
+		return;
+	}
+
+	if (Component == nullptr)
+	{
+		return;
+	}
+	
+	bDataInterfacesNeedInit = false;
+
+	DestroyDataInterfaceInstanceData();
+
+	//Now the interfaces in the simulations are all correct, we can build the per instance data table.
+	int32 InstanceDataSize = 0;
+	DataInterfaceInstanceDataOffsets.Empty();
+	auto CalcInstDataSize = [&](const TArray<UNiagaraDataInterface*>& Interfaces)
+	{
+		for (UNiagaraDataInterface* Interface : Interfaces)
+		{
+			if (int32 Size = Interface->PerInstanceDataSize())
+			{
+				int32* ExistingInstanceDataOffset = DataInterfaceInstanceDataOffsets.Find(Interface);
+				if (!ExistingInstanceDataOffset)//Don't add instance data for interfaces we've seen before.
+				{
+					DataInterfaceInstanceDataOffsets.Add(Interface) = InstanceDataSize;
+					InstanceDataSize += Size;
+				}
+			}
+		}
+	};
+
+	CalcInstDataSize(InstanceParameters.GetDataInterfaces());//This probably should be a proper exec context. 
+
+	//Iterate over interfaces to get size for table and clear their interface bindings.
+	for (TSharedRef<FNiagaraEmitterInstance> Simulation : Emitters)
+	{
+		FNiagaraEmitterInstance& Sim = Simulation.Get();
+		CalcInstDataSize(Sim.GetSpawnExecutionContext().GetDataInterfaces());
+		CalcInstDataSize(Sim.GetUpdateExecutionContext().GetDataInterfaces());
+		for (int32 i = 0; i < Sim.GetEventExecutionContexts().Num(); i++)
+		{
+			CalcInstDataSize(Sim.GetEventExecutionContexts()[i].GetDataInterfaces());
+		}
+
+		//Also force a rebind while we're here.
+		Sim.DirtyDataInterfaces();
+	}
+
+	DataInterfaceInstanceData.SetNumUninitialized(InstanceDataSize);
+
+	bool bOk = true;
+	for (TPair<TWeakObjectPtr<UNiagaraDataInterface>, int32>& Pair : DataInterfaceInstanceDataOffsets)
+	{
+		if (UNiagaraDataInterface* Interface = Pair.Key.Get())
+		{
+			//Ideally when we make the batching changes, we can keep the instance data in big single type blocks that can all be updated together with a single virtual call.
+			bOk &= Pair.Key->InitPerInstanceData(&DataInterfaceInstanceData[Pair.Value], this);
+		}
+		else
+		{
+			UE_LOG(LogNiagara, Error, TEXT("A data interface currently in use by an System has been destroyed."));
+			bOk = false;
+		}
+	}
+
+	if (!bOk)
+	{
+		UE_LOG(LogNiagara, Error, TEXT("Error initializing data interfaces."));
+		bError = true;
+		Component->MarkRenderStateDirty();
+	}
+}
+
+void FNiagaraSystemInstance::TickDataInterfaces(float DeltaSeconds)
+{
+	if (!GetSystem() || !Component || bError)
+	{
+		return;
+	}
+
+	bool bReInitDataInterfaces = false;
+	for (TPair<TWeakObjectPtr<UNiagaraDataInterface>, int32>& Pair : DataInterfaceInstanceDataOffsets)
+	{
+		if(UNiagaraDataInterface* Interface = Pair.Key.Get())
+		{
+			//Ideally when we make the batching changes, we can keep the instance data in big single type blocks that can all be updated together with a single virtual call.
+			bReInitDataInterfaces |= Interface->PerInstanceTick(&DataInterfaceInstanceData[Pair.Value], this, DeltaSeconds);
+		}
+	}
+	
+	//This should ideally really only happen at edit time.
+	if (bReInitDataInterfaces)
+	{
+		InitDataInterfaces();	
+	}
+}
+
+void FNiagaraSystemInstance::TickInstanceParameters(float DeltaSeconds)
+{
+	//TODO: Create helper binding objects to avoid the search in set parameter value.
+	//Set System params.
+	FTransform ComponentTrans = Component->GetComponentTransform();
+	FVector OldPos = OwnerPositionParam.GetValue();// ComponentTrans.GetLocation();
+	FVector CurrPos = ComponentTrans.GetLocation();
+	OwnerPositionParam.SetValue(CurrPos);
+	OwnerVelocityParam.SetValue((CurrPos - OldPos) / DeltaSeconds);
+	OwnerXAxisParam.SetValue(ComponentTrans.GetRotation().GetAxisX());
+	OwnerYAxisParam.SetValue(ComponentTrans.GetRotation().GetAxisY());
+	OwnerZAxisParam.SetValue(ComponentTrans.GetRotation().GetAxisZ());
+
+	FMatrix Transform = ComponentTrans.ToMatrixWithScale();
+	FMatrix Inverse = Transform.Inverse();
+	FMatrix Transpose = Transform.GetTransposed();
+	FMatrix InverseTranspose = Inverse.GetTransposed();
+	OwnerTransformParam.SetValue(Transform);
+	OwnerInverseParam.SetValue(Inverse);
+	OwnerTransposeParam.SetValue(Transpose);
+	OwnerInverseTransposeParam.SetValue(InverseTranspose);
+	OwnerDeltaSecondsParam.SetValue(DeltaSeconds);
+	OwnerInverseDeltaSecondsParam.SetValue(1.0f / DeltaSeconds);
+
+	//Calculate the min distance to a camera.
+	UWorld* World = Component->GetWorld();
+	if (World != NULL)
+	{
+		TArray<FVector, TInlineAllocator<8> > PlayerViewLocations;
+		if (World->GetPlayerControllerIterator())
+		{
+			for (FConstPlayerControllerIterator Iterator = World->GetPlayerControllerIterator(); Iterator; ++Iterator)
+			{
+				APlayerController* PlayerController = Iterator->Get();
+				if (PlayerController->IsLocalPlayerController())
+				{
+					FVector* POVLoc = new(PlayerViewLocations) FVector;
+					FRotator POVRotation;
+					PlayerController->GetPlayerViewPoint(*POVLoc, POVRotation);
+				}
+			}
+		}
+		else
+		{
+			PlayerViewLocations.Append(World->ViewLocationsRenderedLastFrame);
+		}
+
+		float LODDistanceSqr = (PlayerViewLocations.Num() ? FMath::Square(WORLD_MAX) : 0.0f);
+		for (const FVector& ViewLocation : PlayerViewLocations)
+		{
+			const float DistanceToEffectSqr = FVector(ViewLocation - CurrPos).SizeSquared();
+			if (DistanceToEffectSqr < LODDistanceSqr)
+			{
+				LODDistanceSqr = DistanceToEffectSqr;
+			}
+		}
+		OwnerMinDistanceToCameraParam.SetValue(FMath::Sqrt(LODDistanceSqr));
+	}
+
+	int32 NumAlive = 0;
+	for (int32 i = 0; i < Emitters.Num(); i++)
+	{
+		int32 NumParticles = Emitters[i]->GetNumParticles();
+		if (Emitters[i]->GetExecutionState() != ENiagaraExecutionState::Dead)
+		{
+			NumAlive++;
+		}
+		ParameterNumParticleBindings[i].SetValue(NumParticles);
+	}
+	SystemNumEmittersParam.SetValue(Emitters.Num());
+	SystemNumEmittersAliveParam.SetValue(NumAlive);
+	
+	Component->GetInitialParameters().Tick();
+	InstanceParameters.Tick();
+	InstanceParameters.bParametersDirty = true;
+}
+
+#if WITH_EDITORONLY_DATA
+
+bool FNiagaraSystemInstance::UsesEmitter(const UNiagaraEmitter* Emitter)const
+{
+	if (GetSystem())
+	{
+		for (FNiagaraEmitterHandle EmitterHandle : GetSystem()->GetEmitterHandles())
+		{
+			if (Emitter == EmitterHandle.GetSource() || Emitter == EmitterHandle.GetInstance())
+			{
+				return true;
+			}
+		}
+	}
+	return false;
+}
+
+bool FNiagaraSystemInstance::UsesScript(const UNiagaraScript* Script)const
+{
+	if (GetSystem())
+	{
+		for (FNiagaraEmitterHandle EmitterHandle : GetSystem()->GetEmitterHandles())
+		{
+			if ((EmitterHandle.GetSource() && EmitterHandle.GetSource()->UsesScript(Script)) || (EmitterHandle.GetInstance() && EmitterHandle.GetInstance()->UsesScript(Script)))
+			{
+				return true;
+			}
+		}
+	}
+	return false;
+}
+
+// bool FNiagaraSystemInstance::UsesDataInterface(UNiagaraDataInterface* Interface)
+// {
+// 
+// }
+
+bool FNiagaraSystemInstance::UsesCollection(const UNiagaraParameterCollection* Collection)const
+{
+	if (GetSystem())
+	{
+		for (FNiagaraEmitterHandle EmitterHandle : GetSystem()->GetEmitterHandles())
+		{
+			if (EmitterHandle.GetSource()->UsesCollection(Collection) || EmitterHandle.GetInstance()->UsesCollection(Collection))
+			{
+				return true;
+			}
+		}
+	}
+	return false;
+}
+
+#endif
+
+void FNiagaraSystemInstance::InitEmitters()
+{
+	if (Component)
+	{
+		Component->MarkRenderStateDirty();
+	}
+
+	// Just in case this ends up being called more than in Init, we need to 
+	// clear out the update proxy of any renderers that will be destroyed when Emitters.Empty occurs..
+	TArray<NiagaraRenderer*> NewRenderers;
+	UpdateProxy(NewRenderers);
+
+	// Clear out the System renderer from the simulation.
+	for (TSharedRef<FNiagaraEmitterInstance> Simulation : Emitters)
+	{
+		Simulation->ClearRenderer();
+	}
+
+	Emitters.Empty();
+	if (GetSystem() != nullptr)
+	{
+		const TArray<FNiagaraEmitterHandle>& EmitterHandles = GetSystem()->GetEmitterHandles();
+		for (int32 EmitterIdx=0; EmitterIdx < GetSystem()->GetEmitterHandles().Num(); ++EmitterIdx)
+		{
+			const FNiagaraEmitterHandle& EmitterHandle = EmitterHandles[EmitterIdx];
+			TSharedRef<FNiagaraEmitterInstance> Sim = MakeShareable(new FNiagaraEmitterInstance(this));
+			Sim->Init(EmitterIdx, IDName);
+			Emitters.Add(Sim);
+		}
+		bDataInterfacesNeedInit = true;
+	}
+}
+
+void FNiagaraSystemInstance::UpdateRenderModules(ERHIFeatureLevel::Type InFeatureLevel, TArray<NiagaraRenderer*>& OutNewRenderers, TArray<NiagaraRenderer*>& OutOldRenderers)
+{
+	for (TSharedPtr<FNiagaraEmitterInstance> Sim : Emitters)
+	{
+		Sim->UpdateEmitterRenderer(InFeatureLevel, OutNewRenderers, OutOldRenderers);
+	}
+}
+
+void FNiagaraSystemInstance::UpdateProxy(TArray<NiagaraRenderer*>& InRenderers)
+{
+	if (!GetSystem() || !Component)
+	{
+		return;
+	}
+
+	FNiagaraSceneProxy *NiagaraProxy = static_cast<FNiagaraSceneProxy*>(Component->SceneProxy);
+	if (NiagaraProxy)
+	{
+		if (Component != nullptr && Component->GetWorld() != nullptr)
+		{
+			// Tell the scene proxy on the render thread to update its System renderers.
+			ENQUEUE_UNIQUE_RENDER_COMMAND_TWOPARAMETER(
+				FChangeNiagaraRenderModule,
+				FNiagaraSceneProxy*, InProxy, NiagaraProxy,
+				TArray<NiagaraRenderer*>, InRendererArray, InRenderers,
+				{
+					InProxy->UpdateEmitterRenderers(InRendererArray);
+				}
+			);
+		}
+	}
+}
+
+void FNiagaraSystemInstance::ComponentTick(float DeltaSeconds)
+{
+	if (bError)
+	{
+		return;
+	}
+
+	FNiagaraSystemSimulation* Sim = GetSystemSimulation();
+	check(Sim);
+	check(IsInGameThread());
+	check(bSolo);
+	check(Component);
+
+	PreSimulateTick(DeltaSeconds);
+
+	FNiagaraSystemSimulation* SystemSim = GetSystemSimulation();
+	SystemSim->TickSolo(this);
+
+	PostSimulateTick(DeltaSeconds);
+}
+
+void FNiagaraSystemInstance::PreSimulateTick(float DeltaSeconds)
+{
+	HandleResets();
+
+	TickInstanceParameters(DeltaSeconds);
+}
+
+void FNiagaraSystemInstance::PostSimulateTick(float DeltaSeconds)
+{
+	if (GetSystem() == nullptr || Component == nullptr || DeltaSeconds < SMALL_NUMBER || bError)
+	{
+		return;
+	}
+
+	// pass the constants down to the emitter
+	// TODO: should probably just pass a pointer to the table
+	for (TPair<FNiagaraDataSetID, FNiagaraDataSet>& EventSetPair : ExternalEvents)
+	{
+		EventSetPair.Value.Tick();
+	}
+
+	if (bDataInterfacesNeedInit)
+	{
+		InitDataInterfaces();
+	}
+
+	//TODO: Now we're batching in the world manager we can store these in big blocks across systems and update all instance data in one virtual call?
+	TickDataInterfaces(DeltaSeconds);
+
+	// need to run this serially, as PreTick will initialize data interfaces, which can't be done in parallel
+	for (TSharedRef<FNiagaraEmitterInstance>&it : Emitters)
+	{
+		const UNiagaraEmitter *Props = it->GetEmitterHandle().GetInstance();
+		check(Props);
+		it->PreTick();
+	}
+
+	// now tick all emitters in parallel
+	{
+		SCOPE_CYCLE_COUNTER(STAT_NiagaraParallelTick);
+
+ 		if (GbNiagaraParallelEmitterTick && bSolo)
+ 		{
+ 			ParallelFor(Emitters.Num(),
+ 				[&](int32 EmitterIdx)
+ 			{
+ 				TSharedRef<FNiagaraEmitterInstance>&it = Emitters[EmitterIdx];
+ 				if (it->GetExecutionState() != ENiagaraExecutionState::Dead && it->GetExecutionState() != ENiagaraExecutionState::Paused)
+ 				{
+ 					it->Tick(DeltaSeconds);
+ 				}
+ 			});
+ 		}
+ 		else
+		{
+			for (int32 EmitterIdx = 0; EmitterIdx < Emitters.Num(); EmitterIdx++)
+			{
+				TSharedRef<FNiagaraEmitterInstance>&it = Emitters[EmitterIdx];
+				if (it->GetExecutionState() != ENiagaraExecutionState::Dead && it->GetExecutionState() != ENiagaraExecutionState::Paused)
+				{
+					it->Tick(DeltaSeconds);
+				}
+			}
+		}
+	}
+
+	Age += DeltaSeconds;
+}
+
+void FNiagaraSystemInstance::HandleResets()
+{
+	if (bReInitPending)
+	{
+		ReInitInternal();
+	}
+	else if (bResetPending)
+	{
+		ResetInternal();
+	}
+}
+
+void FNiagaraSystemInstance::Disable()
+{
+	FNiagaraSystemSimulation* SystemSim = GetSystemSimulation();
+	SystemSim->RemoveInstance(this);
+
+	// Clear out the System renderer from the proxy.
+	TArray<NiagaraRenderer*> NewRenderers;
+	UpdateProxy(NewRenderers);
+
+	// Clear out the System renderer from the simulation.
+	for (TSharedRef<FNiagaraEmitterInstance> Simulation : Emitters)
+	{
+		Simulation->ClearRenderer();
+	}
+
+	Emitters.Empty();
+	SetExecutionState(ENiagaraExecutionState::Disabled);
+
+	if (Component)
+	{
+		Component->OnSystemDisabled();
+	}
+}
+
+void FNiagaraSystemInstance::Enable()
+{
+	ExecutionState = ENiagaraExecutionState::Inactive;
+
+	InitEmitters();
+	Reset(EResetMode::ImmediateReInit);
+}
+
+void FNiagaraSystemInstance::ReinitializeDataInterfaces()
+{
+	// Update next tick
+	bDataInterfacesNeedInit = true;
+}
+
+void FNiagaraSystemInstance::DestroyDataInterfaceInstanceData()
+{
+	for (TPair<TWeakObjectPtr<UNiagaraDataInterface>, int32>& Pair : DataInterfaceInstanceDataOffsets)
+	{
+		if (UNiagaraDataInterface* Interface = Pair.Key.Get())
+		{
+			Interface->DestroyPerInstanceData(&DataInterfaceInstanceData[Pair.Value], this);
+		}
+	}
+	DataInterfaceInstanceDataOffsets.Empty();
+	DataInterfaceInstanceData.Empty();
+}
+
+TSharedPtr<FNiagaraEmitterInstance> FNiagaraSystemInstance::GetSimulationForHandle(const FNiagaraEmitterHandle& EmitterHandle)
+{
+	for (TSharedPtr<FNiagaraEmitterInstance> Sim : Emitters)
+	{
+		if(Sim->GetEmitterHandle().GetId() == EmitterHandle.GetId())
+		{
+			return Sim;
+		}
+	}
+	return nullptr;
+}
+
+UNiagaraSystem* FNiagaraSystemInstance::GetSystem()const
+{
+	return GetSystemSimulation()->GetSystem(); 
+}
+
+FNiagaraDataSet* FNiagaraSystemInstance::GetDataSet(FNiagaraDataSetID SetID, FName EmitterName)
+{
+	if (EmitterName == NAME_None)
+	{
+		if (FNiagaraDataSet* ExternalSet = ExternalEvents.Find(SetID))
+		{
+			return ExternalSet;
+		}
+	}
+	for (TSharedPtr<FNiagaraEmitterInstance> Emitter : Emitters)
+	{
+		check(Emitter.IsValid());
+		if (Emitter->IsEnabled())
+		{
+			if (Emitter->GetEmitterHandle().GetIdName() == EmitterName)
+			{
+				return Emitter->GetDataSet(SetID);
+			}
+		}
+	}
+
+	return NULL;
+}
+
+FNiagaraSystemInstance::FOnInitialized& FNiagaraSystemInstance::OnInitialized()
+{
+	return OnInitializedDelegate;
+}
+
+#if WITH_EDITOR
+FNiagaraSystemInstance::FOnReset& FNiagaraSystemInstance::OnReset()
+{
+	return OnResetDelegate;
+}
+
+FNiagaraSystemInstance::FOnDestroyed& FNiagaraSystemInstance::OnDestroyed()
+{
+	return OnDestroyedDelegate;
+}
+#endif
