Copying //UE4/Orion-Staging to //UE4/Main (Source: //Orion/Dev-General @ 3606378)

diff --git a/Engine/Plugins/FX/Niagara/Source/Niagara/Private/NiagaraSystemSimulation.cpp b/Engine/Plugins/FX/Niagara/Source/Niagara/Private/NiagaraSystemSimulation.cpp
new file mode 100644
--- /dev/null
+++ b/Engine/Plugins/FX/Niagara/Source/Niagara/Private/NiagaraSystemSimulation.cpp
@@ -0,0 +1,729 @@
+// Copyright 1998-2017 Epic Games, Inc. All Rights Reserved.
+
+#include "NiagaraSystemSimulation.h"
+#include "NiagaraModule.h"
+#include "Modules/ModuleManager.h"
+#include "NiagaraTypes.h"
+#include "NiagaraEvents.h"
+#include "NiagaraSettings.h"
+#include "NiagaraParameterCollection.h"
+#include "NiagaraSystemInstance.h"
+#include "NiagaraConstants.h"
+#include "NiagaraStats.h"
+#include "ParallelFor.h"
+#include "NiagaraComponent.h"
+#include "NiagaraWorldManager.h"
+
+DECLARE_CYCLE_STAT(TEXT("System Simulation (Batched)"), STAT_NiagaraSystemSim, STATGROUP_Niagara);
+DECLARE_CYCLE_STAT(TEXT("System Pre Simulate (Batched)"), STAT_NiagaraSystemSim_PreSimulate, STATGROUP_Niagara);
+DECLARE_CYCLE_STAT(TEXT("System Update (Batched)"), STAT_NiagaraSystemSim_Update, STATGROUP_Niagara);
+DECLARE_CYCLE_STAT(TEXT("System Spawn (Batched)"), STAT_NiagaraSystemSim_Spawn, STATGROUP_Niagara);
+DECLARE_CYCLE_STAT(TEXT("System Transfer Parameters (Batched)"), STAT_NiagaraSystemSim_TransferParameters, STATGROUP_Niagara);
+DECLARE_CYCLE_STAT(TEXT("System Post Simulate (Batched)"), STAT_NiagaraSystemSim_PostSimulate, STATGROUP_Niagara);
+DECLARE_CYCLE_STAT(TEXT("System Mark Component Dirty"), STAT_NiagaraSystemSim_MarkComponentDirty, STATGROUP_Niagara);
+
+DECLARE_CYCLE_STAT(TEXT("System Simulation (Solo)"), STAT_NiagaraSystemSimSolo, STATGROUP_Niagara);
+
+
+static int32 GbDumpSystemData = 0;
+static FAutoConsoleVariableRef CVarNiagaraDumpSystemData(
+	TEXT("fx.DumpSystemData"),
+	GbDumpSystemData,
+	TEXT("If > 0, results of system simulations will be dumped to the log. \n"),
+	ECVF_Default
+);
+
+static int32 GbParallelSystemPreTick = 1;
+static FAutoConsoleVariableRef CVarNiagaraParallelSystemPreTick(
+	TEXT("fx.ParallelSystemPreTick"),
+	GbParallelSystemPreTick,
+	TEXT("If > 0, system pre tick is parallelized. \n"),
+	ECVF_Default
+);
+
+static int32 GbParallelSystemPostTick = 1;
+static FAutoConsoleVariableRef CVarNiagaraParallelSystemPostTick(
+	TEXT("fx.ParallelSystemPostTick"),
+	GbParallelSystemPostTick,
+	TEXT("If > 0, system post tick is parallelized. \n"),
+	ECVF_Default
+);
+
+//TODO: Experiment with parallel param transfer.
+//static int32 GbParallelSystemParamTransfer = 1;
+//static FAutoConsoleVariableRef CVarNiagaraParallelSystemParamTransfer(
+//	TEXT("fx.ParallelSystemParamTransfer"),
+//	GbParallelSystemParamTransfer,
+//	TEXT("If > 0, system param transfer is parallelized. \n"),
+//	ECVF_Default
+//);
+
+//////////////////////////////////////////////////////////////////////////
+
+bool FNiagaraSystemSimulation::Init(UNiagaraSystem* InSystem, UWorld* InWorld)
+{
+	System = InSystem;
+	World = InWorld;
+
+	FNiagaraWorldManager* WorldMan = FNiagaraWorldManager::Get(InWorld);
+	check(WorldMan);
+
+	bCanExecute = System->GetSystemSpawnScript()->IsValid() && System->GetSystemUpdateScript()->IsValid();
+	bCanExecuteSolo = System->GetSystemSpawnScript(true)->IsValid() && System->GetSystemUpdateScript(true)->IsValid();
+	const UEnum* EnumPtr = FNiagaraTypeDefinition::GetExecutionStateEnum();
+
+	if (bCanExecute)
+	{
+		DataSet.Reset();
+		DataSet.AddVariables(System->GetSystemSpawnScript()->Attributes);
+		DataSet.AddVariables(System->GetSystemUpdateScript()->Attributes);
+		DataSet.Finalize();
+
+		SpawnParameterDataSet.Reset();
+		FNiagaraParameters* EngineParamsSpawn = System->GetSystemSpawnScript()->DataSetToParameters.Find(TEXT("Engine"));
+		if (EngineParamsSpawn != nullptr)
+		{
+			SpawnParameterDataSet.AddVariables(EngineParamsSpawn->Parameters);
+		}
+		SpawnParameterDataSet.Finalize();
+		UpdateParameterDataSet.Reset();
+		FNiagaraParameters* EngineParamsUpdate = System->GetSystemUpdateScript()->DataSetToParameters.Find(TEXT("Engine"));
+		if (EngineParamsUpdate != nullptr)
+		{
+			UpdateParameterDataSet.AddVariables(EngineParamsUpdate->Parameters);
+		}
+		UpdateParameterDataSet.Finalize();
+
+		SpawnExecContext.Init(System->GetSystemSpawnScript(false), ENiagaraSimTarget::CPUSim);
+		UpdateExecContext.Init(System->GetSystemUpdateScript(false), ENiagaraSimTarget::CPUSim);
+
+		//Bind parameter collections.
+		for (UNiagaraParameterCollection* Collection : System->GetSystemSpawnScript(false)->ParameterCollections)
+		{
+			WorldMan->GetParameterCollection(Collection)->GetParameterStore().Bind(&SpawnExecContext.Parameters);
+		}
+		for (UNiagaraParameterCollection* Collection : System->GetSystemUpdateScript(false)->ParameterCollections)
+		{
+			WorldMan->GetParameterCollection(Collection)->GetParameterStore().Bind(&UpdateExecContext.Parameters);
+		}
+
+		SystemEnabledAccessor.Create(&DataSet, FNiagaraVariable(FNiagaraTypeDefinition::GetBoolDef(), TEXT("System.Enabled")));
+		SystemExecutionStateAccessor.Create(&DataSet, FNiagaraVariable(EnumPtr, TEXT("System.ExecutionState")));
+		EmitterEnabledAccessors.Reset();
+		EmitterSpawnInfoAccessors.Reset();
+		EmitterExecutionStateAccessors.Reset();
+		EmitterSpawnInfoAccessors.SetNum(System->GetNumEmitters());
+
+		for (int32 EmitterIdx = 0; EmitterIdx < System->GetNumEmitters(); ++EmitterIdx)
+		{
+			FNiagaraEmitterHandle& EmitterHandle = System->GetEmitterHandle(EmitterIdx);
+			UNiagaraEmitter* Emitter = EmitterHandle.GetInstance();
+			FString EmitterName = Emitter->GetUniqueEmitterName();
+			check(Emitter);
+			EmitterEnabledAccessors.Emplace(DataSet, FNiagaraVariable(FNiagaraTypeDefinition::GetBoolDef(), *(EmitterName + TEXT(".Enabled"))));
+			EmitterExecutionStateAccessors.Emplace(DataSet, FNiagaraVariable(EnumPtr, *(EmitterName + TEXT(".ExecutionState"))));
+			const TArray<FNiagaraEmitterSpawnAttributes>& EmitterSpawnAttrNames = System->GetEmitterSpawnAttributes();
+			for (const FNiagaraEmitterSpawnAttributes& SpawnAttrs : EmitterSpawnAttrNames)
+			{
+				for (FName AttrName : SpawnAttrs.SpawnAttributes)
+				{
+					EmitterSpawnInfoAccessors[EmitterIdx].Emplace(DataSet, FNiagaraVariable(FNiagaraTypeDefinition(FNiagaraSpawnInfo::StaticStruct()), AttrName));
+				}
+			}
+		}
+	}
+
+	if (bCanExecuteSolo)
+	{
+		//Init solo simulation.
+		DataSetSolo.Reset();
+		DataSetSolo.AddVariables(System->GetSystemSpawnScript(true)->Attributes);
+		DataSetSolo.AddVariables(System->GetSystemUpdateScript(true)->Attributes);
+		DataSetSolo.Finalize();
+
+		SpawnExecContextSolo.Init(System->GetSystemSpawnScript(true), ENiagaraSimTarget::CPUSim);
+		UpdateExecContextSolo.Init(System->GetSystemUpdateScript(true), ENiagaraSimTarget::CPUSim);
+
+		SoloSystemEnabledAccessor.Create(&DataSetSolo, FNiagaraVariable(FNiagaraTypeDefinition::GetBoolDef(), TEXT("System.Enabled")));
+		SoloSystemExecutionStateAccessor.Create(&DataSetSolo, FNiagaraVariable(EnumPtr, TEXT("System.ExecutionState")));
+		SoloEmitterEnabledAccessors.Reset();
+		SoloEmitterExecutionStateAccessors.Reset();
+		SoloEmitterSpawnInfoAccessors.Reset();
+		SoloEmitterSpawnInfoAccessors.SetNum(System->GetNumEmitters());
+
+		for (int32 EmitterIdx = 0; EmitterIdx < System->GetNumEmitters(); ++EmitterIdx)
+		{
+			FNiagaraEmitterHandle& EmitterHandle = System->GetEmitterHandle(EmitterIdx);
+			UNiagaraEmitter* Emitter = EmitterHandle.GetInstance();
+			FString EmitterName = Emitter->GetUniqueEmitterName();
+			check(Emitter);
+			SoloEmitterEnabledAccessors.Emplace(DataSetSolo, FNiagaraVariable(FNiagaraTypeDefinition::GetBoolDef(), *(EmitterName + TEXT(".Enabled"))));
+			SoloEmitterExecutionStateAccessors.Emplace(DataSetSolo, FNiagaraVariable(EnumPtr, *(EmitterName + TEXT(".ExecutionState"))));
+			
+			for (FName AttrName : System->GetEmitterSpawnAttributes()[EmitterIdx].SpawnAttributes)
+			{
+				SoloEmitterSpawnInfoAccessors[EmitterIdx].Emplace(DataSetSolo, FNiagaraVariable(FNiagaraTypeDefinition(FNiagaraSpawnInfo::StaticStruct()), AttrName));
+			}
+		}
+	}
+
+	return true;
+}
+
+void FNiagaraSystemSimulation::Destroy()
+{
+	while (SystemInstances.Num())
+	{
+		SystemInstances.Last()->Deactivate(true);
+	}
+	while (PendingSystemInstances.Num())
+	{
+		PendingSystemInstances.Last()->Deactivate(true);
+	}
+	while (SoloSystemInstances.Num())
+	{
+		SoloSystemInstances.Last()->Deactivate(true);
+	}
+
+	FNiagaraWorldManager* WorldMan = FNiagaraWorldManager::Get(World);
+	check(WorldMan);
+	//Unbind parameter collections from the batched exec contexts.
+	for (UNiagaraParameterCollection* Collection : System->GetSystemSpawnScript(false)->ParameterCollections)
+	{
+		WorldMan->GetParameterCollection(Collection)->GetParameterStore().Unbind(&SpawnExecContext.Parameters);
+	}
+	for (UNiagaraParameterCollection* Collection : System->GetSystemUpdateScript(false)->ParameterCollections)
+	{
+		WorldMan->GetParameterCollection(Collection)->GetParameterStore().Unbind(&UpdateExecContext.Parameters);
+	}
+}
+
+void FNiagaraSystemSimulation::TickSoloDataSet()
+{
+	//Flip buffers and shrink the solo data set allocation.
+	if (bCanExecuteSolo)
+	{
+		DataSetSolo.Tick();
+		DataSetSolo.Allocate(SoloSystemInstances.Num());
+		DataSetSolo.SetNumInstances(SoloSystemInstances.Num());
+	}
+}
+
+void FNiagaraSystemSimulation::Tick(float DeltaSeconds)
+{
+	if (!System->IsValid())
+	{
+		// TODO: evaluate whether or not we should have removed this from the world manager instead?
+		return;
+	}
+
+	TickSoloDataSet();
+
+	SCOPE_CYCLE_COUNTER(STAT_NiagaraSystemSim);
+
+	int32 OrigNum = SystemInstances.Num();
+	int32 SpawnNum = PendingSystemInstances.Num();
+	int32 NewNum = OrigNum + SpawnNum;
+	
+	SystemInstances.Reserve(NewNum);
+	for (FNiagaraSystemInstance* Inst : PendingSystemInstances)
+	{
+		Inst->SetPendingSpawn(false);
+		Inst->SystemInstanceIndex = SystemInstances.Add(Inst);
+	}
+	PendingSystemInstances.Reset();
+
+	{
+		SCOPE_CYCLE_COUNTER(STAT_NiagaraSystemSim_PreSimulate);
+
+		auto PreTick = [&](int32 SystemIndex)
+		{
+			FNiagaraSystemInstance* Inst = SystemInstances[SystemIndex];
+			Inst->PreSimulateTick(DeltaSeconds);
+
+			check(!Inst->GetParameters().bLayoutDirty);
+			if (Inst->GetParameters().bParametersDirty && bCanExecute)
+			{
+				SpawnParameterToDataSetBinding.ParameterStoreToDataSet(Inst->GetParameters(), SpawnParameterDataSet, SystemIndex);
+				UpdateParameterToDataSetBinding.ParameterStoreToDataSet(Inst->GetParameters(), UpdateParameterDataSet, SystemIndex);
+			}
+			//TODO: Find good way to check that we're not using any instance parameter data interfaces in the system scripts here.
+			//In that case we need to solo and will never get here.
+		};
+
+		SpawnParameterDataSet.Allocate(NewNum);
+		UpdateParameterDataSet.Allocate(NewNum);
+
+		if (GbParallelSystemPreTick)
+		{
+			ParallelFor(SystemInstances.Num(), PreTick);
+		}
+		else
+		{
+			//Tick instances and pull out any per instance parameters we need.
+			for (int32 i = 0; i < SystemInstances.Num(); ++i)
+			{
+				PreTick(i);
+			}
+		}
+
+		SpawnParameterDataSet.Tick();
+		UpdateParameterDataSet.Tick();
+	}
+
+	if (bCanExecute && NewNum > 0)
+	{
+		InitBindings(SystemInstances[0]);
+
+		DataSet.Tick();
+		DataSet.Allocate(NewNum);
+		
+		//Setup the few real constants like delta time.
+		float InvDt = 1.0f / DeltaSeconds;
+		SpawnExecContext.Parameters.SetParameterValue(DeltaSeconds, SYS_PARAM_ENGINE_DELTA_TIME);
+		SpawnExecContext.Parameters.SetParameterValue(InvDt, SYS_PARAM_ENGINE_INV_DELTA_TIME);
+
+		UpdateExecContext.Parameters.SetParameterValue(DeltaSeconds, SYS_PARAM_ENGINE_DELTA_TIME);
+		UpdateExecContext.Parameters.SetParameterValue(InvDt, SYS_PARAM_ENGINE_INV_DELTA_TIME);
+
+		TArray<FNiagaraDataSetExecutionInfo> DataSetExecInfos;
+		DataSetExecInfos.Emplace(&DataSet, 0, false, true);
+
+		{
+			SCOPE_CYCLE_COUNTER(STAT_NiagaraSystemSim_Update);
+			DataSet.SetNumInstances(OrigNum);
+			//Run update.
+			UpdateExecContext.Tick(nullptr);//We can't require a specfic instance here as these are for all instances.
+			DataSetExecInfos.SetNum(1, false);
+			DataSetExecInfos[0].StartInstance = 0;
+			DataSetExecInfos.Emplace(&UpdateParameterDataSet, 0, false, false);
+			UpdateExecContext.Execute(OrigNum, DataSetExecInfos);
+
+			if (GbDumpSystemData)
+			{
+				UE_LOG(LogNiagara, Log, TEXT("=== Updated %d Systems ==="), OrigNum);
+				DataSet.Dump(true, 0, OrigNum);
+				UpdateParameterDataSet.Dump(true, 0, OrigNum);
+			}
+		}
+
+		{
+			SCOPE_CYCLE_COUNTER(STAT_NiagaraSystemSim_Spawn);
+			DataSet.SetNumInstances(NewNum);
+			//Run Spawn
+			SpawnExecContext.Tick(nullptr);//We can't require a specific instance here as these are for all instances.
+			DataSetExecInfos.SetNum(1, false);
+			DataSetExecInfos[0].StartInstance = OrigNum;
+			DataSetExecInfos.Emplace(&SpawnParameterDataSet, OrigNum, false, false);
+			SpawnExecContext.Execute(SpawnNum, DataSetExecInfos);
+
+			if (GbDumpSystemData)
+			{
+				UE_LOG(LogNiagara, Log, TEXT("=== Spwaned %d Systems ==="), SpawnNum);
+				DataSet.Dump(true, OrigNum, SpawnNum);
+				SpawnParameterDataSet.Dump(true, OrigNum, SpawnNum);
+			}
+		}
+
+		{
+			SCOPE_CYCLE_COUNTER(STAT_NiagaraSystemSim_TransferParameters);
+			SystemEnabledAccessor.InitForAccess(true);
+			SystemExecutionStateAccessor.InitForAccess(true);
+			for (int32 EmitterIdx = 0; EmitterIdx < System->GetNumEmitters(); ++EmitterIdx)
+			{
+				EmitterEnabledAccessors[EmitterIdx].InitForAccess(true);
+				EmitterExecutionStateAccessors[EmitterIdx].InitForAccess(true);
+				for (int32 SpawnInfoIdx = 0; SpawnInfoIdx < EmitterSpawnInfoAccessors[EmitterIdx].Num(); ++SpawnInfoIdx)
+				{
+					EmitterSpawnInfoAccessors[EmitterIdx][SpawnInfoIdx].InitForAccess(true);
+				}				
+			}
+
+			int32 SystemIndex = 0;
+			while (SystemIndex < SystemInstances.Num())
+			{
+				ENiagaraExecutionState ExecutionState = (ENiagaraExecutionState)SystemExecutionStateAccessor.GetSafe(SystemIndex, (int32)ENiagaraExecutionState::Active);
+				FNiagaraSystemInstance* SystemInst = SystemInstances[SystemIndex];
+				SystemInst->SetExecutionState(ExecutionState);
+				bool bSystemEnabled = SystemEnabledAccessor.GetSafe(SystemIndex, true);
+				//Kill any instances that are flagged as disabled.
+				if (!bSystemEnabled || ExecutionState == ENiagaraExecutionState::Dead)
+				{
+					RemoveInstance(SystemInst);
+					SystemInst->Disable();
+				}
+				else
+				{
+					//Now pull data out of the simulation and drive the emitters with it.
+					TArray<TSharedRef<FNiagaraEmitterInstance>>& Emitters = SystemInst->GetEmitters();
+					for (int32 EmitterIdx = 0; EmitterIdx < Emitters.Num(); ++EmitterIdx)
+					{
+						FNiagaraEmitterInstance& EmitterInst = Emitters[EmitterIdx].Get();
+
+						FNiagaraBool bEmitterEnabled = EmitterEnabledAccessors[EmitterIdx].GetSafe(SystemIndex, true);
+						EmitterInst.SetEnabled(bEmitterEnabled);//TODO: Alter this to instruct emitter to dump existing particles optionally.
+
+						TArray<FNiagaraSpawnInfo>& EmitterInstSpawnInfos = EmitterInst.GetSpawnInfo();
+						for (int32 SpawnInfoIdx=0; SpawnInfoIdx < EmitterSpawnInfoAccessors[EmitterIdx].Num(); ++SpawnInfoIdx)
+						{
+							if (SpawnInfoIdx < EmitterInstSpawnInfos.Num())
+							{
+								EmitterInstSpawnInfos[SpawnInfoIdx] = EmitterSpawnInfoAccessors[EmitterIdx][SpawnInfoIdx].Get(SystemIndex);
+							}
+							else
+							{
+								ensure(SpawnInfoIdx < EmitterInstSpawnInfos.Num());
+							}
+						}
+
+						ENiagaraExecutionState State = (ENiagaraExecutionState)EmitterExecutionStateAccessors[EmitterIdx].GetSafe(SystemIndex, (int32)ENiagaraExecutionState::Active);
+						EmitterInst.SetExecutionState(State);
+
+						//TODO: Any other fixed function stuff like this?
+
+						FNiagaraScriptExecutionContext& SpawnContext = EmitterInst.GetSpawnExecutionContext();
+						DataSetToEmitterSpawnParameters[EmitterIdx].DataSetToParameterStore(SpawnContext.Parameters, DataSet, SystemIndex);
+
+						FNiagaraScriptExecutionContext& UpdateContext = EmitterInst.GetUpdateExecutionContext();
+						DataSetToEmitterUpdateParameters[EmitterIdx].DataSetToParameterStore(UpdateContext.Parameters, DataSet, SystemIndex);
+
+						TArray<FNiagaraScriptExecutionContext>& EventContexts = EmitterInst.GetEventExecutionContexts();
+						for (int32 EventIdx = 0; EventIdx < EventContexts.Num(); ++EventIdx)
+						{
+							FNiagaraScriptExecutionContext& EventContext = EventContexts[EventIdx];
+							DataSetToEmitterEventParameters[EmitterIdx][EventIdx].DataSetToParameterStore(EventContext.Parameters, DataSet, SystemIndex);
+						}
+					}
+
+					//System is still enabled. 
+					++SystemIndex;
+				}
+			}
+		}
+	}
+
+	{
+		SCOPE_CYCLE_COUNTER(STAT_NiagaraSystemSim_PostSimulate);
+
+		if (GbParallelSystemPostTick)
+		{
+			ParallelFor(SystemInstances.Num(),
+				[&](int32 SystemIndex)
+			{
+				FNiagaraSystemInstance* SystemInst = SystemInstances[SystemIndex];
+				SystemInst->PostSimulateTick(DeltaSeconds);
+			});
+		}
+		else
+		{
+			//Now actually tick emitters.
+			for (int32 SystemIndex = 0; SystemIndex < SystemInstances.Num(); ++SystemIndex)
+			{
+				FNiagaraSystemInstance* SystemInst = SystemInstances[SystemIndex];
+				SystemInst->PostSimulateTick(DeltaSeconds);
+			}
+		}
+	}
+
+	{
+		SCOPE_CYCLE_COUNTER(STAT_NiagaraSystemSim_MarkComponentDirty);
+		//This is not a small amount of the update time. 
+		//Annoyingly these can't be done in parallel it seems.
+		//TODO: Find some way to parallelize this. Especially UpdateComponentToWorld();
+		int32 SystemIndex = 0;
+		while (SystemIndex < SystemInstances.Num())
+		{
+			FNiagaraSystemInstance* SystemInst = SystemInstances[SystemIndex];
+			++SystemIndex;
+			if (SystemIndex < SystemInstances.Num())
+			{
+				FPlatformMisc::Prefetch(SystemInstances[SystemIndex]->GetComponent());
+			}
+			SystemInst->GetComponent()->UpdateComponentToWorld();
+			SystemInst->GetComponent()->MarkRenderDynamicDataDirty();
+		}
+	}
+}
+
+void FNiagaraSystemSimulation::RemoveInstance(FNiagaraSystemInstance* Instance)
+{
+	if (Instance->SystemInstanceIndex == INDEX_NONE)
+	{
+		return;
+	}
+
+	if (Instance->IsSolo())
+	{
+		if (SoloSystemInstances.IsValidIndex(Instance->SystemInstanceIndex))
+		{
+			if (GbDumpSystemData)
+			{
+				UE_LOG(LogNiagara, Log, TEXT("=== Removing System Solo %d ==="), Instance->SystemInstanceIndex);
+				DataSetSolo.Dump(true, Instance->SystemInstanceIndex, 1);
+			}
+
+			int32 NumInstances = DataSetSolo.GetNumInstances();
+			check(SoloSystemInstances.Num() == NumInstances);
+
+			int32 SystemIndex = Instance->SystemInstanceIndex;
+			check(Instance == SoloSystemInstances[SystemIndex]);
+			check(SoloSystemInstances.IsValidIndex(SystemIndex));
+			DataSetSolo.KillInstance(SystemIndex);
+			SoloSystemInstances.RemoveAtSwap(SystemIndex);
+			Instance->SystemInstanceIndex = INDEX_NONE;
+			DataSetSolo.SetNumInstances(SoloSystemInstances.Num());
+
+			if (SoloSystemInstances.IsValidIndex(SystemIndex))
+			{
+				SoloSystemInstances[SystemIndex]->SystemInstanceIndex = SystemIndex;
+			}
+		}
+	}
+	else if (Instance->IsPendingSpawn())
+	{
+		int32 SystemIndex = Instance->SystemInstanceIndex;
+		check(Instance == PendingSystemInstances[SystemIndex]);
+		PendingSystemInstances.RemoveAtSwap(SystemIndex);
+		Instance->SystemInstanceIndex = INDEX_NONE;
+		Instance->SetPendingSpawn(false);
+		if (PendingSystemInstances.IsValidIndex(SystemIndex))
+		{
+			PendingSystemInstances[SystemIndex]->SystemInstanceIndex = SystemIndex;
+		}
+	}
+	else if (SystemInstances.IsValidIndex(Instance->SystemInstanceIndex))
+	{
+		if (GbDumpSystemData)
+		{
+			UE_LOG(LogNiagara, Log, TEXT("=== Removing System %d ==="), Instance->SystemInstanceIndex);
+			DataSet.Dump(true, Instance->SystemInstanceIndex, 1);
+		}
+
+		int32 NumInstances = DataSet.GetNumInstances();
+		check(SystemInstances.Num() == NumInstances);
+
+		int32 SystemIndex = Instance->SystemInstanceIndex;
+		check(Instance == SystemInstances[SystemIndex]);
+		check(SystemInstances.IsValidIndex(SystemIndex));
+		DataSet.KillInstance(SystemIndex);
+		SystemInstances.RemoveAtSwap(SystemIndex);
+		Instance->SystemInstanceIndex = INDEX_NONE;
+
+		if (SystemInstances.IsValidIndex(SystemIndex))
+		{
+			SystemInstances[SystemIndex]->SystemInstanceIndex = SystemIndex;
+		}
+	}
+}
+
+void FNiagaraSystemSimulation::AddInstance(FNiagaraSystemInstance* Instance)
+{
+	Instance->SetPendingSpawn(true);
+	Instance->SystemInstanceIndex = PendingSystemInstances.Add(Instance);
+}
+
+void FNiagaraSystemSimulation::ResetSolo(FNiagaraSystemInstance* Instance)
+{
+	Instance->SetPendingSpawn(true);
+	Instance->SystemInstanceIndex = INDEX_NONE;
+	Instance->SetExecutionState(ENiagaraExecutionState::Active);
+}
+
+void FNiagaraSystemSimulation::TickSolo(FNiagaraSystemInstance* SystemInst)
+{
+	SCOPE_CYCLE_COUNTER(STAT_NiagaraSystemSimSolo);
+	check(System->IsValid());
+
+	if (!System->GetSystemSpawnScript(true)->IsValid() || !System->GetSystemUpdateScript(true)->IsValid())
+	{
+		return;
+	}
+
+	FNiagaraWorldManager* WorldMan = FNiagaraWorldManager::Get(World);
+	check(WorldMan);
+
+	int32 SystemIndex = SystemInst->SystemInstanceIndex;
+
+	TArray<FNiagaraDataSetExecutionInfo> DataSetExecInfos;
+	DataSetExecInfos.Emplace(&DataSetSolo, SystemIndex, false, true);
+
+	if (SystemInst->SystemInstanceIndex == INDEX_NONE || SystemInst->IsPendingSpawn())
+	{
+		SystemInst->SetPendingSpawn(false);
+
+		//Spawn/respawn this system this frame rather than updating.
+		if (SystemIndex == INDEX_NONE)
+		{
+			SystemIndex = SoloSystemInstances.Num();
+			SystemInst->SystemInstanceIndex = SystemIndex;
+			verify(SoloSystemInstances.Add(SystemInst) == SystemIndex);
+			DataSetSolo.Allocate(SoloSystemInstances.Num(), ENiagaraSimTarget::CPUSim, true);
+			DataSetSolo.SetNumInstances(SoloSystemInstances.Num());
+			DataSetExecInfos[0].StartInstance = SystemIndex;
+		}
+
+		check(SoloSystemInstances.IsValidIndex(SystemIndex));
+		//////////////////////////////////////////////////////////////////////////
+		//TODO: This can be optimized as all the offests are the same. Can skip all the searching involved with a full bind().
+		//Unbind parameter collections from the batched exec contexts.
+		for (UNiagaraParameterCollection* Collection : System->GetSystemSpawnScript(true)->ParameterCollections)
+		{
+			WorldMan->GetParameterCollection(Collection)->GetParameterStore().Bind(&SpawnExecContextSolo.Parameters);
+		}
+		SystemInst->GetParameters().Bind(&SpawnExecContextSolo.Parameters);
+		SpawnExecContextSolo.Tick(SystemInst);
+		//////////////////////////////////////////////////////////////////////////
+		
+		SpawnExecContextSolo.Execute(1, DataSetExecInfos);
+
+		//////////////////////////////////////////////////////////////////////////
+		//TODO: This can be optimized as all the offests are the same. Can skip all the searching involved with a full bind().
+		SystemInst->GetParameters().Unbind(&SpawnExecContextSolo.Parameters);
+		for (UNiagaraParameterCollection* Collection : System->GetSystemSpawnScript(true)->ParameterCollections)
+		{
+			WorldMan->GetParameterCollection(Collection)->GetParameterStore().Unbind(&SpawnExecContextSolo.Parameters);
+		}
+		//////////////////////////////////////////////////////////////////////////
+		
+		DataSetSolo.SetNumInstances(SoloSystemInstances.Num());
+
+		if (GbDumpSystemData)
+		{
+			UE_LOG(LogNiagara, Log, TEXT("=== Reset/spawn System Solo %d ==="), SystemIndex);
+			DataSetSolo.Dump(true, SystemIndex, 1);
+		}
+	}
+	else
+	{
+		check(SoloSystemInstances.IsValidIndex(SystemIndex));
+		//////////////////////////////////////////////////////////////////////////
+		//TODO: This can be optimized as all the offests are the same. Can skip all the searching involved with a full bind().
+		for (UNiagaraParameterCollection* Collection : System->GetSystemUpdateScript(true)->ParameterCollections)
+		{
+			WorldMan->GetParameterCollection(Collection)->GetParameterStore().Bind(&SpawnExecContextSolo.Parameters);
+		}
+		SystemInst->GetParameters().Bind(&UpdateExecContextSolo.Parameters);
+		UpdateExecContextSolo.Tick(SystemInst);
+		//////////////////////////////////////////////////////////////////////////
+
+		UpdateExecContextSolo.Execute(1, DataSetExecInfos);
+
+		//////////////////////////////////////////////////////////////////////////
+		//TODO: This can be optimized as all the offests are the same. Can skip all the searching involved with a full bind().
+		SystemInst->GetParameters().Unbind(&UpdateExecContextSolo.Parameters);
+		for (UNiagaraParameterCollection* Collection : System->GetSystemUpdateScript(true)->ParameterCollections)
+		{
+			WorldMan->GetParameterCollection(Collection)->GetParameterStore().Unbind(&UpdateExecContextSolo.Parameters);
+		}
+		//////////////////////////////////////////////////////////////////////////
+
+		DataSetSolo.SetNumInstances(SoloSystemInstances.Num());
+
+		if (GbDumpSystemData)
+		{
+			UE_LOG(LogNiagara, Log, TEXT("=== Updated System Solo %d ==="), SystemIndex);
+			DataSetSolo.Dump(true, SystemIndex, 1);
+		}
+	}
+
+	InitBindings(SystemInst);
+
+	//Kill instance if it's marked for delete.
+	SoloSystemEnabledAccessor.InitForAccess(true);
+	SoloSystemExecutionStateAccessor.InitForAccess(true);
+	for (int32 EmitterIdx = 0; EmitterIdx < System->GetNumEmitters(); ++EmitterIdx)
+	{
+		SoloEmitterEnabledAccessors[EmitterIdx].InitForAccess(true);
+		SoloEmitterExecutionStateAccessors[EmitterIdx].InitForAccess(true);
+		for (int32 SpawnInfoIdx = 0; SpawnInfoIdx < SoloEmitterSpawnInfoAccessors[EmitterIdx].Num(); ++SpawnInfoIdx)
+		{
+			SoloEmitterSpawnInfoAccessors[EmitterIdx][SpawnInfoIdx].InitForAccess(true);
+		}
+	}		
+
+	ensure(SystemIndex != INDEX_NONE);
+	bool bSystemEnabled = SoloSystemEnabledAccessor.GetSafe(SystemIndex, true);
+	ENiagaraExecutionState SoloExecutionState = (ENiagaraExecutionState)SoloSystemExecutionStateAccessor.GetSafe(SystemIndex, (int32)ENiagaraExecutionState::Active);
+	SystemInst->SetExecutionState(SoloExecutionState);
+
+	if (!bSystemEnabled || SoloExecutionState == ENiagaraExecutionState::Dead)
+	{
+		RemoveInstance(SystemInst);
+		SystemInst->Disable();
+	}
+	else
+	{
+		//Now pull data out of the simulation and drive the emitters with it.
+		TArray<TSharedRef<FNiagaraEmitterInstance>>& Emitters = SystemInst->GetEmitters();
+		for (int32 EmitterIdx = 0; EmitterIdx < Emitters.Num(); ++EmitterIdx)
+		{
+			FNiagaraEmitterInstance& EmitterInst = Emitters[EmitterIdx].Get();
+
+			FNiagaraBool bEmitterEnabled = SoloEmitterEnabledAccessors[EmitterIdx].GetSafe(SystemIndex, true);
+			EmitterInst.SetEnabled(bEmitterEnabled);//TODO: Alter this to instruct emitter to dump existing particles optionally.
+
+			ENiagaraExecutionState State = (ENiagaraExecutionState)SoloEmitterExecutionStateAccessors[EmitterIdx].GetSafe(SystemIndex, (int32)ENiagaraExecutionState::Active);
+			EmitterInst.SetExecutionState(State);
+
+			for (int32 SpawnInfoIdx = 0; SpawnInfoIdx < SoloEmitterSpawnInfoAccessors[EmitterIdx].Num(); ++SpawnInfoIdx)
+			{
+				EmitterInst.GetSpawnInfo()[SpawnInfoIdx] = SoloEmitterSpawnInfoAccessors[EmitterIdx][SpawnInfoIdx].Get(SystemIndex);
+			}
+			
+			//TODO: Any other fixed function stuff like this?
+
+			FNiagaraScriptExecutionContext& SpawnContext = EmitterInst.GetSpawnExecutionContext();
+			DataSetToEmitterSpawnParameters[EmitterIdx].DataSetToParameterStore(SpawnContext.Parameters, DataSetSolo, SystemIndex);
+
+			FNiagaraScriptExecutionContext& UpdateContext = EmitterInst.GetUpdateExecutionContext();
+			DataSetToEmitterUpdateParameters[EmitterIdx].DataSetToParameterStore(UpdateContext.Parameters, DataSetSolo, SystemIndex);
+
+			TArray<FNiagaraScriptExecutionContext>& EventContexts = EmitterInst.GetEventExecutionContexts();
+			for (int32 EventIdx = 0; EventIdx < EventContexts.Num(); ++EventIdx)
+			{
+				FNiagaraScriptExecutionContext& EventContext = EventContexts[EventIdx];
+				DataSetToEmitterEventParameters[EmitterIdx][EventIdx].DataSetToParameterStore(EventContext.Parameters, DataSetSolo, SystemIndex);
+			}
+		}
+	}
+}
+
+void FNiagaraSystemSimulation::InitBindings(FNiagaraSystemInstance* SystemInst)
+{
+	//Have to init here as we need an actual parameter store to pull the layout info from.
+	//TODO: Pull the layout stuff out of each data set and store. So much duplicated data.
+	//This assumes that all layouts for all emitters is the same. Which it should be.
+	//Ideally we can store all this layout info in the systm/emitter assets so we can just generate this in Init()
+	if (DataSetToEmitterSpawnParameters.Num() == 0)
+	{
+		SpawnParameterToDataSetBinding.Init(SpawnParameterDataSet, SystemInst->GetInstanceParameters());
+		UpdateParameterToDataSetBinding.Init(UpdateParameterDataSet, SystemInst->GetInstanceParameters());
+
+		TArray<TSharedRef<FNiagaraEmitterInstance>>& Emitters = SystemInst->GetEmitters();
+		check(DataSetToEmitterUpdateParameters.Num() == 0);
+		check(DataSetToEmitterEventParameters.Num() == 0);
+		DataSetToEmitterSpawnParameters.SetNum(Emitters.Num());
+		DataSetToEmitterUpdateParameters.SetNum(Emitters.Num());
+		DataSetToEmitterEventParameters.SetNum(Emitters.Num());
+		for (int32 EmitterIdx = 0; EmitterIdx < Emitters.Num(); ++EmitterIdx)
+		{
+			FNiagaraEmitterInstance& EmitterInst = Emitters[EmitterIdx].Get();
+			FNiagaraScriptExecutionContext& SpawnContext = EmitterInst.GetSpawnExecutionContext();
+			DataSetToEmitterSpawnParameters[EmitterIdx].Init(DataSet, SpawnContext.Parameters);
+
+			FNiagaraScriptExecutionContext& UpdateContext = EmitterInst.GetUpdateExecutionContext();
+			DataSetToEmitterUpdateParameters[EmitterIdx].Init(DataSet, UpdateContext.Parameters);
+
+			TArray<FNiagaraScriptExecutionContext>& EventContexts = EmitterInst.GetEventExecutionContexts();
+			DataSetToEmitterEventParameters[EmitterIdx].SetNum(EventContexts.Num());
+			for (int32 EventIdx = 0; EventIdx < EventContexts.Num(); ++EventIdx)
+			{
+				FNiagaraScriptExecutionContext& EventContext = EventContexts[EventIdx];
+				DataSetToEmitterEventParameters[EmitterIdx][EventIdx].Init(DataSet, EventContext.Parameters);
+			}
+		}
+	}
+}
+
+
+
diff --git a/Engine/Plugins/FX/Niagara/Source/Niagara/Private/NiagaraWorldManager.cpp b/Engine/Plugins/FX/Niagara/Source/Niagara/Private/NiagaraWorldManager.cpp
new file mode 100644
--- /dev/null
+++ b/Engine/Plugins/FX/Niagara/Source/Niagara/Private/NiagaraWorldManager.cpp
@@ -0,0 +1,128 @@
+// Copyright 1998-2017 Epic Games, Inc. All Rights Reserved.
+
+#include "NiagaraWorldManager.h"
+#include "NiagaraModule.h"
+#include "Modules/ModuleManager.h"
+#include "NiagaraTypes.h"
+#include "NiagaraEvents.h"
+#include "NiagaraSettings.h"
+#include "NiagaraParameterCollection.h"
+#include "NiagaraSystemInstance.h"
+#include "Scalability.h"
+#include "ConfigCacheIni.h"
+
+
+FNiagaraWorldManager* FNiagaraWorldManager::Get(UWorld* World)
+{
+	INiagaraModule& NiagaraModule = FModuleManager::LoadModuleChecked<INiagaraModule>("Niagara");
+	return NiagaraModule.GetWorldManager(World);
+}
+
+void FNiagaraWorldManager::AddReferencedObjects(FReferenceCollector& Collector)
+{
+	Collector.AddReferencedObjects(ParameterCollections);
+}
+
+UNiagaraParameterCollectionInstance* FNiagaraWorldManager::GetParameterCollection(UNiagaraParameterCollection* Collection)
+{
+	UNiagaraParameterCollectionInstance** OverrideInst = ParameterCollections.Find(Collection);
+	if (!OverrideInst)
+	{
+		OverrideInst = &ParameterCollections.Add(Collection);
+		*OverrideInst = CastChecked<UNiagaraParameterCollectionInstance>(StaticDuplicateObject(Collection->GetDefaultInstance(), World));
+	}
+
+	check(OverrideInst && *OverrideInst);
+	return *OverrideInst;
+}
+
+void FNiagaraWorldManager::SetParameterCollection(UNiagaraParameterCollectionInstance* NewInstance)
+{
+	check(NewInstance);
+	if (NewInstance)
+	{
+		UNiagaraParameterCollection* Collection = NewInstance->GetParent();
+		UNiagaraParameterCollectionInstance** OverrideInst = ParameterCollections.Find(Collection);
+		if (!OverrideInst)
+		{
+			OverrideInst = &ParameterCollections.Add(Collection);
+		}
+		else
+		{
+			if (*OverrideInst && NewInstance)
+			{
+				//Need to transfer existing bindings from old instance to new one.
+				FNiagaraParameterStore& ExistingStore = (*OverrideInst)->GetParameterStore();
+				FNiagaraParameterStore& NewStore = NewInstance->GetParameterStore();
+
+				ExistingStore.TransferBindings(NewStore);
+			}
+		}
+
+		*OverrideInst = NewInstance;
+	}
+}
+
+FNiagaraSystemSimulation* FNiagaraWorldManager::GetSystemSimulation(UNiagaraSystem* System)
+{
+	FNiagaraSystemSimulation* Sim = SystemSimulations.Find(System);
+	if (Sim)
+	{
+		return Sim;
+	}
+	
+	Sim = &SystemSimulations.Add(System);
+	Sim->Init(System, World);
+	return Sim;
+}
+
+void FNiagaraWorldManager::DestroySystemSimulation(UNiagaraSystem* System)
+{
+	FNiagaraSystemSimulation* Sim = SystemSimulations.Find(System);
+	if (Sim)
+	{
+		Sim->Destroy();
+		SystemSimulations.Remove(System);
+	}	
+}
+
+void FNiagaraWorldManager::Tick(float DeltaSeconds)
+{
+/*
+Some experimental tinkering with links to effects quality.
+Going off this idea tbh
+
+	int32 CurrentEffectsQuality = Scalability::GetEffectsQualityDirect(true);
+	if (CachedEffectsQuality != CurrentEffectsQuality)
+	{
+		CachedEffectsQuality = CurrentEffectsQuality;
+		FString SectionName = FString::Printf(TEXT("%s@%d"), TEXT("EffectsQuality"), CurrentEffectsQuality);
+		if (FConfigFile* File = GConfig->FindConfigFileWithBaseName(TEXT("Niagara")))
+		{
+			FString ScalabilityCollectionString;			
+			if (File->GetString(*SectionName, TEXT("ScalabilityCollection"), ScalabilityCollectionString))
+			{
+				//NewLandscape_Material = LoadObject<UMaterialInterface>(NULL, *NewLandsfgcapeMaterialName, NULL, LOAD_NoWarn);
+				UNiagaraParameterCollectionInstance* ScalabilityCollection = LoadObject<UNiagaraParameterCollectionInstance>(NULL, *ScalabilityCollectionString, NULL, LOAD_NoWarn);
+				if (ScalabilityCollection)
+				{
+					SetParameterCollection(ScalabilityCollection);
+				}
+			}
+		}
+	}
+*/
+
+	//Tick our collections to push any changes to bound stores.
+	for (TPair<UNiagaraParameterCollection*, UNiagaraParameterCollectionInstance*> CollectionInstPair : ParameterCollections)
+	{
+		check(CollectionInstPair.Value);
+		CollectionInstPair.Value->Tick();
+	}
+
+	//Now tick all system instances. 
+	for (TPair<UNiagaraSystem*, FNiagaraSystemSimulation>& SystemSim : SystemSimulations)
+	{
+		SystemSim.Value.Tick(DeltaSeconds);
+	}
+}
\ No newline at end of file
diff --git a/Engine/Plugins/FX/Niagara/Source/Niagara/Public/NiagaraActor.h b/Engine/Plugins/FX/Niagara/Source/Niagara/Public/NiagaraActor.h
new file mode 100644
--- /dev/null
+++ b/Engine/Plugins/FX/Niagara/Source/Niagara/Public/NiagaraActor.h
@@ -0,0 +1,47 @@
+// Copyright 1998-2017 Epic Games, Inc. All Rights Reserved.
+
+#pragma once
+
+#include "CoreMinimal.h"
+#include "UObject/ObjectMacros.h"
+#include "GameFramework/Actor.h"
+#include "NiagaraActor.generated.h"
+
+UCLASS(MinimalAPI)
+class ANiagaraActor : public AActor
+{
+	GENERATED_UCLASS_BODY()
+
+private:
+	/** Pointer to System component */
+	UPROPERTY(VisibleAnywhere, Category=NiagaraActor)
+	class UNiagaraComponent* NiagaraComponent;
+
+#if WITH_EDITORONLY_DATA
+	// Reference to sprite visualization component
+	UPROPERTY()
+	class UBillboardComponent* SpriteComponent;
+
+	// Reference to arrow visualization component
+	UPROPERTY()
+	class UArrowComponent* ArrowComponent;
+
+#endif
+
+public:
+	/** Returns NiagaraComponent subobject **/
+	NIAGARA_API class UNiagaraComponent* GetNiagaraComponent() const { return NiagaraComponent; }
+#if WITH_EDITORONLY_DATA
+	/** Returns SpriteComponent subobject **/
+	class UBillboardComponent* GetSpriteComponent() const { return SpriteComponent; }
+	/** Returns ArrowComponent subobject **/
+	class UArrowComponent* GetArrowComponent() const { return ArrowComponent; }
+#endif
+
+#if WITH_EDITOR
+	// AActor interface
+	virtual bool GetReferencedContentObjects(TArray<UObject*>& Objects) const override;
+	// End of AActor interface
+#endif // WITH_EDITOR
+
+};
diff --git a/Engine/Plugins/FX/Niagara/Source/Niagara/Public/NiagaraCommon.h b/Engine/Plugins/FX/Niagara/Source/Niagara/Public/NiagaraCommon.h
new file mode 100644
--- /dev/null
+++ b/Engine/Plugins/FX/Niagara/Source/Niagara/Public/NiagaraCommon.h
@@ -0,0 +1,383 @@
+// Copyright 1998-2017 Epic Games, Inc. All Rights Reserved.
+
+#pragma once
+
+#include "CoreMinimal.h"
+#include "UObject/ObjectMacros.h"
+#include "NiagaraTypes.h"
+#include "NiagaraCommon.generated.h"
+
+class UNiagaraSystem;
+class UNiagaraScript;
+class UNiagaraDataInterface;
+class UNiagaraEmitter;
+class FNiagaraSystemInstance;
+class UNiagaraParameterCollection;
+
+const uint32 NIAGARA_COMPUTE_THREADGROUP_SIZE = 16;
+const uint32 NIAGARA_MAX_COMPUTE_THREADGROUPS = 2048;
+
+enum ENiagaraBaseTypes
+{
+	NBT_Float,
+	NBT_Int32,
+	NBT_Bool,
+	NBT_Max,
+};
+
+UENUM()
+enum class ENiagaraSimTarget : uint8
+{
+	CPUSim,
+	GPUComputeSim,
+	DynamicLoadBalancedSim
+};
+
+
+
+UENUM()
+enum class ENiagaraDataSetType : uint8
+{
+	ParticleData,
+	Shared,
+	Event,
+};
+
+
+UENUM()
+enum class ENiagaraInputNodeUsage : uint8
+{
+	Undefined = 0,
+	Parameter,
+	Attribute,
+	SystemConstant
+};
+
+/**
+* Enumerates states a Niagara script can be in.
+*/
+UENUM()
+enum class ENiagaraScriptCompileStatus : uint8
+{
+	/** Niagara script is in an unknown state. */
+	NCS_Unknown,
+	/** Niagara script has been modified but not recompiled. */
+	NCS_Dirty,
+	/** Niagara script tried but failed to be compiled. */
+	NCS_Error,
+	/** Niagara script has been compiled since it was last modified. */
+	NCS_UpToDate,
+	/** Niagara script is in the process of being created for the first time. */
+	NCS_BeingCreated,
+	/** Niagara script has been compiled since it was last modified. There are warnings. */
+	NCS_UpToDateWithWarnings,
+	/** Niagara script has been compiled for compute since it was last modified. There are warnings. */
+	NCS_ComputeUpToDateWithWarnings,
+	NCS_MAX,
+};
+
+USTRUCT()
+struct FNiagaraDataSetID
+{
+	GENERATED_USTRUCT_BODY()
+
+	FNiagaraDataSetID()
+	: Name(NAME_None)
+	, Type(ENiagaraDataSetType::Event)
+	{}
+
+	FNiagaraDataSetID(FName InName, ENiagaraDataSetType InType)
+		: Name(InName)
+		, Type(InType)
+	{}
+
+	UPROPERTY(EditAnywhere, Category = "Data Set")
+	FName Name;
+
+	UPROPERTY()
+	ENiagaraDataSetType Type;
+
+	FORCEINLINE bool operator==(const FNiagaraDataSetID& Other)const
+	{
+		return Name == Other.Name && Type == Other.Type;
+	}
+
+	FORCEINLINE bool operator!=(const FNiagaraDataSetID& Other)const
+	{
+		return !(*this == Other);
+	}
+};
+
+
+FORCEINLINE FArchive& operator<<(FArchive& Ar, FNiagaraDataSetID& VarInfo)
+{
+	Ar << VarInfo.Name << VarInfo.Type;
+	return Ar;
+}
+
+FORCEINLINE uint32 GetTypeHash(const FNiagaraDataSetID& Var)
+{
+	return HashCombine(GetTypeHash(Var.Name), (uint32)Var.Type);
+}
+
+USTRUCT()
+struct FNiagaraDataSetProperties
+{
+	GENERATED_BODY()
+	
+	UPROPERTY(VisibleAnywhere, Category = "Data Set")
+	FNiagaraDataSetID ID;
+
+	UPROPERTY()
+	TArray<FNiagaraVariable> Variables;
+};
+
+/** Information about an input or output of a Niagara operation node. */
+class FNiagaraOpInOutInfo
+{
+public:
+	FName Name;
+	FNiagaraTypeDefinition DataType;
+	FText FriendlyName;
+	FText Description;
+	FString Default;
+	FString HlslSnippet;
+
+	FNiagaraOpInOutInfo(FName InName, FNiagaraTypeDefinition InType, FText InFriendlyName, FText InDescription, FString InDefault, FString InHlslSnippet = TEXT(""))
+		: Name(InName)
+		, DataType(InType)
+		, FriendlyName(InFriendlyName)
+		, Description(InDescription)
+		, Default(InDefault)
+		, HlslSnippet(InHlslSnippet)
+	{
+
+	}
+};
+
+
+/** Struct containing usage information about a script. Things such as whether it reads attribute data, reads or writes events data etc.*/
+USTRUCT()
+struct FNiagaraScriptDataUsageInfo
+{
+	GENERATED_BODY()
+
+		FNiagaraScriptDataUsageInfo()
+		: bReadsAttriubteData(false)
+	{}
+
+	/** If true, this script reads attribute data. */
+	UPROPERTY()
+		bool bReadsAttriubteData;
+};
+
+
+USTRUCT()
+struct NIAGARA_API FNiagaraFunctionSignature
+{
+	GENERATED_BODY()
+
+	/** Name of the function. */
+	UPROPERTY()
+	FName Name;
+	/** Input parameters to this function. */
+	UPROPERTY()
+	TArray<FNiagaraVariable> Inputs;
+	/** Input parameters of this function. */
+	UPROPERTY()
+	TArray<FNiagaraVariable> Outputs;
+	/** Id of the owner is this is a member function. */
+	UPROPERTY()
+	FName OwnerName;
+	UPROPERTY()
+	bool bRequiresContext;
+	/** True if this is the signature for a "member" function of a data interface. If this is true, the first input is the owner. */
+	UPROPERTY()
+	bool bMemberFunction;
+
+	/** Localized description of this node. Note that this is *not* used during the operator == below since it may vary from culture to culture.*/
+#if WITH_EDITORONLY_DATA
+	UPROPERTY()
+	FText Description;
+#endif
+
+	FNiagaraFunctionSignature() {}
+	FNiagaraFunctionSignature(FName InName, TArray<FNiagaraVariable>& InInputs, TArray<FNiagaraVariable>& InOutputs, FName InSource, bool bInRequiresContext, bool bInMemberFunction)
+		: Name(InName)
+		, Inputs(InInputs)
+		, Outputs(InOutputs)
+		, bRequiresContext(bInRequiresContext)
+		, bMemberFunction(bInMemberFunction)
+	{
+
+	}
+
+	bool operator==(const FNiagaraFunctionSignature& Other) const
+	{
+		return Name.ToString().Equals(Other.Name.ToString()) && Inputs == Other.Inputs && Outputs == Other.Outputs && bRequiresContext == Other.bRequiresContext && bMemberFunction == Other.bMemberFunction && OwnerName == Other.OwnerName;
+	}
+
+	FString GetName()const { return Name.ToString(); }
+
+	void SetDescription(const FText& Desc)
+	{
+	#if WITH_EDITORONLY_DATA
+		Description = Desc;
+	#endif
+	}
+	FText GetDescription() const
+	{
+	#if WITH_EDITORONLY_DATA
+		return Description;
+	#else
+		return FText::FromName(Name);
+	#endif
+	}
+	bool IsValid()const { return Name != NAME_None && (Inputs.Num() > 0 || Outputs.Num() > 0); }
+};
+
+
+
+USTRUCT()
+struct NIAGARA_API FNiagaraScriptDataInterfaceInfo
+{
+	GENERATED_USTRUCT_BODY()
+public:
+	FNiagaraScriptDataInterfaceInfo()
+		: DataInterface(nullptr)
+		, Name(NAME_None)
+		, UserPtrIdx(INDEX_NONE)
+	{
+
+	}
+
+	UPROPERTY(Instanced)
+	class UNiagaraDataInterface* DataInterface;
+	
+	UPROPERTY()
+	FName Name;
+
+	/** Index of the user pointer for this data interface. */
+	UPROPERTY()
+	int32 UserPtrIdx;
+
+	TArray<FNiagaraFunctionSignature> RegisteredFunctions;
+
+	//TODO: Allow data interfaces to own datasets
+	void CopyTo(FNiagaraScriptDataInterfaceInfo* Destination, UObject* Outer) const;
+};
+
+USTRUCT()
+struct FNiagaraStatScope
+{
+	GENERATED_USTRUCT_BODY();
+
+	FNiagaraStatScope() {}
+	FNiagaraStatScope(FName InFullName, FText InFriendlyName):FullName(InFullName), FriendlyName(InFriendlyName){}
+
+	UPROPERTY()
+	FName FullName;
+
+	UPROPERTY()
+	FText FriendlyName;
+
+	bool operator==(const FNiagaraStatScope& Other) const { return FullName == Other.FullName; }
+};
+
+USTRUCT()
+struct FVMExternalFunctionBindingInfo
+{
+	GENERATED_USTRUCT_BODY();
+
+	UPROPERTY()
+	FName Name;
+
+	UPROPERTY()
+	FName OwnerName;
+
+	UPROPERTY()
+	TArray<bool> InputParamLocations;
+
+	UPROPERTY()
+	int32 NumOutputs;
+
+	FORCEINLINE int32 GetNumInputs()const { return InputParamLocations.Num(); }
+	FORCEINLINE int32 GetNumOutputs()const { return NumOutputs; }
+};
+
+struct NIAGARA_API FNiagaraSystemUpdateContext
+{
+	FNiagaraSystemUpdateContext(const UNiagaraSystem* System, bool bReInit) { Add(System, bReInit); }
+	FNiagaraSystemUpdateContext(const UNiagaraEmitter* Emitter, bool bReInit) { Add(Emitter, bReInit); }
+	FNiagaraSystemUpdateContext(const UNiagaraScript* Script, bool bReInit) { Add(Script, bReInit); }
+	//FNiagaraSystemUpdateContext(UNiagaraDataInterface* Interface, bool bReinit) : Add(Interface, bReinit) {}
+	FNiagaraSystemUpdateContext(const UNiagaraParameterCollection* Collection, bool bReInit) { Add(Collection, bReInit); }
+
+	~FNiagaraSystemUpdateContext();
+
+	void Add(const UNiagaraSystem* System, bool bReInit);
+	void Add(const UNiagaraEmitter* Emitter, bool bReInit);
+	void Add(const UNiagaraScript* Script, bool bReInit);
+	//void Add(UNiagaraDataInterface* Interface, bool bReinit);
+	void Add(const UNiagaraParameterCollection* Collection, bool bReInit);
+
+private:
+	void AddInternal(class UNiagaraComponent* Comp, bool bReInit);
+
+	FNiagaraSystemUpdateContext() { }
+	FNiagaraSystemUpdateContext(FNiagaraSystemUpdateContext& Other) { }
+
+	TArray<UNiagaraComponent*> ComponentsToReset;
+	TArray<UNiagaraComponent*> ComponentsToReInit;
+
+	//TODO: When we allow component less systems we'll also want to find and reset those.
+};
+
+
+
+/** Defines different usages for a niagara script. */
+UENUM()
+enum class ENiagaraScriptUsage : uint8
+{
+	/** The script defines a function for use in modules. */
+	Function,
+	/** The script defines a module for use in particle, emitter, or system scripts. */
+	Module,
+	/** The script defines a dynamic input for use in particle, emitter, or system scripts. */
+	DynamicInput,
+	/** The script is called when spawning particles. */
+	ParticleSpawnScript UMETA(Hidden),
+	/** Particle spawn script that handles intra-frame spawning and also pulls in the update script. */
+	ParticleSpawnScriptInterpolated UMETA(Hidden),
+	/** The script is called to update particles every frame. */
+	ParticleUpdateScript UMETA(Hidden),
+	/** The script is called to update particles in response to an event. */
+	ParticleEventScript UMETA(Hidden),
+	/** The script is called once when the emitter spawns. */
+	EmitterSpawnScript UMETA(Hidden),
+	/** The script is called every frame to tick the emitter. */
+	EmitterUpdateScript UMETA(Hidden),
+	/** The script is called once when the system spawns. */
+	SystemSpawnScript UMETA(Hidden),
+	/** The script is called every frame to tick the system. */
+	SystemUpdateScript UMETA(Hidden),
+};
+
+
+/** Defines all you need to know about a variable.*/
+USTRUCT()
+struct FNiagaraVariableInfo
+{
+	GENERATED_USTRUCT_BODY();
+
+	FNiagaraVariableInfo() : DataInterface(nullptr) {}
+
+	UPROPERTY()
+	FNiagaraVariable Variable;
+
+	UPROPERTY()
+	FText Definition;
+
+	UPROPERTY()
+	UNiagaraDataInterface* DataInterface;
+};
\ No newline at end of file
diff --git a/Engine/Plugins/FX/Niagara/Source/Niagara/Public/NiagaraComponent.h b/Engine/Plugins/FX/Niagara/Source/Niagara/Public/NiagaraComponent.h
new file mode 100644
--- /dev/null
+++ b/Engine/Plugins/FX/Niagara/Source/Niagara/Public/NiagaraComponent.h
@@ -0,0 +1,268 @@
+// Copyright 1998-2017 Epic Games, Inc. All Rights Reserved.
+
+#pragma once
+
+#include "CoreMinimal.h"
+#include "UObject/ObjectMacros.h"
+#include "NiagaraCommon.h"
+#include "PrimitiveViewRelevance.h"
+#include "PrimitiveSceneProxy.h"
+#include "Components/PrimitiveComponent.h"
+#include "NiagaraParameterStore.h"
+#include "NiagaraSystemInstance.h"
+
+#include "NiagaraComponent.generated.h"
+
+class FMeshElementCollector;
+class NiagaraRenderer;
+class UNiagaraSystem;
+class UNiagaraParameterCollection;
+class UNiagaraParameterCollectionInstance;
+
+/**
+* UNiagaraComponent is the primitive component for a Niagara System.
+* @see ANiagaraActor
+* @see UNiagaraSystem
+*/
+UCLASS(ClassGroup = (Rendering, Common), hidecategories = (Object, "Components|Activation"), editinlinenew, meta = (BlueprintSpawnableComponent))
+class NIAGARA_API UNiagaraComponent : public UPrimitiveComponent
+{
+	GENERATED_UCLASS_BODY()
+		
+	/** Defines modes for updating the component's age. */
+	enum class EAgeUpdateMode
+	{
+		/** Update the age using the delta time supplied to the tick function. */
+		TickDeltaTime,
+		/** Update the age by seeking to the DesiredAge. */
+		DesiredAge
+	};
+
+#if WITH_EDITORONLY_DATA
+	DECLARE_MULTICAST_DELEGATE(FOnSystemInstanceChanged);
+#endif
+
+private:
+	UPROPERTY(EditAnywhere, Category="Niagara", meta = (DisplayName = "Niagara System Asset"))
+	UNiagaraSystem* Asset;
+
+	/** Initial values for parameter overrides. 
+	TODO: This should be a minimal set of explicitly override parameters similar to how parameter collection instances override their collections store. 
+	Should expose anything in the "User" namespace.
+	*/
+ 	UPROPERTY(EditAnywhere, Category = Parameters)
+ 	FNiagaraParameterStore InitialParameters;
+
+#if WITH_EDITORONLY_DATA
+	UPROPERTY()
+	TMap<FName, bool> EditorOverridesValue;
+
+	FOnSystemInstanceChanged OnSystemInstanceChangedDelegate;
+#endif
+
+	/**
+	When true, this component's system will be force to update via a slower "solo" path rather than the more optimal batched path with other instances of the same system.
+	*/
+	UPROPERTY(EditAnywhere, Category = Parameters)
+	uint32 bForceSolo : 1;
+
+	TUniquePtr<FNiagaraSystemInstance> SystemInstance;
+
+	/** Defines the mode use when updating the System age. */
+	EAgeUpdateMode AgeUpdateMode;
+	
+	/** The desired age of the System instance.  This is only relevant when using the DesiredAge age update mode. */
+	float DesiredAge;
+
+	/** The delta time used when seeking to the desired age.  This is only relevant when using the DesiredAge age update mode. */
+	float SeekDelta;
+
+	//~ Begin UActorComponent Interface.
+protected:
+	virtual void OnRegister() override;
+	virtual void OnUnregister() override;
+	virtual void SendRenderDynamicData_Concurrent() override;
+	virtual void BeginDestroy() override;
+public:
+
+	virtual void Activate(bool bReset = false)override;
+	virtual void Deactivate()override;
+
+	virtual void TickComponent(float DeltaTime, enum ELevelTick TickType, FActorComponentTickFunction* ThisTickFunction) override;
+	virtual const UObject* AdditionalStatObject() const override;
+	virtual void OnComponentDestroyed(bool bDestroyingHierarchy) override;
+	//~ End UActorComponent Interface.
+
+	//~ Begin UPrimitiveComponent Interface
+	virtual int32 GetNumMaterials() const override;
+	virtual FBoxSphereBounds CalcBounds(const FTransform& LocalToWorld) const override;
+	virtual FPrimitiveSceneProxy* CreateSceneProxy() override;
+	virtual void GetUsedMaterials(TArray<UMaterialInterface*>& OutMaterials, bool bGetDebugMaterials = false) const override;
+	//~ End UPrimitiveComponent Interface
+
+	class FNiagaraSystemSimulation* GetSystemSimulation();
+
+	void DestroyInstance();
+
+	void SetAsset(UNiagaraSystem* InAsset);
+	UNiagaraSystem* GetAsset() const { return Asset; }
+
+	void SetForceSolo(bool bInForceSolo) { bForceSolo = bInForceSolo; }
+	bool GetForceSolo()const { return bForceSolo; }
+
+	EAgeUpdateMode GetAgeUpdateMode() const;
+
+	/** Sets the age update mode for the System instance. */
+	void SetAgeUpdateMode(EAgeUpdateMode InAgeUpdateMode);
+
+	/** Gets the desired age of the System instance.  This is only relevant when using the DesiredAge age update mode. */
+	float GetDesiredAge() const;
+
+	/** Sets the desired age of the System instance.  This is only relevant when using the DesiredAge age update mode. */
+	void SetDesiredAge(float InDesiredAge);
+
+	/** Gets the delta value which is used when seeking from the current age, to the desired age.  This is only relevant
+	when using the DesiredAge age update mode. */
+	float GetSeekDelta() const;
+
+	/** Sets the delta value which is used when seeking from the current age, to the desired age.  This is only relevant
+	when using the DesiredAge age update mode. */
+	void SetSeekDelta(float InSeekDelta);
+
+	FNiagaraSystemInstance* GetSystemInstance() const;
+	
+	void OnSystemDisabled();
+
+	/** Returns true if this component forces it's instances to run in "Solo" mode. A sub optimal path required in some situations. */
+	bool ForcesSolo()const;
+
+	/** Sets a Niagara Vector3 parameter by name, overriding locally if necessary.*/
+	UFUNCTION(BlueprintCallable, Category = Niagara, meta = (DisplayName = "Set Niagara Variable (Vector4)"))
+	void SetNiagaraVariableVec4(const FString& InVariableName, const FVector4& InValue);
+
+	/** Sets a Niagara Vector3 parameter by name, overriding locally if necessary.*/
+	UFUNCTION(BlueprintCallable, Category = Niagara, meta = (DisplayName = "Set Niagara Variable (Vector3)"))
+	void SetNiagaraVariableVec3(const FString& InVariableName, FVector InValue);
+
+	/** Sets a Niagara Vector3 parameter by name, overriding locally if necessary.*/
+	UFUNCTION(BlueprintCallable, Category = Niagara, meta = (DisplayName = "Set Niagara Variable (Vector2)"))
+	void SetNiagaraVariableVec2(const FString& InVariableName, FVector2D InValue);
+
+	/** Sets a Niagara float parameter by name, overriding locally if necessary.*/
+	UFUNCTION(BlueprintCallable, Category = Niagara, meta = (DisplayName = "Set Niagara Variable (Float)"))
+	void SetNiagaraVariableFloat(const FString& InVariableName, float InValue);
+
+	/** Sets a Niagara float parameter by name, overriding locally if necessary.*/
+	UFUNCTION(BlueprintCallable, Category = Niagara, meta = (DisplayName = "Set Niagara Variable (Bool)"))
+	void SetNiagaraVariableBool(const FString& InVariableName, bool InValue);
+
+	/** Debug accessors for getting positions in blueprints. */
+	UFUNCTION(BlueprintCallable, Category = Niagara, meta = (DisplayName = "Get Niagara Emitter Positions"))
+	TArray<FVector> GetNiagaraParticlePositions_DebugOnly(const FString& InEmitterName);
+	
+	/** Debug accessors for getting a float attribute array in blueprints. */
+	UFUNCTION(BlueprintCallable, Category = Niagara, meta = (DisplayName = "Get Niagara Emitter Float Attrib"))
+	TArray<float> GetNiagaraParticleValues_DebugOnly(const FString& InEmitterName, const FString& InValueName);
+	
+	/** Debug accessors for getting a FVector attribute array in blueprints. */
+	UFUNCTION(BlueprintCallable, Category = Niagara, meta = (DisplayName = "Get Niagara Emitter Vec3 Attrib"))
+	TArray<FVector> GetNiagaraParticleValueVec3_DebugOnly(const FString& InEmitterName, const FString& InValueName);
+	/** Resets the System to it's initial pre-simulated state. */
+	UFUNCTION(BlueprintCallable, Category = Niagara, meta = (DisplayName = "Reset System"))
+	void ResetSystem();
+
+	/** Called on when an external object wishes to force this System to reinitialize itself from the System data.*/
+	UFUNCTION(BlueprintCallable, Category = Niagara, meta = (DisplayName = "Reinitialize System"))
+	void ReinitializeSystem();
+
+	/** Gets whether or not rendering is enabled for this component. */
+	bool GetRenderingEnabled() const;
+
+	/** Sets whether or not rendering is enabled for this component. */
+	UFUNCTION(BlueprintCallable, Category = Niagara, meta = (DisplayName = "Set Rendering Enabled"))
+	void SetRenderingEnabled(bool bInRenderingEnabled);
+
+	//~ Begin UObject Interface.
+	virtual void PostLoad();
+#if WITH_EDITOR
+	
+	virtual void PostEditChangeProperty(FPropertyChangedEvent& PropertyChangedEvent) override;
+
+	/** Compare local overrides with the source System. Remove any that have mismatched types or no longer exist on the System. Returns whether or not any changes occurred.*/
+	virtual bool SynchronizeWithSourceSystem();
+
+	bool IsParameterValueOverriddenLocally(const FName& InParamName);
+	void SetParameterValueOverriddenLocally(const FName& InParamName, bool bInOverridden);
+	
+	FOnSystemInstanceChanged& OnSystemInstanceChanged() { return OnSystemInstanceChangedDelegate; }
+#endif
+
+	FNiagaraParameterStore& GetInitialParameters() { return InitialParameters; }
+
+	//~ End UObject Interface.
+	
+	UPROPERTY()
+	bool bRenderingEnabled;
+};
+
+
+
+
+
+
+/**
+* Scene proxy for drawing niagara particle simulations.
+*/
+class FNiagaraSceneProxy : public FPrimitiveSceneProxy
+{
+public:
+
+	FNiagaraSceneProxy(const UNiagaraComponent* InComponent);
+	~FNiagaraSceneProxy();
+
+	/** Called on render thread to assign new dynamic data */
+	void SetDynamicData_RenderThread(struct FNiagaraDynamicDataBase* NewDynamicData);
+	TArray<class NiagaraRenderer*>& GetEmitterRenderers() { return EmitterRenderers; }
+	void AddEmitterRenderer(NiagaraRenderer* Renderer);
+	void UpdateEmitterRenderers(TArray<NiagaraRenderer*>& InRenderers);
+
+	/** Gets whether or not this scene proxy should be rendered. */
+	bool GetRenderingEnabled() const;
+
+	/** Sets whether or not this scene proxy should be rendered. */
+	void SetRenderingEnabled(bool bInRenderingEnabled);
+
+private:
+	void ReleaseRenderThreadResources();
+
+	//~ Begin FPrimitiveSceneProxy Interface.
+	virtual void CreateRenderThreadResources() override;
+
+	//virtual void OnActorPositionChanged() override;
+	virtual void OnTransformChanged() override;
+
+	virtual void GetDynamicMeshElements(const TArray<const FSceneView*>& Views, const FSceneViewFamily& ViewFamily, uint32 VisibilityMap, FMeshElementCollector& Collector) const override;
+
+	virtual FPrimitiveViewRelevance GetViewRelevance(const FSceneView* View) const override;
+
+	/*
+	virtual bool CanBeOccluded() const override
+	{
+	return !MaterialRelevance.bDisableDepthTest;
+	}
+	*/
+
+	/** Callback from the renderer to gather simple lights that this proxy wants renderered. */
+	virtual void GatherSimpleLights(const FSceneViewFamily& ViewFamily, FSimpleLightArray& OutParticleLights) const override;
+
+
+	virtual uint32 GetMemoryFootprint() const override;
+
+	uint32 GetAllocatedSize() const;
+
+private:
+	//class NiagaraRenderer* EmitterRenderer;
+	TArray<class NiagaraRenderer*>EmitterRenderers;
+
+	bool bRenderingEnabled;
+};
diff --git a/Engine/Plugins/FX/Niagara/Source/Niagara/Public/NiagaraDataInterfaceBindingInstance.h b/Engine/Plugins/FX/Niagara/Source/Niagara/Public/NiagaraDataInterfaceBindingInstance.h
new file mode 100644
--- /dev/null
+++ b/Engine/Plugins/FX/Niagara/Source/Niagara/Public/NiagaraDataInterfaceBindingInstance.h
@@ -0,0 +1,6 @@
+// Copyright 1998-2016 Epic Games, Inc. All Rights Reserved.
+
+#pragma once
+#include "NiagaraScript.h"
+struct FNiagaraScriptDataInterfaceInfo;
+
diff --git a/Engine/Plugins/FX/Niagara/Source/Niagara/Public/NiagaraFunctionLibrary.h b/Engine/Plugins/FX/Niagara/Source/Niagara/Public/NiagaraFunctionLibrary.h
new file mode 100644
--- /dev/null
+++ b/Engine/Plugins/FX/Niagara/Source/Niagara/Public/NiagaraFunctionLibrary.h
@@ -0,0 +1,41 @@
+// Copyright 1998-2017 Epic Games, Inc. All Rights Reserved.
+
+#pragma once
+
+#include "CoreMinimal.h"
+#include "UObject/ObjectMacros.h"
+#include "Engine/EngineTypes.h"
+#include "Kismet/BlueprintFunctionLibrary.h"
+#include "NiagaraFunctionLibrary.generated.h"
+
+class UNiagaraComponent;
+class UNiagaraSystem;
+class USceneComponent;
+
+/**
+* A C++ and Blueprint accessible library of utility functions for accessing Niagara simulations
+* All positions & orientations are returned in Unreal reference frame & units, assuming the Leap device is located at the origin.
+*/
+UCLASS()
+class NIAGARA_API UNiagaraFunctionLibrary : public UBlueprintFunctionLibrary
+{
+	GENERATED_UCLASS_BODY()
+
+	/**
+	* Spawns a Niagara System at the specified world location/rotation
+	* @return			The spawned UNiagaraComponent
+	*/
+	UFUNCTION(BlueprintCallable, Category = Niagara, meta = (Keywords = "niagara System", WorldContext = "WorldContextObject", UnsafeDuringActorConstruction = "true"))
+	static UNiagaraComponent* SpawnSystemAtLocation(UObject* WorldContextObject, class UNiagaraSystem* SystemTemplate, FVector Location, FRotator Rotation = FRotator::ZeroRotator, bool bAutoDestroy = true);
+
+	UFUNCTION(BlueprintCallable, Category = Niagara, meta = (Keywords = "niagara System", UnsafeDuringActorConstruction = "true"))
+	static UNiagaraComponent* SpawnSystemAttached(UNiagaraSystem* SystemTemplate, USceneComponent* AttachToComponent, FName AttachPointName, FVector Location, FRotator Rotation, EAttachLocation::Type LocationType, bool bAutoDestroy);
+
+	//This is gonna be totally reworked
+// 	UFUNCTION(BlueprintCallable, Category = Niagara, meta = (Keywords = "niagara System", UnsafeDuringActorConstruction = "true"))
+// 	static void SetUpdateScriptConstant(UNiagaraComponent* Component, FName EmitterName, FName ConstantName, FVector Value);
+
+	UFUNCTION(BlueprintCallable, Category = Niagara, meta = (Keywords = "niagara parameter collection", WorldContext = "WorldContextObject"))
+	static UNiagaraParameterCollectionInstance* GetNiagaraParameterCollection(UObject* WorldContextObject, UNiagaraParameterCollection* Collection);
+
+};
diff --git a/Engine/Plugins/FX/Niagara/Source/Niagara/Public/NiagaraLightRendererProperties.h b/Engine/Plugins/FX/Niagara/Source/Niagara/Public/NiagaraLightRendererProperties.h
new file mode 100644
--- /dev/null
+++ b/Engine/Plugins/FX/Niagara/Source/Niagara/Public/NiagaraLightRendererProperties.h
@@ -0,0 +1,35 @@
+// Copyright 1998-2017 Epic Games, Inc. All Rights Reserved.
+
+#pragma once
+
+#include "CoreMinimal.h"
+#include "UObject/ObjectMacros.h"
+#include "NiagaraRendererProperties.h"
+#include "NiagaraLightRendererProperties.generated.h"
+
+
+
+UCLASS(editinlinenew)
+class UNiagaraLightRendererProperties : public UNiagaraRendererProperties
+{
+public:
+	GENERATED_BODY()
+
+		UNiagaraLightRendererProperties();
+
+	//~ UNiagaraRendererProperties interface
+	virtual NiagaraRenderer* CreateEmitterRenderer(ERHIFeatureLevel::Type FeatureLevel) override;
+	virtual void GetUsedMaterials(TArray<UMaterialInterface*>& OutMaterials) const override;
+#if WITH_EDITORONLY_DATA
+	virtual bool IsMaterialValidForRenderer(UMaterial* Material, FText& InvalidMessage) override;
+	virtual void FixMaterial(UMaterial* Material) override;
+	virtual const TArray<FNiagaraVariable>& GetRequiredAttributes() override;
+	virtual const TArray<FNiagaraVariable>& GetOptionalAttributes() override;
+#endif // WITH_EDITORONLY_DATA
+
+	UPROPERTY(EditAnywhere, Category = "Light Rendering")
+		float RadiusScale;
+
+	UPROPERTY(EditAnywhere, Category = "Light Rendering")
+		FVector ColorAdd;
+};
diff --git a/Engine/Plugins/FX/Niagara/Source/Niagara/Public/NiagaraMeshRendererProperties.h b/Engine/Plugins/FX/Niagara/Source/Niagara/Public/NiagaraMeshRendererProperties.h
new file mode 100644
--- /dev/null
+++ b/Engine/Plugins/FX/Niagara/Source/Niagara/Public/NiagaraMeshRendererProperties.h
@@ -0,0 +1,34 @@
+// Copyright 1998-2017 Epic Games, Inc. All Rights Reserved.
+
+#pragma once
+
+#include "CoreMinimal.h"
+#include "UObject/ObjectMacros.h"
+#include "NiagaraRendererProperties.h"
+#include "StaticMeshResources.h"
+#include "NiagaraMeshRendererProperties.generated.h"
+
+UCLASS(editinlinenew)
+class UNiagaraMeshRendererProperties : public UNiagaraRendererProperties
+{
+public:
+	GENERATED_BODY()
+
+	UNiagaraMeshRendererProperties();
+
+	//~ UNiagaraRendererProperties interface
+	virtual NiagaraRenderer* CreateEmitterRenderer(ERHIFeatureLevel::Type FeatureLevel) override;
+	virtual void GetUsedMaterials(TArray<UMaterialInterface*>& OutMaterials) const override;
+
+#if WITH_EDITORONLY_DATA
+	virtual bool IsMaterialValidForRenderer(UMaterial* Material, FText& InvalidMessage) override;
+	virtual void FixMaterial(UMaterial* Material) override;
+	virtual const TArray<FNiagaraVariable>& GetRequiredAttributes() override;
+	virtual const TArray<FNiagaraVariable>& GetOptionalAttributes() override;
+#endif // WITH_EDITORONLY_DATA
+
+	void PostEditChangeProperty(FPropertyChangedEvent& PropertyChangedEvent);
+
+	UPROPERTY(EditAnywhere, Category = "Mesh Rendering")
+	UStaticMesh *ParticleMesh;
+};
diff --git a/Engine/Plugins/FX/Niagara/Source/Niagara/Public/NiagaraModule.h b/Engine/Plugins/FX/Niagara/Source/Niagara/Public/NiagaraModule.h
new file mode 100644
--- /dev/null
+++ b/Engine/Plugins/FX/Niagara/Source/Niagara/Public/NiagaraModule.h
@@ -0,0 +1,46 @@
+// Copyright 1998-2017 Epic Games, Inc. All Rights Reserved.
+
+#pragma once
+
+#include "CoreMinimal.h"
+#include "Stats/Stats.h"
+#include "Modules/ModuleInterface.h"
+#include "Engine/World.h"
+
+class FNiagaraWorldManager;
+
+/**
+* Niagara module interface
+*/
+class INiagaraModule : public IModuleInterface
+{
+public:
+	DECLARE_DELEGATE_RetVal(void, FOnProcessQueue);
+
+	virtual void StartupModule()override;
+	virtual void ShutdownModule()override;
+
+	FDelegateHandle NIAGARA_API SetOnProcessShaderCompilationQueue(FOnProcessQueue InOnProcessQueue);
+	void NIAGARA_API ResetOnProcessShaderCompilationQueue(FDelegateHandle DelegateHandle);
+	void ProcessShaderCompilationQueue();
+
+
+	FNiagaraWorldManager* GetWorldManager(UWorld* World);
+
+	void DestroyAllSystemSimulations(class UNiagaraSystem* System);
+
+	// Callback function registered with global world delegates to instantiate world manager when a game world is created
+	void OnWorldInit(UWorld* World, const UWorld::InitializationValues IVS);
+
+	// Callback function registered with global world delegates to cleanup world manager contents
+	void OnWorldCleanup(UWorld* World, bool bSessionEnded, bool bCleanupResources);
+
+	// Callback function registered with global world delegates to cleanup world manager when a game world is destroyed
+	void OnPreWorldFinishDestroy(UWorld* World);
+
+	void TickWorld(UWorld* World, ELevelTick TickType, float DeltaSeconds);
+private:
+	TMap<class UWorld*, class FNiagaraWorldManager*> WorldManagers;
+	FOnProcessQueue OnProcessQueue;
+};
+
diff --git a/Engine/Plugins/FX/Niagara/Source/Niagara/Public/NiagaraParameterBindingInstance.h b/Engine/Plugins/FX/Niagara/Source/Niagara/Public/NiagaraParameterBindingInstance.h
new file mode 100644
--- /dev/null
+++ b/Engine/Plugins/FX/Niagara/Source/Niagara/Public/NiagaraParameterBindingInstance.h
@@ -0,0 +1,22 @@
+// Copyright 1998-2016 Epic Games, Inc. All Rights Reserved.
+
+#pragma once
+
+struct FNiagaraVariable;
+
+/** A runtime instance of an System parameter binding. */
+class FNiagaraParameterBindingInstance
+{
+public:
+	FNiagaraParameterBindingInstance(const FNiagaraVariable& InSource, FNiagaraVariable& InDestination);
+
+	/** Updates the paramter binding each frame. */
+	void Tick();
+
+private:
+	/** The source variable for the binding. */
+	const FNiagaraVariable& Source;
+
+	/** The destination variable for the binding. */
+	FNiagaraVariable& Destination;
+};
\ No newline at end of file
diff --git a/Engine/Plugins/FX/Niagara/Source/Niagara/Public/NiagaraParameterStore.h b/Engine/Plugins/FX/Niagara/Source/Niagara/Public/NiagaraParameterStore.h
new file mode 100644
--- /dev/null
+++ b/Engine/Plugins/FX/Niagara/Source/Niagara/Public/NiagaraParameterStore.h
@@ -0,0 +1,459 @@
+// Copyright 1998-2017 Epic Games, Inc. All Rights Reserved.
+
+#pragma once
+
+#include "NiagaraCommon.h"
+#include "NiagaraDataInterface.h"
+#include "NiagaraParameterStore.generated.h"
+
+struct FNiagaraParameterStore;
+
+//Binding from one parameter store to another.
+//This does no tracking of lifetimes etc so the owner must ensure safe use and rebinding when needed etc.
+struct FNiagaraParameterStoreBinding
+{
+	FNiagaraParameterStoreBinding()
+	{}
+
+	typedef TTuple<int32, int32, int32> TParameterBinding;
+	/** Bindings of parameter data. Src offset, Dest offset and Size. */
+	TArray<TParameterBinding> ParmeterBindings;
+
+	typedef TTuple<int32, int32> TInterfaceBinding;
+	/** Bindings of data interfaces. Src and Dest offsets.*/
+	TArray<TInterfaceBinding> InterfaceBindings;
+
+	FORCEINLINE void Empty()
+	{
+		ParmeterBindings.Empty();
+		InterfaceBindings.Empty();
+	}
+
+	FORCEINLINE_DEBUGGABLE void BindParameter(FNiagaraParameterStore& DestStore, int32 DestOffset, FNiagaraParameterStore& SrcStore, FNiagaraVariable& Parameter);
+	
+	/** TODO: Merge contiguous ranges into a single binding. */
+	//FORCEINLINE_DEBUGGABLE void Optimize();
+
+	FORCEINLINE_DEBUGGABLE void Tick(FNiagaraParameterStore& SrcStore, FNiagaraParameterStore& DestStore, bool bForce = false);
+};
+
+/** Base storage class for Niagara parameter values. */
+USTRUCT()
+struct NIAGARA_API FNiagaraParameterStore
+{
+	GENERATED_USTRUCT_BODY()
+
+	/** Owner of this store. Used to provide an outer to data interfaces in this store. */
+	UPROPERTY()
+	UObject* Owner;
+	
+	/** Map from parameter defs to their offset in the data table or the data interface. TODO: Separate out into a layout and instance class to reduce duplicated data for this?  */
+	UPROPERTY()
+	TMap<FNiagaraVariable, int32> ParameterOffsets;
+	
+	/** Buffer containing parameter data. Indexed using offsets in ParameterOffsets */
+	UPROPERTY()
+	TArray<uint8> ParameterData;
+	
+	/** Data interfaces for this script. Possibly overridden with externally owned interfaces. Also indexed by ParameterOffsets. */
+	UPROPERTY()
+	TArray<UNiagaraDataInterface*> DataInterfaces;
+
+	/** Bindings between this parameter store and others we push data into when we tick. */
+	TMap<FNiagaraParameterStore*, FNiagaraParameterStoreBinding> Bindings;
+
+	/** Marks our parameters as dirty. They will be pushed to any bound stores on tick if true. */
+	uint32 bParametersDirty : 1;
+	/** Marks our interfaces as dirty. They will be pushed to any bound stores on tick if true. */
+	uint32 bInterfacesDirty : 1;
+	/** Marks our layout as dirty. All bindings must be recreated and all parameters pushed again. */
+	uint32 bLayoutDirty : 1;
+
+	FNiagaraParameterStore();
+	FNiagaraParameterStore(UObject* InOwner);
+	FNiagaraParameterStore(const FNiagaraParameterStore& Other);
+	FNiagaraParameterStore& operator=(const FNiagaraParameterStore& Other);
+
+	~FNiagaraParameterStore();
+
+	/** Binds this parameter store to another. During Tick the values of this store will be pushed into all bound stores */
+	void Bind(FNiagaraParameterStore* DestStore);
+	/** Unbinds this store form one it's bound to. */
+	void Unbind(FNiagaraParameterStore* DestStore);
+	/** Recreates any bindings to reflect a layout change etc. */
+	void Rebind();
+	/** Recreates any bindings to reflect a layout change etc. */
+	void TransferBindings(FNiagaraParameterStore& OtherStore);
+	/** Handles any update such as pushing parameters to bound stores etc. */
+	void Tick();
+
+	/**
+	Adds the passed parameter to this store.
+	Does nothing if this parameter is already present.
+	Returns true if we added a new parameter.
+	*/
+	bool AddParameter(const FNiagaraVariable& Param, bool bInitialize=true);
+
+	/** Removes the passed parameter if it exists in the store. */
+	bool RemoveParameter(const FNiagaraVariable& Param);
+
+	/** Renames the passed parameter. */
+	void RenameParameter(const FNiagaraVariable& Param, FName NewName);
+
+	/** Removes all parameters from this store and releases any data. */
+	void Empty(bool bClearBindings=true);
+
+	FORCEINLINE void GetParameters(TArray<FNiagaraVariable>& OutParametres) const { return ParameterOffsets.GenerateKeyArray(OutParametres); }
+
+	FORCEINLINE const TArray<UNiagaraDataInterface*>& GetDataInterfaces()const { return DataInterfaces; }
+	FORCEINLINE const TArray<uint8>& GetParameterDataArray()const { return ParameterData; }
+	FORCEINLINE TArray<uint8>& GetParameterDataArray() { return ParameterData; }
+
+	/** Gets the index of the passed parameter. If it is a data interface, this is an offset into the data interface table, otherwise a byte offset into he parameter data buffer. */
+	int32 IndexOf(const FNiagaraVariable& Parameter)const
+	{
+		const int32* Off = ParameterOffsets.Find(Parameter);
+		if ( Off )
+		{
+			return *Off;
+		}
+		return INDEX_NONE;
+	}
+
+	/** Gets the typed parameter data. */
+	template<typename T>
+	FORCEINLINE_DEBUGGABLE void GetParameterValue(T& OutValue, const FNiagaraVariable& Parameter)const
+	{
+		check(Parameter.GetSizeInBytes() == sizeof(T));
+		int32 Offset = IndexOf(Parameter);
+		if (Offset != INDEX_NONE)
+		{
+			OutValue = *(const T*)(GetParameterData(Offset));
+		}
+	}
+
+	template<typename T>
+	FORCEINLINE_DEBUGGABLE T GetParameterValue(const FNiagaraVariable& Parameter)const
+	{
+		check(Parameter.GetSizeInBytes() == sizeof(T));
+		int32 Offset = IndexOf(Parameter);
+		if (Offset != INDEX_NONE)
+		{
+			return *(const T*)(GetParameterData(Offset));
+		}
+		return T();
+	}
+
+	FORCEINLINE const uint8* GetParameterData(int32 Offset)const
+	{
+		return ParameterData.GetData() + Offset;
+	}
+
+	/** Returns the parameter data for the passed parameter if it exists in this store. Null if not. */
+	FORCEINLINE_DEBUGGABLE const uint8* GetParameterData(const FNiagaraVariable& Parameter)const
+	{
+		int32 Offset = IndexOf(Parameter);
+		if (Offset != INDEX_NONE)
+		{
+			return GetParameterData(Offset);
+		}
+		return nullptr;
+	}
+
+	FORCEINLINE_DEBUGGABLE uint8* GetParameterData(const FNiagaraVariable& Parameter)
+	{
+		int32 Offset = IndexOf(Parameter);
+		if (Offset != INDEX_NONE)
+		{
+			return GetParameterData_Internal(Offset);
+		}
+		return nullptr;
+	}
+
+	/** Returns the data interface at the passed offset. */
+	FORCEINLINE void GetDataInterface(UNiagaraDataInterface* OutInterface, int32 Offset)const
+	{
+		OutInterface = GetDataInterface(Offset);
+	}
+
+	/** Returns the data interface at the passed offset. */
+	FORCEINLINE UNiagaraDataInterface* GetDataInterface(int32 Offset)const
+	{
+		if (DataInterfaces.IsValidIndex(Offset))
+		{
+			return DataInterfaces[Offset];
+		}
+
+		return nullptr;
+	}
+
+	/** Returns the data interface for the passed parameter if it exists in this store. */
+	FORCEINLINE_DEBUGGABLE UNiagaraDataInterface* GetDataInterface(const FNiagaraVariable& Parameter)const
+	{
+		int32 Offset = IndexOf(Parameter);
+		UNiagaraDataInterface* Interface = GetDataInterface(Offset);
+		checkSlow(!Interface || Parameter.GetType().GetClass() == Interface->GetClass());
+		return Interface;
+	}
+
+	/** Returns the associated FNiagaraVariable for the passed data interface if it exists in the store. Null if not.*/
+	const FNiagaraVariable* FindVariable(UNiagaraDataInterface* Interface)const;
+
+	/** Copies the passed parameter from this parameter store into another. */
+	FORCEINLINE_DEBUGGABLE void CopyParameterData(FNiagaraParameterStore& DestStore, const FNiagaraVariable& Parameter) const
+	{
+		int32 DestIndex = DestStore.IndexOf(Parameter);
+		int32 SrcIndex = IndexOf(Parameter);
+		if (DestIndex != INDEX_NONE && SrcIndex != INDEX_NONE)
+		{
+			DestStore.SetParameterData(GetParameterData_Internal(SrcIndex), DestIndex, Parameter.GetSizeInBytes());
+		}
+	}
+
+	template<typename T>
+	FORCEINLINE_DEBUGGABLE void SetParameterValue(const T& InValue, const FNiagaraVariable& Param, bool bAdd=false)
+	{
+		check(Param.GetSizeInBytes() == sizeof(T));
+		int32 Offset = IndexOf(Param);
+		if (Offset != INDEX_NONE)
+		{
+			//*(T*)(GetParameterData_Internal(Offset)) = InValue;
+			//Until we solve our alignment issues, temporarily just doing a memcpy here.
+			FMemory::Memcpy(GetParameterData_Internal(Offset), &InValue, sizeof(T));
+			OnParameterChange();
+		}
+		else
+		{
+			if (bAdd)
+			{
+				AddParameter(Param);
+				Offset = IndexOf(Param);
+				check(Offset != INDEX_NONE);
+				*(T*)(GetParameterData_Internal(Offset)) = InValue;
+			}
+		}
+	}
+
+	FORCEINLINE_DEBUGGABLE void SetParameterData(const uint8* Data, int32 Offset, int32 Size)
+	{
+		FMemory::Memcpy(GetParameterData_Internal(Offset), Data, Size);
+		OnParameterChange();
+	}
+
+	FORCEINLINE_DEBUGGABLE void SetParameterData(const uint8* Data, const FNiagaraVariable& Param)
+	{
+		int32 Offset = IndexOf(Param);
+		if (Offset != INDEX_NONE)
+		{
+			checkSlow(!Param.IsDataInterface());
+			FMemory::Memcpy(GetParameterData_Internal(Offset), Data, Param.GetSizeInBytes());
+			OnParameterChange();
+		}
+	}
+
+	/**
+	* Sets the parameter using the internally stored data in the passed FNiagaraVariable.
+	* TODO: Remove this. IMO FNiagaraVariable should be just for data definition and all storage should be done via this class.
+	*/
+	FORCEINLINE_DEBUGGABLE void SetParameter(const FNiagaraVariable& Param)
+	{
+		checkSlow(Param.IsDataAllocated());
+		int32 Offset = IndexOf(Param);
+		if (Offset != INDEX_NONE)
+		{
+			FMemory::Memcpy(GetParameterData_Internal(Offset), Param.GetData(), Param.GetSizeInBytes());
+			OnParameterChange();
+		}
+	}
+
+	FORCEINLINE_DEBUGGABLE void SetDataInterface(UNiagaraDataInterface* InInterface, int32 Offset)
+	{
+		DataInterfaces[Offset] = InInterface;
+		OnInterfaceChange();
+	}
+
+	FORCEINLINE_DEBUGGABLE void SetDataInterface(UNiagaraDataInterface* InInterface, const FNiagaraVariable& Parameter)
+	{
+		int32 Offset = IndexOf(Parameter);
+		if (Offset != INDEX_NONE)
+		{
+			DataInterfaces[Offset] = InInterface;
+			OnInterfaceChange();
+		}
+	}
+
+	FORCEINLINE void OnParameterChange() { bParametersDirty = true; }
+	FORCEINLINE void OnInterfaceChange() { bInterfacesDirty = true; }
+
+protected:
+	FORCEINLINE void OnLayoutChange() { bLayoutDirty = true; }
+
+	/** Returns the parameter data at the passed offset. */
+	FORCEINLINE const uint8* GetParameterData_Internal(int32 Offset) const
+	{
+		return ParameterData.GetData() + Offset;
+	}
+
+	FORCEINLINE uint8* GetParameterData_Internal(int32 Offset) 
+	{
+		return ParameterData.GetData() + Offset;
+	}
+};
+
+
+FORCEINLINE_DEBUGGABLE void FNiagaraParameterStoreBinding::BindParameter(FNiagaraParameterStore& DestStore, int32 DestOffset, FNiagaraParameterStore& SrcStore, FNiagaraVariable& Parameter)
+{
+	int32 SrcOffset = SrcStore.IndexOf(Parameter);
+	checkSlow(DestOffset == DestStore.IndexOf(Parameter));
+
+	if (SrcOffset == INDEX_NONE || DestOffset == INDEX_NONE)
+	{
+		//UE_LOG(LogNiagara, Error, TEXT("Failed to bind parameter %s."), *Parameter.GetName().ToString());
+		return;
+	}
+
+	if (Parameter.IsDataInterface())
+	{
+		InterfaceBindings.AddUnique(TInterfaceBinding(SrcOffset, DestOffset));
+	}
+	else
+	{
+		ParmeterBindings.AddUnique(TParameterBinding(SrcOffset, DestOffset, Parameter.GetSizeInBytes()));
+	}
+}
+
+/** Merge contiguous ranges into a single binding. */
+// FORCEINLINE_DEBUGGABLE void FNiagaraParameterStoreBinding::Optimize()
+// {
+// 	//TODO
+// }
+
+FORCEINLINE_DEBUGGABLE void FNiagaraParameterStoreBinding::Tick(FNiagaraParameterStore& SrcStore, FNiagaraParameterStore& DestStore, bool bForce)
+{
+	if (SrcStore.bParametersDirty || bForce)
+	{
+		DestStore.bParametersDirty = true;
+		for (TParameterBinding& Binding : ParmeterBindings)
+		{
+			int32 SrcOffset = Binding.Get<0>();
+			int32 DestOffset = Binding.Get<1>();
+			int32 Size = Binding.Get<2>();
+			DestStore.SetParameterData(SrcStore.GetParameterData(SrcOffset), DestOffset, Size);
+		}
+	}
+
+	if (SrcStore.bInterfacesDirty || bForce)
+	{
+		DestStore.bInterfacesDirty = true;
+		for (TInterfaceBinding& Binding : InterfaceBindings)
+		{
+			DestStore.SetDataInterface(SrcStore.GetDataInterface(Binding.Get<0>()), Binding.Get<1>());
+		}
+	}
+}
+
+//////////////////////////////////////////////////////////////////////////
+
+/**
+Direct binding to a parameter store to allow efficient gets/sets from code etc. 
+Does no tracking of lifetimes etc so users are responsible for safety.
+*/
+template<typename T>
+struct FNiagaraParameterDirectBinding
+{
+	T* ValuePtr;
+
+	FNiagaraParameterDirectBinding()
+		: ValuePtr(nullptr)
+	{}
+
+	T* Init(FNiagaraParameterStore& InStore, const FNiagaraVariable& DestVariable)
+	{
+		check(DestVariable.GetSizeInBytes() == sizeof(T));
+		ValuePtr = (T*)InStore.GetParameterData(DestVariable);
+		return ValuePtr;
+	}
+
+	FORCEINLINE void SetValue(const T& InValue)
+	{
+		if (ValuePtr)
+		{
+			*ValuePtr = InValue;
+		}
+	}
+
+	FORCEINLINE T GetValue()const 
+	{
+		if (ValuePtr)
+		{
+			return *ValuePtr;
+		}
+		return T();
+	}
+};
+
+template<>
+struct FNiagaraParameterDirectBinding<FMatrix>
+{
+	FMatrix* ValuePtr;
+
+	FNiagaraParameterDirectBinding():ValuePtr(nullptr){}
+
+	FMatrix* Init(FNiagaraParameterStore& InStore, const FNiagaraVariable& DestVariable)
+	{
+		check(DestVariable.GetSizeInBytes() == sizeof(FMatrix));
+		ValuePtr = (FMatrix*)InStore.GetParameterData(DestVariable);
+		return ValuePtr;
+	}
+
+	FORCEINLINE void SetValue(const FMatrix& InValue)
+	{
+		if (ValuePtr)
+		{
+			FMemory::Memcpy(ValuePtr, &InValue, sizeof(FMatrix));//Temp annoyance until we fix the alignment issues with parameter stores.
+		}
+	}
+
+	FORCEINLINE FMatrix GetValue()const
+	{
+		FMatrix Ret;
+		if (ValuePtr)
+		{
+			FMemory::Memcpy(&Ret, ValuePtr, sizeof(FMatrix));//Temp annoyance until we fix the alignment issues with parameter stores.
+		}
+		return Ret;
+	}
+};
+
+template<>
+struct FNiagaraParameterDirectBinding<FVector4>
+{
+	FVector4* ValuePtr;
+
+	FNiagaraParameterDirectBinding() :ValuePtr(nullptr) {}
+
+	FVector4* Init(FNiagaraParameterStore& InStore, const FNiagaraVariable& DestVariable)
+	{
+		check(DestVariable.GetSizeInBytes() == sizeof(FVector4));
+		ValuePtr = (FVector4*)InStore.GetParameterData(DestVariable);
+		return ValuePtr;
+	}
+
+	FORCEINLINE void SetValue(const FVector4& InValue)
+	{
+		if (ValuePtr)
+		{
+			FMemory::Memcpy(ValuePtr, &InValue, sizeof(FVector4));//Temp annoyance until we fix the alignment issues with parameter stores.
+		}
+	}
+
+	FORCEINLINE FVector4 GetValue()const
+	{
+		FVector4 Ret;
+		if (ValuePtr)
+		{
+			FMemory::Memcpy(&Ret, ValuePtr, sizeof(FVector4));//Temp annoyance until we fix the alignment issues with parameter stores.
+		}
+		return Ret;
+	}
+};
\ No newline at end of file
diff --git a/Engine/Plugins/FX/Niagara/Source/Niagara/Public/NiagaraRenderer.h b/Engine/Plugins/FX/Niagara/Source/Niagara/Public/NiagaraRenderer.h
new file mode 100644
--- /dev/null
+++ b/Engine/Plugins/FX/Niagara/Source/Niagara/Public/NiagaraRenderer.h
@@ -0,0 +1,264 @@
+// Copyright 1998-2017 Epic Games, Inc. All Rights Reserved.
+
+/*==============================================================================
+NiagaraRenderer.h: Base class for Niagara render modules
+==============================================================================*/
+#pragma once
+
+#include "CoreMinimal.h"
+#include "NiagaraCommon.h"
+#include "Materials/MaterialInterface.h"
+#include "UniformBuffer.h"
+#include "Materials/Material.h"
+#include "PrimitiveViewRelevance.h"
+#include "ParticleHelper.h"
+#include "NiagaraSpriteRendererProperties.h"
+#include "NiagaraRibbonRendererProperties.h"
+#include "NiagaraMeshRendererProperties.h"
+#include "NiagaraLightRendererProperties.h"
+#include "RenderingThread.h"
+#include "SceneView.h"
+#include "NiagaraComponent.h"
+
+class FNiagaraDataSet;
+
+/** Struct used to pass dynamic data from game thread to render thread */
+struct FNiagaraDynamicDataBase
+{
+};
+
+
+
+class SimpleTimer
+{
+public:
+	SimpleTimer()
+	{
+		StartTime = FPlatformTime::Seconds() * 1000.0;
+	}
+
+	double GetElapsedMilliseconds()
+	{
+		return (FPlatformTime::Seconds()*1000.0) - StartTime;
+	}
+
+	~SimpleTimer()
+	{
+	}
+
+private:
+	double StartTime;
+};
+
+
+/**
+* Base class for Niagara System renderers. System renderers handle generating vertex data for and
+* drawing of simulation data coming out of FNiagaraEmitterInstance instances.
+*/
+class NiagaraRenderer 
+{
+public:
+	virtual void GetDynamicMeshElements(const TArray<const FSceneView*>& Views, const FSceneViewFamily& ViewFamily, uint32 VisibilityMap, FMeshElementCollector& Collector, const FNiagaraSceneProxy *SceneProxy) const = 0;
+
+	virtual void SetDynamicData_RenderThread(FNiagaraDynamicDataBase* NewDynamicData) = 0;
+	virtual void CreateRenderThreadResources() = 0;
+	virtual void ReleaseRenderThreadResources() = 0;
+	virtual FNiagaraDynamicDataBase *GenerateVertexData(const FNiagaraSceneProxy* Proxy, FNiagaraDataSet &Data, const ENiagaraSimTarget Target) = 0;
+	virtual int GetDynamicDataSize() = 0;
+
+	virtual bool HasDynamicData() = 0;
+
+	virtual bool SetMaterialUsage() = 0;
+
+#if WITH_EDITORONLY_DATA
+	virtual const TArray<FNiagaraVariable>& GetRequiredAttributes() = 0;
+	virtual const TArray<FNiagaraVariable>& GetOptionalAttributes() = 0;
+#endif
+
+	virtual FPrimitiveViewRelevance GetViewRelevance(const FSceneView* View, const FNiagaraSceneProxy *SceneProxy)
+	{
+		FPrimitiveViewRelevance Result;
+		bool bHasDynamicData = HasDynamicData();
+
+		Result.bDrawRelevance = bHasDynamicData && SceneProxy->IsShown(View) && View->Family->EngineShowFlags.Particles;
+		Result.bShadowRelevance = bHasDynamicData && SceneProxy->IsShadowCast(View);
+		Result.bDynamicRelevance = bHasDynamicData;
+		if (bHasDynamicData && View->Family->EngineShowFlags.Bounds)
+		{
+			Result.bOpaqueRelevance = true;
+		}
+		MaterialRelevance.SetPrimitiveViewRelevance(Result);
+
+		return Result;
+	}
+
+
+	UMaterialInterface *GetMaterial()	{ return Material; }
+	void SetMaterial(UMaterialInterface *InMaterial, ERHIFeatureLevel::Type FeatureLevel)
+	{
+		Material = InMaterial;
+		if (Material->GetMaterialResource(FeatureLevel))
+		{
+			if (!Material || !SetMaterialUsage())
+			{
+				Material = UMaterial::GetDefaultMaterial(MD_Surface);
+			}
+			check(Material);
+			MaterialRelevance = Material->GetRelevance(FeatureLevel);
+		}
+	}
+	
+	virtual UClass *GetPropertiesClass() = 0;
+	virtual void SetRendererProperties(UNiagaraRendererProperties *Props) = 0;
+
+	float GetCPUTimeMS() { return CPUTimeMS; }
+
+	void SetLocalSpace(bool bInLocalSpace) { bLocalSpace = bInLocalSpace; }
+
+	/** Release enqueues the System renderer to be killed on the render thread safely.*/
+	void Release();
+
+	FNiagaraDynamicDataBase *GetDynamicData() { return DynamicDataRender; }
+
+	bool IsEnabled() const {return bEnabled;}
+
+	void SetEnabled(bool bInEnabled) { bEnabled = bInEnabled; }
+
+protected:
+	NiagaraRenderer()
+		: CPUTimeMS(0.0f)
+		, bLocalSpace(false)
+		, bEnabled(true)
+		, DynamicDataRender(nullptr)
+	{
+		Material = UMaterial::GetDefaultMaterial(MD_Surface);
+	}
+
+	virtual ~NiagaraRenderer();
+
+	mutable float CPUTimeMS;
+	UMaterialInterface* Material;
+	bool bLocalSpace;
+	bool bEnabled;
+
+	FMaterialRelevance MaterialRelevance;
+
+	struct FNiagaraDynamicDataBase *DynamicDataRender;
+};
+
+
+
+
+/**
+* NiagaraRendererSprites renders an FNiagaraEmitterInstance as sprite particles
+*/
+class NIAGARA_API NiagaraRendererSprites : public NiagaraRenderer
+{
+public:
+
+	explicit NiagaraRendererSprites(ERHIFeatureLevel::Type FeatureLevel, UNiagaraRendererProperties *Props);
+	~NiagaraRendererSprites()
+	{
+		ReleaseRenderThreadResources();
+	}
+
+
+	virtual void ReleaseRenderThreadResources() override;
+
+	// FPrimitiveSceneProxy interface.
+	virtual void CreateRenderThreadResources() override;
+
+	virtual void GetDynamicMeshElements(const TArray<const FSceneView*>& Views, const FSceneViewFamily& ViewFamily, uint32 VisibilityMap, FMeshElementCollector& Collector, const FNiagaraSceneProxy *SceneProxy) const override;
+	virtual bool SetMaterialUsage() override;
+	/** Update render data buffer from attributes */
+	FNiagaraDynamicDataBase *GenerateVertexData(const FNiagaraSceneProxy* Proxy, FNiagaraDataSet &Data, const ENiagaraSimTarget Target) override;
+
+	virtual void SetDynamicData_RenderThread(FNiagaraDynamicDataBase* NewDynamicData) override;
+	int GetDynamicDataSize() override;
+	bool HasDynamicData() override;
+
+	UClass *GetPropertiesClass() override { return UNiagaraSpriteRendererProperties::StaticClass(); }
+	void SetRendererProperties(UNiagaraRendererProperties *Props) override { Properties = Cast<UNiagaraSpriteRendererProperties>(Props); }
+
+#if WITH_EDITORONLY_DATA
+	virtual const TArray<FNiagaraVariable>& GetRequiredAttributes() override;
+	virtual const TArray<FNiagaraVariable>& GetOptionalAttributes() override;
+#endif
+
+private:
+	UNiagaraSpriteRendererProperties *Properties;
+	mutable TUniformBuffer<FPrimitiveUniformShaderParameters> WorldSpacePrimitiveUniformBuffer;
+	class FNiagaraSpriteVertexFactory* VertexFactory;
+};
+
+
+/**
+* NiagaraRendererLights renders an FNiagaraEmitterInstance as simple lights
+*/
+class NIAGARA_API NiagaraRendererLights : public NiagaraRenderer
+{
+public:
+	struct SimpleLightData
+	{
+		FSimpleLightEntry LightEntry;
+		FSimpleLightPerViewEntry PerViewEntry;
+	};
+	explicit NiagaraRendererLights(ERHIFeatureLevel::Type FeatureLevel, UNiagaraRendererProperties *Props);
+
+	~NiagaraRendererLights()
+	{
+		ReleaseRenderThreadResources();
+	}
+
+
+	virtual FPrimitiveViewRelevance GetViewRelevance(const FSceneView* View, const FNiagaraSceneProxy *SceneProxy) override
+	{
+		FPrimitiveViewRelevance Result;
+		Result.bDrawRelevance = false;
+		Result.bShadowRelevance = false;
+		Result.bDynamicRelevance = false;
+		Result.bOpaqueRelevance = false;
+		Result.bHasSimpleLights = true;
+		//MaterialRelevance.SetPrimitiveViewRelevance(Result);
+		return Result;
+	}
+
+
+	virtual void ReleaseRenderThreadResources() override;
+
+	// FPrimitiveSceneProxy interface.
+	virtual void CreateRenderThreadResources() override;
+
+	virtual void GetDynamicMeshElements(const TArray<const FSceneView*>& Views, const FSceneViewFamily& ViewFamily, uint32 VisibilityMap, FMeshElementCollector& Collector, const FNiagaraSceneProxy *SceneProxy) const override;
+	virtual bool SetMaterialUsage() override;
+	/** Update render data buffer from attributes */
+	FNiagaraDynamicDataBase *GenerateVertexData(const FNiagaraSceneProxy* Proxy, FNiagaraDataSet &Data, const ENiagaraSimTarget Target) override;
+
+	virtual void SetDynamicData_RenderThread(FNiagaraDynamicDataBase* NewDynamicData) override;
+	int GetDynamicDataSize() override;
+	bool HasDynamicData() override;
+
+	UClass *GetPropertiesClass() override { return UNiagaraLightRendererProperties::StaticClass(); }
+	void SetRendererProperties(UNiagaraRendererProperties *Props) override { Properties = Cast<UNiagaraLightRendererProperties>(Props); }
+#if WITH_EDITORONLY_DATA
+	virtual const TArray<FNiagaraVariable>& GetRequiredAttributes() override;
+	virtual const TArray<FNiagaraVariable>& GetOptionalAttributes() override;
+#endif
+
+	TArray<SimpleLightData> &GetLights() { return LightArray; }
+private:
+	UNiagaraLightRendererProperties *Properties;
+	TArray<SimpleLightData>LightArray;
+};
+
+
+struct FNiagaraDynamicDataLights : public FNiagaraDynamicDataBase
+{
+	TArray<NiagaraRendererLights::SimpleLightData> LightArray;
+};
+
+
+
+
+
+
diff --git a/Engine/Plugins/FX/Niagara/Source/Niagara/Public/NiagaraRendererMeshes.h b/Engine/Plugins/FX/Niagara/Source/Niagara/Public/NiagaraRendererMeshes.h
new file mode 100644
--- /dev/null
+++ b/Engine/Plugins/FX/Niagara/Source/Niagara/Public/NiagaraRendererMeshes.h
@@ -0,0 +1,91 @@
+// Copyright 1998-2017 Epic Games, Inc. All Rights Reserved.
+
+/*==============================================================================
+NiagaraRenderer.h: Base class for Niagara render modules
+==============================================================================*/
+#pragma once
+
+#include "NiagaraMeshVertexFactory.h"
+#include "NiagaraRenderer.h"
+
+class FNiagaraDataSet;
+
+
+
+struct FNiagaraDynamicDataMesh : public FNiagaraDynamicDataBase
+{
+	const FNiagaraDataSet *DataSet;
+	int32 PositionDataOffset;
+	int32 VelocityDataOffset;
+	int32 ColorDataOffset;
+	int32 TransformDataOffset;
+	int32 ScaleDataOffset;
+	int32 SizeDataOffset;
+	int32 MaterialParamDataOffset;
+};
+
+
+
+/**
+* NiagaraRendererSprites renders an FNiagaraEmitterInstance as sprite particles
+*/
+class NIAGARA_API NiagaraRendererMeshes : public NiagaraRenderer
+{
+public:
+
+	explicit NiagaraRendererMeshes(ERHIFeatureLevel::Type FeatureLevel, UNiagaraRendererProperties *Props);
+	~NiagaraRendererMeshes()
+	{
+		ReleaseRenderThreadResources();
+	}
+
+
+	virtual void ReleaseRenderThreadResources() override;
+
+	// FPrimitiveSceneProxy interface.
+	virtual void CreateRenderThreadResources() override;
+
+	virtual void GetDynamicMeshElements(const TArray<const FSceneView*>& Views, const FSceneViewFamily& ViewFamily, uint32 VisibilityMap, FMeshElementCollector& Collector, const FNiagaraSceneProxy *SceneProxy) const override;
+	virtual bool SetMaterialUsage() override;
+	/** Update render data buffer from attributes */
+	FNiagaraDynamicDataBase *GenerateVertexData(const FNiagaraSceneProxy* Proxy, FNiagaraDataSet &Data, const ENiagaraSimTarget Target) override;
+
+	virtual void SetDynamicData_RenderThread(FNiagaraDynamicDataBase* NewDynamicData) override;
+	int GetDynamicDataSize() override;
+	bool HasDynamicData() override;
+
+	virtual FPrimitiveViewRelevance GetViewRelevance(const FSceneView* View, const FNiagaraSceneProxy *SceneProxy)
+	{
+		FPrimitiveViewRelevance Result;
+		bool bHasDynamicData = HasDynamicData();
+		Result.bDrawRelevance = bHasDynamicData && SceneProxy->IsShown(View) && View->Family->EngineShowFlags.Particles;
+		Result.bShadowRelevance = bHasDynamicData && SceneProxy->IsShadowCast(View);
+		Result.bDynamicRelevance = bHasDynamicData;
+
+		if (bHasDynamicData)
+		{
+			Result.bOpaqueRelevance = MaterialRelevance.bOpaque;
+			Result.bNormalTranslucencyRelevance = MaterialRelevance.bNormalTranslucency;
+			Result.bSeparateTranslucencyRelevance = MaterialRelevance.bSeparateTranslucency;
+			Result.bDistortionRelevance = MaterialRelevance.bDistortion;
+		}
+
+		return Result;
+	}
+
+
+
+
+	UClass *GetPropertiesClass() override { return UNiagaraMeshRendererProperties::StaticClass(); }
+	void SetRendererProperties(UNiagaraRendererProperties *Props) override { Properties = Cast<UNiagaraMeshRendererProperties>(Props); }
+#if WITH_EDITORONLY_DATA
+	virtual const TArray<FNiagaraVariable>& GetRequiredAttributes() override;
+	virtual const TArray<FNiagaraVariable>& GetOptionalAttributes() override;
+#endif
+	void SetupVertexFactory(FNiagaraMeshVertexFactory *InVertexFactory, const FStaticMeshLODResources& LODResources) const;
+
+private:
+	UNiagaraMeshRendererProperties *Properties;
+	mutable TUniformBuffer<FPrimitiveUniformShaderParameters> WorldSpacePrimitiveUniformBuffer;
+	class FNiagaraMeshVertexFactory* VertexFactory;
+};
diff --git a/Engine/Plugins/FX/Niagara/Source/Niagara/Public/NiagaraRendererProperties.h b/Engine/Plugins/FX/Niagara/Source/Niagara/Public/NiagaraRendererProperties.h
new file mode 100644
--- /dev/null
+++ b/Engine/Plugins/FX/Niagara/Source/Niagara/Public/NiagaraRendererProperties.h
@@ -0,0 +1,41 @@
+// Copyright 1998-2017 Epic Games, Inc. All Rights Reserved.
+
+#pragma once
+
+#include "CoreMinimal.h"
+#include "UObject/ObjectMacros.h"
+#include "UObject/Object.h"
+#include "RHIDefinitions.h"
+#include "NiagaraTypes.h"
+#include "NiagaraRendererProperties.generated.h"
+/**
+* Emitter properties base class
+* Each EmitterRenderer derives from this with its own class, and returns it in GetProperties; a copy
+* of those specific properties is stored on UNiagaraEmitter (on the System) for serialization 
+* and handed back to the System renderer on load.
+*/
+
+class NiagaraRenderer;
+class UMaterial;
+class UMaterialInterface;
+
+UCLASS(ABSTRACT)
+class NIAGARA_API UNiagaraRendererProperties : public UObject
+{
+	GENERATED_BODY()
+
+public:
+	virtual NiagaraRenderer* CreateEmitterRenderer(ERHIFeatureLevel::Type FeatureLevel) PURE_VIRTUAL ( UNiagaraRendererProperties::CreateEmitterRenderer, return nullptr;);
+	virtual void GetUsedMaterials(TArray<UMaterialInterface*>& OutMaterials) const PURE_VIRTUAL(UNiagaraRendererProperties::GetUsedMaterials,);
+
+#if WITH_EDITORONLY_DATA
+	virtual bool IsMaterialValidForRenderer(UMaterial* Material, FText& InvalidMessage) { return true; }
+
+	virtual void FixMaterial(UMaterial* Material) { }
+
+	virtual const TArray<FNiagaraVariable>& GetRequiredAttributes() { static TArray<FNiagaraVariable> Vars; return Vars; };
+	virtual const TArray<FNiagaraVariable>& GetOptionalAttributes() { static TArray<FNiagaraVariable> Vars; return Vars; };
+#endif // WITH_EDITORONLY_DATA
+};
+
+
diff --git a/Engine/Plugins/FX/Niagara/Source/Niagara/Public/NiagaraRendererRibbons.h b/Engine/Plugins/FX/Niagara/Source/Niagara/Public/NiagaraRendererRibbons.h
new file mode 100644
--- /dev/null
+++ b/Engine/Plugins/FX/Niagara/Source/Niagara/Public/NiagaraRendererRibbons.h
@@ -0,0 +1,75 @@
+// Copyright 1998-2017 Epic Games, Inc. All Rights Reserved.
+
+/*==============================================================================
+NiagaraRenderer.h: Base class for Niagara render modules
+==============================================================================*/
+#pragma once
+
+#include "NiagaraRibbonVertexFactory.h"
+#include "NiagaraRenderer.h"
+
+class FNiagaraDataSet;
+
+/**
+* NiagaraRendererRibbons renders an FNiagaraEmitterInstance as a ribbon connecting all particles
+* in order by particle age.
+*/
+class NIAGARA_API NiagaraRendererRibbons : public NiagaraRenderer
+{
+public:
+	NiagaraRendererRibbons(ERHIFeatureLevel::Type FeatureLevel, UNiagaraRendererProperties *Props);
+	~NiagaraRendererRibbons()
+	{
+		ReleaseRenderThreadResources();
+	}
+
+	virtual void ReleaseRenderThreadResources() override;
+	virtual void CreateRenderThreadResources() override;
+
+
+	virtual void GetDynamicMeshElements(const TArray<const FSceneView*>& Views, const FSceneViewFamily& ViewFamily, uint32 VisibilityMap, FMeshElementCollector& Collector, const FNiagaraSceneProxy *SceneProxy) const override;
+
+	virtual bool SetMaterialUsage() override;
+
+	/** Update render data buffer from attributes */
+	FNiagaraDynamicDataBase *GenerateVertexData(const FNiagaraSceneProxy* Proxy, FNiagaraDataSet &Data, const ENiagaraSimTarget Target) override;
+
+	void AddRibbonVert(TArray<FNiagaraRibbonVertex>& RenderData, FVector ParticlePos, FVector2D UV1,
+		const FLinearColor &Color, const float Age, const float Rotation, const FVector2D &Size)
+	{
+		FNiagaraRibbonVertex& NewVertex = *new(RenderData)FNiagaraRibbonVertex;
+		NewVertex.Position = ParticlePos;
+		NewVertex.OldPosition = NewVertex.Position;
+		NewVertex.Color = Color;
+		NewVertex.ParticleId = RenderData.Num();
+		NewVertex.RelativeTime = Age;
+		NewVertex.Size = Size;
+		NewVertex.Rotation = Rotation;
+		NewVertex.SubImageIndex = 0.f;
+		NewVertex.Tex_U = UV1.X;
+		NewVertex.Tex_V = UV1.Y;
+		NewVertex.Tex_U2 = UV1.X;
+		NewVertex.Tex_V2 = UV1.Y;
+	}
+
+
+
+	virtual void SetDynamicData_RenderThread(FNiagaraDynamicDataBase* NewDynamicData) override;
+	int GetDynamicDataSize() override;
+	bool HasDynamicData() override;
+
+#if WITH_EDITORONLY_DATA
+	virtual const TArray<FNiagaraVariable>& GetRequiredAttributes() override;
+	virtual const TArray<FNiagaraVariable>& GetOptionalAttributes() override;
+#endif
+
+	UClass *GetPropertiesClass() override { return UNiagaraRibbonRendererProperties::StaticClass(); }
+	void SetRendererProperties(UNiagaraRendererProperties *Props) override { Properties = Cast<UNiagaraRibbonRendererProperties>(Props); }
+
+private:
+	class FNiagaraRibbonVertexFactory *VertexFactory;
+	UNiagaraRibbonRendererProperties *Properties;
+	mutable TUniformBuffer<FPrimitiveUniformShaderParameters> WorldSpacePrimitiveUniformBuffer;
+};
+
+
diff --git a/Engine/Plugins/FX/Niagara/Source/Niagara/Public/NiagaraRibbonRendererProperties.h b/Engine/Plugins/FX/Niagara/Source/Niagara/Public/NiagaraRibbonRendererProperties.h
new file mode 100644
--- /dev/null
+++ b/Engine/Plugins/FX/Niagara/Source/Niagara/Public/NiagaraRibbonRendererProperties.h
@@ -0,0 +1,28 @@
+// Copyright 1998-2017 Epic Games, Inc. All Rights Reserved.
+
+#pragma once
+
+#include "CoreMinimal.h"
+#include "UObject/ObjectMacros.h"
+#include "NiagaraRendererProperties.h"
+#include "NiagaraRibbonRendererProperties.generated.h"
+
+UCLASS(editinlinenew)
+class UNiagaraRibbonRendererProperties : public UNiagaraRendererProperties
+{
+public:
+	GENERATED_BODY()
+
+	UPROPERTY(EditAnywhere, Category = "Sprite Rendering")
+	UMaterialInterface* Material;
+
+	//~ UNiagaraRendererProperties interface
+	virtual NiagaraRenderer* CreateEmitterRenderer(ERHIFeatureLevel::Type FeatureLevel) override;
+	virtual void GetUsedMaterials(TArray<UMaterialInterface*>& OutMaterials) const override;
+#if WITH_EDITORONLY_DATA
+	virtual const TArray<FNiagaraVariable>& GetRequiredAttributes() override;
+	virtual const TArray<FNiagaraVariable>& GetOptionalAttributes() override;
+    virtual bool IsMaterialValidForRenderer(UMaterial* Material, FText& InvalidMessage) override;
+	virtual void FixMaterial(UMaterial* Material);
+#endif
+};
