Copying //UE4/Orion-Staging to //UE4/Main (Source: //Orion/Dev-General @ 3606378)

diff --git a/Engine/Plugins/FX/Niagara/Source/Niagara/Public/NiagaraSettings.h b/Engine/Plugins/FX/Niagara/Source/Niagara/Public/NiagaraSettings.h
new file mode 100644
--- /dev/null
+++ b/Engine/Plugins/FX/Niagara/Source/Niagara/Public/NiagaraSettings.h
@@ -0,0 +1,44 @@
+#pragma once
+
+#include "NiagaraSystem.h"
+#include "NiagaraEmitter.h"
+#include "NiagaraScript.h"
+#include "Engine/DeveloperSettings.h"
+#include "InputCoreTypes.h"
+#include "NiagaraSettings.generated.h"
+
+
+UCLASS(config = Engine, defaultconfig, meta=(DisplayName="Niagara"))
+class NIAGARA_API UNiagaraSettings : public UDeveloperSettings
+{
+	GENERATED_UCLASS_BODY()
+
+	UPROPERTY(config, EditAnywhere, Category = Niagara, meta = (AllowedClasses = "ScriptStruct"))
+	TArray<FSoftObjectPath> AdditionalParameterTypes;
+
+	UPROPERTY(config, EditAnywhere, Category = Niagara, meta = (AllowedClasses = "ScriptStruct"))
+	TArray<FSoftObjectPath> AdditionalPayloadTypes;
+
+	UPROPERTY(config, EditAnywhere, Category = Niagara, meta = (AllowedClasses = "Enum"))
+	TArray<FSoftObjectPath> AdditionalParameterEnums;
+	
+	// Begin UDeveloperSettings Interface
+	virtual FName GetCategoryName() const override;
+#if WITH_EDITOR
+	virtual FText GetSectionText() const override;
+#endif
+	// END UDeveloperSettings Interface
+
+#if WITH_EDITOR
+	virtual void PostEditChangeProperty(FPropertyChangedEvent& PropertyChangedEvent) override;
+
+public:
+	DECLARE_MULTICAST_DELEGATE_TwoParams(FOnNiagaraSettingsChanged, const FString&, const UNiagaraSettings*);
+
+	/** Gets a multicast delegate which is called whenever one of the parameters in this settings object changes. */
+	static FOnNiagaraSettingsChanged& OnSettingsChanged();
+
+protected:
+	static FOnNiagaraSettingsChanged SettingsChangedDelegate;
+#endif
+};
\ No newline at end of file
diff --git a/Engine/Plugins/FX/Niagara/Source/Niagara/Public/NiagaraSpriteRendererProperties.h b/Engine/Plugins/FX/Niagara/Source/Niagara/Public/NiagaraSpriteRendererProperties.h
new file mode 100644
--- /dev/null
+++ b/Engine/Plugins/FX/Niagara/Source/Niagara/Public/NiagaraSpriteRendererProperties.h
@@ -0,0 +1,75 @@
+// Copyright 1998-2017 Epic Games, Inc. All Rights Reserved.
+
+#pragma once
+
+#include "CoreMinimal.h"
+#include "UObject/ObjectMacros.h"
+#include "NiagaraRendererProperties.h"
+#include "NiagaraSpriteRendererProperties.generated.h"
+
+
+UENUM()
+enum class ENiagaraSortMode : uint8
+{
+	SortNone,
+	SortViewDepth,
+	SortViewDistance
+};
+
+UENUM()
+enum class ENiagaraSpriteAlignment : uint8
+{
+	Unaligned,
+	VelocityAligned,
+	CustomAlignment
+};
+
+
+UENUM()
+enum class ENiagaraSpriteFacingMode : uint8
+{
+	FaceCamera,
+	FaceCameraPlane,
+	CustomFacingVector
+};
+
+
+UCLASS(editinlinenew)
+class UNiagaraSpriteRendererProperties : public UNiagaraRendererProperties
+{
+public:
+	GENERATED_BODY()
+
+	UNiagaraSpriteRendererProperties();
+
+	//~ UNiagaraRendererProperties interface
+	virtual NiagaraRenderer* CreateEmitterRenderer(ERHIFeatureLevel::Type FeatureLevel) override;
+	virtual void GetUsedMaterials(TArray<UMaterialInterface*>& OutMaterials) const override;
+#if WITH_EDITORONLY_DATA
+	virtual bool IsMaterialValidForRenderer(UMaterial* Material, FText& InvalidMessage) override;
+	virtual void FixMaterial(UMaterial* Material) override;
+	virtual const TArray<FNiagaraVariable>& GetRequiredAttributes() override;
+	virtual const TArray<FNiagaraVariable>& GetOptionalAttributes() override;
+#endif // WITH_EDITORONLY_DATA
+
+	UPROPERTY(EditAnywhere, Category = "Sprite Rendering")
+	UMaterialInterface* Material;
+
+	UPROPERTY(EditAnywhere, Category = "Sprite Rendering")
+	FVector2D SubImageSize;
+
+	UPROPERTY(EditAnywhere, Category = "Sprite Rendering")
+	ENiagaraSpriteAlignment Alignment;
+
+	UPROPERTY(EditAnywhere, Category = "Sprite Rendering")
+	ENiagaraSpriteFacingMode FacingMode;
+
+	UPROPERTY(EditAnywhere, Category = "Sprite Rendering")
+	FVector CustomFacingVectorMask;
+
+	UPROPERTY(EditAnywhere, Category = "Sprite Rendering")
+	ENiagaraSortMode SortMode;
+};
+
+
+
diff --git a/Engine/Plugins/FX/Niagara/Source/Niagara/Public/NiagaraSystemInstance.h b/Engine/Plugins/FX/Niagara/Source/Niagara/Public/NiagaraSystemInstance.h
new file mode 100644
--- /dev/null
+++ b/Engine/Plugins/FX/Niagara/Source/Niagara/Public/NiagaraSystemInstance.h
@@ -0,0 +1,230 @@
+// Copyright 1998-2016 Epic Games, Inc. All Rights Reserved.
+
+#pragma once
+
+#include "NiagaraEmitterInstance.h"
+#include "NiagaraEmitter.h"
+#include "NiagaraEmitterHandle.h"
+#include "NiagaraSystem.h"
+#include "NiagaraParameterBindingInstance.h"
+#include "NiagaraDataInterfaceBindingInstance.h"
+#include "UniquePtr.h"
+#include "NiagaraCommon.h"
+#include "NiagaraDataInterface.h"
+
+class FNiagaraWorldManager;
+class UNiagaraComponent;
+
+class NIAGARA_API FNiagaraSystemInstance 
+{
+public:
+	DECLARE_MULTICAST_DELEGATE(FOnInitialized);
+	
+#if WITH_EDITOR
+	DECLARE_MULTICAST_DELEGATE(FOnReset);
+	DECLARE_MULTICAST_DELEGATE(FOnDestroyed);
+#endif
+
+public:
+	/** Creates a new niagara System instance with the supplied component. */
+	explicit FNiagaraSystemInstance(UNiagaraComponent* InComponent);
+
+	/** Cleanup*/
+	virtual ~FNiagaraSystemInstance();
+
+	/** Initializes this System instance to simulate the supplied System. */
+	void Init(class FNiagaraSystemSimulation* InSystemSimulation, bool bForceReset = false, bool bInForceSolo=false);
+
+	void Activate(bool bReset = false);
+	void Deactivate(bool bImmediate = false);
+
+	//void RebindParameterCollection(UNiagaraParameterCollectionInstance* OldInstance, UNiagaraParameterCollectionInstance* NewInstance);
+	void BindParameters();
+	void UnbindParameters();
+
+	FORCEINLINE FNiagaraParameterStore& GetInstanceParameters() { return InstanceParameters; }
+	
+	FNiagaraWorldManager* GetWorldManager()const;
+
+	/** Defines modes for resetting the System instance. */
+	enum class EResetMode
+	{
+		/** Defers resetting the System instance and simulations until the next tick. */
+		DeferredReset,
+		/** Resets the System instance and simulations immediately. */
+		ImmediateReset,
+		/** same as above, but reinitializes instead of fast resetting */
+		DeferredReInit,
+		ImmediateReInit
+	};
+	
+	/** Requests the the simulation be reset on the next tick. */
+	void Reset(EResetMode Mode);
+
+	void ComponentTick(float DeltaSeconds);
+	void PreSimulateTick(float DeltaSeconds);
+	void PostSimulateTick(float DeltaSeconds);
+	void HandleResets();
+
+	void Enable();
+	void Disable();
+	
+	ENiagaraExecutionState GetExecutionState() { return ExecutionState; }
+	void SetExecutionState(ENiagaraExecutionState InState);
+
+	/** Gets the simulation for the supplied emitter handle. */
+	TSharedPtr<FNiagaraEmitterInstance> GetSimulationForHandle(const FNiagaraEmitterHandle& EmitterHandle);
+
+	UNiagaraSystem* GetSystem()const;
+	FORCEINLINE UNiagaraComponent *GetComponent() { return Component; }
+	FORCEINLINE TArray<TSharedRef<FNiagaraEmitterInstance> > &GetEmitters()	{ return Emitters; }
+	FORCEINLINE FBox &GetSystemBounds()	{ return SystemBounds;  }
+
+	FORCEINLINE bool IsSolo()const { return bSolo; }
+
+	//TEMPORARY. We wont have a single set of parameters when we're executing system scripts.
+	//System params will be pulled in from a data set.
+	FORCEINLINE FNiagaraParameterStore& GetParameters() { return InstanceParameters; }
+
+	/** Gets a data set either from another emitter or one owned by the System itself. */
+	FNiagaraDataSet* GetDataSet(FNiagaraDataSetID SetID, FName EmitterName = NAME_None);
+
+	/** Gets a multicast delegate which is called whenever this instance is initialized with an System asset. */
+	FOnInitialized& OnInitialized();
+
+#if WITH_EDITOR
+	/** Gets a multicast delegate which is called whenever this instance is reset due to external changes in the source System asset. */
+	FOnReset& OnReset();
+
+	FOnDestroyed& OnDestroyed();
+#endif
+
+	FName GetIDName() { return IDName; }
+
+	/** Queue up the data sources to have PrepareForSimulation called on them next tick.*/
+	void ReinitializeDataInterfaces();
+
+	/** Returns the instance data for a particular interface for this System. */
+	FORCEINLINE void* FindDataInterfaceInstanceData(UNiagaraDataInterface* Interface) 
+	{
+		if (int32* InstDataOffset = DataInterfaceInstanceDataOffsets.Find(Interface))
+		{
+			return &DataInterfaceInstanceData[*InstDataOffset];
+		}
+		return nullptr;
+	}
+
+	void DestroyDataInterfaceInstanceData();
+
+	bool UsesEmitter(const UNiagaraEmitter* Emitter)const;
+	bool UsesScript(const UNiagaraScript* Script)const;
+	//bool UsesDataInterface(UNiagaraDataInterface* Interface);
+	bool UsesCollection(const UNiagaraParameterCollection* Collection)const;
+
+	FORCEINLINE bool IsPendingSpawn()const { return bPendingSpawn; }
+	FORCEINLINE void SetPendingSpawn(bool bInValue) { bPendingSpawn = bInValue; }
+
+	FORCEINLINE float GetAge()const { return Age; }
+	
+	FORCEINLINE FNiagaraSystemSimulation* GetSystemSimulation()const
+	{
+		check(SystemSimulation);
+		return SystemSimulation; 
+	}
+
+	/** Index of this instance in the system simulation. */
+	int32 SystemInstanceIndex;
+
+private:
+	/** Builds the emitter simulations. */
+	void InitEmitters();
+
+	/** Resets the System, emitter simulations, and renderers to initial conditions. */
+	void ReInitInternal();
+	/** Resets for restart, assumes no change in emitter setup */
+	void ResetInternal();
+
+	/** Updates the renders for the simulations. Gathers both the EmitterRenderers that were previously set as well as the ones that we  create within.*/
+	void UpdateRenderModules(ERHIFeatureLevel::Type InFeatureLevel, TArray<NiagaraRenderer*>& OutNewRenderers, TArray<NiagaraRenderer*>& OutOldRenderers);
+
+	/** Updates the scene proxy for the System with the specified EmitterRenderer array. Note that this is pushed onto the rendering thread behind the scenes.*/
+	void UpdateProxy(TArray<NiagaraRenderer*>& InRenderers);
+
+	/** Call PrepareForSImulation on each data source from the simulations and determine which need per-tick updates.*/
+	void InitDataInterfaces();
+	
+	/** Perform per-tick updates on data interfaces that need it.*/
+	void TickDataInterfaces(float DeltaSeconds);
+
+	void TickInstanceParameters(float DeltaSeconds);
+
+	void BindParameterCollections(FNiagaraScriptExecutionContext& ExecContext);
+	
+	UNiagaraComponent* Component;
+	class FNiagaraSystemSimulation *SystemSimulation;
+	FBox SystemBounds;
+
+	/** The age of the System instance. */
+	float Age;
+
+	TMap<FNiagaraDataSetID, FNiagaraDataSet> ExternalEvents;
+
+	TArray< TSharedRef<FNiagaraEmitterInstance> > Emitters;
+
+	FOnInitialized OnInitializedDelegate;
+
+#if WITH_EDITOR
+	FOnReset OnResetDelegate;
+	FOnDestroyed OnDestroyedDelegate;
+#endif
+
+	FGuid ID;
+	FName IDName;
+	
+	/** Per instance data for any data interfaces requiring it. */
+	TArray<uint8> DataInterfaceInstanceData;
+
+	/** Map of data interfaces to their instance data. */
+	TMap<TWeakObjectPtr<UNiagaraDataInterface>, int32> DataInterfaceInstanceDataOffsets;
+
+	/** Per system instance parameters. These can be fed by the component and are placed into a dataset for execution for the system scripts. */
+	FNiagaraParameterStore InstanceParameters;
+	
+	FNiagaraParameterDirectBinding<FVector> OwnerPositionParam;
+	FNiagaraParameterDirectBinding<FVector> OwnerVelocityParam;
+	FNiagaraParameterDirectBinding<FVector> OwnerXAxisParam;
+	FNiagaraParameterDirectBinding<FVector> OwnerYAxisParam;
+	FNiagaraParameterDirectBinding<FVector> OwnerZAxisParam;
+	FNiagaraParameterDirectBinding<FMatrix> OwnerTransformParam;
+	FNiagaraParameterDirectBinding<FMatrix> OwnerInverseParam;
+	FNiagaraParameterDirectBinding<FMatrix> OwnerTransposeParam;
+	FNiagaraParameterDirectBinding<FMatrix> OwnerInverseTransposeParam;
+	FNiagaraParameterDirectBinding<float> OwnerDeltaSecondsParam;
+	FNiagaraParameterDirectBinding<float> OwnerInverseDeltaSecondsParam;
+	FNiagaraParameterDirectBinding<float> OwnerMinDistanceToCameraParam;
+	FNiagaraParameterDirectBinding<int32> SystemNumEmittersParam;
+	FNiagaraParameterDirectBinding<int32> SystemNumEmittersAliveParam;
+
+	TArray<FNiagaraParameterDirectBinding<int32>> ParameterNumParticleBindings;
+
+	/** Indicates whether this instance must update itself rather than being batched up as most instances are. */
+	uint32 bSolo : 1;
+	uint32 bForceSolo : 1;
+
+	uint32 bPendingSpawn : 1;
+
+	uint32 bActive : 1;
+
+	/** Flag to ensure the System instance is only reset once per frame. */
+	uint32 bResetPending : 1;
+	uint32 bReInitPending : 1;
+
+	/** Disable ticking and rendering if there was some serious error. */
+	uint32 bError : 1;
+
+	/** Notifier that data interfaces need reinitialization next tick.*/
+	uint32 bDataInterfacesNeedInit : 1;
+
+	/* System tick state */
+	ENiagaraExecutionState ExecutionState;
+};
diff --git a/Engine/Plugins/FX/Niagara/Source/Niagara/Public/NiagaraSystemSimulation.h b/Engine/Plugins/FX/Niagara/Source/Niagara/Public/NiagaraSystemSimulation.h
new file mode 100644
--- /dev/null
+++ b/Engine/Plugins/FX/Niagara/Source/Niagara/Public/NiagaraSystemSimulation.h
@@ -0,0 +1,189 @@
+// Copyright 1998-2017 Epic Games, Inc. All Rights Reserved.
+
+#pragma once
+
+#include "Engine/World.h"
+#include "NiagaraParameterCollection.h"
+#include "GCObject.h"
+#include "NiagaraDataSet.h"
+#include "NiagaraScriptExecutionContext.h"
+
+class UWorld;
+class UNiagaraParameterCollection;
+class UNiagaraParameterCollectionInstance;
+
+//TODO: Pull all the layout information here, in the data set and in parameter stores out into a single layout structure that's shared between all instances of it.
+//Currently there's tons of extra data and work done setting these up.
+struct FNiagaraParameterStoreToDataSetBinding
+{
+	//Array of floats offsets
+	struct FDataOffsets
+	{
+		/** Offset of this value in the parameter store. */
+		int32 ParameterOffset;
+		/** Offset of this value in the data set's components. */
+		int32 DataSetComponentOffset;
+		FDataOffsets(int32 InParamOffset, int32 InDataSetComponentOffset) : ParameterOffset(InParamOffset), DataSetComponentOffset(InDataSetComponentOffset) {}
+	};
+	TArray<FDataOffsets> FloatOffsets;
+	TArray<FDataOffsets> Int32Offsets;
+
+	void Empty()
+	{
+		FloatOffsets.Empty();
+		Int32Offsets.Empty();
+	}
+
+	void Init(FNiagaraDataSet& DataSet, const FNiagaraParameterStore& ParameterStore)
+	{
+		//For now, until I get time to refactor all the layout info into something more coherent we'll init like this and just have to assume the correct layout sets and stores are used later.
+		//Can check it but it'd be v slow.
+
+		for (const FNiagaraVariable& Var : DataSet.GetVariables())
+		{
+			const FNiagaraVariableLayoutInfo* Layout = DataSet.GetVariableLayout(Var);
+			const int32* ParameterOffsetPtr = ParameterStore.ParameterOffsets.Find(Var);
+			if (ParameterOffsetPtr && Layout)
+			{
+				int32 ParameterOffset = *ParameterOffsetPtr;
+				for (uint32 CompIdx = 0; CompIdx < Layout->GetNumFloatComponents(); ++CompIdx)
+				{
+					int32 ParamOffset = ParameterOffset + Layout->LayoutInfo.FloatComponentByteOffsets[CompIdx];
+					int32 DataSetOffset = Layout->FloatComponentStart + Layout->LayoutInfo.FloatComponentRegisterOffsets[CompIdx];
+					FloatOffsets.Emplace(ParamOffset, DataSetOffset);
+				}
+				for (uint32 CompIdx = 0; CompIdx < Layout->GetNumInt32Components(); ++CompIdx)
+				{
+					int32 ParamOffset = ParameterOffset + Layout->LayoutInfo.Int32ComponentByteOffsets[CompIdx];
+					int32 DataSetOffset = Layout->Int32ComponentStart + Layout->LayoutInfo.Int32ComponentRegisterOffsets[CompIdx];
+					Int32Offsets.Emplace(ParamOffset, DataSetOffset);
+				}
+			}
+		}
+	}
+
+	FORCEINLINE_DEBUGGABLE void DataSetToParameterStore(FNiagaraParameterStore& ParameterStore, FNiagaraDataSet& DataSet, int32 DataSetInstanceIndex)
+	{
+		FNiagaraDataBuffer& CurrBuffer = DataSet.CurrData();
+		uint8* ParameterData = ParameterStore.GetParameterDataArray().GetData();
+
+		for (const FDataOffsets& DataOffsets : FloatOffsets)
+		{
+			float* ParamPtr = (float*)(ParameterData + DataOffsets.ParameterOffset);
+			float* DataSetPtr = CurrBuffer.GetInstancePtrFloat(DataOffsets.DataSetComponentOffset, DataSetInstanceIndex);
+			*ParamPtr = *DataSetPtr;
+		}
+		for (const FDataOffsets& DataOffsets : Int32Offsets)
+		{
+			int32* ParamPtr = (int32*)(ParameterData + DataOffsets.ParameterOffset);
+			int32* DataSetPtr = CurrBuffer.GetInstancePtrInt32(DataOffsets.DataSetComponentOffset, DataSetInstanceIndex);
+			*ParamPtr = *DataSetPtr;
+		}
+	}
+
+	FORCEINLINE_DEBUGGABLE void ParameterStoreToDataSet(FNiagaraParameterStore& ParameterStore, FNiagaraDataSet& DataSet, int32 DataSetInstanceIndex)
+	{
+		FNiagaraDataBuffer& CurrBuffer = DataSet.CurrData();
+		uint8* ParameterData = ParameterStore.GetParameterDataArray().GetData();
+
+		for (const FDataOffsets& DataOffsets : FloatOffsets)
+		{
+			float* ParamPtr = (float*)(ParameterData + DataOffsets.ParameterOffset);
+			float* DataSetPtr = CurrBuffer.GetInstancePtrFloat(DataOffsets.DataSetComponentOffset, DataSetInstanceIndex);
+			*DataSetPtr = *ParamPtr;
+		}
+		for (const FDataOffsets& DataOffsets : Int32Offsets)
+		{
+			int32* ParamPtr = (int32*)(ParameterData + DataOffsets.ParameterOffset);
+			int32* DataSetPtr = CurrBuffer.GetInstancePtrInt32(DataOffsets.DataSetComponentOffset, DataSetInstanceIndex);
+			*DataSetPtr = *ParamPtr;
+		}
+	}
+};
+
+/** Simulation performing all system and emitter scripts for a instances of a UNiagaraSystem in a world. */
+class FNiagaraSystemSimulation
+{
+public:
+	bool Init(UNiagaraSystem* InSystem, UWorld* InWorld);
+	void Destroy();
+	void Tick(float DeltaSeconds);
+
+	void RemoveInstance(FNiagaraSystemInstance* Instance);
+	void AddInstance(FNiagaraSystemInstance* Instance);
+
+	void ResetSolo(FNiagaraSystemInstance* Instance);
+	void TickSolo(FNiagaraSystemInstance* Instance);
+	void RemoveSolo(FNiagaraSystemInstance* Instance);
+
+	void TickSoloDataSet();
+
+	FORCEINLINE UNiagaraSystem* GetSystem()const { return System; }
+protected:
+
+	/** System of instances being simulated. No need for GC knowledge as all simulations will be cleaned up by the world manager if the system is invalid. */
+	UNiagaraSystem* System;
+
+	/** World this system simulation belongs to. */
+	UWorld* World;
+
+	/** Main dataset containing system instance attribute data. */
+	FNiagaraDataSet DataSet;
+
+	/**
+	Dataset for system instances doing solo simulation i.e. not batched like most systems.
+	This can be required if strict ordering is needed or the system/emitter scripts use a data interface overridden by the component.
+	*/
+	FNiagaraDataSet DataSetSolo;
+
+	/**
+	As there's a 1 to 1 relationship between system instance and their execution in this simulation we must pull all that instances parameters into a dataset for simulation.
+	In some cases this might be a big waste of memory as there'll be duplicated data from a parameter store that's shared across all instances.
+	Though all these parameters can be unique per instance so for now lets just do the simple thing and improve later.
+	*/
+	FNiagaraDataSet SpawnParameterDataSet;
+	FNiagaraDataSet UpdateParameterDataSet;
+
+	FNiagaraScriptExecutionContext SpawnExecContext;
+	FNiagaraScriptExecutionContext UpdateExecContext;
+
+	FNiagaraScriptExecutionContext SpawnExecContextSolo;
+	FNiagaraScriptExecutionContext UpdateExecContextSolo;
+
+	/** Bindings that pull per component parameters into the spawn parameter dataset. */
+	FNiagaraParameterStoreToDataSetBinding SpawnParameterToDataSetBinding;
+	/** Bindings that pull per component parameters into the update parameter dataset. */
+	FNiagaraParameterStoreToDataSetBinding UpdateParameterToDataSetBinding;
+	
+	/** Binding to push system attributes into each emitter spawn parameters. */
+	TArray<FNiagaraParameterStoreToDataSetBinding> DataSetToEmitterSpawnParameters;
+	/** Binding to push system attributes into each emitter update parameters. */
+	TArray<FNiagaraParameterStoreToDataSetBinding> DataSetToEmitterUpdateParameters;
+	/** Binding to push system attributes into each emitter event parameters. */
+	TArray<TArray<FNiagaraParameterStoreToDataSetBinding>> DataSetToEmitterEventParameters;
+
+	/** System instances that have been spawned and are now simulating. */
+	TArray<FNiagaraSystemInstance*> SystemInstances;
+	/** System instances that are pending to be spawned. */
+	TArray<FNiagaraSystemInstance*> PendingSystemInstances;
+
+	TArray<FNiagaraSystemInstance*> SoloSystemInstances;
+
+	void InitBindings(FNiagaraSystemInstance* SystemInst);
+
+	FNiagaraDataSetAccessor<FNiagaraBool> SystemEnabledAccessor;
+	FNiagaraDataSetAccessor<int32> SystemExecutionStateAccessor;
+	TArray<FNiagaraDataSetAccessor<FNiagaraBool>> EmitterEnabledAccessors;
+	TArray<TArray<FNiagaraDataSetAccessor<FNiagaraSpawnInfo>>> EmitterSpawnInfoAccessors;
+	TArray<FNiagaraDataSetAccessor<int32>> EmitterExecutionStateAccessors;
+
+	//Annoying duplicates required because these access the solo data set. When I rework the where the layout data for parameters and data sets live then these can go away.
+	FNiagaraDataSetAccessor<FNiagaraBool> SoloSystemEnabledAccessor;
+	FNiagaraDataSetAccessor<int32> SoloSystemExecutionStateAccessor;
+	TArray<FNiagaraDataSetAccessor<FNiagaraBool>> SoloEmitterEnabledAccessors;
+	TArray<TArray<FNiagaraDataSetAccessor<FNiagaraSpawnInfo>>> SoloEmitterSpawnInfoAccessors;
+	TArray<FNiagaraDataSetAccessor<int32>> SoloEmitterExecutionStateAccessors;
+
+	uint32 bCanExecute : 1;
+	uint32 bCanExecuteSolo : 1;
+};
diff --git a/Engine/Plugins/FX/Niagara/Source/Niagara/Public/NiagaraTypes.h b/Engine/Plugins/FX/Niagara/Source/Niagara/Public/NiagaraTypes.h
new file mode 100644
--- /dev/null
+++ b/Engine/Plugins/FX/Niagara/Source/Niagara/Public/NiagaraTypes.h
@@ -0,0 +1,700 @@
+// Copyright 1998-2016 Epic Games, Inc. All Rights Reserved.
+
+#pragma once
+
+#include "CoreMinimal.h"
+#include "UObject/UnrealType.h"
+#include "Engine/UserDefinedStruct.h"
+#include "NiagaraTypes.generated.h"
+
+DECLARE_LOG_CATEGORY_EXTERN(LogNiagara, Log, Verbose);
+
+// basic type struct definitions
+
+USTRUCT(meta = (DisplayName = "float"))
+struct FNiagaraFloat
+{
+	GENERATED_USTRUCT_BODY()
+
+	UPROPERTY(EditAnywhere, Category=Parameters)//Parameters? These are used for attrs too.
+	float Value;
+};
+
+USTRUCT(meta = (DisplayName = "int32"))
+struct FNiagaraInt32
+{
+	GENERATED_USTRUCT_BODY()
+
+	UPROPERTY(EditAnywhere, Category = Parameters)//Parameters? These are used for attrs too.
+	int32 Value;
+};
+
+USTRUCT(meta=(DisplayName="bool"))
+struct FNiagaraBool
+{
+	GENERATED_USTRUCT_BODY()
+
+		// The Niagara VM expects this bitmask for its compare and select operators for false.
+	enum BoolValues { 
+		True = INDEX_NONE,
+		False = 0
+	}; 
+
+	
+
+	void SetValue(bool bValue) { Value = bValue ? True : False; }
+	bool GetValue() const { return Value != False; }
+
+	bool IsValid() const { return Value == True || Value == False; }
+	
+	FNiagaraBool():Value(False) {}
+	FNiagaraBool(bool bInValue) : Value(bInValue ? True : False) {}
+	FORCEINLINE operator bool() { return GetValue(); }
+
+	UPROPERTY(EditAnywhere, Category = Parameters)//Parameters? These are used for attrs too. Must be either FNiagaraBool::True or FNiagaraBool::False.
+	int32 Value;
+};
+
+USTRUCT()
+struct FNiagaraNumeric
+{
+	GENERATED_USTRUCT_BODY()
+};
+
+
+USTRUCT()
+struct FNiagaraParameterMap
+{
+	GENERATED_USTRUCT_BODY()
+};
+
+USTRUCT()
+struct FNiagaraTestStructInner
+{
+	GENERATED_USTRUCT_BODY()
+
+	UPROPERTY(EditAnywhere, Category = TestStruct)
+	FVector InnerVector1;
+
+	UPROPERTY(EditAnywhere, Category = TestStruct)
+	FVector InnerVector2;
+};
+
+USTRUCT()
+struct FNiagaraTestStruct
+{
+	GENERATED_USTRUCT_BODY()
+
+	UPROPERTY(EditAnywhere, Category = TestStruct)
+	FVector Vector1;
+
+	UPROPERTY(EditAnywhere, Category = TestStruct)
+	FVector Vector2;
+
+	UPROPERTY(EditAnywhere, Category = TestStruct)
+	FNiagaraTestStructInner InnerStruct1;
+
+	UPROPERTY(EditAnywhere, Category = TestStruct)
+	FNiagaraTestStructInner InnerStruct2;
+};
+
+USTRUCT(meta = (DisplayName = "Matrix"))
+struct FNiagaraMatrix
+{
+	GENERATED_USTRUCT_BODY()
+
+	UPROPERTY(EditAnywhere, Category=NiagaraMatrix)
+	FVector4 Row0;
+
+	UPROPERTY(EditAnywhere, Category = NiagaraMatrix)
+	FVector4 Row1;
+
+	UPROPERTY(EditAnywhere, Category = NiagaraMatrix)
+	FVector4 Row2;
+
+	UPROPERTY(EditAnywhere, Category = NiagaraMatrix)
+	FVector4 Row3;
+};
+
+/** Data controlling the spawning of particles */
+USTRUCT(meta = (DisplayName = "Spawn Info"))
+struct FNiagaraSpawnInfo
+{
+	GENERATED_USTRUCT_BODY();
+	
+	FNiagaraSpawnInfo()
+		: Count(0)
+		, InterpStartDt(0.0f)
+		, IntervalDt(1.0f)
+	{}
+
+	/** How many particles to spawn. */
+	UPROPERTY(EditAnywhere, Category = SpawnInfo)
+	int32 Count;
+	/** The sub frame delta time at which to spawn the first particle. */
+	UPROPERTY(EditAnywhere, Category = SpawnInfo)
+	float InterpStartDt;
+	/** The sub frame delta time between each particle. */
+	UPROPERTY(EditAnywhere, Category = SpawnInfo)
+	float IntervalDt;
+};
+
+/** Information about how this type should be laid out in an FNiagaraDataSet */
+struct FNiagaraTypeLayoutInfo
+{
+	FNiagaraTypeLayoutInfo()
+	{}
+
+	/** Byte offset of each float component in a structured layout. */
+	TArray<uint32> FloatComponentByteOffsets;
+	/** Offset into register table for each float component. */
+	TArray<uint32> FloatComponentRegisterOffsets;
+
+	/** Byte offset of each int32 component in a structured layout. */
+	TArray<uint32> Int32ComponentByteOffsets;
+	/** Offset into register table for each int32 component. */
+	TArray<uint32> Int32ComponentRegisterOffsets;
+
+	FORCEINLINE uint32 GetNumComponents()const { return FloatComponentByteOffsets.Num() + Int32ComponentByteOffsets.Num(); }
+
+	static void GenerateLayoutInfo(FNiagaraTypeLayoutInfo& Layout, const UScriptStruct* Struct)
+	{
+		Layout.FloatComponentByteOffsets.Empty();
+		Layout.FloatComponentRegisterOffsets.Empty();
+		Layout.Int32ComponentByteOffsets.Empty();
+		Layout.Int32ComponentRegisterOffsets.Empty();
+		GenerateLayoutInfoInternal(Layout, Struct);
+	}
+
+private:
+	static void GenerateLayoutInfoInternal(FNiagaraTypeLayoutInfo& Layout, const UScriptStruct* Struct)
+	{
+		for (TFieldIterator<UProperty> PropertyIt(Struct, EFieldIteratorFlags::IncludeSuper); PropertyIt; ++PropertyIt)
+		{
+			UProperty* Property = *PropertyIt;
+			if (Property->IsA(UFloatProperty::StaticClass()))
+			{
+				Layout.FloatComponentRegisterOffsets.Add(Layout.GetNumComponents());
+				Layout.FloatComponentByteOffsets.Add(Property->GetOffset_ForInternal());
+			}
+			else if (Property->IsA(UIntProperty::StaticClass()) || Property->IsA(UBoolProperty::StaticClass()))
+			{
+				Layout.Int32ComponentRegisterOffsets.Add(Layout.GetNumComponents());
+				Layout.Int32ComponentByteOffsets.Add(Property->GetOffset_ForInternal());
+			}
+			//Should be able to support double easily enough
+			else if (UStructProperty* StructProp = CastChecked<UStructProperty>(Property))
+			{
+				GenerateLayoutInfoInternal(Layout, StructProp->Struct);
+			}
+			else
+			{
+				check(false);
+			}
+		}
+	}
+};
+
+/*
+*  Can convert a UStruct with fields of base types only (float, int... - will likely add native vector types here as well)
+*	to an FNiagaraTypeDefinition (internal representation)
+*/
+class NIAGARA_API FNiagaraTypeHelper
+{
+public:
+	static FString ToString(const uint8* ValueData, const UScriptStruct* Struct);
+};
+
+/** Defines different modes for selecting the output numeric type of a function or operation based on the types of the inputs. */
+UENUM()
+enum class ENiagaraNumericOutputTypeSelectionMode : uint8
+{
+	/** Output type selection not supported. */
+	None UMETA(Hidden),
+	/** Select the largest of the numeric inputs. */
+	Largest,
+	/** Select the smallest of the numeric inputs. */
+	Smallest,
+	/** Selects the base scalar type for this numeric inputs. */
+	Scalar,
+};
+
+UENUM()
+enum class ENiagaraExecutionState : uint32
+{
+	/**  Run all scripts. Allow spawning.*/
+	Active,
+	/** Run all scripts but suppress any new spawning.*/
+	Inactive,
+	/** Scripts still run but all ticking of emitters is halted.*/
+	Paused,
+	/** Kill immediately.*/
+	Kill,
+	/** Dead*/
+	Dead,
+	/** Internal use only*/
+	Disabled UMETA(Hidden),
+};
+
+USTRUCT()
+struct NIAGARA_API FNiagaraTypeDefinition
+{
+	GENERATED_USTRUCT_BODY()
+public:
+
+	// Construct blank raw type definition 
+	FNiagaraTypeDefinition(const UClass *ClassDef)
+		: Struct(ClassDef), Enum(nullptr)
+	{
+		check(Struct != nullptr);
+	}
+
+	FNiagaraTypeDefinition(const UEnum *EnumDef)
+		: Struct(IntStruct), Enum(EnumDef)
+	{
+		check(Struct != nullptr);
+	}
+
+	FNiagaraTypeDefinition(const UScriptStruct *StructDef)
+		: Struct(StructDef), Enum(nullptr)
+	{
+		check(Struct != nullptr);
+	}
+
+	FNiagaraTypeDefinition(const FNiagaraTypeDefinition &Other)
+		: Struct(Other.Struct), Enum(Other.Enum)
+	{
+	}
+
+	// Construct a blank raw type definition
+	FNiagaraTypeDefinition()
+		: Struct(nullptr), Enum(nullptr)
+	{}
+
+	bool operator !=(const FNiagaraTypeDefinition &Other) const
+	{
+		return !(*this == Other);
+	}
+
+	bool operator == (const FNiagaraTypeDefinition &Other) const
+	{
+		return Struct == Other.Struct && Enum == Other.Enum;
+	}
+
+	FText GetNameText()const
+	{
+		checkf(IsValid(), TEXT("Type definition is not valid."));
+		if (Enum != nullptr)
+		{
+			return  FText::FromString(Enum->GetName());
+		}
+
+#if WITH_EDITOR
+		return GetStruct()->GetDisplayNameText();
+#else
+		return FText::FromString( GetStruct()->GetName() );
+#endif
+	}
+
+	FString GetName()const
+	{
+		checkf(IsValid(), TEXT("Type definition is not valid."));
+		if (Enum != nullptr)
+		{
+			return  Enum->GetName();
+		}
+		return GetStruct()->GetName();
+	}
+
+	const UStruct* GetStruct()const
+	{
+		return Struct;
+	}
+
+	const UScriptStruct* GetScriptStruct()const
+	{
+		return Cast<UScriptStruct>(Struct);
+	}
+		
+	/** Gets the class ptr for this type if it is a class. */
+	const UClass* GetClass()const
+	{
+		return Cast<UClass>(Struct);
+	}
+
+	const UEnum* GetEnum() const
+	{
+		return Enum;
+	}
+
+	bool IsDataInterface()const;
+
+	bool IsEnum() const { return Enum != nullptr; }
+	
+	int32 GetSize()const
+	{
+		checkf(IsValid(), TEXT("Type definition is not valid."));
+		if (GetClass())
+		{
+			return 0;//TODO: sizeof(void*);//If we're a class then we allocate space for the user to instantiate it. This and stopping it being GCd is up to the user.
+		}
+		else
+		{
+			return CastChecked<UScriptStruct>(Struct)->GetStructureSize();	// TODO: cache this here?
+		}
+	}
+
+	int32 GetAlignment()const
+	{
+		checkf(IsValid(), TEXT("Type definition is not valid."));
+		if (GetClass())
+		{
+			return 0;//TODO: sizeof(void*);//If we're a class then we allocate space for the user to instantiate it. This and stopping it being GCd is up to the user.
+		}
+		else
+		{
+			return CastChecked<UScriptStruct>(Struct)->GetMinAlignment();
+		}
+	}
+
+	bool IsValid() const 
+	{ 
+		return Struct != nullptr;
+	}
+
+	/**
+	UStruct specifying the type for this variable.
+	For most types this will be a UScriptStruct pointing to a something like the struct for an FVector etc.
+	In occasional situations this may be a UClass when we're dealing with DataInterface etc.
+	*/
+	UPROPERTY()
+	const UStruct *Struct;
+
+	UPROPERTY()
+	const UEnum* Enum;
+
+	static void Init();
+	static void RecreateUserDefinedTypeRegistry();
+	static const FNiagaraTypeDefinition& GetFloatDef() { return FloatDef; }
+	static const FNiagaraTypeDefinition& GetBoolDef() { return BoolDef; }
+	static const FNiagaraTypeDefinition& GetIntDef() { return IntDef; }
+	static const FNiagaraTypeDefinition& GetVec2Def() { return Vec2Def; }
+	static const FNiagaraTypeDefinition& GetVec3Def() { return Vec3Def; }
+	static const FNiagaraTypeDefinition& GetVec4Def() { return Vec4Def; }
+	static const FNiagaraTypeDefinition& GetColorDef() { return ColorDef; }
+	static const FNiagaraTypeDefinition& GetMatrix4Def() { return Matrix4Def; }
+	static const FNiagaraTypeDefinition& GetGenericNumericDef() { return NumericDef; }
+	static const FNiagaraTypeDefinition& GetParameterMapDef() { return ParameterMapDef; }
+
+	static const UScriptStruct* GetFloatStruct() { return FloatStruct; }
+	static const UScriptStruct* GetBoolStruct() { return BoolStruct; }
+	static const UScriptStruct* GetIntStruct() { return IntStruct; }
+	static const UScriptStruct* GetVec2Struct() { return Vec2Struct; }
+	static const UScriptStruct* GetVec3Struct() { return Vec3Struct; }
+	static const UScriptStruct* GetVec4Struct() { return Vec4Struct; }
+	static const UScriptStruct* GetColorStruct() { return ColorStruct; }
+	static const UScriptStruct* GetMatrix4Struct() { return Matrix4Struct; }
+	static const UScriptStruct* GetGenericNumericStruct() { return NumericStruct; }
+	static const UScriptStruct* GetParameterMapStruct() { return ParameterMapStruct; }
+
+	static const UEnum* GetExecutionStateEnum() { return ExecutionStateEnum; }
+
+	static const FNiagaraTypeDefinition& GetCollisionEventDef() { return CollisionEventDef; }
+
+	static bool IsScalarDefinition(const FNiagaraTypeDefinition& Type);
+
+	FString ToString(const uint8* ValueData)const
+	{
+		checkf(IsValid(), TEXT("Type definition is not valid."));
+		return FNiagaraTypeHelper::ToString(ValueData, CastChecked<UScriptStruct>(Struct));
+	}
+
+	static bool TypesAreAssignable(const FNiagaraTypeDefinition& TypeA, const FNiagaraTypeDefinition& TypeB);
+	static bool IsLossyConversion(const FNiagaraTypeDefinition& TypeA, const FNiagaraTypeDefinition& TypeB);
+	static FNiagaraTypeDefinition GetNumericOutputType(const TArray<FNiagaraTypeDefinition> TypeDefinintions, ENiagaraNumericOutputTypeSelectionMode SelectionMode);
+
+	static const TArray<FNiagaraTypeDefinition>& GetNumericTypes() { return OrderedNumericTypes; }
+	static bool IsValidNumericInput(const FNiagaraTypeDefinition& TypeDef);
+private:
+
+	static FNiagaraTypeDefinition FloatDef;
+	static FNiagaraTypeDefinition BoolDef;
+	static FNiagaraTypeDefinition IntDef;
+	static FNiagaraTypeDefinition Vec2Def;
+	static FNiagaraTypeDefinition Vec3Def;
+	static FNiagaraTypeDefinition Vec4Def;
+	static FNiagaraTypeDefinition ColorDef;
+	static FNiagaraTypeDefinition Matrix4Def;
+	static FNiagaraTypeDefinition NumericDef;
+	static FNiagaraTypeDefinition ParameterMapDef;
+
+	static const UScriptStruct* FloatStruct;
+	static const UScriptStruct* BoolStruct;
+	static const UScriptStruct* IntStruct;
+	static const UScriptStruct* Vec2Struct;
+	static const UScriptStruct* Vec3Struct;
+	static const UScriptStruct* Vec4Struct;
+	static const UScriptStruct* ColorStruct;
+	static const UScriptStruct* Matrix4Struct;
+	static const UScriptStruct* NumericStruct;
+
+	static const UEnum* ExecutionStateEnum;
+
+	static const UScriptStruct* ParameterMapStruct;
+
+	static TSet<const UScriptStruct*> NumericStructs;
+	static TArray<FNiagaraTypeDefinition> OrderedNumericTypes;
+
+	static TSet<const UScriptStruct*> ScalarStructs;
+
+	static TSet<const UStruct*> FloatStructs;
+	static TSet<const UStruct*> IntStructs;
+	static TSet<const UStruct*> BoolStructs;
+
+	static FNiagaraTypeDefinition CollisionEventDef;
+
+};
+
+/* Contains all types currently available for use in Niagara
+* Used by UI to provide selection; new uniforms and variables
+* may be instanced using the types provided here
+*/
+class NIAGARA_API FNiagaraTypeRegistry
+{
+public:
+	static const TArray<FNiagaraTypeDefinition> &GetRegisteredTypes()
+	{
+		return RegisteredTypes;
+	}
+
+	static const TArray<FNiagaraTypeDefinition> &GetRegisteredParameterTypes()
+	{
+		return RegisteredParamTypes;
+	}
+
+	static const TArray<FNiagaraTypeDefinition> &GetRegisteredPayloadTypes()
+	{
+		return RegisteredPayloadTypes;
+	}
+
+	static const TArray<FNiagaraTypeDefinition>& GetUserDefinedTypes()
+	{
+		return RegisteredUserDefinedTypes;
+	}
+
+	static const TArray<FNiagaraTypeDefinition>& GetNumericTypes()
+	{ 
+		return RegisteredNumericTypes;
+	}
+
+	static void ClearUserDefinedRegistry()
+	{
+		for (const FNiagaraTypeDefinition& Def : RegisteredUserDefinedTypes)
+		{
+			RegisteredTypes.Remove(Def);
+			RegisteredPayloadTypes.Remove(Def);
+			RegisteredParamTypes.Remove(Def);
+		}
+
+		RegisteredNumericTypes.Empty();
+		RegisteredUserDefinedTypes.Empty();
+	}
+
+	static void Register(const FNiagaraTypeDefinition &NewType, bool bCanBeParameter, bool bCanBePayload, bool bIsUserDefined)
+	{
+		//TODO: Make this a map of type to a more verbose set of metadata? Such as the hlsl defs, offset table for conversions etc.
+		RegisteredTypes.AddUnique(NewType);
+
+		if (bCanBeParameter)
+		{
+			RegisteredParamTypes.AddUnique(NewType);
+		}
+
+		if (bCanBePayload)
+		{
+			RegisteredPayloadTypes.AddUnique(NewType);
+		}
+
+		if (bIsUserDefined)
+		{
+			RegisteredUserDefinedTypes.AddUnique(NewType);
+		}
+
+		if (FNiagaraTypeDefinition::IsValidNumericInput(NewType))
+		{
+			RegisteredNumericTypes.AddUnique(NewType);
+		}
+	}
+
+private:
+	static TArray<FNiagaraTypeDefinition> RegisteredTypes;
+	static TArray<FNiagaraTypeDefinition> RegisteredParamTypes;
+	static TArray<FNiagaraTypeDefinition> RegisteredPayloadTypes;
+	static TArray<FNiagaraTypeDefinition> RegisteredUserDefinedTypes;
+	static TArray<FNiagaraTypeDefinition> RegisteredNumericTypes;
+};
+
+FORCEINLINE uint32 GetTypeHash(const FNiagaraTypeDefinition& Type)
+{
+	return HashCombine(GetTypeHash(Type.GetStruct()), GetTypeHash(Type.GetEnum()));
+}
+
+//////////////////////////////////////////////////////////////////////////
+
+USTRUCT()
+struct FNiagaraVariable
+{
+	GENERATED_USTRUCT_BODY()
+
+	FNiagaraVariable()
+		: Name(NAME_None)
+		, TypeDef(FNiagaraTypeDefinition::GetVec4Def())
+	{
+	}
+
+	FNiagaraVariable(const FNiagaraVariable &Other)
+		: Name(Other.Name)
+		, TypeDef(Other.TypeDef)
+	{
+		if (Other.IsDataAllocated())
+		{
+			SetData(Other.GetData());
+		}
+	}
+
+	FNiagaraVariable(FNiagaraTypeDefinition InType, FName InName)
+		: Name(InName)
+		, TypeDef(InType)
+	{
+	}
+	
+	bool operator==(const FNiagaraVariable& Other)const
+	{
+		return Name == Other.Name && TypeDef == Other.TypeDef;
+	}
+
+	bool operator!=(const FNiagaraVariable& Other)const
+	{
+		return !(*this == Other);
+	}
+
+	/** Variables are the same name but if types are auto-assignable, allow them to match. */
+	bool IsEquivalent(const FNiagaraVariable& Other, bool bAllowAssignableTypes = true)const
+	{
+		return Name == Other.Name && (TypeDef == Other.TypeDef || (bAllowAssignableTypes && FNiagaraTypeDefinition::TypesAreAssignable(TypeDef, Other.TypeDef)));
+	}
+	
+	void SetName(FName InName) { Name = InName; }
+	FName GetName()const { return Name; }
+
+	void SetType(const FNiagaraTypeDefinition& InTypeDef) { TypeDef = InTypeDef; }
+	const FNiagaraTypeDefinition& GetType()const { return TypeDef; }
+
+	FORCEINLINE bool IsDataInterface()const { return GetType().IsDataInterface(); }
+
+	void AllocateData()
+	{
+		if (VarData.Num() != TypeDef.GetSize())
+		{
+			VarData.SetNumZeroed(TypeDef.GetSize());
+		}
+	}
+
+	bool IsDataAllocated()const { return VarData.Num() > 0 && VarData.Num() == TypeDef.GetSize(); }
+
+	void CopyTo(uint8* Dest) const
+	{
+		check(TypeDef.GetSize() == VarData.Num());
+		check(IsDataAllocated());
+		FMemory::Memcpy(Dest, VarData.GetData(), VarData.Num());
+	}
+		
+	template<typename T>
+	void SetValue(const T& Data)
+	{
+		check(sizeof(T) == TypeDef.GetSize());
+		AllocateData();
+		FMemory::Memcpy(VarData.GetData(), &Data, VarData.Num());
+	}
+
+	template<typename T>
+	T* GetValue()
+	{
+		check(sizeof(T) == TypeDef.GetSize());
+		check(IsDataAllocated());
+		return (T*)GetData();
+	}
+
+	template<typename T>
+	const T* GetValue() const
+	{
+		check(sizeof(T) == TypeDef.GetSize());
+		check(IsDataAllocated());
+		return (T*)GetData();
+	}
+
+	void SetData(const uint8* Data)
+	{
+		check(Data);
+		AllocateData();
+		FMemory::Memcpy(VarData.GetData(), Data, VarData.Num());
+	}
+
+	const uint8* GetData() const
+	{
+		return VarData.GetData();
+	}
+
+	uint8* GetData()
+	{
+		return VarData.GetData();
+	}
+
+	int32 GetSizeInBytes() const
+	{
+		return TypeDef.GetSize();
+	}
+
+	int32 GetAlignment()const
+	{
+		return TypeDef.GetAlignment();
+	}
+
+	int32 GetAllocatedSizeInBytes() const
+	{
+		return VarData.Num();
+	}
+
+	FString ToString()const
+	{
+		FString Ret = Name.ToString() + TEXT("(");
+		Ret += TypeDef.ToString(VarData.GetData());
+		Ret += TEXT(")");
+		return Ret;
+	}
+
+	bool IsNameValid() const
+	{
+		return Name != NAME_None;
+	}
+
+	FORCEINLINE bool IsInNameSpace(FString Namespace) 
+	{
+		return Name.ToString().StartsWith(Namespace + TEXT("."));
+	}
+
+private:
+	UPROPERTY(EditAnywhere, Category = "Variable")
+	FName Name;
+	UPROPERTY(EditAnywhere, Category = "Variable")
+	FNiagaraTypeDefinition TypeDef;
+	//This gets serialized but do we need to worry about endianness doing things like this? If not, where does that get handled?
+	//TODO: Remove storage here entirely and move everything to an FNiagaraParameterStore.
+	UPROPERTY()
+	TArray<uint8> VarData;
+};
+
+FORCEINLINE uint32 GetTypeHash(const FNiagaraVariable& Var)
+{
+	return HashCombine(GetTypeHash(Var.GetType()), GetTypeHash(Var.GetName()));
+}
diff --git a/Engine/Plugins/FX/Niagara/Source/Niagara/Public/NiagaraWorldManager.h b/Engine/Plugins/FX/Niagara/Source/Niagara/Public/NiagaraWorldManager.h
new file mode 100644
--- /dev/null
+++ b/Engine/Plugins/FX/Niagara/Source/Niagara/Public/NiagaraWorldManager.h
@@ -0,0 +1,50 @@
+// Copyright 1998-2017 Epic Games, Inc. All Rights Reserved.
+
+#pragma once
+
+#include "CoreMinimal.h"
+#include "Engine/World.h"
+#include "NiagaraParameterCollection.h"
+#include "GCObject.h"
+#include "NiagaraDataSet.h"
+#include "NiagaraScriptExecutionContext.h"
+#include "NiagaraSystemSimulation.h"
+
+class UWorld;
+class UNiagaraParameterCollection;
+class UNiagaraParameterCollectionInstance;
+
+/**
+* Manager class for any data relating to a particular world.
+*/
+class FNiagaraWorldManager : public FGCObject
+{
+public:
+	
+	FNiagaraWorldManager(UWorld* InWorld)
+		: World(InWorld)
+		, CachedEffectsQuality(INDEX_NONE)
+	{}
+
+	static FNiagaraWorldManager* Get(UWorld* World);
+
+	//~ GCObject Interface
+	void AddReferencedObjects(FReferenceCollector& Collector);
+	//~ GCObject Interface
+	
+	UNiagaraParameterCollectionInstance* GetParameterCollection(UNiagaraParameterCollection* Collection);
+	void SetParameterCollection(UNiagaraParameterCollectionInstance* NewInstance);
+	FNiagaraSystemSimulation* GetSystemSimulation(UNiagaraSystem* System);
+	void DestroySystemSimulation(UNiagaraSystem* System);
+
+	void Tick(float DeltaSeconds);
+private:
+	UWorld* World;
+
+	TMap<UNiagaraParameterCollection*, UNiagaraParameterCollectionInstance*> ParameterCollections;
+
+	TMap<UNiagaraSystem*, FNiagaraSystemSimulation> SystemSimulations;
+
+	int32 CachedEffectsQuality;
+};
+
diff --git a/Engine/Plugins/FX/Niagara/Source/NiagaraEditor/NiagaraEditor.Build.cs b/Engine/Plugins/FX/Niagara/Source/NiagaraEditor/NiagaraEditor.Build.cs
new file mode 100644
--- /dev/null
+++ b/Engine/Plugins/FX/Niagara/Source/NiagaraEditor/NiagaraEditor.Build.cs
@@ -0,0 +1,82 @@
+// Copyright 1998-2017 Epic Games, Inc. All Rights Reserved.
+
+using UnrealBuildTool;
+
+public class NiagaraEditor : ModuleRules
+{
+	public NiagaraEditor(ReadOnlyTargetRules Target) : base(Target)
+	{
+		PrivateIncludePaths.AddRange(new string[] {
+			"NiagaraEditor/Private",
+			"NiagaraEditor/Private/Toolkits",
+			"NiagaraEditor/Private/Widgets",
+			"NiagaraEditor/Private/Sequencer/NiagaraSequence",
+			"NiagaraEditor/Private/Sequencer/LevelSequence",
+			"NiagaraEditor/Private/ViewModels",
+			"NiagaraEditor/Private/TypeEditorUtilities"
+		});
+
+		PrivateDependencyModuleNames.AddRange(
+			new string[] {
+				"Engine",
+                "RHI",
+                "Core", 
+				"CoreUObject", 
+                "InputCore",
+				"RenderCore",
+				"Slate", 
+				"SlateCore",
+				"Kismet",
+                "EditorStyle",
+				"UnrealEd", 
+				"VectorVM",
+                "Niagara",
+                "NiagaraShader",
+                "MovieScene",
+				"Sequencer",
+				"PropertyEditor",
+				"GraphEditor",
+                "ShaderFormatVectorVM",
+                "TargetPlatform",
+                "AppFramework",
+				"MovieSceneTools",
+                "MovieSceneTracks",
+                "AdvancedPreviewScene",
+            }
+        );
+
+		PrivateIncludePathModuleNames.AddRange(
+			new string[] {
+				"Engine",
+				"Messaging",
+				"LevelEditor",
+				"AssetTools",
+				"ContentBrowser",
+                "DerivedDataCache",
+                "ShaderCore",
+            }
+        );
+
+		PublicDependencyModuleNames.AddRange(
+            new string[] {
+                "NiagaraShader",
+                "Engine",
+                "Niagara",
+                "UnrealEd",
+            }
+        );
+
+        PublicIncludePathModuleNames.AddRange(
+            new string[] {
+				"Engine",
+				"Niagara"
+			}
+		);
+
+        DynamicallyLoadedModuleNames.AddRange(
+            new string[] {
+                "WorkspaceMenuStructure",
+                }
+            );
+	}
+}
diff --git a/Engine/Plugins/FX/Niagara/Source/NiagaraEditor/Private/AssetTypeActions/AssetTypeActions_NiagaraEmitter.cpp b/Engine/Plugins/FX/Niagara/Source/NiagaraEditor/Private/AssetTypeActions/AssetTypeActions_NiagaraEmitter.cpp
new file mode 100644
--- /dev/null
+++ b/Engine/Plugins/FX/Niagara/Source/NiagaraEditor/Private/AssetTypeActions/AssetTypeActions_NiagaraEmitter.cpp
@@ -0,0 +1,33 @@
+// Copyright 1998-2016 Epic Games, Inc. All Rights Reserved.
+
+#include "AssetTypeActions_NiagaraEmitter.h"
+#include "NiagaraEmitter.h"
+#include "NiagaraSystemToolkit.h"
+#include "NiagaraEditorStyle.h"
+
+class UNiagaraEmitter;
+
+FColor FAssetTypeActions_NiagaraEmitter::GetTypeColor() const
+{
+	return FNiagaraEditorStyle::Get().GetColor("NiagaraEditor.AssetColors.Emitter").ToFColor(true);
+}
+
+void FAssetTypeActions_NiagaraEmitter::OpenAssetEditor(const TArray<UObject*>& InObjects, TSharedPtr<IToolkitHost> EditWithinLevelEditor)
+{
+	EToolkitMode::Type Mode = EditWithinLevelEditor.IsValid() ? EToolkitMode::WorldCentric : EToolkitMode::Standalone;
+
+	for (auto ObjIt = InObjects.CreateConstIterator(); ObjIt; ++ObjIt)
+	{
+		auto Emitter = Cast<UNiagaraEmitter>(*ObjIt);
+		if (Emitter != nullptr)
+		{
+			TSharedRef<FNiagaraSystemToolkit> SystemToolkit(new FNiagaraSystemToolkit());
+			SystemToolkit->InitializeWithEmitter(Mode, EditWithinLevelEditor, *Emitter);
+		}
+	}
+}
+
+UClass* FAssetTypeActions_NiagaraEmitter::GetSupportedClass() const
+{
+	return UNiagaraEmitter::StaticClass();
+}
diff --git a/Engine/Plugins/FX/Niagara/Source/NiagaraEditor/Private/AssetTypeActions/AssetTypeActions_NiagaraEmitter.h b/Engine/Plugins/FX/Niagara/Source/NiagaraEditor/Private/AssetTypeActions/AssetTypeActions_NiagaraEmitter.h
new file mode 100644
--- /dev/null
+++ b/Engine/Plugins/FX/Niagara/Source/NiagaraEditor/Private/AssetTypeActions/AssetTypeActions_NiagaraEmitter.h
@@ -0,0 +1,18 @@
+// Copyright 1998-2016 Epic Games, Inc. All Rights Reserved.
+
+#pragma once
+
+#include "AssetTypeActions_Base.h"
+#include "NiagaraEditorModule.h"
+
+/** Asset type actions for FNiagaraEmitter. */
+class FAssetTypeActions_NiagaraEmitter: public FAssetTypeActions_Base
+{
+public:
+	// IAssetTypeActions Implementation
+	virtual FText GetName() const override { return NSLOCTEXT("AssetTypeActions", "AssetTypeActions_NiagaraEmitter", "Niagara Emitter"); }
+	virtual FColor GetTypeColor() const override;
+	virtual UClass* GetSupportedClass() const override;
+	virtual void OpenAssetEditor(const TArray<UObject*>& InObjects, TSharedPtr<class IToolkitHost> EditWithinLevelEditor = TSharedPtr<IToolkitHost>()) override;
+	virtual uint32 GetCategories() override { return FNiagaraEditorModule::GetAssetCategory(); }
+};
\ No newline at end of file
diff --git a/Engine/Plugins/FX/Niagara/Source/NiagaraEditor/Private/AssetTypeActions/AssetTypeActions_NiagaraParameterCollection.cpp b/Engine/Plugins/FX/Niagara/Source/NiagaraEditor/Private/AssetTypeActions/AssetTypeActions_NiagaraParameterCollection.cpp
new file mode 100644
--- /dev/null
+++ b/Engine/Plugins/FX/Niagara/Source/NiagaraEditor/Private/AssetTypeActions/AssetTypeActions_NiagaraParameterCollection.cpp
@@ -0,0 +1,139 @@
+// Copyright 1998-2017 Epic Games, Inc. All Rights Reserved.
+
+#include "AssetTypeActions_NiagaraParameterCollection.h"
+#include "NiagaraParameterCollection.h"
+#include "NiagaraParameterCollectionToolkit.h"
+#include "NiagaraParameterCollectionFactoryNew.h"
+#include "NiagaraEditorStyle.h"
+
+#include "IContentBrowserSingleton.h"
+#include "ContentBrowserModule.h"
+
+#include "MultiBoxBuilder.h"
+#include "EditorStyleSet.h"
+
+#define LOCTEXT_NAMESPACE "AssetTypeActions_NiagaraParameterCollection"
+
+FColor FAssetTypeActions_NiagaraParameterCollection::GetTypeColor() const
+{
+	return FNiagaraEditorStyle::Get().GetColor("NiagaraEditor.AssetColors.ParameterCollection").ToFColor(true);
+}
+
+void FAssetTypeActions_NiagaraParameterCollection::OpenAssetEditor(const TArray<UObject*>& InObjects, TSharedPtr<IToolkitHost> EditWithinLevelEditor)
+{
+	EToolkitMode::Type Mode = EditWithinLevelEditor.IsValid() ? EToolkitMode::WorldCentric : EToolkitMode::Standalone;
+
+	for (auto ObjIt = InObjects.CreateConstIterator(); ObjIt; ++ObjIt)
+	{
+		auto NPC = Cast<UNiagaraParameterCollection>(*ObjIt);
+		if (NPC != NULL)
+		{
+			TSharedRef< FNiagaraParameterCollectionToolkit > NewNiagaraNPCToolkit(new FNiagaraParameterCollectionToolkit());
+			NewNiagaraNPCToolkit->Initialize(Mode, EditWithinLevelEditor, NPC);
+		}
+	}
+}
+
+void FAssetTypeActions_NiagaraParameterCollection::GetActions(const TArray<UObject*>& InObjects, FMenuBuilder& MenuBuilder)
+{
+	auto Collections = GetTypedWeakObjectPtrs<UNiagaraParameterCollection>(InObjects);
+
+	MenuBuilder.AddMenuEntry(
+		LOCTEXT("NewNPC", "Create Niagara Parameter Collection Instance"),
+		LOCTEXT("NewNPCTooltip", "Creates an instance of this Niagara Parameter Collection."),
+		FSlateIcon(FEditorStyle::GetStyleSetName(), "ClassIcon.MaterialInstanceActor"),
+		FUIAction(
+			FExecuteAction::CreateSP(this, &FAssetTypeActions_NiagaraParameterCollection::ExecuteNewNPC, Collections)
+		)
+	);
+}
+
+void FAssetTypeActions_NiagaraParameterCollection::ExecuteNewNPC(TArray<TWeakObjectPtr<UNiagaraParameterCollection>> Objects)
+{
+	const FString DefaultSuffix = TEXT("_Inst");
+
+	FContentBrowserModule& ContentBrowserModule = FModuleManager::LoadModuleChecked<FContentBrowserModule>("ContentBrowser");
+	if (Objects.Num() == 1)
+	{
+		auto Object = Objects[0].Get();
+
+		if (Object)
+		{
+			// Create an appropriate and unique name 
+			FString Name;
+			FString PackageName;
+			CreateUniqueAssetName(Object->GetOutermost()->GetName(), DefaultSuffix, PackageName, Name);
+
+			UNiagaraParameterCollectionInstanceFactoryNew* Factory = NewObject<UNiagaraParameterCollectionInstanceFactoryNew>();
+			Factory->InitialParent = Object;
+
+			ContentBrowserModule.Get().CreateNewAsset(Name, FPackageName::GetLongPackagePath(PackageName), UNiagaraParameterCollectionInstance::StaticClass(), Factory);
+		}
+	}
+	else
+	{
+		TArray<UObject*> ObjectsToSync;
+		for (auto ObjIt = Objects.CreateConstIterator(); ObjIt; ++ObjIt)
+		{
+			auto Object = (*ObjIt).Get();
+			if (Object)
+			{
+				// Determine an appropriate name
+				FString Name;
+				FString PackageName;
+				CreateUniqueAssetName(Object->GetOutermost()->GetName(), DefaultSuffix, PackageName, Name);
+
+				// Create the factory used to generate the asset
+				UNiagaraParameterCollectionInstanceFactoryNew* Factory = NewObject<UNiagaraParameterCollectionInstanceFactoryNew>();
+				Factory->InitialParent = Object;
+
+				FAssetToolsModule& AssetToolsModule = FModuleManager::GetModuleChecked<FAssetToolsModule>("AssetTools");
+				UObject* NewAsset = AssetToolsModule.Get().CreateAsset(Name, FPackageName::GetLongPackagePath(PackageName), UNiagaraParameterCollectionInstance::StaticClass(), Factory);
+
+				if (NewAsset)
+				{
+					ObjectsToSync.Add(NewAsset);
+				}
+			}
+		}
+
+		if (ObjectsToSync.Num() > 0)
+		{
+			ContentBrowserModule.Get().SyncBrowserToAssets(ObjectsToSync);
+		}
+	}
+}
+
+UClass* FAssetTypeActions_NiagaraParameterCollection::GetSupportedClass() const
+{
+	return UNiagaraParameterCollection::StaticClass();
+}
+
+//////////////////////////////////////////////////////////////////////////
+
+FColor FAssetTypeActions_NiagaraParameterCollectionInstance::GetTypeColor() const
+{ 
+	return FNiagaraEditorStyle::Get().GetColor("NiagaraEditor.AssetColors.ParameterCollectionInstance").ToFColor(true); 
+}
+
+void FAssetTypeActions_NiagaraParameterCollectionInstance::OpenAssetEditor(const TArray<UObject*>& InObjects, TSharedPtr<IToolkitHost> EditWithinLevelEditor)
+{
+	EToolkitMode::Type Mode = EditWithinLevelEditor.IsValid() ? EToolkitMode::WorldCentric : EToolkitMode::Standalone;
+
+	for (auto ObjIt = InObjects.CreateConstIterator(); ObjIt; ++ObjIt)
+	{
+		auto Instance = Cast<UNiagaraParameterCollectionInstance>(*ObjIt);
+		if (Instance != NULL)
+		{
+			TSharedRef< FNiagaraParameterCollectionToolkit > NewNiagaraNPCToolkit(new FNiagaraParameterCollectionToolkit());
+			NewNiagaraNPCToolkit->Initialize(Mode, EditWithinLevelEditor, Instance);
+		}
+	}
+}
+
+UClass* FAssetTypeActions_NiagaraParameterCollectionInstance::GetSupportedClass() const
+{
+	return UNiagaraParameterCollectionInstance::StaticClass();
+}
+
+#undef LOCTEXT_NAMESPACE
\ No newline at end of file
diff --git a/Engine/Plugins/FX/Niagara/Source/NiagaraEditor/Private/AssetTypeActions/AssetTypeActions_NiagaraParameterCollection.h b/Engine/Plugins/FX/Niagara/Source/NiagaraEditor/Private/AssetTypeActions/AssetTypeActions_NiagaraParameterCollection.h
new file mode 100644
--- /dev/null
+++ b/Engine/Plugins/FX/Niagara/Source/NiagaraEditor/Private/AssetTypeActions/AssetTypeActions_NiagaraParameterCollection.h
@@ -0,0 +1,34 @@
+// Copyright 1998-2017 Epic Games, Inc. All Rights Reserved.
+
+#pragma once
+
+#include "CoreMinimal.h"
+#include "Toolkits/IToolkitHost.h"
+#include "AssetTypeActions_Base.h"
+#include "NiagaraEditorModule.h"
+
+class FAssetTypeActions_NiagaraParameterCollection : public FAssetTypeActions_Base
+{
+public:
+	// IAssetTypeActions Implementation
+	virtual FText GetName() const override { return NSLOCTEXT("AssetTypeActions", "AssetTypeActions_NiagaraParameterCollection", "Niagara Parameter Collection"); }
+	virtual FColor GetTypeColor() const override;
+	virtual UClass* GetSupportedClass() const override;
+	virtual void OpenAssetEditor(const TArray<UObject*>& InObjects, TSharedPtr<class IToolkitHost> EditWithinLevelEditor = TSharedPtr<IToolkitHost>()) override;
+	virtual uint32 GetCategories() override { return FNiagaraEditorModule::GetAssetCategory(); }
+	virtual bool HasActions(const TArray<UObject*>& InObjects) const override { return true; }
+	virtual void GetActions(const TArray<UObject*>& InObjects, FMenuBuilder& MenuBuilder)override;
+
+	void ExecuteNewNPC(TArray<TWeakObjectPtr<class UNiagaraParameterCollection>> Objects);
+};
+
+class FAssetTypeActions_NiagaraParameterCollectionInstance : public FAssetTypeActions_Base
+{
+public:
+	// IAssetTypeActions Implementation
+	virtual FText GetName() const override { return NSLOCTEXT("AssetTypeActions", "AssetTypeActions_NiagaraParameterCollectionInstance", "Niagara Parameter Collection Instance"); }
+	virtual FColor GetTypeColor() const override;
+	virtual UClass* GetSupportedClass() const override;
+	virtual void OpenAssetEditor(const TArray<UObject*>& InObjects, TSharedPtr<class IToolkitHost> EditWithinLevelEditor = TSharedPtr<IToolkitHost>()) override;
+	virtual uint32 GetCategories() override { return FNiagaraEditorModule::GetAssetCategory(); }
+};
\ No newline at end of file
diff --git a/Engine/Plugins/FX/Niagara/Source/NiagaraEditor/Private/AssetTypeActions/AssetTypeActions_NiagaraScript.cpp b/Engine/Plugins/FX/Niagara/Source/NiagaraEditor/Private/AssetTypeActions/AssetTypeActions_NiagaraScript.cpp
new file mode 100644
--- /dev/null
+++ b/Engine/Plugins/FX/Niagara/Source/NiagaraEditor/Private/AssetTypeActions/AssetTypeActions_NiagaraScript.cpp
@@ -0,0 +1,31 @@
+// Copyright 1998-2017 Epic Games, Inc. All Rights Reserved.
+
+#include "AssetTypeActions/AssetTypeActions_NiagaraScript.h"
+#include "NiagaraScript.h"
+#include "NiagaraScriptToolkit.h"
+#include "NiagaraEditorStyle.h"
+
+FColor FAssetTypeActions_NiagaraScript::GetTypeColor() const
+{ 
+	return FNiagaraEditorStyle::Get().GetColor("NiagaraEditor.AssetColors.Script").ToFColor(true); 
+}
+
+void FAssetTypeActions_NiagaraScript::OpenAssetEditor( const TArray<UObject*>& InObjects, TSharedPtr<IToolkitHost> EditWithinLevelEditor )
+{
+	EToolkitMode::Type Mode = EditWithinLevelEditor.IsValid() ? EToolkitMode::WorldCentric : EToolkitMode::Standalone;
+
+	for (auto ObjIt = InObjects.CreateConstIterator(); ObjIt; ++ObjIt)
+	{
+		auto Script = Cast<UNiagaraScript>(*ObjIt);
+		if (Script != NULL)
+		{
+			TSharedRef< FNiagaraScriptToolkit > NewNiagaraScriptToolkit(new FNiagaraScriptToolkit());
+			NewNiagaraScriptToolkit->Initialize(Mode, EditWithinLevelEditor, Script);
+		}
+	}
+}
+
+UClass* FAssetTypeActions_NiagaraScript::GetSupportedClass() const
+{ 
+	return UNiagaraScript::StaticClass(); 
+}
diff --git a/Engine/Plugins/FX/Niagara/Source/NiagaraEditor/Private/AssetTypeActions/AssetTypeActions_NiagaraScript.h b/Engine/Plugins/FX/Niagara/Source/NiagaraEditor/Private/AssetTypeActions/AssetTypeActions_NiagaraScript.h
new file mode 100644
--- /dev/null
+++ b/Engine/Plugins/FX/Niagara/Source/NiagaraEditor/Private/AssetTypeActions/AssetTypeActions_NiagaraScript.h
@@ -0,0 +1,19 @@
+// Copyright 1998-2017 Epic Games, Inc. All Rights Reserved.
+
+#pragma once
+
+#include "CoreMinimal.h"
+#include "Toolkits/IToolkitHost.h"
+#include "AssetTypeActions_Base.h"
+#include "NiagaraEditorModule.h"
+
+class FAssetTypeActions_NiagaraScript : public FAssetTypeActions_Base
+{
+public:
+	// IAssetTypeActions Implementation
+	virtual FText GetName() const override { return NSLOCTEXT("AssetTypeActions", "AssetTypeActions_NiagaraScript", "Niagara Script"); }
+	virtual FColor GetTypeColor() const override;
+	virtual UClass* GetSupportedClass() const override;
+	virtual void OpenAssetEditor(const TArray<UObject*>& InObjects, TSharedPtr<class IToolkitHost> EditWithinLevelEditor = TSharedPtr<IToolkitHost>()) override;
+	virtual uint32 GetCategories() override { return FNiagaraEditorModule::GetAssetCategory(); }
+};
diff --git a/Engine/Plugins/FX/Niagara/Source/NiagaraEditor/Private/AssetTypeActions/AssetTypeActions_NiagaraSystem.cpp b/Engine/Plugins/FX/Niagara/Source/NiagaraEditor/Private/AssetTypeActions/AssetTypeActions_NiagaraSystem.cpp
new file mode 100644
--- /dev/null
+++ b/Engine/Plugins/FX/Niagara/Source/NiagaraEditor/Private/AssetTypeActions/AssetTypeActions_NiagaraSystem.cpp
@@ -0,0 +1,31 @@
+// Copyright 1998-2017 Epic Games, Inc. All Rights Reserved.
+
+#include "AssetTypeActions/AssetTypeActions_NiagaraSystem.h"
+#include "NiagaraSystem.h"
+#include "NiagaraSystemToolkit.h"
+#include "NiagaraEditorStyle.h"
+
+FColor FAssetTypeActions_NiagaraSystem::GetTypeColor() const
+{
+	return FNiagaraEditorStyle::Get().GetColor("NiagaraEditor.AssetColors.System").ToFColor(true);
+}
+
+void FAssetTypeActions_NiagaraSystem::OpenAssetEditor(const TArray<UObject*>& InObjects, TSharedPtr<IToolkitHost> EditWithinLevelEditor)
+{
+	EToolkitMode::Type Mode = EditWithinLevelEditor.IsValid() ? EToolkitMode::WorldCentric : EToolkitMode::Standalone;
+
+	for (auto ObjIt = InObjects.CreateConstIterator(); ObjIt; ++ObjIt)
+	{
+		auto System = Cast<UNiagaraSystem>(*ObjIt);
+		if (System != NULL)
+		{
+			TSharedRef< FNiagaraSystemToolkit > NewNiagaraSystemToolkit(new FNiagaraSystemToolkit());
+			NewNiagaraSystemToolkit->InitializeWithSystem(Mode, EditWithinLevelEditor, *System);
+		}
+	}
+}
+
+UClass* FAssetTypeActions_NiagaraSystem::GetSupportedClass() const
+{
+	return UNiagaraSystem::StaticClass();
+}
diff --git a/Engine/Plugins/FX/Niagara/Source/NiagaraEditor/Private/AssetTypeActions/AssetTypeActions_NiagaraSystem.h b/Engine/Plugins/FX/Niagara/Source/NiagaraEditor/Private/AssetTypeActions/AssetTypeActions_NiagaraSystem.h
new file mode 100644
--- /dev/null
+++ b/Engine/Plugins/FX/Niagara/Source/NiagaraEditor/Private/AssetTypeActions/AssetTypeActions_NiagaraSystem.h
@@ -0,0 +1,19 @@
+// Copyright 1998-2017 Epic Games, Inc. All Rights Reserved.
+
+#pragma once
+
+#include "CoreMinimal.h"
+#include "Toolkits/IToolkitHost.h"
+#include "AssetTypeActions_Base.h"
+#include "NiagaraEditorModule.h"
+
+class FAssetTypeActions_NiagaraSystem: public FAssetTypeActions_Base
+{
+public:
+	// IAssetTypeActions Implementation
+	virtual FText GetName() const override { return NSLOCTEXT("AssetTypeActions", "AssetTypeActions_NiagaraSystem", "Niagara System"); }
+	virtual FColor GetTypeColor() const override;
+	virtual UClass* GetSupportedClass() const override;
+	virtual void OpenAssetEditor(const TArray<UObject*>& InObjects, TSharedPtr<class IToolkitHost> EditWithinLevelEditor = TSharedPtr<IToolkitHost>()) override;
+	virtual uint32 GetCategories() override { return FNiagaraEditorModule::GetAssetCategory(); }
+};
diff --git a/Engine/Plugins/FX/Niagara/Source/NiagaraEditor/Private/Customizations/NiagaraComponentDetails.cpp b/Engine/Plugins/FX/Niagara/Source/NiagaraEditor/Private/Customizations/NiagaraComponentDetails.cpp
new file mode 100644
--- /dev/null
+++ b/Engine/Plugins/FX/Niagara/Source/NiagaraEditor/Private/Customizations/NiagaraComponentDetails.cpp
@@ -0,0 +1,502 @@
+// Copyright 1998-2016 Epic Games, Inc. All Rights Reserved.
+
+#include "NiagaraComponentDetails.h"
+#include "PropertyHandle.h"
+#include "DetailLayoutBuilder.h"
+#include "DetailCategoryBuilder.h"
+#include "NiagaraComponent.h"
+#include "NiagaraDataInterface.h"
+#include "DetailLayoutBuilder.h"
+#include "IDetailPropertyRow.h"
+#include "DetailWidgetRow.h"
+#include "SInlineEditableTextBlock.h"
+#include "Widgets/SToolTip.h"
+#include "PropertyCustomizationHelpers.h"
+#include "Widgets/Text/STextBlock.h"
+#include "Widgets/Images/SImage.h"
+#include "Widgets/Layout/SBox.h"
+#include "EditorStyleSet.h"
+#include "NiagaraSystemInstance.h"
+#include "Private/ViewModels/NiagaraParameterViewModel.h"
+#include "NiagaraEditorUtilities.h"
+#include "INiagaraEditorTypeUtilities.h"
+#include "NiagaraTypes.h"
+#include "ScopedTransaction.h"
+#include "SNiagaraParameterEditor.h"
+#include "NiagaraEditorModule.h"
+#include "PropertyEditorModule.h"
+#include "ModuleManager.h"
+#include "IStructureDetailsView.h"
+#include "IDetailsView.h"
+#include "Widgets/Input/SButton.h"
+#include "Widgets/Input/SEditableText.h"
+#include "Editor.h"
+#include "NiagaraParameterCollectionViewModel.h"
+#include "NiagaraScriptSource.h"
+#include "NiagaraNodeInput.h"
+#include "NiagaraGraph.h"
+#include "GameDelegates.h"
+#include "NiagaraEditorStyle.h"
+#include "IDetailChildrenBuilder.h"
+#include "MultiBoxBuilder.h"
+#include "SComboButton.h"
+#include "SImage.h"
+#include "NiagaraEditorModule.h"
+#include "ModuleManager.h"
+#include "INiagaraEditorTypeUtilities.h"
+#include "SBox.h"
+#define LOCTEXT_NAMESPACE "NiagaraComponentDetails"
+
+class FNiagaraComponentNodeBuilder : public IDetailCustomNodeBuilder
+{
+public:
+	FNiagaraComponentNodeBuilder(UNiagaraComponent* InComponent, UNiagaraScript* SourceSpawn, UNiagaraScript* SourceUpdate)						   
+	{
+		Component = InComponent;
+		OriginalScripts.Add(SourceSpawn);
+		OriginalScripts.Add(SourceUpdate);
+	}
+
+	~FNiagaraComponentNodeBuilder()
+	{
+	}
+
+	virtual void SetOnRebuildChildren(FSimpleDelegate InOnRegenerateChildren) override
+	{
+		OnRebuildChildren = InOnRegenerateChildren;
+	}
+
+	virtual void GenerateHeaderRowContent(FDetailWidgetRow& NodeRow) {}
+	virtual void Tick(float DeltaTime) override {}
+	virtual bool RequiresTick() const override { return false; }
+	virtual bool InitiallyCollapsed() const { return false; }
+	virtual FName GetName() const  override
+	{
+		static const FName NiagaraComponentNodeBuilder("FNiagaraComponentNodeBuilder");
+		return NiagaraComponentNodeBuilder;
+	}
+
+	virtual void GenerateChildContent(IDetailChildrenBuilder& ChildrenBuilder) override
+	{
+		TArray<FNiagaraVariable> Parameters;
+		FNiagaraParameterStore& ParamStore = Component->GetInitialParameters();
+		ParamStore.GetParameters(Parameters);
+
+		FNiagaraEditorModule& NiagaraEditorModule = FModuleManager::GetModuleChecked<FNiagaraEditorModule>("NiagaraEditor");
+		
+		for (const FNiagaraVariable& Parameter : Parameters)
+		{
+			TSharedPtr<SWidget> NameWidget;
+
+			NameWidget =
+				SNew(STextBlock)
+				.TextStyle(FNiagaraEditorStyle::Get(), "NiagaraEditor.ParameterText")
+				.Text(FText::FromName(Parameter.GetName()));
+			
+			IDetailPropertyRow* Row = nullptr;
+
+			TSharedPtr<SWidget> CustomValueWidget;
+
+			if (!Parameter.IsDataInterface())
+			{
+				TSharedPtr<FStructOnScope> StructOnScope = MakeShareable(new FStructOnScope(Parameter.GetType().GetStruct(), (uint8*)ParamStore.GetParameterData(Parameter)));
+				Row = ChildrenBuilder.AddExternalStructureProperty(StructOnScope.ToSharedRef(), NAME_None, Parameter.GetName());
+
+			}
+			else 
+			{
+				int32 DataInterfaceOffset = ParamStore.IndexOf(Parameter);
+				UObject* DefaultValueObject = ParamStore.DataInterfaces[DataInterfaceOffset];
+
+				TArray<UObject*> Objects;
+				Objects.Add(DefaultValueObject);
+				Row = ChildrenBuilder.AddExternalObjects(Objects, Parameter.GetName());
+
+				CustomValueWidget =
+					SNew(STextBlock)
+					.TextStyle(FNiagaraEditorStyle::Get(), "NiagaraEditor.ParameterText")
+					.Text(FText::FromString(FName::NameToDisplayString(DefaultValueObject->GetClass()->GetName(), false)));
+			}
+
+			if (Row)
+			{
+				TSharedPtr<SWidget> DefaultNameWidget;
+				TSharedPtr<SWidget> DefaultValueWidget;
+
+				TSharedPtr<IPropertyHandle> PropertyHandle = Row->GetPropertyHandle();
+
+				FDetailWidgetRow& CustomWidget = Row->CustomWidget(true);
+
+				TArray<UObject*> Objects;
+				Row->GetPropertyHandle()->GetOuterObjects(Objects);
+
+				Row->GetDefaultWidgets(DefaultNameWidget, DefaultValueWidget, CustomWidget);
+
+				if (Parameter.IsDataInterface())
+				{
+					PropertyHandle->SetOnPropertyValuePreChange(
+						FSimpleDelegate::CreateRaw(this, &FNiagaraComponentNodeBuilder::OnDataInterfacePreChange, Parameter));
+					PropertyHandle->SetOnChildPropertyValuePreChange(
+						FSimpleDelegate::CreateRaw(this, &FNiagaraComponentNodeBuilder::OnDataInterfacePreChange, Parameter));
+					PropertyHandle->SetOnPropertyValueChanged(
+						FSimpleDelegate::CreateRaw(this, &FNiagaraComponentNodeBuilder::OnDataInterfaceChanged, Parameter));
+					PropertyHandle->SetOnChildPropertyValueChanged(
+						FSimpleDelegate::CreateRaw(this, &FNiagaraComponentNodeBuilder::OnDataInterfaceChanged, Parameter));
+				}
+				else
+				{
+					PropertyHandle->SetOnPropertyValuePreChange(
+						FSimpleDelegate::CreateRaw(this, &FNiagaraComponentNodeBuilder::OnParameterPreChange, Parameter));
+					PropertyHandle->SetOnChildPropertyValuePreChange(
+						FSimpleDelegate::CreateRaw(this, &FNiagaraComponentNodeBuilder::OnParameterPreChange, Parameter));
+					PropertyHandle->SetOnPropertyValueChanged(
+						FSimpleDelegate::CreateRaw(this, &FNiagaraComponentNodeBuilder::OnParameterChanged, Parameter));
+					PropertyHandle->SetOnChildPropertyValueChanged(
+						FSimpleDelegate::CreateRaw(this, &FNiagaraComponentNodeBuilder::OnParameterChanged, Parameter));
+				}
+
+				CustomWidget
+					.NameContent()
+					[
+						SNew(SBox)
+						.Padding(FMargin(0.0f, 2.0f))
+						[
+							NameWidget.ToSharedRef()
+						]
+					];
+
+				TSharedPtr<SWidget> ValueWidget = DefaultValueWidget;
+				if (CustomValueWidget.IsValid())
+				{
+					ValueWidget = CustomValueWidget;
+				}
+
+				CustomWidget
+					.ValueContent()
+					[
+						SNew(SHorizontalBox)
+						+ SHorizontalBox::Slot()
+						.HAlign(HAlign_Fill)
+						.Padding(4.0f)
+						[
+							// Add in the parameter editor factoried above.
+							ValueWidget.ToSharedRef()
+						]
+						+ SHorizontalBox::Slot()
+						.VAlign(VAlign_Center)
+						.AutoWidth()
+						[
+							// Add in the "reset to default" buttons
+							SNew(SButton)
+							.OnClicked_Raw(this, &FNiagaraComponentNodeBuilder::OnLocationResetClicked, Parameter)
+							.Visibility_Raw(this, &FNiagaraComponentNodeBuilder::GetLocationResetVisibility, Parameter)
+							.ContentPadding(FMargin(5.f, 0.f))
+							.ToolTipText(LOCTEXT("ResetToDefaultToolTip", "Reset to Default"))
+							.ButtonStyle(FEditorStyle::Get(), "NoBorder")
+							.Content()
+							[
+								SNew(SImage)
+								.Image(FEditorStyle::GetBrush("PropertyWindow.DiffersFromDefault"))
+							]
+						]
+					];
+			}
+		}
+	}
+
+private:
+	void OnCollectionViewModelChanged()
+	{
+		OnRebuildChildren.ExecuteIfBound();
+	}
+
+	void OnParameterPreChange(FNiagaraVariable Var)
+	{
+		Component->Modify();
+	}
+
+	void OnDataInterfacePreChange(FNiagaraVariable Var)
+	{
+		Component->Modify();
+	}
+
+	void OnParameterChanged(FNiagaraVariable Var)
+	{
+		Component->GetInitialParameters().OnParameterChange();
+		Component->SetParameterValueOverriddenLocally(Var.GetName(), true);
+	}
+
+	void OnDataInterfaceChanged(FNiagaraVariable Var)
+	{
+		Component->GetInitialParameters().OnInterfaceChange();
+		Component->SetParameterValueOverriddenLocally(Var.GetName(), true);
+	}
+
+	bool DoesParameterDifferFromDefault(const FNiagaraVariable& Var)
+	{
+		return Component->IsParameterValueOverriddenLocally(Var.GetName());
+	}
+
+	FReply OnLocationResetClicked(FNiagaraVariable Parameter)
+	{
+		FScopedTransaction ScopedTransaction(LOCTEXT("ResetParameterValue", "Reset parameter value to system defaults."));
+		Component->Modify();
+		Component->SetParameterValueOverriddenLocally(Parameter.GetName(), false);
+		return FReply::Handled();
+	}
+
+	EVisibility GetLocationResetVisibility(FNiagaraVariable Parameter) const
+	{
+		return Component->IsParameterValueOverriddenLocally(Parameter.GetName()) ? EVisibility::Visible : EVisibility::Collapsed;
+	}
+
+private:
+	UNiagaraComponent* Component;
+	FSimpleDelegate OnRebuildChildren;
+	TArray<TSharedRef<FStructOnScope>> CreatedStructOnScopes;
+	TArray<UNiagaraScript*> OriginalScripts;
+};
+
+TSharedRef<IDetailCustomization> FNiagaraComponentDetails::MakeInstance()
+{
+	return MakeShareable(new FNiagaraComponentDetails);
+}
+
+FNiagaraComponentDetails::FNiagaraComponentDetails() : Builder(nullptr), bQueueForDetailsRefresh(false)
+{
+	GEditor->RegisterForUndo(this);
+}
+
+
+FNiagaraComponentDetails::~FNiagaraComponentDetails()
+{
+	GEditor->UnregisterForUndo(this);
+
+	if (GEngine)
+	{
+		GEngine->OnWorldDestroyed().RemoveAll(this);
+	}
+
+	FGameDelegates::Get().GetEndPlayMapDelegate().RemoveAll(this);
+
+	if (ObjectsCustomized.Num() == 1)
+	{
+		UNiagaraComponent* Component = Cast<UNiagaraComponent>(ObjectsCustomized[0].Get());
+
+		if (Component)
+		{
+			if (FNiagaraSystemInstance* SystemInstance = Component->GetSystemInstance())
+			{
+				if (OnResetDelegateHandle.IsValid())
+				{
+					SystemInstance->OnReset().Remove(OnResetDelegateHandle);
+				}
+
+				if (OnInitDelegateHandle.IsValid())
+				{
+					SystemInstance->OnInitialized().Remove(OnInitDelegateHandle);
+				}
+
+				SystemInstance->OnDestroyed().RemoveAll(this);
+			}
+		}
+	}
+}
+
+bool FNiagaraComponentDetails::IsTickable() const
+{
+	return bQueueForDetailsRefresh;
+}
+
+void FNiagaraComponentDetails::Tick(float DeltaTime)
+{
+	if (bQueueForDetailsRefresh)
+	{
+		if (Builder)
+		{
+			Builder->ForceRefreshDetails();
+			bQueueForDetailsRefresh = false;
+		}
+	}
+}
+
+void FNiagaraComponentDetails::PostUndo(bool bSuccess)
+{
+	// We may have queued up as a result of an Undo of adding the System itself. The objects we were 
+	// referencing may therefore have been removed. If so, we'll have to take a different path later on in the code.
+	int32 NumValid = 0;
+	for (TWeakObjectPtr<UObject> WeakPtr : ObjectsCustomized)
+	{
+		if (!WeakPtr.IsStale(true) && WeakPtr.IsValid())
+		{
+			NumValid++;
+		}
+	}
+
+	if (Builder)
+	{
+		// This is tricky. There is a high chance that if the component changed, then
+		// any cached FNiagaraVariable that we're holding on to may have been changed out from underneath us.
+		// So we essentially must tear down and start again in the UI.
+		// HOWEVER, a refresh will invoke a new constructor of the FNiagaraComponentDetails, which puts us 
+		// in the queue to handle PostUndo events. This turns quickly into an infinitely loop. Therefore,
+		// we circumvent this by telling the editor that we need to queue up an event
+		// that we will handle in the subsequent frame's Tick event. Not the cleanest approach, 
+		// but because we are doing things like CopyOnWrite, the normal Unreal property editing
+		// path is not available to us.
+		if (NumValid != 0)
+		{
+			bQueueForDetailsRefresh = true;
+		}
+		else
+		{
+			// If we no longer have any valid pointers, but previously had a builder, that means 
+			// that the builder is probably dead or dying soon. We shouldn't trust it any more and
+			// we should make sure that we aren't queueing for ticks to refresh either.
+			Builder = nullptr;
+			bQueueForDetailsRefresh = false;
+		}
+
+#if 0
+		for (TSharedPtr<FNiagaraParameterViewModelCustomDetails>& ViewModel : ViewModels)
+		{
+			ViewModel->Reset();
+		}
+#endif
+	}
+}
+
+TStatId FNiagaraComponentDetails::GetStatId() const
+{
+	RETURN_QUICK_DECLARE_CYCLE_STAT(FNiagaraComponentDetails, STATGROUP_Tickables);
+}
+
+void FNiagaraComponentDetails::OnPiEEnd()
+{
+	UE_LOG(LogNiagaraEditor, Log, TEXT("onPieEnd"));
+	if (ObjectsCustomized.Num() == 1)
+	{
+		UNiagaraComponent* Component = Cast<UNiagaraComponent>(ObjectsCustomized[0].Get());
+
+		if (Component)
+		{
+			if (Component->GetOutermost()->HasAnyPackageFlags(PKG_PlayInEditor))
+			{
+				UE_LOG(LogNiagaraEditor, Log, TEXT("onPieEnd - has package flags"));
+				UWorld* TheWorld = UWorld::FindWorldInPackage(Component->GetOutermost());
+				if (TheWorld)
+				{
+					OnWorldDestroyed(TheWorld);
+				}
+			}
+		}
+	}
+}
+
+void FNiagaraComponentDetails::OnWorldDestroyed(class UWorld* InWorld)
+{
+	// We have to clear out any temp data interfaces that were bound to the component's package when the world goes away or otherwise
+	// we'll report GC leaks..
+	if (ObjectsCustomized.Num() == 1)
+	{
+		UNiagaraComponent* Component = Cast<UNiagaraComponent>(ObjectsCustomized[0].Get());
+
+		if (Component)
+		{
+			if (Component->GetWorld() == InWorld)
+			{
+				UE_LOG(LogNiagaraEditor, Log, TEXT("OnWorldDestroyed - matched up"));
+				Builder = nullptr;
+			}
+		}
+	}
+}
+
+void FNiagaraComponentDetails::CustomizeDetails(IDetailLayoutBuilder& DetailBuilder)
+{
+ 	Builder = &DetailBuilder;
+// 
+// 	// Create a category so this is displayed early in the properties
+// 	IDetailCategoryBuilder& NiagaraCategory = DetailBuilder.EditCategory("Niagara", FText::GetEmpty(), ECategoryPriority::Important);
+// 
+// 	// Store off the properties for analysis in later function calls.
+	static const FName ParamCategoryName = TEXT("NiagaraComponent_Parameters");
+	static const FName ParamCollectionCategoryName = TEXT("NiagaraComponent_ParameterCollectionOverrides");
+	static const FName ScriptCategoryName = TEXT("Parameters");
+
+	DetailBuilder.EditCategory(ScriptCategoryName);
+
+ 	TSharedPtr<IPropertyHandle> LocalOverridesPropertyHandle = DetailBuilder.GetProperty(TEXT("Parameters"));
+ 	if (LocalOverridesPropertyHandle.IsValid())
+ 	{
+ 		LocalOverridesPropertyHandle->MarkHiddenByCustomization();
+		UProperty* Property = LocalOverridesPropertyHandle->GetProperty();
+ 	}
+
+	DetailBuilder.GetObjectsBeingCustomized(ObjectsCustomized);
+
+	if (ObjectsCustomized.Num() == 1)
+	{
+		UNiagaraComponent* Component = Cast<UNiagaraComponent>(ObjectsCustomized[0].Get());
+
+		if (Component == nullptr)
+		{
+			return;
+		}
+
+		if (GEngine)
+		{
+			GEngine->OnWorldDestroyed().AddRaw(this, &FNiagaraComponentDetails::OnWorldDestroyed);
+		}
+
+		FGameDelegates::Get().GetEndPlayMapDelegate().AddRaw(this, &FNiagaraComponentDetails::OnPiEEnd);
+
+		FNiagaraSystemInstance* SystemInstance = Component->GetSystemInstance();
+		if (SystemInstance == nullptr)
+		{
+			return;
+		}
+
+		// We'll want to monitor for any serious changes that might require a rebuild of the UI from scratch.
+		UE_LOG(LogNiagaraEditor, Log, TEXT("Registering with System instance %p"), SystemInstance);
+		SystemInstance->OnReset().RemoveAll(this);
+		SystemInstance->OnInitialized().RemoveAll(this);
+		SystemInstance->OnDestroyed().RemoveAll(this);
+		OnResetDelegateHandle = SystemInstance->OnReset().AddSP(this, &FNiagaraComponentDetails::OnSystemInstanceReset);
+		OnInitDelegateHandle = SystemInstance->OnInitialized().AddSP(this, &FNiagaraComponentDetails::OnSystemInstanceReset);
+		SystemInstance->OnDestroyed().AddSP(this, &FNiagaraComponentDetails::OnSystemInstanceDestroyed);
+		FNiagaraParameterStore* ParamStore = &Component->GetInitialParameters();
+		UNiagaraScript* ScriptSpawn = Component->GetAsset()->GetSystemSpawnScript();
+		UNiagaraScript* ScriptUpdate = Component->GetAsset()->GetSystemUpdateScript();
+		
+		IDetailCategoryBuilder& InputParamCategory = DetailBuilder.EditCategory(ParamCategoryName, LOCTEXT("ParamCategoryName", "Parameters"));
+		//InputParamCategory.HeaderContent(SNew(SAddParameterButton, InputCollectionViewModel));
+		InputParamCategory.AddCustomBuilder(MakeShared<FNiagaraComponentNodeBuilder>(Component, ScriptSpawn, ScriptUpdate));
+	}
+
+}
+
+void FNiagaraComponentDetails::OnSystemInstanceReset()
+{
+	UE_LOG(LogNiagaraEditor, Log, TEXT("OnSystemInstanceReset()"));
+	bQueueForDetailsRefresh = true;
+}
+
+void FNiagaraComponentDetails::OnSystemInstanceDestroyed()
+{
+	UE_LOG(LogNiagaraEditor, Log, TEXT("OnSystemInstanceDestroyed()"));
+	bQueueForDetailsRefresh = true;
+}
+
+FReply FNiagaraComponentDetails::OnSystemOpenRequested(UNiagaraSystem* InSystem)
+{
+	if (InSystem)
+	{
+		FAssetEditorManager::Get().OpenEditorForAsset(InSystem);
+	}
+	return FReply::Handled();
+}
+
+
+
+#undef LOCTEXT_NAMESPACE
diff --git a/Engine/Plugins/FX/Niagara/Source/NiagaraEditor/Private/Customizations/NiagaraComponentDetails.h b/Engine/Plugins/FX/Niagara/Source/NiagaraEditor/Private/Customizations/NiagaraComponentDetails.h
new file mode 100644
--- /dev/null
+++ b/Engine/Plugins/FX/Niagara/Source/NiagaraEditor/Private/Customizations/NiagaraComponentDetails.h
@@ -0,0 +1,56 @@
+// Copyright 1998-2016 Epic Games, Inc. All Rights Reserved.
+
+#pragma once
+
+#include "CoreMinimal.h"
+#include "IDetailCustomization.h"
+#include "EditorUndoClient.h"
+#include "TickableEditorObject.h"
+#include "Reply.h"
+#include "Visibility.h"
+#include "Misc/NotifyHook.h"
+
+class IDetailLayoutBuilder;
+class IPropertyHandle;
+class FNiagaraParameterViewModelCustomDetails;
+class INiagaraParameterViewModel;
+class SNiagaraParameterEditor;
+class FNiagaraParameterViewModelCustomDetails;
+class UNiagaraSystem;
+
+class FNiagaraComponentDetails : public IDetailCustomization, public FEditorUndoClient, public FTickableEditorObject
+{
+public:
+	virtual ~FNiagaraComponentDetails();
+
+	/** Makes a new instance of this detail layout class for a specific detail view requesting it */
+	static TSharedRef<IDetailCustomization> MakeInstance();
+
+	/** IDetailCustomization interface */
+	virtual void CustomizeDetails( IDetailLayoutBuilder& DetailBuilder ) override;
+
+	//~ FEditorUndoClient interface
+	virtual void PostUndo(bool bSuccess) override;
+	virtual void PostRedo(bool bSuccess) override { PostUndo(bSuccess); }
+
+	// FTickableEditorObject
+	virtual bool IsTickable() const override;
+	virtual void Tick(float DeltaTime) override;
+	virtual TStatId GetStatId() const override;
+
+protected:
+	void OnSystemInstanceReset();
+	void OnSystemInstanceDestroyed();
+	FReply OnSystemOpenRequested(UNiagaraSystem* InSystem);
+	void OnWorldDestroyed(class UWorld* InWorld);
+	void OnPiEEnd();
+
+	FNiagaraComponentDetails();
+private:
+	TArray<TWeakObjectPtr<UObject>> ObjectsCustomized;
+	IDetailLayoutBuilder* Builder;
+	bool bQueueForDetailsRefresh;
+
+	FDelegateHandle OnInitDelegateHandle;
+	FDelegateHandle OnResetDelegateHandle;
+};
diff --git a/Engine/Plugins/FX/Niagara/Source/NiagaraEditor/Private/Customizations/NiagaraScriptDetails.cpp b/Engine/Plugins/FX/Niagara/Source/NiagaraEditor/Private/Customizations/NiagaraScriptDetails.cpp
new file mode 100644
--- /dev/null
+++ b/Engine/Plugins/FX/Niagara/Source/NiagaraEditor/Private/Customizations/NiagaraScriptDetails.cpp
@@ -0,0 +1,266 @@
+// Copyright 1998-2016 Epic Games, Inc. All Rights Reserved.
+
+#include "NiagaraScriptDetails.h"
+#include "DetailLayoutBuilder.h"
+#include "DetailCategoryBuilder.h"
+#include "IDetailPropertyRow.h"
+#include "DetailWidgetRow.h"
+#include "SInlineEditableTextBlock.h"
+#include "NiagaraParameterCollectionViewModel.h"
+#include "NiagaraScriptInputCollectionViewModel.h"
+#include "NiagaraScriptOutputCollectionViewModel.h"
+#include "NiagaraScriptViewModel.h"
+#include "NiagaraParameterViewModel.h"
+#include "NiagaraEditorStyle.h"
+#include "IDetailChildrenBuilder.h"
+#include "MultiBoxBuilder.h"
+#include "SComboButton.h"
+#include "SImage.h"
+#include "NiagaraEditorModule.h"
+#include "ModuleManager.h"
+#include "INiagaraEditorTypeUtilities.h"
+#include "SBox.h"
+
+
+#define LOCTEXT_NAMESPACE "NiagaraScriptDetails"
+
+class FNiagaraCustomNodeBuilder : public IDetailCustomNodeBuilder
+{
+public:
+	FNiagaraCustomNodeBuilder(TSharedRef<INiagaraParameterCollectionViewModel> InViewModel)
+		: ViewModel(InViewModel)
+	{
+		ViewModel->OnCollectionChanged().AddRaw(this, &FNiagaraCustomNodeBuilder::OnCollectionViewModelChanged);
+	}
+
+	~FNiagaraCustomNodeBuilder()
+	{
+		ViewModel->OnCollectionChanged().RemoveAll(this);
+	}
+
+	virtual void SetOnRebuildChildren(FSimpleDelegate InOnRegenerateChildren) override
+	{
+		OnRebuildChildren = InOnRegenerateChildren;
+	}
+
+	virtual void GenerateHeaderRowContent(FDetailWidgetRow& NodeRow){}
+	virtual void Tick(float DeltaTime) override {}
+	virtual bool RequiresTick() const override { return false; }
+	virtual bool InitiallyCollapsed() const { return false; }
+	virtual FName GetName() const  override
+	{ 
+		static const FName NiagaraCustomNodeBuilder("NiagaraCustomNodeBuilder");
+		return NiagaraCustomNodeBuilder;
+	}
+
+	virtual void GenerateChildContent(IDetailChildrenBuilder& ChildrenBuilder) override
+	{
+		const TArray<TSharedRef<INiagaraParameterViewModel>>& Parameters = ViewModel->GetParameters();
+
+		FNiagaraEditorModule& NiagaraEditorModule = FModuleManager::GetModuleChecked<FNiagaraEditorModule>("NiagaraEditor");
+
+		for (const TSharedRef<INiagaraParameterViewModel>& Parameter : Parameters)
+		{
+			TSharedPtr<SWidget> NameWidget;
+
+			if (Parameter->CanRenameParameter())
+			{
+				NameWidget =
+					SNew(SInlineEditableTextBlock)
+					.Style(FNiagaraEditorStyle::Get(), "NiagaraEditor.ParameterInlineEditableText")
+					.Text(Parameter, &INiagaraParameterViewModel::GetNameText)
+					.OnVerifyTextChanged(Parameter, &INiagaraParameterViewModel::VerifyNodeNameTextChanged)
+					.OnTextCommitted(Parameter, &INiagaraParameterViewModel::NameTextComitted);
+			}
+			else
+			{
+				NameWidget =
+					SNew(STextBlock)
+					.TextStyle(FNiagaraEditorStyle::Get(), "NiagaraEditor.ParameterText")
+					.Text(Parameter, &INiagaraParameterViewModel::GetNameText);
+			}
+
+			
+
+			IDetailPropertyRow* Row = nullptr;
+
+			TSharedPtr<SWidget> CustomValueWidget;
+
+			if (Parameter->GetDefaultValueType() == INiagaraParameterViewModel::EDefaultValueType::Struct)
+			{
+				Row = ChildrenBuilder.AddExternalStructureProperty(Parameter->GetDefaultValueStruct(), NAME_None, Parameter->GetName());
+				
+			}
+			else if (Parameter->GetDefaultValueType() == INiagaraParameterViewModel::EDefaultValueType::Object)
+			{
+				UObject* DefaultValueObject = Parameter->GetDefaultValueObject();
+
+				TArray<UObject*> Objects;
+				Objects.Add(DefaultValueObject);
+				Row = ChildrenBuilder.AddExternalObjects(Objects, Parameter->GetName());
+
+				CustomValueWidget =
+					SNew(STextBlock)
+					.TextStyle(FNiagaraEditorStyle::Get(), "NiagaraEditor.ParameterText")
+					.Text(FText::FromString(FName::NameToDisplayString(DefaultValueObject->GetClass()->GetName(), false)));
+			}
+
+			if(Row)
+			{
+				TSharedPtr<SWidget> DefaultNameWidget;
+				TSharedPtr<SWidget> DefaultValueWidget;
+
+				TSharedPtr<IPropertyHandle> PropertyHandle = Row->GetPropertyHandle();
+
+				FDetailWidgetRow& CustomWidget = Row->CustomWidget(true);
+
+				Row->GetDefaultWidgets(DefaultNameWidget, DefaultValueWidget, CustomWidget);
+
+				PropertyHandle->SetOnPropertyValueChanged(
+					FSimpleDelegate::CreateSP(Parameter, &INiagaraParameterViewModel::NotifyDefaultValueChanged));
+				PropertyHandle->SetOnChildPropertyValueChanged(
+					FSimpleDelegate::CreateSP(Parameter, &INiagaraParameterViewModel::NotifyDefaultValueChanged));
+
+				CustomWidget
+					.NameContent()
+					[
+						SNew(SBox)
+						.Padding(FMargin(0.0f, 2.0f))
+						[
+							NameWidget.ToSharedRef()
+						]
+					];
+
+				if (CustomValueWidget.IsValid())
+				{
+					CustomWidget.ValueContent()
+					[
+						CustomValueWidget.ToSharedRef()
+					];
+				}
+			}
+
+		}
+
+
+	}
+
+private:
+	void OnCollectionViewModelChanged()
+	{
+		OnRebuildChildren.ExecuteIfBound();
+	}
+
+private:
+	TSharedRef<INiagaraParameterCollectionViewModel> ViewModel;
+	FSimpleDelegate OnRebuildChildren;
+};
+
+class SAddParameterButton : public SCompoundWidget
+{	
+	SLATE_BEGIN_ARGS(SAddParameterButton)
+	{}
+	SLATE_END_ARGS()
+
+	void Construct(const FArguments& InArgs, TSharedPtr<INiagaraParameterCollectionViewModel> InViewModel)
+	{
+		CollectionViewModel = InViewModel;
+
+		ChildSlot
+		.HAlign(HAlign_Right)
+		[
+			SAssignNew(ComboButton, SComboButton)
+			.HasDownArrow(false)
+			.ButtonStyle(FEditorStyle::Get(), "RoundButton")
+			.ForegroundColor(FSlateColor::UseForeground())
+			.OnGetMenuContent(this, &SAddParameterButton::GetAddParameterMenuContent)
+			.Visibility(CollectionViewModel.ToSharedRef(), &INiagaraParameterCollectionViewModel::GetAddButtonVisibility)
+			.ContentPadding(FMargin(2.0f, 1.0f, 0.0f, 1.0f))
+			.HAlign(HAlign_Right)
+			.VAlign(VAlign_Center)
+			.ButtonContent()
+			[
+				SNew(SHorizontalBox)
+				+ SHorizontalBox::Slot()
+				.Padding(0, 1, 2, 1)
+				.AutoWidth()
+				[
+					SNew(SImage)
+					.ColorAndOpacity(FSlateColor::UseForeground())
+					.Image(FEditorStyle::GetBrush("Plus"))
+				]
+				+ SHorizontalBox::Slot()
+				.AutoWidth()
+				.VAlign(VAlign_Center)
+				[
+					SNew(STextBlock)
+					.Font(IDetailLayoutBuilder::GetDetailFontBold())
+					.Text(CollectionViewModel.ToSharedRef(), &INiagaraParameterCollectionViewModel::GetAddButtonText)
+					.Visibility(this, &SAddParameterButton::OnGetAddParameterTextVisibility)
+					.ShadowOffset(FVector2D(1, 1))
+				]
+			]
+		];
+	}
+
+private:
+	EVisibility OnGetAddParameterTextVisibility() const
+	{
+		return IsHovered() || ComboButton->IsOpen() ? EVisibility::Visible : EVisibility::Collapsed;
+	}
+
+	TSharedRef<SWidget> GetAddParameterMenuContent() const
+	{
+		FMenuBuilder AddMenuBuilder(true, nullptr);
+		for (TSharedPtr<FNiagaraTypeDefinition> AvailableType : CollectionViewModel->GetAvailableTypes())
+		{
+			AddMenuBuilder.AddMenuEntry
+			(
+				AvailableType->GetNameText(),
+				FText(),
+				FSlateIcon(),
+				FUIAction(FExecuteAction::CreateSP(CollectionViewModel.ToSharedRef(), &INiagaraParameterCollectionViewModel::AddParameter, AvailableType))
+			);
+		}
+		return AddMenuBuilder.MakeWidget();
+	}
+
+	
+private:
+	TSharedPtr<INiagaraParameterCollectionViewModel> CollectionViewModel;
+	TSharedPtr<SComboButton> ComboButton;
+};
+
+TSharedRef<IDetailCustomization> FNiagaraScriptDetails::MakeInstance(TWeakPtr<FNiagaraScriptViewModel> ScriptViewModel)
+{
+	return MakeShared<FNiagaraScriptDetails>(ScriptViewModel.Pin());
+}
+
+FNiagaraScriptDetails::FNiagaraScriptDetails(TSharedPtr<FNiagaraScriptViewModel> InScriptViewModel)
+	: ScriptViewModel(InScriptViewModel)
+{}
+
+void FNiagaraScriptDetails::CustomizeDetails(IDetailLayoutBuilder& DetailBuilder)
+{
+	static const FName InputParamCategoryName = TEXT("NiagaraScript_InputParams");
+	static const FName OutputParamCategoryName = TEXT("NiagaraScript_OutputParams");
+	static const FName ScriptCategoryName = TEXT("Script");
+
+	DetailBuilder.EditCategory(ScriptCategoryName);
+
+	TSharedRef<INiagaraParameterCollectionViewModel> InputCollectionViewModel = ScriptViewModel->GetInputCollectionViewModel();
+	TSharedRef<INiagaraParameterCollectionViewModel> OutputCollectionViewModel = ScriptViewModel->GetOutputCollectionViewModel();
+
+	IDetailCategoryBuilder& InputParamCategory = DetailBuilder.EditCategory(InputParamCategoryName, LOCTEXT("InputParamCategoryName", "Input Parameters"));
+	InputParamCategory.HeaderContent(SNew(SAddParameterButton, InputCollectionViewModel));
+	InputParamCategory.AddCustomBuilder(MakeShared<FNiagaraCustomNodeBuilder>(InputCollectionViewModel));
+
+
+	IDetailCategoryBuilder& OutputParamCategory = DetailBuilder.EditCategory(OutputParamCategoryName, LOCTEXT("OutputParamCategoryName", "Output Parameters"));
+	OutputParamCategory.HeaderContent(SNew(SAddParameterButton, OutputCollectionViewModel));
+	OutputParamCategory.AddCustomBuilder(MakeShared<FNiagaraCustomNodeBuilder>(OutputCollectionViewModel));
+	
+}
+
+
+#undef LOCTEXT_NAMESPACE
diff --git a/Engine/Plugins/FX/Niagara/Source/NiagaraEditor/Private/Customizations/NiagaraScriptDetails.h b/Engine/Plugins/FX/Niagara/Source/NiagaraEditor/Private/Customizations/NiagaraScriptDetails.h
new file mode 100644
--- /dev/null
+++ b/Engine/Plugins/FX/Niagara/Source/NiagaraEditor/Private/Customizations/NiagaraScriptDetails.h
@@ -0,0 +1,23 @@
+// Copyright 1998-2016 Epic Games, Inc. All Rights Reserved.
+
+#pragma once
+
+#include "CoreMinimal.h"
+#include "IDetailCustomization.h"
+
+class FNiagaraScriptViewModel;
+class INiagaraParameterCollectionViewModel;
+
+class FNiagaraScriptDetails : public IDetailCustomization
+{
+public:
+	static TSharedRef<IDetailCustomization> MakeInstance(TWeakPtr<FNiagaraScriptViewModel> ScriptViewModel);
+	
+	FNiagaraScriptDetails(TSharedPtr<FNiagaraScriptViewModel> InScriptViewModel);
+
+	virtual void CustomizeDetails(IDetailLayoutBuilder& DetailBuilder);
+private:
+	TSharedPtr<FNiagaraScriptViewModel> ScriptViewModel;
+
+};
+
