Copying //UE4/Orion-Staging to //UE4/Main (Source: //Orion/Dev-General @ 3606378)

diff --git a/Engine/Plugins/FX/Niagara/Source/NiagaraEditor/Private/EdGraphSchema_Niagara.cpp b/Engine/Plugins/FX/Niagara/Source/NiagaraEditor/Private/EdGraphSchema_Niagara.cpp
new file mode 100644
--- /dev/null
+++ b/Engine/Plugins/FX/Niagara/Source/NiagaraEditor/Private/EdGraphSchema_Niagara.cpp
@@ -0,0 +1,1310 @@
+// Copyright 1998-2017 Epic Games, Inc. All Rights Reserved.
+
+#include "EdGraphSchema_Niagara.h"
+#include "NiagaraEditorModule.h"
+#include "INiagaraEditorTypeUtilities.h"
+#include "Textures/SlateIcon.h"
+#include "Framework/Commands/UIAction.h"
+#include "Framework/MultiBox/MultiBoxBuilder.h"
+#include "NiagaraCommon.h"
+#include "NiagaraEditorCommon.h"
+#include "INiagaraCompiler.h"
+#include "NiagaraHlslTranslator.h"
+#include "NiagaraComponent.h"
+#include "ScopedTransaction.h"
+#include "NiagaraGraph.h"
+#include "GraphEditorSettings.h"
+#include "GraphEditorActions.h"
+#include "NiagaraConstants.h"
+
+#include "NiagaraScript.h"
+#include "NiagaraNodeOutput.h"
+#include "NiagaraNodeInput.h"
+#include "NiagaraNodeFunctionCall.h"
+#include "NiagaraNodeReadDataSet.h"
+#include "NiagaraNodeWriteDataSet.h"
+#include "NiagaraNodeParameterMapGet.h"
+#include "NiagaraNodeParameterMapSet.h"
+#include "NiagaraNodeParameterCollection.h"
+#include "NiagaraNodeOp.h"
+#include "NiagaraNodeConvert.h"
+#include "NiagaraEditorUtilities.h"
+#include "NiagaraDataInterface.h"
+#include "NiagaraNodeIf.h"
+#include "MessageDialog.h"
+#include "NiagaraScriptSource.h"
+#include "NiagaraEmitter.h"
+#include "NiagaraNodeEmitter.h"
+#include "NiagaraParameterCollection.h"
+#include "EdGraphNode_Comment.h"
+
+#include "ModuleManager.h"
+#include "AssetRegistryModule.h"
+
+#define LOCTEXT_NAMESPACE "NiagaraSchema"
+
+#define SNAP_GRID (16) // @todo ensure this is the same as SNodePanel::GetSnapGridSize()
+
+const FLinearColor UEdGraphSchema_Niagara::NodeTitleColor_Attribute = FLinearColor::Green;
+const FLinearColor UEdGraphSchema_Niagara::NodeTitleColor_Constant = FLinearColor::Red;
+const FLinearColor UEdGraphSchema_Niagara::NodeTitleColor_SystemConstant = FLinearColor::White;
+const FLinearColor UEdGraphSchema_Niagara::NodeTitleColor_FunctionCall = FLinearColor::Blue;
+const FLinearColor UEdGraphSchema_Niagara::NodeTitleColor_Event = FLinearColor::Red;
+
+const FString UEdGraphSchema_Niagara::PinCategoryType("Type");
+const FString UEdGraphSchema_Niagara::PinCategoryMisc("Misc");
+const FString UEdGraphSchema_Niagara::PinCategoryClass("Class");
+const FString UEdGraphSchema_Niagara::PinCategoryEnum("Enum");
+
+namespace 
+{
+	// Maximum distance a drag can be off a node edge to require 'push off' from node
+	const int32 NodeDistance = 60;
+}
+
+UEdGraphNode* FNiagaraSchemaAction_NewNode::PerformAction(class UEdGraph* ParentGraph, UEdGraphPin* FromPin, const FVector2D Location, bool bSelectNewNode/* = true*/)
+{
+	UEdGraphNode* ResultNode = NULL;
+
+	// If there is a template, we actually use it
+	if (NodeTemplate != NULL)
+	{
+		FString OutErrorMsg;
+		UNiagaraNode* NiagaraNodeTemplate = Cast<UNiagaraNode>(NodeTemplate);
+		if (NiagaraNodeTemplate && !NiagaraNodeTemplate->CanAddToGraph(CastChecked<UNiagaraGraph>(ParentGraph), OutErrorMsg))
+		{
+			if (OutErrorMsg.Len() > 0)
+			{
+				FMessageDialog::Open(EAppMsgType::Ok, FText::FromString(OutErrorMsg));
+			}
+			return ResultNode;
+		}
+
+		const FScopedTransaction Transaction(NSLOCTEXT("UnrealEd", "NiagaraEditorNewNode", "Niagara Editor: New Node"));
+		ParentGraph->Modify();
+
+		NodeTemplate->SetFlags(RF_Transactional);
+
+		// set outer to be the graph so it doesn't go away
+		NodeTemplate->Rename(NULL, ParentGraph, REN_NonTransactional);
+		ParentGraph->AddNode(NodeTemplate, true, bSelectNewNode);
+
+		NodeTemplate->CreateNewGuid();
+		NodeTemplate->PostPlacedNewNode();
+		NodeTemplate->AllocateDefaultPins();
+		NodeTemplate->AutowireNewNode(FromPin);
+
+		// For input pins, new node will generally overlap node being dragged off
+		// Work out if we want to visually push away from connected node
+		int32 XLocation = Location.X;
+		if (FromPin && FromPin->Direction == EGPD_Input)
+		{
+			UEdGraphNode* PinNode = FromPin->GetOwningNode();
+			const float XDelta = FMath::Abs(PinNode->NodePosX - Location.X);
+
+			if (XDelta < NodeDistance)
+			{
+				// Set location to edge of current node minus the max move distance
+				// to force node to push off from connect node enough to give selection handle
+				XLocation = PinNode->NodePosX - NodeDistance;
+			}
+		}
+
+		NodeTemplate->NodePosX = XLocation;
+		NodeTemplate->NodePosY = Location.Y;
+		NodeTemplate->SnapToGrid(SNAP_GRID);
+
+		ResultNode = NodeTemplate;
+
+		ParentGraph->NotifyGraphChanged();
+	}
+
+	return ResultNode;
+}
+
+UEdGraphNode* FNiagaraSchemaAction_NewNode::PerformAction(class UEdGraph* ParentGraph, TArray<UEdGraphPin*>& FromPins, const FVector2D Location, bool bSelectNewNode/* = true*/) 
+{
+	UEdGraphNode* ResultNode = NULL;
+
+	if (FromPins.Num() > 0)
+	{
+		ResultNode = PerformAction(ParentGraph, FromPins[0], Location, bSelectNewNode);
+
+		if (ResultNode)
+		{
+			// Try autowiring the rest of the pins
+			for (int32 Index = 1; Index < FromPins.Num(); ++Index)
+			{
+				ResultNode->AutowireNewNode(FromPins[Index]);
+			}
+		}
+	}
+	else
+	{
+		ResultNode = PerformAction(ParentGraph, NULL, Location, bSelectNewNode);
+	}
+
+	return ResultNode;
+}
+
+void FNiagaraSchemaAction_NewNode::AddReferencedObjects( FReferenceCollector& Collector )
+{
+	FEdGraphSchemaAction::AddReferencedObjects( Collector );
+
+	// These don't get saved to disk, but we want to make sure the objects don't get GC'd while the action array is around
+	Collector.AddReferencedObject( NodeTemplate );
+}
+
+//////////////////////////////////////////////////////////////////////////
+
+static int32 GbAllowAllNiagaraNodesInEmitterGraphs = 1;
+static FAutoConsoleVariableRef CVarAllowAllNiagaraNodesInEmitterGraphs(
+	TEXT("niagara.AllowAllNiagaraNodesInEmitterGraphs"),
+	GbAllowAllNiagaraNodesInEmitterGraphs,
+	TEXT("If true, all nodes will be allowed in the Niagara emitter graphs. \n"),
+	ECVF_Default
+);
+
+UEdGraphSchema_Niagara::UEdGraphSchema_Niagara(const FObjectInitializer& ObjectInitializer)
+	: Super(ObjectInitializer)
+{
+}
+
+TSharedPtr<FNiagaraSchemaAction_NewNode> AddNewNodeAction(TArray<TSharedPtr<FNiagaraSchemaAction_NewNode> >& NewActions, const FText& Category, const FText& MenuDesc, const FName& InternalName, const FText& Tooltip, FText Keywords = FText())
+{
+	TSharedPtr<FNiagaraSchemaAction_NewNode> NewAction = TSharedPtr<FNiagaraSchemaAction_NewNode>(new FNiagaraSchemaAction_NewNode(Category, MenuDesc, InternalName, Tooltip, 0, Keywords));
+	NewActions.Add(NewAction);
+	return NewAction;
+}
+
+const UNiagaraGraph* GetAlternateGraph(const UNiagaraGraph* NiagaraGraph)
+{
+	UNiagaraScriptSource* ScriptSource = Cast<UNiagaraScriptSource>(NiagaraGraph->GetOuter());
+	if (ScriptSource != nullptr)
+	{
+		UNiagaraScript* Script = Cast<UNiagaraScript>(ScriptSource->GetOuter());
+		if (Script != nullptr)
+		{
+			UNiagaraEmitter* EmitterProperties = Cast<UNiagaraEmitter>(Script->GetOuter());
+			if (EmitterProperties != nullptr)
+			{
+				if (EmitterProperties->SpawnScriptProps.Script == Script)
+				{
+					return CastChecked<UNiagaraScriptSource>(EmitterProperties->UpdateScriptProps.Script->GetSource())->NodeGraph;
+				}
+				else if (EmitterProperties->UpdateScriptProps.Script == Script)
+				{
+					return CastChecked<UNiagaraScriptSource>(EmitterProperties->SpawnScriptProps.Script->GetSource())->NodeGraph;
+				}
+			}
+		}
+	}
+	return nullptr;
+}
+
+FText GetGraphTypeTitle(const UNiagaraGraph* NiagaraGraph)
+{
+	UNiagaraScriptSource* ScriptSource = Cast<UNiagaraScriptSource>(NiagaraGraph->GetOuter());
+	if (ScriptSource != nullptr)
+	{
+		UNiagaraScript* Script = Cast<UNiagaraScript>(ScriptSource->GetOuter());
+		if (Script != nullptr)
+		{
+			if (Script->IsParticleSpawnScript())
+			{
+				return LOCTEXT("Parameter Menu Title Spawn", "Spawn Parameters");
+			}
+			else if (Script->IsParticleUpdateScript())
+			{
+				return LOCTEXT("Parameter Menu Title Update", "Update Parameters");
+			}
+		}
+	}
+	return LOCTEXT("Parameter Menu Title Generic", "Script Parameters");
+}
+
+void AddParametersForGraph(TArray<TSharedPtr<FNiagaraSchemaAction_NewNode> >& NewActions, const UNiagaraGraph* CurrentGraph,  UEdGraph* OwnerOfTemporaries, const UNiagaraGraph* NiagaraGraph)
+{
+	FText GraphParameterCategory = GetGraphTypeTitle(NiagaraGraph);
+	TArray<UNiagaraNodeInput*> InputNodes;
+	NiagaraGraph->GetNodesOfClass(InputNodes);
+
+	TArray<FNiagaraVariable> SeenParams;
+	for (UNiagaraNodeInput* InputNode : InputNodes)
+	{
+		if (InputNode->Usage == ENiagaraInputNodeUsage::Parameter && !SeenParams.Contains(InputNode->Input))
+		{
+			SeenParams.Add(InputNode->Input);
+			FName Name = InputNode->Input.GetName();
+			FText MenuDesc = FText::FromName(Name);
+			if (NiagaraGraph != CurrentGraph)
+			{
+				Name = UNiagaraNodeInput::GenerateUniqueName(CastChecked<UNiagaraGraph>(CurrentGraph), Name, InputNode->Usage);
+				MenuDesc = FText::Format(LOCTEXT("Parameter Menu Copy Param","Copy \"{0}\" to this Graph"), FText::FromName(Name));
+			}
+
+			TSharedPtr<FNiagaraSchemaAction_NewNode> ExistingInputAction = AddNewNodeAction(NewActions, GraphParameterCategory, MenuDesc, Name, FText::GetEmpty());
+
+			UNiagaraNodeInput* InputNodeTemplate = NewObject<UNiagaraNodeInput>(OwnerOfTemporaries);
+			InputNodeTemplate->Input = InputNode->Input;
+			InputNodeTemplate->Usage = InputNode->Usage;
+			InputNodeTemplate->ExposureOptions = InputNode->ExposureOptions;
+			InputNodeTemplate->DataInterface = nullptr;
+
+			// We also support parameters from an alternate graph. If that was used, then we need to take special care
+			// to make the parameter unique to that graph.
+			if (NiagaraGraph != CurrentGraph)
+			{
+				InputNodeTemplate->Input.SetName(Name);
+
+				if (InputNode->DataInterface)
+				{
+					InputNodeTemplate->DataInterface = Cast<UNiagaraDataInterface>(StaticDuplicateObject(InputNode->DataInterface, InputNodeTemplate, NAME_None, ~RF_Transient));
+				}
+			}
+
+			ExistingInputAction->NodeTemplate = InputNodeTemplate;
+		}
+	}
+}
+
+void AddParameterMenuOptions(TArray<TSharedPtr<FNiagaraSchemaAction_NewNode> >& NewActions, const UNiagaraGraph* CurrentGraph, UEdGraph* OwnerOfTemporaries, const UNiagaraGraph* NiagaraGraph)
+{
+	AddParametersForGraph(NewActions, CurrentGraph, OwnerOfTemporaries, NiagaraGraph);
+
+	const UNiagaraGraph* AltGraph = GetAlternateGraph(NiagaraGraph);
+	if (AltGraph != nullptr)
+	{
+		AddParametersForGraph(NewActions, CurrentGraph, OwnerOfTemporaries, AltGraph);
+	}
+}
+
+void UEdGraphSchema_Niagara::GetGraphContextActions(FGraphContextMenuBuilder& ContextMenuBuilder) const
+{
+	const UNiagaraGraph* NiagaraGraph = CastChecked<UNiagaraGraph>(ContextMenuBuilder.CurrentGraph);
+	TArray<TSharedPtr<FNiagaraSchemaAction_NewNode> > NewActions = GetGraphContextActions(NiagaraGraph, ContextMenuBuilder.SelectedObjects, ContextMenuBuilder.FromPin, ContextMenuBuilder.OwnerOfTemporaries);
+	for (int32 i = 0; i < NewActions.Num(); i++)
+	{
+		ContextMenuBuilder.AddAction(NewActions[i]);
+	}
+}
+
+TArray<TSharedPtr<FNiagaraSchemaAction_NewNode> > UEdGraphSchema_Niagara::GetGraphContextActions(const UEdGraph* CurrentGraph, TArray<UObject*>& SelectedObjects, const UEdGraphPin* FromPin, UEdGraph* OwnerOfTemporaries) const
+{
+	TArray<TSharedPtr<FNiagaraSchemaAction_NewNode> > NewActions;
+
+	const UNiagaraGraph* NiagaraGraph = CastChecked<UNiagaraGraph>(CurrentGraph);
+	TArray<UNiagaraNodeEmitter*> Emitters;
+	NiagaraGraph->GetNodesOfClass<UNiagaraNodeEmitter>(Emitters);
+	bool bSystemGraph = Emitters.Num() != 0 || NiagaraGraph->FindOutputNode(ENiagaraScriptUsage::SystemSpawnScript) != nullptr || NiagaraGraph->FindOutputNode(ENiagaraScriptUsage::SystemUpdateScript) != nullptr;
+	bool bModuleGraph = NiagaraGraph->FindOutputNode(ENiagaraScriptUsage::Module) != nullptr;
+	bool bFunctionGraph = NiagaraGraph->FindOutputNode(ENiagaraScriptUsage::Function) != nullptr;
+	bool bUpdateGraph = NiagaraGraph->FindOutputNode(ENiagaraScriptUsage::ParticleUpdateScript) != nullptr;
+
+	
+
+	if (GbAllowAllNiagaraNodesInEmitterGraphs || bModuleGraph || bFunctionGraph || bSystemGraph)
+	{
+		const TArray<FNiagaraOpInfo>& OpInfos = FNiagaraOpInfo::GetOpInfoArray();
+		for (const FNiagaraOpInfo& OpInfo : OpInfos)
+		{
+			TSharedPtr<FNiagaraSchemaAction_NewNode> AddOpAction = AddNewNodeAction(NewActions, OpInfo.Category, OpInfo.FriendlyName, OpInfo.Name, FText::GetEmpty());
+			UNiagaraNodeOp* OpNode = NewObject<UNiagaraNodeOp>(OwnerOfTemporaries);
+			OpNode->OpName = OpInfo.Name;
+			AddOpAction->NodeTemplate = OpNode;
+		}
+	}
+
+	//Add functions
+	FAssetRegistryModule& AssetRegistryModule = FModuleManager::LoadModuleChecked<FAssetRegistryModule>(TEXT("AssetRegistry"));
+	TArray<FAssetData> ScriptAssets;
+	AssetRegistryModule.Get().GetAssetsByClass(UNiagaraScript::StaticClass()->GetFName(), ScriptAssets);
+	UEnum* NiagaraScriptUsageEnum = FindObjectChecked<UEnum>(ANY_PACKAGE, TEXT("ENiagaraScriptUsage"), true);
+	if (GbAllowAllNiagaraNodesInEmitterGraphs || bModuleGraph || bFunctionGraph || bSystemGraph)
+	{
+		for (const FAssetData& ScriptAsset : ScriptAssets)
+		{
+			FName UsageName;
+			ScriptAsset.GetTagValue(GET_MEMBER_NAME_CHECKED(UNiagaraScript, Usage), UsageName);
+
+			FText AssetDesc;
+			ScriptAsset.GetTagValue(GET_MEMBER_NAME_CHECKED(UNiagaraScript, Description), AssetDesc);
+
+			FString QualifiedUsageName = "ENiagaraScriptUsage::" + UsageName.ToString();
+			int32 UsageIndex = NiagaraScriptUsageEnum->GetIndexByNameString(QualifiedUsageName);
+			if (UsageIndex != INDEX_NONE)
+			{
+				ENiagaraScriptUsage Usage = static_cast<ENiagaraScriptUsage>(NiagaraScriptUsageEnum->GetValueByIndex(UsageIndex));
+				if (Usage == ENiagaraScriptUsage::Function)
+				{
+					FString DisplayNameString = FName::NameToDisplayString(ScriptAsset.AssetName.ToString(), false);
+					const FText MenuDesc = FText::FromString(DisplayNameString);
+					const FText TooltipDesc = FText::Format(LOCTEXT("FunctionPopupTooltip", "Path: {0}\nDescription: {1}"), FText::FromString(ScriptAsset.ObjectPath.ToString()), AssetDesc);
+
+
+					TSharedPtr<FNiagaraSchemaAction_NewNode> FunctionCallAction = AddNewNodeAction(NewActions, LOCTEXT("Function Menu Title", "Functions"), MenuDesc, *DisplayNameString, FText::FromName(ScriptAsset.ObjectPath));
+
+					UNiagaraNodeFunctionCall* FunctionCallNode = NewObject<UNiagaraNodeFunctionCall>(OwnerOfTemporaries);
+					FunctionCallNode->FunctionScriptAssetObjectPath = ScriptAsset.ObjectPath;
+					FunctionCallAction->NodeTemplate = FunctionCallNode;
+				}
+			}
+		}
+	}
+
+	//Add modules
+	//if (GraphUsage == ENiagaraScriptUsage::ParticleSpawnScript || GraphUsage == ENiagaraScriptUsage::ParticleUpdateScript) 
+	if (!bFunctionGraph)
+	{
+		for (const FAssetData& ScriptAsset : ScriptAssets)
+		{
+			FName UsageName;
+			ScriptAsset.GetTagValue(GET_MEMBER_NAME_CHECKED(UNiagaraScript, Usage), UsageName);
+			const FString BitfieldTagValue = ScriptAsset.GetTagValueRef<FString>(GET_MEMBER_NAME_CHECKED(UNiagaraScript, ModuleUsageBitmask));
+			int32 BitfieldValue = FCString::Atoi(*BitfieldTagValue);
+			FText AssetDesc;
+			ScriptAsset.GetTagValue(GET_MEMBER_NAME_CHECKED(UNiagaraScript, Description), AssetDesc);
+
+			ENiagaraScriptUsage TargetUsage = ENiagaraScriptUsage::Module;
+			if (bUpdateGraph)
+			{
+				TargetUsage = ENiagaraScriptUsage::ParticleUpdateScript;
+			}
+			else if (bSystemGraph)
+			{
+				TargetUsage = ENiagaraScriptUsage::SystemSpawnScript;
+			}
+
+			int32 TargetBit = (BitfieldValue >> (int32)TargetUsage) & 1;
+
+			FString QualifiedUsageName = "ENiagaraScriptUsage::" + UsageName.ToString();
+			int32 UsageIndex = NiagaraScriptUsageEnum->GetIndexByNameString(QualifiedUsageName);
+			if (UsageIndex != INDEX_NONE && TargetBit == 1)
+			{
+				ENiagaraScriptUsage Usage = static_cast<ENiagaraScriptUsage>(NiagaraScriptUsageEnum->GetValueByIndex(UsageIndex));
+				if (Usage == ENiagaraScriptUsage::Module)
+				{
+					FString DisplayNameString = FName::NameToDisplayString(ScriptAsset.AssetName.ToString(), false);
+					const FText MenuDesc = FText::FromString(DisplayNameString);
+					const FText TooltipDesc = FText::Format(LOCTEXT("ModulePopupTooltip", "Path: {0}\nDescription: {1}"), FText::FromString(ScriptAsset.ObjectPath.ToString()), AssetDesc);
+
+					TSharedPtr<FNiagaraSchemaAction_NewNode> FunctionCallAction = AddNewNodeAction(NewActions, LOCTEXT("Module Menu Title", "Modules"), MenuDesc, *DisplayNameString, FText::FromName(ScriptAsset.ObjectPath));
+
+					UNiagaraNodeFunctionCall* FunctionCallNode = NewObject<UNiagaraNodeFunctionCall>(OwnerOfTemporaries);
+					FunctionCallNode->FunctionScriptAssetObjectPath = ScriptAsset.ObjectPath;
+					FunctionCallAction->NodeTemplate = FunctionCallNode;
+				}
+			}
+		}
+	}
+
+	//Add collections
+	if (bModuleGraph || bFunctionGraph)
+	{
+		TArray<FAssetData> CollectionAssets;
+		AssetRegistryModule.Get().GetAssetsByClass(UNiagaraParameterCollection::StaticClass()->GetFName(), CollectionAssets);
+
+		for (const FAssetData& CollectionAsset : CollectionAssets)
+		{
+			FString DisplayNameString = FName::NameToDisplayString(CollectionAsset.AssetName.ToString(), false);
+			const FText MenuDesc = FText::FromString(DisplayNameString);
+
+			TSharedPtr<FNiagaraSchemaAction_NewNode> CollectionAction = AddNewNodeAction(NewActions, LOCTEXT("Collection Menu Title", "Collections"), MenuDesc, *DisplayNameString, FText::FromName(CollectionAsset.ObjectPath));
+
+			UNiagaraNodeParameterCollection* CollectionNode = NewObject<UNiagaraNodeParameterCollection>(OwnerOfTemporaries);
+			CollectionNode->CollectionAssetObjectPath = CollectionAsset.ObjectPath;
+			CollectionAction->NodeTemplate = CollectionNode;
+		}
+	}
+
+	//Add event read and writes nodes
+	if (bModuleGraph)
+	{
+		const FText MenuCat = LOCTEXT("NiagaraEventMenuCat", "Events");
+		const TArray<FNiagaraTypeDefinition>& RegisteredTypes = FNiagaraTypeRegistry::GetRegisteredPayloadTypes();
+		for (FNiagaraTypeDefinition Type : RegisteredTypes)
+		{
+			if (Type.GetStruct() && !Type.GetStruct()->IsA(UNiagaraDataInterface::StaticClass()))
+			{
+				{
+					const FText MenuDescFmt = LOCTEXT("AddEventReadFmt", "Add {0} Event Read");
+					const FText MenuDesc = FText::Format(MenuDescFmt, Type.GetStruct()->GetDisplayNameText());
+
+					TSharedPtr<FNiagaraSchemaAction_NewNode> Action = AddNewNodeAction(NewActions, MenuCat, MenuDesc, *MenuDesc.ToString(), FText::GetEmpty());
+
+					UNiagaraNodeReadDataSet* EventReadNode = NewObject<UNiagaraNodeReadDataSet>(OwnerOfTemporaries);
+					EventReadNode->InitializeFromStruct(Type.GetStruct());
+					Action->NodeTemplate = EventReadNode;
+				}
+				{
+					const FText MenuDescFmt = LOCTEXT("AddEventWriteFmt", "Add {0} Event Write");
+					const FText MenuDesc = FText::Format(MenuDescFmt, Type.GetStruct()->GetDisplayNameText());
+
+					TSharedPtr<FNiagaraSchemaAction_NewNode> Action = AddNewNodeAction(NewActions, MenuCat, MenuDesc, *MenuDesc.ToString(), FText::GetEmpty());
+
+					UNiagaraNodeWriteDataSet* EventWriteNode = NewObject<UNiagaraNodeWriteDataSet>(OwnerOfTemporaries);
+					EventWriteNode->InitializeFromStruct(Type.GetStruct());
+					Action->NodeTemplate = EventWriteNode;
+				}
+			}
+		}
+	}
+	
+	TArray<ENiagaraScriptUsage> UsageTypesToAdd;
+	if (bUpdateGraph)
+	{
+		UsageTypesToAdd.Add(ENiagaraScriptUsage::ParticleEventScript);
+		UsageTypesToAdd.Add(ENiagaraScriptUsage::EmitterSpawnScript);
+		UsageTypesToAdd.Add(ENiagaraScriptUsage::EmitterUpdateScript);
+	}
+
+	if (bSystemGraph)
+	{
+		UsageTypesToAdd.Add(ENiagaraScriptUsage::SystemSpawnScript);
+		UsageTypesToAdd.Add(ENiagaraScriptUsage::SystemUpdateScript);
+	}
+
+	if (UsageTypesToAdd.Num() != 0)
+	{
+		for (ENiagaraScriptUsage Usage : UsageTypesToAdd)
+		{
+			const FText MenuCat = LOCTEXT("NiagaraUsageMenuCat", "Output Nodes");
+
+			UNiagaraNodeOutput* OutputNode = NewObject<UNiagaraNodeOutput>(OwnerOfTemporaries);
+			OutputNode->SetUsage(Usage);
+
+			FText MenuDesc = FText::Format(LOCTEXT("AddOutput", "Add {0}"), OutputNode->GetNodeTitle(ENodeTitleType::FullTitle));
+			TSharedPtr<FNiagaraSchemaAction_NewNode> OutputNodeAction = AddNewNodeAction(NewActions, MenuCat, MenuDesc, *MenuDesc.ToString(), FText::GetEmpty());
+
+			UNiagaraNodeOutput* UpdateOutputNode = NiagaraGraph->FindOutputNode(ENiagaraScriptUsage::ParticleUpdateScript);
+			if (UpdateOutputNode)
+			{
+				OutputNode->Outputs = UpdateOutputNode->Outputs;
+			}
+			else
+			{
+				OutputNode->Outputs.Add(FNiagaraVariable(FNiagaraTypeDefinition::GetParameterMapDef(), TEXT("Out")));
+			}
+			OutputNodeAction->NodeTemplate = OutputNode;
+		}
+	}
+
+
+	// Add Convert Nodes
+	{
+		FNiagaraTypeDefinition PinType = FNiagaraTypeDefinition::GetGenericNumericDef();
+		bool bAddMakes = true;
+		bool bAddBreaks = true;
+		if (FromPin)
+		{
+			PinType = PinToTypeDefinition(FromPin);
+			if (FromPin->Direction == EGPD_Input)
+			{
+				bAddBreaks = false;
+			}
+			else
+			{
+				bAddMakes = false;
+			}
+		}
+
+		if (PinType.GetScriptStruct())
+		{
+			FText MakeCat = LOCTEXT("NiagaraMake", "Make");
+			FText BreakCat = LOCTEXT("NiagaraBreak", "Break");
+
+			FText DescFmt = LOCTEXT("NiagaraMakeBreakFmt", "{0}");
+			auto MakeBreakType = [&](FNiagaraTypeDefinition Type, bool bMake)
+			{
+				FText DisplayName = Type.GetStruct()->GetDisplayNameText();
+				if (Type.GetEnum())
+				{
+					DisplayName = FText::FromString(Type.GetEnum()->GetName());
+				}
+				FText Desc = FText::Format(DescFmt, DisplayName);
+				TSharedPtr<FNiagaraSchemaAction_NewNode> Action = AddNewNodeAction(NewActions, bMake ? MakeCat : BreakCat, Desc, *Type.GetStruct()->GetName(), FText::GetEmpty());
+				UNiagaraNodeConvert* ConvertNode = NewObject<UNiagaraNodeConvert>(OwnerOfTemporaries);
+				if (bMake)
+				{
+					ConvertNode->InitAsMake(Type);
+				}
+				else
+				{
+					ConvertNode->InitAsBreak(Type);
+				}
+				Action->NodeTemplate = ConvertNode;
+			};
+
+			if (PinType == FNiagaraTypeDefinition::GetGenericNumericDef())
+			{
+				if (bAddMakes)
+				{
+					const TArray<FNiagaraTypeDefinition>& RegisteredTypes = FNiagaraTypeRegistry::GetRegisteredTypes();
+					for (FNiagaraTypeDefinition Type : RegisteredTypes)
+					{
+						// Data interfaces can't be made.
+						if (!UNiagaraDataInterface::IsDataInterfaceType(Type))
+						{
+							MakeBreakType(Type, true);
+						}
+					}
+				}
+
+				if (bAddBreaks)
+				{
+					const TArray<FNiagaraTypeDefinition>& RegisteredTypes = FNiagaraTypeRegistry::GetRegisteredTypes();
+					for (FNiagaraTypeDefinition Type : RegisteredTypes)
+					{
+						//Don't break scalars. Allow makes for now as a convenient method of getting internal script constants when dealing with numeric pins.
+						// Data interfaces can't be broken.
+						if (!FNiagaraTypeDefinition::IsScalarDefinition(Type) && !UNiagaraDataInterface::IsDataInterfaceType(Type))
+						{
+							MakeBreakType(Type, false);
+						}
+					}
+				}
+			}
+			else
+			{
+				//If we have a valid type then add it as a convenience at the top level.
+				FText TypedMakeBreakFmt = LOCTEXT("NiagaraTypedMakeBreakFmt", "{0} {1}");
+				FText DisplayName = PinType.GetStruct()->GetDisplayNameText();
+				if (PinType.GetEnum())
+				{
+					DisplayName = FText::FromString(PinType.GetEnum()->GetName());
+				}
+				FText Desc = FText::Format(TypedMakeBreakFmt, bAddMakes ? MakeCat : BreakCat, DisplayName);
+				TSharedPtr<FNiagaraSchemaAction_NewNode> Action = AddNewNodeAction(NewActions, FText::GetEmpty(), Desc, *Desc.ToString(), FText::GetEmpty());
+				UNiagaraNodeConvert* ConvertNode = NewObject<UNiagaraNodeConvert>(OwnerOfTemporaries);
+				if (bAddMakes)
+				{
+					ConvertNode->InitAsMake(PinType);
+				}
+				else
+				{
+					ConvertNode->InitAsBreak(PinType);
+				}
+				Action->NodeTemplate = ConvertNode;
+			}
+
+			//Always add generic convert as an option.
+			FText Desc = LOCTEXT("NiagaraConvert", "Convert");
+			TSharedPtr<FNiagaraSchemaAction_NewNode> Action = AddNewNodeAction(NewActions, FText::GetEmpty(), Desc, TEXT("Convert"), FText::GetEmpty());
+			UNiagaraNodeConvert* ConvertNode = NewObject<UNiagaraNodeConvert>(OwnerOfTemporaries);
+			Action->NodeTemplate = ConvertNode;
+		}
+	}
+
+	if (FromPin)
+	{
+		//Add pin specific menu options.
+		FNiagaraTypeDefinition PinType = PinToTypeDefinition(FromPin);
+		UNiagaraDataInterface* DataInterface = nullptr;
+		const UClass* Class = PinType.GetClass();
+		if (Class != nullptr)
+		{
+			if (UNiagaraNodeInput* InputNode = Cast<UNiagaraNodeInput>(FromPin->GetOwningNode()))
+			{
+				DataInterface = InputNode->DataInterface;
+			}
+			else 
+			{
+				DataInterface = Cast<UNiagaraDataInterface>(const_cast<UClass*>(Class)->GetDefaultObject());
+			}
+
+			if (DataInterface)
+			{
+				FText MenuCat = Class->GetDisplayNameText();
+				TArray<FNiagaraFunctionSignature> Functions;
+				DataInterface->GetFunctions(Functions);
+				for (FNiagaraFunctionSignature& Sig : Functions)
+				{
+					TSharedPtr<FNiagaraSchemaAction_NewNode> Action = AddNewNodeAction(NewActions, MenuCat, FText::FromString(Sig.GetName()), *Sig.GetName(), FText::GetEmpty());
+					UNiagaraNodeFunctionCall* FuncNode = NewObject<UNiagaraNodeFunctionCall>(OwnerOfTemporaries);
+					Action->NodeTemplate = FuncNode;
+					FuncNode->Signature = Sig;
+				}
+			}
+		}
+
+		if (FromPin->Direction == EGPD_Output)
+		{
+			//Add all swizzles for this type if it's a vector.
+			if (FHlslNiagaraTranslator::IsHlslBuiltinVector(PinType))
+			{
+				TArray<FString> Components;
+				for (TFieldIterator<UProperty> PropertyIt(PinType.GetStruct(), EFieldIteratorFlags::IncludeSuper); PropertyIt; ++PropertyIt)
+				{
+					UProperty* Property = *PropertyIt;
+					Components.Add(Property->GetName().ToLower());
+				}
+
+				TArray<FString> Swizzles;
+				TFunction<void(FString)> GenSwizzles = [&](FString CurrStr)
+				{
+					if (CurrStr.Len() == 4) return;//Only generate down to float4
+					for (FString& CompStr : Components)
+					{
+						Swizzles.Add(CurrStr + CompStr);
+						GenSwizzles(CurrStr + CompStr);
+					}
+				};
+
+				GenSwizzles(FString());
+
+				for (FString Swiz : Swizzles)
+				{
+					const FText Category = LOCTEXT("NiagaraSwizzles", "Swizzles");
+
+					TSharedPtr<FNiagaraSchemaAction_NewNode> Action = AddNewNodeAction(NewActions, Category, FText::FromString(Swiz), *Swiz, FText::GetEmpty());
+
+					UNiagaraNodeConvert* ConvertNode = NewObject<UNiagaraNodeConvert>(OwnerOfTemporaries);
+					Action->NodeTemplate = ConvertNode;
+					ConvertNode->InitAsSwizzle(Swiz);
+				}
+			}
+		}
+	}
+
+	// Handle parameter map get/set
+	{
+		FText MenuCat = FText::FromString("Parameter Map");
+		{
+			FString Name = TEXT("Parameter Map Get");
+			TSharedPtr<FNiagaraSchemaAction_NewNode> Action = AddNewNodeAction(NewActions, MenuCat, FText::FromString(Name), *Name, FText::GetEmpty());
+			UNiagaraNodeParameterMapGet* BaseNode = NewObject<UNiagaraNodeParameterMapGet>(OwnerOfTemporaries);
+			Action->NodeTemplate = BaseNode;
+		}
+		{
+			FString Name = TEXT("Parameter Map Set");
+			TSharedPtr<FNiagaraSchemaAction_NewNode> Action = AddNewNodeAction(NewActions, MenuCat, FText::FromString(Name), *Name, FText::GetEmpty());
+			UNiagaraNodeParameterMapSet* BaseNode = NewObject<UNiagaraNodeParameterMapSet>(OwnerOfTemporaries);
+			Action->NodeTemplate = BaseNode;
+		}
+	}
+
+	// Handle comment nodes
+	{
+		FText MenuCat = FText::FromString("Comments");
+
+		{
+			FString Name = TEXT("Add Comment");
+			TSharedPtr<FNiagaraSchemaAction_NewNode> Action = AddNewNodeAction(NewActions, MenuCat, FText::FromString(Name), *Name, FText::GetEmpty());
+			UEdGraphNode_Comment* BaseNode = NewObject<UEdGraphNode_Comment>(OwnerOfTemporaries);
+			Action->NodeTemplate = BaseNode;
+		}		
+	}
+
+
+
+	//Add all input node options for input pins or no pin.
+	if (FromPin == nullptr || FromPin->Direction == EGPD_Input)
+	{
+		TArray<UNiagaraNodeInput*> InputNodes;
+		NiagaraGraph->GetNodesOfClass(InputNodes);
+
+		//Emitter constants managed by the system.
+		const TArray<FNiagaraVariable>& SystemConstants = FNiagaraConstants::GetEngineConstants();
+		for (const FNiagaraVariable& SysConst : SystemConstants)
+		{
+			FFormatNamedArguments Args;
+			Args.Add(TEXT("Constant"), FText::FromName(SysConst.GetName()));
+			const FText MenuDesc = FText::Format(LOCTEXT("GetSystemConstant", "Get {Constant}"), Args);
+
+			TSharedPtr<FNiagaraSchemaAction_NewNode> GetConstAction = AddNewNodeAction(NewActions, LOCTEXT("System Parameters Menu Title", "System Parameters"), MenuDesc, SysConst.GetName(), FText::GetEmpty());
+
+			UNiagaraNodeInput* InputNode = NewObject<UNiagaraNodeInput>(OwnerOfTemporaries);
+			InputNode->Usage = ENiagaraInputNodeUsage::SystemConstant;
+			InputNode->Input = SysConst;
+			GetConstAction->NodeTemplate = InputNode;
+		}
+
+		AddParameterMenuOptions(NewActions, NiagaraGraph, OwnerOfTemporaries, NiagaraGraph);
+
+		//Add a generic Parameter node to allow easy creation of parameters.
+		{
+			FNiagaraTypeDefinition PinType = FNiagaraTypeDefinition::GetGenericNumericDef();
+			if (FromPin)
+			{
+				PinType = PinToTypeDefinition(FromPin);
+			}
+
+			if (PinType.GetStruct())
+			{
+				const FText MenuDescFmt = LOCTEXT("Add ParameterFmt", "Add {0} Parameter");
+				const TArray<FNiagaraTypeDefinition>& RegisteredTypes = FNiagaraTypeRegistry::GetRegisteredParameterTypes();
+				for (FNiagaraTypeDefinition Type : RegisteredTypes)
+				{
+					FText MenuCat;
+					if (const UClass* Class = Type.GetClass())
+					{						
+						MenuCat = Class->GetMetaDataText(TEXT("Category"), TEXT("UObjectCategory"), Class->GetFullGroupName(false));
+					}
+					else
+					{
+						MenuCat = LOCTEXT("AddParameterCat", "Add Parameter");
+					}
+						
+					const FText MenuDesc = FText::Format(MenuDescFmt, Type.GetStruct()->GetDisplayNameText());
+					TSharedPtr<FNiagaraSchemaAction_NewNode> InputAction = AddNewNodeAction(NewActions, MenuCat, MenuDesc, *MenuDesc.ToString(), FText::GetEmpty());
+					UNiagaraNodeInput* InputNode = NewObject<UNiagaraNodeInput>(OwnerOfTemporaries);
+					FNiagaraEditorUtilities::InitializeParameterInputNode(*InputNode, Type, NiagaraGraph);
+					InputAction->NodeTemplate = InputNode;
+				}
+
+				if (PinType != FNiagaraTypeDefinition::GetGenericNumericDef())
+				{
+					//For correctly typed pins, offer the correct type at the top level.				
+					const FText MenuDesc = FText::Format(MenuDescFmt, PinType.GetStruct()->GetDisplayNameText());
+					TSharedPtr<FNiagaraSchemaAction_NewNode> InputAction = AddNewNodeAction(NewActions, FText::GetEmpty(), MenuDesc, *MenuDesc.ToString(), FText::GetEmpty());
+					UNiagaraNodeInput* InputNode = NewObject<UNiagaraNodeInput>(OwnerOfTemporaries);
+					FNiagaraEditorUtilities::InitializeParameterInputNode(*InputNode, PinType, NiagaraGraph);
+					InputAction->NodeTemplate = InputNode;
+				}
+			}
+		}
+	}
+
+	const FText MenuCat = LOCTEXT("NiagaraLogicMenuCat", "Logic");
+	{
+		const FText MenuDesc = LOCTEXT("If", "If");
+
+		TSharedPtr<FNiagaraSchemaAction_NewNode> Action = AddNewNodeAction(NewActions, MenuCat, MenuDesc, TEXT("If"), FText::GetEmpty());
+
+		UNiagaraNodeIf* IfNode = NewObject<UNiagaraNodeIf>(OwnerOfTemporaries);
+		Action->NodeTemplate = IfNode;
+	}
+	//TODO: Add quick commands for certain UNiagaraStructs and UNiagaraScripts to be added as functions
+
+	return NewActions;
+}
+
+const FPinConnectionResponse UEdGraphSchema_Niagara::CanCreateConnection(const UEdGraphPin* PinA, const UEdGraphPin* PinB) const
+{
+	// Make sure the pins are not on the same node
+	if (PinA->GetOwningNode() == PinB->GetOwningNode())
+	{
+		return FPinConnectionResponse(CONNECT_RESPONSE_DISALLOW, TEXT("Both are on the same node"));
+	}
+
+	// Check both pins support connections
+	if(PinA->bNotConnectable || PinB->bNotConnectable)
+	{
+		return FPinConnectionResponse(CONNECT_RESPONSE_DISALLOW, TEXT("Pin doesn't support connections."));
+	}
+
+	// Compare the directions
+	const UEdGraphPin* InputPin = NULL;
+	const UEdGraphPin* OutputPin = NULL;
+
+	if (!CategorizePinsByDirection(PinA, PinB, /*out*/ InputPin, /*out*/ OutputPin))
+	{
+		return FPinConnectionResponse(CONNECT_RESPONSE_DISALLOW, TEXT("Directions are not compatible"));
+	}
+
+	// Check for compatible type pins.
+	if (PinA->PinType.PinCategory == PinCategoryType && 
+		PinB->PinType.PinCategory == PinCategoryType && 
+		PinA->PinType != PinB->PinType)
+	{
+		FNiagaraTypeDefinition PinTypeA = PinToTypeDefinition(PinA);
+		FNiagaraTypeDefinition PinTypeB = PinToTypeDefinition(PinB);
+		if (FNiagaraTypeDefinition::TypesAreAssignable(PinTypeA, PinTypeB) == false)
+		{
+			//Do some limiting on auto conversions here?
+			if (PinTypeA.GetClass())
+			{
+				return FPinConnectionResponse(CONNECT_RESPONSE_DISALLOW, TEXT("Types are not compatible"));
+			}
+			else
+			{
+				return FPinConnectionResponse(CONNECT_RESPONSE_MAKE_WITH_CONVERSION_NODE, FString::Printf(TEXT("Convert %s to %s"), *(PinToTypeDefinition(PinA).GetNameText().ToString()), *(PinToTypeDefinition(PinB).GetNameText().ToString())));
+			}
+		}
+	}
+
+	// Check for compatible misc pins
+	if (PinA->PinType.PinCategory == PinCategoryMisc ||
+		PinB->PinType.PinCategory == PinCategoryMisc) 
+	{
+		// TODO: This shouldn't be handled explicitly here.
+		bool PinAIsConvertAddAndPinBIsNonGenericType =
+			PinA->PinType.PinCategory == PinCategoryMisc && PinA->PinType.PinSubCategory == UNiagaraNodeWithDynamicPins::AddPinSubCategory &&
+			PinB->PinType.PinCategory == PinCategoryType && PinToTypeDefinition(PinB) != FNiagaraTypeDefinition::GetGenericNumericDef();
+
+		bool PinBIsConvertAddAndPinAIsNonGenericType =
+			PinB->PinType.PinCategory == PinCategoryMisc && PinB->PinType.PinSubCategory == UNiagaraNodeWithDynamicPins::AddPinSubCategory &&
+			PinA->PinType.PinCategory == PinCategoryType && PinToTypeDefinition(PinA) != FNiagaraTypeDefinition::GetGenericNumericDef();
+
+		if(PinAIsConvertAddAndPinBIsNonGenericType == false && PinBIsConvertAddAndPinAIsNonGenericType == false)
+		{
+			return FPinConnectionResponse(CONNECT_RESPONSE_DISALLOW, TEXT("Types are not compatible"));
+		}
+	}
+
+	if (PinA->PinType.PinCategory == PinCategoryClass || PinB->PinType.PinCategory == PinCategoryClass)
+	{
+		FNiagaraTypeDefinition AType = PinToTypeDefinition(PinA);
+		FNiagaraTypeDefinition BType = PinToTypeDefinition(PinB);
+		if (AType != BType)
+		{
+			return FPinConnectionResponse(CONNECT_RESPONSE_DISALLOW, TEXT("Types are not compatible"));
+		}
+	}
+
+	if (PinA->PinType.PinCategory == PinCategoryEnum || PinB->PinType.PinCategory == PinCategoryEnum)
+	{
+		FNiagaraTypeDefinition PinTypeA = PinToTypeDefinition(PinA);
+		FNiagaraTypeDefinition PinTypeB = PinToTypeDefinition(PinB);
+		if (FNiagaraTypeDefinition::TypesAreAssignable(PinTypeA, PinTypeB) == false)
+		{
+			return FPinConnectionResponse(CONNECT_RESPONSE_DISALLOW, TEXT("Types are not compatible"));
+		}
+	}
+
+	// See if we want to break existing connections (if its an input with an existing connection)
+	const bool bBreakExistingDueToDataInput = (InputPin->LinkedTo.Num() > 0);
+	if (bBreakExistingDueToDataInput)
+	{
+		const ECanCreateConnectionResponse ReplyBreakInputs = (PinA == InputPin) ? CONNECT_RESPONSE_BREAK_OTHERS_A : CONNECT_RESPONSE_BREAK_OTHERS_B;
+		return FPinConnectionResponse(ReplyBreakInputs, TEXT("Replace existing input connections"));
+	}
+	else
+	{
+		return FPinConnectionResponse(CONNECT_RESPONSE_MAKE, FString());
+	}
+}
+
+void UEdGraphSchema_Niagara::BreakSinglePinLink(UEdGraphPin* SourcePin, UEdGraphPin* TargetPin) 
+{
+	const FScopedTransaction Transaction(NSLOCTEXT("UnrealEd", "NiagaraEditorBreakConnection", "Niagara Editor: Break Connection"));
+
+	Super::BreakSinglePinLink(SourcePin, TargetPin);
+}
+
+FConnectionDrawingPolicy* UEdGraphSchema_Niagara::CreateConnectionDrawingPolicy(int32 InBackLayerID, int32 InFrontLayerID, float InZoomFactor, const FSlateRect& InClippingRect, class FSlateWindowElementList& InDrawElements, class UEdGraph* InGraphObj) const
+{
+	return new FNiagaraConnectionDrawingPolicy(InBackLayerID, InFrontLayerID, InZoomFactor, InClippingRect, InDrawElements);
+}
+
+bool UEdGraphSchema_Niagara::TryCreateConnection(UEdGraphPin* PinA, UEdGraphPin* PinB) const
+{
+	const FScopedTransaction Transaction(NSLOCTEXT("UnrealEd", "NiagaraEditorCreateConnection", "Niagara Editor: Create Connection"));
+
+	const FPinConnectionResponse Response = CanCreateConnection(PinA, PinB);
+	bool bModified = false;
+
+	switch (Response.Response)
+	{
+	case CONNECT_RESPONSE_MAKE:
+		PinA->Modify();
+		PinB->Modify();
+		PinA->MakeLinkTo(PinB);
+		bModified = true;
+		break;
+
+	case CONNECT_RESPONSE_BREAK_OTHERS_A:
+		PinA->Modify();
+		PinB->Modify();
+		PinA->BreakAllPinLinks();
+		PinA->MakeLinkTo(PinB);
+		bModified = true;
+		break;
+
+	case CONNECT_RESPONSE_BREAK_OTHERS_B:
+		PinA->Modify();
+		PinB->Modify();
+		PinB->BreakAllPinLinks();
+		PinA->MakeLinkTo(PinB);
+		bModified = true;
+		break;
+
+	case CONNECT_RESPONSE_BREAK_OTHERS_AB:
+		PinA->Modify();
+		PinB->Modify();
+		PinA->BreakAllPinLinks();
+		PinB->BreakAllPinLinks();
+		PinA->MakeLinkTo(PinB);
+		bModified = true;
+		break;
+
+	case CONNECT_RESPONSE_MAKE_WITH_CONVERSION_NODE:
+	{
+		if (PinA->Direction == EGPD_Input)
+		{
+			//Swap so that A is the from pin and B is the to pin.
+			UEdGraphPin* Temp = PinA;
+			PinA = PinB;
+			PinB = Temp;
+		}
+
+		FNiagaraTypeDefinition AType = PinToTypeDefinition(PinA);
+		FNiagaraTypeDefinition BType = PinToTypeDefinition(PinB);
+		if (AType != BType && AType.GetClass() == nullptr && BType.GetClass() == nullptr)
+		{
+			UEdGraphNode* ANode = PinA->GetOwningNode();
+			UEdGraphNode* BNode = PinB->GetOwningNode();
+			UEdGraph* Graph = ANode->GetTypedOuter<UEdGraph>();
+			
+			// Since we'll be adding a node, make sure to modify the graph itself.
+			Graph->Modify();
+			FGraphNodeCreator<UNiagaraNodeConvert> NodeCreator(*Graph);
+			UNiagaraNodeConvert* AutoConvertNode = NodeCreator.CreateNode(false);
+			AutoConvertNode->AllocateDefaultPins();
+			AutoConvertNode->NodePosX = (ANode->NodePosX + BNode->NodePosX) >> 1;
+			AutoConvertNode->NodePosY = (ANode->NodePosY + BNode->NodePosY) >> 1;
+			NodeCreator.Finalize();
+
+			if (AutoConvertNode->InitConversion(PinA, PinB))
+			{
+				PinA->Modify();
+				PinB->Modify();
+				bModified = true;
+			}
+			else
+			{
+				Graph->RemoveNode(AutoConvertNode);
+			}
+		}
+	}
+	break;
+
+	case CONNECT_RESPONSE_DISALLOW:
+	default:
+		break;
+	}
+
+#if WITH_EDITOR
+	if (bModified)
+	{
+		PinA->GetOwningNode()->PinConnectionListChanged(PinA);
+		PinB->GetOwningNode()->PinConnectionListChanged(PinB);
+	}
+#endif	//#if WITH_EDITOR
+
+	return bModified;
+}
+
+FLinearColor UEdGraphSchema_Niagara::GetPinTypeColor(const FEdGraphPinType& PinType) const
+{
+	const UGraphEditorSettings* Settings = GetDefault<UGraphEditorSettings>();
+	if (PinType.PinCategory == PinCategoryType)
+	{
+		FNiagaraTypeDefinition Type(CastChecked<UScriptStruct>(PinType.PinSubCategoryObject.Get()));
+
+		if (Type == FNiagaraTypeDefinition::GetFloatDef())
+		{
+			return Settings->FloatPinTypeColor;
+		}
+		else if (Type == FNiagaraTypeDefinition::GetIntDef())
+		{
+			return Settings->IntPinTypeColor;
+		}
+		else if (Type == FNiagaraTypeDefinition::GetBoolDef())
+		{
+			return Settings->BooleanPinTypeColor;
+		}
+		else if (Type == FNiagaraTypeDefinition::GetVec2Def()
+			|| Type == FNiagaraTypeDefinition::GetVec3Def()
+			|| Type == FNiagaraTypeDefinition::GetVec4Def())
+		{
+			return Settings->VectorPinTypeColor;
+		}
+		else
+		{
+			return Settings->StructPinTypeColor;
+		}
+	}
+	else
+	{
+		return Settings->WildcardPinTypeColor;
+	}
+}
+
+bool UEdGraphSchema_Niagara::ShouldHidePinDefaultValue(UEdGraphPin* Pin) const
+{
+	check(Pin != NULL);
+
+	if (Pin->bDefaultValueIsIgnored)
+	{
+		return true;
+	}
+
+	return false;
+}
+
+FNiagaraVariable UEdGraphSchema_Niagara::PinToNiagaraVariable(const UEdGraphPin* Pin, bool bNeedsValue)const
+{
+	FNiagaraVariable Var = FNiagaraVariable(PinToTypeDefinition(Pin), *Pin->PinName);
+	if (Pin->DefaultValue.IsEmpty() == false)
+	{
+		FNiagaraEditorModule& NiagaraEditorModule = FModuleManager::GetModuleChecked<FNiagaraEditorModule>("NiagaraEditor");
+		TSharedPtr<INiagaraEditorTypeUtilities> TypeEditorUtilities = NiagaraEditorModule.GetTypeUtilities(Var.GetType());
+		if (TypeEditorUtilities.IsValid() && TypeEditorUtilities->CanHandlePinDefaults())
+		{
+			TypeEditorUtilities->SetValueFromPinDefaultString(Pin->DefaultValue, Var);
+		}
+		else if (const UEnum* Enum = Var.GetType().GetEnum())
+		{
+			int64 EnumValue = Enum->GetValueByNameString(Pin->DefaultValue, true);
+			if (EnumValue != INDEX_NONE)
+			{
+				Var.SetValue<int32>((int32)EnumValue);
+			}
+		}
+		else
+		{
+			UE_LOG(LogNiagaraEditor, Warning, TEXT("Pin default values aren't supported for variables of type {%s}"), *Var.GetType().GetName());
+		}
+	}
+	return Var;
+}
+
+FString UEdGraphSchema_Niagara::NiagaraVariableToPinDefaultValue(const FNiagaraVariable& Variable) const
+{
+	FNiagaraEditorModule& NiagaraEditorModule = FModuleManager::GetModuleChecked<FNiagaraEditorModule>("NiagaraEditor");
+	TSharedPtr<INiagaraEditorTypeUtilities> TypeEditorUtilities = NiagaraEditorModule.GetTypeUtilities(Variable.GetType());
+	if (TypeEditorUtilities.IsValid() && TypeEditorUtilities->CanHandlePinDefaults())
+	{
+		return TypeEditorUtilities->GetPinDefaultStringFromValue(Variable);
+	}
+	else
+	{
+		UE_LOG(LogNiagaraEditor, Warning, TEXT("Pin default values aren't supported for variables of type {%s}"), *Variable.GetType().GetName());
+		return FString();
+	}
+}
+
+FNiagaraTypeDefinition UEdGraphSchema_Niagara::PinToTypeDefinition(const UEdGraphPin* Pin) const
+{
+	if (Pin->PinType.PinCategory == PinCategoryType && Pin->PinType.PinSubCategoryObject != nullptr)
+	{
+		return CastChecked<const UScriptStruct>(Pin->PinType.PinSubCategoryObject.Get());
+	}
+	else if (Pin->PinType.PinCategory == PinCategoryClass)
+	{
+		return FNiagaraTypeDefinition(CastChecked<const UClass>(Pin->PinType.PinSubCategoryObject.Get()));
+	}
+	else if (Pin->PinType.PinCategory == PinCategoryEnum)
+	{
+		const UEnum* Enum = Cast<const UEnum>(Pin->PinType.PinSubCategoryObject.Get());
+		if (Enum == nullptr)
+		{
+			UE_LOG(LogNiagaraEditor, Error, TEXT("Pin states that it is of Enum type, but is missing its Enum! Pin Name '%s' Owning Node '%s'. Turning into standard int definition!"), *Pin->PinName,
+				*Pin->GetOwningNode()->GetName());
+			return FNiagaraTypeDefinition(FNiagaraTypeDefinition::GetIntDef());
+		}
+		return FNiagaraTypeDefinition(Enum);
+	}
+	return FNiagaraTypeDefinition();
+}
+
+FEdGraphPinType UEdGraphSchema_Niagara::TypeDefinitionToPinType(FNiagaraTypeDefinition TypeDef)const
+{
+	if (TypeDef.GetClass())
+	{
+		return FEdGraphPinType(PinCategoryClass, FString(), const_cast<UClass*>(TypeDef.GetClass()), EPinContainerType::None, false, FEdGraphTerminalType());
+	}
+	else if (TypeDef.GetEnum())
+	{
+		return FEdGraphPinType(PinCategoryEnum, FString(), const_cast<UEnum*>(TypeDef.GetEnum()), EPinContainerType::None, false, FEdGraphTerminalType());
+	}
+	else
+	{
+		//TODO: Are base types better as structs or done like BPS as a special name?
+		return FEdGraphPinType(PinCategoryType, FString(), const_cast<UScriptStruct*>(TypeDef.GetScriptStruct()), EPinContainerType::None, false, FEdGraphTerminalType());
+	}
+}
+
+bool UEdGraphSchema_Niagara::IsSystemConstant(const FNiagaraVariable& Variable)const
+{
+	return FNiagaraConstants::GetEngineConstants().Find(Variable) != INDEX_NONE;
+}
+
+FNiagaraTypeDefinition UEdGraphSchema_Niagara::GetTypeDefForProperty(const UProperty* Property)const
+{
+	if (Property->IsA(UFloatProperty::StaticClass()))
+	{
+		return FNiagaraTypeDefinition::GetFloatDef();
+	}
+	else if (Property->IsA(UIntProperty::StaticClass()))
+	{
+		return FNiagaraTypeDefinition::GetIntDef();
+	}
+	else if (Property->IsA(UBoolProperty::StaticClass()))
+	{
+		return FNiagaraTypeDefinition::GetBoolDef();
+	}	
+	else if (Property->IsA(UEnumProperty::StaticClass()))
+	{
+		const UEnumProperty* EnumProp = Cast<UEnumProperty>(Property);
+		return FNiagaraTypeDefinition(EnumProp->GetEnum());
+	}
+	else if (const UStructProperty* StructProp = CastChecked<UStructProperty>(Property))
+	{
+		return FNiagaraTypeDefinition(StructProp->Struct);
+	}
+
+	check(0);
+	return FNiagaraTypeDefinition::GetFloatDef();//Some invalid type?
+}
+
+void UEdGraphSchema_Niagara::GetBreakLinkToSubMenuActions(class FMenuBuilder& MenuBuilder, UEdGraphPin* InGraphPin)
+{
+	// Make sure we have a unique name for every entry in the list
+	TMap< FString, uint32 > LinkTitleCount;
+
+	// Add all the links we could break from
+	for (TArray<class UEdGraphPin*>::TConstIterator Links(InGraphPin->LinkedTo); Links; ++Links)
+	{
+		UEdGraphPin* Pin = *Links;
+		FString TitleString = Pin->GetOwningNode()->GetNodeTitle(ENodeTitleType::ListView).ToString();
+		FText Title = FText::FromString(TitleString);
+		if (!Pin->PinName.IsEmpty())
+		{
+			TitleString = FString::Printf(TEXT("%s (%s)"), *TitleString, *Pin->PinName);
+
+			// Add name of connection if possible
+			FFormatNamedArguments Args;
+			Args.Add(TEXT("NodeTitle"), Title);
+			Args.Add(TEXT("PinName"), Pin->GetDisplayName());
+			Title = FText::Format(LOCTEXT("BreakDescPin", "{NodeTitle} ({PinName})"), Args);
+		}
+
+		uint32 &Count = LinkTitleCount.FindOrAdd(TitleString);
+
+		FText Description;
+		FFormatNamedArguments Args;
+		Args.Add(TEXT("NodeTitle"), Title);
+		Args.Add(TEXT("NumberOfNodes"), Count);
+
+		if (Count == 0)
+		{
+			Description = FText::Format(LOCTEXT("BreakDesc", "Break link to {NodeTitle}"), Args);
+		}
+		else
+		{
+			Description = FText::Format(LOCTEXT("BreakDescMulti", "Break link to {NodeTitle} ({NumberOfNodes})"), Args);
+		}
+		++Count;
+		MenuBuilder.AddMenuEntry(Description, Description, FSlateIcon(), FUIAction(
+			FExecuteAction::CreateUObject((UEdGraphSchema_Niagara*const)this, &UEdGraphSchema::BreakSinglePinLink, const_cast<UEdGraphPin*>(InGraphPin), *Links)));
+	}
+}
+
+void UEdGraphSchema_Niagara::ConvertNumericPinToType(UEdGraphPin* InGraphPin, FNiagaraTypeDefinition TypeDef)
+{
+	if (PinToTypeDefinition(InGraphPin) != TypeDef)
+	{
+		UNiagaraNode* Node = Cast<UNiagaraNode>(InGraphPin->GetOwningNode());
+		if (Node)
+		{
+			FScopedTransaction Transaction(NSLOCTEXT("UnrealEd", "NiagaraEditorChangeNumericPinType", "Change Pin Type"));
+			if (false == Node->ConvertNumericPinToType(InGraphPin, TypeDef))
+			{
+				Transaction.Cancel();
+			}
+		}
+	}
+}
+
+void UEdGraphSchema_Niagara::GetNumericConversionToSubMenuActions(class FMenuBuilder& MenuBuilder, UEdGraphPin* InGraphPin)
+{
+	// Add all the types we could convert to
+	for (const FNiagaraTypeDefinition& TypeDef : FNiagaraTypeRegistry::GetNumericTypes())
+	{
+		FText Title = TypeDef.GetNameText();
+
+		FText Description;
+		FFormatNamedArguments Args;
+		Args.Add(TEXT("TypeTitle"), Title);
+		Description = FText::Format(LOCTEXT("NumericConversionText", "{TypeTitle}"), Args);
+		MenuBuilder.AddMenuEntry(Description, Description, FSlateIcon(), FUIAction(
+			FExecuteAction::CreateUObject((UEdGraphSchema_Niagara*const)this, &UEdGraphSchema_Niagara::ConvertNumericPinToType, const_cast<UEdGraphPin*>(InGraphPin), FNiagaraTypeDefinition(TypeDef))));
+	}
+}
+
+void UEdGraphSchema_Niagara::PromoteSinglePinToParameter(UEdGraphPin* SourcePin)
+{
+	if (SourcePin)
+	{
+		const FScopedTransaction Transaction(NSLOCTEXT("UnrealEd", "NiagaraEditorPromote", "Promote To Parameter"));
+		{
+			TSharedPtr<FNiagaraSchemaAction_NewNode> InputAction = TSharedPtr<FNiagaraSchemaAction_NewNode>(new FNiagaraSchemaAction_NewNode(FText::GetEmpty(), FText::GetEmpty(), NAME_None, FText::GetEmpty(), 0));
+			UNiagaraNodeInput* InputNode = NewObject<UNiagaraNodeInput>(GetTransientPackage());
+			FNiagaraVariable Var = PinToNiagaraVariable(SourcePin);
+			UNiagaraGraph* Graph = Cast<UNiagaraGraph>(SourcePin->GetOwningNode()->GetGraph());
+			FNiagaraEditorUtilities::InitializeParameterInputNode(*InputNode, Var.GetType(), Graph);
+			InputAction->NodeTemplate = InputNode;
+
+			UEdGraphNode* PinNode = SourcePin->GetOwningNode();
+
+			const float PinVisualOffsetX = 175.0f;
+			InputAction->PerformAction(Graph, SourcePin, FVector2D(PinNode->NodePosX - PinVisualOffsetX, PinNode->NodePosY));
+		}
+	}
+}
+
+void UEdGraphSchema_Niagara::GetContextMenuActions(const UEdGraph* CurrentGraph, const UEdGraphNode* InGraphNode, const UEdGraphPin* InGraphPin, class FMenuBuilder* MenuBuilder, bool bIsDebugging) const
+{
+	if (InGraphPin)
+	{
+		MenuBuilder->BeginSection("EdGraphSchema_NiagaraPinActions", LOCTEXT("PinActionsMenuHeader", "Pin Actions"));
+		{
+			if (PinToTypeDefinition(InGraphPin) == FNiagaraTypeDefinition::GetGenericNumericDef() && InGraphPin->LinkedTo.Num() == 0)
+			{
+				MenuBuilder->AddSubMenu(
+					LOCTEXT("ConvertNumericSpecific", "Convert Numeric To..."),
+					LOCTEXT("ConvertNumericSpecificToolTip", "Convert Numeric pin to specific typed pin."),
+				FNewMenuDelegate::CreateUObject((UEdGraphSchema_Niagara*const)this, &UEdGraphSchema_Niagara::GetNumericConversionToSubMenuActions, const_cast<UEdGraphPin*>(InGraphPin)));
+			}
+
+			// Only display the 'Break Link' option if there is a link to break!
+			if (InGraphPin->LinkedTo.Num() > 0)
+			{
+				MenuBuilder->AddMenuEntry(FGraphEditorCommands::Get().BreakPinLinks);
+
+				// add sub menu for break link to
+				if (InGraphPin->LinkedTo.Num() > 1)
+				{
+					MenuBuilder->AddSubMenu(
+						LOCTEXT("BreakLinkTo", "Break Link To..."),
+						LOCTEXT("BreakSpecificLinks", "Break a specific link..."),
+						FNewMenuDelegate::CreateUObject((UEdGraphSchema_Niagara*const)this, &UEdGraphSchema_Niagara::GetBreakLinkToSubMenuActions, const_cast<UEdGraphPin*>(InGraphPin)));
+				}
+				else
+				{
+					((UEdGraphSchema_Niagara*const)this)->GetBreakLinkToSubMenuActions(*MenuBuilder, const_cast<UEdGraphPin*>(InGraphPin));
+				}
+			}
+
+			if (InGraphPin->Direction == EEdGraphPinDirection::EGPD_Input)
+			{
+				MenuBuilder->AddMenuEntry(LOCTEXT("PromoteToParameter", "Promote to Parameter"), LOCTEXT("PromoteToParameterTooltip", "Create a parameter argument and connect this pin to that parameter."), FSlateIcon(),
+					FUIAction(FExecuteAction::CreateUObject((UEdGraphSchema_Niagara*const)this, &UEdGraphSchema_Niagara::PromoteSinglePinToParameter, const_cast<UEdGraphPin*>(InGraphPin))));
+			}
+		}
+		MenuBuilder->EndSection();
+	}
+	else if (InGraphNode)
+	{
+	}
+
+	Super::GetContextMenuActions(CurrentGraph, InGraphNode, InGraphPin, MenuBuilder, bIsDebugging);
+}
+
+FNiagaraConnectionDrawingPolicy::FNiagaraConnectionDrawingPolicy(int32 InBackLayerID, int32 InFrontLayerID, float InZoomFactor, const FSlateRect& InClippingRect, FSlateWindowElementList& InDrawElements)
+	: FConnectionDrawingPolicy(InBackLayerID, InFrontLayerID, InZoomFactor, InClippingRect, InDrawElements)
+{
+}
+
+void FNiagaraConnectionDrawingPolicy::DetermineWiringStyle(UEdGraphPin* OutputPin, UEdGraphPin* InputPin, /*inout*/ FConnectionParams& Params)
+{
+	FConnectionDrawingPolicy::DetermineWiringStyle(OutputPin, InputPin, Params);
+	if (HoveredPins.Contains(InputPin) && HoveredPins.Contains(OutputPin))
+	{
+		Params.WireThickness = Params.WireThickness * 5;
+	}
+}
+
+#undef LOCTEXT_NAMESPACE
diff --git a/Engine/Plugins/FX/Niagara/Source/NiagaraEditor/Private/NiagaraCompiler.cpp b/Engine/Plugins/FX/Niagara/Source/NiagaraEditor/Private/NiagaraCompiler.cpp
new file mode 100644
--- /dev/null
+++ b/Engine/Plugins/FX/Niagara/Source/NiagaraEditor/Private/NiagaraCompiler.cpp
@@ -0,0 +1,343 @@
+// Copyright 1998-2017 Epic Games, Inc. All Rights Reserved.
+
+#include "NiagaraCompiler.h"
+#include "NiagaraHlslTranslator.h"
+#include "NiagaraScriptSource.h"
+#include "NiagaraGraph.h"
+#include "NiagaraEditorModule.h"
+#include "NiagaraComponent.h"
+#include "Interfaces/ITargetPlatformManagerModule.h"
+#include "Interfaces/IShaderFormat.h"
+#include "ShaderFormatVectorVM.h"
+#include "NiagaraConstants.h"
+#include "NiagaraSystem.h"
+#include "NiagaraNodeEmitter.h"
+#include "NiagaraNodeInput.h"
+
+#include "NiagaraDataInterface.h"
+#include "NiagaraDataInterfaceStaticMesh.h"
+#include "NiagaraDataInterfaceCurlNoise.h"
+
+#include "ShaderCore.h"
+
+#define LOCTEXT_NAMESPACE "NiagaraCompiler"
+
+DEFINE_LOG_CATEGORY_STATIC(LogNiagaraCompiler, All, All);
+
+
+ENiagaraScriptCompileStatus FNiagaraEditorModule::CompileScript(UNiagaraScript* ScriptToCompile, FString& OutGraphLevelErrorMessages)
+{
+	check(ScriptToCompile != NULL);
+
+	UNiagaraGraph* Graph = Cast<UNiagaraScriptSource>(ScriptToCompile->GetSource())->NodeGraph;
+
+	OutGraphLevelErrorMessages.Empty();
+	FNiagaraCompileResults Results;
+	FHlslNiagaraCompiler Compiler;
+	FNiagaraTranslateResults TranslateResults;
+	FHlslNiagaraTranslator Translator;
+
+	/*
+	if (ScriptToCompile && ScriptToCompile->IsNonParticleScript())
+	{
+		ENiagaraScriptCompileStatus Status = ENiagaraScriptCompileStatus::NCS_UpToDate;
+		if (ScriptToCompile->IsSystemSpawnScript())
+		{
+			Status = CompileSystemScript(ScriptToCompile, OutGraphLevelErrorMessages);
+		}
+		ScriptToCompile->SetChangeID(Graph->GetChangeID());
+		ScriptToCompile->SetLastCompileStatus(Status);
+		return Status;
+	}
+	else
+	*/
+	{
+		FHlslNiagaraTranslatorOptions Options; 
+		Options.SimTarget = ENiagaraSimTarget::CPUSim;
+		TranslateResults = Translator.Translate(ScriptToCompile, Options);
+		Results = Compiler.CompileScript(ScriptToCompile, &Translator.GetTranslateOutput(), Translator.GetTranslatedHLSL());
+		ScriptToCompile->SetChangeID(Graph->GetChangeID());
+		ScriptToCompile->GenerateStatScopeIDs();
+	}
+
+#if defined(NIAGARA_SCRIPT_COMPILE_LOGGING_MEDIUM)
+	if (TranslateResults.bHLSLGenSucceeded)
+	{
+		UE_LOG(LogNiagaraCompiler, Log, TEXT("HLSL generation succeeded: %s"), *ScriptToCompile->GetPathName());
+	}
+	else
+	{
+		UE_LOG(LogNiagaraCompiler, Error, TEXT("HLSL generation failed: %s"), *ScriptToCompile->GetPathName());
+	}
+
+	if (Results.bVMSucceeded)
+	{
+		UE_LOG(LogNiagaraCompiler, Log, TEXT("CPU Compile succeeded: %s"), *ScriptToCompile->GetPathName());
+	}
+	else
+	{
+		UE_LOG(LogNiagaraCompiler, Error, TEXT("CPU Compile failed: %s"), *ScriptToCompile->GetPathName());
+	}
+#endif
+
+	TArray<TSharedRef<FTokenizedMessage>> Messages;
+	if (TranslateResults.MessageLog)
+	{
+		Messages.Append(TranslateResults.MessageLog->Messages);
+	}
+	if (Results.MessageLog)
+	{
+		Messages.Append(Results.MessageLog->Messages);
+	}
+
+	for (TSharedRef<FTokenizedMessage> Message : Messages)
+	{
+		if (Message->GetSeverity() == EMessageSeverity::Info)
+		{
+		#if defined(NIAGARA_SCRIPT_COMPILE_LOGGING_MEDIUM)
+			UE_LOG(LogNiagaraCompiler, Log, TEXT("%s"), *Message->ToText().ToString());
+		#endif
+		}
+		else if (Message->GetSeverity() == EMessageSeverity::Warning || Message->GetSeverity() == EMessageSeverity::PerformanceWarning)
+		{
+		#if defined(NIAGARA_SCRIPT_COMPILE_LOGGING_MEDIUM)
+			UE_LOG(LogNiagaraCompiler, Warning, TEXT("%s"), *Message->ToText().ToString());
+		#endif
+		}
+		else if (Message->GetSeverity() == EMessageSeverity::Error || Message->GetSeverity() == EMessageSeverity::CriticalError)
+		{
+		#if defined(NIAGARA_SCRIPT_COMPILE_LOGGING_MEDIUM)
+			UE_LOG(LogNiagaraCompiler, Error, TEXT("%s"), *Message->ToText().ToString());
+		#endif
+			// Write the error messages to the string as well so that they can be echoed up the chain.
+			if (OutGraphLevelErrorMessages.Len() > 0)
+			{
+				OutGraphLevelErrorMessages += "\n";
+			}
+			OutGraphLevelErrorMessages += Message->ToText().ToString();
+		}
+	}
+
+#if defined(NIAGARA_SCRIPT_COMPILE_LOGGING_MEDIUM)
+	UE_LOG(LogNiagaraCompiler, Log, TEXT("Compile output as text:"));
+	UE_LOG(LogNiagaraCompiler, Log, TEXT("==================================================================================="));
+	TArray<FString> OutputByLines;
+	Results.OutputHLSL.ParseIntoArrayLines(OutputByLines, false);
+	for (int32 i = 0; i < OutputByLines.Num(); i++)
+	{
+		UE_LOG(LogNiagaraCompiler, Log, TEXT("/*%04d*/\t\t%s"), i, *OutputByLines[i]);
+	}
+	UE_LOG(LogNiagaraCompiler, Log, TEXT("==================================================================================="));
+#endif
+
+	ScriptToCompile->SetLastCompileStatus(FNiagaraCompileResults::CompileResultsToSummary(&Results));
+	return ScriptToCompile->GetLastCompileStatus();
+}
+
+
+ENiagaraScriptCompileStatus FNiagaraCompileResults::CompileResultsToSummary(const FNiagaraCompileResults* CompileResults)
+{
+	ENiagaraScriptCompileStatus SummaryStatus = ENiagaraScriptCompileStatus::NCS_Unknown;
+	if (CompileResults != nullptr)
+	{
+		if (CompileResults->MessageLog->NumErrors > 0)
+		{
+			SummaryStatus = ENiagaraScriptCompileStatus::NCS_Error;
+		}
+		else
+		{
+			if (CompileResults->bVMSucceeded)
+			{
+				if (CompileResults->MessageLog->NumWarnings)
+				{
+					SummaryStatus = ENiagaraScriptCompileStatus::NCS_UpToDateWithWarnings;
+				}
+				else
+				{
+					SummaryStatus = ENiagaraScriptCompileStatus::NCS_UpToDate;
+				}
+			}
+
+			if (CompileResults->bComputeSucceeded)
+			{
+				if (CompileResults->MessageLog->NumWarnings)
+				{
+					SummaryStatus = ENiagaraScriptCompileStatus::NCS_ComputeUpToDateWithWarnings;
+				}
+				else
+				{
+					SummaryStatus = ENiagaraScriptCompileStatus::NCS_UpToDate;
+				}
+			}
+		}
+	}
+	return SummaryStatus;
+}
+
+const FNiagaraCompileResults &FHlslNiagaraCompiler::CompileScript(UNiagaraScript* InScript, FNiagaraTranslatorOutput *TranslatorOutput, FString &TranslatedHLSL)
+{
+	//TODO: This should probably be done via the same route that other shaders take through the shader compiler etc.
+	//But that adds the complexity of a new shader type, new shader class and a new shader map to contain them etc.
+	//Can do things simply for now.
+	Script = InScript;
+	Script->LastHlslTranslation = TEXT("");
+	FShaderCompilerInput Input;
+	Input.VirtualSourceFilePath = TEXT("/Engine/Private/NiagaraEmitterInstanceShader.usf");
+	Input.EntryPointName = TEXT("SimulateMain");
+	Input.Environment.SetDefine(TEXT("VM_SIMULATION"), 1);
+	Input.Environment.IncludeVirtualPathToContentsMap.Add(TEXT("/Engine/Generated/NiagaraEmitterInstance.usf"), StringToArray<ANSICHAR>(*TranslatedHLSL, TranslatedHLSL.Len() + 1));
+	FShaderCompilerOutput Output;
+
+	FVectorVMCompilationOutput CompilationOutput;
+
+	if (TranslatorOutput != nullptr && TranslatorOutput->Errors.Len() > 0)
+	{
+		//TODO: Map Lines of HLSL to their source Nodes and flag those nodes with errors associated with their lines.
+		Error(FText::Format(LOCTEXT("HlslTranslateErrorMessageFormat", "The HLSL Translator failed.  Errors:\n{0}"), FText::FromString(TranslatorOutput->Errors)));
+		CompileResults.bVMSucceeded = false;
+	}
+	else if (TranslatedHLSL.Len() == 0)
+	{
+		Error(LOCTEXT("HlslTranslateErrorMessageFormat", "The HLSL Translator failed to generate HLSL!"));
+		CompileResults.bVMSucceeded = false;
+	}
+	else
+	{
+		CompileResults.bVMSucceeded = CompileShader_VectorVM(Input, Output, FString(FPlatformProcess::ShaderDir()), 0, CompilationOutput);
+	}
+
+	if (CompilationOutput.Errors.Len() > 0)
+	{
+		//TODO: Map Lines of HLSL to their source Nodes and flag those nodes with errors associated with their lines.
+		Error(FText::Format(LOCTEXT("VectorVMCompileErrorMessageFormat", "The Vector VM compile failed.  Errors:\n{0}"), FText::FromString(CompilationOutput.Errors)));
+		CompileResults.bVMSucceeded = false;
+	}
+
+	//For now we just copy the shader code over into the script. 
+	//Eventually Niagara will have all the shader plumbing and do things like materials.
+	if (CompileResults.bVMSucceeded)
+	{
+		check(TranslatorOutput);
+		Script->ByteCode = CompilationOutput.ByteCode;
+		Script->LastHlslTranslation = TranslatedHLSL;
+		Script->Attributes = TranslatorOutput->Attributes;
+		Script->Parameters = TranslatorOutput->Parameters;
+		Script->DataUsage.bReadsAttriubteData = TranslatorOutput->bReadsAttributeData;
+		//Build internal parameters
+		Script->InternalParameters.Empty();
+		for (int32 i = 0; i < CompilationOutput.InternalConstantOffsets.Num(); ++i)
+		{
+			EVectorVMBaseTypes Type = CompilationOutput.InternalConstantTypes[i];
+			int32 Offset = CompilationOutput.InternalConstantOffsets[i];
+			switch (Type)
+			{
+			case EVectorVMBaseTypes::Float:
+			{
+				float Val = *(float*)(CompilationOutput.InternalConstantData.GetData() + Offset);
+				Script->InternalParameters.SetOrAdd(FNiagaraVariable(FNiagaraTypeDefinition::GetFloatDef(), *LexicalConversion::ToString(Val)))->SetValue(Val);
+			}
+			break;
+			case EVectorVMBaseTypes::Int:
+			{
+				int32 Val = *(int32*)(CompilationOutput.InternalConstantData.GetData() + Offset);
+				Script->InternalParameters.SetOrAdd(FNiagaraVariable(FNiagaraTypeDefinition::GetIntDef(), *LexicalConversion::ToString(Val)))->SetValue(Val);
+			}
+			break;
+			case EVectorVMBaseTypes::Bool:
+			{
+				int32 Val = *(int32*)(CompilationOutput.InternalConstantData.GetData() + Offset);
+				Script->InternalParameters.SetOrAdd(FNiagaraVariable(FNiagaraTypeDefinition::GetIntDef(), Val == 0 ? TEXT("FALSE") : TEXT("TRUE")))->SetValue(Val);
+			}
+			break;
+			}
+		}
+
+		Script->DataInterfaceInfo.Empty();
+		for (FNiagaraScriptDataInterfaceInfo& DataInterfaceInfo : TranslatorOutput->DataInterfaceInfo)
+		{
+			int32 Idx = Script->DataInterfaceInfo.AddDefaulted();
+			Script->DataInterfaceInfo[Idx].DataInterface = Cast<UNiagaraDataInterface>(StaticDuplicateObject(DataInterfaceInfo.DataInterface, Script, NAME_None, ~RF_Transient));
+			Script->DataInterfaceInfo[Idx].Name = DataInterfaceInfo.Name;
+			Script->DataInterfaceInfo[Idx].UserPtrIdx = DataInterfaceInfo.UserPtrIdx;
+		}
+
+		Script->NumUserPtrs = TranslatorOutput->NumUserPtrs;
+
+		//Extract the external function call table binding info.
+		Script->CalledVMExternalFunctions.Empty(CompilationOutput.CalledVMFunctionTable.Num());
+		for (FVectorVMCompilationOutput::FCalledVMFunction& FuncInfo : CompilationOutput.CalledVMFunctionTable)
+		{
+			//Find the interface corresponding to this call.
+			const FNiagaraFunctionSignature* Sig = nullptr;
+			for (FNiagaraScriptDataInterfaceInfo& NDIInfo : TranslatorOutput->DataInterfaceInfo)
+			{			
+				Sig = NDIInfo.RegisteredFunctions.FindByPredicate([&](const FNiagaraFunctionSignature& CheckSig)
+				{
+					FString SigSymbol = FHlslNiagaraTranslator::GetFunctionSignatureSymbol(CheckSig);
+					return SigSymbol == FuncInfo.Name;
+				});
+				if (Sig)
+				{
+					break;
+				}
+			}
+
+			if (Sig)
+			{
+				int32 NewBindingIdx = Script->CalledVMExternalFunctions.AddDefaulted();
+				Script->CalledVMExternalFunctions[NewBindingIdx].Name = *Sig->GetName();
+				Script->CalledVMExternalFunctions[NewBindingIdx].OwnerName = Sig->OwnerName;
+
+				Script->CalledVMExternalFunctions[NewBindingIdx].InputParamLocations = FuncInfo.InputParamLocations;
+				Script->CalledVMExternalFunctions[NewBindingIdx].NumOutputs = FuncInfo.NumOutputs;
+			}
+			else
+			{
+				Error(FText::Format(LOCTEXT("VectorVMExternalFunctionBindingError", "Failed to bind the exernal function call:  {0}"), FText::FromString(FuncInfo.Name)));
+				CompileResults.bVMSucceeded = false;
+			}
+		}
+	}
+
+	CompileResults.OutputHLSL = TranslatedHLSL;
+
+	if (CompileResults.bVMSucceeded == false)
+	{
+		//Some error. Clear script and exit.
+		Script->ByteCode.Empty();
+		Script->Attributes.Empty();
+		Script->Parameters.Empty();
+		Script->InternalParameters.Empty();
+		Script->DataInterfaceInfo.Empty();
+//		Script->NumUserPtrs = 0;
+	}
+
+	return CompileResults;
+}
+
+
+FHlslNiagaraCompiler::FHlslNiagaraCompiler()
+	: Script(nullptr)
+	, CompileResults(&MessageLog)
+{
+	// Make the message log silent so we're not spamming the blueprint log.
+	MessageLog.bSilentMode = true;
+}
+
+
+
+void FHlslNiagaraCompiler::Error(FText ErrorText)
+{
+	FString ErrorString = FString::Printf(TEXT("%s"), *ErrorText.ToString());
+	MessageLog.Error(*ErrorString);
+}
+
+void FHlslNiagaraCompiler::Warning(FText WarningText)
+{
+	FString WarnString = FString::Printf(TEXT("%s"), *WarningText.ToString());
+	MessageLog.Warning(*WarnString);
+}
+
+//////////////////////////////////////////////////////////////////////////
+
+#undef LOCTEXT_NAMESPACE
diff --git a/Engine/Plugins/FX/Niagara/Source/NiagaraEditor/Private/NiagaraCompiler.h b/Engine/Plugins/FX/Niagara/Source/NiagaraEditor/Private/NiagaraCompiler.h
new file mode 100644
--- /dev/null
+++ b/Engine/Plugins/FX/Niagara/Source/NiagaraEditor/Private/NiagaraCompiler.h
@@ -0,0 +1,41 @@
+// Copyright 1998-2017 Epic Games, Inc. All Rights Reserved.
+
+#pragma once
+
+#include "CoreMinimal.h"
+#include "NiagaraCommon.h"
+#include "NiagaraHlslTranslator.h"
+#include "INiagaraCompiler.h"
+#include "Kismet2/CompilerResultsLog.h"
+
+class Error;
+class UNiagaraGraph;
+class UNiagaraNode;
+class UNiagaraNodeFunctionCall;
+class UNiagaraScript;
+class UNiagaraScriptSource;
+struct FNiagaraTranslatorOutput;
+
+class NIAGARAEDITOR_API FHlslNiagaraCompiler : public INiagaraCompiler
+{
+protected:
+	/** The script we are compiling. */
+	UNiagaraScript* Script;
+
+	/** Message log. Automatically handles marking the NodeGraph with errors. */
+	FCompilerResultsLog MessageLog;
+
+	/** Captures information about a script compile. */
+	FNiagaraCompileResults CompileResults;
+
+public:
+
+	FHlslNiagaraCompiler();
+	virtual ~FHlslNiagaraCompiler() {}
+
+	//Begin INiagaraCompiler Interface
+	virtual const FNiagaraCompileResults& CompileScript(UNiagaraScript* InScript, FNiagaraTranslatorOutput *TranslatorOutput, FString& TranslatedHLSL) override;
+
+	virtual void Error(FText ErrorText) override;
+	virtual void Warning(FText WarningText) override;
+};
diff --git a/Engine/Plugins/FX/Niagara/Source/NiagaraEditor/Private/NiagaraCurveOwner.cpp b/Engine/Plugins/FX/Niagara/Source/NiagaraEditor/Private/NiagaraCurveOwner.cpp
new file mode 100644
--- /dev/null
+++ b/Engine/Plugins/FX/Niagara/Source/NiagaraEditor/Private/NiagaraCurveOwner.cpp
@@ -0,0 +1,126 @@
+// Copyright 1998-2017 Epic Games, Inc. All Rights Reserved.
+
+#include "NiagaraCurveOwner.h"
+
+FNiagaraCurveOwner::FNiagaraCurveOwner()
+	: bIsColorCurve(false)
+{
+}
+
+void FNiagaraCurveOwner::EmptyCurves()
+{
+	ConstCurves.Empty();
+	Curves.Empty();
+	EditInfoToOwnerMap.Empty();
+	EditInfoToColorMap.Empty();
+	EditInfoToNotifyCurveChangedMap.Empty();
+	bIsColorCurve = false;
+}
+
+void FNiagaraCurveOwner::AddCurve(FRichCurve& Curve, FName Name, FLinearColor Color, UObject& Owner, FNotifyCurveChanged CurveChangedHandler)
+{
+	FRichCurveEditInfo EditInfo(&Curve, Name);
+	Curves.Add(EditInfo);
+	ConstCurves.Add(FRichCurveEditInfoConst(&Curve, Name));
+	EditInfoToOwnerMap.Add(EditInfo, &Owner);
+	EditInfoToColorMap.Add(EditInfo, Color);
+	EditInfoToNotifyCurveChangedMap.Add(EditInfo, CurveChangedHandler);
+	bIsColorCurve = false;
+}
+
+void FNiagaraCurveOwner::SetColorCurves(FRichCurve& RedCurve, FRichCurve& GreenCurve, FRichCurve& BlueCurve, FRichCurve& AlphaCurve, FName Name, UObject& Owner, FNotifyCurveChanged CurveChangedHandler)
+{
+	EmptyCurves();
+	FString NamePrefix = Name != NAME_None ? Name.ToString() + TEXT(".") : FString();
+	AddCurve(RedCurve, *(NamePrefix + TEXT("Red")), FLinearColor::Red, Owner, CurveChangedHandler);
+	AddCurve(GreenCurve, *(NamePrefix + TEXT("Green")), FLinearColor::Green, Owner, CurveChangedHandler);
+	AddCurve(BlueCurve, *(NamePrefix + TEXT("Blue")), FLinearColor::Blue, Owner, CurveChangedHandler);
+	AddCurve(AlphaCurve, *(NamePrefix + TEXT("Alpha")), FLinearColor::White, Owner, CurveChangedHandler);
+	bIsColorCurve = true;
+}
+
+TArray<FRichCurveEditInfoConst> FNiagaraCurveOwner::GetCurves() const
+{
+	return ConstCurves;
+}
+
+TArray<FRichCurveEditInfo> FNiagaraCurveOwner::GetCurves()
+{
+	return Curves;
+};
+
+void FNiagaraCurveOwner::ModifyOwner()
+{
+	TArray<UObject*> Owners;
+	EditInfoToOwnerMap.GenerateValueArray(Owners);
+	for (UObject* Owner : Owners)
+	{
+		Owner->Modify();
+	}
+}
+
+TArray<const UObject*> FNiagaraCurveOwner::GetOwners() const
+{
+	TArray<UObject*> Owners;
+	EditInfoToOwnerMap.GenerateValueArray(Owners);
+	TArray<const UObject*> ConstOwners;
+	ConstOwners.Append(Owners);
+	return ConstOwners;
+}
+
+void FNiagaraCurveOwner::MakeTransactional()
+{
+	TArray<UObject*> Owners;
+	EditInfoToOwnerMap.GenerateValueArray(Owners);
+	for (UObject* Owner : Owners)
+	{
+		Owner->SetFlags(RF_Transactional);
+	}
+}
+
+void FNiagaraCurveOwner::OnCurveChanged(const TArray<FRichCurveEditInfo>& ChangedCurveEditInfos)
+{
+	for (const FRichCurveEditInfo& ChangedCurveEditInfo : ChangedCurveEditInfos)
+	{
+		FNotifyCurveChanged* CurveChanged = EditInfoToNotifyCurveChangedMap.Find(ChangedCurveEditInfo);
+		UObject** CurveOwner = EditInfoToOwnerMap.Find(ChangedCurveEditInfo);
+		if (CurveChanged != nullptr && CurveOwner != nullptr)
+		{
+			CurveChanged->Execute(ChangedCurveEditInfo.CurveToEdit, *CurveOwner);
+		}
+	}
+}
+
+bool FNiagaraCurveOwner::IsLinearColorCurve() const
+{
+	return bIsColorCurve;
+}
+
+FLinearColor FNiagaraCurveOwner::GetLinearColorValue(float InTime) const
+{
+	return FLinearColor(
+		Curves[0].CurveToEdit->Eval(InTime),
+		Curves[1].CurveToEdit->Eval(InTime),
+		Curves[2].CurveToEdit->Eval(InTime),
+		Curves[3].CurveToEdit->Eval(InTime));
+}
+
+bool FNiagaraCurveOwner::HasAnyAlphaKeys() const
+{
+	return bIsColorCurve && Curves[3].CurveToEdit->GetNumKeys() > 0;
+}
+
+bool FNiagaraCurveOwner::IsValidCurve(FRichCurveEditInfo CurveInfo)
+{
+	return EditInfoToOwnerMap.Contains(CurveInfo);
+}
+
+FLinearColor FNiagaraCurveOwner::GetCurveColor(FRichCurveEditInfo CurveInfo) const
+{
+	const FLinearColor* Color = EditInfoToColorMap.Find(CurveInfo);
+	if (Color != nullptr)
+	{
+		return *Color;
+	}
+	return FCurveOwnerInterface::GetCurveColor(CurveInfo);
+}
\ No newline at end of file
