Copying //UE4/Orion-Staging to //UE4/Main (Source: //Orion/Dev-General @ 3606378)

diff --git a/Engine/Plugins/FX/Niagara/Source/NiagaraEditor/Private/NiagaraEditorCommon.cpp b/Engine/Plugins/FX/Niagara/Source/NiagaraEditor/Private/NiagaraEditorCommon.cpp
new file mode 100644
--- /dev/null
+++ b/Engine/Plugins/FX/Niagara/Source/NiagaraEditor/Private/NiagaraEditorCommon.cpp
@@ -0,0 +1,926 @@
+// Copyright 1998-2017 Epic Games, Inc. All Rights Reserved.
+
+#include "NiagaraEditorCommon.h"
+#include "AssetData.h"
+#include "INiagaraCompiler.h"
+#include "NiagaraSystem.h"
+
+#include "ActorFactoryNiagara.h"
+#include "NiagaraActor.h"
+#include "NiagaraComponent.h"
+
+#define LOCTEXT_NAMESPACE "NiagaraEditor"
+
+DEFINE_LOG_CATEGORY(LogNiagaraEditor);
+TMap<FName, int32> FNiagaraOpInfo::OpInfoMap;
+TArray<FNiagaraOpInfo> FNiagaraOpInfo::OpInfos;
+
+const FNiagaraOpInfo* FNiagaraOpInfo::GetOpInfo(FName OpName)
+{
+	const int32* Idx = OpInfoMap.Find(OpName);
+	if (Idx)
+	{
+		return &OpInfos[*Idx];
+	}
+	return nullptr;
+}
+
+const TArray<FNiagaraOpInfo>& FNiagaraOpInfo::GetOpInfoArray()
+{
+	return OpInfos;
+}
+
+void FNiagaraOpInfo::BuildName(FString InName, FString InCategory)
+{	
+	Name = FName(*(InCategory + TEXT("::") + InName));
+}
+
+BEGIN_FUNCTION_BUILD_OPTIMIZATION
+void FNiagaraOpInfo::Init()
+{
+	//Common input and output names.
+ 	static FName Result(TEXT("Result"));
+	static FText ResultText = NSLOCTEXT("NiagaraOpInfo", "Operation Result", "Result");
+		
+	static FName In(TEXT("In"));
+	static FText InText = NSLOCTEXT("NiagaraOpInfo", "Single Function Param", "In"); 
+
+	static FName A(TEXT("A"));
+	static FText AText = NSLOCTEXT("NiagaraOpInfo", "First Function Param", "A");
+
+	static FName B(TEXT("B"));
+	static FText BText = NSLOCTEXT("NiagaraOpInfo", "Second Function Param", "B");
+
+	static FName C(TEXT("C"));
+	static FText CText = NSLOCTEXT("NiagaraOpInfo", "Third Function Param", "C");
+
+	FName X(TEXT("X"));	FText XText = NSLOCTEXT("NiagaraOpInfo", "First Vector Component", "X");
+	FName Y(TEXT("Y"));	FText YText = NSLOCTEXT("NiagaraOpInfo", "Second Vector Component", "Y");
+	FName Z(TEXT("Z"));	FText ZText = NSLOCTEXT("NiagaraOpInfo", "Third Vector Component", "Z");
+	FName W(TEXT("W"));	FText WText = NSLOCTEXT("NiagaraOpInfo", "Fourth Vector Component", "W");
+
+	static FName Selection(TEXT("Selection"));	
+	static FText SelectionText = NSLOCTEXT("NiagaraOpInfo", "Selection", "Selection");
+	static FName Mask(TEXT("Mask"));	
+	static FText MaskText = NSLOCTEXT("NiagaraOpInfo", "Mask", "Mask");
+	
+	static FName Min(TEXT("Min"));
+	static FName Max(TEXT("Max"));
+	static FText MinText = NSLOCTEXT("NiagaraOpInfo", "Min", "Min");
+	static FText MaxText = NSLOCTEXT("NiagaraOpInfo", "Max", "Max");
+
+	//Add all numeric ops
+	{
+		FNiagaraTypeDefinition Type = FNiagaraTypeDefinition::GetGenericNumericDef();
+		FString DefaultStr_Zero(TEXT("0.0"));
+		FString DefaultStr_One(TEXT("1.0"));
+		FText CategoryText = NSLOCTEXT("NiagaraOpInfo", "NumericOpCategory", "Numeric");
+		FString CategoryName(TEXT("Numeric"));
+
+		int32 Idx = OpInfos.AddDefaulted();
+		FNiagaraOpInfo* Op = &OpInfos[Idx];
+		Op->Category = CategoryText;
+		Op->FriendlyName = NSLOCTEXT("NiagaraOpInfo", "Add Name", "Add");
+		Op->Description = NSLOCTEXT("NiagaraOpInfo", "Add Desc", "Result = A + B");
+		Op->Keywords = FText::FromString(TEXT("+"));
+		Op->Inputs.Add(FNiagaraOpInOutInfo(A, Type, AText, AText, DefaultStr_Zero));
+		Op->Inputs.Add(FNiagaraOpInOutInfo(B, Type, BText, BText, DefaultStr_Zero));
+		Op->Outputs.Add(FNiagaraOpInOutInfo(Result, Type, ResultText, ResultText, DefaultStr_Zero, TEXT("{0} + {1}")));
+		Op->BuildName(TEXT("Add"), CategoryName);
+		OpInfoMap.Add(Op->Name) = Idx;
+
+		Idx = OpInfos.AddDefaulted();
+		Op = &OpInfos[Idx];
+		Op->Category = CategoryText;
+		Op->FriendlyName = NSLOCTEXT("NiagaraOpInfo", "Subtract Name", "Subtract");
+		Op->Description = NSLOCTEXT("NiagaraOpInfo", "Subtract Desc", "Result = A - B");
+		Op->Keywords = FText::FromString(TEXT("-"));
+		Op->Inputs.Add(FNiagaraOpInOutInfo(A, Type, AText, AText, DefaultStr_Zero));
+		Op->Inputs.Add(FNiagaraOpInOutInfo(B, Type, BText, BText, DefaultStr_Zero));
+		Op->Outputs.Add(FNiagaraOpInOutInfo(Result, Type, ResultText, ResultText, DefaultStr_Zero, TEXT("{0} - {1}")));		
+		Op->BuildName(TEXT("Subtract"), CategoryName);
+		OpInfoMap.Add(Op->Name) = Idx;
+
+		Idx = OpInfos.AddDefaulted();
+		Op = &OpInfos[Idx];
+		Op->Category = CategoryText;
+		Op->FriendlyName = NSLOCTEXT("NiagaraOpInfo", "Multiply Name", "Multiply");
+		Op->Description = NSLOCTEXT("NiagaraOpInfo", "Multiply Desc", "Result = A * B");
+		Op->Keywords = FText::FromString(TEXT("*"));
+		Op->Inputs.Add(FNiagaraOpInOutInfo(A, Type, AText, AText, DefaultStr_One));
+		Op->Inputs.Add(FNiagaraOpInOutInfo(B, Type, BText, BText, DefaultStr_One));
+		Op->Outputs.Add(FNiagaraOpInOutInfo(Result, Type, ResultText, ResultText, DefaultStr_One, TEXT("{0} * {1}")));
+		Op->BuildName(TEXT("Mul"), CategoryName);
+		OpInfoMap.Add(Op->Name) = Idx;
+
+		Idx = OpInfos.AddDefaulted();
+		Op = &OpInfos[Idx];
+		Op->Category = CategoryText;
+		Op->FriendlyName = NSLOCTEXT("NiagaraOpInfo", "Divide Name", "Divide");
+		Op->Description = NSLOCTEXT("NiagaraOpInfo", "Divide Desc", "Result = A / B");
+		Op->Keywords = FText::FromString(TEXT("/"));
+		Op->Inputs.Add(FNiagaraOpInOutInfo(A, Type, AText, AText, DefaultStr_One));
+		Op->Inputs.Add(FNiagaraOpInOutInfo(B, Type, BText, BText, DefaultStr_One));
+		Op->Outputs.Add(FNiagaraOpInOutInfo(Result, Type, ResultText, ResultText, DefaultStr_One, TEXT("{0} / {1}")));
+		Op->BuildName(TEXT("Div"), CategoryName);
+		OpInfoMap.Add(Op->Name) = Idx;
+
+		Idx = OpInfos.AddDefaulted();
+		Op = &OpInfos[Idx];
+		Op->Category = CategoryText;
+		Op->FriendlyName = NSLOCTEXT("NiagaraOpInfo", "MultiplyAdd Name", "MultiplyAdd");
+		Op->Description = NSLOCTEXT("NiagaraOpInfo", "MultiplyAdd Desc", "Result = (A * B) + C");
+		Op->Inputs.Add(FNiagaraOpInOutInfo(A, Type, AText, AText, DefaultStr_One));
+		Op->Inputs.Add(FNiagaraOpInOutInfo(B, Type, BText, BText, DefaultStr_One));
+		Op->Inputs.Add(FNiagaraOpInOutInfo(C, Type, CText, CText, DefaultStr_Zero));
+		Op->Outputs.Add(FNiagaraOpInOutInfo(Result, Type, ResultText, ResultText, DefaultStr_One, TEXT("{0} * {1} + {2}")));
+		Op->BuildName(TEXT("Madd"), CategoryName);
+		OpInfoMap.Add(Op->Name) = Idx;
+
+		Idx = OpInfos.AddDefaulted();
+		Op = &OpInfos[Idx];
+		Op->Category = CategoryText;
+		Op->FriendlyName = NSLOCTEXT("NiagaraOpInfo", "Lerp Name", "Lerp");
+		Op->Description = NSLOCTEXT("NiagaraOpInfo", "Lerp Desc", "Result = (A * (1 - C)) + (B * C)");
+		Op->Keywords = FText::FromString(TEXT("lerp"));
+		Op->Inputs.Add(FNiagaraOpInOutInfo(A, Type, AText, AText, DefaultStr_Zero));
+		Op->Inputs.Add(FNiagaraOpInOutInfo(B, Type, BText, BText, DefaultStr_One));
+		Op->Inputs.Add(FNiagaraOpInOutInfo(C, Type, CText, CText, DefaultStr_Zero));
+		Op->Outputs.Add(FNiagaraOpInOutInfo(Result, Type, ResultText, ResultText, DefaultStr_Zero, TEXT("lerp({0},{1},{2})")));
+		Op->BuildName(TEXT("Lerp"), CategoryName);
+		OpInfoMap.Add(Op->Name) = Idx;
+
+		Idx = OpInfos.AddDefaulted();
+		Op = &OpInfos[Idx];
+		Op->Category = CategoryText;
+		Op->FriendlyName = NSLOCTEXT("NiagaraOpInfo", "Reciprocal Name", "Reciprocal");
+		Op->Description = NSLOCTEXT("NiagaraOpInfo", "Reciprocal Desc", "Result = 1 / A");
+		Op->Inputs.Add(FNiagaraOpInOutInfo(A, Type, AText, AText, DefaultStr_One));
+		Op->Outputs.Add(FNiagaraOpInOutInfo(Result, Type, ResultText, ResultText, DefaultStr_One, TEXT("rcp({0})")));
+		Op->BuildName(TEXT("Rcp"), CategoryName);
+		OpInfoMap.Add(Op->Name) = Idx;
+
+		Idx = OpInfos.AddDefaulted();
+		Op = &OpInfos[Idx];
+		Op->Category = CategoryText;
+		Op->FriendlyName = NSLOCTEXT("NiagaraOpInfo", "Reciprocal Sqrt Name", "Reciprocal Sqrt");
+		Op->Description = NSLOCTEXT("NiagaraOpInfo", "Reciprocal Sqrt Desc", "Result = 1 / sqrt(A)");
+		Op->Inputs.Add(FNiagaraOpInOutInfo(A, Type, AText, AText, DefaultStr_One));
+		Op->Outputs.Add(FNiagaraOpInOutInfo(Result, Type, ResultText, ResultText, DefaultStr_One, TEXT("rsqrt({0})")));
+		Op->BuildName(TEXT("RSqrt"), CategoryName);
+		OpInfoMap.Add(Op->Name) = Idx;
+
+		Idx = OpInfos.AddDefaulted();
+		Op = &OpInfos[Idx];
+		Op->Category = CategoryText;
+		Op->FriendlyName = NSLOCTEXT("NiagaraOpInfo", "Sqrt Name", "Sqrt");
+		Op->Description = NSLOCTEXT("NiagaraOpInfo", "Sqrt Desc", "Result = sqrt(A)");
+		Op->Keywords = FText::FromString(TEXT("sqrt"));
+		Op->Inputs.Add(FNiagaraOpInOutInfo(A, Type, AText, AText, DefaultStr_One));
+		Op->Outputs.Add(FNiagaraOpInOutInfo(Result, Type, ResultText, ResultText, DefaultStr_One, TEXT("sqrt({0})")));
+		Op->BuildName(TEXT("Sqrt"), CategoryName);
+		OpInfoMap.Add(Op->Name) = Idx;
+
+		Idx = OpInfos.AddDefaulted();
+		Op = &OpInfos[Idx];
+		Op->Category = CategoryText;
+		Op->FriendlyName = NSLOCTEXT("NiagaraOpInfo", "One Minus Name", "One Minus");
+		Op->Description = NSLOCTEXT("NiagaraOpInfo", "One Minus Desc", "Result = 1 - A");
+		Op->Keywords = FText::FromString(TEXT("1-x"));
+		Op->Inputs.Add(FNiagaraOpInOutInfo(A, Type, AText, AText, DefaultStr_One));
+		Op->Outputs.Add(FNiagaraOpInOutInfo(Result, Type, ResultText, ResultText, DefaultStr_One, TEXT("1 - {0}")));
+		Op->BuildName(TEXT("OneMinus"), CategoryName);
+		OpInfoMap.Add(Op->Name) = Idx;
+
+		Idx = OpInfos.AddDefaulted();
+		Op = &OpInfos[Idx];
+		Op->Category = CategoryText;
+		Op->FriendlyName = NSLOCTEXT("NiagaraOpInfo", "Negate Name", "Negate");
+		Op->Description = NSLOCTEXT("NiagaraOpInfo", "Negate Desc", "Result = -A");
+		Op->Inputs.Add(FNiagaraOpInOutInfo(A, Type, AText, AText, DefaultStr_One));
+		Op->Outputs.Add(FNiagaraOpInOutInfo(Result, Type, ResultText, ResultText, DefaultStr_One, TEXT("-({0})")));
+		Op->BuildName(TEXT("Negate"), CategoryName);
+		OpInfoMap.Add(Op->Name) = Idx;
+
+		Idx = OpInfos.AddDefaulted();
+		Op = &OpInfos[Idx];
+		Op->Category = CategoryText;
+		Op->FriendlyName = NSLOCTEXT("NiagaraOpInfo", "Abs Name", "Abs");
+		Op->Description = NSLOCTEXT("NiagaraOpInfo", "Abs Desc", "Result = abs(A)");
+		Op->Inputs.Add(FNiagaraOpInOutInfo(A, Type, AText, AText, DefaultStr_One));
+		Op->Outputs.Add(FNiagaraOpInOutInfo(Result, Type, ResultText, ResultText, DefaultStr_One, TEXT("abs({0})")));
+		Op->BuildName(TEXT("Abs"), CategoryName);
+		OpInfoMap.Add(Op->Name) = Idx;
+
+		Idx = OpInfos.AddDefaulted();
+		Op = &OpInfos[Idx];
+		Op->Category = CategoryText;
+		Op->FriendlyName = NSLOCTEXT("NiagaraOpInfo", "Exp Name", "Exp");
+		Op->Description = NSLOCTEXT("NiagaraOpInfo", "Exp Desc", "Result = exp(A)");
+		Op->Inputs.Add(FNiagaraOpInOutInfo(A, Type, AText, AText, DefaultStr_One));
+		Op->Outputs.Add(FNiagaraOpInOutInfo(Result, Type, ResultText, ResultText, DefaultStr_One, TEXT("exp({0})")));
+		Op->BuildName(TEXT("Exp"), CategoryName);
+		OpInfoMap.Add(Op->Name) = Idx;
+
+		Idx = OpInfos.AddDefaulted();
+		Op = &OpInfos[Idx];
+		Op->Category = CategoryText;
+		Op->FriendlyName = NSLOCTEXT("NiagaraOpInfo", "Exp2 Name", "Exp2");
+		Op->Description = NSLOCTEXT("NiagaraOpInfo", "Exp2 Desc", "Result = exp2(A)");
+		Op->Inputs.Add(FNiagaraOpInOutInfo(A, Type, AText, AText, DefaultStr_One));
+		Op->Outputs.Add(FNiagaraOpInOutInfo(Result, Type, ResultText, ResultText, DefaultStr_One, TEXT("exp2({0})")));
+		Op->BuildName(TEXT("Exp2"), CategoryName);
+		OpInfoMap.Add(Op->Name) = Idx;
+
+		Idx = OpInfos.AddDefaulted();
+		Op = &OpInfos[Idx];
+		Op->Category = CategoryText;
+		Op->FriendlyName = NSLOCTEXT("NiagaraOpInfo", "Log Name", "Log");
+		Op->Description = NSLOCTEXT("NiagaraOpInfo", "Log Desc", "Result = log(A)");
+		Op->Inputs.Add(FNiagaraOpInOutInfo(A, Type, AText, AText, DefaultStr_One));
+		Op->Outputs.Add(FNiagaraOpInOutInfo(Result, Type, ResultText, ResultText, DefaultStr_One, TEXT("log({0})")));
+		Op->BuildName(TEXT("Log"), CategoryName);
+		OpInfoMap.Add(Op->Name) = Idx;
+
+		Idx = OpInfos.AddDefaulted();
+		Op = &OpInfos[Idx];
+		Op->Category = CategoryText;
+		Op->FriendlyName = NSLOCTEXT("NiagaraOpInfo", "Log2 Name", "Log2");
+		Op->Description = NSLOCTEXT("NiagaraOpInfo", "Log2 Desc", "Result = log2(A)");
+		Op->Keywords = FText::FromString(TEXT("log2"));
+		Op->Inputs.Add(FNiagaraOpInOutInfo(A, Type, AText, AText, DefaultStr_One));
+		Op->Outputs.Add(FNiagaraOpInOutInfo(Result, Type, ResultText, ResultText, DefaultStr_One, TEXT("log2({0})")));
+		Op->BuildName(TEXT("Log2"), CategoryName);
+		OpInfoMap.Add(Op->Name) = Idx;
+
+		Idx = OpInfos.AddDefaulted();
+		Op = &OpInfos[Idx];
+		Op->Category = CategoryText;
+		Op->FriendlyName = NSLOCTEXT("NiagaraOpInfo", "Sin Name", "Sin");
+		Op->Description = NSLOCTEXT("NiagaraOpInfo", "Sin Desc", "Result = sin(A)");
+		Op->Keywords = FText::FromString(TEXT("sin"));
+		Op->Inputs.Add(FNiagaraOpInOutInfo(A, Type, AText, AText, DefaultStr_One));
+		Op->Outputs.Add(FNiagaraOpInOutInfo(Result, Type, ResultText, ResultText, DefaultStr_One, TEXT("sin({0})")));
+		Op->BuildName(TEXT("Sin"), CategoryName);
+		OpInfoMap.Add(Op->Name) = Idx;
+
+		Idx = OpInfos.AddDefaulted();
+		Op = &OpInfos[Idx];
+		Op->Category = CategoryText;
+		Op->FriendlyName = NSLOCTEXT("NiagaraOpInfo", "Cos Name", "Cos");
+		Op->Description = NSLOCTEXT("NiagaraOpInfo", "Cos Desc", "Result = cos(A)");
+		Op->Keywords = FText::FromString(TEXT("cos"));
+		Op->Inputs.Add(FNiagaraOpInOutInfo(A, Type, AText, AText, DefaultStr_One));
+		Op->Outputs.Add(FNiagaraOpInOutInfo(Result, Type, ResultText, ResultText, DefaultStr_One, TEXT("cos({0})")));
+		Op->BuildName(TEXT("Cos"), CategoryName);
+		OpInfoMap.Add(Op->Name) = Idx;
+
+		Idx = OpInfos.AddDefaulted();
+		Op = &OpInfos[Idx];
+		Op->Category = CategoryText;
+		Op->FriendlyName = NSLOCTEXT("NiagaraOpInfo", "Tan Name", "Tan");
+		Op->Description = NSLOCTEXT("NiagaraOpInfo", "Tan Desc", "Result = tan(A)");
+		Op->Keywords = FText::FromString(TEXT("tan"));
+		Op->Inputs.Add(FNiagaraOpInOutInfo(A, Type, AText, AText, DefaultStr_One));
+		Op->Outputs.Add(FNiagaraOpInOutInfo(Result, Type, ResultText, ResultText, DefaultStr_One, TEXT("tan({0})")));
+		Op->BuildName(TEXT("Tan"), CategoryName);
+		OpInfoMap.Add(Op->Name) = Idx;
+
+		Idx = OpInfos.AddDefaulted();
+		Op = &OpInfos[Idx];
+		Op->Category = CategoryText;
+		Op->FriendlyName = NSLOCTEXT("NiagaraOpInfo", "ASin Name", "ASin");
+		Op->Description = NSLOCTEXT("NiagaraOpInfo", "ASin Desc", "Result = asin(A)");
+		Op->Keywords = FText::FromString(TEXT("asin"));
+		Op->Inputs.Add(FNiagaraOpInOutInfo(A, Type, AText, AText, DefaultStr_One));
+		Op->Outputs.Add(FNiagaraOpInOutInfo(Result, Type, ResultText, ResultText, DefaultStr_One, TEXT("asin({0})")));
+		Op->BuildName(TEXT("ASin"), CategoryName);
+		OpInfoMap.Add(Op->Name) = Idx;
+
+		Idx = OpInfos.AddDefaulted();
+		Op = &OpInfos[Idx];
+		Op->Category = CategoryText;
+		Op->FriendlyName = NSLOCTEXT("NiagaraOpInfo", "ACos Name", "ACos");
+		Op->Description = NSLOCTEXT("NiagaraOpInfo", "ACos Desc", "Result = acos(A)");
+		Op->Keywords = FText::FromString(TEXT("acos"));
+		Op->Inputs.Add(FNiagaraOpInOutInfo(A, Type, AText, AText, DefaultStr_One));
+		Op->Outputs.Add(FNiagaraOpInOutInfo(Result, Type, ResultText, ResultText, DefaultStr_One, TEXT("acos({0})")));
+		Op->BuildName(TEXT("ACos"), CategoryName);
+		OpInfoMap.Add(Op->Name) = Idx;
+
+		Idx = OpInfos.AddDefaulted();
+		Op = &OpInfos[Idx];
+		Op->Category = CategoryText;
+		Op->FriendlyName = NSLOCTEXT("NiagaraOpInfo", "ATan Name", "ATan");
+		Op->Description = NSLOCTEXT("NiagaraOpInfo", "ATan Desc", "Result = atan(A)");
+		Op->Keywords = FText::FromString(TEXT("atan"));
+		Op->Inputs.Add(FNiagaraOpInOutInfo(A, Type, AText, AText, DefaultStr_One));
+		Op->Outputs.Add(FNiagaraOpInOutInfo(Result, Type, ResultText, ResultText, DefaultStr_One, TEXT("atan({0})")));
+		Op->BuildName(TEXT("ATan"), CategoryName);
+		OpInfoMap.Add(Op->Name) = Idx;
+
+		Idx = OpInfos.AddDefaulted();
+		Op = &OpInfos[Idx];
+		Op->Category = CategoryText;
+		Op->FriendlyName = NSLOCTEXT("NiagaraOpInfo", "ATan2 Name", "ATan2");
+		Op->Description = NSLOCTEXT("NiagaraOpInfo", "ATan2 Desc", "Result = atan2(A, B)");
+		Op->Inputs.Add(FNiagaraOpInOutInfo(A, Type, AText, AText, DefaultStr_One));
+		Op->Inputs.Add(FNiagaraOpInOutInfo(B, Type, BText, BText, DefaultStr_One));
+		Op->Outputs.Add(FNiagaraOpInOutInfo(Result, Type, ResultText, ResultText, DefaultStr_One, TEXT("atan2({0},{1})")));
+		Op->BuildName(TEXT("ATan2"), CategoryName);
+		OpInfoMap.Add(Op->Name) = Idx;
+
+		Idx = OpInfos.AddDefaulted();
+		Op = &OpInfos[Idx];
+		Op->Category = CategoryText;
+		Op->FriendlyName = NSLOCTEXT("NiagaraOpInfo", "Ceil Name", "Ceil");
+		Op->Description = NSLOCTEXT("NiagaraOpInfo", "Ceil Desc", "Rounds A to the nearest integer higher than A.");
+		Op->Inputs.Add(FNiagaraOpInOutInfo(A, Type, AText, AText, DefaultStr_One));
+		Op->Outputs.Add(FNiagaraOpInOutInfo(Result, Type, ResultText, ResultText, DefaultStr_One, TEXT("ceil({0})")));
+		Op->BuildName(TEXT("Ceil"), CategoryName);
+		OpInfoMap.Add(Op->Name) = Idx;
+
+		Idx = OpInfos.AddDefaulted();
+		Op = &OpInfos[Idx];
+		Op->Category = CategoryText;
+		Op->FriendlyName = NSLOCTEXT("NiagaraOpInfo", "Floor Name", "Floor");
+		Op->Description = NSLOCTEXT("NiagaraOpInfo", "Floor Desc", "Rounds A to the nearest integer lower than A.");
+		Op->Inputs.Add(FNiagaraOpInOutInfo(A, Type, AText, AText, DefaultStr_One));
+		Op->Outputs.Add(FNiagaraOpInOutInfo(Result, Type, ResultText, ResultText, DefaultStr_One, TEXT("floor({0})")));
+		Op->BuildName(TEXT("Floor"), CategoryName);
+		OpInfoMap.Add(Op->Name) = Idx;
+
+		Idx = OpInfos.AddDefaulted();
+		Op = &OpInfos[Idx];
+		Op->Category = CategoryText;
+		Op->FriendlyName = NSLOCTEXT("NiagaraOpInfo", "Round Name", "Round");
+		Op->Description = NSLOCTEXT("NiagaraOpInfo", "Round Desc", "Rounds A to the nearest integer.");
+		Op->Inputs.Add(FNiagaraOpInOutInfo(A, Type, AText, AText, DefaultStr_One));
+		Op->Outputs.Add(FNiagaraOpInOutInfo(Result, Type, ResultText, ResultText, DefaultStr_One, TEXT("round({0})")));
+		Op->BuildName(TEXT("Round"), CategoryName);
+		OpInfoMap.Add(Op->Name) = Idx;
+
+		Idx = OpInfos.AddDefaulted();
+		Op = &OpInfos[Idx];
+		Op->Category = CategoryText;
+		Op->FriendlyName = NSLOCTEXT("NiagaraOpInfo", "Fmod Name", "FMod");
+		Op->Description = NSLOCTEXT("NiagaraOpInfo", "Fmod Desc", "Result = A % B");
+		Op->Keywords = FText::FromString(TEXT("%"));
+		Op->Inputs.Add(FNiagaraOpInOutInfo(A, Type, AText, AText, DefaultStr_One));
+		Op->Inputs.Add(FNiagaraOpInOutInfo(B, Type, BText, BText, DefaultStr_One));
+		Op->Outputs.Add(FNiagaraOpInOutInfo(Result, Type, ResultText, ResultText, DefaultStr_One, TEXT("Modulo({0}, {1})")));
+		Op->BuildName(TEXT("FMod"), CategoryName);
+		OpInfoMap.Add(Op->Name) = Idx;
+
+		Idx = OpInfos.AddDefaulted();
+		Op = &OpInfos[Idx];
+		Op->Category = CategoryText;
+		Op->FriendlyName = NSLOCTEXT("NiagaraOpInfo", "Frac Name", "Frac");
+		Op->Description = NSLOCTEXT("NiagaraOpInfo", "Frac Desc", "Result = frac(A)");
+		Op->Inputs.Add(FNiagaraOpInOutInfo(A, Type, AText, AText, DefaultStr_One));
+		Op->Outputs.Add(FNiagaraOpInOutInfo(Result, Type, ResultText, ResultText, DefaultStr_One, TEXT("frac({0})")));
+		Op->BuildName(TEXT("Frac"), CategoryName);
+		OpInfoMap.Add(Op->Name) = Idx;
+
+		Idx = OpInfos.AddDefaulted();
+		Op = &OpInfos[Idx];
+		Op->Category = CategoryText;
+		Op->FriendlyName = NSLOCTEXT("NiagaraOpInfo", "Trunc Name", "Trunc");
+		Op->Description = NSLOCTEXT("NiagaraOpInfo", "Trunc Desc", "Result = trunc(A)");
+		Op->Inputs.Add(FNiagaraOpInOutInfo(A, Type, AText, AText, DefaultStr_One));
+		Op->Outputs.Add(FNiagaraOpInOutInfo(Result, Type, ResultText, ResultText, DefaultStr_One, TEXT("trunc({0})")));
+		Op->BuildName(TEXT("Trunc"), CategoryName);
+		OpInfoMap.Add(Op->Name) = Idx;
+
+		Idx = OpInfos.AddDefaulted();
+		Op = &OpInfos[Idx];
+		Op->Category = CategoryText;
+		Op->FriendlyName = NSLOCTEXT("NiagaraOpInfo", "Clamp Name", "Clamp");
+		Op->Description = NSLOCTEXT("NiagaraOpInfo", "Clamp Desc", "Result = clamp(A, Min, Max)");
+		Op->Inputs.Add(FNiagaraOpInOutInfo(A, Type, AText, AText, DefaultStr_One));
+		Op->Inputs.Add(FNiagaraOpInOutInfo(Min, Type, MinText, MinText, DefaultStr_One));
+		Op->Inputs.Add(FNiagaraOpInOutInfo(Max, Type, MaxText, MaxText, DefaultStr_One));
+		Op->Outputs.Add(FNiagaraOpInOutInfo(Result, Type, ResultText, ResultText, DefaultStr_One, TEXT("clamp({0},{1},{2})")));
+		Op->BuildName(TEXT("Clamp"), CategoryName);
+		OpInfoMap.Add(Op->Name) = Idx;
+
+		Idx = OpInfos.AddDefaulted();
+		Op = &OpInfos[Idx];
+		Op->Category = CategoryText;
+		Op->FriendlyName = NSLOCTEXT("NiagaraOpInfo", "Min Name", "Min");
+		Op->Description = NSLOCTEXT("NiagaraOpInfo", "Min Desc", "Result = min(A, B)");
+		Op->Inputs.Add(FNiagaraOpInOutInfo(A, Type, AText, AText, DefaultStr_One));
+		Op->Inputs.Add(FNiagaraOpInOutInfo(B, Type, BText, BText, DefaultStr_One));
+		Op->Outputs.Add(FNiagaraOpInOutInfo(Result, Type, ResultText, ResultText, DefaultStr_One, TEXT("min({0},{1})")));
+		Op->BuildName(TEXT("Min"), CategoryName);
+		OpInfoMap.Add(Op->Name) = Idx;
+
+		Idx = OpInfos.AddDefaulted();
+		Op = &OpInfos[Idx];
+		Op->Category = CategoryText;
+		Op->FriendlyName = NSLOCTEXT("NiagaraOpInfo", "Max Name", "Max");
+		Op->Description = NSLOCTEXT("NiagaraOpInfo", "Max Desc", "Result = max(A, B)");
+		Op->Inputs.Add(FNiagaraOpInOutInfo(A, Type, AText, AText, DefaultStr_One));
+		Op->Inputs.Add(FNiagaraOpInOutInfo(B, Type, BText, BText, DefaultStr_One));
+		Op->Outputs.Add(FNiagaraOpInOutInfo(Result, Type, ResultText, ResultText, DefaultStr_One, TEXT("max({0},{1})")));
+		Op->BuildName(TEXT("Max"), CategoryName);
+		OpInfoMap.Add(Op->Name) = Idx;
+
+		Idx = OpInfos.AddDefaulted();
+		Op = &OpInfos[Idx];
+		Op->Category = CategoryText;
+		Op->FriendlyName = NSLOCTEXT("NiagaraOpInfo", "Pow Name", "Pow");
+		Op->Description = NSLOCTEXT("NiagaraOpInfo", "Pow Desc", "Result = pow(A, B)");
+		Op->Inputs.Add(FNiagaraOpInOutInfo(A, Type, AText, AText, DefaultStr_One));
+		Op->Inputs.Add(FNiagaraOpInOutInfo(B, Type, BText, BText, DefaultStr_One));
+		Op->Outputs.Add(FNiagaraOpInOutInfo(Result, Type, ResultText, ResultText, DefaultStr_One, TEXT("pow({0},{1})")));
+		Op->BuildName(TEXT("Pow"), CategoryName);
+		OpInfoMap.Add(Op->Name) = Idx;
+
+		Idx = OpInfos.AddDefaulted();
+		Op = &OpInfos[Idx];
+		Op->Category = CategoryText;
+		Op->FriendlyName = NSLOCTEXT("NiagaraOpInfo", "Sign Name", "Sign");
+		Op->Description = NSLOCTEXT("NiagaraOpInfo", "Sign Desc", "Result = sign(A)");
+		Op->Inputs.Add(FNiagaraOpInOutInfo(A, Type, AText, AText, DefaultStr_One));
+		Op->Outputs.Add(FNiagaraOpInOutInfo(Result, Type, ResultText, ResultText, DefaultStr_One, TEXT("sign({0})")));
+		Op->BuildName(TEXT("Sign"), CategoryName);
+		OpInfoMap.Add(Op->Name) = Idx;
+
+		Idx = OpInfos.AddDefaulted();
+		Op = &OpInfos[Idx];
+		Op->Category = CategoryText;
+		Op->FriendlyName = NSLOCTEXT("NiagaraOpInfo", "Step Name", "Step");
+		Op->Description = NSLOCTEXT("NiagaraOpInfo", "Step Desc", "Result = step(A)");
+		Op->Inputs.Add(FNiagaraOpInOutInfo(A, Type, AText, AText, DefaultStr_One));
+		Op->Outputs.Add(FNiagaraOpInOutInfo(Result, Type, ResultText, ResultText, DefaultStr_One, TEXT("step({0})")));
+		Op->BuildName(TEXT("Step"), CategoryName);
+		OpInfoMap.Add(Op->Name) = Idx;
+
+		Idx = OpInfos.AddDefaulted();
+		Op = &OpInfos[Idx];
+		Op->Category = CategoryText;
+		Op->FriendlyName = NSLOCTEXT("NiagaraOpInfo", "Noise Name", "Noise");
+		Op->Description = NSLOCTEXT("NiagaraOpInfo", "Noise Desc", "A continuous pseudo random noise function.");
+		Op->Inputs.Add(FNiagaraOpInOutInfo(X, Type, XText, XText, DefaultStr_Zero));
+		Op->Outputs.Add(FNiagaraOpInOutInfo(Result, Type, ResultText, ResultText, DefaultStr_One, TEXT("noise({0})")));
+		Op->NumericOuputTypeSelectionMode = ENiagaraNumericOutputTypeSelectionMode::Scalar;
+		Op->BuildName(TEXT("Noise"), CategoryName);
+		OpInfoMap.Add(Op->Name) = Idx;
+
+		Idx = OpInfos.AddDefaulted();
+		Op = &OpInfos[Idx];
+		Op->Category = CategoryText;
+		Op->FriendlyName = NSLOCTEXT("NiagaraOpInfo", "Dot Name", "Dot");
+		Op->Description = NSLOCTEXT("NiagaraOpInfo", "Dot Desc", "Dot product of two vectors.");
+		Op->Inputs.Add(FNiagaraOpInOutInfo(A, Type, AText, AText, DefaultStr_One));
+		Op->Inputs.Add(FNiagaraOpInOutInfo(B, Type, BText, BText, DefaultStr_One));
+		Op->Outputs.Add(FNiagaraOpInOutInfo(Result, Type, ResultText, ResultText, DefaultStr_One, TEXT("dot({0},{1})")));
+		Op->NumericOuputTypeSelectionMode = ENiagaraNumericOutputTypeSelectionMode::Scalar;
+		Op->BuildName(TEXT("Dot"), CategoryName);
+		OpInfoMap.Add(Op->Name) = Idx;
+
+		Idx = OpInfos.AddDefaulted();
+		Op = &OpInfos[Idx];
+		Op->Category = CategoryText;
+		Op->FriendlyName = NSLOCTEXT("NiagaraOpInfo", "Normalize Name", "Normalize");
+		Op->Description = NSLOCTEXT("NiagaraOpInfo", "Normalize Desc", "Normalizes the passed value.");
+		Op->Inputs.Add(FNiagaraOpInOutInfo(A, Type, AText, AText, DefaultStr_One));
+		Op->Outputs.Add(FNiagaraOpInOutInfo(Result, Type, ResultText, ResultText, DefaultStr_One, TEXT("normalize({0})")));
+		Op->BuildName(TEXT("Normalize"), CategoryName);
+		OpInfoMap.Add(Op->Name) = Idx;
+
+		Idx = OpInfos.AddDefaulted();
+		Op = &OpInfos[Idx];
+		Op->Category = CategoryText;
+		Op->FriendlyName = NSLOCTEXT("NiagaraOpInfo", "Length Name", "Length");
+		Op->Description = NSLOCTEXT("NiagaraOpInfo", "Length Desc", "Returns the length of the passed value.");
+		Op->Inputs.Add(FNiagaraOpInOutInfo(A, Type, AText, AText, DefaultStr_One));
+		Op->Outputs.Add(FNiagaraOpInOutInfo(Result, Type, ResultText, ResultText, DefaultStr_One, TEXT("length({0})")));
+		Op->NumericOuputTypeSelectionMode = ENiagaraNumericOutputTypeSelectionMode::Scalar;
+		Op->BuildName(TEXT("Length"), CategoryName);
+		OpInfoMap.Add(Op->Name) = Idx;
+
+		//Temporarily here. Rand will be reworked shortly.
+		Idx = OpInfos.AddDefaulted();
+		Op = &OpInfos[Idx];
+		Op->Category = CategoryText;
+		Op->FriendlyName = NSLOCTEXT("NiagaraOpInfo", "Rand Name", "Random");
+		Op->Description = NSLOCTEXT("NiagaraOpInfo", "Rand Desc", "Result = rand(A)");
+		Op->Inputs.Add(FNiagaraOpInOutInfo(A, Type, AText, AText, DefaultStr_One));
+		Op->Outputs.Add(FNiagaraOpInOutInfo(Result, Type, ResultText, ResultText, DefaultStr_One, TEXT("rand({0})")));
+		Op->BuildName(TEXT("Rand"), CategoryName);
+		OpInfoMap.Add(Op->Name) = Idx;
+
+		//Comparison ops
+		Idx = OpInfos.AddDefaulted();
+		Op = &OpInfos[Idx];
+		Op->Category = CategoryText;
+		Op->FriendlyName = NSLOCTEXT("NiagaraOpInfo", "CmpLT Name", "Less Than");
+		Op->Description = NSLOCTEXT("NiagaraOpInfo", "CmpLT Desc", "Result = A < B");
+		Op->Keywords = FText::FromString(TEXT("<"));
+		Op->Inputs.Add(FNiagaraOpInOutInfo(A, Type, AText, AText, DefaultStr_Zero));
+		Op->Inputs.Add(FNiagaraOpInOutInfo(B, Type, BText, BText, DefaultStr_One));
+		Op->Outputs.Add(FNiagaraOpInOutInfo(Result, FNiagaraTypeDefinition::GetBoolDef(), ResultText, ResultText, DefaultStr_One, TEXT("{0} < {1}")));
+		Op->BuildName(TEXT("CmpLT"), CategoryName);
+		OpInfoMap.Add(Op->Name) = Idx;
+
+		Idx = OpInfos.AddDefaulted();
+		Op = &OpInfos[Idx];
+		Op->Category = CategoryText;
+		Op->FriendlyName = NSLOCTEXT("NiagaraOpInfo", "CmpLE Name", "Less Than Or Equal");
+		Op->Description = NSLOCTEXT("NiagaraOpInfo", "CmpLE Desc", "Result = A <= B");
+		Op->Keywords = FText::FromString(TEXT("<="));
+		Op->Inputs.Add(FNiagaraOpInOutInfo(A, Type, AText, AText, DefaultStr_Zero));
+		Op->Inputs.Add(FNiagaraOpInOutInfo(B, Type, BText, BText, DefaultStr_One));
+		Op->Outputs.Add(FNiagaraOpInOutInfo(Result, FNiagaraTypeDefinition::GetBoolDef(), ResultText, ResultText, DefaultStr_One, TEXT("{0} <= {1}")));
+		Op->BuildName(TEXT("CmpLE"), CategoryName);
+		OpInfoMap.Add(Op->Name) = Idx;
+
+		Idx = OpInfos.AddDefaulted();
+		Op = &OpInfos[Idx];
+		Op->Category = CategoryText;
+		Op->FriendlyName = NSLOCTEXT("NiagaraOpInfo", "CmpGT Name", "Greater Than");
+		Op->Description = NSLOCTEXT("NiagaraOpInfo", "CmpGT Desc", "Result = A > B");
+		Op->Keywords = FText::FromString(TEXT(">"));
+		Op->Inputs.Add(FNiagaraOpInOutInfo(A, Type, AText, AText, DefaultStr_Zero));
+		Op->Inputs.Add(FNiagaraOpInOutInfo(B, Type, BText, BText, DefaultStr_One));
+		Op->Outputs.Add(FNiagaraOpInOutInfo(Result, FNiagaraTypeDefinition::GetBoolDef(), ResultText, ResultText, DefaultStr_One, TEXT("{0} > {1}")));
+		Op->BuildName(TEXT("CmpGT"), CategoryName);
+		OpInfoMap.Add(Op->Name) = Idx;
+
+		Idx = OpInfos.AddDefaulted();
+		Op = &OpInfos[Idx];
+		Op->Category = CategoryText;
+		Op->FriendlyName = NSLOCTEXT("NiagaraOpInfo", "CmpGE Name", "Greater Than Or Equal");
+		Op->Description = NSLOCTEXT("NiagaraOpInfo", "CmpGE Desc", "Result = A >= B");
+		Op->Keywords = FText::FromString(TEXT(">="));
+		Op->Inputs.Add(FNiagaraOpInOutInfo(A, Type, AText, AText, DefaultStr_Zero));
+		Op->Inputs.Add(FNiagaraOpInOutInfo(B, Type, BText, BText, DefaultStr_One));
+		Op->Outputs.Add(FNiagaraOpInOutInfo(Result, FNiagaraTypeDefinition::GetBoolDef(), ResultText, ResultText, DefaultStr_One, TEXT("{0} >= {1}")));
+		Op->BuildName(TEXT("CmpGE"), CategoryName);
+		OpInfoMap.Add(Op->Name) = Idx;
+
+		Idx = OpInfos.AddDefaulted();
+		Op = &OpInfos[Idx];
+		Op->Category = CategoryText;
+		Op->FriendlyName = NSLOCTEXT("NiagaraOpInfo", "CmpEQ Name", "Equal");
+		Op->Description = NSLOCTEXT("NiagaraOpInfo", "CmpEQ Desc", "Result = A == B");
+		Op->Keywords = FText::FromString(TEXT("=="));
+		Op->Inputs.Add(FNiagaraOpInOutInfo(A, Type, AText, AText, DefaultStr_Zero));
+		Op->Inputs.Add(FNiagaraOpInOutInfo(B, Type, BText, BText, DefaultStr_One));
+		Op->Outputs.Add(FNiagaraOpInOutInfo(Result, FNiagaraTypeDefinition::GetBoolDef(), ResultText, ResultText, DefaultStr_One, TEXT("{0} == {1}")));
+		Op->BuildName(TEXT("CmpEQ"), CategoryName);
+		OpInfoMap.Add(Op->Name) = Idx;
+		
+		Idx = OpInfos.AddDefaulted();
+		Op = &OpInfos[Idx];
+		Op->Category = CategoryText;
+		Op->FriendlyName = NSLOCTEXT("NiagaraOpInfo", "CmpNEQ Name", "Not Equal");
+		Op->Description = NSLOCTEXT("NiagaraOpInfo", "CmpNEQ Desc", "Result = A != B");
+		Op->Keywords = FText::FromString(TEXT("!="));
+		Op->Inputs.Add(FNiagaraOpInOutInfo(A, Type, AText, AText, DefaultStr_Zero));
+		Op->Inputs.Add(FNiagaraOpInOutInfo(B, Type, BText, BText, DefaultStr_One));
+		Op->Outputs.Add(FNiagaraOpInOutInfo(Result, FNiagaraTypeDefinition::GetBoolDef(), ResultText, ResultText, DefaultStr_One, TEXT("{0} != {1}")));
+		Op->BuildName(TEXT("CmpNEQ"), CategoryName);
+		OpInfoMap.Add(Op->Name) = Idx;
+	}
+
+ 	//////////////////////////////////////////////////////////////////////////
+ 	//Integer Only Ops
+
+	FString Default_IntZero(TEXT("0"));
+	FString Default_IntOne(TEXT("1"));
+	FText IntCategory = NSLOCTEXT("NiagaraOpInfo", "IntOpCategory", "Integer");
+	FString IntCategoryName(TEXT("Integer"));
+	FNiagaraTypeDefinition IntType = FNiagaraTypeDefinition::GetIntDef();
+
+	int32 Idx = OpInfos.AddDefaulted();
+	FNiagaraOpInfo* Op = &OpInfos[Idx];
+	Op->Category = IntCategory;
+	Op->FriendlyName = NSLOCTEXT("NiagaraOpInfo", "BitAnd Name", "Bitwise AND");
+	Op->Description = NSLOCTEXT("NiagaraOpInfo", "BitAnd Desc", "Result = A & B");
+	Op->Keywords = FText::FromString(TEXT("&"));
+	Op->Inputs.Add(FNiagaraOpInOutInfo(A, IntType, AText, AText, Default_IntZero));
+	Op->Inputs.Add(FNiagaraOpInOutInfo(B, IntType, BText, BText, Default_IntOne));
+	Op->Outputs.Add(FNiagaraOpInOutInfo(Result, IntType, ResultText, ResultText, Default_IntOne, TEXT("{0} & {1}")));
+	Op->BuildName(TEXT("BitAnd"), IntCategoryName);
+	OpInfoMap.Add(Op->Name) = Idx;
+
+	Idx = OpInfos.AddDefaulted();
+	Op = &OpInfos[Idx];
+	Op->Category = IntCategory;
+	Op->FriendlyName = NSLOCTEXT("NiagaraOpInfo", "BitOr Name", "Bitwise OR");
+	Op->Description = NSLOCTEXT("NiagaraOpInfo", "BitOr Desc", "Result = A | B");
+	Op->Keywords = FText::FromString(TEXT("|"));
+	Op->Inputs.Add(FNiagaraOpInOutInfo(A, IntType, AText, AText, Default_IntZero));
+	Op->Inputs.Add(FNiagaraOpInOutInfo(B, IntType, BText, BText, Default_IntOne));
+	Op->Outputs.Add(FNiagaraOpInOutInfo(Result, IntType, ResultText, ResultText, Default_IntOne, TEXT("{0} | {1}")));
+	Op->BuildName(TEXT("BitOr"), IntCategoryName);
+	OpInfoMap.Add(Op->Name) = Idx;
+
+	Idx = OpInfos.AddDefaulted();
+	Op = &OpInfos[Idx];
+	Op->Category = IntCategory;
+	Op->FriendlyName = NSLOCTEXT("NiagaraOpInfo", "BitXOr Name", "Bitwise XOR");
+	Op->Description = NSLOCTEXT("NiagaraOpInfo", "BitXOr Desc", "Result = A ^ B");
+	Op->Keywords = FText::FromString(TEXT("^"));
+	Op->Inputs.Add(FNiagaraOpInOutInfo(A, IntType, AText, AText, Default_IntZero));
+	Op->Inputs.Add(FNiagaraOpInOutInfo(B, IntType, BText, BText, Default_IntOne));
+	Op->Outputs.Add(FNiagaraOpInOutInfo(Result, IntType, ResultText, ResultText, Default_IntOne, TEXT("{0} ^ {1}")));
+	Op->BuildName(TEXT("BitXOr"), IntCategoryName);
+	OpInfoMap.Add(Op->Name) = Idx;
+
+	Idx = OpInfos.AddDefaulted();
+	Op = &OpInfos[Idx];
+	Op->Category = IntCategory;
+	Op->FriendlyName = NSLOCTEXT("NiagaraOpInfo", "BitNot Name", "Bitwise NOT");
+	Op->Description = NSLOCTEXT("NiagaraOpInfo", "BitNot Desc", "Result =  ~B");
+	Op->Keywords = FText::FromString(TEXT("~"));
+	Op->Inputs.Add(FNiagaraOpInOutInfo(A, IntType, AText, AText, Default_IntOne));
+	Op->Outputs.Add(FNiagaraOpInOutInfo(Result, IntType, ResultText, ResultText, Default_IntOne, TEXT("~{0}")));
+	Op->BuildName(TEXT("BitNot"), IntCategoryName);
+	OpInfoMap.Add(Op->Name) = Idx;
+
+	//////////////////////////////////////////////////////////////////////////
+	// Boolean Only Ops
+
+	FString Default_BoolZero(TEXT("false"));
+	FString Default_BoolOne(TEXT("true"));
+	FText BoolCategory = NSLOCTEXT("NiagaraOpInfo", "BoolOpCategory", "Boolean");
+	FString BoolCategoryName(TEXT("Boolean"));
+	FNiagaraTypeDefinition BoolType = FNiagaraTypeDefinition::GetBoolDef();
+
+	Idx = OpInfos.AddDefaulted();
+	Op = &OpInfos[Idx];
+	Op->Category = BoolCategory;
+	Op->FriendlyName = NSLOCTEXT("NiagaraOpInfo", "LogicAnd Name", "Logic AND");
+	Op->Description = NSLOCTEXT("NiagaraOpInfo", "LogicAnd Desc", "Result = A && B");
+	Op->Keywords = FText::FromString(TEXT("&&"));
+	Op->Inputs.Add(FNiagaraOpInOutInfo(A, BoolType, AText, AText, Default_BoolZero));
+	Op->Inputs.Add(FNiagaraOpInOutInfo(B, BoolType, BText, BText, Default_BoolOne));
+	Op->Outputs.Add(FNiagaraOpInOutInfo(Result, BoolType, ResultText, ResultText, Default_BoolOne, TEXT("{0} && {1}")));
+	Op->BuildName(TEXT("LogicAnd"), BoolCategoryName);
+	OpInfoMap.Add(Op->Name) = Idx;
+
+	Idx = OpInfos.AddDefaulted();
+	Op = &OpInfos[Idx];
+	Op->Category = BoolCategory;
+	Op->FriendlyName = NSLOCTEXT("NiagaraOpInfo", "LogicOr Name", "Logic OR");
+	Op->Description = NSLOCTEXT("NiagaraOpInfo", "LogicOr Desc", "Logic = A || B");
+	Op->Keywords = FText::FromString(TEXT("||"));
+	Op->Inputs.Add(FNiagaraOpInOutInfo(A, BoolType, AText, AText, Default_BoolZero));
+	Op->Inputs.Add(FNiagaraOpInOutInfo(B, BoolType, BText, BText, Default_BoolOne));
+	Op->Outputs.Add(FNiagaraOpInOutInfo(Result, BoolType, ResultText, ResultText, Default_BoolOne, TEXT("{0} || {1}")));
+	Op->BuildName(TEXT("LogicOr"), BoolCategoryName);
+	OpInfoMap.Add(Op->Name) = Idx;
+
+	Idx = OpInfos.AddDefaulted();
+	Op = &OpInfos[Idx];
+	Op->Category = BoolCategory;
+	Op->FriendlyName = NSLOCTEXT("NiagaraOpInfo", "LogicNot Name", "Logic NOT");
+	Op->Description = NSLOCTEXT("NiagaraOpInfo", "LogicNot Desc", "Result = !B");
+	Op->Keywords = FText::FromString(TEXT("!"));
+	Op->Inputs.Add(FNiagaraOpInOutInfo(A, BoolType, AText, AText, Default_BoolOne));
+	Op->Outputs.Add(FNiagaraOpInOutInfo(Result, BoolType, ResultText, ResultText, Default_BoolOne, TEXT("!{0}")));
+	Op->BuildName(TEXT("LogicNot"), BoolCategoryName);
+	OpInfoMap.Add(Op->Name) = Idx;
+
+	//////////////////////////////////////////////////////////////////////////
+	//Matrix only ops
+	FString Default_MatrixZero(TEXT(
+		"0.0,0.0,0.0,0.0,\
+		0.0,0.0,0.0,0.0,\
+		0.0,0.0,0.0,0.0,\
+		0.0,0.0,0.0,0.0"));
+
+	FString Default_MatrixOne(TEXT(
+		"1.0,0.0,0.0,0.0,\
+		0.0,1.0,0.0,0.0,\
+		0.0,0.0,1.0,0.0,\
+		0.0,0.0,0.0,1.0"));
+
+	FText MatrixCategory = NSLOCTEXT("NiagaraOpInfo", "MatrixOpCategory", "Matrix");
+	FString MatrixCategoryName(TEXT("Matrix"));
+	FNiagaraTypeDefinition MatrixType = FNiagaraTypeDefinition::GetMatrix4Def();
+
+	static FName M(TEXT("M"));
+	static FText MText = NSLOCTEXT("NiagaraOpInfo", "Matrix Param", "M");
+
+	static FName V(TEXT("V"));
+	static FText VText = NSLOCTEXT("NiagaraOpInfo", "Vector Param", "V");
+
+	FString Default_VectorOne(TEXT("1.0,1.0,1.0"));
+	FString Default_Vector4One(TEXT("1.0,1.0,1.0,1.0"));
+	FString Default_VectorX(TEXT("1.0,0.0,0.0"));
+	FString Default_VectorY(TEXT("0.0,1.0,0.0"));
+	FString Default_VectorZ(TEXT("0.0,0.0,1.0"));
+
+	Idx = OpInfos.AddDefaulted();
+	Op = &OpInfos[Idx];
+	Op->Category = MatrixCategory;
+	Op->FriendlyName = NSLOCTEXT("NiagaraOpInfo", "Transpose Name", "Transpose");
+	Op->Description = NSLOCTEXT("NiagaraOpInfo", "Transpose Desc", "Returns the trasnpose of the passd matrix.");
+	Op->Inputs.Add(FNiagaraOpInOutInfo(M, MatrixType, MText, MText, Default_MatrixOne));
+	Op->Outputs.Add(FNiagaraOpInOutInfo(Result, MatrixType, ResultText, ResultText, Default_MatrixOne, TEXT("transpose({0})")));
+	Op->BuildName(TEXT("Transpose"), MatrixCategoryName);
+	OpInfoMap.Add(Op->Name) = Idx;
+
+	Idx = OpInfos.AddDefaulted();
+	Op = &OpInfos[Idx];
+	Op->Category = MatrixCategory;
+	Op->FriendlyName = NSLOCTEXT("NiagaraOpInfo", "Row0 Name", "Row 0");
+	Op->Description = NSLOCTEXT("NiagaraOpInfo", "Row0 Desc", "Returns Row 0 of this matrix.");
+	Op->Inputs.Add(FNiagaraOpInOutInfo(M, MatrixType, MText, MText, Default_MatrixOne));
+	Op->Outputs.Add(FNiagaraOpInOutInfo(Result, FNiagaraTypeDefinition::GetVec4Def(), ResultText, ResultText, Default_MatrixOne, TEXT("{0}[0]")));
+	Op->BuildName(TEXT("Row0"), MatrixCategoryName);
+	OpInfoMap.Add(Op->Name) = Idx;
+
+	Idx = OpInfos.AddDefaulted();
+	Op = &OpInfos[Idx];
+	Op->Category = MatrixCategory;
+	Op->FriendlyName = NSLOCTEXT("NiagaraOpInfo", "Row1 Name", "Row 1");
+	Op->Description = NSLOCTEXT("NiagaraOpInfo", "Row1 Desc", "Returns Row 1 of this matrix.");
+	Op->Inputs.Add(FNiagaraOpInOutInfo(M, MatrixType, MText, MText, Default_MatrixOne));
+	Op->Outputs.Add(FNiagaraOpInOutInfo(Result, FNiagaraTypeDefinition::GetVec4Def(), ResultText, ResultText, Default_MatrixOne, TEXT("{0}[1]")));
+	Op->BuildName(TEXT("Row1"), MatrixCategoryName);
+	OpInfoMap.Add(Op->Name) = Idx;
+
+	Idx = OpInfos.AddDefaulted();
+	Op = &OpInfos[Idx];
+	Op->Category = MatrixCategory;
+	Op->FriendlyName = NSLOCTEXT("NiagaraOpInfo", "Row2 Name", "Row 2");
+	Op->Description = NSLOCTEXT("NiagaraOpInfo", "Row2 Desc", "Returns Row 2 of this matrix.");
+	Op->Inputs.Add(FNiagaraOpInOutInfo(M, MatrixType, MText, MText, Default_MatrixOne));
+	Op->Outputs.Add(FNiagaraOpInOutInfo(Result, FNiagaraTypeDefinition::GetVec4Def(), ResultText, ResultText, Default_MatrixOne, TEXT("{0}[2]")));
+	Op->BuildName(TEXT("Row2"), MatrixCategoryName);
+	OpInfoMap.Add(Op->Name) = Idx;
+
+	Idx = OpInfos.AddDefaulted();
+	Op = &OpInfos[Idx];
+	Op->Category = MatrixCategory;
+	Op->FriendlyName = NSLOCTEXT("NiagaraOpInfo", "Row3 Name", "Row 3");
+	Op->Description = NSLOCTEXT("NiagaraOpInfo", "Row3 Desc", "Returns Row 3 of this matrix.");
+	Op->Inputs.Add(FNiagaraOpInOutInfo(M, MatrixType, MText, MText, Default_MatrixOne));
+	Op->Outputs.Add(FNiagaraOpInOutInfo(Result, FNiagaraTypeDefinition::GetVec4Def(), ResultText, ResultText, Default_MatrixOne, TEXT("{0}[3]")));
+	Op->BuildName(TEXT("Row3"), MatrixCategoryName);
+	OpInfoMap.Add(Op->Name) = Idx;
+
+	Idx = OpInfos.AddDefaulted();
+	Op = &OpInfos[Idx];
+	Op->Category = MatrixCategory;
+	Op->FriendlyName = NSLOCTEXT("NiagaraOpInfo", "MatrixMatrix Mul Name", "Multiply (Matrix * Matrix)");
+	Op->Description = NSLOCTEXT("NiagaraOpInfo", "MatrixMatrix Desc", "Multiplies one matrix by another.");
+	Op->Inputs.Add(FNiagaraOpInOutInfo(A, MatrixType, AText, AText, Default_MatrixOne));
+	Op->Inputs.Add(FNiagaraOpInOutInfo(B, MatrixType, BText, BText, Default_MatrixOne));
+	Op->Outputs.Add(FNiagaraOpInOutInfo(Result, MatrixType, ResultText, ResultText, Default_MatrixOne, TEXT("{0} * {1}")));
+	Op->BuildName(TEXT("MatrixMultiply"), MatrixCategoryName);
+	OpInfoMap.Add(Op->Name) = Idx;
+	
+	Idx = OpInfos.AddDefaulted();
+	Op = &OpInfos[Idx];
+	Op->Category = MatrixCategory;
+	Op->FriendlyName = NSLOCTEXT("NiagaraOpInfo", "MatrixVector Mul Name", "Multiply (Matrix * Vector4)");
+	Op->Description = NSLOCTEXT("NiagaraOpInfo", "MatrixVector Mul Desc", "Multiplies a matrix by a vector4.");
+	Op->Inputs.Add(FNiagaraOpInOutInfo(M, MatrixType, MText, MText, Default_MatrixOne));
+	Op->Inputs.Add(FNiagaraOpInOutInfo(V, FNiagaraTypeDefinition::GetVec4Def(), VText, VText, Default_VectorOne));
+	Op->Outputs.Add(FNiagaraOpInOutInfo(Result, FNiagaraTypeDefinition::GetVec4Def(), ResultText, ResultText, Default_VectorOne, TEXT("mul({1},{0})")));
+	Op->BuildName(TEXT("MatrixVectorMultiply"), MatrixCategoryName);
+	OpInfoMap.Add(Op->Name) = Idx;
+
+	Idx = OpInfos.AddDefaulted();
+	Op = &OpInfos[Idx];
+	Op->Category = MatrixCategory;
+	Op->FriendlyName = NSLOCTEXT("NiagaraOpInfo", "MatrixVector Mul Name", "Multiply (Matrix * Vector4)");
+	Op->Description = NSLOCTEXT("NiagaraOpInfo", "MatrixVector Mul Desc", "Multiplies a matrix by a vector4.");
+	Op->Inputs.Add(FNiagaraOpInOutInfo(M, MatrixType, MText, MText, Default_MatrixOne));
+	Op->Inputs.Add(FNiagaraOpInOutInfo(V, FNiagaraTypeDefinition::GetVec4Def(), VText, VText, Default_Vector4One));
+	Op->Outputs.Add(FNiagaraOpInOutInfo(Result, FNiagaraTypeDefinition::GetVec4Def(), ResultText, ResultText, Default_Vector4One, TEXT("mul({1},{0})")));
+	Op->BuildName(TEXT("MatrixVectorMultiply"), MatrixCategoryName);
+	OpInfoMap.Add(Op->Name) = Idx;
+
+	Idx = OpInfos.AddDefaulted();
+	Op = &OpInfos[Idx];
+	Op->Category = MatrixCategory;
+	Op->FriendlyName = NSLOCTEXT("NiagaraOpInfo", "TransformPosition Name", "Transform Position");
+	Op->Description = NSLOCTEXT("NiagaraOpInfo", "TransformPosition Desc", "Transforms a Vector3 as a position.");
+	Op->Inputs.Add(FNiagaraOpInOutInfo(M, MatrixType, MText, MText, Default_MatrixOne));
+	Op->Inputs.Add(FNiagaraOpInOutInfo(V, FNiagaraTypeDefinition::GetVec3Def(), VText, VText, Default_VectorOne));
+	Op->Outputs.Add(FNiagaraOpInOutInfo(Result, FNiagaraTypeDefinition::GetVec3Def(), ResultText, ResultText, Default_VectorOne, TEXT("mul(float4({1},1.0),{0}).xyz")));
+	Op->BuildName(TEXT("TransformPosition"), MatrixCategoryName);
+	OpInfoMap.Add(Op->Name) = Idx;
+
+	Idx = OpInfos.AddDefaulted();
+	Op = &OpInfos[Idx];
+	Op->Category = MatrixCategory;
+	Op->FriendlyName = NSLOCTEXT("NiagaraOpInfo", "TransformVector Name", "Transform Vector");
+	Op->Description = NSLOCTEXT("NiagaraOpInfo", "TransformVector Desc", "Transforms a Vector3 as a vector.");
+	Op->Inputs.Add(FNiagaraOpInOutInfo(M, MatrixType, MText, MText, Default_MatrixOne));
+	Op->Inputs.Add(FNiagaraOpInOutInfo(V, FNiagaraTypeDefinition::GetVec3Def(), VText, VText, Default_VectorOne));
+	Op->Outputs.Add(FNiagaraOpInOutInfo(Result, FNiagaraTypeDefinition::GetVec3Def(), ResultText, ResultText, Default_VectorOne, TEXT("mul(float4({1},0.0),{0}).xyz")));
+	Op->BuildName(TEXT("TransformVector"), MatrixCategoryName);
+	OpInfoMap.Add(Op->Name) = Idx;
+
+	/* Vector3 only ops **/
+	//TODO: Replace all categories with "Math" or something like that?
+	FText Vec3Category = NSLOCTEXT("NiagaraOpInfo", "Vector3OpCategory", "Vector3");
+	FString Vec3CategoryName(TEXT("Vector3"));
+	Idx = OpInfos.AddDefaulted();
+	Op = &OpInfos[Idx];
+	Op->Category = Vec3Category;
+	Op->FriendlyName = NSLOCTEXT("NiagaraOpInfo", "Vector Cross Name", "Cross");
+	Op->Description = NSLOCTEXT("NiagaraOpInfo", "Vector Cross Desc", "Cross product of two vectors.");
+	Op->Inputs.Add(FNiagaraOpInOutInfo(A, FNiagaraTypeDefinition::GetVec3Def(), AText, AText, Default_VectorX));
+	Op->Inputs.Add(FNiagaraOpInOutInfo(B, FNiagaraTypeDefinition::GetVec3Def(), BText, BText, Default_VectorY));
+	Op->Outputs.Add(FNiagaraOpInOutInfo(Result, FNiagaraTypeDefinition::GetVec3Def(), ResultText, ResultText, Default_VectorZ, TEXT("cross({0},{1})")));
+	Op->BuildName(TEXT("Cross"), Vec3CategoryName);
+	OpInfoMap.Add(Op->Name) = Idx;
+
+	FText UtilCategory = NSLOCTEXT("NiagaraOpInfo", "UtilOpCategory", "Util");
+	FString UtilCategoryName(TEXT("Util"));
+	Idx = OpInfos.AddDefaulted();
+	Op = &OpInfos[Idx];
+	Op->Category = UtilCategory;
+	Op->FriendlyName = NSLOCTEXT("NiagaraOpInfo", "ExecIndex Name", "Execution Index");
+	Op->Description = NSLOCTEXT("NiagaraOpInfo", "ExexIndex Desc", "Returns the index of this particle in the current execution. For example, in a spawn script this gives the index of the particle being spawned which can be used to interpolate it's position.");	
+	Op->Outputs.Add(FNiagaraOpInOutInfo(Result, IntType, ResultText, ResultText, Default_IntZero, TEXT("ExecIndex()")));
+	Op->BuildName(TEXT("ExecIndex"), UtilCategoryName);
+	OpInfoMap.Add(Op->Name) = Idx;
+}
+END_FUNCTION_BUILD_OPTIMIZATION
+//////////////////////////////////////////////////////////////////////////
+
+
+/*-----------------------------------------------------------------------------
+UActorFactoryNiagara
+-----------------------------------------------------------------------------*/
+UActorFactoryNiagara::UActorFactoryNiagara(const FObjectInitializer& ObjectInitializer)
+	: Super(ObjectInitializer)
+{
+	DisplayName = LOCTEXT("NiagaraSystemDisplayName", "NiagaraSystem");
+	NewActorClass = ANiagaraActor::StaticClass();
+}
+
+bool UActorFactoryNiagara::CanCreateActorFrom(const FAssetData& AssetData, FText& OutErrorMsg)
+{
+	if (!AssetData.IsValid() || !AssetData.GetClass()->IsChildOf(UNiagaraSystem::StaticClass()))
+	{
+		OutErrorMsg = NSLOCTEXT("CanCreateActor", "NoSystem", "A valid Niagara System must be specified.");
+		return false;
+	}
+
+	return true;
+}
+
+void UActorFactoryNiagara::PostSpawnActor(UObject* Asset, AActor* NewActor)
+{
+	Super::PostSpawnActor(Asset, NewActor);
+
+	UNiagaraSystem* System = CastChecked<UNiagaraSystem>(Asset);
+	ANiagaraActor* NiagaraActor = CastChecked<ANiagaraActor>(NewActor);
+
+	// Term Component
+	NiagaraActor->GetNiagaraComponent()->UnregisterComponent();
+
+	// Change properties
+	NiagaraActor->GetNiagaraComponent()->SetAsset(System);
+
+	// if we're created by Kismet on the server during gameplay, we need to replicate the emitter
+	if (NiagaraActor->GetWorld()->HasBegunPlay() && NiagaraActor->GetWorld()->GetNetMode() != NM_Client)
+	{
+		NiagaraActor->SetReplicates(true);
+		NiagaraActor->bAlwaysRelevant = true;
+		NiagaraActor->NetUpdateFrequency = 0.1f; // could also set bNetTemporary but LD might further trigger it or something
+	}
+
+	// Init Component
+	NiagaraActor->GetNiagaraComponent()->RegisterComponent();
+}
+
+UObject* UActorFactoryNiagara::GetAssetFromActorInstance(AActor* Instance)
+{
+	check(Instance->IsA(NewActorClass));
+	ANiagaraActor* NewActor = CastChecked<ANiagaraActor>(Instance);
+	if (NewActor->GetNiagaraComponent())
+	{
+		return NewActor->GetNiagaraComponent()->GetAsset();
+	}
+	else
+	{
+		return NULL;
+	}
+}
+
+void UActorFactoryNiagara::PostCreateBlueprint(UObject* Asset, AActor* CDO)
+{
+	if (Asset != NULL && CDO != NULL)
+	{
+		UNiagaraSystem* System = CastChecked<UNiagaraSystem>(Asset);
+		ANiagaraActor* Actor = CastChecked<ANiagaraActor>(CDO);
+		Actor->GetNiagaraComponent()->SetAsset(System);
+	}
+}
+
+
+
+#undef LOCTEXT_NAMESPACE
diff --git a/Engine/Plugins/FX/Niagara/Source/NiagaraEditor/Private/NiagaraEditorModule.cpp b/Engine/Plugins/FX/Niagara/Source/NiagaraEditor/Private/NiagaraEditorModule.cpp
new file mode 100644
--- /dev/null
+++ b/Engine/Plugins/FX/Niagara/Source/NiagaraEditor/Private/NiagaraEditorModule.cpp
@@ -0,0 +1,385 @@
+// Copyright 1998-2017 Epic Games, Inc. All Rights Reserved.
+
+#include "NiagaraEditorModule.h"
+#include "NiagaraEditorTickables.h"
+#include "Modules/ModuleManager.h"
+#include "IAssetTypeActions.h"
+#include "AssetToolsModule.h"
+#include "Misc/ConfigCacheIni.h"
+#include "ISequencerModule.h"
+#include "ISettingsModule.h"
+#include "SequencerSettings.h"
+
+#include "AssetTypeActions/AssetTypeActions_NiagaraSystem.h"
+#include "AssetTypeActions/AssetTypeActions_NiagaraEmitter.h"
+#include "AssetTypeActions/AssetTypeActions_NiagaraScript.h"
+#include "AssetTypeActions/AssetTypeActions_NiagaraParameterCollection.h"
+
+#include "EdGraphUtilities.h"
+#include "SGraphPin.h"
+#include "SGraphPinVector4.h"
+#include "SGraphPinNum.h"
+#include "SGraphPinInteger.h"
+#include "SGraphPinVector.h"
+#include "SGraphPinVector2D.h"
+#include "SGraphPinObject.h"
+#include "SGraphPinColor.h"
+#include "SGraphPinBool.h"
+#include "Private/KismetPins/SGraphPinEnum.h"
+#include "SNiagaraGraphPinNumeric.h"
+#include "SNiagaraGraphPinAdd.h"
+#include "NiagaraNodeConvert.h"
+#include "EdGraphSchema_Niagara.h"
+#include "TypeEditorUtilities/NiagaraFloatTypeEditorUtilities.h"
+#include "TypeEditorUtilities/NiagaraIntegerTypeEditorUtilities.h"
+#include "TypeEditorUtilities/NiagaraBoolTypeEditorUtilities.h"
+#include "TypeEditorUtilities/NiagaraFloatTypeEditorUtilities.h"
+#include "TypeEditorUtilities/NiagaraVectorTypeEditorUtilities.h"
+#include "TypeEditorUtilities/NiagaraColorTypeEditorUtilities.h"
+#include "TypeEditorUtilities/NiagaraDataInterfaceCurveTypeEditorUtilities.h"
+#include "NiagaraEditorStyle.h"
+#include "NiagaraEditorCommands.h"
+#include "NiagaraEmitterTrackEditor.h"
+#include "PropertyEditorModule.h"
+#include "NiagaraSettings.h"
+#include "NiagaraModule.h"
+#include "NiagaraShaderModule.h"
+#include "NiagaraDataInterface.h"
+#include "NiagaraDataInterfaceCurve.h"
+#include "NiagaraDataInterfaceVector2DCurve.h"
+#include "NiagaraDataInterfaceVectorCurve.h"
+#include "NiagaraDataInterfaceVector4Curve.h"
+#include "NiagaraDataInterfaceColorCurve.h"
+#include "NiagaraScriptViewModel.h"
+#include "NiagaraSystemViewModel.h"
+#include "NiagaraEmitterViewModel.h"
+#include "TNiagaraGraphPinEditableName.h"
+#include "Class.h"
+
+#include "Customizations/NiagaraComponentDetails.h"
+#include "Customizations/NiagaraTypeCustomizations.h"
+#include "Customizations/NiagaraEventScriptPropertiesCustomization.h"
+
+IMPLEMENT_MODULE( FNiagaraEditorModule, NiagaraEditor );
+
+#define LOCTEXT_NAMESPACE "NiagaraEditorModule"
+
+const FName FNiagaraEditorModule::NiagaraEditorAppIdentifier( TEXT( "NiagaraEditorApp" ) );
+const FLinearColor FNiagaraEditorModule::WorldCentricTabColorScale(0.0f, 0.0f, 0.2f, 0.5f);
+
+EAssetTypeCategories::Type FNiagaraEditorModule::NiagaraAssetCategory;
+
+//////////////////////////////////////////////////////////////////////////
+
+class FNiagaraScriptGraphPanelPinFactory : public FGraphPanelPinFactory
+{
+public:
+	DECLARE_DELEGATE_RetVal_OneParam(TSharedPtr<class SGraphPin>, FCreateGraphPin, UEdGraphPin*);
+
+	/** Registers a delegate for creating a pin for a specific type. */
+	void RegisterTypePin(const UScriptStruct* Type, FCreateGraphPin CreateGraphPin)
+	{
+		TypeToCreatePinDelegateMap.Add(Type, CreateGraphPin);
+	}
+
+	/** Registers a delegate for creating a pin for for a specific miscellaneous sub category. */
+	void RegisterMiscSubCategoryPin(FString SubCategory, FCreateGraphPin CreateGraphPin)
+	{
+		MiscSubCategoryToCreatePinDelegateMap.Add(SubCategory, CreateGraphPin);
+	}
+
+	//~ FGraphPanelPinFactory interface
+	virtual TSharedPtr<class SGraphPin> CreatePin(class UEdGraphPin* InPin) const override
+	{
+		if (const UEdGraphSchema_Niagara* NSchema = Cast<UEdGraphSchema_Niagara>(InPin->GetSchema()))
+		{
+			if (InPin->PinType.PinCategory == UEdGraphSchema_Niagara::PinCategoryType)
+			{
+				const UScriptStruct* Struct = CastChecked<const UScriptStruct>(InPin->PinType.PinSubCategoryObject.Get());
+				const FCreateGraphPin* CreateGraphPin = TypeToCreatePinDelegateMap.Find(Struct);
+				if (CreateGraphPin != nullptr)
+				{
+					return (*CreateGraphPin).Execute(InPin);
+				}
+			}
+			else if (InPin->PinType.PinCategory == UEdGraphSchema_Niagara::PinCategoryEnum)
+			{
+				const UEnum* Enum = Cast<const UEnum>(InPin->PinType.PinSubCategoryObject.Get());
+				if (Enum == nullptr)
+				{
+					UE_LOG(LogNiagaraEditor, Error, TEXT("Pin states that it is of Enum type, but is missing its Enum! Pin Name '%s' Owning Node '%s'. Turning into standard int definition!"), *InPin->PinName,
+						*InPin->GetOwningNode()->GetName());
+					InPin->PinType.PinCategory = UEdGraphSchema_Niagara::PinCategoryType;
+					InPin->PinType.PinSubCategoryObject = FNiagaraTypeDefinition::GetIntStruct();
+					InPin->DefaultValue.Empty();
+					return CreatePin(InPin);
+				}
+				return SNew(TNiagaraGraphPinEditableName<SGraphPinEnum>, InPin);
+			}
+			else if (InPin->PinType.PinCategory == UEdGraphSchema_Niagara::PinCategoryMisc)
+			{
+				const FCreateGraphPin* CreateGraphPin = MiscSubCategoryToCreatePinDelegateMap.Find(InPin->PinType.PinSubCategory);
+				if (CreateGraphPin != nullptr)
+				{
+					return (*CreateGraphPin).Execute(InPin);
+				}
+			}
+
+			return SNew(TNiagaraGraphPinEditableName<SGraphPin>, InPin);
+		}
+		return nullptr;
+	}
+
+private:
+	TMap<const UScriptStruct*, FCreateGraphPin> TypeToCreatePinDelegateMap;
+	TMap<FString, FCreateGraphPin> MiscSubCategoryToCreatePinDelegateMap;
+};
+
+FNiagaraEditorModule::FNiagaraEditorModule() 
+	: SequencerSettings(nullptr)
+{
+}
+
+void FNiagaraEditorModule::StartupModule()
+{
+	MenuExtensibilityManager = MakeShareable(new FExtensibilityManager);
+	ToolBarExtensibilityManager = MakeShareable(new FExtensibilityManager);
+
+	IAssetTools& AssetTools = FModuleManager::LoadModuleChecked<FAssetToolsModule>("AssetTools").Get();
+	NiagaraAssetCategory = AssetTools.RegisterAdvancedAssetCategory(FName(TEXT("FX")), LOCTEXT("NiagaraAssetsCategory", "FX"));
+	RegisterAssetTypeAction(AssetTools, MakeShareable(new FAssetTypeActions_NiagaraSystem()));
+	RegisterAssetTypeAction(AssetTools, MakeShareable(new FAssetTypeActions_NiagaraEmitter()));
+	RegisterAssetTypeAction(AssetTools, MakeShareable(new FAssetTypeActions_NiagaraScript()));
+	RegisterAssetTypeAction(AssetTools, MakeShareable(new FAssetTypeActions_NiagaraParameterCollection()));
+	RegisterAssetTypeAction(AssetTools, MakeShareable(new FAssetTypeActions_NiagaraParameterCollectionInstance()));
+
+	UNiagaraSettings::OnSettingsChanged().AddRaw(this, &FNiagaraEditorModule::OnNiagaraSettingsChangedEvent);
+
+	// register details customization
+	FPropertyEditorModule& PropertyModule = FModuleManager::LoadModuleChecked<FPropertyEditorModule>("PropertyEditor");
+	PropertyModule.RegisterCustomClassLayout("NiagaraComponent", FOnGetDetailCustomizationInstance::CreateStatic(&FNiagaraComponentDetails::MakeInstance));
+	
+	PropertyModule.RegisterCustomPropertyTypeLayout("NiagaraFloat",
+		FOnGetPropertyTypeCustomizationInstance::CreateStatic(&FNiagaraNumericCustomization::MakeInstance)
+	);
+
+	PropertyModule.RegisterCustomPropertyTypeLayout("NiagaraInt32",
+		FOnGetPropertyTypeCustomizationInstance::CreateStatic(&FNiagaraNumericCustomization::MakeInstance)
+	);
+
+	PropertyModule.RegisterCustomPropertyTypeLayout("NiagaraNumeric",
+		FOnGetPropertyTypeCustomizationInstance::CreateStatic(&FNiagaraNumericCustomization::MakeInstance)
+	);
+
+	PropertyModule.RegisterCustomPropertyTypeLayout("NiagaraParameterMap",
+		FOnGetPropertyTypeCustomizationInstance::CreateStatic(&FNiagaraNumericCustomization::MakeInstance)
+	);
+
+
+	PropertyModule.RegisterCustomPropertyTypeLayout("NiagaraBool",
+		FOnGetPropertyTypeCustomizationInstance::CreateStatic(&FNiagaraBoolCustomization::MakeInstance)
+	);
+
+	PropertyModule.RegisterCustomPropertyTypeLayout("NiagaraMatrix",
+		FOnGetPropertyTypeCustomizationInstance::CreateStatic(&FNiagaraMatrixCustomization::MakeInstance)
+	);
+
+	FNiagaraEditorStyle::Initialize();
+	FNiagaraEditorCommands::Register();
+
+	TSharedPtr<FNiagaraScriptGraphPanelPinFactory> GraphPanelPinFactory = MakeShareable(new FNiagaraScriptGraphPanelPinFactory());
+
+	GraphPanelPinFactory->RegisterTypePin(FNiagaraTypeDefinition::GetFloatStruct(), FNiagaraScriptGraphPanelPinFactory::FCreateGraphPin::CreateLambda(
+		[](UEdGraphPin* GraphPin) -> TSharedRef<SGraphPin> { return SNew(TNiagaraGraphPinEditableName<SGraphPinNum>, GraphPin); }));
+
+	GraphPanelPinFactory->RegisterTypePin(FNiagaraTypeDefinition::GetIntStruct(), FNiagaraScriptGraphPanelPinFactory::FCreateGraphPin::CreateLambda(
+		[](UEdGraphPin* GraphPin) -> TSharedRef<SGraphPin> { return SNew(TNiagaraGraphPinEditableName<SGraphPinInteger>, GraphPin); }));
+
+	GraphPanelPinFactory->RegisterTypePin(FNiagaraTypeDefinition::GetVec2Struct(), FNiagaraScriptGraphPanelPinFactory::FCreateGraphPin::CreateLambda(
+		[](UEdGraphPin* GraphPin) -> TSharedRef<SGraphPin> { return SNew(TNiagaraGraphPinEditableName<SGraphPinVector2D>, GraphPin); }));
+
+	GraphPanelPinFactory->RegisterTypePin(FNiagaraTypeDefinition::GetVec3Struct(), FNiagaraScriptGraphPanelPinFactory::FCreateGraphPin::CreateLambda(
+		[](UEdGraphPin* GraphPin) -> TSharedRef<SGraphPin> { return SNew(TNiagaraGraphPinEditableName<SGraphPinVector>, GraphPin); }));
+
+	GraphPanelPinFactory->RegisterTypePin(FNiagaraTypeDefinition::GetVec4Struct(), FNiagaraScriptGraphPanelPinFactory::FCreateGraphPin::CreateLambda(
+		[](UEdGraphPin* GraphPin) -> TSharedRef<SGraphPin> { return SNew(TNiagaraGraphPinEditableName<SGraphPinVector4>, GraphPin); }));
+
+	GraphPanelPinFactory->RegisterTypePin(FNiagaraTypeDefinition::GetColorStruct(), FNiagaraScriptGraphPanelPinFactory::FCreateGraphPin::CreateLambda(
+		[](UEdGraphPin* GraphPin) -> TSharedRef<SGraphPin> { return SNew(TNiagaraGraphPinEditableName<SGraphPinColor>, GraphPin); }));
+
+	GraphPanelPinFactory->RegisterTypePin(FNiagaraTypeDefinition::GetBoolStruct(), FNiagaraScriptGraphPanelPinFactory::FCreateGraphPin::CreateLambda(
+		[](UEdGraphPin* GraphPin) -> TSharedRef<SGraphPin> { return SNew(TNiagaraGraphPinEditableName<SGraphPinBool>, GraphPin); }));
+
+	GraphPanelPinFactory->RegisterTypePin(FNiagaraTypeDefinition::GetGenericNumericStruct(), FNiagaraScriptGraphPanelPinFactory::FCreateGraphPin::CreateLambda(
+		[](UEdGraphPin* GraphPin) -> TSharedRef<SGraphPin> { return SNew(TNiagaraGraphPinEditableName<SNiagaraGraphPinNumeric>, GraphPin); }));
+
+	// TODO: Don't register this here.
+	GraphPanelPinFactory->RegisterMiscSubCategoryPin(UNiagaraNodeWithDynamicPins::AddPinSubCategory, FNiagaraScriptGraphPanelPinFactory::FCreateGraphPin::CreateLambda(
+		[](UEdGraphPin* GraphPin) -> TSharedRef<SGraphPin> { return SNew(SNiagaraGraphPinAdd, GraphPin); }));
+
+	RegisterTypeUtilities(FNiagaraTypeDefinition::GetFloatDef(), MakeShareable(new FNiagaraEditorFloatTypeUtilities()));
+	RegisterTypeUtilities(FNiagaraTypeDefinition::GetIntDef(), MakeShareable(new FNiagaraEditorIntegerTypeUtilities()));
+	RegisterTypeUtilities(FNiagaraTypeDefinition::GetBoolDef(), MakeShareable(new FNiagaraEditorBoolTypeUtilities()));
+	RegisterTypeUtilities(FNiagaraTypeDefinition::GetVec2Def(), MakeShareable(new FNiagaraEditorVector2TypeUtilities()));
+	RegisterTypeUtilities(FNiagaraTypeDefinition::GetVec3Def(), MakeShareable(new FNiagaraEditorVector3TypeUtilities()));
+	RegisterTypeUtilities(FNiagaraTypeDefinition::GetVec4Def(), MakeShareable(new FNiagaraEditorVector4TypeUtilities()));
+	RegisterTypeUtilities(FNiagaraTypeDefinition::GetColorDef(), MakeShareable(new FNiagaraEditorColorTypeUtilities()));
+
+	RegisterTypeUtilities(FNiagaraTypeDefinition(UNiagaraDataInterfaceCurve::StaticClass()), MakeShared<FNiagaraDataInterfaceCurveTypeEditorUtilities>());
+	RegisterTypeUtilities(FNiagaraTypeDefinition(UNiagaraDataInterfaceVector2DCurve::StaticClass()), MakeShared<FNiagaraDataInterfaceCurveTypeEditorUtilities>());
+	RegisterTypeUtilities(FNiagaraTypeDefinition(UNiagaraDataInterfaceVectorCurve::StaticClass()), MakeShared<FNiagaraDataInterfaceVectorCurveTypeEditorUtilities>());
+	RegisterTypeUtilities(FNiagaraTypeDefinition(UNiagaraDataInterfaceVector4Curve::StaticClass()), MakeShared<FNiagaraDataInterfaceVectorCurveTypeEditorUtilities>());
+	RegisterTypeUtilities(FNiagaraTypeDefinition(UNiagaraDataInterfaceColorCurve::StaticClass()), MakeShared<FNiagaraDataInterfaceColorCurveTypeEditorUtilities>());
+
+	FEdGraphUtilities::RegisterVisualPinFactory(GraphPanelPinFactory);
+
+	FNiagaraOpInfo::Init();
+
+	RegisterSettings();
+
+	// Register sequencer track editor
+	ISequencerModule &SequencerModule = FModuleManager::LoadModuleChecked<ISequencerModule>("Sequencer");
+	CreateEmitterTrackEditorHandle = SequencerModule.RegisterTrackEditor(FOnCreateTrackEditor::CreateStatic(&FNiagaraEmitterTrackEditor::CreateTrackEditor));
+
+	// Register the shader queue processor (for cooking)
+	INiagaraModule& NiagaraModule = FModuleManager::LoadModuleChecked<INiagaraModule>("Niagara");
+	NiagaraModule.SetOnProcessShaderCompilationQueue(INiagaraModule::FOnProcessQueue::CreateLambda([]()
+	{
+		FNiagaraShaderQueueTickable::ProcessQueue();
+	}));
+
+	INiagaraShaderModule& NiagaraShaderModule = FModuleManager::LoadModuleChecked<INiagaraShaderModule>("NiagaraShader");
+	NiagaraShaderModule.SetOnProcessShaderCompilationQueue(INiagaraShaderModule::FOnProcessQueue::CreateLambda([]()
+	{
+		FNiagaraShaderQueueTickable::ProcessQueue();
+	}));
+
+}
+
+
+void FNiagaraEditorModule::ShutdownModule()
+{
+	MenuExtensibilityManager.Reset();
+	ToolBarExtensibilityManager.Reset();
+
+	if (FModuleManager::Get().IsModuleLoaded("AssetTools"))
+	{
+		IAssetTools& AssetTools = FModuleManager::GetModuleChecked<FAssetToolsModule>("AssetTools").Get();
+		for (auto CreatedAssetTypeAction : CreatedAssetTypeActions)
+		{
+			AssetTools.UnregisterAssetTypeActions(CreatedAssetTypeAction.ToSharedRef());
+		}
+	}
+	CreatedAssetTypeActions.Empty();
+
+	UNiagaraSettings::OnSettingsChanged().RemoveAll(this);
+
+	
+	if (FModuleManager::Get().IsModuleLoaded("PropertyEditor"))
+	{
+		FPropertyEditorModule& PropertyModule = FModuleManager::LoadModuleChecked<FPropertyEditorModule>("PropertyEditor");
+		PropertyModule.UnregisterCustomClassLayout("NiagaraComponent");
+	}
+
+	FNiagaraEditorStyle::Shutdown();
+
+	UnregisterSettings();
+
+	ISequencerModule* SequencerModule = FModuleManager::GetModulePtr<ISequencerModule>("Sequencer");
+	if (SequencerModule != nullptr)
+	{
+		SequencerModule->UnRegisterTrackEditor(CreateEmitterTrackEditorHandle);
+	}
+
+	// Verify that we've cleaned up all the view models in the world.
+	FNiagaraScriptViewModel::CleanAll();
+	FNiagaraSystemViewModel::CleanAll();
+	FNiagaraEmitterViewModel::CleanAll();
+}
+
+void FNiagaraEditorModule::OnNiagaraSettingsChangedEvent(const FString& PropertyName, const UNiagaraSettings* Settings)
+{
+	if (PropertyName == "AdditionalParameterTypes" || PropertyName == "AdditionalPayloadTypes")
+	{
+		FNiagaraTypeDefinition::RecreateUserDefinedTypeRegistry();
+	}
+}
+
+void FNiagaraEditorModule::RegisterTypeUtilities(FNiagaraTypeDefinition Type, TSharedRef<INiagaraEditorTypeUtilities> EditorUtilities)
+{
+	TypeToEditorUtilitiesMap.Add(Type, EditorUtilities);
+}
+
+
+TSharedPtr<INiagaraEditorTypeUtilities> FNiagaraEditorModule::GetTypeUtilities(const FNiagaraTypeDefinition& Type)
+{
+	TSharedRef<INiagaraEditorTypeUtilities>* EditorUtilities = TypeToEditorUtilitiesMap.Find(Type);
+	if (EditorUtilities != nullptr)
+	{
+		return *EditorUtilities;
+	}
+	return TSharedPtr<INiagaraEditorTypeUtilities>();
+}
+
+TSharedRef<SWidget> FNiagaraEditorModule::CreateStackWidget(UNiagaraStackViewModel* StackViewModel) const
+{
+	checkf(OnCreateStackWidget.IsBound(), TEXT("Can not create stack widget.  Stack creation delegate was never set."));
+	return OnCreateStackWidget.Execute(StackViewModel);
+}
+
+FDelegateHandle FNiagaraEditorModule::SetOnCreateStackWidget(FOnCreateStackWidget InOnCreateStackWidget)
+{
+	checkf(OnCreateStackWidget.IsBound() == false, TEXT("Stack creation delegate already set."));
+	OnCreateStackWidget = InOnCreateStackWidget;
+	return OnCreateStackWidget.GetHandle();
+}
+
+void FNiagaraEditorModule::ResetOnCreateStackWidget(FDelegateHandle Handle)
+{
+	checkf(OnCreateStackWidget.GetHandle() == Handle, TEXT("Can only reset the stack creation module with the handle it was created with."));
+	OnCreateStackWidget.Unbind();
+}
+
+
+
+
+void FNiagaraEditorModule::RegisterAssetTypeAction(IAssetTools& AssetTools, TSharedRef<IAssetTypeActions> Action)
+{
+	AssetTools.RegisterAssetTypeActions(Action);
+	CreatedAssetTypeActions.Add(Action);
+}
+
+void FNiagaraEditorModule::RegisterSettings()
+{
+	ISettingsModule* SettingsModule = FModuleManager::GetModulePtr<ISettingsModule>("Settings");
+
+	if (SettingsModule != nullptr)
+	{
+		SequencerSettings = USequencerSettingsContainer::GetOrCreate<USequencerSettings>(TEXT("NiagaraSequenceEditor"));
+
+		SettingsModule->RegisterSettings("Editor", "ContentEditors", "NiagaraSequenceEditor",
+			LOCTEXT("NiagaraSequenceEditorSettingsName", "Niagara Sequence Editor"),
+			LOCTEXT("NiagaraSequenceEditorSettingsDescription", "Configure the look and feel of the Niagara Sequence Editor."),
+			SequencerSettings);	
+	}
+}
+
+void FNiagaraEditorModule::UnregisterSettings()
+{
+	ISettingsModule* SettingsModule = FModuleManager::GetModulePtr<ISettingsModule>("Settings");
+
+	if (SettingsModule != nullptr)
+	{
+		SettingsModule->UnregisterSettings("Editor", "ContentEditors", "NiagaraSequenceEditor");
+	}
+}
+
+void FNiagaraEditorModule::AddReferencedObjects( FReferenceCollector& Collector )
+{
+	if (SequencerSettings)
+	{
+		Collector.AddReferencedObject(SequencerSettings);
+	}
+}
+
+#undef LOCTEXT_NAMESPACE
\ No newline at end of file
diff --git a/Engine/Plugins/FX/Niagara/Source/NiagaraEditor/Private/NiagaraEditorStyle.cpp b/Engine/Plugins/FX/Niagara/Source/NiagaraEditor/Private/NiagaraEditorStyle.cpp
new file mode 100644
--- /dev/null
+++ b/Engine/Plugins/FX/Niagara/Source/NiagaraEditor/Private/NiagaraEditorStyle.cpp
@@ -0,0 +1,184 @@
+// Copyright 1998-2016 Epic Games, Inc. All Rights Reserved.
+
+#include "NiagaraEditorStyle.h"
+
+#include "SlateApplication.h"
+#include "EditorStyleSet.h"
+#include "Slate/SlateGameResources.h"
+#include "Styling/SlateStyleRegistry.h"
+#include "Styling/SlateTypes.h"
+
+TSharedPtr< FSlateStyleSet > FNiagaraEditorStyle::NiagaraEditorStyleInstance = NULL;
+
+void FNiagaraEditorStyle::Initialize()
+{
+	if (!NiagaraEditorStyleInstance.IsValid())
+	{
+		NiagaraEditorStyleInstance = Create();
+		FSlateStyleRegistry::RegisterSlateStyle(*NiagaraEditorStyleInstance);
+	}
+}
+
+void FNiagaraEditorStyle::Shutdown()
+{
+	FSlateStyleRegistry::UnRegisterSlateStyle(*NiagaraEditorStyleInstance);
+	ensure(NiagaraEditorStyleInstance.IsUnique());
+	NiagaraEditorStyleInstance.Reset();
+}
+
+FName FNiagaraEditorStyle::GetStyleSetName()
+{
+	static FName StyleSetName(TEXT("NiagaraEditorStyle"));
+	return StyleSetName;
+}
+
+#define IMAGE_BRUSH( RelativePath, ... ) FSlateImageBrush( Style->RootToContentDir( RelativePath, TEXT(".png") ), __VA_ARGS__ )
+#define BOX_BRUSH( RelativePath, ... ) FSlateBoxBrush( Style->RootToContentDir( RelativePath, TEXT(".png") ), __VA_ARGS__ )
+#define BORDER_BRUSH( RelativePath, ... ) FSlateBorderBrush( Style->RootToContentDir( RelativePath, TEXT(".png") ), __VA_ARGS__ )
+#define TTF_FONT( RelativePath, ... ) FSlateFontInfo( Style->RootToContentDir( RelativePath, TEXT(".ttf") ), __VA_ARGS__ )
+#define OTF_FONT( RelativePath, ... ) FSlateFontInfo( Style->RootToContentDir( RelativePath, TEXT(".otf") ), __VA_ARGS__ )
+#define TTF_CORE_FONT( RelativePath, ... ) FSlateFontInfo( FPaths::EngineContentDir()  / "Slate" / RelativePath + TEXT(".ttf"), __VA_ARGS__ )
+#define BOX_CORE_BRUSH( RelativePath, ... ) FSlateBoxBrush( FPaths::EngineContentDir() / "Editor/Slate" / RelativePath + TEXT(".png"), __VA_ARGS__ )
+
+const FVector2D Icon16x16(16.0f, 16.0f);
+const FVector2D Icon20x20(20.0f, 20.0f);
+const FVector2D Icon40x40(40.0f, 40.0f);
+
+TSharedRef< FSlateStyleSet > FNiagaraEditorStyle::Create()
+{
+	const FTextBlockStyle NormalText = FEditorStyle::GetWidgetStyle<FTextBlockStyle>("NormalText");
+	const FEditableTextBoxStyle NormalEditableTextBox = FCoreStyle::Get().GetWidgetStyle<FEditableTextBoxStyle>("NormalEditableTextBox");
+	const FSpinBoxStyle NormalSpinBox = FEditorStyle::GetWidgetStyle<FSpinBoxStyle>("SpinBox");
+
+	TSharedRef< FSlateStyleSet > Style = MakeShareable(new FSlateStyleSet("NiagaraEditorStyle"));
+	Style->SetContentRoot(FPaths::EngineContentDir() / TEXT("Editor/Slate/Niagara"));
+
+	// Stats
+	FTextBlockStyle CategoryText = FTextBlockStyle(NormalText)
+		.SetFont(TTF_CORE_FONT("Fonts/Roboto-Regular", 10))
+		.SetShadowOffset(FVector2D(0, 1))
+		.SetShadowColorAndOpacity(FLinearColor(0, 0, 0, 0.9f));
+
+	Style->Set("NiagaraEditor.StatsText", CategoryText);
+
+	// Emitter Header
+	FTextBlockStyle HeadingText = FTextBlockStyle(NormalText)
+		.SetFont(TTF_CORE_FONT("Fonts/Roboto-Regular", 14))
+		.SetShadowOffset(FVector2D(0, 1))
+		.SetShadowColorAndOpacity(FLinearColor(0, 0, 0, 0.9f));
+
+	FEditableTextBoxStyle HeadingEditableTextBox = FEditableTextBoxStyle(NormalEditableTextBox)
+		.SetFont(TTF_CORE_FONT("Fonts/Roboto-Regular", 14));
+
+	Style->Set("NiagaraEditor.HeadingEditableTextBox", HeadingEditableTextBox);
+
+	Style->Set("NiagaraEditor.HeadingInlineEditableText", FInlineEditableTextBlockStyle()
+		.SetTextStyle(HeadingText)
+		.SetEditableTextBoxStyle(HeadingEditableTextBox));
+
+	FTextBlockStyle TabText = FTextBlockStyle(NormalText)
+		.SetFont(TTF_CORE_FONT("Fonts/Roboto-Regular", 12))
+		.SetShadowOffset(FVector2D(0, 1))
+		.SetShadowColorAndOpacity(FLinearColor(0, 0, 0, 0.9f));
+	
+	Style->Set("NiagaraEditor.AttributeSpreadsheetTabText", TabText);
+	
+	// Parameters
+	FSlateFontInfo ParameterFont = TTF_CORE_FONT("Fonts/Roboto-Regular", 8);
+
+	Style->Set("NiagaraEditor.ParameterFont", ParameterFont);
+
+	FTextBlockStyle ParameterText = FTextBlockStyle(NormalText)
+		.SetFont(ParameterFont);
+
+	Style->Set("NiagaraEditor.ParameterText", ParameterText);
+
+	FEditableTextBoxStyle ParameterEditableTextBox = FEditableTextBoxStyle(NormalEditableTextBox)
+		.SetFont(ParameterFont);
+
+	Style->Set("NiagaraEditor.ParameterEditableTextBox", ParameterEditableTextBox);
+
+	Style->Set("NiagaraEditor.ParameterInlineEditableText", FInlineEditableTextBlockStyle()
+		.SetTextStyle(ParameterText)
+		.SetEditableTextBoxStyle(ParameterEditableTextBox));
+
+	FSpinBoxStyle ParameterSpinBox = FSpinBoxStyle(NormalSpinBox)
+		.SetTextPadding(1);
+
+	Style->Set("NiagaraEditor.ParameterSpinbox", ParameterSpinBox);
+
+	// Code View
+	{
+		Style->Set("NiagaraEditor.CodeView.Checkbox.Text", FTextBlockStyle(NormalText)
+			.SetFont(TTF_CORE_FONT("Fonts/Roboto-Bold", 12))
+			.SetColorAndOpacity(FLinearColor(1.0f, 1.0f, 1.0f, 0.9f))
+			.SetShadowOffset(FVector2D(1, 1))
+			.SetShadowColorAndOpacity(FLinearColor(0, 0, 0, 0.9f)));
+
+		const int32 LogFontSize = 9;
+		FSlateFontInfo LogFont = TTF_CORE_FONT("Fonts/DroidSansMono", LogFontSize);
+		FTextBlockStyle NormalLogText = FTextBlockStyle(NormalText)
+			.SetFont(LogFont)
+			.SetColorAndOpacity(FLinearColor(FColor(0xffffffff)))
+			.SetSelectedBackgroundColor(FLinearColor(FColor(0xff666666)));
+		Style->Set("NiagaraEditor.CodeView.Hlsl.Normal", NormalLogText);
+	}
+
+	// Selected Emitter
+	FSlateFontInfo SelectedEmitterUnsupportedSelectionFont = TTF_CORE_FONT("Fonts/Roboto", 10);
+	FTextBlockStyle SelectedEmitterUnsupportedSelectionText = FTextBlockStyle(NormalText)
+		.SetFont(SelectedEmitterUnsupportedSelectionFont)
+		.SetColorAndOpacity(FLinearColor(1.0f, 1.0f, 1.0f, 1.0f))
+		.SetShadowOffset(FVector2D(0, 1))
+		.SetShadowColorAndOpacity(FLinearColor(0, 0, 0, 0.9f));
+	Style->Set("NiagaraEditor.SelectedEmitter.UnsupportedSelectionText", SelectedEmitterUnsupportedSelectionText);
+
+	// Icons
+	Style->Set("NiagaraEditor.Apply", new IMAGE_BRUSH("Icons/icon_Niagara_Apply_40x", Icon40x40));
+	Style->Set("NiagaraEditor.Apply.Small", new IMAGE_BRUSH("Icons/icon_Niagara_Apply_40x", Icon20x20));
+	Style->Set("NiagaraEditor.Compile", new IMAGE_BRUSH("Icons/icon_compile_40x", Icon40x40));
+	Style->Set("NiagaraEditor.Compile.Small", new IMAGE_BRUSH("Icons/icon_compile_40x", Icon20x20));
+	Style->Set("NiagaraEditor.AddEmitter", new IMAGE_BRUSH("Icons/icon_AddObject_40x", Icon40x40));
+	Style->Set("NiagaraEditor.AddEmitter.Small", new IMAGE_BRUSH("Icons/icon_AddObject_40x", Icon20x20));
+	Style->Set("NiagaraEditor.UnlockToChanges", new IMAGE_BRUSH("Icons/icon_levels_unlocked_40x", Icon40x40));
+	Style->Set("NiagaraEditor.UnlockToChanges.Small", new IMAGE_BRUSH("Icons/icon_levels_unlocked_40x", Icon20x20));
+	Style->Set("NiagaraEditor.LockToChanges", new IMAGE_BRUSH("Icons/icon_levels_LockedReadOnly_40x", Icon40x40));
+	Style->Set("NiagaraEditor.LockToChanges.Small", new IMAGE_BRUSH("Icons/icon_levels_LockedReadOnly_40x", Icon20x20));
+
+	Style->Set("Niagara.CompileStatus.Unknown", new IMAGE_BRUSH("Icons/CompileStatus_Working", Icon40x40));
+	Style->Set("Niagara.CompileStatus.Error",   new IMAGE_BRUSH("Icons/CompileStatus_Fail", Icon40x40));
+	Style->Set("Niagara.CompileStatus.Good",    new IMAGE_BRUSH("Icons/CompileStatus_Good", Icon40x40));
+	Style->Set("Niagara.CompileStatus.Warning", new IMAGE_BRUSH("Icons/CompileStatus_Warning", Icon40x40));
+	Style->Set("Niagara.Asset.ReimportAsset.Needed", new IMAGE_BRUSH("Icons/icon_Reimport_Needed_40x", Icon40x40));
+	Style->Set("Niagara.Asset.ReimportAsset.Default", new IMAGE_BRUSH("Icons/icon_Reimport_40x", Icon40x40));
+	
+	// Emitter details customization
+	Style->Set("NiagaraEditor.MaterialWarningBorder", new BOX_CORE_BRUSH("Common/GroupBorderLight", FMargin(4.0f / 16.0f)));
+
+	// Asset colors
+	Style->Set("NiagaraEditor.AssetColors.System", FLinearColor(1.0f, 0.0f, 0.0f));
+	Style->Set("NiagaraEditor.AssetColors.Emitter", FLinearColor(1.0f, 0.3f, 0.0f));
+	Style->Set("NiagaraEditor.AssetColors.Script", FLinearColor(1.0f, 1.0f, 0.0f));
+	Style->Set("NiagaraEditor.AssetColors.ParameterCollection", FLinearColor(1.0f, 1.0f, 0.3f));
+	Style->Set("NiagaraEditor.AssetColors.ParameterCollectionInstance", FLinearColor(1.0f, 1.0f, 0.7f));
+
+	return Style;
+}
+
+#undef IMAGE_BRUSH
+#undef BOX_BRUSH
+#undef BORDER_BRUSH
+#undef TTF_FONT
+#undef OTF_FONT
+#undef TTF_CORE_FONT
+#undef BOX_CORE_BRUSH
+
+void FNiagaraEditorStyle::ReloadTextures()
+{
+	FSlateApplication::Get().GetRenderer()->ReloadTextureResources();
+}
+
+const ISlateStyle& FNiagaraEditorStyle::Get()
+{
+	return *NiagaraEditorStyleInstance;
+}
\ No newline at end of file
diff --git a/Engine/Plugins/FX/Niagara/Source/NiagaraEditor/Private/NiagaraEditorUtilities.cpp b/Engine/Plugins/FX/Niagara/Source/NiagaraEditor/Private/NiagaraEditorUtilities.cpp
new file mode 100644
--- /dev/null
+++ b/Engine/Plugins/FX/Niagara/Source/NiagaraEditor/Private/NiagaraEditorUtilities.cpp
@@ -0,0 +1,582 @@
+// Copyright 1998-2016 Epic Games, Inc. All Rights Reserved.
+
+#include "NiagaraEditorUtilities.h"
+#include "NiagaraEditorModule.h"
+#include "INiagaraEditorTypeUtilities.h"
+#include "NiagaraNodeInput.h"
+#include "NiagaraDataInterface.h"
+#include "NiagaraComponent.h"
+#include "ModuleManager.h"
+#include "StructOnScope.h"
+#include "NiagaraGraph.h"
+#include "NiagaraSystem.h"
+#include "NiagaraScriptSource.h"
+#include "NiagaraScript.h"
+#include "NiagaraNodeOutput.h"
+#include "EdGraphUtilities.h"
+#include "NiagaraConstants.h"
+#include "SWidget.h"
+#include "STextBlock.h"
+#include "SImage.h"
+#include "SBoxPanel.h"
+#include "HAL/PlatformApplicationMisc.h"
+#include "NiagaraEditorStyle.h"
+#include "EditorStyleSet.h"
+#include "NiagaraSystemViewModel.h"
+#include "NiagaraEmitterViewModel.h"
+#include "PlatformApplicationMisc.h"
+
+#define LOCTEXT_NAMESPACE "FNiagaraEditorUtilities"
+
+FName FNiagaraEditorUtilities::GetUniqueName(FName CandidateName, const TSet<FName>& ExistingNames)
+{
+	if (ExistingNames.Contains(CandidateName) == false)
+	{
+		return CandidateName;
+	}
+
+	FString CandidateNameString = CandidateName.ToString();
+	FString BaseNameString = CandidateNameString;
+	if (CandidateNameString.Len() >= 3 && CandidateNameString.Right(3).IsNumeric())
+	{
+		BaseNameString = CandidateNameString.Left(CandidateNameString.Len() - 3);
+	}
+
+	FName UniqueName = FName(*BaseNameString);
+	int32 NameIndex = 1;
+	while (ExistingNames.Contains(UniqueName))
+	{
+		UniqueName = FName(*FString::Printf(TEXT("%s%03i"), *BaseNameString, NameIndex));
+		NameIndex++;
+	}
+
+	return UniqueName;
+}
+
+TSet<FName> FNiagaraEditorUtilities::GetSystemConstantNames()
+{
+	TSet<FName> SystemConstantNames;
+	for (const FNiagaraVariable& SystemConstant : FNiagaraConstants::GetEngineConstants())
+	{
+		SystemConstantNames.Add(SystemConstant.GetName());
+	}
+	return SystemConstantNames;
+}
+
+void FNiagaraEditorUtilities::GetTypeDefaultValue(const FNiagaraTypeDefinition& Type, TArray<uint8>& DefaultData)
+{
+	if (const UScriptStruct* ScriptStruct = Type.GetScriptStruct())
+	{
+		int32 TypeSize = Type.GetSize();
+		DefaultData.SetNumUninitialized(TypeSize);
+
+		FNiagaraEditorModule& NiagaraEditorModule = FModuleManager::GetModuleChecked<FNiagaraEditorModule>("NiagaraEditor");
+		TSharedPtr<INiagaraEditorTypeUtilities> TypeEditorUtilities = NiagaraEditorModule.GetTypeUtilities(Type);
+		if (TypeEditorUtilities.IsValid() && TypeEditorUtilities->CanProvideDefaultValue())
+		{
+			TSharedRef<FStructOnScope> Struct = MakeShareable(new FStructOnScope(ScriptStruct));
+			TypeEditorUtilities->UpdateStructWithDefaultValue(Struct);	
+			FMemory::Memcpy(DefaultData.GetData(), Struct->GetStructMemory(), TypeSize);
+			
+		}
+		else
+		{
+			ScriptStruct->InitializeDefaultValue(DefaultData.GetData());
+		}
+	}
+}
+
+void FNiagaraEditorUtilities::ResetVariableToDefaultValue(FNiagaraVariable& Variable)
+{
+	if (const UScriptStruct* ScriptStruct = Variable.GetType().GetScriptStruct())
+	{
+		FNiagaraEditorModule& NiagaraEditorModule = FModuleManager::GetModuleChecked<FNiagaraEditorModule>("NiagaraEditor");
+		TSharedPtr<INiagaraEditorTypeUtilities> TypeEditorUtilities = NiagaraEditorModule.GetTypeUtilities(Variable.GetType());
+		if (TypeEditorUtilities.IsValid() && TypeEditorUtilities->CanProvideDefaultValue())
+		{
+			TSharedRef<FStructOnScope> Struct = MakeShareable(new FStructOnScope(ScriptStruct));
+			TypeEditorUtilities->UpdateStructWithDefaultValue(Struct);
+			Variable.SetData(Struct->GetStructMemory());
+		}
+		else
+		{
+			Variable.AllocateData();
+			ScriptStruct->InitializeDefaultValue(Variable.GetData());
+		}
+	}
+}
+
+void FNiagaraEditorUtilities::InitializeParameterInputNode(UNiagaraNodeInput& InputNode, const FNiagaraTypeDefinition& Type, const UNiagaraGraph* InGraph, FName InputName)
+{
+	InputNode.Usage = ENiagaraInputNodeUsage::Parameter;
+	InputNode.bCanRenameNode = true;
+	InputName = UNiagaraNodeInput::GenerateUniqueName(InGraph, InputName, ENiagaraInputNodeUsage::Parameter);
+	InputNode.Input.SetName(InputName);
+	InputNode.Input.SetType(Type);
+	if (InGraph) // Only compute sort priority if a graph was passed in, similar to the way that GenrateUniqueName works above.
+	{
+		InputNode.CallSortPriority = UNiagaraNodeInput::GenerateNewSortPriority(InGraph, InputName, ENiagaraInputNodeUsage::Parameter);
+	}
+	if (Type.GetScriptStruct() != nullptr)
+	{
+		ResetVariableToDefaultValue(InputNode.Input);
+		InputNode.DataInterface = nullptr;
+	}
+	else
+	{
+		InputNode.Input.AllocateData(); // Frees previously used memory if we're switching from a struct to a class type.
+		InputNode.DataInterface = NewObject<UNiagaraDataInterface>(&InputNode, const_cast<UClass*>(Type.GetClass()), NAME_None, RF_Transactional);
+	}
+}
+
+void FNiagaraEditorUtilities::GetParameterVariablesFromSystem(UNiagaraSystem& System, TArray<FNiagaraVariable>& ParameterVariables,
+	FNiagaraEditorUtilities::FGetParameterVariablesFromSystemOptions Options)
+{
+	UNiagaraScript* SystemScript = System.GetSystemSpawnScript();
+	if (SystemScript != nullptr)
+	{
+		UNiagaraScriptSource* ScriptSource = Cast<UNiagaraScriptSource>(SystemScript->GetSource());
+		if (ScriptSource != nullptr)
+		{
+			UNiagaraGraph* SystemGraph = ScriptSource->NodeGraph;
+			if (SystemGraph != nullptr)
+			{
+				UNiagaraGraph::FFindInputNodeOptions FindOptions;
+				FindOptions.bIncludeAttributes = false;
+				FindOptions.bIncludeSystemConstants = false;
+				FindOptions.bFilterDuplicates = true;
+
+				TArray<UNiagaraNodeInput*> InputNodes;
+				SystemGraph->FindInputNodes(InputNodes, FindOptions);
+				for (UNiagaraNodeInput* InputNode : InputNodes)
+				{
+					bool bIsStructParameter = InputNode->Input.GetType().GetScriptStruct() != nullptr;
+					bool bIsDataInterfaceParameter = InputNode->Input.GetType().GetClass() != nullptr;
+					if ((bIsStructParameter && Options.bIncludeStructParameters) || (bIsDataInterfaceParameter && Options.bIncludeDataInterfaceParameters))
+					{
+						ParameterVariables.Add(InputNode->Input);
+					}
+				}
+			}
+		}
+	}
+}
+
+bool FNiagaraEditorUtilities::ConvertToMergedGraph(UNiagaraEmitter* InEmitter)
+{
+	if (InEmitter->GraphSource == nullptr)
+	{
+		UNiagaraScriptSource* Source = NewObject<UNiagaraScriptSource>(InEmitter, NAME_None, RF_Transactional);
+		if (Source)
+		{
+			UNiagaraGraph* CreatedGraph = NewObject<UNiagaraGraph>(Source, NAME_None, RF_Transactional);
+			Source->NodeGraph = CreatedGraph;
+
+			TArray<UNiagaraGraph*> GraphsToConvert;
+			TArray<ENiagaraScriptUsage> GraphUsages;
+
+			int32 yMaxPrevious = 0;
+
+			GraphsToConvert.Add(CastChecked<UNiagaraScriptSource>(InEmitter->SpawnScriptProps.Script->GetSource())->NodeGraph);
+			GraphUsages.Add(ENiagaraScriptUsage::ParticleSpawnScript);
+			GraphsToConvert.Add(CastChecked<UNiagaraScriptSource>(InEmitter->UpdateScriptProps.Script->GetSource())->NodeGraph);
+			GraphUsages.Add(ENiagaraScriptUsage::ParticleUpdateScript);
+
+			for (int32 i = 0; i < InEmitter->EventHandlerScriptProps.Num(); i++)
+			{
+				if (InEmitter->EventHandlerScriptProps[i].Script != nullptr)
+				{
+					GraphsToConvert.Add(CastChecked<UNiagaraScriptSource>(InEmitter->EventHandlerScriptProps[i].Script->GetSource())->NodeGraph);
+					GraphUsages.Add(ENiagaraScriptUsage::ParticleEventScript);
+				}
+			}
+
+			for (int32 i = 0; i < GraphsToConvert.Num(); i++)
+			{
+				UNiagaraGraph* Graph = GraphsToConvert[i];
+				ENiagaraScriptUsage GraphUsage = GraphUsages[i];
+
+				TSet<UObject*> NodesToCopy;
+				TArray<UNiagaraNode*> SourceNodes;
+				Graph->GetNodesOfClass<UNiagaraNode>(SourceNodes);
+
+				int32 HighestY = -INT_MAX;
+				int32 EstimatedHeight = 300;
+				for (UNiagaraNode* SelectedGraphNode : SourceNodes)
+				{
+					if (SelectedGraphNode != nullptr)
+					{
+						if (SelectedGraphNode->NodePosY + EstimatedHeight > HighestY)
+						{
+							HighestY = SelectedGraphNode->NodePosY + EstimatedHeight;
+						}
+
+						if (SelectedGraphNode->CanDuplicateNode())
+						{
+							SelectedGraphNode->PrepareForCopying();
+							NodesToCopy.Add(SelectedGraphNode);
+						}
+						else
+						{
+							UE_LOG(LogNiagaraEditor, Error, TEXT("Could not clone node! %s"), *SelectedGraphNode->GetName());
+						}
+					}
+				}
+
+				FString ExportedText;
+				FEdGraphUtilities::ExportNodesToText(NodesToCopy, ExportedText);
+				FPlatformApplicationMisc::ClipboardCopy(*ExportedText);
+
+				// Grab the text to paste from the clipboard.
+				FPlatformApplicationMisc::ClipboardPaste(ExportedText);
+
+				// Import the nodes
+				TSet<UEdGraphNode*> PastedNodes;
+				FEdGraphUtilities::ImportNodesFromText(CreatedGraph, ExportedText, PastedNodes);
+
+
+				for (UEdGraphNode* PastedNode : PastedNodes)
+				{
+					PastedNode->CreateNewGuid();
+					PastedNode->NodePosY += yMaxPrevious;
+
+					UNiagaraNodeOutput* Output = Cast<UNiagaraNodeOutput>(PastedNode);
+					if (Output != nullptr)
+					{
+						Output->SetUsage(GraphUsage);
+					}
+				}
+
+				FixUpPastedInputNodes(CreatedGraph, PastedNodes);
+				yMaxPrevious = yMaxPrevious + HighestY;
+			}
+
+			InEmitter->GraphSource = Source;
+			InEmitter->SpawnScriptProps.Script->SetSource(Source);
+			InEmitter->UpdateScriptProps.Script->SetSource(Source);
+			for (int32 i = 0; i < InEmitter->EventHandlerScriptProps.Num(); i++)
+			{
+				if (InEmitter->EventHandlerScriptProps[i].Script != nullptr)
+				{
+					InEmitter->EventHandlerScriptProps[i].Script->SetSource(Source);
+				}
+			}
+
+			// Also fix up any dependencies' referenced script type..
+			TArray<const UNiagaraGraph*> ReferencedGraphs;
+			CreatedGraph->GetAllReferencedGraphs(ReferencedGraphs);
+			for (const UNiagaraGraph* Graph : ReferencedGraphs)
+			{
+				TArray<UNiagaraNodeOutput*> Outputs;
+				Graph->FindOutputNodes(Outputs);
+				UNiagaraScript* Script = Cast<UNiagaraScript>(Graph->GetOuter());
+				if (Script)
+				{
+					for (UNiagaraNodeOutput* OutputNode : Outputs)
+					{
+						OutputNode->SetUsage(Script->GetUsage());
+					}
+				}
+			}
+			
+			// Now make sure that anyone referencing these graphs knows that they are out-of-date.
+			Source->MarkNotSynchronized();
+			InEmitter->ChangeId = FGuid::NewGuid();
+
+			//FString ErrorMessages;
+			//FNiagaraEditorModule& NiagaraEditorModule = FModuleManager::Get().LoadModuleChecked<FNiagaraEditorModule>(TEXT("NiagaraEditor"));
+			//NiagaraEditorModule.CompileScript(NewScript, ErrorMessages);
+			return true;
+		}
+	}
+	return false;
+}
+
+
+// TODO: This is overly complicated.
+void FNiagaraEditorUtilities::FixUpPastedInputNodes(UEdGraph* Graph, TSet<UEdGraphNode*> PastedNodes)
+{
+	// Collect existing inputs.
+	TArray<UNiagaraNodeInput*> CurrentInputs;
+	Graph->GetNodesOfClass<UNiagaraNodeInput>(CurrentInputs);
+	TSet<FNiagaraVariable> ExistingInputs;
+	TMap<FNiagaraVariable, UNiagaraNodeInput*> ExistingNodes;
+	int32 HighestSortOrder = -1; // Set to -1 initially, so that in the event of no nodes, we still get zero.
+	for (UNiagaraNodeInput* CurrentInput : CurrentInputs)
+	{
+		if (PastedNodes.Contains(CurrentInput) == false && CurrentInput->Usage == ENiagaraInputNodeUsage::Parameter)
+		{
+			ExistingInputs.Add(CurrentInput->Input);
+			ExistingNodes.Add(CurrentInput->Input) = CurrentInput;
+			if (CurrentInput->CallSortPriority > HighestSortOrder)
+			{
+				HighestSortOrder = CurrentInput->CallSortPriority;
+			}
+		}
+	}
+
+	// Collate pasted inputs nodes by their input for further processing.
+	TMap<FNiagaraVariable, TArray<UNiagaraNodeInput*>> InputToPastedInputNodes;
+	for (UEdGraphNode* PastedNode : PastedNodes)
+	{
+		UNiagaraNodeInput* PastedInputNode = Cast<UNiagaraNodeInput>(PastedNode);
+		if (PastedInputNode != nullptr && PastedInputNode->Usage == ENiagaraInputNodeUsage::Parameter && ExistingInputs.Contains(PastedInputNode->Input) == false)
+		{
+			TArray<UNiagaraNodeInput*>* NodesForInput = InputToPastedInputNodes.Find(PastedInputNode->Input);
+			if (NodesForInput == nullptr)
+			{
+				NodesForInput = &InputToPastedInputNodes.Add(PastedInputNode->Input);
+			}
+			NodesForInput->Add(PastedInputNode);
+		}
+	}
+
+	// Fix up the nodes based on their relationship to the existing inputs.
+	for (auto PastedPairIterator = InputToPastedInputNodes.CreateIterator(); PastedPairIterator; ++PastedPairIterator)
+	{
+		FNiagaraVariable PastedInput = PastedPairIterator.Key();
+		TArray<UNiagaraNodeInput*>& PastedNodesForInput = PastedPairIterator.Value();
+
+		// Try to find an existing input which matches the pasted input by both name and type so that the pasted nodes
+		// can be assigned the same id and value, to facilitate pasting multiple times from the same source graph.
+		FNiagaraVariable* MatchingInputByNameAndType = nullptr;
+		UNiagaraNodeInput* MatchingNode = nullptr;
+		for (FNiagaraVariable& ExistingInput : ExistingInputs)
+		{
+			if (PastedInput.GetName() == ExistingInput.GetName() && PastedInput.GetType() == ExistingInput.GetType())
+			{
+				MatchingInputByNameAndType = &ExistingInput;
+				UNiagaraNodeInput** FoundNode = ExistingNodes.Find(ExistingInput);
+				if (FoundNode != nullptr)
+				{
+					MatchingNode = *FoundNode;
+				}
+				break;
+			}
+		}
+
+		if (MatchingInputByNameAndType != nullptr && MatchingNode != nullptr)
+		{
+			// Update the id and value on the matching pasted nodes.
+			for (UNiagaraNodeInput* PastedNodeForInput : PastedNodesForInput)
+			{
+				if (nullptr == PastedNodeForInput)
+				{
+					continue;
+				}
+				PastedNodeForInput->CallSortPriority = MatchingNode->CallSortPriority;
+				PastedNodeForInput->ExposureOptions = MatchingNode->ExposureOptions;
+				PastedNodeForInput->Input.AllocateData();
+				PastedNodeForInput->Input.SetData(MatchingInputByNameAndType->GetData());
+			}
+		}
+		else
+		{
+			// Check for duplicate names
+			TSet<FName> ExistingNames;
+			for (FNiagaraVariable& ExistingInput : ExistingInputs)
+			{
+				ExistingNames.Add(ExistingInput.GetName());
+			}
+			if (ExistingNames.Contains(PastedInput.GetName()))
+			{
+				FName UniqueName = FNiagaraEditorUtilities::GetUniqueName(PastedInput.GetName(), ExistingNames.Union(FNiagaraEditorUtilities::GetSystemConstantNames()));
+				for (UNiagaraNodeInput* PastedNodeForInput : PastedNodesForInput)
+				{
+					PastedNodeForInput->Input.SetName(UniqueName);
+				}
+			}
+
+			// Assign the pasted inputs the same new id and add them to the end of the parameters list.
+			int32 NewSortOrder = ++HighestSortOrder;
+			for (UNiagaraNodeInput* PastedNodeForInput : PastedNodesForInput)
+			{
+				PastedNodeForInput->CallSortPriority = NewSortOrder;
+			}
+		}
+	}
+}
+
+FText FNiagaraEditorUtilities::StatusToText(ENiagaraScriptCompileStatus Status)
+{
+	switch (Status)
+	{
+	default:
+	case ENiagaraScriptCompileStatus::NCS_Unknown:
+		return LOCTEXT("Recompile_Status", "Unknown status; should recompile");
+	case ENiagaraScriptCompileStatus::NCS_Dirty:
+		return LOCTEXT("Dirty_Status", "Dirty; needs to be recompiled");
+	case ENiagaraScriptCompileStatus::NCS_Error:
+		return LOCTEXT("CompileError_Status", "There was an error during compilation, see the log for details");
+	case ENiagaraScriptCompileStatus::NCS_UpToDate:
+		return LOCTEXT("GoodToGo_Status", "Good to go");
+	case ENiagaraScriptCompileStatus::NCS_UpToDateWithWarnings:
+		return LOCTEXT("GoodToGoWarning_Status", "There was a warning during compilation, see the log for details");
+	}
+}
+
+ENiagaraScriptCompileStatus FNiagaraEditorUtilities::UnionCompileStatus(const ENiagaraScriptCompileStatus& StatusA, const ENiagaraScriptCompileStatus& StatusB)
+{
+	if (StatusA != StatusB)
+	{
+		if (StatusA == ENiagaraScriptCompileStatus::NCS_Unknown || StatusB == ENiagaraScriptCompileStatus::NCS_Unknown)
+		{
+			return ENiagaraScriptCompileStatus::NCS_Unknown;
+		}
+		else if (StatusA >= ENiagaraScriptCompileStatus::NCS_MAX || StatusB >= ENiagaraScriptCompileStatus::NCS_MAX)
+		{
+			return ENiagaraScriptCompileStatus::NCS_MAX;
+		}
+		else if (StatusA == ENiagaraScriptCompileStatus::NCS_Dirty || StatusB == ENiagaraScriptCompileStatus::NCS_Dirty)
+		{
+			return ENiagaraScriptCompileStatus::NCS_Dirty;
+		}
+		else if (StatusA == ENiagaraScriptCompileStatus::NCS_Error || StatusB == ENiagaraScriptCompileStatus::NCS_Error)
+		{
+			return ENiagaraScriptCompileStatus::NCS_Error;
+		}
+		else if (StatusA == ENiagaraScriptCompileStatus::NCS_UpToDateWithWarnings || StatusB == ENiagaraScriptCompileStatus::NCS_UpToDateWithWarnings)
+		{
+			return ENiagaraScriptCompileStatus::NCS_UpToDateWithWarnings;
+		}
+		else if (StatusA == ENiagaraScriptCompileStatus::NCS_BeingCreated || StatusB == ENiagaraScriptCompileStatus::NCS_BeingCreated)
+		{
+			return ENiagaraScriptCompileStatus::NCS_BeingCreated;
+		}
+		else if (StatusA == ENiagaraScriptCompileStatus::NCS_UpToDate || StatusB == ENiagaraScriptCompileStatus::NCS_UpToDate)
+		{
+			return ENiagaraScriptCompileStatus::NCS_UpToDate;
+		}
+		else
+		{
+			return ENiagaraScriptCompileStatus::NCS_Unknown;
+		}
+	}
+	else
+	{
+		return StatusA;
+	}
+}
+
+bool FNiagaraEditorUtilities::DataMatches(const FNiagaraVariable& Variable, const FStructOnScope& StructOnScope)
+{
+	if (Variable.GetType().GetScriptStruct() != StructOnScope.GetStruct() ||
+		Variable.IsDataAllocated() == false)
+	{
+		return false;
+	}
+
+	const uint8* VariableMemory = Variable.GetData();
+	const uint8* StructOnScopeMemory = StructOnScope.GetStructMemory();
+	int32 Size = Variable.GetSizeInBytes();
+	for (int32 i = 0; i < Size; i++)
+	{
+		if (VariableMemory[i] != StructOnScopeMemory[i])
+		{
+			return false;
+		}
+	}
+
+	return true;
+}
+
+bool FNiagaraEditorUtilities::DataMatches(const FStructOnScope& StructOnScopeA, const FStructOnScope& StructOnScopeB)
+{
+	if (StructOnScopeA.GetStruct() != StructOnScopeB.GetStruct())
+	{
+		return false;
+	}
+
+	const uint8* StructOnScopeAMemory = StructOnScopeA.GetStructMemory();
+	const uint8* StructOnScopeBMemory = StructOnScopeB.GetStructMemory();
+	int32 Size = StructOnScopeA.GetStruct()->GetStructureSize();
+	for (int32 i = 0; i < Size; i++)
+	{
+		if (StructOnScopeAMemory[i] != StructOnScopeBMemory[i])
+		{
+			return false;
+		}
+	}
+
+	return true;
+}
+
+TSharedPtr<SWidget> FNiagaraEditorUtilities::CreateInlineErrorText(TAttribute<FText> ErrorMessage, TAttribute<FText> ErrorTooltip)
+{
+	TSharedPtr<SHorizontalBox> ErrorInternalBox = SNew(SHorizontalBox);
+		ErrorInternalBox->AddSlot()
+			.HAlign(HAlign_Left)
+			.VAlign(VAlign_Center)
+			.AutoWidth()
+			[
+				SNew(STextBlock)
+				.TextStyle(FNiagaraEditorStyle::Get(), "NiagaraEditor.ParameterText")
+				.Text(ErrorMessage)
+			];
+
+		return SNew(SHorizontalBox)
+			.ToolTipText(ErrorTooltip)
+			+ SHorizontalBox::Slot()
+				.AutoWidth()
+				.HAlign(HAlign_Left)
+				.VAlign(VAlign_Center)
+				[
+					SNew(SImage)
+					.Image(FEditorStyle::GetBrush("Icons.Error"))
+				]
+			+ SHorizontalBox::Slot()
+				.AutoWidth()
+				.HAlign(HAlign_Left)
+				.VAlign(VAlign_Center)
+				[
+					ErrorInternalBox.ToSharedRef()
+				];
+}
+
+void FNiagaraEditorUtilities::UpdateExistingEmitters(const TArray<UNiagaraEmitter*>& AffectedEmitters)
+{
+	// Compile the existing emitters. Also determine which Systems need to be properly updated.
+	TArray<UNiagaraSystem*> AffectedSystemSystems;
+	for (UNiagaraEmitter* Emitter : AffectedEmitters)
+	{
+		if (Emitter->IsPendingKillOrUnreachable())
+		{
+			continue;
+		}
+
+		TSharedPtr<FNiagaraEmitterViewModel> EmitterViewModel = FNiagaraEmitterViewModel::GetExistingViewModelForObject(Emitter);
+		if (!EmitterViewModel.IsValid())
+		{
+			EmitterViewModel = MakeShareable(new FNiagaraEmitterViewModel(Emitter, nullptr));
+		}
+		EmitterViewModel->CompileScripts();
+
+		for (TObjectIterator<UNiagaraSystem> It; It; ++It)
+		{
+			if (It->GetAutoImportChangedEmitters() && It->ReferencesSourceEmitter(Emitter))
+			{
+				AffectedSystemSystems.AddUnique(*It);
+			}
+		}
+	}
+
+	// Now iterate over the affected Systems.
+	for (UNiagaraSystem* System : AffectedSystemSystems)
+	{
+		TSharedPtr<FNiagaraSystemViewModel> SystemViewModel = FNiagaraSystemViewModel::GetExistingViewModelForObject(System);
+		if (!SystemViewModel.IsValid())
+		{
+			FNiagaraSystemViewModelOptions Options;
+			Options.bCanRemoveEmittersFromTimeline = false;
+			Options.bCanRenameEmittersFromTimeline = false;
+			Options.bCanAddEmittersFromTimeline = false;
+			Options.bUseSystemExecStateForTimelineReset = false;
+			SystemViewModel = MakeShareable(new FNiagaraSystemViewModel(*System, Options));
+		}
+
+		SystemViewModel->ResynchronizeAllHandles();
+	}
+}
+
+
+#undef LOCTEXT_NAMESPACE
