Copying //UE4/Orion-Staging to //UE4/Main (Source: //Orion/Dev-General @ 3606378)

diff --git a/Engine/Plugins/FX/Niagara/Source/NiagaraEditor/Private/NiagaraEmitterDetailsCustomization.cpp b/Engine/Plugins/FX/Niagara/Source/NiagaraEditor/Private/NiagaraEmitterDetailsCustomization.cpp
new file mode 100644
--- /dev/null
+++ b/Engine/Plugins/FX/Niagara/Source/NiagaraEditor/Private/NiagaraEmitterDetailsCustomization.cpp
@@ -0,0 +1,272 @@
+// Copyright 1998-2017 Epic Games, Inc. All Rights Reserved.
+
+#include "NiagaraEmitterDetailsCustomization.h"
+#include "NiagaraCommon.h"
+#include "NiagaraEmitter.h"
+#include "NiagaraRendererProperties.h"
+#include "NiagaraComponent.h"
+#include "NiagaraEditorStyle.h"
+
+#include "DetailWidgetRow.h"
+#include "PropertyHandle.h"
+#include "DetailLayoutBuilder.h"
+#include "IDetailGroup.h"
+#include "IDetailChildrenBuilder.h"
+#include "IDetailPropertyRow.h"
+#include "DetailCategoryBuilder.h"
+#include "PropertyCustomizationHelpers.h"
+#include "ScopedTransaction.h"
+
+#include "SBorder.h"
+#include "STextBlock.h"
+#include "SButton.h"
+
+
+#define LOCTEXT_NAMESPACE "NiagaraSystemEditor"
+#if 0
+class SMaterialUsageWarning : public SCompoundWidget
+{
+public:
+	SLATE_BEGIN_ARGS(SMaterialUsageWarning)	{}
+	SLATE_END_ARGS();
+
+	void Construct(const FArguments& InArgs, UNiagaraEmitter* InEmitterProperties)
+	{
+		EmitterProperties = InEmitterProperties;
+		ChildSlot
+		[
+			SNew(SBorder)
+			.BorderImage(FNiagaraEditorStyle::Get().GetBrush("NiagaraEditor.MaterialWarningBorder"))
+			.BorderBackgroundColor(FLinearColor(.6f, 0.0f, 0.0f, 1.0f))
+			[
+				SNew(SHorizontalBox)
+				+ SHorizontalBox::Slot()
+				.VAlign(VAlign_Center)
+				.Padding(3, 0, 0, 0)
+				[
+					SNew(STextBlock)
+					.Text(this, &SMaterialUsageWarning::GetInvalidMessage)
+					.AutoWrapText(true)
+				]
+				+ SHorizontalBox::Slot()
+				.VAlign(VAlign_Top)
+				.Padding(5, 0, 0, 0)
+				.AutoWidth()
+				[
+					SNew(SButton)
+					.OnClicked(this, &SMaterialUsageWarning::OnFixNowClicked)
+					.ToolTipText(NSLOCTEXT("MaterialUsageWarning", "FixButtonLabelToolTip", "Fix the material issue now by modifying the material asset."))
+					.Content()
+					[
+						SNew(STextBlock)
+						.Text(NSLOCTEXT("MaterialUsageWarning", "FixButtonLabel", "Fix Now"))
+					]
+				]
+			]
+		];
+	}
+
+	EVisibility GetVisibility() const
+	{
+		bool bMaterialIsValid =
+			EmitterProperties->RendererProperties == nullptr ||
+			EmitterProperties->Material == nullptr ||
+			EmitterProperties->RendererProperties->IsMaterialValidForRenderer(EmitterProperties->Material, InvalidMessage);
+
+		return bMaterialIsValid
+			? EVisibility::Collapsed
+			: EVisibility::Visible;
+	}
+
+private:
+	FText GetInvalidMessage() const
+	{
+		return InvalidMessage;
+	}
+
+	FReply OnFixNowClicked()
+	{
+		FScopedTransaction ScopedTransaction(NSLOCTEXT("EmitterDetailsCustomization", "FixMaterialTransaction", "Fix material for emitter"));
+		EmitterProperties->RendererProperties->FixMaterial(EmitterProperties->Material);
+		return FReply::Handled();
+	}
+
+private:
+	UNiagaraEmitter* EmitterProperties;
+	mutable FText InvalidMessage;
+};
+#endif
+
+TSharedRef<IDetailCustomization> FNiagaraEmitterDetails::MakeInstance(TWeakObjectPtr<UNiagaraEmitter> EmitterProperties)
+{
+	return MakeShareable(new FNiagaraEmitterDetails(EmitterProperties));
+}
+
+FNiagaraEmitterDetails::FNiagaraEmitterDetails(TWeakObjectPtr<UNiagaraEmitter> EmitterProperties)
+: EmitterProps(EmitterProperties)
+{
+}
+
+void FNiagaraEmitterDetails::CustomizeDetails(IDetailLayoutBuilder& InDetailLayout)
+{
+	DetailLayout = &InDetailLayout;
+	FName EmitterCategoryName = TEXT("Emitter");
+	IDetailCategoryBuilder& EmitterCategory = DetailLayout->EditCategory(EmitterCategoryName);
+	
+	TArray<TWeakObjectPtr<UObject>> ObjectsBeingCustomized;
+	InDetailLayout.GetObjectsBeingCustomized(ObjectsBeingCustomized);
+	UNiagaraEmitter* EmitterProperties = CastChecked<UNiagaraEmitter>(ObjectsBeingCustomized[0].Get());
+
+#if 0
+	IDetailCategoryBuilder& RenderCategory = DetailLayout->EditCategory(TEXT("Render"));
+	TSharedRef<SMaterialUsageWarning> MaterialWarning = SNew(SMaterialUsageWarning, CastChecked<UNiagaraEmitter>(EmitterProperties));
+	RenderCategory.AddCustomRow(FText())
+		.Visibility(TAttribute<EVisibility>::Create(TAttribute<EVisibility>::FGetter::CreateSP(MaterialWarning, &SMaterialUsageWarning::GetVisibility)))
+		[
+			MaterialWarning
+		];
+#endif
+	/*
+	TSharedRef<IPropertyHandle> UpdateScriptProps = DetailLayout->GetProperty(GET_MEMBER_NAME_CHECKED(UNiagaraEmitter, UpdateScriptProps));
+	TSharedRef<IPropertyHandle> SpawnScriptProps = DetailLayout->GetProperty(GET_MEMBER_NAME_CHECKED(UNiagaraEmitter, SpawnScriptProps));
+
+	BuildScriptProperties(SpawnScriptProps, TEXT("Spawn Script"), LOCTEXT("SpawnScriptDisplayName", "Spawn Script"));
+	BuildScriptProperties(UpdateScriptProps, TEXT("Update Script"), LOCTEXT("UpdateScriptDisplayName", "Update Script"));
+	*/
+}
+
+void FNiagaraEmitterDetails::BuildScriptProperties(TSharedRef<IPropertyHandle> ScriptPropsHandle, FName Name, FText DisplayName)
+{
+//	/*
+//	DetailLayout->HideProperty(ScriptPropsHandle);
+//
+//	IDetailCategoryBuilder& ScriptCategory = DetailLayout->EditCategory(Name, DisplayName);
+//
+//	//Script
+//	TSharedPtr<IPropertyHandle> ScriptHandle = ScriptPropsHandle->GetChildHandle(GET_MEMBER_NAME_CHECKED(FNiagaraEmitterScriptProperties, Script));
+//	ScriptCategory.AddProperty(ScriptHandle);
+//
+//	//Constants
+//	/*
+//	bool bGenerateHeader = false;
+//	bool bDisplayResetToDefault = false;
+//	TSharedPtr<IPropertyHandle> Constants = ScriptPropsHandle->GetChildHandle(GET_MEMBER_NAME_CHECKED(FNiagaraEmitterScriptProperties, ExternalConstants));
+//	DetailLayout->HideProperty(Constants);
+//
+//	TSharedPtr<IPropertyHandle> Uniforms = Constants->GetChildHandle(GET_MEMBER_NAME_CHECKED(FNiagaraConstants, Uniforms));
+//	TSharedRef<FDetailArrayBuilder> ConstantsBuilder = MakeShareable(new FDetailArrayBuilder(Uniforms.ToSharedRef(), bGenerateHeader, bDisplayResetToDefault));
+//	ConstantsBuilder->OnGenerateArrayElementWidget(FOnGenerateArrayElementWidget::CreateSP(this, &FNiagaraEmitterDetails::OnGenerateConstantEntry));
+//	ScriptCategory.AddCustomBuilder(ConstantsBuilder);
+//	*/
+//
+//	//Scalar Constants.
+//	/*
+//	TSharedPtr<IPropertyHandle> ScalarConstants = Constants->GetChildHandle(GET_MEMBER_NAME_CHECKED(FNiagaraConstants, ScalarConstants));
+//	TSharedRef<FDetailArrayBuilder> ScalarConstantsBuilder = MakeShareable(new FDetailArrayBuilder(ScalarConstants.ToSharedRef(), bGenerateHeader, bDisplayResetToDefault));
+//	ScalarConstantsBuilder->OnGenerateArrayElementWidget(FOnGenerateArrayElementWidget::CreateSP(this, &FNiagaraEmitterDetails::OnGenerateScalarConstantEntry));
+//	ScriptCategory.AddCustomBuilder(ScalarConstantsBuilder);
+//	//Vector Constants.
+//	TSharedPtr<IPropertyHandle> VectorConstants = Constants->GetChildHandle(GET_MEMBER_NAME_CHECKED(FNiagaraConstants, VectorConstants));
+//	TSharedRef<FDetailArrayBuilder> VectorConstantsBuilder = MakeShareable(new FDetailArrayBuilder(VectorConstants.ToSharedRef(), bGenerateHeader, bDisplayResetToDefault));
+//	VectorConstantsBuilder->OnGenerateArrayElementWidget(FOnGenerateArrayElementWidget::CreateSP(this, &FNiagaraEmitterDetails::OnGenerateVectorConstantEntry));
+//	ScriptCategory.AddCustomBuilder(VectorConstantsBuilder);
+//	*/
+//	//Events
+//	//bGenerateHeader = true;
+//	//Generators
+//// 	TSharedPtr<IPropertyHandle> EventGenerators = ScriptPropsHandle->GetChildHandle(GET_MEMBER_NAME_CHECKED(FNiagaraEmitterScriptProperties, EventGenerators));
+//// 	TSharedRef<FDetailArrayBuilder> EventGeneratorsBuilder = MakeShareable(new FDetailArrayBuilder(EventGenerators.ToSharedRef(), bGenerateHeader, bDisplayResetToDefault));
+//// 	EventGeneratorsBuilder->OnGenerateArrayElementWidget(FOnGenerateArrayElementWidget::CreateSP(this, &FNiagaraEmitterDetails::OnGenerateEventGeneratorEntry));
+//// 	ScriptCategory.AddCustomBuilder(EventGeneratorsBuilder);
+//// 	//Receivers
+//// 	TSharedPtr<IPropertyHandle> EventReceivers = ScriptPropsHandle->GetChildHandle(GET_MEMBER_NAME_CHECKED(FNiagaraEmitterScriptProperties, EventReceivers));
+//// 	TSharedRef<FDetailArrayBuilder> EventReceiversBuilder = MakeShareable(new FDetailArrayBuilder(EventReceivers.ToSharedRef(), bGenerateHeader, bDisplayResetToDefault));
+//// 	EventReceiversBuilder->OnGenerateArrayElementWidget(FOnGenerateArrayElementWidget::CreateSP(this, &FNiagaraEmitterDetails::OnGenerateEventReceiverEntry));
+//// 	ScriptCategory.AddCustomBuilder(EventReceiversBuilder);
+//	
+}
+
+
+
+void FNiagaraEmitterDetails::OnGenerateConstantEntry(TSharedRef<IPropertyHandle> ElementProperty, int32 ElementIndex, IDetailChildrenBuilder& ChildrenBuilder)
+{
+// 	TSharedPtr<IPropertyHandle> ValueProperty = ElementProperty->GetChildHandle(GET_MEMBER_NAME_CHECKED(FNiagaraVariable, VarData));
+// 	TSharedPtr<IPropertyHandle> NameProperty = ElementProperty->GetChildHandle(GET_MEMBER_NAME_CHECKED(FNiagaraVariable, Name));
+ 	
+// 	FName DisplayName;
+// 	NameProperty->GetValue(DisplayName);
+
+	//STOVEY. Not sure this is needed. System constants don't show up in the param UI anyway.
+	//Don't display system constants
+// 	if (UNiagaraComponent::FindSystemConstant(Var) == nullptr)
+// 	{
+//		ChildrenBuilder.AddChildProperty(ValueProperty.ToSharedRef()).DisplayName(FText::FromName(DisplayName));
+//	}
+}
+
+
+void FNiagaraEmitterDetails::OnGenerateEventGeneratorEntry(TSharedRef<IPropertyHandle> ElementProperty, int32 ElementIndex, IDetailChildrenBuilder& ChildrenBuilder)
+{
+	TSharedPtr<IPropertyHandle> IDProperty = ElementProperty->GetChildHandle(GET_MEMBER_NAME_CHECKED(FNiagaraEventGeneratorProperties, ID));
+	TSharedPtr<IPropertyHandle> NameProperty = IDProperty->GetChildHandle(GET_MEMBER_NAME_CHECKED(FNiagaraDataSetID, Name));
+
+	FName DisplayName;
+	NameProperty->GetValue(DisplayName);
+	//IDetailPropertyRow& Row = ChildrenBuilder.AddChildProperty(ElementProperty).DisplayName(FText::FromName(DisplayName));
+
+	IDetailGroup& GenGroup = ChildrenBuilder.AddGroup(DisplayName, FText::FromName(DisplayName));
+	uint32 NumChildren = 0;
+	if (ElementProperty->GetNumChildren(NumChildren) == FPropertyAccess::Success)
+	{
+		for (uint32 i = 0; i < NumChildren; ++i)
+		{
+			TSharedPtr<IPropertyHandle> Child = ElementProperty->GetChildHandle(i);
+			//Dont add the ID. We just grab it's name for the name region of this property.
+			if (Child.IsValid() && Child->GetProperty()->GetName() != GET_MEMBER_NAME_CHECKED(FNiagaraEventGeneratorProperties, ID).ToString())
+			{
+				GenGroup.AddPropertyRow(Child.ToSharedRef());
+			}
+		}
+	}
+}
+
+void FNiagaraEmitterDetails::OnGenerateEventReceiverEntry(TSharedRef<IPropertyHandle> ElementProperty, int32 ElementIndex, IDetailChildrenBuilder& ChildrenBuilder)
+{
+	TSharedPtr<IPropertyHandle> NameProperty = ElementProperty->GetChildHandle(GET_MEMBER_NAME_CHECKED(FNiagaraEventReceiverProperties, Name));
+
+ 	FName DisplayName;
+ 	NameProperty->GetValue(DisplayName);
+// 	ChildrenBuilder.AddChildProperty(ElementProperty).DisplayName(FText::FromName(DisplayName));
+
+	IDetailGroup& Group = ChildrenBuilder.AddGroup(DisplayName, FText::FromName(DisplayName));
+	uint32 NumChildren = 0;
+	if (ElementProperty->GetNumChildren(NumChildren) == FPropertyAccess::Success)
+	{
+		for (uint32 i = 0; i < NumChildren; ++i)
+		{
+			TSharedPtr<IPropertyHandle> Child = ElementProperty->GetChildHandle(i);
+			//Dont add the ID. We just grab it's name for the name region of this property.
+			if (Child.IsValid() && Child->GetProperty()->GetName() != GET_MEMBER_NAME_CHECKED(FNiagaraEventReceiverProperties, Name).ToString())
+			{
+				TSharedPtr<SWidget> NameWidget;
+				TSharedPtr<SWidget> ValueWidget;
+				FDetailWidgetRow DefaultDetailRow;
+			
+				IDetailPropertyRow& Row = Group.AddPropertyRow(Child.ToSharedRef());
+				Row.GetDefaultWidgets(NameWidget, ValueWidget, DefaultDetailRow);
+				Row.CustomWidget(true)
+					.NameContent()
+					[
+						NameWidget.ToSharedRef()
+					]
+					.ValueContent()
+					[
+						ValueWidget.ToSharedRef()
+					];
+			}
+		}
+	}
+}
+
+#undef LOCTEXT_NAMESPACE
+
diff --git a/Engine/Plugins/FX/Niagara/Source/NiagaraEditor/Private/NiagaraEmitterDetailsCustomization.h b/Engine/Plugins/FX/Niagara/Source/NiagaraEditor/Private/NiagaraEmitterDetailsCustomization.h
new file mode 100644
--- /dev/null
+++ b/Engine/Plugins/FX/Niagara/Source/NiagaraEditor/Private/NiagaraEmitterDetailsCustomization.h
@@ -0,0 +1,44 @@
+// Copyright 1998-2017 Epic Games, Inc. All Rights Reserved.
+
+#pragma once
+
+#include "CoreMinimal.h"
+#include "UObject/WeakObjectPtr.h"
+#include "IDetailCustomization.h"
+
+class IDetailChildrenBuilder;
+class IDetailLayoutBuilder;
+class IPropertyHandle;
+class UNiagaraEmitter;
+
+/*-----------------------------------------------------------------------------
+FNiagaraEmitterDetails
+-----------------------------------------------------------------------------*/
+
+class FNiagaraEmitterDetails : public IDetailCustomization
+{
+public:
+	/** Makes a new instance of this detail layout class for a specific detail view requesting it */
+	static TSharedRef<class IDetailCustomization> MakeInstance(TWeakObjectPtr<UNiagaraEmitter> EmitterProperties);
+
+	/** Constructor */
+	FNiagaraEmitterDetails(TWeakObjectPtr<UNiagaraEmitter> EmitterProperties);
+
+	/** IDetailCustomization interface */
+	virtual void CustomizeDetails(IDetailLayoutBuilder& InDetailLayout) override;
+
+	void OnGenerateConstantEntry(TSharedRef<IPropertyHandle> ElementProperty, int32 ElementIndex, IDetailChildrenBuilder& ChildrenBuilder);
+
+	void OnGenerateScalarConstantEntry(TSharedRef<IPropertyHandle> ElementProperty, int32 ElementIndex, IDetailChildrenBuilder& ChildrenBuilder);
+	void OnGenerateVectorConstantEntry(TSharedRef<IPropertyHandle> ElementProperty, int32 ElementIndex, IDetailChildrenBuilder& ChildrenBuilder);
+	void OnGenerateEventGeneratorEntry(TSharedRef<IPropertyHandle> ElementProperty, int32 ElementIndex, IDetailChildrenBuilder& ChildrenBuilder);
+	void OnGenerateEventReceiverEntry(TSharedRef<IPropertyHandle> ElementProperty, int32 ElementIndex, IDetailChildrenBuilder& ChildrenBuilder);
+
+	void BuildScriptProperties(TSharedRef<IPropertyHandle> ScriptPropsHandle, FName Name, FText DisplayName);
+private:
+	/** Object that stores all of the possible parameters we can edit */
+	TWeakObjectPtr<UNiagaraEmitter> EmitterProps;
+
+	IDetailLayoutBuilder* DetailLayout;
+};
+
diff --git a/Engine/Plugins/FX/Niagara/Source/NiagaraEditor/Private/NiagaraEmitterEditorData.h b/Engine/Plugins/FX/Niagara/Source/NiagaraEditor/Private/NiagaraEmitterEditorData.h
new file mode 100644
--- /dev/null
+++ b/Engine/Plugins/FX/Niagara/Source/NiagaraEditor/Private/NiagaraEmitterEditorData.h
@@ -0,0 +1,25 @@
+// Copyright 1998-2017 Epic Games, Inc. All Rights Reserved.
+
+#pragma once
+
+#include "NiagaraEmitterEditorData.generated.h"
+
+class UNiagaraStackEditorData;
+
+/** Editor only UI data for emitters. */
+UCLASS()
+class UNiagaraEmitterEditorData : public UObject
+{
+	GENERATED_BODY()
+
+public:
+	UNiagaraEmitterEditorData(const FObjectInitializer& ObjectInitializer);
+
+	virtual void PostLoad() override;
+
+	UNiagaraStackEditorData& GetStackEditorData() const;
+
+private:
+	UPROPERTY(Instanced)
+	UNiagaraStackEditorData* StackEditorData;
+};
\ No newline at end of file
diff --git a/Engine/Plugins/FX/Niagara/Source/NiagaraEditor/Private/NiagaraEmitterFactoryNew.cpp b/Engine/Plugins/FX/Niagara/Source/NiagaraEditor/Private/NiagaraEmitterFactoryNew.cpp
new file mode 100644
--- /dev/null
+++ b/Engine/Plugins/FX/Niagara/Source/NiagaraEditor/Private/NiagaraEmitterFactoryNew.cpp
@@ -0,0 +1,58 @@
+// Copyright 1998-2016 Epic Games, Inc. All Rights Reserved.
+
+#include "NiagaraEmitterFactoryNew.h"
+#include "NiagaraEmitter.h"
+#include "NiagaraEditorModule.h"
+#include "NiagaraScriptFactoryNew.h"
+#include "NiagaraGraph.h"
+#include "NiagaraScriptSource.h"
+#include "NiagaraNodeOutput.h"
+#include "NiagaraEditorSettings.h"
+
+#include "ConfigCacheIni.h"
+
+#define LOCTEXT_NAMESPACE "NiagaraEmitterFactory"
+
+UNiagaraEmitterFactoryNew::UNiagaraEmitterFactoryNew(const FObjectInitializer& ObjectInitializer)
+: Super(ObjectInitializer)
+{
+	SupportedClass = UNiagaraEmitter::StaticClass();
+	bCreateNew = false;
+	bEditAfterNew = true;
+	bCreateNew = true;
+}
+
+UObject* UNiagaraEmitterFactoryNew::FactoryCreateNew(UClass* Class, UObject* InParent, FName Name, EObjectFlags Flags, UObject* Context, FFeedbackContext* Warn)
+{
+	check(Class->IsChildOf(UNiagaraEmitter::StaticClass()));
+
+	const UNiagaraEditorSettings* Settings = GetDefault<UNiagaraEditorSettings>();
+	check(Settings);
+
+	UNiagaraEmitter* NewEmitter;
+
+	if (UNiagaraEmitter* Default = Cast<UNiagaraEmitter>(Settings->DefaultEmitter.TryLoad()))
+	{
+		NewEmitter = Cast<UNiagaraEmitter>(StaticDuplicateObject(Default, InParent, Name, Flags, Class));
+	}
+	else
+	{
+		NewEmitter = NewObject<UNiagaraEmitter>(InParent, Class, Name, Flags | RF_Transactional);
+
+		UNiagaraScriptSource* Source = NewObject<UNiagaraScriptSource>(NewEmitter, NAME_None, RF_Transactional);
+		if (Source)
+		{
+			UNiagaraGraph* CreatedGraph = NewObject<UNiagaraGraph>(Source, NAME_None, RF_Transactional);
+			Source->NodeGraph = CreatedGraph;
+
+			// Set pointer in script to source
+			NewEmitter->GraphSource = Source;
+			NewEmitter->SpawnScriptProps.Script->SetSource(Source);
+			NewEmitter->UpdateScriptProps.Script->SetSource(Source);
+		}
+	}
+	
+	return NewEmitter;
+}
+
+#undef LOCTEXT_NAMESPACE
\ No newline at end of file
diff --git a/Engine/Plugins/FX/Niagara/Source/NiagaraEditor/Private/NiagaraGraph.cpp b/Engine/Plugins/FX/Niagara/Source/NiagaraEditor/Private/NiagaraGraph.cpp
new file mode 100644
--- /dev/null
+++ b/Engine/Plugins/FX/Niagara/Source/NiagaraEditor/Private/NiagaraGraph.cpp
@@ -0,0 +1,500 @@
+// Copyright 1998-2017 Epic Games, Inc. All Rights Reserved.
+
+#include "NiagaraGraph.h"
+#include "Modules/ModuleManager.h"
+#include "NiagaraCommon.h"
+#include "NiagaraEditorModule.h"
+#include "NiagaraScript.h"
+#include "NiagaraComponent.h"
+#include "UObject/UObjectHash.h"
+#include "UObject/UObjectIterator.h"
+#include "ComponentReregisterContext.h"
+#include "NiagaraConstants.h"
+#include "NiagaraSystem.h"
+#include "NiagaraNodeOutput.h"
+#include "NiagaraNodeInput.h"
+#include "NiagaraNodeWriteDataSet.h"
+#include "NiagaraNodeReadDataSet.h"
+#include "NiagaraScript.h"
+#include "NiagaraScriptSource.h"
+#include "NiagaraDataInterface.h"
+#include "GraphEditAction.h"
+#include "EdGraphSchema_Niagara.h"
+#include "NiagaraNodeParameterMapBase.h"
+#include "NiagaraNodeParameterMapGet.h"
+#include "INiagaraEditorTypeUtilities.h"
+#include "NiagaraConstants.h"
+
+
+UNiagaraGraph::UNiagaraGraph(const FObjectInitializer& ObjectInitializer)
+	: Super(ObjectInitializer)
+{
+	Schema = UEdGraphSchema_Niagara::StaticClass();
+	ChangeId = FGuid::NewGuid();
+}
+
+void UNiagaraGraph::PostLoad()
+{
+	Super::PostLoad();
+
+	// In the past, we didn't bother setting the CallSortPriority and just used lexicographic ordering.
+	// In the event that we have multiple non-matching nodes with a zero call sort priority, this will
+	// give every node a unique order value.
+	TArray<UNiagaraNodeInput*> InputNodes;
+	GetNodesOfClass(InputNodes);
+	bool bAllZeroes = true;
+	TArray<FName> UniqueNames;
+	for (UNiagaraNodeInput* InputNode : InputNodes)
+	{
+		if (InputNode->CallSortPriority != 0)
+		{
+			bAllZeroes = false;
+		}
+
+		if (InputNode->Usage == ENiagaraInputNodeUsage::Parameter)
+		{
+			UniqueNames.AddUnique(InputNode->Input.GetName());
+		}
+
+		if (InputNode->Usage == ENiagaraInputNodeUsage::SystemConstant)
+		{
+			InputNode->Input = FNiagaraConstants::UpdateEngineConstant(InputNode->Input);
+		}
+	}
+
+
+
+	if (bAllZeroes && UniqueNames.Num() > 1)
+	{
+		// Just do the lexicographic sort and assign the call order to their ordered index value.
+		UniqueNames.Sort();
+		for (UNiagaraNodeInput* InputNode : InputNodes)
+		{
+			if (InputNode->Usage == ENiagaraInputNodeUsage::Parameter)
+			{
+				int32 FoundIndex = UniqueNames.Find(InputNode->Input.GetName());
+				check(FoundIndex != -1);
+				InputNode->CallSortPriority = FoundIndex;
+			}
+		}
+	}
+
+	// If this is from a prior version, enforce a valid Change Id!
+	if (ChangeId.IsValid() == false)
+	{
+		MarkGraphRequiresSynchronization();
+	}
+
+	// Assume that all externally referenced assets have changed, so update to match. They will return true if they have changed.
+	TArray<UNiagaraNode*> NiagaraNodes;
+	GetNodesOfClass<UNiagaraNode>(NiagaraNodes);
+	bool bAnyExternalChanges = false;
+	for (UNiagaraNode* NiagaraNode : NiagaraNodes)
+	{
+		UObject* ReferencedAsset = NiagaraNode->GetReferencedAsset();
+		if (ReferencedAsset != nullptr)
+		{
+			ReferencedAsset->ConditionalPostLoad();
+			NiagaraNode->ConditionalPostLoad();
+			bAnyExternalChanges |= (NiagaraNode->RefreshFromExternalChanges());
+		}
+		else
+		{
+			NiagaraNode->ConditionalPostLoad();
+		}
+	}
+
+	if (bAnyExternalChanges)
+	{
+		MarkGraphRequiresSynchronization();
+		NotifyGraphNeedsRecompile();
+	}
+
+	if (GIsEditor)
+	{
+		SetFlags(RF_Transactional);
+	}
+
+}
+
+void UNiagaraGraph::PostEditChangeProperty(FPropertyChangedEvent& PropertyChangedEvent)
+{
+	NotifyGraphChanged();
+}
+
+class UNiagaraScriptSource* UNiagaraGraph::GetSource() const
+{
+	return CastChecked<UNiagaraScriptSource>(GetOuter());
+}
+
+
+UEdGraphPin* UNiagaraGraph::FindParameterMapDefaultValuePin(const FString& VariableName)
+{
+	TArray<UEdGraphPin*> MatchingDefaultPins;
+	TArray<UNiagaraNodeParameterMapGet*> GetNodes;
+	GetNodesOfClass<UNiagaraNodeParameterMapGet>(GetNodes);
+
+	for (UNiagaraNodeParameterMapGet* GetNode : GetNodes)
+	{
+		TArray<UEdGraphPin*> OutputPins;
+		GetNode->GetOutputPins(OutputPins);
+		for (UEdGraphPin* OutputPin : OutputPins)
+		{
+			if (VariableName == OutputPin->PinName)
+			{
+				UEdGraphPin* Pin = GetNode->GetDefaultPin(OutputPin);
+				if (Pin)
+				{
+					MatchingDefaultPins.Add(Pin);
+				}
+			}
+		}
+	}
+
+	if (MatchingDefaultPins.Num() > 0)
+	{
+		return MatchingDefaultPins[0];
+	}
+
+	return nullptr;
+}
+
+void UNiagaraGraph::FindOutputNodes(TArray<UNiagaraNodeOutput*>& OutputNodes) const
+{
+	for (UEdGraphNode* Node : Nodes)
+	{
+		if (UNiagaraNodeOutput* OutNode = Cast<UNiagaraNodeOutput>(Node))
+		{
+			OutputNodes.Add(OutNode);
+		}
+	}
+}
+
+
+void UNiagaraGraph::FindOutputNodes(ENiagaraScriptUsage TargetUsageType, TArray<UNiagaraNodeOutput*>& OutputNodes) const
+{
+	TArray<UNiagaraNodeOutput*> NodesFound;
+	for (UEdGraphNode* Node : Nodes)
+	{
+		UNiagaraNodeOutput* OutNode = Cast<UNiagaraNodeOutput>(Node);
+		if (OutNode && OutNode->GetUsage() == TargetUsageType)
+		{
+			NodesFound.Add(OutNode);
+		}
+	}
+
+	NodesFound.Sort([](const UNiagaraNodeOutput& LHS, const UNiagaraNodeOutput& RHS) { return LHS.GetUsageIndex() > RHS.GetUsageIndex(); });
+	OutputNodes = NodesFound;
+}
+
+UNiagaraNodeOutput* UNiagaraGraph::FindOutputNode(ENiagaraScriptUsage TargetUsageType, int32 TargetOccurence) const
+{
+	for (UEdGraphNode* Node : Nodes)
+	{
+		if (UNiagaraNodeOutput* OutNode = Cast<UNiagaraNodeOutput>(Node))
+		{
+			if (OutNode->GetUsage() == TargetUsageType)
+			{
+				if (OutNode->ScriptTypeIndex == TargetOccurence)
+				{
+					return OutNode;
+				}
+			}
+		}
+	}
+	return nullptr;
+}
+
+void BuildTraversalHelper(TArray<class UNiagaraNode*>& OutNodesTraversed, UNiagaraNode* CurrentNode)
+{
+	if (CurrentNode == nullptr)
+	{
+		return;
+	}
+
+	TArray<UEdGraphPin*> Pins = CurrentNode->GetAllPins();
+	for (int32 i = 0; i < Pins.Num(); i++)
+	{
+		if (Pins[i]->Direction == EEdGraphPinDirection::EGPD_Input && Pins[i]->LinkedTo.Num() == 1)
+		{
+			BuildTraversalHelper(OutNodesTraversed, Cast<UNiagaraNode>(Pins[i]->LinkedTo[0]->GetOwningNode()));
+		}
+	}
+
+	OutNodesTraversed.Add(CurrentNode);
+}
+
+void UNiagaraGraph::BuildTraversal(TArray<class UNiagaraNode*>& OutNodesTraversed, ENiagaraScriptUsage TargetUsage, int32 TargetOccurence) const
+{
+	UNiagaraNodeOutput* Output = FindOutputNode(TargetUsage, TargetOccurence);
+	if (Output)
+	{
+		BuildTraversalHelper(OutNodesTraversed, Output);
+	}
+}
+
+void UNiagaraGraph::BuildTraversal(TArray<class UNiagaraNode*>& OutNodesTraversed, UNiagaraNode* FinalNode) const
+{
+	if (FinalNode)
+	{
+		BuildTraversalHelper(OutNodesTraversed, FinalNode);
+	}
+}
+
+
+void UNiagaraGraph::FindInputNodes(TArray<UNiagaraNodeInput*>& OutInputNodes, UNiagaraGraph::FFindInputNodeOptions Options) const
+{
+	TArray<UNiagaraNodeInput*> InputNodes;
+
+	if (!Options.bFilterByScriptUsage)
+	{
+		for (UEdGraphNode* Node : Nodes)
+		{
+			UNiagaraNodeInput* NiagaraInputNode = Cast<UNiagaraNodeInput>(Node);
+			if (NiagaraInputNode != nullptr &&
+				((NiagaraInputNode->Usage == ENiagaraInputNodeUsage::Parameter && Options.bIncludeParameters) ||
+					(NiagaraInputNode->Usage == ENiagaraInputNodeUsage::Attribute && Options.bIncludeAttributes) ||
+					(NiagaraInputNode->Usage == ENiagaraInputNodeUsage::SystemConstant && Options.bIncludeSystemConstants)))
+			{
+				InputNodes.Add(NiagaraInputNode);
+			}
+		}
+	}
+	else
+	{
+		TArray<class UNiagaraNode*> Traversal;
+		BuildTraversal(Traversal, Options.TargetScriptUsage, Options.TargetOccurence);
+		for (UNiagaraNode* Node : Traversal)
+		{
+			UNiagaraNodeInput* NiagaraInputNode = Cast<UNiagaraNodeInput>(Node);
+			if (NiagaraInputNode != nullptr &&
+				((NiagaraInputNode->Usage == ENiagaraInputNodeUsage::Parameter && Options.bIncludeParameters) ||
+					(NiagaraInputNode->Usage == ENiagaraInputNodeUsage::Attribute && Options.bIncludeAttributes) ||
+					(NiagaraInputNode->Usage == ENiagaraInputNodeUsage::SystemConstant && Options.bIncludeSystemConstants)))
+			{
+				InputNodes.Add(NiagaraInputNode);
+			}
+		}
+	}
+
+	if (Options.bFilterDuplicates)
+	{
+		for (UNiagaraNodeInput* InputNode : InputNodes)
+		{
+			auto NodeMatches = [=](UNiagaraNodeInput* UniqueInputNode)
+			{
+				if (InputNode->Usage == ENiagaraInputNodeUsage::Parameter)
+				{
+					return UniqueInputNode->Input.IsEquivalent(InputNode->Input, false);
+				}
+				else
+				{
+					return UniqueInputNode->Input.IsEquivalent(InputNode->Input);
+				}
+			};
+
+			if (OutInputNodes.ContainsByPredicate(NodeMatches) == false)
+			{
+				OutInputNodes.Add(InputNode);
+			}
+		}
+	}
+	else
+	{
+		OutInputNodes.Append(InputNodes);
+	}
+
+	if (Options.bSort)
+	{
+		UNiagaraNodeInput::SortNodes(OutInputNodes);
+	}
+}
+
+void UNiagaraGraph::GetParameters(TArray<FNiagaraVariable>& Inputs, TArray<FNiagaraVariable>& Outputs)const
+{
+	Inputs.Empty();
+	Outputs.Empty();
+
+	TArray<UNiagaraNodeInput*> InputsNodes;
+	FFindInputNodeOptions Options;
+	Options.bSort = true;
+	FindInputNodes(InputsNodes, Options);
+	for (UNiagaraNodeInput* Input : InputsNodes)
+	{
+		Inputs.Add(Input->Input);
+	}
+
+	TArray<UNiagaraNodeOutput*> OutputNodes;
+	FindOutputNodes(OutputNodes);
+	for (UNiagaraNodeOutput* OutputNode : OutputNodes)
+	{
+		for (FNiagaraVariable& Var : OutputNode->Outputs)
+		{
+			Outputs.AddUnique(Var);
+		}
+	}
+
+	//Do we need to sort outputs?
+	//Should leave them as they're defined in the output node?
+// 	auto SortVars = [](const FNiagaraVariable& A, const FNiagaraVariable& B)
+// 	{
+// 		//Case insensitive lexicographical comparisons of names for sorting.
+// 		return A.GetName().ToString() < B.GetName().ToString();
+// 	};
+// 	Outputs.Sort(SortVars);
+}
+
+void UNiagaraGraph::FindReadDataSetNodes(TArray<class UNiagaraNodeReadDataSet*>& OutReadNodes) const
+{
+	for (UEdGraphNode* Node : Nodes)
+	{
+		if (UNiagaraNodeReadDataSet* InNode = Cast<UNiagaraNodeReadDataSet>(Node))
+		{
+			OutReadNodes.Add(InNode);
+		}
+	}
+}
+
+void UNiagaraGraph::FindWriteDataSetNodes(TArray<class UNiagaraNodeWriteDataSet*>& OutWriteNodes) const
+{
+	for (UEdGraphNode* Node : Nodes)
+	{
+		if (UNiagaraNodeWriteDataSet* InNode = Cast<UNiagaraNodeWriteDataSet>(Node))
+		{
+			OutWriteNodes.Add(InNode);
+		}
+	}
+}
+
+int32 UNiagaraGraph::GetOutputNodeVariableIndex(const FNiagaraVariable& Variable)const
+{
+	TArray<FNiagaraVariable> Variables;
+	GetOutputNodeVariables(Variables);
+	return Variables.Find(Variable);
+}
+
+void UNiagaraGraph::GetOutputNodeVariables(TArray< FNiagaraVariable >& OutVariables)const
+{
+	TArray<UNiagaraNodeOutput*> OutputNodes;
+	FindOutputNodes(OutputNodes);
+	for (UNiagaraNodeOutput* OutputNode : OutputNodes)
+	{
+		for (FNiagaraVariable& Var : OutputNode->Outputs)
+		{
+			OutVariables.AddUnique(Var);
+		}
+	}
+}
+
+void UNiagaraGraph::GetOutputNodeVariables(ENiagaraScriptUsage InScriptUsage, TArray< FNiagaraVariable >& OutVariables)const
+{
+	TArray<UNiagaraNodeOutput*> OutputNodes;
+	FindOutputNodes(InScriptUsage, OutputNodes);
+	for (UNiagaraNodeOutput* OutputNode : OutputNodes)
+	{
+		for (FNiagaraVariable& Var : OutputNode->Outputs)
+		{
+			OutVariables.AddUnique(Var);
+		}
+	}
+}
+
+bool UNiagaraGraph::HasNumericParameters()const
+{
+	TArray<FNiagaraVariable> Inputs;
+	TArray<FNiagaraVariable> Outputs;
+	
+	GetParameters(Inputs, Outputs);
+	
+	for (FNiagaraVariable& Var : Inputs)
+	{
+		if (Var.GetType() == FNiagaraTypeDefinition::GetGenericNumericDef())
+		{
+			return true;
+		}
+	}
+	for (FNiagaraVariable& Var : Outputs)
+	{
+		if (Var.GetType() == FNiagaraTypeDefinition::GetGenericNumericDef())
+		{
+			return true;
+		}
+	}
+
+	return false;
+}
+
+void UNiagaraGraph::NotifyGraphNeedsRecompile()
+{
+	FEdGraphEditAction Action;
+	Action.Action = (EEdGraphActionType)GRAPHACTION_GenericNeedsRecompile;
+	NotifyGraphChanged(Action);
+}
+
+void UNiagaraGraph::SubsumeExternalDependencies(TMap<const UObject*, UObject*>& ExistingConversions)
+{
+	TArray<UNiagaraNode*> NiagaraNodes;
+	GetNodesOfClass<UNiagaraNode>(NiagaraNodes);
+	for (UNiagaraNode* NiagaraNode : NiagaraNodes)
+	{
+		NiagaraNode->SubsumeExternalDependencies(ExistingConversions);
+	}
+}
+
+void UNiagaraGraph::GetAllReferencedGraphs(TArray<const UNiagaraGraph*>& Graphs) const
+{
+	Graphs.AddUnique(this);
+	for (UEdGraphNode* Node : Nodes)
+	{
+		if (UNiagaraNode* InNode = Cast<UNiagaraNode>(Node))
+		{
+			UObject* AssetRef = InNode->GetReferencedAsset();
+			if (AssetRef != nullptr && AssetRef->IsA(UNiagaraScript::StaticClass()))
+			{
+				if (UNiagaraScript* FunctionScript = Cast<UNiagaraScript>(AssetRef))
+				{
+					if (FunctionScript->GetSource())
+					{
+						UNiagaraScriptSource* Source = CastChecked<UNiagaraScriptSource>(FunctionScript->GetSource());
+						if (Source != nullptr)
+						{
+							UNiagaraGraph* FunctionGraph = CastChecked<UNiagaraGraph>(Source->NodeGraph);
+							if (FunctionGraph != nullptr)
+							{
+								if (!Graphs.Contains(FunctionGraph))
+								{
+									FunctionGraph->GetAllReferencedGraphs(Graphs);
+								}
+							}
+						}
+					}
+				}
+				else if (UNiagaraGraph* FunctionGraph = Cast<UNiagaraGraph>(AssetRef))
+				{
+					if (!Graphs.Contains(FunctionGraph))
+					{
+						FunctionGraph->GetAllReferencedGraphs(Graphs);
+					}
+				}
+			}
+		}
+	}
+}
+
+/** Determine if another item has been synchronized with this graph.*/
+bool UNiagaraGraph::IsOtherSynchronized(const FGuid& InChangeId) const
+{
+	if (ChangeId.IsValid() && ChangeId == InChangeId)
+	{
+		return true;
+	}
+	return false;
+}
+
+/** Identify that this graph has undergone changes that will require synchronization with a compiled script.*/
+void UNiagaraGraph::MarkGraphRequiresSynchronization()
+{
+	Modify();
+	ChangeId = FGuid::NewGuid();
+	UE_LOG(LogNiagaraEditor, Log, TEXT("MarkGraphRequiresSynchronization %s"), *ChangeId.ToString());
+}
