Copying //UE4/Orion-Staging to //UE4/Main (Source: //Orion/Dev-General @ 3606378)

diff --git a/Engine/Plugins/FX/Niagara/Source/NiagaraEditor/Private/NiagaraHlslTranslator.cpp b/Engine/Plugins/FX/Niagara/Source/NiagaraEditor/Private/NiagaraHlslTranslator.cpp
new file mode 100644
--- /dev/null
+++ b/Engine/Plugins/FX/Niagara/Source/NiagaraEditor/Private/NiagaraHlslTranslator.cpp
@@ -0,0 +1,4285 @@
+// Copyright 1998-2017 Epic Games, Inc. All Rights Reserved.
+
+#include "NiagaraHlslTranslator.h"
+#include "NiagaraComponent.h"
+#include "NiagaraGraph.h"
+#include "NiagaraScriptSource.h"
+#include "EdGraphUtilities.h"
+#include "UObject/UObjectHash.h"
+#include "NiagaraNode.h"
+#include "NiagaraNodeFunctionCall.h"
+#include "NiagaraNodeInput.h"
+#include "NiagaraNodeOutput.h"
+#include "NiagaraNodeReadDataSet.h"
+#include "NiagaraNodeWriteDataSet.h"
+#include "NiagaraNodeParameterMapGet.h"
+#include "NiagaraNodeParameterMapSet.h"
+#include "NiagaraNodeParameterCollection.h"
+#include "NiagaraNodeOp.h"
+#include "NiagaraNodeConvert.h"
+#include "EdGraphSchema_Niagara.h"
+#include "Interfaces/ITargetPlatformManagerModule.h"
+#include "Interfaces/IShaderFormat.h"
+#include "ShaderFormatVectorVM.h"
+#include "NiagaraConstants.h"
+#include "NiagaraSystem.h"
+#include "NiagaraNodeEmitter.h"
+#include "INiagaraEditorTypeUtilities.h"
+#include "NiagaraEditorUtilities.h"
+
+#include "NiagaraDataInterface.h"
+#include "NiagaraDataInterfaceCurve.h"
+#include "NiagaraDataInterfaceVector2DCurve.h"
+#include "NiagaraDataInterfaceVectorCurve.h"
+#include "NiagaraDataInterfaceVector4Curve.h"
+#include "NiagaraDataInterfaceColorCurve.h"
+#include "NiagaraDataInterfaceStaticMesh.h"
+#include "NiagaraDataInterfaceCurlNoise.h"
+
+#include "NiagaraParameterCollection.h"
+#include "NiagaraEditorTickables.h"
+#include "ShaderCore.h"
+
+#define LOCTEXT_NAMESPACE "NiagaraCompiler"
+
+
+static FNiagaraShaderQueueTickable NiagaraShaderQueueProcessor;
+FNiagaraShaderProcessorTickable NiagaraShaderProcessor;
+
+
+// not pretty. TODO: refactor
+// this will be called via a delegate from UNiagaraScript's cache for cook function,
+// because editor tickables aren't ticked during cooking
+void FNiagaraShaderQueueTickable::ProcessQueue()
+{
+	for (FNiagaraCompilationQueue::NiagaraCompilationQueueItem &Item : FNiagaraCompilationQueue::Get()->GetQueue())
+	{
+		FNiagaraScript *Script = Item.Script;
+		UNiagaraScript* CompilableScript = Script->GetGTScript();
+
+		// For now System scripts don't generate HLSL and go through a special pass...
+		// [OP] thinking they'll likely never run on GPU anyways
+		if (CompilableScript->IsValidLowLevel() == false || CompilableScript->IsSystemSpawnScript() || CompilableScript->IsSystemUpdateScript() || CompilableScript->IsEmitterSpawnScript() || CompilableScript->IsEmitterUpdateScript())
+		{
+			continue;
+		}
+
+		TRefCountPtr<FNiagaraShaderMap>NewShaderMap = Item.ShaderMap;
+		FNiagaraComputeShaderCompilationOutput NewCompilationOutput;
+		FHlslNiagaraTranslator NiagaraTranslator;
+		FHlslNiagaraTranslatorOptions Options;
+		Options.SimTarget = ENiagaraSimTarget::GPUComputeSim;
+
+		bool bNeedsPostCompile = false;
+		if (CompilableScript->GetSource()->IsPreCompiled() == false)
+		{
+		    CompilableScript->GetSource()->PreCompile(Cast<UNiagaraEmitter>(CompilableScript->GetOuter()));
+			bNeedsPostCompile = true;
+		}
+		FNiagaraTranslateResults Results = NiagaraTranslator.Translate(CompilableScript, Options);
+		if (bNeedsPostCompile)
+		{
+			CompilableScript->GetSource()->PostCompile();
+		}
+		Script->SetDatainterfaceBufferDescriptors(NiagaraTranslator.GetDatainterfaceBufferDescriptors());
+
+		// make sure buffers are set up on the original script
+		for (int32 i = 0; i < NiagaraTranslator.GetDatainterfaceBufferDescriptors().Num(); i++)
+		{
+			if (CompilableScript->DataInterfaceInfo.Num() > i)
+			{
+				CompilableScript->DataInterfaceInfo[i].DataInterface->SetupBuffers(NiagaraTranslator.GetDatainterfaceBufferDescriptors()[i]);
+			}
+		}
+
+		if (Results.bHLSLGenSucceeded && Results.MessageLog->NumErrors == 0)
+		{
+			Script->HlslOutput = Results.OutputHLSL;
+			// Create a shader compiler environment for the script that will be shared by all jobs from this script
+			TRefCountPtr<FShaderCompilerEnvironment> CompilerEnvironment = new FShaderCompilerEnvironment();
+
+			FString ShaderCode = NiagaraTranslator.GetTranslatedHLSL();
+			const bool bSynchronousCompile = false;
+
+			// Compile the shaders for the script.
+			NewShaderMap->Compile(Script, Item.ShaderMapId, CompilerEnvironment, NewCompilationOutput, Item.Platform, bSynchronousCompile, Item.bApply);
+		}
+		else
+		{
+			UE_LOG(LogNiagaraEditor, Warning, TEXT("%s"), TEXT("Niagara GPU Translation Failed!"));
+			TArray<TSharedRef<FTokenizedMessage>> Messages;
+			if (Results.MessageLog)
+			{
+				Messages.Append(Results.MessageLog->Messages);
+			}
+
+			for (TSharedRef<FTokenizedMessage> Message : Messages)
+			{
+				if (Message->GetSeverity() == EMessageSeverity::Info)
+				{
+					UE_LOG(LogNiagaraEditor, Log, TEXT("%s"), *Message->ToText().ToString());
+				}
+				else if (Message->GetSeverity() == EMessageSeverity::Warning || Message->GetSeverity() == EMessageSeverity::PerformanceWarning)
+				{
+					UE_LOG(LogNiagaraEditor, Warning, TEXT("%s"), *Message->ToText().ToString());
+				}
+				else if (Message->GetSeverity() == EMessageSeverity::Error || Message->GetSeverity() == EMessageSeverity::CriticalError)
+				{
+					UE_LOG(LogNiagaraEditor, Error, TEXT("%s"), *Message->ToText().ToString());
+				}
+			}
+		}
+	}
+
+	FNiagaraCompilationQueue::Get()->GetQueue().Empty();
+}
+
+void FNiagaraShaderQueueTickable::Tick(float DeltaSeconds)
+{
+	ProcessQueue();
+}
+
+
+template<typename Action>
+void TraverseGraphFromOutputDepthFirst(FHlslNiagaraTranslator& Compiler, const UEdGraphSchema_Niagara* Schema, UNiagaraNode* Node, Action& VisitAction, int32 VisitID)
+{
+	if (Node->ShouldVisit(VisitID))
+	{
+		TArray<UEdGraphPin*> InputPins;
+		Node->GetInputPins(InputPins);
+		for (UEdGraphPin* InputPin : InputPins)
+		{
+			// TODO: Error here if there are multiple links or non-niagara nodes?
+			if (InputPin->LinkedTo.Num() == 1)
+			{
+				UNiagaraNode* LinkedNiagaraNode = Cast<UNiagaraNode>(InputPin->LinkedTo[0]->GetOwningNode());
+				if (LinkedNiagaraNode != nullptr)
+				{
+					TraverseGraphFromOutputDepthFirst(Compiler, Schema, LinkedNiagaraNode, VisitAction, VisitID);
+				}
+			}
+		}
+		VisitAction(Schema, Node);
+	}
+}
+
+void FixUpNumericPinsVisitor(FHlslNiagaraTranslator& Compiler, const UEdGraphSchema_Niagara* Schema, UNiagaraNode* Node)
+{
+	// Fix up numeric input pins and keep track of numeric types to decide the output type.
+	TArray<FNiagaraTypeDefinition> InputTypes;
+	TArray<UEdGraphPin*> InputPins;
+	Node->GetInputPins(InputPins);
+	for (UEdGraphPin* InputPin : InputPins)
+	{
+		if (InputPin->PinType.PinCategory == UEdGraphSchema_Niagara::PinCategoryType)
+		{
+			FNiagaraTypeDefinition InputPinType = Schema->PinToTypeDefinition(InputPin);
+
+			// If the input pin is the generic numeric type set it to the type of the linked output pin which should have been processed already.
+			if (InputPinType == FNiagaraTypeDefinition::GetGenericNumericDef() && InputPin->LinkedTo.Num() == 1)
+			{
+				UEdGraphPin* InputPinLinkedPin = InputPin->LinkedTo[0];
+				FNiagaraTypeDefinition InputPinLinkedPinType = Schema->PinToTypeDefinition(InputPinLinkedPin);
+				if (InputPinLinkedPinType.IsValid())
+				{
+					// Only update the input pin type if the linked pin type is valid.
+					InputPin->PinType = Schema->TypeDefinitionToPinType(InputPinLinkedPinType);
+					InputPinType = InputPinLinkedPinType;
+				}
+			}
+
+			if (InputPinType == FNiagaraTypeDefinition::GetGenericNumericDef())
+			{
+				Compiler.Error(LOCTEXT("NumericPinError", "Unable to deduce type for numeric input pin."), Node, InputPin);
+			}
+			InputTypes.Add(InputPinType);
+		}
+	}
+
+	// Fix up numeric outputs based on the inputs.
+	if (InputTypes.Num() > 0 && Node->GetNumericOutputTypeSelectionMode() != ENiagaraNumericOutputTypeSelectionMode::None)
+	{
+		FNiagaraTypeDefinition OutputNumericType = FNiagaraTypeDefinition::GetNumericOutputType(InputTypes, Node->GetNumericOutputTypeSelectionMode());
+		if (OutputNumericType != FNiagaraTypeDefinition::GetGenericNumericDef())
+		{
+			TArray<UEdGraphPin*> OutputPins;
+			Node->GetOutputPins(OutputPins);
+			for (UEdGraphPin* OutputPin : OutputPins)
+			{
+				FNiagaraTypeDefinition OutputPinType = Schema->PinToTypeDefinition(OutputPin);
+				if (OutputPinType == FNiagaraTypeDefinition::GetGenericNumericDef())
+				{
+					OutputPin->PinType = Schema->TypeDefinitionToPinType(OutputNumericType);
+				}
+			}
+		}
+	}
+}
+
+void FixUpNumericPins(FHlslNiagaraTranslator& Compiler, const UEdGraphSchema_Niagara* Schema, UNiagaraNode* Node)
+{
+	auto FixUpVisitor = [&](const UEdGraphSchema_Niagara* LSchema, UNiagaraNode* LNode) { FixUpNumericPinsVisitor(Compiler, LSchema, LNode); };
+	TraverseGraphFromOutputDepthFirst(Compiler, Schema, Node, FixUpVisitor, ++Compiler.VisitID);
+}
+
+/* Go through the graph and attempt to auto-detect the type of any numeric pins by working back from the leaves of the graph. Only change the types of pins, not FNiagaraVariables.*/
+void PreprocessGraph(FHlslNiagaraTranslator& Compiler, const UEdGraphSchema_Niagara* Schema, UNiagaraGraph* Graph, UNiagaraNodeOutput* OutputNode)
+{
+	if (OutputNode)
+	{
+		FixUpNumericPins(Compiler, Schema, OutputNode);
+	}
+	else
+	{
+		// This should never happen!
+		TArray<FStringFormatArg> Args;
+		Args.Add(FString::FromInt((int32)Compiler.GetTargetUsage()));
+		FString ErrorText = FString::Format(TEXT("Unable to preprocess graph due to missing output node of type '{0}'!"), Args);
+		Compiler.Error(FText::FromString(ErrorText), nullptr, nullptr);
+	}
+}
+
+/* Go through the graph and force any input nodes with Numeric types to a hard-coded type of float. This will allow modules and functions to compile properly.*/
+void PreProcessGraphForInputNumerics(FHlslNiagaraTranslator& Compiler, const UEdGraphSchema_Niagara* Schema, UNiagaraGraph* Graph, TArray<FNiagaraVariable>& OutChangedNumericParams)
+{
+	// Visit all input nodes
+	TArray<UNiagaraNodeInput*> InputNodes;
+	Graph->FindInputNodes(InputNodes);
+	for (UNiagaraNodeInput* InputNode : InputNodes)
+	{
+		// See if any of the output pins are of Numeric type. If so, force to floats.
+		TArray<UEdGraphPin*> OutputPins;
+		InputNode->GetOutputPins(OutputPins);
+		for (UEdGraphPin* OutputPin : OutputPins)
+		{
+			FNiagaraTypeDefinition OutputPinType = Schema->PinToTypeDefinition(OutputPin);
+			if (OutputPinType == FNiagaraTypeDefinition::GetGenericNumericDef())
+			{
+				OutputPin->PinType = Schema->TypeDefinitionToPinType(FNiagaraTypeDefinition::GetFloatDef());
+			}
+		}
+
+		// Record that we touched this variable for later cleanup and make sure that the 
+		// variable's type now matches the pin.
+		if (InputNode->Input.GetType() == FNiagaraTypeDefinition::GetGenericNumericDef())
+		{
+			OutChangedNumericParams.Add(InputNode->Input);
+			InputNode->Input.SetType(FNiagaraTypeDefinition::GetFloatDef());
+		}
+	}
+}
+
+/* Should be called after all pins have been successfully auto-detected for type. This goes through and synchronizes any Numeric FNiagaraVarible outputs with the deduced pin type. This will allow modules and functions to compile properly.*/
+void PreProcessGraphForAttributeNumerics(FHlslNiagaraTranslator& Compiler, const UEdGraphSchema_Niagara* Schema, UNiagaraGraph* Graph, UNiagaraNodeOutput* OutputNode, TArray<FNiagaraVariable>& OutChangedNumericParams)
+{
+	// Visit the output node
+	if (OutputNode != nullptr)
+	{
+		// For each pin, make sure that if it has a valid type, but the associated variable is still Numeric,
+		// force the variable to match the pin's new type. Record that we touched this variable for later cleanup.
+		TArray<UEdGraphPin*> InputPins;
+		OutputNode->GetInputPins(InputPins);
+		check(OutputNode->Outputs.Num() == InputPins.Num());
+		for (int32 i = 0; i < InputPins.Num(); i++)
+		{
+			FNiagaraVariable& Param = OutputNode->Outputs[i];
+			UEdGraphPin* InputPin = InputPins[i];
+
+			FNiagaraTypeDefinition InputPinType = Schema->PinToTypeDefinition(InputPin);
+			if (Param.GetType() == FNiagaraTypeDefinition::GetGenericNumericDef() && 
+				InputPinType != FNiagaraTypeDefinition::GetGenericNumericDef())
+			{
+				OutChangedNumericParams.Add(Param);
+				Param.SetType(InputPinType);
+			}
+		}
+	}
+}
+
+/* Clean up the lingering Systems of PreProcessGraphForInputNumerics and PreProcessGraphForAttributeNumerics by resetting the FNiagaraVariables back to their original types.*/
+void RevertParametersToNumerics(FHlslNiagaraTranslator& Compiler, UNiagaraScript* Script, const TArray<FNiagaraVariable>& ChangedNumericParams)
+{
+	// We either changed an input node variable or an output node variable in the prior functions, let's
+	// check where those ended up and fixup any discrepencies.
+	for (const FNiagaraVariable& ChangedVariable : ChangedNumericParams)
+	{
+		// Check input variables... 
+		FNiagaraVariable* CorrespondingVariable = Script->Parameters.FindParameter(ChangedVariable);
+		if (CorrespondingVariable == nullptr)
+		{
+			// Check output variables... 
+			CorrespondingVariable = Script->Attributes.FindByPredicate([&](FNiagaraVariable& Attribute) { return Attribute.GetName() == ChangedVariable.GetName(); });
+		}
+
+		// Convert back to Numeric so that we maintain consistency externally.
+		if (CorrespondingVariable != nullptr)
+		{
+			check(ChangedVariable.GetType() == FNiagaraTypeDefinition::GetGenericNumericDef());
+			CorrespondingVariable->SetType(ChangedVariable.GetType());
+		}
+		else
+		{
+			// This should never happen!
+			TArray<FStringFormatArg> Args;
+			Args.Add(ChangedVariable.ToString());
+			FString ErrorText = FString::Format(TEXT("Unable to find parameter '{0}' in outputs!"), Args);
+			Compiler.Error(FText::FromString(ErrorText), nullptr, nullptr);
+		}
+	}
+}
+
+
+void PreprocessFunctionGraph(FHlslNiagaraTranslator& Compiler, const UEdGraphSchema_Niagara* Schema, UNiagaraGraph* Graph, const TArray<UEdGraphPin*>& CallInputs, const TArray<UEdGraphPin*>& CallOutputs, ENiagaraScriptUsage ScriptUsage)
+{
+	// Change any numeric inputs or outputs to match the types from the call node.
+	TArray<UNiagaraNodeInput*> InputNodes;
+	
+	// Only handle nodes connected to the correct output node in the event of multiple output nodes in the graph.
+	UNiagaraGraph::FFindInputNodeOptions Options;
+	Options.bFilterByScriptUsage = true;
+	Options.TargetScriptUsage = ScriptUsage;
+
+	Graph->FindInputNodes(InputNodes, Options);
+
+	for (UNiagaraNodeInput* InputNode : InputNodes)
+	{
+		FNiagaraVariable& Input = InputNode->Input;
+		if (Input.GetType() == FNiagaraTypeDefinition::GetGenericNumericDef())
+		{
+			UEdGraphPin* const* MatchingPin = CallInputs.FindByPredicate([&](const UEdGraphPin* Pin) { return *(Pin->PinName) == Input.GetName(); });
+
+			if (MatchingPin != nullptr)
+			{
+				FNiagaraTypeDefinition PinType = Schema->PinToTypeDefinition(*MatchingPin);
+				Input.SetType(PinType);
+				TArray<UEdGraphPin*> OutputPins;
+				InputNode->GetOutputPins(OutputPins);
+				check(OutputPins.Num() == 1);
+				OutputPins[0]->PinType = (*MatchingPin)->PinType;
+			}
+		}
+	}
+
+	UNiagaraNodeOutput* OutputNode = Graph->FindOutputNode(ScriptUsage);
+	check(OutputNode);
+
+	TArray<UEdGraphPin*> InputPins;
+	OutputNode->GetInputPins(InputPins);
+
+	for (FNiagaraVariable& Output : OutputNode->Outputs)
+	{
+		if (Output.GetType() == FNiagaraTypeDefinition::GetGenericNumericDef())
+		{
+			UEdGraphPin* const* MatchingPin = CallOutputs.FindByPredicate([&](const UEdGraphPin* Pin) { return *(Pin->PinName) == Output.GetName(); });
+
+			if (MatchingPin != nullptr)
+			{
+				FNiagaraTypeDefinition PinType = Schema->PinToTypeDefinition(*MatchingPin);
+				Output.SetType(PinType);
+			}
+		}
+	}
+
+	FixUpNumericPins(Compiler, Schema, OutputNode);
+}
+
+ENiagaraScriptCompileStatus FNiagaraTranslateResults::TranslateResultsToSummary(const FNiagaraTranslateResults* TranslateResults)
+{
+	ENiagaraScriptCompileStatus SummaryStatus = ENiagaraScriptCompileStatus::NCS_Unknown;
+	if (TranslateResults != nullptr)
+	{
+		if (TranslateResults->MessageLog->NumErrors > 0)
+		{
+			SummaryStatus = ENiagaraScriptCompileStatus::NCS_Error;
+		}
+		else
+		{
+			if (TranslateResults->bHLSLGenSucceeded)
+			{
+				if (TranslateResults->MessageLog->NumWarnings)
+				{
+					SummaryStatus = ENiagaraScriptCompileStatus::NCS_UpToDateWithWarnings;
+				}
+				else
+				{
+					SummaryStatus = ENiagaraScriptCompileStatus::NCS_UpToDate;
+				}
+			}
+		}
+	}
+	return SummaryStatus;
+}
+
+
+FString FHlslNiagaraTranslator::GetCode(int32 ChunkIdx)
+{
+	FNiagaraCodeChunk& Chunk = CodeChunks[ChunkIdx];
+	return GetCode(Chunk);
+}
+
+FString FHlslNiagaraTranslator::GetCode(FNiagaraCodeChunk& Chunk)
+{
+	TArray<FStringFormatArg> Args;
+	for (int32 i = 0; i < Chunk.SourceChunks.Num(); ++i)
+	{
+		Args.Add(GetCodeAsSource(Chunk.SourceChunks[i]));
+	}
+	FString DefinitionString = FString::Format(*Chunk.Definition, Args);
+
+	FString FinalString;
+
+	if (Chunk.Mode == ENiagaraCodeChunkMode::Body)
+	{
+		FinalString += TEXT("\t");
+	}
+
+	if (Chunk.SymbolName.IsEmpty())
+	{
+		check(!DefinitionString.IsEmpty());
+		FinalString += DefinitionString + (Chunk.bIsTerminated ? TEXT(";\n") : TEXT("\n"));
+	}
+	else
+	{
+		if (DefinitionString.IsEmpty())
+		{
+			check(Chunk.bDecl);//Otherwise, we're doing nothing here.
+			FinalString += GetStructHlslTypeName(Chunk.Type) + TEXT(" ") + Chunk.SymbolName + TEXT(";\n");
+		}
+		else
+		{
+			if (Chunk.bDecl)
+			{
+				FinalString += GetStructHlslTypeName(Chunk.Type) + TEXT(" ") + Chunk.SymbolName + TEXT(" = ") + DefinitionString + TEXT(";\n");
+			}
+			else
+			{
+				FinalString += Chunk.SymbolName + TEXT(" = ") + DefinitionString + TEXT(";\n");
+			}
+		}
+	}
+	return FinalString;
+}
+
+FString FHlslNiagaraTranslator::GetCodeAsSource(int32 ChunkIdx)
+{
+	if(ChunkIdx>=0 && ChunkIdx<CodeChunks.Num())
+	{ 
+		FNiagaraCodeChunk& Chunk = CodeChunks[ChunkIdx];
+		return Chunk.SymbolName;
+	}
+	return "Undefined";
+}
+
+bool FHlslNiagaraTranslator::ValidateTypePins(UNiagaraNode* NodeToValidate)
+{
+	bool bPinsAreValid = true;
+	for (UEdGraphPin* Pin : NodeToValidate->GetAllPins())
+	{
+		if (Pin->PinType.PinCategory == "")
+		{
+			Error(LOCTEXT("InvalidPinTypeError", "Node pin has an undefined type."), NodeToValidate, Pin);
+			bPinsAreValid = false;
+		}
+		else if (Pin->PinType.PinCategory == UEdGraphSchema_Niagara::PinCategoryType)
+		{
+			FNiagaraTypeDefinition Type = Schema->PinToTypeDefinition(Pin);
+			if (Type.IsValid() == false)
+			{
+				Error(LOCTEXT("InvalidPinTypeError", "Node pin has an undefined type."), NodeToValidate, Pin);
+				bPinsAreValid = false;
+			}
+		}
+	}
+	return bPinsAreValid;
+}
+
+
+void FHlslNiagaraTranslator::GenerateFunctionSignature(ENiagaraScriptUsage ScriptUsage, FString InName, const FString& InFullName, UNiagaraGraph* FuncGraph, TArray<int32>& Inputs, bool bHadNumericInputs, FNiagaraFunctionSignature& OutSig)const
+{
+	TArray<FNiagaraVariable> InputVars;
+
+	// Only handle nodes connected to the correct output node in the event of multiple output nodes in the graph.
+	TArray<UNiagaraNodeInput*> InputsNodes;
+	UNiagaraGraph::FFindInputNodeOptions Options;
+	Options.bSort = true;
+	Options.bFilterDuplicates = true;
+	Options.bFilterByScriptUsage = true;
+	Options.TargetScriptUsage = ScriptUsage;
+	FuncGraph->FindInputNodes(InputsNodes, Options);
+
+	if (Inputs.Num() != InputsNodes.Num())
+	{
+		const_cast<FHlslNiagaraTranslator*>(this)->Error(FText::Format(LOCTEXT("GenerateFunctionSignatureFail", "Generating function signature for {0} failed.  The function graph is invalid."), FText::FromString(InFullName)), nullptr, nullptr);
+		return;
+	}
+
+	for (int32 i = 0; i < InputsNodes.Num(); ++i)
+	{
+		//Only add to the signature if the caller has provided it, otherwise we use a local default.
+		if (Inputs[i] != INDEX_NONE)
+		{
+			InputVars.Add(InputsNodes[i]->Input);
+			if (bHadNumericInputs)
+			{
+				InName += TEXT("_In") + InputsNodes[i]->Input.GetType().GetName();
+			}
+		}
+	}
+
+	//Now actually remove the missing inputs so they match the signature.
+	Inputs.Remove(INDEX_NONE);
+
+	TArray<FNiagaraVariable> OutputVars;
+	FuncGraph->GetOutputNodeVariables(ScriptUsage, OutputVars);
+
+	for (int32 i = 0; i < OutputVars.Num(); ++i)
+	{
+		//Only add to the signature if the caller has provided it, otherwise we use a local default.
+		if (bHadNumericInputs)
+		{
+			InName += TEXT("_Out") + OutputVars[i].GetType().GetName();
+		}
+	}
+
+	const FString* ModuleAliasStr = ActiveHistoryForFunctionCalls.GetModuleAlias();
+	const FString* EmitterAliasStr = ActiveHistoryForFunctionCalls.GetEmitterAlias();
+	// For now, we want each module call to be unique due to parameter maps and aliasing causing different variables
+	// to be written within each call.
+	if ((ScriptUsage == ENiagaraScriptUsage::Module || ScriptUsage == ENiagaraScriptUsage::DynamicInput ||
+		ScriptUsage == ENiagaraScriptUsage::EmitterSpawnScript || ScriptUsage == ENiagaraScriptUsage::EmitterUpdateScript)
+		&& (ModuleAliasStr != nullptr || EmitterAliasStr != nullptr))
+	{
+		FString SignatureName;
+		if (ModuleAliasStr != nullptr)
+		{
+			SignatureName = GetSanitizedSymbolName(*ModuleAliasStr);
+		}
+		if (EmitterAliasStr != nullptr)
+		{
+			FString Prefix = ModuleAliasStr != nullptr ? TEXT("_") : TEXT("");
+			SignatureName +=  Prefix + GetSanitizedSymbolName(*EmitterAliasStr);
+		}
+		SignatureName = SignatureName.Replace(TEXT("."), TEXT("_"));
+		OutSig = FNiagaraFunctionSignature(*SignatureName, InputVars, OutputVars, *InFullName, true, false);
+	}
+	else
+	{
+		OutSig = FNiagaraFunctionSignature(*InName, InputVars, OutputVars, *InFullName, true, false);
+	}
+}
+
+//////////////////////////////////////////////////////////////////////////
+
+FHlslNiagaraTranslator::FHlslNiagaraTranslator()
+	: Script(nullptr)
+	, Schema(nullptr)
+	, TranslateResults(&MessageLog)
+	, bInsideInterpolatedSpawnScript(false)
+	, VisitID(INDEX_NONE)
+{
+	// Make the message log silent so we're not spamming the blueprint log.
+	MessageLog.bSilentMode = true;
+}
+
+
+FString FHlslNiagaraTranslator::GetFunctionDefinitions()
+{
+	FString FwdDeclString;
+	FString DefinitionsString;
+
+	for (TPair<FNiagaraFunctionSignature, FString> FuncPair : Functions)
+	{
+		FString Sig = GetFunctionSignature(FuncPair.Key);
+		FwdDeclString += Sig + TEXT(";\n");
+		if (!FuncPair.Value.IsEmpty())
+		{
+			DefinitionsString += Sig + TEXT("\n{\n") + FuncPair.Value + TEXT("}\n\n");
+		}
+		// Don't do anything if the value is empty on the function pair, as this is indicative of 
+		// data interface functions that should be defined differently.
+	}
+
+	return FwdDeclString + TEXT("\n") + DefinitionsString;
+}
+
+UNiagaraGraph* FHlslNiagaraTranslator::CloneGraphAndPrepareForCompilation(UNiagaraScript* InScript, UNiagaraScriptSource* InSource, bool bClearErrors)
+{
+	if (InSource == nullptr)
+	{
+		Error(LOCTEXT("NoSourceErrorMessage", "Script has no source."), nullptr, nullptr);
+		return nullptr;
+	}
+
+	UNiagaraGraph* SourceGraph = InSource->GetPrecomputedNodeGraph();
+	check(SourceGraph != nullptr);
+	
+	TArray<FNiagaraVariable> ChangedNumericParams;
+	// In the case of functions or modules, we may not have enough information at this time to fully resolve the type. In that case,
+	// we circumvent the resulting errors by forcing a type. This gives the user an appropriate level of type checking. We will, however need to clean this up in
+	// the parameters that we output.
+	bool bForceParametersToResolveNumerics = InScript->IsStandaloneScript();
+	if (bForceParametersToResolveNumerics)
+	{
+		PreProcessGraphForInputNumerics(*this, Schema, SourceGraph, ChangedNumericParams);
+	}
+
+	// Auto-deduce the input types for numerics in the graph and overwrite the types on the pins. If PreProcessGraphForInputNumerics occurred, then
+	// we will have pre-populated the inputs with valid types.
+	TArray<UNiagaraNodeOutput*> OutputNodes;
+	SourceGraph->FindOutputNodes(OutputNodes);
+
+	for (UNiagaraNodeOutput* OutputNode : OutputNodes)
+	{
+		PreprocessGraph(*this, Schema, SourceGraph, OutputNode);
+
+		// Now that we've auto-deduced the types, we need to handle any lingering Numerics in the Output's FNiagaraVariable outputs. 
+		// We use the pin's deduced type to temporarily overwrite the variable's type.
+		if (bForceParametersToResolveNumerics)
+		{
+			PreProcessGraphForAttributeNumerics(*this, Schema, SourceGraph, OutputNode, ChangedNumericParams);
+		}
+	}
+
+	return SourceGraph;
+}
+
+FString FHlslNiagaraTranslator::BuildMissingDefaults(TArray<TPair<FNiagaraVariable, const UEdGraphPin*> >& MissingDefaults)
+{
+	FString HlslOutputString = TEXT("void HandleMissingDefaultValues(inout FSimulationContext Context)\n{\n");
+	for (TPair<FNiagaraVariable, const UEdGraphPin*>& Item : MissingDefaults)
+	{
+		FNiagaraVariable Var = Item.Key;
+		const UEdGraphPin* DefaultPin = Item.Value;
+
+		if (Script->IsParticleSpawnScript() || Script->IsEmitterSpawnScript() || Script->IsSystemSpawnScript())
+		{
+			if (FNiagaraParameterMapHistory::IsInitialValue(Var))
+			{
+				FNiagaraVariable SourceForInitialValue = FNiagaraParameterMapHistory::GetSourceForInitialValue(Var);
+				FString ParameterMapInstanceName = GetParameterMapInstanceName(0);
+				HlslOutputString += TEXT("\t") + ParameterMapInstanceName + TEXT(".") + GetSanitizedSymbolName(Var.GetName().ToString() + TEXT(" = ") +
+					ParameterMapInstanceName + TEXT(".") + GetSanitizedSymbolName(SourceForInitialValue.GetName().ToString()) + TEXT(";"));
+				continue;
+			}
+
+			if (DefaultPin && DefaultPin->LinkedTo.Num() == 0)
+			{
+				FNiagaraVariable PinVar = Schema->PinToNiagaraVariable(DefaultPin, true);
+				FString ConstantValue = GenerateConstantString(PinVar);
+				if (Var.GetType().GetClass() == nullptr) // Only need to do this wiring for things that aren't data interfaces.
+				{
+					FString ParameterMapInstanceName = GetParameterMapInstanceName(0);
+					HlslOutputString += TEXT("\t") + ParameterMapInstanceName + TEXT(".") + GetSanitizedSymbolName(Var.GetName().ToString()) + TEXT(" = ") + ConstantValue + TEXT(";");
+				}
+			}
+			else if (DefaultPin)
+			{
+				Error(FText::Format(LOCTEXT("OnlySimpleDefaults", "Only simple constants are supported for defaults of primary values: {0}"), FText::FromName(Var.GetName())), nullptr, nullptr);
+			}
+		}
+		HlslOutputString += TEXT("\n");
+	}
+	HlslOutputString += TEXT("\n}\n");
+	return HlslOutputString;
+}
+
+FString FHlslNiagaraTranslator::BuildParameterMapHlslDefinitions(TArray<FNiagaraVariable>& PrimaryDataSetOutputEntries, TArray<TPair<FNiagaraVariable, const UEdGraphPin*> >& MissingDefaults)
+{
+	FString HlslOutputString;
+
+	// Determine the unique parameter map structs...
+	TArray<const UEdGraphPin*> UniqueParamMapStartingPins;
+	for (int32 ParamMapIdx = 0; ParamMapIdx < ParamMapHistories.Num(); ParamMapIdx++)
+	{
+		const UEdGraphPin* OriginalPin = ParamMapHistories[ParamMapIdx].GetOriginalPin();
+		UniqueParamMapStartingPins.AddUnique(OriginalPin);
+	}
+	
+
+	TArray<FNiagaraVariable> UniqueVariables;
+
+	// Add in currently defined system vars.
+	TArray<FNiagaraVariable> ValueArray;
+	ParamMapDefinedSystemToNamespaceVars.GenerateValueArray(ValueArray);
+	for (FNiagaraVariable& Var : ValueArray)
+	{
+		if (Var.GetType().GetClass() != nullptr)
+		{
+			continue;
+		}
+		UniqueVariables.AddUnique(Var);
+	}
+
+	// Add in currently defined emitter vars.
+	ParamMapDefinedEmitterParameterToNamespaceVars.GenerateValueArray(ValueArray);
+	for (FNiagaraVariable& Var : ValueArray)
+	{
+		if (Var.GetType().GetClass() != nullptr)
+		{
+			continue;
+		}
+
+		UniqueVariables.AddUnique(Var);
+	}
+
+	// Add in currently defined attribute vars.
+	ParamMapDefinedAttributesToNamespaceVars.GenerateValueArray(ValueArray);
+	for (FNiagaraVariable& Var : ValueArray)
+	{
+		if (Var.GetType().GetClass() != nullptr)
+		{
+			continue;
+		}
+
+		UniqueVariables.AddUnique(Var);
+	}
+
+	// Add in any bulk usage vars.
+	for (FNiagaraVariable& Var : ExternalVariablesForBulkUsage)
+	{
+		if (Var.GetType().GetClass() != nullptr)
+		{
+			continue;
+		}
+
+		UniqueVariables.AddUnique(Var);
+	}
+
+	// For now we only care about attributes from the other output parameter map histories.
+	for (int32 ParamMapIdx = 0; ParamMapIdx < OtherOutputParamMapHistories.Num(); ParamMapIdx++)
+	{
+		for (int32 VarIdx = 0; VarIdx < OtherOutputParamMapHistories[ParamMapIdx].Variables.Num(); VarIdx++)
+		{
+			FNiagaraVariable& Var = OtherOutputParamMapHistories[ParamMapIdx].Variables[VarIdx];
+			if (OtherOutputParamMapHistories[ParamMapIdx].IsPrimaryDataSetOutput(Var, Script))
+			{
+				int32 PreviousMax = UniqueVariables.Num();
+				if (UniqueVariables.AddUnique(Var) == PreviousMax) // i.e. we didn't find it previously, so we added to the end.
+				{
+					if (Script->IsParticleSpawnScript() || Script->IsEmitterSpawnScript() || Script->IsSystemSpawnScript())
+					{
+						if (!AddStructToDefinitionSet(Var.GetType()))
+						{
+							Error(FText::Format(LOCTEXT("ParameterMapTypeError", "Cannot handle type {0}! Variable: {1}"), Var.GetType().GetNameText(), FText::FromName(Var.GetName())), nullptr, nullptr);
+						}
+						if (FNiagaraParameterMapHistory::IsInitialValue(Var))
+						{
+							FNiagaraVariable SourceForInitialValue = FNiagaraParameterMapHistory::GetSourceForInitialValue(Var);
+							if (!UniqueVariables.Contains(SourceForInitialValue))
+							{
+								Error(FText::Format(LOCTEXT("MissingInitialValueSource", "Variable {0} is used, but its source variable {1} is not set!"), FText::FromName(Var.GetName()), FText::FromName(SourceForInitialValue.GetName())), nullptr, nullptr);
+							}
+							TPair<FNiagaraVariable, const UEdGraphPin*> Value;
+							Value.Key = Var;
+							Value.Value = nullptr;
+							MissingDefaults.Add(Value);
+							continue;
+						}
+
+						const UEdGraphPin* DefaultPin = OtherOutputParamMapHistories[ParamMapIdx].GetDefaultValuePin(VarIdx);
+						TPair<FNiagaraVariable, const UEdGraphPin*> Value;
+						Value.Key = Var;
+						Value.Value = DefaultPin;
+						MissingDefaults.Add(Value);
+					}
+				}
+			}
+		}
+	}
+
+
+	// Define all the top-level structs and look for sub-structs as yet undefined..
+	for (int32 UniqueParamMapIdx = 0; UniqueParamMapIdx < UniqueParamMapStartingPins.Num(); UniqueParamMapIdx++)
+	{
+		for (int32 ParamMapIdx = 0; ParamMapIdx < ParamMapHistories.Num(); ParamMapIdx++)
+		{
+			// We need to unify the variables across all the parameter maps that we've found during compilation. We 
+			// define the parameter maps as the "same struct type" if they originate from the same input pin.
+			const UEdGraphPin* OriginalPin = ParamMapHistories[ParamMapIdx].GetOriginalPin();
+			if (OriginalPin != UniqueParamMapStartingPins[UniqueParamMapIdx])
+			{
+				continue;
+			}
+
+			for (int32 VarIdx = 0; VarIdx < ParamMapHistories[ParamMapIdx].Variables.Num(); VarIdx++)
+			{
+				const FNiagaraVariable& SrcVariable = ParamMapHistories[ParamMapIdx].Variables[VarIdx];
+
+				if (SrcVariable.GetType().GetClass() != nullptr)
+				{
+					continue;
+				}
+
+				FNiagaraVariable Variable = SrcVariable;
+				UniqueVariables.AddUnique(Variable);
+			}
+		}
+	}
+
+	TMap<FString, TArray<TPair<FString, FString>> > ParamStructNameToMembers;
+	TArray<FString> ParamStructNames;
+
+	for (int32 UniqueVarIdx = 0; UniqueVarIdx < UniqueVariables.Num(); UniqueVarIdx++)
+	{
+		int UniqueParamMapIdx = 0;
+		FNiagaraVariable Variable = UniqueVariables[UniqueVarIdx];
+
+		if (!AddStructToDefinitionSet(Variable.GetType()))
+		{
+			Error(FText::Format(LOCTEXT("ParameterMapTypeError", "Cannot handle type {0}! Variable: {1}"), Variable.GetType().GetNameText(), FText::FromName(Variable.GetName())), nullptr, nullptr);
+		}
+
+		// In order 
+		for (int32 ParamMapIdx = 0; ParamMapIdx < OtherOutputParamMapHistories.Num(); ParamMapIdx++)
+		{
+			if (OtherOutputParamMapHistories[ParamMapIdx].IsPrimaryDataSetOutput(Variable, Script))
+			{
+				PrimaryDataSetOutputEntries.AddUnique(Variable);
+				break;
+			}
+		}
+
+		TArray<FString> StructNameArray;
+		int32 NumFound = Variable.GetName().ToString().ParseIntoArray(StructNameArray, TEXT("."));
+		if (NumFound == 1) // Meaning no split above
+		{
+			Error(FText::Format(LOCTEXT("OnlyOneNamespaceEntry", "Only one namespace entry found for: {0}"), FText::FromName(Variable.GetName())), nullptr, nullptr);
+		}
+		else if (NumFound > 1)
+		{
+			while (StructNameArray.Num())
+			{
+				FString FinalName = StructNameArray[StructNameArray.Num() - 1];
+				StructNameArray.RemoveAt(StructNameArray.Num() - 1);
+				FString StructType = FString::Printf(TEXT("FParamMap%d_%s"), UniqueParamMapIdx, *FString::Join<FString>(StructNameArray, TEXT("_")));
+				if (StructNameArray.Num() == 0)
+				{
+					StructType = FString::Printf(TEXT("FParamMap%d"), UniqueParamMapIdx);
+				}
+
+				FString TypeName = GetStructHlslTypeName(Variable.GetType());
+				FString VarName = GetSanitizedSymbolName(*FinalName);
+				if (NumFound > StructNameArray.Num() + 1 && StructNameArray.Num() != 0)
+				{
+					TypeName = FString::Printf(TEXT("FParamMap%d_%s_%s"), UniqueParamMapIdx, *FString::Join<FString>(StructNameArray, TEXT("_")), *GetSanitizedSymbolName(*FinalName));
+				}
+				else if (StructNameArray.Num() == 0)
+				{
+					TypeName = FString::Printf(TEXT("FParamMap%d_%s"), UniqueParamMapIdx, *GetSanitizedSymbolName(*FinalName));
+				}
+				TPair<FString, FString> Pair(TypeName, VarName);
+				ParamStructNameToMembers.FindOrAdd(StructType).AddUnique(Pair);
+				ParamStructNames.AddUnique(StructType);
+			}		
+		}
+	}
+
+	// Build up the sub-structs..
+	ParamStructNames.Sort();
+	FString StructDefString = "";
+	for (int32 i = ParamStructNames.Num() - 1; i >= 0; i--)
+	{
+		const FString& StructName = ParamStructNames[i];
+		StructDefString += FString::Printf(TEXT("struct %s\n{\n"), *StructName);
+		TArray<TPair<FString, FString>> StructMembers = ParamStructNameToMembers[StructName];
+		auto SortMembers = [](const TPair<FString, FString>& A, const TPair<FString, FString>& B)
+		{
+			return A.Value < B.Value;
+		};
+		StructMembers.Sort(SortMembers);
+		for (const TPair<FString, FString>& Line : StructMembers)
+		{
+			StructDefString += TEXT("\t") + Line.Key + TEXT(" ") + Line.Value + TEXT(";\n");;
+		}
+		StructDefString += TEXT("};\n\n");
+	}
+
+	HlslOutputString += StructDefString;
+
+	return HlslOutputString;
+}
+
+bool FHlslNiagaraTranslator::ShouldConsiderTargetParameterMap(ENiagaraScriptUsage InUsage) const
+{
+	ENiagaraScriptUsage TargetUsage = GetTargetUsage();
+	if (TargetUsage >= ENiagaraScriptUsage::ParticleSpawnScript && TargetUsage <= ENiagaraScriptUsage::ParticleEventScript)
+	{
+		return InUsage >= ENiagaraScriptUsage::ParticleSpawnScript && InUsage <= ENiagaraScriptUsage::ParticleEventScript;
+	}
+	else if (TargetUsage == ENiagaraScriptUsage::SystemSpawnScript)
+	{
+		if (InUsage == ENiagaraScriptUsage::SystemUpdateScript)
+		{
+			return true;
+		}
+		else if (TargetUsage == InUsage)
+		{
+			return true;
+		}
+	}
+	else if (TargetUsage == InUsage)
+	{
+		return true;
+	}
+
+	return false;
+}
+
+void FHlslNiagaraTranslator::HandleNamespacedExternalVariablesToDataSetRead(TArray<FNiagaraVariable>& InDataSetVars, FString InNamespaceStr)
+{
+	for (const FNiagaraVariable& Var : ExternalVariablesForBulkUsage)
+	{
+		if (FNiagaraParameterMapHistory::IsInNamespace(Var, InNamespaceStr))
+		{
+			InDataSetVars.Add(Var);
+		}
+	}
+}
+
+const FNiagaraTranslateResults &FHlslNiagaraTranslator::Translate(UNiagaraScript* InScript, FHlslNiagaraTranslatorOptions Options)
+{
+	check(InScript);
+
+	TranslationOptions = Options;
+	CompilationTarget = TranslationOptions.SimTarget;
+	TranslateResults.bHLSLGenSucceeded = false;
+	TranslateResults.OutputHLSL = "";
+
+	Script = InScript;
+	//Should we roll our own message/error log and put it in a window somewhere?
+	MessageLog.SetSourcePath(InScript->GetPathName());
+
+	UNiagaraGraph* SourceGraph = CloneGraphAndPrepareForCompilation(InScript, Cast<UNiagaraScriptSource>(InScript->GetSource()), true);
+	if (!SourceGraph)
+	{
+		Error(LOCTEXT("GetGraphFail", "Cannot find graph node!"), nullptr, nullptr);
+		return TranslateResults;
+	}
+
+	if (SourceGraph->IsEmpty())
+	{
+		if (Script->IsSystemSpawnScript() || Script->IsSystemUpdateScript())
+		{
+			Error(LOCTEXT("GetNoNodeSystemFail", "Graph contains no nodes! Please add an emitter."), nullptr, nullptr);
+		}
+		else
+		{
+			Error(LOCTEXT("GetNoNodeFail", "Graph contains no nodes! Please add an output node."), nullptr, nullptr);
+		}
+		return TranslateResults;
+	}
+
+	// Find the output node and compile it.
+	UNiagaraNodeOutput* OutputNode = SourceGraph->FindOutputNode(GetTargetUsage(), GetTargetOccurrence());
+	if (OutputNode == nullptr)
+	{
+		Error(FText::Format(LOCTEXT("GetOutputNodeFail", "Cannot find output node of type {0}!"), FText::AsNumber((int32)GetTargetUsage())), nullptr, nullptr);
+		return TranslateResults;
+	}
+	ValidateTypePins(OutputNode);
+	{
+		bool bHasAnyConnections = false;
+		for (int32 PinIdx = 0; PinIdx < OutputNode->Pins.Num(); PinIdx++)
+		{
+			if (OutputNode->Pins[PinIdx]->Direction == EEdGraphPinDirection::EGPD_Input && OutputNode->Pins[PinIdx]->LinkedTo.Num() != 0)
+			{
+				bHasAnyConnections = true;
+			}
+		}
+		if (!bHasAnyConnections)
+		{
+			Error(FText::Format(LOCTEXT("GetOutputNodeConnectivityFail", "Cannot find any connections to output node of type {0}!"), FText::AsNumber((int32)GetTargetUsage())), nullptr, nullptr);
+			return TranslateResults;
+		}
+	}
+
+	// Build up a traversal from the output pin that touches all the parameter maps we might encounter, using the histories we build below.
+	ParamMapHistories.Empty();
+	ParamMapSetVariablesToChunks.Empty();
+	TArray<UNiagaraNodeOutput*> UsageOutputs;
+	if (Script->IsInterpolatedParticleSpawnScript())
+	{
+		UsageOutputs.Add(OutputNode);
+		UsageOutputs.Add(SourceGraph->FindOutputNode(ENiagaraScriptUsage::ParticleUpdateScript));
+		ParamMapHistories.AddDefaulted(2);
+		ParamMapSetVariablesToChunks.AddDefaulted(2);
+	}
+	else
+	{
+		UsageOutputs.Add(OutputNode);
+		ParamMapHistories.AddDefaulted(1);
+		ParamMapSetVariablesToChunks.AddDefaulted(1);
+	}
+
+	// Get all the parameter map histories traced to this graph from output nodes. We'll revisit this shortly in order to build out just the ones we care about for this translation.
+	check(Script->GetSource()->IsPreCompiled());
+	UNiagaraScriptSource* ScriptSource = Cast<UNiagaraScriptSource>(Script->GetSource());
+	OtherOutputParamMapHistories = ScriptSource->GetPrecomputedHistories();
+	for (FNiagaraParameterMapHistory& FoundHistory : OtherOutputParamMapHistories)
+	{
+		const UNiagaraNodeOutput* HistoryOutputNode = FoundHistory.GetFinalOutputNode();
+		if (HistoryOutputNode != nullptr && !ShouldConsiderTargetParameterMap(HistoryOutputNode->GetUsage()))
+		{
+			continue;
+		}
+		
+		// Now see if we want to use any of these specifically..
+		for (int32 ParamMapIdx = 0; ParamMapIdx < UsageOutputs.Num(); ParamMapIdx++)
+		{
+			UNiagaraNodeOutput* TargetOutputNode = UsageOutputs[ParamMapIdx];
+			if (FoundHistory.GetFinalOutputNode() == TargetOutputNode)
+			{
+				ParamMapHistories[ParamMapIdx] = (FoundHistory);
+				TArray<int32> Entries;
+				Entries.AddZeroed(FoundHistory.Variables.Num());
+				for (int32 i = 0; i < Entries.Num(); i++)
+				{
+					Entries[i] = INDEX_NONE;
+				}
+				ParamMapSetVariablesToChunks[ParamMapIdx] = (Entries);
+			}
+		}
+	}
+
+
+	//Create main scope pin cache.
+	PinToCodeChunks.AddDefaulted(1);
+
+	ActiveHistoryForFunctionCalls.BeginTranslation(InScript);
+
+	Script->StatScopes.Empty();
+	EnterStatsScope(FNiagaraStatScope(*Script->GetFullName(), FText::FromString(Script->GetFullName())));
+
+	Script->ParameterCollections.Empty();
+
+	FHlslNiagaraTranslator* ThisTranslator = this;
+	TArray<int32> OutputChunks;
+	if (Script->IsInterpolatedParticleSpawnScript())
+	{
+		//Here we compile the spawn script but write to temporary outputs in the context.
+		AddBodyChunk(TEXT(""), TEXT("//Interpolated Spawn Script!"), FNiagaraTypeDefinition::GetIntDef(), false, false);
+		AddBodyChunk(TEXT(""), TEXT("//Begin Spawn Script!"), FNiagaraTypeDefinition::GetIntDef(), false, false);
+		bInsideInterpolatedSpawnScript = true;
+		CurrentParamMapIndices.Empty();
+		CurrentParamMapIndices.Add(0);
+		OutputNode->Compile(ThisTranslator, OutputChunks);
+		bInsideInterpolatedSpawnScript = false;
+		InstanceWrite = FDataSetAccessInfo(); // Reset after building the output..
+		AddBodyChunk(TEXT(""), TEXT("//End Spawn Script!\n\n"), FNiagaraTypeDefinition::GetIntDef(), false, false);
+
+		AddBodyChunk(TEXT(""), TEXT("//Begin Transfer of Attributes!"), FNiagaraTypeDefinition::GetIntDef(), false, false);
+		AddBodyChunk(TEXT(""), TEXT("Context.MapUpdate.Particles = Context.MapSpawn.Particles"), FNiagaraTypeDefinition::GetIntDef(), false, true);
+		AddBodyChunk(TEXT(""), TEXT("//End Transfer of Attributes!\n\n"), FNiagaraTypeDefinition::GetIntDef(), false, false);
+		
+		AddBodyChunk(TEXT(""), TEXT("HandleMissingDefaultValues(Context);"), FNiagaraTypeDefinition::GetIntDef(), false, false);
+
+		AddBodyChunk(TEXT(""), TEXT("//Begin Update Script!"), FNiagaraTypeDefinition::GetIntDef(), false, false);
+		//Now we compile the update script (with partial dt) and read from the temp values written above.
+ 		UNiagaraNodeOutput* UpdateOutputNode = SourceGraph->FindOutputNode(ENiagaraScriptUsage::ParticleUpdateScript);
+ 		check(UpdateOutputNode);
+		CurrentParamMapIndices.Empty();
+		CurrentParamMapIndices.Add(1);
+		UpdateOutputNode->Compile(ThisTranslator, OutputChunks);
+
+		AddBodyChunk(TEXT(""), TEXT("//End Update Script!\n\n"), FNiagaraTypeDefinition::GetIntDef(), false, false);
+	}
+	else
+	{
+		CurrentParamMapIndices.Empty();
+		CurrentParamMapIndices.Add(0);
+		OutputNode->Compile(ThisTranslator, OutputChunks);
+
+		if (Script->IsParticleSpawnScript() || Script->IsSystemSpawnScript())
+		{
+			AddBodyChunk(TEXT(""), TEXT("HandleMissingDefaultValues(Context);"), FNiagaraTypeDefinition::GetIntDef(), false, false);
+		}
+	}
+	CurrentParamMapIndices.Empty();
+	ExitStatsScope();
+
+	ActiveHistoryForFunctionCalls.EndTranslation(Script);
+
+	TranslateResults.bHLSLGenSucceeded = MessageLog.NumErrors == 0;
+
+	//If we're compiling a function then we have all we need already, we don't want to actually generate shader/vm code.
+	if (FunctionCtx())
+		return TranslateResults;
+
+	//Now evaluate all the code chunks to generate the shader code.
+	//FString HlslOutput;
+	if (TranslateResults.bHLSLGenSucceeded)
+	{
+		//TODO: Declare all used structures up here too.
+		InScript->ReadDataSets.Empty();
+		InScript->WriteDataSets.Empty();
+
+		//Generate function definitions
+		FString FunctionDefinitionString = GetFunctionDefinitions();
+		FunctionDefinitionString += TEXT("\n");
+
+		if (Script->IsInterpolatedParticleSpawnScript())
+		{
+			//ensure the interpolated spawn constants are part of the parameter set.
+			int32 OutputIdx = 0;
+			ParameterMapRegisterExternalConstantNamespaceVariable(SYS_PARAM_ENGINE_DELTA_TIME, nullptr, 0, OutputIdx, nullptr);
+			ParameterMapRegisterExternalConstantNamespaceVariable(SYS_PARAM_ENGINE_INV_DELTA_TIME, nullptr, 0, OutputIdx, nullptr);
+			ParameterMapRegisterExternalConstantNamespaceVariable(SYS_PARAM_ENGINE_EXEC_COUNT, nullptr, 0, OutputIdx, nullptr);
+			ParameterMapRegisterExternalConstantNamespaceVariable(SYS_PARAM_EMITTER_SPAWNRATE, nullptr, 0, OutputIdx, nullptr);
+			ParameterMapRegisterExternalConstantNamespaceVariable(SYS_PARAM_EMITTER_SPAWN_INTERVAL, nullptr, 0, OutputIdx, nullptr);
+			ParameterMapRegisterExternalConstantNamespaceVariable(SYS_PARAM_EMITTER_INTERP_SPAWN_START_DT, nullptr, 0, OutputIdx, nullptr);
+		}
+
+		// Generate the Parameter Map HLSL definitions. We don't add to the final HLSL output here. We just build up the strings and tables
+		// that are needed later.
+		TArray<FNiagaraVariable> PrimaryDataSetOutputEntries;
+		TArray<TPair<FNiagaraVariable, const UEdGraphPin*> > MissingDefaultValues;
+		FString ParameterMapDefinitionStr = BuildParameterMapHlslDefinitions(PrimaryDataSetOutputEntries, MissingDefaultValues);
+
+		for (FNiagaraTypeDefinition Type : StructsToDefine)
+		{
+			HlslOutput += BuildHLSLStructDecl(Type);
+		}
+				
+		//Declare parameters.
+		//TODO: Separate Cbuffer for Global, System and Emitter parameters.
+		{
+			HlslOutput += TEXT("cbuffer FEmitterParameters\n{\n");
+
+			for (int32 i = 0; i < ChunksByMode[(int32)ENiagaraCodeChunkMode::Uniform].Num(); ++i)
+			{
+				FNiagaraCodeChunk& Chunk = CodeChunks[ChunksByMode[(int32)ENiagaraCodeChunkMode::Uniform][i]];
+				HlslOutput += TEXT("\t") + GetCode(ChunksByMode[(int32)ENiagaraCodeChunkMode::Uniform][i]);
+			}
+
+			if (Script->IsInterpolatedParticleSpawnScript())
+			{
+				//Define the params from the previous frame after the main parameters.
+				for (int32 i = 0; i < ChunksByMode[(int32)ENiagaraCodeChunkMode::Uniform].Num(); ++i)
+				{
+					//Copy the chunk so we can fiddle it's symbol name.
+					FNiagaraCodeChunk Chunk = CodeChunks[ChunksByMode[(int32)ENiagaraCodeChunkMode::Uniform][i]];
+					Chunk.SymbolName = TEXT("PREV__") + Chunk.SymbolName;
+					HlslOutput += TEXT("\t") + GetCode(Chunk);
+				}
+			}
+
+			HlslOutput += TEXT("}\n\n");
+		}
+
+		WriteDataSetStructDeclarations(DataSetReadInfo[0], true, HlslOutput);
+		WriteDataSetStructDeclarations(DataSetWriteInfo[0], false, HlslOutput);
+
+		//Map of all variables accessed by all datasets.
+		TMap<FNiagaraDataSetID, TArray<FNiagaraVariable>> DataSetReads;
+		TMap<FNiagaraDataSetID, TArray<FNiagaraVariable>> DataSetWrites;
+
+		TArray<TArray<FNiagaraVariable>* > DataSetReadVars;
+		TArray<TArray<FNiagaraVariable>* > DataSetWriteVars;
+		TArray<FNiagaraDataSetID> DataSetReadIds;
+		TArray<FNiagaraDataSetID> DataSetWriteIds;
+
+		TArray<FNiagaraVariable>& InstanceReadVars = DataSetReads.Add(GetInstanceDataSetID());
+		TArray<FNiagaraVariable>& InstanceWriteVars = DataSetWrites.Add(GetInstanceDataSetID());
+
+		DataSetReadVars.Add(&InstanceReadVars);
+		DataSetReadIds.Add(GetInstanceDataSetID());
+		DataSetWriteVars.Add(&InstanceWriteVars);
+		DataSetWriteIds.Add(GetInstanceDataSetID());
+				
+		if (IsBulkSystemScript())
+		{
+			TArray<FNiagaraVariable>& SystemEngineReadVars = DataSetReads.Add(GetSystemEngineDataSetID());
+			//TArray<FNiagaraVariable>& SystemUserReadVars = DataSetReads.Add(GetSystemUserDataSetID());
+
+			DataSetReadVars.Add(&SystemEngineReadVars);
+			//DataSetReadVars.Add(&SystemUserReadVars);
+			DataSetReadIds.Add(GetSystemEngineDataSetID());
+			//DataSetReadIds.Add(GetSystemUserDataSetID());
+
+			HandleNamespacedExternalVariablesToDataSetRead(SystemEngineReadVars, TEXT("Engine"));
+			HandleNamespacedExternalVariablesToDataSetRead(SystemEngineReadVars/*SystemUserReadVars*/, TEXT("User"));
+
+			// We sort the variables so that they end up in the same ordering between Spawn & Update...
+			SystemEngineReadVars.Sort([&](const FNiagaraVariable& A, const FNiagaraVariable& B)
+			{
+				return A.GetName() < B.GetName();
+			});
+			// We sort the variables so that they end up in the same ordering between Spawn & Update...
+			/*SystemUserReadVars.Sort([&](const FNiagaraVariable& A, const FNiagaraVariable& B)
+			{
+				return A.GetName() < B.GetName();
+			});*/
+
+			FNiagaraParameters ExternalParams;
+			ExternalParams.Parameters = SystemEngineReadVars;
+			Script->DataSetToParameters.Add(GetSystemEngineDataSetID().Name, ExternalParams);
+		}
+
+		// Now we pull in the HLSL generated above by building the parameter map definitions..
+		HlslOutput += ParameterMapDefinitionStr;
+
+		// Gather up all the unique Attribute variables that we generated.
+		TArray<FNiagaraVariable> BasicAttributes;
+		for (FNiagaraVariable& Var : InstanceRead.Variables)
+		{
+			if (Var.GetType().GetClass() != nullptr)
+			{
+				continue;
+			}
+			BasicAttributes.AddUnique(Var);
+		}
+		for (FNiagaraVariable& Var : InstanceWrite.Variables)
+		{
+			if (Var.GetType().GetClass() != nullptr)
+			{
+				continue;
+			}
+			else if (Var.GetType() != FNiagaraTypeDefinition::GetParameterMapDef())
+			{
+				BasicAttributes.AddUnique(Var);
+			}
+			else
+			{
+				for (FNiagaraVariable& ParamMapVar : PrimaryDataSetOutputEntries)
+				{
+					BasicAttributes.AddUnique(ParamMapVar);
+				}
+			}
+		}
+		
+		InstanceReadVars = BasicAttributes;
+		InstanceWriteVars = BasicAttributes;
+	
+
+		// We sort the variables so that they end up in the same ordering between Spawn & Update...
+		InstanceReadVars.Sort([&](const FNiagaraVariable& A, const FNiagaraVariable& B)
+		{
+			return A.GetName() < B.GetName();
+		});
+		// We sort the variables so that they end up in the same ordering between Spawn & Update...
+		InstanceWriteVars.Sort([&](const FNiagaraVariable& A, const FNiagaraVariable& B)
+		{
+			return A.GetName() < B.GetName();
+		});
+		//Define the simulation context. Which is a helper struct containing all the input, result and intermediate data needed for a single simulation.
+		//Allows us to reuse the same simulate function but provide different wrappers for final IO between GPU and CPU sims.
+		{
+			HlslOutput += TEXT("struct FSimulationContext") TEXT("\n{\n");
+
+			// We need to reserve a place in the simulation context for the base Parameter Map.
+			if (PrimaryDataSetOutputEntries.Num() != 0 || ParamMapDefinedSystemToNamespaceVars.Num() != 0 || ParamMapDefinedEmitterParameterToNamespaceVars.Num() != 0 || (ParamMapSetVariablesToChunks.Num() != 0 && ParamMapSetVariablesToChunks[0].Num() > 0))
+			{
+				if (Script->IsInterpolatedParticleSpawnScript())
+				{
+					HlslOutput += TEXT("\tFParamMap0 MapSpawn;\n");
+					HlslOutput += TEXT("\tFParamMap0 MapUpdate;\n");
+				}
+				else
+				{
+					HlslOutput += TEXT("\tFParamMap0 Map;\n");
+				}
+			}
+
+			WriteDataSetContextVars(DataSetReadInfo[0], true, HlslOutput);
+			WriteDataSetContextVars(DataSetWriteInfo[0], false, HlslOutput);
+
+
+			HlslOutput += TEXT("};\n\n");
+		}
+
+		HlslOutput += FunctionDefinitionString;
+
+		TArray<int32> WriteConditionVars;
+
+		// copy the accessed data sets over to the script, so we can grab them during sim
+		for (TPair <FNiagaraDataSetID, TMap<int32, FDataSetAccessInfo>> InfoPair : DataSetReadInfo[0])
+		{
+			Script->ReadDataSets.Add(InfoPair.Key);
+		}
+
+		for (TPair <FNiagaraDataSetID, TMap<int32, FDataSetAccessInfo>> InfoPair : DataSetWriteInfo[0])
+		{
+			FNiagaraDataSetProperties SetProps;
+			SetProps.ID = InfoPair.Key;
+			for (TPair <int32, FDataSetAccessInfo> IndexPair : InfoPair.Value)
+			{
+				SetProps.Variables = IndexPair.Value.Variables;
+			}
+
+			Script->WriteDataSets.Add(SetProps);
+
+			int32* ConditionalWriteChunkIdxPtr = DataSetWriteConditionalInfo[0].Find(InfoPair.Key);
+			if (ConditionalWriteChunkIdxPtr == nullptr)
+			{
+				WriteConditionVars.Add(INDEX_NONE);
+			}
+			else
+			{
+				WriteConditionVars.Add(*ConditionalWriteChunkIdxPtr);
+			}
+		}
+
+		DefineInterpolatedParametersFunction(HlslOutput);
+
+		if (Script->IsParticleSpawnScript() || Script->IsInterpolatedParticleSpawnScript() || Script->IsSystemSpawnScript())
+		{
+			HlslOutput += BuildMissingDefaults(MissingDefaultValues);
+		}
+
+		// define functions for reading and writing all secondary data sets
+		DefineDataSetReadFunction(HlslOutput, InScript->ReadDataSets);
+		DefineDataSetWriteFunction(HlslOutput, InScript->WriteDataSets, WriteConditionVars);
+
+
+		//Define the shared per instance simulation function	
+		{
+			HlslOutput += TEXT("void Simulate(inout FSimulationContext Context)\n{\n");
+			for (int32 i = 0; i < ChunksByMode[(int32)ENiagaraCodeChunkMode::Body].Num(); ++i)
+			{
+				HlslOutput += GetCode(ChunksByMode[(int32)ENiagaraCodeChunkMode::Body][i]);
+			}
+			HlslOutput += TEXT("}\n");
+		}
+
+		if (Options.SimTarget == ENiagaraSimTarget::GPUComputeSim)
+		{
+			FString DataInterfaceHLSL;
+			DefineDataInterfaceHLSL(DataInterfaceHLSL);
+			HlslOutput += DataInterfaceHLSL;
+		}
+
+		//And finally, define the actual main function that handles the reading and writing of data and calls the shared per instance simulate function.
+		//TODO: Different wrappers for GPU and CPU sims. 
+		DefineMain(HlslOutput, DataSetReadVars, DataSetReadIds, DataSetWriteVars, DataSetWriteIds);
+
+
+		//Get full list of instance data accessed by the script as the VM binding assumes same for input and output.
+		for (FNiagaraVariable& Var : InstanceReadVars)
+		{
+			if (FNiagaraParameterMapHistory::IsAttribute(Var))
+			{
+				FNiagaraVariable BasicAttribVar = FNiagaraParameterMapHistory::ResolveAsBasicAttribute(Var);
+				CompilationOutput.Attributes.AddUnique(BasicAttribVar);
+			}
+			else
+			{
+				CompilationOutput.Attributes.AddUnique(Var);
+			}
+		}
+		
+		CompilationOutput.bReadsAttributeData = InstanceRead.Variables.Num() != 0;
+		TranslateResults.OutputHLSL = HlslOutput;
+	}
+
+	return TranslateResults;
+}
+
+void FHlslNiagaraTranslator::GatherVariableForDataSetAccess(const FNiagaraVariable& Var, FString Format, int32& Counter, int32 DataSetIndex, FString InstanceIdxSymbol, FString &HlslOutputString)
+{
+	TArray<FString> Components;
+	const UScriptStruct* Struct = Var.GetType().GetScriptStruct();
+	check(Struct);
+
+	TArray<ENiagaraBaseTypes> Types;
+	GatherComponentsForDataSetAccess(Struct, TEXT(""), false, Components, Types);
+
+	//Add floats and then ints to hlsl
+	TArray<FStringFormatArg> FormatArgs;
+	FormatArgs.Empty(5);
+	FormatArgs.Add(TEXT(""));//We'll set the var name below.
+	FormatArgs.Add(TEXT(""));//We'll set the type name below.
+	// none for the output op (data set comes from acquireindex op)
+	if (DataSetIndex != INDEX_NONE)
+	{
+		FormatArgs.Add(DataSetIndex);
+	}
+	int32 RegIdx = FormatArgs.Add(0);
+	if (!InstanceIdxSymbol.IsEmpty())
+	{
+		FormatArgs.Add(InstanceIdxSymbol);
+	}
+	int32 DefaultIdx = FormatArgs.Add(0);
+
+	check(Components.Num() == Types.Num());
+	for (int32 CompIdx = 0; CompIdx < Components.Num(); ++CompIdx)
+	{
+		if (Types[CompIdx] == NBT_Float)
+		{
+			FormatArgs[1] = TEXT("Float");
+			FormatArgs[DefaultIdx] = TEXT("0.0f");
+		}
+		else if (Types[CompIdx] == NBT_Int32)
+		{
+			FormatArgs[1] = TEXT("Int");
+			FormatArgs[DefaultIdx] = TEXT("0");
+		}
+		else
+		{
+			check(Types[CompIdx] == NBT_Bool);
+			FormatArgs[1] = TEXT("Bool");
+			FormatArgs[DefaultIdx] = TEXT("false");
+		}
+		FormatArgs[0] = Components[CompIdx];
+		FormatArgs[RegIdx] = Counter++;
+		HlslOutputString += FString::Format(*Format, FormatArgs);
+	}
+}
+
+void FHlslNiagaraTranslator::GatherComponentsForDataSetAccess(const UScriptStruct* Struct, FString VariableSymbol, bool bMatrixRoot, TArray<FString>& Components, TArray<ENiagaraBaseTypes>& Types)
+{
+	bool bIsVector = IsHlslBuiltinVector(FNiagaraTypeDefinition(Struct));
+	bool bIsScalar = FNiagaraTypeDefinition::IsScalarDefinition(Struct);
+	bool bIsMatrix = FNiagaraTypeDefinition(Struct) == FNiagaraTypeDefinition::GetMatrix4Def();
+	if (bIsMatrix)
+	{
+		bMatrixRoot = true;
+	}
+	
+	//Bools are an awkward special case. TODO: make neater.
+	if (FNiagaraTypeDefinition(Struct) == FNiagaraTypeDefinition::GetBoolDef())
+	{
+		Types.Add(ENiagaraBaseTypes::NBT_Bool);
+		Components.Add(VariableSymbol);
+		return;
+	}
+
+	for (TFieldIterator<const UProperty> PropertyIt(Struct, EFieldIteratorFlags::IncludeSuper); PropertyIt; ++PropertyIt)
+	{
+		const UProperty* Property = *PropertyIt;
+
+		if (const UStructProperty* StructProp = Cast<const UStructProperty>(Property))
+		{
+			if (bMatrixRoot && FNiagaraTypeDefinition(StructProp->Struct) == FNiagaraTypeDefinition::GetFloatDef())
+			{
+				GatherComponentsForDataSetAccess(StructProp->Struct, VariableSymbol + ComputeMatrixColumnAccess(Property->GetName()), bMatrixRoot, Components, Types);
+			}
+			else if (bMatrixRoot &&  FNiagaraTypeDefinition(StructProp->Struct) == FNiagaraTypeDefinition::GetVec4Def())
+			{
+				GatherComponentsForDataSetAccess(StructProp->Struct, VariableSymbol + ComputeMatrixRowAccess(Property->GetName()), bMatrixRoot, Components, Types);
+			}
+			else
+			{
+				GatherComponentsForDataSetAccess(StructProp->Struct, VariableSymbol + TEXT(".") + Property->GetName(), bMatrixRoot, Components, Types);
+			}
+		}
+		else
+		{
+			FString VarName = VariableSymbol;
+			if (bMatrixRoot)
+			{
+				if (bIsVector && Property->IsA(UFloatProperty::StaticClass())) // Parent is a vector type, we are a float type
+				{
+					VarName += ComputeMatrixColumnAccess(Property->GetName());
+				}
+			}
+			else if (!bIsScalar)
+			{
+				VarName += TEXT(".");
+				VarName += bIsVector ? Property->GetName().ToLower() : Property->GetName();
+			}
+
+			if (Property->IsA(UFloatProperty::StaticClass()))
+			{
+				Types.Add(ENiagaraBaseTypes::NBT_Float);
+				Components.Add(VarName);
+			}
+			else if (Property->IsA(UIntProperty::StaticClass()))
+			{
+				Types.Add(ENiagaraBaseTypes::NBT_Int32);
+				Components.Add(VarName);
+			}
+			else if (Property->IsA(UBoolProperty::StaticClass()))
+			{
+				Types.Add(ENiagaraBaseTypes::NBT_Bool);
+				Components.Add(VarName);
+			}
+		}
+	}
+}
+
+void FHlslNiagaraTranslator::DefineInterpolatedParametersFunction(FString &HlslOutputString)
+{
+	// temporarily not doing this for GPU sim; interpolated spawn doesn't work there
+	if (Script->IsInterpolatedParticleSpawnScript())
+	{
+		FString Emitter_InterpSpawnStartDt = GetSanitizedSymbolName(ActiveHistoryForFunctionCalls.ResolveAliases(SYS_PARAM_EMITTER_INTERP_SPAWN_START_DT).GetName().ToString());
+		Emitter_InterpSpawnStartDt = Emitter_InterpSpawnStartDt.Replace(TEXT("."), TEXT("_"));//TODO: This should be rolled into GetSanitisedSymbolName but currently some usages require the '.' intact. Fix those up!.
+		FString Emitter_SpawnInterval = GetSanitizedSymbolName(ActiveHistoryForFunctionCalls.ResolveAliases(SYS_PARAM_EMITTER_SPAWN_INTERVAL).GetName().ToString());
+		Emitter_SpawnInterval = Emitter_SpawnInterval.Replace(TEXT("."), TEXT("_"));//TODO: This should be rolled into GetSanitisedSymbolName but currently some usages require the '.' intact. Fix those up!.
+		
+		HlslOutputString += TEXT("void InterpolateParameters(inout FSimulationContext Context)\n{\n");
+
+		// TODO: GPU interpolated spawning
+		if (CompilationTarget != ENiagaraSimTarget::GPUComputeSim)
+		{
+			HlslOutputString += TEXT("\tint InterpSpawn_Index = ExecIndex();\n");
+			HlslOutputString += TEXT("\tfloat InterpSpawn_SpawnTime = ") + Emitter_InterpSpawnStartDt + TEXT(" + (") + Emitter_SpawnInterval + TEXT(" * InterpSpawn_Index);\n");
+			HlslOutputString += TEXT("\tfloat InterpSpawn_UpdateTime = Engine_DeltaTime - InterpSpawn_SpawnTime;\n");
+			HlslOutputString += TEXT("\tfloat InterpSpawn_InvSpawnTime = 1.0 / InterpSpawn_SpawnTime;\n");
+			HlslOutputString += TEXT("\tfloat InterpSpawn_InvUpdateTime = 1.0 / InterpSpawn_UpdateTime;\n");
+			HlslOutputString += TEXT("\tfloat SpawnInterp = InterpSpawn_SpawnTime * Engine_InverseDeltaTime ;\n");
+			for (int32 i = 0; i < ChunksByMode[(int32)ENiagaraCodeChunkMode::Uniform].Num(); ++i)
+			{
+				int32 UniformIdx = ChunksByMode[(int32)ENiagaraCodeChunkMode::Uniform][i];
+				if (UniformIdx != INDEX_NONE)
+				{
+					FNiagaraVariable* FoundNamespacedVar = nullptr;
+					const FName* FoundSystemKey = ParamMapDefinedSystemVarsToUniformChunks.FindKey(UniformIdx);
+
+					// This uniform was either an emitter uniform parameter or a system uniform parameter. Search our maps to find out which one it was 
+					// so that we can properly deal with accessors.
+					if (FoundSystemKey != nullptr)
+					{
+						FoundNamespacedVar = ParamMapDefinedSystemToNamespaceVars.Find(*FoundSystemKey);
+					}
+					/*else
+					{
+						const FName* FoundEmitterParameterKey = ParamMapDefinedEmitterParameterVarsToUniformChunks.FindKey(UniformIdx);
+						if (FoundEmitterParameterKey != nullptr)
+						{
+							FoundNamespacedVar = ParamMapDefinedEmitterParameterToNamespaceVars.Find(*FoundEmitterParameterKey);
+						}
+					}*/
+
+					if (FoundNamespacedVar != nullptr)
+					{
+						FString FoundName = FoundNamespacedVar->GetName().ToString();
+						FNiagaraCodeChunk& Chunk = CodeChunks[UniformIdx];
+						HlslOutputString += TEXT("\tContext.MapSpawn.") + FoundName + TEXT(" = lerp(PREV__") + Chunk.SymbolName + TEXT(", ") + Chunk.SymbolName + TEXT(", SpawnInterp);\n");
+					}
+				}
+			}
+			HlslOutputString += TEXT("\tContext.MapSpawn.Engine.DeltaTime = 0.0f;\n");
+			HlslOutputString += TEXT("\tContext.MapSpawn.Engine.InverseDeltaTime = 0.0f;\n");
+			HlslOutputString += TEXT("\tContext.MapUpdate.Engine.DeltaTime = InterpSpawn_UpdateTime;\n");
+			HlslOutputString += TEXT("\tContext.MapUpdate.Engine.InverseDeltaTime = InterpSpawn_InvUpdateTime;\n");
+		}
+
+		HlslOutputString += TEXT("}\n\n");
+	}
+}
+
+void FHlslNiagaraTranslator::DefineDataSetReadFunction(FString &HlslOutputString, TArray<FNiagaraDataSetID> &ReadDataSets)
+{
+	if (Script->IsParticleEventScript() && CompilationTarget == ENiagaraSimTarget::GPUComputeSim)
+	{
+		HlslOutputString += TEXT("void ReadDataSets(inout FSimulationContext Context, int SetInstanceIndex)\n{\n");
+	}
+	else
+	{
+		HlslOutputString += TEXT("void ReadDataSets(inout FSimulationContext Context)\n{\n");
+	}
+
+	// We shouldn't read anything in a Spawn Script!
+	if (Script->IsParticleSpawnScript())
+	{
+		HlslOutputString += TEXT("}\n\n");
+		return;
+	}
+
+	for (TPair<FNiagaraDataSetID, TMap<int32, FDataSetAccessInfo>>& DataSetInfoPair : DataSetReadInfo[0])
+	{
+		FNiagaraDataSetID DataSet = DataSetInfoPair.Key;
+		int32 OffsetCounter = 0;
+		int32 DataSetIndex = 1;
+		for (TPair<int32, FDataSetAccessInfo>& IndexInfoPair : DataSetInfoPair.Value)
+		{
+			FString Symbol = FString("\tContext.") + DataSet.Name.ToString() + "Read.";  // TODO: HACK - need to get the real symbol name here
+			FString SetIdx = FString::FromInt(DataSetIndex);
+			FString DataSetComponentBufferSize = "DSComponentBufferSizeRead{1}" + SetIdx;
+			if (CompilationTarget == ENiagaraSimTarget::GPUComputeSim)
+			{
+				for (FNiagaraVariable &Var : IndexInfoPair.Value.Variables)
+				{
+					// TODO: temp = should really generate output functions for each set
+					FString Fmt = Symbol + Var.GetName().ToString() + FString(TEXT("{0} = ReadDataSet{1}")) + SetIdx + TEXT("[{2}*") + DataSetComponentBufferSize + " + SetInstanceIndex];\n";
+					GatherVariableForDataSetAccess(Var, Fmt, OffsetCounter, -1, TEXT(""), HlslOutputString);
+				}
+			}
+			else
+			{
+				for (FNiagaraVariable &Var : IndexInfoPair.Value.Variables)
+				{
+					// TODO: currently always emitting a non-advancing read, needs to be changed for some of the use cases
+					FString Fmt = TEXT("\tContext.") + DataSet.Name.ToString() + "Read." + Var.GetName().ToString() + TEXT("{0} = InputDataNoadvance{1}({2}, {3});\n");
+					GatherVariableForDataSetAccess(Var, Fmt, OffsetCounter, DataSetIndex, TEXT(""), HlslOutputString);
+				}
+			}
+		}
+	}
+
+	HlslOutputString += TEXT("}\n\n");
+}
+
+
+void FHlslNiagaraTranslator::DefineDataSetWriteFunction(FString &HlslOutputString, TArray<FNiagaraDataSetProperties> &WriteDataSets, TArray<int32>& WriteConditionVarIndices)
+{
+	HlslOutputString += TEXT("void WriteDataSets(inout FSimulationContext Context)\n{\n");
+
+	for (TPair<FNiagaraDataSetID, TMap<int32, FDataSetAccessInfo>>& DataSetInfoPair : DataSetWriteInfo[0])
+	{
+		FNiagaraDataSetID DataSet = DataSetInfoPair.Key;
+		int32 OffsetCounter = 0;
+		int32 DataSetIndex = 1;
+
+		HlslOutputString += "\tint TmpWriteIndex;\n";
+		int32* ConditionalWriteIdxPtr = DataSetWriteConditionalInfo[0].Find(DataSet);
+		if (nullptr == ConditionalWriteIdxPtr || INDEX_NONE == *ConditionalWriteIdxPtr)
+		{
+			HlslOutputString += "\tbool bValid = true;\n";
+		}
+		else
+		{
+			HlslOutputString += "\tbool bValid = " + FString("Context.") + DataSet.Name.ToString() + "Write_Valid;\n";
+		}
+		int32 WriteOffset = 0;
+		// grab the current ouput index; currently pass true, but should use an arbitrary bool to determine whether write should happen or not
+
+		HlslOutputString += "\tTmpWriteIndex = AcquireIndex(1, bValid);\n";
+
+		HlslOutputString += CompilationTarget==ENiagaraSimTarget::GPUComputeSim ? "\tif(TmpWriteIndex>=0)\n\t{\n" : "";
+
+		for (TPair<int32, FDataSetAccessInfo>& IndexInfoPair : DataSetInfoPair.Value)
+		{
+			FString Symbol = FString("Context.") + DataSet.Name.ToString() + "Write";  // TODO: HACK - need to get the real symbol name here
+			if (CompilationTarget == ENiagaraSimTarget::GPUComputeSim)
+			{
+				FString SetIdx = FString::FromInt(DataSetIndex);
+				FString DataSetComponentBufferSize = "DSComponentBufferSizeWrite{1}" + SetIdx;
+				for (FNiagaraVariable &Var : IndexInfoPair.Value.Variables)
+				{
+					// TODO: temp = should really generate output functions for each set
+					FString Fmt = FString(TEXT("\t\tRWWriteDataSet{1}")) + SetIdx + TEXT("[{2}*") + DataSetComponentBufferSize + TEXT(" + {3}] = ") + Symbol + TEXT(".") + Var.GetName().ToString() + TEXT("{0};\n");
+					GatherVariableForDataSetAccess(Var, Fmt, WriteOffset, -1, TEXT("TmpWriteIndex"), HlslOutputString);
+				}
+			}
+			else
+			{
+				for (FNiagaraVariable &Var : IndexInfoPair.Value.Variables)
+				{
+					// TODO: data set index is always 1; need to increase each set
+					FString Fmt = TEXT("\t\tOutputData{1}(1, {2}, {3}, ") + Symbol + "." + Var.GetName().ToString() + TEXT("{0});\n");
+					GatherVariableForDataSetAccess(Var, Fmt, WriteOffset, -1, TEXT("TmpWriteIndex"), HlslOutputString);
+				}
+			}
+		}
+
+		HlslOutputString += CompilationTarget == ENiagaraSimTarget::GPUComputeSim ? "\t}\n" : "";
+		DataSetIndex++;
+	}
+
+	HlslOutput += TEXT("}\n\n");
+}
+
+
+
+void FHlslNiagaraTranslator::DefineDataInterfaceHLSL(FString &InHlslOutput)
+{
+	FString InterfaceUniformHLSL;
+	FString InterfaceFunctionHLSL;
+	TArray<FString> BufferParamNames;
+	for (uint32 i = 0; i < 32; i++)
+	{
+		BufferParamNames.Add( TEXT("DataInterfaceBuffer_") + FString::FromInt(i) );
+	}
+
+	uint32 CurBufferIndex = 0;
+	for (int32 i = 0; i < CompilationOutput.DataInterfaceInfo.Num(); i++)
+	{
+		FNiagaraScriptDataInterfaceInfo &Info = CompilationOutput.DataInterfaceInfo[i];
+		
+		if (Info.DataInterface->CanExecuteOnTarget(ENiagaraSimTarget::GPUComputeSim))
+		{
+			TArray<FNiagaraFunctionSignature> DataInterfaceFunctions;
+			Info.DataInterface->GetFunctions(DataInterfaceFunctions);
+			FString OwnerIDString = Info.Name.ToString();
+
+			// grab the buffer definition from the interface
+			//
+			int32 NewIdx = DIBufferDescriptors.AddDefaulted(1);
+			Info.DataInterface->GetBufferDefinitionHLSL(OwnerIDString, DIBufferDescriptors[NewIdx], InterfaceUniformHLSL);
+
+			// grab the function hlsl from the interface
+			//
+			for (int32 FuncIdx = 0; FuncIdx < DataInterfaceFunctions.Num(); FuncIdx++)
+			{
+				FNiagaraFunctionSignature Sig = DataInterfaceFunctions[FuncIdx];	// make a copy so we can modify the owner id and get the correct hlsl signature
+				Sig.OwnerName = Info.Name;
+				FString DefStr = GetFunctionSignatureSymbol(Sig);
+
+				ensure(Info.DataInterface->GetFunctionHLSL(DefStr, DIBufferDescriptors[NewIdx], OwnerIDString, InterfaceFunctionHLSL) == true);
+			}
+		}
+		else
+		{
+			Error(FText::Format(LOCTEXT("NonGPUDataInterfaceError", "DataInterface {0} ({1}) cannot run on the GPU."), FText::FromName(Info.Name), FText::FromString(Info.DataInterface->GetClass()->GetName()))
+				, nullptr, nullptr);
+		}
+	}
+	InHlslOutput += InterfaceUniformHLSL + InterfaceFunctionHLSL;
+}
+
+
+void FHlslNiagaraTranslator::DefineMain(FString &OutHlslOutput,
+	TArray<TArray<FNiagaraVariable>*> &InstanceReadVars, TArray<FNiagaraDataSetID>& ReadIds,
+	TArray<TArray<FNiagaraVariable>*> &InstanceWriteVars, TArray<FNiagaraDataSetID>& WriteIds)
+{
+	if (CompilationTarget == ENiagaraSimTarget::GPUComputeSim)
+	{
+		OutHlslOutput += TEXT("void SimulateMain(in int InstanceIdx, in int InEventIndex)\n{\n\tFSimulationContext Context = (FSimulationContext)0;\n");
+	}
+	else
+	{
+		OutHlslOutput += TEXT("void SimulateMain()\n{\n\tFSimulationContext Context = (FSimulationContext)0;\n");
+	}
+
+	TMap<FName, int32> InputRegisterAllocations;
+	TMap<FName, int32> OutputRegisterAllocations;
+
+	ReadIdx = 0;
+	WriteIdx = 0;
+
+	//TODO: Grab indices for reading data sets and do the read.
+	//read input.
+
+	//The VM register binding assumes the same inputs as outputs which is obviously not always the case.
+	for (int32 VarArrayIdx = 0; VarArrayIdx < InstanceReadVars.Num(); VarArrayIdx++)
+	{
+		TArray<FNiagaraVariable>* ArrayRef = InstanceReadVars[VarArrayIdx];
+		DefineDataSetVariableReads(HlslOutput, ReadIds[VarArrayIdx], VarArrayIdx, *ArrayRef);
+	}
+
+	// Fill in the defaults for parameters.
+	for (int32 i = 0; i < MainPreSimulateChunks.Num(); ++i)
+	{
+		OutHlslOutput += TEXT("\t") + MainPreSimulateChunks[i] + TEXT("\n");
+	}
+	
+	// if we're in an event handler script on GPU, sim needs to loop over all events as well
+	// TODO: once looping works on CPU, we can do this for both targets
+	//
+	if (Script->IsParticleEventScript() && CompilationTarget == ENiagaraSimTarget::GPUComputeSim)
+	{
+		OutHlslOutput += TEXT("\tfor(int EventIdx=InEventIndex; EventIdx<InEventIndex+NumEventsPerParticle; EventIdx++)\n\t{\n");
+		OutHlslOutput += TEXT("\t\tReadDataSets(Context, EventIdx);\n");
+		OutHlslOutput += TEXT("\t\tSimulate(Context);\n");
+		OutHlslOutput += TEXT("\t}");
+	}
+	else
+	{
+		// call the read data set function
+		OutHlslOutput += TEXT("\tReadDataSets(Context);\n");
+
+		//Interpolate between prev and current parameters for interpolated spawning.
+		if (Script->IsInterpolatedParticleSpawnScript())
+		{
+			OutHlslOutput += TEXT("\tInterpolateParameters(Context);\n");
+		}
+
+		//Call simulate.
+		OutHlslOutput += TEXT("\tSimulate(Context);\n");
+	}
+
+	// write secondary data sets
+	OutHlslOutput += TEXT("\tWriteDataSets(Context);\n");
+
+	//The VM register binding assumes the same inputs as outputs which is obviously not always the case.
+	//We should separate inputs and outputs in the script.
+	for (int32 VarArrayIdx = 0; VarArrayIdx < InstanceWriteVars.Num(); VarArrayIdx++)
+	{
+		TArray<FNiagaraVariable>* ArrayRef = InstanceWriteVars[VarArrayIdx];
+		DefineDataSetVariableWrites(HlslOutput, WriteIds[VarArrayIdx], VarArrayIdx, *ArrayRef);
+	}
+
+	OutHlslOutput += TEXT("}\n");
+
+
+	// define a simple copy function to run on a section of the buffer for GPU event scripts; 
+	//	SimulateMainComputeCS in the usf will decide which one to call for which instances
+	// TODO: we'll want to combine spawn and update here soon so we'll need multiple main functions to be called from NiagaraEmitterInstanceShader.usf
+	//	this will need SpawnMain and UpdateMain for the regular particle scripts; all spawn events should be a single dispatch as well, each with its own function
+	//
+	if (CompilationTarget == ENiagaraSimTarget::GPUComputeSim)
+	{
+		OutHlslOutput += TEXT("void CopyInstance(in int InstanceIdx)\n{\n\tFSimulationContext Context = (FSimulationContext)0;\n");
+		for (int32 VarArrayIdx = 0; VarArrayIdx < InstanceReadVars.Num(); VarArrayIdx++)
+		{
+			TArray<FNiagaraVariable>* ArrayRef = InstanceReadVars[VarArrayIdx];
+			DefineDataSetVariableReads(HlslOutput, ReadIds[VarArrayIdx], VarArrayIdx, *ArrayRef);
+		}
+
+
+		if (Script->IsParticleEventScript())
+		{
+			OutHlslOutput += TEXT("\tContext.Map.DataInstance.Alive = true;\n");
+		}
+
+		for (int32 VarArrayIdx = 0; VarArrayIdx < InstanceWriteVars.Num(); VarArrayIdx++)
+		{
+			TArray<FNiagaraVariable>* ArrayRef = InstanceWriteVars[VarArrayIdx];
+			DefineDataSetVariableWrites(HlslOutput, WriteIds[VarArrayIdx], VarArrayIdx, *ArrayRef);
+		}
+		OutHlslOutput += TEXT("}\n");
+	}
+
+}
+
+void FHlslNiagaraTranslator::DefineDataSetVariableWrites(FString &OutHlslOutput, FNiagaraDataSetID& Id, int32 DataSetIndex, TArray<FNiagaraVariable>& WriteVars)
+{
+	//TODO Grab indices for data set writes (inc output) and do the write
+	OutHlslOutput += "\tint TmpWriteIndex;\n";
+	if (Script->IsNonParticleScript())
+	{
+		OutHlslOutput += "\tbool bValid = true;\n";
+	}
+	else
+	{
+		FString DataSetName = Id.Name.ToString();
+		bool bHasPerParticleAliveSpawn = false;
+		bool bHasPerParticleAliveUpdate = false;
+		bool bHasPerParticleAliveEvent = false;
+		for (int32 i = 0; i < ParamMapHistories.Num(); i++)
+		{
+			const UNiagaraNodeOutput* OutputNode = ParamMapHistories[i].GetFinalOutputNode();
+			bool bFound = (INDEX_NONE != ParamMapHistories[i].FindVariable(*(DataSetName + TEXT(".Alive")), FNiagaraTypeDefinition::GetBoolDef()));
+			if (bFound && OutputNode && (OutputNode->GetUsage() == ENiagaraScriptUsage::ParticleSpawnScript || OutputNode->GetUsage() == ENiagaraScriptUsage::ParticleSpawnScriptInterpolated))
+			{
+				bHasPerParticleAliveSpawn = true;
+			}
+			else if (bFound && OutputNode && OutputNode->GetUsage() == ENiagaraScriptUsage::ParticleUpdateScript)
+			{
+				bHasPerParticleAliveUpdate = true;
+			}
+			else if (bFound && OutputNode && OutputNode->GetUsage() == ENiagaraScriptUsage::ParticleEventScript)
+			{
+				bHasPerParticleAliveEvent = true;
+			}
+		}
+
+		if ((bHasPerParticleAliveSpawn || bHasPerParticleAliveUpdate) && Script->IsInterpolatedParticleSpawnScript())
+		{
+			if (bHasPerParticleAliveSpawn && bHasPerParticleAliveUpdate)
+			{
+				OutHlslOutput += TEXT("\tbool bValid = Context.MapUpdate.") + DataSetName + TEXT(".Alive && Context.MapSpawn.DataInstance.Alive;\n");
+			}
+			else if (bHasPerParticleAliveSpawn)
+			{
+				OutHlslOutput += TEXT("\tbool bValid = Context.MapSpawn.") + DataSetName + TEXT(".Alive;\n");
+			}
+			else if (bHasPerParticleAliveUpdate)
+			{
+				OutHlslOutput += TEXT("\tbool bValid = Context.MapUpdate.") + DataSetName + TEXT(".Alive;\n");
+			}
+		}
+		else if ((Script->IsParticleSpawnScript() && bHasPerParticleAliveSpawn) || (Script->IsParticleUpdateScript() && bHasPerParticleAliveUpdate) ||
+			(Script->IsParticleEventScript() && bHasPerParticleAliveEvent))
+		{
+			OutHlslOutput += TEXT("\tbool bValid = Context.Map.") + DataSetName + TEXT(".Alive;\n");
+		}
+		else
+		{
+			OutHlslOutput += "\tbool bValid = true;\n";
+		}
+	}
+	int32 WriteOffset = 0;
+
+	// grab the current ouput index; currently pass true, but should use an arbitrary bool to determine whether write should happen or not
+	OutHlslOutput += "\tTmpWriteIndex = AcquireIndex(0, bValid);\n";
+	for (FNiagaraVariable &Var : WriteVars)
+	{
+		// If coming from a parameter map, use the one on the context, otherwise use the output.
+		FString Fmt;
+
+		{
+			if (Script->IsInterpolatedParticleSpawnScript())
+			{
+				Fmt = TEXT("\tOutputData{1}(0, {2}, {3}, Context.MapUpdate.") + GetSanitizedSymbolName(Var.GetName().ToString()) + TEXT("{0});\n");
+			}
+			else
+			{
+				Fmt = TEXT("\tOutputData{1}(0, {2}, {3}, Context.Map.") + GetSanitizedSymbolName(Var.GetName().ToString()) + TEXT("{0});\n");
+			}
+			GatherVariableForDataSetAccess(Var, Fmt, WriteOffset, -1, TEXT("TmpWriteIndex"), OutHlslOutput);
+		}
+	}
+}
+
+void FHlslNiagaraTranslator::DefineDataSetVariableReads(FString &OutHlslOutput, FNiagaraDataSetID& Id, int32 DataSetIndex, TArray<FNiagaraVariable>& ReadVars)
+{
+	int32 ReadOffset = 0;
+	FString DataSetName = Id.Name.ToString();
+	
+	for (FNiagaraVariable &Var : ReadVars)
+	{
+		// If coming from a parameter map, use the one on the context, otherwise use the input.
+		FString Fmt;
+
+		{
+			// In the case of interpolated spawn scripts, we read into the MapSpawn and write out of the Map.
+			if (Script->IsInterpolatedParticleSpawnScript())
+			{
+				Fmt = TEXT("\tContext.MapSpawn.") + GetSanitizedSymbolName(Var.GetName().ToString()) + TEXT("{0} = {4};\n");
+			}
+			else if (Script->IsParticleSpawnScript() || Script->IsEmitterSpawnScript() || Script->IsSystemSpawnScript()) 	// We shouldn't read anything in a Spawn Script! Set to defaults.
+			{
+				Fmt = TEXT("\tContext.Map.") + GetSanitizedSymbolName(Var.GetName().ToString()) + TEXT("{0} = {4};\n");
+			}
+			else
+			{
+				if (CompilationTarget == ENiagaraSimTarget::GPUComputeSim)
+				{
+					Fmt = TEXT("\tContext.Map.") + GetSanitizedSymbolName(Var.GetName().ToString()) + TEXT("{0} = InputData{1}({2}, {3}, InstanceIdx);\n");
+				}
+				else
+				{
+					Fmt = TEXT("\tContext.Map.") + GetSanitizedSymbolName(Var.GetName().ToString()) + TEXT("{0} = InputData{1}({2}, {3});\n");
+				}
+			}
+		}
+		GatherVariableForDataSetAccess(Var, Fmt, ReadOffset, DataSetIndex, TEXT(""), OutHlslOutput);
+	}
+}
+
+void FHlslNiagaraTranslator::WriteDataSetContextVars(TMap<FNiagaraDataSetID, TMap<int32, FDataSetAccessInfo>>& DataSetAccessInfo, bool bRead, FString &OutHlslOutput)
+{
+	//Now the intermediate storage for the data set reads and writes.
+	uint32 DataSetIndex = 0;
+	for (TPair<FNiagaraDataSetID, TMap<int32, FDataSetAccessInfo>>& DataSetInfoPair : DataSetAccessInfo)
+	{
+		FNiagaraDataSetID DataSet = DataSetInfoPair.Key;
+
+		if (!bRead)
+		{
+			OutHlslOutput += TEXT("\tbool ") + DataSet.Name.ToString() + TEXT("Write_Valid; \n");
+		}
+
+		OutHlslOutput += TEXT("\tF") + DataSet.Name.ToString() + "DataSet " + DataSet.Name.ToString() + (bRead ? TEXT("Read") : TEXT("Write")) + TEXT(";\n");
+	}
+};
+
+
+void FHlslNiagaraTranslator::WriteDataSetStructDeclarations(TMap<FNiagaraDataSetID, TMap<int32, FDataSetAccessInfo>>& DataSetAccessInfo, bool bRead, FString &OutHlslOutput)
+{
+	uint32 DataSetIndex = 1;
+	for (TPair<FNiagaraDataSetID, TMap<int32, FDataSetAccessInfo>>& DataSetInfoPair : DataSetAccessInfo)
+	{
+		FNiagaraDataSetID DataSet = DataSetInfoPair.Key;
+		FString StructName = TEXT("F") + DataSet.Name.ToString() + "DataSet";
+		OutHlslOutput += TEXT("struct ") + StructName + TEXT("\n{\n");
+
+		for (TPair<int32, FDataSetAccessInfo>& IndexInfoPair : DataSetInfoPair.Value)
+		{
+			for (FNiagaraVariable Var : IndexInfoPair.Value.Variables)
+			{
+				OutHlslOutput += TEXT("\t") + GetStructHlslTypeName(Var.GetType().GetScriptStruct()) + TEXT(" ") + Var.GetName().ToString() + ";\n";
+			}
+		}
+
+		OutHlslOutput += TEXT("};\n");
+
+		// declare buffers for compute shader HLSL only; VM doesn't need htem
+		// because its InputData and OutputData functions handle data set management explicitly
+		//
+		if (CompilationTarget == ENiagaraSimTarget::GPUComputeSim)
+		{
+			FString IndexString = FString::FromInt(DataSetIndex);
+			if (bRead)
+			{
+				OutHlslOutput += FString(TEXT("Buffer<float> ReadDataSetFloat")) + IndexString + ";\n";
+				OutHlslOutput += FString(TEXT("Buffer<int> ReadDataSetInt")) + IndexString + ";\n";
+				OutHlslOutput += "int DSComponentBufferSizeReadFloat" + IndexString + ";\n";
+				OutHlslOutput += "int DSComponentBufferSizeReadInt" + IndexString + ";\n";
+			}
+			else 
+			{
+				OutHlslOutput += FString(TEXT("RWBuffer<float> RWWriteDataSetFloat")) + IndexString + ";\n";
+				OutHlslOutput += FString(TEXT("RWBuffer<int> RWWriteDataSetInt")) + IndexString + ";\n";
+				OutHlslOutput += "int DSComponentBufferSizeWriteFloat" + IndexString + ";\n";
+				OutHlslOutput += "int DSComponentBufferSizeWriteInt" + IndexString + ";\n";
+			}
+		}
+
+		DataSetIndex++;
+	}
+
+}
+
+
+//Decomposes each variable into its constituent register accesses.
+void FHlslNiagaraTranslator::DecomposeVariableAccess(const UStruct* Struct, bool bRead, FString IndexSymbol, FString HLSLString)
+{
+	FString AccessStr;
+
+	for (TFieldIterator<UProperty> PropertyIt(Struct, EFieldIteratorFlags::IncludeSuper); PropertyIt; ++PropertyIt)
+	{
+		UProperty* Property = *PropertyIt;
+
+		if (UStructProperty* StructProp = CastChecked<UStructProperty>(Property))
+		{
+			FNiagaraTypeDefinition PropDef(StructProp->Struct);
+			if (!IsHlslBuiltinVector(PropDef))
+			{
+				DecomposeVariableAccess(StructProp->Struct, bRead, IndexSymbol, AccessStr);
+				return;
+			}
+		}
+
+		int32 Index = INDEX_NONE;
+		if (bRead)
+		{
+			Index = ReadIdx++;
+			AccessStr = TEXT("ReadInput(");
+			AccessStr += FString::FromInt(ReadIdx);
+			AccessStr += ");\n";
+		}
+		else
+		{
+			Index = WriteIdx++;
+			AccessStr = TEXT("WriteOutput(");
+			AccessStr += FString::FromInt(WriteIdx);
+			AccessStr += ");\n";
+		}
+
+		HLSLString += AccessStr;
+
+		FNiagaraTypeDefinition StructDef(Cast<UScriptStruct>(Struct));
+		FString TypeName = GetStructHlslTypeName(StructDef);
+	}
+};
+
+FString FHlslNiagaraTranslator::GetSanitizedSymbolName(FString SymbolName)
+{
+	FString Ret = SymbolName.Replace(TEXT(" "), TEXT(""));
+	Ret = Ret.Replace(TEXT("\\"), TEXT("_"));
+	Ret = Ret.Replace(TEXT("/"), TEXT("_"));
+	Ret = Ret.Replace(TEXT(","), TEXT("_"));
+	Ret = Ret.Replace(TEXT("-"), TEXT("_"));
+	Ret = Ret.Replace(TEXT(":"), TEXT("_"));
+	return Ret;
+}
+
+FString FHlslNiagaraTranslator::GetUniqueSymbolName(FName BaseName)
+{
+	uint32* NameCount = SymbolCounts.Find(BaseName);
+	if (NameCount == nullptr)
+	{
+		SymbolCounts.Add(BaseName) = 1;
+		return GetSanitizedSymbolName(BaseName.ToString());
+	}
+
+	FString Ret = GetSanitizedSymbolName(BaseName.ToString());
+	if(*NameCount > 0)
+	{
+		Ret += LexicalConversion::ToString(*NameCount);
+	}
+	++(*NameCount);
+	return Ret;
+}
+
+void FHlslNiagaraTranslator::EnterFunction(const FString& Name, FNiagaraFunctionSignature& Signature, TArray<int32>& Inputs)
+{
+	FunctionContextStack.Emplace(Name, Signature, Inputs);
+	//May need some more heavy and scoped symbol tracking?
+
+	//Add new scope for pin reuse.
+	PinToCodeChunks.AddDefaulted(1);
+}
+
+void FHlslNiagaraTranslator::ExitFunction()
+{
+	FunctionContextStack.Pop();
+	//May need some more heavy and scoped symbol tracking?
+
+	//Pop pin reuse scope.
+	PinToCodeChunks.Pop();
+}
+
+FString FHlslNiagaraTranslator::GeneratedConstantString(float Constant)
+{
+	return LexicalConversion::ToString(Constant);
+}
+
+static int32 GbNiagaraScriptStatTracking = 1;
+static FAutoConsoleVariableRef CVarNiagaraScriptStatTracking(
+	TEXT("fx.NiagaraScriptStatTracking"),
+	GbNiagaraScriptStatTracking,
+	TEXT("If > 0 stats tracking operations will be compiled into Niagara Scripts. \n"),
+	ECVF_Default
+);
+
+void FHlslNiagaraTranslator::EnterStatsScope(FNiagaraStatScope StatScope)
+{
+	if (GbNiagaraScriptStatTracking)
+	{
+		int32 ScopeIdx = Script->StatScopes.AddUnique(StatScope);
+		AddBodyChunk(TEXT(""), FString::Printf(TEXT("EnterStatScope(%d /**%s*/)"), ScopeIdx, *StatScope.FullName.ToString()), FNiagaraTypeDefinition::GetFloatDef(), false);
+		StatScopeStack.Push(ScopeIdx);
+	}
+}
+
+void FHlslNiagaraTranslator::ExitStatsScope()
+{
+	if (GbNiagaraScriptStatTracking)
+	{
+		int32 ScopeIdx = StatScopeStack.Pop();
+		AddBodyChunk(TEXT(""), FString::Printf(TEXT("ExitStatScope(/**%s*/)"), *Script->StatScopes[ScopeIdx].FullName.ToString()), FNiagaraTypeDefinition::GetFloatDef(), false);
+	}
+}
+
+FString FHlslNiagaraTranslator::GetCallstack()
+{
+	FString Callstack = Script->GetName();
+
+	for (FFunctionContext& Ctx : FunctionContextStack)
+	{
+		Callstack += TEXT(".") + Ctx.Name;
+	}
+
+	return Callstack;
+}
+
+FString FHlslNiagaraTranslator::GeneratedConstantString(FVector4 Constant)
+{
+	TArray<FStringFormatArg> Args;
+	Args.Add(LexicalConversion::ToString(Constant.X));
+	Args.Add(LexicalConversion::ToString(Constant.Y));
+	Args.Add(LexicalConversion::ToString(Constant.Z));
+	Args.Add(LexicalConversion::ToString(Constant.W));
+	return FString::Format(TEXT("float4({0}, {1}, {2}, {3})"), Args);
+}
+
+int32 FHlslNiagaraTranslator::AddUniformChunk(FString SymbolName, const FNiagaraTypeDefinition& Type)
+{
+	int32 Ret = CodeChunks.IndexOfByPredicate(
+		[&](const FNiagaraCodeChunk& Chunk)
+	{
+		return Chunk.Mode == ENiagaraCodeChunkMode::Uniform && Chunk.SymbolName == SymbolName && Chunk.Type == Type;
+	}
+	);
+
+	if (Ret == INDEX_NONE)
+	{
+		Ret = CodeChunks.AddDefaulted();
+		FNiagaraCodeChunk& Chunk = CodeChunks[Ret];
+		Chunk.SymbolName = GetSanitizedSymbolName(SymbolName);
+		Chunk.Type = Type;
+		Chunk.Mode = ENiagaraCodeChunkMode::Uniform;
+
+		ChunksByMode[(int32)ENiagaraCodeChunkMode::Uniform].Add(Ret);
+	}
+	return Ret;
+}
+
+int32 FHlslNiagaraTranslator::AddSourceChunk(FString SymbolName, const FNiagaraTypeDefinition& Type)
+{
+	int32 Ret = CodeChunks.IndexOfByPredicate(
+		[&](const FNiagaraCodeChunk& Chunk)
+	{
+		return Chunk.Mode == ENiagaraCodeChunkMode::Source && Chunk.SymbolName == SymbolName && Chunk.Type == Type;
+	}
+	);
+
+	if (Ret == INDEX_NONE)
+	{
+		Ret = CodeChunks.AddDefaulted();
+		FNiagaraCodeChunk& Chunk = CodeChunks[Ret];
+		Chunk.SymbolName = GetSanitizedSymbolName(SymbolName);
+		Chunk.Type = Type;
+		Chunk.Mode = ENiagaraCodeChunkMode::Source;
+
+		ChunksByMode[(int32)ENiagaraCodeChunkMode::Source].Add(Ret);
+	}
+	return Ret;
+}
+
+int32 FHlslNiagaraTranslator::AddBodyChunk(FString SymbolName, FString Definition, const FNiagaraTypeDefinition& Type, TArray<int32>& SourceChunks, bool bDecl, bool bIsTerminated)
+{
+	int32 Ret = CodeChunks.AddDefaulted();
+	FNiagaraCodeChunk& Chunk = CodeChunks[Ret];
+	Chunk.SymbolName = GetSanitizedSymbolName(SymbolName);
+	Chunk.Definition = Definition;
+	Chunk.Type = Type;
+	Chunk.bDecl = bDecl;
+	Chunk.bIsTerminated = bIsTerminated;
+	Chunk.Mode = ENiagaraCodeChunkMode::Body;
+	Chunk.SourceChunks = SourceChunks;
+
+	ChunksByMode[(int32)ENiagaraCodeChunkMode::Body].Add(Ret);
+	return Ret;
+}
+
+int32 FHlslNiagaraTranslator::AddBodyChunk(FString SymbolName, FString Definition, const FNiagaraTypeDefinition& Type, int32 SourceChunk, bool bDecl, bool bIsTerminated)
+{
+	int32 Ret = CodeChunks.AddDefaulted();
+	FNiagaraCodeChunk& Chunk = CodeChunks[Ret];
+	Chunk.SymbolName = GetSanitizedSymbolName(SymbolName);
+	Chunk.Definition = Definition;
+	Chunk.Type = Type;
+	Chunk.bDecl = bDecl;
+	Chunk.bIsTerminated = bIsTerminated;
+	Chunk.Mode = ENiagaraCodeChunkMode::Body;
+	Chunk.SourceChunks.Add(SourceChunk);
+
+	ChunksByMode[(int32)ENiagaraCodeChunkMode::Body].Add(Ret);
+	return Ret;
+}
+
+int32 FHlslNiagaraTranslator::AddBodyChunk(FString SymbolName, FString Definition, const FNiagaraTypeDefinition& Type, bool bDecl, bool bIsTerminated)
+{
+	int32 Ret = CodeChunks.AddDefaulted();
+	FNiagaraCodeChunk& Chunk = CodeChunks[Ret];
+	Chunk.SymbolName = GetSanitizedSymbolName(SymbolName);
+	Chunk.Definition = Definition;
+	Chunk.Type = Type;
+	Chunk.bDecl = bDecl;
+	Chunk.bIsTerminated = bIsTerminated;
+	Chunk.Mode = ENiagaraCodeChunkMode::Body;
+
+	ChunksByMode[(int32)ENiagaraCodeChunkMode::Body].Add(Ret);
+	return Ret;
+}
+
+bool FHlslNiagaraTranslator::ShouldInterpolateParameter(const FNiagaraVariable& Parameter)
+{
+	//Skip interpolation for some system constants.
+	if (Parameter == SYS_PARAM_ENGINE_DELTA_TIME ||
+		Parameter == SYS_PARAM_ENGINE_INV_DELTA_TIME ||
+		Parameter == SYS_PARAM_ENGINE_EXEC_COUNT || 
+		Parameter == SYS_PARAM_EMITTER_SPAWNRATE ||
+		Parameter == SYS_PARAM_EMITTER_SPAWN_INTERVAL ||
+		Parameter == SYS_PARAM_EMITTER_INTERP_SPAWN_START_DT )
+	{
+		return false;
+	}
+
+	return true;
+}
+
+int32 FHlslNiagaraTranslator::GetParameter(const FNiagaraVariable& Parameter)
+{
+	if (!AddStructToDefinitionSet(Parameter.GetType()))
+	{
+		Error(FText::Format(LOCTEXT("GetParameterFail", "Cannot handle type {0}! Variable: {1}"), Parameter.GetType().GetNameText(), FText::FromName(Parameter.GetName())), nullptr, nullptr);
+	}
+
+	int32 FuncParam = INDEX_NONE;
+	if (GetFunctionParameter(Parameter, FuncParam))
+	{
+		if (FuncParam != INDEX_NONE)
+		{
+			if (Parameter.GetType() == FNiagaraTypeDefinition::GetParameterMapDef())
+			{
+				return FuncParam;
+			}
+			//If this is a valid function parameter, use that.
+			FString SymbolName = TEXT("In_") + GetSanitizedSymbolName(Parameter.GetName().ToString());
+			return AddSourceChunk(SymbolName, Parameter.GetType());
+		}
+	}
+
+	// We don't pass in the input node here (really there could be multiple nodes for the same parameter)
+	// so we have to match up the input parameter map variable value through the pre-traversal histories 
+	// so that we know which parameter map we are referencing.
+	FString SymbolName = GetSanitizedSymbolName(Parameter.GetName().ToString());
+	if (Parameter.GetType() == FNiagaraTypeDefinition::GetParameterMapDef())
+	{
+		if (ParamMapHistories.Num() == 0)
+		{
+			return INDEX_NONE;
+		}
+
+		for (int32 i = 0; i < ParamMapHistories.Num(); i++)
+		{
+			// Double-check against the current output node we are tracing. Ignore any parameter maps
+			// that don't include that node.
+			if (CurrentParamMapIndices.Num() != 0 && !CurrentParamMapIndices.Contains(i))
+			{
+				continue;
+			}
+
+			for (int32 PinIdx = 0; PinIdx < ParamMapHistories[i].MapPinHistory.Num(); PinIdx++)
+			{
+				const UEdGraphPin* Pin = ParamMapHistories[i].MapPinHistory[PinIdx];
+
+				if (Pin != nullptr)
+				{
+					UNiagaraNodeInput* InputNode = Cast<UNiagaraNodeInput>(Pin->GetOwningNode());
+					if (InputNode != nullptr && InputNode->Input == Parameter)
+					{
+						return i;
+					}
+				}
+			}
+		}
+		return INDEX_NONE;
+	}
+
+	//Not a in a function or not a valid function parameter so grab from the main uniforms.
+	int32 OutputChunkIdx = INDEX_NONE;
+	FNiagaraVariable OutputVariable = Parameter;
+	if (FNiagaraParameterMapHistory::IsExternalConstantNamespace(OutputVariable, Script))
+	{
+		if (!ParameterMapRegisterExternalConstantNamespaceVariable(OutputVariable, nullptr, 0, OutputChunkIdx, nullptr))
+		{
+			OutputChunkIdx = INDEX_NONE;
+		}
+	}
+	else
+	{
+		OutputVariable = FNiagaraParameterMapHistory::MoveToExternalConstantNamespaceVariable(OutputVariable, Script);
+		if (!ParameterMapRegisterExternalConstantNamespaceVariable(OutputVariable, nullptr, 0, OutputChunkIdx, nullptr))
+		{
+			OutputChunkIdx = INDEX_NONE;
+		}
+	}
+
+	if (OutputChunkIdx == INDEX_NONE)
+	{
+		Error(FText::Format(LOCTEXT("GetParameterFail", "Cannot handle type {0}! Variable: {1}"), Parameter.GetType().GetNameText(), FText::FromName(Parameter.GetName())), nullptr, nullptr);
+	}
+
+	return OutputChunkIdx;
+}
+
+int32 FHlslNiagaraTranslator::GetConstant(const FNiagaraVariable& Constant)
+{
+	FString ConstantStr = GenerateConstantString(Constant);
+	if (ConstantStr.IsEmpty())
+	{
+		return INDEX_NONE;
+	}
+	return AddBodyChunk(GetUniqueSymbolName(TEXT("Constant")), ConstantStr, Constant.GetType());
+}
+
+int32 FHlslNiagaraTranslator::GetConstantFloat(float InConstantValue)
+{
+	FNiagaraVariable Constant(FNiagaraTypeDefinition::GetFloatDef(), TEXT("Constant"));
+	Constant.SetValue(InConstantValue);
+
+	return GetConstant(Constant);
+}
+
+FString FHlslNiagaraTranslator::GenerateConstantString(const FNiagaraVariable& Constant)
+{
+	FNiagaraTypeDefinition Type = Constant.GetType();
+	if (!AddStructToDefinitionSet(Type))
+	{
+		Error(FText::Format(LOCTEXT("GetConstantFail", "Cannot handle type {0}! Variable: {1}"), Type.GetNameText(), FText::FromName(Constant.GetName())), nullptr, nullptr);
+	}
+	FString ConstantStr = GetStructHlslTypeName(Type) + GetHlslDefaultForType(Type);
+	if (Constant.IsDataAllocated())
+	{
+		if (Type == FNiagaraTypeDefinition::GetFloatDef())
+		{
+			float* ValuePtr = (float*)Constant.GetData();
+			ConstantStr = FString::Printf(TEXT("%g"), *ValuePtr);
+		}
+		else if (Type == FNiagaraTypeDefinition::GetVec2Def())
+		{
+			float* ValuePtr = (float*)Constant.GetData();
+			ConstantStr = FString::Printf(TEXT("float2(%g,%g)"), *ValuePtr, *(ValuePtr + 1));
+		}
+		else if (Type == FNiagaraTypeDefinition::GetVec3Def())
+		{
+			float* ValuePtr = (float*)Constant.GetData();
+			ConstantStr = FString::Printf(TEXT("float3(%g,%g,%g)"), *ValuePtr, *(ValuePtr + 1), *(ValuePtr + 2));
+		}
+		else if (Type == FNiagaraTypeDefinition::GetVec4Def())
+		{
+			float* ValuePtr = (float*)Constant.GetData();
+			ConstantStr = FString::Printf(TEXT("float4(%g,%g,%g,%g)"), *ValuePtr, *(ValuePtr + 1), *(ValuePtr + 2), *(ValuePtr + 3));
+		}
+		else if (Type == FNiagaraTypeDefinition::GetColorDef())
+		{
+			float* ValuePtr = (float*)Constant.GetData();
+			ConstantStr = FString::Printf(TEXT("float4(%g,%g,%g,%g)"), *ValuePtr, *(ValuePtr + 1), *(ValuePtr + 2), *(ValuePtr + 3));
+		}
+		else if (Type == FNiagaraTypeDefinition::GetIntDef() || Type.GetStruct() == FNiagaraTypeDefinition::GetIntStruct())
+		{
+			int32* ValuePtr = (int32*)Constant.GetData();
+			ConstantStr = FString::Printf(TEXT("%d"), *ValuePtr);
+		}
+		else if (Type == FNiagaraTypeDefinition::GetBoolDef())
+		{
+			if (Constant.GetValue<FNiagaraBool>()->IsValid() == false)
+			{
+				Error(FText::Format(LOCTEXT("StructContantsBoolInvalidError", "Boolean constant {0} is not set to explicit True or False. Defaulting to False."), FText::FromName(Constant.GetName())), nullptr, nullptr);
+				ConstantStr = TEXT("false");
+			}
+			else
+			{
+				bool bValue = Constant.GetValue<FNiagaraBool>()->GetValue();
+				ConstantStr = bValue ? TEXT("true") : TEXT("false");
+			}
+		}
+		else
+		{
+			//This is easily doable, just need to keep track of all structs used and define them as well as a ctor function signature with all values decomposed into float1/2/3/4 etc
+			//Then call said function here with the same decomposition literal values.
+			Error(LOCTEXT("StructContantsUnsupportedError", "Constants of struct types are currently unsupported."), nullptr, nullptr);
+			return FString();
+		}
+	}
+	return ConstantStr;
+}
+
+void FHlslNiagaraTranslator::Output(UNiagaraNodeOutput* OutputNode, const TArray<int32>& ComputedInputs)
+{
+	TArray<FNiagaraVariable> Attributes;
+	TArray<int32> Inputs;
+
+	// Build up the attribute list. We don't auto-expand parameter maps here.
+	TArray<FNiagaraVariable> Outputs = OutputNode->GetOutputs();
+	check(ComputedInputs.Num() == Outputs.Num());
+	for (int32 PinIdx = 0; PinIdx < Outputs.Num(); PinIdx++)
+	{
+		Attributes.Add(Outputs[PinIdx]);
+		Inputs.Add(ComputedInputs[PinIdx]);
+	}
+
+	if (FunctionCtx())
+	{
+		for (int32 i = 0; i < Attributes.Num(); ++i)
+		{
+			if (!AddStructToDefinitionSet(Attributes[i].GetType()))
+			{
+				Error(FText::Format(LOCTEXT("GetConstantFail", "Cannot handle type {0}! Variable: {1}"), Attributes[i].GetType().GetNameText(), FText::FromName(Attributes[i].GetName())), nullptr, nullptr);
+			}
+
+			if (Attributes[i].GetType() != FNiagaraTypeDefinition::GetParameterMapDef())
+			{
+				FString SymbolName = *(TEXT("Out_") + GetSanitizedSymbolName(Attributes[i].GetName().ToString()));
+				AddBodyChunk(SymbolName, TEXT("{0}"), Attributes[i].GetType(), Inputs[i], false);
+			}
+		}
+	}
+	else
+	{
+		
+		{
+			check(InstanceWrite.CodeChunks.Num() == 0);//Should only hit one output node.
+
+			FString DataSetAccessName = GetDataSetAccessSymbol(GetInstanceDataSetID(), INDEX_NONE, false);
+			//First chunk for a write is always the condition pin.
+			for (int32 i = 0; i < Attributes.Num(); ++i)
+			{
+				const FNiagaraVariable& Var = Attributes[i];
+
+				if (!AddStructToDefinitionSet(Var.GetType()))
+				{
+					Error(FText::Format(LOCTEXT("GetConstantFail", "Cannot handle type {0}! Variable: {1}"), Var.GetType().GetNameText(), FText::FromName(Var.GetName())), nullptr, nullptr);
+				}
+
+				//DATASET TODO: add and treat input 0 as the 'valid' input for conditional write
+				int32 Input = Inputs[i];
+
+
+				if (Var.GetType() != FNiagaraTypeDefinition::GetParameterMapDef())
+				{
+					FNiagaraVariable VarNamespaced = FNiagaraParameterMapHistory::BasicAttributeToNamespacedAttribute(Var);
+					FString ParameterMapInstanceName = GetParameterMapInstanceName(0);
+					int32 ChunkIdx = AddBodyChunk(ParameterMapInstanceName + TEXT(".") + GetSanitizedSymbolName(VarNamespaced.GetName().ToString()), TEXT("{0}"), VarNamespaced.GetType(), Input, false);
+
+					// Make sure that we end up in the list of Attributes that have been written to by this script.
+					if (ParamMapDefinedAttributesToUniformChunks.Find(Var.GetName()) == nullptr)
+					{
+						ParamMapDefinedAttributesToUniformChunks.Add(Var.GetName(), Input);
+						ParamMapDefinedAttributesToNamespaceVars.Add(Var.GetName(), VarNamespaced);
+					}
+
+					InstanceWrite.Variables.AddUnique(VarNamespaced);
+					InstanceWrite.CodeChunks.Add(ChunkIdx);
+				}
+				else
+				{
+					InstanceWrite.Variables.AddUnique(Var);
+				}
+			}
+		}
+	}
+}
+
+int32 FHlslNiagaraTranslator::GetAttribute(const FNiagaraVariable& Attribute)
+{	
+	if (!AddStructToDefinitionSet(Attribute.GetType()))
+	{
+		Error(FText::Format(LOCTEXT("GetConstantFail", "Cannot handle type {0}! Variable: {1}"), Attribute.GetType().GetNameText(), FText::FromName(Attribute.GetName())), nullptr, nullptr);
+	}
+
+	if (Script->IsParticleSpawnScript())
+	{
+		if (Script->IsInterpolatedParticleSpawnScript() && !bInsideInterpolatedSpawnScript)
+		{
+			//This is a special case where we allow the grabbing of attributes in the update section of an interpolated spawn script.
+			//But we return the results of the previously ran spawn script.
+			FString ParameterMapInstanceName = GetParameterMapInstanceName(0);
+			FNiagaraVariable NamespacedVar = FNiagaraParameterMapHistory::BasicAttributeToNamespacedAttribute(Attribute);
+
+			FString SymbolName = *(ParameterMapInstanceName + TEXT(".") + GetSanitizedSymbolName(NamespacedVar.GetName().ToString()));
+			return AddSourceChunk(SymbolName, Attribute.GetType());
+		}
+		else
+		{
+			Error(LOCTEXT("AttrReadInSpawnError","Cannot read attribute in a spawn script as it's value is not yet initialized."), nullptr, nullptr);
+			return INDEX_NONE;
+		}
+	}
+	else
+	{
+		CompilationOutput.DataUsage.bReadsAttriubteData = true;
+		int32 Chunk = INDEX_NONE;
+		if (!ParameterMapRegisterUniformAttributeVariable(Attribute, nullptr, 0, Chunk))
+		{
+			Error(FText::Format(LOCTEXT("AttrReadError", "Cannot read attribute {0} {1}."), Attribute.GetType().GetNameText(), FText::FromString(*Attribute.GetName().ToString())), nullptr, nullptr);
+			return INDEX_NONE;
+		}
+		return Chunk;
+	}
+}
+
+FString FHlslNiagaraTranslator::GetDataSetAccessSymbol(FNiagaraDataSetID DataSet, int32 IndexChunk, bool bRead)
+{
+	FString Ret = TEXT("\tContext.") + DataSet.Name.ToString() + (bRead ? TEXT("Read") : TEXT("Write"));
+	/*
+	FString Ret = TEXT("Context.") + DataSet.Name.ToString();
+	Ret += bRead ? TEXT("Read") : TEXT("Write");
+	Ret += IndexChunk != INDEX_NONE ? TEXT("_") + CodeChunks[IndexChunk].SymbolName : TEXT("");*/
+	return Ret;
+}
+
+
+void FHlslNiagaraTranslator::ParameterMapSet(UNiagaraNodeParameterMapSet* SetNode, TArray<int32>& Inputs, TArray<int32>& Outputs)
+{
+	Outputs.SetNum(1);
+
+	FString ParameterMapInstanceName = TEXT("Context.Map");
+
+
+	TArray<UEdGraphPin*> InputPins;
+	SetNode->GetInputPins(InputPins);
+	int32 ParamMapHistoryIdx = INDEX_NONE;
+	for (int32 i = 0; i < Inputs.Num(); i++)
+	{
+		int32 Input = Inputs[i];
+		if (i == 0) // This is the parameter map
+		{
+			Outputs[0] = Inputs[0];
+			ParamMapHistoryIdx = Inputs[0];
+			ParameterMapInstanceName = GetParameterMapInstanceName(ParamMapHistoryIdx);
+
+			if (ParamMapHistoryIdx == -1)
+			{
+				Error(LOCTEXT("NoParamMapIdx", "Cannot find parameter map for input!"), SetNode, nullptr);
+				for (int32 j = 0; j < Outputs.Num(); j++)
+				{
+					Outputs[j] = INDEX_NONE;
+					return;
+				}
+			}
+			continue;
+		}
+		else if (SetNode->IsAddPin(InputPins[i]))
+		{
+			// Not a real pin..
+			continue;
+		}
+		else // These are the pins that we are setting on the parameter map.
+		{
+			FNiagaraVariable Var = Schema->PinToNiagaraVariable(InputPins[i], false);
+
+			if (!AddStructToDefinitionSet(Var.GetType()))
+			{
+				Error(FText::Format(LOCTEXT("ParameterMapSetTypeError", "Cannot handle type {0}! Variable: {1}"), Var.GetType().GetNameText(), FText::FromName(Var.GetName())), nullptr, nullptr);
+			}
+
+			FString VarName = Var.GetName().ToString();
+			if (FNiagaraParameterMapHistory::IsExternalConstantNamespace(Var, Script))
+			{
+				Error(FText::Format(LOCTEXT("SetSystemConstantFail", "Cannot Set external constant, Type: {0} Variable: {1}"), Var.GetType().GetNameText(), FText::FromName(Var.GetName())), SetNode, nullptr);
+				continue;
+			}
+
+			Var = ActiveHistoryForFunctionCalls.ResolveAliases(Var);
+
+			if (ParamMapHistoryIdx < ParamMapHistories.Num())
+			{
+				int32 VarIdx = ParamMapHistories[ParamMapHistoryIdx].FindVariableByName(Var.GetName());
+				if (VarIdx != INDEX_NONE && VarIdx < ParamMapSetVariablesToChunks[ParamMapHistoryIdx].Num())
+				{
+					ParamMapSetVariablesToChunks[ParamMapHistoryIdx][VarIdx] = Inputs[i];
+					ParamMapDefinedAttributesToNamespaceVars.FindOrAdd(Var.GetName()) = Var;
+				}
+			}
+
+			if (Var.GetType().GetClass() != nullptr)
+			{
+				// do nothing for now, we've recorded the value for the future.
+			}
+			else
+			{
+				AddBodyChunk(ParameterMapInstanceName + TEXT(".") + GetSanitizedSymbolName(Var.GetName().ToString()), TEXT("{0}"), Var.GetType(), Input, false);
+			}
+		}
+	}
+
+}
+
+bool FHlslNiagaraTranslator::IsBulkSystemScript() const
+{
+	return (Script->Usage == ENiagaraScriptUsage::SystemSpawnScript || Script->Usage == ENiagaraScriptUsage::SystemUpdateScript) &&
+		!Script->GetName().Contains(TEXT("Solo"));
+}
+
+
+bool FHlslNiagaraTranslator::ParameterMapRegisterExternalConstantNamespaceVariable(FNiagaraVariable InVariable, UNiagaraNode* InNode, int32 InParamMapHistoryIdx, int32& Output, UEdGraphPin* InDefaultPin)
+{
+	InVariable = ActiveHistoryForFunctionCalls.ResolveAliases(InVariable);
+	FString VarName = InVariable.GetName().ToString();
+	FString SymbolName = GetSanitizedSymbolName(VarName);
+	FString FlattenedName = SymbolName.Replace(TEXT("."), TEXT("_"));
+	FString ParameterMapInstanceName = GetParameterMapInstanceName(InParamMapHistoryIdx);
+
+	Output = INDEX_NONE;
+	if (InVariable.IsNameValid())
+	{
+		// We don't really want system delta time or inverse system delta time in a spawn script. It leads to trouble.
+		if ((Script->IsInterpolatedParticleSpawnScript() && this->bInsideInterpolatedSpawnScript) || (Script->IsParticleSpawnScript() && !Script->IsInterpolatedParticleSpawnScript()))
+		{
+			if (InVariable == SYS_PARAM_ENGINE_DELTA_TIME || InVariable == SYS_PARAM_ENGINE_INV_DELTA_TIME)
+			{
+				Warning(FText::Format(LOCTEXT("GetParameterInvalidParam", "Cannot call system variable {0} in a spawn script! It is invalid."), FText::FromName(InVariable.GetName())), nullptr, nullptr);
+				Output = GetConstantFloat(0.0f);
+				return true;
+			}
+		}
+
+		if (!IsBulkSystemScript() && false == ParamMapDefinedSystemVarsToUniformChunks.Contains(InVariable.GetName()))
+		{
+			FString SymbolNameDefined = FlattenedName;
+
+			if (!InVariable.IsDataAllocated() && !InDefaultPin)
+			{
+				FNiagaraEditorUtilities::ResetVariableToDefaultValue(InVariable);
+			}
+			else if (!InVariable.IsDataAllocated())
+			{
+				FNiagaraVariable Var = Schema->PinToNiagaraVariable(InDefaultPin, true);
+				FNiagaraEditorUtilities::ResetVariableToDefaultValue(InVariable);
+				if (Var.IsDataAllocated() && Var.GetData() != nullptr)
+				{
+					InVariable.SetData(Var.GetData());
+				}
+			}
+
+			if (InVariable.GetAllocatedSizeInBytes() != InVariable.GetSizeInBytes())
+			{
+				Error(FText::Format(LOCTEXT("GetParameterUnsetParam", "Variable {0} hasn't had its default value set. Required Bytes: {1} vs Allocated Bytes: {2}"), FText::FromName(InVariable.GetName()), FText::AsNumber(InVariable.GetType().GetSize()), FText::AsNumber(InVariable.GetSizeInBytes())), nullptr, nullptr);
+			}
+			
+			CompilationOutput.Parameters.SetOrAdd(InVariable);
+			int32 UniformChunk = AddUniformChunk(SymbolNameDefined, InVariable.GetType());
+			ParamMapDefinedSystemVarsToUniformChunks.Add(InVariable.GetName(), UniformChunk);
+			ParamMapDefinedSystemToNamespaceVars.Add(InVariable.GetName(), InVariable);
+			MainPreSimulateChunks.Add(FString::Printf(TEXT("%s.%s = %s;"), *ParameterMapInstanceName, *VarName, *GetCodeAsSource(UniformChunk)));
+		}
+		else if (IsBulkSystemScript() && !ExternalVariablesForBulkUsage.Contains(InVariable))
+		{
+			ExternalVariablesForBulkUsage.Add(InVariable);
+		}
+		Output = AddSourceChunk(ParameterMapInstanceName + TEXT(".") + SymbolName, InVariable.GetType());
+		return true;
+	}
+
+	if (Output == INDEX_NONE)
+	{
+		Error(FText::Format(LOCTEXT("GetSystemConstantFail", "Unknown System constant, Type: {0} Variable: {1}"), InVariable.GetType().GetNameText(), FText::FromName(InVariable.GetName())), InNode, nullptr);
+	}
+	return false;
+}
+
+
+bool FHlslNiagaraTranslator::ParameterMapRegisterUniformAttributeVariable(const FNiagaraVariable& InVariable, UNiagaraNode* InNode, int32 InParamMapHistoryIdx, int32& Output)
+{
+	FNiagaraVariable NewVar = FNiagaraParameterMapHistory::BasicAttributeToNamespacedAttribute(InVariable);
+	if (NewVar.IsNameValid())
+	{
+		return ParameterMapRegisterNamespaceAttributeVariable(NewVar, InNode, InParamMapHistoryIdx, Output);
+	}
+	return false;
+}
+
+bool FHlslNiagaraTranslator::ParameterMapRegisterNamespaceAttributeVariable(const FNiagaraVariable& InVariable, UNiagaraNode* InNode, int32 InParamMapHistoryIdx, int32& Output)
+{
+	FString VarName = InVariable.GetName().ToString();
+	FString SymbolNameNamespaced = GetSanitizedSymbolName(VarName);
+	FString ParameterMapInstanceName = GetParameterMapInstanceName(InParamMapHistoryIdx);
+	FNiagaraVariable NamespaceVar = InVariable;
+
+	Output = INDEX_NONE;
+	FNiagaraVariable BasicVar = FNiagaraParameterMapHistory::ResolveAsBasicAttribute(InVariable);
+	if (BasicVar.IsNameValid())
+	{
+		if (false == ParamMapDefinedAttributesToUniformChunks.Contains(BasicVar.GetName()))
+		{
+			FString SymbolNameDefined = GetSanitizedSymbolName(BasicVar.GetName().ToString());
+			int32 UniformChunk = INDEX_NONE;
+			int32 Idx = InstanceRead.Variables.Find(NamespaceVar);
+			if (Idx != INDEX_NONE)
+			{
+				UniformChunk = InstanceRead.CodeChunks[Idx];
+			}
+			else
+			{
+				UniformChunk = AddSourceChunk(ParameterMapInstanceName + TEXT(".") + SymbolNameNamespaced, NamespaceVar.GetType());
+				InstanceRead.CodeChunks.Add(UniformChunk);
+				InstanceRead.Variables.Add(NamespaceVar);
+			}
+
+			ParamMapDefinedAttributesToUniformChunks.Add(BasicVar.GetName(), UniformChunk);
+			ParamMapDefinedAttributesToNamespaceVars.Add(BasicVar.GetName(), NamespaceVar);
+		}
+		Output = AddSourceChunk(ParameterMapInstanceName + TEXT(".") + SymbolNameNamespaced, NamespaceVar.GetType());
+		return true;
+	}
+
+	if (Output == INDEX_NONE)
+	{
+		Error(FText::Format(LOCTEXT("GetEmitterUniformFail", "Unknown Emitter Uniform Variable, Type: {0} Variable: {1}"), InVariable.GetType().GetNameText(), FText::FromName(InVariable.GetName())), InNode, nullptr);
+	}
+	return false;
+}
+
+FString FHlslNiagaraTranslator::GetParameterMapInstanceName(int32 ParamMapHistoryIdx)
+{
+	FString ParameterMapInstanceName = TEXT("Context.Map");
+	if (Script->IsInterpolatedParticleSpawnScript())
+	{
+		if (bInsideInterpolatedSpawnScript)
+		{
+			ParameterMapInstanceName = TEXT("Context.MapSpawn");
+		}
+		else
+		{
+			ParameterMapInstanceName = TEXT("Context.MapUpdate");
+		}
+	}
+	return ParameterMapInstanceName;
+}
+
+void FHlslNiagaraTranslator::Emitter(class UNiagaraNodeEmitter* EmitterNode, TArray<int32>& Inputs, TArray<int32>& Outputs)
+{
+	FNiagaraFunctionSignature Signature;
+	UNiagaraScriptSource* Source = EmitterNode->GetScriptSource();
+	if (Source == nullptr)
+	{
+		Error(LOCTEXT("FunctionCallNonexistant", "Emitter call missing ScriptSource"), EmitterNode, nullptr);
+		return;
+	}
+
+	// We need the generated string to generate the proper signature for now.
+	FString EmitterUniqueName = EmitterNode->GetEmitterUniqueName();
+
+	ENiagaraScriptUsage ScriptUsage = EmitterNode->GetUsage();
+	FString Name = EmitterNode->GetName();
+	FString FullName = EmitterNode->GetFullName();
+
+	ActiveHistoryForFunctionCalls.EnterEmitter(EmitterUniqueName, EmitterNode);
+
+	TArray<UEdGraphPin*> CallOutputs;
+	TArray<UEdGraphPin*> CallInputs;
+	EmitterNode->GetOutputPins(CallOutputs);
+	EmitterNode->GetInputPins(CallInputs);
+
+	// We act like a function call here as the semantics are identical.
+	RegisterFunctionCall(ScriptUsage, Name, FullName, Source, Signature, Inputs, CallInputs, CallOutputs, Signature);
+	GenerateFunctionCall(Signature, Inputs, Outputs);
+	ActiveHistoryForFunctionCalls.ExitEmitter(EmitterUniqueName, EmitterNode);
+}
+
+void FHlslNiagaraTranslator::ParameterMapGet(UNiagaraNodeParameterMapGet* GetNode, TArray<int32>& Inputs, TArray<int32>& Outputs)
+{
+	TArray<UEdGraphPin*> OutputPins;
+	GetNode->GetOutputPins(OutputPins);
+
+	TArray<UEdGraphPin*> InputPins;
+	GetNode->GetInputPins(InputPins);
+	
+	int32 ParamMapHistoryIdx = Inputs[0];
+
+	Outputs.SetNum(OutputPins.Num());
+
+	if (ParamMapHistoryIdx == -1)
+	{
+		Error(LOCTEXT("NoParamMapIdx", "Cannot find parameter map for input!"), GetNode, nullptr);
+		for (int32 i = 0; i < Outputs.Num(); i++)
+		{
+			Outputs[i] = INDEX_NONE;
+			return;
+		}
+	}
+
+	FString ParameterMapInstanceName = GetParameterMapInstanceName(ParamMapHistoryIdx);
+	
+	for (int32 i = 0; i < OutputPins.Num(); i++)
+	{
+		if (GetNode->IsAddPin(OutputPins[i]))
+		{
+			// Not a real pin.
+			Outputs[i] = INDEX_NONE;
+			continue;
+		}
+		else // These are the pins that we are getting off the parameter map.
+		{
+			FNiagaraVariable Var = Schema->PinToNiagaraVariable(OutputPins[i], true);
+
+			if (!AddStructToDefinitionSet(Var.GetType()))
+			{
+				Error(FText::Format(LOCTEXT("ParameterMapGetTypeError", "Cannot handle type {0}! Variable: {1}"), Var.GetType().GetNameText(), FText::FromName(Var.GetName())), nullptr, nullptr);
+			}
+
+			// If this is a System parameter, just wire in the system appropriate system attribute.
+			FString VarName = Var.GetName().ToString();
+			FString SymbolName = GetSanitizedSymbolName(VarName);
+
+			bool bIsPerInstanceAttribute = false;
+
+			if (FNiagaraParameterMapHistory::IsExternalConstantNamespace(Var, Script))
+			{
+				if (ParameterMapRegisterExternalConstantNamespaceVariable(Var, GetNode, ParamMapHistoryIdx, Outputs[i], GetNode->GetDefaultPin(OutputPins[i])))
+				{
+					continue;
+				}
+			}
+
+			FNiagaraParameterMapHistory& History = ParamMapHistories[ParamMapHistoryIdx];
+			Var = ActiveHistoryForFunctionCalls.ResolveAliases(Var);
+			if (History.IsPrimaryDataSetOutput(Var, GetTargetUsage()))
+			{
+				bIsPerInstanceAttribute = true;
+			}
+
+			int32 LastSetChunkIdx = INDEX_NONE;
+			if (ParamMapHistoryIdx < ParamMapHistories.Num())
+			{
+				int32 VarIdx = ParamMapHistories[ParamMapHistoryIdx].FindVariableByName(Var.GetName());
+				if (VarIdx != INDEX_NONE && VarIdx < ParamMapSetVariablesToChunks[ParamMapHistoryIdx].Num())
+				{
+					LastSetChunkIdx = ParamMapSetVariablesToChunks[ParamMapHistoryIdx][VarIdx];
+				}
+
+				// We have yet to write to this parameter, use the default value if specified and the parameter 
+				// isn't a per-particle value.
+				bool bIgnoreDefaultValue = false;
+				if (bIsPerInstanceAttribute)
+				{
+					if ((Script->IsInterpolatedParticleSpawnScript() && bInsideInterpolatedSpawnScript == false) || !Script->IsSpawnScript())
+					{
+						bIgnoreDefaultValue = true;
+					}
+				}
+
+				if (LastSetChunkIdx == INDEX_NONE && (Script->IsInterpolatedParticleSpawnScript() || Script->IsParticleSpawnScript() || Script->IsEmitterSpawnScript() || Script->IsSystemSpawnScript()))
+				{
+					if (FNiagaraParameterMapHistory::IsInitialValue(Var))
+					{
+						FNiagaraVariable SourceForInitialValue = FNiagaraParameterMapHistory::GetSourceForInitialValue(Var);
+						bool bFoundExistingSet = false;
+						for (int32 OtherParamIdx = 0; OtherParamIdx < OtherOutputParamMapHistories.Num(); OtherParamIdx++)
+						{
+							if (INDEX_NONE != OtherOutputParamMapHistories[OtherParamIdx].FindVariableByName(SourceForInitialValue.GetName()))
+							{
+								bFoundExistingSet = true;
+							}
+						}
+
+						if (bFoundExistingSet)
+						{
+							LastSetChunkIdx = AddBodyChunk(ParameterMapInstanceName + TEXT(".") + GetSanitizedSymbolName(Var.GetName().ToString()),
+								ParameterMapInstanceName + TEXT(".") + GetSanitizedSymbolName(SourceForInitialValue.GetName().ToString()), Var.GetType(), false);
+							ParamMapDefinedAttributesToNamespaceVars.FindOrAdd(Var.GetName()) = Var;
+						}
+						else
+						{
+							Error(FText::Format(LOCTEXT("MissingInitialValueSource", "Variable {0} is used, but its source variable {1} is not set!"), FText::FromName(Var.GetName()), FText::FromName(SourceForInitialValue.GetName())), nullptr, nullptr);
+						}
+					}
+				}
+
+				if (LastSetChunkIdx == INDEX_NONE && !bIgnoreDefaultValue)
+				{
+					UEdGraphPin* InputPin = GetNode->GetDefaultPin(OutputPins[i]);
+						
+					// Default was found, trace back its inputs.
+					if (InputPin != nullptr )
+					{
+						LastSetChunkIdx = CompilePin(InputPin);
+
+						if (!Var.IsDataInterface())
+						{
+							if (VarIdx != INDEX_NONE && VarIdx < ParamMapSetVariablesToChunks[ParamMapHistoryIdx].Num())
+							{
+								// Record that we wrote to it.
+								ParamMapSetVariablesToChunks[ParamMapHistoryIdx][VarIdx] = LastSetChunkIdx;	
+								ParamMapDefinedAttributesToNamespaceVars.FindOrAdd(Var.GetName()) = Var;
+							}
+							else
+							{
+								Error(FText::Format(LOCTEXT("NoVarDefaultFound", "Default found for {0}, but not found in ParameterMap traversal"), FText::FromName(Var.GetName())), GetNode, nullptr);
+							}
+
+							// Actually insert the text that sets the default value
+							if (LastSetChunkIdx != INDEX_NONE)
+							{
+								if (Var.GetType().GetClass() == nullptr) // Only need to do this wiring for things that aren't data interfaces.
+								{
+									AddBodyChunk(ParameterMapInstanceName + TEXT(".") + GetSanitizedSymbolName(Var.GetName().ToString()), TEXT("{0}"), Var.GetType(), LastSetChunkIdx, false);
+								}
+							}
+						}
+					}
+				}
+			}
+
+			// If we are of a data interface, we should output the data interface registration index, otherwise output
+			// the map namespace that we're writing to.
+			if (Var.IsDataInterface())
+			{
+				// In order for a module to compile successfully, we potentially need to generate default values
+				// for variables encountered without ever being set. We do this by creating an instance of the CDO.
+				if (Script->IsStandaloneScript() && LastSetChunkIdx == INDEX_NONE)
+				{
+					UObject* Obj = const_cast<UClass*>(Var.GetType().GetClass())->GetDefaultObject(true);
+					UNiagaraDataInterface* DataInterface = Cast<UNiagaraDataInterface>(DuplicateObject(Obj, GetTransientPackage()));
+					if (DataInterface)
+					{
+						UE_LOG(LogNiagaraEditor, Log, TEXT("Creating stand-in due to missing data interface!"));
+						LastSetChunkIdx = RegisterDataInterface(Var, DataInterface);
+					}
+				}
+				
+				Outputs[i] = LastSetChunkIdx;
+			}
+			else
+			{
+				Outputs[i] = AddSourceChunk(ParameterMapInstanceName + TEXT(".") + GetSanitizedSymbolName(Var.GetName().ToString()), Var.GetType());
+				ParamMapDefinedAttributesToNamespaceVars.FindOrAdd(Var.GetName()) = Var;
+			}
+		}
+	}
+}
+
+void FHlslNiagaraTranslator::ParameterCollection(class UNiagaraNodeParameterCollection* CollectionNode, TArray<int32>& Outputs)
+{
+	UNiagaraParameterCollection* Collection = CollectionNode->GetCollection();
+	check(Collection);
+	
+	TArray<UEdGraphPin*> OutputPins;
+	CollectionNode->GetOutputPins(OutputPins);
+	
+	check(Outputs.Num() == 0);
+	for (int32 i = 0; i < OutputPins.Num(); i++)
+	{
+		if (CollectionNode->IsAddPin(OutputPins[i]))
+		{
+			Outputs.Add(INDEX_NONE);//Add dummy output for add pin.
+			continue;
+		}
+
+		FNiagaraVariable PinVar = Collection->CollectionParameterFromFriendlyParameter(Schema->PinToNiagaraVariable(OutputPins[i]));
+
+		int32 SrcVarIdx = Collection->IndexOfParameter(PinVar);
+		if (SrcVarIdx == INDEX_NONE)
+		{
+			Error(LOCTEXT("MissingCollectionParameter", "Parameter is not a member of this Parameter Collection."), CollectionNode, OutputPins[i]);
+			Outputs.Add(INDEX_NONE);
+		}
+		else
+		{
+			FNiagaraVariable& SrcVar = Collection->GetParameters()[SrcVarIdx];
+		
+			if (SrcVar.IsDataInterface())
+			{
+				Outputs.Add(RegisterDataInterface(SrcVar, Collection->GetDefaultInstance()->GetParameterStore().GetDataInterface(SrcVar)));
+			}
+			else
+			{
+				FString VarSymbol = GetSanitizedSymbolName(SrcVar.GetName().ToString());
+				//Add uniform access to hlsl.
+				//Need to interact with the parameter map?
+				Outputs.Add(AddUniformChunk(VarSymbol, SrcVar.GetType()));
+				CompilationOutput.Parameters.SetOrAdd(SrcVar);
+			}
+
+			Script->ParameterCollections.AddUnique(Collection);
+		}
+	}
+}
+
+void FHlslNiagaraTranslator::ReadDataSet(const FNiagaraDataSetID DataSet, const TArray<FNiagaraVariable>& Variables, ENiagaraDataSetAccessMode AccessMode, int32 InputChunk, TArray<int32>& Outputs)
+{
+	//Eventually may allow events that take in a direct index or condition but for now we don't
+	check(InputChunk == INDEX_NONE);
+
+	TMap<int32, FDataSetAccessInfo>& Reads = DataSetReadInfo[(int32)AccessMode].FindOrAdd(DataSet);
+	FDataSetAccessInfo* DataSetReadForInput = Reads.Find(InputChunk);
+	if (!DataSetReadForInput)
+	{
+		DataSetReadForInput = &Reads.Add(InputChunk);
+		
+		DataSetReadForInput->Variables = Variables;
+		DataSetReadForInput->CodeChunks.Reserve(Variables.Num()+1);
+
+		FString DataSetAccessSymbol = GetDataSetAccessSymbol(DataSet, InputChunk, true);
+		//Add extra output to indicate if event read is valid data.
+		//DataSetReadForInput->CodeChunks.Add(AddSourceChunk(DataSetAccessSymbol + TEXT("_Valid"), FNiagaraTypeDefinition::GetIntDef()));
+		for (int32 i=0; i<Variables.Num(); ++i)
+		{
+			const FNiagaraVariable& Var = Variables[i];
+			if (!AddStructToDefinitionSet(Var.GetType()))
+			{
+				Error(FText::Format(LOCTEXT("GetConstantFailTypeVar", "Cannot handle type {0}! Variable: {1}"), Var.GetType().GetNameText(), FText::FromName(Var.GetName())), nullptr, nullptr);
+			}
+			DataSetReadForInput->CodeChunks.Add(AddSourceChunk(DataSetAccessSymbol + TEXT(".") + GetSanitizedSymbolName(Var.GetName().ToString()), Var.GetType()));
+		}
+		Outputs = DataSetReadForInput->CodeChunks;
+	}
+	else
+	{
+		check(Variables.Num() == DataSetReadForInput->Variables.Num());
+		Outputs = DataSetReadForInput->CodeChunks;
+	}
+}
+
+void FHlslNiagaraTranslator::WriteDataSet(const FNiagaraDataSetID DataSet, const TArray<FNiagaraVariable>& Variables, ENiagaraDataSetAccessMode AccessMode, const TArray<int32>& Inputs)
+{
+	int32 ConditionalChunk = Inputs[0];
+	int32 InputChunk = Inputs[1];
+	TMap<int32, FDataSetAccessInfo>& Writes = DataSetWriteInfo[(int32)AccessMode].FindOrAdd(DataSet);
+	FDataSetAccessInfo* DataSetWriteForInput = Writes.Find(InputChunk);
+	int32& DataSetWriteConditionalIndex = DataSetWriteConditionalInfo[(int32)AccessMode].FindOrAdd(DataSet);
+	
+	//We should never try to write to the exact same dataset at the same index/condition twice.
+	//This is still possible but we can catch easy cases here.
+	if (DataSetWriteForInput)
+	{
+		//TODO: improve error report.
+		Error(LOCTEXT("WritingToSameDataSetError", "Writing to the same dataset with the same condition/index."), NULL, NULL);
+		return;
+	}
+
+	DataSetWriteConditionalIndex = ConditionalChunk;
+	DataSetWriteForInput = &Writes.Add(InputChunk);
+
+	DataSetWriteForInput->Variables = Variables;
+
+	//FString DataSetAccessName = GetDataSetAccessSymbol(DataSet, InputChunk, false);
+	FString DataSetAccessName = FString("Context.") + DataSet.Name.ToString() + "Write";	// TODO: HACK - need to get the real symbol name here
+
+	//First chunk for a write is always the condition pin.
+	//We always write the event payload into the temp storage but we can access this condition to pass to the final actual write to the buffer.
+	
+	DataSetWriteForInput->CodeChunks.Add(AddBodyChunk(DataSetAccessName + TEXT("_Valid"), TEXT("{0}"), FNiagaraTypeDefinition::GetBoolDef(), Inputs[0], false));
+	for (int32 i = 0; i < Variables.Num(); ++i)
+	{
+		const FNiagaraVariable& Var = Variables[i];
+		int32 Input = Inputs[i + 1];//input 0 is the valid input (no entry in variables array), so we need of offset all other inputs by 1.
+		DataSetWriteForInput->CodeChunks.Add(AddBodyChunk(DataSetAccessName + TEXT(".") + GetSanitizedSymbolName(Var.GetName().ToString()), TEXT("{0}"), Var.GetType(), Input, false));
+	}
+
+}
+
+int32 FHlslNiagaraTranslator::RegisterDataInterface(FNiagaraVariable& Var, UNiagaraDataInterface* DataInterface)
+{
+	int32 FuncParam;
+	if (GetFunctionParameter(Var, FuncParam))
+	{
+		if (FuncParam != INDEX_NONE)
+		{
+			//This data interface param has been overridden by the function call so use that index.	
+			return FuncParam;
+		}
+	}
+	
+	//If we get here then this is a new data interface.
+	int32 Idx = CompilationOutput.DataInterfaceInfo.AddDefaulted();
+	CompilationOutput.DataInterfaceInfo[Idx].DataInterface = DataInterface;
+	CompilationOutput.DataInterfaceInfo[Idx].Name = Var.GetName();
+
+	//Interface requires per instance data so add a user pointer table entry.
+	if (DataInterface->PerInstanceDataSize() > 0)
+	{
+		CompilationOutput.DataInterfaceInfo[Idx].UserPtrIdx = CompilationOutput.NumUserPtrs++;
+	}
+	return Idx;
+}
+
+void FHlslNiagaraTranslator::Operation(class UNiagaraNodeOp* Operation, TArray<int32>& Inputs, TArray<int32>& Outputs)
+{
+	// Use the pins to determine the output type here since they may have been changed due to numeric pin fix up.
+	const FNiagaraOpInfo* OpInfo = FNiagaraOpInfo::GetOpInfo(Operation->OpName);
+
+	//EnterStatsScope(FNiagaraStatScope(*Operation->GetFullName(), OpInfo->FriendlyName));
+
+	TArray<UEdGraphPin*> OutputPins;
+	Operation->GetOutputPins(OutputPins);
+	for (int32 OutputIndex = 0; OutputIndex < OutputPins.Num(); OutputIndex++)
+	{
+		FNiagaraTypeDefinition OutputType = Schema->PinToTypeDefinition(OutputPins[OutputIndex]);
+
+		if (!AddStructToDefinitionSet(OutputType))
+		{
+			FText PinNameText = OutputPins[OutputIndex]->PinFriendlyName.IsEmpty() ? FText::FromString(OutputPins[OutputIndex]->PinName) : OutputPins[OutputIndex]->PinFriendlyName;
+			Error(FText::Format(LOCTEXT("GetConstantFailTypePin", "Cannot handle type {0}! Output Pin: {1}"), OutputType.GetNameText(), PinNameText), Operation, OutputPins[OutputIndex]);
+		}
+
+		const FNiagaraOpInOutInfo& IOInfo = OpInfo->Outputs[OutputIndex];
+		check(!IOInfo.HlslSnippet.IsEmpty());
+		Outputs.Add(AddBodyChunk(GetUniqueSymbolName(IOInfo.Name), IOInfo.HlslSnippet, OutputType, Inputs));
+	}
+
+	//ExitStatsScope();
+}
+
+void FHlslNiagaraTranslator::FunctionCall(UNiagaraNodeFunctionCall* FunctionNode, TArray<int32>& Inputs, TArray<int32>& Outputs)
+{
+	FNiagaraFunctionSignature Signature;
+	if (FunctionNode->FunctionScript == nullptr && !FunctionNode->Signature.IsValid())
+	{
+		Error(LOCTEXT("FunctionCallNonexistant", "Function call missing FunctionScript and invalid signature"), FunctionNode, nullptr);
+		return;
+	}
+
+	// We need the generated string to generate the proper signature for now.
+	ActiveHistoryForFunctionCalls.EnterFunction(FunctionNode->GetFunctionName(), FunctionNode->FunctionScript, FunctionNode);
+
+	TArray<UEdGraphPin*> CallOutputs;
+	TArray<UEdGraphPin*> CallInputs;
+	FunctionNode->GetOutputPins(CallOutputs);
+	FunctionNode->GetInputPins(CallInputs);
+
+	ENiagaraScriptUsage ScriptUsage = ENiagaraScriptUsage::Function;
+	FString Name;
+	FString FullName; 
+	UNiagaraScriptSource* Source = nullptr;
+
+	if (FunctionNode->FunctionScript)
+	{
+		ScriptUsage = FunctionNode->FunctionScript->GetUsage();
+		Name = FunctionNode->FunctionScript->GetName();
+		FullName = FunctionNode->FunctionScript->GetFullName();
+		Source = CastChecked<UNiagaraScriptSource>(FunctionNode->FunctionScript->GetSource());
+	}
+
+	RegisterFunctionCall(ScriptUsage, Name, FullName,  Source, FunctionNode->Signature, Inputs, CallInputs, CallOutputs, Signature);
+	GenerateFunctionCall(Signature, Inputs, Outputs);
+	ActiveHistoryForFunctionCalls.ExitFunction(FunctionNode->GetFunctionName(), FunctionNode->FunctionScript, FunctionNode);
+}
+
+void FHlslNiagaraTranslator::RegisterFunctionCall(ENiagaraScriptUsage ScriptUsage, const FString& InName, const FString& InFullName, UNiagaraScriptSource* Source, FNiagaraFunctionSignature& InSignature, TArray<int32>& Inputs, const TArray<UEdGraphPin*>& CallInputs, const TArray<UEdGraphPin*>& CallOutputs,
+	FNiagaraFunctionSignature& OutSignature)
+{
+	//////////////////////////////////////////////////////////////////////////
+	if (Source)
+	{
+		UNiagaraGraph* SourceGraph = CastChecked<UNiagaraGraph>(Source->NodeGraph);
+
+		bool bHasNumericInputs = false;
+		if (SourceGraph->HasNumericParameters())
+		{
+			TArray<UNiagaraNodeInput*> InputNodes;
+			UNiagaraGraph::FFindInputNodeOptions Options;
+			Options.bFilterDuplicates = true;
+			Options.bIncludeParameters = true;
+			Options.bIncludeAttributes = false;
+			Options.bIncludeSystemConstants = false;
+			SourceGraph->FindInputNodes(InputNodes, Options);
+
+			for (UNiagaraNodeInput* Input : InputNodes)
+			{
+				if (Input->Input.GetType() == FNiagaraTypeDefinition::GetGenericNumericDef())
+				{
+					bHasNumericInputs = true;
+				}
+			}
+
+			//We have to clone and preprocess the function graph to ensure it's numeric types have been fixed up to real types.
+			UNiagaraGraph* PreprocessedGraph = CastChecked<UNiagaraGraph>(FEdGraphUtilities::CloneGraph(SourceGraph, Source, &MessageLog));
+			FEdGraphUtilities::MergeChildrenGraphsIn(PreprocessedGraph, PreprocessedGraph, /*bRequireSchemaMatch=*/ true);
+			PreprocessFunctionGraph(*this, Schema, PreprocessedGraph, CallInputs, CallOutputs, ScriptUsage);
+			SourceGraph = PreprocessedGraph;
+		}
+		else
+		{
+			//If we don't have numeric inputs then we can cache the preprocessed graphs.
+			UNiagaraGraph** CachedGraph = PreprocessedFunctions.Find(SourceGraph);
+			if (!CachedGraph)
+			{
+				CachedGraph = &PreprocessedFunctions.Add(SourceGraph);
+				*CachedGraph = CastChecked<UNiagaraGraph>(FEdGraphUtilities::CloneGraph(Source->NodeGraph, Source, &MessageLog));
+				FEdGraphUtilities::MergeChildrenGraphsIn(*CachedGraph, *CachedGraph, /*bRequireSchemaMatch=*/ true);
+				PreprocessFunctionGraph(*this, Schema, *CachedGraph, CallInputs, CallOutputs, ScriptUsage);
+			}
+			SourceGraph = *CachedGraph;
+		}
+		
+		GenerateFunctionSignature(ScriptUsage, InName, InFullName, SourceGraph, Inputs, bHasNumericInputs, OutSignature);
+
+// 		//Sort the input and outputs to match the sorted parameters. They may be different.
+// 		TArray<FNiagaraVariable> OrderedInputs;
+// 		TArray<FNiagaraVariable> OrderedOutputs;
+// 		SourceGraph->GetParameters(OrderedInputs, OrderedOutputs);
+// 		TArray<UEdGraphPin*> InPins;
+// 		FunctionNode->GetInputPins(InPins);
+// 
+// 		TArray<int32> OrderedInputChunks;
+// 		OrderedInputChunks.SetNumUninitialized(Inputs.Num());
+// 		for (int32 i = 0; i < InPins.Num(); ++i)
+// 		{
+// 			FNiagaraVariable PinVar(Schema->PinToTypeDefinition(InPins[i]), *InPins[i]->PinName);
+// 			int32 InputIdx = OrderedInputs.IndexOfByKey(PinVar);
+// 			check(InputIdx != INDEX_NONE);
+// 			OrderedInputChunks[i] = Inputs[InputIdx];
+// 		}
+// 		Inputs = OrderedInputChunks;
+
+		FString* FuncBody = Functions.Find(OutSignature);
+		if (!FuncBody)
+		{
+			//We've not compiled this function yet so compile it now.
+			EnterFunction(InName, OutSignature, Inputs);
+			
+			UNiagaraNodeOutput* FuncOutput = SourceGraph->FindOutputNode(ScriptUsage);
+			check(FuncOutput);
+
+			//Track the start of this funciton in the chunks so we can remove them after we grab the funcitons hlsl.
+			int32 ChunkStart = CodeChunks.Num();
+			int32 ChunkStartsByMode[(int32)ENiagaraCodeChunkMode::Num];
+			for (int32 i = 0; i < (int32)ENiagaraCodeChunkMode::Num; ++i)
+			{
+				ChunkStartsByMode[i] = ChunksByMode[i].Num();
+			}
+
+			FHlslNiagaraTranslator* ThisTranslator = this;
+			TArray<int32> FuncOutputChunks;
+			FuncOutput->Compile(ThisTranslator, FuncOutputChunks);
+
+			// Find all of the data set writes that were connected to this particular graph.
+			TArray<UNiagaraNode*> OutputTraversal;
+			SourceGraph->BuildTraversal(OutputTraversal, FuncOutput);
+
+			TArray<UNiagaraNodeWriteDataSet *>WriteNodes;
+			SourceGraph->FindWriteDataSetNodes(WriteNodes);
+
+			// If the write data set graph traversal intersects with our normal output graph traversal, then
+			// we consider them part of the same overall graph traversal and include the write nodes.
+			for (UNiagaraNodeWriteDataSet *WriteNode : WriteNodes)
+			{
+				TArray<UNiagaraNode*> WriteTraversal;
+				SourceGraph->BuildTraversal(WriteTraversal, WriteNode);
+
+				bool bReferencesSameGraph = false;
+				for (UNiagaraNode* Node : WriteTraversal)
+				{
+					if (OutputTraversal.Contains(Node))
+					{
+						bReferencesSameGraph = true;
+					}
+				}
+				
+				if (bReferencesSameGraph)
+				{
+					WriteNode->Compile(ThisTranslator, FuncOutputChunks);
+				}
+			}
+
+			//Grab all the body chunks for this function.
+			FString FunctionDefStr;
+			for (int32 i = ChunkStartsByMode[(int32)ENiagaraCodeChunkMode::Body]; i < ChunksByMode[(int32)ENiagaraCodeChunkMode::Body].Num(); ++i)
+			{
+				FunctionDefStr += GetCode(ChunksByMode[(int32)ENiagaraCodeChunkMode::Body][i]);
+			}
+
+			//Now remove all chunks for the function again.			
+			//This is super hacky. Should move chunks etc into a proper scoping system.
+			
+			TArray<FNiagaraCodeChunk> FuncUniforms;
+			for (int32 i = 0; i < (int32)ENiagaraCodeChunkMode::Num; ++i)
+			{
+				//Keep uniform chunks.
+				if (i == (int32)ENiagaraCodeChunkMode::Uniform)
+				{
+					for (int32 ChunkIdx = ChunkStartsByMode[i]; ChunkIdx < ChunksByMode[i].Num(); ++ChunkIdx)
+					{
+						FuncUniforms.Add(CodeChunks[ChunksByMode[i][ChunkIdx]]);
+					}
+				}					
+
+				ChunksByMode[i].RemoveAt(ChunkStartsByMode[i], ChunksByMode[i].Num() - ChunkStartsByMode[i]);
+			}
+			CodeChunks.RemoveAt(ChunkStart, CodeChunks.Num() - ChunkStart, false);
+
+			//Re-add the uniforms. Really this is horrible. Rework soon.
+			for (FNiagaraCodeChunk& Chunk : FuncUniforms)
+			{
+				ChunksByMode[(int32)ENiagaraCodeChunkMode::Uniform].Add(CodeChunks.Add(Chunk));
+			}
+
+			// We don't support an empty function definition when calling a real function.
+			if (FunctionDefStr.IsEmpty())
+			{
+				FunctionDefStr += TEXT("\n");
+			}
+
+			Functions.Add(OutSignature, FunctionDefStr);
+
+			ExitFunction();
+		}
+	}
+	else
+	{
+		check(InSignature.IsValid());
+		check(InSignature.bMemberFunction);
+		check(Inputs.Num() > 0);
+
+		OutSignature = InSignature;		
+
+		//First input for these is the owner of the function.
+		int32 OwnerIdx = Inputs[0];
+		if (OwnerIdx < 0 || OwnerIdx >= CompilationOutput.DataInterfaceInfo.Num())
+		{
+			Error(LOCTEXT("FunctionCallDataInterfaceMissingRegistration", "Function call signature does not match to a registered DataInterface. Valid DataInterfaces should be wired into a DataInterface function call."), nullptr, nullptr);
+			return;
+		}
+		FNiagaraScriptDataInterfaceInfo& Info = CompilationOutput.DataInterfaceInfo[OwnerIdx];
+
+		// Double-check to make sure that the signature matches those specified by the data 
+		// interface. It could be that the existing node has been removed and the graph
+		// needs to be refactored. If that's the case, emit an error.
+		if (Info.DataInterface != nullptr && OutSignature.bMemberFunction)
+		{
+			TArray<FNiagaraFunctionSignature> DataInterfaceFunctions;
+			Info.DataInterface->GetFunctions(DataInterfaceFunctions);
+			bool bFoundMatch = false;
+			for (const FNiagaraFunctionSignature& Sig : DataInterfaceFunctions)
+			{
+				if (Sig == OutSignature)
+				{
+					bFoundMatch = true;
+				}
+			}
+
+			if (!bFoundMatch)
+			{
+				Error(LOCTEXT("FunctionCallDataInterfaceMissing", "Function call signature does not match DataInterface possible signatures?"), nullptr, nullptr);
+				return;
+			}
+
+			if (Info.UserPtrIdx != INDEX_NONE)
+			{
+				//This interface requires per instance data via a user ptr so place the index to it at the end of the inputs.
+				Inputs.Add(AddSourceChunk(LexicalConversion::ToString(Info.UserPtrIdx), FNiagaraTypeDefinition::GetIntDef()));
+				OutSignature.Inputs.Add(FNiagaraVariable(FNiagaraTypeDefinition::GetIntDef(), TEXT("InstanceData")));
+			}
+		}
+
+		//Override the owner id of the signature with the actual caller.
+		OutSignature.OwnerName = Info.Name;
+		Info.RegisteredFunctions.Add(OutSignature);
+
+		Functions.FindOrAdd(OutSignature);
+	}
+}
+
+void FHlslNiagaraTranslator::GenerateFunctionCall(FNiagaraFunctionSignature& FunctionSignature, TArray<int32>& Inputs, TArray<int32>& Outputs)
+{
+	EnterStatsScope(FNiagaraStatScope(*GetFunctionSignatureSymbol(FunctionSignature), FText::FromString(FunctionSignature.GetName())));
+
+	TArray<FString> MissingParameters;
+	int32 ParamIdx = 0;
+	TArray<int32> Params;
+	Params.Reserve(Inputs.Num() + Outputs.Num());
+	FString DefStr = GetFunctionSignatureSymbol(FunctionSignature) + TEXT("(");
+	for (int32 i = 0; i < FunctionSignature.Inputs.Num(); ++i)
+	{
+		FNiagaraTypeDefinition Type = FunctionSignature.Inputs[i].GetType();
+		//We don't write class types as real params in the hlsl
+		if (!Type.GetClass())
+		{
+			if (!AddStructToDefinitionSet(Type))
+			{
+				Error(FText::Format(LOCTEXT("GetConstantFailTypeVar2", "Cannot handle type {0}! Variable: {1}"),  Type.GetNameText(), FText::FromName(FunctionSignature.Inputs[i].GetName())), nullptr, nullptr);
+			}
+
+			if (ParamIdx != 0)
+			{
+				DefStr += TEXT(", ");
+			}
+
+			int32 Input = Inputs[i];
+			bool bSkip = false;
+
+			if (FunctionSignature.Inputs[i].GetType() == FNiagaraTypeDefinition::GetParameterMapDef())
+			{
+				Input = INDEX_NONE;
+				bSkip = true;
+			}
+
+			if (!bSkip)
+			{
+				Params.Add(Input);
+				if (Input == INDEX_NONE)
+				{
+					MissingParameters.Add(FunctionSignature.Inputs[i].GetName().ToString());
+				}
+				else
+				{
+					DefStr += FString::Printf(TEXT("{%d}"), ParamIdx);
+				}
+				++ParamIdx;
+			}
+		}
+	}
+
+	for (int32 i = 0; i < FunctionSignature.Outputs.Num(); ++i)
+	{
+		FNiagaraVariable& OutVar = FunctionSignature.Outputs[i];
+		FNiagaraTypeDefinition Type = OutVar.GetType();
+
+		//We don't write class types as real params in the hlsl
+		if (!Type.GetClass())
+		{
+			if (!AddStructToDefinitionSet(Type))
+			{
+				Error(FText::Format(LOCTEXT("GetConstantFailTypeVar3", "Cannot handle type {0}! Variable: {1}"), Type.GetNameText(), FText::FromName(FunctionSignature.Outputs[i].GetName())), nullptr, nullptr);
+			}
+			
+			if (ParamIdx > 0)
+			{
+				DefStr += TEXT(", ");
+			}
+
+			int32 Output = INDEX_NONE;
+			int32 ParamOutput = INDEX_NONE;
+			bool bSkip = false;
+			if (FunctionSignature.Outputs[i].GetType() == FNiagaraTypeDefinition::GetParameterMapDef())
+			{
+				if (i < FunctionSignature.Inputs.Num() && FunctionSignature.Inputs[i].GetType() == FNiagaraTypeDefinition::GetParameterMapDef())
+				{
+					Output = Inputs[i];
+				}
+				bSkip = true;
+			}
+			else
+			{
+				FString OutputStr = FString::Printf(TEXT("%sOutput_%s"), *GetFunctionSignatureSymbol(FunctionSignature), *OutVar.GetName().ToString());
+				Output = AddBodyChunk(GetUniqueSymbolName(*OutputStr), TEXT(""), OutVar.GetType());
+				ParamOutput = Output;
+			}
+
+			Outputs.Add(Output);
+
+			if (!bSkip)
+			{
+				Params.Add(ParamOutput);
+				if (ParamOutput == INDEX_NONE)
+				{
+					MissingParameters.Add(OutVar.GetName().ToString());
+				}
+				else
+				{
+					DefStr += FString::Printf(TEXT("{%d}"), ParamIdx);
+				}
+				++ParamIdx;
+			}
+		}
+	}
+
+	if (FunctionSignature.bRequiresContext)
+	{
+		if (ParamIdx > 0)
+		{
+			DefStr += TEXT(", ");
+		}
+		DefStr += "Context";
+	}
+
+	DefStr += TEXT(")");
+
+	if (MissingParameters.Num() > 0)
+	{
+		for (FString MissingParam : MissingParameters)
+		{
+			FText Fmt = LOCTEXT("ErrorCompilingParameterFmt", "Error compiling parameter {0} in function call {1}");
+			FText ErrorText = FText::Format(Fmt, FText::FromString(MissingParam), FText::FromString(GetFunctionSignatureSymbol(FunctionSignature)));
+			Error(ErrorText, nullptr, nullptr);
+		}
+		return;
+	}
+
+	AddBodyChunk(TEXT(""), DefStr, FNiagaraTypeDefinition::GetFloatDef(), Params);
+
+	ExitStatsScope();
+}
+
+FString FHlslNiagaraTranslator::GetFunctionSignatureSymbol(const FNiagaraFunctionSignature& Sig)
+{
+	FString SigStr = Sig.GetName();
+	if (!Sig.OwnerName.IsNone() && Sig.OwnerName.IsValid())
+	{
+		SigStr += TEXT("_") + Sig.OwnerName.ToString().Replace(TEXT("."), TEXT(""));;
+	}
+	else
+	{
+		SigStr += TEXT("_Func_");
+	}
+	return GetSanitizedSymbolName(SigStr);
+}
+
+FString FHlslNiagaraTranslator::GetFunctionSignature(const FNiagaraFunctionSignature& Sig)
+{
+	FString SigStr = TEXT("void ") + GetFunctionSignatureSymbol(Sig);
+
+	SigStr += TEXT("(");
+	int32 ParamIdx = 0;
+	for (int32 i = 0; i < Sig.Inputs.Num(); ++i)
+	{
+		const FNiagaraVariable& Input = Sig.Inputs[i];
+		//We don't write class types as real params in the hlsl
+		if (Input.GetType().GetClass() == nullptr)
+		{	
+			if (Input.GetType() == FNiagaraTypeDefinition::GetParameterMapDef())
+			{
+				// Skip parameter maps.
+			}
+			else
+			{
+				if (ParamIdx > 0)
+				{
+					SigStr += TEXT(", ");
+				}
+
+				SigStr += FHlslNiagaraTranslator::GetStructHlslTypeName(Input.GetType()) + TEXT(" In_") + FHlslNiagaraTranslator::GetSanitizedSymbolName(Input.GetName().ToString());
+				++ParamIdx;
+			}
+		}
+	}
+
+	for (int32 i = 0; i < Sig.Outputs.Num(); ++i)
+	{
+		const FNiagaraVariable& Output = Sig.Outputs[i];
+		//We don't write class types as real params in the hlsl
+		if (Output.GetType().GetClass() == nullptr)
+		{
+			if (Output.GetType() == FNiagaraTypeDefinition::GetParameterMapDef())
+			{
+				// Skip output parameter maps..
+			}
+			else
+			{
+				if (ParamIdx > 0)
+				{
+					SigStr += TEXT(", ");
+				}
+
+				SigStr += TEXT("out ") + FHlslNiagaraTranslator::GetStructHlslTypeName(Output.GetType()) + TEXT(" Out_") + FHlslNiagaraTranslator::GetSanitizedSymbolName(Output.GetName().ToString());
+				++ParamIdx;
+			}
+		}
+	}
+	if (Sig.bRequiresContext)
+	{
+		if (ParamIdx > 0)
+		{
+			SigStr += TEXT(", ");
+		}
+		SigStr += TEXT("inout FSimulationContext Context");
+	}
+	return SigStr + TEXT(")");
+}
+
+int32 GetPinIndexById(const TArray<UEdGraphPin*>& Pins, FGuid PinId)
+{
+	for (int32 i = 0; i < Pins.Num(); ++i)
+	{
+		if (Pins[i]->PinId == PinId)
+		{
+			return i;
+		}
+	}
+	return INDEX_NONE;
+}
+
+FNiagaraTypeDefinition FHlslNiagaraTranslator::GetChildType(const FNiagaraTypeDefinition& BaseType, const FName& PropertyName)
+{
+	const UScriptStruct* Struct = BaseType.GetScriptStruct();
+	if (Struct != nullptr)
+	{
+		// Dig through properties to find the matching property native type (if it exists)
+		for (TFieldIterator<UProperty> PropertyIt(Struct, EFieldIteratorFlags::IncludeSuper); PropertyIt; ++PropertyIt)
+		{
+			const UProperty* Property = *PropertyIt;
+			if (Property->GetName() == PropertyName.ToString())
+			{
+				if (Property->IsA(UFloatProperty::StaticClass()))
+				{
+					return FNiagaraTypeDefinition::GetFloatDef();
+				}
+				else if (Property->IsA(UIntProperty::StaticClass()))
+				{
+					return FNiagaraTypeDefinition::GetIntDef();
+				}
+				else if (Property->IsA(UBoolProperty::StaticClass()))
+				{
+					return FNiagaraTypeDefinition::GetBoolDef();
+				}
+				else if (const UStructProperty* StructProp = CastChecked<UStructProperty>(Property))
+				{
+					return FNiagaraTypeDefinition(StructProp->Struct);
+				}
+			}
+		}
+	}
+	return FNiagaraTypeDefinition();
+}
+
+FString FHlslNiagaraTranslator::ComputeMatrixColumnAccess(const FString& Name)
+{
+	FString Value;
+	int32 Column = -1;
+
+	if (Name.Find("X", ESearchCase::IgnoreCase) != -1)
+		Column = 0;
+	else if (Name.Find("Y", ESearchCase::IgnoreCase) != -1)
+		Column = 1;
+	else if (Name.Find("Z", ESearchCase::IgnoreCase) != -1)
+		Column = 2;
+	else if (Name.Find("W", ESearchCase::IgnoreCase) != -1)
+		Column = 3;
+
+	if (Column != -1)
+	{
+		Value.Append("[");
+		Value.AppendInt(Column);
+		Value.Append("]");
+	}
+	else
+	{
+		Error(FText::FromString("Failed to generate type for " + Name + " up to path " + Value), nullptr, nullptr);
+	}
+	return Value;
+}
+
+FString FHlslNiagaraTranslator::ComputeMatrixRowAccess(const FString& Name)
+{
+	FString Value;
+	int32 Row = -1;
+	if (Name.Find("Row0", ESearchCase::IgnoreCase) != -1)
+		Row = 0;
+	else if (Name.Find("Row1", ESearchCase::IgnoreCase) != -1)
+		Row = 1;
+	else if (Name.Find("Row2", ESearchCase::IgnoreCase) != -1)
+		Row = 2;
+	else if (Name.Find("Row3", ESearchCase::IgnoreCase) != -1)
+		Row = 3;
+
+	if (Row != -1)
+	{
+		Value.Append("[");
+		Value.AppendInt(Row);
+		Value.Append("]");
+	}
+	else
+	{
+		Error(FText::FromString("Failed to generate type for " + Name + " up to path " + Value), nullptr, nullptr);
+	}
+	return Value;
+}
+
+FString FHlslNiagaraTranslator::NamePathToString(const FString& Prefix, const FNiagaraTypeDefinition& RootType, const TArray<FName>& NamePath)
+{
+	// We need to deal with matrix parameters differently than any other type by using array syntax.
+	// As we recurse down the tree, we stay away of when we're dealing with a matrix and adjust 
+	// accordingly.
+	FString Value = Prefix;
+	FNiagaraTypeDefinition CurrentType = RootType;
+	bool bParentWasMatrix = (RootType == FNiagaraTypeDefinition::GetMatrix4Def());
+	int32 ParentMatrixRow = -1;
+	for (int32 i = 0; i < NamePath.Num(); i++)
+	{
+		FString Name = NamePath[i].ToString();
+		CurrentType = GetChildType(CurrentType, NamePath[i]);
+		// Found a matrix... brackets from here on out.
+		if (CurrentType == FNiagaraTypeDefinition::GetMatrix4Def())
+		{
+			bParentWasMatrix = true;
+			Value.Append("." + Name);
+		}
+		// Parent was a matrix, determine row..
+		else if (bParentWasMatrix && CurrentType == FNiagaraTypeDefinition::GetVec4Def())
+		{
+			Value.Append(ComputeMatrixRowAccess(Name));
+		}
+		// Parent was a matrix, determine column...
+		else if (bParentWasMatrix && CurrentType == FNiagaraTypeDefinition::GetFloatDef())
+		{
+			Value.Append(ComputeMatrixColumnAccess(Name));
+		}
+		// Handle all other valid types by just using "." 
+		else if (CurrentType.IsValid())
+		{
+			Value.Append("." + Name);
+		}
+		else
+		{
+			Error(FText::FromString("Failed to generate type for " + Name + " up to path " + Value), nullptr, nullptr);
+		}
+	}
+	return Value;
+}
+
+FString FHlslNiagaraTranslator::GenerateAssignment(const FNiagaraTypeDefinition& SrcPinType, const TArray<FName>& ConditionedSourcePath, const FNiagaraTypeDefinition& DestPinType, const TArray<FName>& ConditionedDestinationPath)
+{
+	FString SourceDefinition = NamePathToString("{1}", SrcPinType, ConditionedSourcePath);
+	FString DestinationDefinition = NamePathToString("{0}", DestPinType, ConditionedDestinationPath);
+
+	return DestinationDefinition + " = " + SourceDefinition;
+}
+
+void FHlslNiagaraTranslator::Convert(class UNiagaraNodeConvert* Convert, TArray <int32>& Inputs, TArray<int32>& Outputs)
+{
+	if (ValidateTypePins(Convert) == false)
+	{
+		return;
+	}
+
+	TArray<UEdGraphPin*> InputPins;
+	Convert->GetInputPins(InputPins);
+
+	TArray<UEdGraphPin*> OutputPins;
+	Convert->GetOutputPins(OutputPins);
+
+	// Generate outputs.
+	for (UEdGraphPin* OutputPin : OutputPins)
+	{
+		if (OutputPin->PinType.PinCategory == UEdGraphSchema_Niagara::PinCategoryType || 
+			OutputPin->PinType.PinCategory == UEdGraphSchema_Niagara::PinCategoryEnum)
+		{
+			FNiagaraTypeDefinition Type = Schema->PinToTypeDefinition(OutputPin);
+			int32 OutChunk = AddBodyChunk(GetUniqueSymbolName(*OutputPin->PinName), TEXT(""), Type);
+			Outputs.Add(OutChunk);
+		}
+	}
+
+	// Add an additional invalid output for the add pin which doesn't get compiled.
+	Outputs.Add(INDEX_NONE);
+
+	// Set output values based on connections.
+	for (FNiagaraConvertConnection& Connection : Convert->GetConnections())
+	{
+		int32 SourceIndex = GetPinIndexById(InputPins, Connection.SourcePinId);
+		int32 DestinationIndex = GetPinIndexById(OutputPins, Connection.DestinationPinId);
+		if (SourceIndex != INDEX_NONE && SourceIndex < Inputs.Num() && DestinationIndex != INDEX_NONE && DestinationIndex < Outputs.Num())
+		{
+			FNiagaraTypeDefinition SrcPinType = Schema->PinToTypeDefinition(InputPins[SourceIndex]);
+			TArray<FName> ConditionedSourcePath = ConditionPropertyPath(SrcPinType, Connection.SourcePath);
+			
+			FNiagaraTypeDefinition DestPinType = Schema->PinToTypeDefinition(OutputPins[DestinationIndex]);
+			TArray<FName> ConditionedDestinationPath = ConditionPropertyPath(DestPinType, Connection.DestinationPath);
+
+			FString ConvertDefinition = GenerateAssignment(SrcPinType, ConditionedSourcePath, DestPinType, ConditionedDestinationPath);
+			
+			TArray<int32> SourceChunks;
+			SourceChunks.Add(Outputs[DestinationIndex]);
+			SourceChunks.Add(Inputs[SourceIndex]);
+			AddBodyChunk(TEXT(""), ConvertDefinition, FNiagaraTypeDefinition::GetFloatDef(), SourceChunks);
+		}
+	}
+}
+
+void FHlslNiagaraTranslator::If(TArray<FNiagaraVariable>& Vars, int32 Condition, TArray<int32>& PathA, TArray<int32>& PathB, TArray<int32>& Outputs)
+{
+	int32 NumVars = Vars.Num();
+	check(PathA.Num() == NumVars);
+	check(PathB.Num() == NumVars);
+
+	TArray<FString> OutSymbols;
+	OutSymbols.Reserve(Vars.Num());
+	for (FNiagaraVariable& Var : Vars)
+	{
+		OutSymbols.Add(GetUniqueSymbolName(*(Var.GetName().ToString() + TEXT("_IfResult"))));
+		Outputs.Add(AddBodyChunk(OutSymbols.Last(), TEXT(""), Var.GetType(), true));
+	}
+	AddBodyChunk(TEXT(""), TEXT("if({0})\n\t{"), FNiagaraTypeDefinition::GetFloatDef(), Condition, false, false);
+	for (int32 i = 0; i < NumVars; ++i)
+	{
+		FNiagaraCodeChunk& OutChunk = CodeChunks[Outputs[i]];		
+		FNiagaraCodeChunk& BranchChunk = CodeChunks[AddBodyChunk(OutSymbols[i], TEXT("{0}"), OutChunk.Type, false)];
+		BranchChunk.AddSourceChunk(PathA[i]);
+	}	
+	AddBodyChunk(TEXT(""), TEXT("}\n\telse\n\t{"), FNiagaraTypeDefinition::GetFloatDef(), false, false);
+	for (int32 i = 0; i < NumVars; ++i)
+	{
+		FNiagaraCodeChunk& OutChunk = CodeChunks[Outputs[i]];
+		FNiagaraCodeChunk& BranchChunk = CodeChunks[AddBodyChunk(OutSymbols[i], TEXT("{0}"), OutChunk.Type, false)];
+		BranchChunk.AddSourceChunk(PathB[i]);
+	}
+	AddBodyChunk(TEXT(""), TEXT("}"), FNiagaraTypeDefinition::GetFloatDef(), false, false);
+
+	// Add an additional invalid output for the add pin which doesn't get compiled.
+	Outputs.Add(INDEX_NONE);
+}
+
+int32 FHlslNiagaraTranslator::CompilePin(UEdGraphPin* Pin)
+{
+	check(Pin);
+	int32 Ret = INDEX_NONE;
+	FNiagaraTypeDefinition TypeDef = Schema->PinToTypeDefinition(Pin);
+	if (Pin->Direction == EGPD_Input)
+	{
+		if (Pin->LinkedTo.Num() > 0)
+		{
+			Ret = CompileOutputPin(Pin->LinkedTo[0]);
+		}
+		else if (!Pin->bDefaultValueIsIgnored && Pin->PinType.PinCategory == UEdGraphSchema_Niagara::PinCategoryType)
+		{
+			if (TypeDef == FNiagaraTypeDefinition::GetParameterMapDef())
+			{
+				Error(FText::FromString(TEXT("Cannot create a constant ParameterMap!")), Cast<UNiagaraNode>(Pin->GetOwningNode()), Pin);
+				return INDEX_NONE;
+			}
+			else
+			{
+				//No connections to this input so add the default as a const expression.			
+				FNiagaraVariable PinVar = Schema->PinToNiagaraVariable(Pin, true);
+				return GetConstant(PinVar);
+			}
+		}
+		else if (!Pin->bDefaultValueIsIgnored && Pin->PinType.PinCategory == UEdGraphSchema_Niagara::PinCategoryEnum)
+		{
+			//No connections to this input so add the default as a const expression.			
+			FNiagaraVariable PinVar = Schema->PinToNiagaraVariable(Pin, true);
+			return GetConstant(PinVar);
+		}
+	}
+	else
+	{
+		Ret = CompileOutputPin(Pin);
+	}
+
+	return Ret;
+}
+
+int32 FHlslNiagaraTranslator::CompileOutputPin(UEdGraphPin* Pin)
+{
+	check(Pin->Direction == EGPD_Output);
+
+	int32 Ret = INDEX_NONE;
+
+	int32* Chunk = PinToCodeChunks.Last().Find(Pin);
+	if (Chunk)
+	{
+		Ret = *Chunk; //We've compiled this pin before. Return it's chunk.
+	}
+	else
+	{
+		//Otherwise we need to compile the node to get its output pins.
+		UNiagaraNode* Node = Cast<UNiagaraNode>(Pin->GetOwningNode());
+		if (ValidateTypePins(Node))
+		{
+			TArray<int32> Outputs;
+			TArray<UEdGraphPin*> OutputPins;
+			Node->GetOutputPins(OutputPins);
+			FHlslNiagaraTranslator* ThisTranslator = this;
+			Node->Compile(ThisTranslator, Outputs);
+			if (OutputPins.Num() == Outputs.Num())
+			{
+				for (int32 i = 0; i < Outputs.Num(); ++i)
+				{
+					//Cache off the pin.
+					//Can we allow the caching of local defaults in numerous function calls?
+					PinToCodeChunks.Last().Add(OutputPins[i], Outputs[i]);
+
+					if (Outputs[i] != INDEX_NONE)
+					{
+						//Grab the expression for the pin we're currently interested in. Otherwise we'd have to search the map for it.
+						Ret = OutputPins[i] == Pin ? Outputs[i] : Ret;
+					}
+				}
+			}
+			else
+			{
+				Error(LOCTEXT("IncorrectNumOutputsError", "Incorect number of outputs. Can possibly be fixed with a graph refresh."), Node, Pin);
+			}
+		}
+	}
+
+	return Ret;
+}
+
+void FHlslNiagaraTranslator::Error(FText ErrorText, UNiagaraNode* Node, UEdGraphPin* Pin)
+{
+	FString ErrorString = FString::Printf(TEXT("Node: @@ - Pin: @@ - %s - Callstack: %s"), *ErrorText.ToString(), *GetCallstack());
+	MessageLog.Error(*ErrorString, Node, Pin);
+}
+
+void FHlslNiagaraTranslator::Warning(FText WarningText, UNiagaraNode* Node, UEdGraphPin* Pin)
+{
+	FString WarnString = FString::Printf(TEXT("Node: @@ - Pin: @@ - %s - Callstack: %s"), *WarningText.ToString(), *GetCallstack());
+	MessageLog.Warning(*WarnString, Node, Pin);
+}
+
+bool FHlslNiagaraTranslator::GetFunctionParameter(const FNiagaraVariable& Parameter, int32& OutParam)const
+{
+	// Assume that it wasn't bound by default.
+	OutParam = INDEX_NONE;
+	if (const FFunctionContext* FunctionContext = FunctionCtx())
+	{
+		int32 ParamIdx = FunctionContext->Signature.Inputs.IndexOfByPredicate([&](const FNiagaraVariable& InVar) { return InVar.IsEquivalent(Parameter); });
+		if (ParamIdx != INDEX_NONE)
+		{
+			OutParam = FunctionContext->Inputs[ParamIdx];			
+		}
+		return true;
+	}
+	return false;
+}
+
+bool FHlslNiagaraTranslator::CanReadAttributes()const
+{
+	if (Script->IsParticleUpdateScript() || (Script->IsInterpolatedParticleSpawnScript() && !bInsideInterpolatedSpawnScript))
+	{
+		return true;
+	}
+	return false;
+}
+
+
+ENiagaraScriptUsage FHlslNiagaraTranslator::GetTargetUsage() const
+{
+	if (Script->IsInterpolatedParticleSpawnScript())
+	{
+		return ENiagaraScriptUsage::ParticleSpawnScript;
+	}
+	return Script->GetUsage();
+}
+
+int32 FHlslNiagaraTranslator::GetTargetOccurrence() const
+{
+	return Script->GetUsageIndex();
+}
+//////////////////////////////////////////////////////////////////////////
+
+FString FHlslNiagaraTranslator::GetHlslDefaultForType(FNiagaraTypeDefinition Type)
+{
+	if (Type == FNiagaraTypeDefinition::GetFloatDef())
+	{
+		return "(0.0)";
+	}
+	else if (Type == FNiagaraTypeDefinition::GetVec2Def())
+	{
+		return "(0.0,0.0)";
+	}
+	else if (Type == FNiagaraTypeDefinition::GetVec3Def())
+	{
+		return "(0.0,0.0,0.0)";
+	}
+	else if (Type == FNiagaraTypeDefinition::GetVec4Def())
+	{
+		return "(0.0,0.0,0.0,0.0)";
+	}
+	else if (Type == FNiagaraTypeDefinition::GetColorDef())
+	{
+		return "(1.0,1.0,1.0,1.0)";
+	}
+	else if (Type == FNiagaraTypeDefinition::GetIntDef())
+	{
+		return "(0)";
+	}
+	else if (Type == FNiagaraTypeDefinition::GetBoolDef())
+	{
+		return "(true)";
+	}
+	else
+	{
+		return Type.GetName();
+	}
+}
+
+bool FHlslNiagaraTranslator::IsBuiltInHlslType(FNiagaraTypeDefinition Type)
+{
+	return	
+		Type == FNiagaraTypeDefinition::GetFloatDef() ||
+		Type == FNiagaraTypeDefinition::GetVec2Def() ||
+		Type == FNiagaraTypeDefinition::GetVec3Def() ||
+		Type == FNiagaraTypeDefinition::GetVec4Def() ||
+		Type == FNiagaraTypeDefinition::GetColorDef() ||
+		Type == FNiagaraTypeDefinition::GetMatrix4Def() ||
+		Type == FNiagaraTypeDefinition::GetIntDef() ||
+		Type.GetStruct() == FNiagaraTypeDefinition::GetIntStruct() ||
+		Type == FNiagaraTypeDefinition::GetBoolDef();
+}
+
+FString FHlslNiagaraTranslator::GetStructHlslTypeName(FNiagaraTypeDefinition Type)
+{
+	if (Type.IsValid() == false)
+	{
+		return "undefined";
+	}
+	else if (Type == FNiagaraTypeDefinition::GetFloatDef())
+	{
+		return "float";
+	}
+	else if (Type == FNiagaraTypeDefinition::GetVec2Def())
+	{
+		return "float2";
+	}
+	else if (Type == FNiagaraTypeDefinition::GetVec3Def())
+	{
+		return "float3";
+	}
+	else if (Type == FNiagaraTypeDefinition::GetVec4Def() || Type == FNiagaraTypeDefinition::GetColorDef())
+	{
+		return "float4";
+	}
+	else if (Type == FNiagaraTypeDefinition::GetMatrix4Def())
+	{
+		return "float4x4";
+	}
+	else if (Type == FNiagaraTypeDefinition::GetIntDef() || Type.GetEnum() != nullptr)
+	{
+		return "int";
+	}
+	else if (Type == FNiagaraTypeDefinition::GetBoolDef())
+	{
+		return "bool";
+	}
+	else if (Type == FNiagaraTypeDefinition::GetParameterMapDef())
+	{
+		return "FParamMap0";
+	}
+	else
+	{
+		return Type.GetName();
+	}
+}
+
+FString FHlslNiagaraTranslator::GetPropertyHlslTypeName(const UProperty* Property)
+{
+	if (Property->IsA(UFloatProperty::StaticClass()))
+	{
+		return "float";
+	}
+	else if (Property->IsA(UIntProperty::StaticClass()))
+	{
+		return "int";
+	}
+	else if (Property->IsA(UUInt32Property::StaticClass()))
+	{
+		return "int";
+	}
+	else if (Property->IsA(UStructProperty::StaticClass()))
+	{
+		const UStructProperty* StructProp = Cast<const UStructProperty>(Property);
+		return GetStructHlslTypeName(StructProp->Struct);
+	}
+	else if (Property->IsA(UEnumProperty::StaticClass()))
+	{
+		const UEnumProperty* EnumProp = Cast<const UEnumProperty>(Property);
+		return "int";
+	}
+	else
+	{
+		check(false);	// unknown type
+		return TEXT("UnknownType");
+	}
+}
+
+FString FHlslNiagaraTranslator::BuildHLSLStructDecl(FNiagaraTypeDefinition Type)
+{
+	if (!IsBuiltInHlslType(Type))
+	{
+		FString StructName = GetStructHlslTypeName(Type);
+
+		FString Decl = "struct " + StructName + "\n{\n";
+		for (TFieldIterator<UProperty> PropertyIt(Type.GetStruct(), EFieldIteratorFlags::IncludeSuper); PropertyIt; ++PropertyIt)
+		{
+			UProperty* Property = *PropertyIt;
+			Decl += TEXT("\t") + GetPropertyHlslTypeName(Property) + TEXT(" ") + Property->GetName() + TEXT(";\n");
+		}
+		Decl += "};\n\n";
+		return Decl;
+	}
+
+	return TEXT("");
+}
+
+bool FHlslNiagaraTranslator::IsHlslBuiltinVector(FNiagaraTypeDefinition Type)
+{
+	if ((Type == FNiagaraTypeDefinition::GetVec2Def()) ||
+		(Type == FNiagaraTypeDefinition::GetVec3Def()) ||
+		(Type == FNiagaraTypeDefinition::GetVec4Def()) ||
+		(Type == FNiagaraTypeDefinition::GetColorDef()))
+	{
+		return true;
+	}
+	return false;
+}
+
+
+bool FHlslNiagaraTranslator::AddStructToDefinitionSet(const FNiagaraTypeDefinition& TypeDef)
+{
+	// First make sure that this is a type that we do need to define...
+	if (IsBuiltInHlslType(TypeDef))
+	{
+		return true;
+	}
+
+	if (TypeDef == FNiagaraTypeDefinition::GetGenericNumericDef())
+	{
+		return false;
+	}
+
+	// We build these types on-the-fly.
+	if (TypeDef == FNiagaraTypeDefinition::GetParameterMapDef())
+	{
+		return true;
+	}
+
+	// Now make sure that we don't have any other struct types within our struct. Add them prior to the struct in question to make sure
+	// that the syntax works out properly.
+	const UScriptStruct* Struct = TypeDef.GetScriptStruct();
+	if (Struct != nullptr)
+	{
+		// We need to recursively dig through the struct to get at the lowest level of the input struct, which
+		// could be a native type.
+		for (TFieldIterator<UProperty> PropertyIt(Struct, EFieldIteratorFlags::IncludeSuper); PropertyIt; ++PropertyIt)
+		{
+			const UProperty* Property = *PropertyIt;
+			const UStructProperty* StructProp = Cast<const UStructProperty>(Property);
+			if (StructProp)
+			{
+				if (!AddStructToDefinitionSet(StructProp->Struct))
+				{
+					return false;
+				}
+			}
+		}
+		
+		// Add the new type def
+		StructsToDefine.AddUnique(TypeDef);
+	}
+
+	return true;
+}
+
+TArray<FName> FHlslNiagaraTranslator::ConditionPropertyPath(const FNiagaraTypeDefinition& Type, const TArray<FName>& InPath)
+{
+	// TODO: Build something more extensible and less hard coded for path conditioning.
+	const UScriptStruct* Struct = Type.GetScriptStruct();
+	if (InPath.Num() == 0) // Pointing to the root
+	{
+		return TArray<FName>();
+	}
+	else if (IsHlslBuiltinVector(Type))
+	{
+		checkf(InPath.Num() == 1, TEXT("Invalid path for vector"));
+		TArray<FName> ConditionedPath;
+		ConditionedPath.Add(*InPath[0].ToString().ToLower());
+		return ConditionedPath;
+	}
+	else if (FNiagaraTypeDefinition::IsScalarDefinition(Struct))
+	{
+		return TArray<FName>();
+	}
+	else if (Struct != nullptr)
+	{
+		// We need to recursively dig through the struct to get at the lowest level of the input path specified, which
+		// could be a native type.
+		for (TFieldIterator<UProperty> PropertyIt(Struct, EFieldIteratorFlags::IncludeSuper); PropertyIt; ++PropertyIt)
+		{
+			const UProperty* Property = *PropertyIt;
+			const UStructProperty* StructProp = Cast<const UStructProperty>(Property);
+			// The names match, but even then things might not match up properly..
+			if (InPath[0].ToString() == Property->GetName())
+			{
+				// The names match and this is a nested type, so we can keep digging...
+				if (StructProp != nullptr)
+				{
+					// If our path continues onward, keep recursively digging. Otherwise, just return where we've gotten to so far.
+					if (InPath.Num() > 1)
+					{
+						TArray<FName> ReturnPath;
+						ReturnPath.Add(InPath[0]);
+						TArray<FName> Subset = InPath;
+						Subset.RemoveAt(0);
+						TArray<FName> Children = ConditionPropertyPath(StructProp->Struct, Subset);
+						for (const FName& Child : Children)
+						{
+							ReturnPath.Add(Child);
+						}
+						return ReturnPath;
+					}
+					else
+					{
+						TArray<FName> ConditionedPath;
+						ConditionedPath.Add(InPath[0]);
+						return ConditionedPath;
+					}
+				}
+			}
+		}
+		return InPath;
+	}
+	return InPath;
+}
+//////////////////////////////////////////////////////////////////////////
+
+
+FString FHlslNiagaraTranslator::CompileDataInterfaceFunction(UNiagaraDataInterface* DataInterface, FNiagaraFunctionSignature& Signature)
+{
+	//For now I'm compiling data interface functions like this. 
+	//Not the prettiest thing in the world but it'll suffice for now.
+
+	if (UNiagaraDataInterfaceCurve* Curve = Cast<UNiagaraDataInterfaceCurve>(DataInterface))
+	{
+		//For now, VM only which needs no body. GPU will need a body.
+		return TEXT("");
+	}
+	else if (UNiagaraDataInterfaceVectorCurve* VecCurve = Cast<UNiagaraDataInterfaceVectorCurve>(DataInterface))
+	{
+		//For now, VM only which needs no body. GPU will need a body.
+		return TEXT("");
+	}
+	else if (UNiagaraDataInterfaceColorCurve* ColorCurve = Cast<UNiagaraDataInterfaceColorCurve>(DataInterface))
+	{
+		//For now, VM only which needs no body. GPU will need a body.
+		return TEXT("");
+	}
+	else if (UNiagaraDataInterfaceVector2DCurve* Vec2DCurve = Cast<UNiagaraDataInterfaceVector2DCurve>(DataInterface))
+	{
+		//For now, VM only which needs no body. GPU will need a body.
+		return TEXT("");
+	}
+	else if (UNiagaraDataInterfaceVector4Curve* Vec4Curve = Cast<UNiagaraDataInterfaceVector4Curve>(DataInterface))
+	{
+		//For now, VM only which needs no body. GPU will need a body.
+		return TEXT("");
+	}
+	else if (UNiagaraDataInterfaceStaticMesh* Mesh = Cast<UNiagaraDataInterfaceStaticMesh>(DataInterface))
+	{
+		return TEXT("");
+	}
+	else if (UNiagaraDataInterfaceCurlNoise* Noise = Cast<UNiagaraDataInterfaceCurlNoise>(DataInterface))
+	{
+		return TEXT("");
+	}
+	else
+	{		
+		return TEXT("");
+		check(0);
+	}
+}
+
+
+
+#undef LOCTEXT_NAMESPACE
