Copying //UE4/Orion-Staging to //UE4/Main (Source: //Orion/Dev-General @ 3606378)

diff --git a/Engine/Plugins/FX/Niagara/Source/NiagaraEditor/Private/NiagaraHlslTranslator.h b/Engine/Plugins/FX/Niagara/Source/NiagaraEditor/Private/NiagaraHlslTranslator.h
new file mode 100644
--- /dev/null
+++ b/Engine/Plugins/FX/Niagara/Source/NiagaraEditor/Private/NiagaraHlslTranslator.h
@@ -0,0 +1,511 @@
+// Copyright 1998-2017 Epic Games, Inc. All Rights Reserved.
+
+#pragma once
+
+#include "CoreMinimal.h"
+#include "Kismet2/CompilerResultsLog.h"
+#include "NiagaraScript.h"
+#include "NiagaraParameterMapHistory.h"
+#include "EdGraphUtilities.h"
+#include "UObject/UObjectHash.h"
+#include "ComponentReregisterContext.h"
+#include "NiagaraShaderCompilationManager.h"
+#include "NiagaraDataInterface.h"
+#include "TickableEditorObject.h"
+
+class Error;
+class UNiagaraGraph;
+class UNiagaraNode;
+class UNiagaraNodeFunctionCall;
+class UNiagaraNodeOutput;
+class UNiagaraScriptSource;
+class UNiagaraNodeEmitter;
+
+
+// handles finished shader compile jobs, applying of the shaders to their scripts, and some error handling
+//
+class FNiagaraShaderProcessorTickable : FTickableEditorObject
+{
+	virtual bool IsTickable() const override
+	{
+		return true;
+	}
+
+	virtual void Tick(float DeltaSeconds) override
+	{
+		GNiagaraShaderCompilationManager.Tick(DeltaSeconds);
+		GNiagaraShaderCompilationManager.ProcessAsyncResults();
+	}
+
+	virtual TStatId GetStatId() const
+	{
+		RETURN_QUICK_DECLARE_CYCLE_STAT(FNiagaraShaderQueueTickable, STATGROUP_Tickables);
+	}
+};
+
+
+UENUM()
+enum class ENiagaraDataSetAccessMode : uint8
+{
+	/** Data set reads and writes use shared counters to add and remove the end of available data. Writes are conditional and read */
+	AppendConsume,
+	/** Data set is accessed directly at a specific index. */
+	Direct,
+
+	Num UMETA(Hidden),
+};
+
+
+/** Defines information about the results of a Niagara script compile. */
+struct FNiagaraTranslateResults
+{
+	/** Whether or not HLSL generation was successful */
+	bool bHLSLGenSucceeded;
+
+	/** A results log with messages, warnings, and errors which occurred during the compile. */
+	FCompilerResultsLog* MessageLog;
+
+	/** A string representation of the compilation output. */
+	FString OutputHLSL;
+
+	FNiagaraTranslateResults(FCompilerResultsLog* InMessageLog = nullptr)
+		: MessageLog(InMessageLog)
+	{
+	}
+
+	static ENiagaraScriptCompileStatus TranslateResultsToSummary(const FNiagaraTranslateResults* CompileResults);
+};
+
+
+/** Data which is generated from the hlsl by the VectorVMBackend and fed back into the */
+struct FNiagaraTranslatorOutput
+{
+	FNiagaraTranslatorOutput() : NumUserPtrs(0), bReadsAttributeData(false){}
+
+	/** All external parameters used in the graph. */
+	FNiagaraParameters Parameters;
+
+	/** All internal constants used in the graph. */
+	FNiagaraParameters InternalParameters;
+
+	/** Data sets this script reads. */
+	TArray<FNiagaraDataSetProperties> DataSetReads;
+
+	/** Data sets this script writes. */
+	TArray<FNiagaraDataSetProperties> DataSetWrites;
+
+	/* Per instance attributes */
+	TArray<FNiagaraVariable> Attributes;
+
+	FNiagaraScriptDataUsageInfo DataUsage;
+
+	TArray<FNiagaraScriptDataInterfaceInfo> DataInterfaceInfo;
+
+	/** Ordered table of functions actually called by the VM script. */
+	struct FCalledVMFunction
+	{
+		FString Name;
+		TArray<bool> InputParamLocations;
+		int32 NumOutputs;
+		FCalledVMFunction() :NumOutputs(0) {}
+	};
+	TArray<FCalledVMFunction> CalledVMFunctionTable;
+
+	FString Errors;
+
+	/** Number of user pointers. For passing things like data interface instance data etc. */
+	int32 NumUserPtrs;
+
+	bool bReadsAttributeData;
+};
+
+
+
+enum class ENiagaraCodeChunkMode : uint8
+{
+	Uniform,
+	Source,
+	Body,
+	Num,
+};
+
+FORCEINLINE uint32 GetTypeHash(const FNiagaraFunctionSignature& Sig)
+{
+	uint32 Hash = GetTypeHash(Sig.Name);
+	for (const FNiagaraVariable& Var : Sig.Inputs)
+	{
+		Hash = HashCombine(Hash, GetTypeHash(Var));
+	}
+	for (const FNiagaraVariable& Var : Sig.Outputs)
+	{
+		Hash = HashCombine(Hash, GetTypeHash(Var));
+	}
+	Hash = HashCombine(Hash, GetTypeHash(Sig.OwnerName));
+	return Hash;
+}
+
+struct FNiagaraCodeChunk
+{
+	/** Symbol name for the chunk. Cam be empty for some types of chunk. */
+	FString SymbolName;
+	/** Format definition for incorporating SourceChunks into the final code for this chunk. */
+	FString Definition;
+	/** The returned data type of this chunk. */
+	FNiagaraTypeDefinition Type;
+	/** If this chunk should declare it's symbol name. */
+	bool bDecl;
+	/** If the chunk is unterminated (no semicolon, because it's a scope or similar */
+	bool bIsTerminated;
+	/** Chunks used as input for this chunk. */
+	TArray<int32> SourceChunks;
+
+	ENiagaraCodeChunkMode Mode;
+
+	FNiagaraCodeChunk()
+		: bDecl(true)
+		, bIsTerminated(true)
+		, Mode(ENiagaraCodeChunkMode::Num)
+	{
+		Type = FNiagaraTypeDefinition::GetFloatDef();
+	}
+
+	void AddSourceChunk(int32 ChunkIdx)
+	{
+		SourceChunks.Add(ChunkIdx);
+	}
+
+	int32 GetSourceChunk(int32 i)
+	{
+		return SourceChunks[i];
+	}
+
+	void ReplaceSourceIndex(int32 SourceIdx, int32 NewIdx)
+	{
+		SourceChunks[SourceIdx] = NewIdx;
+	}
+
+	bool operator==(const FNiagaraCodeChunk& Other)
+	{
+		return SymbolName == Other.SymbolName &&
+			Definition == Other.Definition &&
+			Mode == Other.Mode &&
+			Type == Other.Type &&
+			bDecl == Other.bDecl &&
+			SourceChunks == Other.SourceChunks;
+	}
+};
+
+class NIAGARAEDITOR_API FHlslNiagaraTranslatorOptions
+{
+public:
+	FHlslNiagaraTranslatorOptions()
+		: SimTarget(ENiagaraSimTarget::CPUSim)
+	{
+
+	}
+
+	ENiagaraSimTarget SimTarget;
+
+	/** Any parameters in these namespaces will be pulled from an "InstanceParameters" dataset rather than from the uniform table. */
+	TArray<FString> InstanceParameterNamespaces;
+};
+
+class NIAGARAEDITOR_API FHlslNiagaraTranslator
+{
+public:
+
+	struct FDataSetAccessInfo
+	{
+		//Variables accessed.
+		TArray<FNiagaraVariable> Variables;
+		/** Code chunks relating to this access. */
+		TArray<int32> CodeChunks;
+	};
+
+protected:
+
+	/** The script we are compiling. */
+	UNiagaraScript* Script;
+
+	FHlslNiagaraTranslatorOptions TranslationOptions;
+
+	const class UEdGraphSchema_Niagara* Schema;
+
+	/** The set of all generated code chunks for this script. */
+	TArray<FNiagaraCodeChunk> CodeChunks;
+
+	/** Array of code chunks of each different type. */
+	TArray<int32> ChunksByMode[(int32)ENiagaraCodeChunkMode::Num];
+
+	/**
+	Map of Pins to compiled code chunks. Allows easy reuse of previously compiled pins.
+	A stack so that we can track pin reuse within function calls but not have cached pins cross talk with subsequent calls to the same funciton.
+	*/
+	TArray<TMap<UEdGraphPin*, int32>> PinToCodeChunks;
+
+	/** The combined output of the compilation of this script. This is temporary and will be reworked soon. */
+	FNiagaraTranslatorOutput CompilationOutput;
+
+	/** Message log. Automatically handles marking the NodeGraph with errors. */
+	FCompilerResultsLog MessageLog;
+
+	/** Captures information about a script compile. */
+	FNiagaraTranslateResults TranslateResults;
+
+	TMap<FName, uint32> GeneratedSymbolCounts;
+
+	FDataSetAccessInfo InstanceRead;
+	FDataSetAccessInfo InstanceWrite;
+
+	TMap<FNiagaraDataSetID, TMap<int32, FDataSetAccessInfo>> DataSetReadInfo[(int32)ENiagaraDataSetAccessMode::Num];
+	TMap<FNiagaraDataSetID, TMap<int32, FDataSetAccessInfo>> DataSetWriteInfo[(int32)ENiagaraDataSetAccessMode::Num];
+	TMap<FNiagaraDataSetID, int32> DataSetWriteConditionalInfo[(int32)ENiagaraDataSetAccessMode::Num];
+
+	FString GetDataSetAccessSymbol(FNiagaraDataSetID DataSet, int32 IndexChunk, bool bRead);
+	FORCEINLINE FNiagaraDataSetID GetInstanceDataSetID()const { return FNiagaraDataSetID(TEXT("DataInstance"), ENiagaraDataSetType::ParticleData); }
+	FORCEINLINE FNiagaraDataSetID GetSystemEngineDataSetID()const { return FNiagaraDataSetID(TEXT("Engine"), ENiagaraDataSetType::ParticleData); }
+	FORCEINLINE FNiagaraDataSetID GetSystemUserDataSetID()const { return FNiagaraDataSetID(TEXT("User"), ENiagaraDataSetType::ParticleData); }
+
+	/** All functions called in the script. */
+	TMap<FNiagaraFunctionSignature, FString> Functions;
+	/** Map of function graphs we've seen before and already pre-processed. */
+	TMap<const UNiagaraGraph*, UNiagaraGraph*> PreprocessedFunctions;
+
+	void RegisterFunctionCall(ENiagaraScriptUsage ScriptUsage, const FString& InName, const FString& InFullName, UNiagaraScriptSource* Source, FNiagaraFunctionSignature& InSignature, TArray<int32>& Inputs, const TArray<UEdGraphPin*>& CallInputs, const TArray<UEdGraphPin*>& CallOutputs,
+		FNiagaraFunctionSignature& OutSignature);
+	void GenerateFunctionCall(FNiagaraFunctionSignature& FunctionSignature, TArray<int32>& Inputs, TArray<int32>& Outputs);
+	FString GetFunctionSignature(const FNiagaraFunctionSignature& Sig);
+
+	/** Compiles an output Pin on a graph node. Caches the result for any future inputs connected to it. */
+	int32 CompileOutputPin(UEdGraphPin* Pin);
+
+	void WriteDataSetContextVars(TMap<FNiagaraDataSetID, TMap<int32, FDataSetAccessInfo>>& DataSetAccessInfo, bool bRead, FString &HlslOutput);
+	void WriteDataSetStructDeclarations(TMap<FNiagaraDataSetID, TMap<int32, FDataSetAccessInfo>>& DataSetAccessInfo, bool bRead, FString &HlslOutput);
+	void DecomposeVariableAccess(const UStruct* Struct, bool bRead, FString IndexSymbol, FString HLSLString);
+
+	FString GetUniqueSymbolName(FName BaseName);
+
+	/** Stack of all function params. */
+	struct FFunctionContext
+	{
+		FString Name;
+		FNiagaraFunctionSignature& Signature;
+		TArray<int32>& Inputs;
+		FFunctionContext(const FString& InName, FNiagaraFunctionSignature& InSig, TArray<int32>& InInputs)
+			: Name(InName)
+			, Signature(InSig)
+			, Inputs(InInputs)
+		{}
+	};
+	TArray<FFunctionContext> FunctionContextStack;
+	const FFunctionContext* FunctionCtx()const { return FunctionContextStack.Num() > 0 ? &FunctionContextStack.Last() : nullptr; }
+	void EnterFunction(const FString& Name, FNiagaraFunctionSignature& Signature, TArray<int32>& Inputs);
+	void ExitFunction();
+	FString GetCallstack();
+
+	void EnterStatsScope(FNiagaraStatScope StatScope);
+	void ExitStatsScope();
+
+	FString GeneratedConstantString(float Constant);
+	FString GeneratedConstantString(FVector4 Constant);
+
+	/* Add a chunk that is not written to the source, only used as a source chunk for others. */
+	int32 AddSourceChunk(FString SymbolName, const FNiagaraTypeDefinition& Type);
+
+	/** Add a chunk defining a uniform value. */
+	int32 AddUniformChunk(FString SymbolName, const FNiagaraTypeDefinition& Type);
+
+	/* Add a chunk that is written to the body of the shader code. */
+	int32 AddBodyChunk(FString SymbolName, FString Definition, const FNiagaraTypeDefinition& Type, TArray<int32>& SourceChunks, bool bDecl = true, bool bIsTerminated = true);
+	int32 AddBodyChunk(FString SymbolName, FString Definition, const FNiagaraTypeDefinition& Type, int32 SourceChunk, bool bDecl = true, bool bIsTerminated = true);
+	int32 AddBodyChunk(FString SymbolName, FString Definition, const FNiagaraTypeDefinition& Type, bool bDecl = true, bool bIsTerminated = true);
+
+
+	FString GetFunctionDefinitions();
+public:
+
+	FHlslNiagaraTranslator();
+	virtual ~FHlslNiagaraTranslator() {}
+
+	virtual const FNiagaraTranslateResults &Translate(UNiagaraScript* InScript, FHlslNiagaraTranslatorOptions Options);
+	FNiagaraTranslatorOutput &GetTranslateOutput() { return CompilationOutput; }
+
+	virtual int32 CompilePin(UEdGraphPin* Pin);
+
+	virtual int32 RegisterDataInterface(FNiagaraVariable& Var, UNiagaraDataInterface* DataInterface);
+
+	virtual void Operation(class UNiagaraNodeOp* Operation, TArray<int32>& Inputs, TArray<int32>& Outputs);
+	virtual void Output(UNiagaraNodeOutput* OutputNode, const TArray<int32>& ComputedInputs);
+
+	virtual int32 GetParameter(const FNiagaraVariable& Parameter);
+
+	virtual int32 GetAttribute(const FNiagaraVariable& Attribute);
+
+	virtual int32 GetConstant(const FNiagaraVariable& Constant);
+	
+	virtual void ReadDataSet(const FNiagaraDataSetID DataSet, const TArray<FNiagaraVariable>& Variable, ENiagaraDataSetAccessMode AccessMode, int32 InputChunk, TArray<int32>& Outputs);
+	virtual void WriteDataSet(const FNiagaraDataSetID DataSet, const TArray<FNiagaraVariable>& Variable, ENiagaraDataSetAccessMode AccessMode, const TArray<int32>& Inputs);
+	virtual void ParameterMapSet(class UNiagaraNodeParameterMapSet* SetNode, TArray<int32>& Inputs, TArray<int32>& Outputs);
+	virtual void ParameterMapGet(class UNiagaraNodeParameterMapGet* GetNode, TArray<int32>& Inputs, TArray<int32>& Outputs);
+	virtual void ParameterCollection(class UNiagaraNodeParameterCollection* CollectionNode, TArray<int32>& Outputs);
+	virtual void Emitter(class UNiagaraNodeEmitter* GetNode, TArray<int32>& Inputs, TArray<int32>& Outputs);
+
+	void DefineInterpolatedParametersFunction(FString &HlslOutput);
+	void DefineDataSetReadFunction(FString &HlslOutput, TArray<FNiagaraDataSetID> &ReadDataSets);
+	void DefineDataSetWriteFunction(FString &HlslOutput, TArray<FNiagaraDataSetProperties> &WriteDataSets, TArray<int32>& WriteConditionVarIndices);
+	void DefineMain(FString &HLSLOutput, TArray<TArray<FNiagaraVariable>* > &InstanceReadVars, TArray<FNiagaraDataSetID>& ReadIds, TArray<TArray<FNiagaraVariable>* > &InstanceWriteVars, TArray<FNiagaraDataSetID>& WriteIds);
+	void DefineDataSetVariableReads(FString &HLSLOutput, FNiagaraDataSetID& Id, int32 DataSetIndex, TArray<FNiagaraVariable> & ReadVars);
+	void DefineDataSetVariableWrites(FString &HlslOutput, FNiagaraDataSetID& Id, int32 DataSetIndex, TArray<FNiagaraVariable>& WriteVars);
+	void DefineDataInterfaceHLSL(FString &HlslOutput);
+	TArray< TArray<DIGPUBufferParamDescriptor> >&GetDatainterfaceBufferDescriptors() { return DIBufferDescriptors; }
+
+	// Format string should have up to 5 entries, {{0} = Computed Variable Suffix, {1} = Float or Int, {2} = Data Set Index, {3} = Register Index, {4} Default value for that type.
+	void GatherVariableForDataSetAccess(const FNiagaraVariable& Variable, FString Format, int32& RegisterIdx, int32 DataSetIndex, FString InstanceIdxSymbol, FString &HlslOutput);
+	void GatherComponentsForDataSetAccess(const UScriptStruct* Struct, FString VariableSymbol, bool bMatrixRoot, TArray<FString>& Components, TArray<ENiagaraBaseTypes>& Types);
+
+	FString CompileDataInterfaceFunction(UNiagaraDataInterface* DataInterface, FNiagaraFunctionSignature& Signature);
+
+	virtual void FunctionCall(UNiagaraNodeFunctionCall* FunctionNode, TArray<int32>& Inputs, TArray<int32>& Outputs);
+
+	virtual void Convert(class UNiagaraNodeConvert* Convert, TArray <int32>& Inputs, TArray<int32>& Outputs);
+	virtual void If(TArray<FNiagaraVariable>& Vars, int32 Condition, TArray<int32>& PathA, TArray<int32>& PathB, TArray<int32>& Outputs);
+
+	virtual void Error(FText ErrorText, UNiagaraNode* Node, UEdGraphPin* Pin);
+	virtual void Warning(FText WarningText, UNiagaraNode* Node, UEdGraphPin* Pin);
+
+	virtual bool GetFunctionParameter(const FNiagaraVariable& Parameter, int32& OutParam)const;
+
+	virtual bool CanReadAttributes()const;
+	virtual ENiagaraScriptUsage GetTargetUsage() const;
+	int32 GetTargetOccurrence() const;
+
+	static bool IsBuiltInHlslType(FNiagaraTypeDefinition Type);
+	static FString GetStructHlslTypeName(FNiagaraTypeDefinition Type);
+	static FString GetPropertyHlslTypeName(const UProperty* Property);
+	static FString BuildHLSLStructDecl(FNiagaraTypeDefinition Type);
+	static FString GetHlslDefaultForType(FNiagaraTypeDefinition Type);
+	static bool IsHlslBuiltinVector(FNiagaraTypeDefinition Type);
+	static TArray<FName> ConditionPropertyPath(const FNiagaraTypeDefinition& Type, const TArray<FName>& InPath);
+
+
+	static FString GetSanitizedSymbolName(FString SymbolName);
+
+	bool AddStructToDefinitionSet(const FNiagaraTypeDefinition& TypeDef);
+
+	FString &GetTranslatedHLSL()
+	{
+		return HlslOutput;
+	}
+
+	static FString GetFunctionSignatureSymbol(const FNiagaraFunctionSignature& Sig);
+
+	/** ID to allow visitors over the node graph to mark nodes as visited. */
+	int32 VisitID;
+
+private:
+	bool ShouldConsiderTargetParameterMap(ENiagaraScriptUsage InUsage) const;
+	FString BuildParameterMapHlslDefinitions(TArray<FNiagaraVariable>& PrimaryDataSetOutputEntries, TArray<TPair<FNiagaraVariable, const UEdGraphPin*> >& MissingDefaults);
+	FString BuildMissingDefaults(TArray<TPair<FNiagaraVariable, const UEdGraphPin*> >& MissingDefaults);
+
+	void HandleNamespacedExternalVariablesToDataSetRead(TArray<FNiagaraVariable>& InDataSetVars, FString InNamespaceStr);
+
+	FString ComputeMatrixColumnAccess(const FString& Name);
+	FString ComputeMatrixRowAccess(const FString& Name);
+	
+	// Add a raw float constant chunk
+	int32 GetConstantFloat(float InValue);
+
+	FNiagaraTypeDefinition GetChildType(const FNiagaraTypeDefinition& BaseType, const FName& PropertyName);
+	FString NamePathToString(const FString& Prefix, const FNiagaraTypeDefinition& RootType, const TArray<FName>& NamePath);
+	FString GenerateAssignment(const FNiagaraTypeDefinition& SrcType, const TArray<FName>& SrcPath, const FNiagaraTypeDefinition& DestType, const TArray<FName>& DestPath);
+
+	//Generates the code for the passed chunk.
+	FString GetCode(FNiagaraCodeChunk& Chunk);
+	FString GetCode(int32 ChunkIdx);
+	//Retreives the code for this chunk being used as a source for another chunk
+	FString GetCodeAsSource(int32 ChunkIdx);
+
+	// Convert a variable with actual data into a constant string
+	FString GenerateConstantString(const FNiagaraVariable& Constant);
+
+	// Takes the current script state (interpolated or not) and determines the correct context variable.
+	FString GetParameterMapInstanceName(int32 ParamMapHistoryIdx);
+	
+	// Register a System/Engine/read-only variable in its namespaced form
+	bool ParameterMapRegisterExternalConstantNamespaceVariable(FNiagaraVariable InVariable, UNiagaraNode* InNode, int32 InParamMapHistoryIdx, int32& Output, UEdGraphPin* InDefaultPin);
+
+	// Register an attribute in its non namespaced form
+	bool ParameterMapRegisterNamespaceAttributeVariable(const FNiagaraVariable& InVariable, UNiagaraNode* InNode, int32 InParamMapHistoryIdx, int32& Output);
+	// Register an attribute in its namespaced form
+	bool ParameterMapRegisterUniformAttributeVariable(const FNiagaraVariable& InVariable, UNiagaraNode* InNode, int32 InParamMapHistoryIdx, int32& Output);
+
+	bool ValidateTypePins(UNiagaraNode* NodeToValidate);
+	void GenerateFunctionSignature(ENiagaraScriptUsage ScriptUsage, FString InName, const FString& InFullName, UNiagaraGraph* FuncGraph, TArray<int32>& Inputs, bool bHadNumericInputs, FNiagaraFunctionSignature& OutSig)const;
+
+	UNiagaraGraph* CloneGraphAndPrepareForCompilation(UNiagaraScript* InScript, UNiagaraScriptSource* InSource, bool bClearErrors);
+
+	bool ShouldInterpolateParameter(const FNiagaraVariable& Parameter);
+
+	bool IsBulkSystemScript() const;
+
+	/** Map of symbol names to count of times it's been used. Used for generating unique symbol names. */
+	TMap<FName, uint32> SymbolCounts;
+
+	//Set of non-builtin structs we have to define in hlsl.
+	TArray<FNiagaraTypeDefinition> StructsToDefine;
+
+	// Keep track of all the paths that the parameter maps can take through the graph.
+	TArray<FNiagaraParameterMapHistory> ParamMapHistories;
+
+	// Keep track of the other output nodes in the graph's histories so that we can make sure to 
+	// create any variables that are needed downstream.
+	TArray<FNiagaraParameterMapHistory> OtherOutputParamMapHistories;
+
+	// Make sure that the function call names match up on the second traversal.
+	FNiagaraParameterMapHistoryBuilder ActiveHistoryForFunctionCalls;
+
+	// Synced to the ParamMapHistories.
+	TArray<TArray<int32> > ParamMapSetVariablesToChunks;
+
+	// Synced to the System uniforms encountered for parameter maps thus far.
+	TMap<FName, int32> ParamMapDefinedSystemVarsToUniformChunks; // Map from the defined constants to the uniform chunk expressing them (i.e. have we encountered before in this graph?)
+	TMap<FName, FNiagaraVariable> ParamMapDefinedSystemToNamespaceVars; // Map from defined constants to the Namespaced variable expressing it.
+
+	// Synced to the EmitterParameter uniforms encountered for parameter maps thus far.
+	TMap<FName, int32> ParamMapDefinedEmitterParameterVarsToUniformChunks; // Map from the variable name exposed by the emitter as a parameter to the uniform chunk expressing it (i.e. have we encountered before in this graph?)
+	TMap<FName, FNiagaraVariable> ParamMapDefinedEmitterParameterToNamespaceVars; // Map from defined parameter to the Namespaced variable expressing it.
+
+    // Synced to the Attributes encountered for parameter maps thus far.
+	TMap<FName, int32> ParamMapDefinedAttributesToUniformChunks; // Map from the variable name exposed as a attribute to the uniform chunk expressing it (i.e. have we encountered before in this graph?)
+	TMap<FName, FNiagaraVariable> ParamMapDefinedAttributesToNamespaceVars; // Map from defined parameter to the Namespaced variable expressing it.
+
+	// Synced to the external variables used when bulk compiling system scripts.
+	TArray<FNiagaraVariable> ExternalVariablesForBulkUsage;
+
+	// Strings to be inserted within the main function
+	TArray<FString> MainPreSimulateChunks;
+
+	// read and write data set indices
+	int32 ReadIdx;
+	int32 WriteIdx;
+
+	// Data interface buffer descriptors per data interface
+	TArray< TArray<DIGPUBufferParamDescriptor> > DIBufferDescriptors;
+
+	/**
+	True if we're compiling the spawn section of an interpolated spawn script.
+	This will cause us to read from interpolated parameters and write to temp outputs that are fed back into the update script as inputs.
+	*/
+	bool bInsideInterpolatedSpawnScript;
+
+	/** Stack of currently tracked stats scopes. */
+	TArray<int32> StatScopeStack;
+
+	FString HlslOutput;
+
+	ENiagaraSimTarget CompilationTarget;
+
+	// Used to keep track of which output node we are working back from. This allows us 
+	// to find the right parameter map.
+	TArray<int32> CurrentParamMapIndices;
+};
diff --git a/Engine/Plugins/FX/Niagara/Source/NiagaraEditor/Private/NiagaraNodeAssignment.cpp b/Engine/Plugins/FX/Niagara/Source/NiagaraEditor/Private/NiagaraNodeAssignment.cpp
new file mode 100644
--- /dev/null
+++ b/Engine/Plugins/FX/Niagara/Source/NiagaraEditor/Private/NiagaraNodeAssignment.cpp
@@ -0,0 +1,190 @@
+// Copyright 1998-2017 Epic Games, Inc. All Rights Reserved.
+
+#include "NiagaraNodeAssignment.h"
+#include "UObject/UnrealType.h"
+#include "NiagaraGraph.h"
+#include "NiagaraScriptSource.h"
+#include "NiagaraScript.h"
+#include "NiagaraNodeInput.h"
+#include "NiagaraNodeOutput.h"
+#include "EdGraphSchema_Niagara.h"
+#include "ModuleManager.h"
+#include "AssetRegistryModule.h"
+#include "NiagaraComponent.h"
+#include "NiagaraHlslTranslator.h"
+#include "NiagaraEditorUtilities.h"
+#include "NiagaraNodeParameterMapGet.h"
+#include "NiagaraNodeParameterMapSet.h"
+#include "NiagaraConstants.h"
+#include "ViewModels/Stack/NiagaraParameterHandle.h"
+
+#define LOCTEXT_NAMESPACE "NiagaraNodeAssigment"
+
+void UNiagaraNodeAssignment::AllocateDefaultPins()
+{
+	GenerateScript();
+	Super::AllocateDefaultPins();
+}
+
+FText UNiagaraNodeAssignment::GetNodeTitle(ENodeTitleType::Type TitleType) const
+{
+	return FText::Format(LOCTEXT("NodeTitle", "Set {0}"), FText::FromName(AssignmentTarget.GetName()));
+}
+
+bool UNiagaraNodeAssignment::RefreshFromExternalChanges()
+{
+	FunctionScript = nullptr;
+	GenerateScript();
+	ReallocatePins();
+	return true;
+}
+
+void UNiagaraNodeAssignment::GenerateScript()
+{
+	FString NodeName = TEXT("Set_") + AssignmentTarget.GetName().ToString();
+	NodeName = NodeName.Replace(TEXT("."), TEXT("_"));
+
+	if (FunctionScript == nullptr)
+	{
+		FunctionScript = NewObject<UNiagaraScript>(this, FName(*NodeName), RF_Transactional);
+		FunctionScript->SetUsage(ENiagaraScriptUsage::Module);
+		InitializeScript(FunctionScript);
+		ComputeNodeName();
+	}
+	if (FunctionScript->GetName() != NodeName)
+	{
+		FunctionScript->Rename(*NodeName);
+		ComputeNodeName();
+	}
+}
+
+void UNiagaraNodeAssignment::InitializeScript(UNiagaraScript* NewScript)
+{
+	if (NewScript != NULL)
+	{		
+		UNiagaraScriptSource* Source = Cast<UNiagaraScriptSource>(NewScript->GetSource());
+
+		if (nullptr == Source)
+		{
+			Source = NewObject<UNiagaraScriptSource>(NewScript, NAME_None, RF_Transactional);
+			NewScript->SetSource(Source);
+		}
+
+		UNiagaraGraph* CreatedGraph = Source->NodeGraph;
+		if (nullptr == CreatedGraph)
+		{
+			CreatedGraph = NewObject<UNiagaraGraph>(Source, NAME_None, RF_Transactional);
+			Source->NodeGraph = CreatedGraph;
+		}
+		
+		TArray<UNiagaraNodeInput*> InputNodes;
+		CreatedGraph->FindInputNodes(InputNodes);
+		if (InputNodes.Num() == 0)
+		{
+			FGraphNodeCreator<UNiagaraNodeInput> InputNodeCreator(*CreatedGraph);
+			UNiagaraNodeInput* InputNode = InputNodeCreator.CreateNode();
+			InputNode->Input = FNiagaraVariable(FNiagaraTypeDefinition::GetParameterMapDef(), TEXT("InputMap"));
+			InputNode->Usage = ENiagaraInputNodeUsage::Parameter;
+			InputNodeCreator.Finalize();
+			InputNodes.Add(InputNode);
+		}
+
+		UNiagaraNodeOutput* OutputNode = CreatedGraph->FindOutputNode(ENiagaraScriptUsage::Module);
+		if (OutputNode == nullptr)
+		{
+			FGraphNodeCreator<UNiagaraNodeOutput> OutputNodeCreator(*CreatedGraph);
+			OutputNode = OutputNodeCreator.CreateNode();
+			FNiagaraVariable ParamMapAttrib(FNiagaraTypeDefinition::GetParameterMapDef(), TEXT("OutputMap"));
+			OutputNode->SetUsage(ENiagaraScriptUsage::Module);
+			OutputNode->Outputs.Add(ParamMapAttrib);
+			OutputNodeCreator.Finalize();
+		}
+
+		TArray<UNiagaraNodeParameterMapGet*> GetNodes;
+		CreatedGraph->GetNodesOfClass(GetNodes);
+
+		TArray<UNiagaraNodeParameterMapSet*> SetNodes;
+		CreatedGraph->GetNodesOfClass(SetNodes);
+
+		if (SetNodes.Num() == 0)
+		{
+			FGraphNodeCreator<UNiagaraNodeParameterMapSet> InputNodeCreator(*CreatedGraph);
+			UNiagaraNodeParameterMapSet* InputNode = InputNodeCreator.CreateNode();
+			InputNodeCreator.Finalize();
+			SetNodes.Add(InputNode);
+
+			InputNodes[0]->GetOutputPin(0)->MakeLinkTo(SetNodes[0]->GetInputPin(0));
+			SetNodes[0]->GetOutputPin(0)->MakeLinkTo(OutputNode->GetInputPin(0));
+		}
+
+		// We create two get nodes. The first is for the direct values.
+		// The second is in the case of referencing other parameters that we want to use as defaults.
+		if (GetNodes.Num() == 0)
+		{
+			FGraphNodeCreator<UNiagaraNodeParameterMapGet> InputNodeCreator(*CreatedGraph);
+			UNiagaraNodeParameterMapGet* InputNode = InputNodeCreator.CreateNode();
+			InputNodeCreator.Finalize();
+			GetNodes.Add(InputNode);
+
+			InputNodes[0]->GetOutputPin(0)->MakeLinkTo(GetNodes[0]->GetInputPin(0));
+		}
+		if (GetNodes.Num() == 1)
+		{
+			FGraphNodeCreator<UNiagaraNodeParameterMapGet> InputNodeCreator(*CreatedGraph);
+			UNiagaraNodeParameterMapGet* InputNode = InputNodeCreator.CreateNode();
+			InputNodeCreator.Finalize();
+			GetNodes.Add(InputNode);
+
+			InputNodes[0]->GetOutputPin(0)->MakeLinkTo(GetNodes[1]->GetInputPin(0));
+		}
+
+		// Clean out existing pins
+		while (!SetNodes[0]->IsAddPin(SetNodes[0]->GetInputPin(1)))
+		{
+			SetNodes[0]->RemovePin(SetNodes[0]->GetInputPin(1));
+		}
+
+		while (!GetNodes[0]->IsAddPin(GetNodes[0]->GetOutputPin(0)))
+		{
+			GetNodes[0]->RemovePin(GetNodes[0]->GetInputPin(0));
+		}
+
+		while (!GetNodes[1]->IsAddPin(GetNodes[1]->GetOutputPin(0)))
+		{
+			GetNodes[1]->RemovePin(GetNodes[1]->GetInputPin(0));
+		}
+
+		// Now create the proper new pins and connect them.
+		if (AssignmentTarget.GetName() != NAME_None)
+		{
+			FNiagaraParameterHandle TargetHandle(AssignmentTarget.GetName().ToString());
+			UEdGraphPin* SetPin = SetNodes[0]->RequestNewTypedPin(EGPD_Input, AssignmentTarget.GetType(), AssignmentTarget.GetName().ToString());
+			UEdGraphPin* GetPin = GetNodes[0]->RequestNewTypedPin(EGPD_Output, AssignmentTarget.GetType(), TEXT("Module.") + TargetHandle.GetName());
+
+			GetPin->MakeLinkTo(SetPin);
+
+			if (!AssignmentDefaultValue.IsEmpty())
+			{
+				UEdGraphPin* DefaultInputPin = GetNodes[0]->GetDefaultPin(GetPin);
+
+				bool isEngineConstant = false;
+				FNiagaraVariable SeekVar = FNiagaraVariable(AssignmentTarget.GetType(), FName(*AssignmentDefaultValue));
+				const FNiagaraVariable* FoundVar = FNiagaraConstants::FindEngineConstant(SeekVar);
+				if (FoundVar != nullptr)
+				{
+					UEdGraphPin* DefaultGetPin = GetNodes[1]->RequestNewTypedPin(EGPD_Output, AssignmentTarget.GetType(), FoundVar->GetName().ToString());
+					DefaultGetPin->MakeLinkTo(DefaultInputPin);
+				}
+				else
+				{
+					DefaultInputPin->bDefaultValueIsIgnored = false;
+					DefaultInputPin->DefaultValue = AssignmentDefaultValue;
+				}
+			}
+		}
+
+
+	}
+}
+
+#undef LOCTEXT_NAMESPACE
diff --git a/Engine/Plugins/FX/Niagara/Source/NiagaraEditor/Private/NiagaraNodeDataSetBase.cpp b/Engine/Plugins/FX/Niagara/Source/NiagaraEditor/Private/NiagaraNodeDataSetBase.cpp
new file mode 100644
--- /dev/null
+++ b/Engine/Plugins/FX/Niagara/Source/NiagaraEditor/Private/NiagaraNodeDataSetBase.cpp
@@ -0,0 +1,280 @@
+// Copyright 1998-2016 Epic Games, Inc. All Rights Reserved.
+
+#include "NiagaraNodeDataSetBase.h"
+#include "UObject/UnrealType.h"
+#include "INiagaraCompiler.h"
+#include "NiagaraEvents.h"
+#include "EdGraphSchema_Niagara.h"
+
+#define LOCTEXT_NAMESPACE "UNiagaraNodeDataSetBase"
+
+FString UNiagaraNodeDataSetBase::ConditionVarName = TEXT("__Condition");
+
+UNiagaraNodeDataSetBase::UNiagaraNodeDataSetBase(const FObjectInitializer& ObjectInitializer)
+: Super(ObjectInitializer), ExternalStructAsset(nullptr)
+{
+}
+
+bool UNiagaraNodeDataSetBase::InitializeFromStruct(const UStruct* PayloadStruct)
+{
+	if (InitializeFromStructInternal(PayloadStruct))
+	{
+		//ReallocatePins();
+		return true;
+	}
+	return false;
+}
+
+bool UNiagaraNodeDataSetBase::InitializeFromStructInternal(const UStruct* PayloadStruct)
+{
+	if (!PayloadStruct)
+	{
+		return false;
+	}
+
+	ExternalStructAsset = PayloadStruct;
+
+	return SynchronizeWithStruct();
+}
+
+
+void UNiagaraNodeDataSetBase::PostEditChangeProperty(struct FPropertyChangedEvent& PropertyChangedEvent)
+{
+	if (PropertyChangedEvent.Property != nullptr)
+	{
+		ReallocatePins();
+	}
+	Super::PostEditChangeProperty(PropertyChangedEvent);
+}
+
+
+FLinearColor UNiagaraNodeDataSetBase::GetNodeTitleColor() const
+{
+	check(DataSet.Type == ENiagaraDataSetType::Event);//Implement other datasets
+	return CastChecked<UEdGraphSchema_Niagara>(GetSchema())->NodeTitleColor_Event;
+}
+
+void UNiagaraNodeDataSetBase::PostLoad()
+{
+	Super::PostLoad();
+	if (ExternalStructAsset  == nullptr)
+	{
+		const TArray<FNiagaraTypeDefinition>& PayloadTypes = FNiagaraTypeRegistry::GetRegisteredPayloadTypes();
+		const FNiagaraTypeDefinition* PayloadType = PayloadTypes.FindByPredicate([&](const FNiagaraTypeDefinition& Type) { return Type.GetName() == DataSet.Name.ToString(); });
+		if (PayloadType != nullptr && PayloadType->GetStruct() != nullptr)
+		{
+			ExternalStructAsset = PayloadType->GetStruct();
+		}
+	}
+
+	/*if (ExternalStructAsset != nullptr)
+	{
+		UE_LOG(LogNiagaraEditor, Display, TEXT("Niagara script '%s' references struct asset '%s'"), *GetFullName(), *ExternalStructAsset->GetFullName());
+	}*/
+	
+	IsSynchronizedWithStruct(true, nullptr, true);
+}
+
+bool UNiagaraNodeDataSetBase::IsSynchronizedWithStruct(bool bIgnoreConditionVariable, FString* Issues, bool bLogIssues)
+{
+	if (ExternalStructAsset != nullptr)
+	{
+		bool bFoundIssues = false;
+		// First check to see if any variables have been added...
+		{
+			for (TFieldIterator<UProperty> PropertyIt(ExternalStructAsset, EFieldIteratorFlags::IncludeSuper); PropertyIt; ++PropertyIt)
+			{
+				const UProperty* Property = *PropertyIt;
+
+				if (bIgnoreConditionVariable && Property->IsA(UBoolProperty::StaticClass()) && Property->GetName().Equals(ConditionVarName))
+				{
+					continue;
+				}
+
+				const FNiagaraVariable* VarFound = Variables.FindByPredicate([&](const FNiagaraVariable& Var) { return Var.GetName().ToString() == Property->GetName(); });
+				if (VarFound == nullptr)
+				{
+					bFoundIssues = true;
+					if (bLogIssues)
+					{
+						UE_LOG(LogNiagaraEditor, Warning, TEXT("Unable to find matching variable for struct property: '%s' ... possible add?"), *Property->GetName());
+					}
+					if (Issues != nullptr)
+					{
+						Issues->Append(FString::Printf(TEXT("Unable to find matching variable for struct property: '%s' ... possible add?\n"), *Property->GetName()));
+					}
+				}
+				else
+				{
+					FNiagaraTypeDefinition TypeDefFound;
+					if (GetSupportedNiagaraTypeDef(Property, TypeDefFound))
+					{
+						if (VarFound->GetType() != TypeDefFound)
+						{
+							bFoundIssues = true;
+							if (bLogIssues)
+							{
+								UE_LOG(LogNiagaraEditor, Warning, TEXT("Matching variable for struct property '%s', but different type:  Existing Type:'%s' vs New Type:'%s' ... possible type change in user-defined struct?"), *Property->GetName(), *VarFound->GetType().GetName(), *TypeDefFound.GetName());
+							}
+							if (Issues != nullptr)
+							{
+								Issues->Append(FString::Printf(TEXT("Matching variable for struct property '%s', but different type: Existing Type:'%s' vs New Type:'%s' ... possible type change in user-defined struct?\n"), *Property->GetName(), *VarFound->GetType().GetName(), *TypeDefFound.GetName()));
+							}
+						}
+					}
+					else
+					{
+						bFoundIssues = true;
+						if (bLogIssues)
+						{
+							UE_LOG(LogNiagaraEditor, Warning, TEXT("Matching variable for struct property '%s', but different type:  Existing Type:'%s' vs New Type:'Unsupported' ... possible type change in user-defined struct?"), *Property->GetName(), *VarFound->GetType().GetName());
+						}
+						if (Issues != nullptr)
+						{
+							Issues->Append(FString::Printf(TEXT("Matching variable for struct property '%s', but different type: Existing Type:'%s' vs New Type:'Unsupported' ... possible type change in user-defined struct?\n"), *Property->GetName(), *VarFound->GetType().GetName()));
+						}
+					}
+				}
+			}
+		}
+
+		// Now check to see if any variables have been removed..
+		{
+			for (const FNiagaraVariable& Var : Variables)
+			{
+				bool bFound = false;
+				for (TFieldIterator<UProperty> PropertyIt(ExternalStructAsset, EFieldIteratorFlags::IncludeSuper); PropertyIt; ++PropertyIt)
+				{
+					const UProperty* Property = *PropertyIt;
+					if (Property->GetName() == Var.GetName().ToString())
+					{
+						bFound = true;
+						break;
+					}
+				}
+
+				if (!bFound)
+				{
+					bFoundIssues = true;
+					if (bLogIssues)
+					{
+						UE_LOG(LogNiagaraEditor, Warning, TEXT("Unable to find matching struct property for variable: '%s' ... possible removal?"), *Var.GetName().ToString());
+					}
+
+					if (Issues != nullptr)
+					{
+						Issues->Append(FString::Printf(TEXT("Unable to find matching struct property for variable: '%s' ... possible removal?\n"), *Var.GetName().ToString()));
+					}
+				}
+			}
+		}
+		return !bFoundIssues;
+	}
+	else
+	{
+		if (bLogIssues)
+		{
+			UE_LOG(LogNiagaraEditor, Warning, TEXT("Unable to find matching Niagara Type: %s"), *DataSet.Name.ToString());
+		}
+
+		if (Issues != nullptr)
+		{
+			Issues->Append(FString::Printf(TEXT("Unable to find matching Niagara Type: %s\n"), *DataSet.Name.ToString()));
+		}
+		return false;
+	}
+}
+
+bool UNiagaraNodeDataSetBase::RefreshFromExternalChanges()
+{
+	if (!IsSynchronizedWithStruct(true, nullptr, false))
+	{
+		SynchronizeWithStruct();
+		ReallocatePins();
+		return true;
+	}
+	return false;
+}
+
+bool UNiagaraNodeDataSetBase::GetSupportedNiagaraTypeDef(const UProperty* Property, FNiagaraTypeDefinition& TypeDef)
+{
+	const UStructProperty* StructProp = Cast<UStructProperty>(Property);
+	if (Property->IsA(UFloatProperty::StaticClass()))
+	{
+		TypeDef = FNiagaraTypeDefinition::GetFloatDef();
+		return true;
+	}
+	else if (Property->IsA(UBoolProperty::StaticClass()))
+	{
+		TypeDef = FNiagaraTypeDefinition::GetBoolDef();
+		return true;
+	}
+	else if (Property->IsA(UIntProperty::StaticClass()))
+	{
+		TypeDef = FNiagaraTypeDefinition::GetIntDef();
+		return true;
+	}
+	else if (StructProp && StructProp->Struct == FNiagaraTypeDefinition::GetVec2Struct())
+	{
+		TypeDef = FNiagaraTypeDefinition::GetVec2Def();
+		return true;
+	}
+	else if (StructProp && StructProp->Struct == FNiagaraTypeDefinition::GetVec3Struct())
+	{
+		TypeDef = FNiagaraTypeDefinition::GetVec3Def();
+		return true;
+	}
+	else if (StructProp && StructProp->Struct == FNiagaraTypeDefinition::GetVec4Struct())
+	{
+		TypeDef = FNiagaraTypeDefinition::GetVec4Def();
+		return true;
+	}
+	else if (StructProp && StructProp->Struct == FNiagaraTypeDefinition::GetColorStruct())
+	{
+		TypeDef = FNiagaraTypeDefinition::GetColorDef();
+		return true;
+	}
+	return false;
+}
+
+bool UNiagaraNodeDataSetBase::SynchronizeWithStruct()
+{
+
+	Variables.Empty();
+	VariableFriendlyNames.Empty();
+	DataSet = FNiagaraDataSetID();
+
+	// TODO: need to add valid as a variable separately for now; compiler support to validate index is missing	
+	//Variables.Add(FNiagaraVariable(FNiagaraTypeDefinition::GetBoolDef(), "Valid"));
+
+	//UGH! Why do we have our own custom type rep again. Why aren't we just using the UPropertySystem?
+	//
+	// [OP] not really different from anywhere else, nodes everywhere else hold FNiagaraVariables as their outputs; 
+	//  just traversing the ustruct here to build an array of those; this is temporary and should be genericised, of course
+	for (TFieldIterator<UProperty> PropertyIt(ExternalStructAsset, EFieldIteratorFlags::IncludeSuper); PropertyIt; ++PropertyIt)
+	{
+		const UProperty* Property = *PropertyIt;
+		FText DisplayNameText = Property->GetDisplayNameText();
+		FString DisplayName = DisplayNameText.ToString();
+		FNiagaraTypeDefinition TypeDef;
+
+		if (GetSupportedNiagaraTypeDef(Property, TypeDef))
+		{
+			Variables.Add(FNiagaraVariable(TypeDef, *Property->GetName()));
+			VariableFriendlyNames.Add(DisplayName);
+		}
+		else
+		{
+			UE_LOG(LogNiagaraEditor, Warning, TEXT("Could not add property %s in struct %s to NiagaraNodeDataSetBase, class %s"), *Property->GetName(), *ExternalStructAsset->GetName(), ExternalStructAsset->GetClass() ? *ExternalStructAsset->GetClass()->GetName() : *FString("nullptr"));
+		}
+	}
+
+	DataSet.Name = *ExternalStructAsset->GetName();
+	return true;
+}
+
+
+#undef LOCTEXT_NAMESPACE
+
+
+
diff --git a/Engine/Plugins/FX/Niagara/Source/NiagaraEditor/Private/NiagaraNodeEmitter.cpp b/Engine/Plugins/FX/Niagara/Source/NiagaraEditor/Private/NiagaraNodeEmitter.cpp
new file mode 100644
--- /dev/null
+++ b/Engine/Plugins/FX/Niagara/Source/NiagaraEditor/Private/NiagaraNodeEmitter.cpp
@@ -0,0 +1,456 @@
+// Copyright 1998-2016 Epic Games, Inc. All Rights Reserved.
+
+#include "NiagaraNodeEmitter.h"
+#include "NiagaraSystem.h"
+#include "NiagaraEmitter.h"
+#include "NiagaraEditorUtilities.h"
+#include "EdGraphSchema_Niagara.h"
+#include "EdGraph/EdGraph.h"
+#include "NiagaraCommon.h"
+#include "NiagaraDataInterface.h"
+#include "NiagaraScriptSource.h"
+#include "NiagaraGraph.h"
+#include "NiagaraNodeInput.h"
+#include "SNiagaraGraphPinAdd.h"
+#include "NiagaraGraph.h"
+#include "NiagaraNodeParameterMapBase.h"
+#include "NiagaraNodeOutput.h"
+#include "NiagaraHlslTranslator.h"
+#include "MultiBoxBuilder.h"
+
+#define LOCTEXT_NAMESPACE "NiagaraNodeEmitter"
+
+void UNiagaraNodeEmitter::PostInitProperties()
+{
+	Super::PostInitProperties();
+	PinPendingRename = nullptr;
+}
+
+UNiagaraSystem* UNiagaraNodeEmitter::GetOwnerSystem() const
+{
+	return OwnerSystem;
+}
+
+void UNiagaraNodeEmitter::SetOwnerSystem(UNiagaraSystem* InOwnerSystem)
+{
+	OwnerSystem = InOwnerSystem;
+	DisplayName = GetNameFromEmitter();
+}
+
+FGuid UNiagaraNodeEmitter::GetEmitterHandleId() const
+{
+	return EmitterHandleId;
+}
+
+void UNiagaraNodeEmitter::SetEmitterHandleId(FGuid InEmitterHandleId)
+{
+	EmitterHandleId = InEmitterHandleId;
+	DisplayName = GetNameFromEmitter();
+}
+
+void UNiagaraNodeEmitter::PostLoad()
+{
+	Super::PostLoad();
+
+	ReallocatePins();
+}
+
+bool UNiagaraNodeEmitter::IsPinNameEditable(const UEdGraphPin* GraphPinObj) const
+{
+	return false;
+}
+
+bool UNiagaraNodeEmitter::IsPinNameEditableUponCreation(const UEdGraphPin* GraphPinObj) const
+{
+	return false;
+}
+
+
+bool UNiagaraNodeEmitter::VerifyEditablePinName(const FText& InName, FText& OutErrorMessage, const UEdGraphPin* InGraphPinObj) const
+{
+	if (InName.IsEmptyOrWhitespace())
+	{
+		OutErrorMessage = LOCTEXT("InvalidName", "Invalid pin name");
+		return false;
+	}
+	return true;
+}
+
+bool UNiagaraNodeEmitter::CommitEditablePinName(const FText& InName, UEdGraphPin* InGraphPinObj)
+{
+	return false;
+}
+
+
+void UNiagaraNodeEmitter::AllocateDefaultPins()
+{
+	UNiagaraEmitter* Emitter = nullptr;
+	if (OwnerSystem)
+	{
+		for (int32 i = 0; i < OwnerSystem->GetNumEmitters(); ++i)
+		{
+			if (OwnerSystem->GetEmitterHandle(i).GetId() == EmitterHandleId)
+			{
+				Emitter = OwnerSystem->GetEmitterHandle(i).GetInstance();
+			}
+		}
+	}
+
+	if (Emitter != nullptr)
+	{
+		const UEdGraphSchema_Niagara* NiagaraSchema = Cast<UEdGraphSchema_Niagara>(GetSchema());
+		CreatePin(EGPD_Input, NiagaraSchema->TypeDefinitionToPinType(FNiagaraTypeDefinition::GetParameterMapDef()), TEXT("InputMap"));
+		CreatePin(EGPD_Output, NiagaraSchema->TypeDefinitionToPinType(FNiagaraTypeDefinition::GetParameterMapDef()), TEXT("OutputMap"));
+		CachedEmitterChangeId = Emitter->ChangeId;
+	}
+}
+
+bool UNiagaraNodeEmitter::CanUserDeleteNode() const
+{
+	return false;
+}
+
+bool UNiagaraNodeEmitter::CanDuplicateNode() const
+{
+	return false;
+}
+
+FText UNiagaraNodeEmitter::GetNodeTitle(ENodeTitleType::Type TitleType) const
+{
+	FText UsageText;
+	if (ScriptType == ENiagaraScriptUsage::EmitterSpawnScript)
+	{
+		UsageText = LOCTEXT("SpawnTitle", "Spawn");
+	}
+	else if (ScriptType == ENiagaraScriptUsage::EmitterUpdateScript)
+	{
+		UsageText = LOCTEXT("UpdateTitle", "Update");
+	}
+	else
+	{
+		UsageText = LOCTEXT("Unknown Title", "Unknown");
+	}
+	return FText::Format(LOCTEXT("EmitterNameTitle","Emitter {0} {1}"), DisplayName, UsageText);
+}
+
+FLinearColor UNiagaraNodeEmitter::GetNodeTitleColor() const
+{
+	return CastChecked<UEdGraphSchema_Niagara>(GetSchema())->NodeTitleColor_Attribute;
+}
+
+void UNiagaraNodeEmitter::NodeConnectionListChanged()
+{
+	GetGraph()->NotifyGraphChanged();
+}
+
+FString UNiagaraNodeEmitter::GetEmitterUniqueName() const
+{
+	if (OwnerSystem != nullptr && EmitterHandleId.IsValid())
+	{
+		for (const FNiagaraEmitterHandle& EmitterHandle : OwnerSystem->GetEmitterHandles())
+		{
+			if (EmitterHandle.GetId() == EmitterHandleId)
+			{
+				return EmitterHandle.GetUniqueInstanceName();
+			}
+		}
+	}
+	return FString();
+}
+
+UNiagaraScriptSource* UNiagaraNodeEmitter::GetScriptSource() const
+{
+	// First get the emitter that we're referencing..
+	UNiagaraEmitter* Emitter = nullptr;
+	if (OwnerSystem)
+	{
+		for (int32 i = 0; i < OwnerSystem->GetNumEmitters(); ++i)
+		{
+			if (OwnerSystem->GetEmitterHandle(i).GetId() == EmitterHandleId)
+			{
+				Emitter = OwnerSystem->GetEmitterHandle(i).GetInstance();
+			}
+		}
+	}
+
+	// Now get the graph off that emitter
+	if (Emitter != nullptr && Emitter->GraphSource != nullptr)
+	{
+		UNiagaraScriptSource* Source = Cast<UNiagaraScriptSource>(Emitter->GraphSource);
+		return Source;
+	}
+
+	return nullptr;
+}
+
+UNiagaraGraph* UNiagaraNodeEmitter::GetCalledGraph() const
+{
+	// First get the emitter that we're referencing..
+	UNiagaraEmitter* Emitter = nullptr;
+	if (OwnerSystem)
+	{
+		for (int32 i = 0; i < OwnerSystem->GetNumEmitters(); ++i)
+		{
+			if (OwnerSystem->GetEmitterHandle(i).GetId() == EmitterHandleId)
+			{
+				Emitter = OwnerSystem->GetEmitterHandle(i).GetInstance();
+			}
+		}
+	}
+
+	// Now get the graph off that emitter
+	if (Emitter != nullptr && Emitter->GraphSource != nullptr)
+	{
+		UNiagaraScriptSource* Source = Cast<UNiagaraScriptSource>(Emitter->GraphSource);
+		if (Source)
+		{
+			return Source->NodeGraph;
+		}
+	}
+	return nullptr;
+}
+
+bool UNiagaraNodeEmitter::RefreshFromExternalChanges()
+{
+	// First get the emitter that we're referencing..
+	UNiagaraEmitter* Emitter = nullptr;
+	if (OwnerSystem)
+	{
+		for (int32 i = 0; i < OwnerSystem->GetNumEmitters(); ++i)
+		{
+			if (OwnerSystem->GetEmitterHandle(i).GetId() == EmitterHandleId)
+			{
+				Emitter = OwnerSystem->GetEmitterHandle(i).GetInstance();
+			}
+		}
+	}
+
+	if (Emitter && Emitter->ChangeId != CachedEmitterChangeId)
+	{
+		// TODO - Leverage code in reallocate pins to determine if any pins have changed...
+		ReallocatePins();
+		DisplayName = GetNameFromEmitter();
+		return true;
+	}
+	return false;
+}
+
+FText UNiagaraNodeEmitter::GetNameFromEmitter()
+{
+	if (OwnerSystem != nullptr && EmitterHandleId.IsValid())
+	{
+		for (const FNiagaraEmitterHandle& EmitterHandle : OwnerSystem->GetEmitterHandles())
+		{
+			if (EmitterHandle.GetId() == EmitterHandleId)
+			{
+				return FText::FromName(EmitterHandle.GetName());
+			}
+		}
+	}
+	return FText();
+}
+
+void UNiagaraNodeEmitter::BuildParameterMapHistory(FNiagaraParameterMapHistoryBuilder& OutHistory, bool bRecursive)
+{
+	Super::BuildParameterMapHistory(OutHistory, bRecursive);
+
+	const UEdGraphSchema_Niagara* Schema = GetDefault<UEdGraphSchema_Niagara>();
+	TArray<UEdGraphPin*> InputPins;
+	GetInputPins(InputPins);
+	TArray<UEdGraphPin*> OutputPins;
+	GetOutputPins(OutputPins);
+	
+	int32 ParamMapIdx = INDEX_NONE;
+	if (GetInputPin(0)->LinkedTo.Num() != 0)
+	{
+		if (bRecursive)
+		{
+			ParamMapIdx = OutHistory.TraceParameterMapOutputPin(GetInputPin(0)->LinkedTo[0]);
+		}
+		else
+		{
+			ParamMapIdx = OutHistory.CreateParameterMap();
+		}
+	}
+
+	FString EmitterUniqueName = GetEmitterUniqueName();
+	UNiagaraGraph* Graph = GetCalledGraph();
+	if (Graph && ParamMapIdx != INDEX_NONE)
+	{
+		OutHistory.EnterEmitter(EmitterUniqueName, this);
+
+		TArray<ENiagaraScriptUsage> Usages;
+		Usages.Add(ENiagaraScriptUsage::EmitterSpawnScript);
+		Usages.Add(ENiagaraScriptUsage::EmitterUpdateScript);
+		Usages.Add(ENiagaraScriptUsage::ParticleSpawnScript);
+		Usages.Add(ENiagaraScriptUsage::ParticleSpawnScriptInterpolated);
+		Usages.Add(ENiagaraScriptUsage::ParticleUpdateScript);
+		Usages.Add(ENiagaraScriptUsage::ParticleEventScript);
+		
+		for (ENiagaraScriptUsage OutputNodeUsage : Usages)
+		{
+			TArray<UNiagaraNodeOutput*> OutputNodes;
+
+			Graph->FindOutputNodes(OutputNodeUsage, OutputNodes);
+
+			// Build up a new parameter map history with all the child graph nodes..
+			FNiagaraParameterMapHistoryBuilder ChildBuilder;
+
+			ChildBuilder.EnableScriptWhitelist(true, GetUsage());
+			ChildBuilder.EnterEmitter(EmitterUniqueName, this);
+			for (UNiagaraNodeOutput* OutputNode : OutputNodes)
+			{
+				ChildBuilder.BuildParameterMaps(OutputNode, true);
+			}
+			ChildBuilder.ExitEmitter(EmitterUniqueName, this);
+			 
+			for (FNiagaraParameterMapHistory& History : ChildBuilder.Histories)
+			{
+				OutHistory.Histories[ParamMapIdx].MapPinHistory.Append(History.MapPinHistory);
+				for (int32 SrcVarIdx = 0; SrcVarIdx < History.Variables.Num(); SrcVarIdx++)
+				{
+					FNiagaraVariable& Var = History.Variables[SrcVarIdx];
+					int32 ExistingIdx = OutHistory.Histories[ParamMapIdx].FindVariable(Var.GetName(), Var.GetType());
+					if (ExistingIdx == INDEX_NONE)
+					{
+						ExistingIdx = OutHistory.Histories[ParamMapIdx].Variables.Add(Var);
+						OutHistory.Histories[ParamMapIdx].PerVariableReadHistory.AddDefaulted(1);
+						OutHistory.Histories[ParamMapIdx].PerVariableWriteHistory.AddDefaulted(1);
+						OutHistory.Histories[ParamMapIdx].PerVariableWarnings.AddDefaulted(1);
+					}
+					ensure(ExistingIdx < OutHistory.Histories[ParamMapIdx].PerVariableWarnings.Num());
+					ensure(ExistingIdx < OutHistory.Histories[ParamMapIdx].PerVariableReadHistory.Num());
+					ensure(ExistingIdx < OutHistory.Histories[ParamMapIdx].PerVariableWriteHistory.Num());
+					OutHistory.Histories[ParamMapIdx].PerVariableReadHistory[ExistingIdx].Append(History.PerVariableReadHistory[SrcVarIdx]);
+					OutHistory.Histories[ParamMapIdx].PerVariableWriteHistory[ExistingIdx].Append(History.PerVariableWriteHistory[SrcVarIdx]);
+					OutHistory.Histories[ParamMapIdx].PerVariableWarnings[ExistingIdx].Append(History.PerVariableWarnings[SrcVarIdx]);	
+				}
+			}
+		}
+
+		OutHistory.ExitEmitter(EmitterUniqueName, this);
+	}
+
+	for (UEdGraphPin* Pin : OutputPins)
+	{
+		FNiagaraTypeDefinition Type = Schema->PinToTypeDefinition(Pin);
+		if (Type == FNiagaraTypeDefinition::GetParameterMapDef())
+		{
+			OutHistory.RegisterParameterMapPin(ParamMapIdx, Pin);
+		}
+	}
+}
+
+void UNiagaraNodeEmitter::Compile(FHlslNiagaraTranslator *Translator, TArray<int32>& Outputs)
+{
+	TArray<UEdGraphPin*> InputPins;
+	GetInputPins(InputPins);
+	InputPins.RemoveAll([](UEdGraphPin* InputPin) { return (InputPin->PinType.PinCategory != UEdGraphSchema_Niagara::PinCategoryType) && (InputPin->PinType.PinCategory != UEdGraphSchema_Niagara::PinCategoryEnum); });
+
+	TArray<UEdGraphPin*> OutputPins;
+	GetOutputPins(OutputPins);
+
+	check(Outputs.Num() == 0);
+
+	const UEdGraphSchema_Niagara* Schema = CastChecked<UEdGraphSchema_Niagara>(GetSchema());
+
+	// First compile fully down the hierarchy for our predecessors..
+	UNiagaraGraph* CalledGraph = GetCalledGraph();
+	TArray<UNiagaraNodeInput*> InputsNodes;
+	UNiagaraGraph::FFindInputNodeOptions Options;
+	Options.bSort = true;
+	Options.bFilterDuplicates = true;
+	Options.bFilterByScriptUsage = true;
+	Options.TargetScriptUsage = Translator->GetTargetUsage() == ENiagaraScriptUsage::SystemSpawnScript ? ENiagaraScriptUsage::EmitterSpawnScript : ENiagaraScriptUsage::EmitterUpdateScript;
+	CalledGraph->FindInputNodes(InputsNodes, Options);
+
+	TArray<int32> CompileInputs;
+	CompileInputs.Reserve(InputsNodes.Num());
+
+	if (InputsNodes.Num() <= 0)
+	{
+		Translator->Error(LOCTEXT("InputNodesNotFound", "Input nodes on called graph not found"), this, nullptr);
+		return;
+	}
+
+	if (InputPins.Num() > 1)
+	{
+		Translator->Error(LOCTEXT("TooManyOutputPinsError", "Too many input pins on node."), this, nullptr);
+		return;
+	}
+
+	bool bError = false;
+	for (UNiagaraNodeInput* EmitterInputNode : InputsNodes)
+	{
+		if (EmitterInputNode->Input.IsEquivalent(FNiagaraVariable(FNiagaraTypeDefinition::GetParameterMapDef(), TEXT("InputMap"))))
+		{
+			CompileInputs.Add(Translator->CompilePin(InputPins[0]));
+		}
+		else
+		{
+			CompileInputs.Add(INDEX_NONE);
+		}
+	}
+
+	if (!bError)
+	{
+		Translator->Emitter(this, CompileInputs, Outputs);
+	}
+}
+
+TSharedRef<SWidget> UNiagaraNodeEmitter::GenerateAddPinMenu(const FString& InWorkingPinName, SNiagaraGraphPinAdd* InPin)
+{
+	FMenuBuilder MenuBuilder(true, nullptr);
+	
+	UNiagaraGraph* Graph = GetCalledGraph();
+	const UEdGraphSchema_Niagara* Schema = GetDefault<UEdGraphSchema_Niagara>();
+
+	TArray<UEdGraphPin*> InputPins;
+	GetInputPins(InputPins);
+
+	FString EmitterNamespace = FNiagaraParameterMapHistory::MakeSafeNamespaceString(GetNameFromEmitter().ToString());
+
+	FNiagaraParameterMapHistoryBuilder HistoryBuilder;
+	BuildParameterMapHistory(HistoryBuilder, false);
+
+	if (HistoryBuilder.Histories.Num() == 0)
+	{
+		return SNullWidget::NullWidget;
+	}
+
+	TArray<FNiagaraVariable> ExistingVariables;
+	for (UEdGraphPin* InputPin : InputPins)
+	{
+		ExistingVariables.Add(Schema->PinToNiagaraVariable(InputPin));
+	}
+
+	TArray<FNiagaraVariable> Variables;
+	for (const FNiagaraVariable& Var : HistoryBuilder.Histories[0].Variables)
+	{
+		if (ExistingVariables.ContainsByPredicate([Var](const FNiagaraVariable& A) { return A.GetName().ToString().Equals(Var.GetName().ToString()); }))
+		{
+			continue;
+		}
+
+		if (FNiagaraParameterMapHistory::IsInNamespace(Var, TEXT("EmitterUniforms")) || FNiagaraParameterMapHistory::IsInNamespace(Var, EmitterNamespace))
+		{
+			Variables.AddUnique(Var);
+		}
+	}
+
+	Variables.Sort([](const FNiagaraVariable& A, const FNiagaraVariable& B) { return (A.GetName() < B.GetName()); });
+
+	for (FNiagaraVariable& NamespacedVar : Variables)
+	{
+		FString NamespacedName = NamespacedVar.GetName().ToString();
+		MenuBuilder.AddMenuEntry(
+			FText::FromString(NamespacedName),
+			FText::Format(LOCTEXT("AddButtonTypeEntryToolTipFormatSystem", "Add a reference to {0}"), FText::FromName(NamespacedVar.GetName())),
+			FSlateIcon(),
+			FUIAction(FExecuteAction::CreateSP(InPin, &SNiagaraGraphPinAdd::OnAddType, NamespacedVar)));
+	}
+	
+	return MenuBuilder.MakeWidget();
+}
+
+#undef LOCTEXT_NAMESPACE // NiagaraNodeEmitter
diff --git a/Engine/Plugins/FX/Niagara/Source/NiagaraEditor/Private/NiagaraNodeFunctionCall.cpp b/Engine/Plugins/FX/Niagara/Source/NiagaraEditor/Private/NiagaraNodeFunctionCall.cpp
new file mode 100644
--- /dev/null
+++ b/Engine/Plugins/FX/Niagara/Source/NiagaraEditor/Private/NiagaraNodeFunctionCall.cpp
@@ -0,0 +1,601 @@
+// Copyright 1998-2017 Epic Games, Inc. All Rights Reserved.
+
+#include "NiagaraNodeFunctionCall.h"
+#include "UObject/UnrealType.h"
+#include "NiagaraGraph.h"
+#include "NiagaraScriptSource.h"
+#include "NiagaraScript.h"
+#include "NiagaraNodeInput.h"
+#include "NiagaraNodeOutput.h"
+#include "EdGraphSchema_Niagara.h"
+#include "ModuleManager.h"
+#include "AssetRegistryModule.h"
+#include "NiagaraComponent.h"
+#include "NiagaraHlslTranslator.h"
+#include "NiagaraEditorUtilities.h"
+#include "NiagaraNodeParameterMapGet.h"
+#include "NiagaraConstants.h"
+
+#define LOCTEXT_NAMESPACE "NiagaraNodeFunctionCall"
+
+void UNiagaraNodeFunctionCall::PostLoad()
+{
+	Super::PostLoad();
+
+	if (FunctionScript)
+	{
+		FunctionScript->ConditionalPostLoad();
+
+		// We need to make sure that the variables that could potentially be used in AllocateDefaultPins have been properly
+		// loaded. Otherwise, we could be out of date.
+		if (FunctionScript->GetSource())
+		{
+			UNiagaraScriptSource* Source = CastChecked<UNiagaraScriptSource>(FunctionScript->GetSource());
+			Source->ConditionalPostLoad();
+			UNiagaraGraph* Graph = CastChecked<UNiagaraGraph>(Source->NodeGraph);
+			Graph->ConditionalPostLoad();
+		}
+	}
+
+	if (FunctionDisplayName.IsEmpty())
+	{
+		ComputeNodeName();
+	}
+}
+
+void UNiagaraNodeFunctionCall::PostEditChangeProperty(struct FPropertyChangedEvent& PropertyChangedEvent)
+{
+	if (PropertyChangedEvent.Property != nullptr)
+	{
+		ReallocatePins();
+	}
+	Super::PostEditChangeProperty(PropertyChangedEvent);
+
+	GetGraph()->NotifyGraphChanged();
+}
+
+void UNiagaraNodeFunctionCall::AllocateDefaultPins()
+{
+	if (FunctionScriptAssetObjectPath != NAME_None && FunctionScript == nullptr)
+	{
+		FAssetRegistryModule& AssetRegistryModule = FModuleManager::LoadModuleChecked<FAssetRegistryModule>(TEXT("AssetRegistry"));
+		FAssetData ScriptAssetData = AssetRegistryModule.Get().GetAssetByObjectPath(FunctionScriptAssetObjectPath);
+		if (ScriptAssetData.IsValid())
+		{
+			FunctionScript = Cast<UNiagaraScript>(ScriptAssetData.GetAsset());
+		}
+	}
+
+	const UEdGraphSchema_Niagara* Schema = CastChecked<UEdGraphSchema_Niagara>(GetSchema());
+	UNiagaraGraph* CallerGraph = GetNiagaraGraph();
+	if (FunctionScript)
+	{
+		UNiagaraScriptSource* Source = CastChecked<UNiagaraScriptSource>(FunctionScript->GetSource());
+		UNiagaraGraph* Graph = CastChecked<UNiagaraGraph>(Source->NodeGraph);
+
+		//These pins must be refreshed and kept in the correct order for the function
+		TArray<FNiagaraVariable> Inputs;
+		TArray<FNiagaraVariable> Outputs;
+		Graph->GetParameters(Inputs, Outputs);
+
+		TArray<UNiagaraNodeInput*> InputNodes;
+		UNiagaraGraph::FFindInputNodeOptions Options;
+		Options.bSort = true;
+		Options.bFilterDuplicates = true;
+		Graph->FindInputNodes(InputNodes, Options);
+
+		bool bHasAdvancePins = false;
+		for (UNiagaraNodeInput* InputNode : InputNodes)
+		{
+			if (InputNode->IsExposed())
+			{
+				UEdGraphPin* NewPin = CreatePin(EGPD_Input, Schema->TypeDefinitionToPinType(InputNode->Input.GetType()), InputNode->Input.GetName().ToString());
+					
+				//An inline pin default only makes sense if we are required. 
+				//Non exposed or optional inputs will used their own function input nodes defaults when not directly provided by a link.
+				//Special class types cannot have an inline default.
+				NewPin->bDefaultValueIsIgnored = !(InputNode->IsRequired() && InputNode->Input.GetType().GetClass() == nullptr);
+
+				//TODO: Some visual indication of Auto bound pins.
+				//I tried just linking to null but
+// 				FNiagaraVariable AutoBoundVar;
+// 				ENiagaraInputNodeUsage AutBoundUsage = ENiagaraInputNodeUsage::Undefined;
+// 				bool bCanAutoBind = FindAutoBoundInput(InputNode->AutoBindOptions, NewPin, AutoBoundVar, AutBoundUsage);
+// 				if (bCanAutoBind)
+// 				{
+// 
+// 				}
+
+				if (InputNode->IsHidden())
+				{
+					NewPin->bAdvancedView = true;
+					bHasAdvancePins = true;
+				}
+				else
+				{
+					NewPin->bAdvancedView = false;
+				}
+			}
+		}
+
+		AdvancedPinDisplay = bHasAdvancePins ? ENodeAdvancedPins::Hidden : ENodeAdvancedPins::NoPins;
+
+		for (FNiagaraVariable& Output : Outputs)
+		{
+			UEdGraphPin* NewPin = CreatePin(EGPD_Output, Schema->TypeDefinitionToPinType(Output.GetType()), Output.GetName().ToString());
+			NewPin->bDefaultValueIsIgnored = true;
+		}
+
+		// Make sure to note that we've synchronized with the external version.
+		CachedChangeId = FunctionScript->GetChangeID();
+	}
+	else
+	{
+		
+		for (FNiagaraVariable& Input : Signature.Inputs)
+		{
+			UEdGraphPin* NewPin = CreatePin(EGPD_Input, Schema->TypeDefinitionToPinType(Input.GetType()), Input.GetName().ToString());
+			NewPin->bDefaultValueIsIgnored = false;
+		}
+
+		for (FNiagaraVariable& Output : Signature.Outputs)
+		{
+			UEdGraphPin* NewPin = CreatePin(EGPD_Output, Schema->TypeDefinitionToPinType(Output.GetType()), Output.GetName().ToString());
+			NewPin->bDefaultValueIsIgnored = true;
+		}
+
+		// We don't reference an external function, so set an invalid id.
+		CachedChangeId = FGuid();
+	}
+
+	if (FunctionDisplayName.IsEmpty())
+	{
+		ComputeNodeName();
+	}
+}
+
+FText UNiagaraNodeFunctionCall::GetNodeTitle(ENodeTitleType::Type TitleType) const
+{
+	FString DetectedName = FunctionScript ? FunctionScript->GetName() : Signature.GetName();
+	if (DetectedName.IsEmpty())
+	{
+		return FText::FromString(TEXT("Missing ( Was\"") + FunctionDisplayName + TEXT("\")"));
+	}
+	else
+	{
+		return FText::FromString(FName::NameToDisplayString(FunctionDisplayName, false));
+	}
+}
+
+FText UNiagaraNodeFunctionCall::GetTooltipText()const
+{
+	if (FunctionScript != nullptr)
+	{
+		return FunctionScript->GetDescription();
+	}
+	else if (Signature.IsValid())
+	{
+		return Signature.Description;
+	} 
+	else
+	{
+		return LOCTEXT("NiagaraFuncCallUnknownSignatureTooltip", "Unknown function call");
+	}
+}
+
+FLinearColor UNiagaraNodeFunctionCall::GetNodeTitleColor() const
+{
+	return UEdGraphSchema_Niagara::NodeTitleColor_FunctionCall;
+}
+
+bool UNiagaraNodeFunctionCall::CanAddToGraph(UNiagaraGraph* TargetGraph, FString& OutErrorMsg) const
+{
+	if (Super::CanAddToGraph(TargetGraph, OutErrorMsg) == false)
+	{
+		return false;
+	}
+	UPackage* TargetPackage = TargetGraph->GetOutermost();
+
+	TArray<const UNiagaraGraph*> FunctionGraphs;
+	UNiagaraScript* SpawningFunctionScript = FunctionScript;
+	
+	// We probably haven't loaded the script yet. Let's do so now so that we can trace its lineage.
+	if (FunctionScriptAssetObjectPath != NAME_None && FunctionScript == nullptr)
+	{
+		FAssetRegistryModule& AssetRegistryModule = FModuleManager::LoadModuleChecked<FAssetRegistryModule>(TEXT("AssetRegistry"));
+		FAssetData ScriptAssetData = AssetRegistryModule.Get().GetAssetByObjectPath(FunctionScriptAssetObjectPath);
+		if (ScriptAssetData.IsValid())
+		{
+			SpawningFunctionScript = Cast<UNiagaraScript>(ScriptAssetData.GetAsset());
+		}
+	}
+
+	// Now we need to get the graphs referenced by the script that we are about to spawn in.
+	if (SpawningFunctionScript && SpawningFunctionScript->GetSource())
+	{
+		UNiagaraScriptSource* Source = Cast<UNiagaraScriptSource>(SpawningFunctionScript->GetSource());
+		if (Source)
+		{
+			UNiagaraGraph* FunctionGraph = Cast<UNiagaraGraph>(Source->NodeGraph);
+			if (FunctionGraph)
+			{
+				FunctionGraph->GetAllReferencedGraphs(FunctionGraphs);
+			}
+		}
+	}
+
+	// Iterate over each graph referenced by this spawning function call and see if any of them reference the graph that we are about to be spawned into. If 
+	// a match is found, then adding us would introduce a cycle and we need to abort the add.
+	for (const UNiagaraGraph* Graph : FunctionGraphs)
+	{
+		UPackage* FunctionPackage = Graph->GetOutermost();
+		if (FunctionPackage != nullptr && TargetPackage != nullptr && FunctionPackage == TargetPackage)
+		{
+			OutErrorMsg = LOCTEXT("NiagaraFuncCallCannotAddToGraph", "Cannot add to graph because the Function Call used by this node would lead to a cycle.").ToString();
+			return false;
+		}
+	}
+
+	return true;
+}
+
+void UNiagaraNodeFunctionCall::Compile(class FHlslNiagaraTranslator* Translator, TArray<int32>& Outputs)
+{
+	TArray<int32> Inputs;
+
+	bool bError = false;
+
+	const UEdGraphSchema_Niagara* Schema = CastChecked<UEdGraphSchema_Niagara>(GetSchema());
+	UNiagaraGraph* CallerGraph = GetNiagaraGraph();
+	if (FunctionScript)
+	{
+		TArray<UEdGraphPin*> CallerInputPins;
+		GetInputPins(CallerInputPins);
+
+		UNiagaraScriptSource* Source = CastChecked<UNiagaraScriptSource>(FunctionScript->GetSource());
+		UNiagaraGraph* FunctionGraph = CastChecked<UNiagaraGraph>(Source->NodeGraph);
+
+		TArray<UNiagaraNodeInput*> FunctionInputNodes;
+		UNiagaraGraph::FFindInputNodeOptions Options;
+		Options.bSort = true;
+		Options.bFilterDuplicates = true;
+		FunctionGraph->FindInputNodes(FunctionInputNodes, Options);
+
+		for (UNiagaraNodeInput* FunctionInputNode : FunctionInputNodes)
+		{
+			//Finds the matching Pin in the caller.
+			UEdGraphPin** PinPtr = CallerInputPins.FindByPredicate([&](UEdGraphPin* InPin) { return Schema->PinToNiagaraVariable(InPin).IsEquivalent(FunctionInputNode->Input); });
+			if (!PinPtr)
+			{
+				if (FunctionInputNode->IsExposed())
+				{
+					//Couldn't find the matching pin for an exposed input. Probably a stale function call node that needs to be refreshed.
+					Translator->Error(LOCTEXT("StaleFunctionCallError", "Function call is stale and needs to be refreshed."), this, nullptr);
+					bError = true;
+				}
+				else if (FunctionInputNode->ExposureOptions.bRequired == true)
+				{
+					// Not exposed, but required. This means we should just add as a constant.
+					Inputs.Add(Translator->GetConstant(FunctionInputNode->Input));
+					continue;
+				}
+
+
+				Inputs.Add(INDEX_NONE);
+				continue;
+			}
+
+			UEdGraphPin* CallerPin = *PinPtr;
+			UEdGraphPin* CallerLinkedTo = CallerPin->LinkedTo.Num() > 0 ? CallerPin->LinkedTo[0] : nullptr;			
+			FNiagaraVariable PinVar = Schema->PinToNiagaraVariable(CallerPin);
+			if (!CallerLinkedTo)
+			{
+				//if (Translator->CanReadAttributes())
+				{
+					//Try to auto bind if we're not linked to by the caller.
+					FNiagaraVariable AutoBoundVar;
+					ENiagaraInputNodeUsage AutBoundUsage = ENiagaraInputNodeUsage::Undefined;
+					if (FindAutoBoundInput(FunctionInputNode, CallerPin, AutoBoundVar, AutBoundUsage))
+					{
+						UNiagaraNodeInput* NewNode = NewObject<UNiagaraNodeInput>(CallerGraph);
+						NewNode->Input = PinVar;
+						NewNode->Usage = AutBoundUsage;
+						NewNode->AllocateDefaultPins();
+						CallerLinkedTo = NewNode->GetOutputPin(0);
+						CallerPin->BreakAllPinLinks();
+						CallerPin->MakeLinkTo(CallerLinkedTo);
+					}
+				}
+			}
+
+			if (CallerLinkedTo)
+			{
+				//Param is provided by the caller. Typical case.
+				Inputs.Add(Translator->CompilePin(CallerPin));
+				continue;
+			}
+			else
+			{
+				if (FunctionInputNode->IsRequired())
+				{
+					if (CallerPin->bDefaultValueIsIgnored)
+					{
+						//This pin can't use a default and it is required so flag an error.
+						Translator->Error(FText::Format(LOCTEXT("RequiredInputUnboundErrorFmt", "Required input {0} was not bound and could not be automatically bound."), CallerPin->GetDisplayName()),
+							this, CallerPin);
+						bError = true;
+						//We weren't linked to anything and we couldn't auto bind so tell the compiler this input isn't provided and it should use it's local default.
+						Inputs.Add(INDEX_NONE);
+					}
+					else
+					{
+						//We also compile the pin anyway if it is required as we'll be attempting to use it's inline default.
+						Inputs.Add(Translator->CompilePin(CallerPin));
+					}
+				}
+				else
+				{
+					//We optional, weren't linked to anything and we couldn't auto bind so tell the compiler this input isn't provided and it should use it's local default.
+					Inputs.Add(INDEX_NONE);
+				}
+			}
+		}
+	}
+	else
+	{
+		bError = CompileInputPins(Translator, Inputs);
+	}		
+
+	if (!bError)
+	{
+		Translator->FunctionCall(this, Inputs, Outputs);
+	}
+}
+
+UObject*  UNiagaraNodeFunctionCall::GetReferencedAsset() const
+{
+	if (FunctionScript && FunctionScript->GetOutermost() != GetOutermost())
+	{
+		return FunctionScript;
+	}
+	else
+	{
+		return nullptr;
+	}
+}
+
+bool UNiagaraNodeFunctionCall::RefreshFromExternalChanges()
+{
+	bool bReload = false;
+	if (FunctionScript)
+	{
+		bReload = CachedChangeId != FunctionScript->GetChangeID();
+		if (bReload)
+		{
+			bReload = true;
+			check(FunctionScript->GetOutermost()->HasAnyFlags(RF_WasLoaded));
+			UE_LOG(LogNiagaraEditor, Log, TEXT("RefreshFromExternalChanges %s"), *(FunctionScript->GetPathName()));
+		}
+	}
+	else
+	{
+		if (Signature.IsValid())
+		{
+			bReload = true;
+		}
+	}
+
+	if (bReload)
+	{
+		// TODO - Leverage code in reallocate pins to determine if any pins have changed...
+		ReallocatePins();
+		return true;
+	}
+	else
+	{
+		return false;
+	}
+}
+
+void UNiagaraNodeFunctionCall::SubsumeExternalDependencies(TMap<const UObject*, UObject*>& ExistingConversions)
+{
+	if (FunctionScript && FunctionScript->GetOutermost() != this->GetOutermost())
+	{
+		if (ExistingConversions.Contains(FunctionScript))
+		{
+			FunctionScript = CastChecked<UNiagaraScript>(ExistingConversions[FunctionScript]);
+			check(FunctionScript->HasAnyFlags(RF_Standalone) == false);
+			check(FunctionScript->HasAnyFlags(RF_Public) == false);
+		}
+		else
+		{
+			FunctionScript = FunctionScript->MakeRecursiveDeepCopy(this, ExistingConversions);
+		}
+	}
+}
+
+void UNiagaraNodeFunctionCall::BuildParameterMapHistory(FNiagaraParameterMapHistoryBuilder& OutHistory, bool bRecursive)
+{
+	Super::BuildParameterMapHistory(OutHistory, bRecursive);
+
+	if (FunctionScript)
+	{
+		UNiagaraScriptSource* Source = CastChecked<UNiagaraScriptSource>(FunctionScript->GetSource());
+		UNiagaraGraph* FunctionGraph = CastChecked<UNiagaraGraph>(Source->NodeGraph);
+		
+		UNiagaraNodeOutput* OutputNode = FunctionGraph->FindOutputNode(ENiagaraScriptUsage::Function);
+		if (OutputNode == nullptr)
+		{
+			OutputNode = FunctionGraph->FindOutputNode(ENiagaraScriptUsage::Module);
+		}
+		if (OutputNode == nullptr)
+		{
+			OutputNode = FunctionGraph->FindOutputNode(ENiagaraScriptUsage::DynamicInput);
+		}
+		
+		OutHistory.EnterFunction(GetFunctionName(), FunctionScript, this);
+		OutputNode->BuildParameterMapHistory(OutHistory, true);
+
+		// Since we're about to lose the pin calling context, we finish up the function call parameter map pin wiring
+		// here when we have the calling context and the child context still available to us...
+		TArray<UEdGraphPin*> OutputPins;
+		GetOutputPins(OutputPins);
+
+		TArray<TPair<UEdGraphPin*, int32> > MatchedPairs;
+
+		const UEdGraphSchema_Niagara* Schema = CastChecked<UEdGraphSchema_Niagara>(GetSchema());
+
+		// Find the matches of names and types of the sub-graph output pins and this function call nodes' outputs.
+		for (UEdGraphPin* ChildOutputNodePin : OutputNode->GetAllPins())
+		{
+			FNiagaraVariable VarChild = Schema->PinToNiagaraVariable(ChildOutputNodePin);
+			
+			if (ChildOutputNodePin->LinkedTo.Num() > 0 && VarChild.GetType() == FNiagaraTypeDefinition::GetParameterMapDef())
+			{
+				for (int32 i = 0; i < OutputPins.Num(); i++)
+				{
+					FNiagaraVariable OutputVar = Schema->PinToNiagaraVariable(OutputPins[i]);
+					if (OutputVar.IsEquivalent(VarChild))
+					{
+						TPair<UEdGraphPin*, int32> Pair;
+						Pair.Key = OutputPins[i];
+						Pair.Value = OutHistory.TraceParameterMapOutputPin(ChildOutputNodePin->LinkedTo[0]);
+						MatchedPairs.Add(Pair);
+					}
+				}
+			}
+		}
+		OutHistory.ExitFunction(GetFunctionName(), FunctionScript, this);
+
+		for (int32 i = 0; i < MatchedPairs.Num(); i++)
+		{
+			OutHistory.RegisterParameterMapPin(MatchedPairs[i].Value, MatchedPairs[i].Key);
+		}
+	}
+}
+
+
+UEdGraphPin* UNiagaraNodeFunctionCall::FindParameterMapDefaultValuePin(const FString& VariableName)
+{
+	if (FunctionScript)
+	{
+		UNiagaraScriptSource* ScriptSource = Cast<UNiagaraScriptSource>(FunctionScript->GetSource());
+		if (ScriptSource != nullptr)
+		{
+			return ScriptSource->NodeGraph->FindParameterMapDefaultValuePin(VariableName);
+		}
+	}
+	return nullptr;
+}
+
+ENiagaraNumericOutputTypeSelectionMode UNiagaraNodeFunctionCall::GetNumericOutputTypeSelectionMode() const
+{
+	if (FunctionScript)
+	{
+		return FunctionScript->NumericOutputTypeSelectionMode;
+	}	
+	return ENiagaraNumericOutputTypeSelectionMode::None;
+}
+
+void UNiagaraNodeFunctionCall::AutowireNewNode(UEdGraphPin* FromPin)
+{
+	UNiagaraNode::AutowireNewNode(FromPin);
+	ComputeNodeName();
+}
+
+void UNiagaraNodeFunctionCall::ComputeNodeName()
+{
+	const FString CurrentName = FunctionDisplayName;
+	FName ProposedName = *CurrentName;
+	FString FunctionName = FunctionScript ? FunctionScript->GetName() : Signature.GetName();
+	if (!FunctionName.IsEmpty())
+	{
+		ProposedName = *FunctionName;
+	}
+
+	UNiagaraGraph* Graph = GetNiagaraGraph();
+	TArray<UNiagaraNodeFunctionCall*> Nodes;
+	Graph->GetNodesOfClass(Nodes);
+
+	TSet<FName> Names;
+	for (UNiagaraNodeFunctionCall* Node : Nodes)
+	{
+		CA_ASSUME(Node != nullptr);
+		if (Node != this)
+		{
+			Names.Add(*Node->GetFunctionName());
+		}
+	}
+
+	FString NewName = FNiagaraEditorUtilities::GetUniqueName(ProposedName, Names).ToString();
+	if (!FunctionDisplayName.Equals(NewName))
+	{
+		FunctionDisplayName = NewName;
+	}
+}
+
+bool UNiagaraNodeFunctionCall::FindAutoBoundInput(UNiagaraNodeInput* InputNode, UEdGraphPin* PinToAutoBind, FNiagaraVariable& OutFoundVar, ENiagaraInputNodeUsage& OutNodeUsage)
+{
+	check(InputNode && InputNode->IsExposed());
+	if (PinToAutoBind->LinkedTo.Num() > 0 || !InputNode->CanAutoBind())
+		return false;
+
+	const UEdGraphSchema_Niagara* Schema = CastChecked<UEdGraphSchema_Niagara>(GetSchema());
+	FNiagaraVariable PinVar = Schema->PinToNiagaraVariable(PinToAutoBind);
+
+	//See if we can auto bind this pin to something in the caller script.
+	UNiagaraGraph* CallerGraph = GetNiagaraGraph();
+	check(CallerGraph);
+	UNiagaraNodeOutput* CallerOutputNodeSpawn = CallerGraph->FindOutputNode(ENiagaraScriptUsage::ParticleSpawnScript);
+	UNiagaraNodeOutput* CallerOutputNodeUpdate = CallerGraph->FindOutputNode(ENiagaraScriptUsage::ParticleUpdateScript);
+
+	//First, lest see if we're an attribute of this emitter. Only valid if we're a module call off the primary script.
+	if (CallerOutputNodeSpawn || CallerOutputNodeUpdate)
+	{
+		UNiagaraNodeOutput* CallerOutputNode = CallerOutputNodeSpawn != nullptr ? CallerOutputNodeSpawn : CallerOutputNodeUpdate;
+		check(CallerOutputNode);
+		{
+			FNiagaraVariable* AttrVarPtr = CallerOutputNode->Outputs.FindByPredicate([&](const FNiagaraVariable& Attr) { return PinVar.IsEquivalent(Attr); });
+			if (AttrVarPtr)
+			{
+				OutFoundVar = *AttrVarPtr;
+				OutNodeUsage = ENiagaraInputNodeUsage::Attribute;
+				return true;
+			}
+		}
+	}
+	
+	//Next, lets see if we are a system constant.
+	//Do we need a smarter (possibly contextual) handling of system constants?
+	const TArray<FNiagaraVariable>& SysConstants = FNiagaraConstants::GetEngineConstants();
+	if (SysConstants.Contains(PinVar))
+	{
+		OutFoundVar = PinVar;
+		OutNodeUsage = ENiagaraInputNodeUsage::SystemConstant;
+		return true;
+	}
+
+	//Not sure it's a good idea to allow binding to user made parameters.
+// 	if (AutoBindOptions.bBindToParameters)
+// 	{
+// 		//Finally, lets see if we're a parameter of this emitter.
+// 		TArray<UNiagaraNodeInput*> CallerInputNodes;
+// 		CallerGraph->FindInputNodes(CallerInputNodes);
+// 		UNiagaraNodeInput** MatchingParamPtr = CallerInputNodes.FindByPredicate([&](UNiagaraNodeInput*& CallerInputNode)
+// 		{
+// 			return CallerInputNode->Input.IsEquivalent(PinVar);
+// 		});
+// 
+// 		if (MatchingParamPtr)
+// 		{
+// 			OutFoundVar = (*MatchingParamPtr)->Input;
+// 			OutNodeUsage = ENiagaraInputNodeUsage::Parameter;
+// 			return true;
+// 		}
+// 	}
+
+	//Unable to auto bind.
+	return false;
+}
+
+#undef LOCTEXT_NAMESPACE
diff --git a/Engine/Plugins/FX/Niagara/Source/NiagaraEditor/Private/NiagaraNodeIf.cpp b/Engine/Plugins/FX/Niagara/Source/NiagaraEditor/Private/NiagaraNodeIf.cpp
new file mode 100644
--- /dev/null
+++ b/Engine/Plugins/FX/Niagara/Source/NiagaraEditor/Private/NiagaraNodeIf.cpp
@@ -0,0 +1,191 @@
+// Copyright 1998-2016 Epic Games, Inc. All Rights Reserved.
+
+#include "NiagaraNodeIf.h"
+#include "NiagaraEditorUtilities.h"
+#include "NiagaraHlslTranslator.h"
+
+#define LOCTEXT_NAMESPACE "NiagaraNodeIf"
+
+UNiagaraNodeIf::UNiagaraNodeIf(const FObjectInitializer& ObjectInitializer)
+	: Super(ObjectInitializer)
+{
+}
+
+void UNiagaraNodeIf::PostEditChangeProperty(FPropertyChangedEvent& PropertyChangedEvent)
+{
+	// @TODO why do we need to have this post-change property here at all? 
+	// Doing a null check b/c otherwise if doing a Duplicate via Ctrl-W, we die inside AllocateDefaultPins due to 
+	// the point where we get this call not being completely formed.
+	if (PropertyChangedEvent.Property != nullptr)
+	{
+		ReallocatePins();
+	}
+}
+
+void UNiagaraNodeIf::PostLoad()
+{
+	Super::PostLoad();
+
+	if (OutputVars.Num() != OutputVarGuids.Num())
+	{
+		OutputVarGuids.SetNum(OutputVars.Num());
+		for (int32 i = 0; i < OutputVars.Num(); i++)
+		{
+			UEdGraphPin* Pin = FindPin(OutputVars[i].GetName().ToString(), EGPD_Output);
+			if (Pin != nullptr)
+			{
+				OutputVarGuids[i] = Pin->PersistentGuid;
+			}
+			else
+			{
+				UE_LOG(LogNiagaraEditor, Error, TEXT("Unable to find pin named %s"), *OutputVars[i].GetName().ToString());
+			}
+		}
+	}
+}
+
+bool UNiagaraNodeIf::AllowNiagaraTypeForAddPin(const FNiagaraTypeDefinition& InType)
+{
+	return Super::AllowNiagaraTypeForAddPin(InType) && InType != FNiagaraTypeDefinition::GetParameterMapDef();
+}
+
+void UNiagaraNodeIf::AllocateDefaultPins()
+{
+	const UEdGraphSchema_Niagara* Schema = GetDefault<UEdGraphSchema_Niagara>();
+
+	//Add the condition pin.
+	CreatePin(EGPD_Input, Schema->TypeDefinitionToPinType(FNiagaraTypeDefinition::GetBoolDef()), TEXT("Condition"));
+
+	//Create the inputs for each path.
+	for (FNiagaraVariable& Var : OutputVars)
+	{
+		FString PathSuffix = TEXT(" A");
+		CreatePin(EGPD_Input, Schema->TypeDefinitionToPinType(Var.GetType()), Var.GetName().ToString() + PathSuffix);		
+	}
+
+	for (FNiagaraVariable& Var : OutputVars)
+	{
+		FString PathSuffix = TEXT(" B");
+		CreatePin(EGPD_Input, Schema->TypeDefinitionToPinType(Var.GetType()), Var.GetName().ToString() + PathSuffix);
+	}
+
+	for (FNiagaraVariable& Var : OutputVars)
+	{
+		UEdGraphPin* NewPin = CreatePin(EGPD_Output, Schema->TypeDefinitionToPinType(Var.GetType()), Var.GetName().ToString());
+	}
+
+	CreateAddPin(EGPD_Output);
+}
+
+void UNiagaraNodeIf::Compile(class FHlslNiagaraTranslator* Translator, TArray<int32>& Outputs)
+{
+	const UEdGraphSchema_Niagara* Schema = CastChecked<UEdGraphSchema_Niagara>(GetSchema());
+
+	int32 PinIdx = 0;
+	int32 Condition = Translator->CompilePin(Pins[PinIdx++]);
+	TArray<int32> PathA;
+	PathA.Reserve(OutputVars.Num());
+	for (int32 i = 0; i < OutputVars.Num(); ++i)
+	{
+		if (Schema->PinToTypeDefinition(Pins[PinIdx]) == FNiagaraTypeDefinition::GetParameterMapDef())
+		{
+			Translator->Error(LOCTEXT("UnsupportedParamMapInIf","Parameter maps are not supported in if nodes."), this, Pins[PinIdx]);
+		}
+
+		PathA.Add(Translator->CompilePin(Pins[PinIdx++]));
+	}
+	TArray<int32> PathB;
+	PathB.Reserve(OutputVars.Num());
+	for (int32 i = 0; i < OutputVars.Num(); ++i)
+	{
+		if (Schema->PinToTypeDefinition(Pins[PinIdx]) == FNiagaraTypeDefinition::GetParameterMapDef())
+		{
+			Translator->Error(LOCTEXT("UnsupportedParamMapInIf", "Parameter maps are not supported in if nodes."), this, Pins[PinIdx]);
+		}
+
+		PathB.Add(Translator->CompilePin(Pins[PinIdx++]));
+	}
+	Translator->If(OutputVars, Condition, PathA, PathB, Outputs);
+}
+
+bool UNiagaraNodeIf::RefreshFromExternalChanges()
+{
+	// TODO - Leverage code in reallocate pins to determine if any pins have changed...
+	ReallocatePins();
+	return true;
+}
+
+void UNiagaraNodeIf::OnPinRemoved(UEdGraphPin* PinToRemove)
+{
+	auto FindPredicate = [=](const FGuid& Guid) { return Guid == PinToRemove->PersistentGuid; };
+	int32 FoundIndex = OutputVarGuids.IndexOfByPredicate(FindPredicate);
+	if (FoundIndex != INDEX_NONE)
+	{
+		OutputVarGuids.RemoveAt(FoundIndex);
+		OutputVars.RemoveAt(FoundIndex);
+	}
+	ReallocatePins();
+}
+
+void UNiagaraNodeIf::OnNewTypedPinAdded(UEdGraphPin* NewPin)
+{
+	Super::OnNewTypedPinAdded(NewPin);
+
+	const UEdGraphSchema_Niagara* Schema = GetDefault<UEdGraphSchema_Niagara>();
+	FNiagaraTypeDefinition OutputType = Schema->PinToTypeDefinition(NewPin);
+
+	TSet<FName> OutputNames;
+	for (const FNiagaraVariable& Output : OutputVars)
+	{
+		OutputNames.Add(Output.GetName());
+	}
+	FName OutputName = FNiagaraEditorUtilities::GetUniqueName(*OutputType.GetNameText().ToString(), OutputNames);
+
+	FNiagaraVariable NewOutput(OutputType, OutputName);
+	FGuid Guid = FGuid::NewGuid();
+	OutputVars.Add(NewOutput);
+	OutputVarGuids.Add(Guid);
+
+	// Update the pin's data too so that it's connection is maintained after reallocating.
+	NewPin->PersistentGuid = Guid;
+
+	ReallocatePins();
+}
+
+void UNiagaraNodeIf::OnPinRenamed(UEdGraphPin* RenamedPin)
+{
+	auto FindPredicate = [=](const FGuid& Guid) { return Guid == RenamedPin->PersistentGuid; };
+	int32 FoundIndex = OutputVarGuids.IndexOfByPredicate(FindPredicate);
+	if(FoundIndex != INDEX_NONE)
+	{
+		TSet<FName> OutputNames;
+		for (int32 Index = 0; Index < OutputVars.Num(); Index++)
+		{
+			if (FoundIndex != Index)
+			{
+				OutputNames.Add(OutputVars[Index].GetName());
+			}
+		}
+		FName OutputName = FNiagaraEditorUtilities::GetUniqueName(*RenamedPin->PinName, OutputNames);
+		OutputVars[FoundIndex].SetName(OutputName);
+	}
+	ReallocatePins();
+}
+
+bool UNiagaraNodeIf::CanRenamePin(const UEdGraphPin* Pin) const
+{
+	return Super::CanRenamePin(Pin) && Pin->Direction == EGPD_Output;
+}
+
+bool UNiagaraNodeIf::CanRemovePin(const UEdGraphPin* Pin) const
+{
+	return Super::CanRemovePin(Pin) && Pin->Direction == EGPD_Output;
+}
+
+
+FText UNiagaraNodeIf::GetTooltipText() const
+{
+	return LOCTEXT("IfDesc", "If Condition is true, the output value is A, otherwise output B.");
+}
+
+#undef LOCTEXT_NAMESPACE
diff --git a/Engine/Plugins/FX/Niagara/Source/NiagaraEditor/Private/NiagaraNodeIf.h b/Engine/Plugins/FX/Niagara/Source/NiagaraEditor/Private/NiagaraNodeIf.h
new file mode 100644
--- /dev/null
+++ b/Engine/Plugins/FX/Niagara/Source/NiagaraEditor/Private/NiagaraNodeIf.h
@@ -0,0 +1,50 @@
+// Copyright 1998-2016 Epic Games, Inc. All Rights Reserved.
+
+#pragma once
+#include "NiagaraNodeWithDynamicPins.h"
+#include "EdGraphSchema_Niagara.h"
+#include "NiagaraNodeIf.generated.h"
+
+UCLASS(MinimalAPI)
+class UNiagaraNodeIf : public UNiagaraNodeWithDynamicPins
+{
+	GENERATED_UCLASS_BODY()
+
+public:
+
+	/** Outputs of this branch. */
+	UPROPERTY(EditAnywhere, Category=If)
+	TArray<FNiagaraVariable> OutputVars;
+
+	UPROPERTY(EditAnywhere, Category = If)
+	TArray<FGuid> OutputVarGuids;
+
+	//~ Begin UObject Interface
+	virtual void PostEditChangeProperty(FPropertyChangedEvent& PropertyChangedEvent) override;
+	virtual void PostLoad() override;
+	//~ End UObject Interface
+
+	//~ Begin EdGraphNode Interface
+	virtual void AllocateDefaultPins() override;
+	virtual FText GetTooltipText() const override;
+	//~ End EdGraphNode Interface
+
+	//~ Begin UNiagaraNode Interface
+	virtual void Compile(class FHlslNiagaraTranslator* Translator, TArray<int32>& Outputs) override;
+	virtual bool RefreshFromExternalChanges() override;
+	//~ End UNiagaraNode Interface
+
+protected:
+	//~ Begin EdGraphNode Interface
+	virtual void OnPinRemoved(UEdGraphPin* PinToRemove) override;
+	//~ End EdGraphNode Interface
+
+	//~ Begin UNiagaraNodeWithDynamicPins Interface
+	virtual void OnNewTypedPinAdded(UEdGraphPin* NewPin) override;
+	virtual void OnPinRenamed(UEdGraphPin* RenamedPin) override;
+	virtual bool CanRenamePin(const UEdGraphPin* Pin) const override;
+	virtual bool CanRemovePin(const UEdGraphPin* Pin) const override;
+	virtual bool CanMovePin(const UEdGraphPin* Pin) const override { return false; }
+	virtual bool AllowNiagaraTypeForAddPin(const FNiagaraTypeDefinition& InType) override;
+	//~ End UNiagaraNodeWithDynamicPins Interface
+};
diff --git a/Engine/Plugins/FX/Niagara/Source/NiagaraEditor/Private/NiagaraNodeInput.cpp b/Engine/Plugins/FX/Niagara/Source/NiagaraEditor/Private/NiagaraNodeInput.cpp
new file mode 100644
--- /dev/null
+++ b/Engine/Plugins/FX/Niagara/Source/NiagaraEditor/Private/NiagaraNodeInput.cpp
@@ -0,0 +1,490 @@
+// Copyright 1998-2017 Epic Games, Inc. All Rights Reserved.
+
+#include "NiagaraNodeInput.h"
+#include "UObject/UnrealType.h"
+#include "NiagaraHlslTranslator.h"
+#include "NiagaraGraph.h"
+
+#include "NiagaraNodeOutput.h"
+#include "NiagaraNodeOp.h"
+#include "NiagaraNodeFunctionCall.h"
+#include "EdGraphSchema_Niagara.h"
+#include "NiagaraDataInterface.h"
+#include "SNiagaraGraphNodeInput.h"
+#include "NiagaraEditorUtilities.h"
+
+
+#define LOCTEXT_NAMESPACE "NiagaraNodeInput"
+
+
+UNiagaraNodeInput::UNiagaraNodeInput(const FObjectInitializer& ObjectInitializer)
+	: Super(ObjectInitializer)
+	, DataInterface(nullptr)
+	, Usage(ENiagaraInputNodeUsage::Undefined)	
+	, CallSortPriority(0)
+{
+	bCanRenameNode = true;
+}
+
+void UNiagaraNodeInput::PostEditChangeProperty(struct FPropertyChangedEvent& PropertyChangedEvent)
+{
+	if (PropertyChangedEvent.Property != nullptr)
+	{
+		if (UClass* Class = const_cast<UClass*>(Input.GetType().GetClass()))
+		{
+			check(Class->IsChildOf(UNiagaraDataInterface::StaticClass()));
+			if (DataInterface)
+			{
+				if (DataInterface->GetClass() != Class)
+				{
+					//Class has changed so clear this out and allocate pins will create a new instance of the correct type.
+					//Should we preserve old objects somewhere so settings aren't lost when switching around types?
+					DataInterface = nullptr;
+				}
+				else
+				{
+					//Keep it with the same name as the input 
+					if (PropertyChangedEvent.Property->GetName() == TEXT("Input"))
+					{
+						DataInterface->Rename(*Input.GetName().ToString());
+					}
+				}
+			}
+		}
+		else
+		{
+			DataInterface = nullptr;
+		}
+
+		ReallocatePins();
+	}
+
+	Super::PostEditChangeProperty(PropertyChangedEvent);
+}
+
+void UNiagaraNodeInput::BuildParameterMapHistory(FNiagaraParameterMapHistoryBuilder& OutHistory, bool bRecursive)
+{
+	if (Input.GetType() == FNiagaraTypeDefinition::GetParameterMapDef())
+	{
+		int32 ParamMapIdx = OutHistory.FindMatchingParameterMapFromContextInputs(Input);
+
+		if (ParamMapIdx == INDEX_NONE)
+		{
+			ParamMapIdx = OutHistory.CreateParameterMap();
+		}
+		OutHistory.RegisterParameterMapPin(ParamMapIdx, GetOutputPin(0));
+	}
+}
+
+
+void UNiagaraNodeInput::AllocateDefaultPins()
+{
+	if (UClass* Class = const_cast<UClass*>(Input.GetType().GetClass()))
+	{
+		check(Class->IsChildOf(UNiagaraDataInterface::StaticClass()));
+		if (!DataInterface)
+		{
+			DataInterface = NewObject<UNiagaraDataInterface>(this, Class);
+		}
+	}
+
+	const UEdGraphSchema_Niagara* Schema = GetDefault<UEdGraphSchema_Niagara>();
+
+	//If we're a parameter node for a funciton or a module the we allow a "default" input pin.
+	UNiagaraScript* OwnerScript = GetTypedOuter<UNiagaraScript>();
+	if (OwnerScript)
+	{
+		if ((!IsRequired() && IsExposed()) && DataInterface == nullptr && Usage == ENiagaraInputNodeUsage::Parameter && (OwnerScript->IsStandaloneScript()))
+		{
+			UEdGraphPin* NewPin = CreatePin(EGPD_Input, Schema->TypeDefinitionToPinType(Input.GetType()), TEXT("Default"));
+			NewPin->bDefaultValueIsIgnored = true;
+		}
+	}
+
+	CreatePin(EGPD_Output, Schema->TypeDefinitionToPinType(Input.GetType()), TEXT("Input"));
+}
+
+FText UNiagaraNodeInput::GetNodeTitle(ENodeTitleType::Type TitleType) const
+{
+	return FText::FromName(Input.GetName());
+}
+
+int32 UNiagaraNodeInput::GenerateNewSortPriority(const UNiagaraGraph* Graph, FName& ProposedName, ENiagaraInputNodeUsage Usage)
+{
+	int32 HighestSortOrder = -1; // Set to -1 initially, so that if there are no nodes, the return value will be zero.
+	
+	if (Usage == ENiagaraInputNodeUsage::Parameter && Graph != nullptr)
+	{
+		TArray<UNiagaraNodeInput*> InputNodes;
+		Graph->GetNodesOfClass(InputNodes);
+		for (UNiagaraNodeInput* InputNode : InputNodes)
+		{
+			if (InputNode->Usage == Usage && InputNode->CallSortPriority > HighestSortOrder)
+			{
+				HighestSortOrder = InputNode->CallSortPriority;
+			}
+		}
+	}
+	return HighestSortOrder + 1;
+}
+
+FName UNiagaraNodeInput::GenerateUniqueName(const UNiagaraGraph* Graph, FName& ProposedName, ENiagaraInputNodeUsage Usage)
+{
+	check(Usage != ENiagaraInputNodeUsage::SystemConstant && Usage != ENiagaraInputNodeUsage::Undefined);
+	TSet<FName> InputNames;
+	if (Usage == ENiagaraInputNodeUsage::Parameter && Graph != nullptr)
+	{
+		TArray<UNiagaraNodeInput*> InputNodes;
+		Graph->GetNodesOfClass(InputNodes);
+		for (UNiagaraNodeInput* InputNode : InputNodes)
+		{
+			if (InputNode->Usage == Usage)
+			{
+				InputNames.Add(InputNode->Input.GetName());
+			}
+		}
+	}
+	else if (Usage == ENiagaraInputNodeUsage::Attribute && Graph != nullptr)
+	{
+		TArray<UNiagaraNodeOutput*> OutputNodes;
+		Graph->GetNodesOfClass<UNiagaraNodeOutput>(OutputNodes);
+		for (UNiagaraNodeOutput* Node : OutputNodes)
+		{
+			for (const FNiagaraVariable& Output : Node->Outputs)
+			{
+				InputNames.Add(Output.GetName());
+			}
+		}
+	}
+
+	return FNiagaraEditorUtilities::GetUniqueName(ProposedName, FNiagaraEditorUtilities::GetSystemConstantNames().Union(InputNames));
+}
+
+bool UNiagaraNodeInput::VerifyNodeRenameTextCommit(const FText& NewText, UNiagaraNode* NodeBeingChanged, FText& OutErrorMessage)
+{
+	FName NewName = *NewText.ToString();
+	TSet<FName> SystemConstantNames = FNiagaraEditorUtilities::GetSystemConstantNames();
+
+	// Disallow empty names
+	if (NewName == FName())
+	{
+		OutErrorMessage = LOCTEXT("NiagaraInputNameEmptyWarn", "Cannot have empty name!");
+		return false;
+	}
+
+	// Disallow name changes to system constants
+	if (SystemConstantNames.Contains(NewName))
+	{
+		OutErrorMessage = FText::Format(LOCTEXT("NiagaraInputNameSystemWarn", "\"{0}\" is a system constant name."), FText::FromName(NewName));
+		return false;
+	}
+
+	// @TODO: Prevent any hlsl keywords or invalid hlsl characters from being used as names!
+
+	UNiagaraNodeInput* InputNodeBeingChanged = Cast<UNiagaraNodeInput>(NodeBeingChanged);
+	UNiagaraNodeOutput* OutputNodeBeingChanged = Cast<UNiagaraNodeOutput>(NodeBeingChanged);
+
+	// Make sure that we aren't changing names to something already in the graph.
+	if (NodeBeingChanged != nullptr)
+	{
+		UNiagaraGraph* Graph = CastChecked<UNiagaraGraph>(NodeBeingChanged->GetGraph());
+
+		// If dealing with parameter, check to make sure that we don't conflict with any other parameter name.
+		if (InputNodeBeingChanged != nullptr && InputNodeBeingChanged->Usage == ENiagaraInputNodeUsage::Parameter)
+		{
+			TArray<UNiagaraNodeInput*> InputNodes;
+			Graph->GetNodesOfClass<UNiagaraNodeInput>(InputNodes);
+
+			for (UNiagaraNodeInput* Node : InputNodes)
+			{
+				if (Node == nullptr || Node == InputNodeBeingChanged || Node->Usage != InputNodeBeingChanged->Usage)
+				{
+					continue;
+				}
+
+				bool bIsSame = Node->ReferencesSameInput(InputNodeBeingChanged);
+				// This should still allow case changes b/c we test to make sure that they aren't referencing the same node.
+				if (bIsSame == false && Node->Input.GetName().IsEqual(NewName, ENameCase::IgnoreCase))
+				{
+					OutErrorMessage = FText::Format(LOCTEXT("NiagaraInputNameSameParameterWarn", "\"{0}\" is the name of another parameter."), FText::FromName(NewName));
+					return false;
+				}
+			}
+		}
+
+		// If dealing with Attributes, check to make sure that we don't conflict with any other attribute name.
+		if ((InputNodeBeingChanged != nullptr && InputNodeBeingChanged->Usage == ENiagaraInputNodeUsage::Attribute) || OutputNodeBeingChanged != nullptr)
+		{
+			TArray<UNiagaraNodeOutput*> OutputNodes;
+			Graph->GetNodesOfClass<UNiagaraNodeOutput>(OutputNodes);
+			for (UNiagaraNodeOutput* Node : OutputNodes)
+			{
+				for (const FNiagaraVariable& Output : Node->Outputs)
+				{
+					if (InputNodeBeingChanged != nullptr && Output.GetName().IsEqual(InputNodeBeingChanged->Input.GetName(), ENameCase::IgnoreCase))
+					{
+						continue;
+					}
+
+					if (Output.GetName().IsEqual(NewName, ENameCase::IgnoreCase))
+					{
+						OutErrorMessage = FText::Format(LOCTEXT("NiagaraInputNameSameAttributeWarn", "\"{0}\" is the name of another attribute. Hit \"Escape\" to cancel edit."), FText::FromName(NewName));
+						return false;
+					}
+				}
+			}
+		}
+	}
+	return true;
+}
+
+void UNiagaraNodeInput::OnRenameNode(const FString& NewName)
+{
+	UNiagaraGraph* Graph = CastChecked<UNiagaraGraph>(GetGraph());
+	TArray<UNiagaraNodeInput*> InputNodes;
+	Graph->GetNodesOfClass<UNiagaraNodeInput>(InputNodes);
+
+	TArray<UNiagaraNodeInput*> AffectedNodes;
+	AffectedNodes.Add(this);
+	for (UNiagaraNodeInput* Node : InputNodes)
+	{
+		if (Node == this || Node == nullptr)
+		{
+			continue;
+		}
+
+		bool bNeedsSync = Node->ReferencesSameInput(this);
+		if (bNeedsSync)
+		{
+			AffectedNodes.Add(Node);
+		}
+	}
+
+	for (UNiagaraNodeInput* Node : AffectedNodes)
+	{
+		Node->Modify();
+		Node->Input.SetName(FName(*NewName));
+		if (DataInterface != nullptr)
+		{
+			Node->DataInterface->Rename(*NewName);
+		}
+		Node->ReallocatePins();
+	}
+
+	Graph->MarkGraphRequiresSynchronization();
+}
+
+TSharedPtr<SGraphNode> UNiagaraNodeInput::CreateVisualWidget()
+{
+	return SNew(SNiagaraGraphNodeInput, this);
+}
+
+FLinearColor UNiagaraNodeInput::GetNodeTitleColor() const
+{
+	const UEdGraphSchema_Niagara* Schema = CastChecked<UEdGraphSchema_Niagara>(GetSchema());
+	switch (Usage)
+	{
+	case ENiagaraInputNodeUsage::Parameter:
+		return Schema->NodeTitleColor_Constant;
+	case ENiagaraInputNodeUsage::SystemConstant:
+		return Schema->NodeTitleColor_SystemConstant;
+	case ENiagaraInputNodeUsage::Attribute:
+		return Schema->NodeTitleColor_Attribute;
+	default:
+		// TODO: Do something better here.
+		return FLinearColor::Black;
+	}
+}
+
+bool UNiagaraNodeInput::ReferencesSameInput(UNiagaraNodeInput* Other) const
+{
+	if (this == Other)
+	{
+		return true;
+	}
+
+	if (Other == nullptr)
+	{
+		return false;
+	}
+
+	if (Usage == Other->Usage)
+	{	
+		if (Input.GetName() == Other->Input.GetName())
+		{
+			return true;
+		}	
+	}
+	return false;
+}
+
+
+void UNiagaraNodeInput::AutowireNewNode(UEdGraphPin* FromPin)
+{
+	if (FromPin != nullptr)
+	{
+		const UEdGraphSchema_Niagara* Schema = CastChecked<UEdGraphSchema_Niagara>(GetSchema());
+		check(Schema);
+		if (Usage == ENiagaraInputNodeUsage::Parameter)
+		{
+			TArray<UNiagaraNodeInput*> InputNodes;
+			GetGraph()->GetNodesOfClass(InputNodes);
+			TSet<FName> InputNames;
+			int32 NumMatches = 0;
+			int32 HighestSortPriority = -1; // Set to -1 initially, so that in the event of no nodes, we still get zero.
+			for (UNiagaraNodeInput* InputNode : InputNodes)
+			{
+				if (InputNode == nullptr)
+				{
+					continue;
+				}
+
+				if (InputNode != this && InputNode->Usage == ENiagaraInputNodeUsage::Parameter)
+				{
+					if (this->ReferencesSameInput(InputNode))
+					{
+						NumMatches++;
+						check(InputNode->Input.GetName() == this->Input.GetName());
+						check(InputNode->ExposureOptions.bCanAutoBind == this->ExposureOptions.bCanAutoBind);
+						check(InputNode->ExposureOptions.bExposed == this->ExposureOptions.bExposed);
+						check(InputNode->ExposureOptions.bHidden == this->ExposureOptions.bHidden);
+						check(InputNode->ExposureOptions.bRequired == this->ExposureOptions.bRequired);
+						check(InputNode->DataInterface == this->DataInterface);
+						check(InputNode->CallSortPriority == this->CallSortPriority);
+					}
+					InputNames.Add(InputNode->Input.GetName());
+					if (InputNode->CallSortPriority > HighestSortPriority)
+					{
+						HighestSortPriority = InputNode->CallSortPriority;
+					}
+				}
+			}
+
+			FName CandidateName = Input.GetName();
+			FNiagaraTypeDefinition Type = Input.GetType();
+			if (Type == FNiagaraTypeDefinition::GetGenericNumericDef())
+			{
+				//Try to get a real type if we've been set to numeric
+				Type = Schema->PinToTypeDefinition(FromPin);
+			}
+
+			if (NumMatches == 0)
+			{
+				FString PinName = FromPin->PinFriendlyName.IsEmpty() ? FromPin->PinName : FromPin->PinFriendlyName.ToString();
+				if (UNiagaraNodeOp* OpNode = Cast<UNiagaraNodeOp>(FromPin->GetOwningNode()))
+				{
+					const FNiagaraOpInfo* OpInfo = FNiagaraOpInfo::GetOpInfo(OpNode->OpName);
+					check(OpInfo);
+					Input.SetType(Type);
+					CandidateName = (*(OpInfo->FriendlyName.ToString() + TEXT(" ") + PinName));
+				}
+				else if (UNiagaraNodeFunctionCall* FuncNode = Cast<UNiagaraNodeFunctionCall>(FromPin->GetOwningNode()))
+				{
+					Input.SetType(Type);
+					CandidateName = (*(FuncNode->GetNodeTitle(ENodeTitleType::FullTitle).ToString() + TEXT(" ") + PinName));
+				}
+				else
+				{
+					Input.SetType(Type);
+					CandidateName = (*PinName);
+				}
+
+				Input.SetName(FNiagaraEditorUtilities::GetUniqueName(CandidateName, FNiagaraEditorUtilities::GetSystemConstantNames().Union(InputNames)));
+				CallSortPriority = HighestSortPriority + 1;
+			}
+			ReallocatePins();
+		}
+
+		TArray<UEdGraphPin*> OutPins;
+		GetOutputPins(OutPins);
+		check(OutPins.Num() == 1 && OutPins[0] != NULL);
+		
+		if (GetSchema()->TryCreateConnection(FromPin, OutPins[0]))
+		{
+			FromPin->GetOwningNode()->NodeConnectionListChanged();
+		}
+	}
+}
+
+void UNiagaraNodeInput::NotifyInputTypeChanged()
+{
+	ReallocatePins();
+}
+
+
+void UNiagaraNodeInput::NotifyExposureOptionsChanged()
+{
+	ReallocatePins();
+}
+
+void UNiagaraNodeInput::Compile(class FHlslNiagaraTranslator* Translator, TArray<int32>& Outputs)
+{
+	int32 FunctionParam = INDEX_NONE;
+	if (IsExposed() && Translator->GetFunctionParameter(Input, FunctionParam))
+	{
+		//If we're in a function and this parameter hasn't been provided, compile the local default.
+		if (FunctionParam == INDEX_NONE)
+		{
+			TArray<UEdGraphPin*> InputPins;
+			GetInputPins(InputPins);
+			int32 Default = InputPins.Num() > 0 ? Translator->CompilePin(InputPins[0]) : INDEX_NONE;
+			if (Default == INDEX_NONE)
+			{
+				//We failed to compile the default pin so just use the value of the input.
+				if (Usage == ENiagaraInputNodeUsage::Parameter && DataInterface != nullptr)
+				{
+					check(Input.GetType().GetClass());
+					Outputs.Add(Translator->RegisterDataInterface(Input, DataInterface));
+					return;
+				}
+				else
+				{
+					Default = Translator->GetConstant(Input);
+				}
+			}
+			Outputs.Add(Default);
+			return;
+		}
+	}
+
+	switch (Usage)
+	{
+	case ENiagaraInputNodeUsage::Parameter:
+		if (DataInterface)
+		{
+			check(Input.GetType().GetClass());
+			Outputs.Add(Translator->RegisterDataInterface(Input, DataInterface)); break;
+		}
+		else
+		{
+			Outputs.Add(Translator->GetParameter(Input)); break;
+		}
+	case ENiagaraInputNodeUsage::SystemConstant:
+		Outputs.Add(Translator->GetParameter(Input)); break;
+	case ENiagaraInputNodeUsage::Attribute:
+		Outputs.Add(Translator->GetAttribute(Input)); break;
+	default:
+		check(false);
+	}
+}
+
+void UNiagaraNodeInput::SortNodes(TArray<UNiagaraNodeInput*>& InOutNodes)
+{
+	auto SortVars = [](UNiagaraNodeInput& A, UNiagaraNodeInput& B)
+	{
+		if (A.CallSortPriority < B.CallSortPriority)
+		{
+			return true;
+		}
+		else if (A.CallSortPriority > B.CallSortPriority)
+		{
+			return false;
+		}
+
+		//If equal priority, sort lexicographically.
+		return A.Input.GetName().ToString() < B.Input.GetName().ToString();
+	};
+	InOutNodes.Sort(SortVars);
+}
+
+#undef LOCTEXT_NAMESPACE
