Copying //UE4/Orion-Staging to //UE4/Main (Source: //Orion/Dev-General @ 3606378)

diff --git a/Engine/Plugins/FX/Niagara/Source/NiagaraEditor/Private/NiagaraNodeParameterCollection.cpp b/Engine/Plugins/FX/Niagara/Source/NiagaraEditor/Private/NiagaraNodeParameterCollection.cpp
new file mode 100644
--- /dev/null
+++ b/Engine/Plugins/FX/Niagara/Source/NiagaraEditor/Private/NiagaraNodeParameterCollection.cpp
@@ -0,0 +1,206 @@
+// Copyright 1998-2016 Epic Games, Inc. All Rights Reserved.
+
+#include "NiagaraNodeParameterCollection.h"
+#include "EdGraphSchema_Niagara.h"
+#include "NiagaraEditorUtilities.h"
+#include "SNiagaraGraphNodeConvert.h"
+#include "NiagaraGraph.h"
+#include "NiagaraHlslTranslator.h"
+#include "ScopedTransaction.h"
+#include "NiagaraParameterCollection.h"
+#include "NiagaraCommon.h"
+#include "ModuleManager.h"
+#include "AssetRegistryModule.h"
+#include "SNiagaraGraphPinAdd.h"
+#include "MultiBoxBuilder.h"
+
+#define LOCTEXT_NAMESPACE "NiagaraNodeParameterCollection"
+
+UNiagaraNodeParameterCollection::UNiagaraNodeParameterCollection() : UNiagaraNodeWithDynamicPins()
+{
+
+}
+
+void UNiagaraNodeParameterCollection::AllocateDefaultPins()
+{
+	if (CollectionAssetObjectPath != NAME_None && Collection == nullptr)
+	{
+		FAssetRegistryModule& AssetRegistryModule = FModuleManager::LoadModuleChecked<FAssetRegistryModule>(TEXT("AssetRegistry"));
+		FAssetData AssetData = AssetRegistryModule.Get().GetAssetByObjectPath(CollectionAssetObjectPath);
+		if (AssetData.IsValid())
+		{
+			Collection = Cast<UNiagaraParameterCollection>(AssetData.GetAsset());
+		}
+		else
+		{
+			UE_LOG(LogNiagaraEditor, Warning, TEXT("Failed to load Niagara Parameter Collection: {%s}"), *CollectionAssetObjectPath.ToString());
+		}
+	}
+
+ 	const UEdGraphSchema_Niagara* Schema = GetDefault<UEdGraphSchema_Niagara>();
+	if (Collection)
+	{
+		for (const FNiagaraVariable& Var : Variables)
+		{
+			if (Collection->IndexOfParameter(Var) != INDEX_NONE)
+			{
+				FString VarName = Collection->FriendlyNameFromParameterName(Var.GetName().ToString());
+				CreatePin(EGPD_Output, Schema->TypeDefinitionToPinType(Var.GetType()), VarName);
+			}
+		}
+	}
+
+	CreateAddPin(EGPD_Output);
+}
+
+bool UNiagaraNodeParameterCollection::RefreshFromExternalChanges()
+{
+	//Update to reflect any collection changes.
+	if (Collection)
+	{
+		const UEdGraphSchema_Niagara* Schema = GetDefault<UEdGraphSchema_Niagara>();
+		
+		for (UEdGraphPin* Pin : Pins)
+		{
+			if (IsAddPin(Pin))
+			{
+				continue;
+			}
+			FNiagaraVariable PinVar = Collection->CollectionParameterFromFriendlyParameter(Schema->PinToNiagaraVariable(Pin));
+			if (Collection->IndexOfParameter(PinVar) == INDEX_NONE)
+			{
+				//Can't find this parameter in the collection so we must refresh.
+				ReallocatePins();
+				return true;
+			}
+		}
+	}
+	
+	return false;
+}
+
+#if WITH_EDITOR
+
+void UNiagaraNodeParameterCollection::PostEditImport()
+{
+	Super::PostEditImport();
+	//Pins.Empty();
+}
+
+void UNiagaraNodeParameterCollection::PostEditChangeProperty(FPropertyChangedEvent& PropertyChangedEvent)
+{
+	Super::PostEditChangeProperty(PropertyChangedEvent);
+	//RefreshFromExternalChanges();
+}
+
+#endif
+
+void UNiagaraNodeParameterCollection::PostLoad()
+{
+	Super::PostLoad();
+}
+
+bool UNiagaraNodeParameterCollection::CanRenamePin(const UEdGraphPin* GraphPinObj) const
+{
+	return false;
+}
+
+void UNiagaraNodeParameterCollection::OnNewTypedPinAdded(UEdGraphPin* NewPin)
+{
+ 	check(Collection);
+ 	const UEdGraphSchema_Niagara* Schema = GetDefault<UEdGraphSchema_Niagara>();
+	FNiagaraVariable NewVar = Schema->PinToNiagaraVariable(NewPin);
+	Variables.Add(Collection->CollectionParameterFromFriendlyParameter(NewVar));
+}
+
+FText UNiagaraNodeParameterCollection::GetNodeTitle(ENodeTitleType::Type TitleType) const
+{
+	return Collection ? FText::FromString(Collection->GetName()) : LOCTEXT("ParameterCollectionNodeTitle", "Parameter Collection");
+}
+
+FText UNiagaraNodeParameterCollection::GetTooltipText()const
+{
+	return FText::FromString(Collection->GetPathName());
+}
+
+TSharedRef<SWidget> UNiagaraNodeParameterCollection::GenerateAddPinMenu(const FString& InWorkingPinName, SNiagaraGraphPinAdd* InPin)
+{
+	FMenuBuilder MenuBuilder(true, nullptr);
+	if (Collection)
+	{
+		for (FNiagaraVariable Var : Collection->GetParameters())
+		{
+			if(!Variables.Contains(Var))
+			{
+				FNiagaraVariable FriendlyVar = Collection->FriendlyParameterFromCollectionParameter(Var);
+				MenuBuilder.AddMenuEntry(
+					FText::FromName(FriendlyVar.GetName()),
+					FText::Format(LOCTEXT("AddButtonTypeEntryToolTipFormatCollection", "Add a reference to collection parameter {0}"), FText::FromName(FriendlyVar.GetName())),
+					FSlateIcon(),
+					FUIAction(FExecuteAction::CreateSP(InPin, &SNiagaraGraphPinAdd::OnAddType, FriendlyVar)));
+			}
+		}
+	}
+	return MenuBuilder.MakeWidget();
+}
+
+void UNiagaraNodeParameterCollection::RemoveDynamicPin(UEdGraphPin* Pin)
+{
+	Super::RemoveDynamicPin(Pin);
+
+	const UEdGraphSchema_Niagara* Schema = GetDefault<UEdGraphSchema_Niagara>();
+	FNiagaraVariable RemVar = Collection->CollectionParameterFromFriendlyParameter(Schema->PinToNiagaraVariable(Pin));
+	Variables.Remove(RemVar);
+}
+
+// void UNiagaraNodeParameterCollection::BuildParameterMapHistory(const UEdGraphPin* ParameterMapPinFromNode, FNiagaraParameterMapHistory& OutHistory, bool bRecursive)
+// {
+// 	if (Collection)
+// 	{
+// 		const UEdGraphSchema_Niagara* Schema = GetDefault<UEdGraphSchema_Niagara>();
+// 		OutHistory.RegisterNodeVisitation(this);
+// 		if (GetInputPin(0) == ParameterMapPinFromNode)
+// 		{
+// 			FString Namespace = GetCollectionNamespace();
+// 			TArray<UEdGraphPin*> OutputPins;
+// 			GetOutputPins(OutputPins);
+// 			for (int32 i = 0; i < OutputPins.Num(); i++)
+// 			{
+// 				if (IsAddPin(OutputPins[i]))
+// 				{
+// 					continue;
+// 				}
+// 
+// 				FNiagaraVariable Var = Schema->PinToNiagaraVariable(OutputPins[i]);
+// 				Var = FNiagaraParameterMapHistory::VariableToNamespacedVariable(Var, Namespace);
+// 				OutHistory.AddExternalVariable(Var);
+// 			}
+// 		}
+// 	}
+// }
+
+UObject* UNiagaraNodeParameterCollection::GetReferencedAsset()const 
+{
+	return Collection; 
+}
+
+void UNiagaraNodeParameterCollection::Compile(class FHlslNiagaraTranslator* Translator, TArray<int32>& Outputs)
+{
+	if (Collection)
+	{
+		UNiagaraGraph* Graph = Cast<UNiagaraGraph>(GetGraph());
+		const UEdGraphSchema_Niagara* Schema = CastChecked<UEdGraphSchema_Niagara>(GetSchema());
+
+		TArray<UEdGraphPin*> InputPins;
+		GetInputPins(InputPins);
+		check(InputPins.Num() == 0);
+
+		Translator->ParameterCollection(this, Outputs);
+	}
+	else
+	{
+		Translator->Error(LOCTEXT("Invalid Collection", "Parameter Collection is invalid."), this, nullptr);
+	}
+}
+
+#undef LOCTEXT_NAMESPACE
diff --git a/Engine/Plugins/FX/Niagara/Source/NiagaraEditor/Private/NiagaraNodeParameterCollection.h b/Engine/Plugins/FX/Niagara/Source/NiagaraEditor/Private/NiagaraNodeParameterCollection.h
new file mode 100644
--- /dev/null
+++ b/Engine/Plugins/FX/Niagara/Source/NiagaraEditor/Private/NiagaraNodeParameterCollection.h
@@ -0,0 +1,73 @@
+// Copyright 1998-2016 Epic Games, Inc. All Rights Reserved.
+
+#pragma once
+
+#include "SGraphPin.h"
+#include "NiagaraNodeWithDynamicPins.h"
+#include "NiagaraNodeParameterCollection.generated.h"
+
+class UEdGraphPin;
+class UNiagaraParameterCollection;
+
+/** A node that allows a user to get values from a parameter collection. */
+UCLASS()
+class UNiagaraNodeParameterCollection : public UNiagaraNodeWithDynamicPins
+{
+public:
+	GENERATED_BODY()
+
+	UNiagaraNodeParameterCollection();
+
+	//~ UObject interface
+	virtual void PostEditImport()override;
+	virtual void PostEditChangeProperty(FPropertyChangedEvent& PropertyChangedEvent)override;
+	virtual void PostLoad()override;
+	//~ UObject interface
+
+	//~ UEdGraphNode Interface
+	virtual void AllocateDefaultPins() override;
+	virtual FText GetNodeTitle(ENodeTitleType::Type TitleType) const;
+	virtual FText GetTooltipText() const;
+	//~ UEdGraphNode Interface
+
+	//~ UNiagaraNode Interface
+	virtual void Compile(class FHlslNiagaraTranslator* Translator, TArray<int32>& Outputs);
+	virtual UObject* GetReferencedAsset()const;
+	virtual bool RefreshFromExternalChanges()override;
+	//~ UNiagaraNode Interface
+
+	//~ UNiagaraNodeWithDynamicPins Interface
+	virtual bool CanRenamePin(const UEdGraphPin* GraphPinObj) const override;
+	virtual bool CanMovePin(const UEdGraphPin* Pin) const override { return false; }
+	virtual TSharedRef<SWidget> GenerateAddPinMenu(const FString& InWorkingPinName, SNiagaraGraphPinAdd* InPin) override;
+protected:
+	virtual void RemoveDynamicPin(UEdGraphPin* Pin)override;
+public:
+	//~ UNiagaraNodeWithDynamicPins Interface
+
+
+	//virtual void BuildParameterMapHistory(const UEdGraphPin* ParameterMapPinFromNode, FNiagaraParameterMapHistory& OutHistory, bool bRecursive = true) override;
+
+	UNiagaraParameterCollection* GetCollection()const { return Collection; }
+	/**
+	* A path to a collection asset.
+	* This is used so that the nodes can be populated in the graph context
+	* menu without having to load all of the actual collection assets.
+	*/
+	UPROPERTY(Transient)
+	FName CollectionAssetObjectPath;
+
+protected:
+
+	virtual void OnNewTypedPinAdded(UEdGraphPin* NewPin);
+
+	UPROPERTY(EditAnywhere, Category=Collection)
+	UNiagaraParameterCollection* Collection;
+
+	UPROPERTY()
+	TArray<FNiagaraVariable> Variables;
+
+	//Track change ID? Or just validate on load every time and on edit if already loaded.
+	//UPROPERTY()
+	//FGuid ChangeID;
+};
diff --git a/Engine/Plugins/FX/Niagara/Source/NiagaraEditor/Private/NiagaraNodeParameterMapBase.cpp b/Engine/Plugins/FX/Niagara/Source/NiagaraEditor/Private/NiagaraNodeParameterMapBase.cpp
new file mode 100644
--- /dev/null
+++ b/Engine/Plugins/FX/Niagara/Source/NiagaraEditor/Private/NiagaraNodeParameterMapBase.cpp
@@ -0,0 +1,202 @@
+// Copyright 1998-2016 Epic Games, Inc. All Rights Reserved.
+
+#include "NiagaraNodeParameterMapBase.h"
+#include "EdGraphSchema_Niagara.h"
+#include "NiagaraEditorUtilities.h"
+#include "SNiagaraGraphNodeConvert.h"
+#include "INiagaraCompiler.h"
+#include "NiagaraNodeOutput.h"
+#include "ScopedTransaction.h"
+#include "SNiagaraGraphPinAdd.h"
+#include "NiagaraGraph.h"
+#include "NiagaraComponent.h"
+#include "EdGraphSchema_Niagara.h"
+#include "NiagaraScript.h"
+#include "NiagaraScriptSource.h"
+#include "NiagaraConstants.h"
+
+#include "Framework/MultiBox/MultiBoxBuilder.h"
+
+#define LOCTEXT_NAMESPACE "NiagaraNodeParameterMapBase"
+
+UNiagaraNodeParameterMapBase::UNiagaraNodeParameterMapBase() : UNiagaraNodeWithDynamicPins()
+{
+
+}
+
+TArray<FNiagaraParameterMapHistory> UNiagaraNodeParameterMapBase::GetParameterMaps(UNiagaraScriptSourceBase* InSource, FString EmitterNameOverride)
+{
+	TArray<FNiagaraParameterMapHistory> OutputParameterMapHistories;
+	UNiagaraScriptSource* Base = Cast<UNiagaraScriptSource>(InSource);
+	if (Base != nullptr)
+	{
+		OutputParameterMapHistories = GetParameterMaps(Base->NodeGraph, EmitterNameOverride);
+	}
+	return OutputParameterMapHistories;
+
+}
+
+TArray<FNiagaraParameterMapHistory> UNiagaraNodeParameterMapBase::GetParameterMaps(UNiagaraGraph* InGraph, FString EmitterNameOverride)
+{
+	TArray<UNiagaraNodeOutput*> OutputNodes;
+	InGraph->FindOutputNodes(OutputNodes);
+	TArray<FNiagaraParameterMapHistory> OutputParameterMapHistories;
+
+	for (UNiagaraNodeOutput* FoundOutputNode : OutputNodes)
+	{
+		OutputParameterMapHistories.Append(GetParameterMaps(FoundOutputNode, false, EmitterNameOverride));
+	}
+
+	return OutputParameterMapHistories;
+}
+
+TArray<FNiagaraParameterMapHistory> UNiagaraNodeParameterMapBase::GetParameterMaps(UNiagaraNodeOutput* InGraphEnd, bool bLimitToOutputScriptType, FString EmitterNameOverride)
+{
+	const UEdGraphSchema_Niagara* Schema = Cast<UEdGraphSchema_Niagara>(InGraphEnd->GetSchema());
+	
+	FNiagaraParameterMapHistoryBuilder Builder;
+	if (!EmitterNameOverride.IsEmpty())
+	{
+		Builder.EnterEmitter(EmitterNameOverride, nullptr);
+	}
+
+	if (bLimitToOutputScriptType)
+	{
+		Builder.EnableScriptWhitelist(true, InGraphEnd->GetUsage());
+	}
+	
+	Builder.BuildParameterMaps(InGraphEnd);
+	
+	if (!EmitterNameOverride.IsEmpty())
+	{
+		Builder.ExitEmitter(EmitterNameOverride, nullptr);
+	}
+	
+	return Builder.Histories;
+}
+
+
+bool UNiagaraNodeParameterMapBase::AllowNiagaraTypeForAddPin(const FNiagaraTypeDefinition& InType)
+{
+	return  InType != FNiagaraTypeDefinition::GetGenericNumericDef();
+}
+
+TSharedRef<SWidget> UNiagaraNodeParameterMapBase::GenerateAddPinMenu(const FString& InWorkingPinName, SNiagaraGraphPinAdd* InPin)
+{
+	UNiagaraGraph* Graph = GetNiagaraGraph();
+	bool IsModule = Graph->FindOutputNode(ENiagaraScriptUsage::Module) != nullptr || Graph->FindOutputNode(ENiagaraScriptUsage::DynamicInput) != nullptr 
+		|| Graph->FindOutputNode(ENiagaraScriptUsage::Function) != nullptr;
+
+	bool bSupportsAttributes = true;
+	UNiagaraScriptSource* Source = Cast<UNiagaraScriptSource>(Graph->GetOuter());
+	if (Source && IsModule)
+	{
+		UNiagaraScript* Script = Cast<UNiagaraScript>(Source->GetOuter());
+		if (Script)
+		{
+			TArray<ENiagaraScriptUsage> Usages = Script->GetSupportedUsageContexts();
+			if (Usages.Contains(ENiagaraScriptUsage::ParticleEventScript) || Usages.Contains(ENiagaraScriptUsage::ParticleSpawnScript) || Usages.Contains(ENiagaraScriptUsage::ParticleUpdateScript))
+			{
+				bSupportsAttributes = true;
+			}
+			else
+			{
+				bSupportsAttributes = false;
+			}
+		}
+	}
+
+	FString WorkingCustomName;
+	FMenuBuilder MenuBuilder(true, nullptr);
+	MenuBuilder.AddSubMenu(LOCTEXT("CommonEngine", "Common Engine-Provided Variables"),
+		LOCTEXT("CommonSystemTooltip", "Create an entry using one of the common engine variables."),
+		FNewMenuDelegate::CreateLambda([InWorkingPinName, InPin, this](FMenuBuilder& SubMenuBuilder) { BuildEngineMenu(SubMenuBuilder, InWorkingPinName, InPin); }));
+
+	MenuBuilder.AddSubMenu(LOCTEXT("LocalVars", "Existing Graph Variables"),
+		LOCTEXT("LocalVarsTooltip", "Create an entry using existing graph variables."),
+		FNewMenuDelegate::CreateLambda([InWorkingPinName, InPin, this](FMenuBuilder& SubMenuBuilder) { BuildLocalMenu(SubMenuBuilder, InWorkingPinName, InPin); }));
+
+	if (bSupportsAttributes)
+	{
+		MenuBuilder.AddSubMenu(LOCTEXT("CommonAttributes", "Common Attributes"),
+			LOCTEXT("CommonAttributesTooltip", "Create an entry using one of the common attributes."),
+			FNewMenuDelegate::CreateLambda([InWorkingPinName, InPin, this](FMenuBuilder& SubMenuBuilder) { BuildCommonMenu(SubMenuBuilder, InWorkingPinName, InPin); }));
+
+		WorkingCustomName = TEXT("Particles.") + InWorkingPinName;
+		MenuBuilder.AddSubMenu(LOCTEXT("DefineAttribute", "Define Attribute"),
+			LOCTEXT("SupportedTypesTooltip", "Create an entry in the particles namespace that you will name immediately after."),
+			FNewMenuDelegate::CreateLambda([WorkingCustomName, InPin, this](FMenuBuilder& SubMenuBuilder) { BuildTypeMenu(SubMenuBuilder, WorkingCustomName, InPin); }));
+	}
+
+	WorkingCustomName = TEXT("Module.") + InWorkingPinName;
+	MenuBuilder.AddSubMenu(LOCTEXT("DefineModuleLocal", "Define Module Local"),
+		LOCTEXT("SupportedTypesTooltip", "Create an entry in the module namespace that you will name immediately after."),
+		FNewMenuDelegate::CreateLambda([WorkingCustomName, InPin, this](FMenuBuilder& SubMenuBuilder) { BuildTypeMenu(SubMenuBuilder, WorkingCustomName, InPin); }));
+
+	return MenuBuilder.MakeWidget();
+}
+
+void UNiagaraNodeParameterMapBase::BuildCommonMenu(FMenuBuilder& InMenuBuilder, const FString& InWorkingName, SNiagaraGraphPinAdd* InPin)
+{
+	TArray<FNiagaraVariable> Variables = FNiagaraConstants::GetCommonParticleAttributes();
+
+	for (FNiagaraVariable& NamespacedVar : Variables)
+	{
+		FText Desc = FNiagaraConstants::GetAttributeDescription(NamespacedVar);
+		FString NamespacedName = NamespacedVar.GetName().ToString();
+		InMenuBuilder.AddMenuEntry(
+			FText::FromString(NamespacedName),
+			FText::Format(LOCTEXT("AddButtonTypeEntryToolTipFormatCommon", "Add a reference to {0}. {1}"), FText::FromName(NamespacedVar.GetName()), Desc),
+			FSlateIcon(),
+			FUIAction(FExecuteAction::CreateSP(InPin, &SNiagaraGraphPinAdd::OnAddType, NamespacedVar)));
+	}
+}
+
+
+void UNiagaraNodeParameterMapBase::BuildLocalMenu(FMenuBuilder& InMenuBuilder, const FString& InWorkingName, SNiagaraGraphPinAdd* InPin)
+{
+	UNiagaraGraph* Graph = GetNiagaraGraph();
+	const TArray<FNiagaraParameterMapHistory> Histories = GetParameterMaps(Graph);
+
+	TArray<FNiagaraVariable> Variables;
+	for (const FNiagaraParameterMapHistory& History : Histories)
+	{
+		for (const FNiagaraVariable& Var : History.Variables)
+		{
+			Variables.AddUnique(Var);
+		}
+	}
+
+	Variables.Sort([](const FNiagaraVariable& A, const FNiagaraVariable& B) { return (A.GetName() < B.GetName()); });
+
+	for (FNiagaraVariable& NamespacedVar : Variables)
+	{
+		FString NamespacedName = NamespacedVar.GetName().ToString();
+		InMenuBuilder.AddMenuEntry(
+			FText::FromString(NamespacedName),
+			FText::Format(LOCTEXT("AddButtonTypeEntryToolTipFormatSystem", "Add a reference to {0}"), FText::FromName(NamespacedVar.GetName())),
+			FSlateIcon(),
+			FUIAction(FExecuteAction::CreateSP(InPin, &SNiagaraGraphPinAdd::OnAddType, NamespacedVar)));
+	}
+}
+
+void UNiagaraNodeParameterMapBase::BuildEngineMenu(FMenuBuilder& InMenuBuilder, const FString& InWorkingName, SNiagaraGraphPinAdd* InPin)
+{
+	TArray<FNiagaraVariable> Variables = FNiagaraConstants::GetEngineConstants();
+	Variables.Sort([](const FNiagaraVariable& A, const FNiagaraVariable& B) { return (A.GetName() < B.GetName()); });
+
+	for (FNiagaraVariable& Var : Variables)
+	{
+		FText Desc = FNiagaraConstants::GetEngineConstantDescription(Var);
+		FNiagaraVariable NamespacedVar = Var;
+		FString NamespacedName = NamespacedVar.GetName().ToString();
+		InMenuBuilder.AddMenuEntry(
+			FText::FromString(NamespacedName),
+			FText::Format(LOCTEXT("AddButtonTypeEntryToolTipFormatSystem", "Add a reference to {0}. {1}"), FText::FromName(Var.GetName()), Desc),
+			FSlateIcon(),
+			FUIAction(FExecuteAction::CreateSP(InPin, &SNiagaraGraphPinAdd::OnAddType, NamespacedVar)));
+	}
+}
+
+
+#undef LOCTEXT_NAMESPACE
\ No newline at end of file
diff --git a/Engine/Plugins/FX/Niagara/Source/NiagaraEditor/Private/NiagaraNodeParameterMapBase.h b/Engine/Plugins/FX/Niagara/Source/NiagaraEditor/Private/NiagaraNodeParameterMapBase.h
new file mode 100644
--- /dev/null
+++ b/Engine/Plugins/FX/Niagara/Source/NiagaraEditor/Private/NiagaraNodeParameterMapBase.h
@@ -0,0 +1,35 @@
+// Copyright 1998-2016 Epic Games, Inc. All Rights Reserved.
+
+#pragma once
+
+#include "NiagaraNodeWithDynamicPins.h"
+#include "SGraphPin.h"
+#include "NiagaraNodeParameterMapBase.generated.h"
+
+class UEdGraphPin;
+
+
+/** A node which allows the user to build a set of arbitrary output types from an arbitrary set of input types by connecting their inner components. */
+UCLASS()
+class UNiagaraNodeParameterMapBase : public UNiagaraNodeWithDynamicPins
+{
+public:
+	GENERATED_BODY()
+	UNiagaraNodeParameterMapBase();
+
+	/** Traverse the graph looking for the history of the parameter map specified by the input pin. This will return the list of variables discovered, any per-variable warnings (type mismatches, etc)
+		encountered per variable, and an array of pins encountered in order of traversal outward from the input pin.
+	*/
+	static TArray<FNiagaraParameterMapHistory> GetParameterMaps(UNiagaraNodeOutput* InGraphEnd, bool bLimitToOutputScriptType = false, FString EmitterNameOverride = TEXT(""));
+	static TArray<FNiagaraParameterMapHistory> GetParameterMaps(UNiagaraGraph* InGraph, FString EmitterNameOverride = TEXT(""));
+	static TArray<FNiagaraParameterMapHistory> GetParameterMaps(class UNiagaraScriptSourceBase* InSource, FString EmitterNameOverride = TEXT(""));
+
+	virtual bool AllowNiagaraTypeForAddPin(const FNiagaraTypeDefinition& InType) override;
+	
+	virtual TSharedRef<SWidget> GenerateAddPinMenu(const FString& InWorkingPinName, SNiagaraGraphPinAdd* InPin) override;
+
+protected:
+	virtual void BuildCommonMenu(FMenuBuilder& InMenuBuilder, const FString& InWorkingName, SNiagaraGraphPinAdd* InPin);
+	virtual void BuildLocalMenu(FMenuBuilder& InMenuBuilder, const FString& InWorkingName, SNiagaraGraphPinAdd* InPin);
+	virtual void BuildEngineMenu(FMenuBuilder& InMenuBuilder, const FString& InWorkingName, SNiagaraGraphPinAdd* InPin);
+};
diff --git a/Engine/Plugins/FX/Niagara/Source/NiagaraEditor/Private/NiagaraNodeParameterMapGet.cpp b/Engine/Plugins/FX/Niagara/Source/NiagaraEditor/Private/NiagaraNodeParameterMapGet.cpp
new file mode 100644
--- /dev/null
+++ b/Engine/Plugins/FX/Niagara/Source/NiagaraEditor/Private/NiagaraNodeParameterMapGet.cpp
@@ -0,0 +1,380 @@
+// Copyright 1998-2016 Epic Games, Inc. All Rights Reserved.
+
+#include "NiagaraNodeParameterMapGet.h"
+#include "EdGraphSchema_Niagara.h"
+#include "NiagaraEditorUtilities.h"
+#include "SNiagaraGraphNodeConvert.h"
+#include "NiagaraGraph.h"
+#include "NiagaraHlslTranslator.h"
+#include "ScopedTransaction.h"
+#include "SNiagaraGraphParameterMapGetNode.h"
+#include "NiagaraEditorModule.h"
+#include "INiagaraEditorTypeUtilities.h"
+#include "ModuleManager.h"
+
+#define LOCTEXT_NAMESPACE "NiagaraNodeParameterMapGet"
+
+UNiagaraNodeParameterMapGet::UNiagaraNodeParameterMapGet() : UNiagaraNodeParameterMapBase(), PinPendingRename(nullptr)
+{
+
+}
+
+
+void UNiagaraNodeParameterMapGet::AllocateDefaultPins()
+{
+	PinPendingRename = nullptr;
+	const UEdGraphSchema_Niagara* Schema = GetDefault<UEdGraphSchema_Niagara>();
+	CreatePin(EGPD_Input, Schema->TypeDefinitionToPinType(FNiagaraTypeDefinition::GetParameterMapDef()), TEXT("Source"));
+	CreateAddPin(EGPD_Output);
+}
+
+
+TSharedPtr<SGraphNode> UNiagaraNodeParameterMapGet::CreateVisualWidget()
+{
+	return SNew(SNiagaraGraphParameterMapGetNode, this);
+}
+
+bool UNiagaraNodeParameterMapGet::IsPinNameEditable(const UEdGraphPin* GraphPinObj) const
+{
+	const UEdGraphSchema_Niagara* Schema = GetDefault<UEdGraphSchema_Niagara>();
+	FNiagaraTypeDefinition TypeDef = Schema->PinToTypeDefinition(GraphPinObj);
+	if (TypeDef.IsValid() && GraphPinObj && GraphPinObj->Direction == EGPD_Output && CanRenamePin(GraphPinObj))
+	{
+		return true;
+	}
+	else
+	{
+		return false;
+	}
+}
+
+bool UNiagaraNodeParameterMapGet::IsPinNameEditableUponCreation(const UEdGraphPin* GraphPinObj) const
+{
+	if (GraphPinObj == PinPendingRename && GraphPinObj->Direction == EEdGraphPinDirection::EGPD_Output)
+	{
+		return true;
+	}
+	else
+	{
+		return false;
+	}
+}
+
+
+bool UNiagaraNodeParameterMapGet::VerifyEditablePinName(const FText& InName, FText& OutErrorMessage, const UEdGraphPin* InGraphPinObj) const
+{
+	if (InName.IsEmptyOrWhitespace() && InGraphPinObj->Direction == EEdGraphPinDirection::EGPD_Output)
+	{
+		OutErrorMessage = LOCTEXT("InvalidName", "Invalid pin name");
+		return false;
+	}
+	return true;
+}
+
+UEdGraphPin* UNiagaraNodeParameterMapGet::CreateDefaultPin(UEdGraphPin* OutputPin)
+{
+	if (OutputPin == nullptr)
+	{
+		return nullptr;
+	}
+
+	UEdGraphPin* DefaultPin = CreatePin(EEdGraphPinDirection::EGPD_Input, OutputPin->PinType, TEXT(""));
+	const UEdGraphSchema_Niagara* Schema = GetDefault<UEdGraphSchema_Niagara>();
+	FNiagaraVariable Var = Schema->PinToNiagaraVariable(OutputPin);
+	FNiagaraEditorUtilities::ResetVariableToDefaultValue(Var);
+
+	FNiagaraEditorModule& NiagaraEditorModule = FModuleManager::GetModuleChecked<FNiagaraEditorModule>("NiagaraEditor");
+	if (Var.IsDataAllocated())
+	{
+		TSharedPtr<INiagaraEditorTypeUtilities> TypeEditorUtilities = NiagaraEditorModule.GetTypeUtilities(Var.GetType());
+
+		if (TypeEditorUtilities.IsValid() && TypeEditorUtilities->CanHandlePinDefaults())
+		{
+			FString Value = TypeEditorUtilities->GetPinDefaultStringFromValue(Var);
+			DefaultPin->DefaultValue = Value;
+		}
+	}
+	
+	if (!OutputPin->PersistentGuid.IsValid())
+	{
+		OutputPin->PersistentGuid = FGuid::NewGuid();
+	}
+	if (!DefaultPin->PersistentGuid.IsValid())
+	{
+		DefaultPin->PersistentGuid = FGuid::NewGuid();
+	}
+	PinOutputToPinDefaultPersistentId.Add(OutputPin->PersistentGuid, DefaultPin->PersistentGuid);
+
+	SynchronizeDefaultInputPin(DefaultPin, OutputPin);
+	return DefaultPin;
+}
+
+void UNiagaraNodeParameterMapGet::OnPinRenamed(UEdGraphPin* RenamedPin)
+{
+	RenamedPin->PinFriendlyName = FText::FromString(RenamedPin->PinName);
+}
+
+
+void UNiagaraNodeParameterMapGet::OnNewTypedPinAdded(UEdGraphPin* NewPin)
+{
+	if (NewPin->Direction == EEdGraphPinDirection::EGPD_Output)
+	{
+		const UEdGraphSchema_Niagara* Schema = GetDefault<UEdGraphSchema_Niagara>();
+		FNiagaraTypeDefinition TypeDef = Schema->PinToTypeDefinition(NewPin);
+
+		UEdGraphPin* MatchingDefault = GetDefaultPin(NewPin);
+		if (MatchingDefault == nullptr)
+		{
+			UEdGraphPin* DefaultPin = CreateDefaultPin(NewPin);
+		}
+	}
+
+	if (HasAnyFlags(RF_NeedLoad | RF_NeedPostLoad | RF_NeedInitialization))
+	{
+		return;
+	}
+
+	if (NewPin->Direction == EEdGraphPinDirection::EGPD_Output)
+	{
+		PinPendingRename = NewPin;
+	}
+
+}
+
+void UNiagaraNodeParameterMapGet::RemoveDynamicPin(UEdGraphPin* Pin)
+{
+	FScopedTransaction RemovePinTransaction(LOCTEXT("RemovePinTransaction", "Remove pin"));
+
+	UEdGraphPin* DefaultPin = nullptr;
+	if (Pin->Direction == EEdGraphPinDirection::EGPD_Output)
+	{
+		DefaultPin = GetDefaultPin(Pin);
+	}
+
+	RemovePin(Pin);
+
+	if (DefaultPin != nullptr)
+	{
+		RemovePin(DefaultPin);
+	}
+
+	GetGraph()->NotifyGraphChanged();
+}
+
+UEdGraphPin* UNiagaraNodeParameterMapGet::GetDefaultPin(UEdGraphPin* OutputPin)
+{
+	TArray<UEdGraphPin*> InputPins;
+	GetInputPins(InputPins);
+
+	const FGuid* InputGuid = PinOutputToPinDefaultPersistentId.Find(OutputPin->PersistentGuid);
+
+	if (InputGuid != nullptr)
+	{
+		for (UEdGraphPin* InputPin : InputPins)
+		{
+			if ((*InputGuid) == InputPin->PersistentGuid)
+			{
+				return InputPin;
+			}
+		}
+	}
+
+	return nullptr;
+}
+
+
+UEdGraphPin* UNiagaraNodeParameterMapGet::GetOutputPinForDefault(UEdGraphPin* DefaultPin)
+{
+	TArray<UEdGraphPin*> OutputPins;
+	GetOutputPins(OutputPins);
+
+	const FGuid* OutputGuid = PinOutputToPinDefaultPersistentId.Find(DefaultPin->PersistentGuid);
+
+	if (OutputGuid != nullptr)
+	{
+		for (UEdGraphPin* OutputPin : OutputPins)
+		{
+			if ((*OutputGuid) == OutputPin->PersistentGuid)
+			{
+				return OutputPin;
+			}
+		}
+	}
+
+	return nullptr;
+}
+
+void UNiagaraNodeParameterMapGet::PostLoad()
+{
+	Super::PostLoad();
+	TArray<UEdGraphPin*> OutputPins;
+	GetOutputPins(OutputPins);
+	for (int32 i = 0; i < OutputPins.Num(); i++)
+	{
+		if (IsAddPin(OutputPins[i]))
+		{
+			continue;
+		}
+		UEdGraphPin* InputPin = GetDefaultPin(OutputPins[i]);
+		if (InputPin == nullptr)
+		{
+			CreateDefaultPin(OutputPins[i]);
+		}
+		else
+		{
+			SynchronizeDefaultInputPin(InputPin, OutputPins[i]);
+		}
+	}
+}
+
+
+void UNiagaraNodeParameterMapGet::SynchronizeDefaultInputPin(UEdGraphPin* DefaultPin, UEdGraphPin* OutputPin)
+{
+	const UEdGraphSchema_Niagara* Schema = GetDefault<UEdGraphSchema_Niagara>();
+	if (!DefaultPin)
+	{
+		return;
+	}
+
+	if (FNiagaraParameterMapHistory::IsEngineParameter(Schema->PinToNiagaraVariable(OutputPin)))
+	{
+		DefaultPin->bDefaultValueIsIgnored = true;
+		DefaultPin->bNotConnectable = true;
+		DefaultPin->bHidden = true;
+		DefaultPin->PinToolTip = FText::Format(LOCTEXT("DefaultValueTooltip", "Default value for {0}. Disabled for Engine Parameters."), FText::FromString(OutputPin->PinName)).ToString();
+	}
+	else
+	{
+		DefaultPin->bDefaultValueIsIgnored = false;
+		DefaultPin->bNotConnectable = false;
+		DefaultPin->bHidden = false;
+		DefaultPin->PinToolTip = FText::Format(LOCTEXT("DefaultValueTooltip", "Default value for {0} if no other module has set it previously in the stack."), FText::FromString(OutputPin->PinName)).ToString();
+	}
+}
+
+
+FText UNiagaraNodeParameterMapGet::GetNodeTitle(ENodeTitleType::Type TitleType) const
+{
+	return LOCTEXT("UNiagaraNodeParameterMapGetName", "Map Get");
+}
+
+void UNiagaraNodeParameterMapGet::BuildParameterMapHistory(FNiagaraParameterMapHistoryBuilder& OutHistory, bool bRecursive)
+{
+	if (bRecursive)
+	{
+		OutHistory.VisitInputPin(GetInputPin(0), this);
+	}
+
+	int32 ParamMapIdx = INDEX_NONE;
+	if (GetInputPin(0)->LinkedTo.Num() != 0)
+	{
+		ParamMapIdx = OutHistory.TraceParameterMapOutputPin(GetInputPin(0)->LinkedTo[0]);
+	}
+
+	if (ParamMapIdx != INDEX_NONE)
+	{
+		TArray<UEdGraphPin*> OutputPins;
+		GetOutputPins(OutputPins);
+		for (int32 i = 0; i < OutputPins.Num(); i++)
+		{
+			if (IsAddPin(OutputPins[i]))
+			{
+				continue;
+			}
+
+			bool bUsedDefaults = false;
+			if (bRecursive)
+			{
+				OutHistory.HandleVariableRead(ParamMapIdx, OutputPins[i], true, GetDefaultPin(OutputPins[i]), bUsedDefaults);
+			}
+			else
+			{
+				OutHistory.HandleVariableRead(ParamMapIdx, OutputPins[i], true, nullptr, bUsedDefaults);
+			}
+		}
+	}
+}
+
+void UNiagaraNodeParameterMapGet::Compile(class FHlslNiagaraTranslator* Translator, TArray<int32>& Outputs)
+{
+	const UEdGraphSchema_Niagara* Schema = CastChecked<UEdGraphSchema_Niagara>(GetSchema());
+
+	TArray<UEdGraphPin*> InputPins;
+	GetInputPins(InputPins);
+	TArray<UEdGraphPin*> OutputPins;
+	GetOutputPins(OutputPins);
+
+	// Initialize the outputs to invalid values.
+	check(Outputs.Num() == 0);
+	for (int32 i = 0; i < OutputPins.Num(); i++)
+	{
+		if (IsAddPin(OutputPins[i]))
+		{
+			continue;
+		}
+		Outputs.Add(INDEX_NONE);
+	}
+
+	// First compile fully down the hierarchy for our predecessors..
+	TArray<int32> CompileInputs;
+
+	for (int32 i = 0; i < InputPins.Num(); i++)
+	{
+		UEdGraphPin* InputPin = InputPins[i];
+
+		if (InputPin->PinType.PinCategory == UEdGraphSchema_Niagara::PinCategoryType || 
+			InputPin->PinType.PinCategory == UEdGraphSchema_Niagara::PinCategoryEnum)
+		{
+			int32 CompiledInput = INDEX_NONE;
+			if (i == 0) // Only the zeroth item is not an default value pin.
+			{
+				CompiledInput = Translator->CompilePin(InputPin);
+				if (CompiledInput == INDEX_NONE)
+				{
+					Translator->Error(LOCTEXT("InputError", "Error compiling input for param map get node."), this, InputPin);
+				}
+			}
+			CompileInputs.Add(CompiledInput);
+		}
+	}
+
+	UNiagaraGraph* Graph = Cast<UNiagaraGraph>(GetGraph());
+	// By this point, we've visited all of our child nodes in the call graph. We can mine them to find out everyone contributing to the parameter map (and when).
+	if (GetInputPin(0) != nullptr && GetInputPin(0)->LinkedTo.Num() > 0)
+	{
+		Translator->ParameterMapGet(this, CompileInputs, Outputs);
+	}
+}
+
+
+bool UNiagaraNodeParameterMapGet::CommitEditablePinName(const FText& InName, UEdGraphPin* InGraphPinObj) 
+{
+	if (Pins.Contains(InGraphPinObj) && InGraphPinObj->Direction == EEdGraphPinDirection::EGPD_Output)
+	{
+		FScopedTransaction AddNewPinTransaction(LOCTEXT("Rename Pin", "Renamed pin"));
+		Modify();
+		InGraphPinObj->Modify();
+		UEdGraphPin* DefaultPin = GetDefaultPin(InGraphPinObj);
+
+		InGraphPinObj->PinFriendlyName = InName;
+		InGraphPinObj->PinName = InName.ToString();
+
+		if (DefaultPin)
+		{
+			DefaultPin->Modify();
+			SynchronizeDefaultInputPin(DefaultPin, InGraphPinObj);
+		}
+
+
+		if (PinPendingRename == InGraphPinObj)
+		{
+			PinPendingRename = nullptr;
+		}	
+		GetGraph()->NotifyGraphChanged();
+
+		return true;
+	}
+	return false;
+}
+
+#undef LOCTEXT_NAMESPACE
diff --git a/Engine/Plugins/FX/Niagara/Source/NiagaraEditor/Private/NiagaraNodeParameterMapGet.h b/Engine/Plugins/FX/Niagara/Source/NiagaraEditor/Private/NiagaraNodeParameterMapGet.h
new file mode 100644
--- /dev/null
+++ b/Engine/Plugins/FX/Niagara/Source/NiagaraEditor/Private/NiagaraNodeParameterMapGet.h
@@ -0,0 +1,60 @@
+// Copyright 1998-2016 Epic Games, Inc. All Rights Reserved.
+
+#pragma once
+
+#include "SGraphPin.h"
+#include "NiagaraNodeParameterMapBase.h"
+#include "NiagaraNodeParameterMapGet.generated.h"
+
+class UEdGraphPin;
+
+
+/** A node that allows a user to get multiple values from a parameter map. */
+UCLASS()
+class UNiagaraNodeParameterMapGet : public UNiagaraNodeParameterMapBase
+{
+public:
+	GENERATED_BODY()
+
+	UNiagaraNodeParameterMapGet();
+
+	virtual void AllocateDefaultPins() override;
+	virtual TSharedPtr<SGraphNode> CreateVisualWidget() override;
+
+	virtual bool IsPinNameEditable(const UEdGraphPin* GraphPinObj) const override;
+	virtual bool IsPinNameEditableUponCreation(const UEdGraphPin* GraphPinObj) const override;
+	virtual bool VerifyEditablePinName(const FText& InName, FText& OutErrorMessage, const UEdGraphPin* InGraphPinObj) const override;
+	virtual bool CommitEditablePinName(const FText& InName, UEdGraphPin* InGraphPinObj)  override;
+
+	virtual void Compile(class FHlslNiagaraTranslator* Translator, TArray<int32>& Outputs);
+
+	virtual void BuildParameterMapHistory(FNiagaraParameterMapHistoryBuilder& OutHistory, bool bRecursive = true) override;
+
+	virtual FText GetNodeTitle(ENodeTitleType::Type TitleType) const;
+
+	/** Get the default value input pin for one of the output pins specified.*/
+	UEdGraphPin* GetDefaultPin(UEdGraphPin* OutputPin);
+
+	virtual void PostLoad() override;
+
+protected:
+	virtual void OnNewTypedPinAdded(UEdGraphPin* NewPin) override;
+	virtual void OnPinRenamed(UEdGraphPin* RenamedPin) override;
+
+	/** Synchronize the removal of the output pin with its default.*/
+	virtual void RemoveDynamicPin(UEdGraphPin* Pin) override;
+
+	/** Make sure that descriptions match up as well as any other value that can be changed.*/
+	void SynchronizeDefaultInputPin(UEdGraphPin* DefaultPin, UEdGraphPin* OutputPin);
+	
+	/** Reverse the lookup from GetDefaultPin.*/
+	UEdGraphPin* GetOutputPinForDefault(UEdGraphPin* DefaultPin);
+	
+	/** Properly set up the default input pin for an output pin.*/
+	UEdGraphPin* CreateDefaultPin(UEdGraphPin* OutputPin);
+
+	UEdGraphPin* PinPendingRename;
+
+	UPROPERTY()
+	TMap<FGuid, FGuid> PinOutputToPinDefaultPersistentId;
+};
diff --git a/Engine/Plugins/FX/Niagara/Source/NiagaraEditor/Private/NiagaraNodeParameterMapSet.cpp b/Engine/Plugins/FX/Niagara/Source/NiagaraEditor/Private/NiagaraNodeParameterMapSet.cpp
new file mode 100644
--- /dev/null
+++ b/Engine/Plugins/FX/Niagara/Source/NiagaraEditor/Private/NiagaraNodeParameterMapSet.cpp
@@ -0,0 +1,186 @@
+// Copyright 1998-2016 Epic Games, Inc. All Rights Reserved.
+
+#include "NiagaraNodeParameterMapSet.h"
+#include "EdGraphSchema_Niagara.h"
+#include "NiagaraEditorUtilities.h"
+#include "SNiagaraGraphNodeConvert.h"
+#include "NiagaraHlslTranslator.h"
+
+#include "ScopedTransaction.h"
+
+#define LOCTEXT_NAMESPACE "NiagaraNodeParameterMapSet"
+
+UNiagaraNodeParameterMapSet::UNiagaraNodeParameterMapSet() : UNiagaraNodeParameterMapBase(), PinPendingRename(nullptr)
+{
+
+}
+
+void UNiagaraNodeParameterMapSet::AllocateDefaultPins()
+{
+	PinPendingRename = nullptr;
+	const UEdGraphSchema_Niagara* Schema = GetDefault<UEdGraphSchema_Niagara>();
+	CreatePin(EGPD_Input, Schema->TypeDefinitionToPinType(FNiagaraTypeDefinition::GetParameterMapDef()), TEXT("Source"));
+	CreatePin(EGPD_Output, Schema->TypeDefinitionToPinType(FNiagaraTypeDefinition::GetParameterMapDef()), TEXT("Dest"));
+	CreateAddPin(EGPD_Input);
+}
+
+bool UNiagaraNodeParameterMapSet::IsPinNameEditable(const UEdGraphPin* GraphPinObj) const
+{
+	const UEdGraphSchema_Niagara* Schema = GetDefault<UEdGraphSchema_Niagara>();
+	FNiagaraTypeDefinition TypeDef = Schema->PinToTypeDefinition(GraphPinObj);
+	if (TypeDef.IsValid() && GraphPinObj && GraphPinObj->Direction == EGPD_Input && CanRenamePin(GraphPinObj))
+	{
+		return true;
+	}
+	else
+	{
+		return false;
+	}
+}
+
+bool UNiagaraNodeParameterMapSet::IsPinNameEditableUponCreation(const UEdGraphPin* GraphPinObj) const
+{
+	if (GraphPinObj == PinPendingRename)
+	{
+		return true;
+	}
+	else
+	{
+		return false;
+	}
+}
+
+
+bool UNiagaraNodeParameterMapSet::VerifyEditablePinName(const FText& InName, FText& OutErrorMessage, const UEdGraphPin* InGraphPinObj) const
+{
+	if (InName.IsEmptyOrWhitespace())
+	{
+		OutErrorMessage = LOCTEXT("InvalidName", "Invalid pin name");
+		return false;
+	}
+	return true;
+}
+
+void UNiagaraNodeParameterMapSet::OnNewTypedPinAdded(UEdGraphPin* NewPin)
+{
+	if (HasAnyFlags(RF_NeedLoad | RF_NeedPostLoad | RF_NeedInitialization))
+	{
+		return;
+	}
+
+	PinPendingRename = NewPin;
+}
+
+void UNiagaraNodeParameterMapSet::OnPinRenamed(UEdGraphPin* RenamedPin) 
+{
+	RenamedPin->PinFriendlyName = FText::FromString(RenamedPin->PinName);
+
+}
+
+bool UNiagaraNodeParameterMapSet::CommitEditablePinName(const FText& InName, UEdGraphPin* InGraphPinObj) 
+{
+	if (Pins.Contains(InGraphPinObj))
+	{
+		FScopedTransaction AddNewPinTransaction(LOCTEXT("Rename Pin", "Renamed pin"));
+		Modify();
+		InGraphPinObj->Modify();
+		InGraphPinObj->PinFriendlyName = InName;
+		InGraphPinObj->PinName = InName.ToString();
+
+		if (InGraphPinObj == PinPendingRename)
+		{
+			PinPendingRename = nullptr;
+		}
+		GetGraph()->NotifyGraphChanged();
+		return true;
+	}
+	return false;
+}
+
+void UNiagaraNodeParameterMapSet::Compile(class FHlslNiagaraTranslator* Translator, TArray<int32>& Outputs)
+{
+	TArray<UEdGraphPin*> InputPins;
+	GetInputPins(InputPins);
+
+	TArray<UEdGraphPin*> OutputPins;
+	GetOutputPins(OutputPins);
+
+	// Initialize the outputs to invalid values.
+	check(Outputs.Num() == 0);
+	for (int32 i = 0; i < OutputPins.Num(); i++)
+	{
+		Outputs.Add(INDEX_NONE);
+	}
+
+	const UEdGraphSchema_Niagara* Schema = CastChecked<UEdGraphSchema_Niagara>(GetSchema());
+
+	// First compile fully down the hierarchy for our predecessors..
+	TArray<int32> CompileInputs;
+	for (UEdGraphPin* InputPin : InputPins)
+	{
+		if (IsAddPin(InputPin))
+		{
+			continue;
+		}
+
+		int32 CompiledInput = Translator->CompilePin(InputPin);
+		if (CompiledInput == INDEX_NONE)
+		{
+			Translator->Error(LOCTEXT("InputError", "Error compiling input for get node."), this, InputPin);
+		}
+		CompileInputs.Add(CompiledInput);
+	}
+
+	if (GetInputPin(0) != nullptr && GetInputPin(0)->LinkedTo.Num() > 0)
+	{
+		Translator->ParameterMapSet(this, CompileInputs, Outputs);
+	}
+}
+
+FText UNiagaraNodeParameterMapSet::GetNodeTitle(ENodeTitleType::Type TitleType) const
+{
+	return LOCTEXT("UNiagaraNodeParameterMapSetName", "Map Set");
+}
+
+
+void UNiagaraNodeParameterMapSet::BuildParameterMapHistory(FNiagaraParameterMapHistoryBuilder& OutHistory, bool bRecursive)
+{
+	const UEdGraphSchema_Niagara* Schema = GetDefault<UEdGraphSchema_Niagara>();
+	TArray<UEdGraphPin*> InputPins;
+	GetInputPins(InputPins);
+	
+	int32 ParamMapIdx = INDEX_NONE;
+	for (int32 i = 0; i < InputPins.Num(); i++)
+	{
+		if (IsAddPin(InputPins[i]))
+		{
+			continue;
+		}
+
+		OutHistory.VisitInputPin(InputPins[i], this);
+
+		FNiagaraTypeDefinition VarTypeDef = Schema->PinToTypeDefinition(InputPins[i]);
+		if (i == 0 && InputPins[i] != nullptr && VarTypeDef == FNiagaraTypeDefinition::GetParameterMapDef())
+		{
+			UEdGraphPin* PriorParamPin = nullptr;
+			if (InputPins[i]->LinkedTo.Num() > 0)
+			{
+				PriorParamPin = InputPins[i]->LinkedTo[0];
+			}
+
+			// Now plow into our ancestor node
+			if (PriorParamPin)
+			{
+				ParamMapIdx = OutHistory.TraceParameterMapOutputPin(PriorParamPin);
+			}
+		}
+		else if (i > 0 && InputPins[i] != nullptr && ParamMapIdx != INDEX_NONE)
+		{
+			OutHistory.HandleVariableWrite(ParamMapIdx, InputPins[i]);
+		}
+	}
+
+	OutHistory.RegisterParameterMapPin(ParamMapIdx, GetOutputPin(0));
+}
+
+#undef LOCTEXT_NAMESPACE
diff --git a/Engine/Plugins/FX/Niagara/Source/NiagaraEditor/Private/NiagaraNodeParameterMapSet.h b/Engine/Plugins/FX/Niagara/Source/NiagaraEditor/Private/NiagaraNodeParameterMapSet.h
new file mode 100644
--- /dev/null
+++ b/Engine/Plugins/FX/Niagara/Source/NiagaraEditor/Private/NiagaraNodeParameterMapSet.h
@@ -0,0 +1,38 @@
+// Copyright 1998-2016 Epic Games, Inc. All Rights Reserved.
+
+#pragma once
+
+#include "SGraphPin.h"
+#include "NiagaraNodeParameterMapBase.h"
+#include "NiagaraNodeParameterMapSet.generated.h"
+
+class UEdGraphPin;
+
+
+/** A node that allows a user to set multiple values into a parameter map. */
+UCLASS()
+class UNiagaraNodeParameterMapSet : public UNiagaraNodeParameterMapBase
+{
+public:
+	GENERATED_BODY()
+
+	UNiagaraNodeParameterMapSet();
+
+	virtual void AllocateDefaultPins() override;
+
+	virtual bool IsPinNameEditable(const UEdGraphPin* GraphPinObj) const override;
+	virtual bool IsPinNameEditableUponCreation(const UEdGraphPin* GraphPinObj) const override;
+	virtual bool VerifyEditablePinName(const FText& InName, FText& OutErrorMessage, const UEdGraphPin* InGraphPinObj) const override;
+	virtual bool CommitEditablePinName(const FText& InName, UEdGraphPin* InGraphPinObj)  override;
+
+	virtual void BuildParameterMapHistory(FNiagaraParameterMapHistoryBuilder& OutHistory, bool bRecursive = true) override;
+
+	virtual void Compile(class FHlslNiagaraTranslator* Translator, TArray<int32>& Outputs);
+
+	virtual FText GetNodeTitle(ENodeTitleType::Type TitleType) const;
+protected:
+	virtual void OnNewTypedPinAdded(UEdGraphPin* NewPin); 
+	virtual void OnPinRenamed(UEdGraphPin* RenamedPin) override;
+	UEdGraphPin* PinPendingRename;
+
+};
diff --git a/Engine/Plugins/FX/Niagara/Source/NiagaraEditor/Private/NiagaraNodeWithDynamicPins.cpp b/Engine/Plugins/FX/Niagara/Source/NiagaraEditor/Private/NiagaraNodeWithDynamicPins.cpp
new file mode 100644
--- /dev/null
+++ b/Engine/Plugins/FX/Niagara/Source/NiagaraEditor/Private/NiagaraNodeWithDynamicPins.cpp
@@ -0,0 +1,276 @@
+// Copyright 1998-2016 Epic Games, Inc. All Rights Reserved.
+
+
+#include "NiagaraNodeWithDynamicPins.h"
+#include "EdGraphSchema_Niagara.h"
+#include "NiagaraGraph.h"
+#include "UIAction.h"
+#include "ScopedTransaction.h"
+#include "MultiBoxBuilder.h"
+#include "SEditableTextBox.h"
+#include "SBox.h"
+#include "SNiagaraGraphPinAdd.h"
+#include "INiagaraEditorTypeUtilities.h"
+#include "NiagaraEditorUtilities.h"
+
+#define LOCTEXT_NAMESPACE "NiagaraNodeWithDynamicPins"
+
+const FString UNiagaraNodeWithDynamicPins::AddPinSubCategory("DynamicAddPin");
+
+void UNiagaraNodeWithDynamicPins::PinConnectionListChanged(UEdGraphPin* Pin)
+{
+	Super::PinConnectionListChanged(Pin);
+
+	// Check if an add pin was connected and convert it to a typed connection.
+	const UEdGraphSchema_Niagara* Schema = GetDefault<UEdGraphSchema_Niagara>();
+	if (Pin->PinType.PinCategory == UEdGraphSchema_Niagara::PinCategoryMisc && Pin->PinType.PinSubCategory == AddPinSubCategory && Pin->LinkedTo.Num() > 0)
+	{
+		FNiagaraTypeDefinition LinkedPinType = Schema->PinToTypeDefinition(Pin->LinkedTo[0]);
+		Pin->PinType = Schema->TypeDefinitionToPinType(LinkedPinType);
+		Pin->PinName = Pin->LinkedTo[0]->GetName();
+
+		CreateAddPin(Pin->Direction);
+		OnNewTypedPinAdded(Pin);
+		GetGraph()->NotifyGraphChanged();
+	}
+}
+
+UEdGraphPin* GetAddPin(TArray<UEdGraphPin*> Pins, EEdGraphPinDirection Direction)
+{
+	for (UEdGraphPin* Pin : Pins)
+	{
+		if (Pin->Direction == Direction &&
+			Pin->PinType.PinCategory == UEdGraphSchema_Niagara::PinCategoryMisc && 
+			Pin->PinType.PinSubCategory == UNiagaraNodeWithDynamicPins::AddPinSubCategory)
+		{
+			return Pin;
+		}
+	}
+	return nullptr;
+}
+
+bool UNiagaraNodeWithDynamicPins::AllowNiagaraTypeForAddPin(const FNiagaraTypeDefinition& InType)
+{
+	return  InType != FNiagaraTypeDefinition::GetGenericNumericDef() && InType.GetScriptStruct() != nullptr;
+}
+
+UEdGraphPin* UNiagaraNodeWithDynamicPins::RequestNewTypedPin(EEdGraphPinDirection Direction, const FNiagaraTypeDefinition& Type)
+{
+	FString DefaultName;
+	if (Direction == EGPD_Input)
+	{
+		TArray<UEdGraphPin*> InPins;
+		GetInputPins(InPins);
+		DefaultName = TEXT("Input ") + LexicalConversion::ToString(InPins.Num());
+	}
+	else
+	{
+		TArray<UEdGraphPin*> OutPins;
+		GetOutputPins(OutPins);
+		DefaultName = TEXT("Output ") + LexicalConversion::ToString(OutPins.Num());
+	}
+	return RequestNewTypedPin(Direction, Type, DefaultName);
+}
+
+UEdGraphPin* UNiagaraNodeWithDynamicPins::RequestNewTypedPin(EEdGraphPinDirection Direction, const FNiagaraTypeDefinition& Type, FString InName)
+{
+	const UEdGraphSchema_Niagara* Schema = GetDefault<UEdGraphSchema_Niagara>();
+	UEdGraphPin* AddPin = GetAddPin(GetAllPins(), Direction);
+	checkf(AddPin != nullptr, TEXT("Add pin is missing"));
+	AddPin->Modify();
+	AddPin->PinType = Schema->TypeDefinitionToPinType(Type);
+	AddPin->PinName = InName;
+
+	CreateAddPin(Direction);
+	OnNewTypedPinAdded(AddPin);
+	UNiagaraGraph* Graph = GetNiagaraGraph();
+	Graph->NotifyGraphNeedsRecompile();
+
+	return AddPin;
+}
+
+void UNiagaraNodeWithDynamicPins::CreateAddPin(EEdGraphPinDirection Direction)
+{
+	CreatePin(Direction, FEdGraphPinType(UEdGraphSchema_Niagara::PinCategoryMisc, AddPinSubCategory, nullptr, EPinContainerType::None, false, FEdGraphTerminalType()), TEXT("Add"));
+}
+
+bool UNiagaraNodeWithDynamicPins::IsAddPin(const UEdGraphPin* Pin) const
+{
+	return Pin->PinType.PinCategory == UEdGraphSchema_Niagara::PinCategoryMisc && 
+		Pin->PinType.PinSubCategory == UNiagaraNodeWithDynamicPins::AddPinSubCategory;
+}
+
+bool UNiagaraNodeWithDynamicPins::CanRenamePin(const UEdGraphPin* Pin) const
+{
+	return IsAddPin(Pin) == false;
+}
+
+bool UNiagaraNodeWithDynamicPins::CanRemovePin(const UEdGraphPin* Pin) const
+{
+	return IsAddPin(Pin) == false;
+}
+
+bool UNiagaraNodeWithDynamicPins::CanMovePin(const UEdGraphPin* Pin) const
+{
+	return IsAddPin(Pin) == false;
+}
+
+void UNiagaraNodeWithDynamicPins::MoveDynamicPin(UEdGraphPin* Pin, int32 DirectionToMove)
+{
+	TArray<UEdGraphPin*> SameDirectionPins;
+	if (Pin->Direction == EEdGraphPinDirection::EGPD_Input)
+	{
+		GetInputPins(SameDirectionPins);
+	}
+	else
+	{
+		GetOutputPins(SameDirectionPins);
+	}
+
+	for (int32 i = 0; i < SameDirectionPins.Num(); i++)
+	{
+		if (SameDirectionPins[i] == Pin)
+		{
+			if (i + DirectionToMove >= 0 && i + DirectionToMove < SameDirectionPins.Num())
+			{
+				FScopedTransaction MovePinTransaction(LOCTEXT("MovePinTransaction", "Moved pin"));
+				Modify();
+				UEdGraphPin* PinOld = SameDirectionPins[i + DirectionToMove];
+				if (PinOld)
+					PinOld->Modify();
+				Pin->Modify();
+
+				int32 RealPinIdx = INDEX_NONE;
+				int32 SwapRealPinIdx = INDEX_NONE;
+				Pins.Find(Pin, RealPinIdx);
+				Pins.Find(PinOld, SwapRealPinIdx);
+				
+				Pins[SwapRealPinIdx] = Pin;
+				Pins[RealPinIdx] = PinOld;
+				GetGraph()->NotifyGraphChanged();
+				break;
+			}
+		}
+	}
+}
+
+
+void UNiagaraNodeWithDynamicPins::GetContextMenuActions(const FGraphNodeContextMenuBuilder& Context) const
+{
+	Super::GetContextMenuActions(Context);
+	if (Context.Pin != nullptr)
+	{
+		Context.MenuBuilder->BeginSection("EdGraphSchema_NiagaraPinActions", LOCTEXT("EditPinMenuHeader", "Edit Pin"));
+		if (CanRenamePin(Context.Pin))
+		{
+			UEdGraphPin* Pin = const_cast<UEdGraphPin*>(Context.Pin);
+			TSharedRef<SWidget> RenameWidget =
+				SNew(SBox)
+				.WidthOverride(100)
+				.Padding(FMargin(5, 0, 0, 0))
+				[
+					SNew(SEditableTextBox)
+					.Text_UObject(this, &UNiagaraNodeWithDynamicPins::GetPinNameText, Pin)
+					.OnTextCommitted_UObject(this, &UNiagaraNodeWithDynamicPins::PinNameTextCommitted, Pin)
+				];
+			Context.MenuBuilder->AddWidget(RenameWidget, LOCTEXT("NameMenuItem", "Name"));
+		}
+		if (CanRemovePin(Context.Pin))
+		{
+			Context.MenuBuilder->AddMenuEntry(
+				LOCTEXT("RemoveDynamicPin", "Remove pin"),
+				LOCTEXT("RemoveDynamicPinToolTip", "Remove this pin and any connections."),
+				FSlateIcon(),
+				FUIAction(FExecuteAction::CreateUObject(const_cast<UNiagaraNodeWithDynamicPins*>(this), &UNiagaraNodeWithDynamicPins::RemoveDynamicPin, const_cast<UEdGraphPin*>(Context.Pin))));
+		}
+		if (CanMovePin(Context.Pin))
+		{
+			TArray<UEdGraphPin*> SameDirectionPins;
+			if (Context.Pin->Direction == EEdGraphPinDirection::EGPD_Input)
+			{
+				GetInputPins(SameDirectionPins);
+			}
+			else
+			{
+				GetOutputPins(SameDirectionPins);
+			}
+			int32 PinIdx = INDEX_NONE;
+			SameDirectionPins.Find(const_cast<UEdGraphPin*>(Context.Pin), PinIdx);
+
+			if (PinIdx != 0)
+			{
+				Context.MenuBuilder->AddMenuEntry(
+					LOCTEXT("MoveDynamicPinUp", "Move pin up"),
+					LOCTEXT("MoveDynamicPinToolTipUp", "Move this pin and any connections one slot up."),
+					FSlateIcon(),
+					FUIAction(FExecuteAction::CreateUObject(const_cast<UNiagaraNodeWithDynamicPins*>(this), &UNiagaraNodeWithDynamicPins::MoveDynamicPin, const_cast<UEdGraphPin*>(Context.Pin), -1)));
+			}
+			if (PinIdx >= 0 && PinIdx < SameDirectionPins.Num() - 1)
+			{
+				Context.MenuBuilder->AddMenuEntry(
+					LOCTEXT("MoveDynamicPinDown", "Move pin down"),
+					LOCTEXT("MoveDynamicPinToolTipDown", "Move this pin and any connections one slot down."),
+					FSlateIcon(),
+					FUIAction(FExecuteAction::CreateUObject(const_cast<UNiagaraNodeWithDynamicPins*>(this), &UNiagaraNodeWithDynamicPins::MoveDynamicPin, const_cast<UEdGraphPin*>(Context.Pin), 1)));
+			}
+		}
+	}
+}
+
+TSharedRef<SWidget> UNiagaraNodeWithDynamicPins::GenerateAddPinMenu(const FString& InWorkingPinName, SNiagaraGraphPinAdd* InPin)
+{
+	FMenuBuilder MenuBuilder(true, nullptr);
+	BuildTypeMenu(MenuBuilder, InWorkingPinName, InPin);
+	return MenuBuilder.MakeWidget();
+}
+
+void UNiagaraNodeWithDynamicPins::BuildTypeMenu(FMenuBuilder& InMenuBuilder, const FString& InWorkingName, SNiagaraGraphPinAdd* InPin)
+{
+	TArray<FNiagaraTypeDefinition> Types = FNiagaraTypeRegistry::GetRegisteredTypes();
+	Types.Sort([](const FNiagaraTypeDefinition& A, const FNiagaraTypeDefinition& B) { return (A.GetNameText().ToLower().ToString() < B.GetNameText().ToLower().ToString()); });
+
+	for (const FNiagaraTypeDefinition& RegisteredType : Types)
+	{
+		bool bAllowType = false;
+		bAllowType = AllowNiagaraTypeForAddPin(RegisteredType);
+
+		if (bAllowType)
+		{
+			FNiagaraVariable Var(RegisteredType, *InWorkingName);
+			FNiagaraEditorUtilities::ResetVariableToDefaultValue(Var);
+
+			InMenuBuilder.AddMenuEntry(
+				RegisteredType.GetNameText(),
+				FText::Format(LOCTEXT("AddButtonTypeEntryToolTipFormat", "Add a new {0} pin"), RegisteredType.GetNameText()),
+				FSlateIcon(),
+				FUIAction(FExecuteAction::CreateSP(InPin, &SNiagaraGraphPinAdd::OnAddType, Var)));
+		}
+	}
+}
+
+void UNiagaraNodeWithDynamicPins::RemoveDynamicPin(UEdGraphPin* Pin)
+{
+	FScopedTransaction RemovePinTransaction(LOCTEXT("RemovePinTransaction", "Remove pin"));
+	RemovePin(Pin);
+	GetGraph()->NotifyGraphChanged();
+}
+
+FText UNiagaraNodeWithDynamicPins::GetPinNameText(UEdGraphPin* Pin) const
+{
+	return FText::FromString(Pin->PinName);
+}
+
+
+void UNiagaraNodeWithDynamicPins::PinNameTextCommitted(const FText& Text, ETextCommit::Type CommitType, UEdGraphPin* Pin)
+{
+	if (CommitType == ETextCommit::OnEnter)
+	{
+		FScopedTransaction RemovePinTransaction(LOCTEXT("RenamePinTransaction", "Rename pin"));
+		Modify();
+
+		Pin->PinName = Text.ToString();
+		OnPinRenamed(Pin);
+	}
+}
+
+#undef LOCTEXT_NAMESPACE
\ No newline at end of file
diff --git a/Engine/Plugins/FX/Niagara/Source/NiagaraEditor/Private/NiagaraNodeWithDynamicPins.h b/Engine/Plugins/FX/Niagara/Source/NiagaraEditor/Private/NiagaraNodeWithDynamicPins.h
new file mode 100644
--- /dev/null
+++ b/Engine/Plugins/FX/Niagara/Source/NiagaraEditor/Private/NiagaraNodeWithDynamicPins.h
@@ -0,0 +1,79 @@
+// Copyright 1998-2016 Epic Games, Inc. All Rights Reserved.
+
+#pragma once
+
+#include "NiagaraNode.h"
+#include "SlateTypes.h"
+#include "EdGraph/EdGraph.h"
+
+#include "NiagaraNodeWithDynamicPins.generated.h"
+
+class UEdGraphPin;
+class SNiagaraGraphPinAdd;
+
+
+/** A base node for niagara nodes with pins which can be dynamically added and removed by the user. */
+UCLASS()
+class UNiagaraNodeWithDynamicPins : public UNiagaraNode
+{
+public:
+	GENERATED_BODY()
+
+	//~ UEdGraphNode interface
+	virtual void PinConnectionListChanged(UEdGraphPin* Pin) override;
+	virtual void GetContextMenuActions(const FGraphNodeContextMenuBuilder& Context) const override;
+
+	/** Requests a new pin be added to the node with the specified direction, type, and name. */
+	UEdGraphPin* RequestNewTypedPin(EEdGraphPinDirection Direction, const FNiagaraTypeDefinition& Type, FString InName);
+
+	/** Requests a new pin be added to the node with the specified direction and type. */
+	UEdGraphPin* RequestNewTypedPin(EEdGraphPinDirection Direction, const FNiagaraTypeDefinition& Type);
+
+	/** Helper to identify if a pin is an Add pin.*/
+	bool IsAddPin(const UEdGraphPin* Pin) const;
+
+	/** Determine whether or not a Niagara type is supported for an Add Pin possibility.*/
+	virtual bool AllowNiagaraTypeForAddPin(const FNiagaraTypeDefinition& InType);
+
+	/** Used in tandem with SNiagaraGraphPinAdd to generate the menu for selecting the pin to add.*/
+	virtual TSharedRef<SWidget> GenerateAddPinMenu(const FString& InWorkingPinName, SNiagaraGraphPinAdd* InPin);
+
+protected:
+	/** Creates an add pin on the node for the specified direction. */
+	void CreateAddPin(EEdGraphPinDirection Direction);
+
+	/** Called when a new typed pin is added by the user. */
+	virtual void OnNewTypedPinAdded(UEdGraphPin* NewPin) { }
+
+	/** Called when a pin is renamed. */
+	virtual void OnPinRenamed(UEdGraphPin* RenamedPin) { }
+
+	/** Called to determine if a pin can be renamed by the user. */
+	virtual bool CanRenamePin(const UEdGraphPin* Pin) const;
+
+	/** Called to determine if a pin can be removed by the user. */
+	virtual bool CanRemovePin(const UEdGraphPin* Pin) const;
+
+	/** Called to determine if a pin can be moved by the user.*/
+	virtual bool CanMovePin(const UEdGraphPin* Pin) const;
+
+	/** Used GenerateAddPinMenu to build a list of supported types.*/
+	virtual void BuildTypeMenu(FMenuBuilder& InMenuBuilder, const FString& InWorkingName, SNiagaraGraphPinAdd* InPin);
+
+	/** Removes a pin from this node with a transaction. */
+	virtual void RemoveDynamicPin(UEdGraphPin* Pin);
+
+	virtual void MoveDynamicPin(UEdGraphPin* Pin, int32 DirectionToMove);
+
+private:
+
+	/** Gets the display text for a pin. */
+	FText GetPinNameText(UEdGraphPin* Pin) const;
+
+	/** Called when a pin's name text is committed. */
+	void PinNameTextCommitted(const FText& Text, ETextCommit::Type CommitType, UEdGraphPin* Pin);
+
+public:
+	/** The sub category for add pins. */
+	static const FString AddPinSubCategory;
+};
diff --git a/Engine/Plugins/FX/Niagara/Source/NiagaraEditor/Private/NiagaraParameterMapHistory.cpp b/Engine/Plugins/FX/Niagara/Source/NiagaraEditor/Private/NiagaraParameterMapHistory.cpp
new file mode 100644
--- /dev/null
+++ b/Engine/Plugins/FX/Niagara/Source/NiagaraEditor/Private/NiagaraParameterMapHistory.cpp
@@ -0,0 +1,874 @@
+// Copyright 1998-2017 Epic Games, Inc. All Rights Reserved.
+#include "NiagaraParameterMapHistory.h"
+
+#include "NiagaraEditorCommon.h"
+#include "AssetData.h"
+#include "NiagaraHlslTranslator.h"
+#include "NiagaraSystem.h"
+
+#include "ActorFactoryNiagara.h"
+#include "NiagaraActor.h"
+#include "EdGraphSchema_Niagara.h"
+#include "NiagaraNode.h"
+#include "NiagaraCompiler.h"
+#include "NiagaraEditorUtilities.h"
+#include "NiagaraNodeOutput.h"
+#include "NiagaraNodeEmitter.h"
+#include "NiagaraNodeParameterMapGet.h"
+#include "ViewModels/Stack/NiagaraStackEntry.h"
+#include "ViewModels/Stack/NiagaraStackGraphUtilities.h"
+
+#define LOCTEXT_NAMESPACE "NiagaraEditor"
+
+#define PARAM_MAP_ENGINE_STR TEXT("Engine.")
+#define PARAM_MAP_USER_STR TEXT("User.")
+#define PARAM_MAP_SYSTEM_STR TEXT("System.")
+#define PARAM_MAP_EMITTER_STR TEXT("Emitter.")
+#define PARAM_MAP_MODULE_STR TEXT("Module.")
+#define PARAM_MAP_ATTRIBUTE_STR TEXT("Particles.")
+#define PARAM_MAP_INITIAL_STR TEXT("Initial.")
+#define PARAM_MAP_INITIAL_BASE_STR TEXT("Initial")
+
+
+FNiagaraParameterMapHistory::FNiagaraParameterMapHistory() 
+{
+	
+}
+
+int32 FNiagaraParameterMapHistory::RegisterParameterMapPin(const UEdGraphPin* Pin)
+{
+	int32 RetIdx =  MapPinHistory.Add(Pin);
+	return RetIdx;
+}
+
+int32 FNiagaraParameterMapHistory::FindVariableByName(const FName& VariableName)
+{
+	int32 FoundIdx = Variables.IndexOfByPredicate([&](const FNiagaraVariable& InObj) -> bool
+	{
+		return (InObj.GetName() == VariableName);
+	});
+
+	return FoundIdx;
+}
+
+int32 FNiagaraParameterMapHistory::FindVariable(const FName& VariableName, const FNiagaraTypeDefinition& Type)
+{
+	int32 FoundIdx = Variables.IndexOfByPredicate([&](const FNiagaraVariable& InObj) -> bool
+	{
+		return (InObj.GetName() == VariableName && InObj.GetType() == Type);
+	});
+
+	return FoundIdx;
+}
+
+int32 FNiagaraParameterMapHistory::AddVariable(const FNiagaraVariable& InVar, const UEdGraphPin* InPin)
+{
+	FNiagaraVariable Var = InVar;
+	
+	int32 FoundIdx = FindVariable(Var.GetName(), Var.GetType());
+	if (FoundIdx == -1)
+	{
+		FoundIdx = Variables.Add(Var);
+		PerVariableWarnings.AddDefaulted(1);
+		PerVariableWriteHistory.AddDefaulted(1);
+		PerVariableReadHistory.AddDefaulted(1);
+
+		if (InPin != nullptr)
+		{
+			PerVariableWriteHistory[FoundIdx].Add(InPin);
+		}
+
+		check(Variables.Num() == PerVariableWarnings.Num());
+		check(Variables.Num() == PerVariableWriteHistory.Num());
+	}
+	else
+	{
+		if (Variables[FoundIdx].GetType() != Var.GetType())
+		{
+			PerVariableWarnings[FoundIdx].Append(FString::Printf(TEXT("Type mismatch %s instead of %s in map!"), *Var.GetType().GetName(), *Variables[FoundIdx].GetType().GetName()));
+		}
+		if (InPin != nullptr)
+		{
+			PerVariableWriteHistory[FoundIdx].Add(InPin);
+		}
+	}
+
+	return FoundIdx;
+}
+
+int32 FNiagaraParameterMapHistory::AddExternalVariable(const FNiagaraVariable& Var)
+{
+	return AddVariable(Var, nullptr);
+}
+
+const UEdGraphPin* FNiagaraParameterMapHistory::GetFinalPin() const
+{
+	if (MapPinHistory.Num() > 0)
+	{
+		return MapPinHistory[MapPinHistory.Num() - 1];
+	}
+	return nullptr;
+}
+
+const UEdGraphPin* FNiagaraParameterMapHistory::GetOriginalPin() const
+{
+	if (MapPinHistory.Num() > 0)
+	{
+		return MapPinHistory[0];
+	}
+	return nullptr;
+}
+
+FNiagaraVariable FNiagaraParameterMapHistory::ResolveAliases(const FNiagaraVariable& InVar, const TMap<FString, FString>& InAliases, const TCHAR* InJoinSeparator)
+{
+	FNiagaraVariable OutVar = InVar;
+
+	FString OutVarStrName = InVar.GetName().ToString();
+	TArray<FString> SplitName;
+	OutVarStrName.ParseIntoArray(SplitName, TEXT("."));
+
+	for (int32 i = 0; i < SplitName.Num() - 1; i++)
+	{
+		TMap<FString, FString>::TConstIterator It = InAliases.CreateConstIterator();
+		while (It)
+		{
+			if (SplitName[i].Equals(It.Key()))
+			{
+				SplitName[i] = It.Value();
+			}
+			++It;
+		}
+	}
+
+	OutVarStrName = FString::Join<FString>(SplitName, InJoinSeparator);
+
+	OutVar.SetName(*OutVarStrName);
+	return OutVar;
+}
+
+FString FNiagaraParameterMapHistory::MakeSafeNamespaceString(const FString& InStr)
+{
+	FString  Sanitized = FHlslNiagaraTranslator::GetSanitizedSymbolName(InStr);
+	return Sanitized;
+}
+
+
+FNiagaraVariable FNiagaraParameterMapHistory::ResolveAsBasicAttribute(const FNiagaraVariable& InVar)
+{
+	if (IsAttribute(InVar))
+	{
+		FString ParamName = InVar.GetName().ToString();
+		ParamName.RemoveAt(0, FString(PARAM_MAP_ATTRIBUTE_STR).Len());
+
+		ParamName = MakeSafeNamespaceString(ParamName);
+
+		FNiagaraVariable RetVar = InVar;
+		RetVar.SetName(*ParamName);
+		return RetVar;
+	}
+	else
+	{
+		return InVar;
+	}
+}
+
+FNiagaraVariable FNiagaraParameterMapHistory::BasicAttributeToNamespacedAttribute(const FNiagaraVariable& InVar)
+{
+	FString ParamName = InVar.GetName().ToString();
+	ParamName.InsertAt(0, FString(PARAM_MAP_ATTRIBUTE_STR));
+
+	ParamName = MakeSafeNamespaceString(ParamName);
+
+	FNiagaraVariable RetVar = InVar;
+	RetVar.SetName(*ParamName);
+	return RetVar;
+}
+
+FNiagaraVariable FNiagaraParameterMapHistory::VariableToNamespacedVariable(const FNiagaraVariable& InVar, FString Namespace)
+{
+	FString ParamName = Namespace;
+	if (Namespace.EndsWith(TEXT(".")))
+	{
+		ParamName += InVar.GetName().ToString();
+	}
+	else
+	{
+		ParamName += TEXT(".") + InVar.GetName().ToString();
+	}
+	
+
+	FNiagaraVariable RetVar = InVar;
+	RetVar.SetName(*ParamName);
+	return RetVar;
+}
+
+bool FNiagaraParameterMapHistory::IsInNamespace(const FNiagaraVariable& InVar, const FString& Namespace)
+{
+	return InVar.GetName().ToString().StartsWith(Namespace);
+}
+
+bool FNiagaraParameterMapHistory::IsAliasedModuleParameter(const FNiagaraVariable& InVar)
+{
+	return IsInNamespace(InVar, PARAM_MAP_MODULE_STR);
+}
+
+bool FNiagaraParameterMapHistory::IsAliasedEmitterParameter(const FNiagaraVariable& InVar)
+{
+	return IsInNamespace(InVar, PARAM_MAP_EMITTER_STR);
+}
+
+bool FNiagaraParameterMapHistory::IsSystemParameter(const FNiagaraVariable& InVar)
+{
+	return IsInNamespace(InVar, PARAM_MAP_SYSTEM_STR);
+}
+
+bool FNiagaraParameterMapHistory::IsEngineParameter(const FNiagaraVariable& InVar)
+{
+	return IsInNamespace(InVar, PARAM_MAP_ENGINE_STR);
+}
+
+bool FNiagaraParameterMapHistory::IsUserParameter(const FNiagaraVariable& InVar)
+{
+	return IsInNamespace(InVar, PARAM_MAP_USER_STR);
+}
+
+bool FNiagaraParameterMapHistory::IsAttribute(const FNiagaraVariable& InVar)
+{
+	return IsInNamespace(InVar, PARAM_MAP_ATTRIBUTE_STR);
+}
+
+bool FNiagaraParameterMapHistory::IsSystemNamespaceReadOnly(UNiagaraScript* InScript)
+{
+	if (InScript->IsSystemSpawnScript() || InScript->IsSystemUpdateScript())
+	{
+		return false;
+	}
+	else if (InScript->IsStandaloneScript())
+	{
+		TArray<ENiagaraScriptUsage> IntendedUsages = InScript->GetSupportedUsageContexts();
+		if (IntendedUsages.Contains(ENiagaraScriptUsage::SystemSpawnScript) || IntendedUsages.Contains(ENiagaraScriptUsage::SystemUpdateScript))
+		{
+			return false;
+		}
+	}
+
+	return true;
+}
+
+bool FNiagaraParameterMapHistory::IsEmitterNamespaceReadOnly(UNiagaraScript* InScript)
+{
+	if (InScript->IsEmitterSpawnScript() || InScript->IsEmitterUpdateScript() || InScript->IsSystemSpawnScript() || InScript->IsSystemUpdateScript())
+	{
+		return false;
+	}
+	else if (InScript->IsStandaloneScript())
+	{
+		TArray<ENiagaraScriptUsage> IntendedUsages = InScript->GetSupportedUsageContexts();
+		if (IntendedUsages.Contains(ENiagaraScriptUsage::EmitterSpawnScript) || IntendedUsages.Contains(ENiagaraScriptUsage::EmitterUpdateScript) || IntendedUsages.Contains(ENiagaraScriptUsage::SystemSpawnScript) || IntendedUsages.Contains(ENiagaraScriptUsage::SystemUpdateScript))
+		{
+			return false;
+		}
+	}
+
+	return true;
+}
+
+const UEdGraphPin* FNiagaraParameterMapHistory::GetDefaultValuePin(int32 VarIdx) const
+{
+	if (PerVariableWriteHistory[VarIdx].Num() > 0)
+	{
+		const UEdGraphPin* Pin = PerVariableWriteHistory[VarIdx][0];
+		if (Pin != nullptr && Pin->Direction == EEdGraphPinDirection::EGPD_Input && Cast<UNiagaraNodeParameterMapGet>(Pin->GetOwningNode()) != nullptr)
+		{
+			return Pin;
+		}
+	}
+	return nullptr;
+}
+
+bool FNiagaraParameterMapHistory::IsInitialValue(const FNiagaraVariable& InVar)
+{
+	TArray<FString> SplitName;
+	InVar.GetName().ToString().ParseIntoArray(SplitName, TEXT("."));
+
+	for (int32 i = 1; i < SplitName.Num() - 1; i++)
+	{
+		if (SplitName[i].Equals(PARAM_MAP_INITIAL_BASE_STR, ESearchCase::IgnoreCase))
+		{
+			return true;
+		}
+	}
+	return false;
+}
+
+FNiagaraVariable FNiagaraParameterMapHistory::GetSourceForInitialValue(const FNiagaraVariable& InVar)
+{
+	TArray<FString> SplitName;
+	InVar.GetName().ToString().ParseIntoArray(SplitName, TEXT("."));
+
+	TArray<FString> JoinString;
+	bool bFound = false;
+	for (int32 i = 0; i < SplitName.Num(); i++)
+	{
+		if (!bFound && SplitName[i].Equals(PARAM_MAP_INITIAL_BASE_STR, ESearchCase::IgnoreCase))
+		{
+			bFound = true;
+			continue;
+		}
+		else
+		{
+			JoinString.Add(SplitName[i]);
+		}
+	}
+
+	FString OutVarStrName = FString::Join<FString>(JoinString, TEXT("."));
+	FNiagaraVariable Var = InVar;
+	Var.SetName(*OutVarStrName);
+
+	return Var;
+}
+
+bool FNiagaraParameterMapHistory::IsPrimaryDataSetOutput(const FNiagaraVariable& InVar, UNiagaraScript* InScript)
+{
+	return IsPrimaryDataSetOutput(InVar, InScript->GetUsage());
+}
+
+bool FNiagaraParameterMapHistory::IsPrimaryDataSetOutput(const FNiagaraVariable& InVar, ENiagaraScriptUsage Usage)
+{
+	if (InVar.GetType().GetClass() != nullptr)
+	{
+		return false;
+	}
+
+	if (Usage == ENiagaraScriptUsage::EmitterSpawnScript || Usage == ENiagaraScriptUsage::EmitterUpdateScript || 
+		Usage == ENiagaraScriptUsage::SystemSpawnScript || Usage == ENiagaraScriptUsage::SystemUpdateScript)
+	{
+		// In the case of system/emitter scripts we must include the variables in the overall system namespace as well as any of 
+		// the child emitters that were encountered.
+		for (FString EmitterEncounteredNamespace : EmitterNamespacesEncountered)
+		{
+			if (IsInNamespace(InVar, EmitterEncounteredNamespace))
+			{
+				return true;
+			}
+		}
+		return IsInNamespace(InVar, PARAM_MAP_SYSTEM_STR) || IsInNamespace(InVar, PARAM_MAP_EMITTER_STR);
+	}
+	else if (Usage == ENiagaraScriptUsage::Module || Usage == ENiagaraScriptUsage::Function)
+	{
+		return IsInNamespace(InVar, PARAM_MAP_MODULE_STR);
+	}
+	return IsInNamespace(InVar, PARAM_MAP_ATTRIBUTE_STR);
+}
+
+
+FNiagaraVariable FNiagaraParameterMapHistory::MoveToExternalConstantNamespaceVariable(const FNiagaraVariable& InVar, UNiagaraScript* InScript)
+{
+	if (InScript->IsParticleEventScript() || InScript->IsParticleSpawnScript() || InScript->IsParticleUpdateScript())
+	{
+		return VariableToNamespacedVariable(InVar, PARAM_MAP_EMITTER_STR);
+	}
+	else if (InScript->IsStandaloneScript())
+	{
+		return VariableToNamespacedVariable(InVar, PARAM_MAP_MODULE_STR);
+	}
+	else if (InScript->IsEmitterSpawnScript() || InScript->IsEmitterUpdateScript() || InScript->IsSystemSpawnScript() || InScript->IsSystemUpdateScript())
+	{
+		return VariableToNamespacedVariable(InVar, PARAM_MAP_USER_STR);
+	}
+	return InVar;
+}
+
+bool FNiagaraParameterMapHistory::IsExportableExternalConstant(const FNiagaraVariable& InVar, UNiagaraScript* InScript)
+{
+	if (InScript->IsEquivalentUsage(ENiagaraScriptUsage::SystemSpawnScript))
+	{
+		return IsExternalConstantNamespace(InVar, InScript);
+	}
+	else
+	{
+		return false;
+	}
+}
+
+bool FNiagaraParameterMapHistory::IsExternalConstantNamespace(const FNiagaraVariable& InVar, UNiagaraScript* InScript)
+{
+	// Engine parameters are never writable.
+	if (IsInNamespace(InVar, PARAM_MAP_ENGINE_STR))
+	{
+		return true;
+	}
+
+	if (IsInNamespace(InVar, PARAM_MAP_USER_STR))
+	{
+		return true;
+	}
+
+	// Modules and functions need to act as if they are within the script types that they 
+	// say that they support rather than using their exact script type.
+	if (InScript->IsStandaloneScript())
+	{
+		TArray<ENiagaraScriptUsage> SupportedContexts = InScript->GetSupportedUsageContexts();
+		if (((!SupportedContexts.Contains(ENiagaraScriptUsage::EmitterSpawnScript) && !SupportedContexts.Contains(ENiagaraScriptUsage::EmitterUpdateScript)) && IsInNamespace(InVar, PARAM_MAP_EMITTER_STR))
+			|| ((!SupportedContexts.Contains(ENiagaraScriptUsage::SystemSpawnScript) && !SupportedContexts.Contains(ENiagaraScriptUsage::SystemUpdateScript)) && IsInNamespace(InVar, PARAM_MAP_SYSTEM_STR)))
+		{
+			return true;
+		}
+	}
+	
+	// Particle scripts cannot write to the emitter or system namespace.
+	if (InScript->IsParticleEventScript() || InScript->IsParticleSpawnScript() || InScript->IsParticleUpdateScript())
+	{
+		if (IsInNamespace(InVar, PARAM_MAP_EMITTER_STR) || IsInNamespace(InVar, PARAM_MAP_SYSTEM_STR))
+		{
+			return true;
+		}
+	}
+	
+	return false;
+}
+
+const UNiagaraNodeOutput* FNiagaraParameterMapHistory::GetFinalOutputNode() const
+{
+	const UEdGraphPin* Pin = GetFinalPin();
+	if (Pin != nullptr)
+	{
+		const UNiagaraNodeOutput* Output = Cast<const UNiagaraNodeOutput>(Pin->GetOwningNode());
+		if (Output != nullptr)
+		{
+			return Output;
+		}
+	}
+
+	return nullptr;
+}
+
+FNiagaraParameterMapHistoryBuilder::FNiagaraParameterMapHistoryBuilder()
+{
+	ContextuallyVisitedNodes.AddDefaulted(1);
+	PinToParameterMapIndices.AddDefaulted(1);
+	bFilterByScriptWhitelist = false;
+	FilterScriptType = ENiagaraScriptUsage::Function;
+}
+
+void FNiagaraParameterMapHistoryBuilder::BuildParameterMaps(UNiagaraNodeOutput* OutputNode, bool bRecursive)
+{
+	RelevantScriptUsageContext.Emplace(OutputNode->GetUsage());
+	OutputNode->BuildParameterMapHistory(*this, bRecursive);
+	RelevantScriptUsageContext.Pop();
+}
+
+void FNiagaraParameterMapHistoryBuilder::EnableScriptWhitelist(bool bInEnable, ENiagaraScriptUsage InScriptType)
+{
+	bFilterByScriptWhitelist = bInEnable;
+	FilterScriptType = InScriptType;
+}
+
+ENiagaraScriptUsage FNiagaraParameterMapHistoryBuilder::GetCurrentUsageContext()const
+{
+	return RelevantScriptUsageContext.Last();
+}
+
+int32 FNiagaraParameterMapHistoryBuilder::CreateParameterMap()
+{
+	int32 RetValue = Histories.AddDefaulted(1);
+	return RetValue;
+}
+
+int32 FNiagaraParameterMapHistoryBuilder::RegisterParameterMapPin(int32 WhichParameterMap, const UEdGraphPin* Pin)
+{
+	if (WhichParameterMap != INDEX_NONE)
+	{
+		if (Pin->Direction == EEdGraphPinDirection::EGPD_Output)
+		{
+			PinToParameterMapIndices.Last().Add(Pin, WhichParameterMap);
+		}
+
+		return Histories[WhichParameterMap].RegisterParameterMapPin(Pin);
+	}
+	else
+	{
+		return INDEX_NONE;
+	}
+}
+
+int32 FNiagaraParameterMapHistoryBuilder::TraceParameterMapOutputPin(const UEdGraphPin* OutputPin)
+{
+	if (OutputPin->Direction == EEdGraphPinDirection::EGPD_Output)
+	{
+		const int32* IdxPtr = PinToParameterMapIndices.Last().Find(OutputPin);
+		if (IdxPtr != nullptr)
+		{
+			return *IdxPtr;
+		}
+	}
+	return INDEX_NONE;
+}
+
+bool FNiagaraParameterMapHistoryBuilder::GetPinPreviouslyVisited(const UEdGraphPin* InPin) const
+{
+	if (InPin != nullptr)
+		return GetNodePreviouslyVisited(CastChecked<UNiagaraNode>(InPin->GetOwningNode()));
+	else
+		return true;
+}
+
+
+bool FNiagaraParameterMapHistoryBuilder::GetNodePreviouslyVisited(const class UNiagaraNode* Node) const
+{
+	return ContextuallyVisitedNodes.Last().Contains(Node);
+}
+
+int32 FNiagaraParameterMapHistoryBuilder::FindMatchingParameterMapFromContextInputs(const FNiagaraVariable& InVar) const
+{
+	if (CallingContext.Num() == 0)
+	{
+		return INDEX_NONE;
+	}
+	UNiagaraNode* Node = CallingContext.Last();
+	TArray<UEdGraphPin*> Inputs;
+	Node->GetInputPins(Inputs);
+	const UEdGraphSchema_Niagara* Schema = CastChecked<UEdGraphSchema_Niagara>(Node->GetSchema());
+
+	for (int32 i = 0; i < Inputs.Num(); i++)
+	{
+		FNiagaraVariable CallInputVar = Schema->PinToNiagaraVariable(Inputs[i]);
+		if (CallInputVar.IsEquivalent(InVar) && CallInputVar.GetType() == FNiagaraTypeDefinition::GetParameterMapDef())
+		{
+			if (Inputs[i]->LinkedTo.Num() != 0 && PinToParameterMapIndices.Num() >= 2)
+			{
+				const int32* ParamMapIdxPtr = PinToParameterMapIndices[PinToParameterMapIndices.Num() - 2].Find(Inputs[i]->LinkedTo[0]);
+				if (ParamMapIdxPtr != nullptr)
+				{
+					return *ParamMapIdxPtr;
+				}
+				else
+				{
+					FString ScriptUsageDisplayName;
+					UNiagaraNodeOutput* ContextOutputNode = FNiagaraStackGraphUtilities::GetEmitterOutputNodeForStackNode(*Node);
+					if (ContextOutputNode != nullptr)
+					{
+						UEnum* NiagaraScriptUsageEnum = FindObjectChecked<UEnum>(ANY_PACKAGE, TEXT("ENiagaraScriptUsage"), true);
+						ScriptUsageDisplayName = NiagaraScriptUsageEnum->GetDisplayNameTextByValue((uint64)ContextOutputNode->GetUsage()).ToString();
+					}
+					else
+					{
+						ScriptUsageDisplayName = "Unknown";
+					}
+					FString FunctionDisplayName;
+					if (FunctionNameContextStack.Num() > 0)
+					{
+						FunctionDisplayName = FunctionNameContextStack.Last().ToString();
+					}
+					else
+					{
+						FunctionDisplayName = Node->GetName();
+					}
+
+					
+					/* 
+					UE_LOG(LogNiagaraEditor, Error, TEXT("Unable to find matching parameter map for variable.  Name:%s  Function:%s  Usage:%s  Package:%s"),
+						*InVar.GetName().ToString(), *FunctionDisplayName, *ScriptUsageDisplayName, *Node->GetOutermost()->GetName());
+					*/
+					
+				}
+			}
+		}
+	}
+	return INDEX_NONE;
+}
+
+void FNiagaraParameterMapHistoryBuilder::BeginTranslation(UNiagaraScript* Script)
+{
+	//For now this will just tell particle scripts what emitter they're being compiled as part of but maybe we want to do more here.
+	//This is mainly so that parameter names match up between System/Emitter scripts and the parameters they drive within particle scripts.
+	//I dislike this coupling of the translator to emitters but for now it'll have to do.
+	//Will refactor in the future.
+	UNiagaraEmitter* Emitter = Script->GetTypedOuter<UNiagaraEmitter>();
+	BeginTranslation(Emitter);
+}
+
+void FNiagaraParameterMapHistoryBuilder::EndTranslation(UNiagaraScript* Script)
+{
+	EmitterNameContextStack.Reset();
+}
+
+void FNiagaraParameterMapHistoryBuilder::BeginTranslation(UNiagaraEmitter* Emitter)
+{
+	//For now this will just tell particle scripts what emitter they're being compiled as part of but maybe we want to do more here.
+	//This is mainly so that parameter names match up between System/Emitter scripts and the parameters they drive within particle scripts.
+	//I dislike this coupling of the translator to emitters but for now it'll have to do.
+	//Will refactor in the future.
+	if (Emitter)
+	{
+		FString EmitterUniqueName = Emitter->GetUniqueEmitterName();
+		EmitterNameContextStack.Add(*EmitterUniqueName);
+	}
+	BuildCurrentAliases();
+}
+
+void FNiagaraParameterMapHistoryBuilder::EndTranslation(UNiagaraEmitter* Emitter)
+{
+	EmitterNameContextStack.Reset();
+}
+
+void FNiagaraParameterMapHistoryBuilder::EnterFunction(const FString& InNodeName, UNiagaraScript* InScript, UNiagaraNode* Node)
+{
+	if (InScript != nullptr )
+	{
+		RegisterNodeVisitation(Node);
+		CallingContext.Push(Node);
+		PinToParameterMapIndices.Emplace();
+		FunctionNameContextStack.Emplace(*InNodeName);
+		BuildCurrentAliases();
+		ContextuallyVisitedNodes.Emplace();
+	}
+}
+
+void FNiagaraParameterMapHistoryBuilder::ExitFunction(const FString& InNodeName, UNiagaraScript* InScript, UNiagaraNode* Node)
+{
+	if (InScript != nullptr)
+	{
+		CallingContext.Pop();
+		PinToParameterMapIndices.Pop();
+		FunctionNameContextStack.Pop();
+		BuildCurrentAliases();
+		ContextuallyVisitedNodes.Pop();
+	}
+}
+
+void FNiagaraParameterMapHistoryBuilder::EnterEmitter(const FString& InEmitterName, UNiagaraNode* Node)
+{
+	RegisterNodeVisitation(Node);
+	CallingContext.Push(Node);
+	EmitterNameContextStack.Emplace(*InEmitterName);
+	BuildCurrentAliases();
+
+	// Emitters must record their namespaces to their histories as well as
+	// make sure to record their current usage type is so that we can filter variables
+	// for relevance downstream.
+	UNiagaraNodeEmitter* EmitterNode = Cast<UNiagaraNodeEmitter>(Node);
+	if (EmitterNode != nullptr)
+	{
+		RelevantScriptUsageContext.Emplace(EmitterNode->GetUsage());
+	}
+	else
+	{
+		RelevantScriptUsageContext.Emplace(ENiagaraScriptUsage::EmitterSpawnScript);
+	}
+
+	for (FNiagaraParameterMapHistory& History : Histories)
+	{
+		History.EmitterNamespacesEncountered.AddUnique(InEmitterName);
+	}
+	ContextuallyVisitedNodes.Emplace();
+}
+
+void FNiagaraParameterMapHistoryBuilder::ExitEmitter(const FString& InEmitterName, UNiagaraNode* Node)
+{
+	CallingContext.Pop();
+	EmitterNameContextStack.Pop();
+	BuildCurrentAliases();
+	ContextuallyVisitedNodes.Pop();
+}
+
+/**
+* Use the current alias map to resolve any aliases in this input variable name.
+*/
+FNiagaraVariable FNiagaraParameterMapHistoryBuilder::ResolveAliases(const FNiagaraVariable& InVar)
+{
+	return FNiagaraParameterMapHistory::ResolveAliases(InVar, AliasMap, TEXT("."));
+}
+
+
+void FNiagaraParameterMapHistoryBuilder::RegisterNodeVisitation(const UEdGraphNode* Node)
+{
+	ContextuallyVisitedNodes.Last().AddUnique(CastChecked<UNiagaraNode>(Node));
+}
+
+
+const FString* FNiagaraParameterMapHistoryBuilder::GetModuleAlias() const
+{
+	return AliasMap.Find(TEXT("Module"));
+}
+
+const FString* FNiagaraParameterMapHistoryBuilder::GetEmitterAlias() const
+{
+	return AliasMap.Find(TEXT("Emitter"));
+}
+
+void FNiagaraParameterMapHistoryBuilder::VisitInputPin(const UEdGraphPin* Pin, class UNiagaraNode* InNode)
+{
+	const UEdGraphSchema_Niagara* Schema = GetDefault<UEdGraphSchema_Niagara>();
+
+	if (Pin != nullptr && Pin->Direction == EEdGraphPinDirection::EGPD_Input)
+	{
+		for (int32 j = 0; j < Pin->LinkedTo.Num(); j++)
+		{
+			UNiagaraNode* Node = CastChecked<UNiagaraNode>(Pin->LinkedTo[j]->GetOwningNode());
+
+			if (!GetNodePreviouslyVisited(Node))
+			{
+				Node->BuildParameterMapHistory(*this, true);
+				RegisterNodeVisitation(Node);
+			}
+
+			if (Schema->PinToTypeDefinition(Pin) == FNiagaraTypeDefinition::GetParameterMapDef())
+			{
+				int32 ParamMapIdx = TraceParameterMapOutputPin(Pin->LinkedTo[j]);
+				RegisterParameterMapPin(ParamMapIdx, Pin);
+			}
+		} 
+	}
+}
+
+void FNiagaraParameterMapHistoryBuilder::VisitInputPins(class UNiagaraNode* InNode)
+{
+	TArray<UEdGraphPin*> InputPins;
+	InNode->GetInputPins(InputPins);
+
+	for (int32 i = 0; i < InputPins.Num(); i++)
+	{
+		VisitInputPin(InputPins[i], InNode);
+	}
+}
+
+bool FNiagaraParameterMapHistoryBuilder::IsNamespacedVariableRelevantToScriptType(const FNiagaraVariable& InVar, ENiagaraScriptUsage InFilterScriptType)
+{
+	return true;
+}
+
+bool FNiagaraParameterMapHistoryBuilder::ShouldTrackVariable(const FNiagaraVariable& InVar)
+{
+	if (!bFilterByScriptWhitelist)
+	{
+		return true;
+	}
+	if (IsNamespacedVariableRelevantToScriptType(InVar, FilterScriptType))
+	{
+		return true;
+	}
+	return false;
+}
+
+int32 FNiagaraParameterMapHistoryBuilder::HandleVariableWrite(int32 ParamMapIdx, const UEdGraphPin* InPin)
+{
+	const UEdGraphSchema_Niagara* Schema = GetDefault<UEdGraphSchema_Niagara>();
+	FNiagaraVariable Var = Schema->PinToNiagaraVariable(InPin);
+
+	if (!ShouldTrackVariable(Var))
+	{
+		return INDEX_NONE;
+	}
+	Var = ResolveAliases(Var);
+
+	return Histories[ParamMapIdx].AddVariable(Var, InPin);
+}
+
+
+int32 FNiagaraParameterMapHistoryBuilder::HandleVariableRead(int32 ParamMapIdx, const UEdGraphPin* InPin, bool RegisterReadsAsVariables, const UEdGraphPin* InDefaultPin, bool& OutUsedDefault)
+{
+	OutUsedDefault = false;
+	const UEdGraphSchema_Niagara* Schema = GetDefault<UEdGraphSchema_Niagara>();
+	FNiagaraVariable Var = Schema->PinToNiagaraVariable(InPin);
+
+	if (!ShouldTrackVariable(Var))
+	{
+		return INDEX_NONE;
+	}
+	Var = ResolveAliases(Var);
+
+	int32 FoundIdx = Histories[ParamMapIdx].FindVariable(Var.GetName(), Var.GetType());
+	if (FoundIdx == -1)
+	{
+		if (RegisterReadsAsVariables)
+		{
+			OutUsedDefault = false;
+			if (InDefaultPin)
+			{
+				OutUsedDefault = true;
+				VisitInputPin(InDefaultPin, Cast<UNiagaraNode>(InDefaultPin->GetOwningNode()));
+				FoundIdx = Histories[ParamMapIdx].FindVariable(Var.GetName(), Var.GetType());
+			}
+
+			if (FoundIdx == -1)
+			{
+				FoundIdx = Histories[ParamMapIdx].Variables.Add(Var);
+			}
+			Histories[ParamMapIdx].PerVariableWarnings.AddDefaulted(1);
+			Histories[ParamMapIdx].PerVariableWriteHistory.AddDefaulted(1);
+			Histories[ParamMapIdx].PerVariableReadHistory.AddDefaulted(1);
+			Histories[ParamMapIdx].PerVariableReadHistory[FoundIdx].Add(TTuple<const UEdGraphPin*, const UEdGraphPin*>(InPin, nullptr));
+
+			if (InDefaultPin && OutUsedDefault)
+			{
+				Histories[ParamMapIdx].PerVariableWriteHistory[FoundIdx].Add(InDefaultPin);				
+			}
+		}
+		check(Histories[ParamMapIdx].Variables.Num() == Histories[ParamMapIdx].PerVariableWarnings.Num());
+		check(Histories[ParamMapIdx].Variables.Num() == Histories[ParamMapIdx].PerVariableWriteHistory.Num());
+		check(Histories[ParamMapIdx].Variables.Num() == Histories[ParamMapIdx].PerVariableReadHistory.Num());
+	}
+	else
+	{
+		if (Histories[ParamMapIdx].Variables[FoundIdx].GetType() != Var.GetType())
+		{
+			Histories[ParamMapIdx].PerVariableWarnings[FoundIdx].Append(FString::Printf(TEXT("Type mismatch %s instead of %s in map!"), *Var.GetType().GetName(), *Histories[ParamMapIdx].Variables[FoundIdx].GetType().GetName()));
+		}
+
+
+		if (Histories[ParamMapIdx].PerVariableWriteHistory[FoundIdx].Num() > 0)
+		{
+			Histories[ParamMapIdx].PerVariableReadHistory[FoundIdx].Add(TTuple<const UEdGraphPin*, const UEdGraphPin*>(InPin, Histories[ParamMapIdx].PerVariableWriteHistory[FoundIdx][Histories[ParamMapIdx].PerVariableWriteHistory[FoundIdx].Num() - 1]));
+		}
+		else
+		{
+			Histories[ParamMapIdx].PerVariableReadHistory[FoundIdx].Add(TTuple<const UEdGraphPin*, const UEdGraphPin*>(InPin, nullptr));
+		}
+	}
+
+	return FoundIdx;
+}
+
+void FNiagaraParameterMapHistoryBuilder::BuildCurrentAliases()
+{
+	AliasMap = TMap<FString, FString>();
+
+	FString Callstack;
+	for (int32 i = 0; i < FunctionNameContextStack.Num(); i++)
+	{
+		if (i == 0)
+		{
+			Callstack += FunctionNameContextStack[i].ToString();
+		}
+		else
+		{
+			Callstack += TEXT(".") + FunctionNameContextStack[i].ToString();
+		}
+	}
+
+	if (!Callstack.IsEmpty())
+	{
+		AliasMap.Add(TEXT("Module"), Callstack);
+	}
+
+	Callstack.Empty();
+	for (int32 i = 0; i < EmitterNameContextStack.Num(); i++)
+	{
+		if (i == 0)
+		{
+			Callstack += EmitterNameContextStack[i].ToString();
+		}
+		else
+		{
+			Callstack += TEXT(".") + EmitterNameContextStack[i].ToString();
+		}
+	}
+
+	if (!Callstack.IsEmpty())
+	{
+		AliasMap.Add(TEXT("Emitter"), Callstack);
+	}
+}
+
+#undef LOCTEXT_NAMESPACE
diff --git a/Engine/Plugins/FX/Niagara/Source/NiagaraEditor/Private/NiagaraScriptFactory.cpp b/Engine/Plugins/FX/Niagara/Source/NiagaraEditor/Private/NiagaraScriptFactory.cpp
new file mode 100644
--- /dev/null
+++ b/Engine/Plugins/FX/Niagara/Source/NiagaraEditor/Private/NiagaraScriptFactory.cpp
@@ -0,0 +1,90 @@
+// Copyright 1998-2017 Epic Games, Inc. All Rights Reserved.
+
+#include "CoreMinimal.h"
+#include "Misc/ConfigCacheIni.h"
+#include "Modules/ModuleManager.h"
+#include "NiagaraCommon.h"
+#include "NiagaraEditorModule.h"
+#include "NiagaraScript.h"
+#include "EdGraph/EdGraph.h"
+#include "NiagaraGraph.h"
+#include "NiagaraScriptSource.h"
+#include "NiagaraScriptFactoryNew.h"
+#include "NiagaraNodeOutput.h"
+#include "NiagaraEditorSettings.h"
+
+#define LOCTEXT_NAMESPACE "NiagaraScriptFactory"
+
+UNiagaraScriptFactoryNew::UNiagaraScriptFactoryNew(const FObjectInitializer& ObjectInitializer)
+	: Super(ObjectInitializer)
+{
+
+	SupportedClass = UNiagaraScript::StaticClass();
+	bCreateNew = false;
+	bEditAfterNew = true;
+	bCreateNew = true;
+}
+
+UObject* UNiagaraScriptFactoryNew::FactoryCreateNew(UClass* Class, UObject* InParent, FName Name, EObjectFlags Flags, UObject* Context, FFeedbackContext* Warn)
+{
+	check(Class->IsChildOf(UNiagaraScript::StaticClass()));
+	
+	const UNiagaraEditorSettings* Settings = GetDefault<UNiagaraEditorSettings>();
+	check(Settings);
+
+	UNiagaraScript* NewScript;
+
+	if (UNiagaraScript* Default = Cast<UNiagaraScript>(Settings->DefaultScript.TryLoad()))
+	{
+		NewScript = Cast<UNiagaraScript>(StaticDuplicateObject(Default, InParent, Name, Flags, Class));
+	}
+	else
+	{
+		NewScript = NewObject<UNiagaraScript>(InParent, Class, Name, Flags | RF_Transactional);
+		InitializeScript(NewScript);
+	}
+
+	return NewScript;
+}
+
+void UNiagaraScriptFactoryNew::InitializeScript(UNiagaraScript* NewScript)
+{
+	if(NewScript != NULL)
+	{
+		UNiagaraScriptSource* Source = NewObject<UNiagaraScriptSource>(NewScript, NAME_None, RF_Transactional);
+		if(Source)
+		{
+			UNiagaraGraph* CreatedGraph = NewObject<UNiagaraGraph>(Source, NAME_None, RF_Transactional);
+			Source->NodeGraph = CreatedGraph;
+
+			FGraphNodeCreator<UNiagaraNodeOutput> OutputNodeCreator(*CreatedGraph);
+			UNiagaraNodeOutput* OutputNode = OutputNodeCreator.CreateNode();
+
+			FNiagaraVariable PosAttrib(FNiagaraTypeDefinition::GetVec3Def(), "Position");
+			FNiagaraVariable VelAttrib(FNiagaraTypeDefinition::GetVec3Def(), "Velocity");
+			FNiagaraVariable RotAttrib(FNiagaraTypeDefinition::GetFloatDef(), "Rotation");
+			FNiagaraVariable AgeAttrib(FNiagaraTypeDefinition::GetFloatDef(), "NormalizedAge");
+			FNiagaraVariable ColAttrib(FNiagaraTypeDefinition::GetColorDef(), "Color");
+			FNiagaraVariable SizeAttrib(FNiagaraTypeDefinition::GetVec2Def(), "Size");
+
+			OutputNode->Outputs.Add(PosAttrib);
+			OutputNode->Outputs.Add(VelAttrib);
+			OutputNode->Outputs.Add(RotAttrib);
+			OutputNode->Outputs.Add(ColAttrib);
+			OutputNode->Outputs.Add(SizeAttrib);
+			OutputNode->Outputs.Add(AgeAttrib);
+
+			OutputNodeCreator.Finalize();
+
+
+			// Set pointer in script to source
+			NewScript->SetSource(Source);
+
+			FString ErrorMessages;
+			FNiagaraEditorModule& NiagaraEditorModule = FModuleManager::Get().LoadModuleChecked<FNiagaraEditorModule>(TEXT("NiagaraEditor"));
+			NiagaraEditorModule.CompileScript(NewScript, ErrorMessages);
+		}
+	}
+}
+
+#undef LOCTEXT_NAMESPACE
diff --git a/Engine/Plugins/FX/Niagara/Source/NiagaraEditor/Private/NiagaraScriptSource.cpp b/Engine/Plugins/FX/Niagara/Source/NiagaraEditor/Private/NiagaraScriptSource.cpp
new file mode 100644
--- /dev/null
+++ b/Engine/Plugins/FX/Niagara/Source/NiagaraEditor/Private/NiagaraScriptSource.cpp
@@ -0,0 +1,227 @@
+// Copyright 1998-2017 Epic Games, Inc. All Rights Reserved.
+
+#include "NiagaraScriptSource.h"
+#include "Modules/ModuleManager.h"
+#include "NiagaraCommon.h"
+#include "NiagaraEditorModule.h"
+#include "NiagaraScript.h"
+#include "NiagaraComponent.h"
+#include "UObject/UObjectHash.h"
+#include "UObject/UObjectIterator.h"
+#include "ComponentReregisterContext.h"
+#include "NiagaraGraph.h"
+#include "NiagaraConstants.h"
+#include "NiagaraSystem.h"
+#include "NiagaraNodeOutput.h"
+#include "NiagaraNodeInput.h"
+#include "NiagaraNodeWriteDataSet.h"
+#include "NiagaraNodeReadDataSet.h"
+#include "NiagaraScript.h"
+#include "NiagaraDataInterface.h"
+#include "GraphEditAction.h"
+#include "NiagaraNodeParameterMapBase.h"
+#include "NiagaraNodeOutput.h"
+#include "EdGraphUtilities.h"
+#include "EdGraphSchema_Niagara.h"
+#include "TokenizedMessage.h"
+#include "NiagaraEditorUtilities.h"
+
+UNiagaraScriptSource::UNiagaraScriptSource(const FObjectInitializer& ObjectInitializer)
+	: Super(ObjectInitializer), bIsPrecompiled(false)
+{
+}
+
+void UNiagaraScriptSource::PostLoad()
+{
+	Super::PostLoad();
+	// We need to make sure that the node-graph is already resolved b/c we may be asked IsSyncrhonized later...
+	if (NodeGraph)
+	{
+		NodeGraph->ConditionalPostLoad();
+	}
+}
+
+UNiagaraScriptSourceBase* UNiagaraScriptSource::MakeRecursiveDeepCopy(UObject* DestOuter, TMap<const UObject*, UObject*>& ExistingConversions) const
+{
+	check(GetOuter() != DestOuter);
+	EObjectFlags Flags = RF_AllFlags & ~RF_Standalone & ~RF_Public; // Remove Standalone and Public flags.
+	ResetLoaders(GetTransientPackage()); // Make sure that we're not going to get invalid version number linkers into the package we are going into. 
+	GetTransientPackage()->LinkerCustomVersion.Empty();
+
+	UNiagaraScriptSource*	ScriptSource = CastChecked<UNiagaraScriptSource>(StaticDuplicateObject(this, GetTransientPackage(), NAME_None, Flags));
+	check(ScriptSource->HasAnyFlags(RF_Standalone) == false);
+	check(ScriptSource->HasAnyFlags(RF_Public) == false);
+
+	ScriptSource->Rename(nullptr, DestOuter, REN_DoNotDirty | REN_DontCreateRedirectors | REN_NonTransactional);
+	UE_LOG(LogNiagaraEditor, Warning, TEXT("MakeRecursiveDeepCopy %s"), *ScriptSource->GetFullName());
+	ExistingConversions.Add(this, ScriptSource);
+
+	ScriptSource->SubsumeExternalDependencies(ExistingConversions);
+	return ScriptSource;
+}
+
+void UNiagaraScriptSource::SubsumeExternalDependencies(TMap<const UObject*, UObject*>& ExistingConversions)
+{
+	if (NodeGraph)
+	{
+		NodeGraph->SubsumeExternalDependencies(ExistingConversions);
+	}
+}
+
+bool UNiagaraScriptSource::IsSynchronized(const FGuid& InChangeId)
+{
+	if (NodeGraph)
+	{
+		return NodeGraph->IsOtherSynchronized(InChangeId);
+	}
+	else
+	{
+		return false;
+	}
+}
+
+void UNiagaraScriptSource::MarkNotSynchronized()
+{
+	if (NodeGraph)
+	{
+		NodeGraph->MarkGraphRequiresSynchronization();
+	}
+}
+
+bool UNiagaraScriptSource::IsPreCompiled() const
+{
+	return bIsPrecompiled;
+}
+
+void UNiagaraScriptSource::PreCompile(UNiagaraEmitter* InEmitter, bool bClearErrors)
+{
+	if (!bIsPrecompiled)
+	{
+		bIsPrecompiled = true;
+
+		if (bClearErrors)
+		{
+			//Clear previous graph errors.
+			bool bHasClearedGraphErrors = false;
+			for (UEdGraphNode* Node : NodeGraph->Nodes)
+			{
+				if (Node->bHasCompilerMessage)
+				{
+					Node->ErrorMsg.Empty();
+					Node->ErrorType = EMessageSeverity::Info;
+					Node->bHasCompilerMessage = false;
+					Node->Modify(true);
+					bHasClearedGraphErrors = true;
+				}
+			}
+			if (bHasClearedGraphErrors)
+			{
+				NodeGraph->NotifyGraphChanged();
+			}
+		}
+
+		// Clone the source graph so we can modify it as needed; merging in the child graphs
+		NodeGraphDeepCopy = CastChecked<UNiagaraGraph>(FEdGraphUtilities::CloneGraph(NodeGraph, this, 0));
+		FEdGraphUtilities::MergeChildrenGraphsIn(NodeGraphDeepCopy, NodeGraphDeepCopy, /*bRequireSchemaMatch=*/ true);
+		
+		PrecompiledHistories.Empty();
+
+		TArray<UNiagaraNodeOutput*> OutputNodes;
+		NodeGraphDeepCopy->FindOutputNodes(OutputNodes);
+		PrecompiledHistories.Empty();
+		for (UNiagaraNodeOutput* FoundOutputNode : OutputNodes)
+		{
+			// Map all for this output node
+			FNiagaraParameterMapHistoryBuilder Builder;
+			Builder.BeginTranslation(InEmitter);
+			Builder.EnableScriptWhitelist(true, FoundOutputNode->GetUsage());
+			Builder.BuildParameterMaps(FoundOutputNode, true);
+			TArray<FNiagaraParameterMapHistory> Histories = Builder.Histories;
+			ensure(Histories.Num() <= 1);
+			PrecompiledHistories.Append(Histories);
+			Builder.EndTranslation(InEmitter);
+		}
+	}
+}
+
+
+bool UNiagaraScriptSource::GatherPreCompiledVariables(const FString& InNamespaceFilter, TArray<FNiagaraVariable>& OutVars)
+{
+	if (!bIsPrecompiled || PrecompiledHistories.Num() == 0)
+	{
+		return false;
+	}
+
+	for (const FNiagaraParameterMapHistory& History : PrecompiledHistories)
+	{
+		for (const FNiagaraVariable& Var : History.Variables)
+		{
+			if (FNiagaraParameterMapHistory::IsInNamespace(Var, InNamespaceFilter))
+			{
+				FNiagaraVariable NewVar = Var;
+				if (NewVar.IsDataAllocated() == false && !Var.IsDataInterface())
+				{
+					FNiagaraEditorUtilities::ResetVariableToDefaultValue(NewVar);
+				}
+				OutVars.AddUnique(NewVar);
+			}
+		}
+	}
+	return true;
+}
+
+void UNiagaraScriptSource::PostCompile() 
+{
+	bIsPrecompiled = false;
+	PrecompiledHistories.Empty();
+	NodeGraphDeepCopy = nullptr;
+}
+
+
+ENiagaraScriptCompileStatus UNiagaraScriptSource::Compile(UNiagaraScript* ScriptOwner, FString& OutGraphLevelErrorMessages)
+{
+	bool bDoPostCompile = false;
+	if (!bIsPrecompiled)
+	{
+		PreCompile(nullptr);
+		bDoPostCompile = true;
+	}
+
+	FNiagaraEditorModule& NiagaraEditorModule = FModuleManager::Get().LoadModuleChecked<FNiagaraEditorModule>(TEXT("NiagaraEditor"));
+	ENiagaraScriptCompileStatus Status = NiagaraEditorModule.CompileScript(ScriptOwner, OutGraphLevelErrorMessages);
+	check(ScriptOwner != nullptr && IsSynchronized(ScriptOwner->GetChangeID()));
+	
+	if (bDoPostCompile)
+	{
+		PostCompile();
+	}
+	return Status;
+
+// 	FNiagaraConstants& ExternalConsts = ScriptOwner->ConstantData.GetExternalConstants();
+// 
+// 	//Build the constant list. 
+// 	//This is mainly just jumping through some hoops for the custom UI. Should be removed and have the UI just read directly from the constants stored in the UScript.
+// 	const UEdGraphSchema_Niagara* Schema = CastChecked<UEdGraphSchema_Niagara>(NodeGraph->GetSchema());
+// 	ExposedVectorConstants.Empty();
+// 	for (int32 ConstIdx = 0; ConstIdx < ExternalConsts.GetNumVectorConstants(); ConstIdx++)
+// 	{
+// 		FNiagaraVariableInfo Info;
+// 		FVector4 Value;
+// 		ExternalConsts.GetVectorConstant(ConstIdx, Value, Info);
+// 		if (Schema->IsSystemConstant(Info))
+// 		{
+// 			continue;//System constants are "external" but should not be exposed to the editor.
+// 		}
+// 			
+// 		EditorExposedVectorConstant *Const = new EditorExposedVectorConstant();
+// 		Const->ConstName = Info.Name;
+// 		Const->Value = Value;
+// 		ExposedVectorConstants.Add(MakeShareable(Const));
+// 	}
+
+}
+
+FGuid UNiagaraScriptSource::GetChangeID() 
+{ 
+	return NodeGraph->GetChangeID(); 
+}
diff --git a/Engine/Plugins/FX/Niagara/Source/NiagaraEditor/Private/NiagaraStackEditorData.cpp b/Engine/Plugins/FX/Niagara/Source/NiagaraEditor/Private/NiagaraStackEditorData.cpp
new file mode 100644
--- /dev/null
+++ b/Engine/Plugins/FX/Niagara/Source/NiagaraEditor/Private/NiagaraStackEditorData.cpp
@@ -0,0 +1,36 @@
+// Copyright 1998-2017 Epic Games, Inc. All Rights Reserved.
+
+#include "NiagaraStackEditorData.h"
+
+bool UNiagaraStackEditorData::GetModuleInputIsPinned(const FString& ModuleInputKey) const
+{
+	const bool* bIsPinnedPtr = ModuleInputKeyToPinnedMap.Find(ModuleInputKey);
+	return bIsPinnedPtr != nullptr && *bIsPinnedPtr;
+}
+
+void UNiagaraStackEditorData::SetModuleInputIsPinned(const FString& ModuleInputKey, bool bIsPinned)
+{
+	ModuleInputKeyToPinnedMap.FindOrAdd(ModuleInputKey) = bIsPinned;
+}
+
+bool UNiagaraStackEditorData::GetModuleInputIsRenamePending(const FString& ModuleInputKey) const
+{
+	const bool* bIsRenamePendingPtr = ModuleInputKeyToRenamePendingMap.Find(ModuleInputKey);
+	return bIsRenamePendingPtr != nullptr && *bIsRenamePendingPtr;
+}
+
+void UNiagaraStackEditorData::SetModuleInputIsRenamePending(const FString& ModuleInputKey, bool bIsRenamePending)
+{
+	ModuleInputKeyToRenamePendingMap.FindOrAdd(ModuleInputKey) = bIsRenamePending;
+}
+
+bool UNiagaraStackEditorData::GetStackEntryIsExpanded(const FString& StackEntryKey, bool bIsExpandedDefault) const
+{
+	const bool* bIsExpandedPtr = StackEntryKeyToExpandedMap.Find(StackEntryKey);
+	return bIsExpandedPtr != nullptr ? *bIsExpandedPtr : bIsExpandedDefault;
+}
+
+void UNiagaraStackEditorData::SetStackEntryIsExpanded(const FString& StackEntryKey, bool bIsExpanded)
+{
+	StackEntryKeyToExpandedMap.FindOrAdd(StackEntryKey) = bIsExpanded;
+}
\ No newline at end of file
diff --git a/Engine/Plugins/FX/Niagara/Source/NiagaraEditor/Private/NiagaraSystemFactoryNew.cpp b/Engine/Plugins/FX/Niagara/Source/NiagaraEditor/Private/NiagaraSystemFactoryNew.cpp
new file mode 100644
--- /dev/null
+++ b/Engine/Plugins/FX/Niagara/Source/NiagaraEditor/Private/NiagaraSystemFactoryNew.cpp
@@ -0,0 +1,62 @@
+// Copyright 1998-2017 Epic Games, Inc. All Rights Reserved.
+
+#include "NiagaraSystemFactoryNew.h"
+#include "CoreMinimal.h"
+#include "Misc/ConfigCacheIni.h"
+#include "NiagaraSystem.h"
+#include "NiagaraScriptSource.h"
+#include "NiagaraGraph.h"
+#include "NiagaraEditorSettings.h"
+
+#define LOCTEXT_NAMESPACE "NiagaraSystemFactory"
+
+UNiagaraSystemFactoryNew::UNiagaraSystemFactoryNew(const FObjectInitializer& ObjectInitializer)
+: Super(ObjectInitializer)
+{
+
+	SupportedClass = UNiagaraSystem::StaticClass();
+	bCreateNew = false;
+	bEditAfterNew = true;
+	bCreateNew = true;
+}
+
+UObject* UNiagaraSystemFactoryNew::FactoryCreateNew(UClass* Class, UObject* InParent, FName Name, EObjectFlags Flags, UObject* Context, FFeedbackContext* Warn)
+{
+	check(Class->IsChildOf(UNiagaraSystem::StaticClass()));
+
+	const UNiagaraEditorSettings* Settings = GetDefault<UNiagaraEditorSettings>();
+	check(Settings);
+
+	UNiagaraSystem* NewSystem;
+	
+	if (UNiagaraSystem* Default = Cast<UNiagaraSystem>(Settings->DefaultSystem.TryLoad()))
+	{
+		NewSystem = Cast<UNiagaraSystem>(StaticDuplicateObject(Default, InParent, Name, Flags, Class));
+	}
+	else
+	{
+		NewSystem = NewObject<UNiagaraSystem>(InParent, Class, Name, Flags | RF_Transactional);
+	}
+
+	InitializeSystem(NewSystem);
+
+	return NewSystem;
+}
+
+void UNiagaraSystemFactoryNew::InitializeSystem(UNiagaraSystem* System)
+{
+	UNiagaraScript* SystemSpawnScript = System->GetSystemSpawnScript();
+	UNiagaraScript* SystemUpdateScript = System->GetSystemUpdateScript();
+	UNiagaraScript* SystemSpawnScriptSolo = System->GetSystemSpawnScript(true);
+	UNiagaraScript* SystemUpdateScriptSolo = System->GetSystemUpdateScript(true);
+
+	UNiagaraScriptSource* SystemScriptSource = NewObject<UNiagaraScriptSource>(SystemSpawnScript, "SystemScriptSource", RF_Transactional);
+	SystemScriptSource->NodeGraph = NewObject<UNiagaraGraph>(SystemScriptSource, "SystemScriptGraph", RF_Transactional);
+
+	SystemSpawnScript->SetSource(SystemScriptSource);
+	SystemUpdateScript->SetSource(SystemScriptSource);
+	SystemSpawnScriptSolo->SetSource(SystemScriptSource);
+	SystemUpdateScriptSolo->SetSource(SystemScriptSource);
+}
+
+#undef LOCTEXT_NAMESPACE
\ No newline at end of file
diff --git a/Engine/Plugins/FX/Niagara/Source/NiagaraEditor/Private/Sequencer/NiagaraSequence/MovieSceneNiagaraEmitterSection.cpp b/Engine/Plugins/FX/Niagara/Source/NiagaraEditor/Private/Sequencer/NiagaraSequence/MovieSceneNiagaraEmitterSection.cpp
new file mode 100644
--- /dev/null
+++ b/Engine/Plugins/FX/Niagara/Source/NiagaraEditor/Private/Sequencer/NiagaraSequence/MovieSceneNiagaraEmitterSection.cpp
@@ -0,0 +1,66 @@
+// Copyright 1998-2016 Epic Games, Inc. All Rights Reserved.
+
+#include "MovieSceneNiagaraEmitterSection.h"
+#include "NiagaraEmitter.h"
+#include "NiagaraEmitterHandleViewModel.h"
+#include "NiagaraEmitterViewModel.h"
+
+#define LOCTEXT_NAMESPACE "MovieSceneNiagaraEmitterSection"
+
+UMovieSceneNiagaraEmitterSection::UMovieSceneNiagaraEmitterSection(const FObjectInitializer& ObjectInitializer)
+	: Super(ObjectInitializer)
+{
+	BurstCurve = MakeShareable(new FBurstCurve(&Times, &BurstKeys));
+}
+
+TOptional<float> UMovieSceneNiagaraEmitterSection::GetKeyTime(FKeyHandle KeyHandle) const
+{
+	TOptional<float> Time;
+	if (BurstCurve.IsValid())
+	{
+		auto Key = BurstCurve->GetKey(KeyHandle);
+		Time = Key.IsSet()
+			? TOptional<float>(Key->Time)
+			: TOptional<float>();
+	}
+	return Time;
+}
+
+void UMovieSceneNiagaraEmitterSection::SetKeyTime(FKeyHandle KeyHandle, float Time)
+{
+	if (BurstCurve.IsValid())
+	{
+		BurstCurve->SetKeyTime(KeyHandle, Time);
+	}
+}
+
+void UMovieSceneNiagaraEmitterSection::GetKeyHandles(TSet<FKeyHandle>& OutKeyHandles, TRange<float> TimeRange) const
+{
+	if (BurstCurve.IsValid())
+	{
+		for (TKeyTimeIterator<float> KeyIterator = BurstCurve->IterateKeys(); KeyIterator; ++KeyIterator)
+		{
+			if (TimeRange.Contains(*KeyIterator))
+			{
+				OutKeyHandles.Add(KeyIterator.GetKeyHandle());
+			}
+		}
+	}
+}
+
+TSharedPtr<FNiagaraEmitterHandleViewModel> UMovieSceneNiagaraEmitterSection::GetEmitterHandle()
+{
+	return EmitterHandleViewModel;
+}
+
+void UMovieSceneNiagaraEmitterSection::SetEmitterHandle(TSharedRef<FNiagaraEmitterHandleViewModel> InEmitterHandleViewModel)
+{
+	EmitterHandleViewModel = InEmitterHandleViewModel;
+}
+
+TSharedPtr<FBurstCurve> UMovieSceneNiagaraEmitterSection::GetBurstCurve()
+{
+	return BurstCurve;
+}
+
+#undef LOCTEXT_NAMESPACE // MovieSceneNiagaraEmitterSection
\ No newline at end of file
