Copying //UE4/Orion-Staging to //UE4/Main (Source: //Orion/Dev-General @ 3606378)

diff --git a/Engine/Plugins/FX/Niagara/Source/NiagaraEditor/Private/Sequencer/NiagaraSequence/MovieSceneNiagaraEmitterSection.h b/Engine/Plugins/FX/Niagara/Source/NiagaraEditor/Private/Sequencer/NiagaraSequence/MovieSceneNiagaraEmitterSection.h
new file mode 100644
--- /dev/null
+++ b/Engine/Plugins/FX/Niagara/Source/NiagaraEditor/Private/Sequencer/NiagaraSequence/MovieSceneNiagaraEmitterSection.h
@@ -0,0 +1,86 @@
+// Copyright 1998-2016 Epic Games, Inc. All Rights Reserved.
+
+#pragma once
+
+#include "NiagaraEmitter.h"
+#include "MovieSceneSection.h"
+#include "MovieSceneKeyStruct.h"
+#include "Curves/CurveInterface.h"
+#include "MovieSceneClipboard.h"
+#include "MovieSceneNiagaraEmitterSection.generated.h"
+
+class FNiagaraEmitterHandleViewModel;
+
+/** Defines data for burst keys in this emitter section. */
+USTRUCT()
+struct FMovieSceneBurstKey
+{
+	GENERATED_BODY()
+
+	FMovieSceneBurstKey()
+		: TimeRange(0)
+		, SpawnMinimum(0)
+		, SpawnMaximum(0)
+	{
+	}
+
+	/** The time range which will be used around the key time which is used for randomly bursting. */
+	UPROPERTY(EditAnywhere, Category = "Burst")
+	float TimeRange;
+
+	/** The minimum number of particles to spawn with this burst. */
+	UPROPERTY(EditAnywhere, Category = "Burst")
+	uint32 SpawnMinimum;
+
+	/** The maximum number of particles to span with this burst. */
+	UPROPERTY(EditAnywhere, Category = "Burst")
+	uint32 SpawnMaximum;
+};
+
+typedef TCurveInterface<FMovieSceneBurstKey, float> FBurstCurve;
+
+
+
+namespace MovieSceneClipboard
+{
+	template<> inline FName GetKeyTypeName<FMovieSceneBurstKey>()
+	{
+		return "FMovieSceneBurstKey";
+	}
+}
+
+/** 
+ *	Niagara editor movie scene section; represents one emitter in the timeline
+ */
+UCLASS()
+class UMovieSceneNiagaraEmitterSection : public UMovieSceneSection
+{
+	GENERATED_UCLASS_BODY()
+public:
+	//~ UMovieSceneSection interface
+	virtual TOptional<float> GetKeyTime(FKeyHandle KeyHandle) const override;
+	virtual void SetKeyTime(FKeyHandle KeyHandle, float Time) override;
+	virtual void GetKeyHandles(TSet<FKeyHandle>& OutKeyHandles, TRange<float> TimeRange) const override;
+
+	/** Gets the emitter handle for the emitter which this section represents. */
+	TSharedPtr<FNiagaraEmitterHandleViewModel> GetEmitterHandle();
+
+	/** Sets the emitter handle for the emitter which this section represents. */
+	void SetEmitterHandle(TSharedRef<FNiagaraEmitterHandleViewModel> InEmitterHandleViewModel);
+
+	/** Gets the curve containing the burst keyframes. */
+	TSharedPtr<FBurstCurve> GetBurstCurve();
+
+private:
+	/** The view model for the handle to the emitter this section represents. */
+	TSharedPtr<FNiagaraEmitterHandleViewModel> EmitterHandleViewModel;
+
+	/** Storage for the key times in the burst curve. */
+	TArray<float> Times;
+
+	/** Storage for the burst value times in the burst curve. */
+	TArray<FMovieSceneBurstKey> BurstKeys;
+
+	/** A curve containing the burst keyframes. */
+	TSharedPtr<FBurstCurve> BurstCurve;
+};
\ No newline at end of file
diff --git a/Engine/Plugins/FX/Niagara/Source/NiagaraEditor/Private/Sequencer/NiagaraSequence/MovieSceneNiagaraEmitterTrack.cpp b/Engine/Plugins/FX/Niagara/Source/NiagaraEditor/Private/Sequencer/NiagaraSequence/MovieSceneNiagaraEmitterTrack.cpp
new file mode 100644
--- /dev/null
+++ b/Engine/Plugins/FX/Niagara/Source/NiagaraEditor/Private/Sequencer/NiagaraSequence/MovieSceneNiagaraEmitterTrack.cpp
@@ -0,0 +1,63 @@
+// Copyright 1998-2016 Epic Games, Inc. All Rights Reserved.
+
+#include "MovieSceneNiagaraEmitterTrack.h"
+#include "NiagaraEmitterHandle.h"
+#include "NiagaraEmitterHandleViewModel.h"
+#include "NiagaraEmitterViewModel.h"
+#include "MovieSceneNiagaraEmitterSection.h"
+#include "MovieSceneNiagaraEmitterTrackInstance.h"
+
+UMovieSceneNiagaraEmitterTrack::UMovieSceneNiagaraEmitterTrack(const FObjectInitializer& ObjectInitializer)
+: Super(ObjectInitializer)
+{
+}
+
+TSharedPtr<FNiagaraEmitterHandleViewModel> UMovieSceneNiagaraEmitterTrack::GetEmitterHandle()
+{
+	return EmitterHandleViewModel;
+}
+
+void UMovieSceneNiagaraEmitterTrack::SetEmitterHandle(TSharedRef<FNiagaraEmitterHandleViewModel> InEmitterHandleViewModel)
+{
+	EmitterHandleViewModel = InEmitterHandleViewModel;
+
+	if (Sections.Num() == 0)
+	{
+		Sections.Add(NewObject<UMovieSceneNiagaraEmitterSection>(this));
+	}
+}
+
+bool UMovieSceneNiagaraEmitterTrack::HasSection(const UMovieSceneSection& Section) const
+{
+	return Sections.Contains(&Section);
+}
+
+void UMovieSceneNiagaraEmitterTrack::RemoveSection(UMovieSceneSection& Section)
+{
+	Sections.Remove(&Section);
+}
+
+bool UMovieSceneNiagaraEmitterTrack::IsEmpty() const
+{
+	return Sections.Num() == 0;
+}
+
+const TArray<UMovieSceneSection*>& UMovieSceneNiagaraEmitterTrack::GetAllSections() const
+{
+	return Sections;
+}
+
+TRange<float> UMovieSceneNiagaraEmitterTrack::GetSectionBoundaries() const
+{
+	return Sections.Num() == 1
+		? Sections[0]->GetRange()
+		: TRange<float>(0, 0);
+}
+
+FName UMovieSceneNiagaraEmitterTrack::GetTrackName() const
+{
+	return EmitterHandleViewModel.IsValid()
+		? EmitterHandleViewModel->GetName()
+		: NAME_None;
+}
+
diff --git a/Engine/Plugins/FX/Niagara/Source/NiagaraEditor/Private/Sequencer/NiagaraSequence/MovieSceneNiagaraEmitterTrack.h b/Engine/Plugins/FX/Niagara/Source/NiagaraEditor/Private/Sequencer/NiagaraSequence/MovieSceneNiagaraEmitterTrack.h
new file mode 100644
--- /dev/null
+++ b/Engine/Plugins/FX/Niagara/Source/NiagaraEditor/Private/Sequencer/NiagaraSequence/MovieSceneNiagaraEmitterTrack.h
@@ -0,0 +1,46 @@
+// Copyright 1998-2016 Epic Games, Inc. All Rights Reserved.
+
+#pragma once
+
+#include "MovieSceneNameableTrack.h"
+#include "MovieSceneNiagaraEmitterTrack.generated.h"
+
+
+class FNiagaraEmitterHandleViewModel;
+
+/**
+*	Single track containing exactly one UNiagaraMovieSceneSection, representing one emitter
+*/
+UCLASS(MinimalAPI)
+class UMovieSceneNiagaraEmitterTrack
+	: public UMovieSceneNameableTrack
+{
+	GENERATED_UCLASS_BODY()
+public:
+
+	TSharedPtr<FNiagaraEmitterHandleViewModel> GetEmitterHandle();
+
+	void SetEmitterHandle(TSharedRef<FNiagaraEmitterHandleViewModel> InEmitterHandleViewModel);
+
+public:
+
+	//~ UMovieSceneTrack interface
+	virtual void RemoveAllAnimationData() override { }
+	virtual bool HasSection(const UMovieSceneSection& Section) const override;
+	virtual void RemoveSection(UMovieSceneSection& Section) override;
+	virtual bool IsEmpty() const override;
+	virtual const TArray<UMovieSceneSection*>& GetAllSections() const override;
+	virtual TRange<float> GetSectionBoundaries() const override;
+	virtual FName GetTrackName() const override;
+
+private:
+	void EmitterPropertyChanged();
+	void EmitterHandlePropertyChanged();
+	void SynchronizeWithEmitterHandle();
+
+private:
+	TSharedPtr<FNiagaraEmitterHandleViewModel> EmitterHandleViewModel;
+
+	UPROPERTY()
+	TArray<UMovieSceneSection*> Sections;
+};
\ No newline at end of file
diff --git a/Engine/Plugins/FX/Niagara/Source/NiagaraEditor/Private/Sequencer/NiagaraSequence/MovieSceneNiagaraEmitterTrackInstance.cpp b/Engine/Plugins/FX/Niagara/Source/NiagaraEditor/Private/Sequencer/NiagaraSequence/MovieSceneNiagaraEmitterTrackInstance.cpp
new file mode 100644
--- /dev/null
+++ b/Engine/Plugins/FX/Niagara/Source/NiagaraEditor/Private/Sequencer/NiagaraSequence/MovieSceneNiagaraEmitterTrackInstance.cpp
@@ -0,0 +1,3 @@
+// Copyright 1998-2016 Epic Games, Inc. All Rights Reserved.
+
+#include "MovieSceneNiagaraEmitterTrackInstance.h"
\ No newline at end of file
diff --git a/Engine/Plugins/FX/Niagara/Source/NiagaraEditor/Private/Sequencer/NiagaraSequence/MovieSceneNiagaraEmitterTrackInstance.h b/Engine/Plugins/FX/Niagara/Source/NiagaraEditor/Private/Sequencer/NiagaraSequence/MovieSceneNiagaraEmitterTrackInstance.h
new file mode 100644
--- /dev/null
+++ b/Engine/Plugins/FX/Niagara/Source/NiagaraEditor/Private/Sequencer/NiagaraSequence/MovieSceneNiagaraEmitterTrackInstance.h
@@ -0,0 +1,34 @@
+// Copyright 1998-2016 Epic Games, Inc. All Rights Reserved.
+
+#pragma once
+
+#include "IMovieSceneTrackInstance.h"
+
+class UMovieSceneNiagaraEmitterTrack;
+struct FNiagaraEmitterInstance;
+
+/**
+*	One instance of a UEmitterMovieSceneTrack
+*/
+class FMovieSceneNiagaraEmitterTrackInstance
+	: public IMovieSceneTrackInstance
+{
+public:
+	FMovieSceneNiagaraEmitterTrackInstance(UMovieSceneNiagaraEmitterTrack* InTrack)
+		: Track(InTrack)
+	{
+	}
+
+	virtual void Update(EMovieSceneUpdateData& UpdateData, const TArray<TWeakObjectPtr<UObject>>& RuntimeObjects, IMovieScenePlayer& Player, FMovieSceneSequenceInstance& SequenceInstance)
+	{
+	}
+
+	virtual void RefreshInstance(const TArray<TWeakObjectPtr<UObject>>& RuntimeObjects, IMovieScenePlayer& Player, FMovieSceneSequenceInstance& SequenceInstance) override {}
+	virtual void ClearInstance(IMovieScenePlayer& Player, FMovieSceneSequenceInstance& SequenceInstance) override {}
+	virtual void SaveState(const TArray<TWeakObjectPtr<UObject>>& RuntimeObjects, IMovieScenePlayer& Player, FMovieSceneSequenceInstance& SequenceInstance) override {}
+	virtual void RestoreState(const TArray<TWeakObjectPtr<UObject>>& RuntimeObjects, IMovieScenePlayer& Player, FMovieSceneSequenceInstance& SequenceInstance) override {}
+
+private:
+	TSharedPtr<FNiagaraEmitterInstance> Emitter;
+	UMovieSceneNiagaraEmitterTrack *Track;
+};
\ No newline at end of file
diff --git a/Engine/Plugins/FX/Niagara/Source/NiagaraEditor/Private/Sequencer/NiagaraSequence/NiagaraEmitterSection.cpp b/Engine/Plugins/FX/Niagara/Source/NiagaraEditor/Private/Sequencer/NiagaraSequence/NiagaraEmitterSection.cpp
new file mode 100644
--- /dev/null
+++ b/Engine/Plugins/FX/Niagara/Source/NiagaraEditor/Private/Sequencer/NiagaraSequence/NiagaraEmitterSection.cpp
@@ -0,0 +1,74 @@
+// Copyright 1998-2016 Epic Games, Inc. All Rights Reserved.
+
+#include "NiagaraEmitterSection.h"
+#include "MovieSceneNiagaraEmitterSection.h"
+#include "NiagaraEmitterHandleViewModel.h"
+#include "NiagaraEmitterViewModel.h"
+
+#include "SequencerSectionPainter.h"
+#include "GenericKeyArea.h"
+#include "ISectionLayoutBuilder.h"
+
+#define LOCTEXT_NAMESPACE "NiagaraEmitterSection"
+
+FNiagaraEmitterSection::FNiagaraEmitterSection(UMovieSceneSection &SectionObject)
+{
+	UMovieSceneNiagaraEmitterSection *NiagaraSectionObject = Cast<UMovieSceneNiagaraEmitterSection>(&SectionObject);
+	check(NiagaraSectionObject);
+
+	EmitterSection = NiagaraSectionObject;
+}
+
+UMovieSceneSection* FNiagaraEmitterSection::GetSectionObject(void)
+{
+	return EmitterSection;
+}
+
+int32 FNiagaraEmitterSection::OnPaintSection(FSequencerSectionPainter& InPainter) const
+{
+	// draw the first run of the emitter
+	FSlateDrawElement::MakeBox
+	(
+		InPainter.DrawElements,
+		InPainter.LayerId,
+		InPainter.SectionGeometry.ToPaintGeometry(),
+		FEditorStyle::GetBrush("CurveEd.TimelineArea"),
+		ESlateDrawEffect::None,
+		FLinearColor(0.3f, 0.3f, 0.6f)
+	);
+
+	// draw all loops of the emitter as 'ghosts' of the original section
+	float X = InPainter.SectionGeometry.AbsolutePosition.X;
+	float GeomW = InPainter.SectionGeometry.GetDrawSize().X;
+	int32 NumLoops = EmitterSection->GetEmitterHandle()->GetEmitterViewModel()->GetNumLoops() - 1;
+	for (int32 Loop = 0; Loop < NumLoops; Loop++)
+	{
+		FSlateDrawElement::MakeBox
+		(
+			InPainter.DrawElements,
+			InPainter.LayerId,
+			InPainter.SectionGeometry.ToPaintGeometry(FVector2D(GeomW*(Loop + 1), 0.0f), InPainter.SectionGeometry.GetDrawSize(), 1.0f),
+			FEditorStyle::GetBrush("CurveEd.TimelineArea"),
+			ESlateDrawEffect::None,
+			FLinearColor(0.3f, 0.3f, 0.6f, 0.25f)
+		);
+	}
+
+	return InPainter.LayerId;
+}
+
+FText FNiagaraEmitterSection::GetSectionTitle(void) const
+{
+	return EmitterSection->GetEmitterHandle()->GetNameText();
+}
+
+void FNiagaraEmitterSection::GenerateSectionLayout(ISectionLayoutBuilder &LayoutBuilder) const
+{
+	if (EmitterSection->GetBurstCurve().IsValid())
+	{
+		auto KeyArea = MakeShared<TGenericKeyArea<FMovieSceneBurstKey, float>>(MoveTemp(*EmitterSection->GetBurstCurve()), EmitterSection);
+		LayoutBuilder.SetSectionAsKeyArea(KeyArea);
+	}
+}
+
+#undef LOCTEXT_NAMESPACE
diff --git a/Engine/Plugins/FX/Niagara/Source/NiagaraEditor/Private/Sequencer/NiagaraSequence/NiagaraEmitterSection.h b/Engine/Plugins/FX/Niagara/Source/NiagaraEditor/Private/Sequencer/NiagaraSequence/NiagaraEmitterSection.h
new file mode 100644
--- /dev/null
+++ b/Engine/Plugins/FX/Niagara/Source/NiagaraEditor/Private/Sequencer/NiagaraSequence/NiagaraEmitterSection.h
@@ -0,0 +1,26 @@
+// Copyright 1998-2016 Epic Games, Inc. All Rights Reserved.
+
+#pragma once
+
+#include "ISequencerSection.h"
+
+class UMovieSceneNiagaraEmitterSection;
+
+/**
+*	Visual representation of UMovieSceneNiagaraEmitterSection
+*/
+class FNiagaraEmitterSection : public ISequencerSection
+{
+public:
+	FNiagaraEmitterSection(UMovieSceneSection &SectionObject);
+
+	//~ ISequencerSection interface
+	virtual UMovieSceneSection* GetSectionObject(void) override;
+	virtual int32 OnPaintSection(FSequencerSectionPainter& InPainter) const override;
+	virtual FText GetSectionTitle(void) const override;
+	virtual float GetSectionHeight() const override { return 20.0f; }
+	virtual void GenerateSectionLayout(ISectionLayoutBuilder &) const override;
+
+private:
+	UMovieSceneNiagaraEmitterSection* EmitterSection;
+};
\ No newline at end of file
diff --git a/Engine/Plugins/FX/Niagara/Source/NiagaraEditor/Private/TNiagaraGraphPinEditableName.h b/Engine/Plugins/FX/Niagara/Source/NiagaraEditor/Private/TNiagaraGraphPinEditableName.h
new file mode 100644
--- /dev/null
+++ b/Engine/Plugins/FX/Niagara/Source/NiagaraEditor/Private/TNiagaraGraphPinEditableName.h
@@ -0,0 +1,97 @@
+// Copyright 1998-2016 Epic Games, Inc. All Rights Reserved.
+
+#pragma once
+
+#include "SGraphPin.h"
+#include "NiagaraNode.h"
+#include "SInlineEditableTextBlock.h"
+
+/** A graph pin widget for allowing a pin to have an editable name for a pin. */
+template< class BaseClass >
+class TNiagaraGraphPinEditableName : public BaseClass
+{
+public:
+	SLATE_BEGIN_ARGS(TNiagaraGraphPinEditableName<BaseClass>) {}
+	SLATE_END_ARGS()
+
+	FORCENOINLINE void Construct(const FArguments& InArgs, UEdGraphPin* InGraphPinObj)
+	{
+		bPendingRename = false;
+		BaseClass::Construct(BaseClass::FArguments(), InGraphPinObj);
+	}
+
+protected:
+	FText GetParentPinLabel() const
+	{
+		return BaseClass::GetPinLabel();
+	}
+
+	EVisibility GetParentPinVisibility() const
+	{
+		return BaseClass::GetPinLabelVisibility();
+	}
+
+	FSlateColor GetParentPinTextColor() const
+	{
+		return BaseClass::GetPinTextColor();
+	}
+
+	bool OnVerifyTextChanged(const FText& InName, FText& OutErrorMessage) const
+	{
+		UNiagaraNode* ParentNode = Cast<UNiagaraNode>(GraphPinObj->GetOwningNode());
+		if (ParentNode)
+		{
+			return ParentNode->VerifyEditablePinName(InName, OutErrorMessage, GraphPinObj);
+		}
+		return false;
+	}
+
+	void OnTextCommitted(const FText& InText, ETextCommit::Type InCommitType)
+	{
+		UNiagaraNode* ParentNode = Cast<UNiagaraNode>(GraphPinObj->GetOwningNode());
+		if (ParentNode != nullptr)
+		{
+			ParentNode->CommitEditablePinName(InText, GraphPinObj);
+		}
+	}
+
+	virtual void Tick(const FGeometry& AllottedGeometry, const double InCurrentTime, const float InDeltaTime) override
+	{
+		if (bPendingRename && CreatedTextBlock.IsValid())
+		{
+			CreatedTextBlock->EnterEditingMode();
+			bPendingRename = false;
+		}
+		BaseClass::Tick(AllottedGeometry, InCurrentTime, InDeltaTime);
+	}
+
+	virtual TSharedRef<SWidget> GetLabelWidget(const FName& InLabelStyle) override
+	{
+		UNiagaraNode* ParentNode = Cast<UNiagaraNode>(GraphPinObj->GetOwningNode());
+		if (ParentNode && ParentNode->IsPinNameEditable(GraphPinObj))
+		{
+			CreatedTextBlock = SNew(SInlineEditableTextBlock)
+				.Style(&FEditorStyle::Get().GetWidgetStyle<FInlineEditableTextBlockStyle>("Graph.Node.InlineEditablePinName"))
+				.Text(this, &TNiagaraGraphPinEditableName<BaseClass>::GetParentPinLabel)
+				.Visibility(this, &TNiagaraGraphPinEditableName<BaseClass>::GetParentPinVisibility)
+				.ColorAndOpacity(this, &TNiagaraGraphPinEditableName<BaseClass>::GetParentPinTextColor)
+				.OnVerifyTextChanged(this, &TNiagaraGraphPinEditableName<BaseClass>::OnVerifyTextChanged)
+				.OnTextCommitted(this, &TNiagaraGraphPinEditableName<BaseClass>::OnTextCommitted);
+
+			TSharedRef<SWidget> EditableTextBlock = CreatedTextBlock.ToSharedRef();
+
+			if (ParentNode->IsPinNameEditableUponCreation(GraphPinObj))
+			{
+				bPendingRename = true;
+			}
+			return EditableTextBlock;
+		}
+		else
+		{
+			return BaseClass::GetLabelWidget(InLabelStyle);
+		}
+	}
+
+	bool bPendingRename;
+	TSharedPtr<SInlineEditableTextBlock> CreatedTextBlock;
+};
\ No newline at end of file
diff --git a/Engine/Plugins/FX/Niagara/Source/NiagaraEditor/Private/Toolkits/NiagaraScriptToolkit.cpp b/Engine/Plugins/FX/Niagara/Source/NiagaraEditor/Private/Toolkits/NiagaraScriptToolkit.cpp
new file mode 100644
--- /dev/null
+++ b/Engine/Plugins/FX/Niagara/Source/NiagaraEditor/Private/Toolkits/NiagaraScriptToolkit.cpp
@@ -0,0 +1,552 @@
+// Copyright 1998-2017 Epic Games, Inc. All Rights Reserved.
+
+#include "NiagaraScriptToolkit.h"
+#include "NiagaraEditorModule.h"
+#include "NiagaraEditorUtilities.h"
+#include "NiagaraScript.h"
+#include "NiagaraScriptViewModel.h"
+#include "NiagaraScriptGraphViewModel.h"
+#include "NiagaraObjectSelection.h"
+#include "NiagaraEditorCommands.h"
+#include "NiagaraScriptInputCollectionViewModel.h"
+#include "NiagaraScriptOutputCollectionViewModel.h"
+#include "SNiagaraScriptGraph.h"
+#include "SNiagaraSelectedObjectsDetails.h"
+#include "SNiagaraParameterCollection.h"
+#include "Package.h"
+#include "NiagaraEditorStyle.h"
+#include "BusyCursor.h"
+#include "EngineGlobals.h"
+#include "FeedbackContext.h"
+#include "Editor.h"
+#include "Engine/Selection.h"
+#include "UObjectIterator.h"
+#include "NiagaraScriptSource.h"
+#include "NiagaraNode.h"
+#include "NiagaraGraph.h"
+#include "NiagaraEmitter.h"
+#include "NiagaraNodeOutput.h"
+#include "NiagaraModule.h"
+
+#include "ModuleManager.h"
+
+#include "MultiBoxBuilder.h"
+#include "SScrollBox.h"
+#include "SDockTab.h"
+#include "Misc/MessageDialog.h"
+#include "Customizations/NiagaraScriptDetails.h"
+#include "PropertyEditorModule.h"
+
+#define LOCTEXT_NAMESPACE "NiagaraScriptToolkit"
+
+static TAutoConsoleVariable<int32> CVarDevDetails(
+	TEXT("fx.DevDetailsPanels"),
+	0,
+	TEXT("Whether to enable the development details panels inside Niagara."));
+
+
+const FName FNiagaraScriptToolkit::NodeGraphTabId(TEXT("NiagaraEditor_NodeGraph")); 
+const FName FNiagaraScriptToolkit::DetailsTabId(TEXT("NiagaraEditor_Details"));
+const FName FNiagaraScriptToolkit::ParametersTabId(TEXT("NiagaraEditor_Parameters"));
+
+FNiagaraScriptToolkit::FNiagaraScriptToolkit()
+{
+}
+
+FNiagaraScriptToolkit::~FNiagaraScriptToolkit()
+{
+}
+
+
+void FNiagaraScriptToolkit::RegisterTabSpawners(const TSharedRef<class FTabManager>& InTabManager)
+{
+	WorkspaceMenuCategory = InTabManager->AddLocalWorkspaceMenuCategory(LOCTEXT("WorkspaceMenu_NiagaraEditor", "Niagara"));
+
+	FAssetEditorToolkit::RegisterTabSpawners(InTabManager);
+	
+	TSharedRef<FWorkspaceItem> WorkspaceMenuCategoryRef = WorkspaceMenuCategory.ToSharedRef();
+
+	InTabManager->RegisterTabSpawner(NodeGraphTabId, FOnSpawnTab::CreateSP(this, &FNiagaraScriptToolkit::SpawnTabNodeGraph))
+		.SetDisplayName( LOCTEXT("NodeGraph", "Node Graph") )
+		.SetGroup(WorkspaceMenuCategoryRef); 
+
+	InTabManager->RegisterTabSpawner(DetailsTabId, FOnSpawnTab::CreateSP(this, &FNiagaraScriptToolkit::SpawnTabNodeDetails))
+		.SetDisplayName(LOCTEXT("DetailsTab", "Details"))
+		.SetGroup(WorkspaceMenuCategoryRef)
+		.SetIcon(FSlateIcon(FEditorStyle::GetStyleSetName(), "LevelEditor.Tabs.Details"));
+}
+
+void FNiagaraScriptToolkit::UnregisterTabSpawners(const TSharedRef<class FTabManager>& InTabManager)
+{
+	FAssetEditorToolkit::UnregisterTabSpawners(InTabManager);
+
+	InTabManager->UnregisterTabSpawner(NodeGraphTabId);
+	InTabManager->UnregisterTabSpawner(DetailsTabId);
+}
+
+void FNiagaraScriptToolkit::Initialize( const EToolkitMode::Type Mode, const TSharedPtr< class IToolkitHost >& InitToolkitHost, UNiagaraScript* InputScript )
+{	
+	check(InputScript != NULL);
+	OriginalNiagaraScript = InputScript;
+	ResetLoaders(GetTransientPackage()); // Make sure that we're not going to get invalid version number linkers into the package we are going into. 
+	GetTransientPackage()->LinkerCustomVersion.Empty();
+	EditedNiagaraScript = (UNiagaraScript*)StaticDuplicateObject(OriginalNiagaraScript, GetTransientPackage(), NAME_None, ~RF_Standalone, UNiagaraScript::StaticClass());
+
+	ScriptViewModel = MakeShareable(new FNiagaraScriptViewModel(EditedNiagaraScript, LOCTEXT("NiagaraScriptDisplayName", "Niagara Script"), ENiagaraParameterEditMode::EditAll));
+	DetailsSelection = MakeShareable(new FNiagaraObjectSelection());
+	DetailsSelection->SetSelectedObject(EditedNiagaraScript);
+		
+	TSharedRef<FTabManager::FLayout> StandaloneDefaultLayout = FTabManager::NewLayout( "Standalone_Niagara_Layout_v6" )
+	->AddArea
+	(
+		FTabManager::NewPrimaryArea() ->SetOrientation(Orient_Vertical)
+		->Split
+		(
+			FTabManager::NewStack()
+			->SetSizeCoefficient(0.1f)
+			->AddTab(GetToolbarTabId(), ETabState::OpenedTab) 
+			->SetHideTabWell( true )
+		)
+		->Split
+		(
+			FTabManager::NewSplitter()->SetOrientation(Orient_Horizontal)->SetSizeCoefficient(0.9f)
+			->Split
+			(
+				FTabManager::NewStack()
+				->SetSizeCoefficient(0.2f)
+				->AddTab(DetailsTabId, ETabState::OpenedTab)
+				->SetForegroundTab(DetailsTabId)
+			)
+			->Split
+			(
+				FTabManager::NewStack()
+				->SetSizeCoefficient(0.8f)
+				->AddTab(NodeGraphTabId, ETabState::OpenedTab)
+			)
+		)
+	);
+
+	const bool bCreateDefaultStandaloneMenu = true;
+	const bool bCreateDefaultToolbar = true;
+	FAssetEditorToolkit::InitAssetEditor( Mode, InitToolkitHost, FNiagaraEditorModule::NiagaraEditorAppIdentifier, StandaloneDefaultLayout, bCreateDefaultStandaloneMenu, bCreateDefaultToolbar, InputScript );
+
+	FNiagaraEditorModule& NiagaraEditorModule = FModuleManager::LoadModuleChecked<FNiagaraEditorModule>( "NiagaraEditor" );
+	AddMenuExtender(NiagaraEditorModule.GetMenuExtensibilityManager()->GetAllExtenders(GetToolkitCommands(), GetEditingObjects()));
+
+	SetupCommands();
+	ExtendToolbar();
+	RegenerateMenusAndToolbars();
+
+	// @todo toolkit world centric editing
+	/*// Setup our tool's layout
+	if( IsWorldCentricAssetEditor() )
+	{
+		const FString TabInitializationPayload(TEXT(""));		// NOTE: Payload not currently used for table properties
+		SpawnToolkitTab( NodeGraphTabId, TabInitializationPayload, EToolkitTabSpot::Details );
+	}*/
+}
+
+FName FNiagaraScriptToolkit::GetToolkitFName() const
+{
+	return FName("Niagara");
+}
+
+FText FNiagaraScriptToolkit::GetBaseToolkitName() const
+{
+	return LOCTEXT("AppLabel", "Niagara");
+}
+
+FString FNiagaraScriptToolkit::GetWorldCentricTabPrefix() const
+{
+	return LOCTEXT("WorldCentricTabPrefix", "Niagara ").ToString();
+}
+
+
+FLinearColor FNiagaraScriptToolkit::GetWorldCentricTabColorScale() const
+{
+	return FNiagaraEditorModule::WorldCentricTabColorScale;
+}
+
+TSharedRef<SDockTab> FNiagaraScriptToolkit::SpawnTabNodeGraph( const FSpawnTabArgs& Args )
+{
+	checkf(Args.GetTabId().TabType == NodeGraphTabId, TEXT("Wrong tab ID in NiagaraScriptToolkit"));
+	checkf(ScriptViewModel.IsValid(), TEXT("NiagaraScriptToolkit - Script editor view model is invalid"));
+
+	return
+		SNew(SDockTab)
+		[
+			SNew(SNiagaraScriptGraph, ScriptViewModel->GetGraphViewModel())
+			.GraphTitle(LOCTEXT("SpawnGraphTitle", "Script"))
+		];
+}
+
+void FNiagaraScriptToolkit::OnDetailsSelectionPropertyFinishedChanging(const FPropertyChangedEvent& InEvent)
+{
+	// We need to synchronize the Usage field in the property editor with the actual node
+	// in the graph.
+	if (InEvent.Property != nullptr && InEvent.Property->GetName().Equals(TEXT("Usage")))
+	{
+		if (EditedNiagaraScript != nullptr && EditedNiagaraScript->GetSource() != nullptr)
+		{
+			UNiagaraScriptSource* ScriptSource = Cast<UNiagaraScriptSource>(EditedNiagaraScript->GetSource());
+			if (ScriptSource != nullptr)
+			{
+				TArray<UNiagaraNodeOutput*> OutputNodes;
+				ScriptSource->NodeGraph->FindOutputNodes(OutputNodes);
+
+				bool bChanged = false;
+				for (UNiagaraNodeOutput* Output : OutputNodes)
+				{
+					if (Output->GetUsage() != EditedNiagaraScript->GetUsage())
+					{
+						Output->Modify();
+						Output->SetUsage(EditedNiagaraScript->GetUsage());
+						bChanged = true;
+					}
+				}
+
+				if (bChanged)
+				{
+					ScriptSource->NodeGraph->NotifyGraphChanged();
+				}
+			}
+		}
+	}
+}
+
+TSharedRef<SDockTab> FNiagaraScriptToolkit::SpawnTabNodeDetails(const FSpawnTabArgs& Args)
+{
+	checkf(Args.GetTabId().TabType == DetailsTabId, TEXT("Wrong tab ID in NiagaraScriptToolkit"));
+	checkf(ScriptViewModel.IsValid(), TEXT("NiagaraScriptToolkit - Script editor view model is invalid"));
+
+	IConsoleVariable* DevDetailsCVar = IConsoleManager::Get().FindConsoleVariable(TEXT("fx.DevDetailsPanels"));
+
+	TWeakPtr<FNiagaraScriptViewModel> ScriptViewModelWeakPtr = ScriptViewModel;
+
+	FPropertyEditorModule& PropertyEditorModule = FModuleManager::GetModuleChecked<FPropertyEditorModule>("PropertyEditor");
+	FDetailsViewArgs DetailsViewArgs(false, false, true, FDetailsViewArgs::HideNameArea, true);
+	TSharedRef<IDetailsView> DetailsView = PropertyEditorModule.CreateDetailView(DetailsViewArgs);
+
+	DetailsView->OnFinishedChangingProperties().AddRaw(this, &FNiagaraScriptToolkit::OnDetailsSelectionPropertyFinishedChanging);
+	DetailsView->RegisterInstancedCustomPropertyLayout(
+		UNiagaraScript::StaticClass(),
+		FOnGetDetailCustomizationInstance::CreateStatic(&FNiagaraScriptDetails::MakeInstance, ScriptViewModelWeakPtr)
+	);
+
+	DetailsView->SetObjects(DetailsSelection->GetSelectedObjects().Array());
+
+	if (DevDetailsCVar && DevDetailsCVar->GetInt() != 0)
+	{
+		return SNew(SDockTab)
+			.Label(LOCTEXT("ScriptNodeDetailsTabLabel", "Details"))
+			.TabColorScale(GetTabColorScale())
+			[
+				SNew(SScrollBox)
+				+ SScrollBox::Slot()
+				.Padding(0, 3, 0, 0)
+				[
+					DetailsView
+				]
+				+ SScrollBox::Slot()
+				.Padding(0, 3, 0, 0)
+				[
+					SNew(SNiagaraSelectedObjectsDetails, ScriptViewModel->GetGraphViewModel()->GetSelection())
+				]
+			];
+	}
+	else
+	{
+		return SNew(SDockTab)
+			.Label(LOCTEXT("ScriptNodeDetailsTabLabel", "Details"))
+			.TabColorScale(GetTabColorScale())
+			[
+				DetailsView
+			];
+	}
+}
+
+void FNiagaraScriptToolkit::SetupCommands()
+{
+	GetToolkitCommands()->MapAction(
+		FNiagaraEditorCommands::Get().Apply,
+		FExecuteAction::CreateSP(this, &FNiagaraScriptToolkit::OnApply),
+		FCanExecuteAction::CreateSP(this, &FNiagaraScriptToolkit::OnApplyEnabled));
+	GetToolkitCommands()->MapAction(
+		FNiagaraEditorCommands::Get().Compile,
+		FExecuteAction::CreateRaw(this, &FNiagaraScriptToolkit::CompileScript));
+	GetToolkitCommands()->MapAction(
+		FNiagaraEditorCommands::Get().RefreshNodes,
+		FExecuteAction::CreateRaw(this, &FNiagaraScriptToolkit::RefreshNodes));
+}
+
+void FNiagaraScriptToolkit::ExtendToolbar()
+{
+	struct Local
+	{
+		static void FillToolbar(FToolBarBuilder& ToolbarBuilder, FNiagaraScriptToolkit* ScriptToolkit)
+		{
+			ToolbarBuilder.BeginSection("Apply");
+			ToolbarBuilder.AddToolBarButton(FNiagaraEditorCommands::Get().Apply,
+				NAME_None, TAttribute<FText>(), TAttribute<FText>(), 
+				FSlateIcon(FNiagaraEditorStyle::GetStyleSetName(), "NiagaraEditor.Apply"),
+				FName(TEXT("ApplyNiagaraScript")));
+			ToolbarBuilder.EndSection();
+
+			ToolbarBuilder.BeginSection("Compile");
+			ToolbarBuilder.AddToolBarButton(FNiagaraEditorCommands::Get().Compile,
+				NAME_None,
+				TAttribute<FText>(),
+				TAttribute<FText>(ScriptToolkit, &FNiagaraScriptToolkit::GetCompileStatusTooltip),
+				TAttribute<FSlateIcon>(ScriptToolkit, &FNiagaraScriptToolkit::GetCompileStatusImage),
+				FName(TEXT("CompileNiagaraScript")));
+			ToolbarBuilder.AddToolBarButton(FNiagaraEditorCommands::Get().RefreshNodes,
+				NAME_None,
+				TAttribute<FText>(),
+				TAttribute<FText>(ScriptToolkit, &FNiagaraScriptToolkit::GetRefreshStatusTooltip),
+				TAttribute<FSlateIcon>(ScriptToolkit, &FNiagaraScriptToolkit::GetRefreshStatusImage),
+				FName(TEXT("RefreshScriptReferences")));
+			ToolbarBuilder.EndSection();
+		}
+	};
+
+	TSharedPtr<FExtender> ToolbarExtender = MakeShareable(new FExtender);
+
+	ToolbarExtender->AddToolBarExtension(
+		"Asset",
+		EExtensionHook::After,
+		GetToolkitCommands(),
+		FToolBarExtensionDelegate::CreateStatic(&Local::FillToolbar, this));
+
+	AddToolbarExtender(ToolbarExtender);
+
+	FNiagaraEditorModule& NiagaraEditorModule = FModuleManager::LoadModuleChecked<FNiagaraEditorModule>( "NiagaraEditor" );
+	AddToolbarExtender(NiagaraEditorModule.GetToolBarExtensibilityManager()->GetAllExtenders(GetToolkitCommands(), GetEditingObjects()));
+}
+
+FSlateIcon FNiagaraScriptToolkit::GetCompileStatusImage() const
+{
+	ENiagaraScriptCompileStatus Status = ScriptViewModel->GetLatestCompileStatus();
+
+	switch (Status)
+	{
+	default:
+	case ENiagaraScriptCompileStatus::NCS_Unknown:
+	case ENiagaraScriptCompileStatus::NCS_Dirty:
+		return FSlateIcon(FNiagaraEditorStyle::GetStyleSetName(), "Niagara.CompileStatus.Unknown");
+	case ENiagaraScriptCompileStatus::NCS_Error:
+		return FSlateIcon(FNiagaraEditorStyle::GetStyleSetName(), "Niagara.CompileStatus.Error");
+	case ENiagaraScriptCompileStatus::NCS_UpToDate:
+		return FSlateIcon(FNiagaraEditorStyle::GetStyleSetName(), "Niagara.CompileStatus.Good");
+	case ENiagaraScriptCompileStatus::NCS_UpToDateWithWarnings:
+		return FSlateIcon(FNiagaraEditorStyle::GetStyleSetName(), "Niagara.CompileStatus.Warning");
+	}
+}
+
+FText FNiagaraScriptToolkit::GetCompileStatusTooltip() const
+{
+	ENiagaraScriptCompileStatus Status = ScriptViewModel->GetLatestCompileStatus();
+	return FNiagaraEditorUtilities::StatusToText(Status);
+}
+
+FSlateIcon FNiagaraScriptToolkit::GetRefreshStatusImage() const
+{
+	return FSlateIcon(FNiagaraEditorStyle::GetStyleSetName(), "Niagara.Asset.ReimportAsset.Default");
+}
+
+FText FNiagaraScriptToolkit::GetRefreshStatusTooltip() const
+{
+	return LOCTEXT("Refresh_Status", "Currently dependencies up-to-date. Consider refreshing if status isn't accurate.");
+}
+
+void FNiagaraScriptToolkit::CompileScript()
+{
+	ScriptViewModel->CompileStandaloneScript();
+}
+
+void FNiagaraScriptToolkit::RefreshNodes()
+{
+	ScriptViewModel->RefreshNodes();
+}
+
+void FNiagaraScriptToolkit::OnApply()
+{
+	UE_LOG(LogNiagaraEditor, Log, TEXT("Applying Niagara Script %s"), *GetEditingObjects()[0]->GetName());
+
+	UpdateOriginalNiagaraScript();
+}
+
+bool FNiagaraScriptToolkit::OnApplyEnabled() const
+{
+	return ScriptViewModel->GetScriptDirty();
+}
+
+void FNiagaraScriptToolkit::AddReferencedObjects(FReferenceCollector& Collector)
+{
+	Collector.AddReferencedObject(OriginalNiagaraScript);
+	Collector.AddReferencedObject(EditedNiagaraScript);
+}
+
+
+
+void FNiagaraScriptToolkit::GetSaveableObjects(TArray<UObject*>& OutObjects) const
+{
+	OutObjects.Add(OriginalNiagaraScript);
+}
+
+
+void FNiagaraScriptToolkit::SaveAsset_Execute()
+{
+	UE_LOG(LogNiagaraEditor, Log, TEXT("Saving and Compiling NiagaraScript %s"), *GetEditingObjects()[0]->GetName());
+
+	UpdateOriginalNiagaraScript();
+
+	FAssetEditorToolkit::SaveAsset_Execute();
+}
+
+void FNiagaraScriptToolkit::SaveAssetAs_Execute()
+{
+	UE_LOG(LogNiagaraEditor, Log, TEXT("Saving and Compiling NiagaraScript %s"), *GetEditingObjects()[0]->GetName());
+
+	UpdateOriginalNiagaraScript();
+
+	FAssetEditorToolkit::SaveAssetAs_Execute();
+}
+
+void FNiagaraScriptToolkit::UpdateOriginalNiagaraScript()
+{
+	const FScopedBusyCursor BusyCursor;
+
+	const FText LocalizedScriptEditorApply = NSLOCTEXT("UnrealEd", "ToolTip_NiagaraScriptEditorApply", "Apply changes to original script and its use in the world.");
+	GWarn->BeginSlowTask(LocalizedScriptEditorApply, true);
+	GWarn->StatusUpdate(1, 1, LocalizedScriptEditorApply);
+
+	if (OriginalNiagaraScript->IsSelected())
+	{
+		GEditor->GetSelectedObjects()->Deselect(OriginalNiagaraScript);
+	}
+
+	ResetLoaders(OriginalNiagaraScript->GetOutermost()); // Make sure that we're not going to get invalid version number linkers into the package we are going into. 
+	OriginalNiagaraScript->GetOutermost()->LinkerCustomVersion.Empty();
+
+	// overwrite the original script in place by constructing a new one with the same name
+	OriginalNiagaraScript = (UNiagaraScript*)StaticDuplicateObject(EditedNiagaraScript, OriginalNiagaraScript->GetOuter(), OriginalNiagaraScript->GetFName(),
+		RF_AllFlags,
+		OriginalNiagaraScript->GetClass());
+
+	// Restore RF_Standalone on the original material, as it had been removed from the preview material so that it could be GC'd.
+	OriginalNiagaraScript->SetFlags(RF_Standalone);
+
+	// Now there might be other Scripts with functions that referenced this script. So let's update them. They'll need a recompile.
+	// Note that we don't discriminate between the version that are open in transient packages (likely duplicates for editing) and the
+	// original in-scene versions.
+	TArray<UNiagaraScript*> AffectedScripts;
+	UNiagaraGraph* OriginalGraph = CastChecked<UNiagaraScriptSource>(OriginalNiagaraScript->GetSource())->NodeGraph;
+	for (TObjectIterator<UNiagaraScript> It; It; ++It)
+	{
+		if (*It == OriginalNiagaraScript || It->IsPendingKillOrUnreachable())
+		{
+			continue;
+		}
+
+		// First see if it is directly called, as this will force a need to refresh from external changes...
+		UNiagaraScriptSource* Source = Cast<UNiagaraScriptSource>(It->GetSource());
+		TArray<UNiagaraNode*> NiagaraNodes;
+		Source->NodeGraph->GetNodesOfClass<UNiagaraNode>(NiagaraNodes);
+		bool bRefreshed = false;
+		for (UNiagaraNode* NiagaraNode : NiagaraNodes)
+		{
+			UObject* ReferencedAsset = NiagaraNode->GetReferencedAsset();
+			if (ReferencedAsset == OriginalNiagaraScript)
+			{
+				NiagaraNode->RefreshFromExternalChanges();
+				bRefreshed = true;
+			}
+		}
+
+		if (bRefreshed)
+		{
+			Source->NodeGraph->NotifyGraphNeedsRecompile();
+			AffectedScripts.Add(*It);
+		}
+		else
+		{
+			// Now check to see if our graph is anywhere in the dependency chain for a given graph. If it is, 
+			// then it will need to be recompiled against the latest version.
+			TArray<const UNiagaraGraph*> ReferencedGraphs;
+			Source->NodeGraph->GetAllReferencedGraphs(ReferencedGraphs);
+			for (const UNiagaraGraph* Graph : ReferencedGraphs)
+			{
+				if (Graph == OriginalGraph)
+				{
+					Source->NodeGraph->NotifyGraphNeedsRecompile();
+					AffectedScripts.Add(*It);
+					break;
+				}
+			}
+		}
+	}
+
+	// Now determine if any of these scripts were in Emitters. If so, those emitters should be compiled together. If not, go ahead and compile individually.
+	// Use the existing view models if they exist,as they are already wired into the correct UI.
+	TArray<UNiagaraEmitter*> AffectedEmitters;
+	for (UNiagaraScript* Script : AffectedScripts)
+	{
+		if (Script->IsParticleUpdateScript() || Script->IsParticleSpawnScript() || Script->IsParticleEventScript())
+		{
+			UNiagaraEmitter* Emitter = Cast<UNiagaraEmitter>(Script->GetOuter());
+			if (Emitter)
+			{
+				AffectedEmitters.AddUnique(Emitter);
+			}
+		}
+		else
+		{
+			TSharedPtr<FNiagaraScriptViewModel> AffectedScriptViewModel = FNiagaraScriptViewModel::GetExistingViewModelForObject(Script);
+			if (!AffectedScriptViewModel.IsValid())
+			{
+				AffectedScriptViewModel = MakeShareable(new FNiagaraScriptViewModel(Script, FText::FromString(Script->GetName()), ENiagaraParameterEditMode::EditValueOnly));
+			}
+			AffectedScriptViewModel->CompileStandaloneScript();
+		}
+	}
+
+	FNiagaraEditorUtilities::UpdateExistingEmitters(AffectedEmitters);
+
+	GWarn->EndSlowTask();
+	ScriptViewModel->SetScriptDirty(false);
+}
+
+
+bool FNiagaraScriptToolkit::OnRequestClose()
+{
+	if (ScriptViewModel->GetScriptDirty())
+	{
+		// find out the user wants to do with this dirty NiagaraScript
+		EAppReturnType::Type YesNoCancelReply = FMessageDialog::Open(EAppMsgType::YesNoCancel,
+			FText::Format(
+				NSLOCTEXT("UnrealEd", "Prompt_NiagaraScriptEditorClose", "Would you like to apply changes to this NiagaraScript to the original NiagaraScript?\n{0}\n(No will lose all changes!)"),
+				FText::FromString(OriginalNiagaraScript->GetPathName())));
+
+		// act on it
+		switch (YesNoCancelReply)
+		{
+		case EAppReturnType::Yes:
+			// update NiagaraScript and exit
+			UpdateOriginalNiagaraScript();
+			break;
+
+		case EAppReturnType::No:
+			// exit
+			// do nothing.. bNiagaraScriptDirty = false;
+			break;
+
+		case EAppReturnType::Cancel:
+			// don't exit
+			return false;
+		}
+	}
+
+	return true;
+}
+
+#undef LOCTEXT_NAMESPACE
diff --git a/Engine/Plugins/FX/Niagara/Source/NiagaraEditor/Private/Toolkits/NiagaraSystemToolkit.cpp b/Engine/Plugins/FX/Niagara/Source/NiagaraEditor/Private/Toolkits/NiagaraSystemToolkit.cpp
new file mode 100644
--- /dev/null
+++ b/Engine/Plugins/FX/Niagara/Source/NiagaraEditor/Private/Toolkits/NiagaraSystemToolkit.cpp
@@ -0,0 +1,884 @@
+// Copyright 1998-2017 Epic Games, Inc. All Rights Reserved.
+
+#include "NiagaraSystemToolkit.h"
+#include "NiagaraEditorModule.h"
+#include "NiagaraSystem.h"
+#include "NiagaraEmitter.h"
+#include "NiagaraObjectSelection.h"
+#include "NiagaraSystemViewModel.h"
+#include "NiagaraEmitterHandleviewModel.h"
+#include "NiagaraEmitterViewModel.h"
+#include "NiagaraSystemScriptViewModel.h"
+#include "Widgets/SNiagaraCurveEditor.h"
+#include "Widgets/SNiagaraSystemScript.h"
+#include "Widgets/SNiagaraSystemViewport.h"
+#include "Widgets/SNiagaraSelectedObjectsDetails.h"
+#include "Widgets/SNiagaraSelectedEmitterHandle.h"
+#include "Widgets/SNiagaraSpreadsheetView.h"
+#include "Widgets/SNiagaraGeneratedCodeView.h"
+#include "Widgets/SNiagaraScriptGraph.h"
+#include "NiagaraEditorCommands.h"
+#include "NiagaraEditorStyle.h"
+#include "NiagaraEditorSettings.h"
+#include "NiagaraSystemFactoryNew.h"
+
+#include "IContentBrowserSingleton.h"
+#include "ContentBrowserModule.h"
+
+#include "EditorStyleSet.h"
+#include "AssetEditorToolkit.h"
+#include "WorkspaceItem.h"
+#include "ScopedTransaction.h"
+
+#include "SlateApplication.h"
+#include "SBoxPanel.h"
+#include "SBox.h"
+#include "SDockTab.h"
+#include "AdvancedPreviewSceneModule.h"
+#include "BusyCursor.h"
+#include "Misc/FeedbackContext.h"
+#include "Editor.h"
+#include "Engine/Selection.h"
+#include "Misc/MessageDialog.h"
+
+#define LOCTEXT_NAMESPACE "NiagaraSystemEditor"
+
+const FName FNiagaraSystemToolkit::ViewportTabID(TEXT("NiagaraSystemEditor_Viewport"));
+const FName FNiagaraSystemToolkit::CurveEditorTabID(TEXT("NiagaraSystemEditor_CurveEditor"));
+const FName FNiagaraSystemToolkit::SequencerTabID(TEXT("NiagaraSystemEditor_Sequencer"));
+const FName FNiagaraSystemToolkit::SystemScriptTabID(TEXT("NiagaraSystemEditor_SystemScript"));
+const FName FNiagaraSystemToolkit::SystemDetailsTabID(TEXT("NiagaraSystemEditor_SystemDetails"));
+const FName FNiagaraSystemToolkit::SelectedEmitterStackTabID(TEXT("NiagaraSystemEditor_SelectedEmitterStack"));
+const FName FNiagaraSystemToolkit::SelectedEmitterGraphTabID(TEXT("NiagaraSystemEditor_SelectedEmitterGraph"));
+const FName FNiagaraSystemToolkit::DebugSpreadsheetTabID(TEXT("NiagaraSystemEditor_DebugAttributeSpreadsheet"));
+const FName FNiagaraSystemToolkit::PreviewSettingsTabId(TEXT("NiagaraSystemEditor_PreviewSettings"));
+const FName FNiagaraSystemToolkit::GeneratedCodeTabID(TEXT("NiagaraSystemEditor_GeneratedCode"));
+
+void FNiagaraSystemToolkit::RegisterTabSpawners(const TSharedRef<class FTabManager>& InTabManager)
+{
+	WorkspaceMenuCategory = InTabManager->AddLocalWorkspaceMenuCategory(LOCTEXT("WorkspaceMenu_NiagaraSystemEditor", "Niagara System"));
+
+	FAssetEditorToolkit::RegisterTabSpawners(InTabManager);
+
+	InTabManager->RegisterTabSpawner(ViewportTabID, FOnSpawnTab::CreateSP(this, &FNiagaraSystemToolkit::SpawnTab_Viewport))
+		.SetDisplayName(LOCTEXT("Preview", "Preview"))
+		.SetGroup(WorkspaceMenuCategory.ToSharedRef())
+		.SetIcon(FSlateIcon(FEditorStyle::GetStyleSetName(), "LevelEditor.Tabs.Viewports"));
+
+	InTabManager->RegisterTabSpawner(CurveEditorTabID, FOnSpawnTab::CreateSP(this, &FNiagaraSystemToolkit::SpawnTab_CurveEd))
+		.SetDisplayName(LOCTEXT("Curves", "Curves"))
+		.SetGroup(WorkspaceMenuCategory.ToSharedRef());
+
+	InTabManager->RegisterTabSpawner(SequencerTabID, FOnSpawnTab::CreateSP(this, &FNiagaraSystemToolkit::SpawnTab_Sequencer))
+		.SetDisplayName(LOCTEXT("Timeline", "Timeline"))
+		.SetGroup(WorkspaceMenuCategory.ToSharedRef());
+
+	InTabManager->RegisterTabSpawner(SystemScriptTabID, FOnSpawnTab::CreateSP(this, &FNiagaraSystemToolkit::SpawnTab_SystemScript))
+		.SetDisplayName(LOCTEXT("SystemScript", "System Script"))
+		.SetGroup(WorkspaceMenuCategory.ToSharedRef());
+
+	InTabManager->RegisterTabSpawner(SystemDetailsTabID, FOnSpawnTab::CreateSP(this, &FNiagaraSystemToolkit::SpawnTab_SystemDetails))
+		.SetDisplayName(LOCTEXT("SystemDetails", "System Details"))
+		.SetGroup(WorkspaceMenuCategory.ToSharedRef());
+
+	InTabManager->RegisterTabSpawner(SelectedEmitterStackTabID, FOnSpawnTab::CreateSP(this, &FNiagaraSystemToolkit::SpawnTab_SelectedEmitterStack))
+		.SetDisplayName(LOCTEXT("SelectedEmitterStack", "Selected Emitter"))
+		.SetGroup(WorkspaceMenuCategory.ToSharedRef());
+
+	InTabManager->RegisterTabSpawner(SelectedEmitterGraphTabID, FOnSpawnTab::CreateSP(this, &FNiagaraSystemToolkit::SpawnTab_SelectedEmitterGraph))
+		.SetDisplayName(LOCTEXT("SelectedEmitterGraph", "Selected Emitter Graph"))
+		.SetGroup(WorkspaceMenuCategory.ToSharedRef());
+
+	InTabManager->RegisterTabSpawner(DebugSpreadsheetTabID, FOnSpawnTab::CreateSP(this, &FNiagaraSystemToolkit::SpawnTab_DebugSpreadsheet))
+		.SetDisplayName(LOCTEXT("DebugSpreadsheet", "Attribute Spreadsheet"))
+		.SetGroup(WorkspaceMenuCategory.ToSharedRef());
+
+	InTabManager->RegisterTabSpawner(PreviewSettingsTabId, FOnSpawnTab::CreateSP(this, &FNiagaraSystemToolkit::SpawnTab_PreviewSettings))
+		.SetDisplayName(LOCTEXT("PreviewSceneSettingsTab", "Preview Scene Settings"))
+		.SetGroup(WorkspaceMenuCategory.ToSharedRef())
+		.SetIcon(FSlateIcon(FEditorStyle::GetStyleSetName(), "LevelEditor.Tabs.Details"));
+
+	InTabManager->RegisterTabSpawner(GeneratedCodeTabID, FOnSpawnTab::CreateSP(this, &FNiagaraSystemToolkit::SpawnTab_GeneratedCode))
+		.SetDisplayName(LOCTEXT("GeneratedCode", "Generated Code"))
+		.SetGroup(WorkspaceMenuCategory.ToSharedRef());
+}
+
+void FNiagaraSystemToolkit::UnregisterTabSpawners(const TSharedRef<class FTabManager>& InTabManager)
+{
+	FAssetEditorToolkit::UnregisterTabSpawners(InTabManager);
+
+	InTabManager->UnregisterTabSpawner(ViewportTabID);
+	InTabManager->UnregisterTabSpawner(CurveEditorTabID);
+	InTabManager->UnregisterTabSpawner(SequencerTabID);
+	InTabManager->UnregisterTabSpawner(SystemScriptTabID);
+	InTabManager->UnregisterTabSpawner(SystemDetailsTabID);
+	InTabManager->UnregisterTabSpawner(SelectedEmitterStackTabID);
+	InTabManager->UnregisterTabSpawner(SelectedEmitterGraphTabID);
+	InTabManager->UnregisterTabSpawner(DebugSpreadsheetTabID);
+	InTabManager->UnregisterTabSpawner(PreviewSettingsTabId);
+	InTabManager->UnregisterTabSpawner(GeneratedCodeTabID);
+}
+
+
+FNiagaraSystemToolkit::~FNiagaraSystemToolkit()
+{
+	if (SystemViewModel.IsValid())
+	{
+		SystemViewModel->Cleanup();
+	}
+	SystemViewModel.Reset();
+}
+
+void FNiagaraSystemToolkit::AddReferencedObjects(FReferenceCollector& Collector) 
+{
+	Collector.AddReferencedObject(System);
+}
+
+void FNiagaraSystemToolkit::InitializeWithSystem(const EToolkitMode::Type Mode, const TSharedPtr< class IToolkitHost >& InitToolkitHost, UNiagaraSystem& InSystem)
+{
+	System = &InSystem;
+	Emitter = nullptr;
+
+	// In the FNiagaraCustomVersion::UpdateSpawnEventGraphCombination we merged graphs. We update the graph source here because there isn't a good place to do it in the postload pipeline.
+	bool bConverted = false;
+	for (int32 i = 0; i < System->GetNumEmitters(); i++)
+	{
+		FNiagaraEmitterHandle& Handle = System->GetEmitterHandle(i);
+		if (Handle.GetSource() == nullptr)
+		{
+			UE_LOG(LogNiagaraEditor, Error, TEXT("Missing source emitter!"));
+			break;
+		}
+
+		if (Handle.GetSource()->GraphSource == nullptr)
+		{
+			if (FNiagaraEditorUtilities::ConvertToMergedGraph((UNiagaraEmitter*)Handle.GetSource()))
+			{
+				bConverted = true;
+			}
+			else
+			{
+				UE_LOG(LogNiagaraEditor, Error, TEXT("Failed to merge emitter!"));
+			}
+		}
+	}
+
+	if (bConverted)
+	{
+		System->ResynchronizeAllHandles();
+	}
+
+	FNiagaraSystemViewModelOptions SystemOptions;
+	SystemOptions.bCanRemoveEmittersFromTimeline = true;
+	SystemOptions.bCanRenameEmittersFromTimeline = true;
+	SystemOptions.bCanAddEmittersFromTimeline = true;
+	SystemOptions.bUseSystemExecStateForTimelineReset = true;
+	SystemOptions.OnGetSequencerAddMenuContent.BindSP(this, &FNiagaraSystemToolkit::GetSequencerAddMenuContent);
+
+	System->CheckForUpdates();
+
+	SystemViewModel = MakeShareable(new FNiagaraSystemViewModel(*System, SystemOptions));
+	SystemToolkitMode = ESystemToolkitMode::System;
+	InitializeInternal(Mode, InitToolkitHost);
+}
+
+void FNiagaraSystemToolkit::InitializeWithEmitter(const EToolkitMode::Type Mode, const TSharedPtr< class IToolkitHost >& InitToolkitHost, UNiagaraEmitter& InEmitter)
+{
+	// In the FNiagaraCustomVersion::UpdateSpawnEventGraphCombination we merged graphs. We update the graph source here because there isn't a good place to do it in the postload pipeline.
+	if (InEmitter.GraphSource == nullptr)
+	{
+		if (false == FNiagaraEditorUtilities::ConvertToMergedGraph(&InEmitter))
+		{
+			UE_LOG(LogNiagaraEditor, Error, TEXT("Failed to merge emitter!"));
+		}
+	}
+
+	System = NewObject<UNiagaraSystem>(GetTransientPackage(), NAME_None, RF_Transient);
+	UNiagaraSystemFactoryNew::InitializeSystem(System);
+
+	Emitter = &InEmitter;
+
+	ResetLoaders(GetTransientPackage()); // Make sure that we're not going to get invalid version number linkers into the package we are going into. 
+	GetTransientPackage()->LinkerCustomVersion.Empty();
+
+	UNiagaraEmitter* EditableEmitter = (UNiagaraEmitter*)StaticDuplicateObject(Emitter, GetTransientPackage(), NAME_None, ~RF_Standalone, UNiagaraEmitter::StaticClass());
+	System->AddEmitterHandleWithoutCopying(*EditableEmitter);
+
+	FNiagaraSystemViewModelOptions SystemOptions;
+	SystemOptions.bCanRemoveEmittersFromTimeline = false;
+	SystemOptions.bCanRenameEmittersFromTimeline = false;
+	SystemOptions.bCanAddEmittersFromTimeline = false;
+	SystemOptions.bUseSystemExecStateForTimelineReset = false;
+
+	SystemViewModel = MakeShareable(new FNiagaraSystemViewModel(*System, SystemOptions));
+	SystemViewModel->GetSystemScriptViewModel()->RebuildEmitterNodes();
+	SystemToolkitMode = ESystemToolkitMode::Emitter;
+	InitializeInternal(Mode, InitToolkitHost);
+}
+
+void FNiagaraSystemToolkit::InitializeInternal(const EToolkitMode::Type Mode, const TSharedPtr<IToolkitHost>& InitToolkitHost)
+{
+	if (SystemViewModel->GetEmitterHandleViewModels().Num() > 0)
+	{
+		SystemViewModel->SetSelectedEmitterHandleById(SystemViewModel->GetEmitterHandleViewModels()[0]->GetId());
+	}
+
+	const float InTime = -0.02f;
+	const float OutTime = 3.2f;
+
+	TSharedRef<FTabManager::FLayout> StandaloneDefaultLayout = FTabManager::NewLayout("Standalone_Niagara_System_Layout_v17")
+		->AddArea
+		(
+			FTabManager::NewPrimaryArea()->SetOrientation(Orient_Vertical)
+			->Split
+			(
+				FTabManager::NewStack()
+				->SetSizeCoefficient(0.1f)
+				->AddTab(GetToolbarTabId(), ETabState::OpenedTab)
+				->SetHideTabWell(true)
+				)
+			->Split
+			(
+				FTabManager::NewSplitter()->SetOrientation(Orient_Horizontal)
+				->Split
+				(
+					FTabManager::NewSplitter()->SetOrientation(Orient_Vertical)
+					->SetSizeCoefficient(.75f)
+					->Split
+					(
+						FTabManager::NewStack()
+						->SetSizeCoefficient(.75f)
+						->AddTab(ViewportTabID, ETabState::OpenedTab)
+						)
+					->Split
+					(
+						FTabManager::NewStack()
+						->SetSizeCoefficient(0.25f)
+						->AddTab(CurveEditorTabID, ETabState::OpenedTab)
+						->AddTab(SequencerTabID, ETabState::OpenedTab)
+						)
+					)
+				->Split
+				(
+					FTabManager::NewStack()
+					->SetSizeCoefficient(0.25f)
+					->AddTab(SelectedEmitterStackTabID, ETabState::OpenedTab)
+					->AddTab(SelectedEmitterGraphTabID, ETabState::ClosedTab)
+					->AddTab(SystemScriptTabID, ETabState::ClosedTab)
+					->AddTab(SystemDetailsTabID, ETabState::ClosedTab)
+					->AddTab(DebugSpreadsheetTabID, ETabState::ClosedTab)
+					->AddTab(PreviewSettingsTabId, ETabState::ClosedTab)
+					->AddTab(GeneratedCodeTabID, ETabState::ClosedTab)
+				)
+			)
+		);
+
+	const bool bCreateDefaultStandaloneMenu = true;
+	const bool bCreateDefaultToolbar = true;
+	UObject* ToolkitObject = SystemToolkitMode == ESystemToolkitMode::System ? (UObject*)System : (UObject*)Emitter;
+	FAssetEditorToolkit::InitAssetEditor(Mode, InitToolkitHost, FNiagaraEditorModule::NiagaraEditorAppIdentifier,
+		StandaloneDefaultLayout, bCreateDefaultStandaloneMenu, bCreateDefaultToolbar, ToolkitObject);
+	
+	FNiagaraEditorModule& NiagaraEditorModule = FModuleManager::LoadModuleChecked<FNiagaraEditorModule>("NiagaraEditor");
+	AddMenuExtender(NiagaraEditorModule.GetMenuExtensibilityManager()->GetAllExtenders(GetToolkitCommands(), GetEditingObjects()));
+
+	SetupCommands();
+	ExtendToolbar();
+	RegenerateMenusAndToolbars();
+}
+
+FName FNiagaraSystemToolkit::GetToolkitFName() const
+{
+	return FName("Niagara");
+}
+
+FText FNiagaraSystemToolkit::GetBaseToolkitName() const
+{
+	return LOCTEXT("AppLabel", "Niagara");
+}
+
+FString FNiagaraSystemToolkit::GetWorldCentricTabPrefix() const
+{
+	return LOCTEXT("WorldCentricTabPrefix", "Niagara ").ToString();
+}
+
+
+FLinearColor FNiagaraSystemToolkit::GetWorldCentricTabColorScale() const
+{
+	return FNiagaraEditorModule::WorldCentricTabColorScale;
+}
+
+
+TSharedRef<SDockTab> FNiagaraSystemToolkit::SpawnTab_Viewport(const FSpawnTabArgs& Args)
+{
+	check(Args.GetTabId().TabType == ViewportTabID);
+
+	Viewport = SNew(SNiagaraSystemViewport);
+
+	TSharedRef<SDockTab> SpawnedTab =
+		SNew(SDockTab)
+		[
+			Viewport.ToSharedRef()
+		];
+
+	Viewport->SetPreviewComponent(SystemViewModel->GetPreviewComponent());
+	Viewport->OnAddedToTab(SpawnedTab);
+
+	return SpawnedTab;
+}
+
+TSharedRef<SDockTab> FNiagaraSystemToolkit::SpawnTab_PreviewSettings(const FSpawnTabArgs& Args)
+{
+	check(Args.GetTabId() == PreviewSettingsTabId);
+
+	TSharedRef<SWidget> InWidget = SNullWidget::NullWidget;
+	if (Viewport.IsValid())
+	{
+		FAdvancedPreviewSceneModule& AdvancedPreviewSceneModule = FModuleManager::LoadModuleChecked<FAdvancedPreviewSceneModule>("AdvancedPreviewScene");
+		InWidget = AdvancedPreviewSceneModule.CreateAdvancedPreviewSceneSettingsWidget(Viewport->GetPreviewScene());
+	}
+
+	TSharedRef<SDockTab> SpawnedTab = SNew(SDockTab)
+		.Icon(FEditorStyle::GetBrush("LevelEditor.Tabs.Details"))
+		.Label(LOCTEXT("PreviewSceneSettingsTab", "Preview Scene Settings"))
+		[
+			InWidget
+		];
+
+	return SpawnedTab;
+}
+
+
+TSharedRef<SDockTab> FNiagaraSystemToolkit::SpawnTab_CurveEd(const FSpawnTabArgs& Args)
+{
+	check(Args.GetTabId().TabType == CurveEditorTabID);
+
+	TSharedRef<SDockTab> SpawnedTab =
+		SNew(SDockTab)
+		[
+			SNew(SNiagaraCurveEditor, SystemViewModel.ToSharedRef())
+		];
+
+	return SpawnedTab;
+}
+
+
+TSharedRef<SDockTab> FNiagaraSystemToolkit::SpawnTab_Sequencer(const FSpawnTabArgs& Args)
+{
+	check(Args.GetTabId().TabType == SequencerTabID);
+
+	TSharedRef<SDockTab> SpawnedTab =
+		SNew(SDockTab)
+		[
+			SystemViewModel->GetSequencer()->GetSequencerWidget()
+		];
+
+	return SpawnedTab;
+}
+
+
+TSharedRef<SDockTab> FNiagaraSystemToolkit::SpawnTab_SystemScript(const FSpawnTabArgs& Args)
+{
+	check(Args.GetTabId().TabType == SystemScriptTabID);
+
+	TSharedRef<SDockTab> SpawnedTab =
+		SNew(SDockTab)
+		[
+			SNew(SNiagaraSystemScript, SystemViewModel.ToSharedRef())
+		];
+
+	return SpawnedTab;
+}
+
+
+TSharedRef<SDockTab> FNiagaraSystemToolkit::SpawnTab_SystemDetails(const FSpawnTabArgs& Args)
+{
+	check(Args.GetTabId().TabType == SystemDetailsTabID);
+
+	TSharedRef<FNiagaraObjectSelection> SystemSelection = MakeShareable(new FNiagaraObjectSelection());
+	SystemSelection->SetSelectedObject(System);
+
+	TSharedRef<SDockTab> SpawnedTab =
+		SNew(SDockTab)
+		[
+			SNew(SNiagaraSelectedObjectsDetails, SystemSelection)
+		];
+
+	return SpawnedTab;
+}
+
+TSharedRef<SDockTab> FNiagaraSystemToolkit::SpawnTab_SelectedEmitterStack(const FSpawnTabArgs& Args)
+{
+	check(Args.GetTabId().TabType == SelectedEmitterStackTabID);
+
+	TSharedRef<SDockTab> SpawnedTab =
+		SNew(SDockTab)
+		[
+			SNew(SNiagaraSelectedEmitterHandle, SystemViewModel.ToSharedRef())
+		];
+
+	return SpawnedTab;
+}
+
+class SNiagaraSelectedEmitterGraph : public SCompoundWidget
+{
+public:
+	SLATE_BEGIN_ARGS(SNiagaraSelectedEmitterGraph)
+	{}
+	SLATE_END_ARGS();
+
+	void Construct(const FArguments& InArgs, TSharedRef<FNiagaraSystemViewModel> InSystemViewModel)
+	{
+		SystemViewModel = InSystemViewModel;
+		SystemViewModel->OnSelectedEmitterHandlesChanged().AddRaw(this, &SNiagaraSelectedEmitterGraph::SelectedEmitterHandlesChanged);
+		ChildSlot
+		[
+			SAssignNew(GraphWidgetContainer, SBox)
+		];
+		UpdateGraphWidget();
+	}
+
+	~SNiagaraSelectedEmitterGraph()
+	{
+		if (SystemViewModel.IsValid())
+		{
+			SystemViewModel->OnCurveOwnerChanged().RemoveAll(this);
+		}
+	}
+
+private:
+	void SelectedEmitterHandlesChanged()
+	{
+		UpdateGraphWidget();
+	}
+
+	void UpdateGraphWidget()
+	{
+		TArray<TSharedRef<FNiagaraEmitterHandleViewModel>> SelectedEmitterHandles;
+		SystemViewModel->GetSelectedEmitterHandles(SelectedEmitterHandles);
+		if (SelectedEmitterHandles.Num() == 1)
+		{
+			GraphWidgetContainer->SetContent(SNew(SNiagaraScriptGraph, SelectedEmitterHandles[0]->GetEmitterViewModel()->GetSharedScriptViewModel()->GetGraphViewModel()));
+		}
+		else
+		{
+			GraphWidgetContainer->SetContent(SNullWidget::NullWidget);
+		}
+	}
+
+private:
+	TSharedPtr<FNiagaraSystemViewModel> SystemViewModel;
+	TSharedPtr<SBox> GraphWidgetContainer;
+};
+
+TSharedRef<SDockTab> FNiagaraSystemToolkit::SpawnTab_SelectedEmitterGraph(const FSpawnTabArgs& Args)
+{
+	check(Args.GetTabId().TabType == SelectedEmitterGraphTabID);
+
+	TSharedRef<SDockTab> SpawnedTab =
+		SNew(SDockTab)
+		[
+			SNew(SNiagaraSelectedEmitterGraph, SystemViewModel.ToSharedRef())
+		];
+
+	return SpawnedTab;
+}
+
+TSharedRef<SDockTab> FNiagaraSystemToolkit::SpawnTab_DebugSpreadsheet(const FSpawnTabArgs& Args)
+{
+	check(Args.GetTabId().TabType == DebugSpreadsheetTabID);
+
+	TSharedRef<SDockTab> SpawnedTab =
+		SNew(SDockTab)
+		[
+			SNew(SNiagaraSpreadsheetView, SystemViewModel.ToSharedRef())
+		];
+
+	return SpawnedTab;
+}
+
+
+TSharedRef<SDockTab> FNiagaraSystemToolkit::SpawnTab_GeneratedCode(const FSpawnTabArgs& Args)
+{
+	check(Args.GetTabId().TabType == GeneratedCodeTabID);
+
+	TSharedRef<SDockTab> SpawnedTab =
+		SNew(SDockTab)
+		[
+			SNew(SNiagaraGeneratedCodeView, SystemViewModel.ToSharedRef())
+		];
+
+	return SpawnedTab;
+}
+
+void FNiagaraSystemToolkit::SetupCommands()
+{
+	GetToolkitCommands()->MapAction(
+		FNiagaraEditorCommands::Get().ToggleUnlockToChanges,
+		FExecuteAction::CreateSP(this, &FNiagaraSystemToolkit::ToggleUnlockToChanges),
+		FCanExecuteAction(),
+		FIsActionChecked::CreateSP(this, &FNiagaraSystemToolkit::IsToggleUnlockToChangesChecked));
+	GetToolkitCommands()->MapAction(
+		FNiagaraEditorCommands::Get().Compile,
+		FExecuteAction::CreateRaw(this, &FNiagaraSystemToolkit::CompileSystem));
+	GetToolkitCommands()->MapAction(
+		FNiagaraEditorCommands::Get().ResetSimulation,
+		FExecuteAction::CreateRaw(this, &FNiagaraSystemToolkit::ResetSimulation));
+}
+
+void FNiagaraSystemToolkit::ResetSimulation()
+{
+	SystemViewModel->ResetSystem();
+}
+
+void FNiagaraSystemToolkit::ExtendToolbar()
+{
+	struct Local
+	{
+		static void FillToolbar(FToolBarBuilder& ToolbarBuilder, FNiagaraSystemToolkit* Toolkit)
+		{
+			ToolbarBuilder.BeginSection("Compile");
+			{
+				ToolbarBuilder.AddToolBarButton(FNiagaraEditorCommands::Get().Compile,
+					NAME_None,
+					TAttribute<FText>(),
+					TAttribute<FText>(Toolkit, &FNiagaraSystemToolkit::GetCompileStatusTooltip),
+					TAttribute<FSlateIcon>(Toolkit, &FNiagaraSystemToolkit::GetCompileStatusImage),
+					FName(TEXT("CompileNiagaraSystem")));
+				ToolbarBuilder.AddComboButton(
+					FUIAction(),
+					FOnGetContent::CreateStatic(&FNiagaraSystemToolkit::GenerateCompileMenuContent),
+					LOCTEXT("BuildCombo_Label", "Auto-Compile Options"),
+					LOCTEXT("BuildComboToolTip", "Auto-Compile options menu"),
+					FSlateIcon(FEditorStyle::GetStyleSetName(), "LevelEditor.Build"),
+					true);
+			}
+			ToolbarBuilder.EndSection();
+
+			ToolbarBuilder.BeginSection("LockEmitters");
+			{
+				ToolbarBuilder.AddToolBarButton(FNiagaraEditorCommands::Get().ToggleUnlockToChanges, NAME_None,
+					TAttribute<FText>(Toolkit, &FNiagaraSystemToolkit::GetEmitterLockToChangesLabel),
+					TAttribute<FText>(Toolkit, &FNiagaraSystemToolkit::GetEmitterLockToChangesLabelTooltip),
+					TAttribute<FSlateIcon>(Toolkit, &FNiagaraSystemToolkit::GetEmitterLockToChangesIcon));
+			}
+			ToolbarBuilder.EndSection();
+		}
+	};
+
+	TSharedPtr<FExtender> ToolbarExtender = MakeShareable(new FExtender);
+
+	ToolbarExtender->AddToolBarExtension(
+		"Asset",
+		EExtensionHook::After,
+		GetToolkitCommands(),
+		FToolBarExtensionDelegate::CreateStatic(&Local::FillToolbar, this)
+		);
+
+	AddToolbarExtender(ToolbarExtender);
+
+	FNiagaraEditorModule& NiagaraEditorModule = FModuleManager::LoadModuleChecked<FNiagaraEditorModule>("NiagaraEditor");
+	AddToolbarExtender(NiagaraEditorModule.GetToolBarExtensibilityManager()->GetAllExtenders(GetToolkitCommands(), GetEditingObjects()));
+}
+
+void FNiagaraSystemToolkit::GetSequencerAddMenuContent(FMenuBuilder& MenuBuilder, TSharedRef<ISequencer> Sequencer)
+{
+	MenuBuilder.AddSubMenu(
+		LOCTEXT("EmittersLabel", "Emitters..."),
+		LOCTEXT("EmittersToolTip", "Add an existing emitter..."),
+		FNewMenuDelegate::CreateLambda([&](FMenuBuilder& MenuBuilder)
+		{
+			MenuBuilder.AddWidget(CreateAddEmitterMenuContent(), FText());
+		}));
+}
+
+TSharedRef<SWidget> FNiagaraSystemToolkit::CreateAddEmitterMenuContent()
+{
+	FAssetPickerConfig AssetPickerConfig;
+	{
+		AssetPickerConfig.OnAssetSelected = FOnAssetSelected::CreateSP(this, &FNiagaraSystemToolkit::EmitterAssetSelected);
+		AssetPickerConfig.bAllowNullSelection = false;
+		AssetPickerConfig.InitialAssetViewType = EAssetViewType::List;
+		AssetPickerConfig.Filter.ClassNames.Add(UNiagaraEmitter::StaticClass()->GetFName());
+	}
+
+	FContentBrowserModule& ContentBrowserModule = FModuleManager::Get().LoadModuleChecked<FContentBrowserModule>(TEXT("ContentBrowser"));
+
+	return SNew(SBox)
+		.WidthOverride(300.0f)
+		.HeightOverride(300.f)
+		[
+			ContentBrowserModule.Get().CreateAssetPicker(AssetPickerConfig)
+		];
+}
+
+TSharedRef<SWidget> FNiagaraSystemToolkit::GenerateCompileMenuContent()
+{
+	const bool bShouldCloseWindowAfterMenuSelection = true;
+	FMenuBuilder MenuBuilder(bShouldCloseWindowAfterMenuSelection, nullptr);
+
+	FUIAction Action(
+		FExecuteAction::CreateStatic(&FNiagaraSystemToolkit::ToggleCompileEnabled),
+		FCanExecuteAction(),
+		FIsActionChecked::CreateStatic(&FNiagaraSystemToolkit::IsAutoCompileEnabled));
+
+	MenuBuilder.AddMenuEntry(LOCTEXT("AutoCompile", "Automatically compile when graph changes"), FText(), FSlateIcon(), Action, NAME_None, EUserInterfaceActionType::ToggleButton);
+
+	return MenuBuilder.MakeWidget();
+}
+
+FSlateIcon FNiagaraSystemToolkit::GetCompileStatusImage() const
+{
+	ENiagaraScriptCompileStatus Status = SystemViewModel->GetLatestCompileStatus();
+
+	switch (Status)
+	{
+	default:
+	case ENiagaraScriptCompileStatus::NCS_Unknown:
+	case ENiagaraScriptCompileStatus::NCS_Dirty:
+		return FSlateIcon(FNiagaraEditorStyle::GetStyleSetName(), "Niagara.CompileStatus.Unknown");
+	case ENiagaraScriptCompileStatus::NCS_Error:
+		return FSlateIcon(FNiagaraEditorStyle::GetStyleSetName(), "Niagara.CompileStatus.Error");
+	case ENiagaraScriptCompileStatus::NCS_UpToDate:
+		return FSlateIcon(FNiagaraEditorStyle::GetStyleSetName(), "Niagara.CompileStatus.Good");
+	case ENiagaraScriptCompileStatus::NCS_UpToDateWithWarnings:
+		return FSlateIcon(FNiagaraEditorStyle::GetStyleSetName(), "Niagara.CompileStatus.Warning");
+	}
+}
+
+FText FNiagaraSystemToolkit::GetCompileStatusTooltip() const
+{
+	ENiagaraScriptCompileStatus Status = SystemViewModel->GetLatestCompileStatus();
+	return FNiagaraEditorUtilities::StatusToText(Status);
+}
+
+
+void FNiagaraSystemToolkit::CompileSystem()
+{
+	SystemViewModel->CompileSystem();
+}
+
+void FNiagaraSystemToolkit::UpdateOriginalEmitter()
+{
+	checkf(SystemToolkitMode == ESystemToolkitMode::Emitter, TEXT("There is no original emitter to update in system mode."));
+
+	const FScopedBusyCursor BusyCursor;
+	const FText LocalizedScriptEditorApply = NSLOCTEXT("UnrealEd", "ToolTip_NiagaraEmitterEditorApply", "Apply changes to original emitter and its use in the world.");
+	GWarn->BeginSlowTask(LocalizedScriptEditorApply, true);
+	GWarn->StatusUpdate(1, 1, LocalizedScriptEditorApply);
+
+	if (Emitter->IsSelected())
+	{
+		GEditor->GetSelectedObjects()->Deselect(Emitter);
+	}
+
+	ResetLoaders(Emitter->GetOutermost()); // Make sure that we're not going to get invalid version number linkers into the package we are going into. 
+	Emitter->GetOutermost()->LinkerCustomVersion.Empty();
+
+	TSharedPtr<FNiagaraEmitterViewModel> EditableEmitterViewModel = SystemViewModel->GetEmitterHandleViewModels()[0]->GetEmitterViewModel();
+	UNiagaraEmitter* EditableEmitter = EditableEmitterViewModel->GetEmitter();
+
+	// overwrite the original script in place by constructing a new one with the same name
+	Emitter = (UNiagaraEmitter*)StaticDuplicateObject(EditableEmitter, Emitter->GetOuter(), Emitter->GetFName(),
+		RF_AllFlags,
+		Emitter->GetClass());
+
+	// Restore RF_Standalone on the original material, as it had been removed from the preview material so that it could be GC'd.
+	Emitter->SetFlags(RF_Standalone);
+
+	TArray<UNiagaraEmitter*> AffectedEmitters;
+	AffectedEmitters.Add(Emitter);
+	UpdateExistingEmitters(AffectedEmitters);
+
+	GWarn->EndSlowTask();
+	EditableEmitterViewModel->SetDirty(false);
+}
+
+void FNiagaraSystemToolkit::UpdateExistingEmitters(TArray<UNiagaraEmitter*> AffectedEmitters)
+{
+	// Compile the existing emitters. Also determine which Systems need to be properly updated.
+	TArray<UNiagaraSystem*> AffectedSystemSystems;
+	for (UNiagaraEmitter* AffectedEmitter : AffectedEmitters)
+	{
+		if (AffectedEmitter->IsPendingKillOrUnreachable())
+		{
+			continue;
+		}
+
+		TSharedPtr<FNiagaraEmitterViewModel> EmitterViewModel = FNiagaraEmitterViewModel::GetExistingViewModelForObject(AffectedEmitter);
+		if (!EmitterViewModel.IsValid())
+		{
+			EmitterViewModel = MakeShareable(new FNiagaraEmitterViewModel(AffectedEmitter, nullptr));
+		}
+		EmitterViewModel->CompileScripts();
+
+		for (TObjectIterator<UNiagaraSystem> It; It; ++It)
+		{
+			if (It->GetAutoImportChangedEmitters() && It->ReferencesSourceEmitter(AffectedEmitter))
+			{
+				AffectedSystemSystems.AddUnique(*It);
+			}
+		}
+	}
+
+	// Now iterate over the affected Systems.
+	for (UNiagaraSystem* System : AffectedSystemSystems)
+	{
+		TSharedPtr<FNiagaraSystemViewModel> SystemViewModel = FNiagaraSystemViewModel::GetExistingViewModelForObject(System);
+		if (!SystemViewModel.IsValid())
+		{
+			FNiagaraSystemViewModelOptions Options;
+			Options.bCanRemoveEmittersFromTimeline = false;
+			Options.bCanRenameEmittersFromTimeline = false;
+			Options.bCanAddEmittersFromTimeline = false;
+			Options.bUseSystemExecStateForTimelineReset = false;
+			SystemViewModel = MakeShareable(new FNiagaraSystemViewModel(*System, Options));
+		}
+
+		SystemViewModel->ResynchronizeAllHandles();
+	}
+}
+
+void FNiagaraSystemToolkit::GetSaveableObjects(TArray<UObject*>& OutObjects) const
+{
+	if (SystemToolkitMode == ESystemToolkitMode::Emitter)
+	{
+		OutObjects.Add(Emitter);
+	}
+	else
+	{
+		FAssetEditorToolkit::GetSaveableObjects(OutObjects);
+	}
+}
+
+void FNiagaraSystemToolkit::SaveAsset_Execute()
+{
+	if (SystemToolkitMode == ESystemToolkitMode::Emitter)
+	{
+		UE_LOG(LogNiagaraEditor, Log, TEXT("Saving and Compiling NiagaraEmitter %s"), *GetEditingObjects()[0]->GetName());
+		UpdateOriginalEmitter();
+	}
+
+	FAssetEditorToolkit::SaveAsset_Execute();
+}
+
+void FNiagaraSystemToolkit::SaveAssetAs_Execute()
+{
+	if (SystemToolkitMode == ESystemToolkitMode::Emitter)
+	{
+		UE_LOG(LogNiagaraEditor, Log, TEXT("Saving and Compiling NiagaraEmitter %s"), *GetEditingObjects()[0]->GetName());
+		UpdateOriginalEmitter();
+	}
+
+	FAssetEditorToolkit::SaveAssetAs_Execute();
+}
+
+bool FNiagaraSystemToolkit::OnRequestClose()
+{
+	if (SystemToolkitMode == ESystemToolkitMode::Emitter)
+	{
+		TSharedPtr<FNiagaraEmitterViewModel> EmitterViewModel = SystemViewModel->GetEmitterHandleViewModels()[0]->GetEmitterViewModel();
+		if (EmitterViewModel->GetDirty())
+		{
+			// find out the user wants to do with this dirty NiagaraScript
+			EAppReturnType::Type YesNoCancelReply = FMessageDialog::Open(EAppMsgType::YesNoCancel,
+				FText::Format(
+					NSLOCTEXT("UnrealEd", "Prompt_NiagaraEmitterEditorClose", "Would you like to apply changes to this Emitter to the original Emitter?\n{0}\n(No will lose all changes!)"),
+					FText::FromString(Emitter->GetPathName())));
+
+			// act on it
+			switch (YesNoCancelReply)
+			{
+			case EAppReturnType::Yes:
+				// update NiagaraScript and exit
+				UpdateOriginalEmitter();
+				break;
+
+			case EAppReturnType::No:
+				// exit
+				// do nothing.. bNiagaraScriptDirty = false;
+				break;
+
+			case EAppReturnType::Cancel:
+				// don't exit
+				return false;
+			}
+		}
+		return true;
+	}
+
+	return FAssetEditorToolkit::OnRequestClose();
+}
+
+void FNiagaraSystemToolkit::EmitterAssetSelected(const FAssetData& AssetData)
+{
+	FSlateApplication::Get().DismissAllMenus();
+	SystemViewModel->AddEmitterFromAssetData(AssetData);
+}
+
+void FNiagaraSystemToolkit::ToggleUnlockToChanges()
+{
+	const FScopedTransaction ToggleUnlockToChangesTransaction(LOCTEXT("ToggleUnlockToChanges", "Toggle System Unlock To Changes"));
+	System->Modify();
+
+	System->SetAutoImportChangedEmitters(!System->GetAutoImportChangedEmitters());
+
+	if (System->GetAutoImportChangedEmitters())
+	{
+		SystemViewModel->ResynchronizeAllHandles();
+	}
+}
+
+bool FNiagaraSystemToolkit::IsToggleUnlockToChangesChecked()
+{
+	return System->GetAutoImportChangedEmitters();
+}
+
+FText FNiagaraSystemToolkit::GetEmitterLockToChangesLabel() const
+{
+	if (System->GetAutoImportChangedEmitters())
+	{
+		return LOCTEXT("EmitterUnlockToChangesLabel", "Changes Unlocked");
+	}
+	else
+	{
+		return LOCTEXT("EmitterLockToChangesLabel", "Changes Locked");
+	}
+}
+
+FText FNiagaraSystemToolkit::GetEmitterLockToChangesLabelTooltip() const
+{
+	if (System->GetAutoImportChangedEmitters())
+	{
+		return LOCTEXT("EmitterUnlockToChangesLabelTooltip", "If a source emitter changes, the changes will be imported into this System automatically.");
+	}
+	else
+	{
+		return LOCTEXT("EmitterLockToChangesLabelTooltip", "If a source emitter changes, the changes will NOT be imported into this System automatically.");
+	}
+}
+
+FSlateIcon FNiagaraSystemToolkit::GetEmitterLockToChangesIcon() const
+{
+	if (System->GetAutoImportChangedEmitters())
+	{
+		return FSlateIcon(FNiagaraEditorStyle::GetStyleSetName(), "NiagaraEditor.UnlockToChanges");
+	}
+	else
+	{
+		return FSlateIcon(FNiagaraEditorStyle::GetStyleSetName(), "NiagaraEditor.LockToChanges");
+	}
+}
+
+void FNiagaraSystemToolkit::ToggleCompileEnabled()
+{
+	UNiagaraEditorSettings* Settings = GetMutableDefault<UNiagaraEditorSettings>();
+	Settings->bAutoCompile = !Settings->bAutoCompile;
+}
+
+bool FNiagaraSystemToolkit::IsAutoCompileEnabled()
+{
+	return GetDefault<UNiagaraEditorSettings>()->bAutoCompile;
+}
+
+#undef LOCTEXT_NAMESPACE
diff --git a/Engine/Plugins/FX/Niagara/Source/NiagaraEditor/Private/Toolkits/NiagaraSystemToolkit.h b/Engine/Plugins/FX/Niagara/Source/NiagaraEditor/Private/Toolkits/NiagaraSystemToolkit.h
new file mode 100644
--- /dev/null
+++ b/Engine/Plugins/FX/Niagara/Source/NiagaraEditor/Private/Toolkits/NiagaraSystemToolkit.h
@@ -0,0 +1,140 @@
+// Copyright 1998-2017 Epic Games, Inc. All Rights Reserved.
+
+#pragma once
+
+#include "CoreMinimal.h"
+#include "Misc/NotifyHook.h"
+#include "Input/Reply.h"
+#include "UObject/GCObject.h"
+#include "Toolkits/IToolkitHost.h"
+
+#include "AssetEditorToolkit.h"
+
+#include "ISequencer.h"
+#include "ISequencerTrackEditor.h"
+
+class FNiagaraSystemInstance;
+class FNiagaraSystemViewModel;
+class SNiagaraSystemEditorViewport;
+class SNiagaraSystemEditorWidget;
+class SNiagaraSystemViewport;
+class SNiagaraSystemEditor;
+class UNiagaraSystem;
+class UNiagaraEmitter;
+class UNiagaraSequence;
+struct FAssetData;
+class FMenuBuilder;
+class ISequencer;
+
+/** Viewer/editor for a NiagaraSystem
+*/
+class FNiagaraSystemToolkit : public FAssetEditorToolkit, public FGCObject
+{
+	enum class ESystemToolkitMode
+	{
+		System,
+		Emitter,
+	};
+
+public:
+	virtual void RegisterTabSpawners(const TSharedRef<class FTabManager>& TabManager) override;
+	virtual void UnregisterTabSpawners(const TSharedRef<class FTabManager>& TabManager) override;
+
+	/** Edits the specified Niagara System */
+	void InitializeWithSystem(const EToolkitMode::Type Mode, const TSharedPtr< class IToolkitHost >& InitToolkitHost, UNiagaraSystem& InSystem);
+
+	/** Edits the specified Niagara Emitter */
+	void InitializeWithEmitter(const EToolkitMode::Type Mode, const TSharedPtr< class IToolkitHost >& InitToolkitHost, UNiagaraEmitter& InEmitter);
+
+	/** Destructor */
+	virtual ~FNiagaraSystemToolkit();
+
+	//~ Begin IToolkit Interface
+	virtual FName GetToolkitFName() const override;
+	virtual FText GetBaseToolkitName() const override;
+	virtual FString GetWorldCentricTabPrefix() const override;
+	virtual FLinearColor GetWorldCentricTabColorScale() const override;
+	//~ End IToolkit Interface
+
+	//~ FGCObject interface
+	virtual void AddReferencedObjects(FReferenceCollector& Collector) override;
+
+	FSlateIcon GetCompileStatusImage() const;
+	FText GetCompileStatusTooltip() const;
+
+	/** Compiles the system script. */
+	void CompileSystem();
+
+protected:
+	//~ FAssetEditorToolkit interface
+	virtual void GetSaveableObjects(TArray<UObject*>& OutObjects) const override;
+	virtual void SaveAsset_Execute() override;
+	virtual void SaveAssetAs_Execute() override;
+	virtual bool OnRequestClose() override;
+
+private:
+	void InitializeInternal(const EToolkitMode::Type Mode, const TSharedPtr<class IToolkitHost>& InitToolkitHost);
+
+	void UpdateOriginalEmitter();
+	static void UpdateExistingEmitters(TArray<UNiagaraEmitter*> AffectedEmitters);
+
+	TSharedRef<SDockTab> SpawnTab_Viewport(const FSpawnTabArgs& Args);
+	TSharedRef<SDockTab> SpawnTab_CurveEd(const FSpawnTabArgs& Args);
+	TSharedRef<SDockTab> SpawnTab_Sequencer(const FSpawnTabArgs& Args);
+	TSharedRef<SDockTab> SpawnTab_SystemScript(const FSpawnTabArgs& Args);
+	TSharedRef<SDockTab> SpawnTab_SystemDetails(const FSpawnTabArgs& Args);
+	TSharedRef<SDockTab> SpawnTab_SelectedEmitterStack(const FSpawnTabArgs& Args);
+	TSharedRef<SDockTab> SpawnTab_SelectedEmitterGraph(const FSpawnTabArgs& Args);
+	TSharedRef<SDockTab> SpawnTab_DebugSpreadsheet(const FSpawnTabArgs& Args);
+	TSharedRef<SDockTab> SpawnTab_PreviewSettings(const FSpawnTabArgs& Args);
+	TSharedRef<SDockTab> SpawnTab_GeneratedCode(const FSpawnTabArgs& Args);
+
+	/** Builds the toolbar widget */
+	void ExtendToolbar();	
+	void SetupCommands();
+
+	void ResetSimulation();
+
+	void GetSequencerAddMenuContent(FMenuBuilder& MenuBuilder, TSharedRef<ISequencer> Sequencer);
+	TSharedRef<SWidget> CreateAddEmitterMenuContent();
+	static TSharedRef<SWidget> GenerateCompileMenuContent();
+
+	void EmitterAssetSelected(const FAssetData& AssetData);
+	void ToggleUnlockToChanges();
+	bool IsToggleUnlockToChangesChecked();
+
+	FText GetEmitterLockToChangesLabel() const;
+	FText GetEmitterLockToChangesLabelTooltip() const;
+	FSlateIcon GetEmitterLockToChangesIcon() const;
+
+	static void ToggleCompileEnabled();
+	static bool IsAutoCompileEnabled();
+
+private:
+	/** The System being edited in system mode, or the placeholder system being edited in emitter mode. */
+	UNiagaraSystem* System;
+
+	/** The emitter being edited in emitter mode, or null when editing in system mode. */
+	UNiagaraEmitter* Emitter;
+
+	ESystemToolkitMode SystemToolkitMode;
+
+	TSharedPtr<SNiagaraSystemViewport> Viewport;
+
+	/* The view model for the System being edited */
+	TSharedPtr<FNiagaraSystemViewModel> SystemViewModel;
+
+	/** The command list for this editor */
+	TSharedPtr<FUICommandList> EditorCommands;
+
+	static const FName ViewportTabID;
+	static const FName CurveEditorTabID;
+	static const FName SequencerTabID;
+	static const FName SystemScriptTabID;
+	static const FName SystemDetailsTabID;
+	static const FName SelectedEmitterStackTabID;
+	static const FName SelectedEmitterGraphTabID;
+	static const FName DebugSpreadsheetTabID;
+	static const FName PreviewSettingsTabId;
+	static const FName GeneratedCodeTabID;
+};
diff --git a/Engine/Plugins/FX/Niagara/Source/NiagaraEditor/Private/TypeEditorUtilities/NiagaraBoolTypeEditorUtilities.cpp b/Engine/Plugins/FX/Niagara/Source/NiagaraEditor/Private/TypeEditorUtilities/NiagaraBoolTypeEditorUtilities.cpp
new file mode 100644
--- /dev/null
+++ b/Engine/Plugins/FX/Niagara/Source/NiagaraEditor/Private/TypeEditorUtilities/NiagaraBoolTypeEditorUtilities.cpp
@@ -0,0 +1,93 @@
+// Copyright 1998-2016 Epic Games, Inc. All Rights Reserved.
+
+#include "NiagaraBoolTypeEditorUtilities.h"
+#include "NiagaraTypes.h"
+#include "SNiagaraParameterEditor.h"
+#include "DeclarativeSyntaxSupport.h"
+#include "SBoxPanel.h"
+#include "SCheckBox.h"
+
+class SNiagaraBoolParameterEditor : public SNiagaraParameterEditor
+{
+public:
+	SLATE_BEGIN_ARGS(SNiagaraBoolParameterEditor) { }
+	SLATE_END_ARGS();
+
+	void Construct(const FArguments& InArgs)
+	{
+		ChildSlot
+		[
+			SNew(SHorizontalBox)
+			+ SHorizontalBox::Slot()
+			.Padding(0)
+			.AutoWidth()
+			[
+				SNew(SCheckBox)
+				.IsChecked(this, &SNiagaraBoolParameterEditor::GetCheckState)
+				.OnCheckStateChanged(this, &SNiagaraBoolParameterEditor::OnCheckStateChanged)
+			]
+		];
+	}
+
+	virtual void UpdateInternalValueFromStruct(TSharedRef<FStructOnScope> Struct) override
+	{
+		checkf(Struct->GetStruct() == FNiagaraTypeDefinition::GetBoolStruct(), TEXT("Struct type not supported."));
+		bBoolValue = ((FNiagaraBool*)Struct->GetStructMemory())->GetValue();
+	}
+
+	virtual void UpdateStructFromInternalValue(TSharedRef<FStructOnScope> Struct) override
+	{
+		// Note that while bool conventionally have false = 0 and true = 1 (or any non-zero value), Niagara internally uses true == -1. Make 
+		// sure to enforce this convention when setting the value in memory.
+		checkf(Struct->GetStruct() == FNiagaraTypeDefinition::GetBoolStruct(), TEXT("Struct type not supported."));
+		((FNiagaraBool*)Struct->GetStructMemory())->SetValue(bBoolValue);
+	}
+
+private:
+	ECheckBoxState GetCheckState() const
+	{
+		return bBoolValue ? ECheckBoxState::Checked : ECheckBoxState::Unchecked;
+	}
+
+	void OnCheckStateChanged(ECheckBoxState InCheckState)
+	{
+		bBoolValue = InCheckState == ECheckBoxState::Checked;
+		ExecuteOnValueChanged();
+	}
+
+private:
+	bool bBoolValue;
+};
+
+TSharedPtr<SNiagaraParameterEditor> FNiagaraEditorBoolTypeUtilities::CreateParameterEditor() const
+{
+	return SNew(SNiagaraBoolParameterEditor);
+}
+
+bool FNiagaraEditorBoolTypeUtilities::CanHandlePinDefaults() const
+{
+	return true;
+}
+
+FString FNiagaraEditorBoolTypeUtilities::GetPinDefaultStringFromValue(const FNiagaraVariable& Variable) const
+{
+	bool bValue = false;
+	if (Variable.IsDataAllocated())
+	{
+		bValue = Variable.GetValue<FNiagaraBool>()->GetValue();
+	}
+	return bValue ? TEXT("true") : TEXT("false");
+}
+
+void FNiagaraEditorBoolTypeUtilities::SetValueFromPinDefaultString(const FString& StringValue, FNiagaraVariable& Variable) const
+{
+	Variable.AllocateData();
+	if (StringValue == TEXT("true"))
+	{
+		Variable.GetValue<FNiagaraBool>()->SetValue(true);
+	}
+	else
+	{
+		Variable.GetValue<FNiagaraBool>()->SetValue(false);
+	}
+}
\ No newline at end of file
diff --git a/Engine/Plugins/FX/Niagara/Source/NiagaraEditor/Private/TypeEditorUtilities/NiagaraBoolTypeEditorUtilities.h b/Engine/Plugins/FX/Niagara/Source/NiagaraEditor/Private/TypeEditorUtilities/NiagaraBoolTypeEditorUtilities.h
new file mode 100644
--- /dev/null
+++ b/Engine/Plugins/FX/Niagara/Source/NiagaraEditor/Private/TypeEditorUtilities/NiagaraBoolTypeEditorUtilities.h
@@ -0,0 +1,19 @@
+// Copyright 1998-2016 Epic Games, Inc. All Rights Reserved.
+
+#pragma once
+
+#include "INiagaraEditorTypeUtilities.h"
+
+class SNiagaraParameterEditor;
+
+/** Niagara editor utilities for the bool type. */
+class FNiagaraEditorBoolTypeUtilities : public FNiagaraEditorTypeUtilities
+{
+public:
+	//~ INiagaraEditorTypeUtilities interface.
+	virtual bool CanCreateParameterEditor() const override { return true; }
+	virtual TSharedPtr<SNiagaraParameterEditor> CreateParameterEditor() const override;
+	virtual bool CanHandlePinDefaults() const override;
+	virtual FString GetPinDefaultStringFromValue(const FNiagaraVariable& Variable) const override;
+	virtual void SetValueFromPinDefaultString(const FString& StringValue, FNiagaraVariable& Variable) const override;
+};
\ No newline at end of file
diff --git a/Engine/Plugins/FX/Niagara/Source/NiagaraEditor/Private/TypeEditorUtilities/NiagaraColorTypeEditorUtilities.cpp b/Engine/Plugins/FX/Niagara/Source/NiagaraEditor/Private/TypeEditorUtilities/NiagaraColorTypeEditorUtilities.cpp
new file mode 100644
--- /dev/null
+++ b/Engine/Plugins/FX/Niagara/Source/NiagaraEditor/Private/TypeEditorUtilities/NiagaraColorTypeEditorUtilities.cpp
@@ -0,0 +1,224 @@
+// Copyright 1998-2016 Epic Games, Inc. All Rights Reserved.
+
+#include "NiagaraColorTypeEditorUtilities.h"
+#include "SNiagaraParameterEditor.h"
+#include "NiagaraTypes.h"
+#include "NiagaraEditorStyle.h"
+#include "Classes/Engine/Engine.h"
+
+#include "SColorPicker.h"
+#include "SNumericEntryBox.h"
+#include "SColorBlock.h"
+#include "SGridPanel.h"
+
+class SNiagaraColorParameterEditor : public SNiagaraParameterEditor
+{
+public:
+	SLATE_BEGIN_ARGS(SNiagaraColorParameterEditor) { }
+	SLATE_END_ARGS();
+
+	void Construct(const FArguments& InArgs)
+	{
+		ChildSlot
+		[
+			SNew(SGridPanel)
+			.FillColumn(1, 1)
+			.FillColumn(2, 1)
+			.FillColumn(3, 1)
+			+ SGridPanel::Slot(0, 0)
+			.RowSpan(2)
+			[
+				SAssignNew(ColorBlock, SColorBlock)
+				.Color(this, &SNiagaraColorParameterEditor::GetColor)
+				.ShowBackgroundForAlpha(true)
+				.IgnoreAlpha(false)
+				.OnMouseButtonDown(this, &SNiagaraColorParameterEditor::OnMouseButtonDownColorBlock)
+				.Size(FVector2D(35.0f, 12.0f))
+			]
+			+ SGridPanel::Slot(1, 0)
+			.Padding(3, 0, 0, 0)
+			[
+				ConstructComponentWidget(0, NSLOCTEXT("ColorParameterEditor", "RLabel", "R"))
+			]
+			+ SGridPanel::Slot(2, 0)
+			.Padding(3, 0, 0, 0)
+			[
+				ConstructComponentWidget(1, NSLOCTEXT("ColorParameterEditor", "GLabel", "G"))
+			]
+			+ SGridPanel::Slot(3, 0)
+			.Padding(3, 0, 0, 0)
+			[
+				ConstructComponentWidget(2, NSLOCTEXT("ColorParameterEditor", "BLabel", "B"))
+			]
+			+ SGridPanel::Slot(1, 1)
+			.Padding(3, 2, 0, 0)
+			[
+				ConstructComponentWidget(3, NSLOCTEXT("ColorParameterEditor", "ALabel", "A"))
+			]
+		];
+	}
+
+	virtual void UpdateInternalValueFromStruct(TSharedRef<FStructOnScope> Struct) override
+	{
+		checkf(Struct->GetStruct() == FNiagaraTypeDefinition::GetColorStruct(), TEXT("Struct type not supported."));
+		ColorValue = *((FLinearColor*)Struct->GetStructMemory());
+	}
+
+	virtual void UpdateStructFromInternalValue(TSharedRef<FStructOnScope> Struct) override
+	{
+		checkf(Struct->GetStruct() == FNiagaraTypeDefinition::GetColorStruct(), TEXT("Struct type not supported."));
+		*((FLinearColor*)Struct->GetStructMemory()) = ColorValue;
+	}
+
+private:
+	TSharedRef<SWidget> ConstructComponentWidget(int32 Index, FText ComponentLabel)
+	{
+		return SNew(SNumericEntryBox<float>)
+		.Font(FNiagaraEditorStyle::Get().GetFontStyle("NiagaraEditor.ParameterFont"))
+		.OverrideTextMargin(2)
+		.MinValue(TOptional<float>())
+		.MaxValue(TOptional<float>())
+		.MaxSliderValue(TOptional<float>())
+		.MinSliderValue(TOptional<float>())
+		.Delta(0.0f)
+		.Value(this, &SNiagaraColorParameterEditor::GetComponentValue, Index)
+		.OnValueChanged(this, &SNiagaraColorParameterEditor::ComponentValueChanged, Index)
+		.OnValueCommitted(this, &SNiagaraColorParameterEditor::ComponentValueCommitted, Index)
+		.OnBeginSliderMovement(this, &SNiagaraColorParameterEditor::BeginSliderMovement)
+		.OnEndSliderMovement(this, &SNiagaraColorParameterEditor::EndSliderMovement)
+		.AllowSpin(true)
+		.LabelVAlign(EVerticalAlignment::VAlign_Center)
+		.Label()
+		[
+			SNew(STextBlock)
+			.TextStyle(FNiagaraEditorStyle::Get(), "NiagaraEditor.ParameterText")
+			.Text(ComponentLabel)
+		];
+	}
+
+	void BeginSliderMovement()
+	{
+		ExecuteOnBeginValueChange();
+	}
+
+	void EndSliderMovement(float Value)
+	{
+		ExecuteOnEndValueChange();
+	}
+
+	TOptional<float> GetComponentValue(int32 Index) const
+	{
+		return TOptional<float>(ColorValue.Component(Index));
+	}
+
+	void ComponentValueChanged(float ComponentValue, int32 Index)
+	{
+		ColorValue.Component(Index) = ComponentValue;
+		ExecuteOnValueChanged();
+	}
+
+	void ComponentValueCommitted(float ComponentValue, ETextCommit::Type CommitInfo, int32 Index)
+	{
+		if (CommitInfo == ETextCommit::OnEnter || CommitInfo == ETextCommit::OnUserMovedFocus)
+		{
+			ComponentValueChanged(ComponentValue, Index);
+		}
+	}
+
+	FReply OnMouseButtonDownColorBlock(const FGeometry& MyGeometry, const FPointerEvent& MouseEvent)
+	{
+		if (MouseEvent.GetEffectingButton() != EKeys::LeftMouseButton)
+		{
+			return FReply::Unhandled();
+		}
+
+		FColorPickerArgs PickerArgs;
+		{
+			PickerArgs.bUseAlpha = true;
+			PickerArgs.bOnlyRefreshOnMouseUp = false;
+			PickerArgs.bOnlyRefreshOnOk = false;
+			PickerArgs.DisplayGamma = TAttribute<float>::Create(TAttribute<float>::FGetter::CreateUObject(GEngine, &UEngine::GetDisplayGamma));
+			PickerArgs.OnColorCommitted = FOnLinearColorValueChanged::CreateSP(this, &SNiagaraColorParameterEditor::SetColor);
+			PickerArgs.OnColorPickerCancelled = FOnColorPickerCancelled::CreateSP(this, &SNiagaraColorParameterEditor::ColorPickerCancelled);
+			PickerArgs.OnInteractivePickBegin = FSimpleDelegate::CreateSP(this, &SNiagaraColorParameterEditor::InteractivePickBegin);
+			PickerArgs.OnInteractivePickEnd = FSimpleDelegate::CreateSP(this, &SNiagaraColorParameterEditor::InteractivePickEnd);
+			PickerArgs.OnColorPickerWindowClosed = FOnWindowClosed::CreateSP(this, &SNiagaraColorParameterEditor::ColorPickerClosed);
+			PickerArgs.InitialColorOverride = ColorValue;
+			PickerArgs.ParentWidget = ColorBlock;
+		}
+
+		OpenColorPicker(PickerArgs);
+		// Mark this parameter editor as editing exclusively so that the corresponding structure details view doesn't get updated
+		// since it closes all color pickers when it gets updated!
+		SetIsEditingExclusively(true);
+		return FReply::Handled();
+	}
+
+	void InteractivePickBegin()
+	{
+		ExecuteOnBeginValueChange();
+	}
+
+	void InteractivePickEnd()
+	{
+		ExecuteOnEndValueChange();
+	}
+
+	void ColorPickerCancelled(FLinearColor OriginalColor)
+	{
+		ColorValue = OriginalColor;
+		ExecuteOnValueChanged();
+	}
+
+	void ColorPickerClosed(const TSharedRef<SWindow>& Window)
+	{
+		SetIsEditingExclusively(false);
+	}
+
+	FLinearColor GetColor() const
+	{
+		return ColorValue;
+	}
+
+	void SetColor(FLinearColor NewColor)
+	{
+		ColorValue = NewColor;
+		ExecuteOnValueChanged();
+	}
+
+private:
+	TSharedPtr<SColorBlock> ColorBlock;
+
+	FLinearColor ColorValue;
+};
+
+void FNiagaraEditorColorTypeUtilities::UpdateStructWithDefaultValue(TSharedRef<FStructOnScope> Struct) const
+{
+	checkf(Struct->GetStruct() == FNiagaraTypeDefinition::GetColorStruct(), TEXT("Struct type not supported."));
+	*((FLinearColor*)Struct->GetStructMemory()) = FLinearColor(1, 1, 1, 1);
+}
+
+TSharedPtr<SNiagaraParameterEditor> FNiagaraEditorColorTypeUtilities::CreateParameterEditor() const
+{
+	return SNew(SNiagaraColorParameterEditor);
+}
+
+bool FNiagaraEditorColorTypeUtilities::CanHandlePinDefaults() const
+{
+	return true;
+}
+
+FString FNiagaraEditorColorTypeUtilities::GetPinDefaultStringFromValue(const FNiagaraVariable& Variable) const
+{
+	if (Variable.IsDataAllocated())
+	{
+		return Variable.GetValue<FLinearColor>()->ToString();
+	}
+	return FVector4(0.0f, 0.0f, 0.0f, 0.0f).ToString();
+}
+
+void FNiagaraEditorColorTypeUtilities::SetValueFromPinDefaultString(const FString& StringValue, FNiagaraVariable& Variable) const
+{
+	Variable.AllocateData();
+	Variable.GetValue<FLinearColor>()->InitFromString(StringValue);
+}
\ No newline at end of file
diff --git a/Engine/Plugins/FX/Niagara/Source/NiagaraEditor/Private/TypeEditorUtilities/NiagaraColorTypeEditorUtilities.h b/Engine/Plugins/FX/Niagara/Source/NiagaraEditor/Private/TypeEditorUtilities/NiagaraColorTypeEditorUtilities.h
new file mode 100644
--- /dev/null
+++ b/Engine/Plugins/FX/Niagara/Source/NiagaraEditor/Private/TypeEditorUtilities/NiagaraColorTypeEditorUtilities.h
@@ -0,0 +1,21 @@
+// Copyright 1998-2016 Epic Games, Inc. All Rights Reserved.
+
+#pragma once
+
+#include "INiagaraEditorTypeUtilities.h"
+
+class SNiagaraParameterEditor;
+
+/** Niagara editor utilities for the bool type. */
+class FNiagaraEditorColorTypeUtilities : public FNiagaraEditorTypeUtilities
+{
+public:
+	//~ INiagaraEditorTypeUtilities interface.
+	virtual bool CanProvideDefaultValue() const override { return true; }
+	virtual void UpdateStructWithDefaultValue(TSharedRef<FStructOnScope> Struct) const override;
+	virtual bool CanCreateParameterEditor() const override { return true; }
+	virtual TSharedPtr<SNiagaraParameterEditor> CreateParameterEditor() const override;
+	virtual bool CanHandlePinDefaults() const override;
+	virtual FString GetPinDefaultStringFromValue(const FNiagaraVariable& Variable) const override;
+	virtual void SetValueFromPinDefaultString(const FString& StringValue, FNiagaraVariable& Variable) const override;
+};
\ No newline at end of file
diff --git a/Engine/Plugins/FX/Niagara/Source/NiagaraEditor/Private/TypeEditorUtilities/NiagaraFloatTypeEditorUtilities.cpp b/Engine/Plugins/FX/Niagara/Source/NiagaraEditor/Private/TypeEditorUtilities/NiagaraFloatTypeEditorUtilities.cpp
new file mode 100644
--- /dev/null
+++ b/Engine/Plugins/FX/Niagara/Source/NiagaraEditor/Private/TypeEditorUtilities/NiagaraFloatTypeEditorUtilities.cpp
@@ -0,0 +1,112 @@
+// Copyright 1998-2016 Epic Games, Inc. All Rights Reserved.
+
+#include "NiagaraFloatTypeEditorUtilities.h"
+#include "SNiagaraParameterEditor.h"
+#include "NiagaraTypes.h"
+#include "NiagaraEditorStyle.h"
+
+#include "SSpinBox.h"
+
+class SNiagaraFloatParameterEditor : public SNiagaraParameterEditor
+{
+public:
+	SLATE_BEGIN_ARGS(SNiagaraFloatParameterEditor) { }
+	SLATE_END_ARGS();
+
+	void Construct(const FArguments& InArgs)
+	{
+		ChildSlot
+		[
+			SNew(SHorizontalBox)
+			+ SHorizontalBox::Slot()
+			.AutoWidth()
+			[
+				SNew(SSpinBox<float>)
+				.Style(FNiagaraEditorStyle::Get(), "NiagaraEditor.ParameterSpinBox")
+				.Font(FNiagaraEditorStyle::Get().GetFontStyle("NiagaraEditor.ParameterFont"))
+				.MinValue(TOptional<float>())
+				.MaxValue(TOptional<float>())
+				.MaxSliderValue(TOptional<float>())
+				.MinSliderValue(TOptional<float>())
+				.Delta(0.0f)
+				.Value(this, &SNiagaraFloatParameterEditor::GetValue)
+				.OnValueChanged(this, &SNiagaraFloatParameterEditor::ValueChanged)
+				.OnValueCommitted(this, &SNiagaraFloatParameterEditor::ValueCommitted)
+				.OnBeginSliderMovement(this, &SNiagaraFloatParameterEditor::BeginSliderMovement)
+				.OnEndSliderMovement(this, &SNiagaraFloatParameterEditor::EndSliderMovement)
+				.MinDesiredWidth(100)
+			]
+		];
+	}
+
+	virtual void UpdateInternalValueFromStruct(TSharedRef<FStructOnScope> Struct) override
+	{
+		checkf(Struct->GetStruct() == FNiagaraTypeDefinition::GetFloatStruct(), TEXT("Struct type not supported."));
+		FloatValue = ((FNiagaraFloat*)Struct->GetStructMemory())->Value;
+	}
+
+	virtual void UpdateStructFromInternalValue(TSharedRef<FStructOnScope> Struct) override
+	{
+		checkf(Struct->GetStruct() == FNiagaraTypeDefinition::GetFloatStruct(), TEXT("Struct type not supported."));
+		((FNiagaraFloat*)Struct->GetStructMemory())->Value = FloatValue;
+	}
+
+private:
+	void BeginSliderMovement()
+	{
+		ExecuteOnBeginValueChange();
+	}
+
+	void EndSliderMovement(float Value)
+	{
+		ExecuteOnEndValueChange();
+	}
+
+	float GetValue() const
+	{
+		return FloatValue;
+	}
+
+	void ValueChanged(float Value)
+	{
+		FloatValue = Value;
+		ExecuteOnValueChanged();
+	}
+
+	void ValueCommitted(float Value, ETextCommit::Type CommitInfo)
+	{
+		if (CommitInfo == ETextCommit::OnEnter || CommitInfo == ETextCommit::OnUserMovedFocus)
+		{
+			ValueChanged(Value);
+		}
+	}
+
+private:
+	float FloatValue;
+};
+
+TSharedPtr<SNiagaraParameterEditor> FNiagaraEditorFloatTypeUtilities::CreateParameterEditor() const
+{
+	return SNew(SNiagaraFloatParameterEditor);
+}
+
+bool FNiagaraEditorFloatTypeUtilities::CanHandlePinDefaults() const
+{
+	return true;
+}
+
+FString FNiagaraEditorFloatTypeUtilities::GetPinDefaultStringFromValue(const FNiagaraVariable& Variable) const
+{
+	float Value = 0.0f;
+	if (Variable.IsDataAllocated())
+	{
+		Value = Variable.GetValue<FNiagaraFloat>()->Value;
+	}
+	return FString::Printf(TEXT("%3.3f"), Value);
+}
+
+void FNiagaraEditorFloatTypeUtilities::SetValueFromPinDefaultString(const FString& StringValue, FNiagaraVariable& Variable) const
+{
+	Variable.AllocateData();
+	Variable.GetValue<FNiagaraFloat>()->Value = FCString::Atof(*StringValue);
+}
\ No newline at end of file
diff --git a/Engine/Plugins/FX/Niagara/Source/NiagaraEditor/Private/TypeEditorUtilities/NiagaraFloatTypeEditorUtilities.h b/Engine/Plugins/FX/Niagara/Source/NiagaraEditor/Private/TypeEditorUtilities/NiagaraFloatTypeEditorUtilities.h
new file mode 100644
--- /dev/null
+++ b/Engine/Plugins/FX/Niagara/Source/NiagaraEditor/Private/TypeEditorUtilities/NiagaraFloatTypeEditorUtilities.h
@@ -0,0 +1,19 @@
+// Copyright 1998-2016 Epic Games, Inc. All Rights Reserved.
+
+#pragma once
+
+#include "INiagaraEditorTypeUtilities.h"
+
+class SNiagaraParameterEditor;
+
+/** Niagara editor utilities for the float type. */
+class FNiagaraEditorFloatTypeUtilities : public FNiagaraEditorTypeUtilities
+{
+public:
+	//~ INiagaraEditorTypeUtilities interface.
+	virtual bool CanCreateParameterEditor() const override { return true; }
+	virtual TSharedPtr<SNiagaraParameterEditor> CreateParameterEditor() const override;
+	virtual bool CanHandlePinDefaults() const override;
+	virtual FString GetPinDefaultStringFromValue(const FNiagaraVariable& Variable) const override;
+	virtual void SetValueFromPinDefaultString(const FString& StringValue, FNiagaraVariable& Variable) const override;
+};
\ No newline at end of file
diff --git a/Engine/Plugins/FX/Niagara/Source/NiagaraEditor/Private/TypeEditorUtilities/NiagaraIntegerTypeEditorUtilities.cpp b/Engine/Plugins/FX/Niagara/Source/NiagaraEditor/Private/TypeEditorUtilities/NiagaraIntegerTypeEditorUtilities.cpp
new file mode 100644
--- /dev/null
+++ b/Engine/Plugins/FX/Niagara/Source/NiagaraEditor/Private/TypeEditorUtilities/NiagaraIntegerTypeEditorUtilities.cpp
@@ -0,0 +1,112 @@
+// Copyright 1998-2016 Epic Games, Inc. All Rights Reserved.
+
+#include "NiagaraIntegerTypeEditorUtilities.h"
+#include "SNiagaraParameterEditor.h"
+#include "NiagaraTypes.h"
+#include "NiagaraEditorStyle.h"
+
+#include "SSpinBox.h"
+
+class SNiagaraIntegerParameterEditor : public SNiagaraParameterEditor
+{
+public:
+	SLATE_BEGIN_ARGS(SNiagaraIntegerParameterEditor) { }
+	SLATE_END_ARGS();
+
+	void Construct(const FArguments& InArgs)
+	{
+		ChildSlot
+		[
+			SNew(SHorizontalBox)
+			+ SHorizontalBox::Slot()
+			.AutoWidth()
+			[
+				SNew(SSpinBox<int32>)
+				.Style(FNiagaraEditorStyle::Get(), "NiagaraEditor.ParameterSpinBox")
+				.Font(FNiagaraEditorStyle::Get().GetFontStyle("NiagaraEditor.ParameterFont"))
+				.MinValue(TOptional<int32>())
+				.MaxValue(TOptional<int32>())
+				.MaxSliderValue(TOptional<int32>())
+				.MinSliderValue(TOptional<int32>())
+				.Delta(0.001f)
+				.Value(this, &SNiagaraIntegerParameterEditor::GetValue)
+				.OnValueChanged(this, &SNiagaraIntegerParameterEditor::ValueChanged)
+				.OnValueCommitted(this, &SNiagaraIntegerParameterEditor::ValueCommitted)
+				.OnBeginSliderMovement(this, &SNiagaraIntegerParameterEditor::BeginSliderMovement)
+				.OnEndSliderMovement(this, &SNiagaraIntegerParameterEditor::EndSliderMovement)
+				.MinDesiredWidth(100)
+			]
+		];
+	}
+
+	virtual void UpdateInternalValueFromStruct(TSharedRef<FStructOnScope> Struct) override
+	{
+		checkf(Struct->GetStruct() == FNiagaraTypeDefinition::GetIntStruct(), TEXT("Struct type not supported."));
+		IntValue = ((FNiagaraInt32*)Struct->GetStructMemory())->Value;
+	}
+
+	virtual void UpdateStructFromInternalValue(TSharedRef<FStructOnScope> Struct) override
+	{
+		checkf(Struct->GetStruct() == FNiagaraTypeDefinition::GetIntStruct(), TEXT("Struct type not supported."));
+		((FNiagaraInt32*)Struct->GetStructMemory())->Value = IntValue;
+	}
+
+private:
+	void BeginSliderMovement()
+	{
+		ExecuteOnBeginValueChange();
+	}
+
+	void EndSliderMovement(int32 Value)
+	{
+		ExecuteOnEndValueChange();
+	}
+
+	int32 GetValue() const
+	{
+		return IntValue;
+	}
+
+	void ValueChanged(int32 Value)
+	{
+		IntValue = Value;
+		ExecuteOnValueChanged();
+	}
+
+	void ValueCommitted(int32 Value, ETextCommit::Type CommitInfo)
+	{
+		if (CommitInfo == ETextCommit::OnEnter || CommitInfo == ETextCommit::OnUserMovedFocus)
+		{
+			ValueChanged(Value);
+		}
+	}
+
+private:
+	int32 IntValue;
+};
+
+TSharedPtr<SNiagaraParameterEditor> FNiagaraEditorIntegerTypeUtilities::CreateParameterEditor() const
+{
+	return SNew(SNiagaraIntegerParameterEditor);
+}
+
+bool FNiagaraEditorIntegerTypeUtilities::CanHandlePinDefaults() const
+{
+	return true;
+}
+
+FString FNiagaraEditorIntegerTypeUtilities::GetPinDefaultStringFromValue(const FNiagaraVariable& Variable) const
+{
+	int32 Value = 0;
+	if (Variable.IsDataAllocated())
+	{
+		Value = Variable.GetValue<FNiagaraInt32>()->Value;
+	}
+	return FString::Printf(TEXT("%i"), Value);
+}
+
+void FNiagaraEditorIntegerTypeUtilities::SetValueFromPinDefaultString(const FString& StringValue, FNiagaraVariable& Variable) const
+{
+	Variable.AllocateData();
+	Variable.GetValue<FNiagaraInt32>()->Value = FCString::Atoi(*StringValue);
+}
\ No newline at end of file
diff --git a/Engine/Plugins/FX/Niagara/Source/NiagaraEditor/Private/TypeEditorUtilities/NiagaraIntegerTypeEditorUtilities.h b/Engine/Plugins/FX/Niagara/Source/NiagaraEditor/Private/TypeEditorUtilities/NiagaraIntegerTypeEditorUtilities.h
new file mode 100644
--- /dev/null
+++ b/Engine/Plugins/FX/Niagara/Source/NiagaraEditor/Private/TypeEditorUtilities/NiagaraIntegerTypeEditorUtilities.h
@@ -0,0 +1,18 @@
+// Copyright 1998-2016 Epic Games, Inc. All Rights Reserved.
+
+#pragma once
+
+#include "INiagaraEditorTypeUtilities.h"
+
+class SNiagaraParameterEditor;
+
+class FNiagaraEditorIntegerTypeUtilities : public FNiagaraEditorTypeUtilities
+{
+public:
+	//~ INiagaraEditorTypeUtilities interface.
+	virtual bool CanCreateParameterEditor() const override { return true; }
+	virtual TSharedPtr<SNiagaraParameterEditor> CreateParameterEditor() const override;
+	virtual bool CanHandlePinDefaults() const override;
+	virtual FString GetPinDefaultStringFromValue(const FNiagaraVariable& Variable) const override;
+	virtual void SetValueFromPinDefaultString(const FString& StringValue, FNiagaraVariable& Variable) const override;
+};
\ No newline at end of file
