Copying //UE4/Orion-Staging to //UE4/Main (Source: //Orion/Dev-General @ 3606378)

diff --git a/Engine/Plugins/FX/Niagara/Source/NiagaraEditor/Private/TypeEditorUtilities/NiagaraVectorTypeEditorUtilities.cpp b/Engine/Plugins/FX/Niagara/Source/NiagaraEditor/Private/TypeEditorUtilities/NiagaraVectorTypeEditorUtilities.cpp
new file mode 100644
--- /dev/null
+++ b/Engine/Plugins/FX/Niagara/Source/NiagaraEditor/Private/TypeEditorUtilities/NiagaraVectorTypeEditorUtilities.cpp
@@ -0,0 +1,332 @@
+// Copyright 1998-2016 Epic Games, Inc. All Rights Reserved.
+
+#include "NiagaraVectorTypeEditorUtilities.h"
+#include "SNiagaraParameterEditor.h"
+#include "NiagaraEditorStyle.h"
+#include "NiagaraTypes.h"
+
+#include "SNumericEntryBox.h"
+#include "Misc/DefaultValueHelper.h"
+
+class SNiagaraVectorParameterEditorBase : public SNiagaraParameterEditor
+{
+public:
+	SLATE_BEGIN_ARGS(SNiagaraVectorParameterEditorBase) { }
+		SLATE_ARGUMENT(int32, ComponentCount)
+	SLATE_END_ARGS();
+
+	void Construct(const FArguments& InArgs)
+	{
+		ComponentLabels.Add(NSLOCTEXT("VectorParameterEditor", "XLabel", "X"));
+		ComponentLabels.Add(NSLOCTEXT("VectorParameterEditor", "YLabel", "Y"));
+		ComponentLabels.Add(NSLOCTEXT("VectorParameterEditor", "ZLabel", "Z"));
+		ComponentLabels.Add(NSLOCTEXT("VectorParameterEditor", "WLabel", "W"));
+
+		TSharedRef<SHorizontalBox> ComponentBox = SNew(SHorizontalBox);
+		for (int32 ComponentIndex = 0; ComponentIndex < InArgs._ComponentCount; ++ComponentIndex)
+		{
+			ComponentBox->AddSlot()
+			.Padding(ComponentIndex == 0 ? 0 : 3, 0, 0, 0)
+			[
+				ConstructComponentWidget(ComponentIndex)
+			];
+		}
+
+		ChildSlot
+		[
+			ComponentBox
+		];
+	}
+
+protected:
+	virtual float GetValue(int32 Index) const = 0;
+
+	virtual void SetValue(int32 Index, float Value) = 0;
+
+private:
+	TSharedRef<SWidget> ConstructComponentWidget(int32 Index)
+	{
+		return SNew(SNumericEntryBox<float>)
+		.Font(FNiagaraEditorStyle::Get().GetFontStyle("NiagaraEditor.ParameterFont"))
+		.OverrideTextMargin(2)
+		.MinValue(TOptional<float>())
+		.MaxValue(TOptional<float>())
+		.MaxSliderValue(TOptional<float>())
+		.MinSliderValue(TOptional<float>())
+		.Delta(0.0f)
+		.Value(this, &SNiagaraVectorParameterEditorBase::GetValueInternal, Index)
+		.OnValueChanged(this, &SNiagaraVectorParameterEditorBase::ValueChanged, Index)
+		.OnValueCommitted(this, &SNiagaraVectorParameterEditorBase::ValueCommitted, Index)
+		.OnBeginSliderMovement(this, &SNiagaraVectorParameterEditorBase::BeginSliderMovement)
+		.OnEndSliderMovement(this, &SNiagaraVectorParameterEditorBase::EndSliderMovement)
+		.AllowSpin(true)
+		.LabelVAlign(EVerticalAlignment::VAlign_Center)
+		.Label()
+		[
+			SNew(STextBlock)
+			.TextStyle(FNiagaraEditorStyle::Get(), "NiagaraEditor.ParameterText")
+			.Text(ComponentLabels[Index])
+		];
+	}
+
+	void BeginSliderMovement()
+	{
+		ExecuteOnBeginValueChange();
+	}
+
+	void EndSliderMovement(float Value)
+	{
+		ExecuteOnEndValueChange();
+	}
+
+	TOptional<float> GetValueInternal(int32 Index) const
+	{
+		return TOptional<float>(GetValue(Index));
+	}
+
+	void ValueChanged(float Value, int32 Index)
+	{
+		SetValue(Index, Value);
+		ExecuteOnValueChanged();
+	}
+
+	void ValueCommitted(float Value, ETextCommit::Type CommitInfo, int32 Index)
+	{
+		if (CommitInfo == ETextCommit::OnEnter || CommitInfo == ETextCommit::OnUserMovedFocus)
+		{
+			ValueChanged(Value, Index);
+		}
+	}
+
+private:
+	TArray<FText> ComponentLabels;
+};
+
+
+class SNiagaraVector2ParameterEditor : public SNiagaraVectorParameterEditorBase
+{
+public:
+	SLATE_BEGIN_ARGS(SNiagaraVector2ParameterEditor) { }
+	SLATE_END_ARGS();
+
+	void Construct(const FArguments& InArgs)
+	{
+		SNiagaraVectorParameterEditorBase::Construct(
+			SNiagaraVectorParameterEditorBase::FArguments()
+			.ComponentCount(2));
+	}
+
+	virtual void UpdateInternalValueFromStruct(TSharedRef<FStructOnScope> Struct) override
+	{
+		checkf(Struct->GetStruct() == FNiagaraTypeDefinition::GetVec2Struct(), TEXT("Struct type not supported."));
+		VectorValue = *((FVector2D*)Struct->GetStructMemory());
+	}
+
+	virtual void UpdateStructFromInternalValue(TSharedRef<FStructOnScope> Struct) override
+	{
+		checkf(Struct->GetStruct() == FNiagaraTypeDefinition::GetVec2Struct(), TEXT("Struct type not supported."));
+		*((FVector2D*)Struct->GetStructMemory()) = VectorValue;
+	}
+
+protected:
+	virtual float GetValue(int32 Index) const override
+	{
+		return VectorValue[Index];
+	}
+
+	virtual void SetValue(int32 Index, float Value) override
+	{
+		VectorValue[Index] = Value;
+	}
+
+private:
+	FVector2D VectorValue;
+};
+
+TSharedPtr<SNiagaraParameterEditor> FNiagaraEditorVector2TypeUtilities::CreateParameterEditor() const
+{
+	return SNew(SNiagaraVector2ParameterEditor);
+}
+
+bool FNiagaraEditorVector2TypeUtilities::CanHandlePinDefaults() const
+{
+	return true;
+}
+
+FString FNiagaraEditorVector2TypeUtilities::GetPinDefaultStringFromValue(const FNiagaraVariable& Variable) const
+{
+	if (Variable.IsDataAllocated())
+	{
+		return Variable.GetValue<FVector2D>()->ToString();
+	}
+	return FVector2D(0.0f, 0.0f).ToString();
+}
+
+void FNiagaraEditorVector2TypeUtilities::SetValueFromPinDefaultString(const FString& StringValue, FNiagaraVariable& Variable) const
+{
+	Variable.AllocateData();
+	Variable.GetValue<FVector2D>()->InitFromString(StringValue);
+}
+
+class SNiagaraVector3ParameterEditor : public SNiagaraVectorParameterEditorBase
+{
+public:
+	SLATE_BEGIN_ARGS(SNiagaraVector3ParameterEditor) { }
+	SLATE_END_ARGS();
+
+	void Construct(const FArguments& InArgs)
+	{
+		SNiagaraVectorParameterEditorBase::Construct(
+			SNiagaraVectorParameterEditorBase::FArguments()
+			.ComponentCount(3));
+	}
+
+	virtual void UpdateInternalValueFromStruct(TSharedRef<FStructOnScope> Struct) override
+	{
+		checkf(Struct->GetStruct() == FNiagaraTypeDefinition::GetVec3Struct(), TEXT("Struct type not supported."));
+		VectorValue = *((FVector*)Struct->GetStructMemory());
+	}
+
+	virtual void UpdateStructFromInternalValue(TSharedRef<FStructOnScope> Struct) override
+	{
+		checkf(Struct->GetStruct() == FNiagaraTypeDefinition::GetVec3Struct(), TEXT("Struct type not supported."));
+		*((FVector*)Struct->GetStructMemory()) = VectorValue;
+	}
+
+protected:
+	virtual float GetValue(int32 Index) const override
+	{
+		return VectorValue[Index];
+	}
+
+	virtual void SetValue(int32 Index, float Value) override
+	{
+		VectorValue[Index] = Value;
+	}
+
+private:
+	FVector VectorValue;
+};
+
+TSharedPtr<SNiagaraParameterEditor> FNiagaraEditorVector3TypeUtilities::CreateParameterEditor() const
+{
+	return SNew(SNiagaraVector3ParameterEditor);
+}
+
+bool FNiagaraEditorVector3TypeUtilities::CanHandlePinDefaults() const
+{
+	return true;
+}
+
+FString FNiagaraEditorVector3TypeUtilities::GetPinDefaultStringFromValue(const FNiagaraVariable& Variable) const
+{
+	// NOTE: We can not use ToString() here since the vector pin control doesn't use the standard 'X=0,Y=0,Z=0' syntax.
+	FVector Value(0.0f, 0.0f, 0.0f);
+	if (Variable.IsDataAllocated())
+	{
+		Value = *Variable.GetValue<FVector>();
+	}
+	return FString::Printf(TEXT("%3.3f,%3.3f,%3.3f"), Value.X, Value.Y, Value.Z);
+}
+
+void FNiagaraEditorVector3TypeUtilities::SetValueFromPinDefaultString(const FString& StringValue, FNiagaraVariable& Variable) const
+{
+	// NOTE: We can not use InitFromString() here since the vector pin control doesn't use the standard 'X=0,Y=0,Z=0' syntax.
+	FVector Value;
+	if (FDefaultValueHelper::ParseVector(StringValue, Value) == false)
+	{
+		Value = FVector(0.0f, 0.0f, 0.0f);
+	}
+	Variable.AllocateData();
+	*Variable.GetValue<FVector>() = Value;
+}
+
+class SNiagaraVector4ParameterEditor : public SNiagaraVectorParameterEditorBase
+{
+public:
+	SLATE_BEGIN_ARGS(SNiagaraVector4ParameterEditor) { }
+	SLATE_END_ARGS();
+
+	void Construct(const FArguments& InArgs)
+	{
+		SNiagaraVectorParameterEditorBase::Construct(
+			SNiagaraVectorParameterEditorBase::FArguments()
+			.ComponentCount(4));
+	}
+
+	virtual void UpdateInternalValueFromStruct(TSharedRef<FStructOnScope> Struct) override
+	{
+		checkf(Struct->GetStruct() == FNiagaraTypeDefinition::GetVec4Struct(), TEXT("Struct type not supported."));
+		VectorValue = *((FVector4*)Struct->GetStructMemory());
+	}
+
+	virtual void UpdateStructFromInternalValue(TSharedRef<FStructOnScope> Struct) override
+	{
+		checkf(Struct->GetStruct() == FNiagaraTypeDefinition::GetVec4Struct(), TEXT("Struct type not supported."));
+		*((FVector4*)Struct->GetStructMemory()) = VectorValue;
+	}
+
+protected:
+	virtual float GetValue(int32 Index) const override
+	{
+		return VectorValue[Index];
+	}
+
+	virtual void SetValue(int32 Index, float Value) override
+	{
+		VectorValue[Index] = Value;
+	}
+
+private:
+	FVector4 VectorValue;
+};
+
+TSharedPtr<SNiagaraParameterEditor> FNiagaraEditorVector4TypeUtilities::CreateParameterEditor() const
+{
+	return SNew(SNiagaraVector4ParameterEditor);
+}
+
+bool FNiagaraEditorVector4TypeUtilities::CanHandlePinDefaults() const
+{
+	return true;
+}
+
+FString FNiagaraEditorVector4TypeUtilities::GetPinDefaultStringFromValue(const FNiagaraVariable& Variable) const
+{
+	// NOTE: We can not use ToString() here since the vector pin control doesn't use the standard 'X=0,Y=0,Z=0,W=0' syntax.
+	FVector4 Value(0.0f, 0.0f, 0.0f, 0.0f);
+	if (Variable.IsDataAllocated())
+	{
+		// FVector4 is aligned, so casting a pointer to it won't work
+		const uint8 *VarData = Variable.GetData();
+		FPlatformMemory::Memcpy(&Value, VarData, sizeof(FVector4));
+		//Value = *Variable.GetValue<FVector4>();
+	}
+	return FString::Printf(TEXT("%3.3f,%3.3f,%3.3f,%3.3f"), Value.X, Value.Y, Value.Z, Value.W);
+}
+
+void FNiagaraEditorVector4TypeUtilities::SetValueFromPinDefaultString(const FString& StringValue, FNiagaraVariable& Variable) const
+{
+	// NOTE: We can not use InitFromString() here since the vector pin control doesn't use the standard 'X=0,Y=0,Z=0,W=0' syntax.
+	FVector4 Value;
+	if (FDefaultValueHelper::ParseVector4(StringValue, Value) == false)
+	{
+		Value = FVector4(0.0f, 0.0f, 0.0f, 0.0f);
+	}
+
+	// FVector4 requires aligned memory, which cannot be safely guaranteed by the AllocateData call. Therefore, we cannot 
+	// write it as a FVector4. Instead, we must write component by component.
+	Variable.AllocateData();
+	void* ValueWrite = Variable.GetValue<FVector4>();
+	if (nullptr != ValueWrite)
+	{
+		((float*)ValueWrite)[0] = Value.X;
+		((float*)ValueWrite)[1] = Value.Y;
+		((float*)ValueWrite)[2] = Value.Z;
+		((float*)ValueWrite)[3] = Value.W;
+	}
+	else
+	{
+		check(!"UNKNOWNW");
+	}
+}
\ No newline at end of file
diff --git a/Engine/Plugins/FX/Niagara/Source/NiagaraEditor/Private/TypeEditorUtilities/NiagaraVectorTypeEditorUtilities.h b/Engine/Plugins/FX/Niagara/Source/NiagaraEditor/Private/TypeEditorUtilities/NiagaraVectorTypeEditorUtilities.h
new file mode 100644
--- /dev/null
+++ b/Engine/Plugins/FX/Niagara/Source/NiagaraEditor/Private/TypeEditorUtilities/NiagaraVectorTypeEditorUtilities.h
@@ -0,0 +1,43 @@
+// Copyright 1998-2016 Epic Games, Inc. All Rights Reserved.
+
+#pragma once
+
+#include "INiagaraEditorTypeUtilities.h"
+
+class SNiagaraParameterEditor;
+
+/** Niagara editor utilities for the FVector2D type. */
+class FNiagaraEditorVector2TypeUtilities : public FNiagaraEditorTypeUtilities
+{
+public:
+	//~ INiagaraEditorTypeUtilities interface.
+	virtual bool CanCreateParameterEditor() const override { return true; }
+	virtual TSharedPtr<SNiagaraParameterEditor> CreateParameterEditor() const override;
+	virtual bool CanHandlePinDefaults() const override;
+	virtual FString GetPinDefaultStringFromValue(const FNiagaraVariable& Variable) const override;
+	virtual void SetValueFromPinDefaultString(const FString& StringValue, FNiagaraVariable& Variable) const override;
+};
+
+/** Niagara editor utilities for the FVector type. */
+class FNiagaraEditorVector3TypeUtilities : public FNiagaraEditorTypeUtilities
+{
+public:
+	//~ INiagaraEditorTypeUtilities interface.
+	virtual bool CanCreateParameterEditor() const override { return true; }
+	virtual TSharedPtr<SNiagaraParameterEditor> CreateParameterEditor() const override;
+	virtual bool CanHandlePinDefaults() const override;
+	virtual FString GetPinDefaultStringFromValue(const FNiagaraVariable& Variable) const override;
+	virtual void SetValueFromPinDefaultString(const FString& StringValue, FNiagaraVariable& Variable) const override;
+};
+
+/** Niagara editor utilities for the FVector4 type. */
+class FNiagaraEditorVector4TypeUtilities : public FNiagaraEditorTypeUtilities
+{
+public:
+	//~ INiagaraEditorTypeUtilities interface.
+	virtual bool CanCreateParameterEditor() const override { return true; }
+	virtual TSharedPtr<SNiagaraParameterEditor> CreateParameterEditor() const override;
+	virtual bool CanHandlePinDefaults() const override;
+	virtual FString GetPinDefaultStringFromValue(const FNiagaraVariable& Variable) const override;
+	virtual void SetValueFromPinDefaultString(const FString& StringValue, FNiagaraVariable& Variable) const override;
+};
\ No newline at end of file
diff --git a/Engine/Plugins/FX/Niagara/Source/NiagaraEditor/Private/ViewModels/NiagaraEmitterHandleViewModel.cpp b/Engine/Plugins/FX/Niagara/Source/NiagaraEditor/Private/ViewModels/NiagaraEmitterHandleViewModel.cpp
new file mode 100644
--- /dev/null
+++ b/Engine/Plugins/FX/Niagara/Source/NiagaraEditor/Private/ViewModels/NiagaraEmitterHandleViewModel.cpp
@@ -0,0 +1,316 @@
+// Copyright 1998-2016 Epic Games, Inc. All Rights Reserved.
+
+#include "NiagaraEmitterHandleViewModel.h"
+#include "NiagaraSystem.h"
+#include "NiagaraEmitterHandle.h"
+#include "NiagaraEmitterViewModel.h"
+#include "NiagaraScriptViewModel.h"
+#include "NiagaraScriptGraphViewModel.h"
+#include "NiagaraObjectSelection.h"
+#include "NiagaraScriptSource.h"
+#include "NiagaraNodeOutput.h"
+#include "NiagaraGraph.h"
+#include "NiagaraNodeInput.h"
+#include "NiagaraScriptOutputCollectionViewModel.h"
+#include "NotificationManager.h"
+#include "SNotificationList.h"
+
+#include "ScopedTransaction.h"
+#include "AssetEditorManager.h"
+
+#define LOCTEXT_NAMESPACE "EmitterHandleViewModel"
+
+FNiagaraEmitterHandleViewModel::FNiagaraEmitterHandleViewModel(FNiagaraEmitterHandle* InEmitterHandle, TWeakPtr<FNiagaraEmitterInstance> InSimulation, UNiagaraSystem& InOwningSystem)
+	: EmitterHandle(InEmitterHandle)
+	, OwningSystem(InOwningSystem)
+	, EmitterViewModel(MakeShareable(new FNiagaraEmitterViewModel((InEmitterHandle ? InEmitterHandle->GetInstance() : nullptr), InSimulation)))
+{
+}
+
+FNiagaraEmitterHandleViewModel::~FNiagaraEmitterHandleViewModel()
+{
+}
+
+bool FNiagaraEmitterHandleViewModel::Set(FNiagaraEmitterHandle* InEmitterHandle, TWeakPtr<FNiagaraEmitterInstance> InSimulation, UNiagaraSystem& InOwningSystem)
+{
+	if (&OwningSystem != &InOwningSystem)
+	{
+		return false;
+	}
+
+	SetEmitterHandle(InEmitterHandle);
+	SetSimulation(InSimulation);
+	
+	UNiagaraEmitter* EmitterProperties = nullptr;
+	if (InEmitterHandle != nullptr)
+	{
+		EmitterProperties = InEmitterHandle->GetInstance();
+	}
+	return EmitterViewModel->Set(EmitterProperties, InSimulation);
+}
+
+
+void FNiagaraEmitterHandleViewModel::SetEmitterHandle(FNiagaraEmitterHandle* InEmitterHandle)
+{
+	EmitterHandle = InEmitterHandle;
+}
+
+void FNiagaraEmitterHandleViewModel::SetSimulation(TWeakPtr<FNiagaraEmitterInstance> InSimulation)
+{
+	EmitterViewModel->SetSimulation(InSimulation);
+}
+
+FGuid FNiagaraEmitterHandleViewModel::GetId() const
+{
+	if (EmitterHandle)
+	{
+		return EmitterHandle->GetId();
+	}
+	return FGuid();
+}
+
+FText FNiagaraEmitterHandleViewModel::GetIdText() const
+{
+	return FText::FromString( GetId().ToString() );
+}
+
+
+FText FNiagaraEmitterHandleViewModel::GetErrorText() const
+{
+	switch (EmitterViewModel->GetLatestCompileStatus())
+	{
+	case ENiagaraScriptCompileStatus::NCS_Unknown:
+	case ENiagaraScriptCompileStatus::NCS_BeingCreated:
+		return LOCTEXT("NiagaraEmitterHandleCompileStatusUnknown", "Needs compilation & refresh.");
+	case ENiagaraScriptCompileStatus::NCS_UpToDate:
+		return LOCTEXT("NiagaraEmitterHandleCompileStatusUpToDate", "Compiled");
+	default:
+		return LOCTEXT("NiagaraEmitterHandleCompileStatusError", "Error! Needs compilation & refresh.");
+	}
+}
+
+FSlateColor FNiagaraEmitterHandleViewModel::GetErrorTextColor() const
+{
+	switch (EmitterViewModel->GetLatestCompileStatus())
+	{
+	case ENiagaraScriptCompileStatus::NCS_Unknown:
+	case ENiagaraScriptCompileStatus::NCS_BeingCreated:
+		return FSlateColor(FLinearColor::Yellow);
+	case ENiagaraScriptCompileStatus::NCS_UpToDate:
+		return FSlateColor(FLinearColor::Green);
+	default:
+		return FSlateColor(FLinearColor::Red);
+	}
+}
+
+EVisibility FNiagaraEmitterHandleViewModel::GetErrorTextVisibility() const
+{
+	return EmitterViewModel->GetLatestCompileStatus() != ENiagaraScriptCompileStatus::NCS_UpToDate ? EVisibility::Visible : EVisibility::Collapsed;
+}
+
+
+FText FNiagaraEmitterHandleViewModel::GetSourceSynchronizationText() const
+{
+	
+	if (EmitterHandle->IsSynchronizedWithSource())
+	{
+		return LOCTEXT("SynchronizedWithSource", "Up-To-Date");
+	}
+	else
+	{
+		return LOCTEXT("NotSynchronizedWithSource", "Not Up-To-Date");
+	}
+}
+
+EVisibility FNiagaraEmitterHandleViewModel::GetSourceSynchronizationTextVisibility() const
+{
+	if (OwningSystem.GetAutoImportChangedEmitters())
+	{
+		return EVisibility::Collapsed;
+	}
+	else
+	{
+		return EVisibility::Visible;
+	}
+}
+
+FSlateColor FNiagaraEmitterHandleViewModel::GetSourceSynchronizationTextColor() const
+{
+	if (EmitterHandle->IsSynchronizedWithSource())
+	{
+		return FSlateColor::UseForeground();
+	}
+	else
+	{
+		return FSlateColor(FLinearColor::Yellow);
+	}
+}
+
+FName FNiagaraEmitterHandleViewModel::GetName() const
+{
+	if (EmitterHandle)
+	{
+		return EmitterHandle->GetName();
+	}
+	return FName();
+}
+
+void FNiagaraEmitterHandleViewModel::SetName(FName InName)
+{
+	if (EmitterHandle && EmitterHandle->GetName() == InName)
+	{
+		return;
+	}
+
+	if (EmitterHandle)
+	{
+		TSet<FName> OtherEmitterNames;
+		for (const FNiagaraEmitterHandle& OtherEmitterHandle : OwningSystem.GetEmitterHandles())
+		{
+			if (OtherEmitterHandle.GetId() != EmitterHandle->GetId())
+			{
+				OtherEmitterNames.Add(OtherEmitterHandle.GetName());
+			}
+		}
+		FName UniqueName = FNiagaraEditorUtilities::GetUniqueName(InName, OtherEmitterNames);
+
+		FScopedTransaction ScopedTransaction(NSLOCTEXT("NiagaraEmitterEditor", "EditEmitterNameTransaction", "Edit emitter name"));
+		OwningSystem.Modify();
+		EmitterHandle->SetName(UniqueName);
+		OnPropertyChangedDelegate.Broadcast();
+	}
+}
+
+FText FNiagaraEmitterHandleViewModel::GetNameText() const
+{
+	if (EmitterHandle)
+	{
+		return FText::FromName(EmitterHandle->GetName());
+	}
+	return FText();
+}
+
+void FNiagaraEmitterHandleViewModel::OnNameTextComitted(const FText& InText, ETextCommit::Type CommitInfo)
+{
+	SetName(*InText.ToString());
+}
+
+bool FNiagaraEmitterHandleViewModel::VerifyNameTextChanged(const FText& NewText, FText& OutErrorMessage)
+{
+	FName NewName = *NewText.ToString();
+	if (NewName == FName())
+	{
+		OutErrorMessage = NSLOCTEXT("NiagaraEmitterEditor", "NiagaraInputNameEmptyWarn", "Cannot have empty name!");
+		return false;
+	}
+	return true;
+}
+
+bool FNiagaraEmitterHandleViewModel::GetIsEnabled() const
+{
+	if (EmitterHandle)
+	{
+		return EmitterHandle->GetIsEnabled();
+	}
+	return false;
+}
+
+void FNiagaraEmitterHandleViewModel::SetIsEnabled(bool bInIsEnabled)
+{
+	if (EmitterHandle && EmitterHandle->GetIsEnabled() != bInIsEnabled)
+	{
+		FScopedTransaction ScopedTransaction(NSLOCTEXT("NiagaraEmitterEditor", "EditEmitterEnabled", "Change emitter enabled state"));
+		OwningSystem.Modify();
+		EmitterHandle->SetIsEnabled(bInIsEnabled);
+		OnPropertyChangedDelegate.Broadcast();
+	}
+}
+
+ECheckBoxState FNiagaraEmitterHandleViewModel::GetIsEnabledCheckState() const
+{
+	if (EmitterHandle)
+	{
+		return EmitterHandle->GetIsEnabled() ? ECheckBoxState::Checked : ECheckBoxState::Unchecked;
+	}
+	return ECheckBoxState::Undetermined;
+}
+
+void FNiagaraEmitterHandleViewModel::OnIsEnabledCheckStateChanged(ECheckBoxState InCheckState)
+{
+	SetIsEnabled(InCheckState == ECheckBoxState::Checked);
+}
+
+FNiagaraEmitterHandle* FNiagaraEmitterHandleViewModel::GetEmitterHandle()
+{
+	return EmitterHandle;
+}
+
+TSharedRef<FNiagaraEmitterViewModel> FNiagaraEmitterHandleViewModel::GetEmitterViewModel()
+{
+	return EmitterViewModel;
+}
+
+void FNiagaraEmitterHandleViewModel::CompileScripts()
+{
+	EmitterViewModel->CompileScripts();
+}
+
+
+
+void FNiagaraEmitterHandleViewModel::RefreshFromSource()
+{
+	FScopedTransaction ScopedTransaction(NSLOCTEXT("NiagaraEmitterEditor", "RefreshFromSource", "Reset emitter from source."));
+	OwningSystem.Modify();
+
+	if (EmitterHandle)
+	{
+		{
+			// Pull in changes to the emitter asset by copying the source scripts, compiling and then copying over parameter values
+			// where relevant.
+			if (EmitterHandle->RefreshFromSource() == false)
+			{
+				UE_LOG(LogNiagaraEditor, Error, TEXT("Failed to compile during refresh. Refresh cancelled. %s"), *EmitterHandle->GetSource()->GetPathName());
+
+				const FText NotificationText = FText::Format(LOCTEXT("FailedScriptRefresh", "'{0}' failed to refresh due to compile errors. Please see log."), FText::FromString(EmitterHandle->GetSource()->GetName()));
+
+				FNotificationInfo Info(NotificationText);
+				Info.bFireAndForget = true;
+				Info.bUseThrobber = true;
+				Info.bUseSuccessFailIcons = true;
+				Info.ExpireDuration = 10.0f;
+				FSlateNotificationManager::Get().AddNotification(Info);
+			}
+		}
+
+		EmitterViewModel->SetEmitter(EmitterHandle->GetInstance());
+	}
+	OnPropertyChangedDelegate.Broadcast();
+}
+
+void FNiagaraEmitterHandleViewModel::ResetToSource()
+{
+	FScopedTransaction ScopedTransaction(NSLOCTEXT("NiagaraEmitterEditor", "ResetToSource", "Reset emitter to source."));
+	OwningSystem.Modify();
+	
+	if (EmitterHandle)
+	{
+		EmitterHandle->ResetToSource();
+		EmitterViewModel->SetEmitter(EmitterHandle->GetInstance());
+	}
+	OnPropertyChangedDelegate.Broadcast();
+}
+
+void FNiagaraEmitterHandleViewModel::OpenSourceEmitter()
+{
+	if (EmitterHandle)
+	{
+		FAssetEditorManager::Get().OpenEditorForAsset(const_cast<UNiagaraEmitter*>(EmitterHandle->GetSource()));
+	}
+}
+
+FNiagaraEmitterHandleViewModel::FOnPropertyChanged& FNiagaraEmitterHandleViewModel::OnPropertyChanged()
+{
+	return OnPropertyChangedDelegate;
+}
+
+#undef LOCTEXT_NAMESPACE
\ No newline at end of file
diff --git a/Engine/Plugins/FX/Niagara/Source/NiagaraEditor/Private/ViewModels/NiagaraEmitterHandleViewModel.h b/Engine/Plugins/FX/Niagara/Source/NiagaraEditor/Private/ViewModels/NiagaraEmitterHandleViewModel.h
new file mode 100644
--- /dev/null
+++ b/Engine/Plugins/FX/Niagara/Source/NiagaraEditor/Private/ViewModels/NiagaraEmitterHandleViewModel.h
@@ -0,0 +1,113 @@
+// Copyright 1998-2016 Epic Games, Inc. All Rights Reserved.
+
+#pragma once
+
+#include "SharedPointer.h"
+#include "Delegate.h"
+#include "SlateEnums.h"
+#include "SlateTypes.h"
+#include "Layout/Visibility.h"
+
+class UNiagaraSystem;
+struct FNiagaraEmitterHandle;
+struct FNiagaraEmitterInstance;
+class FNiagaraEmitterViewModel;
+
+/** The view model for the FNiagaraEmitterEditorWidget. */
+class FNiagaraEmitterHandleViewModel : public TSharedFromThis<FNiagaraEmitterHandleViewModel>
+{
+public:
+	DECLARE_MULTICAST_DELEGATE(FOnPropertyChanged);
+public:
+	/** Creates a new emitter editor view model with the supplied emitter handle and simulation. */
+	FNiagaraEmitterHandleViewModel(FNiagaraEmitterHandle* InEmitterHandle, TWeakPtr<FNiagaraEmitterInstance> InSimulation, UNiagaraSystem& InOwningSystem);
+	
+	~FNiagaraEmitterHandleViewModel();
+
+	/** Reuses a the emitter editor view model with the supplied emitter handle and simulation.*/
+	bool Set(FNiagaraEmitterHandle* InEmitterHandle, TWeakPtr<FNiagaraEmitterInstance> InSimulation, UNiagaraSystem& InOwningSystem);
+
+	/** Sets the emitter handle.*/
+	void SetEmitterHandle(FNiagaraEmitterHandle* InEmitterHandle);
+
+	/** Sets the simulation for the emitter this handle references. */
+	void SetSimulation(TWeakPtr<FNiagaraEmitterInstance> InSimulation);
+
+	/** Gets the id of the emitter handle. */
+	FGuid GetId() const;
+	FText GetIdText() const;
+
+	/** Gets the name of the emitter handle. */
+	FName GetName() const;
+
+	/** Sets the name of the emitter handle. */
+	void SetName(FName InName);
+
+	/** Gets the text representation of the emitter handle name. */
+	FText GetNameText() const;
+
+	/** Called when the contents of the name text control is committed. */
+	void OnNameTextComitted(const FText& InText, ETextCommit::Type CommitInfo);
+
+	/** Prevent invalid name being set on emitter.*/
+	bool VerifyNameTextChanged(const FText& NewText, FText& OutErrorMessage);
+
+	/** Called to get the error state of the emitter handle.*/
+	FText GetErrorText() const;
+	EVisibility GetErrorTextVisibility() const;
+	FSlateColor GetErrorTextColor() const;
+
+	/** Called to get the synch state of the emitter handle to its source.*/
+	FText GetSourceSynchronizationText() const;
+	EVisibility GetSourceSynchronizationTextVisibility() const;
+	FSlateColor GetSourceSynchronizationTextColor() const;
+	bool IsSynchronized() const;
+
+	/** Gets whether or not this emitter handle is enabled. */
+	bool GetIsEnabled() const;
+
+	/** Sets whether or not this emitter handle is enabled. */
+	void SetIsEnabled(bool bInIsEnabled);
+
+	/** Gets the check state for the is enabled check box. */
+	ECheckBoxState GetIsEnabledCheckState() const;
+
+	/** Called when the check state of the enabled check box changes. */
+	void OnIsEnabledCheckStateChanged(ECheckBoxState InCheckState);
+
+	/** Gets the emitter handled being view and edited by this view model. */
+	FNiagaraEmitterHandle* GetEmitterHandle();
+
+	/** Gets the view model for the emitter this handle references. */
+	TSharedRef<FNiagaraEmitterViewModel> GetEmitterViewModel();
+
+	/** Compiles the spawn and update scripts. */
+	void CompileScripts();
+
+	/** Refreshes the copied emitter's graph and inputs from the source asset.  Input values
+		will be preserved. */
+	void RefreshFromSource();
+
+	/** Replaces the coped emitter instance with a fresh copy of the source emitter asset.  Any changes
+		to input parameters will be lost. */
+	void ResetToSource();
+
+	/** Opens the source emitter in a stand alone asset editor. */
+	void OpenSourceEmitter();
+
+	/** Gets a multicast delegate which is called any time a property on the handle changes. */
+	FOnPropertyChanged& OnPropertyChanged();
+
+private:
+	/** The emitter handle being displayed and edited by this view model. */
+	FNiagaraEmitterHandle* EmitterHandle;
+
+	/** The System which owns the handled being displayed and edited by this view model. */
+	UNiagaraSystem& OwningSystem;
+
+	/** The view model for emitter this handle references. */
+	TSharedRef<FNiagaraEmitterViewModel> EmitterViewModel;
+
+	/** A multicast delegate which is called any time a property on the handle changes. */
+	FOnPropertyChanged OnPropertyChangedDelegate;
+};
\ No newline at end of file
diff --git a/Engine/Plugins/FX/Niagara/Source/NiagaraEditor/Private/ViewModels/NiagaraEmitterViewModel.cpp b/Engine/Plugins/FX/Niagara/Source/NiagaraEditor/Private/ViewModels/NiagaraEmitterViewModel.cpp
new file mode 100644
--- /dev/null
+++ b/Engine/Plugins/FX/Niagara/Source/NiagaraEditor/Private/ViewModels/NiagaraEmitterViewModel.cpp
@@ -0,0 +1,221 @@
+// Copyright 1998-2016 Epic Games, Inc. All Rights Reserved.
+
+#include "NiagaraEmitterViewModel.h"
+#include "NiagaraEmitter.h"
+#include "NiagaraEmitterInstance.h"
+#include "NiagaraScriptSourceBase.h"
+#include "NiagaraScriptViewModel.h"
+#include "NiagaraScriptGraphViewModel.h"
+#include "NiagaraObjectSelection.h"
+#include "NiagaraNodeOutput.h"
+#include "NiagaraScriptSource.h"
+#include "NiagaraEmitterEditorData.h"
+#include "NiagaraParameterEditMode.h"
+
+#include "ScopedTransaction.h"
+
+#define LOCTEXT_NAMESPACE "EmitterEditorViewModel"
+
+template<> TMap<UNiagaraEmitter*, TArray<FNiagaraEmitterViewModel*>> TNiagaraViewModelManager<UNiagaraEmitter, FNiagaraEmitterViewModel>::ObjectsToViewModels{};
+
+const FText FNiagaraEmitterViewModel::StatsFormat = NSLOCTEXT("NiagaraEmitterViewModel", "StatsFormat", "{0} Particles | {1} ms | {2} MB");
+const float Megabyte = 1024.0f * 1024.0f;
+
+FNiagaraEmitterViewModel::FNiagaraEmitterViewModel(UNiagaraEmitter* InEmitter, TWeakPtr<FNiagaraEmitterInstance> InSimulation)
+	: Emitter(InEmitter)
+	, Simulation(InSimulation)
+	, SharedScriptViewModel(MakeShareable(new FNiagaraScriptViewModel(InEmitter, LOCTEXT("SharedDisplayName", "Graph"), ENiagaraParameterEditMode::EditAll)))
+	, bUpdatingSelectionInternally(false)
+	, LastEventScriptStatus(ENiagaraScriptCompileStatus::NCS_Unknown)
+	, bEmitterDirty(false)
+{
+	if (Emitter.IsValid() && Emitter->EventHandlerScriptProps.Num() != 0 && Emitter->EventHandlerScriptProps[0].Script && Emitter->EventHandlerScriptProps[0].Script->ByteCode.Num() != 0)
+	{
+		LastEventScriptStatus = ENiagaraScriptCompileStatus::NCS_UpToDate;
+	}
+
+	RegisteredHandle = RegisterViewModelWithMap(Emitter.Get(), this);
+}
+
+
+bool FNiagaraEmitterViewModel::Set(UNiagaraEmitter* InEmitter, TWeakPtr<FNiagaraEmitterInstance> InSimulation)
+{
+	SetEmitter(InEmitter);
+	SetSimulation(InSimulation);
+	return true;
+}
+
+FNiagaraEmitterViewModel::~FNiagaraEmitterViewModel()
+{
+	SharedScriptViewModel->GetGraphViewModel()->GetSelection()->OnSelectedObjectsChanged().RemoveAll(this);
+	UnregisterViewModelWithMap(RegisteredHandle);
+}
+
+
+void FNiagaraEmitterViewModel::SetEmitter(UNiagaraEmitter* InEmitter)
+{
+	UnregisterViewModelWithMap(RegisteredHandle);
+
+	Emitter = InEmitter;
+
+	RegisteredHandle = RegisterViewModelWithMap(InEmitter, this);
+
+	SharedScriptViewModel->SetScripts(InEmitter);
+
+	OnEmitterChanged().Broadcast();
+}
+
+
+void FNiagaraEmitterViewModel::SetSimulation(TWeakPtr<FNiagaraEmitterInstance> InSimulation)
+{
+	Simulation = InSimulation;
+}
+
+
+float FNiagaraEmitterViewModel::GetStartTime() const
+{
+	// TODO, reference values off the parameter store...
+	return 0.0f;
+}
+
+
+void FNiagaraEmitterViewModel::SetStartTime(float InStartTime)
+{
+	// TODO, reference values off the parameter store...
+}
+
+
+float FNiagaraEmitterViewModel::GetEndTime() const
+{
+	// TODO, reference values off the parameter store...
+	return 0.0f;
+}
+
+
+void FNiagaraEmitterViewModel::SetEndTime(float InEndTime)
+{
+	// TODO, reference values off the parameter store...
+}
+
+
+int32 FNiagaraEmitterViewModel::GetNumLoops() const
+{
+	// TODO, reference values off the parameter store...
+	return 0;
+}
+
+UNiagaraEmitter* FNiagaraEmitterViewModel::GetEmitter()
+{
+	return Emitter.Get();
+}
+
+
+FText FNiagaraEmitterViewModel::GetStatsText() const
+{
+	if (Simulation.IsValid())
+	{
+		TSharedPtr<FNiagaraEmitterInstance> SimInstance = Simulation.Pin();
+		check(SimInstance.IsValid());
+		return FText::Format(StatsFormat,
+			FText::AsNumber(SimInstance->GetNumParticles()),
+			FText::AsNumber(SimInstance->GetTotalCPUTime()),
+			FText::AsNumber(SimInstance->GetTotalBytesUsed() / Megabyte));
+	}
+	else
+	{
+		return LOCTEXT("InvalidSimulation", "Simulation is invalid.");
+	}
+}
+
+TSharedRef<FNiagaraScriptViewModel> FNiagaraEmitterViewModel::GetSharedScriptViewModel()
+{
+	return SharedScriptViewModel;
+}
+
+bool FNiagaraEmitterViewModel::GetDirty() const
+{
+	bool bDirty = SharedScriptViewModel->GetScriptDirty() || bEmitterDirty;
+	return bDirty;
+}
+
+void FNiagaraEmitterViewModel::SetDirty(bool bDirty)
+{
+	SharedScriptViewModel->SetScriptDirty(bDirty);
+	bEmitterDirty = bDirty;
+}
+
+const UNiagaraEmitterEditorData& FNiagaraEmitterViewModel::GetEditorData() const
+{
+	check(Emitter.IsValid());
+
+	const UNiagaraEmitterEditorData* EditorData = Cast<UNiagaraEmitterEditorData>(Emitter->EditorData);
+	if (EditorData == nullptr)
+	{
+		EditorData = GetDefault<UNiagaraEmitterEditorData>();
+	}
+	return *EditorData;
+}
+
+UNiagaraEmitterEditorData& FNiagaraEmitterViewModel::GetOrCreateEditorData()
+{
+	check(Emitter.IsValid());
+	UNiagaraEmitterEditorData* EditorData = Cast<UNiagaraEmitterEditorData>(Emitter->EditorData);
+	if (EditorData == nullptr)
+	{
+		EditorData = NewObject<UNiagaraEmitterEditorData>(Emitter.Get(), NAME_None, RF_Transactional);
+		Emitter->Modify();
+		Emitter->EditorData = EditorData;
+	}
+	return *EditorData;
+}
+
+void FNiagaraEmitterViewModel::CompileScripts()
+{
+	if (Emitter.IsValid())
+	{
+		TArray<ENiagaraScriptCompileStatus> CompileStatuses;
+		TArray<FString> CompileErrors;
+		TArray<FString> CompilePaths;
+		TArray<UNiagaraScript*> Scripts;
+		TArray<TPair<ENiagaraScriptUsage, int32> > Usages;
+		Emitter->CompileScripts(CompileStatuses, CompileErrors, CompilePaths, Scripts);
+
+		ENiagaraScriptCompileStatus AggregateStatus = ENiagaraScriptCompileStatus::NCS_UpToDate;
+		FString AggregateErrors;
+
+		for (int32 i = 0; i < CompileStatuses.Num(); i++)
+		{
+			AggregateStatus = FNiagaraEditorUtilities::UnionCompileStatus(AggregateStatus, CompileStatuses[i]);
+			AggregateErrors += CompilePaths[i] + TEXT(" ") + FNiagaraEditorUtilities::StatusToText(CompileStatuses[i]).ToString() + TEXT("\n");
+			AggregateErrors += CompileErrors[i] + TEXT("\n");
+		}
+		SharedScriptViewModel->UpdateCompileStatus(AggregateStatus, AggregateErrors, CompileStatuses, CompileErrors, CompilePaths, Scripts);
+	}
+	OnScriptCompiled().Broadcast();
+}
+
+
+
+ENiagaraScriptCompileStatus FNiagaraEmitterViewModel::GetLatestCompileStatus()
+{
+	ENiagaraScriptCompileStatus UnionStatus = SharedScriptViewModel->GetLatestCompileStatus();
+	return UnionStatus;
+}
+
+
+FNiagaraEmitterViewModel::FOnEmitterChanged& FNiagaraEmitterViewModel::OnEmitterChanged()
+{
+	return OnEmitterChangedDelegate;
+}
+
+FNiagaraEmitterViewModel::FOnPropertyChanged& FNiagaraEmitterViewModel::OnPropertyChanged()
+{
+	return OnPropertyChangedDelegate;
+}
+
+FNiagaraEmitterViewModel::FOnScriptCompiled& FNiagaraEmitterViewModel::OnScriptCompiled()
+{
+	return OnScriptCompiledDelegate;
+}
+
+#undef LOCTEXT_NAMESPACE
diff --git a/Engine/Plugins/FX/Niagara/Source/NiagaraEditor/Private/ViewModels/NiagaraEmitterViewModel.h b/Engine/Plugins/FX/Niagara/Source/NiagaraEditor/Private/ViewModels/NiagaraEmitterViewModel.h
new file mode 100644
--- /dev/null
+++ b/Engine/Plugins/FX/Niagara/Source/NiagaraEditor/Private/ViewModels/NiagaraEmitterViewModel.h
@@ -0,0 +1,115 @@
+// Copyright 1998-2016 Epic Games, Inc. All Rights Reserved.
+
+#pragma once
+
+#include "NiagaraScriptViewModel.h"
+
+class UNiagaraEmitter;
+class FNiagaraScriptViewModel;
+class FNiagaraScriptGraphViewModel;
+class UNiagaraEmitterEditorData;
+struct FNiagaraEmitterInstance;
+struct FNiagaraVariable;
+
+/** The view model for the UNiagaraEmitter objects */
+class FNiagaraEmitterViewModel : public TSharedFromThis<FNiagaraEmitterViewModel>,  public TNiagaraViewModelManager<UNiagaraEmitter, FNiagaraEmitterViewModel>
+{
+public:
+	DECLARE_MULTICAST_DELEGATE(FOnEmitterChanged);
+	DECLARE_MULTICAST_DELEGATE(FOnPropertyChanged);
+	DECLARE_MULTICAST_DELEGATE(FOnScriptCompiled);
+
+public:
+	/** Creates a new emitter editor view model with the supplied emitter handle and simulation. */
+	FNiagaraEmitterViewModel(UNiagaraEmitter* InEmitter, TWeakPtr<FNiagaraEmitterInstance> InSimulation);
+	virtual ~FNiagaraEmitterViewModel();
+
+	/** Reuse this view model with new parameters.*/
+	bool Set(UNiagaraEmitter* InEmitter, TWeakPtr<FNiagaraEmitterInstance> InSimulation);
+
+	/** Sets this view model to a different emitter. */
+	void SetEmitter(UNiagaraEmitter* InEmitter);
+
+	/** Sets the current simulation for the emitter. */
+	void SetSimulation(TWeakPtr<FNiagaraEmitterInstance> InSimulation);
+
+	/** Gets the start time for the emitter. */
+	float GetStartTime() const;
+
+	/** Sets the start time for the emitter. */
+	void SetStartTime(float InStartTime);
+
+	/** Gets the end time for the emitter. */
+	float GetEndTime() const;
+
+	/** Sets the end time for the emitter. */
+	void SetEndTime(float InEndTime);
+
+	/** Gets the number of loops for the emitter.  0 for infinite. */
+	int32 GetNumLoops() const;
+
+	/** Gets the emitter represented by this view model. */
+	UNiagaraEmitter* GetEmitter();
+
+	/** Gets text representing stats for the emitter. */
+	//~ TODO: Instead of a single string here, we should probably have separate controls with tooltips etc.
+	FText GetStatsText() const;
+	
+	/** Geta a view model for the update/spawn Script. */
+	TSharedRef<FNiagaraScriptViewModel> GetSharedScriptViewModel();
+
+	/** Compiles the spawn and update scripts. */
+	void CompileScripts();
+
+	/* Get the latest status of this view-model's script compilation.*/
+	ENiagaraScriptCompileStatus GetLatestCompileStatus();
+
+	/** Gets a multicast delegate which is called when the emitter for this view model changes to a different emitter. */
+	FOnEmitterChanged& OnEmitterChanged();
+
+	/** Gets a delegate which is called when a property on the emitter changes. */
+	FOnPropertyChanged& OnPropertyChanged();
+
+	/** Gets a delegate which is called when the shared script is compiled. */
+	FOnScriptCompiled& OnScriptCompiled();
+		
+	bool GetDirty() const;
+	void SetDirty(bool bDirty);
+
+	/** Gets editor specific data which can be stored per emitter.  If this data hasn't been created the default version will be returned. */
+	const UNiagaraEmitterEditorData& GetEditorData() const;
+
+	/** Gets editor specific data which is stored per emitter.  If this data hasn't been created then it will be created. */
+	UNiagaraEmitterEditorData& GetOrCreateEditorData();
+
+private:
+
+	/** The text format stats display .*/
+	static const FText StatsFormat;
+
+	/** The emitter object being displayed by the control .*/
+	TWeakObjectPtr<UNiagaraEmitter> Emitter;
+
+	/** The runtime simulation for the emitter being displayed by the control */
+	TWeakPtr<FNiagaraEmitterInstance> Simulation;
+	
+	/** The view model for the update/spawn/event script. */
+	TSharedRef<FNiagaraScriptViewModel> SharedScriptViewModel;
+
+	/** A flag to prevent reentrancy when updating selection sets. */
+	bool bUpdatingSelectionInternally;
+
+	/** A multicast delegate which is called whenever the emitter for this view model is changed to a different emitter. */
+	FOnEmitterChanged OnEmitterChangedDelegate;
+
+	/** A multicast delegate which is called whenever a property on the emitter changes. */
+	FOnPropertyChanged OnPropertyChangedDelegate;
+
+	FOnScriptCompiled OnScriptCompiledDelegate;
+
+	ENiagaraScriptCompileStatus LastEventScriptStatus;
+
+	bool bEmitterDirty;
+
+	TNiagaraViewModelManager<UNiagaraEmitter, FNiagaraEmitterViewModel>::Handle RegisteredHandle;
+};
\ No newline at end of file
diff --git a/Engine/Plugins/FX/Niagara/Source/NiagaraEditor/Private/ViewModels/NiagaraScriptInputCollectionViewModel.cpp b/Engine/Plugins/FX/Niagara/Source/NiagaraEditor/Private/ViewModels/NiagaraScriptInputCollectionViewModel.cpp
new file mode 100644
--- /dev/null
+++ b/Engine/Plugins/FX/Niagara/Source/NiagaraEditor/Private/ViewModels/NiagaraScriptInputCollectionViewModel.cpp
@@ -0,0 +1,607 @@
+// Copyright 1998-2016 Epic Games, Inc. All Rights Reserved.
+
+#include "NiagaraScriptInputCollectionViewModel.h"
+#include "NiagaraScript.h"
+#include "NiagaraEmitter.h"
+#include "NiagaraScriptSource.h"
+#include "NiagaraTypes.h"
+#include "NiagaraScriptParameterViewModel.h"
+#include "NiagaraGraph.h"
+#include "NiagaraNodeInput.h"
+#include "INiagaraEditorTypeUtilities.h"
+#include "NiagaraEditorUtilities.h"
+#include "NiagaraDataInterface.h"
+
+#include "ScopedTransaction.h"
+
+#define LOCTEXT_NAMESPACE "NiagaraScriptInputCollection"
+
+FText DisplayNameFormat = NSLOCTEXT("ScriptInputCollection", "DisplayNameFormat", "{0} Inputs");
+
+FNiagaraScriptInputCollectionViewModel::FNiagaraScriptInputCollectionViewModel(UNiagaraScript* InScript, FText InDisplayName, ENiagaraParameterEditMode InParameterEditMode)
+	: FNiagaraParameterCollectionViewModel(InParameterEditMode)
+	, DisplayName(FText::Format(DisplayNameFormat, InDisplayName))
+{
+	if (InScript)
+	{
+		Scripts.Add(InScript);
+	}
+
+	if (Scripts.Num() == 1 && Scripts[0].IsValid() && Scripts[0]->GetSource())
+	{
+		Graph = Cast<UNiagaraScriptSource>(Scripts[0]->GetSource())->NodeGraph;
+		bCanHaveNumericParameters = (Scripts[0]->IsStandaloneScript());
+	}
+	else
+	{
+		Graph = nullptr;
+		bCanHaveNumericParameters = true;
+	}
+
+	RefreshParameterViewModels();
+
+	if (Graph.IsValid())
+	{
+		OnGraphChangedHandle = Graph->AddOnGraphChangedHandler(
+			FOnGraphChanged::FDelegate::CreateRaw(this, &FNiagaraScriptInputCollectionViewModel::OnGraphChanged));
+	}
+}
+
+
+FNiagaraScriptInputCollectionViewModel::FNiagaraScriptInputCollectionViewModel(UNiagaraEmitter* InEmitter, FText InDisplayName, ENiagaraParameterEditMode InParameterEditMode)
+	: FNiagaraParameterCollectionViewModel(InParameterEditMode)
+	, DisplayName(FText::Format(DisplayNameFormat, InDisplayName))
+{
+	TArray<UNiagaraScript*> InScripts;
+	InEmitter->GetScripts(InScripts);
+	// Because of weak pointers, we need to copy ourselves..
+	for (UNiagaraScript* Script : InScripts)
+	{
+		Scripts.Add(Script);
+		ensure(Script->GetSource() == InEmitter->GraphSource);
+	}
+	UNiagaraScriptSource* Source = Cast<UNiagaraScriptSource>(InEmitter->GraphSource);
+
+	if (Source != nullptr)
+	{
+		Graph = Source->NodeGraph;
+		bCanHaveNumericParameters = false;
+	}
+	else
+	{
+		Graph = nullptr;
+		bCanHaveNumericParameters = false;
+	}
+
+	RefreshParameterViewModels();
+
+	if (Graph.IsValid())
+	{
+		OnGraphChangedHandle = Graph->AddOnGraphChangedHandler(
+			FOnGraphChanged::FDelegate::CreateRaw(this, &FNiagaraScriptInputCollectionViewModel::OnGraphChanged));
+	}
+}
+
+FNiagaraScriptInputCollectionViewModel::~FNiagaraScriptInputCollectionViewModel()
+{
+	for (int32 i = 0; i < ParameterViewModels.Num(); i++)
+	{
+		FNiagaraScriptParameterViewModel* ParameterViewModel = (FNiagaraScriptParameterViewModel*)(&ParameterViewModels[i].Get());
+		if (ParameterViewModel != nullptr)
+		{
+			ParameterViewModel->Reset();
+			ParameterViewModel->OnNameChanged().RemoveAll(this);
+			ParameterViewModel->OnTypeChanged().RemoveAll(this);
+			ParameterViewModel->OnDefaultValueChanged().RemoveAll(this);
+		}
+	}
+	ParameterViewModels.Empty();
+
+	if (Graph.IsValid())
+	{
+		Graph->RemoveOnGraphChangedHandler(OnGraphChangedHandle);
+	}
+}
+
+void FNiagaraScriptInputCollectionViewModel::SetScripts(TArray<UNiagaraScript*> InScripts)
+{
+	if (Graph.IsValid())
+	{
+		Graph->RemoveOnGraphChangedHandler(OnGraphChangedHandle);
+	}
+
+	Scripts.Empty();
+	for (UNiagaraScript* Script : InScripts)
+	{
+		Scripts.Add(Script);
+		check(Script->GetSource() == InScripts[0]->GetSource());
+	}
+
+	if (InScripts.Num() != 0 && Scripts[0].IsValid() && Scripts[0]->GetSource())
+	{
+		Graph = Cast<UNiagaraScriptSource>(Scripts[0]->GetSource())->NodeGraph;
+		OnGraphChangedHandle = Graph->AddOnGraphChangedHandler(
+			FOnGraphChanged::FDelegate::CreateRaw(this, &FNiagaraScriptInputCollectionViewModel::OnGraphChanged));
+		bCanHaveNumericParameters = Scripts[0]->IsStandaloneScript();
+	}
+	else
+	{
+		Graph = nullptr;
+		bCanHaveNumericParameters = true;
+	}
+
+	RefreshParameterViewModels();
+}
+
+FText FNiagaraScriptInputCollectionViewModel::GetDisplayName() const 
+{
+	return DisplayName;
+}
+
+FVector2D GetNewNodeLocation(UNiagaraGraph* Graph, UNiagaraNode* NewInputNode, float VerticalNodeOffset, float HorizontalNodeOffset)
+{
+	FVector2D PlacementLocation;
+	TArray<UNiagaraNodeInput*> InputNodes;
+	Graph->GetNodesOfClass(InputNodes);
+	if (InputNodes.Num() > 1)
+	{
+		// If there are input nodes, try to put it under the lowest one.
+		UNiagaraNodeInput* LowestNode = nullptr;
+		for (UNiagaraNodeInput* InputNode : InputNodes)
+		{
+			if (InputNode->Usage == ENiagaraInputNodeUsage::Parameter && InputNode != NewInputNode && (LowestNode == nullptr || InputNode->NodePosY > LowestNode->NodePosY))
+			{
+				LowestNode = InputNode;
+			}
+		}
+		
+		if (LowestNode)
+		{
+			PlacementLocation = FVector2D(
+				LowestNode->NodePosX,
+				LowestNode->NodePosY + VerticalNodeOffset);
+		}
+		else
+		{
+			PlacementLocation = FVector2D(0.0f, 0.0f);
+		}
+	}
+	else
+	{
+		TArray<UNiagaraNode*> Nodes;
+		Graph->GetNodesOfClass(Nodes);
+		if (Nodes.Num() > 0)
+		{
+			// If there are other nodes, try to put it to the left of the leftmost one.
+			UNiagaraNode* LeftmostNode = nullptr;
+			for (UNiagaraNode* Node : Nodes)
+			{
+				if (Node != NewInputNode && (LeftmostNode == nullptr || Node->NodePosX < LeftmostNode->NodePosX))
+				{
+					LeftmostNode = Node;
+				}
+			}
+			check(LeftmostNode);
+			PlacementLocation = FVector2D(
+				LeftmostNode->NodePosX - HorizontalNodeOffset,
+				LeftmostNode->NodePosY);
+		}
+	}
+	return PlacementLocation;
+}
+
+void FNiagaraScriptInputCollectionViewModel::AddParameter(TSharedPtr<FNiagaraTypeDefinition> ParameterType)
+{
+	if (Graph.IsValid())
+	{
+		FScopedTransaction ScopedTransaction(LOCTEXT("AddScriptInput", "Add script input"));
+		Graph->Modify();
+		FGraphNodeCreator<UNiagaraNodeInput> InputNodeCreator(*Graph);
+		UNiagaraNodeInput* InputNode = InputNodeCreator.CreateNode();
+
+		FNiagaraEditorUtilities::InitializeParameterInputNode(*InputNode, *ParameterType.Get(), Graph.Get());
+
+		FVector2D PlacementLocation = GetNewNodeLocation(Graph.Get(), InputNode, 100, 150);
+		InputNode->NodePosX = PlacementLocation.X;
+		InputNode->NodePosY = PlacementLocation.Y;
+
+		InputNodeCreator.Finalize();
+
+		// The CreateNode notified the graph has changed, but changing the name and type will also need to signal the graph changed event...
+		// We need to do this because NiagaraSystemScriptView model is listening for these changes to update the bindings table. This 
+		// will also cause the RefreshParameterViewModels in our own graph changed handler.
+		Graph->NotifyGraphChanged();
+
+		for (TSharedRef<INiagaraParameterViewModel> ParameterViewModel : ParameterViewModels)
+		{
+			if (ParameterViewModel->GetName() == InputNode->Input.GetName())
+			{
+				GetSelection().SetSelectedObject(ParameterViewModel);
+				break;
+			}
+		}
+	}
+}
+
+void FNiagaraScriptInputCollectionViewModel::DeleteSelectedParameters()
+{
+	if (GetSelection().GetSelectedObjects().Num() > 0)
+	{
+		TSet<FName> InputNamesToDelete;
+		for (TSharedRef<INiagaraParameterViewModel> InputParameter : GetSelection().GetSelectedObjects())
+		{
+			InputNamesToDelete.Add(InputParameter->GetName());
+		}
+		GetSelection().ClearSelectedObjects();
+
+		if (Graph.IsValid())
+		{
+			FScopedTransaction ScopedTransaction(NSLOCTEXT("NiagaraEmitterInputEditor", "DeletedSelectedNodes", "Delete selected nodes"));
+			Graph->Modify();
+
+			TArray<UNiagaraNodeInput*> InputNodes;
+			Graph->GetNodesOfClass(InputNodes);
+			for (UNiagaraNodeInput* InputNode : InputNodes)
+			{
+				if (InputNamesToDelete.Contains(InputNode->Input.GetName()))
+				{
+					InputNode->Modify();
+					InputNode->DestroyNode();
+				}
+			}
+		}
+	}
+}
+
+const TArray<TSharedRef<INiagaraParameterViewModel>>& FNiagaraScriptInputCollectionViewModel::GetParameters()
+{
+	return ParameterViewModels;
+}
+
+void FNiagaraScriptInputCollectionViewModel::RefreshParameterViewModels()
+{
+	for (int32 i = 0; i < ParameterViewModels.Num(); i++)
+	{
+		FNiagaraScriptParameterViewModel* ParameterViewModel = (FNiagaraScriptParameterViewModel*)(&ParameterViewModels[i].Get());
+		if (ParameterViewModel != nullptr)
+		{
+			ParameterViewModel->Reset();
+			ParameterViewModel->OnNameChanged().RemoveAll(this);
+			ParameterViewModel->OnTypeChanged().RemoveAll(this);
+			ParameterViewModel->OnDefaultValueChanged().RemoveAll(this);
+		}
+	}
+
+	ParameterViewModels.Empty();
+
+	TArray<UNiagaraNodeInput*> InputNodes;
+
+	if (Graph.IsValid())
+	{
+		UNiagaraGraph::FFindInputNodeOptions Options;
+		Options.bSort = true;
+		Graph->FindInputNodes(InputNodes, Options);
+	}
+
+	TSet<FName> AddedInputNames;
+	for (UNiagaraNodeInput* InputNode : InputNodes)
+	{
+		if (0 == Scripts.Num())
+		{
+			continue;
+		}
+
+		// We can have multiple input nodes in the graph for each unique input name so make sure we only add one of each.
+		if (InputNode->Usage == ENiagaraInputNodeUsage::Parameter && AddedInputNames.Contains(InputNode->Input.GetName()) == false)
+		{
+			FNiagaraVariable& GraphVariable = InputNode->Input;
+			TSharedPtr<FNiagaraScriptParameterViewModel> ParameterViewModel;
+			if (GraphVariable.GetType().GetScriptStruct() != nullptr)
+			{
+				FNiagaraVariable* EmitterVariable = nullptr;
+				UNiagaraScript* Script = nullptr;
+				for (TWeakObjectPtr<UNiagaraScript> ScriptWeakPtr : Scripts)
+				{
+					if (!ScriptWeakPtr.IsValid())
+					{
+						continue;
+					}
+					if (EmitterVariable != nullptr)
+					{
+						break;
+					}
+
+					for (FNiagaraVariable& EmitterVariableToCheck : ScriptWeakPtr->Parameters.Parameters)
+					{
+
+						// @TODO We should check ID's here, but its possible that the 
+						// EmitterVariable is invalid and we may not have a great way to 
+						// recover without also changing the VM.
+						if (EmitterVariableToCheck.GetName() == GraphVariable.GetName())
+						{
+							EmitterVariable = &EmitterVariableToCheck;
+							Script = ScriptWeakPtr.Get();
+							break;
+						}
+					}
+				}
+				ParameterViewModel = MakeShareable(new FNiagaraScriptParameterViewModel(GraphVariable, *InputNode, EmitterVariable, Script, ParameterEditMode));
+				
+			}
+			else
+			{
+				UNiagaraDataInterface* EmitterDataInterface = InputNode->DataInterface;
+				UNiagaraScript* Script = nullptr;
+				for (TWeakObjectPtr<UNiagaraScript> ScriptWeakPtr : Scripts)
+				{
+					if (EmitterDataInterface != nullptr)
+					{
+						break;
+					}
+
+					for (FNiagaraScriptDataInterfaceInfo& DataInterfaceInfoItem : ScriptWeakPtr->DataInterfaceInfo)
+					{
+						if (DataInterfaceInfoItem.Name == InputNode->Input.GetName())
+						{
+							EmitterDataInterface = DataInterfaceInfoItem.DataInterface;
+							Script = ScriptWeakPtr.Get();
+							break;
+						}
+					}
+				}
+				ParameterViewModel = MakeShareable(new FNiagaraScriptParameterViewModel(GraphVariable, *InputNode, EmitterDataInterface, ParameterEditMode));
+			}
+
+			ParameterViewModel->OnNameChanged().AddRaw(this, &FNiagaraScriptInputCollectionViewModel::OnParameterNameChanged, &GraphVariable);
+			ParameterViewModel->OnTypeChanged().AddRaw(this, &FNiagaraScriptInputCollectionViewModel::OnParameterTypeChanged, &GraphVariable);
+			ParameterViewModel->OnDefaultValueChanged().AddRaw(this, &FNiagaraScriptInputCollectionViewModel::OnParameterValueChangedInternal, ParameterViewModel.ToSharedRef());
+			ParameterViewModels.Add(ParameterViewModel.ToSharedRef());
+			AddedInputNames.Add(InputNode->Input.GetName());
+		}
+	}
+
+	OnCollectionChangedDelegate.Broadcast();
+}
+
+bool FNiagaraScriptInputCollectionViewModel::SupportsType(const FNiagaraTypeDefinition& Type) const
+{
+	return bCanHaveNumericParameters || Type != FNiagaraTypeDefinition::GetGenericNumericDef();
+}
+
+void FNiagaraScriptInputCollectionViewModel::OnGraphChanged(const struct FEdGraphEditAction& InAction)
+{
+	RefreshParameterViewModels();
+}
+
+void FNiagaraScriptInputCollectionViewModel::OnParameterNameChanged(FName OldName, FName NewName, FNiagaraVariable* ParameterVariable)
+{
+	TSet<FName> CurrentNames;
+	TArray<UNiagaraNodeInput*> InputNodes;
+	TArray<UNiagaraNodeInput*> InputNodesToRename;
+
+	// Check the existing input nodes and get a set of the current names, and find nodes with matching names to rename.
+	if (Graph.IsValid())
+	{
+		Graph->GetNodesOfClass(InputNodes);
+	}
+	
+	for (UNiagaraNodeInput* InputNode : InputNodes)
+	{
+		if (InputNode->Usage == ENiagaraInputNodeUsage::Parameter && &InputNode->Input != ParameterVariable)
+		{
+			if (InputNode->Input.GetName() == OldName)
+			{
+				InputNodesToRename.Add(InputNode);
+			}
+			else
+			{
+				CurrentNames.Add(InputNode->Input.GetName());
+			}
+		}
+	}
+
+	TSet<FName> SystemConstantNames = FNiagaraEditorUtilities::GetSystemConstantNames();
+
+	// Rename the nodes and notify the graph that they've changed.
+	FName UniqueNewName = FNiagaraEditorUtilities::GetUniqueName(ParameterVariable->GetName(), CurrentNames.Union(SystemConstantNames));
+	if (ParameterVariable->GetName() != UniqueNewName)
+	{
+		ParameterVariable->SetName(UniqueNewName);
+	}
+	for (UNiagaraNodeInput* InputNodeToRename : InputNodesToRename)
+	{
+		InputNodeToRename->Modify();
+		InputNodeToRename->Input.SetName(UniqueNewName);
+	}
+
+	// Synchronize script variables...
+	for (TWeakObjectPtr<UNiagaraScript> Script : Scripts)
+	{
+		for (FNiagaraVariable& EmitterVariableToCheck : Script->Parameters.Parameters)
+		{
+			if (EmitterVariableToCheck.GetName() == ParameterVariable->GetName() && ParameterVariable != &EmitterVariableToCheck)
+			{
+				EmitterVariableToCheck = *ParameterVariable;
+				break;
+			}
+		}
+	}
+
+	if (Graph.IsValid())
+	{
+		Graph->NotifyGraphChanged();
+	}
+}
+
+void FNiagaraScriptInputCollectionViewModel::SetAllParametersEditingEnabled(bool bInEnabled)
+{
+	for (TSharedRef<INiagaraParameterViewModel>& ParameterViewModel : ParameterViewModels)
+	{
+		ParameterViewModel->SetEditingEnabled(bInEnabled);
+	}
+}
+
+void FNiagaraScriptInputCollectionViewModel::SetAllParametersTooltipOverrides(const FText& Override)
+{
+	for (TSharedRef<INiagaraParameterViewModel>& ParameterViewModel : ParameterViewModels)
+	{
+		ParameterViewModel->SetTooltipOverride(Override);
+	}
+}
+
+TSharedPtr<INiagaraParameterViewModel> FNiagaraScriptInputCollectionViewModel::GetParameterViewModel(const FName& Name)
+{
+	for (TSharedRef<INiagaraParameterViewModel>& ParameterViewModel : ParameterViewModels)
+	{
+		if (ParameterViewModel->GetName() == Name)
+		{
+			return TSharedPtr<INiagaraParameterViewModel>(ParameterViewModel);
+		}
+	}
+	return TSharedPtr<INiagaraParameterViewModel>();
+}
+
+
+void FNiagaraScriptInputCollectionViewModel::OnParameterTypeChanged(FNiagaraVariable* ParameterVariable)
+{
+	TArray<UNiagaraNodeInput*> InputNodes;
+	TArray<UNiagaraNodeInput*> InputNodesToUpdate;
+
+	if (Graph.IsValid())
+	{
+		Graph->GetNodesOfClass(InputNodes);
+	}
+	
+	for (UNiagaraNodeInput* InputNode : InputNodes)
+	{
+		if (InputNode->Usage == ENiagaraInputNodeUsage::Parameter && InputNode->Input.GetName() == ParameterVariable->GetName())
+		{
+			InputNodesToUpdate.Add(InputNode);
+		}
+	}
+
+	if (InputNodesToUpdate.Num() > 0)
+	{
+		// Reinitialize the first node found.
+		UNiagaraNodeInput* FirstNodeToUpdate = InputNodesToUpdate[0];
+		FirstNodeToUpdate->Modify();
+		TSet<FName> EmptyNames;
+		FNiagaraEditorUtilities::InitializeParameterInputNode(*FirstNodeToUpdate, ParameterVariable->GetType(), nullptr, FirstNodeToUpdate->Input.GetName());
+		FirstNodeToUpdate->NotifyInputTypeChanged();
+
+		// Then copy that nodes input variable to the others so that they all have the same id and data object pointers.
+		for (int32 i = 1; i < InputNodesToUpdate.Num(); ++i)
+		{
+			UNiagaraNodeInput* InputNodeToUpdate = InputNodesToUpdate[i];
+			InputNodeToUpdate->Modify();
+			InputNodeToUpdate->Input = FirstNodeToUpdate->Input;
+			InputNodeToUpdate->DataInterface = FirstNodeToUpdate->DataInterface;
+			InputNodeToUpdate->NotifyInputTypeChanged();
+		}
+	}
+
+	// Synchronize script variables...
+	for (TWeakObjectPtr<UNiagaraScript> Script : Scripts)
+	{
+		for (FNiagaraVariable& EmitterVariableToCheck : Script->Parameters.Parameters)
+		{
+			if (EmitterVariableToCheck.GetName() == ParameterVariable->GetName() && ParameterVariable != &EmitterVariableToCheck)
+			{
+				EmitterVariableToCheck = *ParameterVariable;
+				break;
+			}
+		}
+	}
+
+}
+
+void FNiagaraScriptInputCollectionViewModel::OnParameterValueChangedInternal(TSharedRef<FNiagaraScriptParameterViewModel> ChangedParameter)
+{
+	// Since we potentially have multiple input nodes that point to the exact same underlying input variable, we need to make sure that 
+	// all of them are synchronized in their values or else we might end up with confusion when we finally compile. 
+	if (ChangedParameter->GetDefaultValueType() == INiagaraParameterViewModel::EDefaultValueType::Struct || 
+		ChangedParameter->GetDefaultValueType() == INiagaraParameterViewModel::EDefaultValueType::Object)
+	{
+		TArray<UNiagaraNodeInput*> InputNodes;
+		TArray<UNiagaraNodeInput*> InputNodesToUpdate;
+
+		if (Graph.IsValid())
+		{
+			Graph->GetNodesOfClass(InputNodes);
+		}
+
+		for (UNiagaraNodeInput* InputNode : InputNodes)
+		{
+			// Copy the value to all nodes with the matching id to cover both the case where we're updating a graph variable
+			// and the case where we're updating a compiled variable.
+			if (InputNode->Usage == ENiagaraInputNodeUsage::Parameter && InputNode->Input.GetName() == ChangedParameter->GetName())
+			{
+				if (ensureMsgf(InputNode->Input.GetType() == *ChangedParameter->GetType().Get(), TEXT("Can not propagate variable values when the types don't match.")))
+				{
+					InputNodesToUpdate.Add(InputNode);
+				}
+			}
+		}
+
+		if (ChangedParameter->GetDefaultValueType() == INiagaraParameterViewModel::EDefaultValueType::Struct)
+		{
+			for (UNiagaraNodeInput* InputNodeToUpdate : InputNodesToUpdate)
+			{
+				InputNodeToUpdate->Modify();
+				InputNodeToUpdate->Input.AllocateData();
+				InputNodeToUpdate->Input.SetData(ChangedParameter->GetDefaultValueStruct()->GetStructMemory());
+			}
+			
+			// Synchronize script variables...
+			for (TWeakObjectPtr<UNiagaraScript> Script : Scripts)
+			{
+				for (FNiagaraVariable& EmitterVariableToCheck : Script->Parameters.Parameters)
+				{
+					if (EmitterVariableToCheck.GetName() == ChangedParameter->GetName())
+					{
+						EmitterVariableToCheck = InputNodesToUpdate[0]->Input;
+						break;
+					}
+				}
+			}
+		}
+		else if (ChangedParameter->GetDefaultValueType() == INiagaraParameterViewModel::EDefaultValueType::Object)
+		{
+			UNiagaraDataInterface* DataInterface = Cast<UNiagaraDataInterface>(ChangedParameter->GetDefaultValueObject());
+			if (DataInterface != nullptr)
+			{
+				TSet<UNiagaraDataInterface*> DataInterfacesToUpdate;
+				for (UNiagaraNodeInput* InputNodeToUpdate : InputNodesToUpdate)
+				{
+					DataInterfacesToUpdate.Add(InputNodeToUpdate->DataInterface);
+				}
+
+				for (UNiagaraDataInterface* DataInterfaceToUpdate : DataInterfacesToUpdate)
+				{
+					if (DataInterfaceToUpdate->Equals(DataInterface) == false)
+					{
+						DataInterfaceToUpdate->Modify();
+						DataInterface->CopyTo(DataInterfaceToUpdate);
+					}
+				}
+
+				// Synchronize script variables...
+				for (TWeakObjectPtr<UNiagaraScript> Script : Scripts)
+				{
+					for (FNiagaraScriptDataInterfaceInfo& Info : Script->DataInterfaceInfo)
+					{
+						if (Info.Name == ChangedParameter->GetName() && DataInterface != Info.DataInterface)
+						{
+							DataInterface->CopyTo(Info.DataInterface);
+							break;
+						}
+					}
+				}
+			}
+		}
+	}
+
+	OnParameterValueChanged().Broadcast(ChangedParameter->GetName());
+}
+
+#undef LOCTEXT_NAMESPACE // "NiagaraScriptInputCollection"
diff --git a/Engine/Plugins/FX/Niagara/Source/NiagaraEditor/Private/ViewModels/NiagaraScriptViewModel.cpp b/Engine/Plugins/FX/Niagara/Source/NiagaraEditor/Private/ViewModels/NiagaraScriptViewModel.cpp
new file mode 100644
--- /dev/null
+++ b/Engine/Plugins/FX/Niagara/Source/NiagaraEditor/Private/ViewModels/NiagaraScriptViewModel.cpp
@@ -0,0 +1,590 @@
+// Copyright 1998-2016 Epic Games, Inc. All Rights Reserved.
+
+#include "NiagaraScriptViewModel.h"
+#include "NiagaraScript.h"
+#include "NiagaraScriptSource.h"
+#include "NiagaraGraph.h"
+#include "NiagaraScriptGraphViewModel.h"
+#include "NiagaraEmitter.h"
+#include "NiagaraScriptInputCollectionViewModel.h"
+#include "NiagaraScriptOutputCollectionViewModel.h"
+#include "NiagaraNodeInput.h"
+#include "INiagaraCompiler.h"
+#include "CompilerResultsLog.h"
+#include "GraphEditAction.h"
+#include "Package.h"
+#include "Editor.h"
+#include "TNiagaraViewModelManager.h"
+#include "NiagaraDataInterface.h"
+
+template<> TMap<UNiagaraScript*, TArray<FNiagaraScriptViewModel*>> TNiagaraViewModelManager<UNiagaraScript, FNiagaraScriptViewModel>::ObjectsToViewModels{};
+
+FNiagaraScriptViewModel::FNiagaraScriptViewModel(UNiagaraScript* InScript, FText DisplayName, ENiagaraParameterEditMode InParameterEditMode)
+	: InputCollectionViewModel(MakeShareable(new FNiagaraScriptInputCollectionViewModel(InScript, DisplayName, InParameterEditMode)))
+	, OutputCollectionViewModel(MakeShareable(new FNiagaraScriptOutputCollectionViewModel(InScript, InParameterEditMode)))
+	, GraphViewModel(MakeShareable(new FNiagaraScriptGraphViewModel(Cast<UNiagaraScriptSource>(InScript->GetSource()), DisplayName)))
+	, bUpdatingSelectionInternally(false)
+	, LastCompileStatus(ENiagaraScriptCompileStatus::NCS_Unknown)
+	, bNeedsSave(true)
+{
+	if (InScript != nullptr)
+	{
+		Scripts.Add(InScript);
+		Source = Cast<UNiagaraScriptSource>(InScript->GetSource());
+	}
+
+	InputCollectionViewModel->GetSelection().OnSelectedObjectsChanged().AddRaw(this, &FNiagaraScriptViewModel::InputViewModelSelectionChanged);
+	GraphViewModel->GetSelection()->OnSelectedObjectsChanged().AddRaw(this, &FNiagaraScriptViewModel::GraphViewModelSelectedNodesChanged);
+	GEditor->RegisterForUndo(this);
+	
+	if (Source != nullptr)
+	{
+		UNiagaraGraph* Graph = Source->NodeGraph;
+		if (Graph != nullptr)
+		{
+			OnGraphChangedHandle = Graph->AddOnGraphChangedHandler(
+				FOnGraphChanged::FDelegate::CreateRaw(this, &FNiagaraScriptViewModel::OnGraphChanged));
+		}
+	}
+
+	// Guess at initial compile status
+	if (InScript != nullptr && InScript->ByteCode.Num() == 0) // This is either a brand new script or failed in the past. Since we create a default working script, assume invalid.
+	{
+		bNeedsSave = true;
+		LastCompileStatus = ENiagaraScriptCompileStatus::NCS_Error;
+		GraphViewModel->SetErrorTextToolTip("Please recompile for full error stack.");
+	}
+	else // Possibly warnings previously, but still compiled. It *could* have been dirtied somehow, but we assume that it is up-to-date.
+	{
+		if (InScript != nullptr && InScript->AreScriptAndSourceSynchronized())
+		{
+			bNeedsSave = false;
+			LastCompileStatus = InScript->GetLastCompileStatus();
+		}
+		else
+		{
+			bNeedsSave = true;
+			LastCompileStatus = ENiagaraScriptCompileStatus::NCS_Error;
+			GraphViewModel->SetErrorTextToolTip("Please recompile for full error stack.");
+		}
+	}
+
+	RegisteredHandles.Add(RegisterViewModelWithMap(InScript, this));
+}
+
+FNiagaraScriptViewModel::FNiagaraScriptViewModel(UNiagaraEmitter* InEmitter, FText DisplayName, ENiagaraParameterEditMode InParameterEditMode)
+	: InputCollectionViewModel(MakeShareable(new FNiagaraScriptInputCollectionViewModel(InEmitter, DisplayName, InParameterEditMode)))
+	, OutputCollectionViewModel(MakeShareable(new FNiagaraScriptOutputCollectionViewModel(InEmitter, InParameterEditMode)))
+	, GraphViewModel(MakeShareable(new FNiagaraScriptGraphViewModel(Cast<UNiagaraScriptSource>(InEmitter->GraphSource), DisplayName)))
+	, bUpdatingSelectionInternally(false)
+	, LastCompileStatus(ENiagaraScriptCompileStatus::NCS_Unknown)
+	, bNeedsSave(true)
+{
+	TArray<UNiagaraScript*> InScripts;
+	InEmitter->GetScripts(InScripts);
+	// Because of weak pointers, we need to copy ourselves..
+	for (UNiagaraScript* Script : InScripts)
+	{
+		Scripts.Add(Script);
+		check(Script->GetSource() == InEmitter->GraphSource);
+	}
+	Source = Cast<UNiagaraScriptSource>(InEmitter->GraphSource);
+
+	InputCollectionViewModel->GetSelection().OnSelectedObjectsChanged().AddRaw(this, &FNiagaraScriptViewModel::InputViewModelSelectionChanged);
+	GraphViewModel->GetSelection()->OnSelectedObjectsChanged().AddRaw(this, &FNiagaraScriptViewModel::GraphViewModelSelectedNodesChanged);
+	GEditor->RegisterForUndo(this);
+
+	UNiagaraGraph* Graph = Source->NodeGraph;
+	if (Graph != nullptr)
+	{
+		OnGraphChangedHandle = Graph->AddOnGraphChangedHandler(
+			FOnGraphChanged::FDelegate::CreateRaw(this, &FNiagaraScriptViewModel::OnGraphChanged));
+	}
+
+	// Guess at initial compile status
+	LastCompileStatus = ENiagaraScriptCompileStatus::NCS_UpToDate;
+	bNeedsSave = false;
+
+	CompileStatuses.Empty();
+	CompileErrors.Empty();
+	CompilePaths.Empty();
+
+	for (int32 i = 0; i < Scripts.Num(); i++)
+	{
+		FString Message;
+		ENiagaraScriptCompileStatus ScriptStatus = Scripts[i]->GetLastCompileStatus();
+		if (Scripts[i].IsValid() && Scripts[i]->ByteCode.Num() == 0) // This is either a brand new script or failed in the past. Since we create a default working script, assume invalid.
+		{
+			bNeedsSave = true;
+			LastCompileStatus = FNiagaraEditorUtilities::UnionCompileStatus(LastCompileStatus, ENiagaraScriptCompileStatus::NCS_Error);
+			ScriptStatus = ENiagaraScriptCompileStatus::NCS_Error;
+			Message = TEXT("Please recompile for full error stack.");
+			GraphViewModel->SetErrorTextToolTip(Message);
+		}
+		else // Possibly warnings previously, but still compiled. It *could* have been dirtied somehow, but we assume that it is up-to-date.
+		{
+			if (Scripts[i].IsValid() && Scripts[i]->AreScriptAndSourceSynchronized())
+			{
+				LastCompileStatus = FNiagaraEditorUtilities::UnionCompileStatus(LastCompileStatus, Scripts[i]->GetLastCompileStatus());
+			}
+			else
+			{
+				bNeedsSave = true;
+				LastCompileStatus = FNiagaraEditorUtilities::UnionCompileStatus(LastCompileStatus, ENiagaraScriptCompileStatus::NCS_Error);
+				ScriptStatus = ENiagaraScriptCompileStatus::NCS_Error;
+				Message = TEXT("Please recompile for full error stack.");
+				GraphViewModel->SetErrorTextToolTip(Message);
+			}
+		}
+
+		CompilePaths.Add(Scripts[i]->GetPathName());
+		CompileErrors.Add(Message);
+		CompileStatuses.Add(ScriptStatus);
+
+		RegisteredHandles.Add(RegisterViewModelWithMap(Scripts[i].Get(), this));
+	}
+
+	CompileTypes.SetNum(CompileStatuses.Num());
+	for (int32 i = 0; i < CompileStatuses.Num(); i++)
+	{
+		CompileTypes[i].Key = Scripts[i]->GetUsage();
+		if (i > 0 && CompileTypes[i - 1].Key == Scripts[i]->GetUsage())
+			CompileTypes[i].Value = CompileTypes[i - 1].Value + 1;
+		else
+			CompileTypes[i].Value = 0;
+	}
+}
+
+bool FNiagaraScriptViewModel::IsGraphDirty() const
+{
+	for (int32 i = 0; i < Scripts.Num(); i++)
+	{
+		if (!Scripts[i].IsValid())
+		{
+			continue;
+		}
+
+		if (!Scripts[i]->AreScriptAndSourceSynchronized())
+		{
+			return true;
+		}
+	}
+	return false;
+}
+
+
+FNiagaraScriptViewModel::~FNiagaraScriptViewModel()
+{
+	InputCollectionViewModel->GetSelection().OnSelectedObjectsChanged().RemoveAll(this);
+	GraphViewModel->GetSelection()->OnSelectedObjectsChanged().RemoveAll(this);
+
+	if (Source.IsValid() && Source != nullptr)
+	{
+		UNiagaraGraph* Graph = Source->NodeGraph;
+		if (Graph != nullptr)
+		{
+			Graph->RemoveOnGraphChangedHandler(OnGraphChangedHandle);
+		}
+	}
+
+	GEditor->UnregisterForUndo(this);
+
+	for (TNiagaraViewModelManager<UNiagaraScript, FNiagaraScriptViewModel>::Handle RegisteredHandle : RegisteredHandles)
+	{
+		UnregisterViewModelWithMap(RegisteredHandle);
+	}
+}
+
+
+void FNiagaraScriptViewModel::SetScripts(UNiagaraScriptSource* InScriptSource, TArray<UNiagaraScript*>& InScripts)
+{
+	// Remove the graph changed handler on the child.
+	if (Source.IsValid())
+	{
+		UNiagaraGraph* Graph = Source->NodeGraph;
+		if (Graph != nullptr)
+		{
+			Graph->RemoveOnGraphChangedHandler(OnGraphChangedHandle);
+		}
+	}
+	else
+	{
+		Source = nullptr;
+	}
+
+	for (TNiagaraViewModelManager<UNiagaraScript, FNiagaraScriptViewModel>::Handle RegisteredHandle : RegisteredHandles)
+	{
+		UnregisterViewModelWithMap(RegisteredHandle);
+	}
+	RegisteredHandles.Empty();
+
+	Scripts.Empty();
+	for (UNiagaraScript* Script : InScripts)
+	{
+		Scripts.Add(Script);
+		check(Script->GetSource() == InScriptSource);
+	}
+	Source = InScriptSource;
+
+	InputCollectionViewModel->SetScripts(InScripts);
+	OutputCollectionViewModel->SetScripts(InScripts);
+	GraphViewModel->SetScriptSource(Source.Get());
+
+	UNiagaraGraph* Graph = nullptr;
+	if (Source.IsValid() && Source != nullptr)
+	{
+		// The underlying graph may have changed after the previous call.
+		Graph = Source->NodeGraph;
+		if (Graph != nullptr)
+		{
+			OnGraphChangedHandle = Graph->AddOnGraphChangedHandler(
+				FOnGraphChanged::FDelegate::CreateRaw(this, &FNiagaraScriptViewModel::OnGraphChanged));
+		}
+	}
+	
+	// Guess at initial compile status
+	LastCompileStatus = ENiagaraScriptCompileStatus::NCS_UpToDate;
+	bNeedsSave = false;
+
+	CompileStatuses.Empty();
+	CompileErrors.Empty();
+	CompilePaths.Empty();
+
+	for (int32 i = 0; i < Scripts.Num(); i++)
+	{
+		FString Message;
+		ENiagaraScriptCompileStatus ScriptStatus = Scripts[i]->GetLastCompileStatus();
+		if (Scripts[i].IsValid() && Scripts[i]->ByteCode.Num() == 0) // This is either a brand new script or failed in the past. Since we create a default working script, assume invalid.
+		{
+			bNeedsSave = true;
+			LastCompileStatus = FNiagaraEditorUtilities::UnionCompileStatus(LastCompileStatus, ENiagaraScriptCompileStatus::NCS_Error);
+			ScriptStatus = ENiagaraScriptCompileStatus::NCS_Error;
+			Message = TEXT("Please recompile for full error stack.");
+			GraphViewModel->SetErrorTextToolTip(Message);
+		}
+		else // Possibly warnings previously, but still compiled. It *could* have been dirtied somehow, but we assume that it is up-to-date.
+		{
+			if (Scripts[i].IsValid() && Scripts[i]->AreScriptAndSourceSynchronized())
+			{
+				LastCompileStatus = FNiagaraEditorUtilities::UnionCompileStatus(LastCompileStatus, Scripts[i]->GetLastCompileStatus());
+			}
+			else
+			{
+				bNeedsSave = true;
+				LastCompileStatus = FNiagaraEditorUtilities::UnionCompileStatus(LastCompileStatus, ENiagaraScriptCompileStatus::NCS_Error);
+				ScriptStatus = ENiagaraScriptCompileStatus::NCS_Error;
+				Message = TEXT("Please recompile for full error stack.");
+				GraphViewModel->SetErrorTextToolTip(Message);
+			}
+		}
+
+		CompilePaths.Add(Scripts[i]->GetPathName());
+		CompileErrors.Add(Message);
+		CompileStatuses.Add(ScriptStatus);
+
+		RegisteredHandles.Add(RegisterViewModelWithMap(Scripts[i].Get(), this));
+	}
+
+	CompileTypes.SetNum(CompileStatuses.Num());
+	for (int32 i = 0; i < CompileStatuses.Num(); i++)
+	{
+		CompileTypes[i].Key = Scripts[i]->GetUsage();
+		if (i > 0 && CompileTypes[i - 1].Key == Scripts[i]->GetUsage())
+			CompileTypes[i].Value = CompileTypes[i - 1].Value + 1;
+		else
+			CompileTypes[i].Value = 0;
+	}
+}
+
+void FNiagaraScriptViewModel::SetScripts(UNiagaraEmitter* InEmitter)
+{
+	if (InEmitter == nullptr)
+	{
+		TArray<UNiagaraScript*> EmptyScripts;
+		SetScripts(nullptr, EmptyScripts);
+	}
+	else
+	{
+		TArray<UNiagaraScript*> InScripts;
+		InEmitter->GetScripts(InScripts);
+		SetScripts(Cast<UNiagaraScriptSource>(InEmitter->GraphSource), InScripts);
+	}
+}
+
+void FNiagaraScriptViewModel::SetScript(UNiagaraScript* InScript)
+{
+	TArray<UNiagaraScript*> InScripts;
+	UNiagaraScriptSource* InSource = nullptr;
+	if (InScript)
+	{
+		InScripts.Add(InScript);
+		InSource = Cast<UNiagaraScriptSource>(InScript->GetSource());
+	}
+	SetScripts(InSource, InScripts);
+}
+
+void FNiagaraScriptViewModel::OnGraphChanged(const struct FEdGraphEditAction& InAction)
+{
+	if ((InAction.Action & GRAPHACTION_AddNode) != 0 || (InAction.Action & GRAPHACTION_RemoveNode) != 0 ||
+		(InAction.Action & GRAPHACTION_GenericNeedsRecompile) != 0)
+	{
+		for (int32 i = 0; i < Scripts.Num(); i++)
+		{
+			if (Scripts[i].IsValid())
+			{
+				Scripts[i]->MarkScriptAndSourceDesynchronized();
+			}
+		}
+		bNeedsSave = true;
+	}
+}
+
+TSharedRef<FNiagaraScriptInputCollectionViewModel> FNiagaraScriptViewModel::GetInputCollectionViewModel()
+{
+	return InputCollectionViewModel;
+}
+
+
+TSharedRef<FNiagaraScriptOutputCollectionViewModel> FNiagaraScriptViewModel::GetOutputCollectionViewModel()
+{
+	return OutputCollectionViewModel;
+}
+
+TSharedRef<FNiagaraScriptGraphViewModel> FNiagaraScriptViewModel::GetGraphViewModel()
+{
+	return GraphViewModel;
+}
+
+void FNiagaraScriptViewModel::UpdateCompileStatus(ENiagaraScriptCompileStatus InAggregateCompileStatus, const FString& InAggregateCompileErrorString,
+	const TArray<ENiagaraScriptCompileStatus>& InCompileStatuses, const TArray<FString>& InCompileErrors, const TArray<FString>& InCompilePaths,
+	const TArray<UNiagaraScript*>& InCompileSources)
+{
+	if (Source.IsValid() && Source != nullptr)
+	{
+		CompileStatuses = InCompileStatuses;
+		CompileErrors = InCompileErrors;
+		CompilePaths = InCompilePaths;
+
+		CompileTypes.SetNum(CompileStatuses.Num());
+		for (int32 i = 0; i < CompileStatuses.Num(); i++)
+		{
+			CompileTypes[i].Key = InCompileSources[i]->GetUsage();
+			if (i > 0 && CompileTypes[i - 1].Key == InCompileSources[i]->GetUsage())
+				CompileTypes[i].Value = CompileTypes[i - 1].Value + 1;
+			else
+				CompileTypes[i].Value = 0;
+		}
+
+		LastCompileStatus = InAggregateCompileStatus;
+		bNeedsSave = true;
+		
+		if (LastCompileStatus == ENiagaraScriptCompileStatus::NCS_Error)
+		{
+			GraphViewModel->SetErrorTextToolTip(InAggregateCompileErrorString + "\n(These same errors are also in the log)");
+		}
+		else
+		{
+			GraphViewModel->SetErrorTextToolTip("");
+		}
+
+		InputCollectionViewModel->RefreshParameterViewModels();
+		OutputCollectionViewModel->RefreshParameterViewModels();
+	}
+}
+
+ENiagaraScriptCompileStatus FNiagaraScriptViewModel::GetScriptCompileStatus(ENiagaraScriptUsage InUsage, int32 InOccurrence) const
+{
+	ENiagaraScriptCompileStatus Status = ENiagaraScriptCompileStatus::NCS_Unknown;
+	for (int32 i = 0; i < CompileTypes.Num(); i++)
+	{
+		if (UNiagaraScript::IsEquivalentUsage(CompileTypes[i].Key, InUsage) && CompileTypes[i].Value == InOccurrence)
+		{
+			return CompileStatuses[i];
+		}
+	}
+	return Status;
+}
+
+FText FNiagaraScriptViewModel::GetScriptErrors(ENiagaraScriptUsage InUsage, int32 InOccurrence) const
+{
+	FText Text = FText::GetEmpty();
+	for (int32 i = 0; i < CompileTypes.Num(); i++)
+	{
+		if (UNiagaraScript::IsEquivalentUsage(CompileTypes[i].Key, InUsage) && CompileTypes[i].Value == InOccurrence)
+		{
+			return FText::FromString(CompileErrors[i]);
+		}
+	}
+	return Text;
+}
+
+void FNiagaraScriptViewModel::UpdateCompiledDataInterfaces(UNiagaraDataInterface* ChangedDataInterface)
+{
+	FName ChangedInputName = NAME_None;
+	TArray<UNiagaraNodeInput*> InputNodes;
+	GetGraphViewModel()->GetGraph()->GetNodesOfClass<UNiagaraNodeInput>(InputNodes);
+	for (UNiagaraNodeInput* InputNode : InputNodes)
+	{
+		if (InputNode->DataInterface == ChangedDataInterface)
+		{
+			ChangedInputName = InputNode->Input.GetName();
+			break;
+		}
+	}
+
+	if (ChangedInputName != NAME_None)
+	{
+		for (TWeakObjectPtr<UNiagaraScript> Script : Scripts)
+		{
+			if (Script.IsValid())
+			{
+				for (FNiagaraScriptDataInterfaceInfo& DataInterfaceInfo : Script->DataInterfaceInfo)
+				{
+					if (DataInterfaceInfo.Name == ChangedInputName)
+					{
+						ChangedDataInterface->CopyTo(DataInterfaceInfo.DataInterface);
+					}
+				}
+			}
+		}
+	}
+}
+
+const UNiagaraScript* FNiagaraScriptViewModel::GetScript(ENiagaraScriptUsage InUsage) const
+{
+	for (int32 i = 0; i < Scripts.Num(); i++)
+	{
+		if (Scripts[i]->IsEquivalentUsage(InUsage))
+		{
+			return Scripts[i].Get();
+		}
+	}
+	return nullptr;
+}
+
+void FNiagaraScriptViewModel::CompileStandaloneScript()
+{
+	if (Source.IsValid() && Source != nullptr && Scripts.Num() == 1 && Scripts[0].IsValid () && (Scripts[0]->IsStandaloneScript()))
+	{
+		FString ErrorMsg;
+		LastCompileStatus = Scripts[0]->Compile(ErrorMsg);
+		Scripts[0]->CacheResourceShadersForRendering(true);
+		bNeedsSave = true;
+
+		if (LastCompileStatus == ENiagaraScriptCompileStatus::NCS_Error)
+		{
+			GraphViewModel->SetErrorTextToolTip(ErrorMsg + "\n(These same errors are also in the log)");
+		}
+		else
+		{
+			GraphViewModel->SetErrorTextToolTip("");
+		}
+
+		InputCollectionViewModel->RefreshParameterViewModels();
+		OutputCollectionViewModel->RefreshParameterViewModels();
+	}
+	else
+	{
+		ensure(0); // Should not call this for a script that isn't standalone.
+	}
+}
+
+
+ENiagaraScriptCompileStatus FNiagaraScriptViewModel::GetLatestCompileStatus()
+{
+	if (GraphViewModel->GetGraph() && IsGraphDirty())
+		return ENiagaraScriptCompileStatus::NCS_Dirty;
+	return LastCompileStatus;
+}
+
+void FNiagaraScriptViewModel::RefreshNodes()
+{
+	if (GraphViewModel->GetGraph())
+	{
+		TArray<UNiagaraNode*> NiagaraNodes;
+		GraphViewModel->GetGraph()->GetNodesOfClass<UNiagaraNode>(NiagaraNodes);
+		for (UNiagaraNode* NiagaraNode : NiagaraNodes)
+		{
+			if (NiagaraNode->RefreshFromExternalChanges())
+			{
+				for (int32 i = 0; i < Scripts.Num(); i++)
+				{
+					if (Scripts[i].IsValid())
+					{
+						Scripts[i]->MarkScriptAndSourceDesynchronized();
+					}
+				}
+			}
+			bNeedsSave = true;
+		}
+	}
+}
+
+void FNiagaraScriptViewModel::PostUndo(bool bSuccess)
+{
+	InputCollectionViewModel->RefreshParameterViewModels();
+	OutputCollectionViewModel->RefreshParameterViewModels();
+}
+
+void FNiagaraScriptViewModel::GraphViewModelSelectedNodesChanged()
+{
+	if (bUpdatingSelectionInternally == false)
+	{
+		bUpdatingSelectionInternally = true;
+		{
+			TSet<FName> SelectedInputIds;
+			for (UObject* SelectedObject : GraphViewModel->GetSelection()->GetSelectedObjects())
+			{
+				UNiagaraNodeInput* InputNode = Cast<UNiagaraNodeInput>(SelectedObject);
+				if (InputNode != nullptr)
+				{
+					SelectedInputIds.Add(InputNode->Input.GetName());
+				}
+			}
+
+			TSet<TSharedRef<INiagaraParameterViewModel>> ParametersToSelect;
+			for (TSharedRef<INiagaraParameterViewModel> Parameter : InputCollectionViewModel->GetParameters())
+			{
+				if (SelectedInputIds.Contains(Parameter->GetName()))
+				{
+					ParametersToSelect.Add(Parameter);
+				}
+			}
+
+			InputCollectionViewModel->GetSelection().SetSelectedObjects(ParametersToSelect);
+		}
+		bUpdatingSelectionInternally = false;
+	}
+}
+
+void FNiagaraScriptViewModel::InputViewModelSelectionChanged()
+{
+	if (bUpdatingSelectionInternally == false)
+	{
+		bUpdatingSelectionInternally = true;
+		{
+			TSet<FName> SelectedInputIds;
+			for (TSharedRef<INiagaraParameterViewModel> SelectedParameter : InputCollectionViewModel->GetSelection().GetSelectedObjects())
+			{
+				SelectedInputIds.Add(SelectedParameter->GetName());
+			}
+
+			TArray<UNiagaraNodeInput*> InputNodes;
+			if (GraphViewModel->GetGraph())
+			{
+				GraphViewModel->GetGraph()->GetNodesOfClass(InputNodes);
+			}
+			TSet<UObject*> NodesToSelect;
+			for (UNiagaraNodeInput* InputNode : InputNodes)
+			{
+				if (SelectedInputIds.Contains(InputNode->Input.GetName()))
+				{
+					NodesToSelect.Add(InputNode);
+				}
+			}
+
+			GraphViewModel->GetSelection()->SetSelectedObjects(NodesToSelect);
+		}
+		bUpdatingSelectionInternally = false;
+	}
+}
\ No newline at end of file
diff --git a/Engine/Plugins/FX/Niagara/Source/NiagaraEditor/Private/ViewModels/NiagaraScriptViewModel.h b/Engine/Plugins/FX/Niagara/Source/NiagaraEditor/Private/ViewModels/NiagaraScriptViewModel.h
new file mode 100644
--- /dev/null
+++ b/Engine/Plugins/FX/Niagara/Source/NiagaraEditor/Private/ViewModels/NiagaraScriptViewModel.h
@@ -0,0 +1,120 @@
+// Copyright 1998-2016 Epic Games, Inc. All Rights Reserved.
+
+#pragma once
+
+#include "NiagaraParameterEditMode.h"
+#include "INiagaraCompiler.h"
+#include "TNiagaraViewModelManager.h"
+#include "EditorUndoClient.h"
+#include "NotifyHook.h"
+
+class UNiagaraScript;
+class UNiagaraScriptSource;
+class INiagaraParameterCollectionViewModel;
+class FNiagaraScriptGraphViewModel;
+class FNiagaraScriptInputCollectionViewModel;
+class FNiagaraScriptOutputCollectionViewModel;
+class UNiagaraEmitter;
+
+
+/** A view model for niagara scripts which manages other script related view models. */
+class FNiagaraScriptViewModel : public TSharedFromThis<FNiagaraScriptViewModel>, public FEditorUndoClient, public TNiagaraViewModelManager<UNiagaraScript, FNiagaraScriptViewModel>
+{
+public:
+	FNiagaraScriptViewModel(UNiagaraScript* InScript, FText DisplayName, ENiagaraParameterEditMode InParameterEditMode);
+	FNiagaraScriptViewModel(UNiagaraEmitter* InEmitter, FText DisplayName, ENiagaraParameterEditMode InParameterEditMode);
+
+	~FNiagaraScriptViewModel();
+
+	/** Sets the view model to a different script. */
+	void SetScript(UNiagaraScript* InScript);
+
+	void SetScripts(UNiagaraEmitter* InEmitter);
+
+	/** Gets the view model for the input parameter collection. */
+	TSharedRef<FNiagaraScriptInputCollectionViewModel> GetInputCollectionViewModel();
+	
+	/** Gets the view model for the output parameter collection. */
+	TSharedRef<FNiagaraScriptOutputCollectionViewModel> GetOutputCollectionViewModel();
+
+	/** Gets the view model for the graph. */
+	TSharedRef<FNiagaraScriptGraphViewModel> GetGraphViewModel();
+
+	/** Updates the script with the latest compile status. */
+	void UpdateCompileStatus(ENiagaraScriptCompileStatus InAggregateCompileStatus, const FString& InAggregateCompileErrors,
+		const TArray<ENiagaraScriptCompileStatus>& InCompileStatuses, const TArray<FString>& InCompileErrors, const TArray<FString>& InCompilePaths,
+		const TArray<UNiagaraScript*>& InCompileSources);
+
+	/** Compiles a script that isn't part of an emitter or System. */
+	void CompileStandaloneScript();
+
+	/** Get the latest status of this view-model's script compilation.*/
+	ENiagaraScriptCompileStatus GetLatestCompileStatus();
+
+	/** Refreshes the nodes in the script graph, updating the pins to match external changes. */
+	void RefreshNodes();
+
+	//~ FEditorUndoClient Interface
+	virtual void PostUndo(bool bSuccess) override;
+	virtual void PostRedo(bool bSuccess) override { PostUndo(bSuccess); }
+
+	bool GetScriptDirty() const { return bNeedsSave; }
+
+	void SetScriptDirty(bool InNeedsSave) { bNeedsSave = InNeedsSave; }
+
+	const UNiagaraScript* GetScript(ENiagaraScriptUsage InUsage) const;
+
+	ENiagaraScriptCompileStatus GetScriptCompileStatus(ENiagaraScriptUsage InUsage, int32 InOccurrence) const;
+	FText GetScriptErrors(ENiagaraScriptUsage InUsage, int32 InOccurrence) const;
+
+	/** Updates the compiled versions of data interfaces from changes to their source. */
+	void UpdateCompiledDataInterfaces(UNiagaraDataInterface* ChangedDataInterface);
+
+private:
+	/** Handles the selection changing in the graph view model. */
+	void GraphViewModelSelectedNodesChanged();
+
+	/** Handles the selection changing in the input collection view model. */
+	void InputViewModelSelectionChanged();
+
+    /** Handle the graph being changed by the UI notifications to see if we need to mark as needing recompile.*/
+	void OnGraphChanged(const struct FEdGraphEditAction& InAction);
+
+	void SetScripts(UNiagaraScriptSource* InScriptSource, TArray<UNiagaraScript*>& InScripts);
+
+protected:
+	/** The script which provides the data for this view model. */
+	TArray<TWeakObjectPtr<UNiagaraScript>> Scripts;
+
+	TWeakObjectPtr<UNiagaraScriptSource> Source;
+
+	/** The view model for the input parameter collection. */
+	TSharedRef<FNiagaraScriptInputCollectionViewModel> InputCollectionViewModel;
+
+	/** The view model for the output parameter collection .*/
+	TSharedRef<FNiagaraScriptOutputCollectionViewModel> OutputCollectionViewModel;
+
+	/** The view model for the graph. */
+	TSharedRef<FNiagaraScriptGraphViewModel> GraphViewModel;
+
+	/** A flag for preventing reentrancy when synchronizing selection. */
+	bool bUpdatingSelectionInternally;
+
+	/** The stored latest compile status.*/
+	ENiagaraScriptCompileStatus LastCompileStatus;
+	
+	/** The handle to the graph changed delegate needed for removing. */
+	FDelegateHandle OnGraphChangedHandle;
+
+	/** An edit has been made since the last save.*/
+	bool bNeedsSave;
+
+	TArray<TNiagaraViewModelManager<UNiagaraScript, FNiagaraScriptViewModel>::Handle> RegisteredHandles;
+
+	bool IsGraphDirty() const;
+
+	TArray<ENiagaraScriptCompileStatus> CompileStatuses;
+	TArray<FString> CompileErrors;
+	TArray<FString> CompilePaths;
+	TArray<TPair<ENiagaraScriptUsage, int32> > CompileTypes;
+};
\ No newline at end of file
diff --git a/Engine/Plugins/FX/Niagara/Source/NiagaraEditor/Private/ViewModels/NiagaraSystemScriptViewModel.cpp b/Engine/Plugins/FX/Niagara/Source/NiagaraEditor/Private/ViewModels/NiagaraSystemScriptViewModel.cpp
new file mode 100644
--- /dev/null
+++ b/Engine/Plugins/FX/Niagara/Source/NiagaraEditor/Private/ViewModels/NiagaraSystemScriptViewModel.cpp
@@ -0,0 +1,229 @@
+// Copyright 1998-2016 Epic Games, Inc. All Rights Reserved.
+
+#include "NiagaraSystemScriptViewModel.h"
+#include "NiagaraSystem.h"
+#include "NiagaraEmitterHandle.h"
+#include "NiagaraGraph.h"
+#include "NiagaraTypes.h"
+#include "NiagaraScriptSource.h"
+#include "NiagaraScriptViewModel.h"
+#include "NiagaraScriptGraphViewModel.h"
+#include "NiagaraScriptInputCollectionViewModel.h"
+#include "NiagaraNodeEmitter.h"
+#include "NiagaraNodeInput.h"
+#include "NiagaraDataInterface.h"
+#include "NiagaraScriptSourceBase.h"
+#include "EdGraphSchema_Niagara.h"
+#include "NiagaraNodeOUtput.h"
+#include "GraphEditAction.h"
+#include "Stack/NiagaraStackEntry.h"
+#include "NiagaraStackGraphUtilities.h"
+
+FNiagaraSystemScriptViewModel::FNiagaraSystemScriptViewModel(UNiagaraSystem& InSystem)
+	: FNiagaraScriptViewModel(InSystem.GetSystemSpawnScript(), NSLOCTEXT("SystemScriptViewModel", "GraphName", "System"), ENiagaraParameterEditMode::EditAll)
+	, System(InSystem)
+{
+	if (GetGraphViewModel()->GetGraph())
+	{
+		OnGraphChangedHandle = GetGraphViewModel()->GetGraph()->AddOnGraphChangedHandler(
+			FOnGraphChanged::FDelegate::CreateRaw(this, &FNiagaraSystemScriptViewModel::OnGraphChanged));
+		
+		GetGraphViewModel()->SetErrorTextToolTip("");
+	}
+}
+
+FNiagaraSystemScriptViewModel::~FNiagaraSystemScriptViewModel()
+{
+	if (GetGraphViewModel()->GetGraph())
+	{
+		GetGraphViewModel()->GetGraph()->RemoveOnGraphChangedHandler(OnGraphChangedHandle);
+	}
+}
+
+
+float EmitterNodeVerticalOffset = 150.0f;
+
+FVector2D CalculateNewEmitterNodePlacementPosition(UNiagaraGraph* Graph, UNiagaraNodeEmitter* NewEmitterNode)
+{
+	FVector2D PlacementLocation(0.0f, 0.0f);
+	TArray<UNiagaraNodeEmitter*> EmitterNodes;
+	Graph->GetNodesOfClass(EmitterNodes);
+	if (EmitterNodes.Num() > 1)
+	{
+		// If there are Emitter nodes, try to put it under the lowest one.
+		UNiagaraNodeEmitter* LowestNode = nullptr;
+		for (UNiagaraNodeEmitter* EmitterNode : EmitterNodes)
+		{
+			if (EmitterNode != NewEmitterNode && (LowestNode == nullptr || EmitterNode->NodePosY > LowestNode->NodePosY))
+			{
+				LowestNode = EmitterNode;
+			}
+		}
+		check(LowestNode);
+		PlacementLocation = FVector2D(
+			LowestNode->NodePosX,
+			LowestNode->NodePosY + EmitterNodeVerticalOffset);
+	}
+	return PlacementLocation;
+}
+
+void FNiagaraSystemScriptViewModel::RebuildEmitterNodes()
+{
+	UNiagaraGraph* SystemGraph = GetGraphViewModel()->GetGraph();
+	if (SystemGraph == nullptr)
+	{
+		return;
+	}
+	
+	TArray<UNiagaraNodeEmitter*> CurrentEmitterNodes;
+	SystemGraph->GetNodesOfClass<UNiagaraNodeEmitter>(CurrentEmitterNodes);
+
+	const UEdGraphSchema_Niagara* Schema = Cast<UEdGraphSchema_Niagara>(SystemGraph->GetSchema());
+
+	// Remove the old emitter nodes since they will be rebuilt below.
+	for (UNiagaraNodeEmitter* CurrentEmitterNode : CurrentEmitterNodes)
+	{
+		CurrentEmitterNode->Modify();
+		UEdGraphPin* InPin = CurrentEmitterNode->GetInputPin(0);
+		UEdGraphPin* OutPin = CurrentEmitterNode->GetOutputPin(0);
+		if (InPin != nullptr && OutPin != nullptr)
+		{
+			if (0 != InPin->LinkedTo.Num() && 0 != OutPin->LinkedTo.Num())
+			{
+				InPin->LinkedTo[0]->MakeLinkTo(OutPin->LinkedTo[0]);
+			}
+		}
+		CurrentEmitterNode->DestroyNode();
+	}
+
+	// Add output nodes if they don't exist.
+	TArray<UNiagaraNodeInput*> TempInputNodes;
+	TArray<UNiagaraNodeInput*> InputNodes;
+	TArray<UNiagaraNodeOutput*> OutputNodes;
+	OutputNodes.SetNum(2);
+	OutputNodes[0] = SystemGraph->FindOutputNode(ENiagaraScriptUsage::SystemSpawnScript);
+	OutputNodes[1] = SystemGraph->FindOutputNode(ENiagaraScriptUsage::SystemUpdateScript);
+
+	// Add input nodes if they don't exist
+	UNiagaraGraph::FFindInputNodeOptions Options;
+	Options.bFilterDuplicates = false;
+	Options.bIncludeParameters = true;
+	SystemGraph->FindInputNodes(TempInputNodes);
+	for (int32 i = 0; i < TempInputNodes.Num(); i++)
+	{
+		if (Schema->PinToTypeDefinition(TempInputNodes[i]->GetOutputPin(0)) == FNiagaraTypeDefinition::GetParameterMapDef())
+		{
+			InputNodes.Add(TempInputNodes[i]);
+		}
+	}
+
+	// Create a default id variable for the input nodes.
+	FNiagaraVariable SharedInputVar(FNiagaraTypeDefinition::GetParameterMapDef(), TEXT("InputMap"));	
+	InputNodes.SetNum(2);
+
+	// Now create the nodes if they are needed, synchronize if already created.
+	for (int32 i = 0; i < 2; i++)
+	{
+		if (OutputNodes[i] == nullptr)
+		{
+			FGraphNodeCreator<UNiagaraNodeOutput> OutputNodeCreator(*SystemGraph);
+			OutputNodes[i] = OutputNodeCreator.CreateNode();
+			OutputNodes[i]->SetUsage((ENiagaraScriptUsage)(i + (int32)ENiagaraScriptUsage::SystemSpawnScript));
+			
+			OutputNodes[i]->Outputs.Add(FNiagaraVariable(FNiagaraTypeDefinition::GetParameterMapDef(), TEXT("Out")));
+			OutputNodes[i]->NodePosX = 0;
+			OutputNodes[i]->NodePosY = 0 + i*25;
+
+			OutputNodeCreator.Finalize();
+		}
+		if (InputNodes[i] == nullptr)
+		{
+			FGraphNodeCreator<UNiagaraNodeInput> InputNodeCreator(*SystemGraph);
+			InputNodes[i] =  InputNodeCreator.CreateNode();
+			InputNodes[i]->Input = SharedInputVar;
+			InputNodes[i]->Usage = ENiagaraInputNodeUsage::Parameter;
+			InputNodes[i]->NodePosX = -50;
+			InputNodes[i]->NodePosY = 0 + i * 25;
+
+			InputNodeCreator.Finalize();
+
+			InputNodes[i]->GetOutputPin(0)->MakeLinkTo(OutputNodes[i]->GetInputPin(0));
+		}
+	}
+
+	// Add new nodes.
+	UNiagaraNode* TargetNodes[2];
+	TargetNodes[0] = OutputNodes[0];
+	TargetNodes[1] = OutputNodes[1];
+
+	for (const FNiagaraEmitterHandle& EmitterHandle : System.GetEmitterHandles())
+	{
+		for (int32 i = 0; i < 2; i++)
+		{
+			FGraphNodeCreator<UNiagaraNodeEmitter> EmitterNodeCreator(*SystemGraph);
+			UNiagaraNodeEmitter* EmitterNode = EmitterNodeCreator.CreateNode();
+			EmitterNode->SetOwnerSystem(&System);
+			EmitterNode->SetEmitterHandleId(EmitterHandle.GetId());
+			EmitterNode->SetUsage((ENiagaraScriptUsage)(i + (int32)ENiagaraScriptUsage::EmitterSpawnScript));
+
+			FVector2D NewLocation = CalculateNewEmitterNodePlacementPosition(SystemGraph, EmitterNode);
+			EmitterNode->NodePosX = NewLocation.X;
+			EmitterNode->NodePosY = NewLocation.Y;
+			EmitterNode->AllocateDefaultPins();
+			EmitterNodeCreator.Finalize();
+
+			TArray<FNiagaraStackGraphUtilities::FStackNodeGroup> StackNodeGroups;
+			FNiagaraStackGraphUtilities::GetStackNodeGroups(*OutputNodes[i], StackNodeGroups);
+
+			FNiagaraStackGraphUtilities::FStackNodeGroup EmitterGroup;
+			EmitterGroup.StartNodes.Add(EmitterNode);
+			EmitterGroup.EndNode = EmitterNode;
+
+			FNiagaraStackGraphUtilities::FStackNodeGroup& OutputGroup = StackNodeGroups[StackNodeGroups.Num() - 1];
+			FNiagaraStackGraphUtilities::FStackNodeGroup& OutputGroupPrevious = StackNodeGroups[StackNodeGroups.Num() - 2];
+			FNiagaraStackGraphUtilities::ConnectStackNodeGroup(EmitterGroup, OutputGroupPrevious, OutputGroup);
+		}
+	}
+	
+	FNiagaraStackGraphUtilities::RelayoutGraph(*SystemGraph);
+}
+
+FNiagaraSystemScriptViewModel::FOnSystemCompiled& FNiagaraSystemScriptViewModel::OnSystemCompiled()
+{
+	return OnSystemCompiledDelegate;
+}
+
+void FNiagaraSystemScriptViewModel::CompileSystem()
+{
+
+	TArray<ENiagaraScriptCompileStatus> CompileSystemStatuses;
+	TArray<FString> CompileSystemErrors;
+	TArray<FString> CompileSystemPaths;
+	TArray<UNiagaraScript*> CompileSystemScripts;
+	System.CompileScripts(CompileSystemStatuses, CompileSystemErrors, CompileSystemPaths, CompileSystemScripts);
+
+	ENiagaraScriptCompileStatus AggregateStatus = ENiagaraScriptCompileStatus::NCS_UpToDate;
+	FString AggregateErrors;
+
+	for (int32 i = 0; i < CompileSystemStatuses.Num(); i++)
+	{
+		AggregateStatus = FNiagaraEditorUtilities::UnionCompileStatus(AggregateStatus, CompileSystemStatuses[i]);
+		AggregateErrors += CompileSystemPaths[i] + TEXT(" ") + FNiagaraEditorUtilities::StatusToText(CompileSystemStatuses[i]).ToString() + TEXT("\n");
+		AggregateErrors += CompileSystemErrors[i] + TEXT("\n");
+	}
+
+	UpdateCompileStatus(AggregateStatus, AggregateErrors, CompileSystemStatuses, CompileSystemErrors, CompileSystemPaths, CompileSystemScripts);
+
+	if (OnSystemCompiledDelegate.IsBound())
+	{
+		OnSystemCompiledDelegate.Broadcast();
+	}
+}
+
+void FNiagaraSystemScriptViewModel::OnGraphChanged(const struct FEdGraphEditAction& InAction)
+{
+	if (InAction.Action == GRAPHACTION_SelectNode)
+	{
+		return;
+	}
+}
\ No newline at end of file
diff --git a/Engine/Plugins/FX/Niagara/Source/NiagaraEditor/Private/ViewModels/NiagaraSystemScriptViewModel.h b/Engine/Plugins/FX/Niagara/Source/NiagaraEditor/Private/ViewModels/NiagaraSystemScriptViewModel.h
new file mode 100644
--- /dev/null
+++ b/Engine/Plugins/FX/Niagara/Source/NiagaraEditor/Private/ViewModels/NiagaraSystemScriptViewModel.h
@@ -0,0 +1,39 @@
+// Copyright 1998-2016 Epic Games, Inc. All Rights Reserved.
+
+#pragma once
+
+#include "NiagaraScriptViewModel.h"
+
+class FNiagaraScriptViewModel;
+class UNiagaraSystem;
+
+/** View model which manages the System script. */
+class FNiagaraSystemScriptViewModel : public FNiagaraScriptViewModel
+{
+public:
+	DECLARE_MULTICAST_DELEGATE(FOnSystemCompiled)
+
+public:
+	FNiagaraSystemScriptViewModel(UNiagaraSystem& InSystem);
+
+	~FNiagaraSystemScriptViewModel();
+
+	/** Rebuilds the emitter nodes in the System script due to data changes. */
+	void RebuildEmitterNodes();
+
+	FOnSystemCompiled& OnSystemCompiled();
+
+	void CompileSystem();
+
+private:
+	/** Called whenever the graph for the System script changes. */
+	void OnGraphChanged(const struct FEdGraphEditAction& InAction);
+
+	/** The System who's script is getting viewed and edited by this view model. */
+	UNiagaraSystem& System;
+
+	/** A handle to the on graph changed delegate. */
+	FDelegateHandle OnGraphChangedHandle;
+
+	FOnSystemCompiled OnSystemCompiledDelegate;
+};
\ No newline at end of file
