Copying //UE4/Orion-Staging to //UE4/Main (Source: //Orion/Dev-General @ 3606378)

diff --git a/Engine/Plugins/FX/Niagara/Source/NiagaraEditor/Private/ViewModels/NiagaraSystemViewModel.cpp b/Engine/Plugins/FX/Niagara/Source/NiagaraEditor/Private/ViewModels/NiagaraSystemViewModel.cpp
new file mode 100644
--- /dev/null
+++ b/Engine/Plugins/FX/Niagara/Source/NiagaraEditor/Private/ViewModels/NiagaraSystemViewModel.cpp
@@ -0,0 +1,1232 @@
+// Copyright 1998-2016 Epic Games, Inc. All Rights Reserved.
+
+#include "NiagaraSystemViewModel.h"
+#include "NiagaraSystem.h"
+#include "NiagaraEmitterHandle.h"
+#include "NiagaraEmitter.h"
+#include "NiagaraScriptSource.h"
+#include "NiagaraEditorUtilities.h"
+#include "NiagaraEmitterHandleViewModel.h"
+#include "NiagaraEmitterViewModel.h"
+#include "NiagaraSystemScriptViewModel.h"
+#include "NiagaraScriptGraphViewModel.h"
+#include "NiagaraSequence.h"
+#include "MovieSceneNiagaraEmitterTrack.h"
+#include "MovieSceneNiagaraEmitterSection.h"
+#include "NiagaraSystemInstance.h"
+#include "NiagaraEditorSettings.h"
+#include "NiagaraSystemEditorData.h"
+#include "NiagaraGraph.h"
+#include "NiagaraNodeInput.h"
+#include "NiagaraEditorModule.h"
+
+#include "Editor.h"
+
+#include "ScopedTransaction.h"
+#include "MovieScene.h"
+#include "ISequencerModule.h"
+#include "EditorSupportDelegates.h"
+#include "ModuleManager.h"
+#include "TNiagaraViewModelManager.h"
+#include "Math/NumericLimits.h"
+#include "NiagaraComponent.h"
+#include "MovieSceneFolder.h"
+
+DECLARE_CYCLE_STAT(TEXT("SystemViewModel - CompileSystem"), STAT_NiagaraEditor_SystemViewModel_CompileSystem, STATGROUP_NiagaraEditor);
+
+#define LOCTEXT_NAMESPACE "NiagaraSystemViewModel"
+
+template<> TMap<UNiagaraSystem*, TArray<FNiagaraSystemViewModel*>> TNiagaraViewModelManager<UNiagaraSystem, FNiagaraSystemViewModel>::ObjectsToViewModels{};
+
+FNiagaraSystemViewModel::FNiagaraSystemViewModel(UNiagaraSystem& InSystem, FNiagaraSystemViewModelOptions InOptions)
+	: System(InSystem)
+	, PreviewComponent(nullptr)
+	, SystemInstance(nullptr)
+	, SystemScriptViewModel(MakeShareable(new FNiagaraSystemScriptViewModel(System)))
+	, NiagaraSequence(nullptr)
+	, bSettingSequencerTimeDirectly(false)
+	, bCanRemoveEmittersFromTimeline(InOptions.bCanRemoveEmittersFromTimeline)
+	, bCanRenameEmittersFromTimeline(InOptions.bCanRenameEmittersFromTimeline)
+	, bCanAddEmittersFromTimeline(InOptions.bCanAddEmittersFromTimeline)
+	, bUseSystemExecStateForTimelineReset(InOptions.bUseSystemExecStateForTimelineReset)
+	, OnGetSequencerAddMenuContent(InOptions.OnGetSequencerAddMenuContent)
+	, bUpdatingFromSequencerDataChange(false)
+	, bUpdatingSystemSelectionFromSequencer(false)
+	, bUpdatingSequencerSelectionFromSystem(false)
+{
+	SetupPreviewComponentAndInstance();
+	SetupSequencer();
+	RefreshAll();
+	GEditor->RegisterForUndo(this);
+	RegisteredHandle = RegisterViewModelWithMap(&InSystem, this);
+}
+
+void FNiagaraSystemViewModel::Cleanup()
+{
+	UE_LOG(LogNiagaraEditor, Warning, TEXT("Cleanup System view model %p"), this);
+
+	if (PreviewComponent)
+	{
+		PreviewComponent->OnSystemInstanceChanged().RemoveAll(this);
+	}
+
+	CurveOwner.EmptyCurves();
+
+	GEditor->UnregisterForUndo(this);
+
+	// Make sure that we clear out all of our event handlers
+	UnregisterViewModelWithMap(RegisteredHandle);
+
+	for (TSharedRef<FNiagaraEmitterHandleViewModel>& HandleRef : EmitterHandleViewModels)
+	{
+		HandleRef->OnPropertyChanged().RemoveAll(this);
+		HandleRef->GetEmitterViewModel()->OnPropertyChanged().RemoveAll(this);
+		HandleRef->GetEmitterViewModel()->OnScriptCompiled().RemoveAll(this);
+	}
+	EmitterHandleViewModels.Empty();
+
+	if (Sequencer.IsValid())
+	{
+		Sequencer->OnMovieSceneDataChanged().RemoveAll(this);
+		Sequencer->OnGlobalTimeChanged().RemoveAll(this);
+		Sequencer->GetSelectionChangedTracks().RemoveAll(this);
+		Sequencer->GetSelectionChangedSections().RemoveAll(this);
+		Sequencer.Reset();
+	}
+
+	PreviewComponent = nullptr;
+}
+
+FNiagaraSystemViewModel::~FNiagaraSystemViewModel()
+{
+	Cleanup();
+
+	UE_LOG(LogNiagaraEditor, Warning, TEXT("Deleting System view model %p"), this);
+}
+
+const TArray<TSharedRef<FNiagaraEmitterHandleViewModel>>& FNiagaraSystemViewModel::GetEmitterHandleViewModels()
+{
+	return EmitterHandleViewModels;
+}
+
+TSharedRef<FNiagaraSystemScriptViewModel> FNiagaraSystemViewModel::GetSystemScriptViewModel()
+{
+	return SystemScriptViewModel;
+}
+
+void FNiagaraSystemViewModel::CompileSystem()
+{
+	SCOPE_CYCLE_COUNTER(STAT_NiagaraEditor_SystemViewModel_CompileSystem);
+	KillSystemInstances();
+	SystemScriptViewModel->CompileSystem();
+	OnSystemCompiledDelegate.Broadcast();
+}
+
+ENiagaraScriptCompileStatus FNiagaraSystemViewModel::GetLatestCompileStatus()
+{
+	return SystemScriptViewModel->GetLatestCompileStatus();
+}
+
+const TArray<FGuid>& FNiagaraSystemViewModel::GetSelectedEmitterHandleIds()
+{
+	return SelectedEmitterHandleIds;
+}
+
+void FNiagaraSystemViewModel::SetSelectedEmitterHandlesById(TArray<FGuid> InSelectedEmitterHandleIds)
+{
+	bool bSelectionChanged = false;
+	if (SelectedEmitterHandleIds.Num() == InSelectedEmitterHandleIds.Num())
+	{
+		for (FGuid InSelectedEmitterHandleId : InSelectedEmitterHandleIds)
+		{
+			if (SelectedEmitterHandleIds.Contains(InSelectedEmitterHandleId) == false)
+			{
+				bSelectionChanged = true;
+				break;
+			}
+		}
+	}
+	else
+	{
+		bSelectionChanged = true;
+	}
+
+	SelectedEmitterHandleIds.Empty();
+	SelectedEmitterHandleIds.Append(InSelectedEmitterHandleIds);
+	if (bSelectionChanged)
+	{
+		if (bUpdatingSystemSelectionFromSequencer == false)
+		{
+			UpdateSequencerFromEmitterHandleSelection();
+		}
+		OnSelectedEmitterHandlesChangedDelegate.Broadcast();
+	}
+}
+
+void FNiagaraSystemViewModel::SetSelectedEmitterHandleById(FGuid InSelectedEmitterHandleId)
+{
+	TArray<FGuid> SingleSelectedEmitterHandleId;
+	SingleSelectedEmitterHandleId.Add(InSelectedEmitterHandleId);
+	SetSelectedEmitterHandlesById(SingleSelectedEmitterHandleId);
+}
+
+void FNiagaraSystemViewModel::GetSelectedEmitterHandles(TArray<TSharedRef<FNiagaraEmitterHandleViewModel>>& OutSelectedEmitterHanldles)
+{
+	for (TSharedRef<FNiagaraEmitterHandleViewModel> EmitterHandleViewModel : EmitterHandleViewModels)
+	{
+		if (SelectedEmitterHandleIds.Contains(EmitterHandleViewModel->GetId()))
+		{
+			OutSelectedEmitterHanldles.Add(EmitterHandleViewModel);
+		}
+	}
+}
+
+const UNiagaraSystemEditorData& FNiagaraSystemViewModel::GetEditorData() const
+{
+	const UNiagaraSystemEditorData* EditorData = Cast<UNiagaraSystemEditorData>(System.GetEditorData());
+	if (EditorData == nullptr)
+	{
+		EditorData = GetDefault<UNiagaraSystemEditorData>();
+	}
+	return *EditorData;
+}
+
+UNiagaraSystemEditorData& FNiagaraSystemViewModel::GetOrCreateEditorData()
+{
+	UNiagaraSystemEditorData* EditorData = Cast<UNiagaraSystemEditorData>(System.GetEditorData());
+	if (EditorData == nullptr)
+	{
+		EditorData = NewObject<UNiagaraSystemEditorData>(&System, NAME_None, RF_Transactional);
+		System.Modify();
+		System.SetEditorData(EditorData);
+	}
+	return *EditorData;
+}
+
+UNiagaraComponent* FNiagaraSystemViewModel::GetPreviewComponent()
+{
+	return PreviewComponent;
+}
+
+TSharedPtr<ISequencer> FNiagaraSystemViewModel::GetSequencer()
+{
+	return Sequencer;
+}
+
+FNiagaraCurveOwner& FNiagaraSystemViewModel::GetCurveOwner()
+{
+	return CurveOwner;
+}
+
+bool FNiagaraSystemViewModel::GetSystemIsTransient() const
+{
+	return System.HasAnyFlags(RF_Transient);
+}
+
+bool FNiagaraSystemViewModel::GetCanAddEmittersFromTimeline() const
+{
+	return bCanAddEmittersFromTimeline;
+}
+
+void FNiagaraSystemViewModel::AddEmitterFromAssetData(const FAssetData& AssetData)
+{
+	UNiagaraEmitter* Emitter = Cast<UNiagaraEmitter>(AssetData.GetAsset());
+	if (Emitter != nullptr)
+	{
+		AddEmitter(*Emitter);
+	}
+}
+
+void FNiagaraSystemViewModel::AddEmitter(UNiagaraEmitter& Emitter)
+{
+	// Kill all system instances before modifying the emitter handle list to prevent accessing deleted data.
+	KillSystemInstances();
+
+	const FScopedTransaction Transaction(LOCTEXT("AddEmitter", "Add emitter"));
+
+	TSet<FName> EmitterHandleNames;
+	for (const FNiagaraEmitterHandle& EmitterHandle : System.GetEmitterHandles())
+	{
+		EmitterHandleNames.Add(EmitterHandle.GetName());
+	}
+
+	System.Modify();
+	FNiagaraEmitterHandle EmitterHandle = System.AddEmitterHandle(Emitter, FNiagaraEditorUtilities::GetUniqueName(Emitter.GetFName(), EmitterHandleNames));
+	SystemScriptViewModel->RebuildEmitterNodes();
+
+	if (System.GetNumEmitters() == 1)
+	{
+		// When adding a new emitter to an empty system start playing.
+		Sequencer->SetPlaybackStatus(EMovieScenePlayerStatus::Playing);
+	}
+		
+	RefreshAll();
+
+	SetSelectedEmitterHandleById(EmitterHandle.GetId());
+}
+
+void FNiagaraSystemViewModel::DuplicateEmitter(TSharedRef<FNiagaraEmitterHandleViewModel> EmitterHandleToDuplicate)
+{
+	if (EmitterHandleToDuplicate->GetEmitterHandle() == nullptr)
+	{
+		return;
+	}
+
+	// Kill all system instances before modifying the emitter handle list to prevent accessing deleted data.
+	KillSystemInstances();
+
+	const FScopedTransaction DuplicateTransaction(LOCTEXT("DuplicateEmitter", "Duplicate emitter"));
+
+	TSet<FName> EmitterHandleNames;
+	for (const FNiagaraEmitterHandle& EmitterHandle : System.GetEmitterHandles())
+	{
+		EmitterHandleNames.Add(EmitterHandle.GetName());
+	}
+
+	System.Modify();
+	FNiagaraEmitterHandle EmitterHandle = System.DuplicateEmitterHandle(*EmitterHandleToDuplicate->GetEmitterHandle(), FNiagaraEditorUtilities::GetUniqueName(EmitterHandleToDuplicate->GetEmitterHandle()->GetName(), EmitterHandleNames));
+	SystemScriptViewModel->RebuildEmitterNodes();
+
+	RefreshAll();
+}
+
+void FNiagaraSystemViewModel::DeleteEmitter(TSharedRef<FNiagaraEmitterHandleViewModel> EmitterHandleToDelete)
+{
+	TSet<FGuid> IdsToDelete;
+	IdsToDelete.Add(EmitterHandleToDelete->GetId());
+	DeleteEmitters(IdsToDelete);
+}
+
+void FNiagaraSystemViewModel::DeleteEmitters(TSet<FGuid> EmitterHandleIdsToDelete)
+{
+	if (EmitterHandleIdsToDelete.Num() > 0)
+	{
+		// Kill all system instances before modifying the emitter handle list to prevent accessing deleted data.
+		KillSystemInstances();
+
+		const FScopedTransaction DeleteTransaction(EmitterHandleIdsToDelete.Num() == 1 
+			? LOCTEXT("DeleteEmitter", "Delete emitter")
+			: LOCTEXT("DeleteEmitters", "Delete emitters"));
+
+		System.Modify();
+		System.RemoveEmitterHandlesById(EmitterHandleIdsToDelete);
+		SystemScriptViewModel->RebuildEmitterNodes();
+
+		RefreshAll();
+	}
+}
+
+FNiagaraSystemViewModel::FOnEmitterHandleViewModelsChanged& FNiagaraSystemViewModel::OnEmitterHandleViewModelsChanged()
+{
+	return OnEmitterHandleViewModelsChangedDelegate;
+}
+
+FNiagaraSystemViewModel::FOnCurveOwnerChanged& FNiagaraSystemViewModel::OnCurveOwnerChanged()
+{
+	return OnCurveOwnerChangedDelegate;
+}
+
+FNiagaraSystemViewModel::FOnSelectedEmitterHandlesChanged& FNiagaraSystemViewModel::OnSelectedEmitterHandlesChanged()
+{
+	return OnSelectedEmitterHandlesChangedDelegate;
+}
+
+FNiagaraSystemViewModel::FOnPostSequencerTimeChange& FNiagaraSystemViewModel::OnPostSequencerTimeChanged()
+{
+	return OnPostSequencerTimeChangeDelegate;
+}
+
+FNiagaraSystemViewModel::FOnSystemCompiled& FNiagaraSystemViewModel::OnSystemCompiled()
+{
+	return OnSystemCompiledDelegate;
+}
+
+void FNiagaraSystemViewModel::AddReferencedObjects(FReferenceCollector& Collector)
+{
+	if (PreviewComponent != nullptr)
+	{
+		Collector.AddReferencedObject(PreviewComponent);
+	}
+	if (NiagaraSequence != nullptr)
+	{
+		Collector.AddReferencedObject(NiagaraSequence);
+	}
+}
+
+void FNiagaraSystemViewModel::PostUndo(bool bSuccess)
+{
+	RefreshAll();
+}
+
+void FNiagaraSystemViewModel::Tick(float DeltaTime)
+{
+	bool bRecompile = false;
+	if (SystemScriptViewModel->GetLatestCompileStatus() == ENiagaraScriptCompileStatus::NCS_Dirty)
+	{
+		//SystemScriptViewModel->CompileSystem();
+		//UE_LOG(LogNiagaraEditor, Log, TEXT("Compiling %s due to dirty scripts."), *System.GetName());
+		bRecompile |= true;
+	}
+
+	for (TSharedRef<FNiagaraEmitterHandleViewModel> EmitterHandleViewModel : EmitterHandleViewModels)
+	{
+		if (EmitterHandleViewModel->GetEmitterViewModel()->GetLatestCompileStatus() == ENiagaraScriptCompileStatus::NCS_Dirty)
+		{
+			bRecompile |= true;
+			//EmitterHandleViewModel->GetEmitterViewModel()->CompileScripts();
+			//UE_LOG(LogNiagaraEditor, Log, TEXT("Compiling %s - %s due to dirty scripts."), *System.GetName(), *EmitterHandleViewModel->GetName().ToString());
+		}
+	}
+
+	if (bRecompile)
+	{
+		CompileSystem();
+	}
+}
+
+bool FNiagaraSystemViewModel::IsTickable() const
+{
+	return GetDefault<UNiagaraEditorSettings>()->bAutoCompile;
+}
+
+TStatId FNiagaraSystemViewModel::GetStatId() const
+{
+	RETURN_QUICK_DECLARE_CYCLE_STAT(FNiagaraSystemViewModel, STATGROUP_Tickables);
+}
+
+void FNiagaraSystemViewModel::SetupPreviewComponentAndInstance()
+{
+	PreviewComponent = NewObject<UNiagaraComponent>(GetTransientPackage(), NAME_None, RF_Transient);
+	PreviewComponent->CastShadow = 1;
+	PreviewComponent->bCastDynamicShadow = 1;
+	PreviewComponent->SetAsset(&System);
+	PreviewComponent->SetForceSolo(true);
+	PreviewComponent->SetAgeUpdateMode(UNiagaraComponent::EAgeUpdateMode::DesiredAge);
+	PreviewComponent->Activate(true);
+
+	UNiagaraSystemEditorData& EditorData = GetOrCreateEditorData();
+	FTransform OwnerTransform = EditorData.GetOwnerTransform();
+	PreviewComponent->SetRelativeTransform(OwnerTransform);
+
+	PreviewComponent->OnSystemInstanceChanged().AddRaw(this, &FNiagaraSystemViewModel::PreviewComponentSystemInstanceChanged);
+	PreviewComponentSystemInstanceChanged();
+}
+
+void FNiagaraSystemViewModel::RefreshAll()
+{
+	ReInitializeSystemInstances();
+	RefreshEmitterHandleViewModels();
+	RefreshSequencerTracks();
+	ResetCurveData();
+}
+
+void FNiagaraSystemViewModel::NotifyDataObjectChanged(UObject* ChangedObject)
+{
+	UNiagaraDataInterface* ChangedDataInterface = Cast<UNiagaraDataInterface>(ChangedObject);
+	if (ChangedDataInterface)
+	{
+		UpdateCompiledDataInterfaces(ChangedDataInterface);
+	}
+
+	UNiagaraDataInterfaceCurveBase* ChangedDataInterfaceCurve = Cast<UNiagaraDataInterfaceCurveBase>(ChangedDataInterface);
+	if (ChangedDataInterfaceCurve || ChangedObject == nullptr)
+	{
+		ResetCurveData();
+	}
+	ReInitializeSystemInstances();
+}
+
+void FNiagaraSystemViewModel::RefreshEmitterHandleViewModels()
+{
+	TArray<TSharedRef<FNiagaraEmitterHandleViewModel>> OldViewModels = EmitterHandleViewModels;
+	EmitterHandleViewModels.Empty();
+
+	// Map existing view models to the real instances that now exist. Reuse if we can. Create a new one if we cannot.
+	TArray<FGuid> ValidEmitterHandleIds;
+	int32 i;
+	for (i = 0; i < System.GetNumEmitters(); ++i)
+	{
+		FNiagaraEmitterHandle* EmitterHandle = &System.GetEmitterHandle(i);
+		TSharedPtr<FNiagaraEmitterInstance> Simulation = SystemInstance ? SystemInstance->GetSimulationForHandle(*EmitterHandle) : nullptr;
+		ValidEmitterHandleIds.Add(EmitterHandle->GetId());
+
+		bool bAdd = OldViewModels.Num() <= i;
+		if (bAdd)
+		{
+			TSharedRef<FNiagaraEmitterHandleViewModel> ViewModel = MakeShareable(new FNiagaraEmitterHandleViewModel(EmitterHandle, Simulation, System));
+			// Since we're adding fresh, we need to register all the event handlers.
+			ViewModel->OnPropertyChanged().AddRaw(this, &FNiagaraSystemViewModel::EmitterHandlePropertyChanged, ViewModel);
+			ViewModel->GetEmitterViewModel()->OnPropertyChanged().AddRaw(this, &FNiagaraSystemViewModel::EmitterPropertyChanged, ViewModel);
+			ViewModel->GetEmitterViewModel()->OnScriptCompiled().AddRaw(this, &FNiagaraSystemViewModel::ScriptCompiled);
+			EmitterHandleViewModels.Add(ViewModel);
+		}
+		else
+		{
+			TSharedRef<FNiagaraEmitterHandleViewModel> ViewModel = OldViewModels[i];
+			ViewModel->Set(EmitterHandle, Simulation, System);
+			EmitterHandleViewModels.Add(ViewModel);
+		}
+
+	}
+
+	check(EmitterHandleViewModels.Num() == System.GetNumEmitters());
+
+	// Clear out any old view models that may still be left around.
+	for (; i < OldViewModels.Num(); i++)
+	{
+		TSharedRef<FNiagaraEmitterHandleViewModel> ViewModel = OldViewModels[i];
+		ViewModel->OnPropertyChanged().RemoveAll(this);
+		ViewModel->GetEmitterViewModel()->OnPropertyChanged().RemoveAll(this);
+		ViewModel->GetEmitterViewModel()->OnScriptCompiled().RemoveAll(this);
+		ViewModel->Set(nullptr, nullptr, System);
+	}
+
+	// Remove any invalid ids from the handle selection.
+	auto EmitterHandleIdIsInvalid = [&](FGuid& EmitterHandleId) { return ValidEmitterHandleIds.Contains(EmitterHandleId) == false; };
+	int32 NumRemoved = SelectedEmitterHandleIds.RemoveAll(EmitterHandleIdIsInvalid);
+
+	OnEmitterHandleViewModelsChangedDelegate.Broadcast();
+	if (NumRemoved > 0)
+	{
+		OnSelectedEmitterHandlesChangedDelegate.Broadcast();
+	}
+}
+
+TRange<float> SequencerDefaultPlaybackRange(0, 1000.0f);
+TRange<float> SequencerDefaultViewRange(0, 10.0f);
+
+void PopulateChildMovieSceneFoldersFromNiagaraFolders(const UNiagaraSystemEditorFolder* NiagaraFolder, UMovieSceneFolder* MovieSceneFolder, const TMap<FGuid, UMovieSceneNiagaraEmitterTrack*>& EmitterHandleIdToTrackMap)
+{
+	for (const UNiagaraSystemEditorFolder* ChildNiagaraFolder : NiagaraFolder->GetChildFolders())
+	{
+		UMovieSceneFolder* MatchingMovieSceneFolder = nullptr;
+		for (UMovieSceneFolder* ChildMovieSceneFolder : MovieSceneFolder->GetChildFolders())
+		{
+			if (ChildMovieSceneFolder->GetFolderName() == ChildNiagaraFolder->GetFolderName())
+			{
+				MatchingMovieSceneFolder = ChildMovieSceneFolder;
+			}
+		}
+
+		if (MatchingMovieSceneFolder == nullptr)
+		{
+			MatchingMovieSceneFolder = NewObject<UMovieSceneFolder>(MovieSceneFolder, ChildNiagaraFolder->GetFolderName(), RF_Transactional);
+			MatchingMovieSceneFolder->SetFolderName(ChildNiagaraFolder->GetFolderName());
+			MovieSceneFolder->AddChildFolder(MatchingMovieSceneFolder);
+		}
+
+		PopulateChildMovieSceneFoldersFromNiagaraFolders(ChildNiagaraFolder, MatchingMovieSceneFolder, EmitterHandleIdToTrackMap);
+	}
+
+	for (const FGuid& ChildEmitterHandleId : NiagaraFolder->GetChildEmitterHandleIds())
+	{
+		UMovieSceneNiagaraEmitterTrack* const* TrackPtr = EmitterHandleIdToTrackMap.Find(ChildEmitterHandleId);
+		if (TrackPtr != nullptr && MovieSceneFolder->GetChildMasterTracks().Contains(*TrackPtr) == false)
+		{
+			MovieSceneFolder->AddChildMasterTrack(*TrackPtr);
+		}
+	}
+}
+
+void FNiagaraSystemViewModel::RefreshSequencerTracks()
+{
+	TArray<UMovieSceneTrack*> MasterTracks = NiagaraSequence->GetMovieScene()->GetMasterTracks();
+	for (UMovieSceneTrack* MasterTrack : MasterTracks)
+	{
+		if (MasterTrack != nullptr)
+		{
+			NiagaraSequence->GetMovieScene()->RemoveMasterTrack(*MasterTrack);
+		}
+	}
+
+	float MinEmitterTime = 0;
+	float MaxEmitterTime = 0;
+	TMap<FGuid, UMovieSceneNiagaraEmitterTrack*> EmitterHandleIdToTrackMap;
+	for (TSharedRef<FNiagaraEmitterHandleViewModel> EmitterHandleViewModel : EmitterHandleViewModels)
+	{
+		UMovieSceneNiagaraEmitterTrack* EmitterTrack = Cast<UMovieSceneNiagaraEmitterTrack>(NiagaraSequence->GetMovieScene()->AddMasterTrack(UMovieSceneNiagaraEmitterTrack::StaticClass()));
+		EmitterTrack->SetEmitterHandle(EmitterHandleViewModel);
+		RefreshSequencerTrack(EmitterTrack);
+		MinEmitterTime = FMath::Min(MinEmitterTime, EmitterHandleViewModel->GetEmitterViewModel()->GetStartTime());
+		MaxEmitterTime = FMath::Max(MaxEmitterTime, EmitterHandleViewModel->GetEmitterViewModel()->GetEndTime());
+		EmitterHandleIdToTrackMap.Add(EmitterHandleViewModel->GetId(), EmitterTrack);
+	}
+
+	TArray<UMovieSceneFolder*>& MovieSceneRootFolders = NiagaraSequence->GetMovieScene()->GetRootFolders();
+	MovieSceneRootFolders.Empty();
+
+	const UNiagaraSystemEditorData& SystemEditorData = GetEditorData();
+	UNiagaraSystemEditorFolder& RootFolder = SystemEditorData.GetRootFolder();
+	for (const UNiagaraSystemEditorFolder* RootChildFolder : RootFolder.GetChildFolders())
+	{
+		UMovieSceneFolder* MovieSceneRootFolder = NewObject<UMovieSceneFolder>(NiagaraSequence->GetMovieScene(), RootChildFolder->GetFolderName(), RF_Transactional);
+		MovieSceneRootFolder->SetFolderName(RootChildFolder->GetFolderName());
+		MovieSceneRootFolders.Add(MovieSceneRootFolder);
+		PopulateChildMovieSceneFoldersFromNiagaraFolders(RootChildFolder, MovieSceneRootFolder, EmitterHandleIdToTrackMap);
+	}
+
+	// Expand the view range to show all emitters.
+	TRange<float> CurrentViewRange = NiagaraSequence->GetMovieScene()->GetEditorData().ViewRange;
+	float NewMinViewRange = FMath::Min(MinEmitterTime, CurrentViewRange.GetLowerBoundValue());
+	float NewMaxViewRange = FMath::Max(MaxEmitterTime, CurrentViewRange.GetUpperBoundValue());
+
+	NiagaraSequence->GetMovieScene()->GetEditorData().ViewRange = TRange<float>(NewMinViewRange, NewMaxViewRange);
+
+	Sequencer->NotifyMovieSceneDataChanged(EMovieSceneDataChangeType::MovieSceneStructureItemsChanged);
+	Sequencer->SetGlobalTime(0);
+}
+
+UMovieSceneNiagaraEmitterTrack* FNiagaraSystemViewModel::GetTrackForHandleViewModel(TSharedRef<FNiagaraEmitterHandleViewModel> EmitterHandleViewModel)
+{
+	for (UMovieSceneTrack* Track : NiagaraSequence->GetMovieScene()->GetMasterTracks())
+	{
+		UMovieSceneNiagaraEmitterTrack* EmitterTrack = CastChecked<UMovieSceneNiagaraEmitterTrack>(Track);
+		if (EmitterTrack->GetEmitterHandle() == EmitterHandleViewModel)
+		{
+			return EmitterTrack;
+		}
+	}
+	return nullptr;
+}
+
+void FNiagaraSystemViewModel::RefreshSequencerTrack(UMovieSceneNiagaraEmitterTrack* EmitterTrack)
+{
+	if(EmitterTrack != nullptr)
+	{
+		TSharedPtr<FNiagaraEmitterHandleViewModel> EmitterHandleViewModel = EmitterTrack->GetEmitterHandle();
+		UNiagaraEmitter* Emitter = EmitterHandleViewModel->GetEmitterViewModel()->GetEmitter();
+
+		TArray<UMovieSceneSection*> Sections = EmitterTrack->GetAllSections();
+		UMovieSceneNiagaraEmitterSection* EmitterSection = Cast<UMovieSceneNiagaraEmitterSection>(Sections.Num() == 1 ? Sections[0] : nullptr);
+
+		if (EmitterSection != nullptr && Emitter != nullptr)
+		{
+			EmitterTrack->SetDisplayName(EmitterHandleViewModel->GetNameText());
+
+			bool bIsInfinite = EmitterHandleViewModel->GetEmitterViewModel()->GetStartTime() == 0
+				&& EmitterHandleViewModel->GetEmitterViewModel()->GetEndTime() == 0;
+			float StartTime;
+			float EndTime;
+
+// 			if (bIsInfinite && Emitter->Bursts.Num() > 0)
+// 			{
+// 				StartTime = TNumericLimits<float>::Max();
+// 				EndTime = TNumericLimits<float>::Lowest();
+// 				for (const FNiagaraEmitterBurst& Burst : Emitter->Bursts)
+// 				{
+// 					StartTime = FMath::Min(StartTime, Burst.Time);
+// 					EndTime = FMath::Max(EndTime, Burst.Time);
+// 				}
+// 			}
+// 			else
+			{
+				StartTime = EmitterHandleViewModel->GetEmitterViewModel()->GetStartTime();
+				EndTime = EmitterHandleViewModel->GetEmitterViewModel()->GetEndTime();
+			}
+
+			EmitterSection->SetEmitterHandle(EmitterHandleViewModel.ToSharedRef());
+			EmitterSection->SetIsActive(EmitterHandleViewModel->GetIsEnabled());
+			EmitterSection->SetStartTime(StartTime);
+			EmitterSection->SetEndTime(EndTime);
+			EmitterSection->SetIsInfinite(bIsInfinite);
+
+// 			TSharedPtr<FBurstCurve> BurstCurve = EmitterSection->GetBurstCurve();
+// 			if (BurstCurve.IsValid())
+// 			{
+// 				BurstCurve->Reset();
+// 				for (const FNiagaraEmitterBurst& Burst : Emitter->Bursts)
+// 				{
+// 					FMovieSceneBurstKey BurstKey;
+// 					BurstKey.TimeRange = Burst.TimeRange;
+// 					BurstKey.SpawnMinimum = Burst.SpawnMinimum;
+// 					BurstKey.SpawnMaximum = Burst.SpawnMaximum;
+// 					BurstCurve->AddKeyValue(Burst.Time, BurstKey);
+// 				}
+// 			}
+		}
+	}
+}
+
+void FNiagaraSystemViewModel::SetupSequencer()
+{
+	NiagaraSequence = NewObject<UNiagaraSequence>(GetTransientPackage());
+	UMovieScene* MovieScene = NewObject<UMovieScene>(NiagaraSequence, FName("Niagara System MovieScene"), RF_Transactional);
+	MovieScene->SetPlaybackRange(SequencerDefaultPlaybackRange.GetLowerBoundValue(), SequencerDefaultPlaybackRange.GetUpperBoundValue());
+	MovieScene->GetEditorData().ViewRange = SequencerDefaultViewRange;
+	MovieScene->GetEditorData().WorkingRange = SequencerDefaultPlaybackRange;
+
+	NiagaraSequence->Initialize(this, MovieScene);
+
+	FSequencerViewParams ViewParams(TEXT("NiagaraSequencerSettings"));
+	{
+		ViewParams.InitialScrubPosition = 0;
+		ViewParams.UniqueName = "NiagaraSequenceEditor";
+		ViewParams.OnGetAddMenuContent = OnGetSequencerAddMenuContent;
+	}
+
+	FSequencerInitParams SequencerInitParams;
+	{
+		SequencerInitParams.ViewParams = ViewParams;
+		SequencerInitParams.RootSequence = NiagaraSequence;
+		SequencerInitParams.bEditWithinLevelEditor = false;
+		SequencerInitParams.ToolkitHost = nullptr;
+	}
+
+	ISequencerModule &SequencerModule = FModuleManager::LoadModuleChecked< ISequencerModule >("Sequencer");
+	Sequencer = SequencerModule.CreateSequencer(SequencerInitParams);
+
+	Sequencer->OnMovieSceneDataChanged().AddRaw(this, &FNiagaraSystemViewModel::SequencerDataChanged);
+	Sequencer->OnGlobalTimeChanged().AddRaw(this, &FNiagaraSystemViewModel::SequencerTimeChanged);
+	Sequencer->GetSelectionChangedTracks().AddRaw(this, &FNiagaraSystemViewModel::SequencerTrackSelectionChanged);
+	Sequencer->GetSelectionChangedSections().AddRaw(this, &FNiagaraSystemViewModel::SequencerSectionSelectionChanged);
+	Sequencer->SetPlaybackStatus(System.GetNumEmitters() > 0 
+		? EMovieScenePlayerStatus::Playing
+		: EMovieScenePlayerStatus::Stopped);
+}
+
+void FNiagaraSystemViewModel::ResetSystem()
+{
+ 	if (Sequencer->GetPlaybackStatus() == EMovieScenePlayerStatus::Playing)
+	{
+		Sequencer->SetGlobalTime(0.0f);
+	}
+
+	for (TObjectIterator<UNiagaraComponent> ComponentIt; ComponentIt; ++ComponentIt)
+	{
+		UNiagaraComponent* Component = *ComponentIt;
+		if (Component->GetAsset() == &System)
+		{
+			Component->SynchronizeWithSourceSystem();
+			Component->ResetSystem();
+		}
+	}
+	FEditorSupportDelegates::RedrawAllViewports.Broadcast();
+}
+
+void FNiagaraSystemViewModel::KillSystemInstances()
+{
+	for (TObjectIterator<UNiagaraComponent> ComponentIt; ComponentIt; ++ComponentIt)
+	{
+		UNiagaraComponent* Component = *ComponentIt;
+		if (Component->GetAsset() == &System)
+		{
+			Component->DestroyInstance();
+		}
+	}
+}
+
+void FNiagaraSystemViewModel::ReInitializeSystemInstances()
+{
+	if (Sequencer.IsValid() && Sequencer->GetPlaybackStatus() == EMovieScenePlayerStatus::Playing)
+	{
+		Sequencer->SetGlobalTime(0.0f);
+	}
+
+	for (TObjectIterator<UNiagaraComponent> ComponentIt; ComponentIt; ++ComponentIt)
+	{
+		UNiagaraComponent* Component = *ComponentIt;
+		if (Component->GetAsset() == &System)
+		{
+			Component->SynchronizeWithSourceSystem();
+			Component->ReinitializeSystem();
+		}
+	}
+	FEditorSupportDelegates::RedrawAllViewports.Broadcast();
+}
+
+struct FNiagaraSystemCurveData
+{
+	FRichCurve* Curve;
+	FName Name;
+	FLinearColor Color;
+	UObject* Owner;
+};
+
+void GetCurveData(FString CurveSource, UNiagaraGraph* SourceGraph, TArray<FNiagaraSystemCurveData>& OutCurveData)
+{
+	TArray<UNiagaraNodeInput*> InputNodes;
+	SourceGraph->GetNodesOfClass<UNiagaraNodeInput>(InputNodes);
+	TSet<FName> HandledInputs;
+	for (UNiagaraNodeInput* InputNode : InputNodes)
+	{
+		if (HandledInputs.Contains(InputNode->Input.GetName()) == false)
+		{
+			if (InputNode->Usage == ENiagaraInputNodeUsage::Parameter)
+			{
+				UNiagaraDataInterfaceCurveBase* CurveDataInterface = Cast<UNiagaraDataInterfaceCurveBase>(InputNode->DataInterface);
+				if (CurveDataInterface != nullptr)
+				{
+					TArray<UNiagaraDataInterfaceCurveBase::FCurveData> CurveData;
+					CurveDataInterface->GetCurveData(CurveData);
+					for (UNiagaraDataInterfaceCurveBase::FCurveData CurveDataItem : CurveData)
+					{
+						FNiagaraSystemCurveData SystemCurveData;
+						SystemCurveData.Curve = CurveDataItem.Curve;
+						SystemCurveData.Color = CurveDataItem.Color;
+						SystemCurveData.Owner = CurveDataInterface;
+						FString ParameterName = CurveDataItem.Name == NAME_None
+							? InputNode->Input.GetName().ToString()
+							: InputNode->Input.GetName().ToString() + ".";
+						FString DataName = CurveDataItem.Name == NAME_None
+							? FString()
+							: CurveDataItem.Name.ToString();
+						SystemCurveData.Name = *(CurveSource + ParameterName + DataName);
+						OutCurveData.Add(SystemCurveData);
+					}
+				}
+			}
+			HandledInputs.Add(InputNode->Input.GetName());
+		}
+	}
+}
+
+void FNiagaraSystemViewModel::ResetCurveData()
+{
+	CurveOwner.EmptyCurves();
+
+	TArray<FNiagaraSystemCurveData> CurveData;
+
+	GetCurveData(
+		TEXT("System"),
+		SystemScriptViewModel->GetGraphViewModel()->GetGraph(),
+		CurveData);
+
+	for (TSharedRef<FNiagaraEmitterHandleViewModel> EmitterHandleViewModel : EmitterHandleViewModels)
+	{
+		GetCurveData(
+			EmitterHandleViewModel->GetName().ToString(),
+			EmitterHandleViewModel->GetEmitterViewModel()->GetSharedScriptViewModel()->GetGraphViewModel()->GetGraph(),
+			CurveData);
+	}
+
+	for (FNiagaraSystemCurveData& CurveDataItem : CurveData)
+	{
+		CurveOwner.AddCurve(
+			*CurveDataItem.Curve,
+			CurveDataItem.Name,
+			CurveDataItem.Color,
+			*CurveDataItem.Owner,
+			FNiagaraCurveOwner::FNotifyCurveChanged::CreateRaw(this, &FNiagaraSystemViewModel::CurveChanged));
+	}
+
+	OnCurveOwnerChangedDelegate.Broadcast();
+}
+
+void FNiagaraSystemViewModel::UpdateCompiledDataInterfaces(UNiagaraDataInterface* ChangedDataInterface)
+{
+	SystemScriptViewModel->UpdateCompiledDataInterfaces(ChangedDataInterface);
+	for (TSharedRef<FNiagaraEmitterHandleViewModel> EmitterHandleViewModel : EmitterHandleViewModels)
+	{
+		EmitterHandleViewModel->GetEmitterViewModel()->GetSharedScriptViewModel()->UpdateCompiledDataInterfaces(ChangedDataInterface);
+	}
+}
+
+void FNiagaraSystemViewModel::EmitterHandlePropertyChanged(TSharedRef<FNiagaraEmitterHandleViewModel> EmitterHandleViewModel)
+{
+	// When the emitter handle changes, refresh the System scripts emitter nodes just in cast the
+	// property that changed was the handles emitter.
+	if (bUpdatingFromSequencerDataChange == false)
+	{
+		RefreshSequencerTrack(GetTrackForHandleViewModel(EmitterHandleViewModel));
+	}
+	ReInitializeSystemInstances();
+}
+
+void FNiagaraSystemViewModel::EmitterPropertyChanged(TSharedRef<FNiagaraEmitterHandleViewModel> EmitterHandleViewModel)
+{
+	if (bUpdatingFromSequencerDataChange == false)
+	{
+		RefreshSequencerTrack(GetTrackForHandleViewModel(EmitterHandleViewModel));
+	}
+	ReInitializeSystemInstances();
+}
+
+void FNiagaraSystemViewModel::ScriptCompiled()
+{
+	ReInitializeSystemInstances();
+}
+
+void FNiagaraSystemViewModel::CurveChanged(FRichCurve* ChangedCurve, UObject* InCurveOwner)
+{
+	UNiagaraDataInterface* ChangedDataInterface = Cast<UNiagaraDataInterface>(InCurveOwner);
+	if (ChangedDataInterface != nullptr)
+	{
+		UpdateCompiledDataInterfaces(ChangedDataInterface);
+	}
+	ResetSystem();
+}
+
+void PopulateNiagaraFoldersFromMovieSceneFolders(const TArray<UMovieSceneFolder*>& MovieSceneFolders, const TArray<UMovieSceneTrack*>& MovieSceneTracks, UNiagaraSystemEditorFolder* ParentFolder)
+{
+	TArray<FName> ValidFolderNames;
+	for (UMovieSceneFolder* MovieSceneFolder : MovieSceneFolders)
+	{
+		ValidFolderNames.Add(MovieSceneFolder->GetFolderName());
+		UNiagaraSystemEditorFolder* MatchingNiagaraFolder = nullptr;
+		for (UNiagaraSystemEditorFolder* ChildNiagaraFolder : ParentFolder->GetChildFolders())
+		{
+			CA_ASSUME(ChildNiagaraFolder != nullptr);
+			if (ChildNiagaraFolder->GetFolderName() == MovieSceneFolder->GetFolderName())
+			{
+				MatchingNiagaraFolder = ChildNiagaraFolder;
+				break;
+			}
+		}
+
+		if (MatchingNiagaraFolder == nullptr)
+		{
+			MatchingNiagaraFolder = NewObject<UNiagaraSystemEditorFolder>(ParentFolder, MovieSceneFolder->GetFolderName(), RF_Transactional);
+			MatchingNiagaraFolder->SetFolderName(MovieSceneFolder->GetFolderName());
+			ParentFolder->AddChildFolder(MatchingNiagaraFolder);
+		}
+
+		PopulateNiagaraFoldersFromMovieSceneFolders(MovieSceneFolder->GetChildFolders(), MovieSceneFolder->GetChildMasterTracks(), MatchingNiagaraFolder);
+	}
+
+	TArray<UNiagaraSystemEditorFolder*> ChildNiagaraFolders = ParentFolder->GetChildFolders();
+	for (UNiagaraSystemEditorFolder* ChildNiagaraFolder : ChildNiagaraFolders)
+	{
+		if (ValidFolderNames.Contains(ChildNiagaraFolder->GetFolderName()) == false)
+		{
+			ParentFolder->RemoveChildFolder(ChildNiagaraFolder);
+		}
+	}
+
+	TArray<FGuid> ValidEmitterHandleIds;
+	for (UMovieSceneTrack* MovieSceneTrack : MovieSceneTracks)
+	{
+		UMovieSceneNiagaraEmitterTrack* NiagaraEmitterTrack = Cast<UMovieSceneNiagaraEmitterTrack>(MovieSceneTrack);
+		if (NiagaraEmitterTrack != nullptr)
+		{
+			FGuid EmitterHandleId = NiagaraEmitterTrack->GetEmitterHandle()->GetId();
+			ValidEmitterHandleIds.Add(EmitterHandleId);
+			if (ParentFolder->GetChildEmitterHandleIds().Contains(EmitterHandleId) == false)
+			{
+				ParentFolder->AddChildEmitterHandleId(EmitterHandleId);
+			}
+		}
+	}
+
+	TArray<FGuid> ChildEmitterHandleIds = ParentFolder->GetChildEmitterHandleIds();
+	for (FGuid& ChildEmitterHandleId : ChildEmitterHandleIds)
+	{
+		if (ValidEmitterHandleIds.Contains(ChildEmitterHandleId) == false)
+		{
+			ParentFolder->RemoveChildEmitterHandleId(ChildEmitterHandleId);
+		}
+	}
+}
+
+void FNiagaraSystemViewModel::SequencerDataChanged(EMovieSceneDataChangeType DataChangeType)
+{
+	bUpdatingFromSequencerDataChange = true;
+	TSet<FGuid> VaildTrackEmitterHandleIds;
+	for (UMovieSceneTrack* Track : NiagaraSequence->GetMovieScene()->GetMasterTracks())
+	{
+		UMovieSceneNiagaraEmitterTrack* EmitterTrack = CastChecked<UMovieSceneNiagaraEmitterTrack>(Track);
+		TArray<UMovieSceneSection*> Sections = EmitterTrack->GetAllSections();
+		UMovieSceneNiagaraEmitterSection* EmitterSection = Cast<UMovieSceneNiagaraEmitterSection>(Sections.Num() == 1 ? Sections[0] : nullptr);
+		if (EmitterSection != nullptr)
+		{
+			VaildTrackEmitterHandleIds.Add(EmitterTrack->GetEmitterHandle()->GetId());
+
+			TSharedPtr<FNiagaraEmitterHandleViewModel> EmitterHandleViewModel = EmitterTrack->GetEmitterHandle();
+			if (bCanRenameEmittersFromTimeline)
+			{
+				EmitterHandleViewModel->SetName(*EmitterTrack->GetDisplayName().ToString());
+			}
+			else
+			{
+				EmitterTrack->SetDisplayName(EmitterHandleViewModel->GetNameText());
+			}
+			EmitterHandleViewModel->SetIsEnabled(EmitterSection->IsActive());
+
+			TSharedPtr<FNiagaraEmitterViewModel> EmitterViewModel = EmitterTrack->GetEmitterHandle()->GetEmitterViewModel();
+			if (EmitterSection->IsInfinite() == false)
+			{
+				EmitterViewModel->SetStartTime(EmitterSection->GetStartTime());
+				EmitterViewModel->SetEndTime(EmitterSection->GetEndTime());
+			}
+
+			TSharedPtr<FBurstCurve> BurstCurve = EmitterSection->GetBurstCurve();
+			if (BurstCurve.IsValid())
+			{
+				UNiagaraEmitter* Emitter = EmitterViewModel->GetEmitter();
+// 				if (Emitter)
+// 				{
+// 					Emitter->Bursts.Empty();
+// 				}
+				TArray<FKeyHandle> KeyHandles;
+				for (TKeyTimeIterator<float> KeyIterator = BurstCurve->IterateKeys(); KeyIterator; ++KeyIterator)
+				{
+					KeyHandles.Add(KeyIterator.GetKeyHandle());
+				}
+
+				for (const FKeyHandle& KeyHandle : KeyHandles)
+				{
+					auto Key = BurstCurve->GetKey(KeyHandle).GetValue();
+					FNiagaraEmitterBurst Burst;
+					Burst.Time = Key.Time;
+					Burst.TimeRange = Key.Value.TimeRange;
+					Burst.SpawnMinimum = Key.Value.SpawnMinimum;
+					Burst.SpawnMaximum = Key.Value.SpawnMaximum;
+// 					if (Emitter)
+// 					{
+// 						Emitter->Bursts.Add(Burst);
+// 					}
+				}
+			}
+		}
+	}
+
+	TSet<FGuid> AllEmitterHandleIds;
+	for (TSharedRef<FNiagaraEmitterHandleViewModel> EmitterHandleViewModel : EmitterHandleViewModels)
+	{
+		AllEmitterHandleIds.Add(EmitterHandleViewModel->GetId());
+	}
+
+	TSet<FGuid> RemovedEmitterHandleIds = AllEmitterHandleIds.Difference(VaildTrackEmitterHandleIds);
+	if (RemovedEmitterHandleIds.Num() > 0)
+	{
+		if (bCanRemoveEmittersFromTimeline)
+		{
+			DeleteEmitters(RemovedEmitterHandleIds);
+			// When deleting emitters from sequencer, select a new one if one is available.
+			if (SelectedEmitterHandleIds.Num() == 0 && EmitterHandleViewModels.Num() > 0)
+			{
+				SetSelectedEmitterHandleById(EmitterHandleViewModels[0]->GetId());
+			}
+		}
+		else
+		{
+			RefreshSequencerTracks();
+		}
+	}
+
+	TArray<UMovieSceneTrack*> RootTracks;
+	TArray<UMovieSceneFolder*> RootFolders = NiagaraSequence->GetMovieScene()->GetRootFolders();
+	if (RootFolders.Num() != 0 || GetEditorData().GetRootFolder().GetChildFolders().Num() != 0)
+	{
+		PopulateNiagaraFoldersFromMovieSceneFolders(RootFolders, RootTracks, &GetOrCreateEditorData().GetRootFolder());
+	}
+
+	if (SystemInstance != nullptr)
+	{
+		SystemInstance->Reset(FNiagaraSystemInstance::EResetMode::DeferredReset);
+	}
+	bUpdatingFromSequencerDataChange = false;
+}
+
+void FNiagaraSystemViewModel::SequencerTimeChanged()
+{
+	if (!PreviewComponent)
+	{
+		return;
+	}
+	EMovieScenePlayerStatus::Type CurrentStatus = Sequencer->GetPlaybackStatus();
+	float CurrentSequencerTime = Sequencer->GetGlobalTime();
+	if (SystemInstance != nullptr)
+	{
+		// Avoid reentrancy if we're setting the time directly.
+		if (bSettingSequencerTimeDirectly == false && CurrentSequencerTime != PreviousSequencerTime)
+		{
+			bool bUpdateDesiredAge = false;
+			bool bResetSystemInstance = false;
+			bool bSetEnabled = false;
+
+			if (CurrentStatus == EMovieScenePlayerStatus::Playing)
+			{
+				bool bSystemIsAlive = false;
+				//float MaxEmitterTime = 0.0f;
+
+				if (bUseSystemExecStateForTimelineReset)
+				{
+					if (SystemInstance->GetExecutionState() == ENiagaraExecutionState::Disabled)
+					{
+						bSystemIsAlive = false;
+						bSetEnabled = true;
+					}
+					else
+					{
+						bSystemIsAlive = true;
+					}
+				}
+				else 
+				{
+					for (const TSharedRef<FNiagaraEmitterInstance> Simulation : SystemInstance->GetEmitters())
+					{
+						if (Simulation->GetExecutionState() != ENiagaraExecutionState::Dead)
+						{
+							bSystemIsAlive |= true;
+						}
+
+						/*for (const TSharedRef<FNiagaraEmitterHandleViewModel> HandleViewModel : GetEmitterHandleViewModels())
+						{
+							TSharedRef<FNiagaraEmitterViewModel> EmitterViewModel = HandleViewModel->GetEmitterViewModel();
+							if (EmitterViewModel->GetEmitter() == Simulation->GetEmitterHandle().GetInstance())
+							{
+								float EmitterEndTime = EmitterViewModel->GetNumLoops() == 0
+								? EmitterViewModel->GetEndTime()
+								: (EmitterViewModel->GetEndTime() - EmitterViewModel->GetStartTime()) * EmitterViewModel->GetNumLoops();
+								MaxEmitterTime = FMath::Max(MaxEmitterTime, EmitterEndTime);
+
+							}
+						}*/
+					}
+				}
+
+				bool bIsUpdatingOrCanSpawn = bSystemIsAlive /*|| (MaxEmitterTime == 0.0f || CurrentSequencerTime <= MaxEmitterTime)*/;
+				if (bIsUpdatingOrCanSpawn)
+				{
+					// Skip the first update after going from stopped to playing because snapping in sequencer may have made the time
+					// reverse by a small amount, and sending that update to the System will reset it unnecessarily.
+					bUpdateDesiredAge = PreviousSequencerStatus != EMovieScenePlayerStatus::Stopped;
+				}
+				else
+				{
+					// If there are no active particles and no more particles will be spawned reset the System so it loops.
+					bResetSystemInstance = true;
+				}
+			}
+			else
+			{
+				// If the time changed and we're not playing, or stopping playing, than the user is scrubbing, or jumping to a different time using some
+				// other means so just update the System time.  Skip the first update after going from playing to stopped because snapping in sequencer 
+				// may have made the time reverse by a small amount, and sending that update to the System will reset it unnecessarily.
+				bool bStoppedPlaying = PreviousSequencerStatus == EMovieScenePlayerStatus::Playing && CurrentStatus != EMovieScenePlayerStatus::Playing;
+				bUpdateDesiredAge = !bStoppedPlaying;
+			}
+
+			if (bUpdateDesiredAge)
+			{
+				PreviewComponent->SetDesiredAge(FMath::Max(CurrentSequencerTime, 0.0f));
+
+			//SystemInstance->Tick(0);TODO: What was this for? Possibly reinstate but I don't like use dealing with the instance directly.
+			}
+
+			if (bResetSystemInstance)
+			{
+				if (bSetEnabled)
+				{
+					SystemInstance->Enable();
+				}
+				else
+				{
+					SystemInstance->Reset(FNiagaraSystemInstance::EResetMode::ImmediateReset);
+				}
+				TGuardValue<bool>(bSettingSequencerTimeDirectly, true);
+				Sequencer->SetLocalTime(0);
+			}
+		}
+	}
+
+	PreviousSequencerStatus = CurrentStatus;
+	PreviousSequencerTime = CurrentSequencerTime;
+
+	OnPostSequencerTimeChangeDelegate.Broadcast();
+}
+
+void FNiagaraSystemViewModel::SequencerTrackSelectionChanged(TArray<UMovieSceneTrack*> SelectedTracks)
+{
+	if (bUpdatingSequencerSelectionFromSystem == false)
+	{
+		UpdateEmitterHandleSelectionFromSequencer();
+	}
+}
+
+void FNiagaraSystemViewModel::SequencerSectionSelectionChanged(TArray<UMovieSceneSection*> SelectedSections)
+{
+	if (bUpdatingSequencerSelectionFromSystem == false)
+	{
+		UpdateEmitterHandleSelectionFromSequencer();
+	}
+}
+
+void FNiagaraSystemViewModel::UpdateEmitterHandleSelectionFromSequencer()
+{
+	TArray<FGuid> NewSelectedEmitterHandleIds;
+
+	TArray<UMovieSceneTrack*> SelectedTracks;
+	Sequencer->GetSelectedTracks(SelectedTracks);
+	for (UMovieSceneTrack* SelectedTrack : SelectedTracks)
+	{
+		UMovieSceneNiagaraEmitterTrack* SelectedEmitterTrack = Cast<UMovieSceneNiagaraEmitterTrack>(SelectedTrack);
+		if (SelectedEmitterTrack != nullptr && SelectedEmitterTrack->GetEmitterHandle().IsValid())
+		{
+			NewSelectedEmitterHandleIds.AddUnique(SelectedEmitterTrack->GetEmitterHandle()->GetId());
+		}
+	}
+
+	TArray<UMovieSceneSection*> SelectedSections;
+	Sequencer->GetSelectedSections(SelectedSections);
+	for (UMovieSceneSection* SelectedSection : SelectedSections)
+	{
+		UMovieSceneNiagaraEmitterSection* SelectedEmitterSection = Cast<UMovieSceneNiagaraEmitterSection>(SelectedSection);
+		if (SelectedEmitterSection != nullptr && SelectedEmitterSection->GetEmitterHandle().IsValid())
+		{
+			NewSelectedEmitterHandleIds.AddUnique(SelectedEmitterSection->GetEmitterHandle()->GetId());
+		}
+	}
+
+	TGuardValue<bool> UpdateGuard(bUpdatingSystemSelectionFromSequencer, true);
+	SetSelectedEmitterHandlesById(NewSelectedEmitterHandleIds);
+}
+
+void FNiagaraSystemViewModel::UpdateSequencerFromEmitterHandleSelection()
+{
+	TGuardValue<bool> UpdateGuard(bUpdatingSequencerSelectionFromSystem, true);
+	Sequencer->EmptySelection();
+	for (FGuid SelectedEmitterHandleId : SelectedEmitterHandleIds)
+	{
+		for (UMovieSceneTrack* MasterTrack : NiagaraSequence->GetMovieScene()->GetMasterTracks())
+		{
+			UMovieSceneNiagaraEmitterTrack* EmitterTrack = Cast<UMovieSceneNiagaraEmitterTrack>(MasterTrack);
+			if (EmitterTrack != nullptr && EmitterTrack->GetEmitterHandle()->GetId() == SelectedEmitterHandleId)
+			{
+				Sequencer->SelectTrack(EmitterTrack);
+			}
+		}
+	}
+}
+
+void FNiagaraSystemViewModel::PreviewComponentSystemInstanceChanged()
+{
+	FNiagaraSystemInstance* OldSystemInstance = SystemInstance;
+	SystemInstance = PreviewComponent->GetSystemInstance();
+	if (SystemInstance != OldSystemInstance)
+	{
+		if (SystemInstance != nullptr)
+		{
+			SystemInstance->OnInitialized().AddRaw(this, &FNiagaraSystemViewModel::SystemInstanceInitialized);
+		}
+		else
+		{
+			for (TSharedRef<FNiagaraEmitterHandleViewModel> EmitterHandleViewModel : EmitterHandleViewModels)
+			{
+				if (EmitterHandleViewModel->GetEmitterHandle())
+				{
+					EmitterHandleViewModel->SetSimulation(nullptr);
+				}
+			}
+		}
+	}
+}
+
+void FNiagaraSystemViewModel::SystemInstanceInitialized()
+{
+	for (TSharedRef<FNiagaraEmitterHandleViewModel> EmitterHandleViewModel : EmitterHandleViewModels)
+	{
+		if (EmitterHandleViewModel->GetEmitterHandle())
+		{
+			EmitterHandleViewModel->SetSimulation(SystemInstance->GetSimulationForHandle(*EmitterHandleViewModel->GetEmitterHandle()));
+		}
+	}
+}
+
+void FNiagaraSystemViewModel::ResynchronizeAllHandles()
+{
+	System.ResynchronizeAllHandles();
+	RefreshAll();
+}
+
+#undef LOCTEXT_NAMESPACE // NiagaraSystemViewModel
diff --git a/Engine/Plugins/FX/Niagara/Source/NiagaraEditor/Private/ViewModels/NiagaraSystemViewModel.h b/Engine/Plugins/FX/Niagara/Source/NiagaraEditor/Private/ViewModels/NiagaraSystemViewModel.h
new file mode 100644
--- /dev/null
+++ b/Engine/Plugins/FX/Niagara/Source/NiagaraEditor/Private/ViewModels/NiagaraSystemViewModel.h
@@ -0,0 +1,324 @@
+// Copyright 1998-2016 Epic Games, Inc. All Rights Reserved.
+
+#pragma once
+
+#include "IMovieScenePlayer.h"
+
+#include "EditorUndoClient.h"
+#include "GCObject.h"
+#include "NiagaraCurveOwner.h"
+#include "TNiagaraViewModelManager.h"
+#include "ISequencer.h"
+
+#include "TickableEditorObject.h"
+#include "ISequencerModule.h"
+
+class UNiagaraSystem;
+class UNiagaraComponent;
+class UNiagaraSequence;
+class UMovieSceneNiagaraEmitterTrack;
+struct FNiagaraVariable;
+class FNiagaraEmitterHandleViewModel;
+class FNiagaraSystemScriptViewModel;
+class FNiagaraSystemInstance;
+class ISequencer;
+struct FAssetData;
+class UNiagaraSystemEditorData;
+struct FRichCurve;
+
+/** Defines options for the niagara System view model */
+struct FNiagaraSystemViewModelOptions
+{
+	/** Whether or not the user can remove emitters from the timeline. */
+	bool bCanRemoveEmittersFromTimeline;
+
+	/** Whether or not the user can rename emitters from the timeline. */
+	bool bCanRenameEmittersFromTimeline;
+
+	/** Whether or not the user can add emitters from the timeline. */
+	bool bCanAddEmittersFromTimeline;
+
+	/** A delegate which is used to generate the content for the add menu in sequencer. */
+	FOnGetAddMenuContent OnGetSequencerAddMenuContent;
+
+	/** Whether or not we use the system's execution state to drive when we reset the timeline*/
+	bool bUseSystemExecStateForTimelineReset;
+};
+
+/** A view model for viewing and editing a UNiagaraSystem. */
+class FNiagaraSystemViewModel 
+	: public TSharedFromThis<FNiagaraSystemViewModel>
+	, public FGCObject
+	, public FEditorUndoClient
+	, public FTickableEditorObject
+	, public TNiagaraViewModelManager<UNiagaraSystem, FNiagaraSystemViewModel>
+{
+public:
+	DECLARE_MULTICAST_DELEGATE(FOnEmitterHandleViewModelsChanged);
+
+	DECLARE_MULTICAST_DELEGATE(FOnCurveOwnerChanged);
+
+	DECLARE_MULTICAST_DELEGATE(FOnSelectedEmitterHandlesChanged);
+
+	DECLARE_MULTICAST_DELEGATE(FOnPostSequencerTimeChange)
+
+	DECLARE_MULTICAST_DELEGATE(FOnSystemCompiled);
+
+public:
+	/** Creates a new view model with the supplied System and System instance. */
+	FNiagaraSystemViewModel(UNiagaraSystem& InSystem, FNiagaraSystemViewModelOptions InOptions);
+
+	~FNiagaraSystemViewModel();
+
+	/** Gets an array of the view models for the emitter handles owned by this System. */
+	const TArray<TSharedRef<FNiagaraEmitterHandleViewModel>>& GetEmitterHandleViewModels();
+
+	/** Gets the view model for the System script. */
+	TSharedRef<FNiagaraSystemScriptViewModel> GetSystemScriptViewModel();
+
+	/** Gets a niagara component for previewing the simulated System. */
+	UNiagaraComponent* GetPreviewComponent();
+
+	/** Gets the sequencer for this System for displaying the timeline. */
+	TSharedPtr<ISequencer> GetSequencer();
+
+	/** Gets the curve owner for the System represented by this view model, for use with the curve editor widget. */
+	FNiagaraCurveOwner& GetCurveOwner();
+
+	/** Get access to the underlying system*/
+	UNiagaraSystem& GetSystem() { return System; }
+
+	/** Gets whether or not this system is transient.  This will be true for the system view model in the emitter editor. */
+	bool GetSystemIsTransient() const;
+
+	/** Gets whether or not emitters can be added from the timeline. */
+	bool GetCanAddEmittersFromTimeline() const;
+
+	/** Adds a new emitter to the System from an emitter asset data. */
+	void AddEmitterFromAssetData(const FAssetData& AssetData);
+
+	/** Adds a new emitter to the System. */
+	void AddEmitter(UNiagaraEmitter& Emitter);
+
+	/** Duplicates the selected emitter in this System. */
+	void DuplicateEmitter(TSharedRef<FNiagaraEmitterHandleViewModel> EmitterHandleToDuplicate);
+
+	/** Deletes the selected emitter from the System. */
+	void DeleteEmitter(TSharedRef<FNiagaraEmitterHandleViewModel> EmitterHandleToDelete);
+
+	/** Deletes the emitters with the supplied ids from the system */
+	void DeleteEmitters(TSet<FGuid> EmitterHandleIdsToDelete);
+
+	/** Gets a multicast delegate which is called any time the array of emitter handle view models changes. */
+	FOnEmitterHandleViewModelsChanged& OnEmitterHandleViewModelsChanged();
+
+	/** Gets a delegate which is called any time the data in the curve owner is changed internally by this view model. */
+	FOnCurveOwnerChanged& OnCurveOwnerChanged();
+
+	/** Gets a multicast delegate which is called whenever the selected emitter handles changes. */
+	FOnSelectedEmitterHandlesChanged& OnSelectedEmitterHandlesChanged();
+	
+	/** Gets a multicast delegate which is called whenever we've received and handled a sequencer time update.*/
+	FOnPostSequencerTimeChange& OnPostSequencerTimeChanged();
+
+	/** Gets a multicast delegate which is called whenever the system has been compiled. */
+	FOnSystemCompiled& OnSystemCompiled();
+
+	//~ FGCObject interface
+	virtual void AddReferencedObjects(FReferenceCollector& Collector) override;
+
+	//~ FEditorUndoClient interface
+	virtual void PostUndo(bool bSuccess) override;
+	virtual void PostRedo(bool bSuccess) override { PostUndo(bSuccess); }
+
+	// ~ FTickableEditorObject
+	virtual void Tick(float DeltaTime) override;
+	virtual bool IsTickable() const override;
+	virtual TStatId GetStatId() const override;
+
+	void ResynchronizeAllHandles();
+
+	/** Resets the System instance to initial conditions. */
+	void ResetSystem();
+
+	/** Compiles the spawn and update scripts. */
+	void CompileSystem();
+
+	/* Get the latest status of this view-model's script compilation.*/
+	ENiagaraScriptCompileStatus GetLatestCompileStatus();
+
+	/** Gets the ids for the currently selected emitter handles. */
+	const TArray<FGuid>& GetSelectedEmitterHandleIds();
+
+	/** Sets the currently selected emitter handles by id. */
+	void SetSelectedEmitterHandlesById(TArray<FGuid> InSelectedEmitterHandleIds);
+
+	/** Sets the currently selected emitter handle by id. */
+	void SetSelectedEmitterHandleById(FGuid InSelectedEmitterHandleId);
+
+	/** Gets the currently selected emitter handles. */
+	void GetSelectedEmitterHandles(TArray<TSharedRef<FNiagaraEmitterHandleViewModel>>& OutSelectedEmitterHanldles);
+
+	/** Gets editor specific data which can be stored per system.  If this data hasn't been created the default version will be returned. */
+	const UNiagaraSystemEditorData& GetEditorData() const;
+
+	/** Gets editor specific data which is stored per system.  If this data hasn't been created then it will be created. */
+	UNiagaraSystemEditorData& GetOrCreateEditorData();
+	
+	/** Act as if the system has been fully destroyed although references might persist.*/
+	void Cleanup();
+
+	/** Reinitializes all System instances, and rebuilds emitter handle view models and tracks. */
+	void RefreshAll();
+
+	/** Called to notify the system view model that one of the data objects in the system was modified. */
+	void NotifyDataObjectChanged(UObject* ChangedObject);
+private:
+
+	/** Sets up the preview component and System instance. */
+	void SetupPreviewComponentAndInstance();
+
+	/** Rebuilds the emitter handle view models. */
+	void RefreshEmitterHandleViewModels();
+
+	/** Rebuilds the sequencer tracks. */
+	void RefreshSequencerTracks();
+
+	/** Gets the sequencer emitter track for the supplied emitter handle view model. */
+	UMovieSceneNiagaraEmitterTrack* GetTrackForHandleViewModel(TSharedRef<FNiagaraEmitterHandleViewModel> EmitterHandleViewModel);
+
+	/** Refreshes sequencer track data from it's emitter handle. */
+	void RefreshSequencerTrack(UMovieSceneNiagaraEmitterTrack* EmitterTrack);
+
+	/** Sets up the sequencer for this emitter. */
+	void SetupSequencer();
+
+	/** Gets the content for the add menu in sequencer. */
+	void GetSequencerAddMenuContent(FMenuBuilder& MenuBuilder, TSharedRef<ISequencer> Sequencer);
+
+	/** Kills all system instances using the system being displayed by this view model. */
+	void KillSystemInstances();
+
+	/** Reinitializes the System instance to initial conditions - rebuilds all data sets and resets data interfaces. */
+	void ReInitializeSystemInstances();
+
+	/** Resets and rebuilds the data in the curve owner. */
+	void ResetCurveData();
+
+	/** Updates the compiled versions of data interfaces when their sources change. */
+	void UpdateCompiledDataInterfaces(UNiagaraDataInterface* ChangedDataInterface);
+
+	/** Called whenever a property on the emitter handle changes. */
+	void EmitterHandlePropertyChanged(TSharedRef<FNiagaraEmitterHandleViewModel> EmitterHandleViewModel);
+
+	/** Called whenever a property on the emitter changes. */
+	void EmitterPropertyChanged(TSharedRef<FNiagaraEmitterHandleViewModel> EmitterHandleViewModel);
+
+	/** Called when a script is compiled */
+	void ScriptCompiled();
+
+	/** Handles when a curve in the curve owner is changed by the curve editor. */
+	void CurveChanged(FRichCurve* ChangedCurve, UObject* CurveOwnerObject);
+
+	/** Called whenever the data in the sequence is changed. */
+	void SequencerDataChanged(EMovieSceneDataChangeType DataChangeType);
+
+	/** Called whenever the global time in the sequencer is changed. */
+	void SequencerTimeChanged();
+
+	/** Called whenever the track selection in sequencer changes. */
+	void SequencerTrackSelectionChanged(TArray<UMovieSceneTrack*> SelectedTracks);
+
+	/** Called whenever the section selection in sequencer changes. */
+	void SequencerSectionSelectionChanged(TArray<UMovieSceneSection*> SelectedSections);
+
+	/** Updates the current emitter handle selection base on the sequencer selection. */
+	void UpdateEmitterHandleSelectionFromSequencer();
+
+	/** Updates the sequencer selection based on the current emitter handle selection. */
+	void UpdateSequencerFromEmitterHandleSelection();
+
+	/** Called when the system instance on the preview component changes. */
+	void PreviewComponentSystemInstanceChanged();
+
+	/** Called whenever the System instance is initialized. */
+	void SystemInstanceInitialized();
+
+private:
+	/** The System being viewed and edited by this view model. */
+	UNiagaraSystem& System;
+
+	/** The component used for previewing the System in a viewport. */
+	UNiagaraComponent* PreviewComponent;
+
+	/** The system instance currently simulating this system if available. */
+	FNiagaraSystemInstance* SystemInstance;
+
+	/** The view models for the emitter handles owned by the System. */
+	TArray<TSharedRef<FNiagaraEmitterHandleViewModel>> EmitterHandleViewModels;
+
+	/** The view model for the System script. */
+	TSharedRef<FNiagaraSystemScriptViewModel> SystemScriptViewModel;
+
+	/** A niagara sequence for displaying this System in the sequencer timeline. */
+	UNiagaraSequence *NiagaraSequence;
+
+	/** The sequencer instance viewing and editing the niagara sequence. */
+	TSharedPtr<ISequencer> Sequencer;
+
+	/** Flag which indicates we are setting the sequencer time directly in an internal operation. */
+	bool bSettingSequencerTimeDirectly;
+
+	/** The previous play status for sequencer timeline. */
+	EMovieScenePlayerStatus::Type PreviousSequencerStatus;
+
+	/** The previous time for the sequencer timeline. */
+	float PreviousSequencerTime;
+
+	/** Whether or not the user can remove emitters from the timeline. */
+	bool bCanRemoveEmittersFromTimeline;
+
+	/** Whether or not the user can rename emitters from the timeline. */
+	bool bCanRenameEmittersFromTimeline;
+
+	/** Whether or not the user can add emitters from the timeline. */
+	bool bCanAddEmittersFromTimeline;
+
+	/** Whether or not we use the system's execution state to drive when we reset the timeline*/
+	bool bUseSystemExecStateForTimelineReset;
+
+	/** A delegate which is used to generate the content for the add menu in sequencer. */
+	FOnGetAddMenuContent OnGetSequencerAddMenuContent;
+
+	/** A multicast delegate which is called any time the array of emitter handle view models changes. */
+	FOnEmitterHandleViewModelsChanged OnEmitterHandleViewModelsChangedDelegate;
+
+	/** A multicase delegate which is called when the contents of the curve owner is changed internally by this view model. */
+	FOnCurveOwnerChanged OnCurveOwnerChangedDelegate;
+
+	/** A multicast delegate which is called whenever the selected emitter changes. */
+	FOnSelectedEmitterHandlesChanged OnSelectedEmitterHandlesChangedDelegate;
+
+	/** A multicast delegate which is called whenever we've received and handled a sequencer time update.*/
+	FOnPostSequencerTimeChange OnPostSequencerTimeChangeDelegate;
+
+	/** A multicast delegate which is called whenever the system has been compiled. */
+	FOnSystemCompiled OnSystemCompiledDelegate;
+
+	/** A flag for preventing reentrancy when syncrhonizing sequencer data. */
+	bool bUpdatingFromSequencerDataChange;
+
+	/** A flag for preventing reentrancy when synchronizing system selection with sequencer selection */
+	bool bUpdatingSystemSelectionFromSequencer;
+
+	/** A flag for preventing reentrancy when synchronizing sequencer selection with system selection */
+	bool bUpdatingSequencerSelectionFromSystem;
+
+	/** A curve owner implementation for curves in a niagara System. */
+	FNiagaraCurveOwner CurveOwner;
+
+	/** The ids for the currently selected emitter handles. */
+	TArray<FGuid> SelectedEmitterHandleIds;
+
+	TNiagaraViewModelManager<UNiagaraSystem, FNiagaraSystemViewModel>::Handle RegisteredHandle;
+};
\ No newline at end of file
diff --git a/Engine/Plugins/FX/Niagara/Source/NiagaraEditor/Private/ViewModels/Stack/NiagaraParameterHandle.cpp b/Engine/Plugins/FX/Niagara/Source/NiagaraEditor/Private/ViewModels/Stack/NiagaraParameterHandle.cpp
new file mode 100644
--- /dev/null
+++ b/Engine/Plugins/FX/Niagara/Source/NiagaraEditor/Private/ViewModels/Stack/NiagaraParameterHandle.cpp
@@ -0,0 +1,123 @@
+// Copyright 1998-2017 Epic Games, Inc. All Rights Reserved.
+
+#include "NiagaraParameterHandle.h"
+#include "NiagaraNodeFunctionCall.h"
+#include "NiagaraScript.h"
+#include "NiagaraTypes.h"
+
+const FString FNiagaraParameterHandle::UserNamespace(TEXT("User"));
+const FString FNiagaraParameterHandle::EngineNamespace(TEXT("Engine"));
+const FString FNiagaraParameterHandle::SystemNamespace(TEXT("System"));
+const FString FNiagaraParameterHandle::EmitterNamespace(TEXT("Emitter"));
+const FString FNiagaraParameterHandle::ParticleAttributeNamespace(TEXT("Particles"));
+const FString FNiagaraParameterHandle::ModuleNamespace(TEXT("Module"));
+const FString FNiagaraParameterHandle::InitialPrefix(TEXT("Initial"));
+
+FNiagaraParameterHandle::FNiagaraParameterHandle() 
+{
+}
+
+FNiagaraParameterHandle::FNiagaraParameterHandle(const FString& InParameterHandleString)
+{
+	ParameterHandleString = InParameterHandleString;
+	int32 DotIndex;
+	if (ParameterHandleString.FindChar(TEXT('.'), DotIndex))
+	{
+		Name = ParameterHandleString.RightChop(DotIndex + 1);
+		Namespace = ParameterHandleString.Left(DotIndex);
+	}
+	else
+	{
+		Name = ParameterHandleString;
+	}
+}
+
+FNiagaraParameterHandle::FNiagaraParameterHandle(const FString& InNamespace, const FString& InName)
+{
+	Namespace = InNamespace;
+	Name = InName;
+	ParameterHandleString = Namespace + "." + Name;
+}
+
+bool FNiagaraParameterHandle::operator==(const FNiagaraParameterHandle& Other) const
+{
+	return ParameterHandleString == Other.ParameterHandleString;
+}
+
+FNiagaraParameterHandle FNiagaraParameterHandle::CreateAliasedModuleParameterHandle(const FNiagaraParameterHandle& ModuleParameterHandle, UNiagaraNodeFunctionCall* ModuleNode)
+{
+	return ModuleParameterHandle.IsModuleHandle()
+		? FNiagaraParameterHandle(ModuleNode->GetFunctionName(), ModuleParameterHandle.GetName())
+		: ModuleParameterHandle;
+}
+
+FNiagaraParameterHandle FNiagaraParameterHandle::CreateEngineParameterHandle(const FNiagaraVariable& SystemVariable)
+{
+	return FNiagaraParameterHandle(SystemVariable.GetName().ToString());
+}
+
+FNiagaraParameterHandle FNiagaraParameterHandle::CreateEmitterParameterHandle(const FNiagaraVariable& EmitterVariable)
+{
+
+	return FNiagaraParameterHandle(EmitterNamespace, EmitterVariable.GetName().ToString());
+}
+
+FNiagaraParameterHandle FNiagaraParameterHandle::CreateParticleAttributeParameterHandle(const FString& InName)
+{
+	return FNiagaraParameterHandle(ParticleAttributeNamespace, InName);
+}
+
+FNiagaraParameterHandle FNiagaraParameterHandle::CreateModuleParameterHandle(const FString& InName)
+{
+	return FNiagaraParameterHandle(ModuleNamespace, InName);
+}
+
+FNiagaraParameterHandle FNiagaraParameterHandle::CreateInitialParameterHandle(const FNiagaraParameterHandle& Handle)
+{
+	return FNiagaraParameterHandle(Handle.GetNamespace(), InitialPrefix + TEXT(".") + Handle.GetName());
+}
+
+bool FNiagaraParameterHandle::IsValid() const 
+{
+	return ParameterHandleString.IsEmpty() == false;
+}
+
+const FString& FNiagaraParameterHandle::GetParameterHandleString() const 
+{
+	return ParameterHandleString;
+}
+
+const FString& FNiagaraParameterHandle::GetName() const 
+{
+	return Name;
+}
+
+const FString& FNiagaraParameterHandle::GetNamespace() const
+{
+	return Namespace;
+}
+
+bool FNiagaraParameterHandle::IsEngineHandle() const
+{
+	return Namespace == EngineNamespace;
+}
+
+bool FNiagaraParameterHandle::IsSystemHandle() const
+{
+	return Namespace == SystemNamespace;
+}
+
+bool FNiagaraParameterHandle::IsEmitterHandle() const
+{
+	return Namespace == EmitterNamespace;
+}
+
+bool FNiagaraParameterHandle::IsParticleAttributeHandle() const
+{
+	return Namespace == ParticleAttributeNamespace;
+}
+
+bool FNiagaraParameterHandle::IsModuleHandle() const
+{
+	return Namespace == ModuleNamespace;
+}
\ No newline at end of file
diff --git a/Engine/Plugins/FX/Niagara/Source/NiagaraEditor/Private/ViewModels/Stack/NiagaraStackAddEventScriptItem.cpp b/Engine/Plugins/FX/Niagara/Source/NiagaraEditor/Private/ViewModels/Stack/NiagaraStackAddEventScriptItem.cpp
new file mode 100644
--- /dev/null
+++ b/Engine/Plugins/FX/Niagara/Source/NiagaraEditor/Private/ViewModels/Stack/NiagaraStackAddEventScriptItem.cpp
@@ -0,0 +1,74 @@
+// Copyright 1998-2016 Epic Games, Inc. All Rights Reserved.
+
+#include "NiagaraStackAddEventScriptItem.h"
+#include "NiagaraNodeOutput.h"
+#include "NiagaraScript.h"
+#include "NiagaraEmitter.h"
+#include "NiagaraNodeFunctionCall.h"
+#include "NiagaraStackEventScriptItemGroup.h"
+#include "Framework/Notifications/NotificationManager.h"
+#include "Widgets/Notifications/SNotificationList.h"
+#include "NiagaraScriptSource.h"
+#include "NiagaraGraph.h"
+#include "ViewModels/NiagaraEmitterViewModel.h"
+#include "ViewModels/NiagaraScriptGraphViewModel.h"
+#include "NiagaraStackGraphUtilities.h"
+
+#include "ScopedTransaction.h"
+
+#define LOCTEXT_NAMESPACE "NiagaraStackViewModel"
+
+UNiagaraStackAddEventScriptItem::UNiagaraStackAddEventScriptItem()
+{
+}
+
+FText UNiagaraStackAddEventScriptItem::GetDisplayName() const
+{
+	return FText();
+}
+
+ENiagaraScriptUsage UNiagaraStackAddEventScriptItem::GetOutputUsage() const
+{
+	return ENiagaraScriptUsage::ParticleEventScript;
+}
+
+UNiagaraNodeOutput* UNiagaraStackAddEventScriptItem::GetOrCreateOutputNode()
+{
+	UNiagaraEmitter* Emitter = GetEmitterViewModel()->GetEmitter();
+	UNiagaraScriptSource* Source = GetEmitterViewModel()->GetSharedScriptViewModel()->GetGraphViewModel()->GetScriptSource();
+	UNiagaraGraph* Graph = GetEmitterViewModel()->GetSharedScriptViewModel()->GetGraphViewModel()->GetGraph();
+
+	// The stack should not have been created if any of these are null, so bail out if it happens somehow rather than try to handle all of these cases.
+	checkf(Emitter != nullptr && Source != nullptr && Graph != nullptr, TEXT("Stack created for invalid emitter or graph."));
+
+	// Check that there isn't already an event output at the new index before making any changes so that we can early out and don't leave the change
+	// half completed.
+	if(ensureMsgf(Graph->FindOutputNode(ENiagaraScriptUsage::ParticleEventScript, Emitter->EventHandlerScriptProps.Num()) == nullptr,
+		TEXT("Invalid Stack Graph - While creating a new event handler an event script output node already existed with the new index.")))
+	{
+		Emitter->Modify();
+		FNiagaraEventScriptProperties EventScriptProperties;
+		EventScriptProperties.Script = NewObject<UNiagaraScript>(Emitter, MakeUniqueObjectName(Emitter, UNiagaraScript::StaticClass(), "EventScript"), EObjectFlags::RF_Transactional);
+		EventScriptProperties.Script->SetUsage(ENiagaraScriptUsage::ParticleEventScript);
+		EventScriptProperties.Script->SetSource(Source);
+
+		int32 Index = Emitter->EventHandlerScriptProps.Add(EventScriptProperties);
+		Emitter->EventHandlerScriptProps[Index].Script->SetUsageIndex(Index); // synchronize usage index with entry in array
+
+		UNiagaraNodeOutput* OutputNode = FNiagaraStackGraphUtilities::ResetGraphForOutput(*Graph, ENiagaraScriptUsage::ParticleEventScript, Index);
+
+		// Set the emitter here so that the internal state of the view model is updated.
+		// TODO: Move the logic for managing event handlers into the emitter view model or script view model.
+		GetEmitterViewModel()->GetSharedScriptViewModel()->SetScripts(Emitter);
+
+		return OutputNode;
+	}
+	return nullptr;
+}
+
+FText UNiagaraStackAddEventScriptItem::GetInsertTransactionText() const
+{
+	return LOCTEXT("InsertNewEventScript", "Insert new EventScript");
+}
+
+#undef LOCTEXT_NAMESPACE
\ No newline at end of file
diff --git a/Engine/Plugins/FX/Niagara/Source/NiagaraEditor/Private/ViewModels/Stack/NiagaraStackAddModuleItem.cpp b/Engine/Plugins/FX/Niagara/Source/NiagaraEditor/Private/ViewModels/Stack/NiagaraStackAddModuleItem.cpp
new file mode 100644
--- /dev/null
+++ b/Engine/Plugins/FX/Niagara/Source/NiagaraEditor/Private/ViewModels/Stack/NiagaraStackAddModuleItem.cpp
@@ -0,0 +1,124 @@
+// Copyright 1998-2016 Epic Games, Inc. All Rights Reserved.
+
+#include "NiagaraStackAddModuleItem.h"
+#include "NiagaraNodeOutput.h"
+#include "NiagaraScript.h"
+#include "NiagaraNodeFunctionCall.h"
+#include "NiagaraNodeAssignment.h"
+#include "ViewModels/NiagaraEmitterViewModel.h"
+#include "NiagaraStackEditorData.h"
+#include "NiagaraStackGraphUtilities.h"
+
+#include "ScopedTransaction.h"
+#include "EdGraph/EdGraph.h"
+
+#define LOCTEXT_NAMESPACE "NiagaraStackViewModel"
+
+void UNiagaraStackAddModuleItem::Initialize(TSharedRef<FNiagaraSystemViewModel> InSystemViewModel, TSharedRef<FNiagaraEmitterViewModel> InEmitterViewModel, UNiagaraStackEditorData& InStackEditorData)
+{
+	Super::Initialize(InSystemViewModel, InEmitterViewModel);
+	StackEditorData = &InStackEditorData;
+}
+
+FText UNiagaraStackAddModuleItem::GetDisplayName() const
+{
+	return FText();
+}
+
+void UNiagaraStackAddModuleItem::GetAvailableParameters(TArray<FNiagaraVariable>& OutAvailableParameterVariables) const
+{
+}
+
+void UNiagaraStackAddModuleItem::GetNewParameterAvailableTypes(TArray<FNiagaraTypeDefinition>& OutAvailableTypes) const
+{
+}
+
+TOptional<FString> UNiagaraStackAddModuleItem::GetNewParameterNamespace() const
+{
+	return TOptional<FString>();
+}
+
+void UNiagaraStackAddModuleItem::SetOnItemAdded(FOnItemAdded InOnItemAdded)
+{
+	ItemAddedDelegate = InOnItemAdded;
+}
+
+void ConnectModuleNode(UNiagaraNode& ModuleNode, UNiagaraNode& OutputNode)
+{
+	TArray<FNiagaraStackGraphUtilities::FStackNodeGroup> StackNodeGroups;
+	FNiagaraStackGraphUtilities::GetStackNodeGroups(OutputNode, StackNodeGroups);
+
+	FNiagaraStackGraphUtilities::FStackNodeGroup ModuleGroup;
+	ModuleGroup.StartNodes.Add(&ModuleNode);
+	ModuleGroup.EndNode = &ModuleNode;
+
+	FNiagaraStackGraphUtilities::FStackNodeGroup& OutputGroup = StackNodeGroups[StackNodeGroups.Num() - 1];
+	FNiagaraStackGraphUtilities::FStackNodeGroup& OutputGroupPrevious = StackNodeGroups[StackNodeGroups.Num() - 2];
+	FNiagaraStackGraphUtilities::ConnectStackNodeGroup(ModuleGroup, OutputGroupPrevious, OutputGroup);
+}
+
+void UNiagaraStackAddModuleItem::AddScriptModule(FAssetData ModuleScriptAsset)
+{
+	FScopedTransaction ScopedTransaction(GetInsertTransactionText());
+
+	UNiagaraNodeOutput* OutputNode = GetOrCreateOutputNode();
+	if (OutputNode != nullptr)
+	{
+		UEdGraph* Graph = OutputNode->GetGraph();
+		Graph->Modify();
+
+		FGraphNodeCreator<UNiagaraNodeFunctionCall> ModuleNodeCreator(*Graph);
+		UNiagaraNodeFunctionCall* NewModuleNode = ModuleNodeCreator.CreateNode();
+		NewModuleNode->FunctionScriptAssetObjectPath = ModuleScriptAsset.ObjectPath;
+		ModuleNodeCreator.Finalize();
+
+		ConnectModuleNode(*NewModuleNode, *OutputNode);
+		FNiagaraStackGraphUtilities::InitializeDataInterfaceInputs(GetSystemViewModel(), GetEmitterViewModel(), *StackEditorData, *NewModuleNode, *NewModuleNode);
+		FNiagaraStackGraphUtilities::RelayoutGraph(*OutputNode->GetGraph());
+
+		ItemAddedDelegate.ExecuteIfBound();
+	}
+}
+
+void UNiagaraStackAddModuleItem::AddParameterModule(FNiagaraVariable ParameterVariable, bool bRenamePending)
+{
+	FScopedTransaction ScopedTransaction(GetInsertTransactionText());
+
+	UNiagaraNodeOutput* OutputNode = GetOrCreateOutputNode();
+	if (OutputNode != nullptr)
+	{
+		UEdGraph* Graph = OutputNode->GetGraph();
+		Graph->Modify();
+
+		FGraphNodeCreator<UNiagaraNodeAssignment> AssignmentNodeCreator(*Graph);
+		UNiagaraNodeAssignment* NewAssignmentNode = AssignmentNodeCreator.CreateNode();
+		NewAssignmentNode->AssignmentTarget = ParameterVariable;
+		AssignmentNodeCreator.Finalize();
+
+		ConnectModuleNode(*NewAssignmentNode, *OutputNode);
+		FNiagaraStackGraphUtilities::InitializeDataInterfaceInputs(GetSystemViewModel(), GetEmitterViewModel(), *StackEditorData, *NewAssignmentNode, *NewAssignmentNode);
+		FNiagaraStackGraphUtilities::RelayoutGraph(*OutputNode->GetGraph());
+
+		TArray<const UEdGraphPin*> InputPins;
+		FNiagaraStackGraphUtilities::GetStackFunctionInputPins(*NewAssignmentNode, InputPins);
+		if (InputPins.Num() == 1)
+		{
+			FString FunctionInputEditorDataKey = FNiagaraStackGraphUtilities::GenerateStackFunctionInputEditorDataKey(*NewAssignmentNode, InputPins[0]->PinName);
+			StackEditorData->SetModuleInputIsPinned(FunctionInputEditorDataKey, true);
+			StackEditorData->SetStackEntryIsExpanded(FNiagaraStackGraphUtilities::GenerateStackModuleEditorDataKey(*NewAssignmentNode), false);
+			if (bRenamePending)
+			{
+				StackEditorData->SetModuleInputIsRenamePending(FunctionInputEditorDataKey, true);
+			}
+		}
+
+		ItemAddedDelegate.ExecuteIfBound();
+	}
+}
+
+FText UNiagaraStackAddModuleItem::GetInsertTransactionText() const
+{
+	return LOCTEXT("InsertNewModule", "Insert new module");
+}
+
+#undef LOCTEXT_NAMESPACE
\ No newline at end of file
diff --git a/Engine/Plugins/FX/Niagara/Source/NiagaraEditor/Private/ViewModels/Stack/NiagaraStackAddRendererItem.cpp b/Engine/Plugins/FX/Niagara/Source/NiagaraEditor/Private/ViewModels/Stack/NiagaraStackAddRendererItem.cpp
new file mode 100644
--- /dev/null
+++ b/Engine/Plugins/FX/Niagara/Source/NiagaraEditor/Private/ViewModels/Stack/NiagaraStackAddRendererItem.cpp
@@ -0,0 +1,62 @@
+// Copyright 1998-2016 Epic Games, Inc. All Rights Reserved.
+
+#include "NiagaraStackAddRendererItem.h"
+#include "NiagaraNodeOutput.h"
+#include "NiagaraScript.h"
+#include "NiagaraEmitter.h"
+#include "NiagaraRendererProperties.h"
+#include "NiagaraNodeFunctionCall.h"
+#include "NiagaraStackRendererItem.h"
+#include "NiagaraEmitterViewModel.h"
+#include "Framework/Notifications/NotificationManager.h"
+#include "Widgets/Notifications/SNotificationList.h"
+
+#include "ScopedTransaction.h"
+#include "EdGraph/EdGraph.h"
+
+#define LOCTEXT_NAMESPACE "NiagaraStackViewModel"
+
+FText UNiagaraStackAddRendererItem::GetDisplayName() const
+{
+	return FText();
+}
+
+void UNiagaraStackAddRendererItem::SetOnItemAdded(FOnItemAdded InOnItemAdded)
+{
+	ItemAddedDelegate = InOnItemAdded;
+}
+
+void UNiagaraStackAddRendererItem::AddRenderer(UClass* InRendererClass)
+{
+	FScopedTransaction ScopedTransaction(LOCTEXT("InsertNewRenderer", "Insert new renderer"));
+
+	UNiagaraEmitter* Emitter = GetEmitterViewModel()->GetEmitter();
+	Emitter->Modify();
+	UNiagaraRendererProperties* RendererProperties = NewObject<UNiagaraRendererProperties>(Emitter, InRendererClass, NAME_None, RF_Transactional);
+	Emitter->RendererProperties.Add(RendererProperties);
+
+	bool bVarsAdded = false;
+	TArray<FNiagaraVariable> MissingAttributes = UNiagaraStackRendererItem::GetMissingVariables(RendererProperties, Emitter);
+	for (int32 i = 0; i < MissingAttributes.Num(); i++)
+	{
+		if (UNiagaraStackRendererItem::AddMissingVariable(Emitter, MissingAttributes[i]))
+		{
+			bVarsAdded = true;
+		}
+	}
+
+	if (bVarsAdded)
+	{
+		FNotificationInfo Info(LOCTEXT("AddedVariables","One or more variables have been added to the Spawn script to support the added renderer."));
+		Info.ExpireDuration = 5.0f;
+		Info.bFireAndForget = true;
+		Info.Image = FCoreStyle::Get().GetBrush(TEXT("MessageLog.Info"));
+		FSlateNotificationManager::Get().AddNotification(Info);
+	}
+
+	ItemAddedDelegate.ExecuteIfBound();
+}
+
+
+
+#undef LOCTEXT_NAMESPACE
\ No newline at end of file
diff --git a/Engine/Plugins/FX/Niagara/Source/NiagaraEditor/Private/ViewModels/Stack/NiagaraStackAddScriptModuleItem.cpp b/Engine/Plugins/FX/Niagara/Source/NiagaraEditor/Private/ViewModels/Stack/NiagaraStackAddScriptModuleItem.cpp
new file mode 100644
--- /dev/null
+++ b/Engine/Plugins/FX/Niagara/Source/NiagaraEditor/Private/ViewModels/Stack/NiagaraStackAddScriptModuleItem.cpp
@@ -0,0 +1,84 @@
+// Copyright 1998-2016 Epic Games, Inc. All Rights Reserved.
+
+#include "NiagaraStackAddScriptModuleItem.h"
+#include "NiagaraNodeOutput.h"
+#include "NiagaraScript.h"
+#include "NiagaraNodeFunctionCall.h"
+#include "NiagaraNodeParameterMapBase.h"
+#include "NiagaraConstants.h"
+#include "Stack/NiagaraParameterHandle.h"
+
+#include "ScopedTransaction.h"
+
+#define LOCTEXT_NAMESPACE "NiagaraStackViewModel"
+
+void UNiagaraStackAddScriptModuleItem::Initialize(TSharedRef<FNiagaraSystemViewModel> InSystemViewModel, TSharedRef<FNiagaraEmitterViewModel> InEmitterViewModel, UNiagaraStackEditorData& InStackEditorData, UNiagaraNodeOutput& InOutputNode)
+{
+	Super::Initialize(InSystemViewModel, InEmitterViewModel, InStackEditorData);
+	OutputNode = &InOutputNode;
+}
+
+void UNiagaraStackAddScriptModuleItem::GetAvailableParameters(TArray<FNiagaraVariable>& OutAvailableParameterVariables) const
+{
+	TArray<FNiagaraParameterMapHistory> Histories = UNiagaraNodeParameterMapBase::GetParameterMaps(OutputNode->GetNiagaraGraph());
+
+	if (OutputNode->GetUsage() == ENiagaraScriptUsage::ParticleSpawnScript ||
+		OutputNode->GetUsage() == ENiagaraScriptUsage::ParticleSpawnScriptInterpolated ||
+		OutputNode->GetUsage() == ENiagaraScriptUsage::ParticleUpdateScript ||
+		OutputNode->GetUsage() == ENiagaraScriptUsage::ParticleEventScript)
+	{
+		OutAvailableParameterVariables.Append(FNiagaraConstants::GetCommonParticleAttributes());
+	}
+
+	for (FNiagaraParameterMapHistory& History : Histories)
+	{
+		for (FNiagaraVariable& Variable : History.Variables)
+		{
+			if (History.IsPrimaryDataSetOutput(Variable, OutputNode->GetUsage()))
+			{
+				OutAvailableParameterVariables.AddUnique(Variable);
+			}
+		}
+	}
+}
+
+void UNiagaraStackAddScriptModuleItem::GetNewParameterAvailableTypes(TArray<FNiagaraTypeDefinition>& OutAvailableTypes) const
+{
+	for (const FNiagaraTypeDefinition& RegisteredParameterType : FNiagaraTypeRegistry::GetRegisteredParameterTypes())
+	{
+		if (RegisteredParameterType != FNiagaraTypeDefinition::GetGenericNumericDef() && RegisteredParameterType != FNiagaraTypeDefinition::GetParameterMapDef())
+		{
+			OutAvailableTypes.Add(RegisteredParameterType);
+		}
+	}
+}
+
+TOptional<FString> UNiagaraStackAddScriptModuleItem::GetNewParameterNamespace() const
+{
+	switch (GetOutputUsage())
+	{
+	case ENiagaraScriptUsage::ParticleSpawnScript:
+	case ENiagaraScriptUsage::ParticleUpdateScript:
+		return FNiagaraParameterHandle::ParticleAttributeNamespace;
+	case ENiagaraScriptUsage::EmitterSpawnScript:
+	case ENiagaraScriptUsage::EmitterUpdateScript:
+		return FNiagaraParameterHandle::EmitterNamespace;
+	case ENiagaraScriptUsage::SystemSpawnScript:
+	case ENiagaraScriptUsage::SystemUpdateScript:
+		return FNiagaraParameterHandle::SystemNamespace;
+	default:
+		return TOptional<FString>();
+	}
+}
+
+ENiagaraScriptUsage UNiagaraStackAddScriptModuleItem::GetOutputUsage() const
+{
+	return OutputNode->GetUsage();
+}
+
+UNiagaraNodeOutput* UNiagaraStackAddScriptModuleItem::GetOrCreateOutputNode()
+{
+	return OutputNode.Get();
+}
+
+#undef LOCTEXT_NAMESPACE
\ No newline at end of file
diff --git a/Engine/Plugins/FX/Niagara/Source/NiagaraEditor/Private/ViewModels/Stack/NiagaraStackEmitterSpawnScriptItemGroup.cpp b/Engine/Plugins/FX/Niagara/Source/NiagaraEditor/Private/ViewModels/Stack/NiagaraStackEmitterSpawnScriptItemGroup.cpp
new file mode 100644
--- /dev/null
+++ b/Engine/Plugins/FX/Niagara/Source/NiagaraEditor/Private/ViewModels/Stack/NiagaraStackEmitterSpawnScriptItemGroup.cpp
@@ -0,0 +1,71 @@
+// Copyright 1998-2016 Epic Games, Inc. All Rights Reserved.
+
+#include "NiagaraStackEmitterSpawnScriptItemGroup.h"
+#include "NiagaraStackObject.h"
+#include "NiagaraStackSpacer.h"
+#include "NiagaraStackItemExpander.h"
+#include "ViewModels/NiagaraEmitterViewModel.h"
+#include "NiagaraEmitter.h"
+#include "NiagaraStackEditorData.h"
+
+#define LOCTEXT_NAMESPACE "UNiagaraStackScriptItemGroup"
+
+FText UNiagaraStackEmitterPropertiesItem::GetDisplayName() const
+{
+	return LOCTEXT("EmitterPropertiesDisplayName", "Emitter Properties");
+}
+
+void UNiagaraStackEmitterPropertiesItem::RefreshChildrenInternal(const TArray<UNiagaraStackEntry*>& CurrentChildren, TArray<UNiagaraStackEntry*>& NewChildren)
+{
+	if (EmitterObject == nullptr)
+	{
+		EmitterObject = NewObject<UNiagaraStackObject>(this);
+		EmitterObject->Initialize(GetSystemViewModel(), GetEmitterViewModel(), GetEmitterViewModel()->GetEmitter());
+	}
+
+	if (EmitterExpander == nullptr)
+	{
+		EmitterExpander = NewObject<UNiagaraStackItemExpander>(this);
+		EmitterExpander->Initialize(GetSystemViewModel(), GetEmitterViewModel(), GetStackEditorData(), TEXT("Emitter"), false);
+		EmitterExpander->SetOnExpnadedChanged(UNiagaraStackItemExpander::FOnExpandedChanged::CreateUObject(this, &UNiagaraStackEmitterPropertiesItem::EmitterExpandedChanged));
+	}
+
+	if (GetStackEditorData().GetStackEntryIsExpanded(TEXT("Emitter"), false))
+	{
+		NewChildren.Add(EmitterObject);
+	}
+
+	NewChildren.Add(EmitterExpander);
+}
+
+void UNiagaraStackEmitterPropertiesItem::EmitterExpandedChanged()
+{
+	RefreshChildren();
+}
+
+UNiagaraStackEmitterSpawnScriptItemGroup::UNiagaraStackEmitterSpawnScriptItemGroup()
+	: PropertiesItem(nullptr)
+{
+}
+
+void UNiagaraStackEmitterSpawnScriptItemGroup::RefreshChildrenInternal(const TArray<UNiagaraStackEntry*>& CurrentChildren, TArray<UNiagaraStackEntry*>& NewChildren)
+{
+	if (PropertiesItem == nullptr)
+	{
+		PropertiesItem = NewObject<UNiagaraStackEmitterPropertiesItem>(this);
+		PropertiesItem->Initialize(GetSystemViewModel(), GetEmitterViewModel(), GetStackEditorData());
+	}
+
+	if (PropertiesSpacer == nullptr)
+	{
+		PropertiesSpacer = NewObject<UNiagaraStackSpacer>(this);
+		PropertiesSpacer->Initialize(GetSystemViewModel(), GetEmitterViewModel(), "EmitterProperties");
+	}
+
+	NewChildren.Add(PropertiesItem);
+	NewChildren.Add(PropertiesSpacer);
+
+	Super::RefreshChildrenInternal(CurrentChildren, NewChildren);
+}
+
+#undef LOCTEXT_NAMESPACE
\ No newline at end of file
diff --git a/Engine/Plugins/FX/Niagara/Source/NiagaraEditor/Private/ViewModels/Stack/NiagaraStackEntry.cpp b/Engine/Plugins/FX/Niagara/Source/NiagaraEditor/Private/ViewModels/Stack/NiagaraStackEntry.cpp
new file mode 100644
--- /dev/null
+++ b/Engine/Plugins/FX/Niagara/Source/NiagaraEditor/Private/ViewModels/Stack/NiagaraStackEntry.cpp
@@ -0,0 +1,166 @@
+// Copyright 1998-2016 Epic Games, Inc. All Rights Reserved.
+
+#include "NiagaraStackEntry.h"
+#include "NiagaraStackErrorItem.h"
+
+UNiagaraStackEntry::UNiagaraStackEntry()
+{
+}
+
+void UNiagaraStackEntry::Initialize(TSharedRef<FNiagaraSystemViewModel> InSystemViewModel, TSharedRef<FNiagaraEmitterViewModel> InEmitterViewModel)
+{
+	SystemViewModel = InSystemViewModel;
+	EmitterViewModel = InEmitterViewModel;
+	bIsExpanded = IsExpandedByDefault();
+}
+
+FText UNiagaraStackEntry::GetDisplayName() const
+{
+	ensureMsgf(false, TEXT("GetDisplayName not implemented"));
+	return FText::FromName(NAME_None);
+}
+
+FText UNiagaraStackEntry::GetTooltipText() const
+{
+	return FText();
+}
+
+FName UNiagaraStackEntry::GetTextStyleName() const
+{
+	return "NiagaraEditor.Stack.DefaultText";
+}
+
+bool UNiagaraStackEntry::GetCanExpand() const
+{
+	return false;
+}
+
+bool UNiagaraStackEntry::IsExpandedByDefault() const
+{
+	return true;
+}
+
+bool UNiagaraStackEntry::GetIsExpanded() const
+{
+	return bIsExpanded;
+}
+
+void UNiagaraStackEntry::SetIsExpanded(bool bInExpanded)
+{
+	bIsExpanded = bInExpanded;
+}
+
+bool UNiagaraStackEntry::GetShouldShowInStack() const
+{
+	return true;
+}
+
+void UNiagaraStackEntry::GetChildren(TArray<UNiagaraStackEntry*>& OutChildren)
+{
+	OutChildren.Append(ErrorChildren);
+	OutChildren.Append(Children);
+}
+
+void UNiagaraStackEntry::RefreshErrors()
+{
+	ErrorChildren.Empty();
+	if (GetErrorCount() != 0)
+	{
+		for (int32 i = 0; i < GetErrorCount(); i++)
+		{
+			UNiagaraStackErrorItem* ErrorObject = NewObject<UNiagaraStackErrorItem>(this);
+			ErrorObject->Initialize(GetSystemViewModel(), GetEmitterViewModel(), this, i);
+			ErrorChildren.Add(ErrorObject);
+		}
+	}
+}
+
+TSharedRef<FNiagaraSystemViewModel> UNiagaraStackEntry::GetSystemViewModel() const
+{
+	TSharedPtr<FNiagaraSystemViewModel> PinnedSystemViewModel = SystemViewModel.Pin();
+	checkf(PinnedSystemViewModel.IsValid(), TEXT("Base stack entry not initialized or system view model was already deleted."));
+	return PinnedSystemViewModel.ToSharedRef();
+}
+
+TSharedRef<FNiagaraEmitterViewModel> UNiagaraStackEntry::GetEmitterViewModel() const
+{
+	TSharedPtr<FNiagaraEmitterViewModel> PinnedEmitterViewModel = EmitterViewModel.Pin();
+	checkf(PinnedEmitterViewModel.IsValid(), TEXT("Base stack entry not initialized or emitter view model was already deleted."));
+	return PinnedEmitterViewModel.ToSharedRef();
+}
+
+UNiagaraStackEntry::FOnStructureChanged& UNiagaraStackEntry::OnStructureChanged()
+{
+	return StructureChangedDelegate;
+}
+
+UNiagaraStackEntry::FOnDataObjectModified& UNiagaraStackEntry::OnDataObjectModified()
+{
+	return DataObjectModifiedDelegate;
+}
+
+FName UNiagaraStackEntry::GetGroupBackgroundName() const
+{
+	return "NiagaraEditor.Stack.Group.BackgroundColor";
+}
+
+FName UNiagaraStackEntry::GetGroupForegroundName() const
+{
+	return "NiagaraEditor.Stack.ForegroundColor";
+}
+
+FName UNiagaraStackEntry::GetItemBackgroundName() const
+{
+	return "NiagaraEditor.Stack.Item.BackgroundColor";
+}
+
+FName UNiagaraStackEntry::GetItemForegroundName() const
+{
+	return "NiagaraEditor.Stack.ForegroundColor";
+}
+
+int32 UNiagaraStackEntry::GetItemIndentLevel() const
+{
+	return 0;
+}
+
+void UNiagaraStackEntry::RefreshChildren()
+{
+	checkf(SystemViewModel.IsValid() && EmitterViewModel.IsValid(), TEXT("Base stack entry not initialized."));
+
+	for (UNiagaraStackEntry* Child : Children)
+	{
+		Child->OnStructureChanged().RemoveAll(this);
+		Child->OnDataObjectModified().RemoveAll(this);
+	}
+
+	TArray<UNiagaraStackEntry*> NewChildren;
+	RefreshChildrenInternal(Children, NewChildren);
+	Children.Empty();
+	Children.Append(NewChildren);
+
+	for (UNiagaraStackEntry* Child : Children)
+	{
+		Child->RefreshChildren();
+		Child->OnStructureChanged().AddUObject(this, &UNiagaraStackEntry::ChildStructureChanged);
+		Child->OnDataObjectModified().AddUObject(this, &UNiagaraStackEntry::ChildDataObjectModified);
+	}
+
+	RefreshErrors();
+
+	StructureChangedDelegate.Broadcast();
+}
+
+void UNiagaraStackEntry::RefreshChildrenInternal(const TArray<UNiagaraStackEntry*>& CurrentChildren, TArray<UNiagaraStackEntry*>& NewChildren)
+{
+}
+
+void UNiagaraStackEntry::ChildStructureChanged()
+{
+	StructureChangedDelegate.Broadcast();
+}
+
+void UNiagaraStackEntry::ChildDataObjectModified(UObject* ChangedObject)
+{
+	DataObjectModifiedDelegate.Broadcast(ChangedObject);
+}
\ No newline at end of file
diff --git a/Engine/Plugins/FX/Niagara/Source/NiagaraEditor/Private/ViewModels/Stack/NiagaraStackErrorItem.cpp b/Engine/Plugins/FX/Niagara/Source/NiagaraEditor/Private/ViewModels/Stack/NiagaraStackErrorItem.cpp
new file mode 100644
--- /dev/null
+++ b/Engine/Plugins/FX/Niagara/Source/NiagaraEditor/Private/ViewModels/Stack/NiagaraStackErrorItem.cpp
@@ -0,0 +1,70 @@
+// Copyright 1998-2016 Epic Games, Inc. All Rights Reserved.
+
+#include "NiagaraStackErrorItem.h"
+
+UNiagaraStackErrorItem::UNiagaraStackErrorItem() : ErrorSource(nullptr), ErrorIdx(INDEX_NONE)
+{
+}
+
+void UNiagaraStackErrorItem::Initialize(TSharedRef<FNiagaraSystemViewModel> InSystemViewModel, TSharedRef<FNiagaraEmitterViewModel> InEmitterViewModel, UNiagaraStackEntry* InErrorSource, int32 InErrorIndex)
+{
+	checkf(ErrorSource == nullptr, TEXT("Can only initialize once."));
+	Super::Initialize(InSystemViewModel, InEmitterViewModel);
+	ErrorSource = InErrorSource;
+	ErrorIdx = InErrorIndex;
+}
+
+int32 UNiagaraStackErrorItem::GetItemIndentLevel() const
+{
+	return ItemIndentLevel;
+}
+
+void UNiagaraStackErrorItem::SetItemIndentLevel(int32 InItemIndentLevel)
+{
+	ItemIndentLevel = InItemIndentLevel;
+}
+
+FText UNiagaraStackErrorItem::ErrorText() const
+{
+	FText SummaryText = ErrorSource->GetErrorSummaryText(ErrorIdx);
+	if (SummaryText.IsEmpty())
+	{
+		return ErrorSource->GetErrorText(ErrorIdx);
+	}
+	else
+	{
+		return SummaryText;
+	}
+}
+
+FText UNiagaraStackErrorItem::ErrorTextTooltip() const
+{
+	return ErrorSource->GetErrorText(ErrorIdx);
+}
+
+FReply UNiagaraStackErrorItem::OnTryFixError()
+{
+	if (ErrorSource->TryFixError(ErrorIdx))
+	{
+		return FReply::Handled();
+	}
+	return FReply::Handled();
+}
+
+EVisibility UNiagaraStackErrorItem::CanFixVisibility() const
+{
+	if (ErrorSource->GetErrorFixable(ErrorIdx))
+	{
+		return EVisibility::Visible;
+	}
+	else
+	{
+		return EVisibility::Hidden;
+	}
+}
+
+
+FName UNiagaraStackErrorItem::GetItemBackgroundName() const 
+{
+	return "NiagaraEditor.Stack.Item.ErrorBackgroundColor";
+}
\ No newline at end of file
diff --git a/Engine/Plugins/FX/Niagara/Source/NiagaraEditor/Private/ViewModels/Stack/NiagaraStackEventHandlerGroup.cpp b/Engine/Plugins/FX/Niagara/Source/NiagaraEditor/Private/ViewModels/Stack/NiagaraStackEventHandlerGroup.cpp
new file mode 100644
--- /dev/null
+++ b/Engine/Plugins/FX/Niagara/Source/NiagaraEditor/Private/ViewModels/Stack/NiagaraStackEventHandlerGroup.cpp
@@ -0,0 +1,53 @@
+// Copyright 1998-2016 Epic Games, Inc. All Rights Reserved.
+
+#include "NiagaraStackEventHandlerGroup.h"
+#include "NiagaraStackRendererItem.h"
+#include "NiagaraStackEventScriptItemGroup.h"
+#include "NiagaraEmitterViewModel.h"
+#include "NiagaraEmitter.h"
+#include "NiagaraStackAddEventScriptItem.h"
+#include "NiagaraGraph.h"
+#include "NiagaraScriptSource.h"
+
+#define LOCTEXT_NAMESPACE "NiagaraStackEventHandlerGroup"
+
+
+
+FText UNiagaraStackEventHandlerGroup::GetDisplayName() const
+{
+	return DisplayName;
+}
+
+void UNiagaraStackEventHandlerGroup::SetDisplayName(FText InDisplayName)
+{
+	DisplayName = InDisplayName;
+}
+
+void UNiagaraStackEventHandlerGroup::RefreshChildrenInternal(const TArray<UNiagaraStackEntry*>& CurrentChildren, TArray<UNiagaraStackEntry*>& NewChildren)
+{
+	if (AddScriptItem == nullptr)
+	{
+		AddScriptItem = NewObject<UNiagaraStackAddEventScriptItem>(this);
+		AddScriptItem->Initialize(GetSystemViewModel(), GetEmitterViewModel(), GetStackEditorData());
+		AddScriptItem->SetOnItemAdded(UNiagaraStackAddEventScriptItem::FOnItemAdded::CreateUObject(this, &UNiagaraStackEventHandlerGroup::ChildModifiedGroupItems));
+	}
+
+	NewChildren.Add(AddScriptItem);
+}
+
+void UNiagaraStackEventHandlerGroup::SetOnItemAdded(FOnItemAdded InOnItemAdded)
+{
+	ItemAddedDelegate = InOnItemAdded;
+}
+
+bool UNiagaraStackEventHandlerGroup::GetShouldShowInStack() const
+{
+	return true;
+}
+
+void UNiagaraStackEventHandlerGroup::ChildModifiedGroupItems()
+{
+	ItemAddedDelegate.ExecuteIfBound();
+}
+
+#undef LOCTEXT_NAMESPACE
\ No newline at end of file
diff --git a/Engine/Plugins/FX/Niagara/Source/NiagaraEditor/Private/ViewModels/Stack/NiagaraStackEventScriptItemGroup.cpp b/Engine/Plugins/FX/Niagara/Source/NiagaraEditor/Private/ViewModels/Stack/NiagaraStackEventScriptItemGroup.cpp
new file mode 100644
--- /dev/null
+++ b/Engine/Plugins/FX/Niagara/Source/NiagaraEditor/Private/ViewModels/Stack/NiagaraStackEventScriptItemGroup.cpp
@@ -0,0 +1,142 @@
+// Copyright 1998-2016 Epic Games, Inc. All Rights Reserved.
+
+#include "NiagaraStackEventScriptItemGroup.h"
+#include "NiagaraStackModuleItem.h"
+#include "NiagaraStackAddModuleItem.h"
+#include "NiagaraScriptViewModel.h"
+#include "NiagaraScriptGraphViewModel.h"
+#include "NiagaraEmitterViewModel.h"
+#include "NiagaraGraph.h"
+#include "NiagaraNodeOutput.h"
+#include "NiagaraNodeFunctionCall.h"
+#include "EdGraphSchema_Niagara.h"
+#include "NiagaraStackErrorItem.h"
+#include "NiagaraStackStruct.h"
+#include "NiagaraScriptSource.h"
+
+#include "Internationalization.h"
+#include "ScopedTransaction.h"
+#include "NiagaraSystemViewModel.h"
+#include "Customizations/NiagaraEventScriptPropertiesCustomization.h"
+
+#define LOCTEXT_NAMESPACE "UNiagaraStackEventScriptItemGroup"
+
+FText UNiagaraStackEventScriptItemGroup::GetDisplayName() const 
+{
+	FText EventName = FText::FromName(GetEmitterViewModel()->GetEmitter()->EventHandlerScriptProps[GetScriptOccurrence()].SourceEventName);
+	return (FText::Format(LOCTEXT("FormatEventScriptDisplayName", "Event Handler {0} - Source: {1}"), GetScriptOccurrence(), EventName));
+}
+void UNiagaraStackEventScriptItemGroup::RefreshChildrenInternal(const TArray<UNiagaraStackEntry*>& CurrentChildren, TArray<UNiagaraStackEntry*>& NewChildren)
+{
+	UNiagaraEmitter* Emitter = GetEmitterViewModel()->GetEmitter();
+	if (GetScriptOccurrence() < Emitter->EventHandlerScriptProps.Num())
+	{
+		uint8* EventStructMemory = (uint8*)&(Emitter->EventHandlerScriptProps[GetScriptOccurrence()]);
+
+		UNiagaraStackEntry* FoundItem = nullptr;
+		for (UNiagaraStackEntry* CurrentChild : CurrentChildren)
+		{
+			UNiagaraStackStruct* CurrentStructNode = Cast<UNiagaraStackStruct>(CurrentChild);
+			if (CurrentStructNode != nullptr &&
+				CurrentStructNode->GetOwningObject() == Emitter &&
+				CurrentStructNode->GetStructOnScope()->GetStructMemory() == EventStructMemory)
+			{
+				FoundItem = CurrentStructNode;
+				break;
+			}
+		}
+
+		if (FoundItem == nullptr)
+		{
+			UNiagaraStackStruct* StructNode = NewObject<UNiagaraStackStruct>(this);
+			StructNode->Initialize(GetSystemViewModel(), GetEmitterViewModel(), GetEmitterViewModel()->GetEmitter(), FNiagaraEventScriptProperties::StaticStruct(), EventStructMemory);
+			UNiagaraSystem* System = &GetSystemViewModel()->GetSystem();
+			StructNode->SetDetailsCustomization(FOnGetDetailCustomizationInstance::CreateStatic(&FNiagaraEventScriptPropertiesCustomization::MakeInstance, TWeakObjectPtr<UNiagaraSystem>(System), TWeakObjectPtr<UNiagaraEmitter>(Emitter)));
+			FoundItem = StructNode;
+		}
+
+		NewChildren.Add(FoundItem);
+	}
+
+	Super::RefreshChildrenInternal(CurrentChildren, NewChildren);
+}
+
+bool UNiagaraStackEventScriptItemGroup::CanDelete() const
+{
+	return true;
+}
+
+bool UNiagaraStackEventScriptItemGroup::Delete()
+{
+	UNiagaraEmitter* Emitter = GetEmitterViewModel()->GetEmitter();
+	UNiagaraScriptSource* Source = Cast<UNiagaraScriptSource>(Emitter->GraphSource);
+
+	if (!Source || !Source->NodeGraph)
+	{
+		return false;
+	}
+
+	FScopedTransaction Transaction(FText::Format(LOCTEXT("DeleteEventHandler", "Deleted {0}"), GetDisplayName()));
+	Emitter->Modify();
+	Source->NodeGraph->Modify();
+	TArray<UNiagaraNode*> EventIndexNodes;
+	Source->NodeGraph->BuildTraversal(EventIndexNodes, GetScriptUsage(), GetScriptOccurrence());
+	for (UNiagaraNode* Node : EventIndexNodes)
+	{
+		Node->Modify();
+	}
+
+	TArray<UNiagaraNodeOutput*> EventOutputNodes;
+	Source->NodeGraph->FindOutputNodes(GetScriptUsage(), EventOutputNodes);
+	for (UNiagaraNodeOutput* OutputNode : EventOutputNodes)
+	{
+		OutputNode->Modify();
+	}
+
+	int32 OldScriptOccurrence = GetScriptOccurrence();
+	
+	// First, remove the event handler script properties object.
+	Emitter->EventHandlerScriptProps.RemoveAt(OldScriptOccurrence);
+
+	// Now move all the usage indices down the list.
+	for (int32 i = OldScriptOccurrence; i < Emitter->EventHandlerScriptProps.Num(); i++)
+	{
+		Emitter->EventHandlerScriptProps[i].Script->Modify();
+		Emitter->EventHandlerScriptProps[i].Script->UsageIndex -= 1;
+	}
+	
+	// Now remove all graph nodes associated with the event script index.
+	for (UNiagaraNode* Node : EventIndexNodes)
+	{
+		Node->DestroyNode();
+	}
+
+	// Now move all the other output nodes down in their index to match the event script properties array.
+	for (UNiagaraNodeOutput* OutputNode : EventOutputNodes)
+	{
+		if (EventIndexNodes.Contains(OutputNode))
+		{
+			continue;
+		}
+		else if (OutputNode->GetUsageIndex() > OldScriptOccurrence)
+		{
+			OutputNode->SetUsageIndex(OutputNode->GetUsageIndex() - 1);
+		}
+	}
+
+	// Set the emitter here to that the internal state of the view model is updated.
+	// TODO: Move the logic for managing event handlers into the emitter view model or script view model.
+	ScriptViewModel->SetScripts(Emitter);
+	
+	OnModifiedEventHandlersDelegate.ExecuteIfBound();
+
+	return true;
+}
+
+void UNiagaraStackEventScriptItemGroup::SetOnModifiedEventHandlers(FOnModifiedEventHandlers OnModifiedEventHandlers)
+{
+	OnModifiedEventHandlersDelegate = OnModifiedEventHandlers;
+}
+
+
+#undef LOCTEXT_NAMESPACE
\ No newline at end of file
