Copying //UE4/Orion-Staging to //UE4/Main (Source: //Orion/Dev-General @ 3606378)

diff --git a/Engine/Plugins/FX/Niagara/Source/NiagaraEditor/Private/ViewModels/Stack/NiagaraStackFunctionInput.cpp b/Engine/Plugins/FX/Niagara/Source/NiagaraEditor/Private/ViewModels/Stack/NiagaraStackFunctionInput.cpp
new file mode 100644
--- /dev/null
+++ b/Engine/Plugins/FX/Niagara/Source/NiagaraEditor/Private/ViewModels/Stack/NiagaraStackFunctionInput.cpp
@@ -0,0 +1,1094 @@
+// Copyright 1998-2016 Epic Games, Inc. All Rights Reserved.
+
+#include "NiagaraStackFunctionInput.h"
+#include "NiagaraStackFunctionInputCollection.h"
+#include "NiagaraStackObject.h"
+#include "NiagaraNodeFunctionCall.h"
+#include "NiagaraNodeAssignment.h"
+#include "NiagaraNodeParameterMapGet.h"
+#include "NiagaraNodeParameterMapSet.h"
+#include "NiagaraNodeOutput.h"
+#include "NiagaraScript.h"
+#include "NiagaraScriptSource.h"
+#include "NiagaraGraph.h"
+#include "EdGraphSchema_Niagara.h"
+#include "NiagaraEditorUtilities.h"
+#include "NiagaraSystemViewModel.h"
+#include "NiagaraSystemScriptViewModel.h"
+#include "NiagaraEmitterViewModel.h"
+#include "NiagaraScriptGraphViewModel.h"
+#include "NiagaraStackEditorData.h"
+#include "NiagaraComponent.h"
+#include "NiagaraScriptSource.h"
+#include "NiagaraConstants.h"
+#include "NiagaraStackGraphUtilities.h"
+
+#include "ScopedTransaction.h"
+#include "Editor.h"
+#include "UObject/StructOnScope.h"
+#include "Modules/ModuleManager.h"
+#include "AssetRegistryModule.h"
+#include "ARFilter.h"
+#include "EdGraph/EdGraphPin.h"
+
+#define LOCTEXT_NAMESPACE "NiagaraStackViewModel"
+
+UNiagaraStackFunctionInput::UNiagaraStackFunctionInput()
+	: OwningModuleNode(nullptr)
+	, OwningFunctionCallNode(nullptr)
+	, bCanBePinned(true)
+	, bUpdatingGraphDirectly(false)
+{
+}
+
+void UNiagaraStackFunctionInput::BeginDestroy()
+{
+	Super::BeginDestroy();
+	if (OwningFunctionCallNode.IsValid())
+	{
+		OwningFunctionCallNode->GetGraph()->RemoveOnGraphChangedHandler(GraphChangedHandle);
+	}
+}
+
+void UNiagaraStackFunctionInput::AddReferencedObjects(UObject* InThis, FReferenceCollector& Collector)
+{
+	UNiagaraStackFunctionInput* This = Cast<UNiagaraStackFunctionInput>(InThis);
+	if (This != nullptr)
+	{
+		This->AddReferencedObjects(Collector);
+	}
+	Super::AddReferencedObjects(InThis, Collector);
+}
+
+void UNiagaraStackFunctionInput::AddReferencedObjects(FReferenceCollector& Collector)
+{
+	if (InputValues.DataObjects.IsValid() && InputValues.DataObjects.GetDefaultValueOwner() == FDataValues::EDefaultValueOwner::LocallyOwned)
+	{
+		Collector.AddReferencedObject(InputValues.DataObjects.GetDefaultValueObjectRef(), this);
+	}
+}
+
+int32 UNiagaraStackFunctionInput::GetItemIndentLevel() const
+{
+	return ItemIndentLevel;
+}
+
+void UNiagaraStackFunctionInput::SetItemIndentLevel(int32 InItemIndentLevel)
+{
+	ItemIndentLevel = InItemIndentLevel;
+}
+
+// Traverses the path between the owning module node and the function call node this input belongs too collecting up the input handles between them.
+void GenerateInputParameterHandlePath(UNiagaraNodeFunctionCall& ModuleNode, UNiagaraNodeFunctionCall& FunctionCallNode, TArray<FNiagaraParameterHandle>& OutHandlePath)
+{
+	UNiagaraNodeFunctionCall* CurrentFunctionCallNode = &FunctionCallNode;
+	while (CurrentFunctionCallNode != &ModuleNode)
+	{
+		TArray<UEdGraphPin*> FunctionOutputPins;
+		CurrentFunctionCallNode->GetOutputPins(FunctionOutputPins);
+		if (ensureMsgf(FunctionOutputPins.Num() == 1 && FunctionOutputPins[0]->LinkedTo.Num() == 1 && FunctionOutputPins[0]->LinkedTo[0]->GetOwningNode()->IsA<UNiagaraNodeParameterMapSet>(),
+			TEXT("Invalid Stack Graph - Dynamic Input Function call didn't have a valid connected output.")))
+		{
+			FNiagaraParameterHandle AliasedHandle(FunctionOutputPins[0]->LinkedTo[0]->PinName);
+			OutHandlePath.Add(FNiagaraParameterHandle::CreateModuleParameterHandle(AliasedHandle.GetName()));
+			UNiagaraNodeParameterMapSet* NextOverrideNode = CastChecked<UNiagaraNodeParameterMapSet>(FunctionOutputPins[0]->LinkedTo[0]->GetOwningNode());
+			UEdGraphPin* NextOverrideNodeOutputPin = FNiagaraStackGraphUtilities::GetParameterMapOutputPin(*NextOverrideNode);
+			
+			TArray<UNiagaraNodeFunctionCall*> NextFunctionCallNodes;
+			for (UEdGraphPin* NextOverrideNodeOutputPinLinkedPin : NextOverrideNodeOutputPin->LinkedTo)
+			{
+				UNiagaraNodeFunctionCall* NextFunctionCallNode = Cast<UNiagaraNodeFunctionCall>(NextOverrideNodeOutputPinLinkedPin->GetOwningNode());
+				if (NextFunctionCallNode != nullptr)
+				{
+					NextFunctionCallNodes.Add(NextFunctionCallNode);
+				}
+			}
+
+			if (ensureMsgf(NextFunctionCallNodes.Num() == 1, TEXT("Invalid Stack Graph - Override node not corrected to a single function call node.")))
+			{
+				CurrentFunctionCallNode = NextFunctionCallNodes[0];
+			}
+			else
+			{
+				OutHandlePath.Empty();
+				return;
+			}
+		}
+		else
+		{
+			OutHandlePath.Empty();
+			return;
+		}
+	}
+}
+
+void UNiagaraStackFunctionInput::Initialize(
+	TSharedRef<FNiagaraSystemViewModel> InSystemViewModel,
+	TSharedRef<FNiagaraEmitterViewModel> InEmitterViewModel,
+	UNiagaraStackEditorData& InStackEditorData,
+	UNiagaraNodeFunctionCall& InModuleNode,
+	UNiagaraNodeFunctionCall& InInputFunctionCallNode,
+	FString InInputParameterHandle,
+	FNiagaraTypeDefinition InInputType)
+{
+	checkf(OwningModuleNode.IsValid() == false && OwningFunctionCallNode.IsValid() == false, TEXT("Can only initialize once."));
+	Super::Initialize(InSystemViewModel, InEmitterViewModel);
+	StackEditorData = &InStackEditorData;
+	OwningModuleNode = &InModuleNode;
+	OwningFunctionCallNode = &InInputFunctionCallNode;
+	OwningAssignmentNode = Cast<UNiagaraNodeAssignment>(OwningFunctionCallNode.Get());
+
+	GraphChangedHandle = OwningFunctionCallNode->GetGraph()->AddOnGraphChangedHandler(
+		FOnGraphChanged::FDelegate::CreateUObject(this, &UNiagaraStackFunctionInput::OnGraphChanged));
+
+	InputParameterHandle = FNiagaraParameterHandle(InInputParameterHandle);
+	GenerateInputParameterHandlePath(*OwningModuleNode, *OwningFunctionCallNode, InputParameterHandlePath);
+	InputParameterHandlePath.Add(InputParameterHandle);
+
+	DisplayName = FText::FromString(InputParameterHandle.GetName());
+	AliasedInputParameterHandle = FNiagaraParameterHandle::CreateAliasedModuleParameterHandle(InputParameterHandle, OwningFunctionCallNode.Get());
+
+	InputType = InInputType;
+	StackEditorDataKey = FNiagaraStackGraphUtilities::GenerateStackFunctionInputEditorDataKey(*OwningFunctionCallNode.Get(), InputParameterHandle);
+}
+
+const UNiagaraNodeFunctionCall& UNiagaraStackFunctionInput::GetInputFunctionCallNode() const
+{
+	return *OwningFunctionCallNode.Get();
+}
+
+UNiagaraStackFunctionInput::EValueMode UNiagaraStackFunctionInput::GetValueMode()
+{
+	return InputValues.Mode;
+}
+
+bool UNiagaraStackFunctionInput::GetCanBePinned() const
+{
+	return OwningFunctionCallNode.IsValid() && OwningFunctionCallNode.Get()->FunctionScript->Usage == ENiagaraScriptUsage::Module;
+}
+
+const FNiagaraTypeDefinition& UNiagaraStackFunctionInput::GetInputType() const
+{
+	return InputType;
+}
+
+void UNiagaraStackFunctionInput::RefreshChildrenInternal(const TArray<UNiagaraStackEntry*>& CurrentChildren, TArray<UNiagaraStackEntry*>& NewChildren)
+{
+	RefreshValues();
+
+	if (InputValues.Mode == EValueMode::Dynamic && InputValues.DynamicNode.IsValid())
+	{
+		UNiagaraStackFunctionInputCollection* DynamicInputEntry = FindCurrentChildOfTypeByPredicate<UNiagaraStackFunctionInputCollection>(CurrentChildren,
+			[=](UNiagaraStackFunctionInputCollection* CurrentFunctionInputEntry) 
+		{ 
+			return CurrentFunctionInputEntry->GetInputFunctionCallNode() == InputValues.DynamicNode.Get() &&
+				CurrentFunctionInputEntry->GetModuleNode() == OwningModuleNode.Get(); 
+		});
+
+		if (DynamicInputEntry == nullptr)
+		{ 
+			UNiagaraStackFunctionInputCollection::FDisplayOptions DisplayOptions;
+			DisplayOptions.DisplayName = LOCTEXT("DynamicInputCollecitonDisplayName", "Dynamic Input Function Input");
+			DisplayOptions.bShouldShowInStack = false;
+			DisplayOptions.ChildItemIndentLevel = ItemIndentLevel + 1;
+
+			TArray<FString> InputParameterHandleStringPath;
+			for (const FNiagaraParameterHandle& InputParameterHandlePathItem : InputParameterHandlePath)
+			{
+				InputParameterHandleStringPath.Add(InputParameterHandlePathItem.GetParameterHandleString());
+			}
+			DynamicInputEntry = NewObject<UNiagaraStackFunctionInputCollection>(this);
+			DynamicInputEntry->Initialize(GetSystemViewModel(), GetEmitterViewModel(), *StackEditorData, *OwningModuleNode, *InputValues.DynamicNode.Get(), DisplayOptions);
+		}
+		NewChildren.Add(DynamicInputEntry);
+	}
+
+	if (InputValues.Mode == EValueMode::Data && InputValues.DataObjects.GetValueObject() != nullptr)
+	{
+		UNiagaraStackObject* ValueObjectEntry = FindCurrentChildOfTypeByPredicate<UNiagaraStackObject>(CurrentChildren,
+			[=](UNiagaraStackObject* CurrentObjectEntry) { return CurrentObjectEntry->GetObject() == InputValues.DataObjects.GetValueObject(); });
+
+		if(ValueObjectEntry == nullptr)
+		{
+			ValueObjectEntry = NewObject<UNiagaraStackObject>(this);
+			ValueObjectEntry->Initialize(GetSystemViewModel(), GetEmitterViewModel(), InputValues.DataObjects.GetValueObject());
+			ValueObjectEntry->SetItemIndentLevel(ItemIndentLevel + 1);
+		}
+		NewChildren.Add(ValueObjectEntry);
+	}
+}
+
+TSharedPtr<FStructOnScope> UNiagaraStackFunctionInput::FInputValues::GetLocalStructToReuse()
+{
+	return Mode == EValueMode::Local ? LocalStruct : TSharedPtr<FStructOnScope>();
+}
+
+UNiagaraDataInterface* UNiagaraStackFunctionInput::FInputValues::GetDataDefaultValueObjectToReuse()
+{
+	return Mode == EValueMode::Data && DataObjects.IsValid() && DataObjects.GetDefaultValueOwner() == FDataValues::EDefaultValueOwner::LocallyOwned
+		? DataObjects.GetDefaultValueObject()
+		: nullptr;
+}
+
+void UNiagaraStackFunctionInput::RefreshValues()
+{
+	if (ensureMsgf(InputParameterHandle.IsModuleHandle(), TEXT("Function inputs can only be generated for module paramters.")) == false)
+	{
+		return;
+	}
+
+	FInputValues OldValues = InputValues;
+	InputValues = FInputValues();
+
+	UEdGraphPin* DefaultPin = GetDefaultPin();
+	if (DefaultPin != nullptr)
+	{
+		UEdGraphPin* OverridePin = GetOverridePin();
+		UEdGraphPin* ValuePin = OverridePin != nullptr ? OverridePin : DefaultPin;
+
+		if (TryGetCurrentLocalValue(InputValues.LocalStruct, *ValuePin, OldValues.GetLocalStructToReuse()))
+		{
+			InputValues.Mode = EValueMode::Local;
+		}
+		else if (TryGetCurrentLinkedValue(InputValues.LinkedHandle, *ValuePin))
+		{
+			InputValues.Mode = EValueMode::Linked;
+		}
+		else if (TryGetCurrentDataValue(InputValues.DataObjects, OverridePin, *DefaultPin, OldValues.GetDataDefaultValueObjectToReuse()))
+		{
+			InputValues.Mode = EValueMode::Data;
+		}
+		else if (TryGetCurrentDynamicValue(InputValues.DynamicNode, OverridePin))
+		{
+			InputValues.Mode = EValueMode::Dynamic;
+		}
+	}
+
+	ValueChangedDelegate.Broadcast();
+}
+
+FText UNiagaraStackFunctionInput::GetDisplayName() const
+{
+	return DisplayName;
+}
+
+FName UNiagaraStackFunctionInput::GetTextStyleName() const
+{
+	return "NiagaraEditor.Stack.ParameterText";
+}
+
+bool UNiagaraStackFunctionInput::GetCanExpand() const
+{
+	return true;
+}
+
+const TArray<FNiagaraParameterHandle>& UNiagaraStackFunctionInput::GetInputParameterHandlePath() const
+{
+	return InputParameterHandlePath;
+}
+
+const FNiagaraParameterHandle& UNiagaraStackFunctionInput::GetInputParameterHandle() const
+{
+	return InputParameterHandle;
+}
+
+const FNiagaraParameterHandle& UNiagaraStackFunctionInput::GetLinkedValueHandle() const
+{
+	return InputValues.LinkedHandle;
+}
+
+void UNiagaraStackFunctionInput::SetLinkedValueHandle(const FNiagaraParameterHandle& InParameterHandle)
+{
+	if (InParameterHandle == InputValues.LinkedHandle)
+	{
+		return;
+	}
+
+	FScopedTransaction ScopedTransaction(LOCTEXT("UpdateLinkedInputValue", "Update linked input value"));
+	UEdGraphPin& OverridePin = GetOrCreateOverridePin();
+	RemoveAllNodesConnectedToOverridePin(OverridePin, this);
+
+	UEdGraph* Graph = OwningFunctionCallNode->GetGraph();
+	FGraphNodeCreator<UNiagaraNodeParameterMapGet> GetNodeCreator(*Graph);
+	UNiagaraNodeParameterMapGet* GetNode = GetNodeCreator.CreateNode();
+	GetNodeCreator.Finalize();
+
+	UEdGraphPin* GetInputPin = FNiagaraStackGraphUtilities::GetParameterMapInputPin(*GetNode);
+	checkf(GetInputPin != nullptr, TEXT("Parameter map get node was missing it's parameter map input pin."));
+
+	UEdGraphPin* OverrideNodeInputPin = FNiagaraStackGraphUtilities::GetParameterMapInputPin(*GetOverrideNode());
+	UEdGraphPin* PreviousStackNodeOutputPin = OverrideNodeInputPin->LinkedTo[0];
+	if (ensureMsgf(PreviousStackNodeOutputPin != nullptr, TEXT("Invalid Stack Graph - No previous stack node.")))
+	{
+		UEdGraphPin* GetOutputPin = GetNode->RequestNewTypedPin(EGPD_Output, InputType, InParameterHandle.GetParameterHandleString());
+		GetInputPin->MakeLinkTo(PreviousStackNodeOutputPin);
+		GetOutputPin->MakeLinkTo(&OverridePin);
+	}
+
+	FNiagaraStackGraphUtilities::RelayoutGraph(*Graph);
+	RefreshValues();
+}
+
+bool UsageRunsBefore(ENiagaraScriptUsage UsageA, ENiagaraScriptUsage UsageB)
+{
+	static TArray<ENiagaraScriptUsage> UsagesOrderedByExecution
+	{
+		ENiagaraScriptUsage::SystemSpawnScript,
+		ENiagaraScriptUsage::SystemUpdateScript,
+		ENiagaraScriptUsage::EmitterSpawnScript,
+		ENiagaraScriptUsage::EmitterUpdateScript,
+		ENiagaraScriptUsage::ParticleSpawnScript,
+		ENiagaraScriptUsage::ParticleUpdateScript
+	};
+
+	int32 IndexA;
+	int32 IndexB;
+	UsagesOrderedByExecution.Find(UsageA, IndexA);
+	UsagesOrderedByExecution.Find(UsageB, IndexB);
+	return IndexA < IndexB;
+}
+
+bool IsSpawnUsage(ENiagaraScriptUsage Usage)
+{
+	return
+		Usage == ENiagaraScriptUsage::SystemSpawnScript ||
+		Usage == ENiagaraScriptUsage::EmitterSpawnScript ||
+		Usage == ENiagaraScriptUsage::ParticleSpawnScript;
+}
+
+FString GetNamespaceForUsage(ENiagaraScriptUsage Usage)
+{
+	switch (Usage)
+	{
+	case ENiagaraScriptUsage::ParticleSpawnScript:
+	case ENiagaraScriptUsage::ParticleUpdateScript:
+		return FNiagaraParameterHandle::ParticleAttributeNamespace;
+	case ENiagaraScriptUsage::EmitterSpawnScript:
+	case ENiagaraScriptUsage::EmitterUpdateScript:
+		return FNiagaraParameterHandle::EmitterNamespace;
+	case ENiagaraScriptUsage::SystemSpawnScript:
+	case ENiagaraScriptUsage::SystemUpdateScript:
+		return FNiagaraParameterHandle::SystemNamespace;
+	default:
+		return FString();
+	}
+}
+
+void UNiagaraStackFunctionInput::GetAvailableParameterHandles(TArray<FNiagaraParameterHandle>& AvailableParameterHandles)
+{
+	// Engine Handles.
+	for (const FNiagaraVariable& SystemVariable : FNiagaraConstants::GetEngineConstants())
+	{
+		if (SystemVariable.GetType() == InputType)
+		{
+			AvailableParameterHandles.Add(FNiagaraParameterHandle::CreateEngineParameterHandle(SystemVariable));
+		}
+	}
+
+	UNiagaraNodeOutput* CurrentOutputNode = FNiagaraStackGraphUtilities::GetEmitterOutputNodeForStackNode(*OwningModuleNode);
+
+	TArray<UNiagaraNodeOutput*> AllOutputNodes;
+	GetEmitterViewModel()->GetSharedScriptViewModel()->GetGraphViewModel()->GetGraph()->GetNodesOfClass<UNiagaraNodeOutput>(AllOutputNodes);
+	if (GetSystemViewModel()->GetSystemIsTransient() == false)
+	{
+		GetSystemViewModel()->GetSystemScriptViewModel()->GetGraphViewModel()->GetGraph()->GetNodesOfClass<UNiagaraNodeOutput>(AllOutputNodes);
+	}
+
+	for (UNiagaraNodeOutput* OutputNode : AllOutputNodes)
+	{
+		if (OutputNode == CurrentOutputNode || UsageRunsBefore(OutputNode->GetUsage(), CurrentOutputNode->GetUsage()))
+		{
+			TArray<FNiagaraParameterHandle> AvailableParameterHandlesForThisOutput;
+			TArray<FNiagaraStackGraphUtilities::FStackNodeGroup> StackGroups;
+			FNiagaraStackGraphUtilities::GetStackNodeGroups(*OutputNode, StackGroups);
+
+			int32 CurrentModuleIndex = OutputNode == CurrentOutputNode
+				? StackGroups.IndexOfByPredicate([=](const FNiagaraStackGraphUtilities::FStackNodeGroup Group) { return Group.EndNode == OwningModuleNode; })
+				: INDEX_NONE;
+
+			int32 MaxGroupIndex = CurrentModuleIndex != INDEX_NONE ? CurrentModuleIndex : StackGroups.Num() - 1;
+			for (int32 i = 1; i < MaxGroupIndex; i++)
+			{
+				UNiagaraNodeFunctionCall* ModuleToCheck = Cast<UNiagaraNodeFunctionCall>(StackGroups[i].EndNode);
+				FNiagaraParameterMapHistoryBuilder Builder;
+				ModuleToCheck->BuildParameterMapHistory(Builder, false);
+
+				if (Builder.Histories.Num() == 1)
+				{
+					for (int32 j = 0; j < Builder.Histories[0].Variables.Num(); j++)
+					{
+						FNiagaraVariable& Variable = Builder.Histories[0].Variables[j];
+						if (Variable.GetType() == InputType)
+						{
+							TArray<const UEdGraphPin*>& WriteHistory = Builder.Histories[0].PerVariableWriteHistory[j];
+							for (const UEdGraphPin* WritePin : WriteHistory)
+							{
+								if (Cast<UNiagaraNodeParameterMapSet>(WritePin->GetOwningNode()) != nullptr)
+								{
+									FNiagaraParameterHandle AvailableHandle = FNiagaraParameterHandle(Variable.GetName().ToString());
+									AvailableParameterHandles.AddUnique(AvailableHandle);
+									AvailableParameterHandlesForThisOutput.AddUnique(AvailableHandle);
+									break;
+								}
+							}
+						}
+					}
+				}
+			}
+
+			if (OutputNode != CurrentOutputNode && IsSpawnUsage(OutputNode->GetUsage()))
+			{
+				FString OutputNodeNamespace = GetNamespaceForUsage(OutputNode->GetUsage());
+				if (OutputNodeNamespace.IsEmpty() == false)
+				{
+					for (FNiagaraParameterHandle& AvailableParameterHandleForThisOutput : AvailableParameterHandlesForThisOutput)
+					{
+						if (AvailableParameterHandleForThisOutput.GetNamespace() == OutputNodeNamespace)
+						{
+							AvailableParameterHandles.AddUnique(FNiagaraParameterHandle::CreateInitialParameterHandle(AvailableParameterHandleForThisOutput));
+						}
+					}
+				}
+			}
+		}
+	}
+}
+
+UNiagaraNodeFunctionCall* UNiagaraStackFunctionInput::GetDynamicInputNode() const
+{
+	return InputValues.DynamicNode.Get();
+}
+
+void UNiagaraStackFunctionInput::GetAvailableDynamicInputs(TArray<UNiagaraScript*>& AvailableDynamicInputs)
+{
+	UEnum* NiagaraScriptUsageEnum = FindObjectChecked<UEnum>(ANY_PACKAGE, TEXT("ENiagaraScriptUsage"), true);
+	FString QualifiedDynamicInputUsageString = NiagaraScriptUsageEnum->GetNameStringByValue(static_cast<uint8>(ENiagaraScriptUsage::DynamicInput));
+	int32 LastColonIndex;
+	QualifiedDynamicInputUsageString.FindLastChar(TEXT(':'), LastColonIndex);
+	FString UnqualifiedDynamicInputUsageString = QualifiedDynamicInputUsageString.RightChop(LastColonIndex + 1);
+
+	FARFilter DynamicInputFilter;
+	DynamicInputFilter.ClassNames.Add(UNiagaraScript::StaticClass()->GetFName());
+	DynamicInputFilter.TagsAndValues.Add(GET_MEMBER_NAME_CHECKED(UNiagaraScript, Usage), UnqualifiedDynamicInputUsageString);
+
+	FAssetRegistryModule& AssetRegistryModule = FModuleManager::LoadModuleChecked<FAssetRegistryModule>(TEXT("AssetRegistry"));
+	TArray<FAssetData> DynamicInputAssets;
+	AssetRegistryModule.Get().GetAssets(DynamicInputFilter, DynamicInputAssets);
+
+	for (const FAssetData& DynamicInputAsset : DynamicInputAssets)
+	{
+		UNiagaraScript* DynamicInputScript = Cast<UNiagaraScript>(DynamicInputAsset.GetAsset());
+		if (DynamicInputScript != nullptr)
+		{
+			UNiagaraScriptSource* DynamicInputScriptSource = Cast<UNiagaraScriptSource>(DynamicInputScript->GetSource());
+			TArray<UNiagaraNodeOutput*> OutputNodes;
+			DynamicInputScriptSource->NodeGraph->GetNodesOfClass<UNiagaraNodeOutput>(OutputNodes);
+			if (OutputNodes.Num() == 1)
+			{
+				TArray<UEdGraphPin*> InputPins;
+				OutputNodes[0]->GetInputPins(InputPins);
+				if (InputPins.Num() == 1)
+				{
+					const UEdGraphSchema_Niagara* NiagaraSchema = GetDefault<UEdGraphSchema_Niagara>();
+					FNiagaraTypeDefinition PinType = NiagaraSchema->PinToTypeDefinition(InputPins[0]);
+					if (PinType == InputType)
+					{
+						AvailableDynamicInputs.Add(DynamicInputScript);
+					}
+				}
+			}
+		}
+	}
+}
+
+void UNiagaraStackFunctionInput::SetDynamicInput(UNiagaraScript* DynamicInput)
+{
+	FScopedTransaction ScopedTransaction(LOCTEXT("SetDynamicInput", "Make dynamic input"));
+
+	UEdGraphPin& OverridePin = GetOrCreateOverridePin();
+	RemoveAllNodesConnectedToOverridePin(OverridePin, this);
+
+	UEdGraph* Graph = OwningFunctionCallNode->GetGraph();
+	FGraphNodeCreator<UNiagaraNodeFunctionCall> FunctionCallNodeCreator(*Graph);
+	UNiagaraNodeFunctionCall* FunctionCallNode = FunctionCallNodeCreator.CreateNode();
+	FunctionCallNode->FunctionScript = DynamicInput;
+	FunctionCallNodeCreator.Finalize();
+
+	UEdGraphPin* FunctionCallInputPin = FNiagaraStackGraphUtilities::GetParameterMapInputPin(*FunctionCallNode);
+	TArray<UEdGraphPin*> FunctionCallOutputPins;
+	FunctionCallNode->GetOutputPins(FunctionCallOutputPins);
+
+	const UEdGraphSchema_Niagara* NiagaraSchema = GetDefault<UEdGraphSchema_Niagara>();
+
+	if (ensureMsgf(FunctionCallInputPin != nullptr, TEXT("Invalid Stack Graph - Dynamic Input function call did not have a parameter map input pin.")) &&
+		ensureMsgf(FunctionCallOutputPins.Num() == 1 && NiagaraSchema->PinToTypeDefinition(FunctionCallOutputPins[0]) == InputType, 
+			TEXT("Invalid Stack Graph - Dynamic Input function did not have the correct typed output pin")))
+	{
+		UEdGraphPin* OverrideNodeInputPin = FNiagaraStackGraphUtilities::GetParameterMapInputPin(*GetOverrideNode());
+		UEdGraphPin* PreviousStackNodeOutputPin = OverrideNodeInputPin->LinkedTo[0];
+		if (ensureMsgf(PreviousStackNodeOutputPin != nullptr, TEXT("Invalid Stack Graph - No previous stack node.")))
+		{
+			FunctionCallInputPin->MakeLinkTo(PreviousStackNodeOutputPin);
+			FunctionCallOutputPins[0]->MakeLinkTo(&OverridePin);
+		}
+	}
+
+	FNiagaraStackGraphUtilities::InitializeDataInterfaceInputs(GetSystemViewModel(), GetEmitterViewModel(), *StackEditorData, *OwningModuleNode, *FunctionCallNode);
+	FNiagaraStackGraphUtilities::RelayoutGraph(*Graph);
+
+	RefreshChildren();
+}
+
+TSharedPtr<FStructOnScope> UNiagaraStackFunctionInput::GetLocalValueStruct()
+{
+	return InputValues.LocalStruct;
+}
+
+UNiagaraDataInterface* UNiagaraStackFunctionInput::GetDataValueObject()
+{
+	return InputValues.DataObjects.GetValueObject();
+}
+
+bool UNiagaraStackFunctionInput::GetIsPinned() const
+{
+	return StackEditorData->GetModuleInputIsPinned(StackEditorDataKey);
+}
+
+void UNiagaraStackFunctionInput::SetIsPinned(bool bIsPinned)
+{
+	StackEditorData->SetModuleInputIsPinned(StackEditorDataKey, bIsPinned);
+	PinnedChangedDelegate.Broadcast();
+}
+
+void UNiagaraStackFunctionInput::NotifyBeginLocalValueChange()
+{
+	GEditor->BeginTransaction(LOCTEXT("BeginEditModuleInputLocalValue", "Edit input local value."));
+}
+
+void UNiagaraStackFunctionInput::NotifyEndLocalValueChange()
+{
+	if (GEditor->IsTransactionActive())
+	{
+		GEditor->EndTransaction();
+	}
+}
+
+void UNiagaraStackFunctionInput::SetLocalValue(TSharedRef<FStructOnScope> InLocalValue)
+{
+	UEdGraphPin* DefaultPin = GetDefaultPin();
+	UEdGraphPin* CurrentOverridePin = GetOverridePin();
+	UEdGraphPin* ValuePin = CurrentOverridePin != nullptr ? CurrentOverridePin : DefaultPin;
+
+	TSharedPtr<FStructOnScope> CurrentValue;
+	bool bCanHaveLocalValue = ValuePin != nullptr;
+	bool bHasLocalValue = bCanHaveLocalValue && InputValues.Mode == EValueMode::Local && TryGetCurrentLocalValue(CurrentValue, *ValuePin, TSharedPtr<FStructOnScope>());
+	bool bLocalValueMatchesSetValue = bHasLocalValue && FNiagaraEditorUtilities::DataMatches(*CurrentValue.Get(), InLocalValue.Get());
+
+	if (bCanHaveLocalValue == false || bLocalValueMatchesSetValue)
+	{
+		return;
+	}
+
+	FScopedTransaction ScopedTransaction(LOCTEXT("UpdateInputLocalValue", "Update input local value"));
+	UNiagaraGraph* EmitterGraph = Cast<UNiagaraGraph>(OwningFunctionCallNode->GetGraph());
+	UEdGraphPin& OverridePin = GetOrCreateOverridePin();
+
+	bool bGraphWillNeedRelayout = &OverridePin != CurrentOverridePin || OverridePin.LinkedTo.Num() > 0;
+
+	RemoveAllNodesConnectedToOverridePin(OverridePin, this);
+	OverridePin.Modify();
+	FNiagaraVariable LocalValueVariable(InputType, NAME_None);
+	LocalValueVariable.SetData(InLocalValue->GetStructMemory());
+	const UEdGraphSchema_Niagara* NiagaraSchema = GetDefault<UEdGraphSchema_Niagara>();
+	OverridePin.DefaultValue = NiagaraSchema->NiagaraVariableToPinDefaultValue(LocalValueVariable);
+	EmitterGraph->NotifyGraphNeedsRecompile();
+
+	if (bGraphWillNeedRelayout)
+	{
+		FNiagaraStackGraphUtilities::RelayoutGraph(*EmitterGraph);
+	}
+}
+
+bool UNiagaraStackFunctionInput::CanReset() const
+{
+	if(InputValues.Mode == EValueMode::Data)
+	{
+		// For data values a copy of the default object should have been created automatically and attached to the override pin for this input.  If a 
+		// copy of the default object wasn't created, the input can be reset to create one.  If a copy of the data object is available it can be
+		// reset if it's different from it's default value.
+		bool bHasDataValueObject = InputValues.DataObjects.GetValueObject() != nullptr;
+		bool bHasDefaultDataValueObject = InputValues.DataObjects.GetDefaultValueObject() != nullptr;
+		bool bIsDataValueDifferentFromDefaultDataValue = bHasDataValueObject && bHasDefaultDataValueObject 
+			&& InputValues.DataObjects.GetValueObject()->Equals(InputValues.DataObjects.GetDefaultValueObject()) == false;
+		return bHasDataValueObject == false || bHasDefaultDataValueObject == false || bIsDataValueDifferentFromDefaultDataValue;
+	}
+	else
+	{
+		// All other input modes can be reset if there is an override pin available.
+		return GetOverridePin() != nullptr;
+	}
+}
+
+void UNiagaraStackFunctionInput::Reset()
+{
+	if(InputValues.Mode == EValueMode::Data)
+	{
+		// For data values they are reset by making sure the data object owned by this input matches the default
+		// data object.  If there is no data object owned by the input, one is created and and updated to match the default.
+		FScopedTransaction ScopedTransaction(LOCTEXT("ResetInputObjectTransaction", "Reset the inputs data interface object to default."));
+		if (InputValues.DataObjects.GetValueObject() != nullptr && InputValues.DataObjects.GetDefaultValueObject() != nullptr)
+		{
+			InputValues.DataObjects.GetDefaultValueObject()->CopyTo(InputValues.DataObjects.GetValueObject());
+		}
+		else
+		{
+			UEdGraphPin& OverridePin = GetOrCreateOverridePin();
+			RemoveAllNodesConnectedToOverridePin(OverridePin, this);
+
+			FString InputNodeName = InputParameterHandlePath[0].GetName();
+			for (int32 i = 1; i < InputParameterHandlePath.Num(); i++)
+			{
+				InputNodeName += "." + InputParameterHandlePath[i].GetName();
+			}
+
+			UEdGraph* Graph = OwningFunctionCallNode->GetGraph();
+			FGraphNodeCreator<UNiagaraNodeInput> InputNodeCreator(*Graph);
+			UNiagaraNodeInput* InputNode = InputNodeCreator.CreateNode();
+			FNiagaraEditorUtilities::InitializeParameterInputNode(*InputNode, InputType, CastChecked<UNiagaraGraph>(Graph), *InputNodeName);
+
+			UNiagaraDataInterface* InputValueObject = NewObject<UNiagaraDataInterface>(InputNode, const_cast<UClass*>(InputType.GetClass()), *InputParameterHandle.GetName(), RF_Transactional);
+			if (InputValues.DataObjects.GetDefaultValueObject() != nullptr)
+			{
+				InputValues.DataObjects.GetDefaultValueObject()->CopyTo(InputValueObject);
+			}
+			InputNode->DataInterface = InputValueObject;
+
+			InputNodeCreator.Finalize();
+			FNiagaraStackGraphUtilities::ConnectPinToInputNode(OverridePin, *InputNode);
+			FNiagaraStackGraphUtilities::RelayoutGraph(*Graph);
+		}
+	}
+	else
+	{
+		// For all other value modes removing the nodes connected to the override pin resets them.
+		UNiagaraNodeParameterMapSet* OverrideNode = GetOverrideNode();
+		UEdGraphPin* OverridePin = GetOverridePin();
+		if (ensureMsgf(OverrideNode != nullptr && OverridePin != nullptr, TEXT("Can not reset the value of an input that doesn't have a valid override node and override pin")))
+		{
+			FScopedTransaction ScopedTransaction(LOCTEXT("ResetInputStructTransaction", "Reset the inputs value to default."));
+			RemoveAllNodesConnectedToOverridePin(*OverridePin, this);
+			OverrideNode->Modify();
+			OverrideNode->RemovePin(OverridePin);
+
+			OwningFunctionCallNode->GetNiagaraGraph()->NotifyGraphNeedsRecompile();
+			FNiagaraStackGraphUtilities::RelayoutGraph(*OwningFunctionCallNode->GetGraph());
+		}
+	}
+	RefreshChildren();
+}
+
+bool UNiagaraStackFunctionInput::CanRenameInput() const
+{
+	// Only module level assignment node inputs can be renamed.
+	return OwningAssignmentNode.IsValid() && InputParameterHandlePath.Num() == 1;
+}
+
+bool UNiagaraStackFunctionInput::GetIsRenamePending() const
+{
+	return CanRenameInput() && StackEditorData->GetModuleInputIsRenamePending(StackEditorDataKey);
+}
+
+void UNiagaraStackFunctionInput::SetIsRenamePending(bool bIsRenamePending)
+{
+	if (CanRenameInput())
+	{
+		StackEditorData->SetModuleInputIsRenamePending(StackEditorDataKey, bIsRenamePending);
+	}
+}
+
+void UNiagaraStackFunctionInput::RenameInput(FString NewName)
+{
+	if (OwningAssignmentNode.IsValid() && InputParameterHandlePath.Num() == 1 && InputParameterHandle.GetName() != NewName)
+	{
+		bool bIsCurrentlyPinned = GetIsPinned();
+		bool bIsCurrentlyExpanded = StackEditorData->GetStackEntryIsExpanded(FNiagaraStackGraphUtilities::GenerateStackModuleEditorDataKey(*OwningAssignmentNode), false);
+
+		FNiagaraParameterHandle TargetHandle(OwningAssignmentNode->AssignmentTarget.GetName().ToString());
+		FNiagaraParameterHandle RenamedTargetHandle(TargetHandle.GetNamespace(), NewName);
+		OwningAssignmentNode->AssignmentTarget.SetName(*RenamedTargetHandle.GetParameterHandleString());
+		OwningAssignmentNode->RefreshFromExternalChanges();
+
+		InputParameterHandle = FNiagaraParameterHandle(InputParameterHandle.GetNamespace(), NewName);
+		InputParameterHandlePath.Empty();
+		InputParameterHandlePath.Add(InputParameterHandle);
+		AliasedInputParameterHandle = FNiagaraParameterHandle::CreateAliasedModuleParameterHandle(InputParameterHandle, OwningAssignmentNode.Get());
+		DisplayName = FText::FromString(InputParameterHandle.GetName());
+
+		UEdGraphPin* OverridePin = GetOverridePin();
+		if (OverridePin != nullptr)
+		{
+			OverridePin->PinName = AliasedInputParameterHandle.GetParameterHandleString();
+		}
+
+		StackEditorDataKey = FNiagaraStackGraphUtilities::GenerateStackFunctionInputEditorDataKey(*OwningFunctionCallNode.Get(), InputParameterHandle);
+		StackEditorData->SetModuleInputIsPinned(StackEditorDataKey, bIsCurrentlyPinned);
+		StackEditorData->SetStackEntryIsExpanded(FNiagaraStackGraphUtilities::GenerateStackModuleEditorDataKey(*OwningAssignmentNode), bIsCurrentlyExpanded);
+
+		CastChecked<UNiagaraGraph>(OwningAssignmentNode->GetGraph())->NotifyGraphNeedsRecompile();
+	}
+}
+
+void UNiagaraStackFunctionInput::GetNamespacesForNewParameters(TArray<FString>& OutNamespacesForNewParameters) const
+{
+	UNiagaraNodeOutput* OutputNode = FNiagaraStackGraphUtilities::GetEmitterOutputNodeForStackNode(*OwningFunctionCallNode);
+	bool bIsSystemTransient = GetSystemViewModel()->GetSystemIsTransient();
+
+	if (OutputNode->GetUsage() == ENiagaraScriptUsage::ParticleSpawnScript || OutputNode->GetUsage() == ENiagaraScriptUsage::ParticleUpdateScript)
+	{
+		OutNamespacesForNewParameters.Add(FNiagaraParameterHandle::ParticleAttributeNamespace);
+		OutNamespacesForNewParameters.Add(FNiagaraParameterHandle::EmitterNamespace);
+		if (bIsSystemTransient == false)
+		{
+			OutNamespacesForNewParameters.Add(FNiagaraParameterHandle::SystemNamespace);
+			OutNamespacesForNewParameters.Add(FNiagaraParameterHandle::UserNamespace);
+		}
+	}
+	else if (OutputNode->GetUsage() == ENiagaraScriptUsage::EmitterSpawnScript || OutputNode->GetUsage() == ENiagaraScriptUsage::EmitterUpdateScript)
+	{
+		OutNamespacesForNewParameters.Add(FNiagaraParameterHandle::EmitterNamespace);
+		if (bIsSystemTransient == false)
+		{
+			OutNamespacesForNewParameters.Add(FNiagaraParameterHandle::SystemNamespace);
+			OutNamespacesForNewParameters.Add(FNiagaraParameterHandle::UserNamespace);
+		}
+	}
+	else if (OutputNode->GetUsage() == ENiagaraScriptUsage::SystemSpawnScript || OutputNode->GetUsage() == ENiagaraScriptUsage::ParticleUpdateScript && bIsSystemTransient == false)
+	{
+		OutNamespacesForNewParameters.Add(FNiagaraParameterHandle::SystemNamespace);
+		OutNamespacesForNewParameters.Add(FNiagaraParameterHandle::UserNamespace);
+	}
+}
+
+UNiagaraStackFunctionInput::FOnValueChanged& UNiagaraStackFunctionInput::OnValueChanged()
+{
+	return ValueChangedDelegate;
+}
+
+UNiagaraStackFunctionInput::FOnPinnedChanged& UNiagaraStackFunctionInput::OnPinnedChanged()
+{
+	return PinnedChangedDelegate;
+}
+
+void UNiagaraStackFunctionInput::OnGraphChanged(const struct FEdGraphEditAction& InAction)
+{
+	if (bUpdatingGraphDirectly == false)
+	{
+		OverrideNodeCache.Reset();
+		OverridePinCache.Reset();
+	}
+}
+
+UNiagaraNodeParameterMapSet* UNiagaraStackFunctionInput::GetOverrideNode() const
+{
+	if (OverrideNodeCache.IsSet() == false)
+	{
+		UNiagaraNodeParameterMapSet* OverrideNode = nullptr;
+		if (OwningFunctionCallNode.IsValid())
+		{
+			UEdGraphPin* ParameterMapInput = FNiagaraStackGraphUtilities::GetParameterMapInputPin(*OwningFunctionCallNode);
+			if (ParameterMapInput != nullptr && ParameterMapInput->LinkedTo.Num() == 1)
+			{
+				OverrideNode = Cast<UNiagaraNodeParameterMapSet>(ParameterMapInput->LinkedTo[0]->GetOwningNode());
+			}
+		}
+		OverrideNodeCache = OverrideNode;
+	}
+	return OverrideNodeCache.GetValue();
+}
+
+UNiagaraNodeParameterMapSet& UNiagaraStackFunctionInput::GetOrCreateOverrideNode()
+{
+	UNiagaraNodeParameterMapSet* OverrideNode = GetOverrideNode();
+	if (OverrideNode == nullptr)
+	{
+		TGuardValue<bool>(bUpdatingGraphDirectly, true);
+
+		UEdGraph* Graph = OwningModuleNode->GetGraph();
+		FGraphNodeCreator<UNiagaraNodeParameterMapSet> ParameterMapSetNodeCreator(*Graph);
+		OverrideNode = ParameterMapSetNodeCreator.CreateNode();
+		ParameterMapSetNodeCreator.Finalize();
+
+		UEdGraphPin* OverrideNodeInputPin = FNiagaraStackGraphUtilities::GetParameterMapInputPin(*OverrideNode);
+		UEdGraphPin* OverrideNodeOutputPin = FNiagaraStackGraphUtilities::GetParameterMapOutputPin(*OverrideNode);
+
+		UEdGraphPin* OwningFunctionCallInputPin = FNiagaraStackGraphUtilities::GetParameterMapInputPin(*OwningFunctionCallNode);
+		UEdGraphPin* PreviousStackNodeOutputPin = OwningFunctionCallInputPin->LinkedTo[0];
+
+		OwningFunctionCallInputPin->BreakAllPinLinks();
+		OwningFunctionCallInputPin->MakeLinkTo(OverrideNodeOutputPin);
+		for (UEdGraphPin* PreviousStackNodeOutputLinkedPin : PreviousStackNodeOutputPin->LinkedTo)
+		{
+			PreviousStackNodeOutputLinkedPin->MakeLinkTo(OverrideNodeOutputPin);
+		}
+		PreviousStackNodeOutputPin->BreakAllPinLinks();
+		PreviousStackNodeOutputPin->MakeLinkTo(OverrideNodeInputPin);
+
+		OverrideNodeCache = OverrideNode;
+	}
+	return *OverrideNode;
+}
+
+UEdGraphPin* UNiagaraStackFunctionInput::GetDefaultPin() const
+{
+	return OwningFunctionCallNode->FindParameterMapDefaultValuePin(InputParameterHandle.GetParameterHandleString());
+}
+
+UEdGraphPin* UNiagaraStackFunctionInput::GetOverridePin() const
+{
+	if (OverridePinCache.IsSet() == false)
+	{
+		OverridePinCache = nullptr;
+		UNiagaraNodeParameterMapSet* OverrideNode = GetOverrideNode();
+		if (OverrideNode != nullptr)
+		{
+			TArray<UEdGraphPin*> InputPins;
+			OverrideNode->GetInputPins(InputPins);
+			UEdGraphPin** OverridePinPtr = InputPins.FindByPredicate([&](const UEdGraphPin* Pin) { return Pin->PinName == AliasedInputParameterHandle.GetParameterHandleString(); });
+			if (OverridePinPtr != nullptr)
+			{
+				OverridePinCache = *OverridePinPtr;
+			}
+		}
+	}
+	return OverridePinCache.GetValue();
+}
+
+UEdGraphPin& UNiagaraStackFunctionInput::GetOrCreateOverridePin()
+{
+	UEdGraphPin* OverridePin = GetOverridePin();
+	if (OverridePin == nullptr)
+	{
+		TGuardValue<bool>(bUpdatingGraphDirectly, true);
+
+		UNiagaraNodeParameterMapSet& OverrideNode = GetOrCreateOverrideNode();
+		OverrideNode.Modify();
+
+		TArray<UEdGraphPin*> OverrideInputPins;
+		OverrideNode.GetInputPins(OverrideInputPins);
+
+		const UEdGraphSchema_Niagara* NiagaraSchema = GetDefault<UEdGraphSchema_Niagara>();
+		FEdGraphPinType PinType = NiagaraSchema->TypeDefinitionToPinType(InputType);
+		OverridePin = OverrideNode.CreatePin(EEdGraphPinDirection::EGPD_Input, PinType, AliasedInputParameterHandle.GetParameterHandleString(), OverrideInputPins.Num() - 1);
+
+		OverridePinCache = OverridePin;
+	}
+	return *OverridePin;
+}
+
+bool UNiagaraStackFunctionInput::TryGetCurrentLocalValue(TSharedPtr<FStructOnScope>& LocalValue, UEdGraphPin& ValuePin, TSharedPtr<FStructOnScope> OldValueToReuse)
+{
+	if (InputType.IsDataInterface() == false && ValuePin.LinkedTo.Num() == 0)
+	{
+		const UEdGraphSchema_Niagara* NiagaraSchema = GetDefault<UEdGraphSchema_Niagara>();
+		FNiagaraVariable ValueVariable = NiagaraSchema->PinToNiagaraVariable(&ValuePin);
+		if (OldValueToReuse.IsValid() && OldValueToReuse->GetStruct() == ValueVariable.GetType().GetStruct())
+		{
+			LocalValue = OldValueToReuse;
+		}
+		else
+		{
+			LocalValue = MakeShared<FStructOnScope>(ValueVariable.GetType().GetStruct());
+		}
+		ValueVariable.CopyTo(LocalValue->GetStructMemory());
+		return true;
+	}
+	return false;
+}
+
+bool UNiagaraStackFunctionInput::TryGetCurrentDataValue(FDataValues& DataValues, UEdGraphPin* OverrideValuePin, UEdGraphPin& DefaultValuePin, UNiagaraDataInterface* LocallyOwnedDefaultDataValueObjectToReuse)
+{
+	if (InputType.GetClass() != nullptr)
+	{
+		UNiagaraDataInterface* DataValueObject = nullptr;
+		if (OverrideValuePin != nullptr)
+		{
+			if (OverrideValuePin->LinkedTo.Num() == 1)
+			{
+				UNiagaraNodeInput* InputNode = Cast<UNiagaraNodeInput>(OverrideValuePin->LinkedTo[0]->GetOwningNode());
+				if (InputNode != nullptr)
+				{
+					if (InputNode->Usage == ENiagaraInputNodeUsage::Parameter)
+					{
+						DataValueObject = InputNode->DataInterface;
+					}
+				}
+			}
+		}
+
+		UNiagaraDataInterface* DefaultDataValueObject = nullptr;
+		FDataValues::EDefaultValueOwner DefaultDataValueOwner = FDataValues::EDefaultValueOwner::Invalid;
+		if (DefaultValuePin.LinkedTo.Num() == 1)
+		{
+			UNiagaraNodeInput* InputNode = Cast<UNiagaraNodeInput>(DefaultValuePin.LinkedTo[0]->GetOwningNode());
+			if (InputNode != nullptr && InputNode->Usage == ENiagaraInputNodeUsage::Parameter && InputNode->DataInterface != nullptr)
+			{
+				DefaultDataValueObject = InputNode->DataInterface;
+				DefaultDataValueOwner = FDataValues::EDefaultValueOwner::FunctionOwned;
+			}
+		}
+
+		if (DefaultDataValueObject == nullptr)
+		{
+			if (LocallyOwnedDefaultDataValueObjectToReuse == nullptr)
+			{
+				DefaultDataValueObject = NewObject<UNiagaraDataInterface>(this, const_cast<UClass*>(InputType.GetClass()));
+			}
+			else
+			{
+				DefaultDataValueObject = LocallyOwnedDefaultDataValueObjectToReuse;
+			}
+			DefaultDataValueOwner = FDataValues::EDefaultValueOwner::LocallyOwned;
+		}
+		
+		DataValues = FDataValues(DataValueObject, DefaultDataValueObject, DefaultDataValueOwner);
+		return true;
+	}
+	return false;
+}
+
+bool UNiagaraStackFunctionInput::TryGetCurrentLinkedValue(FNiagaraParameterHandle& LinkedValueHandle, UEdGraphPin& ValuePin)
+{
+	if (ValuePin.LinkedTo.Num() == 1)
+	{
+		UEdGraphPin* CurrentValuePin = &ValuePin;
+		TSharedPtr<TArray<FNiagaraParameterHandle>> AvailableHandles;
+		while (CurrentValuePin != nullptr)
+		{
+			UEdGraphPin* LinkedValuePin = CurrentValuePin->LinkedTo[0];
+			CurrentValuePin = nullptr;
+
+			UNiagaraNodeParameterMapGet* GetNode = Cast<UNiagaraNodeParameterMapGet>(LinkedValuePin->GetOwningNode());
+			if (GetNode == nullptr)
+			{
+				// Only parameter map get nodes are supported for linked values.
+				return false;
+			}
+
+			// If a parameter map get node was found, the linked handle will be stored in the pin name.  
+			FNiagaraParameterHandle LinkedValueHandleFromNode(LinkedValuePin->PinName);
+
+			UEdGraphPin* LinkedValueHandleDefaultPin = GetNode->GetDefaultPin(LinkedValuePin);
+			if (LinkedValueHandleDefaultPin->LinkedTo.Num() == 0)
+			{
+				// If the default value pin for this get node isn't connected this is the last read in the chain
+				// so return the handle.
+				LinkedValueHandle = LinkedValueHandleFromNode;
+				return true;
+			}
+			else
+			{
+				// If the default value pin for the get node is connected then there are a chain of possible values.
+				// if the value of the current get node is available it can be returned, otherwise we need to check the
+				// next node.
+				if (AvailableHandles.IsValid() == false)
+				{
+					AvailableHandles = MakeShared<TArray<FNiagaraParameterHandle>>();
+					GetAvailableParameterHandles(*AvailableHandles);
+				}
+
+				if (AvailableHandles->Contains(LinkedValueHandleFromNode))
+				{
+					LinkedValueHandle = LinkedValueHandleFromNode;
+					return true;
+				}
+				else
+				{
+					CurrentValuePin = LinkedValueHandleDefaultPin;
+				}
+			}
+		}
+	}
+	return false;
+}
+
+bool UNiagaraStackFunctionInput::TryGetCurrentDynamicValue(TWeakObjectPtr<UNiagaraNodeFunctionCall>& DynamicValue, UEdGraphPin* OverridePin)
+{
+	if (OverridePin != nullptr && OverridePin->LinkedTo.Num() == 1)
+	{
+		UNiagaraNodeFunctionCall* DynamicNode = Cast<UNiagaraNodeFunctionCall>(OverridePin->LinkedTo[0]->GetOwningNode());
+		if (DynamicNode != nullptr)
+		{
+			DynamicValue = DynamicNode;
+			return true;
+		}
+	}
+	return false;
+}
+
+void UNiagaraStackFunctionInput::RemoveAllNodesConnectedToOverridePin(UEdGraphPin& OverridePin, UNiagaraStackFunctionInput* OwningInput)
+{
+	if (OverridePin.LinkedTo.Num() == 1)
+	{
+		UEdGraphNode* OverrideValueNode = OverridePin.LinkedTo[0]->GetOwningNode();
+		UEdGraph* Graph = OverrideValueNode->GetGraph();
+		if (OverrideValueNode->IsA<UNiagaraNodeInput>() || OverrideValueNode->IsA<UNiagaraNodeParameterMapGet>())
+		{
+			Graph->RemoveNode(OverrideValueNode);
+			if (OwningInput != nullptr && OverrideValueNode->IsA<UNiagaraNodeInput>())
+			{
+				// Input nodes are used to Data Interface objects so we need to call the the data object modified delegate
+				// when removing them.
+				OwningInput->OnDataObjectModified().Broadcast(nullptr);
+			}
+		}
+		else if (OverrideValueNode->IsA<UNiagaraNodeFunctionCall>())
+		{
+			UNiagaraNodeFunctionCall* DynamicInputNode = CastChecked<UNiagaraNodeFunctionCall>(OverrideValueNode);
+			UEdGraphPin* DynamicInputNodeInputPin = FNiagaraStackGraphUtilities::GetParameterMapInputPin(*DynamicInputNode);
+			UNiagaraNodeParameterMapSet* DynamicInputNodeOverrideNode = Cast<UNiagaraNodeParameterMapSet>(DynamicInputNodeInputPin->LinkedTo[0]->GetOwningNode());
+			if (DynamicInputNodeOverrideNode != nullptr)
+			{
+				TArray<UEdGraphPin*> InputPins;
+				DynamicInputNodeOverrideNode->GetInputPins(InputPins);
+				for (UEdGraphPin* InputPin : InputPins)
+				{
+					if (InputPin->PinName.StartsWith(DynamicInputNode->GetFunctionName()))
+					{
+						RemoveAllNodesConnectedToOverridePin(*InputPin, OwningInput);
+						DynamicInputNodeOverrideNode->RemovePin(InputPin);
+					}
+				}
+
+				TArray<UEdGraphPin*> NewInputPins;
+				DynamicInputNodeOverrideNode->GetInputPins(NewInputPins);
+				if (NewInputPins.Num() == 2)
+				{
+					// If there are only 2 pins, they are the parameter map input and the add pin, so the dynamic input's override node can be removed.  This
+					// not always be the case when removing dynamic input nodes because they share the same override node.
+					UEdGraphPin* InputPin = FNiagaraStackGraphUtilities::GetParameterMapInputPin(*DynamicInputNodeOverrideNode);
+					UEdGraphPin* OutputPin = FNiagaraStackGraphUtilities::GetParameterMapOutputPin(*DynamicInputNodeOverrideNode);
+
+					if (ensureMsgf(InputPin != nullptr && InputPin->LinkedTo.Num() == 1 &&
+						OutputPin != nullptr && OutputPin->LinkedTo.Num() == 2, 
+						TEXT("Invalid Stack - Dynamic input node override node not connected correctly.")))
+					{
+						// The DynamicInputOverrideNode will have a single input which is the previous module or override map set, and
+						// two output links, one to the dynamic input node and one to the next override map set.
+						UEdGraphPin* LinkedInputPin = InputPin->LinkedTo[0];
+						UEdGraphPin* LinkedOutputPin = OutputPin->LinkedTo[0]->GetOwningNode() != DynamicInputNode
+							? OutputPin->LinkedTo[0]
+							: OutputPin->LinkedTo[1];
+						InputPin->BreakAllPinLinks();
+						OutputPin->BreakAllPinLinks();
+						LinkedInputPin->MakeLinkTo(LinkedOutputPin);
+						Graph->RemoveNode(DynamicInputNodeOverrideNode);
+					}
+				}
+			}
+
+			Graph->RemoveNode(DynamicInputNode);
+		}
+	}
+}
+
+#undef LOCTEXT_NAMESPACE
diff --git a/Engine/Plugins/FX/Niagara/Source/NiagaraEditor/Private/ViewModels/Stack/NiagaraStackFunctionInputCollection.cpp b/Engine/Plugins/FX/Niagara/Source/NiagaraEditor/Private/ViewModels/Stack/NiagaraStackFunctionInputCollection.cpp
new file mode 100644
--- /dev/null
+++ b/Engine/Plugins/FX/Niagara/Source/NiagaraEditor/Private/ViewModels/Stack/NiagaraStackFunctionInputCollection.cpp
@@ -0,0 +1,104 @@
+// Copyright 1998-2016 Epic Games, Inc. All Rights Reserved.
+
+#include "NiagaraStackFunctionInputCollection.h"
+#include "NiagaraStackFunctionInput.h"
+#include "NiagaraNodeFunctionCall.h"
+#include "NiagaraEmitterEditorData.h"
+#include "EdGraphSchema_Niagara.h"
+#include "NiagaraStackGraphUtilities.h"
+
+#include "EdGraph/EdGraphPin.h"
+
+UNiagaraStackFunctionInputCollection::UNiagaraStackFunctionInputCollection()
+	: ModuleNode(nullptr)
+	, InputFunctionCallNode(nullptr)
+{
+}
+
+UNiagaraNodeFunctionCall* UNiagaraStackFunctionInputCollection::GetModuleNode() const
+{
+	return ModuleNode;
+}
+
+UNiagaraNodeFunctionCall* UNiagaraStackFunctionInputCollection::GetInputFunctionCallNode() const
+{
+	return InputFunctionCallNode;
+}
+
+void UNiagaraStackFunctionInputCollection::Initialize(
+	TSharedRef<FNiagaraSystemViewModel> InSystemViewModel,
+	TSharedRef<FNiagaraEmitterViewModel> InEmitterViewModel,
+	UNiagaraStackEditorData& InStackEditorData,
+	UNiagaraNodeFunctionCall& InModuleNode,
+	UNiagaraNodeFunctionCall& InInputFunctionCallNode,
+	FDisplayOptions InDisplayOptions)
+{
+	checkf(ModuleNode == nullptr && InputFunctionCallNode == nullptr, TEXT("Can not set the node more than once."));
+	Super::Initialize(InSystemViewModel, InEmitterViewModel);
+	StackEditorData = &InStackEditorData;
+	ModuleNode = &InModuleNode;
+	InputFunctionCallNode = &InInputFunctionCallNode;
+	DisplayOptions = InDisplayOptions;
+}
+
+FText UNiagaraStackFunctionInputCollection::GetDisplayName() const
+{
+	return DisplayOptions.DisplayName;
+}
+
+FName UNiagaraStackFunctionInputCollection::GetTextStyleName() const
+{
+	return "NiagaraEditor.Stack.ParameterCollectionText";
+}
+
+bool UNiagaraStackFunctionInputCollection::GetCanExpand() const
+{
+	return true;
+}
+
+bool UNiagaraStackFunctionInputCollection::GetShouldShowInStack() const
+{
+	return DisplayOptions.bShouldShowInStack;
+}
+
+UNiagaraStackFunctionInputCollection::FOnInputPinnedChanged& UNiagaraStackFunctionInputCollection::OnInputPinnedChanged()
+{
+	return InputPinnedChangedDelegate;
+}
+
+void UNiagaraStackFunctionInputCollection::RefreshChildrenInternal(const TArray<UNiagaraStackEntry*>& CurrentChildren, TArray<UNiagaraStackEntry*>& NewChildren)
+{
+	TArray<const UEdGraphPin*> InputPins;
+	FNiagaraStackGraphUtilities::GetStackFunctionInputPins(*InputFunctionCallNode, InputPins, FNiagaraStackGraphUtilities::ENiagaraGetStackFunctionInputPinsOptions::ModuleInputsOnly);
+	const UEdGraphSchema_Niagara* NiagaraSchema = GetDefault<UEdGraphSchema_Niagara>();
+	for(const UEdGraphPin* InputPin : InputPins)
+	{
+		UNiagaraStackFunctionInput* Input = nullptr;
+		for (UNiagaraStackEntry* CurrentChild : CurrentChildren)
+		{
+			UNiagaraStackFunctionInput* ChildInput = CastChecked<UNiagaraStackFunctionInput>(CurrentChild);
+			if (ChildInput->GetInputParameterHandle().GetParameterHandleString() == InputPin->PinName)
+			{
+				Input = ChildInput;
+				break;
+			}
+		}
+		if (Input == nullptr)
+		{
+			Input = NewObject<UNiagaraStackFunctionInput>(this);
+			Input->Initialize(GetSystemViewModel(), GetEmitterViewModel(), *StackEditorData, *ModuleNode, *InputFunctionCallNode, InputPin->PinName, NiagaraSchema->PinToTypeDefinition(InputPin));
+			Input->SetItemIndentLevel(DisplayOptions.ChildItemIndentLevel);
+			Input->OnPinnedChanged().AddUObject(this, &UNiagaraStackFunctionInputCollection::ChildPinnedChanged);
+		}
+
+		if (DisplayOptions.ChildFilter.IsBound() == false || DisplayOptions.ChildFilter.Execute(Input))
+		{
+			NewChildren.Add(Input);
+		}
+	}
+}
+
+void UNiagaraStackFunctionInputCollection::ChildPinnedChanged()
+{
+	InputPinnedChangedDelegate.Broadcast();
+}
\ No newline at end of file
diff --git a/Engine/Plugins/FX/Niagara/Source/NiagaraEditor/Private/ViewModels/Stack/NiagaraStackItem.cpp b/Engine/Plugins/FX/Niagara/Source/NiagaraEditor/Private/ViewModels/Stack/NiagaraStackItem.cpp
new file mode 100644
--- /dev/null
+++ b/Engine/Plugins/FX/Niagara/Source/NiagaraEditor/Private/ViewModels/Stack/NiagaraStackItem.cpp
@@ -0,0 +1,19 @@
+// Copyright 1998-2016 Epic Games, Inc. All Rights Reserved.
+
+#include "NiagaraStackItem.h"
+
+void UNiagaraStackItem::Initialize(TSharedRef<FNiagaraSystemViewModel> InSystemViewModel, TSharedRef<FNiagaraEmitterViewModel> InEmitterViewModel, UNiagaraStackEditorData& InStackEditorData)
+{
+	Super::Initialize(InSystemViewModel, InEmitterViewModel);
+	StackEditorData = &InStackEditorData;
+}
+
+void UNiagaraStackItem::SetOnModifiedGroupItems(FOnModifiedGroupItems OnModifiedGroupItems)
+{
+	ModifiedGroupItemsDelegate = OnModifiedGroupItems;
+}
+
+UNiagaraStackEditorData& UNiagaraStackItem::GetStackEditorData() const
+{
+	return *StackEditorData;
+}
diff --git a/Engine/Plugins/FX/Niagara/Source/NiagaraEditor/Private/ViewModels/Stack/NiagaraStackItemExpander.cpp b/Engine/Plugins/FX/Niagara/Source/NiagaraEditor/Private/ViewModels/Stack/NiagaraStackItemExpander.cpp
new file mode 100644
--- /dev/null
+++ b/Engine/Plugins/FX/Niagara/Source/NiagaraEditor/Private/ViewModels/Stack/NiagaraStackItemExpander.cpp
@@ -0,0 +1,29 @@
+// Copyright 1998-2016 Epic Games, Inc. All Rights Reserved.
+
+#include "NiagaraStackItemExpander.h"
+#include "NiagaraStackEditorData.h"
+#include "NiagaraEmitterViewModel.h"
+
+void UNiagaraStackItemExpander::Initialize(TSharedRef<FNiagaraSystemViewModel> InSystemViewModel, TSharedRef<FNiagaraEmitterViewModel> InEmitterViewModel, UNiagaraStackEditorData& InStackEditorData, FString InEmitterEditorDataKey, bool bInIsExpandedDefault)
+{
+	Super::Initialize(InSystemViewModel, InEmitterViewModel);
+	StackEditorData = &InStackEditorData;
+	EmitterEditorDataKey = InEmitterEditorDataKey;
+	bIsExpandedDefault = bInIsExpandedDefault;
+}
+
+void UNiagaraStackItemExpander::SetOnExpnadedChanged(FOnExpandedChanged OnExpandedChanged)
+{
+	ExpandedChangedDelegate = OnExpandedChanged;
+}
+
+bool UNiagaraStackItemExpander::GetIsExpanded() const
+{
+	return StackEditorData->GetStackEntryIsExpanded(EmitterEditorDataKey, bIsExpandedDefault);
+}
+
+void UNiagaraStackItemExpander::ToggleExpanded()
+{
+	StackEditorData->SetStackEntryIsExpanded(EmitterEditorDataKey, !GetIsExpanded());
+	ExpandedChangedDelegate.ExecuteIfBound();
+}
\ No newline at end of file
diff --git a/Engine/Plugins/FX/Niagara/Source/NiagaraEditor/Private/ViewModels/Stack/NiagaraStackItemGroup.cpp b/Engine/Plugins/FX/Niagara/Source/NiagaraEditor/Private/ViewModels/Stack/NiagaraStackItemGroup.cpp
new file mode 100644
--- /dev/null
+++ b/Engine/Plugins/FX/Niagara/Source/NiagaraEditor/Private/ViewModels/Stack/NiagaraStackItemGroup.cpp
@@ -0,0 +1,51 @@
+// Copyright 1998-2016 Epic Games, Inc. All Rights Reserved.
+
+#include "NiagaraStackItemGroup.h"
+#include "NiagaraStackItem.h"
+#include "NiagaraStackEditorData.h"
+
+void UNiagaraStackItemGroup::Initialize(TSharedRef<FNiagaraSystemViewModel> InSystemViewModel, TSharedRef<FNiagaraEmitterViewModel> InEmitterViewModel, UNiagaraStackEditorData& InStackEditorData)
+{
+	Super::Initialize(InSystemViewModel, InEmitterViewModel);
+	StackEditorData = &InStackEditorData;
+}
+
+bool UNiagaraStackItemGroup::GetCanExpand() const
+{
+	return true;
+}
+
+bool UNiagaraStackItemGroup::GetIsExpanded() const
+{
+	return StackEditorData->GetStackEntryIsExpanded(GetDisplayName().ToString(), IsExpandedByDefault());
+}
+
+void UNiagaraStackItemGroup::SetIsExpanded(bool bInExpanded)
+{
+	StackEditorData->SetStackEntryIsExpanded(GetDisplayName().ToString(), bInExpanded);
+}
+
+FName UNiagaraStackItemGroup::GetItemBackgroundName() const
+{
+	return "NiagaraEditor.Stack.Group.BackgroundColor";
+}
+
+FName UNiagaraStackItemGroup::GetTextStyleName() const
+{
+	return "NiagaraEditor.Stack.GroupText";
+}
+
+FText UNiagaraStackItemGroup::GetTooltipText() const 
+{
+	return TooltipText;
+}
+
+void UNiagaraStackItemGroup::SetTooltipText(const FText& InText)
+{
+	TooltipText = InText;
+}
+
+UNiagaraStackEditorData& UNiagaraStackItemGroup::GetStackEditorData() const
+{
+	return *StackEditorData;
+}
\ No newline at end of file
diff --git a/Engine/Plugins/FX/Niagara/Source/NiagaraEditor/Private/ViewModels/Stack/NiagaraStackModuleItem.cpp b/Engine/Plugins/FX/Niagara/Source/NiagaraEditor/Private/ViewModels/Stack/NiagaraStackModuleItem.cpp
new file mode 100644
--- /dev/null
+++ b/Engine/Plugins/FX/Niagara/Source/NiagaraEditor/Private/ViewModels/Stack/NiagaraStackModuleItem.cpp
@@ -0,0 +1,219 @@
+// Copyright 1998-2016 Epic Games, Inc. All Rights Reserved.
+
+#include "NiagaraStackModuleItem.h"
+#include "NiagaraEmitterViewModel.h"
+#include "NiagaraStackFunctionInputCollection.h"
+#include "NiagaraStackFunctionInput.h"
+#include "NiagaraStackModuleItemOutputCollection.h"
+#include "NiagaraStackItemExpander.h"
+#include "NiagaraNodeFunctionCall.h"
+#include "NiagaraNodeParameterMapSet.h"
+#include "NiagaraNodeOutput.h"
+#include "NiagaraEmitterEditorData.h"
+#include "NiagaraStackEditorData.h"
+#include "NiagaraEditorStyle.h"
+#include "NiagaraGraph.h"
+#include "NiagaraStackGraphUtilities.h"
+
+#include "ScopedTransaction.h"
+
+#define LOCTEXT_NAMESPACE "NiagaraStackViewModel"
+
+UNiagaraStackModuleItem::UNiagaraStackModuleItem()
+	: FunctionCallNode(nullptr)
+	, PinnedInputCollection(nullptr)
+	, UnpinnedInputCollection(nullptr)
+	, ModuleExpander(nullptr)
+{
+}
+
+const UNiagaraNodeFunctionCall& UNiagaraStackModuleItem::GetModuleNode() const
+{
+	return *FunctionCallNode;
+}
+
+void UNiagaraStackModuleItem::Initialize(TSharedRef<FNiagaraSystemViewModel> InSystemViewModel, TSharedRef<FNiagaraEmitterViewModel> InEmitterViewModel, UNiagaraStackEditorData& InStackEditorData, UNiagaraNodeFunctionCall& InFunctionCallNode)
+{
+	checkf(FunctionCallNode == nullptr, TEXT("Can not set the node more than once."));
+	Super::Initialize(InSystemViewModel, InEmitterViewModel, InStackEditorData);
+	FunctionCallNode = &InFunctionCallNode;
+}
+
+FText UNiagaraStackModuleItem::GetDisplayName() const
+{
+	if (FunctionCallNode != nullptr)
+	{
+		return FunctionCallNode->GetNodeTitle(ENodeTitleType::ListView);
+	}
+	else
+	{
+		return FText::FromName(NAME_None);
+	}
+}
+
+
+FText UNiagaraStackModuleItem::GetTooltipText() const
+{
+	if (FunctionCallNode != nullptr)
+	{
+		return FunctionCallNode->GetTooltipText();
+	}
+	else
+	{
+		return FText();
+	}
+}
+
+void UNiagaraStackModuleItem::RefreshChildrenInternal(const TArray<UNiagaraStackEntry*>& CurrentChildren, TArray<UNiagaraStackEntry*>& NewChildren)
+{
+	if (PinnedInputCollection == nullptr)
+	{
+		UNiagaraStackFunctionInputCollection::FDisplayOptions DisplayOptions;
+		DisplayOptions.DisplayName = LOCTEXT("PinnedInputsLabel", "Pinned Inputs");
+		DisplayOptions.ChildItemIndentLevel = 0;
+		DisplayOptions.bShouldShowInStack = false;
+		DisplayOptions.ChildFilter = UNiagaraStackFunctionInputCollection::FOnFilterChildren::CreateLambda(
+			[](UNiagaraStackFunctionInput* FunctionInput) { return FunctionInput->GetIsPinned(); });
+
+		TArray<FString> InputParameterHandlePath;
+		PinnedInputCollection = NewObject<UNiagaraStackFunctionInputCollection>(this);
+		PinnedInputCollection->Initialize(GetSystemViewModel(), GetEmitterViewModel(), GetStackEditorData(), *FunctionCallNode, *FunctionCallNode, DisplayOptions);
+		PinnedInputCollection->OnInputPinnedChanged().AddUObject(this, &UNiagaraStackModuleItem::InputPinnedChanged);
+	}
+
+	if (UnpinnedInputCollection == nullptr)
+	{
+		UNiagaraStackFunctionInputCollection::FDisplayOptions DisplayOptions;
+		DisplayOptions.DisplayName = LOCTEXT("InputsLabel", "Inputs");
+		DisplayOptions.ChildItemIndentLevel = 1;
+		DisplayOptions.bShouldShowInStack = true;
+		DisplayOptions.ChildFilter = UNiagaraStackFunctionInputCollection::FOnFilterChildren::CreateLambda(
+			[](UNiagaraStackFunctionInput* FunctionInput) { return FunctionInput->GetIsPinned() == false; });
+
+		TArray<FString> InputParameterHandlePath;
+		UnpinnedInputCollection = NewObject<UNiagaraStackFunctionInputCollection>(this);
+		UnpinnedInputCollection->Initialize(GetSystemViewModel(), GetEmitterViewModel(), GetStackEditorData(), *FunctionCallNode, *FunctionCallNode, DisplayOptions);
+		UnpinnedInputCollection->OnInputPinnedChanged().AddUObject(this, &UNiagaraStackModuleItem::InputPinnedChanged);
+	}
+
+	if (OutputCollection == nullptr)
+	{
+		OutputCollection = NewObject<UNiagaraStackModuleItemOutputCollection>(this);
+		OutputCollection->Initialize(GetSystemViewModel(), GetEmitterViewModel(), *FunctionCallNode);
+		OutputCollection->SetDisplayName(LOCTEXT("OutputsLabel", "Outputs"));
+	}
+
+	FString ModuleEditorDataKey = FNiagaraStackGraphUtilities::GenerateStackModuleEditorDataKey(*FunctionCallNode);
+	if (ModuleExpander == nullptr)
+	{
+		ModuleExpander = NewObject<UNiagaraStackItemExpander>(this);
+		ModuleExpander->Initialize(GetSystemViewModel(), GetEmitterViewModel(), GetStackEditorData(), ModuleEditorDataKey, true);
+		ModuleExpander->SetOnExpnadedChanged(UNiagaraStackItemExpander::FOnExpandedChanged::CreateUObject(this, &UNiagaraStackModuleItem::ModuleExpandedChanged));
+	}
+
+	NewChildren.Add(PinnedInputCollection);
+	if (GetStackEditorData().GetStackEntryIsExpanded(ModuleEditorDataKey, true))
+	{
+		NewChildren.Add(UnpinnedInputCollection);
+		NewChildren.Add(OutputCollection);
+	}
+	NewChildren.Add(ModuleExpander);
+}
+
+void UNiagaraStackModuleItem::InputPinnedChanged()
+{
+	PinnedInputCollection->RefreshChildren();
+	UnpinnedInputCollection->RefreshChildren();
+}
+
+void UNiagaraStackModuleItem::ModuleExpandedChanged()
+{
+	RefreshChildren();
+}
+
+void UNiagaraStackModuleItem::MoveUp()
+{
+	TArray<FNiagaraStackGraphUtilities::FStackNodeGroup> StackNodeGroups;
+	FNiagaraStackGraphUtilities::GetStackNodeGroups(*FunctionCallNode, StackNodeGroups);
+	int32 ModuleStackIndex = StackNodeGroups.IndexOfByPredicate([&](const FNiagaraStackGraphUtilities::FStackNodeGroup& StackNodeGroup) { return StackNodeGroup.EndNode == FunctionCallNode; });
+
+	if (ModuleStackIndex > 1)
+	{
+		FScopedTransaction ScopedTransaction(LOCTEXT("MoveModuleUpTheStackTransaction", "Move module up the stack"));
+
+		FNiagaraStackGraphUtilities::DisconnectStackNodeGroup(StackNodeGroups[ModuleStackIndex], StackNodeGroups[ModuleStackIndex - 1], StackNodeGroups[ModuleStackIndex + 1]);
+		FNiagaraStackGraphUtilities::ConnectStackNodeGroup(StackNodeGroups[ModuleStackIndex], StackNodeGroups[ModuleStackIndex - 2], StackNodeGroups[ModuleStackIndex - 1]);
+
+		FunctionCallNode->GetNiagaraGraph()->NotifyGraphNeedsRecompile();
+		FNiagaraStackGraphUtilities::RelayoutGraph(*FunctionCallNode->GetGraph());
+		ModifiedGroupItemsDelegate.ExecuteIfBound();
+	}
+}
+
+void UNiagaraStackModuleItem::MoveDown()
+{
+	TArray<FNiagaraStackGraphUtilities::FStackNodeGroup> StackNodeGroups;
+	FNiagaraStackGraphUtilities::GetStackNodeGroups(*FunctionCallNode, StackNodeGroups);
+	int32 ModuleStackIndex = StackNodeGroups.IndexOfByPredicate([&](const FNiagaraStackGraphUtilities::FStackNodeGroup& StackNodeGroup) { return StackNodeGroup.EndNode == FunctionCallNode; });
+
+	if (ModuleStackIndex < StackNodeGroups.Num() - 2)
+	{
+		FScopedTransaction ScopedTransaction(LOCTEXT("MoveModuleDownTheStackTransaction", "Move module down the stack"));
+
+		FNiagaraStackGraphUtilities::DisconnectStackNodeGroup(StackNodeGroups[ModuleStackIndex], StackNodeGroups[ModuleStackIndex - 1], StackNodeGroups[ModuleStackIndex + 1]);
+		FNiagaraStackGraphUtilities::ConnectStackNodeGroup(StackNodeGroups[ModuleStackIndex], StackNodeGroups[ModuleStackIndex + 1], StackNodeGroups[ModuleStackIndex + 2]);
+
+		FunctionCallNode->GetNiagaraGraph()->NotifyGraphNeedsRecompile();
+		FNiagaraStackGraphUtilities::RelayoutGraph(*FunctionCallNode->GetGraph());
+		ModifiedGroupItemsDelegate.ExecuteIfBound();
+	}
+}
+
+void UNiagaraStackModuleItem::Delete()
+{
+	TArray<FNiagaraStackGraphUtilities::FStackNodeGroup> StackNodeGroups;
+	FNiagaraStackGraphUtilities::GetStackNodeGroups(*FunctionCallNode, StackNodeGroups);
+	int32 ModuleStackIndex = StackNodeGroups.IndexOfByPredicate([&](const FNiagaraStackGraphUtilities::FStackNodeGroup& StackNodeGroup) { return StackNodeGroup.EndNode == FunctionCallNode; });
+
+	FScopedTransaction ScopedTransaction(LOCTEXT("RemoveAModuleFromTheStack", "Remove a module from the stack"));
+
+	FNiagaraStackGraphUtilities::DisconnectStackNodeGroup(StackNodeGroups[ModuleStackIndex], StackNodeGroups[ModuleStackIndex - 1], StackNodeGroups[ModuleStackIndex + 1]);
+
+	FNiagaraStackGraphUtilities::FStackNodeGroup ModuleGroup = StackNodeGroups[ModuleStackIndex];
+
+	TArray<UNiagaraNode*> NodesToRemove;
+	TArray<UNiagaraNode*> NodesToCheck;
+	NodesToCheck.Add(ModuleGroup.EndNode);
+	while (NodesToCheck.Num() > 0)
+	{
+		UNiagaraNode* NodeToRemove = NodesToCheck[0];
+		NodesToCheck.RemoveAt(0);
+		NodesToRemove.Add(NodeToRemove);
+
+		TArray<UEdGraphPin*> InputPins;
+		NodeToRemove->GetInputPins(InputPins);
+		for (UEdGraphPin* InputPin : InputPins)
+		{
+			if (InputPin->LinkedTo.Num() == 1)
+			{
+				UNiagaraNode* LinkedNode = Cast<UNiagaraNode>(InputPin->LinkedTo[0]->GetOwningNode());
+				if (LinkedNode != nullptr)
+				{
+					NodesToCheck.Add(LinkedNode);
+				}
+			}
+		}
+	}
+
+	UNiagaraGraph* Graph = FunctionCallNode->GetNiagaraGraph();
+	for (UNiagaraNode* NodeToRemove : NodesToRemove)
+	{
+		Graph->RemoveNode(NodeToRemove);
+	}
+
+	Graph->NotifyGraphNeedsRecompile();
+	FNiagaraStackGraphUtilities::RelayoutGraph(*FunctionCallNode->GetGraph());
+	ModifiedGroupItemsDelegate.ExecuteIfBound();
+}
+
+#undef LOCTEXT_NAMESPACE
diff --git a/Engine/Plugins/FX/Niagara/Source/NiagaraEditor/Private/ViewModels/Stack/NiagaraStackModuleItemOutput.cpp b/Engine/Plugins/FX/Niagara/Source/NiagaraEditor/Private/ViewModels/Stack/NiagaraStackModuleItemOutput.cpp
new file mode 100644
--- /dev/null
+++ b/Engine/Plugins/FX/Niagara/Source/NiagaraEditor/Private/ViewModels/Stack/NiagaraStackModuleItemOutput.cpp
@@ -0,0 +1,53 @@
+// Copyright 1998-2016 Epic Games, Inc. All Rights Reserved.
+
+#include "NiagaraStackModuleItemOutput.h"
+#include "NiagaraNodeFunctionCall.h"
+
+#define LOCTEXT_NAMESPACE "NiagaraStackViewModel"
+
+UNiagaraStackModuleItemOutput::UNiagaraStackModuleItemOutput()
+	: FunctionCallNode(nullptr)
+{
+}
+
+int32 UNiagaraStackModuleItemOutput::GetItemIndentLevel() const
+{
+	return 1;
+}
+
+void UNiagaraStackModuleItemOutput::Initialize(TSharedRef<FNiagaraSystemViewModel> InSystemViewModel, TSharedRef<FNiagaraEmitterViewModel> InEmitterViewModel, UNiagaraNodeFunctionCall& InFunctionCallNode, FString InOutputParameterHandle)
+{
+	checkf(FunctionCallNode.Get() == nullptr, TEXT("Can only set the Output once."));
+	Super::Initialize(InSystemViewModel, InEmitterViewModel);
+	FunctionCallNode = &InFunctionCallNode;
+
+	OutputParameterHandle = FNiagaraParameterHandle(InOutputParameterHandle);
+	DisplayName = FText::FromString(OutputParameterHandle.GetName());
+}
+
+FText UNiagaraStackModuleItemOutput::GetDisplayName() const
+{
+	return DisplayName;
+}
+
+FName UNiagaraStackModuleItemOutput::GetTextStyleName() const
+{
+	return "NiagaraEditor.Stack.ParameterText";
+}
+
+bool UNiagaraStackModuleItemOutput::GetCanExpand() const
+{
+	return true;
+}
+
+const FNiagaraParameterHandle& UNiagaraStackModuleItemOutput::GetOutputParameterHandle() const
+{
+	return OutputParameterHandle;
+}
+
+FText UNiagaraStackModuleItemOutput::GetOutputParameterHandleText() const
+{
+	return FText::FromString(OutputParameterHandle.GetParameterHandleString());
+}
+
+#undef LOCTEXT_NAMESPACE
diff --git a/Engine/Plugins/FX/Niagara/Source/NiagaraEditor/Private/ViewModels/Stack/NiagaraStackModuleItemOutputCollection.cpp b/Engine/Plugins/FX/Niagara/Source/NiagaraEditor/Private/ViewModels/Stack/NiagaraStackModuleItemOutputCollection.cpp
new file mode 100644
--- /dev/null
+++ b/Engine/Plugins/FX/Niagara/Source/NiagaraEditor/Private/ViewModels/Stack/NiagaraStackModuleItemOutputCollection.cpp
@@ -0,0 +1,88 @@
+// Copyright 1998-2016 Epic Games, Inc. All Rights Reserved.
+
+#include "NiagaraStackModuleItemOutputCollection.h"
+#include "NiagaraStackModuleItemOutput.h"
+#include "NiagaraNodeFunctionCall.h"
+#include "NiagaraNodeParameterMapSet.h"
+#include "NiagaraEmitterEditorData.h"
+#include "NiagaraStackGraphUtilities.h"
+
+#include "EdGraph/EdGraphPin.h"
+
+UNiagaraStackModuleItemOutputCollection::UNiagaraStackModuleItemOutputCollection()
+	: FunctionCallNode(nullptr)
+{
+}
+
+void UNiagaraStackModuleItemOutputCollection::Initialize(TSharedRef<FNiagaraSystemViewModel> InSystemViewModel, TSharedRef<FNiagaraEmitterViewModel> InEmitterViewModel, UNiagaraNodeFunctionCall& InFunctionCallNode)
+{
+	checkf(FunctionCallNode == nullptr, TEXT("Can not set the node more than once."));
+	Super::Initialize(InSystemViewModel, InEmitterViewModel);
+	FunctionCallNode = &InFunctionCallNode;
+}
+
+FText UNiagaraStackModuleItemOutputCollection::GetDisplayName() const
+{
+	return DisplayName;
+}
+
+FName UNiagaraStackModuleItemOutputCollection::GetTextStyleName() const
+{
+	return "NiagaraEditor.Stack.ParameterCollectionText";
+}
+
+void UNiagaraStackModuleItemOutputCollection::SetDisplayName(FText InDisplayName)
+{
+	DisplayName = InDisplayName;
+}
+
+bool UNiagaraStackModuleItemOutputCollection::GetCanExpand() const
+{
+	return true;
+}
+
+bool UNiagaraStackModuleItemOutputCollection::IsExpandedByDefault() const
+{
+	return false;
+}
+
+void UNiagaraStackModuleItemOutputCollection::RefreshChildrenInternal(const TArray<UNiagaraStackEntry*>& CurrentChildren, TArray<UNiagaraStackEntry*>& NewChildren)
+{
+	UEdGraphPin* OutputParameterMapPin = FNiagaraStackGraphUtilities::GetParameterMapOutputPin(*FunctionCallNode);
+	if (ensureMsgf(OutputParameterMapPin != nullptr, TEXT("Invalid Stack Graph - Function call node has no output pin.")))
+	{
+		FNiagaraParameterMapHistoryBuilder Builder;
+		FunctionCallNode->BuildParameterMapHistory(Builder, false);
+		check(Builder.Histories.Num() == 1);
+		for (int32 i = 0; i < Builder.Histories[0].Variables.Num(); i++)
+		{
+			FNiagaraVariable& Variable = Builder.Histories[0].Variables[i];
+			TArray<const UEdGraphPin*>& WriteHistory = Builder.Histories[0].PerVariableWriteHistory[i];
+
+			for (const UEdGraphPin* WritePin : WriteHistory)
+			{
+				if (Cast<UNiagaraNodeParameterMapSet>(WritePin->GetOwningNode()) != nullptr)
+				{
+					UNiagaraStackModuleItemOutput* Output = nullptr;
+					for (UNiagaraStackEntry* CurrentChild : CurrentChildren)
+					{
+						UNiagaraStackModuleItemOutput* ChildOutput = CastChecked<UNiagaraStackModuleItemOutput>(CurrentChild);
+						if (ChildOutput->GetOutputParameterHandle().GetParameterHandleString() == Variable.GetName().ToString())
+						{
+							Output = ChildOutput;
+							break;
+						}
+					}
+					if (Output == nullptr)
+					{
+						Output = NewObject<UNiagaraStackModuleItemOutput>(this);
+						Output->Initialize(GetSystemViewModel(), GetEmitterViewModel(), *FunctionCallNode, Variable.GetName().ToString());
+					}
+
+					NewChildren.Add(Output);
+					break;
+				}
+			}
+		}
+	}
+}
\ No newline at end of file
diff --git a/Engine/Plugins/FX/Niagara/Source/NiagaraEditor/Private/ViewModels/Stack/NiagaraStackObject.cpp b/Engine/Plugins/FX/Niagara/Source/NiagaraEditor/Private/ViewModels/Stack/NiagaraStackObject.cpp
new file mode 100644
--- /dev/null
+++ b/Engine/Plugins/FX/Niagara/Source/NiagaraEditor/Private/ViewModels/Stack/NiagaraStackObject.cpp
@@ -0,0 +1,36 @@
+// Copyright 1998-2016 Epic Games, Inc. All Rights Reserved.
+
+#include "NiagaraStackObject.h"
+
+
+UNiagaraStackObject::UNiagaraStackObject()
+	: Object(nullptr)
+{
+}
+
+void UNiagaraStackObject::Initialize(TSharedRef<FNiagaraSystemViewModel> InSystemViewModel, TSharedRef<FNiagaraEmitterViewModel> InEmitterViewModel, UObject* InObject)
+{
+	checkf(Object == nullptr, TEXT("Can only initialize once."));
+	Super::Initialize(InSystemViewModel, InEmitterViewModel);
+	Object = InObject;
+}
+
+UObject* UNiagaraStackObject::GetObject()
+{
+	return Object;
+}
+
+int32 UNiagaraStackObject::GetItemIndentLevel() const
+{
+	return ItemIndentLevel;
+}
+
+void UNiagaraStackObject::SetItemIndentLevel(int32 InItemIndentLevel)
+{
+	ItemIndentLevel = InItemIndentLevel;
+}
+
+void UNiagaraStackObject::NotifyObjectChanged()
+{
+	OnDataObjectModified().Broadcast(Object);
+}
\ No newline at end of file
diff --git a/Engine/Plugins/FX/Niagara/Source/NiagaraEditor/Private/ViewModels/Stack/NiagaraStackParameterStoreEntry.cpp b/Engine/Plugins/FX/Niagara/Source/NiagaraEditor/Private/ViewModels/Stack/NiagaraStackParameterStoreEntry.cpp
new file mode 100644
--- /dev/null
+++ b/Engine/Plugins/FX/Niagara/Source/NiagaraEditor/Private/ViewModels/Stack/NiagaraStackParameterStoreEntry.cpp
@@ -0,0 +1,263 @@
+// Copyright 1998-2016 Epic Games, Inc. All Rights Reserved.
+
+#include "NiagaraStackParameterStoreEntry.h"
+#include "NiagaraScriptSource.h"
+#include "NiagaraGraph.h"
+#include "EdGraphSchema_Niagara.h"
+#include "NiagaraEditorUtilities.h"
+#include "NiagaraSystemViewModel.h"
+#include "NiagaraSystemScriptViewModel.h"
+#include "NiagaraEmitterViewModel.h"
+#include "NiagaraScriptGraphViewModel.h"
+#include "NiagaraStackObject.h"
+
+#include "ScopedTransaction.h"
+#include "Editor.h"
+#include "UObject/StructOnScope.h"
+#include "Modules/ModuleManager.h"
+#include "AssetRegistryModule.h"
+#include "ARFilter.h"
+#include "EdGraph/EdGraphPin.h"
+
+#define LOCTEXT_NAMESPACE "UNiagaraStackParameterStoreEntry"
+UNiagaraStackParameterStoreEntry::UNiagaraStackParameterStoreEntry()
+	: ValueObjectEntry(nullptr)
+	, ItemIndentLevel(0)
+{
+}
+
+void UNiagaraStackParameterStoreEntry::BeginDestroy()
+{
+	Super::BeginDestroy();
+}
+int32 UNiagaraStackParameterStoreEntry::GetItemIndentLevel() const
+{
+	return ItemIndentLevel;
+}
+
+void UNiagaraStackParameterStoreEntry::SetItemIndentLevel(int32 InItemIndentLevel)
+{
+	ItemIndentLevel = InItemIndentLevel;
+}
+
+void UNiagaraStackParameterStoreEntry::Initialize(
+	TSharedRef<FNiagaraSystemViewModel> InSystemViewModel,
+	TSharedRef<FNiagaraEmitterViewModel> InEmitterViewModel,
+	UNiagaraStackEditorData& InStackEditorData,
+	UObject* InOwner,
+	FNiagaraParameterStore* InParameterStore,
+	FString InInputParameterHandle,
+	FNiagaraTypeDefinition InInputType)
+{
+	Super::Initialize(InSystemViewModel, InEmitterViewModel);
+	StackEditorData = &InStackEditorData;
+	DisplayName = FText::FromString(InInputParameterHandle);
+	ParameterName = *InInputParameterHandle;
+	InputType = InInputType;
+	Owner = InOwner;
+	ParameterStore = InParameterStore;
+}
+
+const FNiagaraTypeDefinition& UNiagaraStackParameterStoreEntry::GetInputType() const
+{
+	return InputType;
+}
+
+void UNiagaraStackParameterStoreEntry::RefreshChildrenInternal(const TArray<UNiagaraStackEntry*>& CurrentChildren, TArray<UNiagaraStackEntry*>& NewChildren)
+{
+	RefreshValueAndHandle();
+
+	if (ValueObject != nullptr)
+	{
+		if(ValueObjectEntry == nullptr ||
+			ValueObjectEntry->GetObject() != ValueObject)
+		{
+			ValueObjectEntry = NewObject<UNiagaraStackObject>(this);
+			ValueObjectEntry->Initialize(GetSystemViewModel(), GetEmitterViewModel(), ValueObject);
+			ValueObjectEntry->SetItemIndentLevel(ItemIndentLevel + 1);
+		}
+		NewChildren.Add(ValueObjectEntry);
+	}
+	else
+	{
+		ValueObjectEntry = nullptr;
+	}
+}
+
+void UNiagaraStackParameterStoreEntry::RefreshValueAndHandle()
+{
+	TSharedPtr<FNiagaraVariable> CurrentValueVariable = GetCurrentValueVariable();
+	if (CurrentValueVariable.IsValid() && CurrentValueVariable->GetType() == InputType && CurrentValueVariable->IsDataAllocated())
+	{
+		if (LocalValueStruct.IsValid() == false || LocalValueStruct->GetStruct() != CurrentValueVariable->GetType().GetScriptStruct())
+		{
+			LocalValueStruct = MakeShared<FStructOnScope>(InputType.GetScriptStruct());
+		}
+		CurrentValueVariable->CopyTo(LocalValueStruct->GetStructMemory());
+	}
+	else
+	{
+		LocalValueStruct.Reset();
+	}
+
+	ValueObject = GetCurrentValueObject();
+
+	ValueChangedDelegate.Broadcast();
+}
+
+FText UNiagaraStackParameterStoreEntry::GetDisplayName() const
+{
+	return DisplayName;
+}
+
+FName UNiagaraStackParameterStoreEntry::GetTextStyleName() const
+{
+	return "NiagaraEditor.Stack.ParameterText";
+}
+
+bool UNiagaraStackParameterStoreEntry::GetCanExpand() const
+{
+	return true;
+}
+
+TSharedPtr<FStructOnScope> UNiagaraStackParameterStoreEntry::GetValueStruct()
+{
+	return LocalValueStruct;
+}
+
+UNiagaraDataInterface* UNiagaraStackParameterStoreEntry::GetValueObject()
+{
+	return ValueObject;
+}
+
+void UNiagaraStackParameterStoreEntry::NotifyBeginValueChange()
+{
+	GEditor->BeginTransaction(LOCTEXT("BeginEditModuleInputValue", "Edit module input value."));
+	Owner->Modify();
+}
+
+void UNiagaraStackParameterStoreEntry::NotifyEndValueChange()
+{
+	if (GEditor->IsTransactionActive())
+	{
+		GEditor->EndTransaction();
+	}
+}
+
+void UNiagaraStackParameterStoreEntry::NotifyValueChanged()
+{
+	TSharedPtr<FNiagaraVariable> CurrentValue = GetCurrentValueVariable();
+	if ((CurrentValue.IsValid() && LocalValueStruct.IsValid()) && FNiagaraEditorUtilities::DataMatches(*CurrentValue.Get(), *LocalValueStruct.Get()))
+	{
+		return;
+	}
+	else if ((CurrentValue.IsValid() && LocalValueStruct.IsValid()))
+	{
+		FNiagaraVariable DefaultVariable(InputType, ParameterName);
+		ParameterStore->SetParameterData(LocalValueStruct->GetStructMemory(), DefaultVariable);
+		GetSystemViewModel()->ResetSystem();
+	}
+}
+
+bool UNiagaraStackParameterStoreEntry::CanReset() const
+{
+	return true;
+}
+
+void UNiagaraStackParameterStoreEntry::Reset()
+{
+	if (InputType.GetClass() == nullptr)
+	{		
+		// For struct inputs the override pin and anything attached to it should be removed.
+		FScopedTransaction ScopedTransaction(LOCTEXT("ResetInputStructTransaction", "Reset the inputs value to default."));
+	}
+	else
+	{
+		// For object types make sure the override is setup to an input which matches the default object.
+		FScopedTransaction ScopedTransaction(LOCTEXT("ResetInputObjectTransaction", "Reset the inputs data interface object to default."));
+	}
+	RefreshChildren();
+}
+
+bool UNiagaraStackParameterStoreEntry::CanRenameInput() const
+{
+	return false;
+}
+
+bool UNiagaraStackParameterStoreEntry::GetIsRenamePending() const
+{
+	return CanRenameInput() /*&& StackEditorData->GetModuleInputIsRenamePending(StackEditorDataKey)*/;
+}
+
+void UNiagaraStackParameterStoreEntry::SetIsRenamePending(bool bIsRenamePending)
+{
+	if (CanRenameInput())
+	{
+		//StackEditorData->SetModuleInputIsRenamePending(StackEditorDataKey, bIsRenamePending);
+	}
+}
+
+void UNiagaraStackParameterStoreEntry::RenameInput(FString NewName)
+{
+	/*
+	if (OwningAssignmentNode.IsValid() && InputParameterHandlePath.Num() == 1 && InputParameterHandle.GetName() != NewName)
+	{
+		bool bIsCurrentlyPinned = GetIsPinned();
+		bool bIsCurrentlyExpanded = StackEditorData->GetStackEntryIsExpanded(FNiagaraStackGraphUtilities::GenerateStackModuleEditorDataKey(*OwningAssignmentNode), false);
+
+		FNiagaraParameterHandle TargetHandle(OwningAssignmentNode->AssignmentTarget.GetName().ToString());
+		FNiagaraParameterHandle RenamedTargetHandle(TargetHandle.GetNamespace(), NewName);
+		OwningAssignmentNode->AssignmentTarget.SetName(*RenamedTargetHandle.GetParameterHandleString());
+		OwningAssignmentNode->RefreshFromExternalChanges();
+
+		InputParameterHandle = FNiagaraParameterHandle(InputParameterHandle.GetNamespace(), NewName);
+		InputParameterHandlePath.Empty();
+		InputParameterHandlePath.Add(InputParameterHandle);
+		AliasedInputParameterHandle = FNiagaraParameterHandle::CreateAliasedModuleParameterHandle(InputParameterHandle, OwningAssignmentNode.Get());
+		DisplayName = FText::FromString(InputParameterHandle.GetName());
+
+		UEdGraphPin* OverridePin = GetOverridePin();
+		if (OverridePin != nullptr)
+		{
+			OverridePin->PinName = AliasedInputParameterHandle.GetParameterHandleString();
+		}
+
+		StackEditorDataKey = FNiagaraStackGraphUtilities::GenerateStackFunctionInputEditorDataKey(*OwningFunctionCallNode.Get(), InputParameterHandle);
+		StackEditorData->SetModuleInputIsPinned(StackEditorDataKey, bIsCurrentlyPinned);
+		StackEditorData->SetStackEntryIsExpanded(FNiagaraStackGraphUtilities::GenerateStackModuleEditorDataKey(*OwningAssignmentNode), bIsCurrentlyExpanded);
+
+		CastChecked<UNiagaraGraph>(OwningAssignmentNode->GetGraph())->NotifyGraphNeedsRecompile();
+	}
+	*/
+}
+
+UNiagaraStackParameterStoreEntry::FOnValueChanged& UNiagaraStackParameterStoreEntry::OnValueChanged()
+{
+	return ValueChangedDelegate;
+}
+
+TSharedPtr<FNiagaraVariable> UNiagaraStackParameterStoreEntry::GetCurrentValueVariable()
+{
+	if (InputType.GetClass() == nullptr)
+	{
+		const UEdGraphSchema_Niagara* NiagaraSchema = GetDefault<UEdGraphSchema_Niagara>();
+		
+		FNiagaraVariable DefaultVariable(InputType, ParameterName);
+		const uint8* Data = ParameterStore->GetParameterData(DefaultVariable);
+		DefaultVariable.SetData(Data);
+		return MakeShared<FNiagaraVariable>(DefaultVariable);
+	}
+	return TSharedPtr<FNiagaraVariable>();
+}
+
+UNiagaraDataInterface* UNiagaraStackParameterStoreEntry::GetCurrentValueObject()
+{
+	if (InputType.GetClass() != nullptr)
+	{
+		FNiagaraVariable DefaultVariable(InputType, ParameterName);
+		return ParameterStore->GetDataInterface(DefaultVariable);
+	}
+	return nullptr;
+}
+
+#undef LOCTEXT_NAMESPACE
diff --git a/Engine/Plugins/FX/Niagara/Source/NiagaraEditor/Private/ViewModels/Stack/NiagaraStackParameterStoreGroup.cpp b/Engine/Plugins/FX/Niagara/Source/NiagaraEditor/Private/ViewModels/Stack/NiagaraStackParameterStoreGroup.cpp
new file mode 100644
--- /dev/null
+++ b/Engine/Plugins/FX/Niagara/Source/NiagaraEditor/Private/ViewModels/Stack/NiagaraStackParameterStoreGroup.cpp
@@ -0,0 +1,71 @@
+// Copyright 1998-2016 Epic Games, Inc. All Rights Reserved.
+
+#include "NiagaraStackParameterStoreGroup.h"
+#include "NiagaraStackModuleItem.h"
+#include "NiagaraStackAddScriptModuleItem.h"
+#include "NiagaraScriptViewModel.h"
+#include "NiagaraScriptGraphViewModel.h"
+#include "NiagaraGraph.h"
+#include "NiagaraNodeOutput.h"
+#include "NiagaraNodeFunctionCall.h"
+#include "EdGraphSchema_Niagara.h"
+#include "NiagaraStackErrorItem.h"
+#include "Internationalization.h"
+#include "NiagaraStackParameterStoreEntry.h"
+
+
+#define LOCTEXT_NAMESPACE "UNiagaraStackParameterStoreGroup"
+
+
+void UNiagaraStackParameterStoreGroup::Initialize(
+	TSharedRef<FNiagaraSystemViewModel> InSystemViewModel,
+	TSharedRef<FNiagaraEmitterViewModel> InEmitterViewModel,
+	UNiagaraStackEditorData& InStackEditorData,
+	TSharedRef<FNiagaraScriptViewModel> InScriptViewModel,
+	UObject* InOwner,
+	FNiagaraParameterStore* InParameterStore)
+{
+	Super::Initialize(InSystemViewModel, InEmitterViewModel, InStackEditorData);
+	Owner = InOwner;
+	ParameterStore = InParameterStore;
+}
+
+FText UNiagaraStackParameterStoreGroup::GetDisplayName() const
+{
+	return DisplayName;
+}
+
+void UNiagaraStackParameterStoreGroup::SetDisplayName(FText InDisplayName)
+{
+	DisplayName = InDisplayName;
+}
+
+void UNiagaraStackParameterStoreGroup::RefreshChildrenInternal(const TArray<UNiagaraStackEntry*>& CurrentChildren, TArray<UNiagaraStackEntry*>& NewChildren)
+{
+	if (ParameterStore != nullptr)
+	{
+		TArray<FNiagaraVariable> Variables;
+		ParameterStore->GetParameters(Variables);
+
+		for (FNiagaraVariable& Var : Variables)
+		{
+			UNiagaraStackParameterStoreEntry* ValueObjectEntry = NewObject<UNiagaraStackParameterStoreEntry>(this);
+			ValueObjectEntry->Initialize(GetSystemViewModel(), GetEmitterViewModel(), GetStackEditorData(), Owner, ParameterStore, Var.GetName().ToString(), Var.GetType());
+			ValueObjectEntry->SetItemIndentLevel(1);
+			NewChildren.Add(ValueObjectEntry);
+		}
+	}
+}
+
+void UNiagaraStackParameterStoreGroup::ItemAdded()
+{
+	RefreshChildren();
+}
+
+void UNiagaraStackParameterStoreGroup::ChildModifiedGroupItems()
+{
+	RefreshChildren();
+}
+
+
+#undef LOCTEXT_NAMESPACE
\ No newline at end of file
diff --git a/Engine/Plugins/FX/Niagara/Source/NiagaraEditor/Private/ViewModels/Stack/NiagaraStackRenderItemGroup.cpp b/Engine/Plugins/FX/Niagara/Source/NiagaraEditor/Private/ViewModels/Stack/NiagaraStackRenderItemGroup.cpp
new file mode 100644
--- /dev/null
+++ b/Engine/Plugins/FX/Niagara/Source/NiagaraEditor/Private/ViewModels/Stack/NiagaraStackRenderItemGroup.cpp
@@ -0,0 +1,60 @@
+// Copyright 1998-2016 Epic Games, Inc. All Rights Reserved.
+
+#include "NiagaraStackRenderItemGroup.h"
+#include "NiagaraStackRendererItem.h"
+#include "NiagaraStackSpacer.h"
+#include "NiagaraStackAddRendererItem.h"
+#include "NiagaraEmitterViewModel.h"
+#include "NiagaraEmitter.h"
+#include "NiagaraEmitterEditorData.h"
+
+FText UNiagaraStackRenderItemGroup::GetDisplayName() const
+{
+	return DisplayName;
+}
+
+void UNiagaraStackRenderItemGroup::SetDisplayName(FText InDisplayName)
+{
+	DisplayName = InDisplayName;
+}
+
+void UNiagaraStackRenderItemGroup::RefreshChildrenInternal(const TArray<UNiagaraStackEntry*>& CurrentChildren, TArray<UNiagaraStackEntry*>& NewChildren)
+{
+	int32 RendererIndex = 0;
+	for (UNiagaraRendererProperties* RendererProperties : GetEmitterViewModel()->GetEmitter()->RendererProperties)
+	{
+		UNiagaraStackRendererItem* RendererItem = FindCurrentChildOfTypeByPredicate<UNiagaraStackRendererItem>(CurrentChildren,
+			[=](UNiagaraStackRendererItem* CurrentRendererItem) { return CurrentRendererItem->GetRendererProperties() == RendererProperties; });
+
+		if (RendererItem == nullptr)
+		{
+			RendererItem = NewObject<UNiagaraStackRendererItem>(this);
+			RendererItem->Initialize(GetSystemViewModel(), GetEmitterViewModel(), GetStackEditorData(), RendererProperties);
+			RendererItem->SetOnModifiedGroupItems(UNiagaraStackItem::FOnModifiedGroupItems::CreateUObject(this, &UNiagaraStackRenderItemGroup::ChildModifiedGroupItems));
+		}
+
+		FName RendererSpacerKey = *FString::Printf(TEXT("Renderer%i"), RendererIndex);
+		UNiagaraStackSpacer* RendererSpacer = FindCurrentChildOfTypeByPredicate<UNiagaraStackSpacer>(CurrentChildren,
+			[=](UNiagaraStackSpacer* CurrentRendererSpacer) { return CurrentRendererSpacer->GetSpacerKey() == RendererSpacerKey; });
+
+		if (RendererSpacer == nullptr)
+		{
+			RendererSpacer = NewObject<UNiagaraStackSpacer>(this);
+			RendererSpacer->Initialize(GetSystemViewModel(), GetEmitterViewModel(), RendererSpacerKey);
+		}
+
+		NewChildren.Add(RendererItem);
+		NewChildren.Add(RendererSpacer);
+	}
+
+	UNiagaraStackAddRendererItem* AddRendererItem = NewObject<UNiagaraStackAddRendererItem>(this);
+	AddRendererItem->Initialize(GetSystemViewModel(), GetEmitterViewModel());
+	AddRendererItem->SetOnItemAdded(UNiagaraStackAddRendererItem::FOnItemAdded::CreateUObject(this, &UNiagaraStackRenderItemGroup::ChildModifiedGroupItems));
+	NewChildren.Add(AddRendererItem);
+}
+
+void UNiagaraStackRenderItemGroup::ChildModifiedGroupItems()
+{
+	RefreshChildren();
+}
+
diff --git a/Engine/Plugins/FX/Niagara/Source/NiagaraEditor/Private/ViewModels/Stack/NiagaraStackRendererItem.cpp b/Engine/Plugins/FX/Niagara/Source/NiagaraEditor/Private/ViewModels/Stack/NiagaraStackRendererItem.cpp
new file mode 100644
--- /dev/null
+++ b/Engine/Plugins/FX/Niagara/Source/NiagaraEditor/Private/ViewModels/Stack/NiagaraStackRendererItem.cpp
@@ -0,0 +1,223 @@
+// Copyright 1998-2016 Epic Games, Inc. All Rights Reserved.
+
+#include "NiagaraStackRendererItem.h"
+#include "NiagaraStackObject.h"
+#include "NiagaraStackItemExpander.h"
+#include "NiagaraEmitter.h"
+#include "NiagaraStackEditorData.h"
+#include "NiagaraRendererProperties.h"
+#include "NiagaraScript.h"
+#include "NiagaraEmitterViewModel.h"
+#include "NiagaraScriptViewModel.h"
+#include "Internationalization.h"
+#include "NiagaraNodeAssignment.h"
+#include "NiagaraNodeOutput.h"
+#include "NiagaraConstants.h"
+#include "NiagaraScriptGraphViewModel.h"
+#include "NiagaraGraph.h"
+#include "ScopedTransaction.h"
+#include "NiagaraEmitter.h"
+#include "NiagaraScriptSource.h"
+#include "ScopedTransaction.h"
+#include "NiagaraStackErrorItem.h"
+#include "Framework/Notifications/NotificationManager.h"
+#include "Widgets/Notifications/SNotificationList.h"
+#include "NiagaraStackGraphUtilities.h"
+
+#define LOCTEXT_NAMESPACE "UNiagaraStackRendererItem"
+
+UNiagaraStackRendererItem::UNiagaraStackRendererItem()
+	: RendererProperties(nullptr)
+	, RendererObject(nullptr)
+{
+}
+
+void UNiagaraStackRendererItem::Initialize(TSharedRef<FNiagaraSystemViewModel> InSystemViewModel, TSharedRef<FNiagaraEmitterViewModel> InEmitterViewModel, UNiagaraStackEditorData& InStackEditorData, UNiagaraRendererProperties* InRendererProperties)
+{
+	checkf(RendererProperties == nullptr, TEXT("Can not initialize more than once."));
+	Super::Initialize(InSystemViewModel, InEmitterViewModel, InStackEditorData);
+	RendererProperties = InRendererProperties;
+}
+
+TArray<FNiagaraVariable> UNiagaraStackRendererItem::GetMissingVariables(UNiagaraRendererProperties* RendererProperties, UNiagaraEmitter* Emitter)
+{
+	TArray<FNiagaraVariable> MissingAttributes;
+	const TArray<FNiagaraVariable>& RequiredAttrs = RendererProperties->GetRequiredAttributes();
+	const UNiagaraScript* Script = Emitter->SpawnScriptProps.Script;
+	if (Script != nullptr)
+	{
+		MissingAttributes.Empty();
+		for (FNiagaraVariable Attr : RequiredAttrs)
+		{
+			FNiagaraVariable OriginalAttr = Attr;
+			// TODO .. should we always be namespaced?
+			FString AttrName = Attr.GetName().ToString();
+			if (AttrName.RemoveFromStart(TEXT("Particles.")))
+			{
+				Attr.SetName(*AttrName);
+			}
+
+			bool ContainsVar = Script->Attributes.ContainsByPredicate([&Attr](const FNiagaraVariable& Var) { return Var.GetName() == Attr.GetName(); });
+			if (!ContainsVar)
+			{
+				MissingAttributes.Add(OriginalAttr);
+			}
+		}
+	}
+	return MissingAttributes;
+}
+
+bool UNiagaraStackRendererItem::AddMissingVariable(UNiagaraEmitter* Emitter, const FNiagaraVariable& Variable)
+{
+	UNiagaraScript* Script = Emitter->SpawnScriptProps.Script;
+	if (!Script)
+	{
+		return false;
+	}
+	UNiagaraScriptSource* Source = Cast<UNiagaraScriptSource>(Script->GetSource());
+	if (!Source)
+	{
+		return false;
+	}
+
+	UNiagaraGraph* Graph = Source->NodeGraph;
+	if (!Graph)
+	{
+		return false;
+	}
+
+	UNiagaraNodeOutput* OutputNode = Graph->FindOutputNode(ENiagaraScriptUsage::ParticleSpawnScript);
+	if (!OutputNode)
+	{
+		return false;
+	}
+
+	FScopedTransaction ScopedTransaction(LOCTEXT("FixRendererError", "Fixing rendering module error: Add Attribute"));
+	Graph->Modify();
+
+	FGraphNodeCreator<UNiagaraNodeAssignment> NodeBuilder(*Graph);
+	UNiagaraNodeAssignment* NewAssignmentNode = NodeBuilder.CreateNode();
+	NewAssignmentNode->AssignmentTarget = Variable;
+	FString VarDefaultValue = FNiagaraConstants::GetAttributeDefaultValue(NewAssignmentNode->AssignmentTarget);
+	NewAssignmentNode->AssignmentDefaultValue = VarDefaultValue;
+	NodeBuilder.Finalize();
+
+	TArray<FNiagaraStackGraphUtilities::FStackNodeGroup> StackNodeGroups;
+	FNiagaraStackGraphUtilities::GetStackNodeGroups(*OutputNode, StackNodeGroups);
+
+	FNiagaraStackGraphUtilities::FStackNodeGroup AssignmentGroup;
+	AssignmentGroup.StartNodes.Add(NewAssignmentNode);
+	AssignmentGroup.EndNode = NewAssignmentNode;
+
+	FNiagaraStackGraphUtilities::FStackNodeGroup& OutputGroup = StackNodeGroups[StackNodeGroups.Num() - 1];
+	FNiagaraStackGraphUtilities::FStackNodeGroup& OutputGroupPrevious = StackNodeGroups[StackNodeGroups.Num() - 2];
+	FNiagaraStackGraphUtilities::ConnectStackNodeGroup(AssignmentGroup, OutputGroupPrevious, OutputGroup);
+
+	FNiagaraStackGraphUtilities::RelayoutGraph(*Graph);
+	return true;
+}
+
+UNiagaraRendererProperties* UNiagaraStackRendererItem::GetRendererProperties()
+{
+	return RendererProperties;
+}
+
+FText UNiagaraStackRendererItem::GetDisplayName() const
+{
+	if (RendererProperties != nullptr)
+	{
+		return FText::FromString(RendererProperties->GetClass()->GetName());
+	}
+	else
+	{
+		return FText::FromName(NAME_None);
+	}
+}
+
+void UNiagaraStackRendererItem::Delete()
+{
+	const FScopedTransaction Transaction(LOCTEXT("DeleteRenderer", "Delete Renderer"));
+
+	UNiagaraEmitter* Emitter = GetEmitterViewModel()->GetEmitter();
+	Emitter->Modify();
+	Emitter->RendererProperties.Remove(RendererProperties);
+
+	ModifiedGroupItemsDelegate.ExecuteIfBound();
+}
+
+
+FName UNiagaraStackRendererItem::GetItemBackgroundName() const
+{
+	return "NiagaraEditor.Stack.Item.BackgroundColor";
+}
+
+int32 UNiagaraStackRendererItem::GetErrorCount() const
+{
+	return MissingAttributes.Num();
+}
+
+bool UNiagaraStackRendererItem::GetErrorFixable(int32 ErrorIdx) const
+{
+	if (ErrorIdx < MissingAttributes.Num())
+	{
+		return true;
+	}
+	return false;
+}
+
+bool UNiagaraStackRendererItem::TryFixError(int32 ErrorIdx)
+{
+	FNiagaraVariable MissingVar = MissingAttributes[ErrorIdx];
+	if (AddMissingVariable(GetEmitterViewModel()->GetEmitter(), MissingVar))
+	{
+		FNotificationInfo Info(FText::Format(LOCTEXT("AddedVariableForFix", "Added {0} to the Spawn script to support the renderer."), FText::FromName(MissingVar.GetName())));
+		Info.ExpireDuration = 5.0f;
+		Info.bFireAndForget = true;
+		Info.Image = FCoreStyle::Get().GetBrush(TEXT("MessageLog.Info"));
+		FSlateNotificationManager::Get().AddNotification(Info);
+		return true;
+	}
+	return false;
+}
+
+FText UNiagaraStackRendererItem::GetErrorText(int32 ErrorIdx) const
+{
+	if (ErrorIdx < MissingAttributes.Num())
+	{
+		const FNiagaraVariable& Attr = MissingAttributes[ErrorIdx];
+		return FText::Format(LOCTEXT("FailedRendererBind", "Missing attribute \"{0}\" of Type \"{1}\"."), FText::FromName(Attr.GetName()), Attr.GetType().GetNameText());
+	}
+	return FText();
+}
+
+void UNiagaraStackRendererItem::RefreshChildrenInternal(const TArray<UNiagaraStackEntry*>& CurrentChildren, TArray<UNiagaraStackEntry*>& NewChildren)
+{
+	if (RendererObject == nullptr)
+	{
+		RendererObject = NewObject<UNiagaraStackObject>(this);
+		RendererObject->Initialize(GetSystemViewModel(), GetEmitterViewModel(), RendererProperties);
+	}
+
+	if (RendererExpander == nullptr)
+	{
+		RendererExpander = NewObject<UNiagaraStackItemExpander>(this);
+		RendererExpander->Initialize(GetSystemViewModel(), GetEmitterViewModel(), GetStackEditorData(), RendererProperties->GetName(), false);
+		RendererExpander->SetOnExpnadedChanged(UNiagaraStackItemExpander::FOnExpandedChanged::CreateUObject(this, &UNiagaraStackRendererItem::RendererExpandedChanged));
+	}
+	
+	if (GetStackEditorData().GetStackEntryIsExpanded(RendererProperties->GetName(), false))
+	{
+		NewChildren.Add(RendererObject);
+	}
+
+	NewChildren.Add(RendererExpander);
+
+	MissingAttributes = GetMissingVariables(RendererProperties, GetEmitterViewModel()->GetEmitter());
+}
+
+void UNiagaraStackRendererItem::RendererExpandedChanged()
+{
+	RefreshChildren();
+}
+
+#undef LOCTEXT_NAMESPACE
\ No newline at end of file
