Copying //UE4/Orion-Staging to //UE4/Main (Source: //Orion/Dev-General @ 3606378)

diff --git a/Engine/Plugins/FX/Niagara/Source/NiagaraEditor/Private/ViewModels/Stack/NiagaraStackRoot.cpp b/Engine/Plugins/FX/Niagara/Source/NiagaraEditor/Private/ViewModels/Stack/NiagaraStackRoot.cpp
new file mode 100644
--- /dev/null
+++ b/Engine/Plugins/FX/Niagara/Source/NiagaraEditor/Private/ViewModels/Stack/NiagaraStackRoot.cpp
@@ -0,0 +1,173 @@
+// Copyright 1998-2016 Epic Games, Inc. All Rights Reserved.
+
+#include "NiagaraStackRoot.h"
+#include "NiagaraStackScriptItemGroup.h"
+#include "NiagaraStackEmitterSpawnScriptItemGroup.h"
+#include "NiagaraStackRenderItemGroup.h"
+#include "NiagaraStackEventHandlerGroup.h"
+#include "NiagaraStackEventScriptItemGroup.h"
+#include "NiagaraSystemViewModel.h"
+#include "NiagaraEmitterViewModel.h"
+#include "NiagaraSystemScriptViewModel.h"
+#include "NiagaraScriptViewModel.h"
+#include "NiagaraSystem.h"
+#include "NiagaraSystemEditorData.h"
+#include "NiagaraEmitterEditorData.h"
+#include "NiagaraStackParameterStoreGroup.h"
+
+#define LOCTEXT_NAMESPACE "NiagaraStackViewModel"
+
+UNiagaraStackRoot::UNiagaraStackRoot()
+	: EmitterSpawnGroup(nullptr)
+	, EmitterUpdateGroup(nullptr)
+	, ParticleSpawnGroup(nullptr)
+	, ParticleUpdateGroup(nullptr)
+	, AddEventHandlerGroup(nullptr)
+	, RenderGroup(nullptr)
+	, SystemExposedVariablesGroup(nullptr)
+{
+}
+
+void UNiagaraStackRoot::Initialize(TSharedRef<FNiagaraSystemViewModel> InSystemViewModel, TSharedRef<FNiagaraEmitterViewModel> InEmitterViewModel)
+{
+	Super::Initialize(InSystemViewModel, InEmitterViewModel);
+	EmitterSpawnGroup = nullptr;
+	EmitterUpdateGroup = nullptr;
+	ParticleSpawnGroup = nullptr;
+	ParticleUpdateGroup = nullptr;
+	AddEventHandlerGroup = nullptr;
+	RenderGroup = nullptr;
+	SystemExposedVariablesGroup = nullptr;
+}
+
+
+bool UNiagaraStackRoot::GetShouldShowInStack() const
+{
+	return false;
+}
+
+void UNiagaraStackRoot::RefreshChildrenInternal(const TArray<UNiagaraStackEntry*>& CurrentChildren, TArray<UNiagaraStackEntry*>& NewChildren)
+{
+	// We only allow displaying and editing system stacks if the system isn't transient which is the case in the emitter editor.
+	bool bShowSystemGroups = GetSystemViewModel()->GetSystemIsTransient() == false;
+
+	// Create static entries as needed.
+	if (bShowSystemGroups && SystemExposedVariablesGroup == nullptr)
+	{
+		SystemExposedVariablesGroup = NewObject<UNiagaraStackParameterStoreGroup>(this);
+		SystemExposedVariablesGroup->Initialize(GetSystemViewModel(), GetEmitterViewModel(), GetSystemViewModel()->GetOrCreateEditorData().GetStackEditorData(), GetSystemViewModel()->GetSystemScriptViewModel(), &GetSystemViewModel()->GetSystem(), &GetSystemViewModel()->GetSystem().GetExposedParameters());
+		SystemExposedVariablesGroup->SetDisplayName(LOCTEXT("SystemExposedVariablesGroup", "System Exposed Variables"));
+		SystemExposedVariablesGroup->SetTooltipText(LOCTEXT("SystemExposedVariablesGroupToolTip", "Displays the variables created in the User namespace. These variables are exposed to owning UComponents, blueprints, etc."));
+	}
+
+	if (bShowSystemGroups && SystemSpawnGroup == nullptr)
+	{
+		SystemSpawnGroup = NewObject<UNiagaraStackScriptItemGroup>(this);
+		SystemSpawnGroup->Initialize(GetSystemViewModel(), GetEmitterViewModel(), GetSystemViewModel()->GetOrCreateEditorData().GetStackEditorData(), GetSystemViewModel()->GetSystemScriptViewModel(), ENiagaraScriptUsage::SystemSpawnScript);
+		SystemSpawnGroup->SetDisplayName(LOCTEXT("SystemSpawnGroupName", "System Spawn"));
+		SystemSpawnGroup->SetTooltipText(LOCTEXT("SystemSpawnGroupToolTip", "Occurs once at System creation on the CPU. Modules in this section should initialize defaults and/or do initial setup.\r\nModules are executed in order from top to bottom of the stack."));
+	}
+
+	if (bShowSystemGroups && SystemUpdateGroup == nullptr)
+	{
+		SystemUpdateGroup = NewObject<UNiagaraStackScriptItemGroup>(this);
+		SystemUpdateGroup->Initialize(GetSystemViewModel(), GetEmitterViewModel(), GetSystemViewModel()->GetOrCreateEditorData().GetStackEditorData(), GetSystemViewModel()->GetSystemScriptViewModel(), ENiagaraScriptUsage::SystemUpdateScript);
+		SystemUpdateGroup->SetDisplayName(LOCTEXT("SystemUpdateGroupName", "System Update"));
+		SystemUpdateGroup->SetTooltipText(LOCTEXT("SystemUpdateGroupToolTip", "Occurs every Emitter tick on the CPU.Modules in this section should compute values for parameters for emitter or particle update or spawning this frame.\r\nModules are executed in order from top to bottom of the stack."));
+	}
+
+	if (EmitterSpawnGroup == nullptr)
+	{
+		EmitterSpawnGroup = NewObject<UNiagaraStackEmitterSpawnScriptItemGroup>(this);
+		EmitterSpawnGroup->Initialize(GetSystemViewModel(), GetEmitterViewModel(), GetEmitterViewModel()->GetOrCreateEditorData().GetStackEditorData(), GetEmitterViewModel()->GetSharedScriptViewModel(), ENiagaraScriptUsage::EmitterSpawnScript);
+		EmitterSpawnGroup->SetDisplayName(LOCTEXT("EmitterSpawnGroupName", "Emitter Spawn"));
+		EmitterSpawnGroup->SetTooltipText(LOCTEXT("EmitterSpawnGroupTooltip", "Occurs once at Emitter creation on the CPU. Modules in this section should initialize defaults and/or do initial setup.\r\nModules are executed in order from top to bottom of the stack."));
+	}
+
+	if (EmitterUpdateGroup == nullptr)
+	{
+		EmitterUpdateGroup = NewObject<UNiagaraStackScriptItemGroup>(this);
+		EmitterUpdateGroup->Initialize(GetSystemViewModel(), GetEmitterViewModel(), GetEmitterViewModel()->GetOrCreateEditorData().GetStackEditorData(), GetEmitterViewModel()->GetSharedScriptViewModel(), ENiagaraScriptUsage::EmitterUpdateScript);
+		EmitterUpdateGroup->SetDisplayName(LOCTEXT("EmitterUpdateGroupName", "Emitter Update"));
+		EmitterUpdateGroup->SetTooltipText(LOCTEXT("EmitterUpdateGroupTooltip", "Occurs every Emitter tick on the CPU. Modules in this section should compute values for parameters for Particle Update or Spawning this frame.\r\nModules are executed in order from top to bottom of the stack."));
+	}
+
+	if (ParticleSpawnGroup == nullptr)
+	{
+		ParticleSpawnGroup = NewObject<UNiagaraStackScriptItemGroup>(this);
+		ParticleSpawnGroup->Initialize(GetSystemViewModel(), GetEmitterViewModel(), GetEmitterViewModel()->GetOrCreateEditorData().GetStackEditorData(), GetEmitterViewModel()->GetSharedScriptViewModel(), ENiagaraScriptUsage::ParticleSpawnScript);
+		ParticleSpawnGroup->SetDisplayName(LOCTEXT("ParticleSpawnGroupName", "Particle Spawn"));
+		ParticleSpawnGroup->SetTooltipText(LOCTEXT("ParticleSpawnGroupTooltip", "Called once per created particle. Modules in this section should set up initial values for each particle.\r\nIf \"Use Interpolated Spawning\" is set, we will also run the Particle Update script after the Particle Spawn script.\r\nModules are executed in order from top to bottom of the stack."));
+	}
+
+	if (ParticleUpdateGroup == nullptr)
+	{
+		ParticleUpdateGroup = NewObject<UNiagaraStackScriptItemGroup>(this);
+		ParticleUpdateGroup->Initialize(GetSystemViewModel(), GetEmitterViewModel(), GetEmitterViewModel()->GetOrCreateEditorData().GetStackEditorData(), GetEmitterViewModel()->GetSharedScriptViewModel(), ENiagaraScriptUsage::ParticleUpdateScript);
+		ParticleUpdateGroup->SetDisplayName(LOCTEXT("ParticleUpdateGroupName", "Particle Update"));
+		ParticleUpdateGroup->SetTooltipText(LOCTEXT("ParticleUpdateGroupTooltip", "Called every frame per particle. Modules in this section should update new values for this frame.\r\nModules are executed in order from top to bottom of the stack."));
+	}
+
+	if (AddEventHandlerGroup == nullptr)
+	{
+		AddEventHandlerGroup = NewObject<UNiagaraStackEventHandlerGroup>(this);
+		AddEventHandlerGroup->Initialize(GetSystemViewModel(), GetEmitterViewModel(), GetEmitterViewModel()->GetOrCreateEditorData().GetStackEditorData());
+		AddEventHandlerGroup->SetDisplayName(LOCTEXT("EventGroupName", "Add Event Handler"));
+		AddEventHandlerGroup->SetTooltipText(LOCTEXT("EventGroupTooltip", "Determines how this Emitter responds to incoming events. There can be more than one event handler script stack per Emitter."));
+		AddEventHandlerGroup->SetOnItemAdded(UNiagaraStackEventHandlerGroup::FOnItemAdded::CreateUObject(this, &UNiagaraStackRoot::EmitterEventArraysChanged));
+	}
+
+	if (RenderGroup == nullptr)
+	{
+		RenderGroup = NewObject<UNiagaraStackRenderItemGroup>(this);
+		RenderGroup->Initialize(GetSystemViewModel(), GetEmitterViewModel(), GetEmitterViewModel()->GetOrCreateEditorData().GetStackEditorData());
+		RenderGroup->SetDisplayName(LOCTEXT("RenderGroupName", "Render"));
+		RenderGroup->SetTooltipText(LOCTEXT("RendererGroupTooltip", "Describes how we should display/present each particle. Note that this doesn't have to be visual. Multiple renderers are supported. Order in this stack is not necessarily relevant to draw order."));
+	}
+
+	// Populate new children
+	if (bShowSystemGroups)
+	{
+		NewChildren.Add(SystemExposedVariablesGroup);
+		NewChildren.Add(SystemSpawnGroup);
+		NewChildren.Add(SystemUpdateGroup);
+	}
+	NewChildren.Add(EmitterSpawnGroup);
+	NewChildren.Add(EmitterUpdateGroup);
+	NewChildren.Add(ParticleSpawnGroup);
+	NewChildren.Add(ParticleUpdateGroup);
+
+	for (int32 EventHandlerIdx = 0; EventHandlerIdx < GetEmitterViewModel()->GetEmitter()->EventHandlerScriptProps.Num(); EventHandlerIdx++)
+	{
+		UNiagaraStackEventScriptItemGroup* EventHandlerGroup = nullptr;
+		for (UNiagaraStackEntry* CurrentChild : CurrentChildren)
+		{
+			UNiagaraStackEventScriptItemGroup* CurrentEventHandlerChild = Cast<UNiagaraStackEventScriptItemGroup>(CurrentChild);
+			if (CurrentEventHandlerChild != nullptr && CurrentEventHandlerChild->GetScriptOccurrence() == EventHandlerIdx)
+			{
+				EventHandlerGroup = CurrentEventHandlerChild;
+				break;
+			}
+		}
+
+		if (EventHandlerGroup == nullptr)
+		{
+			EventHandlerGroup = NewObject<UNiagaraStackEventScriptItemGroup>(this, NAME_None, RF_Transactional);
+			EventHandlerGroup->Initialize(GetSystemViewModel(), GetEmitterViewModel(), GetEmitterViewModel()->GetEditorData().GetStackEditorData(), GetEmitterViewModel()->GetSharedScriptViewModel(), ENiagaraScriptUsage::ParticleEventScript, EventHandlerIdx);
+			EventHandlerGroup->SetOnModifiedEventHandlers(UNiagaraStackEventScriptItemGroup::FOnModifiedEventHandlers::CreateUObject(this, &UNiagaraStackRoot::EmitterEventArraysChanged));
+			EventHandlerGroup->SetTooltipText(LOCTEXT("EventGroupTooltip", "Determines how this Emitter responds to incoming events. There can be more than one event handler script stack per Emitter."));
+		}
+
+		NewChildren.Add(EventHandlerGroup);
+	}
+
+	NewChildren.Add(AddEventHandlerGroup);
+	NewChildren.Add(RenderGroup);
+}
+
+void UNiagaraStackRoot::EmitterEventArraysChanged()
+{
+	RefreshChildren();
+}
+
+#undef LOCTEXT_NAMESPACE
\ No newline at end of file
diff --git a/Engine/Plugins/FX/Niagara/Source/NiagaraEditor/Private/ViewModels/Stack/NiagaraStackScriptItemGroup.cpp b/Engine/Plugins/FX/Niagara/Source/NiagaraEditor/Private/ViewModels/Stack/NiagaraStackScriptItemGroup.cpp
new file mode 100644
--- /dev/null
+++ b/Engine/Plugins/FX/Niagara/Source/NiagaraEditor/Private/ViewModels/Stack/NiagaraStackScriptItemGroup.cpp
@@ -0,0 +1,170 @@
+// Copyright 1998-2016 Epic Games, Inc. All Rights Reserved.
+
+#include "NiagaraStackScriptItemGroup.h"
+#include "NiagaraStackModuleItem.h"
+#include "NiagaraStackSpacer.h"
+#include "NiagaraStackAddScriptModuleItem.h"
+#include "NiagaraScriptViewModel.h"
+#include "NiagaraScriptGraphViewModel.h"
+#include "NiagaraGraph.h"
+#include "NiagaraNodeOutput.h"
+#include "NiagaraNodeFunctionCall.h"
+#include "EdGraphSchema_Niagara.h"
+#include "NiagaraStackErrorItem.h"
+#include "Internationalization.h"
+#include "NiagaraStackGraphUtilities.h"
+#include "ScopedTransaction.h"
+
+#define LOCTEXT_NAMESPACE "UNiagaraStackScriptItemGroup"
+
+UNiagaraStackScriptItemGroup::UNiagaraStackScriptItemGroup()
+	: AddModuleItem(nullptr)
+	, BottomSpacer(nullptr)
+{
+}
+
+void UNiagaraStackScriptItemGroup::Initialize(
+	TSharedRef<FNiagaraSystemViewModel> InSystemViewModel,
+	TSharedRef<FNiagaraEmitterViewModel> InEmitterViewModel,
+	UNiagaraStackEditorData& InStackEditorData,
+	TSharedRef<FNiagaraScriptViewModel> InScriptViewModel,
+	ENiagaraScriptUsage InScriptUsage,
+	int32 InScriptOccurrence)
+{
+	checkf(ScriptViewModel.IsValid() == false, TEXT("Can not set the script view model more than once."));
+	Super::Initialize(InSystemViewModel, InEmitterViewModel, InStackEditorData);
+	ScriptViewModel = InScriptViewModel;
+	ScriptUsage = InScriptUsage;
+	ScriptOccurrence = InScriptOccurrence;
+}
+
+FText UNiagaraStackScriptItemGroup::GetDisplayName() const
+{
+	return DisplayName;
+}
+
+void UNiagaraStackScriptItemGroup::SetDisplayName(FText InDisplayName)
+{
+	DisplayName = InDisplayName;
+}
+
+void UNiagaraStackScriptItemGroup::RefreshChildrenInternal(const TArray<UNiagaraStackEntry*>& CurrentChildren, TArray<UNiagaraStackEntry*>& NewChildren)
+{
+	UNiagaraGraph* Graph = ScriptViewModel->GetGraphViewModel()->GetGraph();
+	FText ErrorMessage;
+	if (FNiagaraStackGraphUtilities::ValidateGraphForOutput(*Graph, ScriptUsage, ScriptOccurrence, ErrorMessage) == false)
+	{
+		UE_LOG(LogNiagaraEditor, Error, TEXT("Failed to Create Stack.  Message: %s"), *ErrorMessage.ToString());
+		Error = MakeShared<FError>();
+		Error->ErrorText = LOCTEXT("InvalidErrorText", "The data used to generate the stack has been corrupted and can not be used.\nUsing the fix option will reset this part of the stack to its default empty state.");
+		Error->ErrorSummaryText = LOCTEXT("InvalidErrorSummaryText", "The stack data is invalid");
+		Error->Fix.BindLambda([=]()
+		{
+			FScopedTransaction ScopedTransaction(LOCTEXT("FixStackGraph", "Fix invalid stack graph"));
+			FNiagaraStackGraphUtilities::ResetGraphForOutput(*Graph, ScriptUsage, ScriptOccurrence);
+		});
+	}
+	else
+	{
+		UNiagaraNodeOutput* MatchingOutputNode = Graph->FindOutputNode(ScriptUsage, ScriptOccurrence);
+		TArray<UNiagaraNodeFunctionCall*> ModuleNodes;
+		FNiagaraStackGraphUtilities::GetOrderedModuleNodes(*MatchingOutputNode, ModuleNodes);
+		int32 ModuleIndex = 0;
+		for (UNiagaraNodeFunctionCall* ModuleNode : ModuleNodes)
+		{
+			UNiagaraStackModuleItem* ModuleItem = FindCurrentChildOfTypeByPredicate<UNiagaraStackModuleItem>(CurrentChildren, 
+				[=](UNiagaraStackModuleItem* CurrentModuleItem) { return &CurrentModuleItem->GetModuleNode() == ModuleNode; });
+
+			if (ModuleItem == nullptr)
+			{
+				ModuleItem = NewObject<UNiagaraStackModuleItem>(this);
+				ModuleItem->Initialize(GetSystemViewModel(), GetEmitterViewModel(), GetStackEditorData(), *ModuleNode);
+				ModuleItem->SetOnModifiedGroupItems(UNiagaraStackModuleItem::FOnModifiedGroupItems::CreateUObject(this, &UNiagaraStackScriptItemGroup::ChildModifiedGroupItems));
+			}
+
+			FName ModuleSpacerKey = *FString::Printf(TEXT("Module%i"), ModuleIndex);
+			UNiagaraStackSpacer* ModuleSpacer = FindCurrentChildOfTypeByPredicate<UNiagaraStackSpacer>(CurrentChildren,
+				[=](UNiagaraStackSpacer* CurrentModuleSpacer) { return CurrentModuleSpacer->GetSpacerKey() == ModuleSpacerKey; });
+
+			if (ModuleSpacer == nullptr)
+			{
+				ModuleSpacer = NewObject<UNiagaraStackSpacer>(this);
+				ModuleSpacer->Initialize(GetSystemViewModel(), GetEmitterViewModel(), ModuleSpacerKey);
+			}
+
+			NewChildren.Add(ModuleItem);
+			NewChildren.Add(ModuleSpacer);
+			ModuleIndex++;
+		}
+
+		if (AddModuleItem == nullptr)
+		{
+			AddModuleItem = NewObject<UNiagaraStackAddScriptModuleItem>(this);
+			AddModuleItem->Initialize(GetSystemViewModel(), GetEmitterViewModel(), GetStackEditorData(), *MatchingOutputNode);
+			AddModuleItem->SetOnItemAdded(UNiagaraStackAddModuleItem::FOnItemAdded::CreateUObject(this, &UNiagaraStackScriptItemGroup::ItemAdded));
+		}
+
+		if (BottomSpacer == nullptr)
+		{
+			BottomSpacer = NewObject<UNiagaraStackSpacer>(this);
+			BottomSpacer->Initialize(GetSystemViewModel(), GetEmitterViewModel(), "ScriptStackBottom");
+		}
+
+		NewChildren.Add(AddModuleItem);
+		NewChildren.Add(BottomSpacer);
+
+		ENiagaraScriptCompileStatus Status = ScriptViewModel->GetScriptCompileStatus(GetScriptUsage(), GetScriptOccurrence());
+		if (Status == ENiagaraScriptCompileStatus::NCS_Error)
+		{
+			Error = MakeShared<FError>();
+			Error->ErrorText = ScriptViewModel->GetScriptErrors(GetScriptUsage(), GetScriptOccurrence());
+			Error->ErrorSummaryText = LOCTEXT("ConpileErrorSummary", "The stack has compile errors.");
+		}
+		else
+		{
+			Error.Reset();
+		}
+	}
+}
+
+void UNiagaraStackScriptItemGroup::ItemAdded()
+{
+	RefreshChildren();
+}
+
+void UNiagaraStackScriptItemGroup::ChildModifiedGroupItems()
+{
+	RefreshChildren();
+}
+
+int32 UNiagaraStackScriptItemGroup::GetErrorCount() const
+{
+	return Error.IsValid() ? 1 : 0;
+}
+
+bool UNiagaraStackScriptItemGroup::GetErrorFixable(int32 ErrorIdx) const
+{
+	return Error.IsValid() && Error->Fix.IsBound();
+}
+
+bool UNiagaraStackScriptItemGroup::TryFixError(int32 ErrorIdx)
+{
+	if (Error.IsValid() && Error->Fix.IsBound())
+	{
+		Error->Fix.Execute();
+		return true;
+	}
+	return false;
+}
+
+FText UNiagaraStackScriptItemGroup::GetErrorText(int32 ErrorIdx) const
+{
+	return Error->ErrorText;
+}
+
+FText UNiagaraStackScriptItemGroup::GetErrorSummaryText(int32 ErrorIdx) const
+{
+	return Error->ErrorSummaryText;
+}
+
+#undef LOCTEXT_NAMESPACE
\ No newline at end of file
diff --git a/Engine/Plugins/FX/Niagara/Source/NiagaraEditor/Private/ViewModels/Stack/NiagaraStackSpacer.cpp b/Engine/Plugins/FX/Niagara/Source/NiagaraEditor/Private/ViewModels/Stack/NiagaraStackSpacer.cpp
new file mode 100644
--- /dev/null
+++ b/Engine/Plugins/FX/Niagara/Source/NiagaraEditor/Private/ViewModels/Stack/NiagaraStackSpacer.cpp
@@ -0,0 +1,24 @@
+// Copyright 1998-2016 Epic Games, Inc. All Rights Reserved.
+
+#include "NiagaraStackSpacer.h"
+
+void UNiagaraStackSpacer::Initialize(TSharedRef<FNiagaraSystemViewModel> InSystemViewModel, TSharedRef<FNiagaraEmitterViewModel> InEmitterViewModel, FName InSpacerKey)
+{
+	Super::Initialize(InSystemViewModel, InEmitterViewModel);
+	SpacerKey = InSpacerKey;
+}
+
+FText UNiagaraStackSpacer::GetDisplayName() const
+{
+	return FText();
+}
+
+FName UNiagaraStackSpacer::GetItemBackgroundName() const
+{
+	return "NiagaraEditor.Stack.Group.BackgroundColor";
+}
+
+FName UNiagaraStackSpacer::GetSpacerKey() const
+{
+	return SpacerKey;
+}
\ No newline at end of file
diff --git a/Engine/Plugins/FX/Niagara/Source/NiagaraEditor/Private/ViewModels/Stack/NiagaraStackStruct.cpp b/Engine/Plugins/FX/Niagara/Source/NiagaraEditor/Private/ViewModels/Stack/NiagaraStackStruct.cpp
new file mode 100644
--- /dev/null
+++ b/Engine/Plugins/FX/Niagara/Source/NiagaraEditor/Private/ViewModels/Stack/NiagaraStackStruct.cpp
@@ -0,0 +1,59 @@
+// Copyright 1998-2016 Epic Games, Inc. All Rights Reserved.
+
+#include "NiagaraStackStruct.h"
+
+
+UNiagaraStackStruct::UNiagaraStackStruct()
+	: OwningObject(nullptr)
+{
+}
+
+void UNiagaraStackStruct::Initialize(TSharedRef<FNiagaraSystemViewModel> InSystemViewModel, TSharedRef<FNiagaraEmitterViewModel> InEmitterViewModel, UObject* InOwningObject, const UStruct* InScriptStructClass, uint8* InStructData)
+{
+	checkf(OwningObject == nullptr, TEXT("Can only initialize once."));
+	Super::Initialize(InSystemViewModel, InEmitterViewModel);
+	OwningObject = InOwningObject;
+
+	StructData = MakeShareable(new FStructOnScope(InScriptStructClass, InStructData));
+}
+
+UObject* UNiagaraStackStruct::GetOwningObject()
+{
+	return OwningObject;
+}
+
+
+TSharedPtr<FStructOnScope> UNiagaraStackStruct::GetStructOnScope()
+{
+	return StructData;
+}
+
+int32 UNiagaraStackStruct::GetItemIndentLevel() const
+{
+	return ItemIndentLevel;
+}
+
+void UNiagaraStackStruct::SetItemIndentLevel(int32 InItemIndentLevel)
+{
+	ItemIndentLevel = InItemIndentLevel;
+}
+
+bool UNiagaraStackStruct::HasDetailCustomization() const
+{
+	if (DetailCustomization.IsBound())
+	{
+		return true;
+	}
+	return false;
+}
+
+FOnGetDetailCustomizationInstance UNiagaraStackStruct::GetDetailsCustomization() const
+{
+	return DetailCustomization;
+}
+
+
+void UNiagaraStackStruct::SetDetailsCustomization(FOnGetDetailCustomizationInstance Customization)
+{
+	DetailCustomization = Customization;
+}
\ No newline at end of file
diff --git a/Engine/Plugins/FX/Niagara/Source/NiagaraEditor/Private/ViewModels/Stack/NiagaraStackViewModel.cpp b/Engine/Plugins/FX/Niagara/Source/NiagaraEditor/Private/ViewModels/Stack/NiagaraStackViewModel.cpp
new file mode 100644
--- /dev/null
+++ b/Engine/Plugins/FX/Niagara/Source/NiagaraEditor/Private/ViewModels/Stack/NiagaraStackViewModel.cpp
@@ -0,0 +1,103 @@
+// Copyright 1998-2016 Epic Games, Inc. All Rights Reserved.
+
+#include "NiagaraStackViewModel.h"
+#include "NiagaraStackRoot.h"
+#include "NiagaraSystemViewModel.h"
+#include "NiagaraEmitterViewModel.h"
+#include "NiagaraScriptGraphViewModel.h"
+#include "NiagaraEmitterEditorData.h"
+#include "NiagaraEmitter.h"
+
+#include "Editor.h"
+
+#define LOCTEXT_NAMESPACE "NiagaraStackViewModel"
+
+TSharedPtr<FNiagaraEmitterViewModel> UNiagaraStackViewModel::GetEmitterViewModel()
+{
+	return EmitterViewModel;
+}
+
+void UNiagaraStackViewModel::Initialize(TSharedPtr<FNiagaraSystemViewModel> InSystemViewModel, TSharedPtr<FNiagaraEmitterViewModel> InEmitterViewModel)
+{
+	if (EmitterViewModel.IsValid())
+	{
+		GEditor->UnregisterForUndo(this);
+		EmitterViewModel->OnScriptCompiled().RemoveAll(this);
+		RootEntry->OnStructureChanged().RemoveAll(this);
+		RootEntry->OnDataObjectModified().RemoveAll(this);
+		RootEntries.Empty();
+		RootEntry = nullptr;
+	}
+
+	if (SystemViewModel.IsValid())
+	{
+		SystemViewModel->OnSystemCompiled().RemoveAll(this);
+	}
+
+	SystemViewModel = InSystemViewModel;
+	EmitterViewModel = InEmitterViewModel;
+
+	if (SystemViewModel.IsValid() && EmitterViewModel.IsValid() && EmitterViewModel->GetSharedScriptViewModel()->GetGraphViewModel()->GetGraph() != nullptr)
+	{
+		GEditor->RegisterForUndo(this);
+
+		EmitterViewModel->OnScriptCompiled().AddUObject(this, &UNiagaraStackViewModel::OnEmitterCompiled);
+		SystemViewModel->OnSystemCompiled().AddUObject(this, &UNiagaraStackViewModel::OnSystemCompiled);
+		
+		RootEntry = NewObject<UNiagaraStackRoot>(this);
+		RootEntry->Initialize(SystemViewModel.ToSharedRef(), EmitterViewModel.ToSharedRef());
+		RootEntry->RefreshChildren();
+		RootEntry->OnStructureChanged().AddUObject(this, &UNiagaraStackViewModel::EntryStructureChanged);
+		RootEntry->OnDataObjectModified().AddUObject(this, &UNiagaraStackViewModel::EntryDataObjectModified);
+		RootEntries.Add(RootEntry);
+	}
+
+	StructureChangedDelegate.Broadcast();
+}
+
+void UNiagaraStackViewModel::BeginDestroy()
+{
+	Super::BeginDestroy();
+	if (EmitterViewModel.IsValid())
+	{
+		GEditor->UnregisterForUndo(this);
+		EmitterViewModel->OnScriptCompiled().RemoveAll(this);
+	}
+}
+
+void UNiagaraStackViewModel::OnSystemCompiled()
+{
+	RootEntry->RefreshChildren();
+}
+
+void UNiagaraStackViewModel::OnEmitterCompiled()
+{
+	RootEntry->RefreshChildren();
+}
+
+TArray<UNiagaraStackEntry*>& UNiagaraStackViewModel::GetRootEntries()
+{
+	return RootEntries;
+}
+
+UNiagaraStackViewModel::FOnStructureChanged& UNiagaraStackViewModel::OnStructureChanged()
+{
+	return StructureChangedDelegate;
+}
+
+void UNiagaraStackViewModel::PostUndo(bool bSuccess)
+{
+	RootEntry->RefreshChildren();
+}
+
+void UNiagaraStackViewModel::EntryStructureChanged()
+{
+	StructureChangedDelegate.Broadcast();
+}
+
+void UNiagaraStackViewModel::EntryDataObjectModified(UObject* ChangedObject)
+{
+	SystemViewModel->NotifyDataObjectChanged(ChangedObject);
+}
+
+#undef LOCTEXT_NAMESPACE
diff --git a/Engine/Plugins/FX/Niagara/Source/NiagaraEditor/Private/Widgets/SNiagaraCurveEditor.cpp b/Engine/Plugins/FX/Niagara/Source/NiagaraEditor/Private/Widgets/SNiagaraCurveEditor.cpp
new file mode 100644
--- /dev/null
+++ b/Engine/Plugins/FX/Niagara/Source/NiagaraEditor/Private/Widgets/SNiagaraCurveEditor.cpp
@@ -0,0 +1,229 @@
+// Copyright 1998-2016 Epic Games, Inc. All Rights Reserved.
+
+#include "SNiagaraCurveEditor.h"
+#include "NiagaraSystemViewModel.h"
+#include "SCurveEditor.h"
+#include "SOverlay.h"
+#include "Framework/MultiBox/MultiBoxBuilder.h"
+#include "RichCurveEditorCommands.h"
+#include "Widgets/Input/SNumericDropDown.h"
+#include "SCurveEditor.h"
+#include "SOverlay.h"
+
+#define LOCTEXT_NAMESPACE "NiagaraCurveEditor"
+
+void SNiagaraCurveEditor::Construct(const FArguments& InArgs, TSharedRef<FNiagaraSystemViewModel> InSystemViewModel)
+{
+	SystemViewModel = InSystemViewModel;
+	SystemViewModel->OnCurveOwnerChanged().AddRaw(this, &SNiagaraCurveEditor::OnCurveOwnerChanged);
+	InputSnap = .1f;
+	OutputSnap = .1f;
+
+	SAssignNew(CurveEditor, SCurveEditor)
+		.ShowCurveSelector(true);
+	CurveEditor->SetCurveOwner(&SystemViewModel->GetCurveOwner());
+
+	TSharedPtr<SOverlay> OverlayWidget;
+	this->ChildSlot
+	[
+		SNew(SVerticalBox)
+		+ SVerticalBox::Slot()
+		.AutoHeight()
+		[
+			ConstructToolBar(CurveEditor->GetCommands())
+		]
+		+ SVerticalBox::Slot()
+		[
+			CurveEditor.ToSharedRef()
+		]
+	];
+}
+
+SNiagaraCurveEditor::~SNiagaraCurveEditor()
+{
+	SystemViewModel->OnCurveOwnerChanged().RemoveAll(this);
+}
+
+TSharedRef<SWidget> SNiagaraCurveEditor::ConstructToolBar(TSharedPtr<FUICommandList> CurveEditorCommandList)
+{
+	FToolBarBuilder ToolBarBuilder(CurveEditorCommandList, FMultiBoxCustomization::None, TSharedPtr<FExtender>(), Orient_Horizontal, true);
+
+	// TODO: Move this to a shared location since it's 99% the same as the sequencer curve toolbar.
+	ToolBarBuilder.AddComboButton(
+		FUIAction(),
+		FOnGetContent::CreateSP(this, &SNiagaraCurveEditor::MakeCurveEditorViewOptionsMenu, CurveEditorCommandList),
+		LOCTEXT("CurveEditorViewOptions", "View Options"),
+		LOCTEXT("CurveEditorViewOptionsToolTip", "View Options"),
+		TAttribute<FSlateIcon>(),
+		true);
+
+	TArray<SNumericDropDown<float>::FNamedValue> SnapValues;
+	SnapValues.Add(SNumericDropDown<float>::FNamedValue(0.001f, LOCTEXT("Snap_OneThousandth", "0.001"), LOCTEXT("SnapDescription_OneThousandth", "Set snap to 1/1000th")));
+	SnapValues.Add(SNumericDropDown<float>::FNamedValue(0.01f, LOCTEXT("Snap_OneHundredth", "0.01"), LOCTEXT("SnapDescription_OneHundredth", "Set snap to 1/100th")));
+	SnapValues.Add(SNumericDropDown<float>::FNamedValue(0.1f, LOCTEXT("Snap_OneTenth", "0.1"), LOCTEXT("SnapDescription_OneTenth", "Set snap to 1/10th")));
+	SnapValues.Add(SNumericDropDown<float>::FNamedValue(1.0f, LOCTEXT("Snap_One", "1"), LOCTEXT("SnapDescription_One", "Set snap to 1")));
+	SnapValues.Add(SNumericDropDown<float>::FNamedValue(10.0f, LOCTEXT("Snap_Ten", "10"), LOCTEXT("SnapDescription_Ten", "Set snap to 10")));
+	SnapValues.Add(SNumericDropDown<float>::FNamedValue(100.0f, LOCTEXT("Snap_OneHundred", "100"), LOCTEXT("SnapDescription_OneHundred", "Set snap to 100")));
+
+	TSharedRef<SWidget> InputSnapWidget =
+		SNew(SNumericDropDown<float>)
+		.DropDownValues(SnapValues)
+		.LabelText(LOCTEXT("InputSnapLabel", "Input Snap"))
+		.Value(this, &SNiagaraCurveEditor::GetInputSnap)
+		.OnValueChanged(this, &SNiagaraCurveEditor::SetInputSnap);
+
+	TSharedRef<SWidget> OutputSnapWidget =
+		SNew(SNumericDropDown<float>)
+		.DropDownValues(SnapValues)
+		.LabelText(LOCTEXT("OutputSnapLabel", "Output Snap"))
+		.Value(this, &SNiagaraCurveEditor::GetOutputSnap)
+		.OnValueChanged(this, &SNiagaraCurveEditor::SetOutputSnap);
+
+	ToolBarBuilder.BeginSection("Snap");
+	{
+		ToolBarBuilder.AddToolBarButton(FRichCurveEditorCommands::Get().ToggleInputSnapping);
+		ToolBarBuilder.AddWidget(InputSnapWidget);
+		ToolBarBuilder.AddToolBarButton(FRichCurveEditorCommands::Get().ToggleOutputSnapping);
+		ToolBarBuilder.AddWidget(OutputSnapWidget);
+	}
+	ToolBarBuilder.EndSection();
+
+	ToolBarBuilder.BeginSection("Curve");
+	{
+		ToolBarBuilder.AddToolBarButton(FRichCurveEditorCommands::Get().ZoomToFitHorizontal);
+		ToolBarBuilder.AddToolBarButton(FRichCurveEditorCommands::Get().ZoomToFitVertical);
+		ToolBarBuilder.AddToolBarButton(FRichCurveEditorCommands::Get().ZoomToFit);
+	}
+	ToolBarBuilder.EndSection();
+
+	ToolBarBuilder.BeginSection("Interpolation");
+	{
+		ToolBarBuilder.AddToolBarButton(FRichCurveEditorCommands::Get().InterpolationCubicAuto);
+		ToolBarBuilder.AddToolBarButton(FRichCurveEditorCommands::Get().InterpolationCubicUser);
+		ToolBarBuilder.AddToolBarButton(FRichCurveEditorCommands::Get().InterpolationCubicBreak);
+		ToolBarBuilder.AddToolBarButton(FRichCurveEditorCommands::Get().InterpolationLinear);
+		ToolBarBuilder.AddToolBarButton(FRichCurveEditorCommands::Get().InterpolationConstant);
+	}
+	ToolBarBuilder.EndSection();
+
+	ToolBarBuilder.BeginSection("Tangents");
+	{
+		ToolBarBuilder.AddToolBarButton(FRichCurveEditorCommands::Get().FlattenTangents);
+		ToolBarBuilder.AddToolBarButton(FRichCurveEditorCommands::Get().StraightenTangents);
+	}
+	ToolBarBuilder.EndSection();
+
+	ToolBarBuilder.AddComboButton(
+		FUIAction(),
+		FOnGetContent::CreateSP(this, &SNiagaraCurveEditor::MakeCurveEditorCurveOptionsMenu, CurveEditorCommandList),
+		LOCTEXT("CurveEditorCurveOptions", "Curves Options"),
+		LOCTEXT("CurveEditorCurveOptionsToolTip", "Curve Options"),
+		TAttribute<FSlateIcon>(),
+		true);
+
+	return ToolBarBuilder.MakeWidget();
+}
+
+TSharedRef<SWidget> SNiagaraCurveEditor::MakeCurveEditorViewOptionsMenu(TSharedPtr<FUICommandList> CurveEditorCommandList)
+{
+	FMenuBuilder MenuBuilder(true, CurveEditorCommandList);
+
+	MenuBuilder.BeginSection("CurveVisibility", LOCTEXT("CurveEditorMenuCurveVisibilityHeader", "Curve Visibility"));
+	{
+		MenuBuilder.AddMenuEntry(FRichCurveEditorCommands::Get().SetAllCurveVisibility);
+		MenuBuilder.AddMenuEntry(FRichCurveEditorCommands::Get().SetSelectedCurveVisibility);
+		MenuBuilder.AddMenuEntry(FRichCurveEditorCommands::Get().SetAnimatedCurveVisibility);
+	}
+	MenuBuilder.EndSection();
+
+	MenuBuilder.BeginSection("TangentVisibility", LOCTEXT("CurveEditorMenuTangentVisibilityHeader", "Tangent Visibility"));
+	{
+		MenuBuilder.AddMenuEntry(FRichCurveEditorCommands::Get().SetAllTangentsVisibility);
+		MenuBuilder.AddMenuEntry(FRichCurveEditorCommands::Get().SetSelectedKeysTangentVisibility);
+		MenuBuilder.AddMenuEntry(FRichCurveEditorCommands::Get().SetNoTangentsVisibility);
+	}
+	MenuBuilder.EndSection();
+
+	MenuBuilder.AddMenuSeparator();
+
+	MenuBuilder.AddMenuEntry(FRichCurveEditorCommands::Get().ToggleAutoFrameCurveEditor);
+	MenuBuilder.AddMenuEntry(FRichCurveEditorCommands::Get().ToggleShowCurveEditorCurveToolTips);
+
+	return MenuBuilder.MakeWidget();
+}
+TSharedRef<SWidget> SNiagaraCurveEditor::MakeCurveEditorCurveOptionsMenu(TSharedPtr<FUICommandList> CurveEditorCommandList)
+{
+	struct FExtrapolationMenus
+	{
+		static void MakePreInfinityExtrapSubMenu(FMenuBuilder& MenuBuilder)
+		{
+			MenuBuilder.BeginSection("Pre-Infinity Extrapolation", LOCTEXT("CurveEditorMenuPreInfinityExtrapHeader", "Extrapolation"));
+			{
+				MenuBuilder.AddMenuEntry(FRichCurveEditorCommands::Get().SetPreInfinityExtrapCycle);
+				MenuBuilder.AddMenuEntry(FRichCurveEditorCommands::Get().SetPreInfinityExtrapCycleWithOffset);
+				MenuBuilder.AddMenuEntry(FRichCurveEditorCommands::Get().SetPreInfinityExtrapOscillate);
+				MenuBuilder.AddMenuEntry(FRichCurveEditorCommands::Get().SetPreInfinityExtrapLinear);
+				MenuBuilder.AddMenuEntry(FRichCurveEditorCommands::Get().SetPreInfinityExtrapConstant);
+			}
+			MenuBuilder.EndSection();
+		}
+
+		static void MakePostInfinityExtrapSubMenu(FMenuBuilder& MenuBuilder)
+		{
+			MenuBuilder.BeginSection("Post-Infinity Extrapolation", LOCTEXT("CurveEditorMenuPostInfinityExtrapHeader", "Extrapolation"));
+			{
+				MenuBuilder.AddMenuEntry(FRichCurveEditorCommands::Get().SetPostInfinityExtrapCycle);
+				MenuBuilder.AddMenuEntry(FRichCurveEditorCommands::Get().SetPostInfinityExtrapCycleWithOffset);
+				MenuBuilder.AddMenuEntry(FRichCurveEditorCommands::Get().SetPostInfinityExtrapOscillate);
+				MenuBuilder.AddMenuEntry(FRichCurveEditorCommands::Get().SetPostInfinityExtrapLinear);
+				MenuBuilder.AddMenuEntry(FRichCurveEditorCommands::Get().SetPostInfinityExtrapConstant);
+			}
+			MenuBuilder.EndSection();
+		}
+	};
+
+	FMenuBuilder MenuBuilder(true, CurveEditorCommandList);
+
+	MenuBuilder.AddMenuEntry(FRichCurveEditorCommands::Get().BakeCurve);
+	MenuBuilder.AddMenuEntry(FRichCurveEditorCommands::Get().ReduceCurve);
+
+	MenuBuilder.AddSubMenu(
+		LOCTEXT("PreInfinitySubMenu", "Pre-Infinity"),
+		LOCTEXT("PreInfinitySubMenuToolTip", "Pre-Infinity Extrapolation"),
+		FNewMenuDelegate::CreateStatic(&FExtrapolationMenus::MakePreInfinityExtrapSubMenu));
+
+	MenuBuilder.AddSubMenu(
+		LOCTEXT("PostInfinitySubMenu", "Post-Infinity"),
+		LOCTEXT("PostInfinitySubMenuToolTip", "Post-Infinity Extrapolation"),
+		FNewMenuDelegate::CreateStatic(&FExtrapolationMenus::MakePostInfinityExtrapSubMenu));
+
+	return MenuBuilder.MakeWidget();
+}
+
+
+float SNiagaraCurveEditor::GetInputSnap() const
+{
+	return InputSnap;
+}
+
+void SNiagaraCurveEditor::SetInputSnap(float Value)
+{
+	InputSnap = Value;
+}
+
+float SNiagaraCurveEditor::GetOutputSnap() const
+{
+	return OutputSnap;
+}
+
+void SNiagaraCurveEditor::SetOutputSnap(float Value)
+{
+	OutputSnap = Value;
+}
+
+void SNiagaraCurveEditor::OnCurveOwnerChanged()
+{
+	CurveEditor->SetCurveOwner(&SystemViewModel->GetCurveOwner());
+}
+
+#undef LOCTEXT_NAMESPACE
diff --git a/Engine/Plugins/FX/Niagara/Source/NiagaraEditor/Private/Widgets/SNiagaraCurveEditor.h b/Engine/Plugins/FX/Niagara/Source/NiagaraEditor/Private/Widgets/SNiagaraCurveEditor.h
new file mode 100644
--- /dev/null
+++ b/Engine/Plugins/FX/Niagara/Source/NiagaraEditor/Private/Widgets/SNiagaraCurveEditor.h
@@ -0,0 +1,51 @@
+// Copyright 1998-2016 Epic Games, Inc. All Rights Reserved.
+
+#pragma once
+#include "SCompoundWidget.h"
+#include "DeclarativeSyntaxSupport.h"
+
+#include "Widgets/SCompoundWidget.h"
+#include "Widgets/DeclarativeSyntaxSupport.h"
+
+class FNiagaraSystemViewModel;
+class SCurveEditor;
+class FUICommandList;
+struct FNiagaraSimulation;
+class FNiagaraEffectInstance;
+class UNiagaraEffect;
+class UCurveBase;
+
+/** A curve editor control for curves in a niagara System. */
+class SNiagaraCurveEditor : public SCompoundWidget
+{
+public:
+	SLATE_BEGIN_ARGS(SNiagaraCurveEditor) { }
+	SLATE_END_ARGS()
+
+	virtual void Construct(const FArguments& InArgs, TSharedRef<FNiagaraSystemViewModel> InSystemViewModel);
+
+	~SNiagaraCurveEditor();
+
+private:
+	TSharedRef<SWidget> ConstructToolBar(TSharedPtr<FUICommandList> CurveEditorCommandList);
+
+	/** Makes the curve editor view options menu for the toolbar. */
+	TSharedRef<SWidget> MakeCurveEditorViewOptionsMenu(TSharedPtr<FUICommandList> CurveEditorCommandList);
+
+	/** Makes the curve editor curve options menu for the toolbar. */
+	TSharedRef<SWidget> MakeCurveEditorCurveOptionsMenu(TSharedPtr<FUICommandList> CurveEditorCommandList);
+
+	float GetInputSnap() const;
+	void SetInputSnap(float Value);
+
+	float GetOutputSnap() const;
+	void SetOutputSnap(float Value);
+
+	void OnCurveOwnerChanged();
+
+private:
+	TSharedPtr<FNiagaraSystemViewModel> SystemViewModel;
+	TSharedPtr<SCurveEditor>CurveEditor;
+	float InputSnap;
+	float OutputSnap;
+};
\ No newline at end of file
diff --git a/Engine/Plugins/FX/Niagara/Source/NiagaraEditor/Private/Widgets/SNiagaraEmitterHeader.cpp b/Engine/Plugins/FX/Niagara/Source/NiagaraEditor/Private/Widgets/SNiagaraEmitterHeader.cpp
new file mode 100644
--- /dev/null
+++ b/Engine/Plugins/FX/Niagara/Source/NiagaraEditor/Private/Widgets/SNiagaraEmitterHeader.cpp
@@ -0,0 +1,82 @@
+// Copyright 1998-2016 Epic Games, Inc. All Rights Reserved.
+
+#include "SNiagaraEmitterHeader.h"
+#include "NiagaraEmitterHandleViewModel.h"
+#include "NiagaraEmitterViewModel.h"
+#include "NiagaraEditorStyle.h"
+#include "SInlineEditableTextBlock.h"
+#include "SBoxPanel.h"
+#include "SCheckBox.h"
+#include "STextBlock.h"
+
+#define LOCTEXT_NAMESPACE "NiagaraEmitterHeader"
+
+void SNiagaraEmitterHeader::Construct(const FArguments& InArgs, TSharedRef<FNiagaraEmitterHandleViewModel> InViewModel)
+{
+	ViewModel = InViewModel;
+
+	TSharedRef<SWidget> AdditionalHeaderContent = InArgs._AdditionalHeaderContent.Widget;
+
+	ChildSlot
+	[
+		SNew(SVerticalBox)
+		//~ Enabled check box rename text box, and external header controls.
+		+ SVerticalBox::Slot()
+		.AutoHeight()
+		.HAlign(HAlign_Fill)
+		[
+			//~ Enabled
+			SNew(SHorizontalBox)
+			+ SHorizontalBox::Slot()
+			.HAlign(HAlign_Left)
+			.AutoWidth()
+			.Padding(2)
+			[
+				SNew(SCheckBox)
+				.ToolTipText(LOCTEXT("EnabledToolTip", "Toggles whether this emitter is enabled. Disabled emitters don't simulate or render."))
+				.IsChecked(ViewModel.ToSharedRef(), &FNiagaraEmitterHandleViewModel::GetIsEnabledCheckState)
+				.OnCheckStateChanged(ViewModel.ToSharedRef(), &FNiagaraEmitterHandleViewModel::OnIsEnabledCheckStateChanged)
+			]
+
+			//~ Name
+			+ SHorizontalBox::Slot()
+			.AutoWidth()
+			.Padding(2)
+			.HAlign(HAlign_Left)
+			[
+				SNew(SInlineEditableTextBlock)
+				.ToolTipText(NSLOCTEXT("NiagaraEmitterEditor", "NameTextToolTip", "Click to edit the emitter name."))
+				.Style(FNiagaraEditorStyle::Get(), "NiagaraEditor.HeadingInlineEditableText")
+				.WrapTextAt(150.0f)
+				.Text(ViewModel.ToSharedRef(), &FNiagaraEmitterHandleViewModel::GetNameText)
+				.OnTextCommitted(ViewModel.ToSharedRef(), &FNiagaraEmitterHandleViewModel::OnNameTextComitted)
+				.OnVerifyTextChanged(ViewModel.ToSharedRef(), &FNiagaraEmitterHandleViewModel::VerifyNameTextChanged)
+			]
+
+			//~ External controls.
+			+ SHorizontalBox::Slot()
+			.HAlign(HAlign_Right)
+			[
+				AdditionalHeaderContent
+			]
+		]
+
+		//~ Stats
+		+ SVerticalBox::Slot()
+		.AutoHeight()
+		.HAlign(HAlign_Fill)
+		.Padding(2)
+		[
+			SNew(STextBlock)
+			.Text(ViewModel->GetEmitterViewModel(), &FNiagaraEmitterViewModel::GetStatsText)
+		]
+	];
+}
+
+FReply SNiagaraEmitterHeader::OnMouseButtonDoubleClick(const FGeometry& InMyGeometry, const FPointerEvent& InMouseEvent)
+{
+	ViewModel->OpenSourceEmitter();
+	return FReply::Handled();
+}
+
+#undef LOCTEXT_NAMESPACE // "NiagaraEmitterHeader"
\ No newline at end of file
diff --git a/Engine/Plugins/FX/Niagara/Source/NiagaraEditor/Private/Widgets/SNiagaraEmitterHeader.h b/Engine/Plugins/FX/Niagara/Source/NiagaraEditor/Private/Widgets/SNiagaraEmitterHeader.h
new file mode 100644
--- /dev/null
+++ b/Engine/Plugins/FX/Niagara/Source/NiagaraEditor/Private/Widgets/SNiagaraEmitterHeader.h
@@ -0,0 +1,30 @@
+// Copyright 1998-2016 Epic Games, Inc. All Rights Reserved.
+
+#pragma once
+#include "SharedPointer.h"
+#include "SCompoundWidget.h"
+#include "DeclarativeSyntaxSupport.h"
+
+class FNiagaraEmitterHandleViewModel;
+
+/** A widget for viewing and editing header information for an emitter. */
+class SNiagaraEmitterHeader : public SCompoundWidget
+{
+public:
+	SLATE_BEGIN_ARGS(SNiagaraEmitterHeader)
+	{}
+
+		/** A slot which can be used to add additional content for the header. */
+		SLATE_NAMED_SLOT(FArguments, AdditionalHeaderContent)
+
+	SLATE_END_ARGS();
+
+	void Construct(const FArguments& InArgs, TSharedRef<FNiagaraEmitterHandleViewModel> InViewModel);
+
+	//~ SWidget interface
+	virtual FReply OnMouseButtonDoubleClick(const FGeometry& InMyGeometry, const FPointerEvent& InMouseEvent) override;
+
+private:
+	/** The view model which exposes the data used by the widget. */
+	TSharedPtr<FNiagaraEmitterHandleViewModel> ViewModel;
+};
\ No newline at end of file
diff --git a/Engine/Plugins/FX/Niagara/Source/NiagaraEditor/Private/Widgets/SNiagaraGeneratedCodeView.cpp b/Engine/Plugins/FX/Niagara/Source/NiagaraEditor/Private/Widgets/SNiagaraGeneratedCodeView.cpp
new file mode 100644
--- /dev/null
+++ b/Engine/Plugins/FX/Niagara/Source/NiagaraEditor/Private/Widgets/SNiagaraGeneratedCodeView.cpp
@@ -0,0 +1,442 @@
+// Copyright 1998-2017 Epic Games, Inc. All Rights Reserved.
+
+#include "SNiagaraGeneratedCodeView.h"
+#include "Textures/SlateIcon.h"
+#include "Framework/Commands/UIAction.h"
+#include "Framework/MultiBox/MultiBoxBuilder.h"
+#include "Widgets/Input/SSearchBox.h"
+#include "Widgets/Input/SButton.h"
+#include "ISequencer.h"
+#include "NiagaraSystemViewModel.h"
+#include "NiagaraEmitterHandleViewModel.h"
+#include "NiagaraEmitterHandle.h"
+#include "NiagaraEmitter.h"
+#include "NiagaraScript.h"
+#include "NiagaraSystemScriptViewModel.h"
+#include "EditorStyleSet.h"
+#include "SScrollBox.h"
+#include "UObjectGlobals.h"
+#include "Class.h"
+#include "Package.h"
+#include "SequencerSettings.h"
+#include "NiagaraSystem.h"
+#include "NiagaraEditorStyle.h"
+#include "PlatformApplicationMisc.h"
+
+#define LOCTEXT_NAMESPACE "NiagaraGeneratedCodeView"
+
+void SNiagaraGeneratedCodeView::Construct(const FArguments& InArgs, TSharedRef<FNiagaraSystemViewModel> InSystemViewModel)
+{
+	TabState = 0;
+	ScriptEnum = FindObject<UEnum>(ANY_PACKAGE, TEXT("ENiagaraScriptUsage"));
+	ensure(ScriptEnum);
+
+	SystemViewModel = InSystemViewModel;
+	SystemViewModel->OnSelectedEmitterHandlesChanged().AddRaw(this, &SNiagaraGeneratedCodeView::SelectedEmitterHandlesChanged);
+	SystemViewModel->GetSystemScriptViewModel()->OnSystemCompiled().AddRaw(this, &SNiagaraGeneratedCodeView::OnCodeCompiled);
+
+	this->ChildSlot
+		[
+			SNew(SVerticalBox)
+			+ SVerticalBox::Slot()
+			.AutoHeight()
+			[
+				SNew(SVerticalBox)
+				+ SVerticalBox::Slot()
+				.AutoHeight()
+				[
+					SNew(SHorizontalBox)
+					+ SHorizontalBox::Slot()
+					.AutoWidth()
+					.VAlign(VAlign_Center)
+					.HAlign(HAlign_Left)
+					[
+						SNew(SButton)
+						.HAlign(HAlign_Center)
+						.OnClicked(this, &SNiagaraGeneratedCodeView::OnCopyPressed)
+						.Text(LOCTEXT("CopyOutput", "Copy"))
+						.ToolTipText(LOCTEXT("CopyOutputToolitp", "Press this button to put the contents of this tab in the clipboard."))
+					]
+					+ SHorizontalBox::Slot()
+					[
+						SNullWidget::NullWidget
+					] 
+					+ SHorizontalBox::Slot()
+					.AutoWidth()
+					.HAlign(HAlign_Right)
+					.Padding(2, 4, 2, 4)
+					[
+						SAssignNew(SearchBox, SSearchBox)
+						.OnTextCommitted(this, &SNiagaraGeneratedCodeView::OnSearchTextCommitted)
+						.HintText(NSLOCTEXT("SearchBox", "HelpHint", "Search For Text"))
+						.OnTextChanged(this, &SNiagaraGeneratedCodeView::OnSearchTextChanged)
+						.SelectAllTextWhenFocused(false)
+						.DelayChangeNotificationsWhileTyping(true)
+						.MinDesiredWidth(200)
+					]
+					+ SHorizontalBox::Slot()
+					.AutoWidth()
+					.VAlign(VAlign_Center)
+					.Padding(2, 4, 2, 4)
+					[
+						SAssignNew(SearchFoundMOfNText, STextBlock)
+						.MinDesiredWidth(25)
+					]					
+					/*
+					+ SHorizontalBox::Slot()
+					.AutoWidth()
+					.Padding(2, 4, 2, 4)
+					[
+						SNew(SButton)
+						.ButtonStyle(FEditorStyle::Get(), "HoverHintOnly")
+						.IsFocusable(false)
+						.ToolTipText(LOCTEXT("UpToolTip", "Focus to next found search term"))
+						.OnClicked(this, &SNiagaraGeneratedCodeView::SearchUpClicked)
+						.Content()
+						[
+							SNew(STextBlock)
+							.Font(FEditorStyle::Get().GetFontStyle("FontAwesome.10"))
+							.Text(FText::FromString(FString(TEXT("\xf062"))))
+						]
+					]
+					+ SHorizontalBox::Slot()
+					.AutoWidth()
+					.Padding(2, 4, 2, 4)
+					[
+						SNew(SButton)
+						.ButtonStyle(FEditorStyle::Get(), "HoverHintOnly")
+						.IsFocusable(false)
+						.ToolTipText(LOCTEXT("DownToolTip", "Focus to next found search term"))
+						.OnClicked(this, &SNiagaraGeneratedCodeView::SearchDownClicked)
+						.Content()
+						[
+							SNew(STextBlock)
+							.Font(FEditorStyle::Get().GetFontStyle("FontAwesome.10"))
+							.Text(FText::FromString(FString(TEXT("\xf063"))))
+						]
+					]*/
+				]
+				+ SVerticalBox::Slot()
+				.AutoHeight()
+				[
+					SAssignNew(CheckBoxContainer, SHorizontalBox)
+					// We'll insert here on UI updating..
+				]
+			]
+			+ SVerticalBox::Slot()
+			[
+				SNew(STextBlock)
+				.Text(LOCTEXT("NoDataText", "Failed to compile or has not been compiled."))
+				.Visibility_Lambda([&]() {
+					if (TabHasScriptData())
+						return EVisibility::Collapsed;
+					return EVisibility::Visible;
+				})
+			]
+			+ SVerticalBox::Slot()
+			[
+				SAssignNew(TextBodyContainer, SVerticalBox)
+			]
+		];
+	UpdateUI();
+	DoSearch(SearchBox->GetText());
+}
+
+FReply SNiagaraGeneratedCodeView::SearchUpClicked()
+{
+	CurrentFoundTextEntry--;
+	if (CurrentFoundTextEntry < 0 && ActiveFoundTextEntries.Num() > 0)
+	{
+		CurrentFoundTextEntry = ActiveFoundTextEntries.Num() - 1;
+	}
+	else if (CurrentFoundTextEntry < 0 && ActiveFoundTextEntries.Num() == 0)
+	{
+		CurrentFoundTextEntry = 0;
+	}
+
+	if (ActiveFoundTextEntries.Num() > 0)
+	{
+		//GeneratedCode[TabState].Text->ScrollTo(ActiveFoundTextEntries[CurrentFoundTextEntry]);
+	}
+
+	SetSearchMofN();
+
+	return FReply::Handled();
+}
+
+FReply SNiagaraGeneratedCodeView::SearchDownClicked()
+{
+	CurrentFoundTextEntry++;
+	if (CurrentFoundTextEntry >= ActiveFoundTextEntries.Num())
+	{
+		CurrentFoundTextEntry = 0;
+	}
+
+	if (ActiveFoundTextEntries.Num() > 0)
+	{
+		//GeneratedCode[TabState].Text->ScrollTo(ActiveFoundTextEntries[CurrentFoundTextEntry]);
+	}
+	
+	SetSearchMofN();
+
+	return FReply::Handled();
+}
+
+FReply SNiagaraGeneratedCodeView::OnCopyPressed()
+{
+	if (TabState < (uint32)GeneratedCode.Num())
+	{
+		FPlatformApplicationMisc::ClipboardCopy(*GeneratedCode[TabState].Hlsl.ToString());
+	}
+	return FReply::Handled();
+}
+
+void SNiagaraGeneratedCodeView::OnSearchTextChanged(const FText& InFilterText)
+{
+	DoSearch(InFilterText);
+}
+
+void SNiagaraGeneratedCodeView::DoSearch(const FText& InFilterText)
+{
+	GeneratedCode[TabState].Text->SetHighlightText(InFilterText);
+	InFilterText.ToString();
+
+	FString SearchString = InFilterText.ToString();
+	ActiveFoundTextEntries.Empty();
+	if (SearchString.IsEmpty())
+	{
+		SetSearchMofN();
+		return;
+	}
+
+	ActiveFoundTextEntries.Empty();
+	CurrentFoundTextEntry = INDEX_NONE;
+	for (int32 i = 0; i < GeneratedCode[TabState].HlslByLines.Num(); i++)
+	{
+		int32 LastPos = INDEX_NONE;
+		int32 FoundPos = GeneratedCode[TabState].HlslByLines[i].Find(SearchString, ESearchCase::IgnoreCase, ESearchDir::FromStart, LastPos);
+		while (FoundPos != INDEX_NONE)
+		{
+			ActiveFoundTextEntries.Add(FTextLocation(i, FoundPos));
+			LastPos = FoundPos + 1;
+			FoundPos = GeneratedCode[TabState].HlslByLines[i].Find(SearchString, ESearchCase::IgnoreCase, ESearchDir::FromStart, LastPos);
+		}
+	}
+
+	if (ActiveFoundTextEntries.Num() > 0)
+	{
+		CurrentFoundTextEntry = 0;
+		//GeneratedCode[TabState].Text->ScrollTo(ActiveFoundTextEntries[CurrentFoundTextEntry]);
+	}
+
+	SetSearchMofN();
+}
+
+void SNiagaraGeneratedCodeView::SetSearchMofN()
+{
+	//SearchFoundMOfNText->SetText(FText::Format(LOCTEXT("MOfN", "{0} of {1}"), FText::AsNumber(CurrentFoundTextEntry), FText::AsNumber(ActiveFoundTextEntries.Num())));
+	SearchFoundMOfNText->SetText(FText::Format(LOCTEXT("MOfN", "{1} found"), FText::AsNumber(CurrentFoundTextEntry), FText::AsNumber(ActiveFoundTextEntries.Num())));
+}
+
+void SNiagaraGeneratedCodeView::OnSearchTextCommitted(const FText& InFilterText, ETextCommit::Type InCommitType)
+{
+	OnSearchTextChanged(InFilterText);
+}
+
+void SNiagaraGeneratedCodeView::OnCodeCompiled()
+{
+	UpdateUI();
+	DoSearch(SearchBox->GetText());
+}
+
+void SNiagaraGeneratedCodeView::SelectedEmitterHandlesChanged()
+{
+	UpdateUI();
+	DoSearch(SearchBox->GetText());
+}
+
+void SNiagaraGeneratedCodeView::UpdateUI()
+{
+	TArray<UNiagaraScript*> Scripts;
+	UNiagaraSystem& System = SystemViewModel->GetSystem();
+	Scripts.Add(System.GetSystemSpawnScript(false));
+	Scripts.Add(System.GetSystemUpdateScript(false));
+	Scripts.Add(System.GetSystemSpawnScript(true));
+	Scripts.Add(System.GetSystemUpdateScript(true));
+	
+	TArray<TSharedRef<FNiagaraEmitterHandleViewModel>> SelectedEmitterHandles;
+	SystemViewModel->GetSelectedEmitterHandles(SelectedEmitterHandles);
+	if (SelectedEmitterHandles.Num() == 1)
+	{
+		FNiagaraEmitterHandle* Handle = SelectedEmitterHandles[0]->GetEmitterHandle();
+		if (Handle)
+		{
+			TArray<UNiagaraScript*> EmitterScripts;
+			Handle->GetInstance()->GetScripts(EmitterScripts);
+			Scripts.Append(EmitterScripts);
+		}
+	}
+
+	GeneratedCode.SetNum(Scripts.Num());
+
+	if (TabState >= (uint32)GeneratedCode.Num())
+	{
+		TabState = 0;
+	}
+
+	CheckBoxContainer->ClearChildren();
+	TextBodyContainer->ClearChildren();
+
+	for (int32 i = 0; i < GeneratedCode.Num(); i++)
+	{
+		GeneratedCode[i].Usage = Scripts[i]->Usage;
+		GeneratedCode[i].UsageIndex = Scripts[i]->UsageIndex;
+
+		TArray<FString> OutputByLines;
+		GeneratedCode[i].Hlsl = FText::GetEmpty();
+		FString SumString;
+		Scripts[i]->LastHlslTranslation.ParseIntoArrayLines(OutputByLines, false);
+		GeneratedCode[i].HlslByLines.SetNum(OutputByLines.Num());
+		for (int32 k = 0; k < OutputByLines.Num(); k++)
+		{
+			GeneratedCode[i].HlslByLines[k] = FString::Printf(TEXT("/*%04d*/\t\t%s\r\n"), k, *OutputByLines[k]);
+			SumString.Append(GeneratedCode[i].HlslByLines[k]);
+		}
+		GeneratedCode[i].Hlsl = FText::FromString(SumString);
+		
+		if (Scripts[i]->Usage == ENiagaraScriptUsage::ParticleEventScript)
+		{
+			GeneratedCode[i].UsageName = FText::Format(LOCTEXT("UsageNameEvent", "{0}[{1}]"), ScriptEnum->GetDisplayNameTextByValue((int64)Scripts[i]->Usage), FText::AsNumber(Scripts[i]->UsageIndex));
+		}
+		else
+		{
+			GeneratedCode[i].UsageName = FText::Format(LOCTEXT("UsageName", "{0}"), ScriptEnum->GetDisplayNameTextByValue((int64)Scripts[i]->Usage));
+		}
+
+		if (!GeneratedCode[i].HorizontalScrollBar.IsValid())
+		{
+			GeneratedCode[i].HorizontalScrollBar = SNew(SScrollBar)
+				.Orientation(Orient_Horizontal)
+				.Thickness(FVector2D(8.0f, 8.0f));
+		}
+
+		if (!GeneratedCode[i].VerticalScrollBar.IsValid())
+		{
+			GeneratedCode[i].VerticalScrollBar = SNew(SScrollBar)
+				.Orientation(Orient_Vertical)
+				.Thickness(FVector2D(8.0f, 8.0f));
+		}
+
+		if (!GeneratedCode[i].CheckBox.IsValid())
+		{
+			SAssignNew(GeneratedCode[i].CheckBox, SCheckBox)
+				//.Style(FEditorStyle::Get(), "PlacementBrowser.Tab")
+				//.Style(FEditorStyle::Get(), i == 0 ? "Property.ToggleButton.Start" : (i < GeneratedCode.Num() - 1 ? "Property.ToggleButton.Middle" : "Property.ToggleButton.End"))
+				.OnCheckStateChanged(this, &SNiagaraGeneratedCodeView::OnTabChanged, (uint32)i)
+				.IsChecked(this, &SNiagaraGeneratedCodeView::GetTabCheckedState, (uint32)i)
+				.Style(FEditorStyle::Get(), "PlacementBrowser.Tab")
+				[
+					SNew(SOverlay)
+					+ SOverlay::Slot()
+					.Padding(FMargin(6, 0, 15, 0))
+					.VAlign(VAlign_Center)
+					[
+						SAssignNew(GeneratedCode[i].TextName, STextBlock)
+						.TextStyle(FNiagaraEditorStyle::Get(), "NiagaraEditor.CodeView.Checkbox.Text")
+						.Text(GeneratedCode[i].UsageName)
+					]
+				];
+		}
+		else
+		{
+			GeneratedCode[i].TextName->SetText(GeneratedCode[i].UsageName);
+		}
+
+		if (!GeneratedCode[i].Container.IsValid())
+		{
+			SAssignNew(GeneratedCode[i].Container, SVerticalBox)
+				.Visibility(this, &SNiagaraGeneratedCodeView::GetViewVisibility, (uint32)i)
+				+ SVerticalBox::Slot()
+				[
+					SNew(SHorizontalBox)
+					+ SHorizontalBox::Slot()
+					.FillWidth(1.0f)
+					[
+						SNew(SScrollBox)
+						.Orientation(Orient_Horizontal)
+						.ExternalScrollbar(GeneratedCode[i].HorizontalScrollBar)
+						+ SScrollBox::Slot()
+						[
+							SNew(SScrollBox)
+							.Orientation(Orient_Vertical)
+							.ExternalScrollbar(GeneratedCode[i].VerticalScrollBar)
+							+ SScrollBox::Slot()
+							[
+								SAssignNew(GeneratedCode[i].Text, STextBlock)
+								.TextStyle(FNiagaraEditorStyle::Get(), "NiagaraEditor.CodeView.Hlsl.Normal")
+							]
+						]
+					]
+					+ SHorizontalBox::Slot()
+					.AutoWidth()
+					[
+						GeneratedCode[i].VerticalScrollBar.ToSharedRef()
+					]
+				]
+				+ SVerticalBox::Slot()
+				.AutoHeight()
+				[
+					GeneratedCode[i].HorizontalScrollBar.ToSharedRef()
+				];
+		}
+
+		GeneratedCode[i].Text->SetText(GeneratedCode[i].Hlsl);
+
+		CheckBoxContainer->AddSlot()
+			[
+				GeneratedCode[i].CheckBox.ToSharedRef()
+			];
+		TextBodyContainer->AddSlot()
+			[
+				GeneratedCode[i].Container.ToSharedRef()
+			];		
+	}
+}
+
+SNiagaraGeneratedCodeView::~SNiagaraGeneratedCodeView()
+{
+	if (SystemViewModel.IsValid())
+	{
+		SystemViewModel->OnSelectedEmitterHandlesChanged().RemoveAll(this);
+		SystemViewModel->GetSystemScriptViewModel()->OnSystemCompiled().RemoveAll(this);
+	}
+	
+}
+
+void SNiagaraGeneratedCodeView::OnTabChanged(ECheckBoxState State, uint32 Tab)
+{
+	if (State == ECheckBoxState::Checked)
+	{
+		TabState = Tab;
+		DoSearch(SearchBox->GetText());
+	}
+}
+
+
+bool SNiagaraGeneratedCodeView::TabHasScriptData() const
+{
+	return !GeneratedCode[TabState].Hlsl.IsEmpty();
+}
+
+ECheckBoxState SNiagaraGeneratedCodeView::GetTabCheckedState(uint32 Tab) const
+{
+	return TabState == Tab ? ECheckBoxState::Checked : ECheckBoxState::Unchecked;
+}
+
+EVisibility SNiagaraGeneratedCodeView::GetViewVisibility(uint32 Tab) const
+{
+	return TabState == Tab ? EVisibility::Visible : EVisibility::Collapsed;
+}
+
+
+#undef LOCTEXT_NAMESPACE
\ No newline at end of file
diff --git a/Engine/Plugins/FX/Niagara/Source/NiagaraEditor/Private/Widgets/SNiagaraGeneratedCodeView.h b/Engine/Plugins/FX/Niagara/Source/NiagaraEditor/Private/Widgets/SNiagaraGeneratedCodeView.h
new file mode 100644
--- /dev/null
+++ b/Engine/Plugins/FX/Niagara/Source/NiagaraEditor/Private/Widgets/SNiagaraGeneratedCodeView.h
@@ -0,0 +1,84 @@
+// Copyright 1998-2016 Epic Games, Inc. All Rights Reserved.
+
+#pragma once
+
+#include "EditorUndoClient.h"
+#include "SCompoundWidget.h"
+#include "DeclarativeSyntaxSupport.h"
+#include "STableRow.h"
+#include "STableViewBase.h"
+#include "STreeView.h"
+#include "NiagaraSystemViewModel.h"
+#include "TickableEditorObject.h"
+#include "WeakObjectPtrTemplates.h"
+#include "NiagaraDataSet.h"
+#include "SharedPointer.h"
+#include "Map.h"
+#include "SCheckBox.h"
+#include "SSearchBox.h"
+#include "SMultiLineEditableTextBox.h"
+
+class SNiagaraGeneratedCodeView : public SCompoundWidget
+{
+public:
+	SLATE_BEGIN_ARGS(SNiagaraGeneratedCodeView)
+	{}
+
+	SLATE_END_ARGS();
+
+	void Construct(const FArguments& InArgs, TSharedRef<FNiagaraSystemViewModel> InSystemViewModel);
+	virtual ~SNiagaraGeneratedCodeView();
+
+	void OnCodeCompiled();
+
+protected:
+
+	void UpdateUI();
+
+	struct TabInfo
+	{
+		FText UsageName;
+		FText Hlsl;
+		ENiagaraScriptUsage Usage;
+		int32 UsageIndex;
+
+		TArray<FString> HlslByLines;
+		TSharedPtr<STextBlock> Text;
+		TSharedPtr<STextBlock> TextName;
+		TSharedPtr < SCheckBox > CheckBox;
+		TSharedPtr<SScrollBar> HorizontalScrollBar;
+		TSharedPtr<SScrollBar> VerticalScrollBar;
+		TSharedPtr<SVerticalBox> Container;
+	};
+
+	TArray<TabInfo> GeneratedCode;
+	TSharedPtr<SHorizontalBox> CheckBoxContainer;
+	TSharedPtr<SVerticalBox> TextBodyContainer;
+	TSharedPtr<SSearchBox> SearchBox;
+	TSharedPtr<STextBlock> SearchFoundMOfNText;
+	TArray<FTextLocation> ActiveFoundTextEntries;
+	int32 CurrentFoundTextEntry;
+
+	void SetSearchMofN();
+
+	void SelectedEmitterHandlesChanged();
+	void OnTabChanged(ECheckBoxState State,	uint32 Tab);
+	ECheckBoxState GetTabCheckedState(uint32 Tab) const;
+	EVisibility GetViewVisibility(uint32 Tab) const;
+	bool TabHasScriptData() const;
+	FReply OnCopyPressed();
+	void OnSearchTextChanged(const FText& InFilterText);
+	void OnSearchTextCommitted(const FText& InFilterText, ETextCommit::Type InCommitType);
+	FReply SearchUpClicked();
+	FReply SearchDownClicked();
+
+	void DoSearch(const FText& InFilterText);
+
+	uint32 TabState;
+
+	TSharedPtr<FNiagaraSystemViewModel> SystemViewModel;
+
+	UEnum* ScriptEnum;
+};
+
+
diff --git a/Engine/Plugins/FX/Niagara/Source/NiagaraEditor/Private/Widgets/SNiagaraGraphNodeWriteDataSet.cpp b/Engine/Plugins/FX/Niagara/Source/NiagaraEditor/Private/Widgets/SNiagaraGraphNodeWriteDataSet.cpp
new file mode 100644
--- /dev/null
+++ b/Engine/Plugins/FX/Niagara/Source/NiagaraEditor/Private/Widgets/SNiagaraGraphNodeWriteDataSet.cpp
@@ -0,0 +1,68 @@
+// Copyright 1998-2016 Epic Games, Inc. All Rights Reserved.
+
+#include "SNiagaraGraphNodeWriteDataSet.h"
+#include "NiagaraNodeInput.h"
+#include "NiagaraGraph.h"
+#include "Framework/Commands/Commands.h"
+#include "Framework/MultiBox/MultiBoxBuilder.h"
+#include "Widgets/SBoxPanel.h"
+#include "Widgets/Input/SButton.h"
+#include "Widgets/Images/SImage.h"
+#include "ScopedTransaction.h"
+#include "Editor/PropertyEditor/Public/PropertyEditorModule.h"
+#include "ISinglePropertyView.h"
+#include "NiagaraNodeWriteDataSet.h"
+#include "ModuleManager.h"
+#include "GraphEditorSettings.h"
+
+#define LOCTEXT_NAMESPACE "SNiagaraGraphNodeWriteDataSet"
+
+
+void SNiagaraGraphNodeWriteDataSet::Construct(const FArguments& InArgs, UEdGraphNode* InGraphNode)
+{
+	GraphNode = InGraphNode;
+
+	this->UpdateGraphNode();
+}
+
+TSharedRef<SWidget> SNiagaraGraphNodeWriteDataSet::CreateNodeContentArea()
+{
+	FSinglePropertyParams InitParams;
+	InitParams.NamePlacement = EPropertyNamePlacement::Hidden;
+	FPropertyEditorModule& PropertyEditorModule = FModuleManager::GetModuleChecked<FPropertyEditorModule>("PropertyEditor");
+	TSharedPtr<ISinglePropertyView> SinglePropView = PropertyEditorModule.CreateSingleProperty(GraphNode, GET_MEMBER_NAME_CHECKED(UNiagaraNodeWriteDataSet, EventName), InitParams);
+
+	TSharedRef<SWidget> ContentAreaWidget = SGraphNode::CreateNodeContentArea();
+	TSharedPtr<SVerticalBox> VertContainer = SNew(SVerticalBox)
+		+ SVerticalBox::Slot()
+		.AutoHeight()
+		.VAlign(VAlign_Center)
+		[
+			SNew(SHorizontalBox)
+			+ SHorizontalBox::Slot()
+			.HAlign(HAlign_Center)
+			.AutoWidth()
+			.Padding(Settings->GetInputPinPadding())
+			[
+				SNew(STextBlock)
+				.Text(LOCTEXT("EventName","Event Name"))
+			]
+			+ SHorizontalBox::Slot()
+			.HAlign(HAlign_Center)
+			.AutoWidth()
+			[
+				SinglePropView.ToSharedRef()
+			]
+		]
+		+ SVerticalBox::Slot()
+		.AutoHeight()
+		[
+			ContentAreaWidget
+		];
+	TSharedRef<SWidget> RetWidget = VertContainer.ToSharedRef();
+	return RetWidget;
+}
+
+
+
+#undef LOCTEXT_NAMESPACE
\ No newline at end of file
diff --git a/Engine/Plugins/FX/Niagara/Source/NiagaraEditor/Private/Widgets/SNiagaraGraphNodeWriteDataSet.h b/Engine/Plugins/FX/Niagara/Source/NiagaraEditor/Private/Widgets/SNiagaraGraphNodeWriteDataSet.h
new file mode 100644
--- /dev/null
+++ b/Engine/Plugins/FX/Niagara/Source/NiagaraEditor/Private/Widgets/SNiagaraGraphNodeWriteDataSet.h
@@ -0,0 +1,18 @@
+// Copyright 1998-2016 Epic Games, Inc. All Rights Reserved.
+
+#pragma once
+
+#include "SGraphNode.h"
+#include "SMenuAnchor.h"
+#include "UICommandList.h"
+
+/** A graph node widget representing a niagara write data set node. */
+class SNiagaraGraphNodeWriteDataSet : public SGraphNode
+{
+public:
+	SLATE_BEGIN_ARGS(SNiagaraGraphNodeWriteDataSet) {}
+	SLATE_END_ARGS();
+
+	void Construct(const FArguments& InArgs, UEdGraphNode* InGraphNode);
+	virtual TSharedRef<SWidget> CreateNodeContentArea() override;
+};
\ No newline at end of file
diff --git a/Engine/Plugins/FX/Niagara/Source/NiagaraEditor/Private/Widgets/SNiagaraGraphParameterMapGetNode.cpp b/Engine/Plugins/FX/Niagara/Source/NiagaraEditor/Private/Widgets/SNiagaraGraphParameterMapGetNode.cpp
new file mode 100644
--- /dev/null
+++ b/Engine/Plugins/FX/Niagara/Source/NiagaraEditor/Private/Widgets/SNiagaraGraphParameterMapGetNode.cpp
@@ -0,0 +1,147 @@
+// Copyright 1998-2016 Epic Games, Inc. All Rights Reserved.
+
+#include "SNiagaraGraphParameterMapGetNode.h"
+#include "NiagaraNodeParameterMapGet.h"
+#include "SButton.h"
+#include "GraphEditorSettings.h"
+#include "DrawElements.h"
+#include "Widgets/SBoxPanel.h"
+#include "SGraphPin.h"
+
+#define LOCTEXT_NAMESPACE "SNiagaraGraphParameterMapGetNode"
+
+
+void SNiagaraGraphParameterMapGetNode::Construct(const FArguments& InArgs, UEdGraphNode* InGraphNode)
+{
+	BackgroundBrush = FEditorStyle::GetBrush("Graph.Pin.Background");
+	BackgroundHoveredBrush = FEditorStyle::GetBrush("PlainBorder");
+
+	GraphNode = InGraphNode;
+	UpdateGraphNode();
+}
+
+void SNiagaraGraphParameterMapGetNode::AddPin(const TSharedRef<SGraphPin>& PinToAdd)
+{
+	PinToAdd->SetOwner(SharedThis(this));
+
+	const UEdGraphPin* PinObj = PinToAdd->GetPinObj();
+	const bool bAdvancedParameter = (PinObj != nullptr) && PinObj->bAdvancedView;
+	if (bAdvancedParameter)
+	{
+		PinToAdd->SetVisibility(TAttribute<EVisibility>(PinToAdd, &SGraphPin::IsPinVisibleAsAdvanced));
+	}
+
+	// Save the UI building for later...
+	if (PinToAdd->GetDirection() == EEdGraphPinDirection::EGPD_Input)
+	{
+		InputPins.Add(PinToAdd);
+	}
+	else // Direction == EEdGraphPinDirection::EGPD_Output
+	{
+		OutputPins.Add(PinToAdd);
+	}	
+}
+
+TSharedRef<SWidget> SNiagaraGraphParameterMapGetNode::CreateNodeContentArea()
+{
+	// NODE CONTENT AREA
+	return SNew(SBorder)
+		.BorderImage(FEditorStyle::GetBrush("NoBorder"))
+		.HAlign(HAlign_Fill)
+		.VAlign(VAlign_Fill)
+		.Padding(FMargin(0, 3))
+		[
+			SAssignNew(PinContainerRoot, SVerticalBox)
+		];
+}
+
+void SNiagaraGraphParameterMapGetNode::CreatePinWidgets()
+{
+	SGraphNode::CreatePinWidgets();
+	
+	UNiagaraNodeParameterMapGet* GetNode = Cast<UNiagaraNodeParameterMapGet>(GraphNode);
+
+	// Deferred pin adding to line up input/output pins by name.
+	for (int32 i = 0; i < OutputPins.Num() + 1; i++)
+	{
+		SVerticalBox::FSlot& Slot = PinContainerRoot->AddSlot();
+		Slot.AutoHeight();
+
+		// Get nodes have an unequal number of pins. 
+		TSharedPtr<SWidget> Widget;
+		if (i == 0)
+		{
+
+			SAssignNew(Widget, SHorizontalBox)
+				.Visibility(EVisibility::Visible)
+				+ SHorizontalBox::Slot()
+				.HAlign(HAlign_Left)
+				.FillWidth(1.0f)
+				.Padding(Settings->GetInputPinPadding())
+				[
+					(InputPins.Num() > 0 ? InputPins[0] : SNullWidget::NullWidget)
+				]
+			+ SHorizontalBox::Slot()
+				.AutoWidth()
+				.HAlign(HAlign_Right)
+				.Padding(Settings->GetOutputPinPadding())
+				[
+					SNullWidget::NullWidget
+				];
+		}
+		else
+		{
+			TSharedRef<SGraphPin> OutputPin = OutputPins[i - 1];
+			UEdGraphPin* SrcOutputPin = OutputPin->GetPinObj();
+
+			UEdGraphPin* MatchingInputPin = GetNode->GetDefaultPin(SrcOutputPin);
+
+			TSharedPtr<SWidget> InputPin = SNullWidget::NullWidget;
+			for (TSharedRef<SGraphPin> Pin : InputPins)
+			{
+				UEdGraphPin* SrcInputPin = Pin->GetPinObj();
+				if (SrcInputPin == MatchingInputPin)
+				{
+					InputPin = Pin;
+					Pin->SetShowLabel(false);
+				}
+			}
+
+			SAssignNew(Widget, SHorizontalBox)
+				.Visibility(EVisibility::Visible)
+				+ SHorizontalBox::Slot()
+				.HAlign(HAlign_Left)
+				.FillWidth(1.0f)
+				.Padding(Settings->GetInputPinPadding())
+				[
+					InputPin.ToSharedRef()
+				]
+				+ SHorizontalBox::Slot()
+				.AutoWidth()
+				.HAlign(HAlign_Right)
+				.Padding(Settings->GetOutputPinPadding())
+				[
+					OutputPin
+				];
+		}
+
+		TSharedPtr<SBorder> Border;
+		SAssignNew(Border, SBorder)
+			.HAlign(HAlign_Fill)
+			.VAlign(VAlign_Fill)
+			.Padding(FMargin(0, 3))
+			[
+				Widget.ToSharedRef()
+			];
+		Border->SetBorderImage(TAttribute<const FSlateBrush*>::Create(TAttribute<const FSlateBrush*>::FGetter::CreateRaw(this, &SNiagaraGraphParameterMapGetNode::GetBackgroundBrush, Widget)));
+
+		Slot.AttachWidget(Border.ToSharedRef());
+	}
+}
+
+
+const FSlateBrush* SNiagaraGraphParameterMapGetNode::GetBackgroundBrush(TSharedPtr<SWidget> Border) const
+{
+	return Border->IsHovered() ? BackgroundHoveredBrush	: BackgroundBrush;
+}
+#undef LOCTEXT_NAMESPACE
diff --git a/Engine/Plugins/FX/Niagara/Source/NiagaraEditor/Private/Widgets/SNiagaraGraphPinAdd.cpp b/Engine/Plugins/FX/Niagara/Source/NiagaraEditor/Private/Widgets/SNiagaraGraphPinAdd.cpp
new file mode 100644
--- /dev/null
+++ b/Engine/Plugins/FX/Niagara/Source/NiagaraEditor/Private/Widgets/SNiagaraGraphPinAdd.cpp
@@ -0,0 +1,103 @@
+// Copyright 1998-2016 Epic Games, Inc. All Rights Reserved.
+
+#include "SNiagaraGraphPinAdd.h"
+#include "NiagaraNodeWithDynamicPins.h"
+
+#include "ScopedTransaction.h"
+#include "SComboButton.h"
+#include "SBoxPanel.h"
+#include "SImage.h"
+#include "MultiBoxBuilder.h"
+#include "SEditableTextBox.h"
+#include "SInlineEditableTextBlock.h"
+#include "SBox.h"
+
+#define LOCTEXT_NAMESPACE "NiagaraGraphPinAdd"
+
+void SNiagaraGraphPinAdd::Construct(const FArguments& InArgs, UEdGraphPin* InGraphPinObj)
+{
+	SetShowLabel(false);
+	OwningNode = Cast<UNiagaraNodeWithDynamicPins>(InGraphPinObj->GetOwningNode());
+	SGraphPin::Construct(SGraphPin::FArguments(), InGraphPinObj);
+	TSharedPtr<SHorizontalBox> PinBox = GetFullPinHorizontalRowWidget().Pin();
+	if (PinBox.IsValid())
+	{
+		if (InGraphPinObj->Direction == EGPD_Input)
+		{
+			PinBox->AddSlot()
+			[
+				ConstructAddButton()
+			];
+		}
+		else
+		{
+			PinBox->InsertSlot(0)
+			[
+				ConstructAddButton()
+			];
+		}
+	}
+}
+
+TSharedRef<SWidget>	SNiagaraGraphPinAdd::ConstructAddButton()
+{
+	return SNew(SComboButton)
+		.HasDownArrow(false)
+		.ButtonStyle(FEditorStyle::Get(), "HoverHintOnly")
+		.ForegroundColor(FSlateColor::UseForeground())
+		.OnGetMenuContent(this, &SNiagaraGraphPinAdd::OnGetAddButtonMenuContent)
+		.ContentPadding(FMargin(2))
+		.HAlign(HAlign_Center)
+		.VAlign(VAlign_Center)
+		.ToolTipText(LOCTEXT("AddPinButtonToolTip", "Connect this pin to add a new typed pin, or choose from the drop-down."))
+		.ButtonContent()
+		[
+			SNew(SImage)
+			.ColorAndOpacity(FSlateColor::UseForeground())
+			.Image(FEditorStyle::GetBrush("Plus"))
+		];
+}
+
+TSharedRef<SWidget> SNiagaraGraphPinAdd::OnGetAddButtonMenuContent()
+{
+	TArray<UEdGraphPin*> Pins = GetPinObj()->GetOwningNode()->GetAllPins();
+	EEdGraphPinDirection PinDir = GetPinObj()->Direction;
+
+	int32 FirstPinSameDir = Pins.IndexOfByPredicate([&](UEdGraphPin* InObj) -> bool
+	{
+		return InObj && (InObj->Direction == PinDir);
+	});
+
+	int32 IndexOfPin = Pins.IndexOfByPredicate([&](UEdGraphPin* InObj) -> bool
+	{
+		return InObj && (InObj == GetPinObj());
+	});
+
+	int PinIdx = 0;
+	if (FirstPinSameDir != -1 && IndexOfPin != -1)
+	{
+		PinIdx = IndexOfPin - FirstPinSameDir;
+	}
+
+	FString WorkingName = FString::Printf(TEXT("%s%d"), PinDir == EEdGraphPinDirection::EGPD_Input ? TEXT("Input") : TEXT("Output"), PinIdx);
+
+	if (OwningNode != nullptr)
+	{
+		return OwningNode->GenerateAddPinMenu(WorkingName, this);
+	}
+	else
+	{
+		return SNullWidget::NullWidget;
+	}
+
+}
+
+void SNiagaraGraphPinAdd::OnAddType(FNiagaraVariable InAdd)
+{
+	if (OwningNode != nullptr)
+	{
+		FScopedTransaction AddNewPinTransaction(LOCTEXT("AddNewPinTransaction", "Add pin to node"));
+		OwningNode->RequestNewTypedPin(GetPinObj()->Direction, InAdd.GetType(), InAdd.GetName().ToString());
+	}
+}
+#undef LOCTEXT_NAMESPACE
diff --git a/Engine/Plugins/FX/Niagara/Source/NiagaraEditor/Private/Widgets/SNiagaraSelectedEmitterHandle.cpp b/Engine/Plugins/FX/Niagara/Source/NiagaraEditor/Private/Widgets/SNiagaraSelectedEmitterHandle.cpp
new file mode 100644
--- /dev/null
+++ b/Engine/Plugins/FX/Niagara/Source/NiagaraEditor/Private/Widgets/SNiagaraSelectedEmitterHandle.cpp
@@ -0,0 +1,109 @@
+// Copyright 1998-2016 Epic Games, Inc. All Rights Reserved.
+
+#include "SNiagaraSelectedEmitterHandle.h"
+#include "SNiagaraEmitterHeader.h"
+#include "NiagaraEditorStyle.h"
+#include "ViewModels/NiagaraSystemViewModel.h"
+#include "ViewModels/NiagaraEmitterHandleViewModel.h"
+#include "ViewModels/NiagaraEmitterViewModel.h"
+#include "ViewModels/Stack/NiagaraStackViewModel.h"
+#include "NiagaraEditorModule.h"
+
+#include "Modules/ModuleManager.h"
+#include "Widgets/Layout/SBox.h"
+#include "Widgets/Text/STextBlock.h"
+
+#define LOCTEXT_NAMESPACE "NiagaraSelectedEmitterHandle"
+
+void SNiagaraSelectedEmitterHandle::Construct(const FArguments& InArgs, TSharedRef<FNiagaraSystemViewModel> InSystemViewModel)
+{
+	SystemViewModel = InSystemViewModel;
+	SystemViewModel->OnSelectedEmitterHandlesChanged().AddRaw(this, &SNiagaraSelectedEmitterHandle::SelectedEmitterHandlesChanged);
+	StackViewModel = NewObject<UNiagaraStackViewModel>();
+	FNiagaraEditorModule NiagaraEditorModule = FModuleManager::GetModuleChecked<FNiagaraEditorModule>(TEXT("NiagaraEditor"));
+
+	ChildSlot
+	[
+		SNew(SOverlay)
+		+ SOverlay::Slot()
+		[
+			SNew(SVerticalBox)
+			+ SVerticalBox::Slot()
+			.AutoHeight()
+			[
+				SAssignNew(EmitterHeaderContainer, SBox)
+			]
+			+ SVerticalBox::Slot()
+			[
+				NiagaraEditorModule.CreateStackWidget(StackViewModel)
+			]
+		]
+		+ SOverlay::Slot()
+		.Padding(0, 20, 0, 0)
+		.HAlign(HAlign_Center)
+		[
+			SNew(STextBlock)
+			.Text(this, &SNiagaraSelectedEmitterHandle::GetUnsupportedSelectionText)
+			.TextStyle(FNiagaraEditorStyle::Get(), "NiagaraEditor.SelectedEmitter.UnsupportedSelectionText")
+			.Visibility(this, &SNiagaraSelectedEmitterHandle::GetUnsupportedSelectionTextVisibility)
+		]
+	];
+
+	RefreshEmitterWidgets();
+}
+
+SNiagaraSelectedEmitterHandle::~SNiagaraSelectedEmitterHandle()
+{
+	SystemViewModel->OnEmitterHandleViewModelsChanged().RemoveAll(this);
+	SystemViewModel->OnSelectedEmitterHandlesChanged().RemoveAll(this);
+	StackViewModel->Initialize(nullptr, nullptr);
+}
+
+void SNiagaraSelectedEmitterHandle::AddReferencedObjects(FReferenceCollector& Collector)
+{
+	Collector.AddReferencedObject(StackViewModel);
+}
+
+void SNiagaraSelectedEmitterHandle::RefreshEmitterWidgets()
+{
+	TArray<TSharedRef<FNiagaraEmitterHandleViewModel>> SelectedEmitterHandles;
+	SystemViewModel->GetSelectedEmitterHandles(SelectedEmitterHandles);
+	if (SelectedEmitterHandles.Num() == 1)
+	{
+		EmitterHeaderContainer->SetContent(SNew(SNiagaraEmitterHeader, SelectedEmitterHandles[0]));
+		StackViewModel->Initialize(SystemViewModel, SelectedEmitterHandles[0]->GetEmitterViewModel());
+	}
+	else
+	{
+		EmitterHeaderContainer->SetContent(SNullWidget::NullWidget);
+		StackViewModel->Initialize(SystemViewModel, nullptr);
+	}
+}
+
+void SNiagaraSelectedEmitterHandle::SelectedEmitterHandlesChanged()
+{
+	RefreshEmitterWidgets();
+}
+
+EVisibility SNiagaraSelectedEmitterHandle::GetUnsupportedSelectionTextVisibility() const
+{
+	return SystemViewModel->GetSelectedEmitterHandleIds().Num() != 1 ? EVisibility::Visible : EVisibility::Collapsed;
+}
+
+FText SNiagaraSelectedEmitterHandle::GetUnsupportedSelectionText() const
+{
+	if (SystemViewModel->GetSelectedEmitterHandleIds().Num() == 0)
+	{
+		return LOCTEXT("NoSelectionMessage", "Select an emitter in the timeline.");
+	}
+	else if (SystemViewModel->GetSelectedEmitterHandleIds().Num() > 1)
+	{
+		return LOCTEXT("MultipleSelectionMessage", "Multiple selected emitters are not currently supported.");
+	}
+	else
+	{
+		return FText();
+	}
+}
+
+#undef LOCTEXT_NAMESPACE
\ No newline at end of file
diff --git a/Engine/Plugins/FX/Niagara/Source/NiagaraEditor/Private/Widgets/SNiagaraSelectedEmitterHandle.h b/Engine/Plugins/FX/Niagara/Source/NiagaraEditor/Private/Widgets/SNiagaraSelectedEmitterHandle.h
new file mode 100644
--- /dev/null
+++ b/Engine/Plugins/FX/Niagara/Source/NiagaraEditor/Private/Widgets/SNiagaraSelectedEmitterHandle.h
@@ -0,0 +1,43 @@
+// Copyright 1998-2016 Epic Games, Inc. All Rights Reserved.
+
+#pragma once
+
+#include "SCompoundWidget.h"
+#include "DeclarativeSyntaxSupport.h"
+#include "UObject/GCObject.h"
+
+class FNiagaraSystemViewModel;
+class UNiagaraStackViewModel;
+class SNiagaraStack;
+class SBox;
+
+class SNiagaraSelectedEmitterHandle : public SCompoundWidget, public FGCObject
+{
+public:
+	SLATE_BEGIN_ARGS(SNiagaraSelectedEmitterHandle)
+	{}
+	SLATE_END_ARGS();
+
+	void Construct(const FArguments& InArgs, TSharedRef<FNiagaraSystemViewModel> InSystemViewModel);
+
+	~SNiagaraSelectedEmitterHandle();
+
+	virtual void AddReferencedObjects(FReferenceCollector& Collector) override;
+
+private:
+	void SelectedEmitterHandlesChanged();
+
+	void RefreshEmitterWidgets();
+
+	EVisibility GetUnsupportedSelectionTextVisibility() const;
+
+	FText GetUnsupportedSelectionText() const;
+
+private:
+	TSharedPtr<FNiagaraSystemViewModel> SystemViewModel;
+
+	UNiagaraStackViewModel* StackViewModel;
+
+	TSharedPtr<SBox> EmitterHeaderContainer;
+	TSharedPtr<SNiagaraStack> NiagaraStack;
+};
\ No newline at end of file
diff --git a/Engine/Plugins/FX/Niagara/Source/NiagaraEditor/Private/Widgets/SNiagaraSelectedObjectsDetails.cpp b/Engine/Plugins/FX/Niagara/Source/NiagaraEditor/Private/Widgets/SNiagaraSelectedObjectsDetails.cpp
new file mode 100644
--- /dev/null
+++ b/Engine/Plugins/FX/Niagara/Source/NiagaraEditor/Private/Widgets/SNiagaraSelectedObjectsDetails.cpp
@@ -0,0 +1,42 @@
+// Copyright 1998-2016 Epic Games, Inc. All Rights Reserved.
+
+#include "SNiagaraSelectedObjectsDetails.h"
+#include "NiagaraObjectSelection.h"
+
+#include "ModuleManager.h"
+#include "Editor/PropertyEditor/Public/PropertyEditorModule.h"
+#include "IDetailsView.h"
+
+#define LOCTEXT_NAMESPACE "NiagaraSelectedObjectsDetails"
+
+void SNiagaraSelectedObjectsDetails::Construct(const FArguments& InArgs, TSharedRef<FNiagaraObjectSelection> InSelectedObjects)
+{
+	SelectedObjects = InSelectedObjects;
+	SelectedObjects->OnSelectedObjectsChanged().AddSP(this, &SNiagaraSelectedObjectsDetails::SelectedObjectsChanged);
+
+	FPropertyEditorModule& PropertyEditorModule = FModuleManager::GetModuleChecked<FPropertyEditorModule>("PropertyEditor");
+	FDetailsViewArgs DetailsViewArgs(false, false, true, FDetailsViewArgs::HideNameArea, true);
+	DetailsView = PropertyEditorModule.CreateDetailView(DetailsViewArgs);
+	DetailsView->SetObjects(SelectedObjects->GetSelectedObjects().Array());
+	DetailsView->OnFinishedChangingProperties().AddRaw(this, &SNiagaraSelectedObjectsDetails::OnDetailsPanelFinishedChangingProperties);
+
+	ChildSlot
+	[
+		DetailsView.ToSharedRef()
+	];
+}
+
+void SNiagaraSelectedObjectsDetails::SelectedObjectsChanged()
+{
+	DetailsView->SetObjects(SelectedObjects->GetSelectedObjects().Array());
+}
+
+void SNiagaraSelectedObjectsDetails::OnDetailsPanelFinishedChangingProperties(const FPropertyChangedEvent& InEvent)
+{
+	if (OnFinishedChangingPropertiesDelegate.IsBound())
+	{
+		OnFinishedChangingPropertiesDelegate.Broadcast(InEvent);
+	}
+}
+
+#undef LOCTEXT_NAMESPACE // "NiagaraSelectedObjectsDetails"
\ No newline at end of file
