Copying //UE4/Orion-Staging to //UE4/Main (Source: //Orion/Dev-General @ 3606378)

diff --git a/Engine/Plugins/FX/Niagara/Source/NiagaraEditor/Private/Widgets/SNiagaraSpreadsheetView.cpp b/Engine/Plugins/FX/Niagara/Source/NiagaraEditor/Private/Widgets/SNiagaraSpreadsheetView.cpp
new file mode 100644
--- /dev/null
+++ b/Engine/Plugins/FX/Niagara/Source/NiagaraEditor/Private/Widgets/SNiagaraSpreadsheetView.cpp
@@ -0,0 +1,807 @@
+// Copyright 1998-2017 Epic Games, Inc. All Rights Reserved.
+
+#include "SNiagaraSpreadsheetView.h"
+#include "Textures/SlateIcon.h"
+#include "Framework/Commands/UIAction.h"
+#include "Framework/MultiBox/MultiBoxBuilder.h"
+#include "Widgets/Input/SSearchBox.h"
+#include "Widgets/Input/SButton.h"
+#include "ISequencer.h"
+#include "NiagaraSystemViewModel.h"
+#include "NiagaraEmitterHandleViewModel.h"
+#include "NiagaraEmitterHandle.h"
+#include "NiagaraEmitter.h"
+#include "NiagaraScript.h"
+#include "HAL/PlatformApplicationMisc.h"
+#include "EditorStyleSet.h"
+#include "SScrollBox.h"
+#include "UObjectGlobals.h"
+#include "Class.h"
+#include "Package.h"
+#include "SequencerSettings.h"
+#include "NiagaraSystem.h"
+#include "NiagaraEditorStyle.h"
+#include "PlatformApplicationMisc.h"
+
+#define LOCTEXT_NAMESPACE "SNiagaraSpreadsheetView"
+#define ARRAY_INDEX_COLUMN_NAME TEXT("Array Index")
+#define KEY_COLUMN_NAME TEXT("Property")
+#define VALUE_COLUMN_NAME TEXT("Value")
+
+void SNiagaraSpreadsheetRow::Construct(const FArguments& InArgs, const TSharedRef<STableViewBase>& InOwnerTableView)
+{
+	RowIndex = InArgs._RowIndex;
+	DataSet = InArgs._DataSet;
+	ColumnsAreAttributes = InArgs._ColumnsAreAttributes;
+
+	SupportedFields = InArgs._SupportedFields;
+	FieldInfoMap = InArgs._FieldInfoMap;
+
+	SMultiColumnTableRow< TSharedPtr<int32> >::Construct(FSuperRowType::FArguments(), InOwnerTableView);
+}
+
+
+TSharedRef< SWidget > SNiagaraSpreadsheetRow::GenerateWidgetForColumn(const FName& ColumnName)
+{
+	TSharedPtr<SWidget> EntryWidget;
+	const SNiagaraSpreadsheetView::FieldInfo* FieldInfo = nullptr;
+	int32 RealRowIdx = 0;
+	if (ColumnsAreAttributes && ColumnName == ARRAY_INDEX_COLUMN_NAME)
+	{
+		EntryWidget = SNew(STextBlock)
+			.Text(FText::AsNumber(RowIndex));
+	}
+	else if (!ColumnsAreAttributes && ColumnName == KEY_COLUMN_NAME)
+	{
+		EntryWidget = SNew(STextBlock)
+			.Text(FText::FromName((*SupportedFields.Get())[RowIndex]));
+	}
+	else if (ColumnsAreAttributes)
+	{
+		FieldInfo = FieldInfoMap->Find(ColumnName);
+		RealRowIdx = RowIndex;
+	}
+	else if (!ColumnsAreAttributes && ColumnName == VALUE_COLUMN_NAME)
+	{
+		FieldInfo = FieldInfoMap->Find((*SupportedFields.Get())[RowIndex]);
+	}
+
+	if (FieldInfo != nullptr && DataSet != nullptr && !EntryWidget.IsValid())
+	{
+		
+		if (FieldInfo->bFloat)
+		{
+			uint32 CompBufferOffset = FieldInfo->FloatStartOffset;
+			float* Src = DataSet->PrevData().GetInstancePtrFloat(CompBufferOffset, RealRowIdx);
+
+			EntryWidget = SNew(STextBlock)
+				.Text(FText::AsNumber(Src[0]));
+
+		}
+		else if (FieldInfo->bBoolean)
+		{
+			uint32 CompBufferOffset = FieldInfo->IntStartOffset;
+			int32* Src = DataSet->PrevData().GetInstancePtrInt32(CompBufferOffset, RealRowIdx);
+			FText ValueText;
+			if (Src[0] == 0)
+			{
+				ValueText = LOCTEXT("NiagaraFalse", "False(0)");
+			}
+			else if (Src[0] == -1)
+			{
+				ValueText = LOCTEXT("NiagaraTrue", "True(-1)");
+			}
+			else
+			{
+				ValueText = FText::Format(LOCTEXT("NiagaraUnknown", "Invalid({0}"), FText::AsNumber(Src[0]));
+			}
+			EntryWidget = SNew(STextBlock)
+				.Text(ValueText);
+
+		}
+		else if (FieldInfo->Enum.IsValid())
+		{
+			uint32 CompBufferOffset = FieldInfo->IntStartOffset;
+			int32* Src = DataSet->PrevData().GetInstancePtrInt32(CompBufferOffset, RealRowIdx);
+			EntryWidget = SNew(STextBlock)
+				.Text(FText::Format(LOCTEXT("EnumValue","{0}({1})"), FieldInfo->Enum->GetDisplayNameTextByValue(Src[0]), FText::AsNumber(Src[0])));
+
+		}
+		else
+		{
+			uint32 CompBufferOffset = FieldInfo->IntStartOffset;
+			int32* Src = DataSet->PrevData().GetInstancePtrInt32(CompBufferOffset, RealRowIdx);
+			EntryWidget = SNew(STextBlock)
+				.Text(FText::AsNumber(Src[0]));
+
+		}
+	}
+	else if (!EntryWidget.IsValid())
+	{
+		EntryWidget = SNew(STextBlock)
+			.Text(LOCTEXT("UnsupportedColumn", "n/a"));
+	}
+
+	return SNew(SBorder)
+		.BorderImage(FEditorStyle::GetBrush("NoBorder"))
+		.Padding(3)
+		.HAlign(EHorizontalAlignment::HAlign_Center)
+		[
+			EntryWidget.ToSharedRef()
+		];
+}
+
+/////////////////////////////////////////////////////////////////////////////////////////////////////////////
+
+void SNiagaraSpreadsheetView::Construct(const FArguments& InArgs, TSharedRef<FNiagaraSystemViewModel> InSystemViewModel)
+{
+	TabState = UIPerParticleUpdate;
+	ScriptEnum = FindObject<UEnum>(ANY_PACKAGE, TEXT("ENiagaraScriptUsage"));
+	ensure(ScriptEnum);
+	
+	CaptureData.SetNum(2);
+	CaptureData[(int)UIPerParticleUpdate].TargetUsage = ENiagaraScriptUsage::ParticleUpdateScript;
+	CaptureData[(int)UISystemUpdate].TargetUsage = ENiagaraScriptUsage::SystemUpdateScript;
+	CaptureData[(int)UIPerParticleUpdate].ColumnName = LOCTEXT("PerParticleUpdate", "Per-Particle Update");
+	CaptureData[(int)UISystemUpdate].ColumnName = LOCTEXT("SystemUpdate", "System Update");
+	CaptureData[(int)UIPerParticleUpdate].bColumnsAreAttributes = true;
+	CaptureData[(int)UISystemUpdate].bColumnsAreAttributes = false;
+
+	SystemViewModel = InSystemViewModel;
+	SystemViewModel->OnSelectedEmitterHandlesChanged().AddRaw(this, &SNiagaraSpreadsheetView::SelectedEmitterHandlesChanged);
+	SystemViewModel->OnPostSequencerTimeChanged().AddRaw(this, &SNiagaraSpreadsheetView::OnSequencerTimeChanged);
+
+	for (int32 i = 0; i < (int32)UIMax; i++)
+	{
+		CaptureData[i].HorizontalScrollBar = SNew(SScrollBar)
+			.Orientation(Orient_Horizontal)
+			.Thickness(FVector2D(8.0f, 8.0f));
+
+		CaptureData[i].VerticalScrollBar = SNew(SScrollBar)
+			.Orientation(Orient_Vertical)
+			.Thickness(FVector2D(8.0f, 8.0f));
+
+
+		SAssignNew(CaptureData[i].ListView, STreeView< TSharedPtr<int32> >)
+			.IsEnabled(this, &SNiagaraSpreadsheetView::IsPausedAtRightTimeOnRightHandle)
+			// List view items are this tall
+			.ItemHeight(12)
+			// Tell the list view where to get its source data
+			.TreeItemsSource(&CaptureData[i].SupportedIndices)
+			// When the list view needs to generate a widget for some data item, use this method
+			.OnGenerateRow(this, &SNiagaraSpreadsheetView::OnGenerateWidgetForList, (EUITab)i)
+			// Given some DataItem, this is how we find out if it has any children and what they are.
+			.OnGetChildren(this, &SNiagaraSpreadsheetView::OnGetChildrenForList, (EUITab)i)
+			// Selection mode
+			.SelectionMode(ESelectionMode::Single)
+			.ExternalScrollbar(CaptureData[i].VerticalScrollBar)
+			.ConsumeMouseWheel(EConsumeMouseWheel::Always)
+			.AllowOverscroll(EAllowOverscroll::No)
+			// Selection callback
+			.OnSelectionChanged(this, &SNiagaraSpreadsheetView::OnEventSelectionChanged, (EUITab)i)
+			.HeaderRow
+			(
+				SAssignNew(CaptureData[i].HeaderRow, SHeaderRow)
+			);
+
+		SAssignNew(CaptureData[i].CheckBox, SCheckBox)
+			//.Style(FEditorStyle::Get(), "PlacementBrowser.Tab")
+			.Style(FEditorStyle::Get(), i == 0 ? "Property.ToggleButton.Start" : (i < CaptureData.Num() - 1 ? "Property.ToggleButton.Middle" : "Property.ToggleButton.End"))
+			.OnCheckStateChanged(this, &SNiagaraSpreadsheetView::OnTabChanged, (EUITab)i)
+			.IsChecked(this, &SNiagaraSpreadsheetView::GetTabCheckedState, (EUITab)i)
+			[
+				SNew(SOverlay)
+				+ SOverlay::Slot()
+				.Padding(FMargin(6, 0, 15, 0))
+				.VAlign(VAlign_Center)
+				[
+					SNew(STextBlock)
+					//.TextStyle(FEditorStyle::Get(), "PlacementBrowser.Tab.Text")
+					.TextStyle(FNiagaraEditorStyle::Get(), "NiagaraEditor.AttributeSpreadsheetTabText")
+					.Text(CaptureData[i].ColumnName)
+				]
+			];
+
+		SAssignNew(CaptureData[i].Container, SVerticalBox)
+			.Visibility(this, &SNiagaraSpreadsheetView::GetViewVisibility, (EUITab)i)
+			+ SVerticalBox::Slot()
+			[
+				SNew(SHorizontalBox)
+				+ SHorizontalBox::Slot()
+				[
+					SNew(SScrollBox)
+					.Orientation(Orient_Horizontal)
+					.ExternalScrollbar(CaptureData[i].HorizontalScrollBar)
+					+ SScrollBox::Slot()
+					[
+						CaptureData[i].ListView.ToSharedRef()
+					]
+				]
+				+ SHorizontalBox::Slot()
+				.AutoWidth()
+				[
+					CaptureData[i].VerticalScrollBar.ToSharedRef()
+				]
+			]
+			+ SVerticalBox::Slot()
+			.AutoHeight()
+			[
+				CaptureData[i].HorizontalScrollBar.ToSharedRef()
+			];
+		}
+
+		this->ChildSlot
+		[
+			SNew(SVerticalBox)
+			+ SVerticalBox::Slot()
+			.AutoHeight()
+			[
+				SNew(SBorder)
+				.BorderImage(FEditorStyle::GetBrush(TEXT("ToolPanel.GroupBorder")))
+				[
+					SNew(SHorizontalBox)
+
+					// Toolbar
+					+ SHorizontalBox::Slot()
+					.AutoWidth()
+					.VAlign(VAlign_Center)
+					.HAlign(HAlign_Left)
+					[
+						SNew(SButton)
+						.HAlign(HAlign_Center)
+						.OnClicked(this, &SNiagaraSpreadsheetView::OnCaptureRequestPressed)
+						.Text(LOCTEXT("CaptureLabel", "Capture"))
+						.ToolTipText(LOCTEXT("CaptureToolitp", "Press this button to capture one frame's contents. Can only capture CPU systems."))
+						.IsEnabled(this, &SNiagaraSpreadsheetView::CanCapture)
+					]
+
+					+ SHorizontalBox::Slot()
+					.AutoWidth()
+					.VAlign(VAlign_Center)
+					.HAlign(HAlign_Left)
+					[
+						SNew(SButton)
+						.HAlign(HAlign_Center)
+						.OnClicked(this, &SNiagaraSpreadsheetView::OnCSVOutputPressed)
+						.Text(LOCTEXT("CSVOutput", "Copy For Excel"))
+						.ToolTipText(LOCTEXT("CSVOutputToolitp", "Press this button to put the contents of this spreadsheet in the clipboard in an Excel-friendly format."))
+						.IsEnabled(this, &SNiagaraSpreadsheetView::IsPausedAtRightTimeOnRightHandle)
+					]
+				]
+		]
+		+ SVerticalBox::Slot()
+		.AutoHeight()
+		[
+			SNew(SBorder)
+			.BorderImage(FEditorStyle::GetBrush(TEXT("ToolPanel.GroupBorder")))
+			[
+				SNew(SVerticalBox)
+				+ SVerticalBox::Slot()
+				[
+					SNew(STextBlock)
+					.Text(LOCTEXT("NoDataText", "Please press capture to examine data from a particular frame."))
+					.Visibility_Lambda([&]() {
+						if (IsPausedAtRightTimeOnRightHandle())
+							return EVisibility::Collapsed;
+						return EVisibility::Visible;
+					})
+				]
+				+ SVerticalBox::Slot()
+				[
+					SNew(STextBlock)
+					.Text(this, &SNiagaraSpreadsheetView::LastCapturedInfoText)
+					.Visibility_Lambda([&]() {
+						if (IsPausedAtRightTimeOnRightHandle())
+							return EVisibility::Visible;
+						return EVisibility::Collapsed;
+					})
+				]
+				+ SVerticalBox::Slot()
+				[
+					SNew(SHorizontalBox)
+					+ SHorizontalBox::Slot()
+					.AutoWidth()
+					[
+						CaptureData[UIPerParticleUpdate].CheckBox.ToSharedRef()
+					]
+					+ SHorizontalBox::Slot()
+					.AutoWidth()
+					[
+						CaptureData[UISystemUpdate].CheckBox.ToSharedRef()
+					]
+				]
+			]
+		]
+		+ SVerticalBox::Slot()
+		[
+			CaptureData[UIPerParticleUpdate].Container.ToSharedRef()
+		]
+		+ SVerticalBox::Slot()
+		[
+			CaptureData[UISystemUpdate].Container.ToSharedRef()
+		]
+	];
+}
+
+SNiagaraSpreadsheetView::~SNiagaraSpreadsheetView()
+{
+	if (SystemViewModel.IsValid())
+	{
+		SystemViewModel->OnSelectedEmitterHandlesChanged().RemoveAll(this);
+		SystemViewModel->OnPostSequencerTimeChanged().RemoveAll(this);
+	}
+}
+
+void SNiagaraSpreadsheetView::OnTabChanged(ECheckBoxState State, EUITab Tab)
+{
+	if (State == ECheckBoxState::Checked)
+	{
+		TabState = Tab;
+	}
+}
+
+ECheckBoxState SNiagaraSpreadsheetView::GetTabCheckedState(EUITab Tab) const
+{
+	return TabState == Tab ? ECheckBoxState::Checked : ECheckBoxState::Unchecked;
+}
+
+EVisibility SNiagaraSpreadsheetView::GetViewVisibility(EUITab Tab) const
+{
+	return TabState == Tab ? EVisibility::Visible : EVisibility::Collapsed;
+}
+
+
+TSharedRef< ITableRow > SNiagaraSpreadsheetView::OnGenerateWidgetForList(TSharedPtr<int32> InItem, const TSharedRef<STableViewBase>& OwnerTable, EUITab Tab)
+{
+	return SNew(SNiagaraSpreadsheetRow, OwnerTable)
+		.RowIndex(*InItem)
+		.ColumnsAreAttributes(CaptureData[(int32)Tab].bColumnsAreAttributes)
+		.DataSet(CaptureData[(int32)Tab].DataSet)
+		.SupportedFields(CaptureData[(int32)Tab].SupportedFields)
+		.FieldInfoMap(CaptureData[(int32)Tab].FieldInfoMap);
+}
+
+FText SNiagaraSpreadsheetView::LastCapturedInfoText() const
+{
+	TArray<TSharedRef<FNiagaraEmitterHandleViewModel>> SelectedEmitterHandles;
+	SystemViewModel->GetSelectedEmitterHandles(SelectedEmitterHandles);
+	if (SelectedEmitterHandles.Num() == 1 && SNiagaraSpreadsheetView::IsPausedAtRightTimeOnRightHandle() && CaptureData[(int32)TabState].DataSet != nullptr)
+	{
+		return FText::Format(LOCTEXT("LastCapturedInfoName", "Captured Emitter: \"{0}\"     # Particles: {1}     Local Time: {2}     Script Type: {3}"),
+			SelectedEmitterHandles[0]->GetNameText(),
+			FText::AsNumber(CaptureData[(int32)TabState].DataSet->PrevData().GetNumInstances()),
+			FText::AsNumber(CaptureData[(int32)TabState].LastCaptureTime),
+			ScriptEnum->GetDisplayNameTextByValue((int64)CaptureData[(int32)TabState].TargetUsage));
+	}
+
+	return LOCTEXT("LastCapturedHandleNameStale", "Captured Info: Out-of-date");
+}
+
+void SNiagaraSpreadsheetView::OnGetChildrenForList(TSharedPtr<int32> InItem, TArray<TSharedPtr<int32>>& OutChildren, EUITab Tab)
+{
+	OutChildren.Empty();
+}
+
+void SNiagaraSpreadsheetView::SelectedEmitterHandlesChanged()
+{
+	// Need to reset the attributes list...
+	for (int32 i = 0; i < (int32)UIMax; i++)
+	{
+		CaptureData[i].LastReadWriteId = -1;
+		CaptureData[i].DataSet = nullptr;
+		CaptureData[i].SupportedIndices.SetNum(0);
+		CaptureData[i].ListView->RequestTreeRefresh();
+	}
+	
+}
+
+FReply SNiagaraSpreadsheetView::OnCSVOutputPressed()
+{
+	if (CaptureData[(int32)TabState].SupportedFields.IsValid() && CaptureData[(int32)TabState].FieldInfoMap.IsValid() && IsPausedAtRightTimeOnRightHandle())
+	{
+		FString CSVOutput;
+		int32 SkipIdx = -1;
+		int32 NumWritten = 0;
+		TArray<const SNiagaraSpreadsheetView::FieldInfo*> FieldInfos;
+		FieldInfos.SetNum(CaptureData[(int32)TabState].SupportedFields->Num());
+		FString DelimiterString = TEXT("\t");
+		for (int32 i = 0; i < CaptureData[(int32)TabState].SupportedFields->Num(); i++)
+		{
+			FName Field = (*CaptureData[(int32)TabState].SupportedFields)[i];
+			if (Field == ARRAY_INDEX_COLUMN_NAME)
+			{
+				SkipIdx = i;
+				continue;
+			}
+
+			if (NumWritten != 0)
+			{
+				CSVOutput += DelimiterString;
+			}
+
+			FieldInfos[i] = CaptureData[(int32)TabState].FieldInfoMap->Find(Field);
+
+			CSVOutput += Field.ToString();
+			NumWritten++;
+		}
+
+		CSVOutput += "\r\n";
+
+		for (uint32 RowIndex = 0; RowIndex < CaptureData[(int32)TabState].DataSet->PrevData().GetNumInstances(); RowIndex++)
+		{
+			NumWritten = 0;
+			for (int32 i = 0; i < CaptureData[(int32)TabState].SupportedFields->Num(); i++)
+			{
+				if (i == SkipIdx)
+				{
+					continue;
+				}
+
+				if (NumWritten != 0)
+				{
+					CSVOutput += DelimiterString;
+				}
+				const SNiagaraSpreadsheetView::FieldInfo* FieldInfo = FieldInfos[i];
+				
+				if (FieldInfo != nullptr && CaptureData[(int32)TabState].DataSet != nullptr)
+				{
+					if (FieldInfo->bFloat)
+					{
+						uint32 CompBufferOffset = FieldInfo->FloatStartOffset;
+						float* Src = CaptureData[(int32)TabState].DataSet->PrevData().GetInstancePtrFloat(CompBufferOffset, RowIndex);
+						CSVOutput += FString::Printf(TEXT("%3.3f"), Src[0]);
+					}
+					else
+					{
+						uint32 CompBufferOffset = FieldInfo->IntStartOffset;
+						int32* Src = CaptureData[(int32)TabState].DataSet->PrevData().GetInstancePtrInt32(CompBufferOffset, RowIndex);
+						CSVOutput += FString::Printf(TEXT("%d"), Src[0]);
+					}
+				}
+				NumWritten++;
+			}
+
+			CSVOutput += "\r\n";
+		}
+
+		FPlatformApplicationMisc::ClipboardCopy(*CSVOutput);
+	}
+
+	return FReply::Handled();
+}
+void SNiagaraSpreadsheetView::OnSequencerTimeChanged()
+{
+	HandleTimeChange();
+}
+
+void SNiagaraSpreadsheetView::Tick(float DeltaTime)
+{
+	HandleTimeChange();
+}
+
+void SNiagaraSpreadsheetView::HandleTimeChange()
+{
+	for (int32 i = 0; i < (int32)UIMax; i++)
+	{
+		if (!CaptureData[i].DataSource.IsValid())
+		{
+			CaptureData[i].bAwaitingFrame = false;
+		}
+
+		if (CaptureData[i].bAwaitingFrame)
+		{
+
+			TArray<TSharedRef<FNiagaraEmitterHandleViewModel>> SelectedEmitterHandles;
+			SystemViewModel->GetSelectedEmitterHandles(SelectedEmitterHandles);
+			if (SelectedEmitterHandles.Num() == 1)
+			{
+				TArray<UNiagaraScript*> Scripts;
+				CaptureData[i].DataSource->GetScripts(Scripts);
+				Scripts.Add(SystemViewModel->GetSystem().GetSystemSpawnScript(true));
+				Scripts.Add(SystemViewModel->GetSystem().GetSystemUpdateScript(true));
+
+				UNiagaraScript* FoundScript = nullptr;
+				for (UNiagaraScript* Script : Scripts)
+				{
+					if (Script->IsEquivalentUsage(CaptureData[i].TargetUsage))
+					{
+						FoundScript = Script;
+						break;
+					}
+				}
+
+				float LocalCaptureTime = SystemViewModel->GetSequencer()->GetLocalTime();
+
+				if (FoundScript && FoundScript->GetDebuggerInfo().bRequestDebugFrame == false && CaptureData[i].LastReadWriteId != FoundScript->GetDebuggerInfo().DebugFrameLastWriteId)
+				{
+					CaptureData[i].LastReadWriteId = FoundScript->GetDebuggerInfo().DebugFrameLastWriteId;
+					CaptureData[i].DataSet = &FoundScript->GetDebuggerInfo().DebugFrame;
+					CaptureData[i].DataSet->Tick(ENiagaraSimTarget::CPUSim); // Force a buffer swap, from here out we read from PrevData.
+
+					CaptureData[i].LastCaptureTime = LocalCaptureTime;
+					ensure(CaptureData[i].LastCaptureTime == CaptureData[i].TargetCaptureTime);
+					CaptureData[i].LastCaptureHandleId = SelectedEmitterHandles[0]->GetId();
+
+					ResetColumns((EUITab)i);
+					ResetEntries((EUITab)i);
+
+					CaptureData[i].bAwaitingFrame = false;
+				}
+			}
+		}
+	}
+}
+
+bool SNiagaraSpreadsheetView::IsTickable() const
+{
+	return CaptureData[(int32)TabState].bAwaitingFrame;
+}
+
+TStatId SNiagaraSpreadsheetView::GetStatId() const
+{
+	RETURN_QUICK_DECLARE_CYCLE_STAT(SNiagaraSpreadsheetView, STATGROUP_Tickables);
+}
+
+bool SNiagaraSpreadsheetView::CanCapture() const
+{
+	TArray<TSharedRef<FNiagaraEmitterHandleViewModel>> SelectedEmitterHandles;
+	SystemViewModel->GetSelectedEmitterHandles(SelectedEmitterHandles);
+	if (SelectedEmitterHandles.Num() == 1)
+	{
+		FNiagaraEmitterHandle* Handle = SelectedEmitterHandles[0]->GetEmitterHandle();
+		if (Handle && Handle->GetInstance()->SimTarget == ENiagaraSimTarget::CPUSim)
+		{
+			return true;
+		}
+	}
+	return false;
+}
+
+
+bool SNiagaraSpreadsheetView::IsPausedAtRightTimeOnRightHandle() const
+{
+	TArray<TSharedRef<FNiagaraEmitterHandleViewModel>> SelectedEmitterHandles;
+	SystemViewModel->GetSelectedEmitterHandles(SelectedEmitterHandles);
+	if (SelectedEmitterHandles.Num() == 1)
+	{
+		return SystemViewModel->GetSequencer()->GetPlaybackStatus() == EMovieScenePlayerStatus::Stopped &&
+			CaptureData[(int32)TabState].LastCaptureTime == SystemViewModel->GetSequencer()->GetLocalTime() &&
+			CaptureData[(int32)TabState].LastCaptureHandleId == SelectedEmitterHandles[0]->GetId();
+	}
+	return false;
+}
+
+void SNiagaraSpreadsheetView::ResetEntries(EUITab Tab)
+{
+	if (CaptureData[(int32)Tab].DataSet)
+	{
+		int32 NumInstances = CaptureData[(int32)Tab].DataSet->GetPrevNumInstances();
+		if (!CaptureData[(int32)Tab].bColumnsAreAttributes)
+		{
+			NumInstances = CaptureData[(int32)Tab].SupportedFields->Num();
+		}
+
+		CaptureData[(int32)Tab].SupportedIndices.SetNum(NumInstances);
+
+		for (int32 i = 0; i < NumInstances; i++)
+		{
+			CaptureData[(int32)Tab].SupportedIndices[i] = MakeShared<int32>(i);
+		}
+
+		CaptureData[(int32)Tab].ListView->RequestTreeRefresh();
+	}
+}
+
+void SNiagaraSpreadsheetView::GenerateLayoutInfo(FNiagaraTypeLayoutInfo& Layout, const UScriptStruct* Struct, const UEnum* Enum,  FName BaseName, TArray<FName>& PropertyNames, TArray<SNiagaraSpreadsheetView::FieldInfo>& FieldInfo)
+{
+	TFieldIterator<UProperty> PropertyCountIt(Struct, EFieldIteratorFlags::IncludeSuper);
+	int32 NumProperties = 0;
+	for (; PropertyCountIt; ++PropertyCountIt)
+	{
+		NumProperties++;
+	}
+	
+	for (TFieldIterator<UProperty> PropertyIt(Struct, EFieldIteratorFlags::IncludeSuper); PropertyIt; ++PropertyIt)
+	{
+		UProperty* Property = *PropertyIt;
+		FName PropertyName = (NumProperties == 1) ? *(BaseName.ToString()) : *(BaseName.ToString() + "." + Property->GetName());
+		if (Property->IsA(UFloatProperty::StaticClass()))
+		{
+			SNiagaraSpreadsheetView::FieldInfo Info;
+			Info.bFloat = true;
+			Info.FloatStartOffset = Layout.FloatComponentRegisterOffsets.Num();
+			Info.IntStartOffset = UINT_MAX;
+			Info.Enum = nullptr;
+			FieldInfo.Add(Info);
+				
+			Layout.FloatComponentRegisterOffsets.Add(Layout.GetNumComponents());
+			Layout.FloatComponentByteOffsets.Add(Property->GetOffset_ForInternal());
+			PropertyNames.Add(PropertyName);
+		}
+		else if (Property->IsA(UIntProperty::StaticClass()) || Property->IsA(UBoolProperty::StaticClass()))
+		{
+			SNiagaraSpreadsheetView::FieldInfo Info;
+			Info.bFloat = false;
+			Info.bBoolean = Property->IsA(UBoolProperty::StaticClass());
+			Info.FloatStartOffset = -1.0f;
+			Info.IntStartOffset = Layout.Int32ComponentRegisterOffsets.Num();
+			Info.Enum = Enum;
+			FieldInfo.Add(Info);
+
+			Layout.Int32ComponentRegisterOffsets.Add(Layout.GetNumComponents());
+			Layout.Int32ComponentByteOffsets.Add(Property->GetOffset_ForInternal());
+			PropertyNames.Add(PropertyName);
+		}
+		else if (Property->IsA(UEnumProperty::StaticClass()))
+		{
+			UEnumProperty* EnumProp = CastChecked<UEnumProperty>(Property);
+			GenerateLayoutInfo(Layout, FNiagaraTypeDefinition::GetIntStruct(), EnumProp->GetEnum(), PropertyName, PropertyNames, FieldInfo);
+		}
+		else if (Property->IsA(UStructProperty::StaticClass()))
+		{
+			UStructProperty* StructProp = CastChecked<UStructProperty>(Property);
+			GenerateLayoutInfo(Layout, StructProp->Struct, nullptr, PropertyName, PropertyNames, FieldInfo);
+		}
+		else
+		{
+			check(false);
+		}
+	}
+}
+
+void SNiagaraSpreadsheetView::ResetColumns(EUITab Tab)
+{
+	int32 i = (int32)Tab;
+
+	if (CaptureData[(int32)i].DataSet)
+	{
+		CaptureData[(int32)i].HeaderRow->ClearColumns();
+
+
+		CaptureData[(int32)i].SupportedFields = MakeShared<TArray<FName> >();
+		CaptureData[(int32)i].FieldInfoMap = MakeShared<TMap<FName, FieldInfo> >();
+		uint32 TotalFloatComponents = 0;
+		uint32 TotalInt32Components = 0;
+
+		TArray<FNiagaraVariable> Variables = CaptureData[(int32)i].DataSet->GetVariables();
+
+
+		float ManualWidth = 75.0f;
+		TArray<FName> ColumnNames;
+
+		if (CaptureData[(int32)i].bColumnsAreAttributes)
+		{
+			ColumnNames.Add(ARRAY_INDEX_COLUMN_NAME);
+		}
+		else
+		{
+			ManualWidth = 300.0f;
+			ColumnNames.Add(KEY_COLUMN_NAME);
+			ColumnNames.Add(VALUE_COLUMN_NAME);
+		}
+
+		for (const FNiagaraVariable& Var : Variables)
+		{
+			FNiagaraTypeDefinition TypeDef = Var.GetType();
+			const UScriptStruct* Struct = TypeDef.GetScriptStruct();
+			const UEnum* Enum = TypeDef.GetEnum();
+
+			FNiagaraTypeLayoutInfo Layout;
+			TArray<FName> PropertyNames;
+			TArray<SNiagaraSpreadsheetView::FieldInfo> FieldInfos;
+
+			uint32 TotalFloatComponentsBeforeStruct = TotalFloatComponents;
+			uint32 TotalInt32ComponentsBeforeStruct = TotalInt32Components;
+
+			GenerateLayoutInfo(Layout, Struct, Enum, Var.GetName(), PropertyNames, FieldInfos);
+
+			for (int32 VarIdx = 0; VarIdx < PropertyNames.Num(); VarIdx++)
+			{
+				
+				{
+					if (FieldInfos[VarIdx].bFloat)
+					{
+						FieldInfos[VarIdx].FloatStartOffset += TotalFloatComponentsBeforeStruct;
+						TotalFloatComponents++;
+					}
+					else
+					{
+						FieldInfos[VarIdx].IntStartOffset += TotalInt32ComponentsBeforeStruct;
+						TotalInt32Components++;
+					}
+
+					CaptureData[(int32)i].SupportedFields->Add(PropertyNames[VarIdx]);
+					CaptureData[(int32)i].FieldInfoMap->Add(PropertyNames[VarIdx], FieldInfos[VarIdx]);
+				}
+
+				if (CaptureData[(int32)i].bColumnsAreAttributes)
+				{
+					ColumnNames.Add(PropertyNames[VarIdx]);
+				}
+			}
+		}
+
+
+		for (FName ColumnName : ColumnNames)
+		{
+			SHeaderRow::FColumn::FArguments ColumnArgs;
+
+			ColumnArgs
+				.ColumnId(ColumnName)
+				.DefaultLabel(FText::FromName(ColumnName))
+				.SortMode(EColumnSortMode::None)
+				.HAlignHeader(HAlign_Center)
+				.VAlignHeader(VAlign_Fill)
+				.HeaderContentPadding(TOptional<FMargin>(2.0f))
+				.HAlignCell(HAlign_Fill)
+				.VAlignCell(VAlign_Fill)
+				.ManualWidth(ManualWidth);
+			CaptureData[(int32)i].HeaderRow->AddColumn(ColumnArgs);
+		}
+
+		CaptureData[(int32)i].HeaderRow->ResetColumnWidths();
+		CaptureData[(int32)i].HeaderRow->RefreshColumns();
+		CaptureData[(int32)i].ListView->RequestTreeRefresh();
+	}
+}
+
+FReply SNiagaraSpreadsheetView::OnCaptureRequestPressed()
+{
+	float LocalTime = SystemViewModel->GetSequencer()->GetLocalTime();
+	float SnapInterval = SystemViewModel->GetSequencer()->GetSequencerSettings()->GetTimeSnapInterval();
+	float TargetCaptureTime = LocalTime + SnapInterval;
+
+	TArray<TSharedRef<FNiagaraEmitterHandleViewModel>> SelectedEmitterHandles;
+	SystemViewModel->GetSelectedEmitterHandles(SelectedEmitterHandles);
+	ensure(SelectedEmitterHandles.Num() == 1);
+
+	for (int32 i = 0; i < CaptureData.Num(); i++)
+	{
+		CaptureData[(int32)i].DataSource = SelectedEmitterHandles[0]->GetEmitterHandle()->GetInstance();
+		
+		TArray<UNiagaraScript*> Scripts;
+		CaptureData[(int32)i].DataSource->GetScripts(Scripts);
+		Scripts.Add(SystemViewModel->GetSystem().GetSystemSpawnScript(true));
+		Scripts.Add(SystemViewModel->GetSystem().GetSystemUpdateScript(true));
+
+		UNiagaraScript* FoundScript = nullptr;
+		for (UNiagaraScript* Script : Scripts)
+		{
+			if (Script->IsEquivalentUsage(CaptureData[(int32)i].TargetUsage))
+			{
+				FoundScript = Script;
+				break;
+			}
+		}
+
+		if (FoundScript)
+		{
+			FoundScript->GetDebuggerInfo().bRequestDebugFrame = true;
+			CaptureData[(int32)i].bAwaitingFrame = true;
+
+			CaptureData[(int32)i].TargetCaptureTime = TargetCaptureTime;			
+		}
+	}
+
+	if (SystemViewModel->GetSequencer()->GetPlaybackStatus() == EMovieScenePlayerStatus::Stopped)
+	{
+		SystemViewModel->GetSequencer()->SetLocalTime(TargetCaptureTime, STM_None);
+	}
+	else
+	{
+		SystemViewModel->GetSequencer()->SetPlaybackStatus(EMovieScenePlayerStatus::Stopped);
+		SystemViewModel->GetSequencer()->SetLocalTime(TargetCaptureTime, STM_None);
+	}
+
+	return FReply::Handled();
+}
+
+void SNiagaraSpreadsheetView::OnEventSelectionChanged(TSharedPtr<int32> Selection, ESelectInfo::Type /*SelectInfo*/, EUITab /*Tab*/)
+{
+	if (Selection.IsValid())
+	{
+		// Do nothing for now
+	}
+}
+
+#undef LOCTEXT_NAMESPACE
\ No newline at end of file
diff --git a/Engine/Plugins/FX/Niagara/Source/NiagaraEditor/Private/Widgets/SNiagaraSpreadsheetView.h b/Engine/Plugins/FX/Niagara/Source/NiagaraEditor/Private/Widgets/SNiagaraSpreadsheetView.h
new file mode 100644
--- /dev/null
+++ b/Engine/Plugins/FX/Niagara/Source/NiagaraEditor/Private/Widgets/SNiagaraSpreadsheetView.h
@@ -0,0 +1,173 @@
+// Copyright 1998-2016 Epic Games, Inc. All Rights Reserved.
+
+#pragma once
+
+#include "EditorUndoClient.h"
+#include "SCompoundWidget.h"
+#include "DeclarativeSyntaxSupport.h"
+#include "STableRow.h"
+#include "STableViewBase.h"
+#include "STreeView.h"
+#include "NiagaraSystemViewModel.h"
+#include "TickableEditorObject.h"
+#include "WeakObjectPtrTemplates.h"
+#include "NiagaraDataSet.h"
+#include "SharedPointer.h"
+#include "Map.h"
+#include "SCheckBox.h"
+
+class SNiagaraSpreadsheetView : public SCompoundWidget, public FTickableEditorObject
+{
+public:
+	SLATE_BEGIN_ARGS(SNiagaraSpreadsheetView)
+	{}
+
+	SLATE_END_ARGS();
+
+	void Construct(const FArguments& InArgs, TSharedRef<FNiagaraSystemViewModel> InSystemViewModel);
+	virtual ~SNiagaraSpreadsheetView();
+
+	virtual void Tick(float DeltaTime) override;
+	virtual bool IsTickable() const override;
+	virtual TStatId GetStatId() const override;
+
+	struct FieldInfo
+	{
+		uint32 FloatStartOffset;
+		uint32 IntStartOffset;
+		bool bFloat;
+		bool bBoolean;
+		TWeakObjectPtr<const UEnum> Enum;
+	};
+
+protected:
+
+	enum EUITab : uint8
+	{
+		UIPerParticleUpdate = 0,
+		UISystemUpdate,
+		UIMax
+	};
+
+	void OnTabChanged(ECheckBoxState State,	EUITab Tab);
+	ECheckBoxState GetTabCheckedState(EUITab Tab) const;
+	EVisibility GetViewVisibility(EUITab Tab) const;
+
+	void GenerateLayoutInfo(FNiagaraTypeLayoutInfo& Layout, const UScriptStruct* Struct, const UEnum* Enum, FName BaseName, TArray<FName>& PropertyNames, TArray<FieldInfo>& FieldInfo);
+
+	EUITab TabState;
+
+	struct CapturedUIData
+	{
+		CapturedUIData() : TargetUsage(ENiagaraScriptUsage::ParticleUpdateScript), LastCaptureTime( -FLT_MAX) , bAwaitingFrame(false), LastReadWriteId(-1),	DataSet( nullptr), bColumnsAreAttributes(true)
+		{ }
+
+		TSharedPtr<SHeaderRow> HeaderRow;
+		TSharedPtr < STreeView<TSharedPtr<int32> > > ListView;
+		TSharedPtr < SCheckBox > CheckBox;
+		TArray< TSharedPtr<int32> > SupportedIndices;
+		int32 LastReadWriteId;
+		FNiagaraDataSet* DataSet;
+		TSharedPtr<TArray<FName> > SupportedFields;
+		TSharedPtr<TMap<FName, FieldInfo> > FieldInfoMap;
+		ENiagaraScriptUsage TargetUsage;
+		bool bAwaitingFrame;
+		float LastCaptureTime;
+		float TargetCaptureTime;
+		FGuid LastCaptureHandleId;
+		TWeakObjectPtr<UNiagaraEmitter> DataSource;
+		TSharedPtr<SScrollBar> HorizontalScrollBar;
+		TSharedPtr<SScrollBar> VerticalScrollBar;
+		TSharedPtr<SVerticalBox> Container;
+		bool bColumnsAreAttributes;
+		FText ColumnName;
+	};
+
+	TArray<CapturedUIData> CaptureData;
+
+	TSharedPtr<FNiagaraSystemViewModel> SystemViewModel;
+
+	UEnum* ScriptEnum;
+
+	void HandleTimeChange();
+	
+	void OnSequencerTimeChanged();
+
+	FText LastCapturedInfoText() const;
+
+	FReply OnCaptureRequestPressed();
+	FReply OnCSVOutputPressed();
+
+	bool CanCapture() const;
+	bool IsPausedAtRightTimeOnRightHandle() const;
+
+	void SelectedEmitterHandlesChanged();
+
+	void ResetColumns(EUITab Tab);
+	void ResetEntries(EUITab Tab);
+
+	/**
+	 * Function called when the currently selected event in the list of thread events changes
+	 *
+	 * @param Selection Currently selected event
+	 * @param SelectInfo Provides context on how the selection changed
+	 */
+	void OnEventSelectionChanged(TSharedPtr<int32> Selection, ESelectInfo::Type SelectInfo, EUITab Tab);
+
+	/** 
+	 * Generates SEventItem widgets for the events tree
+	 *
+	 * @param InItem Event to generate SEventItem for
+	 * @param OwnerTable Owner Table
+	 */
+	TSharedRef< ITableRow > OnGenerateWidgetForList(TSharedPtr<int32> InItem, const TSharedRef< STableViewBase >& OwnerTable , EUITab Tab);
+
+	/** 
+	 * Given a profiler event, generates children for it
+	 *
+	 * @param InItem Event to generate children for
+	 * @param OutChildren Generated children
+	 */
+	void OnGetChildrenForList(TSharedPtr<int32> InItem, TArray<TSharedPtr<int32>>& OutChildren, EUITab Tab);
+};
+
+
+class SNiagaraSpreadsheetRow : public SMultiColumnTableRow<TSharedPtr<int32> >
+{
+public:
+	typedef TSharedPtr<TArray<FName> > NamesArray;
+	typedef TSharedPtr<TMap <FName, SNiagaraSpreadsheetView::FieldInfo> > FieldsMap;
+
+	SLATE_BEGIN_ARGS(SNiagaraSpreadsheetRow)
+		: _RowIndex(0), _DataSet(nullptr), _ColumnsAreAttributes(true)
+	{}
+	SLATE_ARGUMENT(int32, RowIndex)
+	SLATE_ARGUMENT(bool, ColumnsAreAttributes)
+	SLATE_ARGUMENT(FNiagaraDataSet*, DataSet)
+	SLATE_ARGUMENT(NamesArray, SupportedFields)
+	SLATE_ARGUMENT(FieldsMap, FieldInfoMap)
+	SLATE_END_ARGS()
+
+		/**
+		* Generates a widget for debug attribute list column
+		*
+		* @param InArgs   A declaration from which to construct the widget
+		*/
+		virtual TSharedRef< SWidget > GenerateWidgetForColumn(const FName& ColumnName) override;
+
+
+	/**
+	* Construct the widget
+	*
+	* @param InArgs   A declaration from which to construct the widget
+	*/
+	void Construct(const FArguments& InArgs, const TSharedRef<STableViewBase>& InOwnerTableView);
+
+private:
+	int32 RowIndex;
+	FNiagaraDataSet* DataSet;
+	TSharedPtr<TArray<FName> > SupportedFields;
+	TSharedPtr<TMap<FName, SNiagaraSpreadsheetView::FieldInfo> > FieldInfoMap;
+	bool ColumnsAreAttributes;
+};
+
diff --git a/Engine/Plugins/FX/Niagara/Source/NiagaraEditor/Private/Widgets/SNiagaraSystemScript.cpp b/Engine/Plugins/FX/Niagara/Source/NiagaraEditor/Private/Widgets/SNiagaraSystemScript.cpp
new file mode 100644
--- /dev/null
+++ b/Engine/Plugins/FX/Niagara/Source/NiagaraEditor/Private/Widgets/SNiagaraSystemScript.cpp
@@ -0,0 +1,30 @@
+// Copyright 1998-2016 Epic Games, Inc. All Rights Reserved.
+
+#include "SNiagaraSystemScript.h"
+#include "NiagaraSystemViewModel.h"
+#include "NiagaraSystemScriptViewModel.h"
+#include "NiagaraScriptViewModel.h"
+#include "NiagaraScriptInputCollectionViewModel.h"
+#include "SNiagaraParameterCollection.h"
+#include "SNiagaraScriptGraph.h"
+
+#include "SSplitter.h"
+
+void SNiagaraSystemScript::Construct(const FArguments& InArgs, TSharedRef<FNiagaraSystemViewModel> InSystemViewModel)
+{
+	SystemViewModel = InSystemViewModel;
+	ChildSlot
+	[
+		SNew(SSplitter)
+		+ SSplitter::Slot()
+		.Value(0.3f)
+		[
+			SNew(SNiagaraParameterCollection, SystemViewModel->GetSystemScriptViewModel()->GetInputCollectionViewModel())
+		]
+		+ SSplitter::Slot()
+		.Value(0.7f)
+		[
+			SNew(SNiagaraScriptGraph, SystemViewModel->GetSystemScriptViewModel()->GetGraphViewModel())
+		]
+	];
+}
\ No newline at end of file
diff --git a/Engine/Plugins/FX/Niagara/Source/NiagaraEditor/Private/Widgets/SNiagaraSystemViewportToolBar.cpp b/Engine/Plugins/FX/Niagara/Source/NiagaraEditor/Private/Widgets/SNiagaraSystemViewportToolBar.cpp
new file mode 100644
--- /dev/null
+++ b/Engine/Plugins/FX/Niagara/Source/NiagaraEditor/Private/Widgets/SNiagaraSystemViewportToolBar.cpp
@@ -0,0 +1,56 @@
+// Copyright 1998-2016 Epic Games, Inc. All Rights Reserved.
+
+#include "SNiagaraSystemViewportToolBar.h"
+#include "Widgets/Layout/SBorder.h"
+#include "EditorStyleSet.h"
+#include "NiagaraEditorCommands.h"
+
+#define LOCTEXT_NAMESPACE "NiagaraSystemViewportToolBar"
+
+
+///////////////////////////////////////////////////////////
+// SNiagaraSystemViewportToolBar
+
+void SNiagaraSystemViewportToolBar::Construct(const FArguments& InArgs, TSharedPtr<class SNiagaraSystemViewport> InViewport)
+{
+	SCommonEditorViewportToolbarBase::Construct(SCommonEditorViewportToolbarBase::FArguments(), InViewport);
+}
+
+TSharedRef<SWidget> SNiagaraSystemViewportToolBar::GenerateShowMenu() const
+{
+	GetInfoProvider().OnFloatingButtonClicked();
+
+	TSharedRef<SEditorViewport> ViewportRef = GetInfoProvider().GetViewportWidget();
+
+	const bool bInShouldCloseWindowAfterMenuSelection = true;
+	FMenuBuilder ShowMenuBuilder(bInShouldCloseWindowAfterMenuSelection, ViewportRef->GetCommandList());
+	{
+		auto Commands = FNiagaraEditorCommands::Get();
+
+		//ShowMenuBuilder.AddMenuEntry(Commands.ToggleMaterialStats);
+		//ShowMenuBuilder.AddMenuEntry(Commands.ToggleMobileStats);
+
+		ShowMenuBuilder.AddMenuSeparator();
+
+		ShowMenuBuilder.AddMenuEntry(Commands.TogglePreviewGrid);
+		//ShowMenuBuilder.AddMenuEntry(Commands.TogglePreviewBackground);
+	}
+
+	return ShowMenuBuilder.MakeWidget();
+}
+
+bool SNiagaraSystemViewportToolBar::IsViewModeSupported(EViewModeIndex ViewModeIndex) const 
+{
+	switch (ViewModeIndex)
+	{
+	case VMI_PrimitiveDistanceAccuracy:
+	case VMI_MeshUVDensityAccuracy:
+	case VMI_RequiredTextureResolution:
+		return false;
+	default:
+		return true;
+	}
+	return true; 
+}
+
+#undef LOCTEXT_NAMESPACE
diff --git a/Engine/Plugins/FX/Niagara/Source/NiagaraEditor/Private/Widgets/SNiagaraSystemViewportToolBar.h b/Engine/Plugins/FX/Niagara/Source/NiagaraEditor/Private/Widgets/SNiagaraSystemViewportToolBar.h
new file mode 100644
--- /dev/null
+++ b/Engine/Plugins/FX/Niagara/Source/NiagaraEditor/Private/Widgets/SNiagaraSystemViewportToolBar.h
@@ -0,0 +1,26 @@
+// Copyright 1998-2016 Epic Games, Inc. All Rights Reserved.
+
+#pragma once
+
+#include "CoreMinimal.h"
+#include "Widgets/DeclarativeSyntaxSupport.h"
+#include "Widgets/SWidget.h"
+#include "SViewportToolBar.h"
+#include "SCommonEditorViewportToolbarBase.h"
+#include "SNiagaraSystemViewport.h"
+
+// In-viewport toolbar widget used in the Niagara editor
+class SNiagaraSystemViewportToolBar : public SCommonEditorViewportToolbarBase
+{
+public:
+	SLATE_BEGIN_ARGS(SNiagaraSystemViewportToolBar) {}
+	SLATE_END_ARGS()
+
+	void Construct(const FArguments& InArgs, TSharedPtr<class SNiagaraSystemViewport> InViewport);
+
+	// SCommonEditorViewportToolbarBase interface
+	virtual TSharedRef<SWidget> GenerateShowMenu() const override;
+	// End of SCommonEditorViewportToolbarBase
+
+	virtual bool IsViewModeSupported(EViewModeIndex ViewModeIndex) const override; 
+};
\ No newline at end of file
diff --git a/Engine/Plugins/FX/Niagara/Source/NiagaraEditor/Public/EdGraphSchema_Niagara.h b/Engine/Plugins/FX/Niagara/Source/NiagaraEditor/Public/EdGraphSchema_Niagara.h
new file mode 100644
--- /dev/null
+++ b/Engine/Plugins/FX/Niagara/Source/NiagaraEditor/Public/EdGraphSchema_Niagara.h
@@ -0,0 +1,109 @@
+// Copyright 1998-2017 Epic Games, Inc. All Rights Reserved.
+
+#pragma once
+
+#include "CoreMinimal.h"
+#include "UObject/ObjectMacros.h"
+#include "EdGraph/EdGraphSchema.h"
+#include "ConnectionDrawingPolicy.h"
+#include "EdGraphSchema_Niagara.generated.h"
+
+class UEdGraph;
+struct FNiagaraVariable;
+struct FNiagaraVariableInfo;
+struct FNiagaraTypeDefinition;
+enum class ENiagaraDataType : uint8;
+
+/** Action to add a node to the graph */
+USTRUCT()
+struct NIAGARAEDITOR_API FNiagaraSchemaAction_NewNode : public FEdGraphSchemaAction
+{
+	GENERATED_USTRUCT_BODY();
+
+	/** Template of node we want to create */
+	UPROPERTY()
+	class UEdGraphNode* NodeTemplate;
+
+	UPROPERTY()
+	FName InternalName;
+
+	FNiagaraSchemaAction_NewNode() 
+		: FEdGraphSchemaAction()
+		, NodeTemplate(nullptr)
+	{}
+
+	FNiagaraSchemaAction_NewNode(FText InNodeCategory, FText InMenuDesc, FName InInternalName, FText InToolTip, const int32 InGrouping, FText InKeywords = FText())
+		: FEdGraphSchemaAction(MoveTemp(InNodeCategory), MoveTemp(InMenuDesc), MoveTemp(InToolTip), InGrouping, InKeywords)
+		, NodeTemplate(nullptr), InternalName(InInternalName)
+	{}
+
+	//~ Begin FEdGraphSchemaAction Interface
+	virtual UEdGraphNode* PerformAction(class UEdGraph* ParentGraph, UEdGraphPin* FromPin, const FVector2D Location, bool bSelectNewNode = true) override;
+	virtual UEdGraphNode* PerformAction(class UEdGraph* ParentGraph, TArray<UEdGraphPin*>& FromPins, const FVector2D Location, bool bSelectNewNode = true) override;
+	virtual void AddReferencedObjects( FReferenceCollector& Collector ) override;
+	//~ End FEdGraphSchemaAction Interface
+
+	template <typename NodeType>
+	static NodeType* SpawnNodeFromTemplate(class UEdGraph* ParentGraph, NodeType* InTemplateNode, const FVector2D Location, bool bSelectNewNode = true)
+	{
+		FNiagaraSchemaAction_NewNode Action;
+		Action.NodeTemplate = InTemplateNode;
+
+		return Cast<NodeType>(Action.PerformAction(ParentGraph, NULL, Location, bSelectNewNode));
+	}
+};
+
+UCLASS()
+class NIAGARAEDITOR_API UEdGraphSchema_Niagara : public UEdGraphSchema
+{
+	GENERATED_UCLASS_BODY()
+
+	// Allowable PinType.PinCategory values
+	static const FString PinCategoryType;
+	static const FString PinCategoryMisc;
+	static const FString PinCategoryClass;
+	static const FString PinCategoryEnum;
+
+	//~ Begin EdGraphSchema Interface
+	virtual void GetGraphContextActions(FGraphContextMenuBuilder& ContextMenuBuilder) const override;
+	virtual void GetContextMenuActions(const UEdGraph* CurrentGraph, const UEdGraphNode* InGraphNode, const UEdGraphPin* InGraphPin, class FMenuBuilder* MenuBuilder, bool bIsDebugging) const override;
+	virtual const FPinConnectionResponse CanCreateConnection(const UEdGraphPin* A, const UEdGraphPin* B) const override;
+	virtual FLinearColor GetPinTypeColor(const FEdGraphPinType& PinType) const override;
+	virtual bool ShouldHidePinDefaultValue(UEdGraphPin* Pin) const override;
+	virtual bool TryCreateConnection(UEdGraphPin* A, UEdGraphPin* B) const override;
+	virtual void BreakSinglePinLink(UEdGraphPin* SourcePin, UEdGraphPin* TargetPin) override;
+	virtual FConnectionDrawingPolicy* CreateConnectionDrawingPolicy(int32 InBackLayerID, int32 InFrontLayerID, float InZoomFactor, const FSlateRect& InClippingRect, class FSlateWindowElementList& InDrawElements, class UEdGraph* InGraphObj) const override;
+	//~ End EdGraphSchema Interface
+
+	virtual TArray<TSharedPtr<FNiagaraSchemaAction_NewNode> > GetGraphContextActions(const UEdGraph* CurrentGraph, TArray<UObject*>& SelectedObjects, const UEdGraphPin* FromPin, UEdGraph* OwnerOfTemporaries) const;
+	void PromoteSinglePinToParameter(UEdGraphPin* SourcePin);
+
+	FNiagaraVariable PinToNiagaraVariable(const UEdGraphPin* Pin, bool bNeedsValue=false)const;
+	FString NiagaraVariableToPinDefaultValue(const FNiagaraVariable& Variable) const;
+	FNiagaraTypeDefinition PinToTypeDefinition(const UEdGraphPin* Pin)const;
+	FEdGraphPinType TypeDefinitionToPinType(FNiagaraTypeDefinition TypeDef)const;
+	
+	bool IsSystemConstant(const FNiagaraVariable& Variable)const;
+
+	FNiagaraTypeDefinition GetTypeDefForProperty(const UProperty* Property)const;
+
+	static const FLinearColor NodeTitleColor_Attribute;
+	static const FLinearColor NodeTitleColor_Constant;
+	static const FLinearColor NodeTitleColor_SystemConstant;
+	static const FLinearColor NodeTitleColor_FunctionCall;
+	static const FLinearColor NodeTitleColor_Event;
+
+private:
+	void GetBreakLinkToSubMenuActions(class FMenuBuilder& MenuBuilder, UEdGraphPin* InGraphPin);
+	void GetNumericConversionToSubMenuActions(class FMenuBuilder& MenuBuilder, UEdGraphPin* InGraphPin);
+	void ConvertNumericPinToType(UEdGraphPin* InPin, FNiagaraTypeDefinition TypeDef);
+
+};
+
+class FNiagaraConnectionDrawingPolicy : public FConnectionDrawingPolicy
+{
+public:
+	FNiagaraConnectionDrawingPolicy(int32 InBackLayerID, int32 InFrontLayerID, float InZoomFactor, const FSlateRect& InClippingRect, FSlateWindowElementList& InDrawElements);
+	virtual void DetermineWiringStyle(UEdGraphPin* OutputPin, UEdGraphPin* InputPin, /*inout*/ FConnectionParams& Params) override;
+};
+
diff --git a/Engine/Plugins/FX/Niagara/Source/NiagaraEditor/Public/INiagaraEditorTypeUtilities.h b/Engine/Plugins/FX/Niagara/Source/NiagaraEditor/Public/INiagaraEditorTypeUtilities.h
new file mode 100644
--- /dev/null
+++ b/Engine/Plugins/FX/Niagara/Source/NiagaraEditor/Public/INiagaraEditorTypeUtilities.h
@@ -0,0 +1,53 @@
+// Copyright 1998-2016 Epic Games, Inc. All Rights Reserved.
+
+#pragma once
+#include "SharedPointer.h"
+#include "Delegate.h"
+
+class FStructOnScope;
+class SNiagaraParameterEditor;
+class SWidget;
+struct FNiagaraVariable;
+
+class INiagaraEditorTypeUtilities
+{
+public:
+	DECLARE_DELEGATE(FNotifyValueChanged);
+public:
+	virtual ~INiagaraEditorTypeUtilities() {}
+
+	virtual bool CanProvideDefaultValue() const = 0;
+
+	virtual void UpdateStructWithDefaultValue(TSharedRef<FStructOnScope> Struct) const = 0;
+
+	virtual bool CanCreateParameterEditor() const = 0;
+
+	virtual TSharedPtr<SNiagaraParameterEditor> CreateParameterEditor() const = 0;
+
+	virtual bool CanCreateDataInterfaceEditor() const = 0;
+
+	virtual TSharedPtr<SWidget> CreateDataInterfaceEditor(UObject* DataInterface, FNotifyValueChanged DataInterfaceChangedHandler) const = 0;
+
+	virtual bool CanHandlePinDefaults() const = 0;
+
+	virtual FString GetPinDefaultStringFromValue(const FNiagaraVariable& Variable) const = 0;
+
+	virtual void SetValueFromPinDefaultString(const FString& StringValue, FNiagaraVariable& Variable) const = 0;
+};
+
+class FNiagaraEditorTypeUtilities : public INiagaraEditorTypeUtilities, public TSharedFromThis<FNiagaraEditorTypeUtilities>
+{
+public:
+	DECLARE_DELEGATE(FNotifyValueChanged);
+public:
+	//~ INiagaraEditorTypeUtilities
+	virtual bool CanProvideDefaultValue() const override { return false; }
+	virtual void UpdateStructWithDefaultValue(TSharedRef<FStructOnScope> Struct) const override { }
+	virtual bool CanCreateParameterEditor() const override { return false; }
+	virtual TSharedPtr<SNiagaraParameterEditor> CreateParameterEditor() const override { return TSharedPtr<SNiagaraParameterEditor>(); }
+	virtual bool CanCreateDataInterfaceEditor() const override { return false; };
+	virtual TSharedPtr<SWidget> CreateDataInterfaceEditor(UObject* DataInterface, FNotifyValueChanged DataInterfaceChangedHandler) const override { return TSharedPtr<SWidget>(); }
+	virtual bool CanHandlePinDefaults() const override { return false; }
+	virtual FString GetPinDefaultStringFromValue(const FNiagaraVariable& Variable) const override { return FString(); }
+	virtual void SetValueFromPinDefaultString(const FString& StringValue, FNiagaraVariable& Variable) const override { }
+};
\ No newline at end of file
diff --git a/Engine/Plugins/FX/Niagara/Source/NiagaraEditor/Public/NiagaraCurveOwner.h b/Engine/Plugins/FX/Niagara/Source/NiagaraEditor/Public/NiagaraCurveOwner.h
new file mode 100644
--- /dev/null
+++ b/Engine/Plugins/FX/Niagara/Source/NiagaraEditor/Public/NiagaraCurveOwner.h
@@ -0,0 +1,56 @@
+// Copyright 1998-2017 Epic Games, Inc. All Rights Reserved.
+
+#pragma once
+
+#include "Curves/CurveOwnerInterface.h"
+#include "NiagaraDataInterface.h"
+
+/** A curve owner for curves in a niagara System. */
+class NIAGARAEDITOR_API FNiagaraCurveOwner : public FCurveOwnerInterface
+{
+public:
+	DECLARE_DELEGATE_TwoParams(FNotifyCurveChanged, FRichCurve*, UObject*);
+
+public:
+	FNiagaraCurveOwner();
+
+	/** Removes all of the curves from the curve owner. */
+	void EmptyCurves();
+
+	/** Adds a curve to this curve owner. */
+	void AddCurve(FRichCurve& Curve, FName Name, FLinearColor Color, UObject& Owner, FNotifyCurveChanged CurveChangedHandler);
+
+	void SetColorCurves(FRichCurve& RedCurve, FRichCurve& GreenCurve, FRichCurve& BlueCurve, FRichCurve& AlphaCurve, FName Name, UObject& Owner, FNotifyCurveChanged CurveChangedHandler);
+
+	/** FCurveOwnerInterface */
+	virtual TArray<FRichCurveEditInfoConst> GetCurves() const override;
+	virtual TArray<FRichCurveEditInfo> GetCurves() override;
+	virtual void ModifyOwner() override;
+	virtual TArray<const UObject*> GetOwners() const override;
+	virtual void MakeTransactional() override;
+	virtual void OnCurveChanged(const TArray<FRichCurveEditInfo>& ChangedCurveEditInfos) override;
+	virtual bool IsLinearColorCurve() const override;
+	virtual FLinearColor GetLinearColorValue(float InTime) const override;
+	virtual bool HasAnyAlphaKeys() const override;
+	virtual bool IsValidCurve(FRichCurveEditInfo CurveInfo) override;
+	virtual FLinearColor GetCurveColor(FRichCurveEditInfo CurveInfo) const override;
+
+private:
+	/** The ordered array of const curves used to implement the curve owner interface. */
+	TArray<FRichCurveEditInfoConst> ConstCurves;
+
+	/** The ordered array or curves used to implement the curve owner interface. */
+	TArray<FRichCurveEditInfo> Curves;
+
+	/** A map of curve edit infos to their corresponding owners. */
+	TMap<FRichCurveEditInfo, UObject*> EditInfoToOwnerMap;
+
+	/** A map of curve edit infos to their colors. */
+	TMap<FRichCurveEditInfo, FLinearColor> EditInfoToColorMap;
+
+	/** A map of curve edit infos to change handler delegates. */
+	TMap<FRichCurveEditInfo, FNotifyCurveChanged> EditInfoToNotifyCurveChangedMap;
+
+	/** Whether or not this set of curves should be treated as a color curve. */
+	bool bIsColorCurve;
+};
\ No newline at end of file
diff --git a/Engine/Plugins/FX/Niagara/Source/NiagaraEditor/Public/NiagaraEditorModule.h b/Engine/Plugins/FX/Niagara/Source/NiagaraEditor/Public/NiagaraEditorModule.h
new file mode 100644
--- /dev/null
+++ b/Engine/Plugins/FX/Niagara/Source/NiagaraEditor/Public/NiagaraEditorModule.h
@@ -0,0 +1,95 @@
+// Copyright 1998-2017 Epic Games, Inc. All Rights Reserved.
+
+#pragma once
+
+#include "CoreMinimal.h"
+#include "Modules/ModuleInterface.h"
+#include "Toolkits/AssetEditorToolkit.h"
+#include "NiagaraTypes.h"
+#include "INiagaraCompiler.h"
+#include "AssetTypeCategories.h"
+
+class IAssetTools;
+class IAssetTypeActions;
+struct FNiagaraCompileResults;
+class INiagaraEditorTypeUtilities;
+class UNiagaraSettings;
+class USequencerSettings;
+class UNiagaraStackViewModel;
+
+DECLARE_STATS_GROUP(TEXT("Niagara Editor"), STATGROUP_NiagaraEditor, STATCAT_Advanced);
+
+/** Niagara Editor module */
+class FNiagaraEditorModule : public IModuleInterface,
+	public IHasMenuExtensibility, public IHasToolBarExtensibility, public FGCObject
+{
+public:
+	DECLARE_DELEGATE_RetVal_OneParam(TSharedRef<SWidget>, FOnCreateStackWidget, UNiagaraStackViewModel*);
+
+public:
+	FNiagaraEditorModule();
+
+	// IModuleInterface
+	virtual void StartupModule() override;
+	virtual void ShutdownModule() override;
+
+	/** Compile the specified script. */
+	virtual ENiagaraScriptCompileStatus CompileScript(class UNiagaraScript* ScriptToCompile, FString& OutGraphLevelErrorMessages);
+
+	/** Gets the extensibility managers for outside entities to extend static mesh editor's menus and toolbars */
+	virtual TSharedPtr<FExtensibilityManager> GetMenuExtensibilityManager() override {return MenuExtensibilityManager;}
+	virtual TSharedPtr<FExtensibilityManager> GetToolBarExtensibilityManager() override {return ToolBarExtensibilityManager;}
+
+	/** Registers niagara editor type utilities for a specific type. */
+	void RegisterTypeUtilities(FNiagaraTypeDefinition Type, TSharedRef<INiagaraEditorTypeUtilities> EditorUtilities);
+
+	/** Register/unregister niagara editor settings. */
+	void RegisterSettings();
+	void UnregisterSettings();
+
+	/** Gets niagara editor type utilities for a specific type if there are any registered. */
+	TSharedPtr<INiagaraEditorTypeUtilities> NIAGARAEDITOR_API GetTypeUtilities(const FNiagaraTypeDefinition& Type);
+
+	static EAssetTypeCategories::Type GetAssetCategory() { return NiagaraAssetCategory; }
+
+	TSharedRef<SWidget> CreateStackWidget(UNiagaraStackViewModel* StackViewModel) const;
+
+	FDelegateHandle NIAGARAEDITOR_API SetOnCreateStackWidget(FOnCreateStackWidget InOnCreateStackWidget);
+	void NIAGARAEDITOR_API ResetOnCreateStackWidget(FDelegateHandle DelegateHandle);
+
+
+	/** Niagara Editor app identifier string */
+	static const FName NiagaraEditorAppIdentifier;
+
+	/** The tab color scale for niagara editors. */
+	static const FLinearColor WorldCentricTabColorScale;
+
+private:
+	void RegisterAssetTypeAction(IAssetTools& AssetTools, TSharedRef<IAssetTypeActions> Action);
+	void OnNiagaraSettingsChangedEvent(const FString& PropertyName, const UNiagaraSettings* Settings);
+
+
+	/** FGCObject interface */
+	virtual void AddReferencedObjects( FReferenceCollector& Collector ) override;
+
+private:
+	TSharedPtr<FExtensibilityManager> MenuExtensibilityManager;
+	TSharedPtr<FExtensibilityManager> ToolBarExtensibilityManager;
+
+	/** All created asset type actions.  Cached here so that we can unregister it during shutdown. */
+	TArray< TSharedPtr<IAssetTypeActions> > CreatedAssetTypeActions;
+
+	TMap<FNiagaraTypeDefinition, TSharedRef<INiagaraEditorTypeUtilities>> TypeToEditorUtilitiesMap;
+
+	static EAssetTypeCategories::Type NiagaraAssetCategory;
+
+	FDelegateHandle CreateEmitterTrackEditorHandle;
+	FDelegateHandle CreateSystemTrackEditorHandle;
+	FDelegateHandle CreateFloatParameterTrackEditorHandle;
+	FDelegateHandle CreateBoolParameterTrackEditorHandle;
+
+	USequencerSettings* SequencerSettings;
+	FOnCreateStackWidget OnCreateStackWidget;
+};
+
+
diff --git a/Engine/Plugins/FX/Niagara/Source/NiagaraEditor/Public/NiagaraEditorSettings.h b/Engine/Plugins/FX/Niagara/Source/NiagaraEditor/Public/NiagaraEditorSettings.h
new file mode 100644
--- /dev/null
+++ b/Engine/Plugins/FX/Niagara/Source/NiagaraEditor/Public/NiagaraEditorSettings.h
@@ -0,0 +1,56 @@
+#pragma once
+
+#include "NiagaraSystem.h"
+#include "NiagaraEmitter.h"
+#include "NiagaraScript.h"
+#include "Engine/DeveloperSettings.h"
+#include "NiagaraSpawnShortcut.h"
+#include "NiagaraEditorSettings.generated.h"
+
+
+UCLASS(config = Engine, defaultconfig, meta=(DisplayName="Niagara"))
+class UNiagaraEditorSettings : public UDeveloperSettings
+{
+public:
+	GENERATED_UCLASS_BODY()
+		
+	/** System to duplicate as the base of all new System assets created. */
+	UPROPERTY(config, EditAnywhere, Category = Niagara)
+	FStringAssetReference DefaultSystem;
+
+	/** Emitter to duplicate as the base of all new emitter assets created. */
+	UPROPERTY(config, EditAnywhere, Category = Niagara)
+	FStringAssetReference DefaultEmitter;
+
+	/** Niagara script to duplicate as the base of all new script assets created. */
+	UPROPERTY(config, EditAnywhere, Category = Niagara)
+	FStringAssetReference DefaultScript;
+
+	/** Shortcut key bindings that if held down while doing a mouse click, will spawn the specified type of Niagara node.*/
+	UPROPERTY(config, EditAnywhere, Category = Niagara)
+	TArray<FNiagaraSpawnShortcut> GraphCreationShortcuts;
+
+	/** Whether or not auto-compile is enabled in the editors. */
+	UPROPERTY(config, EditAnywhere, Category = Niagara)
+	bool bAutoCompile;
+	
+	// Begin UDeveloperSettings Interface
+	virtual FName GetCategoryName() const override;
+#if WITH_EDITOR
+	virtual FText GetSectionText() const override;
+#endif
+	// END UDeveloperSettings Interface
+
+#if WITH_EDITOR
+	virtual void PostEditChangeProperty(FPropertyChangedEvent& PropertyChangedEvent) override;
+
+public:
+	DECLARE_MULTICAST_DELEGATE_TwoParams(FOnNiagaraEditorSettingsChanged, const FString&, const UNiagaraEditorSettings*);
+
+	/** Gets a multicast delegate which is called whenever one of the parameters in this settings object changes. */
+	static FOnNiagaraEditorSettingsChanged& OnSettingsChanged();
+
+protected:
+	static FOnNiagaraEditorSettingsChanged SettingsChangedDelegate;
+#endif
+};
diff --git a/Engine/Plugins/FX/Niagara/Source/NiagaraEditor/Public/NiagaraEditorStyle.h b/Engine/Plugins/FX/Niagara/Source/NiagaraEditor/Public/NiagaraEditorStyle.h
new file mode 100644
--- /dev/null
+++ b/Engine/Plugins/FX/Niagara/Source/NiagaraEditor/Public/NiagaraEditorStyle.h
@@ -0,0 +1,30 @@
+// Copyright 1998-2016 Epic Games, Inc. All Rights Reserved.
+
+#pragma once
+#include "ISlateStyle.h"
+
+/** Manages the style which provides resources for niagara editor widgets. */
+class NIAGARAEDITOR_API FNiagaraEditorStyle
+{
+public:
+
+	static void Initialize();
+
+	static void Shutdown();
+
+	/** reloads textures used by slate renderer */
+	static void ReloadTextures();
+
+	/** @return The Slate style set for niagara editor widgets */
+	static const ISlateStyle& Get();
+
+	static FName GetStyleSetName();
+
+private:
+
+	static TSharedRef< class FSlateStyleSet > Create();
+
+private:
+
+	static TSharedPtr< class FSlateStyleSet > NiagaraEditorStyleInstance;
+};
\ No newline at end of file
diff --git a/Engine/Plugins/FX/Niagara/Source/NiagaraEditor/Public/NiagaraNode.h b/Engine/Plugins/FX/Niagara/Source/NiagaraEditor/Public/NiagaraNode.h
new file mode 100644
--- /dev/null
+++ b/Engine/Plugins/FX/Niagara/Source/NiagaraEditor/Public/NiagaraNode.h
@@ -0,0 +1,96 @@
+// Copyright 1998-2017 Epic Games, Inc. All Rights Reserved.
+
+#pragma once
+
+#include "CoreMinimal.h"
+#include "UObject/ObjectMacros.h"
+#include "EdGraph/EdGraphNode.h"
+#include "NiagaraEditorCommon.h"
+#include "NiagaraParameterMapHistory.h"
+#include "NiagaraNode.generated.h"
+
+class UEdGraphPin;
+class INiagaraCompiler;
+
+UCLASS()
+class NIAGARAEDITOR_API UNiagaraNode : public UEdGraphNode
+{
+	GENERATED_UCLASS_BODY()
+protected:
+
+	bool ReallocatePins();
+
+	bool CompileInputPins(class FHlslNiagaraTranslator *Translator, TArray<int32>& OutCompiledInputs);
+
+public:
+
+	//~ Begin UObject interface
+	virtual void PostLoad() override;
+	//~ End UObject interface
+
+	//~ Begin EdGraphNode Interface
+	virtual void AutowireNewNode(UEdGraphPin* FromPin)override;	
+	virtual void PinDefaultValueChanged(UEdGraphPin* Pin) override;
+	virtual void PinConnectionListChanged(UEdGraphPin* Pin) override;
+	virtual void PinTypeChanged(UEdGraphPin* Pin) override;
+	//~ End EdGraphNode Interface
+
+	/** Get the Niagara graph that owns this node */
+	const class UNiagaraGraph* GetNiagaraGraph()const;
+	class UNiagaraGraph* GetNiagaraGraph();
+
+	/** Get the source object */
+	class UNiagaraScriptSource* GetSource()const;
+
+	/** Gets the asset referenced by this node, or nullptr if there isn't one. */
+	virtual UObject* GetReferencedAsset() const { return nullptr; }
+
+	/** Refreshes the node due to external changes, e.g. the underlying function changed for a function call node. Return true if the graph changed. */
+	virtual bool RefreshFromExternalChanges() { return false; }
+
+	virtual void Compile(class FHlslNiagaraTranslator *Translator, TArray<int32>& Outputs);
+
+	UEdGraphPin* GetInputPin(int32 InputIndex) const;
+	void GetInputPins(TArray<class UEdGraphPin*>& OutInputPins) const;
+	UEdGraphPin* GetOutputPin(int32 OutputIndex) const;
+	void GetOutputPins(TArray<class UEdGraphPin*>& OutOutputPins) const;
+
+	/** Apply any node-specific logic to determine if it is safe to add this node to the graph. This is meant to be called only in the Editor before placing the node.*/
+	virtual bool CanAddToGraph(UNiagaraGraph* TargetGraph, FString& OutErrorMsg) const;
+
+	/** Gets which mode to use when deducing the type of numeric output pins from the types of the input pins. */
+	virtual ENiagaraNumericOutputTypeSelectionMode GetNumericOutputTypeSelectionMode() const;
+
+	/** Convert the type of an existing numeric pin to a more known type.*/
+	virtual bool ConvertNumericPinToType(UEdGraphPin* InGraphPin, FNiagaraTypeDefinition TypeDef);
+
+	/** Determine if there are any external dependencies wrt to scripts and ensure that those dependencies are sucked into the existing package.*/
+	virtual void SubsumeExternalDependencies(TMap<const UObject*, UObject*>& ExistingConversions) {};
+
+	/** Determine whether or not a pin should be renamable. */
+	virtual bool IsPinNameEditable(const UEdGraphPin* GraphPinObj) const { return false; }
+
+	/** Determine whether or not a specific pin should immediately be opened for rename.*/
+	virtual bool IsPinNameEditableUponCreation(const UEdGraphPin* GraphPinObj) const { return false; }
+
+	/** Verify that the potential rename has produced acceptable results for a pin.*/
+	virtual bool VerifyEditablePinName(const FText& InName, FText& OutErrorMessage, const UEdGraphPin* InGraphPinObj) const { return false; }
+
+	/** Verify that the potential rename has produced acceptable results for a pin.*/
+	virtual bool CommitEditablePinName(const FText& InName,  UEdGraphPin* InGraphPinObj) { return false; }
+	
+	virtual void BuildParameterMapHistory(FNiagaraParameterMapHistoryBuilder& OutHistory, bool bRecursive = true);
+
+	bool ShouldVisit(int32 InVisitID) 
+	{ 
+		bool bShouldVisit = InVisitID != VisitID;
+		VisitID = InVisitID;
+		return bShouldVisit;
+	}
+	
+protected:
+	virtual int32 CompileInputPin(class FHlslNiagaraTranslator *Translator, UEdGraphPin* Pin);
+
+	//ID for the most recent visitor. Allows faster tracking of what nodes have been visited.
+	int32 VisitID;
+};
diff --git a/Engine/Plugins/FX/Niagara/Source/NiagaraEditor/Public/NiagaraNodeAssignment.h b/Engine/Plugins/FX/Niagara/Source/NiagaraEditor/Public/NiagaraNodeAssignment.h
new file mode 100644
--- /dev/null
+++ b/Engine/Plugins/FX/Niagara/Source/NiagaraEditor/Public/NiagaraNodeAssignment.h
@@ -0,0 +1,38 @@
+// Copyright 1998-2017 Epic Games, Inc. All Rights Reserved.
+
+#pragma once
+
+#include "CoreMinimal.h"
+#include "UObject/ObjectMacros.h"
+#include "NiagaraEditorCommon.h"
+#include "NiagaraNode.h"
+#include "NiagaraNodeFunctionCall.h"
+#include "NiagaraNodeAssignment.generated.h"
+
+class UNiagaraScript;
+
+UCLASS(MinimalAPI)
+class UNiagaraNodeAssignment : public UNiagaraNodeFunctionCall
+{
+	GENERATED_BODY()
+
+public:
+	UPROPERTY(EditAnywhere, Category = Assignment)
+	FNiagaraVariable AssignmentTarget;
+
+	UPROPERTY(EditAnywhere, Category = Assignment)
+	FString AssignmentDefaultValue;
+
+	//~ Begin EdGraphNode Interface
+	virtual FText GetNodeTitle(ENodeTitleType::Type TitleType) const override;
+	virtual void AllocateDefaultPins() override;
+
+	//~ UNiagaraNodeFunctionCall interface
+	virtual bool RefreshFromExternalChanges() override;
+
+private:
+	void GenerateScript();
+
+	void InitializeScript(UNiagaraScript* NewScript);
+};
+
diff --git a/Engine/Plugins/FX/Niagara/Source/NiagaraEditor/Public/NiagaraNodeDataSetBase.h b/Engine/Plugins/FX/Niagara/Source/NiagaraEditor/Public/NiagaraNodeDataSetBase.h
new file mode 100644
--- /dev/null
+++ b/Engine/Plugins/FX/Niagara/Source/NiagaraEditor/Public/NiagaraNodeDataSetBase.h
@@ -0,0 +1,55 @@
+// Copyright 1998-2016 Epic Games, Inc. All Rights Reserved.
+
+#pragma once
+
+#include "CoreMinimal.h"
+#include "UObject/ObjectMacros.h"
+#include "NiagaraCommon.h"
+#include "NiagaraEditorCommon.h"
+#include "NiagaraNode.h"
+#include "NiagaraNodeDataSetBase.generated.h"
+
+UCLASS(MinimalAPI)
+class UNiagaraNodeDataSetBase : public UNiagaraNode
+{
+	GENERATED_UCLASS_BODY()
+
+public:
+
+	UPROPERTY(EditAnywhere, Category = DataSet)
+	FNiagaraDataSetID DataSet;
+
+	UPROPERTY(EditAnywhere, Category = Variables)
+	TArray<FNiagaraVariable> Variables;
+	
+	UPROPERTY()
+	TArray<FString> VariableFriendlyNames;
+
+	UPROPERTY()
+	const UStruct* ExternalStructAsset;
+
+	//~ Begin UObject Interface
+	virtual void PostEditChangeProperty(struct FPropertyChangedEvent& PropertyChangedEvent) override;
+	virtual void PostLoad() override;
+	//~ End UObject Interface
+
+	//~ Begin EdGraphNode Interface
+	virtual FLinearColor GetNodeTitleColor() const override;
+	//~ End EdGraphNode Interface
+
+	//~ Begin UNiagaraNode Interface
+	virtual bool RefreshFromExternalChanges() override;
+	//~ End UNiagaraNode Interface
+
+	bool InitializeFromStruct(const UStruct* ReferenceStruct);
+	
+protected:
+	bool InitializeFromStructInternal(const UStruct* PayloadStruct);
+
+	bool IsSynchronizedWithStruct(bool bIgnoreConditionVar, FString* Issues, bool bLogIssues = true);
+	bool SynchronizeWithStruct();
+	bool GetSupportedNiagaraTypeDef(const UProperty* Property, FNiagaraTypeDefinition& TypeDef);
+
+	static FString ConditionVarName;
+};
+
diff --git a/Engine/Plugins/FX/Niagara/Source/NiagaraEditor/Public/NiagaraNodeFunctionCall.h b/Engine/Plugins/FX/Niagara/Source/NiagaraEditor/Public/NiagaraNodeFunctionCall.h
new file mode 100644
--- /dev/null
+++ b/Engine/Plugins/FX/Niagara/Source/NiagaraEditor/Public/NiagaraNodeFunctionCall.h
@@ -0,0 +1,78 @@
+// Copyright 1998-2017 Epic Games, Inc. All Rights Reserved.
+
+#pragma once
+
+#include "CoreMinimal.h"
+#include "UObject/ObjectMacros.h"
+#include "NiagaraEditorCommon.h"
+#include "NiagaraNode.h"
+#include "NiagaraNodeInput.h"
+#include "NiagaraNodeFunctionCall.generated.h"
+
+class UNiagaraScript;
+
+UCLASS(MinimalAPI)
+class UNiagaraNodeFunctionCall : public UNiagaraNode
+{
+	GENERATED_BODY()
+
+public:
+
+	UPROPERTY(EditAnywhere, Category = "Function")
+	UNiagaraScript* FunctionScript;
+
+	/** 
+	 * A path to a script asset which can be used to assign the function script the first time that
+	 * default pins are generated. This is used so that the function nodes can be populated in the graph context
+	 * menu without having to load all of the actual script assets.
+	 */
+	UPROPERTY(Transient)
+	FName FunctionScriptAssetObjectPath;
+
+	/** Some functions can be provided a signature directly rather than a script. */
+	UPROPERTY()
+	FNiagaraFunctionSignature Signature;
+
+	//Begin UObject interface
+	virtual void PostLoad()override;
+	virtual void PostEditChangeProperty(struct FPropertyChangedEvent& PropertyChangedEvent)override;
+	//End UObject interface
+
+	//~ Begin UNiagaraNode Interface
+	virtual void Compile(class FHlslNiagaraTranslator* Translator, TArray<int32>& Outputs)override;
+	virtual UObject* GetReferencedAsset() const override;
+	virtual bool RefreshFromExternalChanges() override;
+	virtual ENiagaraNumericOutputTypeSelectionMode GetNumericOutputTypeSelectionMode() const override;
+	virtual bool CanAddToGraph(UNiagaraGraph* TargetGraph, FString& OutErrorMsg) const override;
+	virtual void SubsumeExternalDependencies(TMap<const UObject*, UObject*>& ExistingConversions) override;
+	//End UNiagaraNode interface
+
+	//~ Begin EdGraphNode Interface
+	virtual void AutowireNewNode(UEdGraphPin* FromPin)override;
+	virtual void AllocateDefaultPins() override;
+	virtual FText GetNodeTitle(ENodeTitleType::Type TitleType) const override;
+	virtual FText GetTooltipText() const override;
+	virtual FLinearColor GetNodeTitleColor() const override;
+	//~ End EdGraphNode Interface
+
+	bool FindAutoBoundInput(UNiagaraNodeInput* InputNode, UEdGraphPin* PinToAutoBind, FNiagaraVariable& OutFoundVar, ENiagaraInputNodeUsage& OutNodeUsage);
+
+	virtual void BuildParameterMapHistory(FNiagaraParameterMapHistoryBuilder& OutHistory, bool bRecursive = true) override;
+
+	FString GetFunctionName() const { return FunctionDisplayName; }
+
+	/** Walk through the internal script graph for an ParameterMapGet nodes and see if any of them specify a default for VariableName.*/
+	UEdGraphPin* FindParameterMapDefaultValuePin(const FString& VariableName);
+
+protected:
+	/** Resets the node name based on the referenced script or signature. Guaranteed unique within a given graph instance.*/
+	void ComputeNodeName();
+
+	/** Adjusted every time that we compile this script. Lets us know that we might differ from any cached versions.*/
+	UPROPERTY()
+	FGuid CachedChangeId;
+
+	UPROPERTY()
+	FString FunctionDisplayName;
+};
+
diff --git a/Engine/Plugins/FX/Niagara/Source/NiagaraEditor/Public/NiagaraNodeInput.h b/Engine/Plugins/FX/Niagara/Source/NiagaraEditor/Public/NiagaraNodeInput.h
new file mode 100644
--- /dev/null
+++ b/Engine/Plugins/FX/Niagara/Source/NiagaraEditor/Public/NiagaraNodeInput.h
@@ -0,0 +1,117 @@
+// Copyright 1998-2017 Epic Games, Inc. All Rights Reserved.
+
+#pragma once
+
+#include "CoreMinimal.h"
+#include "UObject/ObjectMacros.h"
+#include "NiagaraCommon.h"
+#include "NiagaraEditorCommon.h"
+#include "NiagaraNode.h"
+#include "NiagaraNodeInput.generated.h"
+
+class UEdGraphPin;
+class UNiagaraGraph;
+
+
+USTRUCT()
+struct FNiagaraInputExposureOptions
+{
+	GENERATED_USTRUCT_BODY()
+		
+	FNiagaraInputExposureOptions()
+		: bExposed(1)
+		, bRequired(1)
+		, bCanAutoBind(0)
+		, bHidden(0)
+	{}
+
+	/** If this input is exposed to it's caller. */
+	UPROPERTY(EditAnywhere, Category = Exposure)
+	uint32 bExposed : 1;
+
+	/** If this input is required to be bound. */
+	UPROPERTY(EditAnywhere, Category = Exposure, meta = (editcondition = "bExposed"))
+	uint32 bRequired : 1;
+
+	/** If this input can auto-bind to system parameters and emitter attributes. Will never auto bind to custom parameters. */
+	UPROPERTY(EditAnywhere, Category = Exposure, meta = (editcondition = "bExposed"))
+	uint32 bCanAutoBind : 1;
+
+	/** If this input should be hidden in the advanced pin section of it's caller. */
+	UPROPERTY(EditAnywhere, Category = Exposure, meta = (editcondition = "bExposed"))
+	uint32 bHidden : 1;
+};
+
+
+UCLASS(MinimalAPI)
+class UNiagaraNodeInput : public UNiagaraNode
+{
+	GENERATED_UCLASS_BODY()
+
+public:
+
+	UPROPERTY(EditAnywhere, Instanced, Category = Input)
+	class UNiagaraDataInterface* DataInterface;
+
+	UPROPERTY(EditAnywhere, Category = Input)
+	FNiagaraVariable Input;
+
+	UPROPERTY(VisibleAnywhere, Category = Input)
+	ENiagaraInputNodeUsage Usage;
+
+	/** Controls where this input is relative to others in the calling node. */
+	UPROPERTY(EditAnywhere, Category = Input)
+	int32 CallSortPriority;
+
+	/** Controls this inputs exposure to callers. */
+	UPROPERTY(EditAnywhere, Category = Input)
+	FNiagaraInputExposureOptions ExposureOptions;
+
+	//~ Begin UObject Interface
+	virtual void PostEditChangeProperty(struct FPropertyChangedEvent& PropertyChangedEvent) override;
+	//~ End UObject Interface
+
+	//~ Begin EdGraphNode Interface
+	virtual void AllocateDefaultPins() override;
+	virtual FText GetNodeTitle(ENodeTitleType::Type TitleType) const override;
+	virtual FLinearColor GetNodeTitleColor() const override;
+	virtual void AutowireNewNode(UEdGraphPin* FromPin)override;
+	virtual TSharedPtr<SGraphNode> CreateVisualWidget() override;
+	virtual void OnRenameNode(const FString& NewName) override;
+	//~ End EdGraphNode Interface
+
+	/** Notifies the node that the type of it's input has been changed externally. */
+	void NotifyInputTypeChanged();
+
+	/** Notifies the node that the exposure options have changed externally.*/
+	void NotifyExposureOptionsChanged();
+
+	virtual void Compile(class FHlslNiagaraTranslator* Translator, TArray<int32>& Outputs)override;
+	
+	bool IsExposed()const { return ExposureOptions.bExposed; }
+	bool IsRequired()const { return ExposureOptions.bExposed && ExposureOptions.bRequired; }
+	bool IsHidden()const { return ExposureOptions.bExposed && ExposureOptions.bHidden; }
+	bool CanAutoBind()const { return ExposureOptions.bExposed && ExposureOptions.bCanAutoBind; }
+
+	bool ReferencesSameInput(UNiagaraNodeInput* Other) const;
+
+	virtual void BuildParameterMapHistory(FNiagaraParameterMapHistoryBuilder& OutHistory, bool bRecursive = true) override;
+
+	static const FLinearColor TitleColor_Attribute;
+	static const FLinearColor TitleColor_Constant;
+
+	/** Verify that the text about to be committed will be valid and doesn't duplicate existing variables based on type.
+	Type validation is done based on the input UObject type, preferably a UNiagaraNodeInput or UNiagaraNodeOutput.*/
+	static bool VerifyNodeRenameTextCommit(const FText& NewText, UNiagaraNode* NodeBeingChanged, FText& OutErrorMessage);
+
+
+	/** Generate a unique name based off of the existing names in the system.*/
+	static FName GenerateUniqueName(const UNiagaraGraph* Graph, FName& ProposedName, ENiagaraInputNodeUsage Usage);
+
+	/** Generate a unique sort index based off the existing nodes in the system.*/
+	static int32 GenerateNewSortPriority(const UNiagaraGraph* Graph, FName& ProposedName, ENiagaraInputNodeUsage Usage);
+
+	/** Given an array of nodes, sort them in place by their sort order, then lexicographically if the same.*/
+	static void SortNodes(TArray<UNiagaraNodeInput*>& InOutNodes);
+};
+
diff --git a/Engine/Plugins/FX/Niagara/Source/NiagaraEditor/Public/NiagaraNodeOp.h b/Engine/Plugins/FX/Niagara/Source/NiagaraEditor/Public/NiagaraNodeOp.h
new file mode 100644
--- /dev/null
+++ b/Engine/Plugins/FX/Niagara/Source/NiagaraEditor/Public/NiagaraNodeOp.h
@@ -0,0 +1,96 @@
+// Copyright 1998-2017 Epic Games, Inc. All Rights Reserved.
+
+#pragma once
+
+#include "CoreMinimal.h"
+#include "UObject/ObjectMacros.h"
+#include "Widgets/SWidget.h"
+#include "NiagaraEditorCommon.h"
+#include "Widgets/DeclarativeSyntaxSupport.h"
+#include "Widgets/SOverlay.h"
+#include "NiagaraNode.h"
+#include "NiagaraNodeOp.generated.h"
+
+class SGraphNode;
+class SGraphPin;
+class SVerticalBox;
+
+UCLASS(MinimalAPI)
+class UNiagaraNodeOp : public UNiagaraNode
+{
+	GENERATED_UCLASS_BODY()
+
+public:
+
+	/** Name of operation */
+	UPROPERTY()
+	FName OpName;
+
+	//~ Begin EdGraphNode Interface
+	virtual void AllocateDefaultPins() override;
+	virtual FText GetNodeTitle(ENodeTitleType::Type TitleType) const override;
+	virtual FText GetTooltipText() const override;
+	virtual FLinearColor GetNodeTitleColor() const override;
+	//~ End EdGraphNode Interface
+
+	//~ Begin UNiagaraNode Interface
+	virtual void Compile(class FHlslNiagaraTranslator* Translator, TArray<int32>& Outputs) override;
+	virtual bool RefreshFromExternalChanges() override;
+	virtual ENiagaraNumericOutputTypeSelectionMode GetNumericOutputTypeSelectionMode() const override;
+	//~ End UNiagaraNode Interface
+};
+
+
+
+#if 0
+
+class SGraphNodeNiagara : public SGraphNode
+{
+public:
+	SLATE_BEGIN_ARGS(SGraphNodeNiagara){}
+	SLATE_END_ARGS()
+
+		void Construct(const FArguments& InArgs, class UNiagaraNode* InNode);
+
+	//~ Begin SGraphNode Interface
+	virtual void CreatePinWidgets() override;
+	//~ End SGraphNode Interface
+
+	//~ Begin SNodePanel::SNode Interface
+	virtual void MoveTo(const FVector2D& NewPosition, FNodeSet& NodeFilter) override;
+	//~ End SNodePanel::SNode Interface
+
+	UNiagaraNode* GetNiagaraGraphNode() const { return NiagaraNode; }
+
+protected:
+	//~ Begin SGraphNode Interface
+	virtual void AddPin(const TSharedRef<SGraphPin>& PinToAdd) override;
+	virtual void CreateBelowPinControls(TSharedPtr<SVerticalBox> MainBox) override;
+	virtual void SetDefaultTitleAreaWidget(TSharedRef<SOverlay> DefaultTitleAreaWidget) override;
+	virtual TSharedRef<SWidget> CreateNodeContentArea() override;
+	//~ End SGraphNode Interface
+
+	/** Creates a preview viewport if necessary */
+	//TSharedRef<SWidget> CreatePreviewWidget();
+
+	/** Returns visibility of Expression Preview viewport */
+	//EVisibility ExpressionPreviewVisibility() const;
+
+	/** Show/hide Expression Preview */
+	//void OnExpressionPreviewChanged(const ECheckBoxState NewCheckedState);
+
+	/** hidden == unchecked, shown == checked */
+	//ECheckBoxState IsExpressionPreviewChecked() const;
+
+	/** Up when shown, down when hidden */
+	//const FSlateBrush* GetExpressionPreviewArrow() const;
+
+private:
+	/** Slate viewport for rendering preview via custom slate element */
+	//TSharedPtr<FPreviewViewport> PreviewViewport;
+
+	/** Cached material graph node pointer to avoid casting */
+	UNiagaraNode* NiagaraNode;
+};
+
+#endif
diff --git a/Engine/Plugins/FX/Niagara/Source/NiagaraEditor/Public/NiagaraParameterMapHistory.h b/Engine/Plugins/FX/Niagara/Source/NiagaraEditor/Public/NiagaraParameterMapHistory.h
new file mode 100644
--- /dev/null
+++ b/Engine/Plugins/FX/Niagara/Source/NiagaraEditor/Public/NiagaraParameterMapHistory.h
@@ -0,0 +1,300 @@
+// Copyright 1998-2017 Epic Games, Inc. All Rights Reserved.
+
+#pragma once
+
+#include "CoreMinimal.h"
+#include "NiagaraCommon.h"
+#include "Templates/Tuple.h"
+
+class UNiagaraNodeOutput;
+class UEdGraphPin;
+class UEdGraphNode;
+
+/** Traverses a Niagara node graph to identify the variables that have been written and read from a parameter map. 
+* 	This class is meant to aid in UI and compilation of the graph. There are several main script types and each one interacts
+*	slightly differently with the history depending on context.
+*/
+struct FNiagaraParameterMapHistory
+{
+public:
+	FNiagaraParameterMapHistory();
+
+	/** The variables that have been identified during the traversal. */
+	TArray<FNiagaraVariable> Variables;
+
+	/** Are there any warnings that were encountered during the traversal of the graph for a given variable? */
+	TArray<FString> PerVariableWarnings;
+
+	/** For each variable that was found, identify the pins that wrote to them in order from first to last write.*/
+	TArray<TArray<const UEdGraphPin*> > PerVariableWriteHistory;
+
+	/** For each variable that was found, identify the pins that read them from the map in order from first to last read. First of the pair has the read pin, second of the pair has the last set that wrote to the pin.*/
+	TArray<TArray<TTuple<const UEdGraphPin*, const UEdGraphPin*> > > PerVariableReadHistory; 
+
+	/** List of pins that manipulated the parameter map from input to output. */
+	TArray<const UEdGraphPin*> MapPinHistory;
+
+	/** List of emitter namespaces encountered as this parameter map was built.*/
+	TArray<FString> EmitterNamespacesEncountered;
+
+	/**
+	* Called in a depth-first traversal to identify a given Niagara Parameter Map pin that was touched during traversal.
+	*/
+	int32 RegisterParameterMapPin(const UEdGraphPin* Pin);
+		
+	/**
+	* Find a variable by name with no concern for type.
+	*/
+	int32 FindVariableByName(const FName& VariableName);
+
+	/**
+	* Find a variable by both name and type. 
+	*/
+	int32 FindVariable(const FName& VariableName, const FNiagaraTypeDefinition& Type);
+
+
+	/**
+	* Add a variable outside the normal get/set node paradigm.
+	*/
+	int32 AddExternalVariable(const FNiagaraVariable& InVar);
+
+	/**
+	* Get the pin that added the parameter map to the graph.
+	*/
+	const UEdGraphPin* GetOriginalPin() const;
+
+	/**
+	* Get the output pin that we traced to build this history object.
+	*/
+	const UEdGraphPin* GetFinalPin() const;
+
+
+	/**
+	* Use the input alias map to resolve any aliases in this input variable name.
+	*/
+	static FNiagaraVariable ResolveAliases(const FNiagaraVariable& InVar, const TMap<FString, FString>& InAliases, const TCHAR* InJoinSeparator);
+		
+	/**
+	* Remove the Particles namespace if it exists.
+	*/
+	static FNiagaraVariable ResolveAsBasicAttribute(const FNiagaraVariable& InVar);
+
+
+	/**
+	* Reverses ResolveAsBasicAttribute.
+	*/
+	static FNiagaraVariable BasicAttributeToNamespacedAttribute(const FNiagaraVariable& InVar);
+
+	
+	/** Prepends the namespace string to the variable name.*/
+	static FNiagaraVariable VariableToNamespacedVariable(const FNiagaraVariable& InVar, FString Namespace);
+		
+	/**
+	* Does this parameter start with the "Module" namespace? Note that the emitter namespace is an alias
+	* for all non-funcion/module script types and will be specialized to the function call node's name using the module.
+	*/
+	static bool IsAliasedModuleParameter(const FNiagaraVariable& InVar);
+
+	/**
+	* Does this parameter start with the "Emitter" namespace? Note that the emitter namespace is an alias
+	* for Emitter and System script types and will be specialized to the name of that specific emitter.
+	*/
+	static bool IsAliasedEmitterParameter(const FNiagaraVariable& InVar);
+	/** Is this parameter in the special "System" namespace?*/
+	static bool IsSystemParameter(const FNiagaraVariable& InVar);
+	/** Is this parameter in the special "Engine" namespace?*/
+	static bool IsEngineParameter(const FNiagaraVariable& InVar);
+	static bool IsUserParameter(const FNiagaraVariable& InVar);
+	
+	/** Take an input string and make it hlsl safe.*/
+	static FString MakeSafeNamespaceString(const FString& InStr);
+
+	/** Does the variable start with this namespace?*/
+	static bool IsInNamespace(const FNiagaraVariable& InVar, const FString& Namespace);
+
+	/** Given the script type and flags, should we treat the System namespace as read-only?*/
+	static bool IsSystemNamespaceReadOnly(UNiagaraScript* InScript);
+
+	/** Given the script type and flags, should we treat the Emitter aliased namespace as read-only?*/
+	static bool IsEmitterNamespaceReadOnly(UNiagaraScript* InScript);
+
+	/** Called to determine if a given variable should be output from a script. It is not static as it requires the overall context to include emitter namespaces visited for system scripts.*/
+	bool IsPrimaryDataSetOutput(const FNiagaraVariable& InVar, UNiagaraScript* InScript);
+	bool IsPrimaryDataSetOutput(const FNiagaraVariable& InVar, ENiagaraScriptUsage InUsage);
+
+	/** Are we required to export this variable as an external constant?*/
+	bool IsExportableExternalConstant(const FNiagaraVariable& InVar, UNiagaraScript* InScript);
+
+	/** Does this variable belong in a namespace that needs to come in as an external constant to this script?*/
+	static bool IsExternalConstantNamespace(const FNiagaraVariable& InVar, UNiagaraScript* InScript);
+	
+	/** Take a non-namespaced variable and move it to an appropriate external constant namespace for this script type.*/
+	static FNiagaraVariable MoveToExternalConstantNamespaceVariable(const FNiagaraVariable& InVar, UNiagaraScript* InScript);
+	
+	/**
+	* Does this parameter start with the "Particles" namespace?
+	*/
+	static bool IsAttribute(const FNiagaraVariable& InVar);
+	
+	/** Does this parameter contain the "Initial" namespace as one of its intermediate namespaces?*/
+	static bool IsInitialValue(const FNiagaraVariable& InVar);
+
+	/** Get the output node associated with this graph.*/
+	const UNiagaraNodeOutput* GetFinalOutputNode() const;
+
+	/** Does this parameter contain the "Initial" namespace as one of its intermediate namespaces? If so, remove the "Initial" namespace and return the original value.*/
+	static FNiagaraVariable GetSourceForInitialValue(const FNiagaraVariable& InVar);
+
+	/**
+	* Helper to add a variable to the known list for a parameter map.
+	*/
+	int32 AddVariable(const FNiagaraVariable& InVar, const UEdGraphPin* InPin);
+
+	/** Get the default value for this variable.*/
+	const UEdGraphPin* GetDefaultValuePin(int32 VarIdx) const;
+
+
+};
+
+
+class FNiagaraParameterMapHistoryBuilder
+{
+public:
+	/** Collection of the build histories from the graph traversal.*/
+	TArray<FNiagaraParameterMapHistory> Histories;
+
+	/** Constructor*/
+	FNiagaraParameterMapHistoryBuilder();
+
+	/** Add a new parameter map to the array.*/
+	int32 CreateParameterMap();
+
+	/** Called in order to set up the correct initial context for an Output node and invokes the UNiagaraNodeOutput's BuildParameterMapHistory method.*/
+	void BuildParameterMaps(UNiagaraNodeOutput* OutputNode, bool bRecursive = true);
+
+	/**
+	* Called first during a node's visitation during traversal to identify that a node has been visited.
+	*/
+	void RegisterNodeVisitation(const UEdGraphNode* Node);
+
+	/** Important. Must be called for each routing of the parameter map. This feeds the list used by TraceParameterMapOutputPin.*/
+	int32 RegisterParameterMapPin(int32 WhichParameterMap, const UEdGraphPin* Pin);
+
+	/** Trace back a pin to whom it was connected to to find the current parameter map to use.*/
+	int32 TraceParameterMapOutputPin(const UEdGraphPin* OutputPin);
+
+	void BeginTranslation(UNiagaraScript* Script);
+
+	void EndTranslation(UNiagaraScript* Script);
+
+	void BeginTranslation(UNiagaraEmitter* Emitter);
+
+	void EndTranslation(UNiagaraEmitter* Emitter);
+
+	/**
+	* Record that we have entered a new function scope.
+	*/
+	void EnterFunction(const FString& InNodeName, class UNiagaraScript* InScript, class UNiagaraNode* Node);
+
+	/**
+	* Record that we have exited a function scope.
+	*/
+	void ExitFunction(const FString& InNodeName, class UNiagaraScript* InScript, class UNiagaraNode* Node);
+
+	/**
+	* Record that we have entered an emitter scope.
+	*/
+	void EnterEmitter(const FString& InEmitterName, class UNiagaraNode* Node);
+
+	/**
+	* Record that we have exited an emitter scope.
+	*/
+	void ExitEmitter(const FString& InEmitterName, class UNiagaraNode* Node);
+
+	/**
+	* Use the current alias map to resolve any aliases in this input variable name.
+	*/
+	FNiagaraVariable ResolveAliases(const FNiagaraVariable& InVar);
+
+	/**
+	* Has RegisterNodeVisitation been called yet on the owning node of this pin?
+	*/
+	bool GetPinPreviouslyVisited(const UEdGraphPin* InPin) const;
+
+	/**
+	* Has RegisterNodeVisitation been called on the input node yet?
+	*/
+	bool GetNodePreviouslyVisited(const class UNiagaraNode* Node) const;
+
+	/** If we haven't already visited the owning nodes, do so.*/
+	void VisitInputPins(class UNiagaraNode*);
+	
+	/** If we haven't already visited the owning node, do so.*/
+	void VisitInputPin(const UEdGraphPin* Pin, class UNiagaraNode*);
+
+	/**
+	* Record that a pin writes to the parameter map. The pin name is expected to be the namespaced parameter map version of the name. If any aliases are in place, they are removed.
+	*/
+	int32 HandleVariableWrite(int32 ParameterMapIndex, const UEdGraphPin* InPin);
+
+	/**
+	* Record that a pin reads from the parameter map. The pin name is expected to be the namespaced parameter map version of the name. If any aliases are in place, they are removed.
+	*/
+	int32 HandleVariableRead(int32 ParameterMapIndex, const UEdGraphPin* InPin, bool RegisterReadsAsVariables, const UEdGraphPin* InDefaultPin, bool& OutUsedDefault);
+	
+	/**
+	* Get the string that the "Module" namespace maps to currently (if it exists)
+	*/
+	const FString* GetModuleAlias() const;
+
+	/**
+	* Get the string that the "Emitter" namespace maps to currently (if it exists)
+	*/
+	const FString* GetEmitterAlias() const;
+
+	/** Get the node calling this sub-graph.*/
+	UNiagaraNode* GetCallingContext() const;
+
+	/** Helper method to identify any matching input nodes from the calling context node to the input variable.*/
+	int32 FindMatchingParameterMapFromContextInputs(const FNiagaraVariable& InVar) const;
+
+	/** In some cases, we don't want all the variables encountered in a traversal. In this case, you can filter 
+	*  the map history to only include variables that are relevant to the specific script type. For instance, a System 
+	*  script doesn't really care about the Particles namespace.
+	*/
+	void EnableScriptWhitelist(bool bInEnable, ENiagaraScriptUsage InScriptType);
+
+	ENiagaraScriptUsage GetCurrentUsageContext()const;
+
+protected:
+	/**
+	* Generate the internal alias map from the current traversal state.
+	*/
+	void BuildCurrentAliases();
+
+	/** Helper function called when variables are added to enable the filtering specified in EnableScriptWhitelist.*/
+	bool ShouldTrackVariable(const FNiagaraVariable& InVar);
+
+	/** Helper method used to take in input script type and determine if the passed in namespaced variable is worth tracking.*/
+	bool IsNamespacedVariableRelevantToScriptType(const FNiagaraVariable& InVar, ENiagaraScriptUsage ScriptType);
+
+	/** Contains the hierarchy of nodes leading to the current graph being processed. Usually made up of FunctionCall and Emitter nodes.*/
+	TArray<UNiagaraNode*> CallingContext;
+	/** Tracker for each context level of the parameter map index associated with a given pin. Used to trace parameter maps through the graph.*/
+	TArray<TMap<const UEdGraphPin*, int32> > PinToParameterMapIndices;
+	/** List of previously visited nodes per context. Note that the same node may be visited multiple times across all graph traversals, but only one time per context level.*/
+	TArray<TArray<const class UNiagaraNode*> > ContextuallyVisitedNodes;
+	/** Contains the hierarchy of emitter node names leading to the current graph being processed.*/
+	TArray<FName> EmitterNameContextStack;
+	/** Contains the hierarchy of function call node names leading to the current graph being processed.*/
+	TArray<FName> FunctionNameContextStack;
+	/** Keeps track of the script usage at the current context level. This allows us to make some decisions about relevence.*/
+	TArray<ENiagaraScriptUsage> RelevantScriptUsageContext;
+	/** Resolved alias map for the current context level. Rebuilt by BuildCurrentAliases.*/
+	TMap<FString, FString> AliasMap;
+	
+	/** Whether or not the script whitelist is active.*/
+	bool bFilterByScriptWhitelist;
+	/** What the script type is that we should be filtering to if the whitelist is enabled.*/
+	ENiagaraScriptUsage FilterScriptType;
+};
\ No newline at end of file
diff --git a/Engine/Plugins/FX/Niagara/Source/NiagaraEditor/Public/NiagaraScriptSource.h b/Engine/Plugins/FX/Niagara/Source/NiagaraEditor/Public/NiagaraScriptSource.h
new file mode 100644
--- /dev/null
+++ b/Engine/Plugins/FX/Niagara/Source/NiagaraEditor/Public/NiagaraScriptSource.h
@@ -0,0 +1,51 @@
+// Copyright 1998-2017 Epic Games, Inc. All Rights Reserved.
+
+
+#pragma once
+
+#include "CoreMinimal.h"
+#include "UObject/ObjectMacros.h"
+#include "NiagaraScriptSourceBase.h"
+#include "INiagaraCompiler.h"
+#include "NiagaraParameterMapHistory.h"
+#include "NiagaraScriptSource.generated.h"
+
+UCLASS(MinimalAPI)
+class UNiagaraScriptSource : public UNiagaraScriptSourceBase
+{
+	GENERATED_UCLASS_BODY()
+
+	/** Graph for particle update expression */
+	UPROPERTY()
+	class UNiagaraGraph*	NodeGraph;
+	
+	// UObject interface
+	virtual void PostLoad() override;
+
+	// UNiagaraScriptSourceBase interface.
+	virtual ENiagaraScriptCompileStatus Compile(UNiagaraScript* ScriptOwner, FString& OutGraphLevelErrorMessages) override;
+	virtual bool IsSynchronized(const FGuid& InChangeId) override;
+	virtual void MarkNotSynchronized() override;
+
+	virtual UNiagaraScriptSourceBase* MakeRecursiveDeepCopy(UObject* DestOuter, TMap<const UObject*, UObject*>& ExistingConversions) const override;
+
+	/** Determine if there are any external dependencies wrt to scripts and ensure that those dependencies are sucked into the existing package.*/
+	virtual void SubsumeExternalDependencies(TMap<const UObject*, UObject*>& ExistingConversions) override;
+
+	virtual FGuid GetChangeID();
+
+	virtual bool IsPreCompiled() const override;
+	virtual void PreCompile(UNiagaraEmitter* Emitter, bool bClearErrors = true) override;
+	virtual bool GatherPreCompiledVariables(const FString& InNamespaceFilter, TArray<FNiagaraVariable>& OutVars) override;
+	virtual void PostCompile() override;
+
+	TArray<FNiagaraParameterMapHistory>& GetPrecomputedHistories() { return PrecompiledHistories; }
+	class UNiagaraGraph* GetPrecomputedNodeGraph() { return NodeGraphDeepCopy; }
+
+private:
+	bool bIsPrecompiled;
+	TArray<FNiagaraParameterMapHistory> PrecompiledHistories;
+
+	UPROPERTY()
+	class UNiagaraGraph* NodeGraphDeepCopy;
+};
diff --git a/Engine/Plugins/FX/Niagara/Source/NiagaraEditor/Public/NiagaraStackEditorData.h b/Engine/Plugins/FX/Niagara/Source/NiagaraEditor/Public/NiagaraStackEditorData.h
new file mode 100644
--- /dev/null
+++ b/Engine/Plugins/FX/Niagara/Source/NiagaraEditor/Public/NiagaraStackEditorData.h
@@ -0,0 +1,64 @@
+// Copyright 1998-2017 Epic Games, Inc. All Rights Reserved.
+
+#pragma once
+
+#include "NiagaraStackEditorData.generated.h"
+
+/** Editor only UI data for emitters. */
+UCLASS()
+class UNiagaraStackEditorData : public UObject
+{
+	GENERATED_BODY()
+
+public:
+	/*
+	 * Gets whether or not a module has been pinned.
+	 * @param ModuleInputKey A unique key for the module. 
+	 */
+	bool GetModuleInputIsPinned(const FString& ModuleInputKey) const;
+
+	/*
+	 * Sets whether or not a module is pinned.
+	 * @param ModuleInputKey A unique key for the module.
+	 * @param bIsPinned Whether or not the module is pinned.
+	 */
+	void SetModuleInputIsPinned(const FString& ModuleInputKey, bool bIsPinned);
+
+	/*
+	* Gets whether or not a module has a rename pending.
+	* @param ModuleInputKey A unique key for the module.
+	*/
+	bool GetModuleInputIsRenamePending(const FString& ModuleInputKey) const;
+
+	/*
+	* Sets whether or not a module is pinned.
+	* @param ModuleInputKey A unique key for the module.
+	* @param bIsRenamePending Whether or not the module is pinned.
+	*/
+	void SetModuleInputIsRenamePending(const FString& ModuleInputKey, bool bIsRenamePending);
+
+	/*
+	 * Gets whether or not a stack entry is Expanded.
+	 * @param bIsExpandedDefault The default value to return if the expanded state hasn't been set for the stack entry.
+	 * @param StackItemKey A unique key for the entry.
+	 */
+	bool GetStackEntryIsExpanded(const FString& StackEntryKey, bool bIsExpandedDefault) const;
+
+	/*
+	 * Sets whether or not a stack entry is Expanded.
+	 * @param StackEntryKey A unique key for the entry.
+	 * @param bIsExpanded Whether or not the entry is expanded.
+	 */
+	void SetStackEntryIsExpanded(const FString& StackEntryKey, bool bIsExpanded);
+
+private:
+
+	UPROPERTY()
+	TMap<FString, bool> ModuleInputKeyToPinnedMap;
+
+	UPROPERTY()
+	TMap<FString, bool> ModuleInputKeyToRenamePendingMap;
+
+	UPROPERTY()
+	TMap<FString, bool> StackEntryKeyToExpandedMap;
+};
\ No newline at end of file
diff --git a/Engine/Plugins/FX/Niagara/Source/NiagaraEditor/Public/ViewModels/Stack/NiagaraParameterHandle.h b/Engine/Plugins/FX/Niagara/Source/NiagaraEditor/Public/ViewModels/Stack/NiagaraParameterHandle.h
new file mode 100644
--- /dev/null
+++ b/Engine/Plugins/FX/Niagara/Source/NiagaraEditor/Public/ViewModels/Stack/NiagaraParameterHandle.h
@@ -0,0 +1,64 @@
+// Copyright 1998-2017 Epic Games, Inc. All Rights Reserved.
+
+#pragma once
+
+#include "Containers/UnrealString.h"
+
+class UNiagaraNodeFunctionCall;
+struct FNiagaraVariable;
+
+class NIAGARAEDITOR_API FNiagaraParameterHandle
+{
+public:
+	FNiagaraParameterHandle();
+
+	FNiagaraParameterHandle(const FString& InParameterHandleString);
+
+	FNiagaraParameterHandle(const FString& InNamespace, const FString& InName);
+
+	bool operator==(const FNiagaraParameterHandle& Other) const;
+
+	static FNiagaraParameterHandle CreateAliasedModuleParameterHandle(const FNiagaraParameterHandle& ModuleParameterHandle, UNiagaraNodeFunctionCall* ModuleNode);
+
+	static FNiagaraParameterHandle CreateEngineParameterHandle(const FNiagaraVariable& SystemVariable);
+
+	static FNiagaraParameterHandle CreateEmitterParameterHandle(const FNiagaraVariable& EmitterVariable);
+
+	static FNiagaraParameterHandle CreateParticleAttributeParameterHandle(const FString& InName);
+
+	static FNiagaraParameterHandle CreateModuleParameterHandle(const FString& InName);
+
+	static FNiagaraParameterHandle CreateInitialParameterHandle(const FNiagaraParameterHandle& Handle);
+
+	bool IsValid() const;
+
+	const FString& GetParameterHandleString() const;
+
+	const FString& GetName() const;
+
+	const FString& GetNamespace() const;
+
+	bool IsEngineHandle() const;
+
+	bool IsSystemHandle() const;
+
+	bool IsEmitterHandle() const;
+
+	bool IsParticleAttributeHandle() const;
+
+	bool IsModuleHandle() const;
+
+public:
+	static const FString UserNamespace;
+	static const FString EngineNamespace;
+	static const FString SystemNamespace;
+	static const FString EmitterNamespace;
+	static const FString ParticleAttributeNamespace;
+	static const FString ModuleNamespace;
+	static const FString InitialPrefix;
+
+private:
+	FString ParameterHandleString;
+	FString Name;
+	FString Namespace;
+};
\ No newline at end of file
diff --git a/Engine/Plugins/FX/Niagara/Source/NiagaraEditor/Public/ViewModels/Stack/NiagaraStackAddEventScriptItem.h b/Engine/Plugins/FX/Niagara/Source/NiagaraEditor/Public/ViewModels/Stack/NiagaraStackAddEventScriptItem.h
new file mode 100644
--- /dev/null
+++ b/Engine/Plugins/FX/Niagara/Source/NiagaraEditor/Public/ViewModels/Stack/NiagaraStackAddEventScriptItem.h
@@ -0,0 +1,33 @@
+// Copyright 1998-2016 Epic Games, Inc. All Rights Reserved.
+
+#pragma once
+
+#include "NiagaraStackEntry.h"
+#include "EdGraph/EdGraphSchema.h"
+#include "NiagaraStackAddModuleItem.h"
+#include "NiagaraStackAddEventScriptItem.generated.h"
+
+class FNiagaraEmitterViewModel;
+
+UCLASS()
+class NIAGARAEDITOR_API UNiagaraStackAddEventScriptItem : public UNiagaraStackAddModuleItem
+{
+	GENERATED_BODY()
+
+public:
+	UNiagaraStackAddEventScriptItem();
+
+	//~ UNiagaraStackEntry interface
+	virtual FText GetDisplayName() const override;
+
+	virtual ENiagaraScriptUsage GetOutputUsage() const override;
+
+protected:
+	//~ UNiagaraStackAddModuleItem interface
+	virtual UNiagaraNodeOutput* GetOrCreateOutputNode() override;
+	virtual FText GetInsertTransactionText() const override;
+
+private:
+	void AddEventScript();
+
+};
\ No newline at end of file
diff --git a/Engine/Plugins/FX/Niagara/Source/NiagaraEditor/Public/ViewModels/Stack/NiagaraStackAddModuleItem.h b/Engine/Plugins/FX/Niagara/Source/NiagaraEditor/Public/ViewModels/Stack/NiagaraStackAddModuleItem.h
new file mode 100644
--- /dev/null
+++ b/Engine/Plugins/FX/Niagara/Source/NiagaraEditor/Public/ViewModels/Stack/NiagaraStackAddModuleItem.h
@@ -0,0 +1,49 @@
+// Copyright 1998-2016 Epic Games, Inc. All Rights Reserved.
+
+#pragma once
+
+#include "NiagaraStackEntry.h"
+#include "NiagaraTypes.h"
+#include "EdGraph/EdGraphSchema.h"
+#include "NiagaraCommon.h"
+#include "NiagaraStackAddModuleItem.generated.h"
+
+class UNiagaraNodeOutput;
+class UNiagaraStackEditorData;
+
+UCLASS()
+class NIAGARAEDITOR_API UNiagaraStackAddModuleItem : public UNiagaraStackEntry
+{
+	GENERATED_BODY()
+public:
+	DECLARE_DELEGATE(FOnItemAdded);
+
+public:
+	void Initialize(TSharedRef<FNiagaraSystemViewModel> InSystemViewModel, TSharedRef<FNiagaraEmitterViewModel> InEmitterViewModel, UNiagaraStackEditorData& InStackEditorData);
+
+	virtual FText GetDisplayName() const override;
+
+	void SetOnItemAdded(FOnItemAdded InOnItemAdded);
+
+	void AddScriptModule(FAssetData ModuleScriptAsset);
+
+	void AddParameterModule(FNiagaraVariable ParameterVariable, bool bRenamePending);
+
+	virtual void GetAvailableParameters(TArray<FNiagaraVariable>& OutAvailableParameterVariables) const;
+
+	virtual void GetNewParameterAvailableTypes(TArray<FNiagaraTypeDefinition>& OutAvailableTypes) const;
+
+	virtual TOptional<FString> GetNewParameterNamespace() const;
+
+	virtual ENiagaraScriptUsage GetOutputUsage() const PURE_VIRTUAL(UNiagaraStackAddModuleItem::GetOutputUsage, return ENiagaraScriptUsage::EmitterSpawnScript;);
+
+protected:
+	virtual UNiagaraNodeOutput* GetOrCreateOutputNode() PURE_VIRTUAL(UNiagaraStackModuleItem::GetOrCreateOutputNode, return nullptr;);
+
+	virtual FText GetInsertTransactionText() const;
+
+protected:
+	FOnItemAdded ItemAddedDelegate;
+
+	UNiagaraStackEditorData* StackEditorData;
+};
\ No newline at end of file
diff --git a/Engine/Plugins/FX/Niagara/Source/NiagaraEditor/Public/ViewModels/Stack/NiagaraStackAddRendererItem.h b/Engine/Plugins/FX/Niagara/Source/NiagaraEditor/Public/ViewModels/Stack/NiagaraStackAddRendererItem.h
new file mode 100644
--- /dev/null
+++ b/Engine/Plugins/FX/Niagara/Source/NiagaraEditor/Public/ViewModels/Stack/NiagaraStackAddRendererItem.h
@@ -0,0 +1,27 @@
+// Copyright 1998-2016 Epic Games, Inc. All Rights Reserved.
+
+#pragma once
+
+#include "NiagaraStackEntry.h"
+#include "EdGraph/EdGraphSchema.h"
+#include "NiagaraStackAddRendererItem.generated.h"
+
+UCLASS()
+class NIAGARAEDITOR_API UNiagaraStackAddRendererItem : public UNiagaraStackEntry
+{
+	GENERATED_BODY()
+
+public:
+	DECLARE_DELEGATE(FOnItemAdded);
+
+public:
+	//~ UNiagaraStackEntry interface
+	virtual FText GetDisplayName() const override;
+
+	void SetOnItemAdded(FOnItemAdded InOnItemAdded);
+
+	void AddRenderer(UClass* RendererClass);
+
+private:
+	FOnItemAdded ItemAddedDelegate;
+};
\ No newline at end of file
diff --git a/Engine/Plugins/FX/Niagara/Source/NiagaraEditor/Public/ViewModels/Stack/NiagaraStackAddScriptModuleItem.h b/Engine/Plugins/FX/Niagara/Source/NiagaraEditor/Public/ViewModels/Stack/NiagaraStackAddScriptModuleItem.h
new file mode 100644
--- /dev/null
+++ b/Engine/Plugins/FX/Niagara/Source/NiagaraEditor/Public/ViewModels/Stack/NiagaraStackAddScriptModuleItem.h
@@ -0,0 +1,32 @@
+// Copyright 1998-2016 Epic Games, Inc. All Rights Reserved.
+
+#pragma once
+
+#include "NiagaraStackAddModuleItem.h"
+#include "EdGraph/EdGraphSchema.h"
+#include "NiagaraStackAddScriptModuleItem.generated.h"
+
+class UNiagaraNodeOutput;
+class UNiagaraStackEditorData;
+
+UCLASS()
+class NIAGARAEDITOR_API UNiagaraStackAddScriptModuleItem : public UNiagaraStackAddModuleItem
+{
+	GENERATED_BODY()
+
+public:
+	void Initialize(TSharedRef<FNiagaraSystemViewModel> InSystemViewModel, TSharedRef<FNiagaraEmitterViewModel> InEmitterViewModel, UNiagaraStackEditorData& InStackEditorData, UNiagaraNodeOutput& InOutputNode);
+
+	virtual void GetAvailableParameters(TArray<FNiagaraVariable>& OutAvailableParameterVariables) const override;
+
+	virtual void GetNewParameterAvailableTypes(TArray<FNiagaraTypeDefinition>& OutAvailableTypes) const override;
+
+	virtual TOptional<FString> GetNewParameterNamespace() const override;
+
+	virtual ENiagaraScriptUsage GetOutputUsage() const override;
+
+	UNiagaraNodeOutput* GetOrCreateOutputNode() override;
+
+private:
+	TWeakObjectPtr<UNiagaraNodeOutput> OutputNode;
+};
\ No newline at end of file
