Copying //UE4/Orion-Staging to //UE4/Main (Source: //Orion/Dev-General @ 3606378)

diff --git a/Engine/Plugins/FX/Niagara/Source/NiagaraEditor/Public/ViewModels/Stack/NiagaraStackEmitterSpawnScriptItemGroup.h b/Engine/Plugins/FX/Niagara/Source/NiagaraEditor/Public/ViewModels/Stack/NiagaraStackEmitterSpawnScriptItemGroup.h
new file mode 100644
--- /dev/null
+++ b/Engine/Plugins/FX/Niagara/Source/NiagaraEditor/Public/ViewModels/Stack/NiagaraStackEmitterSpawnScriptItemGroup.h
@@ -0,0 +1,55 @@
+// Copyright 1998-2016 Epic Games, Inc. All Rights Reserved.
+
+#pragma once
+
+#include "NiagaraStackScriptItemGroup.h"
+#include "NiagaraStackItem.h"
+#include "NiagaraStackEmitterSpawnScriptItemGroup.generated.h"
+
+class FNiagaraEmitterViewModel;
+class FNiagaraScriptViewModel;
+class UNiagaraStackObject;
+class UNiagaraStackSpacer;
+class UNiagaraStackItemExpander;
+
+UCLASS()
+class NIAGARAEDITOR_API UNiagaraStackEmitterPropertiesItem : public UNiagaraStackItem
+{
+	GENERATED_BODY()
+
+public:
+	virtual FText GetDisplayName() const override;
+
+protected:
+	virtual void RefreshChildrenInternal(const TArray<UNiagaraStackEntry*>& CurrentChildren, TArray<UNiagaraStackEntry*>& NewChildren) override;
+
+private:
+	void EmitterExpandedChanged();
+
+private:
+	UPROPERTY()
+	UNiagaraStackObject* EmitterObject;
+
+	UPROPERTY()
+	UNiagaraStackItemExpander* EmitterExpander;
+};
+
+UCLASS()
+class NIAGARAEDITOR_API UNiagaraStackEmitterSpawnScriptItemGroup : public UNiagaraStackScriptItemGroup
+{
+	GENERATED_BODY()
+
+public:
+	UNiagaraStackEmitterSpawnScriptItemGroup();
+
+protected:
+	void RefreshChildrenInternal(const TArray<UNiagaraStackEntry*>& CurrentChildren, TArray<UNiagaraStackEntry*>& NewChildren) override;
+
+private:
+
+	UPROPERTY()
+	UNiagaraStackEmitterPropertiesItem* PropertiesItem;
+
+	UPROPERTY()
+	UNiagaraStackSpacer* PropertiesSpacer;
+};
\ No newline at end of file
diff --git a/Engine/Plugins/FX/Niagara/Source/NiagaraEditor/Public/ViewModels/Stack/NiagaraStackEntry.h b/Engine/Plugins/FX/Niagara/Source/NiagaraEditor/Public/ViewModels/Stack/NiagaraStackEntry.h
new file mode 100644
--- /dev/null
+++ b/Engine/Plugins/FX/Niagara/Source/NiagaraEditor/Public/ViewModels/Stack/NiagaraStackEntry.h
@@ -0,0 +1,105 @@
+// Copyright 1998-2016 Epic Games, Inc. All Rights Reserved.
+
+#pragma once
+
+#include "NiagaraStackEntry.generated.h"
+
+class FNiagaraSystemViewModel;
+class FNiagaraEmitterViewModel;
+
+UCLASS()
+class NIAGARAEDITOR_API UNiagaraStackEntry : public UObject
+{
+	GENERATED_BODY()
+
+public:
+	DECLARE_MULTICAST_DELEGATE(FOnStructureChanged);
+	DECLARE_MULTICAST_DELEGATE_OneParam(FOnDataObjectModified, UObject*);
+
+public:
+	UNiagaraStackEntry();
+
+	void Initialize(TSharedRef<FNiagaraSystemViewModel> InSystemViewModel, TSharedRef<FNiagaraEmitterViewModel> InEmitterViewModel);
+
+	virtual FText GetDisplayName() const;
+
+	virtual FText GetTooltipText() const;
+
+	virtual FName GetTextStyleName() const;
+
+	virtual bool GetCanExpand() const;
+
+	virtual bool IsExpandedByDefault() const;
+
+	virtual bool GetIsExpanded() const;
+
+	virtual void SetIsExpanded(bool bInExpanded);
+
+	virtual FName GetGroupBackgroundName() const;
+
+	virtual FName GetGroupForegroundName() const;
+
+	virtual FName GetItemBackgroundName() const;
+
+	virtual FName GetItemForegroundName() const;
+
+	virtual int32 GetItemIndentLevel() const;
+
+	virtual bool GetShouldShowInStack() const;
+
+	void GetChildren(TArray<UNiagaraStackEntry*>& OutChildren);
+
+	FOnStructureChanged& OnStructureChanged();
+
+	FOnDataObjectModified& OnDataObjectModified();
+
+	void RefreshChildren();
+	void RefreshErrors();
+
+	virtual int32 GetErrorCount() const { return 0; }
+	virtual bool GetErrorFixable(int32 ErrorIdx) const { return false; }
+	virtual bool TryFixError(int32 ErrorIdx) { return false; }
+	virtual FText GetErrorText(int32 ErrorIdx) const { return FText(); }
+	virtual FText GetErrorSummaryText(int32 ErrorIdx) const { return FText(); }
+
+	TSharedRef<FNiagaraSystemViewModel> GetSystemViewModel() const;
+	TSharedRef<FNiagaraEmitterViewModel> GetEmitterViewModel() const;
+
+protected:
+	virtual void RefreshChildrenInternal(const TArray<UNiagaraStackEntry*>& CurrentChildren, TArray<UNiagaraStackEntry*>& NewChildren);
+
+	template<typename ChildType, typename PredicateType>
+	static ChildType* FindCurrentChildOfTypeByPredicate(const TArray<UNiagaraStackEntry*>& CurrentChildren, PredicateType Predicate)
+	{
+		for (UNiagaraStackEntry* CurrentChild : CurrentChildren)
+		{
+			ChildType* TypedCurrentChild = Cast<ChildType>(CurrentChild);
+			if (TypedCurrentChild != nullptr && Predicate(TypedCurrentChild))
+			{
+				return TypedCurrentChild;
+			}
+		}
+		return nullptr;
+	}
+
+private:
+	void ChildStructureChanged();
+	
+	void ChildDataObjectModified(UObject* ChangedObject);
+
+private:
+	TWeakPtr<FNiagaraSystemViewModel> SystemViewModel;
+	TWeakPtr<FNiagaraEmitterViewModel> EmitterViewModel;
+
+	FOnStructureChanged StructureChangedDelegate;
+
+	FOnDataObjectModified DataObjectModifiedDelegate;
+
+	UPROPERTY()
+	TArray<UNiagaraStackEntry*> Children;
+
+	UPROPERTY()
+	TArray<UNiagaraStackEntry*> ErrorChildren;
+
+	bool bIsExpanded;
+};
\ No newline at end of file
diff --git a/Engine/Plugins/FX/Niagara/Source/NiagaraEditor/Public/ViewModels/Stack/NiagaraStackErrorItem.h b/Engine/Plugins/FX/Niagara/Source/NiagaraEditor/Public/ViewModels/Stack/NiagaraStackErrorItem.h
new file mode 100644
--- /dev/null
+++ b/Engine/Plugins/FX/Niagara/Source/NiagaraEditor/Public/ViewModels/Stack/NiagaraStackErrorItem.h
@@ -0,0 +1,39 @@
+// Copyright 1998-2016 Epic Games, Inc. All Rights Reserved.
+
+#pragma once
+
+#include "NiagaraStackEntry.h"
+#include "Reply.h"
+#include "Layout/Visibility.h"
+#include "NiagaraStackErrorItem.generated.h"
+
+UCLASS()
+class NIAGARAEDITOR_API UNiagaraStackErrorItem : public UNiagaraStackEntry
+{
+	GENERATED_BODY()
+
+public:
+	UNiagaraStackErrorItem();
+
+	void Initialize(TSharedRef<FNiagaraSystemViewModel> InSystemViewModel, TSharedRef<FNiagaraEmitterViewModel> InEmitterViewModel, UNiagaraStackEntry* InErrorSource, int32 InErrorIndex);
+
+	//~ UNiagaraStackEntry interface
+	virtual int32 GetItemIndentLevel() const override;
+
+	/** Sets the item indent level for this stack entry. */
+	void SetItemIndentLevel(int32 InItemIndentLevel);
+
+	int32 GetErrorIndex() const { return ErrorIdx; }
+
+	FText ErrorText() const;
+	FText ErrorTextTooltip() const;
+	FReply OnTryFixError();
+	EVisibility CanFixVisibility() const;
+
+	virtual FName GetItemBackgroundName() const override;
+
+protected:
+	UNiagaraStackEntry* ErrorSource;
+	int32 ErrorIdx;
+	int32 ItemIndentLevel;
+};
\ No newline at end of file
diff --git a/Engine/Plugins/FX/Niagara/Source/NiagaraEditor/Public/ViewModels/Stack/NiagaraStackEventHandlerGroup.h b/Engine/Plugins/FX/Niagara/Source/NiagaraEditor/Public/ViewModels/Stack/NiagaraStackEventHandlerGroup.h
new file mode 100644
--- /dev/null
+++ b/Engine/Plugins/FX/Niagara/Source/NiagaraEditor/Public/ViewModels/Stack/NiagaraStackEventHandlerGroup.h
@@ -0,0 +1,43 @@
+// Copyright 1998-2016 Epic Games, Inc. All Rights Reserved.
+
+#pragma once
+
+#include "NiagaraStackItemGroup.h"
+#include "NiagaraStackEventHandlerGroup.generated.h"
+
+class FNiagaraEmitterViewModel;
+class UNiagaraStackAddEventScriptItem;
+
+UCLASS()
+/** Container for one or more NiagaraStackEventScriptItemGroups, allowing multiple event handlers per script.*/
+class NIAGARAEDITOR_API UNiagaraStackEventHandlerGroup : public UNiagaraStackItemGroup
+{
+	GENERATED_BODY()
+
+public:
+	DECLARE_DELEGATE(FOnItemAdded);
+
+public:
+
+	virtual FText GetDisplayName() const override;
+	void SetDisplayName(FText InDisplayName);
+	
+	virtual bool GetShouldShowInStack() const override;
+
+	void SetOnItemAdded(FOnItemAdded InOnItemAdded);
+	
+protected:
+	virtual void RefreshChildrenInternal(const TArray<UNiagaraStackEntry*>& CurrentChildren, TArray<UNiagaraStackEntry*>& NewChildren) override;
+
+	void ChildModifiedGroupItems();
+
+private:
+	FText DisplayName;
+
+	bool bForceRebuild;
+	
+	UPROPERTY()
+	UNiagaraStackAddEventScriptItem* AddScriptItem;
+	
+	FOnItemAdded ItemAddedDelegate;
+};
\ No newline at end of file
diff --git a/Engine/Plugins/FX/Niagara/Source/NiagaraEditor/Public/ViewModels/Stack/NiagaraStackEventScriptItemGroup.h b/Engine/Plugins/FX/Niagara/Source/NiagaraEditor/Public/ViewModels/Stack/NiagaraStackEventScriptItemGroup.h
new file mode 100644
--- /dev/null
+++ b/Engine/Plugins/FX/Niagara/Source/NiagaraEditor/Public/ViewModels/Stack/NiagaraStackEventScriptItemGroup.h
@@ -0,0 +1,34 @@
+// Copyright 1998-2016 Epic Games, Inc. All Rights Reserved.
+
+#pragma once
+
+#include "NiagaraCommon.h"
+#include "NiagaraEmitter.h"
+#include "NiagaraStackScriptItemGroup.h"
+#include "NiagaraStackAddEventScriptItem.h"
+#include "NiagaraStackEventScriptItemGroup.generated.h"
+
+class FNiagaraScriptViewModel;
+
+UCLASS()
+/** Meant to contain a single binding of a Emitter::EventScriptProperties to the stack.*/
+class NIAGARAEDITOR_API UNiagaraStackEventScriptItemGroup : public UNiagaraStackScriptItemGroup
+{
+	GENERATED_BODY()
+
+public:
+	DECLARE_DELEGATE(FOnModifiedEventHandlers);
+
+protected:
+	FOnModifiedEventHandlers OnModifiedEventHandlersDelegate;
+
+	virtual void RefreshChildrenInternal(const TArray<UNiagaraStackEntry*>& CurrentChildren, TArray<UNiagaraStackEntry*>& NewChildren) override;
+
+	virtual bool CanDelete() const override;
+	virtual bool Delete() override;
+	
+public:
+	virtual FText GetDisplayName() const override;
+	void SetOnModifiedEventHandlers(FOnModifiedEventHandlers OnModifiedEventHandlers);
+
+};
\ No newline at end of file
diff --git a/Engine/Plugins/FX/Niagara/Source/NiagaraEditor/Public/ViewModels/Stack/NiagaraStackFunctionInput.h b/Engine/Plugins/FX/Niagara/Source/NiagaraEditor/Public/ViewModels/Stack/NiagaraStackFunctionInput.h
new file mode 100644
--- /dev/null
+++ b/Engine/Plugins/FX/Niagara/Source/NiagaraEditor/Public/ViewModels/Stack/NiagaraStackFunctionInput.h
@@ -0,0 +1,346 @@
+// Copyright 1998-2016 Epic Games, Inc. All Rights Reserved.
+
+#pragma once
+
+#include "NiagaraStackEntry.h"
+#include "NiagaraTypes.h"
+#include "NiagaraParameterHandle.h"
+#include "NiagaraStackFunctionInput.generated.h"
+
+class UNiagaraNodeFunctionCall;
+class UNiagaraNodeAssignment;
+class UNiagaraNodeParameterMapSet;
+class FStructOnScope;
+class UNiagaraStackFunctionInputCollection;
+class UNiagaraStackObject;
+class UNiagaraScript;
+class UEdGraphPin;
+class UNiagaraDataInterface;
+class UNiagaraStackEditorData;
+
+/** Represents a single module input in the module stack view model. */
+UCLASS()
+class NIAGARAEDITOR_API UNiagaraStackFunctionInput : public UNiagaraStackEntry
+{
+	GENERATED_BODY()
+
+public:
+	/** Defines different modes which are used to provide the value for this function input. */
+	enum class EValueMode
+	{
+		/** The value is set to a constant stored locally with this input. */
+		Local,
+		/** The value is linked to a parameter defined outside of this function. */
+		Linked,
+		/** The value is provided by a secondary dynamic input function. */
+		Dynamic,
+		/** The value is provided by a data interface object. */
+		Data,
+		/** The value source for this input was not set, or couldn't be determined. */
+		Invalid
+	};
+
+	DECLARE_MULTICAST_DELEGATE(FOnValueChanged);
+	DECLARE_MULTICAST_DELEGATE(FOnPinnedChanged);
+
+public:
+	UNiagaraStackFunctionInput();
+
+	virtual void BeginDestroy() override;
+
+	static void AddReferencedObjects(UObject* InThis, FReferenceCollector& Collector);
+	void AddReferencedObjects(FReferenceCollector& Collector);
+
+	/** 
+	 * Sets the input data for this entry.
+	 * @param InSystemViewModel The view model for the system which owns the stack containing this entry.
+	 * @param InEmitterViewModel The view model for the emitter which owns the stack containing this entry.
+	 * @param InStackEditorData The stack editor data for this input.
+	 * @param InModuleNode The module function call which owns this input entry. NOTE: This input might not be an input to the module function call, it may be an input to a dynamic input function call which is owned by the module.
+	 * @param InInputFunctionCallNode The function call which this entry is an input to. NOTE: This node can be a module function call node or a dynamic input node.
+	 * @param InInputParameterHandle The input parameter handle for the function call.
+	 * @param InInputType The type of this input.
+	 */
+	void Initialize(
+		TSharedRef<FNiagaraSystemViewModel> InSystemViewModel,
+		TSharedRef<FNiagaraEmitterViewModel> InEmitterViewModel,
+		UNiagaraStackEditorData& InStackEditorData,
+		UNiagaraNodeFunctionCall& InModuleNode,
+		UNiagaraNodeFunctionCall& InInputFunctionCallNode,
+		FString InInputParameterHandle,
+		FNiagaraTypeDefinition InInputType);
+
+	/** Gets the function call node which owns this input. */
+	const UNiagaraNodeFunctionCall& GetInputFunctionCallNode() const;
+
+	/** Gets the current value mode */
+	EValueMode GetValueMode();
+
+	/** Gets whether or not this input can be pinned. */
+	bool GetCanBePinned() const;
+
+	/** Gets the type of this input. */
+	const FNiagaraTypeDefinition& GetInputType() const;
+
+	//~ UNiagaraStackEntry interface
+	virtual FText GetDisplayName() const override;
+	virtual FName GetTextStyleName() const override;
+	virtual bool GetCanExpand() const override;
+	virtual int32 GetItemIndentLevel() const override;
+
+	/** Sets the item indent level for this stack entry. */
+	void SetItemIndentLevel(int32 InItemIndentLevel);
+
+	/** Gets the path of parameter handles from the owning module to the function call which owns this input. */
+	const TArray<FNiagaraParameterHandle>& GetInputParameterHandlePath() const;
+
+	/** Gets the parameter handle which defined this input in the module. */
+	const FNiagaraParameterHandle& GetInputParameterHandle() const;
+
+	/** Gets the handle to the linked value for this input if there is one. */
+	const FNiagaraParameterHandle& GetLinkedValueHandle() const;
+
+	/** Sets the value of this input to a linked parameter handle. */
+	void SetLinkedValueHandle(const FNiagaraParameterHandle& InParameterHandle);
+
+	/** Gets the current set of available parameter handes which can be assigned to this input. */
+	void GetAvailableParameterHandles(TArray<FNiagaraParameterHandle>& AvailableParameterHandles);
+
+	/** Gets the dynamic input node providing the value for this input, if one is available. */
+	UNiagaraNodeFunctionCall* GetDynamicInputNode() const;
+
+	/** Gets the dynamic inputs available for this input. */
+	void GetAvailableDynamicInputs(TArray<UNiagaraScript*>& AvailableDynamicInputs);
+
+	/** Sets the dynamic input script for this input. */
+	void SetDynamicInput(UNiagaraScript* DynamicInput);
+
+	/** Gets the current struct value of this input is there is one. */
+	TSharedPtr<FStructOnScope> GetLocalValueStruct();
+
+	/** Gets the current data object value of this input is there is one. */
+	UNiagaraDataInterface* GetDataValueObject();
+
+	/** Gets whether or not this input is pinned in the UI. */
+	bool GetIsPinned() const;
+
+	/** Sets whether or not this input is pinned in the UI. */
+	void SetIsPinned(bool bIsPinned);
+
+	/** Called to notify the input that an ongoing change to it's value has begun. */
+	void NotifyBeginLocalValueChange();
+
+	/** Called to notify the input that an ongoing change to it's value has ended. */
+	void NotifyEndLocalValueChange();
+
+	/** Sets this input's local value. */
+	void SetLocalValue(TSharedRef<FStructOnScope> InLocalValue);
+
+	/** Returns whether or not the value or handle of this input has been overridden and can be reset. */
+	bool CanReset() const;
+
+	/** Resets the value and handle of this input to the value and handle defined in the module. */
+	void Reset();
+
+	/** Returns whether or not this input can be renamed. */
+	bool CanRenameInput() const;
+
+	/** Gets whether this input has a rename pending. */
+	bool GetIsRenamePending() const;
+
+	/** Sets whether this input has a rename pending. */
+	void SetIsRenamePending(bool bIsRenamePending);
+
+	/** Renames this input to the name specified. */
+	void RenameInput(FString NewName);
+
+	/** Gets the namespaces which new parameters for this input can be read from. */
+	void GetNamespacesForNewParameters(TArray<FString>& OutNamespacesForNewParameters) const;
+
+	/** Gets a multicast delegate which is called whenever the value on this input changes. */
+	FOnValueChanged& OnValueChanged();
+
+	/** Gets a multicast delegate which is called whenever the pinned state of this input changes. */
+	FOnPinnedChanged& OnPinnedChanged();
+
+protected:
+	//~ UNiagaraStackEntry interface
+	virtual void RefreshChildrenInternal(const TArray<UNiagaraStackEntry*>& CurrentChildren, TArray<UNiagaraStackEntry*>& NewChildren) override;
+
+private:
+	struct FDataValues
+	{
+		enum class EDefaultValueOwner
+		{
+			LocallyOwned,
+			FunctionOwned,
+			Invalid
+		};
+
+		FDataValues()
+			: ValueObject(nullptr)
+			, DefaultValueObject(nullptr)
+			, DefaultValueOwner(EDefaultValueOwner::Invalid)
+			, bIsValid(false)
+		{
+		}
+
+		FDataValues(UNiagaraDataInterface* InValueObject, UNiagaraDataInterface* InDefaultValueObject, EDefaultValueOwner InDefaultValueOwner)
+			: ValueObject(InValueObject)
+			, DefaultValueObject(InDefaultValueObject)
+			, DefaultValueOwner(InDefaultValueOwner)
+			, bIsValid(true)
+		{
+			checkf(DefaultValueObject == nullptr || DefaultValueOwner != EDefaultValueOwner::Invalid, TEXT("Must specify a valid owner if the default value object is not null"));
+		}
+
+		UNiagaraDataInterface* GetValueObject() const
+		{
+			return ValueObject;
+		}
+
+		UNiagaraDataInterface* GetDefaultValueObject() const
+		{
+			return DefaultValueObject;
+		}
+
+		UNiagaraDataInterface*& GetDefaultValueObjectRef()
+		{
+			return DefaultValueObject;
+		}
+
+		EDefaultValueOwner GetDefaultValueOwner() const
+		{
+			return DefaultValueOwner;
+		}
+
+		bool IsValid() const
+		{
+			return bIsValid;
+		}
+
+	private:
+		UNiagaraDataInterface* ValueObject;
+		UNiagaraDataInterface* DefaultValueObject;
+		EDefaultValueOwner DefaultValueOwner;
+		bool bIsValid;
+	};
+
+	struct FInputValues
+	{
+		FInputValues()
+			: Mode(EValueMode::Invalid)
+		{
+		}
+
+		TSharedPtr<FStructOnScope> GetLocalStructToReuse();
+		UNiagaraDataInterface* GetDataDefaultValueObjectToReuse();
+
+		EValueMode Mode;
+		TSharedPtr<FStructOnScope> LocalStruct;
+		FNiagaraParameterHandle LinkedHandle;
+		TWeakObjectPtr<UNiagaraNodeFunctionCall> DynamicNode;
+		FDataValues DataObjects;
+	};
+
+private:
+	/** Refreshes the current values for this input from the state of the graph. */
+	void RefreshValues();
+
+	/** Called whenever the graph which generated this input changes. */
+	void OnGraphChanged(const struct FEdGraphEditAction& InAction);
+
+	/** Gets the graph node which owns the local overrides for the module that owns this input if it exists. */
+	UNiagaraNodeParameterMapSet* GetOverrideNode() const;
+
+	/** Gets the graph node which owns the local overrides for the module that owns this input this input.  
+	  * This will create the node and add it to the graph if it doesn't exist. */
+	UNiagaraNodeParameterMapSet& GetOrCreateOverrideNode();
+
+	/** Gets the default value pin from the map get node which generated this input. */
+	UEdGraphPin* GetDefaultPin() const;
+
+	/** Gets the pin on the override node which is associated with this input if it exists. */
+	UEdGraphPin* GetOverridePin() const;
+
+	/** Gets the pin on the override node which is associated with this input.  If either the override node or
+	  * pin don't exist, they will be created. */
+	UEdGraphPin& GetOrCreateOverridePin();
+
+	/** Tries to get a local value for this input if it exists by checking the graph data directly. */
+	bool TryGetCurrentLocalValue(TSharedPtr<FStructOnScope>& LocalValue, UEdGraphPin& ValuePin, TSharedPtr<FStructOnScope> OldValueToReuse);
+
+	/** Tries to get a data interface value for this input if it exists by checking the graph data directly .*/
+	bool TryGetCurrentDataValue(FDataValues& DataValues, UEdGraphPin* OverrideValuePin, UEdGraphPin& DefaultValuePin, UNiagaraDataInterface* LocallyOwnedDefaultDataValueObjectToReuse);
+
+	/** Tries to get the linked value parameter handle for this input if it exists by checking the graph directly. */
+	bool TryGetCurrentLinkedValue(FNiagaraParameterHandle& LinkedValue, UEdGraphPin& ValuePin);
+
+	/** Gets the dynamic input node providing a value to this input if one exists. */
+	bool TryGetCurrentDynamicValue(TWeakObjectPtr<UNiagaraNodeFunctionCall>& DynamicValue, UEdGraphPin* OverridePin);
+
+	/** Recursively removes all nodes connected to the override pin. */
+	static void RemoveAllNodesConnectedToOverridePin(UEdGraphPin& OverridePin, UNiagaraStackFunctionInput* OwningInput);
+
+private:
+	/** The stack editor data for this function input. */
+	UNiagaraStackEditorData* StackEditorData;
+
+	/** The module function call which owns this input entry. NOTE: This input might not be an input to the module function
+		call, it may be an input to a dynamic input function call which is owned by the module. */
+	TWeakObjectPtr<UNiagaraNodeFunctionCall> OwningModuleNode;
+
+	/** The function call which this entry is an input to. NOTE: This node can be a module function call node or a dynamic input node. */
+	TWeakObjectPtr<UNiagaraNodeFunctionCall> OwningFunctionCallNode;
+
+	/** The assignment node which owns this input.  This is only valid for inputs of assignment modules. */
+	TWeakObjectPtr<UNiagaraNodeAssignment> OwningAssignmentNode;
+
+	/** The nigara type definition for this input. */
+	FNiagaraTypeDefinition InputType;
+
+	/** Whether or not this input can be pinned. */
+	bool bCanBePinned;
+
+	/** A unique key for this input for looking up editor only UI data. */
+	FString StackEditorDataKey;
+
+	/** An array representing the path of Namespace.Name handles starting from the owning module to this function input. */
+	TArray<FNiagaraParameterHandle> InputParameterHandlePath;
+
+	/** The parameter handle which defined this input in the module graph. */
+	FNiagaraParameterHandle InputParameterHandle;
+
+	/** The parameter handle which defined this input in the module graph, aliased for use in the current emitter
+	  * graph.  This only affects parameter handles which are local module handles. */
+	FNiagaraParameterHandle AliasedInputParameterHandle;
+
+	/** The name of this input for display in the UI. */
+	FText DisplayName;
+
+	/** Pointers and handles to the various values this input can have. */
+	FInputValues InputValues;
+
+	/** A cached pointer to the override node for this input if it exists.  This value is cached here since the
+	  * UI reads this value every frame due to attribute updates. */
+	mutable TOptional<UNiagaraNodeParameterMapSet*> OverrideNodeCache;
+
+	/** A cached pointer to the override pin for this input if it exists.  This value is cached here since the
+	  * UI reads this value every frame due to attribute updates. */
+	mutable TOptional<UEdGraphPin*> OverridePinCache;
+
+	/** A flag to prevent handling graph changes when it's being updated directly by this object. */
+	bool bUpdatingGraphDirectly;
+
+	/** A handle for removing the graph changed delegate. */
+	FDelegateHandle GraphChangedHandle;
+
+	/** A multicast delegate which is called when the value of this input is changed. */
+	FOnValueChanged ValueChangedDelegate;
+
+	/** A multicast delegate which is called when the pinned state of this input is changed. */
+	FOnPinnedChanged PinnedChangedDelegate;
+
+	/** The item indent level for this stack entry. */
+	int32 ItemIndentLevel;
+};
\ No newline at end of file
diff --git a/Engine/Plugins/FX/Niagara/Source/NiagaraEditor/Public/ViewModels/Stack/NiagaraStackFunctionInputCollection.h b/Engine/Plugins/FX/Niagara/Source/NiagaraEditor/Public/ViewModels/Stack/NiagaraStackFunctionInputCollection.h
new file mode 100644
--- /dev/null
+++ b/Engine/Plugins/FX/Niagara/Source/NiagaraEditor/Public/ViewModels/Stack/NiagaraStackFunctionInputCollection.h
@@ -0,0 +1,67 @@
+// Copyright 1998-2016 Epic Games, Inc. All Rights Reserved.
+
+#pragma once
+
+#include "NiagaraStackEntry.h"
+#include "NiagaraStackFunctionInputCollection.generated.h"
+
+class UNiagaraNodeFunctionCall;
+class UNiagaraStackFunctionInput;
+class UNiagaraStackEditorData;
+
+UCLASS()
+class NIAGARAEDITOR_API UNiagaraStackFunctionInputCollection : public UNiagaraStackEntry
+{
+	GENERATED_BODY()
+
+public:
+	DECLARE_MULTICAST_DELEGATE(FOnInputPinnedChanged);
+	DECLARE_DELEGATE_RetVal_OneParam(bool, FOnFilterChildren, UNiagaraStackFunctionInput*);
+
+public:
+	struct FDisplayOptions
+	{
+		FText DisplayName;
+		bool bShouldShowInStack;
+		int32 ChildItemIndentLevel;
+		FOnFilterChildren ChildFilter;
+	};
+
+public:
+	UNiagaraStackFunctionInputCollection();
+
+	UNiagaraNodeFunctionCall* GetModuleNode() const;
+
+	UNiagaraNodeFunctionCall* GetInputFunctionCallNode() const;
+
+	void Initialize(
+		TSharedRef<FNiagaraSystemViewModel> InSystemViewModel,
+		TSharedRef<FNiagaraEmitterViewModel> InEmitterViewModel,
+		UNiagaraStackEditorData& InStackEditorData,
+		UNiagaraNodeFunctionCall& InModuleNode,
+		UNiagaraNodeFunctionCall& InInputFunctionCallNode,
+		FDisplayOptions InDisplayOptions);
+
+	//~ UNiagaraStackEntry interface
+	virtual FText GetDisplayName() const override;
+	virtual FName GetTextStyleName() const override;
+	virtual bool GetCanExpand() const override;
+	virtual bool GetShouldShowInStack() const override;
+
+	FOnInputPinnedChanged& OnInputPinnedChanged();
+
+protected:
+	virtual void RefreshChildrenInternal(const TArray<UNiagaraStackEntry*>& CurrentChildren, TArray<UNiagaraStackEntry*>& NewChildren) override;
+
+private:
+	void ChildPinnedChanged();
+
+private:
+	UNiagaraStackEditorData* StackEditorData;
+	UNiagaraNodeFunctionCall* ModuleNode;
+	UNiagaraNodeFunctionCall* InputFunctionCallNode;
+
+	FDisplayOptions DisplayOptions;
+
+	FOnInputPinnedChanged InputPinnedChangedDelegate;
+};
\ No newline at end of file
diff --git a/Engine/Plugins/FX/Niagara/Source/NiagaraEditor/Public/ViewModels/Stack/NiagaraStackGraphUtilities.h b/Engine/Plugins/FX/Niagara/Source/NiagaraEditor/Public/ViewModels/Stack/NiagaraStackGraphUtilities.h
new file mode 100644
--- /dev/null
+++ b/Engine/Plugins/FX/Niagara/Source/NiagaraEditor/Public/ViewModels/Stack/NiagaraStackGraphUtilities.h
@@ -0,0 +1,71 @@
+// Copyright 1998-2016 Epic Games, Inc. All Rights Reserved.
+
+#pragma once
+
+#include "NiagaraTypes.h"
+#include "NiagaraCommon.h"
+#include "NiagaraParameterHandle.h"
+
+class UEdGraph;
+class UEdGraphPin;
+class UNiagaraGraph;
+class UNiagaraNode;
+class UNiagaraNodeInput;
+class UNiagaraNodeOutput;
+class UNiagaraNodeFunctionCall;
+class UNiagaraStackEditorData;
+class FNiagaraSystemViewModel;
+class FNiagaraEmitterViewModel;
+
+namespace FNiagaraStackGraphUtilities
+{
+	void RelayoutGraph(UEdGraph& Graph);
+
+	void GetWrittenVariablesForGraph(UEdGraph& Graph, TArray<FNiagaraVariable>& OutWrittenVariables);
+
+	void ConnectPinToInputNode(UEdGraphPin& Pin, UNiagaraNodeInput& InputNode);
+
+	UEdGraphPin* GetParameterMapInputPin(UNiagaraNode& Node);
+
+	UEdGraphPin* GetParameterMapOutputPin(UNiagaraNode& Node);
+
+	void GetOrderedModuleNodes(UNiagaraNodeOutput& OutputNode, TArray<UNiagaraNodeFunctionCall*>& ModuleNodes);
+
+	UNiagaraNodeFunctionCall* GetPreviousModuleNode(UNiagaraNodeFunctionCall& CurrentNode);
+
+	UNiagaraNodeFunctionCall* GetNextModuleNode(UNiagaraNodeFunctionCall& CurrentNode);
+
+	UNiagaraNodeOutput* GetEmitterOutputNodeForStackNode(UNiagaraNode& StackNode);
+
+	UNiagaraNodeInput* GetEmitterInputNodeForStackNode(UNiagaraNode& StackNode);
+
+	struct FStackNodeGroup
+	{
+		TArray<UNiagaraNode*> StartNodes;
+		UNiagaraNode* EndNode;
+	};
+
+	void GetStackNodeGroups(UNiagaraNode& StackNode, TArray<FStackNodeGroup>& OutStackNodeGroups);
+
+	void DisconnectStackNodeGroup(const FStackNodeGroup& DisconnectGroup, const FStackNodeGroup& PreviousGroup, const FStackNodeGroup& NextGroup);
+
+	void ConnectStackNodeGroup(const FStackNodeGroup& ConnectGroup, const FStackNodeGroup& NewPreviousGroup, const FStackNodeGroup& NewNextGroup);
+
+	void InitializeDataInterfaceInputs(TSharedRef<FNiagaraSystemViewModel> SystemViewModel, TSharedRef<FNiagaraEmitterViewModel> EmitterViewModel, UNiagaraStackEditorData& StackEditorData, UNiagaraNodeFunctionCall& ModuleNode, UNiagaraNodeFunctionCall& InputFunctionCallNode);
+
+	FString GenerateStackFunctionInputEditorDataKey(UNiagaraNodeFunctionCall& FunctionCallNode, FNiagaraParameterHandle InputParameterHandle);
+
+	FString GenerateStackModuleEditorDataKey(UNiagaraNodeFunctionCall& ModuleNode);
+
+	enum class ENiagaraGetStackFunctionInputPinsOptions
+	{
+		AllInputs,
+		ModuleInputsOnly
+	};
+
+	void GetStackFunctionInputPins(UNiagaraNodeFunctionCall& FunctionCallNode, TArray<const UEdGraphPin*>& OutInputPins, ENiagaraGetStackFunctionInputPinsOptions Options = ENiagaraGetStackFunctionInputPinsOptions::AllInputs);
+
+	bool ValidateGraphForOutput(UNiagaraGraph& NiagaraGraph, ENiagaraScriptUsage ScriptUsage, int32 ScriptOccurrence, FText& ErrorMessage);
+
+	UNiagaraNodeOutput* ResetGraphForOutput(UNiagaraGraph& NiagaraGraph, ENiagaraScriptUsage ScriptUsage, int32 ScriptOccurrence);
+}
\ No newline at end of file
diff --git a/Engine/Plugins/FX/Niagara/Source/NiagaraEditor/Public/ViewModels/Stack/NiagaraStackItem.h b/Engine/Plugins/FX/Niagara/Source/NiagaraEditor/Public/ViewModels/Stack/NiagaraStackItem.h
new file mode 100644
--- /dev/null
+++ b/Engine/Plugins/FX/Niagara/Source/NiagaraEditor/Public/ViewModels/Stack/NiagaraStackItem.h
@@ -0,0 +1,31 @@
+// Copyright 1998-2016 Epic Games, Inc. All Rights Reserved.
+
+#pragma once
+
+#include "NiagaraStackEntry.h"
+#include "NiagaraStackItem.generated.h"
+
+class UNiagaraStackEditorData;
+
+UCLASS()
+class NIAGARAEDITOR_API UNiagaraStackItem : public UNiagaraStackEntry
+{
+	GENERATED_BODY()
+
+public:
+	DECLARE_DELEGATE(FOnModifiedGroupItems);
+
+public:
+	void Initialize(TSharedRef<FNiagaraSystemViewModel> InSystemViewModel, TSharedRef<FNiagaraEmitterViewModel> InEmitterViewModel, UNiagaraStackEditorData& InStackEditorData);
+
+	void SetOnModifiedGroupItems(FOnModifiedGroupItems OnModifiedGroupItems);
+
+protected:
+	UNiagaraStackEditorData& GetStackEditorData() const;
+
+protected:
+	FOnModifiedGroupItems ModifiedGroupItemsDelegate;
+
+private:
+	UNiagaraStackEditorData* StackEditorData;
+};
\ No newline at end of file
diff --git a/Engine/Plugins/FX/Niagara/Source/NiagaraEditor/Public/ViewModels/Stack/NiagaraStackItemExpander.h b/Engine/Plugins/FX/Niagara/Source/NiagaraEditor/Public/ViewModels/Stack/NiagaraStackItemExpander.h
new file mode 100644
--- /dev/null
+++ b/Engine/Plugins/FX/Niagara/Source/NiagaraEditor/Public/ViewModels/Stack/NiagaraStackItemExpander.h
@@ -0,0 +1,38 @@
+// Copyright 1998-2016 Epic Games, Inc. All Rights Reserved.
+
+#pragma once
+
+#include "NiagaraStackEntry.h"
+#include "NiagaraStackItemExpander.generated.h"
+
+class UNiagaraNodeFunctionCall;
+class UNiagaraStackFunctionInputCollection;
+class UNiagaraStackModuleItemOutputCollection;
+class UNiagaraStackEditorData;
+
+UCLASS()
+class NIAGARAEDITOR_API UNiagaraStackItemExpander : public UNiagaraStackEntry
+{
+	GENERATED_BODY()
+
+public:
+	DECLARE_DELEGATE(FOnExpandedChanged);
+
+public:
+	void Initialize(TSharedRef<FNiagaraSystemViewModel> InSystemViewModel, TSharedRef<FNiagaraEmitterViewModel> InEmitterViewModel, UNiagaraStackEditorData& InStackEditorData, FString InEmitterEditorDataKey, bool bInIsExpandedDefault);
+
+	void SetOnExpnadedChanged(FOnExpandedChanged OnExpandedChanged);
+
+	bool GetIsExpanded() const;
+
+	void ToggleExpanded();
+
+private:
+	UNiagaraStackEditorData* StackEditorData;
+
+	FString EmitterEditorDataKey;
+
+	bool bIsExpandedDefault;
+
+	FOnExpandedChanged ExpandedChangedDelegate;
+};
\ No newline at end of file
diff --git a/Engine/Plugins/FX/Niagara/Source/NiagaraEditor/Public/ViewModels/Stack/NiagaraStackItemGroup.h b/Engine/Plugins/FX/Niagara/Source/NiagaraEditor/Public/ViewModels/Stack/NiagaraStackItemGroup.h
new file mode 100644
--- /dev/null
+++ b/Engine/Plugins/FX/Niagara/Source/NiagaraEditor/Public/ViewModels/Stack/NiagaraStackItemGroup.h
@@ -0,0 +1,41 @@
+// Copyright 1998-2016 Epic Games, Inc. All Rights Reserved.
+
+#pragma once
+
+#include "NiagaraStackEntry.h"
+#include "NiagaraStackItemGroup.generated.h"
+
+class UNiagaraStackEditorData;
+
+UCLASS()
+class NIAGARAEDITOR_API UNiagaraStackItemGroup : public UNiagaraStackEntry
+{
+	GENERATED_BODY()
+
+public:
+	void Initialize(TSharedRef<FNiagaraSystemViewModel> InSystemViewModel, TSharedRef<FNiagaraEmitterViewModel> InEmitterViewModel, UNiagaraStackEditorData& InStackEditorData);
+
+	//~ UNiagaraStackEntry interface
+	virtual bool GetCanExpand() const override;
+	virtual bool GetIsExpanded() const override;
+	virtual void SetIsExpanded(bool bInExpanded) override;
+	virtual FName GetItemBackgroundName() const override;
+	virtual FName GetTextStyleName() const override;
+	virtual FText GetTooltipText() const override;
+
+	virtual void SetTooltipText(const FText& InText);
+
+	virtual bool CanDelete() const { return false; }
+	virtual bool Delete() { return false; }
+
+	virtual bool CanAdd() const { return false; }
+	virtual bool Add() { return false; }
+	
+protected:
+	UNiagaraStackEditorData& GetStackEditorData() const;
+
+private:
+	UNiagaraStackEditorData* StackEditorData;
+
+	FText TooltipText;
+};
\ No newline at end of file
diff --git a/Engine/Plugins/FX/Niagara/Source/NiagaraEditor/Public/ViewModels/Stack/NiagaraStackModuleItem.h b/Engine/Plugins/FX/Niagara/Source/NiagaraEditor/Public/ViewModels/Stack/NiagaraStackModuleItem.h
new file mode 100644
--- /dev/null
+++ b/Engine/Plugins/FX/Niagara/Source/NiagaraEditor/Public/ViewModels/Stack/NiagaraStackModuleItem.h
@@ -0,0 +1,53 @@
+// Copyright 1998-2016 Epic Games, Inc. All Rights Reserved.
+
+#pragma once
+
+#include "NiagaraStackItem.h"
+#include "NiagaraStackModuleItem.generated.h"
+
+class UNiagaraNodeFunctionCall;
+class UNiagaraStackFunctionInputCollection;
+class UNiagaraStackModuleItemOutputCollection;
+class UNiagaraStackItemExpander;
+class UNiagaraStackEditorData;
+
+UCLASS()
+class NIAGARAEDITOR_API UNiagaraStackModuleItem : public UNiagaraStackItem
+{
+	GENERATED_BODY()
+
+public:
+	UNiagaraStackModuleItem();
+
+	const UNiagaraNodeFunctionCall& GetModuleNode() const;
+	void Initialize(TSharedRef<FNiagaraSystemViewModel> InSystemViewModel, TSharedRef<FNiagaraEmitterViewModel> InEmitterViewModel, UNiagaraStackEditorData& InStackEditorData, UNiagaraNodeFunctionCall& InFunctionCallNode);
+
+	virtual FText GetDisplayName() const override;
+	virtual FText GetTooltipText() const override;
+
+	void MoveUp();
+	void MoveDown();
+	void Delete();
+
+protected:
+	virtual void RefreshChildrenInternal(const TArray<UNiagaraStackEntry*>& CurrentChildren, TArray<UNiagaraStackEntry*>& NewChildren) override;
+
+private:
+	void InputPinnedChanged();
+	void ModuleExpandedChanged();
+
+private:
+	UNiagaraNodeFunctionCall* FunctionCallNode;
+
+	UPROPERTY()
+	UNiagaraStackFunctionInputCollection* PinnedInputCollection;
+
+	UPROPERTY()
+	UNiagaraStackFunctionInputCollection* UnpinnedInputCollection;
+
+	UPROPERTY()
+	UNiagaraStackModuleItemOutputCollection* OutputCollection;
+
+	UPROPERTY()
+	UNiagaraStackItemExpander* ModuleExpander;
+};
\ No newline at end of file
diff --git a/Engine/Plugins/FX/Niagara/Source/NiagaraEditor/Public/ViewModels/Stack/NiagaraStackModuleItemOutput.h b/Engine/Plugins/FX/Niagara/Source/NiagaraEditor/Public/ViewModels/Stack/NiagaraStackModuleItemOutput.h
new file mode 100644
--- /dev/null
+++ b/Engine/Plugins/FX/Niagara/Source/NiagaraEditor/Public/ViewModels/Stack/NiagaraStackModuleItemOutput.h
@@ -0,0 +1,48 @@
+// Copyright 1998-2016 Epic Games, Inc. All Rights Reserved.
+
+#pragma once
+
+#include "NiagaraStackEntry.h"
+#include "NiagaraParameterHandle.h"
+#include "NiagaraStackModuleItemOutput.generated.h"
+
+class UNiagaraNodeFunctionCall;
+
+/** Represents a single module Output in the module stack view model. */
+UCLASS()
+class NIAGARAEDITOR_API UNiagaraStackModuleItemOutput : public UNiagaraStackEntry
+{
+	GENERATED_BODY()
+
+public:
+	UNiagaraStackModuleItemOutput();
+
+	/** 
+	 * Sets the Output data for this entry.
+	 * @param InFunctionCallNode The function call node representing the module in the stack graph which owns this Output.
+	 * @param InOutputParameterHandle The Namespace.Name handle for the Output to the owning module.
+	 */
+	void Initialize(TSharedRef<FNiagaraSystemViewModel> InSystemViewModel, TSharedRef<FNiagaraEmitterViewModel> InEmitterViewModel, UNiagaraNodeFunctionCall& InFunctionCallNode, FString InOutputParameterHandle);
+
+	//~ UNiagaraStackEntry interface
+	virtual FText GetDisplayName() const override;
+	virtual FName GetTextStyleName() const override;
+	virtual bool GetCanExpand() const override;
+	virtual int32 GetItemIndentLevel() const override;
+
+	/** Gets the parameter handle which defined this Output in the module. */
+	const FNiagaraParameterHandle& GetOutputParameterHandle() const;
+
+	/** Gets the assigned parameter handle as displayable text. */
+	FText GetOutputParameterHandleText() const;
+
+private:
+	/** The function call node which represents the emitter which owns this Output */
+	TWeakObjectPtr<UNiagaraNodeFunctionCall> FunctionCallNode;
+
+	/** The parameter handle which defined this Output in the module graph. */
+	FNiagaraParameterHandle OutputParameterHandle;
+
+	/** The name of this Output for display in the UI. */
+	FText DisplayName;
+};
\ No newline at end of file
diff --git a/Engine/Plugins/FX/Niagara/Source/NiagaraEditor/Public/ViewModels/Stack/NiagaraStackModuleItemOutputCollection.h b/Engine/Plugins/FX/Niagara/Source/NiagaraEditor/Public/ViewModels/Stack/NiagaraStackModuleItemOutputCollection.h
new file mode 100644
--- /dev/null
+++ b/Engine/Plugins/FX/Niagara/Source/NiagaraEditor/Public/ViewModels/Stack/NiagaraStackModuleItemOutputCollection.h
@@ -0,0 +1,34 @@
+// Copyright 1998-2016 Epic Games, Inc. All Rights Reserved.
+
+#pragma once
+
+#include "NiagaraStackEntry.h"
+#include "NiagaraStackModuleItemOutputCollection.generated.h"
+
+class UNiagaraNodeFunctionCall;
+
+UCLASS()
+class NIAGARAEDITOR_API UNiagaraStackModuleItemOutputCollection : public UNiagaraStackEntry
+{
+	GENERATED_BODY()
+
+public:
+	UNiagaraStackModuleItemOutputCollection();
+
+	void Initialize(TSharedRef<FNiagaraSystemViewModel> InSystemViewModel, TSharedRef<FNiagaraEmitterViewModel> InEmitterViewModel, UNiagaraNodeFunctionCall& InFunctionCallNode);
+
+	//~ UNiagaraStackEntry interface
+	virtual FText GetDisplayName() const override;
+	virtual FName GetTextStyleName() const override;
+	virtual bool GetCanExpand() const override;
+	virtual bool IsExpandedByDefault() const override;
+
+	void SetDisplayName(FText InDisplayName);
+
+protected:
+	virtual void RefreshChildrenInternal(const TArray<UNiagaraStackEntry*>& CurrentChildren, TArray<UNiagaraStackEntry*>& NewChildren) override;
+
+private:
+	UNiagaraNodeFunctionCall* FunctionCallNode;
+	FText DisplayName;
+};
\ No newline at end of file
diff --git a/Engine/Plugins/FX/Niagara/Source/NiagaraEditor/Public/ViewModels/Stack/NiagaraStackObject.h b/Engine/Plugins/FX/Niagara/Source/NiagaraEditor/Public/ViewModels/Stack/NiagaraStackObject.h
new file mode 100644
--- /dev/null
+++ b/Engine/Plugins/FX/Niagara/Source/NiagaraEditor/Public/ViewModels/Stack/NiagaraStackObject.h
@@ -0,0 +1,32 @@
+// Copyright 1998-2016 Epic Games, Inc. All Rights Reserved.
+
+#pragma once
+
+#include "NiagaraStackEntry.h"
+#include "NiagaraStackObject.generated.h"
+
+UCLASS()
+class NIAGARAEDITOR_API UNiagaraStackObject : public UNiagaraStackEntry
+{
+	GENERATED_BODY()
+
+public:
+	UNiagaraStackObject();
+
+	void Initialize(TSharedRef<FNiagaraSystemViewModel> InSystemViewModel, TSharedRef<FNiagaraEmitterViewModel> InEmitterViewModel, UObject* InObject);
+
+	UObject* GetObject();
+
+	//~ UNiagaraStackEntry interface
+	virtual int32 GetItemIndentLevel() const override;
+
+	/** Sets the item indent level for this stack entry. */
+	void SetItemIndentLevel(int32 InItemIndentLevel);
+
+	/** Notifies the stack entry that it's object has been modified. */
+	void NotifyObjectChanged();
+
+private:
+	UObject* Object;
+	int32 ItemIndentLevel;
+};
\ No newline at end of file
diff --git a/Engine/Plugins/FX/Niagara/Source/NiagaraEditor/Public/ViewModels/Stack/NiagaraStackParameterStoreEntry.h b/Engine/Plugins/FX/Niagara/Source/NiagaraEditor/Public/ViewModels/Stack/NiagaraStackParameterStoreEntry.h
new file mode 100644
--- /dev/null
+++ b/Engine/Plugins/FX/Niagara/Source/NiagaraEditor/Public/ViewModels/Stack/NiagaraStackParameterStoreEntry.h
@@ -0,0 +1,137 @@
+// Copyright 1998-2016 Epic Games, Inc. All Rights Reserved.
+
+#pragma once
+
+#include "NiagaraStackEntry.h"
+#include "NiagaraTypes.h"
+#include "NiagaraParameterHandle.h"
+#include "NiagaraParameterStore.h"
+#include "NiagaraStackParameterStoreEntry.generated.h"
+
+class UNiagaraNodeFunctionCall;
+class UNiagaraNodeAssignment;
+class UNiagaraNodeParameterMapSet;
+class FStructOnScope;
+class UNiagaraStackEditorData;
+class UNiagaraStackObject;
+
+/** Represents a single module input in the module stack view model. */
+UCLASS()
+class NIAGARAEDITOR_API UNiagaraStackParameterStoreEntry : public UNiagaraStackEntry
+{
+	GENERATED_BODY()
+
+public:
+	DECLARE_MULTICAST_DELEGATE(FOnValueChanged);
+
+public:
+	UNiagaraStackParameterStoreEntry();
+
+	virtual void BeginDestroy() override;
+	/** 
+	 * Sets the input data for this entry.
+	 * @param InSystemViewModel The view model for the system which owns the stack containing this entry.
+	 * @param InEmitterViewModel The view model for the emitter which owns the stack containing this entry.
+	 * @param InStackEditorData The stack editor data for this input.
+	 * @param InInputParameterHandle The input parameter handle for the function call.
+	 * @param InInputType The type of this input.
+	 */
+	void Initialize(
+		TSharedRef<FNiagaraSystemViewModel> InSystemViewModel,
+		TSharedRef<FNiagaraEmitterViewModel> InEmitterViewModel,
+		UNiagaraStackEditorData& InStackEditorData,
+		UObject* InOwner,
+		FNiagaraParameterStore* InParameterStore,
+		FString InInputParameterHandle,
+		FNiagaraTypeDefinition InInputType);
+
+	/** Gets the type of this input. */
+	const FNiagaraTypeDefinition& GetInputType() const;
+
+	//~ UNiagaraStackEntry interface
+	virtual FText GetDisplayName() const override;
+	virtual FName GetTextStyleName() const override;
+	virtual bool GetCanExpand() const override;
+	virtual int32 GetItemIndentLevel() const override;
+
+	/** Sets the item indent level for this stack entry. */
+	void SetItemIndentLevel(int32 InItemIndentLevel);
+
+	/** Gets the current struct value of this input is there is one. */
+	TSharedPtr<FStructOnScope> GetValueStruct();
+
+	/** Gets the current object value of this input is there is one. */
+	UNiagaraDataInterface* GetValueObject();
+
+	/** Called to notify the input that an ongoing change to it's value has begun. */
+	void NotifyBeginValueChange();
+
+	/** Called to notify the input that an ongoing change to it's value has ended. */
+	void NotifyEndValueChange();
+
+	/** Called to notify the input that it's value has been changed. */
+	void NotifyValueChanged();
+
+	/** Returns whether or not the value or handle of this input has been overridden and can be reset. */
+	bool CanReset() const;
+
+	/** Resets the value and handle of this input to the value and handle defined in the module. */
+	void Reset();
+
+	/** Returns whether or not this input can be renamed. */
+	bool CanRenameInput() const;
+
+	/** Gets whether this input has a rename pending. */
+	bool GetIsRenamePending() const;
+
+	/** Sets whether this input has a rename pending. */
+	void SetIsRenamePending(bool bIsRenamePending);
+
+	/** Renames this input to the name specified. */
+	void RenameInput(FString NewName);
+
+	/** Gets a multicast delegate which is called whenever the value on this input changes. */
+	FOnValueChanged& OnValueChanged();
+
+protected:
+	//~ UNiagaraStackEntry interface
+	virtual void RefreshChildrenInternal(const TArray<UNiagaraStackEntry*>& CurrentChildren, TArray<UNiagaraStackEntry*>& NewChildren) override;
+
+	void RefreshValueAndHandle();
+	TSharedPtr<FNiagaraVariable> GetCurrentValueVariable();
+	UNiagaraDataInterface* GetCurrentValueObject();
+private:
+	/** The stack editor data for this function input. */
+	UNiagaraStackEditorData* StackEditorData;
+
+	/** */
+	FName ParameterName;
+
+	/** The nigara type definition for this input. */
+	FNiagaraTypeDefinition InputType;
+
+	/** The name of this input for display in the UI. */
+	FText DisplayName;
+
+	/** A local copy of the value of this input if one is available. */
+	TSharedPtr<FStructOnScope> LocalValueStruct;
+
+	/** A pointer to the data interface object for this input if one is available. */
+	UPROPERTY()
+	UNiagaraDataInterface* ValueObject;
+
+	/** A multicast delegate which is called when the value of this input is changed. */
+	FOnValueChanged ValueChangedDelegate;
+
+	/** The item indent level for this stack entry. */
+	int32 ItemIndentLevel;
+
+	UPROPERTY()
+	UObject* Owner;
+
+	FNiagaraParameterStore* ParameterStore;
+
+	/** The stack entry for displaying the value object. */
+	UPROPERTY()
+	UNiagaraStackObject* ValueObjectEntry;
+};
diff --git a/Engine/Plugins/FX/Niagara/Source/NiagaraEditor/Public/ViewModels/Stack/NiagaraStackParameterStoreGroup.h b/Engine/Plugins/FX/Niagara/Source/NiagaraEditor/Public/ViewModels/Stack/NiagaraStackParameterStoreGroup.h
new file mode 100644
--- /dev/null
+++ b/Engine/Plugins/FX/Niagara/Source/NiagaraEditor/Public/ViewModels/Stack/NiagaraStackParameterStoreGroup.h
@@ -0,0 +1,45 @@
+// Copyright 1998-2016 Epic Games, Inc. All Rights Reserved.
+
+#pragma once
+
+#include "NiagaraStackItemGroup.h"
+#include "NiagaraCommon.h"
+#include "NiagaraParameterStore.h"
+#include "NiagaraStackParameterStoreGroup.generated.h"
+
+class FNiagaraScriptViewModel;
+class UNiagaraStackEditorData;
+
+UCLASS()
+class NIAGARAEDITOR_API UNiagaraStackParameterStoreGroup : public UNiagaraStackItemGroup
+{
+	GENERATED_BODY()
+
+public:
+	void Initialize(
+		TSharedRef<FNiagaraSystemViewModel> InSystemViewModel,
+		TSharedRef<FNiagaraEmitterViewModel> InEmitterViewModel,
+		UNiagaraStackEditorData& InStackEditorData,
+		TSharedRef<FNiagaraScriptViewModel> InScriptViewModel,
+		UObject* InOwner,
+		FNiagaraParameterStore* InParameterStore);
+	
+	virtual FText GetDisplayName() const override;
+	void SetDisplayName(FText InDisplayName);
+
+protected:
+	virtual void RefreshChildrenInternal(const TArray<UNiagaraStackEntry*>& CurrentChildren, TArray<UNiagaraStackEntry*>& NewChildren) override;
+
+private:
+	void ItemAdded();
+
+	void ChildModifiedGroupItems();
+
+private:
+
+	FText DisplayName;
+
+	UPROPERTY()
+	UObject* Owner;
+	FNiagaraParameterStore* ParameterStore;
+};
\ No newline at end of file
diff --git a/Engine/Plugins/FX/Niagara/Source/NiagaraEditor/Public/ViewModels/Stack/NiagaraStackRenderItemGroup.h b/Engine/Plugins/FX/Niagara/Source/NiagaraEditor/Public/ViewModels/Stack/NiagaraStackRenderItemGroup.h
new file mode 100644
--- /dev/null
+++ b/Engine/Plugins/FX/Niagara/Source/NiagaraEditor/Public/ViewModels/Stack/NiagaraStackRenderItemGroup.h
@@ -0,0 +1,25 @@
+// Copyright 1998-2016 Epic Games, Inc. All Rights Reserved.
+
+#pragma once
+
+#include "NiagaraStackItemGroup.h"
+#include "NiagaraStackRenderItemGroup.generated.h"
+
+UCLASS()
+class NIAGARAEDITOR_API UNiagaraStackRenderItemGroup : public UNiagaraStackItemGroup
+{
+	GENERATED_BODY()
+
+public:
+	virtual FText GetDisplayName() const override;
+	void SetDisplayName(FText InDisplayName);
+
+protected:
+	virtual void RefreshChildrenInternal(const TArray<UNiagaraStackEntry*>& CurrentChildren, TArray<UNiagaraStackEntry*>& NewChildren) override;
+
+private:
+	void ChildModifiedGroupItems();
+
+private:
+	FText DisplayName;
+};
\ No newline at end of file
diff --git a/Engine/Plugins/FX/Niagara/Source/NiagaraEditor/Public/ViewModels/Stack/NiagaraStackRendererItem.h b/Engine/Plugins/FX/Niagara/Source/NiagaraEditor/Public/ViewModels/Stack/NiagaraStackRendererItem.h
new file mode 100644
--- /dev/null
+++ b/Engine/Plugins/FX/Niagara/Source/NiagaraEditor/Public/ViewModels/Stack/NiagaraStackRendererItem.h
@@ -0,0 +1,55 @@
+// Copyright 1998-2016 Epic Games, Inc. All Rights Reserved.
+
+#pragma once
+
+#include "NiagaraStackItem.h"
+#include "NiagaraTypes.h"
+#include "NiagaraStackRendererItem.generated.h"
+
+class UNiagaraEmitter;
+class UNiagaraRendererProperties;
+class UNiagaraStackObject;
+class UNiagaraStackItemExpander;
+
+UCLASS()
+class NIAGARAEDITOR_API UNiagaraStackRendererItem : public UNiagaraStackItem
+{
+	GENERATED_BODY()
+
+public:
+	UNiagaraStackRendererItem();
+
+	void Initialize(TSharedRef<FNiagaraSystemViewModel> InSystemViewModel, TSharedRef<FNiagaraEmitterViewModel> InEmitterViewModel, UNiagaraStackEditorData& InStackEditorData, UNiagaraRendererProperties* InRendererProperties);
+
+	UNiagaraRendererProperties* GetRendererProperties();
+
+	virtual FText GetDisplayName() const override;
+
+	void Delete();
+
+	virtual FName GetItemBackgroundName() const override;
+	virtual int32 GetErrorCount() const override;
+	virtual bool GetErrorFixable(int32 ErrorIdx) const override;
+	virtual bool TryFixError(int32 ErrorIdx) override;
+	virtual FText GetErrorText(int32 ErrorIdx) const override;
+
+	static TArray<FNiagaraVariable> GetMissingVariables(UNiagaraRendererProperties* RendererProperties, UNiagaraEmitter* Emitter);
+	static bool AddMissingVariable(UNiagaraEmitter* Emitter, const FNiagaraVariable& Variable);
+
+protected:
+	virtual void RefreshChildrenInternal(const TArray<UNiagaraStackEntry*>& CurrentChildren, TArray<UNiagaraStackEntry*>& NewChildren) override;
+
+private:
+	void RendererExpandedChanged();
+
+private:
+	UNiagaraRendererProperties* RendererProperties;
+
+	TArray<FNiagaraVariable> MissingAttributes;
+
+	UPROPERTY()
+	UNiagaraStackObject* RendererObject;
+
+	UPROPERTY()
+	UNiagaraStackItemExpander* RendererExpander;
+};
\ No newline at end of file
diff --git a/Engine/Plugins/FX/Niagara/Source/NiagaraEditor/Public/ViewModels/Stack/NiagaraStackRoot.h b/Engine/Plugins/FX/Niagara/Source/NiagaraEditor/Public/ViewModels/Stack/NiagaraStackRoot.h
new file mode 100644
--- /dev/null
+++ b/Engine/Plugins/FX/Niagara/Source/NiagaraEditor/Public/ViewModels/Stack/NiagaraStackRoot.h
@@ -0,0 +1,62 @@
+// Copyright 1998-2016 Epic Games, Inc. All Rights Reserved.
+
+#pragma once
+
+#include "NiagaraStackEntry.h"
+#include "NiagaraStackRoot.generated.h"
+
+class FNiagaraEmitterViewModel;
+class UNiagaraStackEmitterSpawnScriptItemGroup;
+class UNiagaraStackScriptItemGroup;
+class UNiagaraStackScriptItemGroup;
+class UNiagaraStackScriptItemGroup;
+class UNiagaraStackEventHandlerGroup;
+class UNiagaraStackRenderItemGroup;
+class UNiagaraStackParameterStoreGroup;
+
+UCLASS()
+class NIAGARAEDITOR_API UNiagaraStackRoot : public UNiagaraStackEntry
+{
+	GENERATED_BODY()
+
+public:
+	UNiagaraStackRoot();
+
+	void Initialize(TSharedRef<FNiagaraSystemViewModel> InSystemViewModel, TSharedRef<FNiagaraEmitterViewModel> InEmitterViewModel);
+
+	virtual bool GetShouldShowInStack() const override;
+
+protected:
+	virtual void RefreshChildrenInternal(const TArray<UNiagaraStackEntry*>& CurrentChildren, TArray<UNiagaraStackEntry*>& NewChildren) override;
+
+private:
+	void EmitterEventArraysChanged();
+
+private:
+	UPROPERTY()
+	UNiagaraStackScriptItemGroup* SystemSpawnGroup;
+
+	UPROPERTY()
+	UNiagaraStackScriptItemGroup* SystemUpdateGroup;
+
+	UPROPERTY()
+	UNiagaraStackParameterStoreGroup* SystemExposedVariablesGroup;
+
+	UPROPERTY()
+	UNiagaraStackEmitterSpawnScriptItemGroup* EmitterSpawnGroup;
+
+	UPROPERTY()
+	UNiagaraStackScriptItemGroup* EmitterUpdateGroup;
+
+	UPROPERTY()
+	UNiagaraStackScriptItemGroup* ParticleSpawnGroup;
+
+	UPROPERTY()
+	UNiagaraStackScriptItemGroup* ParticleUpdateGroup;
+
+	UPROPERTY()
+	UNiagaraStackEventHandlerGroup* AddEventHandlerGroup;
+
+	UPROPERTY()
+	UNiagaraStackRenderItemGroup* RenderGroup;
+};
\ No newline at end of file
diff --git a/Engine/Plugins/FX/Niagara/Source/NiagaraEditor/Public/ViewModels/Stack/NiagaraStackScriptItemGroup.h b/Engine/Plugins/FX/Niagara/Source/NiagaraEditor/Public/ViewModels/Stack/NiagaraStackScriptItemGroup.h
new file mode 100644
--- /dev/null
+++ b/Engine/Plugins/FX/Niagara/Source/NiagaraEditor/Public/ViewModels/Stack/NiagaraStackScriptItemGroup.h
@@ -0,0 +1,74 @@
+// Copyright 1998-2016 Epic Games, Inc. All Rights Reserved.
+
+#pragma once
+
+#include "NiagaraStackItemGroup.h"
+#include "NiagaraCommon.h"
+#include "NiagaraStackScriptItemGroup.generated.h"
+
+class FNiagaraScriptViewModel;
+class UNiagaraStackEditorData;
+class UNiagaraStackAddScriptModuleItem;
+class UNiagaraStackSpacer;
+class UNiagaraNodeOutput;
+
+UCLASS()
+class NIAGARAEDITOR_API UNiagaraStackScriptItemGroup : public UNiagaraStackItemGroup
+{
+	GENERATED_BODY()
+
+public:
+	UNiagaraStackScriptItemGroup();
+
+	void Initialize(
+		TSharedRef<FNiagaraSystemViewModel> InSystemViewModel,
+		TSharedRef<FNiagaraEmitterViewModel> InEmitterViewModel,
+		UNiagaraStackEditorData& InStackEditorData,
+		TSharedRef<FNiagaraScriptViewModel> InScriptViewModel,
+		ENiagaraScriptUsage InScriptUsage,
+		int32 InScriptOccurrence = 0);
+
+	ENiagaraScriptUsage GetScriptUsage() const { return ScriptUsage; }
+	int32 GetScriptOccurrence() const { return ScriptOccurrence; }
+
+	virtual FText GetDisplayName() const override;
+	void SetDisplayName(FText InDisplayName);
+
+	virtual int32 GetErrorCount() const override;
+	virtual bool GetErrorFixable(int32 ErrorIdx) const override;
+	virtual bool TryFixError(int32 ErrorIdx) override;
+	virtual FText GetErrorText(int32 ErrorIdx) const override;
+	virtual FText GetErrorSummaryText(int32 ErrorIdx) const override;
+
+protected:
+	virtual void RefreshChildrenInternal(const TArray<UNiagaraStackEntry*>& CurrentChildren, TArray<UNiagaraStackEntry*>& NewChildren) override;
+	TSharedPtr<FNiagaraScriptViewModel> ScriptViewModel;
+
+private:
+	void ItemAdded();
+
+	void ChildModifiedGroupItems();
+
+private:
+	ENiagaraScriptUsage ScriptUsage;
+
+	int32 ScriptOccurrence;
+
+	FText DisplayName;
+
+	UPROPERTY()
+	UNiagaraStackAddScriptModuleItem* AddModuleItem;
+
+	UPROPERTY()
+	UNiagaraStackSpacer* BottomSpacer;
+
+	DECLARE_DELEGATE(FFixError);
+	struct FError
+	{
+		FText ErrorText;
+		FText ErrorSummaryText;
+		FFixError Fix;
+	};
+
+	TSharedPtr<FError> Error;
+};
\ No newline at end of file
diff --git a/Engine/Plugins/FX/Niagara/Source/NiagaraEditor/Public/ViewModels/Stack/NiagaraStackSpacer.h b/Engine/Plugins/FX/Niagara/Source/NiagaraEditor/Public/ViewModels/Stack/NiagaraStackSpacer.h
new file mode 100644
--- /dev/null
+++ b/Engine/Plugins/FX/Niagara/Source/NiagaraEditor/Public/ViewModels/Stack/NiagaraStackSpacer.h
@@ -0,0 +1,24 @@
+// Copyright 1998-2016 Epic Games, Inc. All Rights Reserved.
+
+#pragma once
+
+#include "NiagaraStackEntry.h"
+#include "NiagaraStackSpacer.generated.h"
+
+UCLASS()
+class NIAGARAEDITOR_API UNiagaraStackSpacer : public UNiagaraStackEntry
+{
+	GENERATED_BODY()
+
+public:
+	void Initialize(TSharedRef<FNiagaraSystemViewModel> InSystemViewModel, TSharedRef<FNiagaraEmitterViewModel> InEmitterViewModel, FName SpacerKey = NAME_None);
+
+	//~ UNiagaraStackEntry interface
+	virtual FText GetDisplayName() const override;
+	virtual FName GetItemBackgroundName() const override;
+
+	FName GetSpacerKey() const;
+
+private:
+	FName SpacerKey;
+};
\ No newline at end of file
diff --git a/Engine/Plugins/FX/Niagara/Source/NiagaraEditor/Public/ViewModels/Stack/NiagaraStackStruct.h b/Engine/Plugins/FX/Niagara/Source/NiagaraEditor/Public/ViewModels/Stack/NiagaraStackStruct.h
new file mode 100644
--- /dev/null
+++ b/Engine/Plugins/FX/Niagara/Source/NiagaraEditor/Public/ViewModels/Stack/NiagaraStackStruct.h
@@ -0,0 +1,40 @@
+// Copyright 1998-2016 Epic Games, Inc. All Rights Reserved.
+
+#pragma once
+
+#include "NiagaraStackEntry.h"
+#include "StructOnScope.h"
+#include "PropertyEditorDelegates.h"
+#include "NiagaraStackStruct.generated.h"
+
+UCLASS()
+class NIAGARAEDITOR_API UNiagaraStackStruct : public UNiagaraStackEntry
+{
+	GENERATED_BODY()
+
+public:
+	UNiagaraStackStruct();
+
+	void Initialize(TSharedRef<FNiagaraSystemViewModel> InSystemViewModel, TSharedRef<FNiagaraEmitterViewModel> InEmitterViewModel, UObject* InOwningObject, const UStruct* StructClass, uint8* StructData);
+
+	UObject* GetOwningObject();
+	TSharedPtr<FStructOnScope> GetStructOnScope();
+
+	//~ UNiagaraStackEntry interface
+	virtual int32 GetItemIndentLevel() const override;
+
+	/** Sets the item indent level for this stack entry. */
+	void SetItemIndentLevel(int32 InItemIndentLevel);
+
+	bool HasDetailCustomization() const;
+	
+	FOnGetDetailCustomizationInstance GetDetailsCustomization() const;
+	void SetDetailsCustomization(FOnGetDetailCustomizationInstance Customization);
+
+private:
+	UObject* OwningObject;
+	TSharedPtr<FStructOnScope> StructData;
+	int32 ItemIndentLevel;
+
+	FOnGetDetailCustomizationInstance DetailCustomization;
+};
\ No newline at end of file
diff --git a/Engine/Plugins/FX/Niagara/Source/NiagaraEditor/Public/ViewModels/Stack/NiagaraStackViewModel.h b/Engine/Plugins/FX/Niagara/Source/NiagaraEditor/Public/ViewModels/Stack/NiagaraStackViewModel.h
new file mode 100644
--- /dev/null
+++ b/Engine/Plugins/FX/Niagara/Source/NiagaraEditor/Public/ViewModels/Stack/NiagaraStackViewModel.h
@@ -0,0 +1,51 @@
+// Copyright 1998-2016 Epic Games, Inc. All Rights Reserved.
+
+#pragma once
+
+#include "EditorUndoClient.h"
+#include "NiagaraStackViewModel.generated.h"
+
+class FNiagaraSystemViewModel;
+class FNiagaraEmitterViewModel;
+class UNiagaraStackEntry;
+class UNiagaraStackRoot;
+
+UCLASS()
+class NIAGARAEDITOR_API UNiagaraStackViewModel : public UObject, public FEditorUndoClient
+{
+	GENERATED_BODY()
+
+public:
+	DECLARE_MULTICAST_DELEGATE(FOnStructureChanged);
+
+public:
+	TSharedPtr<FNiagaraEmitterViewModel> GetEmitterViewModel();
+	void Initialize(TSharedPtr<FNiagaraSystemViewModel> InSystemViewModel, TSharedPtr<FNiagaraEmitterViewModel> InEmitterViewModel);
+
+	virtual void BeginDestroy() override;
+
+	TArray<UNiagaraStackEntry*>& GetRootEntries();
+
+	FOnStructureChanged& OnStructureChanged();
+
+	//~ FEditorUndoClient Interface
+	virtual void PostUndo(bool bSuccess) override;
+	virtual void PostRedo(bool bSuccess) override { PostUndo(bSuccess); }
+
+private:
+	void EntryStructureChanged();
+	void EntryDataObjectModified(UObject* ChangedObject);
+	void OnSystemCompiled();
+	void OnEmitterCompiled();
+
+private:
+	TSharedPtr<FNiagaraEmitterViewModel> EmitterViewModel;
+	TSharedPtr<FNiagaraSystemViewModel> SystemViewModel;
+
+	TArray<UNiagaraStackEntry*> RootEntries;
+
+	UPROPERTY()
+	UNiagaraStackRoot* RootEntry;
+
+	FOnStructureChanged StructureChangedDelegate;
+};
\ No newline at end of file
diff --git a/Engine/Plugins/FX/Niagara/Source/NiagaraEditorWidgets/Private/DetailCustomizations/NiagaraDataInterfaceCurveDetails.cpp b/Engine/Plugins/FX/Niagara/Source/NiagaraEditorWidgets/Private/DetailCustomizations/NiagaraDataInterfaceCurveDetails.cpp
new file mode 100644
--- /dev/null
+++ b/Engine/Plugins/FX/Niagara/Source/NiagaraEditorWidgets/Private/DetailCustomizations/NiagaraDataInterfaceCurveDetails.cpp
@@ -0,0 +1,371 @@
+// Copyright 1998-2016 Epic Games, Inc. All Rights Reserved.
+
+#include "NiagaraDataInterfaceCurveDetails.h"
+#include "NiagaraCurveOwner.h"
+#include "NiagaraDataInterfaceCurve.h"
+#include "NiagaraDataInterfaceVector2DCurve.h"
+#include "NiagaraDataInterfaceVectorCurve.h"
+#include "NiagaraDataInterfaceVector4Curve.h"
+#include "NiagaraDataInterfaceColorCurve.h"
+#include "NiagaraEditorWidgetsModule.h"
+
+#include "SCurveEditor.h"
+#include "SCurveEditor.h"
+#include "PropertyHandle.h"
+#include "DetailLayoutBuilder.h"
+#include "DetailCategoryBuilder.h"
+#include "DetailWidgetRow.h"
+#include "Widgets/Layout/SBox.h"
+#include "Misc/Optional.h"
+#include "Brushes/SlateColorBrush.h"
+#include "Modules/ModuleManager.h"
+
+class SNiagaraResizeBox : public SCompoundWidget
+{
+public:
+	DECLARE_DELEGATE_OneParam(FOnContentHeightChanged, float);
+
+public:
+	SLATE_BEGIN_ARGS(SNiagaraResizeBox)
+		: _ContentHeight(50)
+		, _HandleHeight(5)
+		, _HandleColor(FLinearColor(0.0f, 0.0f, 0.0f, 0.0f))
+		, _HandleHighlightColor(FLinearColor(1.0f, 1.0f, 1.0f, 0.5f))
+	{}
+		SLATE_ARGUMENT(float, HandleHeight)
+		SLATE_ATTRIBUTE(float, ContentHeight)
+		SLATE_ATTRIBUTE(FLinearColor, HandleColor)
+		SLATE_ATTRIBUTE(FLinearColor, HandleHighlightColor)
+		SLATE_EVENT(FOnContentHeightChanged, ContentHeightChanged)
+		SLATE_DEFAULT_SLOT(FArguments, Content)
+	SLATE_END_ARGS()
+
+	void Construct(const FArguments& InArgs)
+	{
+		ContentHeight = InArgs._ContentHeight;
+		HandleHeight = InArgs._HandleHeight;
+		HandleColor = InArgs._HandleColor;
+		HandleHighlightColor = InArgs._HandleHighlightColor;
+		HandleBrush = FSlateColorBrush(FLinearColor::White);
+		ContentHeightChanged = InArgs._ContentHeightChanged;
+
+		ChildSlot
+		[
+			SNew(SBox)
+			.HeightOverride(this, &SNiagaraResizeBox::GetHeightOverride)
+			.Padding(FMargin(0, 0, 0, HandleHeight))
+			[
+				InArgs._Content.Widget
+			]
+		];
+	}
+
+	virtual FReply OnMouseButtonDown(const FGeometry& MyGeometry, const FPointerEvent& MouseEvent) override
+	{
+		if (MouseEvent.GetEffectingButton() == EKeys::LeftMouseButton)
+		{
+			FVector2D MouseLocation = MyGeometry.AbsoluteToLocal(MouseEvent.GetScreenSpacePosition());
+			if (MyGeometry.GetLocalSize().Y - MouseLocation.Y < HandleHeight)
+			{
+				DragStartLocation = MouseLocation.Y;
+				DragStartContentHeight = ContentHeight.Get();
+				return FReply::Handled().CaptureMouse(SharedThis(this));
+			}
+		}
+		return FReply::Unhandled();
+	}
+
+	virtual FReply OnMouseButtonUp(const FGeometry& MyGeometry, const FPointerEvent& MouseEvent) override
+	{
+		if (this->HasMouseCapture())
+		{
+			return FReply::Handled().ReleaseMouseCapture();
+		}
+		return FReply::Unhandled();
+	}
+
+	virtual FReply OnMouseMove(const FGeometry& MyGeometry, const FPointerEvent& MouseEvent) override
+	{
+		FVector2D MouseLocation = MyGeometry.AbsoluteToLocal(MouseEvent.GetScreenSpacePosition());
+		LastMouseLocation = MouseLocation.Y;
+		if (this->HasMouseCapture())
+		{
+			float NewContentHeight = DragStartContentHeight + (MouseLocation.Y - DragStartLocation);
+			if (ContentHeight.IsBound() && ContentHeightChanged.IsBound())
+			{
+				ContentHeightChanged.Execute(NewContentHeight);
+			}
+			else
+			{
+				ContentHeight = NewContentHeight;
+			}
+			return FReply::Handled();
+		}
+		return FReply::Unhandled();
+	}
+
+	int32 OnPaint(const FPaintArgs& Args, const FGeometry& AllottedGeometry, const FSlateRect& MyClippingRect, FSlateWindowElementList& OutDrawElements, int32 LayerId, const FWidgetStyle& InWidgetStyle, bool bParentEnabled) const
+	{
+		FLinearColor HandleBoxColor;
+		int32 HandleLayerId = LayerId + 1;
+		FVector2D LocalSize = AllottedGeometry.GetLocalSize();
+		if (IsHovered() && LastMouseLocation.IsSet() && LastMouseLocation.GetValue() >= LocalSize.Y - HandleHeight && LastMouseLocation.GetValue() <= LocalSize.Y)
+		{
+			HandleBoxColor = HandleHighlightColor.Get();
+		}
+		else
+		{
+			HandleBoxColor = HandleColor.Get();
+		}
+
+		FVector2D HandleLocation(0, AllottedGeometry.GetLocalSize().Y - HandleHeight);
+		FVector2D HandleSize(AllottedGeometry.GetLocalSize().X, HandleHeight);
+		FSlateDrawElement::MakeBox
+		(
+			OutDrawElements,
+			HandleLayerId,
+			AllottedGeometry.ToPaintGeometry(HandleLocation, HandleSize),
+			&HandleBrush,
+			ESlateDrawEffect::None,
+			HandleBoxColor
+		);
+
+		return SCompoundWidget::OnPaint(Args, AllottedGeometry, MyClippingRect, OutDrawElements, HandleLayerId, InWidgetStyle, bParentEnabled);
+	}
+
+private:
+	FOptionalSize GetHeightOverride() const
+	{
+		return ContentHeight.Get() + HandleHeight;
+	}
+
+private:
+	TOptional<float> LastMouseLocation;
+
+	TAttribute<float> ContentHeight;
+	float HandleHeight;
+
+	float DragStartLocation;
+	float DragStartContentHeight;
+
+	TAttribute<FLinearColor> HandleColor;
+	TAttribute<FLinearColor> HandleHighlightColor;
+	FSlateBrush HandleBrush;
+
+	FOnContentHeightChanged ContentHeightChanged;
+};
+
+class SNiagaraDataInterfaceCurveEditor : public SCompoundWidget
+{
+public:
+	SLATE_BEGIN_ARGS(SNiagaraDataInterfaceCurveEditor) {}
+	SLATE_END_ARGS()
+
+	void Construct(
+		const FArguments& InArgs,
+		TArray<TSharedRef<IPropertyHandle>> InCurveProperties,
+		bool bIsColorCurve,
+		TSharedRef<FNiagaraStackCurveEditorOptions> InStackCurveEditorOptions)
+	{
+		CurveProperties = InCurveProperties;
+		StackCurveEditorOptions = InStackCurveEditorOptions;
+
+		TArray<UObject*> OuterObjects;
+		CurveProperties[0]->GetOuterObjects(OuterObjects);
+		UObject* CurveOwnerObject = OuterObjects[0];
+
+		CurveOwner = MakeShared<FNiagaraCurveOwner>();
+		if (bIsColorCurve)
+		{
+			CurveOwner->SetColorCurves(
+				*GetCurveFromPropertyHandle(CurveProperties[0]),
+				*GetCurveFromPropertyHandle(CurveProperties[1]),
+				*GetCurveFromPropertyHandle(CurveProperties[2]),
+				*GetCurveFromPropertyHandle(CurveProperties[3]),
+				NAME_None,
+				*CurveOwnerObject,
+				FNiagaraCurveOwner::FNotifyCurveChanged::CreateRaw(this, &SNiagaraDataInterfaceCurveEditor::CurveChanged));
+		}
+		else
+		{
+			TArray<FLinearColor> CurveColors{ FLinearColor::Red, FLinearColor::Green, FLinearColor::Blue, FLinearColor::White };
+			int32 ColorIndex = 0;
+			for (TSharedRef<IPropertyHandle> CurveProperty : CurveProperties)
+			{
+				CurveOwner->AddCurve(
+					*GetCurveFromPropertyHandle(CurveProperty),
+					*CurveProperty->GetProperty()->GetDisplayNameText().ToString(),
+					CurveColors[ColorIndex],
+					*CurveOwnerObject,
+					FNiagaraCurveOwner::FNotifyCurveChanged::CreateRaw(this, &SNiagaraDataInterfaceCurveEditor::CurveChanged));
+				ColorIndex++;
+			}
+		}
+
+		ViewMinInput = 0;
+		ViewMaxOutput = 1;
+
+		for (const FRichCurveEditInfo& CurveEditInfo : CurveOwner->GetCurves())
+		{
+			if (CurveEditInfo.CurveToEdit->GetNumKeys())
+			{
+				ViewMinInput = FMath::Min(ViewMinInput, CurveEditInfo.CurveToEdit->GetFirstKey().Time);
+				ViewMaxOutput = FMath::Max(ViewMaxOutput, CurveEditInfo.CurveToEdit->GetLastKey().Time);
+			}
+		}
+
+		ChildSlot
+		[
+			SAssignNew(CurveEditor, SCurveEditor)
+			.HideUI(false)
+			.ViewMinInput(StackCurveEditorOptions.ToSharedRef(), &FNiagaraStackCurveEditorOptions::GetViewMinInput)
+			.ViewMaxInput(StackCurveEditorOptions.ToSharedRef(), &FNiagaraStackCurveEditorOptions::GetViewMaxInput)
+			.ViewMinOutput(StackCurveEditorOptions.ToSharedRef(), &FNiagaraStackCurveEditorOptions::GetViewMinOutput)
+			.ViewMaxOutput(StackCurveEditorOptions.ToSharedRef(), &FNiagaraStackCurveEditorOptions::GetViewMaxOutput)
+			.AreCurvesVisible(StackCurveEditorOptions.ToSharedRef(), &FNiagaraStackCurveEditorOptions::GetAreCurvesVisible)
+			.ZoomToFitVertical(false)
+			.ZoomToFitHorizontal(false)
+			.TimelineLength(StackCurveEditorOptions.ToSharedRef(), &FNiagaraStackCurveEditorOptions::GetTimelineLength)
+			.OnSetInputViewRange(StackCurveEditorOptions.ToSharedRef(), &FNiagaraStackCurveEditorOptions::SetInputViewRange)
+			.OnSetOutputViewRange(StackCurveEditorOptions.ToSharedRef(), &FNiagaraStackCurveEditorOptions::SetOutputViewRange)
+			.OnSetAreCurvesVisible(StackCurveEditorOptions.ToSharedRef(), &FNiagaraStackCurveEditorOptions::SetAreCurvesVisible)
+		];
+
+		CurveEditor->SetCurveOwner(CurveOwner.Get());
+	}
+
+private:
+	FRichCurve* GetCurveFromPropertyHandle(TSharedPtr<IPropertyHandle> Handle)
+	{
+		TArray<void*> RawData;
+		Handle->AccessRawData(RawData);
+		return RawData.Num() == 1 ? static_cast<FRichCurve*>(RawData[0]) : nullptr;
+	}
+
+	void CurveChanged(FRichCurve* ChangedCurve, UObject* CurveOwnerObject)
+	{
+		for (TSharedRef<IPropertyHandle> CurveProperty : CurveProperties)
+		{
+			if (GetCurveFromPropertyHandle(CurveProperty) == ChangedCurve)
+			{
+				CurveProperty->NotifyPostChange();
+				break;
+			}
+		}
+	}
+
+private:
+	float ViewMinInput;
+	float ViewMaxOutput;
+	TArray<TSharedRef<IPropertyHandle>> CurveProperties;
+	TSharedPtr<FNiagaraStackCurveEditorOptions> StackCurveEditorOptions;
+	TSharedPtr<FNiagaraCurveOwner> CurveOwner;
+	TSharedPtr<SCurveEditor> CurveEditor;
+};
+
+
+// Curve Base
+void FNiagaraDataInterfaceCurveDetailsBase::CustomizeDetails(IDetailLayoutBuilder& DetailBuilder)
+{
+	// Only support single objects.
+	TArray<TWeakObjectPtr<UObject>> ObjectsBeingCustomized;
+	DetailBuilder.GetObjectsBeingCustomized(ObjectsBeingCustomized);
+	if (ObjectsBeingCustomized.Num() > 1)
+	{
+		return;
+	}
+
+	FNiagaraEditorWidgetsModule& NiagaraEditorWidgetsModule = FModuleManager::GetModuleChecked<FNiagaraEditorWidgetsModule>("NiagaraEditorWidgets");
+	TSharedRef<FNiagaraStackCurveEditorOptions> StackCurveEditorOptions = NiagaraEditorWidgetsModule.GetOrCreateStackCurveEditorOptionsForObject(
+		ObjectsBeingCustomized[0].Get(), GetDefaultAreCurvesVisible(), GetDefaultHeight());
+
+	TArray<TSharedRef<IPropertyHandle>> CurveProperties;
+	GetCurveProperties(DetailBuilder, CurveProperties);
+
+	for (TSharedRef<IPropertyHandle> CurveProperty : CurveProperties)
+	{
+		CurveProperty->MarkHiddenByCustomization();
+	}
+
+	IDetailCategoryBuilder& CurveCategory = DetailBuilder.EditCategory("Curve");
+	CurveCategory.AddCustomRow(NSLOCTEXT("NiagaraDataInterfaceCurveDetails", "CurveFilterText", "Curve"))
+		.WholeRowContent()
+		[
+			SNew(SNiagaraResizeBox)
+			.ContentHeight(StackCurveEditorOptions, &FNiagaraStackCurveEditorOptions::GetHeight)
+			.ContentHeightChanged(StackCurveEditorOptions, &FNiagaraStackCurveEditorOptions::SetHeight)
+			.Content()
+			[
+				SNew(SNiagaraDataInterfaceCurveEditor, CurveProperties, GetIsColorCurve(), StackCurveEditorOptions)
+			]
+		];
+}
+
+
+// Curve
+TSharedRef<IDetailCustomization> FNiagaraDataInterfaceCurveDetails::MakeInstance()
+{
+	return MakeShared<FNiagaraDataInterfaceCurveDetails>();
+}
+
+void FNiagaraDataInterfaceCurveDetails::GetCurveProperties(IDetailLayoutBuilder& DetailBuilder, TArray<TSharedRef<IPropertyHandle>>& CurveProperties) const
+{
+	CurveProperties.Add(DetailBuilder.GetProperty(GET_MEMBER_NAME_CHECKED(UNiagaraDataInterfaceCurve, Curve)));
+}
+
+
+// Vector 2D Curve
+TSharedRef<IDetailCustomization> FNiagaraDataInterfaceVector2DCurveDetails::MakeInstance()
+{
+	return MakeShared<FNiagaraDataInterfaceVector2DCurveDetails>();
+}
+
+void FNiagaraDataInterfaceVector2DCurveDetails::GetCurveProperties(IDetailLayoutBuilder& DetailBuilder, TArray<TSharedRef<IPropertyHandle>>& OutCurveProperties) const
+{
+	OutCurveProperties.Add(DetailBuilder.GetProperty(GET_MEMBER_NAME_CHECKED(UNiagaraDataInterfaceVector2DCurve, XCurve)));
+	OutCurveProperties.Add(DetailBuilder.GetProperty(GET_MEMBER_NAME_CHECKED(UNiagaraDataInterfaceVector2DCurve, YCurve)));
+}
+
+
+// Vector Curve
+TSharedRef<IDetailCustomization> FNiagaraDataInterfaceVectorCurveDetails::MakeInstance()
+{
+	return MakeShared<FNiagaraDataInterfaceVectorCurveDetails>();
+}
+
+void FNiagaraDataInterfaceVectorCurveDetails::GetCurveProperties(IDetailLayoutBuilder& DetailBuilder, TArray<TSharedRef<IPropertyHandle>>& OutCurveProperties) const
+{
+	OutCurveProperties.Add(DetailBuilder.GetProperty(GET_MEMBER_NAME_CHECKED(UNiagaraDataInterfaceVectorCurve, XCurve)));
+	OutCurveProperties.Add(DetailBuilder.GetProperty(GET_MEMBER_NAME_CHECKED(UNiagaraDataInterfaceVectorCurve, YCurve)));
+	OutCurveProperties.Add(DetailBuilder.GetProperty(GET_MEMBER_NAME_CHECKED(UNiagaraDataInterfaceVectorCurve, ZCurve)));
+}
+
+
+// Vector 4 Curve
+TSharedRef<IDetailCustomization> FNiagaraDataInterfaceVector4CurveDetails::MakeInstance()
+{
+	return MakeShared<FNiagaraDataInterfaceVector4CurveDetails>();
+}
+
+void FNiagaraDataInterfaceVector4CurveDetails::GetCurveProperties(IDetailLayoutBuilder& DetailBuilder, TArray<TSharedRef<IPropertyHandle>>& OutCurveProperties) const
+{
+	OutCurveProperties.Add(DetailBuilder.GetProperty(GET_MEMBER_NAME_CHECKED(UNiagaraDataInterfaceVector4Curve, XCurve)));
+	OutCurveProperties.Add(DetailBuilder.GetProperty(GET_MEMBER_NAME_CHECKED(UNiagaraDataInterfaceVector4Curve, YCurve)));
+	OutCurveProperties.Add(DetailBuilder.GetProperty(GET_MEMBER_NAME_CHECKED(UNiagaraDataInterfaceVector4Curve, ZCurve)));
+	OutCurveProperties.Add(DetailBuilder.GetProperty(GET_MEMBER_NAME_CHECKED(UNiagaraDataInterfaceVector4Curve, WCurve)));
+}
+
+
+// Color Curve
+TSharedRef<IDetailCustomization> FNiagaraDataInterfaceColorCurveDetails::MakeInstance()
+{
+	return MakeShared<FNiagaraDataInterfaceColorCurveDetails>();
+}
+
+void FNiagaraDataInterfaceColorCurveDetails::GetCurveProperties(IDetailLayoutBuilder& DetailBuilder, TArray<TSharedRef<IPropertyHandle>>& OutCurveProperties) const
+{
+	OutCurveProperties.Add(DetailBuilder.GetProperty(GET_MEMBER_NAME_CHECKED(UNiagaraDataInterfaceColorCurve, RedCurve)));
+	OutCurveProperties.Add(DetailBuilder.GetProperty(GET_MEMBER_NAME_CHECKED(UNiagaraDataInterfaceColorCurve, GreenCurve)));
+	OutCurveProperties.Add(DetailBuilder.GetProperty(GET_MEMBER_NAME_CHECKED(UNiagaraDataInterfaceColorCurve, BlueCurve)));
+	OutCurveProperties.Add(DetailBuilder.GetProperty(GET_MEMBER_NAME_CHECKED(UNiagaraDataInterfaceColorCurve, AlphaCurve)));
+}
\ No newline at end of file
diff --git a/Engine/Plugins/FX/Niagara/Source/NiagaraEditorWidgets/Private/NiagaraEditorWidgetsStyle.cpp b/Engine/Plugins/FX/Niagara/Source/NiagaraEditorWidgets/Private/NiagaraEditorWidgetsStyle.cpp
new file mode 100644
--- /dev/null
+++ b/Engine/Plugins/FX/Niagara/Source/NiagaraEditorWidgets/Private/NiagaraEditorWidgetsStyle.cpp
@@ -0,0 +1,122 @@
+// Copyright 1998-2016 Epic Games, Inc. All Rights Reserved.
+
+#include "NiagaraEditorWidgetsStyle.h"
+
+#include "SlateApplication.h"
+#include "EditorStyleSet.h"
+#include "Slate/SlateGameResources.h"
+#include "Styling/SlateStyleRegistry.h"
+#include "Styling/SlateTypes.h"
+
+TSharedPtr< FSlateStyleSet > FNiagaraEditorWidgetsStyle::NiagaraEditorWidgetsStyleInstance = NULL;
+
+void FNiagaraEditorWidgetsStyle::Initialize()
+{
+	if (!NiagaraEditorWidgetsStyleInstance.IsValid())
+	{
+		NiagaraEditorWidgetsStyleInstance = Create();
+		FSlateStyleRegistry::RegisterSlateStyle(*NiagaraEditorWidgetsStyleInstance);
+	}
+}
+
+void FNiagaraEditorWidgetsStyle::Shutdown()
+{
+	FSlateStyleRegistry::UnRegisterSlateStyle(*NiagaraEditorWidgetsStyleInstance);
+	ensure(NiagaraEditorWidgetsStyleInstance.IsUnique());
+	NiagaraEditorWidgetsStyleInstance.Reset();
+}
+
+FName FNiagaraEditorWidgetsStyle::GetStyleSetName()
+{
+	static FName StyleSetName(TEXT("NiagaraEditorWidgetsStyle"));
+	return StyleSetName;
+}
+
+#define IMAGE_BRUSH( RelativePath, ... ) FSlateImageBrush( Style->RootToContentDir( RelativePath, TEXT(".png") ), __VA_ARGS__ )
+#define BOX_BRUSH( RelativePath, ... ) FSlateBoxBrush( Style->RootToContentDir( RelativePath, TEXT(".png") ), __VA_ARGS__ )
+#define BORDER_BRUSH( RelativePath, ... ) FSlateBorderBrush( Style->RootToContentDir( RelativePath, TEXT(".png") ), __VA_ARGS__ )
+#define TTF_FONT( RelativePath, ... ) FSlateFontInfo( Style->RootToContentDir( RelativePath, TEXT(".ttf") ), __VA_ARGS__ )
+#define OTF_FONT( RelativePath, ... ) FSlateFontInfo( Style->RootToContentDir( RelativePath, TEXT(".otf") ), __VA_ARGS__ )
+
+#define IMAGE_CORE_BRUSH( RelativePath, ... ) FSlateImageBrush( FPaths::EngineContentDir() / "Editor/Slate" / RelativePath + TEXT(".png") , __VA_ARGS__ )
+#define BOX_CORE_BRUSH( RelativePath, ... ) FSlateBoxBrush( FPaths::EngineContentDir() / "Editor/Slate" / RelativePath + TEXT(".png"), __VA_ARGS__ )
+#define TTF_CORE_FONT( RelativePath, ... ) FSlateFontInfo( FPaths::EngineContentDir()  / "Slate" / RelativePath + TEXT(".ttf"), __VA_ARGS__ )
+
+const FVector2D Icon8x8(8.0f, 8.0f);
+const FVector2D Icon16x16(16.0f, 16.0f);
+const FVector2D Icon20x20(20.0f, 20.0f);
+const FVector2D Icon40x40(40.0f, 40.0f);
+
+TSharedRef< FSlateStyleSet > FNiagaraEditorWidgetsStyle::Create()
+{
+	const FTextBlockStyle NormalText = FEditorStyle::GetWidgetStyle<FTextBlockStyle>("NormalText");
+	const FEditableTextBoxStyle NormalEditableTextBox = FCoreStyle::Get().GetWidgetStyle<FEditableTextBoxStyle>("NormalEditableTextBox");
+	const FSpinBoxStyle NormalSpinBox = FEditorStyle::GetWidgetStyle<FSpinBoxStyle>("SpinBox");
+
+	TSharedRef< FSlateStyleSet > Style = MakeShareable(new FSlateStyleSet("NiagaraEditorWidgetsStyle"));
+	Style->SetContentRoot(FPaths::EngineContentDir() / TEXT("Editor/Slate/Niagara"));
+
+	// Stack
+	FSlateFontInfo StackGroupFont = TTF_CORE_FONT("Fonts/Roboto-Bold", 10);
+	FTextBlockStyle StackGroupText = FTextBlockStyle(NormalText)
+		.SetFont(StackGroupFont)
+		.SetColorAndOpacity(FLinearColor(1.0f, 1.0f, 1.0f, 1.0f))
+		.SetShadowOffset(FVector2D(0, 1))
+		.SetShadowColorAndOpacity(FLinearColor(0, 0, 0, 0.9f));
+	Style->Set("NiagaraEditor.Stack.GroupText", StackGroupText);
+
+	FSlateFontInfo StackDefaultFont = TTF_CORE_FONT("Fonts/Roboto-Regular", 10);
+	FTextBlockStyle StackDefaultText = FTextBlockStyle(NormalText)
+		.SetFont(StackDefaultFont);
+	Style->Set("NiagaraEditor.Stack.DefaultText", StackDefaultText);
+
+	FSlateFontInfo ParameterFont = TTF_CORE_FONT("Fonts/Roboto-Regular", 8);
+	FTextBlockStyle ParameterText = FTextBlockStyle(NormalText)
+		.SetFont(ParameterFont);
+	Style->Set("NiagaraEditor.Stack.ParameterText", ParameterText);
+
+	FSlateFontInfo ParameterCollectionFont = TTF_CORE_FONT("Fonts/Roboto-Regular", 9);
+	FTextBlockStyle ParameterCollectionText = FTextBlockStyle(NormalText)
+		.SetFont(ParameterCollectionFont);
+	Style->Set("NiagaraEditor.Stack.ParameterCollectionText", ParameterCollectionText);
+
+	FSlateFontInfo StackItemFont = TTF_CORE_FONT("Fonts/Roboto-Regular", 11);
+	FTextBlockStyle StackItemText = FTextBlockStyle(NormalText)
+		.SetFont(StackItemFont);
+	Style->Set("NiagaraEditor.Stack.ItemText", StackItemText);
+
+	Style->Set("NiagaraEditor.Stack.Group.BackgroundColor", FLinearColor(FColor(96, 96, 96)));
+
+	Style->Set("NiagaraEditor.Stack.Item.BackgroundColor", FLinearColor(FColor(35, 35, 35)));
+
+	Style->Set("NiagaraEditor.Stack.Item.ErrorBackgroundColor", FLinearColor(FColor(35, 0, 0)));
+
+	Style->Set("NiagaraEditor.Stack.ItemHeaderFooter.BackgroundBrush", new FSlateColorBrush(FLinearColor(FColor(20, 20, 20))));
+
+	Style->Set("NiagaraEditor.Stack.ForegroundColor", FLinearColor(FColor(191, 191, 191)));
+
+	Style->Set("NiagaraEditor.Stack.Splitter", FSplitterStyle()
+		.SetHandleNormalBrush(IMAGE_CORE_BRUSH("Common/SplitterHandleHighlight", Icon8x8, FLinearColor(.1f, .1f, .1f, 1.0f)))
+		.SetHandleHighlightBrush(IMAGE_CORE_BRUSH("Common/SplitterHandleHighlight", Icon8x8, FLinearColor::White))
+	);
+
+	return Style;
+}
+
+#undef IMAGE_BRUSH
+#undef BOX_BRUSH
+#undef BORDER_BRUSH
+#undef TTF_FONT
+#undef OTF_FONT
+#undef TTF_CORE_FONT
+#undef BOX_CORE_BRUSH
+
+void FNiagaraEditorWidgetsStyle::ReloadTextures()
+{
+	FSlateApplication::Get().GetRenderer()->ReloadTextureResources();
+}
+
+const ISlateStyle& FNiagaraEditorWidgetsStyle::Get()
+{
+	return *NiagaraEditorWidgetsStyleInstance;
+}
\ No newline at end of file
diff --git a/Engine/Plugins/FX/Niagara/Source/NiagaraEditorWidgets/Private/NiagaraEditorWidgetsStyle.h b/Engine/Plugins/FX/Niagara/Source/NiagaraEditorWidgets/Private/NiagaraEditorWidgetsStyle.h
new file mode 100644
--- /dev/null
+++ b/Engine/Plugins/FX/Niagara/Source/NiagaraEditorWidgets/Private/NiagaraEditorWidgetsStyle.h
@@ -0,0 +1,30 @@
+// Copyright 1998-2016 Epic Games, Inc. All Rights Reserved.
+
+#pragma once
+#include "ISlateStyle.h"
+
+/** Manages the style which provides resources for niagara editor widgets. */
+class FNiagaraEditorWidgetsStyle
+{
+public:
+
+	static void Initialize();
+
+	static void Shutdown();
+
+	/** reloads textures used by slate renderer */
+	static void ReloadTextures();
+
+	/** @return The Slate style set for niagara editor widgets */
+	static const ISlateStyle& Get();
+
+	static FName GetStyleSetName();
+
+private:
+
+	static TSharedRef< class FSlateStyleSet > Create();
+
+private:
+
+	static TSharedPtr< class FSlateStyleSet > NiagaraEditorWidgetsStyleInstance;
+};
\ No newline at end of file
