Copying //UE4/Orion-Staging to //UE4/Main (Source: //Orion/Dev-General @ 3606378)

diff --git a/Engine/Plugins/FX/Niagara/Source/NiagaraEditorWidgets/Private/SNiagaraStack.cpp b/Engine/Plugins/FX/Niagara/Source/NiagaraEditorWidgets/Private/SNiagaraStack.cpp
new file mode 100644
--- /dev/null
+++ b/Engine/Plugins/FX/Niagara/Source/NiagaraEditorWidgets/Private/SNiagaraStack.cpp
@@ -0,0 +1,2068 @@
+// Copyright 1998-2016 Epic Games, Inc. All Rights Reserved.
+
+#include "SNiagaraStack.h"
+
+#include "NiagaraEditorModule.h"
+#include "NiagaraTypes.h"
+#include "INiagaraEditorTypeUtilities.h"
+#include "SNiagaraParameterEditor.h"
+#include "NiagaraEditorWidgetsStyle.h"
+#include "NiagaraEditorStyle.h"
+#include "NiagaraScript.h"
+#include "ViewModels/Stack/NiagaraStackViewModel.h"
+#include "ViewModels/Stack/NiagaraStackEntry.h"
+#include "ViewModels/Stack/NiagaraStackSpacer.h"
+#include "ViewModels/Stack/NiagaraStackItemGroup.h"
+#include "ViewModels/Stack/NiagaraStackItem.h"
+#include "ViewModels/Stack/NiagaraStackItemExpander.h"
+#include "ViewModels/Stack/NiagaraStackFunctionInputCollection.h"
+#include "ViewModels/Stack/NiagaraStackFunctionInput.h"
+#include "ViewModels/Stack/NiagaraStackModuleItemOutputCollection.h"
+#include "ViewModels/Stack/NiagaraStackModuleItemOutput.h"
+#include "ViewModels/Stack/NiagaraStackModuleItem.h"
+#include "ViewModels/Stack/NiagaraStackRendererItem.h"
+#include "ViewModels/Stack/NiagaraStackAddModuleItem.h"
+#include "ViewModels/Stack/NiagaraStackAddRendererItem.h"
+#include "ViewModels/Stack/NiagaraStackObject.h"
+#include "ViewModels/Stack/NiagaraStackErrorItem.h"
+#include "ViewModels/Stack/NiagaraStackStruct.h"
+#include "ViewModels/Stack/NiagaraStackParameterStoreEntry.h"
+#include "FrontendFilterBase.h"
+#include "Framework/Commands/UIAction.h"
+#include "Framework/MultiBox/MultiBoxBuilder.h"
+#include "ModuleManager.h"
+#include "AssetRegistryModule.h"
+#include "EdGraph/EdGraphSchema.h"
+#include "Modules/ModuleManager.h"
+#include "IStructureDetailsView.h"
+#include "IDetailsView.h"
+#include "Editor/PropertyEditor/Public/PropertyEditorModule.h"
+#include "IContentBrowserSingleton.h"
+#include "ContentBrowserModule.h"
+#include "SGraphActionMenu.h"
+#include "SButton.h"
+#include "SImage.h"
+#include "SComboButton.h"
+#include "SSpacer.h"
+#include "SlateApplication.h"
+#include "Widgets/Colors/SSimpleGradient.h"
+#include "Widgets/Text/SInlineEditableTextBlock.h"
+#include "NiagaraNodeAssignment.h"
+#include "UObjectGlobals.h"
+#include "NiagaraDataInterface.h"
+#include "ClassViewerModule.h"
+#include "ClassViewerFilter.h"
+#include "NiagaraRendererProperties.h"
+#include "NiagaraConstants.h"
+#include "NotifyHook.h"
+#include "NiagaraEmitter.h"
+#include "Customizations/NiagaraEventScriptPropertiesCustomization.h"
+#include "EdGraph/EdGraph.h"
+#include "Toolkits/AssetEditorManager.h"
+#include "NiagaraEditorUtilities.h"
+
+#define LOCTEXT_NAMESPACE "NiagaraStack"
+
+const float IndentSize = 16;
+const float TextIconSize = 16;
+
+class SNiagaraStackFunctionInputName : public SCompoundWidget
+{
+public:
+	DECLARE_DELEGATE_OneParam(FOnColumnWidthChanged, float);
+
+public:
+	SLATE_BEGIN_ARGS(SNiagaraStackFunctionInputName) { }
+		SLATE_ATTRIBUTE(bool, IsRowActive)
+	SLATE_END_ARGS();
+
+	void Construct(const FArguments& InArgs, UNiagaraStackFunctionInput* InFunctionInput)
+	{
+		FunctionInput = InFunctionInput;
+		IsRowActive = InArgs._IsRowActive;
+
+		PinIsPinnedColor = FNiagaraEditorWidgetsStyle::Get().GetColor(FunctionInput->GetItemForegroundName());
+		PinIsUnpinnedColor = PinIsPinnedColor.Desaturate(.4f);
+
+		ChildSlot
+		[
+			SNew(SHorizontalBox)
+			// Name Label
+			+ SHorizontalBox::Slot()
+			.VAlign(VAlign_Center)
+			[
+				SAssignNew(NameTextBlock, SInlineEditableTextBlock)
+				.Style(FNiagaraEditorStyle::Get(), "NiagaraEditor.ParameterInlineEditableText")
+				.Text_UObject(FunctionInput, &UNiagaraStackEntry::GetDisplayName)
+				.IsReadOnly(this, &SNiagaraStackFunctionInputName::GetIsNameReadOnly)
+				.IsSelected(this, &SNiagaraStackFunctionInputName::GetIsNameWidgetSelected)
+				.OnTextCommitted(this, &SNiagaraStackFunctionInputName::OnNameTextCommitted)
+			]
+			// Pin
+			+ SHorizontalBox::Slot()
+			.VAlign(VAlign_Center)
+			.AutoWidth()
+			[
+				SNew(SButton)
+				.IsFocusable(false)
+				.Visibility(this, &SNiagaraStackFunctionInputName::GetPinVisibility)
+				.ToolTipText(LOCTEXT("PinToolTip", "Pin this input"))
+				.ButtonStyle(FEditorStyle::Get(), "HoverHintOnly")
+				.ForegroundColor(this, &SNiagaraStackFunctionInputName::GetPinColor)
+				.ContentPadding(2)
+				.OnClicked(this, &SNiagaraStackFunctionInputName::PinButtonPressed)
+				.Content()
+				[
+					SNew(STextBlock)
+					.Font(FEditorStyle::Get().GetFontStyle("FontAwesome.9"))
+					.Text(FText::FromString(FString(TEXT("\xf08d"))))
+					.RenderTransformPivot(FVector2D(.5f, .5f))
+				]
+			]
+		];
+	}
+
+	virtual void Tick(const FGeometry& AllottedGeometry, const double InCurrentTime, const float InDeltaTime) override
+	{
+		if (FunctionInput->GetIsRenamePending() && NameTextBlock.IsValid())
+		{
+			NameTextBlock->EnterEditingMode();
+			FunctionInput->SetIsRenamePending(false);
+		}
+		SCompoundWidget::Tick(AllottedGeometry, InCurrentTime, InDeltaTime);
+	}
+
+private:
+	FReply PinButtonPressed() const
+	{
+		FunctionInput->SetIsPinned(!FunctionInput->GetIsPinned());
+		return FReply::Handled();
+	}
+
+	FSlateColor GetPinColor() const
+	{
+		return FunctionInput->GetIsPinned() ? PinIsPinnedColor : PinIsUnpinnedColor;
+	}
+
+	EVisibility GetPinVisibility() const
+	{
+		if (FunctionInput->GetCanBePinned())
+		{
+			return IsRowActive.Get() || FunctionInput->GetIsPinned() ? EVisibility::Visible : EVisibility::Hidden;
+		}
+		else
+		{
+			return EVisibility::Collapsed;
+		}
+	}
+
+	bool GetIsNameReadOnly() const
+	{
+		return FunctionInput->CanRenameInput() == false;
+	}
+
+	bool GetIsNameWidgetSelected() const
+	{
+		return true;
+	}
+
+	void OnNameTextCommitted(const FText& InText, ETextCommit::Type InCommitType)
+	{
+		FunctionInput->RenameInput(InText.ToString());
+	}
+
+private:
+	FLinearColor PinIsPinnedColor;
+	FLinearColor PinIsUnpinnedColor;
+
+	UNiagaraStackFunctionInput* FunctionInput;
+
+	TSharedPtr<SInlineEditableTextBlock> NameTextBlock;
+
+	TAttribute<bool> IsRowActive;
+};
+
+class SNiagaraStackFunctionInputValue : public SCompoundWidget
+{
+public:
+	DECLARE_DELEGATE_OneParam(FOnColumnWidthChanged, float)
+public:
+	SLATE_BEGIN_ARGS(SNiagaraStackFunctionInputValue) { }
+	SLATE_END_ARGS();
+
+	void Construct(const FArguments& InArgs, UNiagaraStackFunctionInput* InFunctionInput)
+	{
+		FunctionInput = InFunctionInput;
+
+		FunctionInput->OnValueChanged().AddSP(this, &SNiagaraStackFunctionInputValue::OnInputValueChanged);
+		DisplayedLocalValueStruct = FunctionInput->GetLocalValueStruct();
+
+		FMargin ItemPadding = FMargin(0);
+		ChildSlot
+		[
+			// Values
+			SNew(SHorizontalBox)
+			+ SHorizontalBox::Slot()
+			.VAlign(VAlign_Center)
+			.AutoWidth()
+			.Padding(0, 0, 3, 0)
+			[
+				// Value Icon
+				SNew(SBox)
+				.WidthOverride(TextIconSize)
+				.VAlign(VAlign_Center)
+				[
+					SNew(STextBlock)
+					.Font(FEditorStyle::Get().GetFontStyle("FontAwesome.10"))
+					.Text(this, &SNiagaraStackFunctionInputValue::GetInputIconText)
+					.ToolTipText(this, &SNiagaraStackFunctionInputValue::GetInputIconToolTip)
+					.ColorAndOpacity(this, &SNiagaraStackFunctionInputValue::GetInputIconColor)
+				]
+			]
+			+ SHorizontalBox::Slot()
+			[
+				// TODO Don't generate all of these widgets for every input, only generate the ones that are used based on the value type.
+
+				// Local struct
+				SNew(SOverlay)
+				+ SOverlay::Slot()
+				[
+					SAssignNew(LocalValueStructContainer, SBox)
+					.Visibility(this, &SNiagaraStackFunctionInputValue::GetValueWidgetVisibility, UNiagaraStackFunctionInput::EValueMode::Local)
+					[
+						ConstructLocalValueStructWidget()
+					]
+				]
+
+				// Linked handle
+				+ SOverlay::Slot()
+				.Padding(0.0f, 0.0f, 0.0f, 2.0f)
+				[
+					SNew(SBox)
+					.Visibility(this, &SNiagaraStackFunctionInputValue::GetValueWidgetVisibility, UNiagaraStackFunctionInput::EValueMode::Linked)
+					.VAlign(VAlign_Center)
+					[
+						SNew(STextBlock)
+						.TextStyle(FNiagaraEditorStyle::Get(), "NiagaraEditor.ParameterText")
+						.Text(this, &SNiagaraStackFunctionInputValue::GetLinkedValueHandleText)
+					]
+				]
+
+				// Data Object
+				+ SOverlay::Slot()
+				.Padding(0.0f, 0.0f, 0.0f, 2.0f)
+				[
+					SNew(SBox)
+					.Visibility(this, &SNiagaraStackFunctionInputValue::GetValueWidgetVisibility, UNiagaraStackFunctionInput::EValueMode::Data)
+					.VAlign(VAlign_Center)
+					[
+						SNew(STextBlock)
+						.TextStyle(FNiagaraEditorStyle::Get(), "NiagaraEditor.ParameterText")
+						.Text(this, &SNiagaraStackFunctionInputValue::GetDataValueText)
+					]
+				]
+
+				// Dynamic input name
+				+ SOverlay::Slot()
+				[
+					SNew(SBox)
+					.Visibility(this, &SNiagaraStackFunctionInputValue::GetValueWidgetVisibility, UNiagaraStackFunctionInput::EValueMode::Dynamic)
+					.VAlign(VAlign_Center)
+					[
+						SNew(STextBlock)
+						.TextStyle(FNiagaraEditorStyle::Get(), "NiagaraEditor.ParameterText")
+						.Text(this, &SNiagaraStackFunctionInputValue::GetDynamicValueText)
+						.OnDoubleClicked(this, &SNiagaraStackFunctionInputValue::DynamicInputTextDoubleClicked)
+					]
+				]
+
+				// Invalid input
+				+ SOverlay::Slot()
+				[
+					SNew(SBox)
+					.Visibility(this, &SNiagaraStackFunctionInputValue::GetValueWidgetVisibility, UNiagaraStackFunctionInput::EValueMode::Invalid)
+					.VAlign(VAlign_Center)
+					[
+						SNew(STextBlock)
+						.TextStyle(FNiagaraEditorStyle::Get(), "NiagaraEditor.ParameterText")
+						.Text(this, &SNiagaraStackFunctionInputValue::GetInvalidValueText)
+						.ToolTipText(this, &SNiagaraStackFunctionInputValue::GetInvalidValueToolTipText)
+					]
+				]
+			]
+
+			// Handle drop-down button
+			+ SHorizontalBox::Slot()
+			.AutoWidth()
+			.VAlign(VAlign_Center)
+			.Padding(3, 0, 0, 0)
+			[
+				SNew(SComboButton)
+				.ButtonStyle(FEditorStyle::Get(), "HoverHintOnly")
+				.ForegroundColor(FSlateColor::UseForeground())
+				.OnGetMenuContent(this, &SNiagaraStackFunctionInputValue::OnGetAvailableHandleMenu)
+				.ContentPadding(FMargin(2))
+				.HAlign(HAlign_Center)
+				.VAlign(VAlign_Center)
+			]
+
+			// Reset Button
+			+ SHorizontalBox::Slot()
+			.AutoWidth()
+			.VAlign(VAlign_Center)
+			.Padding(3, 0, 0, 0)
+			[
+				SNew(SButton)
+				.IsFocusable(false)
+				.ToolTipText(LOCTEXT("ResetToolTip", "Reset to the default value"))
+				.ButtonStyle(FEditorStyle::Get(), "NoBorder")
+				.ContentPadding(0)
+				.Visibility(this, &SNiagaraStackFunctionInputValue::GetResetButtonVisibility)
+				.OnClicked(this, &SNiagaraStackFunctionInputValue::ResetButtonPressed)
+				.Content()
+				[
+					SNew(SImage)
+					.Image(FEditorStyle::GetBrush("PropertyWindow.DiffersFromDefault"))
+				]
+			]
+		];
+	}
+
+private:
+	EVisibility GetValueWidgetVisibility(UNiagaraStackFunctionInput::EValueMode ValidMode) const
+	{
+		return FunctionInput->GetValueMode() == ValidMode ? EVisibility::Visible : EVisibility::Collapsed;
+	}
+
+	TSharedRef<SWidget> ConstructLocalValueStructWidget()
+	{
+		LocalValueStructParameterEditor.Reset();
+		LocalValueStructDetailsView.Reset();
+		if (DisplayedLocalValueStruct.IsValid())
+		{
+			FNiagaraEditorModule& NiagaraEditorModule = FModuleManager::GetModuleChecked<FNiagaraEditorModule>("NiagaraEditor");
+			TSharedPtr<INiagaraEditorTypeUtilities> TypeEditorUtilities = NiagaraEditorModule.GetTypeUtilities(FNiagaraTypeDefinition((UScriptStruct*)DisplayedLocalValueStruct->GetStruct()));
+			if (TypeEditorUtilities.IsValid() && TypeEditorUtilities->CanCreateParameterEditor())
+			{
+				TSharedPtr<SNiagaraParameterEditor> ParameterEditor = TypeEditorUtilities->CreateParameterEditor();
+				ParameterEditor->UpdateInternalValueFromStruct(DisplayedLocalValueStruct.ToSharedRef());
+				ParameterEditor->SetOnBeginValueChange(SNiagaraParameterEditor::FOnValueChange::CreateSP(
+					this, &SNiagaraStackFunctionInputValue::ParameterBeginValueChange));
+				ParameterEditor->SetOnEndValueChange(SNiagaraParameterEditor::FOnValueChange::CreateSP(
+					this, &SNiagaraStackFunctionInputValue::ParameterEndValueChange));
+				ParameterEditor->SetOnValueChanged(SNiagaraParameterEditor::FOnValueChange::CreateSP(
+					this, &SNiagaraStackFunctionInputValue::ParameterValueChanged, ParameterEditor.ToSharedRef()));
+				
+				LocalValueStructParameterEditor = ParameterEditor;
+				return ParameterEditor.ToSharedRef();
+			}
+			else
+			{
+				FPropertyEditorModule& PropertyEditorModule = FModuleManager::GetModuleChecked<FPropertyEditorModule>("PropertyEditor");
+
+				TSharedRef<IStructureDetailsView> StructureDetailsView = PropertyEditorModule.CreateStructureDetailView(
+					FDetailsViewArgs(false, false, false, FDetailsViewArgs::HideNameArea, true),
+					FStructureDetailsViewArgs(),
+					nullptr);
+
+				StructureDetailsView->SetStructureData(DisplayedLocalValueStruct);
+				StructureDetailsView->GetOnFinishedChangingPropertiesDelegate().AddSP(this, &SNiagaraStackFunctionInputValue::ParameterPropertyValueChanged);
+
+				LocalValueStructDetailsView = StructureDetailsView;
+				return StructureDetailsView->GetWidget().ToSharedRef();
+			}
+		}
+		else
+		{
+			return SNullWidget::NullWidget;
+		}
+	}
+
+	void OnInputValueChanged()
+	{
+		TSharedPtr<FStructOnScope> NewLocalValueStruct = FunctionInput->GetLocalValueStruct();
+		if (DisplayedLocalValueStruct == NewLocalValueStruct)
+		{
+			if (LocalValueStructParameterEditor.IsValid())
+			{
+				LocalValueStructParameterEditor->UpdateInternalValueFromStruct(DisplayedLocalValueStruct.ToSharedRef());
+			}
+			if (LocalValueStructDetailsView.IsValid())
+			{
+				LocalValueStructDetailsView->SetStructureData(TSharedPtr<FStructOnScope>());
+				LocalValueStructDetailsView->SetStructureData(DisplayedLocalValueStruct);
+			}
+		}
+		else
+		{
+			DisplayedLocalValueStruct = NewLocalValueStruct;
+			LocalValueStructContainer->SetContent(ConstructLocalValueStructWidget());
+		}
+	}
+
+	void ParameterBeginValueChange()
+	{
+		FunctionInput->NotifyBeginLocalValueChange();
+	}
+
+	void ParameterEndValueChange()
+	{
+		FunctionInput->NotifyEndLocalValueChange();
+	}
+
+	void ParameterValueChanged(TSharedRef<SNiagaraParameterEditor> ParameterEditor)
+	{
+		ParameterEditor->UpdateStructFromInternalValue(FunctionInput->GetLocalValueStruct().ToSharedRef());
+		FunctionInput->SetLocalValue(DisplayedLocalValueStruct.ToSharedRef());
+	}
+
+	void ParameterPropertyValueChanged(const FPropertyChangedEvent& PropertyChangedEvent)
+	{
+		FunctionInput->SetLocalValue(DisplayedLocalValueStruct.ToSharedRef());
+	}
+
+	FText GetLinkedValueHandleText() const
+	{
+		return FText::FromString(FunctionInput->GetLinkedValueHandle().GetParameterHandleString());
+	}
+
+	FText GetDataValueText() const
+	{
+		if (FunctionInput->GetDataValueObject() != nullptr)
+		{
+			return FunctionInput->GetInputType().GetClass()->GetDisplayNameText();
+		}
+		else
+		{
+			return FText::Format(LOCTEXT("InvalidDataObjectFormat", "{0} (Invalid)"), FunctionInput->GetInputType().GetClass()->GetDisplayNameText());
+		}
+	}
+
+	FText GetDynamicValueText() const
+	{
+		if (FunctionInput->GetDynamicInputNode() != nullptr)
+		{
+			return FText::FromString(FName::NameToDisplayString(FunctionInput->GetDynamicInputNode()->GetFunctionName(), false));
+		}
+		else
+		{
+			return LOCTEXT("InvalidDynamicDisplayName", "(Invalid)");
+		}
+	}
+
+	FText GetInvalidValueText() const
+	{
+		if (FunctionInput->CanReset())
+		{
+			return LOCTEXT("InvalidResetLabel", "Unsupported value - Reset to fix.");
+		}
+		else
+		{
+			return LOCTEXT("InvalidLabel", "Unsupported value");
+		}
+	}
+
+	FText GetInvalidValueToolTipText() const
+	{
+		if (FunctionInput->CanReset())
+		{
+			return LOCTEXT("InvalidResetToolTip", "This input has an unsupported value assigned in the stack.\nUse the reset button to remove the unsupported value.");
+		}
+		else
+		{
+			return LOCTEXT("InvalidToolTip", "The script that defines the source of this input has\n a default value that can not be displayed in the stack view.");
+		}
+	}
+
+	FReply DynamicInputTextDoubleClicked()
+	{
+		UNiagaraNodeFunctionCall* DynamicInputNode = FunctionInput->GetDynamicInputNode();
+		if (DynamicInputNode->FunctionScript != nullptr && DynamicInputNode->FunctionScript->IsAsset())
+		{
+			FAssetEditorManager::Get().OpenEditorForAsset(DynamicInputNode->FunctionScript);
+			return FReply::Handled();
+		}
+		return FReply::Unhandled();
+	}
+
+	TSharedRef<SWidget> OnGetAvailableHandleMenu()
+	{
+		FMenuBuilder MenuBuilder(true, nullptr);
+
+		// Set a local value
+		bool bCanSetLocalValue = 
+			FunctionInput->GetValueMode() != UNiagaraStackFunctionInput::EValueMode::Local &&
+			FunctionInput->GetInputType().IsDataInterface() == false;
+		MenuBuilder.AddMenuEntry(
+			LOCTEXT("LocalValue", "Set a local value"),
+			LOCTEXT("LocalValueToolTip", "Set a local editable value for this input."),
+			FSlateIcon(),
+			FUIAction(
+				FExecuteAction::CreateSP(this, &SNiagaraStackFunctionInputValue::SetToLocalValue),
+				FCanExecuteAction::CreateLambda([=]() { return bCanSetLocalValue; })));
+
+		// Add a dynamic input
+		MenuBuilder.BeginSection(NAME_None, LOCTEXT("DynamicInputSection", "Dynamic Inputs"));
+		TArray<UNiagaraScript*> DynamicInputScripts;
+		FunctionInput->GetAvailableDynamicInputs(DynamicInputScripts);
+		for (UNiagaraScript* DynamicInputScript : DynamicInputScripts)
+		{
+			FText DynamicInputText = FText::FromString(FName::NameToDisplayString(DynamicInputScript->GetName(), false));
+			MenuBuilder.AddMenuEntry(
+				DynamicInputText,
+				FText::Format(LOCTEXT("DynamicInputFormat", "Use {0} to provide a value for this input."), DynamicInputText),
+				FSlateIcon(),
+				FUIAction(FExecuteAction::CreateSP(this, &SNiagaraStackFunctionInputValue::DynamicInputScriptSelected, DynamicInputScript)));
+		}
+		MenuBuilder.EndSection();
+
+		// Link existing attribute
+		TArray<FNiagaraParameterHandle> AvailableHandles;
+		FunctionInput->GetAvailableParameterHandles(AvailableHandles);
+
+		TArray<FNiagaraParameterHandle> EngineHandles;
+		TArray<FNiagaraParameterHandle> SystemHandles;
+		TArray<FNiagaraParameterHandle> EmitterHandles;
+		TArray<FNiagaraParameterHandle> ParticleAttributeHandles;
+		TArray<FNiagaraParameterHandle> OtherHandles;
+		for (const FNiagaraParameterHandle AvailableHandle : AvailableHandles)
+		{
+			if (AvailableHandle.IsEngineHandle())
+			{
+				EngineHandles.Add(AvailableHandle);
+			}
+			else if (AvailableHandle.IsSystemHandle())
+			{
+				SystemHandles.Add(AvailableHandle);
+			}
+			else if (AvailableHandle.IsEmitterHandle())
+			{
+				EmitterHandles.Add(AvailableHandle);
+			}
+			else if (AvailableHandle.IsParticleAttributeHandle())
+			{
+				ParticleAttributeHandles.Add(AvailableHandle);
+			}
+			else
+			{
+				OtherHandles.Add(AvailableHandle);
+			}
+		}
+
+		FText MapInputFormat = LOCTEXT("LinkInputFormat", "Link this input to {0}");
+
+		auto AddMenuItemsForHandleList = [&](const TArray<FNiagaraParameterHandle>& Handles, FText SectionDisplayText)
+		{
+			MenuBuilder.BeginSection(NAME_None, SectionDisplayText);
+			for (const FNiagaraParameterHandle& Handle : Handles)
+			{
+				FText HandleDisplayName = FText::FromString(FName::NameToDisplayString(Handle.GetName(), false));
+				MenuBuilder.AddMenuEntry(
+					HandleDisplayName,
+					FText::Format(MapInputFormat, FText::FromString(Handle.GetParameterHandleString())),
+					FSlateIcon(),
+					FUIAction(FExecuteAction::CreateSP(this, &SNiagaraStackFunctionInputValue::ParameterHandleSelected, Handle)));
+			}
+			MenuBuilder.EndSection();
+		};
+
+		AddMenuItemsForHandleList(EngineHandles, LOCTEXT("EngineSection", "Engine"));
+		AddMenuItemsForHandleList(SystemHandles, LOCTEXT("SystemSection", "System"));
+		AddMenuItemsForHandleList(EmitterHandles, LOCTEXT("EmitterSection", "Emitter"));
+		AddMenuItemsForHandleList(ParticleAttributeHandles, LOCTEXT("ParticleAttributeSection", "Particle Attribute"));
+
+		MenuBuilder.BeginSection(NAME_None, LOCTEXT("OtherSection", "Other"));
+		for (const FNiagaraParameterHandle OtherHandle : OtherHandles)
+		{
+			FText HandleDisplayName = FText::FromString(FName::NameToDisplayString(OtherHandle.GetParameterHandleString(), false));
+			MenuBuilder.AddMenuEntry(
+				HandleDisplayName,
+				FText::Format(MapInputFormat, FText::FromString(OtherHandle.GetParameterHandleString())),
+				FSlateIcon(),
+				FUIAction(FExecuteAction::CreateSP(this, &SNiagaraStackFunctionInputValue::ParameterHandleSelected, OtherHandle)));
+		}
+		MenuBuilder.EndSection();
+
+		if (AvailableHandles.Num() > 0 || DynamicInputScripts.Num() > 0)
+		{
+			MenuBuilder.AddMenuSeparator();
+		}
+
+		// Read from new attribute
+		TArray<FString> AvailableNamespaces;
+		FunctionInput->GetNamespacesForNewParameters(AvailableNamespaces);
+
+		TArray<FString> InputNames;
+		for (int32 i = FunctionInput->GetInputParameterHandlePath().Num() - 1; i >= 0; i--)
+		{
+			InputNames.Add(FunctionInput->GetInputParameterHandlePath()[i].GetName());
+		}
+		FString InputName = FString::Join(InputNames, TEXT("."));
+
+		for (const FString& AvailableNamespace : AvailableNamespaces)
+		{
+			FNiagaraParameterHandle HandleToRead(AvailableNamespace, InputName);
+			bool bCanExecute = AvailableHandles.Contains(HandleToRead) == false;
+
+			MenuBuilder.AddMenuEntry(
+				FText::Format(LOCTEXT("ReadLabelFormat", "Read from new {0} parameter"), FText::FromString(AvailableNamespace)),
+				FText::Format(LOCTEXT("ReadToolTipFormat", "Read this input from a new parameter in the {0} namespace."), FText::FromString(AvailableNamespace)), 
+				FSlateIcon(),
+				FUIAction(
+					FExecuteAction::CreateSP(this, &SNiagaraStackFunctionInputValue::ParameterHandleSelected, HandleToRead),
+					FCanExecuteAction::CreateLambda([=]() { return bCanExecute; })));
+		}
+
+		return MenuBuilder.MakeWidget();
+	}
+
+	void SetToLocalValue()
+	{
+		const UScriptStruct* LocalValueStruct = FunctionInput->GetInputType().GetScriptStruct();
+		if (LocalValueStruct != nullptr)
+		{
+			TSharedRef<FStructOnScope> LocalValue = MakeShared<FStructOnScope>(LocalValueStruct);
+			TArray<uint8> DefaultValueData;
+			FNiagaraEditorUtilities::GetTypeDefaultValue(FunctionInput->GetInputType(), DefaultValueData);
+			if (DefaultValueData.Num() == LocalValueStruct->GetStructureSize())
+			{
+				FMemory::Memcpy(LocalValue->GetStructMemory(), DefaultValueData.GetData(), DefaultValueData.Num());
+				FunctionInput->SetLocalValue(LocalValue);
+			}
+		}
+	}
+
+	void DynamicInputScriptSelected(UNiagaraScript* DynamicInputScript)
+	{
+		FunctionInput->SetDynamicInput(DynamicInputScript);
+	}
+
+	void ParameterHandleSelected(FNiagaraParameterHandle Handle)
+	{
+		FunctionInput->SetLinkedValueHandle(Handle);
+	}
+
+	EVisibility GetResetButtonVisibility() const
+	{
+		return FunctionInput->CanReset() ? EVisibility::Visible : EVisibility::Hidden;
+	}
+
+	FReply ResetButtonPressed() const
+	{
+		FunctionInput->Reset();
+		return FReply::Handled();
+	}
+
+	FText GetInputIconText() const
+	{
+		switch (FunctionInput->GetValueMode())
+		{
+		case UNiagaraStackFunctionInput::EValueMode::Local:
+			return FText::FromString(FString(TEXT("\xf040") /* fa-pencil */));
+		case UNiagaraStackFunctionInput::EValueMode::Linked:
+			return FText::FromString(FString(TEXT("\xf0C1") /* fa-link */));
+		case UNiagaraStackFunctionInput::EValueMode::Data:
+			return FText::FromString(FString(TEXT("\xf1C0") /* fa-database */));
+		case UNiagaraStackFunctionInput::EValueMode::Dynamic:
+			return FText::FromString(FString(TEXT("\xf201") /* fa-line-chart */));
+		case UNiagaraStackFunctionInput::EValueMode::Invalid:
+		default:
+			return FText::FromString(FString(TEXT("\xf128") /* fa-question */));
+		}
+	}
+
+	FText GetInputIconToolTip() const
+	{
+		switch (FunctionInput->GetValueMode())
+		{
+		case UNiagaraStackFunctionInput::EValueMode::Local:
+			return LOCTEXT("StructInputIconToolTip", "Local Value");
+		case UNiagaraStackFunctionInput::EValueMode::Linked:
+			return LOCTEXT("LinkInputIconToolTip", "Linked Value");
+		case UNiagaraStackFunctionInput::EValueMode::Data:
+			return LOCTEXT("DataInterfaceInputIconToolTip", "Data Value");
+		case UNiagaraStackFunctionInput::EValueMode::Dynamic:
+			return LOCTEXT("DynamicInputIconToolTip", "Dynamic Value");
+		case UNiagaraStackFunctionInput::EValueMode::Invalid:
+		default:
+			return LOCTEXT("InvalidInputIconToolTip", "Unsupported value type.  Check the graph for issues.");
+		}
+	}
+
+	FSlateColor GetInputIconColor() const
+	{
+		switch (FunctionInput->GetValueMode())
+		{
+		case UNiagaraStackFunctionInput::EValueMode::Local:
+			return FLinearColor(FColor::Orange);
+		case UNiagaraStackFunctionInput::EValueMode::Linked:
+			return FLinearColor(FColor::Purple);
+		case UNiagaraStackFunctionInput::EValueMode::Data:
+			return FLinearColor(FColor::Yellow);
+		case UNiagaraStackFunctionInput::EValueMode::Dynamic:
+			return FLinearColor(FColor::Cyan);
+		case UNiagaraStackFunctionInput::EValueMode::Invalid:
+		default:
+			return FLinearColor(FColor::White);
+		}
+	}
+
+private:
+	UNiagaraStackFunctionInput* FunctionInput;
+
+	TSharedPtr<FStructOnScope> DisplayedLocalValueStruct;
+
+	TSharedPtr<SBox> LocalValueStructContainer;
+	TSharedPtr<SNiagaraParameterEditor> LocalValueStructParameterEditor;
+	TSharedPtr<IStructureDetailsView> LocalValueStructDetailsView;
+};
+
+class SNiagaraStackParameterStoreEntryValue : public SCompoundWidget
+{
+public:
+	DECLARE_DELEGATE_OneParam(FOnColumnWidthChanged, float)
+public:
+	SLATE_BEGIN_ARGS(SNiagaraStackParameterStoreEntryValue) { }
+	SLATE_END_ARGS();
+
+	void Construct(const FArguments& InArgs, UNiagaraStackParameterStoreEntry* InStackEntry)
+	{
+		StackEntry = InStackEntry;
+
+		StackEntry->OnValueChanged().AddSP(this, &SNiagaraStackParameterStoreEntryValue::OnInputValueChanged);
+		DisplayedValueStruct = StackEntry->GetValueStruct();
+
+		FMargin ItemPadding = FMargin(0);
+		ChildSlot
+			[
+				// Values
+				SNew(SHorizontalBox)
+				+ SHorizontalBox::Slot()
+				.VAlign(VAlign_Center)
+				.AutoWidth()
+				.Padding(0, 0, 3, 0)
+				[
+					// Value Icon
+					SNew(SBox)
+					.WidthOverride(TextIconSize)
+					[
+						SNew(STextBlock)
+						.Font(FEditorStyle::Get().GetFontStyle("FontAwesome.10"))
+						.Text(this, &SNiagaraStackParameterStoreEntryValue::GetInputIconText)
+						.ToolTipText(this, &SNiagaraStackParameterStoreEntryValue::GetInputIconToolTip)
+						.ColorAndOpacity(this, &SNiagaraStackParameterStoreEntryValue::GetInputIconColor)
+					]
+				]
+				+ SHorizontalBox::Slot()
+				.VAlign(VAlign_Center)
+				[
+					// Assigned handle
+					SNew(SVerticalBox)
+					// Value struct
+					+ SVerticalBox::Slot()
+					.AutoHeight()
+					[
+						SAssignNew(ValueStructContainer, SBox)
+						[
+							ConstructValueStructWidget()
+						]
+					]
+				]
+
+				// Handle drop-down button
+				+ SHorizontalBox::Slot()
+				.AutoWidth()
+				.VAlign(VAlign_Center)
+				.Padding(3, 0, 0, 0)
+				[
+					SNew(SComboButton)
+					.ButtonStyle(FEditorStyle::Get(), "HoverHintOnly")
+					.ForegroundColor(FSlateColor::UseForeground())
+					.OnGetMenuContent(this, &SNiagaraStackParameterStoreEntryValue::OnGetAvailableHandleMenu)
+					.ContentPadding(FMargin(2))
+					.HAlign(HAlign_Center)
+					.VAlign(VAlign_Center)
+				]
+
+				// Reset Button
+				+ SHorizontalBox::Slot()
+				.AutoWidth()
+				.VAlign(VAlign_Center)
+				.Padding(3, 0, 0, 0)
+				[
+					SNew(SButton)
+					.IsFocusable(false)
+					.ToolTipText(LOCTEXT("ResetToolTip", "Reset to the default value"))
+					.ButtonStyle(FEditorStyle::Get(), "NoBorder")
+					.ContentPadding(0)
+					.Visibility(this, &SNiagaraStackParameterStoreEntryValue::GetResetButtonVisibility)
+					.OnClicked(this, &SNiagaraStackParameterStoreEntryValue::ResetButtonPressed)
+					.Content()
+					[
+						SNew(SImage)
+						.Image(FEditorStyle::GetBrush("PropertyWindow.DiffersFromDefault"))
+					]
+				]
+			];
+	}
+
+private:
+	TSharedRef<SWidget> OnGetAvailableHandleMenu()
+	{
+		// TODO: This will need to be adjusted based on the current stack being edited, i.e. system vs emitter vs particle.
+		FMenuBuilder MenuBuilder(true, nullptr);
+		return MenuBuilder.MakeWidget();
+	}
+
+	TSharedRef<SWidget> ConstructValueStructWidget()
+	{
+		ValueStructParameterEditor.Reset();
+		ValueStructDetailsView.Reset();
+		if (DisplayedValueStruct.IsValid())
+		{
+			FNiagaraEditorModule& NiagaraEditorModule = FModuleManager::GetModuleChecked<FNiagaraEditorModule>("NiagaraEditor");
+			TSharedPtr<INiagaraEditorTypeUtilities> TypeEditorUtilities = NiagaraEditorModule.GetTypeUtilities(FNiagaraTypeDefinition((UScriptStruct*)DisplayedValueStruct->GetStruct()));
+			if (TypeEditorUtilities.IsValid() && TypeEditorUtilities->CanCreateParameterEditor())
+			{
+				TSharedPtr<SNiagaraParameterEditor> ParameterEditor = TypeEditorUtilities->CreateParameterEditor();
+				ParameterEditor->UpdateInternalValueFromStruct(DisplayedValueStruct.ToSharedRef());
+				ParameterEditor->SetOnBeginValueChange(SNiagaraParameterEditor::FOnValueChange::CreateSP(
+					this, &SNiagaraStackParameterStoreEntryValue::ParameterBeginValueChange));
+				ParameterEditor->SetOnEndValueChange(SNiagaraParameterEditor::FOnValueChange::CreateSP(
+					this, &SNiagaraStackParameterStoreEntryValue::ParameterEndValueChange));
+				ParameterEditor->SetOnValueChanged(SNiagaraParameterEditor::FOnValueChange::CreateSP(
+					this, &SNiagaraStackParameterStoreEntryValue::ParameterValueChanged, ParameterEditor.ToSharedRef()));
+
+				ValueStructParameterEditor = ParameterEditor;
+				return ParameterEditor.ToSharedRef();
+			}
+			else
+			{
+				FPropertyEditorModule& PropertyEditorModule = FModuleManager::GetModuleChecked<FPropertyEditorModule>("PropertyEditor");
+
+				TSharedRef<IStructureDetailsView> StructureDetailsView = PropertyEditorModule.CreateStructureDetailView(
+					FDetailsViewArgs(false, false, false, FDetailsViewArgs::HideNameArea, true),
+					FStructureDetailsViewArgs(),
+					nullptr);
+
+				StructureDetailsView->SetStructureData(DisplayedValueStruct);
+				StructureDetailsView->GetOnFinishedChangingPropertiesDelegate().AddSP(this, &SNiagaraStackParameterStoreEntryValue::ParameterPropertyValueChanged);
+
+				ValueStructDetailsView = StructureDetailsView;
+				return StructureDetailsView->GetWidget().ToSharedRef();
+			}
+		}
+		else
+		{
+			return SNullWidget::NullWidget;
+		}
+	}
+
+	void OnInputValueChanged()
+	{
+		TSharedPtr<FStructOnScope> NewValueStruct = StackEntry->GetValueStruct();
+		if (DisplayedValueStruct == NewValueStruct)
+		{
+			if (ValueStructParameterEditor.IsValid())
+			{
+				ValueStructParameterEditor->UpdateInternalValueFromStruct(DisplayedValueStruct.ToSharedRef());
+			}
+			if (ValueStructDetailsView.IsValid())
+			{
+				ValueStructDetailsView->SetStructureData(TSharedPtr<FStructOnScope>());
+				ValueStructDetailsView->SetStructureData(DisplayedValueStruct);
+			}
+		}
+		else
+		{
+			DisplayedValueStruct = NewValueStruct;
+			ValueStructContainer->SetContent(ConstructValueStructWidget());
+		}
+	}
+
+	void ParameterBeginValueChange()
+	{
+		StackEntry->NotifyBeginValueChange();
+	}
+
+	void ParameterEndValueChange()
+	{
+		StackEntry->NotifyEndValueChange();
+	}
+
+	void ParameterValueChanged(TSharedRef<SNiagaraParameterEditor> ParameterEditor)
+	{
+		ParameterEditor->UpdateStructFromInternalValue(StackEntry->GetValueStruct().ToSharedRef());
+		StackEntry->NotifyValueChanged();
+	}
+
+	void ParameterPropertyValueChanged(const FPropertyChangedEvent& PropertyChangedEvent)
+	{
+		StackEntry->NotifyValueChanged();
+	}
+
+	EVisibility GetReferenceVisibility() const
+	{
+		return EVisibility::Collapsed;
+	}
+
+
+	EVisibility GetResetButtonVisibility() const
+	{
+		return StackEntry->CanReset() ? EVisibility::Visible : EVisibility::Hidden;
+	}
+
+	FReply ResetButtonPressed() const
+	{
+		StackEntry->Reset();
+		return FReply::Handled();
+	}
+
+	FText GetInputIconText() const
+	{
+		if (DisplayedValueStruct.IsValid())
+		{
+			return FText::FromString(FString(TEXT("\xf040") /* fa-pencil */));
+		}
+		else if (StackEntry->GetValueObject() != nullptr)
+		{
+			return FText::FromString(FString(TEXT("\xf1C0") /* fa-database */));
+		}
+		else
+		{
+			return FText();
+		}
+	}
+
+	FText GetInputIconToolTip() const
+	{
+		if (DisplayedValueStruct.IsValid())
+		{
+			return LOCTEXT("StructInputIconToolTip", "Local Value");
+		}
+		else if (StackEntry->GetValueObject() != nullptr)
+		{
+			return LOCTEXT("DataInterfaceInputIconToolTip", "Data Value");
+		}
+		else
+		{
+			return FText();
+		}
+	}
+
+	FSlateColor GetInputIconColor() const
+	{
+		if (DisplayedValueStruct.IsValid())
+		{
+			return FLinearColor(FColor::Orange);
+		}
+		else if (StackEntry->GetValueObject() != nullptr)
+		{
+			return FLinearColor(FColor::Yellow);
+		}
+		else
+		{
+			return FLinearColor(FColor::White);
+		}
+	}
+
+private:
+	UNiagaraStackParameterStoreEntry* StackEntry;
+
+	TSharedPtr<FStructOnScope> DisplayedValueStruct;
+
+	TSharedPtr<SBox> ValueStructContainer;
+	TSharedPtr<SNiagaraParameterEditor> ValueStructParameterEditor;
+	TSharedPtr<IStructureDetailsView> ValueStructDetailsView;
+};
+
+class SNiagaraStackObject : public SCompoundWidget, public FNotifyHook
+{
+public:
+	DECLARE_DELEGATE_OneParam(FOnColumnWidthChanged, float)
+public:
+	SLATE_BEGIN_ARGS(SNiagaraStackObject) { }
+	SLATE_END_ARGS();
+
+	void Construct(const FArguments& InArgs, UNiagaraStackObject& InObject)
+	{
+		Object = &InObject;
+
+		FPropertyEditorModule& PropertyEditorModule = FModuleManager::GetModuleChecked<FPropertyEditorModule>("PropertyEditor");
+		TSharedRef<IDetailsView> DetailsView = PropertyEditorModule.CreateDetailView(FDetailsViewArgs(false, false, false, FDetailsViewArgs::HideNameArea, true, this));
+		DetailsView->SetObject(Object->GetObject(), true);
+
+		ChildSlot
+		[
+			DetailsView
+		];
+	}
+
+	virtual void NotifyPostChange(const FPropertyChangedEvent& PropertyChangedEvent, UProperty* PropertyThatChanged) override
+	{
+		Object->NotifyObjectChanged();
+	}
+
+private:
+	UNiagaraStackObject* Object;
+};
+
+class SNiagaraStackStruct : public SCompoundWidget, FNotifyHook
+{
+public:
+	DECLARE_DELEGATE_OneParam(FOnColumnWidthChanged, float)
+public:
+	SLATE_BEGIN_ARGS(SNiagaraStackStruct) { }
+	SLATE_END_ARGS();
+	
+
+	void Construct(const FArguments& InArgs, UNiagaraStackStruct& InObject)
+	{
+		Object = &InObject;
+		
+		FStructureDetailsViewArgs StructureDetailViewArgs;
+		FDetailsViewArgs DetailViewArgs(false, false, false, FDetailsViewArgs::HideNameArea, true);
+		DetailViewArgs.NotifyHook = this;
+		FPropertyEditorModule& PropertyEditorModule = FModuleManager::GetModuleChecked<FPropertyEditorModule>("PropertyEditor");
+
+		TSharedRef<IStructureDetailsView> StructureDetailsView = PropertyEditorModule.CreateStructureDetailView(
+			DetailViewArgs, StructureDetailViewArgs, nullptr);
+		if (InObject.HasDetailCustomization())
+		{
+			StructureDetailsView->GetDetailsView()->RegisterInstancedCustomPropertyLayout((UStruct*)InObject.GetStructOnScope()->GetStruct(), InObject.GetDetailsCustomization());
+		}
+
+		StructureDetailsView->SetStructureData(InObject.GetStructOnScope());
+	
+		ChildSlot
+			[
+				StructureDetailsView->GetWidget().ToSharedRef()
+			];
+	}
+
+	virtual void NotifyPreChange(class FEditPropertyChain* PropertyAboutToChange) override
+	{
+		Object->GetOwningObject()->Modify();
+	}
+
+	virtual void NotifyPostChange(const FPropertyChangedEvent& PropertyChangedEvent, class FEditPropertyChain* PropertyThatChanged) override
+	{
+	}
+
+private:
+	UNiagaraStackStruct* Object;
+};
+
+class SNiagaraStackAddModuleItem : public SCompoundWidget
+{
+public:
+	SLATE_BEGIN_ARGS(SNiagaraStackAddModuleItem) { }
+	SLATE_END_ARGS();
+
+	void Construct(const FArguments& InArgs, UNiagaraStackAddModuleItem& InAddModuleItem)
+	{
+		AddModuleItem = &InAddModuleItem;
+		ChildSlot
+		[
+			SNew(SComboButton)
+			.ButtonStyle(FEditorStyle::Get(), "HoverHintOnly")
+			.ForegroundColor(FSlateColor::UseForeground())
+			.HasDownArrow(false)
+			.OnGetMenuContent(this, &SNiagaraStackAddModuleItem::GetAddModuleMenu)
+			.ContentPadding(3)
+			.HAlign(HAlign_Center)
+			.VAlign(VAlign_Center)
+			.ButtonContent()
+			[
+				SNew(STextBlock)
+				.TextStyle(FEditorStyle::Get(), "NormalText.Important")
+				.Font(FEditorStyle::Get().GetFontStyle("FontAwesome.10"))
+				.Text(FText::FromString(FString(TEXT("\xf067"))) /*fa-plus*/)
+				.ToolTipText(LOCTEXT("AddModuleToolTip", "Add new module"))
+			]
+		];
+	}
+
+private:
+	TSharedRef<SWidget> GetAddModuleMenu()
+	{
+		FGraphActionMenuBuilder MenuBuilder;
+
+		return SNew (SBorder)
+			.BorderImage(FEditorStyle::GetBrush("Menu.Background"))
+			.Padding(5)
+			[
+				SNew(SBox)
+				[
+					SNew(SGraphActionMenu)
+					.OnActionSelected(this, &SNiagaraStackAddModuleItem::OnActionSelected)
+					.OnCreateWidgetForAction(SGraphActionMenu::FOnCreateWidgetForAction::CreateSP(this, &SNiagaraStackAddModuleItem::OnCreateWidgetForAction))
+					.OnCollectAllActions(this, &SNiagaraStackAddModuleItem::CollectAllActions)
+					.AutoExpandActionMenu(false)
+					.ShowFilterTextBox(true)
+				]
+			];
+	}
+
+
+	void CollectAllActions(FGraphActionListBuilderBase& OutAllActions)
+	{
+		if (OutAllActions.OwnerOfTemporaries == nullptr)
+		{
+			OutAllActions.OwnerOfTemporaries = NewObject<UEdGraph>((UObject*)GetTransientPackage());
+		}
+		FAssetRegistryModule& AssetRegistryModule = FModuleManager::LoadModuleChecked<FAssetRegistryModule>(TEXT("AssetRegistry"));
+
+		TArray<FAssetData> ScriptAssets;
+		AssetRegistryModule.Get().GetAssetsByClass(UNiagaraScript::StaticClass()->GetFName(), ScriptAssets);
+		UEnum* NiagaraScriptUsageEnum = FindObjectChecked<UEnum>(ANY_PACKAGE, TEXT("ENiagaraScriptUsage"), true);
+		for (const FAssetData& ScriptAsset : ScriptAssets)
+		{
+			FName UsageName;
+			ScriptAsset.GetTagValue(GET_MEMBER_NAME_CHECKED(UNiagaraScript, Usage), UsageName);
+
+			FText AssetDesc;
+			ScriptAsset.GetTagValue(GET_MEMBER_NAME_CHECKED(UNiagaraScript, Description), AssetDesc);
+
+			FText ModuleCategory;
+			ScriptAsset.GetTagValue(GET_MEMBER_NAME_CHECKED(UNiagaraScript, Category), ModuleCategory);
+
+			FString QualifiedUsageName = "ENiagaraScriptUsage::" + UsageName.ToString();
+			int32 UsageIndex = NiagaraScriptUsageEnum->GetIndexByNameString(QualifiedUsageName);
+
+			const FString BitfieldTagValue = ScriptAsset.GetTagValueRef<FString>(GET_MEMBER_NAME_CHECKED(UNiagaraScript, ModuleUsageBitmask));
+			int32 BitfieldValue = FCString::Atoi(*BitfieldTagValue);
+
+			int32 TargetBit = (BitfieldValue >> (int32)AddModuleItem->GetOutputUsage()) & 1;
+
+			if (UsageIndex != INDEX_NONE && TargetBit == 1)
+			{
+				ENiagaraScriptUsage Usage = static_cast<ENiagaraScriptUsage>(NiagaraScriptUsageEnum->GetValueByIndex(UsageIndex));
+				if (Usage == ENiagaraScriptUsage::Module)
+				{
+					FString DisplayNameString = FName::NameToDisplayString(ScriptAsset.AssetName.ToString(), false);
+					const FText NameText = FText::FromString(DisplayNameString);
+					const FText TooltipDesc = FText::Format(LOCTEXT("FunctionPopupTooltip", "Path: {0}\nDescription: {1}"), FText::FromString(ScriptAsset.ObjectPath.ToString()), AssetDesc);
+					FText CategoryName = LOCTEXT("ModuleNotCategorized", "Uncategorized Modules");
+					if (!ModuleCategory.IsEmptyOrWhitespace())
+						CategoryName = ModuleCategory;
+
+					TSharedPtr<FNiagaraStackGraphSchemaAction> NewNodeAction(new FNiagaraStackGraphSchemaAction(CategoryName, NameText, TooltipDesc, 0, FText(), 
+						FNiagaraStackGraphSchemaAction::FOnPerformStackAction::CreateUObject(AddModuleItem, &UNiagaraStackAddModuleItem::AddScriptModule, ScriptAsset)));
+					OutAllActions.AddAction(NewNodeAction);
+				}
+			}
+		}
+
+		// Generate actions for the available parameters to set.
+		TArray<FNiagaraVariable> AvailableParameters;
+		AddModuleItem->GetAvailableParameters(AvailableParameters);
+		for(const FNiagaraVariable& AvailableParameter : AvailableParameters)
+		{
+			FString DisplayNameString = FName::NameToDisplayString(AvailableParameter.GetName().ToString(), false);
+			const FText NameText = FText::FromString(DisplayNameString);
+			FText VarDesc = FNiagaraConstants::GetAttributeDescription(AvailableParameter);
+			FString VarDefaultValue = FNiagaraConstants::GetAttributeDefaultValue(AvailableParameter);
+			const FText TooltipDesc = FText::Format(LOCTEXT("SetFunctionPopupTooltip", "Description: Set the parameter {0}. {1}"), FText::FromName(AvailableParameter.GetName()), VarDesc);
+			FText CategoryName = LOCTEXT("ModuleSetCategory", "Set Specific Parameters");
+
+			TSharedPtr<FNiagaraStackGraphSchemaAction> NewNodeAction(new FNiagaraStackGraphSchemaAction(CategoryName, NameText, TooltipDesc, 0, FText(),
+				FNiagaraStackGraphSchemaAction::FOnPerformStackAction::CreateUObject(AddModuleItem, &UNiagaraStackAddModuleItem::AddParameterModule, AvailableParameter, false)));
+			OutAllActions.AddAction(NewNodeAction);
+		}
+
+		// Generate actions for setting new typed parameters.
+		TOptional<FString> NewParameterNamespace = AddModuleItem->GetNewParameterNamespace();
+		if (NewParameterNamespace.IsSet())
+		{
+			TArray<FNiagaraTypeDefinition> AvailableTypes;
+			AddModuleItem->GetNewParameterAvailableTypes(AvailableTypes);
+			for (const FNiagaraTypeDefinition& AvailableType : AvailableTypes)
+			{
+				FText NameText = AvailableType.GetNameText();
+				FText Tooltip = FText::Format(LOCTEXT("AddNewParameterTooltipFormat", "Description: Create a new {0} parameter."), NameText);
+				FText CategoryName = LOCTEXT("CreateNewParameterCategory", "Create New Parameter");
+
+				FNiagaraParameterHandle NewParameterHandle(NewParameterNamespace.GetValue(), *(TEXT("New") + AvailableType.GetName()));
+				FNiagaraVariable NewParameter(AvailableType, *NewParameterHandle.GetParameterHandleString());
+				TSharedPtr<FNiagaraStackGraphSchemaAction> NewNodeAction(new FNiagaraStackGraphSchemaAction(CategoryName, NameText, Tooltip, 0, FText(),
+					FNiagaraStackGraphSchemaAction::FOnPerformStackAction::CreateUObject(AddModuleItem, &UNiagaraStackAddModuleItem::AddParameterModule, NewParameter, true)));
+				OutAllActions.AddAction(NewNodeAction);
+			}
+		}
+	}
+
+	TSharedRef<SWidget> OnCreateWidgetForAction(struct FCreateWidgetForActionData* const InCreateData)
+	{
+		return SNew(SVerticalBox)
+			+ SVerticalBox::Slot()
+			.AutoHeight()
+			[
+				SNew(STextBlock)
+				.Text(InCreateData->Action->GetMenuDescription())
+				.ToolTipText(InCreateData->Action->GetTooltipDescription())
+			];
+	}
+
+
+	void OnActionSelected(const TArray< TSharedPtr<FEdGraphSchemaAction> >& SelectedActions, ESelectInfo::Type InSelectionType)
+	{
+		if (InSelectionType == ESelectInfo::OnMouseClick || InSelectionType == ESelectInfo::OnKeyPress || SelectedActions.Num() == 0)
+		{
+			for (int32 ActionIndex = 0; ActionIndex < SelectedActions.Num(); ActionIndex++)
+			{
+				TSharedPtr<FEdGraphSchemaAction> CurrentAction = SelectedActions[ActionIndex];
+
+				if (CurrentAction.IsValid())
+				{
+					FSlateApplication::Get().DismissAllMenus();
+
+					TArray<UEdGraphPin*> Pins;
+					CurrentAction->PerformAction(nullptr, Pins, FVector2D(0.0f, 0.0f));
+				}
+			}
+		}
+	}
+
+private:
+	UNiagaraStackAddModuleItem* AddModuleItem;
+};
+
+class SNiagaraStackAddRendererItem : public SCompoundWidget
+{
+public:
+	SLATE_BEGIN_ARGS(SNiagaraStackAddRendererItem) { }
+	SLATE_END_ARGS();
+
+	void Construct(const FArguments& InArgs, UNiagaraStackAddRendererItem& InAddRendererItem)
+	{
+		AddRendererItem = &InAddRendererItem;
+		ChildSlot
+		[
+			SNew(SComboButton)
+			.ButtonStyle(FEditorStyle::Get(), "HoverHintOnly")
+			.ForegroundColor(FSlateColor::UseForeground())
+			.HasDownArrow(false)
+			.OnGetMenuContent(this, &SNiagaraStackAddRendererItem::GetAddRendererMenu)
+			.ContentPadding(3)
+			.HAlign(HAlign_Center)
+			.VAlign(VAlign_Center)
+			.ButtonContent()
+			[
+				SNew(STextBlock)
+				.TextStyle(FEditorStyle::Get(), "NormalText.Important")
+				.Font(FEditorStyle::Get().GetFontStyle("FontAwesome.10"))
+				.Text(FText::FromString(FString(TEXT("\xf067"))) /*fa-plus*/)
+				.ToolTipText(LOCTEXT("AddRendererToolTip", "Add new renderer"))
+			]
+		];
+	}
+
+private:
+	class FRendererClassFilter : public IClassViewerFilter
+	{
+	public:
+		virtual bool IsClassAllowed(const FClassViewerInitializationOptions& InInitOptions, const UClass* InClass, TSharedRef< class FClassViewerFilterFuncs > InFilterFuncs)
+		{
+			return InClass->HasAnyClassFlags(CLASS_Abstract) == false && InClass->IsChildOf(UNiagaraRendererProperties::StaticClass());
+		}
+
+		virtual bool IsUnloadedClassAllowed(const FClassViewerInitializationOptions& InInitOptions, const TSharedRef< const class IUnloadedBlueprintData > InUnloadedClassData, TSharedRef< class FClassViewerFilterFuncs > InFilterFuncs)
+		{
+			return false;
+		}
+	};
+
+	TSharedRef<SWidget> GetAddRendererMenu()
+	{
+		FClassViewerInitializationOptions Options;
+		Options.bShowDisplayNames = true;
+		Options.ClassFilter = MakeShareable(new FRendererClassFilter());
+
+		FOnClassPicked OnPicked(FOnClassPicked::CreateSP(this, &SNiagaraStackAddRendererItem::RendererClassPicked));
+		return FModuleManager::LoadModuleChecked<FClassViewerModule>("ClassViewer").CreateClassViewer(Options, OnPicked);
+	}
+
+	void RendererClassPicked(UClass* InPickedClass)
+	{
+		AddRendererItem->AddRenderer(InPickedClass);
+	}
+
+private:
+	UNiagaraStackAddRendererItem* AddRendererItem;
+};
+
+class SNiagaraStackModuleItem : public SCompoundWidget
+{
+public:
+	SLATE_BEGIN_ARGS(SNiagaraStackModuleItem) { }
+	SLATE_END_ARGS();
+
+	void Construct(const FArguments& InArgs, UNiagaraStackModuleItem& InModuleItem)
+	{
+		ModuleItem = &InModuleItem;
+
+		ChildSlot
+		[
+			SNew(SHorizontalBox)
+			+ SHorizontalBox::Slot()
+			.Padding(1)
+			[
+				SNew(STextBlock)
+				.TextStyle(FNiagaraEditorWidgetsStyle::Get(), "NiagaraEditor.Stack.ItemText")
+				.ToolTipText_UObject(ModuleItem, &UNiagaraStackEntry::GetTooltipText)
+				.Text_UObject(ModuleItem, &UNiagaraStackEntry::GetDisplayName)
+			]
+			+ SHorizontalBox::Slot()
+			.AutoWidth()
+			[
+				SNew(SButton)
+				.ButtonStyle(FEditorStyle::Get(), "HoverHintOnly")
+				.IsFocusable(false)
+				.ForegroundColor(FNiagaraEditorWidgetsStyle::Get().GetColor(ModuleItem->GetItemForegroundName()))
+				.ToolTipText(LOCTEXT("UpToolTip", "Move this module up the stack"))
+				.OnClicked(this, &SNiagaraStackModuleItem::MoveUpClicked)
+				.Content()
+				[
+					SNew(STextBlock)
+					.Font(FEditorStyle::Get().GetFontStyle("FontAwesome.10"))
+					.Text(FText::FromString(FString(TEXT("\xf062"))))
+				]
+			]
+			+ SHorizontalBox::Slot()
+			.AutoWidth()
+			[
+				SNew(SButton)
+				.ButtonStyle(FEditorStyle::Get(), "HoverHintOnly")
+				.IsFocusable(false)
+				.ForegroundColor(FNiagaraEditorWidgetsStyle::Get().GetColor(ModuleItem->GetItemForegroundName()))
+				.ToolTipText(LOCTEXT("DownToolTip", "Move this module down the stack"))
+				.OnClicked(this, &SNiagaraStackModuleItem::MoveDownClicked)
+				.Content()
+				[
+					SNew(STextBlock)
+					.Font(FEditorStyle::Get().GetFontStyle("FontAwesome.10"))
+					.Text(FText::FromString(FString(TEXT("\xf063"))))
+				]
+			]
+			+ SHorizontalBox::Slot()
+			.AutoWidth()
+			[
+				SNew(SButton)
+				.ButtonStyle(FEditorStyle::Get(), "HoverHintOnly")
+				.IsFocusable(false)
+				.ForegroundColor(FNiagaraEditorWidgetsStyle::Get().GetColor(ModuleItem->GetItemForegroundName()))
+				.ToolTipText(LOCTEXT("DeleteToolTip", "Delete this module"))
+				.OnClicked(this, &SNiagaraStackModuleItem::DeleteClicked)
+				.Content()
+				[
+					SNew(STextBlock)
+					.Font(FEditorStyle::Get().GetFontStyle("FontAwesome.10"))
+					.Text(FText::FromString(FString(TEXT("\xf1f8"))))
+				]
+			]
+		];
+	}
+
+public:
+	//~ SWidget interface
+	virtual FReply OnMouseButtonDoubleClick(const FGeometry& InMyGeometry, const FPointerEvent& InMouseEvent) override
+	{
+		const UNiagaraNodeFunctionCall& ModuleFunctionCall = ModuleItem->GetModuleNode();
+		if (ModuleFunctionCall.FunctionScript != nullptr && ModuleFunctionCall.FunctionScript->IsAsset())
+		{
+			FAssetEditorManager::Get().OpenEditorForAsset(const_cast<UNiagaraScript*>(ModuleFunctionCall.FunctionScript));
+			return FReply::Handled();
+		}
+		return FReply::Unhandled();
+	}
+
+private:
+	FReply MoveUpClicked()
+	{
+		ModuleItem->MoveUp();
+		return FReply::Handled();
+	}
+
+	FReply MoveDownClicked()
+	{
+		ModuleItem->MoveDown();
+		return FReply::Handled();
+	}
+
+	FReply DeleteClicked()
+	{
+		ModuleItem->Delete();
+		return FReply::Handled();
+	}
+
+private:
+	UNiagaraStackModuleItem* ModuleItem;
+};
+
+class SNiagaraStackItemGroup : public SCompoundWidget
+{
+public:
+	SLATE_BEGIN_ARGS(SNiagaraStackItemGroup) { }
+	SLATE_END_ARGS();
+
+	void Construct(const FArguments& InArgs, UNiagaraStackItemGroup& InItem)
+	{
+		Item = &InItem;
+
+		ChildSlot
+			[
+				SNew(SHorizontalBox)
+				+ SHorizontalBox::Slot()
+				.Padding(1)
+				[
+					SNew(STextBlock)
+					.TextStyle(FNiagaraEditorWidgetsStyle::Get(), "NiagaraEditor.Stack.GroupText")
+					.ToolTipText_UObject(Item, &UNiagaraStackEntry::GetTooltipText)
+					.Text_UObject(Item, &UNiagaraStackEntry::GetDisplayName)
+				]
+				+ SHorizontalBox::Slot()
+				.AutoWidth()
+				[
+					SNew(SButton)
+					.Visibility(this, &SNiagaraStackItemGroup::AddVisibility)
+					.ButtonStyle(FEditorStyle::Get(), "HoverHintOnly")
+					.IsFocusable(false)
+					.ForegroundColor(FNiagaraEditorWidgetsStyle::Get().GetColor(Item->GetItemForegroundName()))
+					.ToolTipText(LOCTEXT("AddGroupToolTip", "Add a new group"))
+					.OnClicked(this, &SNiagaraStackItemGroup::AddClicked)
+					.Content()
+					[
+						SNew(STextBlock)
+						.Font(FEditorStyle::Get().GetFontStyle("FontAwesome.10"))
+						.Text(FText::FromString(FString(TEXT("\xf067")))) /*fa-plus*/
+					]
+				]
+				+ SHorizontalBox::Slot()
+				.AutoWidth()
+				[
+					SNew(SButton)
+						.Visibility(this, &SNiagaraStackItemGroup::DeleteVisibility)
+						.ButtonStyle(FEditorStyle::Get(), "HoverHintOnly")
+						.IsFocusable(false)
+						.ForegroundColor(FNiagaraEditorWidgetsStyle::Get().GetColor(Item->GetItemForegroundName()))
+						.ToolTipText(LOCTEXT("DeleteGroupToolTip", "Delete this group"))
+						.OnClicked(this, &SNiagaraStackItemGroup::DeleteClicked)
+						.Content()
+						[
+							SNew(STextBlock)
+							.Font(FEditorStyle::Get().GetFontStyle("FontAwesome.10"))
+							.Text(FText::FromString(FString(TEXT("\xf1f8"))))
+						]
+				]
+			];
+	}
+
+private:
+	FReply DeleteClicked()
+	{
+		Item->Delete();
+		return FReply::Handled();
+	}
+
+	EVisibility DeleteVisibility() const
+	{
+		if (Item->CanDelete())
+		{
+			return EVisibility::Visible;
+		}
+		else
+		{
+			return EVisibility::Collapsed;
+		}
+	}
+
+	FReply AddClicked()
+	{
+		Item->Add();
+		return FReply::Handled();
+	}
+
+	EVisibility AddVisibility() const
+	{
+		if (Item->CanAdd())
+		{
+			return EVisibility::Visible;
+		}
+		else
+		{
+			return EVisibility::Collapsed;
+		}
+	}
+
+private:
+	UNiagaraStackItemGroup* Item;
+};
+
+class SNiagaraStackRendererItem : public SCompoundWidget
+{
+public:
+	SLATE_BEGIN_ARGS(SNiagaraStackRendererItem) { }
+	SLATE_END_ARGS();
+
+	void Construct(const FArguments& InArgs, UNiagaraStackRendererItem& InRendererItem)
+	{
+		RendererItem = &InRendererItem;
+
+		ChildSlot
+		[
+			SNew(SHorizontalBox)
+			+ SHorizontalBox::Slot()
+			.Padding(1)
+			[
+				SNew(STextBlock)
+				.TextStyle(FNiagaraEditorWidgetsStyle::Get(), "NiagaraEditor.Stack.ItemText")
+				.ToolTipText_UObject(RendererItem, &UNiagaraStackEntry::GetTooltipText)
+				.Text_UObject(RendererItem, &UNiagaraStackEntry::GetDisplayName)
+			]
+			+ SHorizontalBox::Slot()
+			.AutoWidth()
+			[
+				SNew(SButton)
+				.ButtonStyle(FEditorStyle::Get(), "HoverHintOnly")
+				.IsFocusable(false)
+				.ForegroundColor(FNiagaraEditorWidgetsStyle::Get().GetColor(RendererItem->GetItemForegroundName()))
+				.ToolTipText(LOCTEXT("DeleteToolTip", "Delete this Renderer"))
+				.OnClicked(this, &SNiagaraStackRendererItem::DeleteClicked)
+				.Content()
+				[
+					SNew(STextBlock)
+					.Font(FEditorStyle::Get().GetFontStyle("FontAwesome.10"))
+					.Text(FText::FromString(FString(TEXT("\xf1f8"))))
+				]
+			]
+		];
+	}
+
+private:
+	FReply DeleteClicked()
+	{
+		RendererItem->Delete();
+		return FReply::Handled();
+	}
+
+private:
+	UNiagaraStackRendererItem* RendererItem;
+};
+
+class SNiagaraStackItemExpander : public SCompoundWidget
+{
+public:
+	SLATE_BEGIN_ARGS(SNiagaraStackItemExpander) { }
+	SLATE_END_ARGS();
+
+	void Construct(const FArguments& InArgs, UNiagaraStackItemExpander& InItemExpander)
+	{
+		ItemExpander = &InItemExpander;
+		ExpandedToolTipText = LOCTEXT("ExpandedItemToolTip", "Collapse this item");
+		CollapsedToolTipText = LOCTEXT("CollapsedItemToolTip", "Expand this item");
+
+		ChildSlot
+		[
+			SNew(SButton)
+			.ButtonStyle(FEditorStyle::Get(), "NoBorder")
+			.HAlign(HAlign_Center)
+			.ContentPadding(2)
+			.ToolTipText(this, &SNiagaraStackItemExpander::GetToolTipText)
+			.OnClicked(this, &SNiagaraStackItemExpander::ExpandButtonClicked)
+			.Content()
+			[
+				SNew(SImage)
+				.Image(this, &SNiagaraStackItemExpander::GetButtonBrush)
+			]
+		];
+	}
+
+private:
+	const FSlateBrush* GetButtonBrush() const
+	{
+		if (IsHovered())
+		{
+			return ItemExpander->GetIsExpanded()
+				? FEditorStyle::GetBrush("DetailsView.PulldownArrow.Up.Hovered") 
+				: FEditorStyle::GetBrush("DetailsView.PulldownArrow.Down.Hovered");
+		}
+		else
+		{
+			return ItemExpander->GetIsExpanded() 
+				? FEditorStyle::GetBrush("DetailsView.PulldownArrow.Up") 
+				: FEditorStyle::GetBrush("DetailsView.PulldownArrow.Down");
+		}
+	}
+
+	FText GetToolTipText() const
+	{
+		return ItemExpander->GetIsExpanded() ? ExpandedToolTipText : CollapsedToolTipText;
+	}
+
+	FReply ExpandButtonClicked()
+	{
+		ItemExpander->ToggleExpanded();
+		return FReply::Handled();
+	}
+
+private:
+	UNiagaraStackItemExpander* ItemExpander;
+	FText ExpandedToolTipText;
+	FText CollapsedToolTipText;
+};
+
+class SNiagaraStackTableRow : public STableRow<UNiagaraStackEntry*>
+{
+public:
+	DECLARE_DELEGATE_OneParam(FOnColumnWidthChanged, float)
+
+public:
+	SLATE_BEGIN_ARGS(SNiagaraStackTableRow)
+		: _GroupPadding(FMargin(5, 0, 5, 0))
+	{}
+		SLATE_ARGUMENT(FMargin, GroupPadding)
+		SLATE_ATTRIBUTE(float, NameColumnWidth)
+		SLATE_ATTRIBUTE(float, ValueColumnWidth)
+		SLATE_EVENT(FOnColumnWidthChanged, OnNameColumnWidthChanged)
+		SLATE_EVENT(FOnColumnWidthChanged, OnValueColumnWidthChanged)
+	SLATE_END_ARGS();
+
+	void Construct(const FArguments& InArgs, UNiagaraStackEntry* InStackEntry, const TSharedRef<STreeView<UNiagaraStackEntry*>>& InOwnerTree)
+	{
+		GroupPadding = InArgs._GroupPadding;
+		NameColumnWidth = InArgs._NameColumnWidth;
+		ValueColumnWidth = InArgs._ValueColumnWidth;
+		NameColumnWidthChanged = InArgs._OnNameColumnWidthChanged;
+		ValueColumnWidthChanged = InArgs._OnValueColumnWidthChanged;
+
+		StackEntry = InStackEntry;
+		OwnerTree = InOwnerTree;
+
+		ExpandedImage = FCoreStyle::Get().GetBrush("TreeArrow_Expanded");
+		CollapsedImage = FCoreStyle::Get().GetBrush("TreeArrow_Collapsed");
+
+		InactiveItemBackgroundColor = FNiagaraEditorWidgetsStyle::Get().GetColor(StackEntry->GetItemBackgroundName());
+		ActiveItemBackgroundColor = InactiveItemBackgroundColor + FLinearColor(.05f, .05f, .05f, 0.0f);
+
+		ConstructInternal(STableRow<UNiagaraStackEntry*>::FArguments(), OwnerTree.ToSharedRef());
+	}
+
+	virtual void SetNameAndValueContent(TSharedRef<SWidget> InNameWidget, TSharedPtr<SWidget> InValueWidget)
+	{
+		TSharedRef<SHorizontalBox> NameContent = SNew(SHorizontalBox)
+			+SHorizontalBox::Slot()
+			.AutoWidth()
+			[
+				SNew(SBox)
+				.WidthOverride(this, &SNiagaraStackTableRow::GetIndentSize)
+			]
+			+SHorizontalBox::Slot()
+			.AutoWidth()
+			.VAlign(VAlign_Center)
+			.Padding(0)
+			[
+				SNew(SButton)
+				.ButtonStyle(FCoreStyle::Get(), "NoBorder")
+				.Visibility(this, &SNiagaraStackTableRow::GetExpanderVisibility)
+				.OnClicked(this, &SNiagaraStackTableRow::ExpandButtonClicked)
+				.ForegroundColor(FSlateColor::UseForeground())
+				.ContentPadding(2)
+				[
+					SNew(SImage)
+					.Image(this, &SNiagaraStackTableRow::GetExpandButtonImage)
+					.ColorAndOpacity(FSlateColor::UseForeground())
+				]
+			]
+			+SHorizontalBox::Slot()
+			[
+				InNameWidget
+			];
+
+		TSharedPtr<SWidget> ChildContent;
+
+		if (InValueWidget.IsValid())
+		{
+			ChildContent = SNew(SSplitter)
+			.Style(FNiagaraEditorWidgetsStyle::Get(), "NiagaraEditor.Stack.Splitter")
+			.PhysicalSplitterHandleSize(1.0f)
+			.HitDetectionSplitterHandleSize(5.0f)
+			+ SSplitter::Slot()
+			.Value(NameColumnWidth)
+			.OnSlotResized(SSplitter::FOnSlotResized::CreateSP(this, &SNiagaraStackTableRow::OnNameColumnWidthChanged))
+			[
+				SNew(SBox)
+				.Padding(FMargin(3, 2, 5, 2))
+				[
+					NameContent
+				]
+			]
+
+			// Value
+			+ SSplitter::Slot()
+			.Value(ValueColumnWidth)
+			.OnSlotResized(SSplitter::FOnSlotResized::CreateSP(this, &SNiagaraStackTableRow::OnValueColumnWidthChanged))
+			[
+				SNew(SBox)
+				.Padding(FMargin(4, 2, 3, 2))
+				[
+					InValueWidget.ToSharedRef()
+				]
+			];
+		}
+		else
+		{
+			ChildContent = SNew(SBox)
+			.Padding(FMargin(3, 2, 3, 2))
+			[
+				NameContent
+			];
+		}
+
+		ChildSlot
+		[
+			SNew(SBorder)
+			.BorderImage(FEditorStyle::GetBrush("WhiteBrush"))
+			.BorderBackgroundColor(FNiagaraEditorWidgetsStyle::Get().GetColor(StackEntry->GetGroupBackgroundName()))
+			.Padding(this, &SNiagaraStackTableRow::GetGroupPadding)
+			.Visibility(this, &SNiagaraStackTableRow::GetRowVisibility)
+			[
+				SNew(SBorder)
+				.BorderImage(FEditorStyle::GetBrush("WhiteBrush"))
+				.BorderBackgroundColor(this, &SNiagaraStackTableRow::GetItemBackgroundColor)
+				.ForegroundColor(FNiagaraEditorWidgetsStyle::Get().GetColor(StackEntry->GetItemForegroundName()))
+				.Padding(FMargin(0))
+				[
+					ChildContent.ToSharedRef()
+				]
+			]
+		];
+	}
+
+	void SetGroupPadding(FMargin InGroupPadding)
+	{
+		GroupPadding = InGroupPadding;
+	}
+
+	bool GetIsRowActive() const
+	{
+		return IsHovered();
+	}
+
+	virtual FReply OnMouseButtonDoubleClick(const FGeometry& InMyGeometry, const FPointerEvent& InMouseEvent) override
+	{
+		return FReply::Unhandled();
+	}
+
+private:
+	FMargin GetGroupPadding() const
+	{
+		return GroupPadding;
+	}
+
+	EVisibility GetRowVisibility() const
+	{
+		return StackEntry->GetShouldShowInStack() 
+			? EVisibility::Visible 
+			: EVisibility::Collapsed;
+	}
+
+	FOptionalSize GetIndentSize() const
+	{
+		return StackEntry->GetItemIndentLevel() * IndentSize;
+	}
+
+	EVisibility GetExpanderVisibility() const
+	{
+		if (StackEntry->GetCanExpand())
+		{
+			TArray<UNiagaraStackEntry*> Children;
+			StackEntry->GetChildren(Children);
+			return Children.Num() > 0
+				? EVisibility::Visible
+				: EVisibility::Hidden;
+		}
+		else
+		{
+			return EVisibility::Collapsed;
+		}
+	}
+
+	FReply ExpandButtonClicked()
+	{
+		StackEntry->SetIsExpanded(!StackEntry->GetIsExpanded());
+		OwnerTree->SetItemExpansion(StackEntry, StackEntry->GetIsExpanded());
+		return FReply::Handled();
+	}
+
+	const FSlateBrush* GetExpandButtonImage() const
+	{
+		return StackEntry->GetIsExpanded() ? ExpandedImage : CollapsedImage;
+	}
+
+	void OnNameColumnWidthChanged(float Width)
+	{
+		NameColumnWidthChanged.ExecuteIfBound(Width);
+	}
+
+	void OnValueColumnWidthChanged(float Width)
+	{
+		ValueColumnWidthChanged.ExecuteIfBound(Width);
+	}
+
+	FSlateColor GetItemBackgroundColor() const
+	{
+		return GetIsRowActive() ? ActiveItemBackgroundColor : InactiveItemBackgroundColor;
+	}
+
+private:
+	UNiagaraStackEntry* StackEntry;
+	TSharedPtr<STreeView<UNiagaraStackEntry*>> OwnerTree;
+
+	TAttribute<float> NameColumnWidth;
+	TAttribute<float> ValueColumnWidth;
+	FOnColumnWidthChanged NameColumnWidthChanged;
+	FOnColumnWidthChanged ValueColumnWidthChanged;
+
+	const FSlateBrush* ExpandedImage;
+	const FSlateBrush* CollapsedImage;
+
+	FLinearColor InactiveItemBackgroundColor;
+	FLinearColor ActiveItemBackgroundColor;
+
+	FMargin GroupPadding;
+};
+
+void SNiagaraStack::Construct(const FArguments& InArgs, UNiagaraStackViewModel* InStackViewModel)
+{
+	StackViewModel = InStackViewModel;
+	StackViewModel->OnStructureChanged().AddSP(this, &SNiagaraStack::StackStructureChanged);
+	NameColumnWidth = .3f;
+	ContentColumnWidth = .7f;
+	ChildSlot
+	[
+		SNew(SBox)
+		.Padding(3)
+		[
+			SAssignNew(StackTree, STreeView<UNiagaraStackEntry*>)
+			.OnGenerateRow(this, &SNiagaraStack::OnGenerateRowForStackItem)
+			.OnGetChildren(this, &SNiagaraStack::OnGetChildren)
+			.TreeItemsSource(&StackViewModel->GetRootEntries())
+		]
+	];
+	PrimeTreeExpansion();
+}
+
+void SNiagaraStack::PrimeTreeExpansion()
+{
+	TArray<UNiagaraStackEntry*> EntriesToProcess(StackViewModel->GetRootEntries());
+	while (EntriesToProcess.Num() > 0)
+	{
+		UNiagaraStackEntry* EntryToProcess = EntriesToProcess[0];
+		EntriesToProcess.RemoveAtSwap(0);
+
+		if (EntryToProcess->GetIsExpanded())
+		{
+			StackTree->SetItemExpansion(EntryToProcess, true);
+			EntryToProcess->GetChildren(EntriesToProcess);
+		}
+	}
+}
+
+TSharedRef<ITableRow> SNiagaraStack::OnGenerateRowForStackItem(UNiagaraStackEntry* Item, const TSharedRef<STableViewBase>& OwnerTable)
+{
+	TSharedRef<SNiagaraStackTableRow> Container = ConstructContainerForItem(Item);
+	TSharedRef<SWidget> NameContent = ConstructNameWidgetForItem(Item, Container);
+	TSharedPtr<SWidget> ValueContent = ConstructValueWidgetForItem(Item, Container);
+	Container->SetNameAndValueContent(NameContent, ValueContent);
+	return Container;
+}
+
+TSharedRef<SNiagaraStackTableRow> SNiagaraStack::ConstructDefaultRow(UNiagaraStackEntry* Item)
+{
+	return SNew(SNiagaraStackTableRow, Item, StackTree.ToSharedRef())
+		.NameColumnWidth(this, &SNiagaraStack::GetNameColumnWidth)
+		.OnNameColumnWidthChanged(this, &SNiagaraStack::OnNameColumnWidthChanged)
+		.ValueColumnWidth(this, &SNiagaraStack::GetContentColumnWidth)
+		.OnValueColumnWidthChanged(this, &SNiagaraStack::OnContentColumnWidthChanged);
+}
+
+TSharedRef<SNiagaraStackTableRow> SNiagaraStack::ConstructContainerForItem(UNiagaraStackEntry* Item)
+{
+	if (Item->GetClass()->IsChildOf(UNiagaraStackItemGroup::StaticClass()))
+	{
+		TSharedRef<SNiagaraStackTableRow> GroupRow = ConstructDefaultRow(Item);
+		GroupRow->SetGroupPadding(FMargin(0));
+		return GroupRow;
+	}
+	else
+	{
+		return ConstructDefaultRow(Item);
+	}
+}
+
+TSharedRef<SWidget> SNiagaraStack::ConstructNameWidgetForItem(UNiagaraStackEntry* Item, TSharedRef<SNiagaraStackTableRow> Container)
+{
+	if (Item->GetClass()->IsChildOf(UNiagaraStackSpacer::StaticClass()))
+	{
+		return SNew(SBox)
+			.HeightOverride(6);
+	}
+	else if (Item->GetClass()->IsChildOf(UNiagaraStackAddModuleItem::StaticClass()))
+	{
+		return SNew(SNiagaraStackAddModuleItem, *CastChecked<UNiagaraStackAddModuleItem>(Item));
+	}
+	else if (Item->GetClass()->IsChildOf(UNiagaraStackAddRendererItem::StaticClass()))
+	{
+		return SNew(SNiagaraStackAddRendererItem, *CastChecked<UNiagaraStackAddRendererItem>(Item));
+	}
+	else if (Item->GetClass()->IsChildOf(UNiagaraStackItemGroup::StaticClass()))
+	{
+		return SNew(SNiagaraStackItemGroup, *CastChecked<UNiagaraStackItemGroup>(Item));
+	}
+	else if (Item->GetClass()->IsChildOf(UNiagaraStackModuleItem::StaticClass()))
+	{
+		return SNew(SNiagaraStackModuleItem, *CastChecked<UNiagaraStackModuleItem>(Item));
+	}
+	else if (Item->GetClass()->IsChildOf(UNiagaraStackRendererItem::StaticClass()))
+	{
+		return SNew(SNiagaraStackRendererItem, *CastChecked<UNiagaraStackRendererItem>(Item));
+	}
+	else if (Item->GetClass()->IsChildOf(UNiagaraStackFunctionInput::StaticClass()))
+	{
+		UNiagaraStackFunctionInput* FunctionInput = CastChecked<UNiagaraStackFunctionInput>(Item);
+		return SNew(SNiagaraStackFunctionInputName, FunctionInput)
+			.IsRowActive(Container, &SNiagaraStackTableRow::GetIsRowActive);
+	}
+	else if (Item->GetClass()->IsChildOf(UNiagaraStackObject::StaticClass()))
+	{
+		return SNew(SNiagaraStackObject, *CastChecked<UNiagaraStackObject>(Item));
+	}
+	else if (Item->GetClass()->IsChildOf(UNiagaraStackStruct::StaticClass()))
+	{
+		return SNew(SNiagaraStackStruct, *CastChecked<UNiagaraStackStruct>(Item));
+	}
+	else if (Item->GetClass()->IsChildOf(UNiagaraStackErrorItem::StaticClass()))
+	{
+		TSharedPtr<SHorizontalBox> ErrorInternalBox = SNew(SHorizontalBox);
+		UNiagaraStackErrorItem* ErrorItem = Cast<UNiagaraStackErrorItem>(Item);
+		ErrorInternalBox->AddSlot()
+			.HAlign(HAlign_Left)
+			.VAlign(VAlign_Center)
+			.AutoWidth()
+			[
+				SNew(STextBlock)
+				.TextStyle(FNiagaraEditorStyle::Get(), "NiagaraEditor.ParameterText")
+				.Text_UObject(ErrorItem, &UNiagaraStackErrorItem::ErrorText)
+				.ToolTipText_UObject(ErrorItem, &UNiagaraStackErrorItem::ErrorTextTooltip)
+			];
+		ErrorInternalBox->AddSlot()
+			.HAlign(HAlign_Right)
+			.VAlign(VAlign_Center)
+			.Padding(10,0)
+			[
+				SNew(SHorizontalBox)
+				.Visibility_UObject(ErrorItem, &UNiagaraStackErrorItem::CanFixVisibility)
+				+ SHorizontalBox::Slot()
+					.VAlign(VAlign_Center)
+					[
+						SNew(SButton)
+						.TextStyle(FNiagaraEditorStyle::Get(), "NiagaraEditor.ParameterText")
+						.Text(LOCTEXT("FixError", "Fix Error"))
+						.OnClicked_UObject(ErrorItem, &UNiagaraStackErrorItem::OnTryFixError)
+					]
+			];
+
+		return SNew(SHorizontalBox)
+			+ SHorizontalBox::Slot()
+				.AutoWidth()
+				.HAlign(HAlign_Left)
+				.VAlign(VAlign_Center)
+				[
+					SNew(SImage)
+					.Image(FEditorStyle::GetBrush("Icons.Error"))
+				]
+			+ SHorizontalBox::Slot()
+				.AutoWidth()
+				.HAlign(HAlign_Left)
+				.VAlign(VAlign_Center)
+				[
+					ErrorInternalBox.ToSharedRef()
+				];
+	}
+	else if (Item->GetClass()->IsChildOf(UNiagaraStackItemExpander::StaticClass()))
+	{
+		UNiagaraStackItemExpander* ItemExpander = CastChecked<UNiagaraStackItemExpander>(Item);
+		return SNew(SNiagaraStackItemExpander, *ItemExpander);
+	}
+	else
+	{
+		return SNew(STextBlock)
+			.TextStyle(FNiagaraEditorWidgetsStyle::Get(), Item->GetTextStyleName())
+			.ToolTipText_UObject(Item, &UNiagaraStackEntry::GetTooltipText)
+			.Text_UObject(Item, &UNiagaraStackEntry::GetDisplayName);
+	}
+}
+
+TSharedPtr<SWidget> SNiagaraStack::ConstructValueWidgetForItem(UNiagaraStackEntry* Item, TSharedRef<SNiagaraStackTableRow> Container)
+{
+	if (Item->GetClass()->IsChildOf(UNiagaraStackFunctionInput::StaticClass()))
+	{
+		UNiagaraStackFunctionInput* FunctionInput = CastChecked<UNiagaraStackFunctionInput>(Item);
+		return SNew(SNiagaraStackFunctionInputValue, FunctionInput);
+	}
+	else if (Item->GetClass()->IsChildOf(UNiagaraStackParameterStoreEntry::StaticClass()))
+	{
+		UNiagaraStackParameterStoreEntry* FunctionInput = CastChecked<UNiagaraStackParameterStoreEntry>(Item);
+		return SNew(SNiagaraStackParameterStoreEntryValue, FunctionInput);
+	}
+	else if (Item->GetClass()->IsChildOf(UNiagaraStackModuleItemOutput::StaticClass()))
+	{
+		UNiagaraStackModuleItemOutput* ModuleItemOutput = CastChecked<UNiagaraStackModuleItemOutput>(Item);
+		return SNew(STextBlock)
+			.TextStyle(FNiagaraEditorWidgetsStyle::Get(), ModuleItemOutput->GetTextStyleName())
+			.Text_UObject(ModuleItemOutput, &UNiagaraStackModuleItemOutput::GetOutputParameterHandleText);
+	}
+	else if (Item->GetClass()->IsChildOf(UNiagaraStackFunctionInputCollection::StaticClass()) ||
+		Item->GetClass()->IsChildOf(UNiagaraStackModuleItemOutputCollection::StaticClass()))
+	{
+		return SNullWidget::NullWidget;
+	}
+	else
+	{
+		return TSharedPtr<SWidget>();
+	}
+}
+
+void SNiagaraStack::OnGetChildren(UNiagaraStackEntry* Item, TArray<UNiagaraStackEntry*>& Children)
+{
+	Item->GetChildren(Children);
+}
+
+float SNiagaraStack::GetNameColumnWidth() const
+{
+	return NameColumnWidth;
+}
+
+float SNiagaraStack::GetContentColumnWidth() const
+{
+	return ContentColumnWidth;
+}
+
+void SNiagaraStack::OnNameColumnWidthChanged(float Width)
+{
+	NameColumnWidth = Width;
+}
+
+void SNiagaraStack::OnContentColumnWidthChanged(float Width)
+{
+	ContentColumnWidth = Width;
+}
+
+void SNiagaraStack::StackStructureChanged()
+{
+	PrimeTreeExpansion();
+	StackTree->RequestTreeRefresh();
+}
+
+#undef LOCTEXT_NAMESPACE
\ No newline at end of file
diff --git a/Engine/Plugins/FX/Niagara/Source/NiagaraEditorWidgets/Private/SNiagaraStack.h b/Engine/Plugins/FX/Niagara/Source/NiagaraEditorWidgets/Private/SNiagaraStack.h
new file mode 100644
--- /dev/null
+++ b/Engine/Plugins/FX/Niagara/Source/NiagaraEditorWidgets/Private/SNiagaraStack.h
@@ -0,0 +1,87 @@
+// Copyright 1998-2016 Epic Games, Inc. All Rights Reserved.
+
+#pragma once
+
+#include "EditorUndoClient.h"
+#include "SCompoundWidget.h"
+#include "DeclarativeSyntaxSupport.h"
+#include "STableRow.h"
+#include "STableViewBase.h"
+#include "STreeView.h"
+#include "ViewModels/Stack/NiagaraStackEntry.h"
+#include "EdGraph/EdGraphSchema.h"
+#include "SNiagaraStack.generated.h"
+
+class UNiagaraStackViewModel;
+class SNiagaraStackTableRow;
+
+class SNiagaraStack : public SCompoundWidget, public FEditorUndoClient
+{
+public:
+	SLATE_BEGIN_ARGS(SNiagaraStack)
+	{}
+
+	SLATE_END_ARGS();
+
+	void Construct(const FArguments& InArgs, UNiagaraStackViewModel* InStackViewModel);
+
+private:
+	void PrimeTreeExpansion();
+
+	TSharedRef<ITableRow> OnGenerateRowForStackItem(UNiagaraStackEntry* Item, const TSharedRef<STableViewBase>& OwnerTable);
+
+	TSharedRef<SNiagaraStackTableRow> ConstructDefaultRow(UNiagaraStackEntry* Item);
+
+	TSharedRef<SNiagaraStackTableRow> ConstructContainerForItem(UNiagaraStackEntry* Item);
+
+	TSharedRef<SWidget> ConstructNameWidgetForItem(UNiagaraStackEntry* Item, TSharedRef<SNiagaraStackTableRow> Container);
+
+	TSharedPtr<SWidget> ConstructValueWidgetForItem(UNiagaraStackEntry* Item, TSharedRef<SNiagaraStackTableRow> Container);
+	
+	void OnGetChildren(UNiagaraStackEntry* Item, TArray<UNiagaraStackEntry*>& Children);
+
+	float GetNameColumnWidth() const;
+	float GetContentColumnWidth() const;
+
+	void OnNameColumnWidthChanged(float Width);
+	void OnContentColumnWidthChanged(float Width);
+	void StackStructureChanged();
+
+	EVisibility GetVisibilityForItem(UNiagaraStackEntry* Item) const;
+
+private:
+	UNiagaraStackViewModel* StackViewModel;
+
+	TSharedPtr<STreeView<UNiagaraStackEntry*>> StackTree;
+
+	float NameColumnWidth;
+
+	float ContentColumnWidth;
+};
+
+USTRUCT()
+struct FNiagaraStackGraphSchemaAction : public FEdGraphSchemaAction
+{
+	GENERATED_USTRUCT_BODY();
+
+	DECLARE_DELEGATE(FOnPerformStackAction);
+
+	FNiagaraStackGraphSchemaAction()
+	{
+	}
+
+	FNiagaraStackGraphSchemaAction(FText InNodeCategory, FText InMenuDesc, FText InToolTip, const int32 InGrouping, FText InKeywords, FOnPerformStackAction InAction)
+		: FEdGraphSchemaAction(MoveTemp(InNodeCategory), MoveTemp(InMenuDesc), MoveTemp(InToolTip), InGrouping, MoveTemp(InKeywords))
+		, Action(InAction)
+	{}
+
+	//~ Begin FEdGraphSchemaAction Interface
+	UEdGraphNode* PerformAction(class UEdGraph* ParentGraph, UEdGraphPin* FromPin, const FVector2D Location, bool bSelectNewNode = true) override
+	{
+		Action.ExecuteIfBound();
+		return nullptr;
+	}
+	//~ End FEdGraphSchemaAction Interface
+
+	FOnPerformStackAction Action;
+};
\ No newline at end of file
