Copying //UE4/Orion-Staging to //UE4/Main (Source: //Orion/Dev-General @ 3606378)

diff --git a/Engine/Plugins/FX/Niagara/Source/NiagaraShader/Private/NiagaraShader.cpp b/Engine/Plugins/FX/Niagara/Source/NiagaraShader/Private/NiagaraShader.cpp
new file mode 100644
--- /dev/null
+++ b/Engine/Plugins/FX/Niagara/Source/NiagaraShader/Private/NiagaraShader.cpp
@@ -0,0 +1,1103 @@
+// Copyright 1998-2017 Epic Games, Inc. All Rights Reserved.
+
+#include "NiagaraShader.h"
+#include "NiagaraShared.h"
+#include "NiagaraShaderMap.h"
+#include "Stats/StatsMisc.h"
+#include "Serialization/MemoryWriter.h"
+#include "Serialization/MemoryReader.h"
+#include "ProfilingDebugging/DiagnosticTable.h"
+#include "ShaderCompiler.h"
+#include "NiagaraShaderDerivedDataVersion.h"
+#include "NiagaraShaderCompilationManager.h"
+#if WITH_EDITOR
+	#include "Interfaces/ITargetPlatformManagerModule.h"
+	#include "TickableEditorObject.h"
+	#include "DerivedDataCacheInterface.h"
+#endif
+#include "ProfilingDebugging/CookStats.h"
+
+IMPLEMENT_SHADER_TYPE(,FNiagaraShader, TEXT("/Engine/Private/NiagaraEmitterInstanceShader.usf"),TEXT("SimulateMain"), SF_Compute)
+
+
+int32 GCreateNiagaraShadersOnLoad = 0;
+static FAutoConsoleVariableRef CVarCreateNiagaraShadersOnLoad(
+	TEXT("niagara.CreateShadersOnLoad"),
+	GCreateNiagaraShadersOnLoad,
+	TEXT("Whether to create Niagara's simulation shaders on load, which can reduce hitching, but use more memory.  Otherwise they will be created as needed.")
+	);
+
+
+#if ENABLE_COOK_STATS
+namespace NiagaraShaderCookStats
+{
+	static FCookStats::FDDCResourceUsageStats UsageStats;
+	static int32 ShadersCompiled = 0;
+	static FCookStatsManager::FAutoRegisterCallback RegisterCookStats([](FCookStatsManager::AddStatFuncRef AddStat)
+	{
+		UsageStats.LogStats(AddStat, TEXT("NiagaraShader.Usage"), TEXT(""));
+		AddStat(TEXT("NiagaraShader.Misc"), FCookStatsManager::CreateKeyValueArray(
+			TEXT("ShadersCompiled"), ShadersCompiled
+			));
+	});
+}
+#endif
+
+
+//
+// Globals
+//
+TMap<FNiagaraShaderMapId, FNiagaraShaderMap*> FNiagaraShaderMap::GIdToNiagaraShaderMap[SP_NumPlatforms];
+TArray<FNiagaraShaderMap*> FNiagaraShaderMap::AllNiagaraShaderMaps;
+
+// The Id of 0 is reserved for global shaders
+uint32 FNiagaraShaderMap::NextCompilingId = 2;
+
+
+/** 
+ * Tracks FNiagaraScripts and their shader maps that are being compiled.
+ * Uses a TRefCountPtr as this will be the only reference to a shader map while it is being compiled.
+ */
+TMap<TRefCountPtr<FNiagaraShaderMap>, TArray<FNiagaraScript*> > FNiagaraShaderMap::NiagaraShaderMapsBeingCompiled;
+
+
+
+static inline bool ShouldCacheNiagaraShader(const FNiagaraShaderType* ShaderType, EShaderPlatform Platform, const FNiagaraScript* Script)
+{
+	return ShaderType->ShouldCache(Platform, Script) && Script->ShouldCache(Platform, ShaderType);
+}
+
+
+
+/** Called for every script shader to update the appropriate stats. */
+void UpdateNiagaraShaderCompilingStats(const FNiagaraScript* Script)
+{
+	INC_DWORD_STAT_BY(STAT_ShaderCompiling_NumTotalNiagaraShaders,1);
+}
+
+
+void FNiagaraShaderMapId::Serialize(FArchive& Ar)
+{
+	// You must bump NIAGARASHADERMAP_DERIVEDDATA_VER if changing the serialization of FNiagaraShaderMapId.
+
+	Ar << BaseScriptID;
+
+	Ar << (int32&)FeatureLevel;
+
+	//ParameterSet.Serialize(Ar);		// NIAGARATODO: at some point we'll need stuff for static switches here
+	//Ar << ReferencedFunctions;		// may want this too, although we've got other mechanisms to determine need to recompile, unsure whether the shader map needs to know
+}
+
+/** Hashes the script-specific part of this shader map Id. */
+void FNiagaraShaderMapId::GetScriptHash(FSHAHash& OutHash) const
+{
+	FSHA1 HashState;
+
+	HashState.Update((const uint8*)&BaseScriptID, sizeof(BaseScriptID));
+	HashState.Update((const uint8*)&FeatureLevel, sizeof(FeatureLevel));
+
+	//ParameterSet.UpdateHash(HashState);		// will need for static switches
+
+	/*
+	for (int32 FunctionIndex = 0; FunctionIndex < ReferencedFunctions.Num(); FunctionIndex++)
+	{
+		HashState.Update((const uint8*)&ReferencedFunctions[FunctionIndex], sizeof(ReferencedFunctions[FunctionIndex]));
+	}
+	*/
+	HashState.Final();
+	HashState.GetHash(&OutHash.Hash[0]);
+}
+
+/** 
+* Tests this set against another for equality, disregarding override settings.
+* 
+* @param ReferenceSet	The set to compare against
+* @return				true if the sets are equal
+*/
+bool FNiagaraShaderMapId::operator==(const FNiagaraShaderMapId& ReferenceSet) const
+{
+	if (BaseScriptID != ReferenceSet.BaseScriptID 
+		//|| QualityLevel != ReferenceSet.QualityLevel
+		|| FeatureLevel != ReferenceSet.FeatureLevel)
+	{
+		return false;
+	}
+	/*
+	if (ParameterSet != ReferenceSet.ParameterSet
+		|| ReferencedFunctions.Num() != ReferenceSet.ReferencedFunctions.Num()
+	{
+		return false;
+	}
+
+	for (int32 RefFunctionIndex = 0; RefFunctionIndex < ReferenceSet.ReferencedFunctions.Num(); RefFunctionIndex++)
+	{
+		const FGuid& ReferenceGuid = ReferenceSet.ReferencedFunctions[RefFunctionIndex];
+
+		if (ReferencedFunctions[RefFunctionIndex] != ReferenceGuid)
+		{
+			return false;
+		}
+	}
+	*/
+	return true;
+}
+
+void FNiagaraShaderMapId::AppendKeyString(FString& KeyString) const
+{
+	KeyString += BaseScriptID.ToString();
+	KeyString += TEXT("_");
+	FString FeatureLevelString;
+	GetFeatureLevelName(FeatureLevel, FeatureLevelString);
+	KeyString += FeatureLevelString + TEXT("_");
+	/*
+	ParameterSet.AppendKeyString(KeyString);
+
+	KeyString += TEXT("_");
+	KeyString += FString::FromInt(Usage);
+	KeyString += TEXT("_");
+
+	// Add any referenced functions to the key so that we will recompile when they are changed
+	for (int32 FunctionIndex = 0; FunctionIndex < ReferencedFunctions.Num(); FunctionIndex++)
+	{
+		KeyString += ReferencedFunctions[FunctionIndex].ToString();
+	}
+	*/
+}
+
+
+/**
+ * Enqueues a compilation for a new shader of this type.
+ * @param script - The script to link the shader with.
+ */
+FNiagaraShaderCompileJob* FNiagaraShaderType::BeginCompileShader(
+	uint32 ShaderMapId,
+	const FNiagaraScript* Script,
+	FShaderCompilerEnvironment* CompilationEnvironment,
+	EShaderPlatform Platform,
+	TArray<FNiagaraShaderCompileJob*>& NewJobs,
+	FShaderTarget Target,
+	TArray< TArray<DIGPUBufferParamDescriptor> >&InBufferDescriptors
+	)
+{
+	FNiagaraShaderCompileJob* NewJob = new FNiagaraShaderCompileJob(ShaderMapId, this, Script->HlslOutput);
+
+	NewJob->DIBufferDescriptors = InBufferDescriptors;	// from hlsl translation and need to be passed to the FNiagaraShader on completion
+
+	NewJob->Input.SharedEnvironment = CompilationEnvironment;
+	NewJob->Input.Target = Target;
+	NewJob->Input.ShaderFormat = LegacyShaderPlatformToShaderFormat(Platform);
+	NewJob->Input.VirtualSourceFilePath = TEXT("/Engine/Private/NiagaraEmitterInstanceShader.usf");
+	NewJob->Input.EntryPointName = TEXT("SimulateMainComputeCS");
+	NewJob->Input.Environment.SetDefine(TEXT("GPU_SIMULATION"), 1);
+	NewJob->Input.Environment.IncludeVirtualPathToContentsMap.Add(TEXT("/Engine/Generated/NiagaraEmitterInstance.usf"), StringToArray<ANSICHAR>(*Script->HlslOutput, Script->HlslOutput.Len() + 1));
+	FShaderCompilerEnvironment& ShaderEnvironment = NewJob->Input.Environment;
+
+	UE_LOG(LogShaders, Verbose, TEXT("			%s"), GetName());
+	COOK_STAT(NiagaraShaderCookStats::ShadersCompiled++);
+
+	//update script shader stats
+	UpdateNiagaraShaderCompilingStats(Script);
+
+	// Allow the shader type to modify the compile environment.
+	SetupCompileEnvironment(Platform, Script, ShaderEnvironment);
+
+	NewJobs.Add(NewJob);
+
+	return NewJob;
+}
+
+/**
+ * Either creates a new instance of this type or returns an equivalent existing shader.
+ * @param CurrentJob - Compile job that was enqueued by BeginCompileShader.
+ */
+FShader* FNiagaraShaderType::FinishCompileShader(
+	//const FUniformExpressionSet& UniformExpressionSet,
+	const FSHAHash& ShaderMapHash,
+	const FNiagaraShaderCompileJob& CurrentJob,
+	const FString& InDebugDescription
+	)
+{
+	check(CurrentJob.bSucceeded);
+
+	FShaderType* SpecificType = CurrentJob.ShaderType->LimitShaderResourceToThisType() ? CurrentJob.ShaderType : NULL;
+
+	// Reuse an existing resource with the same key or create a new one based on the compile output
+	// This allows FShaders to share compiled bytecode and RHI shader references
+	FShaderResource* Resource = FShaderResource::FindOrCreateShaderResource(CurrentJob.Output, SpecificType);
+
+	// Find a shader with the same key in memory
+	FShader* Shader = CurrentJob.ShaderType->FindShaderById(FShaderId(ShaderMapHash, nullptr, nullptr, CurrentJob.ShaderType, CurrentJob.Input.Target));
+
+	// There was no shader with the same key so create a new one with the compile output, which will bind shader parameters
+	if (!Shader)
+	{
+		Shader = (*ConstructCompiledRef)(CompiledShaderInitializerType(this, CurrentJob.Output, Resource, ShaderMapHash, InDebugDescription, CurrentJob.DIBufferDescriptors));
+		CurrentJob.Output.ParameterMap.VerifyBindingsAreComplete(GetName(), CurrentJob.Output.Target, nullptr);
+	}
+
+	return Shader;
+}
+
+/**
+* Finds the shader map for a script.
+* @param ShaderMapId - The script id and static parameter set identifying the shader map
+* @param Platform - The platform to lookup for
+* @return NULL if no cached shader map was found.
+*/
+FNiagaraShaderMap* FNiagaraShaderMap::FindId(const FNiagaraShaderMapId& ShaderMapId, EShaderPlatform InPlatform)
+{
+	check(ShaderMapId.BaseScriptID != FGuid());
+	return GIdToNiagaraShaderMap[InPlatform].FindRef(ShaderMapId);
+}
+
+/** Flushes the given shader types from any loaded FNiagaraShaderMap's. */
+void FNiagaraShaderMap::FlushShaderTypes(TArray<FShaderType*>& ShaderTypesToFlush)
+{
+	for (int32 ShaderMapIndex = 0; ShaderMapIndex < AllNiagaraShaderMaps.Num(); ShaderMapIndex++)
+	{
+		FNiagaraShaderMap* CurrentShaderMap = AllNiagaraShaderMaps[ShaderMapIndex];
+
+		for (int32 ShaderTypeIndex = 0; ShaderTypeIndex < ShaderTypesToFlush.Num(); ShaderTypeIndex++)
+		{
+			CurrentShaderMap->FlushShadersByShaderType(ShaderTypesToFlush[ShaderTypeIndex]);
+		}
+	}
+}
+
+void FNiagaraShaderMap::FixupShaderTypes(EShaderPlatform Platform, const TMap<FShaderType*, FString>& ShaderTypeNames)
+{
+	TArray<FNiagaraShaderMapId> Keys;
+	FNiagaraShaderMap::GIdToNiagaraShaderMap[Platform].GenerateKeyArray(Keys);
+
+	TArray<FNiagaraShaderMap*> Values;
+	FNiagaraShaderMap::GIdToNiagaraShaderMap[Platform].GenerateValueArray(Values);
+
+	//@todo - what about the shader maps in AllNiagaraShaderMaps that are not in GIdToNiagaraShaderMap?
+	FNiagaraShaderMap::GIdToNiagaraShaderMap[Platform].Empty();
+
+	for (int32 PairIndex = 0; PairIndex < Keys.Num(); PairIndex++)
+	{
+		FNiagaraShaderMapId& Key = Keys[PairIndex];
+		FNiagaraShaderMap::GIdToNiagaraShaderMap[Platform].Add(Key, Values[PairIndex]);
+	}
+}
+
+
+void NiagaraShaderMapAppendKeyString(EShaderPlatform Platform, FString& KeyString)
+{
+	// does nothing at the moment, but needs to append to keystring if runtime options impacting selection of sim shader permutations are added
+	// for example static switches will need to go here
+}
+
+#if WITH_EDITOR
+
+/** Creates a string key for the derived data cache given a shader map id. */
+static FString GetNiagaraShaderMapKeyString(const FNiagaraShaderMapId& ShaderMapId, EShaderPlatform Platform)
+{
+	FName Format = LegacyShaderPlatformToShaderFormat(Platform);
+	FString ShaderMapKeyString = Format.ToString() + TEXT("_") + FString(FString::FromInt(GetTargetPlatformManagerRef().ShaderFormatVersion(Format))) + TEXT("_");
+	NiagaraShaderMapAppendKeyString(Platform, ShaderMapKeyString);
+	ShaderMapId.AppendKeyString(ShaderMapKeyString);
+	return FDerivedDataCacheInterface::BuildCacheKey(TEXT("NIAGARASM"), NIAGARASHADERMAP_DERIVEDDATA_VER, *ShaderMapKeyString);
+}
+
+
+
+void FNiagaraShaderMap::LoadFromDerivedDataCache(const FNiagaraScript* Script, const FNiagaraShaderMapId& ShaderMapId, EShaderPlatform Platform, TRefCountPtr<FNiagaraShaderMap>& InOutShaderMap)
+{
+	if (InOutShaderMap != NULL)
+	{
+		check(InOutShaderMap->Platform == Platform);
+		// If the shader map was non-NULL then it was found in memory but is incomplete, attempt to load the missing entries from memory
+		InOutShaderMap->LoadMissingShadersFromMemory(Script);
+	}
+	else
+	{
+		// Shader map was not found in memory, try to load it from the DDC
+		STAT(double NiagaraShaderDDCTime = 0);
+		{
+			SCOPE_SECONDS_COUNTER(NiagaraShaderDDCTime);
+			COOK_STAT(auto Timer = NiagaraShaderCookStats::UsageStats.TimeSyncWork());
+
+			TArray<uint8> CachedData;
+			const FString DataKey = GetNiagaraShaderMapKeyString(ShaderMapId, Platform);
+
+			if (GetDerivedDataCacheRef().GetSynchronous(*DataKey, CachedData))
+			{
+				COOK_STAT(Timer.AddHit(CachedData.Num()));
+				InOutShaderMap = new FNiagaraShaderMap();
+				FMemoryReader Ar(CachedData, true);
+
+				// Deserialize from the cached data
+				InOutShaderMap->Serialize(Ar);
+				InOutShaderMap->RegisterSerializedShaders();
+
+				checkSlow(InOutShaderMap->GetShaderMapId() == ShaderMapId);
+
+				// Register in the global map
+				InOutShaderMap->Register(Platform);
+			}
+			else
+			{
+				// We should be build the data later, and we can track that the resource was built there when we push it to the DDC.
+				COOK_STAT(Timer.TrackCyclesOnly());
+				InOutShaderMap = nullptr;
+			}
+		}
+		INC_FLOAT_STAT_BY(STAT_ShaderCompiling_DDCLoading,(float)NiagaraShaderDDCTime);
+	}
+}
+
+void FNiagaraShaderMap::SaveToDerivedDataCache()
+{
+	COOK_STAT(auto Timer = NiagaraShaderCookStats::UsageStats.TimeSyncWork());
+	TArray<uint8> SaveData;
+	FMemoryWriter Ar(SaveData, true);
+	Serialize(Ar);
+
+	GetDerivedDataCacheRef().Put(*GetNiagaraShaderMapKeyString(ShaderMapId, Platform), SaveData);
+	COOK_STAT(Timer.AddMiss(SaveData.Num()));
+}
+
+TArray<uint8>* FNiagaraShaderMap::BackupShadersToMemory()
+{
+	TArray<uint8>* SavedShaderData = new TArray<uint8>();
+	FMemoryWriter Ar(*SavedShaderData);
+
+	SerializeInline(Ar, true, true);
+	RegisterSerializedShaders();
+	Empty();
+
+	return SavedShaderData;
+}
+
+void FNiagaraShaderMap::RestoreShadersFromMemory(const TArray<uint8>& ShaderData)
+{
+	FMemoryReader Ar(ShaderData);
+	SerializeInline(Ar, true, true);
+	RegisterSerializedShaders();
+}
+
+void FNiagaraShaderMap::SaveForRemoteRecompile(FArchive& Ar, const TMap<FString, TArray<TRefCountPtr<FNiagaraShaderMap> > >& CompiledShaderMaps, const TArray<FShaderResourceId>& ClientResourceIds)
+{
+	UE_LOG(LogShaders, Display, TEXT("Niagara shader map looking for unique resources, %d were on client"), ClientResourceIds.Num());
+
+	// first, we look for the unique shader resources
+	TArray<FShaderResource*> UniqueResources;
+	int32 NumSkippedResources = 0;
+
+	for (TMap<FString, TArray<TRefCountPtr<FNiagaraShaderMap> > >::TConstIterator It(CompiledShaderMaps); It; ++It)
+	{
+		const TArray<TRefCountPtr<FNiagaraShaderMap> >& ShaderMapArray = It.Value();
+
+		for (int32 Index = 0; Index < ShaderMapArray.Num(); Index++)
+		{
+			FNiagaraShaderMap* ShaderMap = ShaderMapArray[Index];
+
+			if (ShaderMap)
+			{
+				// get all shaders in the shader map
+				TMap<FShaderId, FShader*> ShaderList;
+				ShaderMap->GetShaderList(ShaderList);
+
+				// get the resources from the shaders
+				for (auto& KeyValue : ShaderList)
+				{
+					FShader* Shader = KeyValue.Value;
+					FShaderResourceId ShaderId = Shader->GetResourceId();
+
+					// skip this shader if the Id was already on the client (ie, it didn't change)
+					if (ClientResourceIds.Contains(ShaderId) == false)
+					{
+						// lookup the resource by ID
+						FShaderResource* Resource = FShaderResource::FindShaderResourceById(ShaderId);
+						// add it if it's unique
+						UniqueResources.AddUnique(Resource);
+					}
+					else
+					{
+						NumSkippedResources++;
+					}
+				}
+			}
+		}
+	}
+
+	UE_LOG(LogShaders, Display, TEXT("Sending %d new Niagara shader resources, skipped %d existing"), UniqueResources.Num(), NumSkippedResources);
+
+	// now serialize them
+	int32 NumUniqueResources = UniqueResources.Num();
+	Ar << NumUniqueResources;
+
+	for (int32 Index = 0; Index < NumUniqueResources; Index++)
+	{
+		UniqueResources[Index]->Serialize(Ar);
+	}
+
+	// now we serialize a map (for each script)
+	int32 MapSize = CompiledShaderMaps.Num();
+	Ar << MapSize;
+
+	for (TMap<FString, TArray<TRefCountPtr<FNiagaraShaderMap> > >::TConstIterator It(CompiledShaderMaps); It; ++It)
+	{
+		const TArray<TRefCountPtr<FNiagaraShaderMap> >& ShaderMapArray = It.Value();
+
+		FString ScriptName = It.Key();
+		Ar << ScriptName;
+
+		int32 NumShaderMaps = ShaderMapArray.Num();
+		Ar << NumShaderMaps;
+
+		for (int32 Index = 0; Index < ShaderMapArray.Num(); Index++)
+		{
+			FNiagaraShaderMap* ShaderMap = ShaderMapArray[Index];
+
+			if (ShaderMap && NumUniqueResources > 0)
+			{
+				uint8 bIsValid = 1;
+				Ar << bIsValid;
+				ShaderMap->Serialize(Ar, false);
+			}
+			else
+			{
+				uint8 bIsValid = 0;
+				Ar << bIsValid;
+			}
+		}
+	}
+}
+
+void FNiagaraShaderMap::LoadForRemoteRecompile(FArchive& Ar, EShaderPlatform ShaderPlatform, const TArray<FString>& ScriptsForShaderMaps)
+{
+	check(false);	// TODO: implement
+	/*
+	int32 NumResources;
+	Ar << NumResources;
+
+	// KeepAliveReferences keeps resources alive until we are finished serializing in this function
+	TArray<TRefCountPtr<FShaderResource> > KeepAliveReferences;
+
+	// load and register the resources
+	for (int32 Index = 0; Index < NumResources; Index++)
+	{
+		// Load the inlined shader resource
+		FShaderResource* Resource = new FShaderResource();
+		Resource->Serialize(Ar);
+
+		// if this Id is already in memory, that means that this is a repeated resource and so we skip it
+		if (FShaderResource::FindShaderResourceById(Resource->GetId()) != NULL)
+		{
+			delete Resource;
+		}
+		// otherwise, it's a new resource, so we register it for the maps to find below
+		else
+		{
+			Resource->Register();
+
+			// Keep this guy alive until we finish serializing all the FShaders in
+			// The FShaders which are discarded may cause these resources to be discarded 
+			KeepAliveReferences.Add( Resource );
+		}
+	}
+
+	int32 MapSize;
+	Ar << MapSize;
+
+	for (int32 MaterialIndex = 0; MaterialIndex < MapSize; MaterialIndex++)
+	{
+		FString MaterialName;
+		Ar << MaterialName;
+
+		UMaterialInterface* MatchingMaterial = FindObjectChecked<UMaterialInterface>(NULL, *MaterialName);
+
+		int32 NumShaderMaps = 0;
+		Ar << NumShaderMaps;
+
+		TArray<TRefCountPtr<FNiagaraShaderMap> > LoadedShaderMaps;
+
+		for (int32 ShaderMapIndex = 0; ShaderMapIndex < NumShaderMaps; ShaderMapIndex++)
+		{
+			uint8 bIsValid = 0;
+			Ar << bIsValid;
+
+			if (bIsValid)
+			{
+				FNiagaraShaderMap* ShaderMap = new FNiagaraShaderMap;
+
+				// serialize the id and the material shader map
+				ShaderMap->Serialize(Ar, false);
+
+				// Register in the global map
+				ShaderMap->Register(ShaderPlatform);
+
+				LoadedShaderMaps.Add(ShaderMap);
+			}
+		}
+
+		// Assign in two passes: first pass for shader maps with unspecified quality levels,
+		// Second pass for shader maps with a specific quality level
+			for (int32 ShaderMapIndex = 0; ShaderMapIndex < LoadedShaderMaps.Num(); ShaderMapIndex++)
+			{
+				FNiagaraShaderMap* LoadedShaderMap = LoadedShaderMaps[ShaderMapIndex];
+
+				if (LoadedShaderMap->GetShaderPlatform() == ShaderPlatform 
+					&& LoadedShaderMap->GetShaderMapId().FeatureLevel == GetMaxSupportedFeatureLevel(ShaderPlatform))
+				{
+					FMaterialResource* MaterialResource = MatchingMaterial->GetMaterialResource(GetMaxSupportedFeatureLevel(ShaderPlatform), (EMaterialQualityLevel::Type)QualityLevelIndex);
+					MaterialResource->SetGameThreadShaderMap(LoadedShaderMap);
+					MaterialResource->RegisterInlineShaderMap();
+
+					ENQUEUE_UNIQUE_RENDER_COMMAND_TWOPARAMETER(
+						FSetShaderMapOnMaterialResources,
+						FMaterial*,MaterialResource,MaterialResource,
+						FNiagaraShaderMap*,LoadedShaderMap,LoadedShaderMap,
+					{
+						MaterialResource->SetRenderingThreadShaderMap(LoadedShaderMap);
+					});
+				}
+			}
+		}
+	}
+	*/
+}
+
+
+
+
+/**
+* Compiles the shaders for a script and caches them in this shader map.
+* @param script - The script to compile shaders for.
+* @param InShaderMapId - the script id and set of static parameters to compile for
+* @param Platform - The platform to compile to
+*/
+void FNiagaraShaderMap::Compile(
+	FNiagaraScript* Script,
+	const FNiagaraShaderMapId& InShaderMapId,
+	TRefCountPtr<FShaderCompilerEnvironment> CompilationEnvironment,
+	const FNiagaraComputeShaderCompilationOutput& InNiagaraCompilationOutput,
+	EShaderPlatform InPlatform,
+	bool bSynchronousCompile,
+	bool bApplyCompletedShaderMapForRendering)
+{
+	if (FPlatformProperties::RequiresCookedData())
+	{
+		UE_LOG(LogShaders, Fatal, TEXT("Trying to compile Niagara shader %s at run-time, which is not supported on consoles!"), *Script->GetFriendlyName() );
+	}
+	else
+	{
+		// Make sure we are operating on a referenced shader map or the below Find will cause this shader map to be deleted,
+		// Since it creates a temporary ref counted pointer.
+		check(NumRefs > 0);
+  
+		// Add this shader map and to NiagaraShaderMapsBeingCompiled
+		TArray<FNiagaraScript*>* CorrespondingScripts = NiagaraShaderMapsBeingCompiled.Find(this);
+  
+		if (CorrespondingScripts)
+		{
+			check(!bSynchronousCompile);
+			CorrespondingScripts->AddUnique(Script);
+		}
+		else
+		{
+			// Assign a unique identifier so that shaders from this shader map can be associated with it after a deferred compile
+			CompilingId = NextCompilingId;
+			check(NextCompilingId < UINT_MAX);
+			NextCompilingId++;
+  
+			TArray<FNiagaraScript*> NewCorrespondingScripts;
+			NewCorrespondingScripts.Add(Script);
+			NiagaraShaderMapsBeingCompiled.Add(this, NewCorrespondingScripts);
+#if DEBUG_INFINITESHADERCOMPILE
+			UE_LOG(LogTemp, Display, TEXT("Added Niagara ShaderMap 0x%08X%08X with Script 0x%08X%08X to NiagaraShaderMapsBeingCompiled"), (int)((int64)(this) >> 32), (int)((int64)(this)), (int)((int64)(Script) >> 32), (int)((int64)(Script)));
+#endif  
+			// Setup the compilation environment.
+			Script->SetupShaderCompilationEnvironment(InPlatform, *CompilationEnvironment);
+  
+			// Store the script name for debugging purposes.
+			FriendlyName = Script->GetFriendlyName();
+			NiagaraCompilationOutput = InNiagaraCompilationOutput;
+			ShaderMapId = InShaderMapId;
+			Platform = InPlatform;
+
+			uint32 NumShaders = 0;
+			TArray<FNiagaraShaderCompileJob*> NewJobs;
+	
+			// Iterate over all shader types.
+			TMap<FShaderType*, FNiagaraShaderCompileJob*> SharedShaderJobs;
+			for(TLinkedList<FShaderType*>::TIterator ShaderTypeIt(FShaderType::GetTypeList());ShaderTypeIt;ShaderTypeIt.Next())
+			{
+				FNiagaraShaderType* ShaderType = ShaderTypeIt->GetNiagaraShaderType();
+				if (ShaderType && ShouldCacheNiagaraShader(ShaderType, InPlatform, Script))
+				{
+					// Compile this niagara shader .
+					TArray<FString> ShaderErrors;
+  
+					// Only compile the shader if we don't already have it
+					if (!HasShader(ShaderType))
+					{
+						auto* Job = ShaderType->BeginCompileShader(
+							CompilingId,
+							Script,
+							CompilationEnvironment,
+							Platform,
+							NewJobs,
+							FShaderTarget(ShaderType->GetFrequency(), Platform),
+							Script->GetDataInterfaceBufferDescriptors()
+							);
+						check(!SharedShaderJobs.Find(ShaderType));
+						SharedShaderJobs.Add(ShaderType, Job);
+					}
+					NumShaders++;
+				}
+			}
+  
+			if (!CorrespondingScripts)
+			{
+				UE_LOG(LogShaders, Warning, TEXT("		%u Shaders"), NumShaders);
+			}
+
+			// Register this shader map in the global script->shadermap map
+			Register(InPlatform);
+  
+			// Mark the shader map as not having been finalized with ProcessCompilationResults
+			bCompilationFinalized = false;
+  
+			// Mark as not having been compiled
+			bCompiledSuccessfully = false;
+  
+			GNiagaraShaderCompilationManager.AddJobs(NewJobs);
+  
+			// Compile the shaders for this shader map now if not deferring and deferred compiles are not enabled globally
+			if (bSynchronousCompile)
+			{
+				TArray<int32> CurrentShaderMapId;
+				CurrentShaderMapId.Add(CompilingId);
+				//GNiagaraShaderCompilationManager->FinishCompilation(*FriendlyName, CurrentShaderMapId);
+			}
+		}
+	}
+}
+
+FShader* FNiagaraShaderMap::ProcessCompilationResultsForSingleJob(FNiagaraShaderCompileJob* SingleJob, const FSHAHash& ShaderMapHash)
+{
+	check(SingleJob);
+	const FNiagaraShaderCompileJob& CurrentJob = *SingleJob;
+	check(CurrentJob.Id == CompilingId);
+
+	FShader* Shader = nullptr;
+
+	FNiagaraShaderType* NiagaraShaderType = CurrentJob.ShaderType->GetNiagaraShaderType();
+	check(NiagaraShaderType);
+	Shader = NiagaraShaderType->FinishCompileShader(ShaderMapHash, CurrentJob, FriendlyName);
+	FNiagaraShader *NiagaraShader = static_cast<FNiagaraShader*>(Shader);
+	check(Shader);
+	check(!HasShader(NiagaraShaderType));
+	AddShader(NiagaraShaderType, Shader);
+
+	return Shader;
+}
+
+bool FNiagaraShaderMap::ProcessCompilationResults(const TArray<FNiagaraShaderCompileJob*>& InCompilationResults, int32& InOutJobIndex, float& TimeBudget)
+{
+	check(InOutJobIndex < InCompilationResults.Num());
+
+	double StartTime = FPlatformTime::Seconds();
+
+	FSHAHash ShaderMapHash;
+	ShaderMapId.GetScriptHash(ShaderMapHash);
+
+	do
+	{
+		FNiagaraShaderCompileJob* SingleJob = InCompilationResults[InOutJobIndex];
+		ensure(SingleJob);
+
+		{
+			ProcessCompilationResultsForSingleJob(SingleJob, ShaderMapHash);
+		}
+
+		InOutJobIndex++;
+		
+		double NewStartTime = FPlatformTime::Seconds();
+		TimeBudget -= NewStartTime - StartTime;
+		StartTime = NewStartTime;
+	}
+	while ((TimeBudget > 0.0f) && (InOutJobIndex < InCompilationResults.Num()));
+
+	if (InOutJobIndex == InCompilationResults.Num())
+	{
+		SaveToDerivedDataCache();
+		// The shader map can now be used on the rendering thread
+		bCompilationFinalized = true;
+		return true;
+	}
+
+	return false;
+}
+
+bool FNiagaraShaderMap::TryToAddToExistingCompilationTask(FNiagaraScript* Script)
+{
+	check(NumRefs > 0);
+	TArray<FNiagaraScript*>* CorrespondingScripts = FNiagaraShaderMap::NiagaraShaderMapsBeingCompiled.Find(this);
+
+	if (CorrespondingScripts)
+	{
+		CorrespondingScripts->AddUnique(Script);
+#if DEBUG_INFINITESHADERCOMPILE
+		UE_LOG(LogTemp, Display, TEXT("Added shader map 0x%08X%08X from Niagara script 0x%08X%08X"), (int)((int64)(this) >> 32), (int)((int64)(this)), (int)((int64)(Script) >> 32), (int)((int64)(Script)));
+#endif
+		return true;
+	}
+
+	return false;
+}
+
+bool FNiagaraShaderMap::IsNiagaraShaderComplete(const FNiagaraScript* Script, const FNiagaraShaderType* ShaderType, bool bSilent)
+{
+	// If we should cache this script, it's incomplete if the shader is missing
+	if (ShouldCacheNiagaraShader(ShaderType, Platform, Script) &&	!HasShader((FShaderType*)ShaderType))
+	{
+		if (!bSilent)
+		{
+			UE_LOG(LogShaders, Warning, TEXT("Incomplete shader %s, missing FNiagaraShader %s."), *Script->GetFriendlyName(), ShaderType->GetName());
+		}
+		return false;
+	}
+
+	return true;
+}
+
+bool FNiagaraShaderMap::IsComplete(const FNiagaraScript* Script, bool bSilent)
+{
+	check(!IsInRenderingThread());
+	// Make sure we are operating on a referenced shader map or the below Find will cause this shader map to be deleted,
+	// Since it creates a temporary ref counted pointer.
+	check(NumRefs > 0);
+	const TArray<FNiagaraScript*>* CorrespondingScripts = FNiagaraShaderMap::NiagaraShaderMapsBeingCompiled.Find(this);
+
+	if (CorrespondingScripts)
+	{
+		check(!bCompilationFinalized);
+		return false;
+	}
+
+	// Iterate over all shader types.
+	for(TLinkedList<FShaderType*>::TIterator ShaderTypeIt(FShaderType::GetTypeList());ShaderTypeIt;ShaderTypeIt.Next())
+	{
+		// Find this shader type in the script's shader map.
+		const FNiagaraShaderType* ShaderType = ShaderTypeIt->GetNiagaraShaderType();
+		if (ShaderType && !IsNiagaraShaderComplete(Script, ShaderType, bSilent))
+		{
+			return false;
+		}
+	}
+
+	return true;
+}
+
+void FNiagaraShaderMap::LoadMissingShadersFromMemory(const FNiagaraScript* Script)
+{
+	// Make sure we are operating on a referenced shader map or the below Find will cause this shader map to be deleted,
+	// Since it creates a temporary ref counted pointer.
+	check(NumRefs > 0);
+
+	const TArray<FNiagaraScript*>* CorrespondingScripts = FNiagaraShaderMap::NiagaraShaderMapsBeingCompiled.Find(this);
+
+	if (CorrespondingScripts)
+	{
+		check(!bCompilationFinalized);
+		return;
+	}
+
+	FSHAHash ShaderMapHash;
+	ShaderMapId.GetScriptHash(ShaderMapHash);
+
+	// Try to find necessary FNiagaraShaderType's in memory
+	for (TLinkedList<FShaderType*>::TIterator ShaderTypeIt(FShaderType::GetTypeList());ShaderTypeIt;ShaderTypeIt.Next())
+	{
+		FNiagaraShaderType* ShaderType = ShaderTypeIt->GetNiagaraShaderType();
+		if (ShaderType && ShouldCacheNiagaraShader(ShaderType, Platform, Script) && !HasShader(ShaderType))
+		{
+			FShaderId ShaderId(ShaderMapHash, nullptr, nullptr, ShaderType, FShaderTarget(ShaderType->GetFrequency(), Platform));
+			FShader* FoundShader = ShaderType->FindShaderById(ShaderId);
+			if (FoundShader)
+			{
+				AddShader(ShaderType, FoundShader);
+			}
+		}
+	}
+}
+#endif
+
+void FNiagaraShaderMap::GetShaderList(TMap<FShaderId, FShader*>& OutShaders) const
+{
+	TShaderMap<FNiagaraShaderType>::GetShaderList(OutShaders);
+}
+
+/**
+ * Registers a Niagara shader map in the global map so it can be used by scripts.
+ */
+void FNiagaraShaderMap::Register(EShaderPlatform InShaderPlatform)
+{
+	extern int32 GCreateNiagaraShadersOnLoad;
+	if (GCreateNiagaraShadersOnLoad && Platform == InShaderPlatform)
+	{
+		for (auto KeyValue : GetShaders())
+		{
+			FShader* Shader = KeyValue.Value;
+			if (Shader)
+			{
+				Shader->BeginInitializeResources();
+			}
+		}
+	}
+
+	if (!bRegistered)
+	{
+		INC_DWORD_STAT(STAT_Shaders_NumShaderMaps);
+		INC_DWORD_STAT_BY(STAT_Shaders_ShaderMapMemory, GetSizeBytes());
+	}
+
+	GIdToNiagaraShaderMap[Platform].Add(ShaderMapId,this);
+	bRegistered = true;
+}
+
+void FNiagaraShaderMap::AddRef()
+{
+	check(!bDeletedThroughDeferredCleanup);
+	++NumRefs;
+}
+
+void FNiagaraShaderMap::Release()
+{
+	check(NumRefs > 0);
+	if(--NumRefs == 0)
+	{
+		if (bRegistered)
+		{
+			DEC_DWORD_STAT(STAT_Shaders_NumShaderMaps);
+			DEC_DWORD_STAT_BY(STAT_Shaders_ShaderMapMemory, GetSizeBytes());
+
+			GIdToNiagaraShaderMap[Platform].Remove(ShaderMapId);
+			bRegistered = false;
+		}
+
+		BeginCleanup(this);
+	}
+}
+
+FNiagaraShaderMap::FNiagaraShaderMap() :
+	TShaderMap<FNiagaraShaderType>(SP_NumPlatforms),
+	Platform(SP_NumPlatforms),
+	CompilingId(1),
+	NumRefs(0),
+	bDeletedThroughDeferredCleanup(false),
+	bRegistered(false),
+	bCompilationFinalized(true),
+	bCompiledSuccessfully(true),
+	bIsPersistent(true) 
+{
+	checkSlow(IsInGameThread() || IsAsyncLoading());
+	AllNiagaraShaderMaps.Add(this);
+}
+
+FNiagaraShaderMap::~FNiagaraShaderMap()
+{ 
+	checkSlow(IsInGameThread() || IsAsyncLoading());
+	check(bDeletedThroughDeferredCleanup);
+	check(!bRegistered);
+	AllNiagaraShaderMaps.RemoveSwap(this);
+}
+
+/**
+ * Removes all entries in the cache with exceptions based on a shader type
+ * @param ShaderType - The shader type to flush
+ */
+void FNiagaraShaderMap::FlushShadersByShaderType(FShaderType* ShaderType)
+{
+	if (ShaderType->GetNiagaraShaderType())
+	{
+		RemoveShaderType(ShaderType->GetNiagaraShaderType());	
+	}
+}
+
+
+
+void FNiagaraShaderMap::Serialize(FArchive& Ar, bool bInlineShaderResources)
+{
+	// Note: This is saved to the DDC, not into packages (except when cooked)
+	// Backwards compatibility therefore will not work based on the version of Ar
+	// Instead, just bump NIAGARASHADERMAP_DERIVEDDATA_VER
+
+	ShaderMapId.Serialize(Ar);
+
+	// serialize the platform enum as a uint8
+	int32 TempPlatform = (int32)Platform;
+	Ar << TempPlatform;
+	Platform = (EShaderPlatform)TempPlatform;
+
+	Ar << FriendlyName;
+
+	NiagaraCompilationOutput.Serialize(Ar);
+
+	Ar << DebugDescription;
+
+	if (Ar.IsSaving())
+	{
+		TShaderMap<FNiagaraShaderType>::SerializeInline(Ar, bInlineShaderResources, false);
+		RegisterSerializedShaders();
+	}
+
+	if (Ar.IsLoading())
+	{
+		TShaderMap<FNiagaraShaderType>::SerializeInline(Ar, bInlineShaderResources, false);
+	}
+}
+
+void FNiagaraShaderMap::RegisterSerializedShaders()
+{
+	check(IsInGameThread());
+
+	TShaderMap<FNiagaraShaderType>::RegisterSerializedShaders();
+}
+
+void FNiagaraShaderMap::DiscardSerializedShaders()
+{
+	TShaderMap<FNiagaraShaderType>::DiscardSerializedShaders();
+}
+
+void FNiagaraShaderMap::RemovePendingScript(FNiagaraScript* Script)
+{
+	for (TMap<TRefCountPtr<FNiagaraShaderMap>, TArray<FNiagaraScript*> >::TIterator It(NiagaraShaderMapsBeingCompiled); It; ++It)
+	{
+		TArray<FNiagaraScript*>& Scripts = It.Value();
+		int32 Result = Scripts.Remove(Script);
+#if DEBUG_INFINITESHADERCOMPILE
+		if ( Result )
+		{
+			UE_LOG(LogTemp, Display, TEXT("Removed shader map 0x%08X%08X from script 0x%08X%08X"), (int)((int64)(It.Key().GetReference()) >> 32), (int)((int64)(It.Key().GetReference())), (int)((int64)(Script) >> 32), (int)((int64)(Script)));
+		}
+#endif
+	}
+}
+
+const FNiagaraShaderMap* FNiagaraShaderMap::GetShaderMapBeingCompiled(const FNiagaraScript* Script)
+{
+	// Inefficient search, but only when compiling a lot of shaders
+	for (TMap<TRefCountPtr<FNiagaraShaderMap>, TArray<FNiagaraScript*> >::TIterator It(NiagaraShaderMapsBeingCompiled); It; ++It)
+	{
+		TArray<FNiagaraScript*>& Scripts = It.Value();
+		
+		for (int32 ScriptIndex = 0; ScriptIndex < Scripts.Num(); ScriptIndex++)
+		{
+			if (Scripts[ScriptIndex] == Script)
+			{
+				return It.Key();
+			}
+		}
+	}
+
+	return NULL;
+}
+
+
+
+
+FNiagaraShader::FNiagaraShader(const FNiagaraShaderType::CompiledShaderInitializerType& Initializer)
+	: FShader(Initializer)
+	, DebugDescription(Initializer.DebugDescription)
+{
+	check(!DebugDescription.IsEmpty());
+	SetDatainterfaceBufferDescriptors(Initializer.DIBufferDescriptors);
+	BindParams(Initializer.ParameterMap);
+}
+
+
+
+bool FNiagaraShader::Serialize(FArchive& Ar)
+{
+	const bool bShaderHasOutdatedParameters = FShader::Serialize(Ar);
+	Ar << NiagaraUniformBuffer;	// do we need to keep about uniform buffer parameters?
+
+//	Ar.UsingCustomVersion(FNiagaraShaderCustomVersion::GUID);		// only changes version if not loading
+//	const int32 NiagaraVer = Ar.CustomVer(FNiagaraShaderCustomVersion::GUID);
+
+	Ar << DebugDescription;
+
+	Ar << FloatInputBufferParam;
+	Ar << IntInputBufferParam;
+	Ar << FloatOutputBufferParam;
+	Ar << IntOutputBufferParam;
+	Ar << OutputIndexBufferParam;
+	Ar << EmitterTickCounterParam;
+
+	Ar << NumInstancesPerThreadParam;
+	Ar << NumInstancesParam;
+	Ar << StartInstanceParam;
+	Ar << GroupStartInstanceParam;
+	Ar << ComponentBufferSizeReadParam;
+	Ar << ComponentBufferSizeWriteParam;
+
+	for (uint32 i = 0; i < MAX_CONCURRENT_EVENT_DATASETS; i++)
+	{
+		Ar << EventFloatUAVParams[i];
+		Ar << EventIntUAVParams[i];
+		Ar << EventFloatSRVParams[i];
+		Ar << EventIntSRVParams[i];
+	}
+
+	Ar << DIBufferDescriptors;
+
+	// params for data interface buffers
+	if (Ar.IsLoading())
+	{
+		for (TArray<DIGPUBufferParamDescriptor> &InterfaceDescs : DIBufferDescriptors)
+		{
+			int32 Idx = NameToDIBufferParamMap.AddDefaulted(1);
+			for (DIGPUBufferParamDescriptor &Desc : InterfaceDescs)
+			{
+				FShaderResourceParameter &Param = NameToDIBufferParamMap[Idx].Add(*Desc.BufferParamName);
+				Ar << Param;
+			}
+		}
+	}
+	else
+	{
+		uint32 Idx = 0;
+		for (TArray<DIGPUBufferParamDescriptor> &InterfaceDescs : DIBufferDescriptors)
+		{
+			for (DIGPUBufferParamDescriptor &Desc : InterfaceDescs)
+			{
+				FShaderResourceParameter *Param = NameToDIBufferParamMap[Idx].Find(*Desc.BufferParamName);
+				check(Param);
+				Ar << *Param;
+			}
+			Idx++;
+		}
+	}
+
+	Ar << SimulateStartInstanceParam;
+	Ar << NumThreadGroupsParam;
+	Ar << EmitterConstantBufferParam;
+	Ar << DataInterfaceUniformBufferParam;
+	Ar << NumEventsPerParticleParam;
+	Ar << NumParticlesPerEventParam;
+	Ar << CopyInstancesBeforeStartParam;
+
+	return bShaderHasOutdatedParameters;
+}
+
+
+uint32 FNiagaraShader::GetAllocatedSize() const
+{
+	return FShader::GetAllocatedSize()
+		//+ ParameterCollectionUniformBuffers.GetAllocatedSize()
+		+ DebugDescription.GetAllocatedSize();
+}
+
+
+FArchive &operator << (FArchive &Ar, DIGPUBufferParamDescriptor &Desc)
+{
+	Ar << Desc.BufferParamName;
+	Ar << Desc.Index;
+	return Ar;
+}
diff --git a/Engine/Plugins/FX/Niagara/Source/NiagaraShader/Private/NiagaraShaderCompilationManager.cpp b/Engine/Plugins/FX/Niagara/Source/NiagaraShader/Private/NiagaraShaderCompilationManager.cpp
new file mode 100644
--- /dev/null
+++ b/Engine/Plugins/FX/Niagara/Source/NiagaraShader/Private/NiagaraShaderCompilationManager.cpp
@@ -0,0 +1,384 @@
+// Copyright 1998-2017 Epic Games, Inc. All Rights Reserved.
+
+#include "NiagaraShaderCompilationManager.h"
+#include "NiagaraShared.h"
+#if WITH_EDITOR
+#include "IShaderFormat.h"
+#include "ITargetPlatformManagerModule.h"
+#endif
+#include "FileManager.h"
+#include "Paths.h"
+
+
+DEFINE_LOG_CATEGORY_STATIC(LogNiagaraShaderCompiler, All, All);
+
+static int32 GShowNiagaraShaderWarnings = 1;
+static FAutoConsoleVariableRef CVarShowNiagaraShaderWarnings(
+	TEXT("niagara.ShowShaderCompilerWarnings"),
+	GShowNiagaraShaderWarnings,
+	TEXT("When set to 1, will display all warnings from Niagara shader compiles.")
+	);
+
+#if WITH_EDITOR
+
+NIAGARASHADER_API FNiagaraShaderCompilationManager GNiagaraShaderCompilationManager;
+
+void FNiagaraShaderCompilationManager::Tick(float DeltaSeconds)
+{
+	RunCompileJobs();
+}
+
+FNiagaraShaderCompilationManager::FNiagaraShaderCompilationManager()
+{
+	// Ew. Should we just use FShaderCompilingManager's workers instead? Is that safe?
+	const int32 NumVirtualCores = FPlatformMisc::NumberOfCoresIncludingHyperthreads();
+	const uint32 NumNiagaraShaderCompilingThreads = FMath::Min(NumVirtualCores-1, 4);		
+
+	for (uint32 WorkerIndex = 0; WorkerIndex < NumNiagaraShaderCompilingThreads; WorkerIndex++)
+	{
+		WorkerInfos.Add(new FNiagaraShaderCompileWorkerInfo());
+	}
+}
+
+
+void FNiagaraShaderCompilationManager::RunCompileJobs()
+{
+	// If we aren't compiling through workers, so we can just track the serial time here.
+//	COOK_STAT(FScopedDurationTimer CompileTimer(NiagaraShaderCookStats::AsyncCompileTimeSec));
+	int32 NumActiveThreads = 0;
+
+
+
+	for (int32 WorkerIndex = 0; WorkerIndex < WorkerInfos.Num(); WorkerIndex++)
+	{
+		FNiagaraShaderCompileWorkerInfo& CurrentWorkerInfo = *WorkerInfos[WorkerIndex];
+
+		// If this worker doesn't have any queued jobs, look for more in the input queue
+		if (CurrentWorkerInfo.QueuedJobs.Num() == 0)
+		{
+			check(!CurrentWorkerInfo.bComplete);
+
+			if (JobQueue.Num() > 0)
+			{
+				bool bAddedLowLatencyTask = false;
+				int32 JobIndex = 0;
+
+				// Try to grab up to MaxShaderJobBatchSize jobs
+				// Don't put more than one low latency task into a batch
+				for (; JobIndex < JobQueue.Num(); JobIndex++)
+				{
+					CurrentWorkerInfo.QueuedJobs.Add(JobQueue[JobIndex]);
+				}
+
+				// Update the worker state as having new tasks that need to be issued					
+				// don't reset worker app ID, because the shadercompilerworkers don't shutdown immediately after finishing a single job queue.
+				CurrentWorkerInfo.bIssuedTasksToWorker = true;
+				CurrentWorkerInfo.bLaunchedWorker = true;
+				CurrentWorkerInfo.StartTime = FPlatformTime::Seconds();
+				JobQueue.RemoveAt(0, JobIndex);
+			}
+		}
+
+		if (CurrentWorkerInfo.bIssuedTasksToWorker && CurrentWorkerInfo.bLaunchedWorker)
+		{
+			NumActiveThreads++;
+		}
+
+		if (CurrentWorkerInfo.QueuedJobs.Num() > 0)
+		{
+			for (int32 JobIndex = 0; JobIndex < CurrentWorkerInfo.QueuedJobs.Num(); JobIndex++)
+			{
+				FNiagaraShaderCompileJob& CurrentJob = *CurrentWorkerInfo.QueuedJobs[JobIndex];
+
+				check(!CurrentJob.bFinalized);
+				CurrentJob.bFinalized = true;
+
+				static ITargetPlatformManagerModule& TPM = GetTargetPlatformManagerRef();
+
+				const FName Format = LegacyShaderPlatformToShaderFormat(EShaderPlatform(CurrentJob.Input.Target.Platform));
+				const IShaderFormat* Compiler = TPM.FindShaderFormat(Format);
+
+				if (!Compiler)
+				{
+					UE_LOG(LogNiagaraShaderCompiler, Fatal, TEXT("Can't compile shaders for format %s, couldn't load compiler dll"), *Format.ToString());
+				}
+				CA_ASSUME(Compiler != NULL);
+
+				FString AbsoluteDebugInfoDirectory = IFileManager::Get().ConvertToAbsolutePathForExternalAppForWrite(*(FPaths::ProjectSavedDir() / TEXT("ShaderDebugInfo")));
+				FPaths::NormalizeDirectoryName(AbsoluteDebugInfoDirectory);
+				CurrentJob.Input.DumpDebugInfoPath = AbsoluteDebugInfoDirectory / Format.ToString() / FString("Niagara");
+				if (!IFileManager::Get().DirectoryExists(*CurrentJob.Input.DumpDebugInfoPath))
+				{
+					verifyf(IFileManager::Get().MakeDirectory(*CurrentJob.Input.DumpDebugInfoPath, true), TEXT("Failed to create directory for shader debug info '%s'"), *CurrentJob.Input.DumpDebugInfoPath);
+				}
+
+				if (IsValidRef(CurrentJob.Input.SharedEnvironment))
+				{
+					// Merge the shared environment into the per-shader environment before calling into the compile function
+					// Normally this happens in the worker
+					CurrentJob.Input.Environment.Merge(*CurrentJob.Input.SharedEnvironment);
+				}
+
+				// Compile the shader directly through the platform dll (directly from the shader dir as the working directory)
+				Compiler->CompileShader(Format, CurrentJob.Input, CurrentJob.Output, FString(FPlatformProcess::ShaderDir()));
+
+				CurrentJob.bSucceeded = CurrentJob.Output.bSucceeded;
+
+				if (CurrentJob.Output.bSucceeded)
+				{
+					// Generate a hash of the output and cache it
+					// The shader processing this output will use it to search for existing FShaderResources
+					CurrentJob.Output.GenerateOutputHash();
+					UE_LOG(LogNiagaraShaderCompiler, Log, TEXT("GPU shader compile succeeded."));
+				}
+				else
+				{
+					UE_LOG(LogNiagaraShaderCompiler, Log, TEXT("ERROR: GPU shader compile failed!"));
+				}
+
+				CurrentWorkerInfo.bComplete = true;
+			}
+		}
+	}
+
+	for (int32 WorkerIndex = 0; WorkerIndex < WorkerInfos.Num(); WorkerIndex++)
+	{
+		FNiagaraShaderCompileWorkerInfo& CurrentWorkerInfo = *WorkerInfos[WorkerIndex];
+		if (CurrentWorkerInfo.bComplete)
+		{
+			for (int32 JobIndex = 0; JobIndex < CurrentWorkerInfo.QueuedJobs.Num(); JobIndex++)
+			{
+				FNiagaraShaderMapCompileResults& ShaderMapResults = NiagaraShaderMapJobs.FindChecked(CurrentWorkerInfo.QueuedJobs[JobIndex]->Id);
+				ShaderMapResults.FinishedJobs.Add(CurrentWorkerInfo.QueuedJobs[JobIndex]);
+				ShaderMapResults.bAllJobsSucceeded = ShaderMapResults.bAllJobsSucceeded && CurrentWorkerInfo.QueuedJobs[JobIndex]->bSucceeded;
+			}
+		}
+
+		CurrentWorkerInfo.bComplete = false;
+		CurrentWorkerInfo.QueuedJobs.Empty();
+	}
+
+}
+
+
+
+NIAGARASHADER_API void FNiagaraShaderCompilationManager::AddJobs(TArray<FNiagaraShaderCompileJob*> InNewJobs)
+{
+	for (FNiagaraShaderCompileJob *Job : InNewJobs)
+	{
+		FNiagaraShaderMapCompileResults& ShaderMapInfo = NiagaraShaderMapJobs.FindOrAdd(Job->Id);
+//		ShaderMapInfo.bApplyCompletedShaderMapForRendering = bApplyCompletedShaderMapForRendering;
+//		ShaderMapInfo.bRecreateComponentRenderStateOnCompletion = bRecreateComponentRenderStateOnCompletion;
+		ShaderMapInfo.NumJobsQueued++;
+	}
+
+	JobQueue.Append(InNewJobs);
+}
+
+
+void FNiagaraShaderCompilationManager::ProcessAsyncResults()
+{
+	int32 NumCompilingNiagaraShaderMaps = 0;
+	TArray<int32> ShaderMapsToRemove;
+
+	// Get all Niagara shader maps to finalize
+	//
+	for (TMap<int32, FNiagaraShaderMapCompileResults>::TIterator It(NiagaraShaderMapJobs); It; ++It)
+	{
+		const FNiagaraShaderMapCompileResults& Results = It.Value();
+
+		if (Results.FinishedJobs.Num() == Results.NumJobsQueued)
+		{
+			ShaderMapsToRemove.Add(It.Key());
+			PendingFinalizeNiagaraShaderMaps.Add(It.Key(), FNiagaraShaderMapFinalizeResults(Results));
+		}
+	}
+
+	for (int32 RemoveIndex = 0; RemoveIndex < ShaderMapsToRemove.Num(); RemoveIndex++)
+	{
+		NiagaraShaderMapJobs.Remove(ShaderMapsToRemove[RemoveIndex]);
+	}
+
+	NumCompilingNiagaraShaderMaps = NiagaraShaderMapJobs.Num();
+
+	if (PendingFinalizeNiagaraShaderMaps.Num() > 0)
+	{
+		ProcessCompiledNiagaraShaderMaps(PendingFinalizeNiagaraShaderMaps, 0.1f);
+	}
+}
+
+
+
+
+
+
+void FNiagaraShaderCompilationManager::ProcessCompiledNiagaraShaderMaps(
+	TMap<int32, FNiagaraShaderMapFinalizeResults>& CompiledShaderMaps,
+	float TimeBudget)
+{
+	// Keeps shader maps alive as they are passed from the shader compiler and applied to the owning Script
+	TArray<TRefCountPtr<FNiagaraShaderMap> > LocalShaderMapReferences;
+	TMap<FNiagaraScript*, FNiagaraShaderMap*> ScriptsToUpdate;
+
+	// Process compiled shader maps in FIFO order, in case a shader map has been enqueued multiple times,
+	// Which can happen if a script is edited while a background compile is going on
+	for (TMap<int32, FNiagaraShaderMapFinalizeResults>::TIterator ProcessIt(CompiledShaderMaps); ProcessIt; ++ProcessIt)
+	{
+		TRefCountPtr<FNiagaraShaderMap> ShaderMap = NULL;
+		TArray<FNiagaraScript*>* Scripts = NULL;
+
+		for (TMap<TRefCountPtr<FNiagaraShaderMap>, TArray<FNiagaraScript*> >::TIterator ShaderMapIt(FNiagaraShaderMap::GetInFlightShaderMaps()); ShaderMapIt; ++ShaderMapIt)
+		{
+			if (ShaderMapIt.Key()->GetCompilingId() == ProcessIt.Key())
+			{
+				ShaderMap = ShaderMapIt.Key();
+				Scripts = &ShaderMapIt.Value();
+				break;
+			}
+		}
+
+		if (ShaderMap && Scripts)
+		{
+			TArray<FString> Errors;
+			FNiagaraShaderMapFinalizeResults& CompileResults = ProcessIt.Value();
+			const TArray<FNiagaraShaderCompileJob*>& ResultArray = CompileResults.FinishedJobs;
+
+			// Make a copy of the array as this entry of FNiagaraShaderMap::ShaderMapsBeingCompiled will be removed below
+			TArray<FNiagaraScript*> ScriptArray = *Scripts;
+			bool bSuccess = true;
+
+			for (int32 JobIndex = 0; JobIndex < ResultArray.Num(); JobIndex++)
+			{
+				FNiagaraShaderCompileJob& CurrentJob = *ResultArray[JobIndex];
+				bSuccess = bSuccess && CurrentJob.bSucceeded;
+
+				if (bSuccess)
+				{
+					check(CurrentJob.Output.ShaderCode.GetShaderCodeSize() > 0);
+				}
+
+				if (GShowNiagaraShaderWarnings || !CurrentJob.bSucceeded)
+				{
+					for (int32 ErrorIndex = 0; ErrorIndex < CurrentJob.Output.Errors.Num(); ErrorIndex++)
+					{
+						Errors.AddUnique(CurrentJob.Output.Errors[ErrorIndex].GetErrorString());
+					}
+				}
+			}
+
+			bool bShaderMapComplete = true;
+
+			if (bSuccess)
+			{
+				bShaderMapComplete = ShaderMap->ProcessCompilationResults(ResultArray, CompileResults.FinalizeJobIndex, TimeBudget);
+			}
+
+
+			if (bShaderMapComplete)
+			{
+				ShaderMap->SetCompiledSuccessfully(bSuccess);
+
+				// Pass off the reference of the shader map to LocalShaderMapReferences
+				LocalShaderMapReferences.Add(ShaderMap);
+				FNiagaraShaderMap::GetInFlightShaderMaps().Remove(ShaderMap);
+
+				for (FNiagaraScript *Script : ScriptArray)
+				{
+					FNiagaraShaderMap* CompletedShaderMap = ShaderMap;
+
+					Script->RemoveOutstandingCompileId(ShaderMap->GetCompilingId());
+
+					// Only process results that still match the ID which requested a compile
+					// This avoids applying shadermaps which are out of date and a newer one is in the async compiling pipeline
+					if (Script->GetScriptID() == CompletedShaderMap->GetShaderMapId().BaseScriptID)
+					{
+						if (!bSuccess)
+						{
+							// Propagate error messages
+							Script->SetCompileErrors(Errors);
+							ScriptsToUpdate.Add(Script, NULL);
+
+							for (int32 ErrorIndex = 0; ErrorIndex < Errors.Num(); ErrorIndex++)
+							{
+								FString ErrorMessage = Errors[ErrorIndex];
+								// Work around build machine string matching heuristics that will cause a cook to fail
+								ErrorMessage.ReplaceInline(TEXT("error "), TEXT("err0r "), ESearchCase::CaseSensitive);
+								UE_LOG(LogNiagaraShaderCompiler, Warning, TEXT("	%s"), *ErrorMessage);
+							}
+						}
+						else
+						{
+							// if we succeeded and our shader map is not complete this could be because the script was being edited quicker then the compile could be completed
+							// Don't modify scripts for which the compiled shader map is no longer complete
+							// This can happen if a script being compiled is edited
+							if (CompletedShaderMap->IsComplete(Script, true))
+							{
+								ScriptsToUpdate.Add(Script, CompletedShaderMap);
+							}
+
+							if (GShowNiagaraShaderWarnings && Errors.Num() > 0)
+							{
+								UE_LOG(LogNiagaraShaderCompiler, Warning, TEXT("Warnings while compiling Niagara Script %s for platform %s:"),
+									*Script->GetFriendlyName(),
+									*LegacyShaderPlatformToShaderFormat(ShaderMap->GetShaderPlatform()).ToString());
+								for (int32 ErrorIndex = 0; ErrorIndex < Errors.Num(); ErrorIndex++)
+								{
+									UE_LOG(LogNiagaraShaderCompiler, Warning, TEXT("	%s"), *Errors[ErrorIndex]);
+								}
+							}
+						}
+					}
+				}
+
+				// Cleanup shader jobs and compile tracking structures
+				for (int32 JobIndex = 0; JobIndex < ResultArray.Num(); JobIndex++)
+				{
+					delete ResultArray[JobIndex];
+				}
+
+				CompiledShaderMaps.Remove(ShaderMap->GetCompilingId());
+			}
+
+			if (TimeBudget < 0)
+			{
+				break;
+			}
+		}
+	}
+
+	if (ScriptsToUpdate.Num() > 0)
+	{
+		for (TMap<FNiagaraScript*, FNiagaraShaderMap*>::TConstIterator It(ScriptsToUpdate); It; ++It)
+		{
+			FNiagaraScript* Script = It.Key();
+			FNiagaraShaderMap* ShaderMap = It.Value();
+			//check(!ShaderMap || ShaderMap->IsValidForRendering());
+
+			Script->SetGameThreadShaderMap(It.Value());
+
+			ENQUEUE_UNIQUE_RENDER_COMMAND_TWOPARAMETER(
+				FSetShaderMapOnScriptResources,
+				FNiagaraScript*, Script, Script,
+				FNiagaraShaderMap*, CompiledShaderMap, ShaderMap,
+				{
+					Script->SetRenderingThreadShaderMap(CompiledShaderMap);
+				});
+
+
+			Script->NotifyCompilationFinished();
+		}
+	}
+}
+
+
+void FNiagaraShaderCompilationManager::FinishCompilation(const TCHAR* ScriptName, const TArray<int32>& ShaderMapIdsToFinishCompiling)
+{
+	check(!FPlatformProperties::RequiresCookedData());
+
+	RunCompileJobs();	// since we don't async compile through another process, this will run all oustanding jobs
+	ProcessAsyncResults();	// grab compiled shader maps and assign them to their resources
+}
+
+
+
+#endif // WITH_EDITOR
\ No newline at end of file
diff --git a/Engine/Plugins/FX/Niagara/Source/NiagaraShader/Private/NiagaraShaderDerivedDataVersion.h b/Engine/Plugins/FX/Niagara/Source/NiagaraShader/Private/NiagaraShaderDerivedDataVersion.h
new file mode 100644
--- /dev/null
+++ b/Engine/Plugins/FX/Niagara/Source/NiagaraShader/Private/NiagaraShaderDerivedDataVersion.h
@@ -0,0 +1,15 @@
+// Copyright 1998-2017 Epic Games, Inc. All Rights Reserved.
+
+/*=============================================================================
+NiagaraShaderDerivedDataVersion.h: Shader derived data version for Niagara.
+=============================================================================*/
+
+#pragma once
+
+#include "CoreMinimal.h"
+
+// In case of merge conflicts with DDC versions, you MUST generate a new GUID and set this new
+// guid as version
+
+#define NIAGARASHADERMAP_DERIVEDDATA_VER		TEXT("1ecd9c26d37a443a9aa832b8feb72475")
+													  
diff --git a/Engine/Plugins/FX/Niagara/Source/NiagaraShader/Private/NiagaraShared.cpp b/Engine/Plugins/FX/Niagara/Source/NiagaraShader/Private/NiagaraShared.cpp
new file mode 100644
--- /dev/null
+++ b/Engine/Plugins/FX/Niagara/Source/NiagaraShader/Private/NiagaraShared.cpp
@@ -0,0 +1,367 @@
+// Copyright 1998-2017 Epic Games, Inc. All Rights Reserved.
+
+/*=============================================================================
+	NiagaraShared.cpp: Shared Niagara compute shader implementation.
+=============================================================================*/
+
+#include "NiagaraShared.h"
+#include "NiagaraShaderModule.h"
+#include "NiagaraShaderType.h"
+#include "NiagaraShader.h"
+#include "Stats/StatsMisc.h"
+#include "UObject/CoreObjectVersion.h"
+#include "Misc/App.h"
+#include "UObject/UObjectHash.h"
+#include "UObject/UObjectIterator.h"
+#include "ShaderCompiler.h"
+#include "NiagaraShaderCompilationManager.h"
+#include "RendererInterface.h"
+
+#if WITH_EDITOR
+	NIAGARASHADER_API FNiagaraCompilationQueue* FNiagaraCompilationQueue::Singleton = nullptr;
+#endif
+
+
+/** Populates OutEnvironment with defines needed to compile shaders for this script. */
+void FNiagaraScript::SetupShaderCompilationEnvironment(
+	EShaderPlatform Platform,
+	FShaderCompilerEnvironment& OutEnvironment
+	) const
+{
+	OutEnvironment.SetDefine(TEXT("GPU_SIMULATION_SHADER"), TEXT("1"));
+}
+
+
+NIAGARASHADER_API bool FNiagaraScript::ShouldCache(EShaderPlatform Platform, const FShaderType* ShaderType) const
+{
+	check(ShaderType->GetNiagaraShaderType() )
+	return true;
+}
+
+
+NIAGARASHADER_API void FNiagaraScript::LegacySerialize(FArchive& Ar)
+{
+}
+
+
+void FNiagaraScript::GetDependentShaderTypes(EShaderPlatform Platform, TArray<FShaderType*>& OutShaderTypes) const
+{
+	for (TLinkedList<FShaderType*>::TIterator ShaderTypeIt(FShaderType::GetTypeList()); ShaderTypeIt; ShaderTypeIt.Next())
+	{
+		FNiagaraShaderType* ShaderType = ShaderTypeIt->GetNiagaraShaderType();
+
+		if ( ShaderType && ShaderType->ShouldCache(Platform, this) && ShouldCache(Platform, ShaderType) )
+		{
+			OutShaderTypes.Add(ShaderType);
+		}
+	}
+}
+
+
+
+NIAGARASHADER_API void FNiagaraScript::GetShaderMapId(EShaderPlatform Platform, FNiagaraShaderMapId& OutId) const
+{
+	if (bLoadedCookedShaderMapId)
+	{
+		OutId = CookedShaderMapId;
+	}
+	else
+	{
+		TArray<FShaderType*> ShaderTypes;
+		GetDependentShaderTypes(Platform, ShaderTypes);
+		OutId.FeatureLevel = GetFeatureLevel();
+		OutId.BaseScriptID = GetScriptID();
+	}
+}
+
+
+
+void FNiagaraScript::AddReferencedObjects(FReferenceCollector& Collector)
+{
+}
+
+
+void FNiagaraScript::RegisterShaderMap()
+{
+	if (GameThreadShaderMap)
+	{
+		GameThreadShaderMap->RegisterSerializedShaders();
+	}
+}
+
+void FNiagaraScript::ReleaseShaderMap()
+{
+	if (GameThreadShaderMap)
+	{
+		GameThreadShaderMap = nullptr;
+
+		ENQUEUE_UNIQUE_RENDER_COMMAND_ONEPARAMETER(
+			ReleaseShaderMap,
+			FNiagaraScript*, Script, this,
+			{
+				Script->SetRenderingThreadShaderMap(nullptr);
+			});
+	}
+}
+
+void FNiagaraScript::SerializeShaderMap(FArchive& Ar)
+{
+	bool bCooked = Ar.IsCooking();
+	Ar << bCooked;
+
+	if (FPlatformProperties::RequiresCookedData() && !bCooked && Ar.IsLoading())
+	{
+		UE_LOG(LogShaders, Fatal, TEXT("This platform requires cooked packages, and shaders were not cooked into this Niagara script %s."), *GetFriendlyName());
+	}
+
+	if (bCooked)
+	{
+		if (Ar.IsCooking())
+		{
+#if WITH_EDITOR
+			FinishCompilation();
+
+			bool bValid = GameThreadShaderMap != nullptr && GameThreadShaderMap->CompiledSuccessfully();
+			ensure(bValid);
+			Ar << bValid;
+
+			if (bValid)
+			{
+				GameThreadShaderMap->Serialize(Ar);
+			}
+			else
+			{
+				FString Name;
+				UE_LOG(LogShaders, Error, TEXT("ERROR: Failed to compile Niagara shader %s."), *GetFriendlyName());
+			}
+#endif
+		}
+		else
+		{
+			bool bValid = false;
+			Ar << bValid;
+
+			if (bValid)
+			{
+				TRefCountPtr<FNiagaraShaderMap> LoadedShaderMap = new FNiagaraShaderMap();
+				LoadedShaderMap->Serialize(Ar);
+
+				// Toss the loaded shader data if this is a server only instance
+				//@todo - don't cook it in the first place
+				if (FApp::CanEverRender())
+				{
+					GameThreadShaderMap = RenderingThreadShaderMap = LoadedShaderMap;
+				}
+				else
+				{
+					LoadedShaderMap->DiscardSerializedShaders();
+				}
+			}
+		}
+	}
+}
+
+
+
+
+
+void FNiagaraScript::SetScript(UNiagaraScript *InScript, ERHIFeatureLevel::Type InFeatureLevel, FGuid InScriptID, FString InFriendlyName)
+{
+	GTScript = InScript;
+	ScriptId = InScriptID;
+	FriendlyName = InFriendlyName;
+	SetFeatureLevel(InFeatureLevel);
+}
+
+NIAGARASHADER_API  void FNiagaraScript::SetRenderingThreadShaderMap(FNiagaraShaderMap* InShaderMap)
+{
+	check(IsInRenderingThread());
+	RenderingThreadShaderMap = InShaderMap;
+}
+
+/**
+* Cache the script's shaders
+*/
+#if WITH_EDITOR
+
+bool FNiagaraScript::CacheShaders(EShaderPlatform Platform, bool bApplyCompletedShaderMapForRendering, bool bForceRecompile, bool bSynchronous)
+{
+	FNiagaraShaderMapId NoStaticParametersId;
+	GetShaderMapId(Platform, NoStaticParametersId);
+	return CacheShaders(NoStaticParametersId, Platform, bApplyCompletedShaderMapForRendering, bForceRecompile, bSynchronous);
+}
+
+/**
+* Caches the shaders for this script
+*/
+bool FNiagaraScript::CacheShaders(const FNiagaraShaderMapId& ShaderMapId, EShaderPlatform Platform, bool bApplyCompletedShaderMapForRendering, bool bForceRecompile, bool bSynchronous)
+{
+	bool bSucceeded = false;
+
+	{
+		// Find the script's cached shader map.
+		GameThreadShaderMap = FNiagaraShaderMap::FindId(ShaderMapId, Platform);
+
+		// Attempt to load from the derived data cache if we are uncooked
+		if (!bForceRecompile && (!GameThreadShaderMap || !GameThreadShaderMap->IsComplete(this, true)) && !FPlatformProperties::RequiresCookedData())
+		{
+			FNiagaraShaderMap::LoadFromDerivedDataCache(this, ShaderMapId, Platform, GameThreadShaderMap);
+			if (GameThreadShaderMap && GameThreadShaderMap->IsValid())
+			{
+				UE_LOG(LogTemp, Display, TEXT("Loaded shader %s for Niagara script %s from DDC"), *GameThreadShaderMap->GetFriendlyName(), *GetFriendlyName());
+			}
+			else
+			{
+				UE_LOG(LogTemp, Display, TEXT("Loading shader for Niagara script %s from DDC failed. Shader needs recompile."), *GetFriendlyName());
+			}
+		}
+	}
+
+	bool bAssumeShaderMapIsComplete = false;
+#if UE_BUILD_SHIPPING || UE_BUILD_TEST
+	bAssumeShaderMapIsComplete = FPlatformProperties::RequiresCookedData();
+#endif
+
+	if (GameThreadShaderMap && GameThreadShaderMap->TryToAddToExistingCompilationTask(this))
+	{
+		//FNiagaraShaderMap::ShaderMapsBeingCompiled.Find(GameThreadShaderMap);
+#if DEBUG_INFINITESHADERCOMPILE
+		UE_LOG(LogTemp, Display, TEXT("Found existing compiling shader for Niagara script %s, linking to other GameThreadShaderMap 0x%08X%08X"), *GetFriendlyName(), (int)((int64)(GameThreadShaderMap.GetReference()) >> 32), (int)((int64)(GameThreadShaderMap.GetReference())));
+#endif
+		OutstandingCompileShaderMapIds.AddUnique(GameThreadShaderMap->GetCompilingId());
+		// Reset the shader map so we fall back to CPU sim until the compile finishes.
+		GameThreadShaderMap = nullptr;
+		bSucceeded = true;
+	}
+	else if (bForceRecompile || !GameThreadShaderMap || !(bAssumeShaderMapIsComplete || GameThreadShaderMap->IsComplete(this, false)))
+	{
+		if (FPlatformProperties::RequiresCookedData())
+		{
+			UE_LOG(LogShaders, Log, TEXT("Can't compile %s with cooked content!"), *GetFriendlyName());
+			// Reset the shader map so we fall back to CPU sim
+			GameThreadShaderMap = nullptr;
+		}
+		else
+		{
+			UE_LOG(LogShaders, Log, TEXT("%s cached shader map for script %s, compiling."), GameThreadShaderMap? TEXT("Incomplete") : TEXT("Missing"), *GetFriendlyName());
+
+			// If there's no cached shader map for this script compile a new one.
+			// This is just kicking off the compile, GameThreadShaderMap will not be complete yet
+			bSucceeded = BeginCompileShaderMap(ShaderMapId, Platform, GameThreadShaderMap, bApplyCompletedShaderMapForRendering, bSynchronous);
+
+			if (!bSucceeded)
+			{
+				GameThreadShaderMap = nullptr;
+			}
+		}
+	}
+	else
+	{
+		bSucceeded = true;
+	}
+
+
+	ENQUEUE_UNIQUE_RENDER_COMMAND_TWOPARAMETER(
+		FSetShaderMapOnScriptResources,
+		FNiagaraScript*, Script, this,
+		FNiagaraShaderMap*, LoadedShaderMap, GameThreadShaderMap,
+		{
+			Script->SetRenderingThreadShaderMap(LoadedShaderMap);
+		});
+
+	return bSucceeded;
+}
+
+
+void FNiagaraScript::FinishCompilation()
+{
+	TArray<int32> ShaderMapIdsToFinish;
+	GetShaderMapIDsWithUnfinishedCompilation(ShaderMapIdsToFinish);
+
+	if (ShaderMapIdsToFinish.Num() > 0)
+	{
+		// Block until the shader maps that we will save have finished being compiled
+		// NIAGARATODO: implement when async compile works
+		GNiagaraShaderCompilationManager.FinishCompilation(*GetFriendlyName(), ShaderMapIdsToFinish);
+	}
+}
+
+#endif
+
+NIAGARASHADER_API  FNiagaraShader* FNiagaraScript::GetShader() const
+{
+	check(IsInRenderingThread());
+	if (!GIsEditor || RenderingThreadShaderMap /*&& RenderingThreadShaderMap->IsComplete(this, true)*/)
+	{
+		return RenderingThreadShaderMap->GetShader<FNiagaraShader>();
+	}
+	return nullptr;
+};
+
+NIAGARASHADER_API  FNiagaraShader* FNiagaraScript::GetShaderGameThread() const
+{
+	return GameThreadShaderMap->GetShader<FNiagaraShader>();
+};
+
+
+void FNiagaraScript::GetShaderMapIDsWithUnfinishedCompilation(TArray<int32>& ShaderMapIds)
+{
+	// Build an array of the shader map Id's are not finished compiling.
+	if (GameThreadShaderMap && GameThreadShaderMap.IsValid() && !GameThreadShaderMap->IsCompilationFinalized())
+	{
+		ShaderMapIds.Add(GameThreadShaderMap->GetCompilingId());
+	}
+	else if (OutstandingCompileShaderMapIds.Num() != 0)
+	{
+		ShaderMapIds.Append(OutstandingCompileShaderMapIds);
+	}
+}
+
+#if WITH_EDITOR
+
+/**
+* Compiles this script for Platform, storing the result in OutShaderMap
+*
+* @param ShaderMapId - the set of static parameters to compile
+* @param Platform - the platform to compile for
+* @param OutShaderMap - the shader map to compile
+* @return - true if compile succeeded or was not necessary (shader map for ShaderMapId was found and was complete)
+*/
+bool FNiagaraScript::BeginCompileShaderMap(
+	const FNiagaraShaderMapId& ShaderMapId,
+	EShaderPlatform Platform,
+	TRefCountPtr<FNiagaraShaderMap>& OutShaderMap,
+	bool bApplyCompletedShaderMapForRendering,
+	bool bSynchronous)
+{
+#if WITH_EDITORONLY_DATA
+	bool bSuccess = false;
+
+	STAT(double NiagaraCompileTime = 0);
+
+
+	SCOPE_SECONDS_COUNTER(NiagaraCompileTime);
+
+	// Queue hlsl generation and shader compilation - Unlike materials, we queue this here, and compilation happens from the editor module
+	TRefCountPtr<FNiagaraShaderMap> NewShaderMap = new FNiagaraShaderMap();
+	OutstandingCompileShaderMapIds.AddUnique(NewShaderMap->GetCompilingId());
+	FNiagaraCompilationQueue::Get()->Queue(this, NewShaderMap, ShaderMapId, Platform, bApplyCompletedShaderMapForRendering);
+	if (bSynchronous)
+	{
+		INiagaraShaderModule NiagaraShaderModule = FModuleManager::GetModuleChecked<INiagaraShaderModule>(TEXT("NiagaraShader"));
+		NiagaraShaderModule.ProcessShaderCompilationQueue();
+	}
+
+	// For async compile, set to nullptr so that we fall back to CPU side simulation until shader compile is finished
+	OutShaderMap = nullptr;
+
+	INC_FLOAT_STAT_BY(STAT_ShaderCompiling_NiagaraShaders, (float)NiagaraCompileTime);
+
+	return true;
+#else
+	UE_LOG(LogShaders, Fatal, TEXT("Compiling of shaders in a build without editordata is not supported."));
+	return false;
+#endif
+}
+
+#endif
\ No newline at end of file
diff --git a/Engine/Plugins/FX/Niagara/Source/NiagaraShader/Public/NiagaraShader.h b/Engine/Plugins/FX/Niagara/Source/NiagaraShader/Public/NiagaraShader.h
new file mode 100644
--- /dev/null
+++ b/Engine/Plugins/FX/Niagara/Source/NiagaraShader/Public/NiagaraShader.h
@@ -0,0 +1,316 @@
+// Copyright 1998-2017 Epic Games, Inc. All Rights Reserved.
+
+/*=============================================================================
+	MaterialShader.h: Shader base classes
+=============================================================================*/
+
+#pragma once
+
+//#include "HAL/IConsoleManager.h"
+//#include "RHI.h"
+#include "ShaderParameters.h"
+#include "SceneView.h"
+#include "Shader.h"
+#include "GlobalShader.h"
+#include "NiagaraShared.h"
+#include "NiagaraShaderType.h"
+#include "SceneRenderTargetParameters.h"
+
+
+
+template<typename TBufferStruct> class TUniformBufferRef;
+
+template<typename ParameterType> 
+struct TUniformParameter
+{
+	int32 Index;
+	ParameterType ShaderParameter;
+	friend FArchive& operator<<(FArchive& Ar,TUniformParameter<ParameterType>& P)
+	{
+		return Ar << P.Index << P.ShaderParameter;
+	}
+};
+
+
+/**
+ * Debug information related to uniform expression sets.
+ */
+/*
+class FDebugUniformExpressionSet
+{
+public:
+	// The number of each type of expression contained in the set. 
+	int32 NumVectorExpressions;
+	int32 NumScalarExpressions;
+	int32 Num2DTextureExpressions;
+	int32 NumCubeTextureExpressions;
+	int32 NumPerFrameScalarExpressions;
+	int32 NumPerFrameVectorExpressions;
+
+	FDebugUniformExpressionSet()
+		: NumVectorExpressions(0)
+		, NumScalarExpressions(0)
+		, Num2DTextureExpressions(0)
+		, NumCubeTextureExpressions(0)
+		, NumPerFrameScalarExpressions(0)
+		, NumPerFrameVectorExpressions(0)
+	{
+	}
+
+	explicit FDebugUniformExpressionSet(const FUniformExpressionSet& InUniformExpressionSet)
+	{
+		InitFromExpressionSet(InUniformExpressionSet);
+	}
+
+	// Initialize from a uniform expression set. 
+	void InitFromExpressionSet(const FUniformExpressionSet& InUniformExpressionSet)
+	{
+		NumVectorExpressions = InUniformExpressionSet.UniformVectorExpressions.Num();
+		NumScalarExpressions = InUniformExpressionSet.UniformScalarExpressions.Num();
+		Num2DTextureExpressions = InUniformExpressionSet.Uniform2DTextureExpressions.Num();
+		NumCubeTextureExpressions = InUniformExpressionSet.UniformCubeTextureExpressions.Num();
+		NumPerFrameScalarExpressions = InUniformExpressionSet.PerFrameUniformScalarExpressions.Num();
+		NumPerFrameVectorExpressions = InUniformExpressionSet.PerFrameUniformVectorExpressions.Num();
+	}
+
+	// Returns true if the number of uniform expressions matches those with which the debug set was initialized.
+	bool Matches(const FUniformExpressionSet& InUniformExpressionSet) const
+	{
+		return NumVectorExpressions == InUniformExpressionSet.UniformVectorExpressions.Num()
+			&& NumScalarExpressions == InUniformExpressionSet.UniformScalarExpressions.Num()
+			&& Num2DTextureExpressions == InUniformExpressionSet.Uniform2DTextureExpressions.Num()
+			&& NumCubeTextureExpressions == InUniformExpressionSet.UniformCubeTextureExpressions.Num()
+			&& NumPerFrameScalarExpressions == InUniformExpressionSet.PerFrameUniformScalarExpressions.Num()
+			&& NumPerFrameVectorExpressions == InUniformExpressionSet.PerFrameUniformVectorExpressions.Num();
+	}
+};
+*/
+
+/** Serialization for debug uniform expression sets. */
+/*
+inline FArchive& operator<<(FArchive& Ar, FDebugUniformExpressionSet& DebugExpressionSet)
+{
+	Ar << DebugExpressionSet.NumVectorExpressions;
+	Ar << DebugExpressionSet.NumScalarExpressions;
+	Ar << DebugExpressionSet.NumPerFrameScalarExpressions;
+	Ar << DebugExpressionSet.NumPerFrameVectorExpressions;
+	Ar << DebugExpressionSet.Num2DTextureExpressions;
+	Ar << DebugExpressionSet.NumCubeTextureExpressions;
+	return Ar;
+}
+*/
+
+/** Base class of all shaders that need material parameters. */
+class RENDERER_API FNiagaraShader : public FShader
+{
+public:
+	DECLARE_SHADER_TYPE(FNiagaraShader, Niagara);
+
+	static FName UniformBufferLayoutName;
+
+	FNiagaraShader() //: DebugUniformExpressionUBLayout(FRHIUniformBufferLayout::Zero)
+	{
+	}
+
+	static bool ShouldCache(EShaderPlatform Platform, const FNiagaraScript* Script)
+	{
+		//@todo - lit materials only 
+		return IsFeatureLevelSupported(Platform, ERHIFeatureLevel::SM5);
+	}
+
+
+	FNiagaraShader(const FNiagaraShaderType::CompiledShaderInitializerType& Initializer);
+
+	typedef void (*ModifyCompilationEnvironmentType)(EShaderPlatform, const FNiagaraScript*, FShaderCompilerEnvironment&);
+
+	static void ModifyCompilationEnvironment(EShaderPlatform Platform, const FNiagaraScript* Script, FShaderCompilerEnvironment& OutEnvironment)
+	{
+	}
+
+	void SetDatainterfaceBufferDescriptors(const TArray< TArray<DIGPUBufferParamDescriptor>> &InBufferDescriptors)
+	{
+		DIBufferDescriptors = InBufferDescriptors;
+	}
+
+//	FUniformBufferRHIParamRef GetParameterCollectionBuffer(const FGuid& Id, const FSceneInterface* SceneInterface) const;
+	/*
+	template<typename ShaderRHIParamRef>
+	FORCEINLINE_DEBUGGABLE void SetViewParameters(FRHICommandList& RHICmdList, const ShaderRHIParamRef ShaderRHI, const FSceneView& View, const TUniformBufferRef<FViewUniformShaderParameters>& ViewUniformBuffer)
+	{
+		const auto& ViewUniformBufferParameter = GetUniformBufferParameter<FViewUniformShaderParameters>();
+		const auto& BuiltinSamplersUBParameter = GetUniformBufferParameter<FBuiltinSamplersParameters>();
+		CheckShaderIsValid();
+		SetUniformBufferParameter(RHICmdList, ShaderRHI, ViewUniformBufferParameter, ViewUniformBuffer);
+#if USE_GBuiltinSamplersUniformBuffer
+		SetUniformBufferParameter(RHICmdList, ShaderRHI, BuiltinSamplersUBParameter, GBuiltinSamplersUniformBuffer.GetUniformBufferRHI());
+#endif
+
+		if (View.bShouldBindInstancedViewUB && View.Family->Views.Num() > 0)
+		{
+			// When drawing the left eye in a stereo scene, copy the right eye view values into the instanced view uniform buffer.
+			const EStereoscopicPass StereoPassIndex = (View.StereoPass != eSSP_FULL) ? eSSP_RIGHT_EYE : eSSP_FULL;
+
+			const FSceneView& InstancedView = View.Family->GetStereoEyeView(StereoPassIndex);
+			const auto& InstancedViewUniformBufferParameter = GetUniformBufferParameter<FInstancedViewUniformShaderParameters>();
+			SetUniformBufferParameter(RHICmdList, ShaderRHI, InstancedViewUniformBufferParameter, InstancedView.ViewUniformBuffer);
+		}
+	}
+	*/
+	static void ModifyCompilationEnvironment(EShaderPlatform Platform, FShaderCompilerEnvironment& OutEnvironment)
+	{
+		FGlobalShader::ModifyCompilationEnvironment(Platform,OutEnvironment);
+	}
+
+	// Bind parameters
+	void BindParams(const FShaderParameterMap &ParameterMap)
+	{
+		FloatInputBufferParam.Bind(ParameterMap, TEXT("InputFloat"));
+		IntInputBufferParam.Bind(ParameterMap, TEXT("InputInt"));
+		FloatOutputBufferParam.Bind(ParameterMap, TEXT("OutputFloat"));
+		IntOutputBufferParam.Bind(ParameterMap, TEXT("OutputInt"));
+		OutputIndexBufferParam.Bind(ParameterMap, TEXT("DataSetIndices"));
+		EmitterTickCounterParam.Bind(ParameterMap, TEXT("EmitterTickCounter"));
+		NumInstancesPerThreadParam.Bind(ParameterMap, TEXT("NumInstancesPerThread"));		
+		NumEventsPerParticleParam.Bind(ParameterMap, TEXT("NumEventsPerParticle"));
+		NumParticlesPerEventParam.Bind(ParameterMap, TEXT("NumParticlesPerEvent"));
+		CopyInstancesBeforeStartParam.Bind(ParameterMap, TEXT("CopyInstancesBeforeStart"));
+		NumInstancesParam.Bind(ParameterMap, TEXT("NumInstances"));
+		StartInstanceParam.Bind(ParameterMap, TEXT("StartInstance"));
+		SimulateStartInstanceParam.Bind(ParameterMap, TEXT("SimulateStartInstance"));
+		GroupStartInstanceParam.Bind(ParameterMap, TEXT("GroupStartInstance"));
+		ComponentBufferSizeReadParam.Bind(ParameterMap, TEXT("ComponentBufferSizeRead"));
+		ComponentBufferSizeWriteParam.Bind(ParameterMap, TEXT("ComponentBufferSizeWrite"));
+		NumThreadGroupsParam.Bind(ParameterMap, TEXT("NumThreadGroups"));
+		EmitterConstantBufferParam.Bind(ParameterMap, TEXT("FEmitterParameters"));
+
+		// params for event buffers
+		// this is horrendous; need to do this in a uniform buffer instead.
+		//
+		for (uint32 i = 0; i < MAX_CONCURRENT_EVENT_DATASETS; i++)
+		{
+			FString VarName = TEXT("WriteDataSetFloat") + FString::FromInt(i+1);
+			EventFloatUAVParams[i].Bind(ParameterMap, *VarName);
+			VarName = TEXT("WriteDataSetInt") + FString::FromInt(i+1);
+			EventIntUAVParams[i].Bind(ParameterMap, *VarName);
+
+			VarName = TEXT("ReadDataSetFloat") + FString::FromInt(i+1);
+			EventFloatSRVParams[i].Bind(ParameterMap, *VarName);
+			VarName = TEXT("ReadDataSetInt") + FString::FromInt(i+1);
+			EventIntSRVParams[i].Bind(ParameterMap, *VarName);
+
+			VarName = TEXT("DSComponentBufferSizeReadFloat") + FString::FromInt(i + 1);
+			EventReadFloatStrideParams[i].Bind(ParameterMap, *VarName);
+			VarName = TEXT("DSComponentBufferSizeWriteFloat") + FString::FromInt(i + 1);
+			EventWriteFloatStrideParams[i].Bind(ParameterMap, *VarName);
+			VarName = TEXT("DSComponentBufferSizeReadInt") + FString::FromInt(i + 1);
+			EventReadIntStrideParams[i].Bind(ParameterMap, *VarName);
+			VarName = TEXT("DSComponentBufferSizeWriteInt") + FString::FromInt(i + 1);
+			EventWriteIntStrideParams[i].Bind(ParameterMap, *VarName);
+		}
+
+		// params for data interface buffers
+		BuildDIBufferParamMap(ParameterMap);
+
+		ensure(FloatOutputBufferParam.IsBound() || IntOutputBufferParam.IsBound());	// we should have at least one output buffer we're writing to
+		ensure(OutputIndexBufferParam.IsBound());		
+		ensure(NumInstancesPerThreadParam.IsBound());	
+		ensure(ComponentBufferSizeWriteParam.IsBound());
+		ensure(NumInstancesPerThreadParam.IsBound());
+		ensure(StartInstanceParam.IsBound());
+	}
+
+	// FShader interface.
+	virtual bool Serialize(FArchive& Ar) override;
+	virtual uint32 GetAllocatedSize() const override;
+	
+	FShaderResourceParameter FloatInputBufferParam;
+	FShaderResourceParameter IntInputBufferParam;
+	FRWShaderParameter FloatOutputBufferParam;
+	FRWShaderParameter IntOutputBufferParam;
+	FRWShaderParameter OutputIndexBufferParam;
+	FShaderUniformBufferParameter EmitterConstantBufferParam;
+	FShaderUniformBufferParameter DataInterfaceUniformBufferParam;
+	FShaderParameter EmitterTickCounterParam;
+	FShaderParameter NumInstancesPerThreadParam;
+	FShaderParameter NumInstancesParam;
+	FShaderParameter NumEventsPerParticleParam;
+	FShaderParameter NumParticlesPerEventParam;
+	FShaderParameter CopyInstancesBeforeStartParam;
+	FShaderParameter StartInstanceParam;
+	FShaderParameter SimulateStartInstanceParam;
+	FShaderParameter GroupStartInstanceParam;
+	FShaderParameter NumThreadGroupsParam;
+	FShaderParameter ComponentBufferSizeReadParam;
+	FShaderParameter ComponentBufferSizeWriteParam;
+	FRWShaderParameter EventIntUAVParams[MAX_CONCURRENT_EVENT_DATASETS];
+	FRWShaderParameter EventFloatUAVParams[MAX_CONCURRENT_EVENT_DATASETS];
+	FShaderResourceParameter EventIntSRVParams[MAX_CONCURRENT_EVENT_DATASETS];
+	FShaderResourceParameter EventFloatSRVParams[MAX_CONCURRENT_EVENT_DATASETS];
+	FShaderParameter EventWriteFloatStrideParams[MAX_CONCURRENT_EVENT_DATASETS];
+	FShaderParameter EventWriteIntStrideParams[MAX_CONCURRENT_EVENT_DATASETS];
+	FShaderParameter EventReadFloatStrideParams[MAX_CONCURRENT_EVENT_DATASETS];
+	FShaderParameter EventReadIntStrideParams[MAX_CONCURRENT_EVENT_DATASETS];
+
+	FShaderResourceParameter *FindDIBufferParam(int32 DataInterfaceIndex, FName InName)
+	{
+		if (NameToDIBufferParamMap.Num() > DataInterfaceIndex)
+		{
+			return NameToDIBufferParamMap[DataInterfaceIndex].Find(InName);
+		}
+
+		return nullptr;
+	}
+
+	void BuildDIBufferParamMap(const FShaderParameterMap &ParameterMap)
+	{
+		for (TArray<DIGPUBufferParamDescriptor> &InterfaceDescs : DIBufferDescriptors)
+		{
+			int32 Idx = NameToDIBufferParamMap.AddDefaulted(1);
+			for (DIGPUBufferParamDescriptor &Desc : InterfaceDescs)
+			{
+				FShaderResourceParameter &Param = NameToDIBufferParamMap[Idx].Add(*Desc.BufferParamName);
+				Param.Bind(ParameterMap, *Desc.BufferParamName);
+				ensure(Param.IsBound());
+			}
+		}
+	}
+
+private:
+	FShaderUniformBufferParameter NiagaraUniformBuffer;
+
+	// buffer descriptors for data interfaces holding names and params for binding
+	TArray< TArray<DIGPUBufferParamDescriptor> > DIBufferDescriptors;
+
+	// buffer descriptors for event data sets holding names and params for binding
+	TArray< TArray<DIGPUBufferParamDescriptor> > EventBufferDescriptors;
+
+	// one map per data interface, mapping buffer names to their params
+	TArray< TMap<FName, FShaderResourceParameter> > NameToDIBufferParamMap;
+
+	/*
+	FDebugUniformExpressionSet	DebugUniformExpressionSet;
+	FRHIUniformBufferLayout		DebugUniformExpressionUBLayout;
+	*/
+	FString						DebugDescription;
+
+	/* OPTODO: ? */
+	/*
+	// If true, cached uniform expressions are allowed.
+	static int32 bAllowCachedUniformExpressions;
+	// Console variable ref to toggle cached uniform expressions.
+	static FAutoConsoleVariableRef CVarAllowCachedUniformExpressions;
+	*/
+
+
+	/*
+#if !(UE_BUILD_TEST || UE_BUILD_SHIPPING || !WITH_EDITOR)
+	void VerifyExpressionAndShaderMaps(const FMaterialRenderProxy* MaterialRenderProxy, const FMaterial& Material, const FUniformExpressionCache* UniformExpressionCache);
+#endif
+	*/
+};
+
+
+class FNiagaraEmitterInstanceShader : public FNiagaraShader
+{
+
+};
diff --git a/Engine/Plugins/FX/Niagara/Source/NiagaraShader/Public/NiagaraShaderCompilationManager.h b/Engine/Plugins/FX/Niagara/Source/NiagaraShader/Public/NiagaraShaderCompilationManager.h
new file mode 100644
--- /dev/null
+++ b/Engine/Plugins/FX/Niagara/Source/NiagaraShader/Public/NiagaraShaderCompilationManager.h
@@ -0,0 +1,136 @@
+#pragma once
+
+#include "NiagaraShared.h"
+
+
+/** Information tracked for each shader compile worker process instance. */
+struct FNiagaraShaderCompileWorkerInfo
+{
+	/** Process handle of the worker app once launched.  Invalid handle means no process. */
+	FProcHandle WorkerProcess;
+
+	/** Tracks whether tasks have been issued to the worker. */
+	bool bIssuedTasksToWorker;
+
+	/** Whether the worker has been launched for this set of tasks. */
+	bool bLaunchedWorker;
+
+	/** Tracks whether all tasks issued to the worker have been received. */
+	bool bComplete;
+
+	/** Time at which the worker started the most recent batch of tasks. */
+	double StartTime;
+
+	/** Jobs that this worker is responsible for compiling. */
+	TArray<class FNiagaraShaderCompileJob*> QueuedJobs;
+
+	FNiagaraShaderCompileWorkerInfo() :
+		bIssuedTasksToWorker(false),
+		bLaunchedWorker(false),
+		bComplete(false),
+		StartTime(0)
+	{
+	}
+
+	// warning: not virtual
+	~FNiagaraShaderCompileWorkerInfo()
+	{
+		if (WorkerProcess.IsValid())
+		{
+			FPlatformProcess::TerminateProc(WorkerProcess);
+			FPlatformProcess::CloseProc(WorkerProcess);
+		}
+	}
+};
+
+/** Stores all of the input and output information used to compile a single shader. */
+class FNiagaraShaderCompileJob 
+{
+public:
+	/** Id of the shader map this shader belongs to. */
+	uint32 Id;
+	/** true if the results of the shader compile have been processed. */
+	bool bFinalized;
+	/** Output of the shader compile */
+	bool bSucceeded;
+	bool bOptimizeForLowLatency;
+	/** Shader type that this shader belongs to, must be valid */
+	FShaderType* ShaderType;
+	/** Input for the shader compile */
+	FShaderCompilerInput Input;
+	FShaderCompilerOutput Output;
+	FString Hlsl;
+	TArray< TArray<DIGPUBufferParamDescriptor> >DIBufferDescriptors;
+
+	FNiagaraShaderCompileJob(uint32 InId, FShaderType* InShaderType, FString InHlsl) :
+		Id(InId),
+		bFinalized(false),
+		bSucceeded(false),
+		bOptimizeForLowLatency(false),
+		ShaderType(InShaderType),
+		Hlsl(InHlsl)
+	{
+	}
+};
+
+
+
+/** Results for a single compiled shader map. */
+struct FNiagaraShaderMapCompileResults
+{
+	FNiagaraShaderMapCompileResults() :
+		NumJobsQueued(0),
+		bAllJobsSucceeded(true),
+		bRecreateComponentRenderStateOnCompletion(false)
+	{}
+
+	int32 NumJobsQueued;
+	bool bAllJobsSucceeded;
+	bool bRecreateComponentRenderStateOnCompletion;
+	TArray<FNiagaraShaderCompileJob*> FinishedJobs;
+};
+
+
+/** Results for a single compiled and finalized shader map. */
+struct FNiagaraShaderMapFinalizeResults : public FNiagaraShaderMapCompileResults
+{
+	/** Tracks finalization progress on this shader map. */
+	int32 FinalizeJobIndex;
+
+	FNiagaraShaderMapFinalizeResults(const FNiagaraShaderMapCompileResults& InCompileResults) :
+		FNiagaraShaderMapCompileResults(InCompileResults),
+		FinalizeJobIndex(0)
+	{}
+};
+
+
+// handles finished shader compile jobs, applying of the shaders to their scripts, and some error handling
+//
+class FNiagaraShaderCompilationManager
+{
+public:
+	FNiagaraShaderCompilationManager();
+
+	NIAGARASHADER_API void Tick(float DeltaSeconds = 0.0f);
+	NIAGARASHADER_API void AddJobs(TArray<FNiagaraShaderCompileJob*> InNewJobs);
+	NIAGARASHADER_API void ProcessAsyncResults();
+
+	void FinishCompilation(const TCHAR* ScriptName, const TArray<int32>& ShaderMapIdsToFinishCompiling);
+
+private:
+	void ProcessCompiledNiagaraShaderMaps(TMap<int32, FNiagaraShaderMapFinalizeResults>& CompiledShaderMaps, float TimeBudget);
+	void RunCompileJobs();
+
+	TArray<FNiagaraShaderCompileJob*> JobQueue;
+
+	/** Map from shader map Id to the compile results for that map, used to gather compiled results. */
+	TMap<int32, FNiagaraShaderMapCompileResults> NiagaraShaderMapJobs;
+
+	/** Map from shader map id to results being finalized.  Used to track shader finalizations over multiple frames. */
+	TMap<int32, FNiagaraShaderMapFinalizeResults> PendingFinalizeNiagaraShaderMaps;
+
+	TArray<struct FNiagaraShaderCompileWorkerInfo*> WorkerInfos;
+};
+
+
+extern NIAGARASHADER_API FNiagaraShaderCompilationManager GNiagaraShaderCompilationManager;
diff --git a/Engine/Plugins/FX/Niagara/Source/NiagaraShader/Public/NiagaraShaderMap.h b/Engine/Plugins/FX/Niagara/Source/NiagaraShader/Public/NiagaraShaderMap.h
new file mode 100644
--- /dev/null
+++ b/Engine/Plugins/FX/Niagara/Source/NiagaraShader/Public/NiagaraShaderMap.h
@@ -0,0 +1,8 @@
+// Copyright 1998-2017 Epic Games, Inc. All Rights Reserved.
+
+/*=============================================================================
+	MaterialShared.h: Shared material definitions.
+=============================================================================*/
+
+#pragma once
+
diff --git a/Engine/Plugins/FX/Niagara/Source/NiagaraShader/Public/NiagaraShaderType.h b/Engine/Plugins/FX/Niagara/Source/NiagaraShader/Public/NiagaraShaderType.h
new file mode 100644
--- /dev/null
+++ b/Engine/Plugins/FX/Niagara/Source/NiagaraShader/Public/NiagaraShaderType.h
@@ -0,0 +1,142 @@
+// Copyright 1998-2017 Epic Games, Inc. All Rights Reserved.
+
+/*=============================================================================
+	NiagaraShaderType.h: Niagara shader type definition.
+=============================================================================*/
+
+#pragma once
+
+#include "CoreMinimal.h"
+#include "Shader.h"
+#include "GlobalShader.h"
+#include "Engine/EngineTypes.h"
+
+/** A macro to implement Niagara shaders. */
+#define IMPLEMENT_NIAGARA_SHADER_TYPE(TemplatePrefix,ShaderClass,SourceFilename,FunctionName,Frequency) \
+	IMPLEMENT_SHADER_TYPE( \
+		TemplatePrefix, \
+		ShaderClass, \
+		SourceFilename, \
+		FunctionName, \
+		Frequency \
+		);
+
+class FNiagaraScript;
+class FShaderCommonCompileJob;
+class FShaderCompileJob;
+class FUniformExpressionSet;
+
+
+/** Called for every Niagara shader to update the appropriate stats. */
+extern void UpdateNiagaraShaderCompilingStats(const FNiagaraScript* Script);
+
+/**
+ * Dump shader stats for a given platform.
+ * 
+ * @param	Platform	Platform to dump stats for.
+ */
+extern ENGINE_API void DumpComputeShaderStats( EShaderPlatform Platform );
+
+
+
+
+/**
+ * A shader meta type for niagara-linked shaders.
+ */
+class FNiagaraShaderType : public FShaderType
+{
+public:
+	struct CompiledShaderInitializerType : FGlobalShaderType::CompiledShaderInitializerType
+	{
+		//const FUniformExpressionSet& UniformExpressionSet;
+		const FString DebugDescription;
+		TArray< TArray<DIGPUBufferParamDescriptor>> DIBufferDescriptors;
+
+		CompiledShaderInitializerType(
+			FShaderType* InType,
+			const FShaderCompilerOutput& CompilerOutput,
+			FShaderResource* InResource,
+			const FSHAHash& InNiagaraShaderMapHash,
+			const FString& InDebugDescription,
+			const TArray< TArray<DIGPUBufferParamDescriptor>> &InDataInterfaceBufferDescriptors
+			)
+		: FGlobalShaderType::CompiledShaderInitializerType(InType,CompilerOutput,InResource, InNiagaraShaderMapHash,nullptr,nullptr)
+		, DebugDescription(InDebugDescription)
+		, DIBufferDescriptors(InDataInterfaceBufferDescriptors)
+		{}
+	};
+	typedef FShader* (*ConstructCompiledType)(const CompiledShaderInitializerType&);
+	typedef bool (*ShouldCacheType)(EShaderPlatform,const FNiagaraScript*);
+	typedef void (*ModifyCompilationEnvironmentType)(EShaderPlatform, const FNiagaraScript*, FShaderCompilerEnvironment&);
+
+	FNiagaraShaderType(
+		const TCHAR* InName,
+		const TCHAR* InSourceFilename,
+		const TCHAR* InFunctionName,
+		uint32 InFrequency,					// ugly - ignored for Niagara shaders but needed for IMPLEMENT_SHADER_TYPE macro magic
+		ConstructSerializedType InConstructSerializedRef,
+		ConstructCompiledType InConstructCompiledRef,
+		ModifyCompilationEnvironmentType InModifyCompilationEnvironmentRef,
+		ShouldCacheType InShouldCacheRef,
+		GetStreamOutElementsType InGetStreamOutElementsRef
+		):
+		FShaderType(EShaderTypeForDynamicCast::Niagara, InName, InSourceFilename, InFunctionName, SF_Compute, InConstructSerializedRef, InGetStreamOutElementsRef),
+		ConstructCompiledRef(InConstructCompiledRef),
+		ShouldCacheRef(InShouldCacheRef),
+		ModifyCompilationEnvironmentRef(InModifyCompilationEnvironmentRef)
+	{}
+
+	/**
+	 * Enqueues a compilation for a new shader of this type.
+	 * @param Script - The script to link the shader with.
+	 */
+	class FNiagaraShaderCompileJob* BeginCompileShader(
+		uint32 ShaderMapId,
+		const FNiagaraScript* Script,
+		FShaderCompilerEnvironment* ScriptEnvironment,
+		EShaderPlatform Platform,
+		TArray<FNiagaraShaderCompileJob*>& NewJobs,
+		FShaderTarget Target,
+		TArray< TArray<DIGPUBufferParamDescriptor> >&InBufferDescriptors
+		);
+
+	/**
+	 * Either creates a new instance of this type or returns an equivalent existing shader.
+	 * @param script - The script to link the shader with.
+	 * @param CurrentJob - Compile job that was enqueued by BeginCompileShader.
+	 */
+	FShader* FinishCompileShader(
+		const FSHAHash& NiagaraShaderMapHash,
+		const FNiagaraShaderCompileJob& CurrentJob,
+		const FString& InDebugDescription
+		);
+
+	/**
+	 * Checks if the shader type should be cached for a particular platform and script.
+	 * @param Platform - The platform to check.
+	 * @param script - The script to check.
+	 * @return True if this shader type should be cached.
+	 */
+	bool ShouldCache(EShaderPlatform Platform,const FNiagaraScript* Script) const
+	{
+		return (*ShouldCacheRef)(Platform, Script);
+	}
+
+protected:
+
+	/**
+	 * Sets up the environment used to compile an instance of this shader type.
+	 * @param Platform - Platform to compile for.
+	 * @param Environment - The shader compile environment that the function modifies.
+	 */
+	void SetupCompileEnvironment(EShaderPlatform Platform, const FNiagaraScript* Script, FShaderCompilerEnvironment& Environment)
+	{
+		// Allow the shader type to modify its compile environment.
+		(*ModifyCompilationEnvironmentRef)(Platform, Script, Environment);
+	}
+
+private:
+	ConstructCompiledType ConstructCompiledRef;
+	ShouldCacheType ShouldCacheRef;
+	ModifyCompilationEnvironmentType ModifyCompilationEnvironmentRef;
+};
