Copying //UE4/Orion-Staging to //UE4/Main (Source: //Orion/Dev-General @ 3606378)

diff --git a/Engine/Plugins/FX/Niagara/Source/NiagaraShader/Public/NiagaraShared.h b/Engine/Plugins/FX/Niagara/Source/NiagaraShader/Public/NiagaraShared.h
new file mode 100644
--- /dev/null
+++ b/Engine/Plugins/FX/Niagara/Source/NiagaraShader/Public/NiagaraShared.h
@@ -0,0 +1,666 @@
+// Copyright 1998-2017 Epic Games, Inc. All Rights Reserved.
+
+/*=============================================================================
+	NiagaraShared.h: Shared Niagara definitions.
+=============================================================================*/
+
+#pragma once
+
+#include "CoreMinimal.h"
+#include "Containers/IndirectArray.h"
+#include "Misc/Guid.h"
+#include "Engine/EngineTypes.h"
+#include "Templates/RefCounting.h"
+#include "Templates/ScopedPointer.h"
+#include "Misc/SecureHash.h"
+#include "RHI.h"
+#include "RenderResource.h"
+#include "RenderingThread.h"
+#include "UniformBuffer.h"
+#include "Shader.h"
+#include "VertexFactory.h"
+#include "SceneTypes.h"
+#include "StaticParameterSet.h"
+#include "Optional.h"
+
+class FNiagaraScript;
+class FNiagaraShaderMap;
+class FNiagaraShader;
+class FNiagaraShaderMapId;
+class UNiagaraScript;
+
+#define MAX_CONCURRENT_EVENT_DATASETS 4
+
+
+// holds meta data for GPU buffers from data interface; the DI defines them and returns them, they're
+// passed to the shader compiler, so we can bind parameters to dynamically generated hlsl vars
+//  this is only necessary for buffer params; everything else is in a constant buffer, which we
+//  just allocate memory for and copy the data into
+//
+class DIGPUBufferParamDescriptor
+{
+public:
+	DIGPUBufferParamDescriptor(FString Name, uint32 Idx)
+		: BufferParamName(Name)
+		, Index(Idx)
+	{}
+	DIGPUBufferParamDescriptor()
+		: BufferParamName(TEXT("Unknown"))
+		, Index(INDEX_NONE)
+	{}
+
+	friend FArchive &operator << (FArchive &Ar, DIGPUBufferParamDescriptor &Desc);
+
+	FString BufferParamName;	// the name of the parameter in the hlsl the DI generates
+	uint32 Index;				// the index, reflecting the order in which GetBufferDefinitionHLSL declares them
+};
+
+/** Stores outputs from the script compile that need to be saved. */
+class FNiagaraComputeShaderCompilationOutput
+{
+public:
+	FNiagaraComputeShaderCompilationOutput()
+	{}
+
+	NIAGARASHADER_API void Serialize(FArchive& Ar)
+	{}
+};
+
+
+/** Contains all the information needed to uniquely identify a FNiagaraShaderMapID. */
+class FNiagaraShaderMapId
+{
+public:
+
+	/**
+	* The GUID of the script this shader represents
+	* Any change to the FNiagaraScript must modify this guid.
+	*/
+	FGuid BaseScriptID;
+
+	/** Feature level that the shader map is going to be compiled for.  */
+	ERHIFeatureLevel::Type FeatureLevel;
+
+	/** TODO: may need this - Guids of any functions the script was dependent on. */
+	//TArray<FGuid> ReferencedFunctions;
+
+	FNiagaraShaderMapId()
+		: BaseScriptID(0, 0, 0, 0)
+		, FeatureLevel(ERHIFeatureLevel::SM5)
+	{ }
+
+	~FNiagaraShaderMapId()
+	{ }
+
+	//ENGINE_API void SetShaderDependencies(const TArray<FShaderType*>& ShaderTypes, const TArray<const FShaderPipelineType*>& ShaderPipelineTypes, const TArray<FVertexFactoryType*>& VFTypes);
+
+	void Serialize(FArchive& Ar);
+
+	friend uint32 GetTypeHash(const FNiagaraShaderMapId& Ref)
+	{
+		return Ref.BaseScriptID.A;
+	}
+
+	SIZE_T GetSizeBytes() const
+	{
+		return sizeof(*this);
+	}
+
+	/** Hashes the script-specific part of this shader map Id. */
+	void GetScriptHash(FSHAHash& OutHash) const;
+
+	/**
+	* Tests this set against another for equality, disregarding override settings.
+	*
+	* @param ReferenceSet	The set to compare against
+	* @return				true if the sets are equal
+	*/
+	bool operator==(const FNiagaraShaderMapId& ReferenceSet) const;
+
+	bool operator!=(const FNiagaraShaderMapId& ReferenceSet) const
+	{
+		return !(*this == ReferenceSet);
+	}
+
+	/** Appends string representations of this Id to a key string. */
+	void AppendKeyString(FString& KeyString) const;
+
+	/** Returns true if the requested shader type is a dependency of this shader map Id. */
+	//bool ContainsShaderType(const FShaderType* ShaderType) const;
+};
+
+
+
+
+
+
+// Runtime code sticks scripts to compile along with their shader map here
+// Niagara Editor ticks in FNiagaraShaderQueueTickable, kicking off compile jobs
+#if WITH_EDITORONLY_DATA
+
+class FNiagaraCompilationQueue
+{
+public:
+	struct NiagaraCompilationQueueItem
+	{
+		FNiagaraScript *Script;
+		TRefCountPtr<FNiagaraShaderMap>ShaderMap;
+		FNiagaraShaderMapId ShaderMapId;
+		EShaderPlatform Platform;
+		bool bApply;
+	};
+	static FNiagaraCompilationQueue *Get()
+	{
+		if (Singleton == nullptr)
+		{
+			Singleton = new FNiagaraCompilationQueue();
+		}
+		return Singleton;
+	}
+
+	TArray<NiagaraCompilationQueueItem> &GetQueue()
+	{
+		return CompilationQueue;
+	}
+
+	void Queue(FNiagaraScript *InScript, TRefCountPtr<FNiagaraShaderMap>InShaderMap, const FNiagaraShaderMapId &MapId, EShaderPlatform InPlatform, bool InApply)
+	{
+		NiagaraCompilationQueueItem NewQueueItem;
+		NewQueueItem.Script = InScript;
+		NewQueueItem.ShaderMap = InShaderMap;
+		NewQueueItem.ShaderMapId = MapId;
+		NewQueueItem.Platform = InPlatform;
+		NewQueueItem.bApply = InApply;
+		CompilationQueue.Add(NewQueueItem);
+	}
+
+private:
+	TArray<NiagaraCompilationQueueItem> CompilationQueue;
+	NIAGARASHADER_API static FNiagaraCompilationQueue *Singleton;
+};
+
+#endif
+
+
+
+
+
+
+/**
+* The set of shaders for a single script.
+*/
+class FNiagaraShaderMap : public TShaderMap<FNiagaraShaderType>, public FDeferredCleanupInterface
+{
+public:
+
+	/**
+	* Finds the shader map for a script.
+	* @param Platform - The platform to lookup for
+	* @return NULL if no cached shader map was found.
+	*/
+	static FNiagaraShaderMap* FindId(const FNiagaraShaderMapId& ShaderMapId, EShaderPlatform Platform);
+
+	/** Flushes the given shader types from any loaded FNiagaraShaderMaps. */
+	static void FlushShaderTypes(TArray<FShaderType*>& ShaderTypesToFlush);
+
+	static void FixupShaderTypes(EShaderPlatform Platform,
+		const TMap<FShaderType*, FString>& ShaderTypeNames);
+
+	/**
+	* Attempts to load the shader map for the given script from the Derived Data Cache.
+	* If InOutShaderMap is valid, attempts to load the individual missing shaders instead.
+	*/
+	static void LoadFromDerivedDataCache(const FNiagaraScript* Script, const FNiagaraShaderMapId& ShaderMapId, EShaderPlatform Platform, TRefCountPtr<FNiagaraShaderMap>& InOutShaderMap);
+
+	FNiagaraShaderMap();
+
+	// Destructor.
+	~FNiagaraShaderMap();
+
+	/**
+	* Compiles the shaders for a script and caches them in this shader map.
+	* @param script - The script to compile shaders for.
+	* @param ShaderMapId - the set of static parameters to compile for
+	* @param Platform - The platform to compile to
+	*/
+	NIAGARASHADER_API void Compile(
+		FNiagaraScript* Script,
+		const FNiagaraShaderMapId& ShaderMapId,
+		TRefCountPtr<FShaderCompilerEnvironment> CompilationEnvironment,
+		const FNiagaraComputeShaderCompilationOutput& InNiagaraCompilationOutput,
+		EShaderPlatform Platform,
+		bool bSynchronousCompile,
+		bool bApplyCompletedShaderMapForRendering
+		);
+
+	/** Sorts the incoming compiled jobs into the appropriate mesh shader maps, and finalizes this shader map so that it can be used for rendering. */
+	bool ProcessCompilationResults(const TArray<class FNiagaraShaderCompileJob*>& InCompilationResults, int32& ResultIndex, float& TimeBudget);
+
+	/**
+	* Checks whether the shader map is missing any shader types necessary for the given script.
+	* @param Script - The Niagara Script which is checked.
+	* @return True if the shader map has all of the shader types necessary.
+	*/
+	bool IsComplete(const FNiagaraScript* Script, bool bSilent);
+
+	/** Attempts to load missing shaders from memory. */
+	void LoadMissingShadersFromMemory(const FNiagaraScript* Script);
+
+	/**
+	* Checks to see if the shader map is already being compiled for another script, and if so
+	* adds the specified script to the list to be applied to once the compile finishes.
+	* @param Script - The Niagara Script we also wish to apply the compiled shader map to.
+	* @return True if the shader map was being compiled and we added Script to the list to be applied.
+	*/
+	bool TryToAddToExistingCompilationTask(FNiagaraScript* Script);
+
+	/** Builds a list of the shaders in a shader map. */
+	NIAGARASHADER_API  void GetShaderList(TMap<FShaderId, FShader*>& OutShaders) const;
+
+	/** Builds a list of the shader pipelines in a shader map. */
+	//ENGINE_API void GetShaderPipelineList(TArray<FShaderPipeline*>& OutShaderPipelines) const;
+
+	/** Registers a niagara shader map in the global map so it can be used by Niagara scripts. */
+	void Register(EShaderPlatform InShaderPlatform);
+
+	// Reference counting.
+	NIAGARASHADER_API  void AddRef();
+	NIAGARASHADER_API  void Release();
+
+	// FDeferredCleanupInterface
+	virtual void FinishCleanup()
+	{
+		bDeletedThroughDeferredCleanup = true;
+		delete this;
+	}
+
+	/**
+	* Removes all entries in the cache with exceptions based on a shader type
+	* @param ShaderType - The shader type to flush
+	*/
+	void FlushShadersByShaderType(FShaderType* ShaderType);
+
+	/** Removes a Script from NiagaraShaderMapsBeingCompiled. */
+	static void RemovePendingScript(FNiagaraScript* Script);
+
+	/** Finds a shader map currently being compiled that was enqueued for the given script. */
+	static const FNiagaraShaderMap* GetShaderMapBeingCompiled(const FNiagaraScript* Script);
+
+	/** Serializes the shader map. */
+	void Serialize(FArchive& Ar, bool bInlineShaderResources = true);
+
+	/** Saves this shader map to the derived data cache. */
+	void SaveToDerivedDataCache();
+
+	/** Registers all shaders that have been loaded in Serialize */
+	virtual void RegisterSerializedShaders() override;
+	virtual void DiscardSerializedShaders() override;
+
+	/** Backs up any FShaders in this shader map to memory through serialization and clears FShader references. */
+	TArray<uint8>* BackupShadersToMemory();
+	/** Recreates FShaders from the passed in memory, handling shader key changes. */
+	void RestoreShadersFromMemory(const TArray<uint8>& ShaderData);
+
+	/** Serializes a shader map to an archive (used with recompiling shaders for a remote console) */
+	NIAGARASHADER_API  static void SaveForRemoteRecompile(FArchive& Ar, const TMap<FString, TArray<TRefCountPtr<FNiagaraShaderMap> > >& CompiledShaderMaps, const TArray<FShaderResourceId>& ClientResourceIds);
+	NIAGARASHADER_API  static void LoadForRemoteRecompile(FArchive& Ar, EShaderPlatform ShaderPlatform, const TArray<FString>& ScriptsForShaderMaps);
+
+	/** Computes the memory used by this shader map without counting the shaders themselves. */
+	uint32 GetSizeBytes() const
+	{
+		return sizeof(*this)
+			+FriendlyName.GetAllocatedSize();
+	}
+
+	/** Returns the maximum number of texture samplers used by any shader in this shader map. */
+	uint32 GetMaxTextureSamplers() const;
+
+
+	// Accessors.
+	const FNiagaraShaderMapId& GetShaderMapId() const { return ShaderMapId; }
+	EShaderPlatform GetShaderPlatform() const	{ return Platform; }
+	const FString& GetFriendlyName() const		{ return FriendlyName; }
+	uint32 GetCompilingId() const				{ return CompilingId; }
+	bool IsCompilationFinalized() const			{ return bCompilationFinalized; }
+	bool CompiledSuccessfully() const			{ return bCompiledSuccessfully; }
+	const FString& GetDebugDescription() const	{ return DebugDescription; }
+
+	bool IsValid() const
+	{
+		return bCompilationFinalized && bCompiledSuccessfully && !bDeletedThroughDeferredCleanup;
+	}
+
+	//const FUniformExpressionSet& GetUniformExpressionSet() const { return NiagaraCompilationOutput.UniformExpressionSet; }
+
+	int32 GetNumRefs() const { return NumRefs; }
+	uint32 GetCompilingId()  { return CompilingId; }
+	static TMap<TRefCountPtr<FNiagaraShaderMap>, TArray<FNiagaraScript*> > &GetInFlightShaderMaps() { return NiagaraShaderMapsBeingCompiled; }
+
+	void SetCompiledSuccessfully(bool bSuccess) { bCompiledSuccessfully = bSuccess; }
+private:
+
+	/**
+	* A global map from a script's ID and static switch set to any shader map cached for that script.
+	* Note: this does not necessarily contain all script shader maps in memory.  Shader maps with the same key can evict each other.
+	* No ref counting needed as these are removed on destruction of the shader map.
+	*/
+	static TMap<FNiagaraShaderMapId, FNiagaraShaderMap*> GIdToNiagaraShaderMap[SP_NumPlatforms];
+
+	/**
+	* All script shader maps in memory.
+	* No ref counting needed as these are removed on destruction of the shader map.
+	*/
+	static TArray<FNiagaraShaderMap*> AllNiagaraShaderMaps;
+
+	/** The script's user friendly name, typically the object name. */
+	FString FriendlyName;
+
+	/** The platform this shader map was compiled with */
+	EShaderPlatform Platform;
+
+	/** The static parameter set that this shader map was compiled with */
+	FNiagaraShaderMapId ShaderMapId;
+
+	/** Shader compilation output */
+	FNiagaraComputeShaderCompilationOutput NiagaraCompilationOutput;
+
+	/** Next value for CompilingId. */
+	static uint32 NextCompilingId;
+
+	/** Tracks resources and their shader maps that need to be compiled but whose compilation is being deferred. */
+	static TMap<TRefCountPtr<FNiagaraShaderMap>, TArray<FNiagaraScript*> > NiagaraShaderMapsBeingCompiled;
+
+	/** Uniquely identifies this shader map during compilation, needed for deferred compilation where shaders from multiple shader maps are compiled together. */
+	uint32 CompilingId;
+
+	mutable int32 NumRefs;
+
+	/** Used to catch errors where the shader map is deleted directly. */
+	bool bDeletedThroughDeferredCleanup;
+
+	/** Indicates whether this shader map has been registered in GIdToNiagaraShaderMap */
+	uint32 bRegistered : 1;
+
+	/**
+	* Indicates whether this shader map has had ProcessCompilationResults called after Compile.
+	* The shader map must not be used on the rendering thread unless bCompilationFinalized is true.
+	*/
+	uint32 bCompilationFinalized : 1;
+
+	uint32 bCompiledSuccessfully : 1;
+
+	/** Indicates whether the shader map should be stored in the shader cache. */
+	uint32 bIsPersistent : 1;
+
+	/** Debug information about how the shader map was compiled. */
+	FString DebugDescription;
+
+	FShader* ProcessCompilationResultsForSingleJob(class FNiagaraShaderCompileJob* SingleJob, const FSHAHash& ShaderMapHash);
+
+	bool IsNiagaraShaderComplete(const FNiagaraScript* Script, const FNiagaraShaderType* ShaderType, bool bSilent);
+
+	/** Initializes OrderedMeshShaderMaps from the contents of MeshShaderMaps. */
+	void InitOrderedMeshShaderMaps();
+
+	friend NIAGARASHADER_API  void DumpNiagaraStats(EShaderPlatform Platform);
+	friend class FShaderCompilingManager;
+};
+
+
+
+/**
+ * FNiagaraScript represents a Niagara script to the shader compilation process
+ */
+class FNiagaraScript
+{
+public:
+
+	/**
+	 * Minimal initialization constructor.
+	 */
+	FNiagaraScript() :
+		GameThreadShaderMap(NULL),
+		RenderingThreadShaderMap(NULL),
+		FeatureLevel(ERHIFeatureLevel::SM4),
+		bLoadedCookedShaderMapId(false)
+	{}
+
+	/**
+	 * Destructor
+	 */
+	NIAGARASHADER_API  virtual ~FNiagaraScript()
+	{}
+
+	/**
+	 * Caches the shaders for this script with no static parameters on the given platform.
+	 * This is used by FNiagaraScript
+	 */
+	NIAGARASHADER_API  bool CacheShaders(EShaderPlatform Platform, bool bApplyCompletedShaderMapForRendering, bool bForceRecompile, bool bSynchronous = false);
+	bool CacheShaders(const FNiagaraShaderMapId& ShaderMapId, EShaderPlatform Platform, bool bApplyCompletedShaderMapForRendering, bool bForceRecompile, bool bSynchronous = false);
+
+	/**
+	 * Should the shader for this script with the given platform, shader type and vertex
+	 * factory type combination be compiled
+	 *
+	 * @param Platform		The platform currently being compiled for
+	 * @param ShaderType	Which shader is being compiled
+	 * @param VertexFactory	Which vertex factory is being compiled (can be NULL)
+	 *
+	 * @return true if the shader should be compiled
+	 */
+	NIAGARASHADER_API  virtual bool ShouldCache(EShaderPlatform Platform, const FShaderType* ShaderType) const;
+
+	/** Serializes the script. */
+	NIAGARASHADER_API  virtual void LegacySerialize(FArchive& Ar);
+
+	NIAGARASHADER_API void SerializeShaderMap(FArchive& Ar);
+	/** Releases this script's shader map.  Must only be called on scripts not exposed to the rendering thread! */
+	NIAGARASHADER_API void ReleaseShaderMap();
+
+	void GetDependentShaderTypes(EShaderPlatform Platform, TArray<FShaderType*>& OutShaderTypes) const;
+	NIAGARASHADER_API  virtual void GetShaderMapId(EShaderPlatform Platform, FNiagaraShaderMapId& OutId) const;
+
+	/**
+	 * Should shaders compiled for this script be saved to disk?
+	 */
+	virtual bool IsPersistent() const { return true; }
+
+	/**
+	* Called when compilation finishes, after the GameThreadShaderMap is set and the render command to set the RenderThreadShaderMap is queued
+	*/
+	virtual void NotifyCompilationFinished() { }
+
+	/**
+	* Cancels all outstanding compilation jobs
+	*/
+	NIAGARASHADER_API  void CancelCompilation();
+
+	/**
+	 * Blocks until compilation has completed. Returns immediately if a compilation is not outstanding.
+	 */
+	NIAGARASHADER_API  void FinishCompilation();
+
+	/**
+	 * Checks if the compilation for this shader is finished
+	 *
+	 * @return returns true if compilation is complete false otherwise
+	 */
+	NIAGARASHADER_API  bool IsCompilationFinished() const;
+
+	/**
+	* Checks if there is a valid GameThreadShaderMap, that is, the script can be run
+	*
+	* @return returns true if there is a GameThreadShaderMap.
+	*/
+	NIAGARASHADER_API  bool HasValidGameThreadShaderMap() const;
+
+
+	// Accessors.
+	const TArray<FString>& GetCompileErrors() const { return CompileErrors; }
+	void SetCompileErrors(const TArray<FString>& InCompileErrors) { CompileErrors = InCompileErrors; }
+
+	ERHIFeatureLevel::Type GetFeatureLevel() const { return FeatureLevel; }
+
+	class FNiagaraShaderMap* GetGameThreadShaderMap() const
+	{
+		checkSlow(IsInGameThread() || IsInAsyncLoadingThread());
+		return GameThreadShaderMap;
+	}
+
+	NIAGARASHADER_API void RegisterShaderMap();
+
+	/** Note: SetRenderingThreadShaderMap must also be called with the same value, but from the rendering thread. */
+	void SetGameThreadShaderMap(FNiagaraShaderMap* InShaderMap)
+	{
+		checkSlow(IsInGameThread() || IsInAsyncLoadingThread());
+		GameThreadShaderMap = InShaderMap;
+	}
+
+	/** Note: SetGameThreadShaderMap must also be called with the same value, but from the game thread. */
+	NIAGARASHADER_API  void SetRenderingThreadShaderMap(FNiagaraShaderMap* InShaderMap);
+
+
+	void SetShaderMap(FNiagaraShaderMap* InShaderMap)
+	{
+		checkSlow(IsInGameThread() || IsInAsyncLoadingThread());
+		GameThreadShaderMap = InShaderMap;
+		bLoadedCookedShaderMapId = true;
+		CookedShaderMapId = InShaderMap->GetShaderMapId();
+
+	}
+
+	NIAGARASHADER_API  class FNiagaraShaderMap* GetRenderingThreadShaderMap() const;
+
+
+	void RemoveOutstandingCompileId(const int32 OldOutstandingCompileShaderMapId)
+	{
+		OutstandingCompileShaderMapIds.Remove(OldOutstandingCompileShaderMapId);
+	}
+
+	NIAGARASHADER_API  virtual void AddReferencedObjects(FReferenceCollector& Collector);
+
+	//virtual const TArray<UTexture*>& GetReferencedTextures() const = 0;
+
+
+	/** Returns true if this script is allowed to make development shaders via the global CVar CompileShadersForDevelopment. */
+	//virtual bool GetAllowDevelopmentShaderCompile()const{ return true; }
+
+	/**
+	* Get user source code for the shader
+	* @param OutSource - generated source code
+	* @param OutHighlightMap - source code highlight list
+	* @return - true on Success
+	*/
+	NIAGARASHADER_API  bool GetScriptHLSLSource(FString& OutSource);
+
+	const FString& GetFriendlyName()	const { return FriendlyName; }
+
+	/** Gets the Guid that represents this script */
+	virtual FGuid GetScriptID() const { return ScriptId; }
+
+	NIAGARASHADER_API void SetScript(UNiagaraScript *InScript, ERHIFeatureLevel::Type InFeatureLevel, FGuid InScriptId, FString InFriendlyName);
+
+	UNiagaraScript *GetGTScript()
+	{
+		return GTScript;
+	}
+
+	void SetCompileErrors(TArray<FString> &InErrors)
+	{
+		CompileErrors = InErrors;
+	}
+
+	FString HlslOutput;
+
+	NIAGARASHADER_API  FNiagaraShader* GetShader() const;
+	NIAGARASHADER_API  FNiagaraShader* GetShaderGameThread() const;
+	
+	void SetDatainterfaceBufferDescriptors(TArray< TArray<DIGPUBufferParamDescriptor> >&InBufferDescriptors)
+	{
+		DIBufferDescriptors = InBufferDescriptors;
+	}
+
+	TArray< TArray<DIGPUBufferParamDescriptor> >& GetDataInterfaceBufferDescriptors()
+	{
+		return DIBufferDescriptors;
+	}
+
+protected:
+
+	// shared code needed for GetUniformScalarParameterExpressions, GetUniformVectorParameterExpressions, GetUniformCubeTextureExpressions..
+	// @return can be 0
+	const FNiagaraShaderMap* GetShaderMapToUse() const;
+
+	/**
+	* Fills the passed array with IDs of shader maps unfinished compilation jobs.
+	*/
+	void GetShaderMapIDsWithUnfinishedCompilation(TArray<int32>& ShaderMapIds);
+
+
+	void SetFeatureLevel(ERHIFeatureLevel::Type InFeatureLevel)
+	{
+		FeatureLevel = InFeatureLevel;
+	}
+
+private:
+	UNiagaraScript* GTScript;
+
+	TArray<FString> CompileErrors;
+
+	/** 
+	 * Game thread tracked shader map, which is ref counted and manages shader map lifetime. 
+	 * The shader map uses deferred deletion so that the rendering thread has a chance to process a release command when the shader map is no longer referenced.
+	 * Code that sets this is responsible for updating RenderingThreadShaderMap in a thread safe way.
+	 * During an async compile, this will be NULL and will not contain the actual shader map until compilation is complete.
+	 */
+	TRefCountPtr<FNiagaraShaderMap> GameThreadShaderMap;
+
+	/** 
+	 * Shader map for this FniagaraScript which is accessible by the rendering thread. 
+	 * This must be updated along with GameThreadShaderMap, but on the rendering thread.
+	 */
+	FNiagaraShaderMap* RenderingThreadShaderMap;
+
+	// data interface buffer descriptors per data interface; these come from the HLSL translators and need to be passed down to the shader for binding
+	TArray< TArray<DIGPUBufferParamDescriptor> > DIBufferDescriptors;
+
+	/** 
+	 * unique identifier
+	 */
+	FGuid ScriptId;
+
+	/** 
+	 * Contains the compiling id of this shader map when it is being compiled asynchronously. 
+	 * This can be used to access the shader map during async compiling, since GameThreadShaderMap will not have been set yet.
+	 */
+	TArray<int32, TInlineAllocator<1> > OutstandingCompileShaderMapIds;
+
+	/** Feature level that this script is representing. */
+	ERHIFeatureLevel::Type FeatureLevel;
+
+	uint32 bLoadedCookedShaderMapId : 1;
+	FNiagaraShaderMapId CookedShaderMapId;
+
+	/**
+	* Compiles this script for Platform, storing the result in OutShaderMap if the compile was synchronous
+	*/
+	bool BeginCompileShaderMap(
+		const FNiagaraShaderMapId& ShaderMapId,
+		EShaderPlatform Platform, 
+		TRefCountPtr<class FNiagaraShaderMap>& OutShaderMap, 
+		bool bApplyCompletedShaderMapForRendering,
+		bool bSynchronous = false);
+
+	/** Populates OutEnvironment with defines needed to compile shaders for this script. */
+	void SetupShaderCompilationEnvironment(
+		EShaderPlatform Platform,
+		FShaderCompilerEnvironment& OutEnvironment
+		) const;
+
+
+	FString FriendlyName;
+
+	friend class FNiagaraShaderMap;
+	friend class FShaderCompilingManager;
+};
+
diff --git a/Engine/Plugins/FX/Niagara/Source/NiagaraVertexFactories/Private/NiagaraMeshVertexFactory.cpp b/Engine/Plugins/FX/Niagara/Source/NiagaraVertexFactories/Private/NiagaraMeshVertexFactory.cpp
new file mode 100644
--- /dev/null
+++ b/Engine/Plugins/FX/Niagara/Source/NiagaraVertexFactories/Private/NiagaraMeshVertexFactory.cpp
@@ -0,0 +1,235 @@
+// Copyright 1998-2017 Epic Games, Inc. All Rights Reserved.
+
+/*=============================================================================
+	ParticleVertexFactory.cpp: Particle vertex factory implementation.
+=============================================================================*/
+
+#include "NiagaraMeshVertexFactory.h"
+#include "ParticleHelper.h"
+#include "ParticleResources.h"
+#include "ShaderParameterUtils.h"
+
+IMPLEMENT_UNIFORM_BUFFER_STRUCT(FNiagaraMeshUniformParameters,TEXT("NiagaraMeshVF"));
+
+// Copyright 1998-2017 Epic Games, Inc. All Rights Reserved.
+
+class FNiagaraMeshVertexFactoryShaderParameters : public FVertexFactoryShaderParameters
+{
+public:
+
+	virtual void Bind(const FShaderParameterMap& ParameterMap) override
+	{
+		PrevTransformBuffer.Bind(ParameterMap, TEXT("PrevTransformBuffer"));
+		NiagaraParticleDataFloat.Bind(ParameterMap, TEXT("NiagaraParticleDataFloat"));
+		NiagaraParticleDataInt.Bind(ParameterMap, TEXT("NiagaraParticleDataInt"));
+		SafeComponentBufferSizeParam.Bind(ParameterMap, TEXT("SafeComponentBufferSize"));
+	}
+
+	virtual void Serialize(FArchive& Ar) override
+	{
+		Ar << PrevTransformBuffer;
+		Ar << NiagaraParticleDataFloat;
+		Ar << NiagaraParticleDataInt;
+		Ar << SafeComponentBufferSizeParam;
+	}
+
+	virtual void SetMesh(FRHICommandList& RHICmdList, FShader* Shader, const FVertexFactory* VertexFactory, const FSceneView& View, const FMeshBatchElement& BatchElement, uint32 DataFlags) const override
+	{
+		const bool bInstanced = GRHISupportsInstancing;
+		FNiagaraMeshVertexFactory* NiagaraMeshVF = (FNiagaraMeshVertexFactory*)VertexFactory;
+		FVertexShaderRHIParamRef VertexShaderRHI = Shader->GetVertexShader();
+		SetUniformBufferParameter(RHICmdList, VertexShaderRHI, Shader->GetUniformBufferParameter<FNiagaraMeshUniformParameters>(), NiagaraMeshVF->GetUniformBuffer());
+
+		SetSRVParameter(RHICmdList, VertexShaderRHI, PrevTransformBuffer, NiagaraMeshVF->GetPreviousTransformBufferSRV());
+
+		SetSRVParameter(RHICmdList, VertexShaderRHI, NiagaraParticleDataFloat, NiagaraMeshVF->GetFloatDataSRV());
+		SetSRVParameter(RHICmdList, VertexShaderRHI, NiagaraParticleDataInt, NiagaraMeshVF->GetIntDataSRV());
+		SetShaderValue(RHICmdList, VertexShaderRHI, SafeComponentBufferSizeParam, NiagaraMeshVF->GetComponentBufferSize());
+	}
+
+private:
+
+	FShaderResourceParameter PrevTransformBuffer;
+
+	FShaderResourceParameter NiagaraParticleDataFloat;
+	FShaderResourceParameter NiagaraParticleDataInt;
+	FShaderParameter SafeComponentBufferSizeParam;
+
+};
+
+
+void FNiagaraMeshVertexFactory::InitRHI()
+{
+	FVertexDeclarationElementList Elements;
+
+	check(GRHISupportsInstancing);
+
+	if (Data.bInitialized)
+	{
+
+		// TODO: get rid of this; need to modify code further up to not skip entire draw when no streams are present here
+			{
+				//checkf(DynamicVertexStride != -1, TEXT("FNiagaraMeshVertexFactory does not have a valid DynamicVertexStride - likely an empty one was made, but SetStrides was not called"));
+				FVertexStream VertexStream;
+				VertexStream.VertexBuffer = NULL;
+				VertexStream.Stride = 0;
+				VertexStream.Offset = 0;
+				Streams.Add(VertexStream);
+
+				// @todo metal: this will need a valid stride when we get to instanced meshes!
+				Elements.Add(FVertexElement(0, Data.TransformComponent[0].Offset, Data.TransformComponent[0].Type, 8, DynamicVertexStride, Data.TransformComponent[0].bUseInstanceIndex));
+				Elements.Add(FVertexElement(0, Data.TransformComponent[1].Offset, Data.TransformComponent[1].Type, 9, DynamicVertexStride, Data.TransformComponent[1].bUseInstanceIndex));
+				Elements.Add(FVertexElement(0, Data.TransformComponent[2].Offset, Data.TransformComponent[2].Type, 10, DynamicVertexStride, Data.TransformComponent[2].bUseInstanceIndex));
+
+				Elements.Add(FVertexElement(0, Data.SubUVs.Offset, Data.SubUVs.Type, 11, DynamicVertexStride, Data.SubUVs.bUseInstanceIndex));
+				Elements.Add(FVertexElement(0, Data.SubUVLerpAndRelTime.Offset, Data.SubUVLerpAndRelTime.Type, 12, DynamicVertexStride, Data.SubUVLerpAndRelTime.bUseInstanceIndex));
+
+				Elements.Add(FVertexElement(0, Data.ParticleColorComponent.Offset, Data.ParticleColorComponent.Type, 14, DynamicVertexStride, Data.ParticleColorComponent.bUseInstanceIndex));
+				Elements.Add(FVertexElement(0, Data.VelocityComponent.Offset, Data.VelocityComponent.Type, 15, DynamicVertexStride, Data.VelocityComponent.bUseInstanceIndex));
+			}
+
+			// Stream 1 - Dynamic parameter
+			{
+				checkf(DynamicParameterVertexStride != -1, TEXT("FNiagaraMeshVertexFactory does not have a valid DynamicParameterVertexStride - likely an empty one was made, but SetStrides was not called"));
+
+				FVertexStream VertexStream;
+				VertexStream.VertexBuffer = NULL;
+				VertexStream.Stride = 0;
+				VertexStream.Offset = 0;
+				Streams.Add(VertexStream);
+
+				Elements.Add(FVertexElement(1, 0, VET_Float4, 13, DynamicParameterVertexStride, true));
+			}
+
+
+		if (Data.PositionComponent.VertexBuffer != NULL)
+		{
+			Elements.Add(AccessStreamComponent(Data.PositionComponent, 0));
+		}
+
+		// only tangent,normal are used by the stream. the binormal is derived in the shader
+		uint8 TangentBasisAttributes[2] = { 1, 2 };
+		for (int32 AxisIndex = 0; AxisIndex < 2; AxisIndex++)
+		{
+			if (Data.TangentBasisComponents[AxisIndex].VertexBuffer != NULL)
+			{
+				Elements.Add(AccessStreamComponent(Data.TangentBasisComponents[AxisIndex], TangentBasisAttributes[AxisIndex]));
+			}
+		}
+
+		// Vertex color
+		if (Data.VertexColorComponent.VertexBuffer != NULL)
+		{
+			Elements.Add(AccessStreamComponent(Data.VertexColorComponent, 3));
+		}
+		else
+		{
+			//If the mesh has no color component, set the null color buffer on a new stream with a stride of 0.
+			//This wastes 4 bytes of bandwidth per vertex, but prevents having to compile out twice the number of vertex factories.
+			FVertexStreamComponent NullColorComponent(&GNullColorVertexBuffer, 0, 0, VET_Color);
+			Elements.Add(AccessStreamComponent(NullColorComponent, 3));
+		}
+
+		if (Data.TextureCoordinates.Num())
+		{
+			const int32 BaseTexCoordAttribute = 4;
+			for (int32 CoordinateIndex = 0; CoordinateIndex < Data.TextureCoordinates.Num(); CoordinateIndex++)
+			{
+				Elements.Add(AccessStreamComponent(
+					Data.TextureCoordinates[CoordinateIndex],
+					BaseTexCoordAttribute + CoordinateIndex
+					));
+			}
+
+			for (int32 CoordinateIndex = Data.TextureCoordinates.Num(); CoordinateIndex < MAX_TEXCOORDS; CoordinateIndex++)
+			{
+				Elements.Add(AccessStreamComponent(
+					Data.TextureCoordinates[Data.TextureCoordinates.Num() - 1],
+					BaseTexCoordAttribute + CoordinateIndex
+					));
+			}
+		}
+
+		//if (Streams.Num() > 0)
+		{
+			InitDeclaration(Elements);
+			check(IsValidRef(GetDeclaration()));
+		}
+	}
+}
+
+void FNiagaraMeshVertexFactory::SetInstanceBuffer(const FVertexBuffer* InstanceBuffer, uint32 StreamOffset, uint32 Stride)
+{
+	Streams[0].VertexBuffer = InstanceBuffer;
+	Streams[0].Offset = StreamOffset;
+	Streams[0].Stride = Stride;
+}
+
+void FNiagaraMeshVertexFactory::SetDynamicParameterBuffer(const FVertexBuffer* InDynamicParameterBuffer, uint32 StreamOffset, uint32 Stride)
+{
+	if (InDynamicParameterBuffer)
+	{
+		Streams[1].VertexBuffer = InDynamicParameterBuffer;
+		Streams[1].Offset = StreamOffset;
+		Streams[1].Stride = Stride;
+	}
+	else
+	{
+		Streams[1].VertexBuffer = &GNullDynamicParameterVertexBuffer;
+		Streams[1].Offset = 0;
+		Streams[1].Stride = 0;
+	}
+}
+
+uint8* FNiagaraMeshVertexFactory::LockPreviousTransformBuffer(uint32 ParticleCount)
+{
+	const static uint32 ElementSize = sizeof(FVector4);
+	const static uint32 ParticleSize = ElementSize * 3;
+	const uint32 AllocationRequest = ParticleCount * ParticleSize;
+
+	check(!PrevTransformBuffer.MappedBuffer);
+
+	if (AllocationRequest > PrevTransformBuffer.NumBytes)
+	{
+		PrevTransformBuffer.Release();
+		PrevTransformBuffer.Initialize(ElementSize, ParticleCount * 3, PF_A32B32G32R32F, BUF_Dynamic);
+	}
+
+	PrevTransformBuffer.Lock();
+
+	return PrevTransformBuffer.MappedBuffer;
+}
+
+void FNiagaraMeshVertexFactory::UnlockPreviousTransformBuffer()
+{
+	check(PrevTransformBuffer.MappedBuffer);
+
+	PrevTransformBuffer.Unlock();
+}
+
+FShaderResourceViewRHIParamRef FNiagaraMeshVertexFactory::GetPreviousTransformBufferSRV() const
+{
+	return PrevTransformBuffer.SRV;
+}
+
+bool FNiagaraMeshVertexFactory::ShouldCache(EShaderPlatform Platform, const class FMaterial* Material, const class FShaderType* ShaderType)
+{
+	return (Material->IsUsedWithNiagaraMeshParticles() || Material->IsSpecialEngineMaterial());
+}
+
+void FNiagaraMeshVertexFactory::SetData(const FDataType& InData)
+{
+	check(IsInRenderingThread());
+	Data = InData;
+	UpdateRHI();
+}
+
+
+FVertexFactoryShaderParameters* FNiagaraMeshVertexFactory::ConstructShaderParameters(EShaderFrequency ShaderFrequency)
+{
+	return ShaderFrequency == SF_Vertex ? new FNiagaraMeshVertexFactoryShaderParameters() : NULL;
+}
+
+IMPLEMENT_VERTEX_FACTORY_TYPE(FNiagaraMeshVertexFactory, "/Engine/Private/NiagaraMeshVertexFactory.ush", true, false, true, false, false);
+IMPLEMENT_VERTEX_FACTORY_TYPE(FNiagaraMeshVertexFactoryEmulatedInstancing, "/Engine/Private/NiagaraMeshVertexFactory.ush", true, false, true, false, false);
+
diff --git a/Engine/Plugins/FX/Niagara/Source/NiagaraVertexFactories/Private/NiagaraRibbonVertexFactory.cpp b/Engine/Plugins/FX/Niagara/Source/NiagaraVertexFactories/Private/NiagaraRibbonVertexFactory.cpp
new file mode 100644
--- /dev/null
+++ b/Engine/Plugins/FX/Niagara/Source/NiagaraVertexFactories/Private/NiagaraRibbonVertexFactory.cpp
@@ -0,0 +1,152 @@
+// Copyright 1998-2017 Epic Games, Inc. All Rights Reserved.
+
+/*=============================================================================
+	ParticleVertexFactory.cpp: Particle vertex factory implementation.
+=============================================================================*/
+
+#include "NiagaraRibbonVertexFactory.h"
+#include "ParticleHelper.h"
+#include "ParticleResources.h"
+#include "ShaderParameterUtils.h"
+
+IMPLEMENT_UNIFORM_BUFFER_STRUCT(FNiagaraRibbonUniformParameters,TEXT("NiagaraRibbonVF"));
+
+/**
+* Shader parameters for the beam/trail vertex factory.
+*/
+class FNiagaraRibbonVertexFactoryShaderParameters : public FVertexFactoryShaderParameters
+{
+public:
+	virtual void Bind(const FShaderParameterMap& ParameterMap) override
+	{
+	}
+
+	virtual void Serialize(FArchive& Ar) override
+	{
+	}
+
+	virtual void SetMesh(FRHICommandList& RHICmdList, FShader* Shader, const FVertexFactory* VertexFactory, const FSceneView& View, const FMeshBatchElement& BatchElement, uint32 DataFlags) const override
+	{
+		FNiagaraRibbonVertexFactory* BeamTrailVF = (FNiagaraRibbonVertexFactory*)VertexFactory;
+		SetUniformBufferParameter(RHICmdList, Shader->GetVertexShader(), Shader->GetUniformBufferParameter<FNiagaraRibbonUniformParameters>(), BeamTrailVF->GetBeamTrailUniformBuffer());
+	}
+};
+
+///////////////////////////////////////////////////////////////////////////////
+/**
+* The particle system beam trail vertex declaration resource type.
+*/
+class FNiagaraRibbonVertexDeclaration : public FRenderResource
+{
+public:
+	FVertexDeclarationRHIRef VertexDeclarationRHI;
+
+	// Destructor.
+	virtual ~FNiagaraRibbonVertexDeclaration() {}
+
+	virtual void FillDeclElements(FVertexDeclarationElementList& Elements, int32& Offset)
+	{
+		uint16 Stride = sizeof(FNiagaraRibbonVertex);
+		/** The stream to read the vertex position from. */
+		Elements.Add(FVertexElement(0, Offset, VET_Float4, 0, Stride));
+		Offset += sizeof(float) * 4;
+		/** The stream to read the vertex old position from. */
+		Elements.Add(FVertexElement(0, Offset, VET_Float3, 1, Stride));
+		Offset += sizeof(float) * 4;
+		/** The stream to read the vertex size/rot/subimage from. */
+		Elements.Add(FVertexElement(0, Offset, VET_Float4, 2, Stride));
+		Offset += sizeof(float) * 4;
+		/** The stream to read the color from.					*/
+		Elements.Add(FVertexElement(0, Offset, VET_Float4, 4, Stride));
+		Offset += sizeof(float) * 4;
+		/** The stream to read the texture coordinates from.	*/
+		Elements.Add(FVertexElement(0, Offset, VET_Float4, 3, Stride));
+		Offset += sizeof(float) * 4;
+
+		/** Dynamic parameters come from a second stream */
+		Elements.Add(FVertexElement(1, 0, VET_Float4, 5, sizeof(FVector4)));
+	}
+
+	virtual void InitDynamicRHI()
+	{
+		FVertexDeclarationElementList Elements;
+		int32	Offset = 0;
+		FillDeclElements(Elements, Offset);
+
+		// Create the vertex declaration for rendering the factory normally.
+		// This is done in InitDynamicRHI instead of InitRHI to allow FParticleBeamTrailVertexFactory::InitRHI
+		// to rely on it being initialized, since InitDynamicRHI is called before InitRHI.
+		VertexDeclarationRHI = RHICreateVertexDeclaration(Elements);
+	}
+
+	virtual void ReleaseDynamicRHI()
+	{
+		VertexDeclarationRHI.SafeRelease();
+	}
+};
+
+/** The simple element vertex declaration. */
+static TGlobalResource<FNiagaraRibbonVertexDeclaration> GNiagaraRibbonVertexDeclaration;
+
+///////////////////////////////////////////////////////////////////////////////
+
+bool FNiagaraRibbonVertexFactory::ShouldCache(EShaderPlatform Platform, const class FMaterial* Material, const class FShaderType* ShaderType)
+{
+	return Material->IsUsedWithNiagaraRibbons() || Material->IsSpecialEngineMaterial();
+}
+
+/**
+* Can be overridden by FVertexFactory subclasses to modify their compile environment just before compilation occurs.
+*/
+void FNiagaraRibbonVertexFactory::ModifyCompilationEnvironment(EShaderPlatform Platform, const class FMaterial* Material, FShaderCompilerEnvironment& OutEnvironment)
+{
+	FParticleVertexFactoryBase::ModifyCompilationEnvironment(Platform, Material, OutEnvironment);
+	OutEnvironment.SetDefine(TEXT("PARTICLE_BEAMTRAIL_FACTORY"), TEXT("1"));
+}
+
+/**
+*	Initialize the Render Hardware Interface for this vertex factory
+*/
+void FNiagaraRibbonVertexFactory::InitRHI()
+{
+	SetDeclaration(GNiagaraRibbonVertexDeclaration.VertexDeclarationRHI);
+
+	FVertexStream* VertexStream = new(Streams) FVertexStream;
+	FVertexStream* DynamicParameterStream = new(Streams) FVertexStream;
+}
+
+FVertexFactoryShaderParameters* FNiagaraRibbonVertexFactory::ConstructShaderParameters(EShaderFrequency ShaderFrequency)
+{
+	return ShaderFrequency == SF_Vertex ? new FNiagaraRibbonVertexFactoryShaderParameters() : NULL;
+}
+
+void FNiagaraRibbonVertexFactory::SetVertexBuffer(const FVertexBuffer* InBuffer, uint32 StreamOffset, uint32 Stride)
+{
+	check(Streams.Num() == 2);
+	FVertexStream& VertexStream = Streams[0];
+	VertexStream.VertexBuffer = InBuffer;
+	VertexStream.Stride = Stride;
+	VertexStream.Offset = StreamOffset;
+}
+
+void FNiagaraRibbonVertexFactory::SetDynamicParameterBuffer(const FVertexBuffer* InDynamicParameterBuffer, uint32 StreamOffset, uint32 Stride)
+{
+	check(Streams.Num() == 2);
+	FVertexStream& DynamicParameterStream = Streams[1];
+	if (InDynamicParameterBuffer)
+	{
+		DynamicParameterStream.VertexBuffer = InDynamicParameterBuffer;
+		DynamicParameterStream.Stride = Stride;
+		DynamicParameterStream.Offset = StreamOffset;
+	}
+	else
+	{
+		DynamicParameterStream.VertexBuffer = &GNullDynamicParameterVertexBuffer;
+		DynamicParameterStream.Stride = 0;
+		DynamicParameterStream.Offset = 0;
+	}
+}
+
+///////////////////////////////////////////////////////////////////////////////
+
+IMPLEMENT_VERTEX_FACTORY_TYPE(FNiagaraRibbonVertexFactory, "/Engine/Private/ParticleBeamTrailVertexFactory.ush", true, false, true, false, false);
diff --git a/Engine/Plugins/FX/Niagara/Source/NiagaraVertexFactories/Private/NiagaraSpriteVertexFactory.cpp b/Engine/Plugins/FX/Niagara/Source/NiagaraVertexFactories/Private/NiagaraSpriteVertexFactory.cpp
new file mode 100644
--- /dev/null
+++ b/Engine/Plugins/FX/Niagara/Source/NiagaraVertexFactories/Private/NiagaraSpriteVertexFactory.cpp
@@ -0,0 +1,342 @@
+// Copyright 1998-2017 Epic Games, Inc. All Rights Reserved.
+
+/*=============================================================================
+	ParticleVertexFactory.cpp: Particle vertex factory implementation.
+=============================================================================*/
+
+#include "NiagaraSpriteVertexFactory.h"
+#include "ParticleHelper.h"
+#include "ParticleResources.h"
+#include "ShaderParameterUtils.h"
+
+IMPLEMENT_UNIFORM_BUFFER_STRUCT(FNiagaraSpriteUniformParameters,TEXT("NiagaraSpriteVF"));
+
+TGlobalResource<FNullDynamicParameterVertexBuffer> GNullNiagaraDynamicParameterVertexBuffer;
+
+class FNiagaraNullSubUVCutoutVertexBuffer : public FVertexBuffer
+{
+public:
+	/**
+	 * Initialize the RHI for this rendering resource
+	 */
+	virtual void InitRHI() override
+	{
+		// create a static vertex buffer
+		FRHIResourceCreateInfo CreateInfo;
+		void* BufferData = nullptr;
+		VertexBufferRHI = RHICreateAndLockVertexBuffer(sizeof(FVector2D) * 4, BUF_Static | BUF_ShaderResource, CreateInfo, BufferData);
+		FMemory::Memzero(BufferData, sizeof(FVector2D) * 4);
+		RHIUnlockVertexBuffer(VertexBufferRHI);
+		
+		VertexBufferSRV = RHICreateShaderResourceView(VertexBufferRHI, sizeof(FVector2D), PF_G32R32F);
+	}
+	
+	virtual void ReleaseRHI() override
+	{
+		VertexBufferSRV.SafeRelease();
+		FVertexBuffer::ReleaseRHI();
+	}
+	
+	FShaderResourceViewRHIRef VertexBufferSRV;
+};
+TGlobalResource<FNiagaraNullSubUVCutoutVertexBuffer> GFNiagaraNullSubUVCutoutVertexBuffer;
+
+/**
+ * Shader parameters for the particle vertex factory.
+ */
+class FNiagaraSpriteVertexFactoryShaderParameters : public FVertexFactoryShaderParameters
+{
+public:
+
+	virtual void Bind(const FShaderParameterMap& ParameterMap) override
+	{
+	}
+
+	virtual void Serialize(FArchive& Ar) override
+	{
+	}
+};
+
+class FNiagaraSpriteVertexFactoryShaderParametersVS : public FNiagaraSpriteVertexFactoryShaderParameters
+{
+public:
+	virtual void Bind(const FShaderParameterMap& ParameterMap) override
+	{
+		NumCutoutVerticesPerFrame.Bind(ParameterMap, TEXT("NumCutoutVerticesPerFrame"));
+		CutoutGeometry.Bind(ParameterMap, TEXT("CutoutGeometry"));
+		NiagaraParticleDataFloat.Bind(ParameterMap, TEXT("NiagaraParticleDataFloat"));
+		NiagaraParticleDataInt.Bind(ParameterMap, TEXT("NiagaraParticleDataInt"));
+		SafeComponentBufferSizeParam.Bind(ParameterMap, TEXT("SafeComponentBufferSize"));
+		UseCustomAlignmentVector.Bind(ParameterMap, TEXT("UseCustomAlignment"));
+		UseVectorAlignment.Bind(ParameterMap, TEXT("UseVectorAlignment"));
+		UseCameraPlaneFacing.Bind(ParameterMap, TEXT("CameraPlaneFacing"));
+	}
+
+	virtual void Serialize(FArchive& Ar) override
+	{
+		Ar << NumCutoutVerticesPerFrame;
+		Ar << CutoutGeometry;
+		Ar << UseCustomAlignmentVector;
+		Ar << UseVectorAlignment;
+		Ar << UseCameraPlaneFacing;
+		Ar << NiagaraParticleDataFloat;
+		Ar << NiagaraParticleDataInt;
+		Ar << SafeComponentBufferSizeParam;
+
+		Ar << PositionOffsetParam;
+		Ar << SizeOffsetParam;
+		Ar << RotationOffsetParam;
+		Ar << SubimgOffsetParam;
+		Ar << ColorOffsetParam;
+	}
+
+	virtual void SetMesh(FRHICommandList& RHICmdList, FShader* Shader,const FVertexFactory* VertexFactory,const FSceneView& View,const FMeshBatchElement& BatchElement,uint32 DataFlags) const override
+	{
+		FNiagaraSpriteVertexFactory* SpriteVF = (FNiagaraSpriteVertexFactory*)VertexFactory;
+		FVertexShaderRHIParamRef VertexShaderRHI = Shader->GetVertexShader();
+		SetUniformBufferParameter(RHICmdList, VertexShaderRHI, Shader->GetUniformBufferParameter<FNiagaraSpriteUniformParameters>(), SpriteVF->GetSpriteUniformBuffer() );
+		
+		SetShaderValue(RHICmdList, VertexShaderRHI, NumCutoutVerticesPerFrame, SpriteVF->GetNumCutoutVerticesPerFrame());
+		FShaderResourceViewRHIParamRef NullSRV = GFNiagaraNullSubUVCutoutVertexBuffer.VertexBufferSRV;
+		SetSRVParameter(RHICmdList, VertexShaderRHI, CutoutGeometry, SpriteVF->GetCutoutGeometrySRV() ? SpriteVF->GetCutoutGeometrySRV() : NullSRV);
+
+		SetShaderValue(RHICmdList, VertexShaderRHI, UseCustomAlignmentVector, SpriteVF->GetCustomAlignment());
+		SetShaderValue(RHICmdList, VertexShaderRHI, UseVectorAlignment, SpriteVF->GetVectorAligned());
+		SetShaderValue(RHICmdList, VertexShaderRHI, UseCameraPlaneFacing, SpriteVF->GetCameraPlaneFacing());
+
+		SetSRVParameter(RHICmdList, VertexShaderRHI, NiagaraParticleDataFloat, SpriteVF->GetFloatDataSRV());
+		SetSRVParameter(RHICmdList, VertexShaderRHI, NiagaraParticleDataInt, SpriteVF->GetIntDataSRV());
+		SetShaderValue(RHICmdList, VertexShaderRHI, SafeComponentBufferSizeParam, SpriteVF->GetComponentBufferSize());
+	}
+
+private:
+	FShaderParameter NumCutoutVerticesPerFrame;
+	FShaderParameter UseCustomAlignmentVector;
+	FShaderParameter UseVectorAlignment;
+	FShaderParameter UseCameraPlaneFacing;
+	FShaderResourceParameter CutoutGeometry;
+	FShaderResourceParameter NiagaraParticleDataFloat;
+	FShaderResourceParameter NiagaraParticleDataInt;
+	FShaderParameter SafeComponentBufferSizeParam;
+
+	FShaderParameter PositionOffsetParam;
+	FShaderParameter SizeOffsetParam;
+	FShaderParameter RotationOffsetParam;
+	FShaderParameter SubimgOffsetParam;
+	FShaderParameter ColorOffsetParam;
+};
+
+class FNiagaraSpriteVertexFactoryShaderParametersPS : public FNiagaraSpriteVertexFactoryShaderParameters
+{
+public:
+
+	virtual void SetMesh(FRHICommandList& RHICmdList, FShader* Shader,const FVertexFactory* VertexFactory,const FSceneView& View,const FMeshBatchElement& BatchElement,uint32 DataFlags) const override
+	{
+		FNiagaraSpriteVertexFactory* SpriteVF = (FNiagaraSpriteVertexFactory*)VertexFactory;
+		FPixelShaderRHIParamRef PixelShaderRHI = Shader->GetPixelShader();
+		SetUniformBufferParameter(RHICmdList, PixelShaderRHI, Shader->GetUniformBufferParameter<FNiagaraSpriteUniformParameters>(), SpriteVF->GetSpriteUniformBuffer() );
+	}
+};
+
+/**
+ * The particle system vertex declaration resource type.
+ */
+class FNiagaraSpriteVertexDeclaration : public FRenderResource
+{
+public:
+
+	FVertexDeclarationRHIRef VertexDeclarationRHI;
+
+	// Constructor.
+	FNiagaraSpriteVertexDeclaration(bool bInInstanced, int32 InNumVertsInInstanceBuffer) :
+		bInstanced(bInInstanced),
+		NumVertsInInstanceBuffer(InNumVertsInInstanceBuffer)
+	{
+
+	}
+
+	// Destructor.
+	virtual ~FNiagaraSpriteVertexDeclaration() {}
+
+	virtual void FillDeclElements(FVertexDeclarationElementList& Elements, int32& Offset)
+	{
+		uint32 InitialStride = sizeof(float) * 2;
+		uint32 PerParticleStride = (sizeof(float) * 4*4 + sizeof(float)*3*2);
+
+		/** The stream to read the texture coordinates from. */
+		check( Offset == 0 );
+		uint32 Stride = bInstanced ? InitialStride : InitialStride + PerParticleStride;
+		Elements.Add(FVertexElement(0, Offset, VET_Float2, 4, Stride, false));
+		Offset += sizeof(float) * 2;
+
+		/** The per-particle streams follow. */
+		if(bInstanced) 
+		{
+			Offset = 0;
+			// update stride
+			Stride = PerParticleStride;
+		}
+
+		/** The stream to read the vertex position from. */
+		Elements.Add(FVertexElement(bInstanced ? 1 : 0, Offset, VET_Float4, 0, Stride, bInstanced));
+		Offset += sizeof(float) * 4;
+		/** The stream to read the vertex old position from. */
+		Elements.Add(FVertexElement(bInstanced ? 1 : 0, Offset, VET_Float4, 1, Stride, bInstanced));
+		Offset += sizeof(float) * 4;
+		/** The stream to read the vertex size/rot/subimage from. */
+		Elements.Add(FVertexElement(bInstanced ? 1 : 0, Offset, VET_Float4, 2, Stride, bInstanced));
+		Offset += sizeof(float) * 4;
+		/** The stream to read the color from.					*/
+		Elements.Add(FVertexElement(bInstanced ? 1 : 0, Offset, VET_Float4, 3, Stride, bInstanced));
+		Offset += sizeof(float) * 4;
+
+		/** The stream to read the custom alignment vector from. */
+		Elements.Add(FVertexElement(bInstanced ? 1 : 0, Offset, VET_Float3, 6, Stride, bInstanced));
+		Offset += sizeof(float) * 3;
+		/** The stream to read the custom facing vector from. */
+		Elements.Add(FVertexElement(bInstanced ? 1 : 0, Offset, VET_Float3, 7, Stride, bInstanced));
+		Offset += sizeof(float) * 3;
+
+		/** The per-particle dynamic parameter stream */
+
+		// The -V519 disables a warning from PVS-Studio's static analyzer. It noticed that offset is assigned
+		// twice before being read. It is probably safer to leave the redundant assignments here to reduce
+		// the chance of an error being introduced if this code is modified.
+		Offset = 0;  //-V519
+		Stride = sizeof(float) * 4;
+		Elements.Add(FVertexElement(bInstanced ? 2 : 1, Offset, VET_Float4, 5, Stride, bInstanced));
+		Offset += sizeof(float) * 4;
+
+	}
+
+	virtual void InitDynamicRHI()
+	{
+		FVertexDeclarationElementList Elements;
+		int32	Offset = 0;
+
+		FillDeclElements(Elements, Offset);
+
+		// Create the vertex declaration for rendering the factory normally.
+		// This is done in InitDynamicRHI instead of InitRHI to allow FParticleSpriteVertexFactory::InitRHI
+		// to rely on it being initialized, since InitDynamicRHI is called before InitRHI.
+		VertexDeclarationRHI = RHICreateVertexDeclaration(Elements);
+	}
+
+	virtual void ReleaseDynamicRHI()
+	{
+		VertexDeclarationRHI.SafeRelease();
+	}
+
+private:
+
+	bool bInstanced;
+	int32 NumVertsInInstanceBuffer;
+};
+
+/** The simple element vertex declaration. */
+static TGlobalResource<FNiagaraSpriteVertexDeclaration> GParticleSpriteVertexDeclarationInstanced(true, 4);
+static TGlobalResource<FNiagaraSpriteVertexDeclaration> GParticleSpriteEightVertexDeclarationInstanced(true, 8);
+static TGlobalResource<FNiagaraSpriteVertexDeclaration> GParticleSpriteVertexDeclarationNonInstanced(false, 4);
+static TGlobalResource<FNiagaraSpriteVertexDeclaration> GParticleSpriteEightVertexDeclarationNonInstanced(false, 8);
+
+inline TGlobalResource<FNiagaraSpriteVertexDeclaration>& GetNiagaraSpriteVertexDeclaration(bool SupportsInstancing, int32 NumVertsInInstanceBuffer)
+{
+	check(NumVertsInInstanceBuffer == 4 || NumVertsInInstanceBuffer == 8);
+	if (SupportsInstancing)
+	{
+		return NumVertsInInstanceBuffer == 4 ? GParticleSpriteVertexDeclarationInstanced : GParticleSpriteEightVertexDeclarationInstanced;
+	}
+	else
+	{
+		return NumVertsInInstanceBuffer == 4 ? GParticleSpriteVertexDeclarationNonInstanced : GParticleSpriteEightVertexDeclarationNonInstanced;
+	}
+}
+
+bool FNiagaraSpriteVertexFactory::ShouldCache(EShaderPlatform Platform, const class FMaterial* Material, const class FShaderType* ShaderType)
+{
+	return Material->IsUsedWithNiagaraSprites() || Material->IsSpecialEngineMaterial();
+}
+
+/**
+ * Can be overridden by FVertexFactory subclasses to modify their compile environment just before compilation occurs.
+ */
+void FNiagaraSpriteVertexFactory::ModifyCompilationEnvironment(EShaderPlatform Platform, const class FMaterial* Material, FShaderCompilerEnvironment& OutEnvironment)
+{
+	FParticleVertexFactoryBase::ModifyCompilationEnvironment(Platform, Material, OutEnvironment);
+
+	// Set a define so we can tell in MaterialTemplate.usf when we are compiling a sprite vertex factory
+	OutEnvironment.SetDefine(TEXT("PARTICLE_SPRITE_FACTORY"),TEXT("1"));
+}
+
+/**
+ *	Initialize the Render Hardware Interface for this vertex factory
+ */
+void FNiagaraSpriteVertexFactory::InitRHI()
+{
+	InitStreams();
+	SetDeclaration(GetNiagaraSpriteVertexDeclaration(GRHISupportsInstancing, NumVertsInInstanceBuffer).VertexDeclarationRHI);
+}
+
+void FNiagaraSpriteVertexFactory::InitStreams()
+{
+    const bool bInstanced = GRHISupportsInstancing;
+
+	check(Streams.Num() == 0);
+	if(bInstanced) 
+	{
+		FVertexStream* TexCoordStream = new(Streams) FVertexStream;
+		TexCoordStream->VertexBuffer = &GParticleTexCoordVertexBuffer;
+		TexCoordStream->Stride = sizeof(FVector2D);
+		TexCoordStream->Offset = 0;
+	}
+	FVertexStream* InstanceStream = new(Streams) FVertexStream;
+	FVertexStream* DynamicParameterStream = new(Streams) FVertexStream;
+}
+
+void FNiagaraSpriteVertexFactory::SetInstanceBuffer(const FVertexBuffer* InInstanceBuffer, uint32 StreamOffset, uint32 Stride, bool bInstanced)
+{
+	check(Streams.Num() == (bInstanced ? 3 : 2));
+	FVertexStream& InstanceStream = Streams[bInstanced ? 1 : 0];
+	InstanceStream.VertexBuffer = InInstanceBuffer;
+	InstanceStream.Stride = Stride;
+	InstanceStream.Offset = StreamOffset;
+}
+
+void FNiagaraSpriteVertexFactory::SetTexCoordBuffer(const FVertexBuffer* InTexCoordBuffer)
+{
+	FVertexStream& TexCoordStream = Streams[0];
+	TexCoordStream.VertexBuffer = InTexCoordBuffer;
+}
+
+void FNiagaraSpriteVertexFactory::SetDynamicParameterBuffer(const FVertexBuffer* InDynamicParameterBuffer, uint32 StreamOffset, uint32 Stride, bool bInstanced)
+{
+	check(Streams.Num() == (bInstanced ? 3 : 2));
+	FVertexStream& DynamicParameterStream = Streams[bInstanced ? 2 : 1];
+	if (InDynamicParameterBuffer)
+	{
+		DynamicParameterStream.VertexBuffer = InDynamicParameterBuffer;
+		DynamicParameterStream.Stride = Stride;
+		DynamicParameterStream.Offset = StreamOffset;
+	}
+	else
+	{
+		DynamicParameterStream.VertexBuffer = &GNullNiagaraDynamicParameterVertexBuffer;
+		DynamicParameterStream.Stride = 0;
+		DynamicParameterStream.Offset = 0;
+	}
+}
+
+FVertexFactoryShaderParameters* FNiagaraSpriteVertexFactory::ConstructShaderParameters(EShaderFrequency ShaderFrequency)
+{
+	if (ShaderFrequency == SF_Vertex)
+	{
+		return new FNiagaraSpriteVertexFactoryShaderParametersVS();
+	}
+	else if (ShaderFrequency == SF_Pixel)
+	{
+		return new FNiagaraSpriteVertexFactoryShaderParametersPS();
+	}
+	return NULL;
+}
+
+IMPLEMENT_VERTEX_FACTORY_TYPE(FNiagaraSpriteVertexFactory,"/Engine/Private/NiagaraSpriteVertexFactory.ush",true,false,true,false,false);
diff --git a/Engine/Plugins/FX/Niagara/Source/NiagaraVertexFactories/Public/NiagaraMeshVertexFactory.h b/Engine/Plugins/FX/Niagara/Source/NiagaraVertexFactories/Public/NiagaraMeshVertexFactory.h
new file mode 100644
--- /dev/null
+++ b/Engine/Plugins/FX/Niagara/Source/NiagaraVertexFactories/Public/NiagaraMeshVertexFactory.h
@@ -0,0 +1,336 @@
+// Copyright 1998-2017 Epic Games, Inc. All Rights Reserved.
+
+/*=============================================================================
+	ParticleVertexFactory.h: Particle vertex factory definitions.
+=============================================================================*/
+
+#pragma once
+
+#include "CoreMinimal.h"
+#include "RenderResource.h"
+#include "UniformBuffer.h"
+#include "NiagaraVertexFactory.h"
+#include "NiagaraDataSet.h"
+#include "SceneView.h"
+#include "Components.h"
+#include "SceneManagement.h"
+
+
+class FMaterial;
+class FVertexBuffer;
+struct FDynamicReadBuffer;
+struct FShaderCompilerEnvironment;
+
+// Per-particle data sent to the GPU.
+struct FNiagaraMeshInstanceVertex
+{
+	/** The color of the particle. */
+	FLinearColor Color;
+
+	/** The instance to world transform of the particle. Translation vector is packed into W components. */
+	FVector4 Transform[3];
+
+	/** The velocity of the particle, XYZ: direction, W: speed. */
+	FVector4 Velocity;
+
+	/** The sub-image texture offsets for the particle. */
+	int16 SubUVParams[4];
+
+	/** The sub-image lerp value for the particle. */
+	float SubUVLerp;
+
+	/** The relative time of the particle. */
+	float RelativeTime;
+};
+
+struct FNiagaraMeshInstanceVertexDynamicParameter
+{
+	/** The dynamic parameter of the particle. */
+	float DynamicValue[4];
+};
+
+struct FNiagaraMeshInstanceVertexPrevTransform
+{
+	FVector4 PrevTransform0;
+	FVector4 PrevTransform1;
+	FVector4 PrevTransform2;
+};
+
+
+
+
+/**
+* Uniform buffer for mesh particle vertex factories.
+*/
+BEGIN_UNIFORM_BUFFER_STRUCT(FNiagaraMeshUniformParameters, NIAGARAVERTEXFACTORIES_API)
+	DECLARE_UNIFORM_BUFFER_STRUCT_MEMBER(FVector4, SubImageSize)
+	DECLARE_UNIFORM_BUFFER_STRUCT_MEMBER(uint32, TexCoordWeightA)
+	DECLARE_UNIFORM_BUFFER_STRUCT_MEMBER(uint32, TexCoordWeightB)
+	DECLARE_UNIFORM_BUFFER_STRUCT_MEMBER(uint32, PrevTransformAvailable)
+	DECLARE_UNIFORM_BUFFER_STRUCT_MEMBER(float, DeltaSeconds)
+	DECLARE_UNIFORM_BUFFER_STRUCT_MEMBER(int, PositionDataOffset)
+	DECLARE_UNIFORM_BUFFER_STRUCT_MEMBER(int, VelocityDataOffset)
+	DECLARE_UNIFORM_BUFFER_STRUCT_MEMBER(int, ColorDataOffset)
+	DECLARE_UNIFORM_BUFFER_STRUCT_MEMBER(int, TransformDataOffset)
+	DECLARE_UNIFORM_BUFFER_STRUCT_MEMBER(int, ScaleDataOffset)
+	DECLARE_UNIFORM_BUFFER_STRUCT_MEMBER(int, SizeDataOffset)
+	DECLARE_UNIFORM_BUFFER_STRUCT_MEMBER(int, MaterialParamDataOffset)
+
+END_UNIFORM_BUFFER_STRUCT(FNiagaraMeshUniformParameters)
+typedef TUniformBufferRef<FNiagaraMeshUniformParameters> FNiagaraMeshUniformBufferRef;
+
+class FNiagaraMeshInstanceVertices;
+
+
+/**
+* Vertex factory for rendering instanced mesh particles with out dynamic parameter support.
+*/
+class NIAGARAVERTEXFACTORIES_API FNiagaraMeshVertexFactory : public FNiagaraVertexFactoryBase
+{
+	DECLARE_VERTEX_FACTORY_TYPE(FNiagaraMeshVertexFactory);
+public:
+
+	// TODO get rid of the unnecessary components here when streams are no longer necessary
+	struct FDataType
+	{
+		/** The stream to read the vertex position from. */
+		FVertexStreamComponent PositionComponent;
+
+		/** The streams to read the tangent basis from. */
+		FVertexStreamComponent TangentBasisComponents[2];
+
+		/** The streams to read the texture coordinates from. */
+		TArray<FVertexStreamComponent, TFixedAllocator<MAX_TEXCOORDS> > TextureCoordinates;
+
+		/** The stream to read the vertex  color from. */
+		FVertexStreamComponent VertexColorComponent;
+
+		/** The stream to read the vertex  color from. */
+		FVertexStreamComponent ParticleColorComponent;
+
+		/** The stream to read the mesh transform from. */
+		FVertexStreamComponent TransformComponent[3];
+
+		/** The stream to read the particle velocity from */
+		FVertexStreamComponent VelocityComponent;
+
+		/** The stream to read SubUV parameters from.. */
+		FVertexStreamComponent SubUVs;
+
+		/** The stream to read SubUV lerp and the particle relative time from */
+		FVertexStreamComponent SubUVLerpAndRelTime;
+
+		/** Flag to mark as initialized */
+		bool bInitialized;
+
+		FDataType()
+			: bInitialized(false)
+		{
+		}
+	};
+
+	class FBatchParametersCPU : public FOneFrameResource
+	{
+	public:
+		const struct FNiagaraMeshInstanceVertex* InstanceBuffer;
+		const struct FNiagaraMeshInstanceVertexDynamicParameter* DynamicParameterBuffer;
+		const struct FNiagaraMeshInstanceVertexPrevTransform* PrevTransformBuffer;
+	};
+
+	/** Default constructor. */
+	FNiagaraMeshVertexFactory(ENiagaraVertexFactoryType InType, ERHIFeatureLevel::Type InFeatureLevel, int32 InDynamicVertexStride, int32 InDynamicParameterVertexStride)
+		: FNiagaraVertexFactoryBase(InType, InFeatureLevel)
+		, DynamicVertexStride(InDynamicVertexStride)
+		, DynamicParameterVertexStride(InDynamicParameterVertexStride)
+		, InstanceVerticesCPU(nullptr)
+	{}
+
+	FNiagaraMeshVertexFactory()
+		: FNiagaraVertexFactoryBase(NVFT_MAX, ERHIFeatureLevel::Num)
+		, DynamicVertexStride(-1)
+		, DynamicParameterVertexStride(-1)
+		, InstanceVerticesCPU(nullptr)
+	{}
+
+	/**
+	* Should we cache the material's shadertype on this platform with this vertex factory?
+	*/
+	static bool ShouldCache(EShaderPlatform Platform, const class FMaterial* Material, const class FShaderType* ShaderType);
+
+
+	/**
+	* Modify compile environment to enable instancing
+	* @param OutEnvironment - shader compile environment to modify
+	*/
+	static void ModifyCompilationEnvironment(EShaderPlatform Platform, const FMaterial* Material, FShaderCompilerEnvironment& OutEnvironment)
+	{
+		FNiagaraVertexFactoryBase::ModifyCompilationEnvironment(Platform, Material, OutEnvironment);
+
+		// Set a define so we can tell in MaterialTemplate.usf when we are compiling a mesh particle vertex factory
+		OutEnvironment.SetDefine(TEXT("NIAGARA_MESH_FACTORY"), TEXT("1"));
+		OutEnvironment.SetDefine(TEXT("NIAGARA_MESH_INSTANCED"), TEXT("1"));
+	}
+
+	void SetParticleData(const FNiagaraDataSet *InDataSet)
+	{
+		DataSet = InDataSet;
+	}
+
+	inline FShaderResourceViewRHIParamRef GetFloatDataSRV() const
+	{
+#if PLATFORM_PS4
+		return DataSet->PrevData().GetGPUBufferFloat()->SRV;
+#else
+		check(IsInRenderingThread());
+		return DataSet->PrevDataRender().GetGPUBufferFloat()->SRV;
+#endif
+	}
+	inline FShaderResourceViewRHIParamRef GetIntDataSRV() const
+	{
+#if PLATFORM_PS4
+		return DataSet->PrevData().GetGPUBufferInt()->SRV;
+#else
+		check(IsInRenderingThread());
+		return DataSet->PrevDataRender().GetGPUBufferInt()->SRV;
+#endif
+	}
+
+	uint32 GetComponentBufferSize()
+	{
+#if PLATFORM_PS4
+		return DataSet->PrevData().GetFloatStride() / sizeof(float);
+#else
+		check(IsInRenderingThread());
+		return DataSet->PrevDataRender().GetFloatStride() / sizeof(float);
+#endif
+	}
+
+	/**
+	* An implementation of the interface used by TSynchronizedResource to update the resource with new data from the game thread.
+	*/
+	void SetData(const FDataType& InData);
+
+	/**
+	* Set the uniform buffer for this vertex factory.
+	*/
+	FORCEINLINE void SetUniformBuffer(const FNiagaraMeshUniformBufferRef& InMeshParticleUniformBuffer)
+	{
+		MeshParticleUniformBuffer = InMeshParticleUniformBuffer;
+	}
+
+	/**
+	* Retrieve the uniform buffer for this vertex factory.
+	*/
+	FORCEINLINE FUniformBufferRHIParamRef GetUniformBuffer()
+	{
+		return MeshParticleUniformBuffer;
+	}
+
+	/**
+	* Update the data strides (MUST HAPPEN BEFORE InitRHI is called)
+	*/
+	void SetStrides(int32 InDynamicVertexStride, int32 InDynamicParameterVertexStride)
+	{
+		DynamicVertexStride = InDynamicVertexStride;
+		DynamicParameterVertexStride = InDynamicParameterVertexStride;
+	}
+
+	/**
+	* Set the source vertex buffer that contains particle instance data.
+	*/
+	void SetInstanceBuffer(const FVertexBuffer* InstanceBuffer, uint32 StreamOffset, uint32 Stride);
+
+	/**
+	* Set the source vertex buffer that contains particle dynamic parameter data.
+	*/
+	void SetDynamicParameterBuffer(const FVertexBuffer* InDynamicParameterBuffer, uint32 StreamOffset, uint32 Stride);
+
+	uint8* LockPreviousTransformBuffer(uint32 ParticleCount);
+	void UnlockPreviousTransformBuffer();
+	FShaderResourceViewRHIParamRef GetPreviousTransformBufferSRV() const;
+
+	/**
+	* Copy the data from another vertex factory
+	* @param Other - factory to copy from
+	*/
+	void Copy(const FNiagaraMeshVertexFactory& Other);
+
+	// FRenderResource interface.
+	virtual void InitRHI() override;
+
+	static bool SupportsTessellationShaders() { return true; }
+
+	static FVertexFactoryShaderParameters* ConstructShaderParameters(EShaderFrequency ShaderFrequency);
+
+	FNiagaraMeshInstanceVertices*& GetInstanceVerticesCPU()
+	{
+		return InstanceVerticesCPU;
+	}
+
+protected:
+	FDataType Data;
+	const FNiagaraDataSet *DataSet;
+	/** Stride information for instanced mesh particles */
+	int32 DynamicVertexStride;
+	int32 DynamicParameterVertexStride;
+
+	/** Uniform buffer with mesh particle parameters. */
+	FUniformBufferRHIParamRef MeshParticleUniformBuffer;
+
+	FDynamicReadBuffer PrevTransformBuffer;
+
+	/** Used to remember this in the case that we reuse the same vertex factory for multiple renders . */
+	FNiagaraMeshInstanceVertices* InstanceVerticesCPU;
+};
+
+
+class NIAGARAVERTEXFACTORIES_API FNiagaraMeshVertexFactoryEmulatedInstancing : public FNiagaraMeshVertexFactory
+{
+	DECLARE_VERTEX_FACTORY_TYPE(FMeshParticleVertexFactoryEmulatedInstancing);
+
+public:
+	FNiagaraMeshVertexFactoryEmulatedInstancing(ENiagaraVertexFactoryType InType, ERHIFeatureLevel::Type InFeatureLevel, int32 InDynamicVertexStride, int32 InDynamicParameterVertexStride)
+		: FNiagaraMeshVertexFactory(InType, InFeatureLevel, InDynamicVertexStride, InDynamicParameterVertexStride)
+	{}
+
+	FNiagaraMeshVertexFactoryEmulatedInstancing()
+		: FNiagaraMeshVertexFactory()
+	{}
+
+	static bool ShouldCache(EShaderPlatform Platform, const class FMaterial* Material, const class FShaderType* ShaderType)
+	{
+		return (Platform == SP_OPENGL_ES2_ANDROID || Platform == SP_OPENGL_ES2_WEBGL) // Those are only platforms that might not support hardware instancing
+			&& FNiagaraMeshVertexFactory::ShouldCache(Platform, Material, ShaderType);
+	}
+
+	static void ModifyCompilationEnvironment(EShaderPlatform Platform, const FMaterial* Material, FShaderCompilerEnvironment& OutEnvironment)
+	{
+		FNiagaraMeshVertexFactory::ModifyCompilationEnvironment(Platform, Material, OutEnvironment);
+
+		OutEnvironment.SetDefine(TEXT("PARTICLE_MESH_INSTANCED"), TEXT("0"));
+	}
+};
+
+inline FNiagaraMeshVertexFactory* ConstructNiagaraMeshVertexFactory()
+{
+	if (GRHISupportsInstancing)
+	{
+		return new FNiagaraMeshVertexFactory();
+	}
+	else
+	{
+		return new FNiagaraMeshVertexFactoryEmulatedInstancing();
+	}
+}
+
+inline FNiagaraMeshVertexFactory* ConstructNiagaraMeshVertexFactory(ENiagaraVertexFactoryType InType, ERHIFeatureLevel::Type InFeatureLevel, int32 InDynamicVertexStride, int32 InDynamicParameterVertexStride)
+{
+	if (GRHISupportsInstancing)
+	{
+		return new FNiagaraMeshVertexFactory(InType, InFeatureLevel, InDynamicVertexStride, InDynamicParameterVertexStride);
+	}
+	else
+	{
+		return new FNiagaraMeshVertexFactoryEmulatedInstancing(InType, InFeatureLevel, InDynamicVertexStride, InDynamicParameterVertexStride);
+	}
+}
diff --git a/Engine/Plugins/FX/Niagara/Source/NiagaraVertexFactories/Public/NiagaraRibbonVertexFactory.h b/Engine/Plugins/FX/Niagara/Source/NiagaraVertexFactories/Public/NiagaraRibbonVertexFactory.h
new file mode 100644
--- /dev/null
+++ b/Engine/Plugins/FX/Niagara/Source/NiagaraVertexFactories/Public/NiagaraRibbonVertexFactory.h
@@ -0,0 +1,158 @@
+// Copyright 1998-2017 Epic Games, Inc. All Rights Reserved.
+
+/*=============================================================================
+	ParticleVertexFactory.h: Particle vertex factory definitions.
+=============================================================================*/
+
+#pragma once
+
+#include "CoreMinimal.h"
+#include "RenderResource.h"
+#include "UniformBuffer.h"
+#include "NiagaraVertexFactory.h"
+#include "SceneView.h"
+
+class FMaterial;
+
+
+
+//	FParticleBeamTrailVertex
+struct FNiagaraRibbonVertex
+{
+	/** The position of the particle. */
+	FVector Position;
+	/** The relative time of the particle. */
+	float RelativeTime;
+	/** The previous position of the particle. */
+	FVector	OldPosition;
+	/** Value that remains constant over the lifetime of a particle. */
+	float ParticleId;
+	/** The size of the particle. */
+	FVector2D Size;
+	/** The rotation of the particle. */
+	float Rotation;
+	/** The sub-image index for the particle. */
+	float SubImageIndex;
+	/** The color of the particle. */
+	FLinearColor Color;
+
+	float			Tex_U;
+	float			Tex_V;
+
+	/** The second UV set for the particle				*/
+	float			Tex_U2;
+	float			Tex_V2;
+};
+
+
+//	FParticleBeamTrailVertexDynamicParameter
+struct FNiagaraRibbonVertexDynamicParameter
+{
+	/** The dynamic parameter of the particle			*/
+	float			DynamicValue[4];
+};
+
+/**
+* Uniform buffer for particle beam/trail vertex factories.
+*/
+BEGIN_UNIFORM_BUFFER_STRUCT(FNiagaraRibbonUniformParameters, NIAGARAVERTEXFACTORIES_API)
+	DECLARE_UNIFORM_BUFFER_STRUCT_MEMBER(FVector4, CameraRight)
+	DECLARE_UNIFORM_BUFFER_STRUCT_MEMBER(FVector4, CameraUp)
+	DECLARE_UNIFORM_BUFFER_STRUCT_MEMBER(FVector4, ScreenAlignment)
+END_UNIFORM_BUFFER_STRUCT(FNiagaraRibbonUniformParameters)
+typedef TUniformBufferRef<FNiagaraRibbonUniformParameters> FNiagaraRibbonUniformBufferRef;
+
+/**
+* Beam/Trail particle vertex factory.
+*/
+class NIAGARAVERTEXFACTORIES_API FNiagaraRibbonVertexFactory : public FNiagaraVertexFactoryBase
+{
+	DECLARE_VERTEX_FACTORY_TYPE(FParticleBeamTrailVertexFactory);
+
+public:
+
+	/** Default constructor. */
+	FNiagaraRibbonVertexFactory(ENiagaraVertexFactoryType InType, ERHIFeatureLevel::Type InFeatureLevel)
+		: FNiagaraVertexFactoryBase(InType, InFeatureLevel)
+		, IndexBuffer(nullptr)
+		, FirstIndex(0)
+		, OutTriangleCount(0)
+	{}
+
+	FNiagaraRibbonVertexFactory()
+		: FNiagaraVertexFactoryBase(NVFT_MAX, ERHIFeatureLevel::Num)
+		, IndexBuffer(nullptr)
+		, FirstIndex(0)
+		, OutTriangleCount(0)
+	{}
+
+	/**
+	* Should we cache the material's shadertype on this platform with this vertex factory?
+	*/
+	static bool ShouldCache(EShaderPlatform Platform, const class FMaterial* Material, const class FShaderType* ShaderType);
+
+	/**
+	* Can be overridden by FVertexFactory subclasses to modify their compile environment just before compilation occurs.
+	*/
+	static void ModifyCompilationEnvironment(EShaderPlatform Platform, const FMaterial* Material, FShaderCompilerEnvironment& OutEnvironment);
+
+	// FRenderResource interface.
+	virtual void InitRHI() override;
+
+	/**
+	* Set the uniform buffer for this vertex factory.
+	*/
+	FORCEINLINE void SetBeamTrailUniformBuffer(FNiagaraRibbonUniformBufferRef InSpriteUniformBuffer)
+	{
+		BeamTrailUniformBuffer = InSpriteUniformBuffer;
+	}
+
+	/**
+	* Retrieve the uniform buffer for this vertex factory.
+	*/
+	FORCEINLINE FNiagaraRibbonUniformBufferRef GetBeamTrailUniformBuffer()
+	{
+		return BeamTrailUniformBuffer;
+	}
+
+	/**
+	* Set the source vertex buffer.
+	*/
+	void SetVertexBuffer(const FVertexBuffer* InBuffer, uint32 StreamOffset, uint32 Stride);
+
+	/**
+	* Set the source vertex buffer that contains particle dynamic parameter data.
+	*/
+	void SetDynamicParameterBuffer(const FVertexBuffer* InDynamicParameterBuffer, uint32 StreamOffset, uint32 Stride);
+
+
+	/**
+	* Construct shader parameters for this type of vertex factory.
+	*/
+	static FVertexFactoryShaderParameters* ConstructShaderParameters(EShaderFrequency ShaderFrequency);
+
+	FIndexBuffer*& GetIndexBuffer()
+	{
+		return IndexBuffer;
+	}
+
+	uint32& GetFirstIndex()
+	{
+		return FirstIndex;
+	}
+
+	int32& GetOutTriangleCount()
+	{
+		return OutTriangleCount;
+	}
+
+private:
+
+	/** Uniform buffer with beam/trail parameters. */
+	FNiagaraRibbonUniformBufferRef BeamTrailUniformBuffer;
+
+	/** Used to hold the index buffer allocation information when we call GDME more than once per frame. */
+	FIndexBuffer* IndexBuffer;
+	uint32 FirstIndex;
+	int32 OutTriangleCount;
+};
\ No newline at end of file
diff --git a/Engine/Plugins/FX/Niagara/Source/NiagaraVertexFactories/Public/NiagaraSpriteVertexFactory.h b/Engine/Plugins/FX/Niagara/Source/NiagaraVertexFactories/Public/NiagaraSpriteVertexFactory.h
new file mode 100644
--- /dev/null
+++ b/Engine/Plugins/FX/Niagara/Source/NiagaraVertexFactories/Public/NiagaraSpriteVertexFactory.h
@@ -0,0 +1,202 @@
+// Copyright 1998-2017 Epic Games, Inc. All Rights Reserved.
+
+/*=============================================================================
+	ParticleVertexFactory.h: Particle vertex factory definitions.
+=============================================================================*/
+
+#pragma once
+
+#include "CoreMinimal.h"
+#include "RenderResource.h"
+#include "UniformBuffer.h"
+#include "NiagaraVertexFactory.h"
+#include "NiagaraDataSet.h"
+#include "SceneView.h"
+
+class FMaterial;
+
+/**
+ * Uniform buffer for particle sprite vertex factories.
+ */
+BEGIN_UNIFORM_BUFFER_STRUCT( FNiagaraSpriteUniformParameters, NIAGARAVERTEXFACTORIES_API)
+	DECLARE_UNIFORM_BUFFER_STRUCT_MEMBER_EX( FVector, CustomFacingVectorMask, EShaderPrecisionModifier::Half)
+	DECLARE_UNIFORM_BUFFER_STRUCT_MEMBER_EX( FVector4, TangentSelector, EShaderPrecisionModifier::Half )
+	DECLARE_UNIFORM_BUFFER_STRUCT_MEMBER_EX( FVector4, NormalsSphereCenter, EShaderPrecisionModifier::Half )
+	DECLARE_UNIFORM_BUFFER_STRUCT_MEMBER_EX( FVector4, NormalsCylinderUnitDirection, EShaderPrecisionModifier::Half )
+	DECLARE_UNIFORM_BUFFER_STRUCT_MEMBER_EX( FVector4, SubImageSize, EShaderPrecisionModifier::Half )
+	DECLARE_UNIFORM_BUFFER_STRUCT_MEMBER_EX( FVector, CameraFacingBlend, EShaderPrecisionModifier::Half )
+	DECLARE_UNIFORM_BUFFER_STRUCT_MEMBER_EX( float, RemoveHMDRoll, EShaderPrecisionModifier::Half )
+	DECLARE_UNIFORM_BUFFER_STRUCT_MEMBER( FVector4, MacroUVParameters )
+	DECLARE_UNIFORM_BUFFER_STRUCT_MEMBER_EX( float, RotationScale, EShaderPrecisionModifier::Half )
+	DECLARE_UNIFORM_BUFFER_STRUCT_MEMBER_EX( float, RotationBias, EShaderPrecisionModifier::Half )
+	DECLARE_UNIFORM_BUFFER_STRUCT_MEMBER_EX( float, NormalsType, EShaderPrecisionModifier::Half )
+	DECLARE_UNIFORM_BUFFER_STRUCT_MEMBER_EX( float, InvDeltaSeconds, EShaderPrecisionModifier::Half )
+	DECLARE_UNIFORM_BUFFER_STRUCT_MEMBER_EX( FVector2D, PivotOffset, EShaderPrecisionModifier::Half )
+	DECLARE_UNIFORM_BUFFER_STRUCT_MEMBER(int, PositionDataOffset)
+	DECLARE_UNIFORM_BUFFER_STRUCT_MEMBER(int, VelocityDataOffset)
+	DECLARE_UNIFORM_BUFFER_STRUCT_MEMBER(int, RotationDataOffset)
+	DECLARE_UNIFORM_BUFFER_STRUCT_MEMBER(int, SizeDataOffset)
+	DECLARE_UNIFORM_BUFFER_STRUCT_MEMBER(int, SubimageDataOffset)
+	DECLARE_UNIFORM_BUFFER_STRUCT_MEMBER(int, ColorDataOffset)
+	DECLARE_UNIFORM_BUFFER_STRUCT_MEMBER(int, FacingOffset)
+	DECLARE_UNIFORM_BUFFER_STRUCT_MEMBER(int, AlignmentOffset)
+	END_UNIFORM_BUFFER_STRUCT(FNiagaraSpriteUniformParameters)
+
+typedef TUniformBufferRef<FNiagaraSpriteUniformParameters> FNiagaraSpriteUniformBufferRef;
+
+/**
+ * Vertex factory for rendering particle sprites.
+ */
+class NIAGARAVERTEXFACTORIES_API FNiagaraSpriteVertexFactory : public FNiagaraVertexFactoryBase
+{
+	DECLARE_VERTEX_FACTORY_TYPE(FNiagaraSpriteVertexFactory);
+
+public:
+
+	/** Default constructor. */
+	FNiagaraSpriteVertexFactory(ENiagaraVertexFactoryType InType, ERHIFeatureLevel::Type InFeatureLevel )
+		: FNiagaraVertexFactoryBase(InType, InFeatureLevel),
+		NumVertsInInstanceBuffer(0),
+		NumCutoutVerticesPerFrame(0),
+		CutoutGeometrySRV(nullptr),
+		bCustomAlignment(false)
+	{}
+
+	FNiagaraSpriteVertexFactory()
+		: FNiagaraVertexFactoryBase(NVFT_MAX, ERHIFeatureLevel::Num),
+		NumVertsInInstanceBuffer(0),
+		NumCutoutVerticesPerFrame(0),
+		CutoutGeometrySRV(nullptr),
+		bCustomAlignment(false)
+	{}
+
+	// FRenderResource interface.
+	virtual void InitRHI() override;
+
+	/**
+	 * Should we cache the material's shadertype on this platform with this vertex factory? 
+	 */
+	static bool ShouldCache(EShaderPlatform Platform, const class FMaterial* Material, const class FShaderType* ShaderType);
+
+	/**
+	 * Can be overridden by FVertexFactory subclasses to modify their compile environment just before compilation occurs.
+	 */
+	static void ModifyCompilationEnvironment(EShaderPlatform Platform, const FMaterial* Material, FShaderCompilerEnvironment& OutEnvironment);
+
+	/**
+	 * Set the source vertex buffer that contains particle instance data.
+	 */
+	void SetInstanceBuffer(const FVertexBuffer* InInstanceBuffer, uint32 StreamOffset, uint32 Stride, bool bInstanced);
+
+	void SetTexCoordBuffer(const FVertexBuffer* InTexCoordBuffer);
+
+	inline void SetNumVertsInInstanceBuffer(int32 InNumVertsInInstanceBuffer)
+	{
+		NumVertsInInstanceBuffer = InNumVertsInInstanceBuffer;
+	}
+
+	/**
+	 * Set the source vertex buffer that contains particle dynamic parameter data.
+	 */
+	void SetDynamicParameterBuffer(const FVertexBuffer* InDynamicParameterBuffer, uint32 StreamOffset, uint32 Stride, bool bInstanced);
+
+	/**
+	 * Set the uniform buffer for this vertex factory.
+	 */
+	FORCEINLINE void SetSpriteUniformBuffer( const FNiagaraSpriteUniformBufferRef& InSpriteUniformBuffer )
+	{
+		SpriteUniformBuffer = InSpriteUniformBuffer;
+	}
+
+	/**
+	 * Retrieve the uniform buffer for this vertex factory.
+	 */
+	FORCEINLINE FUniformBufferRHIParamRef GetSpriteUniformBuffer()
+	{
+		return SpriteUniformBuffer;
+	}
+
+	void SetCutoutParameters(int32 InNumCutoutVerticesPerFrame, FShaderResourceViewRHIParamRef InCutoutGeometrySRV)
+	{
+		NumCutoutVerticesPerFrame = InNumCutoutVerticesPerFrame;
+		CutoutGeometrySRV = InCutoutGeometrySRV;
+	}
+
+	inline int32 GetNumCutoutVerticesPerFrame() const { return NumCutoutVerticesPerFrame; }
+	inline FShaderResourceViewRHIParamRef GetCutoutGeometrySRV() const { return CutoutGeometrySRV; }
+
+	void SetParticleData(const FNiagaraDataSet *InDataSet)
+	{
+		DataSet = InDataSet;
+	}
+
+	inline FShaderResourceViewRHIParamRef GetFloatDataSRV() const 
+	{ 
+		check(!IsInGameThread());
+		return DataSet->PrevDataRender().GetGPUBufferFloat()->SRV; 
+	}
+	inline FShaderResourceViewRHIParamRef GetIntDataSRV() const 
+	{ 
+		check(!IsInGameThread());
+		return DataSet->PrevDataRender().GetGPUBufferInt()->SRV;
+	}
+
+	uint32 GetComponentBufferSize() 
+	{ 
+		check(!IsInGameThread());
+		return DataSet->PrevDataRender().GetFloatStride() / sizeof(float);
+	}
+
+	void SetCustomAlignment(bool bAlign)
+	{
+		bCustomAlignment = bAlign;
+	}
+
+	void SetVectorAligned(bool bAlign)
+	{
+		bVectorAligned = bAlign;
+	}
+
+	bool GetCustomAlignment()
+	{
+		return bCustomAlignment;
+	}
+
+	bool GetVectorAligned()
+	{
+		return bVectorAligned;
+	}
+
+	void SetCameraPlaneFacing(bool inFacing)
+	{
+		bCameraPlaneFacing = inFacing;
+	}
+
+	bool GetCameraPlaneFacing()
+	{
+		return bCameraPlaneFacing;
+	}
+
+	/**
+	 * Construct shader parameters for this type of vertex factory.
+	 */
+	static FVertexFactoryShaderParameters* ConstructShaderParameters(EShaderFrequency ShaderFrequency);
+
+protected:
+	/** Initialize streams for this vertex factory. */
+	void InitStreams();
+
+private:
+
+	int32 NumVertsInInstanceBuffer;
+
+	/** Uniform buffer with sprite paramters. */
+	FUniformBufferRHIParamRef SpriteUniformBuffer;
+
+	int32 NumCutoutVerticesPerFrame;
+	FShaderResourceViewRHIParamRef CutoutGeometrySRV;
+	const FNiagaraDataSet *DataSet;
+	bool bCustomAlignment;	// uses a custom alignment vector from particle data
+	bool bVectorAligned;	// is either velocity or custom vector aligned (not aligned to camera up and right)
+	bool bCameraPlaneFacing;
+};
diff --git a/Engine/Plugins/FX/Niagara/Source/NiagaraVertexFactories/Public/NiagaraVertexFactory.h b/Engine/Plugins/FX/Niagara/Source/NiagaraVertexFactories/Public/NiagaraVertexFactory.h
new file mode 100644
--- /dev/null
+++ b/Engine/Plugins/FX/Niagara/Source/NiagaraVertexFactories/Public/NiagaraVertexFactory.h
@@ -0,0 +1,109 @@
+// Copyright 1998-2017 Epic Games, Inc. All Rights Reserved.
+
+/*=============================================================================
+	ParticleVertexFactory.h: Particle vertex factory definitions.
+=============================================================================*/
+
+#pragma once
+
+#include "CoreMinimal.h"
+#include "RenderResource.h"
+#include "UniformBuffer.h"
+#include "VertexFactory.h"
+#include "SceneView.h"
+
+class FMaterial;
+
+
+/**
+* Enum identifying the type of a particle vertex factory.
+*/
+enum ENiagaraVertexFactoryType
+{
+	NVFT_Sprite,
+	NVFT_Ribbon,
+	NVFT_Mesh,
+	NVFT_MAX
+};
+
+
+
+/**
+* Base class for particle vertex factories.
+*/
+class FNiagaraVertexFactoryBase : public FVertexFactory
+{
+public:
+
+	/** Default constructor. */
+	explicit FNiagaraVertexFactoryBase(ENiagaraVertexFactoryType Type, ERHIFeatureLevel::Type InFeatureLevel)
+		: FVertexFactory(InFeatureLevel)
+		, LastFrameSetup(MAX_uint32)
+		, LastViewFamily(nullptr)
+		, LastView(nullptr)
+		, LastFrameRealTime(-1.0f)
+		, ParticleFactoryType(Type)
+		, bInUse(false)
+	{
+		bNeedsDeclaration = false;
+	}
+
+	static void ModifyCompilationEnvironment(EShaderPlatform Platform, const FMaterial* Material, FShaderCompilerEnvironment& OutEnvironment)
+	{
+		FVertexFactory::ModifyCompilationEnvironment(Platform, Material, OutEnvironment);
+		OutEnvironment.SetDefine(TEXT("NIAGARA_PARTICLE_FACTORY"), TEXT("1"));
+	}
+
+	/** Return the vertex factory type */
+	FORCEINLINE ENiagaraVertexFactoryType GetParticleFactoryType() const
+	{
+		return ParticleFactoryType;
+	}
+
+	inline void SetParticleFactoryType(ENiagaraVertexFactoryType InType)
+	{
+		ParticleFactoryType = InType;
+	}
+
+	/** Specify whether the factory is in use or not. */
+	FORCEINLINE void SetInUse(bool bInInUse)
+	{
+		bInUse = bInInUse;
+	}
+
+	/** Return the vertex factory type */
+	FORCEINLINE bool GetInUse() const
+	{
+		return bInUse;
+	}
+
+	ERHIFeatureLevel::Type GetFeatureLevel() const { check(HasValidFeatureLevel());  return FRenderResource::GetFeatureLevel(); }
+
+	bool CheckAndUpdateLastFrame(const FSceneViewFamily& ViewFamily, const FSceneView *View = nullptr) const
+	{
+		if (LastFrameSetup != MAX_uint32 && (&ViewFamily == LastViewFamily) && (View == LastView) && ViewFamily.FrameNumber == LastFrameSetup && LastFrameRealTime == ViewFamily.CurrentRealTime)
+		{
+			return false;
+		}
+		LastFrameSetup = ViewFamily.FrameNumber;
+		LastFrameRealTime = ViewFamily.CurrentRealTime;
+		LastViewFamily = &ViewFamily;
+		LastView = View;
+		return true;
+	}
+
+private:
+
+	/** Last state where we set this. We only need to setup these once per frame, so detemine same frame by number, time, and view family. */
+	mutable uint32 LastFrameSetup;
+	mutable const FSceneViewFamily *LastViewFamily;
+	mutable const FSceneView *LastView;
+	mutable float LastFrameRealTime;
+
+	/** The type of the vertex factory. */
+	ENiagaraVertexFactoryType ParticleFactoryType;
+
+	/** Whether the vertex factory is in use. */
+	bool bInUse;
+
+};
\ No newline at end of file
diff --git a/Engine/Plugins/Runtime/GameplayAbilities/Source/GameplayAbilities/Private/AbilitySystemComponent.cpp b/Engine/Plugins/Runtime/GameplayAbilities/Source/GameplayAbilities/Private/AbilitySystemComponent.cpp
--- a/Engine/Plugins/Runtime/GameplayAbilities/Source/GameplayAbilities/Private/AbilitySystemComponent.cpp
+++ b/Engine/Plugins/Runtime/GameplayAbilities/Source/GameplayAbilities/Private/AbilitySystemComponent.cpp
@@ -1552,6 +1552,50 @@ void UAbilitySystemComponent::OnRep_ServerDebugString()
 	}
 }
 
+float UAbilitySystemComponent::GetFilteredAttributeValue(const FGameplayAttribute& Attribute, const FGameplayTagRequirements& SourceTags, const FGameplayTagContainer& TargetTags)
+{
+	float AttributeValue = 0.f;
+
+	if (SourceTags.RequireTags.Num() == 0 && SourceTags.IgnoreTags.Num() == 0)
+	{
+		// No qualifiers so we can just read this attribute normally
+		AttributeValue = GetNumericAttribute(Attribute);
+	}
+	else
+	{
+		// Need to capture qualified attributes
+		FGameplayEffectAttributeCaptureDefinition CaptureDef(Attribute.GetUProperty(), EGameplayEffectAttributeCaptureSource::Source, false);
+		FGameplayEffectAttributeCaptureSpec CaptureSpec(CaptureDef);
+
+		CaptureAttributeForGameplayEffect(CaptureSpec);
+
+		// Source Tags
+		static FGameplayTagContainer QuerySourceTags;
+		QuerySourceTags.Reset();
+
+		GetOwnedGameplayTags(QuerySourceTags);
+		QuerySourceTags.AppendTags(SourceTags.RequireTags);
+
+		// Target Tags
+		static FGameplayTagContainer QueryTargetTags;
+		QueryTargetTags.Reset();
+
+		QueryTargetTags.AppendTags(TargetTags);
+
+		FAggregatorEvaluateParameters Params;
+		Params.SourceTags = &QuerySourceTags;
+		Params.TargetTags = &QueryTargetTags;
+		Params.IncludePredictiveMods = true;
+
+		if (CaptureSpec.AttemptCalculateAttributeMagnitude(Params, AttributeValue) == false)
+		{
+			UE_LOG(LogAbilitySystemComponent, Warning, TEXT("Failed to calculate Attribute %s. On: %s"), *Attribute.GetName(), *GetFullName());
+		}
+	}
+
+	return AttributeValue;
+}
+
 bool UAbilitySystemComponent::ServerPrintDebug_RequestWithStrings_Validate(const TArray<FString>& Strings)
 {
 	return true;
diff --git a/Engine/Plugins/Runtime/GameplayAbilities/Source/GameplayAbilities/Public/AbilitySystemComponent.h b/Engine/Plugins/Runtime/GameplayAbilities/Source/GameplayAbilities/Public/AbilitySystemComponent.h
--- a/Engine/Plugins/Runtime/GameplayAbilities/Source/GameplayAbilities/Public/AbilitySystemComponent.h
+++ b/Engine/Plugins/Runtime/GameplayAbilities/Source/GameplayAbilities/Public/AbilitySystemComponent.h
@@ -1031,6 +1031,8 @@ class GAMEPLAYABILITIES_API UAbilitySystemComponent : public UGameplayTasksCompo
 	UFUNCTION()
 	virtual void OnRep_ServerDebugString();
 
+	float GetFilteredAttributeValue(const FGameplayAttribute& Attribute, const FGameplayTagRequirements& SourceTags, const FGameplayTagContainer& TargetTags);
+
 protected:
 
 	/**
diff --git a/Engine/Shaders/Private/NiagaraEmitterInstanceShader.usf b/Engine/Shaders/Private/NiagaraEmitterInstanceShader.usf
new file mode 100644
--- /dev/null
+++ b/Engine/Shaders/Private/NiagaraEmitterInstanceShader.usf
@@ -0,0 +1,240 @@
+// Copyright 1998-2016 Epic Games, Inc. All Rights Reserved.
+
+/*=============================================================================
+NiagaraSimulationShader.usf: 
+=============================================================================*/
+
+#if GPU_SIMULATION
+	#include "/Engine/Public/Platform.ush"
+#endif
+
+#include "Definitions.usf"
+
+
+//TODO: Add common utility code for things such as reading data sources, reading from and writing to events.
+
+#if GPU_SIMULATION
+	static uint3 GDispatchThreadId;
+	uint EmitterTickCounter;
+
+	#define THREADGROUP_SIZE 16
+
+	// 4D random number generator inspired by PCGs (permuted congruential generator)
+	// COPY FROM RANDOM.USF, which we can't easily include here due to a sizable dependency chain
+	uint4 NiagaraRand4DPCG32(int4 p)
+	{
+		// taking a signed int then reinterpreting as unsigned gives good behavior for negatives
+		uint4 v = uint4(p);
+
+		// Linear congruential step.
+		v = v * 1664525u + 1013904223u;
+
+		// swapping low and high bits makes all 32 bits pretty good
+		v = v * (1u << 16u) + (v >> 16u);
+
+		// final shuffle
+		v.x += v.y*v.w;
+		v.y += v.z*v.x;
+		v.z += v.x*v.y;
+		v.w += v.y*v.z;
+		v.x += v.y*v.w;
+		v.y += v.z*v.x;
+		v.z += v.x*v.y;
+		v.w += v.y*v.z;
+
+		return v;
+	}
+
+
+
+	// using Marc's PCG random implementation, sequentially; initially seed with DispatchThreadID and tick counter,
+	// then repeat seeding with the result from the PRNG; should be deterministic
+	// TODO: use this in the VM for CPU sims as well
+	//
+	static uint4 RandomVal = 2308975;
+	float4 rand4()
+	{
+		RandomVal = NiagaraRand4DPCG32( int4(GDispatchThreadId+(RandomVal/57843).xyz, 54121) + int4(EmitterTickCounter.xxxx));
+		return float4(RandomVal.xyzw) / float4(0xFFFFFFFF,0xFFFFFFFF,0xFFFFFFFF,0xFFFFFFFF);
+	}
+
+	float rand(float x)
+	{
+		return rand4().x * x;
+	}
+
+	float2 rand(float2 x)
+	{
+		return float2(rand4().x, rand4().y) * x;
+	}
+
+	float3 rand(float3 x)
+	{
+		return float3(rand4().x, rand4().y, rand4().z) * x;
+	}
+
+	float4 rand(float4 x) 
+	{
+		return rand4() * x;
+	}
+
+	int rand(int x)
+	{
+		RandomVal = NiagaraRand4DPCG32(int4(GDispatchThreadId + (RandomVal / 57843).xyz, 54121) + int4(EmitterTickCounter.xxxx));
+		return RandomVal.x % x+1;
+	}
+
+#else
+
+	//Temporary random hacks
+	float2 rand(float2 x)
+	{
+		return float2(rand(x.x), rand(x.y));
+	}
+
+	float3 rand(float3 x)
+	{
+		return float3(rand(x.x), rand(x.y), rand(x.z));
+	}
+
+	float4 rand(float4 x) 
+	{
+		return float4(rand(x.x), rand(x.y), rand(x.z), rand(x.w));
+	}
+
+	int rand(int x);
+#endif
+
+
+/** Hooks into the stats system. */
+void EnterStatScope(int ScopeIdx);
+void ExitStatScope();
+
+#if VM_SIMULATION
+	float noise(float x);
+	float noise(float2 x);
+	float noise(float3 x);
+
+	//Total hack to get around the cross compiler converting fmod() to "X - (Y * trunc(X/Y))";
+	//On gpu just define these as fmod(x,y)
+	float4 Modulo(float4 x, float4 y);
+	float3 Modulo(float3 x, float3 y);
+	float2 Modulo(float2 x, float2 y);
+	float Modulo(float x, float y);
+
+	/** Returns the index for this particle in the current execution context. On gpu this'll likely be derived from DispatchThreadId */
+	int ExecIndex();
+
+	//Some functions that we use to map to special VM operations for reading in data.
+	//TODO: replace with proper buffer reading capability and use standard hlsl.
+	int AcquireIndex(int DataSetID, bool DoAcquire);
+
+	float InputDataFloat(int DataSetIndex, int RegisterIdx);  //DataSetIndex is 0 for main dataset
+	int InputDataInt(int DataSetIndex, int RegisterIdx);
+	bool InputDataBool(int DataSetIndex, int RegisterIdx);
+
+	float InputDataNoadvanceFloat(int DataSetIndex, int RegisterIdx);  //DataSetIndex is 0 for main dataset
+	int InputDataNoadvanceInt(int DataSetIndex, int RegisterIdx);
+	bool InputDataNoadvanceBool(int DataSetIndex, int RegisterIdx);
+
+	void OutputDataFloat(int DataSetIndex, int RegisterIndex, int InstanceIndex, float Value);
+	void OutputDataInt(int DataSetIndex, int RegisterIndex, int InstanceIndex, int Value);
+	void OutputDataBool(int DataSetIndex, int RegisterIndex, int InstanceIndex, bool Value);
+#endif
+
+#if GPU_SIMULATION
+	uint NumInstancesPerThread;
+	uint StartInstance;
+	uint SimulateStartInstance;
+	uint NumInstances;
+	uint NumThreadGroups;
+	uint GroupStartInstance;
+	uint ComponentBufferSizeRead;
+	uint ComponentBufferSizeWrite;
+	int	NumEventsPerParticle;
+	int NumParticlesPerEvent;
+	uint CopyInstancesBeforeStart;
+	Buffer<float> InputFloat;
+	Buffer<int> InputInt;
+	RWBuffer<int> RWOutputInt;
+	RWBuffer<float> RWOutputFloat;
+	RWBuffer<int> RWDataSetIndices;
+
+	int AcquireIndex(int DataSetID, bool bDoAcquire)
+	{
+		int PrevIdx = -1;
+		if(bDoAcquire == true)
+		{
+			InterlockedAdd(RWDataSetIndices[DataSetID], 1, PrevIdx);
+		}
+		return PrevIdx;
+	}
+
+	float InputDataFloat(int DataSetIndex, int RegisterIdx, int InstanceIdx)
+	{
+		return InputFloat[RegisterIdx*ComponentBufferSizeRead + InstanceIdx];
+	}
+
+	void OutputDataFloat(int DataSetIndex, int RegisterIndex, int InstanceIndex, float Value)
+	{
+		RWOutputFloat[RegisterIndex*ComponentBufferSizeWrite + InstanceIndex + StartInstance] = Value;
+	}
+
+	void OutputDataInt(int DataSetIndex, int RegisterIndex, int InstanceIndex, int Value)
+	{
+		RWOutputInt[RegisterIndex*ComponentBufferSizeWrite + InstanceIndex + StartInstance] = Value;
+	}
+
+	void EnterStatScope(int ID)	{}
+	void ExitStatScope()	{}
+#endif
+
+//Include the simulation shader code generated by the node graph.
+#include "/Engine/Generated/NiagaraEmitterInstance.usf"
+
+
+//The VectorVM shader needs only the generated code but a compute shader requires some wrapping.
+#if GPU_SIMULATION
+
+[numthreads(THREADGROUP_SIZE, 1, 1)]
+void SimulateMainComputeCS(
+	uint3 GroupId : SV_GroupID,
+	uint3 DispatchThreadId : SV_DispatchThreadID,
+	uint3 GroupThreadId : SV_GroupThreadID)
+{
+	GDispatchThreadId = DispatchThreadId;
+
+	//if(CopyInstancesBeforeStart == 1)
+	//{
+	//	StartInstance = 0;
+	//}
+
+	for(uint InstanceIdx = 0; InstanceIdx<NumInstancesPerThread; InstanceIdx++)
+	{
+		uint InstanceID = StartInstance + InstanceIdx*THREADGROUP_SIZE*NumThreadGroups + DispatchThreadId.x;
+
+		if(InstanceID < StartInstance+NumInstances )
+		{
+			// if we have a start instance >0, then we need to be sure that we don't run on instances before that
+			//
+			if(InstanceID >= SimulateStartInstance)
+			{
+				// calculate event index for the current instance; this is primarily for spawn events, where we do a single read from a secondary DataSet for a subset of particles
+				//  for all-particle events, this will always end up being 0 because NumParticlesPerEvent = total particles in the set
+				uint EventIndex = (InstanceID-SimulateStartInstance) / NumParticlesPerEvent;
+
+				SimulateMain(InstanceID, EventIndex);
+			}
+			// event spawn scripts, for example, will need to copy instances before their start
+			//  TODO: should do this in the VM as well; saves an extra copy
+			//
+			//else if(CopyInstancesBeforeStart)
+			//{
+			//	CopyInstance(InstanceID);
+			//}
+
+		}
+	}
+}
+
+#endif//GPU_SIMULATION
