VS2017 fixes (hopefully)

diff --git a/Engine/Plugins/FX/Niagara/Source/Niagara/Classes/NiagaraEmitter.h b/Engine/Plugins/FX/Niagara/Source/Niagara/Classes/NiagaraEmitter.h
--- a/Engine/Plugins/FX/Niagara/Source/Niagara/Classes/NiagaraEmitter.h
+++ b/Engine/Plugins/FX/Niagara/Source/Niagara/Classes/NiagaraEmitter.h
@@ -247,7 +247,7 @@ public:
 	void NIAGARA_API CompileScripts(TArray<ENiagaraScriptCompileStatus>& OutScriptStatuses, TArray<FString>& OutGraphLevelErrorMessages, TArray<FString>& PathNames, TArray<UNiagaraScript*>& Scripts);
 	ENiagaraScriptCompileStatus NIAGARA_API CompileScript(EScriptCompileIndices InScriptToCompile, FString& OutGraphLevelErrorMessages, int32 InSubScriptIdx = 0);
 
-	UNiagaraEmitter* UNiagaraEmitter::MakeRecursiveDeepCopy(UObject* DestOuter) const;
+	UNiagaraEmitter* MakeRecursiveDeepCopy(UObject* DestOuter) const;
 	UNiagaraEmitter* MakeRecursiveDeepCopy(UObject* DestOuter, TMap<const UObject*, UObject*>& ExistingConversions) const;
 
 	void NIAGARA_API GetScripts(TArray<UNiagaraScript*>& OutScripts);
diff --git a/Engine/Plugins/FX/Niagara/Source/NiagaraEditor/Private/SNiagaraGraphPinNumeric.h b/Engine/Plugins/FX/Niagara/Source/NiagaraEditor/Private/SNiagaraGraphPinNumeric.h
--- a/Engine/Plugins/FX/Niagara/Source/NiagaraEditor/Private/SNiagaraGraphPinNumeric.h
+++ b/Engine/Plugins/FX/Niagara/Source/NiagaraEditor/Private/SNiagaraGraphPinNumeric.h
@@ -2,11 +2,15 @@
 
 #pragma once
 
+#include "CoreMinimal.h"
+#include "Widgets/DeclarativeSyntaxSupport.h"
+#include "Widgets/SWidget.h"
 #include "SGraphPin.h"
 
 /** A graph pin widget for displaying niagara numeric inputs and outputs. */
 class SNiagaraGraphPinNumeric : public SGraphPin
 {
+public:
 	SLATE_BEGIN_ARGS(SNiagaraGraphPinNumeric) {}
 	SLATE_END_ARGS()
 
