Fix VS2017 compile errors.

diff --git a/Engine/Plugins/FX/Niagara/Source/NiagaraEditor/Private/TNiagaraGraphPinEditableName.h b/Engine/Plugins/FX/Niagara/Source/NiagaraEditor/Private/TNiagaraGraphPinEditableName.h
--- a/Engine/Plugins/FX/Niagara/Source/NiagaraEditor/Private/TNiagaraGraphPinEditableName.h
+++ b/Engine/Plugins/FX/Niagara/Source/NiagaraEditor/Private/TNiagaraGraphPinEditableName.h
@@ -38,20 +38,20 @@ protected:
 
 	bool OnVerifyTextChanged(const FText& InName, FText& OutErrorMessage) const
 	{
-		UNiagaraNode* ParentNode = Cast<UNiagaraNode>(GraphPinObj->GetOwningNode());
+		UNiagaraNode* ParentNode = Cast<UNiagaraNode>(this->GraphPinObj->GetOwningNode());
 		if (ParentNode)
 		{
-			return ParentNode->VerifyEditablePinName(InName, OutErrorMessage, GraphPinObj);
+			return ParentNode->VerifyEditablePinName(InName, OutErrorMessage, this->GraphPinObj);
 		}
 		return false;
 	}
 
 	void OnTextCommitted(const FText& InText, ETextCommit::Type InCommitType)
 	{
-		UNiagaraNode* ParentNode = Cast<UNiagaraNode>(GraphPinObj->GetOwningNode());
+		UNiagaraNode* ParentNode = Cast<UNiagaraNode>(this->GraphPinObj->GetOwningNode());
 		if (ParentNode != nullptr)
 		{
-			ParentNode->CommitEditablePinName(InText, GraphPinObj);
+			ParentNode->CommitEditablePinName(InText, this->GraphPinObj);
 		}
 	}
 
@@ -67,8 +67,8 @@ protected:
 
 	virtual TSharedRef<SWidget> GetLabelWidget(const FName& InLabelStyle) override
 	{
-		UNiagaraNode* ParentNode = Cast<UNiagaraNode>(GraphPinObj->GetOwningNode());
-		if (ParentNode && ParentNode->IsPinNameEditable(GraphPinObj))
+		UNiagaraNode* ParentNode = Cast<UNiagaraNode>(this->GraphPinObj->GetOwningNode());
+		if (ParentNode && ParentNode->IsPinNameEditable(this->GraphPinObj))
 		{
 			CreatedTextBlock = SNew(SInlineEditableTextBlock)
 				.Style(&FEditorStyle::Get().GetWidgetStyle<FInlineEditableTextBlockStyle>("Graph.Node.InlineEditablePinName"))
@@ -80,7 +80,7 @@ protected:
 
 			TSharedRef<SWidget> EditableTextBlock = CreatedTextBlock.ToSharedRef();
 
-			if (ParentNode->IsPinNameEditableUponCreation(GraphPinObj))
+			if (ParentNode->IsPinNameEditableUponCreation(this->GraphPinObj))
 			{
 				bPendingRename = true;
 			}
