Copying //UE4/Dev-Framework to //UE4/Dev-Main (Source: //UE4/Dev-Framework @ 3716594)

diff --git a/Engine/Plugins/FX/Niagara/Source/NiagaraEditor/Private/NiagaraNodeOutput.cpp b/Engine/Plugins/FX/Niagara/Source/NiagaraEditor/Private/NiagaraNodeOutput.cpp
--- a/Engine/Plugins/FX/Niagara/Source/NiagaraEditor/Private/NiagaraNodeOutput.cpp
+++ b/Engine/Plugins/FX/Niagara/Source/NiagaraEditor/Private/NiagaraNodeOutput.cpp
@@ -41,7 +41,7 @@ void UNiagaraNodeOutput::PostEditChangeProperty(struct FPropertyChangedEvent& Pr
 void UNiagaraNodeOutput::RemoveOutputPin(UEdGraphPin* Pin)
 {
 	FScopedTransaction RemovePinTransaction(LOCTEXT("RemovePinTransaction", "Remove pin"));
-	int32 Index = Outputs.IndexOfByPredicate([&](const FNiagaraVariable& InVar) { return Pin->PinName == InVar.GetName().ToString(); });
+	int32 Index = Outputs.IndexOfByPredicate([&](const FNiagaraVariable& InVar) { return Pin->PinName == InVar.GetName(); });
 	if (Index >= 0)
 	{
 		Modify();
@@ -53,7 +53,7 @@ void UNiagaraNodeOutput::RemoveOutputPin(UEdGraphPin* Pin)
 
 FText UNiagaraNodeOutput::GetPinNameText(UEdGraphPin* Pin) const
 {
-	return FText::FromString(Pin->PinName);
+	return FText::FromName(Pin->PinName);
 }
 
 
@@ -64,10 +64,10 @@ void UNiagaraNodeOutput::PinNameTextCommitted(const FText& Text, ETextCommit::Ty
 		FScopedTransaction RenamePinTransaction(LOCTEXT("RenamePinTransaction", "Rename pin"));
 
 		Modify();
-		FNiagaraVariable* Var = Outputs.FindByPredicate([&](const FNiagaraVariable& InVar) {return Pin->PinName == InVar.GetName().ToString(); });
+		FNiagaraVariable* Var = Outputs.FindByPredicate([&](const FNiagaraVariable& InVar) {return Pin->PinName == InVar.GetName(); });
 		check(Var != nullptr);
-		Pin->PinName = Text.ToString();
-		Var->SetName(FName(*Pin->PinName));
+		Pin->PinName = *Text.ToString();
+		Var->SetName(Pin->PinName);
 		GetNiagaraGraph()->NotifyGraphNeedsRecompile();
 	}
 }
@@ -116,7 +116,7 @@ void UNiagaraNodeOutput::AllocateDefaultPins()
 
 	for (const FNiagaraVariable& Output : Outputs)
 	{
-		UEdGraphPin* Pin = CreatePin(EGPD_Input, Schema->TypeDefinitionToPinType(Output.GetType()), Output.GetName().ToString());
+		UEdGraphPin* Pin = CreatePin(EGPD_Input, Schema->TypeDefinitionToPinType(Output.GetType()), Output.GetName());
 		if (ScriptType == ENiagaraScriptUsage::ParticleUpdateScript)
 		{
 			Pin->bDefaultValueIsIgnored = true;
diff --git a/Engine/Plugins/FX/Niagara/Source/NiagaraEditor/Private/NiagaraNodeParameterCollection.cpp b/Engine/Plugins/FX/Niagara/Source/NiagaraEditor/Private/NiagaraNodeParameterCollection.cpp
--- a/Engine/Plugins/FX/Niagara/Source/NiagaraEditor/Private/NiagaraNodeParameterCollection.cpp
+++ b/Engine/Plugins/FX/Niagara/Source/NiagaraEditor/Private/NiagaraNodeParameterCollection.cpp
@@ -44,7 +44,7 @@ void UNiagaraNodeParameterCollection::AllocateDefaultPins()
 		{
 			if (Collection->IndexOfParameter(Var) != INDEX_NONE)
 			{
-				FString VarName = Collection->FriendlyNameFromParameterName(Var.GetName().ToString());
+				const FName VarName = *Collection->FriendlyNameFromParameterName(Var.GetName().ToString());
 				CreatePin(EGPD_Output, Schema->TypeDefinitionToPinType(Var.GetType()), VarName);
 			}
 		}
diff --git a/Engine/Plugins/FX/Niagara/Source/NiagaraEditor/Private/NiagaraNodeParameterMapGet.cpp b/Engine/Plugins/FX/Niagara/Source/NiagaraEditor/Private/NiagaraNodeParameterMapGet.cpp
--- a/Engine/Plugins/FX/Niagara/Source/NiagaraEditor/Private/NiagaraNodeParameterMapGet.cpp
+++ b/Engine/Plugins/FX/Niagara/Source/NiagaraEditor/Private/NiagaraNodeParameterMapGet.cpp
@@ -111,7 +111,7 @@ UEdGraphPin* UNiagaraNodeParameterMapGet::CreateDefaultPin(UEdGraphPin* OutputPi
 
 void UNiagaraNodeParameterMapGet::OnPinRenamed(UEdGraphPin* RenamedPin)
 {
-	RenamedPin->PinFriendlyName = FText::FromString(RenamedPin->PinName);
+	RenamedPin->PinFriendlyName = FText::FromName(RenamedPin->PinName);
 }
 
 
@@ -241,14 +241,14 @@ void UNiagaraNodeParameterMapGet::SynchronizeDefaultInputPin(UEdGraphPin* Defaul
 		DefaultPin->bDefaultValueIsIgnored = true;
 		DefaultPin->bNotConnectable = true;
 		DefaultPin->bHidden = true;
-		DefaultPin->PinToolTip = FText::Format(LOCTEXT("DefaultValueTooltip", "Default value for {0}. Disabled for Engine Parameters."), FText::FromString(OutputPin->PinName)).ToString();
+		DefaultPin->PinToolTip = FText::Format(LOCTEXT("DefaultValueTooltip", "Default value for {0}. Disabled for Engine Parameters."), FText::FromName(OutputPin->PinName)).ToString();
 	}
 	else
 	{
 		DefaultPin->bDefaultValueIsIgnored = false;
 		DefaultPin->bNotConnectable = false;
 		DefaultPin->bHidden = false;
-		DefaultPin->PinToolTip = FText::Format(LOCTEXT("DefaultValueTooltip", "Default value for {0} if no other module has set it previously in the stack."), FText::FromString(OutputPin->PinName)).ToString();
+		DefaultPin->PinToolTip = FText::Format(LOCTEXT("DefaultValueTooltip", "Default value for {0} if no other module has set it previously in the stack."), FText::FromName(OutputPin->PinName)).ToString();
 	}
 }
 
@@ -357,7 +357,7 @@ bool UNiagaraNodeParameterMapGet::CommitEditablePinName(const FText& InName, UEd
 		UEdGraphPin* DefaultPin = GetDefaultPin(InGraphPinObj);
 
 		InGraphPinObj->PinFriendlyName = InName;
-		InGraphPinObj->PinName = InName.ToString();
+		InGraphPinObj->PinName = *InName.ToString();
 
 		if (DefaultPin)
 		{
diff --git a/Engine/Plugins/FX/Niagara/Source/NiagaraEditor/Private/NiagaraNodeParameterMapSet.cpp b/Engine/Plugins/FX/Niagara/Source/NiagaraEditor/Private/NiagaraNodeParameterMapSet.cpp
--- a/Engine/Plugins/FX/Niagara/Source/NiagaraEditor/Private/NiagaraNodeParameterMapSet.cpp
+++ b/Engine/Plugins/FX/Niagara/Source/NiagaraEditor/Private/NiagaraNodeParameterMapSet.cpp
@@ -73,7 +73,7 @@ void UNiagaraNodeParameterMapSet::OnNewTypedPinAdded(UEdGraphPin* NewPin)
 
 void UNiagaraNodeParameterMapSet::OnPinRenamed(UEdGraphPin* RenamedPin) 
 {
-	RenamedPin->PinFriendlyName = FText::FromString(RenamedPin->PinName);
+	RenamedPin->PinFriendlyName = FText::FromName(RenamedPin->PinName);
 
 }
 
@@ -85,7 +85,7 @@ bool UNiagaraNodeParameterMapSet::CommitEditablePinName(const FText& InName, UEd
 		Modify();
 		InGraphPinObj->Modify();
 		InGraphPinObj->PinFriendlyName = InName;
-		InGraphPinObj->PinName = InName.ToString();
+		InGraphPinObj->PinName = *InName.ToString();
 
 		if (InGraphPinObj == PinPendingRename)
 		{
diff --git a/Engine/Plugins/FX/Niagara/Source/NiagaraEditor/Private/NiagaraNodeReadDataSet.cpp b/Engine/Plugins/FX/Niagara/Source/NiagaraEditor/Private/NiagaraNodeReadDataSet.cpp
--- a/Engine/Plugins/FX/Niagara/Source/NiagaraEditor/Private/NiagaraNodeReadDataSet.cpp
+++ b/Engine/Plugins/FX/Niagara/Source/NiagaraEditor/Private/NiagaraNodeReadDataSet.cpp
@@ -29,7 +29,7 @@ void UNiagaraNodeReadDataSet::AllocateDefaultPins()
 	for (int32 i = 0; i < Variables.Num(); i++)
 	{
 		FNiagaraVariable& Var = Variables[i];
-		UEdGraphPin* NewPin = CreatePin(EGPD_Output, Schema->TypeDefinitionToPinType(Var.GetType()), Var.GetName().ToString());
+		UEdGraphPin* NewPin = CreatePin(EGPD_Output, Schema->TypeDefinitionToPinType(Var.GetType()), Var.GetName());
 		if (useFriendlyNames && VariableFriendlyNames[i].IsEmpty() == false)
 		{
 			NewPin->PinFriendlyName = FText::FromString(VariableFriendlyNames[i]);
diff --git a/Engine/Plugins/FX/Niagara/Source/NiagaraEditor/Private/NiagaraNodeWithDynamicPins.cpp b/Engine/Plugins/FX/Niagara/Source/NiagaraEditor/Private/NiagaraNodeWithDynamicPins.cpp
--- a/Engine/Plugins/FX/Niagara/Source/NiagaraEditor/Private/NiagaraNodeWithDynamicPins.cpp
+++ b/Engine/Plugins/FX/Niagara/Source/NiagaraEditor/Private/NiagaraNodeWithDynamicPins.cpp
@@ -15,7 +15,7 @@
 
 #define LOCTEXT_NAMESPACE "NiagaraNodeWithDynamicPins"
 
-const FString UNiagaraNodeWithDynamicPins::AddPinSubCategory("DynamicAddPin");
+const FName UNiagaraNodeWithDynamicPins::AddPinSubCategory("DynamicAddPin");
 
 void UNiagaraNodeWithDynamicPins::PinConnectionListChanged(UEdGraphPin* Pin)
 {
@@ -27,7 +27,7 @@ void UNiagaraNodeWithDynamicPins::PinConnectionListChanged(UEdGraphPin* Pin)
 	{
 		FNiagaraTypeDefinition LinkedPinType = Schema->PinToTypeDefinition(Pin->LinkedTo[0]);
 		Pin->PinType = Schema->TypeDefinitionToPinType(LinkedPinType);
-		Pin->PinName = Pin->LinkedTo[0]->GetName();
+		Pin->PinName = Pin->LinkedTo[0]->GetFName();
 
 		CreateAddPin(Pin->Direction);
 		OnNewTypedPinAdded(Pin);
@@ -69,10 +69,10 @@ UEdGraphPin* UNiagaraNodeWithDynamicPins::RequestNewTypedPin(EEdGraphPinDirectio
 		GetOutputPins(OutPins);
 		DefaultName = TEXT("Output ") + LexicalConversion::ToString(OutPins.Num());
 	}
-	return RequestNewTypedPin(Direction, Type, DefaultName);
+	return RequestNewTypedPin(Direction, Type, *DefaultName);
 }
 
-UEdGraphPin* UNiagaraNodeWithDynamicPins::RequestNewTypedPin(EEdGraphPinDirection Direction, const FNiagaraTypeDefinition& Type, FString InName)
+UEdGraphPin* UNiagaraNodeWithDynamicPins::RequestNewTypedPin(EEdGraphPinDirection Direction, const FNiagaraTypeDefinition& Type, const FName InName)
 {
 	const UEdGraphSchema_Niagara* Schema = GetDefault<UEdGraphSchema_Niagara>();
 	UEdGraphPin* AddPin = GetAddPin(GetAllPins(), Direction);
@@ -257,7 +257,7 @@ void UNiagaraNodeWithDynamicPins::RemoveDynamicPin(UEdGraphPin* Pin)
 
 FText UNiagaraNodeWithDynamicPins::GetPinNameText(UEdGraphPin* Pin) const
 {
-	return FText::FromString(Pin->PinName);
+	return FText::FromName(Pin->PinName);
 }
 
 
@@ -268,7 +268,7 @@ void UNiagaraNodeWithDynamicPins::PinNameTextCommitted(const FText& Text, ETextC
 		FScopedTransaction RemovePinTransaction(LOCTEXT("RenamePinTransaction", "Rename pin"));
 		Modify();
 
-		Pin->PinName = Text.ToString();
+		Pin->PinName = *Text.ToString();
 		OnPinRenamed(Pin);
 	}
 }
diff --git a/Engine/Plugins/FX/Niagara/Source/NiagaraEditor/Private/NiagaraNodeWithDynamicPins.h b/Engine/Plugins/FX/Niagara/Source/NiagaraEditor/Private/NiagaraNodeWithDynamicPins.h
--- a/Engine/Plugins/FX/Niagara/Source/NiagaraEditor/Private/NiagaraNodeWithDynamicPins.h
+++ b/Engine/Plugins/FX/Niagara/Source/NiagaraEditor/Private/NiagaraNodeWithDynamicPins.h
@@ -24,7 +24,7 @@ public:
 	virtual void GetContextMenuActions(const FGraphNodeContextMenuBuilder& Context) const override;
 
 	/** Requests a new pin be added to the node with the specified direction, type, and name. */
-	UEdGraphPin* RequestNewTypedPin(EEdGraphPinDirection Direction, const FNiagaraTypeDefinition& Type, FString InName);
+	UEdGraphPin* RequestNewTypedPin(EEdGraphPinDirection Direction, const FNiagaraTypeDefinition& Type, FName InName);
 
 	/** Requests a new pin be added to the node with the specified direction and type. */
 	UEdGraphPin* RequestNewTypedPin(EEdGraphPinDirection Direction, const FNiagaraTypeDefinition& Type);
@@ -75,5 +75,5 @@ private:
 
 public:
 	/** The sub category for add pins. */
-	static const FString AddPinSubCategory;
+	static const FName AddPinSubCategory;
 };
diff --git a/Engine/Plugins/FX/Niagara/Source/NiagaraEditor/Private/NiagaraNodeWriteDataSet.cpp b/Engine/Plugins/FX/Niagara/Source/NiagaraEditor/Private/NiagaraNodeWriteDataSet.cpp
--- a/Engine/Plugins/FX/Niagara/Source/NiagaraEditor/Private/NiagaraNodeWriteDataSet.cpp
+++ b/Engine/Plugins/FX/Niagara/Source/NiagaraEditor/Private/NiagaraNodeWriteDataSet.cpp
@@ -43,7 +43,7 @@ void UNiagaraNodeWriteDataSet::AllocateDefaultPins()
 	for (int32 i = 0; i < Variables.Num(); i++)
 	{
 		FNiagaraVariable& Var = Variables[i];
-		UEdGraphPin* NewPin = CreatePin(EGPD_Input, Schema->TypeDefinitionToPinType(Var.GetType()), Var.GetName().ToString());
+		UEdGraphPin* NewPin = CreatePin(EGPD_Input, Schema->TypeDefinitionToPinType(Var.GetType()), Var.GetName());
 		if (useFriendlyNames && VariableFriendlyNames[i].IsEmpty() == false)
 		{
 			NewPin->PinFriendlyName = FText::FromString(VariableFriendlyNames[i]);
diff --git a/Engine/Plugins/FX/Niagara/Source/NiagaraEditor/Private/ViewModels/Stack/NiagaraParameterHandle.cpp b/Engine/Plugins/FX/Niagara/Source/NiagaraEditor/Private/ViewModels/Stack/NiagaraParameterHandle.cpp
--- a/Engine/Plugins/FX/Niagara/Source/NiagaraEditor/Private/ViewModels/Stack/NiagaraParameterHandle.cpp
+++ b/Engine/Plugins/FX/Niagara/Source/NiagaraEditor/Private/ViewModels/Stack/NiagaraParameterHandle.cpp
@@ -5,94 +5,95 @@
 #include "NiagaraScript.h"
 #include "NiagaraTypes.h"
 
-const FString FNiagaraParameterHandle::UserNamespace(TEXT("User"));
-const FString FNiagaraParameterHandle::EngineNamespace(TEXT("Engine"));
-const FString FNiagaraParameterHandle::SystemNamespace(TEXT("System"));
-const FString FNiagaraParameterHandle::EmitterNamespace(TEXT("Emitter"));
-const FString FNiagaraParameterHandle::ParticleAttributeNamespace(TEXT("Particles"));
-const FString FNiagaraParameterHandle::ModuleNamespace(TEXT("Module"));
+const FName FNiagaraParameterHandle::UserNamespace(TEXT("User"));
+const FName FNiagaraParameterHandle::EngineNamespace(TEXT("Engine"));
+const FName FNiagaraParameterHandle::SystemNamespace(TEXT("System"));
+const FName FNiagaraParameterHandle::EmitterNamespace(TEXT("Emitter"));
+const FName FNiagaraParameterHandle::ParticleAttributeNamespace(TEXT("Particles"));
+const FName FNiagaraParameterHandle::ModuleNamespace(TEXT("Module"));
 const FString FNiagaraParameterHandle::InitialPrefix(TEXT("Initial"));
 
 FNiagaraParameterHandle::FNiagaraParameterHandle() 
 {
 }
 
-FNiagaraParameterHandle::FNiagaraParameterHandle(const FString& InParameterHandleString)
+FNiagaraParameterHandle::FNiagaraParameterHandle(FName InParameterHandleName)
+	: ParameterHandleName(InParameterHandleName)
 {
-	ParameterHandleString = InParameterHandleString;
 	int32 DotIndex;
+	const FString ParameterHandleString = ParameterHandleName.ToString();
 	if (ParameterHandleString.FindChar(TEXT('.'), DotIndex))
 	{
-		Name = ParameterHandleString.RightChop(DotIndex + 1);
-		Namespace = ParameterHandleString.Left(DotIndex);
+		Name = *ParameterHandleString.RightChop(DotIndex + 1);
+		Namespace = *ParameterHandleString.Left(DotIndex);
 	}
 	else
 	{
-		Name = ParameterHandleString;
+		Name = ParameterHandleName;
 	}
 }
 
-FNiagaraParameterHandle::FNiagaraParameterHandle(const FString& InNamespace, const FString& InName)
+FNiagaraParameterHandle::FNiagaraParameterHandle(FName InNamespace, FName InName)
+	: Name(InName)
+	, Namespace(InNamespace)
 {
-	Namespace = InNamespace;
-	Name = InName;
-	ParameterHandleString = Namespace + "." + Name;
+	ParameterHandleName = *FString::Printf(TEXT("%s.%s"), *Namespace.ToString(), *Name.ToString());
 }
 
 bool FNiagaraParameterHandle::operator==(const FNiagaraParameterHandle& Other) const
 {
-	return ParameterHandleString == Other.ParameterHandleString;
+	return ParameterHandleName == Other.ParameterHandleName;
 }
 
 FNiagaraParameterHandle FNiagaraParameterHandle::CreateAliasedModuleParameterHandle(const FNiagaraParameterHandle& ModuleParameterHandle, UNiagaraNodeFunctionCall* ModuleNode)
 {
 	return ModuleParameterHandle.IsModuleHandle()
-		? FNiagaraParameterHandle(ModuleNode->GetFunctionName(), ModuleParameterHandle.GetName())
+		? FNiagaraParameterHandle(*ModuleNode->GetFunctionName(), ModuleParameterHandle.GetName())
 		: ModuleParameterHandle;
 }
 
 FNiagaraParameterHandle FNiagaraParameterHandle::CreateEngineParameterHandle(const FNiagaraVariable& SystemVariable)
 {
-	return FNiagaraParameterHandle(SystemVariable.GetName().ToString());
+	return FNiagaraParameterHandle(SystemVariable.GetName());
 }
 
 FNiagaraParameterHandle FNiagaraParameterHandle::CreateEmitterParameterHandle(const FNiagaraVariable& EmitterVariable)
 {
 
-	return FNiagaraParameterHandle(EmitterNamespace, EmitterVariable.GetName().ToString());
+	return FNiagaraParameterHandle(EmitterNamespace, EmitterVariable.GetName());
 }
 
-FNiagaraParameterHandle FNiagaraParameterHandle::CreateParticleAttributeParameterHandle(const FString& InName)
+FNiagaraParameterHandle FNiagaraParameterHandle::CreateParticleAttributeParameterHandle(const FName InName)
 {
 	return FNiagaraParameterHandle(ParticleAttributeNamespace, InName);
 }
 
-FNiagaraParameterHandle FNiagaraParameterHandle::CreateModuleParameterHandle(const FString& InName)
+FNiagaraParameterHandle FNiagaraParameterHandle::CreateModuleParameterHandle(const FName InName)
 {
 	return FNiagaraParameterHandle(ModuleNamespace, InName);
 }
 
 FNiagaraParameterHandle FNiagaraParameterHandle::CreateInitialParameterHandle(const FNiagaraParameterHandle& Handle)
 {
-	return FNiagaraParameterHandle(Handle.GetNamespace(), InitialPrefix + TEXT(".") + Handle.GetName());
+	return FNiagaraParameterHandle(Handle.GetNamespace(), *FString::Printf(TEXT("%s.%s"), *InitialPrefix, *Handle.GetName().ToString()));
 }
 
 bool FNiagaraParameterHandle::IsValid() const 
 {
-	return ParameterHandleString.IsEmpty() == false;
+	return ParameterHandleName.IsNone() == false;
 }
 
-const FString& FNiagaraParameterHandle::GetParameterHandleString() const 
+const FName FNiagaraParameterHandle::GetParameterHandleString() const 
 {
-	return ParameterHandleString;
+	return ParameterHandleName;
 }
 
-const FString& FNiagaraParameterHandle::GetName() const 
+const FName FNiagaraParameterHandle::GetName() const 
 {
 	return Name;
 }
 
-const FString& FNiagaraParameterHandle::GetNamespace() const
+const FName FNiagaraParameterHandle::GetNamespace() const
 {
 	return Namespace;
 }
diff --git a/Engine/Plugins/FX/Niagara/Source/NiagaraEditor/Private/ViewModels/Stack/NiagaraStackAddModuleItem.cpp b/Engine/Plugins/FX/Niagara/Source/NiagaraEditor/Private/ViewModels/Stack/NiagaraStackAddModuleItem.cpp
--- a/Engine/Plugins/FX/Niagara/Source/NiagaraEditor/Private/ViewModels/Stack/NiagaraStackAddModuleItem.cpp
+++ b/Engine/Plugins/FX/Niagara/Source/NiagaraEditor/Private/ViewModels/Stack/NiagaraStackAddModuleItem.cpp
@@ -33,9 +33,9 @@ void UNiagaraStackAddModuleItem::GetNewParameterAvailableTypes(TArray<FNiagaraTy
 {
 }
 
-TOptional<FString> UNiagaraStackAddModuleItem::GetNewParameterNamespace() const
+TOptional<FName> UNiagaraStackAddModuleItem::GetNewParameterNamespace() const
 {
-	return TOptional<FString>();
+	return TOptional<FName>();
 }
 
 void UNiagaraStackAddModuleItem::SetOnItemAdded(FOnItemAdded InOnItemAdded)
diff --git a/Engine/Plugins/FX/Niagara/Source/NiagaraEditor/Private/ViewModels/Stack/NiagaraStackAddScriptModuleItem.cpp b/Engine/Plugins/FX/Niagara/Source/NiagaraEditor/Private/ViewModels/Stack/NiagaraStackAddScriptModuleItem.cpp
--- a/Engine/Plugins/FX/Niagara/Source/NiagaraEditor/Private/ViewModels/Stack/NiagaraStackAddScriptModuleItem.cpp
+++ b/Engine/Plugins/FX/Niagara/Source/NiagaraEditor/Private/ViewModels/Stack/NiagaraStackAddScriptModuleItem.cpp
@@ -53,7 +53,7 @@ void UNiagaraStackAddScriptModuleItem::GetNewParameterAvailableTypes(TArray<FNia
 	}
 }
 
-TOptional<FString> UNiagaraStackAddScriptModuleItem::GetNewParameterNamespace() const
+TOptional<FName> UNiagaraStackAddScriptModuleItem::GetNewParameterNamespace() const
 {
 	switch (GetOutputUsage())
 	{
@@ -67,7 +67,7 @@ TOptional<FString> UNiagaraStackAddScriptModuleItem::GetNewParameterNamespace()
 	case ENiagaraScriptUsage::SystemUpdateScript:
 		return FNiagaraParameterHandle::SystemNamespace;
 	default:
-		return TOptional<FString>();
+		return TOptional<FName>();
 	}
 }
 
diff --git a/Engine/Plugins/FX/Niagara/Source/NiagaraEditor/Private/ViewModels/Stack/NiagaraStackFunctionInput.cpp b/Engine/Plugins/FX/Niagara/Source/NiagaraEditor/Private/ViewModels/Stack/NiagaraStackFunctionInput.cpp
--- a/Engine/Plugins/FX/Niagara/Source/NiagaraEditor/Private/ViewModels/Stack/NiagaraStackFunctionInput.cpp
+++ b/Engine/Plugins/FX/Niagara/Source/NiagaraEditor/Private/ViewModels/Stack/NiagaraStackFunctionInput.cpp
@@ -128,7 +128,7 @@ void UNiagaraStackFunctionInput::Initialize(
 	UNiagaraStackEditorData& InStackEditorData,
 	UNiagaraNodeFunctionCall& InModuleNode,
 	UNiagaraNodeFunctionCall& InInputFunctionCallNode,
-	FString InInputParameterHandle,
+	FName InInputParameterHandle,
 	FNiagaraTypeDefinition InInputType)
 {
 	checkf(OwningModuleNode.IsValid() == false && OwningFunctionCallNode.IsValid() == false, TEXT("Can only initialize once."));
@@ -145,7 +145,7 @@ void UNiagaraStackFunctionInput::Initialize(
 	GenerateInputParameterHandlePath(*OwningModuleNode, *OwningFunctionCallNode, InputParameterHandlePath);
 	InputParameterHandlePath.Add(InputParameterHandle);
 
-	DisplayName = FText::FromString(InputParameterHandle.GetName());
+	DisplayName = FText::FromName(InputParameterHandle.GetName());
 	AliasedInputParameterHandle = FNiagaraParameterHandle::CreateAliasedModuleParameterHandle(InputParameterHandle, OwningFunctionCallNode.Get());
 
 	InputType = InInputType;
@@ -192,11 +192,6 @@ void UNiagaraStackFunctionInput::RefreshChildrenInternal(const TArray<UNiagaraSt
 			DisplayOptions.bShouldShowInStack = false;
 			DisplayOptions.ChildItemIndentLevel = ItemIndentLevel + 1;
 
-			TArray<FString> InputParameterHandleStringPath;
-			for (const FNiagaraParameterHandle& InputParameterHandlePathItem : InputParameterHandlePath)
-			{
-				InputParameterHandleStringPath.Add(InputParameterHandlePathItem.GetParameterHandleString());
-			}
 			DynamicInputEntry = NewObject<UNiagaraStackFunctionInputCollection>(this);
 			DynamicInputEntry->Initialize(GetSystemViewModel(), GetEmitterViewModel(), *StackEditorData, *OwningModuleNode, *InputValues.DynamicNode.Get(), DisplayOptions);
 		}
@@ -356,7 +351,7 @@ bool IsSpawnUsage(ENiagaraScriptUsage Usage)
 		Usage == ENiagaraScriptUsage::ParticleSpawnScript;
 }
 
-FString GetNamespaceForUsage(ENiagaraScriptUsage Usage)
+FName GetNamespaceForUsage(ENiagaraScriptUsage Usage)
 {
 	switch (Usage)
 	{
@@ -370,7 +365,7 @@ FString GetNamespaceForUsage(ENiagaraScriptUsage Usage)
 	case ENiagaraScriptUsage::SystemUpdateScript:
 		return FNiagaraParameterHandle::SystemNamespace;
 	default:
-		return FString();
+		return NAME_None;
 	}
 }
 
@@ -425,7 +420,7 @@ void UNiagaraStackFunctionInput::GetAvailableParameterHandles(TArray<FNiagaraPar
 							{
 								if (Cast<UNiagaraNodeParameterMapSet>(WritePin->GetOwningNode()) != nullptr)
 								{
-									FNiagaraParameterHandle AvailableHandle = FNiagaraParameterHandle(Variable.GetName().ToString());
+									FNiagaraParameterHandle AvailableHandle = FNiagaraParameterHandle(Variable.GetName());
 									AvailableParameterHandles.AddUnique(AvailableHandle);
 									AvailableParameterHandlesForThisOutput.AddUnique(AvailableHandle);
 									break;
@@ -438,8 +433,8 @@ void UNiagaraStackFunctionInput::GetAvailableParameterHandles(TArray<FNiagaraPar
 
 			if (OutputNode != CurrentOutputNode && IsSpawnUsage(OutputNode->GetUsage()))
 			{
-				FString OutputNodeNamespace = GetNamespaceForUsage(OutputNode->GetUsage());
-				if (OutputNodeNamespace.IsEmpty() == false)
+				FName OutputNodeNamespace = GetNamespaceForUsage(OutputNode->GetUsage());
+				if (OutputNodeNamespace.IsNone() == false)
 				{
 					for (FNiagaraParameterHandle& AvailableParameterHandleForThisOutput : AvailableParameterHandlesForThisOutput)
 					{
@@ -645,10 +640,10 @@ void UNiagaraStackFunctionInput::Reset()
 			UEdGraphPin& OverridePin = GetOrCreateOverridePin();
 			RemoveAllNodesConnectedToOverridePin(OverridePin, this);
 
-			FString InputNodeName = InputParameterHandlePath[0].GetName();
+			FString InputNodeName = InputParameterHandlePath[0].GetName().ToString();
 			for (int32 i = 1; i < InputParameterHandlePath.Num(); i++)
 			{
-				InputNodeName += "." + InputParameterHandlePath[i].GetName();
+				InputNodeName += "." + InputParameterHandlePath[i].GetName().ToString();
 			}
 
 			UEdGraph* Graph = OwningFunctionCallNode->GetGraph();
@@ -656,7 +651,7 @@ void UNiagaraStackFunctionInput::Reset()
 			UNiagaraNodeInput* InputNode = InputNodeCreator.CreateNode();
 			FNiagaraEditorUtilities::InitializeParameterInputNode(*InputNode, InputType, CastChecked<UNiagaraGraph>(Graph), *InputNodeName);
 
-			UNiagaraDataInterface* InputValueObject = NewObject<UNiagaraDataInterface>(InputNode, const_cast<UClass*>(InputType.GetClass()), *InputParameterHandle.GetName(), RF_Transactional);
+			UNiagaraDataInterface* InputValueObject = NewObject<UNiagaraDataInterface>(InputNode, const_cast<UClass*>(InputType.GetClass()), InputParameterHandle.GetName(), RF_Transactional);
 			if (InputValues.DataObjects.GetDefaultValueObject() != nullptr)
 			{
 				InputValues.DataObjects.GetDefaultValueObject()->CopyTo(InputValueObject);
@@ -706,23 +701,23 @@ void UNiagaraStackFunctionInput::SetIsRenamePending(bool bIsRenamePending)
 	}
 }
 
-void UNiagaraStackFunctionInput::RenameInput(FString NewName)
+void UNiagaraStackFunctionInput::RenameInput(FName NewName)
 {
 	if (OwningAssignmentNode.IsValid() && InputParameterHandlePath.Num() == 1 && InputParameterHandle.GetName() != NewName)
 	{
 		bool bIsCurrentlyPinned = GetIsPinned();
 		bool bIsCurrentlyExpanded = StackEditorData->GetStackEntryIsExpanded(FNiagaraStackGraphUtilities::GenerateStackModuleEditorDataKey(*OwningAssignmentNode), false);
 
-		FNiagaraParameterHandle TargetHandle(OwningAssignmentNode->AssignmentTarget.GetName().ToString());
+		FNiagaraParameterHandle TargetHandle(OwningAssignmentNode->AssignmentTarget.GetName());
 		FNiagaraParameterHandle RenamedTargetHandle(TargetHandle.GetNamespace(), NewName);
-		OwningAssignmentNode->AssignmentTarget.SetName(*RenamedTargetHandle.GetParameterHandleString());
+		OwningAssignmentNode->AssignmentTarget.SetName(RenamedTargetHandle.GetParameterHandleString());
 		OwningAssignmentNode->RefreshFromExternalChanges();
 
 		InputParameterHandle = FNiagaraParameterHandle(InputParameterHandle.GetNamespace(), NewName);
-		InputParameterHandlePath.Empty();
+		InputParameterHandlePath.Empty(1);
 		InputParameterHandlePath.Add(InputParameterHandle);
 		AliasedInputParameterHandle = FNiagaraParameterHandle::CreateAliasedModuleParameterHandle(InputParameterHandle, OwningAssignmentNode.Get());
-		DisplayName = FText::FromString(InputParameterHandle.GetName());
+		DisplayName = FText::FromName(InputParameterHandle.GetName());
 
 		UEdGraphPin* OverridePin = GetOverridePin();
 		if (OverridePin != nullptr)
@@ -738,7 +733,7 @@ void UNiagaraStackFunctionInput::RenameInput(FString NewName)
 	}
 }
 
-void UNiagaraStackFunctionInput::GetNamespacesForNewParameters(TArray<FString>& OutNamespacesForNewParameters) const
+void UNiagaraStackFunctionInput::GetNamespacesForNewParameters(TArray<FName>& OutNamespacesForNewParameters) const
 {
 	UNiagaraNodeOutput* OutputNode = FNiagaraStackGraphUtilities::GetEmitterOutputNodeForStackNode(*OwningFunctionCallNode);
 	bool bIsSystemTransient = GetSystemViewModel()->GetSystemIsTransient();
@@ -1052,7 +1047,7 @@ void UNiagaraStackFunctionInput::RemoveAllNodesConnectedToOverridePin(UEdGraphPi
 				DynamicInputNodeOverrideNode->GetInputPins(InputPins);
 				for (UEdGraphPin* InputPin : InputPins)
 				{
-					if (InputPin->PinName.StartsWith(DynamicInputNode->GetFunctionName()))
+					if (InputPin->PinName.ToString().StartsWith(DynamicInputNode->GetFunctionName()))
 					{
 						RemoveAllNodesConnectedToOverridePin(*InputPin, OwningInput);
 						DynamicInputNodeOverrideNode->RemovePin(InputPin);
diff --git a/Engine/Plugins/FX/Niagara/Source/NiagaraEditor/Private/ViewModels/Stack/NiagaraStackGraphUtilities.cpp b/Engine/Plugins/FX/Niagara/Source/NiagaraEditor/Private/ViewModels/Stack/NiagaraStackGraphUtilities.cpp
--- a/Engine/Plugins/FX/Niagara/Source/NiagaraEditor/Private/ViewModels/Stack/NiagaraStackGraphUtilities.cpp
+++ b/Engine/Plugins/FX/Niagara/Source/NiagaraEditor/Private/ViewModels/Stack/NiagaraStackGraphUtilities.cpp
@@ -401,7 +401,7 @@ void FNiagaraStackGraphUtilities::InitializeDataInterfaceInputs(TSharedRef<FNiag
 
 FString FNiagaraStackGraphUtilities::GenerateStackFunctionInputEditorDataKey(UNiagaraNodeFunctionCall& FunctionCallNode, FNiagaraParameterHandle InputParameterHandle)
 {
-	return FunctionCallNode.GetFunctionName() + InputParameterHandle.GetParameterHandleString();
+	return FunctionCallNode.GetFunctionName() + InputParameterHandle.GetParameterHandleString().ToString();
 }
 
 FString FNiagaraStackGraphUtilities::GenerateStackModuleEditorDataKey(UNiagaraNodeFunctionCall& ModuleNode)
diff --git a/Engine/Plugins/FX/Niagara/Source/NiagaraEditor/Private/ViewModels/Stack/NiagaraStackModuleItemOutput.cpp b/Engine/Plugins/FX/Niagara/Source/NiagaraEditor/Private/ViewModels/Stack/NiagaraStackModuleItemOutput.cpp
--- a/Engine/Plugins/FX/Niagara/Source/NiagaraEditor/Private/ViewModels/Stack/NiagaraStackModuleItemOutput.cpp
+++ b/Engine/Plugins/FX/Niagara/Source/NiagaraEditor/Private/ViewModels/Stack/NiagaraStackModuleItemOutput.cpp
@@ -15,14 +15,14 @@ int32 UNiagaraStackModuleItemOutput::GetItemIndentLevel() const
 	return 1;
 }
 
-void UNiagaraStackModuleItemOutput::Initialize(TSharedRef<FNiagaraSystemViewModel> InSystemViewModel, TSharedRef<FNiagaraEmitterViewModel> InEmitterViewModel, UNiagaraNodeFunctionCall& InFunctionCallNode, FString InOutputParameterHandle)
+void UNiagaraStackModuleItemOutput::Initialize(TSharedRef<FNiagaraSystemViewModel> InSystemViewModel, TSharedRef<FNiagaraEmitterViewModel> InEmitterViewModel, UNiagaraNodeFunctionCall& InFunctionCallNode, FName InOutputParameterHandle)
 {
 	checkf(FunctionCallNode.Get() == nullptr, TEXT("Can only set the Output once."));
 	Super::Initialize(InSystemViewModel, InEmitterViewModel);
 	FunctionCallNode = &InFunctionCallNode;
 
 	OutputParameterHandle = FNiagaraParameterHandle(InOutputParameterHandle);
-	DisplayName = FText::FromString(OutputParameterHandle.GetName());
+	DisplayName = FText::FromName(OutputParameterHandle.GetName());
 }
 
 FText UNiagaraStackModuleItemOutput::GetDisplayName() const
@@ -47,7 +47,7 @@ const FNiagaraParameterHandle& UNiagaraStackModuleItemOutput::GetOutputParameter
 
 FText UNiagaraStackModuleItemOutput::GetOutputParameterHandleText() const
 {
-	return FText::FromString(OutputParameterHandle.GetParameterHandleString());
+	return FText::FromName(OutputParameterHandle.GetParameterHandleString());
 }
 
 #undef LOCTEXT_NAMESPACE
diff --git a/Engine/Plugins/FX/Niagara/Source/NiagaraEditor/Private/ViewModels/Stack/NiagaraStackModuleItemOutputCollection.cpp b/Engine/Plugins/FX/Niagara/Source/NiagaraEditor/Private/ViewModels/Stack/NiagaraStackModuleItemOutputCollection.cpp
--- a/Engine/Plugins/FX/Niagara/Source/NiagaraEditor/Private/ViewModels/Stack/NiagaraStackModuleItemOutputCollection.cpp
+++ b/Engine/Plugins/FX/Niagara/Source/NiagaraEditor/Private/ViewModels/Stack/NiagaraStackModuleItemOutputCollection.cpp
@@ -67,7 +67,7 @@ void UNiagaraStackModuleItemOutputCollection::RefreshChildrenInternal(const TArr
 					for (UNiagaraStackEntry* CurrentChild : CurrentChildren)
 					{
 						UNiagaraStackModuleItemOutput* ChildOutput = CastChecked<UNiagaraStackModuleItemOutput>(CurrentChild);
-						if (ChildOutput->GetOutputParameterHandle().GetParameterHandleString() == Variable.GetName().ToString())
+						if (ChildOutput->GetOutputParameterHandle().GetParameterHandleString() == Variable.GetName())
 						{
 							Output = ChildOutput;
 							break;
@@ -76,7 +76,7 @@ void UNiagaraStackModuleItemOutputCollection::RefreshChildrenInternal(const TArr
 					if (Output == nullptr)
 					{
 						Output = NewObject<UNiagaraStackModuleItemOutput>(this);
-						Output->Initialize(GetSystemViewModel(), GetEmitterViewModel(), *FunctionCallNode, Variable.GetName().ToString());
+						Output->Initialize(GetSystemViewModel(), GetEmitterViewModel(), *FunctionCallNode, Variable.GetName());
 					}
 
 					NewChildren.Add(Output);
diff --git a/Engine/Plugins/FX/Niagara/Source/NiagaraEditor/Private/ViewModels/Stack/NiagaraStackParameterStoreEntry.cpp b/Engine/Plugins/FX/Niagara/Source/NiagaraEditor/Private/ViewModels/Stack/NiagaraStackParameterStoreEntry.cpp
--- a/Engine/Plugins/FX/Niagara/Source/NiagaraEditor/Private/ViewModels/Stack/NiagaraStackParameterStoreEntry.cpp
+++ b/Engine/Plugins/FX/Niagara/Source/NiagaraEditor/Private/ViewModels/Stack/NiagaraStackParameterStoreEntry.cpp
@@ -207,7 +207,7 @@ void UNiagaraStackParameterStoreEntry::RenameInput(FString NewName)
 
 		FNiagaraParameterHandle TargetHandle(OwningAssignmentNode->AssignmentTarget.GetName().ToString());
 		FNiagaraParameterHandle RenamedTargetHandle(TargetHandle.GetNamespace(), NewName);
-		OwningAssignmentNode->AssignmentTarget.SetName(*RenamedTargetHandle.GetParameterHandleString());
+		OwningAssignmentNode->AssignmentTarget.SetName(RenamedTargetHandle.GetParameterHandleString());
 		OwningAssignmentNode->RefreshFromExternalChanges();
 
 		InputParameterHandle = FNiagaraParameterHandle(InputParameterHandle.GetNamespace(), NewName);
diff --git a/Engine/Plugins/FX/Niagara/Source/NiagaraEditor/Private/Widgets/SNiagaraGraphPinAdd.cpp b/Engine/Plugins/FX/Niagara/Source/NiagaraEditor/Private/Widgets/SNiagaraGraphPinAdd.cpp
--- a/Engine/Plugins/FX/Niagara/Source/NiagaraEditor/Private/Widgets/SNiagaraGraphPinAdd.cpp
+++ b/Engine/Plugins/FX/Niagara/Source/NiagaraEditor/Private/Widgets/SNiagaraGraphPinAdd.cpp
@@ -97,7 +97,7 @@ void SNiagaraGraphPinAdd::OnAddType(FNiagaraVariable InAdd)
 	if (OwningNode != nullptr)
 	{
 		FScopedTransaction AddNewPinTransaction(LOCTEXT("AddNewPinTransaction", "Add pin to node"));
-		OwningNode->RequestNewTypedPin(GetPinObj()->Direction, InAdd.GetType(), InAdd.GetName().ToString());
+		OwningNode->RequestNewTypedPin(GetPinObj()->Direction, InAdd.GetType(), InAdd.GetName());
 	}
 }
 #undef LOCTEXT_NAMESPACE
diff --git a/Engine/Plugins/FX/Niagara/Source/NiagaraEditor/Public/EdGraphSchema_Niagara.h b/Engine/Plugins/FX/Niagara/Source/NiagaraEditor/Public/EdGraphSchema_Niagara.h
--- a/Engine/Plugins/FX/Niagara/Source/NiagaraEditor/Public/EdGraphSchema_Niagara.h
+++ b/Engine/Plugins/FX/Niagara/Source/NiagaraEditor/Public/EdGraphSchema_Niagara.h
@@ -59,10 +59,10 @@ class NIAGARAEDITOR_API UEdGraphSchema_Niagara : public UEdGraphSchema
 	GENERATED_UCLASS_BODY()
 
 	// Allowable PinType.PinCategory values
-	static const FString PinCategoryType;
-	static const FString PinCategoryMisc;
-	static const FString PinCategoryClass;
-	static const FString PinCategoryEnum;
+	static const FName PinCategoryType;
+	static const FName PinCategoryMisc;
+	static const FName PinCategoryClass;
+	static const FName PinCategoryEnum;
 
 	//~ Begin EdGraphSchema Interface
 	virtual void GetGraphContextActions(FGraphContextMenuBuilder& ContextMenuBuilder) const override;
@@ -71,11 +71,11 @@ class NIAGARAEDITOR_API UEdGraphSchema_Niagara : public UEdGraphSchema
 	virtual FLinearColor GetPinTypeColor(const FEdGraphPinType& PinType) const override;
 	virtual bool ShouldHidePinDefaultValue(UEdGraphPin* Pin) const override;
 	virtual bool TryCreateConnection(UEdGraphPin* A, UEdGraphPin* B) const override;
-	virtual void BreakSinglePinLink(UEdGraphPin* SourcePin, UEdGraphPin* TargetPin) override;
+	virtual void BreakSinglePinLink(UEdGraphPin* SourcePin, UEdGraphPin* TargetPin) const override;
 	virtual FConnectionDrawingPolicy* CreateConnectionDrawingPolicy(int32 InBackLayerID, int32 InFrontLayerID, float InZoomFactor, const FSlateRect& InClippingRect, class FSlateWindowElementList& InDrawElements, class UEdGraph* InGraphObj) const override;
 	//~ End EdGraphSchema Interface
 
-	virtual TArray<TSharedPtr<FNiagaraSchemaAction_NewNode> > GetGraphContextActions(const UEdGraph* CurrentGraph, TArray<UObject*>& SelectedObjects, const UEdGraphPin* FromPin, UEdGraph* OwnerOfTemporaries) const;
+	TArray<TSharedPtr<FNiagaraSchemaAction_NewNode> > GetGraphContextActions(const UEdGraph* CurrentGraph, TArray<UObject*>& SelectedObjects, const UEdGraphPin* FromPin, UEdGraph* OwnerOfTemporaries) const;
 	void PromoteSinglePinToParameter(UEdGraphPin* SourcePin);
 
 	FNiagaraVariable PinToNiagaraVariable(const UEdGraphPin* Pin, bool bNeedsValue=false)const;
diff --git a/Engine/Plugins/FX/Niagara/Source/NiagaraEditor/Public/NiagaraGraph.h b/Engine/Plugins/FX/Niagara/Source/NiagaraEditor/Public/NiagaraGraph.h
--- a/Engine/Plugins/FX/Niagara/Source/NiagaraEditor/Public/NiagaraGraph.h
+++ b/Engine/Plugins/FX/Niagara/Source/NiagaraEditor/Public/NiagaraGraph.h
@@ -110,7 +110,8 @@ class UNiagaraGraph : public UEdGraph
 	FGuid GetChangeID() { return ChangeId; }
 
 	/** Walk through the graph for an ParameterMapGet nodes and see if any of them specify a default for VariableName.*/
-	UEdGraphPin* FindParameterMapDefaultValuePin(const FString& VariableName);
+	UEdGraphPin* FindParameterMapDefaultValuePin(const FName VariableName);
+
 private:
 	
 	/** The current change identifier for this graph. Used to sync status with UNiagaraScripts.*/
diff --git a/Engine/Plugins/FX/Niagara/Source/NiagaraEditor/Public/NiagaraNodeDataSetBase.h b/Engine/Plugins/FX/Niagara/Source/NiagaraEditor/Public/NiagaraNodeDataSetBase.h
--- a/Engine/Plugins/FX/Niagara/Source/NiagaraEditor/Public/NiagaraNodeDataSetBase.h
+++ b/Engine/Plugins/FX/Niagara/Source/NiagaraEditor/Public/NiagaraNodeDataSetBase.h
@@ -50,6 +50,6 @@ protected:
 	bool SynchronizeWithStruct();
 	bool GetSupportedNiagaraTypeDef(const UProperty* Property, FNiagaraTypeDefinition& TypeDef);
 
-	static FString ConditionVarName;
+	static const FName ConditionVarName;
 };
 
diff --git a/Engine/Plugins/FX/Niagara/Source/NiagaraEditor/Public/NiagaraNodeFunctionCall.h b/Engine/Plugins/FX/Niagara/Source/NiagaraEditor/Public/NiagaraNodeFunctionCall.h
--- a/Engine/Plugins/FX/Niagara/Source/NiagaraEditor/Public/NiagaraNodeFunctionCall.h
+++ b/Engine/Plugins/FX/Niagara/Source/NiagaraEditor/Public/NiagaraNodeFunctionCall.h
@@ -62,7 +62,7 @@ public:
 	FString GetFunctionName() const { return FunctionDisplayName; }
 
 	/** Walk through the internal script graph for an ParameterMapGet nodes and see if any of them specify a default for VariableName.*/
-	UEdGraphPin* FindParameterMapDefaultValuePin(const FString& VariableName);
+	UEdGraphPin* FindParameterMapDefaultValuePin(const FName VariableName);
 
 protected:
 	/** Resets the node name based on the referenced script or signature. Guaranteed unique within a given graph instance.*/
diff --git a/Engine/Plugins/FX/Niagara/Source/NiagaraEditor/Public/ViewModels/Stack/NiagaraParameterHandle.h b/Engine/Plugins/FX/Niagara/Source/NiagaraEditor/Public/ViewModels/Stack/NiagaraParameterHandle.h
--- a/Engine/Plugins/FX/Niagara/Source/NiagaraEditor/Public/ViewModels/Stack/NiagaraParameterHandle.h
+++ b/Engine/Plugins/FX/Niagara/Source/NiagaraEditor/Public/ViewModels/Stack/NiagaraParameterHandle.h
@@ -3,6 +3,7 @@
 #pragma once
 
 #include "Containers/UnrealString.h"
+#include "UObject/NameTypes.h"
 
 class UNiagaraNodeFunctionCall;
 struct FNiagaraVariable;
@@ -12,9 +13,9 @@ class NIAGARAEDITOR_API FNiagaraParameterHandle
 public:
 	FNiagaraParameterHandle();
 
-	FNiagaraParameterHandle(const FString& InParameterHandleString);
+	FNiagaraParameterHandle(FName InParameterHandleString);
 
-	FNiagaraParameterHandle(const FString& InNamespace, const FString& InName);
+	FNiagaraParameterHandle(FName InNamespace, FName InName);
 
 	bool operator==(const FNiagaraParameterHandle& Other) const;
 
@@ -24,19 +25,19 @@ public:
 
 	static FNiagaraParameterHandle CreateEmitterParameterHandle(const FNiagaraVariable& EmitterVariable);
 
-	static FNiagaraParameterHandle CreateParticleAttributeParameterHandle(const FString& InName);
+	static FNiagaraParameterHandle CreateParticleAttributeParameterHandle(const FName InName);
 
-	static FNiagaraParameterHandle CreateModuleParameterHandle(const FString& InName);
+	static FNiagaraParameterHandle CreateModuleParameterHandle(const FName InName);
 
 	static FNiagaraParameterHandle CreateInitialParameterHandle(const FNiagaraParameterHandle& Handle);
 
 	bool IsValid() const;
 
-	const FString& GetParameterHandleString() const;
+	const FName GetParameterHandleString() const;
 
-	const FString& GetName() const;
+	const FName GetName() const;
 
-	const FString& GetNamespace() const;
+	const FName GetNamespace() const;
 
 	bool IsEngineHandle() const;
 
@@ -49,16 +50,16 @@ public:
 	bool IsModuleHandle() const;
 
 public:
-	static const FString UserNamespace;
-	static const FString EngineNamespace;
-	static const FString SystemNamespace;
-	static const FString EmitterNamespace;
-	static const FString ParticleAttributeNamespace;
-	static const FString ModuleNamespace;
+	static const FName UserNamespace;
+	static const FName EngineNamespace;
+	static const FName SystemNamespace;
+	static const FName EmitterNamespace;
+	static const FName ParticleAttributeNamespace;
+	static const FName ModuleNamespace;
 	static const FString InitialPrefix;
 
 private:
-	FString ParameterHandleString;
-	FString Name;
-	FString Namespace;
+	FName ParameterHandleName;
+	FName Name;
+	FName Namespace;
 };
\ No newline at end of file
diff --git a/Engine/Plugins/FX/Niagara/Source/NiagaraEditor/Public/ViewModels/Stack/NiagaraStackAddModuleItem.h b/Engine/Plugins/FX/Niagara/Source/NiagaraEditor/Public/ViewModels/Stack/NiagaraStackAddModuleItem.h
--- a/Engine/Plugins/FX/Niagara/Source/NiagaraEditor/Public/ViewModels/Stack/NiagaraStackAddModuleItem.h
+++ b/Engine/Plugins/FX/Niagara/Source/NiagaraEditor/Public/ViewModels/Stack/NiagaraStackAddModuleItem.h
@@ -33,7 +33,7 @@ public:
 
 	virtual void GetNewParameterAvailableTypes(TArray<FNiagaraTypeDefinition>& OutAvailableTypes) const;
 
-	virtual TOptional<FString> GetNewParameterNamespace() const;
+	virtual TOptional<FName> GetNewParameterNamespace() const;
 
 	virtual ENiagaraScriptUsage GetOutputUsage() const PURE_VIRTUAL(UNiagaraStackAddModuleItem::GetOutputUsage, return ENiagaraScriptUsage::EmitterSpawnScript;);
 
diff --git a/Engine/Plugins/FX/Niagara/Source/NiagaraEditor/Public/ViewModels/Stack/NiagaraStackAddScriptModuleItem.h b/Engine/Plugins/FX/Niagara/Source/NiagaraEditor/Public/ViewModels/Stack/NiagaraStackAddScriptModuleItem.h
--- a/Engine/Plugins/FX/Niagara/Source/NiagaraEditor/Public/ViewModels/Stack/NiagaraStackAddScriptModuleItem.h
+++ b/Engine/Plugins/FX/Niagara/Source/NiagaraEditor/Public/ViewModels/Stack/NiagaraStackAddScriptModuleItem.h
@@ -21,7 +21,7 @@ public:
 
 	virtual void GetNewParameterAvailableTypes(TArray<FNiagaraTypeDefinition>& OutAvailableTypes) const override;
 
-	virtual TOptional<FString> GetNewParameterNamespace() const override;
+	virtual TOptional<FName> GetNewParameterNamespace() const override;
 
 	virtual ENiagaraScriptUsage GetOutputUsage() const override;
 
diff --git a/Engine/Plugins/FX/Niagara/Source/NiagaraEditor/Public/ViewModels/Stack/NiagaraStackFunctionInput.h b/Engine/Plugins/FX/Niagara/Source/NiagaraEditor/Public/ViewModels/Stack/NiagaraStackFunctionInput.h
--- a/Engine/Plugins/FX/Niagara/Source/NiagaraEditor/Public/ViewModels/Stack/NiagaraStackFunctionInput.h
+++ b/Engine/Plugins/FX/Niagara/Source/NiagaraEditor/Public/ViewModels/Stack/NiagaraStackFunctionInput.h
@@ -67,7 +67,7 @@ public:
 		UNiagaraStackEditorData& InStackEditorData,
 		UNiagaraNodeFunctionCall& InModuleNode,
 		UNiagaraNodeFunctionCall& InInputFunctionCallNode,
-		FString InInputParameterHandle,
+		FName InInputParameterHandle,
 		FNiagaraTypeDefinition InInputType);
 
 	/** Gets the function call node which owns this input. */
@@ -152,10 +152,10 @@ public:
 	void SetIsRenamePending(bool bIsRenamePending);
 
 	/** Renames this input to the name specified. */
-	void RenameInput(FString NewName);
+	void RenameInput(FName NewName);
 
 	/** Gets the namespaces which new parameters for this input can be read from. */
-	void GetNamespacesForNewParameters(TArray<FString>& OutNamespacesForNewParameters) const;
+	void GetNamespacesForNewParameters(TArray<FName>& OutNamespacesForNewParameters) const;
 
 	/** Gets a multicast delegate which is called whenever the value on this input changes. */
 	FOnValueChanged& OnValueChanged();
diff --git a/Engine/Plugins/FX/Niagara/Source/NiagaraEditor/Public/ViewModels/Stack/NiagaraStackModuleItemOutput.h b/Engine/Plugins/FX/Niagara/Source/NiagaraEditor/Public/ViewModels/Stack/NiagaraStackModuleItemOutput.h
--- a/Engine/Plugins/FX/Niagara/Source/NiagaraEditor/Public/ViewModels/Stack/NiagaraStackModuleItemOutput.h
+++ b/Engine/Plugins/FX/Niagara/Source/NiagaraEditor/Public/ViewModels/Stack/NiagaraStackModuleItemOutput.h
@@ -22,7 +22,7 @@ public:
 	 * @param InFunctionCallNode The function call node representing the module in the stack graph which owns this Output.
 	 * @param InOutputParameterHandle The Namespace.Name handle for the Output to the owning module.
 	 */
-	void Initialize(TSharedRef<FNiagaraSystemViewModel> InSystemViewModel, TSharedRef<FNiagaraEmitterViewModel> InEmitterViewModel, UNiagaraNodeFunctionCall& InFunctionCallNode, FString InOutputParameterHandle);
+	void Initialize(TSharedRef<FNiagaraSystemViewModel> InSystemViewModel, TSharedRef<FNiagaraEmitterViewModel> InEmitterViewModel, UNiagaraNodeFunctionCall& InFunctionCallNode, FName InOutputParameterHandle);
 
 	//~ UNiagaraStackEntry interface
 	virtual FText GetDisplayName() const override;
diff --git a/Engine/Plugins/FX/Niagara/Source/NiagaraEditorWidgets/Private/SNiagaraStack.cpp b/Engine/Plugins/FX/Niagara/Source/NiagaraEditorWidgets/Private/SNiagaraStack.cpp
--- a/Engine/Plugins/FX/Niagara/Source/NiagaraEditorWidgets/Private/SNiagaraStack.cpp
+++ b/Engine/Plugins/FX/Niagara/Source/NiagaraEditorWidgets/Private/SNiagaraStack.cpp
@@ -168,7 +168,7 @@ private:
 
 	void OnNameTextCommitted(const FText& InText, ETextCommit::Type InCommitType)
 	{
-		FunctionInput->RenameInput(InText.ToString());
+		FunctionInput->RenameInput(*InText.ToString());
 	}
 
 private:
@@ -423,7 +423,7 @@ private:
 
 	FText GetLinkedValueHandleText() const
 	{
-		return FText::FromString(FunctionInput->GetLinkedValueHandle().GetParameterHandleString());
+		return FText::FromName(FunctionInput->GetLinkedValueHandle().GetParameterHandleString());
 	}
 
 	FText GetDataValueText() const
@@ -556,10 +556,10 @@ private:
 			MenuBuilder.BeginSection(NAME_None, SectionDisplayText);
 			for (const FNiagaraParameterHandle& Handle : Handles)
 			{
-				FText HandleDisplayName = FText::FromString(FName::NameToDisplayString(Handle.GetName(), false));
+				FText HandleDisplayName = FText::FromString(FName::NameToDisplayString(Handle.GetName().ToString(), false));
 				MenuBuilder.AddMenuEntry(
 					HandleDisplayName,
-					FText::Format(MapInputFormat, FText::FromString(Handle.GetParameterHandleString())),
+					FText::Format(MapInputFormat, FText::FromName(Handle.GetParameterHandleString())),
 					FSlateIcon(),
 					FUIAction(FExecuteAction::CreateSP(this, &SNiagaraStackFunctionInputValue::ParameterHandleSelected, Handle)));
 			}
@@ -574,10 +574,10 @@ private:
 		MenuBuilder.BeginSection(NAME_None, LOCTEXT("OtherSection", "Other"));
 		for (const FNiagaraParameterHandle OtherHandle : OtherHandles)
 		{
-			FText HandleDisplayName = FText::FromString(FName::NameToDisplayString(OtherHandle.GetParameterHandleString(), false));
+			FText HandleDisplayName = FText::FromString(FName::NameToDisplayString(OtherHandle.GetParameterHandleString().ToString(), false));
 			MenuBuilder.AddMenuEntry(
 				HandleDisplayName,
-				FText::Format(MapInputFormat, FText::FromString(OtherHandle.GetParameterHandleString())),
+				FText::Format(MapInputFormat, FText::FromName(OtherHandle.GetParameterHandleString())),
 				FSlateIcon(),
 				FUIAction(FExecuteAction::CreateSP(this, &SNiagaraStackFunctionInputValue::ParameterHandleSelected, OtherHandle)));
 		}
@@ -589,24 +589,27 @@ private:
 		}
 
 		// Read from new attribute
-		TArray<FString> AvailableNamespaces;
+		TArray<FName> AvailableNamespaces;
 		FunctionInput->GetNamespacesForNewParameters(AvailableNamespaces);
 
 		TArray<FString> InputNames;
 		for (int32 i = FunctionInput->GetInputParameterHandlePath().Num() - 1; i >= 0; i--)
 		{
-			InputNames.Add(FunctionInput->GetInputParameterHandlePath()[i].GetName());
+			InputNames.Add(FunctionInput->GetInputParameterHandlePath()[i].GetName().ToString());
 		}
-		FString InputName = FString::Join(InputNames, TEXT("."));
+		FName InputName = *FString::Join(InputNames, TEXT("."));
 
-		for (const FString& AvailableNamespace : AvailableNamespaces)
+		for (const FName AvailableNamespace : AvailableNamespaces)
 		{
 			FNiagaraParameterHandle HandleToRead(AvailableNamespace, InputName);
 			bool bCanExecute = AvailableHandles.Contains(HandleToRead) == false;
 
+			FFormatNamedArguments Args;
+			Args.Add(TEXT("AvailableNamespace"), FText::FromName(AvailableNamespace));
+
 			MenuBuilder.AddMenuEntry(
-				FText::Format(LOCTEXT("ReadLabelFormat", "Read from new {0} parameter"), FText::FromString(AvailableNamespace)),
-				FText::Format(LOCTEXT("ReadToolTipFormat", "Read this input from a new parameter in the {0} namespace."), FText::FromString(AvailableNamespace)), 
+				FText::Format(LOCTEXT("ReadLabelFormat", "Read from new {AvailableNamespace} parameter"), Args),
+				FText::Format(LOCTEXT("ReadToolTipFormat", "Read this input from a new parameter in the {AvailableNamespace} namespace."), Args), 
 				FSlateIcon(),
 				FUIAction(
 					FExecuteAction::CreateSP(this, &SNiagaraStackFunctionInputValue::ParameterHandleSelected, HandleToRead),
@@ -1172,7 +1175,7 @@ private:
 		}
 
 		// Generate actions for setting new typed parameters.
-		TOptional<FString> NewParameterNamespace = AddModuleItem->GetNewParameterNamespace();
+		TOptional<FName> NewParameterNamespace = AddModuleItem->GetNewParameterNamespace();
 		if (NewParameterNamespace.IsSet())
 		{
 			TArray<FNiagaraTypeDefinition> AvailableTypes;
@@ -1184,7 +1187,7 @@ private:
 				FText CategoryName = LOCTEXT("CreateNewParameterCategory", "Create New Parameter");
 
 				FNiagaraParameterHandle NewParameterHandle(NewParameterNamespace.GetValue(), *(TEXT("New") + AvailableType.GetName()));
-				FNiagaraVariable NewParameter(AvailableType, *NewParameterHandle.GetParameterHandleString());
+				FNiagaraVariable NewParameter(AvailableType, NewParameterHandle.GetParameterHandleString());
 				TSharedPtr<FNiagaraStackGraphSchemaAction> NewNodeAction(new FNiagaraStackGraphSchemaAction(CategoryName, NameText, Tooltip, 0, FText(),
 					FNiagaraStackGraphSchemaAction::FOnPerformStackAction::CreateUObject(AddModuleItem, &UNiagaraStackAddModuleItem::AddParameterModule, NewParameter, true)));
 				OutAllActions.AddAction(NewNodeAction);
diff --git a/Engine/Plugins/Runtime/GameplayAbilities/Source/GameplayAbilitiesEditor/Private/GameplayAbilitiesGraphPanelPinFactory.h b/Engine/Plugins/Runtime/GameplayAbilities/Source/GameplayAbilitiesEditor/Private/GameplayAbilitiesGraphPanelPinFactory.h
--- a/Engine/Plugins/Runtime/GameplayAbilities/Source/GameplayAbilitiesEditor/Private/GameplayAbilitiesGraphPanelPinFactory.h
+++ b/Engine/Plugins/Runtime/GameplayAbilities/Source/GameplayAbilitiesEditor/Private/GameplayAbilitiesGraphPanelPinFactory.h
@@ -14,8 +14,7 @@ class FGameplayAbilitiesGraphPanelPinFactory: public FGraphPanelPinFactory
 {
 	virtual TSharedPtr<class SGraphPin> CreatePin(class UEdGraphPin* InPin) const override
 	{
-		const UEdGraphSchema_K2* K2Schema = GetDefault<UEdGraphSchema_K2>();
-		if (InPin->PinType.PinCategory == K2Schema->PC_Struct && InPin->PinType.PinSubCategoryObject == FGameplayAttribute::StaticStruct())
+		if (InPin->PinType.PinCategory == UEdGraphSchema_K2::PC_Struct && InPin->PinType.PinSubCategoryObject == FGameplayAttribute::StaticStruct())
 		{
 			return SNew(SGameplayAttributeGraphPin, InPin);
 		}
diff --git a/Engine/Plugins/Runtime/GameplayAbilities/Source/GameplayAbilitiesEditor/Private/K2Node_LatentAbilityCall.cpp b/Engine/Plugins/Runtime/GameplayAbilities/Source/GameplayAbilitiesEditor/Private/K2Node_LatentAbilityCall.cpp
--- a/Engine/Plugins/Runtime/GameplayAbilities/Source/GameplayAbilitiesEditor/Private/K2Node_LatentAbilityCall.cpp
+++ b/Engine/Plugins/Runtime/GameplayAbilities/Source/GameplayAbilitiesEditor/Private/K2Node_LatentAbilityCall.cpp
@@ -96,7 +96,7 @@ void UK2Node_LatentAbilityCall::ValidateNodeDuringCompilation(class FCompilerRes
 		{
 			if (Property->GetBoolMetaData(FK2Node_LatentAbilityCallHelper_RequiresConnection))
 			{
-				if (UEdGraphPin* DelegateExecPin = FindPin(Property->GetName()))
+				if (UEdGraphPin* DelegateExecPin = FindPin(Property->GetFName()))
 				{
 					if (DelegateExecPin->LinkedTo.Num() < 1)
 					{
diff --git a/Engine/Plugins/Runtime/PhysXVehicles/Source/PhysXVehicles/Public/WheeledVehicle.h b/Engine/Plugins/Runtime/PhysXVehicles/Source/PhysXVehicles/Public/WheeledVehicle.h
--- a/Engine/Plugins/Runtime/PhysXVehicles/Source/PhysXVehicles/Public/WheeledVehicle.h
+++ b/Engine/Plugins/Runtime/PhysXVehicles/Source/PhysXVehicles/Public/WheeledVehicle.h
@@ -12,7 +12,7 @@ class FDebugDisplayInfo;
 /**
  * WheeledVehicle is the base wheeled vehicle pawn actor.
  * By default it uses UWheeledVehicleMovementComponent4W for its simulation, but this can be overridden by inheriting from the class and modifying its constructor like so:
- * Super(ObjectInitializer.SetDefautSubobjectClass<UMyMovement>(VehicleMovementComponentName))
+ * Super(ObjectInitializer.SetDefaultSubobjectClass<UMyMovement>(VehicleMovementComponentName))
  * Where UMyMovement is the new movement type that inherits from UWheeledVehicleMovementComponent
  * 
  * @see https://docs.unrealengine.com/latest/INT/Engine/Physics/Vehicles/VehicleUserGuide/
diff --git a/Engine/Plugins/Runtime/Steam/SteamAudio/Source/SteamAudioEditor/Private/TickableNotification.cpp b/Engine/Plugins/Runtime/Steam/SteamAudio/Source/SteamAudioEditor/Private/TickableNotification.cpp
--- a/Engine/Plugins/Runtime/Steam/SteamAudio/Source/SteamAudioEditor/Private/TickableNotification.cpp
+++ b/Engine/Plugins/Runtime/Steam/SteamAudio/Source/SteamAudioEditor/Private/TickableNotification.cpp
@@ -110,11 +110,6 @@ namespace SteamAudio
 		}
 	}
 
-	bool FTickableNotification::IsTickable() const
-	{
-		return true;
-	}
-
 	TStatId FTickableNotification::GetStatId() const
 	{
 		RETURN_QUICK_DECLARE_CYCLE_STAT(FTickableNotification, STATGROUP_Tickables);
diff --git a/Engine/Plugins/Runtime/Steam/SteamAudio/Source/SteamAudioEditor/Private/TickableNotification.h b/Engine/Plugins/Runtime/Steam/SteamAudio/Source/SteamAudioEditor/Private/TickableNotification.h
--- a/Engine/Plugins/Runtime/Steam/SteamAudio/Source/SteamAudioEditor/Private/TickableNotification.h
+++ b/Engine/Plugins/Runtime/Steam/SteamAudio/Source/SteamAudioEditor/Private/TickableNotification.h
@@ -33,7 +33,7 @@ namespace SteamAudio
 
 	protected:
 		virtual void Tick(float DeltaTime) override;
-		virtual bool IsTickable() const override;
+		virtual ETickableTickType GetTickableTickType() const override { return ETickableTickType::Always; }
 		virtual TStatId GetStatId() const override;
 
 	private:
diff --git a/Engine/Plugins/ScriptPlugin/Source/ScriptEditorPlugin/Private/ScriptBlueprintCompiler.cpp b/Engine/Plugins/ScriptPlugin/Source/ScriptEditorPlugin/Private/ScriptBlueprintCompiler.cpp
--- a/Engine/Plugins/ScriptPlugin/Source/ScriptEditorPlugin/Private/ScriptBlueprintCompiler.cpp
+++ b/Engine/Plugins/ScriptPlugin/Source/ScriptEditorPlugin/Private/ScriptBlueprintCompiler.cpp
@@ -43,37 +43,37 @@ void FScriptBlueprintCompiler::CreateClassVariablesFromBlueprint()
 	UScriptBlueprintGeneratedClass* NewScripClass = CastChecked<UScriptBlueprintGeneratedClass>(NewClass);
 	NewScripClass->ScriptProperties.Empty();
 
-	for (auto& Field : ScriptDefinedFields)
+	for (FScriptField& Field : ScriptDefinedFields)
 	{
 		UClass* InnerType = Field.Class;
 		if (Field.Class->IsChildOf(UProperty::StaticClass()))
 		{
-			FString PinCategory;
+			FName PinCategory;
 			if (Field.Class->IsChildOf(UStrProperty::StaticClass()))
 			{
-				PinCategory = Schema->PC_String;
+				PinCategory = UEdGraphSchema_K2::PC_String;
 			}
 			else if (Field.Class->IsChildOf(UFloatProperty::StaticClass()))
 			{
-				PinCategory = Schema->PC_Float;
+				PinCategory = UEdGraphSchema_K2::PC_Float;
 			}
 			else if (Field.Class->IsChildOf(UIntProperty::StaticClass()))
 			{
-				PinCategory = Schema->PC_Int;
+				PinCategory = UEdGraphSchema_K2::PC_Int;
 			}
 			else if (Field.Class->IsChildOf(UBoolProperty::StaticClass()))
 			{
-				PinCategory = Schema->PC_Boolean;
+				PinCategory = UEdGraphSchema_K2::PC_Boolean;
 			}
 			else if (Field.Class->IsChildOf(UObjectProperty::StaticClass()))
 			{
-				PinCategory = Schema->PC_Object;
+				PinCategory = UEdGraphSchema_K2::PC_Object;
 				// @todo: some scripting extensions (that are strongly typed) can handle this better
 				InnerType = UObject::StaticClass();
 			}
-			if (!PinCategory.IsEmpty())
+			if (!PinCategory.IsNone())
 			{
-				FEdGraphPinType ScriptPinType(PinCategory, FString(), InnerType, EPinContainerType::None, false, FEdGraphTerminalType());
+				FEdGraphPinType ScriptPinType(PinCategory, NAME_None, InnerType, EPinContainerType::None, false, FEdGraphTerminalType());
 				UProperty* ScriptProperty = CreateVariable(Field.Name, ScriptPinType);
 				if (ScriptProperty)
 				{
@@ -103,7 +103,7 @@ void FScriptBlueprintCompiler::CreateScriptContextProperty()
 
 	if (ContextClass)
 	{
-		FEdGraphPinType ScriptContextPinType(Schema->PC_Object, FString(), ContextClass, EPinContainerType::None, false, FEdGraphTerminalType());
+		FEdGraphPinType ScriptContextPinType(UEdGraphSchema_K2::PC_Object, NAME_None, ContextClass, EPinContainerType::None, false, FEdGraphTerminalType());
 		ContextProperty = CastChecked<UObjectProperty>(CreateVariable(TEXT("Generated_ScriptContext"), ScriptContextPinType));
 		ContextProperty->SetPropertyFlags(CPF_ContainsInstancedReference | CPF_InstancedReference);
 	}
@@ -160,7 +160,7 @@ void FScriptBlueprintCompiler::CreateScriptDefinedFunction(FScriptField& Field)
 
 	// Link nodes together
 	UEdGraphPin* ExecPin = Schema->FindExecutionPin(*EntryNode, EGPD_Output);
-	UEdGraphPin* GetVariableOutPin = GetVariableNode->FindPinChecked(ContextProperty->GetName());
+	UEdGraphPin* GetVariableOutPin = GetVariableNode->FindPinChecked(ContextProperty->GetFName());
 	UEdGraphPin* CallFunctionPin = Schema->FindExecutionPin(*CallFunctionNode, EGPD_Input);
 	UEdGraphPin* FunctionTargetPin = CallFunctionNode->FindPinChecked(TEXT("self"));
 	ExecPin->MakeLinkTo(CallFunctionPin);
diff --git a/Engine/Plugins/Tests/EditorTests/Source/EditorTests/Private/UnrealEd/BlueprintAutomationTests.cpp b/Engine/Plugins/Tests/EditorTests/Source/EditorTests/Private/UnrealEd/BlueprintAutomationTests.cpp
--- a/Engine/Plugins/Tests/EditorTests/Source/EditorTests/Private/UnrealEd/BlueprintAutomationTests.cpp
+++ b/Engine/Plugins/Tests/EditorTests/Source/EditorTests/Private/UnrealEd/BlueprintAutomationTests.cpp
@@ -678,8 +678,8 @@ public:
 			
 			FEdGraphPinType const& PinType = MalformedPin->PinType;
 			// only object pins would reference the unloaded blueprint
-			if (!PinType.PinSubCategoryObject.IsValid() || ((PinType.PinCategory != K2Schema->PC_Object) && 
-				(PinType.PinCategory != K2Schema->PSC_Self) && (PinType.PinCategory != K2Schema->PC_Interface)))
+			if (!PinType.PinSubCategoryObject.IsValid() || ((PinType.PinCategory != UEdGraphSchema_K2::PC_Object) && 
+				(PinType.PinCategory != UEdGraphSchema_K2::PSC_Self) && (PinType.PinCategory != UEdGraphSchema_K2::PC_Interface)))
 			{
 				continue;
 			}
diff --git a/Engine/Plugins/Tests/EditorTests/Source/EditorTests/Private/UnrealEd/EditorAutomationTests.cpp b/Engine/Plugins/Tests/EditorTests/Source/EditorTests/Private/UnrealEd/EditorAutomationTests.cpp
--- a/Engine/Plugins/Tests/EditorTests/Source/EditorTests/Private/UnrealEd/EditorAutomationTests.cpp
+++ b/Engine/Plugins/Tests/EditorTests/Source/EditorTests/Private/UnrealEd/EditorAutomationTests.cpp
@@ -766,9 +766,9 @@ void FTraceAllTimelinesAutomationTest::GetTests(TArray<FString>& OutBeautifiedNa
 namespace TraceHelper
 {
 	static const FName CosmeticMacroName(TEXT("Can Execute Cosmetic Events"));
-	static const FString CosmeticCheckedPinName(TEXT("True"));
+	static const FName CosmeticCheckedPinName(TEXT("True"));
 	static const FName SetTimerName(TEXT("K2_SetTimer"));
-	static const FString SetTimerFunctionFieldName(TEXT("FunctionName"));
+	static const FName SetTimerFunctionFieldName(TEXT("FunctionName"));
 
 	typedef TArray<UK2Node_MacroInstance*, TInlineAllocator<2>> TContextStack;
 	typedef TMap<UK2Node_MacroInstance*, struct FVisitedTracker> TGraphNodesVisited;
@@ -901,7 +901,6 @@ bool FTraceAllTimelinesAutomationTest::RunTest(const FString& BlueprintName)
 
 		for (auto TimelineIter = AllTimelines.CreateIterator(); TimelineIter; ++TimelineIter)
 		{
-			const UEdGraphSchema_K2* K2Schema = GetDefault<UEdGraphSchema_K2>();
 			UK2Node_Timeline* TimelinePtr = *TimelineIter;
 			UE_LOG(LogEditorAutomationTests, Log, TEXT("TraceTimeline: %s [%s]"), *TimelinePtr->GetPathName(), *TimelinePtr->TimelineName.ToString());
 
@@ -1056,7 +1055,7 @@ bool FTraceAllTimelinesAutomationTest::RunTest(const FString& BlueprintName)
 					for (auto PinIter = CurrentNode->Pins.CreateConstIterator(); PinIter; ++PinIter)
 					{
 						const UEdGraphPin* CurrentPin = CurrentScope.Pin ? CurrentScope.Pin : *PinIter;
-						if (CurrentPin->Direction == EGPD_Input && CurrentPin->PinType.PinCategory == K2Schema->PC_Exec)
+						if (CurrentPin->Direction == EGPD_Input && CurrentPin->PinType.PinCategory == UEdGraphSchema_K2::PC_Exec)
 						{
 							for (auto LinkedToIter = CurrentPin->LinkedTo.CreateConstIterator(); LinkedToIter; ++LinkedToIter)
 							{
@@ -1080,14 +1079,14 @@ bool FTraceAllTimelinesAutomationTest::RunTest(const FString& BlueprintName)
 												bCosmeticChain = (OtherPin->PinName != TraceHelper::CosmeticCheckedPinName);
 											}
 											// Don't bother trying to identify pins on this special node, all it does is change the cosmetic chain state.
-											OtherPin = NULL;
+											OtherPin = nullptr;
 										}
 									}
 								}
 
 								// We only really care to distinguish pins for tunnel nodes,
 								// because those might have different input pins hooked up to the logic ending in this output pin
-								TraceHelper::AddNode(OtherPinNode, NodeScopeStack, GraphNodesVisited, CurrentScope, bCosmeticChain, Tunnel ? OtherPin : NULL);
+								TraceHelper::AddNode(OtherPinNode, NodeScopeStack, GraphNodesVisited, CurrentScope, bCosmeticChain, Tunnel ? OtherPin : nullptr);
 							}
 						}
 
diff --git a/Engine/Plugins/Tests/EditorTests/Source/EditorTests/Private/UnrealEd/EditorBuildPromotionTests.cpp b/Engine/Plugins/Tests/EditorTests/Source/EditorTests/Private/UnrealEd/EditorBuildPromotionTests.cpp
--- a/Engine/Plugins/Tests/EditorTests/Source/EditorTests/Private/UnrealEd/EditorBuildPromotionTests.cpp
+++ b/Engine/Plugins/Tests/EditorTests/Source/EditorTests/Private/UnrealEd/EditorBuildPromotionTests.cpp
@@ -1863,14 +1863,12 @@ namespace BuildPromotionTestHelper
 				return true;
 			}
 
-			const UEdGraphSchema_K2* K2Schema = GetDefault<UEdGraphSchema_K2>();
-			
 			// Add string member variable
-			FEdGraphPinType StringPinType(K2Schema->PC_String, FString(), nullptr, EPinContainerType::None, false, FEdGraphTerminalType());
+			FEdGraphPinType StringPinType(UEdGraphSchema_K2::PC_String, NAME_None, nullptr, EPinContainerType::None, false, FEdGraphTerminalType());
 			FBlueprintEditorUtils::AddMemberVariable(BlueprintObject, EditorBuildPromotionTestUtils::BlueprintStringVariableName, StringPinType);
 
 			// Add print string node connected to the BeginPlay node; save it for use later
-			UEdGraphPin* PlayThenPin = PostBeginPlayEventNode->FindPin(K2Schema->PN_Then);
+			UEdGraphPin* PlayThenPin = PostBeginPlayEventNode->FindPin(UEdGraphSchema_K2::PN_Then);
 			PrintNode = EditorBuildPromotionTestUtils::AddPrintStringNode(BlueprintObject, EventGraph, PlayThenPin);
 			Test->TestNotNull(TEXT("Print String node added"), PrintNode);
 			
diff --git a/Engine/Source/Developer/BlueprintCompilerCppBackend/Private/BlueprintCompilerCppBackend.cpp b/Engine/Source/Developer/BlueprintCompilerCppBackend/Private/BlueprintCompilerCppBackend.cpp
--- a/Engine/Source/Developer/BlueprintCompilerCppBackend/Private/BlueprintCompilerCppBackend.cpp
+++ b/Engine/Source/Developer/BlueprintCompilerCppBackend/Private/BlueprintCompilerCppBackend.cpp
@@ -1,4 +1,4 @@
-// Copyright 1998-2017 Epic Games, Inc. All Rights Reserved.
+﻿// Copyright 1998-2017 Epic Games, Inc. All Rights Reserved.
 
 #include "BlueprintCompilerCppBackend.h"
 #include "UObject/UnrealType.h"
@@ -219,7 +219,7 @@ static FString GenerateCastRHS(FEmitterLocalContext& EmitterContext, UClass* Cla
 {
 	check(ClassPtr != nullptr);
 
-	auto BPGC = Cast<UBlueprintGeneratedClass>(ClassPtr);
+	UBlueprintGeneratedClass* BPGC = Cast<UBlueprintGeneratedClass>(ClassPtr);
 	if (BPGC && !EmitterContext.Dependencies.WillClassBeConverted(BPGC))
 	{
 		const FString NativeClass = FEmitHelper::GetCppName(EmitterContext.GetFirstNativeOrConvertedClass(ClassPtr));
@@ -486,8 +486,8 @@ FString FBlueprintCompilerCppBackend::EmitSwitchValueStatmentInner(FEmitterLocal
 	const int32 TermsBeforeCases = 1;
 	const int32 TermsPerCase = 2;
 	const int32 NumCases = ((Statement.RHS.Num() - 2) / TermsPerCase);
-	auto IndexTerm = Statement.RHS[0];
-	auto DefaultValueTerm = Statement.RHS.Last();
+	FBPTerminal* IndexTerm = Statement.RHS[0];
+	FBPTerminal* DefaultValueTerm = Statement.RHS.Last();
 
 	const uint32 CppTemplateTypeFlags = EPropertyExportCPPFlags::CPPF_CustomTypeName
 		| EPropertyExportCPPFlags::CPPF_NoConst | EPropertyExportCPPFlags::CPPF_NoRef
@@ -521,9 +521,11 @@ FString FBlueprintCompilerCppBackend::EmitSwitchValueStatmentInner(FEmitterLocal
 				FEmitHelper::GenerateAutomaticCast(EmitterContext, LType, Term->Type, DefaultValueTerm->AssociatedVarProperty, Term->AssociatedVarProperty, BeginCast, EndCast, true);
 			}
 
-			const FString TermEvaluation = TermToText(EmitterContext, Term, ENativizedTermUsage::UnspecifiedOrReference); //should bGetter be false ?
+			const FString TermEvaluation = TermToText(EmitterContext, Term, ENativizedTermUsage::Getter);
 			const FString CastedTerm = FString::Printf(TEXT("%s%s%s"), *BeginCast, *TermEvaluation, *EndCast);
-			if (Term->bIsLiteral) //TODO it should be done for every term, that cannot be handled as reference.
+			// If the term is a literal, a weak pointer, or TSubClassOf we can't take a clean reference to it for the TSwitchPair  
+			// so it needs to be cached in to a local.
+			if (Term->bIsLiteral || Term->Type.bIsWeakPointer || Term->Type.PinCategory == UEdGraphSchema_K2::PC_Class)
 			{
 				const FString LocalVarName = EmitterContext.GenerateUniqueLocalName();
 				EmitterContext.AddLine(FString::Printf(TEXT("%s %s = %s;"), *ValueDeclaration, *LocalVarName, *CastedTerm));
@@ -644,7 +646,7 @@ FString FBlueprintCompilerCppBackend::EmitMethodInputParameterList(FEmitterLocal
 				FString BeginCast;
 				FString CloseCast;
 				FEdGraphPinType LType;
-				auto Schema = GetDefault<UEdGraphSchema_K2>();
+				const UEdGraphSchema_K2* Schema = GetDefault<UEdGraphSchema_K2>();
 				check(Schema);
 				ENativizedTermUsage TermUsage = ENativizedTermUsage::UnspecifiedOrReference;
 				if (Schema->ConvertPropertyToPinType(FuncParamProperty, LType))
@@ -783,7 +785,7 @@ FString FBlueprintCompilerCppBackend::EmitCallStatmentInner(FEmitterLocalContext
 
 			FString BeginCast;
 			FEdGraphPinType RType;
-			auto Schema = GetDefault<UEdGraphSchema_K2>();
+			const UEdGraphSchema_K2* Schema = GetDefault<UEdGraphSchema_K2>();
 			check(Schema);
 			if (Schema->ConvertPropertyToPinType(FuncToCallReturnProperty, RType))
 			{
@@ -820,7 +822,7 @@ FString FBlueprintCompilerCppBackend::EmitCallStatmentInner(FEmitterLocalContext
 	}
 	else
 	{
-		auto OwnerBPGC = Cast<UBlueprintGeneratedClass>(FunctionOwner);
+		UBlueprintGeneratedClass* OwnerBPGC = Cast<UBlueprintGeneratedClass>(FunctionOwner);
 		const bool bUnconvertedClass = OwnerBPGC && !EmitterContext.Dependencies.WillClassBeConverted(OwnerBPGC);
 		const bool bIsCustomThunk = bStaticCall && ( Statement.FunctionToCall->GetBoolMetaData(TEXT("CustomThunk"))
 			|| Statement.FunctionToCall->HasMetaData(TEXT("CustomStructureParam"))
@@ -836,7 +838,7 @@ FString FBlueprintCompilerCppBackend::EmitCallStatmentInner(FEmitterLocalContext
 		}
 		else if (bStaticCall)
 		{
-			auto OwnerClass = Statement.FunctionToCall->GetOuterUClass();
+			UClass* OwnerClass = Statement.FunctionToCall->GetOuterUClass();
 			Result += bIsCustomThunk ? TEXT("FCustomThunkTemplates::") : FString::Printf(TEXT("%s::"), *FEmitHelper::GetCppName(OwnerClass));
 		}
 		else if (bCallOnDifferentObject) //@TODO: Badness, could be a self reference wired to another instance!
@@ -1014,8 +1016,8 @@ FString FBlueprintCompilerCppBackend::TermToText(FEmitterLocalContext& EmitterCo
 			bIsAccessible &= !Term->AssociatedVarProperty->HasAnyPropertyFlags(CPF_NativeAccessSpecifierPrivate)
 				&& ((bPropertyOfParent && bSelfContext) || !Term->AssociatedVarProperty->HasAnyPropertyFlags(CPF_NativeAccessSpecifierProtected));
 
-			auto MinimalClass = Term->AssociatedVarProperty->GetOwnerClass();
-			auto MinimalBPGC = Cast<UBlueprintGeneratedClass>(MinimalClass);
+			UClass* MinimalClass = Term->AssociatedVarProperty->GetOwnerClass();
+			UBlueprintGeneratedClass* MinimalBPGC = Cast<UBlueprintGeneratedClass>(MinimalClass);
 			if (MinimalBPGC && !EmitterContext.Dependencies.WillClassBeConverted(MinimalBPGC))
 			{
 				if (bSelfContext)
@@ -1098,7 +1100,7 @@ FString FBlueprintCompilerCppBackend::TermToText(FEmitterLocalContext& EmitterCo
 
 FString FBlueprintCompilerCppBackend::LatentFunctionInfoTermToText(FEmitterLocalContext& EmitterContext, FBPTerminal* Term, FBlueprintCompiledStatement* TargetLabel)
 {
-	auto LatentInfoStruct = FLatentActionInfo::StaticStruct();
+	UStruct* LatentInfoStruct = FLatentActionInfo::StaticStruct();
 
 	// Find the term name we need to fixup
 	FString FixupTermName;
diff --git a/Engine/Source/Developer/BlueprintNativeCodeGen/Private/BlueprintNativeCodeGenManifest.cpp b/Engine/Source/Developer/BlueprintNativeCodeGen/Private/BlueprintNativeCodeGenManifest.cpp
--- a/Engine/Source/Developer/BlueprintNativeCodeGen/Private/BlueprintNativeCodeGenManifest.cpp
+++ b/Engine/Source/Developer/BlueprintNativeCodeGen/Private/BlueprintNativeCodeGenManifest.cpp
@@ -33,9 +33,6 @@ namespace BlueprintNativeCodeGenManifestImpl
 	static const FString SourceSubDir         = TEXT("Source");
 	static const FString EditorModulePostfix  = TEXT("Editor");
 	static const int32   RootManifestId       = -1;
-	// if you change this placeholder value, then you should reflect those changes in UAT (in 
-	// AddBlueprintPluginPathArgument() - where it folds this string into the requested pathname)
-	static const FString PlatformPlaceholderPattern = TEXT("<PLAT>");
 
 	/**
 	 * Populates the provided manifest object with data from the specified file.
@@ -317,11 +314,27 @@ FUnconvertedDependencyRecord::FUnconvertedDependencyRecord(const FString& InGene
  //------------------------------------------------------------------------------
 FBlueprintNativeCodeGenPaths FBlueprintNativeCodeGenPaths::GetDefaultCodeGenPaths(const FName PlatformName)
 {
-	// NOTE: in UProjectPackagingSettings::PostEditChangeProperty() there are a hardcoded file path/name that are set to match these; 
-	//       if you alter these defaults then you need to update that and likely AddBlueprintPluginPathArgument() in UAT as well
-	FString DefaultPluginName = TEXT("NativizedAssets");
-	FString DefaultPluginDir = FPaths::Combine(*FPaths::Combine(*FPaths::ProjectIntermediateDir(), TEXT("Plugins")), *DefaultPluginName);
-	return FBlueprintNativeCodeGenPaths(DefaultPluginName, DefaultPluginDir, PlatformName);
+	static const FString DefaultPluginName = TEXT("NativizedAssets");
+
+	FString DefaultPluginPath = FPaths::Combine(*FPaths::ProjectIntermediateDir(), TEXT("Plugins"), *DefaultPluginName);
+	if (!PlatformName.IsNone())
+	{
+		for (PlatformInfo::FPlatformEnumerator PlatformIt = PlatformInfo::EnumeratePlatformInfoArray(); PlatformIt; ++PlatformIt)
+		{
+			if (PlatformIt->TargetPlatformName == PlatformName)
+			{
+				static const FName UBTTargetId_Win32 = FName(TEXT("Win32"));
+				static const FName UBTTargetId_Win64 = FName(TEXT("Win64"));
+				static const FName UBTTargetId_Windows = FName(TEXT("Windows"));
+
+				const FName UBTTargetId = (PlatformIt->UBTTargetId == UBTTargetId_Win32 || PlatformIt->UBTTargetId == UBTTargetId_Win64) ? UBTTargetId_Windows : PlatformIt->UBTTargetId;
+				DefaultPluginPath = FPaths::Combine(*DefaultPluginPath, *Lex::ToString(UBTTargetId), *Lex::ToString(PlatformIt->PlatformType));
+				break;
+			}
+		}
+	}
+
+	return FBlueprintNativeCodeGenPaths(DefaultPluginName, DefaultPluginPath, PlatformName);
 }
 
 //------------------------------------------------------------------------------
@@ -347,15 +360,6 @@ FBlueprintNativeCodeGenPaths::FBlueprintNativeCodeGenPaths(const FString& Plugin
 	{
 		PluginsDir = FPaths::GetPath(PluginsDir);
 	}
-	// if you change the replacement of this placeholder value, then you should reflect those changes in UAT (in 
-	// AddBlueprintPluginPathArgument() - where it folds this string into the requested pathname)
-	if (!PlatformName.IsNone() && PluginsDir.ReplaceInline(*BlueprintNativeCodeGenManifestImpl::PlatformPlaceholderPattern, *PlatformName.ToString()) == 0)
-	{
-		// DON'T alter the path that was likely passed down through the command line
-		//
-		// // if there was no platform placeholder in the string, append one to the end
-		// PluginsDir = FPaths::Combine(*PluginsDir, *(PluginNameIn + TEXT("_") + PlatformName.ToString()));
-	}
 }
 
 //------------------------------------------------------------------------------
@@ -380,7 +384,7 @@ FString FBlueprintNativeCodeGenPaths::ManifestFilePath(const int32 ChunkId) cons
 //------------------------------------------------------------------------------
 FString FBlueprintNativeCodeGenPaths::PluginRootDir() const
 {
-	return FPaths::Combine(*PluginsDir, *PluginName);
+	return PluginsDir;
 }
 
 //------------------------------------------------------------------------------
diff --git a/Engine/Source/Developer/BlueprintNativeCodeGen/Private/BlueprintNativeCodeGenModule.cpp b/Engine/Source/Developer/BlueprintNativeCodeGen/Private/BlueprintNativeCodeGenModule.cpp
--- a/Engine/Source/Developer/BlueprintNativeCodeGen/Private/BlueprintNativeCodeGenModule.cpp
+++ b/Engine/Source/Developer/BlueprintNativeCodeGen/Private/BlueprintNativeCodeGenModule.cpp
@@ -40,6 +40,10 @@ public:
 	{
 	}
 
+	//~ Begin IModuleInterface interface
+	virtual void ShutdownModule();
+	//~ End IModuleInterface interface
+
 	//~ Begin IBlueprintNativeCodeGenModule interface
 	virtual void Convert(UPackage* Package, ESavePackageResult ReplacementType, const FName PlatformName) override;
 	virtual void SaveManifest() override;
@@ -52,7 +56,7 @@ public:
 
 	FFileHelper::EEncodingOptions ForcedEncoding() const
 	{
-		return FFileHelper::EEncodingOptions::ForceUTF8;
+		return FFileHelper::EEncodingOptions::ForceUTF8WithoutBOM;
 	}
 	virtual const FCompilerNativizationOptions& GetNativizationOptionsForPlatform(const ITargetPlatform* Platform) const override;
 	virtual void FillPlatformNativizationDetails(const ITargetPlatform* Platform, FPlatformNativizationDetails& OutDetails) override;
@@ -268,11 +272,7 @@ void FBlueprintNativeCodeGenModule::Initialize(const FNativeCodeGenInitData& Ini
 	// Each platform will need a manifest, because each platform could cook different assets:
 	for (const FPlatformNativizationDetails& Platform : InitData.CodegenTargets)
 	{
-		FString TargetPath = InitData.DestPluginPath; 
-		if (TargetPath.IsEmpty())
-		{
-			TargetPath = FBlueprintNativeCodeGenPaths::GetDefaultPluginPath(Platform.PlatformName);
-		}
+		const FString TargetPath = FBlueprintNativeCodeGenPaths::GetDefaultPluginPath(Platform.PlatformName);
 		const FBlueprintNativeCodeGenManifest& Manifest = *Manifests.Add(Platform.PlatformName, TUniquePtr<FBlueprintNativeCodeGenManifest>(new FBlueprintNativeCodeGenManifest(TargetPath, Platform.CompilerNativizationOptions, InitData.ManifestIdentifier)));
 
 		TargetPlatformNames.Add(Platform.PlatformName);
@@ -361,6 +361,18 @@ void FBlueprintNativeCodeGenModule::InitializeForRerunDebugOnly(const TArray<FPl
 	}
 }
 
+void FBlueprintNativeCodeGenModule::ShutdownModule()
+{
+	// Clear the current coordinator reference.
+	IBlueprintNativeCodeGenCore::Register(nullptr);
+
+	// Reset compiler module delegate function bindings.
+	IBlueprintCompilerCppBackendModule& BackEndModule = (IBlueprintCompilerCppBackendModule&)IBlueprintCompilerCppBackendModule::Get();
+	BackEndModule.GetIsFunctionUsedInADelegateCallback().Unbind();
+	BackEndModule.OnIsTargetedForConversionQuery().Unbind();
+	BackEndModule.OnIncludingUnconvertedBP().Unbind();
+}
+
 void FBlueprintNativeCodeGenModule::GenerateFullyConvertedClasses()
 {
 	TSharedPtr<FNativizationSummary> NativizationSummary(new FNativizationSummary());
@@ -856,7 +868,7 @@ EReplacementResult FBlueprintNativeCodeGenModule::IsTargetedForReplacement(const
 			};
 			if (Object && (IsEditorOnlyObject(Object) || IsDeveloperObject(Object)))
 			{
-				UE_LOG(LogBlueprintCodeGen, Warning, TEXT("Object %s depends on Editor or Development stuff. It shouldn't be cooked."), *GetPathNameSafe(Object));
+				UE_LOG(LogBlueprintCodeGen, Verbose, TEXT("Object %s depends on Editor or Development stuff. It shouldn't be cooked."), *GetPathNameSafe(Object));
 				return true;
 			}
 		}
diff --git a/Engine/Source/Developer/BlueprintNativeCodeGen/Public/BlueprintNativeCodeGenModule.h b/Engine/Source/Developer/BlueprintNativeCodeGen/Public/BlueprintNativeCodeGenModule.h
--- a/Engine/Source/Developer/BlueprintNativeCodeGen/Public/BlueprintNativeCodeGenModule.h
+++ b/Engine/Source/Developer/BlueprintNativeCodeGen/Public/BlueprintNativeCodeGenModule.h
@@ -27,10 +27,6 @@ struct FNativeCodeGenInitData
 	// Optional Manifest ManifestIdentifier, used for child cook processes that need a unique manifest name.
 	// The identifier is used to make a unique name for each platform that is converted.
 	int32 ManifestIdentifier;
-
-	// Optional, desired .uplugin file path for the nativized plugin (used so people can manually coordinate 
-	// cook/build processes).
-	FString DestPluginPath;
 };
 
 /** 
diff --git a/Engine/Source/Developer/DesktopPlatform/Private/PlatformInfo.cpp b/Engine/Source/Developer/DesktopPlatform/Private/PlatformInfo.cpp
--- a/Engine/Source/Developer/DesktopPlatform/Private/PlatformInfo.cpp
+++ b/Engine/Source/Developer/DesktopPlatform/Private/PlatformInfo.cpp
@@ -303,4 +303,24 @@ EPlatformType EPlatformTypeFromString(const FString& PlatformTypeName)
 
 } // namespace PlatformFamily
 
+namespace Lex
+{
+	FString ToString(const PlatformInfo::EPlatformType Value)
+	{
+		switch (Value)
+		{
+		case PlatformInfo::EPlatformType::Game:
+			return TEXT("Game");
+		case PlatformInfo::EPlatformType::Editor:
+			return TEXT("Editor");
+		case PlatformInfo::EPlatformType::Client:
+			return TEXT("Client");
+		case PlatformInfo::EPlatformType::Server:
+			return TEXT("Server");
+		}
+
+		return TEXT("");
+	}
+}
+
 #undef LOCTEXT_NAMESPACE
diff --git a/Engine/Source/Developer/DesktopPlatform/Public/PlatformInfo.h b/Engine/Source/Developer/DesktopPlatform/Public/PlatformInfo.h
--- a/Engine/Source/Developer/DesktopPlatform/Public/PlatformInfo.h
+++ b/Engine/Source/Developer/DesktopPlatform/Public/PlatformInfo.h
@@ -363,3 +363,13 @@ namespace PlatformInfo
 	*/
 	DESKTOPPLATFORM_API EPlatformType EPlatformTypeFromString(const FString& PlatformTypeName);
 }
+
+namespace Lex
+{
+	DESKTOPPLATFORM_API FString ToString(const PlatformInfo::EPlatformType Value);
+
+	inline void FromString(PlatformInfo::EPlatformType& OutValue, const TCHAR* Buffer)
+	{
+		OutValue = PlatformInfo::EPlatformTypeFromString(FString(Buffer));
+	}
+}
\ No newline at end of file
diff --git a/Engine/Source/Developer/LogVisualizer/Private/VisualLoggerRenderingActor.cpp b/Engine/Source/Developer/LogVisualizer/Private/VisualLoggerRenderingActor.cpp
--- a/Engine/Source/Developer/LogVisualizer/Private/VisualLoggerRenderingActor.cpp
+++ b/Engine/Source/Developer/LogVisualizer/Private/VisualLoggerRenderingActor.cpp
@@ -148,8 +148,6 @@ AVisualLoggerRenderingActor::~AVisualLoggerRenderingActor()
 
 void AVisualLoggerRenderingActor::ObjectSelectionChanged(const TArray<FName>& Selection)
 {
-	//PrimaryDebugShapes.Reset();
-
 	if (Selection.Num() > 0)
 	{
 		for (auto CurrentName : Selection)
@@ -169,11 +167,11 @@ void AVisualLoggerRenderingActor::ObjectSelectionChanged(const TArray<FName>& Se
 			}
 		}
 
-		for (auto Element : DebugShapesPerRow)
+		for (TMap<FName, FTimelineDebugShapes>::TIterator It(DebugShapesPerRow); It; ++It)
 		{
-			if (Selection.Find(Element.Key) == INDEX_NONE)
+			if (Selection.Find(It->Key) == INDEX_NONE)
 			{
-				DebugShapesPerRow.Remove(Element.Key);
+				It.RemoveCurrent();
 			}
 		}
 	}
diff --git a/Engine/Source/Developer/LogVisualizer/Private/VisualLoggerRenderingActor.h b/Engine/Source/Developer/LogVisualizer/Private/VisualLoggerRenderingActor.h
--- a/Engine/Source/Developer/LogVisualizer/Private/VisualLoggerRenderingActor.h
+++ b/Engine/Source/Developer/LogVisualizer/Private/VisualLoggerRenderingActor.h
@@ -62,7 +62,6 @@ private:
 public:
 	UPrimitiveComponent* RenderingComponent;
 
-	//FTimelineDebugShapes PrimaryDebugShapes;
 	FTimelineDebugShapes TestDebugShapes;
 
 	TArray<FName> CachedRowSelection;
diff --git a/Engine/Source/Editor/AIGraph/Classes/AIGraphSchema.h b/Engine/Source/Editor/AIGraph/Classes/AIGraphSchema.h
--- a/Engine/Source/Editor/AIGraph/Classes/AIGraphSchema.h
+++ b/Engine/Source/Editor/AIGraph/Classes/AIGraphSchema.h
@@ -94,7 +94,7 @@ class AIGRAPH_API UAIGraphSchema : public UEdGraphSchema
 	virtual class FConnectionDrawingPolicy* CreateConnectionDrawingPolicy(int32 InBackLayerID, int32 InFrontLayerID, float InZoomFactor, const FSlateRect& InClippingRect, class FSlateWindowElementList& InDrawElements, class UEdGraph* InGraphObj) const override;
 	virtual void BreakNodeLinks(UEdGraphNode& TargetNode) const override;
 	virtual void BreakPinLinks(UEdGraphPin& TargetPin, bool bSendsNodeNotification) const override;
-	virtual void BreakSinglePinLink(UEdGraphPin* SourcePin, UEdGraphPin* TargetPin) override;
+	virtual void BreakSinglePinLink(UEdGraphPin* SourcePin, UEdGraphPin* TargetPin) const override;
 	//~ End EdGraphSchema Interface
 
 	virtual void GetGraphNodeContextActions(FGraphContextMenuBuilder& ContextMenuBuilder, int32 SubNodeFlags) const;
diff --git a/Engine/Source/Editor/AIGraph/Private/AIGraphNode.cpp b/Engine/Source/Editor/AIGraph/Private/AIGraphNode.cpp
--- a/Engine/Source/Editor/AIGraph/Private/AIGraphNode.cpp
+++ b/Engine/Source/Editor/AIGraph/Private/AIGraphNode.cpp
@@ -104,11 +104,7 @@ void UAIGraphNode::ResetNodeOwner()
 
 		for (auto& SubNode : SubNodes)
 		{
-			if (SubNode->NodeInstance != nullptr)
-			{
-				SubNode->NodeInstance->Rename(NULL, GraphOwner, REN_DontCreateRedirectors | REN_DoNotDirty);
-				SubNode->NodeInstance->ClearFlags(RF_Transient);
-			}
+			SubNode->ResetNodeOwner();
 		}
 	}
 }
diff --git a/Engine/Source/Editor/AIGraph/Private/AIGraphSchema.cpp b/Engine/Source/Editor/AIGraph/Private/AIGraphSchema.cpp
--- a/Engine/Source/Editor/AIGraph/Private/AIGraphSchema.cpp
+++ b/Engine/Source/Editor/AIGraph/Private/AIGraphSchema.cpp
@@ -227,7 +227,7 @@ void UAIGraphSchema::GetBreakLinkToSubMenuActions(class FMenuBuilder& MenuBuilde
 		FText Title = FText::FromString(TitleString);
 		if (Pin->PinName != TEXT(""))
 		{
-			TitleString = FString::Printf(TEXT("%s (%s)"), *TitleString, *Pin->PinName);
+			TitleString = FString::Printf(TEXT("%s (%s)"), *TitleString, *Pin->PinName.ToString());
 
 			// Add name of connection if possible
 			FFormatNamedArguments Args;
@@ -272,7 +272,7 @@ void UAIGraphSchema::BreakPinLinks(UEdGraphPin& TargetPin, bool bSendsNodeNotifi
 	Super::BreakPinLinks(TargetPin, bSendsNodeNotification);
 }
 
-void UAIGraphSchema::BreakSinglePinLink(UEdGraphPin* SourcePin, UEdGraphPin* TargetPin)
+void UAIGraphSchema::BreakSinglePinLink(UEdGraphPin* SourcePin, UEdGraphPin* TargetPin) const
 {
 	const FScopedTransaction Transaction(NSLOCTEXT("UnrealEd", "GraphEd_BreakSinglePinLink", "Break Pin Link"));
 
diff --git a/Engine/Source/Editor/AdvancedPreviewScene/Private/AdvancedPreviewScene.cpp b/Engine/Source/Editor/AdvancedPreviewScene/Private/AdvancedPreviewScene.cpp
--- a/Engine/Source/Editor/AdvancedPreviewScene/Private/AdvancedPreviewScene.cpp
+++ b/Engine/Source/Editor/AdvancedPreviewScene/Private/AdvancedPreviewScene.cpp
@@ -292,11 +292,6 @@ void FAdvancedPreviewScene::Tick(float DeltaTime)
 	}
 }
 
-bool FAdvancedPreviewScene::IsTickable() const
-{
-	return true;
-}
-
 TStatId FAdvancedPreviewScene::GetStatId() const
 {
 	return TStatId();
diff --git a/Engine/Source/Editor/AdvancedPreviewScene/Public/AdvancedPreviewScene.h b/Engine/Source/Editor/AdvancedPreviewScene/Public/AdvancedPreviewScene.h
--- a/Engine/Source/Editor/AdvancedPreviewScene/Public/AdvancedPreviewScene.h
+++ b/Engine/Source/Editor/AdvancedPreviewScene/Public/AdvancedPreviewScene.h
@@ -36,7 +36,7 @@ public:
 
 	/* Begin FTickableEditorObject */
 	virtual void Tick(float DeltaTime) override;
-	virtual bool IsTickable() const override;
+	virtual ETickableTickType GetTickableTickType() const override { return ETickableTickType::Always; }
 	virtual TStatId GetStatId() const override;
 	/* End FTickableEditorObject */
 
diff --git a/Engine/Source/Editor/AnimGraph/Classes/AnimGraphNode_Base.h b/Engine/Source/Editor/AnimGraph/Classes/AnimGraphNode_Base.h
--- a/Engine/Source/Editor/AnimGraph/Classes/AnimGraphNode_Base.h
+++ b/Engine/Source/Editor/AnimGraph/Classes/AnimGraphNode_Base.h
@@ -12,7 +12,7 @@
 #include "K2Node.h"
 #include "AnimGraphNode_Base.generated.h"
 
-class FAnimBlueprintCompiler;
+class FAnimBlueprintCompilerContext;
 class FAnimGraphNodeDetails;
 class FBlueprintActionDatabaseRegistrar;
 class FCanvas;
@@ -284,7 +284,7 @@ class ANIMGRAPH_API UAnimGraphNode_Base : public UK2Node
 	bool IsPinExposedAndLinked(const FString& InPinName, const EEdGraphPinDirection Direction = EGPD_MAX) const;
 
 protected:
-	friend FAnimBlueprintCompiler;
+	friend FAnimBlueprintCompilerContext;
 	friend FAnimGraphNodeDetails;
 
 	// Gets the animation FNode type represented by this ed graph node
diff --git a/Engine/Source/Editor/AnimGraph/Classes/AnimGraphNode_SkeletalControlBase.h b/Engine/Source/Editor/AnimGraph/Classes/AnimGraphNode_SkeletalControlBase.h
--- a/Engine/Source/Editor/AnimGraph/Classes/AnimGraphNode_SkeletalControlBase.h
+++ b/Engine/Source/Editor/AnimGraph/Classes/AnimGraphNode_SkeletalControlBase.h
@@ -52,7 +52,7 @@ public:
 	virtual void CopyNodeDataFrom(const FAnimNode_Base* NewAnimNode){}
 
 	/** Are we currently showing this pin */
-	bool IsPinShown(const FString& PinName) const;
+	bool IsPinShown(const FName PinName) const;
 
 	// Returns the coordinate system that should be used for this bone
 	DEPRECATED(4.14, "This function is deprecated. Use a IAnimNodeEditMode-derived class to implement a new editor mode and override GetEditorMode to return its ID")
@@ -97,11 +97,11 @@ public:
 
 public:
 	// set literal value for FVector
-	void SetDefaultValue(const FString& InDefaultValueName, const FVector& InValue);
+	void SetDefaultValue(const FName InDefaultValueName, const FVector& InValue);
 	// get literal value for vector
-	void GetDefaultValue(const FString& UpdateDefaultValueName, FVector& OutVec);
+	void GetDefaultValue(const FName UpdateDefaultValueName, FVector& OutVec);
 
-	void GetDefaultValue(const FString& PropName, FRotator& OutValue)
+	void GetDefaultValue(const FName PropName, FRotator& OutValue)
 	{
 		FVector Value;
 		GetDefaultValue(PropName, Value);
@@ -111,7 +111,7 @@ public:
 	}
 
 	template<class ValueType>
-	ValueType GetNodeValue(const FString& PropName, const ValueType& CompileNodeValue)
+	ValueType GetNodeValue(const FName PropName, const ValueType& CompileNodeValue)
 	{
 		if (IsPinShown(PropName))
 		{
@@ -122,14 +122,14 @@ public:
 		return CompileNodeValue;
 	}
 
-	void SetDefaultValue(const FString& PropName, const FRotator& InValue)
+	void SetDefaultValue(const FName PropName, const FRotator& InValue)
 	{
 		FVector VecValue(InValue.Pitch, InValue.Yaw, InValue.Roll);
 		SetDefaultValue(PropName, VecValue);
 	}
 
 	template<class ValueType>
-	void SetNodeValue(const FString& PropName, ValueType& CompileNodeValue, const ValueType& InValue)
+	void SetNodeValue(const FName PropName, ValueType& CompileNodeValue, const ValueType& InValue)
 	{
 		if (IsPinShown(PropName))
 		{
diff --git a/Engine/Source/Editor/AnimGraph/Classes/AnimationStateMachineSchema.h b/Engine/Source/Editor/AnimGraph/Classes/AnimationStateMachineSchema.h
--- a/Engine/Source/Editor/AnimGraph/Classes/AnimationStateMachineSchema.h
+++ b/Engine/Source/Editor/AnimGraph/Classes/AnimationStateMachineSchema.h
@@ -68,8 +68,7 @@ class UAnimationStateMachineSchema : public UEdGraphSchema
 {
 	GENERATED_UCLASS_BODY()
 
-	UPROPERTY()
-	FString PC_Exec;
+	ANIMGRAPH_API static const FName PC_Exec;
 
 	//~ Begin UEdGraphSchema Interface
 	virtual void CreateDefaultNodesForGraph(UEdGraph& Graph) const override;
@@ -89,6 +88,6 @@ class UAnimationStateMachineSchema : public UEdGraphSchema
 
 	virtual void BreakNodeLinks(UEdGraphNode& TargetNode) const override;
 	virtual void BreakPinLinks(UEdGraphPin& TargetPin, bool bSendsNodeNotification) const override;
-	virtual void BreakSinglePinLink(UEdGraphPin* SourcePin, UEdGraphPin* TargetPin) override;
+	virtual void BreakSinglePinLink(UEdGraphPin* SourcePin, UEdGraphPin* TargetPin) const override;
 	//~ End UEdGraphSchema Interface
 };
diff --git a/Engine/Source/Editor/AnimGraph/Private/AnimBlueprintNodeOptionalPinManager.cpp b/Engine/Source/Editor/AnimGraph/Private/AnimBlueprintNodeOptionalPinManager.cpp
--- a/Engine/Source/Editor/AnimGraph/Private/AnimBlueprintNodeOptionalPinManager.cpp
+++ b/Engine/Source/Editor/AnimGraph/Private/AnimBlueprintNodeOptionalPinManager.cpp
@@ -12,11 +12,10 @@ FAnimBlueprintNodeOptionalPinManager::FAnimBlueprintNodeOptionalPinManager(class
 	: BaseNode(Node)
 	, OldPins(InOldPins)
 {
-	if (OldPins != NULL)
+	if (OldPins)
 	{
-		for (auto PinIt = OldPins->CreateIterator(); PinIt; ++PinIt)
+		for (UEdGraphPin* Pin : *OldPins)
 		{
-			UEdGraphPin* Pin = *PinIt;
 			OldPinMap.Add(Pin->PinName, Pin);
 		}
 	}
@@ -28,8 +27,8 @@ void FAnimBlueprintNodeOptionalPinManager::GetRecordDefaults(UProperty* TestProp
 
 	// Determine if this is a pose or array of poses
 	UArrayProperty* ArrayProp = Cast<UArrayProperty>(TestProperty);
-	UStructProperty* StructProp = Cast<UStructProperty>((ArrayProp != NULL) ? ArrayProp->Inner : TestProperty);
-	const bool bIsPoseInput = (StructProp != NULL) && (StructProp->Struct->IsChildOf(FPoseLinkBase::StaticStruct()));
+	UStructProperty* StructProp = Cast<UStructProperty>(ArrayProp ? ArrayProp->Inner : TestProperty);
+	const bool bIsPoseInput = (StructProp  && StructProp->Struct->IsChildOf(FPoseLinkBase::StaticStruct()));
 
 	//@TODO: Error if they specified two or more of these flags
 	const bool bAlwaysShow = TestProperty->HasMetaData(Schema->NAME_AlwaysAsPin) || bIsPoseInput;
@@ -111,7 +110,7 @@ void FAnimBlueprintNodeOptionalPinManager::PostRemovedOldPin(FOptionalPinFromPro
 
 	if (Record.bCanToggleVisibility && (OldPins != nullptr))
 	{
-		const FString OldPinName = (ArrayIndex != INDEX_NONE) ? FString::Printf(TEXT("%s_%d"), *(Record.PropertyName.ToString()), ArrayIndex) : Record.PropertyName.ToString();
+		const FName OldPinName = (ArrayIndex != INDEX_NONE) ? *FString::Printf(TEXT("%s_%d"), *(Record.PropertyName.ToString()), ArrayIndex) : Record.PropertyName;
 		// Pin was visible but it's now hidden
 		if (UEdGraphPin* OldPin = OldPinMap.FindRef(OldPinName))
 		{
diff --git a/Engine/Source/Editor/AnimGraph/Private/AnimBlueprintNodeOptionalPinManager.h b/Engine/Source/Editor/AnimGraph/Private/AnimBlueprintNodeOptionalPinManager.h
--- a/Engine/Source/Editor/AnimGraph/Private/AnimBlueprintNodeOptionalPinManager.h
+++ b/Engine/Source/Editor/AnimGraph/Private/AnimBlueprintNodeOptionalPinManager.h
@@ -13,7 +13,7 @@ protected:
 	class UAnimGraphNode_Base* BaseNode;
 	TArray<UEdGraphPin*>* OldPins;
 
-	TMap<FString, UEdGraphPin*> OldPinMap;
+	TMap<FName, UEdGraphPin*> OldPinMap;
 
 public:
 	FAnimBlueprintNodeOptionalPinManager(class UAnimGraphNode_Base* Node, TArray<UEdGraphPin*>* InOldPins);
diff --git a/Engine/Source/Editor/AnimGraph/Private/AnimGraphNode_Base.cpp b/Engine/Source/Editor/AnimGraph/Private/AnimGraphNode_Base.cpp
--- a/Engine/Source/Editor/AnimGraph/Private/AnimGraphNode_Base.cpp
+++ b/Engine/Source/Editor/AnimGraph/Private/AnimGraphNode_Base.cpp
@@ -51,8 +51,7 @@ void UAnimGraphNode_Base::CreateOutputPins()
 {
 	if (!IsSinkNode())
 	{
-		const UAnimationGraphSchema* Schema = GetDefault<UAnimationGraphSchema>();
-		CreatePin(EGPD_Output, Schema->PC_Struct, FString(), FPoseLink::StaticStruct(), TEXT("Pose"));
+		CreatePin(EGPD_Output, UAnimationGraphSchema::PC_Struct, FPoseLink::StaticStruct(), TEXT("Pose"));
 	}
 }
 
@@ -184,14 +183,16 @@ void UAnimGraphNode_Base::GetPinAssociatedProperty(const UScriptStruct* NodeType
 	OutIndex = INDEX_NONE;
 
 	//@TODO: Name-based hackery, avoid the roundtrip and better indicate when it's an array pose pin
-	int32 UnderscoreIndex = InputPin->PinName.Find(TEXT("_"), ESearchCase::CaseSensitive);
+	const FString PinNameStr = InputPin->PinName.ToString();
+	const int32 UnderscoreIndex = PinNameStr.Find(TEXT("_"), ESearchCase::CaseSensitive);
 	if (UnderscoreIndex != INDEX_NONE)
 	{
-		FString ArrayName = InputPin->PinName.Left(UnderscoreIndex);
-		int32 ArrayIndex = FCString::Atoi(*(InputPin->PinName.Mid(UnderscoreIndex + 1)));
+		const FString ArrayName = PinNameStr.Left(UnderscoreIndex);
 
 		if (UArrayProperty* ArrayProperty = FindField<UArrayProperty>(NodeType, *ArrayName))
 		{
+			const int32 ArrayIndex = FCString::Atoi(*(PinNameStr.Mid(UnderscoreIndex + 1)));
+
 			OutProperty = ArrayProperty;
 			OutIndex = ArrayIndex;
 		}
@@ -200,7 +201,7 @@ void UAnimGraphNode_Base::GetPinAssociatedProperty(const UScriptStruct* NodeType
 	// If the array check failed or we have no underscores
 	if(OutProperty == nullptr)
 	{
-		if (UProperty* Property = FindField<UProperty>(NodeType, *(InputPin->PinName)))
+		if (UProperty* Property = FindField<UProperty>(NodeType, InputPin->PinName))
 		{
 			OutProperty = Property;
 			OutIndex = INDEX_NONE;
@@ -215,11 +216,11 @@ FPoseLinkMappingRecord UAnimGraphNode_Base::GetLinkIDLocation(const UScriptStruc
 		if (UAnimGraphNode_Base* LinkedNode = Cast<UAnimGraphNode_Base>(FBlueprintEditorUtils::FindFirstCompilerRelevantNode(SourcePin->LinkedTo[0])))
 		{
 			//@TODO: Name-based hackery, avoid the roundtrip and better indicate when it's an array pose pin
-			int32 UnderscoreIndex = SourcePin->PinName.Find(TEXT("_"), ESearchCase::CaseSensitive);
+			const FString SourcePinName = SourcePin->PinName.ToString();
+			const int32 UnderscoreIndex = SourcePinName.Find(TEXT("_"), ESearchCase::CaseSensitive);
 			if (UnderscoreIndex != INDEX_NONE)
 			{
-				FString ArrayName = SourcePin->PinName.Left(UnderscoreIndex);
-				int32 ArrayIndex = FCString::Atoi(*(SourcePin->PinName.Mid(UnderscoreIndex + 1)));
+				const FString ArrayName = SourcePinName.Left(UnderscoreIndex);
 
 				if (UArrayProperty* ArrayProperty = FindField<UArrayProperty>(NodeType, *ArrayName))
 				{
@@ -227,6 +228,7 @@ FPoseLinkMappingRecord UAnimGraphNode_Base::GetLinkIDLocation(const UScriptStruc
 					{
 						if (Property->Struct->IsChildOf(FPoseLinkBase::StaticStruct()))
 						{
+							const int32 ArrayIndex = FCString::Atoi(*(SourcePinName.Mid(UnderscoreIndex + 1)));
 							return FPoseLinkMappingRecord::MakeFromArrayEntry(this, LinkedNode, ArrayProperty, ArrayIndex);
 						}
 					}
@@ -234,7 +236,7 @@ FPoseLinkMappingRecord UAnimGraphNode_Base::GetLinkIDLocation(const UScriptStruc
 			}
 			else
 			{
-				if (UStructProperty* Property = FindField<UStructProperty>(NodeType, *(SourcePin->PinName)))
+				if (UStructProperty* Property = FindField<UStructProperty>(NodeType, SourcePin->PinName))
 				{
 					if (Property->Struct->IsChildOf(FPoseLinkBase::StaticStruct()))
 					{
@@ -250,12 +252,11 @@ FPoseLinkMappingRecord UAnimGraphNode_Base::GetLinkIDLocation(const UScriptStruc
 
 void UAnimGraphNode_Base::CreatePinsForPoseLink(UProperty* PoseProperty, int32 ArrayIndex)
 {
-	const UAnimationGraphSchema* Schema = GetDefault<UAnimationGraphSchema>();
 	UScriptStruct* A2PoseStruct = FA2Pose::StaticStruct();
 
 	// pose input
-	const FString NewPinName = (ArrayIndex == INDEX_NONE) ? PoseProperty->GetName() : FString::Printf(TEXT("%s_%d"), *(PoseProperty->GetName()), ArrayIndex);
-	CreatePin(EGPD_Input, Schema->PC_Struct, FString(), A2PoseStruct, NewPinName);
+	const FName NewPinName = (ArrayIndex == INDEX_NONE) ? PoseProperty->GetFName() : *FString::Printf(TEXT("%s_%d"), *(PoseProperty->GetName()), ArrayIndex);
+	CreatePin(EGPD_Input, UAnimationGraphSchema::PC_Struct, A2PoseStruct, NewPinName);
 }
 
 void UAnimGraphNode_Base::PostProcessPinName(const UEdGraphPin* Pin, FString& DisplayName) const
