Copying Orion-Staging @ 3748653 (Orion/Dev-General @ 3722124) to //UE4/Main

diff --git a/Engine/Plugins/FX/Niagara/Source/Niagara/Classes/NiagaraEmitter.h b/Engine/Plugins/FX/Niagara/Source/Niagara/Classes/NiagaraEmitter.h
--- a/Engine/Plugins/FX/Niagara/Source/Niagara/Classes/NiagaraEmitter.h
+++ b/Engine/Plugins/FX/Niagara/Source/Niagara/Classes/NiagaraEmitter.h
@@ -13,6 +13,7 @@
 class UMaterial;
 class UNiagaraEmitter;
 class UNiagaraEventReceiverEmitterAction;
+class UNiagaraRendererProperties;
 
 //TODO: Event action that spawns other whole Systems?
 //One that calls a BP exposed delegate?
@@ -124,6 +125,8 @@ struct FNiagaraEmitterScriptProperties
 	TArray<FNiagaraEventGeneratorProperties> EventGenerators;
 
 	NIAGARA_API void InitDataSetAccess();
+
+	NIAGARA_API bool DataSetAccessSynchronized() const;
 };
 
 USTRUCT()
@@ -157,7 +160,6 @@ struct FNiagaraEventScriptProperties : public FNiagaraEmitterScriptProperties
 	/** The name of the event generated. This will be "Collision" for collision events and the Event Name field on the DataSetWrite node in the module graph for others.*/
 	UPROPERTY(EditAnywhere, Category="Event Handler Options")
 	FName SourceEventName;
-
 };
 
 /** Represents timed burst of particles in an emitter. */
@@ -194,10 +196,18 @@ UCLASS(MinimalAPI)
 class UNiagaraEmitter : public UObject
 {
 	GENERATED_UCLASS_BODY()
+
+public:
+#if WITH_EDITOR
+	DECLARE_MULTICAST_DELEGATE(FOnPropertiesChanged);
+#endif
+
 public:
 	//Begin UObject Interface
 #if WITH_EDITOR
 	virtual void PostEditChangeProperty(struct FPropertyChangedEvent& PropertyChangedEvent) override;
+
+	NIAGARA_API FOnPropertiesChanged& OnPropertiesChanged();
 #endif
 	void Serialize(FArchive& Ar)override;
 	virtual void PostInitProperties() override;
@@ -215,9 +225,6 @@ public:
 	//int32 NumLoops;
 	UPROPERTY(EditAnywhere, Category = "Emitter")
 	ENiagaraCollisionMode CollisionMode;
-	
-	UPROPERTY(EditAnywhere, Instanced, Category = "Render")
-	TArray<class UNiagaraRendererProperties*> RendererProperties;
 
 	UPROPERTY()
 	FNiagaraEmitterScriptProperties UpdateScriptProps;
@@ -225,48 +232,138 @@ public:
 	UPROPERTY()
 	FNiagaraEmitterScriptProperties SpawnScriptProps;
 
-	UPROPERTY(EditAnywhere, Category="Event Handler")
-	TArray<FNiagaraEventScriptProperties> EventHandlerScriptProps;
+	UPROPERTY()
+	FNiagaraEmitterScriptProperties EmitterSpawnScriptProps;
+
+	UPROPERTY()
+	FNiagaraEmitterScriptProperties EmitterUpdateScriptProps;
+
+
+	UPROPERTY(EditAnywhere, Category = "Emitter")
+	ENiagaraSimTarget SimTarget;
+	
+	/** The fixed bounding box value. bFixedBounds is the condition whether the fixed bounds can be edited. */
+	UPROPERTY(EditAnywhere, Category = "Emitter", meta = (EditCondition = "bFixedBounds"))
+	FBox FixedBounds;
+	
+	/** If the current engine detail level is below MinDetailLevel then this emitter is disabled. */
+	UPROPERTY(EditAnywhere, Category = "Scalability", meta=(EditCondition = "bUseMinDetailLevel"))
+	int32 MinDetailLevel;
+
+	/** If the current engine detail level is above MaxDetailLevel then this emitter is disabled. */
+	UPROPERTY(EditAnywhere, Category = "Scalability", meta=(EditCondition = "bUseMaxDetailLevel"))
+	int32 MaxDetailLevel;
 
 	/** When enabled, this will spawn using interpolated parameter values and perform a partial update at spawn time. This adds significant additional cost for spawning but will produce much smoother spawning for high spawn rates, erratic frame rates and fast moving emitters. */
 	UPROPERTY(EditAnywhere, Category = "Emitter")
 	uint32 bInterpolatedSpawning : 1;
 
-	UPROPERTY(EditAnywhere, Category = "Emitter")
-	ENiagaraSimTarget SimTarget;
+	/** Whether or not fixed bounds are enabled. */
+	UPROPERTY(EditAnywhere, Category = "Emitter", meta = (InlineEditConditionToggle))
+	uint32 bFixedBounds : 1;
 
-#if WITH_EDITORONLY_DATA
-	/** Adjusted every time that we compile this emitter. Lets us know that we might differ from any cached versions.*/
-	UPROPERTY()
-	FGuid ChangeId;
+	/** Whether to use the min detail or not. */
+	UPROPERTY(EditAnywhere, Category = "Scalability", meta = (InlineEditConditionToggle))
+	uint32 bUseMinDetailLevel : 1;
+	
+	/** Whether to use the min detail or not. */
+	UPROPERTY(EditAnywhere, Category = "Scalability", meta = (InlineEditConditionToggle))
+	uint32 bUseMaxDetailLevel : 1;
+
+	void NIAGARA_API GetScripts(TArray<UNiagaraScript*>& OutScripts, bool bCompilableOnly = true);
 
+	NIAGARA_API UNiagaraScript* GetScript(ENiagaraScriptUsage Usage, FGuid UsageId);
+
+#if WITH_EDITORONLY_DATA
 	/** 'Source' data/graphs for the scripts used by this emitter. */
 	UPROPERTY()
 	class UNiagaraScriptSourceBase*	GraphSource;
 
-	void NIAGARA_API CompileScripts(TArray<ENiagaraScriptCompileStatus>& OutScriptStatuses, TArray<FString>& OutGraphLevelErrorMessages, TArray<FString>& PathNames, TArray<UNiagaraScript*>& Scripts);
-	ENiagaraScriptCompileStatus NIAGARA_API CompileScript(EScriptCompileIndices InScriptToCompile, FString& OutGraphLevelErrorMessages, int32 InSubScriptIdx = 0);
+	bool NIAGARA_API AreAllScriptAndSourcesSynchronized() const;
+	void NIAGARA_API CompileScripts(TArray<ENiagaraScriptCompileStatus>& OutScriptStatuses, TArray<FString>& OutGraphLevelErrorMessages, TArray<FString>& PathNames, TArray<UNiagaraScript*>& Scripts, bool bForce);
+	ENiagaraScriptCompileStatus NIAGARA_API CompileScript(EScriptCompileIndices InScriptToCompile, FString& OutGraphLevelErrorMessages, bool bForce, int32 InSubScriptIdx = 0);
 
 	UNiagaraEmitter* MakeRecursiveDeepCopy(UObject* DestOuter) const;
 	UNiagaraEmitter* MakeRecursiveDeepCopy(UObject* DestOuter, TMap<const UObject*, UObject*>& ExistingConversions) const;
 
-	void NIAGARA_API GetScripts(TArray<UNiagaraScript*>& OutScripts);
-
+	/* Gets a Guid which is updated any time data in this emitter is changed. */
+	FGuid NIAGARA_API GetChangeId() const;
+	
 	/** Data used by the editor to maintain UI state etc.. */
 	UPROPERTY()
 	UObject* EditorData;
+
+	/** Internal: The thumbnail image.*/
+	UPROPERTY()
+	class UTexture2D* ThumbnailImage;
+
+	/** Internal: Indicates the thumbnail image is out of date.*/
+	UPROPERTY()
+	uint32 ThumbnailImageOutOfDate : 1;
 #endif
 
-	bool IsValid()const;
+	/** Is this emitter allowed to be enabled by the current system detail level. */
+	bool IsAllowedByDetailLevel()const;
 
+	NIAGARA_API bool IsValid()const;
+	NIAGARA_API bool IsReadyToRun() const;
 	bool UsesScript(const UNiagaraScript* Script)const;
 	//bool UsesDataInterface(UNiagaraDataInterface* Interface);
 	bool UsesCollection(const class UNiagaraParameterCollection* Collection)const;
 
 	FString NIAGARA_API GetUniqueEmitterName()const;
+	void NIAGARA_API SetUniqueEmitterName(const FString& InName);
 
 	/** Converts an emitter paramter "Emitter.XXXX" into it's real parameter name. */
-	FNiagaraVariable GetEmitterParameter(const FNiagaraVariable& EmitterVar)const;
+	FNiagaraVariable ToEmitterParameter(const FNiagaraVariable& EmitterVar)const;
+
+	const TArray<UNiagaraRendererProperties*>& GetRenderers() const { return RendererProperties; }
+
+	void NIAGARA_API AddRenderer(UNiagaraRendererProperties* Renderer);
+
+	void NIAGARA_API RemoveRenderer(UNiagaraRendererProperties* Renderer);
+
+	FORCEINLINE const TArray<FNiagaraEventScriptProperties>& GetEventHandlers() const { return EventHandlerScriptProps; }
+
+	/* Gets a pointer to an event handler by script usage id.  This method is potentially unsafe because modifications to
+	   the event handler array can make this pointer become invalid without warning. */
+	NIAGARA_API FNiagaraEventScriptProperties* GetEventHandlerByIdUnsafe(FGuid ScriptUsageId);
+
+	void NIAGARA_API AddEventHandler(FNiagaraEventScriptProperties EventHandler);
+
+	void NIAGARA_API RemoveEventHandlerByUsageId(FGuid EventHandlerUsageId);
+
+protected:
+	virtual void BeginDestroy() override;
+
+#if WITH_EDITORONLY_DATA
+private:
+	void UpdateChangeId();
+
+	void ScriptRapidIterationParameterChanged();
+
+	void RendererChanged();
+
+	void GraphSourceChanged();
+
+private:
+	/** Adjusted every time that we compile this emitter. Lets us know that we might differ from any cached versions.*/
+	UPROPERTY()
+	FGuid ChangeId;
+#endif
+
+	UPROPERTY()
+	FString UniqueEmitterName;
+
+	UPROPERTY()
+	TArray<UNiagaraRendererProperties*> RendererProperties;
+
+	UPROPERTY()
+	TArray<FNiagaraEventScriptProperties> EventHandlerScriptProps;
+
+#if WITH_EDITOR
+	FOnPropertiesChanged OnPropertiesChangedDelegate;
+#endif
 };
 
 
diff --git a/Engine/Plugins/FX/Niagara/Source/Niagara/Classes/NiagaraEmitterHandle.h b/Engine/Plugins/FX/Niagara/Source/Niagara/Classes/NiagaraEmitterHandle.h
--- a/Engine/Plugins/FX/Niagara/Source/Niagara/Classes/NiagaraEmitterHandle.h
+++ b/Engine/Plugins/FX/Niagara/Source/Niagara/Classes/NiagaraEmitterHandle.h
@@ -6,6 +6,7 @@
 #include "UObject/ObjectMacros.h"
 #include "UObject/Object.h"
 #include "Guid.h"
+#include "NiagaraModule.h"
 #include "NiagaraEmitterHandle.generated.h"
 
 class UNiagaraSystem;
@@ -29,7 +30,7 @@ public:
 
 #if WITH_EDITORONLY_DATA
 	/** Creates a new emitter handle from an emitter and an owning System. */
-	FNiagaraEmitterHandle(const UNiagaraEmitter& InSourceEmitter, FName InName, UNiagaraSystem& InOuterSystem);
+	FNiagaraEmitterHandle(UNiagaraEmitter& InSourceEmitter, FName InName, UNiagaraSystem& InOuterSystem);
 
 	/** Creates a new emitter handle by duplicating an existing handle.  The new emitter handle will reference the same source emitter
 		but will have it's own copy of the emitter made from the one in the supplied handle and will have it's own Id. */
@@ -48,8 +49,8 @@ public:
 	/** Gets the display name for this emitter in the System. */
 	FName GetName() const;
 
-	/** Sets the display name for this emitter in the System. */
-	void SetName(FName InName);
+	/** Sets the display name for this emitter in the System. The system is needed here in order to ensure uniqueness of the name. */
+	void SetName(FName InName, UNiagaraSystem& InOwnerSystem);
 
 	/** Gets whether or not this emitter is enabled within the System.  Disabled emitters aren't simulated. */
 	bool GetIsEnabled() const;
@@ -60,30 +61,36 @@ public:
 #if WITH_EDITORONLY_DATA
 	/** Gets the source emitter this emitter handle was built from. */
 	const UNiagaraEmitter* GetSource() const;
+
+	bool IsIsolated() const {	return bIsolated; }
+	void SetIsolated(bool bInIsolated) { bIsolated = bInIsolated; }
 #endif
 
 	/** Gets the copied instance of the emitter this handle references. */
 	UNiagaraEmitter* GetInstance() const;
 
-	/** Update the instance this handle references*/
-	void SetInstance(UNiagaraEmitter* InInstance);
-
 	/** Gets a unique name for this emitter instance for use in scripts and parameter stores etc.*/
 	FString GetUniqueInstanceName()const;
 
 #if WITH_EDITORONLY_DATA
-	/** Return this emitter handle instance to its initial state, exactly matching the Source.*/
-	void ResetToSource();
-
-	/** Keep any existing settings, but include any new changes from the source emitter.*/
-	bool RefreshFromSource();
-
 	/** Determine whether or not the Source and Instance refer to the same Emitter ChangeId.*/
 	bool IsSynchronizedWithSource() const;
 
 	/** Determine whether or not the Instance script is in synch with its graph.*/
 	bool NeedsRecompile() const;
 
+	/** Calls conditional post load on all sub-objects this handle references. */
+	void ConditionalPostLoad();
+
+	/** Merges in any changes from the source emitter into the instanced emitter. */
+	INiagaraModule::FMergeEmitterResults MergeSourceChanges();
+
+	/** Gets an object which contains the history of modifications to the emitter instance. */
+	UObject* GetEmitterModificationHistory();
+
+	/** Sets an object which will contain the history of modifications to the emitter instance. */
+	void SetEmitterModificationHistory();
+
 #endif
 public:
 	/** A static const invalid handle. */
@@ -110,7 +117,13 @@ private:
 #if WITH_EDITORONLY_DATA
 	/** The source emitter this emitter handle was built from. */
 	UPROPERTY()
-	const UNiagaraEmitter* Source;
+	UNiagaraEmitter* Source;
+
+	/** An unmodified copy of the emitter this handle references for use when merging change from the source emitter. */
+	UPROPERTY()
+	UNiagaraEmitter* LastMergedSource;
+
+	bool bIsolated;
 #endif
 
 	/** The copied instance of the emitter this handle references. */
diff --git a/Engine/Plugins/FX/Niagara/Source/Niagara/Classes/NiagaraEmitterInstance.h b/Engine/Plugins/FX/Niagara/Source/Niagara/Classes/NiagaraEmitterInstance.h
--- a/Engine/Plugins/FX/Niagara/Source/Niagara/Classes/NiagaraEmitterInstance.h
+++ b/Engine/Plugins/FX/Niagara/Source/Niagara/Classes/NiagaraEmitterInstance.h
@@ -40,50 +40,52 @@ public:
 	virtual ~FNiagaraEmitterInstance();
 
 	void Init(int32 InEmitterIdx, FName SystemInstanceName);
+
 	void ResetSimulation();
-	void ReInitSimulation();
-	void DirtyDataInterfaces();
 
+	/** Called after all emitters in an System have been initialized, allows emitters to access information from one another. */
+	void PostInitSimulation();
+
+	void DirtyDataInterfaces();
+	
 	/** Replaces the binding for a single parameter colleciton instance. If for example the component begins to override the global instance. */
 	//void RebindParameterCollection(UNiagaraParameterCollectionInstance* OldInstance, UNiagaraParameterCollectionInstance* NewInstance);
 	void BindParameters();
 	void UnbindParameters();
-	
-	/** Called after all emitters in an System have been initialized, allows emitters to access information from one another. */
-	void PostResetSimulation();
 
 	void PreTick();
 	void Tick(float DeltaSeconds);
+	void PostProcessParticles();
+	bool HandleCompletion(bool bForce=false);
+
+	FORCEINLINE bool ShouldTick()const { return ExecutionState == ENiagaraExecutionState::Active || GetNumParticles() > 0; }
 
 	uint32 CalculateEventSpawnCount(const FNiagaraEventScriptProperties &EventHandlerProps, TArray<int32> &EventSpawnCounts, FNiagaraDataSet *EventSet);
-	void PostProcessParticles();
 
-	FBox GetBounds() const { return CachedBounds; }
+	NIAGARA_API FBox CalculateDynamicBounds();
 
-	FNiagaraDataSet &GetData()	{ return Data; }
+	FNiagaraDataSet &GetData()	{ return *ParticleDataSet; }
 
 	int32 GetEmitterRendererNum() const {
 		return EmitterRenderer.Num();
 	}
 
 	NiagaraRenderer *GetEmitterRenderer(int32 Index)	{ return EmitterRenderer[Index]; }
-	
-	bool IsEnabled()const 	{ return bIsEnabled;  }
-	
-	/** Set whether or not this simulation is enabled*/
-	void SetEnabled(bool bEnabled) { bIsEnabled = bEnabled; }
 
+	FORCEINLINE bool IsDisabled()const { return ExecutionState == ENiagaraExecutionState::Disabled; }
+	FORCEINLINE bool IsComplete()const { return ExecutionState == ENiagaraExecutionState::Complete || ExecutionState == ENiagaraExecutionState::Disabled; }
+	
 	/** Sets the error state. */
 	void SetError(bool bInError) { bError = bInError; }
 		
 	/** Create a new NiagaraRenderer. The old renderer is not immediately deleted, but instead put in the ToBeRemoved list.*/
 	void NIAGARA_API UpdateEmitterRenderer(ERHIFeatureLevel::Type FeatureLevel, TArray<NiagaraRenderer*>& ToBeAddedList, TArray<NiagaraRenderer*>& ToBeRemovedList);
 
-	int32 GetNumParticles()	{ return Data.GetNumInstances(); }
+	FORCEINLINE int32 GetNumParticles()const	{ return ParticleDataSet->GetNumInstances(); }
 
 	NIAGARA_API const FNiagaraEmitterHandle& GetEmitterHandle() const;
 
-	FNiagaraSystemInstance* GetParentSystemInstance()	{ return ParentSystemInstance; }
+	FNiagaraSystemInstance* GetParentSystemInstance()const	{ return ParentSystemInstance; }
 
 	float NIAGARA_API GetTotalCPUTime();
 	int	NIAGARA_API GetTotalBytesUsed();
@@ -96,15 +98,20 @@ public:
 	/** Tell the renderer thread that we're done with the Niagara renderer on this simulation.*/
 	void ClearRenderer();
 
-	FBox GetCachedBounds() { return CachedBounds; }
+	FBox GetBounds();
 
 	FNiagaraScriptExecutionContext& GetSpawnExecutionContext() { return SpawnExecContext; }
 	FNiagaraScriptExecutionContext& GetUpdateExecutionContext() { return UpdateExecContext; }
 	TArray<FNiagaraScriptExecutionContext>& GetEventExecutionContexts() { return EventExecContexts; }
 
 	TArray<FNiagaraSpawnInfo>& GetSpawnInfo() { return SpawnInfos; }
+
+	NIAGARA_API bool IsReadyToRun() const;
+
 private:
 
+	void CheckForErrors();
+
 	/** The index of our emitter in our parent system instance. */
 	int32 EmitterIdx;
 
@@ -112,11 +119,12 @@ private:
 	float Age;
 	/* how many loops this emitter has gone through */
 	uint32 Loops;
-	/* If false, don't tick or render*/
-	bool bIsEnabled;
 	
 	/* Some error has occurred so stop ticking or rendering. Better to just kill the emitter? */
-	bool bError;
+	uint32 bError : 1;
+	
+	/** Typical resets must be deferred until the tick as the RT could still be using the current buffer. */
+	uint32 bResetPending : 1;
 
 	/* Seconds taken to process everything (including rendering) */
 	float CPUTimeMS;
@@ -130,6 +138,7 @@ private:
 
 	FNiagaraScriptExecutionContext SpawnExecContext;
 	FNiagaraScriptExecutionContext UpdateExecContext;
+	FNiagaraComputeExecutionContext GPUExecContext;
 	TArray<FNiagaraScriptExecutionContext> EventExecContexts;
 
 	FNiagaraParameterDirectBinding<float> SpawnIntervalBinding;
@@ -139,14 +148,16 @@ private:
 	FNiagaraParameterDirectBinding<float> UpdateEmitterAgeBinding;
 	TArray<FNiagaraParameterDirectBinding<float>> EventEmitterAgeBindings;
 
+	FNiagaraParameterDirectBinding<FNiagaraBool> SpawnEmitterLocalSpaceBinding;
+	FNiagaraParameterDirectBinding<FNiagaraBool> UpdateEmitterLocalSpaceBinding;
+	TArray<FNiagaraParameterDirectBinding<FNiagaraBool>> EventEmitterLocalSpaceBindings;
+
 	FNiagaraParameterDirectBinding<int32> SpawnExecCountBinding;
 	FNiagaraParameterDirectBinding<int32> UpdateExecCountBinding;
 	TArray<FNiagaraParameterDirectBinding<int32>> EventExecCountBindings;
 
-	/** particle simulation data */
-	FNiagaraDataSet Data;
-	/** The cached GetComponentTransform() transform. */
-	FTransform CachedComponentToWorld;
+	/** particle simulation data. Must be a shared ref as various things on the RT can have direct ref to it. */
+	FNiagaraDataSet* ParticleDataSet;
 
 	TArray<NiagaraRenderer*> EmitterRenderer;
 	FNiagaraSystemInstance *ParentSystemInstance;
@@ -155,11 +166,14 @@ private:
 	TArray<FNiagaraDataSet*> SpawnScriptEventDataSets;
 	TMap<FNiagaraDataSetID, FNiagaraDataSet*> DataSetMap;
 	
-	FNiagaraCollisionBatch CollisionBatch;
 
 	FName OwnerSystemInstanceName;
 
 #if WITH_EDITORONLY_DATA
 	bool CheckAttributesForRenderer(int32 Index);
 #endif
+
+	FNiagaraDataSetAccessor<FVector> PositionAccessor;
+	FNiagaraDataSetAccessor<FVector2D> SizeAccessor;
+	FNiagaraDataSetAccessor<FVector> MeshScaleAccessor;
 };
diff --git a/Engine/Plugins/FX/Niagara/Source/Niagara/Classes/NiagaraEmitterInstanceBatcher.h b/Engine/Plugins/FX/Niagara/Source/Niagara/Classes/NiagaraEmitterInstanceBatcher.h
--- a/Engine/Plugins/FX/Niagara/Source/Niagara/Classes/NiagaraEmitterInstanceBatcher.h
+++ b/Engine/Plugins/FX/Niagara/Source/Niagara/Classes/NiagaraEmitterInstanceBatcher.h
@@ -17,37 +17,7 @@ the same VectorVM byte code / compute shader code
 #include "ModuleManager.h"
 
 struct FNiagaraScriptExecutionContext;
-
-struct FNiagaraComputeExecutionContext
-{
-	FNiagaraComputeExecutionContext()
-		: MainDataSet(nullptr)
-		, SpawnRateInstances(0)
-		, BurstInstances(0)
-		, EventSpawnTotal(0)
-		, RTUpdateScript(0)
-		, RTSpawnScript(0)
-	{
-		TickCounter++;
-	}
-	class FNiagaraDataSet *MainDataSet;
-	TArray<FNiagaraDataSet*>UpdateEventWriteDataSets;
-	TArray<FNiagaraEventScriptProperties> EventHandlerScriptProps;
-	TArray<FNiagaraDataSet*> EventSets;
-	uint32 SpawnRateInstances;
-	uint32 BurstInstances;
-
-	TArray<int32> EventSpawnCounts;
-	uint32 EventSpawnTotal;
-
-	class FNiagaraScript *RTUpdateScript;
-	class FNiagaraScript *RTSpawnScript;
-	TArray<uint8, TAlignedHeapAllocator<16>> UpdateParams;		// RT side copy of the parameter data
-	TArray<FNiagaraScriptDataInterfaceInfo> UpdateInterfaces;
-	TArray<uint8, TAlignedHeapAllocator<16>> SpawnParams;		// RT side copy of the parameter data
-	static uint32 TickCounter;
-};
-
+struct FNiagaraComputeExecutionContext;
 
 class NiagaraEmitterInstanceBatcher : public FTickableGameObject, public FComputeDispatcher
 {
@@ -60,6 +30,11 @@ public:
 		{
 			RendererModule->RegisterPostOpaqueComputeDispatcher(this);
 		}
+
+		if (DummyWriteIndexBuffer.Buffer == nullptr)
+		{
+			DummyWriteIndexBuffer.Initialize(sizeof(int32), 64, EPixelFormat::PF_R32_SINT);	// always allocate for up to 64 data sets
+		}
 	}
 
 	~NiagaraEmitterInstanceBatcher()
@@ -123,12 +98,22 @@ public:
 	void UnsetEventUAVs(const FNiagaraComputeExecutionContext *Context, FRHICommandList &RHICmdList) const;
 	void SetupDataInterfaceBuffers(const TArray<FNiagaraScriptDataInterfaceInfo> &DIInfos, FNiagaraShader *Shader, FRHICommandList &RHICmdList) const;
 
-	void Run(FNiagaraDataSet *DataSet, const uint32 StartInstance, const uint32 NumInstances, class FNiagaraShader *Shader, const TArray<uint8, TAlignedHeapAllocator<16>>  &Params, FRHICommandList &RhiCmdList, bool bCopyBeforeStart = false) const;
+	void Run(	FNiagaraDataSet *DataSet,
+				uint32 Phase0StartInstance,
+				uint32 Phase1StartInstance,
+				const uint32 TotalNumInstances,
+				FNiagaraShader *Shader,
+				const TArray<uint8, TAlignedHeapAllocator<16>>  &Params,
+				FRHICommandList &RHICmdList,
+				const FRWBuffer &WriteIndexBuffer,
+				bool bCopyBeforeStart = false) const;
+
 	void RunEventHandlers(const FNiagaraComputeExecutionContext *Context, uint32 NumInstancesAfterSim, uint32 NumInstancesAfterSpawn, uint32 NumInstancesAfterNonEventSpawn, FRHICommandList &RhiCmdList) const;
-	void ResolveDatasetWrites(uint32 *OutArray, const FNiagaraComputeExecutionContext *Context) const;
+	void ResolveDatasetWrites(uint32 *OutArray, const FNiagaraComputeExecutionContext *Context, FRWBuffer &WriteIndexBuffer) const;
 private:
 	static NiagaraEmitterInstanceBatcher* BatcherSingleton;
 
 	uint32 CurQueueIndex;
 	TArray<FNiagaraComputeExecutionContext*> SimulationQueue[2];
+	static FRWBuffer DummyWriteIndexBuffer;
 };
diff --git a/Engine/Plugins/FX/Niagara/Source/Niagara/Classes/NiagaraMergeable.h b/Engine/Plugins/FX/Niagara/Source/Niagara/Classes/NiagaraMergeable.h
new file mode 100644
--- /dev/null
+++ b/Engine/Plugins/FX/Niagara/Source/Niagara/Classes/NiagaraMergeable.h
@@ -0,0 +1,39 @@
+// Copyright 1998-2017 Epic Games, Inc. All Rights Reserved.
+
+#pragma once
+
+#include "NiagaraMergeable.generated.h"
+
+UCLASS()
+class NIAGARA_API UNiagaraMergeable : public UObject
+{
+	GENERATED_BODY()
+
+public:
+#if WITH_EDITOR
+	DECLARE_MULTICAST_DELEGATE(FOnChanged);
+#endif
+
+public:
+	UNiagaraMergeable();
+
+#if WITH_EDITOR
+	virtual void PostEditChangeProperty(struct FPropertyChangedEvent& PropertyChangedEvent);
+
+	bool Equals(const UNiagaraMergeable* Other);
+
+	FOnChanged& OnChanged();
+
+	FGuid GetMergeId();
+#endif
+
+private:
+#if WITH_EDITOR
+	FOnChanged OnChangedDelegate;
+#endif
+
+#if WITH_EDITORONLY_DATA
+	UPROPERTY()
+	FGuid MergeId;
+#endif
+};
\ No newline at end of file
diff --git a/Engine/Plugins/FX/Niagara/Source/Niagara/Classes/NiagaraParameterCollection.h b/Engine/Plugins/FX/Niagara/Source/Niagara/Classes/NiagaraParameterCollection.h
--- a/Engine/Plugins/FX/Niagara/Source/Niagara/Classes/NiagaraParameterCollection.h
+++ b/Engine/Plugins/FX/Niagara/Source/Niagara/Classes/NiagaraParameterCollection.h
@@ -113,8 +113,9 @@ class NIAGARA_API UNiagaraParameterCollection : public UObject
 public:
 
 	//~UObject interface
-	virtual void PostInitProperties()override;
-	virtual void PostLoad()override;
+#if WITH_EDITORONLY_DATA
+	virtual void PostEditChangeProperty(struct FPropertyChangedEvent& PropertyChangedEvent)override;
+#endif
 	//~UObject interface
 
 	int32 IndexOfParameter(const FNiagaraVariable& Var);
@@ -144,17 +145,18 @@ public:
 	FNiagaraVariable FriendlyParameterFromCollectionParameter(const FNiagaraVariable& CollectionParameter)const;
 
 	FString FriendlyNameFromParameterName(FString ParameterName)const;
+	FString GetFullNamespace()const;
 protected:
+	
+	void MakeNamespaceNameUnique();
 
-	const FString& GetUniqueName() const { return UniqueName; }
+	/** Namespace for this parameter collection. Is enforced to be unique across all parameter collections. */
+	UPROPERTY(EditAnywhere, Category = "Parameter Collection", AssetRegistrySearchable)
+	FName Namespace;
 	
 	UPROPERTY()
 	TArray<FNiagaraVariable> Parameters;
 	
 	UPROPERTY()
 	UNiagaraParameterCollectionInstance* DefaultInstance;
-
-	/** Unique name used by parameters in this collection and the scripts referencing them to link them with this collection. */
-	UPROPERTY()
-	FString UniqueName;
 };
\ No newline at end of file
diff --git a/Engine/Plugins/FX/Niagara/Source/Niagara/Classes/NiagaraScript.h b/Engine/Plugins/FX/Niagara/Source/Niagara/Classes/NiagaraScript.h
--- a/Engine/Plugins/FX/Niagara/Source/Niagara/Classes/NiagaraScript.h
+++ b/Engine/Plugins/FX/Niagara/Source/Niagara/Classes/NiagaraScript.h
@@ -10,6 +10,8 @@
 #include "NiagaraShader.h"
 #include "NiagaraParameters.h"
 #include "NiagaraDataSet.h"
+#include "NiagaraShared.h"
+#include "NiagaraParameterStore.h"
 
 #include "NiagaraScript.generated.h"
 
@@ -47,6 +49,8 @@ struct FNiagaraScriptDebuggerInfo
 	int32 DebugFrameLastWriteId;
 
 	FNiagaraDataSet DebugFrame;
+
+	FNiagaraParameterStore DebugParameters;
 };
 
 
@@ -57,13 +61,19 @@ class UNiagaraScript : public UObject
 	GENERATED_UCLASS_BODY()
 
 	// how this script is to be used. cannot be private due to use of GET_MEMBER_NAME_CHECKED
-	UPROPERTY(AssetRegistrySearchable, EditAnywhere, Category = Script)
+	UPROPERTY(AssetRegistrySearchable)
 	ENiagaraScriptUsage Usage;
 
-	/** Which instance of the usage in the graph to use.*/
+	/** Which instance of the usage in the graph to use.  This is now deprecated and is handled by UsageId. */
+	UPROPERTY()
+	int32 UsageIndex_DEPRECATED;
+
+private:
+	/** Specifies a unique id for use when there are multiple scripts with the same usage, e.g. events. */
 	UPROPERTY()
-	int32 UsageIndex;
+	FGuid UsageId;
 
+public:
 	/** When used as a module, what are the appropriate script types for referencing this module?*/
 	UPROPERTY(AssetRegistrySearchable, EditAnywhere, Category = Script, meta = (Bitmask, BitmaskEnum = ENiagaraScriptUsage))
 	int32 ModuleUsageBitmask;
@@ -75,14 +85,14 @@ class UNiagaraScript : public UObject
 	/** Number of user pointers we must pass to the VM. */
 	UPROPERTY()
 	int32 NumUserPtrs;
-
-	/** Byte code to execute for this system */
+	
+	/** All the data for using external constants in the script, laid out in the order they are expected in the uniform table.*/
 	UPROPERTY()
-	TArray<uint8> ByteCode;
+	FNiagaraParameters Parameters;
 
-	/** All the data for using constants in the script. */
+	/** Contains all of the top-level values that are iterated on in the UI. These are usually "Module" variables in the graph. They don't necessarily have to be in the order that they are expected in the uniform table.*/
 	UPROPERTY()
-	FNiagaraParameters Parameters;
+	FNiagaraParameterStore RapidIterationParameters;
 
 	UPROPERTY()
 	FNiagaraParameters InternalParameters;
@@ -110,7 +120,13 @@ class UNiagaraScript : public UObject
 	UPROPERTY(EditAnywhere, Category=Script)
 	ENiagaraNumericOutputTypeSelectionMode NumericOutputTypeSelectionMode;
 
+	UPROPERTY()
+	TArray< FDIBufferDescriptorStore > DIBufferDescriptors;
+
+	UPROPERTY()
 	TArray<FNiagaraDataSetID> ReadDataSets;
+	
+	UPROPERTY()
 	TArray<FNiagaraDataSetProperties> WriteDataSets;
 
 	/** Scopes we'll track with stats.*/
@@ -126,16 +142,25 @@ class UNiagaraScript : public UObject
 
 	UPROPERTY()
 	FString LastHlslTranslation;
+	UPROPERTY()
+	FString LastHlslTranslationGPU;
+	UPROPERTY()
+	FString LastAssemblyTranslation;
+	UPROPERTY()
+	uint32 LastOpCount;
+
 	
 	/** Returns true if this script is valid and can be executed. */
 	bool IsValid()const { return ByteCode.Num() > 0; }//More? Differentiate by CPU/GPU?
 
 	void SetUsage(ENiagaraScriptUsage InUsage) { Usage = InUsage; }
 	ENiagaraScriptUsage GetUsage() const { return Usage; }
-	void SetUsageIndex(int32 InUsageIndex) { UsageIndex = InUsageIndex; }
-	int32 GetUsageIndex() const { return UsageIndex; }
 
-	bool IsEquivalentUsage(ENiagaraScriptUsage InUsage) {return (InUsage == Usage) || (Usage == ENiagaraScriptUsage::ParticleSpawnScript && InUsage == ENiagaraScriptUsage::ParticleSpawnScriptInterpolated) || (Usage == ENiagaraScriptUsage::ParticleSpawnScriptInterpolated && InUsage == ENiagaraScriptUsage::ParticleSpawnScript);}
+	void SetUsageId(FGuid InUsageId) { UsageId = InUsageId; }
+	FGuid GetUsageId() const { return UsageId; }
+
+	NIAGARA_API bool ContainsUsage(ENiagaraScriptUsage InUsage) const;
+	bool IsEquivalentUsage(ENiagaraScriptUsage InUsage) const {return (InUsage == Usage) || (Usage == ENiagaraScriptUsage::ParticleSpawnScript && InUsage == ENiagaraScriptUsage::ParticleSpawnScriptInterpolated) || (Usage == ENiagaraScriptUsage::ParticleSpawnScriptInterpolated && InUsage == ENiagaraScriptUsage::ParticleSpawnScript);}
 	static bool IsEquivalentUsage(ENiagaraScriptUsage InUsageA, ENiagaraScriptUsage InUsageB) { return (InUsageA == InUsageB) || (InUsageB == ENiagaraScriptUsage::ParticleSpawnScript && InUsageA == ENiagaraScriptUsage::ParticleSpawnScriptInterpolated) || (InUsageB == ENiagaraScriptUsage::ParticleSpawnScriptInterpolated && InUsageA == ENiagaraScriptUsage::ParticleSpawnScript); }
 
 	bool IsParticleSpawnScript() const { return Usage == ENiagaraScriptUsage::ParticleSpawnScript || Usage == ENiagaraScriptUsage::ParticleSpawnScriptInterpolated; }
@@ -156,9 +181,11 @@ class UNiagaraScript : public UObject
 
 	bool IsSpawnScript()const { return IsParticleSpawnScript() || IsEmitterSpawnScript() || IsSystemSpawnScript(); }
 
+	bool IsCompilable() const { return !IsEmitterSpawnScript() && !IsEmitterUpdateScript(); }
+
 	NIAGARA_API TArray<ENiagaraScriptUsage> GetSupportedUsageContexts() const;
 
-	bool CanBeRunOnGpu() const { return IsParticleSpawnScript() || IsParticleUpdateScript() || IsParticleEventScript(); }
+	NIAGARA_API bool CanBeRunOnGpu() const;
 
 #if WITH_EDITORONLY_DATA
 	class UNiagaraScriptSourceBase *GetSource() { return Source; }
@@ -173,6 +200,12 @@ class UNiagaraScript : public UObject
 #endif
 	//~ End UObject interface
 
+	void SetDatainterfaceBufferDescriptors(TArray< FDIBufferDescriptorStore >&InBufferDescriptors)
+	{
+		DIBufferDescriptors = InBufferDescriptors;
+	}
+
+
 	// Infrastructure for GPU compute Shaders
 	NIAGARA_API void CacheResourceShadersForCooking(EShaderPlatform ShaderPlatform, FNiagaraScript* &OutCachedResource);
 	NIAGARA_API void CacheResourceShadersForRendering(bool bRegenerateId, bool bForceRecompile=false);
@@ -222,10 +255,14 @@ class UNiagaraScript : public UObject
 	}
 
 	NIAGARA_API void GenerateStatScopeIDs();
+	ENiagaraScriptCompileStatus GetLastCompileStatus() { return LastCompileStatus; }
+
+	NIAGARA_API bool IsScriptCompilationPending(bool bGPUScript) const;
+	NIAGARA_API bool DidScriptCompilationSucceed(bool bGPUScript) const;
+	NIAGARA_API void InvalidateScript();
 
 #if WITH_EDITORONLY_DATA
 	FText GetDescription() { return Description.IsEmpty() ? FText::FromString(GetName()) : Description; }
-	ENiagaraScriptCompileStatus GetLastCompileStatus() { return LastCompileStatus; }
 	void SetLastCompileStatus(ENiagaraScriptCompileStatus InStatus) { LastCompileStatus = InStatus; }
 
 	/** Makes a deep copy of any script dependencies, including itself.*/
@@ -240,10 +277,16 @@ class UNiagaraScript : public UObject
 	/** Ensure that the Script and its source graph are marked out of sync.*/
 	NIAGARA_API void MarkScriptAndSourceDesynchronized();
 
-	NIAGARA_API ENiagaraScriptCompileStatus Compile(FString& OutGraphLevelErrorMessages);
+	NIAGARA_API ENiagaraScriptCompileStatus Compile(FString& OutGraphLevelErrorMessages, bool bForce);
 
 	NIAGARA_API FNiagaraScriptDebuggerInfo& GetDebuggerInfo() { return DebuggerInfo; }
 
+#endif
+
+	UFUNCTION()
+	void OnCompilationComplete();
+
+#if STATS
 	const TArray<TStatId>& GetStatScopeIDs()const { return StatScopesIDs; }
 #endif
 
@@ -251,20 +294,27 @@ class UNiagaraScript : public UObject
 	
 	virtual ~UNiagaraScript();
 
-	virtual void PostInitProperties() override;
+	NIAGARA_API void SetByteCode(const TArray<uint8>& InByteCode);
+	const TArray<uint8>& GetByteCode() const { return ByteCode; }
+
 private:
+
+	/** Byte code to execute for this system */
+	UPROPERTY()
+	TArray<uint8> ByteCode;
+
 #if WITH_EDITORONLY_DATA
 	/** 'Source' data/graphs for this script */
 	UPROPERTY()
 	class UNiagaraScriptSourceBase*	Source;
 	
+	FNiagaraScriptDebuggerInfo DebuggerInfo;
+#endif
+
 	/** Last known compile status. Lets us determine the latest state of the script byte buffer.*/
 	UPROPERTY()
 	ENiagaraScriptCompileStatus LastCompileStatus;
 
-	FNiagaraScriptDebuggerInfo DebuggerInfo;
-#endif
-
 	/** Adjusted every time that we compile this script. Lets us know that we might differ from any cached versions.*/
 	UPROPERTY()
 	FGuid ChangeId;
diff --git a/Engine/Plugins/FX/Niagara/Source/Niagara/Classes/NiagaraScriptExecutionContext.h b/Engine/Plugins/FX/Niagara/Source/Niagara/Classes/NiagaraScriptExecutionContext.h
--- a/Engine/Plugins/FX/Niagara/Source/Niagara/Classes/NiagaraScriptExecutionContext.h
+++ b/Engine/Plugins/FX/Niagara/Source/Niagara/Classes/NiagaraScriptExecutionContext.h
@@ -14,6 +14,7 @@ NiagaraEmitterInstance.h: Niagara emitter simulation class
 #include "NiagaraEmitterHandle.h"
 #include "NiagaraEmitter.h"
 #include "NiagaraParameterStore.h"
+#include "NiagaraTypes.h"
 
 /** 
 Storage class containing actual runtime buffers to be used by the VM and the GPU. 
@@ -28,13 +29,18 @@ public:
 
 	//TODO: This function can probably go away entirely when we replace the FNiagaraParameters and DataInterface info in the script with an FNiagaraParameterStore.
 	//Special care with prev params and internal params will have to be taken there.
-	void Init(UNiagaraScript* Script);
-
+	void Init(UNiagaraScript* Script, ENiagaraSimTarget SimTarget);
+	void AddScriptParams(UNiagaraScript* Script, ENiagaraSimTarget SimTarget);
 	void CopyCurrToPrev();
 
 	bool AddParameter(const FNiagaraVariable& Param, bool bInitInterfaces=true)
 	{
-		return FNiagaraParameterStore::AddParameter(Param, bInitInterfaces);
+		if (FNiagaraParameterStore::AddParameter(Param, bInitInterfaces))
+		{
+			AddPaddedParamSize(Param.GetType(), IndexOf(Param));
+			return true;
+		}
+		return false;
 	}
 
 	bool RemoveParameter(FNiagaraVariable& Param)
@@ -48,11 +54,41 @@ public:
 		check(0);//Can't rename parameters for an execution store.
 	}
 
+	// Just the external parameters, not previous or internal...
 	uint32 GetExternalParameterSize() { return ParameterSize; }
+
+	// The entire buffer padded out by the required alignment of the types..
+	uint32 GetPaddedParameterSizeInBytes() { return PaddedParameterSize; }
+
+	// Helper that converts the data from the base type array internally into the padded out renderer-ready format.
+	void CopyParameterDataToPaddedBuffer(uint8* InTargetBuffer, uint32 InTargetBufferSizeInBytes);
+
+	void Clear()
+	{
+		Empty();
+		PaddedParameterSize = 0;
+	}
+protected:
+	void AddPaddedParamSize(const FNiagaraTypeDefinition& InParamType, uint32 InOffset);
+
 private:
 
 	/** Size of the parameter data not including prev frame values or internal constants. Allows copying into previous parameter values for interpolated spawn scripts. */
 	int32 ParameterSize;
+
+	uint32 PaddedParameterSize;
+
+	struct FPaddingInfo
+	{
+		FPaddingInfo(uint32 InSrcOffset, uint32 InDestOffset, uint32 InSize) : SrcOffset(InSrcOffset), DestOffset(InDestOffset), Size(InSize){}
+		uint32 SrcOffset;
+		uint32 DestOffset;
+		uint32 Size;
+	};
+
+	static void GenerateLayoutInfoInternal(TArray<FPaddingInfo>& Members, uint32& NextMemberOffset, const UStruct* InSrcStruct, uint32 InSrcOffset);
+	
+	TArray<FPaddingInfo> PaddingInfo;
 };
 
 struct FNiagaraDataSetExecutionInfo
@@ -93,13 +129,15 @@ struct FNiagaraScriptExecutionContext
 
 	TArray<FDataSetMeta> DataSetMetaTable;
 
-	
 	static uint32 TickCounter;
 
 	FNiagaraScriptExecutionContext();
+	~FNiagaraScriptExecutionContext();
 
 	bool Init(UNiagaraScript* InScript, ENiagaraSimTarget InTarget);
-	bool Tick(class FNiagaraSystemInstance* Instance);
+	
+
+	bool Tick(class FNiagaraSystemInstance* Instance, ENiagaraSimTarget SimTarget = ENiagaraSimTarget::CPUSim);
 	void PostTick();
 
 	bool Execute(uint32 NumInstances, TArray<FNiagaraDataSetExecutionInfo>& DataSetInfos);
@@ -110,3 +148,93 @@ struct FNiagaraScriptExecutionContext
 
 	bool CanExecute()const;
 };
+
+
+
+
+struct FNiagaraComputeExecutionContext
+{
+	FNiagaraComputeExecutionContext()
+		: MainDataSet(nullptr)
+		, SpawnRateInstances(0)
+		, BurstInstances(0)
+		, EventSpawnTotal(0)
+		, SpawnScript(nullptr)
+		, UpdateScript(nullptr)
+		, RTUpdateScript(0)
+		, RTSpawnScript(0)
+	{
+	}
+
+	void InitParams(UNiagaraScript* InSpawnScript, UNiagaraScript *InUpdateScript, ENiagaraSimTarget SimTarget)
+	{
+		CombinedParamStore.Init(InSpawnScript, SimTarget);
+		CombinedParamStore.AddScriptParams(InUpdateScript, SimTarget);
+		SpawnScript = InSpawnScript;
+		UpdateScript = InUpdateScript;
+		
+		FNiagaraShader *Shader = InSpawnScript->GetRenderThreadScript()->GetShaderGameThread();
+		if (Shader)
+		{
+			DataInterfaceDescriptors = Shader->GetDIBufferDescriptors();
+		}
+		else
+		{
+			DataInterfaceDescriptors = InSpawnScript->GetRenderThreadScript()->GetDataInterfaceBufferDescriptors();
+		}
+
+	}
+
+	void DirtyDataInterfaces()
+	{
+		CombinedParamStore.MarkInterfacesDirty();
+	}
+
+	bool Tick(FNiagaraSystemInstance* ParentSystemInstance)
+	{
+		if (CombinedParamStore.GetInterfacesDirty())
+		{
+			const TArray<UNiagaraDataInterface*> &DataInterfaces = CombinedParamStore.GetDataInterfaces();
+			// We must make sure that the data interfaces match up between the original script values and our overrides...
+			if (DataInterfaceDescriptors.Num() != DataInterfaces.Num())
+			{
+				UE_LOG(LogNiagara, Warning, TEXT("Mismatch between Niagara GPU Execution Context data interfaces and those in its script!"));
+				return false;
+			}
+
+			//Fill the instance data table.
+			if (ParentSystemInstance)
+			{
+				for (int32 i = 0; i < DataInterfaces.Num(); i++)
+				{
+					UNiagaraDataInterface* Interface = DataInterfaces[i];
+					// setup GPU interface; TODO: should probably move this to the GPU execution context
+					Interface->SetupBuffers(DataInterfaceDescriptors[i]);
+				}
+			}
+
+			CombinedParamStore.Tick();
+		}
+
+		return true;
+	}
+
+	class FNiagaraDataSet *MainDataSet;
+	TArray<FNiagaraDataSet*>UpdateEventWriteDataSets;
+	TArray<FNiagaraEventScriptProperties> EventHandlerScriptProps;
+	TArray<FNiagaraDataSet*> EventSets;
+	uint32 SpawnRateInstances;
+	uint32 BurstInstances;
+
+	TArray<int32> EventSpawnCounts;
+	uint32 EventSpawnTotal;
+	UNiagaraScript *SpawnScript;
+	UNiagaraScript *UpdateScript;
+	class FNiagaraScript *RTUpdateScript;
+	class FNiagaraScript *RTSpawnScript;
+	TArray<FNiagaraScriptDataInterfaceInfo> UpdateInterfaces;
+	TArray<uint8, TAlignedHeapAllocator<16>> SpawnParams;		// RT side copy of the parameter data
+	FNiagaraScriptExecutionParameterStore CombinedParamStore;
+	TArray< FDIBufferDescriptorStore >  DataInterfaceDescriptors;
+	static uint32 TickCounter;
+};
diff --git a/Engine/Plugins/FX/Niagara/Source/Niagara/Classes/NiagaraScriptSourceBase.h b/Engine/Plugins/FX/Niagara/Source/Niagara/Classes/NiagaraScriptSourceBase.h
--- a/Engine/Plugins/FX/Niagara/Source/Niagara/Classes/NiagaraScriptSourceBase.h
+++ b/Engine/Plugins/FX/Niagara/Source/Niagara/Classes/NiagaraScriptSourceBase.h
@@ -26,6 +26,10 @@ class UNiagaraScriptSourceBase : public UObject
 {
 	GENERATED_UCLASS_BODY()
 
+#if WITH_EDITOR
+	DECLARE_MULTICAST_DELEGATE(FOnChanged);
+#endif
+
 	TArray<TSharedPtr<EditorExposedVectorConstant> > ExposedVectorConstants;
 	TArray<TSharedPtr<EditorExposedVectorCurveConstant> > ExposedVectorCurveConstants;
 
@@ -56,4 +60,20 @@ class UNiagaraScriptSourceBase : public UObject
 
 	/** Do any cleanup to return this data source to its original state before precompilation. Note that this must be called after the "compile" operation.*/
 	virtual void PostCompile() {}
+
+	/** 
+	 * Allows the derived editor only script source to handle a post load requested by an owning emitter. 
+	 * @param OwningEmitter The emitter requesting the post load.
+	 * @returns True if the emitter will need to be recompiled, otherwise false.
+	 */
+	virtual bool PostLoadFromEmitter(UNiagaraEmitter& OwningEmitter) { return false; }
+
+	/** Adds a module if it isn't already in the graph. If the module isn't found bOutFoundModule will be false. If it is found and it did need to be added, the function returns true. If it already exists, it returns false. */
+	NIAGARA_API virtual bool AddModuleIfMissing(FString ModulePath, ENiagaraScriptUsage Usage, bool& bOutFoundModule) { bOutFoundModule = false; return false; }
+#if WITH_EDITOR
+	FOnChanged& OnChanged() { return OnChangedDelegate; }
+
+protected:
+	FOnChanged OnChangedDelegate;
+#endif
 };
diff --git a/Engine/Plugins/FX/Niagara/Source/Niagara/Classes/NiagaraSystem.h b/Engine/Plugins/FX/Niagara/Source/Niagara/Classes/NiagaraSystem.h
--- a/Engine/Plugins/FX/Niagara/Source/Niagara/Classes/NiagaraSystem.h
+++ b/Engine/Plugins/FX/Niagara/Source/Niagara/Classes/NiagaraSystem.h
@@ -9,28 +9,9 @@
 #include "NiagaraDataSet.h"
 #include "NiagaraEmitterInstance.h"
 #include "NiagaraEmitterHandle.h"
+#include "NiagaraParameterCollection.h"
 #include "NiagaraSystem.generated.h"
 
-//TODO:
-/** Defines a number of values and defaults shared among all Niagara Systems that use this category. */
-// UCLASS()
-// class NIAGARA_API UNiagaraSystemCategory : public UObject
-// {
-// 	GENERATED_UCLASS_BODY()
-// 
-// 	/** Name of this System Category. */
-// 	UPROPERTY(EditAnywhere, Category=Category)
-// 	FName Name;
-// 
-// 	/** Category specific to the global instance of a particular Niagara Parameter Collection. */
-// 	UPROPERTY(EditAnywhere, Category=Category)
-// 	TArray<UNiagaraParameterCollectionInstance*> ParameterCollectionOverrides;
-// 
-// 	/** TODO: How to do an icon? */
-// 	//UPROPERTY(EditAnywhere, Category = Category)
-// 	//UTexture2D* Icon
-// };
-
 USTRUCT()
 struct FNiagaraEmitterSpawnAttributes
 {
@@ -41,7 +22,7 @@ struct FNiagaraEmitterSpawnAttributes
 };
 
 /** Container for multiple emitters that combine together to create a particle system effect.*/
-UCLASS()
+UCLASS(BlueprintType)
 class NIAGARA_API UNiagaraSystem : public UObject
 {
 	GENERATED_UCLASS_BODY()
@@ -51,9 +32,13 @@ public:
 	void PostInitProperties();
 	void Serialize(FArchive& Ar)override;
 	virtual void PostLoad() override;
+#if WITH_EDITOR
+	virtual void PostEditChangeProperty(struct FPropertyChangedEvent& PropertyChangedEvent) override;
+#endif
 
 	/** Gets an array of the emitter handles. */
 	const TArray<FNiagaraEmitterHandle>& GetEmitterHandles();
+	const TArray<FNiagaraEmitterHandle>& GetEmitterHandles()const;
 
 	/** Returns true if this system is valid and can be instanced. False otherwise. */
 	bool IsValid()const;
@@ -61,7 +46,7 @@ public:
 #if WITH_EDITORONLY_DATA
 	/** Adds a new emitter handle to this System.  The new handle exposes an Instance value which is a copy of the
 		original asset. */
-	FNiagaraEmitterHandle AddEmitterHandle(const UNiagaraEmitter& SourceEmitter, FName EmitterName);
+	FNiagaraEmitterHandle AddEmitterHandle(UNiagaraEmitter& SourceEmitter, FName EmitterName);
 
 	/** Adds a new emitter handle to this System.  The new handle will not copy the emitter and any changes made to it's
 		Instance value will modify the original asset.  This should only be used in the emitter toolkit for simulation
@@ -71,13 +56,14 @@ public:
 	/** Duplicates an existing emitter handle and adds it to the System.  The new handle will reference the same source asset,
 		but will have a copy of the duplicated Instance value. */
 	FNiagaraEmitterHandle DuplicateEmitterHandle(const FNiagaraEmitterHandle& EmitterHandleToDuplicate, FName EmitterName);
-#endif
 
 	/** Removes the provided emitter handle. */
 	void RemoveEmitterHandle(const FNiagaraEmitterHandle& EmitterHandleToDelete);
 
 	/** Removes the emitter handles which have an Id in the supplied set. */
 	void RemoveEmitterHandlesById(const TSet<FGuid>& HandlesToRemove);
+#endif
+
 
 	FNiagaraEmitterHandle& GetEmitterHandle(int Idx)
 	{
@@ -105,26 +91,26 @@ public:
 	UNiagaraScript* GetSystemSpawnScript(bool bSolo=false);
 	UNiagaraScript* GetSystemUpdateScript(bool bSolo=false);
 
+	bool IsReadyToRun() const;
+
 #if WITH_EDITORONLY_DATA
-	/** Get whether or not we auto-import applied changes. See bAutoImportChangedEmitters member documentation for details.*/
-	bool GetAutoImportChangedEmitters() const;
-	/** Set whether or not we auto-import applied changes. See bAutoImportChangedEmitters member documentation for details.*/
-	void SetAutoImportChangedEmitters(bool bShouldImport);
-		
 	/** Called to query whether or not this emitter is referenced as the source to any emitter handles for this System.*/
-	bool ReferencesSourceEmitter(UNiagaraEmitter* Emitter);
+	bool ReferencesSourceEmitter(UNiagaraEmitter& Emitter);
+
+	/** Updates all handles which use this emitter as their source. */
+	void UpdateFromEmitterChanges(UNiagaraEmitter& ChangedSourceEmitter);
+
+	/** Updates the system's rapid iteration parameters from a specific emitter. */
+	void RefreshSystemParametersFromEmitter(const FNiagaraEmitterHandle& EmitterHandle);
 
-	/** Goes through all handles on the System and compares them to the latest on their source. Refreshes from source if needed. Returns true if any were refreshed.*/
-	bool ResynchronizeAllHandles();
+	/** Removes the system's rapid iteration parameters for a specific emitter. */
+	void RemoveSystemParametersForEmitter(const FNiagaraEmitterHandle& EmitterHandle);
 
 	/** Recompile all emitters and the System script associated with this System.*/
-	void Compile();
+	void Compile(bool bForce);
 
 	/** Compile just the System script associated with this System.*/
-	void CompileScripts(TArray<ENiagaraScriptCompileStatus>& OutScriptStatuses, TArray<FString>& OutGraphLevelErrorMessages, TArray<FString>& PathNames, TArray<UNiagaraScript*>& Scripts);
-
-	/** Uses whether or not the GetAutoImportChangedEmitters is set to ResynchronizeAllHandles.*/
-	bool CheckForUpdates();
+	void CompileScripts(TArray<ENiagaraScriptCompileStatus>& OutScriptStatuses, TArray<FString>& OutGraphLevelErrorMessages, TArray<FString>& PathNames, TArray<UNiagaraScript*>& Scripts, bool bForce);
 
 	/** Gets editor specific data stored with this system. */
 	UObject* GetEditorData();
@@ -135,15 +121,54 @@ public:
 	/** Sets editor specific data stored with this system. */
 	void SetEditorData(UObject* InEditorData);
 
+	/** Internal: The thumbnail image.*/
+	UPROPERTY()
+	class UTexture2D* ThumbnailImage;
+
+	/** Internal: Indicates the thumbnail image is out of date.*/
+	UPROPERTY()
+	uint32 ThumbnailImageOutOfDate : 1;
+
+	bool GetIsolateEnabled() const;
+	void SetIsolateEnabled(bool bIsolate);
 #endif
 
+	bool ShouldAutoDeactivate() const { return bAutoDeactivate; }
+	bool IsLooping() const;
+
 	const TArray<FNiagaraEmitterSpawnAttributes>& GetEmitterSpawnAttributes()const {	return EmitterSpawnAttributes;	};
 
+	bool UsesCollection(const UNiagaraParameterCollection* Collection)const;
+#if WITH_EDITORONLY_DATA
+	bool UsesScript(const UNiagaraScript* Script)const;
+#endif
+
+	FORCEINLINE UNiagaraParameterCollectionInstance* GetParameterCollectionOverride(UNiagaraParameterCollection* Collection)
+	{
+		UNiagaraParameterCollectionInstance** Found = ParameterCollectionOverrides.FindByPredicate(
+			[&](const UNiagaraParameterCollectionInstance* CheckInst)
+		{
+			return CheckInst && Collection == CheckInst->Collection;
+		});
+
+		return Found ? *Found : nullptr;
+	}
+
+
+private:
+#if WITH_EDITORONLY_DATA
+	INiagaraModule::FMergeEmitterResults MergeChangesForEmitterHandle(FNiagaraEmitterHandle& EmitterHandle);
+#endif
+
 protected:
+
 	/** Handles to the emitter this System will simulate. */
 	UPROPERTY(VisibleAnywhere, Category = "Emitters")
 	TArray<FNiagaraEmitterHandle> EmitterHandles;
 
+	UPROPERTY(EditAnywhere, Category="System")
+	TArray<UNiagaraParameterCollectionInstance*> ParameterCollectionOverrides;
+
 // 	/** Category of this system. */
 // 	UPROPERTY(EditAnywhere, Category = System)
 // 	UNiagaraSystemCategory* Category;
@@ -174,17 +199,17 @@ protected:
 	UPROPERTY()
 	FNiagaraParameterStore ExposedParameters;
 
-#if WITH_EDITORONLY_DATA
-	/** Systems are the final step in the process of creating a Niagara system. Artists may wish to lock an System so that it only uses
-	the handle's cached version of the scripts, rather than the external assets that may be subject to changes. If this flag is set, we will only update
-	the emitters if told to do so explicitly by the user.*/
-	UPROPERTY()
-	uint32 bAutoImportChangedEmitters : 1;	
+#if WITH_EDITORONLY_DATA	
 
 	/** Data used by the editor to maintain UI state etc.. */
 	UPROPERTY()
 	UObject* EditorData;
+
+	bool bIsolateEnabled;
 #endif
 
+	UPROPERTY(EditAnywhere, Category = Performance, meta = (ToolTip = "Auto-deactivate system if all emitters are determined to not spawn particles again, regardless of lifetime."))
+	bool bAutoDeactivate;
+
 	void InitEmitterSpawnAttributes();
 };
diff --git a/Engine/Plugins/FX/Niagara/Source/Niagara/Private/NiagaraActor.cpp b/Engine/Plugins/FX/Niagara/Source/Niagara/Private/NiagaraActor.cpp
--- a/Engine/Plugins/FX/Niagara/Source/Niagara/Private/NiagaraActor.cpp
+++ b/Engine/Plugins/FX/Niagara/Source/Niagara/Private/NiagaraActor.cpp
@@ -26,12 +26,12 @@ ANiagaraActor::ANiagaraActor(const FObjectInitializer& ObjectInitializer)
 		struct FConstructorStatics
 		{
 			ConstructorHelpers::FObjectFinderOptional<UTexture2D> SpriteTextureObject;
-			FName ID_Systems;
-			FText NAME_Systems;
+			FName ID_Effects;
+			FText NAME_Effects;
 			FConstructorStatics()
 				: SpriteTextureObject(TEXT("/Engine/EditorResources/S_Emitter"))
-				, ID_Systems(TEXT("Systems"))
-				, NAME_Systems(NSLOCTEXT("SpriteCategory", "Systems", "Systems"))
+				, ID_Effects(TEXT("Effects"))
+				, NAME_Effects(NSLOCTEXT("SpriteCategory", "Effects", "Effects"))
 			{
 			}
 		};
@@ -43,8 +43,8 @@ ANiagaraActor::ANiagaraActor(const FObjectInitializer& ObjectInitializer)
 			SpriteComponent->RelativeScale3D = FVector(0.5f, 0.5f, 0.5f);
 			SpriteComponent->bHiddenInGame = true;
 			SpriteComponent->bIsScreenSizeScaled = true;
-			SpriteComponent->SpriteInfo.Category = ConstructorStatics.ID_Systems;
-			SpriteComponent->SpriteInfo.DisplayName = ConstructorStatics.NAME_Systems;
+			SpriteComponent->SpriteInfo.Category = ConstructorStatics.ID_Effects;
+			SpriteComponent->SpriteInfo.DisplayName = ConstructorStatics.NAME_Effects;
 			SpriteComponent->SetupAttachment(NiagaraComponent);
 			SpriteComponent->bReceivesDecals = false;
 		}
@@ -55,10 +55,11 @@ ANiagaraActor::ANiagaraActor(const FObjectInitializer& ObjectInitializer)
 
 			ArrowComponent->ArrowSize = 1.5f;
 			ArrowComponent->bTreatAsASprite = true;
-			ArrowComponent->SpriteInfo.Category = ConstructorStatics.ID_Systems;
-			ArrowComponent->SpriteInfo.DisplayName = ConstructorStatics.NAME_Systems;
-			ArrowComponent->SetupAttachment(NiagaraComponent);
 			ArrowComponent->bIsScreenSizeScaled = true;
+			ArrowComponent->SpriteInfo.Category = ConstructorStatics.ID_Effects;
+			ArrowComponent->SpriteInfo.DisplayName = ConstructorStatics.NAME_Effects;
+			ArrowComponent->SetupAttachment(NiagaraComponent);
+			ArrowComponent->bAbsoluteScale = true;
 		}
 	}
 #endif // WITH_EDITORONLY_DATA
diff --git a/Engine/Plugins/FX/Niagara/Source/Niagara/Private/NiagaraCollision.cpp b/Engine/Plugins/FX/Niagara/Source/Niagara/Private/NiagaraCollision.cpp
--- a/Engine/Plugins/FX/Niagara/Source/Niagara/Private/NiagaraCollision.cpp
+++ b/Engine/Plugins/FX/Niagara/Source/Niagara/Private/NiagaraCollision.cpp
@@ -5,6 +5,7 @@
 #include "NiagaraSystemInstance.h"
 #include "NiagaraStats.h"
 #include "NiagaraComponent.h"
+#include "PhysicalMaterials/PhysicalMaterial.h"
 
 DECLARE_CYCLE_STAT(TEXT("Collision"), STAT_NiagaraCollision, STATGROUP_Niagara);
 DECLARE_CYCLE_STAT(TEXT("Event Emission"), STAT_NiagaraEventWrite, STATGROUP_Niagara);
@@ -14,11 +15,15 @@ void FNiagaraCollisionBatch::KickoffNewBatch(FNiagaraEmitterInstance *Sim, float
 	SCOPE_CYCLE_COUNTER(STAT_NiagaraCollision);
 	FNiagaraVariable PosVar(FNiagaraTypeDefinition::GetVec3Def(), "Position");
 	FNiagaraVariable VelVar(FNiagaraTypeDefinition::GetVec3Def(), "Velocity");
+	FNiagaraVariable SizeVar(FNiagaraTypeDefinition::GetVec2Def(), "SpriteSize");
 	FNiagaraVariable TstVar(FNiagaraTypeDefinition::GetBoolDef(), "PerformCollision");
 	FNiagaraDataSetIterator<FVector> PosIt(Sim->GetData(), PosVar, 0, false);
 	FNiagaraDataSetIterator<FVector> VelIt(Sim->GetData(), VelVar, 0, false);
+	FNiagaraDataSetIterator<FVector2D> SizeIt(Sim->GetData(), SizeVar, 0, false);
 	//FNiagaraDataSetIterator<int32> TstIt(Sim->GetData(), TstVar, 0, false);
 
+	bool bUseSize = SizeIt.IsValid();
+
 	if (!PosIt.IsValid() || !VelIt.IsValid() /*|| !TstIt.IsValid()*/)
 	{
 		return;
@@ -34,15 +39,33 @@ void FNiagaraCollisionBatch::KickoffNewBatch(FNiagaraEmitterInstance *Sim, float
 			//int32 TestCollision = *TstIt;
 			//if (TestCollision)
 			{
-				check(PosIt.IsValid() && VelIt.IsValid());
+				check(PosIt.IsValid() && VelIt.IsValid() && (bUseSize == false || SizeIt.IsValid()));
 				FVector Position;
+				FVector EndPosition;
 				FVector Velocity;
+
 				PosIt.Get(Position);
 				VelIt.Get(Velocity);
+				EndPosition = Position + Velocity * DeltaSeconds;
+
+				if (bUseSize)
+				{
+					//TODO:  Handle mesh particles too.  Also this can probably be better and or faster.
+					FVector2D SpriteSize;
+					SizeIt.Get(SpriteSize);
+					float MaxSize = FMath::Max(SpriteSize.X, SpriteSize.Y);
+
+					float Length;
+					FVector Direction;
+					Velocity.ToDirectionAndLength(Direction, Length);
+
+					Position -= Direction * (MaxSize / 2);
+					EndPosition += Direction * (MaxSize / 2);
+				}
 
 				FCollisionQueryParams QueryParams(SCENE_QUERY_STAT(NiagraAsync));
 				QueryParams.OwnerTag = "Niagara";
-				FTraceHandle Handle = SystemWorld->AsyncLineTraceByChannel(EAsyncTraceType::Single, Position, Position + Velocity*DeltaSeconds, ECollisionChannel::ECC_WorldStatic, QueryParams, FCollisionResponseParams::DefaultResponseParam, nullptr, i);
+				FTraceHandle Handle = SystemWorld->AsyncLineTraceByChannel(EAsyncTraceType::Single, Position, EndPosition, ECollisionChannel::ECC_WorldStatic, QueryParams, FCollisionResponseParams::DefaultResponseParam, nullptr, i);
 				FNiagaraCollisionTrace Trace;
 				Trace.CollisionTraceHandle = Handle;
 				Trace.SourceParticleIndex = i;
@@ -158,4 +181,126 @@ void FNiagaraCollisionBatch::GenerateEventsFromResults(FNiagaraEmitterInstance *
 		}
 
 	}
-}
\ No newline at end of file
+}
+
+
+
+
+
+
+
+
+int32 FNiagaraDICollisionQueryBatch::SubmitQuery(FVector Position, FVector Velocity, float CollisionSize, float DeltaSeconds)
+{
+	SCOPE_CYCLE_COUNTER(STAT_NiagaraCollision);
+	int32 Ret = INDEX_NONE;
+	if (CollisionWorld)
+	{
+		//int32 TestCollision = *TstIt;
+		//if (TestCollision)
+		{
+			FVector EndPosition = Position + Velocity * DeltaSeconds;
+
+			float Length;
+			FVector Direction;
+			Velocity.ToDirectionAndLength(Direction, Length);
+			Position -= Direction * (CollisionSize / 2);
+			EndPosition += Direction * (CollisionSize/ 2);
+
+			FCollisionQueryParams QueryParams(SCENE_QUERY_STAT(NiagraAsync));
+			QueryParams.OwnerTag = "Niagara";
+			QueryParams.bTraceAsyncScene = true;
+			QueryParams.bFindInitialOverlaps = false;
+			QueryParams.bReturnFaceIndex = false;
+			QueryParams.bReturnPhysicalMaterial = true;
+			QueryParams.bTraceComplex = false;
+			QueryParams.bIgnoreTouches = true;
+			FTraceHandle Handle = CollisionWorld->AsyncLineTraceByChannel(EAsyncTraceType::Single, Position, EndPosition, ECollisionChannel::ECC_WorldStatic, QueryParams, FCollisionResponseParams::DefaultResponseParam, nullptr, TraceID);
+			FNiagaraCollisionTrace Trace;
+			Trace.CollisionTraceHandle = Handle;
+			Trace.SourceParticleIndex = TraceID;
+			Trace.OriginalVelocity = Velocity;
+
+			int32 TraceIdx = CollisionTraces[GetWriteBufferIdx()].Add(Trace);
+			IdToTraceIdx[GetWriteBufferIdx()].Add(TraceID) = TraceIdx;
+			Ret = TraceID;
+			TraceID++;
+		}
+	}
+
+	return Ret;
+}
+
+
+bool FNiagaraDICollisionQueryBatch::GetQueryResult(uint32 InTraceID, FNiagaraDICollsionQueryResult &Result)
+{
+	int32 *TraceIdxPtr = IdToTraceIdx[GetReadBufferIdx()].Find(InTraceID);
+	int32 TraceIdx = INDEX_NONE;
+
+	if (TraceIdxPtr)
+	{
+		TraceIdx = *TraceIdxPtr;
+		check(CollisionTraces[GetReadBufferIdx()].IsValidIndex(TraceIdx));
+		FNiagaraCollisionTrace &CurTrace = CollisionTraces[GetReadBufferIdx()][TraceIdx];
+		FTraceHandle Handle = CurTrace.CollisionTraceHandle;
+		FTraceDatum CurData;
+		// wait for trace handles; this should block rarely to never
+		bool bReady = CollisionWorld->QueryTraceData(Handle, CurData);
+		/*
+		if (!bReady)
+		{
+			// if the query came back false, it's possible that the hanle is invalid for some reason; skip in that case
+			// TODO: handle this more gracefully
+			if (!CollisionWorld->IsTraceHandleValid(Handle, false))
+			{
+				break;
+			}
+			break;
+			}
+		}
+		*/
+
+		if (bReady && CurData.OutHits.Num())
+		{
+			// grab the first hit that blocks
+			FHitResult *Hit = FHitResult::GetFirstBlockingHit(CurData.OutHits);
+			if (Hit && Hit->IsValidBlockingHit())
+			{
+				/*
+				FNiagaraCollisionEventPayload Event;
+				Event.CollisionNormal = Hit->ImpactNormal;
+				Event.CollisionPos = Hit->ImpactPoint;
+				Event.CollisionVelocity = CurCheck.OriginalVelocity;
+				Event.ParticleIndex = CurData.UserData;
+				check(!Event.CollisionNormal.ContainsNaN());
+				check(Event.CollisionNormal.IsNormalized());
+				check(!Event.CollisionPos.ContainsNaN());
+				check(!Event.CollisionVelocity.ContainsNaN());
+
+				// TODO add to unique list of physical materials for Blueprint
+				Event.PhysicalMaterialIndex = 0;// Hit->PhysMaterial->GetUniqueID();
+
+				Payloads.Add(Event);
+				*/
+				Result.CollisionPos = Hit->ImpactPoint;
+				Result.CollisionNormal = Hit->ImpactNormal;
+				Result.CollisionVelocity = CurTrace.OriginalVelocity;
+				Result.TraceID = InTraceID;
+				Result.PhysicalMaterialIdx = Hit->PhysMaterial->GetUniqueID();
+				Result.Friction = Hit->PhysMaterial->Friction;
+				Result.Restitution = Hit->PhysMaterial->Restitution;
+				return true;
+			}
+		}
+	}
+
+	return false;
+}
+
+
+/*
+int32 FNiagaraDICollisionQueryBatch::Tick()
+{
+
+}
+/*/
\ No newline at end of file
diff --git a/Engine/Plugins/FX/Niagara/Source/Niagara/Private/NiagaraCommon.cpp b/Engine/Plugins/FX/Niagara/Source/Niagara/Private/NiagaraCommon.cpp
--- a/Engine/Plugins/FX/Niagara/Source/Niagara/Private/NiagaraCommon.cpp
+++ b/Engine/Plugins/FX/Niagara/Source/Niagara/Private/NiagaraCommon.cpp
@@ -65,6 +65,12 @@ FString FNiagaraTypeHelper::ToString(const uint8* ValueData, const UScriptStruct
 
 FNiagaraSystemUpdateContext::~FNiagaraSystemUpdateContext()
 {
+	INiagaraModule& NiagaraModule = FModuleManager::LoadModuleChecked<INiagaraModule>("Niagara");
+	for (UNiagaraSystem* Sys : SystemSimsToDestroy)
+	{
+		NiagaraModule.DestroyAllSystemSimulations(Sys);
+	}
+
 	for (UNiagaraComponent* Comp : ComponentsToReInit)
 	{
 		Comp->ReinitializeSystem();
@@ -75,6 +81,16 @@ FNiagaraSystemUpdateContext::~FNiagaraSystemUpdateContext()
 	}
 }
 
+void FNiagaraSystemUpdateContext::AddAll(bool bReInit)
+{
+	for (TObjectIterator<UNiagaraComponent> It; It; ++It)
+	{
+		UNiagaraComponent* Comp = *It;
+		check(Comp);
+		AddInternal(Comp, bReInit);
+	}
+}
+
 void FNiagaraSystemUpdateContext::Add(const UNiagaraSystem* System, bool bReInit)
 {
 	for (TObjectIterator<UNiagaraComponent> It; It; ++It)
@@ -109,8 +125,8 @@ void FNiagaraSystemUpdateContext::Add(const UNiagaraScript* Script, bool bReInit
 	{
 		UNiagaraComponent* Comp = *It;
 		check(Comp);
-		FNiagaraSystemInstance* SystemInst = Comp->GetSystemInstance();
-		if (SystemInst && SystemInst->UsesScript(Script))
+		UNiagaraSystem* System = Comp->GetAsset();
+		if (System && System->UsesScript(Script))
 		{
 			AddInternal(Comp, bReInit);
 		}
@@ -153,9 +169,35 @@ void FNiagaraSystemUpdateContext::AddInternal(UNiagaraComponent* Comp, bool bReI
 	if (bReInit)
 	{
 		ComponentsToReInit.AddUnique(Comp);
+		SystemSimsToDestroy.AddUnique(Comp->GetAsset());
 	}
 	else
 	{
 		ComponentsToReset.AddUnique(Comp);
 	}
+}
+
+FName NIAGARA_API FNiagaraUtilities::GetUniqueName(FName CandidateName, const TSet<FName>& ExistingNames)
+{
+	if (ExistingNames.Contains(CandidateName) == false)
+	{
+		return CandidateName;
+	}
+
+	FString CandidateNameString = CandidateName.ToString();
+	FString BaseNameString = CandidateNameString;
+	if (CandidateNameString.Len() >= 3 && CandidateNameString.Right(3).IsNumeric())
+	{
+		BaseNameString = CandidateNameString.Left(CandidateNameString.Len() - 3);
+	}
+
+	FName UniqueName = FName(*BaseNameString);
+	int32 NameIndex = 1;
+	while (ExistingNames.Contains(UniqueName))
+	{
+		UniqueName = FName(*FString::Printf(TEXT("%s%03i"), *BaseNameString, NameIndex));
+		NameIndex++;
+	}
+
+	return UniqueName;
 }
\ No newline at end of file
diff --git a/Engine/Plugins/FX/Niagara/Source/Niagara/Private/NiagaraComponent.cpp b/Engine/Plugins/FX/Niagara/Source/Niagara/Private/NiagaraComponent.cpp
--- a/Engine/Plugins/FX/Niagara/Source/Niagara/Private/NiagaraComponent.cpp
+++ b/Engine/Plugins/FX/Niagara/Source/Niagara/Private/NiagaraComponent.cpp
@@ -19,7 +19,11 @@
 #include "NiagaraParameterCollection.h"
 #include "NiagaraWorldManager.h"
 
-DECLARE_CYCLE_STAT(TEXT("Gen Verts"),STAT_NiagaraGenerateVertices,STATGROUP_Niagara);
+DECLARE_CYCLE_STAT(TEXT("Component Tick (GT)"), STAT_NiagaraComponentTick, STATGROUP_Niagara);
+DECLARE_CYCLE_STAT(TEXT("Activate (GT)"), STAT_NiagaraComponentActivate, STATGROUP_Niagara);
+DECLARE_CYCLE_STAT(TEXT("Deactivate (GT)"), STAT_NiagaraComponentDeactivate, STATGROUP_Niagara);
+DECLARE_CYCLE_STAT(TEXT("Send Render Data (GT)"), STAT_NiagaraComponentSendRenderData, STATGROUP_Niagara);
+DECLARE_CYCLE_STAT(TEXT("Get Dynamic Mesh Elements (RT)"), STAT_NiagaraComponentGetDynamicMeshElements, STATGROUP_Niagara);
 
 DEFINE_LOG_CATEGORY(LogNiagara);
 
@@ -48,6 +52,12 @@ FNiagaraSceneProxy::FNiagaraSceneProxy(const UNiagaraComponent* InComponent)
 	}
 }
 
+SIZE_T FNiagaraSceneProxy::GetTypeHash() const
+{
+	static size_t UniquePointer;
+	return reinterpret_cast<size_t>(&UniquePointer);
+}
+
 void FNiagaraSceneProxy::AddEmitterRenderer(NiagaraRenderer* Renderer)
 {
 	EmitterRenderers.Add(Renderer); 
@@ -164,6 +174,7 @@ void FNiagaraSceneProxy::SetRenderingEnabled(bool bInRenderingEnabled)
 
 void FNiagaraSceneProxy::GetDynamicMeshElements(const TArray<const FSceneView*>& Views, const FSceneViewFamily& ViewFamily, uint32 VisibilityMap, FMeshElementCollector& Collector) const
 {
+	SCOPE_CYCLE_COUNTER(STAT_NiagaraComponentGetDynamicMeshElements);
 	for (NiagaraRenderer* Renderer : EmitterRenderers)
 	{
 		if (Renderer)
@@ -171,6 +182,21 @@ void FNiagaraSceneProxy::GetDynamicMeshElements(const TArray<const FSceneView*>&
 			Renderer->GetDynamicMeshElements(Views, ViewFamily, VisibilityMap, Collector, this);
 		}
 	}
+
+	if (ViewFamily.EngineShowFlags.Particles)
+	{
+		for (int32 ViewIndex = 0; ViewIndex < Views.Num(); ViewIndex++)
+		{
+			if (VisibilityMap & (1 << ViewIndex))
+			{
+				RenderBounds(Collector.GetPDI(ViewIndex), ViewFamily.EngineShowFlags, GetBounds(), IsSelected());
+				if (HasCustomOcclusionBounds())
+				{
+					RenderBounds(Collector.GetPDI(ViewIndex), ViewFamily.EngineShowFlags, GetCustomOcclusionBounds(), IsSelected());
+				}
+			}
+		}
+	}
 }
 
 
@@ -215,10 +241,12 @@ void FNiagaraSceneProxy::GatherSimpleLights(const FSceneViewFamily& ViewFamily,
 
 UNiagaraComponent::UNiagaraComponent(const FObjectInitializer& ObjectInitializer)
 	: Super(ObjectInitializer)
+	, OverrideParameters(this)
 	, bForceSolo(false)
 	, AgeUpdateMode(EAgeUpdateMode::TickDeltaTime)
 	, DesiredAge(0.0f)
 	, SeekDelta(1 / 30.0f)
+	, bAutoDestroy(false)
 {
 	PrimaryComponentTick.bCanEverTick = true;
 	PrimaryComponentTick.TickGroup = TG_DuringPhysics;
@@ -227,13 +255,21 @@ UNiagaraComponent::UNiagaraComponent(const FObjectInitializer& ObjectInitializer
 	bTickInEditor = true;
 	bAutoActivate = true;
 	bRenderingEnabled = true;
+	SavedAutoAttachRelativeScale3D = FVector(1.f, 1.f, 1.f);
 }
 
 
 void UNiagaraComponent::TickComponent(float DeltaSeconds, enum ELevelTick TickType, FActorComponentTickFunction* ThisTickFunction)
 {
+	SCOPE_CYCLE_COUNTER(STAT_NiagaraComponentTick);
+
+	if (!bIsActive && bAutoActivate)
+	{
+		Activate();
+	}
+
 	check(SystemInstance->IsSolo());
-	if (bIsActive && SystemInstance.Get())
+	if (bIsActive && SystemInstance.Get() && !SystemInstance->IsComplete())
 	{
 		if (SceneProxy != nullptr)
 		{
@@ -241,6 +277,12 @@ void UNiagaraComponent::TickComponent(float DeltaSeconds, enum ELevelTick TickTy
 			NiagaraProxy->SetRenderingEnabled(bRenderingEnabled);
 		}
 
+		// If the interfaces have changed in a meaningful way, we need to potentially rebind and update the values.
+		if (OverrideParameters.GetInterfacesDirty())
+		{
+			SystemInstance->Reset(FNiagaraSystemInstance::EResetMode::ReInit);
+		}
+
 		if (AgeUpdateMode == EAgeUpdateMode::TickDeltaTime)
 		{
 			SystemInstance->ComponentTick(DeltaSeconds);
@@ -255,7 +297,7 @@ void UNiagaraComponent::TickComponent(float DeltaSeconds, enum ELevelTick TickTy
 
 			if (AgeDiff < 0.0f)
 			{
-				SystemInstance->Reset(FNiagaraSystemInstance::EResetMode::ImmediateReset);
+				SystemInstance->Reset(FNiagaraSystemInstance::EResetMode::ResetAll);
 				AgeDiff = DesiredAge - SystemInstance->GetAge();
 			}
 
@@ -264,28 +306,30 @@ void UNiagaraComponent::TickComponent(float DeltaSeconds, enum ELevelTick TickTy
 				// TODO: This maximum should be configurable, and we should not draw the System until we are caught up.
 				static const float MaxSimTime = 1;
 				//Take SeekDelta to be the max delta time we'll accept.
-				FNiagaraSystemSimulation* SystemSim = GetSystemSimulation();
-
-				float SimTime = FMath::Min(AgeDiff, MaxSimTime);
-				int32 NumSteps = FMath::CeilToInt(SimTime / SeekDelta);
-				//Num steps must be odd. This is because when doing > 1 we have to flip the solo data set for in the system sim. This has to end up in the same flip state it was in before our ticks.
-				NumSteps = NumSteps % 2 == 0 ? NumSteps + 1 : NumSteps;
-				float RealSeekDelta = SimTime / (float)NumSteps;
-
-				for (int32 i = NumSteps - 1; i >= 0; --i)
+				TSharedPtr<FNiagaraSystemSimulation, ESPMode::ThreadSafe> SystemSim = GetSystemSimulation();
+				if (SystemSim.IsValid())
 				{
-					SystemInstance->ComponentTick(RealSeekDelta);
-					if (i > 0)
+					float SimTime = FMath::Min(AgeDiff, MaxSimTime);
+					int32 NumSteps = FMath::CeilToInt(SimTime / SeekDelta);
+					//Num steps must be odd. This is because when doing > 1 we have to flip the solo data set for in the system sim. This has to end up in the same flip state it was in before our ticks.
+					NumSteps = NumSteps % 2 == 0 ? NumSteps + 1 : NumSteps;
+					float RealSeekDelta = SimTime / (float)NumSteps;
+
+					for (int32 i = NumSteps - 1; i >= 0; --i)
 					{
-						SystemSim->TickSoloDataSet();
+						SystemInstance->ComponentTick(RealSeekDelta);
+						if (i > 0)
+						{
+							SystemSim->TickSoloDataSet();
+						}
 					}
 				}
 			}
 		}
+		
+		UpdateComponentToWorld();
+		MarkRenderDynamicDataDirty();
 	}
-
-	UpdateComponentToWorld();
-	MarkRenderDynamicDataDirty();
 }
 
 const UObject* UNiagaraComponent::AdditionalStatObject() const
@@ -302,11 +346,6 @@ void UNiagaraComponent::ReinitializeSystem()
 {
 	DestroyInstance();
 	Activate();
-
-// 	if (IsRegistered() && SystemInstance.IsValid())
-// 	{
-// 		SystemInstance->Init(Asset, true, bForceSolo);
-// 	}
 }
 
 bool UNiagaraComponent::GetRenderingEnabled() const
@@ -321,31 +360,44 @@ void UNiagaraComponent::SetRenderingEnabled(bool bInRenderingEnabled)
 
 void UNiagaraComponent::Activate(bool bReset /* = false */)
 {
-	if (Asset == nullptr || Asset->IsValid() == false)
+	SCOPE_CYCLE_COUNTER(STAT_NiagaraComponentActivate);
+	if (Asset == nullptr)
 	{
 		DestroyInstance();
 
-		UE_LOG(LogNiagara, Warning, TEXT("Failed to activate Niagara Component!"));
+		UE_LOG(LogNiagara, Warning, TEXT("Failed to activate Niagara Component due to missing or invalid asset!"));
 		SetComponentTickEnabled(false);
 		return;
 	}
 
+	if (!Asset->IsReadyToRun())
+	{
+		return;
+	}
+
 	if(!IsRegistered())
 	{
 		return;
 	}
+
+	// If the particle system can't ever render (ie on dedicated server or in a commandlet) than do not activate...
+	if (!FApp::CanEverRender())
+	{
+		return;
+	}
 	
 	Super::Activate(bReset);
 
 	if (SystemInstance.IsValid() == false)
 	{
-		if (FNiagaraSystemSimulation* SystemSim = GetSystemSimulation())
+		TSharedPtr<FNiagaraSystemSimulation, ESPMode::ThreadSafe> SystemSim = GetSystemSimulation();
+		if (SystemSim.IsValid())
 		{
 			SystemInstance = MakeUnique<FNiagaraSystemInstance>(this);
 #if WITH_EDITORONLY_DATA
 			OnSystemInstanceChangedDelegate.Broadcast();
 #endif
-			SystemInstance->Init(SystemSim, true, bForceSolo);
+			SystemInstance->Init(SystemSim.ToSharedRef(), bForceSolo);
 			bReset = false;
 		}
 		else
@@ -354,6 +406,34 @@ void UNiagaraComponent::Activate(bool bReset /* = false */)
 		}
 	}
 
+	// Auto attach if requested
+	const bool bWasAutoAttached = bDidAutoAttach;
+	bDidAutoAttach = false;
+	if (bAutoManageAttachment)
+	{
+		USceneComponent* NewParent = AutoAttachParent.Get();
+		if (NewParent)
+		{
+			const bool bAlreadyAttached = GetAttachParent() && (GetAttachParent() == NewParent) && (GetAttachSocketName() == AutoAttachSocketName) && GetAttachParent()->GetAttachChildren().Contains(this);
+			if (!bAlreadyAttached)
+			{
+				bDidAutoAttach = bWasAutoAttached;
+				CancelAutoAttachment(true);
+				SavedAutoAttachRelativeLocation = RelativeLocation;
+				SavedAutoAttachRelativeRotation = RelativeRotation;
+				SavedAutoAttachRelativeScale3D = RelativeScale3D;
+				AttachToComponent(NewParent, FAttachmentTransformRules(AutoAttachLocationRule, AutoAttachRotationRule, AutoAttachScaleRule, false), AutoAttachSocketName);
+			}
+
+			bDidAutoAttach = true;
+			//bFlagAsJustAttached = true;
+		}
+		else
+		{
+			CancelAutoAttachment(true);
+		}
+	}
+
 	SystemInstance->Activate(bReset);
 
 	/** We only need to tick the component if we require solo mode. */
@@ -362,9 +442,9 @@ void UNiagaraComponent::Activate(bool bReset /* = false */)
 
 void UNiagaraComponent::Deactivate()
 {
+	SCOPE_CYCLE_COUNTER(STAT_NiagaraComponentDeactivate);
 	Super::Deactivate();
 
-	SetComponentTickEnabled(false);
 	bIsActive = false;
 
 	if (SystemInstance)
@@ -373,6 +453,23 @@ void UNiagaraComponent::Deactivate()
 	}
 }
 
+void UNiagaraComponent::OnSystemComplete()
+{
+	SetComponentTickEnabled(false);
+	bIsActive = false;
+        
+	OnSystemFinished.Broadcast(this);
+
+	if (bAutoDestroy)
+	{
+		DestroyComponent();
+	}
+	else if (bAutoManageAttachment)
+	{
+		CancelAutoAttachment(/*bDetachFromParent=*/ true);
+	}
+}
+
 void UNiagaraComponent::DestroyInstance()
 {
 	SystemInstance = nullptr;
@@ -383,6 +480,41 @@ void UNiagaraComponent::DestroyInstance()
 
 void UNiagaraComponent::OnRegister()
 {
+	if (bAutoManageAttachment && !IsActive())
+	{
+		// Detach from current parent, we are supposed to wait for activation.
+		if (GetAttachParent())
+		{
+			// If no auto attach parent override, use the current parent when we activate
+			if (!AutoAttachParent.IsValid())
+			{
+				AutoAttachParent = GetAttachParent();
+			}
+			// If no auto attach socket override, use current socket when we activate
+			if (AutoAttachSocketName == NAME_None)
+			{
+				AutoAttachSocketName = GetAttachSocketName();
+			}
+
+			// Prevent attachment before Super::OnRegister() tries to attach us, since we only attach when activated.
+			if (GetAttachParent()->GetAttachChildren().Contains(this))
+			{
+				// Only detach if we are not about to auto attach to the same target, that would be wasteful.
+				if (!bAutoActivate || (AutoAttachLocationRule != EAttachmentRule::KeepRelative && AutoAttachRotationRule != EAttachmentRule::KeepRelative && AutoAttachScaleRule != EAttachmentRule::KeepRelative) || (AutoAttachSocketName != GetAttachSocketName()) || (AutoAttachParent != GetAttachParent()))
+				{
+					DetachFromComponent(FDetachmentTransformRules(EDetachmentRule::KeepRelative, /*bCallModify=*/ false));
+				}
+			}
+			else
+			{
+				SetupAttachment(nullptr, NAME_None);
+			}
+		}
+
+		SavedAutoAttachRelativeLocation = RelativeLocation;
+		SavedAutoAttachRelativeRotation = RelativeRotation;
+		SavedAutoAttachRelativeScale3D = RelativeScale3D;
+	}
 	Super::OnRegister();
 }
 
@@ -407,18 +539,19 @@ void UNiagaraComponent::BeginDestroy()
 	Super::BeginDestroy();
 }
 
-class FNiagaraSystemSimulation* UNiagaraComponent::GetSystemSimulation()
+TSharedPtr<FNiagaraSystemSimulation, ESPMode::ThreadSafe> UNiagaraComponent::GetSystemSimulation()
 {
 	UWorld* World = GetWorld();
 	if (World)
 	{
 		return FNiagaraWorldManager::Get(World)->GetSystemSimulation(Asset);
 	}
-	return nullptr;
+	return TSharedRef<FNiagaraSystemSimulation, ESPMode::ThreadSafe>();
 }
 
 void UNiagaraComponent::SendRenderDynamicData_Concurrent()
 {
+	SCOPE_CYCLE_COUNTER(STAT_NiagaraComponentSendRenderData);
 	if (SystemInstance.IsValid() && SceneProxy)
 	{
 		SystemInstance->GetSystemBounds().Init();
@@ -433,7 +566,11 @@ void UNiagaraComponent::SendRenderDynamicData_Concurrent()
 				NiagaraRenderer* Renderer = Emitter->GetEmitterRenderer(EmitterIdx);
 				if (Renderer)
 				{
-					if (Emitter->IsEnabled())
+					bool bRendererEditorEnabled = true;
+#if WITH_EDITORONLY_DATA
+					bRendererEditorEnabled = !SystemInstance->GetIsolateEnabled() || Emitter->GetEmitterHandle().IsIsolated();
+#endif
+					if (bRendererEditorEnabled && !Emitter->IsComplete())
 					{
 						FNiagaraDynamicDataBase* DynamicData = Renderer->GenerateVertexData(NiagaraProxy, Emitter->GetData(), Emitter->GetEmitterHandle().GetInstance()->SimTarget);
 
@@ -476,10 +613,13 @@ FBoxSphereBounds UNiagaraComponent::CalcBounds(const FTransform& LocalToWorld) c
 		for (int32 i = 0; i < SystemInstance->GetEmitters().Num(); i++)
 		{
 			FNiagaraEmitterInstance &Sim = *(SystemInstance->GetEmitters()[i]);
-			SystemInstance->GetSystemBounds() += Sim.GetCachedBounds();
+			SystemInstance->GetSystemBounds() += Sim.GetBounds();
 		}
-		FBoxSphereBounds BSBounds(SystemInstance->GetSystemBounds());
+		FBox BoundingBox = SystemInstance->GetSystemBounds();
+		const FVector ExpandAmount = FVector(0.0f, 0.0f, 0.0f);// BoundingBox.GetExtent() * 0.1f;
+		BoundingBox = FBox(BoundingBox.Min - ExpandAmount, BoundingBox.Max + ExpandAmount);
 
+		FBoxSphereBounds BSBounds(BoundingBox);
 		return BSBounds;
 	}
 	return FBoxSphereBounds(SimBounds);
@@ -503,9 +643,9 @@ void UNiagaraComponent::GetUsedMaterials(TArray<UMaterialInterface*>& OutMateria
 	{
 		if (UNiagaraEmitter* Props = Sim->GetEmitterHandle().GetInstance())
 		{	
-			for (int32 i = 0; i < Props->RendererProperties.Num(); i++)
+			for (int32 i = 0; i < Props->GetRenderers().Num(); i++)
 			{
-				if (UNiagaraRendererProperties* Renderer = Props->RendererProperties[i])
+				if (UNiagaraRendererProperties* Renderer = Props->GetRenderers()[i])
 				{
 					Renderer->GetUsedMaterials(OutMaterials);
 				}
@@ -519,45 +659,46 @@ FNiagaraSystemInstance* UNiagaraComponent::GetSystemInstance() const
 	return SystemInstance.Get();
 }
 
-void UNiagaraComponent::OnSystemDisabled()
+void UNiagaraComponent::SetNiagaraVariableLinearColor(const FString& InVariableName, const FLinearColor& InValue)
 {
-	//TODO: Add KillOnComplete option to asset to kill the component.
+	FName VarName = FName(*InVariableName);
 
+	OverrideParameters.SetParameterValue(InValue, FNiagaraVariable(FNiagaraTypeDefinition::GetColorDef(), VarName), true);
 }
 
 void UNiagaraComponent::SetNiagaraVariableVec4(const FString& InVariableName, const FVector4& InValue)
 {
 	FName VarName = FName(*InVariableName);
 
-	SystemInstance->GetParameters().SetParameterValue(InValue, FNiagaraVariable(FNiagaraTypeDefinition::GetVec4Def(), VarName));
+	OverrideParameters.SetParameterValue(InValue, FNiagaraVariable(FNiagaraTypeDefinition::GetVec4Def(), VarName), true);
 }
 
 void UNiagaraComponent::SetNiagaraVariableVec3(const FString& InVariableName, FVector InValue)
 {
 	FName VarName = FName(*InVariableName);
 	
-	SystemInstance->GetParameters().SetParameterValue(InValue, FNiagaraVariable(FNiagaraTypeDefinition::GetVec3Def(), VarName));
+	OverrideParameters.SetParameterValue(InValue, FNiagaraVariable(FNiagaraTypeDefinition::GetVec3Def(), VarName), true);
 }
 
 void UNiagaraComponent::SetNiagaraVariableVec2(const FString& InVariableName, FVector2D InValue)
 {
 	FName VarName = FName(*InVariableName);
 
-	SystemInstance->GetParameters().SetParameterValue(InValue, FNiagaraVariable(FNiagaraTypeDefinition::GetVec2Def(),VarName));
+	OverrideParameters.SetParameterValue(InValue, FNiagaraVariable(FNiagaraTypeDefinition::GetVec2Def(),VarName), true);
 }
 
 void UNiagaraComponent::SetNiagaraVariableFloat(const FString& InVariableName, float InValue)
 {
 	FName VarName = FName(*InVariableName);
 
-	SystemInstance->GetParameters().SetParameterValue(InValue, FNiagaraVariable(FNiagaraTypeDefinition::GetFloatDef(), VarName));
+	OverrideParameters.SetParameterValue(InValue, FNiagaraVariable(FNiagaraTypeDefinition::GetFloatDef(), VarName), true);
 }
 
 void UNiagaraComponent::SetNiagaraVariableBool(const FString& InVariableName, bool InValue)
 {
 	FName VarName = FName(*InVariableName);
 
-	SystemInstance->GetParameters().SetParameterValue(InValue ? FNiagaraBool::True : FNiagaraBool::False, FNiagaraVariable(FNiagaraTypeDefinition::GetBoolDef(), VarName));
+	OverrideParameters.SetParameterValue(InValue ? FNiagaraBool::True : FNiagaraBool::False, FNiagaraVariable(FNiagaraTypeDefinition::GetBoolDef(), VarName), true);
 }
 
 TArray<FVector> UNiagaraComponent::GetNiagaraParticlePositions_DebugOnly(const FString& InEmitterName)
@@ -648,7 +789,7 @@ void UNiagaraComponent::PostEditChangeProperty(FPropertyChangedEvent& PropertyCh
 	}
 
 	if (PropertyName == GET_MEMBER_NAME_CHECKED(UNiagaraComponent, Asset) ||
-		PropertyName == GET_MEMBER_NAME_CHECKED(UNiagaraComponent, InitialParameters))
+		PropertyName == GET_MEMBER_NAME_CHECKED(UNiagaraComponent, OverrideParameters))
 	{
 		SynchronizeWithSourceSystem();
 	}
@@ -671,17 +812,17 @@ bool UNiagaraComponent::SynchronizeWithSourceSystem()
 	bool bEditsMade = false;
 	for (FNiagaraVariable& Param : SourceVars)
 	{
-		bEditsMade |= InitialParameters.AddParameter(Param, true);
+		bEditsMade |= OverrideParameters.AddParameter(Param, true);
 	}
 
 	TArray<FNiagaraVariable> ExistingVars;
-	InitialParameters.GetParameters(ExistingVars);
+	OverrideParameters.GetParameters(ExistingVars);
 
 	for (FNiagaraVariable ExistingVar : ExistingVars)
 	{
 		if (!SourceVars.Contains(ExistingVar))
 		{
-			InitialParameters.RemoveParameter(ExistingVar);
+			OverrideParameters.RemoveParameter(ExistingVar);
 			EditorOverridesValue.Remove(ExistingVar.GetName());
 		}
 	}
@@ -692,11 +833,11 @@ bool UNiagaraComponent::SynchronizeWithSourceSystem()
 
 		if (!IsParameterValueOverriddenLocally(ExistingVar.GetName()))
 		{
-			Asset->GetExposedParameters().CopyParameterData(InitialParameters, ExistingVar);
+			Asset->GetExposedParameters().CopyParameterData(OverrideParameters, ExistingVar);
 		}
 	}
 
-	InitialParameters.Rebind();
+	OverrideParameters.Rebind();
 
 	//////////////////////////////////////////////////////////////////////////
 	//OLD IMPL
@@ -859,9 +1000,9 @@ bool UNiagaraComponent::IsParameterValueOverriddenLocally(const FName& InParamNa
 	return false;
 }
 
-void UNiagaraComponent::SetParameterValueOverriddenLocally(const FName& InParamName, bool bInOverriden)
+void UNiagaraComponent::SetParameterValueOverriddenLocally(const FNiagaraVariable& InParam, bool bInOverriden)
 {
-	bool* FoundVar = EditorOverridesValue.Find(InParamName);
+	bool* FoundVar = EditorOverridesValue.Find(InParam.GetName());
 
 	if (FoundVar != nullptr && bInOverriden) 
 	{
@@ -869,36 +1010,12 @@ void UNiagaraComponent::SetParameterValueOverriddenLocally(const FName& InParamN
 	}
 	else if (FoundVar == nullptr && bInOverriden)			
 	{
-		EditorOverridesValue.Add(InParamName, true);
+		EditorOverridesValue.Add(InParam.GetName(), true);
 	}
 	else
 	{
-		EditorOverridesValue.Remove(InParamName);
-
-		UNiagaraScript* ScriptSpawn = Asset->GetSystemSpawnScript(true);
-		UNiagaraScript* ScriptUpdate = Asset->GetSystemUpdateScript(true);
-		if (ScriptSpawn != nullptr)
-		{
-			for (FNiagaraVariable Var : ScriptSpawn->Parameters.Parameters)
-			{
-				if (Var.GetName() == InParamName)
-				{
-					InitialParameters.SetParameter(Var);
-					return;
-				}
-			}			
-		}
-		if (ScriptUpdate != nullptr)
-		{
-			for (FNiagaraVariable Var : ScriptUpdate->Parameters.Parameters)
-			{
-				if (Var.GetName() == InParamName)
-				{
-					InitialParameters.SetParameter(Var);
-					return;
-				}
-			}
-		}
+		EditorOverridesValue.Remove(InParam.GetName());
+		Asset->GetExposedParameters().CopyParameterData(OverrideParameters, InParam);
 	}
 }
 
@@ -912,11 +1029,41 @@ void UNiagaraComponent::SetAsset(UNiagaraSystem* InAsset)
 	Asset = InAsset;
 
 #if WITH_EDITOR
-	SynchronizeWithSourceSystem();
+	if (GetWorld() != nullptr && GetWorld()->WorldType == EWorldType::Editor)
+	{
+		SynchronizeWithSourceSystem();
+	}
 #endif
 
 	ReinitializeSystem();
 }
 
+void UNiagaraComponent::SetAutoAttachmentParameters(USceneComponent* Parent, FName SocketName, EAttachmentRule LocationRule, EAttachmentRule RotationRule, EAttachmentRule ScaleRule)
+{
+	AutoAttachParent = Parent;
+	AutoAttachSocketName = SocketName;
+	AutoAttachLocationRule = LocationRule;
+	AutoAttachRotationRule = RotationRule;
+	AutoAttachScaleRule = ScaleRule;
+}
 
 
+void UNiagaraComponent::CancelAutoAttachment(bool bDetachFromParent)
+{
+	if (bAutoManageAttachment)
+	{
+		if (bDidAutoAttach)
+		{
+			// Restore relative transform from before attachment. Actual transform will be updated as part of DetachFromParent().
+			RelativeLocation = SavedAutoAttachRelativeLocation;
+			RelativeRotation = SavedAutoAttachRelativeRotation;
+			RelativeScale3D = SavedAutoAttachRelativeScale3D;
+			bDidAutoAttach = false;
+		}
+
+		if (bDetachFromParent)
+		{
+			DetachFromComponent(FDetachmentTransformRules::KeepRelativeTransform);
+		}
+	}
+}
diff --git a/Engine/Plugins/FX/Niagara/Source/Niagara/Private/NiagaraConstants.cpp b/Engine/Plugins/FX/Niagara/Source/Niagara/Private/NiagaraConstants.cpp
--- a/Engine/Plugins/FX/Niagara/Source/Niagara/Private/NiagaraConstants.cpp
+++ b/Engine/Plugins/FX/Niagara/Source/Niagara/Private/NiagaraConstants.cpp
@@ -14,19 +14,28 @@ const TArray<FNiagaraVariable>& FNiagaraConstants::GetEngineConstants()
 		SystemParameters.Add(SYS_PARAM_ENGINE_DELTA_TIME);
 		SystemParameters.Add(SYS_PARAM_ENGINE_INV_DELTA_TIME);
 		SystemParameters.Add(SYS_PARAM_ENGINE_POSITION);
+		SystemParameters.Add(SYS_PARAM_ENGINE_SCALE);
 		SystemParameters.Add(SYS_PARAM_ENGINE_VELOCITY);
 		SystemParameters.Add(SYS_PARAM_ENGINE_X_AXIS);
 		SystemParameters.Add(SYS_PARAM_ENGINE_Y_AXIS);
 		SystemParameters.Add(SYS_PARAM_ENGINE_Z_AXIS);
+
 		SystemParameters.Add(SYS_PARAM_ENGINE_LOCAL_TO_WORLD);
 		SystemParameters.Add(SYS_PARAM_ENGINE_WORLD_TO_LOCAL);
 		SystemParameters.Add(SYS_PARAM_ENGINE_LOCAL_TO_WORLD_TRANSPOSED);
 		SystemParameters.Add(SYS_PARAM_ENGINE_WORLD_TO_LOCAL_TRANSPOSED);
+		SystemParameters.Add(SYS_PARAM_ENGINE_LOCAL_TO_WORLD_NO_SCALE);
+		SystemParameters.Add(SYS_PARAM_ENGINE_WORLD_TO_LOCAL_NO_SCALE);
+
 		SystemParameters.Add(SYS_PARAM_ENGINE_MIN_DIST_TO_CAMERA);
+		SystemParameters.Add(SYS_PARAM_ENGINE_TIME_SINCE_RENDERED);
+
 		SystemParameters.Add(SYS_PARAM_ENGINE_EXEC_COUNT);
 		SystemParameters.Add(SYS_PARAM_ENGINE_EMITTER_NUM_PARTICLES);
 		SystemParameters.Add(SYS_PARAM_ENGINE_SYSTEM_NUM_EMITTERS_ALIVE);
 		SystemParameters.Add(SYS_PARAM_ENGINE_SYSTEM_NUM_EMITTERS);
+		SystemParameters.Add(SYS_PARAM_EMITTER_AGE);
+		SystemParameters.Add(SYS_PARAM_EMITTER_LOCALSPACE);
 	}
 	return SystemParameters;
 }
@@ -50,19 +59,24 @@ FNiagaraVariable FNiagaraConstants::UpdateEngineConstant(const FNiagaraVariable&
 			UpdatedSystemParameters.Add(FName(TEXT("System X Axis")), SYS_PARAM_ENGINE_X_AXIS);
 			UpdatedSystemParameters.Add(FName(TEXT("System Y Axis")), SYS_PARAM_ENGINE_Y_AXIS);
 			UpdatedSystemParameters.Add(FName(TEXT("System Z Axis")), SYS_PARAM_ENGINE_Z_AXIS);
+
 			UpdatedSystemParameters.Add(FName(TEXT("System Local To World")), SYS_PARAM_ENGINE_LOCAL_TO_WORLD);
 			UpdatedSystemParameters.Add(FName(TEXT("System World To Local")), SYS_PARAM_ENGINE_WORLD_TO_LOCAL);
 			UpdatedSystemParameters.Add(FName(TEXT("System Local To World Transposed")), SYS_PARAM_ENGINE_LOCAL_TO_WORLD_TRANSPOSED);
 			UpdatedSystemParameters.Add(FName(TEXT("System World To Local Transposed")), SYS_PARAM_ENGINE_WORLD_TO_LOCAL_TRANSPOSED);
+			UpdatedSystemParameters.Add(FName(TEXT("System Local To World No Scale")), SYS_PARAM_ENGINE_LOCAL_TO_WORLD_NO_SCALE);
+			UpdatedSystemParameters.Add(FName(TEXT("System World To Local No Scale")), SYS_PARAM_ENGINE_WORLD_TO_LOCAL_NO_SCALE);
 
 			UpdatedSystemParameters.Add(FName(TEXT("Emitter Execution Count")), SYS_PARAM_ENGINE_EXEC_COUNT);
 			UpdatedSystemParameters.Add(FName(TEXT("Emitter Age")), SYS_PARAM_EMITTER_AGE);
+			UpdatedSystemParameters.Add(FName(TEXT("Emitter Local Space")), SYS_PARAM_EMITTER_LOCALSPACE);
 			UpdatedSystemParameters.Add(FName(TEXT("Emitter Spawn Rate")), SYS_PARAM_EMITTER_SPAWNRATE);
 			UpdatedSystemParameters.Add(FName(TEXT("Emitter Spawn Interval")), SYS_PARAM_EMITTER_SPAWN_INTERVAL);
 			UpdatedSystemParameters.Add(FName(TEXT("Emitter Interp Spawn Start Dt")), SYS_PARAM_EMITTER_INTERP_SPAWN_START_DT);
 
 			UpdatedSystemParameters.Add(FName(TEXT("Delta Time")), SYS_PARAM_ENGINE_DELTA_TIME);
 			UpdatedSystemParameters.Add(FName(TEXT("Emitter Age")), SYS_PARAM_EMITTER_AGE);
+			UpdatedSystemParameters.Add(FName(TEXT("Emitter Local Space")), SYS_PARAM_EMITTER_LOCALSPACE);
 			UpdatedSystemParameters.Add(FName(TEXT("Effect Position")), SYS_PARAM_ENGINE_POSITION);
 			UpdatedSystemParameters.Add(FName(TEXT("Effect Velocity")), SYS_PARAM_ENGINE_VELOCITY);
 			UpdatedSystemParameters.Add(FName(TEXT("Effect X Axis")), SYS_PARAM_ENGINE_X_AXIS);
@@ -109,14 +123,20 @@ FText FNiagaraConstants::GetEngineConstantDescription(const FNiagaraVariable& In
 		AttrStrMap.Add(SYS_PARAM_ENGINE_DELTA_TIME, LOCTEXT("EngineDeltaTimeDesc", "Time in seconds since the last tick."));
 		AttrStrMap.Add(SYS_PARAM_ENGINE_INV_DELTA_TIME, LOCTEXT("EngineInvDeltaTimeDesc", "One over Engine.DeltaTime"));
 		AttrStrMap.Add(SYS_PARAM_ENGINE_POSITION, LOCTEXT("EnginePositionDesc", "The owning component's position in world space."));
+		AttrStrMap.Add(SYS_PARAM_ENGINE_SCALE, LOCTEXT("EnginePositionDesc", "The owning component's scale in world space."));
 		AttrStrMap.Add(SYS_PARAM_ENGINE_VELOCITY, LOCTEXT("EngineVelocityDesc", "The owning component's velocity in world space."));
 		AttrStrMap.Add(SYS_PARAM_ENGINE_X_AXIS, LOCTEXT("XAxisDesc", "The X-axis of the owning component."));
 		AttrStrMap.Add(SYS_PARAM_ENGINE_Y_AXIS, LOCTEXT("YAxisDesc", "The Y-axis of the owning component."));
 		AttrStrMap.Add(SYS_PARAM_ENGINE_Z_AXIS, LOCTEXT("ZAxisDesc", "The Z-axis of the owning component."));
-		AttrStrMap.Add(SYS_PARAM_ENGINE_LOCAL_TO_WORLD, LOCTEXT("LocalToWorldDesc", "Conversion of the owning component's local to world space."));
-		AttrStrMap.Add(SYS_PARAM_ENGINE_WORLD_TO_LOCAL, LOCTEXT("WorldToLocalDesc", "Conversion of the owning component's world to local space."));
-		AttrStrMap.Add(SYS_PARAM_ENGINE_LOCAL_TO_WORLD_TRANSPOSED, LOCTEXT("LocalToWorldTransposeDesc", "Conversion of the owning component's local to world space transposed."));
-		AttrStrMap.Add(SYS_PARAM_ENGINE_WORLD_TO_LOCAL_TRANSPOSED, LOCTEXT("WorldToLocalTransposeDesc", "Conversion of the owning component's world to local space transposed."));
+
+		AttrStrMap.Add(SYS_PARAM_ENGINE_LOCAL_TO_WORLD, LOCTEXT("LocalToWorldDesc", "Owning component's local space to world space transform matrix."));
+		AttrStrMap.Add(SYS_PARAM_ENGINE_WORLD_TO_LOCAL, LOCTEXT("WorldToLocalDesc", "Owning component's world space to local space transform matrix."));
+		AttrStrMap.Add(SYS_PARAM_ENGINE_LOCAL_TO_WORLD_TRANSPOSED, LOCTEXT("LocalToWorldTransposeDesc", "Owning component's local space to world space transform matrix transposed."));
+		AttrStrMap.Add(SYS_PARAM_ENGINE_WORLD_TO_LOCAL_TRANSPOSED, LOCTEXT("WorldToLocalTransposeDesc", "Owning component's world space to local space transform matrix transposed."));
+		AttrStrMap.Add(SYS_PARAM_ENGINE_LOCAL_TO_WORLD_NO_SCALE, LOCTEXT("LocalToWorldNoScaleDesc", "Owning component's local space to world space transform matrix with scaling removed."));
+		AttrStrMap.Add(SYS_PARAM_ENGINE_WORLD_TO_LOCAL_NO_SCALE, LOCTEXT("WorldToLocalNoScaleDesc", "Owning component's world space to local space transform matrix with scaling removed."));
+
+		AttrStrMap.Add(SYS_PARAM_ENGINE_TIME_SINCE_RENDERED, LOCTEXT("TimeSinceRendered", "The time in seconds that have passed since this system was last rendered."));
 		AttrStrMap.Add(SYS_PARAM_ENGINE_MIN_DIST_TO_CAMERA, LOCTEXT("MinDistanceToCamera", "The distance from the owner component to the nearest local player viewpoint."));
 		AttrStrMap.Add(SYS_PARAM_ENGINE_EXEC_COUNT, LOCTEXT("ExecCountDesc", "The index of this particle in the read buffer."));
 		AttrStrMap.Add(SYS_PARAM_ENGINE_EMITTER_NUM_PARTICLES, LOCTEXT("EmitterNumParticles", "The number of particles for this emitter at the beginning of simulation. Should only be used in Emitter scripts."));
@@ -148,7 +168,16 @@ const TArray<FNiagaraVariable>& FNiagaraConstants::GetCommonParticleAttributes()
 		Attributes.Add(SYS_PARAM_PARTICLES_SUB_IMAGE_INDEX);
 		Attributes.Add(SYS_PARAM_PARTICLES_DYNAMIC_MATERIAL_PARAM);
 		Attributes.Add(SYS_PARAM_PARTICLES_SCALE);
-		Attributes.Add(SYS_PARAM_PARTICLES_ROTATION);
+		Attributes.Add(SYS_PARAM_PARTICLES_LIFETIME);
+		Attributes.Add(SYS_PARAM_PARTICLES_MESH_ORIENTATION);
+		Attributes.Add(SYS_PARAM_PARTICLES_CAMERA_OFFSET);
+		Attributes.Add(SYS_PARAM_PARTICLES_UV_SCALE);
+		Attributes.Add(SYS_PARAM_PARTICLES_MATERIAL_RANDOM);
+		Attributes.Add(SYS_PARAM_PARTICLES_RIBBONID);
+		Attributes.Add(SYS_PARAM_PARTICLES_RIBBONWIDTH);
+		Attributes.Add(SYS_PARAM_PARTICLES_RIBBONTWIST);
+		Attributes.Add(SYS_PARAM_PARTICLES_RIBBONFACING);
+
 
 		Attributes.Add(SYS_PARAM_INSTANCE_ALIVE);
 	}
@@ -171,7 +200,15 @@ FString FNiagaraConstants::GetAttributeDefaultValue(const FNiagaraVariable& InAt
 		AttrStrMap.Add(SYS_PARAM_PARTICLES_SUB_IMAGE_INDEX, TEXT("0.0"));
 		AttrStrMap.Add(SYS_PARAM_PARTICLES_DYNAMIC_MATERIAL_PARAM, TEXT("1.0,1.0,1.0,1.0"));
 		AttrStrMap.Add(SYS_PARAM_PARTICLES_SCALE, TEXT("1.0,1.0,1.0"));
-		AttrStrMap.Add(SYS_PARAM_PARTICLES_ROTATION, TEXT("0.0,0.0,0.0"));
+		AttrStrMap.Add(SYS_PARAM_PARTICLES_LIFETIME, TEXT("5.0"));
+		AttrStrMap.Add(SYS_PARAM_PARTICLES_MESH_ORIENTATION, TEXT("1.0,0.0,0.0,0.0"));
+		AttrStrMap.Add(SYS_PARAM_PARTICLES_CAMERA_OFFSET, TEXT("0.0"));
+		AttrStrMap.Add(SYS_PARAM_PARTICLES_UV_SCALE, TEXT("X=1.0 Y=1.0"));
+		AttrStrMap.Add(SYS_PARAM_PARTICLES_MATERIAL_RANDOM, TEXT("0.0"));
+		AttrStrMap.Add(SYS_PARAM_PARTICLES_RIBBONID, TEXT("0"));
+		AttrStrMap.Add(SYS_PARAM_PARTICLES_RIBBONWIDTH, TEXT("1.0"));	
+		AttrStrMap.Add(SYS_PARAM_PARTICLES_RIBBONTWIST, TEXT("0.0"));
+		AttrStrMap.Add(SYS_PARAM_PARTICLES_RIBBONFACING, TEXT("0.0, 0.0, 1.0"));
 	}
 
 	FString* FoundStr = AttrStrMap.Find(InAttribute);
@@ -190,18 +227,25 @@ FText FNiagaraConstants::GetAttributeDescription(const FNiagaraVariable& InAttri
 		AttrStrMap.Add(SYS_PARAM_PARTICLES_POSITION, LOCTEXT("PositionDesc", "The position of the particle."));
 		AttrStrMap.Add(SYS_PARAM_PARTICLES_VELOCITY, LOCTEXT("VelocityDesc", "The velocity in cm/s of the particle."));
 		AttrStrMap.Add(SYS_PARAM_PARTICLES_COLOR, LOCTEXT("ColorDesc", "The color of the particle."));
-		AttrStrMap.Add(SYS_PARAM_PARTICLES_SPRITE_ROTATION, LOCTEXT("SpriteRotDesc", "The screen aligned roll of the particle."));
+		AttrStrMap.Add(SYS_PARAM_PARTICLES_SPRITE_ROTATION, LOCTEXT("SpriteRotDesc", "The screen aligned roll of the particle in degrees."));
 		AttrStrMap.Add(SYS_PARAM_PARTICLES_NORMALIZED_AGE, LOCTEXT("NormalizedAgeDesc", "The age in seconds divided by lifetime in seconds. Useful for animation as the value is between 0 and 1."));
 		AttrStrMap.Add(SYS_PARAM_PARTICLES_SPRITE_SIZE, LOCTEXT("SpriteSizeDesc", "The size of the sprite quad."));
-		AttrStrMap.Add(SYS_PARAM_PARTICLES_SPRITE_FACING, LOCTEXT("FacingDesc", "Makes the surface of the sprite face towards a custom vector."));
-		AttrStrMap.Add(SYS_PARAM_PARTICLES_SPRITE_ALIGNMENT, LOCTEXT("AlignmentDesc", "Imagine the texture having an arrow to the right, this attribute makes the arrow point towards the alignment axis."));
-		AttrStrMap.Add(SYS_PARAM_PARTICLES_SUB_IMAGE_INDEX, LOCTEXT("SubImageIndexDesc", "The position in the UV lookup table divided by total number of UV slots."));
+		AttrStrMap.Add(SYS_PARAM_PARTICLES_SPRITE_FACING, LOCTEXT("FacingDesc", "Makes the surface of the sprite face towards a custom vector. Must be used with the SpriteRenderer's CustomFacingVector FacingMode and CustomFacingVectorMask options."));
+		AttrStrMap.Add(SYS_PARAM_PARTICLES_SPRITE_ALIGNMENT, LOCTEXT("AlignmentDesc", "Imagine the texture having an arrow pointing up, this attribute makes the arrow point towards the alignment axis. Must be used with the SpriteRenderer's CustomAlignment Alignment option."));
+		AttrStrMap.Add(SYS_PARAM_PARTICLES_SUB_IMAGE_INDEX, LOCTEXT("SubImageIndexDesc", "A value from 0 to the number of entries in the table of SubUV images."));
 		AttrStrMap.Add(SYS_PARAM_PARTICLES_DYNAMIC_MATERIAL_PARAM, LOCTEXT("DynamicMaterialParameterDesc", "The 4-float vector used to send custom data to renderer."));
 		AttrStrMap.Add(SYS_PARAM_PARTICLES_SCALE, LOCTEXT("ScaleParamDesc", "The XYZ scale of the non-sprite based particle."));
-		AttrStrMap.Add(SYS_PARAM_PARTICLES_ROTATION, LOCTEXT("RotationParamDesc", "The Vector3 euler angle rotation"));
+		AttrStrMap.Add(SYS_PARAM_PARTICLES_LIFETIME, LOCTEXT("LifetimeParamDesc", "The lifetime of a particle in seconds."));
+		AttrStrMap.Add(SYS_PARAM_PARTICLES_MESH_ORIENTATION, LOCTEXT("MeshOrientParamDesc", "The axis-angle rotation to be applied to the mesh particle."));
+		AttrStrMap.Add(SYS_PARAM_PARTICLES_CAMERA_OFFSET, LOCTEXT("CamOffsetParamDesc", "Used to offset position in the direction of the camera. The value is multiplied by the direction vector from the camera to the particle."));
+		AttrStrMap.Add(SYS_PARAM_PARTICLES_UV_SCALE, LOCTEXT("UVScalerParamDesc", "Used to multiply the generated UVs for Sprite renderers."));
+		AttrStrMap.Add(SYS_PARAM_PARTICLES_MATERIAL_RANDOM, LOCTEXT("MaterialRandomParamDesc", "Used to drive the Particle Random node in the Material Editor. Without this set, any Particle Randoms will get 0.0."));
 		AttrStrMap.Add(SYS_PARAM_INSTANCE_ALIVE, LOCTEXT("AliveParamDesc", "Used to determine whether or not this particle instance is still valid or if it can be deleted."));
+		AttrStrMap.Add(SYS_PARAM_PARTICLES_RIBBONID, LOCTEXT("RibbonIDDesc", "Sets the ribbon id for a particle. Particles with the same ribbon id will be connected into a ribbon."));
+		AttrStrMap.Add(SYS_PARAM_PARTICLES_RIBBONWIDTH, LOCTEXT("RibbonWidthDesc", "Sets the ribbon width for a particle, in UE4 units."));
+		AttrStrMap.Add(SYS_PARAM_PARTICLES_RIBBONTWIST, LOCTEXT("RibbonTwistDesc", "Sets the ribbon twist for a particle, in degrees."));
+		AttrStrMap.Add(SYS_PARAM_PARTICLES_RIBBONFACING, LOCTEXT("RibbonFacingDesc", "Sets the facing vector of the ribbon at the particle position."));
 	}
-
 	FText* FoundStr = AttrStrMap.Find(InAttribute);
 	if (FoundStr != nullptr && !FoundStr->IsEmpty())
 	{
@@ -210,4 +254,45 @@ FText FNiagaraConstants::GetAttributeDescription(const FNiagaraVariable& InAttri
 	return FText();
 }
 
+bool FNiagaraConstants::IsNiagaraConstant(const FNiagaraVariable& InVar)
+{
+	if (GetConstantMetaData(InVar) != nullptr)
+	{
+		return true;
+	}
+	return false;
+}
+
+const FNiagaraVariableMetaData* FNiagaraConstants::GetConstantMetaData(const FNiagaraVariable& InVar)
+{
+	static TMap<FNiagaraVariable, FNiagaraVariableMetaData> AttrMetaData;
+	if (AttrMetaData.Num() == 0)
+	{
+		// Add the engine default attributes..
+		{
+			TArray<FNiagaraVariable> Attributes = GetCommonParticleAttributes();
+			for (FNiagaraVariable Var : Attributes)
+			{
+				FNiagaraVariableMetaData MetaData;
+				MetaData.Description = GetAttributeDescription(Var);
+				AttrMetaData.Add(Var, MetaData);
+			}
+		}
+
+		// Add the engine constants..
+		{
+			TArray<FNiagaraVariable> EngineConstants = GetEngineConstants();
+			for (FNiagaraVariable Var : EngineConstants)
+			{
+				FNiagaraVariableMetaData MetaData;
+				MetaData.Description = GetEngineConstantDescription(Var);
+				AttrMetaData.Add(Var, MetaData);
+			}
+		}
+
+	}
+	return AttrMetaData.Find(InVar);
+}
+
+
 #undef LOCTEXT_NAMESPACE
\ No newline at end of file
diff --git a/Engine/Plugins/FX/Niagara/Source/Niagara/Private/NiagaraDataInterface.cpp b/Engine/Plugins/FX/Niagara/Source/Niagara/Private/NiagaraDataInterface.cpp
--- a/Engine/Plugins/FX/Niagara/Source/Niagara/Private/NiagaraDataInterface.cpp
+++ b/Engine/Plugins/FX/Niagara/Source/Niagara/Private/NiagaraDataInterface.cpp
@@ -7,18 +7,16 @@
 #include "NiagaraTypes.h"
 
 UNiagaraDataInterface::UNiagaraDataInterface(FObjectInitializer const& ObjectInitializer)
-	: Super(ObjectInitializer)
 {
-
 }
 
 bool UNiagaraDataInterface::CopyTo(UNiagaraDataInterface* Destination) const 
 {
-	if (Destination == nullptr || Destination->GetClass() != GetClass())
-	{
-		return false;
-	}
-	return true;
+	bool result = CopyToInternal(Destination);
+#if WITH_EDITOR
+	Destination->OnChanged().Broadcast();
+#endif
+	return result;
 }
 
 bool UNiagaraDataInterface::Equals(const UNiagaraDataInterface* Other) const
@@ -40,3 +38,57 @@ bool UNiagaraDataInterface::IsDataInterfaceType(const FNiagaraTypeDefinition& Ty
 	return false;
 }
 
+bool UNiagaraDataInterface::CopyToInternal(UNiagaraDataInterface* Destination) const
+{
+	if (Destination == nullptr || Destination->GetClass() != GetClass())
+	{
+		return false;
+	}
+	return true;
+}
+
+//////////////////////////////////////////////////////////////////////////
+
+bool UNiagaraDataInterfaceCurveBase::CopyToInternal(UNiagaraDataInterface* Destination) const
+{
+	if (!Super::CopyToInternal(Destination))
+	{
+		return false;
+	}
+
+	UNiagaraDataInterfaceCurveBase* DestinationTyped = CastChecked<UNiagaraDataInterfaceCurveBase>(Destination);
+	DestinationTyped->bUseLUT = bUseLUT;
+	return true;
+}
+
+
+bool UNiagaraDataInterfaceCurveBase::CompareLUTS(const TArray<float>& OtherLUT) const
+{
+	if (ShaderLUT.Num() == OtherLUT.Num())
+	{
+		for (int32 i = 0; i < ShaderLUT.Num(); i++)
+		{
+			if (false == FMath::IsNearlyEqual(ShaderLUT[i], OtherLUT[i]))
+			{
+				return false;
+			}
+		}
+		return true;
+	}
+	else
+	{
+		return false;
+	}
+}
+
+bool UNiagaraDataInterfaceCurveBase::Equals(const UNiagaraDataInterface* Other) const
+{
+	if (!Super::Equals(Other))
+	{
+		return false;
+	}
+
+	const UNiagaraDataInterfaceCurveBase* OtherTyped = CastChecked<UNiagaraDataInterfaceCurveBase>(Other);
+	bool bEqual = OtherTyped->bUseLUT == bUseLUT;
+	return bEqual;
+}
\ No newline at end of file
