Copying Orion-Staging @ 3748653 (Orion/Dev-General @ 3722124) to //UE4/Main

diff --git a/Engine/Plugins/FX/Niagara/Source/Niagara/Private/NiagaraDataInterfaceCollisionQuery.cpp b/Engine/Plugins/FX/Niagara/Source/Niagara/Private/NiagaraDataInterfaceCollisionQuery.cpp
new file mode 100644
--- /dev/null
+++ b/Engine/Plugins/FX/Niagara/Source/Niagara/Private/NiagaraDataInterfaceCollisionQuery.cpp
@@ -0,0 +1,428 @@
+// Copyright 1998-2016 Epic Games, Inc. All Rights Reserved.
+
+#include "NiagaraDataInterfaceCollisionQuery.h"
+#include "NiagaraTypes.h"
+#include "NiagaraCustomVersion.h"
+
+//////////////////////////////////////////////////////////////////////////
+//Color Curve
+
+
+UNiagaraDataInterfaceCollisionQuery::UNiagaraDataInterfaceCollisionQuery(FObjectInitializer const& ObjectInitializer)
+	: Super(ObjectInitializer)
+{
+}
+
+bool UNiagaraDataInterfaceCollisionQuery::InitPerInstanceData(void* PerInstanceData, FNiagaraSystemInstance* InSystemInstance)
+{
+	CQDIPerInstanceData *PIData = new (PerInstanceData) CQDIPerInstanceData;
+	PIData->SystemInstance = InSystemInstance;
+	if (InSystemInstance)
+	{
+		PIData->CollisionBatch.Init(InSystemInstance->GetIDName(), InSystemInstance->GetComponent()->GetWorld());
+	}
+	return true;
+}
+
+
+void UNiagaraDataInterfaceCollisionQuery::PostInitProperties()
+{
+	Super::PostInitProperties();
+
+	//Can we register data interfaces as regular types and fold them into the FNiagaraVariable framework for UI and function calls etc?
+	if (HasAnyFlags(RF_ClassDefaultObject))
+	{
+		FNiagaraTypeRegistry::Register(FNiagaraTypeDefinition(GetClass()), true, false, false);
+	}
+}
+
+void UNiagaraDataInterfaceCollisionQuery::PostLoad()
+{
+	Super::PostLoad();
+
+	const int32 NiagaraVer = GetLinkerCustomVersion(FNiagaraCustomVersion::GUID);
+}
+
+#if WITH_EDITOR
+
+void UNiagaraDataInterfaceCollisionQuery::PostEditChangeProperty(struct FPropertyChangedEvent& PropertyChangedEvent)
+{
+	Super::PostEditChangeProperty(PropertyChangedEvent);
+}
+
+#endif
+
+
+bool UNiagaraDataInterfaceCollisionQuery::CopyToInternal(UNiagaraDataInterface* Destination) const
+{
+	if (!Super::CopyToInternal(Destination))
+	{
+		return false;
+	}
+	return true;
+}
+
+bool UNiagaraDataInterfaceCollisionQuery::Equals(const UNiagaraDataInterface* Other) const
+{
+	if (!Super::Equals(Other))
+	{
+		return false;
+	}
+	return true;
+}
+
+
+void UNiagaraDataInterfaceCollisionQuery::GetFunctions(TArray<FNiagaraFunctionSignature>& OutFunctions)
+{
+	FNiagaraFunctionSignature Sig3;
+	Sig3.Name = TEXT("PerformCollisionQuery");
+	Sig3.bMemberFunction = true;
+	Sig3.bRequiresContext = false;
+	Sig3.Inputs.Add(FNiagaraVariable(FNiagaraTypeDefinition(GetClass()), TEXT("CollisionQuery")));
+	//Sig3.Inputs.Add(FNiagaraVariable(FNiagaraTypeDefinition::GetBoolDef(), TEXT("PerformQuery")));
+	Sig3.Inputs.Add(FNiagaraVariable(FNiagaraTypeDefinition::GetIntDef(), TEXT("ReturnQueryID")));
+	Sig3.Inputs.Add(FNiagaraVariable(FNiagaraTypeDefinition::GetVec3Def(), TEXT("ParticlePosition")));
+	Sig3.Inputs.Add(FNiagaraVariable(FNiagaraTypeDefinition::GetVec3Def(), TEXT("ParticleVelocity")));
+	Sig3.Inputs.Add(FNiagaraVariable(FNiagaraTypeDefinition::GetFloatDef(), TEXT("DeltaTime")));
+	Sig3.Inputs.Add(FNiagaraVariable(FNiagaraTypeDefinition::GetFloatDef(), TEXT("CollisionSize")));
+	Sig3.Outputs.Add(FNiagaraVariable(FNiagaraTypeDefinition::GetIntDef(), TEXT("QueryID")));
+	Sig3.Outputs.Add(FNiagaraVariable(FNiagaraTypeDefinition::GetBoolDef(), TEXT("CollisionValid")));
+	Sig3.Outputs.Add(FNiagaraVariable(FNiagaraTypeDefinition::GetVec3Def(), TEXT("CollisionPos")));
+	Sig3.Outputs.Add(FNiagaraVariable(FNiagaraTypeDefinition::GetVec3Def(), TEXT("CollisionVelocity")));
+	Sig3.Outputs.Add(FNiagaraVariable(FNiagaraTypeDefinition::GetVec3Def(), TEXT("CollisionNormal")));
+	Sig3.Outputs.Add(FNiagaraVariable(FNiagaraTypeDefinition::GetFloatDef(), TEXT("Friction")));
+	Sig3.Outputs.Add(FNiagaraVariable(FNiagaraTypeDefinition::GetFloatDef(), TEXT("Restitution")));
+	OutFunctions.Add(Sig3);
+
+
+	FNiagaraFunctionSignature Sig;
+	Sig.Name = TEXT("SubmitQuery");
+	Sig.bMemberFunction = true;
+	Sig.bRequiresContext = false;
+	Sig.Inputs.Add(FNiagaraVariable(FNiagaraTypeDefinition(GetClass()), TEXT("CollisionQuery")));
+	Sig.Inputs.Add(FNiagaraVariable(FNiagaraTypeDefinition::GetVec3Def(), TEXT("ParticlePosition")));
+	Sig.Inputs.Add(FNiagaraVariable(FNiagaraTypeDefinition::GetVec3Def(), TEXT("ParticleVelocity")));
+	Sig.Inputs.Add(FNiagaraVariable(FNiagaraTypeDefinition::GetFloatDef(), TEXT("DeltaTime")));
+	Sig.Outputs.Add(FNiagaraVariable(FNiagaraTypeDefinition::GetIntDef(), TEXT("CollisionID")));
+	OutFunctions.Add(Sig);
+
+	FNiagaraFunctionSignature Sig2;
+	Sig2.Name = TEXT("ReadQuery");
+	Sig2.bMemberFunction = true;
+	Sig2.bRequiresContext = false;
+	Sig2.Inputs.Add(FNiagaraVariable(FNiagaraTypeDefinition(GetClass()), TEXT("CollisionQuery")));
+	Sig2.Inputs.Add(FNiagaraVariable(FNiagaraTypeDefinition::GetIntDef(), TEXT("CollisionID")));
+	Sig2.Outputs.Add(FNiagaraVariable(FNiagaraTypeDefinition::GetBoolDef(), TEXT("CollisionValid")));
+	Sig2.Outputs.Add(FNiagaraVariable(FNiagaraTypeDefinition::GetVec3Def(), TEXT("CollisionPos")));
+	Sig2.Outputs.Add(FNiagaraVariable(FNiagaraTypeDefinition::GetVec3Def(), TEXT("CollisionVelocity")));
+	Sig2.Outputs.Add(FNiagaraVariable(FNiagaraTypeDefinition::GetVec3Def(), TEXT("CollisionNormal")));
+	OutFunctions.Add(Sig2);
+}
+
+// build the shader function HLSL; function name is passed in, as it's defined per-DI; that way, configuration could change
+// the HLSL in the spirit of a static switch
+// TODO: need a way to identify each specific function here
+// 
+bool UNiagaraDataInterfaceCollisionQuery::GetFunctionHLSL(const FName& DefinitionFunctionName, FString InstanceFunctionName, TArray<FDIGPUBufferParamDescriptor> &Descriptors, FString &HLSLInterfaceID, FString &OutHLSL)
+{
+	// a little tricky, since we've got two functions for submitting and retrieving a query; we store submitted queries per thread group,
+	// assuming it'll usually be the same thread trying to call ReadQuery for a particular QueryID, that submitted it in the first place.
+	if (InstanceFunctionName == TEXT("PerformQuery"))
+	{
+		FString TextureName = Descriptors[0].BufferParamName;
+		OutHLSL += TEXT("groupshared int CollisionQueryID = 0;\n");
+		OutHLSL += TEXT("\t\n");
+		OutHLSL += TEXT("void ") + InstanceFunctionName + TEXT("(in int InQueryID, in float3 In_ParticlePos, in float3 In_ParticleVel, in float In_Deltatime, out int OutQueryID, out bool OutCollisionValid, out float3 Out_CollisionPos, out float3 Out_CollisionVel, out float3 Out_CollisionNormal) \n{\n");
+		OutHLSL += TEXT("\tInterlockedAdd(CollisionQueryID, 1);\n");
+		// get the screen position
+		OutHLSL += TEXT("float3 CollisionPosition = InPosition + CollisionOffset;\
+							float4 SamplePosition = float4(CollisionPosition + View.PreViewTranslation, 1);\
+							float4 ClipPosition = mul(SamplePosition, View.TranslatedWorldToClip);\
+							float2 ScreenPosition = ClipPosition.xy / ClipPosition.w;\
+							// Don't try to collide if the particle falls outside the view.\
+							if (all(abs(ScreenPosition.xy) <= float2(1, 1)))\
+							{\
+								// Sample the depth buffer to get a world position near the particle.\
+								float2 ScreenUV = ScreenPosition * View.ScreenPositionScaleBias.xy + View.ScreenPositionScaleBias.wz;\
+								float SceneDepth = CalcSceneDepth(ScreenUV);\
+								if (abs(ClipPosition.w - SceneDepth) < CollisionDepthBounds)");
+		OutHLSL += TEXT("\n}\n");
+	}
+
+
+	return true;
+}
+
+// build buffer definition hlsl
+// 1. Choose a buffer name, add the data interface ID (important!)
+// 2. add a DIGPUBufferParamDescriptor to the array argument; that'll be passed on to the FNiagaraShader for binding to a shader param, that can
+// then later be found by name via FindDIBufferParam for setting; 
+// 3. store buffer declaration hlsl in OutHLSL
+// multiple buffers can be defined at once here
+//
+void UNiagaraDataInterfaceCollisionQuery::GetBufferDefinitionHLSL(FString DataInterfaceID, TArray<FDIGPUBufferParamDescriptor> &BufferDescriptors, FString &OutHLSL)
+{
+	FString TextureName = "DepthBuffer" + DataInterfaceID;
+	OutHLSL += TEXT("Texture2D ") + TextureName + TEXT(";\n");
+	BufferDescriptors.Add(FDIGPUBufferParamDescriptor(TextureName, 0));		// add a descriptor for shader parameter binding
+
+	FString CounterName = "DepthBuffer" + DataInterfaceID;
+	OutHLSL += TEXT("Buffer<uint> ") + TextureName + TEXT(";\n");
+	BufferDescriptors.Add(FDIGPUBufferParamDescriptor(TextureName, 0));		// add a descriptor for shader parameter binding
+}
+
+// called after translate, to setup buffers matching the buffer descriptors generated during hlsl translation
+// need to do this because the script used during translate is a clone, including its DIs
+//
+void UNiagaraDataInterfaceCollisionQuery::SetupBuffers(FDIBufferDescriptorStore &BufferDescriptors)
+{
+	GPUBuffers.Empty();
+	/*
+	for (FDIGPUBufferParamDescriptor &Desc : BufferDescriptors.Descriptors)
+	{
+		FNiagaraDataInterfaceBufferData BufferData(*Desc.BufferParamName);	// store off the data for later use
+		GPUBuffers.Add(BufferData);
+	}
+	GPUBufferDirty = true;
+	*/
+}
+
+// return the GPU buffer array (called from NiagaraInstanceBatcher to get the buffers for setting to the shader)
+// we lazily update the buffer with a new LUT here if necessary
+//
+TArray<FNiagaraDataInterfaceBufferData> &UNiagaraDataInterfaceCollisionQuery::GetBufferDataArray()
+{
+	check(IsInRenderingThread());
+	return GPUBuffers;
+}
+
+
+DEFINE_NDI_FUNC_BINDER(UNiagaraDataInterfaceCollisionQuery, SubmitQuery);
+DEFINE_NDI_FUNC_BINDER(UNiagaraDataInterfaceCollisionQuery, ReadQuery);
+DEFINE_NDI_FUNC_BINDER(UNiagaraDataInterfaceCollisionQuery, PerformQuery);
+
+FVMExternalFunction UNiagaraDataInterfaceCollisionQuery::GetVMExternalFunction(const FVMExternalFunctionBindingInfo& BindingInfo, void* InstanceData)
+{
+	CQDIPerInstanceData *InstData = (CQDIPerInstanceData *)InstanceData;
+	if (BindingInfo.Name == TEXT("SubmitQuery") /*&& BindingInfo.GetNumInputs() == 3 && BindingInfo.GetNumOutputs() == 1*/)
+	{
+		return TNDIParamBinder<0, float, TNDIParamBinder<1, float, TNDIParamBinder<2, float, TNDIParamBinder<3, float, TNDIParamBinder<4, float, TNDIParamBinder<5, float, TNDIParamBinder<6, float, NDI_FUNC_BINDER(UNiagaraDataInterfaceCollisionQuery, SubmitQuery)>>>>>>>::Bind(this, BindingInfo, InstData);
+	}
+	else if (BindingInfo.Name == TEXT("ReadQuery") /*&& BindingInfo.GetNumInputs() == 1 && BindingInfo.GetNumOutputs() == 4*/)
+	{
+		return TNDIParamBinder<0, int32, NDI_FUNC_BINDER(UNiagaraDataInterfaceCollisionQuery, ReadQuery)>::Bind(this, BindingInfo, InstData);
+	}
+	else if (BindingInfo.Name == TEXT("PerformCollisionQuery") /*&& BindingInfo.GetNumInputs() == 1 && BindingInfo.GetNumOutputs() == 4*/)
+	{
+		return TNDIParamBinder<0, int, TNDIParamBinder<1, float, TNDIParamBinder<2, float, TNDIParamBinder<3, float, TNDIParamBinder<4, float, TNDIParamBinder<5, float, TNDIParamBinder<6, float, TNDIParamBinder<7, float, TNDIParamBinder<8, float, NDI_FUNC_BINDER(UNiagaraDataInterfaceCollisionQuery, PerformQuery)>>>>>>>>>::Bind(this, BindingInfo, InstData);
+	}
+	else
+	{
+		UE_LOG(LogNiagara, Error, TEXT("Could not find data interface external function. %s\n"),
+			*BindingInfo.Name.ToString());
+		return FVMExternalFunction();
+	}
+}
+
+
+template<typename InQueryIDType, typename PosTypeX, typename PosTypeY, typename PosTypeZ, typename VelTypeX, typename VelTypeY, typename VelTypeZ, typename DTType, typename SizeType>
+void UNiagaraDataInterfaceCollisionQuery::PerformQuery(FVectorVMContext& Context)
+{
+	//PerformType PerformParam(Context);
+	InQueryIDType InIDParam(Context);
+	PosTypeX PosParamX(Context);
+	PosTypeY PosParamY(Context);
+	PosTypeZ PosParamZ(Context);
+	VelTypeX VelParamX(Context);
+	VelTypeY VelParamY(Context);
+	VelTypeZ VelParamZ(Context);
+	DTType DTParam(Context);
+	SizeType SizeParam(Context);
+
+	FUserPtrHandler<CQDIPerInstanceData> InstanceData(Context);
+
+	FRegisterHandler<int32> OutQueryID(Context);
+
+	FRegisterHandler<int32> OutQueryValid(Context);
+	FRegisterHandler<float> OutCollisionPosX(Context);
+	FRegisterHandler<float> OutCollisionPosY(Context);
+	FRegisterHandler<float> OutCollisionPosZ(Context);
+	FRegisterHandler<float> OutCollisionVelX(Context);
+	FRegisterHandler<float> OutCollisionVelY(Context);
+	FRegisterHandler<float> OutCollisionVelZ(Context);
+	FRegisterHandler<float> OutCollisionNormX(Context);
+	FRegisterHandler<float> OutCollisionNormY(Context);
+	FRegisterHandler<float> OutCollisionNormZ(Context);
+	FRegisterHandler<float> Friction(Context);
+	FRegisterHandler<float> Restitution(Context);
+
+	for (int32 i = 0; i < Context.NumInstances; ++i)
+	{
+		// submit a query, get query id in return
+		//if (PerformParam.Get())
+		{
+			FVector Pos(PosParamX.Get(), PosParamY.Get(), PosParamZ.Get());
+			FVector Vel(VelParamX.Get(), VelParamY.Get(), VelParamZ.Get());
+			ensure(!Pos.ContainsNaN());
+			ensure(!Vel.ContainsNaN());
+			float Dt = DTParam.Get();
+			float Size = SizeParam.Get();
+			*OutQueryID.GetDest() = InstanceData->CollisionBatch.SubmitQuery(Pos, Vel, Size, Dt);
+		}
+
+		// try to retrieve a query with the supplied query ID
+		FNiagaraDICollsionQueryResult Res;
+		int32 ID = InIDParam.Get();
+		bool Valid = InstanceData->CollisionBatch.GetQueryResult(ID, Res);
+		if (Valid)
+		{
+			*OutQueryValid.GetDest() = 0xFFFFFFFF; //->SetValue(true);
+			*OutCollisionPosX.GetDest() = Res.CollisionPos.X;
+			*OutCollisionPosY.GetDest() = Res.CollisionPos.Y;
+			*OutCollisionPosZ.GetDest() = Res.CollisionPos.Z;
+			*OutCollisionVelX.GetDest() = Res.CollisionVelocity.X;
+			*OutCollisionVelY.GetDest() = Res.CollisionVelocity.Y;
+			*OutCollisionVelZ.GetDest() = Res.CollisionVelocity.Z;
+			*OutCollisionNormX.GetDest() = Res.CollisionNormal.X;
+			*OutCollisionNormY.GetDest() = Res.CollisionNormal.Y;
+			*OutCollisionNormZ.GetDest() = Res.CollisionNormal.Z;
+			*Friction.GetDest() = Res.Friction;
+			*Restitution.GetDest() = Res.Restitution;
+		}
+		else
+		{
+			*OutQueryValid.GetDest() = 0;// ->SetValue(false);
+		}
+
+		//PerformParam.Advance();
+		InIDParam.Advance();
+		OutQueryValid.Advance();
+		OutCollisionPosX.Advance();
+		OutCollisionPosY.Advance();
+		OutCollisionPosZ.Advance();
+		OutCollisionVelX.Advance();
+		OutCollisionVelY.Advance();
+		OutCollisionVelZ.Advance();
+		OutCollisionNormX.Advance();
+		OutCollisionNormY.Advance();
+		OutCollisionNormZ.Advance();
+		Friction.Advance();
+		Restitution.Advance();
+
+		PosParamX.Advance();
+		PosParamY.Advance();
+		PosParamZ.Advance();
+		VelParamX.Advance();
+		VelParamY.Advance();
+		VelParamZ.Advance();
+		DTParam.Advance();
+		SizeParam.Advance();
+		OutQueryID.Advance();
+	}
+
+}
+
+
+template<typename PosTypeX, typename PosTypeY, typename PosTypeZ, typename VelTypeX, typename VelTypeY, typename VelTypeZ, typename DTType>
+void UNiagaraDataInterfaceCollisionQuery::SubmitQuery(FVectorVMContext& Context)
+{
+	PosTypeX PosParamX(Context);
+	PosTypeY PosParamY(Context);
+	PosTypeZ PosParamZ(Context);
+	VelTypeX VelParamX(Context);
+	VelTypeY VelParamY(Context);
+	VelTypeZ VelParamZ(Context);
+
+	DTType DTParam(Context);
+	FUserPtrHandler<CQDIPerInstanceData> InstanceData(Context);
+	FRegisterHandler<int32> OutQueryID(Context);
+	for (int32 i = 0; i < Context.NumInstances; ++i)
+	{
+		FVector Pos(PosParamX.Get(), PosParamY.Get(), PosParamZ.Get());
+		FVector Vel(VelParamX.Get(), VelParamY.Get(), VelParamZ.Get());
+		ensure(!Pos.ContainsNaN());
+		ensure(!Vel.ContainsNaN());
+		float Dt = DTParam.Get();
+
+		*OutQueryID.GetDest() = InstanceData->CollisionBatch.SubmitQuery(Pos, Vel, 0.0f, Dt);
+
+		PosParamX.Advance();
+		PosParamY.Advance();
+		PosParamZ.Advance();
+		VelParamX.Advance();
+		VelParamY.Advance();
+		VelParamZ.Advance();
+		DTParam.Advance();
+		OutQueryID.Advance();
+	}
+
+}
+
+
+template<typename IDType>
+void UNiagaraDataInterfaceCollisionQuery::ReadQuery(FVectorVMContext& Context)
+{	
+	IDType IDParam(Context);
+	FUserPtrHandler<CQDIPerInstanceData> InstanceData(Context);
+	FRegisterHandler<int32> OutQueryValid(Context);
+	FRegisterHandler<float> OutCollisionPosX(Context);
+	FRegisterHandler<float> OutCollisionPosY(Context);
+	FRegisterHandler<float> OutCollisionPosZ(Context);
+	FRegisterHandler<float> OutCollisionVelX(Context);
+	FRegisterHandler<float> OutCollisionVelY(Context);
+	FRegisterHandler<float> OutCollisionVelZ(Context);
+	FRegisterHandler<float> OutCollisionNormX(Context);
+	FRegisterHandler<float> OutCollisionNormY(Context);
+	FRegisterHandler<float> OutCollisionNormZ(Context);
+
+	for (int32 i = 0; i < Context.NumInstances; ++i)
+	{
+		FNiagaraDICollsionQueryResult Res;
+		int32 ID = IDParam.Get();
+		bool Valid = InstanceData->CollisionBatch.GetQueryResult(ID, Res);
+
+		if (Valid)
+		{
+			*OutQueryValid.GetDest() = 0xFFFFFFFF; //->SetValue(true);
+			*OutCollisionPosX.GetDest() = Res.CollisionPos.X;
+			*OutCollisionPosY.GetDest() = Res.CollisionPos.Y;
+			*OutCollisionPosZ.GetDest() = Res.CollisionPos.Z;
+			*OutCollisionVelX.GetDest() = Res.CollisionVelocity.X;
+			*OutCollisionVelY.GetDest() = Res.CollisionVelocity.Y;
+			*OutCollisionVelZ.GetDest() = Res.CollisionVelocity.Z;
+			*OutCollisionNormX.GetDest() = Res.CollisionNormal.X;
+			*OutCollisionNormY.GetDest() = Res.CollisionNormal.Y;
+			*OutCollisionNormZ.GetDest() = Res.CollisionNormal.Z;
+		}
+		else
+		{
+			*OutQueryValid.GetDest() = 0;// ->SetValue(false);
+		}
+
+		IDParam.Advance();
+		OutQueryValid.Advance();
+		OutCollisionPosX.Advance();
+		OutCollisionPosY.Advance();
+		OutCollisionPosZ.Advance();
+		OutCollisionVelX.Advance();
+		OutCollisionVelY.Advance();
+		OutCollisionVelZ.Advance();
+		OutCollisionNormX.Advance();
+		OutCollisionNormY.Advance();
+		OutCollisionNormZ.Advance();
+	}
+}
+
+
+bool UNiagaraDataInterfaceCollisionQuery::PerInstanceTick(void* PerInstanceData, FNiagaraSystemInstance* InSystemInstance, float DeltaSeconds)
+{
+	return false;
+}
+
+bool UNiagaraDataInterfaceCollisionQuery::PerInstanceTickPostSimulate(void* PerInstanceData, FNiagaraSystemInstance* InSystemInstance, float DeltaSeconds)
+{
+	CQDIPerInstanceData *PIData = static_cast<CQDIPerInstanceData*>(PerInstanceData);
+	PIData->CollisionBatch.Tick(ENiagaraSimTarget::CPUSim);
+	PIData->CollisionBatch.ClearWrite();
+	return false;
+}
diff --git a/Engine/Plugins/FX/Niagara/Source/Niagara/Private/NiagaraDataInterfaceColorCurve.cpp b/Engine/Plugins/FX/Niagara/Source/Niagara/Private/NiagaraDataInterfaceColorCurve.cpp
--- a/Engine/Plugins/FX/Niagara/Source/Niagara/Private/NiagaraDataInterfaceColorCurve.cpp
+++ b/Engine/Plugins/FX/Niagara/Source/Niagara/Private/NiagaraDataInterfaceColorCurve.cpp
@@ -5,6 +5,7 @@
 #include "Curves/CurveLinearColor.h"
 #include "Curves/CurveFloat.h"
 #include "NiagaraTypes.h"
+#include "NiagaraCustomVersion.h"
 
 //////////////////////////////////////////////////////////////////////////
 //Color Curve
@@ -19,7 +20,7 @@ void UNiagaraDataInterfaceColorCurve::PostInitProperties()
 {
 	Super::PostInitProperties();
 
-	//Can we regitser data interfaces as regular types and fold them into the FNiagaraVariable framework for UI and function calls etc?
+	//Can we register data interfaces as regular types and fold them into the FNiagaraVariable framework for UI and function calls etc?
 	if (HasAnyFlags(RF_ClassDefaultObject))
 	{
 		FNiagaraTypeRegistry::Register(FNiagaraTypeDefinition(GetClass()), true, false, false);
@@ -31,7 +32,20 @@ void UNiagaraDataInterfaceColorCurve::PostInitProperties()
 void UNiagaraDataInterfaceColorCurve::PostLoad()
 {
 	Super::PostLoad();
+
+	const int32 NiagaraVer = GetLinkerCustomVersion(FNiagaraCustomVersion::GUID);
+
+	if (NiagaraVer < FNiagaraCustomVersion::CurveLUTNowOnByDefault)
+	{
+		UpdateLUT();
+	}
+	TArray<float> OldLUT = ShaderLUT;
 	UpdateLUT();
+	if (!CompareLUTS(OldLUT))
+	{
+		UE_LOG(LogNiagara, Log, TEXT("PostLoad LUT generation is out of sync. Please investigate. %s"), *GetPathName());
+	}
+
 }
 
 #if WITH_EDITOR
@@ -47,12 +61,13 @@ void UNiagaraDataInterfaceColorCurve::PostEditChangeProperty(struct FPropertyCha
 		{
 			Modify();
 			RedCurve = ColorCurveAsset->FloatCurves[0];
-			BlueCurve = ColorCurveAsset->FloatCurves[1];
-			GreenCurve = ColorCurveAsset->FloatCurves[2];
+			GreenCurve = ColorCurveAsset->FloatCurves[1];
+			BlueCurve = ColorCurveAsset->FloatCurves[2];
 			AlphaCurve = ColorCurveAsset->FloatCurves[3];
 		}
-		UpdateLUT();
 	}
+
+	UpdateLUT();
 }
 
 #endif
@@ -60,9 +75,32 @@ void UNiagaraDataInterfaceColorCurve::PostEditChangeProperty(struct FPropertyCha
 void UNiagaraDataInterfaceColorCurve::UpdateLUT()
 {
 	ShaderLUT.Empty();
+
+	if (bAllowUnnormalizedLUT && (RedCurve.GetNumKeys() > 0 || GreenCurve.GetNumKeys() > 0 || BlueCurve.GetNumKeys() > 0 || AlphaCurve.GetNumKeys() > 0))
+	{
+		LUTMinTime = FLT_MAX;
+		LUTMinTime = FMath::Min(RedCurve.GetNumKeys() > 0 ? RedCurve.GetFirstKey().Time : LUTMinTime, LUTMinTime);
+		LUTMinTime = FMath::Min(GreenCurve.GetNumKeys() > 0 ? GreenCurve.GetFirstKey().Time : LUTMinTime, LUTMinTime);
+		LUTMinTime = FMath::Min(BlueCurve.GetNumKeys() > 0 ? BlueCurve.GetFirstKey().Time : LUTMinTime, LUTMinTime);
+		LUTMinTime = FMath::Min(AlphaCurve.GetNumKeys() > 0 ? AlphaCurve.GetFirstKey().Time : LUTMinTime, LUTMinTime);
+
+		LUTMaxTime = FLT_MIN;
+		LUTMaxTime = FMath::Max(RedCurve.GetNumKeys() > 0 ? RedCurve.GetLastKey().Time : LUTMaxTime, LUTMaxTime);
+		LUTMaxTime = FMath::Max(GreenCurve.GetNumKeys() > 0 ? GreenCurve.GetLastKey().Time : LUTMaxTime, LUTMaxTime);
+		LUTMaxTime = FMath::Max(BlueCurve.GetNumKeys() > 0 ? BlueCurve.GetLastKey().Time : LUTMaxTime, LUTMaxTime);
+		LUTMaxTime = FMath::Max(AlphaCurve.GetNumKeys() > 0 ? AlphaCurve.GetLastKey().Time : LUTMaxTime, LUTMaxTime);
+		LUTInvTimeRange = 1.0f / (LUTMaxTime - LUTMinTime);
+	}
+	else
+	{
+		LUTMinTime = 0.0f;
+		LUTMaxTime = 1.0f;
+		LUTInvTimeRange = 1.0f;
+	}
+
 	for (uint32 i = 0; i < CurveLUTWidth; i++)
 	{
-		float X = i / (float)CurveLUTWidth;
+		float X = UnnormalizeTime(i / (float)CurveLUTWidth);
 		FLinearColor C(RedCurve.Eval(X), GreenCurve.Eval(X), BlueCurve.Eval(X), AlphaCurve.Eval(X));
 		ShaderLUT.Add(C.R);
 		ShaderLUT.Add(C.G);
@@ -72,9 +110,9 @@ void UNiagaraDataInterfaceColorCurve::UpdateLUT()
 	GPUBufferDirty = true;
 }
 
-bool UNiagaraDataInterfaceColorCurve::CopyTo(UNiagaraDataInterface* Destination) const 
+bool UNiagaraDataInterfaceColorCurve::CopyToInternal(UNiagaraDataInterface* Destination) const 
 {
-	if (!Super::CopyTo(Destination))
+	if (!Super::CopyToInternal(Destination))
 	{
 		return false;
 	}
@@ -85,6 +123,10 @@ bool UNiagaraDataInterfaceColorCurve::CopyTo(UNiagaraDataInterface* Destination)
 	DestinationColorCurve->AlphaCurve = AlphaCurve;
 	DestinationColorCurve->UpdateLUT();
 
+	if (!CompareLUTS(DestinationColorCurve->ShaderLUT))
+	{
+		UE_LOG(LogNiagara, Log, TEXT("CopyToInternal LUT generation is out of sync. Please investigate. %s to %s"), *GetPathName(), *DestinationColorCurve->GetPathName());
+	}
 	return true;
 }
 
@@ -127,10 +169,10 @@ void UNiagaraDataInterfaceColorCurve::GetFunctions(TArray<FNiagaraFunctionSignat
 // the HLSL in the spirit of a static switch
 // TODO: need a way to identify each specific function here
 // 
-bool UNiagaraDataInterfaceColorCurve::GetFunctionHLSL(FString FunctionName, TArray<DIGPUBufferParamDescriptor> &Descriptors, FString &HLSLInterfaceID, FString &OutHLSL)
+bool UNiagaraDataInterfaceColorCurve::GetFunctionHLSL(const FName& DefinitionFunctionName, FString InstanceFunctionName, TArray<FDIGPUBufferParamDescriptor> &Descriptors, FString &HLSLInterfaceID, FString &OutHLSL)
 {
 	FString BufferName = Descriptors[0].BufferParamName;
-	OutHLSL += TEXT("void ") + FunctionName + TEXT("(in float In_X, out float4 Out_Value) \n{\n");
+	OutHLSL += TEXT("void ") + InstanceFunctionName + TEXT("(in float In_X, out float4 Out_Value) \n{\n");
 	OutHLSL += TEXT("\t Out_Value.x = ") + BufferName + TEXT("[(int)(In_X *") + FString::FromInt(CurveLUTWidth) + TEXT(")* 4 ];");
 	OutHLSL += TEXT("\t Out_Value.y = ") + BufferName + TEXT("[1+ (int)(In_X *") + FString::FromInt(CurveLUTWidth) + TEXT(")* 4 ];");
 	OutHLSL += TEXT("\t Out_Value.z = ") + BufferName + TEXT("[2+ (int)(In_X *") + FString::FromInt(CurveLUTWidth) + TEXT(")* 4 ];");
@@ -146,25 +188,27 @@ bool UNiagaraDataInterfaceColorCurve::GetFunctionHLSL(FString FunctionName, TArr
 // 3. store buffer declaration hlsl in OutHLSL
 // multiple buffers can be defined at once here
 //
-void UNiagaraDataInterfaceColorCurve::GetBufferDefinitionHLSL(FString DataInterfaceID, TArray<DIGPUBufferParamDescriptor> &BufferDescriptors, FString &OutHLSL)
+void UNiagaraDataInterfaceColorCurve::GetBufferDefinitionHLSL(FString DataInterfaceID, TArray<FDIGPUBufferParamDescriptor> &BufferDescriptors, FString &OutHLSL)
 {
 	FString BufferName = "CurveLUT" + DataInterfaceID;
 	OutHLSL += TEXT("Buffer<float> ") + BufferName + TEXT(";\n");
 
-	BufferDescriptors.Add( DIGPUBufferParamDescriptor(BufferName, 0) );		// add a descriptor for shader parameter binding
+	BufferDescriptors.Add( FDIGPUBufferParamDescriptor(BufferName, 0) );		// add a descriptor for shader parameter binding
 
 }
 
 // called after translate, to setup buffers matching the buffer descriptors generated during hlsl translation
 // need to do this because the script used during translate is a clone, including its DIs
 //
-void UNiagaraDataInterfaceColorCurve::SetupBuffers(TArray<DIGPUBufferParamDescriptor> &BufferDescriptors)
+void UNiagaraDataInterfaceColorCurve::SetupBuffers(FDIBufferDescriptorStore &BufferDescriptors)
 {
-	for (DIGPUBufferParamDescriptor &Desc : BufferDescriptors)
+	GPUBuffers.Empty();
+	for (FDIGPUBufferParamDescriptor &Desc : BufferDescriptors.Descriptors)
 	{
 		FNiagaraDataInterfaceBufferData BufferData(*Desc.BufferParamName);	// store off the data for later use
 		GPUBuffers.Add(BufferData);
 	}
+	GPUBufferDirty = true;
 }
 
 // return the GPU buffer array (called from NiagaraInstanceBatcher to get the buffers for setting to the shader)
@@ -179,7 +223,13 @@ TArray<FNiagaraDataInterfaceBufferData> &UNiagaraDataInterfaceColorCurve::GetBuf
 
 		FNiagaraDataInterfaceBufferData &GPUBuffer = GPUBuffers[0];
 		GPUBuffer.Buffer.Release();
-		GPUBuffer.Buffer.Initialize(sizeof(float), CurveLUTWidth*4 , EPixelFormat::PF_R32_FLOAT, BUF_Volatile);	// always allocate for up to 64 data sets
+
+/*		TResourceArray<float> FloatRes;
+		FloatRes.AddZeroed(ShaderLUT.Num());
+		uint32 BufferSize = ShaderLUT.Num() * sizeof(float);
+		FPlatformMemory::Memcpy(FloatRes.GetData(), ShaderLUT.GetData(), BufferSize);*/
+
+		GPUBuffer.Buffer.Initialize(sizeof(float), CurveLUTWidth*4 , EPixelFormat::PF_R32_FLOAT, BUF_Static);	// always allocate for up to 64 data sets
 		uint32 BufferSize = ShaderLUT.Num() * sizeof(float);
 		int32 *BufferData = static_cast<int32*>(RHILockVertexBuffer(GPUBuffer.Buffer.Buffer, 0, BufferSize, EResourceLockMode::RLM_WriteOnly));
 		FPlatformMemory::Memcpy(BufferData, ShaderLUT.GetData(), BufferSize);
@@ -198,7 +248,7 @@ FVMExternalFunction UNiagaraDataInterfaceColorCurve::GetVMExternalFunction(const
 {
 	if (BindingInfo.Name == TEXT("SampleColorCurve") && BindingInfo.GetNumInputs() == 1 && BindingInfo.GetNumOutputs() == 4)
 	{
-		return TNDIParamBinder<0, float, NDI_FUNC_BINDER(UNiagaraDataInterfaceColorCurve, SampleCurve)>::Bind(this, BindingInfo, InstanceData);
+		return TCurveUseLUTBinder<TNDIParamBinder<0, float, NDI_FUNC_BINDER(UNiagaraDataInterfaceColorCurve, SampleCurve)>>::Bind(this, BindingInfo, InstanceData);
 	}
 	else
 	{
@@ -208,7 +258,21 @@ FVMExternalFunction UNiagaraDataInterfaceColorCurve::GetVMExternalFunction(const
 	}
 }
 
-template<typename XParamType>
+template<>
+FORCEINLINE_DEBUGGABLE FLinearColor UNiagaraDataInterfaceColorCurve::SampleCurveInternal<TIntegralConstant<bool, true>>(float X)
+{
+	float NormalizedX = NormalizeTime(X);
+	int32 AccessIdx = FMath::Clamp<int32>(FMath::TruncToInt(NormalizedX * CurveLUTWidthMinusOne), 0, CurveLUTWidthMinusOne) * CurveLUTNumElems;
+	return FLinearColor(ShaderLUT[AccessIdx], ShaderLUT[AccessIdx + 1], ShaderLUT[AccessIdx + 2], ShaderLUT[AccessIdx + 3]);
+}
+
+template<>
+FORCEINLINE_DEBUGGABLE FLinearColor UNiagaraDataInterfaceColorCurve::SampleCurveInternal<TIntegralConstant<bool, false>>(float X)
+{
+	return FLinearColor(RedCurve.Eval(X), GreenCurve.Eval(X), BlueCurve.Eval(X), AlphaCurve.Eval(X));
+}
+
+template<typename UseLUT, typename XParamType>
 void UNiagaraDataInterfaceColorCurve::SampleCurve(FVectorVMContext& Context)
 {
 	//TODO: Create some SIMDable optimized representation of the curve to do this faster.
@@ -221,7 +285,7 @@ void UNiagaraDataInterfaceColorCurve::SampleCurve(FVectorVMContext& Context)
 	for (int32 i = 0; i < Context.NumInstances; ++i)
 	{
 		float X = XParam.Get();
-		FLinearColor C(RedCurve.Eval(X), GreenCurve.Eval(X), BlueCurve.Eval(X), AlphaCurve.Eval(X));
+		FLinearColor C = SampleCurveInternal<UseLUT>(X);
 		*SamplePtrR.GetDest() = C.R;
 		*SamplePtrG.GetDest() = C.G;
 		*SamplePtrB.GetDest() = C.B;
diff --git a/Engine/Plugins/FX/Niagara/Source/Niagara/Private/NiagaraDataInterfaceCurlNoise.cpp b/Engine/Plugins/FX/Niagara/Source/Niagara/Private/NiagaraDataInterfaceCurlNoise.cpp
--- a/Engine/Plugins/FX/Niagara/Source/Niagara/Private/NiagaraDataInterfaceCurlNoise.cpp
+++ b/Engine/Plugins/FX/Niagara/Source/Niagara/Private/NiagaraDataInterfaceCurlNoise.cpp
@@ -1,9 +1,9 @@
 #include "NiagaraDataInterfaceCurlNoise.h"
 
-VectorRegister UNiagaraDataInterfaceCurlNoise::NoiseTable[17][17][17];
-
 UNiagaraDataInterfaceCurlNoise::UNiagaraDataInterfaceCurlNoise(FObjectInitializer const& ObjectInitializer)
 	: Super(ObjectInitializer)
+	, bGPUBufferDirty(true)
+	, Seed(0)
 {
 
 }
@@ -18,6 +18,49 @@ void UNiagaraDataInterfaceCurlNoise::PostInitProperties()
 	}
 }
 
+void UNiagaraDataInterfaceCurlNoise::PostLoad()
+{
+	Super::PostLoad();
+	InitNoiseLUT();
+}
+
+#if WITH_EDITOR
+
+void UNiagaraDataInterfaceCurlNoise::PostEditChangeProperty(struct FPropertyChangedEvent& PropertyChangedEvent)
+{
+	Super::PostEditChangeProperty(PropertyChangedEvent);
+
+	if (PropertyChangedEvent.Property && PropertyChangedEvent.Property->GetFName() == GET_MEMBER_NAME_CHECKED(UNiagaraDataInterfaceCurlNoise, Seed))
+	{
+		InitNoiseLUT();
+	}
+}
+
+#endif
+
+bool UNiagaraDataInterfaceCurlNoise::CopyToInternal(UNiagaraDataInterface* Destination) const
+{
+	if (!Super::CopyToInternal(Destination))
+	{
+		return false;
+	}
+	UNiagaraDataInterfaceCurlNoise* DestinationCurlNoise = CastChecked<UNiagaraDataInterfaceCurlNoise>(Destination);
+	DestinationCurlNoise->Seed = Seed;
+	DestinationCurlNoise->InitNoiseLUT();
+
+	return true;
+}
+
+bool UNiagaraDataInterfaceCurlNoise::Equals(const UNiagaraDataInterface* Other) const
+{
+	if (!Super::Equals(Other))
+	{
+		return false;
+	}
+	const UNiagaraDataInterfaceCurlNoise* OtherCurlNoise = CastChecked<const UNiagaraDataInterfaceCurlNoise>(Other);
+	return OtherCurlNoise->Seed == Seed;
+}
+
 void UNiagaraDataInterfaceCurlNoise::GetFunctions(TArray<FNiagaraFunctionSignature>& OutFunctions)
 {
 	FNiagaraFunctionSignature Sig;
@@ -56,42 +99,46 @@ void UNiagaraDataInterfaceCurlNoise::SampleNoiseField(FVectorVMContext& Context)
 
 		const VectorRegister One = MakeVectorRegister(1.0f, 1.0f, 1.0f, 1.0f);
 		const VectorRegister Zero = MakeVectorRegister(0.0f, 0.0f, 0.0f, 0.0f);
-		const VectorRegister VecSize = MakeVectorRegister(16.0f, 16.0f, 16.0f, 16.0f);
+		// We use 15,15,15 here because the value will be between 0 and 1. 1 should indicate the maximal index of the 
+		// real table and we pad that out to handle border work. So our values for the Cxyz below should not
+		// ever be 17 in any dimension as that would overflow the array.
+		const VectorRegister VecSize = MakeVectorRegister(15.0f, 15.0f, 15.0f, 15.0f);
 
 		VectorRegister Dst = MakeVectorRegister(0.0f, 0.0f, 0.0f, 0.0f);
 
-		for (uint32 i = 1; i < 2; i++)
-		{
-			float Di = 0.2f * (1.0f / (1 << i));
-			VectorRegister Div = MakeVectorRegister(Di, Di, Di, Di);
-			VectorRegister Coords = VectorMod(VectorAbs(VectorMultiply(InCoords, Div)), VecSize);
-			Coords = VectorMin(Coords, VecSize);
-			Coords = VectorMax(Coords, Zero);
-			const float *CoordPtr = reinterpret_cast<float const*>(&Coords);
-			const int32 Cx = CoordPtr[0];
-			const int32 Cy = CoordPtr[1];
-			const int32 Cz = CoordPtr[2];
-
-			VectorRegister Frac = VectorFractional(Coords);
-			VectorRegister Alpha = VectorReplicate(Frac, 0);
-			VectorRegister OneMinusAlpha = VectorSubtract(One, Alpha);
-
-			VectorRegister XV1 = VectorMultiplyAdd(NoiseTable[Cx][Cy][Cz], Alpha, VectorMultiply(NoiseTable[Cx + 1][Cy][Cz], OneMinusAlpha));
-			VectorRegister XV2 = VectorMultiplyAdd(NoiseTable[Cx][Cy + 1][Cz], Alpha, VectorMultiply(NoiseTable[Cx + 1][Cy + 1][Cz], OneMinusAlpha));
-			VectorRegister XV3 = VectorMultiplyAdd(NoiseTable[Cx][Cy][Cz + 1], Alpha, VectorMultiply(NoiseTable[Cx + 1][Cy][Cz + 1], OneMinusAlpha));
-			VectorRegister XV4 = VectorMultiplyAdd(NoiseTable[Cx][Cy + 1][Cz + 1], Alpha, VectorMultiply(NoiseTable[Cx + 1][Cy + 1][Cz + 1], OneMinusAlpha));
-
-			Alpha = VectorReplicate(Frac, 1);
-			OneMinusAlpha = VectorSubtract(One, Alpha);
-			VectorRegister YV1 = VectorMultiplyAdd(XV1, Alpha, VectorMultiply(XV2, OneMinusAlpha));
-			VectorRegister YV2 = VectorMultiplyAdd(XV3, Alpha, VectorMultiply(XV4, OneMinusAlpha));
-
-			Alpha = VectorReplicate(Frac, 2);
-			OneMinusAlpha = VectorSubtract(One, Alpha);
-			VectorRegister ZV = VectorMultiplyAdd(YV1, Alpha, VectorMultiply(YV2, OneMinusAlpha));
-
-			Dst = VectorAdd(Dst, ZV);
-		}
+		float Di = 0.2f; // Hard-coded scale TODO remove!
+		VectorRegister Div = MakeVectorRegister(Di, Di, Di, Di);
+		VectorRegister Coords = VectorMod(VectorAbs(VectorMultiply(InCoords, Div)), VecSize);
+		Coords = VectorMin(Coords, VecSize);
+		Coords = VectorMax(Coords, Zero);
+		const float *CoordPtr = reinterpret_cast<float const*>(&Coords);
+		const int32 Cx = CoordPtr[0];
+		const int32 Cy = CoordPtr[1];
+		const int32 Cz = CoordPtr[2];
+
+		VectorRegister Frac = VectorFractional(Coords);
+		VectorRegister Alpha = VectorReplicate(Frac, 0);
+		VectorRegister OneMinusAlpha = VectorSubtract(One, Alpha);
+
+		// Trilinear interpolation, as defined by https://en.wikipedia.org/wiki/Trilinear_interpolation
+		ensure(Cx + 1 < 17);
+		ensure(Cy + 1 < 17);
+		ensure(Cz + 1 < 17);
+		VectorRegister C00 = VectorMultiplyAdd(NoiseTable[Cx][Cy][Cz], OneMinusAlpha, VectorMultiply(NoiseTable[Cx + 1][Cy][Cz], Alpha)); // (x0, y0, z0)(1-a) + (x1, y0, z0)a
+		VectorRegister C01 = VectorMultiplyAdd(NoiseTable[Cx][Cy][Cz + 1], OneMinusAlpha, VectorMultiply(NoiseTable[Cx + 1][Cy][Cz + 1], Alpha)); // (x0, y0, z1)(1-a) + (x1, y0, z1)a
+		VectorRegister C10 = VectorMultiplyAdd(NoiseTable[Cx][Cy + 1][Cz], OneMinusAlpha, VectorMultiply(NoiseTable[Cx + 1][Cy + 1][Cz], Alpha)); // (x0, y1, z0)(1-a) + (x1, y1, z0)a
+		VectorRegister C11 = VectorMultiplyAdd(NoiseTable[Cx][Cy + 1][Cz + 1], OneMinusAlpha, VectorMultiply(NoiseTable[Cx + 1][Cy + 1][Cz + 1], Alpha)); // (x0, y1, z1)(1-a) + (x1, y1, z1)a
+
+		Alpha = VectorReplicate(Frac, 1);
+		OneMinusAlpha = VectorSubtract(One, Alpha);
+		VectorRegister C0 = VectorMultiplyAdd(C00, OneMinusAlpha, VectorMultiply(C10, Alpha));
+		VectorRegister C1 = VectorMultiplyAdd(C01, OneMinusAlpha, VectorMultiply(C11, Alpha));
+
+		Alpha = VectorReplicate(Frac, 2);
+		OneMinusAlpha = VectorSubtract(One, Alpha);
+		VectorRegister ZV = VectorMultiplyAdd(C0, OneMinusAlpha, VectorMultiply(C1, Alpha));
+
+		Dst = VectorAdd(Dst, ZV);
 
 		float *RegPtr = reinterpret_cast<float*>(&Dst);
 		*OutSampleX.GetDest() = RegPtr[0];
@@ -113,29 +160,32 @@ void UNiagaraDataInterfaceCurlNoise::SampleNoiseField(FVectorVMContext& Context)
 // the HLSL in the spirit of a static switch
 // TODO: need a way to identify each specific function here
 // 
-bool UNiagaraDataInterfaceCurlNoise::GetFunctionHLSL(FString FunctionName, TArray<DIGPUBufferParamDescriptor> &Descriptors, FString &HLSLInterfaceID, FString &OutHLSL)
+bool UNiagaraDataInterfaceCurlNoise::GetFunctionHLSL(const FName& DefinitionFunctionName, FString InstanceFunctionName, TArray<FDIGPUBufferParamDescriptor> &Descriptors, FString &HLSLInterfaceID, FString &OutHLSL)
 {
+
 	FString BufferName = Descriptors[0].BufferParamName;
-	OutHLSL += TEXT("void ") + FunctionName + TEXT("(in float3 In_XYZ, out float3 Out_Value) \n{\n");
-	OutHLSL += TEXT("\t float3 a = trunc((In_XYZ*0.2) / 16.0);\n");
-	OutHLSL += TEXT("\t float3 ModXYZ = (In_XYZ*0.2) - a*16.0;\n");
+	OutHLSL += TEXT("void ") + InstanceFunctionName + TEXT("(in float3 In_XYZ, out float3 Out_Value) \n{\n");
+	OutHLSL += TEXT("\t float3 a = trunc((In_XYZ*0.2) / 15.0);\n");
+	OutHLSL += TEXT("\t float3 ModXYZ = (In_XYZ*0.2) - a*15.0;\n");
 	OutHLSL += TEXT("\t int3 IntCoord = int3(ModXYZ.x, ModXYZ.y, ModXYZ.z);\n");
 	OutHLSL += TEXT("\t float3 frc = frac(ModXYZ);\n");
-	OutHLSL += TEXT("\t float3 V1 = ") + BufferName + TEXT("[IntCoord.x + IntCoord.y*17 + IntCoord.z*17*17].xyz;\n");
-	OutHLSL += TEXT("\t float3 V2 = ") + BufferName + TEXT("[IntCoord.x+1 + IntCoord.y*17 + IntCoord.z*17*17].xyz;\n");
-	OutHLSL += TEXT("\t float3 XV1 = lerp(V1, V2, frc.xxx);\n");
-	OutHLSL += TEXT("\t V1 = ") + BufferName + TEXT("[IntCoord.x + (IntCoord.y+1)*17 + IntCoord.z*17*17].xyz;\n");
-	OutHLSL += TEXT("\t V2 = ") + BufferName + TEXT("[IntCoord.x+1 + (IntCoord.y+1)*17 + IntCoord.z*17*17].xyz;\n");
-	OutHLSL += TEXT("\t float3 XV2 = lerp(V1, V2, frc.xxx);\n");
-	OutHLSL += TEXT("\t V1 = ") + BufferName + TEXT("[IntCoord.x + IntCoord.y*17 + (IntCoord.z+1)*17*17].xyz;\n");
-	OutHLSL += TEXT("\t V2 = ") + BufferName + TEXT("[IntCoord.x+1 + IntCoord.y*17 + (IntCoord.z+1)*17*17].xyz;\n");
-	OutHLSL += TEXT("\t float3 XV3 = lerp(V1, V2, frc.xxx);\n");
-	OutHLSL += TEXT("\t V1 = ") + BufferName + TEXT("[IntCoord.x + (IntCoord.y+1)*17 + (IntCoord.z+1)*17*17].xyz;\n");
-	OutHLSL += TEXT("\t V2 = ") + BufferName + TEXT("[IntCoord.x+1 + (IntCoord.y+1)*17 + (IntCoord.z+1)*17*17].xyz;\n");
-	OutHLSL += TEXT("\t float3 XV4 = lerp(V1, V2, frc.xxx);\n");
-	OutHLSL += TEXT("\t float3 YV1 = lerp(XV1, XV2, frc.yyy);\n");
-	OutHLSL += TEXT("\t float3 YV2 = lerp(XV3, XV4, frc.yyy);\n");
-	OutHLSL += TEXT("\t Out_Value = lerp(YV1, YV2, frc.zzz);\n");
+	// Trilinear interpolation, as defined by https://en.wikipedia.org/wiki/Trilinear_interpolation
+	// HLSL lerp is defined as lerp(x, y, s) = x*(1-s) + y*s
+	OutHLSL += TEXT("\t float3 V1 = ") + BufferName + TEXT("[IntCoord.x + IntCoord.y*17 + IntCoord.z*17*17].xyz;\n"); // x0,y0,z0
+	OutHLSL += TEXT("\t float3 V2 = ") + BufferName + TEXT("[IntCoord.x+1 + IntCoord.y*17 + IntCoord.z*17*17].xyz;\n"); //  x1,y0,z0
+	OutHLSL += TEXT("\t float3 C00 = lerp(V1, V2, frc.xxx);\n");
+	OutHLSL += TEXT("\t V1 = ") + BufferName + TEXT("[IntCoord.x + IntCoord.y*17 + (IntCoord.z+1)*17*17].xyz;\n"); // x0, y0, z1
+	OutHLSL += TEXT("\t V2 = ") + BufferName + TEXT("[IntCoord.x+1 + IntCoord.y*17 + (IntCoord.z+1)*17*17].xyz;\n"); //x1, y0, z1
+	OutHLSL += TEXT("\t float3 C01 = lerp(V1, V2, frc.xxx);\n");
+	OutHLSL += TEXT("\t V1 = ") + BufferName + TEXT("[IntCoord.x + (IntCoord.y+1)*17 + IntCoord.z*17*17].xyz;\n"); //x0, y1, z0
+	OutHLSL += TEXT("\t V2 = ") + BufferName + TEXT("[IntCoord.x+1 + (IntCoord.y+1)*17 + IntCoord.z*17*17].xyz;\n"); // x1, y1, z0
+	OutHLSL += TEXT("\t float3 C10 = lerp(V1, V2, frc.xxx);\n");
+	OutHLSL += TEXT("\t V1 = ") + BufferName + TEXT("[IntCoord.x + (IntCoord.y+1)*17 + (IntCoord.z+1)*17*17].xyz;\n"); // x0, y1, z1
+	OutHLSL += TEXT("\t V2 = ") + BufferName + TEXT("[IntCoord.x+1 + (IntCoord.y+1)*17 + (IntCoord.z+1)*17*17].xyz;\n"); //x1, y1, z1
+	OutHLSL += TEXT("\t float3 C11 = lerp(V1, V2, frc.xxx);\n");
+	OutHLSL += TEXT("\t float3 C0 = lerp(C00, C10, frc.yyy);\n");
+	OutHLSL += TEXT("\t float3 C1 = lerp(C01, C11, frc.yyy);\n");
+	OutHLSL += TEXT("\t Out_Value = lerp(C0, C1, frc.zzz);\n");
 	OutHLSL += TEXT("\n}\n");
 	return true;
 }
@@ -147,21 +197,22 @@ bool UNiagaraDataInterfaceCurlNoise::GetFunctionHLSL(FString FunctionName, TArra
 // 3. store buffer declaration hlsl in OutHLSL
 // multiple buffers can be defined at once here
 //
-void UNiagaraDataInterfaceCurlNoise::GetBufferDefinitionHLSL(FString DataInterfaceID, TArray<DIGPUBufferParamDescriptor> &BufferDescriptors, FString &OutHLSL)
+void UNiagaraDataInterfaceCurlNoise::GetBufferDefinitionHLSL(FString DataInterfaceID, TArray<FDIGPUBufferParamDescriptor> &BufferDescriptors, FString &OutHLSL)
 {
 	FString BufferName = "CurlNoiseLUT" + DataInterfaceID;
 	OutHLSL += TEXT("Buffer<float4> ") + BufferName + TEXT(";\n");
 
-	BufferDescriptors.Add(DIGPUBufferParamDescriptor(BufferName, 0));		// add a descriptor for shader parameter binding
+	BufferDescriptors.Add(FDIGPUBufferParamDescriptor(BufferName, 0));		// add a descriptor for shader parameter binding
 
 }
 
 // called after translate, to setup buffers matching the buffer descriptors generated during hlsl translation
 // need to do this because the script used during translate is a clone, including its DIs
 //
-void UNiagaraDataInterfaceCurlNoise::SetupBuffers(TArray<DIGPUBufferParamDescriptor> &BufferDescriptors)
+void UNiagaraDataInterfaceCurlNoise::SetupBuffers(FDIBufferDescriptorStore &BufferDescriptors)
 {
-	for (DIGPUBufferParamDescriptor &Desc : BufferDescriptors)
+	GPUBuffers.Empty();
+	for (FDIGPUBufferParamDescriptor &Desc : BufferDescriptors.Descriptors)
 	{
 		FNiagaraDataInterfaceBufferData BufferData(*Desc.BufferParamName);	// store off the data for later use
 		GPUBuffers.Add(BufferData);
@@ -181,10 +232,10 @@ TArray<FNiagaraDataInterfaceBufferData> &UNiagaraDataInterfaceCurlNoise::GetBuff
 		check(GPUBuffers.Num() > 0);
 		FNiagaraDataInterfaceBufferData &GPUBuffer = GPUBuffers[0];
 		GPUBuffer.Buffer.Release();
-		GPUBuffer.Buffer.Initialize(sizeof(float) * 4, 17 * 17 * 17, EPixelFormat::PF_A32B32G32R32F);
 		uint32 BufferSize = 17 * 17 * 17 * sizeof(float) * 4;
-		int32 *BufferData = static_cast<int32*>(RHILockVertexBuffer(GPUBuffer.Buffer.Buffer, 0, BufferSize, EResourceLockMode::RLM_WriteOnly));
-		FVector4 TempTable[17 * 17 * 17];
+		//int32 *BufferData = static_cast<int32*>(RHILockVertexBuffer(GPUBuffer.Buffer.Buffer, 0, BufferSize, EResourceLockMode::RLM_WriteOnly));
+		TResourceArray<FVector4> TempTable;
+		TempTable.AddUninitialized(17 * 17 * 17);
 		for (int z = 0; z < 17; z++)
 		{
 			for (int y = 0; y < 17; y++)
@@ -196,9 +247,9 @@ TArray<FNiagaraDataInterfaceBufferData> &UNiagaraDataInterfaceCurlNoise::GetBuff
 				}
 			}
 		}
-		FPlatformMemory::Memcpy(BufferData, TempTable, BufferSize);
-		//		FPlatformMemory::Memcpy(BufferData, NoiseTable, BufferSize);
-		RHIUnlockVertexBuffer(GPUBuffer.Buffer.Buffer);
+		GPUBuffer.Buffer.Initialize(sizeof(float) * 4, 17 * 17 * 17, EPixelFormat::PF_A32B32G32R32F, BUF_Static, TEXT("CurlnoiseTable"), &TempTable);
+		//FPlatformMemory::Memcpy(BufferData, TempTable, BufferSize);
+		//RHIUnlockVertexBuffer(GPUBuffer.Buffer.Buffer);
 		bGPUBufferDirty = false;
 	}
 
@@ -207,12 +258,41 @@ TArray<FNiagaraDataInterfaceBufferData> &UNiagaraDataInterfaceCurlNoise::GetBuff
 }
 
 
+// replicate a border for filtering
+template<typename T>
+void UNiagaraDataInterfaceCurlNoise::ReplicateBorder(T* DestBuffer)
+{
+	uint32 Extent = 17;
+	uint32 Square = Extent*Extent;
+	uint32 Last = Extent - 1;
 
+	for (uint32 z = 0; z < Extent; ++z)
+	{
+		for (uint32 y = 0; y < Extent; ++y)
+		{
+			DestBuffer[Last + y * Extent + z * Square] = DestBuffer[0 + y * Extent + z * Square];
+		}
+	}
+	for (uint32 z = 0; z < Extent; ++z)
+	{
+		for (uint32 x = 0; x < Extent; ++x)
+		{
+			DestBuffer[x + Last * Extent + z * Square] = DestBuffer[x + 0 * Extent + z * Square];
+		}
+	}
+	for (uint32 y = 0; y < Extent; ++y)
+	{
+		for (uint32 x = 0; x < Extent; ++x)
+		{
+			DestBuffer[x + y * Extent + Last * Square] = DestBuffer[x + y * Extent + 0 * Square];
+		}
+	}
+}
 
 void UNiagaraDataInterfaceCurlNoise::InitNoiseLUT()
 {
 	// seed random stream
-	FRandomStream RandStream;
+	FRandomStream RandStream(Seed);
 
 	// random noise
 	float TempTable[17][17][17];
@@ -229,6 +309,7 @@ void UNiagaraDataInterfaceCurlNoise::InitNoiseLUT()
 	}
 
 	// pad
+	/*
 	for (int i = 0; i < 17; i++)
 	{
 		for (int j = 0; j < 17; j++)
@@ -238,6 +319,8 @@ void UNiagaraDataInterfaceCurlNoise::InitNoiseLUT()
 			TempTable[16][j][i] = TempTable[0][j][i];
 		}
 	}
+	*/
+	ReplicateBorder(&TempTable[0][0][0]);
 
 	// compute gradients
 	FVector TempTable2[17][17][17];
@@ -257,6 +340,7 @@ void UNiagaraDataInterfaceCurlNoise::InitNoiseLUT()
 		}
 	}
 
+	/*
 	// pad
 	for (int i = 0; i < 17; i++)
 	{
@@ -267,8 +351,10 @@ void UNiagaraDataInterfaceCurlNoise::InitNoiseLUT()
 			TempTable2[16][j][i] = TempTable2[0][j][i];
 		}
 	}
+	*/
+	ReplicateBorder(&TempTable2[0][0][0]);
 
-
+	// http://prideout.net/blog/?p=63           ???
 	// compute curl of gradient field
 	for (int z = 0; z < 16; z++)
 	{
@@ -290,6 +376,7 @@ void UNiagaraDataInterfaceCurlNoise::InitNoiseLUT()
 	}
 
 	// pad
+	/*
 	for (int i = 0; i < 17; i++)
 	{
 		for (int j = 0; j < 17; j++)
@@ -298,6 +385,8 @@ void UNiagaraDataInterfaceCurlNoise::InitNoiseLUT()
 			NoiseTable[i][16][j] = NoiseTable[i][0][j];
 			NoiseTable[16][j][i] = NoiseTable[0][j][i];
 		}
-	}
+	}*/
+	ReplicateBorder(&NoiseTable[0][0][0]);
 
+	bGPUBufferDirty = true;
 }
diff --git a/Engine/Plugins/FX/Niagara/Source/Niagara/Private/NiagaraDataInterfaceCurve.cpp b/Engine/Plugins/FX/Niagara/Source/Niagara/Private/NiagaraDataInterfaceCurve.cpp
--- a/Engine/Plugins/FX/Niagara/Source/Niagara/Private/NiagaraDataInterfaceCurve.cpp
+++ b/Engine/Plugins/FX/Niagara/Source/Niagara/Private/NiagaraDataInterfaceCurve.cpp
@@ -5,6 +5,7 @@
 #include "Curves/CurveLinearColor.h"
 #include "Curves/CurveFloat.h"
 #include "NiagaraTypes.h"
+#include "NiagaraCustomVersion.h"
 
 
 UNiagaraDataInterfaceCurve::UNiagaraDataInterfaceCurve(FObjectInitializer const& ObjectInitializer)
@@ -39,21 +40,21 @@ void UNiagaraDataInterfaceCurve::PostEditChangeProperty(struct FPropertyChangedE
 			Modify();
 			Curve = CurveAsset->FloatCurve;
 		}
-		UpdateLUT();
 	}
+	UpdateLUT();
 }
 
 #endif
 
-bool UNiagaraDataInterfaceCurve::CopyTo(UNiagaraDataInterface* Destination) const 
+bool UNiagaraDataInterfaceCurve::CopyToInternal(UNiagaraDataInterface* Destination) const 
 {
-	if (!Super::CopyTo(Destination))
+	if (!Super::CopyToInternal(Destination))
 	{
 		return false;
 	}
 	CastChecked<UNiagaraDataInterfaceCurve>(Destination)->Curve = Curve;
 	CastChecked<UNiagaraDataInterfaceCurve>(Destination)->UpdateLUT();
-
+	ensure(CompareLUTS(CastChecked<UNiagaraDataInterfaceCurve>(Destination)->ShaderLUT));
 	return true;
 }
 
@@ -86,13 +87,45 @@ void UNiagaraDataInterfaceCurve::GetFunctions(TArray<FNiagaraFunctionSignature>&
 	OutFunctions.Add(Sig);
 }
 
+void UNiagaraDataInterfaceCurve::PostLoad()
+{
+	Super::PostLoad();
+	const int32 NiagaraVer = GetLinkerCustomVersion(FNiagaraCustomVersion::GUID);
+
+	if (NiagaraVer < FNiagaraCustomVersion::CurveLUTNowOnByDefault)
+	{
+		UpdateLUT();
+	}
+
+	TArray<float> OldLUT = ShaderLUT;
+	UpdateLUT();
+
+	if (!CompareLUTS(OldLUT))
+	{
+		UE_LOG(LogNiagara, Log, TEXT("PostLoad LUT generation is out of sync. Please investigate. %s"), *GetPathName());
+	}
+}
+
 
 void UNiagaraDataInterfaceCurve::UpdateLUT()
 {
 	ShaderLUT.Empty();
+	if (bAllowUnnormalizedLUT && Curve.GetNumKeys() > 0)
+	{
+		LUTMinTime = Curve.GetFirstKey().Time;
+		LUTMaxTime = Curve.GetLastKey().Time;
+		LUTInvTimeRange = 1.0f / (LUTMaxTime - LUTMinTime);
+	}
+	else
+	{
+		LUTMinTime = 0.0f;
+		LUTMaxTime = 1.0f;
+		LUTInvTimeRange = 1.0f;
+	}
+
 	for (uint32 i = 0; i < CurveLUTWidth; i++)
 	{
-		float X = i / (float)CurveLUTWidth;
+		float X = UnnormalizeTime(i / (float)CurveLUTWidth);
 		float C = Curve.Eval(X);
 		ShaderLUT.Add(C);
 	}
@@ -103,10 +136,10 @@ void UNiagaraDataInterfaceCurve::UpdateLUT()
 // the HLSL in the spirit of a static switch
 // TODO: need a way to identify each specific function here
 // 
-bool UNiagaraDataInterfaceCurve::GetFunctionHLSL(FString FunctionName, TArray<DIGPUBufferParamDescriptor> &Descriptors, FString &HLSLInterfaceID, FString &OutHLSL)
+bool UNiagaraDataInterfaceCurve::GetFunctionHLSL(const FName& DefinitionFunctionName, FString InstanceFunctionName, TArray<FDIGPUBufferParamDescriptor> &Descriptors, FString &HLSLInterfaceID, FString &OutHLSL)
 {
 	FString BufferName = Descriptors[0].BufferParamName;
-	OutHLSL += TEXT("void ") + FunctionName + TEXT("(in float In_X, out float Out_Value) \n{\n");
+	OutHLSL += TEXT("void ") + InstanceFunctionName + TEXT("(in float In_X, out float Out_Value) \n{\n");
 	OutHLSL += TEXT("\t Out_Value = ") + BufferName + TEXT("[(int)(clamp(In_X, 0.0, 1.0) *") + FString::FromInt(CurveLUTWidth) + TEXT(") ];");
 	OutHLSL += TEXT("\n}\n");
 	return true;
@@ -119,21 +152,21 @@ bool UNiagaraDataInterfaceCurve::GetFunctionHLSL(FString FunctionName, TArray<DI
 // 3. store buffer declaration hlsl in OutHLSL
 // multiple buffers can be defined at once here
 //
-void UNiagaraDataInterfaceCurve::GetBufferDefinitionHLSL(FString DataInterfaceID, TArray<DIGPUBufferParamDescriptor> &BufferDescriptors, FString &OutHLSL)
+void UNiagaraDataInterfaceCurve::GetBufferDefinitionHLSL(FString DataInterfaceID, TArray<FDIGPUBufferParamDescriptor> &BufferDescriptors, FString &OutHLSL)
 {
 	FString BufferName = "CurveLUT" + DataInterfaceID;
 	OutHLSL += TEXT("Buffer<float> ") + BufferName + TEXT(";\n");
 
-	BufferDescriptors.Add(DIGPUBufferParamDescriptor(BufferName, 0));		// add a descriptor for shader parameter binding
+	BufferDescriptors.Add(FDIGPUBufferParamDescriptor(BufferName, 0));		// add a descriptor for shader parameter binding
 
 }
 
 // called after translate, to setup buffers matching the buffer descriptors generated during hlsl translation
 // need to do this because the script used during translate is a clone, including its DIs
 //
-void UNiagaraDataInterfaceCurve::SetupBuffers(TArray<DIGPUBufferParamDescriptor> &BufferDescriptors)
+void UNiagaraDataInterfaceCurve::SetupBuffers(FDIBufferDescriptorStore &BufferDescriptors)
 {
-	for (DIGPUBufferParamDescriptor &Desc : BufferDescriptors)
+	for (FDIGPUBufferParamDescriptor &Desc : BufferDescriptors.Descriptors)
 	{
 		FNiagaraDataInterfaceBufferData BufferData(*Desc.BufferParamName);	// store off the data for later use
 		GPUBuffers.Add(BufferData);
@@ -152,7 +185,7 @@ TArray<FNiagaraDataInterfaceBufferData> &UNiagaraDataInterfaceCurve::GetBufferDa
 
 		FNiagaraDataInterfaceBufferData &GPUBuffer = GPUBuffers[0];
 		GPUBuffer.Buffer.Release();
-		GPUBuffer.Buffer.Initialize(sizeof(float), CurveLUTWidth * 4, EPixelFormat::PF_R32_FLOAT);	// always allocate for up to 64 data sets
+		GPUBuffer.Buffer.Initialize(sizeof(float), CurveLUTWidth * 4, EPixelFormat::PF_R32_FLOAT, BUF_Static);	// always allocate for up to 64 data sets
 		uint32 BufferSize = ShaderLUT.Num() * sizeof(float);
 		int32 *BufferData = static_cast<int32*>(RHILockVertexBuffer(GPUBuffer.Buffer.Buffer, 0, BufferSize, EResourceLockMode::RLM_WriteOnly));
 		FPlatformMemory::Memcpy(BufferData, ShaderLUT.GetData(), BufferSize);
@@ -170,7 +203,7 @@ FVMExternalFunction UNiagaraDataInterfaceCurve::GetVMExternalFunction(const FVME
 {
 	if (BindingInfo.Name == TEXT("SampleCurve") && BindingInfo.GetNumInputs() == 1 && BindingInfo.GetNumOutputs() == 1)
 	{
-		return TNDIParamBinder<0, float, NDI_FUNC_BINDER(UNiagaraDataInterfaceCurve, SampleCurve)>::Bind(this, BindingInfo, InstanceData);
+		return TCurveUseLUTBinder<TNDIParamBinder<0, float, NDI_FUNC_BINDER(UNiagaraDataInterfaceCurve, SampleCurve)>>::Bind(this, BindingInfo, InstanceData);
 	}
 	else
 	{
@@ -180,7 +213,21 @@ FVMExternalFunction UNiagaraDataInterfaceCurve::GetVMExternalFunction(const FVME
 	}
 }
 
-template<typename XParamType>
+template<>
+float UNiagaraDataInterfaceCurve::SampleCurveInternal<TIntegralConstant<bool, true>>(float X)
+{
+	float NormalizedX = NormalizeTime(X);
+	int32 AccessIdx = FMath::Clamp<int32>(FMath::TruncToInt(NormalizedX * CurveLUTWidthMinusOne), 0, CurveLUTWidthMinusOne) * CurveLUTNumElems;
+	return ShaderLUT[AccessIdx];
+}
+
+template<>
+float UNiagaraDataInterfaceCurve::SampleCurveInternal<TIntegralConstant<bool, false>>(float X)
+{
+	return Curve.Eval(X);
+}
+
+template<typename UseLUT, typename XParamType>
 void UNiagaraDataInterfaceCurve::SampleCurve(FVectorVMContext& Context)
 {
 	XParamType XParam(Context);
@@ -188,7 +235,7 @@ void UNiagaraDataInterfaceCurve::SampleCurve(FVectorVMContext& Context)
 
 	for (int32 i = 0; i < Context.NumInstances; ++i)
 	{
-		*OutSample.GetDest() = Curve.Eval(XParam.Get());
+		*OutSample.GetDest() = SampleCurveInternal<UseLUT>(XParam.Get());
 		XParam.Advance();
 		OutSample.Advance();
 	}
diff --git a/Engine/Plugins/FX/Niagara/Source/Niagara/Private/NiagaraDataInterfaceMeshCommon.cpp b/Engine/Plugins/FX/Niagara/Source/Niagara/Private/NiagaraDataInterfaceMeshCommon.cpp
new file mode 100644
--- /dev/null
+++ b/Engine/Plugins/FX/Niagara/Source/Niagara/Private/NiagaraDataInterfaceMeshCommon.cpp
@@ -0,0 +1,12 @@
+// Copyright 1998-2016 Epic Games, Inc. All Rights Reserved.
+
+#include "NiagaraDataInterfaceMeshCommon.h"
+#include "NiagaraEmitterInstance.h"
+#include "NiagaraComponent.h"
+#include "NiagaraSystemInstance.h"
+#include "Internationalization.h"
+#include "NiagaraScript.h"
+
+#define LOCTEXT_NAMESPACE "NiagaraDataInterfaceMeshCommon"
+
+#undef LOCTEXT_NAMESPACE
\ No newline at end of file
