Copying Orion-Staging @ 3748653 (Orion/Dev-General @ 3722124) to //UE4/Main

diff --git a/Engine/Plugins/FX/Niagara/Source/Niagara/Private/NiagaraDataInterfaceSkeletalMesh.cpp b/Engine/Plugins/FX/Niagara/Source/Niagara/Private/NiagaraDataInterfaceSkeletalMesh.cpp
new file mode 100644
--- /dev/null
+++ b/Engine/Plugins/FX/Niagara/Source/Niagara/Private/NiagaraDataInterfaceSkeletalMesh.cpp
@@ -0,0 +1,1451 @@
+// Copyright 1998-2016 Epic Games, Inc. All Rights Reserved.
+
+#include "NiagaraDataInterfaceSkeletalMesh.h"
+#include "NiagaraEmitterInstance.h"
+#include "NiagaraComponent.h"
+#include "NiagaraSystemInstance.h"
+#include "Internationalization.h"
+#include "NiagaraScript.h"
+#include "Animation/SkeletalMeshActor.h"
+#include "Components/SkeletalMeshComponent.h"
+#include "SkeletalMeshTypes.h"
+#include "NiagaraWorldManager.h"
+#include "ParallelFor.h"
+#include "NiagaraStats.h"
+
+#define LOCTEXT_NAMESPACE "NiagaraDataInterfaceSkeletalMesh"
+
+DECLARE_CYCLE_STAT(TEXT("PreSkin"), STAT_NiagaraSkel_PreSkin, STATGROUP_Niagara);
+
+//////////////////////////////////////////////////////////////////////////
+
+FSkeletalMeshSamplingRegionAreaWeightedSampler::FSkeletalMeshSamplingRegionAreaWeightedSampler()
+	: Owner(nullptr)
+{
+}
+
+void FSkeletalMeshSamplingRegionAreaWeightedSampler::Init(FNDISkeletalMesh_InstanceData* InOwner)
+{
+	Owner = InOwner;
+	Initialize();
+}
+
+float FSkeletalMeshSamplingRegionAreaWeightedSampler::GetWeights(TArray<float>& OutWeights)
+{
+	check(Owner && Owner->Mesh);
+	check(Owner->Mesh->LODInfo.IsValidIndex(Owner->GetLODIndex()));
+
+	float Total = 0.0f;
+	int32 NumUsedRegions = Owner->SamplingRegionIndices.Num();
+	if (NumUsedRegions <= 1)
+	{
+		//Use 0 or 1 Sampling region. Only need additional area weighting between regions if we're sampling from multiple.
+		return 0.0f;
+	}
+	
+	const FSkeletalMeshSamplingInfo& SamplingInfo = Owner->Mesh->GetSamplingInfo();
+	OutWeights.Empty(NumUsedRegions);
+	for (int32 i = 0; i < NumUsedRegions; ++i)
+	{
+		int32 RegionIdx = Owner->SamplingRegionIndices[i];
+		const FSkeletalMeshSamplingRegion& Region = SamplingInfo.GetRegion(RegionIdx);
+		float T = SamplingInfo.GetRegionBuiltData(RegionIdx).AreaWeightedSampler.GetTotalWeight();
+		OutWeights.Add(T);
+		Total += T;
+	}
+	return Total;
+}
+
+//////////////////////////////////////////////////////////////////////////
+FSkeletalMeshSkinningDataHandle::FSkeletalMeshSkinningDataHandle()
+	: SkinningData(nullptr)
+{
+}
+
+FSkeletalMeshSkinningDataHandle::FSkeletalMeshSkinningDataHandle(FSkeletalMeshSkinningDataUsage InUsage, TSharedPtr<FSkeletalMeshSkinningData> InSkinningData)
+	: Usage(InUsage)
+	, SkinningData(InSkinningData)
+{
+	if (FSkeletalMeshSkinningData* SkinData = SkinningData.Get())
+	{
+		SkinData->RegisterUser(Usage);
+	}
+}
+
+FSkeletalMeshSkinningDataHandle::~FSkeletalMeshSkinningDataHandle()
+{
+	if (FSkeletalMeshSkinningData* SkinData = SkinningData.Get())
+	{
+		SkinData->UnregisterUser(Usage);
+	}
+}
+
+//////////////////////////////////////////////////////////////////////////
+
+void FSkeletalMeshSkinningData::RegisterUser(FSkeletalMeshSkinningDataUsage Usage)
+{
+	FScopeLock Lock(&CriticalSection);
+	USkeletalMeshComponent* SkelComp = MeshComp.Get();
+
+	int32 LODIndex = Usage.GetLODIndex();
+	check(LODIndex != INDEX_NONE);
+	check(SkelComp);
+
+	LODData.SetNum(LODIndex + 1);
+
+	if (Usage.NeedBoneMatrices())
+	{
+		++BoneMatrixUsers;
+	}
+
+	FLODData& LOD = LODData[LODIndex];
+	if (Usage.NeedPreSkinnedVerts())
+	{
+		++LOD.PreSkinnedVertsUsers;
+	}
+
+	if (Usage.NeedsDataImmediately())
+	{
+		check(IsInGameThread());
+		if (CurrBoneRefToLocals().Num() == 0)
+		{
+			SkelComp->CacheRefToLocalMatrices(CurrBoneRefToLocals());
+		}
+		
+		//Prime the prev matrices if they're missing.
+		if (PrevBoneRefToLocals().Num() != CurrBoneRefToLocals().Num())
+		{
+			PrevBoneRefToLocals() = CurrBoneRefToLocals();
+		}
+
+		if (Usage.NeedPreSkinnedVerts() && CurrSkinnedPositions(LODIndex).Num() == 0)
+		{
+			FSkeletalMeshLODRenderData& SkelMeshLODData = SkelComp->SkeletalMesh->GetResourceForRendering()->LODRenderData[LODIndex];
+			FSkinWeightVertexBuffer* SkinWeightBuffer = SkelComp->GetSkinWeightBuffer(LODIndex);
+			USkeletalMeshComponent::ComputeSkinnedPositions(SkelComp, CurrSkinnedPositions(LODIndex), CurrBoneRefToLocals(), SkelMeshLODData, *SkinWeightBuffer);
+
+			//Prime the previous positions if they're missing
+			if (PrevSkinnedPositions(LODIndex).Num() != CurrSkinnedPositions(LODIndex).Num())
+			{
+				PrevSkinnedPositions(LODIndex) = CurrSkinnedPositions(LODIndex);
+			}
+		}
+	}
+}
+
+void FSkeletalMeshSkinningData::UnregisterUser(FSkeletalMeshSkinningDataUsage Usage)
+{
+	FScopeLock Lock(&CriticalSection);
+	check(LODData.IsValidIndex(Usage.GetLODIndex()));
+
+	if (Usage.NeedBoneMatrices())
+	{
+		--BoneMatrixUsers;
+	}
+
+	FLODData& LOD = LODData[Usage.GetLODIndex()];
+	if (Usage.NeedPreSkinnedVerts())
+	{
+		--LOD.PreSkinnedVertsUsers;
+	}
+}
+
+bool FSkeletalMeshSkinningData::IsUsed()const
+{
+	if (BoneMatrixUsers > 0)
+	{
+		return true;
+	}
+
+	for (const FLODData& LOD : LODData)
+	{
+		if (LOD.PreSkinnedVertsUsers > 0)
+		{
+			return true;
+		}
+	}
+
+	return false;
+}
+
+bool FSkeletalMeshSkinningData::Tick(float InDeltaSeconds)
+{
+	USkeletalMeshComponent* SkelComp = MeshComp.Get();
+	check(SkelComp);
+	DeltaSeconds = InDeltaSeconds;
+	CurrIndex ^= 1;
+
+	if (BoneMatrixUsers > 0)
+	{
+		SkelComp->CacheRefToLocalMatrices(CurrBoneRefToLocals());
+	}
+
+	//Prime the prev matrices if they're missing.
+	if (PrevBoneRefToLocals().Num() != CurrBoneRefToLocals().Num())
+	{
+		PrevBoneRefToLocals() = CurrBoneRefToLocals();
+	}
+
+	for (int32 LODIndex = 0; LODIndex < LODData.Num(); ++LODIndex)
+	{
+		FLODData& LOD = LODData[LODIndex];
+		if (LOD.PreSkinnedVertsUsers > 0)
+		{
+			//TODO: If we pass the sections in the usage too, we can probably skin a minimal set of verts just for the used regions.
+			FSkeletalMeshLODRenderData& SkelMeshLODData = SkelComp->SkeletalMesh->GetResourceForRendering()->LODRenderData[LODIndex];
+			FSkinWeightVertexBuffer* SkinWeightBuffer = SkelComp->GetSkinWeightBuffer(LODIndex);
+			USkeletalMeshComponent::ComputeSkinnedPositions(SkelComp, CurrSkinnedPositions(LODIndex), CurrBoneRefToLocals(), SkelMeshLODData, *SkinWeightBuffer);
+
+			//Prime the previous positions if they're missing
+			if (PrevSkinnedPositions(LODIndex).Num() != CurrSkinnedPositions(LODIndex).Num())
+			{
+				PrevSkinnedPositions(LODIndex) = CurrSkinnedPositions(LODIndex);
+			}
+		}
+	}
+	
+	return true;
+}
+
+//////////////////////////////////////////////////////////////////////////
+
+FSkeletalMeshSkinningDataHandle FNDI_SkeletalMesh_GeneratedData::GetCachedSkinningData(TWeakObjectPtr<USkeletalMeshComponent>& InComponent, FSkeletalMeshSkinningDataUsage Usage)
+{
+	FScopeLock Lock(&CriticalSection);
+	
+	USkeletalMeshComponent* Component = InComponent.Get();
+	check(Component);
+	TSharedPtr<FSkeletalMeshSkinningData> SkinningData = nullptr;
+
+	if (TSharedPtr<FSkeletalMeshSkinningData>* Existing = CachedSkinningData.Find(Component))
+	{
+		check(Existing->IsValid());//We shouldn't be able to have an invalid ptr here.
+		SkinningData = *Existing;
+	}
+	else
+	{
+		SkinningData = MakeShared<FSkeletalMeshSkinningData>(InComponent);
+		CachedSkinningData.Add(Component) = SkinningData;
+	}
+
+	return FSkeletalMeshSkinningDataHandle(Usage, SkinningData);
+}
+
+void FNDI_SkeletalMesh_GeneratedData::TickGeneratedData(float DeltaSeconds)
+{
+	check(IsInGameThread());
+	SCOPE_CYCLE_COUNTER(STAT_NiagaraSkel_PreSkin);
+
+	//Tick skinning data.
+	{
+		TArray<TWeakObjectPtr<USkeletalMeshComponent>, TInlineAllocator<64>> ToRemove;
+		TArray<FSkeletalMeshSkinningData*> ToTick;
+		ToTick.Reserve(CachedSkinningData.Num());
+		for (TPair<TWeakObjectPtr<USkeletalMeshComponent>, TSharedPtr<FSkeletalMeshSkinningData>>& Pair : CachedSkinningData)
+		{
+			TSharedPtr<FSkeletalMeshSkinningData>& Ptr = Pair.Value;
+			FSkeletalMeshSkinningData* SkinData = Ptr.Get();
+			USkeletalMeshComponent* Component = Pair.Key.Get();
+			check(SkinData);
+			if (Ptr.IsUnique() || !Component || !Ptr->IsUsed())
+			{
+				ToRemove.Add(Pair.Key);//Remove unused skin data or for those with GCd components as we go.
+			}
+			else
+			{
+				ToTick.Add(SkinData);
+			}
+		}
+
+		for (TWeakObjectPtr<USkeletalMeshComponent> Key : ToRemove)
+		{
+			CachedSkinningData.Remove(Key);
+		}
+
+		ParallelFor(ToTick.Num(), [&](int32 Index)
+		{
+			ToTick[Index]->Tick(DeltaSeconds);
+		});
+	}
+}
+
+//////////////////////////////////////////////////////////////////////////
+//FNDISkeletalMesh_InstanceData
+
+bool FNDISkeletalMesh_InstanceData::Init(UNiagaraDataInterfaceSkeletalMesh* Interface, FNiagaraSystemInstance* SystemInstance)
+{
+	check(SystemInstance);
+	USkeletalMesh* PrevMesh = Mesh;
+	Component = nullptr;
+	Mesh = nullptr;
+	Transform = FMatrix::Identity;
+	TransformInverseTransposed = FMatrix::Identity;
+	PrevTransform = FMatrix::Identity;
+	PrevTransformInverseTransposed = FMatrix::Identity;
+	DeltaSeconds = 0.0f;
+
+	USkeletalMeshComponent* NewSkelComp = nullptr;
+	if (Interface->Source)
+	{
+		ASkeletalMeshActor* MeshActor = Cast<ASkeletalMeshActor>(Interface->Source);
+		USkeletalMeshComponent* SourceComp = nullptr;
+		if (MeshActor != nullptr)
+		{
+			SourceComp = MeshActor->GetSkeletalMeshComponent();
+		}
+		else
+		{
+			SourceComp = Interface->Source->FindComponentByClass<USkeletalMeshComponent>();
+		}
+
+		if (SourceComp)
+		{
+			Mesh = SourceComp->SkeletalMesh;
+			NewSkelComp = SourceComp;
+		}
+		else
+		{
+			Component = Interface->Source->GetRootComponent();
+		}
+	}
+	else
+	{
+		if (UNiagaraComponent* SimComp = SystemInstance->GetComponent())
+		{
+			if (USkeletalMeshComponent* ParentComp = Cast<USkeletalMeshComponent>(SimComp->GetAttachParent()))
+			{
+				NewSkelComp = ParentComp;
+				Mesh = ParentComp->SkeletalMesh;
+			}
+			else if (USkeletalMeshComponent* OuterComp = SimComp->GetTypedOuter<USkeletalMeshComponent>())
+			{
+				NewSkelComp = OuterComp;
+				Mesh = OuterComp->SkeletalMesh;
+			}
+			else if (AActor* Owner = SimComp->GetAttachmentRootActor())
+			{
+				TArray<UActorComponent*> SourceComps = Owner->GetComponentsByClass(USkeletalMeshComponent::StaticClass());
+				for (UActorComponent* ActorComp : SourceComps)
+				{
+					USkeletalMeshComponent* SourceComp = Cast<USkeletalMeshComponent>(ActorComp);
+					if (SourceComp)
+					{
+						USkeletalMesh* PossibleMesh = SourceComp->SkeletalMesh;
+						if (PossibleMesh != nullptr/* && PossibleMesh->bAllowCPUAccess*/)
+						{
+							Mesh = PossibleMesh;
+							NewSkelComp = SourceComp;
+							break;
+						}
+					}
+				}
+			}
+
+			if (!Component.IsValid())
+			{
+				Component = SimComp;
+			}
+		}
+	}
+
+	if (NewSkelComp)
+	{
+		Component = NewSkelComp;
+	}
+
+	check(Component.IsValid());
+
+	if (!Mesh && Interface->DefaultMesh)
+	{
+		Mesh = Interface->DefaultMesh;
+	}
+
+	if (Component.IsValid() && Mesh)
+	{
+		PrevTransform = Transform;
+		PrevTransformInverseTransposed = TransformInverseTransposed;
+		Transform = Component->GetComponentToWorld().ToMatrixWithScale();
+		TransformInverseTransposed = Transform.InverseFast().GetTransposed();
+	}
+
+	if (!Mesh)
+	{
+		UE_LOG(LogNiagara, Log, TEXT("SkeletalMesh data interface has no valid mesh. Failed InitPerInstanceData - %s"), *Interface->GetFullName());
+		return false;
+	}
+
+#if WITH_EDITOR
+	Mesh->GetOnMeshChanged().AddUObject(SystemInstance->GetComponent(), &UNiagaraComponent::ReinitializeSystem);
+#endif
+
+
+// 	if (!Mesh->bAllowCPUAccess)
+// 	{
+// 		UE_LOG(LogNiagara, Log, TEXT("SkeletalMesh data interface using a mesh that does not allow CPU access. Failed InitPerInstanceData - Mesh: %s"), *Mesh->GetFullName());
+// 		return false;
+// 	}
+
+	if (!Component.IsValid())
+	{
+		UE_LOG(LogNiagara, Log, TEXT("SkeletalMesh data interface has no valid component. Failed InitPerInstanceData - %s"), *Interface->GetFullName());
+		return false;
+	}
+
+	//Setup where to spawn from
+	SamplingRegionIndices.Empty();
+	bool bAllRegionsAreAreaWeighting = true;
+	const FSkeletalMeshSamplingInfo& SamplingInfo = Mesh->GetSamplingInfo();
+	int32 LODIndex = INDEX_NONE;
+	if (Interface->SamplingRegions.Num() == 0)
+	{
+		LODIndex = Interface->WholeMeshLOD;
+		//If we have no regions, sample the whole mesh at the specified LOD.
+		if (LODIndex == INDEX_NONE)
+		{
+			LODIndex = Mesh->LODInfo.Num() - 1;
+		}
+		else
+		{
+			LODIndex = FMath::Clamp(Interface->WholeMeshLOD, 0, Mesh->LODInfo.Num() - 1);
+		}
+
+		if (!Mesh->LODInfo[LODIndex].bAllowCPUAccess)
+		{
+			UE_LOG(LogNiagara, Warning, TEXT("Skeletal Mesh Data Interface is trying to spawn from a whole mesh that does not allow CPU Access.\nInterface: %s\nMesh: %s\nLOD: %d"),
+				*Interface->GetFullName(),
+				*Mesh->GetFullName(),
+				LODIndex);
+
+			return false;
+		}
+	}
+	else
+	{
+		//Sampling from regions. Gather the indices of the regions we'll sample from.
+		for (FName RegionName : Interface->SamplingRegions)
+		{
+			int32 RegionIdx = SamplingInfo.IndexOfRegion(RegionName);
+			if (RegionIdx != INDEX_NONE)
+			{
+				const FSkeletalMeshSamplingRegion& Region = SamplingInfo.GetRegion(RegionIdx);
+				const FSkeletalMeshSamplingRegionBuiltData& RegionBuiltData = SamplingInfo.GetRegionBuiltData(RegionIdx);
+				if (LODIndex == INDEX_NONE)
+				{
+					LODIndex = Region.LODIndex;
+				}
+				else
+				{
+					//ensure we don't try to use two regions from different LODs.
+					if (LODIndex != Region.LODIndex)
+					{
+						UE_LOG(LogNiagara, Warning, TEXT("Skeletal Mesh Data Interface is trying to use regions on different LODs of the mesh. This is currently unsupported.\nInterface: %s\nMesh: %s\nRegion: %s"),
+							*Interface->GetFullName(),
+							*Mesh->GetFullName(),
+							*RegionName.ToString());
+
+						return false;
+					}
+				}
+
+				if (RegionBuiltData.TriangleIndices.Num() > 0)
+				{
+					SamplingRegionIndices.Add(RegionIdx);
+					bAllRegionsAreAreaWeighting &= Region.bSupportUniformlyDistributedSampling;
+				}
+				else
+				{
+					UE_LOG(LogNiagara, Warning, TEXT("Skeletal Mesh Data Interface is trying to use a region with no associated triangles.\nLOD: %d\nInterface: %s\nMesh: %s\nRegion: %s"),
+						LODIndex,
+						*Interface->GetFullName(),
+						*Mesh->GetFullName(),
+						*RegionName.ToString());
+
+					return false;
+				}
+			}
+			else
+			{
+				UE_LOG(LogNiagara, Warning, TEXT("Skeletal Mesh Data Interface is trying to use a region on a mesh that does not provide this region.\nInterface: %s\nMesh: %s\nRegion: %s"),
+					*Interface->GetFullName(),
+					*Mesh->GetFullName(),
+					*RegionName.ToString());
+
+				return false;
+			}
+		}
+	}
+		
+	//Grab a handle to the skinning data if we have a component to skin.
+	ENDISkeletalMesh_SkinningMode SkinningMode = Interface->SkinningMode;
+	FSkeletalMeshSkinningDataUsage Usage(
+		LODIndex,
+		SkinningMode == ENDISkeletalMesh_SkinningMode::SkinOnTheFly || SkinningMode == ENDISkeletalMesh_SkinningMode::PreSkin,
+		SkinningMode == ENDISkeletalMesh_SkinningMode::PreSkin,
+		SystemInstance->IsSolo());
+
+	if (NewSkelComp)
+	{
+		SkinningMode = Interface->SkinningMode;
+		TWeakObjectPtr<USkeletalMeshComponent> SkelWeakCompPtr = NewSkelComp;
+		FNDI_SkeletalMesh_GeneratedData& GeneratedData = SystemInstance->GetWorldManager()->GetSkeletalMeshGeneratedData();
+		SkinningData = GeneratedData.GetCachedSkinningData(SkelWeakCompPtr, Usage);
+	}
+	else
+	{
+		SkinningData = FSkeletalMeshSkinningDataHandle(Usage, nullptr);
+	}
+
+	//Init area weighting sampler for Sampling regions.
+	if (SamplingRegionIndices.Num() > 1 && bAllRegionsAreAreaWeighting)
+	{
+		//We are sampling from multiple area weighted regions so setup the inter-region weighting sampler.
+		SamplingRegionAreaWeightedSampler.Init(this);
+	}
+
+	FSkinWeightVertexBuffer* SkinWeightBuffer = nullptr;
+	FSkeletalMeshLODRenderData& LODData = GetLODRenderDataAndSkinWeights(SkinWeightBuffer);
+
+	//Check for the validity of the Mesh's cpu data.
+	bool bMeshCPUDataValid =
+		LODData.GetNumVertices() > 0 &&
+		LODData.StaticVertexBuffers.PositionVertexBuffer.GetNumVertices() &&
+		SkinWeightBuffer &&
+		SkinWeightBuffer->GetNumVertices() > 0 &&
+		LODData.MultiSizeIndexContainer.IsIndexBufferValid() &&
+		LODData.MultiSizeIndexContainer.GetIndexBuffer() &&
+		LODData.MultiSizeIndexContainer.GetIndexBuffer()->Num() > 0;
+
+	if (!bMeshCPUDataValid)
+	{
+		UE_LOG(LogNiagara, Warning, TEXT("Skeletal Mesh Data Interface is trying to sample from a mesh with missing CPU vertex or index data.\nInterface: %s\nMesh: %s\nLOD: %d"),
+			*Interface->GetFullName(),
+			*Mesh->GetFullName(),
+			LODIndex);
+
+		return false;
+	}
+
+	return true;
+}
+
+bool FNDISkeletalMesh_InstanceData::ResetRequired(UNiagaraDataInterfaceSkeletalMesh* Interface)const
+{
+	USceneComponent* Comp = Component.Get();
+	if (!Comp)
+	{
+		//The component we were bound to is no longer valid so we have to trigger a reset.
+		return true;
+	}
+	
+	if (USkeletalMeshComponent* SkelComp = Cast<USkeletalMeshComponent>(Comp))
+	{
+		if (!SkelComp->SkeletalMesh)
+		{
+			return true;
+		}
+	}
+	else
+	{
+		if (!Interface->DefaultMesh)
+		{
+			return true;
+		}
+	}
+	
+	return false;
+}
+
+bool FNDISkeletalMesh_InstanceData::Tick(UNiagaraDataInterfaceSkeletalMesh* Interface, FNiagaraSystemInstance* SystemInstance, float InDeltaSeconds)
+{
+	if (ResetRequired(Interface))
+	{
+		return true;
+	}
+	else
+	{
+		DeltaSeconds = InDeltaSeconds;
+		if (Component.IsValid() && Mesh)
+		{
+			PrevTransform = Transform;
+			PrevTransformInverseTransposed = TransformInverseTransposed;
+			Transform = Component->GetComponentToWorld().ToMatrixWithScale();
+			TransformInverseTransposed = Transform.InverseFast().GetTransposed();
+		}
+		else
+		{
+			PrevTransform = FMatrix::Identity;
+			PrevTransformInverseTransposed = FMatrix::Identity;
+			Transform = FMatrix::Identity;
+			TransformInverseTransposed = FMatrix::Identity;
+		}
+		return false;
+	}
+}
+
+//Instance Data END
+//////////////////////////////////////////////////////////////////////////
+
+//////////////////////////////////////////////////////////////////////////
+// Helper classes for reducing duplicate code when accessing vertex positions. 
+
+struct FSkeletalMeshAccessorHelper
+{
+	FSkeletalMeshAccessorHelper()
+		: Comp(nullptr)
+		, Mesh(nullptr)
+		, LODData(nullptr)
+		, SkinWeightBuffer(nullptr)
+		, IndexBuffer(nullptr)
+		, SamplingRegion(nullptr)
+		, SamplingRegionBuiltData(nullptr)
+		, SkinningData(nullptr)
+	{
+
+	}
+
+	template<typename FilterMode, typename AreaWeightingMode>
+	FORCEINLINE void Init(FNDISkeletalMesh_InstanceData* InstData);
+
+	USkeletalMeshComponent* Comp;
+	USkeletalMesh* Mesh;
+	FSkeletalMeshLODRenderData* LODData;
+	FSkinWeightVertexBuffer* SkinWeightBuffer;
+	FRawStaticIndexBuffer16or32Interface* IndexBuffer;
+	const FSkeletalMeshSamplingRegion* SamplingRegion;
+	const FSkeletalMeshSamplingRegionBuiltData* SamplingRegionBuiltData;
+	FSkeletalMeshSkinningData* SkinningData;
+	FSkeletalMeshSkinningDataUsage Usage;
+};
+
+template<typename FilterMode, typename AreaWeightingMode>
+void FSkeletalMeshAccessorHelper::Init(FNDISkeletalMesh_InstanceData* InstData)
+{
+	Mesh = InstData->Mesh;
+	LODData = &InstData->GetLODRenderDataAndSkinWeights(SkinWeightBuffer);
+	IndexBuffer = LODData->MultiSizeIndexContainer.GetIndexBuffer();
+	SkinningData = InstData->SkinningData.SkinningData.Get();
+	Usage = InstData->SkinningData.Usage;
+}
+
+template<>
+void FSkeletalMeshAccessorHelper::Init<
+	TIntegralConstant<ENDISkeletalMesh_FilterMode, ENDISkeletalMesh_FilterMode::SingleRegion>,
+	TIntegralConstant<ENDISkelMesh_AreaWeightingMode, ENDISkelMesh_AreaWeightingMode::None>>
+	(FNDISkeletalMesh_InstanceData* InstData)
+{
+	Mesh = InstData->Mesh;
+	LODData = &InstData->GetLODRenderDataAndSkinWeights(SkinWeightBuffer);
+	IndexBuffer = LODData->MultiSizeIndexContainer.GetIndexBuffer();
+	SkinningData = InstData->SkinningData.SkinningData.Get();
+	Usage = InstData->SkinningData.Usage;
+
+	const FSkeletalMeshSamplingInfo& SamplingInfo = InstData->Mesh->GetSamplingInfo();
+	SamplingRegion = &SamplingInfo.GetRegion(InstData->SamplingRegionIndices[0]);
+	SamplingRegionBuiltData = &SamplingInfo.GetRegionBuiltData(InstData->SamplingRegionIndices[0]);
+}
+
+template<>
+void FSkeletalMeshAccessorHelper::Init<
+	TIntegralConstant<ENDISkeletalMesh_FilterMode, ENDISkeletalMesh_FilterMode::SingleRegion>,
+	TIntegralConstant<ENDISkelMesh_AreaWeightingMode, ENDISkelMesh_AreaWeightingMode::AreaWeighted>>
+	(FNDISkeletalMesh_InstanceData* InstData)
+{
+	Mesh = InstData->Mesh;
+	LODData = &InstData->GetLODRenderDataAndSkinWeights(SkinWeightBuffer);
+	IndexBuffer = LODData->MultiSizeIndexContainer.GetIndexBuffer();
+	SkinningData = InstData->SkinningData.SkinningData.Get();
+	Usage = InstData->SkinningData.Usage;
+
+	const FSkeletalMeshSamplingInfo& SamplingInfo = InstData->Mesh->GetSamplingInfo();
+	SamplingRegion = &SamplingInfo.GetRegion(InstData->SamplingRegionIndices[0]);
+	SamplingRegionBuiltData = &SamplingInfo.GetRegionBuiltData(InstData->SamplingRegionIndices[0]);
+}
+
+//////////////////////////////////////////////////////////////////////////
+
+template<typename SkinningMode>
+struct FSkinnedPositionAccessorHelper
+{
+	FORCEINLINE void GetSkinnedTrianglePositions(FSkeletalMeshAccessorHelper& Accessor,
+		int32 Tri, FVector& OutPos0, FVector& OutPos1, FVector& OutPos2, FVector& OutPrev0, FVector& OutPrev1, FVector& OutPrev2)
+	{
+		checkf(false, TEXT("Must provide a specialization for this template type"));
+	}
+
+	FORCEINLINE void GetSkinnedTrianglePositions(FSkeletalMeshAccessorHelper& Accessor,
+		int32 Tri, FVector& OutPos0, FVector& OutPos1, FVector& OutPos2)
+	{
+		checkf(false, TEXT("Must provide a specialization for this template type"));
+	}
+};
+
+template<>
+struct FSkinnedPositionAccessorHelper<TIntegralConstant<ENDISkeletalMesh_SkinningMode, ENDISkeletalMesh_SkinningMode::None>>
+{
+	FORCEINLINE void GetSkinnedTrianglePositions(FSkeletalMeshAccessorHelper& Accessor,
+		int32 Tri, FVector& OutPos0, FVector& OutPos1, FVector& OutPos2)
+	{
+		int32 Idx0 = Accessor.IndexBuffer->Get(Tri);
+		int32 Idx1 = Accessor.IndexBuffer->Get(Tri + 1);
+		int32 Idx2 = Accessor.IndexBuffer->Get(Tri + 2);
+		
+		OutPos0 = GetSkeletalMeshRefVertLocation(Accessor.Mesh, *Accessor.LODData, *Accessor.SkinWeightBuffer, Idx0);
+		OutPos1 = GetSkeletalMeshRefVertLocation(Accessor.Mesh, *Accessor.LODData, *Accessor.SkinWeightBuffer, Idx1);
+		OutPos2 = GetSkeletalMeshRefVertLocation(Accessor.Mesh, *Accessor.LODData, *Accessor.SkinWeightBuffer, Idx2);
+	}
+
+	FORCEINLINE void GetSkinnedTrianglePositions(FSkeletalMeshAccessorHelper& Accessor,
+		int32 Tri, FVector& OutPos0, FVector& OutPos1, FVector& OutPos2, FVector& OutPrev0, FVector& OutPrev1, FVector& OutPrev2)
+	{
+		int32 Idx0 = Accessor.IndexBuffer->Get(Tri);
+		int32 Idx1 = Accessor.IndexBuffer->Get(Tri + 1);
+		int32 Idx2 = Accessor.IndexBuffer->Get(Tri + 2);
+		OutPos0 = GetSkeletalMeshRefVertLocation(Accessor.Mesh, *Accessor.LODData, *Accessor.SkinWeightBuffer, Idx0);
+		OutPos1 = GetSkeletalMeshRefVertLocation(Accessor.Mesh, *Accessor.LODData, *Accessor.SkinWeightBuffer, Idx1);
+		OutPos2 = GetSkeletalMeshRefVertLocation(Accessor.Mesh, *Accessor.LODData, *Accessor.SkinWeightBuffer, Idx2);
+
+		OutPrev0 = OutPos0;
+		OutPrev1 = OutPos1;
+		OutPrev2 = OutPos2;
+	}
+};
+
+template<>
+struct FSkinnedPositionAccessorHelper<TIntegralConstant<ENDISkeletalMesh_SkinningMode, ENDISkeletalMesh_SkinningMode::SkinOnTheFly>>
+{
+	FORCEINLINE void GetSkinnedTrianglePositions(FSkeletalMeshAccessorHelper& Accessor,
+		int32 Tri, FVector& OutPos0, FVector& OutPos1, FVector& OutPos2)
+	{
+		int32 Idx0 = Accessor.IndexBuffer->Get(Tri);
+		int32 Idx1 = Accessor.IndexBuffer->Get(Tri + 1);
+		int32 Idx2 = Accessor.IndexBuffer->Get(Tri + 2);
+		OutPos0 = USkeletalMeshComponent::GetSkinnedVertexPosition(Accessor.Comp, Idx0, *Accessor.LODData, *Accessor.SkinWeightBuffer, Accessor.SkinningData->CurrBoneRefToLocals());
+		OutPos1 = USkeletalMeshComponent::GetSkinnedVertexPosition(Accessor.Comp, Idx1, *Accessor.LODData, *Accessor.SkinWeightBuffer, Accessor.SkinningData->CurrBoneRefToLocals());
+		OutPos2 = USkeletalMeshComponent::GetSkinnedVertexPosition(Accessor.Comp, Idx2, *Accessor.LODData, *Accessor.SkinWeightBuffer, Accessor.SkinningData->CurrBoneRefToLocals());
+	}
+
+	FORCEINLINE void GetSkinnedTrianglePositions(FSkeletalMeshAccessorHelper& Accessor,
+		int32 Tri, FVector& OutPos0, FVector& OutPos1, FVector& OutPos2, FVector& OutPrev0, FVector& OutPrev1, FVector& OutPrev2)
+	{
+		int32 Idx0 = Accessor.IndexBuffer->Get(Tri);
+		int32 Idx1 = Accessor.IndexBuffer->Get(Tri + 1);
+		int32 Idx2 = Accessor.IndexBuffer->Get(Tri + 2);
+		OutPos0 = USkeletalMeshComponent::GetSkinnedVertexPosition(Accessor.Comp, Idx0, *Accessor.LODData, *Accessor.SkinWeightBuffer, Accessor.SkinningData->CurrBoneRefToLocals());
+		OutPos1 = USkeletalMeshComponent::GetSkinnedVertexPosition(Accessor.Comp, Idx1, *Accessor.LODData, *Accessor.SkinWeightBuffer, Accessor.SkinningData->CurrBoneRefToLocals());
+		OutPos2 = USkeletalMeshComponent::GetSkinnedVertexPosition(Accessor.Comp, Idx2, *Accessor.LODData, *Accessor.SkinWeightBuffer, Accessor.SkinningData->CurrBoneRefToLocals());
+		OutPrev0 = USkeletalMeshComponent::GetSkinnedVertexPosition(Accessor.Comp, Idx0, *Accessor.LODData, *Accessor.SkinWeightBuffer, Accessor.SkinningData->PrevBoneRefToLocals());
+		OutPrev1 = USkeletalMeshComponent::GetSkinnedVertexPosition(Accessor.Comp, Idx1, *Accessor.LODData, *Accessor.SkinWeightBuffer, Accessor.SkinningData->PrevBoneRefToLocals());
+		OutPrev2 = USkeletalMeshComponent::GetSkinnedVertexPosition(Accessor.Comp, Idx2, *Accessor.LODData, *Accessor.SkinWeightBuffer, Accessor.SkinningData->PrevBoneRefToLocals());
+	}
+};
+
+template<>
+struct FSkinnedPositionAccessorHelper<TIntegralConstant<ENDISkeletalMesh_SkinningMode, ENDISkeletalMesh_SkinningMode::PreSkin>>
+{
+	FORCEINLINE void GetSkinnedTrianglePositions(FSkeletalMeshAccessorHelper& Accessor,
+		int32 Tri, FVector& OutPos0, FVector& OutPos1, FVector& OutPos2)
+	{
+		int32 Idx0 = Accessor.IndexBuffer->Get(Tri);
+		int32 Idx1 = Accessor.IndexBuffer->Get(Tri + 1);
+		int32 Idx2 = Accessor.IndexBuffer->Get(Tri + 2);
+		OutPos0 = Accessor.SkinningData->GetPosition(Accessor.Usage.GetLODIndex(), Idx0);
+		OutPos1 = Accessor.SkinningData->GetPosition(Accessor.Usage.GetLODIndex(), Idx1);
+		OutPos2 = Accessor.SkinningData->GetPosition(Accessor.Usage.GetLODIndex(), Idx2);
+	}
+
+	FORCEINLINE void GetSkinnedTrianglePositions(FSkeletalMeshAccessorHelper& Accessor,
+		int32 Tri, FVector& OutPos0, FVector& OutPos1, FVector& OutPos2, FVector& OutPrev0, FVector& OutPrev1, FVector& OutPrev2)
+	{
+		int32 Idx0 = Accessor.IndexBuffer->Get(Tri);
+		int32 Idx1 = Accessor.IndexBuffer->Get(Tri + 1);
+		int32 Idx2 = Accessor.IndexBuffer->Get(Tri + 2);
+		OutPos0 = Accessor.SkinningData->GetPosition(Accessor.Usage.GetLODIndex(), Idx0);
+		OutPos1 = Accessor.SkinningData->GetPosition(Accessor.Usage.GetLODIndex(), Idx1);
+		OutPos2 = Accessor.SkinningData->GetPosition(Accessor.Usage.GetLODIndex(), Idx2);
+		OutPrev0 = Accessor.SkinningData->GetPreviousPosition(Accessor.Usage.GetLODIndex(), Idx0);
+		OutPrev1 = Accessor.SkinningData->GetPreviousPosition(Accessor.Usage.GetLODIndex(), Idx1);
+		OutPrev2 = Accessor.SkinningData->GetPreviousPosition(Accessor.Usage.GetLODIndex(), Idx2);
+	}
+};
+
+//////////////////////////////////////////////////////////////////////////
+// Helper for accessing misc vertex data
+template<bool bUseFullPrecisionUVs>
+struct FSkelMeshVertexAccessor
+{
+	FORCEINLINE FVector2D GetVertexUV(FSkeletalMeshLODRenderData& LODData, int32 VertexIdx, int32 UVChannel)const
+	{
+		if (bUseFullPrecisionUVs)
+		{
+			return LODData.StaticVertexBuffers.StaticMeshVertexBuffer.GetVertexUV_Typed<EStaticMeshVertexUVType::HighPrecision>(VertexIdx, UVChannel);
+		}
+		else
+		{
+			return LODData.StaticVertexBuffers.StaticMeshVertexBuffer.GetVertexUV_Typed<EStaticMeshVertexUVType::Default>(VertexIdx, UVChannel);
+		}
+	}
+
+	FORCEINLINE FLinearColor GetVertexColor(FSkeletalMeshLODRenderData& LODData, int32 VertexIdx)const
+	{
+		return LODData.StaticVertexBuffers.ColorVertexBuffer.VertexColor(VertexIdx);
+	}
+};
+
+//////////////////////////////////////////////////////////////////////////
+//Function Binders.
+
+//External function binder choosing between template specializations based on if we're area weighting or not.
+template<typename NextBinder>
+struct TAreaWeightingModeBinder
+{
+	template<typename... ParamTypes>
+	static FVMExternalFunction Bind(UNiagaraDataInterface* Interface, const FVMExternalFunctionBindingInfo& BindingInfo, void* InstanceData)
+	{
+		FNDISkeletalMesh_InstanceData* InstData = (FNDISkeletalMesh_InstanceData*)InstanceData;
+		check(InstData);
+		UNiagaraDataInterfaceSkeletalMesh* MeshInterface = CastChecked<UNiagaraDataInterfaceSkeletalMesh>(Interface);
+		const FSkeletalMeshSamplingInfo& SamplingInfo = InstData->Mesh->GetSamplingInfo();
+
+		bool bAreaWeighting = false;
+		if (InstData->SamplingRegionIndices.Num() > 1)
+		{
+			bAreaWeighting = InstData->SamplingRegionAreaWeightedSampler.IsValid();
+		}
+		else if (InstData->SamplingRegionIndices.Num() == 1)
+		{
+			const FSkeletalMeshSamplingRegion& Region = SamplingInfo.GetRegion(InstData->SamplingRegionIndices[0]);
+			bAreaWeighting = Region.bSupportUniformlyDistributedSampling;
+		}
+		else
+		{
+			int32 LODIndex = InstData->GetLODIndex();
+			check(InstData->Mesh->LODInfo[LODIndex].bAllowCPUAccess);
+			bAreaWeighting = InstData->Mesh->LODInfo[LODIndex].bSupportUniformlyDistributedSampling;
+		}
+
+		if (bAreaWeighting)
+		{
+			return NextBinder::template Bind<ParamTypes..., TIntegralConstant<ENDISkelMesh_AreaWeightingMode, ENDISkelMesh_AreaWeightingMode::AreaWeighted>>(Interface, BindingInfo, InstanceData);
+		}
+			
+		return NextBinder::template Bind<ParamTypes..., TIntegralConstant<ENDISkelMesh_AreaWeightingMode, ENDISkelMesh_AreaWeightingMode::None>>(Interface, BindingInfo, InstanceData);
+	}
+};
+
+//External function binder choosing between template specializations based on filtering methods
+template<typename NextBinder>
+struct TFilterModeBinder
+{
+	template<typename... ParamTypes>
+	static FVMExternalFunction Bind(UNiagaraDataInterface* Interface, const FVMExternalFunctionBindingInfo& BindingInfo, void* InstanceData)
+	{
+		FNDISkeletalMesh_InstanceData* InstData = (FNDISkeletalMesh_InstanceData*)InstanceData;
+		check(InstData);
+		UNiagaraDataInterfaceSkeletalMesh* MeshInterface = CastChecked<UNiagaraDataInterfaceSkeletalMesh>(Interface);
+
+		if (InstData->SamplingRegionIndices.Num() == 1)
+		{
+			return NextBinder::template Bind<ParamTypes..., TIntegralConstant<ENDISkeletalMesh_FilterMode, ENDISkeletalMesh_FilterMode::SingleRegion>>(Interface, BindingInfo, InstanceData);
+		}
+		else if (InstData->SamplingRegionIndices.Num() > 1)
+		{
+			return NextBinder::template Bind<ParamTypes..., TIntegralConstant<ENDISkeletalMesh_FilterMode, ENDISkeletalMesh_FilterMode::MultiRegion>>(Interface, BindingInfo, InstanceData);
+		}
+
+		return NextBinder::template Bind<ParamTypes..., TIntegralConstant<ENDISkeletalMesh_FilterMode, ENDISkeletalMesh_FilterMode::None>>(Interface, BindingInfo, InstanceData);
+	}
+};
+
+//External function binder choosing between template specializations based vetrex data format
+template<typename NextBinder>
+struct TVertexAccessorBinder
+{
+	template<typename... ParamTypes>
+	static FVMExternalFunction Bind(UNiagaraDataInterface* Interface, const FVMExternalFunctionBindingInfo& BindingInfo, void* InstanceData)
+	{
+		FNDISkeletalMesh_InstanceData* InstData = (FNDISkeletalMesh_InstanceData*)InstanceData;
+		UNiagaraDataInterfaceSkeletalMesh* MeshInterface = CastChecked<UNiagaraDataInterfaceSkeletalMesh>(Interface);
+		USkeletalMeshComponent* Component = Cast<USkeletalMeshComponent>(InstData->Component.Get());
+		FSkinWeightVertexBuffer* SkinWeightBuffer = nullptr;
+		FSkeletalMeshLODRenderData& LODData = InstData->GetLODRenderDataAndSkinWeights(SkinWeightBuffer);
+
+		if (LODData.StaticVertexBuffers.StaticMeshVertexBuffer.GetUseFullPrecisionUVs())
+		{
+			return NextBinder::template Bind<ParamTypes..., FSkelMeshVertexAccessor<true>>(Interface, BindingInfo, InstanceData);
+		}
+		else
+		{
+			return NextBinder::template Bind<ParamTypes..., FSkelMeshVertexAccessor<false>>(Interface, BindingInfo, InstanceData);
+		}
+	}
+};
+
+//External function binder choosing between template specializations based on skinning mode
+template<typename NextBinder>
+struct TSkinningModeBinder
+{
+	template<typename... ParamTypes>
+	static FVMExternalFunction Bind(UNiagaraDataInterface* Interface, const FVMExternalFunctionBindingInfo& BindingInfo, void* InstanceData)
+	{
+		FNDISkeletalMesh_InstanceData* InstData = (FNDISkeletalMesh_InstanceData*)InstanceData;
+		UNiagaraDataInterfaceSkeletalMesh* MeshInterface = CastChecked<UNiagaraDataInterfaceSkeletalMesh>(Interface);
+		USkeletalMeshComponent* Component = Cast<USkeletalMeshComponent>(InstData->Component.Get());
+		if (MeshInterface->SkinningMode == ENDISkeletalMesh_SkinningMode::None || !Component)//Can't skin if we have no component.
+		{
+			return NextBinder::template Bind<ParamTypes..., FSkinnedPositionAccessorHelper<TIntegralConstant<ENDISkeletalMesh_SkinningMode, ENDISkeletalMesh_SkinningMode::None>>>(Interface, BindingInfo, InstanceData);
+		}
+		else if (MeshInterface->SkinningMode == ENDISkeletalMesh_SkinningMode::SkinOnTheFly)
+		{
+			check(Component);
+			return NextBinder::template Bind<ParamTypes..., FSkinnedPositionAccessorHelper<TIntegralConstant<ENDISkeletalMesh_SkinningMode, ENDISkeletalMesh_SkinningMode::SkinOnTheFly>>>(Interface, BindingInfo, InstanceData);
+		}
+		else if (MeshInterface->SkinningMode == ENDISkeletalMesh_SkinningMode::PreSkin)
+		{
+			check(Component);
+			return NextBinder::template Bind<ParamTypes..., FSkinnedPositionAccessorHelper<TIntegralConstant<ENDISkeletalMesh_SkinningMode, ENDISkeletalMesh_SkinningMode::PreSkin>>>(Interface, BindingInfo, InstanceData);
+		}
+
+		checkf(false, TEXT("Invalid skinning mode in %s"), *Interface->GetPathName());
+		return FVMExternalFunction();
+	}
+};
+
+//Final binders for all static mesh interface functions.
+DEFINE_NDI_FUNC_BINDER(UNiagaraDataInterfaceSkeletalMesh, RandomTriCoord);
+DEFINE_NDI_FUNC_BINDER(UNiagaraDataInterfaceSkeletalMesh, GetTriCoordPosition);
+DEFINE_NDI_FUNC_BINDER(UNiagaraDataInterfaceSkeletalMesh, GetTriCoordPositionVelocityAndNormal);
+DEFINE_NDI_FUNC_BINDER(UNiagaraDataInterfaceSkeletalMesh, GetTriCoordColor);
+DEFINE_NDI_FUNC_BINDER(UNiagaraDataInterfaceSkeletalMesh, GetTriCoordUV);
+
+//Function Binders END
+//////////////////////////////////////////////////////////////////////////
+
+//////////////////////////////////////////////////////////////////////////
+// UNiagaraDataInterfaceSkeletalMesh
+
+UNiagaraDataInterfaceSkeletalMesh::UNiagaraDataInterfaceSkeletalMesh(FObjectInitializer const& ObjectInitializer)
+	: Super(ObjectInitializer)
+	, DefaultMesh(nullptr)
+	, Source(nullptr)
+	, SkinningMode(ENDISkeletalMesh_SkinningMode::PreSkin)
+	, WholeMeshLOD(INDEX_NONE)
+{
+
+}
+
+#if WITH_EDITOR
+
+void UNiagaraDataInterfaceSkeletalMesh::PostInitProperties()
+{
+	Super::PostInitProperties();
+
+	//Can we register data interfaces as regular types and fold them into the FNiagaraVariable framework for UI and function calls etc?
+	if (HasAnyFlags(RF_ClassDefaultObject))
+	{
+		FNiagaraTypeRegistry::Register(FNiagaraTypeDefinition(GetClass()), true, false, false);
+
+		//Still some issues with using custom structs. Convert node for example throws a wobbler. TODO after GDC.
+		FNiagaraTypeRegistry::Register(FMeshTriCoordinate::StaticStruct(), true, true, false);
+	}
+}
+
+#endif //WITH_EDITOR
+
+static const FName RandomTriCoordName("RandomTriCoord");
+
+static const FName GetTriPositionName("GetTriPosition");
+static const FName GetTriPositionWSName("GetTriPositionWS");
+
+static const FName GetTriNormalName("GetTriNormal");
+static const FName GetTriNormalWSName("GetTriNormalWS");
+
+static const FName GetTriColorName("GetTriColor");
+static const FName GetTriUVName("GetTriUV");
+
+static const FName GetTriPositionVelocityAndNormalName("GetTriPositionVelocityAndNormal");
+static const FName GetTriPositionVelocityAndNormalWSName("GetTriPositionVelocityAndNormalWS");
+
+void UNiagaraDataInterfaceSkeletalMesh::GetFunctions(TArray<FNiagaraFunctionSignature>& OutFunctions)
+{
+	{
+		FNiagaraFunctionSignature Sig;
+		Sig.Name = RandomTriCoordName;
+		Sig.Inputs.Add(FNiagaraVariable(FNiagaraTypeDefinition(GetClass()), TEXT("SkeletalMesh")));
+		Sig.Outputs.Add(FNiagaraVariable(FNiagaraTypeDefinition(FMeshTriCoordinate::StaticStruct()), TEXT("Coord")));
+		Sig.bMemberFunction = true;
+		Sig.bRequiresContext = false;
+		OutFunctions.Add(Sig);
+	}
+	
+	{
+		FNiagaraFunctionSignature Sig;
+		Sig.Name = GetTriPositionName;
+		Sig.Inputs.Add(FNiagaraVariable(FNiagaraTypeDefinition(GetClass()), TEXT("SkeletalMesh")));
+		Sig.Inputs.Add(FNiagaraVariable(FNiagaraTypeDefinition(FMeshTriCoordinate::StaticStruct()), TEXT("Coord")));
+		Sig.Outputs.Add(FNiagaraVariable(FNiagaraTypeDefinition::GetVec3Def(), TEXT("Position")));
+		Sig.bMemberFunction = true;
+		Sig.bRequiresContext = false;
+		OutFunctions.Add(Sig);
+	}
+
+	{
+		FNiagaraFunctionSignature Sig;
+		Sig.Name = GetTriPositionWSName;
+		Sig.Inputs.Add(FNiagaraVariable(FNiagaraTypeDefinition(GetClass()), TEXT("SkeletalMesh")));
+		Sig.Inputs.Add(FNiagaraVariable(FNiagaraTypeDefinition(FMeshTriCoordinate::StaticStruct()), TEXT("Coord")));
+		Sig.Outputs.Add(FNiagaraVariable(FNiagaraTypeDefinition::GetVec3Def(), TEXT("Position")));
+		Sig.bMemberFunction = true;
+		Sig.bRequiresContext = false;
+		OutFunctions.Add(Sig);
+	}
+
+	{
+		FNiagaraFunctionSignature Sig;
+		Sig.Name = GetTriPositionVelocityAndNormalName;
+		Sig.Inputs.Add(FNiagaraVariable(FNiagaraTypeDefinition(GetClass()), TEXT("SkeletalMesh")));
+		Sig.Inputs.Add(FNiagaraVariable(FNiagaraTypeDefinition(FMeshTriCoordinate::StaticStruct()), TEXT("Coord")));
+		Sig.Outputs.Add(FNiagaraVariable(FNiagaraTypeDefinition::GetVec3Def(), TEXT("Position")));
+		Sig.Outputs.Add(FNiagaraVariable(FNiagaraTypeDefinition::GetVec3Def(), TEXT("Velocity")));
+		Sig.Outputs.Add(FNiagaraVariable(FNiagaraTypeDefinition::GetVec3Def(), TEXT("Normal")));
+		Sig.bMemberFunction = true;
+		Sig.bRequiresContext = false;
+		OutFunctions.Add(Sig);
+	}
+
+	{
+		FNiagaraFunctionSignature Sig;
+		Sig.Name = GetTriPositionVelocityAndNormalWSName;
+		Sig.Inputs.Add(FNiagaraVariable(FNiagaraTypeDefinition(GetClass()), TEXT("SkeletalMesh")));
+		Sig.Inputs.Add(FNiagaraVariable(FNiagaraTypeDefinition(FMeshTriCoordinate::StaticStruct()), TEXT("Coord")));
+		Sig.Outputs.Add(FNiagaraVariable(FNiagaraTypeDefinition::GetVec3Def(), TEXT("Position")));
+		Sig.Outputs.Add(FNiagaraVariable(FNiagaraTypeDefinition::GetVec3Def(), TEXT("Velocity")));
+		Sig.Outputs.Add(FNiagaraVariable(FNiagaraTypeDefinition::GetVec3Def(), TEXT("Normal")));
+		Sig.bMemberFunction = true;
+		Sig.bRequiresContext = false;
+		OutFunctions.Add(Sig);
+	}
+	
+	{
+		FNiagaraFunctionSignature Sig;
+		Sig.Name = GetTriColorName;
+		Sig.Inputs.Add(FNiagaraVariable(FNiagaraTypeDefinition(GetClass()), TEXT("SkeletalMesh")));
+		Sig.Inputs.Add(FNiagaraVariable(FNiagaraTypeDefinition(FMeshTriCoordinate::StaticStruct()), TEXT("Coord")));
+		Sig.Outputs.Add(FNiagaraVariable(FNiagaraTypeDefinition::GetColorDef(), TEXT("Color")));
+		Sig.bMemberFunction = true;
+		Sig.bRequiresContext = false;
+		OutFunctions.Add(Sig);
+	}
+
+	{
+		FNiagaraFunctionSignature Sig;
+		Sig.Name = GetTriUVName;
+		Sig.Inputs.Add(FNiagaraVariable(FNiagaraTypeDefinition(GetClass()), TEXT("SkeletalMesh")));
+		Sig.Inputs.Add(FNiagaraVariable(FNiagaraTypeDefinition(FMeshTriCoordinate::StaticStruct()), TEXT("Coord")));
+		Sig.Inputs.Add(FNiagaraVariable(FNiagaraTypeDefinition::GetIntDef(), TEXT("UV Set")));
+
+		Sig.Outputs.Add(FNiagaraVariable(FNiagaraTypeDefinition::GetVec2Def(), TEXT("UV")));
+		Sig.bMemberFunction = true;
+		Sig.bRequiresContext = false;
+		OutFunctions.Add(Sig);
+	}
+}
+
+FVMExternalFunction UNiagaraDataInterfaceSkeletalMesh::GetVMExternalFunction(const FVMExternalFunctionBindingInfo& BindingInfo, void* InstanceData)
+{
+	FNDISkeletalMesh_InstanceData* InstData = (FNDISkeletalMesh_InstanceData*)InstanceData;
+	USkeletalMeshComponent* SkelComp = Cast<USkeletalMeshComponent>(InstData->Component.Get());
+	check(InstData && InstData->Mesh);
+
+	bool bNeedsVertexColors = false;
+
+	FVMExternalFunction Function;
+	if (BindingInfo.Name == RandomTriCoordName)
+	{
+		check(BindingInfo.GetNumInputs() == 1 && BindingInfo.GetNumOutputs() == 4);
+		Function = TFilterModeBinder<TAreaWeightingModeBinder<NDI_FUNC_BINDER(UNiagaraDataInterfaceSkeletalMesh, RandomTriCoord)>>::Bind(this, BindingInfo, InstanceData);
+	}
+	else if (BindingInfo.Name == GetTriPositionName)
+	{
+		check(BindingInfo.GetNumInputs() == 5 && BindingInfo.GetNumOutputs() == 3);
+		Function = TSkinningModeBinder<TNDIExplicitBinder<FNDITransformHandlerNoop, TNDIParamBinder<0, int32, TNDIParamBinder<1, float, TNDIParamBinder<2, float, TNDIParamBinder<3, float, NDI_FUNC_BINDER(UNiagaraDataInterfaceSkeletalMesh, GetTriCoordPosition)>>>>>>::Bind(this, BindingInfo, InstanceData);
+	}
+	else if (BindingInfo.Name == GetTriPositionWSName)
+	{
+		check(BindingInfo.GetNumInputs() == 5 && BindingInfo.GetNumOutputs() == 3);
+		Function = TSkinningModeBinder<TNDIExplicitBinder<FNDITransformHandler, TNDIParamBinder<0, int32, TNDIParamBinder<1, float, TNDIParamBinder<2, float, TNDIParamBinder<3, float, NDI_FUNC_BINDER(UNiagaraDataInterfaceSkeletalMesh, GetTriCoordPosition)>>>>>>::Bind(this, BindingInfo, InstanceData);
+	}
+	else if (BindingInfo.Name == GetTriPositionVelocityAndNormalName)
+	{
+		check(BindingInfo.GetNumInputs() == 5 && BindingInfo.GetNumOutputs() == 9);
+		Function = TSkinningModeBinder<TNDIExplicitBinder<FNDITransformHandlerNoop, TNDIParamBinder<0, int32, TNDIParamBinder<1, float, TNDIParamBinder<2, float, TNDIParamBinder<3, float, NDI_FUNC_BINDER(UNiagaraDataInterfaceSkeletalMesh, GetTriCoordPositionVelocityAndNormal)>>>>>>::Bind(this, BindingInfo, InstanceData);
+	}
+	else if (BindingInfo.Name == GetTriPositionVelocityAndNormalWSName)
+	{
+		check(BindingInfo.GetNumInputs() == 5 && BindingInfo.GetNumOutputs() == 9);
+		Function = TSkinningModeBinder<TNDIExplicitBinder<FNDITransformHandler, TNDIParamBinder<0, int32, TNDIParamBinder<1, float, TNDIParamBinder<2, float, TNDIParamBinder<3, float, NDI_FUNC_BINDER(UNiagaraDataInterfaceSkeletalMesh, GetTriCoordPositionVelocityAndNormal)>>>>>>::Bind(this, BindingInfo, InstanceData);
+	}
+	else if (BindingInfo.Name == GetTriColorName)
+	{
+		check(BindingInfo.GetNumInputs() == 5 && BindingInfo.GetNumOutputs() == 4);
+		bNeedsVertexColors = true;
+		Function = TNDIParamBinder<0, int32, TNDIParamBinder<1, float, TNDIParamBinder<2, float, TNDIParamBinder<3, float, NDI_FUNC_BINDER(UNiagaraDataInterfaceSkeletalMesh, GetTriCoordColor)>>>>::Bind(this, BindingInfo, InstanceData);
+	}
+	else if (BindingInfo.Name == GetTriUVName)
+	{
+		check(BindingInfo.GetNumInputs() == 6 && BindingInfo.GetNumOutputs() == 2);
+		Function = TVertexAccessorBinder<TNDIParamBinder<0, int32, TNDIParamBinder<1, float, TNDIParamBinder<2, float, TNDIParamBinder<3, float, TNDIParamBinder<4, int32, NDI_FUNC_BINDER(UNiagaraDataInterfaceSkeletalMesh, GetTriCoordUV)>>>>>>::Bind(this, BindingInfo, InstanceData);
+	}
+	
+	check(InstData->Mesh);
+	FSkinWeightVertexBuffer* SkinWeightBuffer;
+	FSkeletalMeshLODRenderData& LODData = InstData->GetLODRenderDataAndSkinWeights(SkinWeightBuffer);
+
+	if (bNeedsVertexColors && LODData.StaticVertexBuffers.ColorVertexBuffer.GetNumVertices() == 0)
+	{
+		UE_LOG(LogNiagara, Log, TEXT("Skeletal Mesh data interface is cannot run as it's reading color data on a mesh that does not provide it. - Mesh:%s  "), *InstData->Mesh->GetFullName());
+		Function = FVMExternalFunction();
+	}
+	
+	return Function;
+}
+
+
+bool UNiagaraDataInterfaceSkeletalMesh::CopyToInternal(UNiagaraDataInterface* Destination) const
+{
+	if (!Super::CopyToInternal(Destination))
+	{
+		return false;
+	}
+
+	UNiagaraDataInterfaceSkeletalMesh* OtherTyped = CastChecked<UNiagaraDataInterfaceSkeletalMesh>(Destination);
+	OtherTyped->Source = Source;
+	OtherTyped->DefaultMesh = DefaultMesh;
+	OtherTyped->SkinningMode = SkinningMode;
+	OtherTyped->SamplingRegions = SamplingRegions;
+	OtherTyped->WholeMeshLOD = WholeMeshLOD;
+	return true;
+}
+
+bool UNiagaraDataInterfaceSkeletalMesh::Equals(const UNiagaraDataInterface* Other) const
+{
+	if (!Super::Equals(Other))
+	{
+		return false;
+	}
+	const UNiagaraDataInterfaceSkeletalMesh* OtherTyped = CastChecked<const UNiagaraDataInterfaceSkeletalMesh>(Other);
+	return OtherTyped->Source == Source &&
+		OtherTyped->DefaultMesh == DefaultMesh &&
+		OtherTyped->SkinningMode == SkinningMode &&
+		OtherTyped->SamplingRegions == SamplingRegions &&
+		OtherTyped->WholeMeshLOD == WholeMeshLOD;
+}
+
+bool UNiagaraDataInterfaceSkeletalMesh::InitPerInstanceData(void* PerInstanceData, FNiagaraSystemInstance* SystemInstance)
+{
+	FNDISkeletalMesh_InstanceData* Inst = new (PerInstanceData) FNDISkeletalMesh_InstanceData();
+	return Inst->Init(this, SystemInstance);
+}
+
+void UNiagaraDataInterfaceSkeletalMesh::DestroyPerInstanceData(void* PerInstanceData, FNiagaraSystemInstance* SystemInstance)
+{
+	FNDISkeletalMesh_InstanceData* Inst = (FNDISkeletalMesh_InstanceData*)PerInstanceData;
+
+#if WITH_EDITOR
+	if(Inst->Mesh)
+	{
+		Inst->Mesh->GetOnMeshChanged().RemoveAll(SystemInstance->GetComponent());
+	}
+#endif
+
+	Inst->~FNDISkeletalMesh_InstanceData();
+}
+
+bool UNiagaraDataInterfaceSkeletalMesh::PerInstanceTick(void* PerInstanceData, FNiagaraSystemInstance* SystemInstance, float InDeltaSeconds)
+{
+	FNDISkeletalMesh_InstanceData* Inst = (FNDISkeletalMesh_InstanceData*)PerInstanceData;
+	return Inst->Tick(this, SystemInstance, InDeltaSeconds);
+}
+
+//////////////////////////////////////////////////////////////////////////
+
+template<typename FilterMode, typename AreaWeightingMode>
+FORCEINLINE int32 UNiagaraDataInterfaceSkeletalMesh::RandomTriIndex(FSkeletalMeshAccessorHelper& Accessor, FNDISkeletalMesh_InstanceData* InstData)
+{
+	checkf(false, TEXT("Invalid template call for RandomTriIndex. Bug in Filter binding or Area Weighting binding. Contact code team."));
+	return 0;
+}
+
+template<>
+FORCEINLINE int32 UNiagaraDataInterfaceSkeletalMesh::RandomTriIndex<
+	TIntegralConstant<ENDISkeletalMesh_FilterMode, ENDISkeletalMesh_FilterMode::None>,
+	TIntegralConstant<ENDISkelMesh_AreaWeightingMode, ENDISkelMesh_AreaWeightingMode::None>>
+	(FSkeletalMeshAccessorHelper& Accessor, FNDISkeletalMesh_InstanceData* InstData)
+{
+	int32 SecIdx = InstData->RandStream.RandRange(0, Accessor.LODData->RenderSections.Num() - 1);
+	FSkelMeshRenderSection& Sec = Accessor.LODData->RenderSections[SecIdx];
+	int32 Tri = InstData->RandStream.RandRange(0, Sec.NumTriangles - 1);
+	return Sec.BaseIndex + Tri * 3;
+}
+
+template<>
+FORCEINLINE int32 UNiagaraDataInterfaceSkeletalMesh::RandomTriIndex<
+	TIntegralConstant<ENDISkeletalMesh_FilterMode, ENDISkeletalMesh_FilterMode::None>,
+	TIntegralConstant<ENDISkelMesh_AreaWeightingMode, ENDISkelMesh_AreaWeightingMode::AreaWeighted>>
+	(FSkeletalMeshAccessorHelper& Accessor, FNDISkeletalMesh_InstanceData* InstData)
+{
+	const FSkeletalMeshSamplingInfo& SamplingInfo = InstData->Mesh->GetSamplingInfo();
+	const FSkeletalMeshSamplingLODBuiltData& WholeMeshBuiltData = SamplingInfo.GetWholeMeshLODBuiltData(InstData->GetLODIndex());
+	int32 TriIdx = WholeMeshBuiltData.AreaWeightedTriangleSampler.GetEntryIndex(InstData->RandStream.GetFraction(), InstData->RandStream.GetFraction());
+	return TriIdx * 3;
+}
+
+template<>
+FORCEINLINE int32 UNiagaraDataInterfaceSkeletalMesh::RandomTriIndex<
+	TIntegralConstant<ENDISkeletalMesh_FilterMode, ENDISkeletalMesh_FilterMode::SingleRegion>,
+	TIntegralConstant<ENDISkelMesh_AreaWeightingMode, ENDISkelMesh_AreaWeightingMode::None>>
+	(FSkeletalMeshAccessorHelper& Accessor, FNDISkeletalMesh_InstanceData* InstData)
+{
+	int32 Idx = InstData->RandStream.RandRange(0, Accessor.SamplingRegionBuiltData->TriangleIndices.Num() - 1);
+	return Accessor.SamplingRegionBuiltData->TriangleIndices[Idx];
+}
+
+template<>
+FORCEINLINE int32 UNiagaraDataInterfaceSkeletalMesh::RandomTriIndex<
+	TIntegralConstant<ENDISkeletalMesh_FilterMode, ENDISkeletalMesh_FilterMode::SingleRegion>,
+	TIntegralConstant<ENDISkelMesh_AreaWeightingMode, ENDISkelMesh_AreaWeightingMode::AreaWeighted>>
+	(FSkeletalMeshAccessorHelper& Accessor, FNDISkeletalMesh_InstanceData* InstData)
+{
+	int32 Idx = Accessor.SamplingRegionBuiltData->AreaWeightedSampler.GetEntryIndex(InstData->RandStream.GetFraction(), InstData->RandStream.GetFraction());
+	return Accessor.SamplingRegionBuiltData->TriangleIndices[Idx];
+}
+
+template<>
+FORCEINLINE int32 UNiagaraDataInterfaceSkeletalMesh::RandomTriIndex<
+	TIntegralConstant<ENDISkeletalMesh_FilterMode, ENDISkeletalMesh_FilterMode::MultiRegion>,
+	TIntegralConstant<ENDISkelMesh_AreaWeightingMode, ENDISkelMesh_AreaWeightingMode::None>>
+	(FSkeletalMeshAccessorHelper& Accessor, FNDISkeletalMesh_InstanceData* InstData)
+{
+	int32 RegionIdx = InstData->RandStream.RandRange(0, InstData->SamplingRegionIndices.Num() - 1);
+	const FSkeletalMeshSamplingInfo& SamplingInfo = InstData->Mesh->GetSamplingInfo();
+	const FSkeletalMeshSamplingRegion& Region = SamplingInfo.GetRegion(InstData->SamplingRegionIndices[RegionIdx]);
+	const FSkeletalMeshSamplingRegionBuiltData& RegionBuiltData = SamplingInfo.GetRegionBuiltData(InstData->SamplingRegionIndices[RegionIdx]);
+	int32 Idx = InstData->RandStream.RandRange(0, RegionBuiltData.TriangleIndices.Num() - 1);
+	return RegionBuiltData.TriangleIndices[Idx];
+}
+
+template<>
+FORCEINLINE int32 UNiagaraDataInterfaceSkeletalMesh::RandomTriIndex<
+	TIntegralConstant<ENDISkeletalMesh_FilterMode, ENDISkeletalMesh_FilterMode::MultiRegion>,
+	TIntegralConstant<ENDISkelMesh_AreaWeightingMode, ENDISkelMesh_AreaWeightingMode::AreaWeighted>>
+	(FSkeletalMeshAccessorHelper& Accessor, FNDISkeletalMesh_InstanceData* InstData)
+{
+	int32 RegionIdx = InstData->SamplingRegionAreaWeightedSampler.GetEntryIndex(InstData->RandStream.GetFraction(), InstData->RandStream.GetFraction());
+	const FSkeletalMeshSamplingInfo& SamplingInfo = InstData->Mesh->GetSamplingInfo();
+	const FSkeletalMeshSamplingRegion& Region = SamplingInfo.GetRegion(InstData->SamplingRegionIndices[RegionIdx]);
+	const FSkeletalMeshSamplingRegionBuiltData& RegionBuiltData = SamplingInfo.GetRegionBuiltData(InstData->SamplingRegionIndices[RegionIdx]);
+	int32 Idx = RegionBuiltData.AreaWeightedSampler.GetEntryIndex(InstData->RandStream.GetFraction(), InstData->RandStream.GetFraction());
+	return RegionBuiltData.TriangleIndices[Idx];
+}
+
+template<typename FilterMode, typename AreaWeightingMode>
+void UNiagaraDataInterfaceSkeletalMesh::RandomTriCoord(FVectorVMContext& Context)
+{
+	FUserPtrHandler<FNDISkeletalMesh_InstanceData> InstData(Context);
+	checkfSlow(InstData.Get(), TEXT("Skeletal Mesh Interface has invalid instance data. %s"), *GetPathName());
+	checkfSlow(InstData->Mesh, TEXT("Skeletal Mesh Interface has invalid mesh. %s"), *GetPathName());
+
+	FRegisterHandler<int32> OutTri(Context);
+	FRegisterHandler<float> OutBaryX(Context);	FRegisterHandler<float> OutBaryY(Context);	FRegisterHandler<float> OutBaryZ(Context);
+
+	FSkeletalMeshAccessorHelper MeshAccessor;
+	MeshAccessor.Init<FilterMode, AreaWeightingMode>(InstData);
+	
+	for (int32 i = 0; i < Context.NumInstances; ++i)
+	{
+		*OutTri.GetDest() = RandomTriIndex<FilterMode, AreaWeightingMode>(MeshAccessor, InstData);
+		FVector Bary = RandomBarycentricCoord(InstData->RandStream);
+		*OutBaryX.GetDest() = Bary.X;		*OutBaryY.GetDest() = Bary.Y;		*OutBaryZ.GetDest() = Bary.Z;
+		
+		OutTri.Advance();
+		OutBaryX.Advance();		OutBaryY.Advance();		OutBaryZ.Advance();
+	}
+}
+
+template<typename SkinningHandlerType, typename TransformHandlerType, typename TriType, typename BaryXType, typename BaryYType, typename BaryZType>
+void UNiagaraDataInterfaceSkeletalMesh::GetTriCoordPosition(FVectorVMContext& Context)
+{
+	SkinningHandlerType SkinningHandler;
+	TransformHandlerType TransformHandler;
+	TriType TriParam(Context);
+	BaryXType BaryXParam(Context);	BaryYType BaryYParam(Context);	BaryZType BaryZParam(Context);
+	FUserPtrHandler<FNDISkeletalMesh_InstanceData> InstData(Context);
+
+	checkfSlow(InstData.Get(), TEXT("Skeletal Mesh Interface has invalid instance data. %s"), *GetPathName());
+	checkfSlow(InstData->Mesh, TEXT("Skeletal Mesh Interface has invalid mesh. %s"), *GetPathName());
+
+	FRegisterHandler<float> OutPosX(Context);	FRegisterHandler<float> OutPosY(Context);	FRegisterHandler<float> OutPosZ(Context);
+
+	FSkeletalMeshAccessorHelper Accessor;
+	Accessor.Init<TIntegralConstant<int32, 0>,TIntegralConstant<int32, 0>>(InstData);
+	for (int32 i = 0; i < Context.NumInstances; ++i)
+	{
+		int32 Tri = TriParam.Get();
+		FVector V0;		FVector V1;		FVector V2;
+		SkinningHandler.GetSkinnedTrianglePositions(Accessor, Tri, V0, V1, V2);
+		
+		FVector Pos = BarycentricInterpolate(BaryXParam.Get(), BaryYParam.Get(), BaryZParam.Get(), V0, V1, V2);
+		TransformHandler.TransformPosition(Pos, InstData->Transform);
+
+		*OutPosX.GetDest() = Pos.X;		*OutPosY.GetDest() = Pos.Y;		*OutPosZ.GetDest() = Pos.Z;
+
+		TriParam.Advance();
+		BaryXParam.Advance();	BaryYParam.Advance();		BaryZParam.Advance();
+		OutPosX.Advance();		OutPosY.Advance();			OutPosZ.Advance();
+	}
+}
+
+template<typename SkinningHandlerType, typename TransformHandlerType, typename TriType, typename BaryXType, typename BaryYType, typename BaryZType>
+void UNiagaraDataInterfaceSkeletalMesh::GetTriCoordPositionVelocityAndNormal(FVectorVMContext& Context)
+{
+	SkinningHandlerType SkinningHandler;
+	TransformHandlerType TransformHandler;
+	TriType TriParam(Context);
+	BaryXType BaryXParam(Context);	BaryYType BaryYParam(Context);	BaryZType BaryZParam(Context);
+	FUserPtrHandler<FNDISkeletalMesh_InstanceData> InstData(Context);
+
+	checkfSlow(InstData.Get(), TEXT("Skeletal Mesh Interface has invalid instance data. %s"), *GetPathName());
+	checkfSlow(InstData->Mesh, TEXT("Skeletal Mesh Interface has invalid mesh. %s"), *GetPathName());
+
+	FRegisterHandler<float> OutPosX(Context);	FRegisterHandler<float> OutPosY(Context);	FRegisterHandler<float> OutPosZ(Context);
+	FRegisterHandler<float> OutVelX(Context);	FRegisterHandler<float> OutVelY(Context);	FRegisterHandler<float> OutVelZ(Context);
+	FRegisterHandler<float> OutNormX(Context);	FRegisterHandler<float> OutNormY(Context);	FRegisterHandler<float> OutNormZ(Context);
+
+	FSkeletalMeshAccessorHelper Accessor;
+	Accessor.Init<TIntegralConstant<int32, 0>, TIntegralConstant<int32, 0>>(InstData);
+	float InvDt = 1.0f / InstData->DeltaSeconds;
+	for (int32 i = 0; i < Context.NumInstances; ++i)
+	{
+		int32 Tri = TriParam.Get();
+		FVector Pos0;		FVector Pos1;		FVector Pos2;
+		FVector Prev0;		FVector Prev1;		FVector Prev2;
+		SkinningHandler.GetSkinnedTrianglePositions(Accessor, Tri, Pos0, Pos1, Pos2, Prev0, Prev1, Prev2);
+
+		FVector Pos = BarycentricInterpolate(BaryXParam.Get(), BaryYParam.Get(), BaryZParam.Get(), Pos0, Pos1, Pos2);
+		FVector Prev = BarycentricInterpolate(BaryXParam.Get(), BaryYParam.Get(), BaryZParam.Get(), Prev0, Prev1, Prev2);
+		TransformHandler.TransformPosition(Pos, InstData->Transform);
+		TransformHandler.TransformPosition(Prev, InstData->PrevTransform);
+
+		FVector Vel = (Pos - Prev) * InvDt;
+
+		FVector Normal = ((Pos1 - Pos0) ^ (Pos2 - Pos0)).GetUnsafeNormal();//Temporarily having to get a dirty normal here until the newer pre skinning goodness comes online.
+
+		*OutPosX.GetDest() = Pos.X;		*OutPosY.GetDest() = Pos.Y;		*OutPosZ.GetDest() = Pos.Z;
+		*OutVelX.GetDest() = Vel.X;		*OutVelY.GetDest() = Vel.Y;		*OutVelZ.GetDest() = Vel.Z;
+		*OutNormX.GetDest() = Normal.X;	*OutNormY.GetDest() = Normal.Y;	*OutNormZ.GetDest() = Normal.Z;
+// 		//TODO: Use the velocity temp buffer as storage for the prev position so I can SIMD the velocity calculation.
+// 		*OutVelX.GetDest() = Prev.X;		*OutVelY.GetDest() = Prev.Y;		*OutVelZ.GetDest() = Prev.Z;
+
+		TriParam.Advance();
+		BaryXParam.Advance();	BaryYParam.Advance();		BaryZParam.Advance();
+		OutPosX.Advance();		OutPosY.Advance();			OutPosZ.Advance();
+		OutVelX.Advance();		OutVelY.Advance();			OutVelZ.Advance();
+		OutNormX.Advance();		OutNormY.Advance();			OutNormZ.Advance();
+	}
+}
+
+template<typename TriType, typename BaryXType, typename BaryYType, typename BaryZType>
+void UNiagaraDataInterfaceSkeletalMesh::GetTriCoordColor(FVectorVMContext& Context)
+{
+	TriType TriParam(Context);
+	BaryXType BaryXParam(Context);
+	BaryYType BaryYParam(Context);
+	BaryZType BaryZParam(Context);
+	FUserPtrHandler<FNDISkeletalMesh_InstanceData> InstData(Context);
+
+	FRegisterHandler<float> OutColorR(Context);
+	FRegisterHandler<float> OutColorG(Context);
+	FRegisterHandler<float> OutColorB(Context);
+	FRegisterHandler<float> OutColorA(Context);
+
+	USkeletalMeshComponent* Comp = Cast<USkeletalMeshComponent>(InstData->Component.Get());
+	FSkinWeightVertexBuffer* SkinWeightBuffer;
+	FSkeletalMeshLODRenderData& LODData = InstData->GetLODRenderDataAndSkinWeights(SkinWeightBuffer);
+	const FColorVertexBuffer& Colors = LODData.StaticVertexBuffers.ColorVertexBuffer;
+	checkfSlow(Colors.GetNumVertices() != 0, TEXT("Trying to access vertex colors from mesh without any."));
+
+	FMultiSizeIndexContainer& Indices = LODData.MultiSizeIndexContainer;
+	const FRawStaticIndexBuffer16or32Interface* IndexBuffer = Indices.GetIndexBuffer();
+	for (int32 i = 0; i < Context.NumInstances; ++i)
+	{
+		int32 Tri = TriParam.Get();
+		int32 Idx0 = IndexBuffer->Get(Tri);
+		int32 Idx1 = IndexBuffer->Get(Tri + 1);
+		int32 Idx2 = IndexBuffer->Get(Tri + 2);
+
+		FLinearColor Color = BarycentricInterpolate(BaryXParam.Get(), BaryYParam.Get(), BaryZParam.Get(),
+			Colors.VertexColor(Idx0).ReinterpretAsLinear(), Colors.VertexColor(Idx1).ReinterpretAsLinear(), Colors.VertexColor(Idx2).ReinterpretAsLinear());
+
+		*OutColorR.GetDest() = Color.R;
+		*OutColorG.GetDest() = Color.G;
+		*OutColorB.GetDest() = Color.B;
+		*OutColorA.GetDest() = Color.A;
+		TriParam.Advance();
+		BaryXParam.Advance();
+		BaryYParam.Advance();
+		BaryZParam.Advance();
+		OutColorR.Advance();
+		OutColorG.Advance();
+		OutColorB.Advance();
+		OutColorA.Advance();
+	}
+}
+
+template<typename VertexAccessorType, typename TriType, typename BaryXType, typename BaryYType, typename BaryZType, typename UVSetType>
+void UNiagaraDataInterfaceSkeletalMesh::GetTriCoordUV(FVectorVMContext& Context)
+{
+	VertexAccessorType VertAccessor;
+	TriType TriParam(Context);
+	BaryXType BaryXParam(Context);	BaryYType BaryYParam(Context);	BaryZType BaryZParam(Context);
+	UVSetType UVSetParam(Context);
+	FUserPtrHandler<FNDISkeletalMesh_InstanceData> InstData(Context);
+
+	checkfSlow(InstData.Get(), TEXT("Skeletal Mesh Interface has invalid instance data. %s"), *GetPathName());
+	checkfSlow(InstData->Mesh, TEXT("Skeletal Mesh Interface has invalid mesh. %s"), *GetPathName());
+
+	FRegisterHandler<float> OutUVX(Context);	FRegisterHandler<float> OutUVY(Context);
+
+	USkeletalMeshComponent* Comp = Cast<USkeletalMeshComponent>(InstData->Component.Get());
+	FSkinWeightVertexBuffer* SkinWeightBuffer;
+	FSkeletalMeshLODRenderData& LODData = InstData->GetLODRenderDataAndSkinWeights(SkinWeightBuffer);
+
+	FMultiSizeIndexContainer& Indices = LODData.MultiSizeIndexContainer;
+	FRawStaticIndexBuffer16or32Interface* IndexBuffer = Indices.GetIndexBuffer();
+	float InvDt = 1.0f / InstData->DeltaSeconds;
+	for (int32 i = 0; i < Context.NumInstances; ++i)
+	{
+		int32 Tri = TriParam.Get();
+		int32 Idx0 = IndexBuffer->Get(Tri);
+		int32 Idx1 = IndexBuffer->Get(Tri + 1);
+		int32 Idx2 = IndexBuffer->Get(Tri + 2);
+		FVector2D UV0;		FVector2D UV1;		FVector2D UV2;
+		int32 UVSet = UVSetParam.Get();
+		UV0 = VertAccessor.GetVertexUV(LODData, Idx0, UVSet);
+		UV1 = VertAccessor.GetVertexUV(LODData, Idx1, UVSet);
+		UV2 = VertAccessor.GetVertexUV(LODData, Idx2, UVSet);
+
+		FVector2D UV = BarycentricInterpolate(BaryXParam.Get(), BaryYParam.Get(), BaryZParam.Get(), UV0, UV1, UV2);
+
+		*OutUVX.GetDest() = UV.X;
+		*OutUVY.GetDest() = UV.Y;
+
+		TriParam.Advance();
+		BaryXParam.Advance(); BaryYParam.Advance(); BaryZParam.Advance();
+		UVSetParam.Advance();
+		OutUVX.Advance();
+		OutUVY.Advance();
+	}
+}
+
+//UNiagaraDataInterfaceSkeletalMesh END
+//////////////////////////////////////////////////////////////////////////
+#undef LOCTEXT_NAMESPACE
diff --git a/Engine/Plugins/FX/Niagara/Source/Niagara/Private/NiagaraDataInterfaceSpline.cpp b/Engine/Plugins/FX/Niagara/Source/Niagara/Private/NiagaraDataInterfaceSpline.cpp
--- a/Engine/Plugins/FX/Niagara/Source/Niagara/Private/NiagaraDataInterfaceSpline.cpp
+++ b/Engine/Plugins/FX/Niagara/Source/Niagara/Private/NiagaraDataInterfaceSpline.cpp
@@ -8,29 +8,6 @@
 
 #define LOCTEXT_NAMESPACE "NiagaraDataInterfaceSpline"
 
-namespace SplineUtils
-{
-	struct FTransformHandlerNoop
-	{
-		FORCEINLINE void Transform(FVector& V, FMatrix& M) {  }
-	};
-
-	struct FTransformHandlerPosition
-	{
-		FORCEINLINE void Transform(FVector& P, FMatrix& M) { P = M.TransformPosition(P); }
-	};
-
-	struct FTransformHandlerUnitVector
-	{
-		FORCEINLINE void Transform(FVector& V, FMatrix& M) { V = M.TransformVector(V).GetUnsafeNormal3(); }
-	};
-
-	struct FTransformHandlerTangent
-	{
-		FORCEINLINE void Transform(FVector& V, FMatrix& M) { V = M.TransformVector(V); }
-	};
-};
-
 UNiagaraDataInterfaceSpline::UNiagaraDataInterfaceSpline(FObjectInitializer const& ObjectInitializer)
 	: Super(ObjectInitializer)
 	, Source(nullptr)
@@ -252,56 +229,55 @@ DEFINE_NDI_FUNC_BINDER(UNiagaraDataInterfaceSpline, FindClosestUnitDistanceFromP
 
 FVMExternalFunction UNiagaraDataInterfaceSpline::GetVMExternalFunction(const FVMExternalFunctionBindingInfo& BindingInfo, void* InstanceData)
 {
-	using namespace SplineUtils;
 	if (BindingInfo.Name == SampleSplinePositionByUnitDistanceName)
 	{
 		check(BindingInfo.GetNumInputs() == 2 && BindingInfo.GetNumOutputs() == 3);
-		return TNDIExplicitBinder<FTransformHandlerNoop, TNDIParamBinder<0, float, NDI_FUNC_BINDER(UNiagaraDataInterfaceSpline, SampleSplinePositionByUnitDistance)>>::Bind(this, BindingInfo, InstanceData);
+		return TNDIExplicitBinder<FNDITransformHandlerNoop, TNDIParamBinder<0, float, NDI_FUNC_BINDER(UNiagaraDataInterfaceSpline, SampleSplinePositionByUnitDistance)>>::Bind(this, BindingInfo, InstanceData);
 	}
 	else if (BindingInfo.Name == SampleSplinePositionByUnitDistanceWSName)
 	{
 		check(BindingInfo.GetNumInputs() == 2 && BindingInfo.GetNumOutputs() == 3);
-		return TNDIExplicitBinder<FTransformHandlerPosition, TNDIParamBinder<0, float, NDI_FUNC_BINDER(UNiagaraDataInterfaceSpline, SampleSplinePositionByUnitDistance)>>::Bind(this, BindingInfo, InstanceData);
+		return TNDIExplicitBinder<FNDITransformHandler, TNDIParamBinder<0, float, NDI_FUNC_BINDER(UNiagaraDataInterfaceSpline, SampleSplinePositionByUnitDistance)>>::Bind(this, BindingInfo, InstanceData);
 	}
 	else if (BindingInfo.Name == SampleSplineUpVectorByUnitDistanceName)
 	{
 		check(BindingInfo.GetNumInputs() == 2 && BindingInfo.GetNumOutputs() == 3);
-		return TNDIExplicitBinder<FTransformHandlerNoop, TNDIParamBinder<0, float, NDI_FUNC_BINDER(UNiagaraDataInterfaceSpline, SampleSplineUpVectorByUnitDistance)>>::Bind(this, BindingInfo, InstanceData);
+		return TNDIExplicitBinder<FNDITransformHandlerNoop, TNDIParamBinder<0, float, NDI_FUNC_BINDER(UNiagaraDataInterfaceSpline, SampleSplineUpVectorByUnitDistance)>>::Bind(this, BindingInfo, InstanceData);
 	}
 	else if (BindingInfo.Name == SampleSplineUpVectorByUnitDistanceWSName)
 	{
 		check(BindingInfo.GetNumInputs() == 2 && BindingInfo.GetNumOutputs() == 3);
-		return TNDIExplicitBinder<FTransformHandlerUnitVector, TNDIParamBinder<0, float, NDI_FUNC_BINDER(UNiagaraDataInterfaceSpline, SampleSplineUpVectorByUnitDistance)>>::Bind(this, BindingInfo, InstanceData);
+		return TNDIExplicitBinder<FNDITransformHandler, TNDIParamBinder<0, float, NDI_FUNC_BINDER(UNiagaraDataInterfaceSpline, SampleSplineUpVectorByUnitDistance)>>::Bind(this, BindingInfo, InstanceData);
 	}
 	else if (BindingInfo.Name == SampleSplineDirectionByUnitDistanceName)
 	{
 		check(BindingInfo.GetNumInputs() == 2 && BindingInfo.GetNumOutputs() == 3);
-		return TNDIExplicitBinder<FTransformHandlerNoop, TNDIParamBinder<0, float, NDI_FUNC_BINDER(UNiagaraDataInterfaceSpline, SampleSplineDirectionByUnitDistance)>>::Bind(this, BindingInfo, InstanceData);
+		return TNDIExplicitBinder<FNDITransformHandlerNoop, TNDIParamBinder<0, float, NDI_FUNC_BINDER(UNiagaraDataInterfaceSpline, SampleSplineDirectionByUnitDistance)>>::Bind(this, BindingInfo, InstanceData);
 	}
 	else if (BindingInfo.Name == SampleSplineDirectionByUnitDistanceWSName)
 	{
 		check(BindingInfo.GetNumInputs() == 2 && BindingInfo.GetNumOutputs() == 3);
-		return TNDIExplicitBinder<FTransformHandlerUnitVector, TNDIParamBinder<0, float, NDI_FUNC_BINDER(UNiagaraDataInterfaceSpline, SampleSplineDirectionByUnitDistance)>>::Bind(this, BindingInfo, InstanceData);
+		return TNDIExplicitBinder<FNDITransformHandler, TNDIParamBinder<0, float, NDI_FUNC_BINDER(UNiagaraDataInterfaceSpline, SampleSplineDirectionByUnitDistance)>>::Bind(this, BindingInfo, InstanceData);
 	}
 	else if (BindingInfo.Name == SampleSplineRightVectorByUnitDistanceName)
 	{
 		check(BindingInfo.GetNumInputs() == 2 && BindingInfo.GetNumOutputs() == 3);
-		return TNDIExplicitBinder<FTransformHandlerNoop, TNDIParamBinder<0, float, NDI_FUNC_BINDER(UNiagaraDataInterfaceSpline, SampleSplineRightVectorByUnitDistance)>>::Bind(this, BindingInfo, InstanceData);
+		return TNDIExplicitBinder<FNDITransformHandlerNoop, TNDIParamBinder<0, float, NDI_FUNC_BINDER(UNiagaraDataInterfaceSpline, SampleSplineRightVectorByUnitDistance)>>::Bind(this, BindingInfo, InstanceData);
 	}
 	else if (BindingInfo.Name == SampleSplineRightVectorByUnitDistanceWSName)
 	{
 		check(BindingInfo.GetNumInputs() == 2 && BindingInfo.GetNumOutputs() == 3);
-		return TNDIExplicitBinder<FTransformHandlerUnitVector, TNDIParamBinder<0, float, NDI_FUNC_BINDER(UNiagaraDataInterfaceSpline, SampleSplineRightVectorByUnitDistance)>>::Bind(this, BindingInfo, InstanceData);
+		return TNDIExplicitBinder<FNDITransformHandler, TNDIParamBinder<0, float, NDI_FUNC_BINDER(UNiagaraDataInterfaceSpline, SampleSplineRightVectorByUnitDistance)>>::Bind(this, BindingInfo, InstanceData);
 	}
 	else if (BindingInfo.Name == SampleSplineTangentByUnitDistanceName)
 	{
 		check(BindingInfo.GetNumInputs() == 2 && BindingInfo.GetNumOutputs() == 3);
-		return TNDIExplicitBinder<FTransformHandlerNoop, TNDIParamBinder<0, float, NDI_FUNC_BINDER(UNiagaraDataInterfaceSpline, SampleSplineTangentByUnitDistance)>>::Bind(this, BindingInfo, InstanceData);
+		return TNDIExplicitBinder<FNDITransformHandlerNoop, TNDIParamBinder<0, float, NDI_FUNC_BINDER(UNiagaraDataInterfaceSpline, SampleSplineTangentByUnitDistance)>>::Bind(this, BindingInfo, InstanceData);
 	}
 	else if (BindingInfo.Name == SampleSplineTangentByUnitDistanceWSName)
 	{
 		check(BindingInfo.GetNumInputs() == 2 && BindingInfo.GetNumOutputs() == 3);
-		return TNDIExplicitBinder<FTransformHandlerTangent, TNDIParamBinder<0, float, NDI_FUNC_BINDER(UNiagaraDataInterfaceSpline, SampleSplineTangentByUnitDistance)>>::Bind(this, BindingInfo, InstanceData);
+		return TNDIExplicitBinder<FNDITransformHandler, TNDIParamBinder<0, float, NDI_FUNC_BINDER(UNiagaraDataInterfaceSpline, SampleSplineTangentByUnitDistance)>>::Bind(this, BindingInfo, InstanceData);
 	}
 	else if (BindingInfo.Name == FindClosestUnitDistanceFromPositionWSName)
 	{
@@ -323,9 +299,9 @@ FVMExternalFunction UNiagaraDataInterfaceSpline::GetVMExternalFunction(const FVM
 	return FVMExternalFunction();
 }
 
-bool UNiagaraDataInterfaceSpline::CopyTo(UNiagaraDataInterface* Destination) const
+bool UNiagaraDataInterfaceSpline::CopyToInternal(UNiagaraDataInterface* Destination) const
 {
-	if (!Super::CopyTo(Destination))
+	if (!Super::CopyToInternal(Destination))
 	{
 		return false;
 	}
@@ -430,7 +406,7 @@ void UNiagaraDataInterfaceSpline::SampleSplinePositionByUnitDistance(FVectorVMCo
 		float DistanceUnitDistance = SplineSampleParam.Get();
 
 		FVector Pos = InstData->Component->GetLocationAtDistanceAlongSpline(DistanceUnitDistance * InstData->Component->GetSplineLength(), ESplineCoordinateSpace::Local);
-		TransformHandler.Transform(Pos, InstData->Transform);
+		TransformHandler.TransformPosition(Pos, InstData->Transform);
 
 		*OutPosX.GetDest() = Pos.X;
 		*OutPosY.GetDest() = Pos.Y;
@@ -459,7 +435,7 @@ void UNiagaraDataInterfaceSpline::SampleSplineUpVectorByUnitDistance(FVectorVMCo
 		float DistanceUnitDistance = SplineSampleParam.Get();
 
 		FVector Pos = InstData->Component->GetUpVectorAtDistanceAlongSpline(DistanceUnitDistance * InstData->Component->GetSplineLength(), ESplineCoordinateSpace::Local);
-		TransformHandler.Transform(Pos, InstData->Transform);
+		TransformHandler.TransformPosition(Pos, InstData->Transform);
 
 		*OutPosX.GetDest() = Pos.X;
 		*OutPosY.GetDest() = Pos.Y;
@@ -487,7 +463,7 @@ void UNiagaraDataInterfaceSpline::SampleSplineRightVectorByUnitDistance(FVectorV
 		float DistanceUnitDistance = SplineSampleParam.Get();
 
 		FVector Pos = InstData->Component->GetRightVectorAtDistanceAlongSpline(DistanceUnitDistance * InstData->Component->GetSplineLength(), ESplineCoordinateSpace::Local);
-		TransformHandler.Transform(Pos, InstData->Transform);
+		TransformHandler.TransformPosition(Pos, InstData->Transform);
 
 		*OutPosX.GetDest() = Pos.X;
 		*OutPosY.GetDest() = Pos.Y;
@@ -515,7 +491,7 @@ void UNiagaraDataInterfaceSpline::SampleSplineTangentByUnitDistance(FVectorVMCon
 		float DistanceUnitDistance = SplineSampleParam.Get();
 
 		FVector Pos = InstData->Component->GetTangentAtDistanceAlongSpline(DistanceUnitDistance * InstData->Component->GetSplineLength(), ESplineCoordinateSpace::Local);
-		TransformHandler.Transform(Pos, InstData->Transform);
+		TransformHandler.TransformPosition(Pos, InstData->Transform);
 
 		*OutPosX.GetDest() = Pos.X;
 		*OutPosY.GetDest() = Pos.Y;
@@ -544,7 +520,7 @@ void UNiagaraDataInterfaceSpline::SampleSplineDirectionByUnitDistance(FVectorVMC
 		float DistanceUnitDistance = SplineSampleParam.Get();
 
 		FVector Pos = InstData->Component->GetDirectionAtDistanceAlongSpline(DistanceUnitDistance * InstData->Component->GetSplineLength(), ESplineCoordinateSpace::Local);
-		TransformHandler.Transform(Pos, InstData->Transform);
+		TransformHandler.TransformPosition(Pos, InstData->Transform);
 
 		*OutPosX.GetDest() = Pos.X;
 		*OutPosY.GetDest() = Pos.Y;
diff --git a/Engine/Plugins/FX/Niagara/Source/Niagara/Private/NiagaraDataInterfaceStaticMesh.cpp b/Engine/Plugins/FX/Niagara/Source/Niagara/Private/NiagaraDataInterfaceStaticMesh.cpp
--- a/Engine/Plugins/FX/Niagara/Source/Niagara/Private/NiagaraDataInterfaceStaticMesh.cpp
+++ b/Engine/Plugins/FX/Niagara/Source/Niagara/Private/NiagaraDataInterfaceStaticMesh.cpp
@@ -155,6 +155,10 @@ bool FNDIStaticMesh_InstanceData::Init(UNiagaraDataInterfaceStaticMesh* Interfac
 		return false;
 	}
 
+#if WITH_EDITOR
+	Mesh->GetOnMeshChanged().AddUObject(SystemInstance->GetComponent(), &UNiagaraComponent::ReinitializeSystem);
+#endif
+
 	bIsAreaWeightedSampling = Mesh->bSupportUniformlyDistributedSampling;
 
 	//Init the instance filter
@@ -283,21 +287,6 @@ bool FNDIStaticMesh_InstanceData::Tick(UNiagaraDataInterfaceStaticMesh* Interfac
 //////////////////////////////////////////////////////////////////////////
 
 
-struct FTransformHandlerNoop
-{
-	FORCEINLINE void Transform(FVector& V, FMatrix& M) {  }
-};
-
-struct FTransformHandlerPosition
-{
-	FORCEINLINE void Transform(FVector& P, FMatrix& M) { P = M.TransformPosition(P); }
-};
-
-struct FTransformHandlerVector
-{
-	FORCEINLINE void Transform(FVector& V, FMatrix& M) { V = M.TransformVector(V).GetUnsafeNormal3(); }
-};
-
 UNiagaraDataInterfaceStaticMesh::UNiagaraDataInterfaceStaticMesh(FObjectInitializer const& ObjectInitializer)
 	: Super(ObjectInitializer)
 	, DefaultMesh(nullptr)
@@ -327,35 +316,37 @@ void UNiagaraDataInterfaceStaticMesh::PostInitProperties()
 
 #endif //WITH_EDITOR
 
+namespace StaticMeshHelpers
+{
+	static const FName RandomSectionName("RandomSection");
+	static const FName RandomTriCoordName("RandomTriCoord");
+	static const FName RandomTriCoordOnSectionName("RandomTriCoordOnSection");
+	static const FName RandomTriCoordVCFilteredName("RandomTriCoordUsingVertexColorFilter");
 
-static const FName RandomSectionName("RandomSection");
-static const FName RandomTriCoordName("RandomTriCoord");
-static const FName RandomTriCoordOnSectionName("RandomTriCoordOnSection");
-static const FName RandomTriCoordVCFilteredName("RandomTriCoordUsingVertexColorFilter");
-
-static const FName GetTriPositionName("GetTriPosition");
-static const FName GetTriNormalName("GetTriNormal");
-static const FName GetTriTangentsName("GetTriTangents");
+	static const FName GetTriPositionName("GetTriPosition");
+	static const FName GetTriNormalName("GetTriNormal");
+	static const FName GetTriTangentsName("GetTriTangents");
 
-static const FName GetTriPositionWSName("GetTriPositionWS");
-static const FName GetTriNormalWSName("GetTriNormalWS");
-static const FName GetTriTangentsWSName("GetTriTangentsWS");
+	static const FName GetTriPositionWSName("GetTriPositionWS");
+	static const FName GetTriNormalWSName("GetTriNormalWS");
+	static const FName GetTriTangentsWSName("GetTriTangentsWS");
 
-static const FName GetTriColorName("GetTriColor");
-static const FName GetTriUVName("GetTriUV");
+	static const FName GetTriColorName("GetTriColor");
+	static const FName GetTriUVName("GetTriUV");
 
-static const FName GetTriPositionAndVelocityName("GetTriPositionAndVelocityWS");
+	static const FName GetTriPositionAndVelocityName("GetTriPositionAndVelocityWS");
 
-/** Temporary solution for exposing the transform of a mesh. Ideally this would be done by allowing interfaces to add to the uniform set for a simulation. */
-static const FName GetMeshLocalToWorldName("GetLocalToWorld");
-static const FName GetMeshLocalToWorldInverseTransposedName("GetMeshLocalToWorldInverseTransposed");
-static const FName GetMeshWorldVelocityName("GetWorldVelocity");
+	/** Temporary solution for exposing the transform of a mesh. Ideally this would be done by allowing interfaces to add to the uniform set for a simulation. */
+	static const FName GetMeshLocalToWorldName("GetLocalToWorld");
+	static const FName GetMeshLocalToWorldInverseTransposedName("GetMeshLocalToWorldInverseTransposed");
+	static const FName GetMeshWorldVelocityName("GetWorldVelocity");
+};
 
 void UNiagaraDataInterfaceStaticMesh::GetFunctions(TArray<FNiagaraFunctionSignature>& OutFunctions)
 {
 	{
 		FNiagaraFunctionSignature Sig;
-		Sig.Name = RandomSectionName;
+		Sig.Name = StaticMeshHelpers::RandomSectionName;
 		Sig.Inputs.Add(FNiagaraVariable(FNiagaraTypeDefinition(GetClass()), TEXT("StaticMesh")));
 		Sig.Outputs.Add(FNiagaraVariable(FNiagaraTypeDefinition::GetIntDef(), TEXT("Section")));
 		Sig.bMemberFunction = true;
@@ -366,7 +357,7 @@ void UNiagaraDataInterfaceStaticMesh::GetFunctions(TArray<FNiagaraFunctionSignat
 
 	{
 		FNiagaraFunctionSignature Sig;
-		Sig.Name = RandomTriCoordName;
+		Sig.Name = StaticMeshHelpers::RandomTriCoordName;
 		Sig.Inputs.Add(FNiagaraVariable(FNiagaraTypeDefinition(GetClass()), TEXT("StaticMesh")));
 		Sig.Outputs.Add(FNiagaraVariable(FNiagaraTypeDefinition(FMeshTriCoordinate::StaticStruct()), TEXT("Coord")));
 		//Sig.Outputs.Add(FNiagaraVariable(FNiagaraTypeDefinition::GetIntDef(), TEXT("Tri Index")));
@@ -379,7 +370,7 @@ void UNiagaraDataInterfaceStaticMesh::GetFunctions(TArray<FNiagaraFunctionSignat
 
 	{
 		FNiagaraFunctionSignature Sig;
-		Sig.Name = RandomTriCoordVCFilteredName;
+		Sig.Name = StaticMeshHelpers::RandomTriCoordVCFilteredName;
 		Sig.Inputs.Add(FNiagaraVariable(FNiagaraTypeDefinition(GetClass()), TEXT("StaticMesh")));
 		Sig.Inputs.Add(FNiagaraVariable(FNiagaraTypeDefinition(FNiagaraTypeDefinition::GetFloatDef()), TEXT("Start")));
 		Sig.Inputs.Add(FNiagaraVariable(FNiagaraTypeDefinition(FNiagaraTypeDefinition::GetFloatDef()), TEXT("Range")));
@@ -395,7 +386,7 @@ void UNiagaraDataInterfaceStaticMesh::GetFunctions(TArray<FNiagaraFunctionSignat
 
 	{
 		FNiagaraFunctionSignature Sig;
-		Sig.Name = RandomTriCoordOnSectionName;
+		Sig.Name = StaticMeshHelpers::RandomTriCoordOnSectionName;
 		Sig.Inputs.Add(FNiagaraVariable(FNiagaraTypeDefinition(GetClass()), TEXT("StaticMesh")));
 		Sig.Inputs.Add(FNiagaraVariable(FNiagaraTypeDefinition::GetIntDef(), TEXT("Section")));
 		Sig.Outputs.Add(FNiagaraVariable(FNiagaraTypeDefinition(FMeshTriCoordinate::StaticStruct()), TEXT("Coord")));
@@ -409,7 +400,7 @@ void UNiagaraDataInterfaceStaticMesh::GetFunctions(TArray<FNiagaraFunctionSignat
 
 	{
 		FNiagaraFunctionSignature Sig;
-		Sig.Name = GetTriPositionName;
+		Sig.Name = StaticMeshHelpers::GetTriPositionName;
 		Sig.Inputs.Add(FNiagaraVariable(FNiagaraTypeDefinition(GetClass()), TEXT("StaticMesh")));
 		Sig.Inputs.Add(FNiagaraVariable(FNiagaraTypeDefinition(FMeshTriCoordinate::StaticStruct()), TEXT("Coord")));
 		//Sig.Inputs.Add(FNiagaraVariable(FNiagaraTypeDefinition::GetIntDef(), TEXT("Tri Index")));
@@ -424,7 +415,7 @@ void UNiagaraDataInterfaceStaticMesh::GetFunctions(TArray<FNiagaraFunctionSignat
 
 	{
 		FNiagaraFunctionSignature Sig;
-		Sig.Name = GetTriPositionAndVelocityName;
+		Sig.Name = StaticMeshHelpers::GetTriPositionAndVelocityName;
 		Sig.Inputs.Add(FNiagaraVariable(FNiagaraTypeDefinition(GetClass()), TEXT("StaticMesh")));
 		Sig.Inputs.Add(FNiagaraVariable(FNiagaraTypeDefinition(FMeshTriCoordinate::StaticStruct()), TEXT("Coord")));
 		//Sig.Inputs.Add(FNiagaraVariable(FNiagaraTypeDefinition::GetIntDef(), TEXT("Tri Index")));
@@ -440,7 +431,7 @@ void UNiagaraDataInterfaceStaticMesh::GetFunctions(TArray<FNiagaraFunctionSignat
 
 	{
 		FNiagaraFunctionSignature Sig;
-		Sig.Name = GetTriPositionWSName;
+		Sig.Name = StaticMeshHelpers::GetTriPositionWSName;
 		Sig.Inputs.Add(FNiagaraVariable(FNiagaraTypeDefinition(GetClass()), TEXT("StaticMesh")));
 		Sig.Inputs.Add(FNiagaraVariable(FNiagaraTypeDefinition(FMeshTriCoordinate::StaticStruct()), TEXT("Coord")));
 		//Sig.Inputs.Add(FNiagaraVariable(FNiagaraTypeDefinition::GetIntDef(), TEXT("Tri Index")));
@@ -455,7 +446,7 @@ void UNiagaraDataInterfaceStaticMesh::GetFunctions(TArray<FNiagaraFunctionSignat
 
 	{
 		FNiagaraFunctionSignature Sig;
-		Sig.Name = GetTriNormalName;
+		Sig.Name = StaticMeshHelpers::GetTriNormalName;
 		Sig.Inputs.Add(FNiagaraVariable(FNiagaraTypeDefinition(GetClass()), TEXT("StaticMesh")));
 		Sig.Inputs.Add(FNiagaraVariable(FNiagaraTypeDefinition(FMeshTriCoordinate::StaticStruct()), TEXT("Coord")));
 		//Sig.Inputs.Add(FNiagaraVariable(FNiagaraTypeDefinition::GetIntDef(), TEXT("Tri Index")));
@@ -470,7 +461,7 @@ void UNiagaraDataInterfaceStaticMesh::GetFunctions(TArray<FNiagaraFunctionSignat
 
 	{
 		FNiagaraFunctionSignature Sig;
-		Sig.Name = GetTriNormalWSName;
+		Sig.Name = StaticMeshHelpers::GetTriNormalWSName;
 		Sig.Inputs.Add(FNiagaraVariable(FNiagaraTypeDefinition(GetClass()), TEXT("StaticMesh")));
 		Sig.Inputs.Add(FNiagaraVariable(FNiagaraTypeDefinition(FMeshTriCoordinate::StaticStruct()), TEXT("Coord")));
 		//Sig.Inputs.Add(FNiagaraVariable(FNiagaraTypeDefinition::GetIntDef(), TEXT("Tri Index")));
@@ -485,7 +476,7 @@ void UNiagaraDataInterfaceStaticMesh::GetFunctions(TArray<FNiagaraFunctionSignat
 
 	{
 		FNiagaraFunctionSignature Sig;
-		Sig.Name = GetTriTangentsName;
+		Sig.Name = StaticMeshHelpers::GetTriTangentsName;
 		Sig.Inputs.Add(FNiagaraVariable(FNiagaraTypeDefinition(GetClass()), TEXT("StaticMesh")));
 		Sig.Inputs.Add(FNiagaraVariable(FNiagaraTypeDefinition(FMeshTriCoordinate::StaticStruct()), TEXT("Coord")));
 		//Sig.Inputs.Add(FNiagaraVariable(FNiagaraTypeDefinition::GetIntDef(), TEXT("Tri Index")));
@@ -502,7 +493,7 @@ void UNiagaraDataInterfaceStaticMesh::GetFunctions(TArray<FNiagaraFunctionSignat
 
 	{
 		FNiagaraFunctionSignature Sig;
-		Sig.Name = GetTriTangentsWSName;
+		Sig.Name = StaticMeshHelpers::GetTriTangentsWSName;
 		Sig.Inputs.Add(FNiagaraVariable(FNiagaraTypeDefinition(GetClass()), TEXT("StaticMesh")));
 		Sig.Inputs.Add(FNiagaraVariable(FNiagaraTypeDefinition(FMeshTriCoordinate::StaticStruct()), TEXT("Coord")));
 		//Sig.Inputs.Add(FNiagaraVariable(FNiagaraTypeDefinition::GetIntDef(), TEXT("Tri Index")));
@@ -519,7 +510,7 @@ void UNiagaraDataInterfaceStaticMesh::GetFunctions(TArray<FNiagaraFunctionSignat
 
 	{
 		FNiagaraFunctionSignature Sig;
-		Sig.Name = GetTriColorName;
+		Sig.Name = StaticMeshHelpers::GetTriColorName;
 		Sig.Inputs.Add(FNiagaraVariable(FNiagaraTypeDefinition(GetClass()), TEXT("StaticMesh")));
 		Sig.Inputs.Add(FNiagaraVariable(FNiagaraTypeDefinition(FMeshTriCoordinate::StaticStruct()), TEXT("Coord")));
 		//Sig.Inputs.Add(FNiagaraVariable(FNiagaraTypeDefinition::GetIntDef(), TEXT("Tri Index")));
@@ -534,7 +525,7 @@ void UNiagaraDataInterfaceStaticMesh::GetFunctions(TArray<FNiagaraFunctionSignat
 
 	{
 		FNiagaraFunctionSignature Sig;
-		Sig.Name = GetTriUVName;
+		Sig.Name = StaticMeshHelpers::GetTriUVName;
 		Sig.Inputs.Add(FNiagaraVariable(FNiagaraTypeDefinition(GetClass()), TEXT("StaticMesh")));
 		Sig.Inputs.Add(FNiagaraVariable(FNiagaraTypeDefinition(FMeshTriCoordinate::StaticStruct()), TEXT("Coord")));
 		//Sig.Inputs.Add(FNiagaraVariable(FNiagaraTypeDefinition::GetIntDef(), TEXT("Tri Index")));
@@ -550,7 +541,7 @@ void UNiagaraDataInterfaceStaticMesh::GetFunctions(TArray<FNiagaraFunctionSignat
 
 	{
 		FNiagaraFunctionSignature Sig;
-		Sig.Name = GetMeshLocalToWorldName;
+		Sig.Name = StaticMeshHelpers::GetMeshLocalToWorldName;
 		Sig.Inputs.Add(FNiagaraVariable(FNiagaraTypeDefinition(GetClass()), TEXT("StaticMesh")));
 		//Sig.Inputs.Add(FNiagaraVariable(FNiagaraTypeDefinition(FMeshTriCoordinate::StaticStruct()), TEXT("Coord")));
 		Sig.Outputs.Add(FNiagaraVariable(FNiagaraTypeDefinition::GetMatrix4Def(), TEXT("Transform")));
@@ -562,7 +553,7 @@ void UNiagaraDataInterfaceStaticMesh::GetFunctions(TArray<FNiagaraFunctionSignat
 
 	{
 		FNiagaraFunctionSignature Sig;
-		Sig.Name = GetMeshLocalToWorldInverseTransposedName;
+		Sig.Name = StaticMeshHelpers::GetMeshLocalToWorldInverseTransposedName;
 		Sig.Inputs.Add(FNiagaraVariable(FNiagaraTypeDefinition(GetClass()), TEXT("StaticMesh")));
 		//Sig.Inputs.Add(FNiagaraVariable(FNiagaraTypeDefinition(FMeshTriCoordinate::StaticStruct()), TEXT("Coord")));
 		Sig.Outputs.Add(FNiagaraVariable(FNiagaraTypeDefinition::GetMatrix4Def(), TEXT("Transform")));
@@ -574,7 +565,7 @@ void UNiagaraDataInterfaceStaticMesh::GetFunctions(TArray<FNiagaraFunctionSignat
 
 	{
 		FNiagaraFunctionSignature Sig;
-		Sig.Name = GetMeshWorldVelocityName;
+		Sig.Name = StaticMeshHelpers::GetMeshWorldVelocityName;
 		Sig.Inputs.Add(FNiagaraVariable(FNiagaraTypeDefinition(GetClass()), TEXT("StaticMesh")));
 		//Sig.Inputs.Add(FNiagaraVariable(FNiagaraTypeDefinition(FMeshTriCoordinate::StaticStruct()), TEXT("Coord")));
 		Sig.Outputs.Add(FNiagaraVariable(FNiagaraTypeDefinition::GetVec3Def(), TEXT("Velocity")));
@@ -678,94 +669,94 @@ FVMExternalFunction UNiagaraDataInterfaceStaticMesh::GetVMExternalFunction(const
 	bool bNeedsVertMain = true;//Assuming we always need this?
 	
 	FVMExternalFunction Function;
-	if (BindingInfo.Name == RandomSectionName)
+	if (BindingInfo.Name == StaticMeshHelpers::RandomSectionName)
 	{
 		check(BindingInfo.GetNumInputs() == 1 && BindingInfo.GetNumOutputs() == 1);
 		Function = TUsesAreaWeightingBinder<NDI_FUNC_BINDER(UNiagaraDataInterfaceStaticMesh, RandomSection)>::Bind(this, BindingInfo, InstanceData);
 	}
-	else if (BindingInfo.Name == RandomTriCoordName)
+	else if (BindingInfo.Name == StaticMeshHelpers::RandomTriCoordName)
 	{
 		check(BindingInfo.GetNumInputs() == 1 && BindingInfo.GetNumOutputs() == 4);
 		Function = TUsesAreaWeightingBinder<NDI_FUNC_BINDER(UNiagaraDataInterfaceStaticMesh, RandomTriCoord)>::Bind(this, BindingInfo, InstanceData);
 	}
 	//TODO: Vertex color filtering needs more work.
-	else if (BindingInfo.Name == RandomTriCoordVCFilteredName)
+	else if (BindingInfo.Name == StaticMeshHelpers::RandomTriCoordVCFilteredName)
 	{
 		InstData->InitVertexColorFiltering();
 		check(BindingInfo.GetNumInputs() == 3 && BindingInfo.GetNumOutputs() == 4);
 		Function = TNDIParamBinder<0, float, TNDIParamBinder<1, float, NDI_FUNC_BINDER(UNiagaraDataInterfaceStaticMesh, RandomTriCoordVertexColorFiltered)>>::Bind(this, BindingInfo, InstanceData);
 	}	
-	else if (BindingInfo.Name == RandomTriCoordOnSectionName)
+	else if (BindingInfo.Name == StaticMeshHelpers::RandomTriCoordOnSectionName)
 	{
 		check(BindingInfo.GetNumInputs() == 2 && BindingInfo.GetNumOutputs() == 4);
 		Function = TUsesAreaWeightingBinder<TNDIParamBinder<0, int32, NDI_FUNC_BINDER(UNiagaraDataInterfaceStaticMesh, RandomTriCoordOnSection)>>::Bind(this, BindingInfo, InstanceData);
 	}
-	if (BindingInfo.Name == GetTriPositionName)
+	if (BindingInfo.Name == StaticMeshHelpers::GetTriPositionName)
 	{
 		check(BindingInfo.GetNumInputs() == 5 && BindingInfo.GetNumOutputs() == 3);
 		bNeedsVertexPositions = true;
-		Function = TNDIExplicitBinder<FTransformHandlerNoop, TNDIParamBinder<0, int32, TNDIParamBinder<1, float, TNDIParamBinder<2, float, TNDIParamBinder<3, float, NDI_FUNC_BINDER(UNiagaraDataInterfaceStaticMesh, GetTriCoordPosition)>>>>>::Bind(this, BindingInfo, InstanceData);
+		Function = TNDIExplicitBinder<FNDITransformHandlerNoop, TNDIParamBinder<0, int32, TNDIParamBinder<1, float, TNDIParamBinder<2, float, TNDIParamBinder<3, float, NDI_FUNC_BINDER(UNiagaraDataInterfaceStaticMesh, GetTriCoordPosition)>>>>>::Bind(this, BindingInfo, InstanceData);
 	}
-	else if (BindingInfo.Name == GetTriPositionWSName)
+	else if (BindingInfo.Name == StaticMeshHelpers::GetTriPositionWSName)
 	{
 		check(BindingInfo.GetNumInputs() == 5 && BindingInfo.GetNumOutputs() == 3);
 		bNeedsVertexPositions = true;
-		Function = TNDIExplicitBinder<FTransformHandlerPosition, TNDIParamBinder<0, int32, TNDIParamBinder<1, float, TNDIParamBinder<2, float, TNDIParamBinder<3, float, NDI_FUNC_BINDER(UNiagaraDataInterfaceStaticMesh, GetTriCoordPosition)>>>>>::Bind(this, BindingInfo, InstanceData);
+		Function = TNDIExplicitBinder<FNDITransformHandler, TNDIParamBinder<0, int32, TNDIParamBinder<1, float, TNDIParamBinder<2, float, TNDIParamBinder<3, float, NDI_FUNC_BINDER(UNiagaraDataInterfaceStaticMesh, GetTriCoordPosition)>>>>>::Bind(this, BindingInfo, InstanceData);
 	}
-	else if (BindingInfo.Name == GetTriNormalName)
+	else if (BindingInfo.Name == StaticMeshHelpers::GetTriNormalName)
 	{
 		check(BindingInfo.GetNumInputs() == 6 && BindingInfo.GetNumOutputs() == 3);
 		bNeedsVertMain = true;
-		Function = TNDIExplicitBinder<FTransformHandlerNoop, TNDIParamBinder<0, int32, TNDIParamBinder<1, float, TNDIParamBinder<2, float, TNDIParamBinder<3, float, NDI_FUNC_BINDER(UNiagaraDataInterfaceStaticMesh, GetTriCoordNormal)>>>>>::Bind(this, BindingInfo, InstanceData);
+		Function = TNDIExplicitBinder<FNDITransformHandlerNoop, TNDIParamBinder<0, int32, TNDIParamBinder<1, float, TNDIParamBinder<2, float, TNDIParamBinder<3, float, NDI_FUNC_BINDER(UNiagaraDataInterfaceStaticMesh, GetTriCoordNormal)>>>>>::Bind(this, BindingInfo, InstanceData);
 	}
-	else if (BindingInfo.Name == GetTriNormalWSName)
+	else if (BindingInfo.Name == StaticMeshHelpers::GetTriNormalWSName)
 	{
 		check(BindingInfo.GetNumInputs() == 5 && BindingInfo.GetNumOutputs() == 3);
 		bNeedsVertMain = true;
-		Function = TNDIExplicitBinder<FTransformHandlerVector, TNDIParamBinder<0, int32, TNDIParamBinder<1, float, TNDIParamBinder<2, float, TNDIParamBinder<3, float, NDI_FUNC_BINDER(UNiagaraDataInterfaceStaticMesh, GetTriCoordNormal)>>>>>::Bind(this, BindingInfo, InstanceData);
+		Function = TNDIExplicitBinder<FNDITransformHandler, TNDIParamBinder<0, int32, TNDIParamBinder<1, float, TNDIParamBinder<2, float, TNDIParamBinder<3, float, NDI_FUNC_BINDER(UNiagaraDataInterfaceStaticMesh, GetTriCoordNormal)>>>>>::Bind(this, BindingInfo, InstanceData);
 	}
-	else if (BindingInfo.Name == GetTriTangentsName)
+	else if (BindingInfo.Name == StaticMeshHelpers::GetTriTangentsName)
 	{
 		check(BindingInfo.GetNumInputs() == 5 && BindingInfo.GetNumOutputs() == 9);
 		bNeedsVertMain = true;
-		Function = TTypedMeshAccessorBinder<TNDIExplicitBinder<FTransformHandlerNoop, TNDIParamBinder<0, int32, TNDIParamBinder<1, float, TNDIParamBinder<2, float, TNDIParamBinder<3, float, NDI_FUNC_BINDER(UNiagaraDataInterfaceStaticMesh, GetTriCoordTangents)>>>>>>::Bind(this, BindingInfo, InstanceData);
+		Function = TTypedMeshAccessorBinder<TNDIExplicitBinder<FNDITransformHandlerNoop, TNDIParamBinder<0, int32, TNDIParamBinder<1, float, TNDIParamBinder<2, float, TNDIParamBinder<3, float, NDI_FUNC_BINDER(UNiagaraDataInterfaceStaticMesh, GetTriCoordTangents)>>>>>>::Bind(this, BindingInfo, InstanceData);
 	}
-	else if (BindingInfo.Name == GetTriTangentsWSName)
+	else if (BindingInfo.Name == StaticMeshHelpers::GetTriTangentsWSName)
 	{
 		check(BindingInfo.GetNumInputs() == 5 && BindingInfo.GetNumOutputs() == 9);
 		bNeedsVertMain = true;
-		Function = TTypedMeshAccessorBinder<TNDIExplicitBinder<FTransformHandlerVector, TNDIParamBinder<0, int32, TNDIParamBinder<1, float, TNDIParamBinder<2, float, TNDIParamBinder<3, float, NDI_FUNC_BINDER(UNiagaraDataInterfaceStaticMesh, GetTriCoordTangents)>>>>>>::Bind(this, BindingInfo, InstanceData);
+		Function = TTypedMeshAccessorBinder<TNDIExplicitBinder<FNDITransformHandler, TNDIParamBinder<0, int32, TNDIParamBinder<1, float, TNDIParamBinder<2, float, TNDIParamBinder<3, float, NDI_FUNC_BINDER(UNiagaraDataInterfaceStaticMesh, GetTriCoordTangents)>>>>>>::Bind(this, BindingInfo, InstanceData);
 	}
-	else if (BindingInfo.Name == GetTriColorName)
+	else if (BindingInfo.Name == StaticMeshHelpers::GetTriColorName)
 	{
 		check(BindingInfo.GetNumInputs() == 5 && BindingInfo.GetNumOutputs() == 4);
 		bNeedsVertexColors = true;
 		Function = TNDIParamBinder<0, int32, TNDIParamBinder<1, float, TNDIParamBinder<2, float, TNDIParamBinder<3, float, NDI_FUNC_BINDER(UNiagaraDataInterfaceStaticMesh, GetTriCoordColor)>>>>::Bind(this, BindingInfo, InstanceData);
 	}
-	else if (BindingInfo.Name == GetTriUVName)
+	else if (BindingInfo.Name == StaticMeshHelpers::GetTriUVName)
 	{
 		check(BindingInfo.GetNumInputs() == 6 && BindingInfo.GetNumOutputs() == 2);
 		bNeedsVertMain = true;
 		Function = TTypedMeshAccessorBinder<TNDIParamBinder<0, int32, TNDIParamBinder<1, float, TNDIParamBinder<2, float, TNDIParamBinder<3, float, TNDIParamBinder<4, int32, NDI_FUNC_BINDER(UNiagaraDataInterfaceStaticMesh, GetTriCoordUV)>>>>>>::Bind(this, BindingInfo, InstanceData);
 	}
-	else if (BindingInfo.Name == GetTriPositionAndVelocityName)
+	else if (BindingInfo.Name == StaticMeshHelpers::GetTriPositionAndVelocityName)
 	{
 		check(BindingInfo.GetNumInputs() == 5 && BindingInfo.GetNumOutputs() == 6);
 		bNeedsVertMain = true;
 		bNeedsVertexPositions = true;
 		Function = TNDIParamBinder<0, int32, TNDIParamBinder<1, float, TNDIParamBinder<2, float, TNDIParamBinder<3, float, NDI_FUNC_BINDER(UNiagaraDataInterfaceStaticMesh, GetTriCoordPositionAndVelocity)>>>>::Bind(this, BindingInfo, InstanceData);
 	}
-	else if (BindingInfo.Name == GetMeshLocalToWorldName)
+	else if (BindingInfo.Name == StaticMeshHelpers::GetMeshLocalToWorldName)
 	{
 		check(BindingInfo.GetNumInputs() == 1 && BindingInfo.GetNumOutputs() == 16);
 		Function = FVMExternalFunction::CreateUObject(this, &UNiagaraDataInterfaceStaticMesh::GetLocalToWorld);
 	}
-	else if (BindingInfo.Name == GetMeshLocalToWorldInverseTransposedName)
+	else if (BindingInfo.Name == StaticMeshHelpers::GetMeshLocalToWorldInverseTransposedName)
 	{
 		check(BindingInfo.GetNumInputs() == 1 && BindingInfo.GetNumOutputs() == 16);
 		Function = FVMExternalFunction::CreateUObject(this, &UNiagaraDataInterfaceStaticMesh::GetLocalToWorldInverseTransposed);
 	}
-	else if (BindingInfo.Name == GetMeshWorldVelocityName)
+	else if (BindingInfo.Name == StaticMeshHelpers::GetMeshWorldVelocityName)
 	{
 		check(BindingInfo.GetNumInputs() == 1 && BindingInfo.GetNumOutputs() == 3);
 		Function = FVMExternalFunction::CreateUObject(this, &UNiagaraDataInterfaceStaticMesh::GetWorldVelocity);
@@ -790,9 +781,9 @@ FVMExternalFunction UNiagaraDataInterfaceStaticMesh::GetVMExternalFunction(const
 	return Function;
 }
 
-bool UNiagaraDataInterfaceStaticMesh::CopyTo(UNiagaraDataInterface* Destination) const
+bool UNiagaraDataInterfaceStaticMesh::CopyToInternal(UNiagaraDataInterface* Destination) const
 {
-	if (!Super::CopyTo(Destination))
+	if (!Super::CopyToInternal(Destination))
 	{
 		return false;
 	}
@@ -827,6 +818,14 @@ bool UNiagaraDataInterfaceStaticMesh::InitPerInstanceData(void* PerInstanceData,
 void UNiagaraDataInterfaceStaticMesh::DestroyPerInstanceData(void* PerInstanceData, FNiagaraSystemInstance* SystemInstance)
 {
 	FNDIStaticMesh_InstanceData* Inst = (FNDIStaticMesh_InstanceData*)PerInstanceData;
+
+#if WITH_EDITOR
+	if (Inst->Mesh)
+	{
+		Inst->Mesh->GetOnMeshChanged().RemoveAll(SystemInstance->GetComponent());
+	}
+#endif
+
 	Inst->~FNDIStaticMesh_InstanceData();
 }
 
@@ -836,29 +835,6 @@ bool UNiagaraDataInterfaceStaticMesh::PerInstanceTick(void* PerInstanceData, FNi
 	return Inst->Tick(this, SystemInstance, InDeltaSeconds);
 }
 
-FORCEINLINE FVector RandomBarycentricCoord(FRandomStream& RandStream)
-{
-	//TODO: This is gonna be slooooow. Move to an LUT possibly or find faster method.
-	//Can probably handle lower quality randoms / uniformity for a decent speed win.
-	float r0 = RandStream.GetFraction();
-	float r1 = RandStream.GetFraction();
-	float sqrt0 = FMath::Sqrt(r0);
-	float sqrt1 = FMath::Sqrt(r1);
-	return FVector(1.0f - sqrt0, sqrt0 * (1.0 - r1), r1 * sqrt0);
-}
-
-template<typename T>
-T BarycentricInterpolate(float BaryX, float BaryY, float BaryZ, T V0, T V1, T V2)
-{
-	return V0 * BaryX + V1 * BaryY + V2 * BaryZ;
-}
-
-// Overload for FVector4 to work around C2719: (formal parameter with requested alignment of 16 won't be aligned)
-FVector4 BarycentricInterpolate(float BaryX, float BaryY, float BaryZ, const FVector4& V0, const FVector4& V1, const FVector4& V2)
-{
-	return V0 * BaryX + V1 * BaryY + V2 * BaryZ;
-}
-
 //RandomSection specializations.
 //Each combination for AreaWeighted and Section filtered.
 template<>
@@ -1103,7 +1079,7 @@ void UNiagaraDataInterfaceStaticMesh::GetTriCoordPosition(FVectorVMContext& Cont
 		int32 Idx2 = Indices[Tri + 2];
 
 		FVector Pos = BarycentricInterpolate(BaryXParam.Get(), BaryYParam.Get(), BaryZParam.Get(), Positions.VertexPosition(Idx0), Positions.VertexPosition(Idx1), Positions.VertexPosition(Idx2));
-		TransformHandler.Transform(Pos, InstData->Transform);
+		TransformHandler.TransformPosition(Pos, InstData->Transform);
 
 		*OutPosX.GetDest() = Pos.X;
 		*OutPosY.GetDest() = Pos.Y;
@@ -1146,7 +1122,7 @@ void UNiagaraDataInterfaceStaticMesh::GetTriCoordNormal(FVectorVMContext& Contex
 		int32 Idx2 = Indices[Tri + 2];
 
 		FVector Norm = BarycentricInterpolate(BaryXParam.Get(), BaryYParam.Get(), BaryZParam.Get(), Verts.VertexTangentZ(Idx0), Verts.VertexTangentZ(Idx1), Verts.VertexTangentZ(Idx2));
-		TransformHandler.Transform(Norm, InstData->Transform);
+		TransformHandler.TransformVector(Norm, InstData->TransformInverseTransposed);
 
 		*OutNormX.GetDest() = Norm.X;
 		*OutNormY.GetDest() = Norm.Y;
@@ -1195,9 +1171,9 @@ void UNiagaraDataInterfaceStaticMesh::GetTriCoordTangents(FVectorVMContext& Cont
 		FVector Tangent = BarycentricInterpolate(BaryXParam.Get(), BaryYParam.Get(), BaryZParam.Get(), Verts.GetTangentX(Idx0), Verts.GetTangentX(Idx1), Verts.GetTangentX(Idx2));
 		FVector Binorm = BarycentricInterpolate(BaryXParam.Get(), BaryYParam.Get(), BaryZParam.Get(), Verts.GetTangentY(Idx0), Verts.GetTangentY(Idx1), Verts.GetTangentY(Idx2));
 		FVector Norm = BarycentricInterpolate(BaryXParam.Get(), BaryYParam.Get(), BaryZParam.Get(), Verts.GetTangentZ(Idx0), Verts.GetTangentZ(Idx1), Verts.GetTangentZ(Idx2));
-		TransformHandler.Transform(Tangent, InstData->TransformInverseTransposed);
-		TransformHandler.Transform(Binorm, InstData->TransformInverseTransposed);
-		TransformHandler.Transform(Norm, InstData->TransformInverseTransposed);
+		TransformHandler.TransformVector(Tangent, InstData->TransformInverseTransposed);
+		TransformHandler.TransformVector(Binorm, InstData->TransformInverseTransposed);
+		TransformHandler.TransformVector(Norm, InstData->TransformInverseTransposed);
 		*OutTangentX.GetDest() = Tangent.X;
 		*OutTangentY.GetDest() = Tangent.Y;
 		*OutTangentZ.GetDest() = Tangent.Z;
