Copying Orion-Staging @ 3748653 (Orion/Dev-General @ 3722124) to //UE4/Main

diff --git a/Engine/Plugins/FX/Niagara/Source/Niagara/Private/NiagaraDataInterfaceVector2DCurve.cpp b/Engine/Plugins/FX/Niagara/Source/Niagara/Private/NiagaraDataInterfaceVector2DCurve.cpp
--- a/Engine/Plugins/FX/Niagara/Source/Niagara/Private/NiagaraDataInterfaceVector2DCurve.cpp
+++ b/Engine/Plugins/FX/Niagara/Source/Niagara/Private/NiagaraDataInterfaceVector2DCurve.cpp
@@ -5,6 +5,7 @@
 #include "Curves/CurveLinearColor.h"
 #include "Curves/CurveFloat.h"
 #include "NiagaraTypes.h"
+#include "NiagaraCustomVersion.h"
 
 
 //////////////////////////////////////////////////////////////////////////
@@ -30,15 +31,47 @@ void UNiagaraDataInterfaceVector2DCurve::PostInitProperties()
 void UNiagaraDataInterfaceVector2DCurve::PostLoad()
 {
 	Super::PostLoad();
+
+	const int32 NiagaraVer = GetLinkerCustomVersion(FNiagaraCustomVersion::GUID);
+
+	if (NiagaraVer < FNiagaraCustomVersion::CurveLUTNowOnByDefault)
+	{
+		UpdateLUT();
+	}
+	TArray<float> OldLUT = ShaderLUT;
 	UpdateLUT();
+
+	if (!CompareLUTS(OldLUT))
+	{
+		UE_LOG(LogNiagara, Log, TEXT("PostLoad LUT generation is out of sync. Please investigate. %s"), *GetPathName());
+	}
+
 }
 
 void UNiagaraDataInterfaceVector2DCurve::UpdateLUT()
 {
 	ShaderLUT.Empty();
+	if (bAllowUnnormalizedLUT && (XCurve.GetNumKeys() > 0 || YCurve.GetNumKeys() > 0))
+	{
+		LUTMinTime = FLT_MAX;
+		LUTMinTime = FMath::Min(XCurve.GetNumKeys() > 0 ? XCurve.GetFirstKey().Time : LUTMinTime, LUTMinTime);
+		LUTMinTime = FMath::Min(YCurve.GetNumKeys() > 0 ? YCurve.GetFirstKey().Time : LUTMinTime, LUTMinTime);
+
+		LUTMaxTime = FLT_MIN;
+		LUTMaxTime = FMath::Max(XCurve.GetNumKeys() > 0 ? XCurve.GetLastKey().Time : LUTMaxTime, LUTMaxTime);
+		LUTMaxTime = FMath::Max(YCurve.GetNumKeys() > 0 ? YCurve.GetLastKey().Time : LUTMaxTime, LUTMaxTime);
+		LUTInvTimeRange = 1.0f / (LUTMaxTime - LUTMinTime);
+	}
+	else
+	{
+		LUTMinTime = 0.0f;
+		LUTMaxTime = 1.0f;
+		LUTInvTimeRange = 1.0f;
+	}
+
 	for (uint32 i = 0; i < CurveLUTWidth; i++)
 	{
-		float X = i / (float)CurveLUTWidth;
+		float X = UnnormalizeTime(i / (float)CurveLUTWidth);
 		FVector2D C(XCurve.Eval(X), YCurve.Eval(X));
 		ShaderLUT.Add(C.X);
 		ShaderLUT.Add(C.Y);
@@ -60,15 +93,15 @@ void UNiagaraDataInterfaceVector2DCurve::PostEditChangeProperty(struct FProperty
 			XCurve = VectorCurveAsset->FloatCurves[0];
 			YCurve = VectorCurveAsset->FloatCurves[1];
 		}
-		UpdateLUT();
 	}
 
+	UpdateLUT();
 }
 #endif
 
-bool UNiagaraDataInterfaceVector2DCurve::CopyTo(UNiagaraDataInterface* Destination) const
+bool UNiagaraDataInterfaceVector2DCurve::CopyToInternal(UNiagaraDataInterface* Destination) const
 {
-	if (!Super::CopyTo(Destination))
+	if (!Super::CopyToInternal(Destination))
 	{
 		return false;
 	}
@@ -76,7 +109,7 @@ bool UNiagaraDataInterfaceVector2DCurve::CopyTo(UNiagaraDataInterface* Destinati
 	DestinationVector2DCurve->XCurve = XCurve;
 	DestinationVector2DCurve->YCurve = YCurve;
 	DestinationVector2DCurve->UpdateLUT();
-
+	ensure(CompareLUTS(DestinationVector2DCurve->ShaderLUT));
 	return true;
 }
 
@@ -116,10 +149,10 @@ void UNiagaraDataInterfaceVector2DCurve::GetFunctions(TArray<FNiagaraFunctionSig
 // the HLSL in the spirit of a static switch
 // TODO: need a way to identify each specific function here
 // 
-bool UNiagaraDataInterfaceVector2DCurve::GetFunctionHLSL(FString FunctionName, TArray<DIGPUBufferParamDescriptor> &Descriptors, FString &HLSLInterfaceID, FString &OutHLSL)
+bool UNiagaraDataInterfaceVector2DCurve::GetFunctionHLSL(const FName& DefinitionFunctionName, FString InstanceFunctionName, TArray<FDIGPUBufferParamDescriptor> &Descriptors, FString &HLSLInterfaceID, FString &OutHLSL)
 {
 	FString BufferName = Descriptors[0].BufferParamName;
-	OutHLSL += TEXT("void ") + FunctionName + TEXT("(in float In_X, out float2 Out_Value) \n{\n");
+	OutHLSL += TEXT("void ") + InstanceFunctionName + TEXT("(in float In_X, out float2 Out_Value) \n{\n");
 	OutHLSL += TEXT("\t Out_Value.x = ") + BufferName + TEXT("[(int)(In_X *") + FString::FromInt(CurveLUTWidth) + TEXT(")* 2 ];");
 	OutHLSL += TEXT("\t Out_Value.y = ") + BufferName + TEXT("[1+ (int)(In_X *") + FString::FromInt(CurveLUTWidth) + TEXT(")* 2 ];");
 	OutHLSL += TEXT("\n}\n");
@@ -133,21 +166,21 @@ bool UNiagaraDataInterfaceVector2DCurve::GetFunctionHLSL(FString FunctionName, T
 // 3. store buffer declaration hlsl in OutHLSL
 // multiple buffers can be defined at once here
 //
-void UNiagaraDataInterfaceVector2DCurve::GetBufferDefinitionHLSL(FString DataInterfaceID, TArray<DIGPUBufferParamDescriptor> &BufferDescriptors, FString &OutHLSL)
+void UNiagaraDataInterfaceVector2DCurve::GetBufferDefinitionHLSL(FString DataInterfaceID, TArray<FDIGPUBufferParamDescriptor> &BufferDescriptors, FString &OutHLSL)
 {
 	FString BufferName = "CurveLUT" + DataInterfaceID;
 	OutHLSL += TEXT("Buffer<float> ") + BufferName + TEXT(";\n");
 
-	BufferDescriptors.Add(DIGPUBufferParamDescriptor(BufferName, 0));		// add a descriptor for shader parameter binding
+	BufferDescriptors.Add(FDIGPUBufferParamDescriptor(BufferName, 0));		// add a descriptor for shader parameter binding
 
 }
 
 // called after translate, to setup buffers matching the buffer descriptors generated during hlsl translation
 // need to do this because the script used during translate is a clone, including its DIs
 //
-void UNiagaraDataInterfaceVector2DCurve::SetupBuffers(TArray<DIGPUBufferParamDescriptor> &BufferDescriptors)
+void UNiagaraDataInterfaceVector2DCurve::SetupBuffers(FDIBufferDescriptorStore &BufferDescriptors)
 {
-	for (DIGPUBufferParamDescriptor &Desc : BufferDescriptors)
+	for (FDIGPUBufferParamDescriptor &Desc : BufferDescriptors.Descriptors)
 	{
 		FNiagaraDataInterfaceBufferData BufferData(*Desc.BufferParamName);	// store off the data for later use
 		GPUBuffers.Add(BufferData);
@@ -165,7 +198,7 @@ TArray<FNiagaraDataInterfaceBufferData> &UNiagaraDataInterfaceVector2DCurve::Get
 
 		FNiagaraDataInterfaceBufferData &GPUBuffer = GPUBuffers[0];
 		GPUBuffer.Buffer.Release();
-		GPUBuffer.Buffer.Initialize(sizeof(float), CurveLUTWidth * 2, EPixelFormat::PF_R32_FLOAT);	// always allocate for up to 64 data sets
+		GPUBuffer.Buffer.Initialize(sizeof(float), CurveLUTWidth * 2, EPixelFormat::PF_R32_FLOAT, BUF_Static);	// always allocate for up to 64 data sets
 		uint32 BufferSize = ShaderLUT.Num() * sizeof(float);
 		int32 *BufferData = static_cast<int32*>(RHILockVertexBuffer(GPUBuffer.Buffer.Buffer, 0, BufferSize, EResourceLockMode::RLM_WriteOnly));
 		FPlatformMemory::Memcpy(BufferData, ShaderLUT.GetData(), BufferSize);
@@ -181,7 +214,7 @@ FVMExternalFunction UNiagaraDataInterfaceVector2DCurve::GetVMExternalFunction(co
 {
 	if (BindingInfo.Name == TEXT("SampleVector2DCurve") && BindingInfo.GetNumInputs() == 1 && BindingInfo.GetNumOutputs() == 2)
 	{
-		return TNDIParamBinder<0, float, NDI_FUNC_BINDER(UNiagaraDataInterfaceVector2DCurve, SampleCurve)>::Bind(this, BindingInfo, InstanceData);
+		return TCurveUseLUTBinder<TNDIParamBinder<0, float, NDI_FUNC_BINDER(UNiagaraDataInterfaceVector2DCurve, SampleCurve)>>::Bind(this, BindingInfo, InstanceData);
 	}
 	else
 	{
@@ -191,7 +224,21 @@ FVMExternalFunction UNiagaraDataInterfaceVector2DCurve::GetVMExternalFunction(co
 	}
 }
 
-template<typename XParamType>
+template<>
+FORCEINLINE_DEBUGGABLE FVector2D UNiagaraDataInterfaceVector2DCurve::SampleCurveInternal<TIntegralConstant<bool, true>>(float X)
+{
+	float NormalizedX = NormalizeTime(X);
+	int32 AccessIdx = FMath::Clamp<int32>(FMath::TruncToInt(NormalizedX * CurveLUTWidthMinusOne), 0, CurveLUTWidthMinusOne) * CurveLUTNumElems;
+	return FVector2D(ShaderLUT[AccessIdx], ShaderLUT[AccessIdx + 1]);
+}
+
+template<>
+FORCEINLINE_DEBUGGABLE FVector2D UNiagaraDataInterfaceVector2DCurve::SampleCurveInternal<TIntegralConstant<bool, false>>(float X)
+{
+	return FVector2D(XCurve.Eval(X), YCurve.Eval(X));
+}
+
+template<typename UseLUT, typename XParamType>
 void UNiagaraDataInterfaceVector2DCurve::SampleCurve(FVectorVMContext& Context)
 {
 	//TODO: Create some SIMDable optimized representation of the curve to do this faster.
@@ -202,7 +249,7 @@ void UNiagaraDataInterfaceVector2DCurve::SampleCurve(FVectorVMContext& Context)
 	for (int32 i = 0; i < Context.NumInstances; ++i)
 	{
 		float X = XParam.Get();
-		FVector2D V(XCurve.Eval(X), YCurve.Eval(X));
+		FVector2D V = SampleCurveInternal<UseLUT>(X);
 		*OutSampleX.GetDest() = V.X;
 		*OutSampleY.GetDest() = V.Y;
 		XParam.Advance();
diff --git a/Engine/Plugins/FX/Niagara/Source/Niagara/Private/NiagaraDataInterfaceVector4Curve.cpp b/Engine/Plugins/FX/Niagara/Source/Niagara/Private/NiagaraDataInterfaceVector4Curve.cpp
--- a/Engine/Plugins/FX/Niagara/Source/Niagara/Private/NiagaraDataInterfaceVector4Curve.cpp
+++ b/Engine/Plugins/FX/Niagara/Source/Niagara/Private/NiagaraDataInterfaceVector4Curve.cpp
@@ -5,6 +5,7 @@
 #include "Curves/CurveLinearColor.h"
 #include "Curves/CurveFloat.h"
 #include "NiagaraTypes.h"
+#include "NiagaraCustomVersion.h"
 
 //////////////////////////////////////////////////////////////////////////
 //Color Curve
@@ -31,7 +32,21 @@ void UNiagaraDataInterfaceVector4Curve::PostInitProperties()
 void UNiagaraDataInterfaceVector4Curve::PostLoad()
 {
 	Super::PostLoad();
+
+	const int32 NiagaraVer = GetLinkerCustomVersion(FNiagaraCustomVersion::GUID);
+
+	if (NiagaraVer < FNiagaraCustomVersion::CurveLUTNowOnByDefault)
+	{
+		UpdateLUT();
+	}
+	TArray<float> OldLUT = ShaderLUT;
 	UpdateLUT();
+
+	if (!CompareLUTS(OldLUT))
+	{
+		UE_LOG(LogNiagara, Log, TEXT("PostLoad LUT generation is out of sync. Please investigate. %s"), *GetPathName());
+	}
+
 }
 
 #if WITH_EDITOR
@@ -51,8 +66,9 @@ void UNiagaraDataInterfaceVector4Curve::PostEditChangeProperty(struct FPropertyC
 			ZCurve = ColorCurveAsset->FloatCurves[2];
 			WCurve = ColorCurveAsset->FloatCurves[3];
 		}
-		UpdateLUT();
 	}
+
+	UpdateLUT();
 }
 
 #endif
@@ -60,9 +76,32 @@ void UNiagaraDataInterfaceVector4Curve::PostEditChangeProperty(struct FPropertyC
 void UNiagaraDataInterfaceVector4Curve::UpdateLUT()
 {
 	ShaderLUT.Empty();
+
+	if (bAllowUnnormalizedLUT && (XCurve.GetNumKeys() > 0 || YCurve.GetNumKeys() > 0 || ZCurve.GetNumKeys() > 0 || WCurve.GetNumKeys() > 0))
+	{
+		LUTMinTime = FLT_MAX;
+		LUTMinTime = FMath::Min(XCurve.GetNumKeys() > 0 ? XCurve.GetFirstKey().Time : LUTMinTime, LUTMinTime);
+		LUTMinTime = FMath::Min(YCurve.GetNumKeys() > 0 ? YCurve.GetFirstKey().Time : LUTMinTime, LUTMinTime);
+		LUTMinTime = FMath::Min(ZCurve.GetNumKeys() > 0 ? ZCurve.GetFirstKey().Time : LUTMinTime, LUTMinTime);
+		LUTMinTime = FMath::Min(WCurve.GetNumKeys() > 0 ? WCurve.GetFirstKey().Time : LUTMinTime, LUTMinTime);
+
+		LUTMaxTime = FLT_MIN;
+		LUTMaxTime = FMath::Max(XCurve.GetNumKeys() > 0 ? XCurve.GetLastKey().Time : LUTMaxTime, LUTMaxTime);
+		LUTMaxTime = FMath::Max(YCurve.GetNumKeys() > 0 ? YCurve.GetLastKey().Time : LUTMaxTime, LUTMaxTime);
+		LUTMaxTime = FMath::Max(ZCurve.GetNumKeys() > 0 ? ZCurve.GetLastKey().Time : LUTMaxTime, LUTMaxTime);
+		LUTMaxTime = FMath::Max(WCurve.GetNumKeys() > 0 ? WCurve.GetLastKey().Time : LUTMaxTime, LUTMaxTime);
+		LUTInvTimeRange = 1.0f / (LUTMaxTime - LUTMinTime);
+	}
+	else
+	{
+		LUTMinTime = 0.0f;
+		LUTMaxTime = 1.0f;
+		LUTInvTimeRange = 1.0f;
+	}
+
 	for (uint32 i = 0; i < CurveLUTWidth; i++)
 	{
-		float X = i / (float)CurveLUTWidth;
+		float X = UnnormalizeTime(i / (float)CurveLUTWidth);
 		FLinearColor C(XCurve.Eval(X), YCurve.Eval(X), ZCurve.Eval(X), WCurve.Eval(X));
 		ShaderLUT.Add(C.R);
 		ShaderLUT.Add(C.G);
@@ -72,9 +111,9 @@ void UNiagaraDataInterfaceVector4Curve::UpdateLUT()
 	GPUBufferDirty = true;
 }
 
-bool UNiagaraDataInterfaceVector4Curve::CopyTo(UNiagaraDataInterface* Destination) const 
+bool UNiagaraDataInterfaceVector4Curve::CopyToInternal(UNiagaraDataInterface* Destination) const 
 {
-	if (!Super::CopyTo(Destination))
+	if (!Super::CopyToInternal(Destination))
 	{
 		return false;
 	}
@@ -84,6 +123,7 @@ bool UNiagaraDataInterfaceVector4Curve::CopyTo(UNiagaraDataInterface* Destinatio
 	DestinationColorCurve->ZCurve = ZCurve;
 	DestinationColorCurve->WCurve = WCurve;
 	DestinationColorCurve->UpdateLUT();
+	ensure(CompareLUTS(DestinationColorCurve->ShaderLUT));
 
 	return true;
 }
@@ -127,10 +167,10 @@ void UNiagaraDataInterfaceVector4Curve::GetFunctions(TArray<FNiagaraFunctionSign
 // the HLSL in the spirit of a static switch
 // TODO: need a way to identify each specific function here
 // 
-bool UNiagaraDataInterfaceVector4Curve::GetFunctionHLSL(FString FunctionName, TArray<DIGPUBufferParamDescriptor> &Descriptors, FString &HLSLInterfaceID, FString &OutHLSL)
+bool UNiagaraDataInterfaceVector4Curve::GetFunctionHLSL(const FName& DefinitionFunctionName, FString InstanceFunctionName, TArray<FDIGPUBufferParamDescriptor> &Descriptors, FString &HLSLInterfaceID, FString &OutHLSL)
 {
 	FString BufferName = Descriptors[0].BufferParamName;
-	OutHLSL += TEXT("void ") + FunctionName + TEXT("(in float In_X, out float4 Out_Value) \n{\n");
+	OutHLSL += TEXT("void ") + InstanceFunctionName + TEXT("(in float In_X, out float4 Out_Value) \n{\n");
 	OutHLSL += TEXT("\t Out_Value.x = ") + BufferName + TEXT("[(int)(In_X *") + FString::FromInt(CurveLUTWidth) + TEXT(")* 4 ];");
 	OutHLSL += TEXT("\t Out_Value.y = ") + BufferName + TEXT("[1+ (int)(In_X *") + FString::FromInt(CurveLUTWidth) + TEXT(")* 4 ];");
 	OutHLSL += TEXT("\t Out_Value.z = ") + BufferName + TEXT("[2+ (int)(In_X *") + FString::FromInt(CurveLUTWidth) + TEXT(")* 4 ];");
@@ -146,21 +186,21 @@ bool UNiagaraDataInterfaceVector4Curve::GetFunctionHLSL(FString FunctionName, TA
 // 3. store buffer declaration hlsl in OutHLSL
 // multiple buffers can be defined at once here
 //
-void UNiagaraDataInterfaceVector4Curve::GetBufferDefinitionHLSL(FString DataInterfaceID, TArray<DIGPUBufferParamDescriptor> &BufferDescriptors, FString &OutHLSL)
+void UNiagaraDataInterfaceVector4Curve::GetBufferDefinitionHLSL(FString DataInterfaceID, TArray<FDIGPUBufferParamDescriptor> &BufferDescriptors, FString &OutHLSL)
 {
 	FString BufferName = "CurveLUT" + DataInterfaceID;
 	OutHLSL += TEXT("Buffer<float> ") + BufferName + TEXT(";\n");
 
-	BufferDescriptors.Add( DIGPUBufferParamDescriptor(BufferName, 0) );		// add a descriptor for shader parameter binding
+	BufferDescriptors.Add( FDIGPUBufferParamDescriptor(BufferName, 0) );		// add a descriptor for shader parameter binding
 
 }
 
 // called after translate, to setup buffers matching the buffer descriptors generated during hlsl translation
 // need to do this because the script used during translate is a clone, including its DIs
 //
-void UNiagaraDataInterfaceVector4Curve::SetupBuffers(TArray<DIGPUBufferParamDescriptor> &BufferDescriptors)
+void UNiagaraDataInterfaceVector4Curve::SetupBuffers(FDIBufferDescriptorStore &BufferDescriptors)
 {
-	for (DIGPUBufferParamDescriptor &Desc : BufferDescriptors)
+	for (FDIGPUBufferParamDescriptor &Desc : BufferDescriptors.Descriptors)
 	{
 		FNiagaraDataInterfaceBufferData BufferData(*Desc.BufferParamName);	// store off the data for later use
 		GPUBuffers.Add(BufferData);
@@ -179,7 +219,7 @@ TArray<FNiagaraDataInterfaceBufferData> &UNiagaraDataInterfaceVector4Curve::GetB
 
 		FNiagaraDataInterfaceBufferData &GPUBuffer = GPUBuffers[0];
 		GPUBuffer.Buffer.Release();
-		GPUBuffer.Buffer.Initialize(sizeof(float), CurveLUTWidth*4 , EPixelFormat::PF_R32_FLOAT);	// always allocate for up to 64 data sets
+		GPUBuffer.Buffer.Initialize(sizeof(float), CurveLUTWidth*4 , EPixelFormat::PF_R32_FLOAT, BUF_Static);	// always allocate for up to 64 data sets
 		uint32 BufferSize = ShaderLUT.Num() * sizeof(float);
 		int32 *BufferData = static_cast<int32*>(RHILockVertexBuffer(GPUBuffer.Buffer.Buffer, 0, BufferSize, EResourceLockMode::RLM_WriteOnly));
 		FPlatformMemory::Memcpy(BufferData, ShaderLUT.GetData(), BufferSize);
@@ -198,7 +238,7 @@ FVMExternalFunction UNiagaraDataInterfaceVector4Curve::GetVMExternalFunction(con
 {
 	if (BindingInfo.Name == TEXT("SampleColorCurve") && BindingInfo.GetNumInputs() == 1 && BindingInfo.GetNumOutputs() == 4)
 	{
-		return TNDIParamBinder<0, float, NDI_FUNC_BINDER(UNiagaraDataInterfaceVector4Curve, SampleCurve)>::Bind(this, BindingInfo, InstanceData);
+		return TCurveUseLUTBinder<TNDIParamBinder<0, float, NDI_FUNC_BINDER(UNiagaraDataInterfaceVector4Curve, SampleCurve)>>::Bind(this, BindingInfo, InstanceData);
 	}
 	else
 	{
@@ -208,7 +248,22 @@ FVMExternalFunction UNiagaraDataInterfaceVector4Curve::GetVMExternalFunction(con
 	}
 }
 
-template<typename XParamType>
+template<>
+FORCEINLINE_DEBUGGABLE FVector4 UNiagaraDataInterfaceVector4Curve::SampleCurveInternal<TIntegralConstant<bool, true>>(float X)
+{
+	float NormalizedX = NormalizeTime(X);
+	int32 AccessIdx = FMath::Clamp<int32>(FMath::TruncToInt(NormalizedX * CurveLUTWidthMinusOne), 0, CurveLUTWidthMinusOne) * CurveLUTNumElems;
+	return FVector4(ShaderLUT[AccessIdx], ShaderLUT[AccessIdx + 1], ShaderLUT[AccessIdx + 2], ShaderLUT[AccessIdx + 3]);
+}
+
+template<>
+FORCEINLINE_DEBUGGABLE FVector4 UNiagaraDataInterfaceVector4Curve::SampleCurveInternal<TIntegralConstant<bool, false>>(float X)
+{
+	return FVector4(XCurve.Eval(X), YCurve.Eval(X), ZCurve.Eval(X), WCurve.Eval(X));
+}
+
+
+template<typename UseLUT, typename XParamType>
 void UNiagaraDataInterfaceVector4Curve::SampleCurve(FVectorVMContext& Context)
 {
 	//TODO: Create some SIMDable optimized representation of the curve to do this faster.
@@ -221,7 +276,8 @@ void UNiagaraDataInterfaceVector4Curve::SampleCurve(FVectorVMContext& Context)
 	for (int32 i = 0; i < Context.NumInstances; ++i)
 	{
 		float X = XParam.Get();
-		FLinearColor C(XCurve.Eval(X), YCurve.Eval(X), ZCurve.Eval(X), WCurve.Eval(X));
+		FVector4 Sample = SampleCurveInternal<UseLUT>(X);
+		FLinearColor C(Sample.X, Sample.Y, Sample.Z, Sample.W);
 		*SamplePtrR.GetDest() = C.R;
 		*SamplePtrG.GetDest() = C.G;
 		*SamplePtrB.GetDest() = C.B;
diff --git a/Engine/Plugins/FX/Niagara/Source/Niagara/Private/NiagaraDataInterfaceVectorCurve.cpp b/Engine/Plugins/FX/Niagara/Source/Niagara/Private/NiagaraDataInterfaceVectorCurve.cpp
--- a/Engine/Plugins/FX/Niagara/Source/Niagara/Private/NiagaraDataInterfaceVectorCurve.cpp
+++ b/Engine/Plugins/FX/Niagara/Source/Niagara/Private/NiagaraDataInterfaceVectorCurve.cpp
@@ -5,6 +5,7 @@
 #include "Curves/CurveLinearColor.h"
 #include "Curves/CurveFloat.h"
 #include "NiagaraTypes.h"
+#include "NiagaraCustomVersion.h"
 
 //////////////////////////////////////////////////////////////////////////
 //Vector Curve
@@ -31,7 +32,20 @@ void UNiagaraDataInterfaceVectorCurve::PostInitProperties()
 void UNiagaraDataInterfaceVectorCurve::PostLoad()
 {
 	Super::PostLoad();
+
+	const int32 NiagaraVer = GetLinkerCustomVersion(FNiagaraCustomVersion::GUID);
+
+	if (NiagaraVer < FNiagaraCustomVersion::CurveLUTNowOnByDefault)
+	{
+		UpdateLUT();
+	}
+	TArray<float> OldLUT = ShaderLUT;
 	UpdateLUT();
+	if (!CompareLUTS(OldLUT))
+	{
+		UE_LOG(LogNiagara, Log, TEXT("PostLoad LUT generation is out of sync. Please investigate. %s"), *GetPathName());
+	}
+
 }
 
 #if WITH_EDITOR
@@ -50,8 +64,9 @@ void UNiagaraDataInterfaceVectorCurve::PostEditChangeProperty(struct FPropertyCh
 			YCurve = VectorCurveAsset->FloatCurves[1];
 			ZCurve = VectorCurveAsset->FloatCurves[2];
 		}
-		UpdateLUT();
 	}
+	
+	UpdateLUT();
 }
 
 #endif
@@ -60,21 +75,40 @@ void UNiagaraDataInterfaceVectorCurve::PostEditChangeProperty(struct FPropertyCh
 void UNiagaraDataInterfaceVectorCurve::UpdateLUT()
 {
 	ShaderLUT.Empty();
+	if (bAllowUnnormalizedLUT && (XCurve.GetNumKeys() > 0 || YCurve.GetNumKeys() > 0 || ZCurve.GetNumKeys() > 0))
+	{
+		LUTMinTime = FLT_MAX;
+		LUTMinTime = FMath::Min(XCurve.GetNumKeys() > 0 ? XCurve.GetFirstKey().Time : LUTMinTime, LUTMinTime);
+		LUTMinTime = FMath::Min(YCurve.GetNumKeys() > 0 ? YCurve.GetFirstKey().Time : LUTMinTime, LUTMinTime);
+		LUTMinTime = FMath::Min(ZCurve.GetNumKeys() > 0 ? ZCurve.GetFirstKey().Time : LUTMinTime, LUTMinTime);
+
+		LUTMaxTime = FLT_MIN;
+		LUTMaxTime = FMath::Max(XCurve.GetNumKeys() > 0 ? XCurve.GetLastKey().Time : LUTMaxTime, LUTMaxTime);
+		LUTMaxTime = FMath::Max(YCurve.GetNumKeys() > 0 ? YCurve.GetLastKey().Time : LUTMaxTime, LUTMaxTime);
+		LUTMaxTime = FMath::Max(ZCurve.GetNumKeys() > 0 ? ZCurve.GetLastKey().Time : LUTMaxTime, LUTMaxTime);
+		LUTInvTimeRange = 1.0f / (LUTMaxTime - LUTMinTime);
+	}
+	else
+	{
+		LUTMinTime = 0.0f;
+		LUTMaxTime = 1.0f;
+		LUTInvTimeRange = 1.0f;
+	}
+
 	for (uint32 i = 0; i < CurveLUTWidth; i++)
 	{
-		float X = i / (float)CurveLUTWidth;
+		float X = UnnormalizeTime(i / (float)CurveLUTWidth);
 		FVector C(XCurve.Eval(X), YCurve.Eval(X), ZCurve.Eval(X));
 		ShaderLUT.Add(C.X);
 		ShaderLUT.Add(C.Y);
 		ShaderLUT.Add(C.Z);
-		ShaderLUT.Add(1.0f);
 	}
 	GPUBufferDirty = true;
 }
 
-bool UNiagaraDataInterfaceVectorCurve::CopyTo(UNiagaraDataInterface* Destination) const 
+bool UNiagaraDataInterfaceVectorCurve::CopyToInternal(UNiagaraDataInterface* Destination) const 
 {
-	if (!Super::CopyTo(Destination))
+	if (!Super::CopyToInternal(Destination))
 	{
 		return false;
 	}
@@ -84,6 +118,7 @@ bool UNiagaraDataInterfaceVectorCurve::CopyTo(UNiagaraDataInterface* Destination
 	DestinationVectorCurve->ZCurve = ZCurve;
 	DestinationVectorCurve->UpdateLUT();
 
+	ensure(CompareLUTS(DestinationVectorCurve->ShaderLUT));
 	return true;
 }
 
@@ -126,10 +161,10 @@ void UNiagaraDataInterfaceVectorCurve::GetFunctions(TArray<FNiagaraFunctionSigna
 // the HLSL in the spirit of a static switch
 // TODO: need a way to identify each specific function here
 // 
-bool UNiagaraDataInterfaceVectorCurve::GetFunctionHLSL(FString FunctionName, TArray<DIGPUBufferParamDescriptor> &Descriptors, FString &HLSLInterfaceID, FString &OutHLSL)
+bool UNiagaraDataInterfaceVectorCurve::GetFunctionHLSL(const FName& DefinitionFunctionName, FString InstanceFunctionName, TArray<FDIGPUBufferParamDescriptor> &Descriptors, FString &HLSLInterfaceID, FString &OutHLSL)
 {
 	FString BufferName = Descriptors[0].BufferParamName;
-	OutHLSL += TEXT("void ") + FunctionName + TEXT("(in float In_X, out float3 Out_Value) \n{\n");
+	OutHLSL += TEXT("void ") + InstanceFunctionName + TEXT("(in float In_X, out float3 Out_Value) \n{\n");
 	OutHLSL += TEXT("\t Out_Value.x = ") + BufferName + TEXT("[(int)(In_X *") + FString::FromInt(CurveLUTWidth) + TEXT(")* 3 ];");
 	OutHLSL += TEXT("\t Out_Value.y = ") + BufferName + TEXT("[1+ (int)(In_X *") + FString::FromInt(CurveLUTWidth) + TEXT(")* 3 ];");
 	OutHLSL += TEXT("\t Out_Value.z = ") + BufferName + TEXT("[2+ (int)(In_X *") + FString::FromInt(CurveLUTWidth) + TEXT(")* 3 ];");
@@ -144,21 +179,21 @@ bool UNiagaraDataInterfaceVectorCurve::GetFunctionHLSL(FString FunctionName, TAr
 // 3. store buffer declaration hlsl in OutHLSL
 // multiple buffers can be defined at once here
 //
-void UNiagaraDataInterfaceVectorCurve::GetBufferDefinitionHLSL(FString DataInterfaceID, TArray<DIGPUBufferParamDescriptor> &BufferDescriptors, FString &OutHLSL)
+void UNiagaraDataInterfaceVectorCurve::GetBufferDefinitionHLSL(FString DataInterfaceID, TArray<FDIGPUBufferParamDescriptor> &BufferDescriptors, FString &OutHLSL)
 {
 	FString BufferName = "CurveLUT" + DataInterfaceID;
 	OutHLSL += TEXT("Buffer<float> ") + BufferName + TEXT(";\n");
 
-	BufferDescriptors.Add(DIGPUBufferParamDescriptor(BufferName, 0));		// add a descriptor for shader parameter binding
+	BufferDescriptors.Add(FDIGPUBufferParamDescriptor(BufferName, 0));		// add a descriptor for shader parameter binding
 
 }
 
 // called after translate, to setup buffers matching the buffer descriptors generated during hlsl translation
 // need to do this because the script used during translate is a clone, including its DIs
 //
-void UNiagaraDataInterfaceVectorCurve::SetupBuffers(TArray<DIGPUBufferParamDescriptor> &BufferDescriptors)
+void UNiagaraDataInterfaceVectorCurve::SetupBuffers(FDIBufferDescriptorStore &BufferDescriptors)
 {
-	for (DIGPUBufferParamDescriptor &Desc : BufferDescriptors)
+	for (FDIGPUBufferParamDescriptor &Desc : BufferDescriptors.Descriptors)
 	{
 		FNiagaraDataInterfaceBufferData BufferData(*Desc.BufferParamName);	// store off the data for later use
 		GPUBuffers.Add(BufferData);
@@ -177,7 +212,7 @@ TArray<FNiagaraDataInterfaceBufferData> &UNiagaraDataInterfaceVectorCurve::GetBu
 
 		FNiagaraDataInterfaceBufferData &GPUBuffer = GPUBuffers[0];
 		GPUBuffer.Buffer.Release();
-		GPUBuffer.Buffer.Initialize(sizeof(float), CurveLUTWidth * 3, EPixelFormat::PF_R32_FLOAT);	// always allocate for up to 64 data sets
+		GPUBuffer.Buffer.Initialize(sizeof(float), CurveLUTWidth * 3, EPixelFormat::PF_R32_FLOAT, BUF_Static);	// always allocate for up to 64 data sets
 		uint32 BufferSize = ShaderLUT.Num() * sizeof(float);
 		int32 *BufferData = static_cast<int32*>(RHILockVertexBuffer(GPUBuffer.Buffer.Buffer, 0, BufferSize, EResourceLockMode::RLM_WriteOnly));
 		FPlatformMemory::Memcpy(BufferData, ShaderLUT.GetData(), BufferSize);
@@ -194,7 +229,7 @@ FVMExternalFunction UNiagaraDataInterfaceVectorCurve::GetVMExternalFunction(cons
 {
 	if (BindingInfo.Name == TEXT("SampleVectorCurve") && BindingInfo.GetNumInputs() == 1 && BindingInfo.GetNumOutputs() == 3)
 	{
-		return TNDIParamBinder<0, float, NDI_FUNC_BINDER(UNiagaraDataInterfaceVectorCurve, SampleCurve)>::Bind(this, BindingInfo, InstanceData);
+		return TCurveUseLUTBinder<TNDIParamBinder<0, float, NDI_FUNC_BINDER(UNiagaraDataInterfaceVectorCurve, SampleCurve)>>::Bind(this, BindingInfo, InstanceData);
 	}
 	else
 	{
@@ -204,7 +239,21 @@ FVMExternalFunction UNiagaraDataInterfaceVectorCurve::GetVMExternalFunction(cons
 	}
 }
 
-template<typename XParamType>
+template<>
+FORCEINLINE_DEBUGGABLE FVector UNiagaraDataInterfaceVectorCurve::SampleCurveInternal<TIntegralConstant<bool, true>>(float X)
+{
+	float NormalizedX = NormalizeTime(X);
+	int32 AccessIdx = FMath::Clamp<int32>(FMath::TruncToInt(NormalizedX * CurveLUTWidthMinusOne), 0, CurveLUTWidthMinusOne) * CurveLUTNumElems;
+	return FVector(ShaderLUT[AccessIdx], ShaderLUT[AccessIdx + 1], ShaderLUT[AccessIdx + 2]);
+}
+
+template<>
+FORCEINLINE_DEBUGGABLE FVector UNiagaraDataInterfaceVectorCurve::SampleCurveInternal<TIntegralConstant<bool, false>>(float X)
+{
+	return FVector(XCurve.Eval(X), YCurve.Eval(X), ZCurve.Eval(X));
+}
+
+template<typename UseLUT, typename XParamType>
 void UNiagaraDataInterfaceVectorCurve::SampleCurve(FVectorVMContext& Context)
 {
 	//TODO: Create some SIMDable optimized representation of the curve to do this faster.
@@ -216,7 +265,7 @@ void UNiagaraDataInterfaceVectorCurve::SampleCurve(FVectorVMContext& Context)
 	for (int32 i = 0; i < Context.NumInstances; ++i)
 	{
 		float X = XParam.Get();
-		FVector V(XCurve.Eval(X), YCurve.Eval(X), ZCurve.Eval(X));
+		FVector V = SampleCurveInternal<UseLUT>(X);
 		*OutSampleX.GetDest() = V.X;
 		*OutSampleY.GetDest() = V.Y;
 		*OutSampleZ.GetDest() = V.Z;
diff --git a/Engine/Plugins/FX/Niagara/Source/Niagara/Private/NiagaraDataInterfaceVectorField.cpp b/Engine/Plugins/FX/Niagara/Source/Niagara/Private/NiagaraDataInterfaceVectorField.cpp
new file mode 100644
--- /dev/null
+++ b/Engine/Plugins/FX/Niagara/Source/Niagara/Private/NiagaraDataInterfaceVectorField.cpp
@@ -0,0 +1,922 @@
+// Copyright 1998-2016 Epic Games, Inc. All Rights Reserved.
+
+#include "NiagaraDataInterfaceVectorField.h"
+#include "Internationalization.h"
+#include "VectorField/VectorFieldStatic.h"
+#include "VectorField/VectorFieldAnimated.h"
+#include "Math/Float16Color.h"
+#define LOCTEXT_NAMESPACE "NiagaraDataInterfaceVectorField"
+
+
+#if 0
+//////////////////////////////////////////////////////////////////////////
+//FNDIVectorField_InstanceData
+
+bool FNDIVectorField_InstanceData::Init(UNiagaraDataInterfaceVectorField* Interface, FNiagaraSystemInstance* SystemInstance)
+{
+	UE_LOG(LogNiagara, Log, TEXT("FNDIVectorField_InstanceData %p"), this);
+	check(SystemInstance);
+	/*
+	USkeletalMesh* PrevMesh = Mesh;
+	Component = nullptr;
+	Mesh = nullptr;
+	Transform = FMatrix::Identity;
+	TransformInverseTransposed = FMatrix::Identity;
+	PrevTransform = FMatrix::Identity;
+	PrevTransformInverseTransposed = FMatrix::Identity;
+	*/
+	DeltaSeconds = 0.0f;
+	Field = Interface->Field;
+
+	UVectorFieldStatic* StaticVectorField = Cast<UVectorFieldStatic>(Field);
+	UVectorFieldAnimated* AnimatedVectorField = Cast<UVectorFieldAnimated>(Field);
+	SizeX = 0;
+	SizeY = 0;
+	SizeZ = 0;
+
+	/*
+	OutParameters.WorldToVolume[Index] = VectorFieldInstance->WorldToVolume;
+	OutParameters.VolumeToWorld[Index] = VectorFieldInstance->VolumeToWorldNoScale;
+	OutParameters.VolumeSize[Index] = FVector4(Resource->SizeX, Resource->SizeY, Resource->SizeZ, 0);
+	OutParameters.IntensityAndTightness[Index] = FVector4(Intensity, Tightness, 0, 0);
+	OutParameters.TilingAxes[Index].X = VectorFieldInstance->bTileX ? 1.0f : 0.0f;
+	OutParameters.TilingAxes[Index].Y = VectorFieldInstance->bTileY ? 1.0f : 0.0f;
+	OutParameters.TilingAxes[Index].Z = VectorFieldInstance->bTileZ ? 1.0f : 0.0f;
+	*/
+	TilingAxes.X = Interface->bTileX ? 1.0f : 0.0f;
+	TilingAxes.Y = Interface->bTileY ? 1.0f : 0.0f;
+	TilingAxes.Z = Interface->bTileZ ? 1.0f : 0.0f;
+
+
+	if (StaticVectorField)
+	{
+		SizeX = (uint32)StaticVectorField->SizeX;
+		SizeY = (uint32)StaticVectorField->SizeY;
+		SizeZ = (uint32)StaticVectorField->SizeZ;
+		LocalBounds = StaticVectorField->Bounds;
+	}
+	else if (AnimatedVectorField)
+	{
+		SizeX = (uint32)AnimatedVectorField->VolumeSizeX;
+		SizeY = (uint32)AnimatedVectorField->VolumeSizeY;
+		SizeZ = (uint32)AnimatedVectorField->VolumeSizeZ;
+		LocalBounds = AnimatedVectorField->Bounds;
+	}
+#if 0
+	USkeletalMeshComponent* NewSkelComp = nullptr;
+	if (Interface->Source)
+	{
+		ASkeletalMeshActor* MeshActor = Cast<ASkeletalMeshActor>(Interface->Source);
+		USkeletalMeshComponent* SourceComp = nullptr;
+		if (MeshActor != nullptr)
+		{
+			SourceComp = MeshActor->GetSkeletalMeshComponent();
+		}
+		else
+		{
+			SourceComp = Interface->Source->FindComponentByClass<USkeletalMeshComponent>();
+		}
+
+		if (SourceComp)
+		{
+			Mesh = SourceComp->SkeletalMesh;
+			NewSkelComp = SourceComp;
+		}
+		else
+		{
+			Component = Interface->Source->GetRootComponent();
+		}
+	}
+	else
+	{
+		if (UNiagaraComponent* SimComp = SystemInstance->GetComponent())
+		{
+			if (USkeletalMeshComponent* ParentComp = Cast<USkeletalMeshComponent>(SimComp->GetAttachParent()))
+			{
+				NewSkelComp = ParentComp;
+				Mesh = ParentComp->SkeletalMesh;
+			}
+			else if (USkeletalMeshComponent* OuterComp = SimComp->GetTypedOuter<USkeletalMeshComponent>())
+			{
+				NewSkelComp = OuterComp;
+				Mesh = OuterComp->SkeletalMesh;
+			}
+			else if (AActor* Owner = SimComp->GetAttachmentRootActor())
+			{
+				TArray<UActorComponent*> SourceComps = Owner->GetComponentsByClass(USkeletalMeshComponent::StaticClass());
+				for (UActorComponent* ActorComp : SourceComps)
+				{
+					USkeletalMeshComponent* SourceComp = Cast<USkeletalMeshComponent>(ActorComp);
+					if (SourceComp)
+					{
+						USkeletalMesh* PossibleMesh = SourceComp->SkeletalMesh;
+						if (PossibleMesh != nullptr/* && PossibleMesh->bAllowCPUAccess*/)
+						{
+							Mesh = PossibleMesh;
+							NewSkelComp = SourceComp;
+							break;
+						}
+					}
+				}
+			}
+
+			if (!Component.IsValid())
+			{
+				Component = SimComp;
+			}
+		}
+	}
+
+	if (NewSkelComp)
+	{
+		Component = NewSkelComp;
+	}
+
+	check(Component.IsValid());
+
+	if (!Mesh && Interface->DefaultMesh)
+	{
+		Mesh = Interface->DefaultMesh;
+	}
+
+	if (Component.IsValid() && Mesh)
+	{
+		PrevTransform = Transform;
+		PrevTransformInverseTransposed = TransformInverseTransposed;
+		Transform = Component->GetComponentToWorld().ToMatrixWithScale();
+		TransformInverseTransposed = Transform.InverseFast().GetTransposed();
+	}
+
+	if (!Mesh)
+	{
+		UE_LOG(LogNiagara, Log, TEXT("SkeletalMesh data interface has no valid mesh. Failed InitPerInstanceData - %s"), *Interface->GetFullName());
+		return false;
+	}
+
+	// 	if (!Mesh->bAllowCPUAccess)
+	// 	{
+	// 		UE_LOG(LogNiagara, Log, TEXT("SkeletalMesh data interface using a mesh that does not allow CPU access. Failed InitPerInstanceData - Mesh: %s"), *Mesh->GetFullName());
+	// 		return false;
+	// 	}
+
+	if (!Component.IsValid())
+	{
+		UE_LOG(LogNiagara, Log, TEXT("SkeletalMesh data interface has no valid component. Failed InitPerInstanceData - %s"), *Interface->GetFullName());
+		return false;
+	}
+
+	//	bIsAreaWeightedSampling = Mesh->bSupportUniformlyDistributedSampling;
+
+	// 	//Init the instance filter
+	// 	ValidSections.Empty();
+	// 	FStaticMeshLODResources& Res = Mesh->RenderData->LODResources[0];
+	// 	for (int32 i = 0; i < Res.Sections.Num(); ++i)
+	// 	{
+	// 		if (Interface->SectionFilter.AllowedMaterialSlots.Num() == 0 || Interface->SectionFilter.AllowedMaterialSlots.Contains(Res.Sections[i].MaterialIndex))
+	// 		{
+	// 			ValidSections.Add(i);
+	// 		}
+	// 	}
+	// 
+	// 	if (GetValidSections().Num() == 0)
+	// 	{
+	// 		UE_LOG(LogNiagara, Log, TEXT("StaticMesh data interface has a section filter preventing any spawning. Failed InitPerInstanceData - %s"), *Interface->GetFullName());
+	// 		return false;
+	// 	}
+	// 
+	// 	Sampler.Init(&Res, this);
+
+	if (NewSkelComp)
+	{
+		TWeakObjectPtr<USkeletalMeshComponent> SkelWeakCompPtr = NewSkelComp;
+		SkinningData = SystemInstance->GetWorldManager()->GetSkeletalMeshGeneratedData().GetCachedSkinningData(SkelWeakCompPtr);
+	}
+#endif
+
+	return true;
+}
+
+void FNDIVectorField_InstanceData::UpdateTransforms(const FMatrix& LocalToWorld, FMatrix& OutVolumeToWorld, FMatrix& OutWorldToVolume)
+{
+	const FVector VolumeOffset = LocalBounds.Min;
+	const FVector VolumeScale = LocalBounds.Max - LocalBounds.Min;
+	//VolumeToWorldNoScale = LocalToWorld.GetMatrixWithoutScale().RemoveTranslation();
+	OutVolumeToWorld = FScaleMatrix(VolumeScale) * FTranslationMatrix(VolumeOffset)
+		* LocalToWorld;
+	OutWorldToVolume = OutVolumeToWorld.InverseFast();
+}
+
+const void* FNDIVectorField_InstanceData::Lock()
+{
+	UVectorFieldStatic* StaticVectorField = Cast<UVectorFieldStatic>(Field);
+	UVectorFieldAnimated* AnimatedVectorField = Cast<UVectorFieldAnimated>(Field);
+
+	if (StaticVectorField)
+	{
+		return StaticVectorField->SourceData.LockReadOnly();
+	}
+	return nullptr;
+}
+
+void FNDIVectorField_InstanceData::Unlock()
+{
+	UVectorFieldStatic* StaticVectorField = Cast<UVectorFieldStatic>(Field);
+	UVectorFieldAnimated* AnimatedVectorField = Cast<UVectorFieldAnimated>(Field);
+
+	if (StaticVectorField)
+	{
+		StaticVectorField->SourceData.Unlock();
+	}
+}
+
+
+bool FNDIVectorField_InstanceData::ResetRequired(UNiagaraDataInterfaceVectorField* Interface)const
+{
+	//USceneComponent* Comp = Component.Get();
+	//if (!Comp)
+	//{
+	//	//The component we were bound to is no longer valid so we have to trigger a reset.
+	//	return true;
+	//}
+
+	//if (USkeletalMeshComponent* SkelComp = Cast<USkeletalMeshComponent>(Comp))
+	//{
+	//	if (!SkelComp->SkeletalMesh)
+	//	{
+	//		return true;
+	//	}
+	//}
+	//else
+	//{
+		if (!Interface->Field || Interface->Field != Field)
+		{
+			return true;
+		}
+	//}
+
+	return false;
+}
+
+bool FNDIVectorField_InstanceData::Tick(UNiagaraDataInterfaceVectorField* Interface, FNiagaraSystemInstance* SystemInstance, float InDeltaSeconds)
+{
+	if (ResetRequired(Interface))
+	{
+		return true;
+	}
+	else
+	{
+		DeltaSeconds = InDeltaSeconds;
+		/*if (Component.IsValid() && Mesh)
+		{
+			PrevTransform = Transform;
+			PrevTransformInverseTransposed = TransformInverseTransposed;
+			Transform = Component->GetComponentToWorld().ToMatrixWithScale();
+			TransformInverseTransposed = Transform.InverseFast().GetTransposed();
+		}
+		else
+		{
+			PrevTransform = FMatrix::Identity;
+			PrevTransformInverseTransposed = FMatrix::Identity;
+			Transform = FMatrix::Identity;
+			TransformInverseTransposed = FMatrix::Identity;
+		}*/
+		return false;
+	}
+}
+#endif
+
+//Instance Data END
+//////////////////////////////////////////////////////////////////////////
+UNiagaraDataInterfaceVectorField::UNiagaraDataInterfaceVectorField(FObjectInitializer const& ObjectInitializer)
+	: Super(ObjectInitializer)
+	, Field(nullptr)
+	, bTileX(false)
+	, bTileY(false)
+	, bTileZ(false)
+	, bGPUBufferDirty(true)
+{
+}
+
+#if WITH_EDITOR
+
+void UNiagaraDataInterfaceVectorField::PostEditChangeProperty(struct FPropertyChangedEvent& PropertyChangedEvent)
+{
+	Super::PostEditChangeProperty(PropertyChangedEvent);
+	InitField();
+}
+
+void UNiagaraDataInterfaceVectorField::PostLoad()
+{
+	Super::PostLoad();
+	if (Field)
+	{
+		Field->ConditionalPostLoad();
+	}
+	InitField();
+}
+
+void UNiagaraDataInterfaceVectorField::PostInitProperties()
+{
+	Super::PostInitProperties();
+
+	//Can we register data interfaces as regular types and fold them into the FNiagaraVariable framework for UI and function calls etc?
+	if (HasAnyFlags(RF_ClassDefaultObject))
+	{
+		FNiagaraTypeRegistry::Register(FNiagaraTypeDefinition(GetClass()), true, false, false);
+	}
+}
+
+#endif //WITH_EDITOR
+
+static const FName SampleVectorFieldName("SampleField");
+static const FName GetVectorDimsName("FieldDimensions");
+static const FName GetVectorFieldBoundsName("FieldBounds");
+
+void UNiagaraDataInterfaceVectorField::GetFunctions(TArray<FNiagaraFunctionSignature>& OutFunctions)
+{
+	{
+		FNiagaraFunctionSignature Sig;
+		Sig.Name = SampleVectorFieldName;
+		Sig.Inputs.Add(FNiagaraVariable(FNiagaraTypeDefinition(GetClass()), TEXT("Vector Field")));
+		Sig.Inputs.Add(FNiagaraVariable(FNiagaraTypeDefinition::GetVec3Def(), TEXT("Sample Point")));
+		Sig.Outputs.Add(FNiagaraVariable(FNiagaraTypeDefinition::GetVec3Def(), TEXT("Sampled Value")));
+
+		Sig.bMemberFunction = true;
+		Sig.bRequiresContext = false;
+		OutFunctions.Add(Sig);
+	}
+
+	{
+		FNiagaraFunctionSignature Sig;
+		Sig.Name = GetVectorDimsName;
+		Sig.Inputs.Add(FNiagaraVariable(FNiagaraTypeDefinition(GetClass()), TEXT("Vector Field")));
+		Sig.Outputs.Add(FNiagaraVariable(FNiagaraTypeDefinition::GetVec3Def(), TEXT("Dimensions")));
+
+		Sig.bMemberFunction = true;
+		Sig.bRequiresContext = false;
+		OutFunctions.Add(Sig);
+	}
+
+	{
+		FNiagaraFunctionSignature Sig;
+		Sig.Name = GetVectorFieldBoundsName;
+		Sig.Inputs.Add(FNiagaraVariable(FNiagaraTypeDefinition(GetClass()), TEXT("Vector Field")));
+		Sig.Outputs.Add(FNiagaraVariable(FNiagaraTypeDefinition::GetVec3Def(), TEXT("BoundMin")));
+		Sig.Outputs.Add(FNiagaraVariable(FNiagaraTypeDefinition::GetVec3Def(), TEXT("BoundMax")));
+
+		Sig.bMemberFunction = true;
+		Sig.bRequiresContext = false;
+		OutFunctions.Add(Sig);
+	}
+}
+
+DEFINE_NDI_FUNC_BINDER(UNiagaraDataInterfaceVectorField, SampleVectorField);
+FVMExternalFunction UNiagaraDataInterfaceVectorField::GetVMExternalFunction(const FVMExternalFunctionBindingInfo& BindingInfo, void* InstanceData)
+{
+	FVMExternalFunction Function;
+	if (BindingInfo.Name == SampleVectorFieldName && BindingInfo.GetNumInputs() == 3 && BindingInfo.GetNumOutputs() == 3)
+	{
+		Function = TNDIParamBinder<0, float, TNDIParamBinder<1, float, TNDIParamBinder<2, float, NDI_FUNC_BINDER(UNiagaraDataInterfaceVectorField, SampleVectorField)>>>::Bind(this, BindingInfo, InstanceData);
+	}
+	else if (BindingInfo.Name == GetVectorDimsName && BindingInfo.GetNumInputs() == 0 && BindingInfo.GetNumOutputs() == 3)
+	{
+		Function = FVMExternalFunction::CreateUObject(this, &UNiagaraDataInterfaceVectorField::GetFieldDimensions);
+	}
+	else if (BindingInfo.Name == GetVectorFieldBoundsName && BindingInfo.GetNumInputs() == 0 && BindingInfo.GetNumOutputs() == 6)
+	{
+		Function = FVMExternalFunction::CreateUObject(this, &UNiagaraDataInterfaceVectorField::GetFieldBounds);
+	}
+	return Function;
+}
+
+void UNiagaraDataInterfaceVectorField::GetFieldDimensions(FVectorVMContext& Context)
+{
+	/*FUserPtrHandler<FNDIVectorField_InstanceData> InstData(Context);
+	if (!InstData.Get())
+	{
+		UE_LOG(LogNiagara, Warning, TEXT("FNDIVectorField_InstanceData has invalid instance data. %s"), *GetPathName());
+	}*/
+
+	FRegisterHandler<float> OutSizeX(Context);	
+	FRegisterHandler<float> OutSizeY(Context);
+	FRegisterHandler<float> OutSizeZ(Context);
+	
+	for (int32 i = 0; i < Context.NumInstances; ++i)
+	{
+		*OutSizeX.GetDest() = (float)SizeX;		
+		*OutSizeY.GetDest() = (float)SizeY;
+		*OutSizeZ.GetDest() = (float)SizeZ;
+
+		OutSizeX.Advance();
+		OutSizeY.Advance();
+		OutSizeZ.Advance();
+	}
+}
+
+void UNiagaraDataInterfaceVectorField::GetFieldBounds(FVectorVMContext& Context)
+{
+	/*FUserPtrHandler<FNDIVectorField_InstanceData> InstData(Context);
+	if (!InstData.Get())
+	{
+		UE_LOG(LogNiagara, Warning, TEXT("FNDIVectorField_InstanceData has invalid instance data. %s"), *GetPathName());
+	}*/
+
+	FRegisterHandler<float> OutMinX(Context);
+	FRegisterHandler<float> OutMinY(Context);
+	FRegisterHandler<float> OutMinZ(Context);
+	FRegisterHandler<float> OutMaxX(Context);
+	FRegisterHandler<float> OutMaxY(Context);
+	FRegisterHandler<float> OutMaxZ(Context);
+	
+	for (int32 i = 0; i < Context.NumInstances; ++i)
+	{
+		*OutMinX.GetDest() = LocalBounds.Min.X;
+		*OutMinY.GetDest() = LocalBounds.Min.Y;
+		*OutMinZ.GetDest() = LocalBounds.Min.Z;
+		*OutMaxX.GetDest() = LocalBounds.Max.X;
+		*OutMaxY.GetDest() = LocalBounds.Max.Y;
+		*OutMaxZ.GetDest() = LocalBounds.Max.Z;
+
+		OutMinX.Advance();
+		OutMinY.Advance();
+		OutMinZ.Advance();
+		OutMaxX.Advance();
+		OutMaxY.Advance();
+		OutMaxZ.Advance();
+	}
+}
+
+/*
+* Compute the influence of vector fields on a particle at the given position.
+* @param OutForce - Force to apply to the particle.
+* @param OutVelocity - Direct velocity influence on the particle.
+* @param Position - Position of the particle.
+* @param PerParticleScale - Amount by which to scale the influence on this particle.
+void EvaluateVectorFields(out float3 OutForce, out float4 OutVelocity, float3 Position, float PerParticleScale)
+{
+	float3 TotalForce = 0;
+	float3 WeightedVelocity = 0;
+	float TotalWeight = 0;
+	float FinalWeight = 0;
+
+	for (int VectorFieldIndex = 0; VectorFieldIndex < VectorFields.Count; ++VectorFieldIndex)
+	{
+		float2 IntensityAndTightness = VectorFields.IntensityAndTightness[VectorFieldIndex].xy;
+		float Intensity = IntensityAndTightness.x * PerParticleScale;
+		float Tightness = IntensityAndTightness.y;
+		float3 VolumeSize = VectorFields.VolumeSize[VectorFieldIndex].xyz;
+		float3 VolumeUV = mul(float4(Position.xyz, 1), VectorFields.WorldToVolume[VectorFieldIndex]).xyz;
+		//Tile the UVs if needed. TilingAxes will be 1.0 or 0.0 in each channel depending on which axes are being tiled, if any.
+		VolumeUV -= floor(VolumeUV * VectorFields.TilingAxes[VectorFieldIndex].xyz);
+
+		float3 AxisWeights =
+			saturate(VolumeUV * VolumeSize.xyz) *
+			saturate((1.0f - VolumeUV) * VolumeSize.xyz);
+		float DistanceWeight = min(AxisWeights.x, min(AxisWeights.y, AxisWeights.z));
+
+		// @todo compat hack: Some compilers only allow constant indexing into a texture array
+		//		float3 VectorSample = Texture3DSample(VectorFieldTextures[VectorFieldIndex], VectorFieldTexturesSampler, saturate(VolumeUV)).xyz;
+		float3 VectorSample = SampleVectorFieldTexture(VectorFieldIndex, saturate(VolumeUV));
+
+		float3 Vec = mul(float4(VectorSample, 0), VectorFields.VolumeToWorld[VectorFieldIndex]).xyz;
+		TotalForce += (Vec * DistanceWeight * Intensity);
+		WeightedVelocity += (Vec * Intensity * DistanceWeight * Tightness);
+		TotalWeight += (DistanceWeight * Tightness);
+		FinalWeight = max(FinalWeight, DistanceWeight * Tightness);
+	}
+
+	// Forces are additive.
+	OutForce = TotalForce;
+	// Velocities use a weighted average.
+	OutVelocity.xyz = WeightedVelocity / (TotalWeight + 0.001f);
+	OutVelocity.w = FinalWeight;
+}
+*/
+
+void UNiagaraDataInterfaceVectorField::InitField()
+{
+	UVectorFieldStatic* StaticVectorField = Cast<UVectorFieldStatic>(Field);
+	UVectorFieldAnimated* AnimatedVectorField = Cast<UVectorFieldAnimated>(Field);
+	SizeX = 0;
+	SizeY = 0;
+	SizeZ = 0;
+
+	/*
+	OutParameters.WorldToVolume[Index] = VectorFieldInstance->WorldToVolume;
+	OutParameters.VolumeToWorld[Index] = VectorFieldInstance->VolumeToWorldNoScale;
+	OutParameters.VolumeSize[Index] = FVector4(Resource->SizeX, Resource->SizeY, Resource->SizeZ, 0);
+	OutParameters.IntensityAndTightness[Index] = FVector4(Intensity, Tightness, 0, 0);
+	OutParameters.TilingAxes[Index].X = VectorFieldInstance->bTileX ? 1.0f : 0.0f;
+	OutParameters.TilingAxes[Index].Y = VectorFieldInstance->bTileY ? 1.0f : 0.0f;
+	OutParameters.TilingAxes[Index].Z = VectorFieldInstance->bTileZ ? 1.0f : 0.0f;
+	*/
+	TilingAxes.X = bTileX ? 1.0f : 0.0f;
+	TilingAxes.Y = bTileY ? 1.0f : 0.0f;
+	TilingAxes.Z = bTileZ ? 1.0f : 0.0f;
+
+
+	if (StaticVectorField)
+	{
+		SizeX = (uint32)StaticVectorField->SizeX;
+		SizeY = (uint32)StaticVectorField->SizeY;
+		SizeZ = (uint32)StaticVectorField->SizeZ;
+		LocalBounds = StaticVectorField->Bounds;
+	}
+	else if (AnimatedVectorField)
+	{
+		SizeX = (uint32)AnimatedVectorField->VolumeSizeX;
+		SizeY = (uint32)AnimatedVectorField->VolumeSizeY;
+		SizeZ = (uint32)AnimatedVectorField->VolumeSizeZ;
+		LocalBounds = AnimatedVectorField->Bounds;
+	}
+}
+
+const void* UNiagaraDataInterfaceVectorField::Lock()
+{
+	UVectorFieldStatic* StaticVectorField = Cast<UVectorFieldStatic>(Field);
+	UVectorFieldAnimated* AnimatedVectorField = Cast<UVectorFieldAnimated>(Field);
+
+	if (StaticVectorField)
+	{
+		return StaticVectorField->SourceData.LockReadOnly();
+	}
+	return nullptr;
+}
+
+void UNiagaraDataInterfaceVectorField::Unlock()
+{
+	UVectorFieldStatic* StaticVectorField = Cast<UVectorFieldStatic>(Field);
+	UVectorFieldAnimated* AnimatedVectorField = Cast<UVectorFieldAnimated>(Field);
+
+	if (StaticVectorField)
+	{
+		StaticVectorField->SourceData.Unlock();
+	}
+}
+
+
+template<typename XType, typename YType, typename ZType>
+void UNiagaraDataInterfaceVectorField::SampleVectorField(FVectorVMContext& Context)
+{
+	// Input arguments...
+	XType XParam(Context);
+	YType YParam(Context);
+	ZType ZParam(Context);
+
+	// User pointer buffer (as requested)
+	//FUserPtrHandler<FNDIVectorField_InstanceData> InstData(Context);
+
+	// Outputs...
+	FRegisterHandler<float> OutSampleX(Context);
+	FRegisterHandler<float> OutSampleY(Context);
+	FRegisterHandler<float> OutSampleZ(Context);
+
+	/*UE_LOG(LogNiagara, Log, TEXT("Registers: In: %d %d %d %d Out: %d %d %d"),
+		XParam.RegisterIndex, YParam.RegisterIndex, ZParam.RegisterIndex, InstData.RegisterIndex,
+		OutSampleX.RegisterIndex, OutSampleY.RegisterIndex, OutSampleZ.RegisterIndex);*/
+
+	UVectorFieldStatic* StaticVectorField = Cast<UVectorFieldStatic>(Field);
+	UVectorFieldAnimated* AnimatedVectorField = Cast<UVectorFieldAnimated>(Field);
+
+	//FMatrix VolumeToWorld;
+	//FMatrix WorldToVolume;
+	//FMatrix LocalToWorld = FMatrix::Identity;
+	//InstData->UpdateTransforms(LocalToWorld, VolumeToWorld, WorldToVolume);
+
+	const FFloat16Color* Data = (const FFloat16Color*)Lock();
+
+	if (Data == nullptr || Field == nullptr || (StaticVectorField == nullptr && AnimatedVectorField == nullptr) || AnimatedVectorField != nullptr)
+	{
+		// Set the default vector to positive X axis...
+		for (int32 InstanceIdx = 0; InstanceIdx < Context.NumInstances; ++InstanceIdx)
+		{
+			VectorRegister Dst = MakeVectorRegister(1.0f, 0.0f, 0.0f, 0.0f);
+
+			float *RegPtr = reinterpret_cast<float*>(&Dst);
+			*OutSampleX.GetDest() = RegPtr[0];
+			*OutSampleY.GetDest() = RegPtr[1];
+			*OutSampleZ.GetDest() = RegPtr[2];
+
+			XParam.Advance();
+			YParam.Advance();
+			ZParam.Advance();
+			OutSampleX.Advance();
+			OutSampleY.Advance();
+			OutSampleZ.Advance();
+		}
+	}
+	else if (StaticVectorField != nullptr)
+	{
+		FVector VolumeSize((float)SizeX, (float)SizeY, (float)SizeZ);
+		uint32 VolumeSizeX = (uint32)SizeX;
+		uint32 VolumeSizeY = (uint32)SizeY;
+		uint32 VolumeSizeZ = (uint32)SizeZ;
+
+		for (int32 InstanceIdx = 0; InstanceIdx < Context.NumInstances; ++InstanceIdx)
+		{
+			FVector Pos(XParam.Get(), YParam.Get(), ZParam.Get());
+			
+			// float3 VolumeUV = mul(float4(Position.xyz, 1), VectorFields.WorldToVolume[VectorFieldIndex]).xyz;
+			//FNDITransformHandler Handler;
+			FVector VolumeUV = Pos;
+			//Handler.TransformPosition(VolumeUV, WorldToVolume);
+
+			// VolumeUV -= floor(VolumeUV * VectorFields.TilingAxes[VectorFieldIndex].xyz);
+			FVector TiledVolume = VolumeUV * TilingAxes;
+			VolumeUV.X -= FMath::FloorToFloat(TiledVolume.X);
+			VolumeUV.Y -= FMath::FloorToFloat(TiledVolume.Y);
+			VolumeUV.Z -= FMath::FloorToFloat(TiledVolume.Z);
+
+			// Saturate(VolumeUV) 
+			VolumeUV.X = FMath::Clamp(VolumeUV.X, 0.0f, 1.0f);
+			VolumeUV.Y = FMath::Clamp(VolumeUV.Y, 0.0f, 1.0f);
+			VolumeUV.Z = FMath::Clamp(VolumeUV.Z, 0.0f, 1.0f);
+
+			// Sample texture...
+			VolumeUV = VolumeUV * VolumeSize;
+
+			/**
+			int3 IntCoord = int3(ModXYZ.x, ModXYZ.y, ModXYZ.z);\n");
+			float3 frc = frac(ModXYZ);\n");
+			*/
+			FVector VecIndex;
+			FVector VecFrac;
+			VecFrac.X = FGenericPlatformMath::Modf(VolumeUV.X, &VecIndex.X);
+			VecFrac.Y = FGenericPlatformMath::Modf(VolumeUV.Y, &VecIndex.Y);
+			VecFrac.Z = FGenericPlatformMath::Modf(VolumeUV.Z, &VecIndex.Z);
+
+			uint32 IndexX = (uint32)VecIndex.X;
+			uint32 IndexY = (uint32)VecIndex.Y;
+			uint32 IndexZ = (uint32)VecIndex.Z;
+
+#define VF_MakeIndex(X, Y, Z) ((X)  + ((Y) * VolumeSizeX) + ((Z) * VolumeSizeX * VolumeSizeY))
+#define VF_Tile(I, Range)  (((I) >= Range) ? 0 : (I))
+
+			// Trilinear interpolation, as defined by https://en.wikipedia.org/wiki/Trilinear_interpolation
+			/** 
+			float3 frc = frac(ModXYZ);
+			float3 V1 = Buffer[IntCoord.x + IntCoord.y*VolumeSizeX + IntCoord.z*VolumeSizeX*VolumeSizeY].xyz;  // V(x0, y0, z0)
+			float3 V2 = Buffer[IntCoord.x + 1 + IntCoord.y*VolumeSizeX + IntCoord.z*VolumeSizeX*VolumeSizeY].xyz; // V(x1, y0, z0)
+			float3 C00 = lerp(V1, V2, frc.xxx);
+			V1 = Buffer[IntCoord.x + IntCoord.y*VolumeSizeX + (IntCoord.z+1)*VolumeSizeX*VolumeSizeY].xyz;  // V(x0, y0, z1)
+			V2 = Buffer[IntCoord.x + 1 + IntCoord.y*VolumeSizeX + (IntCoord.z+1)*VolumeSizeX*VolumeSizeY].xyz; // V(x1, y0, z1)
+			float3 C01 = lerp(V1, V2, frc.xxx);
+			V1 = Buffer[IntCoord.x + (IntCoord.y+1)*VolumeSizeX + IntCoord.z*VolumeSizeX*VolumeSizeY].xyz;  // V(x0, y1, z0)
+			V2 = Buffer[IntCoord.x + 1 + (IntCoord.y+1)*VolumeSizeX + IntCoord.z*VolumeSizeX*VolumeSizeY].xyz; // V(x1, y1, z0)
+			float3 C10 = lerp(V1, V2, frc.xxx);
+			V1 = Buffer[IntCoord.x + (IntCoord.y+1)*VolumeSizeX + (IntCoord.z+1)*VolumeSizeX*VolumeSizeY].xyz;  // V(x0, y1, z1)
+			V2 = Buffer[IntCoord.x + 1 + (IntCoord.y+1)*VolumeSizeX + (IntCoord.z+1)*VolumeSizeX*VolumeSizeY].xyz; // V(x1, y1, z1)
+			float3 C11 = lerp(V1, V2, frc.xxx);
+			float3 C0 = lerp(C00, C10, frc.yyy);
+			float3 C1 = lerp(c01, C11, frc.yyy);
+			Out_Value = lerp(C0, C1, frc.zzz);
+			*/
+			uint32 Index = VF_MakeIndex(VF_Tile(IndexX, VolumeSizeX), VF_Tile(IndexY, VolumeSizeY), VF_Tile(IndexZ, VolumeSizeZ));
+			check(Index < (uint32)(VolumeSizeX*VolumeSizeY*VolumeSizeZ));
+			FVector V1 ((float)Data[Index].R, (float)Data[Index].G, (float)Data[Index].B);  // V(x0, y0, z0)
+			Index = VF_MakeIndex(VF_Tile(IndexX + 1, VolumeSizeX), VF_Tile(IndexY, VolumeSizeY), VF_Tile(IndexZ, VolumeSizeZ));
+			check(Index < (uint32)(VolumeSizeX*VolumeSizeY*VolumeSizeZ));
+			FVector V2((float)Data[Index].R, (float)Data[Index].G, (float)Data[Index].B); // V(x1, y0, z0)
+			FVector C00 = FMath::Lerp(V1, V2, VecFrac.X);
+
+			Index = VF_MakeIndex(IndexX, VF_Tile(IndexY, VolumeSizeY), VF_Tile(IndexZ + 1, VolumeSizeZ));
+			check(Index < (uint32)(VolumeSizeX*VolumeSizeY*VolumeSizeZ));
+			V1 = FVector((float)Data[Index].R, (float)Data[Index].G, (float)Data[Index].B);// V(x0, y0, z1)
+			Index = VF_MakeIndex(VF_Tile(IndexX + 1, VolumeSizeX), VF_Tile(IndexY, VolumeSizeY), VF_Tile(IndexZ + 1, VolumeSizeZ));
+			check(Index < (uint32)(VolumeSizeX*VolumeSizeY*VolumeSizeZ));
+			V1 = FVector((float)Data[Index].R, (float)Data[Index].G, (float)Data[Index].B);// V(x1, y0, z1)
+			FVector C01 = FMath::Lerp(V1, V2, VecFrac.X);
+
+			Index = VF_MakeIndex(VF_Tile(IndexX, VolumeSizeX), VF_Tile(IndexY + 1, VolumeSizeY), VF_Tile(IndexZ, VolumeSizeZ));
+			check(Index < (uint32)(VolumeSizeX*VolumeSizeY*VolumeSizeZ));
+			V1 = FVector((float)Data[Index].R, (float)Data[Index].G, (float)Data[Index].B);// V(x0, y1, z0)
+			Index = VF_MakeIndex(VF_Tile(IndexX + 1, VolumeSizeX), VF_Tile(IndexY + 1, VolumeSizeY), VF_Tile(IndexZ, VolumeSizeZ));
+			check(Index < (uint32)(VolumeSizeX*VolumeSizeY*VolumeSizeZ));
+			V1 = FVector((float)Data[Index].R, (float)Data[Index].G, (float)Data[Index].B); // V(x1, y1, z0)
+			FVector C10 = FMath::Lerp(V1, V2, VecFrac.X);
+
+			Index = VF_MakeIndex(VF_Tile(IndexX, VolumeSizeX), VF_Tile(IndexY + 1, VolumeSizeY), VF_Tile(IndexZ + 1, VolumeSizeZ));
+			check(Index < (uint32)(VolumeSizeX*VolumeSizeY*VolumeSizeZ));
+			V1 = FVector((float)Data[Index].R, (float)Data[Index].G, (float)Data[Index].B);// V(x0, y1, z1)
+			Index = VF_MakeIndex(VF_Tile(IndexX + 1, VolumeSizeX), VF_Tile(IndexY + 1, VolumeSizeY), VF_Tile(IndexZ + 1, VolumeSizeZ));
+			check(Index < (uint32)(VolumeSizeX*VolumeSizeY*VolumeSizeZ));
+			V1 = FVector((float)Data[Index].R, (float)Data[Index].G, (float)Data[Index].B);// V(x1, y1, z1)
+			FVector C11 = FMath::Lerp(V1, V2, VecFrac.X);
+
+			FVector C0 = FMath::Lerp(C00, C10, VecFrac.Y);
+			FVector C1 = FMath::Lerp(C01, C11, VecFrac.Y);
+			FVector C = FMath::Lerp(C0, C1, VecFrac.Z);
+
+			// Write final output...
+			*OutSampleX.GetDest() = C.X;
+			*OutSampleY.GetDest() = C.Y;
+			*OutSampleZ.GetDest() = C.Z;
+
+			XParam.Advance();
+			YParam.Advance();
+			ZParam.Advance();
+			OutSampleX.Advance();
+			OutSampleY.Advance();
+			OutSampleZ.Advance();
+		}
+	}
+
+	Unlock();
+}
+
+
+bool UNiagaraDataInterfaceVectorField::CanExecuteOnTarget(ENiagaraSimTarget Target)const 
+{
+	return true; 
+}
+
+bool UNiagaraDataInterfaceVectorField::CopyToInternal(UNiagaraDataInterface* Destination) const
+{
+	if (!Super::CopyToInternal(Destination))
+	{
+		return false;
+	}
+
+	UNiagaraDataInterfaceVectorField* OtherTyped = CastChecked<UNiagaraDataInterfaceVectorField>(Destination);
+	OtherTyped->Field = Field;
+	OtherTyped->bTileX = bTileX;
+	OtherTyped->bTileY = bTileY;
+	OtherTyped->bTileZ = bTileZ;
+	OtherTyped->bGPUBufferDirty = true;
+	return true;
+}
+
+bool UNiagaraDataInterfaceVectorField::Equals(const UNiagaraDataInterface* Other) const
+{
+	if (!Super::Equals(Other))
+	{
+		return false;
+	}
+	const UNiagaraDataInterfaceVectorField* OtherTyped = CastChecked<const UNiagaraDataInterfaceVectorField>(Other);
+	return OtherTyped->Field == Field 
+		&& OtherTyped->bTileX == bTileX
+		&& OtherTyped->bTileY == bTileY
+		&& OtherTyped->bTileZ == bTileZ;
+}
+
+
+int32 UNiagaraDataInterfaceVectorField::PerInstanceDataSize()const  
+{
+	return 0;
+}
+
+bool UNiagaraDataInterfaceVectorField::InitPerInstanceData(void* PerInstanceData, FNiagaraSystemInstance* SystemInstance)
+{
+	//FNDIVectorField_InstanceData* Inst = new (PerInstanceData) FNDIVectorField_InstanceData();
+	//return Inst->Init(this, SystemInstance);
+	return false;
+}
+
+void UNiagaraDataInterfaceVectorField::DestroyPerInstanceData(void* PerInstanceData, FNiagaraSystemInstance* SystemInstance)
+{
+	/*FNDIVectorField_InstanceData* Inst = (FNDIVectorField_InstanceData*)PerInstanceData;
+	Inst->~FNDIVectorField_InstanceData();*/
+}
+
+bool UNiagaraDataInterfaceVectorField::PerInstanceTick(void* PerInstanceData, FNiagaraSystemInstance* SystemInstance, float InDeltaSeconds)
+{
+	/*FNDIVectorField_InstanceData* Inst = (FNDIVectorField_InstanceData*)PerInstanceData;
+	return Inst->Tick(this, SystemInstance, InDeltaSeconds);*/
+	return false;
+}
+
+bool UNiagaraDataInterfaceVectorField::GetFunctionHLSL(const FName& DefinitionFunctionName, FString InstanceFunctionName, TArray<FDIGPUBufferParamDescriptor> &Descriptors, FString &HLSLInterfaceID, FString &OutHLSL)
+{
+	//FString BufferName = Descriptors[0].BufferParamName;
+	if (DefinitionFunctionName == SampleVectorFieldName)
+	{
+		FString BufferName = Descriptors[0].BufferParamName;
+		OutHLSL += TEXT("void ") + InstanceFunctionName + TEXT("(float3 In_SamplePoint,  out float3 Out_Value) \n{\n");
+		OutHLSL += TEXT("\t int3 IntCoord = int3(0,0,0);\n");
+		OutHLSL += TEXT("\t int3 Dimensions = int3(1,1,1);\n");
+		OutHLSL += TEXT("\t Out_Value = ") + BufferName + TEXT("[IntCoord.x + IntCoord.y*Dimensions.x + IntCoord.z*Dimensions.x*Dimensions.y].xyz;\n");
+		OutHLSL += TEXT("\n}\n");
+		return true;
+	}
+	else if (DefinitionFunctionName == GetVectorDimsName)
+	{
+		OutHLSL += TEXT("void ") + InstanceFunctionName + TEXT("(out float3 Out_Value) \n{\n");
+		OutHLSL += TEXT("\t Out_Value = float3(1.0f, 1.0f, 1.0f);\n");
+		OutHLSL += TEXT("\n}\n");
+		return true;
+	}
+	else if (DefinitionFunctionName == GetVectorFieldBoundsName)
+	{
+		OutHLSL += TEXT("void ") + InstanceFunctionName + TEXT("(out float3 Out_MinBounds, out float3 Out_MaxBounds) \n{\n");
+		OutHLSL += TEXT("\t Out_MinBounds = Out_MaxBounds = float3(1.0f, 1.0f, 1.0f);\n");
+		OutHLSL += TEXT("\n}\n");
+		return true;
+	}
+	return false;
+}
+
+// build buffer definition hlsl
+// 1. Choose a buffer name, add the data interface ID (important!)
+// 2. add a DIGPUBufferParamDescriptor to the array argument; that'll be passed on to the FNiagaraShader for binding to a shader param, that can
+// then later be found by name via FindDIBufferParam for setting; 
+// 3. store buffer declaration hlsl in OutHLSL
+// multiple buffers can be defined at once here
+//
+void UNiagaraDataInterfaceVectorField::GetBufferDefinitionHLSL(FString DataInterfaceID, TArray<FDIGPUBufferParamDescriptor> &BufferDescriptors, FString &OutHLSL)
+{
+	FString BufferName = "VectorFieldLUT" + DataInterfaceID;
+	OutHLSL += TEXT("Buffer<float4> ") + BufferName + TEXT(";\n");
+
+	BufferDescriptors.Add(FDIGPUBufferParamDescriptor(BufferName, 0));		// add a descriptor for shader parameter binding
+}
+
+// called after translate, to setup buffers matching the buffer descriptors generated during hlsl translation
+// need to do this because the script used during translate is a clone, including its DIs
+//
+void UNiagaraDataInterfaceVectorField::SetupBuffers(FDIBufferDescriptorStore &BufferDescriptors)
+{
+	GPUBuffers.Empty();
+	for (FDIGPUBufferParamDescriptor &Desc : BufferDescriptors.Descriptors)
+	{
+		FNiagaraDataInterfaceBufferData BufferData(*Desc.BufferParamName);	// store off the data for later use
+		GPUBuffers.Add(BufferData);
+	}
+	bGPUBufferDirty = true;
+}
+
+// return the GPU buffer array (called from NiagaraInstanceBatcher to get the buffers for setting to the shader)
+// we lazily update the buffer with a new LUT here if necessary
+//
+TArray<FNiagaraDataInterfaceBufferData> &UNiagaraDataInterfaceVectorField::GetBufferDataArray()
+{
+	check(IsInRenderingThread());
+
+	if (bGPUBufferDirty)
+	{
+		UVectorFieldStatic* StaticVectorField = Cast<UVectorFieldStatic>(Field);
+		UVectorFieldAnimated* AnimatedVectorField = Cast<UVectorFieldAnimated>(Field);
+
+		uint32 Width = 10;
+		uint32 Height = 10;
+		uint32 Depth = 10;
+		const FFloat16Color* Data = nullptr;
+
+		if (StaticVectorField)
+		{
+			Width = StaticVectorField->SizeX;
+			Height = StaticVectorField->SizeY;
+			Depth = StaticVectorField->SizeZ;
+		}
+		
+		check(GPUBuffers.Num() > 0);
+		FNiagaraDataInterfaceBufferData &GPUBuffer = GPUBuffers[0];
+		GPUBuffer.Buffer.Release();
+		uint32 BufferSize = Width * Height * Depth * sizeof(float) * 4;
+		TResourceArray<FVector4> TempTable;
+		TempTable.AddUninitialized(Width*Height*Depth);
+		
+		if (StaticVectorField)
+		{
+			Data = (const FFloat16Color*)StaticVectorField->SourceData.LockReadOnly();
+
+			// TODO sckime - can we just send this down as half4's ???
+			for (uint32 z = 0; z < Depth; z++)
+			{
+				for (uint32 y = 0; y < Height; y++)
+				{
+					for (uint32 x = 0; x < Width; x++)
+					{
+						uint32 Index = x + (y* Height) + (z * Width * Height);
+						TempTable[Index] = FVector4(Data[Index].R, Data[Index].G, Data[Index].B, 0.0f);
+					}
+				}
+			}
+			StaticVectorField->SourceData.Unlock();
+		}
+		else
+		{
+			// We don't support the other types for now.
+			for (uint32 z = 0; z < Depth; z++)
+			{
+				for (uint32 y = 0; y < Height; y++)
+				{
+					for (uint32 x = 0; x < Width; x++)
+					{
+						uint32 Index = x + (y* Height) + (z * Width * Height);
+						TempTable[Index] = FVector4(1.0f, 0.0f, 0.0f, 0.0f);
+					}
+				}
+			}
+		}
+
+		GPUBuffer.Buffer.Initialize(sizeof(float) * 4, Width*Height*Depth, EPixelFormat::PF_A32B32G32R32F, BUF_Static, TEXT("CurlnoiseTable"), &TempTable);
+		//FPlatformMemory::Memcpy(BufferData, TempTable, BufferSize);
+		//RHIUnlockVertexBuffer(GPUBuffer.Buffer.Buffer);
+		bGPUBufferDirty = false;
+	}
+	return GPUBuffers;
+}
+
+
+
+#undef LOCTEXT_NAMESPACE
\ No newline at end of file
diff --git a/Engine/Plugins/FX/Niagara/Source/Niagara/Private/NiagaraDataSet.cpp b/Engine/Plugins/FX/Niagara/Source/Niagara/Private/NiagaraDataSet.cpp
--- a/Engine/Plugins/FX/Niagara/Source/Niagara/Private/NiagaraDataSet.cpp
+++ b/Engine/Plugins/FX/Niagara/Source/Niagara/Private/NiagaraDataSet.cpp
@@ -5,6 +5,9 @@
 #include "NiagaraShader.h"
 #include "GlobalShader.h"
 #include "UpdateTextureShaders.h"
+#include "NiagaraStats.h"
+
+DECLARE_CYCLE_STAT(TEXT("InitRenderData"), STAT_InitRenderData, STATGROUP_Niagara);
 
 //////////////////////////////////////////////////////////////////////////
 void FNiagaraDataSet::SetShaderParams(FNiagaraShader *Shader, FRHICommandList &CommandList)
@@ -68,11 +71,105 @@ void FNiagaraDataSet::UnsetShaderParams(FNiagaraShader *Shader, FRHICommandList
 	}
 }
 
+//////////////////////////////////////////////////////////////////////////
+static int32 GRenderDataBlockSize = 4096;
+static FAutoConsoleVariableRef CVarRenderDataBlockSize(
+	TEXT("fx.RenderDataBlockSize"),
+	GRenderDataBlockSize,
+	TEXT("Size of alloaction blocks for Niagara render data. \n"),
+	ECVF_Default
+);
+
+static float GGPUBufferShrinkFactor = .5f;
+static FAutoConsoleVariableRef CVarNiagaraRenderBufferShrinkFactor(
+	TEXT("fx.RenderBufferShrinkFactor"),
+	GGPUBufferShrinkFactor,
+	TEXT("What factor should the render buffers grow by when they need to grow. \n"),
+	ECVF_Default
+);
+
+void FNiagaraDataSet::InitGPUFromCPU_RenderThread()
+{
+	SCOPE_CYCLE_COUNTER(STAT_InitRenderData);
+	check(IsInRenderingThread());
+	
+	//Some very basic size handling so we don't have to recreate these buffers every frame. TODO: Improve.
+	auto CalcNewSize = [](int32 ReqNum, int32 CurrNum)
+	{
+		int32 NewSize = AlignArbitrary(ReqNum, GRenderDataBlockSize);
+		int32 ShrinkSize = AlignArbitrary(CurrNum * GGPUBufferShrinkFactor, GRenderDataBlockSize);
+		if (ReqNum < ShrinkSize)
+		{
+			NewSize = ShrinkSize;
+		}
+
+		check(NewSize >= ReqNum);
+		return NewSize;
+	};
+
+	if (GetNumFloatComponents())
+	{
+		int32 CurrBytes = RenderDataFloat.NumBytes;
+		int32 ValidBytes = CurrDataRender().GetFloatBuffer().Num();
+		int32 NewBytes = CalcNewSize(ValidBytes, CurrBytes);
+		if (NewBytes != CurrBytes)
+ 		{
+			RenderDataFloat.Release();
+
+			if (NewBytes > 0)
+			{
+				RenderDataFloat.Initialize(sizeof(float), NewBytes / sizeof(float), EPixelFormat::PF_R32_FLOAT, BUF_Volatile);
+			}
+		}
+
+		if (NewBytes > 0 && ValidBytes > 0)
+		{
+			void* BufferData = RHILockVertexBuffer(RenderDataFloat.Buffer, 0, ValidBytes, RLM_WriteOnly);
+			FPlatformMemory::Memcpy(BufferData, CurrDataRender().GetFloatBuffer().GetData(), ValidBytes);
+			RHIUnlockVertexBuffer(RenderDataFloat.Buffer);
+		}
+	}
+
+	if (GetNumInt32Components())
+	{
+		int32 CurrBytes = RenderDataInt.NumBytes;
+		int32 ValidBytes = CurrDataRender().GetInt32Buffer().Num();
+		int32 NewBytes = CalcNewSize(ValidBytes, CurrBytes);
+ 		if (NewBytes != CurrBytes)
+ 		{
+			RenderDataInt.Release();
+
+			if (NewBytes > 0)
+			{
+				RenderDataInt.Initialize(sizeof(int32), NewBytes / sizeof(int32), EPixelFormat::PF_R32_SINT, BUF_Volatile);
+			}
+		}
+
+		if (NewBytes > 0 && ValidBytes > 0)
+		{
+			void*BufferData = RHILockVertexBuffer(RenderDataInt.Buffer, 0, ValidBytes, RLM_WriteOnly);
+			FPlatformMemory::Memcpy(BufferData, CurrDataRender().GetInt32Buffer().GetData(), ValidBytes);
+			RHIUnlockVertexBuffer(RenderDataInt.Buffer);
+		}
+	}
+
+	CurrentFloatDataSRV = RenderDataFloat.SRV;
+	CurrentIntDataSRV = RenderDataInt.SRV;
+}
+
+
+void FNiagaraDataSet::InitGPUSimSRVs_RenderThread()
+{
+	CurrentFloatDataSRV = CurrDataRender().GetGPUBufferFloat()->SRV;
+	CurrentIntDataSRV = CurrDataRender().GetGPUBufferInt()->SRV;
+}
+
+
 void FNiagaraDataSet::Dump(FNiagaraDataSet& Other, bool bCurr)
 {
 	Other.Reset();
 	Other.Variables = Variables;
-	Other.VariableLayoutMap = VariableLayoutMap;
+	Other.VariableLayouts = VariableLayouts;
 
 	int32 IndexRead = CurrBuffer > 0 ? CurrBuffer - 1 : MaxBufferIdx;
 	if (bCurr)
@@ -106,7 +203,7 @@ void FNiagaraDataSet::Dump(bool bCurr, int32 StartIdx, int32 NumInstances)
 	int32 NumInstancesDumped = 0;
 	TArray<FString> Lines;
 	Lines.Reserve(GetNumInstances());
-	while (Itr.IsValid() && NumInstancesDumped++ < NumInstances)
+	while (Itr.IsValid() && NumInstancesDumped < NumInstances)
 	{
 		Itr.Get();
 
@@ -117,6 +214,7 @@ void FNiagaraDataSet::Dump(bool bCurr, int32 StartIdx, int32 NumInstances)
 		}
 		Lines.Add(Line);
 		Itr.Advance();
+		NumInstancesDumped++;
 	}
 
 	static FString Sep;
@@ -153,6 +251,7 @@ void FNiagaraDataBuffer::Allocate(uint32 InNumInstances, ENiagaraSimTarget Targe
 	if (Target == ENiagaraSimTarget::CPUSim)
 	{
 		NumInstancesAllocated = InNumInstances;
+		NumInstances = 0;
 
 		int32 OldFloatStride = FloatStride;
 		FloatStride = GetSafeComponentBufferSize(NumInstancesAllocated * sizeof(float));
@@ -209,13 +308,18 @@ void FNiagaraDataBuffer::AllocateGPU(uint32 InNumInstances, FRHICommandList &RHI
 		NumChunksAllocatedForGPU = ((InNumInstances + ALLOC_CHUNKSIZE - 1) / ALLOC_CHUNKSIZE);
 		uint32 NumElementsToAlloc = NumChunksAllocatedForGPU * ALLOC_CHUNKSIZE;
 
+		if (NumElementsToAlloc == 0)
+		{
+			return;
+		}
+
 		if (Owner->GetNumFloatComponents())
 		{
 			if (GPUBufferFloat.Buffer)
 			{
 				GPUBufferFloat.Release();
 			}
-			GPUBufferFloat.Initialize(sizeof(float), NumElementsToAlloc * Owner->GetNumFloatComponents(), EPixelFormat::PF_R32_FLOAT);
+			GPUBufferFloat.Initialize(sizeof(float), NumElementsToAlloc * Owner->GetNumFloatComponents(), EPixelFormat::PF_R32_FLOAT, BUF_Static);
 		}
 		if (Owner->GetNumInt32Components())
 		{
@@ -223,26 +327,11 @@ void FNiagaraDataBuffer::AllocateGPU(uint32 InNumInstances, FRHICommandList &RHI
 			{
 				GPUBufferInt.Release();
 			}
-			GPUBufferInt.Initialize(sizeof(int32), NumElementsToAlloc * Owner->GetNumInt32Components(), EPixelFormat::PF_R32_SINT);
+			GPUBufferInt.Initialize(sizeof(int32), NumElementsToAlloc * Owner->GetNumInt32Components(), EPixelFormat::PF_R32_SINT, BUF_Static);
 		}
 	}
 }
 
-
-void FNiagaraDataBuffer::InitGPUFromCPU()
-{
-	if (Owner->GetNumFloatComponents())
-	{
-		GPUBufferFloat.Release();
-		GPUBufferFloat.Initialize(sizeof(float), FloatStride/sizeof(float) * Owner->GetNumFloatComponents(), EPixelFormat::PF_R32_FLOAT, 0, TEXT("GPUBufferFloat"), &FloatData);
-	}
-	if (Owner->GetNumInt32Components())
-	{
-		GPUBufferInt.Release();
-		GPUBufferInt.Initialize(sizeof(int32), Int32Stride/sizeof(int32) * Owner->GetNumInt32Components(), EPixelFormat::PF_R32_SINT, 0, TEXT("GPUBufferInt"), &Int32Data);
-	}
-}
-
 void FNiagaraDataBuffer::SwapInstances(uint32 OldIndex, uint32 NewIndex) 
 {
 	for (uint32 CompIdx = 0; CompIdx < Owner->TotalFloatComponents; ++CompIdx)
diff --git a/Engine/Plugins/FX/Niagara/Source/Niagara/Private/NiagaraEmitter.cpp b/Engine/Plugins/FX/Niagara/Source/Niagara/Private/NiagaraEmitter.cpp
--- a/Engine/Plugins/FX/Niagara/Source/Niagara/Private/NiagaraEmitter.cpp
+++ b/Engine/Plugins/FX/Niagara/Source/Niagara/Private/NiagaraEmitter.cpp
@@ -18,7 +18,6 @@ static FAutoConsoleVariableRef CVarForceNiagaraCompileOnLoad(
 	ECVF_Default
 	);
 
-
 void FNiagaraEmitterScriptProperties::InitDataSetAccess()
 {
 	EventReceivers.Empty();
@@ -26,6 +25,7 @@ void FNiagaraEmitterScriptProperties::InitDataSetAccess()
 
 	if (Script)
 	{
+		//UE_LOG(LogNiagara, Log, TEXT("InitDataSetAccess: %s %d %d"), *Script->GetPathName(), Script->ReadDataSets.Num(), Script->WriteDataSets.Num());
 		// TODO: add event receiver and generator lists to the script properties here
 		//
 		for (FNiagaraDataSetID &ReadID : Script->ReadDataSets)
@@ -41,12 +41,41 @@ void FNiagaraEmitterScriptProperties::InitDataSetAccess()
 	}
 }
 
+bool FNiagaraEmitterScriptProperties::DataSetAccessSynchronized() const
+{
+	if (Script)
+	{
+		if (Script->ReadDataSets.Num() != EventReceivers.Num())
+		{
+			return false;
+		}
+		if (Script->WriteDataSets.Num() != EventGenerators.Num())
+		{
+			return false;
+		}
+		return true;
+	}
+	else
+	{
+		return EventReceivers.Num() == 0 && EventGenerators.Num() == 0;
+	}
+}
+
 //////////////////////////////////////////////////////////////////////////
 
 UNiagaraEmitter::UNiagaraEmitter(const FObjectInitializer& Initializer)
 : Super(Initializer)
 , CollisionMode(ENiagaraCollisionMode::None)
+#if WITH_EDITORONLY_DATA
+, ThumbnailImageOutOfDate(true)
+#endif
+, FixedBounds(FBox(FVector(-100), FVector(100)))
+, MinDetailLevel(0)
+, MaxDetailLevel(4)
 , bInterpolatedSpawning(false)
+, bFixedBounds(false)
+, bUseMinDetailLevel(false)
+, bUseMaxDetailLevel(false)
 {
 }
 
@@ -62,7 +91,15 @@ void UNiagaraEmitter::PostInitProperties()
 
 		UpdateScriptProps.Script = NewObject<UNiagaraScript>(this, "UpdateScript", EObjectFlags::RF_Transactional);
 		UpdateScriptProps.Script->SetUsage(ENiagaraScriptUsage::ParticleUpdateScript);
+
+		EmitterSpawnScriptProps.Script = NewObject<UNiagaraScript>(this, "EmitterSpawnScript", EObjectFlags::RF_Transactional);
+		EmitterSpawnScriptProps.Script->SetUsage(ENiagaraScriptUsage::EmitterSpawnScript);
+		
+		EmitterUpdateScriptProps.Script = NewObject<UNiagaraScript>(this, "EmitterUpdateScript", EObjectFlags::RF_Transactional);
+		EmitterUpdateScriptProps.Script->SetUsage(ENiagaraScriptUsage::EmitterUpdateScript);
+
 	}
+	UniqueEmitterName = TEXT("Emitter");
 }
 
 void UNiagaraEmitter::Serialize(FArchive& Ar)
@@ -93,8 +130,25 @@ void UNiagaraEmitter::PostLoad()
 // 			UE_LOG(LogNiagara, Warning, TEXT("Niagara Emitter had mis matched bInterpolatedSpawning values. You must recompile this emitter. %s"), *GetFullName());
 // 		}
 // 	}
+
+	if (EmitterSpawnScriptProps.Script == nullptr || EmitterUpdateScriptProps.Script == nullptr)
+	{
+		EmitterSpawnScriptProps.Script = NewObject<UNiagaraScript>(this, "EmitterSpawnScript", EObjectFlags::RF_Transactional);
+		EmitterSpawnScriptProps.Script->SetUsage(ENiagaraScriptUsage::EmitterSpawnScript);
+
+		EmitterUpdateScriptProps.Script = NewObject<UNiagaraScript>(this, "EmitterUpdateScript", EObjectFlags::RF_Transactional);
+		EmitterUpdateScriptProps.Script->SetUsage(ENiagaraScriptUsage::EmitterUpdateScript);
+
+#if WITH_EDITORONLY_DATA
+		if (SpawnScriptProps.Script)
+		{
+			EmitterSpawnScriptProps.Script->SetSource(SpawnScriptProps.Script->GetSource());
+			EmitterUpdateScriptProps.Script->SetSource(SpawnScriptProps.Script->GetSource());
+		}
+#endif
+	}
 		
-	//Temporarily disabling interpolated spawn.
+	//Temporarily disabling interpolated spawn if the script type and flag don't match.
 	if (SpawnScriptProps.Script)
 	{
 		SpawnScriptProps.Script->ConditionalPostLoad();
@@ -104,13 +158,11 @@ void UNiagaraEmitter::PostLoad()
 			bInterpolatedSpawning = false;
 			if (bActualInterpolatedSpawning)
 			{
-				SpawnScriptProps.Script->ByteCode.Empty();//clear out the script as it was compiled with interpolated spawn.
-#if WITH_EDITORONLY_DATA
-				SpawnScriptProps.Script->InvalidateChangeID(); // Identify that we have no idea if we're synchronized or not.
-#endif
+				SpawnScriptProps.Script->InvalidateScript();//clear out the script as it was compiled with interpolated spawn.
+
 				SpawnScriptProps.Script->SetUsage(ENiagaraScriptUsage::ParticleSpawnScript);
 			}
-			UE_LOG(LogNiagara, Warning, TEXT("Temporarily disabling interpolated spawn. Emitter may need recompile.. %s"), *GetFullName());
+			UE_LOG(LogNiagara, Warning, TEXT("Disabling interpolated spawn because emitter flag and script type don't match. Did you adjust this value in the UI? Emitter may need recompile.. %s"), *GetFullName());
 		}
 	}
 
@@ -121,7 +173,7 @@ void UNiagaraEmitter::PostLoad()
 	if (SpawnScriptProps.Script)
 	{
 		SpawnScriptProps.Script->ConditionalPostLoad();
-		if (SpawnScriptProps.Script->ByteCode.Num() == 0 || GbForceNiagaraCompileOnLoad > 0)
+		if (SpawnScriptProps.Script->GetByteCode().Num() == 0 || GbForceNiagaraCompileOnLoad > 0)
 		{
 			bNeedsRecompile = true;
 		}
@@ -129,7 +181,7 @@ void UNiagaraEmitter::PostLoad()
 	if (UpdateScriptProps.Script)
 	{
 		UpdateScriptProps.Script->ConditionalPostLoad();
-		if (UpdateScriptProps.Script->ByteCode.Num() == 0 || GbForceNiagaraCompileOnLoad > 0)
+		if (UpdateScriptProps.Script->GetByteCode().Num() == 0 || GbForceNiagaraCompileOnLoad > 0)
 		{
 			bNeedsRecompile = true;
 		}
@@ -140,13 +192,18 @@ void UNiagaraEmitter::PostLoad()
 		if (EventHandlerScriptProps[i].Script)
 		{
 			EventHandlerScriptProps[i].Script->ConditionalPostLoad();
-			if (EventHandlerScriptProps[i].Script->ByteCode.Num() == 0 || GbForceNiagaraCompileOnLoad > 0)
+			if (EventHandlerScriptProps[i].Script->GetByteCode().Num() == 0 || GbForceNiagaraCompileOnLoad > 0)
 			{
 				bNeedsRecompile = true;
 			}
 		}
 	}
 
+#if WITH_EDITORONLY_DATA
+	GraphSource->ConditionalPostLoad();
+	bNeedsRecompile |= GraphSource->PostLoadFromEmitter(*this);
+#endif
+
 	// Check ourselves against the current script rebuild version. Force our children
 	// to need to update.
 	const int32 NiagaraVer = GetLinkerCustomVersion(FNiagaraCustomVersion::GUID);
@@ -159,27 +216,18 @@ void UNiagaraEmitter::PostLoad()
 	{
 		if (SpawnScriptProps.Script)
 		{
-			SpawnScriptProps.Script->ByteCode.Empty();//clear out the script.
-#if WITH_EDITORONLY_DATA
-			SpawnScriptProps.Script->InvalidateChangeID(); // Identify that we have no idea if we're synchronized or not.
-#endif
+			SpawnScriptProps.Script->InvalidateScript();//clear out the script.
 		}
 		if (UpdateScriptProps.Script)
 		{
-			UpdateScriptProps.Script->ByteCode.Empty();//clear out the script.
-#if WITH_EDITORONLY_DATA
-			UpdateScriptProps.Script->InvalidateChangeID(); // Identify that we have no idea if we're synchronized or not.
-#endif
+			UpdateScriptProps.Script->InvalidateScript();//clear out the script.
 		}
 
 		for (int32 i = 0; i < EventHandlerScriptProps.Num(); i++)
 		{
 			if (EventHandlerScriptProps[i].Script)
 			{
-				EventHandlerScriptProps[i].Script->ByteCode.Empty();//clear out the script.
-#if WITH_EDITORONLY_DATA
-				EventHandlerScriptProps[i].Script->InvalidateChangeID(); // Identify that we have no idea if we're synchronized or not.
-#endif
+				EventHandlerScriptProps[i].Script->InvalidateScript();//clear out the script.
 			}
 		}
 	}
@@ -196,7 +244,7 @@ void UNiagaraEmitter::PostLoad()
 			TArray<FString> ScriptErrors;
 			TArray<FString> PathNames;
 			TArray<UNiagaraScript*> Scripts;
-			CompileScripts(ScriptStatuses, ScriptErrors, PathNames, Scripts);
+			CompileScripts(ScriptStatuses, ScriptErrors, PathNames, Scripts, false);
 
 			for (int32 i = 0; i < ScriptStatuses.Num(); i++)
 			{
@@ -211,6 +259,36 @@ void UNiagaraEmitter::PostLoad()
 			}
 		}
 	}
+
+	GraphSource->OnChanged().AddUObject(this, &UNiagaraEmitter::GraphSourceChanged);
+
+	EmitterSpawnScriptProps.Script->RapidIterationParameters.AddOnChangedHandler(
+		FNiagaraParameterStore::FOnChanged::FDelegate::CreateUObject(this, &UNiagaraEmitter::ScriptRapidIterationParameterChanged));
+	EmitterUpdateScriptProps.Script->RapidIterationParameters.AddOnChangedHandler(
+		FNiagaraParameterStore::FOnChanged::FDelegate::CreateUObject(this, &UNiagaraEmitter::ScriptRapidIterationParameterChanged));
+
+	if (SpawnScriptProps.Script)
+	{
+		SpawnScriptProps.Script->RapidIterationParameters.AddOnChangedHandler(
+			FNiagaraParameterStore::FOnChanged::FDelegate::CreateUObject(this, &UNiagaraEmitter::ScriptRapidIterationParameterChanged));
+	}
+	
+	if (UpdateScriptProps.Script)
+	{
+		UpdateScriptProps.Script->RapidIterationParameters.AddOnChangedHandler(
+			FNiagaraParameterStore::FOnChanged::FDelegate::CreateUObject(this, &UNiagaraEmitter::ScriptRapidIterationParameterChanged));
+	}
+
+	for (FNiagaraEventScriptProperties& EventScriptProperties : EventHandlerScriptProps)
+	{
+		EventScriptProperties.Script->RapidIterationParameters.AddOnChangedHandler(
+			FNiagaraParameterStore::FOnChanged::FDelegate::CreateUObject(this, &UNiagaraEmitter::ScriptRapidIterationParameterChanged));
+	}
+
+	for (UNiagaraRendererProperties* Renderer : RendererProperties)
+	{
+		Renderer->OnChanged().AddUObject(this, &UNiagaraEmitter::RendererChanged);
+	}
 #endif
 }
 
@@ -239,21 +317,118 @@ void UNiagaraEmitter::PostEditChangeProperty(struct FPropertyChangedEvent& Prope
 			}
 		}
 	}
+	if (PropertyName == GET_MEMBER_NAME_CHECKED(UNiagaraEmitter, SimTarget))
+	{
+		if (GraphSource != nullptr)
+		{
+			GraphSource->MarkNotSynchronized();
+		}
+	}
+	ThumbnailImageOutOfDate = true;
+	ChangeId = FGuid::NewGuid();
+	OnPropertiesChangedDelegate.Broadcast();
+}
+
+
+UNiagaraEmitter::FOnPropertiesChanged& UNiagaraEmitter::OnPropertiesChanged()
+{
+	return OnPropertiesChangedDelegate;
 }
 #endif
 
 
 bool UNiagaraEmitter::IsValid()const
 {
-	return SpawnScriptProps.Script && SpawnScriptProps.Script->IsValid() && UpdateScriptProps.Script && UpdateScriptProps.Script->IsValid();
+	if (!SpawnScriptProps.Script || !UpdateScriptProps.Script)
+	{
+		return false;
+	}
+
+	if (SimTarget == ENiagaraSimTarget::CPUSim || SimTarget == ENiagaraSimTarget::DynamicLoadBalancedSim)
+	{
+		if (!SpawnScriptProps.Script->IsScriptCompilationPending(false) && !SpawnScriptProps.Script->DidScriptCompilationSucceed(false))
+		{
+			return false;
+		}
+		if (!UpdateScriptProps.Script->IsScriptCompilationPending(false) && !UpdateScriptProps.Script->DidScriptCompilationSucceed(false))
+		{
+			return false;
+		}
+		if (EventHandlerScriptProps.Num() != 0)
+		{
+			for (int32 i = 0; i < EventHandlerScriptProps.Num(); i++)
+			{
+				if (!EventHandlerScriptProps[i].Script->IsScriptCompilationPending(false) &&
+					!EventHandlerScriptProps[i].Script->DidScriptCompilationSucceed(false))
+				{
+					return false;
+				}
+			}
+		}
+	}
+
+	if (SimTarget == ENiagaraSimTarget::GPUComputeSim || SimTarget == ENiagaraSimTarget::DynamicLoadBalancedSim)
+	{
+		if (!SpawnScriptProps.Script->IsScriptCompilationPending(true) && 
+			!SpawnScriptProps.Script->DidScriptCompilationSucceed(true))
+		{
+			return false;
+		}
+	}
+	return true;
 }
 
-#if WITH_EDITORONLY_DATA
+bool UNiagaraEmitter::IsReadyToRun() const
+{
+	//Check for various failure conditions and bail.
+	if (!UpdateScriptProps.Script || !SpawnScriptProps.Script)
+	{
+		return false;
+	}
 
-void UNiagaraEmitter::GetScripts(TArray<UNiagaraScript*>& OutScripts)
+	if (SimTarget == ENiagaraSimTarget::CPUSim || SimTarget == ENiagaraSimTarget::DynamicLoadBalancedSim)
+	{
+		if (SpawnScriptProps.Script->IsScriptCompilationPending(false))
+		{
+			return false;
+		}
+		if (UpdateScriptProps.Script->IsScriptCompilationPending(false))
+		{
+			return false;
+		}
+		if (EventHandlerScriptProps.Num() != 0)
+		{
+			for (int32 i = 0; i < EventHandlerScriptProps.Num(); i++)
+			{
+				if (EventHandlerScriptProps[i].Script->IsScriptCompilationPending(false))
+				{
+					return false;
+				}
+			}
+		}
+	}
+
+	if (SimTarget == ENiagaraSimTarget::GPUComputeSim || SimTarget == ENiagaraSimTarget::DynamicLoadBalancedSim)
+	{
+		if (SpawnScriptProps.Script->IsScriptCompilationPending(true))
+		{
+			return false;
+		}
+	}
+
+	return true;
+}
+
+void UNiagaraEmitter::GetScripts(TArray<UNiagaraScript*>& OutScripts, bool bCompilableOnly)
 {
 	OutScripts.Add(SpawnScriptProps.Script);
 	OutScripts.Add(UpdateScriptProps.Script);
+	if (!bCompilableOnly)
+	{
+		OutScripts.Add(EmitterSpawnScriptProps.Script);
+		OutScripts.Add(EmitterUpdateScriptProps.Script);
+	}
+
 	for (int32 i = 0; i < EventHandlerScriptProps.Num(); i++)
 	{
 		if (EventHandlerScriptProps[i].Script)
@@ -263,8 +438,73 @@ void UNiagaraEmitter::GetScripts(TArray<UNiagaraScript*>& OutScripts)
 	}
 }
 
+UNiagaraScript* UNiagaraEmitter::GetScript(ENiagaraScriptUsage Usage, FGuid UsageId)
+{
+	TArray<UNiagaraScript*> Scripts;
+	GetScripts(Scripts, false);
+	for (UNiagaraScript* Script : Scripts)
+	{
+		if (Script->IsEquivalentUsage(Usage) && Script->GetUsageId() == UsageId)
+		{
+			return Script;
+		}
+	}
+	return nullptr;
+}
+
+bool UNiagaraEmitter::IsAllowedByDetailLevel()const
+{
+	int32 DetailLevel = INiagaraModule::GetDetailLevel();
+	if ((bUseMinDetailLevel && DetailLevel < MinDetailLevel) || (bUseMaxDetailLevel && DetailLevel > MaxDetailLevel))
+	{
+		return false;
+	}
+
+	return true;
+}
+
+#if WITH_EDITORONLY_DATA
+
+FGuid UNiagaraEmitter::GetChangeId() const
+{
+	return ChangeId;
+}
+
+bool UNiagaraEmitter::AreAllScriptAndSourcesSynchronized() const
+{
+	if (SpawnScriptProps.Script->IsCompilable() && !SpawnScriptProps.Script->AreScriptAndSourceSynchronized())
+	{
+		return false;
+	}
+
+	if (UpdateScriptProps.Script->IsCompilable() && !UpdateScriptProps.Script->AreScriptAndSourceSynchronized())
+	{
+		return false;
+	}
+
+	if (EmitterSpawnScriptProps.Script->IsCompilable() && !EmitterSpawnScriptProps.Script->AreScriptAndSourceSynchronized())
+	{
+		return false;
+	}
+
+	if (EmitterUpdateScriptProps.Script->IsCompilable() && !EmitterUpdateScriptProps.Script->AreScriptAndSourceSynchronized())
+	{
+		return false;
+	}
+
+	for (int32 i = 0; i < EventHandlerScriptProps.Num(); i++)
+	{
+		if (EventHandlerScriptProps[i].Script && EventHandlerScriptProps[i].Script->IsCompilable() && !EventHandlerScriptProps[i].Script->AreScriptAndSourceSynchronized())
+		{
+			return false;
+		}
+	}
+
+	return true;
+}
+
 void UNiagaraEmitter::CompileScripts(TArray<ENiagaraScriptCompileStatus>& OutScriptStatuses, TArray<FString>& OutGraphLevelErrorMessages, TArray<FString>& PathNames,
-	TArray<UNiagaraScript*>& Scripts)
+	TArray<UNiagaraScript*>& Scripts, bool bForce)
 {
 	PathNames.AddDefaulted(2);
 	OutScriptStatuses.Empty();
@@ -287,18 +527,25 @@ void UNiagaraEmitter::CompileScripts(TArray<ENiagaraScriptCompileStatus>& OutScr
 		bPostCompile = true;
 	}
 
+	TArray<int32> OutputScriptIndices;
+	TArray<int32> OutputStatusIndices;
+
 	FString ErrorMsg;
 	check((int32)EScriptCompileIndices::SpawnScript == 0);
-	OutScriptStatuses.Add(SpawnScriptProps.Script->Compile(ErrorMsg));
+	int32 StatusIdx = OutScriptStatuses.Add(SpawnScriptProps.Script->Compile(ErrorMsg, bForce));
+	OutputStatusIndices.Add(StatusIdx);
 	OutGraphLevelErrorMessages.Add(ErrorMsg); // EScriptCompileIndices::SpawnScript
-	Scripts.Add(SpawnScriptProps.Script);
+	int32 ScriptAddIdx = Scripts.Add(SpawnScriptProps.Script);
+	OutputScriptIndices.Add(ScriptAddIdx);
 	SpawnScriptProps.InitDataSetAccess();
 	PathNames[(int32)EScriptCompileIndices::SpawnScript] = SpawnScriptProps.Script->GetPathName();
 
 	ErrorMsg.Empty();
 	check((int32)EScriptCompileIndices::UpdateScript == 1);
-	OutScriptStatuses.Add(UpdateScriptProps.Script->Compile(ErrorMsg));
-	Scripts.Add(UpdateScriptProps.Script);
+	StatusIdx = OutScriptStatuses.Add(UpdateScriptProps.Script->Compile(ErrorMsg, bForce));
+	OutputStatusIndices.Add(StatusIdx);
+	ScriptAddIdx = Scripts.Add(UpdateScriptProps.Script);
+	OutputScriptIndices.Add(ScriptAddIdx);
 	OutGraphLevelErrorMessages.Add(ErrorMsg); // EScriptCompileIndices::UpdateScript
 	UpdateScriptProps.InitDataSetAccess();
 	PathNames[(int32)EScriptCompileIndices::UpdateScript] = UpdateScriptProps.Script->GetPathName();
@@ -309,31 +556,41 @@ void UNiagaraEmitter::CompileScripts(TArray<ENiagaraScriptCompileStatus>& OutScr
 		{
 			check((int32)EScriptCompileIndices::EventScript == 2);
 			ErrorMsg.Empty();
-			OutScriptStatuses.Add(EventHandlerScriptProps[i].Script->Compile(ErrorMsg));
-			Scripts.Add(EventHandlerScriptProps[i].Script);
+			StatusIdx = OutScriptStatuses.Add(EventHandlerScriptProps[i].Script->Compile(ErrorMsg, bForce));
+			OutputStatusIndices.Add(StatusIdx);
+			ScriptAddIdx = Scripts.Add(EventHandlerScriptProps[i].Script);
+			OutputScriptIndices.Add(ScriptAddIdx);
 			OutGraphLevelErrorMessages.Add(ErrorMsg); // EScriptCompileIndices::EventScript
 			EventHandlerScriptProps[i].InitDataSetAccess();
 			PathNames.Add(EventHandlerScriptProps[i].Script->GetPathName());
 		}
 	}
 
+	check(OutputScriptIndices.Num() == OutputStatusIndices.Num());
+	check(OutGraphLevelErrorMessages.Num() == OutputScriptIndices.Num());
+
 	// Go through all attached renderers and make sure that the scripts work properly for those renderers.
-	for (int32 ScriptIdx = 0; ScriptIdx < Scripts.Num(); ScriptIdx++)
+	for (int32 OutputScriptIdx = 0; OutputScriptIdx < OutputScriptIndices.Num(); OutputScriptIdx++)
 	{
-		// Only care about spawn and update
-		if (ScriptIdx >= 2)
+		int32 LoopScriptIdx = OutputScriptIndices[OutputScriptIdx];
+		int32 LoopStatusIdx = OutputStatusIndices[OutputScriptIdx];
+		UNiagaraScript* CurrentScript = Scripts[LoopScriptIdx];
+
+		// Only care about scripts that successfully compiled so far.
+		if (OutScriptStatuses[LoopStatusIdx] != ENiagaraScriptCompileStatus::NCS_UpToDate)
 		{
 			continue;
 		}
 
-		// Only care about scripts that successfully compiled so far.
-		if (OutScriptStatuses[ScriptIdx] != ENiagaraScriptCompileStatus::NCS_UpToDate)
+		// Ignore missing scripts.
+		if (CurrentScript == nullptr)
 		{
 			continue;
 		}
 
-		// Ignore missing scripts.
-		if (Scripts[ScriptIdx] == nullptr)
+		// Only care about particle scripts
+		if (!(CurrentScript->IsParticleSpawnScript() ||
+			CurrentScript->IsParticleUpdateScript()))
 		{
 			continue;
 		}
@@ -353,25 +610,23 @@ void UNiagaraEmitter::CompileScripts(TArray<ENiagaraScriptCompileStatus>& OutScr
 						Attr.SetName(*AttrName);
 					}
 
-					if (false == Scripts[ScriptIdx]->Attributes.ContainsByPredicate([&Attr](const FNiagaraVariable& Var) { return Var.GetName() == Attr.GetName(); }))
+					if (false == Scripts[LoopScriptIdx]->Attributes.ContainsByPredicate([&Attr](const FNiagaraVariable& Var) { return Var.GetName() == Attr.GetName(); }))
 					{
-						OutGraphLevelErrorMessages[ScriptIdx] += FString::Printf(TEXT("\nCannot bind to renderer %s because it does not define attribute %s %s."), *RendererProperties[i]->GetName(), *Attr.GetType().GetNameText().ToString(), *Attr.GetName().ToString());
-						OutScriptStatuses[ScriptIdx] = ENiagaraScriptCompileStatus::NCS_Error;
+						OutGraphLevelErrorMessages[LoopStatusIdx] += FString::Printf(TEXT("\nCannot bind to renderer %s because it does not define attribute %s %s."), *RendererProperties[i]->GetName(), *Attr.GetType().GetNameText().ToString(), *Attr.GetName().ToString());
+						OutScriptStatuses[LoopStatusIdx] = ENiagaraScriptCompileStatus::NCS_Error;
 					}
 				}
 			}
 		}
 	}
 
-	ChangeId = FGuid::NewGuid();
-
 	if (GraphSource && bPostCompile) // Cleanup if we precompiled inside this function
 	{
 		GraphSource->PostCompile();
 	}
 }
 
-ENiagaraScriptCompileStatus UNiagaraEmitter::CompileScript(EScriptCompileIndices InScriptToCompile, FString& OutGraphLevelErrorMessages, int32 InSubScriptIdx)
+ENiagaraScriptCompileStatus UNiagaraEmitter::CompileScript(EScriptCompileIndices InScriptToCompile, FString& OutGraphLevelErrorMessages, bool bForce, int32 InSubScriptIdx)
 {
 	ENiagaraScriptCompileStatus CompileStatus = ENiagaraScriptCompileStatus::NCS_Unknown;
 	bool bPostCompile = false;
@@ -384,17 +639,17 @@ ENiagaraScriptCompileStatus UNiagaraEmitter::CompileScript(EScriptCompileIndices
 	switch (InScriptToCompile)
 	{
 	case EScriptCompileIndices::SpawnScript:
-		CompileStatus = SpawnScriptProps.Script->Compile(OutGraphLevelErrorMessages);
+		CompileStatus = SpawnScriptProps.Script->Compile(OutGraphLevelErrorMessages, bForce);
 		SpawnScriptProps.InitDataSetAccess();
 		break;
 	case EScriptCompileIndices::UpdateScript:
-		CompileStatus = UpdateScriptProps.Script->Compile(OutGraphLevelErrorMessages);
+		CompileStatus = UpdateScriptProps.Script->Compile(OutGraphLevelErrorMessages, bForce);
 		UpdateScriptProps.InitDataSetAccess();
 		break;
 	case EScriptCompileIndices::EventScript:
 		if (EventHandlerScriptProps[InSubScriptIdx].Script)
 		{
-			CompileStatus = EventHandlerScriptProps[InSubScriptIdx].Script->Compile(OutGraphLevelErrorMessages);
+			CompileStatus = EventHandlerScriptProps[InSubScriptIdx].Script->Compile(OutGraphLevelErrorMessages, bForce);
 			EventHandlerScriptProps[InSubScriptIdx].InitDataSetAccess();
 		}
 		else
@@ -405,7 +660,6 @@ ENiagaraScriptCompileStatus UNiagaraEmitter::CompileScript(EScriptCompileIndices
 	default:
 		return CompileStatus;
 	}
-	ChangeId = FGuid::NewGuid();
 
 	if (GraphSource && bPostCompile) // Cleanup if we precompiled inside this function
 	{
@@ -451,12 +705,24 @@ UNiagaraEmitter* UNiagaraEmitter::MakeRecursiveDeepCopy(UObject* DestOuter, TMap
 		check(Props->GraphSource == Props->UpdateScriptProps.Script->GetSource());
 	}
 
-	for (int32 i = 0; i < Props->EventHandlerScriptProps.Num(); i++)
+	if (Props->EmitterSpawnScriptProps.Script)
+	{
+		Props->EmitterSpawnScriptProps.Script->SubsumeExternalDependencies(ExistingConversions);
+		check(Props->GraphSource == Props->EmitterSpawnScriptProps.Script->GetSource());
+	}
+	if (Props->EmitterUpdateScriptProps.Script)
+	{
+		Props->EmitterUpdateScriptProps.Script->SubsumeExternalDependencies(ExistingConversions);
+		check(Props->GraphSource == Props->EmitterUpdateScriptProps.Script->GetSource());
+	}
+
+
+	for (int32 i = 0; i < Props->GetEventHandlers().Num(); i++)
 	{
-		if (Props->EventHandlerScriptProps[i].Script)
+		if (Props->GetEventHandlers()[i].Script)
 		{
-			Props->EventHandlerScriptProps[i].Script->SubsumeExternalDependencies(ExistingConversions);
-			check(Props->GraphSource == Props->EventHandlerScriptProps[i].Script->GetSource());
+			Props->GetEventHandlers()[i].Script->SubsumeExternalDependencies(ExistingConversions);
+			check(Props->GraphSource == Props->GetEventHandlers()[i].Script->GetSource());
 		}
 	}
 	return Props;
@@ -465,7 +731,7 @@ UNiagaraEmitter* UNiagaraEmitter::MakeRecursiveDeepCopy(UObject* DestOuter, TMap
 
 bool UNiagaraEmitter::UsesScript(const UNiagaraScript* Script)const
 {
-	if (SpawnScriptProps.Script == Script || UpdateScriptProps.Script == Script)
+	if (SpawnScriptProps.Script == Script || UpdateScriptProps.Script == Script || EmitterSpawnScriptProps.Script == Script || EmitterUpdateScriptProps.Script == Script)
 	{
 		return true;
 	}
@@ -506,12 +772,173 @@ bool UNiagaraEmitter::UsesCollection(const class UNiagaraParameterCollection* Co
 
 FString UNiagaraEmitter::GetUniqueEmitterName()const
 {
-	return GetName() + LexicalConversion::ToString(FCrc::StrCrc32(*GetPathName()));
+	return UniqueEmitterName;
 }
 
-FNiagaraVariable UNiagaraEmitter::GetEmitterParameter(const FNiagaraVariable& EmitterVar)const
+void UNiagaraEmitter::SetUniqueEmitterName(const FString& InName)
+{
+	if (InName != UniqueEmitterName)
+	{
+		Modify();
+		FString OldName = UniqueEmitterName;
+		UniqueEmitterName = InName;
+
+		TMap<FString, FString> RenameMap;
+		RenameMap.Add(OldName, InName);
+
+		TArray<UNiagaraScript*> Scripts;
+		GetScripts(Scripts, false);
+
+		for (UNiagaraScript* Script : Scripts)
+		{
+			Script->Modify();
+
+			// First handle any rapid iteration parameters...
+			TArray<FNiagaraVariable> Params;
+			Script->RapidIterationParameters.GetParameters(Params);
+			for (FNiagaraVariable Var : Params)
+			{
+				FNiagaraVariable NewVar = FNiagaraVariable::ResolveAliases(Var, RenameMap);
+				if (NewVar.GetName() != Var.GetName())
+				{
+					Script->RapidIterationParameters.RenameParameter(Var, NewVar.GetName());
+				}
+			}
+
+			// Now handle any Parameters overall..
+			for (int32 i = 0; i < Script->Parameters.Parameters.Num(); i++)
+			{
+				FNiagaraVariable Var = Script->Parameters.Parameters[i];
+				FNiagaraVariable NewVar = FNiagaraVariable::ResolveAliases(Var, RenameMap);
+				if (NewVar.GetName() != Var.GetName())
+				{
+					Script->Parameters.Parameters[i] = NewVar;
+				}
+			}
+
+			// Also handle any data set mappings...
+			auto Iterator = Script->DataSetToParameters.CreateIterator();
+			while (Iterator)
+			{
+				for (int32 i = 0; i < Iterator.Value().Parameters.Num(); i++)
+				{
+					FNiagaraVariable Var = Iterator.Value().Parameters[i];
+					FNiagaraVariable NewVar = FNiagaraVariable::ResolveAliases(Var, RenameMap);
+					if (NewVar.GetName() != Var.GetName())
+					{
+						Iterator.Value().Parameters[i] = NewVar;
+					}
+				}
+				++Iterator;
+			}
+		}
+
+		// Since we use the emitter name as part of the script, we now need to recompile to work properly.
+#if WITH_EDITORONLY_DATA
+		if (GraphSource != nullptr)
+		{
+			GraphSource->MarkNotSynchronized();
+		}
+#endif
+	}
+}
+
+FNiagaraVariable UNiagaraEmitter::ToEmitterParameter(const FNiagaraVariable& EmitterVar)const
 {
 	FNiagaraVariable Var = EmitterVar;
 	Var.SetName(*Var.GetName().ToString().Replace(TEXT("Emitter."), *(GetUniqueEmitterName() + TEXT("."))));
 	return Var;
-}
\ No newline at end of file
+}
+
+void UNiagaraEmitter::AddRenderer(UNiagaraRendererProperties* Renderer)
+{
+	Modify();
+	RendererProperties.Add(Renderer);
+#if WITH_EDITOR
+	Renderer->OnChanged().AddUObject(this, &UNiagaraEmitter::RendererChanged);
+	UpdateChangeId();
+#endif
+}
+
+void UNiagaraEmitter::RemoveRenderer(UNiagaraRendererProperties* Renderer)
+{
+	Modify();
+	RendererProperties.Remove(Renderer);
+#if WITH_EDITOR
+	Renderer->OnChanged().RemoveAll(this);
+	UpdateChangeId();
+#endif
+}
+
+FNiagaraEventScriptProperties* UNiagaraEmitter::GetEventHandlerByIdUnsafe(FGuid ScriptUsageId)
+{
+	for (FNiagaraEventScriptProperties& EventScriptProperties : EventHandlerScriptProps)
+	{
+		if (EventScriptProperties.Script->GetUsageId() == ScriptUsageId)
+		{
+			return &EventScriptProperties;
+		}
+	}
+	return nullptr;
+}
+
+void UNiagaraEmitter::AddEventHandler(FNiagaraEventScriptProperties EventHandler)
+{
+	EventHandlerScriptProps.Add(EventHandler);
+#if WITH_EDITOR
+	EventHandler.Script->RapidIterationParameters.AddOnChangedHandler(
+		FNiagaraParameterStore::FOnChanged::FDelegate::CreateUObject(this, &UNiagaraEmitter::ScriptRapidIterationParameterChanged));
+	UpdateChangeId();
+#endif
+}
+
+void UNiagaraEmitter::RemoveEventHandlerByUsageId(FGuid EventHandlerUsageId)
+{
+	auto FindEventHandlerById = [=](const FNiagaraEventScriptProperties& EventHandler) { return EventHandler.Script->GetUsageId() == EventHandlerUsageId; };
+#if WITH_EDITOR
+	FNiagaraEventScriptProperties* EventHandler = EventHandlerScriptProps.FindByPredicate(FindEventHandlerById);
+	if (EventHandler != nullptr)
+	{
+		EventHandler->Script->RapidIterationParameters.RemoveAllOnChangedHandlers(this);
+	}
+#endif
+	EventHandlerScriptProps.RemoveAll(FindEventHandlerById);
+#if WITH_EDITOR
+	UpdateChangeId();
+#endif
+}
+
+void UNiagaraEmitter::BeginDestroy()
+{
+#if WITH_EDITOR
+	if (GraphSource != nullptr)
+	{
+		GraphSource->OnChanged().RemoveAll(this);
+	}
+#endif
+	Super::BeginDestroy();
+}
+
+#if WITH_EDITORONLY_DATA
+
+void UNiagaraEmitter::UpdateChangeId()
+{
+	Modify();
+	ChangeId = FGuid::NewGuid();
+}
+
+void UNiagaraEmitter::ScriptRapidIterationParameterChanged()
+{
+	UpdateChangeId();
+}
+
+void UNiagaraEmitter::RendererChanged()
+{
+	UpdateChangeId();
+}
+
+void UNiagaraEmitter::GraphSourceChanged()
+{
+	UpdateChangeId();
+}
+#endif
\ No newline at end of file
