Copying Orion-Staging @ 3748653 (Orion/Dev-General @ 3722124) to //UE4/Main

diff --git a/Engine/Plugins/FX/Niagara/Source/Niagara/Private/NiagaraEmitterHandle.cpp b/Engine/Plugins/FX/Niagara/Source/Niagara/Private/NiagaraEmitterHandle.cpp
--- a/Engine/Plugins/FX/Niagara/Source/Niagara/Private/NiagaraEmitterHandle.cpp
+++ b/Engine/Plugins/FX/Niagara/Source/Niagara/Private/NiagaraEmitterHandle.cpp
@@ -6,14 +6,17 @@
 #include "NiagaraScriptSourceBase.h"
 #include "NiagaraCommon.h"
 #include "NiagaraDataInterface.h"
+#include "NiagaraModule.h"
+
+#include "Modules/ModuleManager.h"
 
 const FNiagaraEmitterHandle FNiagaraEmitterHandle::InvalidHandle;
 
-FNiagaraEmitterHandle::FNiagaraEmitterHandle() 
-	:
+FNiagaraEmitterHandle::FNiagaraEmitterHandle() :
 #if WITH_EDITORONLY_DATA
-	Source(nullptr)
-	,
+	Source(nullptr),
+	LastMergedSource(nullptr),
+	bIsolated(false),
 #endif
 	Instance(nullptr)
 {
@@ -26,20 +29,26 @@ FNiagaraEmitterHandle::FNiagaraEmitterHandle(UNiagaraEmitter& Emitter)
 	, Name(TEXT("Emitter"))
 #if WITH_EDITORONLY_DATA
 	, Source(&Emitter)
+	, LastMergedSource(&Emitter)
+	, bIsolated(false)
 #endif
 	, Instance(&Emitter)
 {
+	
 }
 
 #if WITH_EDITORONLY_DATA
-FNiagaraEmitterHandle::FNiagaraEmitterHandle(const UNiagaraEmitter& InSourceEmitter, FName InName, UNiagaraSystem& InOuterSystem)
+FNiagaraEmitterHandle::FNiagaraEmitterHandle(UNiagaraEmitter& InSourceEmitter, FName InName, UNiagaraSystem& InOuterSystem)
 	: Id(FGuid::NewGuid())
 	, IdName(*Id.ToString())
 	, bIsEnabled(true)
 	, Name(InName)
 	, Source(&InSourceEmitter)
-	, Instance(InSourceEmitter.MakeRecursiveDeepCopy(&InOuterSystem))
+	, LastMergedSource(Cast<UNiagaraEmitter>(StaticDuplicateObject(Source, &InOuterSystem)))
+	, bIsolated(false)
+	, Instance(Cast<UNiagaraEmitter>(StaticDuplicateObject(Source, &InOuterSystem)))
 {
+	Instance->SetUniqueEmitterName(Name.ToString());
 }
 
 FNiagaraEmitterHandle::FNiagaraEmitterHandle(const FNiagaraEmitterHandle& InHandleToDuplicate, FName InDuplicateName, UNiagaraSystem& InDuplicateOwnerSystem)
@@ -48,8 +57,14 @@ FNiagaraEmitterHandle::FNiagaraEmitterHandle(const FNiagaraEmitterHandle& InHand
 	, bIsEnabled(InHandleToDuplicate.bIsEnabled)
 	, Name(InDuplicateName)
 	, Source(InHandleToDuplicate.Source)
-	, Instance(InHandleToDuplicate.Instance->MakeRecursiveDeepCopy(&InDuplicateOwnerSystem))
+	, LastMergedSource(InHandleToDuplicate.LastMergedSource != nullptr ? Cast<UNiagaraEmitter>(StaticDuplicateObject(InHandleToDuplicate.LastMergedSource, &InDuplicateOwnerSystem)) : nullptr)
+	, bIsolated(false)
+	, Instance(Cast<UNiagaraEmitter>(StaticDuplicateObject(InHandleToDuplicate.Instance, &InDuplicateOwnerSystem)))
 {
+	// Clear stand alone and public flags from the referenced emitters since they are not assets.
+	Instance->ClearFlags(RF_Standalone | RF_Public);
+	Instance->SetUniqueEmitterName(Name.ToString());
+	LastMergedSource->ClearFlags(RF_Standalone | RF_Public);
 }
 #endif
 
@@ -73,9 +88,30 @@ FName FNiagaraEmitterHandle::GetName() const
 	return Name;
 }
 
-void FNiagaraEmitterHandle::SetName(FName InName)
+void FNiagaraEmitterHandle::SetName(FName InName, UNiagaraSystem& InOwnerSystem)
 {
-	Name = InName;
+	if (InName == Name)
+	{
+		return;
+	}
+
+	FString TempNameString = InName.ToString();
+	TempNameString.ReplaceInline(TEXT(" "), TEXT("_"));
+	TempNameString.ReplaceInline(TEXT("\t"), TEXT("_"));
+	InName = *TempNameString;
+
+	TSet<FName> OtherEmitterNames;
+	for (const FNiagaraEmitterHandle& OtherEmitterHandle : InOwnerSystem.GetEmitterHandles())
+	{
+		if (OtherEmitterHandle.GetId() != GetId())
+		{
+			OtherEmitterNames.Add(OtherEmitterHandle.GetName());
+		}
+	}
+	FName UniqueName = FNiagaraUtilities::GetUniqueName(InName, OtherEmitterNames);
+
+	Name = UniqueName;
+	Instance->SetUniqueEmitterName(Name.ToString());
 }
 
 bool FNiagaraEmitterHandle::GetIsEnabled() const
@@ -107,57 +143,25 @@ FString FNiagaraEmitterHandle::GetUniqueInstanceName()const
 	return Instance->GetUniqueEmitterName();
 }
 
-void FNiagaraEmitterHandle::SetInstance(UNiagaraEmitter* InInstance)
-{
-	Instance = InInstance;
-}
 #if WITH_EDITORONLY_DATA
 
-void FNiagaraEmitterHandle::ResetToSource()
-{
-	UObject* Outer = GetInstance()->GetOuter();
-	Instance = nullptr; // Clear out our reference to this object first..
-	Instance = GetSource()->MakeRecursiveDeepCopy(Outer);
-}
-
-void CopyParameterValues(UNiagaraScript* Script, UNiagaraScript* PreviousScript)
+bool FNiagaraEmitterHandle::IsSynchronizedWithSource() const
 {
-	for (FNiagaraVariable& InputParameter : Script->Parameters.Parameters)
-	{
-		for (FNiagaraVariable& PreviousInputParameter : PreviousScript->Parameters.Parameters)
-		{
-			if (PreviousInputParameter.IsDataAllocated())
-			{
-				if (InputParameter.GetName() == PreviousInputParameter.GetName() &&
-					InputParameter.GetType() == PreviousInputParameter.GetType())
-				{
-					InputParameter.AllocateData();
-					PreviousInputParameter.CopyTo(InputParameter.GetData());
-				}
-			}
-		}
-	}
-	
-	for (FNiagaraScriptDataInterfaceInfo& InputInfo : Script->DataInterfaceInfo)
+	if (Source == nullptr || LastMergedSource == nullptr)
 	{
-		for (FNiagaraScriptDataInterfaceInfo&  PreviousInputInfo : PreviousScript->DataInterfaceInfo)
-		{
-			if (InputInfo.Name == PreviousInputInfo.Name &&
-				InputInfo.DataInterface->GetClass() == PreviousInputInfo.DataInterface->GetClass())
-			{
-				PreviousInputInfo.CopyTo(&InputInfo, Script->GetOuter());
-			}
-		}
+		// If either the source or last merged sources is missing, then we're not synchronized.  The
+		// merge logic will detect this and print an appropriate message to the log.
+		return false;
 	}
-}
 
-bool FNiagaraEmitterHandle::IsSynchronizedWithSource() const
-{
-	if (Instance && Source && Source->ChangeId.IsValid() && Instance->ChangeId.IsValid())
+	if (Source->GetChangeId().IsValid() == false ||
+		LastMergedSource->GetChangeId().IsValid() == false)
 	{
-		return Instance->ChangeId == Source->ChangeId;		
+		// If any of the change Ids aren't valid then we assume we're out of sync.
+		return false;
 	}
-	return false;
+
+	return Source->GetChangeId() == LastMergedSource->GetChangeId();
 }
 
 bool FNiagaraEmitterHandle::NeedsRecompile() const
@@ -167,7 +171,7 @@ bool FNiagaraEmitterHandle::NeedsRecompile() const
 
 	for (UNiagaraScript* Script : Scripts)
 	{
-		if (!Script->AreScriptAndSourceSynchronized())
+		if (Script->IsCompilable() && !Script->AreScriptAndSourceSynchronized())
 		{
 			return true;
 		}
@@ -175,105 +179,60 @@ bool FNiagaraEmitterHandle::NeedsRecompile() const
 	return false;
 }
 
-bool FNiagaraEmitterHandle::RefreshFromSource()
+void FNiagaraEmitterHandle::ConditionalPostLoad()
 {
-	// TODO: Update this to support events.
-	UNiagaraScript* PreviousInstanceSpawnScript = Instance->SpawnScriptProps.Script;
-	UNiagaraScript* PreviousInstanceUpdateScript = Instance->UpdateScriptProps.Script;
-	TMap<const UObject*, UObject*> ExistingConversions;
+	if (Source != nullptr)
+	{
+		Source->ConditionalPostLoad();
+	}
+	if (LastMergedSource != nullptr)
+	{
+		LastMergedSource->ConditionalPostLoad();
+	}
+	Instance->ConditionalPostLoad();
+}
 
+INiagaraModule::FMergeEmitterResults FNiagaraEmitterHandle::MergeSourceChanges()
+{
 	if (Source == nullptr)
 	{
-		return false;
+		// If we don't have a copy of the source emitter, this emitter can't safely be merged.
+		INiagaraModule::FMergeEmitterResults MergeResults;
+		MergeResults.bSucceeded = false;
+		MergeResults.bModifiedGraph = false;
+		MergeResults.ErrorMessages.Add(NSLOCTEXT("NiagaraEmitterHandle", "NoSourceErrorMessage", "This emitter has no 'Source' so changes can't be merged in."));
+		return MergeResults;
 	}
 
-	// First we need to deep copy the graph, which is shared amongst children
-	UNiagaraScriptSourceBase* GraphSource = Source->GraphSource->MakeRecursiveDeepCopy(Instance, ExistingConversions);
-	ExistingConversions.Add(Source->GraphSource, GraphSource);
-
-	// The script graphs will look up from ExistingConversions
-	UNiagaraScript* NewSpawnScript = Source->SpawnScriptProps.Script->MakeRecursiveDeepCopy(Instance, ExistingConversions);
-	UNiagaraScript* NewUpdateScript = Source->UpdateScriptProps.Script->MakeRecursiveDeepCopy(Instance, ExistingConversions);
-
-	if (NewSpawnScript->GetLastCompileStatus() == ENiagaraScriptCompileStatus::NCS_UpToDate && NewUpdateScript->GetLastCompileStatus() == ENiagaraScriptCompileStatus::NCS_UpToDate)
+	if (LastMergedSource == nullptr)
 	{
-		Instance->GraphSource = GraphSource;
-		Instance->ChangeId = Source->ChangeId;
-		Instance->SpawnScriptProps.Script = NewSpawnScript;
-		Instance->UpdateScriptProps.Script = NewUpdateScript;
-		
-		// We may have overridden the interpolated spawn script value from the default. If so, this requires
-		// a recompile.
-		bool bNeedRecompile = false;
-		if ((int32)Instance->SpawnScriptProps.Script->IsInterpolatedParticleSpawnScript() != Instance->bInterpolatedSpawning)
-		{
-			Instance->SpawnScriptProps.Script->SetUsage(ENiagaraScriptUsage::ParticleSpawnScriptInterpolated);
-			bNeedRecompile = true;
-		}
-
-		int32 NumDifference = Source->EventHandlerScriptProps.Num() - Instance->EventHandlerScriptProps.Num();
-		if (NumDifference < 0)
-		{
-			Instance->EventHandlerScriptProps.SetNum(Source->EventHandlerScriptProps.Num());
-		}
-		else if (NumDifference > 0)
-		{
-			Instance->EventHandlerScriptProps.AddDefaulted(NumDifference);
-		}
-		check(Source->EventHandlerScriptProps.Num() == Instance->EventHandlerScriptProps.Num());
-
-		for (int32 i = 0; i < Source->EventHandlerScriptProps.Num(); i++)
-		{
-			if (Instance->EventHandlerScriptProps[i].Script == nullptr)
-			{
-				Instance->EventHandlerScriptProps[i] = Source->EventHandlerScriptProps[i];
-				Instance->EventHandlerScriptProps[i].Script = Source->EventHandlerScriptProps[i].Script->MakeRecursiveDeepCopy(Instance, ExistingConversions);
-				check(Instance->EventHandlerScriptProps[i].Script->GetSource() == Instance->GraphSource);
-
-			}
-			else if (Source->EventHandlerScriptProps[i].Script != nullptr && Source->GraphSource == Source->EventHandlerScriptProps[i].Script->GetSource())
-			{
-				// Make sure that we copy all required variables and update our bytecode too.
-				Instance->EventHandlerScriptProps[i].Script = Source->EventHandlerScriptProps[i].Script->MakeRecursiveDeepCopy(Instance, ExistingConversions);
-				check(Instance->EventHandlerScriptProps[i].Script->GetSource() == Instance->GraphSource);
-			}
-		}
-
-		if (bNeedRecompile)
-		{
-			TArray<ENiagaraScriptCompileStatus> OutScriptStatuses;
-			TArray<FString> OutGraphLevelErrorMessages;
-			TArray<FString> PathNames;
-			TArray<UNiagaraScript*> Scripts;
-			Instance->CompileScripts(OutScriptStatuses, OutGraphLevelErrorMessages, PathNames, Scripts);
-
-			//ensure that we are now sync'ed with source
-			Instance->ChangeId = Source->ChangeId;
-		}
-
-		// We wait until after the potential recompile above to copy the old values over because
-		// the script parameters array is also written by the CompileScripts call.
-		CopyParameterValues(Instance->SpawnScriptProps.Script, PreviousInstanceSpawnScript);
-		CopyParameterValues(Instance->UpdateScriptProps.Script, PreviousInstanceUpdateScript);
-		check(Instance->SpawnScriptProps.Script->GetSource() == Instance->GraphSource);
-		check(Instance->UpdateScriptProps.Script->GetSource() == Instance->GraphSource);
-
-		UE_LOG(LogNiagara, Log, TEXT("Successful refresh from source %s"), *Instance->GetFullName());
-		check((int32)Instance->SpawnScriptProps.Script->IsInterpolatedParticleSpawnScript() == Instance->bInterpolatedSpawning);
-
-		Instance->SpawnScriptProps.InitDataSetAccess();
-		Instance->UpdateScriptProps.InitDataSetAccess();
-		for (int32 j = 0; j < Instance->EventHandlerScriptProps.Num(); j++)
-		{
-			Instance->EventHandlerScriptProps[j].InitDataSetAccess();
-		}
-		
-		return true;
+		// If we don't have a copy of the last merged source emitter, this emitter can't safely be
+		// merged.
+		INiagaraModule::FMergeEmitterResults MergeResults;
+		MergeResults.bSucceeded = false;
+		MergeResults.bModifiedGraph = false;
+		MergeResults.ErrorMessages.Add(NSLOCTEXT("NiagaraEmitterHandle", "NoLastMergedSourceErrorMessage", "This emitter has no 'LastMergedSource' so changes can't be merged in."));
+		return MergeResults;
 	}
-	else
+
+	INiagaraModule& NiagaraModule = FModuleManager::Get().GetModuleChecked<INiagaraModule>("Niagara");
+	INiagaraModule::FMergeEmitterResults MergeResults = NiagaraModule.MergeEmitter(*Source, *LastMergedSource, *Instance);
+	if (MergeResults.bSucceeded)
 	{
-		UE_LOG(LogNiagara, Warning, TEXT("Failed to refresh from source %s"), *Instance->GetFullName());
+		UObject* Outer = Instance->GetOuter();
+		Instance = MergeResults.MergedInstance;
+
+		// Rename the merged instance into this package with the correct handle name and then clear it's stand alone and public flags since it's not a root asset.
+		FName NewInstanceName = MakeUniqueObjectName(Outer, UNiagaraEmitter::StaticClass(), Name);
+		Instance->Rename(*NewInstanceName.ToString(), Outer, REN_ForceNoResetLoaders);
+		Instance->ClearFlags(RF_Standalone | RF_Public);
+		Instance->SetUniqueEmitterName(Name.ToString());
+
+		// Update the last merged source and clear it's stand alone and public flags since it's not an asset.
+		LastMergedSource = CastChecked<UNiagaraEmitter>(StaticDuplicateObject(Source, Outer));
+		LastMergedSource->ClearFlags(RF_Standalone | RF_Public);
 	}
-	return false;
+	return MergeResults;
 }
+
 #endif
\ No newline at end of file
diff --git a/Engine/Plugins/FX/Niagara/Source/Niagara/Private/NiagaraEmitterInstance.cpp b/Engine/Plugins/FX/Niagara/Source/Niagara/Private/NiagaraEmitterInstance.cpp
--- a/Engine/Plugins/FX/Niagara/Source/Niagara/Private/NiagaraEmitterInstance.cpp
+++ b/Engine/Plugins/FX/Niagara/Source/Niagara/Private/NiagaraEmitterInstance.cpp
@@ -11,6 +11,7 @@
 #include "NiagaraEmitterInstanceBatcher.h"
 #include "NiagaraScriptExecutionContext.h"
 #include "NiagaraParameterCollection.h"
+#include "NiagaraScriptExecutionContext.h"
 #include "NiagaraWorldManager.h"
 
 DECLARE_DWORD_COUNTER_STAT(TEXT("Num Custom Events"), STAT_NiagaraNumCustomEvents, STATGROUP_Niagara);
@@ -35,11 +36,12 @@ static FAutoConsoleVariableRef CVarNiagaraDumpParticleData(
 FNiagaraEmitterInstance::FNiagaraEmitterInstance(FNiagaraSystemInstance* InParentSystemInstance)
 : bError(false)
 , CPUTimeMS(0.0f)
-, ExecutionState(ENiagaraExecutionState::Active)
+, ExecutionState(ENiagaraExecutionState::Inactive)
 , CachedBounds(ForceInit)
 , ParentSystemInstance(InParentSystemInstance)
 {
 	bDumpAfterEvent = false;
+	ParticleDataSet = new FNiagaraDataSet();
 }
 
 FNiagaraEmitterInstance::~FNiagaraEmitterInstance()
@@ -48,6 +50,13 @@ FNiagaraEmitterInstance::~FNiagaraEmitterInstance()
 	ClearRenderer();
 	CachedBounds.Init();
 	UnbindParameters();
+
+	/** We defer the deletion of the particle dataset to the RT to be sure all in-flight RT commands have finished using it.*/
+	ENQUEUE_UNIQUE_RENDER_COMMAND_ONEPARAMETER(FDeleteParticleDataSetCommand,
+		FNiagaraDataSet*, DataSet, ParticleDataSet,
+		{
+			delete DataSet;
+		});
 }
 
 void FNiagaraEmitterInstance::ClearRenderer()
@@ -64,86 +73,55 @@ void FNiagaraEmitterInstance::ClearRenderer()
 	}
 }
 
-void FNiagaraEmitterInstance::Init(int32 InEmitterIdx, FName InSystemInstanceName)
+FBox FNiagaraEmitterInstance::GetBounds()
 {
-	EmitterIdx = InEmitterIdx;
-	OwnerSystemInstanceName = InSystemInstanceName;
-	Data = FNiagaraDataSet(FNiagaraDataSetID(GetEmitterHandle().GetIdName(), ENiagaraDataSetType::ParticleData));
-
-	//Init the spawn infos to the correct number for this system.
-	const TArray<FNiagaraEmitterSpawnAttributes>& EmitterSpawnInfoAttrs = ParentSystemInstance->GetSystem()->GetEmitterSpawnAttributes();
-	if (EmitterSpawnInfoAttrs.IsValidIndex(EmitterIdx))
+	if (UNiagaraEmitter* Emitter = GetEmitterHandle().GetInstance())
 	{
-		SpawnInfos.SetNum(EmitterSpawnInfoAttrs[EmitterIdx].SpawnAttributes.Num());
+		if (Emitter->bFixedBounds || Emitter->SimTarget == ENiagaraSimTarget::GPUComputeSim)//We must use fixed bounds on GPU sim.
+		{
+			return Emitter->FixedBounds.TransformBy(ParentSystemInstance->GetComponent()->GetComponentToWorld());
+		}
 	}
+	return CachedBounds;
 }
 
-void FNiagaraEmitterInstance::ResetSimulation()
+bool FNiagaraEmitterInstance::IsReadyToRun() const
 {
-	Data.ResetNumInstances();
-	Age = 0;
-	Loops = 0;
-	CollisionBatch.Reset();
-	bError = false;
-
 	const UNiagaraEmitter* PinnedProps = GetEmitterHandle().GetInstance();
 
-	SetExecutionState(ENiagaraExecutionState::Active);
-
-	if (!PinnedProps)
-	{
-		UE_LOG(LogNiagara, Error, TEXT("Unknown Error creating Niagara Simulation. Properties were null."));
-		bError = true;
-		return;
-	}
-
-	//Check for various failure conditions and bail.
-	if (!PinnedProps->UpdateScriptProps.Script || !PinnedProps->SpawnScriptProps.Script)
+	if (!PinnedProps || !PinnedProps->IsReadyToRun())
 	{
-		//TODO - Arbitrary named scripts. Would need some base functionality for Spawn/Udpate to be called that can be overriden in BPs for emitters with custom scripts.
-		UE_LOG(LogNiagara, Error, TEXT("Emitter cannot be enabled because it's doesn't have both an update and spawn script."), *PinnedProps->GetFullName());
-		bError = true;
-		return;
-	}
-	if (PinnedProps->UpdateScriptProps.Script->ByteCode.Num() == 0 && PinnedProps->SpawnScriptProps.Script->ByteCode.Num() == 0 && PinnedProps->SimTarget==ENiagaraSimTarget::CPUSim)
-	{
-		UE_LOG(LogNiagara, Error, TEXT("Emitter cannot be enabled because it's spawn or update script was not compiled correctly. %s"), *PinnedProps->GetFullName());
-		bError = true;
-		return;
+		return false;
 	}
 
-	if (PinnedProps->SpawnScriptProps.Script->DataUsage.bReadsAttriubteData)
-	{
-		UE_LOG(LogNiagara, Error, TEXT("%s reads attribute data and so cannot be used as a spawn script. The data being read would be invalid."), *PinnedProps->SpawnScriptProps.Script->GetName());
-		bError = true;
-		return;
-	}
-	if (PinnedProps->UpdateScriptProps.Script->Attributes.Num() == 0 || PinnedProps->SpawnScriptProps.Script->Attributes.Num() == 0)
-	{
-		UE_LOG(LogNiagara, Error, TEXT("This emitter cannot be enabled because it's spawn or update script doesn't have any attriubtes.."));
-		bError = true;
-		return;
-	}
+	return true;
 }
 
-void FNiagaraEmitterInstance::DirtyDataInterfaces()
+
+void FNiagaraEmitterInstance::Init(int32 InEmitterIdx, FName InSystemInstanceName)
 {
-	// Make sure that our function tables need to be regenerated...
-	SpawnExecContext.DirtyDataInterfaces();
-	UpdateExecContext.DirtyDataInterfaces();
-	for (FNiagaraScriptExecutionContext& EventContext : EventExecContexts)
+	check(ParticleDataSet);
+	FNiagaraDataSet& Data = *ParticleDataSet;
+
+	EmitterIdx = InEmitterIdx;
+	OwnerSystemInstanceName = InSystemInstanceName;
+	Data.Init(FNiagaraDataSetID(GetEmitterHandle().GetIdName(), ENiagaraDataSetType::ParticleData));
+
+	//Init the spawn infos to the correct number for this system.
+	const TArray<FNiagaraEmitterSpawnAttributes>& EmitterSpawnInfoAttrs = ParentSystemInstance->GetSystem()->GetEmitterSpawnAttributes();
+	if (EmitterSpawnInfoAttrs.IsValidIndex(EmitterIdx))
 	{
-		EventContext.DirtyDataInterfaces();
+		SpawnInfos.SetNum(EmitterSpawnInfoAttrs[EmitterIdx].SpawnAttributes.Num());
 	}
-}
-
-void FNiagaraEmitterInstance::ReInitSimulation()
-{
 	const FNiagaraEmitterHandle& EmitterHandle = GetEmitterHandle();
-	bIsEnabled = EmitterHandle.GetIsEnabled();
 
 	ResetSimulation();
 
+	if (IsDisabled())
+	{
+		return;
+	}
+
 	Data.Reset();
 	DataSetMap.Empty();
 
@@ -166,9 +144,9 @@ void FNiagaraEmitterInstance::ReInitSimulation()
 			{
 				UE_LOG(LogNiagara, Warning, TEXT("Attribute %s is used in the Update script for %s but it is not initialised in the Spawn script!"), *Attr.GetName().ToString(), *EmitterHandle.GetName().ToString());
 			}
-			for (int32 i = 0; i < PinnedProps->EventHandlerScriptProps.Num(); i++)
+			for (int32 i = 0; i < PinnedProps->GetEventHandlers().Num(); i++)
 			{
-				if (PinnedProps->EventHandlerScriptProps[i].Script && !PinnedProps->EventHandlerScriptProps[i].Script->Attributes.Find(Attr, FoundIdx))
+				if (PinnedProps->GetEventHandlers()[i].Script && !PinnedProps->GetEventHandlers()[i].Script->Attributes.Find(Attr, FoundIdx))
 				{
 					UE_LOG(LogNiagara, Warning, TEXT("Attribute %s is used in the event handler script for %s but it is not initialised in the Spawn script!"), *Attr.GetName().ToString(), *EmitterHandle.GetName().ToString());
 				}
@@ -179,8 +157,7 @@ void FNiagaraEmitterInstance::ReInitSimulation()
 	Data.AddVariables(PinnedProps->SpawnScriptProps.Script->Attributes);
 	Data.Finalize();
 
-	CollisionBatch.Init(ParentSystemInstance->GetIDName(), EmitterHandle.GetIdName());
-
+	ensure(PinnedProps->UpdateScriptProps.DataSetAccessSynchronized());
 	UpdateScriptEventDataSets.Empty();
 	for (const FNiagaraEventGeneratorProperties &GeneratorProps : PinnedProps->UpdateScriptProps.EventGenerators)
 	{
@@ -191,6 +168,7 @@ void FNiagaraEmitterInstance::ReInitSimulation()
 		UpdateScriptEventDataSets.Add(Set);
 	}
 
+	ensure(PinnedProps->SpawnScriptProps.DataSetAccessSynchronized());
 	SpawnScriptEventDataSets.Empty();
 	for (const FNiagaraEventGeneratorProperties &GeneratorProps : PinnedProps->SpawnScriptProps.EventGenerators)
 	{
@@ -203,21 +181,22 @@ void FNiagaraEmitterInstance::ReInitSimulation()
 
 	SpawnExecContext.Init(PinnedProps->SpawnScriptProps.Script, PinnedProps->SimTarget);
 	UpdateExecContext.Init(PinnedProps->UpdateScriptProps.Script, PinnedProps->SimTarget);
+	
+	// setup the parameer store for the GPU execution context; since spawn and update are combined here, we build one with params from both script props
+	if (PinnedProps->SimTarget == ENiagaraSimTarget::GPUComputeSim)
+	{
+		GPUExecContext.InitParams(PinnedProps->SpawnScriptProps.Script, PinnedProps->UpdateScriptProps.Script, PinnedProps->SimTarget);
+		SpawnExecContext.Parameters.Bind(&GPUExecContext.CombinedParamStore);
+		UpdateExecContext.Parameters.Bind(&GPUExecContext.CombinedParamStore);
+	}
 
-	EventExecContexts.SetNum(PinnedProps->EventHandlerScriptProps.Num());
-	int32 NumEvents = PinnedProps->EventHandlerScriptProps.Num();
+	EventExecContexts.SetNum(PinnedProps->GetEventHandlers().Num());
+	int32 NumEvents = PinnedProps->GetEventHandlers().Num();
 	for (int32 i = 0; i < NumEvents; i++)
 	{
-		UNiagaraScript* EventScript = PinnedProps->EventHandlerScriptProps[i].Script;
-		if (EventScript)
-		{
-			if (EventScript->ByteCode.Num() == 0)
-			{
-				UE_LOG(LogNiagara, Error, TEXT("%s has an event handler script hat didn't compile correctly."), *GetEmitterHandle().GetName().ToString());
-				bError = true;
-				return;
-			}
-		}
+		ensure(PinnedProps->GetEventHandlers()[i].DataSetAccessSynchronized());
+
+		UNiagaraScript* EventScript = PinnedProps->GetEventHandlers()[i].Script;
 
 		//This is cpu explicitly? Are we doing event handlers on GPU?
 		EventExecContexts[i].Init(EventScript, ENiagaraSimTarget::CPUSim);
@@ -226,10 +205,10 @@ void FNiagaraEmitterInstance::ReInitSimulation()
 	UNiagaraEmitter* Emitter = GetEmitterHandle().GetInstance();
 
 	//Setup direct bindings for setting parameter values.
-	SpawnIntervalBinding.Init(SpawnExecContext.Parameters, Emitter->GetEmitterParameter(SYS_PARAM_EMITTER_SPAWN_INTERVAL));
-	InterpSpawnStartBinding.Init(SpawnExecContext.Parameters, Emitter->GetEmitterParameter(SYS_PARAM_EMITTER_INTERP_SPAWN_START_DT));
+	SpawnIntervalBinding.Init(SpawnExecContext.Parameters, Emitter->ToEmitterParameter(SYS_PARAM_EMITTER_SPAWN_INTERVAL));
+	InterpSpawnStartBinding.Init(SpawnExecContext.Parameters, Emitter->ToEmitterParameter(SYS_PARAM_EMITTER_INTERP_SPAWN_START_DT));
 
-	FNiagaraVariable EmitterAgeParam = Emitter->GetEmitterParameter(SYS_PARAM_EMITTER_AGE);
+	FNiagaraVariable EmitterAgeParam = Emitter->ToEmitterParameter(SYS_PARAM_EMITTER_AGE);
 	SpawnEmitterAgeBinding.Init(SpawnExecContext.Parameters, EmitterAgeParam);
 	UpdateEmitterAgeBinding.Init(UpdateExecContext.Parameters, EmitterAgeParam);
 	EventEmitterAgeBindings.SetNum(NumEvents);
@@ -238,6 +217,15 @@ void FNiagaraEmitterInstance::ReInitSimulation()
 		EventEmitterAgeBindings[i].Init(EventExecContexts[i].Parameters, EmitterAgeParam);
 	}
 
+	FNiagaraVariable EmitterLocalSpaceParam = Emitter->ToEmitterParameter(SYS_PARAM_EMITTER_LOCALSPACE);
+	SpawnEmitterLocalSpaceBinding.Init(SpawnExecContext.Parameters, EmitterLocalSpaceParam);
+	UpdateEmitterLocalSpaceBinding.Init(UpdateExecContext.Parameters, EmitterLocalSpaceParam);
+	EventEmitterLocalSpaceBindings.SetNum(NumEvents);
+	for (int32 i = 0; i < NumEvents; i++)
+	{
+		EventEmitterLocalSpaceBindings[i].Init(EventExecContexts[i].Parameters, EmitterLocalSpaceParam);
+	}
+
 	SpawnExecCountBinding.Init(SpawnExecContext.Parameters, SYS_PARAM_ENGINE_EXEC_COUNT);
 	UpdateExecCountBinding.Init(UpdateExecContext.Parameters, SYS_PARAM_ENGINE_EXEC_COUNT);
 	EventExecCountBindings.SetNum(NumEvents);
@@ -251,6 +239,124 @@ void FNiagaraEmitterInstance::ReInitSimulation()
 		//Just ensure we've generated the singleton here on the GT as it throws a wobbler if we do this later in parallel.
 		NiagaraEmitterInstanceBatcher::Get();
 	}
+
+	//Init accessors for PostProcessParticles
+	PositionAccessor = FNiagaraDataSetAccessor<FVector>(Data, FNiagaraVariable(FNiagaraTypeDefinition::GetVec3Def(), TEXT("Position")));
+	SizeAccessor = FNiagaraDataSetAccessor<FVector2D>(Data, FNiagaraVariable(FNiagaraTypeDefinition::GetVec2Def(), TEXT("SpriteSize")));
+	MeshScaleAccessor = FNiagaraDataSetAccessor<FVector>(Data, FNiagaraVariable(FNiagaraTypeDefinition::GetVec3Def(), TEXT("Scale")));
+}
+
+void FNiagaraEmitterInstance::ResetSimulation()
+{
+	bResetPending = true;
+	Age = 0;
+	Loops = 0;
+	bError = false;
+	
+	SetExecutionState(ENiagaraExecutionState::Active);
+
+	if (IsDisabled())
+	{
+		return;
+	}
+
+	CheckForErrors();
+}
+
+void FNiagaraEmitterInstance::CheckForErrors()
+{
+	const UNiagaraEmitter* PinnedProps = GetEmitterHandle().GetInstance();
+	if (!PinnedProps)
+	{
+		UE_LOG(LogNiagara, Error, TEXT("Unknown Error creating Niagara Simulation. Properties were null."));
+		bError = true;
+		return;
+	}
+
+	//Check for various failure conditions and bail.
+	if (!PinnedProps->UpdateScriptProps.Script || !PinnedProps->SpawnScriptProps.Script)
+	{
+		//TODO - Arbitrary named scripts. Would need some base functionality for Spawn/Udpate to be called that can be overriden in BPs for emitters with custom scripts.
+		UE_LOG(LogNiagara, Error, TEXT("Emitter cannot be enabled because it's doesn't have both an update and spawn script."), *PinnedProps->GetFullName());
+		bError = true;
+		return;
+	}
+
+	if (PinnedProps->SpawnScriptProps.Script->DataUsage.bReadsAttriubteData)
+	{
+		UE_LOG(LogNiagara, Error, TEXT("%s reads attribute data and so cannot be used as a spawn script. The data being read would be invalid."), *PinnedProps->SpawnScriptProps.Script->GetName());
+		bError = true;
+		return;
+	}
+	if (PinnedProps->UpdateScriptProps.Script->Attributes.Num() == 0 || PinnedProps->SpawnScriptProps.Script->Attributes.Num() == 0)
+	{
+		UE_LOG(LogNiagara, Error, TEXT("This emitter cannot be enabled because it's spawn or update script doesn't have any attriubtes.."));
+		bError = true;
+		return;
+	}
+
+	if (PinnedProps->SimTarget == ENiagaraSimTarget::CPUSim || PinnedProps->SimTarget == ENiagaraSimTarget::DynamicLoadBalancedSim)
+	{
+		bool bFailed = false;
+		if (!PinnedProps->SpawnScriptProps.Script->DidScriptCompilationSucceed(false))
+		{
+			bFailed = true;
+			UE_LOG(LogNiagara, Error, TEXT("This emitter cannot be enabled because it's CPU Spawn script failed to compile."));
+		}
+
+		if (!PinnedProps->UpdateScriptProps.Script->DidScriptCompilationSucceed(false))
+		{
+			bFailed = true;
+			UE_LOG(LogNiagara, Error, TEXT("This emitter cannot be enabled because it's CPU Update script failed to compile."));
+		}
+
+		if (PinnedProps->GetEventHandlers().Num() != 0)
+		{
+			for (int32 i = 0; i < PinnedProps->GetEventHandlers().Num(); i++)
+			{
+				if (!PinnedProps->GetEventHandlers()[i].Script->DidScriptCompilationSucceed(false))
+				{
+					bFailed = true;
+					UE_LOG(LogNiagara, Error, TEXT("This emitter cannot be enabled because one of it's CPU Event scripts failed to compile."));
+				}
+			}
+		}
+
+		if (bFailed)
+		{
+			bError = true;
+			return;
+		}
+	}
+
+	if (PinnedProps->SimTarget == ENiagaraSimTarget::GPUComputeSim || PinnedProps->SimTarget == ENiagaraSimTarget::DynamicLoadBalancedSim)
+	{
+		if (PinnedProps->SpawnScriptProps.Script->IsScriptCompilationPending(true))
+		{
+			UE_LOG(LogNiagara, Error, TEXT("This emitter cannot be enabled because it's GPU script hasn't been compiled.."));
+			bError = true;
+			return;
+		}
+		if (!PinnedProps->SpawnScriptProps.Script->DidScriptCompilationSucceed(true))
+		{
+			UE_LOG(LogNiagara, Error, TEXT("This emitter cannot be enabled because it's GPU script failed to compile."));
+			bError = true;
+			return;
+		}
+	}
+}
+
+void FNiagaraEmitterInstance::DirtyDataInterfaces()
+{
+	// Make sure that our function tables need to be regenerated...
+	SpawnExecContext.DirtyDataInterfaces();
+	UpdateExecContext.DirtyDataInterfaces();
+	GPUExecContext.DirtyDataInterfaces();
+
+	for (FNiagaraScriptExecutionContext& EventContext : EventExecContexts)
+	{
+		EventContext.DirtyDataInterfaces();
+	}
 }
 
 //Unsure on usage of this atm. Possibly useful in future.
@@ -271,55 +377,39 @@ void FNiagaraEmitterInstance::ReInitSimulation()
 
 void FNiagaraEmitterInstance::UnbindParameters()
 {
-	FNiagaraWorldManager* WorldMan = ParentSystemInstance->GetWorldManager();
-	check(WorldMan);
-
-	//Unbind our parameter collections.
-	for (UNiagaraParameterCollection* Collection : SpawnExecContext.Script->ParameterCollections)
-	{
-		WorldMan->GetParameterCollection(Collection)->GetParameterStore().Unbind(&SpawnExecContext.Parameters);
-	}
-	for (UNiagaraParameterCollection* Collection : UpdateExecContext.Script->ParameterCollections)
-	{
-		WorldMan->GetParameterCollection(Collection)->GetParameterStore().Unbind(&UpdateExecContext.Parameters);
-	}
+	SpawnExecContext.Parameters.UnbindFromSourceStores();
+	UpdateExecContext.Parameters.UnbindFromSourceStores();
 
 	for (int32 EventIdx = 0; EventIdx < EventExecContexts.Num(); ++EventIdx)
 	{
-		for (UNiagaraParameterCollection* Collection : EventExecContexts[EventIdx].Script->ParameterCollections)
-		{
-			WorldMan->GetParameterCollection(Collection)->GetParameterStore().Unbind(&EventExecContexts[EventIdx].Parameters);
-		}
-	}
-
-	FNiagaraParameterStore& SystemParams = ParentSystemInstance->GetParameters();
-	SystemParams.Unbind(&SpawnExecContext.Parameters);
-	SystemParams.Unbind(&UpdateExecContext.Parameters);
-	for (FNiagaraScriptExecutionContext& EventContext : EventExecContexts)
-	{
-		SystemParams.Unbind(&EventContext.Parameters);
+		EventExecContexts[EventIdx].Parameters.UnbindFromSourceStores();
 	}
 }
 
 void FNiagaraEmitterInstance::BindParameters()
 {
+	if (IsDisabled())
+	{
+		return;
+	}
+
 	FNiagaraWorldManager* WorldMan = ParentSystemInstance->GetWorldManager();
 	check(WorldMan);
 
 	for (UNiagaraParameterCollection* Collection : SpawnExecContext.Script->ParameterCollections)
 	{
-		WorldMan->GetParameterCollection(Collection)->GetParameterStore().Bind(&SpawnExecContext.Parameters);
+		ParentSystemInstance->GetParameterCollectionInstance(Collection)->GetParameterStore().Bind(&SpawnExecContext.Parameters);
 	}
 	for (UNiagaraParameterCollection* Collection : UpdateExecContext.Script->ParameterCollections)
 	{
-		WorldMan->GetParameterCollection(Collection)->GetParameterStore().Bind(&UpdateExecContext.Parameters);
+		ParentSystemInstance->GetParameterCollectionInstance(Collection)->GetParameterStore().Bind(&UpdateExecContext.Parameters);
 	}
 
 	for (int32 EventIdx = 0; EventIdx < EventExecContexts.Num(); ++EventIdx)
 	{
 		for (UNiagaraParameterCollection* Collection : EventExecContexts[EventIdx].Script->ParameterCollections)
 		{
-			WorldMan->GetParameterCollection(Collection)->GetParameterStore().Bind(&EventExecContexts[EventIdx].Parameters);
+			ParentSystemInstance->GetParameterCollectionInstance(Collection)->GetParameterStore().Bind(&EventExecContexts[EventIdx].Parameters);
 		}
 	}
 
@@ -331,9 +421,19 @@ void FNiagaraEmitterInstance::BindParameters()
 	{
 		InstanceParams.Bind(&EventContext.Parameters);
 	}
+
+	const FNiagaraEmitterHandle& EmitterHandle = GetEmitterHandle();
+	const UNiagaraEmitter* Props = EmitterHandle.GetInstance();
+	Props->SpawnScriptProps.Script->RapidIterationParameters.Bind(&SpawnExecContext.Parameters);
+	Props->UpdateScriptProps.Script->RapidIterationParameters.Bind(&UpdateExecContext.Parameters);
+	ensure(Props->GetEventHandlers().Num() == EventExecContexts.Num());
+	for (int32 i = 0; i < Props->GetEventHandlers().Num(); i++)
+	{
+		Props->GetEventHandlers()[i].Script->RapidIterationParameters.Bind(&EventExecContexts[i].Parameters);
+	}
 }
 
-void FNiagaraEmitterInstance::PostResetSimulation()
+void FNiagaraEmitterInstance::PostInitSimulation()
 {
 	const FNiagaraEmitterHandle& EmitterHandle = GetEmitterHandle();
 	if (!bError)
@@ -354,13 +454,6 @@ void FNiagaraEmitterInstance::PostResetSimulation()
 			//FNiagaraDataSet* ReceiverSet = ParentSystemInstance->GetDataSet(FNiagaraDataSetID(Receiver.SourceEventGenerator, ENiagaraDataSetType::Event), Receiver.SourceEmitter);
 			const FNiagaraDataSet* ReceiverSet = FNiagaraEventDataSetMgr::GetEventDataSet(ParentSystemInstance->GetIDName(), Receiver.SourceEmitter, Receiver.SourceEventGenerator);
 		}
-
-		// add the collision event set
-		// TODO: add spawn and death event data sets (if enabled)
-		if (Props->CollisionMode != ENiagaraCollisionMode::None)
-		{
-			CollisionBatch.Init(ParentSystemInstance->GetIDName(), EmitterHandle.GetIdName());	// creates and sets up the data set for the events
-		}
 	}
 }
 
@@ -402,7 +495,8 @@ float FNiagaraEmitterInstance::GetTotalCPUTime()
 
 int FNiagaraEmitterInstance::GetTotalBytesUsed()
 {
-	int32 BytesUsed = Data.GetSizeBytes();
+	check(ParticleDataSet);
+	int32 BytesUsed = ParticleDataSet->GetSizeBytes();
 	/*
 	for (FNiagaraDataSet& Set : DataSets)
 	{
@@ -412,76 +506,125 @@ int FNiagaraEmitterInstance::GetTotalBytesUsed()
 	return BytesUsed;
 }
 
-
-/** Look for dead particles and move from the end of the list to the dead location, compacting in the process
-  * Also calculates bounds; Kill will be removed from this once we do conditional write
-  */
-void FNiagaraEmitterInstance::PostProcessParticles()
+FBox FNiagaraEmitterInstance::CalculateDynamicBounds()
 {
-	SCOPE_CYCLE_COUNTER(STAT_NiagaraKill);
-	int32 OrigNumParticles = Data.GetNumInstances();
-	int32 CurNumParticles = OrigNumParticles;
-
-	CachedBounds.Init();
-
+	check(ParticleDataSet);
+	FNiagaraDataSet& Data = *ParticleDataSet;
+	int32 NumParticles = Data.GetNumInstances();
 	const FNiagaraEmitterHandle& EmitterHandle = GetEmitterHandle();
-	if (CurNumParticles == 0 || EmitterHandle.GetInstance()->SimTarget == ENiagaraSimTarget::GPUComputeSim)
+	if (NumParticles == 0 || EmitterHandle.GetInstance()->SimTarget == ENiagaraSimTarget::GPUComputeSim)//TODO: Pull data back from gpu buffers to get bounds for GPU sims.
 	{
-		return;
+		return FBox();
 	}
 
-	FNiagaraDataSetIterator<FVector> PosItr(Data, FNiagaraVariable(FNiagaraTypeDefinition::GetVec3Def(), TEXT("Position")));
-	FNiagaraDataSetIterator<FVector2D> SizeItr(Data, FNiagaraVariable(FNiagaraTypeDefinition::GetVec2Def(), TEXT("Size")));
-	FNiagaraDataSetIterator<FVector> MeshScaleItr(Data, FNiagaraVariable(FNiagaraTypeDefinition::GetVec3Def(), TEXT("Scale")));
+	const UNiagaraEmitter* Emitter = EmitterHandle.GetInstance();
+	FBox Ret;
+	Ret.Init();
+
+	PositionAccessor.InitForAccess(true);
+	SizeAccessor.InitForAccess(true);
+	MeshScaleAccessor.InitForAccess(true);
+	int32 NumInstances = Data.GetNumInstances();
 
 	FVector MaxSize(EForceInit::ForceInitToZero);
 
-	while (PosItr.IsValid())
-	{
-		// Only increment the iterators if we don't kill a particle because
-		// we swap with the last particle and it too may have aged out, so we want
-		// to keep looping on the same index as long as it gets swapped out for a 
-		// dead particle.
+	FMatrix SystemLocalToWorld = GetParentSystemInstance()->GetComponent()->GetComponentTransform().ToMatrixNoScale();
 
+	for (int32 InstIdx = 0; InstIdx < NumInstances; ++InstIdx)
+	{
 		FVector Position;
-		PosItr.Get(Position);
+		PositionAccessor.Get(InstIdx, Position);
 
 		// Some graphs have a tendency to divide by zero. This ContainsNaN has been added prophylactically
 		// to keep us safe during GDC. It should be removed as soon as we feel safe that scripts are appropriately warned.
-		if (Position.ContainsNaN() == false)
+		if (!Position.ContainsNaN())
 		{
-			CachedBounds += Position;
+			Ret += Position;
 
 			// We advance the scale or size depending of if we use either.
-			if (MeshScaleItr.IsValid())
+			if (MeshScaleAccessor.IsValid())
 			{
-				MaxSize = MaxSize.ComponentMax(*MeshScaleItr);
-				MeshScaleItr.Advance();
+				MaxSize = MaxSize.ComponentMax(MeshScaleAccessor.Get(InstIdx));
 			}
-			else if (SizeItr.IsValid())
+			else if (SizeAccessor.IsValid())
 			{
-				MaxSize = MaxSize.ComponentMax(FVector((*SizeItr).GetMax()));
-				SizeItr.Advance();
+				MaxSize = MaxSize.ComponentMax(FVector(SizeAccessor.Get(InstIdx).GetMax()));
 			}
-			// Now we advance our main iterator since we've safely handled this particle.
-			PosItr.Advance();
 		}
 		else
 		{
-			//Must always advance otherwise we'll have inf loop if we have nans in the position
-			if (MeshScaleItr.IsValid())
-			{
-				MeshScaleItr.Advance();
-			}
-			else if (SizeItr.IsValid())
-			{
-				SizeItr.Advance();
-			}
-			PosItr.Advance();
+			UE_LOG(LogNiagara, Warning, TEXT("Particle position data contains NaNs. Likely a divide by zero somewhere in your modules."));
 		}
 	}
 
-	CachedBounds = CachedBounds.ExpandBy(MaxSize);
+	FVector MaxBaseSize = FVector(0.0001f, 0.0001f, 0.0001f);
+	for (int32 i = 0; i < EmitterRenderer.Num(); i++)
+	{
+		if (EmitterRenderer[i])
+		{
+			FVector BaseExtents = EmitterRenderer[i]->GetBaseExtents();
+			FVector ComponentMax;
+			MaxBaseSize = BaseExtents.ComponentMax(MaxBaseSize);
+		}
+	}
+
+	Ret = Ret.ExpandBy(MaxSize*MaxBaseSize);
+
+	if (Emitter->bLocalSpace)
+	{
+		Ret = Ret.TransformBy(SystemLocalToWorld);
+	}
+
+	return Ret;
+}
+
+/** Look for dead particles and move from the end of the list to the dead location, compacting in the process
+  * Also calculates bounds; Kill will be removed from this once we do conditional write
+  */
+void FNiagaraEmitterInstance::PostProcessParticles()
+{
+	SCOPE_CYCLE_COUNTER(STAT_NiagaraKill);
+	check(ParticleDataSet);
+	FNiagaraDataSet& Data = *ParticleDataSet;
+	int32 NumParticles = Data.GetNumInstances();
+
+	CachedBounds.Init();
+
+	const FNiagaraEmitterHandle& EmitterHandle = GetEmitterHandle();
+	if (NumParticles == 0 || EmitterHandle.GetInstance()->SimTarget == ENiagaraSimTarget::GPUComputeSim)
+	{
+		return;
+	}
+	
+	const UNiagaraEmitter* Emitter = EmitterHandle.GetInstance();
+	if (Emitter->bFixedBounds)
+	{
+		CachedBounds = Emitter->FixedBounds;
+	}
+	else
+	{
+		CachedBounds = CalculateDynamicBounds();
+	}
+}
+
+bool FNiagaraEmitterInstance::HandleCompletion(bool bForce)
+{
+	if (bForce)
+	{
+		SetExecutionState(ENiagaraExecutionState::Complete);
+	}
+
+	if (IsComplete())
+	{
+		//If we have any particles then clear out the buffers.
+		if (ParticleDataSet->GetNumInstances() > 0 || ParticleDataSet->GetPrevNumInstances() > 0)
+		{
+			ParticleDataSet->ResetBuffers(true);
+		}
+		return true;
+	}
+
+	return false;
 }
 
 /** 
@@ -492,14 +635,29 @@ void FNiagaraEmitterInstance::PreTick()
 	const FNiagaraEmitterHandle& EmitterHandle = GetEmitterHandle();
 	const UNiagaraEmitter* PinnedProps = EmitterHandle.GetInstance();
 
-	if (!PinnedProps || !bIsEnabled || bError || ExecutionState == ENiagaraExecutionState::Dead)
+	if (!PinnedProps || bError || IsComplete())
 	{
 		return;
 	}
 
+	check(ParticleDataSet);
+	FNiagaraDataSet& Data = *ParticleDataSet;
+
+	PinnedProps->SpawnScriptProps.Script->RapidIterationParameters.Tick();
+	PinnedProps->UpdateScriptProps.Script->RapidIterationParameters.Tick();
+	ensure(PinnedProps->GetEventHandlers().Num() == EventExecContexts.Num());
+	for (int32 i = 0; i < PinnedProps->GetEventHandlers().Num(); i++)
+	{
+		PinnedProps->GetEventHandlers()[i].Script->RapidIterationParameters.Tick();
+	}
+
 	bool bOk = true;
 	bOk &= SpawnExecContext.Tick(ParentSystemInstance);
 	bOk &= UpdateExecContext.Tick(ParentSystemInstance);
+	if (PinnedProps->SimTarget == ENiagaraSimTarget::GPUComputeSim)
+	{
+		bOk &= GPUExecContext.Tick(ParentSystemInstance);
+	}
 	for (FNiagaraScriptExecutionContext& EventContext : EventExecContexts)
 	{
 		bOk &= EventContext.Tick(ParentSystemInstance);
@@ -512,22 +670,16 @@ void FNiagaraEmitterInstance::PreTick()
 		return;
 	}
 
-	if (Data.GetNumInstances() == 0)
-	{
-		return;
-	}
-
 	check(Data.GetNumVariables() > 0);
 	check(PinnedProps->SpawnScriptProps.Script);
 	check(PinnedProps->UpdateScriptProps.Script);
 
-	// generate events from collisions
-	if (PinnedProps->CollisionMode == ENiagaraCollisionMode::SceneGeometry)
+	if (bResetPending)
 	{
-		CollisionBatch.GenerateEventsFromResults(this);
+		Data.SetNumInstances(0);
+		bResetPending = false;
 	}
 
-
 	//Swap all data set buffers before doing the main tick on any simulation.
 	if (PinnedProps->SimTarget == ENiagaraSimTarget::CPUSim)
 	{
@@ -535,7 +687,6 @@ void FNiagaraEmitterInstance::PreTick()
 		{
 			SetPair.Value->Tick(PinnedProps->SimTarget);
 		}
-		CollisionBatch.Tick(PinnedProps->SimTarget);
 
 		for (FNiagaraDataSet* Set : UpdateScriptEventDataSets)
 		{
@@ -557,13 +708,36 @@ void FNiagaraEmitterInstance::Tick(float DeltaSeconds)
 
 	const FNiagaraEmitterHandle& EmitterHandle = GetEmitterHandle();
 	const UNiagaraEmitter* PinnedProps = EmitterHandle.GetInstance();
-	if (!PinnedProps || !bIsEnabled || bError || ExecutionState == ENiagaraExecutionState::Dead)
+	if (!PinnedProps || bError)
+	{
+		return;
+	}
+
+	check(ParticleDataSet);
+	FNiagaraDataSet& Data = *ParticleDataSet;
+
+	if (HandleCompletion())
 	{
 		return;
 	}
 
 	Age += DeltaSeconds;
 
+	if (ExecutionState == ENiagaraExecutionState::InactiveClear)
+	{
+		Data.ResetBuffers(true);
+		ExecutionState = ENiagaraExecutionState::Inactive;
+		return;
+	}
+
+	int32 OrigNumParticles = Data.GetPrevNumInstances();
+	if (OrigNumParticles == 0 && ExecutionState != ENiagaraExecutionState::Active)
+	{
+		//Clear out curr buffer in case it had some data in previously.
+		Data.Allocate(0, PinnedProps->SimTarget);
+		return;
+	}
+
 	check(Data.GetNumVariables() > 0);
 	check(PinnedProps->SpawnScriptProps.Script);
 	check(PinnedProps->UpdateScriptProps.Script);
@@ -582,35 +756,46 @@ void FNiagaraEmitterInstance::Tick(float DeltaSeconds)
 		{
 			Binding.SetValue(Age);
 		}
+		
+		SpawnEmitterLocalSpaceBinding.SetValue(PinnedProps->bLocalSpace);
+		UpdateEmitterLocalSpaceBinding.SetValue(PinnedProps->bLocalSpace);
+		for (FNiagaraParameterDirectBinding<FNiagaraBool>& Binding : EventEmitterLocalSpaceBindings)
+		{
+			Binding.SetValue(PinnedProps->bLocalSpace);
+		}
 	}
 	
 	// Calculate number of new particles from regular spawning 
-	int32 OrigNumParticles = Data.GetPrevNumInstances();
-
 	uint32 SpawnTotal = 0;
-	for (FNiagaraSpawnInfo& Info : SpawnInfos)
+	if (ExecutionState == ENiagaraExecutionState::Active)
 	{
-		SpawnTotal += Info.Count;
+		for (FNiagaraSpawnInfo& Info : SpawnInfos)
+		{
+			if (Info.Count > 0)
+			{
+				SpawnTotal += Info.Count;
+			}
+		}
 	}
 
 	// Calculate number of new particles from all event related spawns
 	TArray<TArray<int32>> EventSpawnCounts;
-	EventSpawnCounts.AddDefaulted(PinnedProps->EventHandlerScriptProps.Num());
+	EventSpawnCounts.AddDefaulted(PinnedProps->GetEventHandlers().Num());
 	TArray<int32> EventHandlerSpawnCounts;
-	EventHandlerSpawnCounts.AddDefaulted(PinnedProps->EventHandlerScriptProps.Num());
+	EventHandlerSpawnCounts.AddDefaulted(PinnedProps->GetEventHandlers().Num());
 	uint32 EventSpawnTotal = 0;
 	TArray<FNiagaraDataSet*> EventSet;
-	EventSet.AddZeroed(PinnedProps->EventHandlerScriptProps.Num());
+	EventSet.AddZeroed(PinnedProps->GetEventHandlers().Num());
 	TArray<FGuid> SourceEmitterGuid;
-	SourceEmitterGuid.AddDefaulted(PinnedProps->EventHandlerScriptProps.Num());
+	SourceEmitterGuid.AddDefaulted(PinnedProps->GetEventHandlers().Num());
 	TArray<FName> SourceEmitterName;
-	SourceEmitterName.AddDefaulted(PinnedProps->EventHandlerScriptProps.Num());
+	SourceEmitterName.AddDefaulted(PinnedProps->GetEventHandlers().Num());
 	TArray<bool> bPerformEventSpawning;
-	bPerformEventSpawning.AddDefaulted(PinnedProps->EventHandlerScriptProps.Num());
+	bPerformEventSpawning.AddDefaulted(PinnedProps->GetEventHandlers().Num());
 
-	for (int32 i = 0; i < PinnedProps->EventHandlerScriptProps.Num(); i++)
+	for (int32 i = 0; i < PinnedProps->GetEventHandlers().Num(); i++)
 	{
-		const FNiagaraEventScriptProperties &EventHandlerProps = PinnedProps->EventHandlerScriptProps[i];
+		const FNiagaraEventScriptProperties &EventHandlerProps = PinnedProps->GetEventHandlers()[i];
 		SourceEmitterGuid[i] = EventHandlerProps.SourceEmitterID;
 		SourceEmitterName[i] = SourceEmitterGuid[i].IsValid() ? *SourceEmitterGuid[i].ToString() : EmitterHandle.GetIdName();
 		EventSet[i] = FNiagaraEventDataSetMgr::GetEventDataSet(ParentSystemInstance->GetIDName(), SourceEmitterName[i], EventHandlerProps.SourceEventName);
@@ -628,41 +813,44 @@ void FNiagaraEmitterInstance::Tick(float DeltaSeconds)
 	 */
 	if (PinnedProps->SimTarget == ENiagaraSimTarget::GPUComputeSim)	
 	{
-		FNiagaraComputeExecutionContext *ComputeContext = new FNiagaraComputeExecutionContext();
-		ComputeContext->MainDataSet = &Data;
-		ComputeContext->RTSpawnScript = PinnedProps->SpawnScriptProps.Script->GetRenderThreadScript();
-		ComputeContext->RTUpdateScript = PinnedProps->UpdateScriptProps.Script->GetRenderThreadScript();
-		ComputeContext->SpawnRateInstances = SpawnTotal;
-		ComputeContext->BurstInstances = 0;
-		ComputeContext->EventSpawnTotal = EventSpawnTotal;
+		//FNiagaraComputeExecutionContext *ComputeContext = new FNiagaraComputeExecutionContext();
+		GPUExecContext.MainDataSet = &Data;
+		GPUExecContext.RTSpawnScript = PinnedProps->SpawnScriptProps.Script->GetRenderThreadScript();
+		GPUExecContext.RTUpdateScript = PinnedProps->UpdateScriptProps.Script->GetRenderThreadScript();
+		GPUExecContext.SpawnRateInstances = SpawnTotal;
+		GPUExecContext.BurstInstances = 0;
+		GPUExecContext.EventSpawnTotal = EventSpawnTotal;
 
-		ComputeContext->UpdateInterfaces = PinnedProps->UpdateScriptProps.Script->DataInterfaceInfo;
+		GPUExecContext.UpdateInterfaces = PinnedProps->UpdateScriptProps.Script->DataInterfaceInfo;
 
 		// copy over the constants for the render thread
 		//
-		int32 Size = UpdateExecContext.Parameters.GetExternalParameterSize();
-		if (Size)
+		int32 ParmSize = GPUExecContext.CombinedParamStore.GetPaddedParameterSizeInBytes();
+		if (ParmSize)
 		{
-			ComputeContext->UpdateParams.AddUninitialized(FMath::DivideAndRoundUp(Size, 16) * 16);
-			FMemory::Memcpy(ComputeContext->UpdateParams.GetData(), UpdateExecContext.Parameters.GetParameterData(0), Size);
-		}
-		Size = SpawnExecContext.Parameters.GetExternalParameterSize();
-		if (Size)
-		{
-			ComputeContext->SpawnParams.AddUninitialized(FMath::DivideAndRoundUp(Size, 16) * 16);
-			FMemory::Memcpy(ComputeContext->SpawnParams.GetData(), SpawnExecContext.Parameters.GetParameterData(0), Size);
+#if WITH_EDITORONLY_DATA
+			if (SpawnExecContext.Script->GetDebuggerInfo().bRequestDebugFrame)
+			{
+				SpawnExecContext.Script->GetDebuggerInfo().DebugFrame.Reset();
+				SpawnExecContext.Script->GetDebuggerInfo().bRequestDebugFrame = false;
+				SpawnExecContext.Script->GetDebuggerInfo().DebugFrameLastWriteId = SpawnExecContext.TickCounter;
+				SpawnExecContext.Script->GetDebuggerInfo().DebugParameters = GPUExecContext.CombinedParamStore;
+			}
+#endif
 		}
+		GPUExecContext.SpawnParams.SetNumZeroed(ParmSize);
+		GPUExecContext.CombinedParamStore.CopyParameterDataToPaddedBuffer(GPUExecContext.SpawnParams.GetData(), ParmSize);
 
 		// push event data sets to the context
 		for (FNiagaraDataSet *Set : UpdateScriptEventDataSets)
 		{
-			ComputeContext->UpdateEventWriteDataSets.Add(Set);
+			GPUExecContext.UpdateEventWriteDataSets.Add(Set);
 		}
 
-		ComputeContext->EventHandlerScriptProps = PinnedProps->EventHandlerScriptProps;
-		ComputeContext->EventSets = EventSet;
-		ComputeContext->EventSpawnCounts = EventHandlerSpawnCounts;
-		NiagaraEmitterInstanceBatcher::Get()->Queue(ComputeContext);
+		GPUExecContext.EventHandlerScriptProps = PinnedProps->GetEventHandlers();
+		GPUExecContext.EventSets = EventSet;
+		GPUExecContext.EventSpawnCounts = EventHandlerSpawnCounts;
+		NiagaraEmitterInstanceBatcher::Get()->Queue(&GPUExecContext);
 
 		CachedBounds.Init();
 		CachedBounds = CachedBounds.ExpandBy(FVector(20.0, 20.0, 20.0));	// temp until GPU sims update bounds
@@ -674,12 +862,11 @@ void FNiagaraEmitterInstance::Tick(float DeltaSeconds)
 	Data.Allocate(AllocationSize, PinnedProps->SimTarget);
 
 
-
 	TArray<FNiagaraDataSetExecutionInfo> DataSetExecInfos;
 	DataSetExecInfos.Emplace(&Data, 0, false, true);
 
 	// Simulate existing particles forward by DeltaSeconds.
-	if ((ExecutionState != ENiagaraExecutionState::Dead && ExecutionState != ENiagaraExecutionState::Paused) /*&& OrigNumParticles > 0*/)
+	if (OrigNumParticles > 0)
 	{
 		/*
 		if (bDumpAfterEvent)
@@ -700,10 +887,16 @@ void FNiagaraEmitterInstance::Tick(float DeltaSeconds)
 			DataSetExecInfos.Emplace(EventDataSet, 0, true, false);
 		}
 		UpdateExecContext.Execute(OrigNumParticles, DataSetExecInfos);
+		int32 DeltaParticles = Data.GetNumInstances() - OrigNumParticles;
+
+		ensure(DeltaParticles <= 0); // We either lose particles or stay the same, we should never add particles in update!
 
 		if (GbDumpParticleData)
 		{
-			UE_LOG(LogNiagara, Log, TEXT("=== Updated %d Particles ==="), OrigNumParticles);
+			UE_LOG(LogNiagara, Log, TEXT("=== Update Parameters ===") );
+			UpdateExecContext.Parameters.Dump();
+
+			UE_LOG(LogNiagara, Log, TEXT("=== Updated %d Particles (%d Died) ==="), OrigNumParticles, -DeltaParticles);
 			Data.Dump(true, 0, OrigNumParticles);
 		}
 	}
@@ -711,7 +904,7 @@ void FNiagaraEmitterInstance::Tick(float DeltaSeconds)
 	uint32 EventSpawnStart = Data.GetNumInstances();
 
 	//Init new particles with the spawn script.
-	if (ExecutionState==ENiagaraExecutionState::Active && SpawnTotal + EventSpawnTotal > 0)
+	if (SpawnTotal + EventSpawnTotal > 0)
 	{
 		SCOPE_CYCLE_COUNTER(STAT_NiagaraSpawn);
 
@@ -726,15 +919,20 @@ void FNiagaraEmitterInstance::Tick(float DeltaSeconds)
 				SpawnExecCountBinding.SetValue(Num);
 				DataSetExecInfos.SetNum(1, false);
 				DataSetExecInfos[0].StartInstance = OrigNum;
+				//UE_LOG(LogNiagara, Log, TEXT("SpawnScriptEventDataSets: %d"), SpawnScriptEventDataSets.Num());
 				for (FNiagaraDataSet* EventDataSet : SpawnScriptEventDataSets)
 				{
+					//UE_LOG(LogNiagara, Log, TEXT("SpawnScriptEventDataSets.. %d"), EventDataSet->GetNumVariables());
 					DataSetExecInfos.Emplace(EventDataSet, OrigNum, true, false);
 				}
+				
 				SpawnExecContext.Execute(Num, DataSetExecInfos);
 
 				if (GbDumpParticleData)
 				{
-					UE_LOG(LogNiagara, Log, TEXT("=== Spawned %d Particles ==="), Num);
+					UE_LOG(LogNiagara, Log, TEXT("=== %s Spawn Parameters ==="), *DumpLabel);
+					SpawnExecContext.Parameters.Dump();
+					UE_LOG(LogNiagara, Log, TEXT("===  %s Spawned %d Particles==="), *DumpLabel, Num);
 					Data.Dump(true, OrigNum, Num);
 				}
 			}
@@ -751,7 +949,7 @@ void FNiagaraEmitterInstance::Tick(float DeltaSeconds)
 
 		EventSpawnStart = Data.GetNumInstances();
 
-		for (int32 EventScriptIdx = 0; EventScriptIdx < PinnedProps->EventHandlerScriptProps.Num(); EventScriptIdx++)
+		for (int32 EventScriptIdx = 0; EventScriptIdx < PinnedProps->GetEventHandlers().Num(); EventScriptIdx++)
 		{
 			//Spawn particles coming from events.
 			for (int32 i = 0; i < EventSpawnCounts[EventScriptIdx].Num(); i++)
@@ -766,25 +964,45 @@ void FNiagaraEmitterInstance::Tick(float DeltaSeconds)
 			}
 		}
 	}
+	/*else if (SpawnTotal + EventSpawnTotal > 0)
+	{
+		UE_LOG(LogNiagara, Log, TEXT("Skipping spawning due to execution state! %d"), (uint32)ExecutionState)
+	}*/
 
 	// handle event based spawning
-	for (int32 EventScriptIdx = 0;  EventScriptIdx < PinnedProps->EventHandlerScriptProps.Num(); EventScriptIdx++)
+	for (int32 EventScriptIdx = 0;  EventScriptIdx < PinnedProps->GetEventHandlers().Num(); EventScriptIdx++)
 	{
-		const FNiagaraEventScriptProperties &EventHandlerProps = PinnedProps->EventHandlerScriptProps[EventScriptIdx];
+		const FNiagaraEventScriptProperties &EventHandlerProps = PinnedProps->GetEventHandlers()[EventScriptIdx];
 
 		if (bPerformEventSpawning[EventScriptIdx] && EventSet[EventScriptIdx] && EventSpawnCounts[EventScriptIdx].Num())
 		{
+			uint32 NumParticles = Data.GetNumInstances();
 			SCOPE_CYCLE_COUNTER(STAT_NiagaraEventHandle);
+
+			Data.Tick();
+			Data.CopyPrevToCur();
+
 			for (int32 i = 0; i < EventSpawnCounts[EventScriptIdx].Num(); i++)
 			{
 				int32 EventNumToSpawn = EventSpawnCounts[EventScriptIdx][i];
-				ensure(EventNumToSpawn + EventSpawnStart < Data.GetNumInstances());
+				ensure(EventNumToSpawn + EventSpawnStart <= Data.GetNumInstances());
 				EventExecCountBindings[EventScriptIdx].SetValue(EventNumToSpawn);
 
 				DataSetExecInfos.SetNum(1, false);
 				DataSetExecInfos[0].StartInstance = EventSpawnStart;
+				DataSetExecInfos[0].bUpdateInstanceCount = false;
 				DataSetExecInfos.Emplace(EventSet[EventScriptIdx], i, false, false);
-				SpawnExecContext.Execute(EventNumToSpawn, DataSetExecInfos);
+				EventExecContexts[EventScriptIdx].Execute(EventNumToSpawn, DataSetExecInfos);
+				
+				if (GbDumpParticleData)
+				{
+					UE_LOG(LogNiagara, Log, TEXT("=== Event %d Parameters ==="), EventScriptIdx);
+					EventExecContexts[EventScriptIdx].Parameters.Dump();
+					UE_LOG(LogNiagara, Log, TEXT("=== Event %d %d Particles ==="), EventScriptIdx, EventNumToSpawn);
+					Data.Dump(true, EventSpawnStart, EventNumToSpawn);
+				}
+
+				ensure(Data.GetNumInstances() == NumParticles);
 
 				EventSpawnStart += EventNumToSpawn;
 			}
@@ -794,10 +1012,11 @@ void FNiagaraEmitterInstance::Tick(float DeltaSeconds)
 		// handle all-particle events
 		if (EventHandlerProps.Script && EventHandlerProps.ExecutionMode == EScriptExecutionMode::EveryParticle && EventSet[EventScriptIdx])
 		{
+			uint32 NumParticles = Data.GetNumInstances();
 			if (EventSet[EventScriptIdx]->GetPrevNumInstances())
 			{
 				SCOPE_CYCLE_COUNTER(STAT_NiagaraEventHandle);
-
+				
 				for (uint32 i = 0; i < EventSet[EventScriptIdx]->GetPrevNumInstances(); i++)
 				{
 					// If we have events, Swap buffers, to make sure we don't overwrite previous script results
@@ -806,13 +1025,28 @@ void FNiagaraEmitterInstance::Tick(float DeltaSeconds)
 					Data.Tick();
 					Data.CopyPrevToCur();
 
-					EventExecCountBindings[EventScriptIdx].SetValue(Data.GetPrevNumInstances());
+					uint32 NumInstancesPrev = Data.GetPrevNumInstances();
+					EventExecCountBindings[EventScriptIdx].SetValue(NumInstancesPrev);
 					DataSetExecInfos.SetNum(1, false);
 					DataSetExecInfos[0].StartInstance = 0;
 					DataSetExecInfos.Emplace(EventSet[EventScriptIdx], i, false, false);
-					SpawnExecContext.Execute(Data.GetPrevNumInstances(), DataSetExecInfos);
+					/*if (GbDumpParticleData)
+					{
+						UE_LOG(LogNiagara, Log, TEXT("=== Event %d [%d] Payload ==="), EventScriptIdx, i);
+						DataSetExecInfos[1].DataSet->Dump(true, 0, 1);
+					}*/
+
+					EventExecContexts[EventScriptIdx].Execute(NumInstancesPrev, DataSetExecInfos);
+
+					if (GbDumpParticleData)
+					{
+						UE_LOG(LogNiagara, Log, TEXT("=== Event %d [%d] Parameters ==="), EventScriptIdx, i);
+						EventExecContexts[EventScriptIdx].Parameters.Dump();
+						UE_LOG(LogNiagara, Log, TEXT("=== Event %d %d Particles ==="), EventScriptIdx, NumInstancesPrev);
+						Data.Dump(true, 0, NumInstancesPrev);
+					}
 
-					// TODO: Should we check to see if we wrote all particles here?
+					ensure(NumParticles == Data.GetNumInstances());
 				}
 			}
 		}
@@ -823,35 +1057,44 @@ void FNiagaraEmitterInstance::Tick(float DeltaSeconds)
 		// event/particle arrays; this is currently a very naive (and comparatively slow) implementation, until full indexed reads work
 		if (EventHandlerProps.Script && EventHandlerProps.ExecutionMode == EScriptExecutionMode::SingleParticle && EventSet[EventScriptIdx])
 		{
+
 			SCOPE_CYCLE_COUNTER(STAT_NiagaraEventHandle);
 			FNiagaraVariable IndexVar(FNiagaraTypeDefinition::GetIntDef(), "ParticleIndex");
 			FNiagaraDataSetIterator<int32> IndexItr(*EventSet[EventScriptIdx], IndexVar, 0, false);
-			if (IndexItr.IsValid())
+			if (IndexItr.IsValid() && EventSet[EventScriptIdx]->GetPrevNumInstances() > 0)
 			{
 				EventExecCountBindings[EventScriptIdx].SetValue(1);
 
+				Data.Tick();
+				Data.CopyPrevToCur();
+				uint32 NumParticles = Data.GetNumInstances();
+
 				for (uint32 i = 0; i < EventSet[EventScriptIdx]->GetPrevNumInstances(); i++)
 				{
 					int32 Index = *IndexItr;
 					IndexItr.Advance();
 					DataSetExecInfos.SetNum(1, false);
 					DataSetExecInfos[0].StartInstance = Index;
+					DataSetExecInfos[0].bUpdateInstanceCount = false;
 					DataSetExecInfos.Emplace(EventSet[EventScriptIdx], i, false, false);
-					SpawnExecContext.Execute(1, DataSetExecInfos);
+					EventExecContexts[EventScriptIdx].Execute(1, DataSetExecInfos);
 
-					// TODO: Should we check to see if we wrote all particles here?
+					if (GbDumpParticleData)
+					{
+						ensure(EventHandlerProps.Script->RapidIterationParameters.VerifyBinding(&EventExecContexts[EventScriptIdx].Parameters));
+						UE_LOG(LogNiagara, Log, TEXT("=== Event %d Src Parameters ==="), EventScriptIdx);
+						EventHandlerProps.Script->RapidIterationParameters.Dump();
+						UE_LOG(LogNiagara, Log, TEXT("=== Event %d Context Parameters ==="), EventScriptIdx);
+						EventExecContexts[EventScriptIdx].Parameters.Dump();
+						UE_LOG(LogNiagara, Log, TEXT("=== Event %d Particles (%d index written, %d total) ==="), EventScriptIdx, Index, Data.GetNumInstances());
+						Data.Dump(true, Index, 1);
+					}
+					ensure(NumParticles == Data.GetNumInstances());
 				}
 			}
 		}
 	}
 
-
-	// kick off collision tests from this emitter
-	if (PinnedProps->CollisionMode == ENiagaraCollisionMode::SceneGeometry)
-	{
-		CollisionBatch.KickoffNewBatch(this, DeltaSeconds);
-	}
-
 	PostProcessParticles();
 
 	SpawnExecContext.PostTick();
@@ -886,7 +1129,7 @@ uint32 FNiagaraEmitterInstance::CalculateEventSpawnCount(const FNiagaraEventScri
 
 		for (int32 i = 0; i < NumEventsToProcess; i++)
 		{
-			if (ExecutionState == ENiagaraExecutionState::Active)
+			if (ExecutionState == ENiagaraExecutionState::Active && EventHandlerProps.SpawnNumber > 0)
 			{
 				EventSpawnCounts.Add(EventHandlerProps.SpawnNumber);
 				EventSpawnTotal += EventHandlerProps.SpawnNumber;
@@ -899,18 +1142,26 @@ uint32 FNiagaraEmitterInstance::CalculateEventSpawnCount(const FNiagaraEventScri
 
 void FNiagaraEmitterInstance::SetExecutionState(ENiagaraExecutionState InState)
 {
-	if (InState != ExecutionState)
+	/*if (InState != ExecutionState)
 	{
 		const UEnum* EnumPtr = FNiagaraTypeDefinition::GetExecutionStateEnum();
 		UE_LOG(LogNiagara, Log, TEXT("Emitter \"%s\" change state: %s to %s"), *GetEmitterHandle().GetName().ToString(), *EnumPtr->GetNameStringByValue((int64)ExecutionState),
 			*EnumPtr->GetNameStringByValue((int64)InState));
-	}
+	}*/
 
-	if (InState == ENiagaraExecutionState::Active && ExecutionState == ENiagaraExecutionState::Inactive)
+	/*if (InState == ENiagaraExecutionState::Active && ExecutionState == ENiagaraExecutionState::Inactive)
 	{
 		UE_LOG(LogNiagara, Log, TEXT("Emitter \"%s\" change state N O O O O O "), *GetEmitterHandle().GetName().ToString());
+	}*/
+
+	if (InState != ENiagaraExecutionState::Disabled && GetEmitterHandle().GetInstance()->IsAllowedByDetailLevel() && GetEmitterHandle().GetIsEnabled())
+	{
+		ExecutionState = InState;
+	}
+	else
+	{
+		ExecutionState = ENiagaraExecutionState::Disabled;
 	}
-	ExecutionState = InState;
 }
 
 
@@ -924,6 +1175,8 @@ bool FNiagaraEmitterInstance::CheckAttributesForRenderer(int32 Index)
 		return false;
 	}
 
+	check(ParticleDataSet);
+	FNiagaraDataSet& Data = *ParticleDataSet;
 	bool bOk = true;
 	if (EmitterRenderer[Index])
 	{
@@ -972,18 +1225,18 @@ void FNiagaraEmitterInstance::UpdateEmitterRenderer(ERHIFeatureLevel::Type Featu
 			}
 		}
 
-		if (bIsEnabled && !bError)
+		if (!IsComplete() && !bError)
 		{
 			EmitterRenderer.Empty();
-			EmitterRenderer.AddZeroed(EmitterProperties->RendererProperties.Num());
-			for (int32 SubIdx = 0; SubIdx < EmitterProperties->RendererProperties.Num(); SubIdx++)
+			EmitterRenderer.AddZeroed(EmitterProperties->GetRenderers().Num());
+			for (int32 SubIdx = 0; SubIdx < EmitterProperties->GetRenderers().Num(); SubIdx++)
 			{
 				UMaterialInterface *Material = nullptr;
 
 				TArray<UMaterialInterface*> UsedMats;
-				if (EmitterProperties->RendererProperties[SubIdx] != nullptr)
+				if (EmitterProperties->GetRenderers()[SubIdx] != nullptr)
 				{
-					EmitterProperties->RendererProperties[SubIdx]->GetUsedMaterials(UsedMats);
+					EmitterProperties->GetRenderers()[SubIdx]->GetUsedMaterials(UsedMats);
 					if (UsedMats.Num() != 0)
 					{
 						Material = UsedMats[0];
@@ -995,9 +1248,9 @@ void FNiagaraEmitterInstance::UpdateEmitterRenderer(ERHIFeatureLevel::Type Featu
 					Material = UMaterial::GetDefaultMaterial(MD_Surface);
 				}
 
-				if (EmitterProperties->RendererProperties[SubIdx] != nullptr)
+				if (EmitterProperties->GetRenderers()[SubIdx] != nullptr)
 				{
-					EmitterRenderer[SubIdx] = EmitterProperties->RendererProperties[SubIdx]->CreateEmitterRenderer(FeatureLevel);
+					EmitterRenderer[SubIdx] = EmitterProperties->GetRenderers()[SubIdx]->CreateEmitterRenderer(FeatureLevel);
 					EmitterRenderer[SubIdx]->SetMaterial(Material, FeatureLevel);
 					EmitterRenderer[SubIdx]->SetLocalSpace(EmitterProperties->bLocalSpace);
 					ToBeAddedList.Add(EmitterRenderer[SubIdx]);
diff --git a/Engine/Plugins/FX/Niagara/Source/Niagara/Private/NiagaraEmitterInstanceBatcher.cpp b/Engine/Plugins/FX/Niagara/Source/Niagara/Private/NiagaraEmitterInstanceBatcher.cpp
--- a/Engine/Plugins/FX/Niagara/Source/Niagara/Private/NiagaraEmitterInstanceBatcher.cpp
+++ b/Engine/Plugins/FX/Niagara/Source/Niagara/Private/NiagaraEmitterInstanceBatcher.cpp
@@ -15,7 +15,7 @@ DECLARE_FLOAT_COUNTER_STAT(TEXT("Niagara GPU Sim"), Stat_GPU_NiagaraSim, STATGRO
 
 NiagaraEmitterInstanceBatcher* NiagaraEmitterInstanceBatcher::BatcherSingleton = nullptr;
 uint32 FNiagaraComputeExecutionContext::TickCounter = 0;
-
+FRWBuffer NiagaraEmitterInstanceBatcher::DummyWriteIndexBuffer;
 
 void NiagaraEmitterInstanceBatcher::Queue(FNiagaraComputeExecutionContext *InContext)
 {
@@ -49,9 +49,10 @@ void NiagaraEmitterInstanceBatcher::TickSingle(const FNiagaraComputeExecutionCon
 	check(IsInRenderingThread());
 	Context->MainDataSet->TickRenderThread(ENiagaraSimTarget::GPUComputeSim);
 
-	FNiagaraShader *UpdateShader = Context->RTUpdateScript->GetShader();
+	FNiagaraComputeExecutionContext::TickCounter++;
+
 	FNiagaraShader *SpawnShader = Context->RTSpawnScript->GetShader();
-	if (!UpdateShader || !SpawnShader)
+	if (!SpawnShader)
 	{
 		return;
 	}
@@ -59,9 +60,10 @@ void NiagaraEmitterInstanceBatcher::TickSingle(const FNiagaraComputeExecutionCon
 	uint32 PrevNumInstances = Context->MainDataSet->PrevDataRender().GetNumInstances();
 
 	uint32 EventSpawnTotal = 0;
-	for (int32 i = 0; i < Context->EventHandlerScriptProps.Num(); i++)
+	/*
+	for (int32 i = 0; i < Context->GetEventHandlers().Num(); i++)
 	{
-		const FNiagaraEventScriptProperties &EventHandlerProps = Context->EventHandlerScriptProps[i];
+		const FNiagaraEventScriptProperties &EventHandlerProps = Context->GetEventHandlers()[i];
 		if (EventHandlerProps.ExecutionMode == EScriptExecutionMode::SpawnedParticles && Context->EventSets[i])
 		{
 			uint32 NumEventsToProcess = Context->EventSets[i]->PrevDataRender().GetNumInstances();
@@ -69,6 +71,7 @@ void NiagaraEmitterInstanceBatcher::TickSingle(const FNiagaraComputeExecutionCon
 			EventSpawnTotal += EventSpawnNum;
 		}
 	}
+	*/
 
 
 	// allocate for additional instances spawned and set the new number in the data set, if the new number is greater (meaning if we're spawning in this run)
@@ -91,53 +94,47 @@ void NiagaraEmitterInstanceBatcher::TickSingle(const FNiagaraComputeExecutionCon
 
 	//UE_LOG(LogNiagara, Log, TEXT("Pre sim instances %i"), PrevNumInstances);
 
-	// simulation run
-	RHICmdList.SetComputeShader(UpdateShader->GetComputeShader());
-	SetupDataInterfaceBuffers(Context->UpdateInterfaces, UpdateShader, RHICmdList);
-	SetupEventUAVs(Context, PrevNumInstances, RHICmdList);
-	Run(Context->MainDataSet, 0, PrevNumInstances, UpdateShader, Context->UpdateParams, RHICmdList);
-	UnsetEventUAVs(Context, RHICmdList);
+	// Calc shader parameters
+	uint32 TotalNumInstances = NewNumInstances;
+	uint32 SimulationStartInstance = 0;
+	uint32 SpawnStartInstance = PrevNumInstances;
+
+	// Recreate a clear data set index buffer for the simulation shader to write number of written instances to
+	//
+	{
+		SCOPED_DRAW_EVENT(RHICmdList, NiagaraIndexBufferClear);
+		ClearUAV(RHICmdList, Context->MainDataSet->GetDataSetIndices(), 0);
+	}
+	RHICmdList.TransitionResource(EResourceTransitionAccess::EWritable, EResourceTransitionPipeline::EComputeToCompute, Context->MainDataSet->GetDataSetIndices().UAV);
+
+
+	// run shader, sim and spawn in a single dispatch
+	RHICmdList.SetComputeShader(SpawnShader->GetComputeShader());
+	SetupDataInterfaceBuffers(Context->UpdateInterfaces, SpawnShader, RHICmdList);
+	Run(Context->MainDataSet, SimulationStartInstance, SpawnStartInstance, TotalNumInstances, SpawnShader, Context->SpawnParams, RHICmdList, Context->MainDataSet->GetDataSetIndices());
 
 	// resolve data set writes
 	// grabs the number of instances written from the index set during the simulation run
 	//
 	uint32 NumInstancesAfterSim[64];
 	NumInstancesAfterSim[0] = PrevNumInstances;
-	ResolveDatasetWrites(NumInstancesAfterSim, Context);		// this is going to cause a flush; need to place it differently
+	//we may be able to remove this transition if each step isn't dependent on the previous one.
+	RHICmdList.TransitionResource(EResourceTransitionAccess::EReadable, EResourceTransitionPipeline::EComputeToCompute, Context->MainDataSet->GetDataSetIndices().UAV);
+	ResolveDatasetWrites(NumInstancesAfterSim, Context, Context->MainDataSet->GetDataSetIndices());		// this is going to cause a flush; need to place it differently
 	Context->MainDataSet->CurrDataRender().SetNumInstances(NumInstancesAfterSim[0]);
+	RHICmdList.TransitionResource(EResourceTransitionAccess::EReadable, EResourceTransitionPipeline::EComputeToGfx, Context->MainDataSet->GetDataSetIndices().UAV);
 
 	// TODO: hack - only updating event set 0 on update scripts now; need to match them to their indices and update them all
+	/*
 	if (Context->UpdateEventWriteDataSets.Num())
-	{
+	{ 
 		Context->UpdateEventWriteDataSets[0]->CurrDataRender().SetNumInstances(NumInstancesAfterSim[1]);
 	}
+	*/
 
-	UE_LOG(LogNiagara, Log, TEXT("After sim instances %i"), NumInstancesAfterSim[0]);
-
-	// spawn run
-	// bursts and regular spawn rate happen at once here
-	RHICmdList.SetComputeShader(SpawnShader->GetComputeShader());
-	uint32 SpawnInstances = Context->SpawnRateInstances + Context->BurstInstances + EventSpawnTotal;
-	uint32 NumInstancesAfterSpawn = NumInstancesAfterSim[0] + SpawnInstances;
-	uint32 NumInstancesAfterNonEventSpawn = NumInstancesAfterSim[0] + Context->SpawnRateInstances + Context->BurstInstances;
-	if (SpawnInstances)
-	{
-		Run(Context->MainDataSet, NumInstancesAfterSim[0], Context->SpawnRateInstances + Context->BurstInstances + EventSpawnTotal, SpawnShader, Context->SpawnParams, RHICmdList);
-		Context->MainDataSet->CurrDataRender().SetNumInstances(NumInstancesAfterSpawn);
-
-		//uncomment this to check number of spawn instances versus ones actually processed by the shader
-		uint32 NumInstancesSpawned[64];
-		ResolveDatasetWrites(NumInstancesSpawned, Context);
-		ensure(NumInstancesSpawned[0] == SpawnInstances);
-		UE_LOG(LogNiagara, Log, TEXT("Spawned %i to %i, ran spawn script on %i"), Context->SpawnRateInstances, NumInstancesAfterSpawn, NumInstancesSpawned[0]);
-
-		ensure(NumInstancesAfterNonEventSpawn + EventSpawnTotal == NumInstancesAfterSpawn);
-		ensure(NumInstancesAfterNonEventSpawn + EventSpawnTotal == Context->MainDataSet->CurrDataRender().GetNumInstances());
-	}
-
+	//UE_LOG(LogNiagara, Log, TEXT("After sim instances %i"), NumInstancesAfterSim[0]);
 
-
-	RunEventHandlers(Context, NumInstancesAfterSim[0], NumInstancesAfterSpawn, NumInstancesAfterNonEventSpawn, RHICmdList);
+	//RunEventHandlers(Context, NumInstancesAfterSim[0], NumInstancesAfterSpawn, NumInstancesAfterNonEventSpawn, RHICmdList);
 
 	// the VF grabs PrevDataRender for drawing, so need to transition
 	RHICmdList.TransitionResource(EResourceTransitionAccess::EReadable, EResourceTransitionPipeline::EComputeToGfx, Context->MainDataSet->PrevDataRender().GetGPUBufferFloat()->UAV);
@@ -145,16 +142,16 @@ void NiagaraEmitterInstanceBatcher::TickSingle(const FNiagaraComputeExecutionCon
 }
 
 
-void NiagaraEmitterInstanceBatcher::ResolveDatasetWrites(uint32 *OutArray, const FNiagaraComputeExecutionContext *Context) const
+void NiagaraEmitterInstanceBatcher::ResolveDatasetWrites(uint32 *OutArray, const FNiagaraComputeExecutionContext *Context, FRWBuffer &WriteIndexBuffer) const
 {
-	if (Context->MainDataSet->GetDataSetIndices().NumBytes > 0)
+	if (WriteIndexBuffer.NumBytes > 0)
 	{
-		int32 *DataSetInstanceIndexBuffer = static_cast<int32*>(RHILockVertexBuffer(Context->MainDataSet->GetDataSetIndices().Buffer, 0, 64 * sizeof(int32), EResourceLockMode::RLM_ReadOnly));
+		int32 *DataSetInstanceIndexBuffer = static_cast<int32*>(RHILockVertexBuffer(WriteIndexBuffer.Buffer, 0, 64 * sizeof(int32), EResourceLockMode::RLM_ReadOnly));
 		for (uint32 SetIdx = 0; SetIdx < 64; SetIdx++)
 		{
-			OutArray[SetIdx] = DataSetInstanceIndexBuffer[SetIdx];
+			OutArray[SetIdx] = DataSetInstanceIndexBuffer[SetIdx*4+1];
 		}
-		RHIUnlockVertexBuffer(Context->MainDataSet->GetDataSetIndices().Buffer);
+		RHIUnlockVertexBuffer(WriteIndexBuffer.Buffer);
 	}
 }
 
@@ -178,24 +175,24 @@ void NiagaraEmitterInstanceBatcher::SetupDataInterfaceBuffers(const TArray<FNiag
 	}
 }
 
-void NiagaraEmitterInstanceBatcher::Run(FNiagaraDataSet *DataSet, uint32 StartInstance, const uint32 NumInstances, FNiagaraShader *Shader, const TArray<uint8, TAlignedHeapAllocator<16>>  &Params, FRHICommandList &RHICmdList, bool bCopyBeforeStart) const
+void NiagaraEmitterInstanceBatcher::Run(FNiagaraDataSet *DataSet, uint32 Phase0StartInstance, uint32 Phase1StartInstance, const uint32 TotalNumInstances, FNiagaraShader *Shader, const TArray<uint8, TAlignedHeapAllocator<16>>  &Params, FRHICommandList &RHICmdList, const FRWBuffer &WriteIndexBuffer, bool bCopyBeforeStart) const
 {
-	// Recreate a clear data set index buffer for the simulation shader to write number of written instances to
-	//
-	FRWBuffer &InstanceIdxBuf = DataSet->SetupDataSetIndices();
-	ClearUAV(RHICmdList, InstanceIdxBuf, 0);
-	RHICmdList.TransitionResource(EResourceTransitionAccess::EWritable, EResourceTransitionPipeline::EComputeToCompute, InstanceIdxBuf.UAV);
+	if (TotalNumInstances == 0)
+	{
+		return;
+	}
+
+	RHICmdList.SetComputeShader(Shader->GetComputeShader());
 
 	// set the shader and data set params 
 	//
-	//RHICmdList.SetComputeShader(Shader->GetComputeShader());
 	DataSet->SetShaderParams(Shader, RHICmdList);
 
 	// set the index buffer uav
 	//
 	if (Shader->OutputIndexBufferParam.IsBound())
 	{
-		RHICmdList.SetUAVParameter(Shader->GetComputeShader(), Shader->OutputIndexBufferParam.GetUAVIndex(), InstanceIdxBuf.UAV);
+		RHICmdList.SetUAVParameter(Shader->GetComputeShader(), Shader->OutputIndexBufferParam.GetUAVIndex(), WriteIndexBuffer.UAV);
 	}
 
 	// set the execution parameters
@@ -204,45 +201,19 @@ void NiagaraEmitterInstanceBatcher::Run(FNiagaraDataSet *DataSet, uint32 StartIn
 	{
 		RHICmdList.SetShaderParameter(Shader->GetComputeShader(), 0, Shader->EmitterTickCounterParam.GetBaseIndex(), Shader->EmitterTickCounterParam.GetNumBytes(), &FNiagaraComputeExecutionContext::TickCounter);
 	}
-	uint32 Copy = bCopyBeforeStart ? 1 : 0;
-	RHICmdList.SetShaderParameter(Shader->GetComputeShader(), 0, Shader->CopyInstancesBeforeStartParam.GetBaseIndex(), Shader->CopyInstancesBeforeStartParam.GetNumBytes(), &Copy);
-
-	// Now figure out how many instances each thread needs to simulate and how many thread groups we're running. We do this based on the total number of 
-	// instances we need to run and how many threads per thread group there are as well as the max. number of thread groups (defined in NiagaraCommon.h)
-	//
-	uint32 NumInstancesPerThread = 0;
-	uint32 AdjustedNumInstances = NumInstances;
-	uint32 SimulateStartInstance = StartInstance;
 
-	// if the caller requests to copy all instances before StartInstance, set up the params accordingly
-	// spawn events need to do this, for example
-	if (bCopyBeforeStart)
-	{
-		AdjustedNumInstances += StartInstance;
-		StartInstance = 0;
-	}
+	uint32 Copy = bCopyBeforeStart ? 1 : 0;
 
-	RHICmdList.SetShaderParameter(Shader->GetComputeShader(), 0, Shader->SimulateStartInstanceParam.GetBaseIndex(), Shader->SimulateStartInstanceParam.GetNumBytes(), &SimulateStartInstance);
-	RHICmdList.SetShaderParameter(Shader->GetComputeShader(), 0, Shader->StartInstanceParam.GetBaseIndex(), Shader->StartInstanceParam.GetNumBytes(), &StartInstance);
+	RHICmdList.SetShaderParameter(Shader->GetComputeShader(), 0, Shader->StartInstancePhase0Param.GetBaseIndex(), Shader->StartInstancePhase0Param.GetNumBytes(), &Phase0StartInstance);
+	RHICmdList.SetShaderParameter(Shader->GetComputeShader(), 0, Shader->StartInstancePhase1Param.GetBaseIndex(), Shader->StartInstancePhase1Param.GetNumBytes(), &Phase1StartInstance);
+	RHICmdList.SetShaderParameter(Shader->GetComputeShader(), 0, Shader->TotalNumInstancesParam.GetBaseIndex(), Shader->TotalNumInstancesParam.GetNumBytes(), &TotalNumInstances);
 
-	NumInstancesPerThread = FMath::DivideAndRoundUp(AdjustedNumInstances, NIAGARA_COMPUTE_THREADGROUP_SIZE);
 	uint32 NumThreadGroups = 1;
-	if (AdjustedNumInstances > NIAGARA_COMPUTE_THREADGROUP_SIZE)
+	if (TotalNumInstances > NIAGARA_COMPUTE_THREADGROUP_SIZE)
 	{
-		NumThreadGroups = FMath::Min(NIAGARA_MAX_COMPUTE_THREADGROUPS, FMath::DivideAndRoundUp(AdjustedNumInstances, NIAGARA_COMPUTE_THREADGROUP_SIZE));
-		NumInstancesPerThread = FMath::DivideAndRoundUp(NumInstancesPerThread, NumThreadGroups);
+		NumThreadGroups = FMath::Min(NIAGARA_MAX_COMPUTE_THREADGROUPS, FMath::DivideAndRoundUp(TotalNumInstances, NIAGARA_COMPUTE_THREADGROUP_SIZE));
 	}
 
-	RHICmdList.SetShaderParameter(Shader->GetComputeShader(), 0, Shader->NumInstancesPerThreadParam.GetBaseIndex(), Shader->NumInstancesPerThreadParam.GetNumBytes(), &NumInstancesPerThread);
-	RHICmdList.SetShaderParameter(Shader->GetComputeShader(), 0, Shader->NumInstancesParam.GetBaseIndex(), Shader->NumInstancesParam.GetNumBytes(), &AdjustedNumInstances);
-	RHICmdList.SetShaderParameter(Shader->GetComputeShader(), 0, Shader->NumThreadGroupsParam.GetBaseIndex(), Shader->NumThreadGroupsParam.GetNumBytes(), &NumThreadGroups);
-
-	if (StartInstance > 0)
-	{
-		//UE_LOG(LogNiagara, Log, TEXT("Startinstance %i (%i)    NumInstances %i (%i)   %i/Thread, %i groups"), StartInstance, SimulateStartInstance, NumInstances, AdjustedNumInstances, NumInstancesPerThread, NumThreadGroups);
-	}
-
-
 	// setup script parameters
 	FRHIUniformBufferLayout CBufferLayout(TEXT("Niagara Compute Sim CBuffer"));
 	CBufferLayout.ConstantBufferSize = Params.Num();
@@ -257,7 +228,7 @@ void NiagaraEmitterInstanceBatcher::Run(FNiagaraDataSet *DataSet, uint32 StartIn
 
 	// Dispatch, if anything needs to be done
 	//
-	if (NumInstancesPerThread)
+	if (TotalNumInstances)
 	{
 		SCOPED_DRAW_EVENT(RHICmdList, NiagaraGPUSimulationCS);
 		SCOPED_GPU_STAT(RHICmdList, Stat_GPU_NiagaraSim);
@@ -275,9 +246,10 @@ void NiagaraEmitterInstanceBatcher::RunEventHandlers(const FNiagaraComputeExecut
 {
 	// Event handler run
 	//
-	for (int32 EventScriptIdx = 0; EventScriptIdx < Context->EventHandlerScriptProps.Num(); EventScriptIdx++)
+	/*
+	for (int32 EventScriptIdx = 0; EventScriptIdx < Context->GetEventHandlers().Num(); EventScriptIdx++)
 	{
-		const FNiagaraEventScriptProperties &EventHandlerProps = Context->EventHandlerScriptProps[EventScriptIdx];
+		const FNiagaraEventScriptProperties &EventHandlerProps = Context->GetEventHandlers()[EventScriptIdx];
 		FNiagaraDataSet *EventSet = Context->EventSets[EventScriptIdx];
 		if (EventSet)
 		{
@@ -308,7 +280,7 @@ void NiagaraEmitterInstanceBatcher::RunEventHandlers(const FNiagaraComputeExecut
 					RHICmdList.SetShaderResourceViewParameter(EventHandlerShader->GetComputeShader(), IntParam.GetBaseIndex(), EventSet->PrevDataRender().GetGPUBufferInt()->SRV);
 
 					TArray<uint8, TAlignedHeapAllocator<16>>  BlankParams;
-					Run(Context->MainDataSet, 0, NumInstancesAfterNonEventSpawn, EventHandlerShader, BlankParams, RHICmdList);
+					Run(Context->MainDataSet, 0, NumInstancesAfterNonEventSpawn, EventHandlerShader, BlankParams, RHICmdList, DummyWriteIndexBuffer);
 				}
 			}
 
@@ -345,11 +317,12 @@ void NiagaraEmitterInstanceBatcher::RunEventHandlers(const FNiagaraComputeExecut
 					check(NumInstancesAfterNonEventSpawn + EventSpawnNum == Context->MainDataSet->CurrDataRender().GetNumInstances());
 
 					TArray<uint8, TAlignedHeapAllocator<16>>  BlankParams;
-					Run(Context->MainDataSet, NumInstancesAfterNonEventSpawn, EventSpawnNum, EventHandlerShader, BlankParams, RHICmdList, true);
+					Run(Context->MainDataSet, NumInstancesAfterNonEventSpawn, EventSpawnNum, EventHandlerShader, BlankParams, RHICmdList, DummyWriteIndexBuffer, true);
 				}
 			}
 		}
 	}
+	*/
 }
 
 void NiagaraEmitterInstanceBatcher::SetPrevDataStrideParams(const FNiagaraDataSet *Set, FNiagaraShader *Shader, FRHICommandList &RHICmdList) const
diff --git a/Engine/Plugins/FX/Niagara/Source/Niagara/Private/NiagaraFunctionLibrary.cpp b/Engine/Plugins/FX/Niagara/Source/Niagara/Private/NiagaraFunctionLibrary.cpp
--- a/Engine/Plugins/FX/Niagara/Source/Niagara/Private/NiagaraFunctionLibrary.cpp
+++ b/Engine/Plugins/FX/Niagara/Source/Niagara/Private/NiagaraFunctionLibrary.cpp
@@ -30,6 +30,7 @@ UNiagaraComponent* UNiagaraFunctionLibrary::SpawnSystemAtLocation(UObject* World
 			AActor* Actor = World->GetWorldSettings();
 			PSC = NewObject<UNiagaraComponent>((Actor ? Actor : (UObject*)World));
 			PSC->SetAsset(SystemTemplate);
+			PSC->SetAutoDestroy(bAutoDestroy);
 			PSC->RegisterComponentWithWorld(World);
 
 			PSC->SetAbsolute(true, true, true);
@@ -63,6 +64,7 @@ UNiagaraComponent* UNiagaraFunctionLibrary::SpawnSystemAttached(UNiagaraSystem*
 			AActor* Actor = AttachToComponent->GetOwner();
 			PSC = NewObject<UNiagaraComponent>((Actor ? Actor : (UObject*)AttachToComponent->GetWorld()));
 			PSC->SetAsset(SystemTemplate);
+			PSC->SetAutoDestroy(bAutoDestroy);
 			PSC->RegisterComponentWithWorld(AttachToComponent->GetWorld());
 
 			PSC->AttachToComponent(AttachToComponent, FAttachmentTransformRules::KeepRelativeTransform, AttachPointName);
diff --git a/Engine/Plugins/FX/Niagara/Source/Niagara/Private/NiagaraMergeable.cpp b/Engine/Plugins/FX/Niagara/Source/Niagara/Private/NiagaraMergeable.cpp
new file mode 100644
--- /dev/null
+++ b/Engine/Plugins/FX/Niagara/Source/Niagara/Private/NiagaraMergeable.cpp
@@ -0,0 +1,49 @@
+// Copyright 1998-2017 Epic Games, Inc. All Rights Reserved.
+
+#include "NiagaraMergeable.h"
+#include "UObject/PropertyPortFlags.h"
+
+UNiagaraMergeable::UNiagaraMergeable()
+#if WITH_EDITORONLY_DATA
+	: MergeId(FGuid::NewGuid())
+#endif
+{
+}
+
+#if WITH_EDITOR
+void UNiagaraMergeable::PostEditChangeProperty(struct FPropertyChangedEvent& PropertyChangedEvent)
+{
+	Super::PostEditChangeProperty(PropertyChangedEvent);
+	OnChangedDelegate.Broadcast();
+}
+
+bool UNiagaraMergeable::Equals(const UNiagaraMergeable* Other)
+{
+	if (Other == nullptr || this->GetClass() != Other->GetClass())
+	{
+		return false;
+	}
+
+	for (TFieldIterator<UProperty> PropertyIterator(this->GetClass()); PropertyIterator; ++PropertyIterator)
+	{
+		if (PropertyIterator->Identical(
+			PropertyIterator->ContainerPtrToValuePtr<void>(this),
+			PropertyIterator->ContainerPtrToValuePtr<void>(Other), PPF_DeepComparison) == false)
+		{
+			return false;
+		}
+	}
+
+	return true;
+}
+
+UNiagaraMergeable::FOnChanged& UNiagaraMergeable::OnChanged()
+{
+	return OnChangedDelegate;
+}
+
+FGuid UNiagaraMergeable::GetMergeId()
+{
+	return MergeId;
+}
+#endif
\ No newline at end of file
diff --git a/Engine/Plugins/FX/Niagara/Source/Niagara/Private/NiagaraMeshRendererProperties.cpp b/Engine/Plugins/FX/Niagara/Source/Niagara/Private/NiagaraMeshRendererProperties.cpp
--- a/Engine/Plugins/FX/Niagara/Source/Niagara/Private/NiagaraMeshRendererProperties.cpp
+++ b/Engine/Plugins/FX/Niagara/Source/Niagara/Private/NiagaraMeshRendererProperties.cpp
@@ -19,9 +19,32 @@ void UNiagaraMeshRendererProperties::GetUsedMaterials(TArray<UMaterialInterface*
 {
 	if (ParticleMesh)
 	{
-		for(FStaticMaterial& Mat : ParticleMesh->StaticMaterials)
+		const FStaticMeshLODResources& LODModel = ParticleMesh->RenderData->LODResources[0];
+		if (bOverrideMaterials)
 		{
-			OutMaterials.Add(Mat.MaterialInterface);
+			for (int32 SectionIndex = 0; SectionIndex < LODModel.Sections.Num(); SectionIndex++)
+			{
+				const FStaticMeshSection& Section = LODModel.Sections[SectionIndex];
+				UMaterialInterface* ParticleMeshMaterial = ParticleMesh->GetMaterial(Section.MaterialIndex);
+
+				if (Section.MaterialIndex >= 0 && OverrideMaterials.Num() > Section.MaterialIndex && OverrideMaterials[Section.MaterialIndex] != nullptr)
+				{
+					OutMaterials.Add(OverrideMaterials[Section.MaterialIndex]);
+				}
+				else
+				{
+					OutMaterials.Add(ParticleMeshMaterial);
+				}
+			}
+		}
+		else
+		{
+			for (int32 SectionIndex = 0; SectionIndex < LODModel.Sections.Num(); SectionIndex++)
+			{
+				const FStaticMeshSection& Section = LODModel.Sections[SectionIndex];
+				UMaterialInterface* ParticleMeshMaterial = ParticleMesh->GetMaterial(Section.MaterialIndex);
+				OutMaterials.Add(ParticleMeshMaterial);
+			}
 		}
 	}
 }
@@ -30,9 +53,9 @@ void UNiagaraMeshRendererProperties::GetUsedMaterials(TArray<UMaterialInterface*
 
 bool UNiagaraMeshRendererProperties::IsMaterialValidForRenderer(UMaterial* Material, FText& InvalidMessage)
 {
-	if (Material->bUsedWithMeshParticles == false)
+	if (Material->bUsedWithNiagaraMeshParticles == false)
 	{
-		InvalidMessage = NSLOCTEXT("NiagaraMeshRendererProperties", "InvalidMaterialMessage", "The material isn't marked as \"Used with mesh particles\"");
+		InvalidMessage = NSLOCTEXT("NiagaraMeshRendererProperties", "InvalidMaterialMessage", "The material isn't marked as \"Used with Niagara Mesh particles\"");
 		return false;
 	}
 	return true;
diff --git a/Engine/Plugins/FX/Niagara/Source/Niagara/Private/NiagaraModule.cpp b/Engine/Plugins/FX/Niagara/Source/Niagara/Private/NiagaraModule.cpp
--- a/Engine/Plugins/FX/Niagara/Source/Niagara/Private/NiagaraModule.cpp
+++ b/Engine/Plugins/FX/Niagara/Source/Niagara/Private/NiagaraModule.cpp
@@ -13,6 +13,22 @@
 
 IMPLEMENT_MODULE(INiagaraModule, Niagara);
 
+#define LOCTEXT_NAMESPACE "NiagaraModule"
+
+/**
+Detail Level CVar.
+Effectively replaces the DetaiMode feature but allows for a rolling range of new hardware and emitters to target them.
+TODO: Possible that this might be more broadly useful across the engine as a replacement for DetailMode so placing in "r." rather than "fx."
+*/
+static TAutoConsoleVariable<float> CVarDetailLevel(
+	TEXT("r.DetailLevel"),
+	4,
+	TEXT("The detail level for use with Niagara.\n")
+	TEXT("If this value does not fall within an Emitter's MinDetailLevel and MaxDetailLevel range, then it will be disabled. \n")
+	TEXT("\n")
+	TEXT("Default = 4"),
+	ECVF_Scalability);
+
 void INiagaraModule::StartupModule()
 {
 	VectorVM::Init();
@@ -29,7 +45,9 @@ void INiagaraModule::StartupModule()
 	// @TODO We should remove this once Niagara is fully a plug-in.
 	FModuleManager::Get().LoadModule(TEXT("NiagaraEditor"));
 #endif
-	UNiagaraDataInterfaceCurlNoise::InitNoiseLUT();
+
+	CVarDetailLevel.AsVariable()->SetOnChangedCallback(FConsoleVariableDelegate::CreateRaw(this, &INiagaraModule::OnChangeDetailLevel));
+	OnChangeDetailLevel(CVarDetailLevel.AsVariable());
 }
 
 void INiagaraModule::ShutdownModule()
@@ -41,6 +59,8 @@ void INiagaraModule::ShutdownModule()
 		delete Pair.Value;
 		Pair.Value = nullptr;
 	}
+
+	CVarDetailLevel.AsVariable()->SetOnChangedCallback(FConsoleVariableDelegate());
 }
 
 FNiagaraWorldManager* INiagaraModule::GetWorldManager(UWorld* World)
@@ -65,13 +85,21 @@ void INiagaraModule::OnWorldInit(UWorld* World, const UWorld::InitializationValu
 void INiagaraModule::OnWorldCleanup(UWorld* World, bool bSessionEnded, bool bCleanupResources)
 {
 	//Cleanup world manager contents but not the manager itself.
+	FNiagaraWorldManager** Manager = WorldManagers.Find(World);
+	if (Manager)
+	{
+		(*Manager)->OnWorldCleanup(bSessionEnded, bCleanupResources);
+	}
 }
 
 void INiagaraModule::OnPreWorldFinishDestroy(UWorld* World)
 {
-	FNiagaraWorldManager*& Manager = WorldManagers.FindChecked(World);
-	delete Manager;
-	WorldManagers.Remove(World);
+	FNiagaraWorldManager** Manager = WorldManagers.Find(World);
+	if (Manager)
+	{
+		delete (*Manager);
+		WorldManagers.Remove(World);
+	}
 }
 
 void INiagaraModule::TickWorld(UWorld* World, ELevelTick TickType, float DeltaSeconds)
@@ -79,7 +107,49 @@ void INiagaraModule::TickWorld(UWorld* World, ELevelTick TickType, float DeltaSe
 	GetWorldManager(World)->Tick(DeltaSeconds);
 }
 
+#if WITH_EDITOR
+INiagaraModule::FMergeEmitterResults INiagaraModule::MergeEmitter(UNiagaraEmitter& Source, UNiagaraEmitter& LastMergedSource, UNiagaraEmitter& Instance)
+{
+	if (OnMergeEmitterDelegate.IsBound())
+	{
+		return OnMergeEmitterDelegate.Execute(Source, LastMergedSource, Instance);
+	}
+	FMergeEmitterResults Results;
+	Results.bSucceeded = false;
+	Results.ErrorMessages.Add(FText::Format(LOCTEXT("MergeDelegateNotRegisteredFormat", "Failed to merge emitter {0}.  Merge delegate not registered."), FText::FromString(Instance.GetPathName())));
+	return Results;
+}
+
+FDelegateHandle INiagaraModule::RegisterOnMergeEmitter(FOnMergeEmitter OnMergeEmitter)
+{
+	checkf(OnMergeEmitterDelegate.IsBound() == false, TEXT("Only one handler is allowed for the OnMergeEmitter delegate"));
+	OnMergeEmitterDelegate = OnMergeEmitter;
+	return OnMergeEmitterDelegate.GetHandle();
+}
+
+void INiagaraModule::UnregisterOnMergeEmitter(FDelegateHandle DelegateHandle)
+{
+	checkf(OnMergeEmitterDelegate.IsBound(), TEXT("OnMergeEmitter is not registered"));
+	checkf(OnMergeEmitterDelegate.GetHandle() == DelegateHandle, TEXT("Can only unregister the OnMergeEmitter delegate with the handle it was registered with."));
+	OnMergeEmitterDelegate.Unbind();
+}
+#endif
 
+int32 INiagaraModule::EngineDetailLevel = 4;
+void INiagaraModule::OnChangeDetailLevel(class IConsoleVariable* CVar)
+{
+	//Can only change the detail level at runtime on when not cooked.
+#if WITH_EDITORONLY_DATA
+	int32 NewDetailLevel = CVar->GetInt();
+	if (EngineDetailLevel != NewDetailLevel)
+	{
+		EngineDetailLevel = NewDetailLevel;
+		//If the detail level has changed we have to reset all systems.
+		FNiagaraSystemUpdateContext UpdateContext;
+		UpdateContext.AddAll(true);
+	}
+#endif
+}
 
 //////////////////////////////////////////////////////////////////////////
 
@@ -438,4 +508,6 @@ void INiagaraModule::ProcessShaderCompilationQueue()
 	return OnProcessQueue.Execute();
 }
 
+#undef LOCTEXT_NAMESPACE
+
 
diff --git a/Engine/Plugins/FX/Niagara/Source/Niagara/Private/NiagaraParameterBindingInstance.cpp b/Engine/Plugins/FX/Niagara/Source/Niagara/Private/NiagaraParameterBindingInstance.cpp
--- a/Engine/Plugins/FX/Niagara/Source/Niagara/Private/NiagaraParameterBindingInstance.cpp
+++ b/Engine/Plugins/FX/Niagara/Source/Niagara/Private/NiagaraParameterBindingInstance.cpp
@@ -11,5 +11,6 @@ FNiagaraParameterBindingInstance::FNiagaraParameterBindingInstance(const FNiagar
 
 void FNiagaraParameterBindingInstance::Tick()
 {
+	check(Source.GetSizeInBytes() == Destination.GetSizeInBytes());
 	Source.CopyTo(Destination.GetData());
 }
\ No newline at end of file
diff --git a/Engine/Plugins/FX/Niagara/Source/Niagara/Private/NiagaraParameterCollection.cpp b/Engine/Plugins/FX/Niagara/Source/Niagara/Private/NiagaraParameterCollection.cpp
--- a/Engine/Plugins/FX/Niagara/Source/Niagara/Private/NiagaraParameterCollection.cpp
+++ b/Engine/Plugins/FX/Niagara/Source/Niagara/Private/NiagaraParameterCollection.cpp
@@ -2,6 +2,10 @@
 
 #include "NiagaraParameterCollection.h"
 #include "NiagaraDataInterface.h"
+#if WITH_EDITORONLY_DATA
+	#include "IAssetTools.h"
+#endif
+#include "AssetRegistryModule.h"
 
 //////////////////////////////////////////////////////////////////////////
 
@@ -206,21 +210,19 @@ void UNiagaraParameterCollectionInstance::SetColorParameter(const FString& InVar
 UNiagaraParameterCollection::UNiagaraParameterCollection(const FObjectInitializer& ObjectInitializer)
 	: Super(ObjectInitializer)
 {
+	Namespace = *GetName();
 	DefaultInstance = ObjectInitializer.CreateDefaultSubobject<UNiagaraParameterCollectionInstance>(this, TEXT("Default Instance"));
 	DefaultInstance->SetParent(this);
 }
 
-void UNiagaraParameterCollection::PostInitProperties()
+#if WITH_EDITORONLY_DATA
+void UNiagaraParameterCollection::PostEditChangeProperty(struct FPropertyChangedEvent& PropertyChangedEvent)
 {
-	Super::PostInitProperties();
-	//Unique name is cached off just in case the hash function changes or our name does.
-	UniqueName = GetName() + LexicalConversion::ToString(GetTypeHash(GetFullName()));
-}
+	Super::PostEditChangeProperty(PropertyChangedEvent);
 
-void UNiagaraParameterCollection::PostLoad()
-{
-	Super::PostLoad();
+	MakeNamespaceNameUnique();
 }
+#endif
 
 int32 UNiagaraParameterCollection::IndexOfParameter(const FNiagaraVariable& Var)
 {
@@ -263,6 +265,10 @@ void UNiagaraParameterCollection::RenameParameter(FNiagaraVariable& Parameter, F
 	DefaultInstance->RenameParameter(Parameter, NewName);
 }
 
+FString UNiagaraParameterCollection::GetFullNamespace()const
+{
+	return TEXT("NPC.") + Namespace.ToString() + TEXT(".");
+}
 FNiagaraVariable UNiagaraParameterCollection::CollectionParameterFromFriendlyParameter(const FNiagaraVariable& FriendlyParameter)const
 {
 	return FNiagaraVariable(FriendlyParameter.GetType(), *ParameterNameFromFriendlyName(FriendlyParameter.GetName().ToString()));
@@ -275,11 +281,47 @@ FNiagaraVariable UNiagaraParameterCollection::FriendlyParameterFromCollectionPar
 
 FString UNiagaraParameterCollection::FriendlyNameFromParameterName(FString ParameterName)const
 {
-	ParameterName.ReplaceInline(*(GetUniqueName() + TEXT("_")), TEXT(""), ESearchCase::CaseSensitive);
+	ParameterName.RemoveFromStart(GetFullNamespace());
 	return ParameterName;
 }
 
 FString UNiagaraParameterCollection::ParameterNameFromFriendlyName(const FString& FriendlyName)const
 {
-	return FString::Printf(TEXT("%s_%s"), *GetUniqueName(), *FriendlyName);
+	return FString::Printf(TEXT("%s_%s"), *GetFullNamespace(), *FriendlyName);
+}
+
+void UNiagaraParameterCollection::MakeNamespaceNameUnique()
+{
+ 	FAssetRegistryModule& AssetRegistryModule = FModuleManager::LoadModuleChecked<FAssetRegistryModule>("AssetRegistry");
+ 	TArray<FAssetData> CollectionAssets;
+ 	AssetRegistryModule.Get().GetAssetsByClass(UNiagaraParameterCollection::StaticClass()->GetFName(), CollectionAssets);
+	TArray<FName> ExistingNames;
+ 	for (FAssetData& CollectionAsset : CollectionAssets)
+ 	{
+		if (CollectionAsset.GetFullName() != GetFullName())
+		{
+			ExistingNames.Add(CollectionAsset.GetTagValueRef<FName>(GET_MEMBER_NAME_CHECKED(UNiagaraParameterCollection, Namespace)));
+		}
+	}
+
+	if (ExistingNames.Contains(Namespace))
+	{
+		FString CandidateNameString = Namespace.ToString();
+		FString BaseNameString = CandidateNameString;
+		if (CandidateNameString.Len() >= 3 && CandidateNameString.Right(3).IsNumeric())
+		{
+			BaseNameString = CandidateNameString.Left(CandidateNameString.Len() - 3);
+		}
+
+		FName UniqueName = FName(*BaseNameString);
+		int32 NameIndex = 1;
+		while (ExistingNames.Contains(UniqueName))
+		{
+			UniqueName = FName(*FString::Printf(TEXT("%s%03i"), *BaseNameString, NameIndex));
+			NameIndex++;
+		}
+
+		UE_LOG(LogNiagara, Warning, TEXT("Parameter collection namespace conflict found. \"%s\" is already in use!"), *Namespace.ToString());
+		Namespace = UniqueName;
+	}
 }
\ No newline at end of file
