Copying Orion-Staging @ 3748653 (Orion/Dev-General @ 3722124) to //UE4/Main

diff --git a/Engine/Plugins/FX/Niagara/Source/Niagara/Private/NiagaraParameterStore.cpp b/Engine/Plugins/FX/Niagara/Source/Niagara/Private/NiagaraParameterStore.cpp
--- a/Engine/Plugins/FX/Niagara/Source/Niagara/Private/NiagaraParameterStore.cpp
+++ b/Engine/Plugins/FX/Niagara/Source/Niagara/Private/NiagaraParameterStore.cpp
@@ -8,6 +8,16 @@
 #include "NiagaraParameterCollection.h"
 
 
+#if WITH_EDITORONLY_DATA
+static int32 GbDumpParticleParameterStores = 0;
+static FAutoConsoleVariableRef CVarNiagaraDumpParticleParameterStores(
+	TEXT("fx.DumpParticleParameterStores"),
+	GbDumpParticleParameterStores,
+	TEXT("If > 0 current frame particle parameter stores will be dumped when updated. \n"),
+	ECVF_Default
+);
+#endif
+
 //////////////////////////////////////////////////////////////////////////
 
 FNiagaraParameterStore::FNiagaraParameterStore()
@@ -21,9 +31,13 @@ FNiagaraParameterStore::FNiagaraParameterStore(UObject* InOwner)
 	: Owner(InOwner)
 	, bParametersDirty(true)
 	, bInterfacesDirty(true)
-	, bLayoutDirty(true)
 {
-
+#if WITH_EDITORONLY_DATA
+	if (InOwner != nullptr)
+	{
+		Name = *InOwner->GetFullName();
+	}
+#endif
 }
 
 FNiagaraParameterStore::FNiagaraParameterStore(const FNiagaraParameterStore& Other)
@@ -43,41 +57,39 @@ FNiagaraParameterStore& FNiagaraParameterStore::operator=(const FNiagaraParamete
 
 FNiagaraParameterStore::~FNiagaraParameterStore()
 {
+	//Ensure that any stores bound to drive this one are unbound.
+	UnbindFromSourceStores();
+
+	//Also unbind from any stores we're feeding.
+	for (TPair<FNiagaraParameterStore*, FNiagaraParameterStoreBinding>& Binding : Bindings)
+	{
+		Binding.Value.Empty(Binding.Key, this);
+	}
 }
 
 void FNiagaraParameterStore::Bind(FNiagaraParameterStore* DestStore)
 {
 	check(DestStore);
-	FNiagaraParameterStoreBinding& Binding = Bindings.Add(DestStore);	
-
-	Binding.Empty();
-	for (auto& VarInfo : DestStore->ParameterOffsets)
-	{
-		Binding.BindParameter(*DestStore, VarInfo.Value, *this, VarInfo.Key);
-	}
 
-	//Tick once immediately to ensure initial value is correct.
-	Binding.Tick(*this, *DestStore, true);
+	FNiagaraParameterStoreBinding& Binding = Bindings.FindOrAdd(DestStore);
+	Binding.Initialize(DestStore, this);
 }
 
 void FNiagaraParameterStore::Unbind(FNiagaraParameterStore* DestStore)
 {
-	Bindings.Remove(DestStore);
+	FNiagaraParameterStoreBinding* Binding = Bindings.Find(DestStore);
+	if (Binding)
+	{
+		Binding->Empty(DestStore, this);
+		Bindings.Remove(DestStore);
+	}
 }
 
 void FNiagaraParameterStore::Rebind()
 {
 	for (TPair<FNiagaraParameterStore*, FNiagaraParameterStoreBinding>& Binding : Bindings)
 	{
-		Binding.Value.Empty();
-		for (auto& VarInfo : Binding.Key->ParameterOffsets)
-		{
-			Binding.Value.BindParameter(*Binding.Key, VarInfo.Value, *this, VarInfo.Key);
-		}
-
-		//Tick once immediately to ensure initial value is correct.
-		//TODO: Can probably rework binding so that we don't need to force all to update on a rebind.
-		Binding.Value.Tick(*this, *Binding.Key, true);
+		Binding.Value.Initialize(Binding.Key, this);
 	}
 }
 
@@ -87,27 +99,81 @@ void FNiagaraParameterStore::TransferBindings(FNiagaraParameterStore& OtherStore
 	{
 		OtherStore.Bind(Binding.Key);
 	}
+
 	Bindings.Empty();
 }
 
-void FNiagaraParameterStore::Tick()
+bool FNiagaraParameterStore::VerifyBinding(const FNiagaraParameterStore* DestStore)const
 {
-	if (bLayoutDirty)
+#if WITH_EDITORONLY_DATA
+	const FNiagaraParameterStoreBinding* Binding = Bindings.Find(DestStore);
+	if (Binding)
 	{
-		Rebind();
+		return Binding->VerifyBinding(DestStore, this);
 	}
 	else
 	{
-		for (auto& Binding : Bindings)
-		{
-			Binding.Value.Tick(*this, *Binding.Key);
-		}
+		UE_LOG(LogNiagara, Warning, TEXT("Invalid ParameterStore Binding: % was not bound to %s."), *Name.ToString(), *DestStore->Name.ToString());
 	}
 
+	return false;
+#else
+	return true;
+#endif
+}
+
+void FNiagaraParameterStore::Tick()
+{
+	for (TPair<FNiagaraParameterStore*, FNiagaraParameterStoreBinding>& Binding : Bindings)
+	{
+		Binding.Value.Tick(Binding.Key, this);
+	}
+
+	Dump();
+
 	//We have to have ticked all our source stores before now.
 	bParametersDirty = false;
 	bInterfacesDirty = false;
-	bLayoutDirty = false;
+}
+
+void FNiagaraParameterStore::UnbindFromSourceStores()
+{
+	//Each source store will remove itself from this array as it is unbound so after N unbinds the array should be empty.
+	int32 NumSourceStores = SourceStores.Num();
+	while (NumSourceStores--)
+	{
+		SourceStores[0]->Unbind(this);
+	}
+	ensureMsgf(SourceStores.Num() == 0, TEXT("Parameter store source array was not empty after unbinding all sources. Something seriously wrong."));
+}
+
+void FNiagaraParameterStore::DumpParameters()
+{
+	TArray<FNiagaraVariable> Vars;
+	GetParameters(Vars);
+	for (FNiagaraVariable& Var : Vars)
+	{
+		Var.SetData(GetParameterData_Internal(IndexOf(Var))); // Need to keep the data in sync
+		UE_LOG(LogNiagara, Log, TEXT("Param: %s Offset: %d Type : %s"), *Var.ToString(), IndexOf(Var), *Var.GetType().GetName());
+	}
+}
+
+void FNiagaraParameterStore::Dump()
+{
+#if WITH_EDITORONLY_DATA
+	if (GbDumpParticleParameterStores && GetParametersDirty())
+	{
+		UE_LOG(LogNiagara, Log, TEXT("\nSource Store: %s\n========================\n"), *Name.ToString());
+		DumpParameters();
+
+		for (TPair<FNiagaraParameterStore*, FNiagaraParameterStoreBinding>& Binding : Bindings)
+		{
+			Binding.Value.Dump(Binding.Key, this);
+		}
+
+		UE_LOG(LogNiagara, Log, TEXT("\n========================\n"));
+	}
+#endif
 }
 
 /**
@@ -217,6 +283,45 @@ void FNiagaraParameterStore::RenameParameter(const FNiagaraVariable& Param, FNam
 	}
 }
 
+void FNiagaraParameterStore::CopyParametersTo(FNiagaraParameterStore& DestStore, bool bOnlyAdd)
+{
+	TMap<FNiagaraVariable, int32>::TConstIterator It = ParameterOffsets.CreateConstIterator();
+	while (It)
+	{
+		FNiagaraVariable Parameter = It->Key;
+		int32 DestIndex = DestStore.IndexOf(Parameter);
+		int32 SrcIndex = It->Value;
+		bool bWrite = false;
+		if (DestIndex == INDEX_NONE)
+		{
+			DestStore.AddParameter(Parameter);
+			bWrite = !bOnlyAdd;
+			DestIndex = DestStore.IndexOf(Parameter);
+		}
+		else if (!bOnlyAdd)
+		{
+			bWrite = true;
+		}
+
+		if (bWrite && DestIndex != INDEX_NONE && SrcIndex != INDEX_NONE)
+		{
+			DestStore.SetParameterData(GetParameterData_Internal(SrcIndex), DestIndex, Parameter.GetSizeInBytes());
+		}
+		++It;
+	}
+}
+
+void FNiagaraParameterStore::RemoveParameters(FNiagaraParameterStore& DestStore)
+{
+	TMap<FNiagaraVariable, int32>::TConstIterator It = ParameterOffsets.CreateConstIterator();
+	while (It)
+	{
+		FNiagaraVariable Parameter = It->Key;
+		DestStore.RemoveParameter(Parameter);
+		++It;
+	}
+}
+
 void FNiagaraParameterStore::Empty(bool bClearBindings)
 {
 	ParameterOffsets.Empty();
@@ -226,9 +331,18 @@ void FNiagaraParameterStore::Empty(bool bClearBindings)
 	{
 		Bindings.Empty();
 	}
-	bParametersDirty = true;
-	bInterfacesDirty = true;
-	bLayoutDirty = true;
+	OnParameterChange();
+	OnInterfaceChange();
+	OnLayoutChange();
+}
+
+void FNiagaraParameterStore::OnLayoutChange()
+{
+	Rebind();
+
+#ifdef WITH_EDITOR
+	OnChangedDelegate.Broadcast();
+#endif
 }
 
 const FNiagaraVariable* FNiagaraParameterStore::FindVariable(UNiagaraDataInterface* Interface)const
@@ -250,4 +364,21 @@ const FNiagaraVariable* FNiagaraParameterStore::FindVariable(UNiagaraDataInterfa
 	return nullptr;
 }
 
+#ifdef WITH_EDITOR
+FDelegateHandle FNiagaraParameterStore::AddOnChangedHandler(FOnChanged::FDelegate InOnChanged)
+{
+	return OnChangedDelegate.Add(InOnChanged);
+}
+
+void FNiagaraParameterStore::RemoveOnChangedHandler(FDelegateHandle DelegateHandle)
+{
+	OnChangedDelegate.Remove(DelegateHandle);
+}
+
+void FNiagaraParameterStore::RemoveAllOnChangedHandlers(const void* InUserObject)
+{
+	OnChangedDelegate.RemoveAll(InUserObject);
+}
+#endif
+
 //////////////////////////////////////////////////////////////////////////
\ No newline at end of file
diff --git a/Engine/Plugins/FX/Niagara/Source/Niagara/Private/NiagaraRenderer.cpp b/Engine/Plugins/FX/Niagara/Source/Niagara/Private/NiagaraRenderer.cpp
--- a/Engine/Plugins/FX/Niagara/Source/Niagara/Private/NiagaraRenderer.cpp
+++ b/Engine/Plugins/FX/Niagara/Source/Niagara/Private/NiagaraRenderer.cpp
@@ -5,6 +5,7 @@
 #include "ParticleBeamTrailVertexFactory.h"
 #include "NiagaraDataSet.h"
 #include "NiagaraStats.h"
+#include "NiagaraVertexFactory.h"
 
 DECLARE_CYCLE_STAT(TEXT("Generate Particle Lights"), STAT_NiagaraGenLights, STATGROUP_Niagara);
 
@@ -20,6 +21,17 @@ static FAutoConsoleVariableRef CVarParallelEmitterRenderers(
 
 
 
+NiagaraRenderer::NiagaraRenderer()
+	: CPUTimeMS(0.0f)
+	, bLocalSpace(false)
+	, bEnabled(true)
+	, DynamicDataRender(nullptr)
+	, BaseExtents(1.0f, 1.0f, 1.0f)
+{
+	Material = UMaterial::GetDefaultMaterial(MD_Surface);
+	FNiagaraVertexFactoryBase::Init();
+}
+
 
 NiagaraRenderer::~NiagaraRenderer() 
 {
diff --git a/Engine/Plugins/FX/Niagara/Source/Niagara/Private/NiagaraRendererMeshes.cpp b/Engine/Plugins/FX/Niagara/Source/Niagara/Private/NiagaraRendererMeshes.cpp
--- a/Engine/Plugins/FX/Niagara/Source/Niagara/Private/NiagaraRendererMeshes.cpp
+++ b/Engine/Plugins/FX/Niagara/Source/Niagara/Private/NiagaraRendererMeshes.cpp
@@ -41,17 +41,33 @@ NiagaraRendererMeshes::NiagaraRendererMeshes(ERHIFeatureLevel::Type FeatureLevel
 
 	if (Properties && Properties->ParticleMesh)
 	{
+		if (Properties->bOverrideMaterials && Properties->OverrideMaterials.Num() != 0)
+		{
+			for (UMaterialInterface* Interface : Properties->OverrideMaterials)
+			{
+				if (Interface)
+				{
+					Interface->CheckMaterialUsage_Concurrent(MATUSAGE_NiagaraMeshParticles);
+				}
+			}
+		}
+
 		const FStaticMeshLODResources& LODModel = Properties->ParticleMesh->RenderData->LODResources[0];
 		for (int32 SectionIndex = 0; SectionIndex < LODModel.Sections.Num(); SectionIndex++)
 		{
 			//FMaterialRenderProxy* MaterialProxy = MeshMaterials[SectionIndex]->GetRenderProxy(bSelected);
 			const FStaticMeshSection& Section = LODModel.Sections[SectionIndex];
 			UMaterialInterface* ParticleMeshMaterial = Properties->ParticleMesh->GetMaterial(Section.MaterialIndex);
-			if (ParticleMeshMaterial)
+			if (ParticleMeshMaterial && Properties->bOverrideMaterials == false)
 			{
-				ParticleMeshMaterial->CheckMaterialUsage_Concurrent(MATUSAGE_NiagaraMeshParticles);
+				if (ParticleMeshMaterial)
+				{
+					ParticleMeshMaterial->CheckMaterialUsage_Concurrent(MATUSAGE_NiagaraMeshParticles);
+				}
 			}
 		}
+
+		BaseExtents = Properties->ParticleMesh->GetBounds().BoxExtent;
 	}
 
 }
@@ -138,11 +154,20 @@ void NiagaraRendererMeshes::GetDynamicMeshElements(const TArray<const FSceneView
 	SimpleTimer MeshElementsTimer;
 
 	FNiagaraDynamicDataMesh *DynamicDataMesh = (static_cast<FNiagaraDynamicDataMesh*>(DynamicDataRender));
-	if (!DynamicDataMesh || DynamicDataMesh->DataSet->GetNumInstances()==0)
+	//if (!DynamicDataMesh || !DynamicDataMesh->DataSet || DynamicDataMesh->DataSet->GetNumInstances()==0)
+	if (!DynamicDataMesh 
+		|| DynamicDataMesh->DataSet->CurrDataRender().GetNumInstancesAllocated() == 0
+		|| DynamicDataMesh->DataSet->CurrDataRender().GetNumInstances() == 0
+		|| DynamicDataMesh->DataSet->CurrDataRender().GetGPUBufferFloat() == nullptr
+		|| nullptr == Properties)
 	{
 		return;
 	}
 
+	//check(DynamicDataMesh->DataSet->PrevDataRender().GetNumInstances() > 0);
+	//check(DynamicDataMesh->DataSet->PrevDataRender().GetNumInstancesAllocated() > 0);
+	//check(DynamicDataMesh->DataSet->PrevData().GetGPUBufferFloat()->NumBytes > 0 || DynamicDataMesh->DataSet->PrevData().GetGPUBufferInt()->NumBytes > 0)
+	
 	{
 		// Update the primitive uniform buffer if needed.
 		if (!WorldSpacePrimitiveUniformBuffer.IsInitialized())
@@ -175,6 +200,8 @@ void NiagaraRendererMeshes::GetDynamicMeshElements(const TArray<const FSceneView
 				FNiagaraMeshCollectorResourcesMesh& CollectorResources = Collector.AllocateOneFrameResource<FNiagaraMeshCollectorResourcesMesh>();
 				SetupVertexFactory(&CollectorResources.VertexFactory, LODModel);
 				FNiagaraMeshUniformParameters PerViewUniformParameters;// = UniformParameters;
+				PerViewUniformParameters.LocalToWorld = bLocalSpace ? SceneProxy->GetLocalToWorld() : FMatrix::Identity;//For now just handle local space like this but maybe in future have a VF variant to avoid the transform entirely?
+				PerViewUniformParameters.LocalToWorldInverseTransposed = bLocalSpace ? SceneProxy->GetLocalToWorld().Inverse().GetTransposed() : FMatrix::Identity;
 				PerViewUniformParameters.PrevTransformAvailable = false;
 				PerViewUniformParameters.DeltaSeconds = ViewFamily.DeltaWorldTime;
 				PerViewUniformParameters.PositionDataOffset = DynamicDataMesh->PositionDataOffset;
@@ -194,6 +221,7 @@ void NiagaraRendererMeshes::GetDynamicMeshElements(const TArray<const FSceneView
 
 				// Collector.AllocateOneFrameResource uses default ctor, initialize the vertex factory
 				CollectorResources.VertexFactory.SetParticleFactoryType(NVFT_Mesh);
+				CollectorResources.VertexFactory.SetMeshFacingMode((uint32)Properties->FacingMode);
 				CollectorResources.UniformBuffer = FNiagaraMeshUniformBufferRef::CreateUniformBufferImmediate(PerViewUniformParameters, UniformBuffer_SingleFrame);
 				CollectorResources.VertexFactory.SetStrides(0, 0);
 
@@ -210,7 +238,18 @@ void NiagaraRendererMeshes::GetDynamicMeshElements(const TArray<const FSceneView
 					{
 						continue;
 					}
-					FMaterialRenderProxy* MaterialProxy = ParticleMeshMaterial->GetRenderProxy(false, false);
+					FMaterialRenderProxy* MaterialProxy = nullptr;
+					
+					if (Properties->bOverrideMaterials && Properties->OverrideMaterials.Num() > Section.MaterialIndex && 
+						Properties->OverrideMaterials[Section.MaterialIndex] != nullptr)
+					{
+						MaterialProxy = Properties->OverrideMaterials[Section.MaterialIndex]->GetRenderProxy(false, false);
+					}
+
+					if (MaterialProxy == nullptr)
+					{
+						MaterialProxy = ParticleMeshMaterial->GetRenderProxy(false, false);
+					}
 
 					if ((Section.NumTriangles == 0) || (MaterialProxy == NULL))
 					{
@@ -231,7 +270,8 @@ void NiagaraRendererMeshes::GetDynamicMeshElements(const TArray<const FSceneView
 					BatchElement.FirstIndex = 0;
 					BatchElement.MinVertexIndex = 0;
 					BatchElement.MaxVertexIndex = 0;
-					BatchElement.NumInstances = DynamicDataMesh->DataSet->PrevDataRender().GetNumInstances();
+					BatchElement.NumInstances = DynamicDataMesh->DataSet->CurrDataRender().GetNumInstances();
+					BatchElement.IndirectArgsBuffer = DynamicDataMesh->DataSet->GetDataSetIndices().Buffer;
 
 					if (bIsWireframe)
 					{
@@ -289,7 +329,7 @@ FNiagaraDynamicDataBase *NiagaraRendererMeshes::GenerateVertexData(const FNiagar
 {
 	SCOPE_CYCLE_COUNTER(STAT_NiagaraGenMeshVertexData);
 
-	if (!Properties || Properties->ParticleMesh == nullptr || !bEnabled)
+	if (!Properties || Properties->ParticleMesh == nullptr || !bEnabled || Data.CurrData().GetNumInstances() == 0)
 	{
 		return nullptr;
 	}
@@ -302,9 +342,9 @@ FNiagaraDynamicDataBase *NiagaraRendererMeshes::GenerateVertexData(const FNiagar
 	const FNiagaraVariableLayoutInfo* PositionLayout = Data.GetVariableLayout(FNiagaraVariable(FNiagaraTypeDefinition::GetVec3Def(), TEXT("Position")));
 	const FNiagaraVariableLayoutInfo* VelocityLayout = Data.GetVariableLayout(FNiagaraVariable(FNiagaraTypeDefinition::GetVec3Def(), TEXT("Velocity")));
 	const FNiagaraVariableLayoutInfo* ColorLayout = Data.GetVariableLayout(FNiagaraVariable(FNiagaraTypeDefinition::GetColorDef(), TEXT("Color")));
-	const FNiagaraVariableLayoutInfo* XformLayout = Data.GetVariableLayout(FNiagaraVariable(FNiagaraTypeDefinition::GetVec4Def(), TEXT("Transform")));
+	const FNiagaraVariableLayoutInfo* XformLayout = Data.GetVariableLayout(FNiagaraVariable(FNiagaraTypeDefinition::GetVec4Def(), TEXT("MeshOrientation")));
 	const FNiagaraVariableLayoutInfo* ScaleLayout = Data.GetVariableLayout(FNiagaraVariable(FNiagaraTypeDefinition::GetVec3Def(), TEXT("Scale")));
-	const FNiagaraVariableLayoutInfo* SizeLayout = Data.GetVariableLayout(FNiagaraVariable(FNiagaraTypeDefinition::GetVec2Def(), TEXT("Size")));
+	const FNiagaraVariableLayoutInfo* SizeLayout = Data.GetVariableLayout(FNiagaraVariable(FNiagaraTypeDefinition::GetVec2Def(), TEXT("SpriteSize")));
 	const FNiagaraVariableLayoutInfo* MatParamLayout = Data.GetVariableLayout(FNiagaraVariable(FNiagaraTypeDefinition::GetVec4Def(), TEXT("DynamicMaterialParameter")));
 
 	//Bail if we don't have the required attributes to render this emitter.
@@ -329,18 +369,20 @@ FNiagaraDynamicDataBase *NiagaraRendererMeshes::GenerateVertexData(const FNiagar
 	int32 IntDummy;
 	Data.GetVariableComponentOffsets(FNiagaraVariable(FNiagaraTypeDefinition::GetColorDef(), TEXT("Color")), DynamicData->ColorDataOffset, IntDummy);
 	Data.GetVariableComponentOffsets(FNiagaraVariable(FNiagaraTypeDefinition::GetVec3Def(), TEXT("Scale")), DynamicData->ScaleDataOffset, IntDummy);
-	Data.GetVariableComponentOffsets(FNiagaraVariable(FNiagaraTypeDefinition::GetVec2Def(), TEXT("Size")), DynamicData->SizeDataOffset, IntDummy);
+	Data.GetVariableComponentOffsets(FNiagaraVariable(FNiagaraTypeDefinition::GetVec2Def(), TEXT("SpriteSize")), DynamicData->SizeDataOffset, IntDummy);
 	Data.GetVariableComponentOffsets(FNiagaraVariable(FNiagaraTypeDefinition::GetVec4Def(), TEXT("DynamicMaterialParameter")), DynamicData->MaterialParamDataOffset, IntDummy);
-	Data.GetVariableComponentOffsets(FNiagaraVariable(FNiagaraTypeDefinition::GetVec4Def(), TEXT("Transform")), DynamicData->TransformDataOffset, IntDummy);
+	Data.GetVariableComponentOffsets(FNiagaraVariable(FNiagaraTypeDefinition::GetVec4Def(), TEXT("MeshOrientation")), DynamicData->TransformDataOffset, IntDummy);
 
 	// if we're CPU simulating, need to init the GPU buffers for the vertex factory
-	if (Data.PrevData().GetNumInstances() > 0)
+	if (Target == ENiagaraSimTarget::CPUSim)
 	{
-		if (Target == ENiagaraSimTarget::CPUSim)
-		{
-			Data.ValidateBufferIndices();
-			Data.InitGPUFromCPU();
-		}
+		//Data.ValidateBufferIndices();
+		Data.InitGPUFromCPU();
+	}
+	else
+	{
+		check(Data.GetDataSetIndices().Buffer != nullptr);
+		Data.InitGPUSimSRVs();
 	}
 
 	DynamicData->DataSet = &Data;
diff --git a/Engine/Plugins/FX/Niagara/Source/Niagara/Private/NiagaraRendererRibbons.cpp b/Engine/Plugins/FX/Niagara/Source/Niagara/Private/NiagaraRendererRibbons.cpp
--- a/Engine/Plugins/FX/Niagara/Source/Niagara/Private/NiagaraRendererRibbons.cpp
+++ b/Engine/Plugins/FX/Niagara/Source/Niagara/Private/NiagaraRendererRibbons.cpp
@@ -9,6 +9,7 @@
 DECLARE_CYCLE_STAT(TEXT("Generate Ribbon Vertex Data"), STAT_NiagaraGenRibbonVertexData, STATGROUP_Niagara);
 DECLARE_CYCLE_STAT(TEXT("Render Ribbons"), STAT_NiagaraRenderRibbons, STATGROUP_Niagara);
 
+DECLARE_CYCLE_STAT(TEXT("Genereate GPU Buffers"), STAT_NiagaraGenRibbonGpuBuffers, STATGROUP_Niagara);
 
 
 struct FNiagaraDynamicDataRibbon : public FNiagaraDynamicDataBase
@@ -16,6 +17,12 @@ struct FNiagaraDynamicDataRibbon : public FNiagaraDynamicDataBase
 	TArray<FNiagaraRibbonVertex> VertexData;
 	TArray<int16> IndexData;
 	TArray<FNiagaraRibbonVertexDynamicParameter> MaterialParameterVertexData;
+
+	const FNiagaraDataSet *DataSet;
+	int32 PositionDataOffset;
+	int32 WidthDataOffset;
+	int32 TwistDataOffset;
+	int32 ColorDataOffset;
 };
 
 class FNiagaraMeshCollectorResourcesRibbon : public FOneFrameResource
@@ -61,7 +68,7 @@ void NiagaraRendererRibbons::GetDynamicMeshElements(const TArray<const FSceneVie
 
 	SimpleTimer MeshElementsTimer;
 	FNiagaraDynamicDataRibbon *DynamicDataRibbon = static_cast<FNiagaraDynamicDataRibbon*>(DynamicDataRender);
-	if (!DynamicDataRibbon || DynamicDataRibbon->VertexData.Num() == 0)
+	if (!DynamicDataRibbon || DynamicDataRibbon->VertexData.Num() < 4)
 	{
 		return;
 	}
@@ -125,6 +132,12 @@ void NiagaraRendererRibbons::GetDynamicMeshElements(const TArray<const FSceneVie
 				PerViewUniformParameters.CameraUp = View->GetViewUp(); // FVector4(0.0f, 0.0f, 1.0f, 0.0f);
 				PerViewUniformParameters.CameraRight = View->GetViewRight();//	FVector4(1.0f, 0.0f, 0.0f, 0.0f);
 				PerViewUniformParameters.ScreenAlignment = FVector4(0.0f, 0.0f, 0.0f, 0.0f);
+				PerViewUniformParameters.UseCustomFacing = Properties->FacingMode == ENiagaraRibbonFacingMode::Custom;
+				PerViewUniformParameters.PositionDataOffset = DynamicDataRibbon->PositionDataOffset;
+				PerViewUniformParameters.ColorDataOffset = DynamicDataRibbon->ColorDataOffset;
+				PerViewUniformParameters.WidthDataOffset = DynamicDataRibbon->WidthDataOffset;
+				PerViewUniformParameters.TwistDataOffset = DynamicDataRibbon->TwistDataOffset;
+				CollectorResources.VertexFactory.SetParticleData(DynamicDataRibbon->DataSet);
 
 				// Collector.AllocateOneFrameResource uses default ctor, initialize the vertex factory
 				CollectorResources.VertexFactory.SetParticleFactoryType(NVFT_Ribbon);
@@ -149,7 +162,7 @@ void NiagaraRendererRibbons::GetDynamicMeshElements(const TArray<const FSceneVie
 				MeshBatch.bUseAsOccluder = false;
 				MeshBatch.ReverseCulling = SceneProxy->IsLocalToWorldDeterminantNegative();
 				MeshBatch.bDisableBackfaceCulling = true;
-				MeshBatch.Type = PT_TriangleStrip;
+				MeshBatch.Type = PT_TriangleList;
 				MeshBatch.DepthPriorityGroup = SceneProxy->GetDepthPriorityGroup(View);
 				MeshBatch.bCanApplyViewModeOverrides = true;
 				MeshBatch.bUseWireframeSelectionColoring = SceneProxy->IsSelected();
@@ -165,8 +178,8 @@ void NiagaraRendererRibbons::GetDynamicMeshElements(const TArray<const FSceneVie
 
 				FMeshBatchElement& MeshElement = MeshBatch.Elements[0];
 				MeshElement.IndexBuffer = DynamicIndexAllocation.IndexBuffer;
-				MeshElement.FirstIndex = 0;
-				MeshElement.NumPrimitives = DynamicDataRibbon->IndexData.Num() - 2;
+				MeshElement.FirstIndex = DynamicIndexAllocation.FirstIndex;
+				MeshElement.NumPrimitives = DynamicDataRibbon->IndexData.Num() / 3;
 				MeshElement.NumInstances = 1;
 				MeshElement.MinVertexIndex = 0;
 				MeshElement.MaxVertexIndex = DynamicDataRibbon->VertexData.Num() - 1;
@@ -198,6 +211,8 @@ int NiagaraRendererRibbons::GetDynamicDataSize()
 	if (DynamicDataRender)
 	{
 		Size += (static_cast<FNiagaraDynamicDataRibbon*>(DynamicDataRender))->VertexData.GetAllocatedSize();
+		Size += (static_cast<FNiagaraDynamicDataRibbon*>(DynamicDataRender))->IndexData.GetAllocatedSize();
+		Size += (static_cast<FNiagaraDynamicDataRibbon*>(DynamicDataRender))->MaterialParameterVertexData.GetAllocatedSize();
 	}
 
 	return Size;
@@ -247,18 +262,37 @@ FNiagaraDynamicDataBase *NiagaraRendererRibbons::GenerateVertexData(const FNiaga
 	FNiagaraDataSetIterator<FVector> VelItr(Data, FNiagaraVariable(FNiagaraTypeDefinition::GetVec3Def(), TEXT("Velocity")));
 	FNiagaraDataSetIterator<FLinearColor> ColItr(Data, FNiagaraVariable(FNiagaraTypeDefinition::GetColorDef(), TEXT("Color")));
 	FNiagaraDataSetIterator<float> NormAgeItr(Data, FNiagaraVariable(FNiagaraTypeDefinition::GetFloatDef(), TEXT("NormalizedAge")));
-	FNiagaraDataSetIterator<float> RotItr(Data, FNiagaraVariable(FNiagaraTypeDefinition::GetFloatDef(), TEXT("Rotation")));
+	FNiagaraDataSetIterator<float> AgeItr(Data, FNiagaraVariable(FNiagaraTypeDefinition::GetFloatDef(), TEXT("NormalizedAge")));
 	FNiagaraDataSetIterator<int32> RibbonIdItr(Data, FNiagaraVariable(FNiagaraTypeDefinition::GetIntDef(), TEXT("RibbonId")));
-	FNiagaraDataSetIterator<FVector2D> SizeItr(Data, FNiagaraVariable(FNiagaraTypeDefinition::GetVec2Def(), TEXT("Size")));
-	//FNiagaraDataSetIterator<float> SubImageIndexItr(Data, FNiagaraVariable(FNiagaraTypeDefinition::GetFloatDef(), TEXT("SubImageIndex")));
+	FNiagaraDataSetIterator<float> SizeItr(Data, FNiagaraVariable(FNiagaraTypeDefinition::GetFloatDef(), TEXT("RibbonWidth")));
+	FNiagaraDataSetIterator<float> TwistItr(Data, FNiagaraVariable(FNiagaraTypeDefinition::GetFloatDef(), TEXT("RibbonTwist")));
+	FNiagaraDataSetIterator<FVector> AlignItr(Data, FNiagaraVariable(FNiagaraTypeDefinition::GetVec3Def(), TEXT("RibbonFacing")));
 	FNiagaraDataSetIterator<FVector4> MaterialParamItr(Data, FNiagaraVariable(FNiagaraTypeDefinition::GetVec4Def(), TEXT("DynamicMaterialParameter")));
 
 	//Bail if we don't have the required attributes to render this emitter.
-	if (!PosItr.IsValid() || !ColItr.IsValid() || !NormAgeItr.IsValid() || !SizeItr.IsValid())
+	if (!PosItr.IsValid() || !ColItr.IsValid() || !NormAgeItr.IsValid() || !AgeItr.IsValid() )
 	{
 		return DynamicData;
 	}
 
+	const FNiagaraVariableLayoutInfo* PositionLayout = Data.GetVariableLayout(FNiagaraVariable(FNiagaraTypeDefinition::GetVec3Def(), TEXT("Position")));
+	const FNiagaraVariableLayoutInfo* TwistLayout = Data.GetVariableLayout(FNiagaraVariable(FNiagaraTypeDefinition::GetFloatDef(), TEXT("RibbonTwist")));
+	const FNiagaraVariableLayoutInfo* WidthLayout = Data.GetVariableLayout(FNiagaraVariable(FNiagaraTypeDefinition::GetFloatDef(), TEXT("RibbonWidth")));
+	const FNiagaraVariableLayoutInfo* ColorLayout = Data.GetVariableLayout(FNiagaraVariable(FNiagaraTypeDefinition::GetColorDef(), TEXT("Color")));
+
+	// required attributes
+	DynamicData->PositionDataOffset = PositionLayout->FloatComponentStart;
+	DynamicData->ColorDataOffset = ColorLayout->FloatComponentStart;
+
+	// optional attributes
+	int32 IntDummy;
+	Data.GetVariableComponentOffsets(FNiagaraVariable(FNiagaraTypeDefinition::GetFloatDef(), TEXT("RibbonWidth")), DynamicData->WidthDataOffset, IntDummy);
+	Data.GetVariableComponentOffsets(FNiagaraVariable(FNiagaraTypeDefinition::GetFloatDef(), TEXT("RibbonTwist")), DynamicData->TwistDataOffset, IntDummy);
+
+	DynamicData->DataSet = &Data;
+	
+	////////
+
 	bool bMultiRibbons = RibbonIdItr.IsValid();
 
 	// build a sorted list by age, so we always get particles in order regardless of them being moved around due to dieing and spawning
@@ -267,42 +301,41 @@ FNiagaraDynamicDataBase *NiagaraRendererRibbons::GenerateVertexData(const FNiaga
 	TArray<int32> RibbonIDs;
 	TArray<int32> UniqueRibbonIDs;
 	TArray<FVector> Positions;
-	TArray<FVector2D> Sizes;
+	TArray<FVector> FacingVectors;
+	TArray<FVector4> DynamicParams;
+	TArray<float> Sizes;
 	TArray<FLinearColor> Colors;
 	TArray<float> Rotations;
 	TArray<float> Ages;
 	uint32 CurIdx = 0;
 	for (uint32 Idx = 0; Idx < Data.GetNumInstances(); Idx++)
 	{
-		float Age = NormAgeItr.Get();
-		if (Idx>0 && Age > 0.0f && Age < 1.0f)
+		float Age = NormAgeItr.GetAdvance();
+		if (Age >= 0.0f && Age < 1.0f)
 		{
 			SortedIndices.Add(CurIdx++);
-			Positions.Add(PosItr.Get());
-			Sizes.Add(SizeItr.Get());
-			Rotations.Add(RotItr.Get());
+			Positions.Add(PosItr.GetAdvance());
+			Sizes.Add(SizeItr.GetAdvanceWithDefault(1.0));
+			Rotations.Add(TwistItr.GetAdvanceWithDefault(0.0f));
+			FacingVectors.Add(AlignItr.GetAdvanceWithDefault(FVector(0.0f, 0.0f, 1.0f)));
+			DynamicParams.Add(MaterialParamItr.GetAdvanceWithDefault(FVector4(0.0f, 0.0f, 0.0f, 0.0f)));
+
 			Ages.Add(Age);
-			Colors.Add(ColItr.Get());
+			Colors.Add(ColItr.GetAdvance());
 			if (bMultiRibbons)
 			{
-				int32 RibbonID = RibbonIdItr.Get();
+				int32 RibbonID = RibbonIdItr.GetAdvance();
 				RibbonIDs.Add(RibbonID);
 				UniqueRibbonIDs.AddUnique(RibbonID);
 			}
 		}
-		PosItr.Advance();
-		SizeItr.Advance();
-		RotItr.Advance();
-		NormAgeItr.Advance();
-		ColItr.Advance();
-		RibbonIdItr.Advance();
 	}
 
 	if (!bMultiRibbons)
 	{
 		SortedIndices.Sort(
 			[&Ages](const int32& A, const int32& B) {
-			return (Ages[A] > Ages[B]);
+			return (Ages[A] < Ages[B]);
 		}
 		);
 		UniqueRibbonIDs.Add(0);
@@ -322,12 +355,12 @@ FNiagaraDynamicDataBase *NiagaraRendererRibbons::GenerateVertexData(const FNiaga
 		return DynamicData;
 	}
 
-	RenderData.Reset(SortedIndices.Num() * 2);
+	RenderData.Empty();
 	IndexData.Empty();
 
 	// TODO : deal with the dynamic vertex material parameter should the user have specified it as an output...
 
-	FVector2D UVs[4] = { FVector2D(0.0f, 0.0f), FVector2D(1.0f, 0.0f), FVector2D(1.0f, 1.0f), FVector2D(0.0f, 1.0f) };
+	FVector2D UVs[4] = { FVector2D(1.0f, 0.0f), FVector2D(1.0f, 1.0f) };
 	int32 StartIndex = 0;
 	int32 NumIndices = SortedIndices.Num();
 	int32 *IndexPtr = &SortedIndices[0];
@@ -347,57 +380,79 @@ FNiagaraDynamicDataBase *NiagaraRendererRibbons::GenerateVertexData(const FNiaga
 			}
 			SortedPartialIndices.Sort(
 				[&Ages](const int32& A, const int32& B) {
-				return (Ages[A] > Ages[B]);
+				return (Ages[A] < Ages[B]);
 			}
 			);
 			IndexPtr = &SortedPartialIndices[0];
 		}
 
-		FVector PrevPos, PrevPos2, PrevDir(0.0f, 0.0f, 0.1f);
-		for (int32 i = 0; i < NumIndices - 1; i++)
+		if (NumIndices > 1)
 		{
-			uint32 Index1 = IndexPtr[i];
-			uint32 Index2 = IndexPtr[i + 1];
+			FVector PrevPos, PrevPos2, PrevDir(0.0f, 0.0f, 0.1f);
+			float TotalDistance = 0.0f;
+			FVector2D AgeOffset(Ages.Last(), 0.0f);
 
-			const FVector ParticlePos = Positions[Index1];
-			FVector ParticleDir = Positions[Index2] - ParticlePos;
-
-			// if two ribbon particles were spawned too close together, we skip one
-			if (ParticleDir.SizeSquared() <= FMath::Square(SMALL_NUMBER))
+			for (int32 i = 0; i < NumIndices; i++)
 			{
-				continue;
-			}
-			FVector NormDir = ParticleDir.GetSafeNormal();
-			PrevDir = NormDir;
-
-			FVector ParticleRight = FVector::CrossProduct(NormDir, FVector(0.0f, 0.0f, 1.0f));
-			ParticleRight *= Sizes[Index1].X;
-			FVector ParticleRightRot = ParticleRight.RotateAngleAxis(Rotations[Index1], NormDir);
+				uint32 Index1 = IndexPtr[i];
+				uint32 Index2 = IndexPtr[i + 1];
 
+				const FVector ParticlePos = Positions[Index1];
+				FVector ParticleDir;
+				if (i < NumIndices - 1)
+				{
+					ParticleDir = Positions[Index2] - ParticlePos;
+				}
+				else
+				{
+					Index2 = IndexPtr[i - 1];
+					ParticleDir = ParticlePos - Positions[Index2];
+				}
 
-			AddRibbonVert(RenderData, ParticlePos + ParticleRightRot, UVs[0], Colors[Index1], Ages[Index1], Rotations[Index1], Sizes[Index1]);
-			AddRibbonVert(RenderData, ParticlePos - ParticleRightRot, UVs[1], Colors[Index1], Ages[Index1], Rotations[Index1], Sizes[Index1]);
-
-			IndexData.Add(NumTotalVerts);
-			IndexData.Add(NumTotalVerts + 1);
-			NumTotalVerts += 2;
 
-			PrevPos = ParticlePos;
+				// if two ribbon particles were spawned too close together, we skip one
+				// but never skip the last, because that will result in invalid indices from the prev loop
+				if (ParticleDir.SizeSquared() > 0.002 )
+				{
+					FVector NormDir = ParticleDir.GetSafeNormal();
+					PrevDir = NormDir;
+					FVector2D UV0Mult((float)i / NumIndices, 1.0f);
+					FVector2D UV1Mult((float)i / NumIndices, 1.0f);
+
+					if (Properties->UV0TilingDistance)
+					{
+						UV0Mult = FVector2D(TotalDistance / Properties->UV0TilingDistance, 1.0f);
+					}
+					if (Properties->UV1TilingDistance)
+					{
+						UV1Mult = FVector2D(TotalDistance / Properties->UV1TilingDistance, 1.0f);
+					}
+
+					TotalDistance += ParticleDir.Size();
+					AddRibbonVert(RenderData, ParticlePos, UVs[0] * UV0Mult + AgeOffset, UVs[0] * UV1Mult + AgeOffset, Colors[Index1], Ages[Index1], Rotations[Index1], Sizes[Index1], NormDir, FacingVectors[Index1]);
+					AddRibbonVert(RenderData, ParticlePos, UVs[1] * UV0Mult + AgeOffset, UVs[1] * UV1Mult + AgeOffset, Colors[Index1], Ages[Index1], Rotations[Index1], Sizes[Index1], NormDir, FacingVectors[Index1]);
+					if (DynamicParams.Num())
+					{
+						AddDynamicParam(DynamicData->MaterialParameterVertexData, DynamicParams[Index1]);
+						AddDynamicParam(DynamicData->MaterialParameterVertexData, DynamicParams[Index1]);
+					}
+
+					if (i < NumIndices - 1)
+					{
+						IndexData.Add(NumTotalVerts);
+						IndexData.Add(NumTotalVerts + 1);
+						IndexData.Add(NumTotalVerts + 2);
+						IndexData.Add(NumTotalVerts + 1);
+						IndexData.Add(NumTotalVerts + 3);
+						IndexData.Add(NumTotalVerts + 2);
+					}
+					NumTotalVerts += 2;
+				}
+			}
 		}
-
 		StartIndex += NumIndices;
-
-		// add degenerates to connect with the next ribbon
-		if (RibbonIdx < UniqueRibbonIDs.Num() - 1)
-		{
-			IndexData.Add(FMath::Max(NumTotalVerts - 1, 0));
-			IndexData.Add(FMath::Max(NumTotalVerts - 1, 0));
-			IndexData.Add(FMath::Max(NumTotalVerts, 0));
-			IndexData.Add(FMath::Max(NumTotalVerts, 0));
-		}
 	}
 
-	RenderData.SetNum(NumTotalVerts);
 	CPUTimeMS = VertexDataTimer.GetElapsedMilliseconds();
 
 	return DynamicData;
diff --git a/Engine/Plugins/FX/Niagara/Source/Niagara/Private/NiagaraRendererSprites.cpp b/Engine/Plugins/FX/Niagara/Source/Niagara/Private/NiagaraRendererSprites.cpp
--- a/Engine/Plugins/FX/Niagara/Source/Niagara/Private/NiagaraRendererSprites.cpp
+++ b/Engine/Plugins/FX/Niagara/Source/Niagara/Private/NiagaraRendererSprites.cpp
@@ -9,53 +9,11 @@
 DECLARE_CYCLE_STAT(TEXT("Generate Sprite Vertex Data"), STAT_NiagaraGenSpriteVertexData, STATGROUP_Niagara);
 DECLARE_CYCLE_STAT(TEXT("Render Sprites"), STAT_NiagaraRenderSprites, STATGROUP_Niagara);
 
-DECLARE_CYCLE_STAT(TEXT("Genereate GPU Buffers"), STAT_NiagaraGenGpuBuffers, STATGROUP_Niagara);
-
-
-/**
-* Per-particle data sent to the GPU.
-*/
-struct FNiagaraSpriteVertex
-{
-	/** The position of the particle. */
-	FVector Position;
-	/** The relative time of the particle. */
-	float RelativeTime;
-	/** The previous position of the particle. */
-	FVector	OldPosition;
-	/** Value that remains constant over the lifetime of a particle. */
-	float ParticleId;
-	/** The size of the particle. */
-	FVector2D Size;
-	/** The rotation of the particle. */
-	float Rotation;
-	/** The sub-image index for the particle. */
-	float SubImageIndex;
-	/** The color of the particle. */
-	FLinearColor Color;
-	/* Custom Alignment vector*/
-	FVector CustomAlignmentVector;
-	/* Custom Facing vector*/
-	FVector CustomFacingVector;
-};
-
+DECLARE_CYCLE_STAT(TEXT("Genereate GPU Buffers"), STAT_NiagaraGenSpriteGpuBuffers, STATGROUP_Niagara);
 
 struct FNiagaraDynamicDataSprites : public FNiagaraDynamicDataBase
 {
-	TArray<FNiagaraSpriteVertex> VertexData;
-	TArray<FParticleVertexDynamicParameter> MaterialParameterVertexData;
 	const FNiagaraDataSet *DataSet;
-	//uint32 ComponentBufferSize;
-	int32 PositionDataOffset;
-	int32 VelocityDataOffset;
-	int32 SizeDataOffset;
-	int32 RotationDataOffset;
-	int32 SubimageDataOffset;
-	int32 ColorDataOffset;
-	int32 FacingOffset;
-	int32 AlignmentOffset;
-	//int32 NumInstances;
-	bool bCustomAlignmentAvailable;
 };
 
 
@@ -75,12 +33,25 @@ public:
 
 
 
-NiagaraRendererSprites::NiagaraRendererSprites(ERHIFeatureLevel::Type FeatureLevel, UNiagaraRendererProperties *InProps) :
-	NiagaraRenderer()
+NiagaraRendererSprites::NiagaraRendererSprites(ERHIFeatureLevel::Type FeatureLevel, UNiagaraRendererProperties *InProps) 
+	: NiagaraRenderer()
+	, PositionOffset(INDEX_NONE)
+	, VelocityOffset(INDEX_NONE)
+	, RotationOffset(INDEX_NONE)
+	, SizeOffset(INDEX_NONE)
+	, ColorOffset(INDEX_NONE)
+	, FacingOffset(INDEX_NONE)
+	, AlignmentOffset(INDEX_NONE)
+	, SubImageOffset(INDEX_NONE)
+	, MaterialParamOffset(INDEX_NONE)
+	, CameraOffsetOffset(INDEX_NONE)
+	, UVScaleOffset(INDEX_NONE)
+	, ParticleRandomOffset(INDEX_NONE)
 {
 	//check(InProps);
 	VertexFactory = new FNiagaraSpriteVertexFactory(NVFT_Sprite, FeatureLevel);
 	Properties = Cast<UNiagaraSpriteRendererProperties>(InProps);
+	BaseExtents = FVector(0.5f, 0.5f, 0.5f);
 }
 
 
@@ -106,31 +77,22 @@ void NiagaraRendererSprites::GetDynamicMeshElements(const TArray<const FSceneVie
 	//check(DynamicDataRender)
 
 	FNiagaraDynamicDataSprites *DynamicDataSprites = static_cast<FNiagaraDynamicDataSprites*>(DynamicDataRender);
-	if (!DynamicDataSprites || DynamicDataSprites->DataSet->PrevDataRender().GetNumInstances()==0 || nullptr == Properties)
+
+	if (!DynamicDataSprites
+		|| DynamicDataSprites->DataSet->CurrDataRender().GetNumInstancesAllocated() == 0
+		|| DynamicDataSprites->DataSet->CurrDataRender().GetNumInstances() == 0
+		|| DynamicDataSprites->DataSet->CurrDataRender().GetGPUBufferFloat() == nullptr
+		|| nullptr == Properties
+		)
 	{
 		return;
 	}
 
-	int32 NumInstances = DynamicDataSprites->DataSet->PrevDataRender().GetNumInstances();
+	int32 NumInstances = DynamicDataSprites->DataSet->CurrDataRender().GetNumInstances();
 
 	const bool bIsWireframe = ViewFamily.EngineShowFlags.Wireframe;
 	FMaterialRenderProxy* MaterialRenderProxy = Material->GetRenderProxy(SceneProxy->IsSelected(), SceneProxy->IsHovered());
 
-	int32 SizeInBytes = DynamicDataSprites->VertexData.GetTypeSize() * DynamicDataSprites->VertexData.Num();
-	FGlobalDynamicVertexBuffer::FAllocation LocalDynamicVertexMaterialParamsAllocation;
-
-	if (DynamicDataSprites->MaterialParameterVertexData.Num() > 0)
-	{
-		int32 MatParamSizeInBytes = DynamicDataSprites->MaterialParameterVertexData.GetTypeSize() * DynamicDataSprites->MaterialParameterVertexData.Num();
-		LocalDynamicVertexMaterialParamsAllocation = FGlobalDynamicVertexBuffer::Get().Allocate(MatParamSizeInBytes);
-
-		if (LocalDynamicVertexMaterialParamsAllocation.IsValid())
-		{
-			// Copy the extra material vertex data over.
-			FMemory::Memcpy(LocalDynamicVertexMaterialParamsAllocation.Buffer, DynamicDataSprites->MaterialParameterVertexData.GetData(), MatParamSizeInBytes);
-		}
-	}
-
 	// Update the primitive uniform buffer if needed.
 	if (!WorldSpacePrimitiveUniformBuffer.IsInitialized())
 	{
@@ -157,42 +119,13 @@ void NiagaraRendererSprites::GetDynamicMeshElements(const TArray<const FSceneVie
 	{
 		if (VisibilityMap & (1 << ViewIndex))
 		{
-			FGlobalDynamicVertexBuffer::FAllocation LocalDynamicVertexAllocation = FGlobalDynamicVertexBuffer::Get().Allocate(SizeInBytes);
-			//if (LocalDynamicVertexAllocation.IsValid())
 			{
 				const FSceneView* View = Views[ViewIndex];
-				FMatrix TransMat;
-				/*
-				if (Properties->SortMode == ENiagaraSortMode::SortViewDepth)
-				{
-					TransMat = View->ViewMatrices.GetViewProjectionMatrix();
-
-					// view depth sorting
-					DynamicDataSprites->VertexData.Sort(
-						[&TransMat](const FNiagaraSpriteVertex& A, const FNiagaraSpriteVertex& B) {
-						float W1 = TransMat.TransformPosition(A.Position).W;
-						float W2 = TransMat.TransformPosition(B.Position).W;
-						return (W1 > W2);
-					});
-				}
-				else if (Properties->SortMode == ENiagaraSortMode::SortViewDistance)
-				{
-					FVector ViewLoc = View->ViewLocation;
-					// dist sorting
-					DynamicDataSprites->VertexData.Sort(
-						[&ViewLoc](const FNiagaraSpriteVertex& A, const FNiagaraSpriteVertex& B) {
-						float D1 = (ViewLoc - A.Position).SizeSquared();
-						float D2 = (ViewLoc - B.Position).SizeSquared();
-						return (D1 > D2);
-					});
-				}
-
-				// Copy the vertex data over.
-				*/
-				//FMemory::Memcpy(LocalDynamicVertexAllocation.Buffer, DynamicDataSprites->VertexData.GetData(), SizeInBytes);
 
 				FNiagaraMeshCollectorResourcesSprite& CollectorResources = Collector.AllocateOneFrameResource<FNiagaraMeshCollectorResourcesSprite>();
 				FNiagaraSpriteUniformParameters PerViewUniformParameters;// = UniformParameters;
+				PerViewUniformParameters.LocalToWorld = bLocalSpace ? SceneProxy->GetLocalToWorld() : FMatrix::Identity;//For now just handle local space like this but maybe in future have a VF variant to avoid the transform entirely?
+				PerViewUniformParameters.LocalToWorldInverseTransposed = bLocalSpace ? SceneProxy->GetLocalToWorld().Inverse().GetTransposed() : FMatrix::Identity;
 				PerViewUniformParameters.RotationBias = 0.0f;
 				PerViewUniformParameters.RotationScale = 1.0f;
 				PerViewUniformParameters.TangentSelector = FVector4(0.0f, 0.0f, 0.0f, 1.0f);
@@ -200,20 +133,52 @@ void NiagaraRendererSprites::GetDynamicMeshElements(const TArray<const FSceneVie
 				PerViewUniformParameters.NormalsType = 0.0f;
 				PerViewUniformParameters.NormalsSphereCenter = FVector4(0.0f, 0.0f, 0.0f, 1.0f);
 				PerViewUniformParameters.NormalsCylinderUnitDirection = FVector4(0.0f, 0.0f, 1.0f, 0.0f);
-				PerViewUniformParameters.PivotOffset = FVector2D(-0.5f, -0.5f);
+				PerViewUniformParameters.PivotOffset = Properties->PivotInUVSpace * -1.0f; // We do this because we want to slide the coordinates back since 0,0 is the upper left corner.
 				PerViewUniformParameters.MacroUVParameters = FVector4(0.0f, 0.0f, 1.0f, 1.0f);
 				PerViewUniformParameters.CameraFacingBlend = FVector4(0.0f, 0.0f, 0.0f, 1.0f);
-				PerViewUniformParameters.RemoveHMDRoll = 0.0f;
-				PerViewUniformParameters.CustomFacingVectorMask = Properties->CustomFacingVectorMask;
+				PerViewUniformParameters.RemoveHMDRoll = Properties->bRemoveHMDRollInVR;
+				PerViewUniformParameters.CustomFacingVectorMask = FVector4(0.0f, 0.0f, 0.0f, 0.0f);
 				PerViewUniformParameters.SubImageSize = FVector4(Properties->SubImageSize.X, Properties->SubImageSize.Y, 1.0f / Properties->SubImageSize.X, 1.0f / Properties->SubImageSize.Y);
-				PerViewUniformParameters.PositionDataOffset = DynamicDataSprites->PositionDataOffset;
-				PerViewUniformParameters.VelocityDataOffset = DynamicDataSprites->VelocityDataOffset;
-				PerViewUniformParameters.RotationDataOffset = DynamicDataSprites->RotationDataOffset;
-				PerViewUniformParameters.SizeDataOffset = DynamicDataSprites->SizeDataOffset;
-				PerViewUniformParameters.ColorDataOffset = DynamicDataSprites->ColorDataOffset;
-				PerViewUniformParameters.SubimageDataOffset = DynamicDataSprites->SubimageDataOffset;
-				PerViewUniformParameters.FacingOffset = DynamicDataSprites->FacingOffset;
-				PerViewUniformParameters.AlignmentOffset = DynamicDataSprites->AlignmentOffset;
+				PerViewUniformParameters.PositionDataOffset = PositionOffset;
+				PerViewUniformParameters.VelocityDataOffset = VelocityOffset;
+				PerViewUniformParameters.RotationDataOffset = RotationOffset;
+				PerViewUniformParameters.SizeDataOffset = SizeOffset;
+				PerViewUniformParameters.ColorDataOffset = ColorOffset;
+				PerViewUniformParameters.MaterialParamDataOffset = MaterialParamOffset;
+				PerViewUniformParameters.SubimageDataOffset = SubImageOffset;
+				PerViewUniformParameters.FacingOffset = FacingOffset;
+				PerViewUniformParameters.AlignmentOffset = AlignmentOffset;
+				PerViewUniformParameters.SubImageBlendMode = Properties->bSubImageBlend;
+				PerViewUniformParameters.CameraOffsetOffset = CameraOffsetOffset;
+				PerViewUniformParameters.UVScaleOffset = UVScaleOffset;
+				PerViewUniformParameters.ParticleRandomOffset = ParticleRandomOffset;
+
+				// Collector.AllocateOneFrameResource uses default ctor, initialize the vertex factory
+				ENiagaraSpriteFacingMode FacingMode = Properties->FacingMode;
+				ENiagaraSpriteAlignment AlignmentMode = Properties->Alignment;
+
+				if (FacingOffset == -1 && FacingMode == ENiagaraSpriteFacingMode::CustomFacingVector)
+				{
+					FacingMode = ENiagaraSpriteFacingMode::FaceCamera;
+				}
+
+				if (AlignmentOffset == -1 && AlignmentMode == ENiagaraSpriteAlignment::CustomAlignment)
+				{
+					AlignmentMode = ENiagaraSpriteAlignment::Unaligned;
+				}
+
+				if (FacingMode == ENiagaraSpriteFacingMode::FaceCameraDistanceBlend)
+				{
+					float DistanceBlendMinSq = Properties->MinFacingCameraBlendDistance * Properties->MinFacingCameraBlendDistance;
+					float DistanceBlendMaxSq = Properties->MaxFacingCameraBlendDistance * Properties->MaxFacingCameraBlendDistance;
+					float InvBlendRange = 1.0f / FMath::Max(DistanceBlendMaxSq - DistanceBlendMinSq, 1.0f);
+					float BlendScaledMinDistance = DistanceBlendMinSq * InvBlendRange;
+
+					PerViewUniformParameters.CameraFacingBlend.X = 1.0f;
+					PerViewUniformParameters.CameraFacingBlend.Y = InvBlendRange;
+					PerViewUniformParameters.CameraFacingBlend.Z = BlendScaledMinDistance;
+				}
+
 				if (Properties->Alignment == ENiagaraSpriteAlignment::VelocityAligned)
 				{
 					// velocity aligned
@@ -221,16 +186,14 @@ void NiagaraRendererSprites::GetDynamicMeshElements(const TArray<const FSceneVie
 					PerViewUniformParameters.TangentSelector = FVector4(0.0f, 1.0f, 0.0f, 0.0f);
 				}
 
-				// Collector.AllocateOneFrameResource uses default ctor, initialize the vertex factory
+				if (Properties->FacingMode == ENiagaraSpriteFacingMode::CustomFacingVector)
+				{
+					PerViewUniformParameters.CustomFacingVectorMask = Properties->CustomFacingVectorMask;
+				}
 
-				// use custom alignment if the data is available and if it's set in the props
-				bool bUseCustomAlignment = DynamicDataSprites->bCustomAlignmentAvailable && Properties->Alignment == ENiagaraSpriteAlignment::CustomAlignment;
-				bool bUseVectorAlignment = Properties->Alignment != ENiagaraSpriteAlignment::Unaligned;
-				
 				CollectorResources.VertexFactory.SetParticleData(DynamicDataSprites->DataSet);
-				CollectorResources.VertexFactory.SetCustomAlignment(bUseCustomAlignment);
-				CollectorResources.VertexFactory.SetVectorAligned(bUseVectorAlignment);
-				CollectorResources.VertexFactory.SetCameraPlaneFacing(Properties->FacingMode == ENiagaraSpriteFacingMode::FaceCameraPlane);
+				CollectorResources.VertexFactory.SetAlignmentMode((uint32)AlignmentMode);
+				CollectorResources.VertexFactory.SetFacingMode((uint32)FacingMode);
 
 				CollectorResources.VertexFactory.SetParticleFactoryType(NVFT_Sprite);
 
@@ -241,27 +204,6 @@ void NiagaraRendererSprites::GetDynamicMeshElements(const TArray<const FSceneVie
 				CollectorResources.VertexFactory.InitResource();
 				CollectorResources.VertexFactory.SetSpriteUniformBuffer(CollectorResources.UniformBuffer);
 
-				CollectorResources.VertexFactory.SetInstanceBuffer(
-					/*LocalDynamicVertexAllocation.VertexBuffer*/nullptr,
-					LocalDynamicVertexAllocation.VertexOffset,
-					sizeof(FNiagaraSpriteVertex),
-					true
-					);
-
-
-				/*if (DynamicDataSprites->MaterialParameterVertexData.Num() > 0 && LocalDynamicVertexMaterialParamsAllocation.IsValid())
-				{
-					CollectorResources.VertexFactory.SetDynamicParameterBuffer(
-						LocalDynamicVertexMaterialParamsAllocation.VertexBuffer,
-						LocalDynamicVertexMaterialParamsAllocation.VertexOffset,
-						sizeof(FParticleVertexDynamicParameter),
-						true);
-				}
-				else*/
-				{
-					CollectorResources.VertexFactory.SetDynamicParameterBuffer(NULL, 0, 0, true);
-				}
-
 				FMeshBatch& MeshBatch = Collector.AllocateMesh();
 				MeshBatch.VertexFactory = &CollectorResources.VertexFactory;
 				MeshBatch.CastShadow = SceneProxy->CastsDynamicShadow();
@@ -289,10 +231,9 @@ void NiagaraRendererSprites::GetDynamicMeshElements(const TArray<const FSceneVie
 				MeshElement.MinVertexIndex = 0;
 				MeshElement.MaxVertexIndex = 0;// MeshElement.NumInstances * 4 - 1;
 				MeshElement.PrimitiveUniformBufferResource = &WorldSpacePrimitiveUniformBuffer;
+				MeshElement.IndirectArgsBuffer = DynamicDataSprites->DataSet->GetDataSetIndices().Buffer;
 				
-				
-				Collector.AddMesh(ViewIndex, MeshBatch);
-				
+				Collector.AddMesh(ViewIndex, MeshBatch);				
 			}
 		}
 	}
@@ -311,48 +252,55 @@ FNiagaraDynamicDataBase *NiagaraRendererSprites::GenerateVertexData(const FNiaga
 {
 	FNiagaraDynamicDataSprites *DynamicData = nullptr;
 
-	if (Data.PrevData().GetNumInstances() > 0)
+	if (Data.CurrData().GetNumInstances() > 0 && (Target != ENiagaraSimTarget::CPUSim || (Target == ENiagaraSimTarget::CPUSim && Data.CurrData().GetSizeBytes() != 0)))
 	{
 		SimpleTimer VertexDataTimer;
 
 		SCOPE_CYCLE_COUNTER(STAT_NiagaraGenSpriteVertexData);
 
-		const FNiagaraVariableLayoutInfo* PositionLayout = Data.GetVariableLayout(FNiagaraVariable(FNiagaraTypeDefinition::GetVec3Def(), TEXT("Position")));
-		const FNiagaraVariableLayoutInfo* VelocityLayout = Data.GetVariableLayout(FNiagaraVariable(FNiagaraTypeDefinition::GetVec3Def(), TEXT("Velocity")));
-		const FNiagaraVariableLayoutInfo* RotationLayout = Data.GetVariableLayout(FNiagaraVariable(FNiagaraTypeDefinition::GetFloatDef(), TEXT("SpriteRotation")));
-		const FNiagaraVariableLayoutInfo* SizeLayout = Data.GetVariableLayout(FNiagaraVariable(FNiagaraTypeDefinition::GetVec2Def(), TEXT("SpriteSize")));
-		const FNiagaraVariableLayoutInfo* ColorLayout = Data.GetVariableLayout(FNiagaraVariable(FNiagaraTypeDefinition::GetColorDef(), TEXT("Color")));
-
-		if (!bEnabled || !PositionLayout || !VelocityLayout || !RotationLayout || !SizeLayout || !ColorLayout)
+		if (PositionOffset == INDEX_NONE)
+		{
+			const FNiagaraVariableLayoutInfo* PositionLayout = Data.GetVariableLayout(FNiagaraVariable(FNiagaraTypeDefinition::GetVec3Def(), TEXT("Position")));
+			const FNiagaraVariableLayoutInfo* VelocityLayout = Data.GetVariableLayout(FNiagaraVariable(FNiagaraTypeDefinition::GetVec3Def(), TEXT("Velocity")));
+			const FNiagaraVariableLayoutInfo* RotationLayout = Data.GetVariableLayout(FNiagaraVariable(FNiagaraTypeDefinition::GetFloatDef(), TEXT("SpriteRotation")));
+			const FNiagaraVariableLayoutInfo* SizeLayout = Data.GetVariableLayout(FNiagaraVariable(FNiagaraTypeDefinition::GetVec2Def(), TEXT("SpriteSize")));
+			const FNiagaraVariableLayoutInfo* ColorLayout = Data.GetVariableLayout(FNiagaraVariable(FNiagaraTypeDefinition::GetColorDef(), TEXT("Color")));
+			PositionOffset = PositionLayout->FloatComponentStart;
+			VelocityOffset = VelocityLayout->FloatComponentStart;
+			RotationOffset = RotationLayout->FloatComponentStart;
+			SizeOffset = SizeLayout->FloatComponentStart;
+			ColorOffset = ColorLayout->FloatComponentStart;
+
+			// optional attributes; we pass -1 as the offset so the VF can branch
+			int32 IntDummy;
+			Data.GetVariableComponentOffsets(FNiagaraVariable(FNiagaraTypeDefinition::GetVec3Def(), TEXT("SpriteFacing")), FacingOffset, IntDummy);
+			Data.GetVariableComponentOffsets(FNiagaraVariable(FNiagaraTypeDefinition::GetVec3Def(), TEXT("SpriteAlignment")), AlignmentOffset, IntDummy);
+			Data.GetVariableComponentOffsets(FNiagaraVariable(FNiagaraTypeDefinition::GetFloatDef(), TEXT("SubImageIndex")), SubImageOffset, IntDummy);
+			Data.GetVariableComponentOffsets(FNiagaraVariable(FNiagaraTypeDefinition::GetVec4Def(), TEXT("DynamicMaterialParameter")), MaterialParamOffset, IntDummy);
+			Data.GetVariableComponentOffsets(FNiagaraVariable(FNiagaraTypeDefinition::GetFloatDef(), TEXT("CameraOffset")), CameraOffsetOffset, IntDummy);
+			Data.GetVariableComponentOffsets(FNiagaraVariable(FNiagaraTypeDefinition::GetVec2Def(), TEXT("UVScale")), UVScaleOffset, IntDummy);
+			Data.GetVariableComponentOffsets(FNiagaraVariable(FNiagaraTypeDefinition::GetFloatDef(), TEXT("MaterialRandom")), ParticleRandomOffset, IntDummy);
+		}
+		
+		if (!bEnabled || PositionOffset == INDEX_NONE || VelocityOffset == INDEX_NONE || RotationOffset == INDEX_NONE || SizeOffset == INDEX_NONE || ColorOffset == INDEX_NONE)
 		{
 			return nullptr;
 		}
-
-
-		// required attributes //TODO: This does not need to be done every frame. Losing not insignificant time getting the layout ptrs.
+		
 		DynamicData = new FNiagaraDynamicDataSprites;
-		DynamicData->PositionDataOffset = PositionLayout->FloatComponentStart;
-		DynamicData->VelocityDataOffset = VelocityLayout->FloatComponentStart;
-		DynamicData->RotationDataOffset = RotationLayout->FloatComponentStart;
-		DynamicData->SizeDataOffset = SizeLayout->FloatComponentStart;
-		DynamicData->ColorDataOffset = ColorLayout->FloatComponentStart;
-
-		// optional attributes; we pass -1 as the offset so the VF can branch
-		int32 IntDummy;
-		Data.GetVariableComponentOffsets(FNiagaraVariable(FNiagaraTypeDefinition::GetVec3Def(), TEXT("Facing")), DynamicData->FacingOffset, IntDummy);
-		Data.GetVariableComponentOffsets(FNiagaraVariable(FNiagaraTypeDefinition::GetVec3Def(), TEXT("Alignment")), DynamicData->AlignmentOffset, IntDummy);
-		Data.GetVariableComponentOffsets(FNiagaraVariable(FNiagaraTypeDefinition::GetFloatDef(), TEXT("SubImageIndex")), DynamicData->SubimageDataOffset, IntDummy);
-
-
 		// if we're CPU simulating, need to init the GPU buffers for the vertex factory
 			//This is a killer for perf. Assume its all the GPU buffer allocation.
 			//Need to allocate big shared buffer and allocate ranges out of those.
-		SCOPE_CYCLE_COUNTER(STAT_NiagaraGenGpuBuffers);
+		SCOPE_CYCLE_COUNTER(STAT_NiagaraGenSpriteGpuBuffers);
 		if (Target == ENiagaraSimTarget::CPUSim)
 		{
-			Data.ValidateBufferIndices();
 			Data.InitGPUFromCPU();
 		}
+		else
+		{
+			check(Data.GetDataSetIndices().Buffer != nullptr);
+			Data.InitGPUSimSRVs();
+		}
 
 		DynamicData->DataSet = &Data;
 		CPUTimeMS = VertexDataTimer.GetElapsedMilliseconds();
@@ -378,11 +326,6 @@ void NiagaraRendererSprites::SetDynamicData_RenderThread(FNiagaraDynamicDataBase
 int NiagaraRendererSprites::GetDynamicDataSize()
 {
 	uint32 Size = sizeof(FNiagaraDynamicDataSprites);
-	if (DynamicDataRender)
-	{
-		Size += (static_cast<FNiagaraDynamicDataSprites*>(DynamicDataRender))->VertexData.GetAllocatedSize();
-	}
-
 	return Size;
 }
 
diff --git a/Engine/Plugins/FX/Niagara/Source/Niagara/Private/NiagaraRibbonRendererProperties.cpp b/Engine/Plugins/FX/Niagara/Source/Niagara/Private/NiagaraRibbonRendererProperties.cpp
--- a/Engine/Plugins/FX/Niagara/Source/Niagara/Private/NiagaraRibbonRendererProperties.cpp
+++ b/Engine/Plugins/FX/Niagara/Source/Niagara/Private/NiagaraRibbonRendererProperties.cpp
@@ -23,11 +23,8 @@ const TArray<FNiagaraVariable>& UNiagaraRibbonRendererProperties::GetRequiredAtt
 	if (Attrs.Num() == 0)
 	{
 		Attrs.Add(SYS_PARAM_PARTICLES_POSITION);
-		/*Attrs.Add(SYS_PARAM_PARTICLES_VELOCITY);
-		Attrs.Add(SYS_PARAM_PARTICLES_COLOR);
-		Attrs.Add(SYS_PARAM_PARTICLES_SPRITE_ROTATION);
 		Attrs.Add(SYS_PARAM_PARTICLES_NORMALIZED_AGE);
-		Attrs.Add(SYS_PARAM_PARTICLES_SPRITE_SIZE);*/
+		Attrs.Add(SYS_PARAM_PARTICLES_COLOR);
 	}
 
 	return Attrs;
@@ -40,10 +37,10 @@ const TArray<FNiagaraVariable>& UNiagaraRibbonRendererProperties::GetOptionalAtt
 
 	if (Attrs.Num() == 0)
 	{
-		//Attrs.Add(SYS_PARAM_PARTICLES_SPRITE_FACING);
-		//Attrs.Add(SYS_PARAM_PARTICLES_SPRITE_ALIGNMENT);
-		//Attrs.Add(SYS_PARAM_PARTICLES_SUB_IMAGE_INDEX);
-		//Attrs.Add(SYS_PARAM_PARTICLES_DYNAMIC_MATERIAL_PARAM);
+		Attrs.Add(SYS_PARAM_PARTICLES_RIBBONID);
+		Attrs.Add(SYS_PARAM_PARTICLES_RIBBONTWIST);
+		Attrs.Add(SYS_PARAM_PARTICLES_RIBBONWIDTH);
+		Attrs.Add(SYS_PARAM_PARTICLES_RIBBONFACING);
 	}
 
 	return Attrs;
diff --git a/Engine/Plugins/FX/Niagara/Source/Niagara/Private/NiagaraScript.cpp b/Engine/Plugins/FX/Niagara/Source/Niagara/Private/NiagaraScript.cpp
--- a/Engine/Plugins/FX/Niagara/Source/Niagara/Private/NiagaraScript.cpp
+++ b/Engine/Plugins/FX/Niagara/Source/Niagara/Private/NiagaraScript.cpp
@@ -25,36 +25,60 @@ FNiagaraScriptDebuggerInfo::FNiagaraScriptDebuggerInfo() : bRequestDebugFrame(fa
 {
 }
 
+
+UNiagaraScriptSourceBase::UNiagaraScriptSourceBase(const FObjectInitializer& ObjectInitializer)
+	: Super(ObjectInitializer)
+{
+	
+}
+
+
 UNiagaraScript::UNiagaraScript(const FObjectInitializer& ObjectInitializer)
 	: Super(ObjectInitializer)
 	, Usage(ENiagaraScriptUsage::Function)
 	, ModuleUsageBitmask( (1 << (int32)ENiagaraScriptUsage::ParticleSpawnScript) | (1 << (int32)ENiagaraScriptUsage::ParticleSpawnScriptInterpolated) | (1 << (int32)ENiagaraScriptUsage::ParticleUpdateScript) | (1 << (int32)ENiagaraScriptUsage::ParticleEventScript) )
 	, NumUserPtrs(0)
 	, NumericOutputTypeSelectionMode(ENiagaraNumericOutputTypeSelectionMode::Largest)
-#if WITH_EDITORONLY_DATA
 	, LastCompileStatus(ENiagaraScriptCompileStatus::NCS_Unknown)
-	, UniqueID(FGuid::NewGuid())
+#if WITH_EDITORONLY_DATA
+	, UsageIndex_DEPRECATED(0)
 #endif
 {
-	
+#if WITH_EDITORONLY_DATA
+	ScriptResource.OnCompilationComplete().AddUniqueDynamic(this, &UNiagaraScript::OnCompilationComplete);
+
+	RapidIterationParameters.Name = *GetFullName();
+#endif	
 }
 
 UNiagaraScript::~UNiagaraScript()
 {
 }
 
-void UNiagaraScript::PostInitProperties()
+bool UNiagaraScript::ContainsUsage(ENiagaraScriptUsage InUsage) const
 {
-	Super::PostInitProperties();
-	UsageIndex = 0;
-}
+	if (IsEquivalentUsage(InUsage))
+	{
+		return true;
+	}
 
-UNiagaraScriptSourceBase::UNiagaraScriptSourceBase(const FObjectInitializer& ObjectInitializer)
-	: Super(ObjectInitializer)
-{
+	if (InUsage == ENiagaraScriptUsage::ParticleUpdateScript && Usage == ENiagaraScriptUsage::ParticleSpawnScriptInterpolated)
+	{
+		return true;
+	}
 
-}
+	if (InUsage == ENiagaraScriptUsage::EmitterSpawnScript && Usage == ENiagaraScriptUsage::SystemSpawnScript)
+	{
+		return true;
+	}
 
+	if (InUsage == ENiagaraScriptUsage::EmitterUpdateScript && Usage == ENiagaraScriptUsage::SystemUpdateScript)
+	{
+		return true;
+	}
+
+	return false;
+}
 
 void UNiagaraScript::Serialize(FArchive& Ar)
 {
@@ -63,8 +87,18 @@ void UNiagaraScript::Serialize(FArchive& Ar)
 	Ar.UsingCustomVersion(FNiagaraCustomVersion::GUID);		// only changes version if not loading
 	const int32 NiagaraVer = Ar.CustomVer(FNiagaraCustomVersion::GUID);
 
-	bool IsValidShaderScript = Usage != ENiagaraScriptUsage::Module && Usage != ENiagaraScriptUsage::Function && Usage != ENiagaraScriptUsage::DynamicInput 
-		&& (NiagaraVer<FNiagaraCustomVersion::NiagaraShaderMapCooking2 || (Usage != ENiagaraScriptUsage::SystemSpawnScript && Usage != ENiagaraScriptUsage::SystemUpdateScript));
+	bool IsValidShaderScript = false;
+	if (NiagaraVer < FNiagaraCustomVersion::DontCompileGPUWhenNotNeeded)
+	{
+		IsValidShaderScript = Usage != ENiagaraScriptUsage::Module && Usage != ENiagaraScriptUsage::Function && Usage != ENiagaraScriptUsage::DynamicInput
+			&& (NiagaraVer < FNiagaraCustomVersion::NiagaraShaderMapCooking2 || (Usage != ENiagaraScriptUsage::SystemSpawnScript && Usage != ENiagaraScriptUsage::SystemUpdateScript))
+			&& (NiagaraVer < FNiagaraCustomVersion::NiagaraCombinedGPUSpawnUpdate || (Usage != ENiagaraScriptUsage::ParticleUpdateScript && Usage != ENiagaraScriptUsage::EmitterSpawnScript && Usage != ENiagaraScriptUsage::EmitterUpdateScript));
+	}
+	else
+	{
+		IsValidShaderScript = CanBeRunOnGpu();
+	}
+
 	if ( (!Ar.IsLoading() && IsValidShaderScript)		// saving shader maps only for particle sim and spawn scripts
 		|| (Ar.IsLoading() && NiagaraVer >= FNiagaraCustomVersion::NiagaraShaderMaps && (NiagaraVer < FNiagaraCustomVersion::NiagaraShaderMapCooking || IsValidShaderScript))  // load only if we know shader map is presen
 		)
@@ -90,11 +124,8 @@ void UNiagaraScript::PostLoad()
 		if (Usage != ENiagaraScriptUsage::Module && Usage != ENiagaraScriptUsage::Function && Usage != ENiagaraScriptUsage::DynamicInput)
 		{
 			bNeedsRecompile = true;
-			ByteCode.Empty();
-#if WITH_EDITORONLY_DATA
-			LastCompileStatus = ENiagaraScriptCompileStatus::NCS_Unknown;
-#endif
-			UE_LOG(LogNiagara, Warning, TEXT("Niagara script is out of date and requires recompile to be used! %s"), *GetFullName());
+			InvalidateScript();
+			UE_LOG(LogNiagara, Log, TEXT("Niagara script is out of date and requires recompile to be used! %s"), *GetFullName());
 		}
 	}
 
@@ -121,9 +152,8 @@ void UNiagaraScript::PostLoad()
 			if (ChangeId.IsValid() == false || (Source != nullptr && !Source->IsSynchronized(ChangeId)))
 			{
 				bNeedsRecompile = true;
-				ByteCode.Empty();
-				LastCompileStatus = ENiagaraScriptCompileStatus::NCS_Unknown;
-				UE_LOG(LogNiagara, Warning, TEXT("Niagara script is out of date with source graph requires recompile to be used! %s"), *GetFullName());
+				InvalidateScript();//clear out the script.
+				UE_LOG(LogNiagara, Log, TEXT("Niagara script is out of date with source graph requires recompile to be used! %s"), *GetFullName());
 			}
 		}
 #endif
@@ -133,24 +163,27 @@ void UNiagaraScript::PostLoad()
 //	bNeedsRecompile = true;
 #if WITH_EDITORONLY_DATA
 	CacheResourceShadersForRendering(false, bNeedsRecompile);
+#endif
+#if STATS
 	GenerateStatScopeIDs();
 #endif
 
 }
 
 
-#if WITH_EDITOR
 
+#if STATS
 void UNiagaraScript::GenerateStatScopeIDs()
 {
 	StatScopesIDs.Empty();
-#if STATS
 	for (FNiagaraStatScope& StatScope : StatScopes)
 	{
 		StatScopesIDs.Add(FDynamicStats::CreateStatId<FStatGroup_STATGROUP_NiagaraDetailed>(StatScope.FriendlyName.ToString()));
 	}
-#endif
 }
+#endif
+
+#if WITH_EDITOR
 
 void UNiagaraScript::PostEditChangeProperty(FPropertyChangedEvent& PropertyChangedEvent)
 {
@@ -180,6 +213,19 @@ void UNiagaraScript::MarkScriptAndSourceDesynchronized()
 	}
 }
 
+void UNiagaraScript::SetByteCode(const TArray<uint8>& InByteCode)
+{
+	ByteCode = InByteCode;
+	if (InByteCode.Num() == 0)
+	{
+		LastCompileStatus = ENiagaraScriptCompileStatus::NCS_Error;
+	}
+	else
+	{
+		LastCompileStatus = ENiagaraScriptCompileStatus::NCS_UpToDate;
+	}
+}
+
 
 UNiagaraScript* UNiagaraScript::MakeRecursiveDeepCopy(UObject* DestOuter, TMap<const UObject*, UObject*>& ExistingConversions) const
 {
@@ -252,44 +298,63 @@ void UNiagaraScript::SubsumeExternalDependencies(TMap<const UObject*, UObject*>&
 	Source->SubsumeExternalDependencies(ExistingConversions);
 }
 
-ENiagaraScriptCompileStatus UNiagaraScript::Compile(FString& OutGraphLevelErrorMessages)
+ENiagaraScriptCompileStatus UNiagaraScript::Compile(FString& OutGraphLevelErrorMessages, bool bForce)
 {
 	if (Source == nullptr)
 		return ENiagaraScriptCompileStatus::NCS_Error;
 
+	if (Source->GetChangeID() == ChangeId && ChangeId.IsValid() && !bForce)
+	{
+		UE_LOG(LogNiagara, Log, TEXT("Skipping compilation due to changeid's being in sync already.. path: %s change-id: %s"), *GetFullName(), *ChangeId.ToString())
+		return LastCompileStatus;
+	}
+
+	float StartTime = FPlatformTime::Seconds();
 	ENiagaraScriptCompileStatus CPUCompileStatus = Source->Compile(this, OutGraphLevelErrorMessages);
 	CacheResourceShadersForRendering(false, true);
+	float DeltaTime = FPlatformTime::Seconds() - StartTime;
+	UE_LOG(LogNiagara, Log, TEXT("Compiling %s took %f sec."), *GetFullName(), DeltaTime);
 
 	return CPUCompileStatus;
 }
+
 #endif
 
 
+void UNiagaraScript::OnCompilationComplete()
+{
+#if WITH_EDITORONLY_DATA
+	FNiagaraSystemUpdateContext(this, true);
+#endif
+}
 
 #if WITH_EDITOR
 void UNiagaraScript::BeginCacheForCookedPlatformData(const ITargetPlatform *TargetPlatform)
 {
-	TArray<FName> DesiredShaderFormats;
-	TargetPlatform->GetAllTargetedShaderFormats(DesiredShaderFormats);
-
-	FNiagaraScript** CachedScriptResourceForPlatformPtr = CachedScriptResourcesForCooking.Find(TargetPlatform);
-
-	if (CachedScriptResourceForPlatformPtr == nullptr)
+	if (CanBeRunOnGpu())
 	{
-		CachedScriptResourcesForCooking.Add(TargetPlatform);
-		FNiagaraScript* &CachedScriptResourceForPlatform = *CachedScriptResourcesForCooking.Find(TargetPlatform);
+		TArray<FName> DesiredShaderFormats;
+		TargetPlatform->GetAllTargetedShaderFormats(DesiredShaderFormats);
 
-		if (DesiredShaderFormats.Num())
+		FNiagaraScript** CachedScriptResourceForPlatformPtr = CachedScriptResourcesForCooking.Find(TargetPlatform);
+
+		if (CachedScriptResourceForPlatformPtr == nullptr)
 		{
-			// Cache for all the shader formats that the cooking target requires
-			for (int32 FormatIndex = 0; FormatIndex < DesiredShaderFormats.Num(); FormatIndex++)
+			CachedScriptResourcesForCooking.Add(TargetPlatform);
+			FNiagaraScript* &CachedScriptResourceForPlatform = *CachedScriptResourcesForCooking.Find(TargetPlatform);
+
+			if (DesiredShaderFormats.Num())
 			{
-				if (DesiredShaderFormats[FormatIndex] != TEXT("PCD3D_SM4"))	// TODO: remove once we get rid of SM4 globally
+				// Cache for all the shader formats that the cooking target requires
+				for (int32 FormatIndex = 0; FormatIndex < DesiredShaderFormats.Num(); FormatIndex++)
 				{
-					const EShaderPlatform LegacyShaderPlatform = ShaderFormatToLegacyShaderPlatform(DesiredShaderFormats[FormatIndex]);
+					if (DesiredShaderFormats[FormatIndex] != TEXT("PCD3D_SM4"))	// TODO: remove once we get rid of SM4 globally
+					{
+						const EShaderPlatform LegacyShaderPlatform = ShaderFormatToLegacyShaderPlatform(DesiredShaderFormats[FormatIndex]);
 
-					// Begin caching shaders for the target platform and store the FNiagaraScript being compiled into CachedScriptResourcesForCooking
-					CacheResourceShadersForCooking(LegacyShaderPlatform, CachedScriptResourceForPlatform);
+						// Begin caching shaders for the target platform and store the FNiagaraScript being compiled into CachedScriptResourcesForCooking
+						CacheResourceShadersForCooking(LegacyShaderPlatform, CachedScriptResourceForPlatform);
+					}
 				}
 			}
 		}
@@ -299,23 +364,27 @@ void UNiagaraScript::BeginCacheForCookedPlatformData(const ITargetPlatform *Targ
 
 void UNiagaraScript::CacheResourceShadersForCooking(EShaderPlatform ShaderPlatform, FNiagaraScript* &OutCachedResource)
 {
-	if (Usage != ENiagaraScriptUsage::Function && Usage != ENiagaraScriptUsage::Module && Usage != ENiagaraScriptUsage::DynamicInput)
+	if (CanBeRunOnGpu())
 	{
-		FNiagaraScript *ResourceToCache = nullptr;
-		ERHIFeatureLevel::Type TargetFeatureLevel = GetMaxSupportedFeatureLevel(ShaderPlatform);
+		// spawn and update are combined on GPU, so we only compile spawn scripts
+		if (Usage == ENiagaraScriptUsage::ParticleSpawnScript || Usage == ENiagaraScriptUsage::ParticleSpawnScriptInterpolated)
+		{
+			FNiagaraScript *ResourceToCache = nullptr;
+			ERHIFeatureLevel::Type TargetFeatureLevel = GetMaxSupportedFeatureLevel(ShaderPlatform);
 
-		FNiagaraScript* NewResource = AllocateResource();
-		NewResource->SetScript(this, (ERHIFeatureLevel::Type)TargetFeatureLevel, UniqueID, GetName());
-		ResourceToCache = NewResource;
+			FNiagaraScript* NewResource = AllocateResource();
+			NewResource->SetScript(this, (ERHIFeatureLevel::Type)TargetFeatureLevel, UniqueID, GetName());
+			ResourceToCache = NewResource;
 
-		check(ResourceToCache);
+			check(ResourceToCache);
 
-		CacheShadersForResources(ShaderPlatform, ResourceToCache, false, false, true);
+			CacheShadersForResources(ShaderPlatform, ResourceToCache, false, false, true);
 
-		INiagaraModule NiagaraModule = FModuleManager::GetModuleChecked<INiagaraModule>(TEXT("Niagara"));
-		NiagaraModule.ProcessShaderCompilationQueue();
+			INiagaraModule NiagaraModule = FModuleManager::GetModuleChecked<INiagaraModule>(TEXT("Niagara"));
+			NiagaraModule.ProcessShaderCompilationQueue();
 
-		OutCachedResource = ResourceToCache;
+			OutCachedResource = ResourceToCache;
+		}
 	}
 }
 
@@ -323,26 +392,27 @@ void UNiagaraScript::CacheResourceShadersForCooking(EShaderPlatform ShaderPlatfo
 
 void UNiagaraScript::CacheShadersForResources(EShaderPlatform ShaderPlatform, FNiagaraScript *ResourceToCache, bool bApplyCompletedShaderMapForRendering, bool bForceRecompile, bool bCooking)
 {
-	const bool bSuccess = ResourceToCache->CacheShaders(ShaderPlatform, bApplyCompletedShaderMapForRendering, bForceRecompile, bCooking);
-
-#if defined(NIAGARA_SCRIPT_COMPILE_LOGGING_MEDIUM)
-	if (!bSuccess)
+	if (CanBeRunOnGpu())
 	{
-		UE_LOG(LogNiagara, Warning, TEXT("Failed to compile Niagara shader %s for platform %s."),
-			*GetPathName(),
-			*LegacyShaderPlatformToShaderFormat(ShaderPlatform).ToString());
+		const bool bSuccess = ResourceToCache->CacheShaders(ShaderPlatform, bApplyCompletedShaderMapForRendering, bForceRecompile, bCooking);
 
-		const TArray<FString>& CompileErrors = ResourceToCache->GetCompileErrors();
-		for (int32 ErrorIndex = 0; ErrorIndex < CompileErrors.Num(); ErrorIndex++)
+#if defined(NIAGARA_SCRIPT_COMPILE_LOGGING_MEDIUM)
+		if (!bSuccess)
 		{
-			UE_LOG(LogNiagara, Warning, TEXT("	%s"), *CompileErrors[ErrorIndex]);
+			UE_LOG(LogNiagara, Warning, TEXT("Failed to compile Niagara shader %s for platform %s."),
+				*GetPathName(),
+				*LegacyShaderPlatformToShaderFormat(ShaderPlatform).ToString());
+
+			const TArray<FString>& CompileErrors = ResourceToCache->GetCompileErrors();
+			for (int32 ErrorIndex = 0; ErrorIndex < CompileErrors.Num(); ErrorIndex++)
+			{
+				UE_LOG(LogNiagara, Warning, TEXT("	%s"), *CompileErrors[ErrorIndex]);
+			}
 		}
-	}
 #endif
+	}
 }
 
-
-
 void UNiagaraScript::CacheResourceShadersForRendering(bool bRegenerateId, bool bForceRecompile)
 {
 	if (bRegenerateId)
@@ -381,6 +451,59 @@ void UNiagaraScript::CacheResourceShadersForRendering(bool bRegenerateId, bool b
 
 #endif
 
+
+NIAGARA_API bool UNiagaraScript::IsScriptCompilationPending(bool bGPUScript) const
+{
+	if (bGPUScript)
+	{
+		FNiagaraShader *Shader = ScriptResource.GetShaderGameThread();
+		if (Shader)
+		{
+			return false;
+		}
+		return !ScriptResource.IsCompilationFinished();
+	}
+	else
+	{
+		return ByteCode.Num() == 0 && (LastCompileStatus == ENiagaraScriptCompileStatus::NCS_BeingCreated || LastCompileStatus == ENiagaraScriptCompileStatus::NCS_Unknown);
+	}
+}
+
+NIAGARA_API void UNiagaraScript::InvalidateScript()
+{
+	ByteCode.Empty();
+	ScriptResource.Invalidate();
+	LastCompileStatus = ENiagaraScriptCompileStatus::NCS_Unknown;
+#if WITH_EDITORONLY_DATA
+	InvalidateChangeID(); // Identify that we have no idea if we're synchronized or not.
+#endif
+}
+
+
+NIAGARA_API bool UNiagaraScript::DidScriptCompilationSucceed(bool bGPUScript) const
+{
+	if (bGPUScript)
+	{
+		FNiagaraShader *Shader = ScriptResource.GetShaderGameThread();
+		if (Shader)
+		{
+			return true;
+		}
+
+		if (ScriptResource.IsCompilationFinished())
+		{
+			// If we failed compilation, it would be finished and Shader would be null.
+			return false;
+		}
+	}
+	else
+	{
+		return ByteCode.Num() != 0;
+	}
+
+	return false;
+}
+
 void SerializeNiagaraShaderMaps(const TMap<const ITargetPlatform*, FNiagaraScript*>* PlatformScriptResourcesToSavePtr, FArchive& Ar, FNiagaraScript &OutLoadedResource)
 {
 //	SCOPED_LOADTIMER(SerializeInlineShaderMaps);
@@ -440,7 +563,9 @@ void ProcessSerializedShaderMaps(UNiagaraScript* Owner, FNiagaraScript LoadedRes
 		}
 
 		OutScriptResourcesLoaded[LoadedFeatureLevel]->SetShaderMap(LoadedShaderMap);
+		LoadedResource.SetDatainterfaceBufferDescriptors(LoadedResource.GetShaderGameThread()->GetDIBufferDescriptors());
 	}
+
 }
 
 
@@ -464,6 +589,34 @@ TArray<ENiagaraScriptUsage> UNiagaraScript::GetSupportedUsageContexts() const
 	return Supported;
 }
 
+bool UNiagaraScript::CanBeRunOnGpu()const
+{
+	if (UNiagaraEmitter* Emitter = GetTypedOuter<UNiagaraEmitter>())
+	{
+		if (Emitter->SimTarget == ENiagaraSimTarget::CPUSim)
+		{
+			return false;
+		}
+
+		if (!IsParticleSpawnScript())
+		{
+			return false;
+		}
+
+		for (const FNiagaraScriptDataInterfaceInfo& InterfaceInfo : DataInterfaceInfo)
+		{
+			if (!InterfaceInfo.DataInterface || !InterfaceInfo.DataInterface->CanExecuteOnTarget(ENiagaraSimTarget::GPUComputeSim))
+			{
+				return false;
+			}
+		}
+
+		return true;
+	}
+
+	return false;
+}
+
 bool UNiagaraScript::UsesCollection(const UNiagaraParameterCollection* Collection)const
 {
 	return ParameterCollections.FindByPredicate([&](const UNiagaraParameterCollection* CheckCollection)
diff --git a/Engine/Plugins/FX/Niagara/Source/Niagara/Private/NiagaraScriptExecutionContext.cpp b/Engine/Plugins/FX/Niagara/Source/Niagara/Private/NiagaraScriptExecutionContext.cpp
--- a/Engine/Plugins/FX/Niagara/Source/Niagara/Private/NiagaraScriptExecutionContext.cpp
+++ b/Engine/Plugins/FX/Niagara/Source/Niagara/Private/NiagaraScriptExecutionContext.cpp
@@ -27,49 +27,187 @@ FNiagaraScriptExecutionParameterStore& FNiagaraScriptExecutionParameterStore::op
 	return *this;
 }
 
-void FNiagaraScriptExecutionParameterStore::Init(UNiagaraScript* Script)
+uint32 OffsetAlign(uint32 SrcOffset, uint32 Size)
 {
-	//TEMPORARTY:
-	//We should replace the storage on the script with an FNiagaraParameterStore also so we can just copy that over here. Though that is an even bigger refactor job so this is a convenient place to break that work up.
-
-	bParametersDirty = true;
-	bInterfacesDirty = true;
-	Empty();
+	uint32 OffsetRemaining = UNIFORM_BUFFER_STRUCT_ALIGNMENT - (SrcOffset % UNIFORM_BUFFER_STRUCT_ALIGNMENT);
+	if (Size <= OffsetRemaining)
+	{
+		return SrcOffset;
+	}
+	else
+	{
+		return Align(SrcOffset, UNIFORM_BUFFER_STRUCT_ALIGNMENT);
+	}
+}
 
-	//Here we add the current frame parameters.
-	for (FNiagaraVariable& Param : Script->Parameters.Parameters)
+void FNiagaraScriptExecutionParameterStore::GenerateLayoutInfoInternal(TArray<FNiagaraScriptExecutionParameterStore::FPaddingInfo>& Members, uint32& NextMemberOffset, const UStruct* InSrcStruct, uint32 InSrcOffset)
+{
+	// Now insert an appropriate data member into the mix...
+	if (InSrcStruct == FNiagaraTypeDefinition::GetBoolStruct() || InSrcStruct == FNiagaraTypeDefinition::GetIntStruct())
 	{
-		AddParameter(Param, false);
+		//Members.Emplace(*Property->GetName(), TEXT(""), InSrcOffset, (EUniformBufferBaseType)TUniformBufferTypeInfo<uint32>::BaseType, EShaderPrecisionModifier::Float, TUniformBufferTypeInfo<uint32>::NumRows, TUniformBufferTypeInfo<uint32>::NumColumns, TUniformBufferTypeInfo<uint32>::NumElements, TUniformBufferTypeInfo<uint32>::GetStruct());
+		Members.Emplace(InSrcOffset, Align(NextMemberOffset, TUniformBufferTypeInfo<uint32>::Alignment), (TUniformBufferTypeInfo<uint32>::NumRows * TUniformBufferTypeInfo<uint32>::NumColumns) * sizeof(uint32));
+		InSrcOffset += sizeof(uint32);
+		NextMemberOffset = Members[Members.Num() - 1].DestOffset + Members[Members.Num() - 1].Size;
 	}
-	
-	//Add previous frame values if we're interpolated spawn.
-	bool bIsInterpolatedSpawn = Script->IsInterpolatedParticleSpawnScript();
-	if (bIsInterpolatedSpawn)
+	else if (InSrcStruct == FNiagaraTypeDefinition::GetFloatStruct())
 	{
-		for (FNiagaraVariable& Param : Script->Parameters.Parameters)
+		//Members.Emplace(*Property->GetName(), TEXT(""), InSrcOffset, (EUniformBufferBaseType)TUniformBufferTypeInfo < float > ::BaseType, EShaderPrecisionModifier::Float, TUniformBufferTypeInfo<float>::NumRows, TUniformBufferTypeInfo<float>::NumColumns, TUniformBufferTypeInfo<float>::NumElements, TUniformBufferTypeInfo<float>::GetStruct());
+		Members.Emplace(InSrcOffset, Align(NextMemberOffset, TUniformBufferTypeInfo<float>::Alignment), (TUniformBufferTypeInfo<float>::NumRows * TUniformBufferTypeInfo<float>::NumColumns) * sizeof(float));
+		InSrcOffset += sizeof(float);
+		NextMemberOffset = Members[Members.Num() - 1].DestOffset + Members[Members.Num() - 1].Size;
+	}
+	else if (InSrcStruct == FNiagaraTypeDefinition::GetVec2Struct())
+	{
+		//Members.Emplace(*Property->GetName(), TEXT(""), InSrcOffset, (EUniformBufferBaseType)TUniformBufferTypeInfo < FVector2D > ::BaseType, EShaderPrecisionModifier::Float, TUniformBufferTypeInfo<FVector2D>::NumRows, TUniformBufferTypeInfo<FVector2D>::NumColumns, TUniformBufferTypeInfo<FVector2D>::NumElements, TUniformBufferTypeInfo<FVector2D>::GetStruct());
+		uint32 StructFinalSize = (TUniformBufferTypeInfo<FVector2D>::NumRows * TUniformBufferTypeInfo<FVector2D>::NumColumns) * sizeof(float);
+		Members.Emplace(InSrcOffset, OffsetAlign(NextMemberOffset, StructFinalSize),StructFinalSize);
+		InSrcOffset += sizeof(FVector2D);
+		NextMemberOffset = Members[Members.Num() - 1].DestOffset + Members[Members.Num() - 1].Size;
+	}
+	else if (InSrcStruct == FNiagaraTypeDefinition::GetVec3Struct())
+	{
+		//Members.Emplace(*Property->GetName(), TEXT(""), InSrcOffset, (EUniformBufferBaseType)TUniformBufferTypeInfo < FVector > ::BaseType, EShaderPrecisionModifier::Float, TUniformBufferTypeInfo<FVector>::NumRows, TUniformBufferTypeInfo<FVector>::NumColumns, TUniformBufferTypeInfo<FVector>::NumElements, TUniformBufferTypeInfo<FVector>::GetStruct());
+		uint32 StructFinalSize = (TUniformBufferTypeInfo<FVector>::NumRows * TUniformBufferTypeInfo<FVector>::NumColumns) * sizeof(float);
+		Members.Emplace(InSrcOffset, OffsetAlign(NextMemberOffset, StructFinalSize), StructFinalSize);
+		InSrcOffset += sizeof(FVector);
+		NextMemberOffset = Members[Members.Num() - 1].DestOffset + Members[Members.Num() - 1].Size;
+	}
+	else if (InSrcStruct == FNiagaraTypeDefinition::GetVec4Struct() || InSrcStruct == FNiagaraTypeDefinition::GetColorStruct())
+	{
+		//Members.Emplace(*Property->GetName(), TEXT(""), InSrcOffset, (EUniformBufferBaseType)TUniformBufferTypeInfo < FVector4 > ::BaseType, EShaderPrecisionModifier::Float, TUniformBufferTypeInfo<FVector4>::NumRows, TUniformBufferTypeInfo<FVector4>::NumColumns, TUniformBufferTypeInfo<FVector4>::NumElements, TUniformBufferTypeInfo<FVector4>::GetStruct());
+		Members.Emplace(InSrcOffset, Align(NextMemberOffset, TUniformBufferTypeInfo<FVector4>::Alignment), (TUniformBufferTypeInfo<FVector4>::NumRows * TUniformBufferTypeInfo<FVector4>::NumColumns) * sizeof(float));
+		InSrcOffset += sizeof(FVector4);
+		NextMemberOffset = Members[Members.Num() - 1].DestOffset + Members[Members.Num() - 1].Size;
+	}
+	else if (InSrcStruct == FNiagaraTypeDefinition::GetMatrix4Struct())
+	{
+		//Members.Emplace(*Property->GetName(), TEXT(""), InSrcOffset, (EUniformBufferBaseType)TUniformBufferTypeInfo < FMatrix > ::BaseType, EShaderPrecisionModifier::Float, TUniformBufferTypeInfo<FMatrix>::NumRows, TUniformBufferTypeInfo<FMatrix>::NumColumns, TUniformBufferTypeInfo<FMatrix>::NumElements, TUniformBufferTypeInfo<FMatrix>::GetStruct());
+		Members.Emplace(InSrcOffset, Align(NextMemberOffset, TUniformBufferTypeInfo<FMatrix>::Alignment), (TUniformBufferTypeInfo<FMatrix>::NumRows * TUniformBufferTypeInfo<FMatrix>::NumColumns) * sizeof(float));
+		InSrcOffset += sizeof(FMatrix);
+		NextMemberOffset = Members[Members.Num() - 1].DestOffset + Members[Members.Num() - 1].Size;
+	}
+	else
+	{
+		NextMemberOffset = Align(NextMemberOffset, UNIFORM_BUFFER_STRUCT_ALIGNMENT); // New structs should be aligned to the head..
+
+		for (TFieldIterator<UProperty> PropertyIt(InSrcStruct, EFieldIteratorFlags::IncludeSuper); PropertyIt; ++PropertyIt)
 		{
-			FNiagaraVariable PrevParam(Param.GetType(), FName(*(TEXT("PREV__") + Param.GetName().ToString())));
-			AddParameter(PrevParam, false);
+			UProperty* Property = *PropertyIt;
+			const UStruct* Struct = nullptr;
+
+			// First determine what struct type we're dealing with...
+			if (Property->IsA(UFloatProperty::StaticClass()))
+			{
+				Struct = FNiagaraTypeDefinition::GetFloatStruct();
+			}
+			else if (Property->IsA(UIntProperty::StaticClass()))
+			{
+				Struct = FNiagaraTypeDefinition::GetIntStruct();
+			}
+			else if (Property->IsA(UBoolProperty::StaticClass()))
+			{
+				Struct = FNiagaraTypeDefinition::GetBoolStruct();
+			}
+			//Should be able to support double easily enough
+			else if (UStructProperty* StructProp = CastChecked<UStructProperty>(Property))
+			{
+				Struct = StructProp->Struct;
+			}
+			else
+			{
+				check(false);
+			}
+
+			GenerateLayoutInfoInternal(Members, NextMemberOffset, Struct, InSrcOffset);
 		}
 	}
-	
-	ParameterSize = ParameterData.Num();
-	if (bIsInterpolatedSpawn)
+
+}
+
+void FNiagaraScriptExecutionParameterStore::AddPaddedParamSize(const FNiagaraTypeDefinition& InParamType, uint32 InOffset)
+{
+	if (InParamType.IsDataInterface())
+	{
+		return;
+	}
+
+	uint32 NextMemberOffset = 0;
+	if (PaddingInfo.Num() != 0)
 	{
-		CopyCurrToPrev();
+		NextMemberOffset = PaddingInfo[PaddingInfo.Num() - 1].DestOffset + PaddingInfo[PaddingInfo.Num() - 1].Size;
 	}
+	GenerateLayoutInfoInternal(PaddingInfo, NextMemberOffset, InParamType.GetScriptStruct(), InOffset);
 
-	//Internal constants
-	for (FNiagaraVariable& InternalVar : Script->InternalParameters.Parameters)
+	if (PaddingInfo.Num() != 0)
+	{
+		NextMemberOffset = PaddingInfo[PaddingInfo.Num() - 1].DestOffset + PaddingInfo[PaddingInfo.Num() - 1].Size;
+		PaddedParameterSize = Align(NextMemberOffset, UNIFORM_BUFFER_STRUCT_ALIGNMENT);
+	}
+	else
 	{
-		AddParameter(InternalVar, false);
+		PaddedParameterSize = 0;
 	}
+}
+
+
+void FNiagaraScriptExecutionParameterStore::Init(UNiagaraScript* Script, ENiagaraSimTarget SimTarget)
+{
+	//TEMPORARTY:
+	//We should replace the storage on the script with an FNiagaraParameterStore also so we can just copy that over here. Though that is an even bigger refactor job so this is a convenient place to break that work up.
 
-	for (FNiagaraScriptDataInterfaceInfo& Info : Script->DataInterfaceInfo)
+	Empty();
+	PaddedParameterSize = 0;
+	AddScriptParams(Script, SimTarget);
+}
+
+void FNiagaraScriptExecutionParameterStore::AddScriptParams(UNiagaraScript* Script, ENiagaraSimTarget SimTarget)
+{
+	if (Script != nullptr)
 	{
-		FNiagaraVariable Var(FNiagaraTypeDefinition(Info.DataInterface->GetClass()), Info.Name);
-		AddParameter(Var, false);
-		SetDataInterface(Info.DataInterface, IndexOf(Var));
+		//Here we add the current frame parameters.
+		for (FNiagaraVariable& Param : Script->Parameters.Parameters)
+		{
+			AddParameter(Param, false);
+		}
+
+#if WITH_EDITORONLY_DATA
+		Name = *FString::Printf(TEXT("ScriptExecParamStore %s %p"), *Script->GetFullName(), this);
+#endif
+
+		//Add previous frame values if we're interpolated spawn.
+		bool bIsInterpolatedSpawn = Script->IsInterpolatedParticleSpawnScript();
+
+		if (bIsInterpolatedSpawn)
+		{
+			for (FNiagaraVariable& Param : Script->Parameters.Parameters)
+			{
+				FNiagaraVariable PrevParam(Param.GetType(), FName(*(TEXT("PREV__") + Param.GetName().ToString())));
+				AddParameter(PrevParam, false);
+			}
+		}
+
+		ParameterSize = GetParameterDataArray().Num();
+		if (bIsInterpolatedSpawn)
+		{
+			CopyCurrToPrev();
+		}
+
+		//Internal constants - only needed for non-GPU sim
+		if (SimTarget != ENiagaraSimTarget::GPUComputeSim)
+		{
+			for (FNiagaraVariable& InternalVar : Script->InternalParameters.Parameters)
+			{
+				AddParameter(InternalVar, false);
+			}
+		}
+
+		for (FNiagaraScriptDataInterfaceInfo& Info : Script->DataInterfaceInfo)
+		{
+			FNiagaraVariable Var(Info.Type, Info.Name);
+			AddParameter(Var, false);
+			SetDataInterface(Info.DataInterface, IndexOf(Var));
+		}
 	}
 }
 
@@ -81,32 +219,53 @@ void FNiagaraScriptExecutionParameterStore::CopyCurrToPrev()
 	FMemory::Memcpy(GetParameterData_Internal(ParamStart), GetParameterData(0), ParamStart);
 }
 
+
+void FNiagaraScriptExecutionParameterStore::CopyParameterDataToPaddedBuffer(uint8* InTargetBuffer, uint32 InTargetBufferSizeInBytes)
+{
+	check(InTargetBufferSizeInBytes >= PaddedParameterSize);
+	FMemory::Memzero(InTargetBuffer, InTargetBufferSizeInBytes);
+	const uint8* SrcData = GetParameterDataArray().GetData();
+	for (int32 i = 0; i < PaddingInfo.Num(); i++)
+	{
+		check((PaddingInfo[i].DestOffset + PaddingInfo[i].Size) <= InTargetBufferSizeInBytes);
+		check((PaddingInfo[i].SrcOffset + PaddingInfo[i].Size) <= (uint32)GetParameterDataArray().Num());
+		FMemory::Memcpy(InTargetBuffer + PaddingInfo[i].DestOffset, SrcData + PaddingInfo[i].SrcOffset, PaddingInfo[i].Size);
+	}
+}
+
+uint32 FNiagaraScriptExecutionContext::TickCounter = 0;
+
 FNiagaraScriptExecutionContext::FNiagaraScriptExecutionContext()
 	: Script(nullptr)
 {
 
 }
 
-uint32 FNiagaraScriptExecutionContext::TickCounter = 0;
+FNiagaraScriptExecutionContext::~FNiagaraScriptExecutionContext()
+{
+}
 
 bool FNiagaraScriptExecutionContext::Init(UNiagaraScript* InScript, ENiagaraSimTarget InTarget)
 {
 	Script = InScript;
 
-	Parameters.Init(Script);
+	Parameters.Init(Script, InTarget);
+	
 
 	return true;//TODO: Error cases?
 }
 
-bool FNiagaraScriptExecutionContext::Tick(FNiagaraSystemInstance* ParentSystemInstance)
+bool FNiagaraScriptExecutionContext::Tick(FNiagaraSystemInstance* ParentSystemInstance, ENiagaraSimTarget SimTarget)
 {
 	if (Script)//TODO: Remove. Script can only be null for system instances that currently don't have their script exec context set up correctly.
 	{
 		//Bind data interfaces if needed.
-		if (Parameters.bInterfacesDirty)
+		if (Parameters.GetInterfacesDirty())
 		{
+			// UE_LOG(LogNiagara, Log, TEXT("Updating data interfaces for script %s"), *Script->GetFullName());
+
 			FunctionTable.Empty();
-			const TArray<UNiagaraDataInterface*> DataInterfaces = GetDataInterfaces();
+			const TArray<UNiagaraDataInterface*>& DataInterfaces = GetDataInterfaces();
 			// We must make sure that the data interfaces match up between the original script values and our overrides...
 			if (Script->DataInterfaceInfo.Num() != DataInterfaces.Num())
 			{
@@ -120,10 +279,10 @@ bool FNiagaraScriptExecutionContext::Tick(FNiagaraSystemInstance* ParentSystemIn
 				DataInterfaceInstDataTable.SetNumZeroed(Script->NumUserPtrs);
 				for (int32 i = 0; i < Script->DataInterfaceInfo.Num(); i++)
 				{
+					UNiagaraDataInterface* Interface = DataInterfaces[i];
 					int32 UserPtrIdx = Script->DataInterfaceInfo[i].UserPtrIdx;
 					if (UserPtrIdx != INDEX_NONE)
 					{
-						UNiagaraDataInterface* Interface = DataInterfaces[i];
 						void* InstData = ParentSystemInstance->FindDataInterfaceInstanceData(Interface);
 						DataInterfaceInstDataTable[UserPtrIdx] = InstData;
 					}
@@ -198,6 +357,7 @@ bool FNiagaraScriptExecutionContext::Execute(uint32 NumInstances, TArray<FNiagar
 
 	DataSetMetaTable.Reset();
 
+	bool bRegisterSetupCompleted = true;
 	{
 		SCOPE_CYCLE_COUNTER(STAT_NiagaraSimRegisterSetup);
 		for (FNiagaraDataSetExecutionInfo& DataSetInfo : DataSetInfos)
@@ -210,32 +370,29 @@ bool FNiagaraScriptExecutionContext::Execute(uint32 NumInstances, TArray<FNiagar
 				DataSetInfo.DataSet->Allocate(NumInstances);
 				DataSetInfo.DataSet->SetNumInstances(NumInstances);
 			}
-			for (const FNiagaraVariable &Var : DataSetInfo.DataSet->GetVariables())
-			{
-				DataSetInfo.DataSet->AppendToRegisterTable(Var, InputRegisters, NumInputRegisters, OutputRegisters, NumOutputRegisters, DataSetInfo.StartInstance);
-			}
+			
+			bRegisterSetupCompleted &= DataSetInfo.DataSet->AppendToRegisterTable(InputRegisters, NumInputRegisters, OutputRegisters, NumOutputRegisters, DataSetInfo.StartInstance);
 		}
 	}
 
-	VectorVM::Exec(
-		Script->ByteCode.GetData(),
-		InputRegisters,
-		NumInputRegisters,
-		OutputRegisters,
-		NumOutputRegisters,
-		Parameters.GetParameterDataArray().GetData(),
-		DataSetMetaTable,
-		FunctionTable.GetData(),
-		DataInterfaceInstDataTable.GetData(),
-		NumInstances
+	if (bRegisterSetupCompleted)
+	{
+		VectorVM::Exec(
+			Script->GetByteCode().GetData(),
+			InputRegisters,
+			NumInputRegisters,
+			OutputRegisters,
+			NumOutputRegisters,
+			Parameters.GetParameterDataArray().GetData(),
+			DataSetMetaTable,
+			FunctionTable.GetData(),
+			DataInterfaceInstDataTable.GetData(),
+			NumInstances
 #if STATS
-#if WITH_EDITORONLY_DATA
-		, Script->GetStatScopeIDs()
-#else
-		, TArray<TStatId>()
-#endif
+			, Script->GetStatScopeIDs()
 #endif
-	);
+		);
+	}
 
 	// Tell the datasets we wrote how many instances were actually written.
 	for (int Idx = 0; Idx < DataSetInfos.Num(); Idx++)
@@ -253,6 +410,7 @@ bool FNiagaraScriptExecutionContext::Execute(uint32 NumInstances, TArray<FNiagar
 		DataSetInfos[0].DataSet->Dump(Script->GetDebuggerInfo().DebugFrame, true);
 		Script->GetDebuggerInfo().bRequestDebugFrame = false;
 		Script->GetDebuggerInfo().DebugFrameLastWriteId = TickCounter;
+		Script->GetDebuggerInfo().DebugParameters = Parameters;
 	}
 #endif
 	return true;//TODO: Error cases?
@@ -260,10 +418,10 @@ bool FNiagaraScriptExecutionContext::Execute(uint32 NumInstances, TArray<FNiagar
 
 void FNiagaraScriptExecutionContext::DirtyDataInterfaces()
 {
-	Parameters.bInterfacesDirty = true;
+	Parameters.MarkInterfacesDirty();
 }
 
 bool FNiagaraScriptExecutionContext::CanExecute()const
 {
-	return Script && Script->ByteCode.Num() > 0;
+	return Script && Script->GetByteCode().Num() > 0;
 }
\ No newline at end of file
diff --git a/Engine/Plugins/FX/Niagara/Source/Niagara/Private/NiagaraSpriteRendererProperties.cpp b/Engine/Plugins/FX/Niagara/Source/Niagara/Private/NiagaraSpriteRendererProperties.cpp
--- a/Engine/Plugins/FX/Niagara/Source/Niagara/Private/NiagaraSpriteRendererProperties.cpp
+++ b/Engine/Plugins/FX/Niagara/Source/Niagara/Private/NiagaraSpriteRendererProperties.cpp
@@ -8,11 +8,16 @@
 #define LOCTEXT_NAMESPACE "UNiagaraSpriteRendererProperties"
 
 UNiagaraSpriteRendererProperties::UNiagaraSpriteRendererProperties()
-	: SubImageSize(1.0f, 1.0f)
-	, Alignment(ENiagaraSpriteAlignment::Unaligned)
+	: Alignment(ENiagaraSpriteAlignment::Unaligned)
 	, FacingMode(ENiagaraSpriteFacingMode::FaceCamera)
 	, CustomFacingVectorMask(EForceInit::ForceInitToZero)
+	, PivotInUVSpace(0.5f, 0.5f)
 	, SortMode(ENiagaraSortMode::SortNone)
+	, SubImageSize(1.0f, 1.0f)
+	, bSubImageBlend(false)
+	, bRemoveHMDRollInVR(false)
+	, MinFacingCameraBlendDistance(0.0f)
+	, MaxFacingCameraBlendDistance(0.0f)
 {
 }
 
@@ -57,6 +62,9 @@ const TArray<FNiagaraVariable>& UNiagaraSpriteRendererProperties::GetOptionalAtt
 		Attrs.Add(SYS_PARAM_PARTICLES_SPRITE_ALIGNMENT);
 		Attrs.Add(SYS_PARAM_PARTICLES_SUB_IMAGE_INDEX);
 		Attrs.Add(SYS_PARAM_PARTICLES_DYNAMIC_MATERIAL_PARAM);
+		Attrs.Add(SYS_PARAM_PARTICLES_CAMERA_OFFSET);
+		Attrs.Add(SYS_PARAM_PARTICLES_UV_SCALE);
+		Attrs.Add(SYS_PARAM_PARTICLES_MATERIAL_RANDOM);
 	}
 
 	return Attrs;
