Copying Orion-Staging @ 3748653 (Orion/Dev-General @ 3722124) to //UE4/Main

diff --git a/Engine/Plugins/FX/Niagara/Source/Niagara/Private/NiagaraSystem.cpp b/Engine/Plugins/FX/Niagara/Source/Niagara/Private/NiagaraSystem.cpp
--- a/Engine/Plugins/FX/Niagara/Source/Niagara/Private/NiagaraSystem.cpp
+++ b/Engine/Plugins/FX/Niagara/Source/Niagara/Private/NiagaraSystem.cpp
@@ -9,6 +9,14 @@
 #include "NiagaraCustomVersion.h"
 #include "NiagaraModule.h"
 #include "ModuleManager.h"
+#include "NiagaraEmitter.h"
+#include "Package.h"
+#include "NiagaraEmitterHandle.h"
+#include "AssetData.h"
+#include "NiagaraStats.h"
+
+DECLARE_CYCLE_STAT(TEXT("Niagara - System - CompileScript"), STAT_Niagara_System_CompileScript, STATGROUP_Niagara);
+DECLARE_CYCLE_STAT(TEXT("Niagara - System - CompileScript_ResetAfter"), STAT_Niagara_System_CompileScriptResetAfter, STATGROUP_Niagara);
 
 // UNiagaraSystemCategory::UNiagaraSystemCategory(const FObjectInitializer& ObjectInitializer)
 // 	: Super(ObjectInitializer)
@@ -19,6 +27,11 @@
 
 UNiagaraSystem::UNiagaraSystem(const FObjectInitializer& ObjectInitializer)
 : Super(ObjectInitializer)
+, ExposedParameters(this)
+, bAutoDeactivate(true)
+#if WITH_EDITORONLY_DATA
+, bIsolateEnabled(false)
+#endif
 {
 }
 
@@ -26,7 +39,7 @@ void UNiagaraSystem::PostInitProperties()
 {
 	Super::PostInitProperties();
 #if WITH_EDITORONLY_DATA
-	bAutoImportChangedEmitters = true;
+	ThumbnailImageOutOfDate = true;
 #endif
 	if (HasAnyFlags(RF_ClassDefaultObject | RF_NeedLoad) == false)
 	{
@@ -44,6 +57,55 @@ void UNiagaraSystem::PostInitProperties()
 	}
 }
 
+bool UNiagaraSystem::IsLooping() const
+{ 
+	return false; 
+} //sckime todo fix this!
+
+bool UNiagaraSystem::UsesCollection(const UNiagaraParameterCollection* Collection)const
+{
+	if (SystemSpawnScript->UsesCollection(Collection) ||
+		SystemSpawnScriptSolo->UsesCollection(Collection) ||
+		SystemUpdateScript->UsesCollection(Collection) ||
+		SystemUpdateScriptSolo->UsesCollection(Collection))
+	{
+		return true;
+	}
+
+	for (const FNiagaraEmitterHandle& EmitterHandle : GetEmitterHandles())
+	{
+		if (EmitterHandle.GetInstance()->UsesCollection(Collection))
+		{
+			return true;
+		}
+	}
+
+	return false;
+}
+
+#if WITH_EDITORONLY_DATA
+bool UNiagaraSystem::UsesScript(const UNiagaraScript* Script)const
+{
+	if (SystemSpawnScript == Script ||
+		SystemSpawnScriptSolo == Script ||
+		SystemUpdateScript == Script ||
+		SystemUpdateScriptSolo == Script)
+	{
+		return true;
+	}
+
+	for (FNiagaraEmitterHandle EmitterHandle : GetEmitterHandles())
+	{
+		if ((EmitterHandle.GetSource() && EmitterHandle.GetSource()->UsesScript(Script)) || (EmitterHandle.GetInstance() && EmitterHandle.GetInstance()->UsesScript(Script)))
+		{
+			return true;
+		}
+	}
+	
+	return false;
+}
+#endif
+
 void UNiagaraSystem::Serialize(FArchive& Ar)
 {
 	Super::Serialize(Ar);
@@ -51,6 +113,16 @@ void UNiagaraSystem::Serialize(FArchive& Ar)
 	Ar.UsingCustomVersion(FNiagaraCustomVersion::GUID);
 }
 
+#if WITH_EDITOR
+void UNiagaraSystem::PostEditChangeProperty(struct FPropertyChangedEvent& PropertyChangedEvent)
+{
+	Super::PostEditChangeProperty(PropertyChangedEvent);
+	FNiagaraSystemUpdateContext(this, true);
+
+	ThumbnailImageOutOfDate = true;
+}
+#endif 
+
 void UNiagaraSystem::PostLoad()
 {
 	Super::PostLoad();
@@ -66,7 +138,26 @@ void UNiagaraSystem::PostLoad()
 	if (NiagaraVer < FNiagaraCustomVersion::LatestVersion)
 	{
 		bNeedsRecompile = true;
-	}		
+	}
+
+	// Previously added emitters didn't have their stand alone and public flags cleared so
+	// they 'leak' into the system package.  Clear the flags here so they can be collected
+	// during the next save.
+	UPackage* PackageOuter = Cast<UPackage>(GetOuter());
+	if (PackageOuter != nullptr && HasAnyFlags(RF_Public | RF_Standalone))
+	{
+		TArray<UObject*> ObjectsInPackage;
+		GetObjectsWithOuter((UObject*)PackageOuter, ObjectsInPackage);
+		for (UObject* ObjectInPackage : ObjectsInPackage)
+		{
+			UNiagaraEmitter* Emitter = Cast<UNiagaraEmitter>(ObjectInPackage);
+			if (Emitter != nullptr)
+			{
+				Emitter->ConditionalPostLoad();
+				Emitter->ClearFlags(RF_Standalone | RF_Public);
+			}
+		}
+	}
 	
 	// We will be using this later potentially, so make sure that it's postload is already up-to-date.
 	if (SystemSpawnScript)
@@ -78,9 +169,9 @@ void UNiagaraSystem::PostLoad()
 		SystemUpdateScript->ConditionalPostLoad();
 	}
 
- 	if (NiagaraVer < FNiagaraCustomVersion::SystemEmitterScriptSimulations)
- 	{
- 		bNeedsRecompile = true;
+	if (NiagaraVer < FNiagaraCustomVersion::SystemEmitterScriptSimulations)
+	{
+		bNeedsRecompile = true;
 		SystemSpawnScriptSolo = NewObject<UNiagaraScript>(this, "SystemSpawnScriptSolo", RF_Transactional);
 		SystemSpawnScriptSolo->SetUsage(ENiagaraScriptUsage::SystemSpawnScript);
 
@@ -97,7 +188,7 @@ void UNiagaraSystem::PostLoad()
 			SystemUpdateScriptSolo->SetSource(SystemUpdateScript->GetSource());
 		}
 #endif
- 	}
+	}
 
 	if (bNeedsRecompile)
 	{
@@ -125,29 +216,39 @@ void UNiagaraSystem::PostLoad()
 		SystemUpdateScript->SetUsage(ENiagaraScriptUsage::SystemUpdateScript);
 	}
 
-	if (bNeedsRecompile)
+	//TODO: This causes a crash becuase the script source ptr is null? Fix
+	//For existing emitters before the lifecylce rework, ensure they have the system lifecycle module.
+ 	if (NiagaraVer < FNiagaraCustomVersion::LifeCycleRework)
+ 	{
+ 		/*UNiagaraScriptSourceBase* SystemScriptSource = SystemUpdateScript->GetSource();
+ 		if (SystemScriptSource)
+ 		{
+			bool bFoundModule;
+			if (SystemScriptSource->AddModuleIfMissing(TEXT("/Niagara/Modules/System/SystemLifeCycle.SystemLifeCycle"), ENiagaraScriptUsage::SystemUpdateScript, bFoundModule))
+			{
+				bNeedsRecompile = true;
+			}
+ 		}*/
+ 	}
+
+	for (FNiagaraEmitterHandle& EmitterHandle : EmitterHandles)
 	{
-		Compile();
+		EmitterHandle.ConditionalPostLoad();
+		if (EmitterHandle.IsSynchronizedWithSource() == false)
+		{
+			INiagaraModule::FMergeEmitterResults Results = MergeChangesForEmitterHandle(EmitterHandle);
+			bNeedsRecompile |= Results.bSucceeded && Results.bModifiedGraph;
+		}
 	}
 
-	// Note that we do not call ResynchronizeAllHandles() or CheckForUpdates here.
-	// There are multiple reasons:
-	// 1) You can run into issues with the Linker table associated with this package when doing the deep copy during load. 
-	// 2) We'd have to make sure that all of the referenced nodes that were in the old version of the graph have already been loaded and cleaned out.
-	// Better to just CheckForUpdates anytime this System is used for the first time.
+	if (bNeedsRecompile)
+	{
+		Compile(false);
+	}
 #endif
 }
 
 #if WITH_EDITORONLY_DATA
-bool UNiagaraSystem::CheckForUpdates()
-{
-	// An emitter may have changed since last load, let's refresh if possible.
-	if (bAutoImportChangedEmitters)
-	{
-		return ResynchronizeAllHandles();
-	}
-	return false;
-}
 
 UObject* UNiagaraSystem::GetEditorData()
 {
@@ -164,50 +265,80 @@ void UNiagaraSystem::SetEditorData(UObject* InEditorData)
 	EditorData = InEditorData;
 }
 
-bool UNiagaraSystem::ResynchronizeAllHandles()
+INiagaraModule::FMergeEmitterResults UNiagaraSystem::MergeChangesForEmitterHandle(FNiagaraEmitterHandle& EmitterHandle)
 {
-	uint32 HandlexIdx = 0;
-	bool bAnyRefreshed = false;
-	for (FNiagaraEmitterHandle& Handle : EmitterHandles)
+	INiagaraModule::FMergeEmitterResults Results = EmitterHandle.MergeSourceChanges();
+	if (Results.bSucceeded)
 	{
-		const UNiagaraEmitter* SrcPrps = Handle.GetSource();
-		const UNiagaraEmitter* InstancePrps = Handle.GetInstance();
-
-		// Check to see if the source and handle are synched. If they aren't go ahead and resynch.
-		if (!Handle.IsSynchronizedWithSource())
+		UNiagaraEmitter* Instance = EmitterHandle.GetInstance();
+		RefreshSystemParametersFromEmitter(EmitterHandle);
+		if (Instance->bInterpolatedSpawning)
 		{
-			// Todo.. eliminate synchronization until we have rewritten it...
-#if 0
-			if (Handle.RefreshFromSource())
-			{
-				check(SrcPrps->ChangeId == Handle.GetInstance()->ChangeId);
-				UE_LOG(LogNiagara, Log, TEXT("Refreshed System %s emitter %d from %s"), *GetPathName(), HandlexIdx, *SrcPrps->GetPathName());
-				bAnyRefreshed = true;
-			}
-#endif
+			Instance->UpdateScriptProps.Script->RapidIterationParameters.CopyParametersTo(
+				Instance->SpawnScriptProps.Script->RapidIterationParameters, false);
 		}
-		HandlexIdx++;
 	}
-
-	if (bAnyRefreshed)
+	else
 	{
-		Compile();
+		UE_LOG(LogNiagara, Warning, TEXT("Failed to merge changes for base emitter.  System: %s  Emitter: %s  Error Message: %s"), 
+			*GetPathName(), *EmitterHandle.GetName().ToString(), *Results.GetErrorMessagesString());
 	}
-
-	return bAnyRefreshed;
+	return Results;
 }
 
-bool UNiagaraSystem::ReferencesSourceEmitter(UNiagaraEmitter* Emitter)
+bool UNiagaraSystem::ReferencesSourceEmitter(UNiagaraEmitter& Emitter)
 {
 	for (FNiagaraEmitterHandle& Handle : EmitterHandles)
 	{
-		if (Emitter == Handle.GetSource())
+		if (&Emitter == Handle.GetSource())
 		{
 			return true;
 		}
 	}
 	return false;
 }
+
+void UNiagaraSystem::UpdateFromEmitterChanges(UNiagaraEmitter& ChangedSourceEmitter)
+{
+	bool bNeedsCompile = false;
+	for(FNiagaraEmitterHandle& EmitterHandle : EmitterHandles)
+	{
+		if (EmitterHandle.GetSource() == &ChangedSourceEmitter)
+		{
+			INiagaraModule::FMergeEmitterResults Results = MergeChangesForEmitterHandle(EmitterHandle);
+			bNeedsCompile |= Results.bSucceeded && Results.bModifiedGraph;
+		}
+	}
+
+	if (bNeedsCompile)
+	{
+		Compile(false);
+	}
+}
+
+void UNiagaraSystem::RefreshSystemParametersFromEmitter(const FNiagaraEmitterHandle& EmitterHandle)
+{
+	if (ensureMsgf(EmitterHandles.ContainsByPredicate([=](const FNiagaraEmitterHandle& OwnedEmitterHandle) { return OwnedEmitterHandle.GetId() == EmitterHandle.GetId(); }),
+		TEXT("Can't refresh parameters from an emitter handle this system doesn't own.")))
+	{
+		EmitterHandle.GetInstance()->EmitterSpawnScriptProps.Script->RapidIterationParameters.CopyParametersTo(SystemSpawnScript->RapidIterationParameters, false);
+		EmitterHandle.GetInstance()->EmitterSpawnScriptProps.Script->RapidIterationParameters.CopyParametersTo(SystemSpawnScriptSolo->RapidIterationParameters, false);
+		EmitterHandle.GetInstance()->EmitterUpdateScriptProps.Script->RapidIterationParameters.CopyParametersTo(SystemUpdateScript->RapidIterationParameters, false);
+		EmitterHandle.GetInstance()->EmitterUpdateScriptProps.Script->RapidIterationParameters.CopyParametersTo(SystemUpdateScriptSolo->RapidIterationParameters, false);
+	}
+}
+
+void UNiagaraSystem::RemoveSystemParametersForEmitter(const FNiagaraEmitterHandle& EmitterHandle)
+{
+	if (ensureMsgf(EmitterHandles.ContainsByPredicate([=](const FNiagaraEmitterHandle& OwnedEmitterHandle) { return OwnedEmitterHandle.GetId() == EmitterHandle.GetId(); }),
+		TEXT("Can't remove parameters for an emitter handle this system doesn't own.")))
+	{
+		EmitterHandle.GetInstance()->EmitterSpawnScriptProps.Script->RapidIterationParameters.RemoveParameters(SystemSpawnScript->RapidIterationParameters);
+		EmitterHandle.GetInstance()->EmitterSpawnScriptProps.Script->RapidIterationParameters.RemoveParameters(SystemSpawnScriptSolo->RapidIterationParameters);
+		EmitterHandle.GetInstance()->EmitterUpdateScriptProps.Script->RapidIterationParameters.RemoveParameters(SystemUpdateScript->RapidIterationParameters);
+		EmitterHandle.GetInstance()->EmitterUpdateScriptProps.Script->RapidIterationParameters.RemoveParameters(SystemUpdateScriptSolo->RapidIterationParameters);
+	}
+}
 #endif
 
 
@@ -216,14 +347,50 @@ const TArray<FNiagaraEmitterHandle>& UNiagaraSystem::GetEmitterHandles()
 	return EmitterHandles;
 }
 
+const TArray<FNiagaraEmitterHandle>& UNiagaraSystem::GetEmitterHandles()const
+{
+	return EmitterHandles;
+}
+
+bool UNiagaraSystem::IsReadyToRun() const
+{
+	if (!SystemSpawnScript || !SystemSpawnScriptSolo || !SystemUpdateScript || !SystemUpdateScriptSolo)
+	{
+		return false;
+	}
+
+	if (SystemSpawnScript->IsScriptCompilationPending(false) || SystemSpawnScriptSolo->IsScriptCompilationPending(false) || 
+		SystemUpdateScript->IsScriptCompilationPending(false) || SystemUpdateScriptSolo->IsScriptCompilationPending(false))
+	{
+		return false;
+	}
+
+	for (const FNiagaraEmitterHandle& Handle : EmitterHandles)
+	{
+		if (!Handle.GetInstance()->IsReadyToRun())
+		{
+			return false;
+		}
+	}
+	return true;
+}
 
 bool UNiagaraSystem::IsValid()const
 {
-	if (!SystemSpawnScript->IsValid() || !SystemSpawnScriptSolo->IsValid() || !SystemUpdateScript->IsValid() || !SystemUpdateScriptSolo->IsValid())
+	if (!SystemSpawnScript || !SystemSpawnScriptSolo || !SystemUpdateScript || !SystemUpdateScriptSolo)
+	{
+		return false;
+	}
+
+	if ((!SystemSpawnScript->IsScriptCompilationPending(false) && !SystemSpawnScript->DidScriptCompilationSucceed(false)) ||
+		(!SystemSpawnScriptSolo->IsScriptCompilationPending(false) && !SystemSpawnScriptSolo->DidScriptCompilationSucceed(false)) ||
+		(!SystemUpdateScript->IsScriptCompilationPending(false) && !SystemUpdateScript->DidScriptCompilationSucceed(false)) ||
+		(!SystemUpdateScriptSolo->IsScriptCompilationPending(false) && !SystemUpdateScriptSolo->DidScriptCompilationSucceed(false)))
 	{
 		return false;
 	}
 
+
 	for (const FNiagaraEmitterHandle& Handle : EmitterHandles)
 	{
 		if (!Handle.GetInstance()->IsValid())
@@ -236,10 +403,11 @@ bool UNiagaraSystem::IsValid()const
 }
 #if WITH_EDITORONLY_DATA
 
-FNiagaraEmitterHandle UNiagaraSystem::AddEmitterHandle(const UNiagaraEmitter& SourceEmitter, FName EmitterName)
+FNiagaraEmitterHandle UNiagaraSystem::AddEmitterHandle(UNiagaraEmitter& SourceEmitter, FName EmitterName)
 {
 	FNiagaraEmitterHandle EmitterHandle(SourceEmitter, EmitterName, *this);
 	EmitterHandles.Add(EmitterHandle);
+	RefreshSystemParametersFromEmitter(EmitterHandle);
 	return EmitterHandle;
 }
 
@@ -247,6 +415,7 @@ FNiagaraEmitterHandle UNiagaraSystem::AddEmitterHandleWithoutCopying(UNiagaraEmi
 {
 	FNiagaraEmitterHandle EmitterHandle(Emitter);
 	EmitterHandles.Add(EmitterHandle);
+	RefreshSystemParametersFromEmitter(EmitterHandle);
 	return EmitterHandle;
 }
 
@@ -254,12 +423,14 @@ FNiagaraEmitterHandle UNiagaraSystem::DuplicateEmitterHandle(const FNiagaraEmitt
 {
 	FNiagaraEmitterHandle EmitterHandle(EmitterHandleToDuplicate, EmitterName, *this);
 	EmitterHandles.Add(EmitterHandle);
+	RefreshSystemParametersFromEmitter(EmitterHandle);
 	return EmitterHandle;
 }
-#endif
 
 void UNiagaraSystem::RemoveEmitterHandle(const FNiagaraEmitterHandle& EmitterHandleToDelete)
 {
+	UNiagaraEmitter* EditableEmitter = EmitterHandleToDelete.GetInstance();
+	RemoveSystemParametersForEmitter(EmitterHandleToDelete);
 	auto RemovePredicate = [&](const FNiagaraEmitterHandle& EmitterHandle) { return EmitterHandle.GetId() == EmitterHandleToDelete.GetId(); };
 	EmitterHandles.RemoveAll(RemovePredicate);
 }
@@ -272,6 +443,8 @@ void UNiagaraSystem::RemoveEmitterHandlesById(const TSet<FGuid>& HandlesToRemove
 	};
 	EmitterHandles.RemoveAll(RemovePredicate);
 }
+#endif
+
 
 UNiagaraScript* UNiagaraSystem::GetSystemSpawnScript(bool bSolo)
 {
@@ -284,18 +457,23 @@ UNiagaraScript* UNiagaraSystem::GetSystemUpdateScript(bool bSolo)
 }
 
 #if WITH_EDITORONLY_DATA
-bool UNiagaraSystem::GetAutoImportChangedEmitters() const
+
+bool UNiagaraSystem::GetIsolateEnabled() const
 {
-	return bAutoImportChangedEmitters;
+	return bIsolateEnabled;
 }
 
-void UNiagaraSystem::SetAutoImportChangedEmitters(bool bShouldImport)
+void UNiagaraSystem::SetIsolateEnabled(bool bIsolate)
 {
-	bAutoImportChangedEmitters = bShouldImport;
+	bIsolateEnabled = bIsolate;
 }
 
-void UNiagaraSystem::CompileScripts(TArray<ENiagaraScriptCompileStatus>& OutScriptStatuses, TArray<FString>& OutGraphLevelErrorMessages, TArray<FString>& PathNames, TArray<UNiagaraScript*>& Scripts)
+
+void UNiagaraSystem::CompileScripts(TArray<ENiagaraScriptCompileStatus>& OutScriptStatuses, TArray<FString>& OutGraphLevelErrorMessages, TArray<FString>& PathNames, TArray<UNiagaraScript*>& Scripts, bool bForce)
 {
+	float StartTime = FPlatformTime::Seconds();
+
+	SCOPE_CYCLE_COUNTER(STAT_Niagara_System_CompileScript);
 	PathNames.Reset();
 	OutScriptStatuses.Reset();
 	OutGraphLevelErrorMessages.Reset();
@@ -303,20 +481,31 @@ void UNiagaraSystem::CompileScripts(TArray<ENiagaraScriptCompileStatus>& OutScri
 	TArray<FNiagaraVariable> ExposedVars;
 
 	//Compile all emitters
+	bool bAnyUnsynchronized = false;
 	for (FNiagaraEmitterHandle Handle : EmitterHandles)
 	{
+		TArray<ENiagaraScriptCompileStatus> ScriptStatuses;
 		TArray<FString> ScriptErrors;
+
 		UNiagaraScriptSourceBase* GraphSource = Handle.GetInstance()->GraphSource;
 		check(!GraphSource->IsPreCompiled());
 		GraphSource->PreCompile(Handle.GetInstance());
 
-		Handle.GetInstance()->CompileScripts(OutScriptStatuses, ScriptErrors, PathNames, Scripts);
+		if (!Handle.GetInstance()->AreAllScriptAndSourcesSynchronized())
+		{
+			bAnyUnsynchronized = true;
+		}
+		Handle.GetInstance()->CompileScripts(ScriptStatuses, ScriptErrors, PathNames, Scripts, bForce);
 
 		GraphSource->GatherPreCompiledVariables(TEXT("User"), ExposedVars);
 		GraphSource->PostCompile();
+
+		OutScriptStatuses.Append(ScriptStatuses);
 		OutGraphLevelErrorMessages.Append(ScriptErrors);
 	}
 
+	bool bForceSystems = bForce || bAnyUnsynchronized;
+
 	check(SystemSpawnScript->GetSource() == SystemUpdateScript->GetSource());
 	check(SystemSpawnScript->GetSource()->IsPreCompiled() == false);
 
@@ -325,28 +514,28 @@ void UNiagaraSystem::CompileScripts(TArray<ENiagaraScriptCompileStatus>& OutScri
 
 	FString ErrorMsg;
 	check((int32)EScriptCompileIndices::SpawnScript == 0);
-	OutScriptStatuses.Add(SystemSpawnScript->Compile(ErrorMsg));
+	OutScriptStatuses.Add(SystemSpawnScript->Compile(ErrorMsg, bForceSystems));
 	OutGraphLevelErrorMessages.Add(ErrorMsg); // EScriptCompileIndices::SpawnScript
 	Scripts.Add(SystemSpawnScript);
 	PathNames.Add(SystemSpawnScript->GetPathName());
 
 	ErrorMsg.Empty();
 	check((int32)EScriptCompileIndices::UpdateScript == 1);
-	OutScriptStatuses.Add(SystemUpdateScript->Compile(ErrorMsg));
+	OutScriptStatuses.Add(SystemUpdateScript->Compile(ErrorMsg, bForceSystems));
 	OutGraphLevelErrorMessages.Add(ErrorMsg); // EScriptCompileIndices::SpawnScript
 	Scripts.Add(SystemUpdateScript);
 	PathNames.Add(SystemUpdateScript->GetPathName());
 
 	ErrorMsg.Empty();
 	check((int32)EScriptCompileIndices::SpawnScript == 0);
-	OutScriptStatuses.Add(SystemSpawnScriptSolo->Compile(ErrorMsg));
+	OutScriptStatuses.Add(SystemSpawnScriptSolo->Compile(ErrorMsg, bForceSystems));
 	OutGraphLevelErrorMessages.Add(ErrorMsg); // EScriptCompileIndices::SpawnScript
 	Scripts.Add(SystemSpawnScriptSolo);
 	PathNames.Add(SystemSpawnScriptSolo->GetPathName());
 
 	ErrorMsg.Empty();
 	check((int32)EScriptCompileIndices::UpdateScript == 1);
-	OutScriptStatuses.Add(SystemUpdateScriptSolo->Compile(ErrorMsg));
+	OutScriptStatuses.Add(SystemUpdateScriptSolo->Compile(ErrorMsg, bForceSystems));
 	OutGraphLevelErrorMessages.Add(ErrorMsg); // EScriptCompileIndices::SpawnScript
 	Scripts.Add(SystemUpdateScriptSolo);
 	PathNames.Add(SystemUpdateScriptSolo->GetPathName());
@@ -362,28 +551,28 @@ void UNiagaraSystem::CompileScripts(TArray<ENiagaraScriptCompileStatus>& OutScri
 	{
 		if (OriginalVars.Contains(ExposedVars[i]) == false)
 		{
+			// Just in case it wasn't added previously..
 			ExposedParameters.AddParameter(ExposedVars[i]);
 		}
 	}
-	for (int32 i = 0; i < OriginalVars.Num(); i++)
-	{
-		if (ExposedVars.Contains(OriginalVars[i]) == false)
-		{
-			ExposedParameters.RemoveParameter(OriginalVars[i]);
-		}
-	}
-
-	FNiagaraSystemUpdateContext UpdateCtx(this, true);
 
-	INiagaraModule& NiagaraModule = FModuleManager::LoadModuleChecked<INiagaraModule>("Niagara");
-	NiagaraModule.DestroyAllSystemSimulations(this);
+	{
+		SCOPE_CYCLE_COUNTER(STAT_Niagara_System_CompileScriptResetAfter);
 
+		FNiagaraSystemUpdateContext UpdateCtx(this, true);
+	}
 	//I assume this was added to recompile emitters too but I'm doing that above. We need to rework this compilation flow.
 	//Compile();
+
+
+	UE_LOG(LogNiagara, Log, TEXT("Compiling System %s took %f sec."), *GetFullName(), FPlatformTime::Seconds() - StartTime);
 }
 
-void UNiagaraSystem::Compile()
+void UNiagaraSystem::Compile(bool bForce)
 {
+	float StartTime = FPlatformTime::Seconds();
+
+	bool bAnyUnsynchronized = false;
 	TArray<FNiagaraVariable> ExposedVars;
 	for (FNiagaraEmitterHandle Handle: EmitterHandles)
 	{
@@ -396,7 +585,11 @@ void UNiagaraSystem::Compile()
 		check(!GraphSource->IsPreCompiled());
 		GraphSource->PreCompile(Handle.GetInstance());
 
-		Handle.GetInstance()->CompileScripts(ScriptStatuses, ScriptErrors, PathNames, Scripts);
+		if (!Handle.GetInstance()->AreAllScriptAndSourcesSynchronized())
+		{
+			bAnyUnsynchronized = true;
+		}
+		Handle.GetInstance()->CompileScripts(ScriptStatuses, ScriptErrors, PathNames, Scripts, bForce);
 
 		GraphSource->GatherPreCompiledVariables(TEXT("User"), ExposedVars);
 
@@ -414,6 +607,7 @@ void UNiagaraSystem::Compile()
 			}
 		}
 	}
+	bool bForceSystems = bForce || bAnyUnsynchronized;
 
 	check(SystemSpawnScript->GetSource() == SystemUpdateScript->GetSource());
 	check(SystemSpawnScript->GetSource()->IsPreCompiled() == false);
@@ -422,7 +616,7 @@ void UNiagaraSystem::Compile()
 	SystemSpawnScript->GetSource()->GatherPreCompiledVariables(TEXT("User"), ExposedVars);
 
 	FString CompileErrors;
-	ENiagaraScriptCompileStatus CompileStatus = SystemSpawnScript->Compile(CompileErrors);
+	ENiagaraScriptCompileStatus CompileStatus = SystemSpawnScript->Compile(CompileErrors, bForceSystems);
 	if (CompileErrors.Len() != 0 || CompileStatus != ENiagaraScriptCompileStatus::NCS_UpToDate)
 	{
 		UE_LOG(LogNiagara, Warning, TEXT("System Spawn Script '%s', compile status: %d  errors: %s"), *SystemSpawnScript->GetPathName(), (int32)CompileStatus, *CompileErrors);
@@ -432,7 +626,7 @@ void UNiagaraSystem::Compile()
 		UE_LOG(LogNiagara, Log, TEXT("Script Spawn '%s', compile status: Success!"), *SystemSpawnScript->GetPathName());
 	}
 
-	CompileStatus = SystemUpdateScript->Compile(CompileErrors);
+	CompileStatus = SystemUpdateScript->Compile(CompileErrors, bForceSystems);
 	if (CompileErrors.Len() != 0 || CompileStatus != ENiagaraScriptCompileStatus::NCS_UpToDate)
 	{
 		UE_LOG(LogNiagara, Warning, TEXT("System Update Script '%s', compile status: %d  errors: %s"), *SystemUpdateScript->GetPathName(), (int32)CompileStatus, *CompileErrors);
@@ -444,7 +638,7 @@ void UNiagaraSystem::Compile()
 
 	if(SystemSpawnScriptSolo)
 	{
-		CompileStatus = SystemSpawnScriptSolo->Compile(CompileErrors);
+		CompileStatus = SystemSpawnScriptSolo->Compile(CompileErrors, bForceSystems);
 		if (CompileErrors.Len() != 0 || CompileStatus != ENiagaraScriptCompileStatus::NCS_UpToDate)
 		{
 			UE_LOG(LogNiagara, Warning, TEXT("System Spawn Solo Script '%s', compile status: %d  errors: %s"), *SystemSpawnScriptSolo->GetPathName(), (int32)CompileStatus, *CompileErrors);
@@ -457,7 +651,7 @@ void UNiagaraSystem::Compile()
 
 	if (SystemUpdateScriptSolo)
 	{
-		CompileStatus = SystemUpdateScriptSolo->Compile(CompileErrors);
+		CompileStatus = SystemUpdateScriptSolo->Compile(CompileErrors, bForceSystems);
 		if (CompileErrors.Len() != 0 || CompileStatus != ENiagaraScriptCompileStatus::NCS_UpToDate)
 		{
 			UE_LOG(LogNiagara, Warning, TEXT("System Update Solo Script '%s', compile status: %d  errors: %s"), *SystemUpdateScriptSolo->GetPathName(), (int32)CompileStatus, *CompileErrors);
@@ -484,10 +678,14 @@ void UNiagaraSystem::Compile()
 		}
 	}
 
-	FNiagaraSystemUpdateContext UpdateCtx(this, true);
+	{
+
+		SCOPE_CYCLE_COUNTER(STAT_Niagara_System_CompileScriptResetAfter);
+		FNiagaraSystemUpdateContext UpdateCtx(this, true);
+	}
+
+	UE_LOG(LogNiagara, Log, TEXT("Compiling System %s took %f sec."), *GetFullName(), FPlatformTime::Seconds() - StartTime);
 
-	INiagaraModule& NiagaraModule = FModuleManager::LoadModuleChecked<INiagaraModule>("Niagara");
-	NiagaraModule.DestroyAllSystemSimulations(this);
 }
 #endif
 
@@ -501,7 +699,7 @@ void UNiagaraSystem::InitEmitterSpawnAttributes()
 		for (int32 EmitterIdx = 0; EmitterIdx < EmitterHandles.Num(); ++EmitterIdx)
 		{
 			UNiagaraEmitter* Emitter = EmitterHandles[EmitterIdx].GetInstance();
-			FString EmitterName = Emitter->GetUniqueEmitterName();
+			FString EmitterName = Emitter->GetUniqueEmitterName() + TEXT(".");
 			if (Var.GetType() == SpawnInfoDef && Var.GetName().ToString().StartsWith(EmitterName))
 			{
 				EmitterSpawnAttributes[EmitterIdx].SpawnAttributes.AddUnique(Var.GetName());
@@ -513,7 +711,7 @@ void UNiagaraSystem::InitEmitterSpawnAttributes()
 		for (int32 EmitterIdx = 0; EmitterIdx < EmitterHandles.Num(); ++EmitterIdx)
 		{
 			UNiagaraEmitter* Emitter = EmitterHandles[EmitterIdx].GetInstance();
-			FString EmitterName = Emitter->GetUniqueEmitterName();
+			FString EmitterName = Emitter->GetUniqueEmitterName() + TEXT(".");
 			if (Var.GetType() == SpawnInfoDef && Var.GetName().ToString().StartsWith(EmitterName))
 			{
 				EmitterSpawnAttributes[EmitterIdx].SpawnAttributes.AddUnique(Var.GetName());
diff --git a/Engine/Plugins/FX/Niagara/Source/Niagara/Private/NiagaraSystemInstance.cpp b/Engine/Plugins/FX/Niagara/Source/Niagara/Private/NiagaraSystemInstance.cpp
--- a/Engine/Plugins/FX/Niagara/Source/Niagara/Private/NiagaraSystemInstance.cpp
+++ b/Engine/Plugins/FX/Niagara/Source/Niagara/Private/NiagaraSystemInstance.cpp
@@ -13,44 +13,39 @@
 #include "GameFramework/PlayerController.h"
 
 DECLARE_CYCLE_STAT(TEXT("Parallel Tick"), STAT_NiagaraParallelTick, STATGROUP_Niagara);
+DECLARE_CYCLE_STAT(TEXT("System Reset (GT)"), STAT_NiagaraSystemReset, STATGROUP_Niagara);
+DECLARE_CYCLE_STAT(TEXT("System Reinit (GT)"), STAT_NiagaraSystemReinit, STATGROUP_Niagara);
 
-
-/** Whether to allow emitters to tick in parallel */
-static int32 GbNiagaraParallelEmitterTick = 1;
-static FAutoConsoleVariableRef CVarParallelEmitterTick(
-	TEXT("niagara.ParallelEmitterTick"),
-	GbNiagaraParallelEmitterTick,
-	TEXT("Whether to tick individual emitters in an System in parallel"),
-	ECVF_Default
-	);
  
-
+/** Safety time to allow for the LastRenderTime coming back from the RT. */
+static float GLastRenderTimeSafetyBias = 0.1f;
+static FAutoConsoleVariableRef CVarLastRenderTimeSafetyBias(
+	TEXT("fx.LastRenderTimeSafetyBias"),
+	GLastRenderTimeSafetyBias,
+	TEXT("The time to bias the LastRenderTime value to allow for the delay from it being written by the RT."),
+	ECVF_Default
+);
 
 FNiagaraSystemInstance::FNiagaraSystemInstance(UNiagaraComponent* InComponent)
 	: SystemInstanceIndex(INDEX_NONE)
 	, Component(InComponent)
-	, SystemSimulation(nullptr)
 	, Age(0.0f)
 	, ID(FGuid::NewGuid())
 	, IDName(*ID.ToString())
+	, InstanceParameters(Component)
 	, bSolo(false)
 	, bForceSolo(false)
-	, bResetPending(false)
-	, bReInitPending(false)
-	, bError(false)
+	, bError(true)
 	, bDataInterfacesNeedInit(true)
-	, ExecutionState(ENiagaraExecutionState::Inactive)
+	, bHasTickingEmitters(true)
+	, ExecutionState(ENiagaraExecutionState::Complete)
 {
 	SystemBounds.Init();
 }
 
-void FNiagaraSystemInstance::Init(FNiagaraSystemSimulation* InSystemSimulation, bool bForceReset, bool bInForceSolo)
+void FNiagaraSystemInstance::Init(TSharedRef<FNiagaraSystemSimulation, ESPMode::ThreadSafe> InSystemSimulation, bool bInForceSolo)
 {
 	SystemSimulation = InSystemSimulation;
-// 	if (System != nullptr)NOOOOOPE!. This should be done via the editor or on post load of the system!
-// 	{
-// 		System->CheckForUpdates();
-// 	}
 
 	bError = false;
 	bForceSolo = bInForceSolo;
@@ -58,101 +53,169 @@ void FNiagaraSystemInstance::Init(FNiagaraSystemSimulation* InSystemSimulation,
 
 	//InstanceParameters = GetSystem()->GetInstanceParameters();
 	InitEmitters();
-	Reset(bForceReset ? EResetMode::ImmediateReInit : EResetMode::DeferredReInit);
-
+	Reset(EResetMode::ReInit);
 
+#if WITH_EDITORONLY_DATA
+	InstanceParameters.Name = *FString::Printf(TEXT("SystemInstance %p"), this);
+#endif
 	OnInitializedDelegate.Broadcast();
 }
 
 void FNiagaraSystemInstance::SetExecutionState(ENiagaraExecutionState InState)
 {
+	if (ExecutionState != InState)
+	{
+		const UEnum* EnumPtr = FNiagaraTypeDefinition::GetExecutionStateEnum();
+		UE_LOG(LogNiagara, Log, TEXT("System \"%s\" change state: %s to %s"), *GetSystem()->GetName(), *EnumPtr->GetNameStringByValue((int64)ExecutionState),
+			*EnumPtr->GetNameStringByValue((int64)InState));
+
+	}
 	ExecutionState = InState;
 }
 
 void FNiagaraSystemInstance::Activate(bool bReset /* = false */)
 {
-	if (GetSystem() && GetSystem()->IsValid())
+	if (GetSystem() && GetSystem()->IsValid() && IsReadyToRun())
 	{
-		if (ExecutionState != ENiagaraExecutionState::Active || bReset)
+		if (IsComplete() || bReset)
 		{
-			Reset(EResetMode::DeferredReset);
+			Reset(EResetMode::ResetAll);
+		}
+		else
+		{
+			Reset(EResetMode::ResetSystem);
 		}
 
 		ExecutionState = ENiagaraExecutionState::Active;
 	}
+	else
+	{
+		ExecutionState = ENiagaraExecutionState::Disabled;
+	}
 }
 
 void FNiagaraSystemInstance::Deactivate(bool bImmediate)
 {
-	ExecutionState = ENiagaraExecutionState::Dead;
+	if (IsComplete())
+	{
+		return;
+	}
 
 	if (bImmediate)
 	{
-		//Just kill this instance.
-		FNiagaraSystemSimulation* SystemSim = GetSystemSimulation();
-		SystemSim->RemoveInstance(this);
+		Complete();
 	}
-	
-	//TODO: pipe into simulation?
-	//Drive a bool param at system level which systems can read? Stop spawning at all when not active so we can clean up?
+	else
+	{
+		ExecutionState = ENiagaraExecutionState::Inactive;
+	}
+}
 
-	//Need to work out how we're going to clean up system instances and components.
-	//When all emitters are disabled and have no remaining particles?
+void FNiagaraSystemInstance::Complete()
+{
+	TSharedPtr<FNiagaraSystemSimulation, ESPMode::ThreadSafe> SystemSim = GetSystemSimulation();
+	SystemSim->RemoveInstance(this);
+
+	ExecutionState = ENiagaraExecutionState::Complete;
+
+	for (TSharedRef<FNiagaraEmitterInstance> Simulation : Emitters)
+	{
+		Simulation->HandleCompletion(true);
+	}
+
+	if (Component)
+	{
+		Component->OnSystemComplete();
+	}
+
+	OnCompleteDelegate.Broadcast(this);
 }
 
 void FNiagaraSystemInstance::Reset(FNiagaraSystemInstance::EResetMode Mode)
 {
-	FNiagaraSystemSimulation* SystemSim = GetSystemSimulation();
+	SCOPE_CYCLE_COUNTER(STAT_NiagaraSystemReset);
+
+	TSharedPtr<FNiagaraSystemSimulation, ESPMode::ThreadSafe> SystemSim = GetSystemSimulation();
+	if (!SystemSim.IsValid())
+	{
+		return;
+	}
+
+	Component->LastRenderTime = Component->GetWorld()->GetTimeSeconds();
+
 	SystemSim->RemoveInstance(this);
 	bPendingSpawn = true;
-	
-	//WIP: Currently tesing out having all resets be immediate. 
-	if (Mode == EResetMode::DeferredReset)
+
+	if (bError)
 	{
-		//bResetPending = true;
-		ResetInternal();
+		Mode = EResetMode::ReInit;
 	}
-	else if (Mode == EResetMode::ImmediateReset)
+		
+	if (Mode == EResetMode::ResetSystem)
 	{
-		ResetInternal();
+		ResetInternal(false);
 	}
-	else if (Mode == EResetMode::DeferredReInit)
+	else if (Mode == EResetMode::ResetAll)
 	{
-		DestroyDataInterfaceInstanceData();
-		//bReInitPending = true;
-		ReInitInternal();
+		ResetInternal(true);
 	}
-	else if (Mode == EResetMode::ImmediateReInit)
+	else if (Mode == EResetMode::ReInit)
 	{
 		DestroyDataInterfaceInstanceData();
 		ReInitInternal();
 	}
-
-	if (!bSolo && !bError)
-	{
-		SystemSim->AddInstance(this);
-	}
-	else if (bSolo && !bError)
+	
+	if (!bError)
 	{
-		SystemSim->ResetSolo(this);
+		if (bSolo)
+		{
+			SystemSim->ResetSolo(this);
+		}
+		else
+		{
+			SystemSim->AddInstance(this);
+		}
 	}
 }
 
-void FNiagaraSystemInstance::ResetInternal()
+void FNiagaraSystemInstance::ResetInternal(bool bResetSimulations)
 {
 	Age = 0;
-	if (GetSystem() == nullptr || Component == nullptr || bError)
+	UNiagaraSystem* System = GetSystem();
+	if (System == nullptr || Component == nullptr || bError)
 	{
 		return;
 	}
 
+#if WITH_EDITOR
+	if (Component->GetWorld() != nullptr && Component->GetWorld()->WorldType == EWorldType::Editor)
+	{
+		Component->SynchronizeWithSourceSystem();
+		Component->GetOverrideParameters().Tick();
+	}
+#endif
 
-	for (TSharedRef<FNiagaraEmitterInstance> Simulation : Emitters)
+	bool bAllReadyToRun = IsReadyToRun();
+
+	if (!bAllReadyToRun)
+	{
+		return;
+	}
+
+	if (!System->IsValid())
 	{
-		Simulation->ResetSimulation();
+		bError = true;		
+		UE_LOG(LogNiagara, Warning, TEXT("Failed to activate Niagara System due to invalid asset!"));
+		return;
 	}
 
-	bResetPending = false;
+	if (bResetSimulations)
+	{
+		for (TSharedRef<FNiagaraEmitterInstance> Simulation : Emitters)
+		{
+			Simulation->ResetSimulation();
+		}
+	}
 
 #if WITH_EDITOR
 	//UE_LOG(LogNiagara, Log, TEXT("OnResetInternal %p"), this);
@@ -160,33 +223,72 @@ void FNiagaraSystemInstance::ResetInternal()
 #endif
 }
 
+UNiagaraParameterCollectionInstance* FNiagaraSystemInstance::GetParameterCollectionInstance(UNiagaraParameterCollection* Collection)
+{
+	return SystemSimulation->GetParameterCollectionInstance(Collection);
+}
+
+bool FNiagaraSystemInstance::IsReadyToRun() const
+{
+	bool bAllReadyToRun = true;
+
+	UNiagaraSystem* System = GetSystem();
+
+	if (!System || !System->IsReadyToRun())
+	{
+		return false;
+	}
+
+	for (TSharedRef<FNiagaraEmitterInstance> Simulation : Emitters)
+	{
+		if (!Simulation->IsReadyToRun())
+		{
+			bAllReadyToRun = false;
+		}
+	}
+	return bAllReadyToRun;
+}
+
 void FNiagaraSystemInstance::ReInitInternal()
 {
+	SCOPE_CYCLE_COUNTER(STAT_NiagaraSystemReinit);
 	bool bPrevError = bError;
 	Age = 0;
-	if (GetSystem() == nullptr || Component == nullptr)
+	UNiagaraSystem* System = GetSystem();
+	if (System == nullptr || Component == nullptr)
 	{
 		return;
 	}
 	bError = false;
 
-
-	for (TSharedRef<FNiagaraEmitterInstance> Simulation : Emitters)
+#if WITH_EDITOR
+	if (Component->GetWorld() != nullptr && Component->GetWorld()->WorldType == EWorldType::Editor)
 	{
-		Simulation->ReInitSimulation();
+		Component->SynchronizeWithSourceSystem();
 	}
+#endif
 
-	for (TSharedRef<FNiagaraEmitterInstance> Simulation : Emitters)
+	bool bAllReadyToRun = IsReadyToRun();
+
+	if (!bAllReadyToRun)
 	{
-		Simulation->PostResetSimulation();
+		return;
 	}
 	
-	InstanceParameters.Empty();
-#if WITH_EDITOR
-	InstanceParameters = Component->GetInitialParameters();
-#endif
+	if (!System->IsValid())
+	{
+		bError = true;
+		UE_LOG(LogNiagara, Warning, TEXT("Failed to activate Niagara System due to invalid asset!"));
+		return;
+	}
 
+	//When re initializing, throw away old emitters and init new ones.
+	Emitters.Reset();
+	InitEmitters();
+	
+	InstanceParameters.Empty();
 	InstanceParameters.AddParameter(SYS_PARAM_ENGINE_POSITION);
+	InstanceParameters.AddParameter(SYS_PARAM_ENGINE_SCALE);
 	InstanceParameters.AddParameter(SYS_PARAM_ENGINE_VELOCITY);
 	InstanceParameters.AddParameter(SYS_PARAM_ENGINE_X_AXIS);
 	InstanceParameters.AddParameter(SYS_PARAM_ENGINE_Y_AXIS);
@@ -195,27 +297,37 @@ void FNiagaraSystemInstance::ReInitInternal()
 	InstanceParameters.AddParameter(SYS_PARAM_ENGINE_WORLD_TO_LOCAL);
 	InstanceParameters.AddParameter(SYS_PARAM_ENGINE_LOCAL_TO_WORLD_TRANSPOSED);
 	InstanceParameters.AddParameter(SYS_PARAM_ENGINE_WORLD_TO_LOCAL_TRANSPOSED);
+	InstanceParameters.AddParameter(SYS_PARAM_ENGINE_LOCAL_TO_WORLD_NO_SCALE);
+	InstanceParameters.AddParameter(SYS_PARAM_ENGINE_WORLD_TO_LOCAL_NO_SCALE);
 	InstanceParameters.AddParameter(SYS_PARAM_ENGINE_DELTA_TIME);
 	InstanceParameters.AddParameter(SYS_PARAM_ENGINE_INV_DELTA_TIME);
+	InstanceParameters.AddParameter(SYS_PARAM_ENGINE_TIME_SINCE_RENDERED);
 	InstanceParameters.AddParameter(SYS_PARAM_ENGINE_MIN_DIST_TO_CAMERA);
 	InstanceParameters.AddParameter(SYS_PARAM_ENGINE_SYSTEM_NUM_EMITTERS);
 	InstanceParameters.AddParameter(SYS_PARAM_ENGINE_SYSTEM_NUM_EMITTERS_ALIVE);
 
 	OwnerPositionParam.Init(InstanceParameters, SYS_PARAM_ENGINE_POSITION);
+	OwnerScaleParam.Init(InstanceParameters, SYS_PARAM_ENGINE_SCALE);
 	OwnerVelocityParam.Init(InstanceParameters, SYS_PARAM_ENGINE_VELOCITY);
 	OwnerXAxisParam.Init(InstanceParameters, SYS_PARAM_ENGINE_X_AXIS);
 	OwnerYAxisParam.Init(InstanceParameters, SYS_PARAM_ENGINE_Y_AXIS);
 	OwnerZAxisParam.Init(InstanceParameters, SYS_PARAM_ENGINE_Z_AXIS);
+
 	OwnerTransformParam.Init(InstanceParameters, SYS_PARAM_ENGINE_LOCAL_TO_WORLD);
 	OwnerInverseParam.Init(InstanceParameters, SYS_PARAM_ENGINE_WORLD_TO_LOCAL);
 	OwnerTransposeParam.Init(InstanceParameters, SYS_PARAM_ENGINE_LOCAL_TO_WORLD_TRANSPOSED);
 	OwnerInverseTransposeParam.Init(InstanceParameters, SYS_PARAM_ENGINE_WORLD_TO_LOCAL_TRANSPOSED);
+	OwnerTransformNoScaleParam.Init(InstanceParameters, SYS_PARAM_ENGINE_LOCAL_TO_WORLD_NO_SCALE);
+	OwnerInverseNoScaleParam.Init(InstanceParameters, SYS_PARAM_ENGINE_WORLD_TO_LOCAL_NO_SCALE);
+
 	OwnerDeltaSecondsParam.Init(InstanceParameters, SYS_PARAM_ENGINE_DELTA_TIME);
 	OwnerInverseDeltaSecondsParam.Init(InstanceParameters, SYS_PARAM_ENGINE_INV_DELTA_TIME);
 	OwnerMinDistanceToCameraParam.Init(InstanceParameters, SYS_PARAM_ENGINE_MIN_DIST_TO_CAMERA);
 	SystemNumEmittersParam.Init(InstanceParameters, SYS_PARAM_ENGINE_SYSTEM_NUM_EMITTERS);
 	SystemNumEmittersAliveParam.Init(InstanceParameters, SYS_PARAM_ENGINE_SYSTEM_NUM_EMITTERS_ALIVE);
 
+	SystemTimeSinceRenderedParam.Init(InstanceParameters, SYS_PARAM_ENGINE_TIME_SINCE_RENDERED);
+
 	ParameterNumParticleBindings.SetNum(Emitters.Num());
 	for (int32 i = 0; i < Emitters.Num(); i++)
 	{
@@ -228,34 +340,42 @@ void FNiagaraSystemInstance::ReInitInternal()
 		ParameterNumParticleBindings[i].Init(InstanceParameters, Var);
 	}
 
+	const bool bOnlyAdd = false;
+	System->GetExposedParameters().CopyParametersTo(InstanceParameters, bOnlyAdd);
+
 	TickInstanceParameters(0.01f);
 
 	//Determine if we can update normally or have to update solo.
-	bSolo = bForceSolo || InstanceParameters.GetDataInterfaces().Num() > 0;
+	bSolo = bForceSolo;
 	//If our scripts have any interfaces that require instance data.
 	UNiagaraScript* SystemSpawn = GetSystem()->GetSystemSpawnScript();
 	for (int32 i = 0; !bSolo && i < SystemSpawn->DataInterfaceInfo.Num(); ++i)
 	{
 		FNiagaraScriptDataInterfaceInfo& Info = SystemSpawn->DataInterfaceInfo[i];
-		if (Info.DataInterface->PerInstanceDataSize() > 0)
+		if ( (Info.DataInterface && Info.DataInterface->PerInstanceDataSize() > 0) ||
+			Info.Name.ToString().StartsWith("User."))//Temp hack to force solo on any systems with system scrips needing user (aka per instance) interfaces.
 		{
 			bSolo = true;
+			break;
 		}
+		
 	}	
 	UNiagaraScript* SystemUpdate = GetSystem()->GetSystemUpdateScript();
 	for (int32 i = 0; !bSolo && i < SystemUpdate->DataInterfaceInfo.Num(); ++i)
 	{
 		FNiagaraScriptDataInterfaceInfo& Info = SystemUpdate->DataInterfaceInfo[i];
-		if (Info.DataInterface->PerInstanceDataSize() > 0)
+		if ( (Info.DataInterface && Info.DataInterface->PerInstanceDataSize() > 0) ||
+			Info.Name.ToString().StartsWith("User."))//Temp hack to force solo on any systems with system scrips needing user (aka per instance) interfaces.
 		{
 			bSolo = true;
+			break;
 		}
 	}
 
 	BindParameters();
 
 	InitDataInterfaces(); // We can't wait until tick b/c the data interfaces need to be around when we update render modules.
-	FlushRenderingCommands();
+	//FlushRenderingCommands();//Ah! This causes massive hitch on activation of new fx. 
 
 	// This gets a little tricky, but we want to delete any renderers that are no longer in use on the rendering thread, but
 	// first (to be safe), we want to update the proxy to point to the new renderer objects.
@@ -279,8 +399,6 @@ void FNiagaraSystemInstance::ReInitInternal()
 		}
 	}
 
-	bReInitPending = false;
-
 #if WITH_EDITOR
 	//UE_LOG(LogNiagara, Log, TEXT("OnResetInternal %p"), this);
 	OnResetDelegate.Broadcast();
@@ -296,7 +414,7 @@ FNiagaraSystemInstance::~FNiagaraSystemInstance()
 
 	if (SystemInstanceIndex != INDEX_NONE)
 	{
-		FNiagaraSystemSimulation* SystemSim = GetSystemSimulation();
+		TSharedPtr<FNiagaraSystemSimulation, ESPMode::ThreadSafe> SystemSim = GetSystemSimulation();
 		SystemSim->RemoveInstance(this);
 	}
 
@@ -340,7 +458,7 @@ FNiagaraSystemInstance::~FNiagaraSystemInstance()
 
 void FNiagaraSystemInstance::BindParameters()
 {
-	Component->GetInitialParameters().Bind(&InstanceParameters);
+	Component->GetOverrideParameters().Bind(&InstanceParameters);
 
 	for (TSharedRef<FNiagaraEmitterInstance> Simulation : Emitters)
 	{
@@ -350,7 +468,7 @@ void FNiagaraSystemInstance::BindParameters()
 
 void FNiagaraSystemInstance::UnbindParameters()
 {
-	Component->GetInitialParameters().Unbind(&InstanceParameters);
+	Component->GetOverrideParameters().Unbind(&InstanceParameters);
 
 	for (TSharedRef<FNiagaraEmitterInstance> Simulation : Emitters)
 	{
@@ -377,6 +495,8 @@ void FNiagaraSystemInstance::InitDataInterfaces()
 	{
 		return;
 	}
+
+	Component->GetOverrideParameters().Tick();
 	
 	bDataInterfacesNeedInit = false;
 
@@ -389,6 +509,11 @@ void FNiagaraSystemInstance::InitDataInterfaces()
 	{
 		for (UNiagaraDataInterface* Interface : Interfaces)
 		{
+			if (!Interface)
+			{
+				continue;
+			}
+
 			if (int32 Size = Interface->PerInstanceDataSize())
 			{
 				int32* ExistingInstanceDataOffset = DataInterfaceInstanceDataOffsets.Find(Interface);
@@ -402,6 +527,13 @@ void FNiagaraSystemInstance::InitDataInterfaces()
 	};
 
 	CalcInstDataSize(InstanceParameters.GetDataInterfaces());//This probably should be a proper exec context. 
+	if (bSolo)
+	{
+		TSharedPtr<FNiagaraSystemSimulation, ESPMode::ThreadSafe> SystemSim = GetSystemSimulation();
+		
+		CalcInstDataSize(SystemSim->GetSoloDataInterfacesSpawn());
+		CalcInstDataSize(SystemSim->GetSoloDataInterfacesUpdate());
+	}
 
 	//Iterate over interfaces to get size for table and clear their interface bindings.
 	for (TSharedRef<FNiagaraEmitterInstance> Simulation : Emitters)
@@ -443,7 +575,7 @@ void FNiagaraSystemInstance::InitDataInterfaces()
 	}
 }
 
-void FNiagaraSystemInstance::TickDataInterfaces(float DeltaSeconds)
+void FNiagaraSystemInstance::TickDataInterfaces(float DeltaSeconds, bool bPostSimulate)
 {
 	if (!GetSystem() || !Component || bError)
 	{
@@ -456,7 +588,14 @@ void FNiagaraSystemInstance::TickDataInterfaces(float DeltaSeconds)
 		if(UNiagaraDataInterface* Interface = Pair.Key.Get())
 		{
 			//Ideally when we make the batching changes, we can keep the instance data in big single type blocks that can all be updated together with a single virtual call.
-			bReInitDataInterfaces |= Interface->PerInstanceTick(&DataInterfaceInstanceData[Pair.Value], this, DeltaSeconds);
+			if (bPostSimulate)
+			{
+				bReInitDataInterfaces |= Interface->PerInstanceTickPostSimulate(&DataInterfaceInstanceData[Pair.Value], this, DeltaSeconds);
+			}
+			else
+			{
+				bReInitDataInterfaces |= Interface->PerInstanceTick(&DataInterfaceInstanceData[Pair.Value], this, DeltaSeconds);
+			}
 		}
 	}
 	
@@ -475,6 +614,7 @@ void FNiagaraSystemInstance::TickInstanceParameters(float DeltaSeconds)
 	FVector OldPos = OwnerPositionParam.GetValue();// ComponentTrans.GetLocation();
 	FVector CurrPos = ComponentTrans.GetLocation();
 	OwnerPositionParam.SetValue(CurrPos);
+	OwnerScaleParam.SetValue(ComponentTrans.GetScale3D());
 	OwnerVelocityParam.SetValue((CurrPos - OldPos) / DeltaSeconds);
 	OwnerXAxisParam.SetValue(ComponentTrans.GetRotation().GetAxisX());
 	OwnerYAxisParam.SetValue(ComponentTrans.GetRotation().GetAxisY());
@@ -488,6 +628,12 @@ void FNiagaraSystemInstance::TickInstanceParameters(float DeltaSeconds)
 	OwnerInverseParam.SetValue(Inverse);
 	OwnerTransposeParam.SetValue(Transpose);
 	OwnerInverseTransposeParam.SetValue(InverseTranspose);
+
+	FMatrix TransformNoScale = ComponentTrans.ToMatrixNoScale();
+	FMatrix InverseNoScale = TransformNoScale.Inverse();
+	OwnerTransformNoScaleParam.SetValue(TransformNoScale);
+	OwnerInverseNoScaleParam.SetValue(InverseNoScale);
+
 	OwnerDeltaSecondsParam.SetValue(DeltaSeconds);
 	OwnerInverseDeltaSecondsParam.SetValue(1.0f / DeltaSeconds);
 
@@ -530,7 +676,7 @@ void FNiagaraSystemInstance::TickInstanceParameters(float DeltaSeconds)
 	for (int32 i = 0; i < Emitters.Num(); i++)
 	{
 		int32 NumParticles = Emitters[i]->GetNumParticles();
-		if (Emitters[i]->GetExecutionState() != ENiagaraExecutionState::Dead)
+		if (!Emitters[i]->IsComplete())
 		{
 			NumAlive++;
 		}
@@ -538,10 +684,15 @@ void FNiagaraSystemInstance::TickInstanceParameters(float DeltaSeconds)
 	}
 	SystemNumEmittersParam.SetValue(Emitters.Num());
 	SystemNumEmittersAliveParam.SetValue(NumAlive);
+
+	check(World);
+	float SafeTimeSinceRendererd = FMath::Max(0.0f, World->GetTimeSeconds() - Component->LastRenderTime - GLastRenderTimeSafetyBias);
+	SystemTimeSinceRenderedParam.SetValue(SafeTimeSinceRendererd);
+
 	
-	Component->GetInitialParameters().Tick();
+	Component->GetOverrideParameters().Tick();
 	InstanceParameters.Tick();
-	InstanceParameters.bParametersDirty = true;
+	InstanceParameters.MarkParametersDirty();
 }
 
 #if WITH_EDITORONLY_DATA
@@ -583,14 +734,11 @@ bool FNiagaraSystemInstance::UsesScript(const UNiagaraScript* Script)const
 
 bool FNiagaraSystemInstance::UsesCollection(const UNiagaraParameterCollection* Collection)const
 {
-	if (GetSystem())
+	if (UNiagaraSystem* System = GetSystem())
 	{
-		for (FNiagaraEmitterHandle EmitterHandle : GetSystem()->GetEmitterHandles())
+		if (System->UsesCollection(Collection))
 		{
-			if (EmitterHandle.GetSource()->UsesCollection(Collection) || EmitterHandle.GetInstance()->UsesCollection(Collection))
-			{
-				return true;
-			}
+			return true;
 		}
 	}
 	return false;
@@ -627,6 +775,12 @@ void FNiagaraSystemInstance::InitEmitters()
 			Sim->Init(EmitterIdx, IDName);
 			Emitters.Add(Sim);
 		}
+
+		for (TSharedRef<FNiagaraEmitterInstance> Simulation : Emitters)
+		{
+			Simulation->PostInitSimulation();
+		}
+
 		bDataInterfacesNeedInit = true;
 	}
 }
@@ -671,30 +825,50 @@ void FNiagaraSystemInstance::ComponentTick(float DeltaSeconds)
 		return;
 	}
 
-	FNiagaraSystemSimulation* Sim = GetSystemSimulation();
-	check(Sim);
+	TSharedPtr<FNiagaraSystemSimulation, ESPMode::ThreadSafe> Sim = GetSystemSimulation();
+	check(Sim.IsValid());
 	check(IsInGameThread());
 	check(bSolo);
 	check(Component);
 
 	PreSimulateTick(DeltaSeconds);
 
-	FNiagaraSystemSimulation* SystemSim = GetSystemSimulation();
+	TSharedPtr<FNiagaraSystemSimulation, ESPMode::ThreadSafe> SystemSim = GetSystemSimulation();
 	SystemSim->TickSolo(this);
 
-	PostSimulateTick(DeltaSeconds);
+	if(!HandleCompletion())
+	{
+		PostSimulateTick(DeltaSeconds);
+	}
+}
+
+bool FNiagaraSystemInstance::HandleCompletion()
+{
+	bool bEmittersCompleteOrDisabled = true;
+	bHasTickingEmitters = false;
+	for (TSharedRef<FNiagaraEmitterInstance>&it : Emitters)
+	{
+		FNiagaraEmitterInstance& Inst = *it;
+		bEmittersCompleteOrDisabled &= Inst.HandleCompletion();
+		bHasTickingEmitters |= Inst.ShouldTick();
+	}
+
+	if (IsComplete() || bEmittersCompleteOrDisabled)
+	{
+		Complete();
+		return true;
+	}
+	return false;
 }
 
 void FNiagaraSystemInstance::PreSimulateTick(float DeltaSeconds)
 {
-	HandleResets();
-
 	TickInstanceParameters(DeltaSeconds);
 }
 
 void FNiagaraSystemInstance::PostSimulateTick(float DeltaSeconds)
 {
-	if (GetSystem() == nullptr || Component == nullptr || DeltaSeconds < SMALL_NUMBER || bError)
+	if (IsComplete() || !bHasTickingEmitters || GetSystem() == nullptr || Component == nullptr || DeltaSeconds < SMALL_NUMBER || bError)
 	{
 		return;
 	}
@@ -709,94 +883,54 @@ void FNiagaraSystemInstance::PostSimulateTick(float DeltaSeconds)
 	if (bDataInterfacesNeedInit)
 	{
 		InitDataInterfaces();
-	}
-
-	//TODO: Now we're batching in the world manager we can store these in big blocks across systems and update all instance data in one virtual call?
-	TickDataInterfaces(DeltaSeconds);
-
-	// need to run this serially, as PreTick will initialize data interfaces, which can't be done in parallel
-	for (TSharedRef<FNiagaraEmitterInstance>&it : Emitters)
-	{
-		const UNiagaraEmitter *Props = it->GetEmitterHandle().GetInstance();
-		check(Props);
-		it->PreTick();
-	}
-
-	// now tick all emitters in parallel
-	{
-		SCOPE_CYCLE_COUNTER(STAT_NiagaraParallelTick);
-
- 		if (GbNiagaraParallelEmitterTick && bSolo)
- 		{
- 			ParallelFor(Emitters.Num(),
- 				[&](int32 EmitterIdx)
- 			{
- 				TSharedRef<FNiagaraEmitterInstance>&it = Emitters[EmitterIdx];
- 				if (it->GetExecutionState() != ENiagaraExecutionState::Dead && it->GetExecutionState() != ENiagaraExecutionState::Paused)
- 				{
- 					it->Tick(DeltaSeconds);
- 				}
- 			});
- 		}
- 		else
+		// We may have entered into an error state as a result of InitDataInterfaces.
+		if (bError)
 		{
-			for (int32 EmitterIdx = 0; EmitterIdx < Emitters.Num(); EmitterIdx++)
-			{
-				TSharedRef<FNiagaraEmitterInstance>&it = Emitters[EmitterIdx];
-				if (it->GetExecutionState() != ENiagaraExecutionState::Dead && it->GetExecutionState() != ENiagaraExecutionState::Paused)
-				{
-					it->Tick(DeltaSeconds);
-				}
-			}
+			return;
 		}
 	}
 
-	Age += DeltaSeconds;
-}
+	//TODO: Now we're batching in the world manager we can store these in big blocks across systems and update all instance data in one virtual call?
+	TickDataInterfaces(DeltaSeconds, false);
 
-void FNiagaraSystemInstance::HandleResets()
-{
-	if (bReInitPending)
+	// We may have entered into an error state as a result of TickDataInterfaces.
+	if (bError)
 	{
-		ReInitInternal();
+		return;
 	}
-	else if (bResetPending)
+
+	// need to run this serially, as PreTick will initialize data interfaces, which can't be done in parallel
+	for (int32 EmitterIdx = 0; EmitterIdx < Emitters.Num(); EmitterIdx++)
 	{
-		ResetInternal();
+		FNiagaraEmitterInstance& Inst = Emitters[EmitterIdx].Get();		
+		const UNiagaraEmitter *Props = Inst.GetEmitterHandle().GetInstance();
+		check(Props);
+		Inst.PreTick();
 	}
-}
-
-void FNiagaraSystemInstance::Disable()
-{
-	FNiagaraSystemSimulation* SystemSim = GetSystemSimulation();
-	SystemSim->RemoveInstance(this);
-
-	// Clear out the System renderer from the proxy.
-	TArray<NiagaraRenderer*> NewRenderers;
-	UpdateProxy(NewRenderers);
 
-	// Clear out the System renderer from the simulation.
-	for (TSharedRef<FNiagaraEmitterInstance> Simulation : Emitters)
+	// now tick all emitters
+	for (int32 EmitterIdx = 0; EmitterIdx < Emitters.Num(); EmitterIdx++)
 	{
-		Simulation->ClearRenderer();
+		FNiagaraEmitterInstance& Inst = Emitters[EmitterIdx].Get();
+		Inst.Tick(DeltaSeconds);
 	}
 
-	Emitters.Empty();
-	SetExecutionState(ENiagaraExecutionState::Disabled);
+	TickDataInterfaces(DeltaSeconds, true);
 
-	if (Component)
-	{
-		Component->OnSystemDisabled();
-	}
+	Age += DeltaSeconds;
 }
 
-void FNiagaraSystemInstance::Enable()
+#if WITH_EDITORONLY_DATA
+bool FNiagaraSystemInstance::GetIsolateEnabled() const
 {
-	ExecutionState = ENiagaraExecutionState::Inactive;
-
-	InitEmitters();
-	Reset(EResetMode::ImmediateReInit);
+	UNiagaraSystem* System = GetSystem();
+	if (System)
+	{
+		return System->GetIsolateEnabled();
+	}
+	return false;
 }
+#endif
 
 void FNiagaraSystemInstance::ReinitializeDataInterfaces()
 {
@@ -846,7 +980,7 @@ FNiagaraDataSet* FNiagaraSystemInstance::GetDataSet(FNiagaraDataSetID SetID, FNa
 	for (TSharedPtr<FNiagaraEmitterInstance> Emitter : Emitters)
 	{
 		check(Emitter.IsValid());
-		if (Emitter->IsEnabled())
+		if (!Emitter->IsComplete())
 		{
 			if (Emitter->GetEmitterHandle().GetIdName() == EmitterName)
 			{
@@ -863,6 +997,11 @@ FNiagaraSystemInstance::FOnInitialized& FNiagaraSystemInstance::OnInitialized()
 	return OnInitializedDelegate;
 }
 
+FNiagaraSystemInstance::FOnComplete& FNiagaraSystemInstance::OnComplete()
+{
+	return OnCompleteDelegate;
+}
+
 #if WITH_EDITOR
 FNiagaraSystemInstance::FOnReset& FNiagaraSystemInstance::OnReset()
 {
diff --git a/Engine/Plugins/FX/Niagara/Source/Niagara/Private/NiagaraSystemSimulation.cpp b/Engine/Plugins/FX/Niagara/Source/Niagara/Private/NiagaraSystemSimulation.cpp
--- a/Engine/Plugins/FX/Niagara/Source/Niagara/Private/NiagaraSystemSimulation.cpp
+++ b/Engine/Plugins/FX/Niagara/Source/Niagara/Private/NiagaraSystemSimulation.cpp
@@ -33,7 +33,7 @@ static FAutoConsoleVariableRef CVarNiagaraDumpSystemData(
 	ECVF_Default
 );
 
-static int32 GbParallelSystemPreTick = 1;
+static int32 GbParallelSystemPreTick = 0;
 static FAutoConsoleVariableRef CVarNiagaraParallelSystemPreTick(
 	TEXT("fx.ParallelSystemPreTick"),
 	GbParallelSystemPreTick,
@@ -41,7 +41,7 @@ static FAutoConsoleVariableRef CVarNiagaraParallelSystemPreTick(
 	ECVF_Default
 );
 
-static int32 GbParallelSystemPostTick = 1;
+static int32 GbParallelSystemPostTick = 0;
 static FAutoConsoleVariableRef CVarNiagaraParallelSystemPostTick(
 	TEXT("fx.ParallelSystemPostTick"),
 	GbParallelSystemPostTick,
@@ -60,9 +60,16 @@ static FAutoConsoleVariableRef CVarNiagaraParallelSystemPostTick(
 
 //////////////////////////////////////////////////////////////////////////
 
+FNiagaraSystemSimulation::~FNiagaraSystemSimulation()
+{
+	Destroy();
+}
+
 bool FNiagaraSystemSimulation::Init(UNiagaraSystem* InSystem, UWorld* InWorld)
 {
-	System = InSystem;
+	UNiagaraSystem* System = InSystem;
+	WeakSystem = System;
+
 	World = InWorld;
 
 	FNiagaraWorldManager* WorldMan = FNiagaraWorldManager::Get(InWorld);
@@ -79,37 +86,43 @@ bool FNiagaraSystemSimulation::Init(UNiagaraSystem* InSystem, UWorld* InWorld)
 		DataSet.AddVariables(System->GetSystemUpdateScript()->Attributes);
 		DataSet.Finalize();
 
-		SpawnParameterDataSet.Reset();
-		FNiagaraParameters* EngineParamsSpawn = System->GetSystemSpawnScript()->DataSetToParameters.Find(TEXT("Engine"));
-		if (EngineParamsSpawn != nullptr)
 		{
-			SpawnParameterDataSet.AddVariables(EngineParamsSpawn->Parameters);
-		}
-		SpawnParameterDataSet.Finalize();
-		UpdateParameterDataSet.Reset();
-		FNiagaraParameters* EngineParamsUpdate = System->GetSystemUpdateScript()->DataSetToParameters.Find(TEXT("Engine"));
-		if (EngineParamsUpdate != nullptr)
-		{
-			UpdateParameterDataSet.AddVariables(EngineParamsUpdate->Parameters);
+			SpawnInstanceParameterDataSet.Reset();
+			FNiagaraParameters* EngineParamsSpawn = System->GetSystemSpawnScript()->DataSetToParameters.Find(TEXT("Engine"));
+			if (EngineParamsSpawn != nullptr)
+			{
+				SpawnInstanceParameterDataSet.AddVariables(EngineParamsSpawn->Parameters);
+			}
+			SpawnInstanceParameterDataSet.Finalize();
+			UpdateInstanceParameterDataSet.Reset();
+			FNiagaraParameters* EngineParamsUpdate = System->GetSystemUpdateScript()->DataSetToParameters.Find(TEXT("Engine"));
+			if (EngineParamsUpdate != nullptr)
+			{
+				UpdateInstanceParameterDataSet.AddVariables(EngineParamsUpdate->Parameters);
+			}
+			UpdateInstanceParameterDataSet.Finalize();
 		}
-		UpdateParameterDataSet.Finalize();
 
-		SpawnExecContext.Init(System->GetSystemSpawnScript(false), ENiagaraSimTarget::CPUSim);
-		UpdateExecContext.Init(System->GetSystemUpdateScript(false), ENiagaraSimTarget::CPUSim);
+		UNiagaraScript* SpawnScript = System->GetSystemSpawnScript(false);
+		UNiagaraScript* UpdateScript = System->GetSystemUpdateScript(false);
+
+		SpawnExecContext.Init(SpawnScript, ENiagaraSimTarget::CPUSim);
+		UpdateExecContext.Init(UpdateScript, ENiagaraSimTarget::CPUSim);
 
 		//Bind parameter collections.
-		for (UNiagaraParameterCollection* Collection : System->GetSystemSpawnScript(false)->ParameterCollections)
+		for (UNiagaraParameterCollection* Collection : SpawnScript->ParameterCollections)
 		{
-			WorldMan->GetParameterCollection(Collection)->GetParameterStore().Bind(&SpawnExecContext.Parameters);
+			GetParameterCollectionInstance(Collection)->GetParameterStore().Bind(&SpawnExecContext.Parameters);
 		}
-		for (UNiagaraParameterCollection* Collection : System->GetSystemUpdateScript(false)->ParameterCollections)
+		for (UNiagaraParameterCollection* Collection : UpdateScript->ParameterCollections)
 		{
-			WorldMan->GetParameterCollection(Collection)->GetParameterStore().Bind(&UpdateExecContext.Parameters);
+			GetParameterCollectionInstance(Collection)->GetParameterStore().Bind(&UpdateExecContext.Parameters);
 		}
 
-		SystemEnabledAccessor.Create(&DataSet, FNiagaraVariable(FNiagaraTypeDefinition::GetBoolDef(), TEXT("System.Enabled")));
+		SpawnScript->RapidIterationParameters.Bind(&SpawnExecContext.Parameters);
+		UpdateScript->RapidIterationParameters.Bind(&UpdateExecContext.Parameters);
+
 		SystemExecutionStateAccessor.Create(&DataSet, FNiagaraVariable(EnumPtr, TEXT("System.ExecutionState")));
-		EmitterEnabledAccessors.Reset();
 		EmitterSpawnInfoAccessors.Reset();
 		EmitterExecutionStateAccessors.Reset();
 		EmitterSpawnInfoAccessors.SetNum(System->GetNumEmitters());
@@ -120,33 +133,32 @@ bool FNiagaraSystemSimulation::Init(UNiagaraSystem* InSystem, UWorld* InWorld)
 			UNiagaraEmitter* Emitter = EmitterHandle.GetInstance();
 			FString EmitterName = Emitter->GetUniqueEmitterName();
 			check(Emitter);
-			EmitterEnabledAccessors.Emplace(DataSet, FNiagaraVariable(FNiagaraTypeDefinition::GetBoolDef(), *(EmitterName + TEXT(".Enabled"))));
 			EmitterExecutionStateAccessors.Emplace(DataSet, FNiagaraVariable(EnumPtr, *(EmitterName + TEXT(".ExecutionState"))));
 			const TArray<FNiagaraEmitterSpawnAttributes>& EmitterSpawnAttrNames = System->GetEmitterSpawnAttributes();
-			for (const FNiagaraEmitterSpawnAttributes& SpawnAttrs : EmitterSpawnAttrNames)
+			
+			check(EmitterSpawnAttrNames.Num() == System->GetNumEmitters());
+			for (FName AttrName : EmitterSpawnAttrNames[EmitterIdx].SpawnAttributes)
 			{
-				for (FName AttrName : SpawnAttrs.SpawnAttributes)
-				{
-					EmitterSpawnInfoAccessors[EmitterIdx].Emplace(DataSet, FNiagaraVariable(FNiagaraTypeDefinition(FNiagaraSpawnInfo::StaticStruct()), AttrName));
-				}
+				EmitterSpawnInfoAccessors[EmitterIdx].Emplace(DataSet, FNiagaraVariable(FNiagaraTypeDefinition(FNiagaraSpawnInfo::StaticStruct()), AttrName));
 			}
 		}
 	}
 
 	if (bCanExecuteSolo)
 	{
+		UNiagaraScript* SpawnScript = System->GetSystemSpawnScript(true);
+		UNiagaraScript* UpdateScript = System->GetSystemUpdateScript(true);
+
 		//Init solo simulation.
 		DataSetSolo.Reset();
-		DataSetSolo.AddVariables(System->GetSystemSpawnScript(true)->Attributes);
-		DataSetSolo.AddVariables(System->GetSystemUpdateScript(true)->Attributes);
+		DataSetSolo.AddVariables(SpawnScript->Attributes);
+		DataSetSolo.AddVariables(UpdateScript->Attributes);
 		DataSetSolo.Finalize();
 
-		SpawnExecContextSolo.Init(System->GetSystemSpawnScript(true), ENiagaraSimTarget::CPUSim);
-		UpdateExecContextSolo.Init(System->GetSystemUpdateScript(true), ENiagaraSimTarget::CPUSim);
+		SpawnExecContextSolo.Init(SpawnScript, ENiagaraSimTarget::CPUSim);
+		UpdateExecContextSolo.Init(UpdateScript, ENiagaraSimTarget::CPUSim);
 
-		SoloSystemEnabledAccessor.Create(&DataSetSolo, FNiagaraVariable(FNiagaraTypeDefinition::GetBoolDef(), TEXT("System.Enabled")));
 		SoloSystemExecutionStateAccessor.Create(&DataSetSolo, FNiagaraVariable(EnumPtr, TEXT("System.ExecutionState")));
-		SoloEmitterEnabledAccessors.Reset();
 		SoloEmitterExecutionStateAccessors.Reset();
 		SoloEmitterSpawnInfoAccessors.Reset();
 		SoloEmitterSpawnInfoAccessors.SetNum(System->GetNumEmitters());
@@ -157,7 +169,6 @@ bool FNiagaraSystemSimulation::Init(UNiagaraSystem* InSystem, UWorld* InWorld)
 			UNiagaraEmitter* Emitter = EmitterHandle.GetInstance();
 			FString EmitterName = Emitter->GetUniqueEmitterName();
 			check(Emitter);
-			SoloEmitterEnabledAccessors.Emplace(DataSetSolo, FNiagaraVariable(FNiagaraTypeDefinition::GetBoolDef(), *(EmitterName + TEXT(".Enabled"))));
 			SoloEmitterExecutionStateAccessors.Emplace(DataSetSolo, FNiagaraVariable(EnumPtr, *(EmitterName + TEXT(".ExecutionState"))));
 			
 			for (FName AttrName : System->GetEmitterSpawnAttributes()[EmitterIdx].SpawnAttributes)
@@ -165,6 +176,19 @@ bool FNiagaraSystemSimulation::Init(UNiagaraSystem* InSystem, UWorld* InWorld)
 				SoloEmitterSpawnInfoAccessors[EmitterIdx].Emplace(DataSetSolo, FNiagaraVariable(FNiagaraTypeDefinition(FNiagaraSpawnInfo::StaticStruct()), AttrName));
 			}
 		}
+
+		//Bind the parameter collections and RI parameters into the Solo Execution contexts
+		for (UNiagaraParameterCollection* Collection : SpawnScript->ParameterCollections)
+		{
+			GetParameterCollectionInstance(Collection)->GetParameterStore().Bind(&SpawnExecContextSolo.Parameters);
+		}
+		SpawnScript->RapidIterationParameters.Bind(&SpawnExecContextSolo.Parameters);
+
+		for (UNiagaraParameterCollection* Collection : UpdateScript->ParameterCollections)
+		{
+			GetParameterCollectionInstance(Collection)->GetParameterStore().Bind(&UpdateExecContextSolo.Parameters);
+		}
+		UpdateScript->RapidIterationParameters.Bind(&UpdateExecContextSolo.Parameters);
 	}
 
 	return true;
@@ -187,15 +211,26 @@ void FNiagaraSystemSimulation::Destroy()
 
 	FNiagaraWorldManager* WorldMan = FNiagaraWorldManager::Get(World);
 	check(WorldMan);
-	//Unbind parameter collections from the batched exec contexts.
-	for (UNiagaraParameterCollection* Collection : System->GetSystemSpawnScript(false)->ParameterCollections)
-	{
-		WorldMan->GetParameterCollection(Collection)->GetParameterStore().Unbind(&SpawnExecContext.Parameters);
-	}
-	for (UNiagaraParameterCollection* Collection : System->GetSystemUpdateScript(false)->ParameterCollections)
+	SpawnExecContext.Parameters.UnbindFromSourceStores();
+	UpdateExecContext.Parameters.UnbindFromSourceStores();
+	SpawnExecContextSolo.Parameters.UnbindFromSourceStores();
+	UpdateExecContextSolo.Parameters.UnbindFromSourceStores();
+}
+
+UNiagaraParameterCollectionInstance* FNiagaraSystemSimulation::GetParameterCollectionInstance(UNiagaraParameterCollection* Collection)
+{
+	UNiagaraSystem* System = WeakSystem.Get();
+	check(System != nullptr);
+	UNiagaraParameterCollectionInstance* Ret = System->GetParameterCollectionOverride(Collection);
+
+	//If no explicit override from the system, just get the current instance set on the world.
+	if (!Ret)
 	{
-		WorldMan->GetParameterCollection(Collection)->GetParameterStore().Unbind(&UpdateExecContext.Parameters);
+		FNiagaraWorldManager* WorldMan = FNiagaraWorldManager::Get(World);
+		Ret = WorldMan->GetParameterCollection(Collection);
 	}
+
+	return Ret;
 }
 
 void FNiagaraSystemSimulation::TickSoloDataSet()
@@ -209,14 +244,24 @@ void FNiagaraSystemSimulation::TickSoloDataSet()
 	}
 }
 
-void FNiagaraSystemSimulation::Tick(float DeltaSeconds)
+bool FNiagaraSystemSimulation::Tick(float DeltaSeconds)
 {
-	if (!System->IsValid())
+	UNiagaraSystem* System = WeakSystem.Get();
+	if (System == nullptr)
 	{
 		// TODO: evaluate whether or not we should have removed this from the world manager instead?
-		return;
+		return false;
 	}
 
+	UNiagaraScript* SystemSpawnScript = System->GetSystemSpawnScript(false);
+	UNiagaraScript* SystemUpdateScript = System->GetSystemUpdateScript(false);
+	UNiagaraScript* SystemSpawnScriptSolo = System->GetSystemSpawnScript(true);
+	UNiagaraScript* SystemUpdateScriptSolo = System->GetSystemUpdateScript(true);
+	SystemSpawnScript->RapidIterationParameters.Tick();
+	SystemUpdateScript->RapidIterationParameters.Tick();
+	SystemSpawnScriptSolo->RapidIterationParameters.Tick();
+	SystemUpdateScriptSolo->RapidIterationParameters.Tick();
+
 	TickSoloDataSet();
 
 	SCOPE_CYCLE_COUNTER(STAT_NiagaraSystemSim);
@@ -233,6 +278,16 @@ void FNiagaraSystemSimulation::Tick(float DeltaSeconds)
 	}
 	PendingSystemInstances.Reset();
 
+	// Go ahead an initialize the data set accessors that we will potentially write since they 
+	// may have changed internally from the previous frame. We write to current here because we
+	// call tick after pretick, which will make these values the previous values during execution.
+	SystemExecutionStateAccessor.InitForAccess(true);
+	for (int32 EmitterIdx = 0; EmitterIdx < System->GetNumEmitters(); ++EmitterIdx)
+	{
+		EmitterExecutionStateAccessors[EmitterIdx].InitForAccess(true);
+	}
+	
+
 	{
 		SCOPE_CYCLE_COUNTER(STAT_NiagaraSystemSim_PreSimulate);
 
@@ -241,18 +296,41 @@ void FNiagaraSystemSimulation::Tick(float DeltaSeconds)
 			FNiagaraSystemInstance* Inst = SystemInstances[SystemIndex];
 			Inst->PreSimulateTick(DeltaSeconds);
 
-			check(!Inst->GetParameters().bLayoutDirty);
-			if (Inst->GetParameters().bParametersDirty && bCanExecute)
+			//System exec state can be driven externally so we must feed it in as well as take it out.
+			if (SystemExecutionStateAccessor.IsValid())
+			{
+				SystemExecutionStateAccessor.Set(SystemIndex, (int32)Inst->GetExecutionState());
+			}
+
+			if (Inst->GetParameters().GetParametersDirty() && bCanExecute)
 			{
-				SpawnParameterToDataSetBinding.ParameterStoreToDataSet(Inst->GetParameters(), SpawnParameterDataSet, SystemIndex);
-				UpdateParameterToDataSetBinding.ParameterStoreToDataSet(Inst->GetParameters(), UpdateParameterDataSet, SystemIndex);
+				SpawnInstanceParameterToDataSetBinding.ParameterStoreToDataSet(Inst->GetParameters(), SpawnInstanceParameterDataSet, SystemIndex);
+				UpdateInstanceParameterToDataSetBinding.ParameterStoreToDataSet(Inst->GetParameters(), UpdateInstanceParameterDataSet, SystemIndex);
 			}
+
 			//TODO: Find good way to check that we're not using any instance parameter data interfaces in the system scripts here.
 			//In that case we need to solo and will never get here.
+
+			// Update the individual parameters that may have been manually set. Note that we check for a valid base in multiple situations
+			// as in certain situations (like spawn) these values may not yet have been set or created.
+			if (SystemExecutionStateAccessor.BaseIsValid())
+			{
+				SystemExecutionStateAccessor.Set(SystemIndex, (int32)Inst->GetExecutionState());
+			}
+
+			TArray<TSharedRef<FNiagaraEmitterInstance>>& Emitters = Inst->GetEmitters();
+			for (int32 EmitterIdx = 0; EmitterIdx < Emitters.Num(); ++EmitterIdx)
+			{
+				FNiagaraEmitterInstance& EmitterInst = Emitters[EmitterIdx].Get();
+				if (EmitterExecutionStateAccessors.Num() > EmitterIdx && EmitterExecutionStateAccessors[EmitterIdx].BaseIsValid())
+				{
+					EmitterExecutionStateAccessors[EmitterIdx].Set(SystemIndex, (int32)EmitterInst.GetExecutionState());
+				}
+			}
 		};
 
-		SpawnParameterDataSet.Allocate(NewNum);
-		UpdateParameterDataSet.Allocate(NewNum);
+		SpawnInstanceParameterDataSet.Allocate(NewNum);
+		UpdateInstanceParameterDataSet.Allocate(NewNum);
 
 		if (GbParallelSystemPreTick)
 		{
@@ -267,8 +345,8 @@ void FNiagaraSystemSimulation::Tick(float DeltaSeconds)
 			}
 		}
 
-		SpawnParameterDataSet.Tick();
-		UpdateParameterDataSet.Tick();
+		SpawnInstanceParameterDataSet.Tick();
+		UpdateInstanceParameterDataSet.Tick();
 	}
 
 	if (bCanExecute && NewNum > 0)
@@ -288,50 +366,60 @@ void FNiagaraSystemSimulation::Tick(float DeltaSeconds)
 
 		TArray<FNiagaraDataSetExecutionInfo> DataSetExecInfos;
 		DataSetExecInfos.Emplace(&DataSet, 0, false, true);
+		DataSetExecInfos.Emplace(&UpdateInstanceParameterDataSet, 0, false, false);
 
 		{
 			SCOPE_CYCLE_COUNTER(STAT_NiagaraSystemSim_Update);
 			DataSet.SetNumInstances(OrigNum);
+			UpdateInstanceParameterDataSet.SetNumInstances(OrigNum);
+
 			//Run update.
 			UpdateExecContext.Tick(nullptr);//We can't require a specfic instance here as these are for all instances.
-			DataSetExecInfos.SetNum(1, false);
 			DataSetExecInfos[0].StartInstance = 0;
-			DataSetExecInfos.Emplace(&UpdateParameterDataSet, 0, false, false);
+
+			if (GbDumpSystemData)
+			{
+				UE_LOG(LogNiagara, Log, TEXT("=== PreUpdated %d Systems ==="), OrigNum);
+				DataSet.Dump(true, 0, OrigNum);
+				UpdateInstanceParameterDataSet.Dump(true, 0, OrigNum);
+			}
+
 			UpdateExecContext.Execute(OrigNum, DataSetExecInfos);
 
 			if (GbDumpSystemData)
 			{
 				UE_LOG(LogNiagara, Log, TEXT("=== Updated %d Systems ==="), OrigNum);
 				DataSet.Dump(true, 0, OrigNum);
-				UpdateParameterDataSet.Dump(true, 0, OrigNum);
+				UpdateInstanceParameterDataSet.Dump(false, 0, OrigNum);
 			}
 		}
 
+		if(SpawnNum)
 		{
 			SCOPE_CYCLE_COUNTER(STAT_NiagaraSystemSim_Spawn);
 			DataSet.SetNumInstances(NewNum);
+			SpawnInstanceParameterDataSet.SetNumInstances(NewNum);
+
 			//Run Spawn
 			SpawnExecContext.Tick(nullptr);//We can't require a specific instance here as these are for all instances.
 			DataSetExecInfos.SetNum(1, false);
 			DataSetExecInfos[0].StartInstance = OrigNum;
-			DataSetExecInfos.Emplace(&SpawnParameterDataSet, OrigNum, false, false);
+			DataSetExecInfos.Emplace(&SpawnInstanceParameterDataSet, OrigNum, false, false);
 			SpawnExecContext.Execute(SpawnNum, DataSetExecInfos);
 
 			if (GbDumpSystemData)
 			{
 				UE_LOG(LogNiagara, Log, TEXT("=== Spwaned %d Systems ==="), SpawnNum);
 				DataSet.Dump(true, OrigNum, SpawnNum);
-				SpawnParameterDataSet.Dump(true, OrigNum, SpawnNum);
+				SpawnInstanceParameterDataSet.Dump(false, OrigNum, SpawnNum);
 			}
 		}
 
 		{
 			SCOPE_CYCLE_COUNTER(STAT_NiagaraSystemSim_TransferParameters);
-			SystemEnabledAccessor.InitForAccess(true);
 			SystemExecutionStateAccessor.InitForAccess(true);
 			for (int32 EmitterIdx = 0; EmitterIdx < System->GetNumEmitters(); ++EmitterIdx)
 			{
-				EmitterEnabledAccessors[EmitterIdx].InitForAccess(true);
 				EmitterExecutionStateAccessors[EmitterIdx].InitForAccess(true);
 				for (int32 SpawnInfoIdx = 0; SpawnInfoIdx < EmitterSpawnInfoAccessors[EmitterIdx].Num(); ++SpawnInfoIdx)
 				{
@@ -342,17 +430,10 @@ void FNiagaraSystemSimulation::Tick(float DeltaSeconds)
 			int32 SystemIndex = 0;
 			while (SystemIndex < SystemInstances.Num())
 			{
-				ENiagaraExecutionState ExecutionState = (ENiagaraExecutionState)SystemExecutionStateAccessor.GetSafe(SystemIndex, (int32)ENiagaraExecutionState::Active);
+				ENiagaraExecutionState ExecutionState = (ENiagaraExecutionState)SystemExecutionStateAccessor.GetSafe(SystemIndex, (int32)ENiagaraExecutionState::Disabled);
 				FNiagaraSystemInstance* SystemInst = SystemInstances[SystemIndex];
 				SystemInst->SetExecutionState(ExecutionState);
-				bool bSystemEnabled = SystemEnabledAccessor.GetSafe(SystemIndex, true);
-				//Kill any instances that are flagged as disabled.
-				if (!bSystemEnabled || ExecutionState == ENiagaraExecutionState::Dead)
-				{
-					RemoveInstance(SystemInst);
-					SystemInst->Disable();
-				}
-				else
+				if (!SystemInst->HandleCompletion())
 				{
 					//Now pull data out of the simulation and drive the emitters with it.
 					TArray<TSharedRef<FNiagaraEmitterInstance>>& Emitters = SystemInst->GetEmitters();
@@ -360,9 +441,15 @@ void FNiagaraSystemSimulation::Tick(float DeltaSeconds)
 					{
 						FNiagaraEmitterInstance& EmitterInst = Emitters[EmitterIdx].Get();
 
-						FNiagaraBool bEmitterEnabled = EmitterEnabledAccessors[EmitterIdx].GetSafe(SystemIndex, true);
-						EmitterInst.SetEnabled(bEmitterEnabled);//TODO: Alter this to instruct emitter to dump existing particles optionally.
+						//Early exit before we set the state as if we're complete or disabled we should never let the emitter turn itself back. It needs to be reset/reinited manually.
+						if (EmitterInst.IsComplete())
+						{
+							continue;
+						}
 
+						ENiagaraExecutionState State = (ENiagaraExecutionState)EmitterExecutionStateAccessors[EmitterIdx].GetSafe(SystemIndex, (int32)ENiagaraExecutionState::Disabled);
+						EmitterInst.SetExecutionState(State);
+												
 						TArray<FNiagaraSpawnInfo>& EmitterInstSpawnInfos = EmitterInst.GetSpawnInfo();
 						for (int32 SpawnInfoIdx=0; SpawnInfoIdx < EmitterSpawnInfoAccessors[EmitterIdx].Num(); ++SpawnInfoIdx)
 						{
@@ -376,8 +463,6 @@ void FNiagaraSystemSimulation::Tick(float DeltaSeconds)
 							}
 						}
 
-						ENiagaraExecutionState State = (ENiagaraExecutionState)EmitterExecutionStateAccessors[EmitterIdx].GetSafe(SystemIndex, (int32)ENiagaraExecutionState::Active);
-						EmitterInst.SetExecutionState(State);
 
 						//TODO: Any other fixed function stuff like this?
 
@@ -391,7 +476,14 @@ void FNiagaraSystemSimulation::Tick(float DeltaSeconds)
 						for (int32 EventIdx = 0; EventIdx < EventContexts.Num(); ++EventIdx)
 						{
 							FNiagaraScriptExecutionContext& EventContext = EventContexts[EventIdx];
-							DataSetToEmitterEventParameters[EmitterIdx][EventIdx].DataSetToParameterStore(EventContext.Parameters, DataSet, SystemIndex);
+							if (DataSetToEmitterEventParameters[EmitterIdx].Num() > EventIdx)
+							{
+								DataSetToEmitterEventParameters[EmitterIdx][EventIdx].DataSetToParameterStore(EventContext.Parameters, DataSet, SystemIndex);
+							}
+							else
+							{
+								UE_LOG(LogNiagara, Log, TEXT("Skipping DataSetToEmitterEventParameters because EventIdx is out-of-bounds. %d of %d"), EventIdx, DataSetToEmitterEventParameters[EmitterIdx].Num());
+							}
 						}
 					}
 
@@ -435,14 +527,19 @@ void FNiagaraSystemSimulation::Tick(float DeltaSeconds)
 		{
 			FNiagaraSystemInstance* SystemInst = SystemInstances[SystemIndex];
 			++SystemIndex;
-			if (SystemIndex < SystemInstances.Num())
+			if (SystemInst->HasTickingEmitters())
 			{
-				FPlatformMisc::Prefetch(SystemInstances[SystemIndex]->GetComponent());
+				if (SystemIndex < SystemInstances.Num())
+				{
+					FPlatformMisc::Prefetch(SystemInstances[SystemIndex]->GetComponent());
+				}
+				SystemInst->GetComponent()->UpdateComponentToWorld();//Needed for bounds updates. Can probably skip if using fixed bounds
+				SystemInst->GetComponent()->MarkRenderDynamicDataDirty();
 			}
-			SystemInst->GetComponent()->UpdateComponentToWorld();
-			SystemInst->GetComponent()->MarkRenderDynamicDataDirty();
 		}
 	}
+
+	return true;
 }
 
 void FNiagaraSystemSimulation::RemoveInstance(FNiagaraSystemInstance* Instance)
@@ -526,19 +623,54 @@ void FNiagaraSystemSimulation::ResetSolo(FNiagaraSystemInstance* Instance)
 {
 	Instance->SetPendingSpawn(true);
 	Instance->SystemInstanceIndex = INDEX_NONE;
-	Instance->SetExecutionState(ENiagaraExecutionState::Active);
+}
+
+const TArray<UNiagaraDataInterface*>& FNiagaraSystemSimulation::GetSoloDataInterfacesSpawn()
+{
+	return SpawnExecContextSolo.GetDataInterfaces();
+}
+
+const TArray<UNiagaraDataInterface*>& FNiagaraSystemSimulation::GetSoloDataInterfacesUpdate()
+{
+	return UpdateExecContextSolo.GetDataInterfaces();
 }
 
 void FNiagaraSystemSimulation::TickSolo(FNiagaraSystemInstance* SystemInst)
 {
 	SCOPE_CYCLE_COUNTER(STAT_NiagaraSystemSimSolo);
-	check(System->IsValid());
+
+	UNiagaraSystem* System = WeakSystem.Get();
+	check(System != nullptr);
 
 	if (!System->GetSystemSpawnScript(true)->IsValid() || !System->GetSystemUpdateScript(true)->IsValid())
 	{
 		return;
 	}
 
+	auto SetExecStateFromExternal = [&](int32 SystemIndex, bool bCurr)
+	{
+		SoloSystemExecutionStateAccessor.InitForAccess(bCurr);
+		for (int32 EmitterIdx = 0; EmitterIdx < System->GetNumEmitters(); ++EmitterIdx)
+		{
+			SoloEmitterExecutionStateAccessors[EmitterIdx].InitForAccess(bCurr);
+		}
+		
+		// Update the individual parameters that may have been manually set.
+		if (SoloSystemExecutionStateAccessor.BaseIsValid())
+		{
+			SoloSystemExecutionStateAccessor.Set(SystemIndex, (int32)SystemInst->GetExecutionState());
+		}
+		TArray<TSharedRef<FNiagaraEmitterInstance>>& Emitters = SystemInst->GetEmitters();
+		for (int32 EmitterIdx = 0; EmitterIdx < Emitters.Num(); ++EmitterIdx)
+		{
+			FNiagaraEmitterInstance& EmitterInst = Emitters[EmitterIdx].Get();
+			if (SoloEmitterExecutionStateAccessors.Num() > EmitterIdx && SoloEmitterExecutionStateAccessors[EmitterIdx].BaseIsValid())
+			{
+				SoloEmitterExecutionStateAccessors[EmitterIdx].Set(SystemIndex, (int32)EmitterInst.GetExecutionState());
+			}
+		}
+	};
+
 	FNiagaraWorldManager* WorldMan = FNiagaraWorldManager::Get(World);
 	check(WorldMan);
 
@@ -563,14 +695,16 @@ void FNiagaraSystemSimulation::TickSolo(FNiagaraSystemInstance* SystemInst)
 		}
 
 		check(SoloSystemInstances.IsValidIndex(SystemIndex));
+
+		SetExecStateFromExternal(SystemIndex, true);
+
 		//////////////////////////////////////////////////////////////////////////
 		//TODO: This can be optimized as all the offests are the same. Can skip all the searching involved with a full bind().
-		//Unbind parameter collections from the batched exec contexts.
-		for (UNiagaraParameterCollection* Collection : System->GetSystemSpawnScript(true)->ParameterCollections)
-		{
-			WorldMan->GetParameterCollection(Collection)->GetParameterStore().Bind(&SpawnExecContextSolo.Parameters);
-		}
 		SystemInst->GetParameters().Bind(&SpawnExecContextSolo.Parameters);
+
+		//We must tick the RI parameters here to be sure we have up to date data. NPC data can afford to be a frame delayed but not RI data.
+		SpawnExecContextSolo.Script->RapidIterationParameters.Tick();
+
 		SpawnExecContextSolo.Tick(SystemInst);
 		//////////////////////////////////////////////////////////////////////////
 		
@@ -579,10 +713,6 @@ void FNiagaraSystemSimulation::TickSolo(FNiagaraSystemInstance* SystemInst)
 		//////////////////////////////////////////////////////////////////////////
 		//TODO: This can be optimized as all the offests are the same. Can skip all the searching involved with a full bind().
 		SystemInst->GetParameters().Unbind(&SpawnExecContextSolo.Parameters);
-		for (UNiagaraParameterCollection* Collection : System->GetSystemSpawnScript(true)->ParameterCollections)
-		{
-			WorldMan->GetParameterCollection(Collection)->GetParameterStore().Unbind(&SpawnExecContextSolo.Parameters);
-		}
 		//////////////////////////////////////////////////////////////////////////
 		
 		DataSetSolo.SetNumInstances(SoloSystemInstances.Num());
@@ -596,13 +726,14 @@ void FNiagaraSystemSimulation::TickSolo(FNiagaraSystemInstance* SystemInst)
 	else
 	{
 		check(SoloSystemInstances.IsValidIndex(SystemIndex));
+
+		SetExecStateFromExternal(SystemIndex, true);
+
 		//////////////////////////////////////////////////////////////////////////
 		//TODO: This can be optimized as all the offests are the same. Can skip all the searching involved with a full bind().
-		for (UNiagaraParameterCollection* Collection : System->GetSystemUpdateScript(true)->ParameterCollections)
-		{
-			WorldMan->GetParameterCollection(Collection)->GetParameterStore().Bind(&SpawnExecContextSolo.Parameters);
-		}
 		SystemInst->GetParameters().Bind(&UpdateExecContextSolo.Parameters);
+		//We must tick the RI parameters here to be sure we have up to date data. NPC data can afford to be a frame delayed but not RI data.
+		UpdateExecContextSolo.Script->RapidIterationParameters.Tick();
 		UpdateExecContextSolo.Tick(SystemInst);
 		//////////////////////////////////////////////////////////////////////////
 
@@ -611,10 +742,6 @@ void FNiagaraSystemSimulation::TickSolo(FNiagaraSystemInstance* SystemInst)
 		//////////////////////////////////////////////////////////////////////////
 		//TODO: This can be optimized as all the offests are the same. Can skip all the searching involved with a full bind().
 		SystemInst->GetParameters().Unbind(&UpdateExecContextSolo.Parameters);
-		for (UNiagaraParameterCollection* Collection : System->GetSystemUpdateScript(true)->ParameterCollections)
-		{
-			WorldMan->GetParameterCollection(Collection)->GetParameterStore().Unbind(&UpdateExecContextSolo.Parameters);
-		}
 		//////////////////////////////////////////////////////////////////////////
 
 		DataSetSolo.SetNumInstances(SoloSystemInstances.Num());
@@ -628,12 +755,9 @@ void FNiagaraSystemSimulation::TickSolo(FNiagaraSystemInstance* SystemInst)
 
 	InitBindings(SystemInst);
 
-	//Kill instance if it's marked for delete.
-	SoloSystemEnabledAccessor.InitForAccess(true);
 	SoloSystemExecutionStateAccessor.InitForAccess(true);
 	for (int32 EmitterIdx = 0; EmitterIdx < System->GetNumEmitters(); ++EmitterIdx)
 	{
-		SoloEmitterEnabledAccessors[EmitterIdx].InitForAccess(true);
 		SoloEmitterExecutionStateAccessors[EmitterIdx].InitForAccess(true);
 		for (int32 SpawnInfoIdx = 0; SpawnInfoIdx < SoloEmitterSpawnInfoAccessors[EmitterIdx].Num(); ++SpawnInfoIdx)
 		{
@@ -642,48 +766,50 @@ void FNiagaraSystemSimulation::TickSolo(FNiagaraSystemInstance* SystemInst)
 	}		
 
 	ensure(SystemIndex != INDEX_NONE);
-	bool bSystemEnabled = SoloSystemEnabledAccessor.GetSafe(SystemIndex, true);
-	ENiagaraExecutionState SoloExecutionState = (ENiagaraExecutionState)SoloSystemExecutionStateAccessor.GetSafe(SystemIndex, (int32)ENiagaraExecutionState::Active);
+	ENiagaraExecutionState SoloExecutionState = (ENiagaraExecutionState)SoloSystemExecutionStateAccessor.GetSafe(SystemIndex, (int32)ENiagaraExecutionState::Disabled);
 	SystemInst->SetExecutionState(SoloExecutionState);
 
-	if (!bSystemEnabled || SoloExecutionState == ENiagaraExecutionState::Dead)
+	//Now pull data out of the simulation and drive the emitters with it.
+	TArray<TSharedRef<FNiagaraEmitterInstance>>& Emitters = SystemInst->GetEmitters();
+	for (int32 EmitterIdx = 0; EmitterIdx < Emitters.Num(); ++EmitterIdx)
 	{
-		RemoveInstance(SystemInst);
-		SystemInst->Disable();
-	}
-	else
-	{
-		//Now pull data out of the simulation and drive the emitters with it.
-		TArray<TSharedRef<FNiagaraEmitterInstance>>& Emitters = SystemInst->GetEmitters();
-		for (int32 EmitterIdx = 0; EmitterIdx < Emitters.Num(); ++EmitterIdx)
+		FNiagaraEmitterInstance& EmitterInst = Emitters[EmitterIdx].Get();
+
+		//Early exit before we set the state as if we're complete or disabled we should never let the emitter turn itself back. It needs to be reset/reinited manually.
+		if (EmitterInst.IsComplete())
 		{
-			FNiagaraEmitterInstance& EmitterInst = Emitters[EmitterIdx].Get();
+			continue;
+		}
 
-			FNiagaraBool bEmitterEnabled = SoloEmitterEnabledAccessors[EmitterIdx].GetSafe(SystemIndex, true);
-			EmitterInst.SetEnabled(bEmitterEnabled);//TODO: Alter this to instruct emitter to dump existing particles optionally.
+		ENiagaraExecutionState State = (ENiagaraExecutionState)SoloEmitterExecutionStateAccessors[EmitterIdx].GetSafe(SystemIndex, (int32)ENiagaraExecutionState::Disabled);
+		EmitterInst.SetExecutionState(State);
 
-			ENiagaraExecutionState State = (ENiagaraExecutionState)SoloEmitterExecutionStateAccessors[EmitterIdx].GetSafe(SystemIndex, (int32)ENiagaraExecutionState::Active);
-			EmitterInst.SetExecutionState(State);
+		for (int32 SpawnInfoIdx = 0; SpawnInfoIdx < SoloEmitterSpawnInfoAccessors[EmitterIdx].Num(); ++SpawnInfoIdx)
+		{
+			EmitterInst.GetSpawnInfo()[SpawnInfoIdx] = SoloEmitterSpawnInfoAccessors[EmitterIdx][SpawnInfoIdx].Get(SystemIndex);
+		}
 
-			for (int32 SpawnInfoIdx = 0; SpawnInfoIdx < SoloEmitterSpawnInfoAccessors[EmitterIdx].Num(); ++SpawnInfoIdx)
-			{
-				EmitterInst.GetSpawnInfo()[SpawnInfoIdx] = SoloEmitterSpawnInfoAccessors[EmitterIdx][SpawnInfoIdx].Get(SystemIndex);
-			}
-			
-			//TODO: Any other fixed function stuff like this?
+		//TODO: Any other fixed function stuff like this?
 
-			FNiagaraScriptExecutionContext& SpawnContext = EmitterInst.GetSpawnExecutionContext();
-			DataSetToEmitterSpawnParameters[EmitterIdx].DataSetToParameterStore(SpawnContext.Parameters, DataSetSolo, SystemIndex);
+		FNiagaraScriptExecutionContext& SpawnContext = EmitterInst.GetSpawnExecutionContext();
+		DataSetToEmitterSpawnParameters[EmitterIdx].DataSetToParameterStore(SpawnContext.Parameters, DataSetSolo, SystemIndex);
 
-			FNiagaraScriptExecutionContext& UpdateContext = EmitterInst.GetUpdateExecutionContext();
-			DataSetToEmitterUpdateParameters[EmitterIdx].DataSetToParameterStore(UpdateContext.Parameters, DataSetSolo, SystemIndex);
+		FNiagaraScriptExecutionContext& UpdateContext = EmitterInst.GetUpdateExecutionContext();
+		DataSetToEmitterUpdateParameters[EmitterIdx].DataSetToParameterStore(UpdateContext.Parameters, DataSetSolo, SystemIndex);
 
-			TArray<FNiagaraScriptExecutionContext>& EventContexts = EmitterInst.GetEventExecutionContexts();
-			for (int32 EventIdx = 0; EventIdx < EventContexts.Num(); ++EventIdx)
+		TArray<FNiagaraScriptExecutionContext>& EventContexts = EmitterInst.GetEventExecutionContexts();
+		for (int32 EventIdx = 0; EventIdx < EventContexts.Num(); ++EventIdx)
+		{
+			FNiagaraScriptExecutionContext& EventContext = EventContexts[EventIdx];
+			if (DataSetToEmitterEventParameters[EmitterIdx].Num() > EventIdx)
 			{
-				FNiagaraScriptExecutionContext& EventContext = EventContexts[EventIdx];
 				DataSetToEmitterEventParameters[EmitterIdx][EventIdx].DataSetToParameterStore(EventContext.Parameters, DataSetSolo, SystemIndex);
 			}
+			else
+			{
+				UE_LOG(LogNiagara, Log, TEXT("Skipping DataSetToEmitterEventParameters because EventIdx is out-of-bounds. %d of %d"), EventIdx, DataSetToEmitterEventParameters[EmitterIdx].Num());
+			}
+
 		}
 	}
 }
@@ -696,8 +822,8 @@ void FNiagaraSystemSimulation::InitBindings(FNiagaraSystemInstance* SystemInst)
 	//Ideally we can store all this layout info in the systm/emitter assets so we can just generate this in Init()
 	if (DataSetToEmitterSpawnParameters.Num() == 0)
 	{
-		SpawnParameterToDataSetBinding.Init(SpawnParameterDataSet, SystemInst->GetInstanceParameters());
-		UpdateParameterToDataSetBinding.Init(UpdateParameterDataSet, SystemInst->GetInstanceParameters());
+		SpawnInstanceParameterToDataSetBinding.Init(SpawnInstanceParameterDataSet, SystemInst->GetInstanceParameters());
+		UpdateInstanceParameterToDataSetBinding.Init(UpdateInstanceParameterDataSet, SystemInst->GetInstanceParameters());
 
 		TArray<TSharedRef<FNiagaraEmitterInstance>>& Emitters = SystemInst->GetEmitters();
 		check(DataSetToEmitterUpdateParameters.Num() == 0);
@@ -724,6 +850,3 @@ void FNiagaraSystemSimulation::InitBindings(FNiagaraSystemInstance* SystemInst)
 		}
 	}
 }
-
-
-
diff --git a/Engine/Plugins/FX/Niagara/Source/Niagara/Private/NiagaraWorldManager.cpp b/Engine/Plugins/FX/Niagara/Source/Niagara/Private/NiagaraWorldManager.cpp
--- a/Engine/Plugins/FX/Niagara/Source/Niagara/Private/NiagaraWorldManager.cpp
+++ b/Engine/Plugins/FX/Niagara/Source/Niagara/Private/NiagaraWorldManager.cpp
@@ -10,6 +10,7 @@
 #include "NiagaraSystemInstance.h"
 #include "Scalability.h"
 #include "ConfigCacheIni.h"
+#include "NiagaraDataInterfaceSkeletalMesh.h"
 
 
 FNiagaraWorldManager* FNiagaraWorldManager::Get(UWorld* World)
@@ -25,11 +26,21 @@ void FNiagaraWorldManager::AddReferencedObjects(FReferenceCollector& Collector)
 
 UNiagaraParameterCollectionInstance* FNiagaraWorldManager::GetParameterCollection(UNiagaraParameterCollection* Collection)
 {
+	if (!Collection)
+	{
+		return nullptr;
+	}
+
 	UNiagaraParameterCollectionInstance** OverrideInst = ParameterCollections.Find(Collection);
 	if (!OverrideInst)
 	{
+		UNiagaraParameterCollectionInstance* DefaultInstance = Collection->GetDefaultInstance();
 		OverrideInst = &ParameterCollections.Add(Collection);
-		*OverrideInst = CastChecked<UNiagaraParameterCollectionInstance>(StaticDuplicateObject(Collection->GetDefaultInstance(), World));
+		*OverrideInst = CastChecked<UNiagaraParameterCollectionInstance>(StaticDuplicateObject(DefaultInstance, World));
+#if WITH_EDITORONLY_DATA
+		//Bind to the default instance so that changes to the collection propagate through.
+		DefaultInstance->GetParameterStore().Bind(&(*OverrideInst)->GetParameterStore());
+#endif
 	}
 
 	check(OverrideInst && *OverrideInst);
@@ -51,11 +62,17 @@ void FNiagaraWorldManager::SetParameterCollection(UNiagaraParameterCollectionIns
 		{
 			if (*OverrideInst && NewInstance)
 			{
+				UNiagaraParameterCollectionInstance* DefaultInstance = Collection->GetDefaultInstance();
 				//Need to transfer existing bindings from old instance to new one.
 				FNiagaraParameterStore& ExistingStore = (*OverrideInst)->GetParameterStore();
 				FNiagaraParameterStore& NewStore = NewInstance->GetParameterStore();
 
 				ExistingStore.TransferBindings(NewStore);
+
+#if WITH_EDITOR
+				//If the existing store was this world's duplicate of the default then we must be sure it's unbound.
+				DefaultInstance->GetParameterStore().Unbind(&ExistingStore);
+#endif
 			}
 		}
 
@@ -63,29 +80,55 @@ void FNiagaraWorldManager::SetParameterCollection(UNiagaraParameterCollectionIns
 	}
 }
 
-FNiagaraSystemSimulation* FNiagaraWorldManager::GetSystemSimulation(UNiagaraSystem* System)
+void FNiagaraWorldManager::CleanupParameterCollections()
 {
-	FNiagaraSystemSimulation* Sim = SystemSimulations.Find(System);
-	if (Sim)
+#if WITH_EDITOR
+	for (TPair<UNiagaraParameterCollection*, UNiagaraParameterCollectionInstance*> CollectionInstPair : ParameterCollections)
+	{
+		UNiagaraParameterCollection* Collection = CollectionInstPair.Key;
+		UNiagaraParameterCollectionInstance* CollectionInst = CollectionInstPair.Value;
+		//Ensure that the default instance is not bound to the override.
+		UNiagaraParameterCollectionInstance* DefaultInst = Collection->GetDefaultInstance();
+		DefaultInst->GetParameterStore().Unbind(&CollectionInst->GetParameterStore());
+	}
+#endif
+	ParameterCollections.Empty();
+}
+
+TSharedRef<FNiagaraSystemSimulation, ESPMode::ThreadSafe> FNiagaraWorldManager::GetSystemSimulation(UNiagaraSystem* System)
+{
+	TSharedRef<FNiagaraSystemSimulation, ESPMode::ThreadSafe>* SimPtr = SystemSimulations.Find(System);
+	if (SimPtr != nullptr)
 	{
-		return Sim;
+		return *SimPtr;
 	}
 	
-	Sim = &SystemSimulations.Add(System);
+	TSharedRef<FNiagaraSystemSimulation, ESPMode::ThreadSafe> Sim = MakeShared<FNiagaraSystemSimulation, ESPMode::ThreadSafe>();
+	SystemSimulations.Add(System, Sim);
 	Sim->Init(System, World);
 	return Sim;
 }
 
 void FNiagaraWorldManager::DestroySystemSimulation(UNiagaraSystem* System)
 {
-	FNiagaraSystemSimulation* Sim = SystemSimulations.Find(System);
+	TSharedRef<FNiagaraSystemSimulation, ESPMode::ThreadSafe>* Sim = SystemSimulations.Find(System);
 	if (Sim)
 	{
-		Sim->Destroy();
+		(*Sim)->Destroy();
 		SystemSimulations.Remove(System);
 	}	
 }
 
+void FNiagaraWorldManager::OnWorldCleanup(bool bSessionEnded, bool bCleanupResources)
+{
+	for (TPair<UNiagaraSystem*, TSharedRef<FNiagaraSystemSimulation, ESPMode::ThreadSafe>>& SimPair : SystemSimulations)
+	{
+		SimPair.Value->Destroy();
+	}
+	SystemSimulations.Empty();
+	CleanupParameterCollections();
+}
+
 void FNiagaraWorldManager::Tick(float DeltaSeconds)
 {
 /*
@@ -112,6 +155,7 @@ Going off this idea tbh
 		}
 	}
 */
+	SkeletalMeshGeneratedData.TickGeneratedData(DeltaSeconds);
 
 	//Tick our collections to push any changes to bound stores.
 	for (TPair<UNiagaraParameterCollection*, UNiagaraParameterCollectionInstance*> CollectionInstPair : ParameterCollections)
@@ -121,8 +165,17 @@ Going off this idea tbh
 	}
 
 	//Now tick all system instances. 
-	for (TPair<UNiagaraSystem*, FNiagaraSystemSimulation>& SystemSim : SystemSimulations)
+	TArray<UNiagaraSystem*> DeadSystems;
+	for (TPair<UNiagaraSystem*, TSharedRef<FNiagaraSystemSimulation, ESPMode::ThreadSafe>>& SystemSim : SystemSimulations)
+	{
+		if (SystemSim.Value->Tick(DeltaSeconds) == false)
+		{
+			DeadSystems.Add(SystemSim.Key);
+		}
+	}
+
+	for (UNiagaraSystem* DeadSystem : DeadSystems)
 	{
-		SystemSim.Value.Tick(DeltaSeconds);
+		SystemSimulations.Remove(DeadSystem);
 	}
 }
\ No newline at end of file
diff --git a/Engine/Plugins/FX/Niagara/Source/Niagara/Public/NiagaraActor.h b/Engine/Plugins/FX/Niagara/Source/Niagara/Public/NiagaraActor.h
--- a/Engine/Plugins/FX/Niagara/Source/Niagara/Public/NiagaraActor.h
+++ b/Engine/Plugins/FX/Niagara/Source/Niagara/Public/NiagaraActor.h
@@ -7,14 +7,14 @@
 #include "GameFramework/Actor.h"
 #include "NiagaraActor.generated.h"
 
-UCLASS(MinimalAPI)
+UCLASS(MinimalAPI, hideCategories = (Activation, "Components|Activation", Input, Collision, "Game|Damage"), ComponentWrapperClass)
 class ANiagaraActor : public AActor
 {
 	GENERATED_UCLASS_BODY()
 
 private:
 	/** Pointer to System component */
-	UPROPERTY(VisibleAnywhere, Category=NiagaraActor)
+	UPROPERTY(VisibleAnywhere, BlueprintReadOnly, Category=NiagaraActor, meta = (AllowPrivateAccess = "true"))
 	class UNiagaraComponent* NiagaraComponent;
 
 #if WITH_EDITORONLY_DATA
diff --git a/Engine/Plugins/FX/Niagara/Source/Niagara/Public/NiagaraCommon.h b/Engine/Plugins/FX/Niagara/Source/Niagara/Public/NiagaraCommon.h
--- a/Engine/Plugins/FX/Niagara/Source/Niagara/Public/NiagaraCommon.h
+++ b/Engine/Plugins/FX/Niagara/Source/Niagara/Public/NiagaraCommon.h
@@ -14,8 +14,8 @@ class UNiagaraEmitter;
 class FNiagaraSystemInstance;
 class UNiagaraParameterCollection;
 
-const uint32 NIAGARA_COMPUTE_THREADGROUP_SIZE = 16;
-const uint32 NIAGARA_MAX_COMPUTE_THREADGROUPS = 2048;
+const uint32 NIAGARA_COMPUTE_THREADGROUP_SIZE = 64;
+const uint32 NIAGARA_MAX_COMPUTE_THREADGROUPS = 16384;
 
 enum ENiagaraBaseTypes
 {
@@ -214,7 +214,13 @@ struct NIAGARA_API FNiagaraFunctionSignature
 
 	bool operator==(const FNiagaraFunctionSignature& Other) const
 	{
-		return Name.ToString().Equals(Other.Name.ToString()) && Inputs == Other.Inputs && Outputs == Other.Outputs && bRequiresContext == Other.bRequiresContext && bMemberFunction == Other.bMemberFunction && OwnerName == Other.OwnerName;
+		bool bNamesEqual = Name.ToString().Equals(Other.Name.ToString());
+		bool bInputsEqual = Inputs == Other.Inputs;
+		bool bOutputsEqual = Outputs == Other.Outputs;
+		bool bContextsEqual = bRequiresContext == Other.bRequiresContext;
+		bool bMemberFunctionsEqual = bMemberFunction == Other.bMemberFunction;
+		bool bOwnerNamesEqual = OwnerName == Other.OwnerName;
+		return bNamesEqual && bInputsEqual && bOutputsEqual && bContextsEqual && bMemberFunctionsEqual && bOwnerNamesEqual;
 	}
 
 	FString GetName()const { return Name.ToString(); }
@@ -261,6 +267,9 @@ public:
 	UPROPERTY()
 	int32 UserPtrIdx;
 
+	UPROPERTY()
+	FNiagaraTypeDefinition Type;
+
 	TArray<FNiagaraFunctionSignature> RegisteredFunctions;
 
 	//TODO: Allow data interfaces to own datasets
@@ -273,13 +282,13 @@ struct FNiagaraStatScope
 	GENERATED_USTRUCT_BODY();
 
 	FNiagaraStatScope() {}
-	FNiagaraStatScope(FName InFullName, FText InFriendlyName):FullName(InFullName), FriendlyName(InFriendlyName){}
+	FNiagaraStatScope(FName InFullName, FName InFriendlyName):FullName(InFullName), FriendlyName(InFriendlyName){}
 
 	UPROPERTY()
 	FName FullName;
 
 	UPROPERTY()
-	FText FriendlyName;
+	FName FriendlyName;
 
 	bool operator==(const FNiagaraStatScope& Other) const { return FullName == Other.FullName; }
 };
@@ -312,6 +321,7 @@ struct NIAGARA_API FNiagaraSystemUpdateContext
 	FNiagaraSystemUpdateContext(const UNiagaraScript* Script, bool bReInit) { Add(Script, bReInit); }
 	//FNiagaraSystemUpdateContext(UNiagaraDataInterface* Interface, bool bReinit) : Add(Interface, bReinit) {}
 	FNiagaraSystemUpdateContext(const UNiagaraParameterCollection* Collection, bool bReInit) { Add(Collection, bReInit); }
+	FNiagaraSystemUpdateContext() { }
 
 	~FNiagaraSystemUpdateContext();
 
@@ -321,15 +331,18 @@ struct NIAGARA_API FNiagaraSystemUpdateContext
 	//void Add(UNiagaraDataInterface* Interface, bool bReinit);
 	void Add(const UNiagaraParameterCollection* Collection, bool bReInit);
 
+	/** Adds all currently active systems.*/
+	void AddAll(bool bReInit);
 private:
 	void AddInternal(class UNiagaraComponent* Comp, bool bReInit);
 
-	FNiagaraSystemUpdateContext() { }
 	FNiagaraSystemUpdateContext(FNiagaraSystemUpdateContext& Other) { }
 
 	TArray<UNiagaraComponent*> ComponentsToReset;
 	TArray<UNiagaraComponent*> ComponentsToReInit;
 
+	TArray<UNiagaraSystem*> SystemSimsToDestroy;
+
 	//TODO: When we allow component less systems we'll also want to find and reset those.
 };
 
@@ -380,4 +393,11 @@ struct FNiagaraVariableInfo
 
 	UPROPERTY()
 	UNiagaraDataInterface* DataInterface;
+};
+
+namespace FNiagaraUtilities
+{
+	/** Builds a unique name from a candidate name and a set of existing names.  The candidate name will be made unique
+	if necessary by adding a 3 digit index to the end. */
+	FName NIAGARA_API GetUniqueName(FName CandidateName, const TSet<FName>& ExistingNames);
 };
\ No newline at end of file
