Copying Orion-Staging @ 3748653 (Orion/Dev-General @ 3722124) to //UE4/Main

diff --git a/Engine/Plugins/FX/Niagara/Source/Niagara/Public/NiagaraComponent.h b/Engine/Plugins/FX/Niagara/Source/Niagara/Public/NiagaraComponent.h
--- a/Engine/Plugins/FX/Niagara/Source/Niagara/Public/NiagaraComponent.h
+++ b/Engine/Plugins/FX/Niagara/Source/Niagara/Public/NiagaraComponent.h
@@ -18,13 +18,17 @@ class NiagaraRenderer;
 class UNiagaraSystem;
 class UNiagaraParameterCollection;
 class UNiagaraParameterCollectionInstance;
+class FNiagaraSystemSimulation;
+
+// Called when the particle system is done
+DECLARE_DYNAMIC_MULTICAST_DELEGATE_OneParam(FOnNiagaraSystemFinished, class UNiagaraComponent*, PSystem);
 
 /**
 * UNiagaraComponent is the primitive component for a Niagara System.
 * @see ANiagaraActor
 * @see UNiagaraSystem
 */
-UCLASS(ClassGroup = (Rendering, Common), hidecategories = (Object, "Components|Activation"), editinlinenew, meta = (BlueprintSpawnableComponent))
+UCLASS(ClassGroup = (Rendering, Common), hidecategories = Object, hidecategories = Physics, hidecategories = Collision, showcategories = Trigger, editinlinenew, meta = (BlueprintSpawnableComponent))
 class NIAGARA_API UNiagaraComponent : public UPrimitiveComponent
 {
 	GENERATED_UCLASS_BODY()
@@ -50,8 +54,8 @@ private:
 	TODO: This should be a minimal set of explicitly override parameters similar to how parameter collection instances override their collections store. 
 	Should expose anything in the "User" namespace.
 	*/
- 	UPROPERTY(EditAnywhere, Category = Parameters)
- 	FNiagaraParameterStore InitialParameters;
+	UPROPERTY(EditAnywhere, Category = Parameters)
+	FNiagaraParameterStore OverrideParameters;
 
 #if WITH_EDITORONLY_DATA
 	UPROPERTY()
@@ -77,6 +81,12 @@ private:
 	/** The delta time used when seeking to the desired age.  This is only relevant when using the DesiredAge age update mode. */
 	float SeekDelta;
 
+	UPROPERTY()
+	uint32 bAutoDestroy : 1;
+
+	UPROPERTY()
+	uint32 bRenderingEnabled : 1;
+
 	//~ Begin UActorComponent Interface.
 protected:
 	virtual void OnRegister() override;
@@ -84,6 +94,16 @@ protected:
 	virtual void SendRenderDynamicData_Concurrent() override;
 	virtual void BeginDestroy() override;
 public:
+	/**
+	* True if we should automatically attach to AutoAttachParent when activated, and detach from our parent when completed.
+	* This overrides any current attachment that may be present at the time of activation (deferring initial attachment until activation, if AutoAttachParent is null).
+	* When enabled, detachment occurs regardless of whether AutoAttachParent is assigned, and the relative transform from the time of activation is restored.
+	* This also disables attachment on dedicated servers, where we don't actually activate even if bAutoActivate is true.
+	* @see AutoAttachParent, AutoAttachSocketName, AutoAttachLocationType
+	*/
+	UPROPERTY(EditAnywhere, BlueprintReadOnly, Category = Attachment)
+		uint32 bAutoManageAttachment : 1;
+
 
 	virtual void Activate(bool bReset = false)override;
 	virtual void Deactivate()override;
@@ -100,10 +120,12 @@ public:
 	virtual void GetUsedMaterials(TArray<UMaterialInterface*>& OutMaterials, bool bGetDebugMaterials = false) const override;
 	//~ End UPrimitiveComponent Interface
 
-	class FNiagaraSystemSimulation* GetSystemSimulation();
+	TSharedPtr<FNiagaraSystemSimulation, ESPMode::ThreadSafe> GetSystemSimulation();
 
+	void OnSystemComplete();
 	void DestroyInstance();
 
+		UFUNCTION(BlueprintCallable, Category = Niagara, meta = (DisplayName = "Set Niagara System Asset"))
 	void SetAsset(UNiagaraSystem* InAsset);
 	UNiagaraSystem* GetAsset() const { return Asset; }
 
@@ -129,13 +151,16 @@ public:
 	when using the DesiredAge age update mode. */
 	void SetSeekDelta(float InSeekDelta);
 
+	void SetAutoDestroy(bool bInAutoDestroy) { bAutoDestroy = bInAutoDestroy; }
 	FNiagaraSystemInstance* GetSystemInstance() const;
-	
-	void OnSystemDisabled();
 
 	/** Returns true if this component forces it's instances to run in "Solo" mode. A sub optimal path required in some situations. */
 	bool ForcesSolo()const;
 
+	/** Sets a Niagara FLinearColor parameter by name, overriding locally if necessary.*/
+	UFUNCTION(BlueprintCallable, Category = Niagara, meta = (DisplayName = "Set Niagara Variable (LinearColor)"))
+	void SetNiagaraVariableLinearColor(const FString& InVariableName, const FLinearColor& InValue);
+
 	/** Sets a Niagara Vector3 parameter by name, overriding locally if necessary.*/
 	UFUNCTION(BlueprintCallable, Category = Niagara, meta = (DisplayName = "Set Niagara Variable (Vector4)"))
 	void SetNiagaraVariableVec4(const FString& InVariableName, const FVector4& InValue);
@@ -156,6 +181,9 @@ public:
 	UFUNCTION(BlueprintCallable, Category = Niagara, meta = (DisplayName = "Set Niagara Variable (Bool)"))
 	void SetNiagaraVariableBool(const FString& InVariableName, bool InValue);
 
+	UFUNCTION(BlueprintCallable, Category = Niagara, meta = (DisplayName = "Set Niagara Variable (Actor)"))
+	void SetNiagaraVariableActor(const FString& InVariableName, AActor* Actor) {} // TODO sckime ?????? fix this!!!!
+
 	/** Debug accessors for getting positions in blueprints. */
 	UFUNCTION(BlueprintCallable, Category = Niagara, meta = (DisplayName = "Get Niagara Emitter Positions"))
 	TArray<FVector> GetNiagaraParticlePositions_DebugOnly(const FString& InEmitterName);
@@ -167,6 +195,7 @@ public:
 	/** Debug accessors for getting a FVector attribute array in blueprints. */
 	UFUNCTION(BlueprintCallable, Category = Niagara, meta = (DisplayName = "Get Niagara Emitter Vec3 Attrib"))
 	TArray<FVector> GetNiagaraParticleValueVec3_DebugOnly(const FString& InEmitterName, const FString& InValueName);
+
 	/** Resets the System to it's initial pre-simulated state. */
 	UFUNCTION(BlueprintCallable, Category = Niagara, meta = (DisplayName = "Reset System"))
 	void ResetSystem();
@@ -192,17 +221,82 @@ public:
 	virtual bool SynchronizeWithSourceSystem();
 
 	bool IsParameterValueOverriddenLocally(const FName& InParamName);
-	void SetParameterValueOverriddenLocally(const FName& InParamName, bool bInOverridden);
+	void SetParameterValueOverriddenLocally(const FNiagaraVariable& InParam, bool bInOverridden);
 	
 	FOnSystemInstanceChanged& OnSystemInstanceChanged() { return OnSystemInstanceChangedDelegate; }
 #endif
 
-	FNiagaraParameterStore& GetInitialParameters() { return InitialParameters; }
+	FNiagaraParameterStore& GetOverrideParameters() { return OverrideParameters; }
 
 	//~ End UObject Interface.
-	
-	UPROPERTY()
-	bool bRenderingEnabled;
+
+	// Called when the particle system is done
+	UPROPERTY(BlueprintAssignable)
+	FOnNiagaraSystemFinished OnSystemFinished;
+
+public:
+	/**
+	 * Component we automatically attach to when activated, if bAutoManageAttachment is true.
+	 * If null during registration, we assign the existing AttachParent and defer attachment until we activate.
+	 * @see bAutoManageAttachment
+	 */
+	UPROPERTY(VisibleInstanceOnly, BlueprintReadWrite, Category=Attachment, meta=(EditCondition="bAutoManageAttachment"))
+	TWeakObjectPtr<USceneComponent> AutoAttachParent;
+
+	/**
+	 * Socket we automatically attach to on the AutoAttachParent, if bAutoManageAttachment is true.
+	 * @see bAutoManageAttachment
+	 */
+	UPROPERTY(EditAnywhere, BlueprintReadWrite, Category=Attachment, meta=(EditCondition="bAutoManageAttachment"))
+	FName AutoAttachSocketName;
+
+	/**
+	 * Options for how we handle our location when we attach to the AutoAttachParent, if bAutoManageAttachment is true.
+	 * @see bAutoManageAttachment, EAttachmentRule
+	 */
+	UPROPERTY(EditAnywhere, BlueprintReadWrite, Category = Attachment, meta = (EditCondition = "bAutoManageAttachment"))
+	EAttachmentRule AutoAttachLocationRule;
+
+	/**
+	 * Options for how we handle our rotation when we attach to the AutoAttachParent, if bAutoManageAttachment is true.
+	 * @see bAutoManageAttachment, EAttachmentRule
+	 */
+	UPROPERTY(EditAnywhere, BlueprintReadWrite, Category = Attachment, meta = (EditCondition = "bAutoManageAttachment"))
+	EAttachmentRule AutoAttachRotationRule;
+
+	/**
+	 * Options for how we handle our scale when we attach to the AutoAttachParent, if bAutoManageAttachment is true.
+	 * @see bAutoManageAttachment, EAttachmentRule
+	 */
+	UPROPERTY(EditAnywhere, BlueprintReadWrite, Category = Attachment, meta = (EditCondition = "bAutoManageAttachment"))
+	EAttachmentRule AutoAttachScaleRule;
+
+
+	/**
+	 * Set AutoAttachParent, AutoAttachSocketName, AutoAttachLocationRule, AutoAttachRotationRule, AutoAttachScaleRule to the specified parameters. Does not change bAutoManageAttachment; that must be set separately.
+	 * @param  Parent			Component to attach to. 
+	 * @param  SocketName		Socket on Parent to attach to.
+	 * @param  LocationRule		Option for how we handle our location when we attach to Parent.
+	 * @param  RotationRule		Option for how we handle our rotation when we attach to Parent.
+	 * @param  ScaleRule		Option for how we handle our scale when we attach to Parent.
+	 * @see bAutoManageAttachment, AutoAttachParent, AutoAttachSocketName, AutoAttachLocationRule, AutoAttachRotationRule, AutoAttachScaleRule
+	 */
+	UFUNCTION(BlueprintCallable, Category = Niagara)
+	void SetAutoAttachmentParameters(USceneComponent* Parent, FName SocketName, EAttachmentRule LocationRule, EAttachmentRule RotationRule, EAttachmentRule ScaleRule);
+
+private:
+
+	/** Did we auto attach during activation? Used to determine if we should restore the relative transform during detachment. */
+	uint32 bDidAutoAttach : 1;
+
+	/** Restore relative transform from auto attachment and optionally detach from parent (regardless of whether it was an auto attachment). */
+	void CancelAutoAttachment(bool bDetachFromParent);
+
+
+	/** Saved relative transform before auto attachment. Used during detachment to restore the transform if we had automatically attached. */
+	FVector SavedAutoAttachRelativeLocation;
+	FRotator SavedAutoAttachRelativeRotation;
+	FVector SavedAutoAttachRelativeScale3D;
 };
 
 
@@ -213,9 +307,10 @@ public:
 /**
 * Scene proxy for drawing niagara particle simulations.
 */
-class FNiagaraSceneProxy : public FPrimitiveSceneProxy
+class FNiagaraSceneProxy final : public FPrimitiveSceneProxy
 {
 public:
+	SIZE_T GetTypeHash() const override;
 
 	FNiagaraSceneProxy(const UNiagaraComponent* InComponent);
 	~FNiagaraSceneProxy();
diff --git a/Engine/Plugins/FX/Niagara/Source/Niagara/Public/NiagaraMeshRendererProperties.h b/Engine/Plugins/FX/Niagara/Source/Niagara/Public/NiagaraMeshRendererProperties.h
--- a/Engine/Plugins/FX/Niagara/Source/Niagara/Public/NiagaraMeshRendererProperties.h
+++ b/Engine/Plugins/FX/Niagara/Source/Niagara/Public/NiagaraMeshRendererProperties.h
@@ -8,6 +8,20 @@
 #include "StaticMeshResources.h"
 #include "NiagaraMeshRendererProperties.generated.h"
 
+/** This enum decides how a mesh particle will orient its "facing" axis relative to camera. Must keep these in sync with NiagaraMeshVertexFactory.ush*/
+UENUM()
+enum class ENiagaraMeshFacingMode : uint8
+{
+	/** Ignores the camera altogether. The mesh aligns its local-space X-axis with the particles' local-space X-axis, after transforming by the Particles.Transform vector (if it exists).*/
+	Default = 0,
+	/** The mesh aligns it's local-space X-axis with the particle's Particles.Velocity vector.*/
+	Velocity,
+	/** Has the mesh local-space X-axis point towards the camera's position.*/
+	CameraPosition, 
+	/** Has the mesh local-space X-axis point towards the closest point on the camera view plane.*/
+	CameraPlane
+};
+
 UCLASS(editinlinenew)
 class UNiagaraMeshRendererProperties : public UNiagaraRendererProperties
 {
@@ -29,6 +43,20 @@ public:
 
 	void PostEditChangeProperty(FPropertyChangedEvent& PropertyChangedEvent);
 
+	/** The static mesh to be instanced when rendering mesh particles. If OverrideMaterial is not specified, the mesh's material is used. Note that the material must have the Niagara Mesh Particles flag checked.*/
 	UPROPERTY(EditAnywhere, Category = "Mesh Rendering")
 	UStaticMesh *ParticleMesh;
+
+	/** Whether or not to use the OverrideMaterials array instead of the mesh's existing materials.*/
+	UPROPERTY(EditAnywhere, Category = "Mesh Rendering", meta = (InlineEditConditionToggle))
+	uint32 bOverrideMaterials : 1;
+
+	/** The materials to be used instead of the StaticMesh's materials. Note that each material must have the Niagara Mesh Particles flag checked. If the ParticleMesh 
+	requires more materials than exist in this array or any entry in this array is set to None, we will use the ParticleMesh's existing Material instead.*/
+	UPROPERTY(EditAnywhere, Category = "Mesh Rendering", meta = (EditCondition = "bOverrideMaterials"))
+	TArray<UMaterialInterface*> OverrideMaterials;
+
+	/** Determines how the mesh orients itself relative to the camera.*/
+	UPROPERTY(EditAnywhere, Category = "Mesh Rendering")
+	ENiagaraMeshFacingMode FacingMode;
 };
diff --git a/Engine/Plugins/FX/Niagara/Source/Niagara/Public/NiagaraModule.h b/Engine/Plugins/FX/Niagara/Source/Niagara/Public/NiagaraModule.h
--- a/Engine/Plugins/FX/Niagara/Source/Niagara/Public/NiagaraModule.h
+++ b/Engine/Plugins/FX/Niagara/Source/Niagara/Public/NiagaraModule.h
@@ -8,6 +8,7 @@
 #include "Engine/World.h"
 
 class FNiagaraWorldManager;
+class UNiagaraEmitter;
 
 /**
 * Niagara module interface
@@ -15,8 +16,39 @@ class FNiagaraWorldManager;
 class INiagaraModule : public IModuleInterface
 {
 public:
+#if WITH_EDITOR
+	struct FMergeEmitterResults
+	{
+		FMergeEmitterResults()
+			: bSucceeded(true)
+			, bModifiedGraph(false)
+		{
+		}
+
+		bool bSucceeded;
+		TArray<FText> ErrorMessages;
+		bool bModifiedGraph;
+		UNiagaraEmitter* MergedInstance;
+
+		FString GetErrorMessagesString() const
+		{
+			TArray<FString> ErrorMessageStrings;
+			for (FText ErrorMessage : ErrorMessages)
+			{
+				ErrorMessageStrings.Add(ErrorMessage.ToString());
+			}
+			return FString::Join(ErrorMessageStrings, TEXT("\n"));
+		}
+	};
+#endif
+
+public:
+#if WITH_EDITOR
+	DECLARE_DELEGATE_RetVal_ThreeParams(FMergeEmitterResults, FOnMergeEmitter, UNiagaraEmitter&, UNiagaraEmitter&, UNiagaraEmitter&);
+#endif
 	DECLARE_DELEGATE_RetVal(void, FOnProcessQueue);
 
+public:
 	virtual void StartupModule()override;
 	virtual void ShutdownModule()override;
 
@@ -39,8 +71,25 @@ public:
 	void OnPreWorldFinishDestroy(UWorld* World);
 
 	void TickWorld(UWorld* World, ELevelTick TickType, float DeltaSeconds);
+
+#if WITH_EDITOR
+	FMergeEmitterResults MergeEmitter(UNiagaraEmitter& Source, UNiagaraEmitter& LastMergedSource, UNiagaraEmitter& Instance);
+
+	FDelegateHandle NIAGARA_API RegisterOnMergeEmitter(FOnMergeEmitter OnMergeEmitter);
+
+	void NIAGARA_API UnregisterOnMergeEmitter(FDelegateHandle DelegateHandle);
+#endif
+
+	FORCEINLINE static int32 GetDetailLevel() { return EngineDetailLevel; }
 private:
 	TMap<class UWorld*, class FNiagaraWorldManager*> WorldManagers;
 	FOnProcessQueue OnProcessQueue;
+
+#if WITH_EDITORONLY_DATA
+	FOnMergeEmitter OnMergeEmitterDelegate;
+#endif
+
+	void OnChangeDetailLevel(class IConsoleVariable* CVar);
+	static int32 EngineDetailLevel;
 };
 
diff --git a/Engine/Plugins/FX/Niagara/Source/Niagara/Public/NiagaraParameterStore.h b/Engine/Plugins/FX/Niagara/Source/Niagara/Public/NiagaraParameterStore.h
--- a/Engine/Plugins/FX/Niagara/Source/Niagara/Public/NiagaraParameterStore.h
+++ b/Engine/Plugins/FX/Niagara/Source/Niagara/Public/NiagaraParameterStore.h
@@ -12,37 +12,73 @@ struct FNiagaraParameterStore;
 //This does no tracking of lifetimes etc so the owner must ensure safe use and rebinding when needed etc.
 struct FNiagaraParameterStoreBinding
 {
-	FNiagaraParameterStoreBinding()
-	{}
-
-	typedef TTuple<int32, int32, int32> TParameterBinding;
+	struct FParameterBinding
+	{
+		FParameterBinding(int32 InSrcOffset, int32 InDestOffset, int32 InSize)
+			: SrcOffset((uint16)InSrcOffset), DestOffset((uint16)InDestOffset), Size((uint16)InSize)
+		{
+			//No way these offsets could ever be higher but check just in case.
+			check(InSrcOffset < (int32)0xFFFF);
+			check(InDestOffset < (int32)0xFFFF);
+			check(InSize < (int32)0xFFFF);
+		}
+		FORCEINLINE bool operator==(const FParameterBinding& Other)const
+		{
+			return SrcOffset == Other.SrcOffset && DestOffset == Other.DestOffset && Size == Other.Size;
+		}
+		uint16 SrcOffset;
+		uint16 DestOffset;
+		uint16 Size;
+	};
 	/** Bindings of parameter data. Src offset, Dest offset and Size. */
-	TArray<TParameterBinding> ParmeterBindings;
+	TArray<FParameterBinding> ParameterBindings;
 
-	typedef TTuple<int32, int32> TInterfaceBinding;
+	struct FInterfaceBinding
+	{
+		FInterfaceBinding(int32 InSrcOffset, int32 InDestOffset) 
+			: SrcOffset((uint16)InSrcOffset), DestOffset((uint16)InDestOffset)
+		{
+			//No way these offsets could ever be higher but check just in case.
+			check(InSrcOffset < (int32)0xFFFF);
+			check(InDestOffset < (int32)0xFFFF);
+		}
+		FORCEINLINE bool operator==(const FInterfaceBinding& Other)const
+		{
+			return SrcOffset == Other.SrcOffset && DestOffset == Other.DestOffset;
+		}
+		uint16 SrcOffset;
+		uint16 DestOffset;
+	};
 	/** Bindings of data interfaces. Src and Dest offsets.*/
-	TArray<TInterfaceBinding> InterfaceBindings;
+	TArray<FInterfaceBinding> InterfaceBindings;
 
-	FORCEINLINE void Empty()
+	FNiagaraParameterStoreBinding()
 	{
-		ParmeterBindings.Empty();
-		InterfaceBindings.Empty();
 	}
-
-	FORCEINLINE_DEBUGGABLE void BindParameter(FNiagaraParameterStore& DestStore, int32 DestOffset, FNiagaraParameterStore& SrcStore, FNiagaraVariable& Parameter);
 	
-	/** TODO: Merge contiguous ranges into a single binding. */
+	FORCEINLINE_DEBUGGABLE void Empty(FNiagaraParameterStore* DestStore, FNiagaraParameterStore* SrcStore);
+	FORCEINLINE_DEBUGGABLE void Initialize(FNiagaraParameterStore* DestStore, FNiagaraParameterStore* SrcStore);
+	FORCEINLINE_DEBUGGABLE bool VerifyBinding(const FNiagaraParameterStore* DestStore, const FNiagaraParameterStore* SrcStore)const;
+	FORCEINLINE_DEBUGGABLE void Tick(FNiagaraParameterStore* DestStore, FNiagaraParameterStore* SrcStore, bool bForce = false);
+	FORCEINLINE_DEBUGGABLE void Dump(FNiagaraParameterStore* DestStore, FNiagaraParameterStore* SrcStore)const;
+	/** TODO: Merge contiguous ranges into a single binding? */
 	//FORCEINLINE_DEBUGGABLE void Optimize();
 
-	FORCEINLINE_DEBUGGABLE void Tick(FNiagaraParameterStore& SrcStore, FNiagaraParameterStore& DestStore, bool bForce = false);
+private:
+	void BindParameters(FNiagaraParameterStore* DestStore, FNiagaraParameterStore* SrcStore);
 };
 
 /** Base storage class for Niagara parameter values. */
 USTRUCT()
 struct NIAGARA_API FNiagaraParameterStore
 {
+#ifdef WITH_EDITOR
+	DECLARE_MULTICAST_DELEGATE(FOnChanged);
+#endif
+
 	GENERATED_USTRUCT_BODY()
 
+private:
 	/** Owner of this store. Used to provide an outer to data interfaces in this store. */
 	UPROPERTY()
 	UObject* Owner;
@@ -62,19 +98,39 @@ struct NIAGARA_API FNiagaraParameterStore
 	/** Bindings between this parameter store and others we push data into when we tick. */
 	TMap<FNiagaraParameterStore*, FNiagaraParameterStoreBinding> Bindings;
 
+	/** Parameter stores we've been bound to and are feeding data into us. */
+	TArray<FNiagaraParameterStore*> SourceStores;
+
 	/** Marks our parameters as dirty. They will be pushed to any bound stores on tick if true. */
 	uint32 bParametersDirty : 1;
 	/** Marks our interfaces as dirty. They will be pushed to any bound stores on tick if true. */
 	uint32 bInterfacesDirty : 1;
-	/** Marks our layout as dirty. All bindings must be recreated and all parameters pushed again. */
-	uint32 bLayoutDirty : 1;
 
+#ifdef WITH_EDITOR
+	FOnChanged OnChangedDelegate;
+#endif
+
+public:
 	FNiagaraParameterStore();
 	FNiagaraParameterStore(UObject* InOwner);
 	FNiagaraParameterStore(const FNiagaraParameterStore& Other);
 	FNiagaraParameterStore& operator=(const FNiagaraParameterStore& Other);
 
 	~FNiagaraParameterStore();
+	
+#if WITH_EDITORONLY_DATA
+	UPROPERTY()
+	FName Name;
+#endif
+
+	void Dump();
+	void DumpParameters();
+
+	FORCEINLINE uint32 GetParametersDirty() const { return bParametersDirty; }
+	FORCEINLINE uint32 GetInterfacesDirty() const { return bInterfacesDirty; }
+
+	FORCEINLINE void MarkParametersDirty() { bParametersDirty = true; }
+	FORCEINLINE void MarkInterfacesDirty() { bInterfacesDirty = true; }
 
 	/** Binds this parameter store to another. During Tick the values of this store will be pushed into all bound stores */
 	void Bind(FNiagaraParameterStore* DestStore);
@@ -86,6 +142,10 @@ struct NIAGARA_API FNiagaraParameterStore
 	void TransferBindings(FNiagaraParameterStore& OtherStore);
 	/** Handles any update such as pushing parameters to bound stores etc. */
 	void Tick();
+	/** Unbinds this store from all stores it's being driven by. */
+	void UnbindFromSourceStores();
+	
+	bool VerifyBinding(const FNiagaraParameterStore* InDestStore)const;
 
 	/**
 	Adds the passed parameter to this store.
@@ -103,11 +163,20 @@ struct NIAGARA_API FNiagaraParameterStore
 	/** Removes all parameters from this store and releases any data. */
 	void Empty(bool bClearBindings=true);
 
+	FORCEINLINE TArray<FNiagaraParameterStore*>& GetSourceParameterStores() { return SourceStores; }
+	FORCEINLINE const TMap<FNiagaraVariable, int32>& GetParameterOffests()const { return ParameterOffsets; }
+	/** Get the list of FNiagaraVariables referenced by this store. Note that the values will be stale and are not to be trusted directly. Get the Values using the offset specified by IndexOf or GetParameterValue.*/
 	FORCEINLINE void GetParameters(TArray<FNiagaraVariable>& OutParametres) const { return ParameterOffsets.GenerateKeyArray(OutParametres); }
+	FORCEINLINE int32 GetNumParameters() const { return ParameterOffsets.Num(); }
 
 	FORCEINLINE const TArray<UNiagaraDataInterface*>& GetDataInterfaces()const { return DataInterfaces; }
 	FORCEINLINE const TArray<uint8>& GetParameterDataArray()const { return ParameterData; }
-	FORCEINLINE TArray<uint8>& GetParameterDataArray() { return ParameterData; }
+
+	FORCEINLINE void SetParameterDataArray(const TArray<uint8>& InParameterDataArray)
+	{ 
+		ParameterData = InParameterDataArray;
+		OnParameterChange();
+	}
 
 	/** Gets the index of the passed parameter. If it is a data interface, this is an offset into the data interface table, otherwise a byte offset into he parameter data buffer. */
 	int32 IndexOf(const FNiagaraVariable& Parameter)const
@@ -160,22 +229,6 @@ struct NIAGARA_API FNiagaraParameterStore
 		return nullptr;
 	}
 
-	FORCEINLINE_DEBUGGABLE uint8* GetParameterData(const FNiagaraVariable& Parameter)
-	{
-		int32 Offset = IndexOf(Parameter);
-		if (Offset != INDEX_NONE)
-		{
-			return GetParameterData_Internal(Offset);
-		}
-		return nullptr;
-	}
-
-	/** Returns the data interface at the passed offset. */
-	FORCEINLINE void GetDataInterface(UNiagaraDataInterface* OutInterface, int32 Offset)const
-	{
-		OutInterface = GetDataInterface(Offset);
-	}
-
 	/** Returns the data interface at the passed offset. */
 	FORCEINLINE UNiagaraDataInterface* GetDataInterface(int32 Offset)const
 	{
@@ -199,6 +252,11 @@ struct NIAGARA_API FNiagaraParameterStore
 	/** Returns the associated FNiagaraVariable for the passed data interface if it exists in the store. Null if not.*/
 	const FNiagaraVariable* FindVariable(UNiagaraDataInterface* Interface)const;
 
+	FORCEINLINE_DEBUGGABLE const int32* FindParameterOffset(const FNiagaraVariable& Parameter) const
+	{
+		return ParameterOffsets.Find(Parameter);
+	}
+
 	/** Copies the passed parameter from this parameter store into another. */
 	FORCEINLINE_DEBUGGABLE void CopyParameterData(FNiagaraParameterStore& DestStore, const FNiagaraVariable& Parameter) const
 	{
@@ -206,12 +264,27 @@ struct NIAGARA_API FNiagaraParameterStore
 		int32 SrcIndex = IndexOf(Parameter);
 		if (DestIndex != INDEX_NONE && SrcIndex != INDEX_NONE)
 		{
-			DestStore.SetParameterData(GetParameterData_Internal(SrcIndex), DestIndex, Parameter.GetSizeInBytes());
+			if (Parameter.IsDataInterface())
+			{
+				DataInterfaces[SrcIndex]->CopyTo(DestStore.DataInterfaces[DestIndex]);
+				DestStore.OnInterfaceChange();
+			}
+			else
+			{
+				DestStore.SetParameterData(GetParameterData_Internal(SrcIndex), DestIndex, Parameter.GetSizeInBytes());
+			}
 		}
 	}
 
+	/** Copies all parameters from this parameter store into another.*/
+	void CopyParametersTo(FNiagaraParameterStore& DestStore, bool bOnlyAdd);
+
+	/** Remove all parameters from this parameter store from another.*/
+	void RemoveParameters(FNiagaraParameterStore& DestStore);
+
+
 	template<typename T>
-	FORCEINLINE_DEBUGGABLE void SetParameterValue(const T& InValue, const FNiagaraVariable& Param, bool bAdd=false)
+	FORCEINLINE_DEBUGGABLE bool SetParameterValue(const T& InValue, const FNiagaraVariable& Param, bool bAdd=false)
 	{
 		check(Param.GetSizeInBytes() == sizeof(T));
 		int32 Offset = IndexOf(Param);
@@ -221,6 +294,7 @@ struct NIAGARA_API FNiagaraParameterStore
 			//Until we solve our alignment issues, temporarily just doing a memcpy here.
 			FMemory::Memcpy(GetParameterData_Internal(Offset), &InValue, sizeof(T));
 			OnParameterChange();
+			return true;
 		}
 		else
 		{
@@ -230,8 +304,10 @@ struct NIAGARA_API FNiagaraParameterStore
 				Offset = IndexOf(Param);
 				check(Offset != INDEX_NONE);
 				*(T*)(GetParameterData_Internal(Offset)) = InValue;
+				return true;
 			}
 		}
+		return false;
 	}
 
 	FORCEINLINE_DEBUGGABLE void SetParameterData(const uint8* Data, int32 Offset, int32 Size)
@@ -282,11 +358,30 @@ struct NIAGARA_API FNiagaraParameterStore
 		}
 	}
 
-	FORCEINLINE void OnParameterChange() { bParametersDirty = true; }
-	FORCEINLINE void OnInterfaceChange() { bInterfacesDirty = true; }
+	FORCEINLINE void OnParameterChange() 
+	{ 
+		bParametersDirty = true;
+#ifdef WITH_EDITOR
+		OnChangedDelegate.Broadcast();
+#endif
+	}
+
+	FORCEINLINE void OnInterfaceChange() 
+	{ 
+		bInterfacesDirty = true;
+#ifdef WITH_EDITOR
+		OnChangedDelegate.Broadcast();
+#endif
+	}
+
+#ifdef WITH_EDITOR
+	FDelegateHandle AddOnChangedHandler(FOnChanged::FDelegate InOnChanged);
+	void RemoveOnChangedHandler(FDelegateHandle DelegateHandle);
+	void RemoveAllOnChangedHandlers(const void* InUserObject);
+#endif
 
 protected:
-	FORCEINLINE void OnLayoutChange() { bLayoutDirty = true; }
+	void OnLayoutChange();
 
 	/** Returns the parameter data at the passed offset. */
 	FORCEINLINE const uint8* GetParameterData_Internal(int32 Offset) const
@@ -300,26 +395,89 @@ protected:
 	}
 };
 
-
-FORCEINLINE_DEBUGGABLE void FNiagaraParameterStoreBinding::BindParameter(FNiagaraParameterStore& DestStore, int32 DestOffset, FNiagaraParameterStore& SrcStore, FNiagaraVariable& Parameter)
+FORCEINLINE_DEBUGGABLE void FNiagaraParameterStoreBinding::Empty(FNiagaraParameterStore* DestStore, FNiagaraParameterStore* SrcStore)
 {
-	int32 SrcOffset = SrcStore.IndexOf(Parameter);
-	checkSlow(DestOffset == DestStore.IndexOf(Parameter));
+	if (DestStore)
+	{
+		DestStore->GetSourceParameterStores().RemoveSingleSwap(SrcStore);
+	}
+	DestStore = nullptr;
+	ParameterBindings.Empty();
+	InterfaceBindings.Empty();
+}
 
-	if (SrcOffset == INDEX_NONE || DestOffset == INDEX_NONE)
+FORCEINLINE_DEBUGGABLE void FNiagaraParameterStoreBinding::BindParameters(FNiagaraParameterStore* DestStore, FNiagaraParameterStore* SrcStore)
+{
+	InterfaceBindings.Empty();
+	ParameterBindings.Empty();
+	for (const TPair<FNiagaraVariable, int32>& ParamOffsetPair : DestStore->GetParameterOffests())
 	{
-		//UE_LOG(LogNiagara, Error, TEXT("Failed to bind parameter %s."), *Parameter.GetName().ToString());
-		return;
+		const FNiagaraVariable& Parameter = ParamOffsetPair.Key;
+		int32 DestOffset = ParamOffsetPair.Value;
+		int32 SrcOffset = SrcStore->IndexOf(Parameter);
+
+		if (SrcOffset != INDEX_NONE && DestOffset != INDEX_NONE)
+		{
+			if (Parameter.IsDataInterface())
+			{
+				InterfaceBindings.Add(FInterfaceBinding(SrcOffset, DestOffset));
+			}
+			else
+			{
+				ParameterBindings.Add(FParameterBinding(SrcOffset, DestOffset, Parameter.GetSizeInBytes()));
+			}
+		}
 	}
 
-	if (Parameter.IsDataInterface())
+	//Force an initial tick to prime our values in the destination store.
+	Tick(DestStore, SrcStore, true);
+}
+
+FORCEINLINE_DEBUGGABLE void FNiagaraParameterStoreBinding::Initialize(FNiagaraParameterStore* DestStore, FNiagaraParameterStore* SrcStore)
+{
+	checkSlow(DestStore);
+	checkSlow(SrcStore);
+	DestStore->GetSourceParameterStores().AddUnique(SrcStore);
+
+	BindParameters(DestStore, SrcStore);
+}
+
+FORCEINLINE_DEBUGGABLE bool FNiagaraParameterStoreBinding::VerifyBinding(const FNiagaraParameterStore* DestStore, const FNiagaraParameterStore* SrcStore)const
+{
+	bool bBindingValid = true;
+#if WITH_EDITORONLY_DATA
+	TArray<FName, TInlineAllocator<32>> MissingParameterNames;
+	for (const TPair<FNiagaraVariable, int32>& ParamOffsetPair : DestStore->GetParameterOffests())
 	{
-		InterfaceBindings.AddUnique(TInterfaceBinding(SrcOffset, DestOffset));
+		const FNiagaraVariable& Parameter = ParamOffsetPair.Key;
+		int32 DestOffset = ParamOffsetPair.Value;
+		int32 SrcOffset = SrcStore->IndexOf(Parameter);
+
+		if (Parameter.IsDataInterface())
+		{
+			if (InterfaceBindings.Find(FInterfaceBinding(SrcOffset, DestOffset)) == INDEX_NONE)
+			{
+				MissingParameterNames.Add(Parameter.GetName());
+				bBindingValid = false;
+			}
+		}
+		else
+		{
+			if (ParameterBindings.Find(FParameterBinding(SrcOffset, DestOffset, Parameter.GetSizeInBytes())) == INDEX_NONE)
+			{
+				MissingParameterNames.Add(Parameter.GetName());
+				bBindingValid = false;
+			}
+		}
 	}
-	else
+
+	UE_LOG(LogNiagara, Warning, TEXT("Invalid ParameterStore Binding: Parameters missing from binding between %s and %s. Stores should have been rebound when one changed!"), *SrcStore->Name.ToString(), *DestStore->Name.ToString());
+	for (FName MissingParam : MissingParameterNames)
 	{
-		ParmeterBindings.AddUnique(TParameterBinding(SrcOffset, DestOffset, Parameter.GetSizeInBytes()));
+		UE_LOG(LogNiagara, Warning, TEXT("%s"), *MissingParam.ToString());
 	}
+#endif
+	return bBindingValid;
 }
 
 /** Merge contiguous ranges into a single binding. */
@@ -328,30 +486,96 @@ FORCEINLINE_DEBUGGABLE void FNiagaraParameterStoreBinding::BindParameter(FNiagar
 // 	//TODO
 // }
 
-FORCEINLINE_DEBUGGABLE void FNiagaraParameterStoreBinding::Tick(FNiagaraParameterStore& SrcStore, FNiagaraParameterStore& DestStore, bool bForce)
+FORCEINLINE_DEBUGGABLE void FNiagaraParameterStoreBinding::Tick(FNiagaraParameterStore* DestStore, FNiagaraParameterStore* SrcStore, bool bForce)
 {
-	if (SrcStore.bParametersDirty || bForce)
+	if (SrcStore->GetParametersDirty() || bForce)
 	{
-		DestStore.bParametersDirty = true;
-		for (TParameterBinding& Binding : ParmeterBindings)
+		for (FParameterBinding& Binding : ParameterBindings)
 		{
-			int32 SrcOffset = Binding.Get<0>();
-			int32 DestOffset = Binding.Get<1>();
-			int32 Size = Binding.Get<2>();
-			DestStore.SetParameterData(SrcStore.GetParameterData(SrcOffset), DestOffset, Size);
+			DestStore->SetParameterData(SrcStore->GetParameterData(Binding.SrcOffset), Binding.DestOffset, Binding.Size);
 		}
 	}
 
-	if (SrcStore.bInterfacesDirty || bForce)
+	if (SrcStore->GetInterfacesDirty() || bForce)
 	{
-		DestStore.bInterfacesDirty = true;
-		for (TInterfaceBinding& Binding : InterfaceBindings)
+		for (FInterfaceBinding& Binding : InterfaceBindings)
 		{
-			DestStore.SetDataInterface(SrcStore.GetDataInterface(Binding.Get<0>()), Binding.Get<1>());
+			DestStore->SetDataInterface(SrcStore->GetDataInterface(Binding.SrcOffset), Binding.DestOffset);
 		}
 	}
 }
 
+FORCEINLINE_DEBUGGABLE void FNiagaraParameterStoreBinding::Dump(FNiagaraParameterStore* DestStore, FNiagaraParameterStore* SrcStore)const
+{
+#if WITH_EDITORONLY_DATA
+	UE_LOG(LogNiagara, Log, TEXT("\n\nDest Store: %s\n"), *DestStore->Name.ToString());
+	for (const FParameterBinding& Binding : ParameterBindings)
+	{
+		ensure(Binding.Size != 0);
+		ensure(Binding.SrcOffset != -1);
+		ensure(Binding.DestOffset != -1);
+		FNiagaraVariable Param;
+		bool bFound = false;
+		for (const TPair<FNiagaraVariable, int32>& ParamOffsetPair : DestStore->GetParameterOffests())
+		{
+			if (ParamOffsetPair.Value == Binding.DestOffset && !ParamOffsetPair.Key.IsDataInterface())
+			{
+				bFound = true;
+				Param = ParamOffsetPair.Key;
+			}
+		}
+		if (ensure(bFound))
+		{
+			UE_LOG(LogNiagara, Log, TEXT("| Param | %s %s: Src:%d - Dest:%d - Size:%d\n"), *Param.GetType().GetName(), *Param.GetName().ToString(), Binding.SrcOffset, Binding.DestOffset, Binding.Size);
+
+			//Also ensure the param has been pushed correctly.
+			const uint8* SrcData = SrcStore->GetParameterData(Binding.SrcOffset);
+			const uint8* DestData = DestStore->GetParameterData(Binding.DestOffset);
+			if (!ensure(FMemory::Memcmp(SrcData, DestData, Binding.Size) == 0))
+			{
+				UE_LOG(LogNiagara, Log, TEXT("Parameter in dest store has incorrect data!\n"));
+			}
+		}
+		else
+		{
+			UE_LOG(LogNiagara, Log, TEXT("Failed to find matching param in bound store!\n"));
+		}
+	}
+
+	for (const FInterfaceBinding& Binding : InterfaceBindings)
+	{
+		ensure(Binding.SrcOffset != -1);
+		ensure(Binding.DestOffset != -1);
+		FNiagaraVariable Param;
+		bool bFound = false;
+		for (const TPair<FNiagaraVariable, int32>& ParamOffsetPair : DestStore->GetParameterOffests())
+		{
+			if (ParamOffsetPair.Value == Binding.DestOffset && ParamOffsetPair.Key.IsDataInterface())
+			{
+				bFound = true;
+				Param = ParamOffsetPair.Key;
+			}
+		}
+		if (ensure(bFound))
+		{
+			UE_LOG(LogNiagara, Log, TEXT("| DI | %s %s: Src:%d - Dest:%d\n"), *Param.GetType().GetName(), *Param.GetName().ToString(), Binding.SrcOffset, Binding.DestOffset);
+
+			//Also ensure the param has been pushed correctly.
+			const UNiagaraDataInterface* SrcData = SrcStore->GetDataInterfaces()[Binding.SrcOffset];
+			const UNiagaraDataInterface* DestData = DestStore->GetDataInterfaces()[Binding.DestOffset];
+			if (!ensure(SrcData == DestData))
+			{
+				UE_LOG(LogNiagara, Log, TEXT("Data interface parameter in dest store is incorrect!\n"));
+			}
+		}
+		else
+		{
+			UE_LOG(LogNiagara, Log, TEXT("Failed to find matching data interface param in bound store!\n"));
+		}
+	}
+#endif
+}
+
 //////////////////////////////////////////////////////////////////////////
 
 /**
@@ -361,14 +585,18 @@ Does no tracking of lifetimes etc so users are responsible for safety.
 template<typename T>
 struct FNiagaraParameterDirectBinding
 {
-	T* ValuePtr;
+	mutable T* ValuePtr;
+	FNiagaraParameterStore* BoundStore;
+	FNiagaraVariable BoundVariable;
 
 	FNiagaraParameterDirectBinding()
-		: ValuePtr(nullptr)
+		: ValuePtr(nullptr), BoundStore(nullptr)
 	{}
 
 	T* Init(FNiagaraParameterStore& InStore, const FNiagaraVariable& DestVariable)
 	{
+		BoundStore = &InStore;
+		BoundVariable = DestVariable;
 		check(DestVariable.GetSizeInBytes() == sizeof(T));
 		ValuePtr = (T*)InStore.GetParameterData(DestVariable);
 		return ValuePtr;
@@ -376,6 +604,9 @@ struct FNiagaraParameterDirectBinding
 
 	FORCEINLINE void SetValue(const T& InValue)
 	{
+		check(BoundVariable.GetSizeInBytes() == sizeof(T));
+		ValuePtr = (T*)BoundStore->GetParameterData(BoundVariable);
+
 		if (ValuePtr)
 		{
 			*ValuePtr = InValue;
@@ -384,6 +615,9 @@ struct FNiagaraParameterDirectBinding
 
 	FORCEINLINE T GetValue()const 
 	{
+		check(BoundVariable.GetSizeInBytes() == sizeof(T));
+		ValuePtr = (T*)BoundStore->GetParameterData(BoundVariable);
+
 		if (ValuePtr)
 		{
 			return *ValuePtr;
@@ -395,12 +629,17 @@ struct FNiagaraParameterDirectBinding
 template<>
 struct FNiagaraParameterDirectBinding<FMatrix>
 {
-	FMatrix* ValuePtr;
+	mutable FMatrix* ValuePtr;
+	FNiagaraParameterStore* BoundStore;
+	FNiagaraVariable BoundVariable;
 
-	FNiagaraParameterDirectBinding():ValuePtr(nullptr){}
+	FNiagaraParameterDirectBinding():ValuePtr(nullptr), BoundStore(nullptr) {}
 
 	FMatrix* Init(FNiagaraParameterStore& InStore, const FNiagaraVariable& DestVariable)
 	{
+		BoundStore = &InStore;
+		BoundVariable = DestVariable;
+
 		check(DestVariable.GetSizeInBytes() == sizeof(FMatrix));
 		ValuePtr = (FMatrix*)InStore.GetParameterData(DestVariable);
 		return ValuePtr;
@@ -408,6 +647,9 @@ struct FNiagaraParameterDirectBinding<FMatrix>
 
 	FORCEINLINE void SetValue(const FMatrix& InValue)
 	{
+		check(BoundVariable.GetSizeInBytes() == sizeof(FMatrix));
+		ValuePtr = (FMatrix*)BoundStore->GetParameterData(BoundVariable);
+
 		if (ValuePtr)
 		{
 			FMemory::Memcpy(ValuePtr, &InValue, sizeof(FMatrix));//Temp annoyance until we fix the alignment issues with parameter stores.
@@ -416,6 +658,9 @@ struct FNiagaraParameterDirectBinding<FMatrix>
 
 	FORCEINLINE FMatrix GetValue()const
 	{
+		check(BoundVariable.GetSizeInBytes() == sizeof(FMatrix));
+		ValuePtr = (FMatrix*)BoundStore->GetParameterData(BoundVariable);
+
 		FMatrix Ret;
 		if (ValuePtr)
 		{
@@ -428,12 +673,17 @@ struct FNiagaraParameterDirectBinding<FMatrix>
 template<>
 struct FNiagaraParameterDirectBinding<FVector4>
 {
-	FVector4* ValuePtr;
+	mutable FVector4* ValuePtr;
+	FNiagaraParameterStore* BoundStore;
+	FNiagaraVariable BoundVariable;
 
-	FNiagaraParameterDirectBinding() :ValuePtr(nullptr) {}
+	FNiagaraParameterDirectBinding() :ValuePtr(nullptr), BoundStore(nullptr){}
 
 	FVector4* Init(FNiagaraParameterStore& InStore, const FNiagaraVariable& DestVariable)
 	{
+		BoundStore = &InStore;
+		BoundVariable = DestVariable;
+
 		check(DestVariable.GetSizeInBytes() == sizeof(FVector4));
 		ValuePtr = (FVector4*)InStore.GetParameterData(DestVariable);
 		return ValuePtr;
@@ -441,6 +691,9 @@ struct FNiagaraParameterDirectBinding<FVector4>
 
 	FORCEINLINE void SetValue(const FVector4& InValue)
 	{
+		check(BoundVariable.GetSizeInBytes() == sizeof(FVector4));
+		ValuePtr = (FVector4*)BoundStore->GetParameterData(BoundVariable);
+
 		if (ValuePtr)
 		{
 			FMemory::Memcpy(ValuePtr, &InValue, sizeof(FVector4));//Temp annoyance until we fix the alignment issues with parameter stores.
@@ -449,6 +702,9 @@ struct FNiagaraParameterDirectBinding<FVector4>
 
 	FORCEINLINE FVector4 GetValue()const
 	{
+		check(BoundVariable.GetSizeInBytes() == sizeof(FVector4));
+		ValuePtr = (FVector4*)BoundStore->GetParameterData(BoundVariable);
+
 		FVector4 Ret;
 		if (ValuePtr)
 		{
diff --git a/Engine/Plugins/FX/Niagara/Source/Niagara/Public/NiagaraRenderer.h b/Engine/Plugins/FX/Niagara/Source/Niagara/Public/NiagaraRenderer.h
--- a/Engine/Plugins/FX/Niagara/Source/Niagara/Public/NiagaraRenderer.h
+++ b/Engine/Plugins/FX/Niagara/Source/Niagara/Public/NiagaraRenderer.h
@@ -80,7 +80,8 @@ public:
 		FPrimitiveViewRelevance Result;
 		bool bHasDynamicData = HasDynamicData();
 
-		Result.bDrawRelevance = bHasDynamicData && SceneProxy->IsShown(View) && View->Family->EngineShowFlags.Particles;
+		//Always draw so our LastRenderTime is updated. We may not have dynamic data if we're disabled from visibility culling.
+		Result.bDrawRelevance = /*bHasDynamicData &&*/ SceneProxy->IsShown(View) && View->Family->EngineShowFlags.Particles;
 		Result.bShadowRelevance = bHasDynamicData && SceneProxy->IsShadowCast(View);
 		Result.bDynamicRelevance = bHasDynamicData;
 		if (bHasDynamicData && View->Family->EngineShowFlags.Bounds)
@@ -123,17 +124,10 @@ public:
 	bool IsEnabled() const {return bEnabled;}
 
 	void SetEnabled(bool bInEnabled) { bEnabled = bInEnabled; }
-
+	
+	const FVector& GetBaseExtents() const {	return BaseExtents; }
 protected:
-	NiagaraRenderer()
-		: CPUTimeMS(0.0f)
-		, bLocalSpace(false)
-		, bEnabled(true)
-		, DynamicDataRender(nullptr)
-	{
-		Material = UMaterial::GetDefaultMaterial(MD_Surface);
-	}
-
+	NiagaraRenderer();
 	virtual ~NiagaraRenderer();
 
 	mutable float CPUTimeMS;
@@ -144,6 +138,8 @@ protected:
 	FMaterialRelevance MaterialRelevance;
 
 	struct FNiagaraDynamicDataBase *DynamicDataRender;
+
+	FVector BaseExtents;
 };
 
 
@@ -189,6 +185,20 @@ private:
 	UNiagaraSpriteRendererProperties *Properties;
 	mutable TUniformBuffer<FPrimitiveUniformShaderParameters> WorldSpacePrimitiveUniformBuffer;
 	class FNiagaraSpriteVertexFactory* VertexFactory;
+
+	int32 PositionOffset;
+	int32 VelocityOffset;
+	int32 RotationOffset;
+	int32 SizeOffset;
+	int32 ColorOffset;
+
+	int32 FacingOffset;
+	int32 AlignmentOffset;
+	int32 SubImageOffset;
+	int32 MaterialParamOffset;
+	int32 CameraOffsetOffset;
+	int32 UVScaleOffset;
+	int32 ParticleRandomOffset;
 };
 
 
diff --git a/Engine/Plugins/FX/Niagara/Source/Niagara/Public/NiagaraRendererProperties.h b/Engine/Plugins/FX/Niagara/Source/Niagara/Public/NiagaraRendererProperties.h
--- a/Engine/Plugins/FX/Niagara/Source/Niagara/Public/NiagaraRendererProperties.h
+++ b/Engine/Plugins/FX/Niagara/Source/Niagara/Public/NiagaraRendererProperties.h
@@ -7,6 +7,7 @@
 #include "UObject/Object.h"
 #include "RHIDefinitions.h"
 #include "NiagaraTypes.h"
+#include "NiagaraMergeable.h"
 #include "NiagaraRendererProperties.generated.h"
 /**
 * Emitter properties base class
@@ -20,7 +21,7 @@ class UMaterial;
 class UMaterialInterface;
 
 UCLASS(ABSTRACT)
-class NIAGARA_API UNiagaraRendererProperties : public UObject
+class NIAGARA_API UNiagaraRendererProperties : public UNiagaraMergeable
 {
 	GENERATED_BODY()
 
diff --git a/Engine/Plugins/FX/Niagara/Source/Niagara/Public/NiagaraRendererRibbons.h b/Engine/Plugins/FX/Niagara/Source/Niagara/Public/NiagaraRendererRibbons.h
--- a/Engine/Plugins/FX/Niagara/Source/Niagara/Public/NiagaraRendererRibbons.h
+++ b/Engine/Plugins/FX/Niagara/Source/Niagara/Public/NiagaraRendererRibbons.h
@@ -34,22 +34,32 @@ public:
 	/** Update render data buffer from attributes */
 	FNiagaraDynamicDataBase *GenerateVertexData(const FNiagaraSceneProxy* Proxy, FNiagaraDataSet &Data, const ENiagaraSimTarget Target) override;
 
-	void AddRibbonVert(TArray<FNiagaraRibbonVertex>& RenderData, FVector ParticlePos, FVector2D UV1,
-		const FLinearColor &Color, const float Age, const float Rotation, const FVector2D &Size)
+	void AddRibbonVert(TArray<FNiagaraRibbonVertex>& RenderData, FVector ParticlePos, FVector2D UV1, FVector2D UV2,
+		const FLinearColor &Color, const float Age, const float Rotation, const float Size, const FVector NormDir, const FVector CustomFacing)
 	{
-		FNiagaraRibbonVertex& NewVertex = *new(RenderData)FNiagaraRibbonVertex;
+		FNiagaraRibbonVertex NewVertex;
 		NewVertex.Position = ParticlePos;
-		NewVertex.OldPosition = NewVertex.Position;
+		NewVertex.Direction = NormDir;
 		NewVertex.Color = Color;
-		NewVertex.ParticleId = RenderData.Num();
-		NewVertex.RelativeTime = Age;
 		NewVertex.Size = Size;
 		NewVertex.Rotation = Rotation;
-		NewVertex.SubImageIndex = 0.f;
-		NewVertex.Tex_U = UV1.X;
-		NewVertex.Tex_V = UV1.Y;
-		NewVertex.Tex_U2 = UV1.X;
-		NewVertex.Tex_V2 = UV1.Y;
+		NewVertex.Tex_U = UV1.X * Properties->UV0Scale.X;
+		NewVertex.Tex_V = UV1.Y * Properties->UV0Scale.Y;
+		NewVertex.Tex_U2 = UV2.X * Properties->UV1Scale.X;
+		NewVertex.Tex_V2 = UV2.Y * Properties->UV1Scale.Y;
+		NewVertex.CustomFacingVector = CustomFacing;
+		//NewVertex.NormalizedAge = Age;
+		RenderData.Add(NewVertex);
+	}
+
+	void AddDynamicParam(TArray<FNiagaraRibbonVertexDynamicParameter>& ParamData, const FVector4& DynamicParam)
+	{
+		FNiagaraRibbonVertexDynamicParameter Param;
+		Param.DynamicValue[0] = DynamicParam.X;
+		Param.DynamicValue[1] = DynamicParam.Y;
+		Param.DynamicValue[2] = DynamicParam.Z;
+		Param.DynamicValue[3] = DynamicParam.W;
+		ParamData.Add(Param);
 	}
 
 
diff --git a/Engine/Plugins/FX/Niagara/Source/Niagara/Public/NiagaraRibbonRendererProperties.h b/Engine/Plugins/FX/Niagara/Source/Niagara/Public/NiagaraRibbonRendererProperties.h
--- a/Engine/Plugins/FX/Niagara/Source/Niagara/Public/NiagaraRibbonRendererProperties.h
+++ b/Engine/Plugins/FX/Niagara/Source/Niagara/Public/NiagaraRibbonRendererProperties.h
@@ -7,15 +7,45 @@
 #include "NiagaraRendererProperties.h"
 #include "NiagaraRibbonRendererProperties.generated.h"
 
+UENUM()
+enum class ENiagaraRibbonFacingMode : uint8
+{
+	Screen,
+	Custom
+};
+
 UCLASS(editinlinenew)
 class UNiagaraRibbonRendererProperties : public UNiagaraRendererProperties
 {
 public:
 	GENERATED_BODY()
 
-	UPROPERTY(EditAnywhere, Category = "Sprite Rendering")
+	UNiagaraRibbonRendererProperties()
+		: Material(nullptr)
+		, FacingMode(ENiagaraRibbonFacingMode::Screen)
+		, UV0TilingDistance(0.0f)
+		, UV0Scale(FVector2D(1.0f, 1.0f))
+		, UV1TilingDistance(0.0f)
+		, UV1Scale(FVector2D(1.0f, 1.0f))
+	{
+	}
+
+	UPROPERTY(EditAnywhere, Category = "Ribbon Rendering")
 	UMaterialInterface* Material;
 
+	UPROPERTY(EditAnywhere, Category = "Ribbon Rendering")
+	ENiagaraRibbonFacingMode FacingMode;
+
+	UPROPERTY(EditAnywhere, Category = "Ribbon Rendering")
+	float UV0TilingDistance;
+	UPROPERTY(EditAnywhere, Category = "Ribbon Rendering")
+	FVector2D UV0Scale;
+
+	UPROPERTY(EditAnywhere, Category = "Ribbon Rendering")
+	float UV1TilingDistance;
+	UPROPERTY(EditAnywhere, Category = "Ribbon Rendering")
+	FVector2D UV1Scale;
+
 	//~ UNiagaraRendererProperties interface
 	virtual NiagaraRenderer* CreateEmitterRenderer(ERHIFeatureLevel::Type FeatureLevel) override;
 	virtual void GetUsedMaterials(TArray<UMaterialInterface*>& OutMaterials) const override;
diff --git a/Engine/Plugins/FX/Niagara/Source/Niagara/Public/NiagaraSettings.h b/Engine/Plugins/FX/Niagara/Source/Niagara/Public/NiagaraSettings.h
--- a/Engine/Plugins/FX/Niagara/Source/Niagara/Public/NiagaraSettings.h
+++ b/Engine/Plugins/FX/Niagara/Source/Niagara/Public/NiagaraSettings.h
@@ -8,7 +8,7 @@
 #include "NiagaraSettings.generated.h"
 
 
-UCLASS(config = Engine, defaultconfig, meta=(DisplayName="Niagara"))
+UCLASS(config = Niagara, defaultconfig, meta=(DisplayName="Niagara"))
 class NIAGARA_API UNiagaraSettings : public UDeveloperSettings
 {
 	GENERATED_UCLASS_BODY()
diff --git a/Engine/Plugins/FX/Niagara/Source/Niagara/Public/NiagaraSpriteRendererProperties.h b/Engine/Plugins/FX/Niagara/Source/Niagara/Public/NiagaraSpriteRendererProperties.h
--- a/Engine/Plugins/FX/Niagara/Source/Niagara/Public/NiagaraSpriteRendererProperties.h
+++ b/Engine/Plugins/FX/Niagara/Source/Niagara/Public/NiagaraSpriteRendererProperties.h
@@ -11,26 +11,41 @@
 UENUM()
 enum class ENiagaraSortMode : uint8
 {
+	/** Perform no additional sorting prior to rendering.*/
 	SortNone,
+	/** Sort by depth to the camera's near plane.*/
 	SortViewDepth,
+	/** Sort by distance to the camera's origin.*/
 	SortViewDistance
 };
 
+/** This enum decides how a sprite particle will orient its "up" axis. Must keep these in sync with NiagaraSpriteVertexFactory.ush*/
 UENUM()
 enum class ENiagaraSpriteAlignment : uint8
 {
+	/** Only Particles.SpriteRotation and FacingMode impact the alignment of the particle.*/
 	Unaligned,
+	/** Imagine the particle texture having an arrow pointing up, this mode makes the arrow point in the direction of the Particles.Velocity attribute. FacingMode is ignored unless CustomFacingVector is set.*/
 	VelocityAligned,
+	/** Imagine the particle texture having an arrow pointing up, this mode makes the arrow point towards the axis defined by the "Particles.SpriteAlignment" attribute. FacingMode is ignored unless CustomFacingVector is set. If the "Particles.SpriteAlignment" attribute is missing, this falls back to Unaligned mode.*/
 	CustomAlignment
 };
 
 
+/** This enum decides how a sprite particle will orient its "facing" axis. Must keep these in sync with NiagaraSpriteVertexFactory.ush*/
 UENUM()
 enum class ENiagaraSpriteFacingMode : uint8
 {
+	/** The sprite billboard origin is always "looking at" the camera origin, trying to keep its up axis aligned to the camera's up axis. */
 	FaceCamera,
+	/** The sprite billboard plane is completely parallel to the camera plane. Particle always looks "flat" */
 	FaceCameraPlane,
-	CustomFacingVector
+	/** The sprite billboard faces toward the "Particles.SpriteFacing" vector attribute, using the per-axis CustomFacingVectorMask as a lerp factor from the standard FaceCamera mode. If the "Particles.SpriteFacing" attribute is missing, this falls back to FaceCamera mode.*/
+	CustomFacingVector,
+	/** Faces the camera position, but is not dependent on the camera rotation.  This method produces more stable particles under camera rotation. Uses the up axis of (0,0,1).*/
+	FaceCameraPosition,
+	/** Blends between FaceCamera and FaceCameraPosition.*/
+	FaceCameraDistanceBlend
 };
 
 
@@ -52,23 +67,49 @@ public:
 	virtual const TArray<FNiagaraVariable>& GetOptionalAttributes() override;
 #endif // WITH_EDITORONLY_DATA
 
+	/** The material used to render the particle. Note that it must have the Use with Niagara Sprites flag checked.*/
 	UPROPERTY(EditAnywhere, Category = "Sprite Rendering")
 	UMaterialInterface* Material;
-
-	UPROPERTY(EditAnywhere, Category = "Sprite Rendering")
-	FVector2D SubImageSize;
-
+	
+	/** Imagine the particle texture having an arrow pointing up, these modes define how the particle aligns that texture to other particle attributes.*/
 	UPROPERTY(EditAnywhere, Category = "Sprite Rendering")
 	ENiagaraSpriteAlignment Alignment;
 
+	/** Determines how the particle billboard orients itself relative to the camera.*/
 	UPROPERTY(EditAnywhere, Category = "Sprite Rendering")
 	ENiagaraSpriteFacingMode FacingMode;
 
+	/** Used as a per-axis interpolation factor with the CustomFacingVector mode to determine how the billboard orients itself relative to the camera. A value of 1.0 is fully facing the custom vector. A value of 0.0 uses the standard facing strategy.*/
 	UPROPERTY(EditAnywhere, Category = "Sprite Rendering")
 	FVector CustomFacingVectorMask;
 
+	/** Determines the location of the pivot point of this particle. It follows Unreal's UV space, which has the upper left of the image at 0,0 and bottom right at 1,1. The middle is at 0.5, 0.5. */
+	UPROPERTY(EditAnywhere, Category = "Sprite Rendering")
+	FVector2D PivotInUVSpace;
+
+	/** Determines how we sort the particles prior to rendering.*/
 	UPROPERTY(EditAnywhere, Category = "Sprite Rendering")
 	ENiagaraSortMode SortMode;
+	
+	/** When using SubImage lookups for particles, this variable contains the number of columns in X and the number of rows in Y.*/
+	UPROPERTY(EditAnywhere, Category = "Sprite Rendering")
+	FVector2D SubImageSize;
+
+	/** If true, blends the sub-image UV lookup with its next adjacent member using the fractional part of the SubImageIndex float value as the linear interpolation factor.*/
+	UPROPERTY(EditAnywhere, Category = "Sprite Rendering", meta = (DisplayName = "Sub UV Blending Enabled"))
+	uint32 bSubImageBlend : 1;
+
+	/** If true, removes the HMD view roll (e.g. in VR) */
+	UPROPERTY(EditAnywhere, Category = "Sprite Rendering", meta = (DisplayName = "Remove HMD Roll"))
+	uint32 bRemoveHMDRollInVR : 1;
+
+	/** The distance at which FacingCameraDistanceBlend	is fully FacingCamera */
+	UPROPERTY(EditAnywhere, Category = "Sprite Rendering", meta = (UIMin = "0"))
+	float MinFacingCameraBlendDistance;
+
+	/** The distance at which FacingCameraDistanceBlend	is fully FacingCameraPosition */
+	UPROPERTY(EditAnywhere, Category = "Sprite Rendering", meta = (UIMin = "0"))
+	float MaxFacingCameraBlendDistance;
 };
 
 
diff --git a/Engine/Plugins/FX/Niagara/Source/Niagara/Public/NiagaraSystemInstance.h b/Engine/Plugins/FX/Niagara/Source/Niagara/Public/NiagaraSystemInstance.h
--- a/Engine/Plugins/FX/Niagara/Source/Niagara/Public/NiagaraSystemInstance.h
+++ b/Engine/Plugins/FX/Niagara/Source/Niagara/Public/NiagaraSystemInstance.h
@@ -14,11 +14,13 @@
 
 class FNiagaraWorldManager;
 class UNiagaraComponent;
+class FNiagaraSystemInstance;
 
 class NIAGARA_API FNiagaraSystemInstance 
 {
 public:
 	DECLARE_MULTICAST_DELEGATE(FOnInitialized);
+	DECLARE_MULTICAST_DELEGATE_OneParam(FOnComplete, FNiagaraSystemInstance*);
 	
 #if WITH_EDITOR
 	DECLARE_MULTICAST_DELEGATE(FOnReset);
@@ -33,10 +35,11 @@ public:
 	virtual ~FNiagaraSystemInstance();
 
 	/** Initializes this System instance to simulate the supplied System. */
-	void Init(class FNiagaraSystemSimulation* InSystemSimulation, bool bForceReset = false, bool bInForceSolo=false);
+	void Init(TSharedRef<class FNiagaraSystemSimulation, ESPMode::ThreadSafe> InSystemSimulation, bool bInForceSolo=false);
 
 	void Activate(bool bReset = false);
 	void Deactivate(bool bImmediate = false);
+	void Complete();
 
 	//void RebindParameterCollection(UNiagaraParameterCollectionInstance* OldInstance, UNiagaraParameterCollectionInstance* NewInstance);
 	void BindParameters();
@@ -49,13 +52,12 @@ public:
 	/** Defines modes for resetting the System instance. */
 	enum class EResetMode
 	{
-		/** Defers resetting the System instance and simulations until the next tick. */
-		DeferredReset,
-		/** Resets the System instance and simulations immediately. */
-		ImmediateReset,
-		/** same as above, but reinitializes instead of fast resetting */
-		DeferredReInit,
-		ImmediateReInit
+		/** Resets the System instance and simulations. */
+		ResetAll,
+		/** Resets the System instance but not the simualtions */
+		ResetSystem,
+		/** Full reinitialization of the system and emitters.  */
+		ReInit
 	};
 	
 	/** Requests the the simulation be reset on the next tick. */
@@ -64,14 +66,14 @@ public:
 	void ComponentTick(float DeltaSeconds);
 	void PreSimulateTick(float DeltaSeconds);
 	void PostSimulateTick(float DeltaSeconds);
-	void HandleResets();
+	/** Handles completion of the system and returns true if the system is complete. */
+	bool HandleCompletion();
 
-	void Enable();
-	void Disable();
-	
 	ENiagaraExecutionState GetExecutionState() { return ExecutionState; }
 	void SetExecutionState(ENiagaraExecutionState InState);
 
+	FORCEINLINE bool IsComplete()const { return ExecutionState == ENiagaraExecutionState::Complete || ExecutionState == ENiagaraExecutionState::Disabled; }
+
 	/** Gets the simulation for the supplied emitter handle. */
 	TSharedPtr<FNiagaraEmitterInstance> GetSimulationForHandle(const FNiagaraEmitterHandle& EmitterHandle);
 
@@ -92,6 +94,9 @@ public:
 	/** Gets a multicast delegate which is called whenever this instance is initialized with an System asset. */
 	FOnInitialized& OnInitialized();
 
+	/** Gets a multicast delegate which is called whenever this instance is complete. */
+	FOnComplete& OnComplete();
+
 #if WITH_EDITOR
 	/** Gets a multicast delegate which is called whenever this instance is reset due to external changes in the source System asset. */
 	FOnReset& OnReset();
@@ -99,6 +104,10 @@ public:
 	FOnDestroyed& OnDestroyed();
 #endif
 
+#if WITH_EDITORONLY_DATA
+	bool GetIsolateEnabled() const;
+#endif
+
 	FName GetIDName() { return IDName; }
 
 	/** Queue up the data sources to have PrepareForSimulation called on them next tick.*/
@@ -126,15 +135,24 @@ public:
 
 	FORCEINLINE float GetAge()const { return Age; }
 	
-	FORCEINLINE FNiagaraSystemSimulation* GetSystemSimulation()const
+	FORCEINLINE TSharedPtr<FNiagaraSystemSimulation, ESPMode::ThreadSafe> GetSystemSimulation()const
 	{
-		check(SystemSimulation);
+		check(SystemSimulation.IsValid());
 		return SystemSimulation; 
 	}
 
+	FORCEINLINE bool IsActivated() const { return ExecutionState == ENiagaraExecutionState::Active; }
+
+	bool IsReadyToRun() const;
+
 	/** Index of this instance in the system simulation. */
 	int32 SystemInstanceIndex;
 
+	FORCEINLINE bool IsDeactivatedDueToError() const { return bError; }
+
+	FORCEINLINE bool HasTickingEmitters()const { return bHasTickingEmitters; }
+
+	UNiagaraParameterCollectionInstance* GetParameterCollectionInstance(UNiagaraParameterCollection* Collection);
 private:
 	/** Builds the emitter simulations. */
 	void InitEmitters();
@@ -142,7 +160,7 @@ private:
 	/** Resets the System, emitter simulations, and renderers to initial conditions. */
 	void ReInitInternal();
 	/** Resets for restart, assumes no change in emitter setup */
-	void ResetInternal();
+	void ResetInternal(bool bResetSimulations);
 
 	/** Updates the renders for the simulations. Gathers both the EmitterRenderers that were previously set as well as the ones that we  create within.*/
 	void UpdateRenderModules(ERHIFeatureLevel::Type InFeatureLevel, TArray<NiagaraRenderer*>& OutNewRenderers, TArray<NiagaraRenderer*>& OutOldRenderers);
@@ -154,14 +172,14 @@ private:
 	void InitDataInterfaces();
 	
 	/** Perform per-tick updates on data interfaces that need it.*/
-	void TickDataInterfaces(float DeltaSeconds);
+	void TickDataInterfaces(float DeltaSeconds, bool bPostSimulate);
 
 	void TickInstanceParameters(float DeltaSeconds);
 
 	void BindParameterCollections(FNiagaraScriptExecutionContext& ExecContext);
 	
 	UNiagaraComponent* Component;
-	class FNiagaraSystemSimulation *SystemSimulation;
+	TSharedPtr<class FNiagaraSystemSimulation, ESPMode::ThreadSafe> SystemSimulation;
 	FBox SystemBounds;
 
 	/** The age of the System instance. */
@@ -172,6 +190,7 @@ private:
 	TArray< TSharedRef<FNiagaraEmitterInstance> > Emitters;
 
 	FOnInitialized OnInitializedDelegate;
+	FOnComplete OnCompleteDelegate;
 
 #if WITH_EDITOR
 	FOnReset OnResetDelegate;
@@ -191,20 +210,27 @@ private:
 	FNiagaraParameterStore InstanceParameters;
 	
 	FNiagaraParameterDirectBinding<FVector> OwnerPositionParam;
+	FNiagaraParameterDirectBinding<FVector> OwnerScaleParam;
 	FNiagaraParameterDirectBinding<FVector> OwnerVelocityParam;
 	FNiagaraParameterDirectBinding<FVector> OwnerXAxisParam;
 	FNiagaraParameterDirectBinding<FVector> OwnerYAxisParam;
 	FNiagaraParameterDirectBinding<FVector> OwnerZAxisParam;
+
 	FNiagaraParameterDirectBinding<FMatrix> OwnerTransformParam;
 	FNiagaraParameterDirectBinding<FMatrix> OwnerInverseParam;
 	FNiagaraParameterDirectBinding<FMatrix> OwnerTransposeParam;
 	FNiagaraParameterDirectBinding<FMatrix> OwnerInverseTransposeParam;
+	FNiagaraParameterDirectBinding<FMatrix> OwnerTransformNoScaleParam;
+	FNiagaraParameterDirectBinding<FMatrix> OwnerInverseNoScaleParam;
+
 	FNiagaraParameterDirectBinding<float> OwnerDeltaSecondsParam;
 	FNiagaraParameterDirectBinding<float> OwnerInverseDeltaSecondsParam;
 	FNiagaraParameterDirectBinding<float> OwnerMinDistanceToCameraParam;
 	FNiagaraParameterDirectBinding<int32> SystemNumEmittersParam;
 	FNiagaraParameterDirectBinding<int32> SystemNumEmittersAliveParam;
 
+	FNiagaraParameterDirectBinding<float> SystemTimeSinceRenderedParam;
+
 	TArray<FNiagaraParameterDirectBinding<int32>> ParameterNumParticleBindings;
 
 	/** Indicates whether this instance must update itself rather than being batched up as most instances are. */
@@ -213,18 +239,15 @@ private:
 
 	uint32 bPendingSpawn : 1;
 
-	uint32 bActive : 1;
-
-	/** Flag to ensure the System instance is only reset once per frame. */
-	uint32 bResetPending : 1;
-	uint32 bReInitPending : 1;
-
 	/** Disable ticking and rendering if there was some serious error. */
 	uint32 bError : 1;
 
 	/** Notifier that data interfaces need reinitialization next tick.*/
 	uint32 bDataInterfacesNeedInit : 1;
 
+	/** If this instance has any currently ticking emitters. If false, allows us to skip some work. */
+	uint32 bHasTickingEmitters : 1;
+
 	/* System tick state */
 	ENiagaraExecutionState ExecutionState;
 };
diff --git a/Engine/Plugins/FX/Niagara/Source/Niagara/Public/NiagaraSystemSimulation.h b/Engine/Plugins/FX/Niagara/Source/Niagara/Public/NiagaraSystemSimulation.h
--- a/Engine/Plugins/FX/Niagara/Source/Niagara/Public/NiagaraSystemSimulation.h
+++ b/Engine/Plugins/FX/Niagara/Source/Niagara/Public/NiagaraSystemSimulation.h
@@ -42,7 +42,7 @@ struct FNiagaraParameterStoreToDataSetBinding
 		for (const FNiagaraVariable& Var : DataSet.GetVariables())
 		{
 			const FNiagaraVariableLayoutInfo* Layout = DataSet.GetVariableLayout(Var);
-			const int32* ParameterOffsetPtr = ParameterStore.ParameterOffsets.Find(Var);
+			const int32* ParameterOffsetPtr = ParameterStore.FindParameterOffset(Var);
 			if (ParameterOffsetPtr && Layout)
 			{
 				int32 ParameterOffset = *ParameterOffsetPtr;
@@ -65,7 +65,8 @@ struct FNiagaraParameterStoreToDataSetBinding
 	FORCEINLINE_DEBUGGABLE void DataSetToParameterStore(FNiagaraParameterStore& ParameterStore, FNiagaraDataSet& DataSet, int32 DataSetInstanceIndex)
 	{
 		FNiagaraDataBuffer& CurrBuffer = DataSet.CurrData();
-		uint8* ParameterData = ParameterStore.GetParameterDataArray().GetData();
+		TArray<uint8> ParameterDataArrayCopy = ParameterStore.GetParameterDataArray();
+		uint8* ParameterData = ParameterDataArrayCopy.GetData();
 
 		for (const FDataOffsets& DataOffsets : FloatOffsets)
 		{
@@ -79,12 +80,14 @@ struct FNiagaraParameterStoreToDataSetBinding
 			int32* DataSetPtr = CurrBuffer.GetInstancePtrInt32(DataOffsets.DataSetComponentOffset, DataSetInstanceIndex);
 			*ParamPtr = *DataSetPtr;
 		}
+
+		ParameterStore.SetParameterDataArray(ParameterDataArrayCopy);
 	}
 
 	FORCEINLINE_DEBUGGABLE void ParameterStoreToDataSet(FNiagaraParameterStore& ParameterStore, FNiagaraDataSet& DataSet, int32 DataSetInstanceIndex)
 	{
 		FNiagaraDataBuffer& CurrBuffer = DataSet.CurrData();
-		uint8* ParameterData = ParameterStore.GetParameterDataArray().GetData();
+		const uint8* ParameterData = ParameterStore.GetParameterDataArray().GetData();
 
 		for (const FDataOffsets& DataOffsets : FloatOffsets)
 		{
@@ -105,9 +108,10 @@ struct FNiagaraParameterStoreToDataSetBinding
 class FNiagaraSystemSimulation
 {
 public:
+	~FNiagaraSystemSimulation();
 	bool Init(UNiagaraSystem* InSystem, UWorld* InWorld);
 	void Destroy();
-	void Tick(float DeltaSeconds);
+	bool Tick(float DeltaSeconds);
 
 	void RemoveInstance(FNiagaraSystemInstance* Instance);
 	void AddInstance(FNiagaraSystemInstance* Instance);
@@ -118,11 +122,16 @@ public:
 
 	void TickSoloDataSet();
 
-	FORCEINLINE UNiagaraSystem* GetSystem()const { return System; }
+	FORCEINLINE UNiagaraSystem* GetSystem()const { return WeakSystem.Get(); }
+
+	const TArray<UNiagaraDataInterface*>& GetSoloDataInterfacesSpawn();
+	const TArray<UNiagaraDataInterface*>& GetSoloDataInterfacesUpdate();
+
+	UNiagaraParameterCollectionInstance* GetParameterCollectionInstance(UNiagaraParameterCollection* Collection);
 protected:
 
-	/** System of instances being simulated. No need for GC knowledge as all simulations will be cleaned up by the world manager if the system is invalid. */
-	UNiagaraSystem* System;
+	/** System of instances being simulated.  We use a weak object ptr here because once the last referencing object goes away this system may be come invalid at runtime. */
+	TWeakObjectPtr<UNiagaraSystem> WeakSystem;
 
 	/** World this system simulation belongs to. */
 	UWorld* World;
@@ -141,8 +150,8 @@ protected:
 	In some cases this might be a big waste of memory as there'll be duplicated data from a parameter store that's shared across all instances.
 	Though all these parameters can be unique per instance so for now lets just do the simple thing and improve later.
 	*/
-	FNiagaraDataSet SpawnParameterDataSet;
-	FNiagaraDataSet UpdateParameterDataSet;
+	FNiagaraDataSet SpawnInstanceParameterDataSet;
+ 	FNiagaraDataSet UpdateInstanceParameterDataSet;
 
 	FNiagaraScriptExecutionContext SpawnExecContext;
 	FNiagaraScriptExecutionContext UpdateExecContext;
@@ -151,10 +160,12 @@ protected:
 	FNiagaraScriptExecutionContext UpdateExecContextSolo;
 
 	/** Bindings that pull per component parameters into the spawn parameter dataset. */
-	FNiagaraParameterStoreToDataSetBinding SpawnParameterToDataSetBinding;
+	FNiagaraParameterStoreToDataSetBinding SpawnInstanceParameterToDataSetBinding;
 	/** Bindings that pull per component parameters into the update parameter dataset. */
-	FNiagaraParameterStoreToDataSetBinding UpdateParameterToDataSetBinding;
-	
+	FNiagaraParameterStoreToDataSetBinding UpdateInstanceParameterToDataSetBinding;
+	FNiagaraParameterStoreToDataSetBinding SpawnConstantParameterToDataSetBinding;
+	FNiagaraParameterStoreToDataSetBinding UpdateConstantParameterToDataSetBinding;
+
 	/** Binding to push system attributes into each emitter spawn parameters. */
 	TArray<FNiagaraParameterStoreToDataSetBinding> DataSetToEmitterSpawnParameters;
 	/** Binding to push system attributes into each emitter update parameters. */
@@ -171,16 +182,12 @@ protected:
 
 	void InitBindings(FNiagaraSystemInstance* SystemInst);
 
-	FNiagaraDataSetAccessor<FNiagaraBool> SystemEnabledAccessor;
 	FNiagaraDataSetAccessor<int32> SystemExecutionStateAccessor;
-	TArray<FNiagaraDataSetAccessor<FNiagaraBool>> EmitterEnabledAccessors;
 	TArray<TArray<FNiagaraDataSetAccessor<FNiagaraSpawnInfo>>> EmitterSpawnInfoAccessors;
 	TArray<FNiagaraDataSetAccessor<int32>> EmitterExecutionStateAccessors;
 
 	//Annoying duplicates required because these access the solo data set. When I rework the where the layout data for parameters and data sets live then these can go away.
-	FNiagaraDataSetAccessor<FNiagaraBool> SoloSystemEnabledAccessor;
 	FNiagaraDataSetAccessor<int32> SoloSystemExecutionStateAccessor;
-	TArray<FNiagaraDataSetAccessor<FNiagaraBool>> SoloEmitterEnabledAccessors;
 	TArray<TArray<FNiagaraDataSetAccessor<FNiagaraSpawnInfo>>> SoloEmitterSpawnInfoAccessors;
 	TArray<FNiagaraDataSetAccessor<int32>> SoloEmitterExecutionStateAccessors;
 
diff --git a/Engine/Plugins/FX/Niagara/Source/Niagara/Public/NiagaraTypes.h b/Engine/Plugins/FX/Niagara/Source/Niagara/Public/NiagaraTypes.h
--- a/Engine/Plugins/FX/Niagara/Source/Niagara/Public/NiagaraTypes.h
+++ b/Engine/Plugins/FX/Niagara/Source/Niagara/Public/NiagaraTypes.h
@@ -5,6 +5,7 @@
 #include "CoreMinimal.h"
 #include "UObject/UnrealType.h"
 #include "Engine/UserDefinedStruct.h"
+#include "SharedPointer.h"
 #include "NiagaraTypes.generated.h"
 
 DECLARE_LOG_CATEGORY_EXTERN(LogNiagara, Log, Verbose);
@@ -45,12 +46,19 @@ struct FNiagaraBool
 	void SetValue(bool bValue) { Value = bValue ? True : False; }
 	bool GetValue() const { return Value != False; }
 
+	/** Sets this niagara bool's raw integer value directly using the special raw integer values expected by the VM and HLSL. */
+	FORCEINLINE void SetRawValue(int32 RawValue) { Value = RawValue; }
+
+	/** Gets this niagara bools raw integer value expected by the VM and HLSL. */
+	FORCEINLINE int32 GetRawValue() const { return Value; }
+
 	bool IsValid() const { return Value == True || Value == False; }
 	
 	FNiagaraBool():Value(False) {}
 	FNiagaraBool(bool bInValue) : Value(bInValue ? True : False) {}
 	FORCEINLINE operator bool() { return GetValue(); }
 
+private:
 	UPROPERTY(EditAnywhere, Category = Parameters)//Parameters? These are used for attrs too. Must be either FNiagaraBool::True or FNiagaraBool::False.
 	int32 Value;
 };
@@ -117,7 +125,7 @@ struct FNiagaraMatrix
 };
 
 /** Data controlling the spawning of particles */
-USTRUCT(meta = (DisplayName = "Spawn Info"))
+USTRUCT(meta = (DisplayName = "Spawn Info", NiagaraClearEachFrame = "true"))
 struct FNiagaraSpawnInfo
 {
 	GENERATED_USTRUCT_BODY();
@@ -226,14 +234,30 @@ enum class ENiagaraExecutionState : uint32
 	Active,
 	/** Run all scripts but suppress any new spawning.*/
 	Inactive,
-	/** Scripts still run but all ticking of emitters is halted.*/
-	Paused,
-	/** Kill immediately.*/
-	Kill,
-	/** Dead*/
-	Dead,
-	/** Internal use only*/
-	Disabled UMETA(Hidden),
+	/** Clear all existing particles and move to inactive.*/
+	InactiveClear,
+	/** Complete. When the system or all emitters are complete the effect is considered finished. */
+	Complete,
+	/** Emitter only. Emitter is disabled. Will not tick or render again until a full re initialization of the system. */
+	Disabled,
+};
+
+USTRUCT()
+struct NIAGARA_API FNiagaraVariableMetaData
+{
+	GENERATED_USTRUCT_BODY()
+public:
+	UPROPERTY(EditAnywhere, Category = "Variable")
+	TMap<FName, FString> PropertyMetaData;
+
+	UPROPERTY(EditAnywhere, Category = "Variable")
+	FText Description;
+
+	UPROPERTY()
+	int32 CallSortPriority;
+
+	UPROPERTY()
+	TArray<TWeakObjectPtr<UObject>> ReferencerNodes;
 };
 
 USTRUCT()
@@ -283,7 +307,11 @@ public:
 
 	FText GetNameText()const
 	{
-		checkf(IsValid(), TEXT("Type definition is not valid."));
+		if (IsValid() == false)
+		{
+			return NSLOCTEXT("NiagaraTypeDefinition", "InvalidNameText", "Invalid (null type)");
+		}
+
 		if (Enum != nullptr)
 		{
 			return  FText::FromString(Enum->GetName());
@@ -298,7 +326,11 @@ public:
 
 	FString GetName()const
 	{
-		checkf(IsValid(), TEXT("Type definition is not valid."));
+		if (IsValid() == false)
+		{
+			return TEXT("Invalid");
+		}
+
 		if (Enum != nullptr)
 		{
 			return  Enum->GetName();
@@ -357,6 +389,12 @@ public:
 		}
 	}
 
+	bool IsFloatPrimitive() const
+	{
+		return Struct == FNiagaraTypeDefinition::GetFloatStruct() || Struct == FNiagaraTypeDefinition::GetVec2Struct() || Struct == FNiagaraTypeDefinition::GetVec3Struct() || Struct == FNiagaraTypeDefinition::GetVec4Struct() ||
+			Struct == FNiagaraTypeDefinition::GetMatrix4Struct() || Struct == FNiagaraTypeDefinition::GetColorStruct();
+ 	}
+
 	bool IsValid() const 
 	{ 
 		return Struct != nullptr;
@@ -406,6 +444,10 @@ public:
 	FString ToString(const uint8* ValueData)const
 	{
 		checkf(IsValid(), TEXT("Type definition is not valid."));
+		if (ValueData == nullptr)
+		{
+			return TEXT("(null)");
+		}
 		return FNiagaraTypeHelper::ToString(ValueData, CastChecked<UScriptStruct>(Struct));
 	}
 
@@ -568,11 +610,13 @@ struct FNiagaraVariable
 	{
 	}
 	
+	/** Check if Name and Type definition are the same. The actual stored value is not checked here.*/
 	bool operator==(const FNiagaraVariable& Other)const
 	{
 		return Name == Other.Name && TypeDef == Other.TypeDef;
 	}
 
+	/** Check if Name and Type definition are not the same. The actual stored value is not checked here.*/
 	bool operator!=(const FNiagaraVariable& Other)const
 	{
 		return !(*this == Other);
@@ -618,19 +662,13 @@ struct FNiagaraVariable
 	}
 
 	template<typename T>
-	T* GetValue()
-	{
-		check(sizeof(T) == TypeDef.GetSize());
-		check(IsDataAllocated());
-		return (T*)GetData();
-	}
-
-	template<typename T>
-	const T* GetValue() const
+	T GetValue() const
 	{
 		check(sizeof(T) == TypeDef.GetSize());
 		check(IsDataAllocated());
-		return (T*)GetData();
+		T Value;
+		FMemory::Memcpy(&Value, GetData(), TypeDef.GetSize());
+		return Value;
 	}
 
 	void SetData(const uint8* Data)
@@ -673,9 +711,9 @@ struct FNiagaraVariable
 		return Ret;
 	}
 
-	bool IsNameValid() const
+	bool IsValid() const
 	{
-		return Name != NAME_None;
+		return Name != NAME_None && TypeDef.IsValid();
 	}
 
 	FORCEINLINE bool IsInNameSpace(FString Namespace) 
@@ -683,6 +721,33 @@ struct FNiagaraVariable
 		return Name.ToString().StartsWith(Namespace + TEXT("."));
 	}
 
+	static FNiagaraVariable ResolveAliases(const FNiagaraVariable& InVar, const TMap<FString, FString>& InAliases, const TCHAR* InJoinSeparator = TEXT("."))
+	{
+		FNiagaraVariable OutVar = InVar;
+
+		FString OutVarStrName = InVar.GetName().ToString();
+		TArray<FString> SplitName;
+		OutVarStrName.ParseIntoArray(SplitName, TEXT("."));
+
+		for (int32 i = 0; i < SplitName.Num() - 1; i++)
+		{
+			TMap<FString, FString>::TConstIterator It = InAliases.CreateConstIterator();
+			while (It)
+			{
+				if (SplitName[i].Equals(It.Key()))
+				{
+					SplitName[i] = It.Value();
+				}
+				++It;
+			}
+		}
+
+		OutVarStrName = FString::Join<FString>(SplitName, InJoinSeparator);
+
+		OutVar.SetName(*OutVarStrName);
+		return OutVar;
+	}
+
 private:
 	UPROPERTY(EditAnywhere, Category = "Variable")
 	FName Name;
diff --git a/Engine/Plugins/FX/Niagara/Source/Niagara/Public/NiagaraWorldManager.h b/Engine/Plugins/FX/Niagara/Source/Niagara/Public/NiagaraWorldManager.h
--- a/Engine/Plugins/FX/Niagara/Source/Niagara/Public/NiagaraWorldManager.h
+++ b/Engine/Plugins/FX/Niagara/Source/Niagara/Public/NiagaraWorldManager.h
@@ -10,6 +10,8 @@
 #include "NiagaraScriptExecutionContext.h"
 #include "NiagaraSystemSimulation.h"
 
+#include "NiagaraDataInterfaceSkeletalMesh.h"
+
 class UWorld;
 class UNiagaraParameterCollection;
 class UNiagaraParameterCollectionInstance;
@@ -34,17 +36,25 @@ public:
 	
 	UNiagaraParameterCollectionInstance* GetParameterCollection(UNiagaraParameterCollection* Collection);
 	void SetParameterCollection(UNiagaraParameterCollectionInstance* NewInstance);
-	FNiagaraSystemSimulation* GetSystemSimulation(UNiagaraSystem* System);
+	void CleanupParameterCollections();
+	TSharedRef<FNiagaraSystemSimulation, ESPMode::ThreadSafe> GetSystemSimulation(UNiagaraSystem* System);
 	void DestroySystemSimulation(UNiagaraSystem* System);
 
 	void Tick(float DeltaSeconds);
+
+	void OnWorldCleanup(bool bSessionEnded, bool bCleanupResources);
+
+	FORCEINLINE FNDI_SkeletalMesh_GeneratedData& GetSkeletalMeshGeneratedData() { return SkeletalMeshGeneratedData; }
 private:
 	UWorld* World;
 
 	TMap<UNiagaraParameterCollection*, UNiagaraParameterCollectionInstance*> ParameterCollections;
 
-	TMap<UNiagaraSystem*, FNiagaraSystemSimulation> SystemSimulations;
+	TMap<UNiagaraSystem*, TSharedRef<FNiagaraSystemSimulation, ESPMode::ThreadSafe>> SystemSimulations;
 
 	int32 CachedEffectsQuality;
+
+	/** Generated data used by data interfaces*/
+	FNDI_SkeletalMesh_GeneratedData SkeletalMeshGeneratedData;
 };
 
diff --git a/Engine/Plugins/FX/Niagara/Source/NiagaraEditor/Private/AssetTypeActions/AssetTypeActions_NiagaraEmitter.cpp b/Engine/Plugins/FX/Niagara/Source/NiagaraEditor/Private/AssetTypeActions/AssetTypeActions_NiagaraEmitter.cpp
--- a/Engine/Plugins/FX/Niagara/Source/NiagaraEditor/Private/AssetTypeActions/AssetTypeActions_NiagaraEmitter.cpp
+++ b/Engine/Plugins/FX/Niagara/Source/NiagaraEditor/Private/AssetTypeActions/AssetTypeActions_NiagaraEmitter.cpp
@@ -2,10 +2,19 @@
 
 #include "AssetTypeActions_NiagaraEmitter.h"
 #include "NiagaraEmitter.h"
+#include "NiagaraSystem.h"
 #include "NiagaraSystemToolkit.h"
 #include "NiagaraEditorStyle.h"
+#include "NiagaraSystemFactoryNew.h"
+#include "NiagaraSystemViewModel.h"
+#include "NiagaraSystemScriptViewModel.h"
+#include "Framework/MultiBox/MultiBoxBuilder.h"
+#include "EditorStyleSet.h"
+#include "AssetEditorToolkit.h"
+#include "IContentBrowserSingleton.h"
+#include "ContentBrowserModule.h"
 
-class UNiagaraEmitter;
+#define LOCTEXT_NAMESPACE "AssetTypeActions"
 
 FColor FAssetTypeActions_NiagaraEmitter::GetTypeColor() const
 {
@@ -31,3 +40,64 @@ UClass* FAssetTypeActions_NiagaraEmitter::GetSupportedClass() const
 {
 	return UNiagaraEmitter::StaticClass();
 }
+
+void FAssetTypeActions_NiagaraEmitter::GetActions(const TArray<UObject*>& InObjects, FMenuBuilder& MenuBuilder)
+{
+	auto NiagaraEmitters = GetTypedWeakObjectPtrs<UNiagaraEmitter>(InObjects);
+
+	MenuBuilder.AddMenuEntry(
+		LOCTEXT("Emitter_NewNiagaraSystem", "Create Niagara System"),
+		LOCTEXT("Emitter_NewNiagaraSystemTooltip", "Creates a niagara system using this emitter as a base."),
+		FSlateIcon(FEditorStyle::GetStyleSetName(), "ClassIcon.ParticleSystem"),
+		FUIAction(
+			FExecuteAction::CreateSP(this, &FAssetTypeActions_NiagaraEmitter::ExecuteNewNiagaraSystem, NiagaraEmitters)
+		)
+	);
+}
+
+void FAssetTypeActions_NiagaraEmitter::ExecuteNewNiagaraSystem(TArray<TWeakObjectPtr<UNiagaraEmitter>> Objects)
+{
+	const FString DefaultSuffix = TEXT("_System");
+
+	FNiagaraSystemViewModelOptions SystemOptions;
+	SystemOptions.bCanModifyEmittersFromTimeline = true;
+	SystemOptions.bUseSystemExecStateForTimelineReset = true;
+	SystemOptions.EditMode = ENiagaraSystemViewModelEditMode::SystemAsset;
+
+	TArray<UObject*> ObjectsToSync;
+	for (auto ObjIt = Objects.CreateConstIterator(); ObjIt; ++ObjIt)
+	{
+		auto Emitter = (*ObjIt).Get();
+		if (Emitter)
+		{
+			// Determine an appropriate names
+			FString Name;
+			FString PackageName;
+			CreateUniqueAssetName(Emitter->GetOutermost()->GetName(), DefaultSuffix, PackageName, Name);
+
+			// Create the factory used to generate the asset
+			UNiagaraSystemFactoryNew* Factory = NewObject<UNiagaraSystemFactoryNew>();
+			FAssetToolsModule& AssetToolsModule = FModuleManager::GetModuleChecked<FAssetToolsModule>("AssetTools");
+			UObject* NewAsset = AssetToolsModule.Get().CreateAsset(Name, FPackageName::GetLongPackagePath(PackageName), UNiagaraSystem::StaticClass(), Factory);
+			
+			UNiagaraSystem* System = Cast<UNiagaraSystem>(NewAsset);
+			if (System != nullptr)
+			{
+				System->AddEmitterHandle(*Emitter, FName(*Emitter->GetName()));
+
+				FNiagaraSystemViewModel SystemViewModel = FNiagaraSystemViewModel(*System, SystemOptions);
+				SystemViewModel.GetSystemScriptViewModel()->RebuildEmitterNodes();
+
+				ObjectsToSync.Add(NewAsset);
+			}
+		}
+	}
+
+	if (ObjectsToSync.Num() > 0)
+	{
+		FContentBrowserModule& ContentBrowserModule = FModuleManager::LoadModuleChecked<FContentBrowserModule>("ContentBrowser");
+		ContentBrowserModule.Get().SyncBrowserToAssets(ObjectsToSync);
+	}
+}
+
+#undef LOCTEXT_NAMESPACE
\ No newline at end of file
diff --git a/Engine/Plugins/FX/Niagara/Source/NiagaraEditor/Private/AssetTypeActions/AssetTypeActions_NiagaraEmitter.h b/Engine/Plugins/FX/Niagara/Source/NiagaraEditor/Private/AssetTypeActions/AssetTypeActions_NiagaraEmitter.h
--- a/Engine/Plugins/FX/Niagara/Source/NiagaraEditor/Private/AssetTypeActions/AssetTypeActions_NiagaraEmitter.h
+++ b/Engine/Plugins/FX/Niagara/Source/NiagaraEditor/Private/AssetTypeActions/AssetTypeActions_NiagaraEmitter.h
@@ -13,6 +13,11 @@ public:
 	virtual FText GetName() const override { return NSLOCTEXT("AssetTypeActions", "AssetTypeActions_NiagaraEmitter", "Niagara Emitter"); }
 	virtual FColor GetTypeColor() const override;
 	virtual UClass* GetSupportedClass() const override;
+	virtual bool HasActions(const TArray<UObject*>& InObjects) const override { return true; }
+	virtual void GetActions( const TArray<UObject*>& InObjects, FMenuBuilder& MenuBuilder ) override;
 	virtual void OpenAssetEditor(const TArray<UObject*>& InObjects, TSharedPtr<class IToolkitHost> EditWithinLevelEditor = TSharedPtr<IToolkitHost>()) override;
 	virtual uint32 GetCategories() override { return FNiagaraEditorModule::GetAssetCategory(); }
+
+private:
+	void ExecuteNewNiagaraSystem(TArray<TWeakObjectPtr<UNiagaraEmitter>> Emitters);
 };
\ No newline at end of file
diff --git a/Engine/Plugins/FX/Niagara/Source/NiagaraEditor/Private/AssetTypeActions/AssetTypeActions_NiagaraScript.cpp b/Engine/Plugins/FX/Niagara/Source/NiagaraEditor/Private/AssetTypeActions/AssetTypeActions_NiagaraScript.cpp
--- a/Engine/Plugins/FX/Niagara/Source/NiagaraEditor/Private/AssetTypeActions/AssetTypeActions_NiagaraScript.cpp
+++ b/Engine/Plugins/FX/Niagara/Source/NiagaraEditor/Private/AssetTypeActions/AssetTypeActions_NiagaraScript.cpp
@@ -4,6 +4,12 @@
 #include "NiagaraScript.h"
 #include "NiagaraScriptToolkit.h"
 #include "NiagaraEditorStyle.h"
+#include "AssetData.h"
+
+const FName FAssetTypeActions_NiagaraScriptFunctions::NiagaraFunctionScriptName = FName("Niagara Function Script");
+const FName FAssetTypeActions_NiagaraScriptModules::NiagaraModuleScriptName = FName("Niagara Module Script");
+const FName FAssetTypeActions_NiagaraScriptDynamicInputs::NiagaraDynamicInputScriptName = FName("Niagara Dynamic Input Script");
+
 
 FColor FAssetTypeActions_NiagaraScript::GetTypeColor() const
 { 
@@ -25,6 +31,31 @@ void FAssetTypeActions_NiagaraScript::OpenAssetEditor( const TArray<UObject*>& I
 	}
 }
 
+FText FAssetTypeActions_NiagaraScript::GetDisplayNameFromAssetData(const FAssetData& AssetData) const
+{
+	const FString* Usage = AssetData.TagsAndValues.Find("Usage");
+	if (Usage != nullptr)
+	{
+		static const FString FunctionString = FString("Function");
+		static const FString ModuleString = FString("Module");
+		static const FString DynamicInputString = FString("DynamicInput");
+		if (*Usage == FunctionString)
+		{
+			return FAssetTypeActions_NiagaraScriptFunctions::GetFormattedName();
+		}
+		else if (*Usage == ModuleString)
+		{
+			return FAssetTypeActions_NiagaraScriptModules::GetFormattedName();
+		}
+		else if (*Usage == DynamicInputString)
+		{
+			return FAssetTypeActions_NiagaraScriptDynamicInputs::GetFormattedName();
+		}
+	}
+
+	return FAssetTypeActions_NiagaraScript::GetName();
+}
+
 UClass* FAssetTypeActions_NiagaraScript::GetSupportedClass() const
 { 
 	return UNiagaraScript::StaticClass(); 
diff --git a/Engine/Plugins/FX/Niagara/Source/NiagaraEditor/Private/AssetTypeActions/AssetTypeActions_NiagaraScript.h b/Engine/Plugins/FX/Niagara/Source/NiagaraEditor/Private/AssetTypeActions/AssetTypeActions_NiagaraScript.h
--- a/Engine/Plugins/FX/Niagara/Source/NiagaraEditor/Private/AssetTypeActions/AssetTypeActions_NiagaraScript.h
+++ b/Engine/Plugins/FX/Niagara/Source/NiagaraEditor/Private/AssetTypeActions/AssetTypeActions_NiagaraScript.h
@@ -10,10 +10,51 @@
 class FAssetTypeActions_NiagaraScript : public FAssetTypeActions_Base
 {
 public:
-	// IAssetTypeActions Implementation
+	// Begin IAssetTypeActions Interface
 	virtual FText GetName() const override { return NSLOCTEXT("AssetTypeActions", "AssetTypeActions_NiagaraScript", "Niagara Script"); }
 	virtual FColor GetTypeColor() const override;
 	virtual UClass* GetSupportedClass() const override;
 	virtual void OpenAssetEditor(const TArray<UObject*>& InObjects, TSharedPtr<class IToolkitHost> EditWithinLevelEditor = TSharedPtr<IToolkitHost>()) override;
 	virtual uint32 GetCategories() override { return FNiagaraEditorModule::GetAssetCategory(); }
+	virtual FText GetDisplayNameFromAssetData(const FAssetData& AssetData) const override;
+	// End IAssetTypeActions Interface
+
+	/** Get the action type name. */
+	virtual FName GetTypeName() const = 0;
+};
+
+class FAssetTypeActions_NiagaraScriptFunctions: public FAssetTypeActions_NiagaraScript
+{
+public:
+	// Begin IAssetTypeActions Interface
+	virtual FText GetName() const override { return GetFormattedName(); }
+	virtual FName GetTypeName() const override { return NiagaraFunctionScriptName; }
+	// End IAssetTypeActions Interface
+
+	static FText GetFormattedName() { return FText::Format(NSLOCTEXT("AssetTypeActions", "AssetTypeActions_NiagaraScriptFunctions", "{0}"), FText::FromName(NiagaraFunctionScriptName)); }
+	static const FName NiagaraFunctionScriptName;
+};
+
+class FAssetTypeActions_NiagaraScriptModules: public FAssetTypeActions_NiagaraScript
+{
+public:
+	// Begin IAssetTypeActions Interface
+	virtual FText GetName() const override { return GetFormattedName(); }
+	virtual FName GetTypeName() const override { return NiagaraModuleScriptName; }
+	// End IAssetTypeActions Interface
+
+	static FText GetFormattedName() { return FText::Format(NSLOCTEXT("AssetTypeActions", "AssetTypeActions_NiagaraScriptModules", "{0}"), FText::FromName(NiagaraModuleScriptName)); }
+	static const FName NiagaraModuleScriptName;
+};
+
+class FAssetTypeActions_NiagaraScriptDynamicInputs: public FAssetTypeActions_NiagaraScript
+{
+public:
+	// Begin IAssetTypeActions Interface
+	virtual FText GetName() const override { return GetFormattedName(); }
+	virtual FName GetTypeName() const override { return NiagaraDynamicInputScriptName; }
+	// End IAssetTypeActions Interface
+
+	static FText GetFormattedName() { return FText::Format(NSLOCTEXT("AssetTypeActions", "AssetTypeActions_NiagaraScriptDynamicInputs", "{0}"), FText::FromName(NiagaraDynamicInputScriptName)); }
+	static const FName NiagaraDynamicInputScriptName;
 };
diff --git a/Engine/Plugins/FX/Niagara/Source/NiagaraEditor/Private/Customizations/NiagaraComponentDetails.cpp b/Engine/Plugins/FX/Niagara/Source/NiagaraEditor/Private/Customizations/NiagaraComponentDetails.cpp
--- a/Engine/Plugins/FX/Niagara/Source/NiagaraEditor/Private/Customizations/NiagaraComponentDetails.cpp
+++ b/Engine/Plugins/FX/Niagara/Source/NiagaraEditor/Private/Customizations/NiagaraComponentDetails.cpp
@@ -79,7 +79,7 @@ public:
 	virtual void GenerateChildContent(IDetailChildrenBuilder& ChildrenBuilder) override
 	{
 		TArray<FNiagaraVariable> Parameters;
-		FNiagaraParameterStore& ParamStore = Component->GetInitialParameters();
+		FNiagaraParameterStore& ParamStore = Component->GetOverrideParameters();
 		ParamStore.GetParameters(Parameters);
 
 		FNiagaraEditorModule& NiagaraEditorModule = FModuleManager::GetModuleChecked<FNiagaraEditorModule>("NiagaraEditor");
@@ -106,11 +106,14 @@ public:
 			else 
 			{
 				int32 DataInterfaceOffset = ParamStore.IndexOf(Parameter);
-				UObject* DefaultValueObject = ParamStore.DataInterfaces[DataInterfaceOffset];
+				UObject* DefaultValueObject = ParamStore.GetDataInterfaces()[DataInterfaceOffset];
 
 				TArray<UObject*> Objects;
 				Objects.Add(DefaultValueObject);
-				Row = ChildrenBuilder.AddExternalObjects(Objects, Parameter.GetName());
+
+				TOptional<bool> bAllowChildrenOverride(true);
+				TOptional<bool> bCreateCategoryNodesOverride(false);
+				Row = ChildrenBuilder.AddExternalObjectProperty(Objects, NAME_None, Parameter.GetName(), bAllowChildrenOverride, bCreateCategoryNodesOverride); 
 
 				CustomValueWidget =
 					SNew(STextBlock)
@@ -222,14 +225,14 @@ private:
 
 	void OnParameterChanged(FNiagaraVariable Var)
 	{
-		Component->GetInitialParameters().OnParameterChange();
-		Component->SetParameterValueOverriddenLocally(Var.GetName(), true);
+		Component->GetOverrideParameters().OnParameterChange();
+		Component->SetParameterValueOverriddenLocally(Var, true);
 	}
 
 	void OnDataInterfaceChanged(FNiagaraVariable Var)
 	{
-		Component->GetInitialParameters().OnInterfaceChange();
-		Component->SetParameterValueOverriddenLocally(Var.GetName(), true);
+		Component->GetOverrideParameters().OnInterfaceChange();
+		Component->SetParameterValueOverriddenLocally(Var, true);
 	}
 
 	bool DoesParameterDifferFromDefault(const FNiagaraVariable& Var)
@@ -241,7 +244,7 @@ private:
 	{
 		FScopedTransaction ScopedTransaction(LOCTEXT("ResetParameterValue", "Reset parameter value to system defaults."));
 		Component->Modify();
-		Component->SetParameterValueOverriddenLocally(Parameter.GetName(), false);
+		Component->SetParameterValueOverriddenLocally(Parameter, false);
 		return FReply::Handled();
 	}
 
@@ -426,7 +429,7 @@ void FNiagaraComponentDetails::CustomizeDetails(IDetailLayoutBuilder& DetailBuil
 
 	DetailBuilder.EditCategory(ScriptCategoryName);
 
- 	TSharedPtr<IPropertyHandle> LocalOverridesPropertyHandle = DetailBuilder.GetProperty(TEXT("Parameters"));
+ 	TSharedPtr<IPropertyHandle> LocalOverridesPropertyHandle = DetailBuilder.GetProperty(TEXT("OverrideParameters"));
  	if (LocalOverridesPropertyHandle.IsValid())
  	{
  		LocalOverridesPropertyHandle->MarkHiddenByCustomization();
@@ -465,11 +468,11 @@ void FNiagaraComponentDetails::CustomizeDetails(IDetailLayoutBuilder& DetailBuil
 		OnResetDelegateHandle = SystemInstance->OnReset().AddSP(this, &FNiagaraComponentDetails::OnSystemInstanceReset);
 		OnInitDelegateHandle = SystemInstance->OnInitialized().AddSP(this, &FNiagaraComponentDetails::OnSystemInstanceReset);
 		SystemInstance->OnDestroyed().AddSP(this, &FNiagaraComponentDetails::OnSystemInstanceDestroyed);
-		FNiagaraParameterStore* ParamStore = &Component->GetInitialParameters();
+		FNiagaraParameterStore* ParamStore = &Component->GetOverrideParameters();
 		UNiagaraScript* ScriptSpawn = Component->GetAsset()->GetSystemSpawnScript();
 		UNiagaraScript* ScriptUpdate = Component->GetAsset()->GetSystemUpdateScript();
 		
-		IDetailCategoryBuilder& InputParamCategory = DetailBuilder.EditCategory(ParamCategoryName, LOCTEXT("ParamCategoryName", "Parameters"));
+		IDetailCategoryBuilder& InputParamCategory = DetailBuilder.EditCategory(ParamCategoryName, LOCTEXT("ParamCategoryName", "Override Parameters"));
 		//InputParamCategory.HeaderContent(SNew(SAddParameterButton, InputCollectionViewModel));
 		InputParamCategory.AddCustomBuilder(MakeShared<FNiagaraComponentNodeBuilder>(Component, ScriptSpawn, ScriptUpdate));
 	}
diff --git a/Engine/Plugins/FX/Niagara/Source/NiagaraEditor/Private/Customizations/NiagaraMetaDataCustomNodeBuilder.h b/Engine/Plugins/FX/Niagara/Source/NiagaraEditor/Private/Customizations/NiagaraMetaDataCustomNodeBuilder.h
new file mode 100644
--- /dev/null
+++ b/Engine/Plugins/FX/Niagara/Source/NiagaraEditor/Private/Customizations/NiagaraMetaDataCustomNodeBuilder.h
@@ -0,0 +1,112 @@
+// Copyright 1998-2017 Epic Games, Inc. All Rights Reserved.
+
+#pragma once
+#include "DetailLayoutBuilder.h"
+#include "DetailCategoryBuilder.h"
+#include "IDetailPropertyRow.h"
+#include "DetailWidgetRow.h"
+#include "SInlineEditableTextBlock.h"
+#include "NiagaraMetaDataCollectionViewModel.h"
+#include "NiagaraScriptInputCollectionViewModel.h"
+#include "NiagaraScriptOutputCollectionViewModel.h"
+#include "NiagaraMetaDataCollectionViewModel.h"
+#include "NiagaraMetaDataViewModel.h"
+#include "NiagaraScriptViewModel.h"
+#include "NiagaraParameterViewModel.h"
+#include "NiagaraEditorStyle.h"
+#include "IDetailChildrenBuilder.h"
+#include "MultiBoxBuilder.h"
+#include "SComboButton.h"
+#include "SImage.h"
+#include "NiagaraEditorModule.h"
+#include "ModuleManager.h"
+#include "INiagaraEditorTypeUtilities.h"
+#include "SBox.h"
+#include "NiagaraScript.h"
+#include "StructOnScope.h"
+#include "Text.h"
+#include "IDetailGroup.h"
+
+
+#define LOCTEXT_NAMESPACE "NiagaraMetaDataCustomNodeBuilder"
+
+class FNiagaraMetaDataCustomNodeBuilder : public IDetailCustomNodeBuilder
+{
+public:
+	FNiagaraMetaDataCustomNodeBuilder(TSharedRef<FNiagaraMetaDataCollectionViewModel> InViewModel)
+		: ViewModel(InViewModel)
+	{
+		ViewModel->OnCollectionChanged().AddRaw(this, &FNiagaraMetaDataCustomNodeBuilder::OnCollectionViewModelChanged);
+	}
+
+	~FNiagaraMetaDataCustomNodeBuilder()
+	{
+		ViewModel->OnCollectionChanged().RemoveAll(this);
+	}
+
+	virtual void SetOnRebuildChildren(FSimpleDelegate InOnRegenerateChildren) override
+	{
+		OnRebuildChildren = InOnRegenerateChildren;
+	}
+
+	virtual void GenerateHeaderRowContent(FDetailWidgetRow& NodeRow) {}
+	virtual void Tick(float DeltaTime) override {}
+	virtual bool RequiresTick() const override { return false; }
+	virtual bool InitiallyCollapsed() const { return false; }
+	
+	virtual FName GetName() const  override
+	{
+		static const FName NiagaraMetadataCustomNodeBuilder("NiagaraMetadataCustomNodeBuilder");
+		return NiagaraMetadataCustomNodeBuilder;
+	}
+
+	virtual void GenerateChildContent(IDetailChildrenBuilder& ChildrenBuilder) override
+	{
+		const TArray<TSharedRef<FNiagaraMetaDataViewModel>>& VariableModels = ViewModel->GetVariableModels();
+
+		FNiagaraEditorModule& NiagaraEditorModule = FModuleManager::GetModuleChecked<FNiagaraEditorModule>("NiagaraEditor");
+
+		for (const TSharedRef<FNiagaraMetaDataViewModel>& MetadataViewModel : VariableModels)
+		{
+			IDetailGroup& MetaDataGroup = ChildrenBuilder.AddGroup(MetadataViewModel->GetName(), FText::FromName(MetadataViewModel->GetName()));
+			MetaDataGroup.ToggleExpansion(true);
+			TSharedRef<FStructOnScope> StructData = MetadataViewModel->GetValueStruct();
+			
+			IDetailPropertyRow* DescriptionRow = ChildrenBuilder.AddExternalStructureProperty(StructData, FName(TEXT("Description")), "Description");
+
+			if (DescriptionRow)
+			{
+				DescriptionRow->Visibility(EVisibility::Hidden); // hide it here, show it in groups only
+				MetaDataGroup.AddPropertyRow(DescriptionRow->GetPropertyHandle()->AsShared());
+
+				TSharedPtr<IPropertyHandle> PropertyHandle = DescriptionRow->GetPropertyHandle();
+				PropertyHandle->SetOnPropertyValueChanged(FSimpleDelegate::CreateSP(MetadataViewModel, &FNiagaraMetaDataViewModel::NotifyMetaDataChanged));
+			}
+
+			IDetailPropertyRow* MetaDataRow = ChildrenBuilder.AddExternalStructureProperty(StructData, FName(TEXT("PropertyMetaData")), "PropertyMetaData");
+			
+			if (MetaDataRow)
+			{
+				MetaDataRow->ShouldAutoExpand(true);
+				MetaDataRow->Visibility(EVisibility::Hidden);
+				MetaDataGroup.AddPropertyRow(MetaDataRow->GetPropertyHandle()->AsShared());
+
+				TSharedPtr<IPropertyHandle> PropertyHandle = MetaDataRow->GetPropertyHandle();
+				PropertyHandle->SetOnPropertyValueChanged(FSimpleDelegate::CreateSP(MetadataViewModel, &FNiagaraMetaDataViewModel::NotifyMetaDataChanged));
+				PropertyHandle->SetOnChildPropertyValueChanged(FSimpleDelegate::CreateSP(MetadataViewModel, &FNiagaraMetaDataViewModel::NotifyMetaDataChanged));
+			}
+		}
+	}
+
+private:
+	void OnCollectionViewModelChanged()
+	{
+		OnRebuildChildren.ExecuteIfBound();
+	}
+
+private:
+	TSharedRef<FNiagaraMetaDataCollectionViewModel> ViewModel;
+	FSimpleDelegate OnRebuildChildren;
+};
+
+#undef LOCTEXT_NAMESPACE // "NiagaraMetaDataCustomNodeBuilder"
\ No newline at end of file
diff --git a/Engine/Plugins/FX/Niagara/Source/NiagaraEditor/Private/Customizations/NiagaraParameterCollectionCustomNodeBuilder.h b/Engine/Plugins/FX/Niagara/Source/NiagaraEditor/Private/Customizations/NiagaraParameterCollectionCustomNodeBuilder.h
new file mode 100644
--- /dev/null
+++ b/Engine/Plugins/FX/Niagara/Source/NiagaraEditor/Private/Customizations/NiagaraParameterCollectionCustomNodeBuilder.h
@@ -0,0 +1,162 @@
+// Copyright 1998-2016 Epic Games, Inc. All Rights Reserved.
+
+#pragma once
+#include "DetailLayoutBuilder.h"
+#include "DetailCategoryBuilder.h"
+#include "IDetailPropertyRow.h"
+#include "DetailWidgetRow.h"
+#include "SInlineEditableTextBlock.h"
+#include "NiagaraParameterCollectionViewModel.h"
+#include "NiagaraScriptInputCollectionViewModel.h"
+#include "NiagaraScriptOutputCollectionViewModel.h"
+#include "NiagaraScriptViewModel.h"
+#include "NiagaraParameterViewModel.h"
+#include "NiagaraEditorStyle.h"
+#include "IDetailChildrenBuilder.h"
+#include "MultiBoxBuilder.h"
+#include "SComboButton.h"
+#include "SImage.h"
+#include "NiagaraEditorModule.h"
+#include "ModuleManager.h"
+#include "INiagaraEditorTypeUtilities.h"
+#include "SBox.h"
+#include "NiagaraScript.h"
+
+class FNiagaraParameterCollectionCustomNodeBuilder : public IDetailCustomNodeBuilder
+{
+public:
+	FNiagaraParameterCollectionCustomNodeBuilder(TSharedRef<INiagaraParameterCollectionViewModel> InViewModel, bool bAllowMetaData = true)
+		: ViewModel(InViewModel)
+	{
+		ViewModel->OnCollectionChanged().AddRaw(this, &FNiagaraParameterCollectionCustomNodeBuilder::OnCollectionViewModelChanged);
+	}
+
+	~FNiagaraParameterCollectionCustomNodeBuilder()
+	{
+		ViewModel->OnCollectionChanged().RemoveAll(this);
+	}
+
+	virtual void SetOnRebuildChildren(FSimpleDelegate InOnRegenerateChildren) override
+	{
+		OnRebuildChildren = InOnRegenerateChildren;
+	}
+
+	virtual void GenerateHeaderRowContent(FDetailWidgetRow& NodeRow) {}
+	virtual void Tick(float DeltaTime) override {}
+	virtual bool RequiresTick() const override { return false; }
+	virtual bool InitiallyCollapsed() const { return false; }
+	virtual FName GetName() const  override
+	{
+		static const FName NiagaraCustomNodeBuilder("NiagaraCustomNodeBuilder");
+		return NiagaraCustomNodeBuilder;
+	}
+
+	virtual void GenerateChildContent(IDetailChildrenBuilder& ChildrenBuilder) override
+	{
+		const TArray<TSharedRef<INiagaraParameterViewModel>>& Parameters = ViewModel->GetParameters();
+
+		FNiagaraEditorModule& NiagaraEditorModule = FModuleManager::GetModuleChecked<FNiagaraEditorModule>("NiagaraEditor");
+
+		for (const TSharedRef<INiagaraParameterViewModel>& Parameter : Parameters)
+		{
+			TSharedPtr<SWidget> NameWidget;
+
+			if (Parameter->CanRenameParameter())
+			{
+				NameWidget =
+					SNew(SInlineEditableTextBlock)
+					.Style(FNiagaraEditorStyle::Get(), "NiagaraEditor.ParameterInlineEditableText")
+					.Text(Parameter, &INiagaraParameterViewModel::GetNameText)
+					.OnVerifyTextChanged(Parameter, &INiagaraParameterViewModel::VerifyNodeNameTextChanged)
+					.OnTextCommitted(Parameter, &INiagaraParameterViewModel::NameTextComitted);
+			}
+			else
+			{
+				NameWidget =
+					SNew(STextBlock)
+					.TextStyle(FNiagaraEditorStyle::Get(), "NiagaraEditor.ParameterText")
+					.Text(Parameter, &INiagaraParameterViewModel::GetNameText);
+			}
+
+
+
+			IDetailPropertyRow* Row = nullptr;
+
+			TSharedPtr<SWidget> CustomValueWidget;
+			bool bCustomize = true;
+
+			if (Parameter->GetDefaultValueType() == INiagaraParameterViewModel::EDefaultValueType::Struct)
+			{
+				Row = ChildrenBuilder.AddExternalStructureProperty(Parameter->GetDefaultValueStruct(), NAME_None, Parameter->GetName());
+
+			}
+			else if (Parameter->GetDefaultValueType() == INiagaraParameterViewModel::EDefaultValueType::Object)
+			{
+				UObject* DefaultValueObject = Parameter->GetDefaultValueObject();
+
+				TArray<UObject*> Objects;
+				Objects.Add(DefaultValueObject);
+				
+				TOptional<bool> bAllowChildrenOverride(true);
+				TOptional<bool> bCreateCategoryNodesOverride(false);
+
+				Row = ChildrenBuilder.AddExternalObjectProperty(Objects, NAME_None, Parameter->GetName(), bAllowChildrenOverride, bCreateCategoryNodesOverride);
+				CustomValueWidget =
+					SNew(STextBlock)
+					.TextStyle(FNiagaraEditorStyle::Get(), "NiagaraEditor.ParameterText")
+					.Text(FText::FromString(FName::NameToDisplayString(DefaultValueObject->GetClass()->GetName(), false)));
+			}
+
+			if (Row)
+			{
+
+				TSharedPtr<IPropertyHandle> PropertyHandle = Row->GetPropertyHandle();
+				PropertyHandle->SetOnPropertyValueChanged(
+					FSimpleDelegate::CreateSP(Parameter, &INiagaraParameterViewModel::NotifyDefaultValueChanged));
+				PropertyHandle->SetOnChildPropertyValueChanged(
+					FSimpleDelegate::CreateSP(Parameter, &INiagaraParameterViewModel::NotifyDefaultValueChanged));
+
+				if (bCustomize)
+				{
+					TSharedPtr<SWidget> DefaultNameWidget;
+					TSharedPtr<SWidget> DefaultValueWidget;
+					FDetailWidgetRow& CustomWidget = Row->CustomWidget(true);
+					Row->GetDefaultWidgets(DefaultNameWidget, DefaultValueWidget, CustomWidget);
+					CustomWidget
+						.NameContent()
+						[
+							SNew(SBox)
+							.Padding(FMargin(0.0f, 2.0f))
+							[
+								NameWidget.ToSharedRef()
+							]
+						];
+
+					if (CustomValueWidget.IsValid())
+					{
+						CustomWidget.ValueContent()
+							[
+								CustomValueWidget.ToSharedRef()
+							];
+					}
+
+				}
+
+			}
+
+		}
+
+
+	}
+
+private:
+	void OnCollectionViewModelChanged()
+	{
+		OnRebuildChildren.ExecuteIfBound();
+	}
+
+private:
+	TSharedRef<INiagaraParameterCollectionViewModel> ViewModel;
+	FSimpleDelegate OnRebuildChildren;
+	bool bAllowMetaData;
+};
\ No newline at end of file
