Copying Orion-Staging @ 3748653 (Orion/Dev-General @ 3722124) to //UE4/Main

diff --git a/Engine/Plugins/FX/Niagara/Source/NiagaraEditor/Private/Customizations/NiagaraScriptDetails.cpp b/Engine/Plugins/FX/Niagara/Source/NiagaraEditor/Private/Customizations/NiagaraScriptDetails.cpp
--- a/Engine/Plugins/FX/Niagara/Source/NiagaraEditor/Private/Customizations/NiagaraScriptDetails.cpp
+++ b/Engine/Plugins/FX/Niagara/Source/NiagaraEditor/Private/Customizations/NiagaraScriptDetails.cpp
@@ -20,141 +20,15 @@
 #include "ModuleManager.h"
 #include "INiagaraEditorTypeUtilities.h"
 #include "SBox.h"
+#include "NiagaraScript.h"
+#include "NiagaraParameterCollectionCustomNodeBuilder.h"
+#include "NiagaraMetaDataCustomNodeBuilder.h"
+#include "NiagaraMetaDataCollectionViewModel.h"
 
 
 #define LOCTEXT_NAMESPACE "NiagaraScriptDetails"
 
-class FNiagaraCustomNodeBuilder : public IDetailCustomNodeBuilder
-{
-public:
-	FNiagaraCustomNodeBuilder(TSharedRef<INiagaraParameterCollectionViewModel> InViewModel)
-		: ViewModel(InViewModel)
-	{
-		ViewModel->OnCollectionChanged().AddRaw(this, &FNiagaraCustomNodeBuilder::OnCollectionViewModelChanged);
-	}
-
-	~FNiagaraCustomNodeBuilder()
-	{
-		ViewModel->OnCollectionChanged().RemoveAll(this);
-	}
-
-	virtual void SetOnRebuildChildren(FSimpleDelegate InOnRegenerateChildren) override
-	{
-		OnRebuildChildren = InOnRegenerateChildren;
-	}
-
-	virtual void GenerateHeaderRowContent(FDetailWidgetRow& NodeRow){}
-	virtual void Tick(float DeltaTime) override {}
-	virtual bool RequiresTick() const override { return false; }
-	virtual bool InitiallyCollapsed() const { return false; }
-	virtual FName GetName() const  override
-	{ 
-		static const FName NiagaraCustomNodeBuilder("NiagaraCustomNodeBuilder");
-		return NiagaraCustomNodeBuilder;
-	}
-
-	virtual void GenerateChildContent(IDetailChildrenBuilder& ChildrenBuilder) override
-	{
-		const TArray<TSharedRef<INiagaraParameterViewModel>>& Parameters = ViewModel->GetParameters();
-
-		FNiagaraEditorModule& NiagaraEditorModule = FModuleManager::GetModuleChecked<FNiagaraEditorModule>("NiagaraEditor");
-
-		for (const TSharedRef<INiagaraParameterViewModel>& Parameter : Parameters)
-		{
-			TSharedPtr<SWidget> NameWidget;
-
-			if (Parameter->CanRenameParameter())
-			{
-				NameWidget =
-					SNew(SInlineEditableTextBlock)
-					.Style(FNiagaraEditorStyle::Get(), "NiagaraEditor.ParameterInlineEditableText")
-					.Text(Parameter, &INiagaraParameterViewModel::GetNameText)
-					.OnVerifyTextChanged(Parameter, &INiagaraParameterViewModel::VerifyNodeNameTextChanged)
-					.OnTextCommitted(Parameter, &INiagaraParameterViewModel::NameTextComitted);
-			}
-			else
-			{
-				NameWidget =
-					SNew(STextBlock)
-					.TextStyle(FNiagaraEditorStyle::Get(), "NiagaraEditor.ParameterText")
-					.Text(Parameter, &INiagaraParameterViewModel::GetNameText);
-			}
-
-			
-
-			IDetailPropertyRow* Row = nullptr;
-
-			TSharedPtr<SWidget> CustomValueWidget;
-
-			if (Parameter->GetDefaultValueType() == INiagaraParameterViewModel::EDefaultValueType::Struct)
-			{
-				Row = ChildrenBuilder.AddExternalStructureProperty(Parameter->GetDefaultValueStruct(), NAME_None, Parameter->GetName());
-				
-			}
-			else if (Parameter->GetDefaultValueType() == INiagaraParameterViewModel::EDefaultValueType::Object)
-			{
-				UObject* DefaultValueObject = Parameter->GetDefaultValueObject();
-
-				TArray<UObject*> Objects;
-				Objects.Add(DefaultValueObject);
-				Row = ChildrenBuilder.AddExternalObjects(Objects, Parameter->GetName());
-
-				CustomValueWidget =
-					SNew(STextBlock)
-					.TextStyle(FNiagaraEditorStyle::Get(), "NiagaraEditor.ParameterText")
-					.Text(FText::FromString(FName::NameToDisplayString(DefaultValueObject->GetClass()->GetName(), false)));
-			}
-
-			if(Row)
-			{
-				TSharedPtr<SWidget> DefaultNameWidget;
-				TSharedPtr<SWidget> DefaultValueWidget;
-
-				TSharedPtr<IPropertyHandle> PropertyHandle = Row->GetPropertyHandle();
 
-				FDetailWidgetRow& CustomWidget = Row->CustomWidget(true);
-
-				Row->GetDefaultWidgets(DefaultNameWidget, DefaultValueWidget, CustomWidget);
-
-				PropertyHandle->SetOnPropertyValueChanged(
-					FSimpleDelegate::CreateSP(Parameter, &INiagaraParameterViewModel::NotifyDefaultValueChanged));
-				PropertyHandle->SetOnChildPropertyValueChanged(
-					FSimpleDelegate::CreateSP(Parameter, &INiagaraParameterViewModel::NotifyDefaultValueChanged));
-
-				CustomWidget
-					.NameContent()
-					[
-						SNew(SBox)
-						.Padding(FMargin(0.0f, 2.0f))
-						[
-							NameWidget.ToSharedRef()
-						]
-					];
-
-				if (CustomValueWidget.IsValid())
-				{
-					CustomWidget.ValueContent()
-					[
-						CustomValueWidget.ToSharedRef()
-					];
-				}
-			}
-
-		}
-
-
-	}
-
-private:
-	void OnCollectionViewModelChanged()
-	{
-		OnRebuildChildren.ExecuteIfBound();
-	}
-
-private:
-	TSharedRef<INiagaraParameterCollectionViewModel> ViewModel;
-	FSimpleDelegate OnRebuildChildren;
-};
 
 class SAddParameterButton : public SCompoundWidget
 {	
@@ -245,21 +119,49 @@ void FNiagaraScriptDetails::CustomizeDetails(IDetailLayoutBuilder& DetailBuilder
 	static const FName InputParamCategoryName = TEXT("NiagaraScript_InputParams");
 	static const FName OutputParamCategoryName = TEXT("NiagaraScript_OutputParams");
 	static const FName ScriptCategoryName = TEXT("Script");
+	static const FName MetadataCategoryName = TEXT("Metadata");
 
 	DetailBuilder.EditCategory(ScriptCategoryName);
 
-	TSharedRef<INiagaraParameterCollectionViewModel> InputCollectionViewModel = ScriptViewModel->GetInputCollectionViewModel();
-	TSharedRef<INiagaraParameterCollectionViewModel> OutputCollectionViewModel = ScriptViewModel->GetOutputCollectionViewModel();
+	UNiagaraScript* StandaloneScript = ScriptViewModel->GetStandaloneScript();
+	bool bAddParameters = false;
+	TSharedPtr<INiagaraParameterCollectionViewModel> InputCollectionViewModel;
+	TSharedPtr<INiagaraParameterCollectionViewModel> OutputCollectionViewModel;
+	TSharedPtr<FNiagaraMetaDataCollectionViewModel> MetadataCollectionViewModel;
 
-	IDetailCategoryBuilder& InputParamCategory = DetailBuilder.EditCategory(InputParamCategoryName, LOCTEXT("InputParamCategoryName", "Input Parameters"));
-	InputParamCategory.HeaderContent(SNew(SAddParameterButton, InputCollectionViewModel));
-	InputParamCategory.AddCustomBuilder(MakeShared<FNiagaraCustomNodeBuilder>(InputCollectionViewModel));
+	if (StandaloneScript)
+	{
+		InputCollectionViewModel = ScriptViewModel->GetInputCollectionViewModel();
+		OutputCollectionViewModel = ScriptViewModel->GetOutputCollectionViewModel();
+		MetadataCollectionViewModel = ScriptViewModel->GetMetadataCollectionViewModel();
+		bAddParameters = true;
+	}
 
+	if (InputCollectionViewModel.IsValid())
+	{
+		IDetailCategoryBuilder& InputParamCategory = DetailBuilder.EditCategory(InputParamCategoryName, LOCTEXT("InputParamCategoryName", "Input Parameters"));
+		if (bAddParameters)
+		{
+			InputParamCategory.HeaderContent(SNew(SAddParameterButton, InputCollectionViewModel.ToSharedRef()));
+		}
+		InputParamCategory.AddCustomBuilder(MakeShared<FNiagaraParameterCollectionCustomNodeBuilder>(InputCollectionViewModel.ToSharedRef()));
+	}
 
-	IDetailCategoryBuilder& OutputParamCategory = DetailBuilder.EditCategory(OutputParamCategoryName, LOCTEXT("OutputParamCategoryName", "Output Parameters"));
-	OutputParamCategory.HeaderContent(SNew(SAddParameterButton, OutputCollectionViewModel));
-	OutputParamCategory.AddCustomBuilder(MakeShared<FNiagaraCustomNodeBuilder>(OutputCollectionViewModel));
+	if (OutputCollectionViewModel.IsValid())
+	{
+		IDetailCategoryBuilder& OutputParamCategory = DetailBuilder.EditCategory(OutputParamCategoryName, LOCTEXT("OutputParamCategoryName", "Output Parameters"));
+		if (bAddParameters)
+		{
+			OutputParamCategory.HeaderContent(SNew(SAddParameterButton, OutputCollectionViewModel.ToSharedRef()));
+		}
+		OutputParamCategory.AddCustomBuilder(MakeShared<FNiagaraParameterCollectionCustomNodeBuilder>(OutputCollectionViewModel.ToSharedRef()));
+	}
 	
+	if (MetadataCollectionViewModel.IsValid())
+	{
+		IDetailCategoryBuilder& MetadataDetailCategory = DetailBuilder.EditCategory(MetadataCategoryName, LOCTEXT("MetadataParamCategoryName", "Variable Metadata"));
+		MetadataDetailCategory.AddCustomBuilder(MakeShared<FNiagaraMetaDataCustomNodeBuilder>(MetadataCollectionViewModel.ToSharedRef()));
+	}	
 }
 
 
diff --git a/Engine/Plugins/FX/Niagara/Source/NiagaraEditor/Private/EdGraphSchema_Niagara.cpp b/Engine/Plugins/FX/Niagara/Source/NiagaraEditor/Private/EdGraphSchema_Niagara.cpp
--- a/Engine/Plugins/FX/Niagara/Source/NiagaraEditor/Private/EdGraphSchema_Niagara.cpp
+++ b/Engine/Plugins/FX/Niagara/Source/NiagaraEditor/Private/EdGraphSchema_Niagara.cpp
@@ -25,7 +25,7 @@
 #include "NiagaraNodeWriteDataSet.h"
 #include "NiagaraNodeParameterMapGet.h"
 #include "NiagaraNodeParameterMapSet.h"
-#include "NiagaraNodeParameterCollection.h"
+#include "NiagaraNodeCustomHlsl.h"
 #include "NiagaraNodeOp.h"
 #include "NiagaraNodeConvert.h"
 #include "NiagaraEditorUtilities.h"
@@ -49,6 +49,7 @@ const FLinearColor UEdGraphSchema_Niagara::NodeTitleColor_Attribute = FLinearCol
 const FLinearColor UEdGraphSchema_Niagara::NodeTitleColor_Constant = FLinearColor::Red;
 const FLinearColor UEdGraphSchema_Niagara::NodeTitleColor_SystemConstant = FLinearColor::White;
 const FLinearColor UEdGraphSchema_Niagara::NodeTitleColor_FunctionCall = FLinearColor::Blue;
+const FLinearColor UEdGraphSchema_Niagara::NodeTitleColor_CustomHlsl = FLinearColor::Yellow;
 const FLinearColor UEdGraphSchema_Niagara::NodeTitleColor_Event = FLinearColor::Red;
 
 const FName UEdGraphSchema_Niagara::PinCategoryType("Type");
@@ -249,7 +250,7 @@ void AddParametersForGraph(TArray<TSharedPtr<FNiagaraSchemaAction_NewNode> >& Ne
 			InputNodeTemplate->Input = InputNode->Input;
 			InputNodeTemplate->Usage = InputNode->Usage;
 			InputNodeTemplate->ExposureOptions = InputNode->ExposureOptions;
-			InputNodeTemplate->DataInterface = nullptr;
+			InputNodeTemplate->SetDataInterface(nullptr);
 
 			// We also support parameters from an alternate graph. If that was used, then we need to take special care
 			// to make the parameter unique to that graph.
@@ -257,9 +258,9 @@ void AddParametersForGraph(TArray<TSharedPtr<FNiagaraSchemaAction_NewNode> >& Ne
 			{
 				InputNodeTemplate->Input.SetName(Name);
 
-				if (InputNode->DataInterface)
+				if (InputNode->GetDataInterface())
 				{
-					InputNodeTemplate->DataInterface = Cast<UNiagaraDataInterface>(StaticDuplicateObject(InputNode->DataInterface, InputNodeTemplate, NAME_None, ~RF_Transient));
+					InputNodeTemplate->SetDataInterface(Cast<UNiagaraDataInterface>(StaticDuplicateObject(InputNode->GetDataInterface(), InputNodeTemplate, NAME_None, ~RF_Transient)));
 				}
 			}
 
@@ -315,6 +316,16 @@ TArray<TSharedPtr<FNiagaraSchemaAction_NewNode> > UEdGraphSchema_Niagara::GetGra
 		}
 	}
 
+	// Add custom code
+	{
+		const FText MenuDesc = LOCTEXT("CustomHLSLNode","Custom Hlsl");
+		const FText TooltipDesc = LOCTEXT("CustomHlslPopupTooltip", "Add a node with custom hlsl content");
+		TSharedPtr<FNiagaraSchemaAction_NewNode> FunctionCallAction = AddNewNodeAction(NewActions, LOCTEXT("Function Menu Title", "Functions"), MenuDesc, TEXT("CustomHLSL"), TooltipDesc);
+		UNiagaraNodeCustomHlsl* CustomHlslNode = NewObject<UNiagaraNodeCustomHlsl>(OwnerOfTemporaries);
+		CustomHlslNode->CustomHlsl = TEXT("// Insert the body of the function here and add any inputs\r\n// and outputs by name using the add pins above.\r\n// Currently, complicated branches, for loops, switches, etc are not advised.");
+		FunctionCallAction->NodeTemplate = CustomHlslNode;
+	}
+
 	//Add functions
 	FAssetRegistryModule& AssetRegistryModule = FModuleManager::LoadModuleChecked<FAssetRegistryModule>(TEXT("AssetRegistry"));
 	TArray<FAssetData> ScriptAssets;
@@ -398,25 +409,6 @@ TArray<TSharedPtr<FNiagaraSchemaAction_NewNode> > UEdGraphSchema_Niagara::GetGra
 		}
 	}
 
-	//Add collections
-	if (bModuleGraph || bFunctionGraph)
-	{
-		TArray<FAssetData> CollectionAssets;
-		AssetRegistryModule.Get().GetAssetsByClass(UNiagaraParameterCollection::StaticClass()->GetFName(), CollectionAssets);
-
-		for (const FAssetData& CollectionAsset : CollectionAssets)
-		{
-			FString DisplayNameString = FName::NameToDisplayString(CollectionAsset.AssetName.ToString(), false);
-			const FText MenuDesc = FText::FromString(DisplayNameString);
-
-			TSharedPtr<FNiagaraSchemaAction_NewNode> CollectionAction = AddNewNodeAction(NewActions, LOCTEXT("Collection Menu Title", "Collections"), MenuDesc, *DisplayNameString, FText::FromName(CollectionAsset.ObjectPath));
-
-			UNiagaraNodeParameterCollection* CollectionNode = NewObject<UNiagaraNodeParameterCollection>(OwnerOfTemporaries);
-			CollectionNode->CollectionAssetObjectPath = CollectionAsset.ObjectPath;
-			CollectionAction->NodeTemplate = CollectionNode;
-		}
-	}
-
 	//Add event read and writes nodes
 	if (bModuleGraph)
 	{
@@ -605,7 +597,7 @@ TArray<TSharedPtr<FNiagaraSchemaAction_NewNode> > UEdGraphSchema_Niagara::GetGra
 		{
 			if (UNiagaraNodeInput* InputNode = Cast<UNiagaraNodeInput>(FromPin->GetOwningNode()))
 			{
-				DataInterface = InputNode->DataInterface;
+				DataInterface = InputNode->GetDataInterface();
 			}
 			else 
 			{
@@ -887,6 +879,15 @@ FConnectionDrawingPolicy* UEdGraphSchema_Niagara::CreateConnectionDrawingPolicy(
 	return new FNiagaraConnectionDrawingPolicy(InBackLayerID, InFrontLayerID, InZoomFactor, InClippingRect, InDrawElements);
 }
 
+void UEdGraphSchema_Niagara::ResetPinToAutogeneratedDefaultValue(UEdGraphPin* Pin, bool bCallModifyCallbacks) const
+{
+	Pin->DefaultValue = Pin->AutogeneratedDefaultValue;
+	if (bCallModifyCallbacks)
+	{
+		Pin->GetOwningNode()->PinDefaultValueChanged(Pin);
+	}
+}
+
 bool UEdGraphSchema_Niagara::TryCreateConnection(UEdGraphPin* PinA, UEdGraphPin* PinB) const
 {
 	const FScopedTransaction Transaction(NSLOCTEXT("UnrealEd", "NiagaraEditorCreateConnection", "Niagara Editor: Create Connection"));
@@ -1036,54 +1037,79 @@ bool UEdGraphSchema_Niagara::ShouldHidePinDefaultValue(UEdGraphPin* Pin) const
 FNiagaraVariable UEdGraphSchema_Niagara::PinToNiagaraVariable(const UEdGraphPin* Pin, bool bNeedsValue)const
 {
 	FNiagaraVariable Var = FNiagaraVariable(PinToTypeDefinition(Pin), Pin->PinName);
-	if (Pin->DefaultValue.IsEmpty() == false)
+	bool bHasValue = false;
+	if (Pin->bDefaultValueIsIgnored == false && Pin->DefaultValue.IsEmpty() == false)
 	{
 		FNiagaraEditorModule& NiagaraEditorModule = FModuleManager::GetModuleChecked<FNiagaraEditorModule>("NiagaraEditor");
 		TSharedPtr<INiagaraEditorTypeUtilities> TypeEditorUtilities = NiagaraEditorModule.GetTypeUtilities(Var.GetType());
 		if (TypeEditorUtilities.IsValid() && TypeEditorUtilities->CanHandlePinDefaults())
 		{
-			TypeEditorUtilities->SetValueFromPinDefaultString(Pin->DefaultValue, Var);
-		}
-		else if (const UEnum* Enum = Var.GetType().GetEnum())
-		{
-			int64 EnumValue = Enum->GetValueByNameString(Pin->DefaultValue, EGetByNameFlags::None);
-			if (EnumValue != INDEX_NONE)
+			bHasValue = TypeEditorUtilities->SetValueFromPinDefaultString(Pin->DefaultValue, Var);
+			if (bHasValue == false)
 			{
-				Var.SetValue<int32>((int32)EnumValue);
+				FString OwningNodePath = Pin->GetOwningNode() != nullptr ? Pin->GetOwningNode()->GetPathName() : TEXT("Unknown");
+				UE_LOG(LogNiagaraEditor, Error, TEXT("PinToNiagaraVariable: Failed to convert default value '%s' to type %s. Owning node path: %s"), *Pin->DefaultValue, *Var.GetType().GetName(), *OwningNodePath);
 			}
 		}
 		else
 		{
-			UE_LOG(LogNiagaraEditor, Warning, TEXT("Pin default values aren't supported for variables of type {%s}"), *Var.GetType().GetName());
+			FString OwningNodePath = Pin->GetOwningNode() != nullptr ? Pin->GetOwningNode()->GetPathName() : TEXT("Unknown");
+			UE_LOG(LogNiagaraEditor, Error, TEXT("Pin had default value string, but default values aren't supported for variables of type {%s}. Owning node path: %s"), *Var.GetType().GetName(), *OwningNodePath);
 		}
 	}
+
+	if (bNeedsValue && bHasValue == false)
+	{
+		FNiagaraEditorUtilities::ResetVariableToDefaultValue(Var);
+	}
+
 	return Var;
 }
 
-FString UEdGraphSchema_Niagara::NiagaraVariableToPinDefaultValue(const FNiagaraVariable& Variable) const
+bool UEdGraphSchema_Niagara::TryGetPinDefaultValueFromNiagaraVariable(const FNiagaraVariable& Variable, FString& OutPinDefaultValue) const
 {
-	FNiagaraEditorModule& NiagaraEditorModule = FModuleManager::GetModuleChecked<FNiagaraEditorModule>("NiagaraEditor");
-	TSharedPtr<INiagaraEditorTypeUtilities> TypeEditorUtilities = NiagaraEditorModule.GetTypeUtilities(Variable.GetType());
-	if (TypeEditorUtilities.IsValid() && TypeEditorUtilities->CanHandlePinDefaults())
+	// Create a variable we can be sure is allocated since it's required for the call to GetPinDefaultStringFromValue.
+	FNiagaraVariable PinDefaultVariable = Variable;
+	if (Variable.IsDataAllocated() == false)
 	{
-		return TypeEditorUtilities->GetPinDefaultStringFromValue(Variable);
+		FNiagaraEditorUtilities::ResetVariableToDefaultValue(PinDefaultVariable);
 	}
-	else
+
+	FNiagaraEditorModule& NiagaraEditorModule = FModuleManager::GetModuleChecked<FNiagaraEditorModule>("NiagaraEditor");
+	TSharedPtr<INiagaraEditorTypeUtilities> TypeEditorUtilities = NiagaraEditorModule.GetTypeUtilities(PinDefaultVariable.GetType());
+	if (TypeEditorUtilities.IsValid() && TypeEditorUtilities->CanHandlePinDefaults())
 	{
-		UE_LOG(LogNiagaraEditor, Warning, TEXT("Pin default values aren't supported for variables of type {%s}"), *Variable.GetType().GetName());
-		return FString();
+		OutPinDefaultValue = TypeEditorUtilities->GetPinDefaultStringFromValue(PinDefaultVariable);
+		return true;
 	}
+	
+	OutPinDefaultValue = FString();
+	return false;
 }
 
 FNiagaraTypeDefinition UEdGraphSchema_Niagara::PinToTypeDefinition(const UEdGraphPin* Pin) const
 {
 	if (Pin->PinType.PinCategory == PinCategoryType && Pin->PinType.PinSubCategoryObject != nullptr)
 	{
-		return CastChecked<const UScriptStruct>(Pin->PinType.PinSubCategoryObject.Get());
+		const UScriptStruct* Struct = Cast<const UScriptStruct>(Pin->PinType.PinSubCategoryObject.Get());
+		if (Struct == nullptr)
+		{
+			UE_LOG(LogNiagaraEditor, Error, TEXT("Pin states that it is of struct type, but is missing its struct object. This is usually the result of a registered type going away. Pin Name '%s' Owning Node '%s'."),
+				*Pin->PinName.ToString(), *Pin->GetOwningNode()->GetName());
+			return FNiagaraTypeDefinition();
+		}
+		return FNiagaraTypeDefinition(Struct);
 	}
 	else if (Pin->PinType.PinCategory == PinCategoryClass)
 	{
-		return FNiagaraTypeDefinition(CastChecked<const UClass>(Pin->PinType.PinSubCategoryObject.Get()));
+		const UClass* Class = Cast<const UClass>(Pin->PinType.PinSubCategoryObject.Get());
+		if (Class == nullptr)
+		{
+			UE_LOG(LogNiagaraEditor, Error, TEXT("Pin states that it is of class type, but is missing its class object. This is usually the result of a registered type going away. Pin Name '%s' Owning Node '%s'."),
+				*Pin->PinName.ToString(), *Pin->GetOwningNode()->GetName());
+			return FNiagaraTypeDefinition();
+		}
+		return FNiagaraTypeDefinition(Class);
 	}
 	else if (Pin->PinType.PinCategory == PinCategoryEnum)
 	{
@@ -1121,6 +1147,29 @@ bool UEdGraphSchema_Niagara::IsSystemConstant(const FNiagaraVariable& Variable)c
 	return FNiagaraConstants::GetEngineConstants().Find(Variable) != INDEX_NONE;
 }
 
+UNiagaraParameterCollection* UEdGraphSchema_Niagara::VariableIsFromParameterCollection(const FNiagaraVariable& Var)const
+{
+	FString VarName = Var.GetName().ToString();
+	if (VarName.StartsWith(TEXT("NPC.")))
+	{
+		FAssetRegistryModule& AssetRegistryModule = FModuleManager::LoadModuleChecked<FAssetRegistryModule>("AssetRegistry");
+		TArray<FAssetData> CollectionAssets;
+		AssetRegistryModule.Get().GetAssetsByClass(UNiagaraParameterCollection::StaticClass()->GetFName(), CollectionAssets);
+		TArray<FName> ExistingNames;
+		for (FAssetData& CollectionAsset : CollectionAssets)
+		{
+			if (UNiagaraParameterCollection* Collection = CastChecked<UNiagaraParameterCollection>(CollectionAsset.GetAsset()))
+			{
+				if (VarName.StartsWith(Collection->GetFullNamespace()))
+				{
+					return Collection;
+				}
+			}
+		}
+	}
+	return nullptr;
+}
+
 FNiagaraTypeDefinition UEdGraphSchema_Niagara::GetTypeDefForProperty(const UProperty* Property)const
 {
 	if (Property->IsA(UFloatProperty::StaticClass()))
@@ -1224,6 +1273,27 @@ void UEdGraphSchema_Niagara::GetNumericConversionToSubMenuActions(class FMenuBui
 	}
 }
 
+void UEdGraphSchema_Niagara::ToggleNodeEnabledState(UNiagaraNode* InNode) const
+{
+	if (InNode != nullptr)
+	{
+		if (InNode->GetDesiredEnabledState() == ENodeEnabledState::Disabled)
+		{
+			const FScopedTransaction Transaction(NSLOCTEXT("UnrealEd", "NiagaraEditorSetNodeEnabled", "Enabled Node"));
+			InNode->Modify();
+			InNode->SetEnabledState(ENodeEnabledState::Enabled, true);
+			InNode->GetNiagaraGraph()->NotifyGraphNeedsRecompile();
+		}
+		else if (InNode->GetDesiredEnabledState() == ENodeEnabledState::Enabled)
+		{
+			const FScopedTransaction Transaction(NSLOCTEXT("UnrealEd", "NiagaraEditorSetNodeDisabled", "Disabled Node"));
+			InNode->Modify();
+			InNode->SetEnabledState(ENodeEnabledState::Disabled, true);
+			InNode->GetNiagaraGraph()->NotifyGraphNeedsRecompile();
+		}
+	}
+}
+
 void UEdGraphSchema_Niagara::PromoteSinglePinToParameter(UEdGraphPin* SourcePin)
 {
 	if (SourcePin)
@@ -1245,6 +1315,11 @@ void UEdGraphSchema_Niagara::PromoteSinglePinToParameter(UEdGraphPin* SourcePin)
 	}
 }
 
+bool CanResetPinToDefault(const UEdGraphSchema_Niagara* Schema, const UEdGraphPin* Pin)
+{
+	return Schema->DoesDefaultValueMatchAutogenerated(*Pin) == false;
+}
+
 void UEdGraphSchema_Niagara::GetContextMenuActions(const UEdGraph* CurrentGraph, const UEdGraphNode* InGraphNode, const UEdGraphPin* InGraphPin, class FMenuBuilder* MenuBuilder, bool bIsDebugging) const
 {
 	if (InGraphPin)
@@ -1282,12 +1357,27 @@ void UEdGraphSchema_Niagara::GetContextMenuActions(const UEdGraph* CurrentGraph,
 			{
 				MenuBuilder->AddMenuEntry(LOCTEXT("PromoteToParameter", "Promote to Parameter"), LOCTEXT("PromoteToParameterTooltip", "Create a parameter argument and connect this pin to that parameter."), FSlateIcon(),
 					FUIAction(FExecuteAction::CreateUObject((UEdGraphSchema_Niagara*const)this, &UEdGraphSchema_Niagara::PromoteSinglePinToParameter, const_cast<UEdGraphPin*>(InGraphPin))));
+				if (InGraphPin->LinkedTo.Num() == 0 && InGraphPin->bDefaultValueIsIgnored == false)
+				{
+					MenuBuilder->AddMenuEntry(
+						LOCTEXT("ResetInputToDefault", "Reset to Default"), 
+						LOCTEXT("ResetInputToDefaultToolTip", "Reset this input to its default value."),
+						FSlateIcon(),
+						FUIAction(
+							FExecuteAction::CreateUObject((UEdGraphSchema_Niagara*const)this, &UEdGraphSchema_Niagara::ResetPinToAutogeneratedDefaultValue, const_cast<UEdGraphPin*>(InGraphPin), true),
+							FCanExecuteAction::CreateStatic(&CanResetPinToDefault, this, InGraphPin)));
+				}
 			}
 		}
 		MenuBuilder->EndSection();
 	}
 	else if (InGraphNode)
 	{
+		const UNiagaraNode* Node = Cast<UNiagaraNode>(InGraphNode);
+		MenuBuilder->BeginSection("EdGraphSchema_NiagaraNodeActions", LOCTEXT("NodeActionsMenuHeader", "Node Actions"));
+		MenuBuilder->AddMenuEntry(LOCTEXT("ToggleEnabledState", "Toggle Enabled State"), LOCTEXT("ToggleEnabledStateTooltip", "Toggle this node between Enbled (default) and Disabled (skipped from compilation)."), FSlateIcon(),
+			FUIAction(FExecuteAction::CreateUObject((UEdGraphSchema_Niagara*const)this, &UEdGraphSchema_Niagara::ToggleNodeEnabledState, const_cast<UNiagaraNode*>(Node))));
+		MenuBuilder->EndSection();
 	}
 
 	Super::GetContextMenuActions(CurrentGraph, InGraphNode, InGraphPin, MenuBuilder, bIsDebugging);
diff --git a/Engine/Plugins/FX/Niagara/Source/NiagaraEditor/Private/NiagaraCompiler.cpp b/Engine/Plugins/FX/Niagara/Source/NiagaraEditor/Private/NiagaraCompiler.cpp
--- a/Engine/Plugins/FX/Niagara/Source/NiagaraEditor/Private/NiagaraCompiler.cpp
+++ b/Engine/Plugins/FX/Niagara/Source/NiagaraEditor/Private/NiagaraCompiler.cpp
@@ -24,9 +24,14 @@
 
 DEFINE_LOG_CATEGORY_STATIC(LogNiagaraCompiler, All, All);
 
+DECLARE_CYCLE_STAT(TEXT("Niagara - Module - CompileScript"), STAT_NiagaraEditor_Module_CompileScript, STATGROUP_NiagaraEditor);
+DECLARE_CYCLE_STAT(TEXT("Niagara - HlslCompiler - CompileScript"), STAT_NiagaraEditor_HlslCompiler_CompileScript, STATGROUP_NiagaraEditor);
+DECLARE_CYCLE_STAT(TEXT("Niagara - HlslCompiler - CompileShader_VectorVM"), STAT_NiagaraEditor_HlslCompiler_CompileShader_VectorVM, STATGROUP_NiagaraEditor);
+DECLARE_CYCLE_STAT(TEXT("Niagara - Module - CompileShader_VectorVMSucceeded"), STAT_NiagaraEditor_HlslCompiler_CompileShader_VectorVMSucceeded, STATGROUP_NiagaraEditor);
 
 ENiagaraScriptCompileStatus FNiagaraEditorModule::CompileScript(UNiagaraScript* ScriptToCompile, FString& OutGraphLevelErrorMessages)
 {
+	SCOPE_CYCLE_COUNTER(STAT_NiagaraEditor_Module_CompileScript);
 	check(ScriptToCompile != NULL);
 
 	UNiagaraGraph* Graph = Cast<UNiagaraScriptSource>(ScriptToCompile->GetSource())->NodeGraph;
@@ -176,6 +181,8 @@ ENiagaraScriptCompileStatus FNiagaraCompileResults::CompileResultsToSummary(cons
 
 const FNiagaraCompileResults &FHlslNiagaraCompiler::CompileScript(UNiagaraScript* InScript, FNiagaraTranslatorOutput *TranslatorOutput, FString &TranslatedHLSL)
 {
+	SCOPE_CYCLE_COUNTER(STAT_NiagaraEditor_HlslCompiler_CompileScript);
+
 	//TODO: This should probably be done via the same route that other shaders take through the shader compiler etc.
 	//But that adds the complexity of a new shader type, new shader class and a new shader map to contain them etc.
 	//Can do things simply for now.
@@ -203,6 +210,7 @@ const FNiagaraCompileResults &FHlslNiagaraCompiler::CompileScript(UNiagaraScript
 	}
 	else
 	{
+		SCOPE_CYCLE_COUNTER(STAT_NiagaraEditor_HlslCompiler_CompileShader_VectorVM);
 		CompileResults.bVMSucceeded = CompileShader_VectorVM(Input, Output, FString(FPlatformProcess::ShaderDir()), 0, CompilationOutput);
 	}
 
@@ -217,9 +225,12 @@ const FNiagaraCompileResults &FHlslNiagaraCompiler::CompileScript(UNiagaraScript
 	//Eventually Niagara will have all the shader plumbing and do things like materials.
 	if (CompileResults.bVMSucceeded)
 	{
+		SCOPE_CYCLE_COUNTER(STAT_NiagaraEditor_HlslCompiler_CompileShader_VectorVMSucceeded);
 		check(TranslatorOutput);
-		Script->ByteCode = CompilationOutput.ByteCode;
+		Script->SetByteCode(CompilationOutput.ByteCode);
 		Script->LastHlslTranslation = TranslatedHLSL;
+		Script->LastAssemblyTranslation = CompilationOutput.AssemblyAsString;
+		Script->LastOpCount = CompilationOutput.NumOps;
 		Script->Attributes = TranslatorOutput->Attributes;
 		Script->Parameters = TranslatorOutput->Parameters;
 		Script->DataUsage.bReadsAttriubteData = TranslatorOutput->bReadsAttributeData;
@@ -256,7 +267,15 @@ const FNiagaraCompileResults &FHlslNiagaraCompiler::CompileScript(UNiagaraScript
 		for (FNiagaraScriptDataInterfaceInfo& DataInterfaceInfo : TranslatorOutput->DataInterfaceInfo)
 		{
 			int32 Idx = Script->DataInterfaceInfo.AddDefaulted();
-			Script->DataInterfaceInfo[Idx].DataInterface = Cast<UNiagaraDataInterface>(StaticDuplicateObject(DataInterfaceInfo.DataInterface, Script, NAME_None, ~RF_Transient));
+			if (DataInterfaceInfo.DataInterface)
+			{
+				Script->DataInterfaceInfo[Idx].DataInterface = Cast<UNiagaraDataInterface>(StaticDuplicateObject(DataInterfaceInfo.DataInterface, Script, NAME_None, ~RF_Transient));
+			}
+			else
+			{
+				Script->DataInterfaceInfo[Idx].DataInterface = nullptr;
+			}
+			Script->DataInterfaceInfo[Idx].Type = DataInterfaceInfo.Type;
 			Script->DataInterfaceInfo[Idx].Name = DataInterfaceInfo.Name;
 			Script->DataInterfaceInfo[Idx].UserPtrIdx = DataInterfaceInfo.UserPtrIdx;
 		}
@@ -304,7 +323,8 @@ const FNiagaraCompileResults &FHlslNiagaraCompiler::CompileScript(UNiagaraScript
 	if (CompileResults.bVMSucceeded == false)
 	{
 		//Some error. Clear script and exit.
-		Script->ByteCode.Empty();
+		Script->LastHlslTranslation = TranslatedHLSL;
+		Script->SetByteCode(CompilationOutput.ByteCode);
 		Script->Attributes.Empty();
 		Script->Parameters.Empty();
 		Script->InternalParameters.Empty();
diff --git a/Engine/Plugins/FX/Niagara/Source/NiagaraEditor/Private/NiagaraCurveOwner.cpp b/Engine/Plugins/FX/Niagara/Source/NiagaraEditor/Private/NiagaraCurveOwner.cpp
--- a/Engine/Plugins/FX/Niagara/Source/NiagaraEditor/Private/NiagaraCurveOwner.cpp
+++ b/Engine/Plugins/FX/Niagara/Source/NiagaraEditor/Private/NiagaraCurveOwner.cpp
@@ -87,6 +87,7 @@ void FNiagaraCurveOwner::OnCurveChanged(const TArray<FRichCurveEditInfo>& Change
 		if (CurveChanged != nullptr && CurveOwner != nullptr)
 		{
 			CurveChanged->Execute(ChangedCurveEditInfo.CurveToEdit, *CurveOwner);
+			(*CurveOwner)->PostEditChange();
 		}
 	}
 }
diff --git a/Engine/Plugins/FX/Niagara/Source/NiagaraEditor/Private/NiagaraEditorCommands.cpp b/Engine/Plugins/FX/Niagara/Source/NiagaraEditor/Private/NiagaraEditorCommands.cpp
--- a/Engine/Plugins/FX/Niagara/Source/NiagaraEditor/Private/NiagaraEditorCommands.cpp
+++ b/Engine/Plugins/FX/Niagara/Source/NiagaraEditor/Private/NiagaraEditorCommands.cpp
@@ -15,6 +15,12 @@ void FNiagaraEditorCommands::RegisterCommands()
 	UI_COMMAND(TogglePreviewGrid, "Grid", "Toggles the preview pane's grid.", EUserInterfaceActionType::ToggleButton, FInputChord() );
 	UI_COMMAND(TogglePreviewBackground, "Background", "Toggles the preview pane's background.", EUserInterfaceActionType::ToggleButton, FInputChord());
 	UI_COMMAND(ToggleUnlockToChanges, "Lock/Unlock To Changes", "Toggles whether or not changes in the source asset get pulled into this asset automatically.", EUserInterfaceActionType::Button, FInputChord());
+	UI_COMMAND(ToggleBounds, "Bounds", "Display Bounds", EUserInterfaceActionType::ToggleButton, FInputChord());
+	UI_COMMAND(ToggleBounds_SetFixedBounds, "Set Fixed Bounds", "Set Fixed Bounds", EUserInterfaceActionType::Button, FInputChord());
+	UI_COMMAND(ToggleOrbit, "Orbit Mode", "Toggle Orbit Navigation", EUserInterfaceActionType::ToggleButton, FInputChord());
+
+	UI_COMMAND(SaveThumbnailImage, "Thumbnail", "Generate Thumbnail", EUserInterfaceActionType::Button, FInputChord());
+
 }
 
 #undef LOCTEXT_NAMESPACE // NiagaraEditorCommands
diff --git a/Engine/Plugins/FX/Niagara/Source/NiagaraEditor/Private/NiagaraEditorCommands.h b/Engine/Plugins/FX/Niagara/Source/NiagaraEditor/Private/NiagaraEditorCommands.h
--- a/Engine/Plugins/FX/Niagara/Source/NiagaraEditor/Private/NiagaraEditorCommands.h
+++ b/Engine/Plugins/FX/Niagara/Source/NiagaraEditor/Private/NiagaraEditorCommands.h
@@ -35,4 +35,9 @@ public:
 
 	/** Toggles the locking/unlocking of refreshing from changes*/
 	TSharedPtr<FUICommandInfo> ToggleUnlockToChanges;
+
+	TSharedPtr<FUICommandInfo> ToggleOrbit;
+	TSharedPtr<FUICommandInfo> ToggleBounds;
+	TSharedPtr<FUICommandInfo> ToggleBounds_SetFixedBounds;
+	TSharedPtr<FUICommandInfo> SaveThumbnailImage;
 };
\ No newline at end of file
diff --git a/Engine/Plugins/FX/Niagara/Source/NiagaraEditor/Private/NiagaraEditorModule.cpp b/Engine/Plugins/FX/Niagara/Source/NiagaraEditor/Private/NiagaraEditorModule.cpp
--- a/Engine/Plugins/FX/Niagara/Source/NiagaraEditor/Private/NiagaraEditorModule.cpp
+++ b/Engine/Plugins/FX/Niagara/Source/NiagaraEditor/Private/NiagaraEditorModule.cpp
@@ -1,6 +1,7 @@
 // Copyright 1998-2017 Epic Games, Inc. All Rights Reserved.
 
 #include "NiagaraEditorModule.h"
+#include "NiagaraModule.h"
 #include "NiagaraEditorTickables.h"
 #include "Modules/ModuleManager.h"
 #include "IAssetTypeActions.h"
@@ -32,6 +33,7 @@
 #include "EdGraphSchema_Niagara.h"
 #include "TypeEditorUtilities/NiagaraFloatTypeEditorUtilities.h"
 #include "TypeEditorUtilities/NiagaraIntegerTypeEditorUtilities.h"
+#include "TypeEditorUtilities/NiagaraEnumTypeEditorUtilities.h"
 #include "TypeEditorUtilities/NiagaraBoolTypeEditorUtilities.h"
 #include "TypeEditorUtilities/NiagaraFloatTypeEditorUtilities.h"
 #include "TypeEditorUtilities/NiagaraVectorTypeEditorUtilities.h"
@@ -55,6 +57,8 @@
 #include "NiagaraEmitterViewModel.h"
 #include "TNiagaraGraphPinEditableName.h"
 #include "Class.h"
+#include "NiagaraScriptMergeManager.h"
+#include "NiagaraEmitter.h"
 
 #include "Customizations/NiagaraComponentDetails.h"
 #include "Customizations/NiagaraTypeCustomizations.h"
@@ -149,7 +153,9 @@ void FNiagaraEditorModule::StartupModule()
 	NiagaraAssetCategory = AssetTools.RegisterAdvancedAssetCategory(FName(TEXT("FX")), LOCTEXT("NiagaraAssetsCategory", "FX"));
 	RegisterAssetTypeAction(AssetTools, MakeShareable(new FAssetTypeActions_NiagaraSystem()));
 	RegisterAssetTypeAction(AssetTools, MakeShareable(new FAssetTypeActions_NiagaraEmitter()));
-	RegisterAssetTypeAction(AssetTools, MakeShareable(new FAssetTypeActions_NiagaraScript()));
+	RegisterAssetTypeAction(AssetTools, MakeShareable(new FAssetTypeActions_NiagaraScriptFunctions()));
+	RegisterAssetTypeAction(AssetTools, MakeShareable(new FAssetTypeActions_NiagaraScriptModules()));
+	RegisterAssetTypeAction(AssetTools, MakeShareable(new FAssetTypeActions_NiagaraScriptDynamicInputs()));
 	RegisterAssetTypeAction(AssetTools, MakeShareable(new FAssetTypeActions_NiagaraParameterCollection()));
 	RegisterAssetTypeAction(AssetTools, MakeShareable(new FAssetTypeActions_NiagaraParameterCollectionInstance()));
 
@@ -217,6 +223,7 @@ void FNiagaraEditorModule::StartupModule()
 	GraphPanelPinFactory->RegisterMiscSubCategoryPin(UNiagaraNodeWithDynamicPins::AddPinSubCategory, FNiagaraScriptGraphPanelPinFactory::FCreateGraphPin::CreateLambda(
 		[](UEdGraphPin* GraphPin) -> TSharedRef<SGraphPin> { return SNew(SNiagaraGraphPinAdd, GraphPin); }));
 
+	EnumTypeUtilities = MakeShareable(new FNiagaraEditorEnumTypeUtilities());
 	RegisterTypeUtilities(FNiagaraTypeDefinition::GetFloatDef(), MakeShareable(new FNiagaraEditorFloatTypeUtilities()));
 	RegisterTypeUtilities(FNiagaraTypeDefinition::GetIntDef(), MakeShareable(new FNiagaraEditorIntegerTypeUtilities()));
 	RegisterTypeUtilities(FNiagaraTypeDefinition::GetBoolDef(), MakeShareable(new FNiagaraEditorBoolTypeUtilities()));
@@ -254,6 +261,9 @@ void FNiagaraEditorModule::StartupModule()
 		FNiagaraShaderQueueTickable::ProcessQueue();
 	}));
 
+	// Register the emitter merge handler.
+	ScriptMergeManager = MakeShared<FNiagaraScriptMergeManager>();
+	MergeEmitterHandle = NiagaraModule.RegisterOnMergeEmitter(INiagaraModule::FOnMergeEmitter::CreateSP(ScriptMergeManager.ToSharedRef(), &FNiagaraScriptMergeManager::MergeEmitter));
 }
 
 
@@ -291,6 +301,12 @@ void FNiagaraEditorModule::ShutdownModule()
 		SequencerModule->UnRegisterTrackEditor(CreateEmitterTrackEditorHandle);
 	}
 
+	INiagaraModule* NiagaraModule = FModuleManager::GetModulePtr<INiagaraModule>("Niagara");
+	if (NiagaraModule != nullptr)
+	{
+		NiagaraModule->UnregisterOnMergeEmitter(MergeEmitterHandle);
+	}
+
 	// Verify that we've cleaned up all the view models in the world.
 	FNiagaraScriptViewModel::CleanAll();
 	FNiagaraSystemViewModel::CleanAll();
@@ -314,10 +330,16 @@ void FNiagaraEditorModule::RegisterTypeUtilities(FNiagaraTypeDefinition Type, TS
 TSharedPtr<INiagaraEditorTypeUtilities> FNiagaraEditorModule::GetTypeUtilities(const FNiagaraTypeDefinition& Type)
 {
 	TSharedRef<INiagaraEditorTypeUtilities>* EditorUtilities = TypeToEditorUtilitiesMap.Find(Type);
-	if (EditorUtilities != nullptr)
+	if(EditorUtilities != nullptr)
 	{
 		return *EditorUtilities;
 	}
+
+	if (Type.IsEnum())
+	{
+		return EnumTypeUtilities;
+	}
+
 	return TSharedPtr<INiagaraEditorTypeUtilities>();
 }
 
@@ -340,8 +362,10 @@ void FNiagaraEditorModule::ResetOnCreateStackWidget(FDelegateHandle Handle)
 	OnCreateStackWidget.Unbind();
 }
 
-
-
+TSharedRef<FNiagaraScriptMergeManager> FNiagaraEditorModule::GetScriptMergeManager() const
+{
+	return ScriptMergeManager.ToSharedRef();
+}
 
 void FNiagaraEditorModule::RegisterAssetTypeAction(IAssetTools& AssetTools, TSharedRef<IAssetTypeActions> Action)
 {
diff --git a/Engine/Plugins/FX/Niagara/Source/NiagaraEditor/Private/NiagaraEditorStyle.cpp b/Engine/Plugins/FX/Niagara/Source/NiagaraEditor/Private/NiagaraEditorStyle.cpp
--- a/Engine/Plugins/FX/Niagara/Source/NiagaraEditor/Private/NiagaraEditorStyle.cpp
+++ b/Engine/Plugins/FX/Niagara/Source/NiagaraEditor/Private/NiagaraEditorStyle.cpp
@@ -42,6 +42,7 @@ FName FNiagaraEditorStyle::GetStyleSetName()
 const FVector2D Icon16x16(16.0f, 16.0f);
 const FVector2D Icon20x20(20.0f, 20.0f);
 const FVector2D Icon40x40(40.0f, 40.0f);
+const FVector2D Icon64x64(64.0f, 64.0f);
 
 TSharedRef< FSlateStyleSet > FNiagaraEditorStyle::Create()
 {
@@ -161,6 +162,11 @@ TSharedRef< FSlateStyleSet > FNiagaraEditorStyle::Create()
 	Style->Set("NiagaraEditor.AssetColors.ParameterCollection", FLinearColor(1.0f, 1.0f, 0.3f));
 	Style->Set("NiagaraEditor.AssetColors.ParameterCollectionInstance", FLinearColor(1.0f, 1.0f, 0.7f));
 
+	// Script factory thumbnails
+	Style->Set("NiagaraEditor.Thumbnails.DynamicInputs", new IMAGE_BRUSH("Icons/NiagaraScriptDynamicInputs_64x", Icon64x64));
+	Style->Set("NiagaraEditor.Thumbnails.Functions", new IMAGE_BRUSH("Icons/NiagaraScriptFunction_64x", Icon64x64));
+	Style->Set("NiagaraEditor.Thumbnails.Modules", new IMAGE_BRUSH("Icons/NiagaraScriptModules_64x", Icon64x64));
+
 	return Style;
 }
 
diff --git a/Engine/Plugins/FX/Niagara/Source/NiagaraEditor/Private/NiagaraEditorUtilities.cpp b/Engine/Plugins/FX/Niagara/Source/NiagaraEditor/Private/NiagaraEditorUtilities.cpp
--- a/Engine/Plugins/FX/Niagara/Source/NiagaraEditor/Private/NiagaraEditorUtilities.cpp
+++ b/Engine/Plugins/FX/Niagara/Source/NiagaraEditor/Private/NiagaraEditorUtilities.cpp
@@ -28,30 +28,7 @@
 
 #define LOCTEXT_NAMESPACE "FNiagaraEditorUtilities"
 
-FName FNiagaraEditorUtilities::GetUniqueName(FName CandidateName, const TSet<FName>& ExistingNames)
-{
-	if (ExistingNames.Contains(CandidateName) == false)
-	{
-		return CandidateName;
-	}
-
-	FString CandidateNameString = CandidateName.ToString();
-	FString BaseNameString = CandidateNameString;
-	if (CandidateNameString.Len() >= 3 && CandidateNameString.Right(3).IsNumeric())
-	{
-		BaseNameString = CandidateNameString.Left(CandidateNameString.Len() - 3);
-	}
-
-	FName UniqueName = FName(*BaseNameString);
-	int32 NameIndex = 1;
-	while (ExistingNames.Contains(UniqueName))
-	{
-		UniqueName = FName(*FString::Printf(TEXT("%s%03i"), *BaseNameString, NameIndex));
-		NameIndex++;
-	}
 
-	return UniqueName;
-}
 
 TSet<FName> FNiagaraEditorUtilities::GetSystemConstantNames()
 {
@@ -67,22 +44,11 @@ void FNiagaraEditorUtilities::GetTypeDefaultValue(const FNiagaraTypeDefinition&
 {
 	if (const UScriptStruct* ScriptStruct = Type.GetScriptStruct())
 	{
-		int32 TypeSize = Type.GetSize();
-		DefaultData.SetNumUninitialized(TypeSize);
+		FNiagaraVariable DefaultVariable(Type, NAME_None);
+		ResetVariableToDefaultValue(DefaultVariable);
 
-		FNiagaraEditorModule& NiagaraEditorModule = FModuleManager::GetModuleChecked<FNiagaraEditorModule>("NiagaraEditor");
-		TSharedPtr<INiagaraEditorTypeUtilities> TypeEditorUtilities = NiagaraEditorModule.GetTypeUtilities(Type);
-		if (TypeEditorUtilities.IsValid() && TypeEditorUtilities->CanProvideDefaultValue())
-		{
-			TSharedRef<FStructOnScope> Struct = MakeShareable(new FStructOnScope(ScriptStruct));
-			TypeEditorUtilities->UpdateStructWithDefaultValue(Struct);	
-			FMemory::Memcpy(DefaultData.GetData(), Struct->GetStructMemory(), TypeSize);
-			
-		}
-		else
-		{
-			ScriptStruct->InitializeDefaultValue(DefaultData.GetData());
-		}
+		DefaultData.SetNumUninitialized(Type.GetSize());
+		DefaultVariable.CopyTo(DefaultData.GetData());
 	}
 }
 
@@ -94,9 +60,7 @@ void FNiagaraEditorUtilities::ResetVariableToDefaultValue(FNiagaraVariable& Vari
 		TSharedPtr<INiagaraEditorTypeUtilities> TypeEditorUtilities = NiagaraEditorModule.GetTypeUtilities(Variable.GetType());
 		if (TypeEditorUtilities.IsValid() && TypeEditorUtilities->CanProvideDefaultValue())
 		{
-			TSharedRef<FStructOnScope> Struct = MakeShareable(new FStructOnScope(ScriptStruct));
-			TypeEditorUtilities->UpdateStructWithDefaultValue(Struct);
-			Variable.SetData(Struct->GetStructMemory());
+			TypeEditorUtilities->UpdateVariableWithDefaultValue(Variable);
 		}
 		else
 		{
@@ -120,12 +84,12 @@ void FNiagaraEditorUtilities::InitializeParameterInputNode(UNiagaraNodeInput& In
 	if (Type.GetScriptStruct() != nullptr)
 	{
 		ResetVariableToDefaultValue(InputNode.Input);
-		InputNode.DataInterface = nullptr;
+		InputNode.SetDataInterface(nullptr);
 	}
 	else
 	{
 		InputNode.Input.AllocateData(); // Frees previously used memory if we're switching from a struct to a class type.
-		InputNode.DataInterface = NewObject<UNiagaraDataInterface>(&InputNode, const_cast<UClass*>(Type.GetClass()), NAME_None, RF_Transactional);
+		InputNode.SetDataInterface(NewObject<UNiagaraDataInterface>(&InputNode, const_cast<UClass*>(Type.GetClass()), NAME_None, RF_Transactional));
 	}
 }
 
@@ -182,11 +146,11 @@ bool FNiagaraEditorUtilities::ConvertToMergedGraph(UNiagaraEmitter* InEmitter)
 			GraphsToConvert.Add(CastChecked<UNiagaraScriptSource>(InEmitter->UpdateScriptProps.Script->GetSource())->NodeGraph);
 			GraphUsages.Add(ENiagaraScriptUsage::ParticleUpdateScript);
 
-			for (int32 i = 0; i < InEmitter->EventHandlerScriptProps.Num(); i++)
+			for (int32 i = 0; i < InEmitter->GetEventHandlers().Num(); i++)
 			{
-				if (InEmitter->EventHandlerScriptProps[i].Script != nullptr)
+				if (InEmitter->GetEventHandlers()[i].Script != nullptr)
 				{
-					GraphsToConvert.Add(CastChecked<UNiagaraScriptSource>(InEmitter->EventHandlerScriptProps[i].Script->GetSource())->NodeGraph);
+					GraphsToConvert.Add(CastChecked<UNiagaraScriptSource>(InEmitter->GetEventHandlers()[i].Script->GetSource())->NodeGraph);
 					GraphUsages.Add(ENiagaraScriptUsage::ParticleEventScript);
 				}
 			}
@@ -254,11 +218,11 @@ bool FNiagaraEditorUtilities::ConvertToMergedGraph(UNiagaraEmitter* InEmitter)
 			InEmitter->GraphSource = Source;
 			InEmitter->SpawnScriptProps.Script->SetSource(Source);
 			InEmitter->UpdateScriptProps.Script->SetSource(Source);
-			for (int32 i = 0; i < InEmitter->EventHandlerScriptProps.Num(); i++)
+			for (int32 i = 0; i < InEmitter->GetEventHandlers().Num(); i++)
 			{
-				if (InEmitter->EventHandlerScriptProps[i].Script != nullptr)
+				if (InEmitter->GetEventHandlers()[i].Script != nullptr)
 				{
-					InEmitter->EventHandlerScriptProps[i].Script->SetSource(Source);
+					InEmitter->GetEventHandlers()[i].Script->SetSource(Source);
 				}
 			}
 
@@ -281,7 +245,6 @@ bool FNiagaraEditorUtilities::ConvertToMergedGraph(UNiagaraEmitter* InEmitter)
 			
 			// Now make sure that anyone referencing these graphs knows that they are out-of-date.
 			Source->MarkNotSynchronized();
-			InEmitter->ChangeId = FGuid::NewGuid();
 
 			//FString ErrorMessages;
 			//FNiagaraEditorModule& NiagaraEditorModule = FModuleManager::Get().LoadModuleChecked<FNiagaraEditorModule>(TEXT("NiagaraEditor"));
@@ -380,7 +343,7 @@ void FNiagaraEditorUtilities::FixUpPastedInputNodes(UEdGraph* Graph, TSet<UEdGra
 			}
 			if (ExistingNames.Contains(PastedInput.GetName()))
 			{
-				FName UniqueName = FNiagaraEditorUtilities::GetUniqueName(PastedInput.GetName(), ExistingNames.Union(FNiagaraEditorUtilities::GetSystemConstantNames()));
+				FName UniqueName = FNiagaraUtilities::GetUniqueName(PastedInput.GetName(), ExistingNames.Union(FNiagaraEditorUtilities::GetSystemConstantNames()));
 				for (UNiagaraNodeInput* PastedNodeForInput : PastedNodesForInput)
 				{
 					PastedNodeForInput->Input.SetName(UniqueName);
@@ -466,15 +429,24 @@ bool FNiagaraEditorUtilities::DataMatches(const FNiagaraVariable& Variable, cons
 		return false;
 	}
 
-	const uint8* VariableMemory = Variable.GetData();
-	const uint8* StructOnScopeMemory = StructOnScope.GetStructMemory();
-	int32 Size = Variable.GetSizeInBytes();
-	for (int32 i = 0; i < Size; i++)
+	return FMemory::Memcmp(Variable.GetData(), StructOnScope.GetStructMemory(), Variable.GetSizeInBytes()) == 0;
+}
+
+bool FNiagaraEditorUtilities::DataMatches(const FNiagaraVariable& VariableA, const FNiagaraVariable& VariableB)
+{
+	if (VariableA.GetType() != VariableB.GetType())
 	{
-		if (VariableMemory[i] != StructOnScopeMemory[i])
-		{
-			return false;
-		}
+		return false;
+	}
+
+	if (VariableA.IsDataAllocated() != VariableB.IsDataAllocated())
+	{
+		return false;
+	}
+
+	if (VariableA.IsDataAllocated())
+	{
+		return FMemory::Memcmp(VariableA.GetData(), VariableB.GetData(), VariableA.GetSizeInBytes()) == 0;
 	}
 
 	return true;
@@ -487,18 +459,7 @@ bool FNiagaraEditorUtilities::DataMatches(const FStructOnScope& StructOnScopeA,
 		return false;
 	}
 
-	const uint8* StructOnScopeAMemory = StructOnScopeA.GetStructMemory();
-	const uint8* StructOnScopeBMemory = StructOnScopeB.GetStructMemory();
-	int32 Size = StructOnScopeA.GetStruct()->GetStructureSize();
-	for (int32 i = 0; i < Size; i++)
-	{
-		if (StructOnScopeAMemory[i] != StructOnScopeBMemory[i])
-		{
-			return false;
-		}
-	}
-
-	return true;
+	return FMemory::Memcmp(StructOnScopeA.GetStructMemory(), StructOnScopeB.GetStructMemory(), StructOnScopeA.GetStruct()->GetStructureSize()) == 0;
 }
 
 TSharedPtr<SWidget> FNiagaraEditorUtilities::CreateInlineErrorText(TAttribute<FText> ErrorMessage, TAttribute<FText> ErrorTooltip)
@@ -533,50 +494,55 @@ TSharedPtr<SWidget> FNiagaraEditorUtilities::CreateInlineErrorText(TAttribute<FT
 				];
 }
 
-void FNiagaraEditorUtilities::UpdateExistingEmitters(const TArray<UNiagaraEmitter*>& AffectedEmitters)
+void FNiagaraEditorUtilities::CompileExistingEmitters(const TArray<UNiagaraEmitter*>& AffectedEmitters)
 {
-	// Compile the existing emitters. Also determine which Systems need to be properly updated.
-	TArray<UNiagaraSystem*> AffectedSystemSystems;
+	TSet<UNiagaraEmitter*> CompiledEmitters;
 	for (UNiagaraEmitter* Emitter : AffectedEmitters)
 	{
-		if (Emitter->IsPendingKillOrUnreachable())
+		// If we've already compiled this emitter, or it's invalid skip it.
+		if (CompiledEmitters.Contains(Emitter) || Emitter->IsPendingKillOrUnreachable())
 		{
 			continue;
 		}
 
-		TSharedPtr<FNiagaraEmitterViewModel> EmitterViewModel = FNiagaraEmitterViewModel::GetExistingViewModelForObject(Emitter);
-		if (!EmitterViewModel.IsValid())
-		{
-			EmitterViewModel = MakeShareable(new FNiagaraEmitterViewModel(Emitter, nullptr));
-		}
-		EmitterViewModel->CompileScripts();
-
-		for (TObjectIterator<UNiagaraSystem> It; It; ++It)
+		// We only need to compile emitters referenced directly by systems since emitters can now only be used in the context 
+		// of a system.
+		for (TObjectIterator<UNiagaraSystem> SystemIterator; SystemIterator; ++SystemIterator)
 		{
-			if (It->GetAutoImportChangedEmitters() && It->ReferencesSourceEmitter(Emitter))
+			if (SystemIterator->ReferencesSourceEmitter(*Emitter))
 			{
-				AffectedSystemSystems.AddUnique(*It);
+				SystemIterator->Compile(false);
+
+				TArray<TSharedPtr<FNiagaraSystemViewModel>> ExistingSystemViewModels;
+				FNiagaraSystemViewModel::GetAllViewModelsForObject(*SystemIterator, ExistingSystemViewModels);
+				for (TSharedPtr<FNiagaraSystemViewModel> SystemViewModel : ExistingSystemViewModels)
+				{
+					SystemViewModel->RefreshAll();
+				}
+
+				for (const FNiagaraEmitterHandle& EmitterHandle : SystemIterator->GetEmitterHandles())
+				{
+					CompiledEmitters.Add(EmitterHandle.GetInstance());
+				}
 			}
 		}
 	}
+}
 
-	// Now iterate over the affected Systems.
-	for (UNiagaraSystem* System : AffectedSystemSystems)
+bool FNiagaraEditorUtilities::TryGetEventDisplayName(UNiagaraEmitter* Emitter, FGuid EventUsageId, FText& OutEventDisplayName)
+{
+	if (Emitter != nullptr)
 	{
-		TSharedPtr<FNiagaraSystemViewModel> SystemViewModel = FNiagaraSystemViewModel::GetExistingViewModelForObject(System);
-		if (!SystemViewModel.IsValid())
+		for (const FNiagaraEventScriptProperties& EventScriptProperties : Emitter->GetEventHandlers())
 		{
-			FNiagaraSystemViewModelOptions Options;
-			Options.bCanRemoveEmittersFromTimeline = false;
-			Options.bCanRenameEmittersFromTimeline = false;
-			Options.bCanAddEmittersFromTimeline = false;
-			Options.bUseSystemExecStateForTimelineReset = false;
-			SystemViewModel = MakeShareable(new FNiagaraSystemViewModel(*System, Options));
+			if (EventScriptProperties.Script->GetUsageId() == EventUsageId)
+			{
+				OutEventDisplayName = FText::FromName(EventScriptProperties.SourceEventName);
+				return true;
+			}
 		}
-
-		SystemViewModel->ResynchronizeAllHandles();
 	}
+	return false;
 }
 
-
 #undef LOCTEXT_NAMESPACE
diff --git a/Engine/Plugins/FX/Niagara/Source/NiagaraEditor/Private/NiagaraGraph.cpp b/Engine/Plugins/FX/Niagara/Source/NiagaraEditor/Private/NiagaraGraph.cpp
--- a/Engine/Plugins/FX/Niagara/Source/NiagaraEditor/Private/NiagaraGraph.cpp
+++ b/Engine/Plugins/FX/Niagara/Source/NiagaraEditor/Private/NiagaraGraph.cpp
@@ -24,7 +24,15 @@
 #include "NiagaraNodeParameterMapGet.h"
 #include "INiagaraEditorTypeUtilities.h"
 #include "NiagaraConstants.h"
+#include "NiagaraEditorModule.h"
 
+DECLARE_CYCLE_STAT(TEXT("NiagaraEditor - Graph - FindInputNodes"), STAT_NiagaraEditor_Graph_FindInputNodes, STATGROUP_NiagaraEditor);
+DECLARE_CYCLE_STAT(TEXT("NiagaraEditor - Graph - FindInputNodes_NotFilterUsage"), STAT_NiagaraEditor_Graph_FindInputNodes_NotFilterUsage, STATGROUP_NiagaraEditor);
+DECLARE_CYCLE_STAT(TEXT("NiagaraEditor - Graph - FindInputNodes_FilterUsage"), STAT_NiagaraEditor_Graph_FindInputNodes_FilterUsage, STATGROUP_NiagaraEditor);
+DECLARE_CYCLE_STAT(TEXT("NiagaraEditor - Graph - FindInputNodes_FilterDupes"), STAT_NiagaraEditor_Graph_FindInputNodes_FilterDupes, STATGROUP_NiagaraEditor);
+DECLARE_CYCLE_STAT(TEXT("NiagaraEditor - Graph - FindInputNodes_FindInputNodes_Sort"), STAT_NiagaraEditor_Graph_FindInputNodes_Sort, STATGROUP_NiagaraEditor);
+DECLARE_CYCLE_STAT(TEXT("NiagaraEditor - Graph - FindOutputNode"), STAT_NiagaraEditor_Graph_FindOutputNode, STATGROUP_NiagaraEditor);
+DECLARE_CYCLE_STAT(TEXT("NiagaraEditor - Graph - BuildTraversalHelper"), STAT_NiagaraEditor_Graph_BuildTraversalHelper, STATGROUP_NiagaraEditor);
 
 UNiagaraGraph::UNiagaraGraph(const FObjectInitializer& ObjectInitializer)
 	: Super(ObjectInitializer)
@@ -32,6 +40,20 @@ UNiagaraGraph::UNiagaraGraph(const FObjectInitializer& ObjectInitializer)
 	Schema = UEdGraphSchema_Niagara::StaticClass();
 	ChangeId = FGuid::NewGuid();
 }
+void UNiagaraGraph::NotifyGraphChanged(const FEdGraphEditAction& InAction)
+{
+	if ((InAction.Action & GRAPHACTION_AddNode) != 0 || (InAction.Action & GRAPHACTION_RemoveNode) != 0 ||
+		(InAction.Action & GRAPHACTION_GenericNeedsRecompile) != 0)
+	{
+		MarkGraphRequiresSynchronization();
+	}
+	Super::NotifyGraphChanged(InAction);
+}
+
+void UNiagaraGraph::NotifyGraphChanged()
+{
+	Super::NotifyGraphChanged();
+}
 
 void UNiagaraGraph::PostLoad()
 {
@@ -176,22 +198,19 @@ void UNiagaraGraph::FindOutputNodes(ENiagaraScriptUsage TargetUsageType, TArray<
 		}
 	}
 
-	NodesFound.Sort([](const UNiagaraNodeOutput& LHS, const UNiagaraNodeOutput& RHS) { return LHS.GetUsageIndex() > RHS.GetUsageIndex(); });
 	OutputNodes = NodesFound;
 }
 
-UNiagaraNodeOutput* UNiagaraGraph::FindOutputNode(ENiagaraScriptUsage TargetUsageType, int32 TargetOccurence) const
+UNiagaraNodeOutput* UNiagaraGraph::FindOutputNode(ENiagaraScriptUsage TargetUsageType, FGuid TargetUsageId) const
 {
+	SCOPE_CYCLE_COUNTER(STAT_NiagaraEditor_Graph_FindOutputNode);
 	for (UEdGraphNode* Node : Nodes)
 	{
 		if (UNiagaraNodeOutput* OutNode = Cast<UNiagaraNodeOutput>(Node))
 		{
-			if (OutNode->GetUsage() == TargetUsageType)
+			if (OutNode->GetUsage() == TargetUsageType && OutNode->GetUsageId() == TargetUsageId)
 			{
-				if (OutNode->ScriptTypeIndex == TargetOccurence)
-				{
-					return OutNode;
-				}
+				return OutNode;
 			}
 		}
 	}
@@ -205,21 +224,28 @@ void BuildTraversalHelper(TArray<class UNiagaraNode*>& OutNodesTraversed, UNiaga
 		return;
 	}
 
+	SCOPE_CYCLE_COUNTER(STAT_NiagaraEditor_Graph_BuildTraversalHelper);
+
 	TArray<UEdGraphPin*> Pins = CurrentNode->GetAllPins();
 	for (int32 i = 0; i < Pins.Num(); i++)
 	{
 		if (Pins[i]->Direction == EEdGraphPinDirection::EGPD_Input && Pins[i]->LinkedTo.Num() == 1)
 		{
-			BuildTraversalHelper(OutNodesTraversed, Cast<UNiagaraNode>(Pins[i]->LinkedTo[0]->GetOwningNode()));
+			UNiagaraNode* Node = Cast<UNiagaraNode>(Pins[i]->LinkedTo[0]->GetOwningNode());
+			if (OutNodesTraversed.Contains(Node))
+			{
+				continue;
+			}
+			BuildTraversalHelper(OutNodesTraversed, Node);
 		}
 	}
 
 	OutNodesTraversed.Add(CurrentNode);
 }
 
-void UNiagaraGraph::BuildTraversal(TArray<class UNiagaraNode*>& OutNodesTraversed, ENiagaraScriptUsage TargetUsage, int32 TargetOccurence) const
+void UNiagaraGraph::BuildTraversal(TArray<class UNiagaraNode*>& OutNodesTraversed, ENiagaraScriptUsage TargetUsage, FGuid TargetUsageId) const
 {
-	UNiagaraNodeOutput* Output = FindOutputNode(TargetUsage, TargetOccurence);
+	UNiagaraNodeOutput* Output = FindOutputNode(TargetUsage, TargetUsageId);
 	if (Output)
 	{
 		BuildTraversalHelper(OutNodesTraversed, Output);
@@ -237,10 +263,13 @@ void UNiagaraGraph::BuildTraversal(TArray<class UNiagaraNode*>& OutNodesTraverse
 
 void UNiagaraGraph::FindInputNodes(TArray<UNiagaraNodeInput*>& OutInputNodes, UNiagaraGraph::FFindInputNodeOptions Options) const
 {
+	SCOPE_CYCLE_COUNTER(STAT_NiagaraEditor_Graph_FindInputNodes);
 	TArray<UNiagaraNodeInput*> InputNodes;
 
 	if (!Options.bFilterByScriptUsage)
 	{
+		SCOPE_CYCLE_COUNTER(STAT_NiagaraEditor_Graph_FindInputNodes_NotFilterUsage);
+
 		for (UEdGraphNode* Node : Nodes)
 		{
 			UNiagaraNodeInput* NiagaraInputNode = Cast<UNiagaraNodeInput>(Node);
@@ -255,8 +284,10 @@ void UNiagaraGraph::FindInputNodes(TArray<UNiagaraNodeInput*>& OutInputNodes, UN
 	}
 	else
 	{
+		SCOPE_CYCLE_COUNTER(STAT_NiagaraEditor_Graph_FindInputNodes_FilterUsage);
+
 		TArray<class UNiagaraNode*> Traversal;
-		BuildTraversal(Traversal, Options.TargetScriptUsage, Options.TargetOccurence);
+		BuildTraversal(Traversal, Options.TargetScriptUsage, Options.TargetScriptUsageId);
 		for (UNiagaraNode* Node : Traversal)
 		{
 			UNiagaraNodeInput* NiagaraInputNode = Cast<UNiagaraNodeInput>(Node);
@@ -272,6 +303,8 @@ void UNiagaraGraph::FindInputNodes(TArray<UNiagaraNodeInput*>& OutInputNodes, UN
 
 	if (Options.bFilterDuplicates)
 	{
+		SCOPE_CYCLE_COUNTER(STAT_NiagaraEditor_Graph_FindInputNodes_FilterDupes);
+
 		for (UNiagaraNodeInput* InputNode : InputNodes)
 		{
 			auto NodeMatches = [=](UNiagaraNodeInput* UniqueInputNode)
@@ -299,6 +332,8 @@ void UNiagaraGraph::FindInputNodes(TArray<UNiagaraNodeInput*>& OutInputNodes, UN
 
 	if (Options.bSort)
 	{
+		SCOPE_CYCLE_COUNTER(STAT_NiagaraEditor_Graph_FindInputNodes_Sort);
+
 		UNiagaraNodeInput::SortNodes(OutInputNodes);
 	}
 }
@@ -392,6 +427,31 @@ void UNiagaraGraph::GetOutputNodeVariables(ENiagaraScriptUsage InScriptUsage, TA
 	}
 }
 
+bool UNiagaraGraph::HasParameterMapParameters()const
+{
+	TArray<FNiagaraVariable> Inputs;
+	TArray<FNiagaraVariable> Outputs;
+
+	GetParameters(Inputs, Outputs);
+
+	for (FNiagaraVariable& Var : Inputs)
+	{
+		if (Var.GetType() == FNiagaraTypeDefinition::GetParameterMapDef())
+		{
+			return true;
+		}
+	}
+	for (FNiagaraVariable& Var : Outputs)
+	{
+		if (Var.GetType() == FNiagaraTypeDefinition::GetParameterMapDef())
+		{
+			return true;
+		}
+	}
+
+	return false;
+}
+
 bool UNiagaraGraph::HasNumericParameters()const
 {
 	TArray<FNiagaraVariable> Inputs;
@@ -424,6 +484,11 @@ void UNiagaraGraph::NotifyGraphNeedsRecompile()
 	NotifyGraphChanged(Action);
 }
 
+void UNiagaraGraph::NotifyGraphDataInterfaceChanged()
+{
+	OnDataInterfaceChangedDelegate.Broadcast();
+}
+
 void UNiagaraGraph::SubsumeExternalDependencies(TMap<const UObject*, UObject*>& ExistingConversions)
 {
 	TArray<UNiagaraNode*> NiagaraNodes;
@@ -491,3 +556,58 @@ void UNiagaraGraph::MarkGraphRequiresSynchronization()
 	ChangeId = FGuid::NewGuid();
 	UE_LOG(LogNiagaraEditor, Log, TEXT("MarkGraphRequiresSynchronization %s"), *ChangeId.ToString());
 }
+
+/** Get the meta-data associated with this variable, if it exists.*/
+FNiagaraVariableMetaData* UNiagaraGraph::GetMetaData(const FNiagaraVariable& InVar)
+{
+	return VariableToMetaData.Find(InVar);
+}
+
+const FNiagaraVariableMetaData* UNiagaraGraph::GetMetaData(const FNiagaraVariable& InVar) const
+{
+	return VariableToMetaData.Find(InVar);
+}
+
+/** Return the meta-data associated with this variable. This should only be called on variables defined within this Graph, otherwise meta-data may leak.*/
+FNiagaraVariableMetaData& UNiagaraGraph::FindOrAddMetaData(const FNiagaraVariable& InVar)
+{
+	FNiagaraVariableMetaData* FoundMetaData = VariableToMetaData.Find(InVar);
+	if (FoundMetaData)
+	{
+		return *FoundMetaData;
+	}
+	// We shouldn't add constants to the graph's meta-data list. Those are stored globally.
+	ensure(FNiagaraConstants::IsNiagaraConstant(InVar) == false);
+	return VariableToMetaData.Add(InVar);
+}
+
+void UNiagaraGraph::PurgeUnreferencedMetaData()
+{
+	TArray<FNiagaraVariable> VarsToRemove;
+	for (auto It = VariableToMetaData.CreateConstIterator(); It; ++It)
+	{
+		int32 NumValid = 0;
+		for (TWeakObjectPtr<UObject> WeakPtr : It.Value().ReferencerNodes)
+		{
+			if (WeakPtr.IsValid())
+			{
+				NumValid++;
+			}
+		}
+
+		if (NumValid == 0)
+		{
+			VarsToRemove.Add(It.Key());
+		}
+	}
+
+	for (FNiagaraVariable& Var : VarsToRemove)
+	{
+		VariableToMetaData.Remove(Var);
+	}
+}
+
+UNiagaraGraph::FOnDataInterfaceChanged& UNiagaraGraph::OnDataInterfaceChanged()
+{
+	return OnDataInterfaceChangedDelegate;
+}
