Copying Orion-Staging @ 3748653 (Orion/Dev-General @ 3722124) to //UE4/Main

diff --git a/Engine/Plugins/FX/Niagara/Source/NiagaraEditor/Private/NiagaraHlslTranslator.cpp b/Engine/Plugins/FX/Niagara/Source/NiagaraEditor/Private/NiagaraHlslTranslator.cpp
--- a/Engine/Plugins/FX/Niagara/Source/NiagaraEditor/Private/NiagaraHlslTranslator.cpp
+++ b/Engine/Plugins/FX/Niagara/Source/NiagaraEditor/Private/NiagaraHlslTranslator.cpp
@@ -14,7 +14,7 @@
 #include "NiagaraNodeWriteDataSet.h"
 #include "NiagaraNodeParameterMapGet.h"
 #include "NiagaraNodeParameterMapSet.h"
-#include "NiagaraNodeParameterCollection.h"
+#include "NiagaraNodeCustomHlsl.h"
 #include "NiagaraNodeOp.h"
 #include "NiagaraNodeConvert.h"
 #include "EdGraphSchema_Niagara.h"
@@ -26,6 +26,7 @@
 #include "NiagaraNodeEmitter.h"
 #include "INiagaraEditorTypeUtilities.h"
 #include "NiagaraEditorUtilities.h"
+#include "NiagaraEditorModule.h"
 
 #include "NiagaraDataInterface.h"
 #include "NiagaraDataInterfaceCurve.h"
@@ -39,9 +40,42 @@
 #include "NiagaraParameterCollection.h"
 #include "NiagaraEditorTickables.h"
 #include "ShaderCore.h"
+#include "NiagaraShaderCompilationManager.h"
 
 #define LOCTEXT_NAMESPACE "NiagaraCompiler"
 
+DECLARE_CYCLE_STAT(TEXT("Niagara - HlslTranslator - Translate"), STAT_NiagaraEditor_HlslTranslator_Translate, STATGROUP_NiagaraEditor);
+DECLARE_CYCLE_STAT(TEXT("Niagara - HlslTranslator - CloneGraphAndPrepareForCompilation"), STAT_NiagaraEditor_HlslTranslator_CloneGraphAndPrepareForCompilation, STATGROUP_NiagaraEditor);
+DECLARE_CYCLE_STAT(TEXT("Niagara - HlslTranslator - BuildParameterMapHlslDefinitions"), STAT_NiagaraEditor_HlslTranslator_BuildParameterMapHlslDefinitions, STATGROUP_NiagaraEditor);
+DECLARE_CYCLE_STAT(TEXT("Niagara - HlslTranslator - Emitter"), STAT_NiagaraEditor_Module_NiagaraHLSLTranslator_Emitter, STATGROUP_NiagaraEditor);
+DECLARE_CYCLE_STAT(TEXT("Niagara - HlslTranslator - MapGet"), STAT_NiagaraEditor_Module_NiagaraHLSLTranslator_MapGet, STATGROUP_NiagaraEditor);
+DECLARE_CYCLE_STAT(TEXT("Niagara - HlslTranslator - FunctionCall"), STAT_NiagaraEditor_Module_NiagaraHLSLTranslator_FunctionCall, STATGROUP_NiagaraEditor);
+DECLARE_CYCLE_STAT(TEXT("Niagara - HlslTranslator - FunctionCallCloneGraphNumeric"), STAT_NiagaraEditor_Module_NiagaraHLSLTranslator_FunctionCallCloneGraphNumeric, STATGROUP_NiagaraEditor);
+DECLARE_CYCLE_STAT(TEXT("Niagara - HlslTranslator - FunctionCallCloneGraphNonNumeric"), STAT_NiagaraEditor_Module_NiagaraHLSLTranslator_FunctionCallCloneGraphNonNumeric, STATGROUP_NiagaraEditor);
+DECLARE_CYCLE_STAT(TEXT("Niagara - HlslTranslator - RegisterFunctionCall"), STAT_NiagaraEditor_Module_NiagaraHLSLTranslator_RegisterFunctionCall, STATGROUP_NiagaraEditor);
+DECLARE_CYCLE_STAT(TEXT("Niagara - HlslTranslator - GenerateFunctionCall"), STAT_NiagaraEditor_Module_NiagaraHLSLTranslator_GenerateFunctionCall, STATGROUP_NiagaraEditor);
+DECLARE_CYCLE_STAT(TEXT("Niagara - HlslTranslator - CustomHLSL"), STAT_NiagaraEditor_Module_NiagaraHLSLTranslator_CustomHLSL, STATGROUP_NiagaraEditor);
+DECLARE_CYCLE_STAT(TEXT("Niagara - HlslTranslator - FuncBody"), STAT_NiagaraEditor_Module_NiagaraHLSLTranslator_FuncBody, STATGROUP_NiagaraEditor);
+DECLARE_CYCLE_STAT(TEXT("Niagara - HlslTranslator - Output"), STAT_NiagaraEditor_HlslTranslator_Output, STATGROUP_NiagaraEditor);
+DECLARE_CYCLE_STAT(TEXT("Niagara - HlslTranslator - MapSet"), STAT_NiagaraEditor_HlslTranslator_MapSet, STATGROUP_NiagaraEditor);
+DECLARE_CYCLE_STAT(TEXT("Niagara - HlslTranslator - Operation"), STAT_NiagaraEditor_HlslTranslator_Operation, STATGROUP_NiagaraEditor);
+DECLARE_CYCLE_STAT(TEXT("Niagara - HlslTranslator - If"), STAT_NiagaraEditor_HlslTranslator_If, STATGROUP_NiagaraEditor);
+DECLARE_CYCLE_STAT(TEXT("Niagara - HlslTranslator - CompilePin"), STAT_NiagaraEditor_HlslTranslator_CompilePin, STATGROUP_NiagaraEditor);
+DECLARE_CYCLE_STAT(TEXT("Niagara - HlslTranslator - CompileOutputPin"), STAT_NiagaraEditor_HlslTranslator_CompileOutputPin, STATGROUP_NiagaraEditor);
+DECLARE_CYCLE_STAT(TEXT("Niagara - HlslTranslator - GetParameter"), STAT_NiagaraEditor_HlslTranslator_GetParameter, STATGROUP_NiagaraEditor);
+DECLARE_CYCLE_STAT(TEXT("Niagara - HlslTranslator - GenerateFunctionSignature"), STAT_NiagaraEditor_Module_NiagaraHLSLTranslator_GenerateFunctionSignature, STATGROUP_NiagaraEditor);
+DECLARE_CYCLE_STAT(TEXT("Niagara - HlslTranslator - RegisterFunctionCall_Source"), STAT_NiagaraEditor_Module_NiagaraHLSLTranslator_RegisterFunctionCall_Source, STATGROUP_NiagaraEditor);
+DECLARE_CYCLE_STAT(TEXT("Niagara - HlslTranslator - RegisterFunctionCall_Compile"), STAT_NiagaraEditor_Module_NiagaraHLSLTranslator_RegisterFunctionCall_Compile, STATGROUP_NiagaraEditor);
+DECLARE_CYCLE_STAT(TEXT("Niagara - HlslTranslator - RegisterFunctionCall_Signature"), STAT_NiagaraEditor_Module_NiagaraHLSLTranslator_RegisterFunctionCall_Signature, STATGROUP_NiagaraEditor);
+DECLARE_CYCLE_STAT(TEXT("Niagara - HlslTranslator - RegisterFunctionCall_FunctionDefStr"), STAT_NiagaraEditor_Module_NiagaraHLSLTranslator_RegisterFunctionCall_FunctionDefStr, STATGROUP_NiagaraEditor);
+DECLARE_CYCLE_STAT(TEXT("Niagara - HlslTranslator - GenerateFunctionSignature_UniqueDueToMaps"), STAT_NiagaraEditor_Module_NiagaraHLSLTranslator_GenerateFunctionSignature_UniqueDueToMaps, STATGROUP_NiagaraEditor);
+DECLARE_CYCLE_STAT(TEXT("Niagara - HlslTranslator - GenerateFunctionSignature_Outputs"), STAT_NiagaraEditor_Module_NiagaraHLSLTranslator_GenerateFunctionSignature_Outputs, STATGROUP_NiagaraEditor);
+DECLARE_CYCLE_STAT(TEXT("Niagara - HlslTranslator - GenerateFunctionSignature_Inputs"), STAT_NiagaraEditor_Module_NiagaraHLSLTranslator_GenerateFunctionSignature_Inputs, STATGROUP_NiagaraEditor);
+DECLARE_CYCLE_STAT(TEXT("Niagara - HlslTranslator - GenerateFunctionSignature_FindInputNodes"), STAT_NiagaraEditor_Module_NiagaraHLSLTranslator_GenerateFunctionSignature_FindInputNodes, STATGROUP_NiagaraEditor);
+
+
+
+
 
 static FNiagaraShaderQueueTickable NiagaraShaderQueueProcessor;
 FNiagaraShaderProcessorTickable NiagaraShaderProcessor;
@@ -55,16 +89,28 @@ void FNiagaraShaderQueueTickable::ProcessQueue()
 	for (FNiagaraCompilationQueue::NiagaraCompilationQueueItem &Item : FNiagaraCompilationQueue::Get()->GetQueue())
 	{
 		FNiagaraScript *Script = Item.Script;
+		TRefCountPtr<FNiagaraShaderMap>NewShaderMap = Item.ShaderMap;
+
+		if (Script == nullptr) // This script has been removed from the pending queue post submission... just skip it.
+		{
+			FNiagaraShaderMap::RemovePendingMap(NewShaderMap);
+			NewShaderMap->SetCompiledSuccessfully(false);
+			UE_LOG(LogNiagaraEditor, Log, TEXT("GPU shader compile skipped. Id %d"), NewShaderMap->GetCompilingId());
+			continue;
+		}
 		UNiagaraScript* CompilableScript = Script->GetGTScript();
 
 		// For now System scripts don't generate HLSL and go through a special pass...
 		// [OP] thinking they'll likely never run on GPU anyways
-		if (CompilableScript->IsValidLowLevel() == false || CompilableScript->IsSystemSpawnScript() || CompilableScript->IsSystemUpdateScript() || CompilableScript->IsEmitterSpawnScript() || CompilableScript->IsEmitterUpdateScript())
+		if (CompilableScript->IsValidLowLevel() == false || CompilableScript->CanBeRunOnGpu() == false)
 		{
+			NewShaderMap->SetCompiledSuccessfully(false);
+			FNiagaraShaderMap::RemovePendingMap(NewShaderMap);
+			Script->RemoveOutstandingCompileId(NewShaderMap->GetCompilingId());
+			UE_LOG(LogNiagaraEditor, Log, TEXT("GPU shader compile skipped. Id %d"), NewShaderMap->GetCompilingId());
 			continue;
 		}
 
-		TRefCountPtr<FNiagaraShaderMap>NewShaderMap = Item.ShaderMap;
 		FNiagaraComputeShaderCompilationOutput NewCompilationOutput;
 		FHlslNiagaraTranslator NiagaraTranslator;
 		FHlslNiagaraTranslatorOptions Options;
@@ -73,38 +119,33 @@ void FNiagaraShaderQueueTickable::ProcessQueue()
 		bool bNeedsPostCompile = false;
 		if (CompilableScript->GetSource()->IsPreCompiled() == false)
 		{
-		    CompilableScript->GetSource()->PreCompile(Cast<UNiagaraEmitter>(CompilableScript->GetOuter()));
+			CompilableScript->GetSource()->PreCompile(Cast<UNiagaraEmitter>(CompilableScript->GetOuter()));
 			bNeedsPostCompile = true;
 		}
 		FNiagaraTranslateResults Results = NiagaraTranslator.Translate(CompilableScript, Options);
+		CompilableScript->LastHlslTranslationGPU = Results.OutputHLSL;
+
 		if (bNeedsPostCompile)
 		{
 			CompilableScript->GetSource()->PostCompile();
 		}
 		Script->SetDatainterfaceBufferDescriptors(NiagaraTranslator.GetDatainterfaceBufferDescriptors());
-
+		CompilableScript->SetDatainterfaceBufferDescriptors(NiagaraTranslator.GetDatainterfaceBufferDescriptors());
 		// make sure buffers are set up on the original script
 		for (int32 i = 0; i < NiagaraTranslator.GetDatainterfaceBufferDescriptors().Num(); i++)
 		{
 			if (CompilableScript->DataInterfaceInfo.Num() > i)
 			{
-				CompilableScript->DataInterfaceInfo[i].DataInterface->SetupBuffers(NiagaraTranslator.GetDatainterfaceBufferDescriptors()[i]);
+				if (CompilableScript->DataInterfaceInfo[i].DataInterface)
+				{
+					CompilableScript->DataInterfaceInfo[i].DataInterface->SetupBuffers(NiagaraTranslator.GetDatainterfaceBufferDescriptors()[i]);
+				}
 			}
 		}
 
-		if (Results.bHLSLGenSucceeded && Results.MessageLog->NumErrors == 0)
-		{
-			Script->HlslOutput = Results.OutputHLSL;
-			// Create a shader compiler environment for the script that will be shared by all jobs from this script
-			TRefCountPtr<FShaderCompilerEnvironment> CompilerEnvironment = new FShaderCompilerEnvironment();
-
-			FString ShaderCode = NiagaraTranslator.GetTranslatedHLSL();
-			const bool bSynchronousCompile = false;
+		Script->HlslOutput = Results.OutputHLSL;
 
-			// Compile the shaders for the script.
-			NewShaderMap->Compile(Script, Item.ShaderMapId, CompilerEnvironment, NewCompilationOutput, Item.Platform, bSynchronousCompile, Item.bApply);
-		}
-		else
+		if (!Results.bHLSLGenSucceeded || Results.MessageLog->NumErrors > 0)
 		{
 			UE_LOG(LogNiagaraEditor, Warning, TEXT("%s"), TEXT("Niagara GPU Translation Failed!"));
 			TArray<TSharedRef<FTokenizedMessage>> Messages;
@@ -128,7 +169,22 @@ void FNiagaraShaderQueueTickable::ProcessQueue()
 					UE_LOG(LogNiagaraEditor, Error, TEXT("%s"), *Message->ToText().ToString());
 				}
 			}
+			NewShaderMap->SetCompiledSuccessfully(false);
+			FNiagaraShaderMap::RemovePendingMap(NewShaderMap);
+
+		}
+		else
+		{
+			// Create a shader compiler environment for the script that will be shared by all jobs from this script
+			TRefCountPtr<FShaderCompilerEnvironment> CompilerEnvironment = new FShaderCompilerEnvironment();
+
+			FString ShaderCode = NiagaraTranslator.GetTranslatedHLSL();
+			const bool bSynchronousCompile = false;
+
+			// Compile the shaders for the script.
+			NewShaderMap->Compile(Script, Item.ShaderMapId, CompilerEnvironment, NewCompilationOutput, Item.Platform, bSynchronousCompile, Item.bApply);
 		}
+
 	}
 
 	FNiagaraCompilationQueue::Get()->GetQueue().Empty();
@@ -381,6 +437,14 @@ void PreprocessFunctionGraph(FHlslNiagaraTranslator& Compiler, const UEdGraphSch
 	}
 
 	FixUpNumericPins(Compiler, Schema, OutputNode);
+
+	// Fix up numeric pins for data set writes since they may have numerics which were missed traversing from the output node.
+	TArray<UNiagaraNodeWriteDataSet*> DataSetWriteNodes;
+	Graph->GetNodesOfClass<UNiagaraNodeWriteDataSet>(DataSetWriteNodes);
+	for (UNiagaraNodeWriteDataSet* DataSetWriteNode : DataSetWriteNodes)
+	{
+		FixUpNumericPins(Compiler, Schema, DataSetWriteNode);
+	}
 }
 
 ENiagaraScriptCompileStatus FNiagaraTranslateResults::TranslateResultsToSummary(const FNiagaraTranslateResults* TranslateResults)
@@ -494,50 +558,70 @@ bool FHlslNiagaraTranslator::ValidateTypePins(UNiagaraNode* NodeToValidate)
 }
 
 
-void FHlslNiagaraTranslator::GenerateFunctionSignature(ENiagaraScriptUsage ScriptUsage, FString InName, const FString& InFullName, UNiagaraGraph* FuncGraph, TArray<int32>& Inputs, bool bHadNumericInputs, FNiagaraFunctionSignature& OutSig)const
+void FHlslNiagaraTranslator::GenerateFunctionSignature(ENiagaraScriptUsage ScriptUsage, FString InName, const FString& InFullName, UNiagaraGraph* FuncGraph, TArray<int32>& Inputs, 
+	bool bHadNumericInputs, bool bHasParameterMapParameters, FNiagaraFunctionSignature& OutSig)const
 {
-	TArray<FNiagaraVariable> InputVars;
+	SCOPE_CYCLE_COUNTER(STAT_NiagaraEditor_Module_NiagaraHLSLTranslator_GenerateFunctionSignature);
 
-	// Only handle nodes connected to the correct output node in the event of multiple output nodes in the graph.
+	TArray<FNiagaraVariable> InputVars;
 	TArray<UNiagaraNodeInput*> InputsNodes;
-	UNiagaraGraph::FFindInputNodeOptions Options;
-	Options.bSort = true;
-	Options.bFilterDuplicates = true;
-	Options.bFilterByScriptUsage = true;
-	Options.TargetScriptUsage = ScriptUsage;
-	FuncGraph->FindInputNodes(InputsNodes, Options);
 
-	if (Inputs.Num() != InputsNodes.Num())
+	// Only handle nodes connected to the correct output node in the event of multiple output nodes in the graph.
 	{
-		const_cast<FHlslNiagaraTranslator*>(this)->Error(FText::Format(LOCTEXT("GenerateFunctionSignatureFail", "Generating function signature for {0} failed.  The function graph is invalid."), FText::FromString(InFullName)), nullptr, nullptr);
-		return;
+		SCOPE_CYCLE_COUNTER(STAT_NiagaraEditor_Module_NiagaraHLSLTranslator_GenerateFunctionSignature_FindInputNodes);
+		InputsNodes.Reserve(100);
+		UNiagaraGraph::FFindInputNodeOptions Options;
+		Options.bSort = true;
+		Options.bFilterDuplicates = true;
+		Options.bFilterByScriptUsage = true;
+		Options.TargetScriptUsage = ScriptUsage;
+		FuncGraph->FindInputNodes(InputsNodes, Options);
+
+		if (Inputs.Num() != InputsNodes.Num())
+		{
+			const_cast<FHlslNiagaraTranslator*>(this)->Error(FText::Format(LOCTEXT("GenerateFunctionSignatureFail", "Generating function signature for {0} failed.  The function graph is invalid."), FText::FromString(InFullName)), nullptr, nullptr);
+			return;
+		}
 	}
 
-	for (int32 i = 0; i < InputsNodes.Num(); ++i)
 	{
-		//Only add to the signature if the caller has provided it, otherwise we use a local default.
-		if (Inputs[i] != INDEX_NONE)
+		SCOPE_CYCLE_COUNTER(STAT_NiagaraEditor_Module_NiagaraHLSLTranslator_GenerateFunctionSignature_Inputs);
+
+		InName.Reserve(100 * InputsNodes.Num());
+		InputVars.Reserve(InputsNodes.Num());
+		for (int32 i = 0; i < InputsNodes.Num(); ++i)
 		{
-			InputVars.Add(InputsNodes[i]->Input);
-			if (bHadNumericInputs)
+			//Only add to the signature if the caller has provided it, otherwise we use a local default.
+			if (Inputs[i] != INDEX_NONE)
 			{
-				InName += TEXT("_In") + InputsNodes[i]->Input.GetType().GetName();
+				InputVars.Add(InputsNodes[i]->Input);
+				if (bHadNumericInputs)
+				{
+					InName += TEXT("_In");
+					InName += InputsNodes[i]->Input.GetType().GetName();
+				}
 			}
 		}
-	}
 
-	//Now actually remove the missing inputs so they match the signature.
-	Inputs.Remove(INDEX_NONE);
+		//Now actually remove the missing inputs so they match the signature.
+		Inputs.Remove(INDEX_NONE);
+	}
 
 	TArray<FNiagaraVariable> OutputVars;
-	FuncGraph->GetOutputNodeVariables(ScriptUsage, OutputVars);
-
-	for (int32 i = 0; i < OutputVars.Num(); ++i)
 	{
-		//Only add to the signature if the caller has provided it, otherwise we use a local default.
-		if (bHadNumericInputs)
+		SCOPE_CYCLE_COUNTER(STAT_NiagaraEditor_Module_NiagaraHLSLTranslator_GenerateFunctionSignature_Outputs);
+
+		OutputVars.Reserve(10);
+		FuncGraph->GetOutputNodeVariables(ScriptUsage, OutputVars);
+
+		for (int32 i = 0; i < OutputVars.Num(); ++i)
 		{
-			InName += TEXT("_Out") + OutputVars[i].GetType().GetName();
+			//Only add to the signature if the caller has provided it, otherwise we use a local default.
+			if (bHadNumericInputs)
+			{
+				InName += TEXT("_Out");
+				InName += OutputVars[i].GetType().GetName();
+			}
 		}
 	}
 
@@ -546,10 +630,13 @@ void FHlslNiagaraTranslator::GenerateFunctionSignature(ENiagaraScriptUsage Scrip
 	// For now, we want each module call to be unique due to parameter maps and aliasing causing different variables
 	// to be written within each call.
 	if ((ScriptUsage == ENiagaraScriptUsage::Module || ScriptUsage == ENiagaraScriptUsage::DynamicInput ||
-		ScriptUsage == ENiagaraScriptUsage::EmitterSpawnScript || ScriptUsage == ENiagaraScriptUsage::EmitterUpdateScript)
+		ScriptUsage == ENiagaraScriptUsage::EmitterSpawnScript || ScriptUsage == ENiagaraScriptUsage::EmitterUpdateScript
+		|| bHasParameterMapParameters)
 		&& (ModuleAliasStr != nullptr || EmitterAliasStr != nullptr))
 	{
+		SCOPE_CYCLE_COUNTER(STAT_NiagaraEditor_Module_NiagaraHLSLTranslator_GenerateFunctionSignature_UniqueDueToMaps);
 		FString SignatureName;
+		SignatureName.Reserve(1024);
 		if (ModuleAliasStr != nullptr)
 		{
 			SignatureName = GetSanitizedSymbolName(*ModuleAliasStr);
@@ -557,9 +644,10 @@ void FHlslNiagaraTranslator::GenerateFunctionSignature(ENiagaraScriptUsage Scrip
 		if (EmitterAliasStr != nullptr)
 		{
 			FString Prefix = ModuleAliasStr != nullptr ? TEXT("_") : TEXT("");
-			SignatureName +=  Prefix + GetSanitizedSymbolName(*EmitterAliasStr);
+			SignatureName += Prefix;
+			SignatureName += GetSanitizedSymbolName(*EmitterAliasStr);
 		}
-		SignatureName = SignatureName.Replace(TEXT("."), TEXT("_"));
+		SignatureName.ReplaceInline(TEXT("."), TEXT("_"));
 		OutSig = FNiagaraFunctionSignature(*SignatureName, InputVars, OutputVars, *InFullName, true, false);
 	}
 	else
@@ -576,6 +664,7 @@ FHlslNiagaraTranslator::FHlslNiagaraTranslator()
 	, TranslateResults(&MessageLog)
 	, bInsideInterpolatedSpawnScript(false)
 	, VisitID(INDEX_NONE)
+	, CurrentBodyChunkMode(ENiagaraCodeChunkMode::Body)
 {
 	// Make the message log silent so we're not spamming the blueprint log.
 	MessageLog.bSilentMode = true;
@@ -604,6 +693,7 @@ FString FHlslNiagaraTranslator::GetFunctionDefinitions()
 
 UNiagaraGraph* FHlslNiagaraTranslator::CloneGraphAndPrepareForCompilation(UNiagaraScript* InScript, UNiagaraScriptSource* InSource, bool bClearErrors)
 {
+	SCOPE_CYCLE_COUNTER(STAT_NiagaraEditor_HlslTranslator_CloneGraphAndPrepareForCompilation);
 	if (InSource == nullptr)
 	{
 		Error(LOCTEXT("NoSourceErrorMessage", "Script has no source."), nullptr, nullptr);
@@ -685,6 +775,7 @@ FString FHlslNiagaraTranslator::BuildMissingDefaults(TArray<TPair<FNiagaraVariab
 
 FString FHlslNiagaraTranslator::BuildParameterMapHlslDefinitions(TArray<FNiagaraVariable>& PrimaryDataSetOutputEntries, TArray<TPair<FNiagaraVariable, const UEdGraphPin*> >& MissingDefaults)
 {
+	SCOPE_CYCLE_COUNTER(STAT_NiagaraEditor_HlslTranslator_BuildParameterMapHlslDefinitions);
 	FString HlslOutputString;
 
 	// Determine the unique parameter map structs...
@@ -938,6 +1029,7 @@ void FHlslNiagaraTranslator::HandleNamespacedExternalVariablesToDataSetRead(TArr
 
 const FNiagaraTranslateResults &FHlslNiagaraTranslator::Translate(UNiagaraScript* InScript, FHlslNiagaraTranslatorOptions Options)
 {
+	SCOPE_CYCLE_COUNTER(STAT_NiagaraEditor_HlslTranslator_Translate);
 	check(InScript);
 
 	TranslationOptions = Options;
@@ -970,7 +1062,7 @@ const FNiagaraTranslateResults &FHlslNiagaraTranslator::Translate(UNiagaraScript
 	}
 
 	// Find the output node and compile it.
-	UNiagaraNodeOutput* OutputNode = SourceGraph->FindOutputNode(GetTargetUsage(), GetTargetOccurrence());
+	UNiagaraNodeOutput* OutputNode = SourceGraph->FindOutputNode(GetTargetUsage(), GetTargetUsageId());
 	if (OutputNode == nullptr)
 	{
 		Error(FText::Format(LOCTEXT("GetOutputNodeFail", "Cannot find output node of type {0}!"), FText::AsNumber((int32)GetTargetUsage())), nullptr, nullptr);
@@ -997,7 +1089,7 @@ const FNiagaraTranslateResults &FHlslNiagaraTranslator::Translate(UNiagaraScript
 	ParamMapHistories.Empty();
 	ParamMapSetVariablesToChunks.Empty();
 	TArray<UNiagaraNodeOutput*> UsageOutputs;
-	if (Script->IsInterpolatedParticleSpawnScript())
+	if (Script->IsInterpolatedParticleSpawnScript() || Options.SimTarget == ENiagaraSimTarget::GPUComputeSim)
 	{
 		UsageOutputs.Add(OutputNode);
 		UsageOutputs.Add(SourceGraph->FindOutputNode(ENiagaraScriptUsage::ParticleUpdateScript));
@@ -1041,24 +1133,37 @@ const FNiagaraTranslateResults &FHlslNiagaraTranslator::Translate(UNiagaraScript
 		}
 	}
 
-
+	Script->ParameterCollections.Empty();
+	for (FNiagaraParameterMapHistory& History : ParamMapHistories)
+	{
+		for (UNiagaraParameterCollection* Collection : History.ParameterCollections)
+		{
+			Script->ParameterCollections.AddUnique(Collection);
+		}
+	}
+	
 	//Create main scope pin cache.
 	PinToCodeChunks.AddDefaulted(1);
 
 	ActiveHistoryForFunctionCalls.BeginTranslation(InScript);
 
 	Script->StatScopes.Empty();
-	EnterStatsScope(FNiagaraStatScope(*Script->GetFullName(), FText::FromString(Script->GetFullName())));
-
-	Script->ParameterCollections.Empty();
+	EnterStatsScope(FNiagaraStatScope(*Script->GetFullName(), *Script->GetFullName()));
 
 	FHlslNiagaraTranslator* ThisTranslator = this;
 	TArray<int32> OutputChunks;
-	if (Script->IsInterpolatedParticleSpawnScript())
+	if (Script->IsInterpolatedParticleSpawnScript() || Options.SimTarget == ENiagaraSimTarget::GPUComputeSim)
 	{
+		CurrentBodyChunkMode = ENiagaraCodeChunkMode::SpawnBody;
 		//Here we compile the spawn script but write to temporary outputs in the context.
-		AddBodyChunk(TEXT(""), TEXT("//Interpolated Spawn Script!"), FNiagaraTypeDefinition::GetIntDef(), false, false);
-		AddBodyChunk(TEXT(""), TEXT("//Begin Spawn Script!"), FNiagaraTypeDefinition::GetIntDef(), false, false);
+		if (Script->IsInterpolatedParticleSpawnScript())
+		{
+			AddBodyChunk(TEXT(""), TEXT("//Begin Interpolated Spawn Script!"), FNiagaraTypeDefinition::GetIntDef(), false, false);
+		}
+		else
+		{
+			AddBodyChunk(TEXT(""), TEXT("//Begin Spawn Script!"), FNiagaraTypeDefinition::GetIntDef(), false, false);
+		}
 		bInsideInterpolatedSpawnScript = true;
 		CurrentParamMapIndices.Empty();
 		CurrentParamMapIndices.Add(0);
@@ -1067,21 +1172,25 @@ const FNiagaraTranslateResults &FHlslNiagaraTranslator::Translate(UNiagaraScript
 		InstanceWrite = FDataSetAccessInfo(); // Reset after building the output..
 		AddBodyChunk(TEXT(""), TEXT("//End Spawn Script!\n\n"), FNiagaraTypeDefinition::GetIntDef(), false, false);
 
-		AddBodyChunk(TEXT(""), TEXT("//Begin Transfer of Attributes!"), FNiagaraTypeDefinition::GetIntDef(), false, false);
-		AddBodyChunk(TEXT(""), TEXT("Context.MapUpdate.Particles = Context.MapSpawn.Particles"), FNiagaraTypeDefinition::GetIntDef(), false, true);
-		AddBodyChunk(TEXT(""), TEXT("//End Transfer of Attributes!\n\n"), FNiagaraTypeDefinition::GetIntDef(), false, false);
-		
-		AddBodyChunk(TEXT(""), TEXT("HandleMissingDefaultValues(Context);"), FNiagaraTypeDefinition::GetIntDef(), false, false);
+		if (Script->IsInterpolatedParticleSpawnScript())
+		{
+			AddBodyChunk(TEXT(""), TEXT("//Begin Transfer of Attributes!"), FNiagaraTypeDefinition::GetIntDef(), false, false);
+			AddBodyChunk(TEXT(""), TEXT("Context.MapUpdate.Particles = Context.MapSpawn.Particles"), FNiagaraTypeDefinition::GetIntDef(), false, true);
+			AddBodyChunk(TEXT(""), TEXT("//End Transfer of Attributes!\n\n"), FNiagaraTypeDefinition::GetIntDef(), false, false);
+			AddBodyChunk(TEXT(""), TEXT("HandleMissingDefaultValues(Context);"), FNiagaraTypeDefinition::GetIntDef(), false, false);
+		}
 
+		CurrentBodyChunkMode = ENiagaraCodeChunkMode::UpdateBody;
 		AddBodyChunk(TEXT(""), TEXT("//Begin Update Script!"), FNiagaraTypeDefinition::GetIntDef(), false, false);
 		//Now we compile the update script (with partial dt) and read from the temp values written above.
- 		UNiagaraNodeOutput* UpdateOutputNode = SourceGraph->FindOutputNode(ENiagaraScriptUsage::ParticleUpdateScript);
- 		check(UpdateOutputNode);
+		UNiagaraNodeOutput* UpdateOutputNode = SourceGraph->FindOutputNode(ENiagaraScriptUsage::ParticleUpdateScript);
+		check(UpdateOutputNode);
 		CurrentParamMapIndices.Empty();
 		CurrentParamMapIndices.Add(1);
 		UpdateOutputNode->Compile(ThisTranslator, OutputChunks);
 
 		AddBodyChunk(TEXT(""), TEXT("//End Update Script!\n\n"), FNiagaraTypeDefinition::GetIntDef(), false, false);
+		CurrentBodyChunkMode = ENiagaraCodeChunkMode::Body;
 	}
 	else
 	{
@@ -1188,13 +1297,15 @@ const FNiagaraTranslateResults &FHlslNiagaraTranslator::Translate(UNiagaraScript
 				
 		if (IsBulkSystemScript())
 		{
+			// We have two sets of data that can change independently.. The engine data set are variables
+			// that are essentially set once per system. The constants are rapid iteration variables
+			// that exist per emitter and change infrequently. Since they are so different, putting
+			// them in two distinct read data sets seems warranted.
 			TArray<FNiagaraVariable>& SystemEngineReadVars = DataSetReads.Add(GetSystemEngineDataSetID());
-			//TArray<FNiagaraVariable>& SystemUserReadVars = DataSetReads.Add(GetSystemUserDataSetID());
+			TArray<FNiagaraVariable>& SystemConstantReadVars = DataSetReads.Add(GetSystemConstantDataSetID());
 
 			DataSetReadVars.Add(&SystemEngineReadVars);
-			//DataSetReadVars.Add(&SystemUserReadVars);
 			DataSetReadIds.Add(GetSystemEngineDataSetID());
-			//DataSetReadIds.Add(GetSystemUserDataSetID());
 
 			HandleNamespacedExternalVariablesToDataSetRead(SystemEngineReadVars, TEXT("Engine"));
 			HandleNamespacedExternalVariablesToDataSetRead(SystemEngineReadVars/*SystemUserReadVars*/, TEXT("User"));
@@ -1204,15 +1315,12 @@ const FNiagaraTranslateResults &FHlslNiagaraTranslator::Translate(UNiagaraScript
 			{
 				return A.GetName() < B.GetName();
 			});
-			// We sort the variables so that they end up in the same ordering between Spawn & Update...
-			/*SystemUserReadVars.Sort([&](const FNiagaraVariable& A, const FNiagaraVariable& B)
-			{
-				return A.GetName() < B.GetName();
-			});*/
 
-			FNiagaraParameters ExternalParams;
-			ExternalParams.Parameters = SystemEngineReadVars;
-			Script->DataSetToParameters.Add(GetSystemEngineDataSetID().Name, ExternalParams);
+			{
+				FNiagaraParameters ExternalParams;
+				ExternalParams.Parameters = SystemEngineReadVars;
+				Script->DataSetToParameters.Add(GetSystemEngineDataSetID().Name, ExternalParams);
+			}
 		}
 
 		// Now we pull in the HLSL generated above by building the parameter map definitions..
@@ -1331,7 +1439,34 @@ const FNiagaraTranslateResults &FHlslNiagaraTranslator::Translate(UNiagaraScript
 		DefineDataSetWriteFunction(HlslOutput, InScript->WriteDataSets, WriteConditionVars);
 
 
-		//Define the shared per instance simulation function	
+		//Define the shared per instance simulation function
+		// for interpolated scripts AND GPU sim, define spawn and sim in separate functions
+		if (Script->IsInterpolatedParticleSpawnScript() || Options.SimTarget == ENiagaraSimTarget::GPUComputeSim)
+		{
+			HlslOutput += TEXT("void SimulateSpawn(inout FSimulationContext Context)\n{\n");
+			for (int32 i = 0; i < ChunksByMode[(int32)ENiagaraCodeChunkMode::SpawnBody].Num(); ++i)
+			{
+				HlslOutput += TEXT("\t") + GetCode(ChunksByMode[(int32)ENiagaraCodeChunkMode::SpawnBody][i]);
+			}
+			HlslOutput += TEXT("}\n");
+
+			HlslOutput += TEXT("void SimulateUpdate(inout FSimulationContext Context)\n{\n");
+			for (int32 i = 0; i < ChunksByMode[(int32)ENiagaraCodeChunkMode::UpdateBody].Num(); ++i)
+			{
+				HlslOutput += TEXT("\t") + GetCode(ChunksByMode[(int32)ENiagaraCodeChunkMode::UpdateBody][i]);
+			}
+			HlslOutput += TEXT("}\n");
+
+			// add simulate function 
+			if (Options.SimTarget == ENiagaraSimTarget::CPUSim)
+			{
+				HlslOutput += TEXT("void Simulate(inout FSimulationContext Context)\n{\n");
+				HlslOutput += TEXT("\tSimulateSpawn(Context);\n");
+				HlslOutput += TEXT("\tSimulateUpdate(Context);\n");
+				HlslOutput += TEXT("}\n");
+			}
+		}
+		else
 		{
 			HlslOutput += TEXT("void Simulate(inout FSimulationContext Context)\n{\n");
 			for (int32 i = 0; i < ChunksByMode[(int32)ENiagaraCodeChunkMode::Body].Num(); ++i)
@@ -1366,7 +1501,16 @@ const FNiagaraTranslateResults &FHlslNiagaraTranslator::Translate(UNiagaraScript
 				CompilationOutput.Attributes.AddUnique(Var);
 			}
 		}
-		
+
+		// We may have created some transient data interfaces. This cleans up the ones that we created.
+		for (int32 Idx = 0; Idx < CompilationOutput.DataInterfaceInfo.Num(); Idx++)
+		{
+			if (CompilationOutput.DataInterfaceInfo[Idx].DataInterface != nullptr &&
+				PlaceholderDataInterfaces.Contains(CompilationOutput.DataInterfaceInfo[Idx].DataInterface))
+			{
+				CompilationOutput.DataInterfaceInfo[Idx].DataInterface = nullptr;
+			}
+		}
 		CompilationOutput.bReadsAttributeData = InstanceRead.Variables.Num() != 0;
 		TranslateResults.OutputHLSL = HlslOutput;
 	}
@@ -1518,9 +1662,9 @@ void FHlslNiagaraTranslator::DefineInterpolatedParametersFunction(FString &HlslO
 			HlslOutputString += TEXT("\tfloat InterpSpawn_InvSpawnTime = 1.0 / InterpSpawn_SpawnTime;\n");
 			HlslOutputString += TEXT("\tfloat InterpSpawn_InvUpdateTime = 1.0 / InterpSpawn_UpdateTime;\n");
 			HlslOutputString += TEXT("\tfloat SpawnInterp = InterpSpawn_SpawnTime * Engine_InverseDeltaTime ;\n");
-			for (int32 i = 0; i < ChunksByMode[(int32)ENiagaraCodeChunkMode::Uniform].Num(); ++i)
+			for (int32 UniformIdx = 0; UniformIdx < ChunksByMode[(int32)ENiagaraCodeChunkMode::Uniform].Num(); ++UniformIdx)
 			{
-				int32 UniformIdx = ChunksByMode[(int32)ENiagaraCodeChunkMode::Uniform][i];
+				int32 ChunkIdx = ChunksByMode[(int32)ENiagaraCodeChunkMode::Uniform][UniformIdx];
 				if (UniformIdx != INDEX_NONE)
 				{
 					FNiagaraVariable* FoundNamespacedVar = nullptr;
@@ -1543,9 +1687,17 @@ void FHlslNiagaraTranslator::DefineInterpolatedParametersFunction(FString &HlslO
 
 					if (FoundNamespacedVar != nullptr)
 					{
-						FString FoundName = FoundNamespacedVar->GetName().ToString();
-						FNiagaraCodeChunk& Chunk = CodeChunks[UniformIdx];
-						HlslOutputString += TEXT("\tContext.MapSpawn.") + FoundName + TEXT(" = lerp(PREV__") + Chunk.SymbolName + TEXT(", ") + Chunk.SymbolName + TEXT(", SpawnInterp);\n");
+						FString FoundName = GetSanitizedSymbolName(FoundNamespacedVar->GetName().ToString());
+						FNiagaraCodeChunk& Chunk = CodeChunks[ChunkIdx];
+						if (ShouldInterpolateParameter(*FoundNamespacedVar))
+						{
+							HlslOutputString += TEXT("\tContext.MapSpawn.") + FoundName + TEXT(" = lerp(PREV__") + Chunk.SymbolName + TEXT(", ") + Chunk.SymbolName + TEXT(", SpawnInterp);\n");
+						}
+						else
+						{
+							// For now, we do nothing for non-floating point variables..
+							//HlslOutputString += TEXT("\tContext.MapSpawn.") + FoundName + TEXT(" = ") + Chunk.SymbolName + TEXT(";\n");
+						}
 					}
 				}
 			}
@@ -1688,16 +1840,23 @@ void FHlslNiagaraTranslator::DefineDataInterfaceHLSL(FString &InHlslOutput)
 	{
 		FNiagaraScriptDataInterfaceInfo &Info = CompilationOutput.DataInterfaceInfo[i];
 		
+		if (!Info.DataInterface)
+		{
+			Error(LOCTEXT("InvalidDataInterfacePtr", "Null DataInterface pointer!"), nullptr, nullptr);
+			return;
+		}
+
 		if (Info.DataInterface->CanExecuteOnTarget(ENiagaraSimTarget::GPUComputeSim))
 		{
 			TArray<FNiagaraFunctionSignature> DataInterfaceFunctions;
 			Info.DataInterface->GetFunctions(DataInterfaceFunctions);
 			FString OwnerIDString = Info.Name.ToString();
+			FString SanitizedOwnerIDString = GetSanitizedSymbolName(OwnerIDString, true);
 
 			// grab the buffer definition from the interface
 			//
 			int32 NewIdx = DIBufferDescriptors.AddDefaulted(1);
-			Info.DataInterface->GetBufferDefinitionHLSL(OwnerIDString, DIBufferDescriptors[NewIdx], InterfaceUniformHLSL);
+			Info.DataInterface->GetBufferDefinitionHLSL(SanitizedOwnerIDString, DIBufferDescriptors[NewIdx].Descriptors, InterfaceUniformHLSL);
 
 			// grab the function hlsl from the interface
 			//
@@ -1707,7 +1866,8 @@ void FHlslNiagaraTranslator::DefineDataInterfaceHLSL(FString &InHlslOutput)
 				Sig.OwnerName = Info.Name;
 				FString DefStr = GetFunctionSignatureSymbol(Sig);
 
-				ensure(Info.DataInterface->GetFunctionHLSL(DefStr, DIBufferDescriptors[NewIdx], OwnerIDString, InterfaceFunctionHLSL) == true);
+				bool HlslOK = Info.DataInterface->GetFunctionHLSL(Sig.Name, DefStr, DIBufferDescriptors[NewIdx].Descriptors, SanitizedOwnerIDString, InterfaceFunctionHLSL);
+				ensure(HlslOK == true);
 			}
 		}
 		else
@@ -1726,13 +1886,17 @@ void FHlslNiagaraTranslator::DefineMain(FString &OutHlslOutput,
 {
 	if (CompilationTarget == ENiagaraSimTarget::GPUComputeSim)
 	{
-		OutHlslOutput += TEXT("void SimulateMain(in int InstanceIdx, in int InEventIndex)\n{\n\tFSimulationContext Context = (FSimulationContext)0;\n");
+		OutHlslOutput += TEXT("void SimulateMain(in int InstanceIdx, in int InEventIndex, in int Phase)\n{\n");
 	}
 	else
 	{
-		OutHlslOutput += TEXT("void SimulateMain()\n{\n\tFSimulationContext Context = (FSimulationContext)0;\n");
+		OutHlslOutput += TEXT("void SimulateMain()\n{\n");
 	}
 
+	EnterStatsScope(FNiagaraStatScope(*(Script->GetName() + TEXT("_Main")), TEXT("Main")), OutHlslOutput);
+
+
+	OutHlslOutput += TEXT("\n\tFSimulationContext Context = (FSimulationContext)0;\n");
 	TMap<FName, int32> InputRegisterAllocations;
 	TMap<FName, int32> OutputRegisterAllocations;
 
@@ -1776,8 +1940,25 @@ void FHlslNiagaraTranslator::DefineMain(FString &OutHlslOutput,
 			OutHlslOutput += TEXT("\tInterpolateParameters(Context);\n");
 		}
 
-		//Call simulate.
-		OutHlslOutput += TEXT("\tSimulate(Context);\n");
+		// branch between spawn and update
+		if (CompilationTarget == ENiagaraSimTarget::GPUComputeSim)
+		{
+			//Call simulate.
+			if (Script->IsInterpolatedParticleSpawnScript())
+			{
+				OutHlslOutput += TEXT("\tif(Phase==0)\n\t{\n\t\tSimulateSpawn(Context);\n\t\tContext.MapSpawn.DataInstance.Alive=true;\n\t}\n");
+			}
+			else
+			{
+				OutHlslOutput += TEXT("\tif(Phase==0)\n\t{\n\t\tSimulateSpawn(Context);\n\t\tContext.Map.DataInstance.Alive=true;\n\t}\n");
+			}
+			OutHlslOutput += TEXT("\telse if(Phase==1)\n\t{\n\t\tSimulateUpdate(Context);\n\t}\n");
+		}
+		else
+		{
+			//Call simulate.
+			OutHlslOutput += TEXT("\tSimulate(Context);\n");
+		}
 	}
 
 	// write secondary data sets
@@ -1791,6 +1972,7 @@ void FHlslNiagaraTranslator::DefineMain(FString &OutHlslOutput,
 		DefineDataSetVariableWrites(HlslOutput, WriteIds[VarArrayIdx], VarArrayIdx, *ArrayRef);
 	}
 
+	ExitStatsScope(OutHlslOutput);
 	OutHlslOutput += TEXT("}\n");
 
 
@@ -1802,22 +1984,21 @@ void FHlslNiagaraTranslator::DefineMain(FString &OutHlslOutput,
 	if (CompilationTarget == ENiagaraSimTarget::GPUComputeSim)
 	{
 		OutHlslOutput += TEXT("void CopyInstance(in int InstanceIdx)\n{\n\tFSimulationContext Context = (FSimulationContext)0;\n");
-		for (int32 VarArrayIdx = 0; VarArrayIdx < InstanceReadVars.Num(); VarArrayIdx++)
-		{
-			TArray<FNiagaraVariable>* ArrayRef = InstanceReadVars[VarArrayIdx];
-			DefineDataSetVariableReads(HlslOutput, ReadIds[VarArrayIdx], VarArrayIdx, *ArrayRef);
-		}
-
-
 		if (Script->IsParticleEventScript())
 		{
+			for (int32 VarArrayIdx = 0; VarArrayIdx < InstanceReadVars.Num(); VarArrayIdx++)
+			{
+				TArray<FNiagaraVariable>* ArrayRef = InstanceReadVars[VarArrayIdx];
+				DefineDataSetVariableReads(HlslOutput, ReadIds[VarArrayIdx], VarArrayIdx, *ArrayRef);
+			}
+
 			OutHlslOutput += TEXT("\tContext.Map.DataInstance.Alive = true;\n");
-		}
 
-		for (int32 VarArrayIdx = 0; VarArrayIdx < InstanceWriteVars.Num(); VarArrayIdx++)
-		{
-			TArray<FNiagaraVariable>* ArrayRef = InstanceWriteVars[VarArrayIdx];
-			DefineDataSetVariableWrites(HlslOutput, WriteIds[VarArrayIdx], VarArrayIdx, *ArrayRef);
+			for (int32 VarArrayIdx = 0; VarArrayIdx < InstanceWriteVars.Num(); VarArrayIdx++)
+			{
+				TArray<FNiagaraVariable>* ArrayRef = InstanceWriteVars[VarArrayIdx];
+				DefineDataSetVariableWrites(HlslOutput, WriteIds[VarArrayIdx], VarArrayIdx, *ArrayRef);
+			}
 		}
 		OutHlslOutput += TEXT("}\n");
 	}
@@ -1871,8 +2052,10 @@ void FHlslNiagaraTranslator::DefineDataSetVariableWrites(FString &OutHlslOutput,
 				OutHlslOutput += TEXT("\tbool bValid = Context.MapUpdate.") + DataSetName + TEXT(".Alive;\n");
 			}
 		}
-		else if ((Script->IsParticleSpawnScript() && bHasPerParticleAliveSpawn) || (Script->IsParticleUpdateScript() && bHasPerParticleAliveUpdate) ||
-			(Script->IsParticleEventScript() && bHasPerParticleAliveEvent))
+		else if ((Script->IsParticleSpawnScript() && bHasPerParticleAliveSpawn) 
+			|| (Script->IsParticleUpdateScript() && bHasPerParticleAliveUpdate) 
+			|| (Script->IsParticleEventScript() && bHasPerParticleAliveEvent)
+			|| (Script->IsParticleSpawnScript() && bHasPerParticleAliveUpdate && TranslationOptions.SimTarget == ENiagaraSimTarget::GPUComputeSim))
 		{
 			OutHlslOutput += TEXT("\tbool bValid = Context.Map.") + DataSetName + TEXT(".Alive;\n");
 		}
@@ -1908,36 +2091,108 @@ void FHlslNiagaraTranslator::DefineDataSetVariableReads(FString &OutHlslOutput,
 {
 	int32 ReadOffset = 0;
 	FString DataSetName = Id.Name.ToString();
-	
-	for (FNiagaraVariable &Var : ReadVars)
+	FString Fmt;
+	FString SpawnCondition = "\t{\n", UpdateCondition = "\t{\n";
+
+	if (TranslationOptions.SimTarget == ENiagaraSimTarget::GPUComputeSim)
 	{
-		// If coming from a parameter map, use the one on the context, otherwise use the input.
-		FString Fmt;
+		SpawnCondition = "\tif(Phase==0)\n\t{\n";
+		UpdateCondition = "\tif(Phase==1)\n\t{\n";
+	}
 
+	// interpolated spawn or particle spawn on CPU
+	FString VarReads;
+	if (Script->IsInterpolatedParticleSpawnScript())
+	{
+		for (FNiagaraVariable &Var : ReadVars)
 		{
-			// In the case of interpolated spawn scripts, we read into the MapSpawn and write out of the Map.
-			if (Script->IsInterpolatedParticleSpawnScript())
+			Fmt = TEXT("\tContext.MapSpawn.") + GetSanitizedSymbolName(Var.GetName().ToString()) + TEXT("{0} = {4};\n");
+			GatherVariableForDataSetAccess(Var, Fmt, ReadOffset, DataSetIndex, TEXT(""), VarReads);
+		}
+		OutHlslOutput += SpawnCondition;
+		OutHlslOutput += VarReads + "\t}\n";
+	}
+
+	// regular emitter or system spawn script
+	else if ((Script->IsParticleSpawnScript()) || Script->IsEmitterSpawnScript() || Script->IsSystemSpawnScript()) 	// We shouldn't read anything in a Spawn Script! Set to defaults.
+	{
+		for (FNiagaraVariable &Var : ReadVars)
+		{
+			Fmt = TEXT("\tContext.Map.") + GetSanitizedSymbolName(Var.GetName().ToString()) + TEXT("{0} = {4};\n");
+			GatherVariableForDataSetAccess(Var, Fmt, ReadOffset, DataSetIndex, TEXT(""), VarReads);
+		}
+		OutHlslOutput += SpawnCondition;
+		OutHlslOutput += VarReads + "\t}\n";
+	}
+
+	VarReads = "";
+
+	if (Script->IsParticleUpdateScript() || Script->IsParticleEventScript() || Script->IsEmitterUpdateScript() || Script->IsSystemUpdateScript()
+		|| (Script->IsParticleSpawnScript() && TranslationOptions.SimTarget == ENiagaraSimTarget::GPUComputeSim))
+	{
+		FString ContextName = TEXT("\tContext.Map.");
+		if (Script->IsInterpolatedParticleSpawnScript())
+		{
+			ContextName = TEXT("\tContext.MapUpdate.");
+		}
+
+		// if we're a GPU spawn script (meaning a combined spawn/update script), we need to reset register index counter
+		if (Script->IsParticleSpawnScript() && TranslationOptions.SimTarget == ENiagaraSimTarget::GPUComputeSim)
+		{
+			ReadOffset = 0;
+		}
+		for (FNiagaraVariable &Var : ReadVars)
+		{
+			// If the NiagaraClearEachFrame value is set on the data set, we don't bother reading it in each frame as we know that it is is invalid. However,
+			// this is only used for the base data set. Other reads are potentially from events and are therefore perfectly valid.
+			if (DataSetIndex == 0 && Var.GetType().GetScriptStruct() != nullptr && Var.GetType().GetScriptStruct()->GetMetaData(TEXT("NiagaraClearEachFrame")).Equals(TEXT("true"), ESearchCase::IgnoreCase))
 			{
-				Fmt = TEXT("\tContext.MapSpawn.") + GetSanitizedSymbolName(Var.GetName().ToString()) + TEXT("{0} = {4};\n");
+				Fmt = ContextName + GetSanitizedSymbolName(Var.GetName().ToString()) + TEXT("{0} = {4};\n");
 			}
-			else if (Script->IsParticleSpawnScript() || Script->IsEmitterSpawnScript() || Script->IsSystemSpawnScript()) 	// We shouldn't read anything in a Spawn Script! Set to defaults.
+			else if (CompilationTarget == ENiagaraSimTarget::GPUComputeSim)
 			{
-				Fmt = TEXT("\tContext.Map.") + GetSanitizedSymbolName(Var.GetName().ToString()) + TEXT("{0} = {4};\n");
+				Fmt = ContextName + GetSanitizedSymbolName(Var.GetName().ToString()) + TEXT("{0} = InputData{1}({2}, {3}, InstanceIdx);\n");
 			}
 			else
 			{
-				if (CompilationTarget == ENiagaraSimTarget::GPUComputeSim)
-				{
-					Fmt = TEXT("\tContext.Map.") + GetSanitizedSymbolName(Var.GetName().ToString()) + TEXT("{0} = InputData{1}({2}, {3}, InstanceIdx);\n");
-				}
-				else
-				{
-					Fmt = TEXT("\tContext.Map.") + GetSanitizedSymbolName(Var.GetName().ToString()) + TEXT("{0} = InputData{1}({2}, {3});\n");
-				}
+				Fmt = ContextName + GetSanitizedSymbolName(Var.GetName().ToString()) + TEXT("{0} = InputData{1}({2}, {3});\n");
+			}
+			GatherVariableForDataSetAccess(Var, Fmt, ReadOffset, DataSetIndex, TEXT(""), VarReads);
+		}
+
+		OutHlslOutput += UpdateCondition;
+		OutHlslOutput += VarReads + "\n\t}\n";
+	}
+
+	/*
+	for (FNiagaraVariable &Var : ReadVars)
+	{
+		// If coming from a parameter map, use the one on the context, otherwise use the input.
+
+		// In the case of interpolated spawn scripts, we read into the MapSpawn and write out of the Map.
+		if (Script->IsInterpolatedParticleSpawnScript())
+		{
+			Fmt = TEXT("\tContext.MapSpawn.") + GetSanitizedSymbolName(Var.GetName().ToString()) + TEXT("{0} = {4};\n");
+		}
+		else if ( (Script->IsParticleSpawnScript()) || Script->IsEmitterSpawnScript() || Script->IsSystemSpawnScript()) 	// We shouldn't read anything in a Spawn Script! Set to defaults.
+		{
+			Fmt = TEXT("\tContext.Map.") + GetSanitizedSymbolName(Var.GetName().ToString()) + TEXT("{0} = {4};\n");
+		}
+		else
+		{
+			if (CompilationTarget == ENiagaraSimTarget::GPUComputeSim)
+			{
+				Fmt = TEXT("\tContext.Map.") + GetSanitizedSymbolName(Var.GetName().ToString()) + TEXT("{0} = InputData{1}({2}, {3}, InstanceIdx);\n");
+			}
+			else
+			{
+				Fmt = TEXT("\tContext.Map.") + GetSanitizedSymbolName(Var.GetName().ToString()) + TEXT("{0} = InputData{1}({2}, {3});\n");
 			}
 		}
+
 		GatherVariableForDataSetAccess(Var, Fmt, ReadOffset, DataSetIndex, TEXT(""), OutHlslOutput);
 	}
+	*/
 }
 
 void FHlslNiagaraTranslator::WriteDataSetContextVars(TMap<FNiagaraDataSetID, TMap<int32, FDataSetAccessInfo>>& DataSetAccessInfo, bool bRead, FString &OutHlslOutput)
@@ -2047,33 +2302,39 @@ void FHlslNiagaraTranslator::DecomposeVariableAccess(const UStruct* Struct, bool
 	}
 };
 
-FString FHlslNiagaraTranslator::GetSanitizedSymbolName(FString SymbolName)
+FString FHlslNiagaraTranslator::GetSanitizedSymbolName(FString SymbolName, bool bCollapsNamespaces)
 {
-	FString Ret = SymbolName.Replace(TEXT(" "), TEXT(""));
-	Ret = Ret.Replace(TEXT("\\"), TEXT("_"));
-	Ret = Ret.Replace(TEXT("/"), TEXT("_"));
-	Ret = Ret.Replace(TEXT(","), TEXT("_"));
-	Ret = Ret.Replace(TEXT("-"), TEXT("_"));
-	Ret = Ret.Replace(TEXT(":"), TEXT("_"));
+	FString Ret = SymbolName;
+	Ret.ReplaceInline(TEXT(" "), TEXT(""));
+	Ret.ReplaceInline(TEXT("\\"), TEXT("_"));
+	Ret.ReplaceInline(TEXT("/"), TEXT("_"));
+	Ret.ReplaceInline(TEXT(","), TEXT("_"));
+	Ret.ReplaceInline(TEXT("-"), TEXT("_"));
+	Ret.ReplaceInline(TEXT(":"), TEXT("_"));
+	if (bCollapsNamespaces)
+	{
+		Ret.ReplaceInline(TEXT("."), TEXT("_"));
+	}
 	return Ret;
 }
 
 FString FHlslNiagaraTranslator::GetUniqueSymbolName(FName BaseName)
 {
-	uint32* NameCount = SymbolCounts.Find(BaseName);
+	FString RetString = GetSanitizedSymbolName(BaseName.ToString());
+	FName RetName = *RetString;
+	uint32* NameCount = SymbolCounts.Find(RetName);
 	if (NameCount == nullptr)
 	{
-		SymbolCounts.Add(BaseName) = 1;
-		return GetSanitizedSymbolName(BaseName.ToString());
+		SymbolCounts.Add(RetName) = 1;
+		return RetString;
 	}
 
-	FString Ret = GetSanitizedSymbolName(BaseName.ToString());
 	if(*NameCount > 0)
 	{
-		Ret += LexicalConversion::ToString(*NameCount);
+		RetString += LexicalConversion::ToString(*NameCount);
 	}
 	++(*NameCount);
-	return Ret;
+	return RetString;
 }
 
 void FHlslNiagaraTranslator::EnterFunction(const FString& Name, FNiagaraFunctionSignature& Signature, TArray<int32>& Inputs)
@@ -2126,6 +2387,25 @@ void FHlslNiagaraTranslator::ExitStatsScope()
 	}
 }
 
+void FHlslNiagaraTranslator::EnterStatsScope(FNiagaraStatScope StatScope, FString& OutHlsl)
+{
+	if (GbNiagaraScriptStatTracking)
+	{
+		int32 ScopeIdx = Script->StatScopes.AddUnique(StatScope);
+		OutHlsl += FString::Printf(TEXT("EnterStatScope(%d /**%s*/);\n"), ScopeIdx, *StatScope.FullName.ToString());
+		StatScopeStack.Push(ScopeIdx);
+	}
+}
+
+void FHlslNiagaraTranslator::ExitStatsScope(FString& OutHlsl)
+{
+	if (GbNiagaraScriptStatTracking)
+	{
+		int32 ScopeIdx = StatScopeStack.Pop();
+		OutHlsl += FString::Printf(TEXT("ExitStatScope(/**%s*/);\n"), *Script->StatScopes[ScopeIdx].FullName.ToString());
+	}
+}
+
 FString FHlslNiagaraTranslator::GetCallstack()
 {
 	FString Callstack = Script->GetName();
@@ -2192,8 +2472,11 @@ int32 FHlslNiagaraTranslator::AddSourceChunk(FString SymbolName, const FNiagaraT
 	return Ret;
 }
 
+
 int32 FHlslNiagaraTranslator::AddBodyChunk(FString SymbolName, FString Definition, const FNiagaraTypeDefinition& Type, TArray<int32>& SourceChunks, bool bDecl, bool bIsTerminated)
 {
+	check(CurrentBodyChunkMode == ENiagaraCodeChunkMode::Body || CurrentBodyChunkMode == ENiagaraCodeChunkMode::SpawnBody || CurrentBodyChunkMode == ENiagaraCodeChunkMode::UpdateBody);
+
 	int32 Ret = CodeChunks.AddDefaulted();
 	FNiagaraCodeChunk& Chunk = CodeChunks[Ret];
 	Chunk.SymbolName = GetSanitizedSymbolName(SymbolName);
@@ -2201,15 +2484,19 @@ int32 FHlslNiagaraTranslator::AddBodyChunk(FString SymbolName, FString Definitio
 	Chunk.Type = Type;
 	Chunk.bDecl = bDecl;
 	Chunk.bIsTerminated = bIsTerminated;
-	Chunk.Mode = ENiagaraCodeChunkMode::Body;
+	Chunk.Mode = CurrentBodyChunkMode;
 	Chunk.SourceChunks = SourceChunks;
 
-	ChunksByMode[(int32)ENiagaraCodeChunkMode::Body].Add(Ret);
+	ChunksByMode[(int32)CurrentBodyChunkMode].Add(Ret);
 	return Ret;
 }
 
+
+
 int32 FHlslNiagaraTranslator::AddBodyChunk(FString SymbolName, FString Definition, const FNiagaraTypeDefinition& Type, int32 SourceChunk, bool bDecl, bool bIsTerminated)
 {
+	check(CurrentBodyChunkMode == ENiagaraCodeChunkMode::Body || CurrentBodyChunkMode == ENiagaraCodeChunkMode::SpawnBody || CurrentBodyChunkMode == ENiagaraCodeChunkMode::UpdateBody);
+
 	int32 Ret = CodeChunks.AddDefaulted();
 	FNiagaraCodeChunk& Chunk = CodeChunks[Ret];
 	Chunk.SymbolName = GetSanitizedSymbolName(SymbolName);
@@ -2217,15 +2504,17 @@ int32 FHlslNiagaraTranslator::AddBodyChunk(FString SymbolName, FString Definitio
 	Chunk.Type = Type;
 	Chunk.bDecl = bDecl;
 	Chunk.bIsTerminated = bIsTerminated;
-	Chunk.Mode = ENiagaraCodeChunkMode::Body;
+	Chunk.Mode = CurrentBodyChunkMode;
 	Chunk.SourceChunks.Add(SourceChunk);
 
-	ChunksByMode[(int32)ENiagaraCodeChunkMode::Body].Add(Ret);
+	ChunksByMode[(int32)CurrentBodyChunkMode].Add(Ret);
 	return Ret;
 }
 
 int32 FHlslNiagaraTranslator::AddBodyChunk(FString SymbolName, FString Definition, const FNiagaraTypeDefinition& Type, bool bDecl, bool bIsTerminated)
 {
+	check(CurrentBodyChunkMode == ENiagaraCodeChunkMode::Body || CurrentBodyChunkMode == ENiagaraCodeChunkMode::SpawnBody || CurrentBodyChunkMode == ENiagaraCodeChunkMode::UpdateBody);
+
 	int32 Ret = CodeChunks.AddDefaulted();
 	FNiagaraCodeChunk& Chunk = CodeChunks[Ret];
 	Chunk.SymbolName = GetSanitizedSymbolName(SymbolName);
@@ -2233,14 +2522,34 @@ int32 FHlslNiagaraTranslator::AddBodyChunk(FString SymbolName, FString Definitio
 	Chunk.Type = Type;
 	Chunk.bDecl = bDecl;
 	Chunk.bIsTerminated = bIsTerminated;
-	Chunk.Mode = ENiagaraCodeChunkMode::Body;
+	Chunk.Mode = CurrentBodyChunkMode;
 
-	ChunksByMode[(int32)ENiagaraCodeChunkMode::Body].Add(Ret);
+	ChunksByMode[(int32)CurrentBodyChunkMode].Add(Ret);
 	return Ret;
 }
 
 bool FHlslNiagaraTranslator::ShouldInterpolateParameter(const FNiagaraVariable& Parameter)
 {
+	//TODO: Some data driven method of deciding what parameters to interpolate and how to do it.
+	//Possibly allow definition of a dynamic input for the interpolation?
+	//With defaults for various types. Matrix=none, quat=slerp, everything else = Lerp.
+
+	//We don't want to interpolate matrices. Possibly consider moving to an FTransform like representation rather than matrices which could be interpolated?
+	if (Parameter.GetType() == FNiagaraTypeDefinition::GetMatrix4Def())
+	{
+		return false;
+	}
+
+	if (!Parameter.GetType().IsFloatPrimitive())
+	{
+		return false;
+	}
+
+	if (FNiagaraParameterMapHistory::IsRapidIterationParameter(Parameter))
+	{
+		return false;
+	}
+
 	//Skip interpolation for some system constants.
 	if (Parameter == SYS_PARAM_ENGINE_DELTA_TIME ||
 		Parameter == SYS_PARAM_ENGINE_INV_DELTA_TIME ||
@@ -2257,6 +2566,8 @@ bool FHlslNiagaraTranslator::ShouldInterpolateParameter(const FNiagaraVariable&
 
 int32 FHlslNiagaraTranslator::GetParameter(const FNiagaraVariable& Parameter)
 {
+
+	SCOPE_CYCLE_COUNTER(STAT_NiagaraEditor_HlslTranslator_GetParameter);
 	if (!AddStructToDefinitionSet(Parameter.GetType()))
 	{
 		Error(FText::Format(LOCTEXT("GetParameterFail", "Cannot handle type {0}! Variable: {1}"), Parameter.GetType().GetNameText(), FText::FromName(Parameter.GetName())), nullptr, nullptr);
@@ -2401,14 +2712,14 @@ FString FHlslNiagaraTranslator::GenerateConstantString(const FNiagaraVariable& C
 		}
 		else if (Type == FNiagaraTypeDefinition::GetBoolDef())
 		{
-			if (Constant.GetValue<FNiagaraBool>()->IsValid() == false)
+			if (Constant.GetValue<FNiagaraBool>().IsValid() == false)
 			{
 				Error(FText::Format(LOCTEXT("StructContantsBoolInvalidError", "Boolean constant {0} is not set to explicit True or False. Defaulting to False."), FText::FromName(Constant.GetName())), nullptr, nullptr);
 				ConstantStr = TEXT("false");
 			}
 			else
 			{
-				bool bValue = Constant.GetValue<FNiagaraBool>()->GetValue();
+				bool bValue = Constant.GetValue<FNiagaraBool>().GetValue();
 				ConstantStr = bValue ? TEXT("true") : TEXT("false");
 			}
 		}
@@ -2425,6 +2736,8 @@ FString FHlslNiagaraTranslator::GenerateConstantString(const FNiagaraVariable& C
 
 void FHlslNiagaraTranslator::Output(UNiagaraNodeOutput* OutputNode, const TArray<int32>& ComputedInputs)
 {
+	SCOPE_CYCLE_COUNTER(STAT_NiagaraEditor_HlslTranslator_Output);
+
 	TArray<FNiagaraVariable> Attributes;
 	TArray<int32> Inputs;
 
@@ -2449,7 +2762,10 @@ void FHlslNiagaraTranslator::Output(UNiagaraNodeOutput* OutputNode, const TArray
 			if (Attributes[i].GetType() != FNiagaraTypeDefinition::GetParameterMapDef())
 			{
 				FString SymbolName = *(TEXT("Out_") + GetSanitizedSymbolName(Attributes[i].GetName().ToString()));
+				ENiagaraCodeChunkMode OldMode = CurrentBodyChunkMode;
+				CurrentBodyChunkMode = ENiagaraCodeChunkMode::Body;
 				AddBodyChunk(SymbolName, TEXT("{0}"), Attributes[i].GetType(), Inputs[i], false);
+				CurrentBodyChunkMode = OldMode;
 			}
 		}
 	}
@@ -2550,6 +2866,8 @@ FString FHlslNiagaraTranslator::GetDataSetAccessSymbol(FNiagaraDataSetID DataSet
 
 void FHlslNiagaraTranslator::ParameterMapSet(UNiagaraNodeParameterMapSet* SetNode, TArray<int32>& Inputs, TArray<int32>& Outputs)
 {
+	SCOPE_CYCLE_COUNTER(STAT_NiagaraEditor_HlslTranslator_MapSet);
+
 	Outputs.SetNum(1);
 
 	FString ParameterMapInstanceName = TEXT("Context.Map");
@@ -2557,6 +2875,18 @@ void FHlslNiagaraTranslator::ParameterMapSet(UNiagaraNodeParameterMapSet* SetNod
 
 	TArray<UEdGraphPin*> InputPins;
 	SetNode->GetInputPins(InputPins);
+
+	// There is only one output pin for a set node, the parameter map must 
+	// continue to route through it.
+	if (!SetNode->IsNodeEnabled())
+	{
+		if (InputPins.Num() >= 1)
+		{
+			Outputs[0] = Inputs[0];
+		}
+		return;
+	}
+
 	int32 ParamMapHistoryIdx = INDEX_NONE;
 	for (int32 i = 0; i < Inputs.Num(); i++)
 	{
@@ -2640,10 +2970,12 @@ bool FHlslNiagaraTranslator::ParameterMapRegisterExternalConstantNamespaceVariab
 	FString ParameterMapInstanceName = GetParameterMapInstanceName(InParamMapHistoryIdx);
 
 	Output = INDEX_NONE;
-	if (InVariable.IsNameValid())
+	if (InVariable.IsValid())
 	{
 		// We don't really want system delta time or inverse system delta time in a spawn script. It leads to trouble.
-		if ((Script->IsInterpolatedParticleSpawnScript() && this->bInsideInterpolatedSpawnScript) || (Script->IsParticleSpawnScript() && !Script->IsInterpolatedParticleSpawnScript()))
+		if ((Script->IsInterpolatedParticleSpawnScript() && this->bInsideInterpolatedSpawnScript) 
+			|| (Script->IsParticleSpawnScript() && !Script->IsInterpolatedParticleSpawnScript() && TranslationOptions.SimTarget != ENiagaraSimTarget::GPUComputeSim)
+			)
 		{
 			if (InVariable == SYS_PARAM_ENGINE_DELTA_TIME || InVariable == SYS_PARAM_ENGINE_INV_DELTA_TIME)
 			{
@@ -2653,36 +2985,78 @@ bool FHlslNiagaraTranslator::ParameterMapRegisterExternalConstantNamespaceVariab
 			}
 		}
 
-		if (!IsBulkSystemScript() && false == ParamMapDefinedSystemVarsToUniformChunks.Contains(InVariable.GetName()))
+		bool bMissingParameter = false;
+		UNiagaraParameterCollection* Collection = ParamMapHistories[InParamMapHistoryIdx].IsParameterCollectionParameter(InVariable, bMissingParameter);
+		if (Collection && bMissingParameter)
 		{
-			FString SymbolNameDefined = FlattenedName;
+			Error(FText::Format(LOCTEXT("MissingNPCParameterError", "Parameter {0} was not found in Parameter Collection {1}"), FText::FromName(InVariable.GetName()), FText::FromString(Collection->GetFullName())), InNode, nullptr);
+		}
 
-			if (!InVariable.IsDataAllocated() && !InDefaultPin)
-			{
-				FNiagaraEditorUtilities::ResetVariableToDefaultValue(InVariable);
-			}
-			else if (!InVariable.IsDataAllocated())
+		bool bIsDataInterface = InVariable.GetType().IsDataInterface();
+		bool bIsPerInstanceBulkSystemParam = IsBulkSystemScript() && !bIsDataInterface && (FNiagaraParameterMapHistory::IsUserParameter(InVariable) || FNiagaraParameterMapHistory::IsEngineParameter(InVariable));
+
+		if (!bIsPerInstanceBulkSystemParam)
+		{
+			int32 UniformIdx = 0;
+			int32 UniformChunk = 0;
+
+			if (false == ParamMapDefinedSystemVarsToUniformChunks.Contains(InVariable.GetName()))
 			{
-				FNiagaraVariable Var = Schema->PinToNiagaraVariable(InDefaultPin, true);
-				FNiagaraEditorUtilities::ResetVariableToDefaultValue(InVariable);
-				if (Var.IsDataAllocated() && Var.GetData() != nullptr)
+				FString SymbolNameDefined = FlattenedName;
+
+				if (InVariable.GetType().IsDataInterface())
 				{
-					InVariable.SetData(Var.GetData());
-				}
-			}
+					UNiagaraDataInterface* DataInterface = nullptr;
+					if (Collection)
+					{
+						DataInterface = Collection->GetDefaultInstance()->GetParameterStore().GetDataInterface(InVariable);
+					}
+					else
+					{
+						UObject* Obj = const_cast<UClass*>(InVariable.GetType().GetClass())->GetDefaultObject(true);
+						DataInterface = Cast<UNiagaraDataInterface>(DuplicateObject(Obj, GetTransientPackage()));
+					}
+					if (ensure(DataInterface))
+					{
+						Output = RegisterDataInterface(InVariable, DataInterface, true);
+						return true;
+					}
+				}
+				if (!InVariable.IsDataAllocated() && !InDefaultPin)
+				{
+					FNiagaraEditorUtilities::ResetVariableToDefaultValue(InVariable);
+				}
+				else if (!InVariable.IsDataAllocated())
+				{
+					FNiagaraVariable Var = Schema->PinToNiagaraVariable(InDefaultPin, true);
+					FNiagaraEditorUtilities::ResetVariableToDefaultValue(InVariable);
+					if (Var.IsDataAllocated() && Var.GetData() != nullptr)
+					{
+						InVariable.SetData(Var.GetData());
+					}
+				}
+
+				if (InVariable.GetAllocatedSizeInBytes() != InVariable.GetSizeInBytes())
+				{
+					Error(FText::Format(LOCTEXT("GetParameterUnsetParam", "Variable {0} hasn't had its default value set. Required Bytes: {1} vs Allocated Bytes: {2}"), FText::FromName(InVariable.GetName()), FText::AsNumber(InVariable.GetType().GetSize()), FText::AsNumber(InVariable.GetSizeInBytes())), nullptr, nullptr);
+				}
 
-			if (InVariable.GetAllocatedSizeInBytes() != InVariable.GetSizeInBytes())
+				CompilationOutput.Parameters.SetOrAdd(InVariable);
+				UniformIdx = ChunksByMode[(int32)ENiagaraCodeChunkMode::Uniform].Num();
+				UniformChunk = AddUniformChunk(SymbolNameDefined, InVariable.GetType());
+				ParamMapDefinedSystemVarsToUniformChunks.Add(InVariable.GetName(), UniformIdx);
+				ParamMapDefinedSystemToNamespaceVars.Add(InVariable.GetName(), InVariable);
+			}
+			else
 			{
-				Error(FText::Format(LOCTEXT("GetParameterUnsetParam", "Variable {0} hasn't had its default value set. Required Bytes: {1} vs Allocated Bytes: {2}"), FText::FromName(InVariable.GetName()), FText::AsNumber(InVariable.GetType().GetSize()), FText::AsNumber(InVariable.GetSizeInBytes())), nullptr, nullptr);
+				UniformIdx = ParamMapDefinedSystemVarsToUniformChunks.FindChecked(InVariable.GetName());
+				UniformChunk = ChunksByMode[(int32)ENiagaraCodeChunkMode::Uniform][UniformIdx];
 			}
-			
-			CompilationOutput.Parameters.SetOrAdd(InVariable);
-			int32 UniformChunk = AddUniformChunk(SymbolNameDefined, InVariable.GetType());
-			ParamMapDefinedSystemVarsToUniformChunks.Add(InVariable.GetName(), UniformChunk);
-			ParamMapDefinedSystemToNamespaceVars.Add(InVariable.GetName(), InVariable);
-			MainPreSimulateChunks.Add(FString::Printf(TEXT("%s.%s = %s;"), *ParameterMapInstanceName, *VarName, *GetCodeAsSource(UniformChunk)));
+
+			//Add this separately as the same uniform can appear in the pre sim chunks more than once in different param maps.
+			MainPreSimulateChunks.AddUnique(FString::Printf(TEXT("%s.%s = %s;"), *ParameterMapInstanceName, *GetSanitizedSymbolName(VarName), *GetCodeAsSource(UniformChunk)));
 		}
-		else if (IsBulkSystemScript() && !ExternalVariablesForBulkUsage.Contains(InVariable))
+		else if (bIsPerInstanceBulkSystemParam && !ExternalVariablesForBulkUsage.Contains(InVariable))
 		{
 			ExternalVariablesForBulkUsage.Add(InVariable);
 		}
@@ -2701,7 +3075,7 @@ bool FHlslNiagaraTranslator::ParameterMapRegisterExternalConstantNamespaceVariab
 bool FHlslNiagaraTranslator::ParameterMapRegisterUniformAttributeVariable(const FNiagaraVariable& InVariable, UNiagaraNode* InNode, int32 InParamMapHistoryIdx, int32& Output)
 {
 	FNiagaraVariable NewVar = FNiagaraParameterMapHistory::BasicAttributeToNamespacedAttribute(InVariable);
-	if (NewVar.IsNameValid())
+	if (NewVar.IsValid())
 	{
 		return ParameterMapRegisterNamespaceAttributeVariable(NewVar, InNode, InParamMapHistoryIdx, Output);
 	}
@@ -2717,7 +3091,7 @@ bool FHlslNiagaraTranslator::ParameterMapRegisterNamespaceAttributeVariable(cons
 
 	Output = INDEX_NONE;
 	FNiagaraVariable BasicVar = FNiagaraParameterMapHistory::ResolveAsBasicAttribute(InVariable);
-	if (BasicVar.IsNameValid())
+	if (BasicVar.IsValid())
 	{
 		if (false == ParamMapDefinedAttributesToUniformChunks.Contains(BasicVar.GetName()))
 		{
@@ -2768,6 +3142,8 @@ FString FHlslNiagaraTranslator::GetParameterMapInstanceName(int32 ParamMapHistor
 
 void FHlslNiagaraTranslator::Emitter(class UNiagaraNodeEmitter* EmitterNode, TArray<int32>& Inputs, TArray<int32>& Outputs)
 {
+	SCOPE_CYCLE_COUNTER(STAT_NiagaraEditor_Module_NiagaraHLSLTranslator_Emitter);
+
 	FNiagaraFunctionSignature Signature;
 	UNiagaraScriptSource* Source = EmitterNode->GetScriptSource();
 	if (Source == nullptr)
@@ -2783,24 +3159,79 @@ void FHlslNiagaraTranslator::Emitter(class UNiagaraNodeEmitter* EmitterNode, TAr
 	FString Name = EmitterNode->GetName();
 	FString FullName = EmitterNode->GetFullName();
 
-	ActiveHistoryForFunctionCalls.EnterEmitter(EmitterUniqueName, EmitterNode);
-
 	TArray<UEdGraphPin*> CallOutputs;
 	TArray<UEdGraphPin*> CallInputs;
 	EmitterNode->GetOutputPins(CallOutputs);
 	EmitterNode->GetInputPins(CallInputs);
 
+
+	if (Inputs.Num() == 0 || Schema->PinToNiagaraVariable(CallInputs[0]).GetType() != FNiagaraTypeDefinition::GetParameterMapDef())
+	{
+		Error(LOCTEXT("EmitterMissingParamMap", "Emitter call missing ParameterMap input pin!"), EmitterNode, nullptr);
+		return;
+	}
+
+
+	int32 ParamMapHistoryIdx = Inputs[0];
+	if (ParamMapHistoryIdx == INDEX_NONE)
+	{
+		Error(LOCTEXT("EmitterMissingParamMapIndex", "Emitter call missing valid ParameterMap index!"), EmitterNode, nullptr);
+		return;
+	}
+	ActiveHistoryForFunctionCalls.EnterEmitter(EmitterUniqueName, EmitterNode);
+
+	// Clear out the parameter map writes to emitter module parameters as they should not be shared across emitters.
+	if (ParamMapHistoryIdx != -1 && ParamMapHistoryIdx < ParamMapHistories.Num())
+	{
+		for (int32 i = 0; i < ParamMapHistories[ParamMapHistoryIdx].Variables.Num(); i++)
+		{
+			check(ParamMapHistories[ParamMapHistoryIdx].VariablesWithOriginalAliasesIntact.Num() > i);
+			FNiagaraVariable Var = ParamMapHistories[ParamMapHistoryIdx].VariablesWithOriginalAliasesIntact[i];
+			if (FNiagaraParameterMapHistory::IsAliasedModuleParameter(Var))
+			{
+				ParamMapSetVariablesToChunks[ParamMapHistoryIdx][i] = INDEX_NONE;
+			}
+		}
+	}
+
 	// We act like a function call here as the semantics are identical.
-	RegisterFunctionCall(ScriptUsage, Name, FullName, Source, Signature, Inputs, CallInputs, CallOutputs, Signature);
+	RegisterFunctionCall(ScriptUsage, Name, FullName, Source, Signature, false, FString(), Inputs, CallInputs, CallOutputs, Signature);
 	GenerateFunctionCall(Signature, Inputs, Outputs);
+
+	// Clear out the parameter map writes to emitter module parameters as they should not be shared across emitters.
+	if (ParamMapHistoryIdx != -1 && ParamMapHistoryIdx < ParamMapHistories.Num())
+	{
+		for (int32 i = 0; i < ParamMapHistories[ParamMapHistoryIdx].Variables.Num(); i++)
+		{
+			check(ParamMapHistories[ParamMapHistoryIdx].VariablesWithOriginalAliasesIntact.Num() > i);
+			FNiagaraVariable Var = ParamMapHistories[ParamMapHistoryIdx].VariablesWithOriginalAliasesIntact[i];
+			if (ActiveHistoryForFunctionCalls.IsInEncounteredFunctionNamespace(Var) || FNiagaraParameterMapHistory::IsAliasedModuleParameter(Var))
+			{
+				ParamMapSetVariablesToChunks[ParamMapHistoryIdx][i] = INDEX_NONE;
+			}
+		}
+	}
 	ActiveHistoryForFunctionCalls.ExitEmitter(EmitterUniqueName, EmitterNode);
 }
 
 void FHlslNiagaraTranslator::ParameterMapGet(UNiagaraNodeParameterMapGet* GetNode, TArray<int32>& Inputs, TArray<int32>& Outputs)
 {
+	SCOPE_CYCLE_COUNTER(STAT_NiagaraEditor_Module_NiagaraHLSLTranslator_MapGet);
+
 	TArray<UEdGraphPin*> OutputPins;
 	GetNode->GetOutputPins(OutputPins);
 
+	// Push out invalid values for all output pins if the node is disabled.
+	if (!GetNode->IsNodeEnabled())
+	{
+		Outputs.SetNum(OutputPins.Num());
+		for (int32 i = 0; i < OutputPins.Num(); i++)
+		{
+			Outputs[i] = INDEX_NONE;
+		}
+		return;
+	}
+
 	TArray<UEdGraphPin*> InputPins;
 	GetNode->GetInputPins(InputPins);
 	
@@ -2817,6 +3248,15 @@ void FHlslNiagaraTranslator::ParameterMapGet(UNiagaraNodeParameterMapGet* GetNod
 			return;
 		}
 	}
+	else if (ParamMapHistoryIdx >= ParamMapHistories.Num())
+	{
+		Error(FText::Format(LOCTEXT("InvalidParamMapIdx", "Invalid parameter map index for input {0} of {1}!"), ParamMapHistoryIdx, ParamMapHistories.Num()), GetNode, nullptr);
+		for (int32 i = 0; i < Outputs.Num(); i++)
+		{
+			Outputs[i] = INDEX_NONE;
+			return;
+		}
+	}
 
 	FString ParameterMapInstanceName = GetParameterMapInstanceName(ParamMapHistoryIdx);
 	
@@ -2842,6 +3282,8 @@ void FHlslNiagaraTranslator::ParameterMapGet(UNiagaraNodeParameterMapGet* GetNod
 			FString SymbolName = GetSanitizedSymbolName(VarName);
 
 			bool bIsPerInstanceAttribute = false;
+			bool bIsCandidateForRapidIteration = false;
+			UEdGraphPin* InputPin = GetNode->GetDefaultPin(OutputPins[i]);
 
 			if (FNiagaraParameterMapHistory::IsExternalConstantNamespace(Var, Script))
 			{
@@ -2850,8 +3292,16 @@ void FHlslNiagaraTranslator::ParameterMapGet(UNiagaraNodeParameterMapGet* GetNod
 					continue;
 				}
 			}
+			else if (FNiagaraParameterMapHistory::IsAliasedModuleParameter(Var) && ActiveHistoryForFunctionCalls.InTopLevelFunctionCall(Script->Usage))
+			{
+				if (InputPin != nullptr && InputPin->LinkedTo.Num() == 0 && Var.GetType() != FNiagaraTypeDefinition::GetBoolDef() && !Var.GetType().IsEnum() && !Var.GetType().IsDataInterface())
+				{
+					bIsCandidateForRapidIteration = true;
+				}
+			}
 
 			FNiagaraParameterMapHistory& History = ParamMapHistories[ParamMapHistoryIdx];
+			bool bWasEmitterAliased = FNiagaraParameterMapHistory::IsAliasedEmitterParameter(Var);
 			Var = ActiveHistoryForFunctionCalls.ResolveAliases(Var);
 			if (History.IsPrimaryDataSetOutput(Var, GetTargetUsage()))
 			{
@@ -2861,18 +3311,59 @@ void FHlslNiagaraTranslator::ParameterMapGet(UNiagaraNodeParameterMapGet* GetNod
 			int32 LastSetChunkIdx = INDEX_NONE;
 			if (ParamMapHistoryIdx < ParamMapHistories.Num())
 			{
+				// See if we've written this variable before, if so we can reuse the index
 				int32 VarIdx = ParamMapHistories[ParamMapHistoryIdx].FindVariableByName(Var.GetName());
 				if (VarIdx != INDEX_NONE && VarIdx < ParamMapSetVariablesToChunks[ParamMapHistoryIdx].Num())
 				{
 					LastSetChunkIdx = ParamMapSetVariablesToChunks[ParamMapHistoryIdx][VarIdx];
 				}
 
+				// Check to see if this is the first time we've encountered this node and it is a viable candidate for rapid iteration
+				if (LastSetChunkIdx == INDEX_NONE && bIsCandidateForRapidIteration && TranslationOptions.bParameterRapidIteration)
+				{
+					bool bVarChanged = false;
+					if (!bWasEmitterAliased && ActiveHistoryForFunctionCalls.GetEmitterAlias() != nullptr)
+					{
+						Var = FNiagaraParameterMapHistory::ConvertVariableToRapidIterationConstantName(Var, *(*ActiveHistoryForFunctionCalls.GetEmitterAlias()));
+						bVarChanged = true;
+					}
+					else if (Script->IsSystemSpawnScript() || Script->IsSystemUpdateScript())
+					{
+						Var = FNiagaraParameterMapHistory::ConvertVariableToRapidIterationConstantName(Var, nullptr);
+						bVarChanged = true;
+					}
+
+					// Now try to look up with the new name.. we may have already made this an external variable before..
+					if (bVarChanged)
+					{
+						VarIdx = ParamMapHistories[ParamMapHistoryIdx].FindVariableByName(Var.GetName());
+						if (VarIdx != INDEX_NONE && VarIdx < ParamMapSetVariablesToChunks[ParamMapHistoryIdx].Num())
+						{
+							LastSetChunkIdx = ParamMapSetVariablesToChunks[ParamMapHistoryIdx][VarIdx];
+						}
+					}
+
+					// If it isn't found yet.. go ahead and make it into a constant variable..
+					if (LastSetChunkIdx == INDEX_NONE && ParameterMapRegisterExternalConstantNamespaceVariable(Var, GetNode, ParamMapHistoryIdx, Outputs[i], InputPin))
+					{
+						LastSetChunkIdx = Outputs[i];
+						if (VarIdx != INDEX_NONE && VarIdx < ParamMapSetVariablesToChunks[ParamMapHistoryIdx].Num())
+						{
+							// Record that we wrote to it.
+							ParamMapSetVariablesToChunks[ParamMapHistoryIdx][VarIdx] = LastSetChunkIdx;
+							ParamMapDefinedAttributesToNamespaceVars.FindOrAdd(Var.GetName()) = Var;
+						}
+						continue;
+					}
+				}
+
 				// We have yet to write to this parameter, use the default value if specified and the parameter 
 				// isn't a per-particle value.
 				bool bIgnoreDefaultValue = false;
 				if (bIsPerInstanceAttribute)
 				{
-					if ((Script->IsInterpolatedParticleSpawnScript() && bInsideInterpolatedSpawnScript == false) || !Script->IsSpawnScript())
+					if ((Script->IsInterpolatedParticleSpawnScript() && bInsideInterpolatedSpawnScript == false) 
+						|| ((TranslationOptions.SimTarget == ENiagaraSimTarget::GPUComputeSim && bInsideInterpolatedSpawnScript == false) || !Script->IsSpawnScript()))
 					{
 						bIgnoreDefaultValue = true;
 					}
@@ -2907,11 +3398,33 @@ void FHlslNiagaraTranslator::ParameterMapGet(UNiagaraNodeParameterMapGet* GetNod
 
 				if (LastSetChunkIdx == INDEX_NONE && !bIgnoreDefaultValue)
 				{
-					UEdGraphPin* InputPin = GetNode->GetDefaultPin(OutputPins[i]);
-						
 					// Default was found, trace back its inputs.
 					if (InputPin != nullptr )
-					{
+					{						
+						// Check to see if there are any overrides passed in to the translator. This allows us to bake in rapid iteration variables for performance.
+						if (InputPin->LinkedTo.Num() == 0 && bIsCandidateForRapidIteration && !TranslationOptions.bParameterRapidIteration)
+						{
+							bool bVarChanged = false;
+							FNiagaraVariable RapidIterationConstantVar;
+							if (!bWasEmitterAliased && ActiveHistoryForFunctionCalls.GetEmitterAlias() != nullptr)
+							{
+								RapidIterationConstantVar = FNiagaraParameterMapHistory::ConvertVariableToRapidIterationConstantName(Var, *(*ActiveHistoryForFunctionCalls.GetEmitterAlias()));
+								bVarChanged = true;
+							}
+							else if (Script->IsSystemSpawnScript() || Script->IsSystemUpdateScript())
+							{
+								RapidIterationConstantVar = FNiagaraParameterMapHistory::ConvertVariableToRapidIterationConstantName(Var, nullptr);
+								bVarChanged = true;
+							}
+
+							int32 FoundIdx = TranslationOptions.OverrideModuleConstants.Find(RapidIterationConstantVar);
+							if (FoundIdx != INDEX_NONE)
+							{
+								Outputs[i] = GetConstant(TranslationOptions.OverrideModuleConstants[FoundIdx]);
+								continue;
+							}
+						}
+
 						LastSetChunkIdx = CompilePin(InputPin);
 
 						if (!Var.IsDataInterface())
@@ -2952,8 +3465,7 @@ void FHlslNiagaraTranslator::ParameterMapGet(UNiagaraNodeParameterMapGet* GetNod
 					UNiagaraDataInterface* DataInterface = Cast<UNiagaraDataInterface>(DuplicateObject(Obj, GetTransientPackage()));
 					if (DataInterface)
 					{
-						UE_LOG(LogNiagaraEditor, Log, TEXT("Creating stand-in due to missing data interface!"));
-						LastSetChunkIdx = RegisterDataInterface(Var, DataInterface);
+						LastSetChunkIdx = RegisterDataInterface(Var, DataInterface, true);
 					}
 				}
 				
@@ -2968,52 +3480,6 @@ void FHlslNiagaraTranslator::ParameterMapGet(UNiagaraNodeParameterMapGet* GetNod
 	}
 }
 
-void FHlslNiagaraTranslator::ParameterCollection(class UNiagaraNodeParameterCollection* CollectionNode, TArray<int32>& Outputs)
-{
-	UNiagaraParameterCollection* Collection = CollectionNode->GetCollection();
-	check(Collection);
-	
-	TArray<UEdGraphPin*> OutputPins;
-	CollectionNode->GetOutputPins(OutputPins);
-	
-	check(Outputs.Num() == 0);
-	for (int32 i = 0; i < OutputPins.Num(); i++)
-	{
-		if (CollectionNode->IsAddPin(OutputPins[i]))
-		{
-			Outputs.Add(INDEX_NONE);//Add dummy output for add pin.
-			continue;
-		}
-
-		FNiagaraVariable PinVar = Collection->CollectionParameterFromFriendlyParameter(Schema->PinToNiagaraVariable(OutputPins[i]));
-
-		int32 SrcVarIdx = Collection->IndexOfParameter(PinVar);
-		if (SrcVarIdx == INDEX_NONE)
-		{
-			Error(LOCTEXT("MissingCollectionParameter", "Parameter is not a member of this Parameter Collection."), CollectionNode, OutputPins[i]);
-			Outputs.Add(INDEX_NONE);
-		}
-		else
-		{
-			FNiagaraVariable& SrcVar = Collection->GetParameters()[SrcVarIdx];
-		
-			if (SrcVar.IsDataInterface())
-			{
-				Outputs.Add(RegisterDataInterface(SrcVar, Collection->GetDefaultInstance()->GetParameterStore().GetDataInterface(SrcVar)));
-			}
-			else
-			{
-				FString VarSymbol = GetSanitizedSymbolName(SrcVar.GetName().ToString());
-				//Add uniform access to hlsl.
-				//Need to interact with the parameter map?
-				Outputs.Add(AddUniformChunk(VarSymbol, SrcVar.GetType()));
-				CompilationOutput.Parameters.SetOrAdd(SrcVar);
-			}
-
-			Script->ParameterCollections.AddUnique(Collection);
-		}
-	}
-}
 
 void FHlslNiagaraTranslator::ReadDataSet(const FNiagaraDataSetID DataSet, const TArray<FNiagaraVariable>& Variables, ENiagaraDataSetAccessMode AccessMode, int32 InputChunk, TArray<int32>& Outputs)
 {
@@ -3088,7 +3554,7 @@ void FHlslNiagaraTranslator::WriteDataSet(const FNiagaraDataSetID DataSet, const
 
 }
 
-int32 FHlslNiagaraTranslator::RegisterDataInterface(FNiagaraVariable& Var, UNiagaraDataInterface* DataInterface)
+int32 FHlslNiagaraTranslator::RegisterDataInterface(FNiagaraVariable& Var, UNiagaraDataInterface* DataInterface, bool bPlaceholder)
 {
 	int32 FuncParam;
 	if (GetFunctionParameter(Var, FuncParam))
@@ -3101,20 +3567,41 @@ int32 FHlslNiagaraTranslator::RegisterDataInterface(FNiagaraVariable& Var, UNiag
 	}
 	
 	//If we get here then this is a new data interface.
-	int32 Idx = CompilationOutput.DataInterfaceInfo.AddDefaulted();
-	CompilationOutput.DataInterfaceInfo[Idx].DataInterface = DataInterface;
-	CompilationOutput.DataInterfaceInfo[Idx].Name = Var.GetName();
+	int32 Idx = CompilationOutput.DataInterfaceInfo.IndexOfByPredicate([&](const FNiagaraScriptDataInterfaceInfo& OtherInfo)
+	{
+		return OtherInfo.Name == Var.GetName();
+	});
+	if (Idx == INDEX_NONE)
+	{
+		Idx = CompilationOutput.DataInterfaceInfo.AddDefaulted();
+		CompilationOutput.DataInterfaceInfo[Idx].DataInterface = DataInterface;
+		CompilationOutput.DataInterfaceInfo[Idx].Name = Var.GetName();
+		CompilationOutput.DataInterfaceInfo[Idx].Type = Var.GetType();
 
-	//Interface requires per instance data so add a user pointer table entry.
-	if (DataInterface->PerInstanceDataSize() > 0)
+		//Interface requires per instance data so add a user pointer table entry.
+		if (DataInterface->PerInstanceDataSize() > 0)
+		{
+			CompilationOutput.DataInterfaceInfo[Idx].UserPtrIdx = CompilationOutput.NumUserPtrs++;
+		}
+	}
+	else
+	{
+		check(CompilationOutput.DataInterfaceInfo[Idx].Name == Var.GetName());
+		check(CompilationOutput.DataInterfaceInfo[Idx].Type == Var.GetType());
+	}
+
+	if (bPlaceholder)
 	{
-		CompilationOutput.DataInterfaceInfo[Idx].UserPtrIdx = CompilationOutput.NumUserPtrs++;
+		PlaceholderDataInterfaces.AddUnique(DataInterface);
 	}
+
 	return Idx;
 }
 
 void FHlslNiagaraTranslator::Operation(class UNiagaraNodeOp* Operation, TArray<int32>& Inputs, TArray<int32>& Outputs)
 {
+	SCOPE_CYCLE_COUNTER(STAT_NiagaraEditor_HlslTranslator_Operation);
+
 	// Use the pins to determine the output type here since they may have been changed due to numeric pin fix up.
 	const FNiagaraOpInfo* OpInfo = FNiagaraOpInfo::GetOpInfo(Operation->OpName);
 
@@ -3142,7 +3629,47 @@ void FHlslNiagaraTranslator::Operation(class UNiagaraNodeOp* Operation, TArray<i
 
 void FHlslNiagaraTranslator::FunctionCall(UNiagaraNodeFunctionCall* FunctionNode, TArray<int32>& Inputs, TArray<int32>& Outputs)
 {
-	FNiagaraFunctionSignature Signature;
+	SCOPE_CYCLE_COUNTER(STAT_NiagaraEditor_Module_NiagaraHLSLTranslator_FunctionCall);
+
+	TArray<UEdGraphPin*> CallOutputs;
+	TArray<UEdGraphPin*> CallInputs;
+	FunctionNode->GetOutputPins(CallOutputs);
+	FunctionNode->GetInputPins(CallInputs);
+	
+	// If the function call is disabled, we 
+	// need to route the input parameter map pin to the output parameter map pin.
+	// Any other outputs become invalid.
+	if (!FunctionNode->IsNodeEnabled())
+	{
+		int32 InputPinIdx = INDEX_NONE;
+		int32 OutputPinIdx = INDEX_NONE;
+
+		for (int32 i = 0; i < CallInputs.Num(); i++)
+		{
+			const UEdGraphPin* Pin = CallInputs[i];
+			if (Schema->PinToTypeDefinition(Pin) == FNiagaraTypeDefinition::GetParameterMapDef())
+			{
+				// Found the input pin
+				InputPinIdx = Inputs[i];
+				break;
+			}
+		}
+
+		Outputs.SetNum(CallOutputs.Num());
+		for (int32 i = 0; i < CallOutputs.Num(); i++)
+		{
+			Outputs[i] = INDEX_NONE;
+			const UEdGraphPin* Pin = CallOutputs[i];
+			if (Schema->PinToTypeDefinition(Pin) == FNiagaraTypeDefinition::GetParameterMapDef())
+			{
+				// Mapping the input parameter map pin to the output.
+				Outputs[i] = InputPinIdx;
+			}
+		}
+		return;
+	}
+
+	FNiagaraFunctionSignature OutputSignature;
 	if (FunctionNode->FunctionScript == nullptr && !FunctionNode->Signature.IsValid())
 	{
 		Error(LOCTEXT("FunctionCallNonexistant", "Function call missing FunctionScript and invalid signature"), FunctionNode, nullptr);
@@ -3151,16 +3678,34 @@ void FHlslNiagaraTranslator::FunctionCall(UNiagaraNodeFunctionCall* FunctionNode
 
 	// We need the generated string to generate the proper signature for now.
 	ActiveHistoryForFunctionCalls.EnterFunction(FunctionNode->GetFunctionName(), FunctionNode->FunctionScript, FunctionNode);
+	
+	// Remove input add pin if it exists
+	for (int32 i = 0; i < CallOutputs.Num(); i++)
+	{
+		if (FunctionNode->IsAddPin(CallOutputs[i]))
+		{
+			CallOutputs.RemoveAt(i);
+			break;
+		}
+	}
 
-	TArray<UEdGraphPin*> CallOutputs;
-	TArray<UEdGraphPin*> CallInputs;
-	FunctionNode->GetOutputPins(CallOutputs);
-	FunctionNode->GetInputPins(CallInputs);
+	// Remove output add pin if it exists
+	for (int32 i = 0; i < CallInputs.Num(); i++)
+	{
+		if (FunctionNode->IsAddPin(CallInputs[i]))
+		{
+			CallInputs.RemoveAt(i);
+			break;
+		}
+	}
 
 	ENiagaraScriptUsage ScriptUsage = ENiagaraScriptUsage::Function;
 	FString Name;
 	FString FullName; 
 	UNiagaraScriptSource* Source = nullptr;
+	bool bCustomHlsl = false;
+	FString CustomHlsl;
+	FNiagaraFunctionSignature Signature = FunctionNode->Signature;
 
 	if (FunctionNode->FunctionScript)
 	{
@@ -3169,18 +3714,83 @@ void FHlslNiagaraTranslator::FunctionCall(UNiagaraNodeFunctionCall* FunctionNode
 		FullName = FunctionNode->FunctionScript->GetFullName();
 		Source = CastChecked<UNiagaraScriptSource>(FunctionNode->FunctionScript->GetSource());
 	}
+	UNiagaraNodeCustomHlsl* CustomFunctionHlsl = Cast<UNiagaraNodeCustomHlsl>(FunctionNode);
+	if (CustomFunctionHlsl != nullptr)
+	{
+		SCOPE_CYCLE_COUNTER(STAT_NiagaraEditor_Module_NiagaraHLSLTranslator_CustomHLSL);
 
-	RegisterFunctionCall(ScriptUsage, Name, FullName,  Source, FunctionNode->Signature, Inputs, CallInputs, CallOutputs, Signature);
-	GenerateFunctionCall(Signature, Inputs, Outputs);
+		ScriptUsage = CustomFunctionHlsl->ScriptUsage;
+		Name = GetSanitizedSymbolName(CustomFunctionHlsl->Signature.Name.ToString() + FunctionNode->NodeGuid.ToString());
+		Signature.Name = *Name; // Force the name to be set to include the node guid for safety...
+		bCustomHlsl = true;
+		CustomHlsl = CustomFunctionHlsl->CustomHlsl;
+
+		TArray<FNiagaraVariable> SigInputs;
+		for (FNiagaraVariable Input : Signature.Inputs)
+		{
+			if (Input.GetType() == FNiagaraTypeDefinition::GetParameterMapDef())
+			{
+				FString ParameterMapInstanceName = GetParameterMapInstanceName(0);
+				FString ReplaceSrc = Input.GetName().ToString() + TEXT(".");
+				FString ReplaceDest = ParameterMapInstanceName + TEXT(".");
+				CustomHlsl = CustomHlsl.Replace(*ReplaceSrc, *ReplaceDest);
+			}
+			else
+			{
+				FString ReplaceSrc = Input.GetName().ToString();
+				FString ReplaceDest = TEXT("In_") + ReplaceSrc;
+				CustomHlsl = CustomHlsl.Replace(*ReplaceSrc, *ReplaceDest);
+				SigInputs.Add(Input);
+			}
+		}
+		Signature.Inputs = SigInputs;
+
+		TArray<FNiagaraVariable> SigOutputs;
+		for (FNiagaraVariable Output : Signature.Outputs)
+		{
+			if (Output.GetType() == FNiagaraTypeDefinition::GetParameterMapDef())
+			{
+				FString ParameterMapInstanceName = GetParameterMapInstanceName(0);
+				FString ReplaceSrc = Output.GetName().ToString() + TEXT(".");
+				FString ReplaceDest = ParameterMapInstanceName + TEXT(".");
+				CustomHlsl = CustomHlsl.Replace(*ReplaceSrc, *ReplaceDest);
+			}
+			else
+			{
+				FString ReplaceSrc = Output.GetName().ToString();
+				FString ReplaceDest = TEXT("Out_") + ReplaceSrc;
+				CustomHlsl = CustomHlsl.Replace(*ReplaceSrc, *ReplaceDest);
+				SigOutputs.Add(Output);
+			}
+			Signature.Outputs = SigOutputs;
+		}
+
+		CustomHlsl = CustomHlsl.Replace(TEXT("\n"), TEXT("\n\t"));
+		CustomHlsl = TEXT("\n") + CustomHlsl + TEXT("\n");
+	}
+	
+
+	RegisterFunctionCall(ScriptUsage, Name, FullName,  Source, Signature, bCustomHlsl, CustomHlsl, Inputs, CallInputs, CallOutputs, OutputSignature);
+	GenerateFunctionCall(OutputSignature, Inputs, Outputs);
+
+	if (bCustomHlsl)
+	{
+		// Re-add the add pins.
+		Inputs.Add(INDEX_NONE);
+		Outputs.Add(INDEX_NONE);
+	}
 	ActiveHistoryForFunctionCalls.ExitFunction(FunctionNode->GetFunctionName(), FunctionNode->FunctionScript, FunctionNode);
 }
 
-void FHlslNiagaraTranslator::RegisterFunctionCall(ENiagaraScriptUsage ScriptUsage, const FString& InName, const FString& InFullName, UNiagaraScriptSource* Source, FNiagaraFunctionSignature& InSignature, TArray<int32>& Inputs, const TArray<UEdGraphPin*>& CallInputs, const TArray<UEdGraphPin*>& CallOutputs,
+void FHlslNiagaraTranslator::RegisterFunctionCall(ENiagaraScriptUsage ScriptUsage, const FString& InName, const FString& InFullName, UNiagaraScriptSource* Source, FNiagaraFunctionSignature& InSignature, bool bIsCustomHlsl, const FString& InCustomHlsl, TArray<int32>& Inputs, const TArray<UEdGraphPin*>& CallInputs, const TArray<UEdGraphPin*>& CallOutputs,
 	FNiagaraFunctionSignature& OutSignature)
 {
+	SCOPE_CYCLE_COUNTER(STAT_NiagaraEditor_Module_NiagaraHLSLTranslator_RegisterFunctionCall);
+
 	//////////////////////////////////////////////////////////////////////////
 	if (Source)
 	{
+		SCOPE_CYCLE_COUNTER(STAT_NiagaraEditor_Module_NiagaraHLSLTranslator_RegisterFunctionCall_Source); 
 		UNiagaraGraph* SourceGraph = CastChecked<UNiagaraGraph>(Source->NodeGraph);
 
 		bool bHasNumericInputs = false;
@@ -3202,6 +3812,7 @@ void FHlslNiagaraTranslator::RegisterFunctionCall(ENiagaraScriptUsage ScriptUsag
 				}
 			}
 
+			SCOPE_CYCLE_COUNTER(STAT_NiagaraEditor_Module_NiagaraHLSLTranslator_FunctionCallCloneGraphNumeric);
 			//We have to clone and preprocess the function graph to ensure it's numeric types have been fixed up to real types.
 			UNiagaraGraph* PreprocessedGraph = CastChecked<UNiagaraGraph>(FEdGraphUtilities::CloneGraph(SourceGraph, Source, &MessageLog));
 			FEdGraphUtilities::MergeChildrenGraphsIn(PreprocessedGraph, PreprocessedGraph, /*bRequireSchemaMatch=*/ true);
@@ -3210,6 +3821,7 @@ void FHlslNiagaraTranslator::RegisterFunctionCall(ENiagaraScriptUsage ScriptUsag
 		}
 		else
 		{
+			SCOPE_CYCLE_COUNTER(STAT_NiagaraEditor_Module_NiagaraHLSLTranslator_FunctionCallCloneGraphNonNumeric);
 			//If we don't have numeric inputs then we can cache the preprocessed graphs.
 			UNiagaraGraph** CachedGraph = PreprocessedFunctions.Find(SourceGraph);
 			if (!CachedGraph)
@@ -3222,7 +3834,9 @@ void FHlslNiagaraTranslator::RegisterFunctionCall(ENiagaraScriptUsage ScriptUsag
 			SourceGraph = *CachedGraph;
 		}
 		
-		GenerateFunctionSignature(ScriptUsage, InName, InFullName, SourceGraph, Inputs, bHasNumericInputs, OutSignature);
+		bool bHasParameterMapParameters = SourceGraph->HasParameterMapParameters();
+
+		GenerateFunctionSignature(ScriptUsage, InName, InFullName, SourceGraph, Inputs, bHasNumericInputs, bHasParameterMapParameters, OutSignature);
 
 // 		//Sort the input and outputs to match the sorted parameters. They may be different.
 // 		TArray<FNiagaraVariable> OrderedInputs;
@@ -3245,9 +3859,11 @@ void FHlslNiagaraTranslator::RegisterFunctionCall(ENiagaraScriptUsage ScriptUsag
 		FString* FuncBody = Functions.Find(OutSignature);
 		if (!FuncBody)
 		{
+			SCOPE_CYCLE_COUNTER(STAT_NiagaraEditor_Module_NiagaraHLSLTranslator_FuncBody);
+
 			//We've not compiled this function yet so compile it now.
 			EnterFunction(InName, OutSignature, Inputs);
-			
+
 			UNiagaraNodeOutput* FuncOutput = SourceGraph->FindOutputNode(ScriptUsage);
 			check(FuncOutput);
 
@@ -3261,10 +3877,18 @@ void FHlslNiagaraTranslator::RegisterFunctionCall(ENiagaraScriptUsage ScriptUsag
 
 			FHlslNiagaraTranslator* ThisTranslator = this;
 			TArray<int32> FuncOutputChunks;
-			FuncOutput->Compile(ThisTranslator, FuncOutputChunks);
+
+			ENiagaraCodeChunkMode OldMode = CurrentBodyChunkMode;
+			CurrentBodyChunkMode = ENiagaraCodeChunkMode::Body;
+			{
+				SCOPE_CYCLE_COUNTER(STAT_NiagaraEditor_Module_NiagaraHLSLTranslator_RegisterFunctionCall_Compile);
+				FuncOutput->Compile(ThisTranslator, FuncOutputChunks);
+			}
+			CurrentBodyChunkMode = OldMode;
 
 			// Find all of the data set writes that were connected to this particular graph.
 			TArray<UNiagaraNode*> OutputTraversal;
+			OutputTraversal.Reserve(1024);
 			SourceGraph->BuildTraversal(OutputTraversal, FuncOutput);
 
 			TArray<UNiagaraNodeWriteDataSet *>WriteNodes;
@@ -3275,6 +3899,7 @@ void FHlslNiagaraTranslator::RegisterFunctionCall(ENiagaraScriptUsage ScriptUsag
 			for (UNiagaraNodeWriteDataSet *WriteNode : WriteNodes)
 			{
 				TArray<UNiagaraNode*> WriteTraversal;
+				WriteTraversal.Reserve(1024);
 				SourceGraph->BuildTraversal(WriteTraversal, WriteNode);
 
 				bool bReferencesSameGraph = false;
@@ -3285,114 +3910,148 @@ void FHlslNiagaraTranslator::RegisterFunctionCall(ENiagaraScriptUsage ScriptUsag
 						bReferencesSameGraph = true;
 					}
 				}
-				
+
 				if (bReferencesSameGraph)
 				{
+					SCOPE_CYCLE_COUNTER(STAT_NiagaraEditor_Module_NiagaraHLSLTranslator_RegisterFunctionCall_Compile);
 					WriteNode->Compile(ThisTranslator, FuncOutputChunks);
 				}
 			}
 
-			//Grab all the body chunks for this function.
-			FString FunctionDefStr;
-			for (int32 i = ChunkStartsByMode[(int32)ENiagaraCodeChunkMode::Body]; i < ChunksByMode[(int32)ENiagaraCodeChunkMode::Body].Num(); ++i)
 			{
-				FunctionDefStr += GetCode(ChunksByMode[(int32)ENiagaraCodeChunkMode::Body][i]);
-			}
+				SCOPE_CYCLE_COUNTER(STAT_NiagaraEditor_Module_NiagaraHLSLTranslator_RegisterFunctionCall_FunctionDefStr);
+				//Grab all the body chunks for this function.
+				FString FunctionDefStr;
+				FunctionDefStr.Reserve(256 * ChunksByMode[(int32)ENiagaraCodeChunkMode::Body].Num());
+				for (int32 i = ChunkStartsByMode[(int32)ENiagaraCodeChunkMode::Body]; i < ChunksByMode[(int32)ENiagaraCodeChunkMode::Body].Num(); ++i)
+				{
+					FunctionDefStr += GetCode(ChunksByMode[(int32)ENiagaraCodeChunkMode::Body][i]);
+				}
 
-			//Now remove all chunks for the function again.			
-			//This is super hacky. Should move chunks etc into a proper scoping system.
-			
-			TArray<FNiagaraCodeChunk> FuncUniforms;
-			for (int32 i = 0; i < (int32)ENiagaraCodeChunkMode::Num; ++i)
-			{
-				//Keep uniform chunks.
-				if (i == (int32)ENiagaraCodeChunkMode::Uniform)
+				//Now remove all chunks for the function again.			
+				//This is super hacky. Should move chunks etc into a proper scoping system.
+
+				TArray<FNiagaraCodeChunk> FuncUniforms;
+				FuncUniforms.Reserve(1024);
+				for (int32 i = 0; i < (int32)ENiagaraCodeChunkMode::Num; ++i)
 				{
-					for (int32 ChunkIdx = ChunkStartsByMode[i]; ChunkIdx < ChunksByMode[i].Num(); ++ChunkIdx)
+					//Keep uniform chunks.
+					if (i == (int32)ENiagaraCodeChunkMode::Uniform)
 					{
-						FuncUniforms.Add(CodeChunks[ChunksByMode[i][ChunkIdx]]);
+						for (int32 ChunkIdx = ChunkStartsByMode[i]; ChunkIdx < ChunksByMode[i].Num(); ++ChunkIdx)
+						{
+							FuncUniforms.Add(CodeChunks[ChunksByMode[i][ChunkIdx]]);
+						}
 					}
-				}					
 
-				ChunksByMode[i].RemoveAt(ChunkStartsByMode[i], ChunksByMode[i].Num() - ChunkStartsByMode[i]);
-			}
-			CodeChunks.RemoveAt(ChunkStart, CodeChunks.Num() - ChunkStart, false);
+					ChunksByMode[i].RemoveAt(ChunkStartsByMode[i], ChunksByMode[i].Num() - ChunkStartsByMode[i]);
+				}
+				CodeChunks.RemoveAt(ChunkStart, CodeChunks.Num() - ChunkStart, false);
 
-			//Re-add the uniforms. Really this is horrible. Rework soon.
-			for (FNiagaraCodeChunk& Chunk : FuncUniforms)
-			{
-				ChunksByMode[(int32)ENiagaraCodeChunkMode::Uniform].Add(CodeChunks.Add(Chunk));
-			}
+				//Re-add the uniforms. Really this is horrible. Rework soon.
+				for (FNiagaraCodeChunk& Chunk : FuncUniforms)
+				{
+					ChunksByMode[(int32)ENiagaraCodeChunkMode::Uniform].Add(CodeChunks.Add(Chunk));
+				}
 
-			// We don't support an empty function definition when calling a real function.
-			if (FunctionDefStr.IsEmpty())
-			{
-				FunctionDefStr += TEXT("\n");
-			}
+				// We don't support an empty function definition when calling a real function.
+				if (FunctionDefStr.IsEmpty())
+				{
+					FunctionDefStr += TEXT("\n");
+				}
 
-			Functions.Add(OutSignature, FunctionDefStr);
+				Functions.Add(OutSignature, FunctionDefStr);
+			}
 
 			ExitFunction();
 		}
 	}
 	else
 	{
+		SCOPE_CYCLE_COUNTER(STAT_NiagaraEditor_Module_NiagaraHLSLTranslator_RegisterFunctionCall_Signature);
+
 		check(InSignature.IsValid());
-		check(InSignature.bMemberFunction);
+		check(InSignature.bMemberFunction || bIsCustomHlsl);
 		check(Inputs.Num() > 0);
 
 		OutSignature = InSignature;		
 
 		//First input for these is the owner of the function.
-		int32 OwnerIdx = Inputs[0];
-		if (OwnerIdx < 0 || OwnerIdx >= CompilationOutput.DataInterfaceInfo.Num())
+		if (bIsCustomHlsl)
 		{
-			Error(LOCTEXT("FunctionCallDataInterfaceMissingRegistration", "Function call signature does not match to a registered DataInterface. Valid DataInterfaces should be wired into a DataInterface function call."), nullptr, nullptr);
-			return;
-		}
-		FNiagaraScriptDataInterfaceInfo& Info = CompilationOutput.DataInterfaceInfo[OwnerIdx];
-
-		// Double-check to make sure that the signature matches those specified by the data 
-		// interface. It could be that the existing node has been removed and the graph
-		// needs to be refactored. If that's the case, emit an error.
-		if (Info.DataInterface != nullptr && OutSignature.bMemberFunction)
-		{
-			TArray<FNiagaraFunctionSignature> DataInterfaceFunctions;
-			Info.DataInterface->GetFunctions(DataInterfaceFunctions);
-			bool bFoundMatch = false;
-			for (const FNiagaraFunctionSignature& Sig : DataInterfaceFunctions)
+			FString* FuncBody = Functions.Find(OutSignature);
+			if (!FuncBody)
 			{
-				if (Sig == OutSignature)
+				//We've not compiled this function yet so compile it now.
+				EnterFunction(InName, OutSignature, Inputs);
+
+				FString FunctionDefStr = InCustomHlsl;
+				// We don't support an empty function definition when calling a real function.
+				if (FunctionDefStr.IsEmpty())
 				{
-					bFoundMatch = true;
+					FunctionDefStr += TEXT("\n");
 				}
-			}
 
-			if (!bFoundMatch)
+				Functions.Add(OutSignature, FunctionDefStr);
+
+				ExitFunction();
+			}
+		}
+		else
+		{
+			int32 OwnerIdx = Inputs[0];
+			if (OwnerIdx < 0 || OwnerIdx >= CompilationOutput.DataInterfaceInfo.Num())
 			{
-				Error(LOCTEXT("FunctionCallDataInterfaceMissing", "Function call signature does not match DataInterface possible signatures?"), nullptr, nullptr);
+				Error(LOCTEXT("FunctionCallDataInterfaceMissingRegistration", "Function call signature does not match to a registered DataInterface. Valid DataInterfaces should be wired into a DataInterface function call."), nullptr, nullptr);
 				return;
 			}
+			FNiagaraScriptDataInterfaceInfo& Info = CompilationOutput.DataInterfaceInfo[OwnerIdx];
 
-			if (Info.UserPtrIdx != INDEX_NONE)
+			// Double-check to make sure that the signature matches those specified by the data 
+			// interface. It could be that the existing node has been removed and the graph
+			// needs to be refactored. If that's the case, emit an error.
+			if (Info.DataInterface != nullptr && OutSignature.bMemberFunction)
 			{
-				//This interface requires per instance data via a user ptr so place the index to it at the end of the inputs.
-				Inputs.Add(AddSourceChunk(LexicalConversion::ToString(Info.UserPtrIdx), FNiagaraTypeDefinition::GetIntDef()));
-				OutSignature.Inputs.Add(FNiagaraVariable(FNiagaraTypeDefinition::GetIntDef(), TEXT("InstanceData")));
+				TArray<FNiagaraFunctionSignature> DataInterfaceFunctions;
+				Info.DataInterface->GetFunctions(DataInterfaceFunctions);
+				bool bFoundMatch = false;
+				for (const FNiagaraFunctionSignature& Sig : DataInterfaceFunctions)
+				{
+					if (Sig == OutSignature)
+					{
+						bFoundMatch = true;
+					}
+				}
+
+				if (!bFoundMatch)
+				{
+					Error(LOCTEXT("FunctionCallDataInterfaceMissing", "Function call signature does not match DataInterface possible signatures?"), nullptr, nullptr);
+					return;
+				}
+
+				if (Info.UserPtrIdx != INDEX_NONE)
+				{
+					//This interface requires per instance data via a user ptr so place the index to it at the end of the inputs.
+					Inputs.Add(AddSourceChunk(LexicalConversion::ToString(Info.UserPtrIdx), FNiagaraTypeDefinition::GetIntDef()));
+					OutSignature.Inputs.Add(FNiagaraVariable(FNiagaraTypeDefinition::GetIntDef(), TEXT("InstanceData")));
+				}
 			}
-		}
 
-		//Override the owner id of the signature with the actual caller.
-		OutSignature.OwnerName = Info.Name;
-		Info.RegisteredFunctions.Add(OutSignature);
+			//Override the owner id of the signature with the actual caller.
+			OutSignature.OwnerName = Info.Name;
+			Info.RegisteredFunctions.Add(OutSignature);
+
+			Functions.FindOrAdd(OutSignature);
+		}
 
-		Functions.FindOrAdd(OutSignature);
 	}
 }
 
 void FHlslNiagaraTranslator::GenerateFunctionCall(FNiagaraFunctionSignature& FunctionSignature, TArray<int32>& Inputs, TArray<int32>& Outputs)
 {
-	EnterStatsScope(FNiagaraStatScope(*GetFunctionSignatureSymbol(FunctionSignature), FText::FromString(FunctionSignature.GetName())));
+	SCOPE_CYCLE_COUNTER(STAT_NiagaraEditor_Module_NiagaraHLSLTranslator_GenerateFunctionCall);
+
+	EnterStatsScope(FNiagaraStatScope(*GetFunctionSignatureSymbol(FunctionSignature), *(FunctionSignature.GetName())));
 
 	TArray<FString> MissingParameters;
 	int32 ParamIdx = 0;
@@ -3410,11 +4069,6 @@ void FHlslNiagaraTranslator::GenerateFunctionCall(FNiagaraFunctionSignature& Fun
 				Error(FText::Format(LOCTEXT("GetConstantFailTypeVar2", "Cannot handle type {0}! Variable: {1}"),  Type.GetNameText(), FText::FromName(FunctionSignature.Inputs[i].GetName())), nullptr, nullptr);
 			}
 
-			if (ParamIdx != 0)
-			{
-				DefStr += TEXT(", ");
-			}
-
 			int32 Input = Inputs[i];
 			bool bSkip = false;
 
@@ -3426,6 +4080,11 @@ void FHlslNiagaraTranslator::GenerateFunctionCall(FNiagaraFunctionSignature& Fun
 
 			if (!bSkip)
 			{
+				if (ParamIdx != 0)
+				{
+					DefStr += TEXT(", ");
+				}
+
 				Params.Add(Input);
 				if (Input == INDEX_NONE)
 				{
@@ -3453,11 +4112,6 @@ void FHlslNiagaraTranslator::GenerateFunctionCall(FNiagaraFunctionSignature& Fun
 				Error(FText::Format(LOCTEXT("GetConstantFailTypeVar3", "Cannot handle type {0}! Variable: {1}"), Type.GetNameText(), FText::FromName(FunctionSignature.Outputs[i].GetName())), nullptr, nullptr);
 			}
 			
-			if (ParamIdx > 0)
-			{
-				DefStr += TEXT(", ");
-			}
-
 			int32 Output = INDEX_NONE;
 			int32 ParamOutput = INDEX_NONE;
 			bool bSkip = false;
@@ -3480,6 +4134,11 @@ void FHlslNiagaraTranslator::GenerateFunctionCall(FNiagaraFunctionSignature& Fun
 
 			if (!bSkip)
 			{
+				if (ParamIdx > 0)
+				{
+					DefStr += TEXT(", ");
+				}
+
 				Params.Add(ParamOutput);
 				if (ParamOutput == INDEX_NONE)
 				{
@@ -3558,7 +4217,7 @@ FString FHlslNiagaraTranslator::GetFunctionSignature(const FNiagaraFunctionSigna
 					SigStr += TEXT(", ");
 				}
 
-				SigStr += FHlslNiagaraTranslator::GetStructHlslTypeName(Input.GetType()) + TEXT(" In_") + FHlslNiagaraTranslator::GetSanitizedSymbolName(Input.GetName().ToString());
+				SigStr += FHlslNiagaraTranslator::GetStructHlslTypeName(Input.GetType()) + TEXT(" In_") + FHlslNiagaraTranslator::GetSanitizedSymbolName(Input.GetName().ToString(), true);
 				++ParamIdx;
 			}
 		}
@@ -3798,6 +4457,8 @@ void FHlslNiagaraTranslator::Convert(class UNiagaraNodeConvert* Convert, TArray
 
 void FHlslNiagaraTranslator::If(TArray<FNiagaraVariable>& Vars, int32 Condition, TArray<int32>& PathA, TArray<int32>& PathB, TArray<int32>& Outputs)
 {
+	SCOPE_CYCLE_COUNTER(STAT_NiagaraEditor_HlslTranslator_If);
+
 	int32 NumVars = Vars.Num();
 	check(PathA.Num() == NumVars);
 	check(PathB.Num() == NumVars);
@@ -3831,6 +4492,8 @@ void FHlslNiagaraTranslator::If(TArray<FNiagaraVariable>& Vars, int32 Condition,
 
 int32 FHlslNiagaraTranslator::CompilePin(UEdGraphPin* Pin)
 {
+	SCOPE_CYCLE_COUNTER(STAT_NiagaraEditor_HlslTranslator_CompilePin);
+
 	check(Pin);
 	int32 Ret = INDEX_NONE;
 	FNiagaraTypeDefinition TypeDef = Schema->PinToTypeDefinition(Pin);
@@ -3871,6 +4534,8 @@ int32 FHlslNiagaraTranslator::CompilePin(UEdGraphPin* Pin)
 
 int32 FHlslNiagaraTranslator::CompileOutputPin(UEdGraphPin* Pin)
 {
+	SCOPE_CYCLE_COUNTER(STAT_NiagaraEditor_HlslTranslator_CompileOutputPin);
+
 	check(Pin->Direction == EGPD_Output);
 
 	int32 Ret = INDEX_NONE;
@@ -3963,9 +4628,9 @@ ENiagaraScriptUsage FHlslNiagaraTranslator::GetTargetUsage() const
 	return Script->GetUsage();
 }
 
-int32 FHlslNiagaraTranslator::GetTargetOccurrence() const
+FGuid FHlslNiagaraTranslator::GetTargetUsageId() const
 {
-	return Script->GetUsageIndex();
+	return Script->GetUsageId();
 }
 //////////////////////////////////////////////////////////////////////////
 
@@ -3997,7 +4662,7 @@ FString FHlslNiagaraTranslator::GetHlslDefaultForType(FNiagaraTypeDefinition Typ
 	}
 	else if (Type == FNiagaraTypeDefinition::GetBoolDef())
 	{
-		return "(true)";
+		return "(false)";
 	}
 	else
 	{
