Copying Orion-Staging @ 3748653 (Orion/Dev-General @ 3722124) to //UE4/Main

diff --git a/Engine/Plugins/FX/Niagara/Source/NiagaraEditor/Private/NiagaraHlslTranslator.h b/Engine/Plugins/FX/Niagara/Source/NiagaraEditor/Private/NiagaraHlslTranslator.h
--- a/Engine/Plugins/FX/Niagara/Source/NiagaraEditor/Private/NiagaraHlslTranslator.h
+++ b/Engine/Plugins/FX/Niagara/Source/NiagaraEditor/Private/NiagaraHlslTranslator.h
@@ -126,6 +126,8 @@ enum class ENiagaraCodeChunkMode : uint8
 	Uniform,
 	Source,
 	Body,
+	SpawnBody,
+	UpdateBody,
 	Num,
 };
 
@@ -200,6 +202,7 @@ class NIAGARAEDITOR_API FHlslNiagaraTranslatorOptions
 public:
 	FHlslNiagaraTranslatorOptions()
 		: SimTarget(ENiagaraSimTarget::CPUSim)
+		, bParameterRapidIteration(true)
 	{
 
 	}
@@ -208,6 +211,12 @@ public:
 
 	/** Any parameters in these namespaces will be pulled from an "InstanceParameters" dataset rather than from the uniform table. */
 	TArray<FString> InstanceParameterNamespaces;
+
+	/** Whether or not to treat top-level module variables as external values for rapid iteration without need for compilation.*/
+	bool bParameterRapidIteration;
+
+	/** Whether or not to override top-level module variables with values from the constant override table. This is only used for variables that were candidates for rapid iteration.*/
+	TArray<FNiagaraVariable> OverrideModuleConstants;
 };
 
 class NIAGARAEDITOR_API FHlslNiagaraTranslator
@@ -265,13 +274,14 @@ protected:
 	FORCEINLINE FNiagaraDataSetID GetInstanceDataSetID()const { return FNiagaraDataSetID(TEXT("DataInstance"), ENiagaraDataSetType::ParticleData); }
 	FORCEINLINE FNiagaraDataSetID GetSystemEngineDataSetID()const { return FNiagaraDataSetID(TEXT("Engine"), ENiagaraDataSetType::ParticleData); }
 	FORCEINLINE FNiagaraDataSetID GetSystemUserDataSetID()const { return FNiagaraDataSetID(TEXT("User"), ENiagaraDataSetType::ParticleData); }
+	FORCEINLINE FNiagaraDataSetID GetSystemConstantDataSetID()const { return FNiagaraDataSetID(TEXT("Constant"), ENiagaraDataSetType::ParticleData); }
 
 	/** All functions called in the script. */
 	TMap<FNiagaraFunctionSignature, FString> Functions;
 	/** Map of function graphs we've seen before and already pre-processed. */
 	TMap<const UNiagaraGraph*, UNiagaraGraph*> PreprocessedFunctions;
 
-	void RegisterFunctionCall(ENiagaraScriptUsage ScriptUsage, const FString& InName, const FString& InFullName, UNiagaraScriptSource* Source, FNiagaraFunctionSignature& InSignature, TArray<int32>& Inputs, const TArray<UEdGraphPin*>& CallInputs, const TArray<UEdGraphPin*>& CallOutputs,
+	void RegisterFunctionCall(ENiagaraScriptUsage ScriptUsage, const FString& InName, const FString& InFullName, UNiagaraScriptSource* Source, FNiagaraFunctionSignature& InSignature, bool bIsCustomHlsl, const FString& InCustomHlsl, TArray<int32>& Inputs, const TArray<UEdGraphPin*>& CallInputs, const TArray<UEdGraphPin*>& CallOutputs,
 		FNiagaraFunctionSignature& OutSignature);
 	void GenerateFunctionCall(FNiagaraFunctionSignature& FunctionSignature, TArray<int32>& Inputs, TArray<int32>& Outputs);
 	FString GetFunctionSignature(const FNiagaraFunctionSignature& Sig);
@@ -305,6 +315,8 @@ protected:
 
 	void EnterStatsScope(FNiagaraStatScope StatScope);
 	void ExitStatsScope();
+	void EnterStatsScope(FNiagaraStatScope StatScope, FString& OutHlsl);
+	void ExitStatsScope(FString& OutHlsl);
 
 	FString GeneratedConstantString(float Constant);
 	FString GeneratedConstantString(FVector4 Constant);
@@ -320,7 +332,6 @@ protected:
 	int32 AddBodyChunk(FString SymbolName, FString Definition, const FNiagaraTypeDefinition& Type, int32 SourceChunk, bool bDecl = true, bool bIsTerminated = true);
 	int32 AddBodyChunk(FString SymbolName, FString Definition, const FNiagaraTypeDefinition& Type, bool bDecl = true, bool bIsTerminated = true);
 
-
 	FString GetFunctionDefinitions();
 public:
 
@@ -332,7 +343,7 @@ public:
 
 	virtual int32 CompilePin(UEdGraphPin* Pin);
 
-	virtual int32 RegisterDataInterface(FNiagaraVariable& Var, UNiagaraDataInterface* DataInterface);
+	virtual int32 RegisterDataInterface(FNiagaraVariable& Var, UNiagaraDataInterface* DataInterface, bool bPlaceholder);
 
 	virtual void Operation(class UNiagaraNodeOp* Operation, TArray<int32>& Inputs, TArray<int32>& Outputs);
 	virtual void Output(UNiagaraNodeOutput* OutputNode, const TArray<int32>& ComputedInputs);
@@ -347,7 +358,6 @@ public:
 	virtual void WriteDataSet(const FNiagaraDataSetID DataSet, const TArray<FNiagaraVariable>& Variable, ENiagaraDataSetAccessMode AccessMode, const TArray<int32>& Inputs);
 	virtual void ParameterMapSet(class UNiagaraNodeParameterMapSet* SetNode, TArray<int32>& Inputs, TArray<int32>& Outputs);
 	virtual void ParameterMapGet(class UNiagaraNodeParameterMapGet* GetNode, TArray<int32>& Inputs, TArray<int32>& Outputs);
-	virtual void ParameterCollection(class UNiagaraNodeParameterCollection* CollectionNode, TArray<int32>& Outputs);
 	virtual void Emitter(class UNiagaraNodeEmitter* GetNode, TArray<int32>& Inputs, TArray<int32>& Outputs);
 
 	void DefineInterpolatedParametersFunction(FString &HlslOutput);
@@ -357,7 +367,7 @@ public:
 	void DefineDataSetVariableReads(FString &HLSLOutput, FNiagaraDataSetID& Id, int32 DataSetIndex, TArray<FNiagaraVariable> & ReadVars);
 	void DefineDataSetVariableWrites(FString &HlslOutput, FNiagaraDataSetID& Id, int32 DataSetIndex, TArray<FNiagaraVariable>& WriteVars);
 	void DefineDataInterfaceHLSL(FString &HlslOutput);
-	TArray< TArray<DIGPUBufferParamDescriptor> >&GetDatainterfaceBufferDescriptors() { return DIBufferDescriptors; }
+	TArray< FDIBufferDescriptorStore >&GetDatainterfaceBufferDescriptors() { return DIBufferDescriptors; }
 
 	// Format string should have up to 5 entries, {{0} = Computed Variable Suffix, {1} = Float or Int, {2} = Data Set Index, {3} = Register Index, {4} Default value for that type.
 	void GatherVariableForDataSetAccess(const FNiagaraVariable& Variable, FString Format, int32& RegisterIdx, int32 DataSetIndex, FString InstanceIdxSymbol, FString &HlslOutput);
@@ -377,7 +387,7 @@ public:
 
 	virtual bool CanReadAttributes()const;
 	virtual ENiagaraScriptUsage GetTargetUsage() const;
-	int32 GetTargetOccurrence() const;
+	FGuid GetTargetUsageId() const;
 
 	static bool IsBuiltInHlslType(FNiagaraTypeDefinition Type);
 	static FString GetStructHlslTypeName(FNiagaraTypeDefinition Type);
@@ -388,7 +398,7 @@ public:
 	static TArray<FName> ConditionPropertyPath(const FNiagaraTypeDefinition& Type, const TArray<FName>& InPath);
 
 
-	static FString GetSanitizedSymbolName(FString SymbolName);
+	static FString GetSanitizedSymbolName(FString SymbolName, bool bCollapseNamespaces=false);
 
 	bool AddStructToDefinitionSet(const FNiagaraTypeDefinition& TypeDef);
 
@@ -440,7 +450,8 @@ private:
 	bool ParameterMapRegisterUniformAttributeVariable(const FNiagaraVariable& InVariable, UNiagaraNode* InNode, int32 InParamMapHistoryIdx, int32& Output);
 
 	bool ValidateTypePins(UNiagaraNode* NodeToValidate);
-	void GenerateFunctionSignature(ENiagaraScriptUsage ScriptUsage, FString InName, const FString& InFullName, UNiagaraGraph* FuncGraph, TArray<int32>& Inputs, bool bHadNumericInputs, FNiagaraFunctionSignature& OutSig)const;
+	void GenerateFunctionSignature(ENiagaraScriptUsage ScriptUsage, FString InName, const FString& InFullName, UNiagaraGraph* FuncGraph, TArray<int32>& Inputs, 
+		bool bHadNumericInputs, bool bHasParameterMapParameters, FNiagaraFunctionSignature& OutSig)const;
 
 	UNiagaraGraph* CloneGraphAndPrepareForCompilation(UNiagaraScript* InScript, UNiagaraScriptSource* InSource, bool bClearErrors);
 
@@ -475,7 +486,7 @@ private:
 	TMap<FName, int32> ParamMapDefinedEmitterParameterVarsToUniformChunks; // Map from the variable name exposed by the emitter as a parameter to the uniform chunk expressing it (i.e. have we encountered before in this graph?)
 	TMap<FName, FNiagaraVariable> ParamMapDefinedEmitterParameterToNamespaceVars; // Map from defined parameter to the Namespaced variable expressing it.
 
-    // Synced to the Attributes encountered for parameter maps thus far.
+	// Synced to the Attributes encountered for parameter maps thus far.
 	TMap<FName, int32> ParamMapDefinedAttributesToUniformChunks; // Map from the variable name exposed as a attribute to the uniform chunk expressing it (i.e. have we encountered before in this graph?)
 	TMap<FName, FNiagaraVariable> ParamMapDefinedAttributesToNamespaceVars; // Map from defined parameter to the Namespaced variable expressing it.
 
@@ -490,7 +501,7 @@ private:
 	int32 WriteIdx;
 
 	// Data interface buffer descriptors per data interface
-	TArray< TArray<DIGPUBufferParamDescriptor> > DIBufferDescriptors;
+	TArray< FDIBufferDescriptorStore > DIBufferDescriptors;
 
 	/**
 	True if we're compiling the spawn section of an interpolated spawn script.
@@ -508,4 +519,8 @@ private:
 	// Used to keep track of which output node we are working back from. This allows us 
 	// to find the right parameter map.
 	TArray<int32> CurrentParamMapIndices;
+
+	ENiagaraCodeChunkMode CurrentBodyChunkMode;
+
+	TArray<UNiagaraDataInterface*> PlaceholderDataInterfaces;
 };
diff --git a/Engine/Plugins/FX/Niagara/Source/NiagaraEditor/Private/NiagaraMaterialParameterNode.cpp b/Engine/Plugins/FX/Niagara/Source/NiagaraEditor/Private/NiagaraMaterialParameterNode.cpp
new file mode 100644
--- /dev/null
+++ b/Engine/Plugins/FX/Niagara/Source/NiagaraEditor/Private/NiagaraMaterialParameterNode.cpp
@@ -0,0 +1,306 @@
+// Copyright 1998-2017 Epic Games, Inc. All Rights Reserved.
+#include "NiagaraMaterialParameterNode.h"
+#include "UObject/UnrealType.h"
+#include "NiagaraGraph.h"
+#include "NiagaraEditorCommon.h"
+#include "NiagaraScriptSource.h"
+#include "NiagaraScript.h"
+#include "NiagaraNodeInput.h"
+#include "NiagaraNodeOutput.h"
+#include "NiagaraComponent.h"
+#include "NiagaraEditorUtilities.h"
+#include "NiagaraNodeParameterMapGet.h"
+#include "NiagaraNodeParameterMapSet.h"
+#include "Materials/Material.h"
+#include "Materials/MaterialExpressionDynamicParameter.h"
+#include "NiagaraNodeConvert.h"
+#include "NiagaraEmitter.h"
+#include "NiagaraRendererProperties.h"
+#include "NiagaraNodeIf.h"
+#include "ModuleManager.h"
+#include "NiagaraEditorUtilities.h"
+#include "NiagaraEditorModule.h"
+#include "INiagaraEditorTypeUtilities.h"
+#include "NiagaraStackEditorData.h"
+#include "NiagaraStackGraphUtilities.h"
+
+#define LOCTEXT_NAMESPACE "NiagaraMaterialParameterNode"
+
+void UNiagaraMaterialParameterNode::AllocateDefaultPins()
+{
+	GenerateScript();
+	Super::AllocateDefaultPins();
+}
+
+FText UNiagaraMaterialParameterNode::GetNodeTitle(ENodeTitleType::Type TitleType) const
+{
+	return NSLOCTEXT("NiagaraMaterialParameterNode", "NiagaraMaterialParameterNode", "Material Parameters");
+}
+
+bool UNiagaraMaterialParameterNode::RefreshFromExternalChanges()
+{
+	FunctionScript = nullptr;
+	GenerateScript();
+	ReallocatePins();
+	return true;
+}
+
+void UNiagaraMaterialParameterNode::SetEmitter(UNiagaraEmitter* InEmitter)
+{
+	Emitter = InEmitter;
+}
+
+void UNiagaraMaterialParameterNode::GenerateScript()
+{
+	FString NodeName = FString("MaterialParameter");
+
+	if (FunctionScript == nullptr)
+	{
+		FunctionScript = NewObject<UNiagaraScript>(this, FName(*NodeName), RF_Transactional);
+		FunctionScript->SetUsage(ENiagaraScriptUsage::Module);
+		InitializeScript(FunctionScript);
+		ComputeNodeName();
+	}
+}
+
+void UNiagaraMaterialParameterNode::InitializeScript(UNiagaraScript* NewScript)
+{
+	check(NewScript != nullptr);
+
+	// Find the dynamic parameters expression from the material.
+	// @YannickLange todo: Notify user that the material did not have any dynamic parameters. Maybe add them from code?
+	UMaterialExpressionDynamicParameter* DynParamExp = GetMaterialExpressionDynamicParameter(Emitter);
+
+	UNiagaraScriptSource* Source = Cast<UNiagaraScriptSource>(NewScript->GetSource());
+
+	if (Source == nullptr)
+	{
+		Source = NewObject<UNiagaraScriptSource>(NewScript, NAME_None, RF_Transactional);
+		NewScript->SetSource(Source);
+	}
+
+	UNiagaraGraph* CreatedGraph = Source->NodeGraph;
+	if (CreatedGraph == nullptr)
+	{
+		CreatedGraph = NewObject<UNiagaraGraph>(Source, NAME_None, RF_Transactional);
+		Source->NodeGraph = CreatedGraph;
+	}
+
+	// Create input map
+	TArray<UNiagaraNodeInput*> InputNodes;
+	CreatedGraph->FindInputNodes(InputNodes);
+	if (InputNodes.Num() == 0)
+	{
+		FGraphNodeCreator<UNiagaraNodeInput> InputNodeCreator(*CreatedGraph);
+		UNiagaraNodeInput* InputNode = InputNodeCreator.CreateNode();
+		InputNode->Input = FNiagaraVariable(FNiagaraTypeDefinition::GetParameterMapDef(), TEXT("InputMap"));
+		InputNode->Usage = ENiagaraInputNodeUsage::Parameter;
+		InputNodeCreator.Finalize();
+		InputNodes.Add(InputNode);
+	}
+
+	// Create output map
+	UNiagaraNodeOutput* OutputNode = CreatedGraph->FindOutputNode(ENiagaraScriptUsage::Module);
+	if (OutputNode == nullptr)
+	{
+		FGraphNodeCreator<UNiagaraNodeOutput> OutputNodeCreator(*CreatedGraph);
+		OutputNode = OutputNodeCreator.CreateNode();
+		FNiagaraVariable ParamMapAttrib(FNiagaraTypeDefinition::GetParameterMapDef(), TEXT("OutputMap"));
+		OutputNode->SetUsage(ENiagaraScriptUsage::Module);
+		OutputNode->Outputs.Add(ParamMapAttrib);
+		OutputNodeCreator.Finalize();
+	}
+
+	TArray<UNiagaraNodeParameterMapGet*> GetNodes;
+	CreatedGraph->GetNodesOfClass(GetNodes);
+
+	TArray<UNiagaraNodeParameterMapSet*> SetNodes;
+	CreatedGraph->GetNodesOfClass(SetNodes);
+
+	if (SetNodes.Num() == 0)
+	{
+		FGraphNodeCreator<UNiagaraNodeParameterMapSet> InputNodeCreator(*CreatedGraph);
+		UNiagaraNodeParameterMapSet* InputNode = InputNodeCreator.CreateNode();
+		InputNodeCreator.Finalize();
+		SetNodes.Add(InputNode);
+
+		InputNodes[0]->GetOutputPin(0)->MakeLinkTo(SetNodes[0]->GetInputPin(0));
+		SetNodes[0]->GetOutputPin(0)->MakeLinkTo(OutputNode->GetInputPin(0));
+	}
+
+	// We create two get nodes. The first is for the direct values.
+	// The second is in the case of referencing other parameters that we want to use as defaults.
+	if (GetNodes.Num() == 0)
+	{
+		FGraphNodeCreator<UNiagaraNodeParameterMapGet> InputNodeCreator(*CreatedGraph);
+		UNiagaraNodeParameterMapGet* InputNode = InputNodeCreator.CreateNode();
+		InputNodeCreator.Finalize();
+		GetNodes.Add(InputNode);
+
+		InputNodes[0]->GetOutputPin(0)->MakeLinkTo(GetNodes[0]->GetInputPin(0));
+	}
+	if (GetNodes.Num() == 1)
+	{
+		FGraphNodeCreator<UNiagaraNodeParameterMapGet> InputNodeCreator(*CreatedGraph);
+		UNiagaraNodeParameterMapGet* InputNode = InputNodeCreator.CreateNode();
+		InputNodeCreator.Finalize();
+		GetNodes.Add(InputNode);
+
+		InputNodes[0]->GetOutputPin(0)->MakeLinkTo(GetNodes[1]->GetInputPin(0));
+	}
+
+	// Clean out existing pins
+	while (!SetNodes[0]->IsAddPin(SetNodes[0]->GetInputPin(1)))
+	{
+		SetNodes[0]->RemovePin(SetNodes[0]->GetInputPin(1));
+	}
+
+	while (!GetNodes[0]->IsAddPin(GetNodes[0]->GetOutputPin(0)))
+	{
+		GetNodes[0]->RemovePin(GetNodes[0]->GetInputPin(0));
+	}
+
+	while (!GetNodes[1]->IsAddPin(GetNodes[1]->GetOutputPin(0)))
+	{
+		GetNodes[1]->RemovePin(GetNodes[1]->GetInputPin(0));
+	}
+
+	FNiagaraEditorModule& NiagaraEditorModule = FModuleManager::GetModuleChecked<FNiagaraEditorModule>("NiagaraEditor");
+	TSharedPtr<INiagaraEditorTypeUtilities> TypeEditorUtilitiesFloat = NiagaraEditorModule.GetTypeUtilities(FNiagaraTypeDefinition::GetFloatDef());
+	TSharedPtr<INiagaraEditorTypeUtilities> TypeEditorUtilitiesBool = NiagaraEditorModule.GetTypeUtilities(FNiagaraTypeDefinition::GetBoolDef());
+
+	// Add a vector4 pin (referencing the dynamic material parameter) to the set node.
+	UEdGraphPin* SetPinDynParam = SetNodes[0]->RequestNewTypedPin(EGPD_Input, FNiagaraTypeDefinition::GetVec4Def(), TEXT("Particles.DynamicMaterialParameter"));
+	UEdGraphPin* GetPinDynParam = GetNodes[0]->RequestNewTypedPin(EGPD_Output, FNiagaraTypeDefinition::GetVec4Def(), TEXT("Particles.DynamicMaterialParameter"));
+
+	// Create the convert node that converts a vector4 to 4 floats.
+	UNiagaraNodeConvert* BreakNode = NewObject<UNiagaraNodeConvert>(CreatedGraph);
+	BreakNode->InitAsBreak(FNiagaraTypeDefinition::GetVec4Def());
+	BreakNode->CreateNewGuid();
+	BreakNode->PostPlacedNewNode();
+	BreakNode->AllocateDefaultPins();
+
+	// Wire the result of the Break to the set pin.
+	BreakNode->AutowireNewNode(GetPinDynParam);
+
+	// Create the convert node that converts 4 floats to a vector4.
+	UNiagaraNodeConvert* MakeNode = NewObject<UNiagaraNodeConvert>(CreatedGraph);
+	MakeNode->InitAsMake(FNiagaraTypeDefinition::GetVec4Def());
+	MakeNode->CreateNewGuid();
+	MakeNode->PostPlacedNewNode();
+	MakeNode->AllocateDefaultPins();
+
+	// Wire the result of the convert to the set pin.
+	MakeNode->AutowireNewNode(SetPinDynParam);
+
+	TArray<UEdGraphPin*> InputMakeVectorPins;
+	MakeNode->GetInputPins(InputMakeVectorPins);
+
+	TArray<UEdGraphPin*> OutputBreakVectorPins;
+	BreakNode->GetOutputPins(OutputBreakVectorPins);
+
+	if (DynParamExp != nullptr)
+	{
+		// Create input pins for the module and connect those to the input of the convert node.
+		TArray<FExpressionOutput>& Expressions =  DynParamExp->GetOutputs();
+		for (int32 Index = 0; Index < Expressions.Num(); Index++)
+		{
+			FExpressionOutput& Expression = Expressions[Index];
+			UEdGraphPin* GetPin = GetNodes[0]->RequestNewTypedPin(EGPD_Output, FNiagaraTypeDefinition::GetFloatDef(), FName(*(TEXT("Module.") + Expression.OutputName.ToString())));
+			UEdGraphPin* BoolPin = GetNodes[0]->RequestNewTypedPin(EGPD_Output, FNiagaraTypeDefinition::GetBoolDef(), FName(*(TEXT("Module.") + Expression.OutputName.ToString() + TEXT(" Enabled"))));
+
+			// Assume that we want the default for the bool pin to be checked
+			{
+				UEdGraphPin* DefaultValueCheckedPin = GetNodes[0]->GetDefaultPin(BoolPin);
+				FNiagaraVariable DefaultValueBoolVar(FNiagaraTypeDefinition::GetBoolDef(), TEXT("TempBool"));
+				DefaultValueBoolVar.SetValue<FNiagaraBool>(FNiagaraBool(true));
+				DefaultValueCheckedPin->DefaultValue = TypeEditorUtilitiesBool->GetPinDefaultStringFromValue(DefaultValueBoolVar);
+			}
+
+			// Set the default value on the pin.
+			{
+				// Gather the default from the material for the value...
+				UEdGraphPin* DefaultValuePin = GetNodes[0]->GetDefaultPin(GetPin);
+				const float DefaultValue = DynParamExp->DefaultValue.Component(Index);
+				FNiagaraVariable DefaultValueVar(FNiagaraTypeDefinition::GetFloatDef(), TEXT("Temp"));
+				DefaultValueVar.SetValue<float>(DefaultValue);
+				DefaultValuePin->DefaultValue = TypeEditorUtilitiesFloat->GetPinDefaultStringFromValue(DefaultValueVar);
+			}
+
+			// Create if statement node.
+			UNiagaraNodeIf* IfNode = nullptr;
+			{
+				FGraphNodeCreator<UNiagaraNodeIf> BranchNodeCreator(*CreatedGraph);
+				IfNode = BranchNodeCreator.CreateNode();
+				IfNode->AddOutput(FNiagaraTypeDefinition::GetFloatDef(), *Expression.OutputName.ToString());
+				BranchNodeCreator.Finalize();
+			}
+
+			{
+				// Get the pins from the branch node.
+				UEdGraphPin* ConditionPin = IfNode->GetPinAt(0);
+				UEdGraphPin* ConditionAPin = IfNode->GetPinAt(1);
+				UEdGraphPin* ConditionBPin = IfNode->GetPinAt(2);
+				check (ConditionPin != nullptr && ConditionAPin != nullptr && ConditionBPin != nullptr);
+
+				// Connect the pins to the branch node.
+				BoolPin->MakeLinkTo(ConditionPin);
+				IfNode->GetPinAt(3)->MakeLinkTo(InputMakeVectorPins[Index]);
+				ConditionAPin->MakeLinkTo(GetPin);
+				ConditionBPin->MakeLinkTo(OutputBreakVectorPins[Index]);
+			}
+		}
+	}
+}
+
+UMaterial* UNiagaraMaterialParameterNode::GetMaterialFromEmitter(UNiagaraEmitter* InEmitter)
+{
+	UMaterial* ResultMaterial = nullptr;
+	if (InEmitter->GetRenderers().Num() > 0)
+	{
+		for (UNiagaraRendererProperties* RenderProperties : InEmitter->GetRenderers())
+		{
+			TArray<UMaterialInterface*> UsedMaterialInteraces;
+			RenderProperties->GetUsedMaterials(UsedMaterialInteraces);
+			for (UMaterialInterface* UsedMaterialInterface : UsedMaterialInteraces)
+			{
+				UMaterial* UsedMaterial = Cast<UMaterial>(UsedMaterialInterface);
+				if (UsedMaterial != nullptr)
+				{
+					ResultMaterial = UsedMaterial;
+					break;
+				}
+			}
+		}
+	}
+
+	// @YannickLange todo: Notify the user if we couldn't find a material because there are no render properties.
+	// @YannickLange todo: Notify the user if we couldn't find a material because there are no materials in any render properties.
+	return ResultMaterial;
+}
+
+
+UMaterialExpressionDynamicParameter* UNiagaraMaterialParameterNode::GetMaterialExpressionDynamicParameter(class UNiagaraEmitter* InEmitter)
+{
+	UMaterial* Material = GetMaterialFromEmitter(Emitter);
+
+	// Find the dynamic parameters expression from the material.
+	// @YannickLange todo: Notify user that the material did not have any dynamic parameters. Maybe add them from code?
+	UMaterialExpressionDynamicParameter* DynParamExp = nullptr;
+	if (Material != nullptr)
+	{
+		for (UMaterialExpression* Expression : Material->Expressions)
+		{
+			DynParamExp = Cast<UMaterialExpressionDynamicParameter>(Expression);
+
+			if (DynParamExp != nullptr)
+			{
+				break;
+			}
+		}
+	}
+
+	return DynParamExp;
+}
+
+#undef LOCTEXT_NAMESPACE
\ No newline at end of file
diff --git a/Engine/Plugins/FX/Niagara/Source/NiagaraEditor/Private/NiagaraNode.cpp b/Engine/Plugins/FX/Niagara/Source/NiagaraEditor/Private/NiagaraNode.cpp
--- a/Engine/Plugins/FX/Niagara/Source/NiagaraEditor/Private/NiagaraNode.cpp
+++ b/Engine/Plugins/FX/Niagara/Source/NiagaraEditor/Private/NiagaraNode.cpp
@@ -72,7 +72,6 @@ bool UNiagaraNode::ReallocatePins()
 			if (bTypeMismatch && !bRetainOldTypeDueToNumerics)
 			{
 				DefaultValue = (*MatchingNewPin)->DefaultValue;
-				AutogeneratedDefaultValue = (*MatchingNewPin)->AutogeneratedDefaultValue;
 				DefaultObject = (*MatchingNewPin)->DefaultObject;
 				DefaultTextValue = (*MatchingNewPin)->DefaultTextValue;
 			}
@@ -87,7 +86,6 @@ bool UNiagaraNode::ReallocatePins()
 			if (bTypeMismatch && !bRetainOldTypeDueToNumerics)
 			{
 				(*MatchingNewPin)->DefaultValue = DefaultValue;
-				(*MatchingNewPin)->AutogeneratedDefaultValue = AutogeneratedDefaultValue;
 				(*MatchingNewPin)->DefaultObject = DefaultObject;
 				(*MatchingNewPin)->DefaultTextValue = DefaultTextValue;
 				bAllSame = false;
@@ -128,6 +126,12 @@ bool UNiagaraNode::CompileInputPins(FHlslNiagaraTranslator *Translator, TArray<i
 	for (int32 i = 0; i < InputPins.Num(); ++i)
 	{
 		UEdGraphPin* Pin = InputPins[i];
+		if (!IsValidPinToCompile(Pin))
+		{
+			OutCompiledInputs.Add(INDEX_NONE);
+			continue;
+		}
+
 		check(Pin->Direction == EGPD_Input);
 		int32 Result = CompileInputPin(Translator, Pin);
 		if (Result == INDEX_NONE)
@@ -342,4 +346,41 @@ ENiagaraNumericOutputTypeSelectionMode UNiagaraNode::GetNumericOutputTypeSelecti
 	return ENiagaraNumericOutputTypeSelectionMode::None;
 }
 
+void UNiagaraNode::RouteParameterMapAroundMe(FNiagaraParameterMapHistoryBuilder& OutHistory, bool bRecursive)
+{
+	const UEdGraphSchema_Niagara* Schema = CastChecked<UEdGraphSchema_Niagara>(GetSchema());
+
+	TArray<UEdGraphPin*> InputPins;
+	GetInputPins(InputPins);
+	TArray<UEdGraphPin*> OutputPins;
+	GetOutputPins(OutputPins);
+
+	const UEdGraphPin* InputPin = nullptr;
+	const UEdGraphPin* OutputPin = nullptr;
+
+	for (const UEdGraphPin* Pin : InputPins)
+	{
+		if (Schema->PinToTypeDefinition(Pin) == FNiagaraTypeDefinition::GetParameterMapDef())
+		{
+			InputPin = Pin;
+			break;
+		}
+	}
+
+	for (const UEdGraphPin* Pin : OutputPins)
+	{
+		if (Schema->PinToTypeDefinition(Pin) == FNiagaraTypeDefinition::GetParameterMapDef())
+		{
+			OutputPin = Pin;
+			break;
+		}
+	}
+
+	if (InputPin != nullptr && OutputPin != nullptr && InputPin->LinkedTo.Num() != 0)
+	{
+		int32 PMIdx = OutHistory.TraceParameterMapOutputPin(InputPin->LinkedTo[0]);
+		OutHistory.RegisterParameterMapPin(PMIdx, OutputPin);
+	}
+}
+
 #undef LOCTEXT_NAMESPACE
\ No newline at end of file
diff --git a/Engine/Plugins/FX/Niagara/Source/NiagaraEditor/Private/NiagaraNodeEmitter.cpp b/Engine/Plugins/FX/Niagara/Source/NiagaraEditor/Private/NiagaraNodeEmitter.cpp
--- a/Engine/Plugins/FX/Niagara/Source/NiagaraEditor/Private/NiagaraNodeEmitter.cpp
+++ b/Engine/Plugins/FX/Niagara/Source/NiagaraEditor/Private/NiagaraNodeEmitter.cpp
@@ -17,9 +17,14 @@
 #include "NiagaraNodeOutput.h"
 #include "NiagaraHlslTranslator.h"
 #include "MultiBoxBuilder.h"
+#include "Stats.h"
+#include "NiagaraEditorModule.h"
 
 #define LOCTEXT_NAMESPACE "NiagaraNodeEmitter"
 
+DECLARE_CYCLE_STAT(TEXT("Niagara - Module - NiagaraNodeEmitter_Compile"), STAT_NiagaraEditor_Module_NiagaraNodeEmitter_Compile, STATGROUP_NiagaraEditor);
+
+
 void UNiagaraNodeEmitter::PostInitProperties()
 {
 	Super::PostInitProperties();
@@ -101,7 +106,6 @@ void UNiagaraNodeEmitter::AllocateDefaultPins()
 		const UEdGraphSchema_Niagara* NiagaraSchema = Cast<UEdGraphSchema_Niagara>(GetSchema());
 		CreatePin(EGPD_Input, NiagaraSchema->TypeDefinitionToPinType(FNiagaraTypeDefinition::GetParameterMapDef()), TEXT("InputMap"));
 		CreatePin(EGPD_Output, NiagaraSchema->TypeDefinitionToPinType(FNiagaraTypeDefinition::GetParameterMapDef()), TEXT("OutputMap"));
-		CachedEmitterChangeId = Emitter->ChangeId;
 	}
 }
 
@@ -212,27 +216,8 @@ UNiagaraGraph* UNiagaraNodeEmitter::GetCalledGraph() const
 
 bool UNiagaraNodeEmitter::RefreshFromExternalChanges()
 {
-	// First get the emitter that we're referencing..
-	UNiagaraEmitter* Emitter = nullptr;
-	if (OwnerSystem)
-	{
-		for (int32 i = 0; i < OwnerSystem->GetNumEmitters(); ++i)
-		{
-			if (OwnerSystem->GetEmitterHandle(i).GetId() == EmitterHandleId)
-			{
-				Emitter = OwnerSystem->GetEmitterHandle(i).GetInstance();
-			}
-		}
-	}
-
-	if (Emitter && Emitter->ChangeId != CachedEmitterChangeId)
-	{
-		// TODO - Leverage code in reallocate pins to determine if any pins have changed...
-		ReallocatePins();
-		DisplayName = GetNameFromEmitter();
-		return true;
-	}
-	return false;
+	DisplayName = GetNameFromEmitter();
+	return true;
 }
 
 FText UNiagaraNodeEmitter::GetNameFromEmitter()
@@ -254,6 +239,12 @@ void UNiagaraNodeEmitter::BuildParameterMapHistory(FNiagaraParameterMapHistoryBu
 {
 	Super::BuildParameterMapHistory(OutHistory, bRecursive);
 
+	if (!IsNodeEnabled() && OutHistory.GetIgnoreDisabled())
+	{
+		RouteParameterMapAroundMe(OutHistory, bRecursive);
+		return;
+	}
+
 	const UEdGraphSchema_Niagara* Schema = GetDefault<UEdGraphSchema_Niagara>();
 	TArray<UEdGraphPin*> InputPins;
 	GetInputPins(InputPins);
@@ -314,6 +305,7 @@ void UNiagaraNodeEmitter::BuildParameterMapHistory(FNiagaraParameterMapHistoryBu
 					if (ExistingIdx == INDEX_NONE)
 					{
 						ExistingIdx = OutHistory.Histories[ParamMapIdx].Variables.Add(Var);
+						OutHistory.Histories[ParamMapIdx].VariablesWithOriginalAliasesIntact.Add(History.VariablesWithOriginalAliasesIntact[SrcVarIdx]);
 						OutHistory.Histories[ParamMapIdx].PerVariableReadHistory.AddDefaulted(1);
 						OutHistory.Histories[ParamMapIdx].PerVariableWriteHistory.AddDefaulted(1);
 						OutHistory.Histories[ParamMapIdx].PerVariableWarnings.AddDefaulted(1);
@@ -325,6 +317,9 @@ void UNiagaraNodeEmitter::BuildParameterMapHistory(FNiagaraParameterMapHistoryBu
 					OutHistory.Histories[ParamMapIdx].PerVariableWriteHistory[ExistingIdx].Append(History.PerVariableWriteHistory[SrcVarIdx]);
 					OutHistory.Histories[ParamMapIdx].PerVariableWarnings[ExistingIdx].Append(History.PerVariableWarnings[SrcVarIdx]);	
 				}
+				OutHistory.Histories[ParamMapIdx].ParameterCollections.Append(History.ParameterCollections);
+				OutHistory.Histories[ParamMapIdx].ParameterCollectionNamespaces.Append(History.ParameterCollectionNamespaces);
+				OutHistory.Histories[ParamMapIdx].ParameterCollectionVariables.Append(History.ParameterCollectionVariables);
 			}
 		}
 
@@ -343,6 +338,7 @@ void UNiagaraNodeEmitter::BuildParameterMapHistory(FNiagaraParameterMapHistoryBu
 
 void UNiagaraNodeEmitter::Compile(FHlslNiagaraTranslator *Translator, TArray<int32>& Outputs)
 {
+	SCOPE_CYCLE_COUNTER(STAT_NiagaraEditor_Module_NiagaraNodeEmitter_Compile);
 	TArray<UEdGraphPin*> InputPins;
 	GetInputPins(InputPins);
 	InputPins.RemoveAll([](UEdGraphPin* InputPin) { return (InputPin->PinType.PinCategory != UEdGraphSchema_Niagara::PinCategoryType) && (InputPin->PinType.PinCategory != UEdGraphSchema_Niagara::PinCategoryEnum); });
diff --git a/Engine/Plugins/FX/Niagara/Source/NiagaraEditor/Private/NiagaraNodeEmitter.h b/Engine/Plugins/FX/Niagara/Source/NiagaraEditor/Private/NiagaraNodeEmitter.h
--- a/Engine/Plugins/FX/Niagara/Source/NiagaraEditor/Private/NiagaraNodeEmitter.h
+++ b/Engine/Plugins/FX/Niagara/Source/NiagaraEditor/Private/NiagaraNodeEmitter.h
@@ -81,9 +81,6 @@ private:
 	UPROPERTY()
 	FText DisplayName;
 
-	UPROPERTY()
-	FGuid CachedEmitterChangeId;
-
 	UPROPERTY()
 	ENiagaraScriptUsage ScriptType;
 };
\ No newline at end of file
diff --git a/Engine/Plugins/FX/Niagara/Source/NiagaraEditor/Private/NiagaraNodeFunctionCall.cpp b/Engine/Plugins/FX/Niagara/Source/NiagaraEditor/Private/NiagaraNodeFunctionCall.cpp
--- a/Engine/Plugins/FX/Niagara/Source/NiagaraEditor/Private/NiagaraNodeFunctionCall.cpp
+++ b/Engine/Plugins/FX/Niagara/Source/NiagaraEditor/Private/NiagaraNodeFunctionCall.cpp
@@ -15,6 +15,7 @@
 #include "NiagaraEditorUtilities.h"
 #include "NiagaraNodeParameterMapGet.h"
 #include "NiagaraConstants.h"
+#include "NiagaraCustomVersion.h"
 
 #define LOCTEXT_NAMESPACE "NiagaraNodeFunctionCall"
 
@@ -34,6 +35,37 @@ void UNiagaraNodeFunctionCall::PostLoad()
 			Source->ConditionalPostLoad();
 			UNiagaraGraph* Graph = Source->NodeGraph;
 			Graph->ConditionalPostLoad();
+
+			// Fix up autogenerated default values if neccessary.
+			const int32 NiagaraCustomVersion = GetLinkerCustomVersion(FNiagaraCustomVersion::GUID);
+			if (NiagaraCustomVersion < FNiagaraCustomVersion::EnabledAutogeneratedDefaultValuesForFunctionCallNodes)
+			{
+				TArray<UEdGraphPin*> InputPins;
+				GetInputPins(InputPins);
+
+				TArray<UNiagaraNodeInput*> InputNodes;
+				UNiagaraGraph::FFindInputNodeOptions Options;
+				Options.bSort = true;
+				Options.bFilterDuplicates = true;
+				Graph->FindInputNodes(InputNodes, Options);
+
+				for (UEdGraphPin* InputPin : InputPins)
+				{
+					auto FindInputNodeByName = [&](UNiagaraNodeInput* InputNode) { return InputNode->Input.GetName().ToString() == InputPin->PinName.ToString(); };
+					UNiagaraNodeInput** MatchingInputNodePtr = InputNodes.FindByPredicate(FindInputNodeByName);
+					if (MatchingInputNodePtr != nullptr)
+					{
+						UNiagaraNodeInput* MatchingInputNode = *MatchingInputNodePtr;
+						SetPinAutoGeneratedDefaultValue(*InputPin, *MatchingInputNode);
+
+						// If the default value wasn't set, update it with the new autogenerated default.
+						if (InputPin->DefaultValue.IsEmpty())
+						{
+							InputPin->DefaultValue = InputPin->AutogeneratedDefaultValue;
+						}
+					}
+				}
+			}
 		}
 	}
 
@@ -96,6 +128,9 @@ void UNiagaraNodeFunctionCall::AllocateDefaultPins()
 				//Special class types cannot have an inline default.
 				NewPin->bDefaultValueIsIgnored = !(InputNode->IsRequired() && InputNode->Input.GetType().GetClass() == nullptr);
 
+				SetPinAutoGeneratedDefaultValue(*NewPin, *InputNode);
+				NewPin->DefaultValue = NewPin->AutogeneratedDefaultValue;
+
 				//TODO: Some visual indication of Auto bound pins.
 				//I tried just linking to null but
 // 				FNiagaraVariable AutoBoundVar;
@@ -127,7 +162,7 @@ void UNiagaraNodeFunctionCall::AllocateDefaultPins()
 		}
 
 		// Make sure to note that we've synchronized with the external version.
-		CachedChangeId = FunctionScript->GetChangeID();
+		CachedChangeId = Graph->GetChangeID();
 	}
 	else
 	{
@@ -144,6 +179,12 @@ void UNiagaraNodeFunctionCall::AllocateDefaultPins()
 			NewPin->bDefaultValueIsIgnored = true;
 		}
 
+		if (AllowDynamicPins())
+		{
+			CreateAddPin(EGPD_Input);
+			CreateAddPin(EGPD_Output);
+		}
+
 		// We don't reference an external function, so set an invalid id.
 		CachedChangeId = FGuid();
 	}
@@ -341,10 +382,15 @@ void UNiagaraNodeFunctionCall::Compile(class FHlslNiagaraTranslator* Translator,
 			}
 		}
 	}
-	else
+	else if (Signature.IsValid())
 	{
 		bError = CompileInputPins(Translator, Inputs);
 	}		
+	else
+	{
+		Translator->Error(FText::Format(LOCTEXT("UnknownFunction", "Unknown Function Call! Missing Script or Data Interface Signature. Stack Name: {0}"), FText::FromString(GetFunctionName())), this, nullptr);
+		bError = true;
+	}
 
 	if (!bError)
 	{
@@ -352,6 +398,11 @@ void UNiagaraNodeFunctionCall::Compile(class FHlslNiagaraTranslator* Translator,
 	}
 }
 
+bool UNiagaraNodeFunctionCall::IsValidPinToCompile(UEdGraphPin* Pin) const
+{
+	return !IsAddPin(Pin);
+}
+
 UObject*  UNiagaraNodeFunctionCall::GetReferencedAsset() const
 {
 	if (FunctionScript && FunctionScript->GetOutermost() != GetOutermost())
@@ -369,12 +420,15 @@ bool UNiagaraNodeFunctionCall::RefreshFromExternalChanges()
 	bool bReload = false;
 	if (FunctionScript)
 	{
-		bReload = CachedChangeId != FunctionScript->GetChangeID();
-		if (bReload)
+		UNiagaraScriptSource* Source = CastChecked<UNiagaraScriptSource>(FunctionScript->GetSource());
+		if (Source != nullptr)
 		{
-			bReload = true;
-			check(FunctionScript->GetOutermost()->HasAnyFlags(RF_WasLoaded));
-			UE_LOG(LogNiagaraEditor, Log, TEXT("RefreshFromExternalChanges %s"), *(FunctionScript->GetPathName()));
+			bReload = CachedChangeId != Source->NodeGraph->GetChangeID();
+			if (bReload)
+			{
+				bReload = true;
+				UE_LOG(LogNiagaraEditor, Log, TEXT("RefreshFromExternalChanges %s"), *(FunctionScript->GetPathName()));
+			}
 		}
 	}
 	else
@@ -414,10 +468,29 @@ void UNiagaraNodeFunctionCall::SubsumeExternalDependencies(TMap<const UObject*,
 	}
 }
 
+bool UNiagaraNodeFunctionCall::ScriptIsValid() const
+{
+	if (FunctionScript != nullptr)
+	{
+		return true;
+	}
+	else if (Signature.IsValid())
+	{
+		return true;
+	}
+	return false;
+}
+
 void UNiagaraNodeFunctionCall::BuildParameterMapHistory(FNiagaraParameterMapHistoryBuilder& OutHistory, bool bRecursive)
 {
 	Super::BuildParameterMapHistory(OutHistory, bRecursive);
+	if (!IsNodeEnabled() && OutHistory.GetIgnoreDisabled())
+	{
+		RouteParameterMapAroundMe(OutHistory, bRecursive);
+		return;
+	}
 
+	const UEdGraphSchema_Niagara* Schema = CastChecked<UEdGraphSchema_Niagara>(GetSchema());
 	if (FunctionScript)
 	{
 		UNiagaraScriptSource* Source = CastChecked<UNiagaraScriptSource>(FunctionScript->GetSource());
@@ -442,9 +515,7 @@ void UNiagaraNodeFunctionCall::BuildParameterMapHistory(FNiagaraParameterMapHist
 		GetOutputPins(OutputPins);
 
 		TArray<TPair<UEdGraphPin*, int32> > MatchedPairs;
-
-		const UEdGraphSchema_Niagara* Schema = CastChecked<UEdGraphSchema_Niagara>(GetSchema());
-
+		
 		// Find the matches of names and types of the sub-graph output pins and this function call nodes' outputs.
 		for (UEdGraphPin* ChildOutputNodePin : OutputNode->GetAllPins())
 		{
@@ -472,6 +543,10 @@ void UNiagaraNodeFunctionCall::BuildParameterMapHistory(FNiagaraParameterMapHist
 			OutHistory.RegisterParameterMapPin(MatchedPairs[i].Value, MatchedPairs[i].Key);
 		}
 	}
+	else if (!ScriptIsValid())
+	{
+		RouteParameterMapAroundMe(OutHistory, bRecursive);
+	}
 }
 
 UEdGraphPin* UNiagaraNodeFunctionCall::FindParameterMapDefaultValuePin(const FName VariableName)
@@ -526,13 +601,36 @@ void UNiagaraNodeFunctionCall::ComputeNodeName()
 		}
 	}
 
-	FString NewName = FNiagaraEditorUtilities::GetUniqueName(ProposedName, Names).ToString();
+	FString NewName = FNiagaraUtilities::GetUniqueName(ProposedName, Names).ToString();
 	if (!FunctionDisplayName.Equals(NewName))
 	{
 		FunctionDisplayName = NewName;
 	}
 }
 
+void UNiagaraNodeFunctionCall::SetPinAutoGeneratedDefaultValue(UEdGraphPin& FunctionInputPin, UNiagaraNodeInput& FunctionScriptInputNode)
+{
+	if (FunctionInputPin.bDefaultValueIsIgnored == false)
+	{
+		TArray<UEdGraphPin*> InputPins;
+		FunctionScriptInputNode.GetInputPins(InputPins);
+		if (InputPins.Num() == 1 && InputPins[0]->bDefaultValueIsIgnored == false)
+		{
+			// If the function graph's input node had an input pin, and that pin's default wasn't ignored, use that value. 
+			FunctionInputPin.AutogeneratedDefaultValue = InputPins[0]->DefaultValue;
+		}
+		else
+		{
+			const UEdGraphSchema_Niagara* Schema = CastChecked<UEdGraphSchema_Niagara>(GetSchema());
+			FString PinDefaultValue;
+			if (Schema->TryGetPinDefaultValueFromNiagaraVariable(FunctionScriptInputNode.Input, PinDefaultValue))
+			{
+				FunctionInputPin.AutogeneratedDefaultValue = PinDefaultValue;
+			}
+		}
+	}
+}
+
 bool UNiagaraNodeFunctionCall::FindAutoBoundInput(UNiagaraNodeInput* InputNode, UEdGraphPin* PinToAutoBind, FNiagaraVariable& OutFoundVar, ENiagaraInputNodeUsage& OutNodeUsage)
 {
 	check(InputNode && InputNode->IsExposed());
diff --git a/Engine/Plugins/FX/Niagara/Source/NiagaraEditor/Private/NiagaraNodeIf.cpp b/Engine/Plugins/FX/Niagara/Source/NiagaraEditor/Private/NiagaraNodeIf.cpp
--- a/Engine/Plugins/FX/Niagara/Source/NiagaraEditor/Private/NiagaraNodeIf.cpp
+++ b/Engine/Plugins/FX/Niagara/Source/NiagaraEditor/Private/NiagaraNodeIf.cpp
@@ -69,9 +69,12 @@ void UNiagaraNodeIf::AllocateDefaultPins()
 		CreatePin(EGPD_Input, Schema->TypeDefinitionToPinType(Var.GetType()), *(Var.GetName().ToString() + PathSuffix));
 	}
 
-	for (FNiagaraVariable& Var : OutputVars)
+	for (int32 Index = 0; Index < OutputVars.Num(); Index++)
 	{
+		const FNiagaraVariable& Var = OutputVars[Index];
+		const FGuid& Guid = OutputVarGuids[Index];
 		UEdGraphPin* NewPin = CreatePin(EGPD_Output, Schema->TypeDefinitionToPinType(Var.GetType()), Var.GetName());
+		NewPin->PersistentGuid = Guid;
 	}
 
 	CreateAddPin(EGPD_Output);
@@ -115,6 +118,15 @@ bool UNiagaraNodeIf::RefreshFromExternalChanges()
 	return true;
 }
 
+FGuid UNiagaraNodeIf::AddOutput(FNiagaraTypeDefinition Type, const FName& Name)
+{
+	FNiagaraVariable NewOutput(Type, Name);
+	FGuid Guid = FGuid::NewGuid();
+	OutputVars.Add(NewOutput);
+	OutputVarGuids.Add(Guid);
+	return Guid;
+}
+
 void UNiagaraNodeIf::OnPinRemoved(UEdGraphPin* PinToRemove)
 {
 	auto FindPredicate = [=](const FGuid& Guid) { return Guid == PinToRemove->PersistentGuid; };
@@ -139,12 +151,9 @@ void UNiagaraNodeIf::OnNewTypedPinAdded(UEdGraphPin* NewPin)
 	{
 		OutputNames.Add(Output.GetName());
 	}
-	FName OutputName = FNiagaraEditorUtilities::GetUniqueName(*OutputType.GetNameText().ToString(), OutputNames);
+	FName OutputName = FNiagaraUtilities::GetUniqueName(*OutputType.GetNameText().ToString(), OutputNames);
 
-	FNiagaraVariable NewOutput(OutputType, OutputName);
-	FGuid Guid = FGuid::NewGuid();
-	OutputVars.Add(NewOutput);
-	OutputVarGuids.Add(Guid);
+	FGuid Guid = AddOutput(OutputType, OutputName);
 
 	// Update the pin's data too so that it's connection is maintained after reallocating.
 	NewPin->PersistentGuid = Guid;
@@ -152,7 +161,7 @@ void UNiagaraNodeIf::OnNewTypedPinAdded(UEdGraphPin* NewPin)
 	ReallocatePins();
 }
 
-void UNiagaraNodeIf::OnPinRenamed(UEdGraphPin* RenamedPin)
+void UNiagaraNodeIf::OnPinRenamed(UEdGraphPin* RenamedPin, const FString& OldName)
 {
 	auto FindPredicate = [=](const FGuid& Guid) { return Guid == RenamedPin->PersistentGuid; };
 	int32 FoundIndex = OutputVarGuids.IndexOfByPredicate(FindPredicate);
@@ -166,7 +175,7 @@ void UNiagaraNodeIf::OnPinRenamed(UEdGraphPin* RenamedPin)
 				OutputNames.Add(OutputVars[Index].GetName());
 			}
 		}
-		const FName OutputName = FNiagaraEditorUtilities::GetUniqueName(RenamedPin->PinName, OutputNames);
+		const FName OutputName = FNiagaraUtilities::GetUniqueName(RenamedPin->PinName, OutputNames);
 		OutputVars[FoundIndex].SetName(OutputName);
 	}
 	ReallocatePins();
@@ -188,4 +197,9 @@ FText UNiagaraNodeIf::GetTooltipText() const
 	return LOCTEXT("IfDesc", "If Condition is true, the output value is A, otherwise output B.");
 }
 
+FText UNiagaraNodeIf::GetNodeTitle(ENodeTitleType::Type TitleType) const
+{
+	return LOCTEXT("IfTitle", "If");
+}
+
 #undef LOCTEXT_NAMESPACE
diff --git a/Engine/Plugins/FX/Niagara/Source/NiagaraEditor/Private/NiagaraNodeIf.h b/Engine/Plugins/FX/Niagara/Source/NiagaraEditor/Private/NiagaraNodeIf.h
--- a/Engine/Plugins/FX/Niagara/Source/NiagaraEditor/Private/NiagaraNodeIf.h
+++ b/Engine/Plugins/FX/Niagara/Source/NiagaraEditor/Private/NiagaraNodeIf.h
@@ -27,6 +27,7 @@ public:
 	//~ Begin EdGraphNode Interface
 	virtual void AllocateDefaultPins() override;
 	virtual FText GetTooltipText() const override;
+	virtual FText GetNodeTitle(ENodeTitleType::Type TitleType) const override;
 	//~ End EdGraphNode Interface
 
 	//~ Begin UNiagaraNode Interface
@@ -34,6 +35,9 @@ public:
 	virtual bool RefreshFromExternalChanges() override;
 	//~ End UNiagaraNode Interface
 
+	/** Helper function to create a variable to add to the OutputVars and FGuid to add to OutputVarGuids. */
+	FGuid AddOutput(FNiagaraTypeDefinition Type, const FName& Name);
+
 protected:
 	//~ Begin EdGraphNode Interface
 	virtual void OnPinRemoved(UEdGraphPin* PinToRemove) override;
@@ -41,7 +45,7 @@ protected:
 
 	//~ Begin UNiagaraNodeWithDynamicPins Interface
 	virtual void OnNewTypedPinAdded(UEdGraphPin* NewPin) override;
-	virtual void OnPinRenamed(UEdGraphPin* RenamedPin) override;
+	virtual void OnPinRenamed(UEdGraphPin* RenamedPin, const FString& OldName) override;
 	virtual bool CanRenamePin(const UEdGraphPin* Pin) const override;
 	virtual bool CanRemovePin(const UEdGraphPin* Pin) const override;
 	virtual bool CanMovePin(const UEdGraphPin* Pin) const override { return false; }
diff --git a/Engine/Plugins/FX/Niagara/Source/NiagaraEditor/Private/NiagaraNodeInput.cpp b/Engine/Plugins/FX/Niagara/Source/NiagaraEditor/Private/NiagaraNodeInput.cpp
--- a/Engine/Plugins/FX/Niagara/Source/NiagaraEditor/Private/NiagaraNodeInput.cpp
+++ b/Engine/Plugins/FX/Niagara/Source/NiagaraEditor/Private/NiagaraNodeInput.cpp
@@ -12,10 +12,11 @@
 #include "NiagaraDataInterface.h"
 #include "SNiagaraGraphNodeInput.h"
 #include "NiagaraEditorUtilities.h"
-
+#include "NiagaraEditorModule.h"
 
 #define LOCTEXT_NAMESPACE "NiagaraNodeInput"
 
+DECLARE_CYCLE_STAT(TEXT("NiagaraEditor - UNiagaraNodeInput - SortNodes"), STAT_NiagaraEditor_UNiagaraNodeInput_SortNodes, STATGROUP_NiagaraEditor);
 
 UNiagaraNodeInput::UNiagaraNodeInput(const FObjectInitializer& ObjectInitializer)
 	: Super(ObjectInitializer)
@@ -62,8 +63,23 @@ void UNiagaraNodeInput::PostEditChangeProperty(struct FPropertyChangedEvent& Pro
 	Super::PostEditChangeProperty(PropertyChangedEvent);
 }
 
+void UNiagaraNodeInput::PostLoad()
+{
+	Super::PostLoad();
+	if (DataInterface != nullptr)
+	{
+		DataInterface->OnChanged().AddUObject(this, &UNiagaraNodeInput::DataInterfaceChanged);
+	}
+}
+
 void UNiagaraNodeInput::BuildParameterMapHistory(FNiagaraParameterMapHistoryBuilder& OutHistory, bool bRecursive)
 {
+	if (!IsNodeEnabled() && OutHistory.GetIgnoreDisabled())
+	{
+		RouteParameterMapAroundMe(OutHistory, bRecursive);
+		return;
+	}
+
 	if (Input.GetType() == FNiagaraTypeDefinition::GetParameterMapDef())
 	{
 		int32 ParamMapIdx = OutHistory.FindMatchingParameterMapFromContextInputs(Input);
@@ -157,7 +173,7 @@ FName UNiagaraNodeInput::GenerateUniqueName(const UNiagaraGraph* Graph, FName& P
 		}
 	}
 
-	return FNiagaraEditorUtilities::GetUniqueName(ProposedName, FNiagaraEditorUtilities::GetSystemConstantNames().Union(InputNames));
+	return FNiagaraUtilities::GetUniqueName(ProposedName, FNiagaraEditorUtilities::GetSystemConstantNames().Union(InputNames));
 }
 
 bool UNiagaraNodeInput::VerifyNodeRenameTextCommit(const FText& NewText, UNiagaraNode* NodeBeingChanged, FText& OutErrorMessage)
@@ -266,7 +282,7 @@ void UNiagaraNodeInput::OnRenameNode(const FString& NewName)
 		Node->Input.SetName(FName(*NewName));
 		if (DataInterface != nullptr)
 		{
-			Node->DataInterface->Rename(*NewName);
+			Node->GetDataInterface()->Rename(*NewName);
 		}
 		Node->ReallocatePins();
 	}
@@ -389,7 +405,7 @@ void UNiagaraNodeInput::AutowireNewNode(UEdGraphPin* FromPin)
 					CandidateName = (*PinName);
 				}
 
-				Input.SetName(FNiagaraEditorUtilities::GetUniqueName(CandidateName, FNiagaraEditorUtilities::GetSystemConstantNames().Union(InputNames)));
+				Input.SetName(FNiagaraUtilities::GetUniqueName(CandidateName, FNiagaraEditorUtilities::GetSystemConstantNames().Union(InputNames)));
 				CallSortPriority = HighestSortPriority + 1;
 			}
 			ReallocatePins();
@@ -419,6 +435,12 @@ void UNiagaraNodeInput::NotifyExposureOptionsChanged()
 
 void UNiagaraNodeInput::Compile(class FHlslNiagaraTranslator* Translator, TArray<int32>& Outputs)
 {
+	if (!IsNodeEnabled())
+	{
+		Outputs.Add(INDEX_NONE);
+		return;
+	}
+
 	int32 FunctionParam = INDEX_NONE;
 	if (IsExposed() && Translator->GetFunctionParameter(Input, FunctionParam))
 	{
@@ -434,7 +456,7 @@ void UNiagaraNodeInput::Compile(class FHlslNiagaraTranslator* Translator, TArray
 				if (Usage == ENiagaraInputNodeUsage::Parameter && DataInterface != nullptr)
 				{
 					check(Input.GetType().GetClass());
-					Outputs.Add(Translator->RegisterDataInterface(Input, DataInterface));
+					Outputs.Add(Translator->RegisterDataInterface(Input, DataInterface, false));
 					return;
 				}
 				else
@@ -453,7 +475,7 @@ void UNiagaraNodeInput::Compile(class FHlslNiagaraTranslator* Translator, TArray
 		if (DataInterface)
 		{
 			check(Input.GetType().GetClass());
-			Outputs.Add(Translator->RegisterDataInterface(Input, DataInterface)); break;
+			Outputs.Add(Translator->RegisterDataInterface(Input, DataInterface, false)); break;
 		}
 		else
 		{
@@ -470,6 +492,8 @@ void UNiagaraNodeInput::Compile(class FHlslNiagaraTranslator* Translator, TArray
 
 void UNiagaraNodeInput::SortNodes(TArray<UNiagaraNodeInput*>& InOutNodes)
 {
+	SCOPE_CYCLE_COUNTER(STAT_NiagaraEditor_UNiagaraNodeInput_SortNodes);
+
 	auto SortVars = [](UNiagaraNodeInput& A, UNiagaraNodeInput& B)
 	{
 		if (A.CallSortPriority < B.CallSortPriority)
@@ -487,4 +511,34 @@ void UNiagaraNodeInput::SortNodes(TArray<UNiagaraNodeInput*>& InOutNodes)
 	InOutNodes.Sort(SortVars);
 }
 
+UNiagaraDataInterface* UNiagaraNodeInput::GetDataInterface() const
+{
+	return DataInterface;
+}
+
+void UNiagaraNodeInput::SetDataInterface(UNiagaraDataInterface* InDataInterface)
+{
+	if (DataInterface != nullptr)
+	{
+		DataInterface->OnChanged().RemoveAll(this);
+	}
+	DataInterface = InDataInterface;
+	if (DataInterface != nullptr)
+	{
+		DataInterface->OnChanged().AddUObject(this, &UNiagaraNodeInput::DataInterfaceChanged);
+	}
+	DataInterfaceChanged();
+}
+
+void UNiagaraNodeInput::DataInterfaceChanged()
+{
+	// Don't use GetNiagaraGraph() here since this may be called on a temporary node which isn't
+	// in a proper graph yet.
+	UNiagaraGraph* NiagaraGraph = Cast<UNiagaraGraph>(GetGraph());
+	if (NiagaraGraph != nullptr)
+	{
+		NiagaraGraph->NotifyGraphDataInterfaceChanged();
+	}
+}
+
 #undef LOCTEXT_NAMESPACE
diff --git a/Engine/Plugins/FX/Niagara/Source/NiagaraEditor/Private/NiagaraNodeOutput.cpp b/Engine/Plugins/FX/Niagara/Source/NiagaraEditor/Private/NiagaraNodeOutput.cpp
--- a/Engine/Plugins/FX/Niagara/Source/NiagaraEditor/Private/NiagaraNodeOutput.cpp
+++ b/Engine/Plugins/FX/Niagara/Source/NiagaraEditor/Private/NiagaraNodeOutput.cpp
@@ -3,6 +3,7 @@
 #include "NiagaraNodeOutput.h"
 #include "UObject/UnrealType.h"
 #include "NiagaraHlslTranslator.h"
+#include "NiagaraEmitter.h"
 #include "NiagaraScript.h"
 #include "NiagaraGraph.h"
 #include "NiagaraScriptSource.h"
@@ -12,11 +13,12 @@
 #include "MultiBoxBuilder.h"
 #include "SBox.h"
 #include "SEditableTextBox.h"
+#include "NiagaraEditorUtilities.h"
 
 #define LOCTEXT_NAMESPACE "NiagaraNodeOutput"
 
 UNiagaraNodeOutput::UNiagaraNodeOutput(const FObjectInitializer& ObjectInitializer)
-: Super(ObjectInitializer), ScriptTypeIndex(0)
+: Super(ObjectInitializer), ScriptTypeIndex_DEPRECATED(0)
 {
 }
 
@@ -154,7 +156,12 @@ FText UNiagaraNodeOutput::GetNodeTitle(ENodeTitleType::Type TitleType) const
 	}
 	else if (ScriptType == ENiagaraScriptUsage::ParticleEventScript)
 	{
-		return FText::Format(NSLOCTEXT("NiagaraNodeOutput", "OutputEvent", "Output Event {0}"), FText::AsNumber(ScriptTypeIndex));
+		FText EventName;
+		if (FNiagaraEditorUtilities::TryGetEventDisplayName(GetTypedOuter<UNiagaraEmitter>(), ScriptTypeId, EventName) == false)
+		{
+			EventName = NSLOCTEXT("NiagaraNodeOutput", "UnknownEventName", "Unknown");
+		}
+		return FText::Format(NSLOCTEXT("NiagaraNodeOutput", "OutputEvent", "Output Event {0}"), EventName);
 	}
 	else if (ScriptType == ENiagaraScriptUsage::Function)
 	{
diff --git a/Engine/Plugins/FX/Niagara/Source/NiagaraEditor/Private/NiagaraNodeParameterMapBase.cpp b/Engine/Plugins/FX/Niagara/Source/NiagaraEditor/Private/NiagaraNodeParameterMapBase.cpp
--- a/Engine/Plugins/FX/Niagara/Source/NiagaraEditor/Private/NiagaraNodeParameterMapBase.cpp
+++ b/Engine/Plugins/FX/Niagara/Source/NiagaraEditor/Private/NiagaraNodeParameterMapBase.cpp
@@ -14,6 +14,10 @@
 #include "NiagaraScript.h"
 #include "NiagaraScriptSource.h"
 #include "NiagaraConstants.h"
+#include "NiagaraParameterCollection.h"
+
+#include "IAssetTools.h"
+#include "AssetRegistryModule.h"
 
 #include "Framework/MultiBox/MultiBoxBuilder.h"
 
@@ -133,6 +137,11 @@ TSharedRef<SWidget> UNiagaraNodeParameterMapBase::GenerateAddPinMenu(const FStri
 		LOCTEXT("SupportedTypesTooltip", "Create an entry in the module namespace that you will name immediately after."),
 		FNewMenuDelegate::CreateLambda([WorkingCustomName, InPin, this](FMenuBuilder& SubMenuBuilder) { BuildTypeMenu(SubMenuBuilder, WorkingCustomName, InPin); }));
 
+	WorkingCustomName = TEXT("ParameterCollections.") + InWorkingPinName;
+	MenuBuilder.AddSubMenu(LOCTEXT("ParameterCollectionsMenuTitle", "ParameterCollections"),
+		LOCTEXT("ParameterCollectionsMenuTooltip", "Create a reference to a parameter in a Parameter Collection."),
+		FNewMenuDelegate::CreateLambda([WorkingCustomName, InPin, this](FMenuBuilder& SubMenuBuilder) { BuildParameterCollectionsMenu(SubMenuBuilder, WorkingCustomName, InPin); }));
+
 	return MenuBuilder.MakeWidget();
 }
 
@@ -152,6 +161,35 @@ void UNiagaraNodeParameterMapBase::BuildCommonMenu(FMenuBuilder& InMenuBuilder,
 	}
 }
 
+void UNiagaraNodeParameterMapBase::BuildParameterCollectionsMenu(FMenuBuilder& InMenuBuilder, const FString& InWorkingName, SNiagaraGraphPinAdd* InPin)
+{
+	//Create sub menus for parameter collections.
+	FAssetRegistryModule& AssetRegistryModule = FModuleManager::LoadModuleChecked<FAssetRegistryModule>("AssetRegistry");
+	TArray<FAssetData> CollectionAssets;
+	AssetRegistryModule.Get().GetAssetsByClass(UNiagaraParameterCollection::StaticClass()->GetFName(), CollectionAssets);
+	for (FAssetData& CollectionAsset : CollectionAssets)
+	{
+		UNiagaraParameterCollection* Collection = CastChecked<UNiagaraParameterCollection>(CollectionAsset.GetAsset());
+
+		InMenuBuilder.AddSubMenu(FText::FromName(CollectionAsset.AssetName),
+			FText::FromString(CollectionAsset.GetFullName()),
+			FNewMenuDelegate::CreateLambda([Collection, InPin, this](FMenuBuilder& SubMenuBuilder) { BuildParameterCollectionMenu(SubMenuBuilder, Collection, InPin); }));
+	}
+}
+
+void UNiagaraNodeParameterMapBase::BuildParameterCollectionMenu(FMenuBuilder& InMenuBuilder, UNiagaraParameterCollection* Collection, SNiagaraGraphPinAdd* InPin)
+{
+	for (const FNiagaraVariable& Parameter : Collection->GetParameters())
+	{
+		FText Desc = FText::GetEmpty();//FNiagaraConstants::GetAttributeDescription(Parameter);
+		FString FriendlyName = Collection->FriendlyNameFromParameterName(Parameter.GetName().ToString());
+		InMenuBuilder.AddMenuEntry(
+			FText::FromString(FriendlyName),
+			FText::Format(LOCTEXT("AddButtonTypeEntryToolTipFormatNPC", "Add a reference to Collection Parameter {0}. {1}"), FText::FromName(Parameter.GetName()), Desc),
+			FSlateIcon(),
+			FUIAction(FExecuteAction::CreateSP(InPin, &SNiagaraGraphPinAdd::OnAddType, Parameter)));
+	}
+}
 
 void UNiagaraNodeParameterMapBase::BuildLocalMenu(FMenuBuilder& InMenuBuilder, const FString& InWorkingName, SNiagaraGraphPinAdd* InPin)
 {
@@ -198,5 +236,55 @@ void UNiagaraNodeParameterMapBase::BuildEngineMenu(FMenuBuilder& InMenuBuilder,
 	}
 }
 
+FText UNiagaraNodeParameterMapBase::GetPinDescriptionText(UEdGraphPin* Pin) const
+{
+	FNiagaraVariable Var = CastChecked<UEdGraphSchema_Niagara>(GetSchema())->PinToNiagaraVariable(Pin);
+
+	const UNiagaraGraph* Graph = GetNiagaraGraph();
+	const FNiagaraVariableMetaData* OldMetaData = Graph->GetMetaData(Var);
+	if (OldMetaData)
+	{
+		ensure(OldMetaData->ReferencerNodes.Contains(TWeakObjectPtr<UObject>(this)));
+		return OldMetaData->Description;
+	}
+	return FText::GetEmpty();
+}
+
+/** Called when a pin's description text is committed. */
+void UNiagaraNodeParameterMapBase::PinDescriptionTextCommitted(const FText& Text, ETextCommit::Type CommitType, UEdGraphPin* Pin)
+{
+	FNiagaraVariable Var = CastChecked<UEdGraphSchema_Niagara>(GetSchema())->PinToNiagaraVariable(Pin);
+	UNiagaraGraph* Graph = GetNiagaraGraph();
+	if (FNiagaraConstants::IsNiagaraConstant(Var))
+	{
+		UE_LOG(LogNiagaraEditor, Error, TEXT("You cannot set the description for a Niagara internal constant \"%s\""),*Var.GetName().ToString());
+		return;
+	}
+	FNiagaraVariableMetaData* OldMetaData = Graph->GetMetaData(Var);
+	bool bSet = !Text.IsEmptyOrWhitespace();
+	if (OldMetaData && !OldMetaData->Description.IsEmptyOrWhitespace())
+	{
+		bSet = true;
+	}
+
+	FScopedTransaction AddNewPinTransaction(LOCTEXT("Rename Pin Desc", "Changed variable description"));
+	Modify();
+	Pin->Modify();
+
+	if (OldMetaData)
+	{
+		ensure(OldMetaData->ReferencerNodes.Contains(TWeakObjectPtr<UObject>(this)));
+		Graph->Modify();
+		OldMetaData->Description = Text;
+	}
+	else
+	{
+		Graph->Modify();
+		FNiagaraVariableMetaData& NewMetaData = Graph->FindOrAddMetaData(Var);
+		NewMetaData.Description = Text;
+		NewMetaData.ReferencerNodes.Add(TWeakObjectPtr<UObject>(this));
+	}
+}
+
 
 #undef LOCTEXT_NAMESPACE
\ No newline at end of file
diff --git a/Engine/Plugins/FX/Niagara/Source/NiagaraEditor/Private/NiagaraNodeParameterMapBase.h b/Engine/Plugins/FX/Niagara/Source/NiagaraEditor/Private/NiagaraNodeParameterMapBase.h
--- a/Engine/Plugins/FX/Niagara/Source/NiagaraEditor/Private/NiagaraNodeParameterMapBase.h
+++ b/Engine/Plugins/FX/Niagara/Source/NiagaraEditor/Private/NiagaraNodeParameterMapBase.h
@@ -28,8 +28,16 @@ public:
 	
 	virtual TSharedRef<SWidget> GenerateAddPinMenu(const FString& InWorkingPinName, SNiagaraGraphPinAdd* InPin) override;
 
+	/** Gets the description text for a pin. */
+	FText GetPinDescriptionText(UEdGraphPin* Pin) const;
+
+	/** Called when a pin's description text is committed. */
+	void PinDescriptionTextCommitted(const FText& Text, ETextCommit::Type CommitType, UEdGraphPin* Pin);
 protected:
 	virtual void BuildCommonMenu(FMenuBuilder& InMenuBuilder, const FString& InWorkingName, SNiagaraGraphPinAdd* InPin);
 	virtual void BuildLocalMenu(FMenuBuilder& InMenuBuilder, const FString& InWorkingName, SNiagaraGraphPinAdd* InPin);
 	virtual void BuildEngineMenu(FMenuBuilder& InMenuBuilder, const FString& InWorkingName, SNiagaraGraphPinAdd* InPin);
+	virtual void BuildParameterCollectionsMenu(FMenuBuilder& InMenuBuilder, const FString& InWorkingName, SNiagaraGraphPinAdd* InPin);
+	virtual void BuildParameterCollectionMenu(FMenuBuilder& InMenuBuilder, UNiagaraParameterCollection* Collection, SNiagaraGraphPinAdd* InPin);
+
 };
diff --git a/Engine/Plugins/FX/Niagara/Source/NiagaraEditor/Private/NiagaraNodeParameterMapGet.cpp b/Engine/Plugins/FX/Niagara/Source/NiagaraEditor/Private/NiagaraNodeParameterMapGet.cpp
--- a/Engine/Plugins/FX/Niagara/Source/NiagaraEditor/Private/NiagaraNodeParameterMapGet.cpp
+++ b/Engine/Plugins/FX/Niagara/Source/NiagaraEditor/Private/NiagaraNodeParameterMapGet.cpp
@@ -11,6 +11,12 @@
 #include "NiagaraEditorModule.h"
 #include "INiagaraEditorTypeUtilities.h"
 #include "ModuleManager.h"
+#include "SharedPointer.h"
+#include "NiagaraConstants.h"
+#include "MultiBoxBuilder.h"
+#include "SBox.h"
+#include "SEditableTextBox.h"
+#include "EdGraph/EdGraphNode.h"
 
 #define LOCTEXT_NAMESPACE "NiagaraNodeParameterMapGet"
 
@@ -79,20 +85,13 @@ UEdGraphPin* UNiagaraNodeParameterMapGet::CreateDefaultPin(UEdGraphPin* OutputPi
 	}
 
 	UEdGraphPin* DefaultPin = CreatePin(EEdGraphPinDirection::EGPD_Input, OutputPin->PinType, TEXT(""));
-	const UEdGraphSchema_Niagara* Schema = GetDefault<UEdGraphSchema_Niagara>();
-	FNiagaraVariable Var = Schema->PinToNiagaraVariable(OutputPin);
-	FNiagaraEditorUtilities::ResetVariableToDefaultValue(Var);
 
-	FNiagaraEditorModule& NiagaraEditorModule = FModuleManager::GetModuleChecked<FNiagaraEditorModule>("NiagaraEditor");
-	if (Var.IsDataAllocated())
+	const UEdGraphSchema_Niagara* Schema = GetDefault<UEdGraphSchema_Niagara>();
+	FNiagaraVariable Var = Schema->PinToNiagaraVariable(OutputPin, true);
+	FString PinDefaultValue;
+	if (Schema->TryGetPinDefaultValueFromNiagaraVariable(Var, PinDefaultValue))
 	{
-		TSharedPtr<INiagaraEditorTypeUtilities> TypeEditorUtilities = NiagaraEditorModule.GetTypeUtilities(Var.GetType());
-
-		if (TypeEditorUtilities.IsValid() && TypeEditorUtilities->CanHandlePinDefaults())
-		{
-			FString Value = TypeEditorUtilities->GetPinDefaultStringFromValue(Var);
-			DefaultPin->DefaultValue = Value;
-		}
+		DefaultPin->DefaultValue = PinDefaultValue;
 	}
 	
 	if (!OutputPin->PersistentGuid.IsValid())
@@ -109,9 +108,54 @@ UEdGraphPin* UNiagaraNodeParameterMapGet::CreateDefaultPin(UEdGraphPin* OutputPi
 	return DefaultPin;
 }
 
-void UNiagaraNodeParameterMapGet::OnPinRenamed(UEdGraphPin* RenamedPin)
+void UNiagaraNodeParameterMapGet::OnPinRenamed(UEdGraphPin* RenamedPin, const FString& OldName)
 {
 	RenamedPin->PinFriendlyName = FText::FromName(RenamedPin->PinName);
+
+	FNiagaraTypeDefinition VarType = CastChecked<UEdGraphSchema_Niagara>(GetSchema())->PinToTypeDefinition(RenamedPin);
+	FNiagaraVariable Var(VarType, *OldName);
+
+	UNiagaraGraph* Graph = GetNiagaraGraph();
+	FNiagaraVariableMetaData* OldMetaData = Graph->GetMetaData(Var);
+
+	if (OldMetaData)
+	{
+		Graph->Modify();
+		OldMetaData->ReferencerNodes.Remove(TWeakObjectPtr<UObject>(this));
+	}
+
+	FNiagaraVariable NewVar(VarType, *RenamedPin->GetName());
+	FNiagaraVariableMetaData* NewMetaData = Graph->GetMetaData(NewVar);
+
+	// If no variable has already defined the meta-data for this entry and it isn't being redefined to one of our constants, copy over the old values that make sense (not display name).
+	if (NewMetaData == nullptr && OldMetaData != nullptr && !FNiagaraConstants::IsNiagaraConstant(NewVar))
+	{
+		FNiagaraVariableMetaData& NewVarMetaData = Graph->FindOrAddMetaData(NewVar);
+		NewVarMetaData.PropertyMetaData = OldMetaData->PropertyMetaData;
+		NewVarMetaData.PropertyMetaData.Remove(TEXT("DisplayName")); // DisplayName is probably incorrect, so drop it.
+		NewVarMetaData.ReferencerNodes.Add(TWeakObjectPtr<UObject>(this));
+	}
+
+	if (OldMetaData)
+	{
+		Graph->PurgeUnreferencedMetaData();
+	}
+
+	if (RenamedPin == PinPendingRename)
+	{
+		PinPendingRename = nullptr;
+	}
+
+
+	UEdGraphPin* DefaultPin = GetDefaultPin(RenamedPin);
+	
+	if (DefaultPin)
+	{
+		DefaultPin->Modify();
+		SynchronizeDefaultInputPin(DefaultPin, RenamedPin);
+	}
+
+	Graph->NotifyGraphNeedsRecompile();
 }
 
 
@@ -149,6 +193,17 @@ void UNiagaraNodeParameterMapGet::RemoveDynamicPin(UEdGraphPin* Pin)
 	if (Pin->Direction == EEdGraphPinDirection::EGPD_Output)
 	{
 		DefaultPin = GetDefaultPin(Pin);
+
+		FNiagaraVariable Var = CastChecked<UEdGraphSchema_Niagara>(GetSchema())->PinToNiagaraVariable(Pin);
+
+		UNiagaraGraph* Graph = GetNiagaraGraph();
+		FNiagaraVariableMetaData* OldMetaData = Graph->GetMetaData(Var);
+
+		if (OldMetaData)
+		{
+			Graph->Modify();
+			OldMetaData->ReferencerNodes.Remove(TWeakObjectPtr<UObject>(this));
+		}
 	}
 
 	RemovePin(Pin);
@@ -265,6 +320,12 @@ void UNiagaraNodeParameterMapGet::BuildParameterMapHistory(FNiagaraParameterMapH
 		OutHistory.VisitInputPin(GetInputPin(0), this);
 	}
 
+	if (!IsNodeEnabled() && OutHistory.GetIgnoreDisabled())
+	{
+		RouteParameterMapAroundMe(OutHistory, bRecursive);
+		return;
+	}
+
 	int32 ParamMapIdx = INDEX_NONE;
 	if (GetInputPin(0)->LinkedTo.Num() != 0)
 	{
@@ -354,27 +415,44 @@ bool UNiagaraNodeParameterMapGet::CommitEditablePinName(const FText& InName, UEd
 		FScopedTransaction AddNewPinTransaction(LOCTEXT("Rename Pin", "Renamed pin"));
 		Modify();
 		InGraphPinObj->Modify();
-		UEdGraphPin* DefaultPin = GetDefaultPin(InGraphPinObj);
 
-		InGraphPinObj->PinFriendlyName = InName;
+		FString OldPinName = InGraphPinObj->PinName.ToString();
 		InGraphPinObj->PinName = *InName.ToString();
 
-		if (DefaultPin)
-		{
-			DefaultPin->Modify();
-			SynchronizeDefaultInputPin(DefaultPin, InGraphPinObj);
-		}
-
-
-		if (PinPendingRename == InGraphPinObj)
-		{
-			PinPendingRename = nullptr;
-		}	
-		GetGraph()->NotifyGraphChanged();
+		OnPinRenamed(InGraphPinObj, OldPinName);
 
 		return true;
 	}
 	return false;
 }
 
+void UNiagaraNodeParameterMapGet::GetContextMenuActions(const FGraphNodeContextMenuBuilder& Context) const
+{
+	UNiagaraNodeParameterMapBase::GetContextMenuActions(Context);
+
+	UEdGraphPin* Pin = const_cast<UEdGraphPin*>(Context.Pin);
+	if (Pin && Pin->Direction == EEdGraphPinDirection::EGPD_Output)
+	{
+		
+		FNiagaraVariable Var = CastChecked<UEdGraphSchema_Niagara>(GetSchema())->PinToNiagaraVariable(Pin);
+		const UNiagaraGraph* Graph = GetNiagaraGraph();
+
+		if (!FNiagaraConstants::IsNiagaraConstant(Var))
+		{
+			Context.MenuBuilder->BeginSection("EdGraphSchema_NiagaraMetaDataActions", LOCTEXT("EditPinMenuHeader", "Meta-Data"));
+			TSharedRef<SWidget> RenameWidget =
+				SNew(SBox)
+				.WidthOverride(100)
+				.Padding(FMargin(5, 0, 0, 0))
+				[
+					SNew(SEditableTextBox)
+					.Text_UObject(this, &UNiagaraNodeParameterMapBase::GetPinDescriptionText, Pin)
+					.OnTextCommitted_UObject(this, &UNiagaraNodeParameterMapBase::PinDescriptionTextCommitted, Pin)
+				];
+			Context.MenuBuilder->AddWidget(RenameWidget, LOCTEXT("DescMenuItem", "Description"));
+			Context.MenuBuilder->EndSection();
+		}
+	}
+}
+
 #undef LOCTEXT_NAMESPACE
diff --git a/Engine/Plugins/FX/Niagara/Source/NiagaraEditor/Private/NiagaraNodeParameterMapGet.h b/Engine/Plugins/FX/Niagara/Source/NiagaraEditor/Private/NiagaraNodeParameterMapGet.h
--- a/Engine/Plugins/FX/Niagara/Source/NiagaraEditor/Private/NiagaraNodeParameterMapGet.h
+++ b/Engine/Plugins/FX/Niagara/Source/NiagaraEditor/Private/NiagaraNodeParameterMapGet.h
@@ -21,6 +21,7 @@ public:
 	virtual void AllocateDefaultPins() override;
 	virtual TSharedPtr<SGraphNode> CreateVisualWidget() override;
 
+	virtual void GetContextMenuActions(const FGraphNodeContextMenuBuilder& Context) const override;
 	virtual bool IsPinNameEditable(const UEdGraphPin* GraphPinObj) const override;
 	virtual bool IsPinNameEditableUponCreation(const UEdGraphPin* GraphPinObj) const override;
 	virtual bool VerifyEditablePinName(const FText& InName, FText& OutErrorMessage, const UEdGraphPin* InGraphPinObj) const override;
@@ -39,7 +40,7 @@ public:
 
 protected:
 	virtual void OnNewTypedPinAdded(UEdGraphPin* NewPin) override;
-	virtual void OnPinRenamed(UEdGraphPin* RenamedPin) override;
+	virtual void OnPinRenamed(UEdGraphPin* RenamedPin, const FString& OldName) override;
 
 	/** Synchronize the removal of the output pin with its default.*/
 	virtual void RemoveDynamicPin(UEdGraphPin* Pin) override;
diff --git a/Engine/Plugins/FX/Niagara/Source/NiagaraEditor/Private/NiagaraNodeParameterMapSet.cpp b/Engine/Plugins/FX/Niagara/Source/NiagaraEditor/Private/NiagaraNodeParameterMapSet.cpp
--- a/Engine/Plugins/FX/Niagara/Source/NiagaraEditor/Private/NiagaraNodeParameterMapSet.cpp
+++ b/Engine/Plugins/FX/Niagara/Source/NiagaraEditor/Private/NiagaraNodeParameterMapSet.cpp
@@ -5,6 +5,13 @@
 #include "NiagaraEditorUtilities.h"
 #include "SNiagaraGraphNodeConvert.h"
 #include "NiagaraHlslTranslator.h"
+#include "SharedPointer.h"
+#include "NiagaraGraph.h"
+#include "NiagaraConstants.h"
+#include "MultiBoxBuilder.h"
+#include "SBox.h"
+#include "SEditableTextBox.h"
+#include "EdGraph/EdGraphNode.h"
 
 #include "ScopedTransaction.h"
 
@@ -71,10 +78,45 @@ void UNiagaraNodeParameterMapSet::OnNewTypedPinAdded(UEdGraphPin* NewPin)
 	PinPendingRename = NewPin;
 }
 
-void UNiagaraNodeParameterMapSet::OnPinRenamed(UEdGraphPin* RenamedPin) 
+void UNiagaraNodeParameterMapSet::OnPinRenamed(UEdGraphPin* RenamedPin, const FString& OldName)
 {
 	RenamedPin->PinFriendlyName = FText::FromName(RenamedPin->PinName);
 
+	FNiagaraTypeDefinition VarType = CastChecked<UEdGraphSchema_Niagara>(GetSchema())->PinToTypeDefinition(RenamedPin);
+	FNiagaraVariable Var(VarType, *OldName);
+
+	UNiagaraGraph* Graph = GetNiagaraGraph();
+	FNiagaraVariableMetaData* OldMetaData = Graph->GetMetaData(Var);
+
+	if (OldMetaData)
+	{
+		Graph->Modify();
+		OldMetaData->ReferencerNodes.Remove(TWeakObjectPtr<UObject>(this));
+	}
+
+	FNiagaraVariable NewVar(VarType, *RenamedPin->GetName());
+	FNiagaraVariableMetaData* NewMetaData = Graph->GetMetaData(NewVar);
+
+	// If no variable has already defined the meta-data for this entry and it isn't being redefined to one of our constants, copy over the old values that make sense (not display name).
+	if (NewMetaData == nullptr && OldMetaData != nullptr && !FNiagaraConstants::IsNiagaraConstant(NewVar))
+	{
+		FNiagaraVariableMetaData& NewVarMetaData = Graph->FindOrAddMetaData(NewVar);
+		NewVarMetaData.PropertyMetaData = OldMetaData->PropertyMetaData;
+		NewVarMetaData.PropertyMetaData.Remove(TEXT("DisplayName")); // DisplayName is probably incorrect, so drop it.
+		NewVarMetaData.ReferencerNodes.Add(TWeakObjectPtr<UObject>(this));
+	}
+
+	if (OldMetaData)
+	{
+		Graph->PurgeUnreferencedMetaData();
+	}
+
+	if (RenamedPin == PinPendingRename)
+	{
+		PinPendingRename = nullptr;
+	}
+	
+	Graph->NotifyGraphNeedsRecompile();
 }
 
 bool UNiagaraNodeParameterMapSet::CommitEditablePinName(const FText& InName, UEdGraphPin* InGraphPinObj) 
@@ -84,14 +126,11 @@ bool UNiagaraNodeParameterMapSet::CommitEditablePinName(const FText& InName, UEd
 		FScopedTransaction AddNewPinTransaction(LOCTEXT("Rename Pin", "Renamed pin"));
 		Modify();
 		InGraphPinObj->Modify();
-		InGraphPinObj->PinFriendlyName = InName;
+		
+		FString OldPinName = InGraphPinObj->PinName.ToString();
 		InGraphPinObj->PinName = *InName.ToString();
+		OnPinRenamed(InGraphPinObj, OldPinName);
 
-		if (InGraphPinObj == PinPendingRename)
-		{
-			PinPendingRename = nullptr;
-		}
-		GetGraph()->NotifyGraphChanged();
 		return true;
 	}
 	return false;
@@ -123,6 +162,11 @@ void UNiagaraNodeParameterMapSet::Compile(class FHlslNiagaraTranslator* Translat
 			continue;
 		}
 
+		if (IsNodeEnabled() == false && Schema->PinToTypeDefinition(InputPin) != FNiagaraTypeDefinition::GetParameterMapDef())
+		{
+			continue;
+		}
+
 		int32 CompiledInput = Translator->CompilePin(InputPin);
 		if (CompiledInput == INDEX_NONE)
 		{
@@ -159,6 +203,12 @@ void UNiagaraNodeParameterMapSet::BuildParameterMapHistory(FNiagaraParameterMapH
 
 		OutHistory.VisitInputPin(InputPins[i], this);
 
+
+		if (!IsNodeEnabled() && OutHistory.GetIgnoreDisabled())
+		{
+			continue;
+		}
+
 		FNiagaraTypeDefinition VarTypeDef = Schema->PinToTypeDefinition(InputPins[i]);
 		if (i == 0 && InputPins[i] != nullptr && VarTypeDef == FNiagaraTypeDefinition::GetParameterMapDef())
 		{
@@ -180,7 +230,42 @@ void UNiagaraNodeParameterMapSet::BuildParameterMapHistory(FNiagaraParameterMapH
 		}
 	}
 
+	if (!IsNodeEnabled() && OutHistory.GetIgnoreDisabled())
+	{
+		RouteParameterMapAroundMe(OutHistory, bRecursive);
+		return;
+	}
+
 	OutHistory.RegisterParameterMapPin(ParamMapIdx, GetOutputPin(0));
 }
 
+void UNiagaraNodeParameterMapSet::GetContextMenuActions(const FGraphNodeContextMenuBuilder& Context) const
+{
+	UNiagaraNodeParameterMapBase::GetContextMenuActions(Context);
+
+	UEdGraphPin* Pin = const_cast<UEdGraphPin*>(Context.Pin);
+	if (Pin && Pin->Direction == EEdGraphPinDirection::EGPD_Input)
+	{
+		FNiagaraVariable Var = CastChecked<UEdGraphSchema_Niagara>(GetSchema())->PinToNiagaraVariable(Pin);
+		const UNiagaraGraph* Graph = GetNiagaraGraph();
+
+		if (!FNiagaraConstants::IsNiagaraConstant(Var))
+		{
+			Context.MenuBuilder->BeginSection("EdGraphSchema_NiagaraMetaDataActions", LOCTEXT("EditPinMenuHeader", "Meta-Data"));
+			TSharedRef<SWidget> RenameWidget =
+				SNew(SBox)
+				.WidthOverride(100)
+				.Padding(FMargin(5, 0, 0, 0))
+				[
+					SNew(SEditableTextBox)
+					.Text_UObject(this, &UNiagaraNodeParameterMapBase::GetPinDescriptionText, Pin)
+					.OnTextCommitted_UObject(this, &UNiagaraNodeParameterMapBase::PinDescriptionTextCommitted, Pin)
+				];
+			Context.MenuBuilder->AddWidget(RenameWidget, LOCTEXT("DescMenuItem", "Description"));
+
+			Context.MenuBuilder->EndSection();
+		}
+	}
+}
+
 #undef LOCTEXT_NAMESPACE
diff --git a/Engine/Plugins/FX/Niagara/Source/NiagaraEditor/Private/NiagaraNodeParameterMapSet.h b/Engine/Plugins/FX/Niagara/Source/NiagaraEditor/Private/NiagaraNodeParameterMapSet.h
--- a/Engine/Plugins/FX/Niagara/Source/NiagaraEditor/Private/NiagaraNodeParameterMapSet.h
+++ b/Engine/Plugins/FX/Niagara/Source/NiagaraEditor/Private/NiagaraNodeParameterMapSet.h
@@ -30,9 +30,12 @@ public:
 	virtual void Compile(class FHlslNiagaraTranslator* Translator, TArray<int32>& Outputs);
 
 	virtual FText GetNodeTitle(ENodeTitleType::Type TitleType) const;
+
+	virtual void GetContextMenuActions(const FGraphNodeContextMenuBuilder& Context) const override;
+
 protected:
 	virtual void OnNewTypedPinAdded(UEdGraphPin* NewPin); 
-	virtual void OnPinRenamed(UEdGraphPin* RenamedPin) override;
+	virtual void OnPinRenamed(UEdGraphPin* RenamedPin, const FString& OldName) override;
 	UEdGraphPin* PinPendingRename;
 
 };
diff --git a/Engine/Plugins/FX/Niagara/Source/NiagaraEditor/Private/NiagaraNodeWithDynamicPins.cpp b/Engine/Plugins/FX/Niagara/Source/NiagaraEditor/Private/NiagaraNodeWithDynamicPins.cpp
--- a/Engine/Plugins/FX/Niagara/Source/NiagaraEditor/Private/NiagaraNodeWithDynamicPins.cpp
+++ b/Engine/Plugins/FX/Niagara/Source/NiagaraEditor/Private/NiagaraNodeWithDynamicPins.cpp
@@ -91,6 +91,10 @@ UEdGraphPin* UNiagaraNodeWithDynamicPins::RequestNewTypedPin(EEdGraphPinDirectio
 
 void UNiagaraNodeWithDynamicPins::CreateAddPin(EEdGraphPinDirection Direction)
 {
+	if (!AllowDynamicPins())
+	{
+		return;
+	}
 	CreatePin(Direction, FEdGraphPinType(UEdGraphSchema_Niagara::PinCategoryMisc, AddPinSubCategory, nullptr, EPinContainerType::None, false, FEdGraphTerminalType()), TEXT("Add"));
 }
 
@@ -214,6 +218,7 @@ void UNiagaraNodeWithDynamicPins::GetContextMenuActions(const FGraphNodeContextM
 					FUIAction(FExecuteAction::CreateUObject(const_cast<UNiagaraNodeWithDynamicPins*>(this), &UNiagaraNodeWithDynamicPins::MoveDynamicPin, const_cast<UEdGraphPin*>(Context.Pin), 1)));
 			}
 		}
+		Context.MenuBuilder->EndSection();
 	}
 }
 
@@ -267,9 +272,9 @@ void UNiagaraNodeWithDynamicPins::PinNameTextCommitted(const FText& Text, ETextC
 	{
 		FScopedTransaction RemovePinTransaction(LOCTEXT("RenamePinTransaction", "Rename pin"));
 		Modify();
-
+		FString PinOldName = Pin->PinName.ToString();
 		Pin->PinName = *Text.ToString();
-		OnPinRenamed(Pin);
+		OnPinRenamed(Pin, PinOldName);
 	}
 }
 
diff --git a/Engine/Plugins/FX/Niagara/Source/NiagaraEditor/Private/NiagaraParameterMapHistory.cpp b/Engine/Plugins/FX/Niagara/Source/NiagaraEditor/Private/NiagaraParameterMapHistory.cpp
--- a/Engine/Plugins/FX/Niagara/Source/NiagaraEditor/Private/NiagaraParameterMapHistory.cpp
+++ b/Engine/Plugins/FX/Niagara/Source/NiagaraEditor/Private/NiagaraParameterMapHistory.cpp
@@ -5,7 +5,7 @@
 #include "AssetData.h"
 #include "NiagaraHlslTranslator.h"
 #include "NiagaraSystem.h"
-
+#include "NiagaraGraph.h"
 #include "ActorFactoryNiagara.h"
 #include "NiagaraActor.h"
 #include "EdGraphSchema_Niagara.h"
@@ -17,9 +17,11 @@
 #include "NiagaraNodeParameterMapGet.h"
 #include "ViewModels/Stack/NiagaraStackEntry.h"
 #include "ViewModels/Stack/NiagaraStackGraphUtilities.h"
+#include "NiagaraParameterCollection.h"
 
 #define LOCTEXT_NAMESPACE "NiagaraEditor"
 
+#define PARAM_MAP_NPC_STR TEXT("NPC.")
 #define PARAM_MAP_ENGINE_STR TEXT("Engine.")
 #define PARAM_MAP_USER_STR TEXT("User.")
 #define PARAM_MAP_SYSTEM_STR TEXT("System.")
@@ -28,6 +30,7 @@
 #define PARAM_MAP_ATTRIBUTE_STR TEXT("Particles.")
 #define PARAM_MAP_INITIAL_STR TEXT("Initial.")
 #define PARAM_MAP_INITIAL_BASE_STR TEXT("Initial")
+#define PARAM_MAP_RAPID_ITERATION_STR TEXT("Constants.")
 
 
 FNiagaraParameterMapHistory::FNiagaraParameterMapHistory() 
@@ -51,6 +54,7 @@ int32 FNiagaraParameterMapHistory::FindVariableByName(const FName& VariableName)
 	return FoundIdx;
 }
 
+
 int32 FNiagaraParameterMapHistory::FindVariable(const FName& VariableName, const FNiagaraTypeDefinition& Type)
 {
 	int32 FoundIdx = Variables.IndexOfByPredicate([&](const FNiagaraVariable& InObj) -> bool
@@ -61,7 +65,7 @@ int32 FNiagaraParameterMapHistory::FindVariable(const FName& VariableName, const
 	return FoundIdx;
 }
 
-int32 FNiagaraParameterMapHistory::AddVariable(const FNiagaraVariable& InVar, const UEdGraphPin* InPin)
+int32 FNiagaraParameterMapHistory::AddVariable(const FNiagaraVariable& InVar, const FNiagaraVariable& InAliasedVar, const UEdGraphPin* InPin)
 {
 	FNiagaraVariable Var = InVar;
 	
@@ -69,6 +73,7 @@ int32 FNiagaraParameterMapHistory::AddVariable(const FNiagaraVariable& InVar, co
 	if (FoundIdx == -1)
 	{
 		FoundIdx = Variables.Add(Var);
+		VariablesWithOriginalAliasesIntact.Add(InAliasedVar);
 		PerVariableWarnings.AddDefaulted(1);
 		PerVariableWriteHistory.AddDefaulted(1);
 		PerVariableReadHistory.AddDefaulted(1);
@@ -98,7 +103,7 @@ int32 FNiagaraParameterMapHistory::AddVariable(const FNiagaraVariable& InVar, co
 
 int32 FNiagaraParameterMapHistory::AddExternalVariable(const FNiagaraVariable& Var)
 {
-	return AddVariable(Var, nullptr);
+	return AddVariable(Var, Var, nullptr);
 }
 
 const UEdGraphPin* FNiagaraParameterMapHistory::GetFinalPin() const
@@ -121,29 +126,7 @@ const UEdGraphPin* FNiagaraParameterMapHistory::GetOriginalPin() const
 
 FNiagaraVariable FNiagaraParameterMapHistory::ResolveAliases(const FNiagaraVariable& InVar, const TMap<FString, FString>& InAliases, const TCHAR* InJoinSeparator)
 {
-	FNiagaraVariable OutVar = InVar;
-
-	FString OutVarStrName = InVar.GetName().ToString();
-	TArray<FString> SplitName;
-	OutVarStrName.ParseIntoArray(SplitName, TEXT("."));
-
-	for (int32 i = 0; i < SplitName.Num() - 1; i++)
-	{
-		TMap<FString, FString>::TConstIterator It = InAliases.CreateConstIterator();
-		while (It)
-		{
-			if (SplitName[i].Equals(It.Key()))
-			{
-				SplitName[i] = It.Value();
-			}
-			++It;
-		}
-	}
-
-	OutVarStrName = FString::Join<FString>(SplitName, InJoinSeparator);
-
-	OutVar.SetName(*OutVarStrName);
-	return OutVar;
+	return FNiagaraVariable::ResolveAliases(InVar, InAliases, InJoinSeparator);
 }
 
 FString FNiagaraParameterMapHistory::MakeSafeNamespaceString(const FString& InStr)
@@ -204,7 +187,14 @@ FNiagaraVariable FNiagaraParameterMapHistory::VariableToNamespacedVariable(const
 
 bool FNiagaraParameterMapHistory::IsInNamespace(const FNiagaraVariable& InVar, const FString& Namespace)
 {
-	return InVar.GetName().ToString().StartsWith(Namespace);
+	if (Namespace.EndsWith(TEXT(".")))
+	{
+		return InVar.GetName().ToString().StartsWith(Namespace);
+	}
+	else
+	{
+		return InVar.GetName().ToString().StartsWith(Namespace + TEXT("."));
+	}
 }
 
 bool FNiagaraParameterMapHistory::IsAliasedModuleParameter(const FNiagaraVariable& InVar)
@@ -232,6 +222,24 @@ bool FNiagaraParameterMapHistory::IsUserParameter(const FNiagaraVariable& InVar)
 	return IsInNamespace(InVar, PARAM_MAP_USER_STR);
 }
 
+bool FNiagaraParameterMapHistory::IsRapidIterationParameter(const FNiagaraVariable& InVar)
+{
+	return IsInNamespace(InVar, PARAM_MAP_RAPID_ITERATION_STR);
+}
+
+bool FNiagaraParameterMapHistory::TryGetEmitterAndFunctionCallNamesFromRapidIterationParameter(const FNiagaraVariable& InVar, FString& EmitterName, FString& FunctionCallName)
+{
+	TArray<FString> SplitName;
+	InVar.GetName().ToString().ParseIntoArray(SplitName, TEXT("."));
+	if (SplitName.Num() >= 4 && (SplitName[0] + ".") == PARAM_MAP_RAPID_ITERATION_STR)
+	{
+		EmitterName = SplitName[1];
+		FunctionCallName = SplitName[2];
+		return true;
+	}
+	return false;
+}
+
 bool FNiagaraParameterMapHistory::IsAttribute(const FNiagaraVariable& InVar)
 {
 	return IsInNamespace(InVar, PARAM_MAP_ATTRIBUTE_STR);
@@ -328,6 +336,48 @@ FNiagaraVariable FNiagaraParameterMapHistory::GetSourceForInitialValue(const FNi
 	return Var;
 }
 
+const FNiagaraVariableMetaData* FNiagaraParameterMapHistory::GetMetaData(int32 VarIdx) const
+{
+	if (PerVariableWriteHistory[VarIdx].Num() > 0)
+	{
+		const UEdGraphPin* Pin = PerVariableWriteHistory[VarIdx][0];
+		if (Pin != nullptr)
+		{
+			return CastChecked<UNiagaraNode>(Pin->GetOwningNode())->GetNiagaraGraph()->GetMetaData(Variables[VarIdx]);
+		}
+	}
+	else if (PerVariableReadHistory[VarIdx].Num() > 0)
+	{
+		const UEdGraphPin* Pin = PerVariableReadHistory[VarIdx][0].Key;
+		if (Pin != nullptr)
+		{
+			return CastChecked<UNiagaraNode>(Pin->GetOwningNode())->GetNiagaraGraph()->GetMetaData(Variables[VarIdx]);
+		}
+	}
+	return nullptr;
+}
+
+FNiagaraVariableMetaData* FNiagaraParameterMapHistory::GetMetaData(int32 VarIdx)
+{
+	if (PerVariableWriteHistory[VarIdx].Num() > 0)
+	{
+		const UEdGraphPin* Pin = PerVariableWriteHistory[VarIdx][0];
+		if (Pin != nullptr)
+		{
+			return CastChecked<UNiagaraNode>(Pin->GetOwningNode())->GetNiagaraGraph()->GetMetaData(Variables[VarIdx]);
+		}
+	}
+	else if (PerVariableReadHistory[VarIdx].Num() > 0)
+	{
+		const UEdGraphPin* Pin = PerVariableReadHistory[VarIdx][0].Key;
+		if (Pin != nullptr)
+		{
+			return CastChecked<UNiagaraNode>(Pin->GetOwningNode())->GetNiagaraGraph()->GetMetaData(Variables[VarIdx]);
+		}
+	}
+	return nullptr;
+}
+
 bool FNiagaraParameterMapHistory::IsPrimaryDataSetOutput(const FNiagaraVariable& InVar, UNiagaraScript* InScript)
 {
 	return IsPrimaryDataSetOutput(InVar, InScript->GetUsage());
@@ -393,6 +443,12 @@ bool FNiagaraParameterMapHistory::IsExportableExternalConstant(const FNiagaraVar
 
 bool FNiagaraParameterMapHistory::IsExternalConstantNamespace(const FNiagaraVariable& InVar, UNiagaraScript* InScript)
 {
+	// Parameter collections are always constants
+	if (IsInNamespace(InVar, PARAM_MAP_NPC_STR))
+	{
+		return true;
+	}
+
 	// Engine parameters are never writable.
 	if (IsInNamespace(InVar, PARAM_MAP_ENGINE_STR))
 	{
@@ -443,11 +499,40 @@ const UNiagaraNodeOutput* FNiagaraParameterMapHistory::GetFinalOutputNode() cons
 	return nullptr;
 }
 
+FNiagaraVariable FNiagaraParameterMapHistory::ConvertVariableToRapidIterationConstantName(FNiagaraVariable InVar, const TCHAR* InEmitterName)
+{
+	FNiagaraVariable Var = InVar;
+	FString NewName = PARAM_MAP_RAPID_ITERATION_STR + Var.GetName().ToString();
+	FString EmitterName = FString(InEmitterName);
+	if (InEmitterName != nullptr && false == InVar.GetName().ToString().StartsWith(EmitterName + TEXT(".")))
+	{
+		NewName = PARAM_MAP_RAPID_ITERATION_STR + EmitterName + TEXT(".") + Var.GetName().ToString();
+	}
+	Var.SetName(*NewName);
+	return Var;
+}
+
+UNiagaraParameterCollection* FNiagaraParameterMapHistory::IsParameterCollectionParameter(FNiagaraVariable& InVar, bool& bMissingParameter)
+{
+	bMissingParameter = false;
+	FString VarName = InVar.GetName().ToString();
+	for (int32 i = 0; i < ParameterCollections.Num(); ++i)
+	{
+		if (VarName.StartsWith(ParameterCollectionNamespaces[i]))
+		{
+			bMissingParameter = !ParameterCollectionVariables[i].Contains(InVar);
+			return ParameterCollections[i];
+		}
+	}
+	return nullptr;
+}
+
 FNiagaraParameterMapHistoryBuilder::FNiagaraParameterMapHistoryBuilder()
 {
 	ContextuallyVisitedNodes.AddDefaulted(1);
 	PinToParameterMapIndices.AddDefaulted(1);
 	bFilterByScriptWhitelist = false;
+	bIgnoreDisabled = true;
 	FilterScriptType = ENiagaraScriptUsage::Function;
 }
 
@@ -612,6 +697,44 @@ void FNiagaraParameterMapHistoryBuilder::EndTranslation(UNiagaraEmitter* Emitter
 	EmitterNameContextStack.Reset();
 }
 
+UNiagaraNode* FNiagaraParameterMapHistoryBuilder::GetCallingContext() const
+{
+	if (CallingContext.Num() == 0)
+	{
+		return nullptr;
+	}
+	else
+	{
+		return CallingContext.Last();
+	}
+}
+
+bool FNiagaraParameterMapHistoryBuilder::InTopLevelFunctionCall(ENiagaraScriptUsage InFilterScriptType) const
+{
+	if (InFilterScriptType == ENiagaraScriptUsage::EmitterSpawnScript || InFilterScriptType == ENiagaraScriptUsage::EmitterUpdateScript || InFilterScriptType == ENiagaraScriptUsage::SystemSpawnScript || InFilterScriptType == ENiagaraScriptUsage::SystemUpdateScript)
+	{
+		if (CallingContext.Num() <= 1) // Handles top-level system graph and any function calls off of it.
+		{
+			return true;
+		}
+		else if (CallingContext.Num() <= 2 && Cast<UNiagaraNodeEmitter>(CallingContext[0]) != nullptr) // Handle a function call off of an emitter
+		{
+			return true;
+		}
+
+	}
+	else if (InFilterScriptType == ENiagaraScriptUsage::ParticleSpawnScript || InFilterScriptType == ENiagaraScriptUsage::ParticleEventScript || InFilterScriptType == ENiagaraScriptUsage::ParticleSpawnScriptInterpolated || InFilterScriptType == ENiagaraScriptUsage::ParticleUpdateScript)
+	{
+		if (CallingContext.Num() <= 1) // Handle a function call
+		{
+			return true;
+		}
+	}
+
+	return false;
+}
+
+
 void FNiagaraParameterMapHistoryBuilder::EnterFunction(const FString& InNodeName, UNiagaraScript* InScript, UNiagaraNode* Node)
 {
 	if (InScript != nullptr )
@@ -621,6 +744,10 @@ void FNiagaraParameterMapHistoryBuilder::EnterFunction(const FString& InNodeName
 		PinToParameterMapIndices.Emplace();
 		FunctionNameContextStack.Emplace(*InNodeName);
 		BuildCurrentAliases();
+		if (EncounteredFunctionNames.Num() != 0)
+		{
+			EncounteredFunctionNames.Last().AddUnique(InNodeName);
+		}
 		ContextuallyVisitedNodes.Emplace();
 	}
 }
@@ -661,6 +788,9 @@ void FNiagaraParameterMapHistoryBuilder::EnterEmitter(const FString& InEmitterNa
 	{
 		History.EmitterNamespacesEncountered.AddUnique(InEmitterName);
 	}
+	EncounteredEmitterNames.AddUnique(InEmitterName);
+	TArray<FString> EmptyFuncs;
+	EncounteredFunctionNames.Push(EmptyFuncs);
 	ContextuallyVisitedNodes.Emplace();
 }
 
@@ -670,6 +800,35 @@ void FNiagaraParameterMapHistoryBuilder::ExitEmitter(const FString& InEmitterNam
 	EmitterNameContextStack.Pop();
 	BuildCurrentAliases();
 	ContextuallyVisitedNodes.Pop();
+	EncounteredFunctionNames.Pop();
+}
+
+
+bool FNiagaraParameterMapHistoryBuilder::IsInEncounteredFunctionNamespace(FNiagaraVariable& InVar) const
+{
+	if (EncounteredFunctionNames.Num() != 0)
+	{
+		for (FString EncounteredNamespace : EncounteredFunctionNames.Last())
+		{
+			if (FNiagaraParameterMapHistory::IsInNamespace(InVar, EncounteredNamespace))
+			{
+				return true;
+			}
+		}
+	}
+	return false;
+}
+
+bool FNiagaraParameterMapHistoryBuilder::IsInEncounteredEmitterNamespace(FNiagaraVariable& InVar) const
+{
+	for (FString EmitterEncounteredNamespace : EncounteredEmitterNames)
+	{
+		if (FNiagaraParameterMapHistory::IsInNamespace(InVar, EmitterEncounteredNamespace))
+		{
+			return true;
+		}
+	}
+	return false;
 }
 
 /**
@@ -705,18 +864,21 @@ void FNiagaraParameterMapHistoryBuilder::VisitInputPin(const UEdGraphPin* Pin, c
 	{
 		for (int32 j = 0; j < Pin->LinkedTo.Num(); j++)
 		{
-			UNiagaraNode* Node = CastChecked<UNiagaraNode>(Pin->LinkedTo[j]->GetOwningNode());
-
-			if (!GetNodePreviouslyVisited(Node))
+			if (Pin->LinkedTo[j])
 			{
-				Node->BuildParameterMapHistory(*this, true);
-				RegisterNodeVisitation(Node);
-			}
+				UNiagaraNode* Node = CastChecked<UNiagaraNode>(Pin->LinkedTo[j]->GetOwningNode());
 
-			if (Schema->PinToTypeDefinition(Pin) == FNiagaraTypeDefinition::GetParameterMapDef())
-			{
-				int32 ParamMapIdx = TraceParameterMapOutputPin(Pin->LinkedTo[j]);
-				RegisterParameterMapPin(ParamMapIdx, Pin);
+				if (!GetNodePreviouslyVisited(Node))
+				{
+					Node->BuildParameterMapHistory(*this, true);
+					RegisterNodeVisitation(Node);
+				}
+
+				if (Schema->PinToTypeDefinition(Pin) == FNiagaraTypeDefinition::GetParameterMapDef())
+				{
+					int32 ParamMapIdx = TraceParameterMapOutputPin(Pin->LinkedTo[j]);
+					RegisterParameterMapPin(ParamMapIdx, Pin);
+				}
 			}
 		} 
 	}
@@ -760,9 +922,10 @@ int32 FNiagaraParameterMapHistoryBuilder::HandleVariableWrite(int32 ParamMapIdx,
 	{
 		return INDEX_NONE;
 	}
+	FNiagaraVariable AliasedVar = Var;
 	Var = ResolveAliases(Var);
 
-	return Histories[ParamMapIdx].AddVariable(Var, InPin);
+	return Histories[ParamMapIdx].AddVariable(Var,AliasedVar, InPin);
 }
 
 
@@ -772,12 +935,25 @@ int32 FNiagaraParameterMapHistoryBuilder::HandleVariableRead(int32 ParamMapIdx,
 	const UEdGraphSchema_Niagara* Schema = GetDefault<UEdGraphSchema_Niagara>();
 	FNiagaraVariable Var = Schema->PinToNiagaraVariable(InPin);
 
+	FNiagaraParameterMapHistory& History = Histories[ParamMapIdx];
+
 	if (!ShouldTrackVariable(Var))
 	{
 		return INDEX_NONE;
 	}
+	FNiagaraVariable AliasedVar = Var;
 	Var = ResolveAliases(Var);
 
+	//Track any parameter collections we're referencing.
+	if (UNiagaraParameterCollection* Collection = Schema->VariableIsFromParameterCollection(Var))
+	{
+		int32 Index = History.ParameterCollections.AddUnique(Collection);
+		History.ParameterCollectionNamespaces.SetNum(History.ParameterCollections.Num());
+		History.ParameterCollectionVariables.SetNum(History.ParameterCollections.Num());
+		History.ParameterCollectionNamespaces[Index] = Collection->GetFullNamespace();
+		History.ParameterCollectionVariables[Index] = Collection->GetParameters();
+	}
+
 	int32 FoundIdx = Histories[ParamMapIdx].FindVariable(Var.GetName(), Var.GetType());
 	if (FoundIdx == -1)
 	{
@@ -794,6 +970,8 @@ int32 FNiagaraParameterMapHistoryBuilder::HandleVariableRead(int32 ParamMapIdx,
 			if (FoundIdx == -1)
 			{
 				FoundIdx = Histories[ParamMapIdx].Variables.Add(Var);
+				Histories[ParamMapIdx].VariablesWithOriginalAliasesIntact.Add(AliasedVar);
+
 			}
 			Histories[ParamMapIdx].PerVariableWarnings.AddDefaulted(1);
 			Histories[ParamMapIdx].PerVariableWriteHistory.AddDefaulted(1);
diff --git a/Engine/Plugins/FX/Niagara/Source/NiagaraEditor/Private/NiagaraScriptFactory.cpp b/Engine/Plugins/FX/Niagara/Source/NiagaraEditor/Private/NiagaraScriptFactory.cpp
--- a/Engine/Plugins/FX/Niagara/Source/NiagaraEditor/Private/NiagaraScriptFactory.cpp
+++ b/Engine/Plugins/FX/Niagara/Source/NiagaraEditor/Private/NiagaraScriptFactory.cpp
@@ -12,6 +12,9 @@
 #include "NiagaraScriptFactoryNew.h"
 #include "NiagaraNodeOutput.h"
 #include "NiagaraEditorSettings.h"
+#include "IAssetTypeActions.h"
+#include "AssetToolsModule.h"
+#include "AssetTypeActions/AssetTypeActions_NiagaraScript.h"
 
 #define LOCTEXT_NAMESPACE "NiagaraScriptFactory"
 
@@ -33,8 +36,8 @@ UObject* UNiagaraScriptFactoryNew::FactoryCreateNew(UClass* Class, UObject* InPa
 	check(Settings);
 
 	UNiagaraScript* NewScript;
-
-	if (UNiagaraScript* Default = Cast<UNiagaraScript>(Settings->DefaultScript.TryLoad()))
+	const FStringAssetReference& SettingDefaultScript = GetDefaultScriptFromSettings(Settings);
+	if (UNiagaraScript* Default = Cast<UNiagaraScript>(SettingDefaultScript.TryLoad()))
 	{
 		NewScript = Cast<UNiagaraScript>(StaticDuplicateObject(Default, InParent, Name, Flags, Class));
 	}
@@ -47,6 +50,59 @@ UObject* UNiagaraScriptFactoryNew::FactoryCreateNew(UClass* Class, UObject* InPa
 	return NewScript;
 }
 
+const FStringAssetReference& UNiagaraScriptFactoryNew::GetDefaultScriptFromSettings(const UNiagaraEditorSettings* Settings)
+{
+	switch (GetScriptUsage())
+	{
+	case ENiagaraScriptUsage::Module:
+		if (Settings->DefaultModuleScript.IsValid())
+		{
+			return Settings->DefaultModuleScript;
+		}
+		break;
+	case ENiagaraScriptUsage::DynamicInput:
+		if (Settings->DefaultDynamicInputScript.IsValid())
+		{
+			return Settings->DefaultDynamicInputScript;
+		}
+		break;
+	case ENiagaraScriptUsage::Function:
+		if (Settings->DefaultFunctionScript.IsValid())
+		{
+			return Settings->DefaultFunctionScript;
+		}
+		break;
+	}
+	return Settings->DefaultScript;
+}
+
+FText UNiagaraScriptFactoryNew::GetDisplayName() const
+{
+	// Get the displayname for this factory from the action type.
+	FAssetToolsModule& AssetToolsModule = FModuleManager::LoadModuleChecked<FAssetToolsModule>("AssetTools");
+
+	UClass* LocalSupportedClass = GetSupportedClass();
+	if (LocalSupportedClass)
+	{
+		TArray<TWeakPtr<IAssetTypeActions>> AssetTypeActions = AssetToolsModule.Get().GetAssetTypeActionsListForClass(LocalSupportedClass);
+		for (auto& CurrentAssetTypeAction : AssetTypeActions)
+		{
+			TSharedPtr<FAssetTypeActions_NiagaraScript> AssetTypeAction = StaticCastSharedPtr<FAssetTypeActions_NiagaraScript>(CurrentAssetTypeAction.Pin());
+			if (AssetTypeAction.IsValid() && AssetTypeAction->GetTypeName() == GetAssetTypeActionName())
+			{
+				FText Name = AssetTypeAction->GetName();
+				if (!Name.IsEmpty())
+				{
+					return Name;
+				}
+			}
+		}
+	}
+
+	// Factories that have no supported class have no display name.
+	return FText();
+}
+
 void UNiagaraScriptFactoryNew::InitializeScript(UNiagaraScript* NewScript)
 {
 	if(NewScript != NULL)
@@ -87,4 +143,78 @@ void UNiagaraScriptFactoryNew::InitializeScript(UNiagaraScript* NewScript)
 	}
 }
 
+ENiagaraScriptUsage UNiagaraScriptFactoryNew::GetScriptUsage() const
+{
+	return ENiagaraScriptUsage::Module;
+}
+
+/************************************************************************/
+/* UNiagaraScriptModulesFactory											*/
+/************************************************************************/
+UNiagaraModuleScriptFactory::UNiagaraModuleScriptFactory(const FObjectInitializer& ObjectInitializer)
+	: Super(ObjectInitializer)
+{
+}
+
+FName UNiagaraModuleScriptFactory::GetNewAssetThumbnailOverride() const
+{
+	return TEXT("NiagaraEditor.Thumbnails.Modules");
+}
+
+FName UNiagaraModuleScriptFactory::GetAssetTypeActionName() const
+{
+	return FAssetTypeActions_NiagaraScriptModules::NiagaraModuleScriptName;
+}
+
+ENiagaraScriptUsage UNiagaraModuleScriptFactory::GetScriptUsage() const
+{
+	return ENiagaraScriptUsage::Module;
+}
+
+/************************************************************************/
+/* UNiagaraScriptFunctionsFactory										*/
+/************************************************************************/
+UNiagaraFunctionScriptFactory::UNiagaraFunctionScriptFactory(const FObjectInitializer& ObjectInitializer)
+	: Super(ObjectInitializer)
+{
+}
+
+FName UNiagaraFunctionScriptFactory::GetNewAssetThumbnailOverride() const
+{
+	return TEXT("NiagaraEditor.Thumbnails.Functions");
+}
+
+FName UNiagaraFunctionScriptFactory::GetAssetTypeActionName() const
+{
+	return FAssetTypeActions_NiagaraScriptFunctions::NiagaraFunctionScriptName;
+}
+
+ENiagaraScriptUsage UNiagaraFunctionScriptFactory::GetScriptUsage() const
+{
+	return ENiagaraScriptUsage::Function;
+}
+
+/************************************************************************/
+/* UNiagaraScriptDynamicInputsFactory									*/
+/************************************************************************/
+UNiagaraDynamicInputScriptFactory::UNiagaraDynamicInputScriptFactory(const FObjectInitializer& ObjectInitializer)
+	: Super(ObjectInitializer)
+{
+}
+
+FName UNiagaraDynamicInputScriptFactory::GetNewAssetThumbnailOverride() const
+{
+	return TEXT("NiagaraEditor.Thumbnails.DynamicInputs");
+}
+
+FName UNiagaraDynamicInputScriptFactory::GetAssetTypeActionName() const
+{
+	return FAssetTypeActions_NiagaraScriptDynamicInputs::NiagaraDynamicInputScriptName;
+}
+
+ENiagaraScriptUsage UNiagaraDynamicInputScriptFactory::GetScriptUsage() const
+{
+	return ENiagaraScriptUsage::DynamicInput;
+}
+
 #undef LOCTEXT_NAMESPACE
