Copying Orion-Staging @ 3748653 (Orion/Dev-General @ 3722124) to //UE4/Main

diff --git a/Engine/Plugins/FX/Niagara/Source/NiagaraEditor/Private/NiagaraScriptMergeManager.cpp b/Engine/Plugins/FX/Niagara/Source/NiagaraEditor/Private/NiagaraScriptMergeManager.cpp
new file mode 100644
--- /dev/null
+++ b/Engine/Plugins/FX/Niagara/Source/NiagaraEditor/Private/NiagaraScriptMergeManager.cpp
@@ -0,0 +1,1730 @@
+// Copyright 1998-2017 Epic Games, Inc. All Rights Reserved.
+
+#include "NiagaraScriptMergeManager.h"
+#include "NiagaraEmitter.h"
+#include "NiagaraScript.h"
+#include "NiagaraScriptSource.h"
+#include "NiagaraRendererProperties.h"
+#include "EdGraphSchema_Niagara.h"
+#include "NiagaraGraph.h"
+#include "NiagaraDataInterface.h"
+#include "NiagaraNodeInput.h"
+#include "NiagaraNodeOutput.h"
+#include "NiagaraNodeFunctionCall.h"
+#include "NiagaraNodeAssignment.h"
+#include "NiagaraNodeParameterMapGet.h"
+#include "NiagaraNodeParameterMapSet.h"
+#include "NiagaraEditorUtilities.h"
+#include "ViewModels/Stack/NiagaraStackGraphUtilities.h"
+#include "NiagaraEditorModule.h"
+#include "UObject/PropertyPortFlags.h"
+
+#include "EdGraph/EdGraphNode.h"
+#include "EdGraph/EdGraphPin.h"
+#include "Modules/ModuleManager.h"
+
+#define LOCTEXT_NAMESPACE "NiagaraScriptMergeManager"
+
+DECLARE_CYCLE_STAT(TEXT("Niagara - ScriptMergeManager - DiffEmitters"), STAT_NiagaraEditor_ScriptMergeManager_DiffEmitters, STATGROUP_NiagaraEditor);
+DECLARE_CYCLE_STAT(TEXT("Niagara - ScriptMergeManager - MergeEmitter"), STAT_NiagaraEditor_ScriptMergeManager_MergeEmitter, STATGROUP_NiagaraEditor);
+DECLARE_CYCLE_STAT(TEXT("Niagara - ScriptMergeManager - IsModuleInputDifferentFromBase"), STAT_NiagaraEditor_ScriptMergeManager_IsModuleInputDifferentFromBase, STATGROUP_NiagaraEditor);
+
+FNiagaraStackFunctionInputOverrideMergeAdapter::FNiagaraStackFunctionInputOverrideMergeAdapter(
+	FString InUniqueEmitterName,
+	UNiagaraScript& InOwningScript,
+	UNiagaraNodeFunctionCall& InOwningFunctionCallNode,
+	UEdGraphPin& InOverridePin
+)
+	: UniqueEmitterName(InUniqueEmitterName)
+	, OwningScript(&InOwningScript)
+	, OwningFunctionCallNode(&InOwningFunctionCallNode)
+	, OverridePin(&InOverridePin)
+	, DataValueObject(nullptr)
+{
+	
+	InputName = FNiagaraParameterHandle(OverridePin->PinName).GetName().ToString();
+	OverrideNode = CastChecked<UNiagaraNodeParameterMapSet>(OverridePin->GetOwningNode());
+
+	if (OverridePin->LinkedTo.Num() == 0)
+	{
+		LocalValueString = OverridePin->DefaultValue;
+	}
+	else if (OverridePin->LinkedTo.Num() == 1)
+	{
+		if (OverridePin->LinkedTo[0]->GetOwningNode()->IsA<UNiagaraNodeParameterMapGet>())
+		{
+			LinkedValueHandle = FNiagaraParameterHandle(OverridePin->LinkedTo[0]->PinName);
+		}
+		else if (OverridePin->LinkedTo[0]->GetOwningNode()->IsA<UNiagaraNodeInput>())
+		{
+			UNiagaraNodeInput* DataInputNode = CastChecked<UNiagaraNodeInput>(OverridePin->LinkedTo[0]->GetOwningNode());
+			DataValueObject = DataInputNode->GetDataInterface();
+		}
+		else if (OverridePin->LinkedTo[0]->GetOwningNode()->IsA<UNiagaraNodeFunctionCall>())
+		{
+			DynamicValueFunction = MakeShared<FNiagaraStackFunctionMergeAdapter>(UniqueEmitterName, *OwningScript.Get(), *CastChecked<UNiagaraNodeFunctionCall>(OverridePin->LinkedTo[0]->GetOwningNode()), INDEX_NONE);
+		}
+		else
+		{
+			UE_LOG(LogNiagaraEditor, Error, TEXT("Invalid Stack Graph - Unsupported input node connection. Owning Node: %s"), *OverrideNode->GetPathName());
+		}
+	}
+	else
+	{
+		UE_LOG(LogNiagaraEditor, Error, TEXT("Invalid Stack Graph - Input had multiple connections. Owning Node: %s"), *OverrideNode->GetPathName());
+	}
+}
+
+FNiagaraStackFunctionInputOverrideMergeAdapter::FNiagaraStackFunctionInputOverrideMergeAdapter(
+	FString InUniqueEmitterName,
+	UNiagaraScript& InOwningScript,
+	UNiagaraNodeFunctionCall& InOwningFunctionCallNode,
+	FString InInputName,
+	FNiagaraVariable InRapidIterationParameter
+)
+	: UniqueEmitterName(InUniqueEmitterName)
+	, OwningScript(&InOwningScript)
+	, OwningFunctionCallNode(&InOwningFunctionCallNode)
+	, InputName(InInputName)
+	, OverridePin(nullptr)
+	, LocalValueRapidIterationParameter(InRapidIterationParameter)
+	, DataValueObject(nullptr)
+{
+}
+
+UNiagaraScript* FNiagaraStackFunctionInputOverrideMergeAdapter::GetOwningScript() const
+{
+	return OwningScript.Get();
+}
+
+UNiagaraNodeFunctionCall* FNiagaraStackFunctionInputOverrideMergeAdapter::GetOwningFunctionCall() const
+{
+	return OwningFunctionCallNode.Get();
+}
+
+FString FNiagaraStackFunctionInputOverrideMergeAdapter::GetInputName() const
+{
+	return InputName;
+}
+
+UNiagaraNodeParameterMapSet* FNiagaraStackFunctionInputOverrideMergeAdapter::GetOverrideNode() const
+{
+	return OverrideNode.Get();
+}
+
+UEdGraphPin* FNiagaraStackFunctionInputOverrideMergeAdapter::GetOverridePin() const
+{
+	return OverridePin;
+}
+
+TOptional<FString> FNiagaraStackFunctionInputOverrideMergeAdapter::GetLocalValueString() const
+{
+	return LocalValueString;
+}
+
+TOptional<FNiagaraVariable> FNiagaraStackFunctionInputOverrideMergeAdapter::GetLocalValueRapidIterationParameter() const
+{
+	return LocalValueRapidIterationParameter;
+}
+
+TOptional<FNiagaraParameterHandle> FNiagaraStackFunctionInputOverrideMergeAdapter::GetLinkedValueHandle() const
+{
+	return LinkedValueHandle;
+}
+
+UNiagaraDataInterface* FNiagaraStackFunctionInputOverrideMergeAdapter::GetDataValueObject() const
+{
+	return DataValueObject;
+}
+
+TSharedPtr<FNiagaraStackFunctionMergeAdapter> FNiagaraStackFunctionInputOverrideMergeAdapter::GetDynamicValueFunction() const
+{
+	return DynamicValueFunction;
+}
+
+FNiagaraStackFunctionMergeAdapter::FNiagaraStackFunctionMergeAdapter(FString InEmitterUniqueName, UNiagaraScript& InOwningScript, UNiagaraNodeFunctionCall& InFunctionCallNode, int32 InStackIndex)
+{
+	UniqueEmitterName = InEmitterUniqueName;
+	OwningScript = &InOwningScript;
+	FunctionCallNode = &InFunctionCallNode;
+	StackIndex = InStackIndex;
+
+	TSet<FString> AliasedInputsAdded;
+	UNiagaraNodeParameterMapSet* OverrideNode = FNiagaraStackGraphUtilities::GetStackFunctionOverrideNode(*FunctionCallNode);
+	if (OverrideNode != nullptr)
+	{
+		TArray<UEdGraphPin*> OverridePins;
+		OverrideNode->GetInputPins(OverridePins);
+		for (UEdGraphPin* OverridePin : OverridePins)
+		{
+			if (OverridePin->PinType.PinCategory != UEdGraphSchema_Niagara::PinCategoryMisc &&
+				OverridePin->PinType.PinSubCategoryObject != FNiagaraTypeDefinition::GetParameterMapStruct())
+			{
+				FNiagaraParameterHandle InputHandle(OverridePin->PinName);
+				if (InputHandle.GetNamespace().ToString() == FunctionCallNode->GetFunctionName())
+				{
+					InputOverrides.Add(MakeShared<FNiagaraStackFunctionInputOverrideMergeAdapter>(UniqueEmitterName, *OwningScript.Get(), *FunctionCallNode.Get(), *OverridePin));
+					AliasedInputsAdded.Add(OverridePin->PinName.ToString());
+				}
+			}
+		}
+	}
+
+	FString RapidIterationParameterNamePrefix = TEXT("Constants." + UniqueEmitterName + ".");
+	TArray<FNiagaraVariable> RapidIterationParameters;
+	OwningScript->RapidIterationParameters.GetParameters(RapidIterationParameters);
+	for (const FNiagaraVariable& RapidIterationParameter : RapidIterationParameters)
+	{
+		FNiagaraParameterHandle AliasedInputHandle(*RapidIterationParameter.GetName().ToString().RightChop(RapidIterationParameterNamePrefix.Len()));
+		if (AliasedInputHandle.GetNamespace().ToString() == FunctionCallNode->GetFunctionName() && AliasedInputsAdded.Contains(AliasedInputHandle.GetParameterHandleString().ToString()) == false)
+		{
+			InputOverrides.Add(MakeShared<FNiagaraStackFunctionInputOverrideMergeAdapter>(UniqueEmitterName, *OwningScript.Get(), *FunctionCallNode.Get(), AliasedInputHandle.GetName().ToString(), RapidIterationParameter));
+		}
+	}
+}
+
+UNiagaraNodeFunctionCall* FNiagaraStackFunctionMergeAdapter::GetFunctionCallNode() const
+{
+	return FunctionCallNode.Get();
+}
+
+int32 FNiagaraStackFunctionMergeAdapter::GetStackIndex() const
+{
+	return StackIndex;
+}
+
+const TArray<TSharedRef<FNiagaraStackFunctionInputOverrideMergeAdapter>>& FNiagaraStackFunctionMergeAdapter::GetInputOverrides() const
+{
+	return InputOverrides;
+}
+
+TSharedPtr<FNiagaraStackFunctionInputOverrideMergeAdapter> FNiagaraStackFunctionMergeAdapter::GetInputOverrideByInputName(FString InputName) const
+{
+	for (TSharedPtr<FNiagaraStackFunctionInputOverrideMergeAdapter> InputOverride : InputOverrides)
+	{
+		if (InputOverride->GetInputName() == InputName)
+		{
+			return InputOverride;
+		}
+	}
+	return TSharedPtr<FNiagaraStackFunctionInputOverrideMergeAdapter>();
+}
+
+FNiagaraScriptStackMergeAdapter::FNiagaraScriptStackMergeAdapter(UNiagaraNodeOutput& InOutputNode, UNiagaraScript& InScript, FString InUniqueEmitterName)
+{
+	OutputNode = &InOutputNode;
+	Script = &InScript;
+	UniqueEmitterName = InUniqueEmitterName;
+
+	TArray<FNiagaraStackGraphUtilities::FStackNodeGroup> StackGroups;
+	FNiagaraStackGraphUtilities::GetStackNodeGroups(*OutputNode, StackGroups);
+
+	if (StackGroups.Num() > 2 && StackGroups[0].EndNode->IsA<UNiagaraNodeInput>() && StackGroups.Last().EndNode->IsA<UNiagaraNodeOutput>())
+	{
+		for (int i = 1; i < StackGroups.Num() - 1; i++)
+		{
+			UNiagaraNodeFunctionCall* ModuleFunctionCallNode = Cast<UNiagaraNodeFunctionCall>(StackGroups[i].EndNode);
+			if (ModuleFunctionCallNode != nullptr)
+			{
+				// The first stack node group is the input node, so we subtract one to get the index of the module.
+				int32 StackIndex = i - 1;
+				ModuleFunctions.Add(MakeShared<FNiagaraStackFunctionMergeAdapter>(UniqueEmitterName, *Script.Get(), *ModuleFunctionCallNode, StackIndex));
+			}
+		}
+	}
+}
+
+UNiagaraNodeOutput* FNiagaraScriptStackMergeAdapter::GetOutputNode() const
+{
+	return OutputNode.Get();
+}
+
+UNiagaraScript* FNiagaraScriptStackMergeAdapter::GetScript() const
+{
+	return Script.Get();
+}
+
+FString FNiagaraScriptStackMergeAdapter::GetUniqueEmitterName() const
+{
+	return UniqueEmitterName;
+}
+
+const TArray<TSharedRef<FNiagaraStackFunctionMergeAdapter>>& FNiagaraScriptStackMergeAdapter::GetModuleFunctions() const
+{
+	return ModuleFunctions;
+}
+
+TSharedPtr<FNiagaraStackFunctionMergeAdapter> FNiagaraScriptStackMergeAdapter::GetModuleFunctionById(FGuid FunctionCallNodeId) const
+{
+	for (TSharedRef<FNiagaraStackFunctionMergeAdapter> Modulefunction : ModuleFunctions)
+	{
+		if (Modulefunction->GetFunctionCallNode()->NodeGuid == FunctionCallNodeId)
+		{
+			return Modulefunction;
+		}
+	}
+	return TSharedPtr<FNiagaraStackFunctionMergeAdapter>();
+}
+
+FNiagaraEventHandlerMergeAdapter::FNiagaraEventHandlerMergeAdapter(const UNiagaraEmitter& InEmitter, const FNiagaraEventScriptProperties* InEventScriptProperties, UNiagaraNodeOutput* InOutputNode)
+{
+	Initialize(InEmitter, InEventScriptProperties, nullptr, InOutputNode);
+}
+
+FNiagaraEventHandlerMergeAdapter::FNiagaraEventHandlerMergeAdapter(const UNiagaraEmitter& InEmitter, FNiagaraEventScriptProperties* InEventScriptProperties, UNiagaraNodeOutput* InOutputNode)
+{
+	Initialize(InEmitter, InEventScriptProperties, InEventScriptProperties, InOutputNode);
+}
+
+FNiagaraEventHandlerMergeAdapter::FNiagaraEventHandlerMergeAdapter(const UNiagaraEmitter& InEmitter, UNiagaraNodeOutput* InOutputNode)
+{
+	Initialize(InEmitter, nullptr, nullptr, InOutputNode);
+}
+
+void FNiagaraEventHandlerMergeAdapter::Initialize(const UNiagaraEmitter& InEmitter, const FNiagaraEventScriptProperties* InEventScriptProperties, FNiagaraEventScriptProperties* InEditableEventScriptProperties, UNiagaraNodeOutput* InOutputNode)
+{
+	Emitter = &InEmitter;
+
+	EventScriptProperties = InEventScriptProperties;
+	EditableEventScriptProperties = InEditableEventScriptProperties;
+
+	OutputNode = InOutputNode;
+
+	if (EventScriptProperties != nullptr && OutputNode != nullptr)
+	{
+		EventStack = MakeShared<FNiagaraScriptStackMergeAdapter>(*OutputNode.Get(), *EventScriptProperties->Script, Emitter->GetUniqueEmitterName());
+	}
+}
+
+const UNiagaraEmitter* FNiagaraEventHandlerMergeAdapter::GetEmitter() const
+{
+	return Emitter.Get();
+}
+
+FGuid FNiagaraEventHandlerMergeAdapter::GetUsageId() const
+{
+	if (EventScriptProperties != nullptr)
+	{
+		return EventScriptProperties->Script->GetUsageId();
+	}
+	else
+	{
+		return OutputNode->GetUsageId();
+	}
+}
+
+const FNiagaraEventScriptProperties* FNiagaraEventHandlerMergeAdapter::GetEventScriptProperties() const
+{
+	return EventScriptProperties;
+}
+
+FNiagaraEventScriptProperties* FNiagaraEventHandlerMergeAdapter::GetEditableEventScriptProperties() const
+{
+	return EditableEventScriptProperties;
+}
+
+UNiagaraNodeOutput* FNiagaraEventHandlerMergeAdapter::GetOutputNode() const
+{
+	return OutputNode.Get();
+}
+
+TSharedPtr<FNiagaraScriptStackMergeAdapter> FNiagaraEventHandlerMergeAdapter::GetEventStack() const
+{
+	return EventStack;
+}
+
+FNiagaraRendererMergeAdapter::FNiagaraRendererMergeAdapter(UNiagaraRendererProperties& InRenderer)
+{
+	Renderer = &InRenderer;
+}
+
+UNiagaraRendererProperties* FNiagaraRendererMergeAdapter::GetRenderer()
+{
+	return Renderer.Get();
+}
+
+FNiagaraEmitterMergeAdapter::FNiagaraEmitterMergeAdapter(const UNiagaraEmitter& InEmitter)
+{
+	Initialize(InEmitter, nullptr);
+}
+
+FNiagaraEmitterMergeAdapter::FNiagaraEmitterMergeAdapter(UNiagaraEmitter& InEmitter)
+{
+	Initialize(InEmitter, &InEmitter);
+}
+
+void FNiagaraEmitterMergeAdapter::Initialize(const UNiagaraEmitter& InEmitter, UNiagaraEmitter* InEditableEmitter)
+{
+	Emitter = &InEmitter;
+	EditableEmitter = InEditableEmitter;
+	UNiagaraScriptSource* EmitterScriptSource = Cast<UNiagaraScriptSource>(Emitter->GraphSource);
+	UNiagaraGraph* Graph = EmitterScriptSource->NodeGraph;
+	TArray<UNiagaraNodeOutput*> OutputNodes;
+	Graph->GetNodesOfClass<UNiagaraNodeOutput>(OutputNodes);
+
+	TArray<UNiagaraNodeOutput*> EventOutputNodes;
+	for (UNiagaraNodeOutput* OutputNode : OutputNodes)
+	{
+		switch (OutputNode->GetUsage())
+		{
+		case ENiagaraScriptUsage::EmitterSpawnScript:
+			EmitterSpawnStack = MakeShared<FNiagaraScriptStackMergeAdapter>(*OutputNode, *Emitter->EmitterSpawnScriptProps.Script, Emitter->GetUniqueEmitterName());
+			break;
+		case ENiagaraScriptUsage::EmitterUpdateScript:
+			EmitterUpdateStack = MakeShared<FNiagaraScriptStackMergeAdapter>(*OutputNode, *Emitter->EmitterUpdateScriptProps.Script, Emitter->GetUniqueEmitterName());
+			break;
+		case ENiagaraScriptUsage::ParticleSpawnScript:
+			ParticleSpawnStack = MakeShared<FNiagaraScriptStackMergeAdapter>(*OutputNode, *Emitter->SpawnScriptProps.Script, Emitter->GetUniqueEmitterName());
+			break;
+		case ENiagaraScriptUsage::ParticleUpdateScript:
+			ParticleUpdateStack = MakeShared<FNiagaraScriptStackMergeAdapter>(*OutputNode, *Emitter->UpdateScriptProps.Script, Emitter->GetUniqueEmitterName());
+			break;
+		case ENiagaraScriptUsage::ParticleEventScript:
+			EventOutputNodes.Add(OutputNode);
+			break;
+		}
+	}
+
+	// Create an event handler adapter for each usage id even if it's missing an event script properties struct or an output node.  These
+	// incomplete adapters will be caught if they are diffed.
+	for (const FNiagaraEventScriptProperties& EventScriptProperties : Emitter->GetEventHandlers())
+	{
+		UNiagaraNodeOutput** MatchingOutputNodePtr = EventOutputNodes.FindByPredicate(
+			[=](UNiagaraNodeOutput* EventOutputNode) { return EventOutputNode->GetUsageId() == EventScriptProperties.Script->GetUsageId(); });
+
+		UNiagaraNodeOutput* MatchingOutputNode = MatchingOutputNodePtr != nullptr ? *MatchingOutputNodePtr : nullptr;
+
+		if (EditableEmitter == nullptr)
+		{
+			EventHandlers.Add(MakeShared<FNiagaraEventHandlerMergeAdapter>(*Emitter.Get(), &EventScriptProperties, MatchingOutputNode));
+		}
+		else
+		{
+			FNiagaraEventScriptProperties* EditableEventScriptProperties = EditableEmitter->GetEventHandlerByIdUnsafe(EventScriptProperties.Script->GetUsageId());
+			EventHandlers.Add(MakeShared<FNiagaraEventHandlerMergeAdapter>(*Emitter.Get(), EditableEventScriptProperties, MatchingOutputNode));
+		}
+
+		if (MatchingOutputNode != nullptr)
+		{
+			EventOutputNodes.Remove(MatchingOutputNode);
+		}
+	}
+
+	for (UNiagaraNodeOutput* EventOutputNode : EventOutputNodes)
+	{
+		EventHandlers.Add(MakeShared<FNiagaraEventHandlerMergeAdapter>(*Emitter.Get(), EventOutputNode));
+	}
+
+	// Renderers
+	for (UNiagaraRendererProperties* RendererProperties : Emitter->GetRenderers())
+	{
+		Renderers.Add(MakeShared<FNiagaraRendererMergeAdapter>(*RendererProperties));
+	}
+}
+
+UNiagaraEmitter* FNiagaraEmitterMergeAdapter::GetEditableEmitter() const
+{
+	return EditableEmitter.Get();
+}
+
+TSharedPtr<FNiagaraScriptStackMergeAdapter> FNiagaraEmitterMergeAdapter::GetEmitterSpawnStack() const
+{
+	return EmitterSpawnStack;
+}
+
+TSharedPtr<FNiagaraScriptStackMergeAdapter> FNiagaraEmitterMergeAdapter::GetEmitterUpdateStack() const
+{
+	return EmitterUpdateStack;
+}
+
+TSharedPtr<FNiagaraScriptStackMergeAdapter> FNiagaraEmitterMergeAdapter::GetParticleSpawnStack() const
+{
+	return ParticleSpawnStack;
+}
+
+TSharedPtr<FNiagaraScriptStackMergeAdapter> FNiagaraEmitterMergeAdapter::GetParticleUpdateStack() const
+{
+	return ParticleUpdateStack;
+}
+
+const TArray<TSharedRef<FNiagaraEventHandlerMergeAdapter>> FNiagaraEmitterMergeAdapter::GetEventHandlers() const
+{
+	return EventHandlers;
+}
+
+const TArray<TSharedRef<FNiagaraRendererMergeAdapter>> FNiagaraEmitterMergeAdapter::GetRenderers() const
+{
+	return Renderers;
+}
+
+TSharedPtr<FNiagaraScriptStackMergeAdapter> FNiagaraEmitterMergeAdapter::GetScriptStack(ENiagaraScriptUsage Usage, FGuid ScriptUsageId)
+{
+	switch (Usage)
+	{
+	case ENiagaraScriptUsage::EmitterSpawnScript:
+		return EmitterSpawnStack;
+	case ENiagaraScriptUsage::EmitterUpdateScript:
+		return EmitterUpdateStack;
+	case ENiagaraScriptUsage::ParticleSpawnScript:
+		return ParticleSpawnStack;
+	case ENiagaraScriptUsage::ParticleUpdateScript:
+		return ParticleUpdateStack;
+	case ENiagaraScriptUsage::ParticleEventScript:
+		for (TSharedPtr<FNiagaraEventHandlerMergeAdapter> EventHandler : EventHandlers)
+		{
+			if (EventHandler->GetUsageId() == ScriptUsageId)
+			{
+				return EventHandler->GetEventStack();
+			}
+		}
+		break;
+	default:
+		checkf(false, TEXT("Unsupported usage"));
+	}
+
+	return TSharedPtr<FNiagaraScriptStackMergeAdapter>();
+}
+
+TSharedPtr<FNiagaraEventHandlerMergeAdapter> FNiagaraEmitterMergeAdapter::GetEventHandler(FGuid EventScriptUsageId)
+{
+	for (TSharedRef<FNiagaraEventHandlerMergeAdapter> EventHandler : EventHandlers)
+	{
+		if (EventHandler->GetUsageId() == EventScriptUsageId)
+		{
+			return EventHandler;
+		}
+	}
+	return TSharedPtr<FNiagaraEventHandlerMergeAdapter>();
+}
+
+TSharedPtr<FNiagaraRendererMergeAdapter> FNiagaraEmitterMergeAdapter::GetRenderer(FGuid RendererMergeId)
+{
+	for (TSharedRef<FNiagaraRendererMergeAdapter> Renderer : Renderers)
+	{
+		if (Renderer->GetRenderer()->GetMergeId() == RendererMergeId)
+		{
+			return Renderer;
+		}
+	}
+	return TSharedPtr<FNiagaraRendererMergeAdapter>();
+}
+
+FNiagaraScriptStackDiffResults::FNiagaraScriptStackDiffResults()
+	: bIsValid(true)
+{
+}
+
+bool FNiagaraScriptStackDiffResults::IsEmpty() const
+{
+	return
+		RemovedBaseModules.Num() == 0 &&
+		AddedOtherModules.Num() == 0 &&
+		MovedBaseModules.Num() == 0 &&
+		MovedOtherModules.Num() == 0 &&
+		EnabledChangedBaseModules.Num() == 0 &&
+		EnabledChangedOtherModules.Num() == 0 &&
+		RemovedBaseInputOverrides.Num() == 0 &&
+		AddedOtherInputOverrides.Num() == 0 &&
+		ModifiedOtherInputOverrides.Num() == 0;
+}
+
+bool FNiagaraScriptStackDiffResults::IsValid() const
+{
+	return bIsValid;
+}
+
+void FNiagaraScriptStackDiffResults::AddError(FText ErrorMessage)
+{
+	ErrorMessages.Add(ErrorMessage);
+	bIsValid = false;
+}
+
+const TArray<FText>& FNiagaraScriptStackDiffResults::GetErrorMessages() const
+{
+	return ErrorMessages;
+}
+
+FNiagaraEmitterDiffResults::FNiagaraEmitterDiffResults()
+	: bIsValid(true)
+{
+}
+
+bool FNiagaraEmitterDiffResults::IsValid() const
+{
+	return bIsValid &&
+		EmitterSpawnDiffResults.IsValid() &&
+		EmitterUpdateDiffResults.IsValid() &&
+		ParticleSpawnDiffResults.IsValid() &&
+		ParticleUpdateDiffResults.IsValid();
+}
+
+bool FNiagaraEmitterDiffResults::IsEmpty() const
+{
+	return DifferentEmitterProperties.Num() == 0 &&
+		EmitterSpawnDiffResults.IsEmpty() &&
+		EmitterUpdateDiffResults.IsEmpty() &&
+		ParticleSpawnDiffResults.IsEmpty() &&
+		ParticleUpdateDiffResults.IsEmpty() &&
+		RemovedBaseEventHandlers.Num() == 0 &&
+		AddedOtherEventHandlers.Num() == 0 &&
+		ModifiedEventHandlers.Num() == 0 &&
+		RemovedBaseRenderers.Num() == 0 &&
+		AddedOtherRenderers.Num() == 0 &&
+		ModifiedBaseRenderers.Num() == 0 &&
+		ModifiedOtherRenderers.Num() == 0;
+}
+
+void FNiagaraEmitterDiffResults::AddError(FText ErrorMessage)
+{
+	ErrorMessages.Add(ErrorMessage);
+	bIsValid = false;
+}
+
+const TArray<FText>& FNiagaraEmitterDiffResults::GetErrorMessages() const
+{
+	return ErrorMessages;
+}
+
+FString FNiagaraEmitterDiffResults::GetErrorMessagesString() const
+{
+	TArray<FString> ErrorMessageStrings;
+	for (FText ErrorMessage : ErrorMessages)
+	{
+		ErrorMessageStrings.Add(ErrorMessage.ToString());
+	}
+	return FString::Join(ErrorMessageStrings, TEXT("\n"));
+}
+
+TSharedRef<FNiagaraScriptMergeManager> FNiagaraScriptMergeManager::Get()
+{
+	FNiagaraEditorModule& NiagaraEditorModule = FModuleManager::GetModuleChecked<FNiagaraEditorModule>("NiagaraEditor");
+	return NiagaraEditorModule.GetScriptMergeManager();
+}
+
+INiagaraModule::FMergeEmitterResults FNiagaraScriptMergeManager::MergeEmitter(UNiagaraEmitter& Source, UNiagaraEmitter& LastMergedSource, UNiagaraEmitter& Instance)
+{
+	SCOPE_CYCLE_COUNTER(STAT_NiagaraEditor_ScriptMergeManager_MergeEmitter);
+	INiagaraModule::FMergeEmitterResults MergeResults;
+	FNiagaraEmitterDiffResults DiffResults = DiffEmitters(LastMergedSource, Instance);
+	if (DiffResults.IsValid())
+	{
+		UNiagaraEmitter* MergedInstance = CastChecked<UNiagaraEmitter>(StaticDuplicateObject(&Source, (UObject*)GetTransientPackage()));
+		TSharedRef<FNiagaraEmitterMergeAdapter> MergedInstanceAdapter = MakeShared<FNiagaraEmitterMergeAdapter>(*MergedInstance);
+
+		FApplyDiffResults EmitterSpawnResults = ApplyScriptStackDiff(MergedInstanceAdapter->GetEmitterSpawnStack().ToSharedRef(), DiffResults.EmitterSpawnDiffResults);
+		MergeResults.bSucceeded &= EmitterSpawnResults.bSucceeded;
+		MergeResults.bModifiedGraph |= EmitterSpawnResults.bModifiedGraph;
+		MergeResults.ErrorMessages.Append(EmitterSpawnResults.ErrorMessages);
+
+		FApplyDiffResults EmitterUpdateResults = ApplyScriptStackDiff(MergedInstanceAdapter->GetEmitterUpdateStack().ToSharedRef(), DiffResults.EmitterUpdateDiffResults);
+		MergeResults.bSucceeded &= EmitterUpdateResults.bSucceeded;
+		MergeResults.bModifiedGraph |= EmitterUpdateResults.bModifiedGraph;
+		MergeResults.ErrorMessages.Append(EmitterUpdateResults.ErrorMessages);
+
+		FApplyDiffResults ParticleSpawnResults = ApplyScriptStackDiff(MergedInstanceAdapter->GetParticleSpawnStack().ToSharedRef(), DiffResults.ParticleSpawnDiffResults);
+		MergeResults.bSucceeded &= ParticleSpawnResults.bSucceeded;
+		MergeResults.bModifiedGraph |= ParticleSpawnResults.bModifiedGraph;
+		MergeResults.ErrorMessages.Append(ParticleSpawnResults.ErrorMessages);
+
+		FApplyDiffResults ParticleUpdateResults = ApplyScriptStackDiff(MergedInstanceAdapter->GetParticleUpdateStack().ToSharedRef(), DiffResults.ParticleUpdateDiffResults);
+		MergeResults.bSucceeded &= ParticleUpdateResults.bSucceeded;
+		MergeResults.bModifiedGraph |= ParticleUpdateResults.bModifiedGraph;
+		MergeResults.ErrorMessages.Append(ParticleUpdateResults.ErrorMessages);
+
+		FApplyDiffResults EventHandlerResults = ApplyEventHandlerDiff(MergedInstanceAdapter, DiffResults);
+		MergeResults.bSucceeded &= EventHandlerResults.bSucceeded;
+		MergeResults.bModifiedGraph |= EventHandlerResults.bModifiedGraph;
+		MergeResults.ErrorMessages.Append(EventHandlerResults.ErrorMessages);
+
+		FApplyDiffResults RendererResults = ApplyRendererDiff(*MergedInstance, DiffResults);
+		MergeResults.bSucceeded &= RendererResults.bSucceeded;
+		MergeResults.bModifiedGraph |= RendererResults.bModifiedGraph;
+		MergeResults.ErrorMessages.Append(RendererResults.ErrorMessages);
+
+		CopyPropertiesToBase(MergedInstance, &Instance, DiffResults.DifferentEmitterProperties);
+
+		FNiagaraStackGraphUtilities::CleanUpStaleRapidIterationParameters(*MergedInstance);
+
+		if (MergeResults.bSucceeded)
+		{
+			UNiagaraScriptSource* ScriptSource = Cast<UNiagaraScriptSource>(MergedInstance->GraphSource);
+			FNiagaraStackGraphUtilities::RelayoutGraph(*ScriptSource->NodeGraph);
+			MergeResults.MergedInstance = MergedInstance;
+		}
+	}
+	else
+	{
+		MergeResults.bSucceeded = false;
+		MergeResults.ErrorMessages = DiffResults.GetErrorMessages();
+	}
+	return MergeResults;
+}
+
+bool FNiagaraScriptMergeManager::IsMergeableScriptUsage(ENiagaraScriptUsage ScriptUsage) const
+{
+	return ScriptUsage == ENiagaraScriptUsage::EmitterSpawnScript ||
+		ScriptUsage == ENiagaraScriptUsage::EmitterUpdateScript ||
+		ScriptUsage == ENiagaraScriptUsage::ParticleSpawnScript ||
+		ScriptUsage == ENiagaraScriptUsage::ParticleUpdateScript ||
+		ScriptUsage == ENiagaraScriptUsage::ParticleEventScript;
+}
+
+bool FNiagaraScriptMergeManager::HasBaseModule(const UNiagaraEmitter& BaseEmitter, ENiagaraScriptUsage ScriptUsage, FGuid ScriptUsageId, FGuid ModuleId)
+{
+	TSharedRef<FNiagaraEmitterMergeAdapter> BaseEmitterAdapter = GetEmitterMergeAdapterUsingCache(BaseEmitter);
+	TSharedPtr<FNiagaraScriptStackMergeAdapter> BaseScriptStackAdapter = BaseEmitterAdapter->GetScriptStack(ScriptUsage, ScriptUsageId);
+	return BaseScriptStackAdapter.IsValid() && BaseScriptStackAdapter->GetModuleFunctionById(ModuleId).IsValid();
+}
+
+bool FNiagaraScriptMergeManager::IsModuleInputDifferentFromBase(UNiagaraEmitter& Emitter, const UNiagaraEmitter& BaseEmitter, ENiagaraScriptUsage ScriptUsage, FGuid ScriptUsageId, FGuid ModuleId, FString InputName)
+{
+	SCOPE_CYCLE_COUNTER(STAT_NiagaraEditor_ScriptMergeManager_IsModuleInputDifferentFromBase);
+
+	TSharedRef<FNiagaraEmitterMergeAdapter> EmitterAdapter = GetEmitterMergeAdapterUsingCache(Emitter);
+	TSharedRef<FNiagaraEmitterMergeAdapter> BaseEmitterAdapter = GetEmitterMergeAdapterUsingCache(BaseEmitter);
+
+	TSharedRef<FNiagaraScriptStackMergeAdapter> ScriptStackAdapter = EmitterAdapter->GetScriptStack(ScriptUsage, ScriptUsageId).ToSharedRef();
+	TSharedPtr<FNiagaraScriptStackMergeAdapter> BaseScriptStackAdapter = BaseEmitterAdapter->GetScriptStack(ScriptUsage, ScriptUsageId);
+
+	if (BaseScriptStackAdapter.IsValid() == false)
+	{
+		return false;
+	}
+
+	FNiagaraScriptStackDiffResults ScriptStackDiffResults;
+	DiffScriptStacks(BaseScriptStackAdapter.ToSharedRef(), ScriptStackAdapter, ScriptStackDiffResults);
+
+	if (ScriptStackDiffResults.IsValid() == false)
+	{
+		return true;
+	}
+
+	if (ScriptStackDiffResults.IsEmpty())
+	{
+		return false;
+	}
+
+	auto FindInputOverrideByInputName = [=](TSharedRef<FNiagaraStackFunctionInputOverrideMergeAdapter> InputOverride)
+	{
+		return InputOverride->GetOwningFunctionCall()->NodeGuid == ModuleId && InputOverride->GetInputName() == InputName;
+	};
+
+	return
+		ScriptStackDiffResults.RemovedBaseInputOverrides.FindByPredicate(FindInputOverrideByInputName) != nullptr ||
+		ScriptStackDiffResults.AddedOtherInputOverrides.FindByPredicate(FindInputOverrideByInputName) != nullptr ||
+		ScriptStackDiffResults.ModifiedOtherInputOverrides.FindByPredicate(FindInputOverrideByInputName) != nullptr;
+}
+
+FNiagaraScriptMergeManager::FApplyDiffResults FNiagaraScriptMergeManager::ResetModuleInputToBase(UNiagaraEmitter& Emitter, const UNiagaraEmitter& BaseEmitter, ENiagaraScriptUsage ScriptUsage, FGuid ScriptUsageId, FGuid ModuleId, FString InputName)
+{
+	TSharedRef<FNiagaraEmitterMergeAdapter> EmitterAdapter = GetEmitterMergeAdapterUsingCache(Emitter);
+	TSharedRef<FNiagaraEmitterMergeAdapter> BaseEmitterAdapter = GetEmitterMergeAdapterUsingCache(BaseEmitter);
+
+	// Diff from the emitter to the base to create a diff which will reset the emitter back to the base.
+	FNiagaraScriptStackDiffResults ResetDiffResults;
+	DiffScriptStacks(EmitterAdapter->GetScriptStack(ScriptUsage, ScriptUsageId).ToSharedRef(), BaseEmitterAdapter->GetScriptStack(ScriptUsage, ScriptUsageId).ToSharedRef(), ResetDiffResults);
+
+	if (ResetDiffResults.IsValid() == false)
+	{
+		FApplyDiffResults Results;
+		Results.bSucceeded = false;
+		Results.bModifiedGraph = false;
+		Results.ErrorMessages.Add(FText::Format(LOCTEXT("ResetFailedBecauseOfDiffMessage", "Failed to reset input back to it's base value.  It couldn't be diffed successfully.  Emitter: {0}  Input:{1}"),
+			FText::FromString(Emitter.GetPathName()), FText::FromString(InputName)));
+		return Results;
+	}
+
+	if (ResetDiffResults.IsEmpty())
+	{
+		FApplyDiffResults Results;
+		Results.bSucceeded = false;
+		Results.bModifiedGraph = false;
+		Results.ErrorMessages.Add(FText::Format(LOCTEXT("ResetFailedBecauseOfEmptyDiffMessage", "Failed to reset input back to it's base value.  It wasn't different from the base.  Emitter: {0}  Input:{1}"),
+			FText::FromString(Emitter.GetPathName()), FText::FromString(InputName)));
+		return Results;
+	}
+	
+	// Remove items from the diff which are not relevant to this input.
+	ResetDiffResults.RemovedBaseModules.Empty();
+	ResetDiffResults.AddedOtherModules.Empty();
+
+	auto FindUnrelatedInputOverrides = [=](TSharedRef<FNiagaraStackFunctionInputOverrideMergeAdapter> InputOverride)
+	{
+		return InputOverride->GetOwningFunctionCall()->NodeGuid != ModuleId || InputOverride->GetInputName() != InputName;
+	};
+
+	ResetDiffResults.RemovedBaseInputOverrides.RemoveAll(FindUnrelatedInputOverrides);
+	ResetDiffResults.AddedOtherInputOverrides.RemoveAll(FindUnrelatedInputOverrides);
+	ResetDiffResults.ModifiedBaseInputOverrides.RemoveAll(FindUnrelatedInputOverrides);
+	ResetDiffResults.ModifiedOtherInputOverrides.RemoveAll(FindUnrelatedInputOverrides);
+
+	return ApplyScriptStackDiff(EmitterAdapter->GetScriptStack(ScriptUsage, ScriptUsageId).ToSharedRef(), ResetDiffResults);
+}
+
+bool FNiagaraScriptMergeManager::HasBaseEventHandler(const UNiagaraEmitter& BaseEmitter, FGuid EventScriptUsageId)
+{
+	TSharedRef<FNiagaraEmitterMergeAdapter> BaseEmitterAdapter = GetEmitterMergeAdapterUsingCache(BaseEmitter);
+	return BaseEmitterAdapter->GetEventHandler(EventScriptUsageId).IsValid();
+}
+
+bool FNiagaraScriptMergeManager::IsEventHandlerPropertySetDifferentFromBase(UNiagaraEmitter& Emitter, const UNiagaraEmitter& BaseEmitter, FGuid EventScriptUsageId)
+{
+	TSharedRef<FNiagaraEmitterMergeAdapter> EmitterAdapter = GetEmitterMergeAdapterUsingCache(Emitter);
+	TSharedRef<FNiagaraEmitterMergeAdapter> BaseEmitterAdapter = GetEmitterMergeAdapterUsingCache(BaseEmitter);
+
+	TSharedPtr<FNiagaraEventHandlerMergeAdapter> EventHandlerAdapter = EmitterAdapter->GetEventHandler(EventScriptUsageId);
+	TSharedPtr<FNiagaraEventHandlerMergeAdapter> BaseEventHandlerAdapter = BaseEmitterAdapter->GetEventHandler(EventScriptUsageId);
+
+	if (EventHandlerAdapter->GetEditableEventScriptProperties() == nullptr || BaseEventHandlerAdapter->GetEventScriptProperties() == nullptr)
+	{
+		return true;
+	}
+
+	TArray<UProperty*> DifferentProperties;
+	DiffEditableProperties(BaseEventHandlerAdapter->GetEventScriptProperties(), EventHandlerAdapter->GetEventScriptProperties(), *FNiagaraEventScriptProperties::StaticStruct(), DifferentProperties);
+	return DifferentProperties.Num() > 0;
+}
+
+void FNiagaraScriptMergeManager::ResetEventHandlerPropertySetToBase(UNiagaraEmitter& Emitter, const UNiagaraEmitter& BaseEmitter, FGuid EventScriptUsageId)
+{
+	TSharedRef<FNiagaraEmitterMergeAdapter> EmitterAdapter = GetEmitterMergeAdapterUsingCache(Emitter);
+	TSharedRef<FNiagaraEmitterMergeAdapter> BaseEmitterAdapter = GetEmitterMergeAdapterUsingCache(BaseEmitter);
+
+	TSharedPtr<FNiagaraEventHandlerMergeAdapter> EventHandlerAdapter = EmitterAdapter->GetEventHandler(EventScriptUsageId);
+	TSharedPtr<FNiagaraEventHandlerMergeAdapter> BaseEventHandlerAdapter = BaseEmitterAdapter->GetEventHandler(EventScriptUsageId);
+
+	if (EventHandlerAdapter->GetEditableEventScriptProperties() == nullptr || BaseEventHandlerAdapter->GetEventScriptProperties() == nullptr)
+	{
+		// TODO: Display an error to the user.
+		return;
+	}
+
+	TArray<UProperty*> DifferentProperties;
+	DiffEditableProperties(BaseEventHandlerAdapter->GetEventScriptProperties(), EventHandlerAdapter->GetEventScriptProperties(), *FNiagaraEventScriptProperties::StaticStruct(), DifferentProperties);
+	CopyPropertiesToBase(EventHandlerAdapter->GetEditableEventScriptProperties(), BaseEventHandlerAdapter->GetEventScriptProperties(), DifferentProperties);
+	Emitter.PostEditChange();
+}
+
+bool FNiagaraScriptMergeManager::HasBaseRenderer(const UNiagaraEmitter& BaseEmitter, FGuid RendererMergeId)
+{
+	TSharedRef<FNiagaraEmitterMergeAdapter> BaseEmitterAdapter = GetEmitterMergeAdapterUsingCache(BaseEmitter);
+	return BaseEmitterAdapter->GetRenderer(RendererMergeId).IsValid();
+}
+
+bool FNiagaraScriptMergeManager::IsRendererDifferentFromBase(UNiagaraEmitter& Emitter, const UNiagaraEmitter& BaseEmitter, FGuid RendererMergeId)
+{
+	TSharedRef<FNiagaraEmitterMergeAdapter> EmitterAdapter = GetEmitterMergeAdapterUsingCache(Emitter);
+	TSharedRef<FNiagaraEmitterMergeAdapter> BaseEmitterAdapter = GetEmitterMergeAdapterUsingCache(BaseEmitter);
+
+	FNiagaraEmitterDiffResults DiffResults;
+	DiffRenderers(BaseEmitterAdapter->GetRenderers(), EmitterAdapter->GetRenderers(), DiffResults);
+
+	if (DiffResults.IsValid() == false)
+	{
+		return true;
+	}
+
+	if (DiffResults.ModifiedOtherRenderers.Num() == 0)
+	{
+		return false;
+	}
+
+	auto FindRendererByMergeId = [=](TSharedRef<FNiagaraRendererMergeAdapter> Renderer) { return Renderer->GetRenderer()->GetMergeId() == RendererMergeId; };
+	return DiffResults.ModifiedOtherRenderers.FindByPredicate(FindRendererByMergeId) != nullptr;
+}
+
+void FNiagaraScriptMergeManager::ResetRendererToBase(UNiagaraEmitter& Emitter, const UNiagaraEmitter& BaseEmitter, FGuid RendererMergeId)
+{
+	TSharedRef<FNiagaraEmitterMergeAdapter> EmitterAdapter = GetEmitterMergeAdapterUsingCache(Emitter);
+	TSharedRef<FNiagaraEmitterMergeAdapter> BaseEmitterAdapter = GetEmitterMergeAdapterUsingCache(BaseEmitter);
+
+	// Diff from the current emitter to the base emitter to create a diff which will reset the emitter back to the base.
+	FNiagaraEmitterDiffResults ResetDiffResults;
+	DiffRenderers(EmitterAdapter->GetRenderers(), BaseEmitterAdapter->GetRenderers(), ResetDiffResults);
+
+	auto FindUnrelatedRenderers = [=](TSharedRef<FNiagaraRendererMergeAdapter> Renderer)
+	{
+		return Renderer->GetRenderer()->GetMergeId() != RendererMergeId;
+	};
+
+	// Removed added and removed renderers, as well as changes to renderers with different ids from the one being reset.
+	ResetDiffResults.RemovedBaseRenderers.Empty();
+	ResetDiffResults.AddedOtherRenderers.Empty();
+	ResetDiffResults.ModifiedBaseRenderers.RemoveAll(FindUnrelatedRenderers);
+	ResetDiffResults.ModifiedOtherRenderers.RemoveAll(FindUnrelatedRenderers);
+
+	ApplyRendererDiff(Emitter, ResetDiffResults);
+}
+
+bool FNiagaraScriptMergeManager::IsEmitterEditablePropertySetDifferentFromBase(const UNiagaraEmitter& Emitter, const UNiagaraEmitter& BaseEmitter)
+{
+	TArray<UProperty*> DifferentProperties;
+	DiffEditableProperties(&BaseEmitter, &Emitter, *UNiagaraEmitter::StaticClass(), DifferentProperties);
+	return DifferentProperties.Num() > 0;
+}
+
+void FNiagaraScriptMergeManager::ResetEmitterEditablePropertySetToBase(UNiagaraEmitter& Emitter, const UNiagaraEmitter& BaseEmitter)
+{
+	TArray<UProperty*> DifferentProperties;
+	DiffEditableProperties(&BaseEmitter, &Emitter, *UNiagaraEmitter::StaticClass(), DifferentProperties);
+	CopyPropertiesToBase(&Emitter, &BaseEmitter, DifferentProperties);
+	Emitter.PostEditChange();
+}
+
+FNiagaraEmitterDiffResults FNiagaraScriptMergeManager::DiffEmitters(UNiagaraEmitter& BaseEmitter, UNiagaraEmitter& OtherEmitter)
+{
+	SCOPE_CYCLE_COUNTER(STAT_NiagaraEditor_ScriptMergeManager_DiffEmitters);
+
+	TSharedRef<FNiagaraEmitterMergeAdapter> BaseEmitterAdapter = MakeShared<FNiagaraEmitterMergeAdapter>(BaseEmitter);
+	TSharedRef<FNiagaraEmitterMergeAdapter> OtherEmitterAdapter = MakeShared<FNiagaraEmitterMergeAdapter>(OtherEmitter);
+
+	FNiagaraEmitterDiffResults EmitterDiffResults;
+	if (BaseEmitterAdapter->GetEmitterSpawnStack().IsValid() && OtherEmitterAdapter->GetEmitterSpawnStack().IsValid())
+	{
+		DiffScriptStacks(BaseEmitterAdapter->GetEmitterSpawnStack().ToSharedRef(), OtherEmitterAdapter->GetEmitterSpawnStack().ToSharedRef(), EmitterDiffResults.EmitterSpawnDiffResults);
+		if (EmitterDiffResults.EmitterSpawnDiffResults.IsValid() == false)
+		{
+			EmitterDiffResults.AddError(LOCTEXT("EmitterSpawnStacksDiffInvalidMessage", "Emitter spawn diff is invalid."));
+		}
+	}
+	else
+	{
+		EmitterDiffResults.AddError(LOCTEXT("EmitterSpawnStacksInvalidMessage", "One of the emitter spawn script stacks was invalid."));
+	}
+
+	if (BaseEmitterAdapter->GetEmitterUpdateStack().IsValid() && OtherEmitterAdapter->GetEmitterUpdateStack().IsValid())
+	{
+		DiffScriptStacks(BaseEmitterAdapter->GetEmitterUpdateStack().ToSharedRef(), OtherEmitterAdapter->GetEmitterUpdateStack().ToSharedRef(), EmitterDiffResults.EmitterUpdateDiffResults);
+		if (EmitterDiffResults.EmitterUpdateDiffResults.IsValid() == false)
+		{
+			EmitterDiffResults.AddError(LOCTEXT("EmitterUpdateStacksDiffInvalidMessage", "Emitter update diff is invalid."));
+		}
+	}
+	else
+	{
+		EmitterDiffResults.AddError(LOCTEXT("EmitterUpdateStacksInvalidMessage", "One of the emitter update script stacks was invalid."));
+	}
+
+	if (BaseEmitterAdapter->GetParticleSpawnStack().IsValid() && OtherEmitterAdapter->GetParticleSpawnStack().IsValid())
+	{
+		DiffScriptStacks(BaseEmitterAdapter->GetParticleSpawnStack().ToSharedRef(), OtherEmitterAdapter->GetParticleSpawnStack().ToSharedRef(), EmitterDiffResults.ParticleSpawnDiffResults);
+		if (EmitterDiffResults.ParticleSpawnDiffResults.IsValid() == false)
+		{
+			EmitterDiffResults.AddError(LOCTEXT("ParticleSpawnStacksDiffInvalidMessage", "Particle spawn diff is invalid."));
+		}
+	}
+	else
+	{
+		EmitterDiffResults.AddError(LOCTEXT("ParticleSpawnStacksInvalidMessage", "One of the particle spawn script stacks was invalid."));
+	}
+
+	if (BaseEmitterAdapter->GetParticleUpdateStack().IsValid() && OtherEmitterAdapter->GetParticleUpdateStack().IsValid())
+	{
+		DiffScriptStacks(BaseEmitterAdapter->GetParticleUpdateStack().ToSharedRef(), OtherEmitterAdapter->GetParticleUpdateStack().ToSharedRef(), EmitterDiffResults.ParticleUpdateDiffResults);
+		if (EmitterDiffResults.ParticleUpdateDiffResults.IsValid() == false)
+		{
+			EmitterDiffResults.AddError(LOCTEXT("ParticleUpdateStacksDiffInvalidMessage", "Particle update diff is invalid."));
+		}
+	}
+	else
+	{
+		EmitterDiffResults.AddError(LOCTEXT("ParticleUpdateStacksInvalidMessage", "One of the particle update script stacks was invalid."));
+	}
+
+	DiffEventHandlers(BaseEmitterAdapter->GetEventHandlers(), OtherEmitterAdapter->GetEventHandlers(), EmitterDiffResults);
+	DiffRenderers(BaseEmitterAdapter->GetRenderers(), OtherEmitterAdapter->GetRenderers(), EmitterDiffResults);
+	DiffEditableProperties(&BaseEmitter, &OtherEmitter, *UNiagaraEmitter::StaticClass(), EmitterDiffResults.DifferentEmitterProperties);
+
+	return EmitterDiffResults;
+}
+
+template<typename ValueType>
+struct FCommonValuePair
+{
+	FCommonValuePair(ValueType InBaseValue, ValueType InOtherValue)
+		: BaseValue(InBaseValue)
+		, OtherValue(InOtherValue)
+	{
+	}
+
+	ValueType BaseValue;
+	ValueType OtherValue;
+};
+
+template<typename ValueType>
+struct FListDiffResults
+{
+	TArray<ValueType> RemovedBaseValues;
+	TArray<ValueType> AddedOtherValues;
+	TArray<FCommonValuePair<ValueType>> CommonValuePairs;
+};
+
+template<typename ValueType, typename KeyType, typename KeyFromValueDelegate>
+FListDiffResults<ValueType> DiffLists(const TArray<ValueType>& BaseList, const TArray<ValueType> OtherList, KeyFromValueDelegate KeyFromValue)
+{
+	FListDiffResults<ValueType> DiffResults;
+
+	TMap<KeyType, ValueType> BaseKeyToValueMap;
+	TSet<KeyType> BaseKeys;
+	for (ValueType BaseValue : BaseList)
+	{
+		KeyType BaseKey = KeyFromValue(BaseValue);
+		BaseKeyToValueMap.Add(BaseKey, BaseValue);
+		BaseKeys.Add(BaseKey);
+	}
+
+	TMap<KeyType, ValueType> OtherKeyToValueMap;
+	TSet<KeyType> OtherKeys;
+	for (ValueType OtherValue : OtherList)
+	{
+		KeyType OtherKey = KeyFromValue(OtherValue);
+		OtherKeyToValueMap.Add(OtherKey, OtherValue);
+		OtherKeys.Add(OtherKey);
+	}
+
+	for (KeyType RemovedKey : BaseKeys.Difference(OtherKeys))
+	{
+		DiffResults.RemovedBaseValues.Add(BaseKeyToValueMap[RemovedKey]);
+	}
+
+	for (KeyType AddedKey : OtherKeys.Difference(BaseKeys))
+	{
+		DiffResults.AddedOtherValues.Add(OtherKeyToValueMap[AddedKey]);
+	}
+
+	for (KeyType CommonKey : BaseKeys.Intersect(OtherKeys))
+	{
+		DiffResults.CommonValuePairs.Add(FCommonValuePair<ValueType>(BaseKeyToValueMap[CommonKey], OtherKeyToValueMap[CommonKey]));
+	}
+
+	return DiffResults;
+}
+
+void FNiagaraScriptMergeManager::DiffEventHandlers(const TArray<TSharedRef<FNiagaraEventHandlerMergeAdapter>>& BaseEventHandlers, const TArray<TSharedRef<FNiagaraEventHandlerMergeAdapter>>& OtherEventHandlers, FNiagaraEmitterDiffResults& DiffResults)
+{
+	FListDiffResults<TSharedRef<FNiagaraEventHandlerMergeAdapter>> EventHandlerListDiffResults = DiffLists<TSharedRef<FNiagaraEventHandlerMergeAdapter>, FGuid>(
+		BaseEventHandlers,
+		OtherEventHandlers,
+		[](TSharedRef<FNiagaraEventHandlerMergeAdapter> EventHandler) { return EventHandler->GetUsageId(); });
+
+	DiffResults.RemovedBaseEventHandlers.Append(EventHandlerListDiffResults.RemovedBaseValues);
+	DiffResults.AddedOtherEventHandlers.Append(EventHandlerListDiffResults.AddedOtherValues);
+
+	for (const FCommonValuePair<TSharedRef<FNiagaraEventHandlerMergeAdapter>>& CommonValuePair : EventHandlerListDiffResults.CommonValuePairs)
+	{
+		if (CommonValuePair.BaseValue->GetEventScriptProperties() == nullptr || CommonValuePair.BaseValue->GetOutputNode() == nullptr)
+		{
+			DiffResults.AddError(FText::Format(LOCTEXT("InvalidBaseEventHandlerDiffFailedFormat", "Failed to diff event handlers, the base event handler was invalid.  Script Usage Id: {0}"),
+				FText::FromString(CommonValuePair.BaseValue->GetUsageId().ToString())));
+		}
+		else if(CommonValuePair.OtherValue->GetEventScriptProperties() == nullptr || CommonValuePair.OtherValue->GetOutputNode() == nullptr)
+		{
+			DiffResults.AddError(FText::Format(LOCTEXT("InvalidOtherEventHandlerDiffFailedFormat", "Failed to diff event handlers, the other event handler was invalid.  Script Usage Id: {0}"),
+				FText::FromString(CommonValuePair.OtherValue->GetUsageId().ToString())));
+		}
+		else
+		{
+			TArray<UProperty*> DifferentProperties;
+			DiffEditableProperties(CommonValuePair.BaseValue->GetEventScriptProperties(), CommonValuePair.OtherValue->GetEventScriptProperties(), *FNiagaraEventScriptProperties::StaticStruct(), DifferentProperties);
+
+			FNiagaraScriptStackDiffResults EventHandlerScriptStackDiffResults;
+			DiffScriptStacks(CommonValuePair.BaseValue->GetEventStack().ToSharedRef(), CommonValuePair.OtherValue->GetEventStack().ToSharedRef(), EventHandlerScriptStackDiffResults);
+
+			if (DifferentProperties.Num() > 0 || EventHandlerScriptStackDiffResults.IsValid() == false || EventHandlerScriptStackDiffResults.IsEmpty() == false)
+			{
+				FNiagaraModifiedEventHandlerDiffResults ModifiedEventHandlerResults;
+				ModifiedEventHandlerResults.BaseAdapter = CommonValuePair.BaseValue;
+				ModifiedEventHandlerResults.OtherAdapter = CommonValuePair.OtherValue;
+				ModifiedEventHandlerResults.ChangedProperties.Append(DifferentProperties);
+				ModifiedEventHandlerResults.ScriptDiffResults = EventHandlerScriptStackDiffResults;
+				DiffResults.ModifiedEventHandlers.Add(ModifiedEventHandlerResults);
+			}
+
+			if (EventHandlerScriptStackDiffResults.IsValid() == false)
+			{
+				for (const FText& EventHandlerScriptStackDiffErrorMessage : EventHandlerScriptStackDiffResults.GetErrorMessages())
+				{
+					DiffResults.AddError(EventHandlerScriptStackDiffErrorMessage);
+				}
+			}
+		}
+	}
+}
+
+void FNiagaraScriptMergeManager::DiffRenderers(const TArray<TSharedRef<FNiagaraRendererMergeAdapter>>& BaseRenderers, const TArray<TSharedRef<FNiagaraRendererMergeAdapter>>& OtherRenderers, FNiagaraEmitterDiffResults& DiffResults)
+{
+	FListDiffResults<TSharedRef<FNiagaraRendererMergeAdapter>> RendererListDiffResults = DiffLists<TSharedRef<FNiagaraRendererMergeAdapter>, FGuid>(
+		BaseRenderers,
+		OtherRenderers,
+		[](TSharedRef<FNiagaraRendererMergeAdapter> Renderer) { return Renderer->GetRenderer()->GetMergeId(); });
+
+	DiffResults.RemovedBaseRenderers.Append(RendererListDiffResults.RemovedBaseValues);
+	DiffResults.AddedOtherRenderers.Append(RendererListDiffResults.AddedOtherValues);
+
+	for (const FCommonValuePair<TSharedRef<FNiagaraRendererMergeAdapter>>& CommonValuePair : RendererListDiffResults.CommonValuePairs)
+	{
+		if (CommonValuePair.BaseValue->GetRenderer()->Equals(CommonValuePair.OtherValue->GetRenderer()) == false)
+		{
+			DiffResults.ModifiedBaseRenderers.Add(CommonValuePair.BaseValue);
+			DiffResults.ModifiedOtherRenderers.Add(CommonValuePair.OtherValue);
+		}
+	}
+}
+
+void FNiagaraScriptMergeManager::DiffScriptStacks(TSharedRef<FNiagaraScriptStackMergeAdapter> BaseScriptStackAdapter, TSharedRef<FNiagaraScriptStackMergeAdapter> OtherScriptStackAdapter, FNiagaraScriptStackDiffResults& DiffResults)
+{
+	// Diff the module lists.
+	FListDiffResults<TSharedRef<FNiagaraStackFunctionMergeAdapter>> ModuleListDiffResults = DiffLists<TSharedRef<FNiagaraStackFunctionMergeAdapter>, FGuid>(
+		BaseScriptStackAdapter->GetModuleFunctions(),
+		OtherScriptStackAdapter->GetModuleFunctions(),
+		[](TSharedRef<FNiagaraStackFunctionMergeAdapter> FunctionAdapter) { return FunctionAdapter->GetFunctionCallNode()->NodeGuid; });
+
+	// Sort the diff results for easier diff applying and testing.
+	auto OrderModuleByStackIndex = [](TSharedRef<FNiagaraStackFunctionMergeAdapter> ModuleA, TSharedRef<FNiagaraStackFunctionMergeAdapter> ModuleB)
+	{
+		return ModuleA->GetStackIndex() < ModuleB->GetStackIndex();
+	};
+
+	ModuleListDiffResults.RemovedBaseValues.Sort(OrderModuleByStackIndex);
+	ModuleListDiffResults.AddedOtherValues.Sort(OrderModuleByStackIndex);
+
+	auto OrderCommonModulePairByBaseStackIndex = [](
+		const FCommonValuePair<TSharedRef<FNiagaraStackFunctionMergeAdapter>>& CommonValuesA,
+		const FCommonValuePair<TSharedRef<FNiagaraStackFunctionMergeAdapter>>& CommonValuesB)
+	{
+		return CommonValuesA.BaseValue->GetStackIndex() < CommonValuesB.BaseValue->GetStackIndex();
+	};
+
+	ModuleListDiffResults.CommonValuePairs.Sort(OrderCommonModulePairByBaseStackIndex);
+
+	// Populate results from the sorted diff.
+	DiffResults.RemovedBaseModules.Append(ModuleListDiffResults.RemovedBaseValues);
+	DiffResults.AddedOtherModules.Append(ModuleListDiffResults.AddedOtherValues);
+
+	for (const FCommonValuePair<TSharedRef<FNiagaraStackFunctionMergeAdapter>>& CommonValuePair : ModuleListDiffResults.CommonValuePairs)
+	{
+		if (CommonValuePair.BaseValue->GetStackIndex() != CommonValuePair.OtherValue->GetStackIndex())
+		{
+			DiffResults.MovedBaseModules.Add(CommonValuePair.BaseValue);
+			DiffResults.MovedOtherModules.Add(CommonValuePair.OtherValue);
+		}
+
+		if (CommonValuePair.BaseValue->GetFunctionCallNode()->IsNodeEnabled() != CommonValuePair.OtherValue->GetFunctionCallNode()->IsNodeEnabled())
+		{
+			DiffResults.EnabledChangedBaseModules.Add(CommonValuePair.BaseValue);
+			DiffResults.EnabledChangedOtherModules.Add(CommonValuePair.OtherValue);
+		}
+
+		DiffFunctionInputs(CommonValuePair.BaseValue, CommonValuePair.OtherValue, DiffResults);
+	}
+}
+
+void FNiagaraScriptMergeManager::DiffFunctionInputs(TSharedRef<FNiagaraStackFunctionMergeAdapter> BaseFunctionAdapter, TSharedRef<FNiagaraStackFunctionMergeAdapter> OtherFunctionAdapter, FNiagaraScriptStackDiffResults& DiffResults)
+{
+	FListDiffResults<TSharedRef<FNiagaraStackFunctionInputOverrideMergeAdapter>> ListDiffResults = DiffLists<TSharedRef<FNiagaraStackFunctionInputOverrideMergeAdapter>, FString>(
+		BaseFunctionAdapter->GetInputOverrides(),
+		OtherFunctionAdapter->GetInputOverrides(),
+		[](TSharedRef<FNiagaraStackFunctionInputOverrideMergeAdapter> InputOverrideAdapter) { return InputOverrideAdapter->GetInputName(); });
+
+	DiffResults.RemovedBaseInputOverrides.Append(ListDiffResults.RemovedBaseValues);
+	DiffResults.AddedOtherInputOverrides.Append(ListDiffResults.AddedOtherValues);
+
+	for (const FCommonValuePair<TSharedRef<FNiagaraStackFunctionInputOverrideMergeAdapter>>& CommonValuePair : ListDiffResults.CommonValuePairs)
+	{
+		TOptional<bool> FunctionMatch = DoFunctionInputOverridesMatch(CommonValuePair.BaseValue, CommonValuePair.OtherValue);
+		if (FunctionMatch.IsSet())
+		{
+			if (FunctionMatch.GetValue() == false)
+			{
+				DiffResults.ModifiedBaseInputOverrides.Add(CommonValuePair.BaseValue);
+				DiffResults.ModifiedOtherInputOverrides.Add(CommonValuePair.OtherValue);
+			}
+		}
+		else
+		{
+			DiffResults.AddError(FText::Format(LOCTEXT("FunctionInputDiffFailedFormat", "Failed to diff function inputs.  Function name: {0}  Input Name: {1}"),
+				FText::FromString(BaseFunctionAdapter->GetFunctionCallNode()->GetFunctionName()),
+				FText::FromString(CommonValuePair.BaseValue->GetInputName())));
+		}
+	}
+}
+
+void FNiagaraScriptMergeManager::DiffEditableProperties(const void* BaseDataAddress, const void* OtherDataAddress, UStruct& Struct, TArray<UProperty*>& OutDifferentProperties)
+{
+	for (TFieldIterator<UProperty> PropertyIterator(&Struct); PropertyIterator; ++PropertyIterator)
+	{
+		if (PropertyIterator->HasAllPropertyFlags(CPF_Edit))
+		{
+			if (PropertyIterator->Identical(
+				PropertyIterator->ContainerPtrToValuePtr<void>(BaseDataAddress),
+				PropertyIterator->ContainerPtrToValuePtr<void>(OtherDataAddress), PPF_DeepComparison) == false)
+			{
+				OutDifferentProperties.Add(*PropertyIterator);
+			}
+		}
+	}
+}
+
+TOptional<bool> FNiagaraScriptMergeManager::DoFunctionInputOverridesMatch(TSharedRef<FNiagaraStackFunctionInputOverrideMergeAdapter> BaseFunctionInputAdapter, TSharedRef<FNiagaraStackFunctionInputOverrideMergeAdapter> OtherFunctionInputAdapter)
+{
+	// Local String Value.
+	if ((BaseFunctionInputAdapter->GetLocalValueString().IsSet() && OtherFunctionInputAdapter->GetLocalValueString().IsSet() == false) ||
+		(BaseFunctionInputAdapter->GetLocalValueString().IsSet() == false && OtherFunctionInputAdapter->GetLocalValueString().IsSet()))
+	{
+		return false;
+	}
+
+	if (BaseFunctionInputAdapter->GetLocalValueString().IsSet() && OtherFunctionInputAdapter->GetLocalValueString().IsSet())
+	{
+		return BaseFunctionInputAdapter->GetLocalValueString().GetValue() == OtherFunctionInputAdapter->GetLocalValueString().GetValue();
+	}
+
+	// Local rapid iteration parameter value.
+	if ((BaseFunctionInputAdapter->GetLocalValueRapidIterationParameter().IsSet() && OtherFunctionInputAdapter->GetLocalValueRapidIterationParameter().IsSet() == false) ||
+		(BaseFunctionInputAdapter->GetLocalValueRapidIterationParameter().IsSet() == false && OtherFunctionInputAdapter->GetLocalValueRapidIterationParameter().IsSet()))
+	{
+		return false;
+	}
+
+	if (BaseFunctionInputAdapter->GetLocalValueRapidIterationParameter().IsSet() && OtherFunctionInputAdapter->GetLocalValueRapidIterationParameter().IsSet())
+	{
+		const uint8* BaseRapidIterationParameterValue = BaseFunctionInputAdapter->GetOwningScript()->RapidIterationParameters
+			.GetParameterData(BaseFunctionInputAdapter->GetLocalValueRapidIterationParameter().GetValue());
+		const uint8* OtherRapidIterationParameterValue = OtherFunctionInputAdapter->GetOwningScript()->RapidIterationParameters
+			.GetParameterData(OtherFunctionInputAdapter->GetLocalValueRapidIterationParameter().GetValue());
+		return FMemory::Memcmp(
+			BaseRapidIterationParameterValue,
+			OtherRapidIterationParameterValue,
+			BaseFunctionInputAdapter->GetLocalValueRapidIterationParameter().GetValue().GetSizeInBytes()) == 0;
+	}
+
+	// Linked value
+	if ((BaseFunctionInputAdapter->GetLinkedValueHandle().IsSet() && OtherFunctionInputAdapter->GetLinkedValueHandle().IsSet() == false) ||
+		(BaseFunctionInputAdapter->GetLinkedValueHandle().IsSet() == false && OtherFunctionInputAdapter->GetLinkedValueHandle().IsSet()))
+	{
+		return false;
+	}
+
+	if (BaseFunctionInputAdapter->GetLinkedValueHandle().IsSet() && OtherFunctionInputAdapter->GetLinkedValueHandle().IsSet())
+	{
+		return BaseFunctionInputAdapter->GetLinkedValueHandle().GetValue() == OtherFunctionInputAdapter->GetLinkedValueHandle().GetValue();
+	}
+
+	// Data value
+	if ((BaseFunctionInputAdapter->GetDataValueObject()!= nullptr && OtherFunctionInputAdapter->GetDataValueObject()== nullptr) ||
+		(BaseFunctionInputAdapter->GetDataValueObject()== nullptr && OtherFunctionInputAdapter->GetDataValueObject()!= nullptr))
+	{
+		return false;
+	}
+
+	if (BaseFunctionInputAdapter->GetDataValueObject()!= nullptr && OtherFunctionInputAdapter->GetDataValueObject()!= nullptr)
+	{
+		return BaseFunctionInputAdapter->GetDataValueObject()->Equals(OtherFunctionInputAdapter->GetDataValueObject());
+	}
+
+	// Dynamic value
+	if ((BaseFunctionInputAdapter->GetDynamicValueFunction().IsValid() && OtherFunctionInputAdapter->GetDynamicValueFunction().IsValid() == false) ||
+		(BaseFunctionInputAdapter->GetDynamicValueFunction().IsValid() == false && OtherFunctionInputAdapter->GetDynamicValueFunction().IsValid()))
+	{
+		return false;
+	}
+
+	if (BaseFunctionInputAdapter->GetDynamicValueFunction().IsValid() && OtherFunctionInputAdapter->GetDynamicValueFunction().IsValid())
+	{
+		if (BaseFunctionInputAdapter->GetDynamicValueFunction()->GetFunctionCallNode()->FunctionScript != OtherFunctionInputAdapter->GetDynamicValueFunction()->GetFunctionCallNode()->FunctionScript)
+		{
+			return false;
+		}
+
+		FNiagaraScriptStackDiffResults FunctionDiffResults;
+		DiffFunctionInputs(BaseFunctionInputAdapter->GetDynamicValueFunction().ToSharedRef(), OtherFunctionInputAdapter->GetDynamicValueFunction().ToSharedRef(), FunctionDiffResults);
+		return
+			FunctionDiffResults.RemovedBaseInputOverrides.Num() == 0 &&
+			FunctionDiffResults.AddedOtherInputOverrides.Num() == 0 &&
+			FunctionDiffResults.ModifiedOtherInputOverrides.Num() == 0;
+	}
+
+	return nullptr;
+}
+
+FNiagaraScriptMergeManager::FApplyDiffResults FNiagaraScriptMergeManager::AddModule(FString UniqueEmitterName, UNiagaraScript& OwningScript, UNiagaraNodeOutput& TargetOutputNode, TSharedRef<FNiagaraStackFunctionMergeAdapter> AddModule)
+{
+	FApplyDiffResults Results;
+
+	UNiagaraNodeFunctionCall* AddedModuleNode = nullptr;
+	if (AddModule->GetFunctionCallNode()->IsA<UNiagaraNodeAssignment>())
+	{
+		AddedModuleNode = FNiagaraStackGraphUtilities::AddParameterModuleToStack(CastChecked<UNiagaraNodeAssignment>(AddModule->GetFunctionCallNode())->AssignmentTarget, TargetOutputNode, AddModule->GetStackIndex());
+		Results.bModifiedGraph = true;
+	}
+	else
+	{
+		if (AddModule->GetFunctionCallNode()->FunctionScript != nullptr)
+		{
+			AddedModuleNode = FNiagaraStackGraphUtilities::AddScriptModuleToStack(AddModule->GetFunctionCallNode()->FunctionScript, TargetOutputNode, AddModule->GetStackIndex());
+			Results.bModifiedGraph = true;
+		}
+		else
+		{
+			Results.bSucceeded = false;
+			Results.ErrorMessages.Add(FText::Format(
+				LOCTEXT("AddModuleFailedDueToMissingModuleScriptFormat", "Can not add module {0} because it's script was missing."),
+				FText::FromString(AddModule->GetFunctionCallNode()->GetFunctionName())));
+		}
+	}
+
+	if (AddedModuleNode != nullptr)
+	{
+		for (TSharedRef<FNiagaraStackFunctionInputOverrideMergeAdapter> InputOverride : AddModule->GetInputOverrides())
+		{
+			FApplyDiffResults AddInputResults = AddInputOverride(UniqueEmitterName, OwningScript, *AddedModuleNode, InputOverride);
+			Results.bSucceeded &= AddInputResults.bSucceeded;
+			Results.bModifiedGraph |= AddInputResults.bModifiedGraph;
+			Results.ErrorMessages.Append(AddInputResults.ErrorMessages);
+		}
+	}
+	else
+	{
+		Results.bSucceeded = false;
+		Results.ErrorMessages.Add(LOCTEXT("AddModuleFailed", "Failed to add module from diff."));
+	}
+
+	return Results;
+}
+
+FNiagaraScriptMergeManager::FApplyDiffResults FNiagaraScriptMergeManager::RemoveInputOverride(UNiagaraScript& OwningScript, TSharedRef<FNiagaraStackFunctionInputOverrideMergeAdapter> OverrideToRemove)
+{
+	FApplyDiffResults Results;
+	if (OverrideToRemove->GetOverridePin() != nullptr && OverrideToRemove->GetOverrideNode() != nullptr)
+	{
+		FNiagaraStackGraphUtilities::RemoveNodesForStackFunctionInputOverridePin(*OverrideToRemove->GetOverridePin());
+		OverrideToRemove->GetOverrideNode()->RemovePin(OverrideToRemove->GetOverridePin());
+		Results.bSucceeded = true;
+		Results.bModifiedGraph = true;
+	}
+	else if (OverrideToRemove->GetLocalValueRapidIterationParameter().IsSet())
+	{
+		OwningScript.Modify();
+		OwningScript.RapidIterationParameters.RemoveParameter(OverrideToRemove->GetLocalValueRapidIterationParameter().GetValue());
+		Results.bSucceeded = true;
+		Results.bModifiedGraph = false;
+	}
+	else
+	{
+		Results.bSucceeded = false;
+		Results.bModifiedGraph = false;
+		Results.ErrorMessages.Add(LOCTEXT("RemoveInputOverrideFailed", "Failed to remove input override because it was invalid."));
+	}
+	return Results;
+}
+
+FNiagaraScriptMergeManager::FApplyDiffResults FNiagaraScriptMergeManager::AddInputOverride(FString UniqueEmitterName, UNiagaraScript& OwningScript, UNiagaraNodeFunctionCall& TargetFunctionCall, TSharedRef<FNiagaraStackFunctionInputOverrideMergeAdapter> OverrideToAdd)
+{
+	FApplyDiffResults Results;
+
+	if (OverrideToAdd->GetOverridePin() != nullptr)
+	{
+		FNiagaraParameterHandle FunctionInputHandle(FNiagaraParameterHandle::ModuleNamespace, *OverrideToAdd->GetInputName());
+		FNiagaraParameterHandle AliasedFunctionInputHandle = FNiagaraParameterHandle::CreateAliasedModuleParameterHandle(FunctionInputHandle, &TargetFunctionCall);
+
+		const UEdGraphSchema_Niagara* NiagaraSchema = GetDefault<UEdGraphSchema_Niagara>();
+		FNiagaraTypeDefinition InputType = NiagaraSchema->PinToTypeDefinition(OverrideToAdd->GetOverridePin());
+
+		UEdGraphPin& InputOverridePin = FNiagaraStackGraphUtilities::GetOrCreateStackFunctionInputOverridePin(TargetFunctionCall, AliasedFunctionInputHandle, InputType);
+
+		if (OverrideToAdd->GetLocalValueString().IsSet())
+		{
+			InputOverridePin.DefaultValue = OverrideToAdd->GetLocalValueString().GetValue();
+			Results.bSucceeded = true;
+		}
+		else if (OverrideToAdd->GetLinkedValueHandle().IsSet())
+		{
+			FNiagaraStackGraphUtilities::SetLinkedValueHandleForFunctionInput(InputOverridePin, OverrideToAdd->GetLinkedValueHandle().GetValue());
+			Results.bSucceeded = true;
+		}
+		else if (OverrideToAdd->GetDataValueObject() != nullptr)
+		{
+			UNiagaraDataInterface* OverrideValueObject = OverrideToAdd->GetDataValueObject();
+			UNiagaraDataInterface* NewOverrideValueObject;
+			FNiagaraStackGraphUtilities::SetDataValueObjectForFunctionInput(InputOverridePin, OverrideToAdd->GetDataValueObject()->GetClass(), OverrideValueObject->GetName(), NewOverrideValueObject);
+			OverrideValueObject->CopyTo(NewOverrideValueObject);
+			Results.bSucceeded = true;
+		}
+		else if (OverrideToAdd->GetDynamicValueFunction().IsValid())
+		{
+			if (OverrideToAdd->GetDynamicValueFunction()->GetFunctionCallNode()->FunctionScript != nullptr)
+			{
+				UNiagaraNodeFunctionCall* DynamicInputFunctionCall;
+				FNiagaraStackGraphUtilities::SetDynamicInputForFunctionInput(InputOverridePin, OverrideToAdd->GetDynamicValueFunction()->GetFunctionCallNode()->FunctionScript, DynamicInputFunctionCall);
+				for (TSharedRef<FNiagaraStackFunctionInputOverrideMergeAdapter> DynamicInputInputOverride : OverrideToAdd->GetDynamicValueFunction()->GetInputOverrides())
+				{
+					FApplyDiffResults AddResults = AddInputOverride(UniqueEmitterName, OwningScript, *DynamicInputFunctionCall, DynamicInputInputOverride);
+					Results.bSucceeded &= AddResults.bSucceeded;
+					Results.bModifiedGraph |= AddResults.bModifiedGraph;
+					Results.ErrorMessages.Append(AddResults.ErrorMessages);
+				}
+			}
+			else
+			{
+				Results.bSucceeded = false;
+				Results.ErrorMessages.Add(LOCTEXT("AddPinBasedInputOverrideFailedInvalidDynamicInput", "Failed to add input override because it's dynamic function call's function script was null."));
+			}
+		}
+		else
+		{
+			Results.bSucceeded = false;
+			Results.ErrorMessages.Add(LOCTEXT("AddPinBasedInputOverrideFailed", "Failed to add input override because it was invalid."));
+		}
+		Results.bModifiedGraph = true;
+	}
+	else
+	{
+		if (OverrideToAdd->GetLocalValueRapidIterationParameter().IsSet())
+		{
+			FString RapidIterationParameterName = FString::Printf(TEXT("Constants.%s.%s.%s"), *UniqueEmitterName, *TargetFunctionCall.GetFunctionName(), *OverrideToAdd->GetInputName());
+			FNiagaraVariable RapidIterationParameter(OverrideToAdd->GetLocalValueRapidIterationParameter().GetValue().GetType(), *RapidIterationParameterName);
+			const uint8* SourceData = OverrideToAdd->GetOwningScript()->RapidIterationParameters.GetParameterData(OverrideToAdd->GetLocalValueRapidIterationParameter().GetValue());
+			OwningScript.Modify();
+			OwningScript.RapidIterationParameters.AddParameter(RapidIterationParameter);
+			OwningScript.RapidIterationParameters.SetParameterData(SourceData, RapidIterationParameter);
+			Results.bSucceeded = true;
+		}
+		else
+		{
+			Results.bSucceeded = false;
+			Results.ErrorMessages.Add(LOCTEXT("AddParameterBasedInputOverrideFailed", "Failed to add input override because it was invalid."));
+		}
+		Results.bModifiedGraph = false;
+	}
+	return Results;
+}
+
+void FNiagaraScriptMergeManager::CopyPropertiesToBase(void* BaseDataAddress, const void* OtherDataAddress, TArray<UProperty*> PropertiesToCopy)
+{
+	for (UProperty* PropertyToCopy : PropertiesToCopy)
+	{
+		PropertyToCopy->CopyCompleteValue(PropertyToCopy->ContainerPtrToValuePtr<void>(BaseDataAddress), PropertyToCopy->ContainerPtrToValuePtr<void>(OtherDataAddress));
+	}
+}
+
+FNiagaraScriptMergeManager::FApplyDiffResults FNiagaraScriptMergeManager::ApplyScriptStackDiff(TSharedRef<FNiagaraScriptStackMergeAdapter> BaseScriptStackAdapter, const FNiagaraScriptStackDiffResults& DiffResults)
+{
+	FApplyDiffResults Results;
+
+	if (DiffResults.IsEmpty())
+	{
+		Results.bSucceeded = true;
+		Results.bModifiedGraph = false;
+		return Results;
+	}
+
+	struct FAddInputOverrideActionData
+	{
+		UNiagaraNodeFunctionCall* TargetFunctionCall;
+		TSharedPtr<FNiagaraStackFunctionInputOverrideMergeAdapter> OverrideToAdd;
+	};
+
+	// Collect the graph actions from the adapter and diff first.
+	TArray<TSharedRef<FNiagaraStackFunctionMergeAdapter>> RemoveModules;
+	TArray<TSharedRef<FNiagaraStackFunctionMergeAdapter>> AddModules;
+	TArray<TSharedRef<FNiagaraStackFunctionInputOverrideMergeAdapter>> RemoveInputOverrides;
+	TArray<FAddInputOverrideActionData> AddInputOverrideActionDatas;
+	TArray<TSharedRef<FNiagaraStackFunctionMergeAdapter>> EnableModules;
+	TArray<TSharedRef<FNiagaraStackFunctionMergeAdapter>> DisableModules;
+
+	for (TSharedRef<FNiagaraStackFunctionMergeAdapter> RemovedModule : DiffResults.RemovedBaseModules)
+	{
+		TSharedPtr<FNiagaraStackFunctionMergeAdapter> MatchingModuleAdapter = BaseScriptStackAdapter->GetModuleFunctionById(RemovedModule->GetFunctionCallNode()->NodeGuid);
+		if (MatchingModuleAdapter.IsValid())
+		{
+			RemoveModules.Add(MatchingModuleAdapter.ToSharedRef());
+		}
+	}
+
+	AddModules.Append(DiffResults.AddedOtherModules);
+
+	for (TSharedRef<FNiagaraStackFunctionInputOverrideMergeAdapter> RemovedInputOverrideAdapter : DiffResults.RemovedBaseInputOverrides)
+	{
+		TSharedPtr<FNiagaraStackFunctionMergeAdapter> MatchingModuleAdapter = BaseScriptStackAdapter->GetModuleFunctionById(RemovedInputOverrideAdapter->GetOwningFunctionCall()->NodeGuid);
+		if (MatchingModuleAdapter.IsValid())
+		{
+			TSharedPtr<FNiagaraStackFunctionInputOverrideMergeAdapter> MatchingInputOverrideAdapter = MatchingModuleAdapter->GetInputOverrideByInputName(RemovedInputOverrideAdapter->GetInputName());
+			if (MatchingInputOverrideAdapter.IsValid())
+			{
+				RemoveInputOverrides.Add(MatchingInputOverrideAdapter.ToSharedRef());
+			}
+		}
+	}
+
+	for (TSharedRef<FNiagaraStackFunctionInputOverrideMergeAdapter> AddedInputOverrideAdapter : DiffResults.AddedOtherInputOverrides)
+	{
+		TSharedPtr<FNiagaraStackFunctionMergeAdapter> MatchingModuleAdapter = BaseScriptStackAdapter->GetModuleFunctionById(AddedInputOverrideAdapter->GetOwningFunctionCall()->NodeGuid);
+		if (MatchingModuleAdapter.IsValid())
+		{
+			TSharedPtr<FNiagaraStackFunctionInputOverrideMergeAdapter> MatchingInputOverrideAdapter = MatchingModuleAdapter->GetInputOverrideByInputName(AddedInputOverrideAdapter->GetInputName());
+			if (MatchingInputOverrideAdapter.IsValid())
+			{
+				RemoveInputOverrides.AddUnique(MatchingInputOverrideAdapter.ToSharedRef());
+			}
+
+			FAddInputOverrideActionData AddInputOverrideActionData;
+			AddInputOverrideActionData.TargetFunctionCall = MatchingModuleAdapter->GetFunctionCallNode();
+			AddInputOverrideActionData.OverrideToAdd = AddedInputOverrideAdapter;
+			AddInputOverrideActionDatas.Add(AddInputOverrideActionData);
+		}
+	}
+
+	for (TSharedRef<FNiagaraStackFunctionInputOverrideMergeAdapter> ModifiedInputOverrideAdapter : DiffResults.ModifiedOtherInputOverrides)
+	{
+		TSharedPtr<FNiagaraStackFunctionMergeAdapter> MatchingModuleAdapter = BaseScriptStackAdapter->GetModuleFunctionById(ModifiedInputOverrideAdapter->GetOwningFunctionCall()->NodeGuid);
+		if (MatchingModuleAdapter.IsValid())
+		{
+			TSharedPtr<FNiagaraStackFunctionInputOverrideMergeAdapter> MatchingInputOverrideAdapter = MatchingModuleAdapter->GetInputOverrideByInputName(ModifiedInputOverrideAdapter->GetInputName());
+			if (MatchingInputOverrideAdapter.IsValid())
+			{
+				RemoveInputOverrides.AddUnique(MatchingInputOverrideAdapter.ToSharedRef());
+			}
+
+			FAddInputOverrideActionData AddInputOverrideActionData;
+			AddInputOverrideActionData.TargetFunctionCall = MatchingModuleAdapter->GetFunctionCallNode();
+			AddInputOverrideActionData.OverrideToAdd = ModifiedInputOverrideAdapter;
+			AddInputOverrideActionDatas.Add(AddInputOverrideActionData);
+		}
+	}
+
+	for (TSharedRef<FNiagaraStackFunctionMergeAdapter> EnabledChangedModule : DiffResults.EnabledChangedOtherModules)
+	{
+		TSharedPtr<FNiagaraStackFunctionMergeAdapter> MatchingModuleAdapter = BaseScriptStackAdapter->GetModuleFunctionById(EnabledChangedModule->GetFunctionCallNode()->NodeGuid);
+		if (MatchingModuleAdapter.IsValid())
+		{
+			if (EnabledChangedModule->GetFunctionCallNode()->IsNodeEnabled())
+			{
+				EnableModules.Add(MatchingModuleAdapter.ToSharedRef());
+			}
+			else
+			{
+				DisableModules.Add(MatchingModuleAdapter.ToSharedRef());
+			}
+		}
+	}
+
+	// Apply the graph actions.
+	for (TSharedRef<FNiagaraStackFunctionMergeAdapter> RemoveModule : RemoveModules)
+	{
+		bool bRemoveResults = FNiagaraStackGraphUtilities::RemoveModuleFromStack(*RemoveModule->GetFunctionCallNode());
+		if (bRemoveResults == false)
+		{
+			Results.bSucceeded = false;
+			Results.ErrorMessages.Add(LOCTEXT("RemoveModuleFailedMessage", "Failed to remove module while applying diff"));
+		}
+		else
+		{
+			Results.bModifiedGraph = true;
+		}
+	}
+
+	for (TSharedRef<FNiagaraStackFunctionMergeAdapter> AddModuleAdapter : AddModules)
+	{
+		FApplyDiffResults AddModuleResults = AddModule(BaseScriptStackAdapter->GetUniqueEmitterName(), *BaseScriptStackAdapter->GetScript(), *BaseScriptStackAdapter->GetOutputNode(), AddModuleAdapter);
+		Results.bSucceeded &= AddModuleResults.bSucceeded;
+		Results.bModifiedGraph |= AddModuleResults.bModifiedGraph;
+		Results.ErrorMessages.Append(AddModuleResults.ErrorMessages);
+	}
+
+	for (TSharedRef<FNiagaraStackFunctionInputOverrideMergeAdapter> RemoveInputOverrideItem : RemoveInputOverrides)
+	{
+		FApplyDiffResults RemoveInputOverrideResults = RemoveInputOverride(*BaseScriptStackAdapter->GetScript(), RemoveInputOverrideItem);
+		Results.bSucceeded &= RemoveInputOverrideResults.bSucceeded;
+		Results.bModifiedGraph |= RemoveInputOverrideResults.bModifiedGraph;
+		Results.ErrorMessages.Append(RemoveInputOverrideResults.ErrorMessages);
+	}
+
+	for (const FAddInputOverrideActionData& AddInputOverrideActionData : AddInputOverrideActionDatas)
+	{
+		FApplyDiffResults AddInputOverrideResults = AddInputOverride(BaseScriptStackAdapter->GetUniqueEmitterName(), *BaseScriptStackAdapter->GetScript(), *AddInputOverrideActionData.TargetFunctionCall, AddInputOverrideActionData.OverrideToAdd.ToSharedRef());
+		Results.bSucceeded &= AddInputOverrideResults.bSucceeded;
+		Results.bModifiedGraph |= AddInputOverrideResults.bModifiedGraph;
+		Results.ErrorMessages.Append(AddInputOverrideResults.ErrorMessages);
+	}
+
+	// Apply enabled state last so that it applies to function calls added  from input overrides;
+	for (TSharedRef<FNiagaraStackFunctionMergeAdapter> EnableModule : EnableModules)
+	{
+		FNiagaraStackGraphUtilities::SetModuleIsEnabled(*EnableModule->GetFunctionCallNode(), true);
+	}
+	for (TSharedRef<FNiagaraStackFunctionMergeAdapter> DisableModule : DisableModules)
+	{
+		FNiagaraStackGraphUtilities::SetModuleIsEnabled(*DisableModule->GetFunctionCallNode(), false);
+	}
+
+	return Results;
+}
+
+FNiagaraScriptMergeManager::FApplyDiffResults FNiagaraScriptMergeManager::ApplyEventHandlerDiff(TSharedRef<FNiagaraEmitterMergeAdapter> BaseEmitterAdapter, const FNiagaraEmitterDiffResults& DiffResults)
+{
+	FApplyDiffResults Results;
+	if (DiffResults.RemovedBaseEventHandlers.Num() > 0)
+	{
+		Results.bSucceeded = false;
+		Results.bModifiedGraph = false;
+		Results.ErrorMessages.Add(LOCTEXT("RemovedEventHandlersUnsupported", "Apply diff failed, removed event handlers are currently unsupported."));
+		return Results;
+	}
+
+	// Apply the modifications first since adding new event handlers may invalidate the adapter.
+	for (const FNiagaraModifiedEventHandlerDiffResults& ModifiedEventHandler : DiffResults.ModifiedEventHandlers)
+	{
+		if (ModifiedEventHandler.OtherAdapter->GetEventScriptProperties() == nullptr)
+		{
+			Results.bSucceeded = false;
+			Results.ErrorMessages.Add(LOCTEXT("MissingModifiedEventProperties", "Apply diff failed.  The modified event handler with id: {0} was missing it's event properties."));
+		}
+		else if (ModifiedEventHandler.OtherAdapter->GetOutputNode() == nullptr)
+		{
+			Results.bSucceeded = false;
+			Results.ErrorMessages.Add(LOCTEXT("MissingModifiedEventOutputNode", "Apply diff failed.  The modified event handler with id: {0} was missing it's output node."));
+		}
+		else
+		{
+			TSharedPtr<FNiagaraEventHandlerMergeAdapter> MatchingBaseEventHandlerAdapter = BaseEmitterAdapter->GetEventHandler(ModifiedEventHandler.OtherAdapter->GetUsageId());
+			if (MatchingBaseEventHandlerAdapter.IsValid())
+			{
+				if (ModifiedEventHandler.ChangedProperties.Num() > 0)
+				{
+					CopyPropertiesToBase(MatchingBaseEventHandlerAdapter->GetEditableEventScriptProperties(), ModifiedEventHandler.OtherAdapter->GetEditableEventScriptProperties(), ModifiedEventHandler.ChangedProperties);
+				}
+				if (ModifiedEventHandler.ScriptDiffResults.IsEmpty() == false)
+				{
+					FApplyDiffResults ApplyEventHandlerStackDiffResults = ApplyScriptStackDiff(MatchingBaseEventHandlerAdapter->GetEventStack().ToSharedRef(), ModifiedEventHandler.ScriptDiffResults);
+					Results.bSucceeded &= ApplyEventHandlerStackDiffResults.bSucceeded;
+					Results.bModifiedGraph |= ApplyEventHandlerStackDiffResults.bModifiedGraph;
+					Results.ErrorMessages.Append(ApplyEventHandlerStackDiffResults.ErrorMessages);
+				}
+			}
+		}
+	}
+
+	UNiagaraScriptSource* EmitterSource = CastChecked<UNiagaraScriptSource>(BaseEmitterAdapter->GetEditableEmitter()->GraphSource);
+	UNiagaraGraph* EmitterGraph = EmitterSource->NodeGraph;
+	for (TSharedRef<FNiagaraEventHandlerMergeAdapter> AddedEventHandler : DiffResults.AddedOtherEventHandlers)
+	{
+		if (AddedEventHandler->GetEventScriptProperties() == nullptr)
+		{
+			Results.bSucceeded = false;
+			Results.ErrorMessages.Add(LOCTEXT("MissingAddedEventProperties", "Apply diff failed.  The added event handler with id: {0} was missing it's event properties."));
+		}
+		else if (AddedEventHandler->GetOutputNode() == nullptr)
+		{
+			Results.bSucceeded = false;
+			Results.ErrorMessages.Add(LOCTEXT("MissingAddedEventOutputNode", "Apply diff failed.  The added event handler with id: {0} was missing it's output node."));
+		}
+		else
+		{
+			UNiagaraEmitter* BaseEmitter = BaseEmitterAdapter->GetEditableEmitter();
+			FNiagaraEventScriptProperties AddedEventScriptProperties = *AddedEventHandler->GetEventScriptProperties();
+			AddedEventScriptProperties.Script = NewObject<UNiagaraScript>(BaseEmitter, MakeUniqueObjectName(BaseEmitter, UNiagaraScript::StaticClass(), "EventScript"), EObjectFlags::RF_Transactional);
+			AddedEventScriptProperties.Script->SetUsage(ENiagaraScriptUsage::ParticleEventScript);
+			AddedEventScriptProperties.Script->SetUsageId(FGuid::NewGuid());
+			AddedEventScriptProperties.Script->SetSource(EmitterSource);
+			BaseEmitter->AddEventHandler(AddedEventScriptProperties);
+
+			UNiagaraNodeOutput* EventOutputNode = FNiagaraStackGraphUtilities::ResetGraphForOutput(*EmitterGraph, ENiagaraScriptUsage::ParticleEventScript, AddedEventScriptProperties.Script->GetUsageId());
+			for (TSharedRef<FNiagaraStackFunctionMergeAdapter> ModuleAdapter : AddedEventHandler->GetEventStack()->GetModuleFunctions())
+			{
+				FApplyDiffResults AddModuleResults = AddModule(BaseEmitter->GetUniqueEmitterName(), *AddedEventScriptProperties.Script, *EventOutputNode, ModuleAdapter);
+				Results.bSucceeded &= AddModuleResults.bSucceeded;
+				Results.ErrorMessages.Append(AddModuleResults.ErrorMessages);
+			}
+			Results.bModifiedGraph = true;
+		}
+	}
+	return Results;
+}
+
+FNiagaraScriptMergeManager::FApplyDiffResults FNiagaraScriptMergeManager::ApplyRendererDiff(UNiagaraEmitter& BaseEmitter, const FNiagaraEmitterDiffResults& DiffResults)
+{
+	TArray<UNiagaraRendererProperties*> RenderersToRemove;
+	TArray<UNiagaraRendererProperties*> RenderersToAdd;
+
+	for (TSharedRef<FNiagaraRendererMergeAdapter> RemovedRenderer : DiffResults.RemovedBaseRenderers)
+	{
+		auto FindRendererByMergeId = [=](UNiagaraRendererProperties* Renderer) { return Renderer->GetMergeId() == RemovedRenderer->GetRenderer()->GetMergeId(); };
+		UNiagaraRendererProperties*const* MatchingRendererPtr = BaseEmitter.GetRenderers().FindByPredicate(FindRendererByMergeId);
+		if (MatchingRendererPtr != nullptr)
+		{
+			RenderersToRemove.Add(*MatchingRendererPtr);
+		}
+	}
+
+	for (TSharedRef<FNiagaraRendererMergeAdapter> AddedRenderer : DiffResults.AddedOtherRenderers)
+	{
+		RenderersToAdd.Add(Cast<UNiagaraRendererProperties>(StaticDuplicateObject(AddedRenderer->GetRenderer(), &BaseEmitter)));
+	}
+
+	for (TSharedRef<FNiagaraRendererMergeAdapter> ModifiedRenderer : DiffResults.ModifiedOtherRenderers)
+	{
+		auto FindRendererByMergeId = [=](UNiagaraRendererProperties* Renderer) { return Renderer->GetMergeId() == ModifiedRenderer->GetRenderer()->GetMergeId(); };
+		UNiagaraRendererProperties*const* MatchingRendererPtr = BaseEmitter.GetRenderers().FindByPredicate(FindRendererByMergeId);
+		if (MatchingRendererPtr != nullptr)
+		{
+			RenderersToRemove.Add(*MatchingRendererPtr);
+			RenderersToAdd.Add(Cast<UNiagaraRendererProperties>(StaticDuplicateObject(ModifiedRenderer->GetRenderer(), &BaseEmitter)));
+		}
+	}
+
+	for (UNiagaraRendererProperties* RendererToRemove : RenderersToRemove)
+	{
+		BaseEmitter.RemoveRenderer(RendererToRemove);
+	}
+
+	for (UNiagaraRendererProperties* RendererToAdd : RenderersToAdd)
+	{
+		BaseEmitter.AddRenderer(RendererToAdd);
+	}
+
+	FApplyDiffResults Results;
+	Results.bSucceeded = true;
+	Results.bModifiedGraph = false;
+	return Results;
+}
+
+TSharedRef<FNiagaraEmitterMergeAdapter> FNiagaraScriptMergeManager::GetEmitterMergeAdapterUsingCache(const UNiagaraEmitter& Emitter)
+{
+	FCachedMergeAdapter* CachedMergeAdapter = CachedMergeAdapters.Find(FObjectKey(&Emitter));
+	if (CachedMergeAdapter == nullptr)
+	{
+		CachedMergeAdapter = &CachedMergeAdapters.Add(FObjectKey(&Emitter));
+	}
+
+	if (CachedMergeAdapter->EmitterMergeAdapter.IsValid() == false ||
+		CachedMergeAdapter->EmitterMergeAdapter->GetEditableEmitter() != nullptr ||
+		CachedMergeAdapter->ChangeId != Emitter.GetChangeId())
+	{
+		CachedMergeAdapter->EmitterMergeAdapter = MakeShared<FNiagaraEmitterMergeAdapter>(Emitter);
+		CachedMergeAdapter->ChangeId = Emitter.GetChangeId();
+	}
+
+	return CachedMergeAdapter->EmitterMergeAdapter.ToSharedRef();
+}
+
+TSharedRef<FNiagaraEmitterMergeAdapter> FNiagaraScriptMergeManager::GetEmitterMergeAdapterUsingCache(UNiagaraEmitter& Emitter)
+{
+	FCachedMergeAdapter* CachedMergeAdapter = CachedMergeAdapters.Find(FObjectKey(&Emitter));
+	if (CachedMergeAdapter == nullptr)
+	{
+		CachedMergeAdapter = &CachedMergeAdapters.Add(FObjectKey(&Emitter));
+	}
+
+	if (CachedMergeAdapter->EmitterMergeAdapter.IsValid() == false ||
+		CachedMergeAdapter->EmitterMergeAdapter->GetEditableEmitter() == nullptr ||
+		CachedMergeAdapter->ChangeId != Emitter.GetChangeId())
+	{
+		CachedMergeAdapter->EmitterMergeAdapter = MakeShared<FNiagaraEmitterMergeAdapter>(Emitter);
+		CachedMergeAdapter->ChangeId = Emitter.GetChangeId();
+	}
+
+	return CachedMergeAdapter->EmitterMergeAdapter.ToSharedRef();
+}
+
+
+#undef LOCTEXT_NAMESPACE
\ No newline at end of file
diff --git a/Engine/Plugins/FX/Niagara/Source/NiagaraEditor/Private/NiagaraScriptMergeManager.h b/Engine/Plugins/FX/Niagara/Source/NiagaraEditor/Private/NiagaraScriptMergeManager.h
new file mode 100644
--- /dev/null
+++ b/Engine/Plugins/FX/Niagara/Source/NiagaraEditor/Private/NiagaraScriptMergeManager.h
@@ -0,0 +1,349 @@
+// Copyright 1998-2017 Epic Games, Inc. All Rights Reserved.
+
+#pragma once
+
+#include "ViewModels/Stack/NiagaraParameterHandle.h"
+#include "NiagaraModule.h"
+#include "NiagaraTypes.h"
+#include "NiagaraCommon.h"
+
+#include "Templates/SharedPointer.h"
+#include "UObject/WeakObjectPtr.h"
+#include "UObject/WeakObjectPtrTemplates.h"
+#include "UObject/ObjectKey.h"
+
+class UNiagaraEmitter;
+class UNiagaraScript;
+class UNiagaraNodeOutput;
+class UNiagaraNodeFunctionCall;
+class UNiagaraNodeParameterMapSet;
+class UEdGraphNode;
+class UEdGraphPin;
+class UNiagaraDataInterface;
+struct FNiagaraEventScriptProperties;
+class UNiagaraRendererProperties;
+
+class FNiagaraStackFunctionMergeAdapter;
+
+class FNiagaraStackFunctionInputOverrideMergeAdapter
+{
+public:
+	FNiagaraStackFunctionInputOverrideMergeAdapter(
+		FString InUniqueEmitterName, 
+		UNiagaraScript& InOwningScript, 
+		UNiagaraNodeFunctionCall& InOwningFunctionCallNode, 
+		UEdGraphPin& InOverridePin);
+
+	FNiagaraStackFunctionInputOverrideMergeAdapter(
+		FString InUniqueEmitterName,
+		UNiagaraScript& InOwningScript,
+		UNiagaraNodeFunctionCall& InOwningFunctionCallNode,
+		FString InInputName,
+		FNiagaraVariable InRapidIterationParameter);
+
+	UNiagaraScript* GetOwningScript() const;
+	UNiagaraNodeFunctionCall* GetOwningFunctionCall() const;
+	FString GetInputName() const;
+	UNiagaraNodeParameterMapSet* GetOverrideNode() const;
+	UEdGraphPin* GetOverridePin() const;
+
+	TOptional<FString> GetLocalValueString() const;
+	TOptional<FNiagaraVariable> GetLocalValueRapidIterationParameter() const;
+	TOptional<FNiagaraParameterHandle> GetLinkedValueHandle() const;
+	UNiagaraDataInterface* GetDataValueObject() const;
+	TSharedPtr<FNiagaraStackFunctionMergeAdapter> GetDynamicValueFunction() const;
+
+private:
+	FString UniqueEmitterName;
+	TWeakObjectPtr<UNiagaraScript> OwningScript;
+	TWeakObjectPtr<UNiagaraNodeFunctionCall> OwningFunctionCallNode;
+	FString InputName;
+
+	TWeakObjectPtr<UNiagaraNodeParameterMapSet> OverrideNode;
+	UEdGraphPin* OverridePin;
+
+	TOptional<FString> LocalValueString;
+	TOptional<FNiagaraVariable> LocalValueRapidIterationParameter;
+	TOptional<FNiagaraParameterHandle> LinkedValueHandle;
+	UNiagaraDataInterface* DataValueObject;
+	TSharedPtr<FNiagaraStackFunctionMergeAdapter> DynamicValueFunction;
+};
+
+class FNiagaraStackFunctionMergeAdapter
+{
+public:
+	FNiagaraStackFunctionMergeAdapter(FString InUniqueEmitterName, UNiagaraScript& InOwningScript, UNiagaraNodeFunctionCall& InFunctionCallNode, int32 InStackIndex);
+
+	UNiagaraNodeFunctionCall* GetFunctionCallNode() const;
+
+	int32 GetStackIndex() const;
+
+	const TArray<TSharedRef<FNiagaraStackFunctionInputOverrideMergeAdapter>>& GetInputOverrides() const;
+
+	TSharedPtr<FNiagaraStackFunctionInputOverrideMergeAdapter> GetInputOverrideByInputName(FString InputName) const;
+
+private:
+	FString UniqueEmitterName;
+	TWeakObjectPtr<UNiagaraScript> OwningScript;
+	TWeakObjectPtr<UNiagaraNodeFunctionCall> FunctionCallNode;
+	int32 StackIndex;
+
+	TArray<TSharedRef<FNiagaraStackFunctionInputOverrideMergeAdapter>> InputOverrides;
+};
+
+class FNiagaraScriptStackMergeAdapter
+{
+public:
+	FNiagaraScriptStackMergeAdapter(UNiagaraNodeOutput& InOutputNode, UNiagaraScript& InScript, FString InUniqueEmitterName);
+
+	UNiagaraNodeOutput* GetOutputNode() const;
+
+	UNiagaraScript* GetScript() const;
+
+	FString GetUniqueEmitterName() const;
+
+	const TArray<TSharedRef<FNiagaraStackFunctionMergeAdapter>>& GetModuleFunctions() const;
+
+	TSharedPtr<FNiagaraStackFunctionMergeAdapter> GetModuleFunctionById(FGuid FunctionCallNodeId) const;
+
+private:
+	TWeakObjectPtr<UNiagaraNodeOutput> OutputNode;
+	TWeakObjectPtr<UNiagaraScript> Script;
+	FString UniqueEmitterName;
+	TArray<TSharedRef<FNiagaraStackFunctionMergeAdapter>> ModuleFunctions;
+};
+
+class FNiagaraEventHandlerMergeAdapter
+{
+public:
+	FNiagaraEventHandlerMergeAdapter(const UNiagaraEmitter& InEmitter, const FNiagaraEventScriptProperties* InEventScriptProperties, UNiagaraNodeOutput* InOutputNode);
+	FNiagaraEventHandlerMergeAdapter(const UNiagaraEmitter& InEmitter, FNiagaraEventScriptProperties* InEventScriptProperties, UNiagaraNodeOutput* InOutputNode);
+	FNiagaraEventHandlerMergeAdapter(const UNiagaraEmitter& InEmitter, UNiagaraNodeOutput* InOutputNode);
+
+	FGuid GetUsageId() const;
+	const UNiagaraEmitter* GetEmitter() const;
+	const FNiagaraEventScriptProperties* GetEventScriptProperties() const;
+	FNiagaraEventScriptProperties* GetEditableEventScriptProperties() const;
+	UNiagaraNodeOutput* GetOutputNode() const;
+	TSharedPtr<FNiagaraScriptStackMergeAdapter> GetEventStack() const;
+
+private:
+	void Initialize(const UNiagaraEmitter& InEmitter, const FNiagaraEventScriptProperties* InEventScriptProperties, FNiagaraEventScriptProperties* InEditableEventScriptProperties, UNiagaraNodeOutput* InOutputNode);
+
+private:
+	TWeakObjectPtr<UNiagaraEmitter> Emitter;
+	const FNiagaraEventScriptProperties* EventScriptProperties;
+	FNiagaraEventScriptProperties* EditableEventScriptProperties;
+	TWeakObjectPtr<UNiagaraNodeOutput> OutputNode;
+
+	TSharedPtr<FNiagaraScriptStackMergeAdapter> EventStack;
+};
+
+class FNiagaraRendererMergeAdapter
+{
+public:
+	FNiagaraRendererMergeAdapter(UNiagaraRendererProperties& InRenderer);
+
+	UNiagaraRendererProperties* GetRenderer();
+
+private:
+	TWeakObjectPtr<UNiagaraRendererProperties> Renderer;
+};
+
+class FNiagaraEmitterMergeAdapter
+{
+public:
+	FNiagaraEmitterMergeAdapter(const UNiagaraEmitter& InEmitter);
+	FNiagaraEmitterMergeAdapter(UNiagaraEmitter& InEmitter);
+
+	UNiagaraEmitter* GetEditableEmitter() const;
+
+	TSharedPtr<FNiagaraScriptStackMergeAdapter> GetEmitterSpawnStack() const;
+	TSharedPtr<FNiagaraScriptStackMergeAdapter> GetEmitterUpdateStack() const;
+	TSharedPtr<FNiagaraScriptStackMergeAdapter> GetParticleSpawnStack() const;
+	TSharedPtr<FNiagaraScriptStackMergeAdapter> GetParticleUpdateStack() const;
+	const TArray<TSharedRef<FNiagaraEventHandlerMergeAdapter>> GetEventHandlers() const;
+	const TArray<TSharedRef<FNiagaraRendererMergeAdapter>> GetRenderers() const;
+
+	TSharedPtr<FNiagaraScriptStackMergeAdapter> GetScriptStack(ENiagaraScriptUsage Usage, FGuid UsageId);
+
+	TSharedPtr<FNiagaraEventHandlerMergeAdapter> GetEventHandler(FGuid EventScriptUsageId);
+
+	TSharedPtr<FNiagaraRendererMergeAdapter> GetRenderer(FGuid RendererMergeId);
+
+private:
+	void Initialize(const UNiagaraEmitter& InEmitter, UNiagaraEmitter* InEditableEmitter);
+
+private:
+	TWeakObjectPtr<const UNiagaraEmitter> Emitter;
+	TWeakObjectPtr<UNiagaraEmitter> EditableEmitter;
+	TSharedPtr<FNiagaraScriptStackMergeAdapter> EmitterSpawnStack;
+	TSharedPtr<FNiagaraScriptStackMergeAdapter> EmitterUpdateStack;
+	TSharedPtr<FNiagaraScriptStackMergeAdapter> ParticleSpawnStack;
+	TSharedPtr<FNiagaraScriptStackMergeAdapter> ParticleUpdateStack;
+	TArray<TSharedRef<FNiagaraEventHandlerMergeAdapter>> EventHandlers;
+	TArray<TSharedRef<FNiagaraRendererMergeAdapter>> Renderers;
+};
+
+struct FNiagaraScriptStackDiffResults
+{
+	FNiagaraScriptStackDiffResults();
+
+	bool IsEmpty() const;
+
+	bool IsValid() const;
+
+	void AddError(FText ErrorMessage);
+
+	const TArray<FText>& GetErrorMessages() const;
+
+	TArray<TSharedRef<FNiagaraStackFunctionMergeAdapter>> RemovedBaseModules;
+	TArray<TSharedRef<FNiagaraStackFunctionMergeAdapter>> AddedOtherModules;
+
+	TArray<TSharedRef<FNiagaraStackFunctionMergeAdapter>> MovedBaseModules;
+	TArray<TSharedRef<FNiagaraStackFunctionMergeAdapter>> MovedOtherModules;
+
+	TArray<TSharedRef<FNiagaraStackFunctionMergeAdapter>> EnabledChangedBaseModules;
+	TArray<TSharedRef<FNiagaraStackFunctionMergeAdapter>> EnabledChangedOtherModules;
+
+	TArray<TSharedRef<FNiagaraStackFunctionInputOverrideMergeAdapter>> RemovedBaseInputOverrides;
+	TArray<TSharedRef<FNiagaraStackFunctionInputOverrideMergeAdapter>> AddedOtherInputOverrides;
+	TArray<TSharedRef<FNiagaraStackFunctionInputOverrideMergeAdapter>> ModifiedBaseInputOverrides;
+	TArray<TSharedRef<FNiagaraStackFunctionInputOverrideMergeAdapter>> ModifiedOtherInputOverrides;
+
+private:
+	bool bIsValid;
+	TArray<FText> ErrorMessages;
+};
+
+struct FNiagaraModifiedEventHandlerDiffResults
+{
+	TSharedPtr<FNiagaraEventHandlerMergeAdapter> BaseAdapter;
+	TSharedPtr<FNiagaraEventHandlerMergeAdapter> OtherAdapter;
+	TArray<UProperty*> ChangedProperties;
+	FNiagaraScriptStackDiffResults ScriptDiffResults;
+};
+
+struct FNiagaraEmitterDiffResults
+{
+	FNiagaraEmitterDiffResults();
+
+	bool IsValid() const;
+
+	bool IsEmpty() const;
+
+	void AddError(FText ErrorMessage);
+
+	const TArray<FText>& GetErrorMessages() const;
+
+	FString GetErrorMessagesString() const;
+
+	TArray<UProperty*> DifferentEmitterProperties;
+
+	FNiagaraScriptStackDiffResults EmitterSpawnDiffResults;
+	FNiagaraScriptStackDiffResults EmitterUpdateDiffResults;
+	FNiagaraScriptStackDiffResults ParticleSpawnDiffResults;
+	FNiagaraScriptStackDiffResults ParticleUpdateDiffResults;
+
+	TArray<TSharedRef<FNiagaraEventHandlerMergeAdapter>> RemovedBaseEventHandlers;
+	TArray<TSharedRef<FNiagaraEventHandlerMergeAdapter>> AddedOtherEventHandlers;
+	TArray<FNiagaraModifiedEventHandlerDiffResults> ModifiedEventHandlers;
+
+	TArray<TSharedRef<FNiagaraRendererMergeAdapter>> RemovedBaseRenderers;
+	TArray<TSharedRef<FNiagaraRendererMergeAdapter>> AddedOtherRenderers;
+	TArray<TSharedRef<FNiagaraRendererMergeAdapter>> ModifiedBaseRenderers;
+	TArray<TSharedRef<FNiagaraRendererMergeAdapter>> ModifiedOtherRenderers;
+
+private:
+	bool bIsValid;
+	TArray<FText> ErrorMessages;
+};
+
+class FNiagaraScriptMergeManager
+{
+public:
+	struct FApplyDiffResults
+	{
+		FApplyDiffResults()
+			: bSucceeded(true)
+			, bModifiedGraph(false)
+		{
+		}
+
+		bool bSucceeded;
+		bool bModifiedGraph;
+		TArray<FText> ErrorMessages;
+	};
+
+public:
+	INiagaraModule::FMergeEmitterResults MergeEmitter(UNiagaraEmitter& Source, UNiagaraEmitter& LastMergedSource, UNiagaraEmitter& Instance);
+
+	static TSharedRef<FNiagaraScriptMergeManager> Get();
+
+	FNiagaraEmitterDiffResults DiffEmitters(UNiagaraEmitter& BaseEmitter, UNiagaraEmitter& OtherEmitter);
+
+	bool IsMergeableScriptUsage(ENiagaraScriptUsage ScriptUsage) const;
+
+	bool HasBaseModule(const UNiagaraEmitter& BaseEmitter, ENiagaraScriptUsage ScriptUsage, FGuid ScriptUsageId, FGuid ModuleId);
+
+	bool IsModuleInputDifferentFromBase(UNiagaraEmitter& Emitter, const UNiagaraEmitter& BaseEmitter, ENiagaraScriptUsage ScriptUsage, FGuid ScriptUsageId, FGuid ModuleId, FString InputName);
+
+	FApplyDiffResults ResetModuleInputToBase(UNiagaraEmitter& Emitter, const UNiagaraEmitter& BaseEmitter, ENiagaraScriptUsage ScriptUsage, FGuid ScriptUsageId, FGuid ModuleId, FString InputName);
+
+	bool HasBaseEventHandler(const UNiagaraEmitter& BaseEmitter, FGuid EventScriptUsageId);
+
+	bool IsEventHandlerPropertySetDifferentFromBase(UNiagaraEmitter& Emitter, const UNiagaraEmitter& BaseEmitter, FGuid EventScriptUsageId);
+
+	void ResetEventHandlerPropertySetToBase(UNiagaraEmitter& Emitter, const UNiagaraEmitter& BaseEmitter, FGuid EventScriptUsageId);
+
+	bool HasBaseRenderer(const UNiagaraEmitter& BaseEmitter, FGuid RendererMergeId);
+
+	bool IsRendererDifferentFromBase(UNiagaraEmitter& Emitter, const UNiagaraEmitter& BaseEmitter, FGuid RendererMergeId);
+
+	void ResetRendererToBase(UNiagaraEmitter& Emitter, const UNiagaraEmitter& BaseEmitter, FGuid RendererMergeId);
+
+	bool IsEmitterEditablePropertySetDifferentFromBase(const UNiagaraEmitter& Emitter, const UNiagaraEmitter& BaseEmitter);
+
+	void ResetEmitterEditablePropertySetToBase(UNiagaraEmitter& Emitter, const UNiagaraEmitter& BaseEmitter);
+
+	void DiffEventHandlers(const TArray<TSharedRef<FNiagaraEventHandlerMergeAdapter>>& BaseEventHandlers, const TArray<TSharedRef<FNiagaraEventHandlerMergeAdapter>>& OtherEventHandlers, FNiagaraEmitterDiffResults& DiffResults);
+
+	void DiffRenderers(const TArray<TSharedRef<FNiagaraRendererMergeAdapter>>& BaseRenderers, const TArray<TSharedRef<FNiagaraRendererMergeAdapter>>& OtherRenderers, FNiagaraEmitterDiffResults& DiffResults);
+
+	void DiffScriptStacks(TSharedRef<FNiagaraScriptStackMergeAdapter> BaseScriptStackAdapter, TSharedRef<FNiagaraScriptStackMergeAdapter> OtherScriptStackAdapter, FNiagaraScriptStackDiffResults& DiffResults);
+
+	void DiffFunctionInputs(TSharedRef<FNiagaraStackFunctionMergeAdapter> BaseFunctionAdapter, TSharedRef<FNiagaraStackFunctionMergeAdapter> OtherFunctionAdapter, FNiagaraScriptStackDiffResults& DiffResults);
+
+	void DiffEditableProperties(const void* BaseDataAddress, const void* OtherDataAddress, UStruct& Struct, TArray<UProperty*>& OutDifferentProperties);
+
+private:
+	TOptional<bool> DoFunctionInputOverridesMatch(TSharedRef<FNiagaraStackFunctionInputOverrideMergeAdapter> BaseFunctionInputAdapter, TSharedRef<FNiagaraStackFunctionInputOverrideMergeAdapter> OtherFunctionInputAdapter);
+
+	FApplyDiffResults ApplyScriptStackDiff(TSharedRef<FNiagaraScriptStackMergeAdapter> BaseScriptStackAdapter, const FNiagaraScriptStackDiffResults& DiffResults);
+
+	FApplyDiffResults ApplyEventHandlerDiff(TSharedRef<FNiagaraEmitterMergeAdapter> BaseEmitterAdapter, const FNiagaraEmitterDiffResults& DiffResults);
+
+	FApplyDiffResults ApplyRendererDiff(UNiagaraEmitter& BaseEmitter, const FNiagaraEmitterDiffResults& DiffResults);
+
+	FApplyDiffResults AddModule(FString UniqueEmitterName, UNiagaraScript& OwningScript, UNiagaraNodeOutput& TargetOutputNode, TSharedRef<FNiagaraStackFunctionMergeAdapter> AddModule);
+
+	FApplyDiffResults RemoveInputOverride(UNiagaraScript& OwningScript, TSharedRef<FNiagaraStackFunctionInputOverrideMergeAdapter> OverrideToRemove);
+
+	FApplyDiffResults AddInputOverride(FString UniqueEmitterName, UNiagaraScript& OwningScript, UNiagaraNodeFunctionCall& TargetFunctionCall, TSharedRef<FNiagaraStackFunctionInputOverrideMergeAdapter> OverrideToAdd);
+
+	void CopyPropertiesToBase(void* BaseDataAddress, const void* OtherDataAddress, TArray<UProperty*> PropertiesToCopy);
+
+private:
+	TSharedRef<FNiagaraEmitterMergeAdapter> GetEmitterMergeAdapterUsingCache(const UNiagaraEmitter& Emitter);
+
+	TSharedRef<FNiagaraEmitterMergeAdapter> GetEmitterMergeAdapterUsingCache(UNiagaraEmitter& Emitter);
+
+private:
+	struct FCachedMergeAdapter
+	{
+		FGuid ChangeId;
+		TSharedPtr<FNiagaraEmitterMergeAdapter> EmitterMergeAdapter;
+	};
+
+	TMap<FObjectKey, FCachedMergeAdapter> CachedMergeAdapters;
+};
\ No newline at end of file
