Copying Orion-Staging @ 3748653 (Orion/Dev-General @ 3722124) to //UE4/Main

diff --git a/Engine/Plugins/FX/Niagara/Source/NiagaraEditor/Private/NiagaraScriptSource.cpp b/Engine/Plugins/FX/Niagara/Source/NiagaraEditor/Private/NiagaraScriptSource.cpp
--- a/Engine/Plugins/FX/Niagara/Source/NiagaraEditor/Private/NiagaraScriptSource.cpp
+++ b/Engine/Plugins/FX/Niagara/Source/NiagaraEditor/Private/NiagaraScriptSource.cpp
@@ -25,6 +25,14 @@
 #include "EdGraphSchema_Niagara.h"
 #include "TokenizedMessage.h"
 #include "NiagaraEditorUtilities.h"
+#include "NiagaraCustomVersion.h"
+#include "NiagaraStackGraphUtilities.h"
+#include "NiagaraNodeFunctionCall.h"
+#include "NiagaraNodeParameterMapSet.h"
+#include "INiagaraEditorTypeUtilities.h"
+
+DECLARE_CYCLE_STAT(TEXT("Niagara - ScriptSource - Compile"), STAT_NiagaraEditor_ScriptSource_Compile, STATGROUP_NiagaraEditor);
+DECLARE_CYCLE_STAT(TEXT("Niagara - ScriptSource - PreCompile"), STAT_NiagaraEditor_ScriptSource_PreCompile, STATGROUP_NiagaraEditor);
 
 UNiagaraScriptSource::UNiagaraScriptSource(const FObjectInitializer& ObjectInitializer)
 	: Super(ObjectInitializer), bIsPrecompiled(false)
@@ -34,10 +42,215 @@ UNiagaraScriptSource::UNiagaraScriptSource(const FObjectInitializer& ObjectIniti
 void UNiagaraScriptSource::PostLoad()
 {
 	Super::PostLoad();
-	// We need to make sure that the node-graph is already resolved b/c we may be asked IsSyncrhonized later...
+
 	if (NodeGraph)
 	{
+		// We need to make sure that the node-graph is already resolved b/c we may be asked IsSyncrhonized later...
 		NodeGraph->ConditionalPostLoad();
+
+		// Hook up event handlers so the on changed handler can be called correctly.
+		NodeGraph->AddOnGraphChangedHandler(FOnGraphChanged::FDelegate::CreateUObject(this, &UNiagaraScriptSource::OnGraphChanged));
+		NodeGraph->OnDataInterfaceChanged().AddUObject(this, &UNiagaraScriptSource::OnGraphDataInterfaceChanged);
+	}
+
+	const int32 NiagaraVer = GetLinkerCustomVersion(FNiagaraCustomVersion::GUID);
+	if (NiagaraVer < FNiagaraCustomVersion::AddedScriptRapidIterationVariables)
+	{
+		UpgradeForScriptRapidIterationVariables();
+	}
+}
+
+void UNiagaraScriptSource::UpgradeForScriptRapidIterationVariables()
+{
+	// In this version we added the ability to rapidly iterate on top-level module parameters. This means that there are 
+	// loose set nodes that were just setting constants that need to be removed from the graph as well as needing to build the script
+	// rapid iteration parameters list.
+	TArray<UNiagaraNodeOutput*> OutputNodes;
+	NodeGraph->GetNodesOfClass(OutputNodes);
+	UNiagaraEmitter* ParentEmitter = Cast<UNiagaraEmitter>(GetOuter());
+	UNiagaraSystem* ParentSystem = Cast<UNiagaraSystem>(GetOuter());
+	FNiagaraEditorModule& NiagaraEditorModule = FModuleManager::GetModuleChecked<FNiagaraEditorModule>("NiagaraEditor");
+
+	if (ParentEmitter != nullptr || ParentSystem != nullptr)
+	{
+		TArray<UNiagaraScript*> Scripts;
+		if (ParentEmitter != nullptr)
+		{
+			ParentEmitter->ConditionalPostLoad();
+			ParentEmitter->GetScripts(Scripts, false);
+		}
+		else if (ParentSystem != nullptr)
+		{
+			ParentSystem->ConditionalPostLoad();
+			Scripts.Add(ParentSystem->GetSystemSpawnScript(true));
+			Scripts.Add(ParentSystem->GetSystemUpdateScript(true));
+			Scripts.Add(ParentSystem->GetSystemSpawnScript(false));
+			Scripts.Add(ParentSystem->GetSystemUpdateScript(false));
+		}
+
+		for (UNiagaraScript* Script : Scripts)
+		{
+			ensure(Script->RapidIterationParameters.GetParameterDataArray().Num() == 0);
+		}
+
+		bool bNeedsRecompile = false;
+
+		TArray<UEdGraphPin*> PinsToRemove;
+
+		for (UNiagaraNodeOutput* Output : OutputNodes)
+		{
+			FNiagaraParameterMapHistoryBuilder Builder;
+			Builder.BuildParameterMaps(Output);
+				
+			TArray<UNiagaraNodeFunctionCall*> FunctionCallNodes;
+			TArray<UNiagaraNode*> Nodes;
+			NodeGraph->BuildTraversal(Nodes, Output);
+			for (UNiagaraNode* Node : Nodes)
+			{
+				UNiagaraNodeFunctionCall* FuncCallNode = Cast<UNiagaraNodeFunctionCall>(Node);
+				if (FuncCallNode != nullptr)
+				{
+					FunctionCallNodes.Add(FuncCallNode);
+				}
+			}
+
+			TArray<const UEdGraphPin*> CandidateInputPins;
+			TArray<FNiagaraVariable> CandidateVariables;
+			for (UNiagaraNodeFunctionCall* FunctionCallNode : FunctionCallNodes)
+			{
+				TArray<const UEdGraphPin*> InputPins;
+				FNiagaraStackGraphUtilities::GetStackFunctionInputPins(*FunctionCallNode, InputPins, FNiagaraStackGraphUtilities::ENiagaraGetStackFunctionInputPinsOptions::ModuleInputsOnly);
+
+				for (const UEdGraphPin* InputPin : InputPins)
+				{
+					CandidateInputPins.Add(InputPin);
+					const UEdGraphSchema_Niagara* NiagaraSchema = CastChecked<UEdGraphSchema_Niagara>(InputPin->GetSchema());
+					FNiagaraVariable FunctionVar = NiagaraSchema->PinToNiagaraVariable(InputPin);
+					TMap<FString, FString> AliasMap = TMap<FString, FString>();
+					AliasMap.Add(TEXT("Module"), FunctionCallNode->GetFunctionName());
+					if (ParentEmitter)
+					{
+						AliasMap.Add(TEXT("Emitter"), ParentEmitter->GetUniqueEmitterName());
+					}
+					FunctionVar = FNiagaraParameterMapHistory::ResolveAliases(FunctionVar, AliasMap, TEXT("."));
+					CandidateVariables.Add(FunctionVar);
+				}
+			}
+
+			UE_LOG(LogNiagaraEditor, Log, TEXT("Converting OutputNode: %s"), *Output->GetNodeTitle(ENodeTitleType::FullTitle).ToString());
+
+			// As of this version there aren't multiple instances of Event graphs, so we ignore the occurrence parameter.
+			// We do however have the solo and multi versions of the system scripts, so we need to be able to set multiple scripts.
+			TArray<UNiagaraScript*> ActiveScripts;
+			for (UNiagaraScript* Script : Scripts)
+			{
+				if (Script->ContainsUsage(Output->GetUsage()))
+				{
+					ActiveScripts.Add(Script);
+				}
+			}
+
+			if (ActiveScripts.Num() == 0)
+			{
+				continue;
+			}
+
+			for (int32 i = 0; i < CandidateInputPins.Num(); i++)
+			{
+				const UEdGraphPin* Pin = CandidateInputPins[i];
+				FNiagaraVariable CandidateVar = CandidateVariables[i];
+
+				check(Builder.Histories.Num() == 1);
+				const UEdGraphSchema_Niagara* NiagaraSchema = CastChecked<UEdGraphSchema_Niagara>(Pin->GetSchema());
+				FNiagaraVariable FunctionVar = NiagaraSchema->PinToNiagaraVariable(Pin);
+				FNiagaraTypeDefinition PinDefinition = FunctionVar.GetType();
+
+				// Now check to see if this is one of the appropriate types:
+				if (PinDefinition == FNiagaraTypeDefinition::GetBoolDef() || PinDefinition.IsEnum() || PinDefinition == FNiagaraTypeDefinition::GetParameterMapDef())
+				{
+					continue;
+				}
+				
+				// Is this pin even in the subgraph we care about right now?
+				int32 ParamIdx = Builder.Histories[0].FindVariable(CandidateVar.GetName(), CandidateVar.GetType());
+				if (ParamIdx == INDEX_NONE)
+				{
+					continue;
+				}
+
+				FString DefaultValue;
+				// Now check to see if anyone writes to this variable before reading in the graph.
+				if (Builder.Histories[0].PerVariableReadHistory[ParamIdx].Num() != 0)
+				{
+					const UEdGraphPin* FirstReadPin = Builder.Histories[0].PerVariableReadHistory[ParamIdx][0].Key;
+					const UEdGraphPin* FirstWritePin = Builder.Histories[0].PerVariableReadHistory[ParamIdx][0].Value;
+					const UEdGraphPin* DefaultValuePin = Builder.Histories[0].GetDefaultValuePin(ParamIdx);
+
+					UNiagaraNodeParameterMapSet* MapSetNode = nullptr;
+					if (FirstWritePin != nullptr)
+					{
+						MapSetNode = Cast<UNiagaraNodeParameterMapSet>(FirstWritePin->GetOwningNode());
+					}
+
+					if (FirstWritePin != nullptr && FirstWritePin->LinkedTo.Num() == 0 && nullptr != MapSetNode) // Set by a ParameterMapSet pin linked to a constant...
+					{
+						DefaultValue = FirstWritePin->DefaultValue;
+						PinsToRemove.Add((UEdGraphPin*)FirstWritePin);
+						UE_LOG(LogNiagaraEditor, Log, TEXT("Removing Set pin: %s"), *Builder.Histories[0].Variables[ParamIdx].GetName().ToString());
+					}
+					else if (FirstWritePin != nullptr && FirstWritePin->LinkedTo.Num() > 0) // The Set node is set by something other than a constant, so skip it.
+					{
+						continue;
+					}
+					else if (FirstWritePin == nullptr && DefaultValuePin != nullptr) // Get the default value of a pin...
+					{
+						DefaultValue = DefaultValuePin->DefaultValue;
+					}
+
+					FNiagaraVariable FinalVar = Builder.Histories[0].Variables[ParamIdx];
+					FinalVar = FNiagaraParameterMapHistory::ConvertVariableToRapidIterationConstantName(FinalVar, ParentEmitter ? *ParentEmitter->GetUniqueEmitterName() : nullptr);
+
+					TSharedPtr<INiagaraEditorTypeUtilities> TypeEditorUtilities = NiagaraEditorModule.GetTypeUtilities(FinalVar.GetType());
+					if (TypeEditorUtilities.IsValid() && TypeEditorUtilities->CanHandlePinDefaults())
+					{
+						bool bValueSet = TypeEditorUtilities->SetValueFromPinDefaultString(DefaultValue, FinalVar);
+						if (bValueSet)
+						{
+							UE_LOG(LogNiagaraEditor, Log, TEXT("Adding Rapid Iteration Variable: %s  Default Value: %s"), *FinalVar.GetName().ToString(), *DefaultValue);
+						}
+						else
+						{
+							UE_LOG(LogNiagaraEditor, Warning, TEXT("Adding Rapid Iteration Variable: %s  Default Value '%s' could not be set for type %s"),
+								*FinalVar.GetName().ToString(), *DefaultValue, *FinalVar.GetType().GetName());
+						}
+					}
+
+					for (UNiagaraScript* ActiveScript : ActiveScripts)
+					{
+						ActiveScript->RapidIterationParameters.AddParameter(FinalVar);
+					}
+				}					
+			}
+		}
+
+		for (UEdGraphPin* PinToRemove : PinsToRemove)
+		{
+			if (PinToRemove->IsPendingKill())
+			{
+				continue;
+			}
+			UNiagaraNodeParameterMapSet* SetNode = Cast<UNiagaraNodeParameterMapSet>(PinToRemove->GetOwningNode());
+			if (SetNode)
+			{
+				SetNode->RemovePin(PinToRemove);
+				bNeedsRecompile = true;
+			}
+		}
+
+		if (bNeedsRecompile)
+		{
+			NodeGraph->NotifyGraphNeedsRecompile();
+		}
 	}
 }
 
@@ -95,6 +308,7 @@ bool UNiagaraScriptSource::IsPreCompiled() const
 
 void UNiagaraScriptSource::PreCompile(UNiagaraEmitter* InEmitter, bool bClearErrors)
 {
+	SCOPE_CYCLE_COUNTER(STAT_NiagaraEditor_ScriptSource_PreCompile);
 	if (!bIsPrecompiled)
 	{
 		bIsPrecompiled = true;
@@ -177,9 +391,62 @@ void UNiagaraScriptSource::PostCompile()
 	NodeGraphDeepCopy = nullptr;
 }
 
+bool UNiagaraScriptSource::PostLoadFromEmitter(UNiagaraEmitter& OwningEmitter)
+{
+	bool bNeedsCompile = false;
+	const int32 NiagaraCustomVersion = GetLinkerCustomVersion(FNiagaraCustomVersion::GUID);
+	if (NiagaraCustomVersion < FNiagaraCustomVersion::ScriptsNowUseAGuidForIdentificationInsteadOfAnIndex)
+	{
+		TArray<UNiagaraNodeOutput*> OutputNodes;
+		NodeGraph->GetNodesOfClass<UNiagaraNodeOutput>(OutputNodes);
+		for (int32 i = 0; i < OwningEmitter.GetEventHandlers().Num(); i++)
+		{
+			const FNiagaraEventScriptProperties& EventScriptProperties = OwningEmitter.GetEventHandlers()[i];
+			EventScriptProperties.Script->SetUsageId(FGuid::NewGuid());
+
+			auto FindOutputNodeByUsageIndex = [=](UNiagaraNodeOutput* OutputNode) 
+			{ 
+				return OutputNode->GetUsage() == ENiagaraScriptUsage::ParticleEventScript && OutputNode->ScriptTypeIndex_DEPRECATED == EventScriptProperties.Script->UsageIndex_DEPRECATED; 
+			};
+			UNiagaraNodeOutput** MatchingOutputNodePtr = OutputNodes.FindByPredicate(FindOutputNodeByUsageIndex);
+			if (MatchingOutputNodePtr != nullptr)
+			{
+				UNiagaraNodeOutput* MatchingOutputNode = *MatchingOutputNodePtr;
+				MatchingOutputNode->SetUsageId(EventScriptProperties.Script->GetUsageId());
+			}
+		}
+		bNeedsCompile = true;
+	}
+	return bNeedsCompile;
+}
+
+bool UNiagaraScriptSource::AddModuleIfMissing(FString ModulePath, ENiagaraScriptUsage Usage, bool& bOutFoundModule)
+{
+	FStringAssetReference SystemUpdateScriptRef(ModulePath);
+	FAssetData ModuleScriptAsset;
+	ModuleScriptAsset.ObjectPath = SystemUpdateScriptRef.GetAssetPathName();
+	bOutFoundModule = false;
+
+	if (ModuleScriptAsset.IsValid())
+	{
+		bOutFoundModule = true;
+		if (UNiagaraNodeOutput* OutputNode = NodeGraph->FindOutputNode(Usage))
+		{
+			TArray<UNiagaraNodeFunctionCall*> FoundCalls;
+			if (!FNiagaraStackGraphUtilities::FindScriptModulesInStack(ModuleScriptAsset, *OutputNode, FoundCalls))
+			{
+				FNiagaraStackGraphUtilities::AddScriptModuleToStack(ModuleScriptAsset, *OutputNode);
+				return true;
+			}
+		}
+	}
+
+	return false;
+}
 
 ENiagaraScriptCompileStatus UNiagaraScriptSource::Compile(UNiagaraScript* ScriptOwner, FString& OutGraphLevelErrorMessages)
 {
+	SCOPE_CYCLE_COUNTER(STAT_NiagaraEditor_ScriptSource_Compile);
 	bool bDoPostCompile = false;
 	if (!bIsPrecompiled)
 	{
@@ -221,6 +488,16 @@ ENiagaraScriptCompileStatus UNiagaraScriptSource::Compile(UNiagaraScript* Script
 
 }
 
+void UNiagaraScriptSource::OnGraphChanged(const FEdGraphEditAction &Action)
+{
+	OnChangedDelegate.Broadcast();
+}
+
+void UNiagaraScriptSource::OnGraphDataInterfaceChanged()
+{
+	OnChangedDelegate.Broadcast();
+}
+
 FGuid UNiagaraScriptSource::GetChangeID() 
 { 
 	return NodeGraph->GetChangeID(); 
diff --git a/Engine/Plugins/FX/Niagara/Source/NiagaraEditor/Private/NiagaraSystemFactoryNew.cpp b/Engine/Plugins/FX/Niagara/Source/NiagaraEditor/Private/NiagaraSystemFactoryNew.cpp
--- a/Engine/Plugins/FX/Niagara/Source/NiagaraEditor/Private/NiagaraSystemFactoryNew.cpp
+++ b/Engine/Plugins/FX/Niagara/Source/NiagaraEditor/Private/NiagaraSystemFactoryNew.cpp
@@ -7,6 +7,8 @@
 #include "NiagaraScriptSource.h"
 #include "NiagaraGraph.h"
 #include "NiagaraEditorSettings.h"
+#include "AssetData.h"
+#include "NiagaraStackGraphUtilities.h"
 
 #define LOCTEXT_NAMESPACE "NiagaraSystemFactory"
 
@@ -32,18 +34,19 @@ UObject* UNiagaraSystemFactoryNew::FactoryCreateNew(UClass* Class, UObject* InPa
 	if (UNiagaraSystem* Default = Cast<UNiagaraSystem>(Settings->DefaultSystem.TryLoad()))
 	{
 		NewSystem = Cast<UNiagaraSystem>(StaticDuplicateObject(Default, InParent, Name, Flags, Class));
+		InitializeSystem(NewSystem, false);
 	}
 	else
 	{
 		NewSystem = NewObject<UNiagaraSystem>(InParent, Class, Name, Flags | RF_Transactional);
+		InitializeSystem(NewSystem, true);
 	}
 
-	InitializeSystem(NewSystem);
 
 	return NewSystem;
 }
 
-void UNiagaraSystemFactoryNew::InitializeSystem(UNiagaraSystem* System)
+void UNiagaraSystemFactoryNew::InitializeSystem(UNiagaraSystem* System, bool bCreateDefaultNodes)
 {
 	UNiagaraScript* SystemSpawnScript = System->GetSystemSpawnScript();
 	UNiagaraScript* SystemUpdateScript = System->GetSystemUpdateScript();
@@ -51,12 +54,35 @@ void UNiagaraSystemFactoryNew::InitializeSystem(UNiagaraSystem* System)
 	UNiagaraScript* SystemUpdateScriptSolo = System->GetSystemUpdateScript(true);
 
 	UNiagaraScriptSource* SystemScriptSource = NewObject<UNiagaraScriptSource>(SystemSpawnScript, "SystemScriptSource", RF_Transactional);
-	SystemScriptSource->NodeGraph = NewObject<UNiagaraGraph>(SystemScriptSource, "SystemScriptGraph", RF_Transactional);
+
+	if (SystemScriptSource)
+	{
+		SystemScriptSource->NodeGraph = NewObject<UNiagaraGraph>(SystemScriptSource, "SystemScriptGraph", RF_Transactional);
+	}
 
 	SystemSpawnScript->SetSource(SystemScriptSource);
 	SystemUpdateScript->SetSource(SystemScriptSource);
 	SystemSpawnScriptSolo->SetSource(SystemScriptSource);
 	SystemUpdateScriptSolo->SetSource(SystemScriptSource);
+
+	if (bCreateDefaultNodes)
+	{
+		FStringAssetReference SystemUpdateScriptRef(TEXT("/Niagara/Modules/System/SystemLifeCycle.SystemLifeCycle"));
+		UNiagaraScript* Script = Cast<UNiagaraScript>(SystemUpdateScriptRef.TryLoad());
+
+		FAssetData ModuleScriptAsset(Script);
+		if (SystemScriptSource && ModuleScriptAsset.IsValid())
+		{
+			UNiagaraNodeOutput* SpawnOutputNode = FNiagaraStackGraphUtilities::ResetGraphForOutput(*SystemScriptSource->NodeGraph, ENiagaraScriptUsage::SystemSpawnScript, SystemSpawnScript->GetUsageId());
+			UNiagaraNodeOutput* UpdateOutputNode = FNiagaraStackGraphUtilities::ResetGraphForOutput(*SystemScriptSource->NodeGraph, ENiagaraScriptUsage::SystemUpdateScript, SystemUpdateScript->GetUsageId());
+
+			if (UpdateOutputNode)
+			{
+				FNiagaraStackGraphUtilities::AddScriptModuleToStack(ModuleScriptAsset, *UpdateOutputNode);
+			}
+			FNiagaraStackGraphUtilities::RelayoutGraph(*SystemScriptSource->NodeGraph);
+		}
+	}
 }
 
 #undef LOCTEXT_NAMESPACE
\ No newline at end of file
diff --git a/Engine/Plugins/FX/Niagara/Source/NiagaraEditor/Private/Sequencer/NiagaraSequence/MovieSceneNiagaraEmitterSection.cpp b/Engine/Plugins/FX/Niagara/Source/NiagaraEditor/Private/Sequencer/NiagaraSequence/MovieSceneNiagaraEmitterSection.cpp
--- a/Engine/Plugins/FX/Niagara/Source/NiagaraEditor/Private/Sequencer/NiagaraSequence/MovieSceneNiagaraEmitterSection.cpp
+++ b/Engine/Plugins/FX/Niagara/Source/NiagaraEditor/Private/Sequencer/NiagaraSequence/MovieSceneNiagaraEmitterSection.cpp
@@ -50,7 +50,7 @@ void UMovieSceneNiagaraEmitterSection::GetKeyHandles(TSet<FKeyHandle>& OutKeyHan
 
 TSharedPtr<FNiagaraEmitterHandleViewModel> UMovieSceneNiagaraEmitterSection::GetEmitterHandle()
 {
-	return EmitterHandleViewModel;
+	return EmitterHandleViewModel.Pin();
 }
 
 void UMovieSceneNiagaraEmitterSection::SetEmitterHandle(TSharedRef<FNiagaraEmitterHandleViewModel> InEmitterHandleViewModel)
diff --git a/Engine/Plugins/FX/Niagara/Source/NiagaraEditor/Private/Sequencer/NiagaraSequence/MovieSceneNiagaraEmitterSection.h b/Engine/Plugins/FX/Niagara/Source/NiagaraEditor/Private/Sequencer/NiagaraSequence/MovieSceneNiagaraEmitterSection.h
--- a/Engine/Plugins/FX/Niagara/Source/NiagaraEditor/Private/Sequencer/NiagaraSequence/MovieSceneNiagaraEmitterSection.h
+++ b/Engine/Plugins/FX/Niagara/Source/NiagaraEditor/Private/Sequencer/NiagaraSequence/MovieSceneNiagaraEmitterSection.h
@@ -73,7 +73,7 @@ public:
 
 private:
 	/** The view model for the handle to the emitter this section represents. */
-	TSharedPtr<FNiagaraEmitterHandleViewModel> EmitterHandleViewModel;
+	TWeakPtr<FNiagaraEmitterHandleViewModel> EmitterHandleViewModel;
 
 	/** Storage for the key times in the burst curve. */
 	TArray<float> Times;
diff --git a/Engine/Plugins/FX/Niagara/Source/NiagaraEditor/Private/Sequencer/NiagaraSequence/MovieSceneNiagaraEmitterTrack.cpp b/Engine/Plugins/FX/Niagara/Source/NiagaraEditor/Private/Sequencer/NiagaraSequence/MovieSceneNiagaraEmitterTrack.cpp
--- a/Engine/Plugins/FX/Niagara/Source/NiagaraEditor/Private/Sequencer/NiagaraSequence/MovieSceneNiagaraEmitterTrack.cpp
+++ b/Engine/Plugins/FX/Niagara/Source/NiagaraEditor/Private/Sequencer/NiagaraSequence/MovieSceneNiagaraEmitterTrack.cpp
@@ -14,12 +14,13 @@ UMovieSceneNiagaraEmitterTrack::UMovieSceneNiagaraEmitterTrack(const FObjectInit
 
 TSharedPtr<FNiagaraEmitterHandleViewModel> UMovieSceneNiagaraEmitterTrack::GetEmitterHandle()
 {
-	return EmitterHandleViewModel;
+	return EmitterHandleViewModel.Pin();
 }
 
 void UMovieSceneNiagaraEmitterTrack::SetEmitterHandle(TSharedRef<FNiagaraEmitterHandleViewModel> InEmitterHandleViewModel)
 {
 	EmitterHandleViewModel = InEmitterHandleViewModel;
+	EmitterHandleId = InEmitterHandleViewModel->GetId();
 
 	if (Sections.Num() == 0)
 	{
@@ -56,8 +57,14 @@ TRange<float> UMovieSceneNiagaraEmitterTrack::GetSectionBoundaries() const
 
 FName UMovieSceneNiagaraEmitterTrack::GetTrackName() const
 {
-	return EmitterHandleViewModel.IsValid()
-		? EmitterHandleViewModel->GetName()
+	TSharedPtr<FNiagaraEmitterHandleViewModel> EmitterHandleViewModelPinned = EmitterHandleViewModel.Pin();
+	return EmitterHandleViewModelPinned.IsValid()
+		? EmitterHandleViewModelPinned->GetName()
 		: NAME_None;
 }
 
+FGuid UMovieSceneNiagaraEmitterTrack::GetEmitterHandleId() const
+{
+	return EmitterHandleId;
+}
+
diff --git a/Engine/Plugins/FX/Niagara/Source/NiagaraEditor/Private/Sequencer/NiagaraSequence/MovieSceneNiagaraEmitterTrack.h b/Engine/Plugins/FX/Niagara/Source/NiagaraEditor/Private/Sequencer/NiagaraSequence/MovieSceneNiagaraEmitterTrack.h
--- a/Engine/Plugins/FX/Niagara/Source/NiagaraEditor/Private/Sequencer/NiagaraSequence/MovieSceneNiagaraEmitterTrack.h
+++ b/Engine/Plugins/FX/Niagara/Source/NiagaraEditor/Private/Sequencer/NiagaraSequence/MovieSceneNiagaraEmitterTrack.h
@@ -32,6 +32,8 @@ public:
 	virtual const TArray<UMovieSceneSection*>& GetAllSections() const override;
 	virtual TRange<float> GetSectionBoundaries() const override;
 	virtual FName GetTrackName() const override;
+	/** Gets the unique id for the emitter handle that was associated with this track; used for copy/paste detection */
+	FGuid GetEmitterHandleId() const;
 
 private:
 	void EmitterPropertyChanged();
@@ -39,8 +41,12 @@ private:
 	void SynchronizeWithEmitterHandle();
 
 private:
-	TSharedPtr<FNiagaraEmitterHandleViewModel> EmitterHandleViewModel;
+	TWeakPtr<FNiagaraEmitterHandleViewModel> EmitterHandleViewModel;
 
 	UPROPERTY()
 	TArray<UMovieSceneSection*> Sections;
+
+	// Used for detecting copy/paste 
+	UPROPERTY()
+	FGuid EmitterHandleId;
 };
\ No newline at end of file
diff --git a/Engine/Plugins/FX/Niagara/Source/NiagaraEditor/Private/Sequencer/NiagaraSequence/NiagaraEmitterSection.cpp b/Engine/Plugins/FX/Niagara/Source/NiagaraEditor/Private/Sequencer/NiagaraSequence/NiagaraEmitterSection.cpp
--- a/Engine/Plugins/FX/Niagara/Source/NiagaraEditor/Private/Sequencer/NiagaraSequence/NiagaraEmitterSection.cpp
+++ b/Engine/Plugins/FX/Niagara/Source/NiagaraEditor/Private/Sequencer/NiagaraSequence/NiagaraEmitterSection.cpp
@@ -26,6 +26,12 @@ UMovieSceneSection* FNiagaraEmitterSection::GetSectionObject(void)
 
 int32 FNiagaraEmitterSection::OnPaintSection(FSequencerSectionPainter& InPainter) const
 {
+	TSharedPtr<FNiagaraEmitterHandleViewModel> EmitterHandleViewModel = EmitterSection->GetEmitterHandle();
+	if (EmitterHandleViewModel.IsValid() == false)
+	{
+		return InPainter.LayerId;
+	}
+
 	// draw the first run of the emitter
 	FSlateDrawElement::MakeBox
 	(
@@ -40,7 +46,7 @@ int32 FNiagaraEmitterSection::OnPaintSection(FSequencerSectionPainter& InPainter
 	// draw all loops of the emitter as 'ghosts' of the original section
 	float X = InPainter.SectionGeometry.AbsolutePosition.X;
 	float GeomW = InPainter.SectionGeometry.GetDrawSize().X;
-	int32 NumLoops = EmitterSection->GetEmitterHandle()->GetEmitterViewModel()->GetNumLoops() - 1;
+	int32 NumLoops = EmitterHandleViewModel->GetEmitterViewModel()->GetNumLoops() - 1;
 	for (int32 Loop = 0; Loop < NumLoops; Loop++)
 	{
 		FSlateDrawElement::MakeBox
@@ -59,7 +65,12 @@ int32 FNiagaraEmitterSection::OnPaintSection(FSequencerSectionPainter& InPainter
 
 FText FNiagaraEmitterSection::GetSectionTitle(void) const
 {
-	return EmitterSection->GetEmitterHandle()->GetNameText();
+	TSharedPtr<FNiagaraEmitterHandleViewModel> EmitterHandleviewModel = EmitterSection->GetEmitterHandle();
+	if (EmitterHandleviewModel.IsValid())
+	{
+		return EmitterHandleviewModel->GetNameText();
+	}
+	return FText();
 }
 
 void FNiagaraEmitterSection::GenerateSectionLayout(ISectionLayoutBuilder &LayoutBuilder) const
diff --git a/Engine/Plugins/FX/Niagara/Source/NiagaraEditor/Private/Sequencer/NiagaraSequence/NiagaraEmitterTrackEditor.cpp b/Engine/Plugins/FX/Niagara/Source/NiagaraEditor/Private/Sequencer/NiagaraSequence/NiagaraEmitterTrackEditor.cpp
--- a/Engine/Plugins/FX/Niagara/Source/NiagaraEditor/Private/Sequencer/NiagaraSequence/NiagaraEmitterTrackEditor.cpp
+++ b/Engine/Plugins/FX/Niagara/Source/NiagaraEditor/Private/Sequencer/NiagaraSequence/NiagaraEmitterTrackEditor.cpp
@@ -5,6 +5,9 @@
 #include "NiagaraEmitter.h"
 #include "ViewModels/NiagaraSystemViewModel.h"
 #include "Sequencer/NiagaraSequence/NiagaraSequence.h"
+#include "Framework/MultiBox/MultiBoxBuilder.h"
+
+#define LOCTEXT_NAMESPACE "NiagaraEmitterTrackEditor"
 
 FNiagaraEmitterTrackEditor::FNiagaraEmitterTrackEditor(TSharedPtr<ISequencer> Sequencer) 
 	: FMovieSceneTrackEditor(Sequencer.ToSharedRef())
@@ -34,9 +37,48 @@ bool FNiagaraEmitterTrackEditor::HandleAssetAdded(UObject* Asset, const FGuid& T
 {
 	UNiagaraEmitter* EmitterAsset = Cast<UNiagaraEmitter>(Asset);
 	UNiagaraSequence* NiagaraSequence = Cast<UNiagaraSequence>(GetSequencer()->GetRootMovieSceneSequence());
-	if (EmitterAsset != nullptr && NiagaraSequence != nullptr && NiagaraSequence->GetSystemViewModel().GetCanAddEmittersFromTimeline())
+	if (EmitterAsset != nullptr && NiagaraSequence != nullptr && NiagaraSequence->GetSystemViewModel().GetCanModifyEmittersFromTimeline())
 	{
 		NiagaraSequence->GetSystemViewModel().AddEmitter(*EmitterAsset);
 	}
 	return false;
-}
\ No newline at end of file
+}
+
+void FNiagaraEmitterTrackEditor::BuildTrackContextMenu( FMenuBuilder& MenuBuilder, UMovieSceneTrack* Track )
+{
+	UNiagaraSequence* NiagaraSequencer = Cast<UNiagaraSequence>(GetSequencer()->GetRootMovieSceneSequence());
+	if (NiagaraSequencer != nullptr)
+	{
+		FNiagaraSystemViewModel* SystemViewModel = &NiagaraSequencer->GetSystemViewModel();
+		if (SystemViewModel->GetEditMode() == ENiagaraSystemViewModelEditMode::SystemAsset)
+		{
+			MenuBuilder.BeginSection("Niagara", LOCTEXT("NiagaraContextMenuSectionName", "Niagara"));
+			{
+				MenuBuilder.AddMenuEntry(
+					LOCTEXT("IsolateSelectedEmitters", "Isolate Selected"),
+					LOCTEXT("IsolateSelectedEmittersTooltip", "Force selected emitters to only be isolated."),
+					FSlateIcon(),
+					FUIAction(
+						FExecuteAction::CreateSP(SystemViewModel, &FNiagaraSystemViewModel::IsolateSelectedEmitters),
+						FCanExecuteAction::CreateSP(SystemViewModel, &FNiagaraSystemViewModel::CanExecuteToggleIsolation)
+					)
+				);
+				MenuBuilder.AddMenuEntry(
+					LOCTEXT("IsolateEmitter", "Isolated"),
+					LOCTEXT("IsolateEmitterTooltip", "Toggle the isolation state of the selected emitters."),
+					FSlateIcon(),
+					FUIAction(
+						FExecuteAction::CreateSP(SystemViewModel, &FNiagaraSystemViewModel::IsolateEmitterToggle),
+						FCanExecuteAction::CreateSP(SystemViewModel, &FNiagaraSystemViewModel::CanExecuteToggleIsolation),
+						FIsActionChecked::CreateSP(SystemViewModel, &FNiagaraSystemViewModel::IsEmitterIsolationActive)
+					),
+					NAME_None,
+					EUserInterfaceActionType::ToggleButton
+				);
+			}
+			MenuBuilder.EndSection();
+		}
+	}
+}
+
+#undef LOCTEXT_NAMESPACE
\ No newline at end of file
diff --git a/Engine/Plugins/FX/Niagara/Source/NiagaraEditor/Private/Tests/NiagaraEmitterInheritanceTests.h b/Engine/Plugins/FX/Niagara/Source/NiagaraEditor/Private/Tests/NiagaraEmitterInheritanceTests.h
new file mode 100644
--- /dev/null
+++ b/Engine/Plugins/FX/Niagara/Source/NiagaraEditor/Private/Tests/NiagaraEmitterInheritanceTests.h
@@ -0,0 +1,3 @@
+// Copyright 1998-2017 Epic Games, Inc. All Rights Reserved.
+
+#pragma once
\ No newline at end of file
diff --git a/Engine/Plugins/FX/Niagara/Source/NiagaraEditor/Private/Tests/NiagaraStackViewModelTests.h b/Engine/Plugins/FX/Niagara/Source/NiagaraEditor/Private/Tests/NiagaraStackViewModelTests.h
new file mode 100644
--- /dev/null
+++ b/Engine/Plugins/FX/Niagara/Source/NiagaraEditor/Private/Tests/NiagaraStackViewModelTests.h
@@ -0,0 +1 @@
+// Copyright 1998-2017 Epic Games, Inc. All Rights Reserved.
\ No newline at end of file
diff --git a/Engine/Plugins/FX/Niagara/Source/NiagaraEditor/Private/Toolkits/NiagaraParameterCollectionToolkit.cpp b/Engine/Plugins/FX/Niagara/Source/NiagaraEditor/Private/Toolkits/NiagaraParameterCollectionToolkit.cpp
--- a/Engine/Plugins/FX/Niagara/Source/NiagaraEditor/Private/Toolkits/NiagaraParameterCollectionToolkit.cpp
+++ b/Engine/Plugins/FX/Niagara/Source/NiagaraEditor/Private/Toolkits/NiagaraParameterCollectionToolkit.cpp
@@ -170,24 +170,30 @@ TSharedRef<SDockTab> FNiagaraParameterCollectionToolkit::SpawnTab_Main(const FSp
 
 	TSharedRef<SVerticalBox> Contents = SNew(SVerticalBox);
 
-	if (!Instance->IsDefaultInstance())
-	{
-		FDetailsViewArgs DetailArgs;
-		DetailArgs.bAllowSearch = false;
 
-		FPropertyEditorModule& PropertyEditorModule = FModuleManager::GetModuleChecked<FPropertyEditorModule>("PropertyEditor");
-		FDetailsViewArgs DetailsViewArgs(false, false, true, FDetailsViewArgs::HideNameArea, true, ParameterCollectionViewModel.Get());
-		TSharedRef<IDetailsView> DetailsView = PropertyEditorModule.CreateDetailView(DetailsViewArgs);
-		DetailsView->SetObject(Instance);
+	FDetailsViewArgs DetailArgs;
+	DetailArgs.bAllowSearch = false;
 
-		Contents->AddSlot()
-			.AutoHeight()
-			.Padding(FMargin(0.0f, 2.0f))
-			[
-				DetailsView
-			];
+	FPropertyEditorModule& PropertyEditorModule = FModuleManager::GetModuleChecked<FPropertyEditorModule>("PropertyEditor");
+	FDetailsViewArgs DetailsViewArgs(false, false, true, FDetailsViewArgs::HideNameArea, true, ParameterCollectionViewModel.Get());
+	TSharedRef<IDetailsView> DetailsView = PropertyEditorModule.CreateDetailView(DetailsViewArgs);
+
+	if (Instance->IsDefaultInstance())
+	{
+		DetailsView->SetObject(Collection);
+	}
+	else
+	{
+		DetailsView->SetObject(Instance);
 	}
 
+	Contents->AddSlot()
+		.AutoHeight()
+		.Padding(FMargin(0.0f, 2.0f))
+		[
+			DetailsView
+		];
+
 	Contents->AddSlot()
 		.AutoHeight()
 		.Padding(FMargin(0.0f, 2.0f))
diff --git a/Engine/Plugins/FX/Niagara/Source/NiagaraEditor/Private/Toolkits/NiagaraScriptToolkit.cpp b/Engine/Plugins/FX/Niagara/Source/NiagaraEditor/Private/Toolkits/NiagaraScriptToolkit.cpp
--- a/Engine/Plugins/FX/Niagara/Source/NiagaraEditor/Private/Toolkits/NiagaraScriptToolkit.cpp
+++ b/Engine/Plugins/FX/Niagara/Source/NiagaraEditor/Private/Toolkits/NiagaraScriptToolkit.cpp
@@ -39,6 +39,8 @@
 
 #define LOCTEXT_NAMESPACE "NiagaraScriptToolkit"
 
+DECLARE_CYCLE_STAT(TEXT("Niagara - ScriptToolkit - OnApply"), STAT_NiagaraEditor_ScriptToolkit_OnApply, STATGROUP_NiagaraEditor);
+
 static TAutoConsoleVariable<int32> CVarDevDetails(
 	TEXT("fx.DevDetailsPanels"),
 	0,
@@ -273,7 +275,7 @@ void FNiagaraScriptToolkit::SetupCommands()
 		FCanExecuteAction::CreateSP(this, &FNiagaraScriptToolkit::OnApplyEnabled));
 	GetToolkitCommands()->MapAction(
 		FNiagaraEditorCommands::Get().Compile,
-		FExecuteAction::CreateRaw(this, &FNiagaraScriptToolkit::CompileScript));
+		FExecuteAction::CreateRaw(this, &FNiagaraScriptToolkit::CompileScript, true));
 	GetToolkitCommands()->MapAction(
 		FNiagaraEditorCommands::Get().RefreshNodes,
 		FExecuteAction::CreateRaw(this, &FNiagaraScriptToolkit::RefreshNodes));
@@ -358,9 +360,9 @@ FText FNiagaraScriptToolkit::GetRefreshStatusTooltip() const
 	return LOCTEXT("Refresh_Status", "Currently dependencies up-to-date. Consider refreshing if status isn't accurate.");
 }
 
-void FNiagaraScriptToolkit::CompileScript()
+void FNiagaraScriptToolkit::CompileScript(bool bForce)
 {
-	ScriptViewModel->CompileStandaloneScript();
+	ScriptViewModel->CompileStandaloneScript(bForce);
 }
 
 void FNiagaraScriptToolkit::RefreshNodes()
@@ -370,8 +372,8 @@ void FNiagaraScriptToolkit::RefreshNodes()
 
 void FNiagaraScriptToolkit::OnApply()
 {
+	SCOPE_CYCLE_COUNTER(STAT_NiagaraEditor_ScriptToolkit_OnApply);
 	UE_LOG(LogNiagaraEditor, Log, TEXT("Applying Niagara Script %s"), *GetEditingObjects()[0]->GetName());
-
 	UpdateOriginalNiagaraScript();
 }
 
@@ -466,7 +468,7 @@ void FNiagaraScriptToolkit::UpdateOriginalNiagaraScript()
 		if (bRefreshed)
 		{
 			Source->NodeGraph->NotifyGraphNeedsRecompile();
-			AffectedScripts.Add(*It);
+			AffectedScripts.AddUnique(*It);
 		}
 		else
 		{
@@ -479,7 +481,7 @@ void FNiagaraScriptToolkit::UpdateOriginalNiagaraScript()
 				if (Graph == OriginalGraph)
 				{
 					Source->NodeGraph->NotifyGraphNeedsRecompile();
-					AffectedScripts.Add(*It);
+					AffectedScripts.AddUnique(*It);
 					break;
 				}
 			}
@@ -506,11 +508,11 @@ void FNiagaraScriptToolkit::UpdateOriginalNiagaraScript()
 			{
 				AffectedScriptViewModel = MakeShareable(new FNiagaraScriptViewModel(Script, FText::FromString(Script->GetName()), ENiagaraParameterEditMode::EditValueOnly));
 			}
-			AffectedScriptViewModel->CompileStandaloneScript();
+			AffectedScriptViewModel->CompileStandaloneScript(false);
 		}
 	}
 
-	FNiagaraEditorUtilities::UpdateExistingEmitters(AffectedEmitters);
+	FNiagaraEditorUtilities::CompileExistingEmitters(AffectedEmitters);
 
 	GWarn->EndSlowTask();
 	ScriptViewModel->SetScriptDirty(false);
diff --git a/Engine/Plugins/FX/Niagara/Source/NiagaraEditor/Private/Toolkits/NiagaraScriptToolkit.h b/Engine/Plugins/FX/Niagara/Source/NiagaraEditor/Private/Toolkits/NiagaraScriptToolkit.h
--- a/Engine/Plugins/FX/Niagara/Source/NiagaraEditor/Private/Toolkits/NiagaraScriptToolkit.h
+++ b/Engine/Plugins/FX/Niagara/Source/NiagaraEditor/Private/Toolkits/NiagaraScriptToolkit.h
@@ -77,7 +77,7 @@ private:
 	void ExtendToolbar();
 
 	/** Compiles the script. */
-	void CompileScript();
+	void CompileScript(bool bForce);
 
 	/** Refreshes the nodes in the script graph, updating the pins to match external changes. */
 	void RefreshNodes();
diff --git a/Engine/Plugins/FX/Niagara/Source/NiagaraEditor/Private/Toolkits/NiagaraSystemToolkit.cpp b/Engine/Plugins/FX/Niagara/Source/NiagaraEditor/Private/Toolkits/NiagaraSystemToolkit.cpp
--- a/Engine/Plugins/FX/Niagara/Source/NiagaraEditor/Private/Toolkits/NiagaraSystemToolkit.cpp
+++ b/Engine/Plugins/FX/Niagara/Source/NiagaraEditor/Private/Toolkits/NiagaraSystemToolkit.cpp
@@ -21,6 +21,8 @@
 #include "NiagaraEditorStyle.h"
 #include "NiagaraEditorSettings.h"
 #include "NiagaraSystemFactoryNew.h"
+#include "NiagaraEmitter.h"
+#include "NiagaraComponent.h"
 
 #include "IContentBrowserSingleton.h"
 #include "ContentBrowserModule.h"
@@ -40,9 +42,14 @@
 #include "Editor.h"
 #include "Engine/Selection.h"
 #include "Misc/MessageDialog.h"
+#include "Modules/ModuleManager.h"
+#include "AssetRegistryModule.h"
+#include "IAssetRegistry.h"
 
 #define LOCTEXT_NAMESPACE "NiagaraSystemEditor"
 
+DECLARE_CYCLE_STAT(TEXT("Niagara - SystemToolkit - OnApply"), STAT_NiagaraEditor_SystemToolkit_OnApply, STATGROUP_NiagaraEditor);
+
 const FName FNiagaraSystemToolkit::ViewportTabID(TEXT("NiagaraSystemEditor_Viewport"));
 const FName FNiagaraSystemToolkit::CurveEditorTabID(TEXT("NiagaraSystemEditor_CurveEditor"));
 const FName FNiagaraSystemToolkit::SequencerTabID(TEXT("NiagaraSystemEditor_Sequencer"));
@@ -163,20 +170,12 @@ void FNiagaraSystemToolkit::InitializeWithSystem(const EToolkitMode::Type Mode,
 		}
 	}
 
-	if (bConverted)
-	{
-		System->ResynchronizeAllHandles();
-	}
-
 	FNiagaraSystemViewModelOptions SystemOptions;
-	SystemOptions.bCanRemoveEmittersFromTimeline = true;
-	SystemOptions.bCanRenameEmittersFromTimeline = true;
-	SystemOptions.bCanAddEmittersFromTimeline = true;
+	SystemOptions.bCanModifyEmittersFromTimeline = true;
 	SystemOptions.bUseSystemExecStateForTimelineReset = true;
+	SystemOptions.EditMode = ENiagaraSystemViewModelEditMode::SystemAsset;
 	SystemOptions.OnGetSequencerAddMenuContent.BindSP(this, &FNiagaraSystemToolkit::GetSequencerAddMenuContent);
 
-	System->CheckForUpdates();
-
 	SystemViewModel = MakeShareable(new FNiagaraSystemViewModel(*System, SystemOptions));
 	SystemToolkitMode = ESystemToolkitMode::System;
 	InitializeInternal(Mode, InitToolkitHost);
@@ -194,7 +193,7 @@ void FNiagaraSystemToolkit::InitializeWithEmitter(const EToolkitMode::Type Mode,
 	}
 
 	System = NewObject<UNiagaraSystem>(GetTransientPackage(), NAME_None, RF_Transient);
-	UNiagaraSystemFactoryNew::InitializeSystem(System);
+	UNiagaraSystemFactoryNew::InitializeSystem(System, true);
 
 	Emitter = &InEmitter;
 
@@ -205,13 +204,13 @@ void FNiagaraSystemToolkit::InitializeWithEmitter(const EToolkitMode::Type Mode,
 	System->AddEmitterHandleWithoutCopying(*EditableEmitter);
 
 	FNiagaraSystemViewModelOptions SystemOptions;
-	SystemOptions.bCanRemoveEmittersFromTimeline = false;
-	SystemOptions.bCanRenameEmittersFromTimeline = false;
-	SystemOptions.bCanAddEmittersFromTimeline = false;
+	SystemOptions.bCanModifyEmittersFromTimeline = false;
 	SystemOptions.bUseSystemExecStateForTimelineReset = false;
+	SystemOptions.EditMode = ENiagaraSystemViewModelEditMode::EmitterAsset;
 
 	SystemViewModel = MakeShareable(new FNiagaraSystemViewModel(*System, SystemOptions));
 	SystemViewModel->GetSystemScriptViewModel()->RebuildEmitterNodes();
+	SystemViewModel->GetSystemScriptViewModel()->CompileSystem(false);
 	SystemToolkitMode = ESystemToolkitMode::Emitter;
 	InitializeInternal(Mode, InitToolkitHost);
 }
@@ -513,17 +512,39 @@ TSharedRef<SDockTab> FNiagaraSystemToolkit::SpawnTab_GeneratedCode(const FSpawnT
 
 void FNiagaraSystemToolkit::SetupCommands()
 {
-	GetToolkitCommands()->MapAction(
-		FNiagaraEditorCommands::Get().ToggleUnlockToChanges,
-		FExecuteAction::CreateSP(this, &FNiagaraSystemToolkit::ToggleUnlockToChanges),
-		FCanExecuteAction(),
-		FIsActionChecked::CreateSP(this, &FNiagaraSystemToolkit::IsToggleUnlockToChangesChecked));
 	GetToolkitCommands()->MapAction(
 		FNiagaraEditorCommands::Get().Compile,
-		FExecuteAction::CreateRaw(this, &FNiagaraSystemToolkit::CompileSystem));
+		FExecuteAction::CreateRaw(this, &FNiagaraSystemToolkit::CompileSystem, true));
 	GetToolkitCommands()->MapAction(
 		FNiagaraEditorCommands::Get().ResetSimulation,
 		FExecuteAction::CreateRaw(this, &FNiagaraSystemToolkit::ResetSimulation));
+
+	GetToolkitCommands()->MapAction(
+		FNiagaraEditorCommands::Get().ToggleBounds,
+		FExecuteAction::CreateSP(this, &FNiagaraSystemToolkit::OnToggleBounds),
+		FCanExecuteAction(),
+		FIsActionChecked::CreateSP(this, &FNiagaraSystemToolkit::IsToggleBoundsChecked));
+
+	GetToolkitCommands()->MapAction(
+		FNiagaraEditorCommands::Get().ToggleBounds_SetFixedBounds,
+		FExecuteAction::CreateSP(this, &FNiagaraSystemToolkit::OnToggleBoundsSetFixedBounds));
+
+	GetToolkitCommands()->MapAction(
+		FNiagaraEditorCommands::Get().SaveThumbnailImage,
+		FExecuteAction::CreateSP(this, &FNiagaraSystemToolkit::OnSaveThumbnailImage));
+
+	GetToolkitCommands()->MapAction(
+		FNiagaraEditorCommands::Get().Apply,
+		FExecuteAction::CreateSP(this, &FNiagaraSystemToolkit::OnApply),
+		FCanExecuteAction::CreateSP(this, &FNiagaraSystemToolkit::OnApplyEnabled));
+}
+
+void FNiagaraSystemToolkit::OnSaveThumbnailImage()
+{
+	if (Viewport.IsValid() && Viewport->GetViewportClient().IsValid())
+	{
+		Viewport->CreateThumbnail();
+	}
 }
 
 void FNiagaraSystemToolkit::ResetSimulation()
@@ -537,6 +558,17 @@ void FNiagaraSystemToolkit::ExtendToolbar()
 	{
 		static void FillToolbar(FToolBarBuilder& ToolbarBuilder, FNiagaraSystemToolkit* Toolkit)
 		{
+			if (Toolkit->Emitter != nullptr)
+			{
+				ToolbarBuilder.BeginSection("Apply");
+				{
+					ToolbarBuilder.AddToolBarButton(FNiagaraEditorCommands::Get().Apply,
+						NAME_None, TAttribute<FText>(), TAttribute<FText>(),
+						FSlateIcon(FNiagaraEditorStyle::GetStyleSetName(), "NiagaraEditor.Apply"),
+						FName(TEXT("ApplyNiagaraEmitter")));
+				}
+				ToolbarBuilder.EndSection();
+			}
 			ToolbarBuilder.BeginSection("Compile");
 			{
 				ToolbarBuilder.AddToolBarButton(FNiagaraEditorCommands::Get().Compile,
@@ -555,12 +587,29 @@ void FNiagaraSystemToolkit::ExtendToolbar()
 			}
 			ToolbarBuilder.EndSection();
 
-			ToolbarBuilder.BeginSection("LockEmitters");
+			ToolbarBuilder.BeginSection("NiagaraThumbnail");
 			{
-				ToolbarBuilder.AddToolBarButton(FNiagaraEditorCommands::Get().ToggleUnlockToChanges, NAME_None,
-					TAttribute<FText>(Toolkit, &FNiagaraSystemToolkit::GetEmitterLockToChangesLabel),
-					TAttribute<FText>(Toolkit, &FNiagaraSystemToolkit::GetEmitterLockToChangesLabelTooltip),
-					TAttribute<FSlateIcon>(Toolkit, &FNiagaraSystemToolkit::GetEmitterLockToChangesIcon));
+				ToolbarBuilder.AddToolBarButton(FNiagaraEditorCommands::Get().SaveThumbnailImage, NAME_None,
+					LOCTEXT("GenerateThumbnail", "Thumbnail"),
+					LOCTEXT("GenerateThumbnailTooltip","Generate a thumbnail image."),
+					FSlateIcon(FEditorStyle::GetStyleSetName(), "Cascade.SaveThumbnailImage"));
+			}
+			ToolbarBuilder.EndSection();
+
+			ToolbarBuilder.BeginSection("NiagaraPreviewOptions");
+			{
+				ToolbarBuilder.AddToolBarButton(FNiagaraEditorCommands::Get().ToggleBounds, NAME_None,
+					LOCTEXT("ShowBounds", "Bounds"),
+					LOCTEXT("ShowBoundsTooltip", "Show the bounds for the scene."),
+					FSlateIcon(FEditorStyle::GetStyleSetName(), "Cascade.ToggleBounds"));
+				ToolbarBuilder.AddComboButton(
+					FUIAction(),
+					FOnGetContent::CreateRaw(Toolkit, &FNiagaraSystemToolkit::GenerateBoundsMenuContent, Toolkit->GetToolkitCommands()),
+					LOCTEXT("BoundsMenuCombo_Label", "Bounds Options"),
+					LOCTEXT("BoundsMenuCombo_ToolTip", "Bounds options"),
+					FSlateIcon(FEditorStyle::GetStyleSetName(), "Cascade.ToggleBounds"),
+					true
+				);
 			}
 			ToolbarBuilder.EndSection();
 		}
@@ -581,6 +630,16 @@ void FNiagaraSystemToolkit::ExtendToolbar()
 	AddToolbarExtender(NiagaraEditorModule.GetToolBarExtensibilityManager()->GetAllExtenders(GetToolkitCommands(), GetEditingObjects()));
 }
 
+TSharedRef<SWidget> FNiagaraSystemToolkit::GenerateBoundsMenuContent(TSharedRef<FUICommandList> InCommandList)
+{
+	const bool bShouldCloseWindowAfterMenuSelection = true;
+	FMenuBuilder MenuBuilder(bShouldCloseWindowAfterMenuSelection, InCommandList);
+
+	MenuBuilder.AddMenuEntry(FNiagaraEditorCommands::Get().ToggleBounds_SetFixedBounds);
+
+	return MenuBuilder.MakeWidget();
+}
+
 void FNiagaraSystemToolkit::GetSequencerAddMenuContent(FMenuBuilder& MenuBuilder, TSharedRef<ISequencer> Sequencer)
 {
 	MenuBuilder.AddSubMenu(
@@ -653,9 +712,72 @@ FText FNiagaraSystemToolkit::GetCompileStatusTooltip() const
 }
 
 
-void FNiagaraSystemToolkit::CompileSystem()
+void FNiagaraSystemToolkit::CompileSystem(bool bForce)
+{
+	SystemViewModel->CompileSystem(bForce);
+}
+
+void FNiagaraSystemToolkit::OnToggleBounds()
+{
+	ToggleDrawOption(SNiagaraSystemViewport::Bounds);
+}
+
+bool FNiagaraSystemToolkit::IsToggleBoundsChecked() const
+{
+	return IsDrawOptionEnabled(SNiagaraSystemViewport::Bounds);
+}
+
+void FNiagaraSystemToolkit::ToggleDrawOption(int32 Element)
+{
+	if (Viewport.IsValid() && Viewport->GetViewportClient().IsValid())
+	{
+		Viewport->ToggleDrawElement((SNiagaraSystemViewport::EDrawElements)Element);
+		Viewport->RefreshViewport();
+	}
+}
+
+bool FNiagaraSystemToolkit::IsDrawOptionEnabled(int32 Element) const
+{
+	if (Viewport.IsValid() && Viewport->GetViewportClient().IsValid())
+	{
+		return Viewport->GetDrawElement((SNiagaraSystemViewport::EDrawElements)Element);
+	}
+	else
+	{
+		return false;
+	}
+}
+
+void FNiagaraSystemToolkit::OnToggleBoundsSetFixedBounds()
 {
-	SystemViewModel->CompileSystem();
+	FScopedTransaction Transaction(LOCTEXT("SetFixedBounds", "Set Fixed Bounds"));
+
+	SystemViewModel->UpdateEmitterFixedBounds();
+
+	/*
+	// Force the component to update its bounds.
+	ParticleSystemComponent->ForceUpdateBounds();
+
+	// Grab the current bounds of the PSysComp & set it on the PSystem itself
+	ParticleSystem->Modify();
+	ParticleSystem->FixedRelativeBoundingBox.Min = ParticleSystemComponent->Bounds.GetBoxExtrema(0);
+	ParticleSystem->FixedRelativeBoundingBox.Max = ParticleSystemComponent->Bounds.GetBoxExtrema(1);
+	ParticleSystem->FixedRelativeBoundingBox.IsValid = true;
+	ParticleSystem->bUseFixedRelativeBoundingBox = true;
+
+	ParticleSystem->MarkPackageDirty();
+
+	EndTransaction(Transaction);
+
+	if ((SelectedModule == NULL) && (SelectedEmitter == NULL))
+	{
+		TArray<UObject*> NewSelection;
+		NewSelection.Add(ParticleSystem);
+		SetSelection(NewSelection);
+	}
+
+	ReassociateParticleSystem();
+	*/
 }
 
 void FNiagaraSystemToolkit::UpdateOriginalEmitter()
@@ -683,59 +805,48 @@ void FNiagaraSystemToolkit::UpdateOriginalEmitter()
 		RF_AllFlags,
 		Emitter->GetClass());
 
-	// Restore RF_Standalone on the original material, as it had been removed from the preview material so that it could be GC'd.
+	// Restore RF_Standalone on the original emitter, as it had been removed from the preview emitter so that it could be GC'd.
 	Emitter->SetFlags(RF_Standalone);
 
+
 	TArray<UNiagaraEmitter*> AffectedEmitters;
 	AffectedEmitters.Add(Emitter);
-	UpdateExistingEmitters(AffectedEmitters);
+	UpdateExistingEmitters();
 
 	GWarn->EndSlowTask();
-	EditableEmitterViewModel->SetDirty(false);
 }
 
-void FNiagaraSystemToolkit::UpdateExistingEmitters(TArray<UNiagaraEmitter*> AffectedEmitters)
+void FNiagaraSystemToolkit::UpdateExistingEmitters()
 {
-	// Compile the existing emitters. Also determine which Systems need to be properly updated.
-	TArray<UNiagaraSystem*> AffectedSystemSystems;
-	for (UNiagaraEmitter* AffectedEmitter : AffectedEmitters)
+	for (TObjectIterator<UNiagaraSystem> SystemIterator; SystemIterator; ++SystemIterator)
 	{
-		if (AffectedEmitter->IsPendingKillOrUnreachable())
+		UNiagaraSystem* LoadedSystem = *SystemIterator;
+		if (LoadedSystem->IsPendingKill() == false && 
+			LoadedSystem->HasAnyFlags(RF_ClassDefaultObject) == false &&
+			LoadedSystem->ReferencesSourceEmitter(*Emitter))
 		{
-			continue;
-		}
+			LoadedSystem->UpdateFromEmitterChanges(*Emitter);
+			TArray<TSharedPtr<FNiagaraSystemViewModel>> ReferencingSystemViewModels;
+			FNiagaraSystemViewModel::GetAllViewModelsForObject(LoadedSystem, ReferencingSystemViewModels);
 
-		TSharedPtr<FNiagaraEmitterViewModel> EmitterViewModel = FNiagaraEmitterViewModel::GetExistingViewModelForObject(AffectedEmitter);
-		if (!EmitterViewModel.IsValid())
-		{
-			EmitterViewModel = MakeShareable(new FNiagaraEmitterViewModel(AffectedEmitter, nullptr));
-		}
-		EmitterViewModel->CompileScripts();
-
-		for (TObjectIterator<UNiagaraSystem> It; It; ++It)
-		{
-			if (It->GetAutoImportChangedEmitters() && It->ReferencesSourceEmitter(AffectedEmitter))
+			for (TSharedPtr<FNiagaraSystemViewModel> ReferencingSystemViewModel : ReferencingSystemViewModels)
 			{
-				AffectedSystemSystems.AddUnique(*It);
+				ReferencingSystemViewModel->RefreshAll();
 			}
-		}
-	}
 
-	// Now iterate over the affected Systems.
-	for (UNiagaraSystem* System : AffectedSystemSystems)
-	{
-		TSharedPtr<FNiagaraSystemViewModel> SystemViewModel = FNiagaraSystemViewModel::GetExistingViewModelForObject(System);
-		if (!SystemViewModel.IsValid())
-		{
-			FNiagaraSystemViewModelOptions Options;
-			Options.bCanRemoveEmittersFromTimeline = false;
-			Options.bCanRenameEmittersFromTimeline = false;
-			Options.bCanAddEmittersFromTimeline = false;
-			Options.bUseSystemExecStateForTimelineReset = false;
-			SystemViewModel = MakeShareable(new FNiagaraSystemViewModel(*System, Options));
+			if (ReferencingSystemViewModels.Num() == 0)
+			{
+				for (TObjectIterator<UNiagaraComponent> ComponentIterator; ComponentIterator; ++ComponentIterator)
+				{
+					UNiagaraComponent* Component = *ComponentIterator;
+					if (Component->GetAsset() == LoadedSystem)
+					{
+						Component->SynchronizeWithSourceSystem();
+						Component->ReinitializeSystem();
+					}
+				}
+			}
 		}
-
-		SystemViewModel->ResynchronizeAllHandles();
 	}
 }
 
@@ -778,7 +889,7 @@ bool FNiagaraSystemToolkit::OnRequestClose()
 	if (SystemToolkitMode == ESystemToolkitMode::Emitter)
 	{
 		TSharedPtr<FNiagaraEmitterViewModel> EmitterViewModel = SystemViewModel->GetEmitterHandleViewModels()[0]->GetEmitterViewModel();
-		if (EmitterViewModel->GetDirty())
+		if (EmitterViewModel->GetEmitter()->GetChangeId() != Emitter->GetChangeId())
 		{
 			// find out the user wants to do with this dirty NiagaraScript
 			EAppReturnType::Type YesNoCancelReply = FMessageDialog::Open(EAppMsgType::YesNoCancel,
@@ -816,69 +927,31 @@ void FNiagaraSystemToolkit::EmitterAssetSelected(const FAssetData& AssetData)
 	SystemViewModel->AddEmitterFromAssetData(AssetData);
 }
 
-void FNiagaraSystemToolkit::ToggleUnlockToChanges()
-{
-	const FScopedTransaction ToggleUnlockToChangesTransaction(LOCTEXT("ToggleUnlockToChanges", "Toggle System Unlock To Changes"));
-	System->Modify();
-
-	System->SetAutoImportChangedEmitters(!System->GetAutoImportChangedEmitters());
-
-	if (System->GetAutoImportChangedEmitters())
-	{
-		SystemViewModel->ResynchronizeAllHandles();
-	}
-}
-
-bool FNiagaraSystemToolkit::IsToggleUnlockToChangesChecked()
+void FNiagaraSystemToolkit::ToggleCompileEnabled()
 {
-	return System->GetAutoImportChangedEmitters();
+	UNiagaraEditorSettings* Settings = GetMutableDefault<UNiagaraEditorSettings>();
+	Settings->bAutoCompile = !Settings->bAutoCompile;
 }
 
-FText FNiagaraSystemToolkit::GetEmitterLockToChangesLabel() const
+bool FNiagaraSystemToolkit::IsAutoCompileEnabled()
 {
-	if (System->GetAutoImportChangedEmitters())
-	{
-		return LOCTEXT("EmitterUnlockToChangesLabel", "Changes Unlocked");
-	}
-	else
-	{
-		return LOCTEXT("EmitterLockToChangesLabel", "Changes Locked");
-	}
+	return GetDefault<UNiagaraEditorSettings>()->bAutoCompile;
 }
 
-FText FNiagaraSystemToolkit::GetEmitterLockToChangesLabelTooltip() const
+void FNiagaraSystemToolkit::OnApply()
 {
-	if (System->GetAutoImportChangedEmitters())
-	{
-		return LOCTEXT("EmitterUnlockToChangesLabelTooltip", "If a source emitter changes, the changes will be imported into this System automatically.");
-	}
-	else
-	{
-		return LOCTEXT("EmitterLockToChangesLabelTooltip", "If a source emitter changes, the changes will NOT be imported into this System automatically.");
-	}
+	SCOPE_CYCLE_COUNTER(STAT_NiagaraEditor_SystemToolkit_OnApply);
+	UpdateOriginalEmitter();
 }
 
-FSlateIcon FNiagaraSystemToolkit::GetEmitterLockToChangesIcon() const
+bool FNiagaraSystemToolkit::OnApplyEnabled() const
 {
-	if (System->GetAutoImportChangedEmitters())
+	if (Emitter != nullptr)
 	{
-		return FSlateIcon(FNiagaraEditorStyle::GetStyleSetName(), "NiagaraEditor.UnlockToChanges");
-	}
-	else
-	{
-		return FSlateIcon(FNiagaraEditorStyle::GetStyleSetName(), "NiagaraEditor.LockToChanges");
+		TSharedPtr<FNiagaraEmitterViewModel> EmitterViewModel = SystemViewModel->GetEmitterHandleViewModels()[0]->GetEmitterViewModel();
+		return EmitterViewModel->GetEmitter()->GetChangeId() != Emitter->GetChangeId();
 	}
-}
-
-void FNiagaraSystemToolkit::ToggleCompileEnabled()
-{
-	UNiagaraEditorSettings* Settings = GetMutableDefault<UNiagaraEditorSettings>();
-	Settings->bAutoCompile = !Settings->bAutoCompile;
-}
-
-bool FNiagaraSystemToolkit::IsAutoCompileEnabled()
-{
-	return GetDefault<UNiagaraEditorSettings>()->bAutoCompile;
+	return false;
 }
 
 #undef LOCTEXT_NAMESPACE
diff --git a/Engine/Plugins/FX/Niagara/Source/NiagaraEditor/Private/Toolkits/NiagaraSystemToolkit.h b/Engine/Plugins/FX/Niagara/Source/NiagaraEditor/Private/Toolkits/NiagaraSystemToolkit.h
--- a/Engine/Plugins/FX/Niagara/Source/NiagaraEditor/Private/Toolkits/NiagaraSystemToolkit.h
+++ b/Engine/Plugins/FX/Niagara/Source/NiagaraEditor/Private/Toolkits/NiagaraSystemToolkit.h
@@ -63,9 +63,16 @@ public:
 	FText GetCompileStatusTooltip() const;
 
 	/** Compiles the system script. */
-	void CompileSystem();
+	void CompileSystem(bool bForce);
 
 protected:
+	void OnToggleBounds();
+	bool IsToggleBoundsChecked() const;
+	void OnToggleBoundsSetFixedBounds();
+
+	void ToggleDrawOption(int32 Element);
+	bool IsDrawOptionEnabled(int32 Element) const;
+
 	//~ FAssetEditorToolkit interface
 	virtual void GetSaveableObjects(TArray<UObject*>& OutObjects) const override;
 	virtual void SaveAsset_Execute() override;
@@ -76,7 +83,7 @@ private:
 	void InitializeInternal(const EToolkitMode::Type Mode, const TSharedPtr<class IToolkitHost>& InitToolkitHost);
 
 	void UpdateOriginalEmitter();
-	static void UpdateExistingEmitters(TArray<UNiagaraEmitter*> AffectedEmitters);
+	void UpdateExistingEmitters();
 
 	TSharedRef<SDockTab> SpawnTab_Viewport(const FSpawnTabArgs& Args);
 	TSharedRef<SDockTab> SpawnTab_CurveEd(const FSpawnTabArgs& Args);
@@ -100,17 +107,17 @@ private:
 	static TSharedRef<SWidget> GenerateCompileMenuContent();
 
 	void EmitterAssetSelected(const FAssetData& AssetData);
-	void ToggleUnlockToChanges();
-	bool IsToggleUnlockToChangesChecked();
-
-	FText GetEmitterLockToChangesLabel() const;
-	FText GetEmitterLockToChangesLabelTooltip() const;
-	FSlateIcon GetEmitterLockToChangesIcon() const;
 
 	static void ToggleCompileEnabled();
 	static bool IsAutoCompileEnabled();
 
+	void OnApply();
+	bool OnApplyEnabled() const;
+
 private:
+	TSharedRef<SWidget> GenerateBoundsMenuContent(TSharedRef<FUICommandList> InCommandList);
+	void OnSaveThumbnailImage();
+
 	/** The System being edited in system mode, or the placeholder system being edited in emitter mode. */
 	UNiagaraSystem* System;
 
diff --git a/Engine/Plugins/FX/Niagara/Source/NiagaraEditor/Private/TypeEditorUtilities/NiagaraBoolTypeEditorUtilities.cpp b/Engine/Plugins/FX/Niagara/Source/NiagaraEditor/Private/TypeEditorUtilities/NiagaraBoolTypeEditorUtilities.cpp
--- a/Engine/Plugins/FX/Niagara/Source/NiagaraEditor/Private/TypeEditorUtilities/NiagaraBoolTypeEditorUtilities.cpp
+++ b/Engine/Plugins/FX/Niagara/Source/NiagaraEditor/Private/TypeEditorUtilities/NiagaraBoolTypeEditorUtilities.cpp
@@ -59,7 +59,7 @@ private:
 	bool bBoolValue;
 };
 
-TSharedPtr<SNiagaraParameterEditor> FNiagaraEditorBoolTypeUtilities::CreateParameterEditor() const
+TSharedPtr<SNiagaraParameterEditor> FNiagaraEditorBoolTypeUtilities::CreateParameterEditor(const FNiagaraTypeDefinition& ParameterType) const
 {
 	return SNew(SNiagaraBoolParameterEditor);
 }
@@ -69,25 +69,21 @@ bool FNiagaraEditorBoolTypeUtilities::CanHandlePinDefaults() const
 	return true;
 }
 
-FString FNiagaraEditorBoolTypeUtilities::GetPinDefaultStringFromValue(const FNiagaraVariable& Variable) const
+FString FNiagaraEditorBoolTypeUtilities::GetPinDefaultStringFromValue(const FNiagaraVariable& AllocatedVariable) const
 {
-	bool bValue = false;
-	if (Variable.IsDataAllocated())
-	{
-		bValue = Variable.GetValue<FNiagaraBool>()->GetValue();
-	}
-	return bValue ? TEXT("true") : TEXT("false");
+	checkf(AllocatedVariable.IsDataAllocated(), TEXT("Can not generate a default value string for an unallocated variable."));
+	return Lex::ToString(AllocatedVariable.GetValue<FNiagaraBool>().GetValue());
 }
 
-void FNiagaraEditorBoolTypeUtilities::SetValueFromPinDefaultString(const FString& StringValue, FNiagaraVariable& Variable) const
+bool FNiagaraEditorBoolTypeUtilities::SetValueFromPinDefaultString(const FString& StringValue, FNiagaraVariable& Variable) const
 {
-	Variable.AllocateData();
-	if (StringValue == TEXT("true"))
-	{
-		Variable.GetValue<FNiagaraBool>()->SetValue(true);
-	}
-	else
+	bool bBoolValue;
+	if (Lex::TryParseString(bBoolValue, *StringValue))
 	{
-		Variable.GetValue<FNiagaraBool>()->SetValue(false);
+		FNiagaraBool BoolValue;
+		BoolValue.SetValue(bBoolValue);
+		Variable.SetValue<FNiagaraBool>(BoolValue);
+		return true;
 	}
+	return false;
 }
\ No newline at end of file
diff --git a/Engine/Plugins/FX/Niagara/Source/NiagaraEditor/Private/TypeEditorUtilities/NiagaraBoolTypeEditorUtilities.h b/Engine/Plugins/FX/Niagara/Source/NiagaraEditor/Private/TypeEditorUtilities/NiagaraBoolTypeEditorUtilities.h
--- a/Engine/Plugins/FX/Niagara/Source/NiagaraEditor/Private/TypeEditorUtilities/NiagaraBoolTypeEditorUtilities.h
+++ b/Engine/Plugins/FX/Niagara/Source/NiagaraEditor/Private/TypeEditorUtilities/NiagaraBoolTypeEditorUtilities.h
@@ -12,8 +12,8 @@ class FNiagaraEditorBoolTypeUtilities : public FNiagaraEditorTypeUtilities
 public:
 	//~ INiagaraEditorTypeUtilities interface.
 	virtual bool CanCreateParameterEditor() const override { return true; }
-	virtual TSharedPtr<SNiagaraParameterEditor> CreateParameterEditor() const override;
+	virtual TSharedPtr<SNiagaraParameterEditor> CreateParameterEditor(const FNiagaraTypeDefinition& ParameterType) const override;
 	virtual bool CanHandlePinDefaults() const override;
-	virtual FString GetPinDefaultStringFromValue(const FNiagaraVariable& Variable) const override;
-	virtual void SetValueFromPinDefaultString(const FString& StringValue, FNiagaraVariable& Variable) const override;
+	virtual FString GetPinDefaultStringFromValue(const FNiagaraVariable& AllocatedVariable) const override;
+	virtual bool SetValueFromPinDefaultString(const FString& StringValue, FNiagaraVariable& Variable) const override;
 };
\ No newline at end of file
diff --git a/Engine/Plugins/FX/Niagara/Source/NiagaraEditor/Private/TypeEditorUtilities/NiagaraColorTypeEditorUtilities.cpp b/Engine/Plugins/FX/Niagara/Source/NiagaraEditor/Private/TypeEditorUtilities/NiagaraColorTypeEditorUtilities.cpp
--- a/Engine/Plugins/FX/Niagara/Source/NiagaraEditor/Private/TypeEditorUtilities/NiagaraColorTypeEditorUtilities.cpp
+++ b/Engine/Plugins/FX/Niagara/Source/NiagaraEditor/Private/TypeEditorUtilities/NiagaraColorTypeEditorUtilities.cpp
@@ -192,13 +192,13 @@ private:
 	FLinearColor ColorValue;
 };
 
-void FNiagaraEditorColorTypeUtilities::UpdateStructWithDefaultValue(TSharedRef<FStructOnScope> Struct) const
+void FNiagaraEditorColorTypeUtilities::UpdateVariableWithDefaultValue(FNiagaraVariable& Variable) const
 {
-	checkf(Struct->GetStruct() == FNiagaraTypeDefinition::GetColorStruct(), TEXT("Struct type not supported."));
-	*((FLinearColor*)Struct->GetStructMemory()) = FLinearColor(1, 1, 1, 1);
+	checkf(Variable.GetType().GetStruct() == FNiagaraTypeDefinition::GetColorStruct(), TEXT("Struct type not supported."));
+	Variable.SetValue<FLinearColor>(FLinearColor(1, 1, 1, 1));
 }
 
-TSharedPtr<SNiagaraParameterEditor> FNiagaraEditorColorTypeUtilities::CreateParameterEditor() const
+TSharedPtr<SNiagaraParameterEditor> FNiagaraEditorColorTypeUtilities::CreateParameterEditor(const FNiagaraTypeDefinition& ParameterType) const
 {
 	return SNew(SNiagaraColorParameterEditor);
 }
@@ -208,17 +208,19 @@ bool FNiagaraEditorColorTypeUtilities::CanHandlePinDefaults() const
 	return true;
 }
 
-FString FNiagaraEditorColorTypeUtilities::GetPinDefaultStringFromValue(const FNiagaraVariable& Variable) const
+FString FNiagaraEditorColorTypeUtilities::GetPinDefaultStringFromValue(const FNiagaraVariable& AllocatedVariable) const
 {
-	if (Variable.IsDataAllocated())
-	{
-		return Variable.GetValue<FLinearColor>()->ToString();
-	}
-	return FVector4(0.0f, 0.0f, 0.0f, 0.0f).ToString();
+	checkf(AllocatedVariable.IsDataAllocated(), TEXT("Can not generate a default value string for an unallocated variable."));
+	return AllocatedVariable.GetValue<FLinearColor>().ToString();
 }
 
-void FNiagaraEditorColorTypeUtilities::SetValueFromPinDefaultString(const FString& StringValue, FNiagaraVariable& Variable) const
+bool FNiagaraEditorColorTypeUtilities::SetValueFromPinDefaultString(const FString& StringValue, FNiagaraVariable& Variable) const
 {
-	Variable.AllocateData();
-	Variable.GetValue<FLinearColor>()->InitFromString(StringValue);
+	FLinearColor ColorValue;
+	if (ColorValue.InitFromString(StringValue))
+	{
+		Variable.SetValue<FLinearColor>(ColorValue);
+		return true;
+	}
+	return false;
 }
\ No newline at end of file
diff --git a/Engine/Plugins/FX/Niagara/Source/NiagaraEditor/Private/TypeEditorUtilities/NiagaraColorTypeEditorUtilities.h b/Engine/Plugins/FX/Niagara/Source/NiagaraEditor/Private/TypeEditorUtilities/NiagaraColorTypeEditorUtilities.h
--- a/Engine/Plugins/FX/Niagara/Source/NiagaraEditor/Private/TypeEditorUtilities/NiagaraColorTypeEditorUtilities.h
+++ b/Engine/Plugins/FX/Niagara/Source/NiagaraEditor/Private/TypeEditorUtilities/NiagaraColorTypeEditorUtilities.h
@@ -12,10 +12,10 @@ class FNiagaraEditorColorTypeUtilities : public FNiagaraEditorTypeUtilities
 public:
 	//~ INiagaraEditorTypeUtilities interface.
 	virtual bool CanProvideDefaultValue() const override { return true; }
-	virtual void UpdateStructWithDefaultValue(TSharedRef<FStructOnScope> Struct) const override;
+	virtual void UpdateVariableWithDefaultValue(FNiagaraVariable& Variable) const override;
 	virtual bool CanCreateParameterEditor() const override { return true; }
-	virtual TSharedPtr<SNiagaraParameterEditor> CreateParameterEditor() const override;
+	virtual TSharedPtr<SNiagaraParameterEditor> CreateParameterEditor(const FNiagaraTypeDefinition& ParameterType) const override;
 	virtual bool CanHandlePinDefaults() const override;
-	virtual FString GetPinDefaultStringFromValue(const FNiagaraVariable& Variable) const override;
-	virtual void SetValueFromPinDefaultString(const FString& StringValue, FNiagaraVariable& Variable) const override;
+	virtual FString GetPinDefaultStringFromValue(const FNiagaraVariable& AllocatedVariable) const override;
+	virtual bool SetValueFromPinDefaultString(const FString& StringValue, FNiagaraVariable& Variable) const override;
 };
\ No newline at end of file
diff --git a/Engine/Plugins/FX/Niagara/Source/NiagaraEditor/Private/TypeEditorUtilities/NiagaraEnumTypeEditorUtilities.cpp b/Engine/Plugins/FX/Niagara/Source/NiagaraEditor/Private/TypeEditorUtilities/NiagaraEnumTypeEditorUtilities.cpp
new file mode 100644
--- /dev/null
+++ b/Engine/Plugins/FX/Niagara/Source/NiagaraEditor/Private/TypeEditorUtilities/NiagaraEnumTypeEditorUtilities.cpp
@@ -0,0 +1,212 @@
+// Copyright 1998-2016 Epic Games, Inc. All Rights Reserved.
+
+#include "NiagaraEnumTypeEditorUtilities.h"
+#include "SNiagaraParameterEditor.h"
+#include "NiagaraTypes.h"
+#include "NiagaraEditorStyle.h"
+#include "NiagaraEditorCommon.h"
+
+#include "Widgets/Input/SSpinBox.h"
+#include "Widgets/Input/SComboBox.h"
+
+class SEnumCombobox : public SComboBox<TSharedPtr<int32>>
+{
+public:
+	DECLARE_DELEGATE_TwoParams(FOnValueChanged, int32, ESelectInfo::Type);
+
+public:
+	SLATE_BEGIN_ARGS(SEnumCombobox) {}
+		SLATE_ATTRIBUTE(int32, Value)
+		SLATE_EVENT(FOnValueChanged, OnValueChanged)
+	SLATE_END_ARGS()
+
+	void Construct(const FArguments& InArgs, const UEnum* InEnum)
+	{
+		Enum = InEnum;
+		Value = InArgs._Value;
+		check(Value.IsBound());
+		OnValueChanged = InArgs._OnValueChanged;
+
+		bUpdatingSelectionInternally = false;
+
+		for (int32 i = 0; i < Enum->NumEnums() - 1; i++)
+		{
+			if (Enum->HasMetaData(TEXT("Hidden"), i) == false)
+			{
+				VisibleEnumNameIndices.Add(MakeShareable(new int32(i)));
+			}
+		}
+
+		SComboBox::Construct(SComboBox<TSharedPtr<int32>>::FArguments()
+			//.ButtonStyle(FEditorStyle::Get(), "FlatButton.Light")
+			.OptionsSource(&VisibleEnumNameIndices)
+			.OnGenerateWidget(this, &SEnumCombobox::OnGenerateWidget)
+			.OnSelectionChanged(this, &SEnumCombobox::OnComboSelectionChanged)
+			.OnComboBoxOpening(this, &SEnumCombobox::OnComboMenuOpening)
+			.ContentPadding(FMargin(2, 0))
+			[
+				SNew(STextBlock)
+				//.Font(FEditorStyle::GetFontStyle("Sequencer.AnimationOutliner.RegularFont"))
+				.Text(this, &SEnumCombobox::GetValueText)
+			]);
+	}
+
+private:
+	FText GetValueText() const
+	{
+		int32 ValueNameIndex = Enum->GetIndexByValue(Value.Get());
+		return Enum->GetDisplayNameTextByIndex(ValueNameIndex);
+	}
+
+	TSharedRef<SWidget> OnGenerateWidget(TSharedPtr<int32> InItem)
+	{
+		return SNew(STextBlock)
+			.Text(Enum->GetDisplayNameTextByIndex(*InItem));
+	}
+
+	void OnComboSelectionChanged(TSharedPtr<int32> InSelectedItem, ESelectInfo::Type SelectInfo)
+	{
+		if (bUpdatingSelectionInternally == false)
+		{
+			OnValueChanged.ExecuteIfBound(*InSelectedItem, SelectInfo);
+		}
+	}
+
+	void OnComboMenuOpening()
+	{
+		int32 CurrentNameIndex = Enum->GetIndexByValue(Value.Get());
+		TSharedPtr<int32> FoundNameIndexItem;
+		for (int32 i = 0; i < VisibleEnumNameIndices.Num(); i++)
+		{
+			if (*VisibleEnumNameIndices[i] == CurrentNameIndex)
+			{
+				FoundNameIndexItem = VisibleEnumNameIndices[i];
+				break;
+			}
+		}
+		if (FoundNameIndexItem.IsValid())
+		{
+			bUpdatingSelectionInternally = true;
+			SetSelectedItem(FoundNameIndexItem);
+			bUpdatingSelectionInternally = false;
+		}
+	}
+
+private:
+	const UEnum* Enum;
+
+	TAttribute<int32> Value;
+
+	TArray<TSharedPtr<int32>> VisibleEnumNameIndices;
+
+	bool bUpdatingSelectionInternally;
+
+	FOnValueChanged OnValueChanged;
+};
+
+class SNiagaraEnumParameterEditor : public SNiagaraParameterEditor
+{
+public:
+	SLATE_BEGIN_ARGS(SNiagaraEnumParameterEditor) { }
+	SLATE_END_ARGS();
+
+	void Construct(const FArguments& InArgs, const UEnum* Enum)
+	{
+		ChildSlot
+		[
+			SNew(SEnumCombobox, Enum)
+			.Value(this, &SNiagaraEnumParameterEditor::GetValue)
+			.OnValueChanged(this, &SNiagaraEnumParameterEditor::ValueChanged)
+		];
+	}
+
+	virtual void UpdateInternalValueFromStruct(TSharedRef<FStructOnScope> Struct) override
+	{
+		checkf(Struct->GetStruct() == FNiagaraTypeDefinition::GetIntStruct(), TEXT("Struct type not supported."));
+		IntValue = ((FNiagaraInt32*)Struct->GetStructMemory())->Value;
+	}
+
+	virtual void UpdateStructFromInternalValue(TSharedRef<FStructOnScope> Struct) override
+	{
+		checkf(Struct->GetStruct() == FNiagaraTypeDefinition::GetIntStruct(), TEXT("Struct type not supported."));
+		((FNiagaraInt32*)Struct->GetStructMemory())->Value = IntValue;
+	}
+
+private:
+	int32 GetValue() const
+	{
+		return IntValue;
+	}
+
+	void ValueChanged(int32 Value, ESelectInfo::Type SelectionType)
+	{
+		IntValue = Value;
+		ExecuteOnValueChanged();
+	}
+
+private:
+	int32 IntValue;
+};
+
+bool FNiagaraEditorEnumTypeUtilities::CanProvideDefaultValue() const
+{
+	return true;
+}
+
+void FNiagaraEditorEnumTypeUtilities::UpdateVariableWithDefaultValue(FNiagaraVariable& Variable) const
+{
+	const UEnum* Enum = Variable.GetType().GetEnum();
+	checkf(Enum != nullptr, TEXT("Variable is not an enum type."));
+
+	FNiagaraInt32 EnumIntValue;
+	EnumIntValue.Value = Enum->GetValueByIndex(0);
+
+	Variable.SetValue<FNiagaraInt32>(EnumIntValue);
+}
+
+TSharedPtr<SNiagaraParameterEditor> FNiagaraEditorEnumTypeUtilities::CreateParameterEditor(const FNiagaraTypeDefinition& ParameterType) const
+{
+	return SNew(SNiagaraEnumParameterEditor, ParameterType.GetEnum());
+}
+
+bool FNiagaraEditorEnumTypeUtilities::CanHandlePinDefaults() const
+{
+	return true;
+}
+
+FString FNiagaraEditorEnumTypeUtilities::GetPinDefaultStringFromValue(const FNiagaraVariable& AllocatedVariable) const
+{
+	checkf(AllocatedVariable.IsDataAllocated(), TEXT("Can not generate a default value string for an unallocated variable."));
+
+	const UEnum* Enum = AllocatedVariable.GetType().GetEnum();
+	checkf(Enum != nullptr, TEXT("Variable is not an enum type."));
+
+	int64 EnumValue = AllocatedVariable.GetValue<int32>();
+	if(Enum->IsValidEnumValue(EnumValue))
+	{
+		return Enum->GetNameStringByValue(EnumValue);
+	}
+	else
+	{
+		FString ReplacementForInvalidValue = Enum->GetNameStringByIndex(0);
+		UE_LOG(LogNiagaraEditor, Error, TEXT("Error getting default value for enum pin.  Enum value %i is not supported for enum type %s.  Using value %s"),
+			EnumValue, *AllocatedVariable.GetType().GetName(), *ReplacementForInvalidValue);
+		return ReplacementForInvalidValue;
+	}
+}
+
+bool FNiagaraEditorEnumTypeUtilities::SetValueFromPinDefaultString(const FString& StringValue, FNiagaraVariable& Variable) const
+{
+	const UEnum* Enum = Variable.GetType().GetEnum();
+	checkf(Enum != nullptr, TEXT("Variable is not an enum type."));
+
+	FNiagaraInt32 EnumValue;
+	EnumValue.Value = (int32)Enum->GetValueByNameString(StringValue, EGetByNameFlags::ErrorIfNotFound);
+	if(EnumValue.Value != INDEX_NONE)
+	{
+		Variable.AllocateData();
+		Variable.SetValue<FNiagaraInt32>(EnumValue);
+		return true;
+	}
+	return false;
+}
\ No newline at end of file
diff --git a/Engine/Plugins/FX/Niagara/Source/NiagaraEditor/Private/TypeEditorUtilities/NiagaraFloatTypeEditorUtilities.cpp b/Engine/Plugins/FX/Niagara/Source/NiagaraEditor/Private/TypeEditorUtilities/NiagaraFloatTypeEditorUtilities.cpp
--- a/Engine/Plugins/FX/Niagara/Source/NiagaraEditor/Private/TypeEditorUtilities/NiagaraFloatTypeEditorUtilities.cpp
+++ b/Engine/Plugins/FX/Niagara/Source/NiagaraEditor/Private/TypeEditorUtilities/NiagaraFloatTypeEditorUtilities.cpp
@@ -85,7 +85,7 @@ private:
 	float FloatValue;
 };
 
-TSharedPtr<SNiagaraParameterEditor> FNiagaraEditorFloatTypeUtilities::CreateParameterEditor() const
+TSharedPtr<SNiagaraParameterEditor> FNiagaraEditorFloatTypeUtilities::CreateParameterEditor(const FNiagaraTypeDefinition& ParameterType) const
 {
 	return SNew(SNiagaraFloatParameterEditor);
 }
@@ -95,18 +95,19 @@ bool FNiagaraEditorFloatTypeUtilities::CanHandlePinDefaults() const
 	return true;
 }
 
-FString FNiagaraEditorFloatTypeUtilities::GetPinDefaultStringFromValue(const FNiagaraVariable& Variable) const
+FString FNiagaraEditorFloatTypeUtilities::GetPinDefaultStringFromValue(const FNiagaraVariable& AllocatedVariable) const
 {
-	float Value = 0.0f;
-	if (Variable.IsDataAllocated())
-	{
-		Value = Variable.GetValue<FNiagaraFloat>()->Value;
-	}
-	return FString::Printf(TEXT("%3.3f"), Value);
+	checkf(AllocatedVariable.IsDataAllocated(), TEXT("Can not generate a default value string for an unallocated variable."));
+	return Lex::ToString(AllocatedVariable.GetValue<FNiagaraFloat>().Value);
 }
 
-void FNiagaraEditorFloatTypeUtilities::SetValueFromPinDefaultString(const FString& StringValue, FNiagaraVariable& Variable) const
+bool FNiagaraEditorFloatTypeUtilities::SetValueFromPinDefaultString(const FString& StringValue, FNiagaraVariable& Variable) const
 {
-	Variable.AllocateData();
-	Variable.GetValue<FNiagaraFloat>()->Value = FCString::Atof(*StringValue);
+	FNiagaraFloat FloatValue;
+	if (Lex::TryParseString(FloatValue.Value, *StringValue))
+	{
+		Variable.SetValue<FNiagaraFloat>(FloatValue);
+		return true;
+	}
+	return false;
 }
\ No newline at end of file
diff --git a/Engine/Plugins/FX/Niagara/Source/NiagaraEditor/Private/TypeEditorUtilities/NiagaraFloatTypeEditorUtilities.h b/Engine/Plugins/FX/Niagara/Source/NiagaraEditor/Private/TypeEditorUtilities/NiagaraFloatTypeEditorUtilities.h
--- a/Engine/Plugins/FX/Niagara/Source/NiagaraEditor/Private/TypeEditorUtilities/NiagaraFloatTypeEditorUtilities.h
+++ b/Engine/Plugins/FX/Niagara/Source/NiagaraEditor/Private/TypeEditorUtilities/NiagaraFloatTypeEditorUtilities.h
@@ -12,8 +12,8 @@ class FNiagaraEditorFloatTypeUtilities : public FNiagaraEditorTypeUtilities
 public:
 	//~ INiagaraEditorTypeUtilities interface.
 	virtual bool CanCreateParameterEditor() const override { return true; }
-	virtual TSharedPtr<SNiagaraParameterEditor> CreateParameterEditor() const override;
+	virtual TSharedPtr<SNiagaraParameterEditor> CreateParameterEditor(const FNiagaraTypeDefinition& ParameterType) const override;
 	virtual bool CanHandlePinDefaults() const override;
-	virtual FString GetPinDefaultStringFromValue(const FNiagaraVariable& Variable) const override;
-	virtual void SetValueFromPinDefaultString(const FString& StringValue, FNiagaraVariable& Variable) const override;
+	virtual FString GetPinDefaultStringFromValue(const FNiagaraVariable& AllocatedVariable) const override;
+	virtual bool SetValueFromPinDefaultString(const FString& StringValue, FNiagaraVariable& Variable) const override;
 };
\ No newline at end of file
diff --git a/Engine/Plugins/FX/Niagara/Source/NiagaraEditor/Private/TypeEditorUtilities/NiagaraIntegerTypeEditorUtilities.cpp b/Engine/Plugins/FX/Niagara/Source/NiagaraEditor/Private/TypeEditorUtilities/NiagaraIntegerTypeEditorUtilities.cpp
--- a/Engine/Plugins/FX/Niagara/Source/NiagaraEditor/Private/TypeEditorUtilities/NiagaraIntegerTypeEditorUtilities.cpp
+++ b/Engine/Plugins/FX/Niagara/Source/NiagaraEditor/Private/TypeEditorUtilities/NiagaraIntegerTypeEditorUtilities.cpp
@@ -5,7 +5,8 @@
 #include "NiagaraTypes.h"
 #include "NiagaraEditorStyle.h"
 
-#include "SSpinBox.h"
+#include "Widgets/Input/SSpinBox.h"
+#include "Widgets/Input/SComboBox.h"
 
 class SNiagaraIntegerParameterEditor : public SNiagaraParameterEditor
 {
@@ -85,7 +86,7 @@ private:
 	int32 IntValue;
 };
 
-TSharedPtr<SNiagaraParameterEditor> FNiagaraEditorIntegerTypeUtilities::CreateParameterEditor() const
+TSharedPtr<SNiagaraParameterEditor> FNiagaraEditorIntegerTypeUtilities::CreateParameterEditor(const FNiagaraTypeDefinition& ParameterType) const
 {
 	return SNew(SNiagaraIntegerParameterEditor);
 }
@@ -95,18 +96,19 @@ bool FNiagaraEditorIntegerTypeUtilities::CanHandlePinDefaults() const
 	return true;
 }
 
-FString FNiagaraEditorIntegerTypeUtilities::GetPinDefaultStringFromValue(const FNiagaraVariable& Variable) const
+FString FNiagaraEditorIntegerTypeUtilities::GetPinDefaultStringFromValue(const FNiagaraVariable& AllocatedVariable) const
 {
-	int32 Value = 0;
-	if (Variable.IsDataAllocated())
-	{
-		Value = Variable.GetValue<FNiagaraInt32>()->Value;
-	}
-	return FString::Printf(TEXT("%i"), Value);
+	checkf(AllocatedVariable.IsDataAllocated(), TEXT("Can not generate a default value string for an unallocated variable."));
+	return Lex::ToString(AllocatedVariable.GetValue<FNiagaraFloat>().Value);
 }
 
-void FNiagaraEditorIntegerTypeUtilities::SetValueFromPinDefaultString(const FString& StringValue, FNiagaraVariable& Variable) const
+bool FNiagaraEditorIntegerTypeUtilities::SetValueFromPinDefaultString(const FString& StringValue, FNiagaraVariable& Variable) const
 {
-	Variable.AllocateData();
-	Variable.GetValue<FNiagaraInt32>()->Value = FCString::Atoi(*StringValue);
+	FNiagaraInt32 IntegerValue;
+	if(Lex::TryParseString(IntegerValue.Value, *StringValue))
+	{
+		Variable.SetValue<FNiagaraInt32>(IntegerValue);
+		return true;
+	}
+	return false;
 }
\ No newline at end of file
diff --git a/Engine/Plugins/FX/Niagara/Source/NiagaraEditor/Private/TypeEditorUtilities/NiagaraVectorTypeEditorUtilities.cpp b/Engine/Plugins/FX/Niagara/Source/NiagaraEditor/Private/TypeEditorUtilities/NiagaraVectorTypeEditorUtilities.cpp
--- a/Engine/Plugins/FX/Niagara/Source/NiagaraEditor/Private/TypeEditorUtilities/NiagaraVectorTypeEditorUtilities.cpp
+++ b/Engine/Plugins/FX/Niagara/Source/NiagaraEditor/Private/TypeEditorUtilities/NiagaraVectorTypeEditorUtilities.cpp
@@ -143,7 +143,7 @@ private:
 	FVector2D VectorValue;
 };
 
-TSharedPtr<SNiagaraParameterEditor> FNiagaraEditorVector2TypeUtilities::CreateParameterEditor() const
+TSharedPtr<SNiagaraParameterEditor> FNiagaraEditorVector2TypeUtilities::CreateParameterEditor(const FNiagaraTypeDefinition& ParameterType) const
 {
 	return SNew(SNiagaraVector2ParameterEditor);
 }
@@ -153,19 +153,21 @@ bool FNiagaraEditorVector2TypeUtilities::CanHandlePinDefaults() const
 	return true;
 }
 
-FString FNiagaraEditorVector2TypeUtilities::GetPinDefaultStringFromValue(const FNiagaraVariable& Variable) const
+FString FNiagaraEditorVector2TypeUtilities::GetPinDefaultStringFromValue(const FNiagaraVariable& AllocatedVariable) const
 {
-	if (Variable.IsDataAllocated())
-	{
-		return Variable.GetValue<FVector2D>()->ToString();
-	}
-	return FVector2D(0.0f, 0.0f).ToString();
+	checkf(AllocatedVariable.IsDataAllocated(), TEXT("Can not generate a default value string for an unallocated variable."));
+	return AllocatedVariable.GetValue<FVector2D>().ToString();
 }
 
-void FNiagaraEditorVector2TypeUtilities::SetValueFromPinDefaultString(const FString& StringValue, FNiagaraVariable& Variable) const
+bool FNiagaraEditorVector2TypeUtilities::SetValueFromPinDefaultString(const FString& StringValue, FNiagaraVariable& Variable) const
 {
-	Variable.AllocateData();
-	Variable.GetValue<FVector2D>()->InitFromString(StringValue);
+	FVector2D VectorValue;
+	if (VectorValue.InitFromString(StringValue))
+	{
+		Variable.SetValue<FVector2D>(VectorValue);
+		return true;
+	}
+	return false;
 }
 
 class SNiagaraVector3ParameterEditor : public SNiagaraVectorParameterEditorBase
@@ -208,7 +210,7 @@ private:
 	FVector VectorValue;
 };
 
-TSharedPtr<SNiagaraParameterEditor> FNiagaraEditorVector3TypeUtilities::CreateParameterEditor() const
+TSharedPtr<SNiagaraParameterEditor> FNiagaraEditorVector3TypeUtilities::CreateParameterEditor(const FNiagaraTypeDefinition& ParameterType) const
 {
 	return SNew(SNiagaraVector3ParameterEditor);
 }
@@ -218,27 +220,25 @@ bool FNiagaraEditorVector3TypeUtilities::CanHandlePinDefaults() const
 	return true;
 }
 
-FString FNiagaraEditorVector3TypeUtilities::GetPinDefaultStringFromValue(const FNiagaraVariable& Variable) const
+FString FNiagaraEditorVector3TypeUtilities::GetPinDefaultStringFromValue(const FNiagaraVariable& AllocatedVariable) const
 {
+	checkf(AllocatedVariable.IsDataAllocated(), TEXT("Can not generate a default value string for an unallocated variable."));
+
 	// NOTE: We can not use ToString() here since the vector pin control doesn't use the standard 'X=0,Y=0,Z=0' syntax.
-	FVector Value(0.0f, 0.0f, 0.0f);
-	if (Variable.IsDataAllocated())
-	{
-		Value = *Variable.GetValue<FVector>();
-	}
+	FVector Value = AllocatedVariable.GetValue<FVector>();
 	return FString::Printf(TEXT("%3.3f,%3.3f,%3.3f"), Value.X, Value.Y, Value.Z);
 }
 
-void FNiagaraEditorVector3TypeUtilities::SetValueFromPinDefaultString(const FString& StringValue, FNiagaraVariable& Variable) const
+bool FNiagaraEditorVector3TypeUtilities::SetValueFromPinDefaultString(const FString& StringValue, FNiagaraVariable& Variable) const
 {
 	// NOTE: We can not use InitFromString() here since the vector pin control doesn't use the standard 'X=0,Y=0,Z=0' syntax.
 	FVector Value;
-	if (FDefaultValueHelper::ParseVector(StringValue, Value) == false)
+	if (FDefaultValueHelper::ParseVector(StringValue, Value))
 	{
-		Value = FVector(0.0f, 0.0f, 0.0f);
+		Variable.SetValue<FVector>(Value);
+		return true;
 	}
-	Variable.AllocateData();
-	*Variable.GetValue<FVector>() = Value;
+	return false;
 }
 
 class SNiagaraVector4ParameterEditor : public SNiagaraVectorParameterEditorBase
@@ -281,7 +281,7 @@ private:
 	FVector4 VectorValue;
 };
 
-TSharedPtr<SNiagaraParameterEditor> FNiagaraEditorVector4TypeUtilities::CreateParameterEditor() const
+TSharedPtr<SNiagaraParameterEditor> FNiagaraEditorVector4TypeUtilities::CreateParameterEditor(const FNiagaraTypeDefinition& ParameterType) const
 {
 	return SNew(SNiagaraVector4ParameterEditor);
 }
@@ -291,42 +291,22 @@ bool FNiagaraEditorVector4TypeUtilities::CanHandlePinDefaults() const
 	return true;
 }
 
-FString FNiagaraEditorVector4TypeUtilities::GetPinDefaultStringFromValue(const FNiagaraVariable& Variable) const
+FString FNiagaraEditorVector4TypeUtilities::GetPinDefaultStringFromValue(const FNiagaraVariable& AllocatedVariable) const
 {
+	checkf(AllocatedVariable.IsDataAllocated(), TEXT("Can not generate a default value string for an unallocated variable."));
 	// NOTE: We can not use ToString() here since the vector pin control doesn't use the standard 'X=0,Y=0,Z=0,W=0' syntax.
-	FVector4 Value(0.0f, 0.0f, 0.0f, 0.0f);
-	if (Variable.IsDataAllocated())
-	{
-		// FVector4 is aligned, so casting a pointer to it won't work
-		const uint8 *VarData = Variable.GetData();
-		FPlatformMemory::Memcpy(&Value, VarData, sizeof(FVector4));
-		//Value = *Variable.GetValue<FVector4>();
-	}
+	FVector4 Value = AllocatedVariable.GetValue<FVector4>();
 	return FString::Printf(TEXT("%3.3f,%3.3f,%3.3f,%3.3f"), Value.X, Value.Y, Value.Z, Value.W);
 }
 
-void FNiagaraEditorVector4TypeUtilities::SetValueFromPinDefaultString(const FString& StringValue, FNiagaraVariable& Variable) const
+bool FNiagaraEditorVector4TypeUtilities::SetValueFromPinDefaultString(const FString& StringValue, FNiagaraVariable& Variable) const
 {
 	// NOTE: We can not use InitFromString() here since the vector pin control doesn't use the standard 'X=0,Y=0,Z=0,W=0' syntax.
 	FVector4 Value;
-	if (FDefaultValueHelper::ParseVector4(StringValue, Value) == false)
-	{
-		Value = FVector4(0.0f, 0.0f, 0.0f, 0.0f);
-	}
-
-	// FVector4 requires aligned memory, which cannot be safely guaranteed by the AllocateData call. Therefore, we cannot 
-	// write it as a FVector4. Instead, we must write component by component.
-	Variable.AllocateData();
-	void* ValueWrite = Variable.GetValue<FVector4>();
-	if (nullptr != ValueWrite)
-	{
-		((float*)ValueWrite)[0] = Value.X;
-		((float*)ValueWrite)[1] = Value.Y;
-		((float*)ValueWrite)[2] = Value.Z;
-		((float*)ValueWrite)[3] = Value.W;
-	}
-	else
+	if (FDefaultValueHelper::ParseVector4(StringValue, Value))
 	{
-		check(!"UNKNOWNW");
+		Variable.SetValue<FVector4>(Value);
+		return true;
 	}
+	return false;
 }
\ No newline at end of file
diff --git a/Engine/Plugins/FX/Niagara/Source/NiagaraEditor/Private/TypeEditorUtilities/NiagaraVectorTypeEditorUtilities.h b/Engine/Plugins/FX/Niagara/Source/NiagaraEditor/Private/TypeEditorUtilities/NiagaraVectorTypeEditorUtilities.h
--- a/Engine/Plugins/FX/Niagara/Source/NiagaraEditor/Private/TypeEditorUtilities/NiagaraVectorTypeEditorUtilities.h
+++ b/Engine/Plugins/FX/Niagara/Source/NiagaraEditor/Private/TypeEditorUtilities/NiagaraVectorTypeEditorUtilities.h
@@ -12,10 +12,10 @@ class FNiagaraEditorVector2TypeUtilities : public FNiagaraEditorTypeUtilities
 public:
 	//~ INiagaraEditorTypeUtilities interface.
 	virtual bool CanCreateParameterEditor() const override { return true; }
-	virtual TSharedPtr<SNiagaraParameterEditor> CreateParameterEditor() const override;
+	virtual TSharedPtr<SNiagaraParameterEditor> CreateParameterEditor(const FNiagaraTypeDefinition& ParameterType) const override;
 	virtual bool CanHandlePinDefaults() const override;
-	virtual FString GetPinDefaultStringFromValue(const FNiagaraVariable& Variable) const override;
-	virtual void SetValueFromPinDefaultString(const FString& StringValue, FNiagaraVariable& Variable) const override;
+	virtual FString GetPinDefaultStringFromValue(const FNiagaraVariable& AllocatedVariable) const override;
+	virtual bool SetValueFromPinDefaultString(const FString& StringValue, FNiagaraVariable& Variable) const override;
 };
 
 /** Niagara editor utilities for the FVector type. */
@@ -24,10 +24,10 @@ class FNiagaraEditorVector3TypeUtilities : public FNiagaraEditorTypeUtilities
 public:
 	//~ INiagaraEditorTypeUtilities interface.
 	virtual bool CanCreateParameterEditor() const override { return true; }
-	virtual TSharedPtr<SNiagaraParameterEditor> CreateParameterEditor() const override;
+	virtual TSharedPtr<SNiagaraParameterEditor> CreateParameterEditor(const FNiagaraTypeDefinition& ParameterType) const override;
 	virtual bool CanHandlePinDefaults() const override;
-	virtual FString GetPinDefaultStringFromValue(const FNiagaraVariable& Variable) const override;
-	virtual void SetValueFromPinDefaultString(const FString& StringValue, FNiagaraVariable& Variable) const override;
+	virtual FString GetPinDefaultStringFromValue(const FNiagaraVariable& AllocatedVariable) const override;
+	virtual bool SetValueFromPinDefaultString(const FString& StringValue, FNiagaraVariable& Variable) const override;
 };
 
 /** Niagara editor utilities for the FVector4 type. */
@@ -36,8 +36,8 @@ class FNiagaraEditorVector4TypeUtilities : public FNiagaraEditorTypeUtilities
 public:
 	//~ INiagaraEditorTypeUtilities interface.
 	virtual bool CanCreateParameterEditor() const override { return true; }
-	virtual TSharedPtr<SNiagaraParameterEditor> CreateParameterEditor() const override;
+	virtual TSharedPtr<SNiagaraParameterEditor> CreateParameterEditor(const FNiagaraTypeDefinition& ParameterType) const override;
 	virtual bool CanHandlePinDefaults() const override;
-	virtual FString GetPinDefaultStringFromValue(const FNiagaraVariable& Variable) const override;
-	virtual void SetValueFromPinDefaultString(const FString& StringValue, FNiagaraVariable& Variable) const override;
+	virtual FString GetPinDefaultStringFromValue(const FNiagaraVariable& AllocatedVariable) const override;
+	virtual bool SetValueFromPinDefaultString(const FString& StringValue, FNiagaraVariable& Variable) const override;
 };
\ No newline at end of file
diff --git a/Engine/Plugins/FX/Niagara/Source/NiagaraEditor/Private/ViewModels/NiagaraEmitterHandleViewModel.cpp b/Engine/Plugins/FX/Niagara/Source/NiagaraEditor/Private/ViewModels/NiagaraEmitterHandleViewModel.cpp
--- a/Engine/Plugins/FX/Niagara/Source/NiagaraEditor/Private/ViewModels/NiagaraEmitterHandleViewModel.cpp
+++ b/Engine/Plugins/FX/Niagara/Source/NiagaraEditor/Private/ViewModels/NiagaraEmitterHandleViewModel.cpp
@@ -108,44 +108,6 @@ EVisibility FNiagaraEmitterHandleViewModel::GetErrorTextVisibility() const
 	return EmitterViewModel->GetLatestCompileStatus() != ENiagaraScriptCompileStatus::NCS_UpToDate ? EVisibility::Visible : EVisibility::Collapsed;
 }
 
-
-FText FNiagaraEmitterHandleViewModel::GetSourceSynchronizationText() const
-{
-	
-	if (EmitterHandle->IsSynchronizedWithSource())
-	{
-		return LOCTEXT("SynchronizedWithSource", "Up-To-Date");
-	}
-	else
-	{
-		return LOCTEXT("NotSynchronizedWithSource", "Not Up-To-Date");
-	}
-}
-
-EVisibility FNiagaraEmitterHandleViewModel::GetSourceSynchronizationTextVisibility() const
-{
-	if (OwningSystem.GetAutoImportChangedEmitters())
-	{
-		return EVisibility::Collapsed;
-	}
-	else
-	{
-		return EVisibility::Visible;
-	}
-}
-
-FSlateColor FNiagaraEmitterHandleViewModel::GetSourceSynchronizationTextColor() const
-{
-	if (EmitterHandle->IsSynchronizedWithSource())
-	{
-		return FSlateColor::UseForeground();
-	}
-	else
-	{
-		return FSlateColor(FLinearColor::Yellow);
-	}
-}
-
 FName FNiagaraEmitterHandleViewModel::GetName() const
 {
 	if (EmitterHandle)
@@ -164,20 +126,13 @@ void FNiagaraEmitterHandleViewModel::SetName(FName InName)
 
 	if (EmitterHandle)
 	{
-		TSet<FName> OtherEmitterNames;
-		for (const FNiagaraEmitterHandle& OtherEmitterHandle : OwningSystem.GetEmitterHandles())
-		{
-			if (OtherEmitterHandle.GetId() != EmitterHandle->GetId())
-			{
-				OtherEmitterNames.Add(OtherEmitterHandle.GetName());
-			}
-		}
-		FName UniqueName = FNiagaraEditorUtilities::GetUniqueName(InName, OtherEmitterNames);
-
 		FScopedTransaction ScopedTransaction(NSLOCTEXT("NiagaraEmitterEditor", "EditEmitterNameTransaction", "Edit emitter name"));
 		OwningSystem.Modify();
-		EmitterHandle->SetName(UniqueName);
+		OwningSystem.RemoveSystemParametersForEmitter(*EmitterHandle);
+		EmitterHandle->SetName(InName, OwningSystem);
+		OwningSystem.RefreshSystemParametersFromEmitter(*EmitterHandle);
 		OnPropertyChangedDelegate.Broadcast();
+		OnNameChangedDelegate.Broadcast();
 	}
 }
 
@@ -250,54 +205,9 @@ TSharedRef<FNiagaraEmitterViewModel> FNiagaraEmitterHandleViewModel::GetEmitterV
 	return EmitterViewModel;
 }
 
-void FNiagaraEmitterHandleViewModel::CompileScripts()
-{
-	EmitterViewModel->CompileScripts();
-}
-
-
-
-void FNiagaraEmitterHandleViewModel::RefreshFromSource()
-{
-	FScopedTransaction ScopedTransaction(NSLOCTEXT("NiagaraEmitterEditor", "RefreshFromSource", "Reset emitter from source."));
-	OwningSystem.Modify();
-
-	if (EmitterHandle)
-	{
-		{
-			// Pull in changes to the emitter asset by copying the source scripts, compiling and then copying over parameter values
-			// where relevant.
-			if (EmitterHandle->RefreshFromSource() == false)
-			{
-				UE_LOG(LogNiagaraEditor, Error, TEXT("Failed to compile during refresh. Refresh cancelled. %s"), *EmitterHandle->GetSource()->GetPathName());
-
-				const FText NotificationText = FText::Format(LOCTEXT("FailedScriptRefresh", "'{0}' failed to refresh due to compile errors. Please see log."), FText::FromString(EmitterHandle->GetSource()->GetName()));
-
-				FNotificationInfo Info(NotificationText);
-				Info.bFireAndForget = true;
-				Info.bUseThrobber = true;
-				Info.bUseSuccessFailIcons = true;
-				Info.ExpireDuration = 10.0f;
-				FSlateNotificationManager::Get().AddNotification(Info);
-			}
-		}
-
-		EmitterViewModel->SetEmitter(EmitterHandle->GetInstance());
-	}
-	OnPropertyChangedDelegate.Broadcast();
-}
-
-void FNiagaraEmitterHandleViewModel::ResetToSource()
+void FNiagaraEmitterHandleViewModel::CompileScripts(bool bForce)
 {
-	FScopedTransaction ScopedTransaction(NSLOCTEXT("NiagaraEmitterEditor", "ResetToSource", "Reset emitter to source."));
-	OwningSystem.Modify();
-	
-	if (EmitterHandle)
-	{
-		EmitterHandle->ResetToSource();
-		EmitterViewModel->SetEmitter(EmitterHandle->GetInstance());
-	}
-	OnPropertyChangedDelegate.Broadcast();
+	EmitterViewModel->CompileScripts(bForce);
 }
 
 void FNiagaraEmitterHandleViewModel::OpenSourceEmitter()
@@ -313,4 +223,9 @@ FNiagaraEmitterHandleViewModel::FOnPropertyChanged& FNiagaraEmitterHandleViewMod
 	return OnPropertyChangedDelegate;
 }
 
+FNiagaraEmitterHandleViewModel::FOnNameChanged& FNiagaraEmitterHandleViewModel::OnNameChanged()
+{
+	return OnNameChangedDelegate;
+}
+
 #undef LOCTEXT_NAMESPACE
\ No newline at end of file
diff --git a/Engine/Plugins/FX/Niagara/Source/NiagaraEditor/Private/ViewModels/NiagaraEmitterHandleViewModel.h b/Engine/Plugins/FX/Niagara/Source/NiagaraEditor/Private/ViewModels/NiagaraEmitterHandleViewModel.h
--- a/Engine/Plugins/FX/Niagara/Source/NiagaraEditor/Private/ViewModels/NiagaraEmitterHandleViewModel.h
+++ b/Engine/Plugins/FX/Niagara/Source/NiagaraEditor/Private/ViewModels/NiagaraEmitterHandleViewModel.h
@@ -18,6 +18,7 @@ class FNiagaraEmitterHandleViewModel : public TSharedFromThis<FNiagaraEmitterHan
 {
 public:
 	DECLARE_MULTICAST_DELEGATE(FOnPropertyChanged);
+	DECLARE_MULTICAST_DELEGATE(FOnNameChanged);
 public:
 	/** Creates a new emitter editor view model with the supplied emitter handle and simulation. */
 	FNiagaraEmitterHandleViewModel(FNiagaraEmitterHandle* InEmitterHandle, TWeakPtr<FNiagaraEmitterInstance> InSimulation, UNiagaraSystem& InOwningSystem);
@@ -58,9 +59,6 @@ public:
 	FSlateColor GetErrorTextColor() const;
 
 	/** Called to get the synch state of the emitter handle to its source.*/
-	FText GetSourceSynchronizationText() const;
-	EVisibility GetSourceSynchronizationTextVisibility() const;
-	FSlateColor GetSourceSynchronizationTextColor() const;
 	bool IsSynchronized() const;
 
 	/** Gets whether or not this emitter handle is enabled. */
@@ -82,15 +80,7 @@ public:
 	TSharedRef<FNiagaraEmitterViewModel> GetEmitterViewModel();
 
 	/** Compiles the spawn and update scripts. */
-	void CompileScripts();
-
-	/** Refreshes the copied emitter's graph and inputs from the source asset.  Input values
-		will be preserved. */
-	void RefreshFromSource();
-
-	/** Replaces the coped emitter instance with a fresh copy of the source emitter asset.  Any changes
-		to input parameters will be lost. */
-	void ResetToSource();
+	void CompileScripts(bool bForce);
 
 	/** Opens the source emitter in a stand alone asset editor. */
 	void OpenSourceEmitter();
@@ -98,6 +88,9 @@ public:
 	/** Gets a multicast delegate which is called any time a property on the handle changes. */
 	FOnPropertyChanged& OnPropertyChanged();
 
+	/** Gets a multicast delegate which is called any time this handle is renamed. */
+	FOnNameChanged& OnNameChanged();
+
 private:
 	/** The emitter handle being displayed and edited by this view model. */
 	FNiagaraEmitterHandle* EmitterHandle;
@@ -110,4 +103,7 @@ private:
 
 	/** A multicast delegate which is called any time a property on the handle changes. */
 	FOnPropertyChanged OnPropertyChangedDelegate;
+
+	/** A multicast delegate which is called any time this emitter handle is renamed. */
+	FOnNameChanged OnNameChangedDelegate;
 };
\ No newline at end of file
diff --git a/Engine/Plugins/FX/Niagara/Source/NiagaraEditor/Private/ViewModels/NiagaraEmitterViewModel.cpp b/Engine/Plugins/FX/Niagara/Source/NiagaraEditor/Private/ViewModels/NiagaraEmitterViewModel.cpp
--- a/Engine/Plugins/FX/Niagara/Source/NiagaraEditor/Private/ViewModels/NiagaraEmitterViewModel.cpp
+++ b/Engine/Plugins/FX/Niagara/Source/NiagaraEditor/Private/ViewModels/NiagaraEmitterViewModel.cpp
@@ -18,7 +18,7 @@
 
 template<> TMap<UNiagaraEmitter*, TArray<FNiagaraEmitterViewModel*>> TNiagaraViewModelManager<UNiagaraEmitter, FNiagaraEmitterViewModel>::ObjectsToViewModels{};
 
-const FText FNiagaraEmitterViewModel::StatsFormat = NSLOCTEXT("NiagaraEmitterViewModel", "StatsFormat", "{0} Particles | {1} ms | {2} MB");
+const FText FNiagaraEmitterViewModel::StatsFormat = NSLOCTEXT("NiagaraEmitterViewModel", "StatsFormat", "{0} Particles | {1} ms | {2} MB | {3}");
 const float Megabyte = 1024.0f * 1024.0f;
 
 FNiagaraEmitterViewModel::FNiagaraEmitterViewModel(UNiagaraEmitter* InEmitter, TWeakPtr<FNiagaraEmitterInstance> InSimulation)
@@ -27,9 +27,10 @@ FNiagaraEmitterViewModel::FNiagaraEmitterViewModel(UNiagaraEmitter* InEmitter, T
 	, SharedScriptViewModel(MakeShareable(new FNiagaraScriptViewModel(InEmitter, LOCTEXT("SharedDisplayName", "Graph"), ENiagaraParameterEditMode::EditAll)))
 	, bUpdatingSelectionInternally(false)
 	, LastEventScriptStatus(ENiagaraScriptCompileStatus::NCS_Unknown)
-	, bEmitterDirty(false)
 {
-	if (Emitter.IsValid() && Emitter->EventHandlerScriptProps.Num() != 0 && Emitter->EventHandlerScriptProps[0].Script && Emitter->EventHandlerScriptProps[0].Script->ByteCode.Num() != 0)
+	ExecutionStateEnum = FindObject<UEnum>(ANY_PACKAGE, TEXT("ENiagaraExecutionState")); 
+
+	if (Emitter.IsValid() && Emitter->GetEventHandlers().Num() != 0 && Emitter->GetEventHandlers()[0].Script && Emitter->GetEventHandlers()[0].Script->GetByteCode().Num() != 0)
 	{
 		LastEventScriptStatus = ENiagaraScriptCompileStatus::NCS_UpToDate;
 	}
@@ -115,16 +116,51 @@ FText FNiagaraEmitterViewModel::GetStatsText() const
 	if (Simulation.IsValid())
 	{
 		TSharedPtr<FNiagaraEmitterInstance> SimInstance = Simulation.Pin();
-		check(SimInstance.IsValid());
-		return FText::Format(StatsFormat,
-			FText::AsNumber(SimInstance->GetNumParticles()),
-			FText::AsNumber(SimInstance->GetTotalCPUTime()),
-			FText::AsNumber(SimInstance->GetTotalBytesUsed() / Megabyte));
+		if (SimInstance.IsValid())
+		{
+			if (!SimInstance->IsReadyToRun())
+			{
+				return LOCTEXT("PendingCompile", "Compilation in progress...");
+			}
+
+			const FNiagaraEmitterHandle& Handle = SimInstance->GetEmitterHandle();
+			if (Handle.GetInstance())
+			{
+				if (Handle.IsValid() == false)
+				{
+					return LOCTEXT("InvalidHandle", "Invalid handle");
+				}
+
+				UNiagaraEmitter* HandleEmitter = Handle.GetInstance();
+				if (HandleEmitter == nullptr)
+				{
+					return LOCTEXT("NullInstance", "Invalid instance");
+				}
+
+				if (!HandleEmitter->IsValid())
+				{
+					return LOCTEXT("InvalidInstance", "Invalid Emitter! May have compile errors.");
+				}
+
+				if (Handle.GetIsEnabled() == false)
+				{
+					return LOCTEXT("DisabledSimulation", "Simulation is not enabled.");
+				}
+
+				return FText::Format(StatsFormat,
+					FText::AsNumber(SimInstance->GetNumParticles()),
+					FText::AsNumber(SimInstance->GetTotalCPUTime()),
+					FText::AsNumber(SimInstance->GetTotalBytesUsed() / Megabyte),
+					ExecutionStateEnum->GetDisplayNameTextByValue((int32)SimInstance->GetExecutionState()));
+			}
+		}
 	}
-	else
+	else if(!Emitter->IsReadyToRun())
 	{
-		return LOCTEXT("InvalidSimulation", "Simulation is invalid.");
+		return LOCTEXT("SimulationNotReady", "Preparing simulation...");
 	}
+	
+	return LOCTEXT("InvalidSimulation", "Simulation is invalid.");
 }
 
 TSharedRef<FNiagaraScriptViewModel> FNiagaraEmitterViewModel::GetSharedScriptViewModel()
@@ -132,18 +168,6 @@ TSharedRef<FNiagaraScriptViewModel> FNiagaraEmitterViewModel::GetSharedScriptVie
 	return SharedScriptViewModel;
 }
 
-bool FNiagaraEmitterViewModel::GetDirty() const
-{
-	bool bDirty = SharedScriptViewModel->GetScriptDirty() || bEmitterDirty;
-	return bDirty;
-}
-
-void FNiagaraEmitterViewModel::SetDirty(bool bDirty)
-{
-	SharedScriptViewModel->SetScriptDirty(bDirty);
-	bEmitterDirty = bDirty;
-}
-
 const UNiagaraEmitterEditorData& FNiagaraEmitterViewModel::GetEditorData() const
 {
 	check(Emitter.IsValid());
@@ -169,7 +193,7 @@ UNiagaraEmitterEditorData& FNiagaraEmitterViewModel::GetOrCreateEditorData()
 	return *EditorData;
 }
 
-void FNiagaraEmitterViewModel::CompileScripts()
+void FNiagaraEmitterViewModel::CompileScripts(bool bForce)
 {
 	if (Emitter.IsValid())
 	{
@@ -178,7 +202,7 @@ void FNiagaraEmitterViewModel::CompileScripts()
 		TArray<FString> CompilePaths;
 		TArray<UNiagaraScript*> Scripts;
 		TArray<TPair<ENiagaraScriptUsage, int32> > Usages;
-		Emitter->CompileScripts(CompileStatuses, CompileErrors, CompilePaths, Scripts);
+		Emitter->CompileScripts(CompileStatuses, CompileErrors, CompilePaths, Scripts, bForce);
 
 		ENiagaraScriptCompileStatus AggregateStatus = ENiagaraScriptCompileStatus::NCS_UpToDate;
 		FString AggregateErrors;
diff --git a/Engine/Plugins/FX/Niagara/Source/NiagaraEditor/Private/ViewModels/NiagaraEmitterViewModel.h b/Engine/Plugins/FX/Niagara/Source/NiagaraEditor/Private/ViewModels/NiagaraEmitterViewModel.h
--- a/Engine/Plugins/FX/Niagara/Source/NiagaraEditor/Private/ViewModels/NiagaraEmitterViewModel.h
+++ b/Engine/Plugins/FX/Niagara/Source/NiagaraEditor/Private/ViewModels/NiagaraEmitterViewModel.h
@@ -59,7 +59,7 @@ public:
 	TSharedRef<FNiagaraScriptViewModel> GetSharedScriptViewModel();
 
 	/** Compiles the spawn and update scripts. */
-	void CompileScripts();
+	void CompileScripts(bool bForce);
 
 	/* Get the latest status of this view-model's script compilation.*/
 	ENiagaraScriptCompileStatus GetLatestCompileStatus();
@@ -72,9 +72,6 @@ public:
 
 	/** Gets a delegate which is called when the shared script is compiled. */
 	FOnScriptCompiled& OnScriptCompiled();
-		
-	bool GetDirty() const;
-	void SetDirty(bool bDirty);
 
 	/** Gets editor specific data which can be stored per emitter.  If this data hasn't been created the default version will be returned. */
 	const UNiagaraEmitterEditorData& GetEditorData() const;
@@ -109,7 +106,7 @@ private:
 
 	ENiagaraScriptCompileStatus LastEventScriptStatus;
 
-	bool bEmitterDirty;
-
 	TNiagaraViewModelManager<UNiagaraEmitter, FNiagaraEmitterViewModel>::Handle RegisteredHandle;
+
+	UEnum* ExecutionStateEnum;
 };
\ No newline at end of file
