Copying Orion-Staging @ 3748653 (Orion/Dev-General @ 3722124) to //UE4/Main

diff --git a/Engine/Plugins/FX/Niagara/Source/NiagaraEditor/Private/ViewModels/NiagaraMetaDataCollectionViewModel.cpp b/Engine/Plugins/FX/Niagara/Source/NiagaraEditor/Private/ViewModels/NiagaraMetaDataCollectionViewModel.cpp
new file mode 100644
--- /dev/null
+++ b/Engine/Plugins/FX/Niagara/Source/NiagaraEditor/Private/ViewModels/NiagaraMetaDataCollectionViewModel.cpp
@@ -0,0 +1,115 @@
+// Copyright 1998-2017 Epic Games, Inc. All Rights Reserved.
+
+#include "NiagaraMetaDataCollectionViewModel.h"
+#include "NiagaraTypes.h"
+#include "NiagaraScript.h"
+#include "NiagaraScriptSourceBase.h"
+#include "NiagaraScriptSource.h"
+#include "NiagaraGraph.h"
+#include "NiagaraStackGraphUtilities.h"
+#include "NiagaraNodeOutput.h"
+#include "NiagaraConstants.h"
+#include "NiagaraMetaDataViewModel.h"
+#include "GraphEditAction.h"
+#include "EdGraphSchema_Niagara.h"
+
+#define LOCTEXT_NAMESPACE "NiagaraMetaDataCollectionViewModel"
+
+FNiagaraMetaDataCollectionViewModel::FNiagaraMetaDataCollectionViewModel()
+{
+	ModuleGraph = nullptr;
+}
+
+void FNiagaraMetaDataCollectionViewModel::SetGraph(UNiagaraGraph* InGraph)
+{
+	if (InGraph == nullptr)
+	{
+		return;
+	}
+	// now build variables
+	ModuleGraph = InGraph;
+	
+	// Get the variable metadata from the graph
+	const UEdGraphSchema_Niagara* Schema = GetDefault<UEdGraphSchema_Niagara>();
+	TArray<UNiagaraNode*> NiagaraNodes;
+	ModuleGraph->GetNodesOfClass<UNiagaraNode>(NiagaraNodes);
+	for (UEdGraphNode* GraphNode : NiagaraNodes)
+	{
+		for (UEdGraphPin* CurrentPin : GraphNode->Pins)
+		{
+			// parse all pins, add all valid metadata to object array.
+			FNiagaraVariable NiagaraVariable = Schema->PinToNiagaraVariable(CurrentPin, true);
+			if (NiagaraVariable.IsInNameSpace(FNiagaraParameterHandle::ModuleNamespace.ToString()) && !FNiagaraConstants::IsNiagaraConstant(NiagaraVariable))
+			{
+				TSharedPtr<FNiagaraMetaDataViewModel> MetadataViewModel = GetMetadataViewModelForVariable(NiagaraVariable);
+				bool bNeedsToAdd = false;
+				if (!MetadataViewModel.IsValid())
+				{
+					MetadataViewModel = TSharedPtr<FNiagaraMetaDataViewModel>(new FNiagaraMetaDataViewModel(NiagaraVariable, *ModuleGraph));
+					bNeedsToAdd = true;
+				}
+				MetadataViewModel->AssociateNode(CurrentPin->GetOwningNode());
+				MetadataViewModel->OnMetadataChanged().AddRaw(this, &FNiagaraMetaDataCollectionViewModel::BroadcastChanged);
+				if (bNeedsToAdd)
+				{
+					MetaDataViewModels.Add(MetadataViewModel.ToSharedRef());
+				}
+			}
+		}
+	}
+
+	OnGraphChangedHandle = ModuleGraph->AddOnGraphChangedHandler(FOnGraphChanged::FDelegate::CreateRaw(this, &FNiagaraMetaDataCollectionViewModel::OnGraphChanged));
+}
+
+TSharedPtr<FNiagaraMetaDataViewModel> FNiagaraMetaDataCollectionViewModel::GetMetadataViewModelForVariable(FNiagaraVariable InVariable)
+{
+	for (TSharedPtr<FNiagaraMetaDataViewModel> Model : MetaDataViewModels)
+	{
+		if (InVariable == Model->GetVariable())
+			return Model;
+	}
+	return TSharedPtr<FNiagaraMetaDataViewModel>();
+}
+
+TArray<TSharedRef<FNiagaraMetaDataViewModel>>& FNiagaraMetaDataCollectionViewModel::GetVariableModels()
+{
+	return MetaDataViewModels;
+}
+
+void FNiagaraMetaDataCollectionViewModel::OnGraphChanged(const struct FEdGraphEditAction& InAction)
+{
+	Reload(); 
+} 
+
+void FNiagaraMetaDataCollectionViewModel::BroadcastChanged()
+{
+	OnCollectionChangedDelegate.Broadcast();
+}
+
+void FNiagaraMetaDataCollectionViewModel::Reload()
+{
+	Cleanup();
+	SetGraph(ModuleGraph);
+	BroadcastChanged();
+}
+
+void FNiagaraMetaDataCollectionViewModel::Cleanup()
+{
+	for (int32 i = 0; i < MetaDataViewModels.Num(); i++)
+	{
+		FNiagaraMetaDataViewModel* MetaDataViewModel = (FNiagaraMetaDataViewModel*)(&MetaDataViewModels[i].Get());
+		if (MetaDataViewModel != nullptr)
+		{
+			MetaDataViewModel->OnMetadataChanged().RemoveAll(this);
+		}
+	}
+	MetaDataViewModels.Empty();
+	ModuleGraph->RemoveOnGraphChangedHandler(OnGraphChangedHandle);
+}
+
+FNiagaraMetaDataCollectionViewModel::~FNiagaraMetaDataCollectionViewModel()
+{
+	Cleanup();
+}
+
+#undef LOCTEXT_NAMESPACE // NiagaraMetaDataCollectionViewModel
\ No newline at end of file
diff --git a/Engine/Plugins/FX/Niagara/Source/NiagaraEditor/Private/ViewModels/NiagaraMetaDataCollectionViewModel.h b/Engine/Plugins/FX/Niagara/Source/NiagaraEditor/Private/ViewModels/NiagaraMetaDataCollectionViewModel.h
new file mode 100644
--- /dev/null
+++ b/Engine/Plugins/FX/Niagara/Source/NiagaraEditor/Private/ViewModels/NiagaraMetaDataCollectionViewModel.h
@@ -0,0 +1,50 @@
+// Copyright 1998-2017 Epic Games, Inc. All Rights Reserved.
+
+#pragma once
+
+#include "NiagaraTypes.h"
+class UNiagaraScript;
+class UNiagaraGraph;
+class FNiagaraMetaDataViewModel;
+
+class FNiagaraMetaDataCollectionViewModel 
+{
+public:
+	DECLARE_MULTICAST_DELEGATE(FOnCollectionChanged);
+
+public:
+	FNiagaraMetaDataCollectionViewModel();
+
+	~FNiagaraMetaDataCollectionViewModel();
+
+	/** Sets the view model to a new script. */
+	void SetGraph(UNiagaraGraph* InGraph);
+
+	/** Gets the metadata view models. */
+	TArray<TSharedRef<FNiagaraMetaDataViewModel>>& GetVariableModels();
+	
+	/** Reloads the data from the graph and builds the viewmodel*/
+	void Reload();
+
+	/** returns the delegate to be called when the collection changes*/
+	FOnCollectionChanged& OnCollectionChanged() { return OnCollectionChangedDelegate; }
+	
+private:
+	/** Callback for when the graph changes and the collection viewmodel needs to react */
+	void OnGraphChanged(const struct FEdGraphEditAction& InAction);
+	void Cleanup();
+	void BroadcastChanged();
+	TSharedPtr<FNiagaraMetaDataViewModel> GetMetadataViewModelForVariable(FNiagaraVariable InVariable);
+private:
+	/** The variables */
+	TArray <TSharedRef<FNiagaraMetaDataViewModel>> MetaDataViewModels;
+
+	/** The actual graph of the module we're editing */
+	UNiagaraGraph* ModuleGraph;
+
+	/** The handle to the graph changed delegate. */
+	FDelegateHandle OnGraphChangedHandle;
+
+	/** A multicast delegate which is called whenever the parameter collection is changed. */
+	FOnCollectionChanged OnCollectionChangedDelegate;
+};
\ No newline at end of file
diff --git a/Engine/Plugins/FX/Niagara/Source/NiagaraEditor/Private/ViewModels/NiagaraMetaDataViewModel.cpp b/Engine/Plugins/FX/Niagara/Source/NiagaraEditor/Private/ViewModels/NiagaraMetaDataViewModel.cpp
new file mode 100644
--- /dev/null
+++ b/Engine/Plugins/FX/Niagara/Source/NiagaraEditor/Private/ViewModels/NiagaraMetaDataViewModel.cpp
@@ -0,0 +1,90 @@
+// Copyright 1998-2016 Epic Games, Inc. All Rights Reserved.
+
+#include "NiagaraMetaDataViewModel.h"
+#include "NiagaraEditorUtilities.h"
+#include "NiagaraTypes.h"
+#include "NiagaraNodeInput.h"
+#include "NiagaraGraph.h"
+#include "NiagaraScriptSource.h"
+#include "ScopedTransaction.h"
+#include "Editor.h"
+#include "Array.h"
+#include "StructOnScope.h"
+
+#define LOCTEXT_NAMESPACE "NiagaraMetaDataViewModel"
+
+FNiagaraMetaDataViewModel::FNiagaraMetaDataViewModel(FNiagaraVariable& InGraphVariable, UNiagaraGraph& InGraph)
+	: GraphVariable(InGraphVariable)
+	, CurrentGraph(&InGraph)
+{
+	RefreshMetaDataValue();
+}
+
+FNiagaraMetaDataViewModel::~FNiagaraMetaDataViewModel()
+{
+	//UE_LOG(LogNiagaraEditor, Log, TEXT("Delete %p Var %s"), this, *DebugName);
+	GraphMetaData = nullptr;
+	CurrentGraph = nullptr;
+}
+
+FName FNiagaraMetaDataViewModel::GetName() const
+{
+	return GraphVariable.GetName();
+}
+
+void FNiagaraMetaDataViewModel::AssociateNode(UEdGraphNode* InNode)
+{
+	AssociatedNodes.AddUnique(TWeakObjectPtr<UObject>(InNode));
+}
+
+void FNiagaraMetaDataViewModel::NotifyMetaDataChanged()
+{
+	FScopedTransaction ScopedTransaction(LOCTEXT("EditVariableMetadata", "Edit variable metadata"));
+	CurrentGraph->Modify();
+	
+	FNiagaraVariableMetaData* ValueData = (FNiagaraVariableMetaData*)ValueStruct->GetStructMemory();
+	if (GraphMetaData == nullptr)
+	{
+		GraphMetaData = &CurrentGraph->FindOrAddMetaData(GraphVariable);
+	}
+	*GraphMetaData = *ValueData; // deep copy
+	for (TWeakObjectPtr<UObject> RefNode : AssociatedNodes)
+	{
+		if (RefNode.IsValid())
+			GraphMetaData->ReferencerNodes.AddUnique(RefNode);
+	}
+	
+	OnMetadataChangedDelegate.Broadcast();
+	
+	// Trigger save prompt on close; it's sufficiently cheap
+	CurrentGraph->NotifyGraphNeedsRecompile();
+}
+
+FNiagaraVariableMetaData*  FNiagaraMetaDataViewModel::GetGraphMetaData()
+{
+	return GraphMetaData;
+}
+
+FNiagaraVariable FNiagaraMetaDataViewModel::GetVariable() const
+{
+	return GraphVariable;
+}
+
+void FNiagaraMetaDataViewModel::RefreshMetaDataValue()
+{
+	GraphMetaData = CurrentGraph->GetMetaData(GraphVariable); 
+	ValueStruct = MakeShareable(new FStructOnScope(FNiagaraVariableMetaData::StaticStruct()));
+	FNiagaraVariableMetaData* ValueData = (FNiagaraVariableMetaData*)ValueStruct->GetStructMemory();
+	if (GraphMetaData != nullptr)
+	{
+		// this will create a deep copy due to the nature of the structure members. it will be owned and managed by ValueStruct
+		*ValueData = *GraphMetaData; 
+	}
+}
+
+TSharedRef<FStructOnScope> FNiagaraMetaDataViewModel::GetValueStruct()
+{
+	return ValueStruct.ToSharedRef();
+}
+
+#undef LOCTEXT_NAMESPACE // "NiagaraMetaDataViewModel"
\ No newline at end of file
diff --git a/Engine/Plugins/FX/Niagara/Source/NiagaraEditor/Private/ViewModels/NiagaraMetaDataViewModel.h b/Engine/Plugins/FX/Niagara/Source/NiagaraEditor/Private/ViewModels/NiagaraMetaDataViewModel.h
new file mode 100644
--- /dev/null
+++ b/Engine/Plugins/FX/Niagara/Source/NiagaraEditor/Private/ViewModels/NiagaraMetaDataViewModel.h
@@ -0,0 +1,71 @@
+// Copyright 1998-2016 Epic Games, Inc. All Rights Reserved.
+
+#pragma once
+
+#include "CoreMinimal.h"
+#include "NiagaraTypes.h"
+
+class UNiagaraNodeInput;
+class UNiagaraGraph;
+class UEdGraphNode;
+class FStructOnScope;
+struct FNiagaraVariable;
+struct FNiagaraVariableMetaData;
+
+/** A view model for a variable metadata. */
+class FNiagaraMetaDataViewModel
+{
+public:
+	DECLARE_MULTICAST_DELEGATE(FOnMetadataChanged);
+public:
+	/*
+	 * Create a new variable metadata view model.
+	 * @param InGraphVariable The variable which is owned by the graph which connects to this metadata.
+	 * @param InGraph The graph where the metadata lives.
+	 */
+	FNiagaraMetaDataViewModel(FNiagaraVariable& InGraphVariable, UNiagaraGraph& InGraph);
+	
+	virtual ~FNiagaraMetaDataViewModel();
+
+	virtual FName GetName() const;
+
+	FNiagaraVariable GetVariable() const;
+	
+	/** Gets a pointer to the actual metadata on the graph */
+	FNiagaraVariableMetaData* GetGraphMetaData();
+	
+	/** Callback for when a change in the value structure needs to be copied into the graph*/
+	void NotifyMetaDataChanged();
+
+	/** Gets the FStructOnScope object that holds the editable value of the metadata*/
+	TSharedRef<FStructOnScope> GetValueStruct();
+
+	/** Gets a multicast delegate which is called whenever the metadata changes. */
+	FOnMetadataChanged& OnMetadataChanged() { return OnMetadataChangedDelegate; }
+
+	/** associates a node to this metadata*/
+	void AssociateNode(UEdGraphNode* InNode);
+	
+private:
+	/** Refreshes the parameter value struct from the variable data. */
+	void RefreshMetaDataValue();
+
+private:
+	/** The graph variable which corresponds to the metadata in our viewmodel */
+	FNiagaraVariable GraphVariable;
+	
+	/** The metadata which is being displayed and edited by this view model. */
+	FNiagaraVariableMetaData* GraphMetaData;
+	
+	/** graph who owns the metadata*/
+	UNiagaraGraph* CurrentGraph;
+
+	/** graph nodes associated with this metadata*/
+	TArray<TWeakObjectPtr<UObject>> AssociatedNodes;
+	
+	/** actual value being edited*/
+	TSharedPtr<FStructOnScope> ValueStruct;
+	
+	/** A multicast delegate which is called whenever the metadata is changed. */
+	FOnMetadataChanged OnMetadataChangedDelegate;
+};
\ No newline at end of file
diff --git a/Engine/Plugins/FX/Niagara/Source/NiagaraEditor/Private/ViewModels/NiagaraParameterCollectionAssetViewModel.cpp b/Engine/Plugins/FX/Niagara/Source/NiagaraEditor/Private/ViewModels/NiagaraParameterCollectionAssetViewModel.cpp
--- a/Engine/Plugins/FX/Niagara/Source/NiagaraEditor/Private/ViewModels/NiagaraParameterCollectionAssetViewModel.cpp
+++ b/Engine/Plugins/FX/Niagara/Source/NiagaraEditor/Private/ViewModels/NiagaraParameterCollectionAssetViewModel.cpp
@@ -15,8 +15,6 @@
 
 #include "NiagaraDataInterface.h"
 
-#include "NiagaraNodeParameterCollection.h"
-
 #define LOCTEXT_NAMESPACE "NiagaraScriptInputCollection"
 
 template<> TMap<UNiagaraParameterCollection*, TArray<FNiagaraParameterCollectionAssetViewModel*>> TNiagaraViewModelManager<UNiagaraParameterCollection, FNiagaraParameterCollectionAssetViewModel>::ObjectsToViewModels{};
@@ -105,7 +103,7 @@ FName FNiagaraParameterCollectionAssetViewModel::GenerateNewName(FNiagaraTypeDef
 		Existing.Add(ParameterViewModel->GetName());
 	}
 
-	return *Collection->ParameterNameFromFriendlyName(FNiagaraEditorUtilities::GetUniqueName(ProposedName, Existing).ToString());
+	return *Collection->ParameterNameFromFriendlyName(FNiagaraUtilities::GetUniqueName(ProposedName, Existing).ToString());
 }
 
 void FNiagaraParameterCollectionAssetViewModel::AddParameter(TSharedPtr<FNiagaraTypeDefinition> ParameterType)
@@ -223,14 +221,14 @@ void FNiagaraParameterCollectionAssetViewModel::CollectionChanged(bool bRecompil
 	//Refresh any existing view models that might be showing changed instances.
 	UpdateOpenInstances();
 
-	//Refresh any nodes that are referencing this collection.
-	for (TObjectIterator<UNiagaraNodeParameterCollection> It; It; ++It)
-	{
-		if (It->GetReferencedAsset() == Collection)
-		{
-			It->RefreshFromExternalChanges();
-		}
-	}
+// 	//Refresh any nodes that are referencing this collection.
+// 	for (TObjectIterator<UNiagaraNodeParameterCollection> It; It; ++It)
+// 	{
+// 		if (It->GetReferencedAsset() == Collection)
+// 		{
+// 			It->RefreshFromExternalChanges();
+// 		}
+// 	}
 
 	if (bRecompile)
 	{
@@ -359,7 +357,13 @@ void FNiagaraParameterCollectionAssetViewModel::OnParameterProvidedChanged(FNiag
 
 void FNiagaraParameterCollectionAssetViewModel::OnParameterValueChangedInternal(TSharedRef<FNiagaraCollectionParameterViewModel> ChangedParameter)
 {
+	//restart any systems using this collection.
+	FNiagaraSystemUpdateContext UpdateContext(Collection, true);
+
 	OnParameterValueChanged().Broadcast(ChangedParameter->GetName());
+	
+	//Push the change to anyone already bound.
+	Instance->GetParameterStore().Tick();
 }
 
 #undef LOCTEXT_NAMESPACE // "NiagaraScriptInputCollectionViewModel"
diff --git a/Engine/Plugins/FX/Niagara/Source/NiagaraEditor/Private/ViewModels/NiagaraScriptInputCollectionViewModel.cpp b/Engine/Plugins/FX/Niagara/Source/NiagaraEditor/Private/ViewModels/NiagaraScriptInputCollectionViewModel.cpp
--- a/Engine/Plugins/FX/Niagara/Source/NiagaraEditor/Private/ViewModels/NiagaraScriptInputCollectionViewModel.cpp
+++ b/Engine/Plugins/FX/Niagara/Source/NiagaraEditor/Private/ViewModels/NiagaraScriptInputCollectionViewModel.cpp
@@ -330,7 +330,7 @@ void FNiagaraScriptInputCollectionViewModel::RefreshParameterViewModels()
 			}
 			else
 			{
-				UNiagaraDataInterface* EmitterDataInterface = InputNode->DataInterface;
+				UNiagaraDataInterface* EmitterDataInterface = InputNode->GetDataInterface();
 				UNiagaraScript* Script = nullptr;
 				for (TWeakObjectPtr<UNiagaraScript> ScriptWeakPtr : Scripts)
 				{
@@ -403,7 +403,7 @@ void FNiagaraScriptInputCollectionViewModel::OnParameterNameChanged(FName OldNam
 	TSet<FName> SystemConstantNames = FNiagaraEditorUtilities::GetSystemConstantNames();
 
 	// Rename the nodes and notify the graph that they've changed.
-	FName UniqueNewName = FNiagaraEditorUtilities::GetUniqueName(ParameterVariable->GetName(), CurrentNames.Union(SystemConstantNames));
+	FName UniqueNewName = FNiagaraUtilities::GetUniqueName(ParameterVariable->GetName(), CurrentNames.Union(SystemConstantNames));
 	if (ParameterVariable->GetName() != UniqueNewName)
 	{
 		ParameterVariable->SetName(UniqueNewName);
@@ -495,7 +495,7 @@ void FNiagaraScriptInputCollectionViewModel::OnParameterTypeChanged(FNiagaraVari
 			UNiagaraNodeInput* InputNodeToUpdate = InputNodesToUpdate[i];
 			InputNodeToUpdate->Modify();
 			InputNodeToUpdate->Input = FirstNodeToUpdate->Input;
-			InputNodeToUpdate->DataInterface = FirstNodeToUpdate->DataInterface;
+			InputNodeToUpdate->SetDataInterface(FirstNodeToUpdate->GetDataInterface());
 			InputNodeToUpdate->NotifyInputTypeChanged();
 		}
 	}
@@ -573,7 +573,7 @@ void FNiagaraScriptInputCollectionViewModel::OnParameterValueChangedInternal(TSh
 				TSet<UNiagaraDataInterface*> DataInterfacesToUpdate;
 				for (UNiagaraNodeInput* InputNodeToUpdate : InputNodesToUpdate)
 				{
-					DataInterfacesToUpdate.Add(InputNodeToUpdate->DataInterface);
+					DataInterfacesToUpdate.Add(InputNodeToUpdate->GetDataInterface());
 				}
 
 				for (UNiagaraDataInterface* DataInterfaceToUpdate : DataInterfacesToUpdate)
diff --git a/Engine/Plugins/FX/Niagara/Source/NiagaraEditor/Private/ViewModels/NiagaraScriptOutputCollectionViewModel.cpp b/Engine/Plugins/FX/Niagara/Source/NiagaraEditor/Private/ViewModels/NiagaraScriptOutputCollectionViewModel.cpp
--- a/Engine/Plugins/FX/Niagara/Source/NiagaraEditor/Private/ViewModels/NiagaraScriptOutputCollectionViewModel.cpp
+++ b/Engine/Plugins/FX/Niagara/Source/NiagaraEditor/Private/ViewModels/NiagaraScriptOutputCollectionViewModel.cpp
@@ -173,7 +173,7 @@ void FNiagaraScriptOutputCollectionViewModel::AddParameter(TSharedPtr<FNiagaraTy
 	}
 
 	FScopedTransaction ScopedTransaction(LOCTEXT("AddScriptOutput", "Add script output"));
-	FName OutputName = FNiagaraEditorUtilities::GetUniqueName(FName(TEXT("NewOutput")), GetParameterNames());
+	FName OutputName = FNiagaraUtilities::GetUniqueName(FName(TEXT("NewOutput")), GetParameterNames());
 	for (UNiagaraNodeOutput* OtherOutputNode : OutputNodes)
 	{
 		OtherOutputNode->Modify();
@@ -297,7 +297,7 @@ void FNiagaraScriptOutputCollectionViewModel::OnParameterNameChanged(FName OldNa
 	// If the new name matches a current output name, rename it to something unique.
 	if (CurrentNames.Contains(ParameterVariable->GetName()))
 	{
-		ParameterVariable->SetName(FNiagaraEditorUtilities::GetUniqueName(ParameterVariable->GetName(), CurrentNames));
+		ParameterVariable->SetName(FNiagaraUtilities::GetUniqueName(ParameterVariable->GetName(), CurrentNames));
 	}
 	
 	OutputNode->NotifyOutputVariablesChanged();
diff --git a/Engine/Plugins/FX/Niagara/Source/NiagaraEditor/Private/ViewModels/NiagaraScriptViewModel.cpp b/Engine/Plugins/FX/Niagara/Source/NiagaraEditor/Private/ViewModels/NiagaraScriptViewModel.cpp
--- a/Engine/Plugins/FX/Niagara/Source/NiagaraEditor/Private/ViewModels/NiagaraScriptViewModel.cpp
+++ b/Engine/Plugins/FX/Niagara/Source/NiagaraEditor/Private/ViewModels/NiagaraScriptViewModel.cpp
@@ -8,6 +8,7 @@
 #include "NiagaraEmitter.h"
 #include "NiagaraScriptInputCollectionViewModel.h"
 #include "NiagaraScriptOutputCollectionViewModel.h"
+#include "NiagaraMetaDataCollectionViewModel.h"
 #include "NiagaraNodeInput.h"
 #include "INiagaraCompiler.h"
 #include "CompilerResultsLog.h"
@@ -22,6 +23,7 @@ template<> TMap<UNiagaraScript*, TArray<FNiagaraScriptViewModel*>> TNiagaraViewM
 FNiagaraScriptViewModel::FNiagaraScriptViewModel(UNiagaraScript* InScript, FText DisplayName, ENiagaraParameterEditMode InParameterEditMode)
 	: InputCollectionViewModel(MakeShareable(new FNiagaraScriptInputCollectionViewModel(InScript, DisplayName, InParameterEditMode)))
 	, OutputCollectionViewModel(MakeShareable(new FNiagaraScriptOutputCollectionViewModel(InScript, InParameterEditMode)))
+	, MetaDataCollectionViewModel(MakeShareable(new FNiagaraMetaDataCollectionViewModel()))
 	, GraphViewModel(MakeShareable(new FNiagaraScriptGraphViewModel(Cast<UNiagaraScriptSource>(InScript->GetSource()), DisplayName)))
 	, bUpdatingSelectionInternally(false)
 	, LastCompileStatus(ENiagaraScriptCompileStatus::NCS_Unknown)
@@ -34,9 +36,11 @@ FNiagaraScriptViewModel::FNiagaraScriptViewModel(UNiagaraScript* InScript, FText
 	}
 
 	InputCollectionViewModel->GetSelection().OnSelectedObjectsChanged().AddRaw(this, &FNiagaraScriptViewModel::InputViewModelSelectionChanged);
+	InputCollectionViewModel->OnParameterValueChanged().AddRaw(this, &FNiagaraScriptViewModel::InputParameterValueChanged);
+	OutputCollectionViewModel->OnParameterValueChanged().AddRaw(this, &FNiagaraScriptViewModel::OutputParameterValueChanged);
 	GraphViewModel->GetSelection()->OnSelectedObjectsChanged().AddRaw(this, &FNiagaraScriptViewModel::GraphViewModelSelectedNodesChanged);
 	GEditor->RegisterForUndo(this);
-	
+
 	if (Source != nullptr)
 	{
 		UNiagaraGraph* Graph = Source->NodeGraph;
@@ -45,13 +49,14 @@ FNiagaraScriptViewModel::FNiagaraScriptViewModel(UNiagaraScript* InScript, FText
 			OnGraphChangedHandle = Graph->AddOnGraphChangedHandler(
 				FOnGraphChanged::FDelegate::CreateRaw(this, &FNiagaraScriptViewModel::OnGraphChanged));
 		}
+		MetaDataCollectionViewModel->SetGraph(Graph);
 	}
 
 	// Guess at initial compile status
-	if (InScript != nullptr && InScript->ByteCode.Num() == 0) // This is either a brand new script or failed in the past. Since we create a default working script, assume invalid.
+	if (InScript != nullptr && InScript->GetByteCode().Num() == 0) // This is either a brand new script or failed in the past. Since we create a default working script, assume invalid.
 	{
 		bNeedsSave = true;
-		LastCompileStatus = ENiagaraScriptCompileStatus::NCS_Error;
+		LastCompileStatus = InScript->GetLastCompileStatus();
 		GraphViewModel->SetErrorTextToolTip("Please recompile for full error stack.");
 	}
 	else // Possibly warnings previously, but still compiled. It *could* have been dirtied somehow, but we assume that it is up-to-date.
@@ -75,6 +80,7 @@ FNiagaraScriptViewModel::FNiagaraScriptViewModel(UNiagaraScript* InScript, FText
 FNiagaraScriptViewModel::FNiagaraScriptViewModel(UNiagaraEmitter* InEmitter, FText DisplayName, ENiagaraParameterEditMode InParameterEditMode)
 	: InputCollectionViewModel(MakeShareable(new FNiagaraScriptInputCollectionViewModel(InEmitter, DisplayName, InParameterEditMode)))
 	, OutputCollectionViewModel(MakeShareable(new FNiagaraScriptOutputCollectionViewModel(InEmitter, InParameterEditMode)))
+	, MetaDataCollectionViewModel(MakeShareable(new FNiagaraMetaDataCollectionViewModel()))
 	, GraphViewModel(MakeShareable(new FNiagaraScriptGraphViewModel(Cast<UNiagaraScriptSource>(InEmitter->GraphSource), DisplayName)))
 	, bUpdatingSelectionInternally(false)
 	, LastCompileStatus(ENiagaraScriptCompileStatus::NCS_Unknown)
@@ -91,10 +97,14 @@ FNiagaraScriptViewModel::FNiagaraScriptViewModel(UNiagaraEmitter* InEmitter, FTe
 	Source = Cast<UNiagaraScriptSource>(InEmitter->GraphSource);
 
 	InputCollectionViewModel->GetSelection().OnSelectedObjectsChanged().AddRaw(this, &FNiagaraScriptViewModel::InputViewModelSelectionChanged);
+	InputCollectionViewModel->OnParameterValueChanged().AddRaw(this, &FNiagaraScriptViewModel::InputParameterValueChanged);
+	OutputCollectionViewModel->OnParameterValueChanged().AddRaw(this, &FNiagaraScriptViewModel::OutputParameterValueChanged);
 	GraphViewModel->GetSelection()->OnSelectedObjectsChanged().AddRaw(this, &FNiagaraScriptViewModel::GraphViewModelSelectedNodesChanged);
 	GEditor->RegisterForUndo(this);
 
 	UNiagaraGraph* Graph = Source->NodeGraph;
+
+	MetaDataCollectionViewModel->SetGraph(Graph);
 	if (Graph != nullptr)
 	{
 		OnGraphChangedHandle = Graph->AddOnGraphChangedHandler(
@@ -113,20 +123,24 @@ FNiagaraScriptViewModel::FNiagaraScriptViewModel(UNiagaraEmitter* InEmitter, FTe
 	{
 		FString Message;
 		ENiagaraScriptCompileStatus ScriptStatus = Scripts[i]->GetLastCompileStatus();
-		if (Scripts[i].IsValid() && Scripts[i]->ByteCode.Num() == 0) // This is either a brand new script or failed in the past. Since we create a default working script, assume invalid.
+		if (Scripts[i].IsValid() && Scripts[i]->IsCompilable() && Scripts[i]->GetByteCode().Num() == 0) // This is either a brand new script or failed in the past. Since we create a default working script, assume invalid.
 		{
 			bNeedsSave = true;
-			LastCompileStatus = FNiagaraEditorUtilities::UnionCompileStatus(LastCompileStatus, ENiagaraScriptCompileStatus::NCS_Error);
-			ScriptStatus = ENiagaraScriptCompileStatus::NCS_Error;
+			LastCompileStatus = ScriptStatus;
 			Message = TEXT("Please recompile for full error stack.");
 			GraphViewModel->SetErrorTextToolTip(Message);
 		}
 		else // Possibly warnings previously, but still compiled. It *could* have been dirtied somehow, but we assume that it is up-to-date.
 		{
-			if (Scripts[i].IsValid() && Scripts[i]->AreScriptAndSourceSynchronized())
+			if (Scripts[i].IsValid() && Scripts[i]->IsCompilable()&& Scripts[i]->AreScriptAndSourceSynchronized())
 			{
 				LastCompileStatus = FNiagaraEditorUtilities::UnionCompileStatus(LastCompileStatus, Scripts[i]->GetLastCompileStatus());
 			}
+			else if (Scripts[i].IsValid() && !Scripts[i]->IsCompilable())
+			{
+				// Do nothing..
+				ScriptStatus = ENiagaraScriptCompileStatus::NCS_UpToDate;
+			}
 			else
 			{
 				bNeedsSave = true;
@@ -148,10 +162,7 @@ FNiagaraScriptViewModel::FNiagaraScriptViewModel(UNiagaraEmitter* InEmitter, FTe
 	for (int32 i = 0; i < CompileStatuses.Num(); i++)
 	{
 		CompileTypes[i].Key = Scripts[i]->GetUsage();
-		if (i > 0 && CompileTypes[i - 1].Key == Scripts[i]->GetUsage())
-			CompileTypes[i].Value = CompileTypes[i - 1].Value + 1;
-		else
-			CompileTypes[i].Value = 0;
+		CompileTypes[i].Value = Scripts[i]->GetUsageId();
 	}
 }
 
@@ -164,7 +175,7 @@ bool FNiagaraScriptViewModel::IsGraphDirty() const
 			continue;
 		}
 
-		if (!Scripts[i]->AreScriptAndSourceSynchronized())
+		if (Scripts[i]->IsCompilable() && !Scripts[i]->AreScriptAndSourceSynchronized())
 		{
 			return true;
 		}
@@ -241,7 +252,8 @@ void FNiagaraScriptViewModel::SetScripts(UNiagaraScriptSource* InScriptSource, T
 				FOnGraphChanged::FDelegate::CreateRaw(this, &FNiagaraScriptViewModel::OnGraphChanged));
 		}
 	}
-	
+	MetaDataCollectionViewModel->SetGraph(Graph);
+
 	// Guess at initial compile status
 	LastCompileStatus = ENiagaraScriptCompileStatus::NCS_UpToDate;
 	bNeedsSave = false;
@@ -254,20 +266,24 @@ void FNiagaraScriptViewModel::SetScripts(UNiagaraScriptSource* InScriptSource, T
 	{
 		FString Message;
 		ENiagaraScriptCompileStatus ScriptStatus = Scripts[i]->GetLastCompileStatus();
-		if (Scripts[i].IsValid() && Scripts[i]->ByteCode.Num() == 0) // This is either a brand new script or failed in the past. Since we create a default working script, assume invalid.
+		if (Scripts[i].IsValid() && Scripts[i]->IsCompilable() && Scripts[i]->GetByteCode().Num() == 0) // This is either a brand new script or failed in the past. Since we create a default working script, assume invalid.
 		{
 			bNeedsSave = true;
-			LastCompileStatus = FNiagaraEditorUtilities::UnionCompileStatus(LastCompileStatus, ENiagaraScriptCompileStatus::NCS_Error);
-			ScriptStatus = ENiagaraScriptCompileStatus::NCS_Error;
+			LastCompileStatus = LastCompileStatus;
 			Message = TEXT("Please recompile for full error stack.");
 			GraphViewModel->SetErrorTextToolTip(Message);
 		}
 		else // Possibly warnings previously, but still compiled. It *could* have been dirtied somehow, but we assume that it is up-to-date.
 		{
-			if (Scripts[i].IsValid() && Scripts[i]->AreScriptAndSourceSynchronized())
+			if (Scripts[i].IsValid() && Scripts[i]->IsCompilable() && Scripts[i]->AreScriptAndSourceSynchronized())
 			{
 				LastCompileStatus = FNiagaraEditorUtilities::UnionCompileStatus(LastCompileStatus, Scripts[i]->GetLastCompileStatus());
 			}
+			else if (Scripts[i]->IsCompilable())
+			{
+				//DO nothing
+				ScriptStatus = ENiagaraScriptCompileStatus::NCS_UpToDate;
+			}
 			else
 			{
 				bNeedsSave = true;
@@ -289,10 +305,7 @@ void FNiagaraScriptViewModel::SetScripts(UNiagaraScriptSource* InScriptSource, T
 	for (int32 i = 0; i < CompileStatuses.Num(); i++)
 	{
 		CompileTypes[i].Key = Scripts[i]->GetUsage();
-		if (i > 0 && CompileTypes[i - 1].Key == Scripts[i]->GetUsage())
-			CompileTypes[i].Value = CompileTypes[i - 1].Value + 1;
-		else
-			CompileTypes[i].Value = 0;
+		CompileTypes[i].Value = Scripts[i]->GetUsageId();
 	}
 }
 
@@ -323,19 +336,24 @@ void FNiagaraScriptViewModel::SetScript(UNiagaraScript* InScript)
 	SetScripts(InSource, InScripts);
 }
 
+void FNiagaraScriptViewModel::MarkAllDirty()
+{
+	for (int32 i = 0; i < Scripts.Num(); i++)
+	{
+		if (Scripts[i].IsValid() && Scripts[i]->IsCompilable())
+		{
+			Scripts[i]->MarkScriptAndSourceDesynchronized();
+		}
+	}
+	bNeedsSave = true;
+}
+
 void FNiagaraScriptViewModel::OnGraphChanged(const struct FEdGraphEditAction& InAction)
 {
 	if ((InAction.Action & GRAPHACTION_AddNode) != 0 || (InAction.Action & GRAPHACTION_RemoveNode) != 0 ||
 		(InAction.Action & GRAPHACTION_GenericNeedsRecompile) != 0)
 	{
-		for (int32 i = 0; i < Scripts.Num(); i++)
-		{
-			if (Scripts[i].IsValid())
-			{
-				Scripts[i]->MarkScriptAndSourceDesynchronized();
-			}
-		}
-		bNeedsSave = true;
+		MarkAllDirty();
 	}
 }
 
@@ -350,11 +368,30 @@ TSharedRef<FNiagaraScriptOutputCollectionViewModel> FNiagaraScriptViewModel::Get
 	return OutputCollectionViewModel;
 }
 
+TSharedRef<FNiagaraMetaDataCollectionViewModel> FNiagaraScriptViewModel::GetMetadataCollectionViewModel()
+{
+	return MetaDataCollectionViewModel;
+}
+
 TSharedRef<FNiagaraScriptGraphViewModel> FNiagaraScriptViewModel::GetGraphViewModel()
 {
 	return GraphViewModel;
 }
 
+UNiagaraScript* FNiagaraScriptViewModel::GetStandaloneScript()
+{
+	if (Scripts.Num() == 1)
+	{
+		UNiagaraScript* ScriptOutput = Scripts[0].Get();
+		if (ScriptOutput && ScriptOutput->IsStandaloneScript())
+		{
+			return ScriptOutput;
+		}
+	}
+	return nullptr;
+}
+
+
 void FNiagaraScriptViewModel::UpdateCompileStatus(ENiagaraScriptCompileStatus InAggregateCompileStatus, const FString& InAggregateCompileErrorString,
 	const TArray<ENiagaraScriptCompileStatus>& InCompileStatuses, const TArray<FString>& InCompileErrors, const TArray<FString>& InCompilePaths,
 	const TArray<UNiagaraScript*>& InCompileSources)
@@ -369,10 +406,7 @@ void FNiagaraScriptViewModel::UpdateCompileStatus(ENiagaraScriptCompileStatus In
 		for (int32 i = 0; i < CompileStatuses.Num(); i++)
 		{
 			CompileTypes[i].Key = InCompileSources[i]->GetUsage();
-			if (i > 0 && CompileTypes[i - 1].Key == InCompileSources[i]->GetUsage())
-				CompileTypes[i].Value = CompileTypes[i - 1].Value + 1;
-			else
-				CompileTypes[i].Value = 0;
+			CompileTypes[i].Value = InCompileSources[i]->GetUsageId();
 		}
 
 		LastCompileStatus = InAggregateCompileStatus;
@@ -392,12 +426,12 @@ void FNiagaraScriptViewModel::UpdateCompileStatus(ENiagaraScriptCompileStatus In
 	}
 }
 
-ENiagaraScriptCompileStatus FNiagaraScriptViewModel::GetScriptCompileStatus(ENiagaraScriptUsage InUsage, int32 InOccurrence) const
+ENiagaraScriptCompileStatus FNiagaraScriptViewModel::GetScriptCompileStatus(ENiagaraScriptUsage InUsage, FGuid InUsageId) const
 {
 	ENiagaraScriptCompileStatus Status = ENiagaraScriptCompileStatus::NCS_Unknown;
 	for (int32 i = 0; i < CompileTypes.Num(); i++)
 	{
-		if (UNiagaraScript::IsEquivalentUsage(CompileTypes[i].Key, InUsage) && CompileTypes[i].Value == InOccurrence)
+		if (UNiagaraScript::IsEquivalentUsage(CompileTypes[i].Key, InUsage) && CompileTypes[i].Value == InUsageId)
 		{
 			return CompileStatuses[i];
 		}
@@ -405,12 +439,12 @@ ENiagaraScriptCompileStatus FNiagaraScriptViewModel::GetScriptCompileStatus(ENia
 	return Status;
 }
 
-FText FNiagaraScriptViewModel::GetScriptErrors(ENiagaraScriptUsage InUsage, int32 InOccurrence) const
+FText FNiagaraScriptViewModel::GetScriptErrors(ENiagaraScriptUsage InUsage, FGuid InUsageId) const
 {
 	FText Text = FText::GetEmpty();
 	for (int32 i = 0; i < CompileTypes.Num(); i++)
 	{
-		if (UNiagaraScript::IsEquivalentUsage(CompileTypes[i].Key, InUsage) && CompileTypes[i].Value == InOccurrence)
+		if (UNiagaraScript::IsEquivalentUsage(CompileTypes[i].Key, InUsage) && CompileTypes[i].Value == InUsageId)
 		{
 			return FText::FromString(CompileErrors[i]);
 		}
@@ -425,7 +459,7 @@ void FNiagaraScriptViewModel::UpdateCompiledDataInterfaces(UNiagaraDataInterface
 	GetGraphViewModel()->GetGraph()->GetNodesOfClass<UNiagaraNodeInput>(InputNodes);
 	for (UNiagaraNodeInput* InputNode : InputNodes)
 	{
-		if (InputNode->DataInterface == ChangedDataInterface)
+		if (InputNode->GetDataInterface() == ChangedDataInterface)
 		{
 			ChangedInputName = InputNode->Input.GetName();
 			break;
@@ -450,11 +484,11 @@ void FNiagaraScriptViewModel::UpdateCompiledDataInterfaces(UNiagaraDataInterface
 	}
 }
 
-const UNiagaraScript* FNiagaraScriptViewModel::GetScript(ENiagaraScriptUsage InUsage) const
+UNiagaraScript* FNiagaraScriptViewModel::GetContainerScript(ENiagaraScriptUsage InUsage, FGuid InUsageId) 
 {
 	for (int32 i = 0; i < Scripts.Num(); i++)
 	{
-		if (Scripts[i]->IsEquivalentUsage(InUsage))
+		if (Scripts[i]->ContainsUsage(InUsage) && Scripts[i]->GetUsageId() == InUsageId)
 		{
 			return Scripts[i].Get();
 		}
@@ -462,26 +496,49 @@ const UNiagaraScript* FNiagaraScriptViewModel::GetScript(ENiagaraScriptUsage InU
 	return nullptr;
 }
 
-void FNiagaraScriptViewModel::CompileStandaloneScript()
+UNiagaraScript* FNiagaraScriptViewModel::GetScript(ENiagaraScriptUsage InUsage, FGuid InUsageId) 
 {
-	if (Source.IsValid() && Source != nullptr && Scripts.Num() == 1 && Scripts[0].IsValid () && (Scripts[0]->IsStandaloneScript()))
+	for (int32 i = 0; i < Scripts.Num(); i++)
 	{
-		FString ErrorMsg;
-		LastCompileStatus = Scripts[0]->Compile(ErrorMsg);
-		Scripts[0]->CacheResourceShadersForRendering(true);
-		bNeedsSave = true;
+		if (Scripts[i]->IsEquivalentUsage(InUsage) && Scripts[i]->GetUsageId() == InUsageId)
+		{
+			return Scripts[i].Get();
+		}
+	}
+	return nullptr;
+}
 
-		if (LastCompileStatus == ENiagaraScriptCompileStatus::NCS_Error)
+void FNiagaraScriptViewModel::CompileStandaloneScript(bool bForce)
+{
+	if (Source.IsValid() && Source != nullptr && Scripts.Num() == 1 && Scripts[0].IsValid ())
+	{
+		if (Scripts[0]->IsStandaloneScript() && Scripts[0]->IsCompilable())
 		{
-			GraphViewModel->SetErrorTextToolTip(ErrorMsg + "\n(These same errors are also in the log)");
+			FString ErrorMsg;
+			LastCompileStatus = Scripts[0]->Compile(ErrorMsg, bForce);
+			Scripts[0]->CacheResourceShadersForRendering(true);
+			bNeedsSave = true;
+
+			if (LastCompileStatus == ENiagaraScriptCompileStatus::NCS_Error)
+			{
+				GraphViewModel->SetErrorTextToolTip(ErrorMsg + "\n(These same errors are also in the log)");
+			}
+			else
+			{
+				GraphViewModel->SetErrorTextToolTip("");
+			}
+
+			InputCollectionViewModel->RefreshParameterViewModels();
+			OutputCollectionViewModel->RefreshParameterViewModels();
+		}
+		else if (!Scripts[0]->IsCompilable())
+		{
+			// do nothing
 		}
 		else
 		{
-			GraphViewModel->SetErrorTextToolTip("");
+			ensure(0); // Should not call this for a script that isn't standalone.
 		}
-
-		InputCollectionViewModel->RefreshParameterViewModels();
-		OutputCollectionViewModel->RefreshParameterViewModels();
 	}
 	else
 	{
@@ -509,7 +566,7 @@ void FNiagaraScriptViewModel::RefreshNodes()
 			{
 				for (int32 i = 0; i < Scripts.Num(); i++)
 				{
-					if (Scripts[i].IsValid())
+					if (Scripts[i].IsValid() && Scripts[i]->IsCompilable())
 					{
 						Scripts[i]->MarkScriptAndSourceDesynchronized();
 					}
@@ -587,4 +644,14 @@ void FNiagaraScriptViewModel::InputViewModelSelectionChanged()
 		}
 		bUpdatingSelectionInternally = false;
 	}
+}
+
+void FNiagaraScriptViewModel::InputParameterValueChanged(FName ParameterName)
+{
+	MarkAllDirty();
+}
+
+void FNiagaraScriptViewModel::OutputParameterValueChanged(FName ParameterName)
+{
+	MarkAllDirty();
 }
\ No newline at end of file
diff --git a/Engine/Plugins/FX/Niagara/Source/NiagaraEditor/Private/ViewModels/NiagaraScriptViewModel.h b/Engine/Plugins/FX/Niagara/Source/NiagaraEditor/Private/ViewModels/NiagaraScriptViewModel.h
--- a/Engine/Plugins/FX/Niagara/Source/NiagaraEditor/Private/ViewModels/NiagaraScriptViewModel.h
+++ b/Engine/Plugins/FX/Niagara/Source/NiagaraEditor/Private/ViewModels/NiagaraScriptViewModel.h
@@ -14,10 +14,11 @@ class INiagaraParameterCollectionViewModel;
 class FNiagaraScriptGraphViewModel;
 class FNiagaraScriptInputCollectionViewModel;
 class FNiagaraScriptOutputCollectionViewModel;
+class FNiagaraMetaDataCollectionViewModel;
 class UNiagaraEmitter;
 
 
-/** A view model for niagara scripts which manages other script related view models. */
+/** A view model for Niagara scripts which manages other script related view models. */
 class FNiagaraScriptViewModel : public TSharedFromThis<FNiagaraScriptViewModel>, public FEditorUndoClient, public TNiagaraViewModelManager<UNiagaraScript, FNiagaraScriptViewModel>
 {
 public:
@@ -37,6 +38,12 @@ public:
 	/** Gets the view model for the output parameter collection. */
 	TSharedRef<FNiagaraScriptOutputCollectionViewModel> GetOutputCollectionViewModel();
 
+	/** Gets the view model for the metadata collection. */
+	TSharedRef<FNiagaraMetaDataCollectionViewModel> GetMetadataCollectionViewModel();
+
+	TSharedRef<INiagaraParameterCollectionViewModel> GetInputParameterMapViewModel();
+	TSharedRef<INiagaraParameterCollectionViewModel> GetOutputParameterMapViewModel();
+
 	/** Gets the view model for the graph. */
 	TSharedRef<FNiagaraScriptGraphViewModel> GetGraphViewModel();
 
@@ -46,7 +53,7 @@ public:
 		const TArray<UNiagaraScript*>& InCompileSources);
 
 	/** Compiles a script that isn't part of an emitter or System. */
-	void CompileStandaloneScript();
+	void CompileStandaloneScript(bool bForce);
 
 	/** Get the latest status of this view-model's script compilation.*/
 	ENiagaraScriptCompileStatus GetLatestCompileStatus();
@@ -62,14 +69,18 @@ public:
 
 	void SetScriptDirty(bool InNeedsSave) { bNeedsSave = InNeedsSave; }
 
-	const UNiagaraScript* GetScript(ENiagaraScriptUsage InUsage) const;
+	UNiagaraScript* GetContainerScript(ENiagaraScriptUsage InUsage, FGuid InUsageId = FGuid());
+	UNiagaraScript* GetScript(ENiagaraScriptUsage InUsage, FGuid InUsageId = FGuid());
 
-	ENiagaraScriptCompileStatus GetScriptCompileStatus(ENiagaraScriptUsage InUsage, int32 InOccurrence) const;
-	FText GetScriptErrors(ENiagaraScriptUsage InUsage, int32 InOccurrence) const;
+	ENiagaraScriptCompileStatus GetScriptCompileStatus(ENiagaraScriptUsage InUsage, FGuid InUsageId) const;
+	FText GetScriptErrors(ENiagaraScriptUsage InUsage, FGuid InUsageId) const;
 
 	/** Updates the compiled versions of data interfaces from changes to their source. */
 	void UpdateCompiledDataInterfaces(UNiagaraDataInterface* ChangedDataInterface);
 
+	/** If this is editing a standalone script, returns the script being edited.*/
+	UNiagaraScript* GetStandaloneScript();
+
 private:
 	/** Handles the selection changing in the graph view model. */
 	void GraphViewModelSelectedNodesChanged();
@@ -77,11 +88,20 @@ private:
 	/** Handles the selection changing in the input collection view model. */
 	void InputViewModelSelectionChanged();
 
+	/** Marks this script view model as dirty and marks the scripts as needing synchrnozation. */
+	void MarkAllDirty();
+
     /** Handle the graph being changed by the UI notifications to see if we need to mark as needing recompile.*/
 	void OnGraphChanged(const struct FEdGraphEditAction& InAction);
 
 	void SetScripts(UNiagaraScriptSource* InScriptSource, TArray<UNiagaraScript*>& InScripts);
 
+	/** Handles when a value in the input parameter collection changes. */
+	void InputParameterValueChanged(FName ParameterName);
+
+	/** Handles when a value in the output parameter collection changes. */
+	void OutputParameterValueChanged(FName ParameterName);
+
 protected:
 	/** The script which provides the data for this view model. */
 	TArray<TWeakObjectPtr<UNiagaraScript>> Scripts;
@@ -94,6 +114,9 @@ protected:
 	/** The view model for the output parameter collection .*/
 	TSharedRef<FNiagaraScriptOutputCollectionViewModel> OutputCollectionViewModel;
 
+	/** The view model for the metadata collection .*/
+	TSharedRef<FNiagaraMetaDataCollectionViewModel> MetaDataCollectionViewModel;
+
 	/** The view model for the graph. */
 	TSharedRef<FNiagaraScriptGraphViewModel> GraphViewModel;
 
@@ -116,5 +139,5 @@ protected:
 	TArray<ENiagaraScriptCompileStatus> CompileStatuses;
 	TArray<FString> CompileErrors;
 	TArray<FString> CompilePaths;
-	TArray<TPair<ENiagaraScriptUsage, int32> > CompileTypes;
+	TArray<TPair<ENiagaraScriptUsage, FGuid>> CompileTypes;
 };
\ No newline at end of file
diff --git a/Engine/Plugins/FX/Niagara/Source/NiagaraEditor/Private/ViewModels/NiagaraSystemScriptViewModel.cpp b/Engine/Plugins/FX/Niagara/Source/NiagaraEditor/Private/ViewModels/NiagaraSystemScriptViewModel.cpp
--- a/Engine/Plugins/FX/Niagara/Source/NiagaraEditor/Private/ViewModels/NiagaraSystemScriptViewModel.cpp
+++ b/Engine/Plugins/FX/Niagara/Source/NiagaraEditor/Private/ViewModels/NiagaraSystemScriptViewModel.cpp
@@ -18,11 +18,17 @@
 #include "GraphEditAction.h"
 #include "Stack/NiagaraStackEntry.h"
 #include "NiagaraStackGraphUtilities.h"
+#include "NiagaraSystemViewModel.h"
+#include "NiagaraEmitterHandleViewModel.h"
+#include "NiagaraEmitterViewModel.h"
 
-FNiagaraSystemScriptViewModel::FNiagaraSystemScriptViewModel(UNiagaraSystem& InSystem)
+FNiagaraSystemScriptViewModel::FNiagaraSystemScriptViewModel(UNiagaraSystem& InSystem, FNiagaraSystemViewModel* InParent)
 	: FNiagaraScriptViewModel(InSystem.GetSystemSpawnScript(), NSLOCTEXT("SystemScriptViewModel", "GraphName", "System"), ENiagaraParameterEditMode::EditAll)
 	, System(InSystem)
+	, Parent(InParent)
 {
+	Scripts.Add(InSystem.GetSystemUpdateScript());
+
 	if (GetGraphViewModel()->GetGraph())
 	{
 		OnGraphChangedHandle = GetGraphViewModel()->GetGraph()->AddOnGraphChangedHandler(
@@ -193,27 +199,59 @@ FNiagaraSystemScriptViewModel::FOnSystemCompiled& FNiagaraSystemScriptViewModel:
 	return OnSystemCompiledDelegate;
 }
 
-void FNiagaraSystemScriptViewModel::CompileSystem()
+void FNiagaraSystemScriptViewModel::CompileSystem(bool bForce)
 {
 
 	TArray<ENiagaraScriptCompileStatus> CompileSystemStatuses;
 	TArray<FString> CompileSystemErrors;
 	TArray<FString> CompileSystemPaths;
 	TArray<UNiagaraScript*> CompileSystemScripts;
-	System.CompileScripts(CompileSystemStatuses, CompileSystemErrors, CompileSystemPaths, CompileSystemScripts);
-
-	ENiagaraScriptCompileStatus AggregateStatus = ENiagaraScriptCompileStatus::NCS_UpToDate;
-	FString AggregateErrors;
+	System.CompileScripts(CompileSystemStatuses, CompileSystemErrors, CompileSystemPaths, CompileSystemScripts, bForce);
 
-	for (int32 i = 0; i < CompileSystemStatuses.Num(); i++)
+	// First we need to find all the possibly affected view models..
+	TArray<FNiagaraScriptViewModel*> ScriptViewModels;
+	ScriptViewModels.Add(this);
+	if (Parent)
 	{
-		AggregateStatus = FNiagaraEditorUtilities::UnionCompileStatus(AggregateStatus, CompileSystemStatuses[i]);
-		AggregateErrors += CompileSystemPaths[i] + TEXT(" ") + FNiagaraEditorUtilities::StatusToText(CompileSystemStatuses[i]).ToString() + TEXT("\n");
-		AggregateErrors += CompileSystemErrors[i] + TEXT("\n");
+		for (const TSharedRef<FNiagaraEmitterHandleViewModel>& EmitterHandleVM : Parent->GetEmitterHandleViewModels())
+		{
+			TSharedRef<FNiagaraEmitterViewModel> EmitterVM = EmitterHandleVM->GetEmitterViewModel();
+			ScriptViewModels.Add(&EmitterVM->GetSharedScriptViewModel().Get());
+		}
 	}
 
-	UpdateCompileStatus(AggregateStatus, AggregateErrors, CompileSystemStatuses, CompileSystemErrors, CompileSystemPaths, CompileSystemScripts);
+	// Now we need to match up the results with the view models that need to know about them..
+	for (FNiagaraScriptViewModel* ScriptVM : ScriptViewModels)
+	{
+		TArray<ENiagaraScriptCompileStatus> CompileChildStatuses;
+		TArray<FString> CompileChildErrors;
+		TArray<FString> CompileChildPaths;
+		TArray<UNiagaraScript*> CompileChildScripts;
 
+		for (int32 i = 0; i < CompileSystemScripts.Num(); i++)
+		{
+			UNiagaraScript* Script = CompileSystemScripts[i];
+			if (ScriptVM->GetScript(Script->GetUsage(), Script->GetUsageId()) == Script) // Check to see if it contains this script
+			{
+				CompileChildScripts.Add(Script);
+				CompileChildStatuses.Add(CompileSystemStatuses[i]);
+				CompileChildErrors.Add(CompileSystemErrors[i]);
+				CompileChildPaths.Add(CompileSystemPaths[i]);
+			}
+		}
+		ENiagaraScriptCompileStatus AggregateStatus = ENiagaraScriptCompileStatus::NCS_UpToDate;
+		FString AggregateErrors;
+
+		for (int32 i = 0; i < CompileChildStatuses.Num(); i++)
+		{
+			AggregateStatus = FNiagaraEditorUtilities::UnionCompileStatus(AggregateStatus, CompileChildStatuses[i]);
+			AggregateErrors += CompileChildPaths[i] + TEXT(" ") + FNiagaraEditorUtilities::StatusToText(CompileChildStatuses[i]).ToString() + TEXT("\n");
+			AggregateErrors += CompileChildErrors[i] + TEXT("\n");
+		}
+		
+		ScriptVM->UpdateCompileStatus(AggregateStatus, AggregateErrors, CompileChildStatuses, CompileChildErrors, CompileChildPaths, CompileChildScripts);
+	}
+	
 	if (OnSystemCompiledDelegate.IsBound())
 	{
 		OnSystemCompiledDelegate.Broadcast();
diff --git a/Engine/Plugins/FX/Niagara/Source/NiagaraEditor/Private/ViewModels/NiagaraSystemScriptViewModel.h b/Engine/Plugins/FX/Niagara/Source/NiagaraEditor/Private/ViewModels/NiagaraSystemScriptViewModel.h
--- a/Engine/Plugins/FX/Niagara/Source/NiagaraEditor/Private/ViewModels/NiagaraSystemScriptViewModel.h
+++ b/Engine/Plugins/FX/Niagara/Source/NiagaraEditor/Private/ViewModels/NiagaraSystemScriptViewModel.h
@@ -14,7 +14,7 @@ public:
 	DECLARE_MULTICAST_DELEGATE(FOnSystemCompiled)
 
 public:
-	FNiagaraSystemScriptViewModel(UNiagaraSystem& InSystem);
+	FNiagaraSystemScriptViewModel(UNiagaraSystem& InSystem, class FNiagaraSystemViewModel* InParent);
 
 	~FNiagaraSystemScriptViewModel();
 
@@ -23,12 +23,14 @@ public:
 
 	FOnSystemCompiled& OnSystemCompiled();
 
-	void CompileSystem();
+	void CompileSystem(bool bForce);
 
 private:
 	/** Called whenever the graph for the System script changes. */
 	void OnGraphChanged(const struct FEdGraphEditAction& InAction);
 
+	FNiagaraSystemViewModel* Parent;
+
 	/** The System who's script is getting viewed and edited by this view model. */
 	UNiagaraSystem& System;
 
diff --git a/Engine/Plugins/FX/Niagara/Source/NiagaraEditor/Private/ViewModels/NiagaraSystemViewModel.cpp b/Engine/Plugins/FX/Niagara/Source/NiagaraEditor/Private/ViewModels/NiagaraSystemViewModel.cpp
--- a/Engine/Plugins/FX/Niagara/Source/NiagaraEditor/Private/ViewModels/NiagaraSystemViewModel.cpp
+++ b/Engine/Plugins/FX/Niagara/Source/NiagaraEditor/Private/ViewModels/NiagaraSystemViewModel.cpp
@@ -32,27 +32,35 @@
 #include "NiagaraComponent.h"
 #include "MovieSceneFolder.h"
 
-DECLARE_CYCLE_STAT(TEXT("SystemViewModel - CompileSystem"), STAT_NiagaraEditor_SystemViewModel_CompileSystem, STATGROUP_NiagaraEditor);
+DECLARE_CYCLE_STAT(TEXT("Niagara - SystemViewModel - CompileSystem"), STAT_NiagaraEditor_SystemViewModel_CompileSystem, STATGROUP_NiagaraEditor);
 
 #define LOCTEXT_NAMESPACE "NiagaraSystemViewModel"
 
 template<> TMap<UNiagaraSystem*, TArray<FNiagaraSystemViewModel*>> TNiagaraViewModelManager<UNiagaraSystem, FNiagaraSystemViewModel>::ObjectsToViewModels{};
 
+FNiagaraSystemViewModelOptions::FNiagaraSystemViewModelOptions()
+	: bCanAutoCompile(true)
+	, bCanSimulate(true)
+{
+}
+
 FNiagaraSystemViewModel::FNiagaraSystemViewModel(UNiagaraSystem& InSystem, FNiagaraSystemViewModelOptions InOptions)
 	: System(InSystem)
 	, PreviewComponent(nullptr)
 	, SystemInstance(nullptr)
-	, SystemScriptViewModel(MakeShareable(new FNiagaraSystemScriptViewModel(System)))
+	, SystemScriptViewModel(MakeShareable(new FNiagaraSystemScriptViewModel(System, this)))
 	, NiagaraSequence(nullptr)
 	, bSettingSequencerTimeDirectly(false)
-	, bCanRemoveEmittersFromTimeline(InOptions.bCanRemoveEmittersFromTimeline)
-	, bCanRenameEmittersFromTimeline(InOptions.bCanRenameEmittersFromTimeline)
-	, bCanAddEmittersFromTimeline(InOptions.bCanAddEmittersFromTimeline)
+	, bCanModifyEmittersFromTimeline(InOptions.bCanModifyEmittersFromTimeline)
 	, bUseSystemExecStateForTimelineReset(InOptions.bUseSystemExecStateForTimelineReset)
+	, bCanAutoCompile(InOptions.bCanAutoCompile)
+	, bCanSimulate(InOptions.bCanSimulate)
+	, EditMode(InOptions.EditMode)
 	, OnGetSequencerAddMenuContent(InOptions.OnGetSequencerAddMenuContent)
 	, bUpdatingFromSequencerDataChange(false)
 	, bUpdatingSystemSelectionFromSequencer(false)
 	, bUpdatingSequencerSelectionFromSystem(false)
+	, bForceAutoCompileOnce(false)
 {
 	SetupPreviewComponentAndInstance();
 	SetupSequencer();
@@ -114,11 +122,11 @@ TSharedRef<FNiagaraSystemScriptViewModel> FNiagaraSystemViewModel::GetSystemScri
 	return SystemScriptViewModel;
 }
 
-void FNiagaraSystemViewModel::CompileSystem()
+void FNiagaraSystemViewModel::CompileSystem(bool bForce)
 {
 	SCOPE_CYCLE_COUNTER(STAT_NiagaraEditor_SystemViewModel_CompileSystem);
 	KillSystemInstances();
-	SystemScriptViewModel->CompileSystem();
+	SystemScriptViewModel->CompileSystem(bForce);
 	OnSystemCompiledDelegate.Broadcast();
 }
 
@@ -218,14 +226,15 @@ FNiagaraCurveOwner& FNiagaraSystemViewModel::GetCurveOwner()
 	return CurveOwner;
 }
 
-bool FNiagaraSystemViewModel::GetSystemIsTransient() const
+bool FNiagaraSystemViewModel::GetCanModifyEmittersFromTimeline() const
 {
-	return System.HasAnyFlags(RF_Transient);
+	return bCanModifyEmittersFromTimeline;
 }
 
-bool FNiagaraSystemViewModel::GetCanAddEmittersFromTimeline() const
+/** Gets the current editing mode for this system view model. */
+ENiagaraSystemViewModelEditMode FNiagaraSystemViewModel::GetEditMode() const
 {
-	return bCanAddEmittersFromTimeline;
+	return EditMode;
 }
 
 void FNiagaraSystemViewModel::AddEmitterFromAssetData(const FAssetData& AssetData)
@@ -251,7 +260,7 @@ void FNiagaraSystemViewModel::AddEmitter(UNiagaraEmitter& Emitter)
 	}
 
 	System.Modify();
-	FNiagaraEmitterHandle EmitterHandle = System.AddEmitterHandle(Emitter, FNiagaraEditorUtilities::GetUniqueName(Emitter.GetFName(), EmitterHandleNames));
+	FNiagaraEmitterHandle EmitterHandle = System.AddEmitterHandle(Emitter, FNiagaraUtilities::GetUniqueName(Emitter.GetFName(), EmitterHandleNames));
 	SystemScriptViewModel->RebuildEmitterNodes();
 
 	if (System.GetNumEmitters() == 1)
@@ -263,6 +272,8 @@ void FNiagaraSystemViewModel::AddEmitter(UNiagaraEmitter& Emitter)
 	RefreshAll();
 
 	SetSelectedEmitterHandleById(EmitterHandle.GetId());
+
+	bForceAutoCompileOnce = true;
 }
 
 void FNiagaraSystemViewModel::DuplicateEmitter(TSharedRef<FNiagaraEmitterHandleViewModel> EmitterHandleToDuplicate)
@@ -271,11 +282,24 @@ void FNiagaraSystemViewModel::DuplicateEmitter(TSharedRef<FNiagaraEmitterHandleV
 	{
 		return;
 	}
+	TSet<FGuid> HandlesToDuplicate;
+	HandlesToDuplicate.Add(EmitterHandleToDuplicate->GetEmitterHandle()->GetId());
+	DuplicateEmitters(HandlesToDuplicate);
+	bForceAutoCompileOnce = true;
+}
 
+void FNiagaraSystemViewModel::DuplicateEmitters(TSet<FGuid> EmitterHandleIdsToDuplicate)
+{
+	if (EmitterHandleIdsToDuplicate.Num() <= 0)
+	{
+		return;
+	}
 	// Kill all system instances before modifying the emitter handle list to prevent accessing deleted data.
 	KillSystemInstances();
+	const FScopedTransaction DeleteTransaction(EmitterHandleIdsToDuplicate.Num() == 1
+		? LOCTEXT("DuplicateEmitter", "Duplicate emitter")
+		: LOCTEXT("DuplicateEmitters", "Duplicate emitters"));
 
-	const FScopedTransaction DuplicateTransaction(LOCTEXT("DuplicateEmitter", "Duplicate emitter"));
 
 	TSet<FName> EmitterHandleNames;
 	for (const FNiagaraEmitterHandle& EmitterHandle : System.GetEmitterHandles())
@@ -284,10 +308,26 @@ void FNiagaraSystemViewModel::DuplicateEmitter(TSharedRef<FNiagaraEmitterHandleV
 	}
 
 	System.Modify();
-	FNiagaraEmitterHandle EmitterHandle = System.DuplicateEmitterHandle(*EmitterHandleToDuplicate->GetEmitterHandle(), FNiagaraEditorUtilities::GetUniqueName(EmitterHandleToDuplicate->GetEmitterHandle()->GetName(), EmitterHandleNames));
-	SystemScriptViewModel->RebuildEmitterNodes();
+	for (FGuid OriginalEmitterHandleId : EmitterHandleIdsToDuplicate)
+	{
+		FNiagaraEmitterHandle OriginalEmitterHandle;
+		for (const FNiagaraEmitterHandle& EmitterHandle : System.GetEmitterHandles())
+		{
+			if (EmitterHandle.GetId() == OriginalEmitterHandleId)
+			{
+				OriginalEmitterHandle = EmitterHandle;
+				break;
+			}
+		}
+		if (OriginalEmitterHandle.IsValid())
+		{
+			FNiagaraEmitterHandle EmitterHandle = System.DuplicateEmitterHandle(OriginalEmitterHandle, FNiagaraUtilities::GetUniqueName(OriginalEmitterHandle.GetName(), EmitterHandleNames));
+		}
+	}
 
+	SystemScriptViewModel->RebuildEmitterNodes();
 	RefreshAll();
+	bForceAutoCompileOnce = true;
 }
 
 void FNiagaraSystemViewModel::DeleteEmitter(TSharedRef<FNiagaraEmitterHandleViewModel> EmitterHandleToDelete)
@@ -295,6 +335,7 @@ void FNiagaraSystemViewModel::DeleteEmitter(TSharedRef<FNiagaraEmitterHandleView
 	TSet<FGuid> IdsToDelete;
 	IdsToDelete.Add(EmitterHandleToDelete->GetId());
 	DeleteEmitters(IdsToDelete);
+	bForceAutoCompileOnce = true;
 }
 
 void FNiagaraSystemViewModel::DeleteEmitters(TSet<FGuid> EmitterHandleIdsToDelete)
@@ -313,6 +354,7 @@ void FNiagaraSystemViewModel::DeleteEmitters(TSet<FGuid> EmitterHandleIdsToDelet
 		SystemScriptViewModel->RebuildEmitterNodes();
 
 		RefreshAll();
+		bForceAutoCompileOnce = true;
 	}
 }
 
@@ -378,15 +420,17 @@ void FNiagaraSystemViewModel::Tick(float DeltaTime)
 		}
 	}
 
-	if (bRecompile)
+	if (bRecompile || bForceAutoCompileOnce )
 	{
-		CompileSystem();
+		CompileSystem(false);
+
+		bForceAutoCompileOnce = false;
 	}
 }
 
 bool FNiagaraSystemViewModel::IsTickable() const
 {
-	return GetDefault<UNiagaraEditorSettings>()->bAutoCompile;
+	return bForceAutoCompileOnce || (GetDefault<UNiagaraEditorSettings>()->bAutoCompile && bCanAutoCompile);
 }
 
 TStatId FNiagaraSystemViewModel::GetStatId() const
@@ -396,20 +440,23 @@ TStatId FNiagaraSystemViewModel::GetStatId() const
 
 void FNiagaraSystemViewModel::SetupPreviewComponentAndInstance()
 {
-	PreviewComponent = NewObject<UNiagaraComponent>(GetTransientPackage(), NAME_None, RF_Transient);
-	PreviewComponent->CastShadow = 1;
-	PreviewComponent->bCastDynamicShadow = 1;
-	PreviewComponent->SetAsset(&System);
-	PreviewComponent->SetForceSolo(true);
-	PreviewComponent->SetAgeUpdateMode(UNiagaraComponent::EAgeUpdateMode::DesiredAge);
-	PreviewComponent->Activate(true);
+	if (bCanSimulate)
+	{
+		PreviewComponent = NewObject<UNiagaraComponent>(GetTransientPackage(), NAME_None, RF_Transient);
+		PreviewComponent->CastShadow = 1;
+		PreviewComponent->bCastDynamicShadow = 1;
+		PreviewComponent->SetAsset(&System);
+		PreviewComponent->SetForceSolo(true);
+		PreviewComponent->SetAgeUpdateMode(UNiagaraComponent::EAgeUpdateMode::DesiredAge);
+		PreviewComponent->Activate(true);
 
-	UNiagaraSystemEditorData& EditorData = GetOrCreateEditorData();
-	FTransform OwnerTransform = EditorData.GetOwnerTransform();
-	PreviewComponent->SetRelativeTransform(OwnerTransform);
+		UNiagaraSystemEditorData& EditorData = GetOrCreateEditorData();
+		FTransform OwnerTransform = EditorData.GetOwnerTransform();
+		PreviewComponent->SetRelativeTransform(OwnerTransform);
 
-	PreviewComponent->OnSystemInstanceChanged().AddRaw(this, &FNiagaraSystemViewModel::PreviewComponentSystemInstanceChanged);
-	PreviewComponentSystemInstanceChanged();
+		PreviewComponent->OnSystemInstanceChanged().AddRaw(this, &FNiagaraSystemViewModel::PreviewComponentSystemInstanceChanged);
+		PreviewComponentSystemInstanceChanged();
+	}
 }
 
 void FNiagaraSystemViewModel::RefreshAll()
@@ -436,6 +483,79 @@ void FNiagaraSystemViewModel::NotifyDataObjectChanged(UObject* ChangedObject)
 	ReInitializeSystemInstances();
 }
 
+bool FNiagaraSystemViewModel::CanExecuteToggleIsolation()
+{
+	return SelectedEmitterHandleIds.Num() > 0;
+}
+
+void FNiagaraSystemViewModel::IsolateSelectedEmitters()
+{
+	for (TSharedRef<FNiagaraEmitterHandleViewModel> EmitterHandle : EmitterHandleViewModels)
+	{
+		EmitterHandle->GetEmitterHandle()->SetIsolated(false);
+	}
+
+	bool bAnyEmitterIsolated = false;
+	TArray<TSharedRef<FNiagaraEmitterHandleViewModel>> SelectedEmitterHandleViewModels;
+	GetSelectedEmitterHandles(SelectedEmitterHandleViewModels);
+	for (TSharedRef<FNiagaraEmitterHandleViewModel> SelectedEmitterHandle : SelectedEmitterHandleViewModels)
+	{
+		bAnyEmitterIsolated = true;
+		SelectedEmitterHandle->GetEmitterHandle()->SetIsolated(true);
+	}
+
+	System.SetIsolateEnabled(bAnyEmitterIsolated);
+}
+
+void FNiagaraSystemViewModel::IsolateEmitterToggle()
+{
+	bool bAnyEmitterIsolated = false;
+	TArray<TSharedRef<FNiagaraEmitterHandleViewModel>> SelectedEmitterHandleViewModels;
+	GetSelectedEmitterHandles(SelectedEmitterHandleViewModels);
+	for (TSharedRef<FNiagaraEmitterHandleViewModel> SelectedEmitterHandle : SelectedEmitterHandleViewModels)
+	{
+		const bool bShouldIsolate = !SelectedEmitterHandle->GetEmitterHandle()->IsIsolated();
+		SelectedEmitterHandle->GetEmitterHandle()->SetIsolated(bShouldIsolate);
+		if (!bAnyEmitterIsolated && bShouldIsolate)
+		{
+			bAnyEmitterIsolated = true;
+		}
+	}
+
+	// If all the selected emitters are set to be not isolated we have to check if there is any emitter that is isolated.
+	if (!bAnyEmitterIsolated)
+	{
+		for (TSharedRef<FNiagaraEmitterHandleViewModel> EmitterHandle : EmitterHandleViewModels)
+		{
+			if (EmitterHandle->GetEmitterHandle()->IsIsolated())
+			{
+				bAnyEmitterIsolated = true;
+				break;
+			}
+		}
+	}
+
+	System.SetIsolateEnabled(bAnyEmitterIsolated);
+}
+
+bool FNiagaraSystemViewModel::IsEmitterIsolationActive()
+{
+	bool bShowActive = true;
+	TArray<TSharedRef<FNiagaraEmitterHandleViewModel>> SelectedEmitterHandleViewModels;
+	GetSelectedEmitterHandles(SelectedEmitterHandleViewModels);
+	for (TSharedRef<FNiagaraEmitterHandleViewModel> SelectedEmitterHandle : SelectedEmitterHandleViewModels)
+	{
+		if (!SelectedEmitterHandle->GetEmitterHandle()->IsIsolated())
+		{
+			bShowActive = false;
+			break;
+		}
+	}
+
+	return bShowActive;
+}
+
+
 void FNiagaraSystemViewModel::RefreshEmitterHandleViewModels()
 {
 	TArray<TSharedRef<FNiagaraEmitterHandleViewModel>> OldViewModels = EmitterHandleViewModels;
@@ -456,6 +576,7 @@ void FNiagaraSystemViewModel::RefreshEmitterHandleViewModels()
 			TSharedRef<FNiagaraEmitterHandleViewModel> ViewModel = MakeShareable(new FNiagaraEmitterHandleViewModel(EmitterHandle, Simulation, System));
 			// Since we're adding fresh, we need to register all the event handlers.
 			ViewModel->OnPropertyChanged().AddRaw(this, &FNiagaraSystemViewModel::EmitterHandlePropertyChanged, ViewModel);
+			ViewModel->OnNameChanged().AddRaw(this, &FNiagaraSystemViewModel::EmitterHandleNameChanged, ViewModel);
 			ViewModel->GetEmitterViewModel()->OnPropertyChanged().AddRaw(this, &FNiagaraSystemViewModel::EmitterPropertyChanged, ViewModel);
 			ViewModel->GetEmitterViewModel()->OnScriptCompiled().AddRaw(this, &FNiagaraSystemViewModel::ScriptCompiled);
 			EmitterHandleViewModels.Add(ViewModel);
@@ -686,11 +807,9 @@ void FNiagaraSystemViewModel::SetupSequencer()
 
 void FNiagaraSystemViewModel::ResetSystem()
 {
- 	if (Sequencer->GetPlaybackStatus() == EMovieScenePlayerStatus::Playing)
-	{
-		Sequencer->SetGlobalTime(0.0f);
-	}
-
+	Sequencer->SetGlobalTime(0.0f);
+	Sequencer->SetPlaybackStatus(EMovieScenePlayerStatus::Playing);
+	
 	for (TObjectIterator<UNiagaraComponent> ComponentIt; ComponentIt; ++ComponentIt)
 	{
 		UNiagaraComponent* Component = *ComponentIt;
@@ -753,7 +872,7 @@ void GetCurveData(FString CurveSource, UNiagaraGraph* SourceGraph, TArray<FNiaga
 		{
 			if (InputNode->Usage == ENiagaraInputNodeUsage::Parameter)
 			{
-				UNiagaraDataInterfaceCurveBase* CurveDataInterface = Cast<UNiagaraDataInterfaceCurveBase>(InputNode->DataInterface);
+				UNiagaraDataInterfaceCurveBase* CurveDataInterface = Cast<UNiagaraDataInterfaceCurveBase>(InputNode->GetDataInterface());
 				if (CurveDataInterface != nullptr)
 				{
 					TArray<UNiagaraDataInterfaceCurveBase::FCurveData> CurveData;
@@ -832,6 +951,11 @@ void FNiagaraSystemViewModel::EmitterHandlePropertyChanged(TSharedRef<FNiagaraEm
 	ReInitializeSystemInstances();
 }
 
+void FNiagaraSystemViewModel::EmitterHandleNameChanged(TSharedRef<FNiagaraEmitterHandleViewModel> EmitterHandleViewModel)
+{
+	CompileSystem(false);
+}
+
 void FNiagaraSystemViewModel::EmitterPropertyChanged(TSharedRef<FNiagaraEmitterHandleViewModel> EmitterHandleViewModel)
 {
 	if (bUpdatingFromSequencerDataChange == false)
@@ -921,6 +1045,7 @@ void FNiagaraSystemViewModel::SequencerDataChanged(EMovieSceneDataChangeType Dat
 {
 	bUpdatingFromSequencerDataChange = true;
 	TSet<FGuid> VaildTrackEmitterHandleIds;
+	TSet<FGuid> EmittersToDuplicate;
 	for (UMovieSceneTrack* Track : NiagaraSequence->GetMovieScene()->GetMasterTracks())
 	{
 		UMovieSceneNiagaraEmitterTrack* EmitterTrack = CastChecked<UMovieSceneNiagaraEmitterTrack>(Track);
@@ -928,17 +1053,25 @@ void FNiagaraSystemViewModel::SequencerDataChanged(EMovieSceneDataChangeType Dat
 		UMovieSceneNiagaraEmitterSection* EmitterSection = Cast<UMovieSceneNiagaraEmitterSection>(Sections.Num() == 1 ? Sections[0] : nullptr);
 		if (EmitterSection != nullptr)
 		{
+			if (!EmitterTrack->GetEmitterHandle().IsValid())
+			{
+				if (EmitterTrack->GetEmitterHandleId().IsValid())
+				{
+					// The emitter handle is invalid, but the track has a valid Id, most probably because of a copy/paste event
+					EmittersToDuplicate.Add(EmitterTrack->GetEmitterHandleId());
+				}
+				continue;
+			}
 			VaildTrackEmitterHandleIds.Add(EmitterTrack->GetEmitterHandle()->GetId());
 
 			TSharedPtr<FNiagaraEmitterHandleViewModel> EmitterHandleViewModel = EmitterTrack->GetEmitterHandle();
-			if (bCanRenameEmittersFromTimeline)
+			if (bCanModifyEmittersFromTimeline)
 			{
 				EmitterHandleViewModel->SetName(*EmitterTrack->GetDisplayName().ToString());
 			}
-			else
-			{
-				EmitterTrack->SetDisplayName(EmitterHandleViewModel->GetNameText());
-			}
+			
+			EmitterTrack->SetDisplayName(EmitterHandleViewModel->GetNameText());
+			
 			EmitterHandleViewModel->SetIsEnabled(EmitterSection->IsActive());
 
 			TSharedPtr<FNiagaraEmitterViewModel> EmitterViewModel = EmitterTrack->GetEmitterHandle()->GetEmitterViewModel();
@@ -988,7 +1121,7 @@ void FNiagaraSystemViewModel::SequencerDataChanged(EMovieSceneDataChangeType Dat
 	TSet<FGuid> RemovedEmitterHandleIds = AllEmitterHandleIds.Difference(VaildTrackEmitterHandleIds);
 	if (RemovedEmitterHandleIds.Num() > 0)
 	{
-		if (bCanRemoveEmittersFromTimeline)
+		if (bCanModifyEmittersFromTimeline)
 		{
 			DeleteEmitters(RemovedEmitterHandleIds);
 			// When deleting emitters from sequencer, select a new one if one is available.
@@ -1002,6 +1135,18 @@ void FNiagaraSystemViewModel::SequencerDataChanged(EMovieSceneDataChangeType Dat
 			RefreshSequencerTracks();
 		}
 	}
+	
+	if (EmittersToDuplicate.Num() > 0)
+	{
+		if (bCanModifyEmittersFromTimeline)
+		{
+			DuplicateEmitters(EmittersToDuplicate);
+		}
+		else
+		{
+			RefreshSequencerTracks();
+		}
+	}
 
 	TArray<UMovieSceneTrack*> RootTracks;
 	TArray<UMovieSceneFolder*> RootFolders = NiagaraSequence->GetMovieScene()->GetRootFolders();
@@ -1012,7 +1157,7 @@ void FNiagaraSystemViewModel::SequencerDataChanged(EMovieSceneDataChangeType Dat
 
 	if (SystemInstance != nullptr)
 	{
-		SystemInstance->Reset(FNiagaraSystemInstance::EResetMode::DeferredReset);
+		SystemInstance->Reset(FNiagaraSystemInstance::EResetMode::ResetAll);
 	}
 	bUpdatingFromSequencerDataChange = false;
 }
@@ -1032,7 +1177,6 @@ void FNiagaraSystemViewModel::SequencerTimeChanged()
 		{
 			bool bUpdateDesiredAge = false;
 			bool bResetSystemInstance = false;
-			bool bSetEnabled = false;
 
 			if (CurrentStatus == EMovieScenePlayerStatus::Playing)
 			{
@@ -1041,10 +1185,9 @@ void FNiagaraSystemViewModel::SequencerTimeChanged()
 
 				if (bUseSystemExecStateForTimelineReset)
 				{
-					if (SystemInstance->GetExecutionState() == ENiagaraExecutionState::Disabled)
+					if (SystemInstance->IsComplete())
 					{
 						bSystemIsAlive = false;
-						bSetEnabled = true;
 					}
 					else
 					{
@@ -1055,7 +1198,7 @@ void FNiagaraSystemViewModel::SequencerTimeChanged()
 				{
 					for (const TSharedRef<FNiagaraEmitterInstance> Simulation : SystemInstance->GetEmitters())
 					{
-						if (Simulation->GetExecutionState() != ENiagaraExecutionState::Dead)
+						if (!Simulation->IsComplete())
 						{
 							bSystemIsAlive |= true;
 						}
@@ -1097,23 +1240,26 @@ void FNiagaraSystemViewModel::SequencerTimeChanged()
 				bUpdateDesiredAge = !bStoppedPlaying;
 			}
 
+			if (CurrentSequencerTime < PreviousSequencerTime)
+			{
+				//Always reset if we've gone backwards.
+				bResetSystemInstance = true;
+			}
+
 			if (bUpdateDesiredAge)
 			{
 				PreviewComponent->SetDesiredAge(FMath::Max(CurrentSequencerTime, 0.0f));
-
+				if (CurrentSequencerTime < PreviewComponent->GetSystemInstance()->GetAge())
+				{
+					bResetSystemInstance = true;
+				}
 			//SystemInstance->Tick(0);TODO: What was this for? Possibly reinstate but I don't like use dealing with the instance directly.
 			}
 
-			if (bResetSystemInstance)
+			if (bResetSystemInstance || !PreviewComponent->IsActive())
 			{
-				if (bSetEnabled)
-				{
-					SystemInstance->Enable();
-				}
-				else
-				{
-					SystemInstance->Reset(FNiagaraSystemInstance::EResetMode::ImmediateReset);
-				}
+				PreviewComponent->Activate(true);
+
 				TGuardValue<bool>(bSettingSequencerTimeDirectly, true);
 				Sequencer->SetLocalTime(0);
 			}
@@ -1189,6 +1335,11 @@ void FNiagaraSystemViewModel::UpdateSequencerFromEmitterHandleSelection()
 	}
 }
 
+void FNiagaraSystemViewModel::SystemInstanceReset()
+{
+	SystemInstanceInitialized();
+}
+
 void FNiagaraSystemViewModel::PreviewComponentSystemInstanceChanged()
 {
 	FNiagaraSystemInstance* OldSystemInstance = SystemInstance;
@@ -1198,6 +1349,7 @@ void FNiagaraSystemViewModel::PreviewComponentSystemInstanceChanged()
 		if (SystemInstance != nullptr)
 		{
 			SystemInstance->OnInitialized().AddRaw(this, &FNiagaraSystemViewModel::SystemInstanceInitialized);
+			SystemInstance->OnReset().AddRaw(this, &FNiagaraSystemViewModel::SystemInstanceReset);
 		}
 		else
 		{
@@ -1223,10 +1375,28 @@ void FNiagaraSystemViewModel::SystemInstanceInitialized()
 	}
 }
 
-void FNiagaraSystemViewModel::ResynchronizeAllHandles()
+void FNiagaraSystemViewModel::UpdateEmitterFixedBounds()
 {
-	System.ResynchronizeAllHandles();
-	RefreshAll();
+	TArray<TSharedRef<FNiagaraEmitterHandleViewModel>> SelectedEmitterHandles;
+	GetSelectedEmitterHandles(SelectedEmitterHandles);
+
+	for (TSharedRef<FNiagaraEmitterHandleViewModel>& SelectedEmitterHandleViewModel : SelectedEmitterHandles)
+	{
+		FNiagaraEmitterHandle* SelectedEmitterHandle = SelectedEmitterHandleViewModel->GetEmitterHandle();
+		check(SelectedEmitterHandle);
+		UNiagaraEmitter* Emitter = SelectedEmitterHandle->GetInstance();
+		for (TSharedRef<FNiagaraEmitterInstance>& EmitterInst : PreviewComponent->GetSystemInstance()->GetEmitters())
+		{
+			if (&EmitterInst->GetEmitterHandle() == SelectedEmitterHandle && !EmitterInst->IsComplete())
+			{
+				FBox EmitterBounds = EmitterInst->CalculateDynamicBounds();
+				Emitter->Modify();
+				Emitter->bFixedBounds = true;
+				//Dynamic bounds are in world space. Transform back to local.
+				Emitter->FixedBounds = EmitterBounds.TransformBy(PreviewComponent->GetComponentToWorld().Inverse());
+			}
+		}
+	}
 }
 
 #undef LOCTEXT_NAMESPACE // NiagaraSystemViewModel
diff --git a/Engine/Plugins/FX/Niagara/Source/NiagaraEditor/Private/ViewModels/NiagaraSystemViewModel.h b/Engine/Plugins/FX/Niagara/Source/NiagaraEditor/Private/ViewModels/NiagaraSystemViewModel.h
--- a/Engine/Plugins/FX/Niagara/Source/NiagaraEditor/Private/ViewModels/NiagaraSystemViewModel.h
+++ b/Engine/Plugins/FX/Niagara/Source/NiagaraEditor/Private/ViewModels/NiagaraSystemViewModel.h
@@ -18,6 +18,7 @@ class UNiagaraComponent;
 class UNiagaraSequence;
 class UMovieSceneNiagaraEmitterTrack;
 struct FNiagaraVariable;
+struct FNiagaraEmitterHandle;
 class FNiagaraEmitterHandleViewModel;
 class FNiagaraSystemScriptViewModel;
 class FNiagaraSystemInstance;
@@ -26,23 +27,37 @@ struct FAssetData;
 class UNiagaraSystemEditorData;
 struct FRichCurve;
 
+/** Defines different editing modes for this system view model. */
+enum class ENiagaraSystemViewModelEditMode
+{
+	/** A system asset is being edited.  This assumes that emitters should be inheriting from a base version and that emitter editing will be restricted. */
+	SystemAsset,
+	/** An emitter asset is being edited.  In this mode the system scripts will not be editable and all emitter editing options are available. */
+	EmitterAsset,
+};
+
 /** Defines options for the niagara System view model */
 struct FNiagaraSystemViewModelOptions
 {
-	/** Whether or not the user can remove emitters from the timeline. */
-	bool bCanRemoveEmittersFromTimeline;
-
-	/** Whether or not the user can rename emitters from the timeline. */
-	bool bCanRenameEmittersFromTimeline;
+	FNiagaraSystemViewModelOptions();
 
-	/** Whether or not the user can add emitters from the timeline. */
-	bool bCanAddEmittersFromTimeline;
+	/** Whether or not the user can edit emitters from the timeline. */
+	bool bCanModifyEmittersFromTimeline;
 
 	/** A delegate which is used to generate the content for the add menu in sequencer. */
 	FOnGetAddMenuContent OnGetSequencerAddMenuContent;
 
 	/** Whether or not we use the system's execution state to drive when we reset the timeline*/
 	bool bUseSystemExecStateForTimelineReset;
+
+	/** Whether or not the system represented by this view model can be automatically compiled.  True by default. */
+	bool bCanAutoCompile;
+
+	/** Whether or not the system represented by this view model can be simulated. True by default. */
+	bool bCanSimulate;
+
+	/** Gets the current editing mode for this system. */
+	ENiagaraSystemViewModelEditMode EditMode;
 };
 
 /** A view model for viewing and editing a UNiagaraSystem. */
@@ -88,11 +103,11 @@ public:
 	/** Get access to the underlying system*/
 	UNiagaraSystem& GetSystem() { return System; }
 
-	/** Gets whether or not this system is transient.  This will be true for the system view model in the emitter editor. */
-	bool GetSystemIsTransient() const;
-
 	/** Gets whether or not emitters can be added from the timeline. */
-	bool GetCanAddEmittersFromTimeline() const;
+	bool GetCanModifyEmittersFromTimeline() const;
+
+	/** Gets the current editing mode for this system view model. */
+	ENiagaraSystemViewModelEditMode GetEditMode() const;
 
 	/** Adds a new emitter to the System from an emitter asset data. */
 	void AddEmitterFromAssetData(const FAssetData& AssetData);
@@ -136,13 +151,14 @@ public:
 	virtual bool IsTickable() const override;
 	virtual TStatId GetStatId() const override;
 
-	void ResynchronizeAllHandles();
-
 	/** Resets the System instance to initial conditions. */
 	void ResetSystem();
 
+	/** Reinitializes the System instance to initial conditions - rebuilds all data sets and resets data interfaces. */
+	void ReInitializeSystemInstances();
+
 	/** Compiles the spawn and update scripts. */
-	void CompileSystem();
+	void CompileSystem(bool bForce);
 
 	/* Get the latest status of this view-model's script compilation.*/
 	ENiagaraScriptCompileStatus GetLatestCompileStatus();
@@ -173,6 +189,22 @@ public:
 
 	/** Called to notify the system view model that one of the data objects in the system was modified. */
 	void NotifyDataObjectChanged(UObject* ChangedObject);
+
+	/** Updates all selected emitter's fixed bounds with their current dynamic bounds. */
+	void UpdateEmitterFixedBounds();
+
+	/** Whether it is possible to isolate an emitter. */
+	bool CanExecuteToggleIsolation();
+
+	/** Isolates the selected emitters and sets all the others to not be isolated. */
+	void IsolateSelectedEmitters();
+
+	/** Toggle the isolate state of the selected emitters. */
+	void IsolateEmitterToggle();
+
+	/** Whether all the selected emitters are isolated. */
+	bool IsEmitterIsolationActive();
+
 private:
 
 	/** Sets up the preview component and System instance. */
@@ -199,9 +231,6 @@ private:
 	/** Kills all system instances using the system being displayed by this view model. */
 	void KillSystemInstances();
 
-	/** Reinitializes the System instance to initial conditions - rebuilds all data sets and resets data interfaces. */
-	void ReInitializeSystemInstances();
-
 	/** Resets and rebuilds the data in the curve owner. */
 	void ResetCurveData();
 
@@ -211,6 +240,9 @@ private:
 	/** Called whenever a property on the emitter handle changes. */
 	void EmitterHandlePropertyChanged(TSharedRef<FNiagaraEmitterHandleViewModel> EmitterHandleViewModel);
 
+	/** Called whenever the name on an emitter handle changes. */
+	void EmitterHandleNameChanged(TSharedRef<FNiagaraEmitterHandleViewModel> EmitterHandleViewModel);
+
 	/** Called whenever a property on the emitter changes. */
 	void EmitterPropertyChanged(TSharedRef<FNiagaraEmitterHandleViewModel> EmitterHandleViewModel);
 
@@ -244,6 +276,12 @@ private:
 	/** Called whenever the System instance is initialized. */
 	void SystemInstanceInitialized();
 
+	/** Called whenever the System instance is reset.*/
+	void SystemInstanceReset();
+
+	/** Duplicates a set of emitters and refreshes everything.*/
+	void DuplicateEmitters(TSet<FGuid> EmitterHandleIdsToDuplicate);
+
 private:
 	/** The System being viewed and edited by this view model. */
 	UNiagaraSystem& System;
@@ -275,18 +313,24 @@ private:
 	/** The previous time for the sequencer timeline. */
 	float PreviousSequencerTime;
 
-	/** Whether or not the user can remove emitters from the timeline. */
-	bool bCanRemoveEmittersFromTimeline;
-
-	/** Whether or not the user can rename emitters from the timeline. */
-	bool bCanRenameEmittersFromTimeline;
-
-	/** Whether or not the user can add emitters from the timeline. */
-	bool bCanAddEmittersFromTimeline;
+	/** Whether or not the user can edit emitters from the timeline. */
+	bool bCanModifyEmittersFromTimeline;
 
 	/** Whether or not we use the system's execution state to drive when we reset the timeline*/
 	bool bUseSystemExecStateForTimelineReset;
 
+	/** Whether or not the system represented by this view model can be automatically compiled. */
+	bool bCanAutoCompile;
+
+	/** Whether or not the system requires a compilation*/
+	bool bForceAutoCompileOnce;
+
+	/** Whether or not the system represented by this view model can be simulated. */
+	bool bCanSimulate;
+
+	/** The current editing mode for this view model. */
+	ENiagaraSystemViewModelEditMode EditMode;
+
 	/** A delegate which is used to generate the content for the add menu in sequencer. */
 	FOnGetAddMenuContent OnGetSequencerAddMenuContent;
 
diff --git a/Engine/Plugins/FX/Niagara/Source/NiagaraEditor/Private/ViewModels/Stack/NiagaraStackAddEventScriptItem.cpp b/Engine/Plugins/FX/Niagara/Source/NiagaraEditor/Private/ViewModels/Stack/NiagaraStackAddEventScriptItem.cpp
--- a/Engine/Plugins/FX/Niagara/Source/NiagaraEditor/Private/ViewModels/Stack/NiagaraStackAddEventScriptItem.cpp
+++ b/Engine/Plugins/FX/Niagara/Source/NiagaraEditor/Private/ViewModels/Stack/NiagaraStackAddEventScriptItem.cpp
@@ -41,29 +41,21 @@ UNiagaraNodeOutput* UNiagaraStackAddEventScriptItem::GetOrCreateOutputNode()
 	// The stack should not have been created if any of these are null, so bail out if it happens somehow rather than try to handle all of these cases.
 	checkf(Emitter != nullptr && Source != nullptr && Graph != nullptr, TEXT("Stack created for invalid emitter or graph."));
 
-	// Check that there isn't already an event output at the new index before making any changes so that we can early out and don't leave the change
-	// half completed.
-	if(ensureMsgf(Graph->FindOutputNode(ENiagaraScriptUsage::ParticleEventScript, Emitter->EventHandlerScriptProps.Num()) == nullptr,
-		TEXT("Invalid Stack Graph - While creating a new event handler an event script output node already existed with the new index.")))
-	{
-		Emitter->Modify();
-		FNiagaraEventScriptProperties EventScriptProperties;
-		EventScriptProperties.Script = NewObject<UNiagaraScript>(Emitter, MakeUniqueObjectName(Emitter, UNiagaraScript::StaticClass(), "EventScript"), EObjectFlags::RF_Transactional);
-		EventScriptProperties.Script->SetUsage(ENiagaraScriptUsage::ParticleEventScript);
-		EventScriptProperties.Script->SetSource(Source);
+	Emitter->Modify();
+	FNiagaraEventScriptProperties EventScriptProperties;
+	EventScriptProperties.Script = NewObject<UNiagaraScript>(Emitter, MakeUniqueObjectName(Emitter, UNiagaraScript::StaticClass(), "EventScript"), EObjectFlags::RF_Transactional);
+	EventScriptProperties.Script->SetUsage(ENiagaraScriptUsage::ParticleEventScript);
+	EventScriptProperties.Script->SetUsageId(FGuid::NewGuid());
+	EventScriptProperties.Script->SetSource(Source);
+	Emitter->AddEventHandler(EventScriptProperties);
 
-		int32 Index = Emitter->EventHandlerScriptProps.Add(EventScriptProperties);
-		Emitter->EventHandlerScriptProps[Index].Script->SetUsageIndex(Index); // synchronize usage index with entry in array
+	UNiagaraNodeOutput* OutputNode = FNiagaraStackGraphUtilities::ResetGraphForOutput(*Graph, ENiagaraScriptUsage::ParticleEventScript, EventScriptProperties.Script->GetUsageId());
 
-		UNiagaraNodeOutput* OutputNode = FNiagaraStackGraphUtilities::ResetGraphForOutput(*Graph, ENiagaraScriptUsage::ParticleEventScript, Index);
+	// Set the emitter here so that the internal state of the view model is updated.
+	// TODO: Move the logic for managing event handlers into the emitter view model or script view model.
+	GetEmitterViewModel()->GetSharedScriptViewModel()->SetScripts(Emitter);
 
-		// Set the emitter here so that the internal state of the view model is updated.
-		// TODO: Move the logic for managing event handlers into the emitter view model or script view model.
-		GetEmitterViewModel()->GetSharedScriptViewModel()->SetScripts(Emitter);
-
-		return OutputNode;
-	}
-	return nullptr;
+	return OutputNode;
 }
 
 FText UNiagaraStackAddEventScriptItem::GetInsertTransactionText() const
diff --git a/Engine/Plugins/FX/Niagara/Source/NiagaraEditor/Private/ViewModels/Stack/NiagaraStackAddModuleItem.cpp b/Engine/Plugins/FX/Niagara/Source/NiagaraEditor/Private/ViewModels/Stack/NiagaraStackAddModuleItem.cpp
--- a/Engine/Plugins/FX/Niagara/Source/NiagaraEditor/Private/ViewModels/Stack/NiagaraStackAddModuleItem.cpp
+++ b/Engine/Plugins/FX/Niagara/Source/NiagaraEditor/Private/ViewModels/Stack/NiagaraStackAddModuleItem.cpp
@@ -5,6 +5,7 @@
 #include "NiagaraScript.h"
 #include "NiagaraNodeFunctionCall.h"
 #include "NiagaraNodeAssignment.h"
+#include "NiagaraMaterialParameterNode.h"
 #include "ViewModels/NiagaraEmitterViewModel.h"
 #include "NiagaraStackEditorData.h"
 #include "NiagaraStackGraphUtilities.h"
@@ -25,6 +26,11 @@ FText UNiagaraStackAddModuleItem::GetDisplayName() const
 	return FText();
 }
 
+UNiagaraStackAddModuleItem::EDisplayMode UNiagaraStackAddModuleItem::GetDisplayMode() const
+{
+	return EDisplayMode::Standard;
+}
+
 void UNiagaraStackAddModuleItem::GetAvailableParameters(TArray<FNiagaraVariable>& OutAvailableParameterVariables) const
 {
 }
@@ -38,23 +44,14 @@ TOptional<FName> UNiagaraStackAddModuleItem::GetNewParameterNamespace() const
 	return TOptional<FName>();
 }
 
-void UNiagaraStackAddModuleItem::SetOnItemAdded(FOnItemAdded InOnItemAdded)
+int32 UNiagaraStackAddModuleItem::GetTargetIndex() const
 {
-	ItemAddedDelegate = InOnItemAdded;
+	return INDEX_NONE;
 }
 
-void ConnectModuleNode(UNiagaraNode& ModuleNode, UNiagaraNode& OutputNode)
+void UNiagaraStackAddModuleItem::SetOnItemAdded(FOnItemAdded InOnItemAdded)
 {
-	TArray<FNiagaraStackGraphUtilities::FStackNodeGroup> StackNodeGroups;
-	FNiagaraStackGraphUtilities::GetStackNodeGroups(OutputNode, StackNodeGroups);
-
-	FNiagaraStackGraphUtilities::FStackNodeGroup ModuleGroup;
-	ModuleGroup.StartNodes.Add(&ModuleNode);
-	ModuleGroup.EndNode = &ModuleNode;
-
-	FNiagaraStackGraphUtilities::FStackNodeGroup& OutputGroup = StackNodeGroups[StackNodeGroups.Num() - 1];
-	FNiagaraStackGraphUtilities::FStackNodeGroup& OutputGroupPrevious = StackNodeGroups[StackNodeGroups.Num() - 2];
-	FNiagaraStackGraphUtilities::ConnectStackNodeGroup(ModuleGroup, OutputGroupPrevious, OutputGroup);
+	ItemAddedDelegate = InOnItemAdded;
 }
 
 void UNiagaraStackAddModuleItem::AddScriptModule(FAssetData ModuleScriptAsset)
@@ -64,15 +61,7 @@ void UNiagaraStackAddModuleItem::AddScriptModule(FAssetData ModuleScriptAsset)
 	UNiagaraNodeOutput* OutputNode = GetOrCreateOutputNode();
 	if (OutputNode != nullptr)
 	{
-		UEdGraph* Graph = OutputNode->GetGraph();
-		Graph->Modify();
-
-		FGraphNodeCreator<UNiagaraNodeFunctionCall> ModuleNodeCreator(*Graph);
-		UNiagaraNodeFunctionCall* NewModuleNode = ModuleNodeCreator.CreateNode();
-		NewModuleNode->FunctionScriptAssetObjectPath = ModuleScriptAsset.ObjectPath;
-		ModuleNodeCreator.Finalize();
-
-		ConnectModuleNode(*NewModuleNode, *OutputNode);
+		UNiagaraNodeFunctionCall* NewModuleNode = FNiagaraStackGraphUtilities::AddScriptModuleToStack(ModuleScriptAsset, *OutputNode, GetTargetIndex());
 		FNiagaraStackGraphUtilities::InitializeDataInterfaceInputs(GetSystemViewModel(), GetEmitterViewModel(), *StackEditorData, *NewModuleNode, *NewModuleNode);
 		FNiagaraStackGraphUtilities::RelayoutGraph(*OutputNode->GetGraph());
 
@@ -87,15 +76,7 @@ void UNiagaraStackAddModuleItem::AddParameterModule(FNiagaraVariable ParameterVa
 	UNiagaraNodeOutput* OutputNode = GetOrCreateOutputNode();
 	if (OutputNode != nullptr)
 	{
-		UEdGraph* Graph = OutputNode->GetGraph();
-		Graph->Modify();
-
-		FGraphNodeCreator<UNiagaraNodeAssignment> AssignmentNodeCreator(*Graph);
-		UNiagaraNodeAssignment* NewAssignmentNode = AssignmentNodeCreator.CreateNode();
-		NewAssignmentNode->AssignmentTarget = ParameterVariable;
-		AssignmentNodeCreator.Finalize();
-
-		ConnectModuleNode(*NewAssignmentNode, *OutputNode);
+		UNiagaraNodeAssignment* NewAssignmentNode = FNiagaraStackGraphUtilities::AddParameterModuleToStack(ParameterVariable, *OutputNode, GetTargetIndex());
 		FNiagaraStackGraphUtilities::InitializeDataInterfaceInputs(GetSystemViewModel(), GetEmitterViewModel(), *StackEditorData, *NewAssignmentNode, *NewAssignmentNode);
 		FNiagaraStackGraphUtilities::RelayoutGraph(*OutputNode->GetGraph());
 
@@ -116,6 +97,31 @@ void UNiagaraStackAddModuleItem::AddParameterModule(FNiagaraVariable ParameterVa
 	}
 }
 
+void UNiagaraStackAddModuleItem::AddMaterialParameterModule()
+{
+	FScopedTransaction ScopedTransaction(GetInsertTransactionText());
+
+	UNiagaraNodeOutput* OutputNode = GetOrCreateOutputNode();
+	if (OutputNode != nullptr)
+	{
+		TSharedRef<FNiagaraEmitterViewModel> EmitterViewModelRef = GetEmitterViewModel();
+		UNiagaraMaterialParameterNode* NewModuleNode = FNiagaraStackGraphUtilities::AddMaterialParameterModuleToStack(EmitterViewModelRef, *OutputNode);
+		FNiagaraStackGraphUtilities::InitializeDataInterfaceInputs(GetSystemViewModel(), EmitterViewModelRef, *StackEditorData, *NewModuleNode, *NewModuleNode);
+		FNiagaraStackGraphUtilities::RelayoutGraph(*OutputNode->GetGraph());
+		
+		TArray<const UEdGraphPin*> InputPins;
+		FNiagaraStackGraphUtilities::GetStackFunctionInputPins(*NewModuleNode, InputPins);
+		for (int32 Index = 1; Index < InputPins.Num() - 1; Index += 2)
+		{
+			FString FunctionInputEditorDataKey = FNiagaraStackGraphUtilities::GenerateStackFunctionInputEditorDataKey(*NewModuleNode, InputPins[Index]->PinName);
+			StackEditorData->SetModuleInputIsPinned(FunctionInputEditorDataKey, true);
+		}
+
+		StackEditorData->SetStackEntryIsExpanded(FNiagaraStackGraphUtilities::GenerateStackModuleEditorDataKey(*NewModuleNode), false);
+		ItemAddedDelegate.ExecuteIfBound();
+	}
+}
+
 FText UNiagaraStackAddModuleItem::GetInsertTransactionText() const
 {
 	return LOCTEXT("InsertNewModule", "Insert new module");
diff --git a/Engine/Plugins/FX/Niagara/Source/NiagaraEditor/Private/ViewModels/Stack/NiagaraStackAddRendererItem.cpp b/Engine/Plugins/FX/Niagara/Source/NiagaraEditor/Private/ViewModels/Stack/NiagaraStackAddRendererItem.cpp
--- a/Engine/Plugins/FX/Niagara/Source/NiagaraEditor/Private/ViewModels/Stack/NiagaraStackAddRendererItem.cpp
+++ b/Engine/Plugins/FX/Niagara/Source/NiagaraEditor/Private/ViewModels/Stack/NiagaraStackAddRendererItem.cpp
@@ -33,7 +33,7 @@ void UNiagaraStackAddRendererItem::AddRenderer(UClass* InRendererClass)
 	UNiagaraEmitter* Emitter = GetEmitterViewModel()->GetEmitter();
 	Emitter->Modify();
 	UNiagaraRendererProperties* RendererProperties = NewObject<UNiagaraRendererProperties>(Emitter, InRendererClass, NAME_None, RF_Transactional);
-	Emitter->RendererProperties.Add(RendererProperties);
+	Emitter->AddRenderer(RendererProperties);
 
 	bool bVarsAdded = false;
 	TArray<FNiagaraVariable> MissingAttributes = UNiagaraStackRendererItem::GetMissingVariables(RendererProperties, Emitter);
diff --git a/Engine/Plugins/FX/Niagara/Source/NiagaraEditor/Private/ViewModels/Stack/NiagaraStackAddScriptModuleItem.cpp b/Engine/Plugins/FX/Niagara/Source/NiagaraEditor/Private/ViewModels/Stack/NiagaraStackAddScriptModuleItem.cpp
--- a/Engine/Plugins/FX/Niagara/Source/NiagaraEditor/Private/ViewModels/Stack/NiagaraStackAddScriptModuleItem.cpp
+++ b/Engine/Plugins/FX/Niagara/Source/NiagaraEditor/Private/ViewModels/Stack/NiagaraStackAddScriptModuleItem.cpp
@@ -12,10 +12,20 @@
 
 #define LOCTEXT_NAMESPACE "NiagaraStackViewModel"
 
-void UNiagaraStackAddScriptModuleItem::Initialize(TSharedRef<FNiagaraSystemViewModel> InSystemViewModel, TSharedRef<FNiagaraEmitterViewModel> InEmitterViewModel, UNiagaraStackEditorData& InStackEditorData, UNiagaraNodeOutput& InOutputNode)
+void UNiagaraStackAddScriptModuleItem::Initialize(TSharedRef<FNiagaraSystemViewModel> InSystemViewModel, TSharedRef<FNiagaraEmitterViewModel> InEmitterViewModel, UNiagaraStackEditorData& InStackEditorData, UNiagaraNodeOutput& InOutputNode, int32 InTargetIndex)
 {
 	Super::Initialize(InSystemViewModel, InEmitterViewModel, InStackEditorData);
 	OutputNode = &InOutputNode;
+	TargetIndex = InTargetIndex;
+}
+
+UNiagaraStackAddModuleItem::EDisplayMode UNiagaraStackAddScriptModuleItem::GetDisplayMode() const
+{
+	if (TargetIndex != INDEX_NONE)
+	{
+		return EDisplayMode::Compact;
+	}
+	return EDisplayMode::Standard;
 }
 
 void UNiagaraStackAddScriptModuleItem::GetAvailableParameters(TArray<FNiagaraVariable>& OutAvailableParameterVariables) const
@@ -81,4 +91,21 @@ UNiagaraNodeOutput* UNiagaraStackAddScriptModuleItem::GetOrCreateOutputNode()
 	return OutputNode.Get();
 }
 
+int32 UNiagaraStackAddScriptModuleItem::GetTargetIndex() const
+{
+	return TargetIndex;
+}
+
+FName UNiagaraStackAddScriptModuleItem::GetItemBackgroundName() const
+{
+	if (GetDisplayMode() == UNiagaraStackAddModuleItem::EDisplayMode::Compact)
+	{
+		return "NiagaraEditor.Stack.Group.BackgroundColor";
+	}
+	else
+	{
+		return Super::GetItemBackgroundName();
+	}
+}
+
 #undef LOCTEXT_NAMESPACE
\ No newline at end of file
diff --git a/Engine/Plugins/FX/Niagara/Source/NiagaraEditor/Private/ViewModels/Stack/NiagaraStackEmitterSpawnScriptItemGroup.cpp b/Engine/Plugins/FX/Niagara/Source/NiagaraEditor/Private/ViewModels/Stack/NiagaraStackEmitterSpawnScriptItemGroup.cpp
--- a/Engine/Plugins/FX/Niagara/Source/NiagaraEditor/Private/ViewModels/Stack/NiagaraStackEmitterSpawnScriptItemGroup.cpp
+++ b/Engine/Plugins/FX/Niagara/Source/NiagaraEditor/Private/ViewModels/Stack/NiagaraStackEmitterSpawnScriptItemGroup.cpp
@@ -4,23 +4,70 @@
 #include "NiagaraStackObject.h"
 #include "NiagaraStackSpacer.h"
 #include "NiagaraStackItemExpander.h"
+#include "ViewModels/NiagaraSystemViewModel.h"
 #include "ViewModels/NiagaraEmitterViewModel.h"
 #include "NiagaraEmitter.h"
 #include "NiagaraStackEditorData.h"
+#include "ViewModels/Stack/NiagaraStackGraphUtilities.h"
+#include "NiagaraScriptMergeManager.h"
 
 #define LOCTEXT_NAMESPACE "UNiagaraStackScriptItemGroup"
 
+void UNiagaraStackEmitterPropertiesItem::Initialize(TSharedRef<FNiagaraSystemViewModel> InSystemViewModel, TSharedRef<FNiagaraEmitterViewModel> InEmitterViewModel, UNiagaraStackEditorData& InStackEditorData)
+{
+	Super::Initialize(InSystemViewModel, InEmitterViewModel, InStackEditorData);
+	Emitter = GetEmitterViewModel()->GetEmitter();
+	Emitter->OnPropertiesChanged().AddUObject(this, &UNiagaraStackEmitterPropertiesItem::EmitterPropertiesChanged);
+}
+
 FText UNiagaraStackEmitterPropertiesItem::GetDisplayName() const
 {
 	return LOCTEXT("EmitterPropertiesDisplayName", "Emitter Properties");
 }
 
+bool UNiagaraStackEmitterPropertiesItem::CanResetToBase() const
+{
+	if (bCanResetToBase.IsSet() == false)
+	{
+		const UNiagaraEmitter* BaseEmitter = FNiagaraStackGraphUtilities::GetBaseEmitter(*Emitter.Get(), GetSystemViewModel()->GetSystem());
+		if (BaseEmitter != nullptr && Emitter != BaseEmitter)
+		{
+			TSharedRef<FNiagaraScriptMergeManager> MergeManager = FNiagaraScriptMergeManager::Get();
+			bCanResetToBase = MergeManager->IsEmitterEditablePropertySetDifferentFromBase(*Emitter.Get(), *BaseEmitter);
+		}
+		else
+		{
+			bCanResetToBase = false;
+		}
+	}
+	return bCanResetToBase.GetValue();
+}
+
+void UNiagaraStackEmitterPropertiesItem::ResetToBase()
+{
+	if (bCanResetToBase.GetValue())
+	{
+		const UNiagaraEmitter* BaseEmitter = FNiagaraStackGraphUtilities::GetBaseEmitter(*Emitter.Get(), GetSystemViewModel()->GetSystem());
+		TSharedRef<FNiagaraScriptMergeManager> MergeManager = FNiagaraScriptMergeManager::Get();
+		MergeManager->ResetEmitterEditablePropertySetToBase(*Emitter, *BaseEmitter);
+	}
+}
+
+void UNiagaraStackEmitterPropertiesItem::BeginDestroy()
+{
+	if (Emitter.IsValid())
+	{
+		Emitter->OnPropertiesChanged().RemoveAll(this);
+	}
+	Super::BeginDestroy();
+}
+
 void UNiagaraStackEmitterPropertiesItem::RefreshChildrenInternal(const TArray<UNiagaraStackEntry*>& CurrentChildren, TArray<UNiagaraStackEntry*>& NewChildren)
 {
 	if (EmitterObject == nullptr)
 	{
 		EmitterObject = NewObject<UNiagaraStackObject>(this);
-		EmitterObject->Initialize(GetSystemViewModel(), GetEmitterViewModel(), GetEmitterViewModel()->GetEmitter());
+		EmitterObject->Initialize(GetSystemViewModel(), GetEmitterViewModel(), Emitter.Get());
 	}
 
 	if (EmitterExpander == nullptr)
@@ -36,6 +83,7 @@ void UNiagaraStackEmitterPropertiesItem::RefreshChildrenInternal(const TArray<UN
 	}
 
 	NewChildren.Add(EmitterExpander);
+	bCanResetToBase.Reset();
 }
 
 void UNiagaraStackEmitterPropertiesItem::EmitterExpandedChanged()
@@ -43,6 +91,11 @@ void UNiagaraStackEmitterPropertiesItem::EmitterExpandedChanged()
 	RefreshChildren();
 }
 
+void UNiagaraStackEmitterPropertiesItem::EmitterPropertiesChanged()
+{
+	bCanResetToBase.Reset();
+}
+
 UNiagaraStackEmitterSpawnScriptItemGroup::UNiagaraStackEmitterSpawnScriptItemGroup()
 	: PropertiesItem(nullptr)
 {
diff --git a/Engine/Plugins/FX/Niagara/Source/NiagaraEditor/Private/ViewModels/Stack/NiagaraStackEntry.cpp b/Engine/Plugins/FX/Niagara/Source/NiagaraEditor/Private/ViewModels/Stack/NiagaraStackEntry.cpp
--- a/Engine/Plugins/FX/Niagara/Source/NiagaraEditor/Private/ViewModels/Stack/NiagaraStackEntry.cpp
+++ b/Engine/Plugins/FX/Niagara/Source/NiagaraEditor/Private/ViewModels/Stack/NiagaraStackEntry.cpp
@@ -2,6 +2,7 @@
 
 #include "NiagaraStackEntry.h"
 #include "NiagaraStackErrorItem.h"
+#include "NiagaraUnmergeableStackEntryFilter.h"
 
 UNiagaraStackEntry::UNiagaraStackEntry()
 {
@@ -12,6 +13,7 @@ void UNiagaraStackEntry::Initialize(TSharedRef<FNiagaraSystemViewModel> InSystem
 	SystemViewModel = InSystemViewModel;
 	EmitterViewModel = InEmitterViewModel;
 	bIsExpanded = IsExpandedByDefault();
+	AddChildFilter(FNiagaraUnmergeableStackEntryFilter::CreateFilter());
 }
 
 FText UNiagaraStackEntry::GetDisplayName() const
@@ -55,10 +57,32 @@ bool UNiagaraStackEntry::GetShouldShowInStack() const
 	return true;
 }
 
-void UNiagaraStackEntry::GetChildren(TArray<UNiagaraStackEntry*>& OutChildren)
+void UNiagaraStackEntry::GetFilteredChildren(TArray<UNiagaraStackEntry*>& OutFilteredChildren)
 {
-	OutChildren.Append(ErrorChildren);
-	OutChildren.Append(Children);
+	OutFilteredChildren.Append(ErrorChildren);
+	for (UNiagaraStackEntry* Child : Children)
+	{
+		bool bPassesFilter = true;
+		for(const FOnFilterChild& ChildFilter : ChildFilters)
+		{ 
+			if (ChildFilter.Execute(*Child) == false)
+			{
+				bPassesFilter = false;
+				break;
+			}
+		}
+
+		if (bPassesFilter)
+		{
+			OutFilteredChildren.Add(Child);
+		}
+	}
+}
+
+void UNiagaraStackEntry::GetUnfilteredChildren(TArray<UNiagaraStackEntry*>& OutUnfilteredChildren)
+{
+	OutUnfilteredChildren.Append(ErrorChildren);
+	OutUnfilteredChildren.Append(Children);
 }
 
 void UNiagaraStackEntry::RefreshErrors()
@@ -75,6 +99,19 @@ void UNiagaraStackEntry::RefreshErrors()
 	}
 }
 
+FDelegateHandle UNiagaraStackEntry::AddChildFilter(FOnFilterChild ChildFilter)
+{
+	ChildFilters.Add(ChildFilter);
+	StructureChangedDelegate.Broadcast();
+	return ChildFilters.Last().GetHandle();
+}
+
+void UNiagaraStackEntry::RemoveChildFilter(FDelegateHandle FilterHandle)
+{
+	ChildFilters.RemoveAll([=](const FOnFilterChild& ChildFilter) { return ChildFilter.GetHandle() == FilterHandle; });
+	StructureChangedDelegate.Broadcast();
+}
+
 TSharedRef<FNiagaraSystemViewModel> UNiagaraStackEntry::GetSystemViewModel() const
 {
 	TSharedPtr<FNiagaraSystemViewModel> PinnedSystemViewModel = SystemViewModel.Pin();
