Copying Orion-Staging @ 3748653 (Orion/Dev-General @ 3722124) to //UE4/Main

diff --git a/Engine/Plugins/FX/Niagara/Source/NiagaraEditor/Private/ViewModels/Stack/NiagaraStackEventScriptItemGroup.cpp b/Engine/Plugins/FX/Niagara/Source/NiagaraEditor/Private/ViewModels/Stack/NiagaraStackEventScriptItemGroup.cpp
--- a/Engine/Plugins/FX/Niagara/Source/NiagaraEditor/Private/ViewModels/Stack/NiagaraStackEventScriptItemGroup.cpp
+++ b/Engine/Plugins/FX/Niagara/Source/NiagaraEditor/Private/ViewModels/Stack/NiagaraStackEventScriptItemGroup.cpp
@@ -3,6 +3,7 @@
 #include "NiagaraStackEventScriptItemGroup.h"
 #include "NiagaraStackModuleItem.h"
 #include "NiagaraStackAddModuleItem.h"
+#include "NiagaraStackSpacer.h"
 #include "NiagaraScriptViewModel.h"
 #include "NiagaraScriptGraphViewModel.h"
 #include "NiagaraEmitterViewModel.h"
@@ -13,6 +14,9 @@
 #include "NiagaraStackErrorItem.h"
 #include "NiagaraStackStruct.h"
 #include "NiagaraScriptSource.h"
+#include "NiagaraStackGraphUtilities.h"
+#include "NiagaraScriptMergeManager.h"
+#include "NiagaraStackEditorData.h"
 
 #include "Internationalization.h"
 #include "ScopedTransaction.h"
@@ -21,53 +25,178 @@
 
 #define LOCTEXT_NAMESPACE "UNiagaraStackEventScriptItemGroup"
 
-FText UNiagaraStackEventScriptItemGroup::GetDisplayName() const 
+void UNiagaraStackEventHandlerPropertiesItem::Initialize(TSharedRef<FNiagaraSystemViewModel> InSystemViewModel, TSharedRef<FNiagaraEmitterViewModel> InEmitterViewModel, FGuid InEventScriptUsageId, UNiagaraStackEditorData& InStackEditorData)
 {
-	FText EventName = FText::FromName(GetEmitterViewModel()->GetEmitter()->EventHandlerScriptProps[GetScriptOccurrence()].SourceEventName);
-	return (FText::Format(LOCTEXT("FormatEventScriptDisplayName", "Event Handler {0} - Source: {1}"), GetScriptOccurrence(), EventName));
+	Super::Initialize(InSystemViewModel, InEmitterViewModel, InStackEditorData);
+	Emitter = GetEmitterViewModel()->GetEmitter();
+	EventScriptUsageId = InEventScriptUsageId;
+	Emitter->OnPropertiesChanged().AddUObject(this, &UNiagaraStackEventHandlerPropertiesItem::EventHandlerPropertiesChanged);
+	bForceRebuildChildStruct = true;
+
+	const UNiagaraEmitter* BaseEmitter = FNiagaraStackGraphUtilities::GetBaseEmitter(*Emitter.Get(), GetSystemViewModel()->GetSystem());
+	if (BaseEmitter != nullptr && Emitter != BaseEmitter)
+	{
+		TSharedRef<FNiagaraScriptMergeManager> MergeManager = FNiagaraScriptMergeManager::Get();
+		bHasBaseEventHandler = MergeManager->HasBaseEventHandler(*BaseEmitter, EventScriptUsageId);
+	}
 }
-void UNiagaraStackEventScriptItemGroup::RefreshChildrenInternal(const TArray<UNiagaraStackEntry*>& CurrentChildren, TArray<UNiagaraStackEntry*>& NewChildren)
+
+FText UNiagaraStackEventHandlerPropertiesItem::GetDisplayName() const
 {
-	UNiagaraEmitter* Emitter = GetEmitterViewModel()->GetEmitter();
-	if (GetScriptOccurrence() < Emitter->EventHandlerScriptProps.Num())
-	{
-		uint8* EventStructMemory = (uint8*)&(Emitter->EventHandlerScriptProps[GetScriptOccurrence()]);
+	return LOCTEXT("EventHandlerPropertiesDisplayName", "Event Handler Properties");
+}
 
-		UNiagaraStackEntry* FoundItem = nullptr;
-		for (UNiagaraStackEntry* CurrentChild : CurrentChildren)
+bool UNiagaraStackEventHandlerPropertiesItem::CanResetToBase() const
+{
+	if (bCanResetToBase.IsSet() == false)
+	{
+		if (bHasBaseEventHandler)
 		{
-			UNiagaraStackStruct* CurrentStructNode = Cast<UNiagaraStackStruct>(CurrentChild);
-			if (CurrentStructNode != nullptr &&
-				CurrentStructNode->GetOwningObject() == Emitter &&
-				CurrentStructNode->GetStructOnScope()->GetStructMemory() == EventStructMemory)
+			const UNiagaraEmitter* BaseEmitter = FNiagaraStackGraphUtilities::GetBaseEmitter(*Emitter.Get(), GetSystemViewModel()->GetSystem());
+			if (BaseEmitter != nullptr && Emitter != BaseEmitter)
+			{
+				TSharedRef<FNiagaraScriptMergeManager> MergeManager = FNiagaraScriptMergeManager::Get();
+				bCanResetToBase = MergeManager->IsEventHandlerPropertySetDifferentFromBase(*Emitter.Get(), *BaseEmitter, EventScriptUsageId);
+			}
+			else
 			{
-				FoundItem = CurrentStructNode;
-				break;
+				bCanResetToBase = false;
 			}
 		}
+		else
+		{
+			bCanResetToBase = false;
+		}
+	}
+	return bCanResetToBase.GetValue();
+}
 
-		if (FoundItem == nullptr)
+void UNiagaraStackEventHandlerPropertiesItem::ResetToBase()
+{
+	if (bCanResetToBase.GetValue())
+	{
+		const UNiagaraEmitter* BaseEmitter = FNiagaraStackGraphUtilities::GetBaseEmitter(*Emitter.Get(), GetSystemViewModel()->GetSystem());
+		TSharedRef<FNiagaraScriptMergeManager> MergeManager = FNiagaraScriptMergeManager::Get();
+		MergeManager->ResetEventHandlerPropertySetToBase(*Emitter, *BaseEmitter, EventScriptUsageId);
+		bForceRebuildChildStruct = true;
+		RefreshChildren();
+	}
+}
+
+void UNiagaraStackEventHandlerPropertiesItem::BeginDestroy()
+{
+	if (Emitter.IsValid())
+	{
+		Emitter->OnPropertiesChanged().RemoveAll(this);
+	}
+	Super::BeginDestroy();
+}
+
+void UNiagaraStackEventHandlerPropertiesItem::RefreshChildrenInternal(const TArray<UNiagaraStackEntry*>& CurrentChildren, TArray<UNiagaraStackEntry*>& NewChildren)
+{
+	FNiagaraEventScriptProperties* EventScriptProperties = Emitter->GetEventHandlerByIdUnsafe(EventScriptUsageId);
+
+	UNiagaraStackStruct* EventHandlerStruct = nullptr;
+	if (EventScriptProperties != nullptr)
+	{
+		uint8* EventStructMemory = (uint8*)EventScriptProperties;
+
+		if (bForceRebuildChildStruct == false)
+		{
+			EventHandlerStruct = FindCurrentChildOfTypeByPredicate<UNiagaraStackStruct>(CurrentChildren,
+				[=](UNiagaraStackStruct* CurrentChild) { return CurrentChild->GetOwningObject() == Emitter && CurrentChild->GetStructOnScope()->GetStructMemory() == EventStructMemory; });
+		}
+
+		if (EventHandlerStruct == nullptr)
 		{
-			UNiagaraStackStruct* StructNode = NewObject<UNiagaraStackStruct>(this);
-			StructNode->Initialize(GetSystemViewModel(), GetEmitterViewModel(), GetEmitterViewModel()->GetEmitter(), FNiagaraEventScriptProperties::StaticStruct(), EventStructMemory);
 			UNiagaraSystem* System = &GetSystemViewModel()->GetSystem();
-			StructNode->SetDetailsCustomization(FOnGetDetailCustomizationInstance::CreateStatic(&FNiagaraEventScriptPropertiesCustomization::MakeInstance, TWeakObjectPtr<UNiagaraSystem>(System), TWeakObjectPtr<UNiagaraEmitter>(Emitter)));
-			FoundItem = StructNode;
+			EventHandlerStruct = NewObject<UNiagaraStackStruct>(this);
+			EventHandlerStruct->Initialize(GetSystemViewModel(), GetEmitterViewModel(), GetEmitterViewModel()->GetEmitter(), FNiagaraEventScriptProperties::StaticStruct(), EventStructMemory);
+			EventHandlerStruct->SetDetailsCustomization(FOnGetDetailCustomizationInstance::CreateStatic(&FNiagaraEventScriptPropertiesCustomization::MakeInstance, TWeakObjectPtr<UNiagaraSystem>(System), TWeakObjectPtr<UNiagaraEmitter>(Emitter)));
+			bForceRebuildChildStruct = false;
 		}
+	}
+
+	FString ExpandedKey = TEXT("EventHandler_") + EventScriptUsageId.ToString();
+
+	if (EventHandlerExpander == nullptr)
+	{
+		EventHandlerExpander = NewObject<UNiagaraStackItemExpander>(this);
+		EventHandlerExpander->Initialize(GetSystemViewModel(), GetEmitterViewModel(), GetStackEditorData(), ExpandedKey, false);
+		EventHandlerExpander->SetOnExpnadedChanged(UNiagaraStackItemExpander::FOnExpandedChanged::CreateUObject(this, &UNiagaraStackEventHandlerPropertiesItem::EventHandlerExpandedChanged));
+	}
 
-		NewChildren.Add(FoundItem);
+	if (EventHandlerStruct != nullptr && GetStackEditorData().GetStackEntryIsExpanded(ExpandedKey, false))
+	{
+		NewChildren.Add(EventHandlerStruct);
 	}
 
+	NewChildren.Add(EventHandlerExpander);
+	bCanResetToBase.Reset();
+}
+
+void UNiagaraStackEventHandlerPropertiesItem::EventHandlerExpandedChanged()
+{
+	RefreshChildren();
+}
+
+void UNiagaraStackEventHandlerPropertiesItem::EventHandlerPropertiesChanged()
+{
+	bCanResetToBase.Reset();
+}
+
+FText UNiagaraStackEventScriptItemGroup::GetDisplayName() const 
+{
+	return DisplayName;
+}
+
+void UNiagaraStackEventScriptItemGroup::RefreshChildrenInternal(const TArray<UNiagaraStackEntry*>& CurrentChildren, TArray<UNiagaraStackEntry*>& NewChildren)
+{
+	UNiagaraEmitter* Emitter = GetEmitterViewModel()->GetEmitter();
+
+	const FNiagaraEventScriptProperties* EventScriptProperties = Emitter->GetEventHandlers().FindByPredicate(
+		[=](const FNiagaraEventScriptProperties& EventScriptProperties) { return EventScriptProperties.Script->GetUsageId() == GetScriptUsageId(); });
+
+	if (EventScriptProperties != nullptr)
+	{
+		DisplayName = FText::Format(LOCTEXT("FormatEventScriptDisplayName", "Event Handler - Source: {0}"), FText::FromName(EventScriptProperties->SourceEventName));
+
+		const UNiagaraEmitter* BaseEmitter = FNiagaraStackGraphUtilities::GetBaseEmitter(*Emitter, GetSystemViewModel()->GetSystem());
+		bHasBaseEventHandler = BaseEmitter != nullptr && FNiagaraScriptMergeManager::Get()->HasBaseEventHandler(*BaseEmitter, GetScriptUsageId());
+	}
+	else
+	{
+		DisplayName = LOCTEXT("UnassignedEventDisplayName", "Unassigned Event");
+		bHasBaseEventHandler = false;
+	}
+
+	if (EventHandlerProperties == nullptr)
+	{
+		EventHandlerProperties = NewObject<UNiagaraStackEventHandlerPropertiesItem>(this);
+		EventHandlerProperties->Initialize(GetSystemViewModel(), GetEmitterViewModel(), GetScriptUsageId(), GetStackEditorData());
+	}
+	NewChildren.Add(EventHandlerProperties);
+
+	if (EventHandlerPropertiesSpacer == nullptr)
+	{
+		EventHandlerPropertiesSpacer = NewObject<UNiagaraStackSpacer>(this);
+		EventHandlerPropertiesSpacer->Initialize(GetSystemViewModel(), GetEmitterViewModel(), "EventHandlerProperties");
+	}
+	NewChildren.Add(EventHandlerPropertiesSpacer);
+
 	Super::RefreshChildrenInternal(CurrentChildren, NewChildren);
 }
 
 bool UNiagaraStackEventScriptItemGroup::CanDelete() const
 {
-	return true;
+	return bHasBaseEventHandler == false;
 }
 
 bool UNiagaraStackEventScriptItemGroup::Delete()
 {
+	TSharedPtr<FNiagaraScriptViewModel> ScriptViewModelPinned = ScriptViewModel.Pin();
+	checkf(ScriptViewModelPinned.IsValid(), TEXT("Can not delete when the script view model has been deleted."));
+
 	UNiagaraEmitter* Emitter = GetEmitterViewModel()->GetEmitter();
 	UNiagaraScriptSource* Source = Cast<UNiagaraScriptSource>(Emitter->GraphSource);
 
@@ -80,30 +209,14 @@ bool UNiagaraStackEventScriptItemGroup::Delete()
 	Emitter->Modify();
 	Source->NodeGraph->Modify();
 	TArray<UNiagaraNode*> EventIndexNodes;
-	Source->NodeGraph->BuildTraversal(EventIndexNodes, GetScriptUsage(), GetScriptOccurrence());
+	Source->NodeGraph->BuildTraversal(EventIndexNodes, GetScriptUsage(), GetScriptUsageId());
 	for (UNiagaraNode* Node : EventIndexNodes)
 	{
 		Node->Modify();
 	}
-
-	TArray<UNiagaraNodeOutput*> EventOutputNodes;
-	Source->NodeGraph->FindOutputNodes(GetScriptUsage(), EventOutputNodes);
-	for (UNiagaraNodeOutput* OutputNode : EventOutputNodes)
-	{
-		OutputNode->Modify();
-	}
-
-	int32 OldScriptOccurrence = GetScriptOccurrence();
 	
 	// First, remove the event handler script properties object.
-	Emitter->EventHandlerScriptProps.RemoveAt(OldScriptOccurrence);
-
-	// Now move all the usage indices down the list.
-	for (int32 i = OldScriptOccurrence; i < Emitter->EventHandlerScriptProps.Num(); i++)
-	{
-		Emitter->EventHandlerScriptProps[i].Script->Modify();
-		Emitter->EventHandlerScriptProps[i].Script->UsageIndex -= 1;
-	}
+	Emitter->RemoveEventHandlerByUsageId(GetScriptUsageId());
 	
 	// Now remove all graph nodes associated with the event script index.
 	for (UNiagaraNode* Node : EventIndexNodes)
@@ -111,22 +224,9 @@ bool UNiagaraStackEventScriptItemGroup::Delete()
 		Node->DestroyNode();
 	}
 
-	// Now move all the other output nodes down in their index to match the event script properties array.
-	for (UNiagaraNodeOutput* OutputNode : EventOutputNodes)
-	{
-		if (EventIndexNodes.Contains(OutputNode))
-		{
-			continue;
-		}
-		else if (OutputNode->GetUsageIndex() > OldScriptOccurrence)
-		{
-			OutputNode->SetUsageIndex(OutputNode->GetUsageIndex() - 1);
-		}
-	}
-
 	// Set the emitter here to that the internal state of the view model is updated.
 	// TODO: Move the logic for managing event handlers into the emitter view model or script view model.
-	ScriptViewModel->SetScripts(Emitter);
+	ScriptViewModelPinned->SetScripts(Emitter);
 	
 	OnModifiedEventHandlersDelegate.ExecuteIfBound();
 
@@ -138,5 +238,4 @@ void UNiagaraStackEventScriptItemGroup::SetOnModifiedEventHandlers(FOnModifiedEv
 	OnModifiedEventHandlersDelegate = OnModifiedEventHandlers;
 }
 
-
 #undef LOCTEXT_NAMESPACE
\ No newline at end of file
diff --git a/Engine/Plugins/FX/Niagara/Source/NiagaraEditor/Private/ViewModels/Stack/NiagaraStackFunctionInput.cpp b/Engine/Plugins/FX/Niagara/Source/NiagaraEditor/Private/ViewModels/Stack/NiagaraStackFunctionInput.cpp
--- a/Engine/Plugins/FX/Niagara/Source/NiagaraEditor/Private/ViewModels/Stack/NiagaraStackFunctionInput.cpp
+++ b/Engine/Plugins/FX/Niagara/Source/NiagaraEditor/Private/ViewModels/Stack/NiagaraStackFunctionInput.cpp
@@ -15,6 +15,7 @@
 #include "NiagaraEditorUtilities.h"
 #include "NiagaraSystemViewModel.h"
 #include "NiagaraSystemScriptViewModel.h"
+#include "NiagaraEmitterHandleViewModel.h"
 #include "NiagaraEmitterViewModel.h"
 #include "NiagaraScriptGraphViewModel.h"
 #include "NiagaraStackEditorData.h"
@@ -22,11 +23,11 @@
 #include "NiagaraScriptSource.h"
 #include "NiagaraConstants.h"
 #include "NiagaraStackGraphUtilities.h"
+#include "NiagaraScriptMergeManager.h"
 
 #include "ScopedTransaction.h"
 #include "Editor.h"
 #include "UObject/StructOnScope.h"
-#include "Modules/ModuleManager.h"
 #include "AssetRegistryModule.h"
 #include "ARFilter.h"
 #include "EdGraph/EdGraphPin.h"
@@ -47,6 +48,12 @@ void UNiagaraStackFunctionInput::BeginDestroy()
 	if (OwningFunctionCallNode.IsValid())
 	{
 		OwningFunctionCallNode->GetGraph()->RemoveOnGraphChangedHandler(GraphChangedHandle);
+
+		if (SourceScript.IsValid())
+		{
+			SourceScript->RapidIterationParameters.RemoveOnChangedHandler(RapidIterationParametersChangedHandle);
+			SourceScript->GetSource()->OnChanged().RemoveAll(this);
+		}
 	}
 }
 
@@ -138,6 +145,55 @@ void UNiagaraStackFunctionInput::Initialize(
 	OwningFunctionCallNode = &InInputFunctionCallNode;
 	OwningAssignmentNode = Cast<UNiagaraNodeAssignment>(OwningFunctionCallNode.Get());
 
+	UNiagaraNodeOutput* OutputNode = FNiagaraStackGraphUtilities::GetEmitterOutputNodeForStackNode(*OwningModuleNode.Get());
+	UNiagaraEmitter* ParentEmitter = InEmitterViewModel->GetEmitter();
+	UNiagaraSystem* ParentSystem = &InSystemViewModel->GetSystem();
+	if (OutputNode)
+	{
+		TArray<UNiagaraScript*> Scripts;
+		if (ParentEmitter != nullptr)
+		{
+			ParentEmitter->GetScripts(Scripts, false);
+		}
+		if (ParentSystem != nullptr)
+		{
+			Scripts.Add(ParentSystem->GetSystemSpawnScript(true));
+			Scripts.Add(ParentSystem->GetSystemUpdateScript(true));
+			Scripts.Add(ParentSystem->GetSystemSpawnScript(false));
+			Scripts.Add(ParentSystem->GetSystemUpdateScript(false));
+		}
+
+		for (UNiagaraScript* Script : Scripts)
+		{
+			if (OutputNode->GetUsage() == ENiagaraScriptUsage::ParticleEventScript)
+			{
+				if (Script->GetUsage() == ENiagaraScriptUsage::ParticleEventScript && Script->GetUsageId() == OutputNode->GetUsageId())
+				{
+					AffectedScripts.Add(Script);
+					break;
+				}
+			}
+			else if (Script->ContainsUsage(OutputNode->GetUsage()))
+			{
+				AffectedScripts.Add(Script);
+			}
+		}
+
+		for (TWeakObjectPtr<UNiagaraScript> AffectedScript : AffectedScripts)
+		{
+			if (AffectedScript.IsValid() && AffectedScript->IsEquivalentUsage(OutputNode->GetUsage()) && AffectedScript->GetUsageId() == OutputNode->GetUsageId())
+			{
+				SourceScript = AffectedScript;
+				RapidIterationParametersChangedHandle = SourceScript->RapidIterationParameters.AddOnChangedHandler(
+					FNiagaraParameterStore::FOnChanged::FDelegate::CreateUObject(this, &UNiagaraStackFunctionInput::OnRapidIterationParametersChanged));
+				SourceScript->GetSource()->OnChanged().AddUObject(this, &UNiagaraStackFunctionInput::OnScriptSourceChanged);
+				break;
+			}
+		}
+	}
+
+	checkf(SourceScript.IsValid(), TEXT("Coudn't find source script in affected scripts."));
+
 	GraphChangedHandle = OwningFunctionCallNode->GetGraph()->AddOnGraphChangedHandler(
 		FOnGraphChanged::FDelegate::CreateUObject(this, &UNiagaraStackFunctionInput::OnGraphChanged));
 
@@ -150,6 +206,15 @@ void UNiagaraStackFunctionInput::Initialize(
 
 	InputType = InInputType;
 	StackEditorDataKey = FNiagaraStackGraphUtilities::GenerateStackFunctionInputEditorDataKey(*OwningFunctionCallNode.Get(), InputParameterHandle);
+
+	/*
+	UE_LOG(LogNiagaraEditor, Log, TEXT("Scripts for Parameter '%s'"), *RapidIterationParameterName);
+	for (UNiagaraScript* Script : AffectedScripts)
+	{
+		UE_LOG(LogNiagaraEditor, Log, TEXT("Added... '%s'"), *Script->GetFullName());
+	}
+	*/
+
 }
 
 const UNiagaraNodeFunctionCall& UNiagaraStackFunctionInput::GetInputFunctionCallNode() const
@@ -172,8 +237,58 @@ const FNiagaraTypeDefinition& UNiagaraStackFunctionInput::GetInputType() const
 	return InputType;
 }
 
+FText UNiagaraStackFunctionInput::GetTooltipText() const
+{
+	return GetTooltipText(InputValues.Mode);
+}
+
+FText UNiagaraStackFunctionInput::GetTooltipText(EValueMode InValueMode) const
+{
+	FNiagaraVariable ValueVariable;
+	UNiagaraGraph* NodeGraph = nullptr;
+
+	if (InValueMode == EValueMode::Linked)
+	{
+		UEdGraphPin* OverridePin = GetOverridePin();
+		UEdGraphPin* ValuePin = OverridePin != nullptr ? OverridePin : GetDefaultPin();
+		ValueVariable = FNiagaraVariable(InputType, InputValues.LinkedHandle.GetParameterHandleString());
+		if (ValuePin != nullptr)
+		{
+			NodeGraph = Cast<UNiagaraGraph>(ValuePin->GetOwningNode()->GetGraph());
+		}
+	}
+	else
+	{
+		ValueVariable = FNiagaraVariable(InputType, InputParameterHandle.GetParameterHandleString());
+		if (OwningFunctionCallNode.IsValid() && OwningFunctionCallNode->FunctionScript != nullptr)
+		{
+			UNiagaraScriptSource* Source = Cast<UNiagaraScriptSource>(OwningFunctionCallNode->FunctionScript->GetSource());
+			const UEdGraphSchema_Niagara* NiagaraSchema = GetDefault<UEdGraphSchema_Niagara>();
+			NodeGraph = Source->NodeGraph;
+		}
+	}
+
+	const FNiagaraVariableMetaData* MetaData = nullptr;
+	if (FNiagaraConstants::IsNiagaraConstant(ValueVariable))
+	{
+		MetaData = FNiagaraConstants::GetConstantMetaData(ValueVariable);
+	}
+	else if (NodeGraph != nullptr)
+	{
+		MetaData = NodeGraph->GetMetaData(ValueVariable);
+	}
+
+	if (MetaData != nullptr)
+	{
+		return MetaData->Description;
+	}
+
+	return FText::FromName(ValueVariable.GetName());
+}
+
 void UNiagaraStackFunctionInput::RefreshChildrenInternal(const TArray<UNiagaraStackEntry*>& CurrentChildren, TArray<UNiagaraStackEntry*>& NewChildren)
 {
+	RapidIterationParameterName = CreateRapidIterationVariableName(AliasedInputParameterHandle.GetParameterHandleString().ToString());
 	RefreshValues();
 
 	if (InputValues.Mode == EValueMode::Dynamic && InputValues.DynamicNode.IsValid())
@@ -259,6 +374,7 @@ void UNiagaraStackFunctionInput::RefreshValues()
 		}
 	}
 
+	bCanResetToBase.Reset();
 	ValueChangedDelegate.Broadcast();
 }
 
@@ -301,26 +417,15 @@ void UNiagaraStackFunctionInput::SetLinkedValueHandle(const FNiagaraParameterHan
 
 	FScopedTransaction ScopedTransaction(LOCTEXT("UpdateLinkedInputValue", "Update linked input value"));
 	UEdGraphPin& OverridePin = GetOrCreateOverridePin();
-	RemoveAllNodesConnectedToOverridePin(OverridePin, this);
-
-	UEdGraph* Graph = OwningFunctionCallNode->GetGraph();
-	FGraphNodeCreator<UNiagaraNodeParameterMapGet> GetNodeCreator(*Graph);
-	UNiagaraNodeParameterMapGet* GetNode = GetNodeCreator.CreateNode();
-	GetNodeCreator.Finalize();
-
-	UEdGraphPin* GetInputPin = FNiagaraStackGraphUtilities::GetParameterMapInputPin(*GetNode);
-	checkf(GetInputPin != nullptr, TEXT("Parameter map get node was missing it's parameter map input pin."));
-
-	UEdGraphPin* OverrideNodeInputPin = FNiagaraStackGraphUtilities::GetParameterMapInputPin(*GetOverrideNode());
-	UEdGraphPin* PreviousStackNodeOutputPin = OverrideNodeInputPin->LinkedTo[0];
-	if (ensureMsgf(PreviousStackNodeOutputPin != nullptr, TEXT("Invalid Stack Graph - No previous stack node.")))
+	RemoveNodesForOverridePin(OverridePin);
+	if (IsRapidIterationCandidate())
 	{
-		UEdGraphPin* GetOutputPin = GetNode->RequestNewTypedPin(EGPD_Output, InputType, InParameterHandle.GetParameterHandleString());
-		GetInputPin->MakeLinkTo(PreviousStackNodeOutputPin);
-		GetOutputPin->MakeLinkTo(&OverridePin);
+		RemoveRapidIterationParametersForAffectedScripts();
 	}
 
-	FNiagaraStackGraphUtilities::RelayoutGraph(*Graph);
+	FNiagaraStackGraphUtilities::SetLinkedValueHandleForFunctionInput(OverridePin, InParameterHandle);
+	FNiagaraStackGraphUtilities::RelayoutGraph(*OwningFunctionCallNode->GetGraph());
+
 	RefreshValues();
 }
 
@@ -384,11 +489,21 @@ void UNiagaraStackFunctionInput::GetAvailableParameterHandles(TArray<FNiagaraPar
 
 	TArray<UNiagaraNodeOutput*> AllOutputNodes;
 	GetEmitterViewModel()->GetSharedScriptViewModel()->GetGraphViewModel()->GetGraph()->GetNodesOfClass<UNiagaraNodeOutput>(AllOutputNodes);
-	if (GetSystemViewModel()->GetSystemIsTransient() == false)
+	if (GetSystemViewModel()->GetEditMode() == ENiagaraSystemViewModelEditMode::SystemAsset)
 	{
 		GetSystemViewModel()->GetSystemScriptViewModel()->GetGraphViewModel()->GetGraph()->GetNodesOfClass<UNiagaraNodeOutput>(AllOutputNodes);
 	}
 
+	TArray<FNiagaraVariable> ExposedVars;
+	GetSystemViewModel()->GetSystem().GetExposedParameters().GetParameters(ExposedVars);
+	for (const FNiagaraVariable& ExposedVar : ExposedVars)
+	{
+		if (ExposedVar.GetType() == InputType)
+		{
+			AvailableParameterHandles.Add(FNiagaraParameterHandle::CreateEngineParameterHandle(ExposedVar));
+		}
+	}
+
 	for (UNiagaraNodeOutput* OutputNode : AllOutputNodes)
 	{
 		if (OutputNode == CurrentOutputNode || UsageRunsBefore(OutputNode->GetUsage(), CurrentOutputNode->GetUsage()))
@@ -501,35 +616,16 @@ void UNiagaraStackFunctionInput::SetDynamicInput(UNiagaraScript* DynamicInput)
 	FScopedTransaction ScopedTransaction(LOCTEXT("SetDynamicInput", "Make dynamic input"));
 
 	UEdGraphPin& OverridePin = GetOrCreateOverridePin();
-	RemoveAllNodesConnectedToOverridePin(OverridePin, this);
-
-	UEdGraph* Graph = OwningFunctionCallNode->GetGraph();
-	FGraphNodeCreator<UNiagaraNodeFunctionCall> FunctionCallNodeCreator(*Graph);
-	UNiagaraNodeFunctionCall* FunctionCallNode = FunctionCallNodeCreator.CreateNode();
-	FunctionCallNode->FunctionScript = DynamicInput;
-	FunctionCallNodeCreator.Finalize();
-
-	UEdGraphPin* FunctionCallInputPin = FNiagaraStackGraphUtilities::GetParameterMapInputPin(*FunctionCallNode);
-	TArray<UEdGraphPin*> FunctionCallOutputPins;
-	FunctionCallNode->GetOutputPins(FunctionCallOutputPins);
-
-	const UEdGraphSchema_Niagara* NiagaraSchema = GetDefault<UEdGraphSchema_Niagara>();
-
-	if (ensureMsgf(FunctionCallInputPin != nullptr, TEXT("Invalid Stack Graph - Dynamic Input function call did not have a parameter map input pin.")) &&
-		ensureMsgf(FunctionCallOutputPins.Num() == 1 && NiagaraSchema->PinToTypeDefinition(FunctionCallOutputPins[0]) == InputType, 
-			TEXT("Invalid Stack Graph - Dynamic Input function did not have the correct typed output pin")))
+	RemoveNodesForOverridePin(OverridePin);
+	if (IsRapidIterationCandidate())
 	{
-		UEdGraphPin* OverrideNodeInputPin = FNiagaraStackGraphUtilities::GetParameterMapInputPin(*GetOverrideNode());
-		UEdGraphPin* PreviousStackNodeOutputPin = OverrideNodeInputPin->LinkedTo[0];
-		if (ensureMsgf(PreviousStackNodeOutputPin != nullptr, TEXT("Invalid Stack Graph - No previous stack node.")))
-		{
-			FunctionCallInputPin->MakeLinkTo(PreviousStackNodeOutputPin);
-			FunctionCallOutputPins[0]->MakeLinkTo(&OverridePin);
-		}
+		RemoveRapidIterationParametersForAffectedScripts();
 	}
 
+	UNiagaraNodeFunctionCall* FunctionCallNode;
+	FNiagaraStackGraphUtilities::SetDynamicInputForFunctionInput(OverridePin, DynamicInput, FunctionCallNode);
 	FNiagaraStackGraphUtilities::InitializeDataInterfaceInputs(GetSystemViewModel(), GetEmitterViewModel(), *StackEditorData, *OwningModuleNode, *FunctionCallNode);
-	FNiagaraStackGraphUtilities::RelayoutGraph(*Graph);
+	FNiagaraStackGraphUtilities::RelayoutGraph(*OwningFunctionCallNode->GetGraph());
 
 	RefreshChildren();
 }
@@ -568,11 +664,25 @@ void UNiagaraStackFunctionInput::NotifyEndLocalValueChange()
 	}
 }
 
+bool UNiagaraStackFunctionInput::IsRapidIterationCandidate() const
+{
+	return InputType != FNiagaraTypeDefinition::GetBoolDef() && !InputType.IsEnum() &&
+		InputType != FNiagaraTypeDefinition::GetParameterMapDef() && !InputType.IsDataInterface();
+}
+
 void UNiagaraStackFunctionInput::SetLocalValue(TSharedRef<FStructOnScope> InLocalValue)
 {
 	UEdGraphPin* DefaultPin = GetDefaultPin();
-	UEdGraphPin* CurrentOverridePin = GetOverridePin();
-	UEdGraphPin* ValuePin = CurrentOverridePin != nullptr ? CurrentOverridePin : DefaultPin;
+	UEdGraphPin* OverridePin = GetOverridePin();
+	UEdGraphPin* ValuePin = DefaultPin;
+	
+	// If the default pin in the function graph is connected internally, rapid iteration parameters can't be used since
+	// the compilation currently won't use them.
+	bool bCanUseRapidIterationParameter = IsRapidIterationCandidate() && DefaultPin->LinkedTo.Num() == 0;
+	if (bCanUseRapidIterationParameter == false)
+	{
+		ValuePin = OverridePin != nullptr ? OverridePin : DefaultPin;
+	}
 
 	TSharedPtr<FStructOnScope> CurrentValue;
 	bool bCanHaveLocalValue = ValuePin != nullptr;
@@ -586,22 +696,72 @@ void UNiagaraStackFunctionInput::SetLocalValue(TSharedRef<FStructOnScope> InLoca
 
 	FScopedTransaction ScopedTransaction(LOCTEXT("UpdateInputLocalValue", "Update input local value"));
 	UNiagaraGraph* EmitterGraph = Cast<UNiagaraGraph>(OwningFunctionCallNode->GetGraph());
-	UEdGraphPin& OverridePin = GetOrCreateOverridePin();
 
-	bool bGraphWillNeedRelayout = &OverridePin != CurrentOverridePin || OverridePin.LinkedTo.Num() > 0;
+	bool bGraphWillNeedRelayout = false;
+	if (OverridePin != nullptr && OverridePin->LinkedTo.Num() > 0)
+	{
+		RemoveNodesForOverridePin(*OverridePin);
+		bGraphWillNeedRelayout = true;
+	}
 
-	RemoveAllNodesConnectedToOverridePin(OverridePin, this);
-	OverridePin.Modify();
-	FNiagaraVariable LocalValueVariable(InputType, NAME_None);
-	LocalValueVariable.SetData(InLocalValue->GetStructMemory());
 	const UEdGraphSchema_Niagara* NiagaraSchema = GetDefault<UEdGraphSchema_Niagara>();
-	OverridePin.DefaultValue = NiagaraSchema->NiagaraVariableToPinDefaultValue(LocalValueVariable);
-	EmitterGraph->NotifyGraphNeedsRecompile();
+	if (bCanUseRapidIterationParameter)
+	{
+		for (TWeakObjectPtr<UNiagaraScript> Script : AffectedScripts)
+		{
+			Script->Modify();
+		}
 
+		// If there is currently an override, we need to get rid of it.
+		if (OverridePin != nullptr)
+		{
+			UNiagaraNode* OverrideNode = CastChecked<UNiagaraNode>(OverridePin->GetOwningNode());
+			OverrideNode->Modify();
+			OverrideNode->RemovePin(OverridePin);
+		}
+	
+		FNiagaraVariable ValueVariable = NiagaraSchema->PinToNiagaraVariable(ValuePin, true);
+		ValueVariable.SetName(*RapidIterationParameterName);
+
+		for (TWeakObjectPtr<UNiagaraScript> Script : AffectedScripts)
+		{
+			int32 FoundOffset = Script->RapidIterationParameters.IndexOf(ValueVariable);
+			checkSlow(!ValueVariable.IsDataInterface());
+			if (FoundOffset == INDEX_NONE)
+			{
+				Script->RapidIterationParameters.AddParameter(ValueVariable);
+				UE_LOG(LogNiagaraEditor, Log, TEXT("Adding Parameter %s to Script %s"), *ValueVariable.GetName().ToString(), *Script->GetFullName());
+			}
+			Script->RapidIterationParameters.SetParameterData(InLocalValue->GetStructMemory(), ValueVariable);
+		}
+		GetSystemViewModel()->ReInitializeSystemInstances();
+	}
+	else 
+	{
+		FNiagaraVariable LocalValueVariable(InputType, NAME_None);
+		LocalValueVariable.SetData(InLocalValue->GetStructMemory());
+		FString PinDefaultValue;
+		if(ensureMsgf(NiagaraSchema->TryGetPinDefaultValueFromNiagaraVariable(LocalValueVariable, PinDefaultValue),
+			TEXT("Could not generate default value string for non-rapid iteration parameter.")))
+		{
+			if (OverridePin == nullptr)
+			{
+				OverridePin = &GetOrCreateOverridePin();
+				bGraphWillNeedRelayout = true;
+			}
+
+			OverridePin->Modify();
+			OverridePin->DefaultValue = PinDefaultValue;
+			EmitterGraph->NotifyGraphNeedsRecompile();
+		}
+	}
+	
 	if (bGraphWillNeedRelayout)
 	{
 		FNiagaraStackGraphUtilities::RelayoutGraph(*EmitterGraph);
 	}
+
+	RefreshValues();
 }
 
 bool UNiagaraStackFunctionInput::CanReset() const
@@ -619,9 +779,88 @@ bool UNiagaraStackFunctionInput::CanReset() const
 	}
 	else
 	{
-		// All other input modes can be reset if there is an override pin available.
-		return GetOverridePin() != nullptr;
+		if (IsRapidIterationCandidate())
+		{
+			if (GetOverridePin() != nullptr)
+			{
+				return true;
+			}
+			else if (InputValues.LocalStruct.IsValid() && !FNiagaraEditorUtilities::DataMatches(GetDefaultValue(), *(InputValues.LocalStruct.Get())))
+			{
+				return true;
+			}
+			else
+			{
+				return false;
+			}
+		}
+		else
+		{
+			// All other input modes can be reset if there is an override pin available.
+			return GetOverridePin() != nullptr;
+		}
+	}
+}
+
+FNiagaraVariable UNiagaraStackFunctionInput::GetDefaultValue() const
+{
+	FNiagaraVariable Var;
+	UEdGraphPin* DefaultPin = GetDefaultPin();
+	if (DefaultPin != nullptr)
+	{
+		const UEdGraphSchema_Niagara* NiagaraSchema = GetDefault<UEdGraphSchema_Niagara>();
+		Var = NiagaraSchema->PinToNiagaraVariable(DefaultPin, true);
+		Var.SetName(*RapidIterationParameterName);
+	}
+	return Var;
+}
+
+bool UNiagaraStackFunctionInput::UpdateRapidIterationParametersForAffectedScripts(const uint8* Data)
+{
+	for (TWeakObjectPtr<UNiagaraScript> Script : AffectedScripts)
+	{
+		Script->Modify();
+	}
+
+	FNiagaraVariable Var(InputType, *RapidIterationParameterName);
+
+	for (TWeakObjectPtr<UNiagaraScript> Script : AffectedScripts)
+	{
+		int32 FoundOffset = Script->RapidIterationParameters.IndexOf(Var);
+		checkSlow(!Var.IsDataInterface());
+		if (FoundOffset == INDEX_NONE)
+		{
+			Script->RapidIterationParameters.AddParameter(Var);
+			UE_LOG(LogNiagaraEditor, Log, TEXT("Adding Parameter %s to Script %s"), *Var.GetName().ToString(), *Script->GetFullName());
+
+		}
+		Script->RapidIterationParameters.SetParameterData(Data, Var);
+	}
+	GetSystemViewModel()->ResetSystem();
+	return true;
+}
+
+bool UNiagaraStackFunctionInput::RemoveRapidIterationParametersForAffectedScripts()
+{
+	for (TWeakObjectPtr<UNiagaraScript> Script : AffectedScripts)
+	{
+		Script->Modify();
+	}
+
+	FNiagaraVariable Var(InputType, *RapidIterationParameterName);
+	for (TWeakObjectPtr<UNiagaraScript> Script : AffectedScripts)
+	{
+		int32 FoundOffset = Script->RapidIterationParameters.IndexOf(Var);
+		if (FoundOffset == INDEX_NONE)
+		{
+			return false;
+		}
+
+		UE_LOG(LogNiagaraEditor, Log, TEXT("Removed Var '%s' from Script %s"), *Var.GetName().ToString(), *Script->GetFullName());
+
+		Script->RapidIterationParameters.RemoveParameter(Var);
 	}
+	return true;
 }
 
 void UNiagaraStackFunctionInput::Reset()
@@ -638,7 +877,7 @@ void UNiagaraStackFunctionInput::Reset()
 		else
 		{
 			UEdGraphPin& OverridePin = GetOrCreateOverridePin();
-			RemoveAllNodesConnectedToOverridePin(OverridePin, this);
+			RemoveNodesForOverridePin(OverridePin);
 
 			FString InputNodeName = InputParameterHandlePath[0].GetName().ToString();
 			for (int32 i = 1; i < InputParameterHandlePath.Num(); i++)
@@ -646,21 +885,14 @@ void UNiagaraStackFunctionInput::Reset()
 				InputNodeName += "." + InputParameterHandlePath[i].GetName().ToString();
 			}
 
-			UEdGraph* Graph = OwningFunctionCallNode->GetGraph();
-			FGraphNodeCreator<UNiagaraNodeInput> InputNodeCreator(*Graph);
-			UNiagaraNodeInput* InputNode = InputNodeCreator.CreateNode();
-			FNiagaraEditorUtilities::InitializeParameterInputNode(*InputNode, InputType, CastChecked<UNiagaraGraph>(Graph), *InputNodeName);
-
-			UNiagaraDataInterface* InputValueObject = NewObject<UNiagaraDataInterface>(InputNode, const_cast<UClass*>(InputType.GetClass()), InputParameterHandle.GetName(), RF_Transactional);
+			UNiagaraDataInterface* InputValueObject;
+			FNiagaraStackGraphUtilities::SetDataValueObjectForFunctionInput(OverridePin, const_cast<UClass*>(InputType.GetClass()), InputNodeName, InputValueObject);
 			if (InputValues.DataObjects.GetDefaultValueObject() != nullptr)
 			{
 				InputValues.DataObjects.GetDefaultValueObject()->CopyTo(InputValueObject);
 			}
-			InputNode->DataInterface = InputValueObject;
 
-			InputNodeCreator.Finalize();
-			FNiagaraStackGraphUtilities::ConnectPinToInputNode(OverridePin, *InputNode);
-			FNiagaraStackGraphUtilities::RelayoutGraph(*Graph);
+			FNiagaraStackGraphUtilities::RelayoutGraph(*OwningFunctionCallNode->GetGraph());
 		}
 	}
 	else
@@ -668,20 +900,169 @@ void UNiagaraStackFunctionInput::Reset()
 		// For all other value modes removing the nodes connected to the override pin resets them.
 		UNiagaraNodeParameterMapSet* OverrideNode = GetOverrideNode();
 		UEdGraphPin* OverridePin = GetOverridePin();
-		if (ensureMsgf(OverrideNode != nullptr && OverridePin != nullptr, TEXT("Can not reset the value of an input that doesn't have a valid override node and override pin")))
+		bool bGraphNeedsRecompile = false;
 		{
 			FScopedTransaction ScopedTransaction(LOCTEXT("ResetInputStructTransaction", "Reset the inputs value to default."));
-			RemoveAllNodesConnectedToOverridePin(*OverridePin, this);
-			OverrideNode->Modify();
-			OverrideNode->RemovePin(OverridePin);
+			
+			if (IsRapidIterationCandidate())
+			{
+				if (OverrideNode != nullptr && OverridePin != nullptr)
+				{
+					RemoveNodesForOverridePin(*OverridePin);
+					OverrideNode->Modify();
+					OverrideNode->RemovePin(OverridePin);
+					bGraphNeedsRecompile = true;
+				}
 
-			OwningFunctionCallNode->GetNiagaraGraph()->NotifyGraphNeedsRecompile();
-			FNiagaraStackGraphUtilities::RelayoutGraph(*OwningFunctionCallNode->GetGraph());
+				// Get the default value of the graph pin and use that to reset the rapid iteration variables...
+				UEdGraphPin* DefaultPin = GetDefaultPin();
+				if (DefaultPin != nullptr)
+				{
+					FNiagaraVariable DefaultVar = GetDefaultValue();
+					UpdateRapidIterationParametersForAffectedScripts(DefaultVar.GetData());
+				}
+			}
+			else
+			{
+				if (ensureMsgf(OverrideNode != nullptr && OverridePin != nullptr, TEXT("Can not reset the value of an input that doesn't have a valid override node and override pin")))
+				{
+					RemoveNodesForOverridePin(*OverridePin);
+					OverrideNode->Modify();
+					OverrideNode->RemovePin(OverridePin);
+					bGraphNeedsRecompile =  true;
+				}
+			}
+
+			if (bGraphNeedsRecompile)
+			{
+				OwningFunctionCallNode->GetNiagaraGraph()->NotifyGraphNeedsRecompile();
+				FNiagaraStackGraphUtilities::RelayoutGraph(*OwningFunctionCallNode->GetGraph());
+			}
 		}
 	}
 	RefreshChildren();
 }
 
+bool UNiagaraStackFunctionInput::EmitterHasBase() const
+{
+	return GetSystemViewModel()->GetEditMode() == ENiagaraSystemViewModelEditMode::SystemAsset;
+}
+
+bool UNiagaraStackFunctionInput::CanResetToBase() const
+{
+	if (EmitterHasBase())
+	{
+		if (bCanResetToBase.IsSet() == false)
+		{
+			bool bIsModuleInput = OwningFunctionCallNode == OwningModuleNode;
+			if (bIsModuleInput)
+			{
+				TSharedRef<FNiagaraScriptMergeManager> MergeManager = FNiagaraScriptMergeManager::Get();
+
+				UNiagaraNodeOutput* OutputNode = FNiagaraStackGraphUtilities::GetEmitterOutputNodeForStackNode(*OwningFunctionCallNode.Get());
+				if(MergeManager->IsMergeableScriptUsage(OutputNode->GetUsage()))
+				{
+					const UNiagaraEmitter* BaseEmitter = FNiagaraStackGraphUtilities::GetBaseEmitter(*GetEmitterViewModel()->GetEmitter(), GetSystemViewModel()->GetSystem());
+
+					bCanResetToBase = BaseEmitter != nullptr && MergeManager->IsModuleInputDifferentFromBase(
+						*GetEmitterViewModel()->GetEmitter(),
+						*BaseEmitter,
+						OutputNode->GetUsage(),
+						OutputNode->GetUsageId(),
+						OwningModuleNode->NodeGuid,
+						InputParameterHandle.GetName().ToString());
+				}
+				else
+				{
+					bCanResetToBase = false;
+				}
+			}
+			else
+			{
+				bCanResetToBase = false;
+			}
+		}
+		return bCanResetToBase.GetValue();
+	}
+	return false;
+}
+
+void UNiagaraStackFunctionInput::ResetToBase()
+{
+	if (CanResetToBase())
+	{
+		TSharedRef<FNiagaraScriptMergeManager> MergeManager = FNiagaraScriptMergeManager::Get();
+
+		TSharedPtr<FNiagaraEmitterHandleViewModel> ThisEmitterHandleViewModel;
+		for (TSharedRef<FNiagaraEmitterHandleViewModel> EmitterHandleViewModel : GetSystemViewModel()->GetEmitterHandleViewModels())
+		{
+			if (EmitterHandleViewModel->GetEmitterViewModel() == GetEmitterViewModel())
+			{
+				ThisEmitterHandleViewModel = EmitterHandleViewModel;
+				break;
+			}
+		}
+
+		const UNiagaraEmitter* BaseEmitter = ThisEmitterHandleViewModel->GetEmitterHandle()->GetSource();
+		UNiagaraNodeOutput* OutputNode = FNiagaraStackGraphUtilities::GetEmitterOutputNodeForStackNode(*OwningFunctionCallNode.Get());
+
+		FScopedTransaction ScopedTransaction(LOCTEXT("ResetInputToBaseTransaction", "Reset this input to match the parent emitter."));
+		FNiagaraScriptMergeManager::FApplyDiffResults Results = MergeManager->ResetModuleInputToBase(
+			*GetEmitterViewModel()->GetEmitter(),
+			*BaseEmitter,
+			OutputNode->GetUsage(),
+			OutputNode->GetUsageId(),
+			OwningModuleNode->NodeGuid,
+			InputParameterHandle.GetName().ToString());
+
+		if (Results.bSucceeded)
+		{
+			// If resetting to the base succeeded, an unknown number of rapid iteration parameters may have been added.  To fix
+			// this copy all of the owning scripts rapid iteration parameters to all other affected scripts.
+			// TODO: Either the merge should take care of this directly, or at least provide more information about what changed.
+			UNiagaraScript* OwningScript = GetEmitterViewModel()->GetEmitter()->GetScript(OutputNode->GetUsage(), OutputNode->GetUsageId());
+			TArray<FNiagaraVariable> OwningScriptRapidIterationParameters;
+			OwningScript->RapidIterationParameters.GetParameters(OwningScriptRapidIterationParameters);
+			if (OwningScriptRapidIterationParameters.Num() > 0)
+			{
+				for (TWeakObjectPtr<UNiagaraScript> AffectedScript : AffectedScripts)
+				{
+					if (AffectedScript.Get() != OwningScript)
+					{
+						AffectedScript->Modify();
+						for (FNiagaraVariable& OwningScriptRapidIterationParameter : OwningScriptRapidIterationParameters)
+						{
+							AffectedScript->RapidIterationParameters.AddParameter(OwningScriptRapidIterationParameter);
+							AffectedScript->RapidIterationParameters.SetParameterData(
+								OwningScript->RapidIterationParameters.GetParameterData(OwningScriptRapidIterationParameter), OwningScriptRapidIterationParameter);
+						}
+					}
+				}
+			}
+		}
+		RefreshChildren();
+	}
+}
+
+FString UNiagaraStackFunctionInput::CreateRapidIterationVariableName(const FString& InName)
+{
+	UNiagaraNodeOutput* OutputNode = FNiagaraStackGraphUtilities::GetEmitterOutputNodeForStackNode(*OwningModuleNode.Get());
+	UNiagaraEmitter* ParentEmitter = GetEmitterViewModel()->GetEmitter();
+	UNiagaraSystem* ParentSystem = &GetSystemViewModel()->GetSystem();
+	FNiagaraVariable Temp(InputType, *InName);
+	if (OutputNode->GetUsage() == ENiagaraScriptUsage::SystemSpawnScript || OutputNode->GetUsage() == ENiagaraScriptUsage::SystemUpdateScript)
+	{
+		Temp = FNiagaraParameterMapHistory::ConvertVariableToRapidIterationConstantName(Temp, nullptr); // These names *should* have the emitter baked in...
+	}
+	else
+	{
+		Temp = FNiagaraParameterMapHistory::ConvertVariableToRapidIterationConstantName(Temp, *ParentEmitter->GetUniqueEmitterName());
+	}
+
+	return Temp.GetName().ToString();
+}
+
+
 bool UNiagaraStackFunctionInput::CanRenameInput() const
 {
 	// Only module level assignment node inputs can be renamed.
@@ -710,6 +1091,7 @@ void UNiagaraStackFunctionInput::RenameInput(FName NewName)
 
 		FNiagaraParameterHandle TargetHandle(OwningAssignmentNode->AssignmentTarget.GetName());
 		FNiagaraParameterHandle RenamedTargetHandle(TargetHandle.GetNamespace(), NewName);
+
 		OwningAssignmentNode->AssignmentTarget.SetName(RenamedTargetHandle.GetParameterHandleString());
 		OwningAssignmentNode->RefreshFromExternalChanges();
 
@@ -719,6 +1101,20 @@ void UNiagaraStackFunctionInput::RenameInput(FName NewName)
 		AliasedInputParameterHandle = FNiagaraParameterHandle::CreateAliasedModuleParameterHandle(InputParameterHandle, OwningAssignmentNode.Get());
 		DisplayName = FText::FromName(InputParameterHandle.GetName());
 
+		if (IsRapidIterationCandidate())
+		{
+			FString OriginalName = RapidIterationParameterName;
+			FString NewNameForRIParam = CreateRapidIterationVariableName(AliasedInputParameterHandle.GetParameterHandleString().ToString());
+
+			FNiagaraVariable Temp(InputType, *OriginalName);
+			for (TWeakObjectPtr<UNiagaraScript> Script : AffectedScripts)
+			{
+				Script->RapidIterationParameters.RenameParameter(Temp, *NewNameForRIParam);
+			}
+			UE_LOG(LogNiagaraEditor, Log, TEXT("Renaming %s to %s"), *RapidIterationParameterName, *NewNameForRIParam);
+			RapidIterationParameterName = NewNameForRIParam;
+		}
+
 		UEdGraphPin* OverridePin = GetOverridePin();
 		if (OverridePin != nullptr)
 		{
@@ -736,13 +1132,13 @@ void UNiagaraStackFunctionInput::RenameInput(FName NewName)
 void UNiagaraStackFunctionInput::GetNamespacesForNewParameters(TArray<FName>& OutNamespacesForNewParameters) const
 {
 	UNiagaraNodeOutput* OutputNode = FNiagaraStackGraphUtilities::GetEmitterOutputNodeForStackNode(*OwningFunctionCallNode);
-	bool bIsSystemTransient = GetSystemViewModel()->GetSystemIsTransient();
+	bool bIsEditingSystem = GetSystemViewModel()->GetEditMode() == ENiagaraSystemViewModelEditMode::SystemAsset;
 
 	if (OutputNode->GetUsage() == ENiagaraScriptUsage::ParticleSpawnScript || OutputNode->GetUsage() == ENiagaraScriptUsage::ParticleUpdateScript)
 	{
 		OutNamespacesForNewParameters.Add(FNiagaraParameterHandle::ParticleAttributeNamespace);
 		OutNamespacesForNewParameters.Add(FNiagaraParameterHandle::EmitterNamespace);
-		if (bIsSystemTransient == false)
+		if (bIsEditingSystem)
 		{
 			OutNamespacesForNewParameters.Add(FNiagaraParameterHandle::SystemNamespace);
 			OutNamespacesForNewParameters.Add(FNiagaraParameterHandle::UserNamespace);
@@ -751,13 +1147,13 @@ void UNiagaraStackFunctionInput::GetNamespacesForNewParameters(TArray<FName>& Ou
 	else if (OutputNode->GetUsage() == ENiagaraScriptUsage::EmitterSpawnScript || OutputNode->GetUsage() == ENiagaraScriptUsage::EmitterUpdateScript)
 	{
 		OutNamespacesForNewParameters.Add(FNiagaraParameterHandle::EmitterNamespace);
-		if (bIsSystemTransient == false)
+		if (bIsEditingSystem)
 		{
 			OutNamespacesForNewParameters.Add(FNiagaraParameterHandle::SystemNamespace);
 			OutNamespacesForNewParameters.Add(FNiagaraParameterHandle::UserNamespace);
 		}
 	}
-	else if (OutputNode->GetUsage() == ENiagaraScriptUsage::SystemSpawnScript || OutputNode->GetUsage() == ENiagaraScriptUsage::ParticleUpdateScript && bIsSystemTransient == false)
+	else if (OutputNode->GetUsage() == ENiagaraScriptUsage::SystemSpawnScript || OutputNode->GetUsage() == ENiagaraScriptUsage::ParticleUpdateScript && bIsEditingSystem)
 	{
 		OutNamespacesForNewParameters.Add(FNiagaraParameterHandle::SystemNamespace);
 		OutNamespacesForNewParameters.Add(FNiagaraParameterHandle::UserNamespace);
@@ -783,6 +1179,16 @@ void UNiagaraStackFunctionInput::OnGraphChanged(const struct FEdGraphEditAction&
 	}
 }
 
+void UNiagaraStackFunctionInput::OnRapidIterationParametersChanged()
+{
+	bCanResetToBase.Reset();
+}
+
+void UNiagaraStackFunctionInput::OnScriptSourceChanged()
+{
+	bCanResetToBase.Reset();
+}
+
 UNiagaraNodeParameterMapSet* UNiagaraStackFunctionInput::GetOverrideNode() const
 {
 	if (OverrideNodeCache.IsSet() == false)
@@ -790,11 +1196,7 @@ UNiagaraNodeParameterMapSet* UNiagaraStackFunctionInput::GetOverrideNode() const
 		UNiagaraNodeParameterMapSet* OverrideNode = nullptr;
 		if (OwningFunctionCallNode.IsValid())
 		{
-			UEdGraphPin* ParameterMapInput = FNiagaraStackGraphUtilities::GetParameterMapInputPin(*OwningFunctionCallNode);
-			if (ParameterMapInput != nullptr && ParameterMapInput->LinkedTo.Num() == 1)
-			{
-				OverrideNode = Cast<UNiagaraNodeParameterMapSet>(ParameterMapInput->LinkedTo[0]->GetOwningNode());
-			}
+			OverrideNode = FNiagaraStackGraphUtilities::GetStackFunctionOverrideNode(*OwningFunctionCallNode);
 		}
 		OverrideNodeCache = OverrideNode;
 	}
@@ -807,27 +1209,7 @@ UNiagaraNodeParameterMapSet& UNiagaraStackFunctionInput::GetOrCreateOverrideNode
 	if (OverrideNode == nullptr)
 	{
 		TGuardValue<bool>(bUpdatingGraphDirectly, true);
-
-		UEdGraph* Graph = OwningModuleNode->GetGraph();
-		FGraphNodeCreator<UNiagaraNodeParameterMapSet> ParameterMapSetNodeCreator(*Graph);
-		OverrideNode = ParameterMapSetNodeCreator.CreateNode();
-		ParameterMapSetNodeCreator.Finalize();
-
-		UEdGraphPin* OverrideNodeInputPin = FNiagaraStackGraphUtilities::GetParameterMapInputPin(*OverrideNode);
-		UEdGraphPin* OverrideNodeOutputPin = FNiagaraStackGraphUtilities::GetParameterMapOutputPin(*OverrideNode);
-
-		UEdGraphPin* OwningFunctionCallInputPin = FNiagaraStackGraphUtilities::GetParameterMapInputPin(*OwningFunctionCallNode);
-		UEdGraphPin* PreviousStackNodeOutputPin = OwningFunctionCallInputPin->LinkedTo[0];
-
-		OwningFunctionCallInputPin->BreakAllPinLinks();
-		OwningFunctionCallInputPin->MakeLinkTo(OverrideNodeOutputPin);
-		for (UEdGraphPin* PreviousStackNodeOutputLinkedPin : PreviousStackNodeOutputPin->LinkedTo)
-		{
-			PreviousStackNodeOutputLinkedPin->MakeLinkTo(OverrideNodeOutputPin);
-		}
-		PreviousStackNodeOutputPin->BreakAllPinLinks();
-		PreviousStackNodeOutputPin->MakeLinkTo(OverrideNodeInputPin);
-
+		OverrideNode = &FNiagaraStackGraphUtilities::GetOrCreateStackFunctionOverrideNode(*OwningFunctionCallNode);
 		OverrideNodeCache = OverrideNode;
 	}
 	return *OverrideNode;
@@ -842,18 +1224,7 @@ UEdGraphPin* UNiagaraStackFunctionInput::GetOverridePin() const
 {
 	if (OverridePinCache.IsSet() == false)
 	{
-		OverridePinCache = nullptr;
-		UNiagaraNodeParameterMapSet* OverrideNode = GetOverrideNode();
-		if (OverrideNode != nullptr)
-		{
-			TArray<UEdGraphPin*> InputPins;
-			OverrideNode->GetInputPins(InputPins);
-			UEdGraphPin** OverridePinPtr = InputPins.FindByPredicate([&](const UEdGraphPin* Pin) { return Pin->PinName == AliasedInputParameterHandle.GetParameterHandleString(); });
-			if (OverridePinPtr != nullptr)
-			{
-				OverridePinCache = *OverridePinPtr;
-			}
-		}
+		OverridePinCache = FNiagaraStackGraphUtilities::GetStackFunctionInputOverridePin(*OwningFunctionCallNode, AliasedInputParameterHandle);
 	}
 	return OverridePinCache.GetValue();
 }
@@ -864,17 +1235,7 @@ UEdGraphPin& UNiagaraStackFunctionInput::GetOrCreateOverridePin()
 	if (OverridePin == nullptr)
 	{
 		TGuardValue<bool>(bUpdatingGraphDirectly, true);
-
-		UNiagaraNodeParameterMapSet& OverrideNode = GetOrCreateOverrideNode();
-		OverrideNode.Modify();
-
-		TArray<UEdGraphPin*> OverrideInputPins;
-		OverrideNode.GetInputPins(OverrideInputPins);
-
-		const UEdGraphSchema_Niagara* NiagaraSchema = GetDefault<UEdGraphSchema_Niagara>();
-		FEdGraphPinType PinType = NiagaraSchema->TypeDefinitionToPinType(InputType);
-		OverridePin = OverrideNode.CreatePin(EEdGraphPinDirection::EGPD_Input, PinType, AliasedInputParameterHandle.GetParameterHandleString(), OverrideInputPins.Num() - 1);
-
+		OverridePin = &FNiagaraStackGraphUtilities::GetOrCreateStackFunctionInputOverridePin(*OwningFunctionCallNode, AliasedInputParameterHandle, InputType);
 		OverridePinCache = OverridePin;
 	}
 	return *OverridePin;
@@ -885,7 +1246,7 @@ bool UNiagaraStackFunctionInput::TryGetCurrentLocalValue(TSharedPtr<FStructOnSco
 	if (InputType.IsDataInterface() == false && ValuePin.LinkedTo.Num() == 0)
 	{
 		const UEdGraphSchema_Niagara* NiagaraSchema = GetDefault<UEdGraphSchema_Niagara>();
-		FNiagaraVariable ValueVariable = NiagaraSchema->PinToNiagaraVariable(&ValuePin);
+		FNiagaraVariable ValueVariable = NiagaraSchema->PinToNiagaraVariable(&ValuePin, true);
 		if (OldValueToReuse.IsValid() && OldValueToReuse->GetStruct() == ValueVariable.GetType().GetStruct())
 		{
 			LocalValue = OldValueToReuse;
@@ -894,7 +1255,35 @@ bool UNiagaraStackFunctionInput::TryGetCurrentLocalValue(TSharedPtr<FStructOnSco
 		{
 			LocalValue = MakeShared<FStructOnScope>(ValueVariable.GetType().GetStruct());
 		}
-		ValueVariable.CopyTo(LocalValue->GetStructMemory());
+
+		if (IsRapidIterationCandidate())
+		{
+			ValueVariable.SetName(*RapidIterationParameterName);
+
+			int32 FoundOffset = SourceScript->RapidIterationParameters.IndexOf(ValueVariable);
+			if (FoundOffset == INDEX_NONE)
+			{
+				// If the source script didn't have a value for this variable, update all affected scripts
+				// with the default value from the function.
+				UpdateRapidIterationParametersForAffectedScripts(ValueVariable.GetData());
+			}
+
+			FoundOffset = SourceScript->RapidIterationParameters.IndexOf(ValueVariable);
+			if (FoundOffset != INDEX_NONE)
+			{
+				FMemory::Memcpy(LocalValue->GetStructMemory(), SourceScript->RapidIterationParameters.GetParameterData(FoundOffset), ValueVariable.GetSizeInBytes());
+			}
+			else
+			{
+				// How would we ever get here in practice, we just added above?
+				ensure(false);
+				return false;
+			}
+		}
+		else
+		{
+			ValueVariable.CopyTo(LocalValue->GetStructMemory());
+		}
 		return true;
 	}
 	return false;
@@ -914,7 +1303,7 @@ bool UNiagaraStackFunctionInput::TryGetCurrentDataValue(FDataValues& DataValues,
 				{
 					if (InputNode->Usage == ENiagaraInputNodeUsage::Parameter)
 					{
-						DataValueObject = InputNode->DataInterface;
+						DataValueObject = InputNode->GetDataInterface();
 					}
 				}
 			}
@@ -925,9 +1314,9 @@ bool UNiagaraStackFunctionInput::TryGetCurrentDataValue(FDataValues& DataValues,
 		if (DefaultValuePin.LinkedTo.Num() == 1)
 		{
 			UNiagaraNodeInput* InputNode = Cast<UNiagaraNodeInput>(DefaultValuePin.LinkedTo[0]->GetOwningNode());
-			if (InputNode != nullptr && InputNode->Usage == ENiagaraInputNodeUsage::Parameter && InputNode->DataInterface != nullptr)
+			if (InputNode != nullptr && InputNode->Usage == ENiagaraInputNodeUsage::Parameter && InputNode->GetDataInterface() != nullptr)
 			{
-				DefaultDataValueObject = InputNode->DataInterface;
+				DefaultDataValueObject = InputNode->GetDataInterface();
 				DefaultDataValueOwner = FDataValues::EDefaultValueOwner::FunctionOwned;
 			}
 		}
@@ -1020,68 +1409,15 @@ bool UNiagaraStackFunctionInput::TryGetCurrentDynamicValue(TWeakObjectPtr<UNiaga
 	return false;
 }
 
-void UNiagaraStackFunctionInput::RemoveAllNodesConnectedToOverridePin(UEdGraphPin& OverridePin, UNiagaraStackFunctionInput* OwningInput)
+void UNiagaraStackFunctionInput::RemoveNodesForOverridePin(UEdGraphPin& OverridePin)
 {
-	if (OverridePin.LinkedTo.Num() == 1)
+	TArray<TWeakObjectPtr<UNiagaraDataInterface>> RemovedDataObjects;
+	FNiagaraStackGraphUtilities::RemoveNodesForStackFunctionInputOverridePin(OverridePin, RemovedDataObjects);
+	for (TWeakObjectPtr<UNiagaraDataInterface> RemovedDataObject : RemovedDataObjects)
 	{
-		UEdGraphNode* OverrideValueNode = OverridePin.LinkedTo[0]->GetOwningNode();
-		UEdGraph* Graph = OverrideValueNode->GetGraph();
-		if (OverrideValueNode->IsA<UNiagaraNodeInput>() || OverrideValueNode->IsA<UNiagaraNodeParameterMapGet>())
-		{
-			Graph->RemoveNode(OverrideValueNode);
-			if (OwningInput != nullptr && OverrideValueNode->IsA<UNiagaraNodeInput>())
-			{
-				// Input nodes are used to Data Interface objects so we need to call the the data object modified delegate
-				// when removing them.
-				OwningInput->OnDataObjectModified().Broadcast(nullptr);
-			}
-		}
-		else if (OverrideValueNode->IsA<UNiagaraNodeFunctionCall>())
+		if (RemovedDataObject.IsValid())
 		{
-			UNiagaraNodeFunctionCall* DynamicInputNode = CastChecked<UNiagaraNodeFunctionCall>(OverrideValueNode);
-			UEdGraphPin* DynamicInputNodeInputPin = FNiagaraStackGraphUtilities::GetParameterMapInputPin(*DynamicInputNode);
-			UNiagaraNodeParameterMapSet* DynamicInputNodeOverrideNode = Cast<UNiagaraNodeParameterMapSet>(DynamicInputNodeInputPin->LinkedTo[0]->GetOwningNode());
-			if (DynamicInputNodeOverrideNode != nullptr)
-			{
-				TArray<UEdGraphPin*> InputPins;
-				DynamicInputNodeOverrideNode->GetInputPins(InputPins);
-				for (UEdGraphPin* InputPin : InputPins)
-				{
-					if (InputPin->PinName.ToString().StartsWith(DynamicInputNode->GetFunctionName()))
-					{
-						RemoveAllNodesConnectedToOverridePin(*InputPin, OwningInput);
-						DynamicInputNodeOverrideNode->RemovePin(InputPin);
-					}
-				}
-
-				TArray<UEdGraphPin*> NewInputPins;
-				DynamicInputNodeOverrideNode->GetInputPins(NewInputPins);
-				if (NewInputPins.Num() == 2)
-				{
-					// If there are only 2 pins, they are the parameter map input and the add pin, so the dynamic input's override node can be removed.  This
-					// not always be the case when removing dynamic input nodes because they share the same override node.
-					UEdGraphPin* InputPin = FNiagaraStackGraphUtilities::GetParameterMapInputPin(*DynamicInputNodeOverrideNode);
-					UEdGraphPin* OutputPin = FNiagaraStackGraphUtilities::GetParameterMapOutputPin(*DynamicInputNodeOverrideNode);
-
-					if (ensureMsgf(InputPin != nullptr && InputPin->LinkedTo.Num() == 1 &&
-						OutputPin != nullptr && OutputPin->LinkedTo.Num() == 2, 
-						TEXT("Invalid Stack - Dynamic input node override node not connected correctly.")))
-					{
-						// The DynamicInputOverrideNode will have a single input which is the previous module or override map set, and
-						// two output links, one to the dynamic input node and one to the next override map set.
-						UEdGraphPin* LinkedInputPin = InputPin->LinkedTo[0];
-						UEdGraphPin* LinkedOutputPin = OutputPin->LinkedTo[0]->GetOwningNode() != DynamicInputNode
-							? OutputPin->LinkedTo[0]
-							: OutputPin->LinkedTo[1];
-						InputPin->BreakAllPinLinks();
-						OutputPin->BreakAllPinLinks();
-						LinkedInputPin->MakeLinkTo(LinkedOutputPin);
-						Graph->RemoveNode(DynamicInputNodeOverrideNode);
-					}
-				}
-			}
-
-			Graph->RemoveNode(DynamicInputNode);
+			OnDataObjectModified().Broadcast(RemovedDataObject.Get());
 		}
 	}
 }
diff --git a/Engine/Plugins/FX/Niagara/Source/NiagaraEditor/Private/ViewModels/Stack/NiagaraStackFunctionInputCollection.cpp b/Engine/Plugins/FX/Niagara/Source/NiagaraEditor/Private/ViewModels/Stack/NiagaraStackFunctionInputCollection.cpp
--- a/Engine/Plugins/FX/Niagara/Source/NiagaraEditor/Private/ViewModels/Stack/NiagaraStackFunctionInputCollection.cpp
+++ b/Engine/Plugins/FX/Niagara/Source/NiagaraEditor/Private/ViewModels/Stack/NiagaraStackFunctionInputCollection.cpp
@@ -6,8 +6,13 @@
 #include "NiagaraEmitterEditorData.h"
 #include "EdGraphSchema_Niagara.h"
 #include "NiagaraStackGraphUtilities.h"
+#include "NiagaraNodeParameterMapSet.h"
+#include "NiagaraDataInterface.h"
 
 #include "EdGraph/EdGraphPin.h"
+#include "ScopedTransaction.h"
+
+#define LOCTEXT_NAMESPACE "UNiagaraStackFunctionInputCollection"
 
 UNiagaraStackFunctionInputCollection::UNiagaraStackFunctionInputCollection()
 	: ModuleNode(nullptr)
@@ -39,6 +44,10 @@ void UNiagaraStackFunctionInputCollection::Initialize(
 	ModuleNode = &InModuleNode;
 	InputFunctionCallNode = &InInputFunctionCallNode;
 	DisplayOptions = InDisplayOptions;
+	if (DisplayOptions.ChildFilter.IsBound())
+	{
+		AddChildFilter(DisplayOptions.ChildFilter);
+	}
 }
 
 FText UNiagaraStackFunctionInputCollection::GetDisplayName() const
@@ -61,6 +70,32 @@ bool UNiagaraStackFunctionInputCollection::GetShouldShowInStack() const
 	return DisplayOptions.bShouldShowInStack;
 }
 
+int32 UNiagaraStackFunctionInputCollection::GetErrorCount() const
+{
+	return Errors.Num();
+}
+
+bool UNiagaraStackFunctionInputCollection::GetErrorFixable(int32 ErrorIdx) const
+{
+	return Errors[ErrorIdx].Fix.IsBound();
+}
+
+bool UNiagaraStackFunctionInputCollection::TryFixError(int32 ErrorIdx)
+{
+	Errors[ErrorIdx].Fix.Execute();
+	return true;
+}
+
+FText UNiagaraStackFunctionInputCollection::GetErrorText(int32 ErrorIdx) const
+{
+	return Errors[ErrorIdx].ErrorText;
+}
+
+FText UNiagaraStackFunctionInputCollection::GetErrorSummaryText(int32 ErrorIdx) const
+{
+	return Errors[ErrorIdx].ErrorSummaryText;
+}
+
 UNiagaraStackFunctionInputCollection::FOnInputPinnedChanged& UNiagaraStackFunctionInputCollection::OnInputPinnedChanged()
 {
 	return InputPinnedChangedDelegate;
@@ -71,34 +106,111 @@ void UNiagaraStackFunctionInputCollection::RefreshChildrenInternal(const TArray<
 	TArray<const UEdGraphPin*> InputPins;
 	FNiagaraStackGraphUtilities::GetStackFunctionInputPins(*InputFunctionCallNode, InputPins, FNiagaraStackGraphUtilities::ENiagaraGetStackFunctionInputPinsOptions::ModuleInputsOnly);
 	const UEdGraphSchema_Niagara* NiagaraSchema = GetDefault<UEdGraphSchema_Niagara>();
+	TArray<FName> ValidAliasedInputNames;
+	TArray<FName> ProcessedInputNames;
+	TArray<FName> DuplicateInputNames;
+	TArray<const UEdGraphPin*> PinsWithInvalidTypes;
 	for(const UEdGraphPin* InputPin : InputPins)
 	{
-		UNiagaraStackFunctionInput* Input = nullptr;
-		for (UNiagaraStackEntry* CurrentChild : CurrentChildren)
+		if (ProcessedInputNames.Contains(InputPin->PinName) == false)
 		{
-			UNiagaraStackFunctionInput* ChildInput = CastChecked<UNiagaraStackFunctionInput>(CurrentChild);
-			if (ChildInput->GetInputParameterHandle().GetParameterHandleString() == InputPin->PinName)
+			UNiagaraStackFunctionInput* Input = FindCurrentChildOfTypeByPredicate<UNiagaraStackFunctionInput>(CurrentChildren,
+				[=](UNiagaraStackFunctionInput* CurrentFunctionInput) { return CurrentFunctionInput->GetInputParameterHandle().GetParameterHandleString() == InputPin->PinName; });
+			
+			if (Input == nullptr)
 			{
-				Input = ChildInput;
-				break;
+				FNiagaraTypeDefinition InputType = NiagaraSchema->PinToTypeDefinition(InputPin);
+				if (InputType.IsValid())
+				{
+					Input = NewObject<UNiagaraStackFunctionInput>(this);
+					Input->Initialize(GetSystemViewModel(), GetEmitterViewModel(), *StackEditorData, *ModuleNode, *InputFunctionCallNode, InputPin->PinName, NiagaraSchema->PinToTypeDefinition(InputPin));
+					Input->SetItemIndentLevel(DisplayOptions.ChildItemIndentLevel);
+					Input->OnPinnedChanged().AddUObject(this, &UNiagaraStackFunctionInputCollection::ChildPinnedChanged);
+				}
+				else
+				{
+					PinsWithInvalidTypes.Add(InputPin);
+				}
 			}
+
+			if (Input != nullptr)
+			{
+				NewChildren.Add(Input);
+			}
+			
+			ValidAliasedInputNames.Add(FNiagaraParameterHandle::CreateAliasedModuleParameterHandle(FNiagaraParameterHandle(InputPin->PinName), InputFunctionCallNode).GetParameterHandleString());
+			ProcessedInputNames.Add(InputPin->PinName);
 		}
-		if (Input == nullptr)
+		else
 		{
-			Input = NewObject<UNiagaraStackFunctionInput>(this);
-			Input->Initialize(GetSystemViewModel(), GetEmitterViewModel(), *StackEditorData, *ModuleNode, *InputFunctionCallNode, InputPin->PinName, NiagaraSchema->PinToTypeDefinition(InputPin));
-			Input->SetItemIndentLevel(DisplayOptions.ChildItemIndentLevel);
-			Input->OnPinnedChanged().AddUObject(this, &UNiagaraStackFunctionInputCollection::ChildPinnedChanged);
+			DuplicateInputNames.AddUnique(InputPin->PinName);
 		}
+	}
 
-		if (DisplayOptions.ChildFilter.IsBound() == false || DisplayOptions.ChildFilter.Execute(Input))
+	// Try to find function input overrides which are no longer valid so we can generate errors for them.
+	Errors.Empty();
+	UNiagaraNodeParameterMapSet* OverrideNode = FNiagaraStackGraphUtilities::GetStackFunctionOverrideNode(*InputFunctionCallNode);
+	if (OverrideNode != nullptr)
+	{
+		TArray<UEdGraphPin*> OverridePins;
+		OverrideNode->GetInputPins(OverridePins);
+		for (UEdGraphPin* OverridePin : OverridePins)
 		{
-			NewChildren.Add(Input);
+			// If the pin isn't in the misc category for the add pin, and not the parameter map pin, and it's for this function call,
+			// check to see if it's in the list of valid input names, and if not generate an error.
+			if (OverridePin->PinType.PinCategory != UEdGraphSchema_Niagara::PinCategoryMisc &&
+				OverridePin->PinType.PinSubCategoryObject != FNiagaraTypeDefinition::GetParameterMapStruct() &&
+				FNiagaraParameterHandle(OverridePin->PinName).GetNamespace().ToString() == InputFunctionCallNode->GetFunctionName() &&
+				ValidAliasedInputNames.Contains(OverridePin->PinName) == false)
+			{
+				FError InvalidInputError;
+				InvalidInputError.ErrorSummaryText = FText::Format(LOCTEXT("InvalidInputSummaryFormat", "Invalid Input: {0}"), FText::FromString(OverridePin->PinName.ToString()));
+				InvalidInputError.ErrorText = FText::Format(LOCTEXT("InvalidInputFormat", "The input {0} was previously overriden but is no longer exposed by the function {1}.\nPress the fix button to remove this unused override data,\nor check the function definition to see why this input is no longer exposed."),
+					FText::FromString(OverridePin->PinName.ToString()), FText::FromString(InputFunctionCallNode->GetFunctionName()));
+				InvalidInputError.Fix.BindLambda([=]()
+				{
+					FScopedTransaction ScopedTransaction(LOCTEXT("RemoveInvalidInputTransaction", "Remove invalid input override."));
+					TArray<TWeakObjectPtr<UNiagaraDataInterface>> RemovedDataObjects;
+					FNiagaraStackGraphUtilities::RemoveNodesForStackFunctionInputOverridePin(*OverridePin, RemovedDataObjects);
+					for (TWeakObjectPtr<UNiagaraDataInterface> RemovedDataObject : RemovedDataObjects)
+					{
+						if (RemovedDataObject.IsValid())
+						{
+							OnDataObjectModified().Broadcast(RemovedDataObject.Get());
+						}
+					}
+					OverridePin->GetOwningNode()->RemovePin(OverridePin);
+				});
+				Errors.Add(InvalidInputError);
+			}
 		}
 	}
+
+	for (const FName& DuplicateInputName : DuplicateInputNames)
+	{
+		FError DuplicateInputError;
+
+		FString DuplicateInputNameString = DuplicateInputName.ToString();
+
+		DuplicateInputError.ErrorSummaryText = FText::Format(LOCTEXT("DuplicateInputSummaryFormat", "Duplicate Input: {0}"), FText::FromString(DuplicateInputNameString));
+		DuplicateInputError.ErrorText = FText::Format(LOCTEXT("DuplicateInputFormat", "There are multiple inputs with the same name {0}, but different types exposed by the function {1}.\nThis is not suppored and must be fixed in the script that defines this function."),
+			FText::FromString(DuplicateInputNameString), FText::FromString(InputFunctionCallNode->GetFunctionName()));
+		Errors.Add(DuplicateInputError);
+	}
+
+	for (const UEdGraphPin* PinWithInvalidType : PinsWithInvalidTypes)
+	{
+		FError InputWithInvalidTypeError;
+		InputWithInvalidTypeError.ErrorSummaryText = FText::Format(LOCTEXT("InputWithInvalidTypeSummaryFormat", "Input has an invalid type: {0}"), FText::FromString(PinWithInvalidType->PinName.ToString()));
+		InputWithInvalidTypeError.ErrorText = FText::Format(LOCTEXT("InputWithInvalidTypeFormat", "The input {0} on function {1} has a type which is invalid.\nThe type of this input likely doesn't exist anymore.\nThis input must be fixed in the script before this module can be used."),
+			FText::FromString(PinWithInvalidType->PinName.ToString()), FText::FromString(InputFunctionCallNode->GetFunctionName()));
+		Errors.Add(InputWithInvalidTypeError);
+	}
 }
 
 void UNiagaraStackFunctionInputCollection::ChildPinnedChanged()
 {
 	InputPinnedChangedDelegate.Broadcast();
-}
\ No newline at end of file
+}
+
+#undef LOCTEXT_NAMESPACE
\ No newline at end of file
diff --git a/Engine/Plugins/FX/Niagara/Source/NiagaraEditor/Private/ViewModels/Stack/NiagaraStackGraphUtilities.cpp b/Engine/Plugins/FX/Niagara/Source/NiagaraEditor/Private/ViewModels/Stack/NiagaraStackGraphUtilities.cpp
--- a/Engine/Plugins/FX/Niagara/Source/NiagaraEditor/Private/ViewModels/Stack/NiagaraStackGraphUtilities.cpp
+++ b/Engine/Plugins/FX/Niagara/Source/NiagaraEditor/Private/ViewModels/Stack/NiagaraStackGraphUtilities.cpp
@@ -9,18 +9,24 @@
 #include "NiagaraNodeOutput.h"
 #include "NiagaraNodeInput.h"
 #include "NiagaraNodeFunctionCall.h"
+#include "NiagaraNodeAssignment.h"
+#include "NiagaraMaterialParameterNode.h"
+#include "NiagaraNodeParameterMapGet.h"
+#include "NiagaraNodeParameterMapSet.h"
+#include "NiagaraScriptSource.h"
 #include "EdGraphSchema_Niagara.h"
 #include "NiagaraStackEntry.h"
 #include "NiagaraStackFunctionInputCollection.h"
 #include "NiagaraStackFunctionInput.h"
 #include "NiagaraSystem.h"
 #include "NiagaraEditorModule.h"
-
+#include "NiagaraEditorUtilities.h"
 #include "EdGraph/EdGraph.h"
 #include "EdGraph/EdGraphNode.h"
 #include "EdGraph/EdGraphPin.h"
+#include "NiagaraEmitterViewModel.h"
 
-DECLARE_CYCLE_STAT(TEXT("StackGraphUtilities - RelayoutGraph"), STAT_NiagaraEditor_StackGraphUtilities_RelayoutGraph, STATGROUP_NiagaraEditor);
+DECLARE_CYCLE_STAT(TEXT("Niagara - StackGraphUtilities - RelayoutGraph"), STAT_NiagaraEditor_StackGraphUtilities_RelayoutGraph, STATGROUP_NiagaraEditor);
 
 #define LOCTEXT_NAMESPACE "NiagaraStackGraphUtilities"
 
@@ -305,6 +311,40 @@ UNiagaraNodeInput* FNiagaraStackGraphUtilities::GetEmitterInputNodeForStackNode(
 	return nullptr;
 }
 
+void GetGroupNodesRecursive(const TArray<UNiagaraNode*>& CurrentStartNodes, UNiagaraNode* EndNode, TArray<UNiagaraNode*>& OutAllNodes)
+{
+	for (UNiagaraNode* CurrentStartNode : CurrentStartNodes)
+	{
+		if (OutAllNodes.Contains(CurrentStartNode) == false)
+		{
+			OutAllNodes.Add(CurrentStartNode);
+			if (CurrentStartNode != EndNode)
+			{
+				TArray<UNiagaraNode*> LinkedNodes;
+				TArray<UEdGraphPin*> OutputPins;
+				CurrentStartNode->GetOutputPins(OutputPins);
+				for (UEdGraphPin* OutputPin : OutputPins)
+				{
+					for (UEdGraphPin* LinkedPin : OutputPin->LinkedTo)
+					{
+						UNiagaraNode* LinkedNode = Cast<UNiagaraNode>(LinkedPin->GetOwningNode());
+						if (LinkedNode != nullptr)
+						{
+							LinkedNodes.Add(LinkedNode);
+						}
+					}
+				}
+				GetGroupNodesRecursive(LinkedNodes, EndNode, OutAllNodes);
+			}
+		}
+	}
+}
+
+void FNiagaraStackGraphUtilities::FStackNodeGroup::GetAllNodesInGroup(TArray<UNiagaraNode*>& OutAllNodes) const
+{
+	GetGroupNodesRecursive(StartNodes, EndNode, OutAllNodes);
+}
+
 void FNiagaraStackGraphUtilities::GetStackNodeGroups(UNiagaraNode& StackNode, TArray<FNiagaraStackGraphUtilities::FStackNodeGroup>& OutStackNodeGroups)
 {
 	UNiagaraNodeOutput* OutputNode = GetEmitterOutputNodeForStackNode(StackNode);
@@ -386,7 +426,7 @@ void FNiagaraStackGraphUtilities::InitializeDataInterfaceInputs(TSharedRef<FNiag
 	FunctionInputCollection->RefreshChildren();
 
 	TArray<UNiagaraStackEntry*> Children;
-	FunctionInputCollection->GetChildren(Children);
+	FunctionInputCollection->GetUnfilteredChildren(Children);
 	for (UNiagaraStackEntry* Child : Children)
 	{
 		UNiagaraStackFunctionInput* FunctionInput = Cast<UNiagaraStackFunctionInput>(Child);
@@ -409,9 +449,10 @@ FString FNiagaraStackGraphUtilities::GenerateStackModuleEditorDataKey(UNiagaraNo
 	return ModuleNode.GetFunctionName();
 }
 
-void FNiagaraStackGraphUtilities::GetStackFunctionInputPins(UNiagaraNodeFunctionCall& FunctionCallNode, TArray<const UEdGraphPin*>& OutInputPins, ENiagaraGetStackFunctionInputPinsOptions Options)
+void FNiagaraStackGraphUtilities::GetStackFunctionInputPins(UNiagaraNodeFunctionCall& FunctionCallNode, TArray<const UEdGraphPin*>& OutInputPins, ENiagaraGetStackFunctionInputPinsOptions Options, bool bIgnoreDisabled)
 {
 	FNiagaraParameterMapHistoryBuilder Builder;
+	Builder.SetIgnoreDisabled(bIgnoreDisabled);
 	FunctionCallNode.BuildParameterMapHistory(Builder, false);
 	
 	if (Builder.Histories.Num() == 1)
@@ -434,9 +475,408 @@ void FNiagaraStackGraphUtilities::GetStackFunctionInputPins(UNiagaraNodeFunction
 	}
 }
 
-bool FNiagaraStackGraphUtilities::ValidateGraphForOutput(UNiagaraGraph& NiagaraGraph, ENiagaraScriptUsage ScriptUsage, int32 ScriptOccurrence, FText& ErrorMessage)
+UNiagaraNodeParameterMapSet* FNiagaraStackGraphUtilities::GetStackFunctionOverrideNode(UNiagaraNodeFunctionCall& FunctionCallNode)
+{
+	UEdGraphPin* ParameterMapInput = FNiagaraStackGraphUtilities::GetParameterMapInputPin(FunctionCallNode);
+	if (ParameterMapInput != nullptr && ParameterMapInput->LinkedTo.Num() == 1)
+	{
+		return Cast<UNiagaraNodeParameterMapSet>(ParameterMapInput->LinkedTo[0]->GetOwningNode());
+	}
+	return nullptr;
+}
+
+UNiagaraNodeParameterMapSet& FNiagaraStackGraphUtilities::GetOrCreateStackFunctionOverrideNode(UNiagaraNodeFunctionCall& StackFunctionCall)
+{
+	UNiagaraNodeParameterMapSet* OverrideNode = GetStackFunctionOverrideNode(StackFunctionCall);
+	if (OverrideNode == nullptr)
+	{
+		UEdGraph* Graph = StackFunctionCall.GetGraph();
+		FGraphNodeCreator<UNiagaraNodeParameterMapSet> ParameterMapSetNodeCreator(*Graph);
+		OverrideNode = ParameterMapSetNodeCreator.CreateNode();
+		ParameterMapSetNodeCreator.Finalize();
+		OverrideNode->SetEnabledState(StackFunctionCall.GetDesiredEnabledState(), StackFunctionCall.HasUserSetTheEnabledState());
+
+		UEdGraphPin* OverrideNodeInputPin = FNiagaraStackGraphUtilities::GetParameterMapInputPin(*OverrideNode);
+		UEdGraphPin* OverrideNodeOutputPin = FNiagaraStackGraphUtilities::GetParameterMapOutputPin(*OverrideNode);
+
+		UEdGraphPin* OwningFunctionCallInputPin = FNiagaraStackGraphUtilities::GetParameterMapInputPin(StackFunctionCall);
+		UEdGraphPin* PreviousStackNodeOutputPin = OwningFunctionCallInputPin->LinkedTo[0];
+
+		OwningFunctionCallInputPin->BreakAllPinLinks();
+		OwningFunctionCallInputPin->MakeLinkTo(OverrideNodeOutputPin);
+		for (UEdGraphPin* PreviousStackNodeOutputLinkedPin : PreviousStackNodeOutputPin->LinkedTo)
+		{
+			PreviousStackNodeOutputLinkedPin->MakeLinkTo(OverrideNodeOutputPin);
+		}
+		PreviousStackNodeOutputPin->BreakAllPinLinks();
+		PreviousStackNodeOutputPin->MakeLinkTo(OverrideNodeInputPin);
+	}
+	return *OverrideNode;
+}
+
+UEdGraphPin* FNiagaraStackGraphUtilities::GetStackFunctionInputOverridePin(UNiagaraNodeFunctionCall& StackFunctionCall, FNiagaraParameterHandle AliasedInputParameterHandle)
+{
+	UNiagaraNodeParameterMapSet* OverrideNode = GetStackFunctionOverrideNode(StackFunctionCall);
+	if (OverrideNode != nullptr)
+	{
+		TArray<UEdGraphPin*> InputPins;
+		OverrideNode->GetInputPins(InputPins);
+		UEdGraphPin** OverridePinPtr = InputPins.FindByPredicate([&](const UEdGraphPin* Pin) { return Pin->PinName == AliasedInputParameterHandle.GetParameterHandleString(); });
+		if (OverridePinPtr != nullptr)
+		{
+			return *OverridePinPtr;
+		}
+	}
+	return nullptr;
+}
+
+UEdGraphPin& FNiagaraStackGraphUtilities::GetOrCreateStackFunctionInputOverridePin(UNiagaraNodeFunctionCall& StackFunctionCall, FNiagaraParameterHandle AliasedInputParameterHandle, FNiagaraTypeDefinition InputType)
+{
+	UEdGraphPin* OverridePin = GetStackFunctionInputOverridePin(StackFunctionCall, AliasedInputParameterHandle);
+	if (OverridePin == nullptr)
+	{
+		UNiagaraNodeParameterMapSet& OverrideNode = GetOrCreateStackFunctionOverrideNode(StackFunctionCall);
+		OverrideNode.Modify();
+
+		TArray<UEdGraphPin*> OverrideInputPins;
+		OverrideNode.GetInputPins(OverrideInputPins);
+
+		const UEdGraphSchema_Niagara* NiagaraSchema = GetDefault<UEdGraphSchema_Niagara>();
+		FEdGraphPinType PinType = NiagaraSchema->TypeDefinitionToPinType(InputType);
+		OverridePin = OverrideNode.CreatePin(EEdGraphPinDirection::EGPD_Input, PinType, AliasedInputParameterHandle.GetParameterHandleString(), OverrideInputPins.Num() - 1);
+	}
+	return *OverridePin;
+}
+
+void FNiagaraStackGraphUtilities::RemoveNodesForStackFunctionInputOverridePin(UEdGraphPin& StackFunctionInputOverridePin)
+{
+	TArray<TWeakObjectPtr<UNiagaraDataInterface>> RemovedDataObjects;
+	RemoveNodesForStackFunctionInputOverridePin(StackFunctionInputOverridePin, RemovedDataObjects);
+}
+
+void FNiagaraStackGraphUtilities::RemoveNodesForStackFunctionInputOverridePin(UEdGraphPin& StackFunctionInputOverridePin, TArray<TWeakObjectPtr<UNiagaraDataInterface>>& OutRemovedDataObjects)
+{
+	if (StackFunctionInputOverridePin.LinkedTo.Num() == 1)
+	{
+		UEdGraphNode* OverrideValueNode = StackFunctionInputOverridePin.LinkedTo[0]->GetOwningNode();
+		UEdGraph* Graph = OverrideValueNode->GetGraph();
+		if (OverrideValueNode->IsA<UNiagaraNodeInput>() || OverrideValueNode->IsA<UNiagaraNodeParameterMapGet>())
+		{
+			UNiagaraNodeInput* InputNode = Cast<UNiagaraNodeInput>(OverrideValueNode);
+			if (InputNode != nullptr && InputNode->GetDataInterface() != nullptr)
+			{
+				OutRemovedDataObjects.Add(InputNode->GetDataInterface());
+			}
+			Graph->RemoveNode(OverrideValueNode);
+		}
+		else if (OverrideValueNode->IsA<UNiagaraNodeFunctionCall>())
+		{
+			UNiagaraNodeFunctionCall* DynamicInputNode = CastChecked<UNiagaraNodeFunctionCall>(OverrideValueNode);
+			UEdGraphPin* DynamicInputNodeInputPin = FNiagaraStackGraphUtilities::GetParameterMapInputPin(*DynamicInputNode);
+			UNiagaraNodeParameterMapSet* DynamicInputNodeOverrideNode = Cast<UNiagaraNodeParameterMapSet>(DynamicInputNodeInputPin->LinkedTo[0]->GetOwningNode());
+			if (DynamicInputNodeOverrideNode != nullptr)
+			{
+				TArray<UEdGraphPin*> InputPins;
+				DynamicInputNodeOverrideNode->GetInputPins(InputPins);
+				for (UEdGraphPin* InputPin : InputPins)
+				{
+					if (InputPin->PinName.ToString().StartsWith(DynamicInputNode->GetFunctionName()))
+					{
+						RemoveNodesForStackFunctionInputOverridePin(*InputPin, OutRemovedDataObjects);
+						DynamicInputNodeOverrideNode->RemovePin(InputPin);
+					}
+				}
+
+				TArray<UEdGraphPin*> NewInputPins;
+				DynamicInputNodeOverrideNode->GetInputPins(NewInputPins);
+				if (NewInputPins.Num() == 2)
+				{
+					// If there are only 2 pins, they are the parameter map input and the add pin, so the dynamic input's override node can be removed.  This
+					// not always be the case when removing dynamic input nodes because they share the same override node.
+					UEdGraphPin* InputPin = FNiagaraStackGraphUtilities::GetParameterMapInputPin(*DynamicInputNodeOverrideNode);
+					UEdGraphPin* OutputPin = FNiagaraStackGraphUtilities::GetParameterMapOutputPin(*DynamicInputNodeOverrideNode);
+
+					if (ensureMsgf(InputPin != nullptr && InputPin->LinkedTo.Num() == 1 &&
+						OutputPin != nullptr && OutputPin->LinkedTo.Num() == 2,
+						TEXT("Invalid Stack - Dynamic input node override node not connected correctly.")))
+					{
+						// The DynamicInputOverrideNode will have a single input which is the previous module or override map set, and
+						// two output links, one to the dynamic input node and one to the next override map set.
+						UEdGraphPin* LinkedInputPin = InputPin->LinkedTo[0];
+						UEdGraphPin* LinkedOutputPin = OutputPin->LinkedTo[0]->GetOwningNode() != DynamicInputNode
+							? OutputPin->LinkedTo[0]
+							: OutputPin->LinkedTo[1];
+						InputPin->BreakAllPinLinks();
+						OutputPin->BreakAllPinLinks();
+						LinkedInputPin->MakeLinkTo(LinkedOutputPin);
+						Graph->RemoveNode(DynamicInputNodeOverrideNode);
+					}
+				}
+			}
+
+			Graph->RemoveNode(DynamicInputNode);
+		}
+	}
+}
+
+void FNiagaraStackGraphUtilities::SetLinkedValueHandleForFunctionInput(UEdGraphPin& OverridePin, FNiagaraParameterHandle LinkedParameterHandle)
+{
+	checkf(OverridePin.LinkedTo.Num() == 0, TEXT("Can't set a linked value handle when the override pin already has a value."));
+
+	UNiagaraNodeParameterMapSet* OverrideNode = CastChecked<UNiagaraNodeParameterMapSet>(OverridePin.GetOwningNode());
+	UEdGraph* Graph = OverrideNode->GetGraph();
+	FGraphNodeCreator<UNiagaraNodeParameterMapGet> GetNodeCreator(*Graph);
+	UNiagaraNodeParameterMapGet* GetNode = GetNodeCreator.CreateNode();
+	GetNodeCreator.Finalize();
+	GetNode->SetEnabledState(OverrideNode->GetDesiredEnabledState(), OverrideNode->HasUserSetTheEnabledState());
+
+	UEdGraphPin* GetInputPin = FNiagaraStackGraphUtilities::GetParameterMapInputPin(*GetNode);
+	checkf(GetInputPin != nullptr, TEXT("Parameter map get node was missing it's parameter map input pin."));
+
+	UEdGraphPin* OverrideNodeInputPin = FNiagaraStackGraphUtilities::GetParameterMapInputPin(*OverrideNode);
+	UEdGraphPin* PreviousStackNodeOutputPin = OverrideNodeInputPin->LinkedTo[0];
+	checkf(PreviousStackNodeOutputPin != nullptr, TEXT("Invalid Stack Graph - No previous stack node."));
+
+	const UEdGraphSchema_Niagara* NiagaraSchema = GetDefault<UEdGraphSchema_Niagara>();
+	FNiagaraTypeDefinition InputType = NiagaraSchema->PinToTypeDefinition(&OverridePin);
+	UEdGraphPin* GetOutputPin = GetNode->RequestNewTypedPin(EGPD_Output, InputType, LinkedParameterHandle.GetParameterHandleString());
+	GetInputPin->MakeLinkTo(PreviousStackNodeOutputPin);
+	GetOutputPin->MakeLinkTo(&OverridePin);
+}
+
+void FNiagaraStackGraphUtilities::SetDataValueObjectForFunctionInput(UEdGraphPin& OverridePin, UClass* DataObjectType, FString DataObjectName, UNiagaraDataInterface*& OutDataObject)
+{
+	checkf(OverridePin.LinkedTo.Num() == 0, TEXT("Can't set a data value when the override pin already has a value."));
+	checkf(DataObjectType->IsChildOf<UNiagaraDataInterface>(), TEXT("Can only set a function input to a data interface value object"));
+
+	UNiagaraNodeParameterMapSet* OverrideNode = CastChecked<UNiagaraNodeParameterMapSet>(OverridePin.GetOwningNode());
+	UEdGraph* Graph = OverrideNode->GetGraph();
+	FGraphNodeCreator<UNiagaraNodeInput> InputNodeCreator(*Graph);
+	UNiagaraNodeInput* InputNode = InputNodeCreator.CreateNode();
+	FNiagaraEditorUtilities::InitializeParameterInputNode(*InputNode, FNiagaraTypeDefinition(DataObjectType), CastChecked<UNiagaraGraph>(Graph), *DataObjectName);
+
+	OutDataObject = NewObject<UNiagaraDataInterface>(InputNode, DataObjectType, *DataObjectName, RF_Transactional);
+	InputNode->SetDataInterface(OutDataObject);
+
+	InputNodeCreator.Finalize();
+	FNiagaraStackGraphUtilities::ConnectPinToInputNode(OverridePin, *InputNode);
+}
+
+void FNiagaraStackGraphUtilities::SetDynamicInputForFunctionInput(UEdGraphPin& OverridePin, UNiagaraScript* DynamicInput, UNiagaraNodeFunctionCall*& OutDynamicInputFunctionCall)
+{
+	checkf(OverridePin.LinkedTo.Num() == 0, TEXT("Can't set a data value when the override pin already has a value."));
+
+	UNiagaraNodeParameterMapSet* OverrideNode = CastChecked<UNiagaraNodeParameterMapSet>(OverridePin.GetOwningNode());
+	UEdGraph* Graph = OverrideNode->GetGraph();
+	FGraphNodeCreator<UNiagaraNodeFunctionCall> FunctionCallNodeCreator(*Graph);
+	UNiagaraNodeFunctionCall* FunctionCallNode = FunctionCallNodeCreator.CreateNode();
+	FunctionCallNode->FunctionScript = DynamicInput;
+	FunctionCallNodeCreator.Finalize();
+	FunctionCallNode->SetEnabledState(OverrideNode->GetDesiredEnabledState(), OverrideNode->HasUserSetTheEnabledState());
+
+	UEdGraphPin* FunctionCallInputPin = FNiagaraStackGraphUtilities::GetParameterMapInputPin(*FunctionCallNode);
+	TArray<UEdGraphPin*> FunctionCallOutputPins;
+	FunctionCallNode->GetOutputPins(FunctionCallOutputPins);
+
+	const UEdGraphSchema_Niagara* NiagaraSchema = GetDefault<UEdGraphSchema_Niagara>();
+
+	FNiagaraTypeDefinition InputType = NiagaraSchema->PinToTypeDefinition(&OverridePin);
+	checkf(FunctionCallInputPin != nullptr, TEXT("Dynamic Input function call did not have a parameter map input pin."));
+	checkf(FunctionCallOutputPins.Num() == 1 && NiagaraSchema->PinToTypeDefinition(FunctionCallOutputPins[0]) == InputType, TEXT("Invalid Stack Graph - Dynamic Input function did not have the correct typed output pin"));
+
+	UEdGraphPin* OverrideNodeInputPin = FNiagaraStackGraphUtilities::GetParameterMapInputPin(*OverrideNode);
+	UEdGraphPin* PreviousStackNodeOutputPin = OverrideNodeInputPin->LinkedTo[0];
+	checkf(PreviousStackNodeOutputPin != nullptr, TEXT("Invalid Stack Graph - No previous stack node."));
+
+	FunctionCallInputPin->MakeLinkTo(PreviousStackNodeOutputPin);
+	FunctionCallOutputPins[0]->MakeLinkTo(&OverridePin);
+
+	OutDynamicInputFunctionCall = FunctionCallNode;
+}
+
+bool FNiagaraStackGraphUtilities::RemoveModuleFromStack(UNiagaraNodeFunctionCall& ModuleNode)
+{
+	TArray<FNiagaraStackGraphUtilities::FStackNodeGroup> StackNodeGroups;
+	FNiagaraStackGraphUtilities::GetStackNodeGroups(ModuleNode, StackNodeGroups);
+
+	int32 ModuleStackIndex = StackNodeGroups.IndexOfByPredicate([&](const FNiagaraStackGraphUtilities::FStackNodeGroup& StackNodeGroup) { return StackNodeGroup.EndNode == &ModuleNode; });
+	if (ModuleStackIndex == INDEX_NONE)
+	{
+		return false;
+	}
+
+	// Disconnect the group from the stack first to make collecting the nodes to remove easier.
+	FNiagaraStackGraphUtilities::DisconnectStackNodeGroup(StackNodeGroups[ModuleStackIndex], StackNodeGroups[ModuleStackIndex - 1], StackNodeGroups[ModuleStackIndex + 1]);
+
+	// Traverse all of the nodes in the group to find the nodes to remove.
+	FNiagaraStackGraphUtilities::FStackNodeGroup ModuleGroup = StackNodeGroups[ModuleStackIndex];
+	TArray<UNiagaraNode*> NodesToRemove;
+	TArray<UNiagaraNode*> NodesToCheck;
+	NodesToCheck.Add(ModuleGroup.EndNode);
+	while (NodesToCheck.Num() > 0)
+	{
+		UNiagaraNode* NodeToRemove = NodesToCheck[0];
+		NodesToCheck.RemoveAt(0);
+		NodesToRemove.Add(NodeToRemove);
+
+		TArray<UEdGraphPin*> InputPins;
+		NodeToRemove->GetInputPins(InputPins);
+		for (UEdGraphPin* InputPin : InputPins)
+		{
+			if (InputPin->LinkedTo.Num() == 1)
+			{
+				UNiagaraNode* LinkedNode = Cast<UNiagaraNode>(InputPin->LinkedTo[0]->GetOwningNode());
+				if (LinkedNode != nullptr)
+				{
+					NodesToCheck.Add(LinkedNode);
+				}
+			}
+		}
+	}
+
+	// Remove the nodes in the group from the graph.
+	UNiagaraGraph* Graph = ModuleNode.GetNiagaraGraph();
+	for (UNiagaraNode* NodeToRemove : NodesToRemove)
+	{
+		Graph->RemoveNode(NodeToRemove);
+	}
+
+	return true;
+}
+
+void ConnectModuleNode(UNiagaraNodeFunctionCall& ModuleNode, UNiagaraNodeOutput& TargetOutputNode, int32 TargetIndex)
+{
+	FNiagaraStackGraphUtilities::FStackNodeGroup ModuleGroup;
+	ModuleGroup.StartNodes.Add(&ModuleNode);
+	ModuleGroup.EndNode = &ModuleNode;
+
+	TArray<FNiagaraStackGraphUtilities::FStackNodeGroup> StackNodeGroups;
+	FNiagaraStackGraphUtilities::GetStackNodeGroups(TargetOutputNode, StackNodeGroups);
+	checkf(StackNodeGroups.Num() >= 2, TEXT("Stack graph is invalid, can not connect module"));
+
+	int32 InsertIndex;
+	if (TargetIndex != INDEX_NONE)
+	{
+		// The first stack node group is always the input node so we add one to the target module index to get the insertion index.
+		InsertIndex = FMath::Clamp(TargetIndex + 1, 1, StackNodeGroups.Num() - 1);
+	}
+	else
+	{
+		// If no insert index was specified, add the module at the end.
+		InsertIndex = StackNodeGroups.Num() - 1;
+	}
+
+	FNiagaraStackGraphUtilities::FStackNodeGroup& TargetInsertGroup = StackNodeGroups[InsertIndex];
+	FNiagaraStackGraphUtilities::FStackNodeGroup& TargetInsertPreviousGroup = StackNodeGroups[InsertIndex - 1];
+	FNiagaraStackGraphUtilities::ConnectStackNodeGroup(ModuleGroup, TargetInsertPreviousGroup, TargetInsertGroup);
+}
+
+bool FNiagaraStackGraphUtilities::FindScriptModulesInStack(FAssetData ModuleScriptAsset, UNiagaraNodeOutput& TargetOutputNode, TArray<UNiagaraNodeFunctionCall*> OutFunctionCalls)
+{
+	UNiagaraGraph* Graph = TargetOutputNode.GetNiagaraGraph();
+	TArray<UNiagaraNode*> Nodes;
+	Graph->BuildTraversal(Nodes, &TargetOutputNode);
+
+	OutFunctionCalls.Empty();
+	FString ModuleObjectName = ModuleScriptAsset.ObjectPath.ToString();
+	for (UEdGraphNode* Node : Nodes)
+	{
+		if (UNiagaraNodeFunctionCall* FunctionCallNode = Cast<UNiagaraNodeFunctionCall>(Node))
+		{
+			if (FunctionCallNode->FunctionScriptAssetObjectPath == ModuleScriptAsset.ObjectPath ||
+				(FunctionCallNode->FunctionScript != nullptr && FunctionCallNode->FunctionScript->GetPathName() == ModuleObjectName))
+			{
+				OutFunctionCalls.Add(FunctionCallNode);
+			}
+		}
+	}
+
+	return OutFunctionCalls.Num() > 0;
+}
+
+UNiagaraNodeFunctionCall* FNiagaraStackGraphUtilities::AddScriptModuleToStack(FAssetData ModuleScriptAsset, UNiagaraNodeOutput& TargetOutputNode, int32 TargetIndex)
 {
-	UNiagaraNodeOutput* OutputNode = NiagaraGraph.FindOutputNode(ScriptUsage, ScriptOccurrence);
+	UEdGraph* Graph = TargetOutputNode.GetGraph();
+	Graph->Modify();
+
+	FGraphNodeCreator<UNiagaraNodeFunctionCall> ModuleNodeCreator(*Graph);
+	UNiagaraNodeFunctionCall* NewModuleNode = ModuleNodeCreator.CreateNode();
+	NewModuleNode->FunctionScriptAssetObjectPath = ModuleScriptAsset.ObjectPath;
+	ModuleNodeCreator.Finalize();
+
+	ConnectModuleNode(*NewModuleNode, TargetOutputNode, TargetIndex);
+	return NewModuleNode;
+}
+
+UNiagaraNodeAssignment* FNiagaraStackGraphUtilities::AddParameterModuleToStack(FNiagaraVariable ParameterVariable, UNiagaraNodeOutput& TargetOutputNode, int32 TargetIndex)
+{
+	UEdGraph* Graph = TargetOutputNode.GetGraph();
+	Graph->Modify();
+
+	FGraphNodeCreator<UNiagaraNodeAssignment> AssignmentNodeCreator(*Graph);
+	UNiagaraNodeAssignment* NewAssignmentNode = AssignmentNodeCreator.CreateNode();
+	NewAssignmentNode->AssignmentTarget = ParameterVariable;
+	AssignmentNodeCreator.Finalize();
+
+	ConnectModuleNode(*NewAssignmentNode, TargetOutputNode, TargetIndex);
+	return NewAssignmentNode;
+}
+
+UNiagaraMaterialParameterNode* FNiagaraStackGraphUtilities::AddMaterialParameterModuleToStack(TSharedRef<FNiagaraEmitterViewModel>& EmitterViewModel, UNiagaraNodeOutput& TargetOutputNode, int32 TargetIndex)
+{
+	UEdGraph* Graph = TargetOutputNode.GetGraph();
+	Graph->Modify();
+
+	FGraphNodeCreator<UNiagaraMaterialParameterNode> MaterialParameterNodeCreator(*Graph);
+	UNiagaraMaterialParameterNode* NewMaterialParameterNode = MaterialParameterNodeCreator.CreateNode();
+	NewMaterialParameterNode->SetEmitter(EmitterViewModel->GetEmitter());
+	MaterialParameterNodeCreator.Finalize();
+	
+	ConnectModuleNode(*NewMaterialParameterNode, TargetOutputNode, TargetIndex);
+	return NewMaterialParameterNode;
+}
+
+void GetAllNodesForModule(UNiagaraNodeFunctionCall& ModuleFunctionCall, TArray<UNiagaraNode*>& ModuleNodes)
+{
+	TArray<FNiagaraStackGraphUtilities::FStackNodeGroup> StackNodeGroups;
+	FNiagaraStackGraphUtilities::GetStackNodeGroups(ModuleFunctionCall, StackNodeGroups);
+
+	int32 ThisGroupIndex = StackNodeGroups.IndexOfByPredicate([&](const FNiagaraStackGraphUtilities::FStackNodeGroup& Group) { return Group.EndNode == &ModuleFunctionCall; });
+	checkf(ThisGroupIndex > 0 && ThisGroupIndex < StackNodeGroups.Num() - 1, TEXT("Stack graph is invalid"));
+
+	TArray<UNiagaraNode*> AllGroupNodes;
+	StackNodeGroups[ThisGroupIndex].GetAllNodesInGroup(ModuleNodes);
+}
+
+TOptional<bool> FNiagaraStackGraphUtilities::GetModuleIsEnabled(UNiagaraNodeFunctionCall& FunctionCallNode)
+{
+	TArray<UNiagaraNode*> AllModuleNodes;
+	GetAllNodesForModule(FunctionCallNode, AllModuleNodes);
+	bool bIsEnabled = AllModuleNodes[0]->IsNodeEnabled();
+	for (int32 i = 1; i < AllModuleNodes.Num(); i++)
+	{
+		if (AllModuleNodes[i]->IsNodeEnabled() != bIsEnabled)
+		{
+			return TOptional<bool>();
+		}
+	}
+	return TOptional<bool>(bIsEnabled);
+}
+
+void FNiagaraStackGraphUtilities::SetModuleIsEnabled(UNiagaraNodeFunctionCall& FunctionCallNode, bool bIsEnabled)
+{
+	TArray<UNiagaraNode*> ModuleNodes;
+	GetAllNodesForModule(FunctionCallNode, ModuleNodes);
+	for (UNiagaraNode* ModuleNode : ModuleNodes)
+	{
+		ModuleNode->SetEnabledState(bIsEnabled ? ENodeEnabledState::Enabled : ENodeEnabledState::Disabled, true);
+	}
+	FunctionCallNode.GetNiagaraGraph()->NotifyGraphNeedsRecompile();
+}
+
+bool FNiagaraStackGraphUtilities::ValidateGraphForOutput(UNiagaraGraph& NiagaraGraph, ENiagaraScriptUsage ScriptUsage, FGuid ScriptUsageId, FText& ErrorMessage)
+{
+	UNiagaraNodeOutput* OutputNode = NiagaraGraph.FindOutputNode(ScriptUsage, ScriptUsageId);
 	if (OutputNode == nullptr)
 	{
 		ErrorMessage = LOCTEXT("ValidateNoOutputMessage", "Output node doesn't exist for script.");
@@ -455,10 +895,10 @@ bool FNiagaraStackGraphUtilities::ValidateGraphForOutput(UNiagaraGraph& NiagaraG
 	return true;
 }
 
-UNiagaraNodeOutput* FNiagaraStackGraphUtilities::ResetGraphForOutput(UNiagaraGraph& NiagaraGraph, ENiagaraScriptUsage ScriptUsage, int32 ScriptOccurrence)
+UNiagaraNodeOutput* FNiagaraStackGraphUtilities::ResetGraphForOutput(UNiagaraGraph& NiagaraGraph, ENiagaraScriptUsage ScriptUsage, FGuid ScriptUsageId)
 {
 	NiagaraGraph.Modify();
-	UNiagaraNodeOutput* OutputNode = NiagaraGraph.FindOutputNode(ScriptUsage, ScriptOccurrence);
+	UNiagaraNodeOutput* OutputNode = NiagaraGraph.FindOutputNode(ScriptUsage, ScriptUsageId);
 	UEdGraphPin* OutputNodeInputPin = OutputNode != nullptr ? GetParameterMapInputPin(*OutputNode) : nullptr;
 	if (OutputNode != nullptr && OutputNodeInputPin == nullptr)
 	{
@@ -471,7 +911,7 @@ UNiagaraNodeOutput* FNiagaraStackGraphUtilities::ResetGraphForOutput(UNiagaraGra
 		FGraphNodeCreator<UNiagaraNodeOutput> OutputNodeCreator(NiagaraGraph);
 		OutputNode = OutputNodeCreator.CreateNode();
 		OutputNode->SetUsage(ScriptUsage);
-		OutputNode->SetUsageIndex(ScriptOccurrence);
+		OutputNode->SetUsageId(ScriptUsageId);
 		OutputNode->Outputs.Add(FNiagaraVariable(FNiagaraTypeDefinition::GetParameterMapDef(), TEXT("Out")));
 		OutputNodeCreator.Finalize();
 
@@ -497,16 +937,102 @@ UNiagaraNodeOutput* FNiagaraStackGraphUtilities::ResetGraphForOutput(UNiagaraGra
 	{
 		// TODO: Move the emitter node wrangling to a utility function instead of getting the typed outer here and creating a view model.
 		UNiagaraSystem* System = NiagaraGraph.GetTypedOuter<UNiagaraSystem>();
-		if (System != nullptr)
+		if (System != nullptr && System->GetEmitterHandles().Num() != 0)
 		{
-			TSharedRef<FNiagaraSystemScriptViewModel> SystemScriptViewModel = MakeShared<FNiagaraSystemScriptViewModel>(*System);
+			TSharedRef<FNiagaraSystemScriptViewModel> SystemScriptViewModel = MakeShared<FNiagaraSystemScriptViewModel>(*System, nullptr);
 			SystemScriptViewModel->RebuildEmitterNodes();
 		}
 	}
 
-	RelayoutGraph(NiagaraGraph);
-
 	return OutputNode;
 }
 
+const UNiagaraEmitter* FNiagaraStackGraphUtilities::GetBaseEmitter(UNiagaraEmitter& Emitter, UNiagaraSystem& OwningSystem)
+{
+	for(const FNiagaraEmitterHandle& Handle : OwningSystem.GetEmitterHandles())
+	{ 
+		if (Handle.GetInstance() == &Emitter)
+		{
+			if (Handle.GetSource() != nullptr && Handle.GetSource() != &Emitter)
+			{
+				return Handle.GetSource();
+			}
+			else
+			{
+				// If the source is null then it was deleted and if the source is the same as the emitter the owning
+				// system is transient and the emitter doesn't have base.
+				return nullptr;
+			}
+		}
+	}
+	return nullptr;
+}
+
+void GetFunctionNamesRecursive(UNiagaraNode* CurrentNode, TArray<UNiagaraNode*>& VisitedNodes, TArray<FString>& FunctionNames)
+{
+	if (VisitedNodes.Contains(CurrentNode) == false)
+	{
+		VisitedNodes.Add(CurrentNode);
+		UNiagaraNodeFunctionCall* FunctionCall = Cast<UNiagaraNodeFunctionCall>(CurrentNode);
+		if (FunctionCall != nullptr)
+		{
+			FunctionNames.Add(FunctionCall->GetFunctionName());
+		}
+		TArray<UEdGraphPin*> InputPins;
+		CurrentNode->GetInputPins(InputPins);
+		for (UEdGraphPin* InputPin : InputPins)
+		{
+			for (UEdGraphPin* LinkedPin : InputPin->LinkedTo)
+			{
+				UNiagaraNode* LinkedNode = Cast<UNiagaraNode>(LinkedPin->GetOwningNode());
+				if (LinkedNode != nullptr)
+				{
+					GetFunctionNamesRecursive(LinkedNode, VisitedNodes, FunctionNames);
+				}
+			}
+		}
+	}
+}
+
+void GetFunctionNamesForOutputNode(UNiagaraNodeOutput& OutputNode, TArray<FString>& FunctionNames)
+{
+	TArray<UNiagaraNode*> VisitedNodes;
+	GetFunctionNamesRecursive(&OutputNode, VisitedNodes, FunctionNames);
+}
+
+void FNiagaraStackGraphUtilities::CleanUpStaleRapidIterationParameters(UNiagaraScript& Script, UNiagaraEmitter& OwningEmitter)
+{
+	UNiagaraScriptSource* Source = CastChecked<UNiagaraScriptSource>(Script.GetSource());
+	UNiagaraNodeOutput* OutputNode = Source->NodeGraph->FindOutputNode(Script.GetUsage(), Script.GetUsageId());
+	if (OutputNode != nullptr)
+	{
+		TArray<FString> ValidFunctionCallNames;
+		GetFunctionNamesForOutputNode(*OutputNode, ValidFunctionCallNames);
+		TArray<FNiagaraVariable> RapidIterationParameters;
+		Script.RapidIterationParameters.GetParameters(RapidIterationParameters);
+		for (const FNiagaraVariable& RapidIterationParameter : RapidIterationParameters)
+		{
+			FString EmitterName;
+			FString FunctionCallName;
+			if (FNiagaraParameterMapHistory::TryGetEmitterAndFunctionCallNamesFromRapidIterationParameter(RapidIterationParameter, EmitterName, FunctionCallName))
+			{
+				if (EmitterName != OwningEmitter.GetUniqueEmitterName() || ValidFunctionCallNames.Contains(FunctionCallName) == false)
+				{
+					Script.RapidIterationParameters.RemoveParameter(RapidIterationParameter);
+				}
+			}
+		}
+	}
+}
+
+void FNiagaraStackGraphUtilities::CleanUpStaleRapidIterationParameters(UNiagaraEmitter& Emitter)
+{
+	TArray<UNiagaraScript*> Scripts;
+	Emitter.GetScripts(Scripts, false);
+	for (UNiagaraScript* Script : Scripts)
+	{
+		CleanUpStaleRapidIterationParameters(*Script, Emitter);
+	}
+}
+
 #undef LOCTEXT_NAMESPACE
\ No newline at end of file
