Copying Orion-Staging @ 3748653 (Orion/Dev-General @ 3722124) to //UE4/Main

diff --git a/Engine/Plugins/FX/Niagara/Source/NiagaraEditor/Private/ViewModels/Stack/NiagaraStackModuleItem.cpp b/Engine/Plugins/FX/Niagara/Source/NiagaraEditor/Private/ViewModels/Stack/NiagaraStackModuleItem.cpp
--- a/Engine/Plugins/FX/Niagara/Source/NiagaraEditor/Private/ViewModels/Stack/NiagaraStackModuleItem.cpp
+++ b/Engine/Plugins/FX/Niagara/Source/NiagaraEditor/Private/ViewModels/Stack/NiagaraStackModuleItem.cpp
@@ -1,6 +1,9 @@
 // Copyright 1998-2016 Epic Games, Inc. All Rights Reserved.
 
 #include "NiagaraStackModuleItem.h"
+#include "NiagaraEditorModule.h"
+#include "NiagaraSystemViewModel.h"
+#include "NiagaraEmitterHandleViewModel.h"
 #include "NiagaraEmitterViewModel.h"
 #include "NiagaraStackFunctionInputCollection.h"
 #include "NiagaraStackFunctionInput.h"
@@ -14,6 +17,9 @@
 #include "NiagaraEditorStyle.h"
 #include "NiagaraGraph.h"
 #include "NiagaraStackGraphUtilities.h"
+#include "NiagaraEmitterHandle.h"
+#include "NiagaraScriptMergeManager.h"
+#include "NiagaraStackEditorData.h"
 
 #include "ScopedTransaction.h"
 
@@ -37,6 +43,34 @@ void UNiagaraStackModuleItem::Initialize(TSharedRef<FNiagaraSystemViewModel> InS
 	checkf(FunctionCallNode == nullptr, TEXT("Can not set the node more than once."));
 	Super::Initialize(InSystemViewModel, InEmitterViewModel, InStackEditorData);
 	FunctionCallNode = &InFunctionCallNode;
+	ModuleEditorDataKey = FNiagaraStackGraphUtilities::GenerateStackModuleEditorDataKey(*FunctionCallNode);
+
+	FOnFilterChild CollapsedFilter = FOnFilterChild::CreateLambda([this](const UNiagaraStackEntry& Child)
+	{
+		bool bChildShouldHiddenWhenNotExpanded = &Child == UnpinnedInputCollection || &Child == OutputCollection;
+		bool bIsExpanded = GetStackEditorData().GetStackEntryIsExpanded(ModuleEditorDataKey, true);
+		return bIsExpanded || bChildShouldHiddenWhenNotExpanded == false;
+	});
+	AddChildFilter(CollapsedFilter);
+
+	// Update bCanMoveAndDelete
+	if (GetSystemViewModel()->GetEditMode() == ENiagaraSystemViewModelEditMode::EmitterAsset)
+	{
+		// When editing emitters all modules can be moved and deleted.
+		bCanMoveAndDelete = true;
+	}
+	else
+	{
+		// When editing systems only non-base modules can be moved and deleted.
+		TSharedRef<FNiagaraScriptMergeManager> MergeManager = FNiagaraScriptMergeManager::Get();
+
+		const UNiagaraEmitter* BaseEmitter = FNiagaraStackGraphUtilities::GetBaseEmitter(*GetEmitterViewModel()->GetEmitter(), GetSystemViewModel()->GetSystem());
+		UNiagaraNodeOutput* OutputNode = FNiagaraStackGraphUtilities::GetEmitterOutputNodeForStackNode(*FunctionCallNode);
+
+		bool bIsMergeable = MergeManager->IsMergeableScriptUsage(OutputNode->GetUsage());
+		bool bHasBaseModule = bIsMergeable && BaseEmitter != nullptr && MergeManager->HasBaseModule(*BaseEmitter, OutputNode->GetUsage(), OutputNode->GetUsageId(), FunctionCallNode->NodeGuid);
+		bCanMoveAndDelete = bHasBaseModule == false;
+	}
 }
 
 FText UNiagaraStackModuleItem::GetDisplayName() const
@@ -66,58 +100,119 @@ FText UNiagaraStackModuleItem::GetTooltipText() const
 
 void UNiagaraStackModuleItem::RefreshChildrenInternal(const TArray<UNiagaraStackEntry*>& CurrentChildren, TArray<UNiagaraStackEntry*>& NewChildren)
 {
-	if (PinnedInputCollection == nullptr)
+	Errors.Empty();
+	if (FunctionCallNode->ScriptIsValid())
 	{
-		UNiagaraStackFunctionInputCollection::FDisplayOptions DisplayOptions;
-		DisplayOptions.DisplayName = LOCTEXT("PinnedInputsLabel", "Pinned Inputs");
-		DisplayOptions.ChildItemIndentLevel = 0;
-		DisplayOptions.bShouldShowInStack = false;
-		DisplayOptions.ChildFilter = UNiagaraStackFunctionInputCollection::FOnFilterChildren::CreateLambda(
-			[](UNiagaraStackFunctionInput* FunctionInput) { return FunctionInput->GetIsPinned(); });
-
-		TArray<FString> InputParameterHandlePath;
-		PinnedInputCollection = NewObject<UNiagaraStackFunctionInputCollection>(this);
-		PinnedInputCollection->Initialize(GetSystemViewModel(), GetEmitterViewModel(), GetStackEditorData(), *FunctionCallNode, *FunctionCallNode, DisplayOptions);
-		PinnedInputCollection->OnInputPinnedChanged().AddUObject(this, &UNiagaraStackModuleItem::InputPinnedChanged);
-	}
+		if (PinnedInputCollection == nullptr)
+		{
+			UNiagaraStackFunctionInputCollection::FDisplayOptions DisplayOptions;
+			DisplayOptions.DisplayName = LOCTEXT("PinnedInputsLabel", "Pinned Inputs");
+			DisplayOptions.ChildItemIndentLevel = 0;
+			DisplayOptions.bShouldShowInStack = false;
+			DisplayOptions.ChildFilter = FOnFilterChild::CreateLambda([](const UNiagaraStackEntry& Child)
+			{
+				const UNiagaraStackFunctionInput* FunctionInput = CastChecked<const UNiagaraStackFunctionInput>(&Child);
+				return FunctionInput->GetIsPinned();
+			});
+
+			TArray<FString> InputParameterHandlePath;
+			PinnedInputCollection = NewObject<UNiagaraStackFunctionInputCollection>(this);
+			PinnedInputCollection->Initialize(GetSystemViewModel(), GetEmitterViewModel(), GetStackEditorData(), *FunctionCallNode, *FunctionCallNode, DisplayOptions);
+			PinnedInputCollection->OnInputPinnedChanged().AddUObject(this, &UNiagaraStackModuleItem::InputPinnedChanged);
+		}
 
-	if (UnpinnedInputCollection == nullptr)
-	{
-		UNiagaraStackFunctionInputCollection::FDisplayOptions DisplayOptions;
-		DisplayOptions.DisplayName = LOCTEXT("InputsLabel", "Inputs");
-		DisplayOptions.ChildItemIndentLevel = 1;
-		DisplayOptions.bShouldShowInStack = true;
-		DisplayOptions.ChildFilter = UNiagaraStackFunctionInputCollection::FOnFilterChildren::CreateLambda(
-			[](UNiagaraStackFunctionInput* FunctionInput) { return FunctionInput->GetIsPinned() == false; });
-
-		TArray<FString> InputParameterHandlePath;
-		UnpinnedInputCollection = NewObject<UNiagaraStackFunctionInputCollection>(this);
-		UnpinnedInputCollection->Initialize(GetSystemViewModel(), GetEmitterViewModel(), GetStackEditorData(), *FunctionCallNode, *FunctionCallNode, DisplayOptions);
-		UnpinnedInputCollection->OnInputPinnedChanged().AddUObject(this, &UNiagaraStackModuleItem::InputPinnedChanged);
-	}
+		if (UnpinnedInputCollection == nullptr)
+		{
+			UNiagaraStackFunctionInputCollection::FDisplayOptions DisplayOptions;
+			DisplayOptions.DisplayName = LOCTEXT("InputsLabel", "Inputs");
+			DisplayOptions.ChildItemIndentLevel = 1;
+			DisplayOptions.bShouldShowInStack = true;
+			DisplayOptions.ChildFilter = FOnFilterChild::CreateLambda([](const UNiagaraStackEntry& Child)
+			{
+				const UNiagaraStackFunctionInput* FunctionInput = CastChecked<const UNiagaraStackFunctionInput>(&Child);
+				return FunctionInput->GetIsPinned() == false;
+			});
+
+			TArray<FString> InputParameterHandlePath;
+			UnpinnedInputCollection = NewObject<UNiagaraStackFunctionInputCollection>(this);
+			UnpinnedInputCollection->Initialize(GetSystemViewModel(), GetEmitterViewModel(), GetStackEditorData(), *FunctionCallNode, *FunctionCallNode, DisplayOptions);
+			UnpinnedInputCollection->OnInputPinnedChanged().AddUObject(this, &UNiagaraStackModuleItem::InputPinnedChanged);
+		}
+
+		if (OutputCollection == nullptr)
+		{
+			OutputCollection = NewObject<UNiagaraStackModuleItemOutputCollection>(this);
+			OutputCollection->Initialize(GetSystemViewModel(), GetEmitterViewModel(), *FunctionCallNode);
+			OutputCollection->SetDisplayName(LOCTEXT("OutputsLabel", "Outputs"));
+		}
 
-	if (OutputCollection == nullptr)
+		if (ModuleExpander == nullptr)
+		{
+			ModuleExpander = NewObject<UNiagaraStackItemExpander>(this);
+			ModuleExpander->Initialize(GetSystemViewModel(), GetEmitterViewModel(), GetStackEditorData(), ModuleEditorDataKey, true);
+			ModuleExpander->SetOnExpnadedChanged(UNiagaraStackItemExpander::FOnExpandedChanged::CreateUObject(this, &UNiagaraStackModuleItem::ModuleExpandedChanged));
+		}
+
+		NewChildren.Add(PinnedInputCollection);
+		NewChildren.Add(UnpinnedInputCollection);
+		NewChildren.Add(OutputCollection);
+		NewChildren.Add(ModuleExpander);
+	}
+	else
 	{
-		OutputCollection = NewObject<UNiagaraStackModuleItemOutputCollection>(this);
-		OutputCollection->Initialize(GetSystemViewModel(), GetEmitterViewModel(), *FunctionCallNode);
-		OutputCollection->SetDisplayName(LOCTEXT("OutputsLabel", "Outputs"));
+		FError InvalidScriptError;
+		InvalidScriptError.ErrorSummaryText = LOCTEXT("MissingModule", "The referenced module's script asset is missing.");
+		InvalidScriptError.ErrorText = InvalidScriptError.ErrorSummaryText;
 	}
 
-	FString ModuleEditorDataKey = FNiagaraStackGraphUtilities::GenerateStackModuleEditorDataKey(*FunctionCallNode);
-	if (ModuleExpander == nullptr)
+	RefreshIsEnabled();
+}
+
+void UNiagaraStackModuleItem::RefreshIsEnabled()
+{
+	TOptional<bool> IsEnabled = FNiagaraStackGraphUtilities::GetModuleIsEnabled(*FunctionCallNode);
+	if (IsEnabled.IsSet())
 	{
-		ModuleExpander = NewObject<UNiagaraStackItemExpander>(this);
-		ModuleExpander->Initialize(GetSystemViewModel(), GetEmitterViewModel(), GetStackEditorData(), ModuleEditorDataKey, true);
-		ModuleExpander->SetOnExpnadedChanged(UNiagaraStackItemExpander::FOnExpandedChanged::CreateUObject(this, &UNiagaraStackModuleItem::ModuleExpandedChanged));
+		bIsEnabled = IsEnabled.GetValue();
 	}
-
-	NewChildren.Add(PinnedInputCollection);
-	if (GetStackEditorData().GetStackEntryIsExpanded(ModuleEditorDataKey, true))
+	else
 	{
-		NewChildren.Add(UnpinnedInputCollection);
-		NewChildren.Add(OutputCollection);
+		bIsEnabled = false;
+		FError InconsistentEnabledError;
+		InconsistentEnabledError.ErrorSummaryText = LOCTEXT("InconsistentEnabledErrorSummary", "The endabled state for module is inconsistent.");
+		InconsistentEnabledError.ErrorText = LOCTEXT("InconsistentEnabledError", "This module is using multiple functions and their enabled state is inconsistent.\nClick fix to make all of the functions for this module enabled.");
+		InconsistentEnabledError.Fix = FFixError::CreateLambda([=]()
+		{
+			SetIsEnabled(true);
+			return true;
+		});
+		Errors.Add(InconsistentEnabledError);
 	}
-	NewChildren.Add(ModuleExpander);
+}
+
+int32 UNiagaraStackModuleItem::GetErrorCount() const 
+{
+	return Errors.Num();
+}
+
+bool UNiagaraStackModuleItem::GetErrorFixable(int32 ErrorIdx) const
+{
+	return Errors[ErrorIdx].Fix.IsBound();
+}
+
+bool UNiagaraStackModuleItem::TryFixError(int32 ErrorIdx)
+{
+	return Errors[ErrorIdx].Fix.IsBound() ? Errors[ErrorIdx].Fix.Execute() : false;
+}
+
+FText UNiagaraStackModuleItem::GetErrorText(int32 ErrorIdx) const
+{
+	return Errors[ErrorIdx].ErrorText;
+}
+
+FText UNiagaraStackModuleItem::GetErrorSummaryText(int32 ErrorIdx) const 
+{
+	return Errors[ErrorIdx].ErrorSummaryText;
 }
 
 void UNiagaraStackModuleItem::InputPinnedChanged()
@@ -128,11 +223,44 @@ void UNiagaraStackModuleItem::InputPinnedChanged()
 
 void UNiagaraStackModuleItem::ModuleExpandedChanged()
 {
-	RefreshChildren();
+	OnStructureChanged().Broadcast();
+}
+
+bool UNiagaraStackModuleItem::CanMoveAndDelete() const
+{
+	return bCanMoveAndDelete;
+}
+
+bool UNiagaraStackModuleItem::CanRefresh() const
+{
+	return FunctionCallNode != nullptr ? FunctionCallNode->CanRefreshFromModuleItem() : false;
+}
+
+void UNiagaraStackModuleItem::Refresh()
+{
+	if (CanRefresh())
+	{
+		FunctionCallNode->RefreshFromExternalChanges();
+		FunctionCallNode->GetNiagaraGraph()->NotifyGraphNeedsRecompile();
+		GetSystemViewModel()->ResetSystem();
+	}
+}
+
+bool UNiagaraStackModuleItem::GetIsEnabled() const
+{
+	return bIsEnabled;
+}
+
+void UNiagaraStackModuleItem::SetIsEnabled(bool bInIsEnabled)
+{
+	FNiagaraStackGraphUtilities::SetModuleIsEnabled(*FunctionCallNode, bInIsEnabled);
+	bIsEnabled = bInIsEnabled;
 }
 
 void UNiagaraStackModuleItem::MoveUp()
 {
+	checkf(CanMoveAndDelete(), TEXT("This module can't be moved or deleted"));
+
 	TArray<FNiagaraStackGraphUtilities::FStackNodeGroup> StackNodeGroups;
 	FNiagaraStackGraphUtilities::GetStackNodeGroups(*FunctionCallNode, StackNodeGroups);
 	int32 ModuleStackIndex = StackNodeGroups.IndexOfByPredicate([&](const FNiagaraStackGraphUtilities::FStackNodeGroup& StackNodeGroup) { return StackNodeGroup.EndNode == FunctionCallNode; });
@@ -152,6 +280,8 @@ void UNiagaraStackModuleItem::MoveUp()
 
 void UNiagaraStackModuleItem::MoveDown()
 {
+	checkf(CanMoveAndDelete(), TEXT("This module can't be moved or deleted"));
+
 	TArray<FNiagaraStackGraphUtilities::FStackNodeGroup> StackNodeGroups;
 	FNiagaraStackGraphUtilities::GetStackNodeGroups(*FunctionCallNode, StackNodeGroups);
 	int32 ModuleStackIndex = StackNodeGroups.IndexOfByPredicate([&](const FNiagaraStackGraphUtilities::FStackNodeGroup& StackNodeGroup) { return StackNodeGroup.EndNode == FunctionCallNode; });
@@ -171,49 +301,18 @@ void UNiagaraStackModuleItem::MoveDown()
 
 void UNiagaraStackModuleItem::Delete()
 {
-	TArray<FNiagaraStackGraphUtilities::FStackNodeGroup> StackNodeGroups;
-	FNiagaraStackGraphUtilities::GetStackNodeGroups(*FunctionCallNode, StackNodeGroups);
-	int32 ModuleStackIndex = StackNodeGroups.IndexOfByPredicate([&](const FNiagaraStackGraphUtilities::FStackNodeGroup& StackNodeGroup) { return StackNodeGroup.EndNode == FunctionCallNode; });
+	checkf(CanMoveAndDelete(), TEXT("This module can't be moved or deleted"));
 
 	FScopedTransaction ScopedTransaction(LOCTEXT("RemoveAModuleFromTheStack", "Remove a module from the stack"));
 
-	FNiagaraStackGraphUtilities::DisconnectStackNodeGroup(StackNodeGroups[ModuleStackIndex], StackNodeGroups[ModuleStackIndex - 1], StackNodeGroups[ModuleStackIndex + 1]);
-
-	FNiagaraStackGraphUtilities::FStackNodeGroup ModuleGroup = StackNodeGroups[ModuleStackIndex];
-
-	TArray<UNiagaraNode*> NodesToRemove;
-	TArray<UNiagaraNode*> NodesToCheck;
-	NodesToCheck.Add(ModuleGroup.EndNode);
-	while (NodesToCheck.Num() > 0)
+	bool bRemoved = FNiagaraStackGraphUtilities::RemoveModuleFromStack(*FunctionCallNode);
+	if (bRemoved)
 	{
-		UNiagaraNode* NodeToRemove = NodesToCheck[0];
-		NodesToCheck.RemoveAt(0);
-		NodesToRemove.Add(NodeToRemove);
-
-		TArray<UEdGraphPin*> InputPins;
-		NodeToRemove->GetInputPins(InputPins);
-		for (UEdGraphPin* InputPin : InputPins)
-		{
-			if (InputPin->LinkedTo.Num() == 1)
-			{
-				UNiagaraNode* LinkedNode = Cast<UNiagaraNode>(InputPin->LinkedTo[0]->GetOwningNode());
-				if (LinkedNode != nullptr)
-				{
-					NodesToCheck.Add(LinkedNode);
-				}
-			}
-		}
-	}
-
-	UNiagaraGraph* Graph = FunctionCallNode->GetNiagaraGraph();
-	for (UNiagaraNode* NodeToRemove : NodesToRemove)
-	{
-		Graph->RemoveNode(NodeToRemove);
+		UNiagaraGraph* Graph = FunctionCallNode->GetNiagaraGraph();
+		Graph->NotifyGraphNeedsRecompile();
+		FNiagaraStackGraphUtilities::RelayoutGraph(*FunctionCallNode->GetGraph());
+		ModifiedGroupItemsDelegate.ExecuteIfBound();
 	}
-
-	Graph->NotifyGraphNeedsRecompile();
-	FNiagaraStackGraphUtilities::RelayoutGraph(*FunctionCallNode->GetGraph());
-	ModifiedGroupItemsDelegate.ExecuteIfBound();
 }
 
 #undef LOCTEXT_NAMESPACE
diff --git a/Engine/Plugins/FX/Niagara/Source/NiagaraEditor/Private/ViewModels/Stack/NiagaraStackModuleItemOutput.cpp b/Engine/Plugins/FX/Niagara/Source/NiagaraEditor/Private/ViewModels/Stack/NiagaraStackModuleItemOutput.cpp
--- a/Engine/Plugins/FX/Niagara/Source/NiagaraEditor/Private/ViewModels/Stack/NiagaraStackModuleItemOutput.cpp
+++ b/Engine/Plugins/FX/Niagara/Source/NiagaraEditor/Private/ViewModels/Stack/NiagaraStackModuleItemOutput.cpp
@@ -2,6 +2,11 @@
 
 #include "NiagaraStackModuleItemOutput.h"
 #include "NiagaraNodeFunctionCall.h"
+#include "NiagaraConstants.h"
+#include "NiagaraScriptSource.h"
+#include "NiagaraGraph.h"
+#include "EdGraphSchema_Niagara.h"
+
 
 #define LOCTEXT_NAMESPACE "NiagaraStackViewModel"
 
@@ -15,12 +20,13 @@ int32 UNiagaraStackModuleItemOutput::GetItemIndentLevel() const
 	return 1;
 }
 
-void UNiagaraStackModuleItemOutput::Initialize(TSharedRef<FNiagaraSystemViewModel> InSystemViewModel, TSharedRef<FNiagaraEmitterViewModel> InEmitterViewModel, UNiagaraNodeFunctionCall& InFunctionCallNode, FName InOutputParameterHandle)
+void UNiagaraStackModuleItemOutput::Initialize(TSharedRef<FNiagaraSystemViewModel> InSystemViewModel, TSharedRef<FNiagaraEmitterViewModel> InEmitterViewModel, UNiagaraNodeFunctionCall& InFunctionCallNode, FName InOutputParameterHandle,
+	FNiagaraTypeDefinition InOutputType)
 {
 	checkf(FunctionCallNode.Get() == nullptr, TEXT("Can only set the Output once."));
 	Super::Initialize(InSystemViewModel, InEmitterViewModel);
 	FunctionCallNode = &InFunctionCallNode;
-
+	OutputType = InOutputType;
 	OutputParameterHandle = FNiagaraParameterHandle(InOutputParameterHandle);
 	DisplayName = FText::FromName(OutputParameterHandle.GetName());
 }
@@ -40,6 +46,31 @@ bool UNiagaraStackModuleItemOutput::GetCanExpand() const
 	return true;
 }
 
+FText UNiagaraStackModuleItemOutput::GetTooltipText() const
+{
+	FNiagaraVariable ValueVariable(OutputType, OutputParameterHandle.GetParameterHandleString());
+	if (FunctionCallNode.IsValid() && FunctionCallNode->FunctionScript != nullptr)
+	{
+		UNiagaraScriptSource* Source = Cast<UNiagaraScriptSource>(FunctionCallNode->FunctionScript->GetSource());
+		const UEdGraphSchema_Niagara* NiagaraSchema = GetDefault<UEdGraphSchema_Niagara>();
+		const FNiagaraVariableMetaData* MetaData = nullptr;
+		if (FNiagaraConstants::IsNiagaraConstant(ValueVariable))
+		{
+			MetaData = FNiagaraConstants::GetConstantMetaData(ValueVariable);
+		}
+		else if (Source->NodeGraph != nullptr)
+		{
+			MetaData = Source->NodeGraph->GetMetaData(ValueVariable);
+		}
+
+		if (MetaData != nullptr)
+		{
+			return MetaData->Description;
+		}
+	}
+	return FText::FromName(ValueVariable.GetName());
+}
+
 const FNiagaraParameterHandle& UNiagaraStackModuleItemOutput::GetOutputParameterHandle() const
 {
 	return OutputParameterHandle;
diff --git a/Engine/Plugins/FX/Niagara/Source/NiagaraEditor/Private/ViewModels/Stack/NiagaraStackModuleItemOutputCollection.cpp b/Engine/Plugins/FX/Niagara/Source/NiagaraEditor/Private/ViewModels/Stack/NiagaraStackModuleItemOutputCollection.cpp
--- a/Engine/Plugins/FX/Niagara/Source/NiagaraEditor/Private/ViewModels/Stack/NiagaraStackModuleItemOutputCollection.cpp
+++ b/Engine/Plugins/FX/Niagara/Source/NiagaraEditor/Private/ViewModels/Stack/NiagaraStackModuleItemOutputCollection.cpp
@@ -52,35 +52,39 @@ void UNiagaraStackModuleItemOutputCollection::RefreshChildrenInternal(const TArr
 	if (ensureMsgf(OutputParameterMapPin != nullptr, TEXT("Invalid Stack Graph - Function call node has no output pin.")))
 	{
 		FNiagaraParameterMapHistoryBuilder Builder;
+		Builder.SetIgnoreDisabled(false);
 		FunctionCallNode->BuildParameterMapHistory(Builder, false);
-		check(Builder.Histories.Num() == 1);
-		for (int32 i = 0; i < Builder.Histories[0].Variables.Num(); i++)
-		{
-			FNiagaraVariable& Variable = Builder.Histories[0].Variables[i];
-			TArray<const UEdGraphPin*>& WriteHistory = Builder.Histories[0].PerVariableWriteHistory[i];
 
-			for (const UEdGraphPin* WritePin : WriteHistory)
+		if (ensureMsgf(Builder.Histories.Num() == 1, TEXT("Invalid Stack Graph - Function call node has invalid history count!")))
+		{
+			for (int32 i = 0; i < Builder.Histories[0].Variables.Num(); i++)
 			{
-				if (Cast<UNiagaraNodeParameterMapSet>(WritePin->GetOwningNode()) != nullptr)
+				FNiagaraVariable& Variable = Builder.Histories[0].Variables[i];
+				TArray<const UEdGraphPin*>& WriteHistory = Builder.Histories[0].PerVariableWriteHistory[i];
+
+				for (const UEdGraphPin* WritePin : WriteHistory)
 				{
-					UNiagaraStackModuleItemOutput* Output = nullptr;
-					for (UNiagaraStackEntry* CurrentChild : CurrentChildren)
+					if (Cast<UNiagaraNodeParameterMapSet>(WritePin->GetOwningNode()) != nullptr)
 					{
-						UNiagaraStackModuleItemOutput* ChildOutput = CastChecked<UNiagaraStackModuleItemOutput>(CurrentChild);
+						UNiagaraStackModuleItemOutput* Output = nullptr;
+						for (UNiagaraStackEntry* CurrentChild : CurrentChildren)
+						{
+							UNiagaraStackModuleItemOutput* ChildOutput = CastChecked<UNiagaraStackModuleItemOutput>(CurrentChild);
 						if (ChildOutput->GetOutputParameterHandle().GetParameterHandleString() == Variable.GetName())
+							{
+								Output = ChildOutput;
+								break;
+							}
+						}
+						if (Output == nullptr)
 						{
-							Output = ChildOutput;
-							break;
+							Output = NewObject<UNiagaraStackModuleItemOutput>(this);
+							Output->Initialize(GetSystemViewModel(), GetEmitterViewModel(), *FunctionCallNode, Variable.GetName(), Variable.GetType());
 						}
-					}
-					if (Output == nullptr)
-					{
-						Output = NewObject<UNiagaraStackModuleItemOutput>(this);
-						Output->Initialize(GetSystemViewModel(), GetEmitterViewModel(), *FunctionCallNode, Variable.GetName());
-					}
 
-					NewChildren.Add(Output);
-					break;
+						NewChildren.Add(Output);
+						break;
+					}
 				}
 			}
 		}
diff --git a/Engine/Plugins/FX/Niagara/Source/NiagaraEditor/Private/ViewModels/Stack/NiagaraStackRenderItemGroup.cpp b/Engine/Plugins/FX/Niagara/Source/NiagaraEditor/Private/ViewModels/Stack/NiagaraStackRenderItemGroup.cpp
--- a/Engine/Plugins/FX/Niagara/Source/NiagaraEditor/Private/ViewModels/Stack/NiagaraStackRenderItemGroup.cpp
+++ b/Engine/Plugins/FX/Niagara/Source/NiagaraEditor/Private/ViewModels/Stack/NiagaraStackRenderItemGroup.cpp
@@ -7,6 +7,7 @@
 #include "NiagaraEmitterViewModel.h"
 #include "NiagaraEmitter.h"
 #include "NiagaraEmitterEditorData.h"
+#include "NiagaraSystemViewModel.h"
 
 FText UNiagaraStackRenderItemGroup::GetDisplayName() const
 {
@@ -21,7 +22,7 @@ void UNiagaraStackRenderItemGroup::SetDisplayName(FText InDisplayName)
 void UNiagaraStackRenderItemGroup::RefreshChildrenInternal(const TArray<UNiagaraStackEntry*>& CurrentChildren, TArray<UNiagaraStackEntry*>& NewChildren)
 {
 	int32 RendererIndex = 0;
-	for (UNiagaraRendererProperties* RendererProperties : GetEmitterViewModel()->GetEmitter()->RendererProperties)
+	for (UNiagaraRendererProperties* RendererProperties : GetEmitterViewModel()->GetEmitter()->GetRenderers())
 	{
 		UNiagaraStackRendererItem* RendererItem = FindCurrentChildOfTypeByPredicate<UNiagaraStackRendererItem>(CurrentChildren,
 			[=](UNiagaraStackRendererItem* CurrentRendererItem) { return CurrentRendererItem->GetRendererProperties() == RendererProperties; });
@@ -45,6 +46,8 @@ void UNiagaraStackRenderItemGroup::RefreshChildrenInternal(const TArray<UNiagara
 
 		NewChildren.Add(RendererItem);
 		NewChildren.Add(RendererSpacer);
+
+		RendererIndex++;
 	}
 
 	UNiagaraStackAddRendererItem* AddRendererItem = NewObject<UNiagaraStackAddRendererItem>(this);
@@ -55,6 +58,7 @@ void UNiagaraStackRenderItemGroup::RefreshChildrenInternal(const TArray<UNiagara
 
 void UNiagaraStackRenderItemGroup::ChildModifiedGroupItems()
 {
+	GetSystemViewModel()->RefreshAll();
 	RefreshChildren();
 }
 
diff --git a/Engine/Plugins/FX/Niagara/Source/NiagaraEditor/Private/ViewModels/Stack/NiagaraStackRendererItem.cpp b/Engine/Plugins/FX/Niagara/Source/NiagaraEditor/Private/ViewModels/Stack/NiagaraStackRendererItem.cpp
--- a/Engine/Plugins/FX/Niagara/Source/NiagaraEditor/Private/ViewModels/Stack/NiagaraStackRendererItem.cpp
+++ b/Engine/Plugins/FX/Niagara/Source/NiagaraEditor/Private/ViewModels/Stack/NiagaraStackRendererItem.cpp
@@ -7,6 +7,7 @@
 #include "NiagaraStackEditorData.h"
 #include "NiagaraRendererProperties.h"
 #include "NiagaraScript.h"
+#include "NiagaraSystemViewModel.h"
 #include "NiagaraEmitterViewModel.h"
 #include "NiagaraScriptViewModel.h"
 #include "Internationalization.h"
@@ -23,20 +24,32 @@
 #include "Framework/Notifications/NotificationManager.h"
 #include "Widgets/Notifications/SNotificationList.h"
 #include "NiagaraStackGraphUtilities.h"
+#include "NiagaraScriptMergeManager.h"
 
 #define LOCTEXT_NAMESPACE "UNiagaraStackRendererItem"
 
 UNiagaraStackRendererItem::UNiagaraStackRendererItem()
-	: RendererProperties(nullptr)
-	, RendererObject(nullptr)
+	: RendererObject(nullptr)
 {
 }
 
 void UNiagaraStackRendererItem::Initialize(TSharedRef<FNiagaraSystemViewModel> InSystemViewModel, TSharedRef<FNiagaraEmitterViewModel> InEmitterViewModel, UNiagaraStackEditorData& InStackEditorData, UNiagaraRendererProperties* InRendererProperties)
 {
-	checkf(RendererProperties == nullptr, TEXT("Can not initialize more than once."));
+	checkf(RendererProperties.IsValid() == false, TEXT("Can not initialize more than once."));
 	Super::Initialize(InSystemViewModel, InEmitterViewModel, InStackEditorData);
 	RendererProperties = InRendererProperties;
+	RendererProperties->OnChanged().AddUObject(this, &UNiagaraStackRendererItem::RendererChanged);
+
+	if (GetSystemViewModel()->GetEditMode() == ENiagaraSystemViewModelEditMode::EmitterAsset)
+	{
+		bHasBaseRenderer = false;
+	}
+	else
+	{
+		TSharedRef<FNiagaraScriptMergeManager> MergeManager = FNiagaraScriptMergeManager::Get();
+		const UNiagaraEmitter* BaseEmitter = FNiagaraStackGraphUtilities::GetBaseEmitter(*GetEmitterViewModel()->GetEmitter(), GetSystemViewModel()->GetSystem());
+		bHasBaseRenderer = BaseEmitter != nullptr && MergeManager->HasBaseRenderer(*BaseEmitter, RendererProperties->GetMergeId());
+	}
 }
 
 TArray<FNiagaraVariable> UNiagaraStackRendererItem::GetMissingVariables(UNiagaraRendererProperties* RendererProperties, UNiagaraEmitter* Emitter)
@@ -44,7 +57,7 @@ TArray<FNiagaraVariable> UNiagaraStackRendererItem::GetMissingVariables(UNiagara
 	TArray<FNiagaraVariable> MissingAttributes;
 	const TArray<FNiagaraVariable>& RequiredAttrs = RendererProperties->GetRequiredAttributes();
 	const UNiagaraScript* Script = Emitter->SpawnScriptProps.Script;
-	if (Script != nullptr)
+	if (Script != nullptr && Script->GetByteCode().Num() != 0)
 	{
 		MissingAttributes.Empty();
 		for (FNiagaraVariable Attr : RequiredAttrs)
@@ -119,7 +132,7 @@ bool UNiagaraStackRendererItem::AddMissingVariable(UNiagaraEmitter* Emitter, con
 
 UNiagaraRendererProperties* UNiagaraStackRendererItem::GetRendererProperties()
 {
-	return RendererProperties;
+	return RendererProperties.Get();
 }
 
 FText UNiagaraStackRendererItem::GetDisplayName() const
@@ -134,17 +147,52 @@ FText UNiagaraStackRendererItem::GetDisplayName() const
 	}
 }
 
+bool UNiagaraStackRendererItem::CanDelete() const
+{
+	return bHasBaseRenderer == false;
+}
+
 void UNiagaraStackRendererItem::Delete()
 {
 	const FScopedTransaction Transaction(LOCTEXT("DeleteRenderer", "Delete Renderer"));
 
 	UNiagaraEmitter* Emitter = GetEmitterViewModel()->GetEmitter();
 	Emitter->Modify();
-	Emitter->RendererProperties.Remove(RendererProperties);
+	Emitter->RemoveRenderer(RendererProperties.Get());
 
 	ModifiedGroupItemsDelegate.ExecuteIfBound();
 }
 
+bool UNiagaraStackRendererItem::CanHaveBase() const
+{
+	return GetSystemViewModel()->GetEditMode() == ENiagaraSystemViewModelEditMode::SystemAsset;
+}
+
+bool UNiagaraStackRendererItem::CanResetToBase() const
+{
+	if (CanHaveBase())
+	{
+		if (bCanResetToBase.IsSet() == false)
+		{
+			TSharedRef<FNiagaraScriptMergeManager> MergeManager = FNiagaraScriptMergeManager::Get();
+			const UNiagaraEmitter* BaseEmitter = FNiagaraStackGraphUtilities::GetBaseEmitter(*GetEmitterViewModel()->GetEmitter(), GetSystemViewModel()->GetSystem());
+			bCanResetToBase = BaseEmitter != nullptr && MergeManager->IsRendererDifferentFromBase(*GetEmitterViewModel()->GetEmitter(), *BaseEmitter, RendererProperties->GetMergeId());
+		}
+		return bCanResetToBase.GetValue();
+	}
+	return false;
+}
+
+void UNiagaraStackRendererItem::ResetToBase()
+{
+	if (CanResetToBase())
+	{
+		TSharedRef<FNiagaraScriptMergeManager> MergeManager = FNiagaraScriptMergeManager::Get();
+		const UNiagaraEmitter* BaseEmitter = FNiagaraStackGraphUtilities::GetBaseEmitter(*GetEmitterViewModel()->GetEmitter(), GetSystemViewModel()->GetSystem());
+		MergeManager->ResetRendererToBase(*GetEmitterViewModel()->GetEmitter(), *BaseEmitter, RendererProperties->GetMergeId());
+		ModifiedGroupItemsDelegate.ExecuteIfBound();
+	}
+}
 
 FName UNiagaraStackRendererItem::GetItemBackgroundName() const
 {
@@ -190,12 +238,21 @@ FText UNiagaraStackRendererItem::GetErrorText(int32 ErrorIdx) const
 	return FText();
 }
 
+void UNiagaraStackRendererItem::BeginDestroy()
+{
+	if (RendererProperties.IsValid())
+	{
+		RendererProperties->OnChanged().RemoveAll(this);
+	}
+	Super::BeginDestroy();
+}
+
 void UNiagaraStackRendererItem::RefreshChildrenInternal(const TArray<UNiagaraStackEntry*>& CurrentChildren, TArray<UNiagaraStackEntry*>& NewChildren)
 {
 	if (RendererObject == nullptr)
 	{
 		RendererObject = NewObject<UNiagaraStackObject>(this);
-		RendererObject->Initialize(GetSystemViewModel(), GetEmitterViewModel(), RendererProperties);
+		RendererObject->Initialize(GetSystemViewModel(), GetEmitterViewModel(), RendererProperties.Get());
 	}
 
 	if (RendererExpander == nullptr)
@@ -212,7 +269,8 @@ void UNiagaraStackRendererItem::RefreshChildrenInternal(const TArray<UNiagaraSta
 
 	NewChildren.Add(RendererExpander);
 
-	MissingAttributes = GetMissingVariables(RendererProperties, GetEmitterViewModel()->GetEmitter());
+	MissingAttributes = GetMissingVariables(RendererProperties.Get(), GetEmitterViewModel()->GetEmitter());
+	bCanResetToBase.Reset();
 }
 
 void UNiagaraStackRendererItem::RendererExpandedChanged()
@@ -220,4 +278,9 @@ void UNiagaraStackRendererItem::RendererExpandedChanged()
 	RefreshChildren();
 }
 
+void UNiagaraStackRendererItem::RendererChanged()
+{
+	bCanResetToBase.Reset();
+}
+
 #undef LOCTEXT_NAMESPACE
\ No newline at end of file
diff --git a/Engine/Plugins/FX/Niagara/Source/NiagaraEditor/Private/ViewModels/Stack/NiagaraStackRoot.cpp b/Engine/Plugins/FX/Niagara/Source/NiagaraEditor/Private/ViewModels/Stack/NiagaraStackRoot.cpp
--- a/Engine/Plugins/FX/Niagara/Source/NiagaraEditor/Private/ViewModels/Stack/NiagaraStackRoot.cpp
+++ b/Engine/Plugins/FX/Niagara/Source/NiagaraEditor/Private/ViewModels/Stack/NiagaraStackRoot.cpp
@@ -49,7 +49,7 @@ bool UNiagaraStackRoot::GetShouldShowInStack() const
 void UNiagaraStackRoot::RefreshChildrenInternal(const TArray<UNiagaraStackEntry*>& CurrentChildren, TArray<UNiagaraStackEntry*>& NewChildren)
 {
 	// We only allow displaying and editing system stacks if the system isn't transient which is the case in the emitter editor.
-	bool bShowSystemGroups = GetSystemViewModel()->GetSystemIsTransient() == false;
+	bool bShowSystemGroups = GetSystemViewModel()->GetEditMode() == ENiagaraSystemViewModelEditMode::SystemAsset;
 
 	// Create static entries as needed.
 	if (bShowSystemGroups && SystemExposedVariablesGroup == nullptr)
@@ -137,23 +137,15 @@ void UNiagaraStackRoot::RefreshChildrenInternal(const TArray<UNiagaraStackEntry*
 	NewChildren.Add(ParticleSpawnGroup);
 	NewChildren.Add(ParticleUpdateGroup);
 
-	for (int32 EventHandlerIdx = 0; EventHandlerIdx < GetEmitterViewModel()->GetEmitter()->EventHandlerScriptProps.Num(); EventHandlerIdx++)
+	for (const FNiagaraEventScriptProperties& EventScriptProperties : GetEmitterViewModel()->GetEmitter()->GetEventHandlers())
 	{
-		UNiagaraStackEventScriptItemGroup* EventHandlerGroup = nullptr;
-		for (UNiagaraStackEntry* CurrentChild : CurrentChildren)
-		{
-			UNiagaraStackEventScriptItemGroup* CurrentEventHandlerChild = Cast<UNiagaraStackEventScriptItemGroup>(CurrentChild);
-			if (CurrentEventHandlerChild != nullptr && CurrentEventHandlerChild->GetScriptOccurrence() == EventHandlerIdx)
-			{
-				EventHandlerGroup = CurrentEventHandlerChild;
-				break;
-			}
-		}
+		UNiagaraStackEventScriptItemGroup* EventHandlerGroup = FindCurrentChildOfTypeByPredicate<UNiagaraStackEventScriptItemGroup>(CurrentChildren,
+			[&](UNiagaraStackEventScriptItemGroup* CurrentEventHandlerGroup) { return CurrentEventHandlerGroup->GetScriptUsageId() == EventScriptProperties.Script->GetUsageId(); });
 
 		if (EventHandlerGroup == nullptr)
 		{
 			EventHandlerGroup = NewObject<UNiagaraStackEventScriptItemGroup>(this, NAME_None, RF_Transactional);
-			EventHandlerGroup->Initialize(GetSystemViewModel(), GetEmitterViewModel(), GetEmitterViewModel()->GetEditorData().GetStackEditorData(), GetEmitterViewModel()->GetSharedScriptViewModel(), ENiagaraScriptUsage::ParticleEventScript, EventHandlerIdx);
+			EventHandlerGroup->Initialize(GetSystemViewModel(), GetEmitterViewModel(), GetEmitterViewModel()->GetEditorData().GetStackEditorData(), GetEmitterViewModel()->GetSharedScriptViewModel(), ENiagaraScriptUsage::ParticleEventScript, EventScriptProperties.Script->GetUsageId());
 			EventHandlerGroup->SetOnModifiedEventHandlers(UNiagaraStackEventScriptItemGroup::FOnModifiedEventHandlers::CreateUObject(this, &UNiagaraStackRoot::EmitterEventArraysChanged));
 			EventHandlerGroup->SetTooltipText(LOCTEXT("EventGroupTooltip", "Determines how this Emitter responds to incoming events. There can be more than one event handler script stack per Emitter."));
 		}
diff --git a/Engine/Plugins/FX/Niagara/Source/NiagaraEditor/Private/ViewModels/Stack/NiagaraStackScriptItemGroup.cpp b/Engine/Plugins/FX/Niagara/Source/NiagaraEditor/Private/ViewModels/Stack/NiagaraStackScriptItemGroup.cpp
--- a/Engine/Plugins/FX/Niagara/Source/NiagaraEditor/Private/ViewModels/Stack/NiagaraStackScriptItemGroup.cpp
+++ b/Engine/Plugins/FX/Niagara/Source/NiagaraEditor/Private/ViewModels/Stack/NiagaraStackScriptItemGroup.cpp
@@ -17,25 +17,19 @@
 
 #define LOCTEXT_NAMESPACE "UNiagaraStackScriptItemGroup"
 
-UNiagaraStackScriptItemGroup::UNiagaraStackScriptItemGroup()
-	: AddModuleItem(nullptr)
-	, BottomSpacer(nullptr)
-{
-}
-
 void UNiagaraStackScriptItemGroup::Initialize(
 	TSharedRef<FNiagaraSystemViewModel> InSystemViewModel,
 	TSharedRef<FNiagaraEmitterViewModel> InEmitterViewModel,
 	UNiagaraStackEditorData& InStackEditorData,
 	TSharedRef<FNiagaraScriptViewModel> InScriptViewModel,
 	ENiagaraScriptUsage InScriptUsage,
-	int32 InScriptOccurrence)
+	FGuid InScriptUsageId)
 {
 	checkf(ScriptViewModel.IsValid() == false, TEXT("Can not set the script view model more than once."));
 	Super::Initialize(InSystemViewModel, InEmitterViewModel, InStackEditorData);
 	ScriptViewModel = InScriptViewModel;
 	ScriptUsage = InScriptUsage;
-	ScriptOccurrence = InScriptOccurrence;
+	ScriptUsageId = InScriptUsageId;
 }
 
 FText UNiagaraStackScriptItemGroup::GetDisplayName() const
@@ -50,9 +44,12 @@ void UNiagaraStackScriptItemGroup::SetDisplayName(FText InDisplayName)
 
 void UNiagaraStackScriptItemGroup::RefreshChildrenInternal(const TArray<UNiagaraStackEntry*>& CurrentChildren, TArray<UNiagaraStackEntry*>& NewChildren)
 {
-	UNiagaraGraph* Graph = ScriptViewModel->GetGraphViewModel()->GetGraph();
+	TSharedPtr<FNiagaraScriptViewModel> ScriptViewModelPinned = ScriptViewModel.Pin();
+	checkf(ScriptViewModelPinned.IsValid(), TEXT("Can not refresh children when the script view model has been deleted."));
+
+	UNiagaraGraph* Graph = ScriptViewModelPinned->GetGraphViewModel()->GetGraph();
 	FText ErrorMessage;
-	if (FNiagaraStackGraphUtilities::ValidateGraphForOutput(*Graph, ScriptUsage, ScriptOccurrence, ErrorMessage) == false)
+	if (FNiagaraStackGraphUtilities::ValidateGraphForOutput(*Graph, ScriptUsage, ScriptUsageId, ErrorMessage) == false)
 	{
 		UE_LOG(LogNiagaraEditor, Error, TEXT("Failed to Create Stack.  Message: %s"), *ErrorMessage.ToString());
 		Error = MakeShared<FError>();
@@ -61,17 +58,28 @@ void UNiagaraStackScriptItemGroup::RefreshChildrenInternal(const TArray<UNiagara
 		Error->Fix.BindLambda([=]()
 		{
 			FScopedTransaction ScopedTransaction(LOCTEXT("FixStackGraph", "Fix invalid stack graph"));
-			FNiagaraStackGraphUtilities::ResetGraphForOutput(*Graph, ScriptUsage, ScriptOccurrence);
+			FNiagaraStackGraphUtilities::ResetGraphForOutput(*Graph, ScriptUsage, ScriptUsageId);
+			FNiagaraStackGraphUtilities::RelayoutGraph(*Graph);
 		});
 	}
 	else
 	{
-		UNiagaraNodeOutput* MatchingOutputNode = Graph->FindOutputNode(ScriptUsage, ScriptOccurrence);
+		UNiagaraNodeOutput* MatchingOutputNode = Graph->FindOutputNode(ScriptUsage, ScriptUsageId);
 		TArray<UNiagaraNodeFunctionCall*> ModuleNodes;
 		FNiagaraStackGraphUtilities::GetOrderedModuleNodes(*MatchingOutputNode, ModuleNodes);
 		int32 ModuleIndex = 0;
 		for (UNiagaraNodeFunctionCall* ModuleNode : ModuleNodes)
 		{
+			UNiagaraStackAddScriptModuleItem* AddModuleAtIndexItem = FindCurrentChildOfTypeByPredicate<UNiagaraStackAddScriptModuleItem>(CurrentChildren,
+				[=](UNiagaraStackAddScriptModuleItem* CurrentAddModuleItem) { return CurrentAddModuleItem->GetTargetIndex() == ModuleIndex; });
+
+			if (AddModuleAtIndexItem == nullptr)
+			{
+				AddModuleAtIndexItem = NewObject<UNiagaraStackAddScriptModuleItem>(this);
+				AddModuleAtIndexItem->Initialize(GetSystemViewModel(), GetEmitterViewModel(), GetStackEditorData(), *MatchingOutputNode, ModuleIndex);
+				AddModuleAtIndexItem->SetOnItemAdded(UNiagaraStackAddModuleItem::FOnItemAdded::CreateUObject(this, &UNiagaraStackScriptItemGroup::ItemAdded));
+			}
+
 			UNiagaraStackModuleItem* ModuleItem = FindCurrentChildOfTypeByPredicate<UNiagaraStackModuleItem>(CurrentChildren, 
 				[=](UNiagaraStackModuleItem* CurrentModuleItem) { return &CurrentModuleItem->GetModuleNode() == ModuleNode; });
 
@@ -82,47 +90,82 @@ void UNiagaraStackScriptItemGroup::RefreshChildrenInternal(const TArray<UNiagara
 				ModuleItem->SetOnModifiedGroupItems(UNiagaraStackModuleItem::FOnModifiedGroupItems::CreateUObject(this, &UNiagaraStackScriptItemGroup::ChildModifiedGroupItems));
 			}
 
-			FName ModuleSpacerKey = *FString::Printf(TEXT("Module%i"), ModuleIndex);
-			UNiagaraStackSpacer* ModuleSpacer = FindCurrentChildOfTypeByPredicate<UNiagaraStackSpacer>(CurrentChildren,
-				[=](UNiagaraStackSpacer* CurrentModuleSpacer) { return CurrentModuleSpacer->GetSpacerKey() == ModuleSpacerKey; });
+			NewChildren.Add(AddModuleAtIndexItem);
+			NewChildren.Add(ModuleItem);
+			ModuleIndex++;
+		}
 
-			if (ModuleSpacer == nullptr)
+		bool bForcedError = false;
+		if (ScriptUsage == ENiagaraScriptUsage::SystemUpdateScript)
+		{
+			// We need to make sure that System Update Scripts have the SystemLifecycle script for now.
+			// The factor ensures this, but older assets may not have it or it may have been removed accidentally.
+			// For now, treat this as an error and allow them to resolve.
+			FString ModulePath = TEXT("/Niagara/Modules/System/SystemLifeCycle.SystemLifeCycle");
+			FStringAssetReference SystemUpdateScriptRef(ModulePath);
+			FAssetData ModuleScriptAsset;
+			ModuleScriptAsset.ObjectPath = SystemUpdateScriptRef.GetAssetPathName();
+
+			TArray<UNiagaraNodeFunctionCall*> FoundCalls;
+			if (!FNiagaraStackGraphUtilities::FindScriptModulesInStack(ModuleScriptAsset, *MatchingOutputNode, FoundCalls))
 			{
-				ModuleSpacer = NewObject<UNiagaraStackSpacer>(this);
-				ModuleSpacer->Initialize(GetSystemViewModel(), GetEmitterViewModel(), ModuleSpacerKey);
+				bForcedError = true;
+				Error = MakeShared<FError>();
+				Error->ErrorText = LOCTEXT("SystemLifeCycleWarning", "The stack needs a SystemLifeCycle module.");;
+				Error->ErrorSummaryText = LOCTEXT("MissingRequiredMode", "Missing required module.");
+				Error->Fix.BindLambda([=]()
+				{
+					FScopedTransaction ScopedTransaction(LOCTEXT("AddingSystemLifecycleModule", "Adding System Lifecycle Module."));
+					FNiagaraStackGraphUtilities::AddScriptModuleToStack(ModuleScriptAsset, *MatchingOutputNode);
+				});
 			}
+		}
 
-			NewChildren.Add(ModuleItem);
-			NewChildren.Add(ModuleSpacer);
-			ModuleIndex++;
+		UNiagaraStackSpacer* AddButtonSpacer = FindCurrentChildOfTypeByPredicate<UNiagaraStackSpacer>(CurrentChildren,
+			[=](UNiagaraStackSpacer* CurrentSpacer) { return CurrentSpacer->GetSpacerKey() == "AddButton"; });
+
+		if (AddButtonSpacer == nullptr)
+		{
+			AddButtonSpacer = NewObject<UNiagaraStackSpacer>(this);
+			AddButtonSpacer->Initialize(GetSystemViewModel(), GetEmitterViewModel(), "AddButton");
 		}
 
+		UNiagaraStackAddScriptModuleItem* AddModuleItem = FindCurrentChildOfTypeByPredicate<UNiagaraStackAddScriptModuleItem>(CurrentChildren,
+			[=](UNiagaraStackAddScriptModuleItem* CurrentAddModuleItem) { return CurrentAddModuleItem->GetTargetIndex() == INDEX_NONE; });
+
 		if (AddModuleItem == nullptr)
 		{
 			AddModuleItem = NewObject<UNiagaraStackAddScriptModuleItem>(this);
-			AddModuleItem->Initialize(GetSystemViewModel(), GetEmitterViewModel(), GetStackEditorData(), *MatchingOutputNode);
+			AddModuleItem->Initialize(GetSystemViewModel(), GetEmitterViewModel(), GetStackEditorData(), *MatchingOutputNode, INDEX_NONE);
 			AddModuleItem->SetOnItemAdded(UNiagaraStackAddModuleItem::FOnItemAdded::CreateUObject(this, &UNiagaraStackScriptItemGroup::ItemAdded));
 		}
 
+		UNiagaraStackSpacer* BottomSpacer = FindCurrentChildOfTypeByPredicate<UNiagaraStackSpacer>(CurrentChildren,
+			[=](UNiagaraStackSpacer* CurrentSpacer) { return CurrentSpacer->GetSpacerKey() == "ScriptStackBottom"; });
+
 		if (BottomSpacer == nullptr)
 		{
 			BottomSpacer = NewObject<UNiagaraStackSpacer>(this);
 			BottomSpacer->Initialize(GetSystemViewModel(), GetEmitterViewModel(), "ScriptStackBottom");
 		}
 
+		NewChildren.Add(AddButtonSpacer);
 		NewChildren.Add(AddModuleItem);
 		NewChildren.Add(BottomSpacer);
 
-		ENiagaraScriptCompileStatus Status = ScriptViewModel->GetScriptCompileStatus(GetScriptUsage(), GetScriptOccurrence());
-		if (Status == ENiagaraScriptCompileStatus::NCS_Error)
-		{
-			Error = MakeShared<FError>();
-			Error->ErrorText = ScriptViewModel->GetScriptErrors(GetScriptUsage(), GetScriptOccurrence());
-			Error->ErrorSummaryText = LOCTEXT("ConpileErrorSummary", "The stack has compile errors.");
-		}
-		else
+		ENiagaraScriptCompileStatus Status = ScriptViewModelPinned->GetScriptCompileStatus(GetScriptUsage(), GetScriptUsageId());
+		if (!bForcedError)
 		{
-			Error.Reset();
+			if (Status == ENiagaraScriptCompileStatus::NCS_Error)
+			{
+				Error = MakeShared<FError>();
+				Error->ErrorText = ScriptViewModelPinned->GetScriptErrors(GetScriptUsage(), GetScriptUsageId());
+				Error->ErrorSummaryText = LOCTEXT("ConpileErrorSummary", "The stack has compile errors.");
+			}
+			else
+			{
+				Error.Reset();
+			}
 		}
 	}
 }
@@ -159,12 +202,26 @@ bool UNiagaraStackScriptItemGroup::TryFixError(int32 ErrorIdx)
 
 FText UNiagaraStackScriptItemGroup::GetErrorText(int32 ErrorIdx) const
 {
-	return Error->ErrorText;
+	if (Error.IsValid())
+	{
+		return Error->ErrorText;
+	}
+	else
+	{
+		return FText::GetEmpty();
+	}
 }
 
 FText UNiagaraStackScriptItemGroup::GetErrorSummaryText(int32 ErrorIdx) const
 {
-	return Error->ErrorSummaryText;
+	if (Error.IsValid())
+	{
+		return Error->ErrorSummaryText;
+	}
+	else
+	{
+		return FText::GetEmpty();
+	}
 }
 
 #undef LOCTEXT_NAMESPACE
\ No newline at end of file
diff --git a/Engine/Plugins/FX/Niagara/Source/NiagaraEditor/Private/ViewModels/Stack/NiagaraUnmergeableStackEntryFilter.cpp b/Engine/Plugins/FX/Niagara/Source/NiagaraEditor/Private/ViewModels/Stack/NiagaraUnmergeableStackEntryFilter.cpp
new file mode 100644
--- /dev/null
+++ b/Engine/Plugins/FX/Niagara/Source/NiagaraEditor/Private/ViewModels/Stack/NiagaraUnmergeableStackEntryFilter.cpp
@@ -0,0 +1,19 @@
+// Copyright 1998-2017 Epic Games, Inc. All Rights Reserved.
+
+#include "NiagaraUnmergeableStackEntryFilter.h"
+#include "ViewModels/NiagaraSystemViewModel.h"
+#include "ViewModels/Stack/NiagaraStackEntry.h"
+#include "ViewModels/Stack/NiagaraStackEmitterSpawnScriptItemGroup.h"
+#include "ViewModels/Stack/NiagaraStackSpacer.h"
+#include "ViewModels/Stack/NiagaraStackAddScriptModuleItem.h"
+#include "ViewModels/Stack/NiagaraStackEventHandlerGroup.h"
+#include "ViewModels/Stack/NiagaraStackEventScriptItemGroup.h"
+#include "ViewModels/Stack/NiagaraStackRenderItemGroup.h"
+
+UNiagaraStackEntry::FOnFilterChild FNiagaraUnmergeableStackEntryFilter::CreateFilter()
+{
+	return UNiagaraStackEntry::FOnFilterChild::CreateLambda([](const UNiagaraStackEntry& StackEntry)
+	{
+		return true;
+	});
+}
\ No newline at end of file
diff --git a/Engine/Plugins/FX/Niagara/Source/NiagaraEditor/Private/ViewModels/Stack/NiagaraUnmergeableStackEntryFilter.h b/Engine/Plugins/FX/Niagara/Source/NiagaraEditor/Private/ViewModels/Stack/NiagaraUnmergeableStackEntryFilter.h
new file mode 100644
--- /dev/null
+++ b/Engine/Plugins/FX/Niagara/Source/NiagaraEditor/Private/ViewModels/Stack/NiagaraUnmergeableStackEntryFilter.h
@@ -0,0 +1,10 @@
+// Copyright 1998-2017 Epic Games, Inc. All Rights Reserved.
+
+#pragma once
+
+#include "NiagaraStackEntry.h"
+
+namespace FNiagaraUnmergeableStackEntryFilter
+{
+	UNiagaraStackEntry::FOnFilterChild CreateFilter();
+}
\ No newline at end of file
diff --git a/Engine/Plugins/FX/Niagara/Source/NiagaraEditor/Private/Widgets/SNiagaraGeneratedCodeView.cpp b/Engine/Plugins/FX/Niagara/Source/NiagaraEditor/Private/Widgets/SNiagaraGeneratedCodeView.cpp
--- a/Engine/Plugins/FX/Niagara/Source/NiagaraEditor/Private/Widgets/SNiagaraGeneratedCodeView.cpp
+++ b/Engine/Plugins/FX/Niagara/Source/NiagaraEditor/Private/Widgets/SNiagaraGeneratedCodeView.cpp
@@ -22,6 +22,7 @@
 #include "NiagaraSystem.h"
 #include "NiagaraEditorStyle.h"
 #include "PlatformApplicationMisc.h"
+#include "NiagaraEditorUtilities.h"
 
 #define LOCTEXT_NAMESPACE "NiagaraGeneratedCodeView"
 
@@ -35,129 +36,192 @@ void SNiagaraGeneratedCodeView::Construct(const FArguments& InArgs, TSharedRef<F
 	SystemViewModel->OnSelectedEmitterHandlesChanged().AddRaw(this, &SNiagaraGeneratedCodeView::SelectedEmitterHandlesChanged);
 	SystemViewModel->GetSystemScriptViewModel()->OnSystemCompiled().AddRaw(this, &SNiagaraGeneratedCodeView::OnCodeCompiled);
 
+	TSharedRef<SWidget> HeaderContentsFirstLine = SNew(SHorizontalBox)
+		+ SHorizontalBox::Slot()
+		.AutoWidth()
+		.VAlign(VAlign_Center)
+		.HAlign(HAlign_Left)
+		[
+			SNew(SButton)
+			.HAlign(HAlign_Center)
+			.OnClicked(this, &SNiagaraGeneratedCodeView::OnCopyPressed)
+			.Text(LOCTEXT("CopyOutput", "Copy"))
+			.ToolTipText(LOCTEXT("CopyOutputToolitp", "Press this button to put the contents of this tab in the clipboard."))
+		]
+		+ SHorizontalBox::Slot()
+		[
+			SNullWidget::NullWidget
+		]
+		+ SHorizontalBox::Slot()
+		.AutoWidth()
+		.HAlign(HAlign_Right)
+		.Padding(2, 4, 2, 4)
+		[
+			SAssignNew(SearchBox, SSearchBox)
+			.OnTextCommitted(this, &SNiagaraGeneratedCodeView::OnSearchTextCommitted)
+			.HintText(NSLOCTEXT("SearchBox", "HelpHint", "Search For Text"))
+			.OnTextChanged(this, &SNiagaraGeneratedCodeView::OnSearchTextChanged)
+			.SelectAllTextWhenFocused(false)
+			.DelayChangeNotificationsWhileTyping(true)
+			.MinDesiredWidth(200)
+		]
+		+ SHorizontalBox::Slot()
+		.AutoWidth()
+		.VAlign(VAlign_Center)
+		.Padding(2, 4, 2, 4)
+		[
+			SAssignNew(SearchFoundMOfNText, STextBlock)
+			.MinDesiredWidth(25)
+		]
+		+ SHorizontalBox::Slot()
+		.AutoWidth()
+		.Padding(2, 4, 2, 4)
+		[
+			SNew(SButton)
+			.ButtonStyle(FEditorStyle::Get(), "HoverHintOnly")
+			.IsFocusable(false)
+			.ToolTipText(LOCTEXT("UpToolTip", "Focus to next found search term"))
+			.OnClicked(this, &SNiagaraGeneratedCodeView::SearchUpClicked)
+			.Content()
+			[
+				SNew(STextBlock)
+				.Font(FEditorStyle::Get().GetFontStyle("FontAwesome.10"))
+				.Text(FText::FromString(FString(TEXT("\xf062"))))
+			]
+		]
+		+ SHorizontalBox::Slot()
+		.AutoWidth()
+		.Padding(2, 4, 2, 4)
+		[
+			SNew(SButton)
+			.ButtonStyle(FEditorStyle::Get(), "HoverHintOnly")
+			.IsFocusable(false)
+			.ToolTipText(LOCTEXT("DownToolTip", "Focus to next found search term"))
+			.OnClicked(this, &SNiagaraGeneratedCodeView::SearchDownClicked)
+			.Content()
+			[
+				SNew(STextBlock)
+				.Font(FEditorStyle::Get().GetFontStyle("FontAwesome.10"))
+				.Text(FText::FromString(FString(TEXT("\xf063"))))
+			]
+		];
+
 	this->ChildSlot
+	[
+		SNew(SVerticalBox)
+		+ SVerticalBox::Slot()
+		.AutoHeight() // Header block
 		[
-			SNew(SVerticalBox)
-			+ SVerticalBox::Slot()
-			.AutoHeight()
+			SNew(SBorder)
+			.BorderImage(FEditorStyle::GetBrush(TEXT("ToolPanel.GroupBorder")))
 			[
 				SNew(SVerticalBox)
 				+ SVerticalBox::Slot()
 				.AutoHeight()
 				[
-					SNew(SHorizontalBox)
-					+ SHorizontalBox::Slot()
-					.AutoWidth()
-					.VAlign(VAlign_Center)
-					.HAlign(HAlign_Left)
-					[
-						SNew(SButton)
-						.HAlign(HAlign_Center)
-						.OnClicked(this, &SNiagaraGeneratedCodeView::OnCopyPressed)
-						.Text(LOCTEXT("CopyOutput", "Copy"))
-						.ToolTipText(LOCTEXT("CopyOutputToolitp", "Press this button to put the contents of this tab in the clipboard."))
-					]
-					+ SHorizontalBox::Slot()
-					[
-						SNullWidget::NullWidget
-					] 
+					HeaderContentsFirstLine						
+				]
+				+ SVerticalBox::Slot()
+				.AutoHeight()
+				[
+					SAssignNew(ScriptNameContainer, SHorizontalBox)
 					+ SHorizontalBox::Slot()
 					.AutoWidth()
-					.HAlign(HAlign_Right)
 					.Padding(2, 4, 2, 4)
 					[
-						SAssignNew(SearchBox, SSearchBox)
-						.OnTextCommitted(this, &SNiagaraGeneratedCodeView::OnSearchTextCommitted)
-						.HintText(NSLOCTEXT("SearchBox", "HelpHint", "Search For Text"))
-						.OnTextChanged(this, &SNiagaraGeneratedCodeView::OnSearchTextChanged)
-						.SelectAllTextWhenFocused(false)
-						.DelayChangeNotificationsWhileTyping(true)
-						.MinDesiredWidth(200)
+						SAssignNew(ScriptNameCombo, SComboButton)
+						.OnGetMenuContent(this, &SNiagaraGeneratedCodeView::MakeScriptMenu)
+						.ComboButtonStyle(FEditorStyle::Get(), "ContentBrowser.Filters.Style")
+						.ForegroundColor(FLinearColor::White)
+						.ContentPadding(0)
+						.ToolTipText(LOCTEXT("ScriptsToolTip", "Select a script to view below."))
+						.HasDownArrow(true)
+						.ContentPadding(FMargin(1, 0))
+						.ButtonContent()
+						[
+							SNew(SHorizontalBox)
+							+ SHorizontalBox::Slot()
+								.AutoWidth()
+								.Padding(2, 0, 0, 0)
+								[
+									SNew(STextBlock)
+									.TextStyle(FEditorStyle::Get(), "ContentBrowser.Filters.Text")
+									.Text(LOCTEXT("Scripts", "Scripts"))
+								]
+						]
 					]
 					+ SHorizontalBox::Slot()
 					.AutoWidth()
 					.VAlign(VAlign_Center)
-					.Padding(2, 4, 2, 4)
+					.Padding(20, 4, 2, 4)
 					[
-						SAssignNew(SearchFoundMOfNText, STextBlock)
+						SNew(STextBlock)
 						.MinDesiredWidth(25)
-					]					
-					/*
-					+ SHorizontalBox::Slot()
-					.AutoWidth()
-					.Padding(2, 4, 2, 4)
-					[
-						SNew(SButton)
-						.ButtonStyle(FEditorStyle::Get(), "HoverHintOnly")
-						.IsFocusable(false)
-						.ToolTipText(LOCTEXT("UpToolTip", "Focus to next found search term"))
-						.OnClicked(this, &SNiagaraGeneratedCodeView::SearchUpClicked)
-						.Content()
-						[
-							SNew(STextBlock)
-							.Font(FEditorStyle::Get().GetFontStyle("FontAwesome.10"))
-							.Text(FText::FromString(FString(TEXT("\xf062"))))
-						]
+						.Text(this, &SNiagaraGeneratedCodeView::GetCurrentScriptNameText)
 					]
-					+ SHorizontalBox::Slot()
-					.AutoWidth()
-					.Padding(2, 4, 2, 4)
-					[
-						SNew(SButton)
-						.ButtonStyle(FEditorStyle::Get(), "HoverHintOnly")
-						.IsFocusable(false)
-						.ToolTipText(LOCTEXT("DownToolTip", "Focus to next found search term"))
-						.OnClicked(this, &SNiagaraGeneratedCodeView::SearchDownClicked)
-						.Content()
-						[
-							SNew(STextBlock)
-							.Font(FEditorStyle::Get().GetFontStyle("FontAwesome.10"))
-							.Text(FText::FromString(FString(TEXT("\xf063"))))
-						]
-					]*/
+					// We'll insert here on UI updating..
 				]
 				+ SVerticalBox::Slot()
-				.AutoHeight()
 				[
-					SAssignNew(CheckBoxContainer, SHorizontalBox)
-					// We'll insert here on UI updating..
+					SNew(STextBlock)
+					.Text(LOCTEXT("NoDataText", "Failed to compile or has not been compiled."))
+					.Visibility_Lambda([&]() {
+						if (TabHasScriptData())
+							return EVisibility::Collapsed;
+						return EVisibility::Visible;
+					})
 				]
 			]
-			+ SVerticalBox::Slot()
-			[
-				SNew(STextBlock)
-				.Text(LOCTEXT("NoDataText", "Failed to compile or has not been compiled."))
-				.Visibility_Lambda([&]() {
-					if (TabHasScriptData())
-						return EVisibility::Collapsed;
-					return EVisibility::Visible;
-				})
-			]
-			+ SVerticalBox::Slot()
-			[
-				SAssignNew(TextBodyContainer, SVerticalBox)
-			]
-		];
+		]
+		+ SVerticalBox::Slot() // Text body block
+		[
+			SAssignNew(TextBodyContainer, SVerticalBox)
+		]
+	]; // this->ChildSlot
 	UpdateUI();
 	DoSearch(SearchBox->GetText());
 }
 
-FReply SNiagaraGeneratedCodeView::SearchUpClicked()
+FText SNiagaraGeneratedCodeView::GetCurrentScriptNameText() const
 {
-	CurrentFoundTextEntry--;
-	if (CurrentFoundTextEntry < 0 && ActiveFoundTextEntries.Num() > 0)
+	if (TabState < (uint32)GeneratedCode.Num())
 	{
-		CurrentFoundTextEntry = ActiveFoundTextEntries.Num() - 1;
+		return GeneratedCode[TabState].UsageName;
 	}
-	else if (CurrentFoundTextEntry < 0 && ActiveFoundTextEntries.Num() == 0)
+	else
 	{
-		CurrentFoundTextEntry = 0;
+		return FText::GetEmpty();
+	}
+}
+
+TSharedRef<SWidget> SNiagaraGeneratedCodeView::MakeScriptMenu()
+{
+	FMenuBuilder MenuBuilder(true, nullptr);
+
+	for (int32 i = 0; i < GeneratedCode.Num(); i++)
+	{
+		MenuBuilder.AddMenuEntry(
+			GeneratedCode[i].UsageName,
+			FText::Format(LOCTEXT("MakeScriptMenuTooltip", "View {0}"), GeneratedCode[i].UsageName),
+			FSlateIcon(),
+			FUIAction(FExecuteAction::CreateSP(this, &SNiagaraGeneratedCodeView::OnTabChanged, (uint32)i)));
 	}
 
+	return MenuBuilder.MakeWidget();
+}
+
+FReply SNiagaraGeneratedCodeView::SearchUpClicked()
+{
 	if (ActiveFoundTextEntries.Num() > 0)
 	{
-		//GeneratedCode[TabState].Text->ScrollTo(ActiveFoundTextEntries[CurrentFoundTextEntry]);
+		CurrentFoundTextEntry++;
+		if (CurrentFoundTextEntry == ActiveFoundTextEntries.Num())
+		{
+			CurrentFoundTextEntry = 0;
+		}
 	}
+	
+	GeneratedCode[TabState].Text->AdvanceSearch(true);
 
 	SetSearchMofN();
 
@@ -166,16 +230,15 @@ FReply SNiagaraGeneratedCodeView::SearchUpClicked()
 
 FReply SNiagaraGeneratedCodeView::SearchDownClicked()
 {
-	CurrentFoundTextEntry++;
-	if (CurrentFoundTextEntry >= ActiveFoundTextEntries.Num())
-	{
-		CurrentFoundTextEntry = 0;
-	}
-
 	if (ActiveFoundTextEntries.Num() > 0)
 	{
-		//GeneratedCode[TabState].Text->ScrollTo(ActiveFoundTextEntries[CurrentFoundTextEntry]);
+		CurrentFoundTextEntry--;
+		if (CurrentFoundTextEntry < 0)
+		{
+			CurrentFoundTextEntry = ActiveFoundTextEntries.Num() - 1;
+		}
 	}
+	GeneratedCode[TabState].Text->AdvanceSearch(false);
 	
 	SetSearchMofN();
 
@@ -198,7 +261,7 @@ void SNiagaraGeneratedCodeView::OnSearchTextChanged(const FText& InFilterText)
 
 void SNiagaraGeneratedCodeView::DoSearch(const FText& InFilterText)
 {
-	GeneratedCode[TabState].Text->SetHighlightText(InFilterText);
+	GeneratedCode[TabState].Text->BeginSearch(InFilterText, ESearchCase::IgnoreCase, false);
 	InFilterText.ToString();
 
 	FString SearchString = InFilterText.ToString();
@@ -234,8 +297,8 @@ void SNiagaraGeneratedCodeView::DoSearch(const FText& InFilterText)
 
 void SNiagaraGeneratedCodeView::SetSearchMofN()
 {
-	//SearchFoundMOfNText->SetText(FText::Format(LOCTEXT("MOfN", "{0} of {1}"), FText::AsNumber(CurrentFoundTextEntry), FText::AsNumber(ActiveFoundTextEntries.Num())));
-	SearchFoundMOfNText->SetText(FText::Format(LOCTEXT("MOfN", "{1} found"), FText::AsNumber(CurrentFoundTextEntry), FText::AsNumber(ActiveFoundTextEntries.Num())));
+	SearchFoundMOfNText->SetText(FText::Format(LOCTEXT("MOfN", "{0} of {1}"), FText::AsNumber(CurrentFoundTextEntry), FText::AsNumber(ActiveFoundTextEntries.Num())));
+	//SearchFoundMOfNText->SetText(FText::Format(LOCTEXT("MOfN", "{1} found"), FText::AsNumber(CurrentFoundTextEntry), FText::AsNumber(ActiveFoundTextEntries.Num())));
 }
 
 void SNiagaraGeneratedCodeView::OnSearchTextCommitted(const FText& InFilterText, ETextCommit::Type InCommitType)
@@ -258,12 +321,17 @@ void SNiagaraGeneratedCodeView::SelectedEmitterHandlesChanged()
 void SNiagaraGeneratedCodeView::UpdateUI()
 {
 	TArray<UNiagaraScript*> Scripts;
+	TArray<uint32> ScriptDisplayTypes;
 	UNiagaraSystem& System = SystemViewModel->GetSystem();
 	Scripts.Add(System.GetSystemSpawnScript(false));
+	ScriptDisplayTypes.Add(0);
 	Scripts.Add(System.GetSystemUpdateScript(false));
+	ScriptDisplayTypes.Add(0);
 	Scripts.Add(System.GetSystemSpawnScript(true));
+	ScriptDisplayTypes.Add(0);
 	Scripts.Add(System.GetSystemUpdateScript(true));
-	
+	ScriptDisplayTypes.Add(0);
+
 	TArray<TSharedRef<FNiagaraEmitterHandleViewModel>> SelectedEmitterHandles;
 	SystemViewModel->GetSelectedEmitterHandles(SelectedEmitterHandles);
 	if (SelectedEmitterHandles.Num() == 1)
@@ -274,9 +342,26 @@ void SNiagaraGeneratedCodeView::UpdateUI()
 			TArray<UNiagaraScript*> EmitterScripts;
 			Handle->GetInstance()->GetScripts(EmitterScripts);
 			Scripts.Append(EmitterScripts);
+			ScriptDisplayTypes.AddZeroed(EmitterScripts.Num());
 		}
 	}
 
+	// find the particle spawn script
+	TArray<UNiagaraScript*> DupeScriptsForAssembly;
+	UNiagaraScript *ParticleSpawnScript = nullptr;
+	for (UNiagaraScript *Script : Scripts)
+	{
+		DupeScriptsForAssembly.Add(Script);
+		ScriptDisplayTypes.Add(2);
+		if (Script->Usage == ENiagaraScriptUsage::ParticleSpawnScript || Script->Usage == ENiagaraScriptUsage::ParticleSpawnScriptInterpolated)
+		{
+			ParticleSpawnScript = Script;
+		}
+	}
+	Scripts.Append(DupeScriptsForAssembly);
+	Scripts.Add(ParticleSpawnScript); // add for the GPU update/spawn script
+	ScriptDisplayTypes.Add(1);
+		
 	GeneratedCode.SetNum(Scripts.Num());
 
 	if (TabState >= (uint32)GeneratedCode.Num())
@@ -284,35 +369,87 @@ void SNiagaraGeneratedCodeView::UpdateUI()
 		TabState = 0;
 	}
 
-	CheckBoxContainer->ClearChildren();
 	TextBodyContainer->ClearChildren();
 
 	for (int32 i = 0; i < GeneratedCode.Num(); i++)
 	{
-		GeneratedCode[i].Usage = Scripts[i]->Usage;
-		GeneratedCode[i].UsageIndex = Scripts[i]->UsageIndex;
-
 		TArray<FString> OutputByLines;
 		GeneratedCode[i].Hlsl = FText::GetEmpty();
 		FString SumString;
-		Scripts[i]->LastHlslTranslation.ParseIntoArrayLines(OutputByLines, false);
+
+		bool bIsGPU = ScriptDisplayTypes[i] == 1;
+		bool bIsAssembly = ScriptDisplayTypes[i] == 2;
+		if (Scripts[i] != nullptr)
+		{
+			// GPU combined spawn / update script
+			if (bIsGPU)
+			{
+				GeneratedCode[i].Usage = ENiagaraScriptUsage::ParticleSpawnScript;
+				ParticleSpawnScript->LastHlslTranslationGPU.ParseIntoArrayLines(OutputByLines, false);
+			}
+			else if (bIsAssembly)
+			{
+				GeneratedCode[i].Usage = Scripts[i]->Usage;
+				GeneratedCode[i].UsageId = Scripts[i]->GetUsageId();
+				Scripts[i]->LastAssemblyTranslation.ParseIntoArrayLines(OutputByLines, false);
+			}
+			else
+			{
+				GeneratedCode[i].Usage = Scripts[i]->Usage;
+				GeneratedCode[i].UsageId = Scripts[i]->GetUsageId();
+				Scripts[i]->LastHlslTranslation.ParseIntoArrayLines(OutputByLines, false);
+			}
+		}
+		else
+		{
+			GeneratedCode[i].Usage = ENiagaraScriptUsage::ParticleSpawnScript;
+		}
+
 		GeneratedCode[i].HlslByLines.SetNum(OutputByLines.Num());
-		for (int32 k = 0; k < OutputByLines.Num(); k++)
+		if (bIsAssembly)
 		{
-			GeneratedCode[i].HlslByLines[k] = FString::Printf(TEXT("/*%04d*/\t\t%s\r\n"), k, *OutputByLines[k]);
-			SumString.Append(GeneratedCode[i].HlslByLines[k]);
+			if (Scripts[i] != nullptr)
+			{
+				GeneratedCode[i].Hlsl = FText::FromString(Scripts[i]->LastAssemblyTranslation);
+			}
+			GeneratedCode[i].HlslByLines = OutputByLines;
 		}
-		GeneratedCode[i].Hlsl = FText::FromString(SumString);
-		
-		if (Scripts[i]->Usage == ENiagaraScriptUsage::ParticleEventScript)
+		else
+		{
+			for (int32 k = 0; k < OutputByLines.Num(); k++)
+			{
+				GeneratedCode[i].HlslByLines[k] = FString::Printf(TEXT("/*%04d*/\t\t%s\r\n"), k, *OutputByLines[k]);
+				SumString.Append(GeneratedCode[i].HlslByLines[k]);
+			}
+			GeneratedCode[i].Hlsl = FText::FromString(SumString);
+		}
+		FText AssemblyIdText = LOCTEXT("IsAssembly", "Assembly");
+
+		if (Scripts[i] == nullptr)
+		{
+			GeneratedCode[i].UsageName = LOCTEXT("UsageNameInvalid", "Invalid");
+		}
+		else if (Scripts[i]->Usage == ENiagaraScriptUsage::ParticleEventScript)
+		{
+			FText EventName;
+			if (FNiagaraEditorUtilities::TryGetEventDisplayName(Scripts[i]->GetTypedOuter<UNiagaraEmitter>(), Scripts[i]->GetUsageId(), EventName) == false)
+			{
+				EventName = NSLOCTEXT("NiagaraNodeOutput", "UnknownEventName", "Unknown");
+			}
+			GeneratedCode[i].UsageName = FText::Format(LOCTEXT("UsageNameEvent", "{0}-{1}{2}"), ScriptEnum->GetDisplayNameTextByValue((int64)Scripts[i]->Usage), EventName, bIsAssembly ? AssemblyIdText : FText::GetEmpty());
+		}
+		// GPU combined spawn / update script
+		else if (i == GeneratedCode.Num() - 1 && Scripts[i]->Usage == ENiagaraScriptUsage::ParticleSpawnScript)
 		{
-			GeneratedCode[i].UsageName = FText::Format(LOCTEXT("UsageNameEvent", "{0}[{1}]"), ScriptEnum->GetDisplayNameTextByValue((int64)Scripts[i]->Usage), FText::AsNumber(Scripts[i]->UsageIndex));
+			GeneratedCode[i].UsageName = LOCTEXT("UsageNameGPU", "GPU Spawn/Update");
 		}
 		else
 		{
-			GeneratedCode[i].UsageName = FText::Format(LOCTEXT("UsageName", "{0}"), ScriptEnum->GetDisplayNameTextByValue((int64)Scripts[i]->Usage));
+			GeneratedCode[i].UsageName = FText::Format(LOCTEXT("UsageName", "{0}{1}"), ScriptEnum->GetDisplayNameTextByValue((int64)Scripts[i]->Usage), bIsAssembly ? AssemblyIdText : FText::GetEmpty());
 		}
 
+
+
 		if (!GeneratedCode[i].HorizontalScrollBar.IsValid())
 		{
 			GeneratedCode[i].HorizontalScrollBar = SNew(SScrollBar)
@@ -326,32 +463,7 @@ void SNiagaraGeneratedCodeView::UpdateUI()
 				.Orientation(Orient_Vertical)
 				.Thickness(FVector2D(8.0f, 8.0f));
 		}
-
-		if (!GeneratedCode[i].CheckBox.IsValid())
-		{
-			SAssignNew(GeneratedCode[i].CheckBox, SCheckBox)
-				//.Style(FEditorStyle::Get(), "PlacementBrowser.Tab")
-				//.Style(FEditorStyle::Get(), i == 0 ? "Property.ToggleButton.Start" : (i < GeneratedCode.Num() - 1 ? "Property.ToggleButton.Middle" : "Property.ToggleButton.End"))
-				.OnCheckStateChanged(this, &SNiagaraGeneratedCodeView::OnTabChanged, (uint32)i)
-				.IsChecked(this, &SNiagaraGeneratedCodeView::GetTabCheckedState, (uint32)i)
-				.Style(FEditorStyle::Get(), "PlacementBrowser.Tab")
-				[
-					SNew(SOverlay)
-					+ SOverlay::Slot()
-					.Padding(FMargin(6, 0, 15, 0))
-					.VAlign(VAlign_Center)
-					[
-						SAssignNew(GeneratedCode[i].TextName, STextBlock)
-						.TextStyle(FNiagaraEditorStyle::Get(), "NiagaraEditor.CodeView.Checkbox.Text")
-						.Text(GeneratedCode[i].UsageName)
-					]
-				];
-		}
-		else
-		{
-			GeneratedCode[i].TextName->SetText(GeneratedCode[i].UsageName);
-		}
-
+		
 		if (!GeneratedCode[i].Container.IsValid())
 		{
 			SAssignNew(GeneratedCode[i].Container, SVerticalBox)
@@ -372,7 +484,9 @@ void SNiagaraGeneratedCodeView::UpdateUI()
 							.ExternalScrollbar(GeneratedCode[i].VerticalScrollBar)
 							+ SScrollBox::Slot()
 							[
-								SAssignNew(GeneratedCode[i].Text, STextBlock)
+								SAssignNew(GeneratedCode[i].Text, SMultiLineEditableText)
+								.ClearTextSelectionOnFocusLoss(false)
+								.IsReadOnly(true)
 								.TextStyle(FNiagaraEditorStyle::Get(), "NiagaraEditor.CodeView.Hlsl.Normal")
 							]
 						]
@@ -392,10 +506,6 @@ void SNiagaraGeneratedCodeView::UpdateUI()
 
 		GeneratedCode[i].Text->SetText(GeneratedCode[i].Hlsl);
 
-		CheckBoxContainer->AddSlot()
-			[
-				GeneratedCode[i].CheckBox.ToSharedRef()
-			];
 		TextBodyContainer->AddSlot()
 			[
 				GeneratedCode[i].Container.ToSharedRef()
@@ -413,13 +523,15 @@ SNiagaraGeneratedCodeView::~SNiagaraGeneratedCodeView()
 	
 }
 
-void SNiagaraGeneratedCodeView::OnTabChanged(ECheckBoxState State, uint32 Tab)
+FText SNiagaraGeneratedCodeView::GetSearchText() const
 {
-	if (State == ECheckBoxState::Checked)
-	{
-		TabState = Tab;
-		DoSearch(SearchBox->GetText());
-	}
+	return SearchBox->GetText();
+}
+
+void SNiagaraGeneratedCodeView::OnTabChanged(uint32 Tab)
+{
+	TabState = Tab;
+	DoSearch(SearchBox->GetText());
 }
 
 
@@ -428,9 +540,9 @@ bool SNiagaraGeneratedCodeView::TabHasScriptData() const
 	return !GeneratedCode[TabState].Hlsl.IsEmpty();
 }
 
-ECheckBoxState SNiagaraGeneratedCodeView::GetTabCheckedState(uint32 Tab) const
+bool SNiagaraGeneratedCodeView::GetTabCheckedState(uint32 Tab) const
 {
-	return TabState == Tab ? ECheckBoxState::Checked : ECheckBoxState::Unchecked;
+	return TabState == Tab ? true : false;
 }
 
 EVisibility SNiagaraGeneratedCodeView::GetViewVisibility(uint32 Tab) const
@@ -439,4 +551,4 @@ EVisibility SNiagaraGeneratedCodeView::GetViewVisibility(uint32 Tab) const
 }
 
 
-#undef LOCTEXT_NAMESPACE
\ No newline at end of file
+#undef LOCTEXT_NAMESPACE
diff --git a/Engine/Plugins/FX/Niagara/Source/NiagaraEditor/Private/Widgets/SNiagaraGeneratedCodeView.h b/Engine/Plugins/FX/Niagara/Source/NiagaraEditor/Private/Widgets/SNiagaraGeneratedCodeView.h
--- a/Engine/Plugins/FX/Niagara/Source/NiagaraEditor/Private/Widgets/SNiagaraGeneratedCodeView.h
+++ b/Engine/Plugins/FX/Niagara/Source/NiagaraEditor/Private/Widgets/SNiagaraGeneratedCodeView.h
@@ -16,6 +16,7 @@
 #include "Map.h"
 #include "SCheckBox.h"
 #include "SSearchBox.h"
+#include "SComboButton.h"
 #include "SMultiLineEditableTextBox.h"
 
 class SNiagaraGeneratedCodeView : public SCompoundWidget
@@ -40,19 +41,18 @@ protected:
 		FText UsageName;
 		FText Hlsl;
 		ENiagaraScriptUsage Usage;
-		int32 UsageIndex;
+		FGuid UsageId;
 
 		TArray<FString> HlslByLines;
-		TSharedPtr<STextBlock> Text;
-		TSharedPtr<STextBlock> TextName;
-		TSharedPtr < SCheckBox > CheckBox;
+		TSharedPtr<SMultiLineEditableText> Text;
 		TSharedPtr<SScrollBar> HorizontalScrollBar;
 		TSharedPtr<SScrollBar> VerticalScrollBar;
 		TSharedPtr<SVerticalBox> Container;
 	};
 
 	TArray<TabInfo> GeneratedCode;
-	TSharedPtr<SHorizontalBox> CheckBoxContainer;
+	TSharedPtr<SComboButton> ScriptNameCombo;
+	TSharedPtr<SHorizontalBox> ScriptNameContainer;
 	TSharedPtr<SVerticalBox> TextBodyContainer;
 	TSharedPtr<SSearchBox> SearchBox;
 	TSharedPtr<STextBlock> SearchFoundMOfNText;
@@ -62,8 +62,8 @@ protected:
 	void SetSearchMofN();
 
 	void SelectedEmitterHandlesChanged();
-	void OnTabChanged(ECheckBoxState State,	uint32 Tab);
-	ECheckBoxState GetTabCheckedState(uint32 Tab) const;
+	void OnTabChanged(uint32 Tab);
+	bool GetTabCheckedState(uint32 Tab) const;
 	EVisibility GetViewVisibility(uint32 Tab) const;
 	bool TabHasScriptData() const;
 	FReply OnCopyPressed();
@@ -71,8 +71,11 @@ protected:
 	void OnSearchTextCommitted(const FText& InFilterText, ETextCommit::Type InCommitType);
 	FReply SearchUpClicked();
 	FReply SearchDownClicked();
-
+	TSharedRef<SWidget> MakeScriptMenu();
 	void DoSearch(const FText& InFilterText);
+	FText GetCurrentScriptNameText() const;
+
+	FText GetSearchText() const;
 
 	uint32 TabState;
 
diff --git a/Engine/Plugins/FX/Niagara/Source/NiagaraEditor/Private/Widgets/SNiagaraGraphNodeCustomHlsl.cpp b/Engine/Plugins/FX/Niagara/Source/NiagaraEditor/Private/Widgets/SNiagaraGraphNodeCustomHlsl.cpp
new file mode 100644
--- /dev/null
+++ b/Engine/Plugins/FX/Niagara/Source/NiagaraEditor/Private/Widgets/SNiagaraGraphNodeCustomHlsl.cpp
@@ -0,0 +1,100 @@
+// Copyright 1998-2016 Epic Games, Inc. All Rights Reserved.
+
+#include "SNiagaraGraphNodeCustomHlsl.h"
+#include "NiagaraNodeCustomHlsl.h"
+#include "NiagaraGraph.h"
+#include "Framework/Commands/Commands.h"
+#include "Framework/MultiBox/MultiBoxBuilder.h"
+#include "Widgets/SBoxPanel.h"
+#include "Widgets/Input/SButton.h"
+#include "Widgets/Images/SImage.h"
+#include "ScopedTransaction.h"
+#include "Editor/PropertyEditor/Public/PropertyEditorModule.h"
+#include "ISinglePropertyView.h"
+#include "ModuleManager.h"
+#include "GraphEditorSettings.h"
+#include "SMultiLineEditableTextBox.h"
+
+#define LOCTEXT_NAMESPACE "SNiagaraGraphNodeCustomHlsl"
+
+SNiagaraGraphNodeCustomHlsl::SNiagaraGraphNodeCustomHlsl() : SGraphNode()
+{
+
+}
+
+void SNiagaraGraphNodeCustomHlsl::Construct(const FArguments& InArgs, UEdGraphNode* InGraphNode)
+{
+	GraphNode = InGraphNode;
+	UpdateGraphNode();
+}
+
+bool SNiagaraGraphNodeCustomHlsl::IsNameReadOnly() const
+{
+	return false;
+}
+
+void SNiagaraGraphNodeCustomHlsl::RequestRenameOnSpawn()
+{
+	// We only want to initiate the rename if this is a uniquely added node.
+	UNiagaraNodeCustomHlsl* CustomNode = CastChecked<UNiagaraNodeCustomHlsl>(GraphNode);
+	UNiagaraGraph* Graph = CastChecked<UNiagaraGraph>(CustomNode->GetGraph());
+	TArray<UNiagaraNodeCustomHlsl*> CustomNodes;
+	Graph->GetNodesOfClass<UNiagaraNodeCustomHlsl>(CustomNodes);
+
+	int32 NumMatches = 0;
+	for (UNiagaraNodeCustomHlsl* Node : CustomNodes)
+	{
+		if (Node == CustomNode)
+		{
+			continue;
+		}
+
+		bool bNeedsSync = Node->Signature.Name == CustomNode->Signature.Name;
+		if (bNeedsSync)
+		{
+			NumMatches++;
+		}
+	}
+	if (NumMatches == 0)
+	{
+		RequestRename();
+	}
+}
+
+TSharedRef<SWidget> SNiagaraGraphNodeCustomHlsl::CreateTitleWidget(TSharedPtr<SNodeTitle> NodeTitle)
+{
+	return SNew(SHorizontalBox)
+		+ SHorizontalBox::Slot()
+		.AutoWidth()
+		[
+			SGraphNode::CreateTitleWidget(NodeTitle)
+		];
+}
+
+void SNiagaraGraphNodeCustomHlsl::CreateBelowPinControls(TSharedPtr<SVerticalBox> MainBox)
+{
+	UNiagaraNodeCustomHlsl* CustomNode = CastChecked<UNiagaraNodeCustomHlsl>(GraphNode);
+
+	TAttribute<FText> GetText;
+	GetText.BindUObject(CustomNode, &UNiagaraNodeCustomHlsl::GetHlslText);
+	
+	FOnTextCommitted TextCommit;
+	TextCommit.BindUObject(CustomNode, &UNiagaraNodeCustomHlsl::OnCustomHlslTextCommitted);
+
+	TSharedRef<SMultiLineEditableTextBox> TextBox = SNew(SMultiLineEditableTextBox)
+		.AutoWrapText(false)
+		.Margin(FMargin(5, 5, 5, 5))
+		.Text(GetText)
+		.OnTextCommitted(TextCommit);
+	MainBox->AddSlot()[
+		SNew(SHorizontalBox)
+		+ SHorizontalBox::Slot()
+		.AutoWidth()
+		.Padding(FMargin(10.f, 5.f, 10.f, 10.f))
+		[
+			TextBox
+		]
+	];
+}
+
+#undef LOCTEXT_NAMESPACE
\ No newline at end of file
diff --git a/Engine/Plugins/FX/Niagara/Source/NiagaraEditor/Private/Widgets/SNiagaraGraphNodeCustomHlsl.h b/Engine/Plugins/FX/Niagara/Source/NiagaraEditor/Private/Widgets/SNiagaraGraphNodeCustomHlsl.h
new file mode 100644
--- /dev/null
+++ b/Engine/Plugins/FX/Niagara/Source/NiagaraEditor/Private/Widgets/SNiagaraGraphNodeCustomHlsl.h
@@ -0,0 +1,29 @@
+// Copyright 1998-2016 Epic Games, Inc. All Rights Reserved.
+
+#pragma once
+
+#include "SGraphNode.h"
+#include "SMenuAnchor.h"
+#include "UICommandList.h"
+
+/** A graph node widget representing a niagara input node. */
+class SNiagaraGraphNodeCustomHlsl : public SGraphNode
+{
+public:
+	SLATE_BEGIN_ARGS(SNiagaraGraphNodeCustomHlsl) {}
+	SLATE_END_ARGS();
+
+	SNiagaraGraphNodeCustomHlsl();
+
+	void Construct(const FArguments& InArgs, UEdGraphNode* InGraphNode);
+
+	//~ SGraphNode api
+	virtual TSharedRef<SWidget> CreateTitleWidget(TSharedPtr<SNodeTitle> NodeTitle) override;
+	virtual bool IsNameReadOnly() const override;
+	virtual void RequestRenameOnSpawn() override;
+
+	virtual void CreateBelowPinControls(TSharedPtr<SVerticalBox> MainBox);
+protected:
+
+
+};
\ No newline at end of file
diff --git a/Engine/Plugins/FX/Niagara/Source/NiagaraEditor/Private/Widgets/SNiagaraParameterCollection.cpp b/Engine/Plugins/FX/Niagara/Source/NiagaraEditor/Private/Widgets/SNiagaraParameterCollection.cpp
--- a/Engine/Plugins/FX/Niagara/Source/NiagaraEditor/Private/Widgets/SNiagaraParameterCollection.cpp
+++ b/Engine/Plugins/FX/Niagara/Source/NiagaraEditor/Private/Widgets/SNiagaraParameterCollection.cpp
@@ -490,11 +490,12 @@ TSharedRef<ITableRow> SNiagaraParameterCollection::OnGenerateRowForParameter(TSh
 	if (bCanEdit && Item->GetDefaultValueType() == INiagaraParameterViewModel::EDefaultValueType::Struct)
 	{
 		FNiagaraEditorModule& NiagaraEditorModule = FModuleManager::GetModuleChecked<FNiagaraEditorModule>("NiagaraEditor");
-		TSharedPtr<INiagaraEditorTypeUtilities> TypeEditorUtilities = NiagaraEditorModule.GetTypeUtilities(*Item->GetType().Get());
+		FNiagaraTypeDefinition ParameterType = *Item->GetType().Get();
+		TSharedPtr<INiagaraEditorTypeUtilities> TypeEditorUtilities = NiagaraEditorModule.GetTypeUtilities(ParameterType);
 		TSharedPtr<SNiagaraParameterEditor> ParameterEditor;
 		if (TypeEditorUtilities.IsValid() && TypeEditorUtilities->CanCreateParameterEditor())
 		{
-			ParameterEditor = TypeEditorUtilities->CreateParameterEditor();
+			ParameterEditor = TypeEditorUtilities->CreateParameterEditor(ParameterType);
 			ParameterEditor->UpdateInternalValueFromStruct(Item->GetDefaultValueStruct());
 			ParameterEditor->SetOnBeginValueChange(SNiagaraParameterEditor::FOnValueChange::CreateSP(
 				this, &SNiagaraParameterCollection::ParameterEditorBeginValueChange, Item));
diff --git a/Engine/Plugins/FX/Niagara/Source/NiagaraEditor/Private/Widgets/SNiagaraSelectedEmitterHandle.cpp b/Engine/Plugins/FX/Niagara/Source/NiagaraEditor/Private/Widgets/SNiagaraSelectedEmitterHandle.cpp
--- a/Engine/Plugins/FX/Niagara/Source/NiagaraEditor/Private/Widgets/SNiagaraSelectedEmitterHandle.cpp
+++ b/Engine/Plugins/FX/Niagara/Source/NiagaraEditor/Private/Widgets/SNiagaraSelectedEmitterHandle.cpp
@@ -19,6 +19,7 @@ void SNiagaraSelectedEmitterHandle::Construct(const FArguments& InArgs, TSharedR
 {
 	SystemViewModel = InSystemViewModel;
 	SystemViewModel->OnSelectedEmitterHandlesChanged().AddRaw(this, &SNiagaraSelectedEmitterHandle::SelectedEmitterHandlesChanged);
+	SystemViewModel->OnEmitterHandleViewModelsChanged().AddRaw(this, &SNiagaraSelectedEmitterHandle::EmitterHandleViewModelsChanged);
 	StackViewModel = NewObject<UNiagaraStackViewModel>();
 	FNiagaraEditorModule NiagaraEditorModule = FModuleManager::GetModuleChecked<FNiagaraEditorModule>(TEXT("NiagaraEditor"));
 
@@ -85,6 +86,11 @@ void SNiagaraSelectedEmitterHandle::SelectedEmitterHandlesChanged()
 	RefreshEmitterWidgets();
 }
 
+void SNiagaraSelectedEmitterHandle::EmitterHandleViewModelsChanged()
+{
+	RefreshEmitterWidgets();
+}
+
 EVisibility SNiagaraSelectedEmitterHandle::GetUnsupportedSelectionTextVisibility() const
 {
 	return SystemViewModel->GetSelectedEmitterHandleIds().Num() != 1 ? EVisibility::Visible : EVisibility::Collapsed;
diff --git a/Engine/Plugins/FX/Niagara/Source/NiagaraEditor/Private/Widgets/SNiagaraSelectedEmitterHandle.h b/Engine/Plugins/FX/Niagara/Source/NiagaraEditor/Private/Widgets/SNiagaraSelectedEmitterHandle.h
--- a/Engine/Plugins/FX/Niagara/Source/NiagaraEditor/Private/Widgets/SNiagaraSelectedEmitterHandle.h
+++ b/Engine/Plugins/FX/Niagara/Source/NiagaraEditor/Private/Widgets/SNiagaraSelectedEmitterHandle.h
@@ -27,6 +27,8 @@ public:
 private:
 	void SelectedEmitterHandlesChanged();
 
+	void EmitterHandleViewModelsChanged();
+
 	void RefreshEmitterWidgets();
 
 	EVisibility GetUnsupportedSelectionTextVisibility() const;
