Copying Orion-Staging @ 3748653 (Orion/Dev-General @ 3722124) to //UE4/Main

diff --git a/Engine/Plugins/FX/Niagara/Source/NiagaraEditor/Private/Widgets/SNiagaraSpreadsheetView.cpp b/Engine/Plugins/FX/Niagara/Source/NiagaraEditor/Private/Widgets/SNiagaraSpreadsheetView.cpp
--- a/Engine/Plugins/FX/Niagara/Source/NiagaraEditor/Private/Widgets/SNiagaraSpreadsheetView.cpp
+++ b/Engine/Plugins/FX/Niagara/Source/NiagaraEditor/Private/Widgets/SNiagaraSpreadsheetView.cpp
@@ -25,8 +25,10 @@
 
 #define LOCTEXT_NAMESPACE "SNiagaraSpreadsheetView"
 #define ARRAY_INDEX_COLUMN_NAME TEXT("Array Index")
-#define KEY_COLUMN_NAME TEXT("Property")
+#define OUTPUT_KEY_COLUMN_NAME TEXT("Output Property")
+#define INPUT_KEY_COLUMN_NAME TEXT("Input Property")
 #define VALUE_COLUMN_NAME TEXT("Value")
+#define FILLER_COLUMN_NAME TEXT("__FILLER__")
 
 void SNiagaraSpreadsheetRow::Construct(const FArguments& InArgs, const TSharedRef<STableViewBase>& InOwnerTableView)
 {
@@ -36,6 +38,8 @@ void SNiagaraSpreadsheetRow::Construct(const FArguments& InArgs, const TSharedRe
 
 	SupportedFields = InArgs._SupportedFields;
 	FieldInfoMap = InArgs._FieldInfoMap;
+	UseGlobalOffsets = InArgs._UseGlobalOffsets;
+	ParameterStore = InArgs._ParameterStore;
 
 	SMultiColumnTableRow< TSharedPtr<int32> >::Construct(FSuperRowType::FArguments(), InOwnerTableView);
 }
@@ -43,6 +47,11 @@ void SNiagaraSpreadsheetRow::Construct(const FArguments& InArgs, const TSharedRe
 
 TSharedRef< SWidget > SNiagaraSpreadsheetRow::GenerateWidgetForColumn(const FName& ColumnName)
 {
+	if (ColumnName == FILLER_COLUMN_NAME)
+	{
+		return SNullWidget::NullWidget;
+	}
+
 	TSharedPtr<SWidget> EntryWidget;
 	const SNiagaraSpreadsheetView::FieldInfo* FieldInfo = nullptr;
 	int32 RealRowIdx = 0;
@@ -51,7 +60,7 @@ TSharedRef< SWidget > SNiagaraSpreadsheetRow::GenerateWidgetForColumn(const FNam
 		EntryWidget = SNew(STextBlock)
 			.Text(FText::AsNumber(RowIndex));
 	}
-	else if (!ColumnsAreAttributes && ColumnName == KEY_COLUMN_NAME)
+	else if (!ColumnsAreAttributes && (ColumnName == INPUT_KEY_COLUMN_NAME || ColumnName == OUTPUT_KEY_COLUMN_NAME))
 	{
 		EntryWidget = SNew(STextBlock)
 			.Text(FText::FromName((*SupportedFields.Get())[RowIndex]));
@@ -66,13 +75,22 @@ TSharedRef< SWidget > SNiagaraSpreadsheetRow::GenerateWidgetForColumn(const FNam
 		FieldInfo = FieldInfoMap->Find((*SupportedFields.Get())[RowIndex]);
 	}
 
-	if (FieldInfo != nullptr && DataSet != nullptr && !EntryWidget.IsValid())
+	if (FieldInfo != nullptr && (UseGlobalOffsets ? ParameterStore != nullptr : DataSet != nullptr) && !EntryWidget.IsValid())
 	{
 		
 		if (FieldInfo->bFloat)
 		{
-			uint32 CompBufferOffset = FieldInfo->FloatStartOffset;
-			float* Src = DataSet->PrevData().GetInstancePtrFloat(CompBufferOffset, RealRowIdx);
+			float* Src = nullptr;
+			if (UseGlobalOffsets)
+			{
+				uint32 CompBufferOffset = FieldInfo->GlobalStartOffset;
+				Src = (float*)(ParameterStore->GetParameterDataArray().GetData() + CompBufferOffset);
+			}
+			else
+			{
+				uint32 CompBufferOffset = FieldInfo->FloatStartOffset;
+				Src = DataSet->PrevData().GetInstancePtrFloat(CompBufferOffset, RealRowIdx);
+			}
 
 			EntryWidget = SNew(STextBlock)
 				.Text(FText::AsNumber(Src[0]));
@@ -80,8 +98,17 @@ TSharedRef< SWidget > SNiagaraSpreadsheetRow::GenerateWidgetForColumn(const FNam
 		}
 		else if (FieldInfo->bBoolean)
 		{
-			uint32 CompBufferOffset = FieldInfo->IntStartOffset;
-			int32* Src = DataSet->PrevData().GetInstancePtrInt32(CompBufferOffset, RealRowIdx);
+			int32* Src = nullptr;
+			if (UseGlobalOffsets)
+			{
+				uint32 CompBufferOffset = FieldInfo->GlobalStartOffset;
+				Src = (int32*)(ParameterStore->GetParameterDataArray().GetData() + CompBufferOffset);
+			}
+			else
+			{
+				uint32 CompBufferOffset = FieldInfo->IntStartOffset;
+				Src = DataSet->PrevData().GetInstancePtrInt32(CompBufferOffset, RealRowIdx);
+			}
 			FText ValueText;
 			if (Src[0] == 0)
 			{
@@ -101,16 +128,34 @@ TSharedRef< SWidget > SNiagaraSpreadsheetRow::GenerateWidgetForColumn(const FNam
 		}
 		else if (FieldInfo->Enum.IsValid())
 		{
-			uint32 CompBufferOffset = FieldInfo->IntStartOffset;
-			int32* Src = DataSet->PrevData().GetInstancePtrInt32(CompBufferOffset, RealRowIdx);
+			int32* Src = nullptr;
+			if (UseGlobalOffsets)
+			{
+				uint32 CompBufferOffset = FieldInfo->GlobalStartOffset;
+				Src = (int32*)(ParameterStore->GetParameterDataArray().GetData() + CompBufferOffset);
+			}
+			else
+			{
+				uint32 CompBufferOffset = FieldInfo->IntStartOffset;
+				Src = DataSet->PrevData().GetInstancePtrInt32(CompBufferOffset, RealRowIdx);
+			}
 			EntryWidget = SNew(STextBlock)
 				.Text(FText::Format(LOCTEXT("EnumValue","{0}({1})"), FieldInfo->Enum->GetDisplayNameTextByValue(Src[0]), FText::AsNumber(Src[0])));
 
 		}
 		else
 		{
-			uint32 CompBufferOffset = FieldInfo->IntStartOffset;
-			int32* Src = DataSet->PrevData().GetInstancePtrInt32(CompBufferOffset, RealRowIdx);
+			int32* Src = nullptr;
+			if (UseGlobalOffsets)
+			{
+				uint32 CompBufferOffset = FieldInfo->GlobalStartOffset;
+				Src = (int32*)(ParameterStore->GetParameterDataArray().GetData() + CompBufferOffset);
+			}
+			else
+			{
+				uint32 CompBufferOffset = FieldInfo->IntStartOffset;
+				Src = DataSet->PrevData().GetInstancePtrInt32(CompBufferOffset, RealRowIdx);
+			}
 			EntryWidget = SNew(STextBlock)
 				.Text(FText::AsNumber(Src[0]));
 
@@ -144,8 +189,10 @@ void SNiagaraSpreadsheetView::Construct(const FArguments& InArgs, TSharedRef<FNi
 	CaptureData[(int)UISystemUpdate].TargetUsage = ENiagaraScriptUsage::SystemUpdateScript;
 	CaptureData[(int)UIPerParticleUpdate].ColumnName = LOCTEXT("PerParticleUpdate", "Per-Particle Update");
 	CaptureData[(int)UISystemUpdate].ColumnName = LOCTEXT("SystemUpdate", "System Update");
-	CaptureData[(int)UIPerParticleUpdate].bColumnsAreAttributes = true;
-	CaptureData[(int)UISystemUpdate].bColumnsAreAttributes = false;
+	CaptureData[(int)UIPerParticleUpdate].bOutputColumnsAreAttributes = true;
+	CaptureData[(int)UISystemUpdate].bOutputColumnsAreAttributes = false;
+	CaptureData[(int)UIPerParticleUpdate].bInputColumnsAreAttributes = false;
+	CaptureData[(int)UISystemUpdate].bInputColumnsAreAttributes = false;
 
 	SystemViewModel = InSystemViewModel;
 	SystemViewModel->OnSelectedEmitterHandlesChanged().AddRaw(this, &SNiagaraSpreadsheetView::SelectedEmitterHandlesChanged);
@@ -153,35 +200,64 @@ void SNiagaraSpreadsheetView::Construct(const FArguments& InArgs, TSharedRef<FNi
 
 	for (int32 i = 0; i < (int32)UIMax; i++)
 	{
-		CaptureData[i].HorizontalScrollBar = SNew(SScrollBar)
+		CaptureData[i].OutputHorizontalScrollBar = SNew(SScrollBar)
 			.Orientation(Orient_Horizontal)
 			.Thickness(FVector2D(8.0f, 8.0f));
 
-		CaptureData[i].VerticalScrollBar = SNew(SScrollBar)
+		CaptureData[i].OutputVerticalScrollBar = SNew(SScrollBar)
 			.Orientation(Orient_Vertical)
 			.Thickness(FVector2D(8.0f, 8.0f));
 
+		CaptureData[i].InputHorizontalScrollBar = SNew(SScrollBar)
+			.Orientation(Orient_Horizontal)
+			.Thickness(FVector2D(8.0f, 8.0f));
+
+		CaptureData[i].InputVerticalScrollBar = SNew(SScrollBar)
+			.Orientation(Orient_Vertical)
+			.Thickness(FVector2D(8.0f, 8.0f));
 
-		SAssignNew(CaptureData[i].ListView, STreeView< TSharedPtr<int32> >)
+		SAssignNew(CaptureData[i].OutputsListView, STreeView< TSharedPtr<int32> >)
 			.IsEnabled(this, &SNiagaraSpreadsheetView::IsPausedAtRightTimeOnRightHandle)
 			// List view items are this tall
 			.ItemHeight(12)
 			// Tell the list view where to get its source data
-			.TreeItemsSource(&CaptureData[i].SupportedIndices)
+			.TreeItemsSource(&CaptureData[i].SupportedOutputIndices)
 			// When the list view needs to generate a widget for some data item, use this method
-			.OnGenerateRow(this, &SNiagaraSpreadsheetView::OnGenerateWidgetForList, (EUITab)i)
+			.OnGenerateRow(this, &SNiagaraSpreadsheetView::OnGenerateWidgetForList, (EUITab)i, false)
 			// Given some DataItem, this is how we find out if it has any children and what they are.
-			.OnGetChildren(this, &SNiagaraSpreadsheetView::OnGetChildrenForList, (EUITab)i)
+			.OnGetChildren(this, &SNiagaraSpreadsheetView::OnGetChildrenForList, (EUITab)i, false)
 			// Selection mode
 			.SelectionMode(ESelectionMode::Single)
-			.ExternalScrollbar(CaptureData[i].VerticalScrollBar)
+			.ExternalScrollbar(CaptureData[i].OutputVerticalScrollBar)
 			.ConsumeMouseWheel(EConsumeMouseWheel::Always)
 			.AllowOverscroll(EAllowOverscroll::No)
 			// Selection callback
-			.OnSelectionChanged(this, &SNiagaraSpreadsheetView::OnEventSelectionChanged, (EUITab)i)
+			.OnSelectionChanged(this, &SNiagaraSpreadsheetView::OnEventSelectionChanged, (EUITab)i, false)
 			.HeaderRow
 			(
-				SAssignNew(CaptureData[i].HeaderRow, SHeaderRow)
+				SAssignNew(CaptureData[i].OutputHeaderRow, SHeaderRow)
+			);
+
+		SAssignNew(CaptureData[i].InputsListView, STreeView< TSharedPtr<int32> >)
+			.IsEnabled(this, &SNiagaraSpreadsheetView::IsPausedAtRightTimeOnRightHandle)
+			// List view items are this tall
+			.ItemHeight(12)
+			// Tell the list view where to get its source data
+			.TreeItemsSource(&CaptureData[i].SupportedInputIndices)
+			// When the list view needs to generate a widget for some data item, use this method
+			.OnGenerateRow(this, &SNiagaraSpreadsheetView::OnGenerateWidgetForList, (EUITab)i, true)
+			// Given some DataItem, this is how we find out if it has any children and what they are.
+			.OnGetChildren(this, &SNiagaraSpreadsheetView::OnGetChildrenForList, (EUITab)i, true)
+			// Selection mode
+			.SelectionMode(ESelectionMode::Single)
+			.ExternalScrollbar(CaptureData[i].InputVerticalScrollBar)
+			.ConsumeMouseWheel(EConsumeMouseWheel::Always)
+			.AllowOverscroll(EAllowOverscroll::No)
+			// Selection callback
+			.OnSelectionChanged(this, &SNiagaraSpreadsheetView::OnEventSelectionChanged, (EUITab)i, true)
+			.HeaderRow
+			(
+				SAssignNew(CaptureData[i].InputHeaderRow, SHeaderRow)
 			);
 
 		SAssignNew(CaptureData[i].CheckBox, SCheckBox)
@@ -205,28 +281,53 @@ void SNiagaraSpreadsheetView::Construct(const FArguments& InArgs, TSharedRef<FNi
 		SAssignNew(CaptureData[i].Container, SVerticalBox)
 			.Visibility(this, &SNiagaraSpreadsheetView::GetViewVisibility, (EUITab)i)
 			+ SVerticalBox::Slot()
+			.FillHeight(0.25f)
+			[
+				SNew(SHorizontalBox)
+				+ SHorizontalBox::Slot()
+				[
+					SNew(SScrollBox)
+					.Orientation(Orient_Horizontal)
+					.ExternalScrollbar(CaptureData[i].InputHorizontalScrollBar)
+					+ SScrollBox::Slot()
+					[
+						CaptureData[i].InputsListView.ToSharedRef()
+					]
+				]
+				+ SHorizontalBox::Slot()
+				.AutoWidth()
+				[
+					CaptureData[i].InputVerticalScrollBar.ToSharedRef()
+				]
+			]
+			+ SVerticalBox::Slot()
+			.AutoHeight()
+			[
+				CaptureData[i].InputHorizontalScrollBar.ToSharedRef()
+			]
+			+ SVerticalBox::Slot()
 			[
 				SNew(SHorizontalBox)
 				+ SHorizontalBox::Slot()
 				[
 					SNew(SScrollBox)
 					.Orientation(Orient_Horizontal)
-					.ExternalScrollbar(CaptureData[i].HorizontalScrollBar)
+					.ExternalScrollbar(CaptureData[i].OutputHorizontalScrollBar)
 					+ SScrollBox::Slot()
 					[
-						CaptureData[i].ListView.ToSharedRef()
+						CaptureData[i].OutputsListView.ToSharedRef()
 					]
 				]
 				+ SHorizontalBox::Slot()
 				.AutoWidth()
 				[
-					CaptureData[i].VerticalScrollBar.ToSharedRef()
+					CaptureData[i].OutputVerticalScrollBar.ToSharedRef()
 				]
 			]
 			+ SVerticalBox::Slot()
 			.AutoHeight()
 			[
-				CaptureData[i].HorizontalScrollBar.ToSharedRef()
+				CaptureData[i].OutputHorizontalScrollBar.ToSharedRef()
 			];
 		}
 
@@ -351,14 +452,30 @@ EVisibility SNiagaraSpreadsheetView::GetViewVisibility(EUITab Tab) const
 }
 
 
-TSharedRef< ITableRow > SNiagaraSpreadsheetView::OnGenerateWidgetForList(TSharedPtr<int32> InItem, const TSharedRef<STableViewBase>& OwnerTable, EUITab Tab)
+TSharedRef< ITableRow > SNiagaraSpreadsheetView::OnGenerateWidgetForList(TSharedPtr<int32> InItem, const TSharedRef<STableViewBase>& OwnerTable, EUITab Tab, bool bInputList)
 {
-	return SNew(SNiagaraSpreadsheetRow, OwnerTable)
-		.RowIndex(*InItem)
-		.ColumnsAreAttributes(CaptureData[(int32)Tab].bColumnsAreAttributes)
-		.DataSet(CaptureData[(int32)Tab].DataSet)
-		.SupportedFields(CaptureData[(int32)Tab].SupportedFields)
-		.FieldInfoMap(CaptureData[(int32)Tab].FieldInfoMap);
+	if (bInputList)
+	{
+		return SNew(SNiagaraSpreadsheetRow, OwnerTable)
+			.RowIndex(*InItem)
+			.ColumnsAreAttributes(CaptureData[(int32)Tab].bInputColumnsAreAttributes)
+			.DataSet(nullptr)
+			.SupportedFields(CaptureData[(int32)Tab].SupportedInputFields)
+			.FieldInfoMap(CaptureData[(int32)Tab].InputFieldInfoMap)
+			.UseGlobalOffsets(true)
+			.ParameterStore(&CaptureData[(int32)Tab].InputParams);
+	}
+	else
+	{
+		return SNew(SNiagaraSpreadsheetRow, OwnerTable)
+			.RowIndex(*InItem)
+			.ColumnsAreAttributes(CaptureData[(int32)Tab].bOutputColumnsAreAttributes)
+			.DataSet(CaptureData[(int32)Tab].DataSet)
+			.SupportedFields(CaptureData[(int32)Tab].SupportedOutputFields)
+			.FieldInfoMap(CaptureData[(int32)Tab].OutputFieldInfoMap)
+			.UseGlobalOffsets(false)
+			.ParameterStore(nullptr);
+	}
 }
 
 FText SNiagaraSpreadsheetView::LastCapturedInfoText() const
@@ -377,7 +494,7 @@ FText SNiagaraSpreadsheetView::LastCapturedInfoText() const
 	return LOCTEXT("LastCapturedHandleNameStale", "Captured Info: Out-of-date");
 }
 
-void SNiagaraSpreadsheetView::OnGetChildrenForList(TSharedPtr<int32> InItem, TArray<TSharedPtr<int32>>& OutChildren, EUITab Tab)
+void SNiagaraSpreadsheetView::OnGetChildrenForList(TSharedPtr<int32> InItem, TArray<TSharedPtr<int32>>& OutChildren, EUITab Tab, bool bInputList)
 {
 	OutChildren.Empty();
 }
@@ -389,25 +506,27 @@ void SNiagaraSpreadsheetView::SelectedEmitterHandlesChanged()
 	{
 		CaptureData[i].LastReadWriteId = -1;
 		CaptureData[i].DataSet = nullptr;
-		CaptureData[i].SupportedIndices.SetNum(0);
-		CaptureData[i].ListView->RequestTreeRefresh();
+		CaptureData[i].SupportedInputIndices.SetNum(0);
+		CaptureData[i].SupportedOutputIndices.SetNum(0);
+		CaptureData[i].OutputsListView->RequestTreeRefresh();
+		CaptureData[i].InputsListView->RequestTreeRefresh();
 	}
 	
 }
 
 FReply SNiagaraSpreadsheetView::OnCSVOutputPressed()
 {
-	if (CaptureData[(int32)TabState].SupportedFields.IsValid() && CaptureData[(int32)TabState].FieldInfoMap.IsValid() && IsPausedAtRightTimeOnRightHandle())
+	if (CaptureData[(int32)TabState].SupportedOutputFields.IsValid() && CaptureData[(int32)TabState].OutputFieldInfoMap.IsValid() && IsPausedAtRightTimeOnRightHandle())
 	{
 		FString CSVOutput;
 		int32 SkipIdx = -1;
 		int32 NumWritten = 0;
 		TArray<const SNiagaraSpreadsheetView::FieldInfo*> FieldInfos;
-		FieldInfos.SetNum(CaptureData[(int32)TabState].SupportedFields->Num());
+		FieldInfos.SetNum(CaptureData[(int32)TabState].SupportedOutputFields->Num());
 		FString DelimiterString = TEXT("\t");
-		for (int32 i = 0; i < CaptureData[(int32)TabState].SupportedFields->Num(); i++)
+		for (int32 i = 0; i < CaptureData[(int32)TabState].SupportedOutputFields->Num(); i++)
 		{
-			FName Field = (*CaptureData[(int32)TabState].SupportedFields)[i];
+			FName Field = (*CaptureData[(int32)TabState].SupportedOutputFields)[i];
 			if (Field == ARRAY_INDEX_COLUMN_NAME)
 			{
 				SkipIdx = i;
@@ -419,7 +538,7 @@ FReply SNiagaraSpreadsheetView::OnCSVOutputPressed()
 				CSVOutput += DelimiterString;
 			}
 
-			FieldInfos[i] = CaptureData[(int32)TabState].FieldInfoMap->Find(Field);
+			FieldInfos[i] = CaptureData[(int32)TabState].OutputFieldInfoMap->Find(Field);
 
 			CSVOutput += Field.ToString();
 			NumWritten++;
@@ -430,7 +549,7 @@ FReply SNiagaraSpreadsheetView::OnCSVOutputPressed()
 		for (uint32 RowIndex = 0; RowIndex < CaptureData[(int32)TabState].DataSet->PrevData().GetNumInstances(); RowIndex++)
 		{
 			NumWritten = 0;
-			for (int32 i = 0; i < CaptureData[(int32)TabState].SupportedFields->Num(); i++)
+			for (int32 i = 0; i < CaptureData[(int32)TabState].SupportedOutputFields->Num(); i++)
 			{
 				if (i == SkipIdx)
 				{
@@ -503,7 +622,7 @@ void SNiagaraSpreadsheetView::HandleTimeChange()
 				UNiagaraScript* FoundScript = nullptr;
 				for (UNiagaraScript* Script : Scripts)
 				{
-					if (Script->IsEquivalentUsage(CaptureData[i].TargetUsage))
+					if (Script && Script->IsEquivalentUsage(CaptureData[i].TargetUsage))
 					{
 						FoundScript = Script;
 						break;
@@ -516,6 +635,7 @@ void SNiagaraSpreadsheetView::HandleTimeChange()
 				{
 					CaptureData[i].LastReadWriteId = FoundScript->GetDebuggerInfo().DebugFrameLastWriteId;
 					CaptureData[i].DataSet = &FoundScript->GetDebuggerInfo().DebugFrame;
+					CaptureData[i].InputParams = FoundScript->GetDebuggerInfo().DebugParameters;
 					CaptureData[i].DataSet->Tick(ENiagaraSimTarget::CPUSim); // Force a buffer swap, from here out we read from PrevData.
 
 					CaptureData[i].LastCaptureTime = LocalCaptureTime;
@@ -549,7 +669,7 @@ bool SNiagaraSpreadsheetView::CanCapture() const
 	if (SelectedEmitterHandles.Num() == 1)
 	{
 		FNiagaraEmitterHandle* Handle = SelectedEmitterHandles[0]->GetEmitterHandle();
-		if (Handle && Handle->GetInstance()->SimTarget == ENiagaraSimTarget::CPUSim)
+		if (Handle )
 		{
 			return true;
 		}
@@ -575,20 +695,39 @@ void SNiagaraSpreadsheetView::ResetEntries(EUITab Tab)
 {
 	if (CaptureData[(int32)Tab].DataSet)
 	{
-		int32 NumInstances = CaptureData[(int32)Tab].DataSet->GetPrevNumInstances();
-		if (!CaptureData[(int32)Tab].bColumnsAreAttributes)
 		{
-			NumInstances = CaptureData[(int32)Tab].SupportedFields->Num();
-		}
+			int32 NumInstances = CaptureData[(int32)Tab].DataSet->GetPrevNumInstances();
+			if (!CaptureData[(int32)Tab].bOutputColumnsAreAttributes)
+			{
+				NumInstances = CaptureData[(int32)Tab].SupportedOutputFields->Num();
+			}
 
-		CaptureData[(int32)Tab].SupportedIndices.SetNum(NumInstances);
+			CaptureData[(int32)Tab].SupportedOutputIndices.SetNum(NumInstances);
 
-		for (int32 i = 0; i < NumInstances; i++)
-		{
-			CaptureData[(int32)Tab].SupportedIndices[i] = MakeShared<int32>(i);
+			for (int32 i = 0; i < NumInstances; i++)
+			{
+				CaptureData[(int32)Tab].SupportedOutputIndices[i] = MakeShared<int32>(i);
+			}
+
+			CaptureData[(int32)Tab].OutputsListView->RequestTreeRefresh();
 		}
+	
+		{
+			int32 NumInstances = CaptureData[(int32)Tab].InputParams.GetNumParameters();
+			if (!CaptureData[(int32)Tab].bInputColumnsAreAttributes)
+			{
+				NumInstances = CaptureData[(int32)Tab].SupportedInputFields->Num();
+			}
+
+			CaptureData[(int32)Tab].SupportedInputIndices.SetNum(NumInstances);
+
+			for (int32 i = 0; i < NumInstances; i++)
+			{
+				CaptureData[(int32)Tab].SupportedInputIndices[i] = MakeShared<int32>(i);
+			}
 
-		CaptureData[(int32)Tab].ListView->RequestTreeRefresh();
+			CaptureData[(int32)Tab].InputsListView->RequestTreeRefresh();
+		}
 	}
 }
 
@@ -611,6 +750,7 @@ void SNiagaraSpreadsheetView::GenerateLayoutInfo(FNiagaraTypeLayoutInfo& Layout,
 			Info.bFloat = true;
 			Info.FloatStartOffset = Layout.FloatComponentRegisterOffsets.Num();
 			Info.IntStartOffset = UINT_MAX;
+			Info.GlobalStartOffset = sizeof(float) * Layout.FloatComponentRegisterOffsets.Num() + sizeof(int32) * Layout.Int32ComponentByteOffsets.Num();
 			Info.Enum = nullptr;
 			FieldInfo.Add(Info);
 				
@@ -625,6 +765,7 @@ void SNiagaraSpreadsheetView::GenerateLayoutInfo(FNiagaraTypeLayoutInfo& Layout,
 			Info.bBoolean = Property->IsA(UBoolProperty::StaticClass());
 			Info.FloatStartOffset = -1.0f;
 			Info.IntStartOffset = Layout.Int32ComponentRegisterOffsets.Num();
+			Info.GlobalStartOffset = sizeof(float) * Layout.FloatComponentRegisterOffsets.Num() + sizeof(int32) * Layout.Int32ComponentByteOffsets.Num();
 			Info.Enum = Enum;
 			FieldInfo.Add(Info);
 
@@ -655,93 +796,194 @@ void SNiagaraSpreadsheetView::ResetColumns(EUITab Tab)
 
 	if (CaptureData[(int32)i].DataSet)
 	{
-		CaptureData[(int32)i].HeaderRow->ClearColumns();
+		float ManualWidth = 75.0f;
 
+		// Handle output columns
+		{
+			CaptureData[(int32)i].OutputHeaderRow->ClearColumns();
 
-		CaptureData[(int32)i].SupportedFields = MakeShared<TArray<FName> >();
-		CaptureData[(int32)i].FieldInfoMap = MakeShared<TMap<FName, FieldInfo> >();
-		uint32 TotalFloatComponents = 0;
-		uint32 TotalInt32Components = 0;
+			CaptureData[(int32)i].SupportedOutputFields = MakeShared<TArray<FName> >();
+			CaptureData[(int32)i].OutputFieldInfoMap = MakeShared<TMap<FName, FieldInfo> >();
+			uint32 TotalFloatComponents = 0;
+			uint32 TotalInt32Components = 0;
 
-		TArray<FNiagaraVariable> Variables = CaptureData[(int32)i].DataSet->GetVariables();
+			TArray<FNiagaraVariable> Variables = CaptureData[(int32)i].DataSet->GetVariables();
 
+			TArray<FName> ColumnNames;
 
-		float ManualWidth = 75.0f;
-		TArray<FName> ColumnNames;
+			if (CaptureData[(int32)i].bOutputColumnsAreAttributes)
+			{
+				ColumnNames.Add(ARRAY_INDEX_COLUMN_NAME);
+			}
+			else
+			{
+				ManualWidth = 300.0f;
+				ColumnNames.Add(OUTPUT_KEY_COLUMN_NAME);
+				ColumnNames.Add(VALUE_COLUMN_NAME);
+				ColumnNames.Add(FILLER_COLUMN_NAME);
+			}
 
-		if (CaptureData[(int32)i].bColumnsAreAttributes)
-		{
-			ColumnNames.Add(ARRAY_INDEX_COLUMN_NAME);
-		}
-		else
-		{
-			ManualWidth = 300.0f;
-			ColumnNames.Add(KEY_COLUMN_NAME);
-			ColumnNames.Add(VALUE_COLUMN_NAME);
+			for (const FNiagaraVariable& Var : Variables)
+			{
+				FNiagaraTypeDefinition TypeDef = Var.GetType();
+				const UScriptStruct* Struct = TypeDef.GetScriptStruct();
+				const UEnum* Enum = TypeDef.GetEnum();
+
+				FNiagaraTypeLayoutInfo Layout;
+				TArray<FName> PropertyNames;
+				TArray<SNiagaraSpreadsheetView::FieldInfo> FieldInfos;
+
+				uint32 TotalFloatComponentsBeforeStruct = TotalFloatComponents;
+				uint32 TotalInt32ComponentsBeforeStruct = TotalInt32Components;
+
+				GenerateLayoutInfo(Layout, Struct, Enum, Var.GetName(), PropertyNames, FieldInfos);
+
+				for (int32 VarIdx = 0; VarIdx < PropertyNames.Num(); VarIdx++)
+				{
+
+					{
+						if (FieldInfos[VarIdx].bFloat)
+						{
+							FieldInfos[VarIdx].FloatStartOffset += TotalFloatComponentsBeforeStruct;
+							TotalFloatComponents++;
+						}
+						else
+						{
+							FieldInfos[VarIdx].IntStartOffset += TotalInt32ComponentsBeforeStruct;
+							TotalInt32Components++;
+						}
+
+						CaptureData[(int32)i].SupportedOutputFields->Add(PropertyNames[VarIdx]);
+						CaptureData[(int32)i].OutputFieldInfoMap->Add(PropertyNames[VarIdx], FieldInfos[VarIdx]);
+					}
+
+					if (CaptureData[(int32)i].bOutputColumnsAreAttributes)
+					{
+						ColumnNames.Add(PropertyNames[VarIdx]);
+					}
+				}
+			}
+
+
+			for (int32 ColIdx = 0; ColIdx < ColumnNames.Num(); ColIdx++)
+			{
+				FName ColumnName = ColumnNames[ColIdx];
+				SHeaderRow::FColumn::FArguments ColumnArgs;
+				ColumnArgs
+					.ColumnId(ColumnName)
+					.SortMode(EColumnSortMode::None)
+					.HAlignHeader(HAlign_Center)
+					.VAlignHeader(VAlign_Fill)
+					.HeaderContentPadding(TOptional<FMargin>(2.0f))
+					.HAlignCell(HAlign_Fill)
+					.VAlignCell(VAlign_Fill);
+
+				if (ColumnName != FILLER_COLUMN_NAME)
+				{
+					ColumnArgs.DefaultLabel(FText::FromName(ColumnName));
+					ColumnArgs.ManualWidth(ManualWidth);
+				}
+				else
+				{
+					ColumnArgs.DefaultLabel(FText::FromString(TEXT(" ")));
+					ColumnArgs.ManualWidth(ManualWidth);
+				}
+				CaptureData[(int32)i].OutputHeaderRow->AddColumn(ColumnArgs);
+			}
+
+			CaptureData[(int32)i].OutputHeaderRow->ResetColumnWidths();
+			CaptureData[(int32)i].OutputHeaderRow->RefreshColumns();
+			CaptureData[(int32)i].OutputsListView->RequestTreeRefresh();
 		}
 
-		for (const FNiagaraVariable& Var : Variables)
+		// Handle input columns
 		{
-			FNiagaraTypeDefinition TypeDef = Var.GetType();
-			const UScriptStruct* Struct = TypeDef.GetScriptStruct();
-			const UEnum* Enum = TypeDef.GetEnum();
+			CaptureData[(int32)i].InputHeaderRow->ClearColumns();
+
 
-			FNiagaraTypeLayoutInfo Layout;
-			TArray<FName> PropertyNames;
-			TArray<SNiagaraSpreadsheetView::FieldInfo> FieldInfos;
+			CaptureData[(int32)i].SupportedInputFields = MakeShared<TArray<FName> >();
+			CaptureData[(int32)i].InputFieldInfoMap = MakeShared<TMap<FName, FieldInfo> >();
 
-			uint32 TotalFloatComponentsBeforeStruct = TotalFloatComponents;
-			uint32 TotalInt32ComponentsBeforeStruct = TotalInt32Components;
+			TArray<FNiagaraVariable> Variables;
+			CaptureData[(int32)i].InputParams.GetParameters(Variables);
 
-			GenerateLayoutInfo(Layout, Struct, Enum, Var.GetName(), PropertyNames, FieldInfos);
+			TArray<FName> ColumnNames;
 
-			for (int32 VarIdx = 0; VarIdx < PropertyNames.Num(); VarIdx++)
+			if (CaptureData[(int32)i].bInputColumnsAreAttributes)
 			{
-				
+				ColumnNames.Add(ARRAY_INDEX_COLUMN_NAME);
+			}
+			else
+			{
+				ManualWidth = 300.0f;
+				ColumnNames.Add(INPUT_KEY_COLUMN_NAME);
+				ColumnNames.Add(VALUE_COLUMN_NAME);
+				ColumnNames.Add(FILLER_COLUMN_NAME);
+			}
+
+			for (const FNiagaraVariable& Var : Variables)
+			{
+				FNiagaraTypeDefinition TypeDef = Var.GetType();
+				const UScriptStruct* Struct = TypeDef.GetScriptStruct();
+				const UEnum* Enum = TypeDef.GetEnum();
+
+				FNiagaraTypeLayoutInfo Layout;
+				TArray<FName> PropertyNames;
+				TArray<SNiagaraSpreadsheetView::FieldInfo> FieldInfos;
+
+				int32 ByteOffset = CaptureData[(int32)i].InputParams.IndexOf(Var);
+
+				GenerateLayoutInfo(Layout, Struct, Enum, Var.GetName(), PropertyNames, FieldInfos);
+
+				for (int32 VarIdx = 0; VarIdx < PropertyNames.Num(); VarIdx++)
 				{
-					if (FieldInfos[VarIdx].bFloat)
+
 					{
-						FieldInfos[VarIdx].FloatStartOffset += TotalFloatComponentsBeforeStruct;
-						TotalFloatComponents++;
+						
+						FieldInfos[VarIdx].GlobalStartOffset += ByteOffset;
+
+						CaptureData[(int32)i].SupportedInputFields->Add(PropertyNames[VarIdx]);
+						CaptureData[(int32)i].InputFieldInfoMap->Add(PropertyNames[VarIdx], FieldInfos[VarIdx]);
 					}
-					else
+
+					if (CaptureData[(int32)i].bInputColumnsAreAttributes)
 					{
-						FieldInfos[VarIdx].IntStartOffset += TotalInt32ComponentsBeforeStruct;
-						TotalInt32Components++;
+						ColumnNames.Add(PropertyNames[VarIdx]);
 					}
-
-					CaptureData[(int32)i].SupportedFields->Add(PropertyNames[VarIdx]);
-					CaptureData[(int32)i].FieldInfoMap->Add(PropertyNames[VarIdx], FieldInfos[VarIdx]);
 				}
+			}
+
 
-				if (CaptureData[(int32)i].bColumnsAreAttributes)
+			for (int32 ColIdx = 0; ColIdx < ColumnNames.Num(); ColIdx++)
+			{
+				FName ColumnName = ColumnNames[ColIdx];
+				SHeaderRow::FColumn::FArguments ColumnArgs;
+				ColumnArgs
+					.ColumnId(ColumnName)
+					.SortMode(EColumnSortMode::None)
+					.HAlignHeader(HAlign_Center)
+					.VAlignHeader(VAlign_Fill)
+					.HeaderContentPadding(TOptional<FMargin>(2.0f))
+					.HAlignCell(HAlign_Fill)
+					.VAlignCell(VAlign_Fill);
+
+				if (ColumnName != FILLER_COLUMN_NAME)
+				{
+					ColumnArgs.DefaultLabel(FText::FromName(ColumnName));
+					ColumnArgs.ManualWidth(ManualWidth);
+				}
+				else
 				{
-					ColumnNames.Add(PropertyNames[VarIdx]);
+					ColumnArgs.DefaultLabel(FText::FromString(TEXT(" ")));
+					ColumnArgs.ManualWidth(ManualWidth);
 				}
+				CaptureData[(int32)i].InputHeaderRow->AddColumn(ColumnArgs);
 			}
-		}
-
 
-		for (FName ColumnName : ColumnNames)
-		{
-			SHeaderRow::FColumn::FArguments ColumnArgs;
-
-			ColumnArgs
-				.ColumnId(ColumnName)
-				.DefaultLabel(FText::FromName(ColumnName))
-				.SortMode(EColumnSortMode::None)
-				.HAlignHeader(HAlign_Center)
-				.VAlignHeader(VAlign_Fill)
-				.HeaderContentPadding(TOptional<FMargin>(2.0f))
-				.HAlignCell(HAlign_Fill)
-				.VAlignCell(VAlign_Fill)
-				.ManualWidth(ManualWidth);
-			CaptureData[(int32)i].HeaderRow->AddColumn(ColumnArgs);
+			CaptureData[(int32)i].InputHeaderRow->ResetColumnWidths();
+			CaptureData[(int32)i].InputHeaderRow->RefreshColumns();
+			CaptureData[(int32)i].InputsListView->RequestTreeRefresh();
 		}
-
-		CaptureData[(int32)i].HeaderRow->ResetColumnWidths();
-		CaptureData[(int32)i].HeaderRow->RefreshColumns();
-		CaptureData[(int32)i].ListView->RequestTreeRefresh();
 	}
 }
 
@@ -767,7 +1009,7 @@ FReply SNiagaraSpreadsheetView::OnCaptureRequestPressed()
 		UNiagaraScript* FoundScript = nullptr;
 		for (UNiagaraScript* Script : Scripts)
 		{
-			if (Script->IsEquivalentUsage(CaptureData[(int32)i].TargetUsage))
+			if (Script && Script->IsEquivalentUsage(CaptureData[(int32)i].TargetUsage))
 			{
 				FoundScript = Script;
 				break;
@@ -796,7 +1038,7 @@ FReply SNiagaraSpreadsheetView::OnCaptureRequestPressed()
 	return FReply::Handled();
 }
 
-void SNiagaraSpreadsheetView::OnEventSelectionChanged(TSharedPtr<int32> Selection, ESelectInfo::Type /*SelectInfo*/, EUITab /*Tab*/)
+void SNiagaraSpreadsheetView::OnEventSelectionChanged(TSharedPtr<int32> Selection, ESelectInfo::Type /*SelectInfo*/, EUITab /*Tab*/, bool /*bInputList*/)
 {
 	if (Selection.IsValid())
 	{
diff --git a/Engine/Plugins/FX/Niagara/Source/NiagaraEditor/Private/Widgets/SNiagaraSpreadsheetView.h b/Engine/Plugins/FX/Niagara/Source/NiagaraEditor/Private/Widgets/SNiagaraSpreadsheetView.h
--- a/Engine/Plugins/FX/Niagara/Source/NiagaraEditor/Private/Widgets/SNiagaraSpreadsheetView.h
+++ b/Engine/Plugins/FX/Niagara/Source/NiagaraEditor/Private/Widgets/SNiagaraSpreadsheetView.h
@@ -15,6 +15,7 @@
 #include "SharedPointer.h"
 #include "Map.h"
 #include "SCheckBox.h"
+#include "NiagaraParameterStore.h"
 
 class SNiagaraSpreadsheetView : public SCompoundWidget, public FTickableEditorObject
 {
@@ -35,6 +36,7 @@ public:
 	{
 		uint32 FloatStartOffset;
 		uint32 IntStartOffset;
+		uint32 GlobalStartOffset;
 		bool bFloat;
 		bool bBoolean;
 		TWeakObjectPtr<const UEnum> Enum;
@@ -59,27 +61,36 @@ protected:
 
 	struct CapturedUIData
 	{
-		CapturedUIData() : TargetUsage(ENiagaraScriptUsage::ParticleUpdateScript), LastCaptureTime( -FLT_MAX) , bAwaitingFrame(false), LastReadWriteId(-1),	DataSet( nullptr), bColumnsAreAttributes(true)
+		CapturedUIData() : TargetUsage(ENiagaraScriptUsage::ParticleUpdateScript), LastCaptureTime( -FLT_MAX) , bAwaitingFrame(false), LastReadWriteId(-1),	DataSet( nullptr), bOutputColumnsAreAttributes(true), bInputColumnsAreAttributes(true)
 		{ }
 
-		TSharedPtr<SHeaderRow> HeaderRow;
-		TSharedPtr < STreeView<TSharedPtr<int32> > > ListView;
+		TSharedPtr<SHeaderRow> OutputHeaderRow;
+		TSharedPtr<SHeaderRow> InputHeaderRow;
+		TSharedPtr < STreeView<TSharedPtr<int32> > > OutputsListView;
+		TSharedPtr < STreeView<TSharedPtr<int32> > > InputsListView;
 		TSharedPtr < SCheckBox > CheckBox;
-		TArray< TSharedPtr<int32> > SupportedIndices;
+		TArray< TSharedPtr<int32> > SupportedInputIndices;
+		TArray< TSharedPtr<int32> > SupportedOutputIndices;
 		int32 LastReadWriteId;
 		FNiagaraDataSet* DataSet;
-		TSharedPtr<TArray<FName> > SupportedFields;
-		TSharedPtr<TMap<FName, FieldInfo> > FieldInfoMap;
+		FNiagaraParameterStore InputParams;
+		TSharedPtr<TArray<FName> > SupportedInputFields;
+		TSharedPtr<TArray<FName> > SupportedOutputFields;
+		TSharedPtr<TMap<FName, FieldInfo> > InputFieldInfoMap;
+		TSharedPtr<TMap<FName, FieldInfo> > OutputFieldInfoMap;
 		ENiagaraScriptUsage TargetUsage;
 		bool bAwaitingFrame;
 		float LastCaptureTime;
 		float TargetCaptureTime;
 		FGuid LastCaptureHandleId;
 		TWeakObjectPtr<UNiagaraEmitter> DataSource;
-		TSharedPtr<SScrollBar> HorizontalScrollBar;
-		TSharedPtr<SScrollBar> VerticalScrollBar;
+		TSharedPtr<SScrollBar> OutputHorizontalScrollBar;
+		TSharedPtr<SScrollBar> OutputVerticalScrollBar;
+		TSharedPtr<SScrollBar> InputHorizontalScrollBar;
+		TSharedPtr<SScrollBar> InputVerticalScrollBar;
 		TSharedPtr<SVerticalBox> Container;
-		bool bColumnsAreAttributes;
+		bool bInputColumnsAreAttributes;
+		bool bOutputColumnsAreAttributes;
 		FText ColumnName;
 	};
 
@@ -112,7 +123,7 @@ protected:
 	 * @param Selection Currently selected event
 	 * @param SelectInfo Provides context on how the selection changed
 	 */
-	void OnEventSelectionChanged(TSharedPtr<int32> Selection, ESelectInfo::Type SelectInfo, EUITab Tab);
+	void OnEventSelectionChanged(TSharedPtr<int32> Selection, ESelectInfo::Type SelectInfo, EUITab Tab, bool bInputList);
 
 	/** 
 	 * Generates SEventItem widgets for the events tree
@@ -120,7 +131,7 @@ protected:
 	 * @param InItem Event to generate SEventItem for
 	 * @param OwnerTable Owner Table
 	 */
-	TSharedRef< ITableRow > OnGenerateWidgetForList(TSharedPtr<int32> InItem, const TSharedRef< STableViewBase >& OwnerTable , EUITab Tab);
+	TSharedRef< ITableRow > OnGenerateWidgetForList(TSharedPtr<int32> InItem, const TSharedRef< STableViewBase >& OwnerTable , EUITab Tab, bool bInputList);
 
 	/** 
 	 * Given a profiler event, generates children for it
@@ -128,7 +139,7 @@ protected:
 	 * @param InItem Event to generate children for
 	 * @param OutChildren Generated children
 	 */
-	void OnGetChildrenForList(TSharedPtr<int32> InItem, TArray<TSharedPtr<int32>>& OutChildren, EUITab Tab);
+	void OnGetChildrenForList(TSharedPtr<int32> InItem, TArray<TSharedPtr<int32>>& OutChildren, EUITab Tab, bool bInputList);
 };
 
 
@@ -146,6 +157,8 @@ public:
 	SLATE_ARGUMENT(FNiagaraDataSet*, DataSet)
 	SLATE_ARGUMENT(NamesArray, SupportedFields)
 	SLATE_ARGUMENT(FieldsMap, FieldInfoMap)
+	SLATE_ARGUMENT(bool, UseGlobalOffsets)
+	SLATE_ARGUMENT(FNiagaraParameterStore*, ParameterStore)
 	SLATE_END_ARGS()
 
 		/**
@@ -166,8 +179,10 @@ public:
 private:
 	int32 RowIndex;
 	FNiagaraDataSet* DataSet;
+	FNiagaraParameterStore* ParameterStore;
 	TSharedPtr<TArray<FName> > SupportedFields;
 	TSharedPtr<TMap<FName, SNiagaraSpreadsheetView::FieldInfo> > FieldInfoMap;
 	bool ColumnsAreAttributes;
+	bool UseGlobalOffsets;
 };
 
diff --git a/Engine/Plugins/FX/Niagara/Source/NiagaraEditor/Private/Widgets/SNiagaraSystemViewport.cpp b/Engine/Plugins/FX/Niagara/Source/NiagaraEditor/Private/Widgets/SNiagaraSystemViewport.cpp
--- a/Engine/Plugins/FX/Niagara/Source/NiagaraEditor/Private/Widgets/SNiagaraSystemViewport.cpp
+++ b/Engine/Plugins/FX/Niagara/Source/NiagaraEditor/Private/Widgets/SNiagaraSystemViewport.cpp
@@ -16,6 +16,7 @@
 #include "NiagaraEditorCommands.h"
 #include "EditorViewportCommands.h"
 #include "AdvancedPreviewScene.h"
+#include "ImageUtils.h"
 
 /** Viewport Client for the preview viewport */
 class FNiagaraSystemViewportClient : public FEditorViewportClient
@@ -31,11 +32,18 @@ public:
 	virtual FSceneView* CalcSceneView(FSceneViewFamily* ViewFamily, const EStereoscopicPass StereoPass = eSSP_FULL) override;
 	
 	void SetShowGrid(bool bShowGrid);
+
+	virtual void SetIsSimulateInEditorViewport(bool bInIsSimulateInEditorViewport)override;
+
+	TSharedPtr<SNiagaraSystemViewport> NiagaraViewport;
+	bool bCaptureScreenShot;
 };
 
 FNiagaraSystemViewportClient::FNiagaraSystemViewportClient(FAdvancedPreviewScene& InPreviewScene, const TSharedRef<SNiagaraSystemViewport>& InNiagaraEditorViewport)
 	: FEditorViewportClient(nullptr, &InPreviewScene, StaticCastSharedRef<SEditorViewport>(InNiagaraEditorViewport))
 {
+	NiagaraViewport = InNiagaraEditorViewport;
+
 	// Setup defaults for the common draw helper.
 	DrawHelper.bDrawPivot = false;
 	DrawHelper.bDrawWorldBox = false;
@@ -53,6 +61,10 @@ FNiagaraSystemViewportClient::FNiagaraSystemViewportClient(FAdvancedPreviewScene
 	
 	OverrideNearClipPlane(1.0f);
 	bUsingOrbitCamera = true;
+	bCaptureScreenShot = false;
+
+	//This seems to be needed to get the correct world time in the preview.
+	SetIsSimulateInEditorViewport(true);
 }
 
 
@@ -67,15 +79,54 @@ void FNiagaraSystemViewportClient::Tick(float DeltaSeconds)
 	}
 }
 
-
 void FNiagaraSystemViewportClient::Draw(FViewport* InViewport,FCanvas* Canvas)
 {
+	if (NiagaraViewport.IsValid() && NiagaraViewport->GetDrawElement(SNiagaraSystemViewport::EDrawElements::Bounds))
+	{
+		EngineShowFlags.SetBounds(true);
+		EngineShowFlags.Game = 1;
+	}
+	else
+	{
+		EngineShowFlags.SetBounds(false);
+		EngineShowFlags.Game = 0;
+	}
+
 	FEditorViewportClient::Draw(InViewport, Canvas);
+
+	if (bCaptureScreenShot)
+	{
+		UNiagaraSystem* ParticleSystem = NiagaraViewport->GetPreviewComponent()->GetAsset();
+		int32 SrcWidth = InViewport->GetSizeXY().X;
+		int32 SrcHeight = InViewport->GetSizeXY().Y;
+		// Read the contents of the viewport into an array.
+		TArray<FColor> OrigBitmap;
+		if (InViewport->ReadPixels(OrigBitmap))
+		{
+			check(OrigBitmap.Num() == SrcWidth * SrcHeight);
+
+			// Resize image to enforce max size.
+			TArray<FColor> ScaledBitmap;
+			int32 ScaledWidth = 512;
+			int32 ScaledHeight = 512;
+			FImageUtils::ImageResize(SrcWidth, SrcHeight, OrigBitmap, ScaledWidth, ScaledHeight, ScaledBitmap, true);
+
+			// Compress.
+			FCreateTexture2DParameters Params;
+			Params.bDeferCompression = true;
+			ParticleSystem->ThumbnailImage = FImageUtils::CreateTexture2D(ScaledWidth, ScaledHeight, ScaledBitmap, ParticleSystem, TEXT("ThumbnailTexture"), RF_NoFlags, Params);
+
+			ParticleSystem->ThumbnailImageOutOfDate = false;
+			ParticleSystem->MarkPackageDirty();
+		}
+
+		bCaptureScreenShot = false;
+	}
 }
 
 bool FNiagaraSystemViewportClient::ShouldOrbitCamera() const
 {
-	return true;
+	return bUsingOrbitCamera;
 }
 
 
@@ -99,13 +150,32 @@ void FNiagaraSystemViewportClient::SetShowGrid(bool bShowGrid)
 	DrawHelper.bDrawGrid = bShowGrid;
 }
 
+void FNiagaraSystemViewportClient::SetIsSimulateInEditorViewport(bool bInIsSimulateInEditorViewport)
+{
+	bIsSimulateInEditorViewport = bInIsSimulateInEditorViewport;
+
+// 	if (bInIsSimulateInEditorViewport)
+// 	{
+// 		TSharedRef<FPhysicsManipulationEdModeFactory> Factory = MakeShareable(new FPhysicsManipulationEdModeFactory);
+// 		FEditorModeRegistry::Get().RegisterMode(FBuiltinEditorModes::EM_Physics, Factory);
+// 	}
+// 	else
+// 	{
+// 		FEditorModeRegistry::Get().UnregisterMode(FBuiltinEditorModes::EM_Physics);
+// 	}
+}
+
+//////////////////////////////////////////////////////////////////////////
+
 void SNiagaraSystemViewport::Construct(const FArguments& InArgs)
 {
+	DrawFlags = 0;
 	bShowGrid = false;
 	bShowBackground = false;
 	PreviewComponent = nullptr;
 	AdvancedPreviewScene = MakeShareable(new FAdvancedPreviewScene(FPreviewScene::ConstructionValues()));
-	
+	AdvancedPreviewScene->SetFloorVisibility(false);
+
 	SEditorViewport::Construct( SEditorViewport::FArguments() );
 }
 
@@ -117,6 +187,15 @@ SNiagaraSystemViewport::~SNiagaraSystemViewport()
 	}
 }
 
+void SNiagaraSystemViewport::CreateThumbnail()
+{
+	if (SystemViewportClient.IsValid() && PreviewComponent != nullptr)
+	{
+		SystemViewportClient->bCaptureScreenShot = true;
+	}
+}
+
+
 void SNiagaraSystemViewport::AddReferencedObjects( FReferenceCollector& Collector )
 {
 	if (PreviewComponent != nullptr)
@@ -125,6 +204,27 @@ void SNiagaraSystemViewport::AddReferencedObjects( FReferenceCollector& Collecto
 	}
 }
 
+
+bool SNiagaraSystemViewport::GetDrawElement(EDrawElements Element) const
+{
+	return (DrawFlags & Element) != 0;
+}
+
+void SNiagaraSystemViewport::ToggleDrawElement(EDrawElements Element)
+{
+	DrawFlags = DrawFlags ^ Element;
+}
+
+bool SNiagaraSystemViewport::IsToggleOrbitChecked() const
+{
+	return SystemViewportClient->bUsingOrbitCamera;
+}
+
+void SNiagaraSystemViewport::ToggleOrbit()
+{
+	SystemViewportClient->ToggleOrbitCamera(!SystemViewportClient->bUsingOrbitCamera);
+}
+
 void SNiagaraSystemViewport::RefreshViewport()
 {
 	//reregister the preview components, so if the preview material changed it will be propagated to the render thread
@@ -198,6 +298,12 @@ void SNiagaraSystemViewport::BindCommands()
 		FExecuteAction::CreateSP( this, &SNiagaraSystemViewport::TogglePreviewBackground ),
 								  FCanExecuteAction(),
 								  FIsActionChecked::CreateSP( this, &SNiagaraSystemViewport::IsTogglePreviewBackgroundChecked ) );
+
+	CommandList->MapAction(
+		Commands.ToggleOrbit,
+		FExecuteAction::CreateSP(this, &SNiagaraSystemViewport::ToggleOrbit),
+		FCanExecuteAction(),
+		FIsActionChecked::CreateSP(this, &SNiagaraSystemViewport::IsToggleOrbitChecked));
 								  
 }
 
diff --git a/Engine/Plugins/FX/Niagara/Source/NiagaraEditor/Private/Widgets/SNiagaraSystemViewport.h b/Engine/Plugins/FX/Niagara/Source/NiagaraEditor/Private/Widgets/SNiagaraSystemViewport.h
--- a/Engine/Plugins/FX/Niagara/Source/NiagaraEditor/Private/Widgets/SNiagaraSystemViewport.h
+++ b/Engine/Plugins/FX/Niagara/Source/NiagaraEditor/Private/Widgets/SNiagaraSystemViewport.h
@@ -66,6 +66,21 @@ public:
 	virtual TSharedPtr<FExtender> GetExtenders() const override;
 	virtual void OnFloatingButtonClicked() override;
 	// End of ICommonEditorViewportToolbarInfoProvider interface
+
+	/** Draw flag types */
+	enum EDrawElements
+	{
+		Bounds = 0x020,
+	};
+
+	bool GetDrawElement(EDrawElements Element) const;
+	void ToggleDrawElement(EDrawElements Element);
+	void CreateThumbnail();
+
+
+	bool IsToggleOrbitChecked() const;
+	void ToggleOrbit();
+
 protected:
 	/** SEditorViewport interface */
 	virtual TSharedRef<FEditorViewportClient> MakeEditorViewportClient() override;
@@ -87,8 +102,10 @@ private:
 	/** Pointer back to the material editor tool that owns us */
 	//TWeakPtr<INiagaraSystemEditor> SystemEditorPtr;
 	
-	class UNiagaraComponent *PreviewComponent;
+	class UNiagaraComponent* PreviewComponent;
 	
 	/** Level viewport client */
 	TSharedPtr<class FNiagaraSystemViewportClient> SystemViewportClient;
+
+	uint32 DrawFlags;
 };
diff --git a/Engine/Plugins/FX/Niagara/Source/NiagaraEditor/Private/Widgets/SNiagaraSystemViewportToolBar.cpp b/Engine/Plugins/FX/Niagara/Source/NiagaraEditor/Private/Widgets/SNiagaraSystemViewportToolBar.cpp
--- a/Engine/Plugins/FX/Niagara/Source/NiagaraEditor/Private/Widgets/SNiagaraSystemViewportToolBar.cpp
+++ b/Engine/Plugins/FX/Niagara/Source/NiagaraEditor/Private/Widgets/SNiagaraSystemViewportToolBar.cpp
@@ -4,6 +4,7 @@
 #include "Widgets/Layout/SBorder.h"
 #include "EditorStyleSet.h"
 #include "NiagaraEditorCommands.h"
+#include "EditorViewportCommands.h"
 
 #define LOCTEXT_NAMESPACE "NiagaraSystemViewportToolBar"
 
@@ -53,4 +54,11 @@ bool SNiagaraSystemViewportToolBar::IsViewModeSupported(EViewModeIndex ViewModeI
 	return true; 
 }
 
+void SNiagaraSystemViewportToolBar::ExtendOptionsMenu(FMenuBuilder& OptionsMenuBuilder) const
+{
+	OptionsMenuBuilder.BeginSection("LevelViewportNavigationOptions", LOCTEXT("NavOptionsMenuHeader", "Navigation Options"));
+	OptionsMenuBuilder.AddMenuEntry(FNiagaraEditorCommands::Get().ToggleOrbit);
+	OptionsMenuBuilder.EndSection();
+}
+
 #undef LOCTEXT_NAMESPACE
diff --git a/Engine/Plugins/FX/Niagara/Source/NiagaraEditor/Private/Widgets/SNiagaraSystemViewportToolBar.h b/Engine/Plugins/FX/Niagara/Source/NiagaraEditor/Private/Widgets/SNiagaraSystemViewportToolBar.h
--- a/Engine/Plugins/FX/Niagara/Source/NiagaraEditor/Private/Widgets/SNiagaraSystemViewportToolBar.h
+++ b/Engine/Plugins/FX/Niagara/Source/NiagaraEditor/Private/Widgets/SNiagaraSystemViewportToolBar.h
@@ -20,6 +20,8 @@ public:
 
 	// SCommonEditorViewportToolbarBase interface
 	virtual TSharedRef<SWidget> GenerateShowMenu() const override;
+	virtual void ExtendOptionsMenu(FMenuBuilder& OptionsMenuBuilder) const override;
+
 	// End of SCommonEditorViewportToolbarBase
 
 	virtual bool IsViewModeSupported(EViewModeIndex ViewModeIndex) const override; 
diff --git a/Engine/Plugins/FX/Niagara/Source/NiagaraEditor/Public/EdGraphSchema_Niagara.h b/Engine/Plugins/FX/Niagara/Source/NiagaraEditor/Public/EdGraphSchema_Niagara.h
--- a/Engine/Plugins/FX/Niagara/Source/NiagaraEditor/Public/EdGraphSchema_Niagara.h
+++ b/Engine/Plugins/FX/Niagara/Source/NiagaraEditor/Public/EdGraphSchema_Niagara.h
@@ -73,24 +73,43 @@ class NIAGARAEDITOR_API UEdGraphSchema_Niagara : public UEdGraphSchema
 	virtual bool TryCreateConnection(UEdGraphPin* A, UEdGraphPin* B) const override;
 	virtual void BreakSinglePinLink(UEdGraphPin* SourcePin, UEdGraphPin* TargetPin) const override;
 	virtual FConnectionDrawingPolicy* CreateConnectionDrawingPolicy(int32 InBackLayerID, int32 InFrontLayerID, float InZoomFactor, const FSlateRect& InClippingRect, class FSlateWindowElementList& InDrawElements, class UEdGraph* InGraphObj) const override;
+	virtual void ResetPinToAutogeneratedDefaultValue(UEdGraphPin* Pin, bool bCallModifyCallbacks = true) const override;
 	//~ End EdGraphSchema Interface
 
 	TArray<TSharedPtr<FNiagaraSchemaAction_NewNode> > GetGraphContextActions(const UEdGraph* CurrentGraph, TArray<UObject*>& SelectedObjects, const UEdGraphPin* FromPin, UEdGraph* OwnerOfTemporaries) const;
 	void PromoteSinglePinToParameter(UEdGraphPin* SourcePin);
-
+	void ToggleNodeEnabledState(class UNiagaraNode* InNode) const;
+
+	/** 
+	  * Creates a niagara variable using the name, type, and default value stored on an ed graph pin.
+	  * Pin The pin to create a variable from.
+	  * bNeedsValue Whether or not the returned variable must be allocated to a valid value. When true if the pin doesn't have a valid default value itself the variable will be reset to default before return.
+	  * returns The newly created variable.
+	  */
 	FNiagaraVariable PinToNiagaraVariable(const UEdGraphPin* Pin, bool bNeedsValue=false)const;
-	FString NiagaraVariableToPinDefaultValue(const FNiagaraVariable& Variable) const;
+
+	/** 
+	  * Tries to get a default value string for a graph pin from a niagara variable.
+	  * @param Variable The variable to get a default pin value for.
+	  * @param OutPinDefaultValue The pin default value string if this call was successful, otherwise empty.
+	  * @returns Whether or not we could get a default value form the supplied variable.
+	  */
+	bool TryGetPinDefaultValueFromNiagaraVariable(const FNiagaraVariable& Variable, FString& OutPinDefaultValue) const;
+
 	FNiagaraTypeDefinition PinToTypeDefinition(const UEdGraphPin* Pin)const;
 	FEdGraphPinType TypeDefinitionToPinType(FNiagaraTypeDefinition TypeDef)const;
 	
 	bool IsSystemConstant(const FNiagaraVariable& Variable)const;
 
+	class UNiagaraParameterCollection* VariableIsFromParameterCollection(const FNiagaraVariable& Var)const;
+
 	FNiagaraTypeDefinition GetTypeDefForProperty(const UProperty* Property)const;
 
 	static const FLinearColor NodeTitleColor_Attribute;
 	static const FLinearColor NodeTitleColor_Constant;
 	static const FLinearColor NodeTitleColor_SystemConstant;
 	static const FLinearColor NodeTitleColor_FunctionCall;
+	static const FLinearColor NodeTitleColor_CustomHlsl;
 	static const FLinearColor NodeTitleColor_Event;
 
 private:
diff --git a/Engine/Plugins/FX/Niagara/Source/NiagaraEditor/Public/INiagaraEditorTypeUtilities.h b/Engine/Plugins/FX/Niagara/Source/NiagaraEditor/Public/INiagaraEditorTypeUtilities.h
--- a/Engine/Plugins/FX/Niagara/Source/NiagaraEditor/Public/INiagaraEditorTypeUtilities.h
+++ b/Engine/Plugins/FX/Niagara/Source/NiagaraEditor/Public/INiagaraEditorTypeUtilities.h
@@ -7,6 +7,7 @@
 class FStructOnScope;
 class SNiagaraParameterEditor;
 class SWidget;
+struct FNiagaraTypeDefinition;
 struct FNiagaraVariable;
 
 class INiagaraEditorTypeUtilities
@@ -18,11 +19,11 @@ public:
 
 	virtual bool CanProvideDefaultValue() const = 0;
 
-	virtual void UpdateStructWithDefaultValue(TSharedRef<FStructOnScope> Struct) const = 0;
+	virtual void UpdateVariableWithDefaultValue(FNiagaraVariable& Variable) const = 0;
 
 	virtual bool CanCreateParameterEditor() const = 0;
 
-	virtual TSharedPtr<SNiagaraParameterEditor> CreateParameterEditor() const = 0;
+	virtual TSharedPtr<SNiagaraParameterEditor> CreateParameterEditor(const FNiagaraTypeDefinition& ParameterType) const = 0;
 
 	virtual bool CanCreateDataInterfaceEditor() const = 0;
 
@@ -30,9 +31,9 @@ public:
 
 	virtual bool CanHandlePinDefaults() const = 0;
 
-	virtual FString GetPinDefaultStringFromValue(const FNiagaraVariable& Variable) const = 0;
+	virtual FString GetPinDefaultStringFromValue(const FNiagaraVariable& AllocatedVariable) const = 0;
 
-	virtual void SetValueFromPinDefaultString(const FString& StringValue, FNiagaraVariable& Variable) const = 0;
+	virtual bool SetValueFromPinDefaultString(const FString& StringValue, FNiagaraVariable& Variable) const = 0;
 };
 
 class FNiagaraEditorTypeUtilities : public INiagaraEditorTypeUtilities, public TSharedFromThis<FNiagaraEditorTypeUtilities>
@@ -42,12 +43,12 @@ public:
 public:
 	//~ INiagaraEditorTypeUtilities
 	virtual bool CanProvideDefaultValue() const override { return false; }
-	virtual void UpdateStructWithDefaultValue(TSharedRef<FStructOnScope> Struct) const override { }
+	virtual void UpdateVariableWithDefaultValue(FNiagaraVariable& Variable) const override { }
 	virtual bool CanCreateParameterEditor() const override { return false; }
-	virtual TSharedPtr<SNiagaraParameterEditor> CreateParameterEditor() const override { return TSharedPtr<SNiagaraParameterEditor>(); }
+	virtual TSharedPtr<SNiagaraParameterEditor> CreateParameterEditor(const FNiagaraTypeDefinition& ParameterType) const override { return TSharedPtr<SNiagaraParameterEditor>(); }
 	virtual bool CanCreateDataInterfaceEditor() const override { return false; };
 	virtual TSharedPtr<SWidget> CreateDataInterfaceEditor(UObject* DataInterface, FNotifyValueChanged DataInterfaceChangedHandler) const override { return TSharedPtr<SWidget>(); }
 	virtual bool CanHandlePinDefaults() const override { return false; }
-	virtual FString GetPinDefaultStringFromValue(const FNiagaraVariable& Variable) const override { return FString(); }
-	virtual void SetValueFromPinDefaultString(const FString& StringValue, FNiagaraVariable& Variable) const override { }
+	virtual FString GetPinDefaultStringFromValue(const FNiagaraVariable& AllocatedVariable) const override { return FString(); }
+	virtual bool SetValueFromPinDefaultString(const FString& StringValue, FNiagaraVariable& Variable) const override { return false; }
 };
\ No newline at end of file
diff --git a/Engine/Plugins/FX/Niagara/Source/NiagaraEditor/Public/NiagaraEditorModule.h b/Engine/Plugins/FX/Niagara/Source/NiagaraEditor/Public/NiagaraEditorModule.h
--- a/Engine/Plugins/FX/Niagara/Source/NiagaraEditor/Public/NiagaraEditorModule.h
+++ b/Engine/Plugins/FX/Niagara/Source/NiagaraEditor/Public/NiagaraEditorModule.h
@@ -16,6 +16,7 @@ class INiagaraEditorTypeUtilities;
 class UNiagaraSettings;
 class USequencerSettings;
 class UNiagaraStackViewModel;
+class FNiagaraScriptMergeManager;
 
 DECLARE_STATS_GROUP(TEXT("Niagara Editor"), STATGROUP_NiagaraEditor, STATCAT_Advanced);
 
@@ -57,6 +58,7 @@ public:
 	FDelegateHandle NIAGARAEDITOR_API SetOnCreateStackWidget(FOnCreateStackWidget InOnCreateStackWidget);
 	void NIAGARAEDITOR_API ResetOnCreateStackWidget(FDelegateHandle DelegateHandle);
 
+	TSharedRef<FNiagaraScriptMergeManager> GetScriptMergeManager() const;
 
 	/** Niagara Editor app identifier string */
 	static const FName NiagaraEditorAppIdentifier;
@@ -80,6 +82,7 @@ private:
 	TArray< TSharedPtr<IAssetTypeActions> > CreatedAssetTypeActions;
 
 	TMap<FNiagaraTypeDefinition, TSharedRef<INiagaraEditorTypeUtilities>> TypeToEditorUtilitiesMap;
+	TSharedPtr<INiagaraEditorTypeUtilities> EnumTypeUtilities;
 
 	static EAssetTypeCategories::Type NiagaraAssetCategory;
 
@@ -87,9 +90,12 @@ private:
 	FDelegateHandle CreateSystemTrackEditorHandle;
 	FDelegateHandle CreateFloatParameterTrackEditorHandle;
 	FDelegateHandle CreateBoolParameterTrackEditorHandle;
+	FDelegateHandle MergeEmitterHandle;
 
 	USequencerSettings* SequencerSettings;
 	FOnCreateStackWidget OnCreateStackWidget;
+
+	TSharedPtr<FNiagaraScriptMergeManager> ScriptMergeManager;
 };
 
 
diff --git a/Engine/Plugins/FX/Niagara/Source/NiagaraEditor/Public/NiagaraEditorSettings.h b/Engine/Plugins/FX/Niagara/Source/NiagaraEditor/Public/NiagaraEditorSettings.h
--- a/Engine/Plugins/FX/Niagara/Source/NiagaraEditor/Public/NiagaraEditorSettings.h
+++ b/Engine/Plugins/FX/Niagara/Source/NiagaraEditor/Public/NiagaraEditorSettings.h
@@ -8,7 +8,7 @@
 #include "NiagaraEditorSettings.generated.h"
 
 
-UCLASS(config = Engine, defaultconfig, meta=(DisplayName="Niagara"))
+UCLASS(config = Niagara, defaultconfig, meta=(DisplayName="Niagara"))
 class UNiagaraEditorSettings : public UDeveloperSettings
 {
 public:
@@ -26,6 +26,18 @@ public:
 	UPROPERTY(config, EditAnywhere, Category = Niagara)
 	FStringAssetReference DefaultScript;
 
+	/** Niagara script to duplicate as the base of all new script assets created. */
+	UPROPERTY(config, EditAnywhere, Category = Niagara)
+	FStringAssetReference DefaultDynamicInputScript;
+
+	/** Niagara script to duplicate as the base of all new script assets created. */
+	UPROPERTY(config, EditAnywhere, Category = Niagara)
+	FStringAssetReference DefaultFunctionScript;
+
+	/** Niagara script to duplicate as the base of all new script assets created. */
+	UPROPERTY(config, EditAnywhere, Category = Niagara)
+	FStringAssetReference DefaultModuleScript;
+
 	/** Shortcut key bindings that if held down while doing a mouse click, will spawn the specified type of Niagara node.*/
 	UPROPERTY(config, EditAnywhere, Category = Niagara)
 	TArray<FNiagaraSpawnShortcut> GraphCreationShortcuts;
diff --git a/Engine/Plugins/FX/Niagara/Source/NiagaraEditor/Public/NiagaraEditorUtilities.h b/Engine/Plugins/FX/Niagara/Source/NiagaraEditor/Public/NiagaraEditorUtilities.h
--- a/Engine/Plugins/FX/Niagara/Source/NiagaraEditor/Public/NiagaraEditorUtilities.h
+++ b/Engine/Plugins/FX/Niagara/Source/NiagaraEditor/Public/NiagaraEditorUtilities.h
@@ -62,11 +62,7 @@ namespace FNiagaraEditorUtilities
 
 		return true;
 	}
-
-	/** Builds a unique name from a candidate name and a set of existing names.  The candidate name will be made unique 
-		if necessary by adding a 3 digit index to the end. */
-	FName GetUniqueName(FName CandidateName, const TSet<FName>& ExistingNames);
-
+	
 	/** Gets a set of the system constant names. */
 	TSet<FName> GetSystemConstantNames();
 
@@ -110,10 +106,15 @@ namespace FNiagaraEditorUtilities
 	/** Returns whether the data in a niagara variable and a struct on scope match */
 	bool DataMatches(const FNiagaraVariable& Variable, const FStructOnScope& StructOnScope);
 
+	/** Returns whether the data in two niagara variables match. */
+	bool DataMatches(const FNiagaraVariable& VariableA, const FNiagaraVariable& VariableB);
+
 	/** Returns whether the data in two structs on scope matches. */
 	bool DataMatches(const FStructOnScope& StructOnScopeA, const FStructOnScope& StructOnScopeB);
 
 	TSharedPtr<SWidget> CreateInlineErrorText(TAttribute<FText> ErrorMessage, TAttribute<FText> ErrorTooltip);
 
-	void UpdateExistingEmitters(const TArray<UNiagaraEmitter*>& AffectedEmitters);
+	void CompileExistingEmitters(const TArray<UNiagaraEmitter*>& AffectedEmitters);
+	
+	bool TryGetEventDisplayName(UNiagaraEmitter* Emitter, FGuid EventUsageId, FText& OutEventDisplayName);
 }
\ No newline at end of file
diff --git a/Engine/Plugins/FX/Niagara/Source/NiagaraEditor/Public/NiagaraGraph.h b/Engine/Plugins/FX/Niagara/Source/NiagaraEditor/Public/NiagaraGraph.h
--- a/Engine/Plugins/FX/Niagara/Source/NiagaraEditor/Public/NiagaraGraph.h
+++ b/Engine/Plugins/FX/Niagara/Source/NiagaraEditor/Public/NiagaraGraph.h
@@ -21,6 +21,8 @@ class UNiagaraGraph : public UEdGraph
 {
 	GENERATED_UCLASS_BODY()
 
+	DECLARE_MULTICAST_DELEGATE(FOnDataInterfaceChanged);
+
 	//~ Begin UObject Interface
 	virtual void PostLoad() override;
 	virtual void PostEditChangeProperty(struct FPropertyChangedEvent& PropertyChangedEvent)override;
@@ -33,7 +35,7 @@ class UNiagaraGraph : public UEdGraph
 	bool IsEmpty() const { return Nodes.Num() == 0; }
 			
 	/** Find the first output node bound to the target usage type.*/
-	class UNiagaraNodeOutput* FindOutputNode(ENiagaraScriptUsage TargetUsageType, int32 TargetOccurrence = 0) const;
+	class UNiagaraNodeOutput* FindOutputNode(ENiagaraScriptUsage TargetUsageType, FGuid TargetUsageId = FGuid()) const;
 
 	/** Find all output nodes.*/
 	void FindOutputNodes(TArray<UNiagaraNodeOutput*>& OutputNodes) const;
@@ -50,7 +52,6 @@ class UNiagaraGraph : public UEdGraph
 			, bFilterDuplicates(false)
 			, bFilterByScriptUsage(false)
 			, TargetScriptUsage(ENiagaraScriptUsage::Function)
-			, TargetOccurence(0)
 		{
 		}
 
@@ -68,15 +69,15 @@ class UNiagaraGraph : public UEdGraph
 		bool bFilterByScriptUsage;
 		/** The specified script usage required for an input.*/
 		ENiagaraScriptUsage TargetScriptUsage;
-		/** The specified occurrence within the graph of the script usage*/
-		int32 TargetOccurence;
+		/** The specified id within the graph of the script usage*/
+		FGuid TargetScriptUsageId;
 	};
 
 	/** Finds input nodes in the graph with. */
 	void FindInputNodes(TArray<class UNiagaraNodeInput*>& OutInputNodes, FFindInputNodeOptions Options = FFindInputNodeOptions()) const;
 
 	/** Get an in-order traversal of a graph by the specified target output script usage.*/
-	void BuildTraversal(TArray<class UNiagaraNode*>& OutNodesTraversed, ENiagaraScriptUsage TargetUsage, int32 TargetOccurence = 0) const;
+	void BuildTraversal(TArray<class UNiagaraNode*>& OutNodesTraversed, ENiagaraScriptUsage TargetUsage, FGuid TargetUsageId) const;
 	void BuildTraversal(TArray<class UNiagaraNode*>& OutNodesTraversed, UNiagaraNode* FinalNode) const;
 
 	/** Generates a list of unique input and output parameters for when this script is used as a function. */
@@ -92,9 +93,14 @@ class UNiagaraGraph : public UEdGraph
 
 	bool HasNumericParameters()const;
 
+	bool HasParameterMapParameters()const;
+
 	/** Signal to listeners that the graph has changed */
 	void NotifyGraphNeedsRecompile();
 
+	/** Notifies the graph that a contained data interface has changed. */
+	void NotifyGraphDataInterfaceChanged();
+
 	/** Get all referenced graphs in this specified graph, including this graph. */
 	void GetAllReferencedGraphs(TArray<const UNiagaraGraph*>& Graphs) const;
 
@@ -107,15 +113,37 @@ class UNiagaraGraph : public UEdGraph
 	/** Identify that this graph has undergone changes that will require synchronization with a compiled script.*/
 	void MarkGraphRequiresSynchronization();
 
+	virtual void NotifyGraphChanged() override;
+
 	FGuid GetChangeID() { return ChangeId; }
 
 	/** Walk through the graph for an ParameterMapGet nodes and see if any of them specify a default for VariableName.*/
 	UEdGraphPin* FindParameterMapDefaultValuePin(const FName VariableName);
 
+	/** Get the meta-data associated with this variable, if it exists.*/
+	FNiagaraVariableMetaData* GetMetaData(const FNiagaraVariable& InVar);
+	const FNiagaraVariableMetaData* GetMetaData(const FNiagaraVariable& InVar) const;
+
+	/** Return the meta-data associated with this variable. This should only be called on variables defined explicitly for this Graph, otherwise meta-data may leak.*/
+	FNiagaraVariableMetaData& FindOrAddMetaData(const FNiagaraVariable& InVar);
+
+	/** Remove any meta-data that is no longer being referenced within this graph.*/
+	void PurgeUnreferencedMetaData();
+
+	/** Gets a delegate which is called whenever a contained data interfaces changes. */
+	FOnDataInterfaceChanged& OnDataInterfaceChanged();
+
 private:
+	virtual void NotifyGraphChanged(const FEdGraphEditAction& InAction) override;
 	
 	/** The current change identifier for this graph. Used to sync status with UNiagaraScripts.*/
 	UPROPERTY()
 	FGuid ChangeId;
+
+	/** Storage of meta-data for variables defined for use explicitly with this graph.*/
+	UPROPERTY()
+	TMap<FNiagaraVariable,FNiagaraVariableMetaData> VariableToMetaData;
+
+	FOnDataInterfaceChanged OnDataInterfaceChangedDelegate;
 };
 
diff --git a/Engine/Plugins/FX/Niagara/Source/NiagaraEditor/Public/NiagaraMaterialParameterNode.h b/Engine/Plugins/FX/Niagara/Source/NiagaraEditor/Public/NiagaraMaterialParameterNode.h
new file mode 100644
--- /dev/null
+++ b/Engine/Plugins/FX/Niagara/Source/NiagaraEditor/Public/NiagaraMaterialParameterNode.h
@@ -0,0 +1,40 @@
+// Copyright 1998-2017 Epic Games, Inc. All Rights Reserved.
+
+#pragma once
+
+#include "CoreMinimal.h"
+#include "UObject/ObjectMacros.h"
+#include "NiagaraNodeFunctionCall.h"
+#include "NiagaraMaterialParameterNode.generated.h"
+
+UCLASS(MinimalAPI)
+class UNiagaraMaterialParameterNode : public UNiagaraNodeFunctionCall
+{
+	GENERATED_BODY()
+
+private:
+
+	UPROPERTY()
+	class UNiagaraEmitter* Emitter;
+
+public:
+
+	//~ Begin EdGraphNode Interface
+	virtual void AllocateDefaultPins() override;
+	virtual FText GetNodeTitle(ENodeTitleType::Type TitleType) const override;
+
+	//~ UNiagaraNodeFunctionCall interface
+	virtual bool RefreshFromExternalChanges() override;
+	virtual bool CanRefreshFromModuleItem() const override { return true; };
+
+	/** Set the owning emitter of this node to get the material. */
+	void SetEmitter(class UNiagaraEmitter* InEmitter);
+
+private:		
+	void GenerateScript();
+	void InitializeScript(class UNiagaraScript* NewScript);
+	
+	/** Get a the first material of a renderproperty used by the emitter with a material. */
+	class UMaterial* GetMaterialFromEmitter(class UNiagaraEmitter* InEmitter);
+	class UMaterialExpressionDynamicParameter* GetMaterialExpressionDynamicParameter(class UNiagaraEmitter* InEmitter);
+};
diff --git a/Engine/Plugins/FX/Niagara/Source/NiagaraEditor/Public/NiagaraNode.h b/Engine/Plugins/FX/Niagara/Source/NiagaraEditor/Public/NiagaraNode.h
--- a/Engine/Plugins/FX/Niagara/Source/NiagaraEditor/Public/NiagaraNode.h
+++ b/Engine/Plugins/FX/Niagara/Source/NiagaraEditor/Public/NiagaraNode.h
@@ -3,6 +3,7 @@
 #pragma once
 
 #include "CoreMinimal.h"
+#include "Private/NiagaraHlslTranslator.h"
 #include "UObject/ObjectMacros.h"
 #include "EdGraph/EdGraphNode.h"
 #include "NiagaraEditorCommon.h"
@@ -90,6 +91,11 @@ public:
 	
 protected:
 	virtual int32 CompileInputPin(class FHlslNiagaraTranslator *Translator, UEdGraphPin* Pin);
+	virtual bool IsValidPinToCompile(UEdGraphPin* Pin) const { return true; };
+
+	/** Route input parameter map to output parameter map if it exists. Note that before calling this function,
+		the input pins should have been visited already.*/
+	virtual void RouteParameterMapAroundMe(FNiagaraParameterMapHistoryBuilder& OutHistory, bool bRecursive);
 
 	//ID for the most recent visitor. Allows faster tracking of what nodes have been visited.
 	int32 VisitID;
diff --git a/Engine/Plugins/FX/Niagara/Source/NiagaraEditor/Public/NiagaraNodeCustomHlsl.h b/Engine/Plugins/FX/Niagara/Source/NiagaraEditor/Public/NiagaraNodeCustomHlsl.h
new file mode 100644
--- /dev/null
+++ b/Engine/Plugins/FX/Niagara/Source/NiagaraEditor/Public/NiagaraNodeCustomHlsl.h
@@ -0,0 +1,60 @@
+// Copyright 1998-2017 Epic Games, Inc. All Rights Reserved.
+
+#pragma once
+
+#include "CoreMinimal.h"
+#include "UObject/ObjectMacros.h"
+#include "NiagaraEditorCommon.h"
+#include "NiagaraNodeFunctionCall.h"
+#include "NiagaraNodeInput.h"
+#include "NiagaraNodeCustomHlsl.generated.h"
+
+class UNiagaraScript;
+
+UCLASS(MinimalAPI)
+class UNiagaraNodeCustomHlsl : public UNiagaraNodeFunctionCall
+{
+	GENERATED_UCLASS_BODY()
+
+public:
+	UPROPERTY(EditAnywhere, Category = "Function", meta = (MultiLine = true))
+	FString CustomHlsl;
+
+	UPROPERTY(EditAnywhere, Category = "Function")
+	ENiagaraScriptUsage ScriptUsage;
+
+	virtual TSharedPtr<SGraphNode> CreateVisualWidget() override;
+	virtual void OnRenameNode(const FString& NewName) override;
+	virtual FLinearColor UNiagaraNodeCustomHlsl::GetNodeTitleColor() const override;
+
+	FText GetHlslText() const;
+	void OnCustomHlslTextCommitted(const FText& InText, ETextCommit::Type InType);
+
+#if WITH_EDITOR
+	virtual void PostEditChangeProperty(struct FPropertyChangedEvent& PropertyChangedEvent) override;
+#endif
+protected:
+	virtual bool AllowDynamicPins() const override { return true; }
+
+	virtual bool CanRenamePin(const UEdGraphPin* Pin) const override { return UNiagaraNodeWithDynamicPins::CanRenamePin(Pin); }
+	virtual bool CanRemovePin(const UEdGraphPin* Pin) const override {
+		return UNiagaraNodeWithDynamicPins::CanRemovePin(Pin);
+	}
+	virtual bool CanMovePin(const UEdGraphPin* Pin) const override {
+		return UNiagaraNodeWithDynamicPins::CanMovePin(Pin);
+	}
+
+	/** Called when a new typed pin is added by the user. */
+	virtual void OnNewTypedPinAdded(UEdGraphPin* NewPin) override;
+
+	/** Called when a pin is renamed. */
+	virtual void OnPinRenamed(UEdGraphPin* RenamedPin, const FString& OldPinName) override;
+	
+	/** Removes a pin from this node with a transaction. */
+	virtual void RemoveDynamicPin(UEdGraphPin* Pin) override;
+
+	virtual void MoveDynamicPin(UEdGraphPin* Pin, int32 DirectionToMove) override;
+
+	void RebuildSignatureFromPins();
+
+};
diff --git a/Engine/Plugins/FX/Niagara/Source/NiagaraEditor/Public/NiagaraNodeFunctionCall.h b/Engine/Plugins/FX/Niagara/Source/NiagaraEditor/Public/NiagaraNodeFunctionCall.h
--- a/Engine/Plugins/FX/Niagara/Source/NiagaraEditor/Public/NiagaraNodeFunctionCall.h
+++ b/Engine/Plugins/FX/Niagara/Source/NiagaraEditor/Public/NiagaraNodeFunctionCall.h
@@ -5,14 +5,14 @@
 #include "CoreMinimal.h"
 #include "UObject/ObjectMacros.h"
 #include "NiagaraEditorCommon.h"
-#include "NiagaraNode.h"
+#include "NiagaraNodeWithDynamicPins.h"
 #include "NiagaraNodeInput.h"
 #include "NiagaraNodeFunctionCall.generated.h"
 
 class UNiagaraScript;
 
 UCLASS(MinimalAPI)
-class UNiagaraNodeFunctionCall : public UNiagaraNode
+class UNiagaraNodeFunctionCall : public UNiagaraNodeWithDynamicPins
 {
 	GENERATED_BODY()
 
@@ -33,6 +33,8 @@ public:
 	UPROPERTY()
 	FNiagaraFunctionSignature Signature;
 
+	bool ScriptIsValid() const;
+
 	//Begin UObject interface
 	virtual void PostLoad()override;
 	virtual void PostEditChangeProperty(struct FPropertyChangedEvent& PropertyChangedEvent)override;
@@ -64,10 +66,21 @@ public:
 	/** Walk through the internal script graph for an ParameterMapGet nodes and see if any of them specify a default for VariableName.*/
 	UEdGraphPin* FindParameterMapDefaultValuePin(const FName VariableName);
 
+	virtual bool CanRefreshFromModuleItem() const { return false; };
+
 protected:
+	virtual bool IsValidPinToCompile(UEdGraphPin* Pin) const;
+
+	virtual bool AllowDynamicPins() const override { return false; }
+	virtual bool CanRenamePin(const UEdGraphPin* Pin) const override { return false; }
+	virtual bool CanRemovePin(const UEdGraphPin* Pin) const override { return false; }
+	virtual bool CanMovePin(const UEdGraphPin* Pin) const override { return false; }
+
 	/** Resets the node name based on the referenced script or signature. Guaranteed unique within a given graph instance.*/
 	void ComputeNodeName();
 
+	void SetPinAutoGeneratedDefaultValue(UEdGraphPin& FunctionInputPin, UNiagaraNodeInput& FunctionScriptInputNode);
+
 	/** Adjusted every time that we compile this script. Lets us know that we might differ from any cached versions.*/
 	UPROPERTY()
 	FGuid CachedChangeId;
diff --git a/Engine/Plugins/FX/Niagara/Source/NiagaraEditor/Public/NiagaraNodeInput.h b/Engine/Plugins/FX/Niagara/Source/NiagaraEditor/Public/NiagaraNodeInput.h
--- a/Engine/Plugins/FX/Niagara/Source/NiagaraEditor/Public/NiagaraNodeInput.h
+++ b/Engine/Plugins/FX/Niagara/Source/NiagaraEditor/Public/NiagaraNodeInput.h
@@ -49,10 +49,6 @@ class UNiagaraNodeInput : public UNiagaraNode
 	GENERATED_UCLASS_BODY()
 
 public:
-
-	UPROPERTY(EditAnywhere, Instanced, Category = Input)
-	class UNiagaraDataInterface* DataInterface;
-
 	UPROPERTY(EditAnywhere, Category = Input)
 	FNiagaraVariable Input;
 
@@ -69,6 +65,7 @@ public:
 
 	//~ Begin UObject Interface
 	virtual void PostEditChangeProperty(struct FPropertyChangedEvent& PropertyChangedEvent) override;
+	virtual void PostLoad() override;
 	//~ End UObject Interface
 
 	//~ Begin EdGraphNode Interface
@@ -113,5 +110,16 @@ public:
 
 	/** Given an array of nodes, sort them in place by their sort order, then lexicographically if the same.*/
 	static void SortNodes(TArray<UNiagaraNodeInput*>& InOutNodes);
+
+	UNiagaraDataInterface* GetDataInterface() const;
+
+	void SetDataInterface(UNiagaraDataInterface* InDataInterface);
+
+private:
+	void DataInterfaceChanged();
+
+private:
+	UPROPERTY()
+	class UNiagaraDataInterface* DataInterface;
 };
 
diff --git a/Engine/Plugins/FX/Niagara/Source/NiagaraEditor/Public/NiagaraNodeOutput.h b/Engine/Plugins/FX/Niagara/Source/NiagaraEditor/Public/NiagaraNodeOutput.h
--- a/Engine/Plugins/FX/Niagara/Source/NiagaraEditor/Public/NiagaraNodeOutput.h
+++ b/Engine/Plugins/FX/Niagara/Source/NiagaraEditor/Public/NiagaraNodeOutput.h
@@ -22,7 +22,11 @@ class UNiagaraNodeOutput : public UNiagaraNode
 	ENiagaraScriptUsage ScriptType;
 
 	UPROPERTY()
-	int32 ScriptTypeIndex;
+	FGuid ScriptTypeId;
+
+public:
+	UPROPERTY()
+	int32 ScriptTypeIndex_DEPRECATED;
 
 public:
 
@@ -43,16 +47,17 @@ public:
 	/** Notifies the node that it's output variables have been modified externally. */
 	void NotifyOutputVariablesChanged();
 
-	virtual void Compile(class FHlslNiagaraTranslator *Translator, TArray<int32>& OutputExpressions)override;
+	virtual void Compile(class FHlslNiagaraTranslator *Translator, TArray<int32>& OutputExpressions) override;
 	const TArray<FNiagaraVariable>& GetOutputs() const {return Outputs;}
 	
 	/** Gets the usage of this graph root. */
 	ENiagaraScriptUsage GetUsage() const { return ScriptType; }
 	void SetUsage(ENiagaraScriptUsage InUsage) { ScriptType = InUsage; }
 
-	/** Gets the usage index of this graph root. */
-	int32 GetUsageIndex() const { return ScriptTypeIndex; }
-	void SetUsageIndex(int32 InIndex) { ScriptTypeIndex = InIndex; }
+	/** Gets the id of the usage of this script.  For use when there are multiple scripts with the same usage e.g. Events. */
+	FGuid GetUsageId() const { return ScriptTypeId; }
+	void SetUsageId(FGuid InId) { ScriptTypeId = InId; }
+
 protected:
 	virtual int32 CompileInputPin(class FHlslNiagaraTranslator *Translator, UEdGraphPin* Pin) override;
 
diff --git a/Engine/Plugins/FX/Niagara/Source/NiagaraEditor/Public/NiagaraNodeWithDynamicPins.h b/Engine/Plugins/FX/Niagara/Source/NiagaraEditor/Public/NiagaraNodeWithDynamicPins.h
new file mode 100644
--- /dev/null
+++ b/Engine/Plugins/FX/Niagara/Source/NiagaraEditor/Public/NiagaraNodeWithDynamicPins.h
@@ -0,0 +1,81 @@
+// Copyright 1998-2016 Epic Games, Inc. All Rights Reserved.
+
+#pragma once
+
+#include "NiagaraNode.h"
+#include "SlateTypes.h"
+#include "EdGraph/EdGraph.h"
+
+#include "NiagaraNodeWithDynamicPins.generated.h"
+
+class UEdGraphPin;
+class SNiagaraGraphPinAdd;
+
+
+/** A base node for niagara nodes with pins which can be dynamically added and removed by the user. */
+UCLASS()
+class UNiagaraNodeWithDynamicPins : public UNiagaraNode
+{
+public:
+	GENERATED_BODY()
+
+	//~ UEdGraphNode interface
+	virtual void PinConnectionListChanged(UEdGraphPin* Pin) override;
+	virtual void GetContextMenuActions(const FGraphNodeContextMenuBuilder& Context) const override;
+
+	/** Requests a new pin be added to the node with the specified direction, type, and name. */
+	UEdGraphPin* RequestNewTypedPin(EEdGraphPinDirection Direction, const FNiagaraTypeDefinition& Type, FName InName);
+
+	/** Requests a new pin be added to the node with the specified direction and type. */
+	UEdGraphPin* RequestNewTypedPin(EEdGraphPinDirection Direction, const FNiagaraTypeDefinition& Type);
+
+	/** Helper to identify if a pin is an Add pin.*/
+	bool IsAddPin(const UEdGraphPin* Pin) const;
+
+	/** Determine whether or not a Niagara type is supported for an Add Pin possibility.*/
+	virtual bool AllowNiagaraTypeForAddPin(const FNiagaraTypeDefinition& InType);
+
+	/** Used in tandem with SNiagaraGraphPinAdd to generate the menu for selecting the pin to add.*/
+	virtual TSharedRef<SWidget> GenerateAddPinMenu(const FString& InWorkingPinName, SNiagaraGraphPinAdd* InPin);
+
+protected:
+	virtual bool AllowDynamicPins() const { return true; }
+
+	/** Creates an add pin on the node for the specified direction. */
+	void CreateAddPin(EEdGraphPinDirection Direction);
+
+	/** Called when a new typed pin is added by the user. */
+	virtual void OnNewTypedPinAdded(UEdGraphPin* NewPin) { }
+
+	/** Called when a pin is renamed. */
+	virtual void OnPinRenamed(UEdGraphPin* RenamedPin, const FString& OldPinName) { }
+
+	/** Called to determine if a pin can be renamed by the user. */
+	virtual bool CanRenamePin(const UEdGraphPin* Pin) const;
+
+	/** Called to determine if a pin can be removed by the user. */
+	virtual bool CanRemovePin(const UEdGraphPin* Pin) const;
+
+	/** Called to determine if a pin can be moved by the user.*/
+	virtual bool CanMovePin(const UEdGraphPin* Pin) const;
+
+	/** Used GenerateAddPinMenu to build a list of supported types.*/
+	virtual void BuildTypeMenu(FMenuBuilder& InMenuBuilder, const FString& InWorkingName, SNiagaraGraphPinAdd* InPin);
+
+	/** Removes a pin from this node with a transaction. */
+	virtual void RemoveDynamicPin(UEdGraphPin* Pin);
+
+	virtual void MoveDynamicPin(UEdGraphPin* Pin, int32 DirectionToMove);
+
+private:
+
+	/** Gets the display text for a pin. */
+	FText GetPinNameText(UEdGraphPin* Pin) const;
+
+	/** Called when a pin's name text is committed. */
+	void PinNameTextCommitted(const FText& Text, ETextCommit::Type CommitType, UEdGraphPin* Pin);
+
+public:
+	/** The sub category for add pins. */
+	static const FName AddPinSubCategory;
+};
diff --git a/Engine/Plugins/FX/Niagara/Source/NiagaraEditor/Public/NiagaraParameterMapHistory.h b/Engine/Plugins/FX/Niagara/Source/NiagaraEditor/Public/NiagaraParameterMapHistory.h
--- a/Engine/Plugins/FX/Niagara/Source/NiagaraEditor/Public/NiagaraParameterMapHistory.h
+++ b/Engine/Plugins/FX/Niagara/Source/NiagaraEditor/Public/NiagaraParameterMapHistory.h
@@ -9,6 +9,7 @@
 class UNiagaraNodeOutput;
 class UEdGraphPin;
 class UEdGraphNode;
+class UNiagaraParameterCollection;
 
 /** Traverses a Niagara node graph to identify the variables that have been written and read from a parameter map. 
 * 	This class is meant to aid in UI and compilation of the graph. There are several main script types and each one interacts
@@ -22,6 +23,16 @@ public:
 	/** The variables that have been identified during the traversal. */
 	TArray<FNiagaraVariable> Variables;
 
+	TArray<FNiagaraVariable> VariablesWithOriginalAliasesIntact;
+
+
+	/** Used parameter collections identified during the traversal. TODO: Need to ensure these cannot be GCd if the asset is deleted while it's being used in an in flight compilation. */
+	TArray<UNiagaraParameterCollection*> ParameterCollections;
+	/** Cached off contents of used parameter collections, in case they change during threaded compilation. */
+	TArray<TArray<FNiagaraVariable>> ParameterCollectionVariables;
+	/** Cached off contents of used parameter collections, in case they change during threaded compilation. */
+	TArray<FString> ParameterCollectionNamespaces;
+	
 	/** Are there any warnings that were encountered during the traversal of the graph for a given variable? */
 	TArray<FString> PerVariableWarnings;
 
@@ -36,7 +47,7 @@ public:
 
 	/** List of emitter namespaces encountered as this parameter map was built.*/
 	TArray<FString> EmitterNamespacesEncountered;
-
+	
 	/**
 	* Called in a depth-first traversal to identify a given Niagara Parameter Map pin that was touched during traversal.
 	*/
@@ -46,7 +57,7 @@ public:
 	* Find a variable by name with no concern for type.
 	*/
 	int32 FindVariableByName(const FName& VariableName);
-
+	
 	/**
 	* Find a variable by both name and type. 
 	*/
@@ -105,6 +116,8 @@ public:
 	/** Is this parameter in the special "Engine" namespace?*/
 	static bool IsEngineParameter(const FNiagaraVariable& InVar);
 	static bool IsUserParameter(const FNiagaraVariable& InVar);
+	static bool IsRapidIterationParameter(const FNiagaraVariable& InVar);
+	static bool TryGetEmitterAndFunctionCallNamesFromRapidIterationParameter(const FNiagaraVariable& InVar, FString& EmitterName, FString& FunctionCallName);
 	
 	/** Take an input string and make it hlsl safe.*/
 	static FString MakeSafeNamespaceString(const FString& InStr);
@@ -148,11 +161,22 @@ public:
 	/**
 	* Helper to add a variable to the known list for a parameter map.
 	*/
-	int32 AddVariable(const FNiagaraVariable& InVar, const UEdGraphPin* InPin);
+	int32 AddVariable(const FNiagaraVariable& InVar, const FNiagaraVariable& InAliasedVar, const UEdGraphPin* InPin);
 
 	/** Get the default value for this variable.*/
 	const UEdGraphPin* GetDefaultValuePin(int32 VarIdx) const;
 
+	const FNiagaraVariableMetaData* GetMetaData(int32 VarIdx) const;
+	FNiagaraVariableMetaData* GetMetaData(int32 VarIdx);
+
+	static FNiagaraVariable ConvertVariableToRapidIterationConstantName(FNiagaraVariable InVar, const TCHAR* InEmitterName);
+
+	/**
+	If this is variable is a parameter in one of our tracked collections, return it.
+	@param InVar	Variable to test.
+	@param bMissingParameter	bool set to mark if this parameter was a collection parameter but is now missing from it's collection.
+	*/
+	UNiagaraParameterCollection* IsParameterCollectionParameter(FNiagaraVariable& InVar, bool& bMissingParameter);
 
 };
 
@@ -255,6 +279,9 @@ public:
 	/** Get the node calling this sub-graph.*/
 	UNiagaraNode* GetCallingContext() const;
 
+	/** Are we currently in a top-level function call context (ie a node in the main graph or an argument into a function for the main graph.)*/
+	bool InTopLevelFunctionCall(ENiagaraScriptUsage InFilterScriptType) const;
+
 	/** Helper method to identify any matching input nodes from the calling context node to the input variable.*/
 	int32 FindMatchingParameterMapFromContextInputs(const FNiagaraVariable& InVar) const;
 
@@ -266,6 +293,12 @@ public:
 
 	ENiagaraScriptUsage GetCurrentUsageContext()const;
 
+	bool GetIgnoreDisabled() const { return bIgnoreDisabled; }
+	void SetIgnoreDisabled(bool bInIgnore) { bIgnoreDisabled = bInIgnore; }
+
+	bool IsInEncounteredFunctionNamespace(FNiagaraVariable& InVar) const;
+	bool IsInEncounteredEmitterNamespace(FNiagaraVariable& InVar) const;
+
 protected:
 	/**
 	* Generate the internal alias map from the current traversal state.
@@ -292,9 +325,15 @@ protected:
 	TArray<ENiagaraScriptUsage> RelevantScriptUsageContext;
 	/** Resolved alias map for the current context level. Rebuilt by BuildCurrentAliases.*/
 	TMap<FString, FString> AliasMap;
+
+	TArray<TArray<FString> > EncounteredFunctionNames;
+	TArray<FString> EncounteredEmitterNames;
 	
 	/** Whether or not the script whitelist is active.*/
 	bool bFilterByScriptWhitelist;
 	/** What the script type is that we should be filtering to if the whitelist is enabled.*/
 	ENiagaraScriptUsage FilterScriptType;
+
+	/** Whether or not to ignore disabled nodes.*/
+	bool bIgnoreDisabled;
 };
\ No newline at end of file
diff --git a/Engine/Plugins/FX/Niagara/Source/NiagaraEditor/Public/NiagaraScriptFactoryNew.h b/Engine/Plugins/FX/Niagara/Source/NiagaraEditor/Public/NiagaraScriptFactoryNew.h
--- a/Engine/Plugins/FX/Niagara/Source/NiagaraEditor/Public/NiagaraScriptFactoryNew.h
+++ b/Engine/Plugins/FX/Niagara/Source/NiagaraEditor/Public/NiagaraScriptFactoryNew.h
@@ -6,20 +6,88 @@
 #include "UObject/ObjectMacros.h"
 #include "Factories/Factory.h"
 #include "NiagaraSettings.h"
+#include "NiagaraCommon.h"
 #include "NiagaraScriptFactoryNew.generated.h"
 
-UCLASS(hidecategories=Object)
+UCLASS(hidecategories=Object, ABSTRACT)
 class UNiagaraScriptFactoryNew : public UFactory
 {
 	GENERATED_UCLASS_BODY()
 
-	//~ Begin UFactory Interface
+	// Begin UFactory Interface
 	virtual UObject* FactoryCreateNew(UClass* Class,UObject* InParent,FName Name,EObjectFlags Flags,UObject* Context,FFeedbackContext* Warn) override;
-	//~ Begin UFactory Interface
+	virtual FText GetDisplayName() const override;
+	// End UFactory Interface
 
 public:
 	static void InitializeScript(UNiagaraScript* Script);
+
+protected:
+
+	/** Gets the script asset to create a new one from that based on the script usage type. */
+	const FStringAssetReference& GetDefaultScriptFromSettings(const class UNiagaraEditorSettings* Settings);
+
+	/** Give child factory classes the ability to define which asset action it represents. */
+	virtual FName GetAssetTypeActionName() const { return NAME_None; };
+
+	/** Give child factory classes the ability to define which script usage type it represent. */
+	virtual ENiagaraScriptUsage GetScriptUsage() const;
+};
+
+/**
+ * Niagara module script factory.
+ */
+UCLASS()
+class UNiagaraModuleScriptFactory: public UNiagaraScriptFactoryNew
+{
+	GENERATED_UCLASS_BODY()
+
+	// Begin UFactory Interface
+	virtual FName GetNewAssetThumbnailOverride() const override;
+	// End UFactory Interface
+
+protected:
+
+	// Begin UNiagaraScriptFactoryNew Interface
+	virtual FName GetAssetTypeActionName() const override;
+	virtual ENiagaraScriptUsage GetScriptUsage() const override;
+	// End UNiagaraScriptFactoryNew Interface
+};
+
+/**
+ * Niagara function factory.
+ */
+UCLASS()
+class UNiagaraFunctionScriptFactory: public UNiagaraScriptFactoryNew
+{
+	GENERATED_UCLASS_BODY()
+
+	// Begin UFactory Interface
+	virtual FName GetNewAssetThumbnailOverride() const override;
+	// End UFactory Interface
+
+protected:
+	// Begin UNiagaraScriptFactoryNew Interface
+	virtual FName GetAssetTypeActionName() const override;
+	virtual ENiagaraScriptUsage GetScriptUsage() const override;
+	// End UNiagaraScriptFactoryNew Interface
 };
 
+/**
+ * Niagara dynamic input script factory.
+ */
+UCLASS()
+class UNiagaraDynamicInputScriptFactory: public UNiagaraScriptFactoryNew
+{
+	GENERATED_UCLASS_BODY()
 
+	// Begin UFactory Interface
+	virtual FName GetNewAssetThumbnailOverride() const override;
+	// End UFactory Interface
 
+protected:
+	// Begin UNiagaraScriptFactoryNew Interface
+	virtual FName GetAssetTypeActionName() const override;
+	virtual ENiagaraScriptUsage GetScriptUsage() const override;
+	// End UNiagaraScriptFactoryNew Interface
+};
diff --git a/Engine/Plugins/FX/Niagara/Source/NiagaraEditor/Public/NiagaraScriptSource.h b/Engine/Plugins/FX/Niagara/Source/NiagaraEditor/Public/NiagaraScriptSource.h
--- a/Engine/Plugins/FX/Niagara/Source/NiagaraEditor/Public/NiagaraScriptSource.h
+++ b/Engine/Plugins/FX/Niagara/Source/NiagaraEditor/Public/NiagaraScriptSource.h
@@ -8,6 +8,7 @@
 #include "NiagaraScriptSourceBase.h"
 #include "INiagaraCompiler.h"
 #include "NiagaraParameterMapHistory.h"
+#include "GraphEditAction.h"
 #include "NiagaraScriptSource.generated.h"
 
 UCLASS(MinimalAPI)
@@ -38,14 +39,23 @@ class UNiagaraScriptSource : public UNiagaraScriptSourceBase
 	virtual void PreCompile(UNiagaraEmitter* Emitter, bool bClearErrors = true) override;
 	virtual bool GatherPreCompiledVariables(const FString& InNamespaceFilter, TArray<FNiagaraVariable>& OutVars) override;
 	virtual void PostCompile() override;
+	virtual bool PostLoadFromEmitter(UNiagaraEmitter& OwningEmitter) override;
+
+	NIAGARAEDITOR_API virtual bool AddModuleIfMissing(FString ModulePath, ENiagaraScriptUsage Usage, bool& bOutFoundModule)override;
 
 	TArray<FNiagaraParameterMapHistory>& GetPrecomputedHistories() { return PrecompiledHistories; }
 	class UNiagaraGraph* GetPrecomputedNodeGraph() { return NodeGraphDeepCopy; }
 
+private:
+	void OnGraphChanged(const FEdGraphEditAction &Action);
+	void OnGraphDataInterfaceChanged();
+
 private:
 	bool bIsPrecompiled;
 	TArray<FNiagaraParameterMapHistory> PrecompiledHistories;
 
+	void UpgradeForScriptRapidIterationVariables();
+
 	UPROPERTY()
 	class UNiagaraGraph* NodeGraphDeepCopy;
 };
diff --git a/Engine/Plugins/FX/Niagara/Source/NiagaraEditor/Public/NiagaraSystemFactoryNew.h b/Engine/Plugins/FX/Niagara/Source/NiagaraEditor/Public/NiagaraSystemFactoryNew.h
--- a/Engine/Plugins/FX/Niagara/Source/NiagaraEditor/Public/NiagaraSystemFactoryNew.h
+++ b/Engine/Plugins/FX/Niagara/Source/NiagaraEditor/Public/NiagaraSystemFactoryNew.h
@@ -19,7 +19,7 @@ class UNiagaraSystemFactoryNew : public UFactory
 	//~ Begin UFactory Interface	
 
 public:
-	static void InitializeSystem(UNiagaraSystem* System);
+	static void InitializeSystem(UNiagaraSystem* System, bool bCreateDefaultNodes);
 };
 
 
diff --git a/Engine/Plugins/FX/Niagara/Source/NiagaraEditor/Public/ViewModels/Stack/NiagaraStackAddModuleItem.h b/Engine/Plugins/FX/Niagara/Source/NiagaraEditor/Public/ViewModels/Stack/NiagaraStackAddModuleItem.h
--- a/Engine/Plugins/FX/Niagara/Source/NiagaraEditor/Public/ViewModels/Stack/NiagaraStackAddModuleItem.h
+++ b/Engine/Plugins/FX/Niagara/Source/NiagaraEditor/Public/ViewModels/Stack/NiagaraStackAddModuleItem.h
@@ -18,17 +18,28 @@ class NIAGARAEDITOR_API UNiagaraStackAddModuleItem : public UNiagaraStackEntry
 public:
 	DECLARE_DELEGATE(FOnItemAdded);
 
+public:
+	enum class EDisplayMode
+	{
+		Standard,
+		Compact
+	};
+
 public:
 	void Initialize(TSharedRef<FNiagaraSystemViewModel> InSystemViewModel, TSharedRef<FNiagaraEmitterViewModel> InEmitterViewModel, UNiagaraStackEditorData& InStackEditorData);
 
 	virtual FText GetDisplayName() const override;
 
+	virtual EDisplayMode GetDisplayMode() const;
+
 	void SetOnItemAdded(FOnItemAdded InOnItemAdded);
 
 	void AddScriptModule(FAssetData ModuleScriptAsset);
 
 	void AddParameterModule(FNiagaraVariable ParameterVariable, bool bRenamePending);
 
+	void AddMaterialParameterModule();
+
 	virtual void GetAvailableParameters(TArray<FNiagaraVariable>& OutAvailableParameterVariables) const;
 
 	virtual void GetNewParameterAvailableTypes(TArray<FNiagaraTypeDefinition>& OutAvailableTypes) const;
@@ -37,6 +48,8 @@ public:
 
 	virtual ENiagaraScriptUsage GetOutputUsage() const PURE_VIRTUAL(UNiagaraStackAddModuleItem::GetOutputUsage, return ENiagaraScriptUsage::EmitterSpawnScript;);
 
+	virtual int32 GetTargetIndex() const;
+
 protected:
 	virtual UNiagaraNodeOutput* GetOrCreateOutputNode() PURE_VIRTUAL(UNiagaraStackModuleItem::GetOrCreateOutputNode, return nullptr;);
 
diff --git a/Engine/Plugins/FX/Niagara/Source/NiagaraEditor/Public/ViewModels/Stack/NiagaraStackAddScriptModuleItem.h b/Engine/Plugins/FX/Niagara/Source/NiagaraEditor/Public/ViewModels/Stack/NiagaraStackAddScriptModuleItem.h
--- a/Engine/Plugins/FX/Niagara/Source/NiagaraEditor/Public/ViewModels/Stack/NiagaraStackAddScriptModuleItem.h
+++ b/Engine/Plugins/FX/Niagara/Source/NiagaraEditor/Public/ViewModels/Stack/NiagaraStackAddScriptModuleItem.h
@@ -15,7 +15,9 @@ class NIAGARAEDITOR_API UNiagaraStackAddScriptModuleItem : public UNiagaraStackA
 	GENERATED_BODY()
 
 public:
-	void Initialize(TSharedRef<FNiagaraSystemViewModel> InSystemViewModel, TSharedRef<FNiagaraEmitterViewModel> InEmitterViewModel, UNiagaraStackEditorData& InStackEditorData, UNiagaraNodeOutput& InOutputNode);
+	void Initialize(TSharedRef<FNiagaraSystemViewModel> InSystemViewModel, TSharedRef<FNiagaraEmitterViewModel> InEmitterViewModel, UNiagaraStackEditorData& InStackEditorData, UNiagaraNodeOutput& InOutputNode, int32 InTargetIndex);
+
+	virtual EDisplayMode GetDisplayMode() const override;
 
 	virtual void GetAvailableParameters(TArray<FNiagaraVariable>& OutAvailableParameterVariables) const override;
 
@@ -25,8 +27,14 @@ public:
 
 	virtual ENiagaraScriptUsage GetOutputUsage() const override;
 
-	UNiagaraNodeOutput* GetOrCreateOutputNode() override;
+	virtual UNiagaraNodeOutput* GetOrCreateOutputNode() override;
+
+	virtual int32 GetTargetIndex() const override;
+
+	virtual FName GetItemBackgroundName() const override;
 
 private:
 	TWeakObjectPtr<UNiagaraNodeOutput> OutputNode;
+
+	int32 TargetIndex;
 };
\ No newline at end of file
diff --git a/Engine/Plugins/FX/Niagara/Source/NiagaraEditor/Public/ViewModels/Stack/NiagaraStackEmitterSpawnScriptItemGroup.h b/Engine/Plugins/FX/Niagara/Source/NiagaraEditor/Public/ViewModels/Stack/NiagaraStackEmitterSpawnScriptItemGroup.h
--- a/Engine/Plugins/FX/Niagara/Source/NiagaraEditor/Public/ViewModels/Stack/NiagaraStackEmitterSpawnScriptItemGroup.h
+++ b/Engine/Plugins/FX/Niagara/Source/NiagaraEditor/Public/ViewModels/Stack/NiagaraStackEmitterSpawnScriptItemGroup.h
@@ -18,15 +18,29 @@ class NIAGARAEDITOR_API UNiagaraStackEmitterPropertiesItem : public UNiagaraStac
 	GENERATED_BODY()
 
 public:
+	void Initialize(TSharedRef<FNiagaraSystemViewModel> InSystemViewModel, TSharedRef<FNiagaraEmitterViewModel> InEmitterViewModel, UNiagaraStackEditorData& InStackEditorData);
+
 	virtual FText GetDisplayName() const override;
 
+	bool CanResetToBase() const;
+
+	void ResetToBase();
+
 protected:
+	virtual void BeginDestroy() override;
+
 	virtual void RefreshChildrenInternal(const TArray<UNiagaraStackEntry*>& CurrentChildren, TArray<UNiagaraStackEntry*>& NewChildren) override;
 
 private:
 	void EmitterExpandedChanged();
 
+	void EmitterPropertiesChanged();
+
 private:
+	mutable TOptional<bool> bCanResetToBase;
+
+	TWeakObjectPtr<UNiagaraEmitter> Emitter;
+
 	UPROPERTY()
 	UNiagaraStackObject* EmitterObject;
 
diff --git a/Engine/Plugins/FX/Niagara/Source/NiagaraEditor/Public/ViewModels/Stack/NiagaraStackEntry.h b/Engine/Plugins/FX/Niagara/Source/NiagaraEditor/Public/ViewModels/Stack/NiagaraStackEntry.h
--- a/Engine/Plugins/FX/Niagara/Source/NiagaraEditor/Public/ViewModels/Stack/NiagaraStackEntry.h
+++ b/Engine/Plugins/FX/Niagara/Source/NiagaraEditor/Public/ViewModels/Stack/NiagaraStackEntry.h
@@ -15,6 +15,7 @@ class NIAGARAEDITOR_API UNiagaraStackEntry : public UObject
 public:
 	DECLARE_MULTICAST_DELEGATE(FOnStructureChanged);
 	DECLARE_MULTICAST_DELEGATE_OneParam(FOnDataObjectModified, UObject*);
+	DECLARE_DELEGATE_RetVal_OneParam(bool, FOnFilterChild, const UNiagaraStackEntry&);
 
 public:
 	UNiagaraStackEntry();
@@ -47,7 +48,9 @@ public:
 
 	virtual bool GetShouldShowInStack() const;
 
-	void GetChildren(TArray<UNiagaraStackEntry*>& OutChildren);
+	void GetFilteredChildren(TArray<UNiagaraStackEntry*>& OutFilteredChildren);
+
+	void GetUnfilteredChildren(TArray<UNiagaraStackEntry*>& OutUnfilteredChildren);
 
 	FOnStructureChanged& OnStructureChanged();
 
@@ -56,6 +59,9 @@ public:
 	void RefreshChildren();
 	void RefreshErrors();
 
+	FDelegateHandle AddChildFilter(FOnFilterChild ChildFilter);
+	void RemoveChildFilter(FDelegateHandle FilterHandle);
+
 	virtual int32 GetErrorCount() const { return 0; }
 	virtual bool GetErrorFixable(int32 ErrorIdx) const { return false; }
 	virtual bool TryFixError(int32 ErrorIdx) { return false; }
@@ -65,9 +71,6 @@ public:
 	TSharedRef<FNiagaraSystemViewModel> GetSystemViewModel() const;
 	TSharedRef<FNiagaraEmitterViewModel> GetEmitterViewModel() const;
 
-protected:
-	virtual void RefreshChildrenInternal(const TArray<UNiagaraStackEntry*>& CurrentChildren, TArray<UNiagaraStackEntry*>& NewChildren);
-
 	template<typename ChildType, typename PredicateType>
 	static ChildType* FindCurrentChildOfTypeByPredicate(const TArray<UNiagaraStackEntry*>& CurrentChildren, PredicateType Predicate)
 	{
@@ -82,6 +85,17 @@ protected:
 		return nullptr;
 	}
 
+	template<typename ChildType, typename PredicateType>
+	ChildType* FindChildOfTypeByPredicate(PredicateType Predicate)
+	{
+		TArray<UNiagaraStackEntry*> CurrentChildren;
+		GetUnfilteredChildren(CurrentChildren);
+		return FindCurrentChildOfTypeByPredicate<ChildType>(CurrentChildren, Predicate);
+	}
+
+protected:
+	virtual void RefreshChildrenInternal(const TArray<UNiagaraStackEntry*>& CurrentChildren, TArray<UNiagaraStackEntry*>& NewChildren);
+
 private:
 	void ChildStructureChanged();
 	
@@ -95,6 +109,8 @@ private:
 
 	FOnDataObjectModified DataObjectModifiedDelegate;
 
+	TArray<FOnFilterChild> ChildFilters;
+
 	UPROPERTY()
 	TArray<UNiagaraStackEntry*> Children;
 
diff --git a/Engine/Plugins/FX/Niagara/Source/NiagaraEditor/Public/ViewModels/Stack/NiagaraStackEventScriptItemGroup.h b/Engine/Plugins/FX/Niagara/Source/NiagaraEditor/Public/ViewModels/Stack/NiagaraStackEventScriptItemGroup.h
--- a/Engine/Plugins/FX/Niagara/Source/NiagaraEditor/Public/ViewModels/Stack/NiagaraStackEventScriptItemGroup.h
+++ b/Engine/Plugins/FX/Niagara/Source/NiagaraEditor/Public/ViewModels/Stack/NiagaraStackEventScriptItemGroup.h
@@ -4,12 +4,53 @@
 
 #include "NiagaraCommon.h"
 #include "NiagaraEmitter.h"
+#include "NiagaraStackItem.h"
 #include "NiagaraStackScriptItemGroup.h"
+#include "NiagaraStackItemExpander.h"
 #include "NiagaraStackAddEventScriptItem.h"
 #include "NiagaraStackEventScriptItemGroup.generated.h"
 
 class FNiagaraScriptViewModel;
 
+UCLASS()
+class NIAGARAEDITOR_API UNiagaraStackEventHandlerPropertiesItem : public UNiagaraStackItem
+{
+	GENERATED_BODY()
+
+public:
+	void Initialize(TSharedRef<FNiagaraSystemViewModel> InSystemViewModel, TSharedRef<FNiagaraEmitterViewModel> InEmitterViewModel, FGuid InEventScriptUsageId, UNiagaraStackEditorData& InStackEditorData);
+
+	virtual FText GetDisplayName() const override;
+
+	bool CanResetToBase() const;
+
+	void ResetToBase();
+
+protected:
+	virtual void BeginDestroy() override;
+
+	virtual void RefreshChildrenInternal(const TArray<UNiagaraStackEntry*>& CurrentChildren, TArray<UNiagaraStackEntry*>& NewChildren) override;
+
+private:
+	void EventHandlerExpandedChanged();
+
+	void EventHandlerPropertiesChanged();
+
+private:
+	bool bHasBaseEventHandler;
+
+	mutable TOptional<bool> bCanResetToBase;
+
+	bool bForceRebuildChildStruct;
+
+	TWeakObjectPtr<UNiagaraEmitter> Emitter;
+
+	FGuid EventScriptUsageId;
+
+	UPROPERTY()
+	UNiagaraStackItemExpander* EventHandlerExpander;
+};
+
 UCLASS()
 /** Meant to contain a single binding of a Emitter::EventScriptProperties to the stack.*/
 class NIAGARAEDITOR_API UNiagaraStackEventScriptItemGroup : public UNiagaraStackScriptItemGroup
@@ -31,4 +72,14 @@ public:
 	virtual FText GetDisplayName() const override;
 	void SetOnModifiedEventHandlers(FOnModifiedEventHandlers OnModifiedEventHandlers);
 
+private:
+	FText DisplayName;
+
+	bool bHasBaseEventHandler;
+
+	UPROPERTY()
+	UNiagaraStackEventHandlerPropertiesItem* EventHandlerProperties;
+
+	UPROPERTY()
+	UNiagaraStackSpacer* EventHandlerPropertiesSpacer;
 };
\ No newline at end of file
diff --git a/Engine/Plugins/FX/Niagara/Source/NiagaraEditor/Public/ViewModels/Stack/NiagaraStackFunctionInput.h b/Engine/Plugins/FX/Niagara/Source/NiagaraEditor/Public/ViewModels/Stack/NiagaraStackFunctionInput.h
--- a/Engine/Plugins/FX/Niagara/Source/NiagaraEditor/Public/ViewModels/Stack/NiagaraStackFunctionInput.h
+++ b/Engine/Plugins/FX/Niagara/Source/NiagaraEditor/Public/ViewModels/Stack/NiagaraStackFunctionInput.h
@@ -87,6 +87,9 @@ public:
 	virtual FName GetTextStyleName() const override;
 	virtual bool GetCanExpand() const override;
 	virtual int32 GetItemIndentLevel() const override;
+	virtual FText GetTooltipText() const override;
+
+	FText GetTooltipText(EValueMode InValueMode) const;
 
 	/** Sets the item indent level for this stack entry. */
 	void SetItemIndentLevel(int32 InItemIndentLevel);
@@ -142,6 +145,15 @@ public:
 	/** Resets the value and handle of this input to the value and handle defined in the module. */
 	void Reset();
 
+	/** Whether or not this input has a base value.  This is true for emitter instances in systems. */
+	bool EmitterHasBase() const;
+
+	/** Whether or not this input can be reset to a base value. */
+	bool CanResetToBase() const;
+
+	/** Resets this input to its base value. */
+	void ResetToBase();
+
 	/** Returns whether or not this input can be renamed. */
 	bool CanRenameInput() const;
 
@@ -167,6 +179,10 @@ protected:
 	//~ UNiagaraStackEntry interface
 	virtual void RefreshChildrenInternal(const TArray<UNiagaraStackEntry*>& CurrentChildren, TArray<UNiagaraStackEntry*>& NewChildren) override;
 
+	FNiagaraVariable GetDefaultValue() const;
+	bool UpdateRapidIterationParametersForAffectedScripts(const uint8* Data);
+	bool RemoveRapidIterationParametersForAffectedScripts();
+
 private:
 	struct FDataValues
 	{
@@ -250,6 +266,12 @@ private:
 	/** Called whenever the graph which generated this input changes. */
 	void OnGraphChanged(const struct FEdGraphEditAction& InAction);
 
+	/** Called whenever rapid iteration parameters are changed for the script that owns the function that owns this input. */
+	void OnRapidIterationParametersChanged();
+
+	/** Called whenever the script source that owns the function that owns this input changes. */
+	void OnScriptSourceChanged();
+
 	/** Gets the graph node which owns the local overrides for the module that owns this input if it exists. */
 	UNiagaraNodeParameterMapSet* GetOverrideNode() const;
 
@@ -279,8 +301,13 @@ private:
 	/** Gets the dynamic input node providing a value to this input if one exists. */
 	bool TryGetCurrentDynamicValue(TWeakObjectPtr<UNiagaraNodeFunctionCall>& DynamicValue, UEdGraphPin* OverridePin);
 
-	/** Recursively removes all nodes connected to the override pin. */
-	static void RemoveAllNodesConnectedToOverridePin(UEdGraphPin& OverridePin, UNiagaraStackFunctionInput* OwningInput);
+	/** Removes all nodes connected to the override pin which provide it's value. */
+	void RemoveNodesForOverridePin(UEdGraphPin& OverridePin);
+
+	/** Determine if the values in this input are possibly under the control of the rapid iteration array on the script.*/
+	bool IsRapidIterationCandidate() const;
+
+	FString CreateRapidIterationVariableName(const FString& InName);
 
 private:
 	/** The stack editor data for this function input. */
@@ -296,7 +323,7 @@ private:
 	/** The assignment node which owns this input.  This is only valid for inputs of assignment modules. */
 	TWeakObjectPtr<UNiagaraNodeAssignment> OwningAssignmentNode;
 
-	/** The nigara type definition for this input. */
+	/** The niagara type definition for this input. */
 	FNiagaraTypeDefinition InputType;
 
 	/** Whether or not this input can be pinned. */
@@ -315,6 +342,9 @@ private:
 	  * graph.  This only affects parameter handles which are local module handles. */
 	FNiagaraParameterHandle AliasedInputParameterHandle;
 
+	/** The computed name for the rapid iteration variable that could potentially drive this entry..*/
+	FString RapidIterationParameterName;
+
 	/** The name of this input for display in the UI. */
 	FText DisplayName;
 
@@ -329,12 +359,19 @@ private:
 	  * UI reads this value every frame due to attribute updates. */
 	mutable TOptional<UEdGraphPin*> OverridePinCache;
 
+	/** Whether or not this input can be reset to a base value. */
+	mutable TOptional<bool> bCanResetToBase;
+
 	/** A flag to prevent handling graph changes when it's being updated directly by this object. */
 	bool bUpdatingGraphDirectly;
 
 	/** A handle for removing the graph changed delegate. */
 	FDelegateHandle GraphChangedHandle;
 
+	/** A handle for removing the changed delegate for the rapid iteration parameters for the script
+	   that owns the function that owns this input. */
+	FDelegateHandle RapidIterationParametersChangedHandle;
+
 	/** A multicast delegate which is called when the value of this input is changed. */
 	FOnValueChanged ValueChangedDelegate;
 
@@ -343,4 +380,10 @@ private:
 
 	/** The item indent level for this stack entry. */
 	int32 ItemIndentLevel;
+
+	/** The script which owns the function which owns this input.  This is also the autoritative version of the rapid iteration parameters. */
+	TWeakObjectPtr<UNiagaraScript> SourceScript;
+
+	/** An array of scripts which this input affects. */
+	TArray<TWeakObjectPtr<UNiagaraScript>> AffectedScripts;
 };
\ No newline at end of file
