Copying Orion-Staging @ 3748653 (Orion/Dev-General @ 3722124) to //UE4/Main

diff --git a/Engine/Plugins/FX/Niagara/Source/NiagaraEditor/Public/ViewModels/Stack/NiagaraStackFunctionInputCollection.h b/Engine/Plugins/FX/Niagara/Source/NiagaraEditor/Public/ViewModels/Stack/NiagaraStackFunctionInputCollection.h
--- a/Engine/Plugins/FX/Niagara/Source/NiagaraEditor/Public/ViewModels/Stack/NiagaraStackFunctionInputCollection.h
+++ b/Engine/Plugins/FX/Niagara/Source/NiagaraEditor/Public/ViewModels/Stack/NiagaraStackFunctionInputCollection.h
@@ -16,7 +16,6 @@ class NIAGARAEDITOR_API UNiagaraStackFunctionInputCollection : public UNiagaraSt
 
 public:
 	DECLARE_MULTICAST_DELEGATE(FOnInputPinnedChanged);
-	DECLARE_DELEGATE_RetVal_OneParam(bool, FOnFilterChildren, UNiagaraStackFunctionInput*);
 
 public:
 	struct FDisplayOptions
@@ -24,7 +23,7 @@ public:
 		FText DisplayName;
 		bool bShouldShowInStack;
 		int32 ChildItemIndentLevel;
-		FOnFilterChildren ChildFilter;
+		UNiagaraStackEntry::FOnFilterChild ChildFilter;
 	};
 
 public:
@@ -47,6 +46,11 @@ public:
 	virtual FName GetTextStyleName() const override;
 	virtual bool GetCanExpand() const override;
 	virtual bool GetShouldShowInStack() const override;
+	virtual int32 GetErrorCount() const override;
+	virtual bool GetErrorFixable(int32 ErrorIdx) const override;
+	virtual bool TryFixError(int32 ErrorIdx) override;
+	virtual FText GetErrorText(int32 ErrorIdx) const override;
+	virtual FText GetErrorSummaryText(int32 ErrorIdx) const override;
 
 	FOnInputPinnedChanged& OnInputPinnedChanged();
 
@@ -57,6 +61,16 @@ private:
 	void ChildPinnedChanged();
 
 private:
+	DECLARE_DELEGATE(FFixError);
+	struct FError
+	{
+		FText ErrorText;
+		FText ErrorSummaryText;
+		FFixError Fix;
+	};
+
+	TArray<FError> Errors;
+
 	UNiagaraStackEditorData* StackEditorData;
 	UNiagaraNodeFunctionCall* ModuleNode;
 	UNiagaraNodeFunctionCall* InputFunctionCallNode;
diff --git a/Engine/Plugins/FX/Niagara/Source/NiagaraEditor/Public/ViewModels/Stack/NiagaraStackGraphUtilities.h b/Engine/Plugins/FX/Niagara/Source/NiagaraEditor/Public/ViewModels/Stack/NiagaraStackGraphUtilities.h
--- a/Engine/Plugins/FX/Niagara/Source/NiagaraEditor/Public/ViewModels/Stack/NiagaraStackGraphUtilities.h
+++ b/Engine/Plugins/FX/Niagara/Source/NiagaraEditor/Public/ViewModels/Stack/NiagaraStackGraphUtilities.h
@@ -5,6 +5,7 @@
 #include "NiagaraTypes.h"
 #include "NiagaraCommon.h"
 #include "NiagaraParameterHandle.h"
+#include "AssetData.h"
 
 class UEdGraph;
 class UEdGraphPin;
@@ -13,6 +14,9 @@ class UNiagaraNode;
 class UNiagaraNodeInput;
 class UNiagaraNodeOutput;
 class UNiagaraNodeFunctionCall;
+class UNiagaraNodeAssignment;
+class UNiagaraMaterialParameterNode;
+class UNiagaraNodeParameterMapSet;
 class UNiagaraStackEditorData;
 class FNiagaraSystemViewModel;
 class FNiagaraEmitterViewModel;
@@ -43,6 +47,7 @@ namespace FNiagaraStackGraphUtilities
 	{
 		TArray<UNiagaraNode*> StartNodes;
 		UNiagaraNode* EndNode;
+		void GetAllNodesInGroup(TArray<UNiagaraNode*>& OutAllNodes) const;
 	};
 
 	void GetStackNodeGroups(UNiagaraNode& StackNode, TArray<FStackNodeGroup>& OutStackNodeGroups);
@@ -63,9 +68,47 @@ namespace FNiagaraStackGraphUtilities
 		ModuleInputsOnly
 	};
 
-	void GetStackFunctionInputPins(UNiagaraNodeFunctionCall& FunctionCallNode, TArray<const UEdGraphPin*>& OutInputPins, ENiagaraGetStackFunctionInputPinsOptions Options = ENiagaraGetStackFunctionInputPinsOptions::AllInputs);
+	void GetStackFunctionInputPins(UNiagaraNodeFunctionCall& FunctionCallNode, TArray<const UEdGraphPin*>& OutInputPins, ENiagaraGetStackFunctionInputPinsOptions Options = ENiagaraGetStackFunctionInputPinsOptions::AllInputs, bool bIgnoreDisabled = false);
 
-	bool ValidateGraphForOutput(UNiagaraGraph& NiagaraGraph, ENiagaraScriptUsage ScriptUsage, int32 ScriptOccurrence, FText& ErrorMessage);
+	UNiagaraNodeParameterMapSet* GetStackFunctionOverrideNode(UNiagaraNodeFunctionCall& FunctionCallNode);
 
-	UNiagaraNodeOutput* ResetGraphForOutput(UNiagaraGraph& NiagaraGraph, ENiagaraScriptUsage ScriptUsage, int32 ScriptOccurrence);
+	UNiagaraNodeParameterMapSet& GetOrCreateStackFunctionOverrideNode(UNiagaraNodeFunctionCall& FunctionCallNode);
+
+	UEdGraphPin* GetStackFunctionInputOverridePin(UNiagaraNodeFunctionCall& StackFunctionCall, FNiagaraParameterHandle AliasedInputParameterHandle);
+
+	UEdGraphPin& GetOrCreateStackFunctionInputOverridePin(UNiagaraNodeFunctionCall& StackFunctionCall, FNiagaraParameterHandle AliasedInputParameterHandle, FNiagaraTypeDefinition InputType);
+
+	void RemoveNodesForStackFunctionInputOverridePin(UEdGraphPin& StackFunctionInputOverridePin);
+
+	void RemoveNodesForStackFunctionInputOverridePin(UEdGraphPin& StackFunctinoInputOverridePin, TArray<TWeakObjectPtr<UNiagaraDataInterface>>& OutRemovedDataObjects);
+
+	void SetLinkedValueHandleForFunctionInput(UEdGraphPin& OverridePin, FNiagaraParameterHandle LinkedParameterHandle);
+
+	void SetDataValueObjectForFunctionInput(UEdGraphPin& OverridePin, UClass* DataObjectType, FString DataObjectName, UNiagaraDataInterface*& OutDataObject);
+
+	void SetDynamicInputForFunctionInput(UEdGraphPin& OverridePin, UNiagaraScript* DynamicInput, UNiagaraNodeFunctionCall*& OutDynamicInputFunctionCall);
+
+	bool RemoveModuleFromStack(UNiagaraNodeFunctionCall& ModuleNode);
+
+	UNiagaraNodeFunctionCall* AddScriptModuleToStack(FAssetData ModuleScriptAsset, UNiagaraNodeOutput& TargetOutputNode, int32 TargetIndex = INDEX_NONE);
+	
+	bool FindScriptModulesInStack(FAssetData ModuleScriptAsset, UNiagaraNodeOutput& TargetOutputNode, TArray<UNiagaraNodeFunctionCall*> OutFunctionCalls);
+
+	UNiagaraNodeAssignment* AddParameterModuleToStack(FNiagaraVariable ParameterVariable, UNiagaraNodeOutput& TargetOutputNode, int32 TargetIndex = INDEX_NONE);
+		
+	UNiagaraMaterialParameterNode* AddMaterialParameterModuleToStack(TSharedRef<FNiagaraEmitterViewModel>& EmitterViewModel, UNiagaraNodeOutput& TargetOutputNode, int32 TargetIndex = INDEX_NONE);
+
+	TOptional<bool> GetModuleIsEnabled(UNiagaraNodeFunctionCall& FunctionCallNode);
+
+	void SetModuleIsEnabled(UNiagaraNodeFunctionCall& FunctionCallNode, bool bIsEnabled);
+
+	bool ValidateGraphForOutput(UNiagaraGraph& NiagaraGraph, ENiagaraScriptUsage ScriptUsage, FGuid ScriptUsageId, FText& ErrorMessage);
+
+	UNiagaraNodeOutput* ResetGraphForOutput(UNiagaraGraph& NiagaraGraph, ENiagaraScriptUsage ScriptUsage, FGuid ScriptUsageId);
+
+	const UNiagaraEmitter* GetBaseEmitter(UNiagaraEmitter& Emitter, UNiagaraSystem& OwningSystem);
+
+	void CleanUpStaleRapidIterationParameters(UNiagaraScript& Script, UNiagaraEmitter& OwningEmitter);
+
+	void CleanUpStaleRapidIterationParameters(UNiagaraEmitter& Emitter);
 }
\ No newline at end of file
diff --git a/Engine/Plugins/FX/Niagara/Source/NiagaraEditor/Public/ViewModels/Stack/NiagaraStackModuleItem.h b/Engine/Plugins/FX/Niagara/Source/NiagaraEditor/Public/ViewModels/Stack/NiagaraStackModuleItem.h
--- a/Engine/Plugins/FX/Niagara/Source/NiagaraEditor/Public/ViewModels/Stack/NiagaraStackModuleItem.h
+++ b/Engine/Plugins/FX/Niagara/Source/NiagaraEditor/Public/ViewModels/Stack/NiagaraStackModuleItem.h
@@ -6,6 +6,7 @@
 #include "NiagaraStackModuleItem.generated.h"
 
 class UNiagaraNodeFunctionCall;
+class UNiagaraNode;
 class UNiagaraStackFunctionInputCollection;
 class UNiagaraStackModuleItemOutputCollection;
 class UNiagaraStackItemExpander;
@@ -25,6 +26,13 @@ public:
 	virtual FText GetDisplayName() const override;
 	virtual FText GetTooltipText() const override;
 
+	bool CanMoveAndDelete() const;
+	bool CanRefresh() const;
+	void Refresh();
+
+	bool GetIsEnabled() const;
+	void SetIsEnabled(bool bInIsEnabled);
+
 	void MoveUp();
 	void MoveDown();
 	void Delete();
@@ -33,11 +41,32 @@ protected:
 	virtual void RefreshChildrenInternal(const TArray<UNiagaraStackEntry*>& CurrentChildren, TArray<UNiagaraStackEntry*>& NewChildren) override;
 
 private:
+	void RefreshIsEnabled();
 	void InputPinnedChanged();
 	void ModuleExpandedChanged();
 
+private:
+	virtual int32 GetErrorCount() const override;
+	virtual bool GetErrorFixable(int32 ErrorIdx) const override;
+	virtual bool TryFixError(int32 ErrorIdx) override;
+	virtual FText GetErrorText(int32 ErrorIdx) const override;
+	virtual FText GetErrorSummaryText(int32 ErrorIdx) const override;
+
 private:
 	UNiagaraNodeFunctionCall* FunctionCallNode;
+	FString ModuleEditorDataKey;
+	bool bCanMoveAndDelete;
+	bool bIsEnabled;
+
+	DECLARE_DELEGATE_RetVal(bool, FFixError);
+	struct FError
+	{
+		FText ErrorText;
+		FText ErrorSummaryText;
+		FFixError Fix;
+	};
+
+	TArray<FError> Errors;
 
 	UPROPERTY()
 	UNiagaraStackFunctionInputCollection* PinnedInputCollection;
diff --git a/Engine/Plugins/FX/Niagara/Source/NiagaraEditor/Public/ViewModels/Stack/NiagaraStackModuleItemOutput.h b/Engine/Plugins/FX/Niagara/Source/NiagaraEditor/Public/ViewModels/Stack/NiagaraStackModuleItemOutput.h
--- a/Engine/Plugins/FX/Niagara/Source/NiagaraEditor/Public/ViewModels/Stack/NiagaraStackModuleItemOutput.h
+++ b/Engine/Plugins/FX/Niagara/Source/NiagaraEditor/Public/ViewModels/Stack/NiagaraStackModuleItemOutput.h
@@ -4,6 +4,7 @@
 
 #include "NiagaraStackEntry.h"
 #include "NiagaraParameterHandle.h"
+#include "NiagaraTypes.h"
 #include "NiagaraStackModuleItemOutput.generated.h"
 
 class UNiagaraNodeFunctionCall;
@@ -22,13 +23,15 @@ public:
 	 * @param InFunctionCallNode The function call node representing the module in the stack graph which owns this Output.
 	 * @param InOutputParameterHandle The Namespace.Name handle for the Output to the owning module.
 	 */
-	void Initialize(TSharedRef<FNiagaraSystemViewModel> InSystemViewModel, TSharedRef<FNiagaraEmitterViewModel> InEmitterViewModel, UNiagaraNodeFunctionCall& InFunctionCallNode, FName InOutputParameterHandle);
+	void Initialize(TSharedRef<FNiagaraSystemViewModel> InSystemViewModel, TSharedRef<FNiagaraEmitterViewModel> InEmitterViewModel, UNiagaraNodeFunctionCall& InFunctionCallNode, FName InOutputParameterHandle,
+		FNiagaraTypeDefinition InOutputType);
 
 	//~ UNiagaraStackEntry interface
 	virtual FText GetDisplayName() const override;
 	virtual FName GetTextStyleName() const override;
 	virtual bool GetCanExpand() const override;
 	virtual int32 GetItemIndentLevel() const override;
+	virtual FText GetTooltipText() const override;
 
 	/** Gets the parameter handle which defined this Output in the module. */
 	const FNiagaraParameterHandle& GetOutputParameterHandle() const;
@@ -45,4 +48,7 @@ private:
 
 	/** The name of this Output for display in the UI. */
 	FText DisplayName;
+
+	/** The niagara type definition for this input. */
+	FNiagaraTypeDefinition OutputType;
 };
\ No newline at end of file
diff --git a/Engine/Plugins/FX/Niagara/Source/NiagaraEditor/Public/ViewModels/Stack/NiagaraStackRendererItem.h b/Engine/Plugins/FX/Niagara/Source/NiagaraEditor/Public/ViewModels/Stack/NiagaraStackRendererItem.h
--- a/Engine/Plugins/FX/Niagara/Source/NiagaraEditor/Public/ViewModels/Stack/NiagaraStackRendererItem.h
+++ b/Engine/Plugins/FX/Niagara/Source/NiagaraEditor/Public/ViewModels/Stack/NiagaraStackRendererItem.h
@@ -25,8 +25,16 @@ public:
 
 	virtual FText GetDisplayName() const override;
 
+	bool CanDelete() const;
+
 	void Delete();
 
+	bool CanHaveBase() const;
+
+	bool CanResetToBase() const;
+
+	void ResetToBase();
+
 	virtual FName GetItemBackgroundName() const override;
 	virtual int32 GetErrorCount() const override;
 	virtual bool GetErrorFixable(int32 ErrorIdx) const override;
@@ -37,13 +45,21 @@ public:
 	static bool AddMissingVariable(UNiagaraEmitter* Emitter, const FNiagaraVariable& Variable);
 
 protected:
+	virtual void BeginDestroy() override;
+
 	virtual void RefreshChildrenInternal(const TArray<UNiagaraStackEntry*>& CurrentChildren, TArray<UNiagaraStackEntry*>& NewChildren) override;
 
 private:
 	void RendererExpandedChanged();
 
+	void RendererChanged();
+
 private:
-	UNiagaraRendererProperties* RendererProperties;
+	TWeakObjectPtr<UNiagaraRendererProperties> RendererProperties;
+
+	bool bHasBaseRenderer;
+
+	mutable TOptional<bool> bCanResetToBase;
 
 	TArray<FNiagaraVariable> MissingAttributes;
 
diff --git a/Engine/Plugins/FX/Niagara/Source/NiagaraEditor/Public/ViewModels/Stack/NiagaraStackScriptItemGroup.h b/Engine/Plugins/FX/Niagara/Source/NiagaraEditor/Public/ViewModels/Stack/NiagaraStackScriptItemGroup.h
--- a/Engine/Plugins/FX/Niagara/Source/NiagaraEditor/Public/ViewModels/Stack/NiagaraStackScriptItemGroup.h
+++ b/Engine/Plugins/FX/Niagara/Source/NiagaraEditor/Public/ViewModels/Stack/NiagaraStackScriptItemGroup.h
@@ -18,18 +18,16 @@ class NIAGARAEDITOR_API UNiagaraStackScriptItemGroup : public UNiagaraStackItemG
 	GENERATED_BODY()
 
 public:
-	UNiagaraStackScriptItemGroup();
-
 	void Initialize(
 		TSharedRef<FNiagaraSystemViewModel> InSystemViewModel,
 		TSharedRef<FNiagaraEmitterViewModel> InEmitterViewModel,
 		UNiagaraStackEditorData& InStackEditorData,
 		TSharedRef<FNiagaraScriptViewModel> InScriptViewModel,
 		ENiagaraScriptUsage InScriptUsage,
-		int32 InScriptOccurrence = 0);
+		FGuid InScriptUsageId = FGuid());
 
 	ENiagaraScriptUsage GetScriptUsage() const { return ScriptUsage; }
-	int32 GetScriptOccurrence() const { return ScriptOccurrence; }
+	FGuid GetScriptUsageId() const { return ScriptUsageId; }
 
 	virtual FText GetDisplayName() const override;
 	void SetDisplayName(FText InDisplayName);
@@ -42,7 +40,7 @@ public:
 
 protected:
 	virtual void RefreshChildrenInternal(const TArray<UNiagaraStackEntry*>& CurrentChildren, TArray<UNiagaraStackEntry*>& NewChildren) override;
-	TSharedPtr<FNiagaraScriptViewModel> ScriptViewModel;
+	TWeakPtr<FNiagaraScriptViewModel> ScriptViewModel;
 
 private:
 	void ItemAdded();
@@ -52,16 +50,10 @@ private:
 private:
 	ENiagaraScriptUsage ScriptUsage;
 
-	int32 ScriptOccurrence;
+	FGuid ScriptUsageId;
 
 	FText DisplayName;
 
-	UPROPERTY()
-	UNiagaraStackAddScriptModuleItem* AddModuleItem;
-
-	UPROPERTY()
-	UNiagaraStackSpacer* BottomSpacer;
-
 	DECLARE_DELEGATE(FFixError);
 	struct FError
 	{
diff --git a/Engine/Plugins/FX/Niagara/Source/NiagaraEditorWidgets/Private/DetailCustomizations/NiagaraDataInterfaceCurveDetails.cpp b/Engine/Plugins/FX/Niagara/Source/NiagaraEditorWidgets/Private/DetailCustomizations/NiagaraDataInterfaceCurveDetails.cpp
--- a/Engine/Plugins/FX/Niagara/Source/NiagaraEditorWidgets/Private/DetailCustomizations/NiagaraDataInterfaceCurveDetails.cpp
+++ b/Engine/Plugins/FX/Niagara/Source/NiagaraEditorWidgets/Private/DetailCustomizations/NiagaraDataInterfaceCurveDetails.cpp
@@ -232,6 +232,8 @@ public:
 		];
 
 		CurveEditor->SetCurveOwner(CurveOwner.Get());
+		// Allow users to scroll over the widget in the editor using the scroll wheel (unless it has keyboard focus, in which case it will zoom in/out)
+		CurveEditor->SetRequireFocusToZoom(true);
 	}
 
 private:
@@ -282,6 +284,15 @@ void FNiagaraDataInterfaceCurveDetailsBase::CustomizeDetails(IDetailLayoutBuilde
 	TArray<TSharedRef<IPropertyHandle>> CurveProperties;
 	GetCurveProperties(DetailBuilder, CurveProperties);
 
+	// Make sure all property handles are valid.
+	for (TSharedRef<IPropertyHandle> CurveProperty : CurveProperties)
+	{
+		if (CurveProperty->IsValidHandle() == false)
+		{
+			return;
+		}
+	}
+
 	for (TSharedRef<IPropertyHandle> CurveProperty : CurveProperties)
 	{
 		CurveProperty->MarkHiddenByCustomization();
diff --git a/Engine/Plugins/FX/Niagara/Source/NiagaraEditorWidgets/Private/SNiagaraStack.cpp b/Engine/Plugins/FX/Niagara/Source/NiagaraEditorWidgets/Private/SNiagaraStack.cpp
--- a/Engine/Plugins/FX/Niagara/Source/NiagaraEditorWidgets/Private/SNiagaraStack.cpp
+++ b/Engine/Plugins/FX/Niagara/Source/NiagaraEditorWidgets/Private/SNiagaraStack.cpp
@@ -27,6 +27,8 @@
 #include "ViewModels/Stack/NiagaraStackErrorItem.h"
 #include "ViewModels/Stack/NiagaraStackStruct.h"
 #include "ViewModels/Stack/NiagaraStackParameterStoreEntry.h"
+#include "ViewModels/Stack/NiagaraStackEmitterSpawnScriptItemGroup.h"
+#include "ViewModels/Stack/NiagaraStackEventScriptItemGroup.h"
 #include "FrontendFilterBase.h"
 #include "Framework/Commands/UIAction.h"
 #include "Framework/MultiBox/MultiBoxBuilder.h"
@@ -47,6 +49,7 @@
 #include "SlateApplication.h"
 #include "Widgets/Colors/SSimpleGradient.h"
 #include "Widgets/Text/SInlineEditableTextBlock.h"
+#include "Widgets/Input/SCheckBox.h"
 #include "NiagaraNodeAssignment.h"
 #include "UObjectGlobals.h"
 #include "NiagaraDataInterface.h"
@@ -60,6 +63,7 @@
 #include "EdGraph/EdGraph.h"
 #include "Toolkits/AssetEditorManager.h"
 #include "NiagaraEditorUtilities.h"
+#include "SSearchBox.h"
 
 #define LOCTEXT_NAMESPACE "NiagaraStack"
 
@@ -97,6 +101,7 @@ public:
 				.IsReadOnly(this, &SNiagaraStackFunctionInputName::GetIsNameReadOnly)
 				.IsSelected(this, &SNiagaraStackFunctionInputName::GetIsNameWidgetSelected)
 				.OnTextCommitted(this, &SNiagaraStackFunctionInputName::OnNameTextCommitted)
+				.ToolTipText_UObject(FunctionInput, &UNiagaraStackFunctionInput::GetTooltipText, UNiagaraStackFunctionInput::EValueMode::Local)
 			]
 			// Pin
 			+ SHorizontalBox::Slot()
@@ -240,6 +245,7 @@ public:
 				[
 					SNew(SBox)
 					.Visibility(this, &SNiagaraStackFunctionInputValue::GetValueWidgetVisibility, UNiagaraStackFunctionInput::EValueMode::Linked)
+					.ToolTipText_UObject(InFunctionInput, &UNiagaraStackFunctionInput::GetTooltipText, UNiagaraStackFunctionInput::EValueMode::Linked)
 					.VAlign(VAlign_Center)
 					[
 						SNew(STextBlock)
@@ -325,6 +331,27 @@ public:
 					.Image(FEditorStyle::GetBrush("PropertyWindow.DiffersFromDefault"))
 				]
 			]
+
+			// Reset to base Button
+			+ SHorizontalBox::Slot()
+			.AutoWidth()
+			.VAlign(VAlign_Center)
+			.Padding(3, 0, 0, 0)
+			[
+				SNew(SButton)
+				.IsFocusable(false)
+				.ToolTipText(LOCTEXT("ResetToBaseToolTip", "Reset this input to the value defined by the parent emitter"))
+				.ButtonStyle(FEditorStyle::Get(), "NoBorder")
+				.ContentPadding(0)
+				.Visibility(this, &SNiagaraStackFunctionInputValue::GetResetToBaseButtonVisibility)
+				.OnClicked(this, &SNiagaraStackFunctionInputValue::ResetToBaseButtonPressed)
+				.Content()
+				[
+					SNew(SImage)
+					.Image(FEditorStyle::GetBrush("PropertyWindow.DiffersFromDefault"))
+					.ColorAndOpacity(FSlateColor(FLinearColor::Green))
+				]
+			]
 		];
 	}
 
@@ -341,10 +368,10 @@ private:
 		if (DisplayedLocalValueStruct.IsValid())
 		{
 			FNiagaraEditorModule& NiagaraEditorModule = FModuleManager::GetModuleChecked<FNiagaraEditorModule>("NiagaraEditor");
-			TSharedPtr<INiagaraEditorTypeUtilities> TypeEditorUtilities = NiagaraEditorModule.GetTypeUtilities(FNiagaraTypeDefinition((UScriptStruct*)DisplayedLocalValueStruct->GetStruct()));
+			TSharedPtr<INiagaraEditorTypeUtilities> TypeEditorUtilities = NiagaraEditorModule.GetTypeUtilities(FunctionInput->GetInputType());
 			if (TypeEditorUtilities.IsValid() && TypeEditorUtilities->CanCreateParameterEditor())
 			{
-				TSharedPtr<SNiagaraParameterEditor> ParameterEditor = TypeEditorUtilities->CreateParameterEditor();
+				TSharedPtr<SNiagaraParameterEditor> ParameterEditor = TypeEditorUtilities->CreateParameterEditor(FunctionInput->GetInputType());
 				ParameterEditor->UpdateInternalValueFromStruct(DisplayedLocalValueStruct.ToSharedRef());
 				ParameterEditor->SetOnBeginValueChange(SNiagaraParameterEditor::FOnValueChange::CreateSP(
 					this, &SNiagaraStackFunctionInputValue::ParameterBeginValueChange));
@@ -656,6 +683,24 @@ private:
 		return FReply::Handled();
 	}
 
+	EVisibility GetResetToBaseButtonVisibility() const
+	{
+		if (FunctionInput->EmitterHasBase())
+		{
+			return FunctionInput->CanResetToBase() ? EVisibility::Visible : EVisibility::Hidden;
+		}
+		else
+		{
+			return EVisibility::Collapsed;
+		}
+	}
+
+	FReply ResetToBaseButtonPressed() const
+	{
+		FunctionInput->ResetToBase();
+		return FReply::Handled();
+	}
+
 	FText GetInputIconText() const
 	{
 		switch (FunctionInput->GetValueMode())
@@ -824,10 +869,11 @@ private:
 		if (DisplayedValueStruct.IsValid())
 		{
 			FNiagaraEditorModule& NiagaraEditorModule = FModuleManager::GetModuleChecked<FNiagaraEditorModule>("NiagaraEditor");
-			TSharedPtr<INiagaraEditorTypeUtilities> TypeEditorUtilities = NiagaraEditorModule.GetTypeUtilities(FNiagaraTypeDefinition((UScriptStruct*)DisplayedValueStruct->GetStruct()));
+			FNiagaraTypeDefinition ParameterType = FNiagaraTypeDefinition((UScriptStruct*)DisplayedValueStruct->GetStruct());
+			TSharedPtr<INiagaraEditorTypeUtilities> TypeEditorUtilities = NiagaraEditorModule.GetTypeUtilities(ParameterType);
 			if (TypeEditorUtilities.IsValid() && TypeEditorUtilities->CanCreateParameterEditor())
 			{
-				TSharedPtr<SNiagaraParameterEditor> ParameterEditor = TypeEditorUtilities->CreateParameterEditor();
+				TSharedPtr<SNiagaraParameterEditor> ParameterEditor = TypeEditorUtilities->CreateParameterEditor(ParameterType);
 				ParameterEditor->UpdateInternalValueFromStruct(DisplayedValueStruct.ToSharedRef());
 				ParameterEditor->SetOnBeginValueChange(SNiagaraParameterEditor::FOnValueChange::CreateSP(
 					this, &SNiagaraStackParameterStoreEntryValue::ParameterBeginValueChange));
@@ -1050,6 +1096,7 @@ public:
 
 	virtual void NotifyPostChange(const FPropertyChangedEvent& PropertyChangedEvent, class FEditPropertyChain* PropertyThatChanged) override
 	{
+		Object->GetOwningObject()->PostEditChange();
 	}
 
 private:
@@ -1065,46 +1112,92 @@ public:
 	void Construct(const FArguments& InArgs, UNiagaraStackAddModuleItem& InAddModuleItem)
 	{
 		AddModuleItem = &InAddModuleItem;
+
+		FSlateColor TextColor;
+		FSlateColor BackgroundColor;
+		float Padding;
+		FName TextStyleName;
+		FName FontStyleName;
+		FText ToolTipText;
+		if (AddModuleItem->GetDisplayMode() == UNiagaraStackAddModuleItem::EDisplayMode::Standard)
+		{
+			TextColor = FSlateColor::UseForeground();
+			BackgroundColor = FSlateColor(FLinearColor::Transparent);
+			Padding = 3;
+			TextStyleName = "NormalText.Important";
+			FontStyleName = "FontAwesome.10";
+			ToolTipText = LOCTEXT("AddModuleToolTip", "Add new module");
+		}
+		else // Compact
+		{
+			TextColor = FSlateColor(FLinearColor::Black);
+			BackgroundColor = FSlateColor(FLinearColor::White);
+			Padding = 0;
+			TextStyleName = "NormalText";
+			FontStyleName = "FontAwesome.7";
+			ToolTipText = LOCTEXT("InsertModuleToolTip", "Insert new module");
+		}
+
 		ChildSlot
 		[
-			SNew(SComboButton)
+			SAssignNew(AddModuleButton, SComboButton)
 			.ButtonStyle(FEditorStyle::Get(), "HoverHintOnly")
-			.ForegroundColor(FSlateColor::UseForeground())
+			.ToolTipText(ToolTipText)
+			.ForegroundColor(TextColor)
 			.HasDownArrow(false)
 			.OnGetMenuContent(this, &SNiagaraStackAddModuleItem::GetAddModuleMenu)
-			.ContentPadding(3)
-			.HAlign(HAlign_Center)
-			.VAlign(VAlign_Center)
+			.ContentPadding(Padding)
 			.ButtonContent()
 			[
-				SNew(STextBlock)
-				.TextStyle(FEditorStyle::Get(), "NormalText.Important")
-				.Font(FEditorStyle::Get().GetFontStyle("FontAwesome.10"))
-				.Text(FText::FromString(FString(TEXT("\xf067"))) /*fa-plus*/)
-				.ToolTipText(LOCTEXT("AddModuleToolTip", "Add new module"))
+				SNew(SBorder)
+				.BorderImage(FEditorStyle::GetBrush("WhiteBrush"))
+				.BorderBackgroundColor(BackgroundColor)
+				.Visibility(this, &SNiagaraStackAddModuleItem::GetButtonContentVisibility)
+				.HAlign(HAlign_Center)
+				.VAlign(VAlign_Center)
+				[
+					SNew(STextBlock)
+					.TextStyle(FEditorStyle::Get(), TextStyleName)
+					.Font(FEditorStyle::Get().GetFontStyle(FontStyleName))
+					.Text(FText::FromString(FString(TEXT("\xf067"))) /*fa-plus*/)
+				]
 			]
 		];
 	}
 
 private:
+	EVisibility GetButtonContentVisibility() const 
+	{
+		if (AddModuleItem->GetDisplayMode() == UNiagaraStackAddModuleItem::EDisplayMode::Standard)
+		{
+			return EVisibility::Visible;
+		}
+		else
+		{
+			return IsHovered() ? EVisibility::Visible : EVisibility::Hidden;
+		}
+	}
+
 	TSharedRef<SWidget> GetAddModuleMenu()
 	{
 		FGraphActionMenuBuilder MenuBuilder;
+		TSharedPtr<SGraphActionMenu> AddModuleMenu;
 
-		return SNew (SBorder)
+		TSharedRef<SBorder> MenuWidget =  SNew(SBorder)
 			.BorderImage(FEditorStyle::GetBrush("Menu.Background"))
 			.Padding(5)
 			[
 				SNew(SBox)
 				[
-					SNew(SGraphActionMenu)
+					SAssignNew(AddModuleMenu, SGraphActionMenu)
 					.OnActionSelected(this, &SNiagaraStackAddModuleItem::OnActionSelected)
-					.OnCreateWidgetForAction(SGraphActionMenu::FOnCreateWidgetForAction::CreateSP(this, &SNiagaraStackAddModuleItem::OnCreateWidgetForAction))
 					.OnCollectAllActions(this, &SNiagaraStackAddModuleItem::CollectAllActions)
 					.AutoExpandActionMenu(false)
 					.ShowFilterTextBox(true)
 				]
 			];
+		AddModuleButton->SetMenuContentWidgetToFocus(AddModuleMenu->GetFilterTextBox()->AsShared()); 
+		return MenuWidget;
 	}
 
 
@@ -1193,20 +1286,20 @@ private:
 				OutAllActions.AddAction(NewNodeAction);
 			}
 		}
-	}
 
-	TSharedRef<SWidget> OnCreateWidgetForAction(struct FCreateWidgetForActionData* const InCreateData)
-	{
-		return SNew(SVerticalBox)
-			+ SVerticalBox::Slot()
-			.AutoHeight()
-			[
-				SNew(STextBlock)
-				.Text(InCreateData->Action->GetMenuDescription())
-				.ToolTipText(InCreateData->Action->GetTooltipDescription())
-			];
-	}
+		const ENiagaraScriptUsage AddModuleOutputUsage = AddModuleItem->GetOutputUsage();
+		if (AddModuleOutputUsage != ENiagaraScriptUsage::EmitterSpawnScript && AddModuleOutputUsage != ENiagaraScriptUsage::EmitterUpdateScript)
+		{
+			// Generate actions for material
+			FText NameText = LOCTEXT("MaterialParameterCategory", "Create Dynamic Parameters");
+			FText Tooltip = LOCTEXT("MaterialParameterTooltip", "Description: Create dynamic material parameters.");
+			FText CategoryName = LOCTEXT("MaterialParameter", "Material");
 
+			TSharedPtr<FNiagaraStackGraphSchemaAction> NewNodeAction(new FNiagaraStackGraphSchemaAction(CategoryName, NameText, Tooltip, 0, FText(),
+				FNiagaraStackGraphSchemaAction::FOnPerformStackAction::CreateUObject(AddModuleItem, &UNiagaraStackAddModuleItem::AddMaterialParameterModule)));
+			OutAllActions.AddAction(NewNodeAction);
+		}
+	}
 
 	void OnActionSelected(const TArray< TSharedPtr<FEdGraphSchemaAction> >& SelectedActions, ESelectInfo::Type InSelectionType)
 	{
@@ -1229,6 +1322,7 @@ private:
 
 private:
 	UNiagaraStackAddModuleItem* AddModuleItem;
+	TSharedPtr<SComboButton> AddModuleButton;
 };
 
 class SNiagaraStackAddRendererItem : public SCompoundWidget
@@ -1309,6 +1403,14 @@ public:
 		[
 			SNew(SHorizontalBox)
 			+ SHorizontalBox::Slot()
+			.Padding(0, 0, 2, 0)
+			.AutoWidth()
+			[
+				SNew(SCheckBox)
+				.IsChecked(this, &SNiagaraStackModuleItem::GetCheckState)
+				.OnCheckStateChanged(this, &SNiagaraStackModuleItem::OnCheckStateChanged)
+			]
+			+ SHorizontalBox::Slot()
 			.Padding(1)
 			[
 				SNew(STextBlock)
@@ -1324,6 +1426,7 @@ public:
 				.IsFocusable(false)
 				.ForegroundColor(FNiagaraEditorWidgetsStyle::Get().GetColor(ModuleItem->GetItemForegroundName()))
 				.ToolTipText(LOCTEXT("UpToolTip", "Move this module up the stack"))
+				.Visibility(this, &SNiagaraStackModuleItem::GetEditButtonVisibility)
 				.OnClicked(this, &SNiagaraStackModuleItem::MoveUpClicked)
 				.Content()
 				[
@@ -1340,6 +1443,7 @@ public:
 				.IsFocusable(false)
 				.ForegroundColor(FNiagaraEditorWidgetsStyle::Get().GetColor(ModuleItem->GetItemForegroundName()))
 				.ToolTipText(LOCTEXT("DownToolTip", "Move this module down the stack"))
+				.Visibility(this, &SNiagaraStackModuleItem::GetEditButtonVisibility)
 				.OnClicked(this, &SNiagaraStackModuleItem::MoveDownClicked)
 				.Content()
 				[
@@ -1350,12 +1454,30 @@ public:
 			]
 			+ SHorizontalBox::Slot()
 			.AutoWidth()
+			[
+				SNew(SButton)
+				.ButtonStyle(FEditorStyle::Get(), "HoverHintOnly")
+				.IsFocusable(false)
+				.ForegroundColor(FNiagaraEditorWidgetsStyle::Get().GetColor(ModuleItem->GetItemForegroundName()))
+				.ToolTipText(LOCTEXT("RefreshTooltip", "Refresh this module"))
+				.Visibility(this, &SNiagaraStackModuleItem::GetRefreshVisibility)
+				.OnClicked(this, &SNiagaraStackModuleItem::RefreshClicked)
+				.Content()
+				[
+					SNew(STextBlock)
+					.Font(FEditorStyle::Get().GetFontStyle("FontAwesome.10"))
+					.Text(FText::FromString(FString(TEXT("\xf021"))))
+				]
+			]
+			+ SHorizontalBox::Slot()
+			.AutoWidth()
 			[
 				SNew(SButton)
 				.ButtonStyle(FEditorStyle::Get(), "HoverHintOnly")
 				.IsFocusable(false)
 				.ForegroundColor(FNiagaraEditorWidgetsStyle::Get().GetColor(ModuleItem->GetItemForegroundName()))
 				.ToolTipText(LOCTEXT("DeleteToolTip", "Delete this module"))
+				.Visibility(this, &SNiagaraStackModuleItem::GetEditButtonVisibility)
 				.OnClicked(this, &SNiagaraStackModuleItem::DeleteClicked)
 				.Content()
 				[
@@ -1381,6 +1503,26 @@ public:
 	}
 
 private:
+	ECheckBoxState GetCheckState() const
+	{
+		return ModuleItem->GetIsEnabled() ? ECheckBoxState::Checked : ECheckBoxState::Unchecked;
+	}
+
+	void OnCheckStateChanged(ECheckBoxState InCheckState)
+	{
+		ModuleItem->SetIsEnabled(InCheckState == ECheckBoxState::Checked);
+	}
+
+	EVisibility GetEditButtonVisibility() const
+	{
+		return ModuleItem->CanMoveAndDelete() ? EVisibility::Visible : EVisibility::Collapsed;
+	}
+
+	EVisibility GetRefreshVisibility() const
+	{
+		return ModuleItem->CanRefresh() ? EVisibility::Visible : EVisibility::Collapsed;
+	}
+
 	FReply MoveUpClicked()
 	{
 		ModuleItem->MoveUp();
@@ -1399,6 +1541,12 @@ private:
 		return FReply::Handled();
 	}
 
+	FReply RefreshClicked()
+	{
+		ModuleItem->Refresh();
+		return FReply::Handled();
+	}
+
 private:
 	UNiagaraStackModuleItem* ModuleItem;
 };
@@ -1417,7 +1565,6 @@ public:
 			[
 				SNew(SHorizontalBox)
 				+ SHorizontalBox::Slot()
-				.Padding(1)
 				[
 					SNew(STextBlock)
 					.TextStyle(FNiagaraEditorWidgetsStyle::Get(), "NiagaraEditor.Stack.GroupText")
@@ -1523,6 +1670,7 @@ public:
 				.ToolTipText_UObject(RendererItem, &UNiagaraStackEntry::GetTooltipText)
 				.Text_UObject(RendererItem, &UNiagaraStackEntry::GetDisplayName)
 			]
+			// Delete button
 			+ SHorizontalBox::Slot()
 			.AutoWidth()
 			[
@@ -1530,7 +1678,8 @@ public:
 				.ButtonStyle(FEditorStyle::Get(), "HoverHintOnly")
 				.IsFocusable(false)
 				.ForegroundColor(FNiagaraEditorWidgetsStyle::Get().GetColor(RendererItem->GetItemForegroundName()))
-				.ToolTipText(LOCTEXT("DeleteToolTip", "Delete this Renderer"))
+				.ToolTipText(LOCTEXT("DeleteRendererToolTip", "Delete this Renderer"))
+				.Visibility(this, &SNiagaraStackRendererItem::GetDeleteButtonVisibility)
 				.OnClicked(this, &SNiagaraStackRendererItem::DeleteClicked)
 				.Content()
 				[
@@ -1539,16 +1688,59 @@ public:
 					.Text(FText::FromString(FString(TEXT("\xf1f8"))))
 				]
 			]
+			// Reset to base Button
+			+ SHorizontalBox::Slot()
+			.AutoWidth()
+			.VAlign(VAlign_Center)
+			.Padding(3, 0, 0, 0)
+			[
+				SNew(SButton)
+				.IsFocusable(false)
+				.ToolTipText(LOCTEXT("ResetRendererToBaseToolTip", "Reset this renderer to the state defined by the parent emitter"))
+				.ButtonStyle(FEditorStyle::Get(), "NoBorder")
+				.ContentPadding(0)
+				.Visibility(this, &SNiagaraStackRendererItem::GetResetToBaseButtonVisibility)
+				.OnClicked(this, &SNiagaraStackRendererItem::ResetToBaseButtonClicked)
+				.Content()
+				[
+					SNew(SImage)
+					.Image(FEditorStyle::GetBrush("PropertyWindow.DiffersFromDefault"))
+					.ColorAndOpacity(FSlateColor(FLinearColor::Green))
+				]
+			]
 		];
 	}
 
 private:
+	EVisibility GetDeleteButtonVisibility() const
+	{
+		return RendererItem->CanDelete() ? EVisibility::Visible : EVisibility::Collapsed;
+	}
+
 	FReply DeleteClicked()
 	{
 		RendererItem->Delete();
 		return FReply::Handled();
 	}
 
+	EVisibility GetResetToBaseButtonVisibility() const
+	{
+		if (RendererItem->CanHaveBase())
+		{
+			return RendererItem->CanResetToBase() ? EVisibility::Visible : EVisibility::Hidden;
+		}
+		else
+		{
+			return EVisibility::Collapsed;
+		}
+	}
+
+	FReply ResetToBaseButtonClicked()
+	{
+		RendererItem->ResetToBase();
+		return FReply::Handled();
+	}
+
 private:
 	UNiagaraStackRendererItem* RendererItem;
 };
@@ -1615,6 +1807,126 @@ private:
 	FText CollapsedToolTipText;
 };
 
+class SNiagaraStackEmitterPropertiesItem : public SCompoundWidget
+{
+public:
+	SLATE_BEGIN_ARGS(SNiagaraStackEmitterPropertiesItem) { }
+	SLATE_END_ARGS();
+
+	void Construct(const FArguments& InArgs, UNiagaraStackEmitterPropertiesItem& InEmitterPropertiesItem)
+	{
+		EmitterPropertiesItem = &InEmitterPropertiesItem;
+
+		ChildSlot
+		[
+			SNew(SHorizontalBox)
+			+ SHorizontalBox::Slot()
+			.Padding(1)
+			[
+				SNew(STextBlock)
+				.TextStyle(FNiagaraEditorWidgetsStyle::Get(), "NiagaraEditor.Stack.ItemText")
+				.ToolTipText_UObject(EmitterPropertiesItem, &UNiagaraStackEntry::GetTooltipText)
+				.Text_UObject(EmitterPropertiesItem, &UNiagaraStackEntry::GetDisplayName)
+			]
+			// Reset to base Button
+			+ SHorizontalBox::Slot()
+			.AutoWidth()
+			.VAlign(VAlign_Center)
+			.Padding(3, 0, 0, 0)
+			[
+				SNew(SButton)
+				.IsFocusable(false)
+				.ToolTipText(LOCTEXT("ResetEmitterPropertiesToBaseToolTip", "Reset the emitter properties to the state defined by the parent emitter"))
+				.ButtonStyle(FEditorStyle::Get(), "NoBorder")
+				.ContentPadding(0)
+				.Visibility(this, &SNiagaraStackEmitterPropertiesItem::GetResetToBaseButtonVisibility)
+				.OnClicked(this, &SNiagaraStackEmitterPropertiesItem::ResetToBaseButtonClicked)
+				.Content()
+				[
+					SNew(SImage)
+					.Image(FEditorStyle::GetBrush("PropertyWindow.DiffersFromDefault"))
+					.ColorAndOpacity(FSlateColor(FLinearColor::Green))
+				]
+			]
+		];
+	}
+
+private:
+	EVisibility GetResetToBaseButtonVisibility() const
+	{
+		return EmitterPropertiesItem->CanResetToBase() ? EVisibility::Visible : EVisibility::Collapsed;
+	}
+
+	FReply ResetToBaseButtonClicked()
+	{
+		EmitterPropertiesItem->ResetToBase();
+		return FReply::Handled();
+	}
+
+private:
+	UNiagaraStackEmitterPropertiesItem* EmitterPropertiesItem;
+};
+
+class SNiagaraStackEventHandlerPropertiesItem : public SCompoundWidget
+{
+public:
+	SLATE_BEGIN_ARGS(SNiagaraStackEventHandlerPropertiesItem) { }
+	SLATE_END_ARGS();
+
+	void Construct(const FArguments& InArgs, UNiagaraStackEventHandlerPropertiesItem& InEventHandlerPropertiesItem)
+	{
+		EventHandlerPropertiesItem = &InEventHandlerPropertiesItem;
+
+		ChildSlot
+		[
+			SNew(SHorizontalBox)
+			+ SHorizontalBox::Slot()
+			.Padding(1)
+			[
+				SNew(STextBlock)
+				.TextStyle(FNiagaraEditorWidgetsStyle::Get(), "NiagaraEditor.Stack.ItemText")
+				.ToolTipText_UObject(EventHandlerPropertiesItem, &UNiagaraStackEntry::GetTooltipText)
+				.Text_UObject(EventHandlerPropertiesItem, &UNiagaraStackEntry::GetDisplayName)
+			]
+			// Reset to base Button
+			+ SHorizontalBox::Slot()
+			.AutoWidth()
+			.VAlign(VAlign_Center)
+			.Padding(3, 0, 0, 0)
+			[
+				SNew(SButton)
+				.IsFocusable(false)
+				.ToolTipText(LOCTEXT("ResetEventHandlerPropertiesToBaseToolTip", "Reset the event handler properties to the state defined by the parent emitter"))
+				.ButtonStyle(FEditorStyle::Get(), "NoBorder")
+				.ContentPadding(0)
+				.Visibility(this, &SNiagaraStackEventHandlerPropertiesItem::GetResetToBaseButtonVisibility)
+				.OnClicked(this, &SNiagaraStackEventHandlerPropertiesItem::ResetToBaseButtonClicked)
+				.Content()
+				[
+					SNew(SImage)
+					.Image(FEditorStyle::GetBrush("PropertyWindow.DiffersFromDefault"))
+					.ColorAndOpacity(FSlateColor(FLinearColor::Green))
+				]
+			]
+		];
+	}
+
+private:
+	EVisibility GetResetToBaseButtonVisibility() const
+	{
+		return EventHandlerPropertiesItem->CanResetToBase() ? EVisibility::Visible : EVisibility::Collapsed;
+	}
+
+	FReply ResetToBaseButtonClicked()
+	{
+		EventHandlerPropertiesItem->ResetToBase();
+		return FReply::Handled();
+	}
+
+private:
+	UNiagaraStackEventHandlerPropertiesItem* EventHandlerPropertiesItem;
+};
+
 class SNiagaraStackTableRow : public STableRow<UNiagaraStackEntry*>
 {
 public:
@@ -1779,8 +2091,9 @@ private:
 	{
 		if (StackEntry->GetCanExpand())
 		{
+			// TODO Cache this and refresh the cache when the entries structure changes.
 			TArray<UNiagaraStackEntry*> Children;
-			StackEntry->GetChildren(Children);
+			StackEntry->GetFilteredChildren(Children);
 			return Children.Num() > 0
 				? EVisibility::Visible
 				: EVisibility::Hidden;
@@ -1867,7 +2180,7 @@ void SNiagaraStack::PrimeTreeExpansion()
 		if (EntryToProcess->GetIsExpanded())
 		{
 			StackTree->SetItemExpansion(EntryToProcess, true);
-			EntryToProcess->GetChildren(EntriesToProcess);
+			EntryToProcess->GetFilteredChildren(EntriesToProcess);
 		}
 	}
 }
@@ -1998,6 +2311,16 @@ TSharedRef<SWidget> SNiagaraStack::ConstructNameWidgetForItem(UNiagaraStackEntry
 		UNiagaraStackItemExpander* ItemExpander = CastChecked<UNiagaraStackItemExpander>(Item);
 		return SNew(SNiagaraStackItemExpander, *ItemExpander);
 	}
+	else if (Item->GetClass()->IsChildOf(UNiagaraStackEmitterPropertiesItem::StaticClass()))
+	{
+		UNiagaraStackEmitterPropertiesItem* PropertiesItem = CastChecked<UNiagaraStackEmitterPropertiesItem>(Item);
+		return SNew(SNiagaraStackEmitterPropertiesItem, *PropertiesItem);
+	}
+	else if (Item->GetClass()->IsChildOf(UNiagaraStackEventHandlerPropertiesItem::StaticClass()))
+	{
+		UNiagaraStackEventHandlerPropertiesItem* PropertiesItem = CastChecked<UNiagaraStackEventHandlerPropertiesItem>(Item);
+		return SNew(SNiagaraStackEventHandlerPropertiesItem, *PropertiesItem);
+	}
 	else
 	{
 		return SNew(STextBlock)
@@ -2039,7 +2362,7 @@ TSharedPtr<SWidget> SNiagaraStack::ConstructValueWidgetForItem(UNiagaraStackEntr
 
 void SNiagaraStack::OnGetChildren(UNiagaraStackEntry* Item, TArray<UNiagaraStackEntry*>& Children)
 {
-	Item->GetChildren(Children);
+	Item->GetFilteredChildren(Children);
 }
 
 float SNiagaraStack::GetNameColumnWidth() const
diff --git a/Engine/Plugins/FX/Niagara/Source/NiagaraShader/Private/NiagaraShader.cpp b/Engine/Plugins/FX/Niagara/Source/NiagaraShader/Private/NiagaraShader.cpp
--- a/Engine/Plugins/FX/Niagara/Source/NiagaraShader/Private/NiagaraShader.cpp
+++ b/Engine/Plugins/FX/Niagara/Source/NiagaraShader/Private/NiagaraShader.cpp
@@ -176,7 +176,7 @@ FNiagaraShaderCompileJob* FNiagaraShaderType::BeginCompileShader(
 	EShaderPlatform Platform,
 	TArray<FNiagaraShaderCompileJob*>& NewJobs,
 	FShaderTarget Target,
-	TArray< TArray<DIGPUBufferParamDescriptor> >&InBufferDescriptors
+	TArray< FDIBufferDescriptorStore >&InBufferDescriptors
 	)
 {
 	FNiagaraShaderCompileJob* NewJob = new FNiagaraShaderCompileJob(ShaderMapId, this, Script->HlslOutput);
@@ -742,6 +742,9 @@ bool FNiagaraShaderMap::TryToAddToExistingCompilationTask(FNiagaraScript* Script
 	if (CorrespondingScripts)
 	{
 		CorrespondingScripts->AddUnique(Script);
+
+		//UE_LOG(LogShaders, Error, TEXT("TryToAddToExistingCompilationTask %p %d"), Script, GetCompilingId());
+
 #if DEBUG_INFINITESHADERCOMPILE
 		UE_LOG(LogTemp, Display, TEXT("Added shader map 0x%08X%08X from Niagara script 0x%08X%08X"), (int)((int64)(this) >> 32), (int)((int64)(this)), (int)((int64)(Script) >> 32), (int)((int64)(Script)));
 #endif
@@ -971,6 +974,11 @@ void FNiagaraShaderMap::RemovePendingScript(FNiagaraScript* Script)
 	{
 		TArray<FNiagaraScript*>& Scripts = It.Value();
 		int32 Result = Scripts.Remove(Script);
+		if (Result)
+		{
+			Script->RemoveOutstandingCompileId(It.Key()->CompilingId);
+			Script->NotifyCompilationFinished();
+		}
 #if DEBUG_INFINITESHADERCOMPILE
 		if ( Result )
 		{
@@ -980,6 +988,22 @@ void FNiagaraShaderMap::RemovePendingScript(FNiagaraScript* Script)
 	}
 }
 
+
+void FNiagaraShaderMap::RemovePendingMap(FNiagaraShaderMap* Map)
+{
+	TArray<FNiagaraScript*>* Scripts = NiagaraShaderMapsBeingCompiled.Find(Map);
+	if (Scripts)
+	{
+		for (FNiagaraScript* Script : *Scripts)
+		{
+			Script->RemoveOutstandingCompileId(Map->CompilingId);
+			Script->NotifyCompilationFinished();
+		}
+	}
+
+	NiagaraShaderMapsBeingCompiled.Remove(Map);
+}
+
 const FNiagaraShaderMap* FNiagaraShaderMap::GetShaderMapBeingCompiled(const FNiagaraScript* Script)
 {
 	// Inefficient search, but only when compiling a lot of shaders
@@ -1030,10 +1054,9 @@ bool FNiagaraShader::Serialize(FArchive& Ar)
 	Ar << OutputIndexBufferParam;
 	Ar << EmitterTickCounterParam;
 
-	Ar << NumInstancesPerThreadParam;
-	Ar << NumInstancesParam;
-	Ar << StartInstanceParam;
-	Ar << GroupStartInstanceParam;
+	Ar << TotalNumInstancesParam;
+	Ar << StartInstancePhase0Param;
+	Ar << StartInstancePhase1Param;
 	Ar << ComponentBufferSizeReadParam;
 	Ar << ComponentBufferSizeWriteParam;
 
@@ -1050,10 +1073,10 @@ bool FNiagaraShader::Serialize(FArchive& Ar)
 	// params for data interface buffers
 	if (Ar.IsLoading())
 	{
-		for (TArray<DIGPUBufferParamDescriptor> &InterfaceDescs : DIBufferDescriptors)
+		for (FDIBufferDescriptorStore &InterfaceDescs : DIBufferDescriptors)
 		{
 			int32 Idx = NameToDIBufferParamMap.AddDefaulted(1);
-			for (DIGPUBufferParamDescriptor &Desc : InterfaceDescs)
+			for (FDIGPUBufferParamDescriptor &Desc : InterfaceDescs.Descriptors)
 			{
 				FShaderResourceParameter &Param = NameToDIBufferParamMap[Idx].Add(*Desc.BufferParamName);
 				Ar << Param;
@@ -1063,9 +1086,9 @@ bool FNiagaraShader::Serialize(FArchive& Ar)
 	else
 	{
 		uint32 Idx = 0;
-		for (TArray<DIGPUBufferParamDescriptor> &InterfaceDescs : DIBufferDescriptors)
+		for (FDIBufferDescriptorStore &InterfaceDescs : DIBufferDescriptors)
 		{
-			for (DIGPUBufferParamDescriptor &Desc : InterfaceDescs)
+			for (FDIGPUBufferParamDescriptor &Desc : InterfaceDescs.Descriptors)
 			{
 				FShaderResourceParameter *Param = NameToDIBufferParamMap[Idx].Find(*Desc.BufferParamName);
 				check(Param);
@@ -1075,8 +1098,6 @@ bool FNiagaraShader::Serialize(FArchive& Ar)
 		}
 	}
 
-	Ar << SimulateStartInstanceParam;
-	Ar << NumThreadGroupsParam;
 	Ar << EmitterConstantBufferParam;
 	Ar << DataInterfaceUniformBufferParam;
 	Ar << NumEventsPerParticleParam;
@@ -1094,8 +1115,15 @@ uint32 FNiagaraShader::GetAllocatedSize() const
 		+ DebugDescription.GetAllocatedSize();
 }
 
+FArchive &operator << (FArchive &Ar, FDIBufferDescriptorStore &Store)
+{
+	Ar << Store.DataInterfaceName;
+	Ar << Store.Descriptors;
+	return Ar;
+}
+
 
-FArchive &operator << (FArchive &Ar, DIGPUBufferParamDescriptor &Desc)
+FArchive &operator << (FArchive &Ar, FDIGPUBufferParamDescriptor &Desc)
 {
 	Ar << Desc.BufferParamName;
 	Ar << Desc.Index;
diff --git a/Engine/Plugins/FX/Niagara/Source/NiagaraShader/Private/NiagaraShaderCompilationManager.cpp b/Engine/Plugins/FX/Niagara/Source/NiagaraShader/Private/NiagaraShaderCompilationManager.cpp
--- a/Engine/Plugins/FX/Niagara/Source/NiagaraShader/Private/NiagaraShaderCompilationManager.cpp
+++ b/Engine/Plugins/FX/Niagara/Source/NiagaraShader/Private/NiagaraShaderCompilationManager.cpp
@@ -129,11 +129,11 @@ void FNiagaraShaderCompilationManager::RunCompileJobs()
 					// Generate a hash of the output and cache it
 					// The shader processing this output will use it to search for existing FShaderResources
 					CurrentJob.Output.GenerateOutputHash();
-					UE_LOG(LogNiagaraShaderCompiler, Log, TEXT("GPU shader compile succeeded."));
+					UE_LOG(LogNiagaraShaderCompiler, Log, TEXT("GPU shader compile succeeded. Id %d"), CurrentJob.Id);
 				}
 				else
 				{
-					UE_LOG(LogNiagaraShaderCompiler, Log, TEXT("ERROR: GPU shader compile failed!"));
+					UE_LOG(LogNiagaraShaderCompiler, Log, TEXT("ERROR: GPU shader compile failed! Id %d"), CurrentJob.Id);
 				}
 
 				CurrentWorkerInfo.bComplete = true;
@@ -208,10 +208,6 @@ void FNiagaraShaderCompilationManager::ProcessAsyncResults()
 }
 
 
-
-
-
-
 void FNiagaraShaderCompilationManager::ProcessCompiledNiagaraShaderMaps(
 	TMap<int32, FNiagaraShaderMapFinalizeResults>& CompiledShaderMaps,
 	float TimeBudget)
@@ -292,6 +288,21 @@ void FNiagaraShaderCompilationManager::ProcessCompiledNiagaraShaderMaps(
 					// This avoids applying shadermaps which are out of date and a newer one is in the async compiling pipeline
 					if (Script->GetScriptID() == CompletedShaderMap->GetShaderMapId().BaseScriptID)
 					{
+						if (Errors.Num() != 0)
+						{
+							FString SourceCode;
+							Script->GetScriptHLSLSource(SourceCode);
+							UE_LOG(LogNiagaraShaderCompiler, Log, TEXT("Compile output as text:"));
+							UE_LOG(LogNiagaraShaderCompiler, Log, TEXT("==================================================================================="));
+							TArray<FString> OutputByLines;
+							SourceCode.ParseIntoArrayLines(OutputByLines, false);
+							for (int32 i = 0; i < OutputByLines.Num(); i++)
+							{
+								UE_LOG(LogNiagaraShaderCompiler, Log, TEXT("/*%04d*/\t\t%s"), i + 1, *OutputByLines[i]);
+							}
+							UE_LOG(LogNiagaraShaderCompiler, Log, TEXT("==================================================================================="));
+						}
+
 						if (!bSuccess)
 						{
 							// Propagate error messages
@@ -328,6 +339,13 @@ void FNiagaraShaderCompilationManager::ProcessCompiledNiagaraShaderMaps(
 							}
 						}
 					}
+					else
+					{
+						if (CompletedShaderMap->IsComplete(Script, true))
+						{
+							Script->NotifyCompilationFinished();
+						}
+					}
 				}
 
 				// Cleanup shader jobs and compile tracking structures
diff --git a/Engine/Plugins/FX/Niagara/Source/NiagaraShader/Private/NiagaraShaderDerivedDataVersion.h b/Engine/Plugins/FX/Niagara/Source/NiagaraShader/Private/NiagaraShaderDerivedDataVersion.h
--- a/Engine/Plugins/FX/Niagara/Source/NiagaraShader/Private/NiagaraShaderDerivedDataVersion.h
+++ b/Engine/Plugins/FX/Niagara/Source/NiagaraShader/Private/NiagaraShaderDerivedDataVersion.h
@@ -11,5 +11,5 @@ NiagaraShaderDerivedDataVersion.h: Shader derived data version for Niagara.
 // In case of merge conflicts with DDC versions, you MUST generate a new GUID and set this new
 // guid as version
 
-#define NIAGARASHADERMAP_DERIVEDDATA_VER		TEXT("1ecd9c26d37a443a9aa832b8feb72475")
+#define NIAGARASHADERMAP_DERIVEDDATA_VER		TEXT("1ecd9c26d37a443a9aa832b8feb72476")
 													  
diff --git a/Engine/Plugins/FX/Niagara/Source/NiagaraShader/Private/NiagaraShared.cpp b/Engine/Plugins/FX/Niagara/Source/NiagaraShader/Private/NiagaraShared.cpp
--- a/Engine/Plugins/FX/Niagara/Source/NiagaraShader/Private/NiagaraShared.cpp
+++ b/Engine/Plugins/FX/Niagara/Source/NiagaraShader/Private/NiagaraShared.cpp
@@ -21,6 +21,15 @@
 	NIAGARASHADER_API FNiagaraCompilationQueue* FNiagaraCompilationQueue::Singleton = nullptr;
 #endif
 
+FNiagaraScript::~FNiagaraScript()
+{
+#if WITH_EDITOR
+	if (IsInGameThread())
+	{
+		FNiagaraCompilationQueue::Get()->RemovePending(this);
+	}
+#endif
+}
 
 /** Populates OutEnvironment with defines needed to compile shaders for this script. */
 void FNiagaraScript::SetupShaderCompilationEnvironment(
@@ -38,6 +47,39 @@ NIAGARASHADER_API bool FNiagaraScript::ShouldCache(EShaderPlatform Platform, con
 	return true;
 }
 
+NIAGARASHADER_API void FNiagaraScript::NotifyCompilationFinished()
+{
+	OnCompilationCompleteDelegate.Broadcast();
+}
+
+NIAGARASHADER_API void FNiagaraScript::CancelCompilation()
+{
+#if WITH_EDITOR
+	if (IsInGameThread())
+	{
+		FNiagaraShaderMap::RemovePendingScript(this);
+		FNiagaraCompilationQueue::Get()->RemovePending(this);
+
+		//UE_LOG(LogShaders, Error, TEXT("CancelCompilation %p."), this);
+		OutstandingCompileShaderMapIds.Empty();
+	}
+#endif
+}
+
+NIAGARASHADER_API void FNiagaraScript::RemoveOutstandingCompileId(const int32 OldOutstandingCompileShaderMapId)
+{
+	if (0 < OutstandingCompileShaderMapIds.Remove(OldOutstandingCompileShaderMapId))
+	{
+		//UE_LOG(LogShaders, Error, TEXT("RemoveOutstandingCompileId %p %d"), this, OldOutstandingCompileShaderMapId);
+	}
+}
+
+NIAGARASHADER_API void FNiagaraScript::Invalidate()
+{
+	CancelCompilation();
+	ReleaseShaderMap();
+}
+
 
 NIAGARASHADER_API void FNiagaraScript::LegacySerialize(FArchive& Ar)
 {
@@ -179,6 +221,16 @@ NIAGARASHADER_API  void FNiagaraScript::SetRenderingThreadShaderMap(FNiagaraShad
 	RenderingThreadShaderMap = InShaderMap;
 }
 
+NIAGARASHADER_API  bool FNiagaraScript::IsCompilationFinished() const
+{
+	bool bRet = GameThreadShaderMap && GameThreadShaderMap.IsValid() && GameThreadShaderMap->IsCompilationFinalized();
+	if (OutstandingCompileShaderMapIds.Num() == 0)
+	{
+		return true;
+	}
+	return bRet;
+}
+
 /**
 * Cache the script's shaders
 */
@@ -229,6 +281,8 @@ bool FNiagaraScript::CacheShaders(const FNiagaraShaderMapId& ShaderMapId, EShade
 		UE_LOG(LogTemp, Display, TEXT("Found existing compiling shader for Niagara script %s, linking to other GameThreadShaderMap 0x%08X%08X"), *GetFriendlyName(), (int)((int64)(GameThreadShaderMap.GetReference()) >> 32), (int)((int64)(GameThreadShaderMap.GetReference())));
 #endif
 		OutstandingCompileShaderMapIds.AddUnique(GameThreadShaderMap->GetCompilingId());
+		//UE_LOG(LogShaders, Error, TEXT("AddUniqueExisting %p %d"), this, GameThreadShaderMap->GetCompilingId());
+
 		// Reset the shader map so we fall back to CPU sim until the compile finishes.
 		GameThreadShaderMap = nullptr;
 		bSucceeded = true;
@@ -300,7 +354,12 @@ NIAGARASHADER_API  FNiagaraShader* FNiagaraScript::GetShader() const
 
 NIAGARASHADER_API  FNiagaraShader* FNiagaraScript::GetShaderGameThread() const
 {
-	return GameThreadShaderMap->GetShader<FNiagaraShader>();
+	if (GameThreadShaderMap)
+	{
+		return GameThreadShaderMap->GetShader<FNiagaraShader>();
+	}
+
+	return nullptr;
 };
 
 
@@ -344,7 +403,9 @@ bool FNiagaraScript::BeginCompileShaderMap(
 
 	// Queue hlsl generation and shader compilation - Unlike materials, we queue this here, and compilation happens from the editor module
 	TRefCountPtr<FNiagaraShaderMap> NewShaderMap = new FNiagaraShaderMap();
-	OutstandingCompileShaderMapIds.AddUnique(NewShaderMap->GetCompilingId());
+	OutstandingCompileShaderMapIds.AddUnique(NewShaderMap->GetCompilingId());		
+	//UE_LOG(LogShaders, Error, TEXT("AddUnique %p %d"), this, NewShaderMap->GetCompilingId());
+
 	FNiagaraCompilationQueue::Get()->Queue(this, NewShaderMap, ShaderMapId, Platform, bApplyCompletedShaderMapForRendering);
 	if (bSynchronous)
 	{
diff --git a/Engine/Plugins/FX/Niagara/Source/NiagaraShader/Public/NiagaraShader.h b/Engine/Plugins/FX/Niagara/Source/NiagaraShader/Public/NiagaraShader.h
--- a/Engine/Plugins/FX/Niagara/Source/NiagaraShader/Public/NiagaraShader.h
+++ b/Engine/Plugins/FX/Niagara/Source/NiagaraShader/Public/NiagaraShader.h
@@ -127,7 +127,7 @@ public:
 	{
 	}
 
-	void SetDatainterfaceBufferDescriptors(const TArray< TArray<DIGPUBufferParamDescriptor>> &InBufferDescriptors)
+	void SetDatainterfaceBufferDescriptors(const TArray< FDIBufferDescriptorStore > &InBufferDescriptors)
 	{
 		DIBufferDescriptors = InBufferDescriptors;
 	}
@@ -170,17 +170,16 @@ public:
 		IntOutputBufferParam.Bind(ParameterMap, TEXT("OutputInt"));
 		OutputIndexBufferParam.Bind(ParameterMap, TEXT("DataSetIndices"));
 		EmitterTickCounterParam.Bind(ParameterMap, TEXT("EmitterTickCounter"));
-		NumInstancesPerThreadParam.Bind(ParameterMap, TEXT("NumInstancesPerThread"));		
 		NumEventsPerParticleParam.Bind(ParameterMap, TEXT("NumEventsPerParticle"));
 		NumParticlesPerEventParam.Bind(ParameterMap, TEXT("NumParticlesPerEvent"));
 		CopyInstancesBeforeStartParam.Bind(ParameterMap, TEXT("CopyInstancesBeforeStart"));
-		NumInstancesParam.Bind(ParameterMap, TEXT("NumInstances"));
-		StartInstanceParam.Bind(ParameterMap, TEXT("StartInstance"));
-		SimulateStartInstanceParam.Bind(ParameterMap, TEXT("SimulateStartInstance"));
-		GroupStartInstanceParam.Bind(ParameterMap, TEXT("GroupStartInstance"));
+
+		TotalNumInstancesParam.Bind(ParameterMap, TEXT("TotalNumInstances"));
+		StartInstancePhase0Param.Bind(ParameterMap, TEXT("StartInstancePhase0"));
+		StartInstancePhase1Param.Bind(ParameterMap, TEXT("StartInstancePhase1"));
+
 		ComponentBufferSizeReadParam.Bind(ParameterMap, TEXT("ComponentBufferSizeRead"));
 		ComponentBufferSizeWriteParam.Bind(ParameterMap, TEXT("ComponentBufferSizeWrite"));
-		NumThreadGroupsParam.Bind(ParameterMap, TEXT("NumThreadGroups"));
 		EmitterConstantBufferParam.Bind(ParameterMap, TEXT("FEmitterParameters"));
 
 		// params for event buffers
@@ -212,11 +211,11 @@ public:
 		BuildDIBufferParamMap(ParameterMap);
 
 		ensure(FloatOutputBufferParam.IsBound() || IntOutputBufferParam.IsBound());	// we should have at least one output buffer we're writing to
-		ensure(OutputIndexBufferParam.IsBound());		
-		ensure(NumInstancesPerThreadParam.IsBound());	
 		ensure(ComponentBufferSizeWriteParam.IsBound());
-		ensure(NumInstancesPerThreadParam.IsBound());
-		ensure(StartInstanceParam.IsBound());
+		ensure(OutputIndexBufferParam.IsBound());
+		ensure(StartInstancePhase0Param.IsBound());
+		ensure(StartInstancePhase1Param.IsBound());
+		ensure(TotalNumInstancesParam.IsBound());
 	}
 
 	// FShader interface.
@@ -231,15 +230,12 @@ public:
 	FShaderUniformBufferParameter EmitterConstantBufferParam;
 	FShaderUniformBufferParameter DataInterfaceUniformBufferParam;
 	FShaderParameter EmitterTickCounterParam;
-	FShaderParameter NumInstancesPerThreadParam;
-	FShaderParameter NumInstancesParam;
 	FShaderParameter NumEventsPerParticleParam;
 	FShaderParameter NumParticlesPerEventParam;
 	FShaderParameter CopyInstancesBeforeStartParam;
-	FShaderParameter StartInstanceParam;
-	FShaderParameter SimulateStartInstanceParam;
-	FShaderParameter GroupStartInstanceParam;
-	FShaderParameter NumThreadGroupsParam;
+	FShaderParameter TotalNumInstancesParam;
+	FShaderParameter StartInstancePhase0Param;
+	FShaderParameter StartInstancePhase1Param;
 	FShaderParameter ComponentBufferSizeReadParam;
 	FShaderParameter ComponentBufferSizeWriteParam;
 	FRWShaderParameter EventIntUAVParams[MAX_CONCURRENT_EVENT_DATASETS];
@@ -263,10 +259,10 @@ public:
 
 	void BuildDIBufferParamMap(const FShaderParameterMap &ParameterMap)
 	{
-		for (TArray<DIGPUBufferParamDescriptor> &InterfaceDescs : DIBufferDescriptors)
+		for (FDIBufferDescriptorStore &InterfaceDescs : DIBufferDescriptors)
 		{
 			int32 Idx = NameToDIBufferParamMap.AddDefaulted(1);
-			for (DIGPUBufferParamDescriptor &Desc : InterfaceDescs)
+			for (FDIGPUBufferParamDescriptor &Desc : InterfaceDescs.Descriptors)
 			{
 				FShaderResourceParameter &Param = NameToDIBufferParamMap[Idx].Add(*Desc.BufferParamName);
 				Param.Bind(ParameterMap, *Desc.BufferParamName);
@@ -275,14 +271,19 @@ public:
 		}
 	}
 
+	TArray< FDIBufferDescriptorStore > &GetDIBufferDescriptors()
+	{
+		return DIBufferDescriptors;
+	}
+
 private:
 	FShaderUniformBufferParameter NiagaraUniformBuffer;
 
 	// buffer descriptors for data interfaces holding names and params for binding
-	TArray< TArray<DIGPUBufferParamDescriptor> > DIBufferDescriptors;
+	TArray< FDIBufferDescriptorStore > DIBufferDescriptors;
 
 	// buffer descriptors for event data sets holding names and params for binding
-	TArray< TArray<DIGPUBufferParamDescriptor> > EventBufferDescriptors;
+	TArray< TArray<FDIGPUBufferParamDescriptor> > EventBufferDescriptors;
 
 	// one map per data interface, mapping buffer names to their params
 	TArray< TMap<FName, FShaderResourceParameter> > NameToDIBufferParamMap;
diff --git a/Engine/Plugins/FX/Niagara/Source/NiagaraShader/Public/NiagaraShaderCompilationManager.h b/Engine/Plugins/FX/Niagara/Source/NiagaraShader/Public/NiagaraShaderCompilationManager.h
--- a/Engine/Plugins/FX/Niagara/Source/NiagaraShader/Public/NiagaraShaderCompilationManager.h
+++ b/Engine/Plugins/FX/Niagara/Source/NiagaraShader/Public/NiagaraShaderCompilationManager.h
@@ -60,7 +60,7 @@ public:
 	FShaderCompilerInput Input;
 	FShaderCompilerOutput Output;
 	FString Hlsl;
-	TArray< TArray<DIGPUBufferParamDescriptor> >DIBufferDescriptors;
+	TArray< FDIBufferDescriptorStore >DIBufferDescriptors;
 
 	FNiagaraShaderCompileJob(uint32 InId, FShaderType* InShaderType, FString InHlsl) :
 		Id(InId),
diff --git a/Engine/Plugins/FX/Niagara/Source/NiagaraShader/Public/NiagaraShaderType.h b/Engine/Plugins/FX/Niagara/Source/NiagaraShader/Public/NiagaraShaderType.h
--- a/Engine/Plugins/FX/Niagara/Source/NiagaraShader/Public/NiagaraShaderType.h
+++ b/Engine/Plugins/FX/Niagara/Source/NiagaraShader/Public/NiagaraShaderType.h
@@ -50,7 +50,7 @@ public:
 	{
 		//const FUniformExpressionSet& UniformExpressionSet;
 		const FString DebugDescription;
-		TArray< TArray<DIGPUBufferParamDescriptor>> DIBufferDescriptors;
+		TArray< FDIBufferDescriptorStore > DIBufferDescriptors;
 
 		CompiledShaderInitializerType(
 			FShaderType* InType,
@@ -58,7 +58,7 @@ public:
 			FShaderResource* InResource,
 			const FSHAHash& InNiagaraShaderMapHash,
 			const FString& InDebugDescription,
-			const TArray< TArray<DIGPUBufferParamDescriptor>> &InDataInterfaceBufferDescriptors
+			const TArray< FDIBufferDescriptorStore > &InDataInterfaceBufferDescriptors
 			)
 		: FGlobalShaderType::CompiledShaderInitializerType(InType,CompilerOutput,InResource, InNiagaraShaderMapHash,nullptr,nullptr)
 		, DebugDescription(InDebugDescription)
@@ -97,7 +97,7 @@ public:
 		EShaderPlatform Platform,
 		TArray<FNiagaraShaderCompileJob*>& NewJobs,
 		FShaderTarget Target,
-		TArray< TArray<DIGPUBufferParamDescriptor> >&InBufferDescriptors
+		TArray< FDIBufferDescriptorStore >&InBufferDescriptors
 		);
 
 	/**
diff --git a/Engine/Plugins/FX/Niagara/Source/NiagaraShader/Public/NiagaraShared.h b/Engine/Plugins/FX/Niagara/Source/NiagaraShader/Public/NiagaraShared.h
--- a/Engine/Plugins/FX/Niagara/Source/NiagaraShader/Public/NiagaraShared.h
+++ b/Engine/Plugins/FX/Niagara/Source/NiagaraShader/Public/NiagaraShared.h
@@ -22,6 +22,7 @@
 #include "SceneTypes.h"
 #include "StaticParameterSet.h"
 #include "Optional.h"
+#include "NiagaraShared.generated.h"
 
 class FNiagaraScript;
 class FNiagaraShaderMap;
@@ -32,29 +33,53 @@ class UNiagaraScript;
 #define MAX_CONCURRENT_EVENT_DATASETS 4
 
 
+
 // holds meta data for GPU buffers from data interface; the DI defines them and returns them, they're
 // passed to the shader compiler, so we can bind parameters to dynamically generated hlsl vars
 //  this is only necessary for buffer params; everything else is in a constant buffer, which we
 //  just allocate memory for and copy the data into
 //
-class DIGPUBufferParamDescriptor
+USTRUCT()
+struct NIAGARASHADER_API FDIGPUBufferParamDescriptor
 {
+	GENERATED_USTRUCT_BODY()
 public:
-	DIGPUBufferParamDescriptor(FString Name, uint32 Idx)
+	FDIGPUBufferParamDescriptor(FString Name, uint32 Idx)
 		: BufferParamName(Name)
 		, Index(Idx)
 	{}
-	DIGPUBufferParamDescriptor()
+	FDIGPUBufferParamDescriptor()
 		: BufferParamName(TEXT("Unknown"))
 		, Index(INDEX_NONE)
 	{}
 
-	friend FArchive &operator << (FArchive &Ar, DIGPUBufferParamDescriptor &Desc);
-
+	friend FArchive &operator << (FArchive &Ar, FDIGPUBufferParamDescriptor &Desc);
+	UPROPERTY()
 	FString BufferParamName;	// the name of the parameter in the hlsl the DI generates
+	UPROPERTY()
 	uint32 Index;				// the index, reflecting the order in which GetBufferDefinitionHLSL declares them
 };
 
+
+
+/** Data interface GPU buffer descriptors for one data interface
+*/
+USTRUCT()
+struct NIAGARASHADER_API FDIBufferDescriptorStore
+{
+	GENERATED_USTRUCT_BODY()
+public:
+	friend FArchive &operator << (FArchive &Ar, FDIBufferDescriptorStore &Desc);
+
+	UPROPERTY()
+	FName DataInterfaceName;
+	UPROPERTY()
+	TArray< FDIGPUBufferParamDescriptor > Descriptors;
+};
+
+
+
+
 /** Stores outputs from the script compile that need to be saved. */
 class FNiagaraComputeShaderCompilationOutput
 {
@@ -67,6 +92,7 @@ public:
 };
 
 
+
 /** Contains all the information needed to uniquely identify a FNiagaraShaderMapID. */
 class FNiagaraShaderMapId
 {
@@ -143,12 +169,13 @@ class FNiagaraCompilationQueue
 public:
 	struct NiagaraCompilationQueueItem
 	{
-		FNiagaraScript *Script;
+		FNiagaraScript* Script;
 		TRefCountPtr<FNiagaraShaderMap>ShaderMap;
 		FNiagaraShaderMapId ShaderMapId;
 		EShaderPlatform Platform;
 		bool bApply;
 	};
+
 	static FNiagaraCompilationQueue *Get()
 	{
 		if (Singleton == nullptr)
@@ -165,6 +192,7 @@ public:
 
 	void Queue(FNiagaraScript *InScript, TRefCountPtr<FNiagaraShaderMap>InShaderMap, const FNiagaraShaderMapId &MapId, EShaderPlatform InPlatform, bool InApply)
 	{
+		check(IsInGameThread());
 		NiagaraCompilationQueueItem NewQueueItem;
 		NewQueueItem.Script = InScript;
 		NewQueueItem.ShaderMap = InShaderMap;
@@ -174,6 +202,18 @@ public:
 		CompilationQueue.Add(NewQueueItem);
 	}
 
+	void RemovePending(FNiagaraScript* InScript)
+	{
+		check(IsInGameThread());
+		for (NiagaraCompilationQueueItem& Item : CompilationQueue)
+		{
+			if (Item.Script == InScript)
+			{
+				Item.Script = nullptr;
+			}
+		}
+	}
+
 private:
 	TArray<NiagaraCompilationQueueItem> CompilationQueue;
 	NIAGARASHADER_API static FNiagaraCompilationQueue *Singleton;
@@ -281,7 +321,8 @@ public:
 	void FlushShadersByShaderType(FShaderType* ShaderType);
 
 	/** Removes a Script from NiagaraShaderMapsBeingCompiled. */
-	static void RemovePendingScript(FNiagaraScript* Script);
+	NIAGARASHADER_API static void RemovePendingScript(FNiagaraScript* Script);
+	NIAGARASHADER_API static void RemovePendingMap(FNiagaraShaderMap* Map);
 
 	/** Finds a shader map currently being compiled that was enqueued for the given script. */
 	static const FNiagaraShaderMap* GetShaderMapBeingCompiled(const FNiagaraScript* Script);
@@ -407,6 +448,7 @@ private:
 };
 
 
+DECLARE_DYNAMIC_MULTICAST_DELEGATE(FOnNiagaraScriptCompilationComplete);
 
 /**
  * FNiagaraScript represents a Niagara script to the shader compilation process
@@ -428,8 +470,7 @@ public:
 	/**
 	 * Destructor
 	 */
-	NIAGARASHADER_API  virtual ~FNiagaraScript()
-	{}
+	NIAGARASHADER_API  virtual ~FNiagaraScript();
 
 	/**
 	 * Caches the shaders for this script with no static parameters on the given platform.
@@ -460,6 +501,8 @@ public:
 	void GetDependentShaderTypes(EShaderPlatform Platform, TArray<FShaderType*>& OutShaderTypes) const;
 	NIAGARASHADER_API  virtual void GetShaderMapId(EShaderPlatform Platform, FNiagaraShaderMapId& OutId) const;
 
+	NIAGARASHADER_API void Invalidate();
+
 	/**
 	 * Should shaders compiled for this script be saved to disk?
 	 */
@@ -468,7 +511,7 @@ public:
 	/**
 	* Called when compilation finishes, after the GameThreadShaderMap is set and the render command to set the RenderThreadShaderMap is queued
 	*/
-	virtual void NotifyCompilationFinished() { }
+	NIAGARASHADER_API virtual void NotifyCompilationFinished();
 
 	/**
 	* Cancels all outstanding compilation jobs
@@ -532,10 +575,7 @@ public:
 	NIAGARASHADER_API  class FNiagaraShaderMap* GetRenderingThreadShaderMap() const;
 
 
-	void RemoveOutstandingCompileId(const int32 OldOutstandingCompileShaderMapId)
-	{
-		OutstandingCompileShaderMapIds.Remove(OldOutstandingCompileShaderMapId);
-	}
+	NIAGARASHADER_API void RemoveOutstandingCompileId(const int32 OldOutstandingCompileShaderMapId);
 
 	NIAGARASHADER_API  virtual void AddReferencedObjects(FReferenceCollector& Collector);
 
@@ -551,7 +591,10 @@ public:
 	* @param OutHighlightMap - source code highlight list
 	* @return - true on Success
 	*/
-	NIAGARASHADER_API  bool GetScriptHLSLSource(FString& OutSource);
+	NIAGARASHADER_API  bool GetScriptHLSLSource(FString& OutSource) {
+		OutSource = HlslOutput;
+		return true;
+	};
 
 	const FString& GetFriendlyName()	const { return FriendlyName; }
 
@@ -569,22 +612,27 @@ public:
 	{
 		CompileErrors = InErrors;
 	}
-
+	
 	FString HlslOutput;
 
 	NIAGARASHADER_API  FNiagaraShader* GetShader() const;
 	NIAGARASHADER_API  FNiagaraShader* GetShaderGameThread() const;
 	
-	void SetDatainterfaceBufferDescriptors(TArray< TArray<DIGPUBufferParamDescriptor> >&InBufferDescriptors)
+	void SetDatainterfaceBufferDescriptors(TArray< FDIBufferDescriptorStore >&InBufferDescriptors)
 	{
 		DIBufferDescriptors = InBufferDescriptors;
 	}
 
-	TArray< TArray<DIGPUBufferParamDescriptor> >& GetDataInterfaceBufferDescriptors()
+	TArray< FDIBufferDescriptorStore >& GetDataInterfaceBufferDescriptors()
 	{
 		return DIBufferDescriptors;
 	}
 
+	NIAGARASHADER_API FOnNiagaraScriptCompilationComplete& OnCompilationComplete()
+	{
+		return OnCompilationCompleteDelegate;
+	}
+
 protected:
 
 	// shared code needed for GetUniformScalarParameterExpressions, GetUniformVectorParameterExpressions, GetUniformCubeTextureExpressions..
@@ -622,7 +670,7 @@ private:
 	FNiagaraShaderMap* RenderingThreadShaderMap;
 
 	// data interface buffer descriptors per data interface; these come from the HLSL translators and need to be passed down to the shader for binding
-	TArray< TArray<DIGPUBufferParamDescriptor> > DIBufferDescriptors;
+	TArray< FDIBufferDescriptorStore > DIBufferDescriptors;
 
 	/** 
 	 * unique identifier
@@ -641,6 +689,8 @@ private:
 	uint32 bLoadedCookedShaderMapId : 1;
 	FNiagaraShaderMapId CookedShaderMapId;
 
+	FOnNiagaraScriptCompilationComplete OnCompilationCompleteDelegate;
+
 	/**
 	* Compiles this script for Platform, storing the result in OutShaderMap if the compile was synchronous
 	*/
@@ -664,3 +714,4 @@ private:
 	friend class FShaderCompilingManager;
 };
 
+
diff --git a/Engine/Plugins/FX/Niagara/Source/NiagaraVertexFactories/Private/NiagaraMeshVertexFactory.cpp b/Engine/Plugins/FX/Niagara/Source/NiagaraVertexFactories/Private/NiagaraMeshVertexFactory.cpp
--- a/Engine/Plugins/FX/Niagara/Source/NiagaraVertexFactories/Private/NiagaraMeshVertexFactory.cpp
+++ b/Engine/Plugins/FX/Niagara/Source/NiagaraVertexFactories/Private/NiagaraMeshVertexFactory.cpp
@@ -19,18 +19,20 @@ public:
 
 	virtual void Bind(const FShaderParameterMap& ParameterMap) override
 	{
-		PrevTransformBuffer.Bind(ParameterMap, TEXT("PrevTransformBuffer"));
+		//PrevTransformBuffer.Bind(ParameterMap, TEXT("PrevTransformBuffer"));
 		NiagaraParticleDataFloat.Bind(ParameterMap, TEXT("NiagaraParticleDataFloat"));
 		NiagaraParticleDataInt.Bind(ParameterMap, TEXT("NiagaraParticleDataInt"));
 		SafeComponentBufferSizeParam.Bind(ParameterMap, TEXT("SafeComponentBufferSize"));
+		MeshFacingMode.Bind(ParameterMap, TEXT("MeshFacingMode"));
 	}
 
 	virtual void Serialize(FArchive& Ar) override
 	{
-		Ar << PrevTransformBuffer;
+		//Ar << PrevTransformBuffer;
 		Ar << NiagaraParticleDataFloat;
 		Ar << NiagaraParticleDataInt;
 		Ar << SafeComponentBufferSizeParam;
+		Ar << MeshFacingMode;
 	}
 
 	virtual void SetMesh(FRHICommandList& RHICmdList, FShader* Shader, const FVertexFactory* VertexFactory, const FSceneView& View, const FMeshBatchElement& BatchElement, uint32 DataFlags) const override
@@ -40,20 +42,22 @@ public:
 		FVertexShaderRHIParamRef VertexShaderRHI = Shader->GetVertexShader();
 		SetUniformBufferParameter(RHICmdList, VertexShaderRHI, Shader->GetUniformBufferParameter<FNiagaraMeshUniformParameters>(), NiagaraMeshVF->GetUniformBuffer());
 
-		SetSRVParameter(RHICmdList, VertexShaderRHI, PrevTransformBuffer, NiagaraMeshVF->GetPreviousTransformBufferSRV());
+		//SetSRVParameter(RHICmdList, VertexShaderRHI, PrevTransformBuffer, NiagaraMeshVF->GetPreviousTransformBufferSRV());
 
 		SetSRVParameter(RHICmdList, VertexShaderRHI, NiagaraParticleDataFloat, NiagaraMeshVF->GetFloatDataSRV());
 		SetSRVParameter(RHICmdList, VertexShaderRHI, NiagaraParticleDataInt, NiagaraMeshVF->GetIntDataSRV());
 		SetShaderValue(RHICmdList, VertexShaderRHI, SafeComponentBufferSizeParam, NiagaraMeshVF->GetComponentBufferSize());
+		SetShaderValue(RHICmdList, VertexShaderRHI, MeshFacingMode, NiagaraMeshVF->GetMeshFacingMode());
 	}
 
 private:
 
-	FShaderResourceParameter PrevTransformBuffer;
+	//FShaderResourceParameter PrevTransformBuffer;
 
 	FShaderResourceParameter NiagaraParticleDataFloat;
 	FShaderResourceParameter NiagaraParticleDataInt;
 	FShaderParameter SafeComponentBufferSizeParam;
+	FShaderParameter MeshFacingMode;
 
 };
 
@@ -187,6 +191,7 @@ void FNiagaraMeshVertexFactory::SetDynamicParameterBuffer(const FVertexBuffer* I
 	}
 }
 
+/*
 uint8* FNiagaraMeshVertexFactory::LockPreviousTransformBuffer(uint32 ParticleCount)
 {
 	const static uint32 ElementSize = sizeof(FVector4);
@@ -217,6 +222,7 @@ FShaderResourceViewRHIParamRef FNiagaraMeshVertexFactory::GetPreviousTransformBu
 {
 	return PrevTransformBuffer.SRV;
 }
+*/
 
 bool FNiagaraMeshVertexFactory::ShouldCache(EShaderPlatform Platform, const class FMaterial* Material, const class FShaderType* ShaderType)
 {
diff --git a/Engine/Plugins/FX/Niagara/Source/NiagaraVertexFactories/Private/NiagaraRibbonVertexFactory.cpp b/Engine/Plugins/FX/Niagara/Source/NiagaraVertexFactories/Private/NiagaraRibbonVertexFactory.cpp
--- a/Engine/Plugins/FX/Niagara/Source/NiagaraVertexFactories/Private/NiagaraRibbonVertexFactory.cpp
+++ b/Engine/Plugins/FX/Niagara/Source/NiagaraVertexFactories/Private/NiagaraRibbonVertexFactory.cpp
@@ -11,27 +11,91 @@
 
 IMPLEMENT_UNIFORM_BUFFER_STRUCT(FNiagaraRibbonUniformParameters,TEXT("NiagaraRibbonVF"));
 
+
+class FNiagaraRibbonVertexFactoryShaderParameters : public FVertexFactoryShaderParameters
+{
+public:
+	virtual void Bind(const FShaderParameterMap& ParameterMap) override
+	{
+	}
+
+	virtual void Serialize(FArchive& Ar) override
+	{
+	}
+};
+
 /**
 * Shader parameters for the beam/trail vertex factory.
 */
-class FNiagaraRibbonVertexFactoryShaderParameters : public FVertexFactoryShaderParameters
+class FNiagaraRibbonVertexFactoryShaderParametersVS : public FNiagaraRibbonVertexFactoryShaderParameters
 {
 public:
 	virtual void Bind(const FShaderParameterMap& ParameterMap) override
 	{
+		/*NiagaraParticleDataFloat.Bind(ParameterMap, TEXT("NiagaraParticleDataFloat"));
+		NiagaraParticleDataInt.Bind(ParameterMap, TEXT("NiagaraParticleDataInt"));
+		SafeComponentBufferSizeParam.Bind(ParameterMap, TEXT("SafeComponentBufferSize"));*/
 	}
 
 	virtual void Serialize(FArchive& Ar) override
 	{
+		/*Ar << PositionOffsetParam;
+		Ar << WidthOffsetParam;
+		Ar << TwistOffsetParam;
+		Ar << ColorOffsetParam;
+
+		Ar << NiagaraParticleDataFloat;
+		Ar << NiagaraParticleDataInt;
+		Ar << SafeComponentBufferSizeParam;*/
 	}
 
 	virtual void SetMesh(FRHICommandList& RHICmdList, FShader* Shader, const FVertexFactory* VertexFactory, const FSceneView& View, const FMeshBatchElement& BatchElement, uint32 DataFlags) const override
 	{
-		FNiagaraRibbonVertexFactory* BeamTrailVF = (FNiagaraRibbonVertexFactory*)VertexFactory;
-		SetUniformBufferParameter(RHICmdList, Shader->GetVertexShader(), Shader->GetUniformBufferParameter<FNiagaraRibbonUniformParameters>(), BeamTrailVF->GetBeamTrailUniformBuffer());
+		FNiagaraRibbonVertexFactory* RibbonVF = (FNiagaraRibbonVertexFactory*)VertexFactory;
+		FVertexShaderRHIParamRef VertexShaderRHI = Shader->GetVertexShader();
+		SetUniformBufferParameter(RHICmdList, Shader->GetVertexShader(), Shader->GetUniformBufferParameter<FNiagaraRibbonUniformParameters>(), RibbonVF->GetBeamTrailUniformBuffer());
+
+		/*
+		SetSRVParameter(RHICmdList, VertexShaderRHI, NiagaraParticleDataFloat, RibbonVF->GetFloatDataSRV());
+		SetSRVParameter(RHICmdList, VertexShaderRHI, NiagaraParticleDataInt, RibbonVF->GetIntDataSRV());
+		SetShaderValue(RHICmdList, VertexShaderRHI, SafeComponentBufferSizeParam, RibbonVF->GetComponentBufferSize());
+		*/
 	}
+
+private:
+	FShaderParameter PositionOffsetParam;
+	FShaderParameter WidthOffsetParam;
+	FShaderParameter TwistOffsetParam;
+	FShaderParameter ColorOffsetParam;
+	FShaderResourceParameter NiagaraParticleDataFloat;
+	FShaderResourceParameter NiagaraParticleDataInt;
+	FShaderParameter SafeComponentBufferSizeParam;
 };
 
+
+
+/**
+* Shader parameters for the beam/trail vertex factory.
+*/
+class FNiagaraRibbonVertexFactoryShaderParametersPS : public FNiagaraRibbonVertexFactoryShaderParameters
+{
+public:
+	virtual void Bind(const FShaderParameterMap& ParameterMap) override
+	{
+	}
+
+	virtual void Serialize(FArchive& Ar) override
+	{
+	}
+
+	virtual void SetMesh(FRHICommandList& RHICmdList, FShader* Shader, const FVertexFactory* VertexFactory, const FSceneView& View, const FMeshBatchElement& BatchElement, uint32 DataFlags) const override
+	{
+		FNiagaraRibbonVertexFactory* RibbonVF = (FNiagaraRibbonVertexFactory*)VertexFactory;
+		SetUniformBufferParameter(RHICmdList, Shader->GetPixelShader(), Shader->GetUniformBufferParameter<FNiagaraRibbonUniformParameters>(), RibbonVF->GetBeamTrailUniformBuffer());
+	}
+};
+
+
 ///////////////////////////////////////////////////////////////////////////////
 /**
 * The particle system beam trail vertex declaration resource type.
@@ -48,23 +112,33 @@ public:
 	{
 		uint16 Stride = sizeof(FNiagaraRibbonVertex);
 		/** The stream to read the vertex position from. */
-		Elements.Add(FVertexElement(0, Offset, VET_Float4, 0, Stride));
-		Offset += sizeof(float) * 4;
-		/** The stream to read the vertex old position from. */
+		Elements.Add(FVertexElement(0, Offset, VET_Float3, 0, Stride));
+		Offset += sizeof(float) * 3;
+		/** The stream to read the direction from. */
 		Elements.Add(FVertexElement(0, Offset, VET_Float3, 1, Stride));
-		Offset += sizeof(float) * 4;
-		/** The stream to read the vertex size/rot/subimage from. */
-		Elements.Add(FVertexElement(0, Offset, VET_Float4, 2, Stride));
-		Offset += sizeof(float) * 4;
+		Offset += sizeof(float) * 3;
+		/** The stream to read the vertex size from. */
+		Elements.Add(FVertexElement(0, Offset, VET_Float1, 2, Stride));
+		Offset += sizeof(float) * 1;
 		/** The stream to read the color from.					*/
-		Elements.Add(FVertexElement(0, Offset, VET_Float4, 4, Stride));
+		Elements.Add(FVertexElement(0, Offset, VET_Float4, 3, Stride));
 		Offset += sizeof(float) * 4;
 		/** The stream to read the texture coordinates from.	*/
-		Elements.Add(FVertexElement(0, Offset, VET_Float4, 3, Stride));
+		Elements.Add(FVertexElement(0, Offset, VET_Float4, 4, Stride));
 		Offset += sizeof(float) * 4;
 
+		/** The stream to read the twist from. */
+		Elements.Add(FVertexElement(0, Offset, VET_Float1, 5, Stride));
+		Offset += sizeof(float);
+
+		/** The stream to read the custom alignment vector from. */
+		Elements.Add(FVertexElement(0, Offset, VET_Float3, 6, Stride));
+		Offset += sizeof(float) * 3;
+
 		/** Dynamic parameters come from a second stream */
-		Elements.Add(FVertexElement(1, 0, VET_Float4, 5, sizeof(FVector4)));
+		Elements.Add(FVertexElement(1, 0, VET_Float4, 7, sizeof(FVector4)));
+		Offset += sizeof(float) * 4;
+
 	}
 
 	virtual void InitDynamicRHI()
@@ -117,7 +191,15 @@ void FNiagaraRibbonVertexFactory::InitRHI()
 
 FVertexFactoryShaderParameters* FNiagaraRibbonVertexFactory::ConstructShaderParameters(EShaderFrequency ShaderFrequency)
 {
-	return ShaderFrequency == SF_Vertex ? new FNiagaraRibbonVertexFactoryShaderParameters() : NULL;
+	if (ShaderFrequency == SF_Vertex)
+	{
+		return new FNiagaraRibbonVertexFactoryShaderParametersVS();
+	}
+	else if (ShaderFrequency == SF_Pixel)
+	{
+		return new FNiagaraRibbonVertexFactoryShaderParametersPS();
+	}
+	return NULL;
 }
 
 void FNiagaraRibbonVertexFactory::SetVertexBuffer(const FVertexBuffer* InBuffer, uint32 StreamOffset, uint32 Stride)
@@ -149,4 +231,4 @@ void FNiagaraRibbonVertexFactory::SetDynamicParameterBuffer(const FVertexBuffer*
 
 ///////////////////////////////////////////////////////////////////////////////
 
-IMPLEMENT_VERTEX_FACTORY_TYPE(FNiagaraRibbonVertexFactory, "/Engine/Private/ParticleBeamTrailVertexFactory.ush", true, false, true, false, false);
+IMPLEMENT_VERTEX_FACTORY_TYPE(FNiagaraRibbonVertexFactory, "/Engine/Private/NiagaraRibbonVertexFactory.ush", true, false, true, false, false);
diff --git a/Engine/Plugins/FX/Niagara/Source/NiagaraVertexFactories/Private/NiagaraSpriteVertexFactory.cpp b/Engine/Plugins/FX/Niagara/Source/NiagaraVertexFactories/Private/NiagaraSpriteVertexFactory.cpp
--- a/Engine/Plugins/FX/Niagara/Source/NiagaraVertexFactories/Private/NiagaraSpriteVertexFactory.cpp
+++ b/Engine/Plugins/FX/Niagara/Source/NiagaraVertexFactories/Private/NiagaraSpriteVertexFactory.cpp
@@ -67,18 +67,16 @@ public:
 		NiagaraParticleDataFloat.Bind(ParameterMap, TEXT("NiagaraParticleDataFloat"));
 		NiagaraParticleDataInt.Bind(ParameterMap, TEXT("NiagaraParticleDataInt"));
 		SafeComponentBufferSizeParam.Bind(ParameterMap, TEXT("SafeComponentBufferSize"));
-		UseCustomAlignmentVector.Bind(ParameterMap, TEXT("UseCustomAlignment"));
-		UseVectorAlignment.Bind(ParameterMap, TEXT("UseVectorAlignment"));
-		UseCameraPlaneFacing.Bind(ParameterMap, TEXT("CameraPlaneFacing"));
+		ParticleAlignmentMode.Bind(ParameterMap, TEXT("ParticleAlignmentMode"));
+		ParticleFacingMode.Bind(ParameterMap, TEXT("ParticleFacingMode"));
 	}
 
 	virtual void Serialize(FArchive& Ar) override
 	{
 		Ar << NumCutoutVerticesPerFrame;
 		Ar << CutoutGeometry;
-		Ar << UseCustomAlignmentVector;
-		Ar << UseVectorAlignment;
-		Ar << UseCameraPlaneFacing;
+		Ar << ParticleFacingMode;
+		Ar << ParticleAlignmentMode;
 		Ar << NiagaraParticleDataFloat;
 		Ar << NiagaraParticleDataInt;
 		Ar << SafeComponentBufferSizeParam;
@@ -100,9 +98,8 @@ public:
 		FShaderResourceViewRHIParamRef NullSRV = GFNiagaraNullSubUVCutoutVertexBuffer.VertexBufferSRV;
 		SetSRVParameter(RHICmdList, VertexShaderRHI, CutoutGeometry, SpriteVF->GetCutoutGeometrySRV() ? SpriteVF->GetCutoutGeometrySRV() : NullSRV);
 
-		SetShaderValue(RHICmdList, VertexShaderRHI, UseCustomAlignmentVector, SpriteVF->GetCustomAlignment());
-		SetShaderValue(RHICmdList, VertexShaderRHI, UseVectorAlignment, SpriteVF->GetVectorAligned());
-		SetShaderValue(RHICmdList, VertexShaderRHI, UseCameraPlaneFacing, SpriteVF->GetCameraPlaneFacing());
+		SetShaderValue(RHICmdList, VertexShaderRHI, ParticleAlignmentMode, SpriteVF->GetAlignmentMode());
+		SetShaderValue(RHICmdList, VertexShaderRHI, ParticleFacingMode, SpriteVF->GetFacingMode());
 
 		SetSRVParameter(RHICmdList, VertexShaderRHI, NiagaraParticleDataFloat, SpriteVF->GetFloatDataSRV());
 		SetSRVParameter(RHICmdList, VertexShaderRHI, NiagaraParticleDataInt, SpriteVF->GetIntDataSRV());
@@ -111,9 +108,10 @@ public:
 
 private:
 	FShaderParameter NumCutoutVerticesPerFrame;
-	FShaderParameter UseCustomAlignmentVector;
-	FShaderParameter UseVectorAlignment;
-	FShaderParameter UseCameraPlaneFacing;
+
+	FShaderParameter ParticleAlignmentMode;
+	FShaderParameter ParticleFacingMode;
+
 	FShaderResourceParameter CutoutGeometry;
 	FShaderResourceParameter NiagaraParticleDataFloat;
 	FShaderResourceParameter NiagaraParticleDataInt;
@@ -190,6 +188,10 @@ public:
 		Elements.Add(FVertexElement(bInstanced ? 1 : 0, Offset, VET_Float4, 3, Stride, bInstanced));
 		Offset += sizeof(float) * 4;
 
+		/** The per-particle dynamic parameter stream */
+		Elements.Add(FVertexElement(bInstanced ? 1 : 0, Offset, VET_Float4, 5, Stride, bInstanced));
+		Offset += sizeof(float) * 4;
+
 		/** The stream to read the custom alignment vector from. */
 		Elements.Add(FVertexElement(bInstanced ? 1 : 0, Offset, VET_Float3, 6, Stride, bInstanced));
 		Offset += sizeof(float) * 3;
@@ -197,15 +199,6 @@ public:
 		Elements.Add(FVertexElement(bInstanced ? 1 : 0, Offset, VET_Float3, 7, Stride, bInstanced));
 		Offset += sizeof(float) * 3;
 
-		/** The per-particle dynamic parameter stream */
-
-		// The -V519 disables a warning from PVS-Studio's static analyzer. It noticed that offset is assigned
-		// twice before being read. It is probably safer to leave the redundant assignments here to reduce
-		// the chance of an error being introduced if this code is modified.
-		Offset = 0;  //-V519
-		Stride = sizeof(float) * 4;
-		Elements.Add(FVertexElement(bInstanced ? 2 : 1, Offset, VET_Float4, 5, Stride, bInstanced));
-		Offset += sizeof(float) * 4;
 
 	}
 
@@ -290,12 +283,11 @@ void FNiagaraSpriteVertexFactory::InitStreams()
 		TexCoordStream->Offset = 0;
 	}
 	FVertexStream* InstanceStream = new(Streams) FVertexStream;
-	FVertexStream* DynamicParameterStream = new(Streams) FVertexStream;
 }
 
 void FNiagaraSpriteVertexFactory::SetInstanceBuffer(const FVertexBuffer* InInstanceBuffer, uint32 StreamOffset, uint32 Stride, bool bInstanced)
 {
-	check(Streams.Num() == (bInstanced ? 3 : 2));
+	check(Streams.Num() == (bInstanced ? 2 : 1));
 	FVertexStream& InstanceStream = Streams[bInstanced ? 1 : 0];
 	InstanceStream.VertexBuffer = InInstanceBuffer;
 	InstanceStream.Stride = Stride;
@@ -308,24 +300,6 @@ void FNiagaraSpriteVertexFactory::SetTexCoordBuffer(const FVertexBuffer* InTexCo
 	TexCoordStream.VertexBuffer = InTexCoordBuffer;
 }
 
-void FNiagaraSpriteVertexFactory::SetDynamicParameterBuffer(const FVertexBuffer* InDynamicParameterBuffer, uint32 StreamOffset, uint32 Stride, bool bInstanced)
-{
-	check(Streams.Num() == (bInstanced ? 3 : 2));
-	FVertexStream& DynamicParameterStream = Streams[bInstanced ? 2 : 1];
-	if (InDynamicParameterBuffer)
-	{
-		DynamicParameterStream.VertexBuffer = InDynamicParameterBuffer;
-		DynamicParameterStream.Stride = Stride;
-		DynamicParameterStream.Offset = StreamOffset;
-	}
-	else
-	{
-		DynamicParameterStream.VertexBuffer = &GNullNiagaraDynamicParameterVertexBuffer;
-		DynamicParameterStream.Stride = 0;
-		DynamicParameterStream.Offset = 0;
-	}
-}
-
 FVertexFactoryShaderParameters* FNiagaraSpriteVertexFactory::ConstructShaderParameters(EShaderFrequency ShaderFrequency)
 {
 	if (ShaderFrequency == SF_Vertex)
diff --git a/Engine/Plugins/FX/Niagara/Source/NiagaraVertexFactories/Private/NiagaraVertexFactoriesModule.cpp b/Engine/Plugins/FX/Niagara/Source/NiagaraVertexFactories/Private/NiagaraVertexFactoriesModule.cpp
--- a/Engine/Plugins/FX/Niagara/Source/NiagaraVertexFactories/Private/NiagaraVertexFactoriesModule.cpp
+++ b/Engine/Plugins/FX/Niagara/Source/NiagaraVertexFactories/Private/NiagaraVertexFactoriesModule.cpp
@@ -2,5 +2,8 @@
 
 #include "NiagaraVertexFactoriesModule.h"
 #include "Modules/ModuleManager.h"
+#include "NiagaraVertexFactory.h"
 
-IMPLEMENT_MODULE(INiagaraVertexFactoriesModule, NiagaraVertexFactories);
\ No newline at end of file
+IMPLEMENT_MODULE(INiagaraVertexFactoriesModule, NiagaraVertexFactories);
+
+FRWBuffer FNiagaraVertexFactoryBase::DummyBuffer;
\ No newline at end of file
diff --git a/Engine/Plugins/FX/Niagara/Source/NiagaraVertexFactories/Public/NiagaraMeshVertexFactory.h b/Engine/Plugins/FX/Niagara/Source/NiagaraVertexFactories/Public/NiagaraMeshVertexFactory.h
--- a/Engine/Plugins/FX/Niagara/Source/NiagaraVertexFactories/Public/NiagaraMeshVertexFactory.h
+++ b/Engine/Plugins/FX/Niagara/Source/NiagaraVertexFactories/Public/NiagaraMeshVertexFactory.h
@@ -64,6 +64,8 @@ struct FNiagaraMeshInstanceVertexPrevTransform
 * Uniform buffer for mesh particle vertex factories.
 */
 BEGIN_UNIFORM_BUFFER_STRUCT(FNiagaraMeshUniformParameters, NIAGARAVERTEXFACTORIES_API)
+	DECLARE_UNIFORM_BUFFER_STRUCT_MEMBER_EX( FMatrix, LocalToWorld, EShaderPrecisionModifier::Half)
+	DECLARE_UNIFORM_BUFFER_STRUCT_MEMBER_EX( FMatrix, LocalToWorldInverseTransposed, EShaderPrecisionModifier::Half)
 DECLARE_UNIFORM_BUFFER_STRUCT_MEMBER(FVector4, SubImageSize)
 DECLARE_UNIFORM_BUFFER_STRUCT_MEMBER(uint32, TexCoordWeightA)
 DECLARE_UNIFORM_BUFFER_STRUCT_MEMBER(uint32, TexCoordWeightB)
@@ -131,6 +133,7 @@ public:
 		: FNiagaraVertexFactoryBase(InType, InFeatureLevel)
 		, DynamicVertexStride(InDynamicVertexStride)
 		, DynamicParameterVertexStride(InDynamicParameterVertexStride)
+		, MeshFacingMode(0)
 		, InstanceVerticesCPU(nullptr)
 	{}
 
@@ -138,6 +141,7 @@ public:
 		: FNiagaraVertexFactoryBase(NVFT_MAX, ERHIFeatureLevel::Num)
 		, DynamicVertexStride(-1)
 		, DynamicParameterVertexStride(-1)
+		, MeshFacingMode(0)
 		, InstanceVerticesCPU(nullptr)
 	{}
 
@@ -167,31 +171,28 @@ public:
 
 	inline FShaderResourceViewRHIParamRef GetFloatDataSRV() const
 	{
-#if PLATFORM_PS4
-		return DataSet->PrevData().GetGPUBufferFloat()->SRV;
-#else
-		check(IsInRenderingThread());
-		return DataSet->PrevDataRender().GetGPUBufferFloat()->SRV;
-#endif
+		const FShaderResourceViewRHIRef& Ret = DataSet->GetRenderDataFloatSRV();
+		if (Ret.IsValid())
+		{
+			return Ret;
+		}
+		return DummyBuffer.SRV;
 	}
+
 	inline FShaderResourceViewRHIParamRef GetIntDataSRV() const
 	{
-#if PLATFORM_PS4
-		return DataSet->PrevData().GetGPUBufferInt()->SRV;
-#else
-		check(IsInRenderingThread());
-		return DataSet->PrevDataRender().GetGPUBufferInt()->SRV;
-#endif
+		const FShaderResourceViewRHIRef& Ret = DataSet->GetRenderDataInt32SRV();
+		if (Ret.IsValid())
+		{
+			return Ret;
+		}
+		return DummyBuffer.SRV;
 	}
 
 	uint32 GetComponentBufferSize()
 	{
-#if PLATFORM_PS4
-		return DataSet->PrevData().GetFloatStride() / sizeof(float);
-#else
-		check(IsInRenderingThread());
-		return DataSet->PrevDataRender().GetFloatStride() / sizeof(float);
-#endif
+		check(!IsInGameThread());
+		return DataSet->CurrDataRender().GetFloatStride() / sizeof(float);
 	}
 
 	/**
@@ -234,9 +235,9 @@ public:
 	*/
 	void SetDynamicParameterBuffer(const FVertexBuffer* InDynamicParameterBuffer, uint32 StreamOffset, uint32 Stride);
 
-	uint8* LockPreviousTransformBuffer(uint32 ParticleCount);
-	void UnlockPreviousTransformBuffer();
-	FShaderResourceViewRHIParamRef GetPreviousTransformBufferSRV() const;
+	//uint8* LockPreviousTransformBuffer(uint32 ParticleCount);
+	//void UnlockPreviousTransformBuffer();
+	//FShaderResourceViewRHIParamRef GetPreviousTransformBufferSRV() const;
 
 	/**
 	* Copy the data from another vertex factory
@@ -256,18 +257,27 @@ public:
 		return InstanceVerticesCPU;
 	}
 
+	uint32 GetMeshFacingMode() const
+	{
+		return MeshFacingMode;
+	}
+
+	void SetMeshFacingMode(uint32 InMode)
+	{
+		MeshFacingMode = InMode;
+	}
+
 protected:
 	FDataType Data;
 	const FNiagaraDataSet *DataSet;
 	/** Stride information for instanced mesh particles */
 	int32 DynamicVertexStride;
 	int32 DynamicParameterVertexStride;
+	uint32 MeshFacingMode;
 
 	/** Uniform buffer with mesh particle parameters. */
 	FUniformBufferRHIParamRef MeshParticleUniformBuffer;
-
-	FDynamicReadBuffer PrevTransformBuffer;
-
+	
 	/** Used to remember this in the case that we reuse the same vertex factory for multiple renders . */
 	FNiagaraMeshInstanceVertices* InstanceVerticesCPU;
 };
diff --git a/Engine/Plugins/FX/Niagara/Source/NiagaraVertexFactories/Public/NiagaraRibbonVertexFactory.h b/Engine/Plugins/FX/Niagara/Source/NiagaraVertexFactories/Public/NiagaraRibbonVertexFactory.h
--- a/Engine/Plugins/FX/Niagara/Source/NiagaraVertexFactories/Public/NiagaraRibbonVertexFactory.h
+++ b/Engine/Plugins/FX/Niagara/Source/NiagaraVertexFactories/Public/NiagaraRibbonVertexFactory.h
@@ -11,7 +11,8 @@
 #include "UniformBuffer.h"
 #include "NiagaraVertexFactory.h"
 #include "SceneView.h"
-
+#include "NiagaraDataSet.h"
+#
 class FMaterial;
 
 
@@ -21,27 +22,28 @@ struct FNiagaraRibbonVertex
 {
 	/** The position of the particle. */
 	FVector Position;
-	/** The relative time of the particle. */
-	float RelativeTime;
-	/** The previous position of the particle. */
-	FVector	OldPosition;
-	/** Value that remains constant over the lifetime of a particle. */
-	float ParticleId;
+	/** The direction of the particle. */
+	FVector Direction;
 	/** The size of the particle. */
-	FVector2D Size;
-	/** The rotation of the particle. */
-	float Rotation;
-	/** The sub-image index for the particle. */
-	float SubImageIndex;
+	float Size;
 	/** The color of the particle. */
 	FLinearColor Color;
 
+	/** The second UV set for the particle				*/
 	float			Tex_U;
 	float			Tex_V;
-
-	/** The second UV set for the particle				*/
 	float			Tex_U2;
 	float			Tex_V2;
+
+
+	/** The rotation of the particle. */
+	float Rotation;
+
+	/** The relative time of the particle. */
+	FVector CustomFacingVector;
+
+	/** The relative time of the particle. */
+	//float NormalizedAge;
 };
 
 
@@ -59,7 +61,13 @@ BEGIN_UNIFORM_BUFFER_STRUCT(FNiagaraRibbonUniformParameters, NIAGARAVERTEXFACTOR
 	DECLARE_UNIFORM_BUFFER_STRUCT_MEMBER(FVector4, CameraRight)
 	DECLARE_UNIFORM_BUFFER_STRUCT_MEMBER(FVector4, CameraUp)
 	DECLARE_UNIFORM_BUFFER_STRUCT_MEMBER(FVector4, ScreenAlignment)
-END_UNIFORM_BUFFER_STRUCT(FNiagaraRibbonUniformParameters)
+	DECLARE_UNIFORM_BUFFER_STRUCT_MEMBER(int, PositionDataOffset)
+	DECLARE_UNIFORM_BUFFER_STRUCT_MEMBER(int, WidthDataOffset)
+	DECLARE_UNIFORM_BUFFER_STRUCT_MEMBER(int, TwistDataOffset)
+	DECLARE_UNIFORM_BUFFER_STRUCT_MEMBER(int, ColorDataOffset)
+	DECLARE_UNIFORM_BUFFER_STRUCT_MEMBER(int, MaterialParamDataOffset)
+	DECLARE_UNIFORM_BUFFER_STRUCT_MEMBER(bool, UseCustomFacing)
+	END_UNIFORM_BUFFER_STRUCT(FNiagaraRibbonUniformParameters)
 typedef TUniformBufferRef<FNiagaraRibbonUniformParameters> FNiagaraRibbonUniformBufferRef;
 
 /**
@@ -77,6 +85,7 @@ public:
 		, IndexBuffer(nullptr)
 		, FirstIndex(0)
 		, OutTriangleCount(0)
+		, DataSet(0)
 	{}
 
 	FNiagaraRibbonVertexFactory()
@@ -84,6 +93,7 @@ public:
 		, IndexBuffer(nullptr)
 		, FirstIndex(0)
 		, OutTriangleCount(0)
+		, DataSet(0)
 	{}
 
 	/**
@@ -126,6 +136,32 @@ public:
 	void SetDynamicParameterBuffer(const FVertexBuffer* InDynamicParameterBuffer, uint32 StreamOffset, uint32 Stride);
 
 
+	void SetParticleData(const FNiagaraDataSet *InDataSet)
+	{
+		DataSet = InDataSet;
+	}
+
+
+	inline FShaderResourceViewRHIParamRef GetFloatDataSRV() const
+	{
+		check(!IsInGameThread());
+		return DataSet->GetRenderDataFloat().SRV;
+	}
+	inline FShaderResourceViewRHIParamRef GetIntDataSRV() const
+	{
+		check(!IsInGameThread());
+		return DataSet->GetRenderDataInt32().SRV;
+	}
+
+	uint32 GetComponentBufferSize()
+	{
+		check(!IsInGameThread());
+		return DataSet->CurrDataRender().GetFloatStride() / sizeof(float);
+	}
+
+
+
+
 	/**
 	* Construct shader parameters for this type of vertex factory.
 	*/
@@ -155,4 +191,6 @@ private:
 	FIndexBuffer* IndexBuffer;
 	uint32 FirstIndex;
 	int32 OutTriangleCount;
+
+	const FNiagaraDataSet *DataSet;
 };
\ No newline at end of file
diff --git a/Engine/Plugins/FX/Niagara/Source/NiagaraVertexFactories/Public/NiagaraSpriteVertexFactory.h b/Engine/Plugins/FX/Niagara/Source/NiagaraVertexFactories/Public/NiagaraSpriteVertexFactory.h
--- a/Engine/Plugins/FX/Niagara/Source/NiagaraVertexFactories/Public/NiagaraSpriteVertexFactory.h
+++ b/Engine/Plugins/FX/Niagara/Source/NiagaraVertexFactories/Public/NiagaraSpriteVertexFactory.h
@@ -19,6 +19,8 @@ class FMaterial;
  * Uniform buffer for particle sprite vertex factories.
  */
 BEGIN_UNIFORM_BUFFER_STRUCT( FNiagaraSpriteUniformParameters, NIAGARAVERTEXFACTORIES_API)
+	DECLARE_UNIFORM_BUFFER_STRUCT_MEMBER_EX( FMatrix, LocalToWorld, EShaderPrecisionModifier::Half)
+	DECLARE_UNIFORM_BUFFER_STRUCT_MEMBER_EX( FMatrix, LocalToWorldInverseTransposed, EShaderPrecisionModifier::Half)
 	DECLARE_UNIFORM_BUFFER_STRUCT_MEMBER_EX( FVector, CustomFacingVectorMask, EShaderPrecisionModifier::Half)
 	DECLARE_UNIFORM_BUFFER_STRUCT_MEMBER_EX( FVector4, TangentSelector, EShaderPrecisionModifier::Half )
 	DECLARE_UNIFORM_BUFFER_STRUCT_MEMBER_EX( FVector4, NormalsSphereCenter, EShaderPrecisionModifier::Half )
@@ -38,8 +40,13 @@ BEGIN_UNIFORM_BUFFER_STRUCT( FNiagaraSpriteUniformParameters, NIAGARAVERTEXFACTO
 	DECLARE_UNIFORM_BUFFER_STRUCT_MEMBER(int, SizeDataOffset)
 	DECLARE_UNIFORM_BUFFER_STRUCT_MEMBER(int, SubimageDataOffset)
 	DECLARE_UNIFORM_BUFFER_STRUCT_MEMBER(int, ColorDataOffset)
+	DECLARE_UNIFORM_BUFFER_STRUCT_MEMBER(int, MaterialParamDataOffset)
 	DECLARE_UNIFORM_BUFFER_STRUCT_MEMBER(int, FacingOffset)
 	DECLARE_UNIFORM_BUFFER_STRUCT_MEMBER(int, AlignmentOffset)
+	DECLARE_UNIFORM_BUFFER_STRUCT_MEMBER(int, SubImageBlendMode)
+	DECLARE_UNIFORM_BUFFER_STRUCT_MEMBER(int, CameraOffsetOffset)
+	DECLARE_UNIFORM_BUFFER_STRUCT_MEMBER(int, UVScaleOffset)
+	DECLARE_UNIFORM_BUFFER_STRUCT_MEMBER(int, ParticleRandomOffset)
 	END_UNIFORM_BUFFER_STRUCT(FNiagaraSpriteUniformParameters)
 
 typedef TUniformBufferRef<FNiagaraSpriteUniformParameters> FNiagaraSpriteUniformBufferRef;
@@ -59,7 +66,8 @@ public:
 		NumVertsInInstanceBuffer(0),
 		NumCutoutVerticesPerFrame(0),
 		CutoutGeometrySRV(nullptr),
-		bCustomAlignment(false)
+		AlignmentMode(0),
+		FacingMode(0)
 	{}
 
 	FNiagaraSpriteVertexFactory()
@@ -67,7 +75,8 @@ public:
 		NumVertsInInstanceBuffer(0),
 		NumCutoutVerticesPerFrame(0),
 		CutoutGeometrySRV(nullptr),
-		bCustomAlignment(false)
+		AlignmentMode(0),
+		FacingMode(0)
 	{}
 
 	// FRenderResource interface.
@@ -96,12 +105,7 @@ public:
 	{
 		NumVertsInInstanceBuffer = InNumVertsInInstanceBuffer;
 	}
-
-	/**
-	 * Set the source vertex buffer that contains particle dynamic parameter data.
-	 */
-	void SetDynamicParameterBuffer(const FVertexBuffer* InDynamicParameterBuffer, uint32 StreamOffset, uint32 Stride, bool bInstanced);
-
+	
 	/**
 	 * Set the uniform buffer for this vertex factory.
 	 */
@@ -135,49 +139,40 @@ public:
 	inline FShaderResourceViewRHIParamRef GetFloatDataSRV() const 
 	{ 
 		check(!IsInGameThread());
-		return DataSet->PrevDataRender().GetGPUBufferFloat()->SRV; 
+		return DataSet->GetRenderDataFloatSRV();
 	}
 	inline FShaderResourceViewRHIParamRef GetIntDataSRV() const 
 	{ 
 		check(!IsInGameThread());
-		return DataSet->PrevDataRender().GetGPUBufferInt()->SRV;
+		return DataSet->GetRenderDataInt32SRV();
 	}
 
 	uint32 GetComponentBufferSize() 
 	{ 
 		check(!IsInGameThread());
-		return DataSet->PrevDataRender().GetFloatStride() / sizeof(float);
+		return DataSet->CurrDataRender().GetFloatStride() / sizeof(float);
 	}
 
-	void SetCustomAlignment(bool bAlign)
+	void SetFacingMode(uint32 InMode)
 	{
-		bCustomAlignment = bAlign;
+		FacingMode = InMode;
 	}
 
-	void SetVectorAligned(bool bAlign)
+	uint32 GetFacingMode()
 	{
-		bVectorAligned = bAlign;
+		return FacingMode;
 	}
 
-	bool GetCustomAlignment()
+	void SetAlignmentMode(uint32 InMode)
 	{
-		return bCustomAlignment;
+		AlignmentMode = InMode;
 	}
 
-	bool GetVectorAligned()
+	uint32 GetAlignmentMode()
 	{
-		return bVectorAligned;
+		return AlignmentMode;
 	}
 
-	void SetCameraPlaneFacing(bool inFacing)
-	{
-		bCameraPlaneFacing = inFacing;
-	}
-
-	bool GetCameraPlaneFacing()
-	{
-		return bCameraPlaneFacing;
-	}
 
 	/**
 	 * Construct shader parameters for this type of vertex factory.
@@ -192,13 +187,12 @@ private:
 
 	int32 NumVertsInInstanceBuffer;
 
-	/** Uniform buffer with sprite paramters. */
+	/** Uniform buffer with sprite parameters. */
 	FUniformBufferRHIParamRef SpriteUniformBuffer;
 
 	int32 NumCutoutVerticesPerFrame;
 	FShaderResourceViewRHIParamRef CutoutGeometrySRV;
 	const FNiagaraDataSet *DataSet;
-	bool bCustomAlignment;	// uses a custom alignment vector from particle data
-	bool bVectorAligned;	// is either velocity or custom vector aligned (not aligned to camera up and right)
-	bool bCameraPlaneFacing;
+	uint32 AlignmentMode;
+	uint32 FacingMode;
 };
diff --git a/Engine/Plugins/FX/Niagara/Source/NiagaraVertexFactories/Public/NiagaraVertexFactory.h b/Engine/Plugins/FX/Niagara/Source/NiagaraVertexFactories/Public/NiagaraVertexFactory.h
--- a/Engine/Plugins/FX/Niagara/Source/NiagaraVertexFactories/Public/NiagaraVertexFactory.h
+++ b/Engine/Plugins/FX/Niagara/Source/NiagaraVertexFactories/Public/NiagaraVertexFactory.h
@@ -48,6 +48,14 @@ public:
 		bNeedsDeclaration = false;
 	}
 
+	NIAGARAVERTEXFACTORIES_API static void Init()
+	{
+		if (DummyBuffer.Buffer == nullptr)
+		{
+			DummyBuffer.Initialize(4, 1, PF_R32_UINT);
+		}
+	}
+
 	static void ModifyCompilationEnvironment(EShaderPlatform Platform, const FMaterial* Material, FShaderCompilerEnvironment& OutEnvironment)
 	{
 		FVertexFactory::ModifyCompilationEnvironment(Platform, Material, OutEnvironment);
@@ -92,6 +100,8 @@ public:
 		return true;
 	}
 
+protected:
+	static FRWBuffer DummyBuffer;
 private:
 
 	/** Last state where we set this. We only need to setup these once per frame, so detemine same frame by number, time, and view family. */
@@ -105,5 +115,4 @@ private:
 
 	/** Whether the vertex factory is in use. */
 	bool bInUse;
-
 };
\ No newline at end of file
diff --git a/Engine/Plugins/FX/NiagaraExtras/NiagaraExtras.uplugin b/Engine/Plugins/FX/NiagaraExtras/NiagaraExtras.uplugin
new file mode 100644
--- /dev/null
+++ b/Engine/Plugins/FX/NiagaraExtras/NiagaraExtras.uplugin
@@ -0,0 +1,16 @@
+{
+	"FileVersion": 3,
+	"Version": 1,
+	"VersionName": "1.0",
+	"FriendlyName": "Niagara Extras",
+	"Description": "Additional content that is too experimental/unproductized/special purpose for the official Niagara plugin.",
+	"Category": "FX",
+	"CreatedBy" : "Epic Games, Inc.",
+	"CreatedByURL" : "http://epicgames.com",
+	"DocsURL": "",
+	"MarketplaceURL": "",
+	"SupportURL": "",
+	"CanContainContent": true,
+	"IsBetaVersion": false,
+	"Installed": false
+}
\ No newline at end of file
diff --git a/Engine/Plugins/Media/MediaCompositing/Source/MediaCompositing/Private/Components/MediaPlaneComponent.cpp b/Engine/Plugins/Media/MediaCompositing/Source/MediaCompositing/Private/Components/MediaPlaneComponent.cpp
--- a/Engine/Plugins/Media/MediaCompositing/Source/MediaCompositing/Private/Components/MediaPlaneComponent.cpp
+++ b/Engine/Plugins/Media/MediaCompositing/Source/MediaCompositing/Private/Components/MediaPlaneComponent.cpp
@@ -140,6 +140,12 @@ namespace
 			VertexFactory.ReleaseResource();
 		}
 
+		SIZE_T GetTypeHash() const
+		{
+			static size_t UniquePointer;
+			return reinterpret_cast<size_t>(&UniquePointer);
+		}
+
 		virtual void CreateRenderThreadResources() override
 		{
 			BuildMesh();
diff --git a/Engine/Plugins/Media/MediaCompositing/Source/MediaCompositing/Private/Components/MediaPlaneFrustumComponent.cpp b/Engine/Plugins/Media/MediaCompositing/Source/MediaCompositing/Private/Components/MediaPlaneFrustumComponent.cpp
--- a/Engine/Plugins/Media/MediaCompositing/Source/MediaCompositing/Private/Components/MediaPlaneFrustumComponent.cpp
+++ b/Engine/Plugins/Media/MediaCompositing/Source/MediaCompositing/Private/Components/MediaPlaneFrustumComponent.cpp
@@ -31,6 +31,12 @@ public:
 		}
 	}
 
+	SIZE_T GetTypeHash() const
+	{
+		static size_t UniquePointer;
+		return reinterpret_cast<size_t>(&UniquePointer);
+	}
+
 	virtual void GetDynamicMeshElements(const TArray<const FSceneView*>& Views, const FSceneViewFamily& ViewFamily, uint32 VisibilityMap, FMeshElementCollector& Collector) const override
 	{
 		if (!ViewTarget)
diff --git a/Engine/Plugins/NetcodeUnitTest/NUTUnrealEngine4/NUTUnrealEngine4.uplugin b/Engine/Plugins/NetcodeUnitTest/NUTUnrealEngine4/NUTUnrealEngine4.uplugin
--- a/Engine/Plugins/NetcodeUnitTest/NUTUnrealEngine4/NUTUnrealEngine4.uplugin
+++ b/Engine/Plugins/NetcodeUnitTest/NUTUnrealEngine4/NUTUnrealEngine4.uplugin
@@ -30,10 +30,6 @@
 	],
 	"Plugins":
 	[
-		{
-			"Name": "OnlineSubsystemUtils",
-			"Enabled": true
-		},
 		{
 			"Name": "NetcodeUnitTest",
 			"Enabled": true
diff --git a/Engine/Plugins/Runtime/CableComponent/Source/CableComponent/Private/CableComponent.cpp b/Engine/Plugins/Runtime/CableComponent/Source/CableComponent/Private/CableComponent.cpp
--- a/Engine/Plugins/Runtime/CableComponent/Source/CableComponent/Private/CableComponent.cpp
+++ b/Engine/Plugins/Runtime/CableComponent/Source/CableComponent/Private/CableComponent.cpp
@@ -51,9 +51,14 @@ struct FCableDynamicData
 //////////////////////////////////////////////////////////////////////////
 // FCableSceneProxy
 
-class FCableSceneProxy : public FPrimitiveSceneProxy
+class FCableSceneProxy final : public FPrimitiveSceneProxy
 {
 public:
+	SIZE_T GetTypeHash() const override
+	{
+		static size_t UniquePointer;
+		return reinterpret_cast<size_t>(&UniquePointer);
+	}
 
 	FCableSceneProxy(UCableComponent* Component)
 		: FPrimitiveSceneProxy(Component)
diff --git a/Engine/Plugins/Runtime/CustomMeshComponent/Source/CustomMeshComponent/Private/CustomMeshComponent.cpp b/Engine/Plugins/Runtime/CustomMeshComponent/Source/CustomMeshComponent/Private/CustomMeshComponent.cpp
--- a/Engine/Plugins/Runtime/CustomMeshComponent/Source/CustomMeshComponent/Private/CustomMeshComponent.cpp
+++ b/Engine/Plugins/Runtime/CustomMeshComponent/Source/CustomMeshComponent/Private/CustomMeshComponent.cpp
@@ -17,9 +17,14 @@
 #include "StaticMeshResources.h"
 
 /** Scene proxy */
-class FCustomMeshSceneProxy : public FPrimitiveSceneProxy
+class FCustomMeshSceneProxy final : public FPrimitiveSceneProxy
 {
 public:
+	SIZE_T GetTypeHash() const override
+	{
+		static size_t UniquePointer;
+		return reinterpret_cast<size_t>(&UniquePointer);
+	}
 
 	FCustomMeshSceneProxy(UCustomMeshComponent* Component)
 		: FPrimitiveSceneProxy(Component)
