Fixed Clang NonUnity compile errors in Niagara plugin

diff --git a/Engine/Plugins/FX/Niagara/Source/Niagara/Private/NiagaraDataSet.cpp b/Engine/Plugins/FX/Niagara/Source/Niagara/Private/NiagaraDataSet.cpp
--- a/Engine/Plugins/FX/Niagara/Source/Niagara/Private/NiagaraDataSet.cpp
+++ b/Engine/Plugins/FX/Niagara/Source/Niagara/Private/NiagaraDataSet.cpp
@@ -5,6 +5,7 @@
 #include "NiagaraShader.h"
 #include "GlobalShader.h"
 #include "UpdateTextureShaders.h"
+#include "ShaderParameterUtils.h"
 #include "NiagaraStats.h"
 
 DECLARE_CYCLE_STAT(TEXT("InitRenderData"), STAT_InitRenderData, STATGROUP_Niagara);
@@ -379,4 +380,4 @@ void FNiagaraDataBuffer::CopyTo(FNiagaraDataBuffer& DestBuffer)
 	DestBuffer.Int32Data = Int32Data;
 	DestBuffer.NumInstancesAllocated = NumInstancesAllocated;
 	DestBuffer.NumInstances = NumInstances;
-}
\ No newline at end of file
+}
diff --git a/Engine/Plugins/FX/Niagara/Source/NiagaraEditor/Private/Customizations/NiagaraTypeCustomizations.cpp b/Engine/Plugins/FX/Niagara/Source/NiagaraEditor/Private/Customizations/NiagaraTypeCustomizations.cpp
--- a/Engine/Plugins/FX/Niagara/Source/NiagaraEditor/Private/Customizations/NiagaraTypeCustomizations.cpp
+++ b/Engine/Plugins/FX/Niagara/Source/NiagaraEditor/Private/Customizations/NiagaraTypeCustomizations.cpp
@@ -1,7 +1,5 @@
 // Copyright 1998-2016 Epic Games, Inc. All Rights Reserved.
 
-#pragma once
-
 #include "NiagaraTypeCustomizations.h"
 #include "CoreMinimal.h"
 #include "PropertyHandle.h"
