Copying //UE4/Dev-Niagara to //UE4/Dev-Main (Source: //UE4/Dev-Niagara @ 4074996)

diff --git a/Engine/Plugins/FX/Niagara/Source/Niagara/Classes/NiagaraConstants.h b/Engine/Plugins/FX/Niagara/Source/Niagara/Classes/NiagaraConstants.h
--- a/Engine/Plugins/FX/Niagara/Source/Niagara/Classes/NiagaraConstants.h
+++ b/Engine/Plugins/FX/Niagara/Source/Niagara/Classes/NiagaraConstants.h
@@ -2,63 +2,85 @@
 
 #pragma once
 #include "NiagaraTypes.h"
+#include "NiagaraCommon.h"
 
-#define SYS_PARAM_ENGINE_DELTA_TIME					FNiagaraVariable(FNiagaraTypeDefinition::GetFloatDef(), TEXT("Engine.DeltaTime"))
-#define SYS_PARAM_ENGINE_INV_DELTA_TIME				FNiagaraVariable(FNiagaraTypeDefinition::GetFloatDef(), TEXT("Engine.InverseDeltaTime"))
-#define SYS_PARAM_ENGINE_POSITION					FNiagaraVariable(FNiagaraTypeDefinition::GetVec3Def(), TEXT("Engine.Owner.Position"))
-#define SYS_PARAM_ENGINE_VELOCITY					FNiagaraVariable(FNiagaraTypeDefinition::GetVec3Def(), TEXT("Engine.Owner.Velocity"))
-#define SYS_PARAM_ENGINE_X_AXIS						FNiagaraVariable(FNiagaraTypeDefinition::GetVec3Def(), TEXT("Engine.Owner.SystemXAxis"))
-#define SYS_PARAM_ENGINE_Y_AXIS						FNiagaraVariable(FNiagaraTypeDefinition::GetVec3Def(), TEXT("Engine.Owner.SystemYAxis"))
-#define SYS_PARAM_ENGINE_Z_AXIS						FNiagaraVariable(FNiagaraTypeDefinition::GetVec3Def(), TEXT("Engine.Owner.SystemZAxis"))
-#define SYS_PARAM_ENGINE_SCALE						FNiagaraVariable(FNiagaraTypeDefinition::GetVec3Def(), TEXT("Engine.Owner.Scale"))
-
-#define SYS_PARAM_ENGINE_LOCAL_TO_WORLD				FNiagaraVariable(FNiagaraTypeDefinition::GetMatrix4Def(), TEXT("Engine.Owner.SystemLocalToWorld"))
-#define SYS_PARAM_ENGINE_WORLD_TO_LOCAL				FNiagaraVariable(FNiagaraTypeDefinition::GetMatrix4Def(), TEXT("Engine.Owner.SystemWorldToLocal"))
-#define SYS_PARAM_ENGINE_LOCAL_TO_WORLD_TRANSPOSED	FNiagaraVariable(FNiagaraTypeDefinition::GetMatrix4Def(), TEXT("Engine.Owner.SystemLocalToWorldTransposed"))
-#define SYS_PARAM_ENGINE_WORLD_TO_LOCAL_TRANSPOSED	FNiagaraVariable(FNiagaraTypeDefinition::GetMatrix4Def(), TEXT("Engine.Owner.SystemWorldToLocalTransposed"))
-#define SYS_PARAM_ENGINE_LOCAL_TO_WORLD_NO_SCALE	FNiagaraVariable(FNiagaraTypeDefinition::GetMatrix4Def(), TEXT("Engine.Owner.SystemLocalToWorldNoScale"))
-#define SYS_PARAM_ENGINE_WORLD_TO_LOCAL_NO_SCALE	FNiagaraVariable(FNiagaraTypeDefinition::GetMatrix4Def(), TEXT("Engine.Owner.SystemWorldToLocalNoScale"))
-
-#define SYS_PARAM_ENGINE_TIME_SINCE_RENDERED		FNiagaraVariable(FNiagaraTypeDefinition::GetFloatDef(), TEXT("Engine.Owner.TimeSinceRendered"))
-#define SYS_PARAM_ENGINE_MIN_DIST_TO_CAMERA			FNiagaraVariable(FNiagaraTypeDefinition::GetFloatDef(), TEXT("Engine.Owner.MinDistanceToCamera"))
-
-#define SYS_PARAM_ENGINE_EXEC_COUNT					FNiagaraVariable(FNiagaraTypeDefinition::GetIntDef(), TEXT("Engine.ExecutionCount"))
-#define SYS_PARAM_ENGINE_EMITTER_NUM_PARTICLES		FNiagaraVariable(FNiagaraTypeDefinition::GetIntDef(), TEXT("Engine.Emitter.NumParticles"))
-#define SYS_PARAM_ENGINE_SYSTEM_NUM_EMITTERS_ALIVE	FNiagaraVariable(FNiagaraTypeDefinition::GetIntDef(), TEXT("Engine.System.NumEmittersAlive"))
-#define SYS_PARAM_ENGINE_SYSTEM_NUM_EMITTERS		FNiagaraVariable(FNiagaraTypeDefinition::GetIntDef(), TEXT("Engine.System.NumEmitters"))
-
-#define SYS_PARAM_EMITTER_AGE						FNiagaraVariable(FNiagaraTypeDefinition::GetFloatDef(), TEXT("Emitter.Age"))
-#define SYS_PARAM_EMITTER_LOCALSPACE				FNiagaraVariable(FNiagaraTypeDefinition::GetBoolDef(), TEXT("Emitter.LocalSpace"))
-#define SYS_PARAM_EMITTER_SPAWNRATE					FNiagaraVariable(FNiagaraTypeDefinition::GetFloatDef(), TEXT("Emitter.SpawnRate"))
-#define SYS_PARAM_EMITTER_SPAWN_INTERVAL			FNiagaraVariable(FNiagaraTypeDefinition::GetFloatDef(), TEXT("Emitter.SpawnInterval"))
-#define SYS_PARAM_EMITTER_INTERP_SPAWN_START_DT		FNiagaraVariable(FNiagaraTypeDefinition::GetFloatDef(), TEXT("Emitter.InterpSpawnStartDt"))
-
-#define SYS_PARAM_PARTICLES_POSITION				FNiagaraVariable(FNiagaraTypeDefinition::GetVec3Def(), TEXT("Particles.Position"))
-#define SYS_PARAM_PARTICLES_VELOCITY				FNiagaraVariable(FNiagaraTypeDefinition::GetVec3Def(), TEXT("Particles.Velocity"))
-#define SYS_PARAM_PARTICLES_COLOR					FNiagaraVariable(FNiagaraTypeDefinition::GetColorDef(), TEXT("Particles.Color"))
-#define SYS_PARAM_PARTICLES_SPRITE_ROTATION			FNiagaraVariable(FNiagaraTypeDefinition::GetFloatDef(), TEXT("Particles.SpriteRotation"))
-#define SYS_PARAM_PARTICLES_NORMALIZED_AGE			FNiagaraVariable(FNiagaraTypeDefinition::GetFloatDef(), TEXT("Particles.NormalizedAge"))
-#define SYS_PARAM_PARTICLES_SPRITE_SIZE				FNiagaraVariable(FNiagaraTypeDefinition::GetVec2Def(), TEXT("Particles.SpriteSize"))
-#define SYS_PARAM_PARTICLES_SPRITE_FACING			FNiagaraVariable(FNiagaraTypeDefinition::GetVec3Def(), TEXT("Particles.SpriteFacing"))
-#define SYS_PARAM_PARTICLES_SPRITE_ALIGNMENT		FNiagaraVariable(FNiagaraTypeDefinition::GetVec3Def(), TEXT("Particles.SpriteAlignment"))
-#define SYS_PARAM_PARTICLES_SUB_IMAGE_INDEX			FNiagaraVariable(FNiagaraTypeDefinition::GetFloatDef(), TEXT("Particles.SubImageIndex"))
-#define SYS_PARAM_PARTICLES_DYNAMIC_MATERIAL_PARAM	FNiagaraVariable(FNiagaraTypeDefinition::GetVec4Def(), TEXT("Particles.DynamicMaterialParameter"))
-#define SYS_PARAM_PARTICLES_SCALE					FNiagaraVariable(FNiagaraTypeDefinition::GetVec3Def(), TEXT("Particles.Scale"))
-#define SYS_PARAM_PARTICLES_LIFETIME				FNiagaraVariable(FNiagaraTypeDefinition::GetFloatDef(), TEXT("Particles.Lifetime"))
-#define SYS_PARAM_PARTICLES_MESH_ORIENTATION		FNiagaraVariable(FNiagaraTypeDefinition::GetVec4Def(), TEXT("Particles.MeshOrientation"))
-#define SYS_PARAM_INSTANCE_ALIVE					FNiagaraVariable(FNiagaraTypeDefinition::GetBoolDef(), TEXT("DataInstance.Alive"))
-#define SYS_PARAM_PARTICLES_UV_SCALE				FNiagaraVariable(FNiagaraTypeDefinition::GetVec2Def(), TEXT("Particles.UVScale"))
-#define SYS_PARAM_PARTICLES_CAMERA_OFFSET			FNiagaraVariable(FNiagaraTypeDefinition::GetFloatDef(), TEXT("Particles.CameraOffset"))
-#define SYS_PARAM_PARTICLES_MATERIAL_RANDOM			FNiagaraVariable(FNiagaraTypeDefinition::GetFloatDef(), TEXT("Particles.MaterialRandom"))
-
-#define SYS_PARAM_PARTICLES_RIBBONID				FNiagaraVariable(FNiagaraTypeDefinition::GetIntDef(), TEXT("Particles.RibbonID"))
-#define SYS_PARAM_PARTICLES_RIBBONWIDTH			FNiagaraVariable(FNiagaraTypeDefinition::GetFloatDef(), TEXT("Particles.RibbonWidth"))
-#define SYS_PARAM_PARTICLES_RIBBONTWIST			FNiagaraVariable(FNiagaraTypeDefinition::GetFloatDef(), TEXT("Particles.RibbonTwist"))
-#define SYS_PARAM_PARTICLES_RIBBONFACING			FNiagaraVariable(FNiagaraTypeDefinition::GetVec3Def(), TEXT("Particles.RibbonFacing"))
+#define SYS_PARAM_ENGINE_DELTA_TIME					INiagaraModule::GetVar_Engine_DeltaTime()
+#define SYS_PARAM_ENGINE_INV_DELTA_TIME				INiagaraModule::GetVar_Engine_InvDeltaTime()
+#define SYS_PARAM_ENGINE_TIME						INiagaraModule::GetVar_Engine_Time()
+#define SYS_PARAM_ENGINE_REAL_TIME					INiagaraModule::GetVar_Engine_RealTime()
+#define SYS_PARAM_ENGINE_POSITION					INiagaraModule::GetVar_Engine_Owner_Position()
+#define SYS_PARAM_ENGINE_VELOCITY					INiagaraModule::GetVar_Engine_Owner_Velocity()
+#define SYS_PARAM_ENGINE_X_AXIS						INiagaraModule::GetVar_Engine_Owner_XAxis()
+#define SYS_PARAM_ENGINE_Y_AXIS						INiagaraModule::GetVar_Engine_Owner_YAxis()
+#define SYS_PARAM_ENGINE_Z_AXIS						INiagaraModule::GetVar_Engine_Owner_ZAxis()
+#define SYS_PARAM_ENGINE_SCALE						INiagaraModule::GetVar_Engine_Owner_Scale()
+
+#define SYS_PARAM_ENGINE_LOCAL_TO_WORLD				INiagaraModule::GetVar_Engine_Owner_SystemLocalToWorld()
+#define SYS_PARAM_ENGINE_WORLD_TO_LOCAL				INiagaraModule::GetVar_Engine_Owner_SystemWorldToLocal()
+#define SYS_PARAM_ENGINE_LOCAL_TO_WORLD_TRANSPOSED	INiagaraModule::GetVar_Engine_Owner_SystemLocalToWorldTransposed()
+#define SYS_PARAM_ENGINE_WORLD_TO_LOCAL_TRANSPOSED	INiagaraModule::GetVar_Engine_Owner_SystemWorldToLocalTransposed()
+#define SYS_PARAM_ENGINE_LOCAL_TO_WORLD_NO_SCALE	INiagaraModule::GetVar_Engine_Owner_SystemLocalToWorldNoScale()
+#define SYS_PARAM_ENGINE_WORLD_TO_LOCAL_NO_SCALE	INiagaraModule::GetVar_Engine_Owner_SystemWorldToLocalNoScale()
+
+#define SYS_PARAM_ENGINE_TIME_SINCE_RENDERED		INiagaraModule::GetVar_Engine_Owner_TimeSinceRendered()
+#define SYS_PARAM_ENGINE_MIN_DIST_TO_CAMERA			INiagaraModule::GetVar_Engine_Owner_MinDistanceToCamera()
+
+#define SYS_PARAM_ENGINE_EXECUTION_STATE			INiagaraModule::GetVar_Engine_Owner_ExecutionState()
+
+#define SYS_PARAM_ENGINE_EXEC_COUNT					INiagaraModule::GetVar_Engine_ExecutionCount()
+#define SYS_PARAM_ENGINE_EMITTER_NUM_PARTICLES		INiagaraModule::GetVar_Engine_Emitter_NumParticles()
+#define SYS_PARAM_ENGINE_SYSTEM_NUM_EMITTERS_ALIVE	INiagaraModule::GetVar_Engine_System_NumEmittersAlive()
+#define SYS_PARAM_ENGINE_SYSTEM_NUM_EMITTERS		INiagaraModule::GetVar_Engine_System_NumEmitters()
+#define SYS_PARAM_ENGINE_NUM_SYSTEM_INSTANCES		INiagaraModule::GetVar_Engine_NumSystemInstances()
+
+#define SYS_PARAM_ENGINE_GLOBAL_SPAWN_COUNT_SCALE	INiagaraModule::GetVar_Engine_GlobalSpawnCountScale()
+#define SYS_PARAM_ENGINE_GLOBAL_SYSTEM_COUNT_SCALE	INiagaraModule::GetVar_Engine_GlobalSystemScale()
+
+#define SYS_PARAM_ENGINE_SYSTEM_AGE					INiagaraModule::GetVar_Engine_System_Age()
+
+#define SYS_PARAM_EMITTER_AGE						INiagaraModule::GetVar_Emitter_Age()
+#define SYS_PARAM_EMITTER_LOCALSPACE				INiagaraModule::GetVar_Emitter_LocalSpace()
+#define SYS_PARAM_EMITTER_SPAWNRATE					INiagaraModule::GetVar_Emitter_SpawnRate()
+#define SYS_PARAM_EMITTER_SPAWN_INTERVAL			INiagaraModule::GetVar_Emitter_SpawnInterval()
+#define SYS_PARAM_EMITTER_INTERP_SPAWN_START_DT		INiagaraModule::GetVar_Emitter_InterpSpawnStartDt()
+
+#define SYS_PARAM_PARTICLES_ID						INiagaraModule::GetVar_Particles_ID()
+#define SYS_PARAM_PARTICLES_POSITION				INiagaraModule::GetVar_Particles_Position()
+#define SYS_PARAM_PARTICLES_VELOCITY				INiagaraModule::GetVar_Particles_Velocity()
+#define SYS_PARAM_PARTICLES_COLOR					INiagaraModule::GetVar_Particles_Color()
+#define SYS_PARAM_PARTICLES_SPRITE_ROTATION			INiagaraModule::GetVar_Particles_SpriteRotation()
+#define SYS_PARAM_PARTICLES_NORMALIZED_AGE			INiagaraModule::GetVar_Particles_NormalizedAge()
+#define SYS_PARAM_PARTICLES_SPRITE_SIZE				INiagaraModule::GetVar_Particles_SpriteSize()
+#define SYS_PARAM_PARTICLES_SPRITE_FACING			INiagaraModule::GetVar_Particles_SpriteFacing()
+#define SYS_PARAM_PARTICLES_SPRITE_ALIGNMENT		INiagaraModule::GetVar_Particles_SpriteAlignment()
+#define SYS_PARAM_PARTICLES_SUB_IMAGE_INDEX			INiagaraModule::GetVar_Particles_SubImageIndex()
+#define SYS_PARAM_PARTICLES_DYNAMIC_MATERIAL_PARAM	INiagaraModule::GetVar_Particles_DynamicMaterialParameter()
+#define SYS_PARAM_PARTICLES_DYNAMIC_MATERIAL_PARAM_1	INiagaraModule::GetVar_Particles_DynamicMaterialParameter1()
+#define SYS_PARAM_PARTICLES_DYNAMIC_MATERIAL_PARAM_2	INiagaraModule::GetVar_Particles_DynamicMaterialParameter2()
+#define SYS_PARAM_PARTICLES_DYNAMIC_MATERIAL_PARAM_3	INiagaraModule::GetVar_Particles_DynamicMaterialParameter3()
+#define SYS_PARAM_PARTICLES_SCALE					INiagaraModule::GetVar_Particles_Scale()
+#define SYS_PARAM_PARTICLES_LIFETIME				INiagaraModule::GetVar_Particles_Lifetime()
+#define SYS_PARAM_PARTICLES_MESH_ORIENTATION		INiagaraModule::GetVar_Particles_MeshOrientation()
+#define SYS_PARAM_PARTICLES_UV_SCALE				INiagaraModule::GetVar_Particles_UVScale()
+#define SYS_PARAM_PARTICLES_CAMERA_OFFSET			INiagaraModule::GetVar_Particles_CameraOffset()
+#define SYS_PARAM_PARTICLES_MATERIAL_RANDOM			INiagaraModule::GetVar_Particles_MaterialRandom()
+#define SYS_PARAM_PARTICLES_LIGHT_RADIUS			INiagaraModule::GetVar_Particles_LightRadius()
+
+#define SYS_PARAM_PARTICLES_RIBBONID				INiagaraModule::GetVar_Particles_RibbonID()
+#define SYS_PARAM_PARTICLES_RIBBONWIDTH				INiagaraModule::GetVar_Particles_RibbonWidth()
+#define SYS_PARAM_PARTICLES_RIBBONTWIST				INiagaraModule::GetVar_Particles_RibbonTwist()
+#define SYS_PARAM_PARTICLES_RIBBONFACING			INiagaraModule::GetVar_Particles_RibbonFacing()
+#define SYS_PARAM_PARTICLES_RIBBONLINKORDER			INiagaraModule::GetVar_Particles_RibbonLinkOrder()
+
+#define SYS_PARAM_INSTANCE_ALIVE					INiagaraModule::GetVar_DataInstance_Alive()
+
+#define TRANSLATOR_PARAM_BEGIN_DEFAULTS				INiagaraModule::GetVar_BeginDefaults()
 
 struct NIAGARA_API FNiagaraConstants
 {
+	static void Init();
 	static const TArray<FNiagaraVariable>& GetEngineConstants();
+	static const TArray<FNiagaraVariable>& GetTranslatorConstants();
 	static FNiagaraVariable UpdateEngineConstant(const FNiagaraVariable& InVar);
 	static const FNiagaraVariable *FindEngineConstant(const FNiagaraVariable& InVar);
 	static FText GetEngineConstantDescription(const FNiagaraVariable& InVar);
@@ -66,7 +88,29 @@ struct NIAGARA_API FNiagaraConstants
 	static const TArray<FNiagaraVariable>& GetCommonParticleAttributes();
 	static FText GetAttributeDescription(const FNiagaraVariable& InVar);
 	static FString GetAttributeDefaultValue(const FNiagaraVariable& InVar);
+	static FNiagaraVariable GetAttributeWithDefaultValue(const FNiagaraVariable& InAttribute);
+	static FNiagaraVariable GetAttributeAsDataSetKey(const FNiagaraVariable& InAttribute);
+	static FNiagaraVariableAttributeBinding GetAttributeDefaultBinding(const FNiagaraVariable& InAttribute);
 
 	static bool IsNiagaraConstant(const FNiagaraVariable& InVar);
 	static const FNiagaraVariableMetaData* GetConstantMetaData(const FNiagaraVariable& InVar);
+
+	static const FNiagaraVariable* GetKnownConstant(const FName& InName, bool bAllowPartialNameMatch);
+
+	static bool IsEngineManagedAttribute(const FNiagaraVariable& Var);
+
+private:
+	static TArray<FNiagaraVariable> SystemParameters;
+	static TArray<FNiagaraVariable> TranslatorParameters;
+	static TMap<FName, FNiagaraVariable> UpdatedSystemParameters;
+	static TMap<FNiagaraVariable, FText> SystemStrMap;
+	static TArray<FNiagaraVariable> Attributes;
+	static TMap<FNiagaraVariable, FString> AttrDefaultsStrMap;
+	static TMap<FNiagaraVariable, FNiagaraVariable> AttrDefaultsValueMap;
+	static TMap<FNiagaraVariable, FNiagaraVariable> AttrDataSetKeyMap;
+	static TMap<FNiagaraVariable, FText> AttrDescStrMap;
+	static TMap<FNiagaraVariable, FNiagaraVariableMetaData> AttrMetaData;
+	
+	static TArray<FNiagaraVariable> EngineManagedAttributes;
+
 };
\ No newline at end of file
diff --git a/Engine/Plugins/FX/Niagara/Source/Niagara/Classes/NiagaraDataInterface.h b/Engine/Plugins/FX/Niagara/Source/Niagara/Classes/NiagaraDataInterface.h
--- a/Engine/Plugins/FX/Niagara/Source/Niagara/Classes/NiagaraDataInterface.h
+++ b/Engine/Plugins/FX/Niagara/Source/Niagara/Classes/NiagaraDataInterface.h
@@ -1,12 +1,13 @@
 // Copyright 1998-2018 Epic Games, Inc. All Rights Reserved.
 #pragma once
 
-#include "NiagaraCommon.h"
+#include "Niagara/Public/NiagaraCommon.h"
 #include "NiagaraShared.h"
 #include "VectorVM.h"
 #include "StaticMeshResources.h"
 #include "Curves/RichCurve.h"
 #include "NiagaraMergeable.h"
+#include "NiagaraDataInterfaceBase.h"
 #include "NiagaraDataInterface.generated.h"
 
 class INiagaraCompiler;
@@ -15,23 +16,6 @@ class UCurveLinearColor;
 class UCurveFloat;
 class FNiagaraSystemInstance;
 
-USTRUCT()
-struct FNiagaraDataInterfaceBufferData
-{
-	GENERATED_BODY()
-public:
-	FNiagaraDataInterfaceBufferData()
-		:UniformName(TEXT("Undefined"))
-	{}
-
-	FNiagaraDataInterfaceBufferData(FName InName)
-		:UniformName(InName)
-	{
-	}
-	FRWBuffer Buffer;
-	FName UniformName;
-};
-
 struct FNDITransformHandlerNoop
 {
 	FORCEINLINE void TransformPosition(FVector& V, FMatrix& M) {  }
@@ -52,9 +36,9 @@ template<typename DirectType, typename NextBinder>
 struct TNDIExplicitBinder
 {
 	template<typename... ParamTypes>
-	static FVMExternalFunction Bind(UNiagaraDataInterface* Interface, const FVMExternalFunctionBindingInfo& BindingInfo, void* InstanceData)
+	static void Bind(UNiagaraDataInterface* Interface, const FVMExternalFunctionBindingInfo& BindingInfo, void* InstanceData, FVMExternalFunction &OutFunc)
 	{
-		return NextBinder::template Bind<ParamTypes..., DirectType>(Interface, BindingInfo, InstanceData);
+		NextBinder::template Bind<ParamTypes..., DirectType>(Interface, BindingInfo, InstanceData, OutFunc);
 	}
 };
 
@@ -63,15 +47,15 @@ template<int32 ParamIdx, typename DataType, typename NextBinder>
 struct TNDIParamBinder
 {
 	template<typename... ParamTypes>
-	static FVMExternalFunction Bind(UNiagaraDataInterface* Interface, const FVMExternalFunctionBindingInfo& BindingInfo, void* InstanceData)
+	static void Bind(UNiagaraDataInterface* Interface, const FVMExternalFunctionBindingInfo& BindingInfo, void* InstanceData, FVMExternalFunction &OutFunc)
 	{
 		if (BindingInfo.InputParamLocations[ParamIdx])
 		{
-			return NextBinder::template Bind<ParamTypes..., FConstantHandler<DataType>>(Interface, BindingInfo, InstanceData);
+			NextBinder::template Bind<ParamTypes..., FConstantHandler<DataType>>(Interface, BindingInfo, InstanceData, OutFunc);
 		}
 		else
 		{
-			return NextBinder::template Bind<ParamTypes..., FRegisterHandler<DataType>>(Interface, BindingInfo, InstanceData);
+			NextBinder::template Bind<ParamTypes..., FRegisterHandler<DataType>>(Interface, BindingInfo, InstanceData, OutFunc);
 		}
 	}
 };
@@ -83,22 +67,88 @@ struct TNDIParamBinder
 struct NDI_FUNC_BINDER(ClassName, FuncName)\
 {\
 	template<typename ... ParamTypes>\
-	static FVMExternalFunction Bind(UNiagaraDataInterface* Interface, const FVMExternalFunctionBindingInfo& BindingInfo, void* InstanceData)\
+	static void Bind(UNiagaraDataInterface* Interface, const FVMExternalFunctionBindingInfo& BindingInfo, void* InstanceData, FVMExternalFunction &OutFunc)\
 	{\
-		return FVMExternalFunction::CreateUObject(CastChecked<ClassName>(Interface), &ClassName::FuncName<ParamTypes...>);\
+		FVMExternalFunction::CreateUObject(CastChecked<ClassName>(Interface), &ClassName::FuncName<ParamTypes...>, OutFunc);\
 	}\
 };
 
+
+
+#define NDI_RAW_FUNC_BINDER(ClassName, FuncName) T##ClassName##_##FuncName##Binder
+
+#define DEFINE_NDI_RAW_FUNC_BINDER(ClassName, FuncName)\
+struct NDI_RAW_FUNC_BINDER(ClassName, FuncName)\
+{\
+	template<typename ... ParamTypes>\
+	static void Bind(UNiagaraDataInterface* Interface, const FVMExternalFunctionBindingInfo& BindingInfo, void* InstanceData, FVMExternalFunction &OutFunc)\
+	{\
+		auto Lambda = [Interface](FVectorVMContext& Context) { static_cast<ClassName*>(Interface)->FuncName<ParamTypes...>(Context); };\
+		OutFunc = FVMExternalFunction::CreateLambda(Lambda);\
+	}\
+};
+
+#if WITH_EDITOR
+// Helper class for GUI error handling
+DECLARE_DELEGATE_RetVal(bool, FNiagaraDataInterfaceFix);
+class FNiagaraDataInterfaceError
+{
+public:
+	FNiagaraDataInterfaceError(FText InErrorText,
+		FText InErrorSummaryText,
+		FNiagaraDataInterfaceFix InFix)
+		: ErrorText(InErrorText)
+		, ErrorSummaryText(InErrorSummaryText)
+		, Fix(InFix)
+
+	{};
+	FNiagaraDataInterfaceError()
+	{};
+	/** Returns true if the error can be fixed automatically. */
+	bool GetErrorFixable() const
+	{
+		return Fix.IsBound();
+	};
+
+	/** Applies the fix if a delegate is bound for it.*/
+	bool TryFixError()
+	{
+		return Fix.IsBound() ? Fix.Execute() : false;
+	};
+
+	/** Full error description text */
+	FText GetErrorText() const
+	{
+		return ErrorText;
+	};
+
+	/** Shortened error description text*/
+	FText GetErrorSummaryText() const
+	{
+		return ErrorSummaryText;
+	};
+
+private:
+	FText ErrorText;
+	FText ErrorSummaryText;
+	FNiagaraDataInterfaceFix Fix;
+};
+#endif
+
 //////////////////////////////////////////////////////////////////////////
 
 /** Base class for all Niagara data interfaces. */
 UCLASS(abstract, EditInlineNew)
-class NIAGARA_API UNiagaraDataInterface : public UNiagaraMergeable
+class NIAGARA_API UNiagaraDataInterface : public UNiagaraDataInterfaceBase
 {
 	GENERATED_UCLASS_BODY()
 		 
 public: 
 
+	// UObject Interface
+	virtual void PostLoad()override;
+	// UObject Interface END
+
 	/** Initializes the per instance data for this interface. Returns false if there was some error and the simulation should be disabled. */
 	virtual bool InitPerInstanceData(void* PerInstanceData, FNiagaraSystemInstance* SystemInstance) { return true; }
 
@@ -119,7 +169,7 @@ public:
 	virtual void GetFunctions(TArray<FNiagaraFunctionSignature>& OutFunctions) {}
 
 	/** Returns the delegate for the passed function signature. */
-	virtual FVMExternalFunction GetVMExternalFunction(const FVMExternalFunctionBindingInfo& BindingInfo, void* InstanceData) { return FVMExternalFunction(); };
+	virtual void GetVMExternalFunction(const FVMExternalFunctionBindingInfo& BindingInfo, void* InstanceData, FVMExternalFunction &OutFunc) { };
 
 	/** Copies the contents of this DataInterface to another.*/
 	bool CopyTo(UNiagaraDataInterface* Destination) const;
@@ -132,32 +182,23 @@ public:
 	/** Determines if this type definition matches to a known data interface type.*/
 	static bool IsDataInterfaceType(const FNiagaraTypeDefinition& TypeDef);
 
-	virtual bool GetFunctionHLSL(const FName& DefinitionFunctionName, FString InstanceFunctionName, TArray<FDIGPUBufferParamDescriptor> &Descriptors, FString &HLSLInterfaceID, FString &OutHLSL)
+	virtual bool GetFunctionHLSL(const FName& DefinitionFunctionName, FString InstanceFunctionName, FNiagaraDataInterfaceGPUParamInfo& ParamInfo, FString& OutHLSL)
 	{
-		check("Unefined HLSL in data interface. Interfaces need to be able to return HLSL for each function they define in GetFunctions.");
+//		checkf(false, TEXT("Unefined HLSL in data interface. Interfaces need to be able to return HLSL for each function they define in GetFunctions."));
 		return false;
 	}
-	virtual void GetBufferDefinitionHLSL(FString DataInterfaceID, TArray<FDIGPUBufferParamDescriptor> &BufferDescriptors, FString &OutHLSL)
-	{
-		check("Unefined HLSL in data interface. Interfaces need to define HLSL for uniforms their functions access.");
-	}
-	virtual TArray<FNiagaraDataInterfaceBufferData> &GetBufferDataArray()
+	virtual void GetParameterDefinitionHLSL(FNiagaraDataInterfaceGPUParamInfo& ParamInfo, FString& OutHLSL)
 	{
-		check("Undefined buffer array access.");
-		return GPUBuffers;
+//		checkf(false, TEXT("Unefined HLSL in data interface. Interfaces need to define HLSL for uniforms their functions access."));
 	}
 
-	virtual void SetupBuffers(FDIBufferDescriptorStore &BufferDescriptors)
-	{
-		check("Undefined buffer setup.");
-	}
+#if WITH_EDITOR	
+	/** Refreshes and returns the errors detected with the corresponding data, if any.*/
+	virtual TArray<FNiagaraDataInterfaceError> GetErrors() { return TArray<FNiagaraDataInterfaceError>(); }
+#endif
 
 protected:
 	virtual bool CopyToInternal(UNiagaraDataInterface* Destination) const;
-
-protected:
-	UPROPERTY()
-	TArray<FNiagaraDataInterfaceBufferData> GPUBuffers;  // not sure whether storing this on the base is the right idea
 };
 
 /** Base class for curve data interfaces which facilitates handling the curve data in a standardized way. */
@@ -180,13 +221,13 @@ protected:
 	/** Remap a sample time for this curve to 0 to 1 between first and last keys for LUT access.*/
 	FORCEINLINE float NormalizeTime(float T)
 	{
-		return bAllowUnnormalizedLUT ? (T - LUTMinTime) * LUTInvTimeRange : T;
+		return (T - LUTMinTime) * LUTInvTimeRange;
 	}
 
 	/** Remap a 0 to 1 value between the first and last keys to a real sample time for this curve. */
 	FORCEINLINE float UnnormalizeTime(float T)
 	{
-		return bAllowUnnormalizedLUT ? (T / LUTInvTimeRange) + LUTMinTime : T;
+		return (T / LUTInvTimeRange) + LUTMinTime;
 	}
 
 public:
@@ -196,7 +237,9 @@ public:
 		, LUTMaxTime(1.0f)
 		, LUTInvTimeRange(1.0f)
 		, bUseLUT(true)
-		, bAllowUnnormalizedLUT(false)
+#if WITH_EDITORONLY_DATA
+		, ShowInCurveEditor(false)
+#endif
 	{
 	}
 
@@ -206,15 +249,18 @@ public:
 		, LUTMaxTime(1.0f)
 		, LUTInvTimeRange(1.0f)
 		, bUseLUT(true)
-		, bAllowUnnormalizedLUT(false)
+#if WITH_EDITORONLY_DATA
+		, ShowInCurveEditor(false)
+#endif
 	{}
 
-	UPROPERTY(EditAnywhere, Category = "Curve")
+	UPROPERTY(EditAnywhere, AdvancedDisplay, Category = "Curve")
 	uint32 bUseLUT : 1;
 	
-	/** Allow the LUT to capture a curve outside the 0 to 1 range. Currently does not work for GPU particles. */
-	UPROPERTY(EditAnywhere, Category = "Curve")
-	uint32 bAllowUnnormalizedLUT : 1;
+#if WITH_EDITORONLY_DATA
+	UPROPERTY(EditAnywhere, Transient, Category = "Curve")
+		bool ShowInCurveEditor;
+#endif
 
 	enum
 	{
@@ -241,33 +287,34 @@ public:
 	/** Gets information for all of the curves owned by this curve data interface. */
 	virtual void GetCurveData(TArray<FCurveData>& OutCurveData) { }
 
-	virtual bool GetFunctionHLSL(const FName&  DefinitionFunctionName, FString InstanceFunctionName, TArray<FDIGPUBufferParamDescriptor> &Descriptors, FString &HLSLInterfaceID, FString &OutHLSL)
-	{
-		check("Undefined HLSL in data interface. All curve interfaces need to define this function and return SampleCurve code");
-		return false;
-	}
-	virtual void GetBufferDefinitionHLSL(FString DataInterfaceID, TArray<FDIGPUBufferParamDescriptor> &BufferDescriptors, FString &OutHLSL)
-	{
-		check("Unefined HLSL in data interface. Interfaces need to define HLSL for uniforms their functions access.");
-	}
-	virtual TArray<FNiagaraDataInterfaceBufferData> &GetBufferDataArray()
-	{
-		check("Undefined buffer array access.");
-		return GPUBuffers;
-	}
-	virtual void SetupBuffers(FDIBufferDescriptorStore &BufferDescriptors)
+	virtual void GetParameterDefinitionHLSL(FNiagaraDataInterfaceGPUParamInfo& ParamInfo, FString& OutHLSL)override;
+	virtual FNiagaraDataInterfaceParametersCS* ConstructComputeParameters()const override;
+
+	virtual int32 GetCurveNumElems()const 
 	{
-		check("Undefined buffer setup.");
+		checkf(false, TEXT("You must implement this function so the GPU buffer can be created of the correct size."));
+		return 0; 
 	}
 
+	//TODO: Make this a texture and get HW filter + clamping?
+	FRWBuffer& GetCurveLUTGPUBuffer();
+
 	//UNiagaraDataInterface interface
 	virtual bool Equals(const UNiagaraDataInterface* Other) const override;
 	virtual bool CanExecuteOnTarget(ENiagaraSimTarget Target)const override { return true; }
+	virtual void UpdateLUT() { };
+
+	FORCEINLINE float GetMinTime()const { return LUTMinTime; }
+	FORCEINLINE float GetMaxTime()const { return LUTMaxTime; }
+	FORCEINLINE float GetInvTimeRange()const { return LUTInvTimeRange; }
+
 
 protected:
 	virtual bool CopyToInternal(UNiagaraDataInterface* Destination) const override;
 	virtual bool CompareLUTS(const TArray<float>& OtherLUT) const;
 	//UNiagaraDataInterface interface END
+
+	FRWBuffer CurveLUT;
 };
 
 //External function binder choosing between template specializations based on if a curve should use the LUT over full evaluation.
@@ -275,16 +322,16 @@ template<typename NextBinder>
 struct TCurveUseLUTBinder
 {
 	template<typename... ParamTypes>
-	static FVMExternalFunction Bind(UNiagaraDataInterface* Interface, const FVMExternalFunctionBindingInfo& BindingInfo, void* InstanceData)
+	static void Bind(UNiagaraDataInterface* Interface, const FVMExternalFunctionBindingInfo& BindingInfo, void* InstanceData, FVMExternalFunction &OutFunc)
 	{
 		UNiagaraDataInterfaceCurveBase* CurveInterface = CastChecked<UNiagaraDataInterfaceCurveBase>(Interface);
 		if (CurveInterface->bUseLUT)
 		{
-			return NextBinder::template Bind<ParamTypes..., TIntegralConstant<bool, true>>(Interface, BindingInfo, InstanceData);
+			NextBinder::template Bind<ParamTypes..., TIntegralConstant<bool, true>>(Interface, BindingInfo, InstanceData, OutFunc);
 		}
 		else
 		{
-			return NextBinder::template Bind<ParamTypes..., TIntegralConstant<bool, false>>(Interface, BindingInfo, InstanceData);
+			NextBinder::template Bind<ParamTypes..., TIntegralConstant<bool, false>>(Interface, BindingInfo, InstanceData, OutFunc);
 		}
 	}
 };
diff --git a/Engine/Plugins/FX/Niagara/Source/Niagara/Classes/NiagaraDataInterfaceCollisionQuery.h b/Engine/Plugins/FX/Niagara/Source/Niagara/Classes/NiagaraDataInterfaceCollisionQuery.h
--- a/Engine/Plugins/FX/Niagara/Source/Niagara/Classes/NiagaraDataInterfaceCollisionQuery.h
+++ b/Engine/Plugins/FX/Niagara/Source/Niagara/Classes/NiagaraDataInterfaceCollisionQuery.h
@@ -53,7 +53,7 @@ public:
 
 
 	virtual void GetFunctions(TArray<FNiagaraFunctionSignature>& OutFunctions)override;
-	virtual FVMExternalFunction GetVMExternalFunction(const FVMExternalFunctionBindingInfo& BindingInfo, void* InstanceData)override;
+	virtual void GetVMExternalFunction(const FVMExternalFunctionBindingInfo& BindingInfo, void* InstanceData, FVMExternalFunction &OutFunc) override;
 
 	// VM functions
 	template<typename PosTypeX, typename PosTypeY, typename PosTypeZ, typename VelTypeX, typename VelTypeY, typename VelTypeZ, typename DTType>
@@ -62,16 +62,21 @@ public:
 	template<typename IDType>
 	void ReadQuery(FVectorVMContext& Context);
 
-	template<typename InQueryIDType, typename PosTypeX, typename PosTypeY, typename PosTypeZ, typename VelTypeX, typename VelTypeY, typename VelTypeZ, typename DTType, typename SizeType>
+	template<typename InQueryIDType, typename PosTypeX, typename PosTypeY, typename PosTypeZ, typename DirTypeX, typename DirTypeY, typename DirTypeZ, typename DTType, typename SizeType, typename DepthBoundsType>
 	void PerformQuery(FVectorVMContext& Context);
 
 	virtual bool Equals(const UNiagaraDataInterface* Other) const override;
+	virtual bool CanExecuteOnTarget(ENiagaraSimTarget Target) const override { return true; }
 
-	virtual bool GetFunctionHLSL(const FName&  DefinitionFunctionName, FString InstanceFunctionName, TArray<FDIGPUBufferParamDescriptor> &Descriptors, FString &HLSLInterfaceID, FString &OutHLSL) override;
-	virtual void GetBufferDefinitionHLSL(FString DataInterfaceID, TArray<FDIGPUBufferParamDescriptor> &BufferDescriptors, FString &OutHLSL) override;
-	virtual TArray<FNiagaraDataInterfaceBufferData> &GetBufferDataArray() override;
-	virtual void SetupBuffers(FDIBufferDescriptorStore &BufferDescriptors) override;
+	virtual bool GetFunctionHLSL(const FName&  DefinitionFunctionName, FString InstanceFunctionName, FNiagaraDataInterfaceGPUParamInfo& ParamInfo, FString& OutHLSL) override;
+	virtual void GetParameterDefinitionHLSL(FNiagaraDataInterfaceGPUParamInfo& ParamInfo, FString& OutHLSL) override;
+	virtual FNiagaraDataInterfaceParametersCS* ConstructComputeParameters()const override;
+	
 
 protected:
 	virtual bool CopyToInternal(UNiagaraDataInterface* Destination) const override;
+
+private:
+
+	static FCriticalSection CriticalSection;
 };
diff --git a/Engine/Plugins/FX/Niagara/Source/Niagara/Classes/NiagaraDataInterfaceColorCurve.h b/Engine/Plugins/FX/Niagara/Source/Niagara/Classes/NiagaraDataInterfaceColorCurve.h
--- a/Engine/Plugins/FX/Niagara/Source/Niagara/Classes/NiagaraDataInterfaceColorCurve.h
+++ b/Engine/Plugins/FX/Niagara/Source/Niagara/Classes/NiagaraDataInterfaceColorCurve.h
@@ -22,12 +22,6 @@ class NIAGARA_API UNiagaraDataInterfaceColorCurve : public UNiagaraDataInterface
 {
 	GENERATED_UCLASS_BODY()
 public:
-
-#if WITH_EDITORONLY_DATA
-	UPROPERTY(EditAnywhere, Category = "Curve", meta=(AllowedClasses = CurveLinearColor))
-	FSoftObjectPath CurveToCopy;
-#endif
-
 	UPROPERTY(EditAnywhere, Category = "Curve")
 	FRichCurve RedCurve;
 
@@ -43,9 +37,6 @@ public:
 	//UObject Interface
 	virtual void PostInitProperties() override;
 	virtual void PostLoad() override;
-#if WITH_EDITOR
-	virtual void PostEditChangeProperty(struct FPropertyChangedEvent& PropertyChangedEvent) override;
-#endif
 	//UObject Interface End
 
 	enum
@@ -57,7 +48,7 @@ public:
 	void UpdateLUT();
 
 	virtual void GetFunctions(TArray<FNiagaraFunctionSignature>& OutFunctions)override;
-	virtual FVMExternalFunction GetVMExternalFunction(const FVMExternalFunctionBindingInfo& BindingInfo, void* InstanceData)override;
+	virtual void GetVMExternalFunction(const FVMExternalFunctionBindingInfo& BindingInfo, void* InstanceData, FVMExternalFunction &OutFunc) override;
 
 	template<typename UseLUT, typename XParamType>
 	void SampleCurve(FVectorVMContext& Context);
@@ -67,15 +58,15 @@ public:
 	//~ UNiagaraDataInterfaceCurveBase interface
 	virtual void GetCurveData(TArray<FCurveData>& OutCurveData) override;
 
-	virtual bool GetFunctionHLSL(const FName&  DefinitionFunctionName, FString InstanceFunctionName, TArray<FDIGPUBufferParamDescriptor> &Descriptors, FString &HLSLInterfaceID, FString &OutHLSL) override;
-	virtual void GetBufferDefinitionHLSL(FString DataInterfaceID, TArray<FDIGPUBufferParamDescriptor> &BufferDescriptors, FString &OutHLSL) override;
-	virtual TArray<FNiagaraDataInterfaceBufferData> &GetBufferDataArray() override;
-	virtual void SetupBuffers(FDIBufferDescriptorStore &BufferDescriptors) override;
+	virtual bool GetFunctionHLSL(const FName&  DefinitionFunctionName, FString InstanceFunctionName, FNiagaraDataInterfaceGPUParamInfo& ParamInfo, FString& OutHLSL) override;
 
+	virtual int32 GetCurveNumElems()const { return CurveLUTNumElems; }
 protected:
 	virtual bool CopyToInternal(UNiagaraDataInterface* Destination) const override;
 
 private:
 	template<typename UseLUT>
 	FORCEINLINE_DEBUGGABLE FLinearColor SampleCurveInternal(float X);
+
+	static const FName SampleCurveName;
 };
diff --git a/Engine/Plugins/FX/Niagara/Source/Niagara/Classes/NiagaraDataInterfaceCurlNoise.h b/Engine/Plugins/FX/Niagara/Source/Niagara/Classes/NiagaraDataInterfaceCurlNoise.h
--- a/Engine/Plugins/FX/Niagara/Source/Niagara/Classes/NiagaraDataInterfaceCurlNoise.h
+++ b/Engine/Plugins/FX/Niagara/Source/Niagara/Classes/NiagaraDataInterfaceCurlNoise.h
@@ -28,7 +28,7 @@ public:
 
 	//UNiagaraDataInterface Interface
 	virtual void GetFunctions(TArray<FNiagaraFunctionSignature>& OutFunctions)override;
-	virtual FVMExternalFunction GetVMExternalFunction(const FVMExternalFunctionBindingInfo& BindingInfo, void* InstanceData)override;
+	virtual void GetVMExternalFunction(const FVMExternalFunctionBindingInfo& BindingInfo, void* InstanceData, FVMExternalFunction &OutFunc) override;
 	virtual bool CanExecuteOnTarget(ENiagaraSimTarget Target)const override { return true; }
 	//UNiagaraDataInterface Interface End
 
@@ -38,11 +38,12 @@ public:
 	virtual bool Equals(const UNiagaraDataInterface* Other) const override;
 
 	// GPU sim functionality
-	virtual bool GetFunctionHLSL(const FName&  DefinitionFunctionName, FString InstanceFunctionName, TArray<FDIGPUBufferParamDescriptor> &Descriptors, FString &HLSLInterfaceID, FString &OutHLSL) override;
-	virtual void GetBufferDefinitionHLSL(FString DataInterfaceID, TArray<FDIGPUBufferParamDescriptor> &BufferDescriptors, FString &OutHLSL) override;
-	virtual TArray<FNiagaraDataInterfaceBufferData> &GetBufferDataArray() override;
-	virtual void SetupBuffers(FDIBufferDescriptorStore &BufferDescriptors) override;
+	virtual bool GetFunctionHLSL(const FName&  DefinitionFunctionName, FString InstanceFunctionName, FNiagaraDataInterfaceGPUParamInfo& ParamInfo, FString& OutHLSL) override;
+	virtual void GetParameterDefinitionHLSL(FNiagaraDataInterfaceGPUParamInfo& ParamInfo, FString& OutHLSL) override;
+	virtual FNiagaraDataInterfaceParametersCS* ConstructComputeParameters()const override;
 
+	FRWBuffer& GetGPUBuffer();
+	static const FString CurlNoiseBufferName;
 protected:
 	virtual bool CopyToInternal(UNiagaraDataInterface* Destination) const override;
 	void InitNoiseLUT();
@@ -51,4 +52,8 @@ protected:
 
 	template<typename T>
 	void ReplicateBorder(T* DestBuffer);
+
+	FRWBuffer GPUBuffer;
+
+	static const FName SampleNoiseFieldName;
 };
diff --git a/Engine/Plugins/FX/Niagara/Source/Niagara/Classes/NiagaraDataInterfaceCurve.h b/Engine/Plugins/FX/Niagara/Source/Niagara/Classes/NiagaraDataInterfaceCurve.h
--- a/Engine/Plugins/FX/Niagara/Source/Niagara/Classes/NiagaraDataInterfaceCurve.h
+++ b/Engine/Plugins/FX/Niagara/Source/Niagara/Classes/NiagaraDataInterfaceCurve.h
@@ -17,13 +17,8 @@ class NIAGARA_API UNiagaraDataInterfaceCurve : public UNiagaraDataInterfaceCurve
 	GENERATED_UCLASS_BODY()
 
 public:
-#if WITH_EDITORONLY_DATA
-	UPROPERTY(EditAnywhere, Category = "Curve", meta = (AllowedClasses = CurveFloat))
-	FSoftObjectPath CurveToCopy;
-#endif
-
 	UPROPERTY(EditAnywhere, Category = "Curve")
-	FRichCurve Curve;
+	FRichCurve Curve; 
 
 	enum
 	{
@@ -36,14 +31,10 @@ public:
 	//UObject Interface
 	virtual void PostInitProperties() override;
 	virtual void PostLoad() override;
-
-#if WITH_EDITOR
-	virtual void PostEditChangeProperty(struct FPropertyChangedEvent& PropertyChangedEvent) override;
-#endif
 	//UObject Interface End
 
 	virtual void GetFunctions(TArray<FNiagaraFunctionSignature>& OutFunctions)override;
-	virtual FVMExternalFunction GetVMExternalFunction(const FVMExternalFunctionBindingInfo& BindingInfo, void* InstanceData)override;
+	virtual void GetVMExternalFunction(const FVMExternalFunctionBindingInfo& BindingInfo, void* InstanceData, FVMExternalFunction &OutFunc) override;
 
 	template<typename UseLUT, typename XParamType>
 	void SampleCurve(FVectorVMContext& Context);
@@ -53,15 +44,15 @@ public:
 	//~ UNiagaraDataInterfaceCurveBase interface
 	virtual void GetCurveData(TArray<FCurveData>& OutCurveData) override;
 
-	virtual bool GetFunctionHLSL(const FName&  DefinitionFunctionName, FString InstanceFunctionName, TArray<FDIGPUBufferParamDescriptor> &Descriptors, FString &HLSLInterfaceID, FString &OutHLSL) override;
-	virtual void GetBufferDefinitionHLSL(FString DataInterfaceID, TArray<FDIGPUBufferParamDescriptor> &BufferDescriptors, FString &OutHLSL) override;
-	virtual TArray<FNiagaraDataInterfaceBufferData> &GetBufferDataArray() override;
-	virtual void SetupBuffers(FDIBufferDescriptorStore &BufferDescriptors) override;
+	virtual bool GetFunctionHLSL(const FName&  DefinitionFunctionName, FString InstanceFunctionName, FNiagaraDataInterfaceGPUParamInfo& ParamInfo, FString& OutHLSL) override;
 
+	virtual int32 GetCurveNumElems()const { return CurveLUTNumElems; }
 protected:
 	virtual bool CopyToInternal(UNiagaraDataInterface* Destination) const override;
 
 private:
 	template<typename UseLUT>
-	float SampleCurveInternal(float X);
+	FORCEINLINE_DEBUGGABLE float SampleCurveInternal(float X);
+
+	static const FName SampleCurveName;
 };
diff --git a/Engine/Plugins/FX/Niagara/Source/Niagara/Classes/NiagaraDataInterfaceSimpleCounter.h b/Engine/Plugins/FX/Niagara/Source/Niagara/Classes/NiagaraDataInterfaceSimpleCounter.h
new file mode 100644
--- /dev/null
+++ b/Engine/Plugins/FX/Niagara/Source/Niagara/Classes/NiagaraDataInterfaceSimpleCounter.h
@@ -0,0 +1,68 @@
+// Copyright 1998-2016 Epic Games, Inc. All Rights Reserved.
+#pragma once
+
+#include "NiagaraCommon.h"
+#include "NiagaraCollision.h"
+#include "NiagaraComponent.h"
+#include "NiagaraSystemInstance.h"
+#include "NiagaraShared.h"
+#include "VectorVM.h"
+#include "NiagaraDataInterface.h"
+#include "HAL/PlatformAtomics.h"
+#include "NiagaraDataInterfaceSimpleCounter.generated.h"
+
+class INiagaraCompiler;
+class FNiagaraSystemInstance;
+
+
+struct CounterInstanceData
+{
+	int32 Counter;
+};
+
+
+/** Data Interface allowing a counter that starts at zero when created and increases every time it is sampled. Note that for now this is a signed integer, meaning that it can go negative when it loops past INT_MAX. */
+UCLASS(EditInlineNew, Category = "Counting", meta = (DisplayName = "Simple Counter"))
+class NIAGARA_API UNiagaraDataInterfaceSimpleCounter : public UNiagaraDataInterface
+{
+	GENERATED_UCLASS_BODY()
+public:
+
+	//UObject Interface
+	virtual void PostInitProperties() override;
+	virtual void PostLoad() override;
+#if WITH_EDITOR
+	virtual void PostEditChangeProperty(struct FPropertyChangedEvent& PropertyChangedEvent) override;
+#endif
+	//UObject Interface End
+
+	/** Initializes the per instance data for this interface. Returns false if there was some error and the simulation should be disabled. */
+	virtual bool InitPerInstanceData(void* PerInstanceData, FNiagaraSystemInstance* InSystemInstance) override;
+	/** Destroys the per instence data for this interface. */
+	virtual void DestroyPerInstanceData(void* PerInstanceData, FNiagaraSystemInstance* InSystemInstance) {}
+
+	/** Ticks the per instance data for this interface, if it has any. */
+	virtual bool PerInstanceTick(void* PerInstanceData, FNiagaraSystemInstance* SystemInstance, float DeltaSeconds);
+	virtual bool PerInstanceTickPostSimulate(void* PerInstanceData, FNiagaraSystemInstance* SystemInstance, float DeltaSeconds);
+
+	virtual int32 PerInstanceDataSize()const { return sizeof(CounterInstanceData); }
+
+
+	virtual void GetFunctions(TArray<FNiagaraFunctionSignature>& OutFunctions)override;
+	virtual void GetVMExternalFunction(const FVMExternalFunctionBindingInfo& BindingInfo, void* InstanceData, FVMExternalFunction &OutFunc) override;
+
+	// VM functions
+	void GetNextValue(FVectorVMContext& Context);
+
+	virtual bool Equals(const UNiagaraDataInterface* Other) const override;
+	virtual bool CanExecuteOnTarget(ENiagaraSimTarget Target) const override { return true; }
+
+	virtual bool GetFunctionHLSL(const FName&  DefinitionFunctionName, FString InstanceFunctionName, FNiagaraDataInterfaceGPUParamInfo& ParamInfo, FString& OutHLSL) override;
+	virtual void GetParameterDefinitionHLSL(FNiagaraDataInterfaceGPUParamInfo& ParamInfo, FString& OutHLSL) override;
+	
+
+protected:
+	virtual bool CopyToInternal(UNiagaraDataInterface* Destination) const override;
+
+private:
+};
diff --git a/Engine/Plugins/FX/Niagara/Source/Niagara/Classes/NiagaraDataInterfaceSkeletalMesh.h b/Engine/Plugins/FX/Niagara/Source/Niagara/Classes/NiagaraDataInterfaceSkeletalMesh.h
--- a/Engine/Plugins/FX/Niagara/Source/Niagara/Classes/NiagaraDataInterfaceSkeletalMesh.h
+++ b/Engine/Plugins/FX/Niagara/Source/Niagara/Classes/NiagaraDataInterfaceSkeletalMesh.h
@@ -6,6 +6,7 @@
 #include "Engine/SkeletalMesh.h"
 #include "Rendering/SkeletalMeshRenderData.h"
 #include "Components/SkeletalMeshComponent.h"
+#include "Containers/Array.h"
 #include "NiagaraDataInterfaceSkeletalMesh.generated.h"
 
 class UNiagaraDataInterfaceSkeletalMesh;
@@ -98,6 +99,7 @@ struct FSkeletalMeshSkinningData
 	void RegisterUser(FSkeletalMeshSkinningDataUsage Usage);
 	void UnregisterUser(FSkeletalMeshSkinningDataUsage Usage);
 	bool IsUsed()const;
+	void ForceDataRefresh();
 
 	bool Tick(float InDeltaSeconds);
 
@@ -160,6 +162,8 @@ private:
 		TArray<FVector> SkinnedCPUPositions[2];
 	};
 	TArray<FLODData> LODData;
+
+	bool bForceDataRefresh;
 };
 
 class FNDI_SkeletalMesh_GeneratedData
@@ -215,19 +219,21 @@ struct FSkeletalMeshSamplingRegionAreaWeightedSampler : FWeightedRandomSampler
 	virtual float GetWeights(TArray<float>& OutWeights)override;
 
 	FORCEINLINE bool IsValid() { return TotalWeight > 0.0f; }
+
+	int32 GetEntries() const { return Alias.Num(); }
 protected:
 	FNDISkeletalMesh_InstanceData* Owner;
 };
 
 struct FNDISkeletalMesh_InstanceData
 {
-	FRandomStream RandStream;//Might remove this when we rework randoms. I don't like stateful randoms as they can't work on GPU and likely not threaded VM either!
-
 	//Cached ptr to component we sample from. 
 	TWeakObjectPtr<USceneComponent> Component;
 
 	USkeletalMesh* Mesh;
 
+	TWeakObjectPtr<USkeletalMesh> MeshSafe;
+
 	/** Handle to our skinning data. */
 	FSkeletalMeshSkinningDataHandle SkinningData;
 	
@@ -250,9 +256,11 @@ struct FNDISkeletalMesh_InstanceData
 	/** Time separating Transform and PrevTransform. */
 	float DeltaSeconds;
 
+	uint32 ChangeId;
+
 	FORCEINLINE_DEBUGGABLE bool ResetRequired(UNiagaraDataInterfaceSkeletalMesh* Interface)const;
 
-	FORCEINLINE_DEBUGGABLE bool Init(UNiagaraDataInterfaceSkeletalMesh* Interface, FNiagaraSystemInstance* SystemInstance);
+	bool Init(UNiagaraDataInterfaceSkeletalMesh* Interface, FNiagaraSystemInstance* SystemInstance);
 	FORCEINLINE_DEBUGGABLE bool Cleanup(UNiagaraDataInterfaceSkeletalMesh* Interface, FNiagaraSystemInstance* SystemInstance);
 	FORCEINLINE_DEBUGGABLE bool Tick(UNiagaraDataInterfaceSkeletalMesh* Interface, FNiagaraSystemInstance* SystemInstance, float InDeltaSeconds);
 	FORCEINLINE int32 GetLODIndex()const { return SkinningData.Usage.GetLODIndex(); }
@@ -299,11 +307,14 @@ public:
 	UPROPERTY(EditAnywhere, Category="Mesh")
 	int32 WholeMeshLOD;
 
+	/** Cached change id off of the data interface.*/
+	uint32 ChangeId;
 public:
 
 	//~ UObject interface
 #if WITH_EDITOR
 	virtual void PostInitProperties()override;
+	virtual void PostEditChangeProperty(struct FPropertyChangedEvent& PropertyChangedEvent) override;
 #endif
 	//~ UObject interface END
 
@@ -316,32 +327,56 @@ public:
 	virtual int32 PerInstanceDataSize()const override { return sizeof(FNDISkeletalMesh_InstanceData); }
 
 	virtual void GetFunctions(TArray<FNiagaraFunctionSignature>& OutFunctions)override;
-	virtual FVMExternalFunction GetVMExternalFunction(const FVMExternalFunctionBindingInfo& BindingInfo, void* InstanceData)override;
+	virtual void GetVMExternalFunction(const FVMExternalFunctionBindingInfo& BindingInfo, void* InstanceData, FVMExternalFunction &OutFunc)override;
 	virtual bool Equals(const UNiagaraDataInterface* Other) const override;
 	virtual bool CanExecuteOnTarget(ENiagaraSimTarget Target)const override { return Target == ENiagaraSimTarget::CPUSim; }
+#if WITH_EDITOR	
+	virtual TArray<FNiagaraDataInterfaceError> GetErrors() override;
+#endif
 	//~ UNiagaraDataInterface interface END
 public:
 
+	template<typename FilterMode, typename AreaWeightingMode>
+	void GetFilteredTriangleCount(FVectorVMContext& Context);
+
+	template<typename FilterMode, typename AreaWeightingMode, typename TriType>
+	void GetFilteredTriangleAt(FVectorVMContext& Context);
+
 	template<typename FilterMode, typename AreaWeightingMode>
 	void RandomTriCoord(FVectorVMContext& Context);
 
+
+	template<typename FilterMode, typename AreaWeightingMode, typename TriType, typename BaryXType, typename BaryYType, typename BaryZType>
+	void IsValidTriCoord(FVectorVMContext& Context);
+
 	template<typename SkinningHandlerType, typename TransformHandlerType, typename TriType, typename BaryXType, typename BaryYType, typename BaryZType>
 	void GetTriCoordPosition(FVectorVMContext& Context);
 
 	template<typename SkinningHandlerType, typename TransformHandlerType, typename TriType, typename BaryXType, typename BaryYType, typename BaryZType>
 	void GetTriCoordPositionVelocityAndNormal(FVectorVMContext& Context);
+
+	template<typename SkinningHandlerType, typename TransformHandlerType, typename VertexAccessorType, typename TriType, typename BaryXType, typename BaryYType, typename BaryZType, typename UVSetType>
+	void GetTriCoordPositionVelocityAndNormalBinormalTangent(FVectorVMContext& Context);
 	
 	template<typename TriType, typename BaryXType, typename BaryYType, typename BaryZType>
 	void GetTriCoordColor(FVectorVMContext& Context);
 
 	template<typename VertexAccessorType, typename TriType, typename BaryXType, typename BaryYType, typename BaryZType, typename UVSetType>
 	void GetTriCoordUV(FVectorVMContext& Context);
- 
+
+	static USkeletalMesh* GetSkeletalMeshHelper(UNiagaraDataInterfaceSkeletalMesh* Interface, class UNiagaraComponent* OwningComponent, TWeakObjectPtr<USceneComponent>& SceneComponent, USkeletalMeshComponent*& FoundSkelComp);
 protected:
 	virtual bool CopyToInternal(UNiagaraDataInterface* Destination) const override;
 
 private:
 
 	template<typename FilterMode, typename AreaWeightingMode>
-	FORCEINLINE_DEBUGGABLE int32 RandomTriIndex(FSkeletalMeshAccessorHelper& Accessor, FNDISkeletalMesh_InstanceData* InstData);
+	FORCEINLINE_DEBUGGABLE int32 RandomTriIndex(FRandomStream& RandStream, FSkeletalMeshAccessorHelper& Accessor, FNDISkeletalMesh_InstanceData* InstData);
+
+	template<typename FilterMode, typename AreaWeightingMode>
+	FORCEINLINE_DEBUGGABLE int32 GetSpecificTriangleCount(FSkeletalMeshAccessorHelper& Accessor, FNDISkeletalMesh_InstanceData* InstData);
+
+	template<typename FilterMode, typename AreaWeightingMode>
+	FORCEINLINE_DEBUGGABLE int32 GetSpecificTriangleAt(FSkeletalMeshAccessorHelper& Accessor, FNDISkeletalMesh_InstanceData* InstData, int32 FilteredIdx);
+
 };
diff --git a/Engine/Plugins/FX/Niagara/Source/Niagara/Classes/NiagaraDataInterfaceSpline.h b/Engine/Plugins/FX/Niagara/Source/Niagara/Classes/NiagaraDataInterfaceSpline.h
--- a/Engine/Plugins/FX/Niagara/Source/Niagara/Classes/NiagaraDataInterfaceSpline.h
+++ b/Engine/Plugins/FX/Niagara/Source/Niagara/Classes/NiagaraDataInterfaceSpline.h
@@ -41,7 +41,7 @@ public:
 	virtual bool PerInstanceTick(void* PerInstanceData, FNiagaraSystemInstance* SystemInstance, float DeltaSeconds) override;
 	virtual int32 PerInstanceDataSize()const override;
 	virtual void GetFunctions(TArray<FNiagaraFunctionSignature>& OutFunctions)override;
-	virtual FVMExternalFunction GetVMExternalFunction(const FVMExternalFunctionBindingInfo& BindingInfo, void* InstanceData)override;
+	virtual void GetVMExternalFunction(const FVMExternalFunctionBindingInfo& BindingInfo, void* InstanceData, FVMExternalFunction &OutFunc) override;
 	virtual bool Equals(const UNiagaraDataInterface* Other) const override;
 	virtual bool CanExecuteOnTarget(ENiagaraSimTarget Target)const override { return Target == ENiagaraSimTarget::CPUSim; }
 	//UNiagaraDataInterface Interface End
diff --git a/Engine/Plugins/FX/Niagara/Source/Niagara/Classes/NiagaraDataInterfaceStaticMesh.h b/Engine/Plugins/FX/Niagara/Source/Niagara/Classes/NiagaraDataInterfaceStaticMesh.h
--- a/Engine/Plugins/FX/Niagara/Source/Niagara/Classes/NiagaraDataInterfaceStaticMesh.h
+++ b/Engine/Plugins/FX/Niagara/Source/Niagara/Classes/NiagaraDataInterfaceStaticMesh.h
@@ -41,8 +41,6 @@ class UNiagaraDataInterfaceStaticMesh;
 
 struct FNDIStaticMesh_InstanceData
 {
-	FRandomStream RandStream;//Might remove this when we rework randoms. I don't like stateful randoms as they can't work on GPU and likely not threaded VM either!
-
 	 //Cached ptr to component we sample from. 
 	TWeakObjectPtr<USceneComponent> Component;
 
@@ -73,13 +71,16 @@ struct FNDIStaticMesh_InstanceData
 	/** Allows sampling of the mesh's tris based on a dynamic color range. */
 	TSharedPtr<struct FDynamicVertexColorFilterData> DynamicVertexColorSampler;
 
+	/** Cached change id off of the data interface.*/
+	uint32 ChangeId;
+
 	FORCEINLINE UStaticMesh* GetActualMesh()const { return Mesh; }
 	FORCEINLINE bool UsesAreaWeighting()const { return bIsAreaWeightedSampling; }
 	FORCEINLINE bool MeshHasPositions()const { return Mesh && Mesh->RenderData->LODResources[0].VertexBuffers.PositionVertexBuffer.GetNumVertices() > 0; }
 	FORCEINLINE bool MeshHasVerts()const { return Mesh && Mesh->RenderData->LODResources[0].VertexBuffers.StaticMeshVertexBuffer.GetNumVertices() > 0; }
 	FORCEINLINE bool MeshHasColors()const { return Mesh && Mesh->RenderData->LODResources[0].VertexBuffers.ColorVertexBuffer.GetNumVertices() > 0; }
 
-	FORCEINLINE_DEBUGGABLE bool ResetRequired()const;
+	FORCEINLINE_DEBUGGABLE bool ResetRequired(UNiagaraDataInterfaceStaticMesh* Interface)const;
 
 	FORCEINLINE const TArray<int32>& GetValidSections()const { return ValidSections; }
 	FORCEINLINE const FStaticMeshAreaWeightedSectionSampler& GetAreaWeigtedSampler()const { return Sampler; }
@@ -110,12 +111,16 @@ public:
 	UPROPERTY(EditAnywhere, Category = "Mesh")
 	FNDIStaticMeshSectionFilter SectionFilter;
 
+    /** Changed within the editor on PostEditChangeProperty. Should be changed whenever a refresh is desired.*/
+	uint32 ChangeId;
+
 public:
 
 	//~ UObject interface
 
 #if WITH_EDITOR
 	virtual void PostInitProperties()override;
+	virtual void PostEditChangeProperty(struct FPropertyChangedEvent& PropertyChangedEvent) override;
 #endif
 
 public:
@@ -128,10 +133,14 @@ public:
 	virtual int32 PerInstanceDataSize()const override { return sizeof(FNDIStaticMesh_InstanceData); }
 
 	virtual void GetFunctions(TArray<FNiagaraFunctionSignature>& OutFunctions)override;
-	virtual FVMExternalFunction GetVMExternalFunction(const FVMExternalFunctionBindingInfo& BindingInfo, void* InstanceData)override;
+	virtual void GetVMExternalFunction(const FVMExternalFunctionBindingInfo& BindingInfo, void* InstanceData, FVMExternalFunction &OutFunc) override;
 	virtual bool Equals(const UNiagaraDataInterface* Other) const override;
 	virtual bool CanExecuteOnTarget(ENiagaraSimTarget Target)const override { return Target == ENiagaraSimTarget::CPUSim; }
+#if WITH_EDITOR
+	virtual TArray<FNiagaraDataInterfaceError> GetErrors() override;
+#endif
 public:
+	void GetNumTriangles(FVectorVMContext& Context);
 
 	template<typename TAreaWeighted>
 	void RandomSection(FVectorVMContext& Context);
@@ -167,6 +176,9 @@ public:
 	void GetLocalToWorldInverseTransposed(FVectorVMContext& Context);
 	void GetWorldVelocity(FVectorVMContext& Context);
 
+	template<typename TransformHandlerType, typename VertexIndexType>
+	void GetVertexPosition(FVectorVMContext& Context);
+
 	FORCEINLINE_DEBUGGABLE bool UsesSectionFilter()const { return SectionFilter.CanEverReject(); }
 
 	//TODO: Vertex color filtering requires a bit more work.
@@ -174,17 +186,17 @@ public:
 
 protected:
 	virtual bool CopyToInternal(UNiagaraDataInterface* Destination) const override;
-
+	
 private:
 	
 	template<typename TAreaWeighted, bool bFiltered>
-	FORCEINLINE_DEBUGGABLE int32 RandomSection(FStaticMeshLODResources& Res, FNDIStaticMesh_InstanceData* InstData);
+	FORCEINLINE_DEBUGGABLE int32 RandomSection(FRandomStream& RandStream, FStaticMeshLODResources& Res, FNDIStaticMesh_InstanceData* InstData);
 
 	template<typename TAreaWeighted, bool bFiltered>
-	FORCEINLINE_DEBUGGABLE int32 RandomTriIndex(FStaticMeshLODResources& Res, FNDIStaticMesh_InstanceData* InstData);
+	FORCEINLINE_DEBUGGABLE int32 RandomTriIndex(FRandomStream& RandStream, FStaticMeshLODResources& Res, FNDIStaticMesh_InstanceData* InstData);
 
 	template<typename TAreaWeighted>
-	FORCEINLINE_DEBUGGABLE int32 RandomTriIndexOnSection(FStaticMeshLODResources& Res, int32 SectionIdx, FNDIStaticMesh_InstanceData* InstData);
+	FORCEINLINE_DEBUGGABLE int32 RandomTriIndexOnSection(FRandomStream& RandStream, FStaticMeshLODResources& Res, int32 SectionIdx, FNDIStaticMesh_InstanceData* InstData);
 
 	void WriteTransform(const FMatrix& ToWrite, FVectorVMContext& Context);
 };
diff --git a/Engine/Plugins/FX/Niagara/Source/Niagara/Classes/NiagaraDataInterfaceVector2DCurve.h b/Engine/Plugins/FX/Niagara/Source/Niagara/Classes/NiagaraDataInterfaceVector2DCurve.h
--- a/Engine/Plugins/FX/Niagara/Source/Niagara/Classes/NiagaraDataInterfaceVector2DCurve.h
+++ b/Engine/Plugins/FX/Niagara/Source/Niagara/Classes/NiagaraDataInterfaceVector2DCurve.h
@@ -18,11 +18,6 @@ class NIAGARA_API UNiagaraDataInterfaceVector2DCurve : public UNiagaraDataInterf
 	GENERATED_UCLASS_BODY()
 
 public:
-#if WITH_EDITORONLY_DATA
-	UPROPERTY(EditAnywhere, Category = "Curve", meta = (AllowedClasses = CurveVector))
-	FSoftObjectPath CurveToCopy;
-#endif
-
 	UPROPERTY(EditAnywhere, Category = "Curve")
 	FRichCurve XCurve;
 
@@ -41,9 +36,6 @@ public:
 	virtual void PostInitProperties() override;
 	virtual void PostLoad() override;
 
-#if WITH_EDITOR
-	virtual void PostEditChangeProperty(struct FPropertyChangedEvent& PropertyChangedEvent) override;
-#endif
 	//UObject Interface End
 
 	virtual bool Equals(const UNiagaraDataInterface* Other) const override;
@@ -51,7 +43,7 @@ public:
 
 	virtual void GetFunctions(TArray<FNiagaraFunctionSignature>& OutFunctions)override;
 
-	virtual FVMExternalFunction GetVMExternalFunction(const FVMExternalFunctionBindingInfo& BindingInfo, void* InstanceData)override;
+	virtual void GetVMExternalFunction(const FVMExternalFunctionBindingInfo& BindingInfo, void* InstanceData, FVMExternalFunction &OutFunc) override;
 
 	template<typename UseLUT, typename XParamType>
 	void SampleCurve(FVectorVMContext& Context);
@@ -59,15 +51,15 @@ public:
 	//~ UNiagaraDataInterfaceCurveBase interface
 	virtual void GetCurveData(TArray<FCurveData>& OutCurveData) override;
 
-	virtual bool GetFunctionHLSL(const FName&  DefinitionFunctionName, FString InstanceFunctionName, TArray<FDIGPUBufferParamDescriptor> &Descriptors, FString &HLSLInterfaceID, FString &OutHLSL) override;
-	virtual void GetBufferDefinitionHLSL(FString DataInterfaceID, TArray<FDIGPUBufferParamDescriptor> &BufferDescriptors, FString &OutHLSL) override;
-	virtual TArray<FNiagaraDataInterfaceBufferData> &GetBufferDataArray() override;
-	virtual void SetupBuffers(FDIBufferDescriptorStore &BufferDescriptors) override;
+	virtual bool GetFunctionHLSL(const FName&  DefinitionFunctionName, FString InstanceFunctionName, FNiagaraDataInterfaceGPUParamInfo& ParamInfo, FString& OutHLSL) override;
 
+	virtual int32 GetCurveNumElems()const { return CurveLUTNumElems; }
 protected:
 	virtual bool CopyToInternal(UNiagaraDataInterface* Destination) const override;
 
 private:
 	template<typename UseLUT>
-	FVector2D SampleCurveInternal(float X);
+	FORCEINLINE_DEBUGGABLE FVector2D SampleCurveInternal(float X);
+
+	static const FName SampleCurveName;
 };
diff --git a/Engine/Plugins/FX/Niagara/Source/Niagara/Classes/NiagaraDataInterfaceVector4Curve.h b/Engine/Plugins/FX/Niagara/Source/Niagara/Classes/NiagaraDataInterfaceVector4Curve.h
--- a/Engine/Plugins/FX/Niagara/Source/Niagara/Classes/NiagaraDataInterfaceVector4Curve.h
+++ b/Engine/Plugins/FX/Niagara/Source/Niagara/Classes/NiagaraDataInterfaceVector4Curve.h
@@ -22,11 +22,6 @@ class NIAGARA_API UNiagaraDataInterfaceVector4Curve : public UNiagaraDataInterfa
 	GENERATED_UCLASS_BODY()
 
 public:
-#if WITH_EDITORONLY_DATA
-	UPROPERTY(EditAnywhere, Category = "Curve", meta = (AllowedClasses = CurveLinearColor))
-	FSoftObjectPath CurveToCopy;
-#endif
-
 	UPROPERTY(EditAnywhere, Category = "Curve")
 	FRichCurve XCurve;
 
@@ -50,13 +45,10 @@ public:
 	//UObject Interface
 	virtual void PostInitProperties() override;
 	virtual void PostLoad() override;
-#if WITH_EDITOR
-	virtual void PostEditChangeProperty(struct FPropertyChangedEvent& PropertyChangedEvent) override;
-#endif
 	//UObject Interface End
 
 	virtual void GetFunctions(TArray<FNiagaraFunctionSignature>& OutFunctions)override;
-	virtual FVMExternalFunction GetVMExternalFunction(const FVMExternalFunctionBindingInfo& BindingInfo, void* InstanceData)override;
+	virtual void GetVMExternalFunction(const FVMExternalFunctionBindingInfo& BindingInfo, void* InstanceData, FVMExternalFunction &OutFunc) override;
 
 	template<typename UseLUT, typename XParamType>
 	void SampleCurve(FVectorVMContext& Context);
@@ -65,15 +57,15 @@ public:
 
 	//~ UNiagaraDataInterfaceCurveBase interface
 	virtual void GetCurveData(TArray<FCurveData>& OutCurveData) override;
-	virtual bool GetFunctionHLSL (const FName& DefinitionFunctionName, FString InstanceFunctionName, TArray<FDIGPUBufferParamDescriptor> &Descriptors, FString &HLSLInterfaceID, FString &OutHLSL) override;
-	virtual void GetBufferDefinitionHLSL(FString DataInterfaceID, TArray<FDIGPUBufferParamDescriptor> &BufferDescriptors, FString &OutHLSL) override;
-	virtual TArray<FNiagaraDataInterfaceBufferData> &GetBufferDataArray() override;
-	virtual void SetupBuffers(FDIBufferDescriptorStore &BufferDescriptors) override;
+	virtual bool GetFunctionHLSL (const FName& DefinitionFunctionName, FString InstanceFunctionName, FNiagaraDataInterfaceGPUParamInfo& ParamInfo, FString& OutHLSL) override;
 
+	virtual int32 GetCurveNumElems()const { return CurveLUTNumElems; }
 protected:
 	virtual bool CopyToInternal(UNiagaraDataInterface* Destination) const override;
 
 private: 
 	template<typename UseLUT>
-	FVector4 SampleCurveInternal(float X);
+	FORCEINLINE_DEBUGGABLE FVector4 SampleCurveInternal(float X);
+
+	static const FName SampleCurveName;
 };
diff --git a/Engine/Plugins/FX/Niagara/Source/Niagara/Classes/NiagaraDataInterfaceVectorCurve.h b/Engine/Plugins/FX/Niagara/Source/Niagara/Classes/NiagaraDataInterfaceVectorCurve.h
--- a/Engine/Plugins/FX/Niagara/Source/Niagara/Classes/NiagaraDataInterfaceVectorCurve.h
+++ b/Engine/Plugins/FX/Niagara/Source/Niagara/Classes/NiagaraDataInterfaceVectorCurve.h
@@ -22,11 +22,6 @@ class NIAGARA_API UNiagaraDataInterfaceVectorCurve : public UNiagaraDataInterfac
 {
 	GENERATED_UCLASS_BODY()
 public:
-#if WITH_EDITORONLY_DATA
-	UPROPERTY(EditAnywhere, Category = "Curve", meta = (AllowedClasses = CurveVector))
-	FSoftObjectPath CurveToCopy;
-#endif
-
 	UPROPERTY(EditAnywhere, Category = "Curve")
 	FRichCurve XCurve;
 
@@ -47,13 +42,10 @@ public:
 	//UObject Interface
 	virtual void PostInitProperties() override;
 	virtual void PostLoad() override;
-#if WITH_EDITOR
-	virtual void PostEditChangeProperty(struct FPropertyChangedEvent& PropertyChangedEvent) override;
-#endif
 	//UObject Interface End
 
 	virtual void GetFunctions(TArray<FNiagaraFunctionSignature>& OutFunctions)override;
-	virtual FVMExternalFunction GetVMExternalFunction(const FVMExternalFunctionBindingInfo& BindingInfo, void* InstanceData)override;
+	virtual void GetVMExternalFunction(const FVMExternalFunctionBindingInfo& BindingInfo, void* InstanceData, FVMExternalFunction &OutFunc) override;
 
 	template<typename UseLUT, typename XParamType>
 	void SampleCurve(FVectorVMContext& Context);
@@ -62,16 +54,17 @@ public:
 
 	//~ UNiagaraDataInterfaceCurveBase interface
 	virtual void GetCurveData(TArray<FCurveData>& OutCurveData) override;
+	
+	virtual int32 GetCurveNumElems()const { return CurveLUTNumElems; }
 
-	virtual bool GetFunctionHLSL(const FName&  DefinitionFunctionName, FString InstanceFunctionName, TArray<FDIGPUBufferParamDescriptor> &Descriptors, FString &HLSLInterfaceID, FString &OutHLSL) override;
-	virtual void GetBufferDefinitionHLSL(FString DataInterfaceID, TArray<FDIGPUBufferParamDescriptor> &BufferDescriptors, FString &OutHLSL) override;
-	virtual TArray<FNiagaraDataInterfaceBufferData> &GetBufferDataArray() override;
-	virtual void SetupBuffers(FDIBufferDescriptorStore &BufferDescriptors) override;
+	virtual bool GetFunctionHLSL(const FName&  DefinitionFunctionName, FString InstanceFunctionName, FNiagaraDataInterfaceGPUParamInfo& ParamInfo, FString& OutHLSL) override;
 
 protected:
 	virtual bool CopyToInternal(UNiagaraDataInterface* Destination) const override;
 
 private:
 	template<typename UseLUT>
-	FVector SampleCurveInternal(float X);
+	FORCEINLINE_DEBUGGABLE FVector SampleCurveInternal(float X);
+
+	static const FName SampleCurveName;
 };
diff --git a/Engine/Plugins/FX/Niagara/Source/Niagara/Classes/NiagaraDataInterfaceVectorField.h b/Engine/Plugins/FX/Niagara/Source/Niagara/Classes/NiagaraDataInterfaceVectorField.h
--- a/Engine/Plugins/FX/Niagara/Source/Niagara/Classes/NiagaraDataInterfaceVectorField.h
+++ b/Engine/Plugins/FX/Niagara/Source/Niagara/Classes/NiagaraDataInterfaceVectorField.h
@@ -62,15 +62,14 @@ public:
 	virtual int32 PerInstanceDataSize()const override;
 
 	virtual void GetFunctions(TArray<FNiagaraFunctionSignature>& OutFunctions)override;
-	virtual FVMExternalFunction GetVMExternalFunction(const FVMExternalFunctionBindingInfo& BindingInfo, void* InstanceData)override;
+	virtual void GetVMExternalFunction(const FVMExternalFunctionBindingInfo& BindingInfo, void* InstanceData, FVMExternalFunction &OutFunc)override;
 	virtual bool Equals(const UNiagaraDataInterface* Other) const override;
 	virtual bool CanExecuteOnTarget(ENiagaraSimTarget Target)const override;
 
 	// GPU sim functionality
-	virtual bool GetFunctionHLSL(const FName&  DefinitionFunctionName, FString InstanceFunctionName, TArray<FDIGPUBufferParamDescriptor> &Descriptors, FString &HLSLInterfaceID, FString &OutHLSL) override;
-	virtual void GetBufferDefinitionHLSL(FString DataInterfaceID, TArray<FDIGPUBufferParamDescriptor> &BufferDescriptors, FString &OutHLSL) override;
-	virtual TArray<FNiagaraDataInterfaceBufferData> &GetBufferDataArray() override;
-	virtual void SetupBuffers(FDIBufferDescriptorStore &BufferDescriptors) override;
+	virtual bool GetFunctionHLSL(const FName&  DefinitionFunctionName, FString InstanceFunctionName, FNiagaraDataInterfaceGPUParamInfo& ParamInfo, FString& OutHLSL) override;
+	virtual void GetParameterDefinitionHLSL(FNiagaraDataInterfaceGPUParamInfo& ParamInfo, FString& OutHLSL) override;
+	virtual FNiagaraDataInterfaceParametersCS* ConstructComputeParameters()const override;
 	//~ UNiagaraDataInterface interface END
 
 	template<typename XType, typename YType, typename ZType>
@@ -82,7 +81,17 @@ public:
 
 	bool bGPUBufferDirty;
 
+	static const FString BufferBaseName;
+	static const FString DimentionsBaseName;
+	static const FString BoundsMinBaseName;
+	static const FString BoundsMaxBaseName;
+
+	FRWBuffer& GetGPUBuffer();
+	FORCEINLINE FVector GetDimentions()const { return FVector(SizeX, SizeY, SizeZ); }
+	FORCEINLINE FVector GetBoundsMin()const { return LocalBounds.Min; }
+	FORCEINLINE FVector GetBoundsMax()const { return LocalBounds.Max; }
 protected:
+
 	const void* Lock();
 	void Unlock();
 	virtual bool CopyToInternal(UNiagaraDataInterface* Destination) const override;
@@ -91,4 +100,6 @@ protected:
 	uint32 SizeX, SizeY, SizeZ;
 	FVector TilingAxes;
 	FBox LocalBounds;
+
+	FRWBuffer GPUBuffer;
 };
diff --git a/Engine/Plugins/FX/Niagara/Source/Niagara/Classes/NiagaraDataSet.h b/Engine/Plugins/FX/Niagara/Source/Niagara/Classes/NiagaraDataSet.h
--- a/Engine/Plugins/FX/Niagara/Source/Niagara/Classes/NiagaraDataSet.h
+++ b/Engine/Plugins/FX/Niagara/Source/Niagara/Classes/NiagaraDataSet.h
@@ -27,23 +27,35 @@ struct FNiagaraVariableLayoutInfo
 class FNiagaraDataSet;
 
 /** Buffer containing one frame of Niagara simulation data. */
-class FNiagaraDataBuffer
+class NIAGARA_API FNiagaraDataBuffer
 {
 public:
 	FNiagaraDataBuffer() : Owner(nullptr)
 	{
 		Reset();
 	}
+	~FNiagaraDataBuffer();
 	void Init(FNiagaraDataSet* InOwner);
-	void Allocate(uint32 NumInstances, ENiagaraSimTarget Target = ENiagaraSimTarget::CPUSim, bool bMaintainExisting=false);
+	void Allocate(uint32 NumInstances, bool bMaintainExisting=false);
 	void AllocateGPU(uint32 InNumInstances, FRHICommandList &RHICmdList);
 	void SwapInstances(uint32 OldIndex, uint32 NewIndex);
 	void KillInstance(uint32 InstanceIdx);
+	void CopyTo(FNiagaraDataBuffer& DestBuffer, int32 StartIdx, int32 NumInstances);
 	void CopyTo(FNiagaraDataBuffer& DestBuffer);
 
 	const TArray<uint8>& GetFloatBuffer()const { return FloatData; }
 	const TArray<uint8>& GetInt32Buffer()const { return Int32Data; }
 
+	FORCEINLINE const uint8* GetComponentPtrFloat(uint32 ComponentIdx)const
+	{
+		return FloatData.GetData() + FloatStride * ComponentIdx;
+	}
+
+	FORCEINLINE const uint8* GetComponentPtrInt32(uint32 ComponentIdx)const
+	{
+		return Int32Data.GetData() + Int32Stride * ComponentIdx;
+	}
+
 	FORCEINLINE uint8* GetComponentPtrFloat(uint32 ComponentIdx)
 	{
 		return FloatData.GetData() + FloatStride * ComponentIdx;
@@ -72,16 +84,7 @@ public:
 		NumInstances = InNumInstances;
 	}
 
-	FORCEINLINE void Reset()
-	{
-		FloatData.Empty();
-		Int32Data.Empty();
-		FloatStride = 0;
-		Int32Stride = 0;
-		NumInstances = 0;
-		NumInstancesAllocated = 0;
-		NumChunksAllocatedForGPU = 0;
-	}
+	void Reset();
 
 	FORCEINLINE uint32 GetSizeBytes()const
 	{
@@ -107,13 +110,15 @@ public:
 	uint32 GetInt32Stride() const { return Int32Stride; }
 
 	FORCEINLINE const FNiagaraDataSet* GetOwner()const { return Owner; }
+
+	int32 TransferInstance(FNiagaraDataBuffer& SourceBuffer, int32 InstanceIndex);
 private:
 
 	FORCEINLINE int32 GetSafeComponentBufferSize(int32 RequiredSize) const
 	{
 		//Round up to VECTOR_WIDTH_BYTES.
 		//Both aligns the component buffers to the vector width but also ensures their ops cannot stomp over one another.		
-		return RequiredSize + VECTOR_WIDTH_BYTES - (RequiredSize % VECTOR_WIDTH_BYTES);
+		return RequiredSize + VECTOR_WIDTH_BYTES - (RequiredSize % VECTOR_WIDTH_BYTES) + VECTOR_WIDTH_BYTES;
 	}
 
 	/** Back ptr to our owning data set. Used to access layout info for the buffer. */
@@ -136,7 +141,6 @@ private:
 	/** Number of instances the buffer has been allocated for. */
 	uint32 NumInstancesAllocated;
 
-
 	FRWBuffer GPUBufferFloat;
 	FRWBuffer GPUBufferInt;
 };
@@ -146,35 +150,42 @@ private:
 /**
 General storage class for all per instance simulation data in Niagara.
 */
-class FNiagaraDataSet
+class NIAGARA_API FNiagaraDataSet
 {
 	friend FNiagaraDataBuffer;
-public:
-	FNiagaraDataSet()
-	{
-		Reset();
-	}
-	
-	void Init(FNiagaraDataSetID InID)
-	{
-		Reset();
-		ID = InID;
-	}
 
 	void Reset()
 	{
 		Variables.Empty();
 		VariableLayouts.Empty();
-		Data[0].Reset();
-		Data[1].Reset();
 		CurrBuffer = 0;
-		CurrRenderBuffer = 0;
 		bFinalized = false;
 		TotalFloatComponents = 0;
 		TotalInt32Components = 0;
 		MaxBufferIdx = 1;
+		bNeedsPersistentIDs = 0;
+
+		SimTarget = ENiagaraSimTarget::CPUSim;
+
+		ResetBuffersInternal();
+	}
+
+public:
+	FNiagaraDataSet()
+	{
+		Reset();
+	}
+	
+	~FNiagaraDataSet();
+
+	void Init(FNiagaraDataSetID InID, ENiagaraSimTarget InSimTarget)
+	{
+		Reset();
+		ID = InID;
+		SimTarget = InSimTarget;
 	}
 
+
 	void AddVariable(FNiagaraVariable& Variable)
 	{
 		check(!bFinalized);
@@ -190,6 +201,13 @@ public:
 		}
 	}
 
+	FORCEINLINE void SetNeedsPersistentIDs(bool bNeedsIDs)
+	{
+		bNeedsPersistentIDs = bNeedsIDs;
+	}
+
+	FORCEINLINE bool GetNeedsPersistentIDs()const { return bNeedsPersistentIDs; }
+
 	/** Finalize the addition of variables and other setup before this data set can be used. */
 	FORCEINLINE void Finalize()
 	{
@@ -202,19 +220,29 @@ public:
 	FORCEINLINE void KillInstance(uint32 InstanceIdx)
 	{
 		check(bFinalized);
+		CheckCorrectThread();
 		CurrData().KillInstance(InstanceIdx);
 	}
 
 	FORCEINLINE void SwapInstances(uint32 OldIndex, uint32 NewIndex)
 	{
 		check(bFinalized);
+		CheckCorrectThread();
 		PrevData().SwapInstances(OldIndex, NewIndex);
 	}
 
+	int32 TransferInstance(FNiagaraDataSet& SourceDataset, int32 InstanceIndex)
+	{
+		check(bFinalized);
+		CheckCorrectThread();
+		return CurrData().TransferInstance(SourceDataset.CurrData(), InstanceIndex);
+	}
+
 	/** Appends all variables in this dataset to a register table ready for execution by the VectorVM. */
 	bool AppendToRegisterTable(uint8** InputRegisters, int32& NumInputRegisters, uint8** OutputRegisters, int32& NumOutputRegisters, int32 StartInstance)
 	{
 		check(bFinalized);
+		CheckCorrectThread();
 		int32 TotalComponents = GetNumFloatComponents() + GetNumInt32Components();
 		if (NumInputRegisters + TotalComponents > VectorVM::MaxInputRegisters || NumOutputRegisters + TotalComponents > VectorVM::MaxOutputRegisters)
 		{
@@ -250,127 +278,197 @@ public:
 
 	void SetShaderParams(class FNiagaraShader *Shader, FRHICommandList &CommandList);
 	void UnsetShaderParams(class FNiagaraShader *Shader, FRHICommandList &CommandList);
-	FORCEINLINE void Allocate(uint32 NumInstances, ENiagaraSimTarget Target = ENiagaraSimTarget::CPUSim, bool bMaintainExisting=false)
+	void Allocate(int32 NumInstances, bool bMaintainExisting=false)
 	{
 		check(bFinalized);
-		CurrData().Allocate(NumInstances, Target, bMaintainExisting);
-	}
-	
-	void InitGPUFromCPU_RenderThread();
+		CheckCorrectThread();
+		CurrData().Allocate(NumInstances, bMaintainExisting);
 
-	void InitGPUFromCPU()
-	{
-		ENQUEUE_UNIQUE_RENDER_COMMAND_ONEPARAMETER(InitGPUFromCPUCommand,
-			FNiagaraDataSet*, DataSet, this,
-			{
-				DataSet->InitGPUFromCPU_RenderThread();
-			});
-	}
+		if (bNeedsPersistentIDs)
+		{
+			int32 NumUsedIDs = MaxUsedID + 1;
 
-	void InitGPUSimSRVs()
-	{
-		ENQUEUE_UNIQUE_RENDER_COMMAND_ONEPARAMETER(InitGPUFromCPUCommand,
-			FNiagaraDataSet*, DataSet, this,
-			{
-				DataSet->InitGPUSimSRVs_RenderThread();
-			});
-	}
+			int32 RequiredIDs = FMath::Max(NumInstances, NumUsedIDs);
+			int32 ExistingNumIDs = PrevIDTable().Num();
 
-	void InitGPUSimSRVs_RenderThread();
+			//Free ID Table must always be at least as large as the data buffer + it's current size in the case all particles die this frame.
+			FreeIDsTable.SetNumUninitialized(NumInstances + NumFreeIDs);
 
+			if (RequiredIDs > ExistingNumIDs)
+			{
+				//UE_LOG(LogNiagara, Warning, TEXT("Growing ID Table! OldSize:%d | NewSize:%d"), ExistingNumIDs, RequiredIDs);
+				IDToIndexTable[CurrBuffer].SetNumUninitialized(RequiredIDs);
 
-	FORCEINLINE void Tick(ENiagaraSimTarget SimTarget = ENiagaraSimTarget::CPUSim)
-	{
-		SwapBuffers(SimTarget);
-	}
-	FORCEINLINE void TickRenderThread(ENiagaraSimTarget SimTarget = ENiagaraSimTarget::CPUSim)
-	{
-		SwapBuffersRenderThread(SimTarget);
-	}
+				int32 NumNewIDs = RequiredIDs - ExistingNumIDs;
 
+				//Free table should always have enough room for these new IDs.
+				check(NumFreeIDs + NumNewIDs <= FreeIDsTable.Num());
+
+				//ID Table grows so add any new IDs to the free array. Add in reverse order to maintain a continuous increasing allocation when popping.
+				for (int32 NewFreeID = RequiredIDs - 1; NewFreeID >= ExistingNumIDs; --NewFreeID)
+				{
+					FreeIDsTable[NumFreeIDs++] = NewFreeID ;
+				}
+				//UE_LOG(LogNiagara, Warning, TEXT("DataSetAllocate: Adding New Free IDs: %d - "), RequiredIDs - 1, ExistingNumIDs);
+			}
+			else if (RequiredIDs < ExistingNumIDs >> 1)//Configurable?
+			{
+				//If the max id we use has reduced significantly then we can shrink the tables.
+				//Have to go through the FreeIDs and remove any that are greater than the new table size.
+				//UE_LOG(LogNiagara, Warning, TEXT("DataSetAllocate: Shrinking ID Table! OldSize:%d | NewSize:%d"), ExistingNumIDs, RequiredIDs);
+				for (int32 CheckedFreeID = 0; CheckedFreeID < NumFreeIDs;)
+				{
+					checkSlow(NumFreeIDs <= FreeIDsTable.Num());
+					if (FreeIDsTable[CheckedFreeID] >= RequiredIDs)
+					{
+						//UE_LOG(LogNiagara, Warning, TEXT("RemoveSwap FreeID: Removed:%d | Swapped:%d"), FreeIDsTable[CheckedFreeID], FreeIDsTable.Last());		
+						int32 FreeIDIndex = --NumFreeIDs;
+						FreeIDsTable[CheckedFreeID] = FreeIDsTable[FreeIDIndex];
+						FreeIDsTable[FreeIDIndex] = INDEX_NONE;
+					}
+					else
+					{
+						++CheckedFreeID;
+					}
+				}
+			}
+			else
+			{
+				//Drop in required size not great enough so just allocate same size.
+				RequiredIDs = ExistingNumIDs;
+			}
 
+			IDToIndexTable[CurrBuffer].SetNumUninitialized(RequiredIDs);
+			MaxUsedID = INDEX_NONE;//reset the max ID ready for it to be filled in during simulation.
 
-	// called before rendering to make sure we access the correct buffer
-	FORCEINLINE void ValidateBufferIndices()
+			//UE_LOG(LogNiagara, Warning, TEXT("DataSetAllocate: NumInstances:%d | ID Table Size:%d | NumFreeIDs:%d | FreeTableSize:%d"), NumInstances, IDToIndexTable[CurrBuffer].Num(), NumFreeIDs, FreeIDsTable.Num());
+		}
+	}
+	
+	FORCEINLINE void Tick()
 	{
-		CurrRenderBuffer = CurrBuffer;
+		SwapBuffers();
 	}
-
-	FORCEINLINE void CopyPrevToCur()
+	
+	FORCEINLINE void CopyCurToPrev()
 	{
-		PrevData().CopyTo(CurrData());
+		CurrData().CopyTo(PrevData());
 	}
 
 	FORCEINLINE FNiagaraDataSetID GetID()const { return ID; }
 	FORCEINLINE void SetID(FNiagaraDataSetID InID) { ID = InID; }
 
+	FORCEINLINE int32 GetPrevBufferIdx() const
+	{
+		return CurrBuffer > 0 ? CurrBuffer - 1 : MaxBufferIdx;
+	}
+
 	FORCEINLINE FNiagaraDataBuffer& CurrData() 
 	{
-		//check(IsInGameThread());
+		//CheckCorrectThread();//TODO: We should be able to enable these checks and have well defined GT/RT ownership but GPU sims fire these a lot currently.
 		return Data[CurrBuffer]; 
 	}
 	FORCEINLINE FNiagaraDataBuffer& PrevData() 
-	{ 
-		//check(IsInGameThread());
-		return Data[CurrBuffer > 0 ? CurrBuffer-1 : MaxBufferIdx];
+	{
+		//CheckCorrectThread();//TODO: We should be able to enable these checks and have well defined GT/RT ownership but GPU sims fire these a lot currently.
+		return Data[GetPrevBufferIdx()];
 	}
 	FORCEINLINE const FNiagaraDataBuffer& CurrData()const 
-	{ 
-		//check(IsInGameThread());
+	{
+		//CheckCorrectThread();//TODO: We should be able to enable these checks and have well defined GT/RT ownership but GPU sims fire these a lot currently.
 		return Data[CurrBuffer];
 	}
 	FORCEINLINE const FNiagaraDataBuffer& PrevData()const 
-	{ 
+	{
+		//CheckCorrectThread();//TODO: We should be able to enable these checks and have well defined GT/RT ownership but GPU sims fire these a lot currently.
+		return Data[GetPrevBufferIdx()];
+	}
+
+	FORCEINLINE uint32 GetNumInstances()const { return CurrData().GetNumInstances(); }
+	FORCEINLINE uint32 GetNumInstancesAllocated()const { return CurrData().GetNumInstancesAllocated(); }
+	FORCEINLINE void SetNumInstances(uint32 InNumInstances) { CurrData().SetNumInstances(InNumInstances); }
+
+	void ResetCurrentBuffers()
+	{
+		SetNumInstances(0);
+		if (bNeedsPersistentIDs)
+		{
+			CurrIDTable().Empty();
+			FreeIDsTable.Empty();
+			NumFreeIDs = 0;
+			MaxUsedID = INDEX_NONE;
+		}
+	}
+
+	FORCEINLINE TArray<int32>& CurrIDTable()
+	{
 		//check(IsInGameThread());
-		return Data[CurrBuffer > 0 ? CurrBuffer - 1 : MaxBufferIdx];
+		return IDToIndexTable[CurrBuffer];
 	}
 
-	FORCEINLINE FNiagaraDataBuffer& CurrDataRender() 
-	{ 
-		check(!IsInGameThread());
-		return Data[CurrRenderBuffer];
+	FORCEINLINE TArray<int32>& PrevIDTable()
+	{
+		//check(IsInGameThread());
+		return IDToIndexTable[GetPrevBufferIdx()];
 	}
-	FORCEINLINE FNiagaraDataBuffer& PrevDataRender() 
-	{ 
-		check(!IsInGameThread());
-		return Data[CurrRenderBuffer > 0 ? CurrRenderBuffer - 1 : MaxBufferIdx];
+
+	FORCEINLINE TArray<int32>& GetFreeIDTable()
+	{
+		return FreeIDsTable;
 	}
-	FORCEINLINE const FNiagaraDataBuffer& CurrDataRender() const 
-	{ 
-		check(!IsInGameThread());
-		return Data[CurrRenderBuffer];
+	
+	FORCEINLINE int32& GetNumFreeIDs()
+	{
+		return NumFreeIDs;
 	}
-	FORCEINLINE const FNiagaraDataBuffer& PrevDataRender() const 
+
+	FORCEINLINE int32& GetMaxUsedID()
 	{
-		check(!IsInGameThread());
-		return Data[CurrRenderBuffer > 0 ? CurrRenderBuffer - 1 : MaxBufferIdx];
+		return MaxUsedID;
 	}
 
-	FORCEINLINE uint32 GetNumInstances()const { return CurrData().GetNumInstances(); }
-	FORCEINLINE uint32 GetNumInstancesAllocated()const { return CurrData().GetNumInstancesAllocated(); }
-	FORCEINLINE void SetNumInstances(uint32 InNumInstances) { CurrData().SetNumInstances(InNumInstances); }
+	FORCEINLINE int32& GetIDAcquireTag()
+	{
+		return IDAcquireTag;
+	}
+
+	FORCEINLINE void SetIDAcquireTag(int32 InTag)
+	{
+		IDAcquireTag = InTag;
+	}
 
+	FORCEINLINE ENiagaraSimTarget GetSimTarget()const { return SimTarget; }
 
 	FORCEINLINE void ResetBuffersInternal()
 	{
 		Data[0].Reset();
 		Data[1].Reset();
 		Data[2].Reset();
+
+		FreeIDsTable.Reset();
+		NumFreeIDs = 0;
+
+		IDToIndexTable[0].Reset();
+		IDToIndexTable[1].Reset();
+		IDToIndexTable[2].Reset();
+		MaxUsedID = INDEX_NONE;
 	}
 
-	FORCEINLINE void ResetBuffers(bool bOnRT)
+	FORCEINLINE void ResetBuffers()
 	{
-		if (bOnRT)
+		if (SimTarget == ENiagaraSimTarget::CPUSim)
 		{
-			ENQUEUE_UNIQUE_RENDER_COMMAND_ONEPARAMETER(InitGPUFromCPUCommand,
+			ResetBuffersInternal();
+		}
+		else
+		{			
+			ENQUEUE_UNIQUE_RENDER_COMMAND_ONEPARAMETER(
+				ResetBuffersRT,
 				FNiagaraDataSet*, DataSet, this,
 				{
 					DataSet->ResetBuffersInternal();
-				});
-		}
-		else
-		{
-			ResetBuffersInternal();
+				}
+			);
 		}
 	}
 
@@ -412,49 +510,57 @@ public:
 	}
 
 	void Dump(bool bCurr, int32 StartIdx = 0, int32 NumInstances = INDEX_NONE);
-	void Dump(FNiagaraDataSet& Other, bool bCurr);
+	void Dump(FNiagaraDataSet& Other, bool bCurr, int32 StartIdx = 0, int32 NumInstances = INDEX_NONE);
 	FORCEINLINE const TArray<FNiagaraVariable> &GetVariables() { return Variables; }
 
-	// called before dispatch from NiagaraEmitterInstanceBatcher
-	FRWBuffer &SetupDataSetIndices() 
-	{ 
-		check(IsInRenderingThread());
-		if (DataSetIndices[CurrRenderBuffer].Buffer)
-		{
-			DataSetIndices[CurrRenderBuffer].Release();
-		}
-		DataSetIndices[CurrRenderBuffer].Initialize(sizeof(int32), 64 /*Context->NumDataSets*/, EPixelFormat::PF_R32_UINT, BUF_Transient | BUF_DrawIndirect | BUF_Static);	// always allocate for up to 64 data sets
-		return DataSetIndices[CurrRenderBuffer];
+
+	// Data set index buffer management
+	// these buffers hold the number of instances for the buffers; the first five uint32s are the DrawIndirect parameters for rendering of the main particle data set
+	//
+	FRWBuffer &GetCurDataSetIndices()		{	return GetDataSetIndices(CurrBuffer);			}
+	FRWBuffer &GetPrevDataSetIndices()		{	return GetDataSetIndices(GetPrevBufferIdx());	}
+
+	bool HasDatasetIndices(bool bCur = true) const
+	{
+		uint32 BufIdx = bCur == true ? CurrBuffer : GetPrevBufferIdx();
+		CheckCorrectThread();
+		return DataSetIndices[BufIdx].Buffer != nullptr;
 	}
 
-	FRWBuffer &GetDataSetIndices()
+	const FRWBuffer &GetCurDataSetIndices() const	
 	{
-		if (DataSetIndices[CurrRenderBuffer].Buffer == nullptr)
-		{
-			DataSetIndices[CurrRenderBuffer].Initialize(sizeof(int32), 64 /*Context->NumDataSets*/, EPixelFormat::PF_R32_UINT, BUF_Transient | BUF_DrawIndirect | BUF_Static);	// always allocate for up to 64 data sets
-		}
-		return DataSetIndices[CurrRenderBuffer];
+		CheckCorrectThread();
+		return DataSetIndices[CurrBuffer];
+	}
+
+	const FRWBuffer &GetPrevDataSetIndices() const
+	{
+		CheckCorrectThread();
+		return DataSetIndices[GetPrevBufferIdx()];
 	}
 
-	const FRWBuffer &GetDataSetIndices() const
+	void SetupCurDatasetIndices()
 	{
-		//ensure(DataSetIndices[CurrRenderBuffer].Buffer != nullptr);
-		return DataSetIndices[CurrRenderBuffer];
+		if (DataSetIndices[CurrBuffer].Buffer != nullptr)
+		{
+			DataSetIndices[CurrBuffer].Release();
+		}
+		DataSetIndices[CurrBuffer].Initialize(sizeof(int32), 64 /*Context->NumDataSets*/, EPixelFormat::PF_R32_UINT, BUF_DrawIndirect | BUF_Static);	// always allocate for up to 64 data sets
 	}
 
 	FORCEINLINE uint32 GetNumFloatComponents()const { return TotalFloatComponents; }
 	FORCEINLINE uint32 GetNumInt32Components()const { return TotalInt32Components; }
 
-	FORCEINLINE const FDynamicReadBuffer& GetRenderDataFloat()const { return RenderDataFloat; }
-	FORCEINLINE const FDynamicReadBuffer& GetRenderDataInt32()const { return RenderDataInt; }
-
-	FORCEINLINE const FShaderResourceViewRHIRef& GetRenderDataFloatSRV()const { return CurrentFloatDataSRV; }
-	FORCEINLINE const FShaderResourceViewRHIRef& GetRenderDataInt32SRV()const { return CurrentIntDataSRV; }
-
 private:
+	FRWBuffer &GetDataSetIndices(uint32 BufIdx)
+	{
+		CheckCorrectThread();
+		return DataSetIndices[BufIdx];
+	}
 
-	FORCEINLINE void SwapBuffers(ENiagaraSimTarget SimTarget)
+	FORCEINLINE void SwapBuffers()
 	{
+		CheckCorrectThread();
 		if (SimTarget == ENiagaraSimTarget::CPUSim)
 		{
 			MaxBufferIdx = 2;
@@ -465,33 +571,19 @@ private:
 			MaxBufferIdx = 1;
 			CurrBuffer = CurrBuffer == 0 ? 1 : 0;
 		}
-
-		ENQUEUE_UNIQUE_RENDER_COMMAND_TWOPARAMETER(SwapDataSetBuffersGPU,
-			uint32, CurrBuffer, CurrBuffer,
-			uint32 &, CurrRenderBuffer, CurrRenderBuffer,
-			{
-				CurrRenderBuffer = CurrBuffer;
-			});
 	}
 
-	FORCEINLINE void SwapBuffersRenderThread(ENiagaraSimTarget SimTarget)
+	FORCEINLINE void CheckCorrectThread()const
 	{
-		check(IsInRenderingThread());
-		if (SimTarget == ENiagaraSimTarget::CPUSim)
-		{
-			MaxBufferIdx = 2;
-			CurrBuffer = CurrBuffer < 2 ? CurrBuffer + 1 : 0;
-		}
-		else
-		{
-			MaxBufferIdx = 1;
-			CurrBuffer = CurrBuffer == 0 ? 1 : 0;
-		}
-
-		CurrRenderBuffer = CurrBuffer;
+		// In some rare occasions, the render thread might be null, like when offloading work to Lightmass 
+		// The final GRenderingThread check keeps us from inadvertently failing when that happens.
+#if DO_CHECK
+		bool CPUSimOK = (SimTarget == ENiagaraSimTarget::CPUSim && !IsInRenderingThread());
+		bool GPUSimOK = (SimTarget == ENiagaraSimTarget::GPUComputeSim && IsInRenderingThread());
+		checkf(!GRenderingThread || CPUSimOK || GPUSimOK, TEXT("NiagaraDataSet function being called on incorrect thread."));
+#endif
 	}
 
-
 	void BuildLayout()
 	{
 		VariableLayouts.Empty();
@@ -530,20 +622,31 @@ private:
 	
 	/** Index of current state data. */
 	uint32 CurrBuffer;
-	uint32 CurrRenderBuffer;
 	uint32 MaxBufferIdx;
 
+	ENiagaraSimTarget SimTarget;
+
 	/** Once finalized, the data layout etc is built and no more variables can be added. */
-	bool bFinalized;
+	uint32 bFinalized : 1;
+	uint32 bNeedsPersistentIDs : 1;
+	
+	/** Table of IDs to real buffer indices. Multi buffered so we can access previous frame data. */
+	TArray<int32> IDToIndexTable[3];
+
+	/** Table of free IDs available to allocate next tick. */
+	TArray<int32> FreeIDsTable;
+
+	/** Number of free IDs in FreeIDTable. */
+	int32 NumFreeIDs;
+
+	/** Max ID seen in last execution. Allows us to shrink the IDTable. */
+	int32 MaxUsedID;
+
+	/** Tag to use when new IDs are acquired. Should be unique per tick. */
+	int32 IDAcquireTag;
 
 	FNiagaraDataBuffer Data[3];
 	FRWBuffer DataSetIndices[3]; 
-
-	/** GPU side copies of the most recent cpu sim data for use in rendering. */
-	FDynamicReadBuffer RenderDataFloat;
-	FDynamicReadBuffer RenderDataInt;
-	FShaderResourceViewRHIRef CurrentFloatDataSRV;
-	FShaderResourceViewRHIRef CurrentIntDataSRV;
 };
 
 /**
@@ -1069,6 +1172,84 @@ private:
 	float* WBase;
 };
 
+
+template<>
+struct FNiagaraDataSetAccessor<FQuat> : public FNiagaraDataSetAccessorBase
+{
+	FNiagaraDataSetAccessor<FQuat>() {}
+	FNiagaraDataSetAccessor<FQuat>(FNiagaraDataSet& InDataSet, FNiagaraVariable InVar, bool bCurrBuffer = true)
+		: FNiagaraDataSetAccessorBase(&InDataSet, InVar, bCurrBuffer)
+	{
+		check(sizeof(FQuat) == InVar.GetType().GetSize());
+		InitForAccess(bCurrBuffer);
+	}
+
+	void InitForAccess(bool bCurrBuffer = true)
+	{
+		DataBuffer = bCurrBuffer ? &DataSet->CurrData() : &DataSet->PrevData();
+		if (VarLayout)
+		{
+			XBase = (float*)DataBuffer->GetComponentPtrFloat(VarLayout->FloatComponentStart);
+			YBase = (float*)DataBuffer->GetComponentPtrFloat(VarLayout->FloatComponentStart + 1);
+			ZBase = (float*)DataBuffer->GetComponentPtrFloat(VarLayout->FloatComponentStart + 2);
+			WBase = (float*)DataBuffer->GetComponentPtrFloat(VarLayout->FloatComponentStart + 3);
+		}
+		else
+		{
+			XBase = nullptr;
+			YBase = nullptr;
+			ZBase = nullptr;
+			WBase = nullptr;
+		}
+	}
+
+	FORCEINLINE FQuat operator[](int32 Index)const
+	{
+		return Get(Index);
+	}
+
+	FORCEINLINE FQuat GetSafe(int32 Index, const FQuat& Default = FQuat(0.0f, 0.0f, 0.0f, 1.0f))const
+	{
+		if (IsValid())
+		{
+			return Get(Index);
+		}
+
+		return Default;
+	}
+
+	FORCEINLINE FQuat Get(int32 Index)const
+	{
+		FQuat Ret;
+		Get(Index, Ret);
+		return Ret;
+	}
+
+	FORCEINLINE void Get(int32 Index, FQuat& OutValue)const
+	{
+		OutValue.X = XBase[Index];
+		OutValue.Y = YBase[Index];
+		OutValue.Z = ZBase[Index];
+		OutValue.W = WBase[Index];
+	}
+
+	FORCEINLINE void Set(int32 Index, const FQuat& InValue)
+	{
+		XBase[Index] = InValue.X;
+		YBase[Index] = InValue.Y;
+		ZBase[Index] = InValue.Z;
+		WBase[Index] = InValue.W;
+	}
+
+	FORCEINLINE bool BaseIsValid() const { return XBase != nullptr && YBase != nullptr && ZBase != nullptr && WBase != nullptr; }
+private:
+
+	float* XBase;
+	float* YBase;
+	float* ZBase;
+	float* WBase;
+};
+
 template<>
 struct FNiagaraDataSetAccessor<FLinearColor> : public FNiagaraDataSetAccessorBase
 {
@@ -1220,6 +1401,73 @@ private:
 	float* IntervalDtBase;
 };
 
+template<>
+struct FNiagaraDataSetAccessor<FNiagaraID> : public FNiagaraDataSetAccessorBase
+{
+	FNiagaraDataSetAccessor<FNiagaraID>() {}
+	FNiagaraDataSetAccessor<FNiagaraID>(FNiagaraDataSet& InDataSet, FNiagaraVariable InVar, bool bCurrBuffer = true)
+		: FNiagaraDataSetAccessorBase(&InDataSet, InVar, bCurrBuffer)
+	{
+		InitForAccess(bCurrBuffer);
+	}
+
+	void InitForAccess(bool bCurrBuffer = true)
+	{
+		DataBuffer = bCurrBuffer ? &DataSet->CurrData() : &DataSet->PrevData();
+		if (VarLayout != nullptr)
+		{
+			IndexBase = (int32*)DataBuffer->GetComponentPtrInt32(VarLayout->Int32ComponentStart);
+			TagBase = (int32*)DataBuffer->GetComponentPtrInt32(VarLayout->Int32ComponentStart + 1);
+		}
+		else
+		{
+			IndexBase = nullptr;
+			TagBase = nullptr;
+		}
+	}
+
+	FORCEINLINE FNiagaraID operator[](int32 Index)const
+	{
+		return Get(Index);
+	}
+
+	FORCEINLINE FNiagaraID GetSafe(int32 Index, FNiagaraID Default = FNiagaraID())const
+	{
+		if (IsValid())
+		{
+			return Get(Index);
+		}
+
+		return Default;
+	}
+
+	FORCEINLINE FNiagaraID Get(int32 Index)const
+	{
+		FNiagaraID Ret;
+		Get(Index, Ret);
+		return Ret;
+	}
+
+	FORCEINLINE void Get(int32 Index, FNiagaraID& OutValue)const
+	{
+		OutValue.Index = IndexBase[Index];
+		OutValue.AcquireTag = TagBase[Index];
+	}
+
+	FORCEINLINE void Set(int32 Index, const FNiagaraID& InValue)
+	{
+		IndexBase[Index] = InValue.Index;
+		TagBase[Index] = InValue.AcquireTag;
+	}
+
+	FORCEINLINE bool BaseIsValid() const { return IndexBase != nullptr && TagBase != nullptr; }
+
+private:
+
+	int32* IndexBase;
+	int32* TagBase;
+};
+
 //Provide iterator to keep iterator access patterns still in use.
 template<typename T>
 struct FNiagaraDataSetIterator : public FNiagaraDataSetAccessor<T>
@@ -1230,6 +1478,12 @@ struct FNiagaraDataSetIterator : public FNiagaraDataSetAccessor<T>
 		, CurrIdx(StartIndex)
 	{}
 
+	void Create(FNiagaraDataSet* InDataSet, FNiagaraVariable InVar, uint32 StartIndex=0)
+	{
+		FNiagaraDataSetAccessor<T>::Create(InDataSet, InVar);
+		CurrIdx = StartIndex;
+	}
+
 	FORCEINLINE T operator*() { return Get(); }
 	FORCEINLINE T Get()const
 	{
@@ -1372,4 +1626,4 @@ private:
 	TArray<const FNiagaraVariableLayoutInfo*> VarLayouts;
 
 	uint32 CurrIdx;
-};
+};
\ No newline at end of file
diff --git a/Engine/Plugins/FX/Niagara/Source/Niagara/Classes/NiagaraEmitter.h b/Engine/Plugins/FX/Niagara/Source/Niagara/Classes/NiagaraEmitter.h
--- a/Engine/Plugins/FX/Niagara/Source/Niagara/Classes/NiagaraEmitter.h
+++ b/Engine/Plugins/FX/Niagara/Source/Niagara/Classes/NiagaraEmitter.h
@@ -51,8 +51,8 @@ struct FNiagaraEventReceiverProperties
 	UPROPERTY(EditAnywhere, Category = "Event Receiver")
 	FName SourceEmitter;
 
-	UPROPERTY(EditAnywhere, Instanced, Category = "Event Receiver")
-	TArray<UNiagaraEventReceiverEmitterAction*> EmitterActions;
+	//UPROPERTY(EditAnywhere, Category = "Event Receiver")
+	//TArray<UNiagaraEventReceiverEmitterAction*> EmitterActions;
 };
 
 USTRUCT()
@@ -97,14 +97,6 @@ enum class EScriptExecutionMode : uint8
 	SingleParticle
 };
 
-UENUM()
-enum class EScriptCompileIndices : uint8
-{
-	SpawnScript = 0,
-	UpdateScript,
-	EventScript
-};
-
 USTRUCT()
 struct FNiagaraEmitterScriptProperties
 {
@@ -115,8 +107,8 @@ struct FNiagaraEmitterScriptProperties
 
 	GENERATED_BODY()
 	
- 	UPROPERTY()
- 	UNiagaraScript *Script;
+	UPROPERTY()
+	UNiagaraScript *Script;
 
 	UPROPERTY()
 	TArray<FNiagaraEventReceiverProperties> EventReceivers;
@@ -139,13 +131,15 @@ struct FNiagaraEventScriptProperties : public FNiagaraEmitterScriptProperties
 		ExecutionMode = EScriptExecutionMode::EveryParticle;
 		SpawnNumber = 0;
 		MaxEventsPerFrame = 0;
+		bRandomSpawnNumber = false;
+		MinSpawnNumber = 0;
 	}
 	
 	/** Controls which particles have the event script run on them.*/
 	UPROPERTY(EditAnywhere, Category="Event Handler Options")
 	EScriptExecutionMode ExecutionMode;
 
-	/** Controls whether or not particles are spawned as a result of handling the event. Only valid for EScriptExecutionMode::SpawnedParticles*/
+	/** Controls whether or not particles are spawned as a result of handling the event. Only valid for EScriptExecutionMode::SpawnedParticles. If Random Spawn Number is used, this will act as the maximum spawn range. */
 	UPROPERTY(EditAnywhere, Category="Event Handler Options")
 	uint32 SpawnNumber;
 
@@ -160,32 +154,14 @@ struct FNiagaraEventScriptProperties : public FNiagaraEmitterScriptProperties
 	/** The name of the event generated. This will be "Collision" for collision events and the Event Name field on the DataSetWrite node in the module graph for others.*/
 	UPROPERTY(EditAnywhere, Category="Event Handler Options")
 	FName SourceEventName;
-};
 
-/** Represents timed burst of particles in an emitter. */
-USTRUCT()
-struct FNiagaraEmitterBurst
-{
-	GENERATED_BODY()
-
-	/** The base time of the burst absolute emitter time. */
-	UPROPERTY(EditAnywhere, Category = "Burst")
-	float Time;
+	/** Whether using a random spawn number. */
+	UPROPERTY(EditAnywhere, Category = "Event Handler Options")
+	bool bRandomSpawnNumber;
 
-	/** 
-	 * A range around the base time which can be used to randomize burst timing.  The resulting range is 
-	 * from Time - (TimeRange / 2) to Time + (TimeRange / 2).
-	 */
-	UPROPERTY(EditAnywhere, Category = "Burst")
-	float TimeRange;
-
-	/** The minimum number of particles to spawn. */
-	UPROPERTY(EditAnywhere, Category = "Burst")
-	uint32 SpawnMinimum;
-
-	/** The maximum number of particles to spawn. */
-	UPROPERTY(EditAnywhere, Category = "Burst")
-	uint32 SpawnMaximum;
+	/** The minimum spawn number when random spawn is used. Spawn Number is used as the maximum range. */
+	UPROPERTY(EditAnywhere, Category = "Event Handler Options")
+	uint32 MinSpawnNumber;
 };
 
 /** 
@@ -200,6 +176,12 @@ class UNiagaraEmitter : public UObject
 public:
 #if WITH_EDITOR
 	DECLARE_MULTICAST_DELEGATE(FOnPropertiesChanged);
+	DECLARE_MULTICAST_DELEGATE_OneParam(FOnEmitterCompiled, UNiagaraEmitter*);
+
+	struct NIAGARA_API PrivateMemberNames
+	{
+		static const FName EventHandlerScriptProps;
+	};
 #endif
 
 public:
@@ -222,9 +204,9 @@ public:
 	//UPROPERTY(EditAnywhere, Category = "Emitter")
 	//float EndTime;
 	//UPROPERTY(EditAnywhere, Category = "Emitter")
-	//int32 NumLoops;
-	UPROPERTY(EditAnywhere, Category = "Emitter")
-	ENiagaraCollisionMode CollisionMode;
+	//int32 NumLoops
+	//UPROPERTY(EditAnywhere, Category = "Emitter")
+	//ENiagaraCollisionMode CollisionMode;
 
 	UPROPERTY()
 	FNiagaraEmitterScriptProperties UpdateScriptProps;
@@ -270,18 +252,23 @@ public:
 	UPROPERTY(EditAnywhere, Category = "Scalability", meta = (InlineEditConditionToggle))
 	uint32 bUseMaxDetailLevel : 1;
 
+	/** Do particles in this emitter require a persistent ID? */
+	UPROPERTY(EditAnywhere, Category = "Emitter")
+	uint32 bRequiresPersistentIDs : 1;
+
 	void NIAGARA_API GetScripts(TArray<UNiagaraScript*>& OutScripts, bool bCompilableOnly = true);
 
 	NIAGARA_API UNiagaraScript* GetScript(ENiagaraScriptUsage Usage, FGuid UsageId);
 
+	NIAGARA_API UNiagaraScript* GetGPUComputeScript() { return GPUComputeScript; }
+
 #if WITH_EDITORONLY_DATA
 	/** 'Source' data/graphs for the scripts used by this emitter. */
 	UPROPERTY()
 	class UNiagaraScriptSourceBase*	GraphSource;
 
 	bool NIAGARA_API AreAllScriptAndSourcesSynchronized() const;
-	void NIAGARA_API CompileScripts(TArray<ENiagaraScriptCompileStatus>& OutScriptStatuses, TArray<FString>& OutGraphLevelErrorMessages, TArray<FString>& PathNames, TArray<UNiagaraScript*>& Scripts, bool bForce);
-	ENiagaraScriptCompileStatus NIAGARA_API CompileScript(EScriptCompileIndices InScriptToCompile, FString& OutGraphLevelErrorMessages, bool bForce, int32 InSubScriptIdx = 0);
+	void NIAGARA_API OnPostCompile();
 
 	UNiagaraEmitter* MakeRecursiveDeepCopy(UObject* DestOuter) const;
 	UNiagaraEmitter* MakeRecursiveDeepCopy(UObject* DestOuter, TMap<const UObject*, UObject*>& ExistingConversions) const;
@@ -300,10 +287,16 @@ public:
 	/** Internal: Indicates the thumbnail image is out of date.*/
 	UPROPERTY()
 	uint32 ThumbnailImageOutOfDate : 1;
+	
+	/** Callback issued whenever a VM compilation successfully happened (even if the results are a script that cannot be executed due to errors)*/
+	NIAGARA_API FOnEmitterCompiled& OnEmitterVMCompiled();
+
+	NIAGARA_API static bool GetForceCompileOnLoad();
 #endif
 
 	/** Is this emitter allowed to be enabled by the current system detail level. */
 	bool IsAllowedByDetailLevel()const;
+	NIAGARA_API bool RequiresPersistantIDs()const;
 
 	NIAGARA_API bool IsValid()const;
 	NIAGARA_API bool IsReadyToRun() const;
@@ -312,7 +305,7 @@ public:
 	bool UsesCollection(const class UNiagaraParameterCollection* Collection)const;
 
 	FString NIAGARA_API GetUniqueEmitterName()const;
-	void NIAGARA_API SetUniqueEmitterName(const FString& InName);
+	bool NIAGARA_API SetUniqueEmitterName(const FString& InName);
 
 	/** Converts an emitter paramter "Emitter.XXXX" into it's real parameter name. */
 	FNiagaraVariable ToEmitterParameter(const FNiagaraVariable& EmitterVar)const;
@@ -337,6 +330,8 @@ protected:
 	virtual void BeginDestroy() override;
 
 #if WITH_EDITORONLY_DATA
+	void SyncEmitterAlias(const FString& InOldName, const FString& InNewName);
+
 private:
 	void UpdateChangeId();
 
@@ -350,6 +345,9 @@ private:
 	/** Adjusted every time that we compile this emitter. Lets us know that we might differ from any cached versions.*/
 	UPROPERTY()
 	FGuid ChangeId;
+
+	/** A multicast delegate which is called whenever all the scripts for this emitter have been compiled (successfully or not). */
+	FOnEmitterCompiled OnVMScriptCompiledDelegate;
 #endif
 
 	UPROPERTY()
@@ -358,9 +356,12 @@ private:
 	UPROPERTY()
 	TArray<UNiagaraRendererProperties*> RendererProperties;
 
-	UPROPERTY()
+	UPROPERTY(EditAnywhere, Category = "Events", meta=(NiagaraNoMerge))
 	TArray<FNiagaraEventScriptProperties> EventHandlerScriptProps;
 
+	UPROPERTY()
+	UNiagaraScript* GPUComputeScript;
+
 #if WITH_EDITOR
 	FOnPropertiesChanged OnPropertiesChangedDelegate;
 #endif
