Copying //UE4/Dev-Niagara to //UE4/Dev-Main (Source: //UE4/Dev-Niagara @ 4074996)

diff --git a/Engine/Plugins/FX/Niagara/Source/Niagara/Classes/NiagaraEmitterInstance.h b/Engine/Plugins/FX/Niagara/Source/Niagara/Classes/NiagaraEmitterInstance.h
--- a/Engine/Plugins/FX/Niagara/Source/Niagara/Classes/NiagaraEmitterInstance.h
+++ b/Engine/Plugins/FX/Niagara/Source/Niagara/Classes/NiagaraEmitterInstance.h
@@ -26,14 +26,6 @@ class UNiagaraParameterCollectionInstance;
 */
 struct FNiagaraEmitterInstance
 {
-
-private:
-	struct FNiagaraBurstInstance
-	{
-		float Time;
-		uint32 NumberToSpawn;
-	};
-
 public:
 	explicit FNiagaraEmitterInstance(FNiagaraSystemInstance* InParentSystemInstance);
 	bool bDumpAfterEvent;
@@ -58,9 +50,11 @@ public:
 	void PostProcessParticles();
 	bool HandleCompletion(bool bForce=false);
 
+	bool RequiredPersistentID()const;
+
 	FORCEINLINE bool ShouldTick()const { return ExecutionState == ENiagaraExecutionState::Active || GetNumParticles() > 0; }
 
-	uint32 CalculateEventSpawnCount(const FNiagaraEventScriptProperties &EventHandlerProps, TArray<int32> &EventSpawnCounts, FNiagaraDataSet *EventSet);
+	uint32 CalculateEventSpawnCount(const FNiagaraEventScriptProperties &EventHandlerProps, TArray<int32, TInlineAllocator<16>>& EventSpawnCounts, FNiagaraDataSet *EventSet);
 
 	NIAGARA_API FBox CalculateDynamicBounds();
 
@@ -74,9 +68,6 @@ public:
 
 	FORCEINLINE bool IsDisabled()const { return ExecutionState == ENiagaraExecutionState::Disabled; }
 	FORCEINLINE bool IsComplete()const { return ExecutionState == ENiagaraExecutionState::Complete || ExecutionState == ENiagaraExecutionState::Disabled; }
-	
-	/** Sets the error state. */
-	void SetError(bool bInError) { bError = bInError; }
 		
 	/** Create a new NiagaraRenderer. The old renderer is not immediately deleted, but instead put in the ToBeRemoved list.*/
 	void NIAGARA_API UpdateEmitterRenderer(ERHIFeatureLevel::Type FeatureLevel, TArray<NiagaraRenderer*>& ToBeAddedList, TArray<NiagaraRenderer*>& ToBeRemovedList);
@@ -104,6 +95,9 @@ public:
 	FNiagaraScriptExecutionContext& GetUpdateExecutionContext() { return UpdateExecContext; }
 	TArray<FNiagaraScriptExecutionContext>& GetEventExecutionContexts() { return EventExecContexts; }
 
+	FORCEINLINE FName GetCachedIDName()const { return CachedIDName; }
+	FORCEINLINE UNiagaraEmitter* GetCachedEmitter()const { return CachedEmitter; }
+
 	TArray<FNiagaraSpawnInfo>& GetSpawnInfo() { return SpawnInfos; }
 
 	NIAGARA_API bool IsReadyToRun() const;
@@ -119,9 +113,8 @@ private:
 	float Age;
 	/* how many loops this emitter has gone through */
 	uint32 Loops;
-	
-	/* Some error has occurred so stop ticking or rendering. Better to just kill the emitter? */
-	uint32 bError : 1;
+
+	int32 TickCount;
 	
 	/** Typical resets must be deferred until the tick as the RT could still be using the current buffer. */
 	uint32 bResetPending : 1;
@@ -144,6 +137,11 @@ private:
 	FNiagaraParameterDirectBinding<float> SpawnIntervalBinding;
 	FNiagaraParameterDirectBinding<float> InterpSpawnStartBinding;
 
+	FNiagaraParameterDirectBinding<float> SpawnIntervalBindingGPU;
+	FNiagaraParameterDirectBinding<float> InterpSpawnStartBindingGPU;
+	FNiagaraParameterDirectBinding<float> EmitterAgeBindingGPU;
+	FNiagaraParameterDirectBinding<FNiagaraBool> EmitterLocalSpaceBindingGPU;
+
 	FNiagaraParameterDirectBinding<float> SpawnEmitterAgeBinding;
 	FNiagaraParameterDirectBinding<float> UpdateEmitterAgeBinding;
 	TArray<FNiagaraParameterDirectBinding<float>> EventEmitterAgeBindings;
@@ -162,6 +160,10 @@ private:
 	TArray<NiagaraRenderer*> EmitterRenderer;
 	FNiagaraSystemInstance *ParentSystemInstance;
 
+	/** Raw pointer to the emitter that we're instanced from. Raw ptr should be safe here as we check for the validity of the system and it's emitters higher up before any ticking. */
+	UNiagaraEmitter* CachedEmitter;
+	FName CachedIDName;
+
 	TArray<FNiagaraDataSet*> UpdateScriptEventDataSets;
 	TArray<FNiagaraDataSet*> SpawnScriptEventDataSets;
 	TMap<FNiagaraDataSetID, FNiagaraDataSet*> DataSetMap;
@@ -176,4 +178,15 @@ private:
 	FNiagaraDataSetAccessor<FVector> PositionAccessor;
 	FNiagaraDataSetAccessor<FVector2D> SizeAccessor;
 	FNiagaraDataSetAccessor<FVector> MeshScaleAccessor;
+
+	static const FName PositionName;
+	static const FName SizeName;
+	static const FName MeshScaleName;
+
+#if !UE_BUILD_SHIPPING
+	bool bEncounteredNaNs;
+#endif
+
+	/** A parameter store which contains the data interfaces parameters which were defined by the scripts. */
+	FNiagaraParameterStore ScriptDefinedDataInterfaceParameters;
 };
diff --git a/Engine/Plugins/FX/Niagara/Source/Niagara/Classes/NiagaraEmitterInstanceBatcher.h b/Engine/Plugins/FX/Niagara/Source/Niagara/Classes/NiagaraEmitterInstanceBatcher.h
--- a/Engine/Plugins/FX/Niagara/Source/Niagara/Classes/NiagaraEmitterInstanceBatcher.h
+++ b/Engine/Plugins/FX/Niagara/Source/Niagara/Classes/NiagaraEmitterInstanceBatcher.h
@@ -14,7 +14,7 @@ the same VectorVM byte code / compute shader code
 #include "NiagaraEmitter.h"
 #include "Tickable.h"
 #include "Modules/ModuleManager.h"
-#include "Modules/ModuleManager.h"
+#include "RHIResources.h"
 
 struct FNiagaraScriptExecutionContext;
 struct FNiagaraComputeExecutionContext;
@@ -30,11 +30,6 @@ public:
 		{
 			RendererModule->RegisterPostOpaqueComputeDispatcher(this);
 		}
-
-		if (DummyWriteIndexBuffer.Buffer == nullptr)
-		{
-			DummyWriteIndexBuffer.Initialize(sizeof(int32), 64, EPixelFormat::PF_R32_SINT);	// always allocate for up to 64 data sets
-		}
 	}
 
 	~NiagaraEmitterInstanceBatcher()
@@ -81,39 +76,42 @@ public:
 	{
 	}
 
+	uint32 GetEventSpawnTotal(const FNiagaraComputeExecutionContext *InContext) const;
+
 	// from FComputeDispatcher; called once per frame by the render thread, swaps buffers and works down 
 	// the queue submitted by the game thread; means we're one frame behind with this; we need a mechanism to determine execution order here.
-	virtual void Execute(FRHICommandList &RHICmdList)
+	virtual void Execute(FRHICommandList &RHICmdList, FUniformBufferRHIParamRef ViewUniformBuffer)
 	{
 		CurQueueIndex ^= 0x1;
-		ExecuteAll(RHICmdList);
+		ExecuteAll(RHICmdList, ViewUniformBuffer);
 	}
 
-	void ExecuteAll(FRHICommandList &RHICmdList);
-	void TickSingle(const FNiagaraComputeExecutionContext *Context, FRHICommandList &RHICmdList) const;
+	void ExecuteAll(FRHICommandList &RHICmdList, FUniformBufferRHIParamRef ViewUniformBuffer);
+	void TickSingle(FNiagaraComputeExecutionContext *Context, FRHICommandList &RHICmdList, FUniformBufferRHIParamRef ViewUniformBuffer) const;
 
 	void SetPrevDataStrideParams(const FNiagaraDataSet *Set, FNiagaraShader *Shader, FRHICommandList &RHICmdList) const;
 
 	void SetupEventUAVs(const FNiagaraComputeExecutionContext *Context, uint32 NumInstances, FRHICommandList &RHICmdList) const;
 	void UnsetEventUAVs(const FNiagaraComputeExecutionContext *Context, FRHICommandList &RHICmdList) const;
-	void SetupDataInterfaceBuffers(const TArray<FNiagaraScriptDataInterfaceInfo> &DIInfos, FNiagaraShader *Shader, FRHICommandList &RHICmdList) const;
-
-	void Run(	FNiagaraDataSet *DataSet,
-				uint32 Phase0StartInstance,
-				uint32 Phase1StartInstance,
-				const uint32 TotalNumInstances,
-				FNiagaraShader *Shader,
-				const TArray<uint8, TAlignedHeapAllocator<16>>  &Params,
-				FRHICommandList &RHICmdList,
-				const FRWBuffer &WriteIndexBuffer,
-				bool bCopyBeforeStart = false) const;
+	void SetDataInterfaceParameters(const TArray<UNiagaraDataInterface*> &DataInterfaces, FNiagaraShader *Shader, FRHICommandList &RHICmdList) const;
+
+	void Run(	const FNiagaraComputeExecutionContext *Context, 
+				uint32 UpdateStartInstance, 
+				const uint32 TotalNumInstances, 
+				FNiagaraShader* Shader,
+				FRHICommandList &RHICmdList, 
+				FUniformBufferRHIParamRef ViewUniformBuffer, 
+				bool bCopyBeforeStart = false
+			) const;
 
 	void RunEventHandlers(const FNiagaraComputeExecutionContext *Context, uint32 NumInstancesAfterSim, uint32 NumInstancesAfterSpawn, uint32 NumInstancesAfterNonEventSpawn, FRHICommandList &RhiCmdList) const;
-	void ResolveDatasetWrites(uint32 *OutArray, const FNiagaraComputeExecutionContext *Context, FRWBuffer &WriteIndexBuffer) const;
+
+	void ClearIndexBufferCur(FRHICommandList &RHICmdList, FNiagaraComputeExecutionContext *Context) const;
+	void ResolveDatasetWrites(FNiagaraComputeExecutionContext *Context) const;
+	void ResizeCurrentBuffer(FRHICommandList &RHICmdList, FNiagaraComputeExecutionContext *Context, uint32 NewNumInstances, uint32 PrevNumInstances) const;
 private:
 	static NiagaraEmitterInstanceBatcher* BatcherSingleton;
 
 	uint32 CurQueueIndex;
 	TArray<FNiagaraComputeExecutionContext*> SimulationQueue[2];
-	static FRWBuffer DummyWriteIndexBuffer;
 };
diff --git a/Engine/Plugins/FX/Niagara/Source/Niagara/Classes/NiagaraParameterCollection.h b/Engine/Plugins/FX/Niagara/Source/Niagara/Classes/NiagaraParameterCollection.h
--- a/Engine/Plugins/FX/Niagara/Source/Niagara/Classes/NiagaraParameterCollection.h
+++ b/Engine/Plugins/FX/Niagara/Source/Niagara/Classes/NiagaraParameterCollection.h
@@ -80,6 +80,9 @@ public:
 	UFUNCTION(BlueprintCallable, Category = Niagara, meta = (DisplayName = "Get Vector4 Parameter"))
 	FVector4 GetVector4Parameter(const FString& InVariableName);
 
+	UFUNCTION(BlueprintCallable, Category = Niagara, meta = (DisplayName = "Get Quaternion Parameter"))
+	FQuat GetQuatParameter(const FString& InVariableName);
+
 	UFUNCTION(BlueprintCallable, Category = Niagara, meta = (DisplayName = "Get Color Parameter"))
 	FLinearColor GetColorParameter(const FString& InVariableName);
 
@@ -103,6 +106,9 @@ public:
 
 	UFUNCTION(BlueprintCallable, Category = Niagara, meta = (DisplayName = "Set Color Parameter"))
 	void SetColorParameter(const FString& InVariableName, FLinearColor InValue);
+	
+	UFUNCTION(BlueprintCallable, Category = Niagara, meta = (DisplayName = "Set Quaternion Parameter"))
+	void SetQuatParameter(const FString& InVariableName, const FQuat& InValue);
 };
 
 /** Asset containing a collection of global parameters usable by Niagara. */
@@ -146,6 +152,16 @@ public:
 
 	FString FriendlyNameFromParameterName(FString ParameterName)const;
 	FString GetFullNamespace()const;
+
+	/** The compile Id is an indicator to any compiled scripts that reference this collection that contents may have changed and a recompile is recommended to be safe.*/
+	FGuid GetCompileId() const { return CompileId; };
+
+	/** If for any reason this data has changed externally and needs recompilation and it isn't autodetected, use this method. */
+	void RefreshCompileId();
+
+	//~UObject interface
+	virtual void PostLoad()override;
+	//~UObject interface
 protected:
 	
 	void MakeNamespaceNameUnique();
@@ -159,4 +175,8 @@ protected:
 	
 	UPROPERTY()
 	UNiagaraParameterCollectionInstance* DefaultInstance;
+
+	/** Used to track whenever something of note changes in this parameter collection that might invalidate a compilation downstream of a script/emitter/system.*/
+	UPROPERTY()
+	FGuid CompileId;
 };
\ No newline at end of file
diff --git a/Engine/Plugins/FX/Niagara/Source/Niagara/Classes/NiagaraScript.h b/Engine/Plugins/FX/Niagara/Source/Niagara/Classes/NiagaraScript.h
--- a/Engine/Plugins/FX/Niagara/Source/Niagara/Classes/NiagaraScript.h
+++ b/Engine/Plugins/FX/Niagara/Source/Niagara/Classes/NiagaraScript.h
@@ -11,18 +11,20 @@
 #include "NiagaraParameters.h"
 #include "NiagaraDataSet.h"
 #include "NiagaraShared.h"
-#include "NiagaraParameterStore.h"
+#include "NiagaraScriptExecutionParameterStore.h"
+#include "NiagaraCustomVersion.h"
 
 #include "NiagaraScript.generated.h"
 
 class UNiagaraDataInterface;
+class FNiagaraCompileRequestDataBase;
 
-void SerializeNiagaraShaderMaps(const TMap<const ITargetPlatform*, FNiagaraScript*>* PlatformScriptResourcesToSavePtr, FArchive& Ar, FNiagaraScript& OutLoadedResources);
-void ProcessSerializedShaderMaps(UNiagaraScript* Owner, FNiagaraScript LoadedResource, FNiagaraScript* (&OutScriptResourcesLoaded)[ERHIFeatureLevel::Num]);
+void SerializeNiagaraShaderMaps(const TMap<const ITargetPlatform*, FNiagaraShaderScript*>* PlatformScriptResourcesToSavePtr, FArchive& Ar, FNiagaraShaderScript& OutLoadedResources);
+void ProcessSerializedShaderMaps(UNiagaraScript* Owner, FNiagaraShaderScript LoadedResource, FNiagaraShaderScript* (&OutScriptResourcesLoaded)[ERHIFeatureLevel::Num]);
 
 #define NIAGARA_INVALID_MEMORY (0xBA)
 
-//#define NIAGARA_SCRIPT_COMPILE_LOGGING_MEDIUM
+#define NIAGARA_SCRIPT_COMPILE_LOGGING_MEDIUM
 
 /** Defines what will happen to unused attributes when a script is run. */
 UENUM()
@@ -40,60 +42,151 @@
 	PassThrough,
 };
 
+UENUM()
+enum class ENiagaraModuleDependencyType : uint8
+{
+	/** The dependency belongs before the module. */
+	PreDependency,
+	/** The dependency belongs after the module. */
+	PostDependency
+};
+
+USTRUCT()
+struct FNiagaraModuleDependency
+{
+	GENERATED_USTRUCT_BODY()
+public:
+	/** Specifies the provided id of the required dependent module (e.g. 'ProvidesNormalizedAge') */
+	UPROPERTY(AssetRegistrySearchable, EditAnywhere, Category = Script)
+		FName Id;
+	/** Whether the dependency belongs before or after this module */
+	UPROPERTY(AssetRegistrySearchable, EditAnywhere, Category = Script)
+		ENiagaraModuleDependencyType Type; // e.g. PreDependency,
+										   /** Detailed description of the dependency */
+	UPROPERTY(AssetRegistrySearchable, EditAnywhere, Category = Script)
+		FText Description;
+};
+
 struct FNiagaraScriptDebuggerInfo
 {
 	FNiagaraScriptDebuggerInfo();
+	FNiagaraScriptDebuggerInfo(FName InName, ENiagaraScriptUsage InUsage, const FGuid& InUsageId);
 
-	bool bRequestDebugFrame;
-	
-	int32 DebugFrameLastWriteId;
+	FName HandleName;
 
-	FNiagaraDataSet DebugFrame;
-
-	FNiagaraParameterStore DebugParameters;
-};
-
-
-/** Runtime script for a Niagara system */
-UCLASS(MinimalAPI)
-class UNiagaraScript : public UObject
-{
-	GENERATED_UCLASS_BODY()
-
-	// how this script is to be used. cannot be private due to use of GET_MEMBER_NAME_CHECKED
-	UPROPERTY(AssetRegistrySearchable)
 	ENiagaraScriptUsage Usage;
 
-	/** Which instance of the usage in the graph to use.  This is now deprecated and is handled by UsageId. */
-	UPROPERTY()
-	int32 UsageIndex_DEPRECATED;
-
-private:
-	/** Specifies a unique id for use when there are multiple scripts with the same usage, e.g. events. */
-	UPROPERTY()
 	FGuid UsageId;
 
+	int32 FrameLastWriteId;
+
+	FNiagaraDataSet Frame;
+
+	FNiagaraParameterStore Parameters;
+};
+
+
+/** Struct containing all of the data necessary to look up a NiagaraScript's VM executable results from the Derived Data Cache.*/
+USTRUCT()
+struct NIAGARA_API FNiagaraVMExecutableDataId
+{
+	GENERATED_USTRUCT_BODY()
 public:
-	/** When used as a module, what are the appropriate script types for referencing this module?*/
-	UPROPERTY(AssetRegistrySearchable, EditAnywhere, Category = Script, meta = (Bitmask, BitmaskEnum = ENiagaraScriptUsage))
-	int32 ModuleUsageBitmask;
+	/** The version of the compiler that this needs to be built against.*/
+	UPROPERTY()
+	FGuid CompilerVersionID;
 
-	/** Used to break up scripts of the same Usage type in UI display.*/
-	UPROPERTY(AssetRegistrySearchable, EditAnywhere, Category = Script)
-	FText Category;
+	/** The type of script this was used for.*/
+	UPROPERTY()
+	ENiagaraScriptUsage ScriptUsageType;
 
+	/** The instance id of this script usage type.*/
+	UPROPERTY()
+	FGuid ScriptUsageTypeID;
+
+	/** Configuration options*/
+	UPROPERTY()
+	TArray<FString> AdditionalDefines;
+
+	/**
+	* The GUID of the subgraph this shader primarily represents.
+	*/
+	UPROPERTY()
+	FGuid BaseScriptID;
+
+	/** Guids of any functions, module scripts, parameter collections, or other assets the script was dependent on that might trigger a recompile if they change. */
+	UPROPERTY()
+	TArray<FGuid> ReferencedDependencyIds;
+
+	/** Temp storage while generating the Id. This is NOT serialized and shouldn't be used in any comparisons*/
+	TArray<UObject*> ReferencedObjects;
+
+	FNiagaraVMExecutableDataId()
+		: CompilerVersionID()
+		, BaseScriptID(0, 0, 0, 0)
+	{ }
+
+
+	~FNiagaraVMExecutableDataId()
+	{ }
+
+	bool IsValid() const;
+	void Invalidate();
+	
+	friend uint32 GetTypeHash(const FNiagaraVMExecutableDataId& Ref)
+	{
+		return Ref.BaseScriptID.A;
+	}
+
+	SIZE_T GetSizeBytes() const
+	{
+		return sizeof(*this);
+	}
+
+	bool HasInterpolatedParameters() const;
+	bool RequiresPersistentIDs() const;
+#if 0
+	/** Hashes the script-specific part of this shader map Id. */
+	void GetScriptHash(FSHAHash& OutHash) const;
+#endif
+
+	/**
+	* Tests this set against another for equality, disregarding override settings.
+	*
+	* @param ReferenceSet	The set to compare against
+	* @return				true if the sets are equal
+	*/
+	bool operator==(const FNiagaraVMExecutableDataId& ReferenceSet) const;
+
+	bool operator!=(const FNiagaraVMExecutableDataId& ReferenceSet) const
+	{
+		return !(*this == ReferenceSet);
+	}
+
+	/** Appends string representations of this Id to a key string. */
+	void AppendKeyString(FString& KeyString) const;
+};
+
+/** Struct containing all of the data needed to run a Niagara VM executable script.*/
+USTRUCT()
+struct NIAGARA_API FNiagaraVMExecutableData
+{
+	GENERATED_USTRUCT_BODY()
+public:
+	FNiagaraVMExecutableData();
+
+	/** Byte code to execute for this system */
+	UPROPERTY()
+	TArray<uint8> ByteCode;
+
 	/** Number of user pointers we must pass to the VM. */
 	UPROPERTY()
 	int32 NumUserPtrs;
-	
+
 	/** All the data for using external constants in the script, laid out in the order they are expected in the uniform table.*/
 	UPROPERTY()
 	FNiagaraParameters Parameters;
 
-	/** Contains all of the top-level values that are iterated on in the UI. These are usually "Module" variables in the graph. They don't necessarily have to be in the order that they are expected in the uniform table.*/
 	UPROPERTY()
-	FNiagaraParameterStore RapidIterationParameters;
-
-	UPROPERTY()
 	FNiagaraParameters InternalParameters;
 
@@ -102,7 +195,7 @@
 
 	/** Attributes used by this script. */
 	UPROPERTY()
- 	TArray<FNiagaraVariable> Attributes;
+	TArray<FNiagaraVariable> Attributes;
 
 	/** Contains various usage information for this script. */
 	UPROPERTY()
@@ -110,48 +203,123 @@
 
 	/** Information about all data interfaces used by this script. */
 	UPROPERTY()
-	TArray<FNiagaraScriptDataInterfaceInfo> DataInterfaceInfo;
+	TArray<FNiagaraScriptDataInterfaceCompileInfo> DataInterfaceInfo;
 
 	/** Array of ordered vm external functions to place in the function table. */
 	UPROPERTY()
 	TArray<FVMExternalFunctionBindingInfo> CalledVMExternalFunctions;
 
-	/** The mode to use when deducing the type of numeric output pins from the types of the input pins. */
-	UPROPERTY(EditAnywhere, Category=Script)
-	ENiagaraNumericOutputTypeSelectionMode NumericOutputTypeSelectionMode;
-
 	UPROPERTY()
-	TArray< FDIBufferDescriptorStore > DIBufferDescriptors;
-
-	UPROPERTY()
 	TArray<FNiagaraDataSetID> ReadDataSets;
-	
+
 	UPROPERTY()
 	TArray<FNiagaraDataSetProperties> WriteDataSets;
 
 	/** Scopes we'll track with stats.*/
 	UPROPERTY()
 	TArray<FNiagaraStatScope> StatScopes;
-	
-	/** The parameter collections used by this script. */
-	UPROPERTY()
-	TArray<class UNiagaraParameterCollection*> ParameterCollections;
-
-	UPROPERTY(AssetRegistrySearchable, EditAnywhere, Category = Script)
-	FText Description;
 
 	UPROPERTY()
 	FString LastHlslTranslation;
+
 	UPROPERTY()
 	FString LastHlslTranslationGPU;
+
 	UPROPERTY()
 	FString LastAssemblyTranslation;
+
 	UPROPERTY()
 	uint32 LastOpCount;
 
+	UPROPERTY()
+	TArray<FNiagaraDataInterfaceGPUParamInfo> DIParamInfo; //TODO: GPU Param info should not be in the "VM executable data"
+
+	/** The parameter collections used by this script. */
+	UPROPERTY()
+	TArray<FString> ParameterCollectionPaths;
 	
-	/** Returns true if this script is valid and can be executed. */
-	bool IsValid()const { return ByteCode.Num() > 0; }//More? Differentiate by CPU/GPU?
+	/** Last known compile status. Lets us determine the latest state of the script byte buffer.*/
+	UPROPERTY()
+	ENiagaraScriptCompileStatus LastCompileStatus;
 
+	UPROPERTY()
+	bool bReadsAttributeData;
+
+	UPROPERTY()
+	FString ErrorMsg;
+
+	UPROPERTY()
+	float CompileTime;
+
+	void SerializeData(FArchive& Ar, bool bDDCData);
+	
+	bool IsValid() const;
+	
+	void Reset();
+};
+
+
+/** Runtime script for a Niagara system */
+UCLASS(MinimalAPI)
+class UNiagaraScript : public UObject
+{
+	GENERATED_UCLASS_BODY()
+public:
+	// how this script is to be used. cannot be private due to use of GET_MEMBER_NAME_CHECKED
+	UPROPERTY(AssetRegistrySearchable)
+	ENiagaraScriptUsage Usage;
+
+	/** Which instance of the usage in the graph to use.  This is now deprecated and is handled by UsageId. */
+	UPROPERTY()
+	int32 UsageIndex_DEPRECATED;
+
+#if WITH_EDITOR
+	DECLARE_MULTICAST_DELEGATE_OneParam(FOnScriptCompiled, UNiagaraScript*);
+#endif
+
+private:
+	/** Specifies a unique id for use when there are multiple scripts with the same usage, e.g. events. */
+	UPROPERTY()
+	FGuid UsageId;
+
+public:
+	/** When used as a module, what are the appropriate script types for referencing this module?*/
+	UPROPERTY(AssetRegistrySearchable, EditAnywhere, Category = Script, meta = (Bitmask, BitmaskEnum = ENiagaraScriptUsage))
+	int32 ModuleUsageBitmask;
+
+#if WITH_EDITORONLY_DATA
+	/** Used to break up scripts of the same Usage type in UI display.*/
+	UPROPERTY(AssetRegistrySearchable, EditAnywhere, Category = Script)
+	FText Category;
+	
+	/** Array of Ids of dependencies provided by this module to other modules on the stack (e.g. 'ProvidesNormalizedAge') */
+	UPROPERTY(EditAnywhere, Category = Script)
+	TArray<FName> ProvidedDependencies;
+	
+	/** Dependencies required by this module from other modules on the stack */
+	UPROPERTY(EditAnywhere, Category = Script)
+	TArray<FNiagaraModuleDependency> RequiredDependencies;
+#endif
+
+
+	/** Contains all of the top-level values that are iterated on in the UI. These are usually "Module" variables in the graph. They don't necessarily have to be in the order that they are expected in the uniform table.*/
+	UPROPERTY()
+	FNiagaraParameterStore RapidIterationParameters;
+
+	/** The mode to use when deducing the type of numeric output pins from the types of the input pins. */
+	UPROPERTY(EditAnywhere, Category=Script)
+	ENiagaraNumericOutputTypeSelectionMode NumericOutputTypeSelectionMode;
+
+#if WITH_EDITORONLY_DATA
+	UPROPERTY(AssetRegistrySearchable, EditAnywhere, Category = Script)
+	FText Description;
+
+	UPROPERTY(EditAnywhere, Category = Script)
+	TMap<FName, FString> ScriptMetaData;
+#endif
+
+	NIAGARA_API void ComputeVMCompilationId(FNiagaraVMExecutableDataId& Id) const;
+	NIAGARA_API const FNiagaraVMExecutableDataId& GetComputedVMCompilationId() const { return LastGeneratedVMId; }
+
 	void SetUsage(ENiagaraScriptUsage InUsage) { Usage = InUsage; }
 	ENiagaraScriptUsage GetUsage() const { return Usage; }
@@ -162,6 +330,8 @@
 	NIAGARA_API bool ContainsUsage(ENiagaraScriptUsage InUsage) const;
 	bool IsEquivalentUsage(ENiagaraScriptUsage InUsage) const {return (InUsage == Usage) || (Usage == ENiagaraScriptUsage::ParticleSpawnScript && InUsage == ENiagaraScriptUsage::ParticleSpawnScriptInterpolated) || (Usage == ENiagaraScriptUsage::ParticleSpawnScriptInterpolated && InUsage == ENiagaraScriptUsage::ParticleSpawnScript);}
 	static bool IsEquivalentUsage(ENiagaraScriptUsage InUsageA, ENiagaraScriptUsage InUsageB) { return (InUsageA == InUsageB) || (InUsageB == ENiagaraScriptUsage::ParticleSpawnScript && InUsageA == ENiagaraScriptUsage::ParticleSpawnScriptInterpolated) || (InUsageB == ENiagaraScriptUsage::ParticleSpawnScriptInterpolated && InUsageA == ENiagaraScriptUsage::ParticleSpawnScript); }
+	/** Is usage A dependent on Usage B?*/
+	NIAGARA_API static bool IsUsageDependentOn(ENiagaraScriptUsage InUsageA, ENiagaraScriptUsage InUsageB);
 
 	bool IsParticleSpawnScript() const { return Usage == ENiagaraScriptUsage::ParticleSpawnScript || Usage == ENiagaraScriptUsage::ParticleSpawnScriptInterpolated; }
 	bool IsInterpolatedParticleSpawnScript()const { return Usage == ENiagaraScriptUsage::ParticleSpawnScriptInterpolated; }
@@ -170,6 +340,7 @@
 	bool IsFunctionScript()	const { return Usage == ENiagaraScriptUsage::Function; }
 	bool IsDynamicInputScript()	const { return Usage == ENiagaraScriptUsage::DynamicInput; }
 	bool IsParticleEventScript()const { return Usage == ENiagaraScriptUsage::ParticleEventScript; }
+	bool IsParticleScript() const {	return Usage >= ENiagaraScriptUsage::ParticleSpawnScript && Usage <= ENiagaraScriptUsage::ParticleGPUComputeScript;}
 
 	bool IsNonParticleScript()const { return Usage >= ENiagaraScriptUsage::EmitterSpawnScript; }
 	
@@ -183,13 +354,49 @@
 
 	bool IsCompilable() const { return !IsEmitterSpawnScript() && !IsEmitterUpdateScript(); }
 
+
+	static bool IsGPUScript(ENiagaraScriptUsage Usage) { return Usage == ENiagaraScriptUsage::ParticleGPUComputeScript; }
+	static bool IsParticleSpawnScript(ENiagaraScriptUsage Usage)  { return Usage == ENiagaraScriptUsage::ParticleSpawnScript || Usage == ENiagaraScriptUsage::ParticleSpawnScriptInterpolated; }
+	static bool IsInterpolatedParticleSpawnScript(ENiagaraScriptUsage Usage) { return Usage == ENiagaraScriptUsage::ParticleSpawnScriptInterpolated; }
+	static bool IsParticleUpdateScript(ENiagaraScriptUsage Usage) { return Usage == ENiagaraScriptUsage::ParticleUpdateScript; }
+	static bool IsModuleScript(ENiagaraScriptUsage Usage) { return Usage == ENiagaraScriptUsage::Module; }
+	static bool IsFunctionScript(ENiagaraScriptUsage Usage)	 { return Usage == ENiagaraScriptUsage::Function; }
+	static bool IsDynamicInputScript(ENiagaraScriptUsage Usage)	 { return Usage == ENiagaraScriptUsage::DynamicInput; }
+	static bool IsParticleEventScript(ENiagaraScriptUsage Usage) { return Usage == ENiagaraScriptUsage::ParticleEventScript; }
+	static bool IsParticleScript(ENiagaraScriptUsage Usage)  { return Usage >= ENiagaraScriptUsage::ParticleSpawnScript && Usage <= ENiagaraScriptUsage::ParticleGPUComputeScript; }
+
+	static bool IsNonParticleScript(ENiagaraScriptUsage Usage) { return Usage >= ENiagaraScriptUsage::EmitterSpawnScript; }
+
+	static bool IsSystemSpawnScript(ENiagaraScriptUsage Usage) { return Usage == ENiagaraScriptUsage::SystemSpawnScript; }
+	static bool IsSystemUpdateScript(ENiagaraScriptUsage Usage) { return Usage == ENiagaraScriptUsage::SystemUpdateScript; }
+	static bool IsSystemScript(ENiagaraScriptUsage Usage) { return IsSystemSpawnScript(Usage) || IsSystemUpdateScript(Usage);}
+	static bool IsEmitterSpawnScript(ENiagaraScriptUsage Usage) { return Usage == ENiagaraScriptUsage::EmitterSpawnScript; }
+	static bool IsEmitterUpdateScript(ENiagaraScriptUsage Usage) { return Usage == ENiagaraScriptUsage::EmitterUpdateScript; }
+	static bool IsStandaloneScript(ENiagaraScriptUsage Usage)  { return IsDynamicInputScript(Usage) || IsFunctionScript(Usage) || IsModuleScript(Usage); }
+
+	static bool IsSpawnScript(ENiagaraScriptUsage Usage) { return IsParticleSpawnScript(Usage) || IsEmitterSpawnScript(Usage) || IsSystemSpawnScript(Usage); }
+
+	static bool IsCompilable(ENiagaraScriptUsage Usage)  { return !IsEmitterSpawnScript(Usage) && !IsEmitterUpdateScript(Usage); }
+
+	static bool NIAGARA_API ConvertUsageToGroup(ENiagaraScriptUsage InUsage, ENiagaraScriptGroup& OutGroup);
+
 	NIAGARA_API TArray<ENiagaraScriptUsage> GetSupportedUsageContexts() const;
+	static NIAGARA_API TArray<ENiagaraScriptUsage> GetSupportedUsageContextsForBitmask(int32 InModuleUsageBitmask);
 
 	NIAGARA_API bool CanBeRunOnGpu() const;
+	NIAGARA_API bool IsReadyToRun(ENiagaraSimTarget SimTarget) const;
 
 #if WITH_EDITORONLY_DATA
 	class UNiagaraScriptSourceBase *GetSource() { return Source; }
+	const class UNiagaraScriptSourceBase *GetSource() const  { return Source; }
 	void SetSource(class UNiagaraScriptSourceBase *InSource) { Source = InSource; }
+	void CleanUpOldAndInitializeNewRapidIterationParameters(FString UniqueEmitterName);
+
+	NIAGARA_API FGuid GetBaseChangeID() const;
+	NIAGARA_API ENiagaraScriptCompileStatus GetLastCompileStatus() const;
+	void InvalidateCachedCompileIds();
+
+	NIAGARA_API bool HandleVariableRenames(const TMap<FNiagaraVariable, FNiagaraVariable>& OldToNewVars, const FString& UniqueEmitterName);
 #endif
 
 	//~ Begin UObject interface
@@ -198,21 +405,16 @@
 #if WITH_EDITOR
 	virtual void PostEditChangeProperty(FPropertyChangedEvent& PropertyChangedEvent) override;
 #endif
+	virtual void GetAssetRegistryTags(TArray<FAssetRegistryTag>& OutTags) const override;
 	//~ End UObject interface
 
-	void SetDatainterfaceBufferDescriptors(TArray< FDIBufferDescriptorStore >&InBufferDescriptors)
-	{
-		DIBufferDescriptors = InBufferDescriptors;
-	}
-
-
 	// Infrastructure for GPU compute Shaders
-	NIAGARA_API void CacheResourceShadersForCooking(EShaderPlatform ShaderPlatform, FNiagaraScript* &OutCachedResource);
+	NIAGARA_API void CacheResourceShadersForCooking(EShaderPlatform ShaderPlatform, FNiagaraShaderScript* &OutCachedResource);
 	NIAGARA_API void CacheResourceShadersForRendering(bool bRegenerateId, bool bForceRecompile=false);
 	void BeginCacheForCookedPlatformData(const ITargetPlatform *TargetPlatform);
-	void CacheShadersForResources(EShaderPlatform ShaderPlatform, FNiagaraScript *ResourceToCache, bool bApplyCompletedShaderMapForRendering, bool bForceRecompile = false, bool bCooking=false);
-	FNiagaraScript* AllocateResource();
-	FNiagaraScript *GetRenderThreadScript()
+	void CacheShadersForResources(EShaderPlatform ShaderPlatform, FNiagaraShaderScript *ResourceToCache, bool bApplyCompletedShaderMapForRendering, bool bForceRecompile = false, bool bCooking=false);
+	FNiagaraShaderScript* AllocateResource();
+	FNiagaraShaderScript *GetRenderThreadScript()
 	{
 		return &ScriptResource;
 	}
@@ -237,33 +439,14 @@
 
 	void SetFeatureLevel(ERHIFeatureLevel::Type InFeatureLevel)		{ FeatureLevel = InFeatureLevel; }
 
-
-	// SourceChangeID is updated with the source ID everytime the source graph changes
-	// UniqueID is regenerated in that case, so we have separate Ids for scripts coming from the 
-	// same source (e.g. spawn and update of the same emitter)
-	void InvalidateChangeID()	{ ChangeId.Invalidate(); }
-
-	const FGuid &GetChangeID()	{ return ChangeId; }
-
-	void SetChangeID(FGuid InGuid)
-	{
-		ChangeId = InGuid;
-		
-		// whenever source changes, regenerate so we know to recompile; need this because we may have 
-		//	multiple scripts with the same source, for example spawn and update scripts from one graph
-		UniqueID = FGuid::NewGuid();
-	}
-
 	NIAGARA_API void GenerateStatScopeIDs();
-	ENiagaraScriptCompileStatus GetLastCompileStatus() { return LastCompileStatus; }
 
 	NIAGARA_API bool IsScriptCompilationPending(bool bGPUScript) const;
 	NIAGARA_API bool DidScriptCompilationSucceed(bool bGPUScript) const;
-	NIAGARA_API void InvalidateScript();
 
 #if WITH_EDITORONLY_DATA
+	NIAGARA_API void InvalidateCompileResults();
 	FText GetDescription() { return Description.IsEmpty() ? FText::FromString(GetName()) : Description; }
-	void SetLastCompileStatus(ENiagaraScriptCompileStatus InStatus) { LastCompileStatus = InStatus; }
 
 	/** Makes a deep copy of any script dependencies, including itself.*/
 	NIAGARA_API virtual UNiagaraScript* MakeRecursiveDeepCopy(UObject* DestOuter, TMap<const UObject*, UObject*>& ExistingConversions) const;
@@ -275,17 +458,40 @@
 	NIAGARA_API bool AreScriptAndSourceSynchronized() const;
 
 	/** Ensure that the Script and its source graph are marked out of sync.*/
-	NIAGARA_API void MarkScriptAndSourceDesynchronized();
+	NIAGARA_API void MarkScriptAndSourceDesynchronized(FString Reason);
+	
+	/** Request a synchronous compile for the script, possibly forcing it to compile.*/
+	NIAGARA_API void RequestCompile();
 
-	NIAGARA_API ENiagaraScriptCompileStatus Compile(FString& OutGraphLevelErrorMessages, bool bForce);
+	/** Request an asynchronous compile for the script, possibly forcing it to compile. The output values are the compilation id of the data as well as the async handle to 
+		gather up the results with. bTrulyAsync tells the system whether or not the compile task must be completed on the main thread (mostly used for debugging). The
+		overall function returns whether or not any compiles were actually issued. They will be skipped if none of the data is dirty.*/
+	NIAGARA_API bool RequestExternallyManagedAsyncCompile(const TSharedPtr<FNiagaraCompileRequestDataBase, ESPMode::ThreadSafe>& RequestData, FNiagaraVMExecutableDataId& OutCompileId, uint32& OutAsyncHandle, bool bTrulyAsync);
 
-	NIAGARA_API FNiagaraScriptDebuggerInfo& GetDebuggerInfo() { return DebuggerInfo; }
+	/** Callback issued whenever a compilation successfully happened (even if the results are a script that cannot be executed due to errors)*/
+	NIAGARA_API FOnScriptCompiled& OnVMScriptCompiled();
+	
+	/** External call used to identify the values for a successful VM script compilation. OnVMScriptCompiled will be issued in this case.*/
+	void SetVMCompilationResults(const FNiagaraVMExecutableDataId& InCompileId, FNiagaraVMExecutableData& InScriptVM, FNiagaraCompileRequestDataBase* InRequestData);
 
+	/** In the event where we "merge" we duplicate the changes of the master copy onto the newly cloned copy. This function will synchronize the compiled script 
+		results assuming that the scripts themselves are bound to the same key. This saves looking things up in the DDC. It returns true if successfully synchronized and 
+		false if not.*/
+	NIAGARA_API bool SynchronizeExecutablesWithMaster(const UNiagaraScript* Script, const TMap<FString, FString>& RenameMap);
+
+	NIAGARA_API void SyncAliases(const TMap<FString, FString>& RenameMap);
 #endif
-
+	
 	UFUNCTION()
 	void OnCompilationComplete();
 
+	NIAGARA_API FNiagaraVMExecutableData& GetVMExecutableData() { return CachedScriptVM; }
+	NIAGARA_API const FNiagaraVMExecutableData& GetVMExecutableData() const { return CachedScriptVM; }
+	NIAGARA_API const FNiagaraVMExecutableDataId& GetVMExecutableDataCompilationId() const { return CachedScriptVMId; }
+
+	TArray<UNiagaraParameterCollection*>& GetCachedParameterCollectionReferences() { return CachedParameterCollectionReferences; }
+	TArray<FNiagaraScriptDataInterfaceInfo>& GetCachedDefaultDataInterfaces() { return CachedDefaultDataInterfaces; }
+
 #if STATS
 	const TArray<TStatId>& GetStatScopeIDs()const { return StatScopesIDs; }
 #endif
@@ -294,37 +500,40 @@
 	
 	virtual ~UNiagaraScript();
 
-	NIAGARA_API void SetByteCode(const TArray<uint8>& InByteCode);
-	const TArray<uint8>& GetByteCode() const { return ByteCode; }
-
+	FNiagaraScriptExecutionParameterStore* GetExecutionReadyParameterStore(ENiagaraSimTarget SimTarget);
+	void InvalidateExecutionReadyParameterStores();
 private:
+	bool LegacyCanBeRunOnGpu()const;
 
-	/** Byte code to execute for this system */
-	UPROPERTY()
-	TArray<uint8> ByteCode;
+	UPROPERTY(Transient)
+	FNiagaraScriptExecutionParameterStore ScriptExecutionParamStoreCPU;
 
+	UPROPERTY(Transient)
+	FNiagaraScriptExecutionParameterStore ScriptExecutionParamStoreGPU;
+
 #if WITH_EDITORONLY_DATA
+
 	/** 'Source' data/graphs for this script */
 	UPROPERTY()
 	class UNiagaraScriptSourceBase*	Source;
 	
-	FNiagaraScriptDebuggerInfo DebuggerInfo;
+	/** A multicast delegate which is called whenever the script has been compiled (successfully or not). */
+	FOnScriptCompiled OnVMScriptCompiledDelegate;
+
+	mutable FNiagaraVMExecutableDataId LastReportedVMId;
 #endif
 
-	/** Last known compile status. Lets us determine the latest state of the script byte buffer.*/
-	UPROPERTY()
-	ENiagaraScriptCompileStatus LastCompileStatus;
-
 	/** Adjusted every time that we compile this script. Lets us know that we might differ from any cached versions.*/
 	UPROPERTY()
-	FGuid ChangeId;
+	FNiagaraVMExecutableDataId CachedScriptVMId;
 
-	/** unique ID for this script */
+	/** Adjusted every time ComputeVMCompilationId is called.*/
 	UPROPERTY()
-	FGuid UniqueID;
+	mutable FNiagaraVMExecutableDataId LastGeneratedVMId;
 
-	FNiagaraScript ScriptResource;
-	FNiagaraScript* ScriptResourcesByFeatureLevel[ERHIFeatureLevel::Num];
+	FNiagaraShaderScript ScriptResource;
 
+	FNiagaraShaderScript* ScriptResourcesByFeatureLevel[ERHIFeatureLevel::Num];
+
 	/** Feature level that the shader map is going to be compiled for. */
 	ERHIFeatureLevel::Type FeatureLevel;
@@ -337,9 +546,21 @@
 	TArray<TStatId> StatScopesIDs;
 #endif
 
-#if WITH_EDITOR
+#if WITH_EDITORONLY_DATA
 	/* script resources being cached for cooking. */
-	TMap<const class ITargetPlatform*, FNiagaraScript*> CachedScriptResourcesForCooking;
+	TMap<const class ITargetPlatform*, FNiagaraShaderScript*> CachedScriptResourcesForCooking;
+
+	UPROPERTY(Transient)
+	TArray<UObject*> ActiveCompileRoots;
 #endif
 
+	/** Compiled VM bytecode and data necessary to run this script.*/
+	UPROPERTY()
+	FNiagaraVMExecutableData CachedScriptVM;
+	
+	UPROPERTY()
+	TArray<UNiagaraParameterCollection*> CachedParameterCollectionReferences;
+
+	UPROPERTY()
+	TArray<FNiagaraScriptDataInterfaceInfo> CachedDefaultDataInterfaces;
 };
diff --git a/Engine/Plugins/FX/Niagara/Source/Niagara/Classes/NiagaraScriptExecutionContext.h b/Engine/Plugins/FX/Niagara/Source/Niagara/Classes/NiagaraScriptExecutionContext.h
--- a/Engine/Plugins/FX/Niagara/Source/Niagara/Classes/NiagaraScriptExecutionContext.h
+++ b/Engine/Plugins/FX/Niagara/Source/Niagara/Classes/NiagaraScriptExecutionContext.h
@@ -13,83 +13,11 @@ NiagaraEmitterInstance.h: Niagara emitter simulation class
 #include "NiagaraCollision.h"
 #include "NiagaraEmitterHandle.h"
 #include "NiagaraEmitter.h"
-#include "NiagaraParameterStore.h"
+#include "NiagaraScriptExecutionParameterStore.h"
 #include "NiagaraTypes.h"
+#include "RHIGPUReadback.h"
 
-/** 
-Storage class containing actual runtime buffers to be used by the VM and the GPU. 
-Is not the actual source for any parameter data, rather just the final place it's gathered from various other places ready for execution.
-*/
-class FNiagaraScriptExecutionParameterStore : public FNiagaraParameterStore
-{
-public:
-	FNiagaraScriptExecutionParameterStore();
-	FNiagaraScriptExecutionParameterStore(const FNiagaraParameterStore& Other);
-	FNiagaraScriptExecutionParameterStore& operator=(const FNiagaraParameterStore& Other);
-
-	//TODO: This function can probably go away entirely when we replace the FNiagaraParameters and DataInterface info in the script with an FNiagaraParameterStore.
-	//Special care with prev params and internal params will have to be taken there.
-	void Init(UNiagaraScript* Script, ENiagaraSimTarget SimTarget);
-	void AddScriptParams(UNiagaraScript* Script, ENiagaraSimTarget SimTarget);
-	void CopyCurrToPrev();
-
-	bool AddParameter(const FNiagaraVariable& Param, bool bInitInterfaces=true)
-	{
-		if (FNiagaraParameterStore::AddParameter(Param, bInitInterfaces))
-		{
-			AddPaddedParamSize(Param.GetType(), IndexOf(Param));
-			return true;
-		}
-		return false;
-	}
-
-	bool RemoveParameter(FNiagaraVariable& Param)
-	{
-		check(0);//Not allowed to remove parameters from an execution store as it will adjust the table layout mess up the 
-		return false;
-	}
-
-	void RenameParameter(FNiagaraVariable& Param, FName NewName)
-	{
-		check(0);//Can't rename parameters for an execution store.
-	}
-
-	// Just the external parameters, not previous or internal...
-	uint32 GetExternalParameterSize() { return ParameterSize; }
-
-	// The entire buffer padded out by the required alignment of the types..
-	uint32 GetPaddedParameterSizeInBytes() { return PaddedParameterSize; }
-
-	// Helper that converts the data from the base type array internally into the padded out renderer-ready format.
-	void CopyParameterDataToPaddedBuffer(uint8* InTargetBuffer, uint32 InTargetBufferSizeInBytes);
-
-	void Clear()
-	{
-		Empty();
-		PaddedParameterSize = 0;
-	}
-protected:
-	void AddPaddedParamSize(const FNiagaraTypeDefinition& InParamType, uint32 InOffset);
-
-private:
 
-	/** Size of the parameter data not including prev frame values or internal constants. Allows copying into previous parameter values for interpolated spawn scripts. */
-	int32 ParameterSize;
-
-	uint32 PaddedParameterSize;
-
-	struct FPaddingInfo
-	{
-		FPaddingInfo(uint32 InSrcOffset, uint32 InDestOffset, uint32 InSize) : SrcOffset(InSrcOffset), DestOffset(InDestOffset), Size(InSize){}
-		uint32 SrcOffset;
-		uint32 DestOffset;
-		uint32 Size;
-	};
-
-	static void GenerateLayoutInfoInternal(TArray<FPaddingInfo>& Members, uint32& NextMemberOffset, const UStruct* InSrcStruct, uint32 InSrcOffset);
-	
-	TArray<FPaddingInfo> PaddingInfo;
-};
 
 struct FNiagaraDataSetExecutionInfo
 {
@@ -140,7 +68,7 @@ struct FNiagaraScriptExecutionContext
 	bool Tick(class FNiagaraSystemInstance* Instance, ENiagaraSimTarget SimTarget = ENiagaraSimTarget::CPUSim);
 	void PostTick();
 
-	bool Execute(uint32 NumInstances, TArray<FNiagaraDataSetExecutionInfo>& DataSetInfos);
+	bool Execute(uint32 NumInstances, TArray<FNiagaraDataSetExecutionInfo, TInlineAllocator<8>>& DataSetInfos);
 
 	const TArray<UNiagaraDataInterface*>& GetDataInterfaces()const { return Parameters.GetDataInterfaces(); }
 
@@ -157,32 +85,52 @@ struct FNiagaraComputeExecutionContext
 	FNiagaraComputeExecutionContext()
 		: MainDataSet(nullptr)
 		, SpawnRateInstances(0)
-		, BurstInstances(0)
 		, EventSpawnTotal(0)
 		, SpawnScript(nullptr)
 		, UpdateScript(nullptr)
+		, GPUScript(nullptr)
 		, RTUpdateScript(0)
-		, RTSpawnScript(0)
+		, RTSpawnScript(0)	
+		, RTGPUScript(0)
+		, GPUDataReadback(nullptr)
+		, AccumulatedSpawnRate(0)
+		, NumIndicesPerInstance(0)
+	{
+	}
+
+	void Reset()
 	{
+		AccumulatedSpawnRate = 0;
+		if (GPUDataReadback)
+		{
+			delete GPUDataReadback;
+		}
+		GPUDataReadback = nullptr;
 	}
 
-	void InitParams(UNiagaraScript* InSpawnScript, UNiagaraScript *InUpdateScript, ENiagaraSimTarget SimTarget)
+	void InitParams(UNiagaraScript* InGPUComputeScript, UNiagaraScript* InSpawnScript, UNiagaraScript *InUpdateScript, ENiagaraSimTarget SimTarget)
 	{
-		CombinedParamStore.Init(InSpawnScript, SimTarget);
-		CombinedParamStore.AddScriptParams(InUpdateScript, SimTarget);
+		CombinedParamStore.InitFromOwningContext(InGPUComputeScript, SimTarget, true);
+
+		GPUScript = InGPUComputeScript;
 		SpawnScript = InSpawnScript;
 		UpdateScript = InUpdateScript;
-		
-		FNiagaraShader *Shader = InSpawnScript->GetRenderThreadScript()->GetShaderGameThread();
+
+#if DO_CHECK
+		FNiagaraShader *Shader = InGPUComputeScript->GetRenderThreadScript()->GetShaderGameThread();
+		DIParamInfo.Empty();
 		if (Shader)
 		{
-			DataInterfaceDescriptors = Shader->GetDIBufferDescriptors();
+			for (FNiagaraDataInterfaceParamRef& DIParams : Shader->GetDIParameters())
+			{
+				DIParamInfo.Add(DIParams.ParameterInfo);
+			}
 		}
 		else
 		{
-			DataInterfaceDescriptors = InSpawnScript->GetRenderThreadScript()->GetDataInterfaceBufferDescriptors();
+			DIParamInfo = InGPUComputeScript->GetRenderThreadScript()->GetDataInterfaceParamInfo();
 		}
-
+#endif
 	}
 
 	void DirtyDataInterfaces()
@@ -194,24 +142,25 @@ struct FNiagaraComputeExecutionContext
 	{
 		if (CombinedParamStore.GetInterfacesDirty())
 		{
+#if DO_CHECK
 			const TArray<UNiagaraDataInterface*> &DataInterfaces = CombinedParamStore.GetDataInterfaces();
 			// We must make sure that the data interfaces match up between the original script values and our overrides...
-			if (DataInterfaceDescriptors.Num() != DataInterfaces.Num())
+			if (DIParamInfo.Num() != DataInterfaces.Num())
 			{
 				UE_LOG(LogNiagara, Warning, TEXT("Mismatch between Niagara GPU Execution Context data interfaces and those in its script!"));
 				return false;
 			}
 
-			//Fill the instance data table.
-			if (ParentSystemInstance)
+			for (int32 i=0; i<DIParamInfo.Num(); ++i)
 			{
-				for (int32 i = 0; i < DataInterfaces.Num(); i++)
+				FString UsedClassName = DataInterfaces[i]->GetClass()->GetName();
+				if (DIParamInfo[i].DIClassName != UsedClassName)
 				{
-					UNiagaraDataInterface* Interface = DataInterfaces[i];
-					// setup GPU interface; TODO: should probably move this to the GPU execution context
-					Interface->SetupBuffers(DataInterfaceDescriptors[i]);
+					UE_LOG(LogNiagara, Warning, TEXT("Mismatched class between Niagara GPU Execution Context data interfaces and those in its script!\nIndex:%d\nShader:%s\nScript:%s")
+						, i, *DIParamInfo[i].DIClassName, *UsedClassName);
 				}
 			}
+#endif
 
 			CombinedParamStore.Tick();
 		}
@@ -219,22 +168,30 @@ struct FNiagaraComputeExecutionContext
 		return true;
 	}
 
+	const TArray<FNiagaraEventScriptProperties> &GetEventHandlers() const { return EventHandlerScriptProps; }
+
 	class FNiagaraDataSet *MainDataSet;
 	TArray<FNiagaraDataSet*>UpdateEventWriteDataSets;
 	TArray<FNiagaraEventScriptProperties> EventHandlerScriptProps;
 	TArray<FNiagaraDataSet*> EventSets;
 	uint32 SpawnRateInstances;
-	uint32 BurstInstances;
 
 	TArray<int32> EventSpawnCounts;
 	uint32 EventSpawnTotal;
-	UNiagaraScript *SpawnScript;
-	UNiagaraScript *UpdateScript;
-	class FNiagaraScript *RTUpdateScript;
-	class FNiagaraScript *RTSpawnScript;
-	TArray<FNiagaraScriptDataInterfaceInfo> UpdateInterfaces;
-	TArray<uint8, TAlignedHeapAllocator<16>> SpawnParams;		// RT side copy of the parameter data
+	UNiagaraScript* SpawnScript;
+	UNiagaraScript* UpdateScript;
+	UNiagaraScript* GPUScript;
+	class FNiagaraShaderScript*  RTUpdateScript;
+	class FNiagaraShaderScript*  RTSpawnScript;
+	class FNiagaraShaderScript*  RTGPUScript;
+	TArray<uint8, TAlignedHeapAllocator<16>> ParamData_RT;		// RT side copy of the parameter data
 	FNiagaraScriptExecutionParameterStore CombinedParamStore;
-	TArray< FDIBufferDescriptorStore >  DataInterfaceDescriptors;
 	static uint32 TickCounter;
+#if DO_CHECK
+	TArray< FNiagaraDataInterfaceGPUParamInfo >  DIParamInfo;
+#endif
+
+	FRHIGPUMemoryReadback *GPUDataReadback;
+	uint32 AccumulatedSpawnRate;
+	uint32 NumIndicesPerInstance;	// how many vtx indices per instance the renderer is going to have for its draw call
 };
diff --git a/Engine/Plugins/FX/Niagara/Source/Niagara/Classes/NiagaraScriptSourceBase.h b/Engine/Plugins/FX/Niagara/Source/Niagara/Classes/NiagaraScriptSourceBase.h
--- a/Engine/Plugins/FX/Niagara/Source/Niagara/Classes/NiagaraScriptSourceBase.h
+++ b/Engine/Plugins/FX/Niagara/Source/Niagara/Classes/NiagaraScriptSourceBase.h
@@ -20,6 +20,47 @@ struct EditorExposedVectorCurveConstant
 	class UCurveVector *Value;
 };
 
+/** External reference to the compile request data generated.*/
+class FNiagaraCompileRequestDataBase
+{
+public:
+	virtual ~FNiagaraCompileRequestDataBase() {}
+	virtual bool GatherPreCompiledVariables(const FString& InNamespaceFilter, TArray<FNiagaraVariable>& OutVars) = 0;
+	virtual void GetReferencedObjects(TArray<UObject*>& Objects) = 0;
+	virtual const TMap<FName, UNiagaraDataInterface*>& GetObjectNameMap() = 0;
+	virtual int32 GetDependentRequestCount() const = 0;
+	virtual TSharedPtr<FNiagaraCompileRequestDataBase, ESPMode::ThreadSafe> GetDependentRequest(int32 Index) = 0;
+	virtual FName ResolveEmitterAlias(FName VariableName) const = 0;
+};
+
+class FNiagaraCompileOptions
+{
+public:
+	FNiagaraCompileOptions() : TargetUsage(ENiagaraScriptUsage::Function), TargetUsageBitmask(0)
+	{
+	}
+
+	FNiagaraCompileOptions(ENiagaraScriptUsage InTargetUsage, FGuid InTargetUsageId, int32 InTargetUsageBitmask,  const FString& InPathName, const FString& InFullName, const FString& InName)
+		: TargetUsage(InTargetUsage), TargetUsageId(InTargetUsageId), PathName(InPathName), FullName(InFullName), Name(InName), TargetUsageBitmask(InTargetUsageBitmask)
+	{
+	}
+
+	const FString& GetFullName() const { return FullName; }
+	const FString& GetName() const { return Name; }
+	const FString& GetPathName() const { return PathName; }
+	int32 GetTargetUsageBitmask() const { return TargetUsageBitmask; }
+
+	ENiagaraScriptUsage TargetUsage;
+	FGuid TargetUsageId;
+	FString PathName;
+	FString FullName;
+	FString Name;
+	int32 TargetUsageBitmask;
+	TArray<FString> AdditionalDefines;
+};
+
+struct FNiagaraParameterStore;
+
 /** Runtime data for a Niagara system */
 UCLASS(MinimalAPI)
 class UNiagaraScriptSourceBase : public UObject
@@ -42,37 +83,31 @@ class UNiagaraScriptSourceBase : public UObject
 	virtual void SubsumeExternalDependencies(TMap<const UObject*, UObject*>& ExistingConversions) {}
 
 	/** Enforce that the source graph is now out of sync with the script.*/
-	virtual void MarkNotSynchronized() {}
+	virtual void MarkNotSynchronized(FString Reason) {}
 
 	virtual FGuid GetChangeID() { return FGuid(); };
 
-	/** Have we called PreCompile on this source previously?*/
-	virtual bool IsPreCompiled() const { return false; }
+	virtual void ComputeVMCompilationId(struct FNiagaraVMExecutableDataId& Id, ENiagaraScriptUsage InUsage, const FGuid& InUsageId) const {};
 	
 	/** Cause the source to build up any internal variables that will be useful in the compilation process.*/
-	virtual void PreCompile(UNiagaraEmitter* Emitter, bool bClearErrors = true) {}
+	virtual TSharedPtr<FNiagaraCompileRequestDataBase, ESPMode::ThreadSafe> PreCompile(UNiagaraEmitter* Emitter, const TArray<FNiagaraVariable>& EncounterableVariables, TArray<TSharedPtr<FNiagaraCompileRequestDataBase, ESPMode::ThreadSafe>>& ReferencedCompileRequests, bool bClearErrors = true) { return nullptr; }
 	
-	/** If the user previously pre-compiled the source, dig through the data to find any variables defined. */
-	virtual bool GatherPreCompiledVariables(const FString& InNamespaceFilter, TArray<FNiagaraVariable>& OutVars) { return false; }
-
-	/** Implements compilation of a Niagara script.*/
-	virtual ENiagaraScriptCompileStatus Compile(UNiagaraScript* ScriptOwner, FString& OutGraphLevelErrorMessages) { return ENiagaraScriptCompileStatus::NCS_Unknown; }
-
-	/** Do any cleanup to return this data source to its original state before precompilation. Note that this must be called after the "compile" operation.*/
-	virtual void PostCompile() {}
-
 	/** 
 	 * Allows the derived editor only script source to handle a post load requested by an owning emitter. 
 	 * @param OwningEmitter The emitter requesting the post load.
-	 * @returns True if the emitter will need to be recompiled, otherwise false.
 	 */
-	virtual bool PostLoadFromEmitter(UNiagaraEmitter& OwningEmitter) { return false; }
+	virtual void PostLoadFromEmitter(UNiagaraEmitter& OwningEmitter) { }
 
 	/** Adds a module if it isn't already in the graph. If the module isn't found bOutFoundModule will be false. If it is found and it did need to be added, the function returns true. If it already exists, it returns false. */
 	NIAGARA_API virtual bool AddModuleIfMissing(FString ModulePath, ENiagaraScriptUsage Usage, bool& bOutFoundModule) { bOutFoundModule = false; return false; }
+
 #if WITH_EDITOR
+	virtual void CleanUpOldAndInitializeNewRapidIterationParameters(FString UniqueEmitterName, ENiagaraScriptUsage ScriptUsage, FGuid ScriptUsageId, FNiagaraParameterStore& RapidIterationParameters) const { checkf(false, TEXT("Not implemented")); }
+
 	FOnChanged& OnChanged() { return OnChangedDelegate; }
 
+	virtual void InvalidateCachedCompileIds() {}
+
 protected:
 	FOnChanged OnChangedDelegate;
 #endif
diff --git a/Engine/Plugins/FX/Niagara/Source/Niagara/Classes/NiagaraSystem.h b/Engine/Plugins/FX/Niagara/Source/Niagara/Classes/NiagaraSystem.h
--- a/Engine/Plugins/FX/Niagara/Source/Niagara/Classes/NiagaraSystem.h
+++ b/Engine/Plugins/FX/Niagara/Source/Niagara/Classes/NiagaraSystem.h
@@ -12,6 +12,10 @@
 #include "NiagaraParameterCollection.h"
 #include "NiagaraSystem.generated.h"
 
+#if WITH_EDITORONLY_DATA
+class UNiagaraEditorDataBase;
+#endif
+
 USTRUCT()
 struct FNiagaraEmitterSpawnAttributes
 {
@@ -21,6 +25,36 @@ struct FNiagaraEmitterSpawnAttributes
 	TArray<FName> SpawnAttributes;
 };
 
+
+
+USTRUCT()
+struct FEmitterCompiledScriptPair
+{
+	GENERATED_USTRUCT_BODY()
+	
+	bool bResultsReady;
+	UNiagaraEmitter* Emitter;
+	UNiagaraScript* CompiledScript;
+	uint32 PendingDDCID;
+	FNiagaraVMExecutableDataId CompileId;
+	TSharedPtr<FNiagaraVMExecutableData> CompileResults;
+};
+
+USTRUCT()
+struct FNiagaraSystemCompileRequest
+{
+	GENERATED_USTRUCT_BODY()
+		
+	double StartTime;
+
+	UPROPERTY()
+	TArray<UObject*> RootObjects;
+
+	TArray<FEmitterCompiledScriptPair> EmitterCompiledScriptPairs;
+	
+	TMap<UNiagaraScript*, TSharedPtr<FNiagaraCompileRequestDataBase, ESPMode::ThreadSafe> > MappedData;
+};
+
 /** Container for multiple emitters that combine together to create a particle system effect.*/
 UCLASS(BlueprintType)
 class NIAGARA_API UNiagaraSystem : public UObject
@@ -28,12 +62,19 @@ class NIAGARA_API UNiagaraSystem : public UObject
 	GENERATED_UCLASS_BODY()
 
 public:
+#if WITH_EDITOR
+	DECLARE_MULTICAST_DELEGATE_OneParam(FOnSystemCompiled, UNiagaraSystem*);
+#endif
+
 	//~ UObject interface
 	void PostInitProperties();
 	void Serialize(FArchive& Ar)override;
-	virtual void PostLoad() override;
+	virtual void PostLoad() override; 
+	virtual void BeginDestroy() override;
+	virtual void PreSave(const class ITargetPlatform * TargetPlatform) override;
 #if WITH_EDITOR
-	virtual void PostEditChangeProperty(struct FPropertyChangedEvent& PropertyChangedEvent) override;
+	virtual void PostEditChangeProperty(struct FPropertyChangedEvent& PropertyChangedEvent) override; 
+	virtual void BeginCacheForCookedPlatformData(const ITargetPlatform *TargetPlatform) override;
 #endif
 
 	/** Gets an array of the emitter handles. */
@@ -88,11 +129,22 @@ public:
 
 
 	/** Gets the System script which is used to populate the System parameters and parameter bindings. */
-	UNiagaraScript* GetSystemSpawnScript(bool bSolo=false);
-	UNiagaraScript* GetSystemUpdateScript(bool bSolo=false);
+	UNiagaraScript* GetSystemSpawnScript();
+	UNiagaraScript* GetSystemUpdateScript();
 
 	bool IsReadyToRun() const;
 
+	/** Are there any pending compile requests?*/
+	bool HasOutstandingCompilationRequests() const;
+
+	/** Returns whether this system has to be run in solo or not. */
+	bool IsSolo()const;
+
+	FORCEINLINE bool NeedsWarmup()const { return WarmupTickCount > 0 && WarmupTickDelta > SMALL_NUMBER; }
+	FORCEINLINE float GetWarmupTime()const { return WarmupTime; }
+	FORCEINLINE int32 GetWarmupTickCount()const { return WarmupTickCount; }
+	FORCEINLINE float GetWarmupTickDelta()const { return WarmupTickDelta; }
+
 #if WITH_EDITORONLY_DATA
 	/** Called to query whether or not this emitter is referenced as the source to any emitter handles for this System.*/
 	bool ReferencesSourceEmitter(UNiagaraEmitter& Emitter);
@@ -106,20 +158,26 @@ public:
 	/** Removes the system's rapid iteration parameters for a specific emitter. */
 	void RemoveSystemParametersForEmitter(const FNiagaraEmitterHandle& EmitterHandle);
 
-	/** Recompile all emitters and the System script associated with this System.*/
-	void Compile(bool bForce);
+	/** Request that any dirty scripts referenced by this system be compiled.*/
+	bool RequestCompile(bool bForce);
+
+	/** If we have a pending compile request, is it done with yet? */
+	bool PollForCompilationComplete();
 
-	/** Compile just the System script associated with this System.*/
-	void CompileScripts(TArray<ENiagaraScriptCompileStatus>& OutScriptStatuses, TArray<FString>& OutGraphLevelErrorMessages, TArray<FString>& PathNames, TArray<UNiagaraScript*>& Scripts, bool bForce);
+	/** */
+	void WaitForCompilationComplete();
+
+	/** Delegate called when the system's dependencies have all been compiled.*/
+	FOnSystemCompiled& OnSystemCompiled();
 
 	/** Gets editor specific data stored with this system. */
-	UObject* GetEditorData();
+	UNiagaraEditorDataBase* GetEditorData();
 
 	/** Gets editor specific data stored with this system. */
-	const UObject* GetEditorData() const;
+	const UNiagaraEditorDataBase* GetEditorData() const;
 
 	/** Sets editor specific data stored with this system. */
-	void SetEditorData(UObject* InEditorData);
+	void SetEditorData(UNiagaraEditorDataBase* InEditorData);
 
 	/** Internal: The thumbnail image.*/
 	UPROPERTY()
@@ -140,7 +198,11 @@ public:
 
 	bool UsesCollection(const UNiagaraParameterCollection* Collection)const;
 #if WITH_EDITORONLY_DATA
-	bool UsesScript(const UNiagaraScript* Script)const;
+	bool UsesEmitter(const UNiagaraEmitter* Emitter) const;
+	bool UsesScript(const UNiagaraScript* Script)const; 
+	void InvalidateCachedCompileIds();
+
+	static void RequestCompileForEmitter(UNiagaraEmitter* InEmitter);
 #endif
 
 	FORCEINLINE UNiagaraParameterCollectionInstance* GetParameterCollectionOverride(UNiagaraParameterCollection* Collection)
@@ -153,13 +215,20 @@ public:
 
 		return Found ? *Found : nullptr;
 	}
+	
+	UPROPERTY(EditAnywhere, Category = "Debug")
+	bool bDumpDebugSystemInfo;
 
+	UPROPERTY(EditAnywhere, Category = "Debug")
+	bool bDumpDebugEmitterInfo;
 
 private:
 #if WITH_EDITORONLY_DATA
 	INiagaraModule::FMergeEmitterResults MergeChangesForEmitterHandle(FNiagaraEmitterHandle& EmitterHandle);
+	bool QueryCompileComplete(bool bWait, bool bDoPost, bool bDoNotApply = false);
 #endif
 
+	void DetermineIfSolo();
 protected:
 
 	/** Handles to the emitter this System will simulate. */
@@ -169,6 +238,9 @@ protected:
 	UPROPERTY(EditAnywhere, Category="System")
 	TArray<UNiagaraParameterCollectionInstance*> ParameterCollectionOverrides;
 
+	UPROPERTY(Transient)
+	TArray<FNiagaraSystemCompileRequest> ActiveCompilations;
+
 // 	/** Category of this system. */
 // 	UPROPERTY(EditAnywhere, Category = System)
 // 	UNiagaraSystemCategory* Category;
@@ -183,14 +255,6 @@ protected:
 	UPROPERTY()
 	UNiagaraScript* SystemUpdateScript;
 
-	/** Spawn script compiled to be run individually on a single instance of the system rather than batched as the main spawn script */
-	UPROPERTY()
-	UNiagaraScript* SystemSpawnScriptSolo;
-
-	/** Update script compiled to be run individually on a single instance of the system rather than batched as the main spawn script */
-	UPROPERTY()
-	UNiagaraScript* SystemUpdateScriptSolo;
-
 	/** Attribute names in the data set that are driving each emitter's spawning. */
 	UPROPERTY()
 	TArray<FNiagaraEmitterSpawnAttributes> EmitterSpawnAttributes;
@@ -203,13 +267,32 @@ protected:
 
 	/** Data used by the editor to maintain UI state etc.. */
 	UPROPERTY()
-	UObject* EditorData;
+	UNiagaraEditorDataBase* EditorData;
 
 	bool bIsolateEnabled;
+
+	/** A multicast delegate which is called whenever the script has been compiled (successfully or not). */
+	FOnSystemCompiled OnSystemCompiledDelegate;
 #endif
 
 	UPROPERTY(EditAnywhere, Category = Performance, meta = (ToolTip = "Auto-deactivate system if all emitters are determined to not spawn particles again, regardless of lifetime."))
 	bool bAutoDeactivate;
 
+	/** Warm up time in seconds. Used to calculate WarmupTickCount. Rounds down to the nearest multiple of WarmupTickDelta. */
+	UPROPERTY(EditAnywhere, Category = Warmup)
+	float WarmupTime;
+
+	/** Number of ticks to process for warmup. You can set by this or by time via WarmupTime. */
+	UPROPERTY(EditAnywhere, Category = Warmup)
+	int32 WarmupTickCount;
+
+	/** Delta time to use for warmup ticks. */
+	UPROPERTY(EditAnywhere, Category = Warmup)
+	float WarmupTickDelta;
+
+
+	UPROPERTY()
+	uint32 bSolo : 1;
+
 	void InitEmitterSpawnAttributes();
 };
diff --git a/Engine/Plugins/FX/Niagara/Source/Niagara/Niagara.Build.cs b/Engine/Plugins/FX/Niagara/Source/Niagara/Niagara.Build.cs
--- a/Engine/Plugins/FX/Niagara/Source/Niagara/Niagara.Build.cs
+++ b/Engine/Plugins/FX/Niagara/Source/Niagara/Niagara.Build.cs
@@ -8,18 +8,21 @@ public class Niagara : ModuleRules
     {
         PrivateDependencyModuleNames.AddRange(
             new string[] {
+                "NiagaraCore",
                 "NiagaraShader",
                 "Core",
                 "Engine",
 				"RenderCore",
                 "UtilityShaders",
-                "ShaderCore"
+                "ShaderCore",
+                "TimeManagement"
             }
         );
 
 
         PublicDependencyModuleNames.AddRange(
             new string[] {
+                "NiagaraCore",
                 "NiagaraShader",
                 "MovieScene",
 				"MovieSceneTracks",
@@ -35,6 +38,7 @@ public class Niagara : ModuleRules
 
         PrivateIncludePaths.AddRange(
             new string[] {
+                "Niagara/Private",
             })
         ;
 
diff --git a/Engine/Plugins/FX/Niagara/Source/Niagara/Private/MovieScene/MovieSceneNiagaraSystemTrack.cpp b/Engine/Plugins/FX/Niagara/Source/Niagara/Private/MovieScene/MovieSceneNiagaraSystemTrack.cpp
new file mode 100644
--- /dev/null
+++ b/Engine/Plugins/FX/Niagara/Source/Niagara/Private/MovieScene/MovieSceneNiagaraSystemTrack.cpp
@@ -0,0 +1,53 @@
+// Copyright 1998-2017 Epic Games, Inc. All Rights Reserved.
+
+#include "MovieScene/MovieSceneNiagaraSystemTrack.h"
+#include "CoreMinimal.h"
+#include "MovieScene.h"
+#include "MovieScene/MovieSceneNiagaraSystemSpawnSection.h"
+#include "MovieScene/MovieSceneNiagaraSystemTrackTemplate.h"
+#include "Evaluation/MovieSceneEvalTemplate.h"
+#include "Evaluation/MovieSceneEvaluationTrack.h"
+#include "Compilation/IMovieSceneTemplateGenerator.h"
+#include "Compilation/MovieSceneSegmentCompiler.h"
+#include "Compilation/MovieSceneCompilerRules.h"
+
+
+UMovieSceneSection* UMovieSceneNiagaraSystemTrack::CreateNewSection()
+{
+	return NewObject<UMovieSceneNiagaraSystemSpawnSection>(this);
+}
+
+FMovieSceneEvalTemplatePtr UMovieSceneNiagaraSystemTrack::CreateTemplateForSection(const UMovieSceneSection& InSection) const
+{
+	return FMovieSceneEvalTemplatePtr();
+}
+
+void UMovieSceneNiagaraSystemTrack::PostCompile(FMovieSceneEvaluationTrack& OutTrack, const FMovieSceneTrackCompilerArgs& Args) const
+{
+	UMovieSceneSection* const* SpawnSectionPtr = Sections.FindByPredicate([](const UMovieSceneSection* Section) { return Section->GetClass() == UMovieSceneNiagaraSystemSpawnSection::StaticClass(); });
+	if (SpawnSectionPtr != nullptr)
+	{
+		UMovieSceneNiagaraSystemSpawnSection* SpawnSection = CastChecked<UMovieSceneNiagaraSystemSpawnSection>(*SpawnSectionPtr);
+		UMovieScene* ParentMovieScene = GetTypedOuter<UMovieScene>();
+		OutTrack.SetTrackImplementation(FMovieSceneNiagaraSystemTrackImplementation(SpawnSection->GetInclusiveStartFrame(), SpawnSection->GetExclusiveEndFrame()));
+	}
+}
+
+struct FNiagaraSystemTrackSegmentBlender : public FMovieSceneTrackSegmentBlender
+{
+	FNiagaraSystemTrackSegmentBlender()
+	{
+		bAllowEmptySegments = true;
+		bCanFillEmptySpace = true;
+	}
+
+	virtual TOptional<FMovieSceneSegment> InsertEmptySpace(const TRange<FFrameNumber>& Range, const FMovieSceneSegment* PreviousSegment, const FMovieSceneSegment* NextSegment) const override
+	{
+		return MovieSceneSegmentCompiler::EvaluateNearestSegment(Range, PreviousSegment, NextSegment);
+	}
+};
+
+FMovieSceneTrackSegmentBlenderPtr UMovieSceneNiagaraSystemTrack::GetTrackSegmentBlender() const
+{
+	return FNiagaraSystemTrackSegmentBlender();
+}
diff --git a/Engine/Plugins/FX/Niagara/Source/Niagara/Private/MovieScene/MovieSceneNiagaraSystemTrackTemplate.cpp b/Engine/Plugins/FX/Niagara/Source/Niagara/Private/MovieScene/MovieSceneNiagaraSystemTrackTemplate.cpp
new file mode 100644
--- /dev/null
+++ b/Engine/Plugins/FX/Niagara/Source/Niagara/Private/MovieScene/MovieSceneNiagaraSystemTrackTemplate.cpp
@@ -0,0 +1,160 @@
+// Copyright 1998-2017 Epic Games, Inc. All Rights Reserved.
+
+#include "MovieSceneNiagaraSystemTrackTemplate.h"
+#include "MovieSceneExecutionToken.h"
+#include "NiagaraComponent.h"
+#include "NiagaraSystemInstance.h"
+#include "IMovieScenePlayer.h"
+
+#include "Evaluation/MovieSceneEvaluationTemplateInstance.h"
+#include "MovieSceneSequence.h"
+#include "MovieScene.h"
+
+struct FPreAnimatedNiagaraComponentToken : IMovieScenePreAnimatedToken
+{
+	FPreAnimatedNiagaraComponentToken(bool bInComponentIsActive, bool bInComponentForceSolo, bool bInComponentRenderingEnabled, TOptional<ENiagaraExecutionState> InSystemInstanceExecutionState, ENiagaraAgeUpdateMode InComponentAgeUpdateMode, float InComponentSeekDelta, float InComponentDesiredAge)
+		: bComponentIsActive(bInComponentIsActive)
+		, bComponentForceSolo(bInComponentForceSolo)
+		, bComponentRenderingEnabled(bInComponentRenderingEnabled)
+		, SystemInstanceExecutionState(InSystemInstanceExecutionState)
+		, ComponentAgeUpdateMode(InComponentAgeUpdateMode)
+		, ComponentSeekDelta(InComponentSeekDelta)
+		, ComponentDesiredAge(InComponentDesiredAge)
+	{ }
+
+	virtual void RestoreState(UObject& InObject, IMovieScenePlayer& InPlayer)
+	{
+		UNiagaraComponent* NiagaraComponent = CastChecked<UNiagaraComponent>(&InObject);
+		FNiagaraSystemInstance* SystemInstance = NiagaraComponent->GetSystemInstance();
+		if (bComponentIsActive)
+		{
+			NiagaraComponent->Activate();
+		}
+		else
+		{
+			if (SystemInstance != nullptr)
+			{
+				SystemInstance->Reset(FNiagaraSystemInstance::EResetMode::ResetSystem);
+			}
+			NiagaraComponent->Deactivate();
+		}
+		NiagaraComponent->SetForceSolo(bComponentForceSolo);
+		NiagaraComponent->SetRenderingEnabled(bComponentRenderingEnabled);
+		if (SystemInstance != nullptr && SystemInstanceExecutionState.IsSet())
+		{
+			SystemInstance->SetRequestedExecutionState(SystemInstanceExecutionState.GetValue());
+		}
+		NiagaraComponent->SetAgeUpdateMode(ComponentAgeUpdateMode);
+		NiagaraComponent->SetSeekDelta(ComponentSeekDelta);
+		NiagaraComponent->SetDesiredAge(ComponentDesiredAge);
+	}
+
+	bool bComponentIsActive;
+	bool bComponentForceSolo;
+	bool bComponentRenderingEnabled;
+	TOptional<ENiagaraExecutionState> SystemInstanceExecutionState;
+	ENiagaraAgeUpdateMode ComponentAgeUpdateMode;
+	float ComponentSeekDelta;
+	float ComponentDesiredAge;
+};
+
+struct FPreAnimatedNiagaraComponentTokenProducer : IMovieScenePreAnimatedTokenProducer
+{
+	virtual IMovieScenePreAnimatedTokenPtr CacheExistingState(UObject& InObject) const override
+	{
+		UNiagaraComponent* NiagaraComponent = CastChecked<UNiagaraComponent>(&InObject);
+		FNiagaraSystemInstance* SystemInstance = NiagaraComponent->GetSystemInstance();
+		return FPreAnimatedNiagaraComponentToken(
+			NiagaraComponent->IsActive(),
+			NiagaraComponent->GetForceSolo(),
+			NiagaraComponent->GetRenderingEnabled(),
+			SystemInstance != nullptr ? SystemInstance->GetRequestedExecutionState() : TOptional<ENiagaraExecutionState>(),
+			NiagaraComponent->GetAgeUpdateMode(),
+			NiagaraComponent->GetSeekDelta(),
+			NiagaraComponent->GetDesiredAge());
+	}
+};
+
+struct FNiagaraSystemUpdateDesiredAgeExecutionToken : IMovieSceneExecutionToken
+{
+	FNiagaraSystemUpdateDesiredAgeExecutionToken(FFrameNumber InSpawnSectionStartFrame, FFrameNumber InSpawnSectionEndFrame)
+		: SpawnSectionStartFrame(InSpawnSectionStartFrame)
+		, SpawnSectionEndFrame(InSpawnSectionEndFrame)
+		
+	{}
+
+	virtual void Execute(const FMovieSceneContext& Context, const FMovieSceneEvaluationOperand& Operand, FPersistentEvaluationData& PersistentData, IMovieScenePlayer& Player) override
+	{
+		for (TWeakObjectPtr<> Object : Player.FindBoundObjects(Operand))
+		{
+			UObject* ObjectPtr = Object.Get();
+			UNiagaraComponent* NiagaraComponent = Cast<UNiagaraComponent>(ObjectPtr);
+
+			static FMovieSceneAnimTypeID TypeID = TMovieSceneAnimTypeID<FNiagaraSystemUpdateDesiredAgeExecutionToken, 0>();
+			Player.SavePreAnimatedState(*NiagaraComponent, TypeID, FPreAnimatedNiagaraComponentTokenProducer(), PersistentData.GetTrackKey());
+
+			NiagaraComponent->SetForceSolo(true);
+			NiagaraComponent->SetAgeUpdateMode(ENiagaraAgeUpdateMode::DesiredAge);
+
+			UMovieSceneSequence* MovieSceneSequence = Player.GetEvaluationTemplate().GetSequence(Operand.SequenceID);
+			if (MovieSceneSequence != nullptr)
+			{
+				UMovieScene* MovieScene = MovieSceneSequence->GetMovieScene();
+				if (MovieScene != nullptr)
+				{
+					NiagaraComponent->SetSeekDelta((float)MovieScene->GetDisplayRate().Denominator / MovieScene->GetDisplayRate().Numerator);
+				}
+			}
+
+			if (NiagaraComponent->IsActive() == false || NiagaraComponent->GetSystemInstance() == nullptr)
+			{
+				NiagaraComponent->Activate();
+			}
+
+			if (Context.GetTime() < SpawnSectionStartFrame)
+			{
+				FNiagaraSystemInstance* SystemInstance = NiagaraComponent->GetSystemInstance();
+				if (SystemInstance != nullptr)
+				{
+					SystemInstance->SetRequestedExecutionState(ENiagaraExecutionState::Complete);
+				}
+			}
+			else if (Context.GetTime() < SpawnSectionEndFrame)
+			{
+				FNiagaraSystemInstance* SystemInstance = NiagaraComponent->GetSystemInstance();
+				if (SystemInstance != nullptr)
+				{
+					SystemInstance->SetRequestedExecutionState(ENiagaraExecutionState::Active);
+				}
+			}
+			else
+			{
+				FNiagaraSystemInstance* SystemInstance = NiagaraComponent->GetSystemInstance();
+				if (SystemInstance != nullptr)
+				{
+					SystemInstance->SetRequestedExecutionState(ENiagaraExecutionState::Inactive);
+				}
+			}
+
+			bool bRenderingEnabled = Context.IsPreRoll() == false;
+			NiagaraComponent->SetRenderingEnabled(bRenderingEnabled);
+			NiagaraComponent->SetDesiredAge(Context.GetFrameRate().AsSeconds(Context.GetTime() - SpawnSectionStartFrame));
+		}
+	}
+
+	FFrameNumber SpawnSectionStartFrame;
+	FFrameNumber SpawnSectionEndFrame;
+};
+
+FMovieSceneNiagaraSystemTrackImplementation::FMovieSceneNiagaraSystemTrackImplementation(FFrameNumber InSpawnSectionStartFrame, FFrameNumber InSpawnSectionEndFrame)
+	: SpawnSectionStartFrame(InSpawnSectionStartFrame)
+	, SpawnSectionEndFrame(InSpawnSectionEndFrame)
+{
+	EnableOverrides(CustomEvaluateFlag);
+}
+
+void FMovieSceneNiagaraSystemTrackImplementation::Evaluate(const FMovieSceneEvaluationTrack& Track, FMovieSceneSegmentIdentifier SegmentID, const FMovieSceneEvaluationOperand& Operand, const FMovieSceneContext& Context, const FPersistentEvaluationData& PersistentData, FMovieSceneExecutionTokens& ExecutionTokens) const
+{
+	ExecutionTokens.SetContext(Context);
+	ExecutionTokens.Add(FNiagaraSystemUpdateDesiredAgeExecutionToken(SpawnSectionStartFrame, SpawnSectionEndFrame));
+}
\ No newline at end of file
diff --git a/Engine/Plugins/FX/Niagara/Source/Niagara/Private/MovieScene/MovieSceneNiagaraSystemTrackTemplate.h b/Engine/Plugins/FX/Niagara/Source/Niagara/Private/MovieScene/MovieSceneNiagaraSystemTrackTemplate.h
new file mode 100644
--- /dev/null
+++ b/Engine/Plugins/FX/Niagara/Source/Niagara/Private/MovieScene/MovieSceneNiagaraSystemTrackTemplate.h
@@ -0,0 +1,37 @@
+// Copyright 1998-2017 Epic Games, Inc. All Rights Reserved.
+
+#pragma once
+
+#include "Evaluation/MovieSceneEvalTemplate.h"
+#include "Evaluation/MovieSceneTrackImplementation.h"
+#include "MovieSceneNiagaraSystemTrackTemplate.generated.h"
+
+USTRUCT()
+struct FMovieSceneNiagaraSystemTrackTemplate : public FMovieSceneEvalTemplate
+{
+	GENERATED_BODY()
+
+public:
+	FMovieSceneNiagaraSystemTrackTemplate() { }
+
+private:
+	virtual UScriptStruct& GetScriptStructImpl() const override { return *StaticStruct(); }
+};
+
+USTRUCT()
+struct FMovieSceneNiagaraSystemTrackImplementation : public FMovieSceneTrackImplementation
+{
+	GENERATED_BODY()
+
+public:
+	FMovieSceneNiagaraSystemTrackImplementation() { }
+	FMovieSceneNiagaraSystemTrackImplementation(FFrameNumber InSpawnSectionStartFrame, FFrameNumber InSpawnSectionEndFrame);
+
+private:
+	virtual UScriptStruct& GetScriptStructImpl() const override { return *StaticStruct(); }
+	virtual void Evaluate(const FMovieSceneEvaluationTrack& Track, FMovieSceneSegmentIdentifier SegmentID, const FMovieSceneEvaluationOperand& Operand, const FMovieSceneContext& Context, const FPersistentEvaluationData& PersistentData, FMovieSceneExecutionTokens& ExecutionTokens) const override;
+
+private:
+	FFrameNumber SpawnSectionStartFrame;
+	FFrameNumber SpawnSectionEndFrame;
+};
\ No newline at end of file
diff --git a/Engine/Plugins/FX/Niagara/Source/Niagara/Private/MovieScene/Parameters/MovieSceneNiagaraBoolParameterSectionTemplate.cpp b/Engine/Plugins/FX/Niagara/Source/Niagara/Private/MovieScene/Parameters/MovieSceneNiagaraBoolParameterSectionTemplate.cpp
new file mode 100644
--- /dev/null
+++ b/Engine/Plugins/FX/Niagara/Source/Niagara/Private/MovieScene/Parameters/MovieSceneNiagaraBoolParameterSectionTemplate.cpp
@@ -0,0 +1,31 @@
+// Copyright 1998-2017 Epic Games, Inc. All Rights Reserved.
+
+#include "MovieSceneNiagaraBoolParameterSectionTemplate.h"
+#include "NiagaraComponent.h"
+#include "NiagaraTypes.h"
+#include "IMovieScenePlayer.h"
+
+FMovieSceneNiagaraBoolParameterSectionTemplate::FMovieSceneNiagaraBoolParameterSectionTemplate()
+{
+}
+
+FMovieSceneNiagaraBoolParameterSectionTemplate::FMovieSceneNiagaraBoolParameterSectionTemplate(FNiagaraVariable InParameter, const FMovieSceneBoolChannel& InBoolChannel)
+	: FMovieSceneNiagaraParameterSectionTemplate(InParameter)
+	, BoolChannel(InBoolChannel)
+{
+}
+
+void FMovieSceneNiagaraBoolParameterSectionTemplate::GetParameterValue(FFrameTime InTime, const TArray<uint8>& InCurrentValueData, TArray<uint8>& OutAnimatedValueData) const
+{
+	FNiagaraBool const* CurrentValue = (FNiagaraBool const*)InCurrentValueData.GetData();
+	FNiagaraBool AnimatedNiagaraValue = *CurrentValue;
+
+	bool AnimatedValue;
+	if (BoolChannel.Evaluate(InTime, AnimatedValue))
+	{
+		AnimatedNiagaraValue.SetValue(AnimatedValue);
+	}
+	
+	OutAnimatedValueData.AddUninitialized(sizeof(FNiagaraBool));
+	FMemory::Memcpy(OutAnimatedValueData.GetData(), (uint8*)&AnimatedValue, sizeof(FNiagaraBool));
+}
\ No newline at end of file
diff --git a/Engine/Plugins/FX/Niagara/Source/Niagara/Private/MovieScene/Parameters/MovieSceneNiagaraFloatParameterSectionTemplate.cpp b/Engine/Plugins/FX/Niagara/Source/Niagara/Private/MovieScene/Parameters/MovieSceneNiagaraFloatParameterSectionTemplate.cpp
new file mode 100644
--- /dev/null
+++ b/Engine/Plugins/FX/Niagara/Source/Niagara/Private/MovieScene/Parameters/MovieSceneNiagaraFloatParameterSectionTemplate.cpp
@@ -0,0 +1,27 @@
+// Copyright 1998-2017 Epic Games, Inc. All Rights Reserved.
+
+#include "MovieSceneNiagaraFloatParameterSectionTemplate.h"
+#include "NiagaraComponent.h"
+#include "NiagaraTypes.h"
+#include "IMovieScenePlayer.h"
+
+FMovieSceneNiagaraFloatParameterSectionTemplate::FMovieSceneNiagaraFloatParameterSectionTemplate()
+{
+}
+
+FMovieSceneNiagaraFloatParameterSectionTemplate::FMovieSceneNiagaraFloatParameterSectionTemplate(FNiagaraVariable InParameter, const FMovieSceneFloatChannel& InFloatChannel)
+	: FMovieSceneNiagaraParameterSectionTemplate(InParameter)
+	, FloatChannel(InFloatChannel)
+{
+}
+
+void FMovieSceneNiagaraFloatParameterSectionTemplate::GetParameterValue(FFrameTime InTime, const TArray<uint8>& InCurrentValueData, TArray<uint8>& OutAnimatedValueData) const
+{
+	FNiagaraFloat const* CurrentValue = (FNiagaraFloat const*)InCurrentValueData.GetData();
+	FNiagaraFloat AnimatedValue = *CurrentValue;
+
+	FloatChannel.Evaluate(InTime, AnimatedValue.Value);
+
+	OutAnimatedValueData.AddUninitialized(sizeof(FNiagaraFloat));
+	FMemory::Memcpy(OutAnimatedValueData.GetData(), (uint8*)&AnimatedValue, sizeof(FNiagaraFloat));
+}
\ No newline at end of file
diff --git a/Engine/Plugins/FX/Niagara/Source/Niagara/Private/MovieScene/Parameters/MovieSceneNiagaraIntegerParameterSectionTemplate.cpp b/Engine/Plugins/FX/Niagara/Source/Niagara/Private/MovieScene/Parameters/MovieSceneNiagaraIntegerParameterSectionTemplate.cpp
new file mode 100644
--- /dev/null
+++ b/Engine/Plugins/FX/Niagara/Source/Niagara/Private/MovieScene/Parameters/MovieSceneNiagaraIntegerParameterSectionTemplate.cpp
@@ -0,0 +1,27 @@
+// Copyright 1998-2017 Epic Games, Inc. All Rights Reserved.
+
+#include "MovieSceneNiagaraIntegerParameterSectionTemplate.h"
+#include "NiagaraComponent.h"
+#include "NiagaraTypes.h"
+#include "IMovieScenePlayer.h"
+
+FMovieSceneNiagaraIntegerParameterSectionTemplate::FMovieSceneNiagaraIntegerParameterSectionTemplate()
+{
+}
+
+FMovieSceneNiagaraIntegerParameterSectionTemplate::FMovieSceneNiagaraIntegerParameterSectionTemplate(FNiagaraVariable InParameter, const FMovieSceneIntegerChannel& InIntegerChannel)
+	: FMovieSceneNiagaraParameterSectionTemplate(InParameter)
+	, IntegerChannel(InIntegerChannel)
+{
+}
+
+void FMovieSceneNiagaraIntegerParameterSectionTemplate::GetParameterValue(FFrameTime InTime, const TArray<uint8>& InCurrentValueData, TArray<uint8>& OutAnimatedValueData) const
+{
+	FNiagaraInt32 const* CurrentValue = (FNiagaraInt32 const*)InCurrentValueData.GetData();
+	FNiagaraInt32 AnimatedValue = *CurrentValue;
+
+	IntegerChannel.Evaluate(InTime, AnimatedValue.Value);
+
+	OutAnimatedValueData.AddUninitialized(sizeof(FNiagaraInt32));
+	FMemory::Memcpy(OutAnimatedValueData.GetData(), (uint8*)&AnimatedValue, sizeof(FNiagaraInt32));
+}
\ No newline at end of file
diff --git a/Engine/Plugins/FX/Niagara/Source/Niagara/Private/MovieScene/Parameters/MovieSceneNiagaraVectorParameterSectionTemplate.cpp b/Engine/Plugins/FX/Niagara/Source/Niagara/Private/MovieScene/Parameters/MovieSceneNiagaraVectorParameterSectionTemplate.cpp
new file mode 100644
--- /dev/null
+++ b/Engine/Plugins/FX/Niagara/Source/Niagara/Private/MovieScene/Parameters/MovieSceneNiagaraVectorParameterSectionTemplate.cpp
@@ -0,0 +1,60 @@
+// Copyright 1998-2017 Epic Games, Inc. All Rights Reserved.
+
+#include "MovieSceneNiagaraVectorParameterSectionTemplate.h"
+#include "NiagaraComponent.h"
+#include "NiagaraTypes.h"
+#include "IMovieScenePlayer.h"
+
+FMovieSceneNiagaraVectorParameterSectionTemplate::FMovieSceneNiagaraVectorParameterSectionTemplate()
+{
+}
+
+FMovieSceneNiagaraVectorParameterSectionTemplate::FMovieSceneNiagaraVectorParameterSectionTemplate(FNiagaraVariable InParameter, TArray<FMovieSceneFloatChannel>&& InVectorChannels, int32 InChannelsUsed)
+	: FMovieSceneNiagaraParameterSectionTemplate(InParameter)
+	, ChannelsUsed(InChannelsUsed)
+{
+	for (int32 i = 0; i < ChannelsUsed; i++)
+	{
+		VectorChannels[i] = InVectorChannels[i];
+	}
+}
+
+void FMovieSceneNiagaraVectorParameterSectionTemplate::GetParameterValue(FFrameTime InTime, const TArray<uint8>& InCurrentValueData, TArray<uint8>& OutAnimatedValueData) const
+{
+	if (ChannelsUsed == 2)
+	{
+		FVector2D const* CurrentValue = (FVector2D const*)InCurrentValueData.GetData();
+		FVector2D AnimatedValue = *CurrentValue;
+
+		VectorChannels[0].Evaluate(InTime, AnimatedValue.X);
+		VectorChannels[1].Evaluate(InTime, AnimatedValue.Y);
+
+		OutAnimatedValueData.AddUninitialized(sizeof(FVector2D));
+		FMemory::Memcpy(OutAnimatedValueData.GetData(), (uint8*)&AnimatedValue, sizeof(FVector2D));
+	}
+	else if (ChannelsUsed == 3)
+	{
+		FVector const* CurrentValue = (FVector const*)InCurrentValueData.GetData();
+		FVector AnimatedValue = *CurrentValue;
+
+		VectorChannels[0].Evaluate(InTime, AnimatedValue.X);
+		VectorChannels[1].Evaluate(InTime, AnimatedValue.Y);
+		VectorChannels[2].Evaluate(InTime, AnimatedValue.Z);
+
+		OutAnimatedValueData.AddUninitialized(sizeof(FVector));
+		FMemory::Memcpy(OutAnimatedValueData.GetData(), (uint8*)&AnimatedValue, sizeof(FVector));
+	}
+	else if (ChannelsUsed == 4)
+	{
+		FVector4 const* CurrentValue = (FVector4 const*)InCurrentValueData.GetData();
+		FVector4 AnimatedValue = *CurrentValue;
+
+		VectorChannels[0].Evaluate(InTime, AnimatedValue.X);
+		VectorChannels[1].Evaluate(InTime, AnimatedValue.Y);
+		VectorChannels[2].Evaluate(InTime, AnimatedValue.Z);
+		VectorChannels[3].Evaluate(InTime, AnimatedValue.W);
+
+		OutAnimatedValueData.AddUninitialized(sizeof(FVector4));
+		FMemory::Memcpy(OutAnimatedValueData.GetData(), (uint8*)&AnimatedValue, sizeof(FVector4));
+	}
+}
\ No newline at end of file
diff --git a/Engine/Plugins/FX/Niagara/Source/Niagara/Private/NiagaraActor.cpp b/Engine/Plugins/FX/Niagara/Source/Niagara/Private/NiagaraActor.cpp
--- a/Engine/Plugins/FX/Niagara/Source/Niagara/Private/NiagaraActor.cpp
+++ b/Engine/Plugins/FX/Niagara/Source/Niagara/Private/NiagaraActor.cpp
@@ -29,7 +29,7 @@ ANiagaraActor::ANiagaraActor(const FObjectInitializer& ObjectInitializer)
 			FName ID_Effects;
 			FText NAME_Effects;
 			FConstructorStatics()
-				: SpriteTextureObject(TEXT("/Engine/EditorResources/S_Emitter"))
+				: SpriteTextureObject(TEXT("/Niagara/Icons/S_ParticleSystem"))
 				, ID_Effects(TEXT("Effects"))
 				, NAME_Effects(NSLOCTEXT("SpriteCategory", "Effects", "Effects"))
 			{
diff --git a/Engine/Plugins/FX/Niagara/Source/Niagara/Private/NiagaraCollision.cpp b/Engine/Plugins/FX/Niagara/Source/Niagara/Private/NiagaraCollision.cpp
--- a/Engine/Plugins/FX/Niagara/Source/Niagara/Private/NiagaraCollision.cpp
+++ b/Engine/Plugins/FX/Niagara/Source/Niagara/Private/NiagaraCollision.cpp
@@ -69,7 +69,6 @@ void FNiagaraCollisionBatch::KickoffNewBatch(FNiagaraEmitterInstance *Sim, float
 				FNiagaraCollisionTrace Trace;
 				Trace.CollisionTraceHandle = Handle;
 				Trace.SourceParticleIndex = i;
-				Trace.OriginalVelocity = Velocity;
 				CollisionTraces.Add(Trace);
 			}
 
@@ -122,7 +121,6 @@ void FNiagaraCollisionBatch::GenerateEventsFromResults(FNiagaraEmitterInstance *
 					FNiagaraCollisionEventPayload Event;
 					Event.CollisionNormal = Hit->ImpactNormal;
 					Event.CollisionPos = Hit->ImpactPoint;
-					Event.CollisionVelocity = CurCheck.OriginalVelocity;
 					Event.ParticleIndex = CurTrace.UserData;
 					check(!Event.CollisionNormal.ContainsNaN());
 					check(Event.CollisionNormal.IsNormalized());
@@ -190,7 +188,7 @@ void FNiagaraCollisionBatch::GenerateEventsFromResults(FNiagaraEmitterInstance *
 
 
 
-int32 FNiagaraDICollisionQueryBatch::SubmitQuery(FVector Position, FVector Velocity, float CollisionSize, float DeltaSeconds)
+int32 FNiagaraDICollisionQueryBatch::SubmitQuery(FVector Position, FVector Direction, float CollisionSize, float DeltaSeconds)
 {
 	SCOPE_CYCLE_COUNTER(STAT_NiagaraCollision);
 	int32 Ret = INDEX_NONE;
@@ -199,13 +197,13 @@ int32 FNiagaraDICollisionQueryBatch::SubmitQuery(FVector Position, FVector Veloc
 		//int32 TestCollision = *TstIt;
 		//if (TestCollision)
 		{
-			FVector EndPosition = Position + Velocity * DeltaSeconds;
+			FVector EndPosition = Position + Direction * DeltaSeconds;
 
 			float Length;
-			FVector Direction;
-			Velocity.ToDirectionAndLength(Direction, Length);
-			Position -= Direction * (CollisionSize / 2);
-			EndPosition += Direction * (CollisionSize/ 2);
+			FVector NormDir;
+			Direction.ToDirectionAndLength(NormDir, Length);
+			Position -= NormDir * (CollisionSize / 2);
+			EndPosition += NormDir * (CollisionSize/ 2);
 
 			FCollisionQueryParams QueryParams(SCENE_QUERY_STAT(NiagraAsync));
 			QueryParams.OwnerTag = "Niagara";
@@ -219,7 +217,8 @@ int32 FNiagaraDICollisionQueryBatch::SubmitQuery(FVector Position, FVector Veloc
 			FNiagaraCollisionTrace Trace;
 			Trace.CollisionTraceHandle = Handle;
 			Trace.SourceParticleIndex = TraceID;
-			Trace.OriginalVelocity = Velocity;
+			Trace.CollisionSize = CollisionSize;
+			Trace.DeltaSeconds = DeltaSeconds;
 
 			int32 TraceIdx = CollisionTraces[GetWriteBufferIdx()].Add(Trace);
 			IdToTraceIdx[GetWriteBufferIdx()].Add(TraceID) = TraceIdx;
@@ -282,13 +281,21 @@ bool FNiagaraDICollisionQueryBatch::GetQueryResult(uint32 InTraceID, FNiagaraDIC
 
 				Payloads.Add(Event);
 				*/
-				Result.CollisionPos = Hit->ImpactPoint;
+				Result.CollisionPos = Hit->ImpactPoint;// -NormVel*(CurTrace.CollisionSize / 2);
 				Result.CollisionNormal = Hit->ImpactNormal;
-				Result.CollisionVelocity = CurTrace.OriginalVelocity;
 				Result.TraceID = InTraceID;
-				Result.PhysicalMaterialIdx = Hit->PhysMaterial->GetUniqueID();
-				Result.Friction = Hit->PhysMaterial->Friction;
-				Result.Restitution = Hit->PhysMaterial->Restitution;
+				if (Hit->PhysMaterial.IsValid())
+				{
+					Result.PhysicalMaterialIdx = Hit->PhysMaterial->GetUniqueID();
+					Result.Friction = Hit->PhysMaterial->Friction;
+					Result.Restitution = Hit->PhysMaterial->Restitution;
+				}
+				else
+				{
+					Result.PhysicalMaterialIdx = -1;
+					Result.Friction = 0.0f;
+					Result.Restitution = 0.0f;
+				}
 				return true;
 			}
 		}
diff --git a/Engine/Plugins/FX/Niagara/Source/Niagara/Private/NiagaraCommon.cpp b/Engine/Plugins/FX/Niagara/Source/Niagara/Private/NiagaraCommon.cpp
--- a/Engine/Plugins/FX/Niagara/Source/Niagara/Private/NiagaraCommon.cpp
+++ b/Engine/Plugins/FX/Niagara/Source/Niagara/Private/NiagaraCommon.cpp
@@ -178,6 +178,9 @@ void FNiagaraSystemUpdateContext::AddInternal(UNiagaraComponent* Comp, bool bReI
 	}
 }
 
+
+//////////////////////////////////////////////////////////////////////////
+
 FName NIAGARA_API FNiagaraUtilities::GetUniqueName(FName CandidateName, const TSet<FName>& ExistingNames)
 {
 	if (ExistingNames.Contains(CandidateName) == false)
@@ -202,3 +205,99 @@ FName NIAGARA_API FNiagaraUtilities::GetUniqueName(FName CandidateName, const TS
 
 	return UniqueName;
 }
+
+FNiagaraVariable FNiagaraUtilities::ConvertVariableToRapidIterationConstantName(FNiagaraVariable InVar, const TCHAR* InEmitterName, ENiagaraScriptUsage InUsage)
+{
+	FNiagaraVariable Var = InVar;
+
+	TArray<FString> SplitName;
+	Var.GetName().ToString().ParseIntoArray(SplitName, TEXT("."));
+	int32 NumSlots = SplitName.Num();
+	if (InEmitterName != nullptr)
+	{
+		for (int32 i = 0; i < NumSlots; i++)
+		{
+			if (SplitName[i] == TEXT("Emitter"))
+			{
+				SplitName[i] = InEmitterName;
+			}
+		}
+
+		if (NumSlots >= 3 && SplitName[0] == InEmitterName)
+		{
+			// Do nothing
+			UE_LOG(LogNiagara, Log, TEXT("ConvertVariableToRapidIterationConstantName Got here!"));
+		}
+		else
+		{
+			SplitName.Insert(InEmitterName, 0);
+		}
+		SplitName.Insert(TEXT("Constants"), 0);
+	}
+	else
+	{
+		SplitName.Insert(TEXT("Constants"), 0);
+	}
+
+	FString OutVarStrName = FString::Join<FString>(SplitName, TEXT("."));
+	Var.SetName(*OutVarStrName);
+	return Var;
+}
+
+void FNiagaraUtilities::CollectScriptDataInterfaceParameters(const UObject& Owner, const TArray<UNiagaraScript*>& Scripts, FNiagaraParameterStore& OutDataInterfaceParameters)
+{
+	for (UNiagaraScript* Script : Scripts)
+	{
+		for (FNiagaraScriptDataInterfaceInfo& DataInterfaceInfo : Script->GetCachedDefaultDataInterfaces())
+		{
+			if (DataInterfaceInfo.RegisteredParameterMapWrite != NAME_None)
+			{
+				FNiagaraVariable DataInterfaceParameter(DataInterfaceInfo.Type, DataInterfaceInfo.RegisteredParameterMapWrite);
+				if (OutDataInterfaceParameters.AddParameter(DataInterfaceParameter, false, false))
+				{
+					OutDataInterfaceParameters.SetDataInterface(DataInterfaceInfo.DataInterface, DataInterfaceParameter);
+				}
+				else
+				{
+					UE_LOG(LogNiagara, Error, TEXT("Duplicate data interface parameter writes found, simulation will be incorrect.  Owner: %s Parameter: %s"),
+						*Owner.GetPathName(), *DataInterfaceInfo.RegisteredParameterMapWrite.ToString());
+				}
+			}
+		}
+	}
+}
+
+bool FNiagaraScriptDataInterfaceCompileInfo::CanExecuteOnTarget(ENiagaraSimTarget SimTarget) const
+{
+	check(IsInGameThread());
+	UNiagaraDataInterface* Obj = GetDefaultDataInterface();
+	if (Obj)
+	{
+		return Obj->CanExecuteOnTarget(SimTarget);
+	}
+	check(false);
+	return false;
+}
+
+bool FNiagaraScriptDataInterfaceCompileInfo::IsSystemSolo() const
+{
+	check(IsInGameThread());
+	if (Name.ToString().StartsWith("User."))
+	{
+		return true;
+	}
+
+	UNiagaraDataInterface* Obj = GetDefaultDataInterface();
+	if (Obj && Obj->PerInstanceDataSize() > 0)
+	{
+		return true;
+	}
+	return false;
+}
+
+UNiagaraDataInterface* FNiagaraScriptDataInterfaceCompileInfo::GetDefaultDataInterface() const
+{
+	check(IsInGameThread());
+	UNiagaraDataInterface* Obj = CastChecked<UNiagaraDataInterface>(const_cast<UClass*>(Type.GetClass())->GetDefaultObject(true));
+	return Obj;
+}
\ No newline at end of file
