Copying //UE4/Dev-Niagara to //UE4/Dev-Main (Source: //UE4/Dev-Niagara @ 4074996)

diff --git a/Engine/Plugins/FX/Niagara/Source/Niagara/Private/NiagaraComponent.cpp b/Engine/Plugins/FX/Niagara/Source/Niagara/Private/NiagaraComponent.cpp
--- a/Engine/Plugins/FX/Niagara/Source/Niagara/Private/NiagaraComponent.cpp
+++ b/Engine/Plugins/FX/Niagara/Source/Niagara/Private/NiagaraComponent.cpp
@@ -18,7 +18,9 @@
 #include "UObject/NameTypes.h"
 #include "NiagaraParameterCollection.h"
 #include "NiagaraWorldManager.h"
+#include "EngineUtils.h"
 
+DECLARE_CYCLE_STAT(TEXT("Sceneproxy create (GT)"), STAT_NiagaraCreateSceneProxy, STATGROUP_Niagara);
 DECLARE_CYCLE_STAT(TEXT("Component Tick (GT)"), STAT_NiagaraComponentTick, STATGROUP_Niagara);
 DECLARE_CYCLE_STAT(TEXT("Activate (GT)"), STAT_NiagaraComponentActivate, STATGROUP_Niagara);
 DECLARE_CYCLE_STAT(TEXT("Deactivate (GT)"), STAT_NiagaraComponentDeactivate, STATGROUP_Niagara);
@@ -27,6 +29,68 @@ DECLARE_CYCLE_STAT(TEXT("Get Dynamic Mesh Elements (RT)"), STAT_NiagaraComponent
 
 DEFINE_LOG_CATEGORY(LogNiagara);
 
+static int32 GbSuppressNiagaraSystems = 0;
+static FAutoConsoleVariableRef CVarSuppressNiagaraSystems(
+	TEXT("fx.SuppressNiagaraSystems"),
+	GbSuppressNiagaraSystems,
+	TEXT("If > 0 Niagara particle systems will not be activated. \n"),
+	ECVF_Default
+);
+
+
+void DumpNiagaraComponents(UWorld* World)
+{
+	for (TActorIterator<AActor> ActorItr(World); ActorItr; ++ActorItr)
+	{
+		TArray<UNiagaraComponent*> Components;
+		ActorItr->GetComponents<UNiagaraComponent>(Components, true);
+		if (Components.Num() != 0)
+		{
+			UE_LOG(LogNiagara, Log, TEXT("Actor: \"%s\" ... %d Components"), *ActorItr->GetName(), Components.Num());
+		}
+
+		for (UNiagaraComponent* Component : Components)
+		{
+			if (Component != nullptr)
+			{
+				UNiagaraSystem* Sys = Component->GetAsset();
+				FNiagaraSystemInstance* SysInst = Component->GetSystemInstance();
+				if (!Sys)
+				{
+					UE_LOG(LogNiagara, Log, TEXT("Component: \"%s\" ... no system"), *Component->GetName());
+
+				}
+				else if (Sys && !SysInst)
+				{
+					UE_LOG(LogNiagara, Log, TEXT("Component: \"%s\" System: \"%s\" ... no instance"), *Component->GetName(), *Sys->GetName());
+
+				}
+				else
+				{
+					UE_LOG(LogNiagara, Log, TEXT("Component: \"%s\" System: \"%s\" | ReqExecState: %d | ExecState: %d | bIsActive: %d"), *Component->GetName(), *Sys->GetName(),
+						(int32)SysInst->GetRequestedExecutionState(), (int32)SysInst->GetActualExecutionState(), Component->bIsActive);
+
+					if (!SysInst->IsComplete())
+					{
+						for (TSharedRef<FNiagaraEmitterInstance> Emitter : SysInst->GetEmitters())
+						{
+							UE_LOG(LogNiagara, Log, TEXT("    Emitter: \"%s\" | ExecState: %d | NumParticles: %d | CPUTime: %f"), *Emitter->GetEmitterHandle().GetUniqueInstanceName(),
+								(int32)Emitter->GetExecutionState(), Emitter->GetNumParticles(), Emitter->GetTotalCPUTime());
+						}
+					}
+				}
+			}
+		}
+	}
+}
+
+FAutoConsoleCommandWithWorld DumpNiagaraComponentsCommand(
+	TEXT("DumpNiagaraComponents"),
+	TEXT("Dump Existing Niagara Components"),
+	FConsoleCommandWithWorldDelegate::CreateStatic(&DumpNiagaraComponents)
+);
+
+
 FNiagaraSceneProxy::FNiagaraSceneProxy(const UNiagaraComponent* InComponent)
 		: FPrimitiveSceneProxy(InComponent)
 		, bRenderingEnabled(true)
@@ -36,7 +100,7 @@ FNiagaraSceneProxy::FNiagaraSceneProxy(const UNiagaraComponent* InComponent)
 	FNiagaraSystemInstance* SystemInst = InComponent->GetSystemInstance();
 	if (SystemInst)
 	{
-		TArray<TSharedRef<FNiagaraEmitterInstance> > Sims = SystemInst->GetEmitters();
+		TArray<TSharedRef<FNiagaraEmitterInstance> > &Sims = SystemInst->GetEmitters();
 		TArray<NiagaraRenderer*> RenderersFromSims;
 		for (TSharedRef<FNiagaraEmitterInstance> Sim : Sims)
 		{
@@ -57,20 +121,26 @@ SIZE_T FNiagaraSceneProxy::GetTypeHash() const
 	static size_t UniquePointer;
 	return reinterpret_cast<size_t>(&UniquePointer);
 }
-
-void FNiagaraSceneProxy::AddEmitterRenderer(NiagaraRenderer* Renderer)
-{
-	EmitterRenderers.Add(Renderer); 
-	//UE_LOG(LogNiagara, Warning, TEXT("FNiagaraSceneProxy::AddEmitterRenderer %p"), Renderer);
-}
-
 void FNiagaraSceneProxy::UpdateEmitterRenderers(TArray<NiagaraRenderer*>& InRenderers)
 {
 	EmitterRenderers.Empty();
 	for (NiagaraRenderer* EmitterRenderer : InRenderers)
 	{
-		AddEmitterRenderer(EmitterRenderer);
+		if (EmitterRenderer)
+		{
+			EmitterRenderers.Add(EmitterRenderer);
+		}
 	}
+
+	// We sort by the sort hint in order to guarantee that we submit according to the preferred sort order..
+	EmitterRenderers.Sort([&](const NiagaraRenderer& A, const NiagaraRenderer& B)
+	{
+		check(A.GetRendererProperties());
+		check(B.GetRendererProperties());
+		int32 AIndex = A.GetRendererProperties()->SortOrderHint;
+		int32 BIndex = B.GetRendererProperties()->SortOrderHint;
+		return AIndex < BIndex;
+	});
 }
 
 FNiagaraSceneProxy::~FNiagaraSceneProxy()
@@ -135,7 +205,7 @@ FPrimitiveViewRelevance FNiagaraSceneProxy::GetViewRelevance(const FSceneView* V
 
 	for (NiagaraRenderer* Renderer : EmitterRenderers)
 	{
-		if (Renderer)
+		if (Renderer && Renderer->GetRendererProperties()->GetIsEnabled())
 		{
 			Relevance |= Renderer->GetViewRelevance(View, this);
 		}
@@ -243,10 +313,18 @@ UNiagaraComponent::UNiagaraComponent(const FObjectInitializer& ObjectInitializer
 	: Super(ObjectInitializer)
 	, OverrideParameters(this)
 	, bForceSolo(false)
-	, AgeUpdateMode(EAgeUpdateMode::TickDeltaTime)
+	, AgeUpdateMode(ENiagaraAgeUpdateMode::TickDeltaTime)
 	, DesiredAge(0.0f)
+	, bCanRenderWhileSeeking(true)
 	, SeekDelta(1 / 30.0f)
+	, MaxSimTime(33.0f / 1000.0f)
+	, bIsSeeking(false)
 	, bAutoDestroy(false)
+#if WITH_EDITOR
+	, bWaitForCompilationOnActivate(false)
+#endif
+	, bAwaitingActivationDueToNotReady(false)
+	//, bIsChangingAutoAttachment(false)
 {
 	PrimaryComponentTick.bCanEverTick = true;
 	PrimaryComponentTick.TickGroup = TG_DuringPhysics;
@@ -263,72 +341,80 @@ void UNiagaraComponent::TickComponent(float DeltaSeconds, enum ELevelTick TickTy
 {
 	SCOPE_CYCLE_COUNTER(STAT_NiagaraComponentTick);
 
+	if (bAwaitingActivationDueToNotReady)
+	{
+		Activate(bActivateShouldResetWhenReady);
+		return;
+	}
+
 	if (!bIsActive && bAutoActivate)
 	{
 		Activate();
 	}
 
+	if (!SystemInstance)
+	{
+		return;
+	}
+
 	check(SystemInstance->IsSolo());
 	if (bIsActive && SystemInstance.Get() && !SystemInstance->IsComplete())
 	{
-		if (SceneProxy != nullptr)
-		{
-			FNiagaraSceneProxy* NiagaraProxy = static_cast<FNiagaraSceneProxy*>(SceneProxy);
-			NiagaraProxy->SetRenderingEnabled(bRenderingEnabled);
-		}
-
 		// If the interfaces have changed in a meaningful way, we need to potentially rebind and update the values.
 		if (OverrideParameters.GetInterfacesDirty())
 		{
 			SystemInstance->Reset(FNiagaraSystemInstance::EResetMode::ReInit);
 		}
 
-		if (AgeUpdateMode == EAgeUpdateMode::TickDeltaTime)
+		if (AgeUpdateMode == ENiagaraAgeUpdateMode::TickDeltaTime)
 		{
 			SystemInstance->ComponentTick(DeltaSeconds);
 		}
 		else
 		{
-			float AgeDiff = DesiredAge - SystemInstance->GetAge();
+			float AgeDiff = FMath::Max(DesiredAge, 0.0f) - SystemInstance->GetAge();
+			int32 TicksToProcess = 0;
 			if (FMath::Abs(AgeDiff) < KINDA_SMALL_NUMBER)
 			{
 				AgeDiff = 0.0f;
 			}
-
-			if (AgeDiff < 0.0f)
-			{
-				SystemInstance->Reset(FNiagaraSystemInstance::EResetMode::ResetAll);
-				AgeDiff = DesiredAge - SystemInstance->GetAge();
-			}
-
-			if (AgeDiff > 0.0f)
+			else
 			{
-				// TODO: This maximum should be configurable, and we should not draw the System until we are caught up.
-				static const float MaxSimTime = 1;
-				//Take SeekDelta to be the max delta time we'll accept.
-				TSharedPtr<FNiagaraSystemSimulation, ESPMode::ThreadSafe> SystemSim = GetSystemSimulation();
-				if (SystemSim.IsValid())
+				if (AgeDiff < 0.0f)
 				{
-					float SimTime = FMath::Min(AgeDiff, MaxSimTime);
-					int32 NumSteps = FMath::CeilToInt(SimTime / SeekDelta);
-					//Num steps must be odd. This is because when doing > 1 we have to flip the solo data set for in the system sim. This has to end up in the same flip state it was in before our ticks.
-					NumSteps = NumSteps % 2 == 0 ? NumSteps + 1 : NumSteps;
-					float RealSeekDelta = SimTime / (float)NumSteps;
+					SystemInstance->Reset(FNiagaraSystemInstance::EResetMode::ResetAll);
+					AgeDiff = DesiredAge - SystemInstance->GetAge();
+				}
 
-					for (int32 i = NumSteps - 1; i >= 0; --i)
+				if (AgeDiff > 0.0f)
+				{
+					FNiagaraSystemSimulation* SystemSim = GetSystemSimulation().Get();
+					if (SystemSim)
 					{
-						SystemInstance->ComponentTick(RealSeekDelta);
-						if (i > 0)
+						double StartTime = FPlatformTime::Seconds();
+						double CurrentTime = StartTime;
+
+						TicksToProcess = FMath::FloorToInt(AgeDiff / SeekDelta);
+						for (; TicksToProcess > 0 && CurrentTime - StartTime < MaxSimTime; --TicksToProcess)
 						{
-							SystemSim->TickSoloDataSet();
+							SystemInstance->ComponentTick(SeekDelta);
+							CurrentTime = FPlatformTime::Seconds();
 						}
 					}
 				}
 			}
+
+			if (TicksToProcess == 0)
+			{
+				bIsSeeking = false;
+			}
+		}
+
+		if (SceneProxy != nullptr)
+		{
+			FNiagaraSceneProxy* NiagaraProxy = static_cast<FNiagaraSceneProxy*>(SceneProxy);
+			NiagaraProxy->SetRenderingEnabled(bRenderingEnabled && (bCanRenderWhileSeeking || bIsSeeking == false));
 		}
-		
-		UpdateComponentToWorld();
-		MarkRenderDynamicDataDirty();
 	}
 }
 
@@ -358,54 +444,109 @@ void UNiagaraComponent::SetRenderingEnabled(bool bInRenderingEnabled)
 	bRenderingEnabled = bInRenderingEnabled;
 }
 
+void UNiagaraComponent::AdvanceSimulation(int32 TickCount, float TickDeltaSeconds)
+{
+	if (SystemInstance.IsValid() && TickDeltaSeconds > SMALL_NUMBER)
+	{
+		SystemInstance->AdvanceSimulation(TickCount, TickDeltaSeconds);
+	}
+}
+
+void UNiagaraComponent::AdvanceSimulationByTime(float SimulateTime, float TickDeltaSeconds)
+{
+	if (SystemInstance.IsValid() && TickDeltaSeconds > SMALL_NUMBER)
+	{
+		int32 TickCount = SimulateTime / TickDeltaSeconds;
+		SystemInstance->AdvanceSimulation(TickCount, TickDeltaSeconds);
+	}
+}
+
+bool UNiagaraComponent::InitializeSystem()
+{
+	if (SystemInstance.IsValid() == false)
+	{
+		//UE_LOG(LogNiagara, Warning, TEXT("Create System: %0X | %s\n"), this, *GetAsset()->GetFullName());
+		SystemInstance = MakeUnique<FNiagaraSystemInstance>(this);
+#if WITH_EDITORONLY_DATA
+		OnSystemInstanceChangedDelegate.Broadcast();
+#endif
+		SystemInstance->Init(GetAsset(), bForceSolo);
+		return true;
+	}
+	return false;
+}
+
 void UNiagaraComponent::Activate(bool bReset /* = false */)
 {
+	bAwaitingActivationDueToNotReady = false;
+
+	if (GbSuppressNiagaraSystems != 0)
+	{
+		OnSystemComplete();
+		return;
+	}
+
+	if (IsSwitchPlatform(GMaxRHIShaderPlatform))
+	{
+		UE_LOG(LogNiagara, Warning, TEXT("Failed to activate Niagara component as Niagara is not yet supported on this platform: %s"), *LegacyShaderPlatformToShaderFormat(GMaxRHIShaderPlatform).ToString());
+		OnSystemComplete();
+		return;
+	}
+
 	SCOPE_CYCLE_COUNTER(STAT_NiagaraComponentActivate);
 	if (Asset == nullptr)
 	{
 		DestroyInstance();
-
-		UE_LOG(LogNiagara, Warning, TEXT("Failed to activate Niagara Component due to missing or invalid asset!"));
+		if (!HasAnyFlags(RF_DefaultSubObject | RF_ArchetypeObject | RF_ClassDefaultObject))
+		{
+			UE_LOG(LogNiagara, Warning, TEXT("Failed to activate Niagara Component due to missing or invalid asset!"));
+		}
 		SetComponentTickEnabled(false);
 		return;
 	}
-
-	if (!Asset->IsReadyToRun())
+	
+	// If the particle system can't ever render (ie on dedicated server or in a commandlet) than do not activate...
+	if (!FApp::CanEverRender())
 	{
 		return;
 	}
 
-	if(!IsRegistered())
+	if (!IsRegistered())
 	{
 		return;
 	}
 
-	// If the particle system can't ever render (ie on dedicated server or in a commandlet) than do not activate...
-	if (!FApp::CanEverRender())
+	// On the off chance that the user changed the asset, we need to clear out the existing data.
+	if (SystemInstance.IsValid() && SystemInstance->GetSystem() != Asset)
 	{
-		return;
+		OnSystemComplete();
 	}
-	
-	Super::Activate(bReset);
 
-	if (SystemInstance.IsValid() == false)
+#if WITH_EDITOR
+	// In case we're not yet ready to run due to compilation requests, go ahead and keep polling there..
+	if (Asset->HasOutstandingCompilationRequests())
 	{
-		TSharedPtr<FNiagaraSystemSimulation, ESPMode::ThreadSafe> SystemSim = GetSystemSimulation();
-		if (SystemSim.IsValid())
-		{
-			SystemInstance = MakeUnique<FNiagaraSystemInstance>(this);
-#if WITH_EDITORONLY_DATA
-			OnSystemInstanceChangedDelegate.Broadcast();
-#endif
-			SystemInstance->Init(SystemSim.ToSharedRef(), bForceSolo);
-			bReset = false;
-		}
-		else
+		if (bWaitForCompilationOnActivate)
 		{
-			return;
+			Asset->WaitForCompilationComplete();
 		}
+		Asset->PollForCompilationComplete();
 	}
+#endif
 
+	if (!Asset->IsReadyToRun())
+	{
+		bAwaitingActivationDueToNotReady = true;
+		bActivateShouldResetWhenReady = bReset;
+		SetComponentTickEnabled(true);
+		return;
+	}
+
+	
+	Super::Activate(bReset);
+
+	//UE_LOG(LogNiagara, Log, TEXT("Activate: %u - %s"), this, *Asset->GetName());
+	
 	// Auto attach if requested
 	const bool bWasAutoAttached = bDidAutoAttach;
 	bDidAutoAttach = false;
@@ -422,7 +563,9 @@ void UNiagaraComponent::Activate(bool bReset /* = false */)
 				SavedAutoAttachRelativeLocation = RelativeLocation;
 				SavedAutoAttachRelativeRotation = RelativeRotation;
 				SavedAutoAttachRelativeScale3D = RelativeScale3D;
+				//bIsChangingAutoAttachment = true;
 				AttachToComponent(NewParent, FAttachmentTransformRules(AutoAttachLocationRule, AutoAttachRotationRule, AutoAttachScaleRule, false), AutoAttachSocketName);
+				//bIsChangingAutoAttachment = false;
 			}
 
 			bDidAutoAttach = true;
@@ -434,7 +577,18 @@ void UNiagaraComponent::Activate(bool bReset /* = false */)
 		}
 	}
 
-	SystemInstance->Activate(bReset);
+	FNiagaraSystemInstance::EResetMode ResetMode = bReset ? FNiagaraSystemInstance::EResetMode::ResetAll : FNiagaraSystemInstance::EResetMode::ResetSystem;
+	if (InitializeSystem())
+	{
+		ResetMode = FNiagaraSystemInstance::EResetMode::None;//Already done a reinit sete
+	}
+
+	if (!SystemInstance)
+	{
+		return;
+	}
+
+	SystemInstance->Activate(ResetMode);
 
 	/** We only need to tick the component if we require solo mode. */
 	SetComponentTickEnabled(SystemInstance->IsSolo());
@@ -445,6 +599,8 @@ void UNiagaraComponent::Deactivate()
 	SCOPE_CYCLE_COUNTER(STAT_NiagaraComponentDeactivate);
 	Super::Deactivate();
 
+	//UE_LOG(LogNiagara, Log, TEXT("Deactivate: %u - %s"), this, *Asset->GetName());
+
 	bIsActive = false;
 
 	if (SystemInstance)
@@ -453,11 +609,32 @@ void UNiagaraComponent::Deactivate()
 	}
 }
 
+void UNiagaraComponent::DeactivateImmediate()
+{
+	SCOPE_CYCLE_COUNTER(STAT_NiagaraComponentDeactivate);
+	Super::Deactivate();
+
+	//UE_LOG(LogNiagara, Log, TEXT("DeactivateImmediate: %u - %s"), this, *Asset->GetName());
+
+	//UE_LOG(LogNiagara, Log, TEXT("Deactivate %s"), *GetName());
+
+	bIsActive = false;
+
+	if (SystemInstance)
+	{
+		SystemInstance->Deactivate(true);
+	}
+}
+
 void UNiagaraComponent::OnSystemComplete()
 {
+	//UE_LOG(LogNiagara, Log, TEXT("OnSystemComplete: %u - %s"), this, *Asset->GetName());
+
 	SetComponentTickEnabled(false);
 	bIsActive = false;
-        
+
+	MarkRenderDynamicDataDirty();
+		
 	OnSystemFinished.Broadcast(this);
 
 	if (bAutoDestroy)
@@ -472,6 +649,9 @@ void UNiagaraComponent::OnSystemComplete()
 
 void UNiagaraComponent::DestroyInstance()
 {
+	//UE_LOG(LogNiagara, Warning, TEXT("Destroy System: %s\n"), *GetAsset()->GetFullName());
+	//UE_LOG(LogNiagara, Log, TEXT("DestroyInstance: %u - %s"), this, *Asset->GetName());
+	bIsActive = false;
 	SystemInstance = nullptr;
 #if WITH_EDITORONLY_DATA
 	OnSystemInstanceChangedDelegate.Broadcast();
@@ -502,7 +682,9 @@ void UNiagaraComponent::OnRegister()
 				// Only detach if we are not about to auto attach to the same target, that would be wasteful.
 				if (!bAutoActivate || (AutoAttachLocationRule != EAttachmentRule::KeepRelative && AutoAttachRotationRule != EAttachmentRule::KeepRelative && AutoAttachScaleRule != EAttachmentRule::KeepRelative) || (AutoAttachSocketName != GetAttachSocketName()) || (AutoAttachParent != GetAttachParent()))
 				{
+					//bIsChangingAutoAttachment = true;
 					DetachFromComponent(FDetachmentTransformRules(EDetachmentRule::KeepRelative, /*bCallModify=*/ false));
+					//bIsChangingAutoAttachment = false;
 				}
 			}
 			else
@@ -521,15 +703,21 @@ void UNiagaraComponent::OnRegister()
 void UNiagaraComponent::OnComponentDestroyed(bool bDestroyingHierarchy)
 {
 	//UE_LOG(LogNiagara, Log, TEXT("OnComponentDestroyed %p %p"), this, SystemInstance.Get());
-	DestroyInstance();
+	//DestroyInstance();//Can't do this here as we can call this from inside the system instance currently during completion 
 
 	Super::OnComponentDestroyed(bDestroyingHierarchy);
 }
 
 void UNiagaraComponent::OnUnregister()
 {
-	DestroyInstance();
 	Super::OnUnregister();
+
+	bIsActive = false;
+
+	if (SystemInstance)
+	{
+		SystemInstance->Deactivate(true);
+	}
 }
 
 void UNiagaraComponent::BeginDestroy()
@@ -539,14 +727,24 @@ void UNiagaraComponent::BeginDestroy()
 	Super::BeginDestroy();
 }
 
+// Uncertain about this. 
+// void UNiagaraComponent::OnAttachmentChanged()
+// {
+// 	Super::OnAttachmentChanged();
+// 	if (bIsActive && !bIsChangingAutoAttachment && !GetOwner()->IsPendingKillPending())
+// 	{
+// 		ResetSystem();
+// 	}
+// }
+
 TSharedPtr<FNiagaraSystemSimulation, ESPMode::ThreadSafe> UNiagaraComponent::GetSystemSimulation()
 {
-	UWorld* World = GetWorld();
-	if (World)
+	if (SystemInstance)
 	{
-		return FNiagaraWorldManager::Get(World)->GetSystemSimulation(Asset);
+		return SystemInstance->GetSystemSimulation();
 	}
-	return TSharedRef<FNiagaraSystemSimulation, ESPMode::ThreadSafe>();
+
+	return nullptr;
 }
 
 void UNiagaraComponent::SendRenderDynamicData_Concurrent()
@@ -560,7 +758,7 @@ void UNiagaraComponent::SendRenderDynamicData_Concurrent()
 
 		for (int32 i = 0; i < SystemInstance->GetEmitters().Num(); i++)
 		{
-			TSharedPtr<FNiagaraEmitterInstance> Emitter = SystemInstance->GetEmitters()[i];
+			FNiagaraEmitterInstance* Emitter = &SystemInstance->GetEmitters()[i].Get();
 			for (int32 EmitterIdx = 0; EmitterIdx < Emitter->GetEmitterRendererNum(); EmitterIdx++)
 			{
 				NiagaraRenderer* Renderer = Emitter->GetEmitterRenderer(EmitterIdx);
@@ -568,9 +766,10 @@ void UNiagaraComponent::SendRenderDynamicData_Concurrent()
 				{
 					bool bRendererEditorEnabled = true;
 #if WITH_EDITORONLY_DATA
-					bRendererEditorEnabled = !SystemInstance->GetIsolateEnabled() || Emitter->GetEmitterHandle().IsIsolated();
+					bRendererEditorEnabled = (!SystemInstance->GetIsolateEnabled() || Emitter->GetEmitterHandle().IsIsolated());
 #endif
-					if (bRendererEditorEnabled && !Emitter->IsComplete())
+					bRendererEditorEnabled &= Renderer->GetRendererProperties()->GetIsEnabled();
+					if (bRendererEditorEnabled && !Emitter->IsComplete() && !SystemInstance->IsComplete())
 					{
 						FNiagaraDynamicDataBase* DynamicData = Renderer->GenerateVertexData(NiagaraProxy, Emitter->GetData(), Emitter->GetEmitterHandle().GetInstance()->SimTarget);
 
@@ -627,6 +826,7 @@ FBoxSphereBounds UNiagaraComponent::CalcBounds(const FTransform& LocalToWorld) c
 
 FPrimitiveSceneProxy* UNiagaraComponent::CreateSceneProxy()
 {
+	SCOPE_CYCLE_COUNTER(STAT_NiagaraCreateSceneProxy);
 	// The constructor will set up the System renderers from the component.
 	FNiagaraSceneProxy* Proxy = new FNiagaraSceneProxy(this);
 	return Proxy;
@@ -666,6 +866,13 @@ void UNiagaraComponent::SetNiagaraVariableLinearColor(const FString& InVariableN
 	OverrideParameters.SetParameterValue(InValue, FNiagaraVariable(FNiagaraTypeDefinition::GetColorDef(), VarName), true);
 }
 
+void UNiagaraComponent::SetNiagaraVariableQuat(const FString& InVariableName, const FQuat& InValue)
+{
+	FName VarName = FName(*InVariableName);
+
+	OverrideParameters.SetParameterValue(InValue, FNiagaraVariable(FNiagaraTypeDefinition::GetQuatDef(), VarName), true);
+}
+
 void UNiagaraComponent::SetNiagaraVariableVec4(const FString& InVariableName, const FVector4& InValue)
 {
 	FName VarName = FName(*InVariableName);
@@ -772,15 +979,24 @@ void UNiagaraComponent::PostLoad()
 		Asset->ConditionalPostLoad();
 #if WITH_EDITOR
 		SynchronizeWithSourceSystem();
+		AssetExposedParametersChangedHandle = Asset->GetExposedParameters().AddOnChangedHandler(
+			FNiagaraParameterStore::FOnChanged::FDelegate::CreateUObject(this, &UNiagaraComponent::AssetExposedParametersChanged));
 #endif
 	}
 }
 
 #if WITH_EDITOR
 
+void UNiagaraComponent::PreEditChange(UProperty* PropertyAboutToChange)
+{
+	if (PropertyAboutToChange != nullptr && PropertyAboutToChange->GetFName() == GET_MEMBER_NAME_CHECKED(UNiagaraComponent, Asset) && Asset != nullptr)
+	{
+		Asset->GetExposedParameters().RemoveOnChangedHandler(AssetExposedParametersChangedHandle);
+	}
+}
+
 void UNiagaraComponent::PostEditChangeProperty(FPropertyChangedEvent& PropertyChangedEvent)
 {
-	Super::PostEditChangeProperty(PropertyChangedEvent);
 
 	FName PropertyName;
 	if (PropertyChangedEvent.Property)
@@ -788,31 +1004,47 @@ void UNiagaraComponent::PostEditChangeProperty(FPropertyChangedEvent& PropertyCh
 		PropertyName = PropertyChangedEvent.Property->GetFName();
 	}
 
-	if (PropertyName == GET_MEMBER_NAME_CHECKED(UNiagaraComponent, Asset) ||
-		PropertyName == GET_MEMBER_NAME_CHECKED(UNiagaraComponent, OverrideParameters))
+	if (PropertyName == GET_MEMBER_NAME_CHECKED(UNiagaraComponent, Asset))
+	{
+		SynchronizeWithSourceSystem();
+		if (Asset != nullptr)
+		{
+			AssetExposedParametersChangedHandle = Asset->GetExposedParameters().AddOnChangedHandler(
+				FNiagaraParameterStore::FOnChanged::FDelegate::CreateUObject(this, &UNiagaraComponent::AssetExposedParametersChanged));
+		}
+	}
+	else if(PropertyName == GET_MEMBER_NAME_CHECKED(UNiagaraComponent, OverrideParameters))
 	{
 		SynchronizeWithSourceSystem();
 	}
 
 	ReinitializeSystem();
+
+	Super::PostEditChangeProperty(PropertyChangedEvent);
 }
 
 
-bool UNiagaraComponent::SynchronizeWithSourceSystem()
+void UNiagaraComponent::SynchronizeWithSourceSystem()
 {
+	// Synchronizing parameters will create new data interface objects and if the old data
+	// interface objects are currently being used by a simulation they may be destroyed due to garbage
+	// collection, so preemptively kill the instance here.
+	DestroyInstance();
+
 	//TODO: Look through params in system in "Owner" namespace and add to our parameters.
 	if (Asset == nullptr)
 	{
-		return false;
+		OverrideParameters.Empty();
+		EditorOverridesValue.Empty();
+		return;
 	}
 
 	TArray<FNiagaraVariable> SourceVars;
 	Asset->GetExposedParameters().GetParameters(SourceVars);
 
-	bool bEditsMade = false;
 	for (FNiagaraVariable& Param : SourceVars)
 	{
-		bEditsMade |= OverrideParameters.AddParameter(Param, true);
+		OverrideParameters.AddParameter(Param, true);
 	}
 
 	TArray<FNiagaraVariable> ExistingVars;
@@ -839,131 +1071,23 @@ bool UNiagaraComponent::SynchronizeWithSourceSystem()
 
 	OverrideParameters.Rebind();
 
-	//////////////////////////////////////////////////////////////////////////
-	//OLD IMPL
-// 	if (Asset == nullptr)
-// 	{
-// 		return false; 
-// 	}
-// 
-// 	UNiagaraScript* Script = Asset->GetSystemSpawnScript();
-// 	if (Script == nullptr)
-// 	{
-// 		return false;
-// 	}
-// 
-// 	bool bEditsMade = false;
-// 	{
-// 		TArray<int32> VarIndicesToRemove;
-// 
-// 		// Check over all of our overrides to see if they still match up with the source System.
-// 		for (int32 i = 0; i < SystemParameterLocalOverrides.Num(); i++)
-// 		{
-// 			FNiagaraVariable* Variable = Script->Parameters.FindParameter(SystemParameterLocalOverrides[i].GetId());
-// 			// If the variable still exists, keep digging... otherwise remove it.
-// 			if (Variable)
-// 			{
-// 				if (Variable->GetType() != SystemParameterLocalOverrides[i].GetType())
-// 				{
-// 					UE_LOG(LogNiagara, Warning, TEXT("Variable '%s' types changed.. possibly losing override value."), *Variable->GetName().ToString());
-// 					VarIndicesToRemove.Add(i);
-// 				}
-// 				else if (Variable->GetName() != SystemParameterLocalOverrides[i].GetName())
-// 				{
-// 					if (bEditsMade == false)
-// 					{
-// 						bEditsMade = true;
-// 						Modify();
-// 					}
-// 					SystemParameterLocalOverrides[i].SetName(Variable->GetName());
-// 				}
-// 			}
-// 			else
-// 			{
-// 				VarIndicesToRemove.Add(i);
-// 			}
-// 		}
-// 
-// 		// Handle the variables that we want to remove.
-// 		if (VarIndicesToRemove.Num() != 0)
-// 		{
-// 			if (bEditsMade == false)
-// 			{
-// 				bEditsMade = true;
-// 				Modify();
-// 			}
-// 
-// 			for (int32 removeIdx = VarIndicesToRemove.Num() - 1; removeIdx >= 0; removeIdx--)
-// 			{
-// 				UE_LOG(LogNiagara, Warning, TEXT("Removing local override '%s'."), *SystemParameterLocalOverrides[VarIndicesToRemove[removeIdx]].GetName().ToString());
-// 
-// 				SystemParameterLocalOverrides.RemoveAt(VarIndicesToRemove[removeIdx]);
-// 			}
-// 		}
-// 	}
-// 
-// 	{
-// 		TArray<int32> VarIndicesToRemove;
-// 
-// 		// Check over all of our overrides to see if they still match up with the source System.
-// 		for (int32 i = 0; i < SystemDataInterfaceLocalOverrides.Num(); i++)
-// 		{
-// 			FNiagaraScriptDataInterfaceInfo* Variable = Script->DataInterfaceInfo.FindByPredicate([&](const FNiagaraScriptDataInterfaceInfo& Var)
-// 			{
-// 				return Var.Id == SystemDataInterfaceLocalOverrides[i].Id;
-// 			}); 
-// 			// If the variable still exists, keep digging... otherwise remove it.
-// 			if (Variable)
-// 			{
-// 				if (Variable->DataInterface->GetClass() != SystemDataInterfaceLocalOverrides[i].DataInterface->GetClass())
-// 				{
-// 					UE_LOG(LogNiagara, Warning, TEXT("Variable '%s' types changed.. losing override value."), *Variable->Name.ToString());
-// 					VarIndicesToRemove.Add(i);
-// 				}
-// 				else if (Variable->Name != SystemDataInterfaceLocalOverrides[i].Name)
-// 				{
-// 					if (bEditsMade == false)
-// 					{
-// 						bEditsMade = true;
-// 						Modify();
-// 					}
-// 					SystemDataInterfaceLocalOverrides[i].Name = Variable->Name;
-// 				}
-// 			}
-// 			else
-// 			{
-// 				VarIndicesToRemove.Add(i);
-// 			}
-// 		}
-// 
-// 		// Handle the variables that we want to remove.
-// 		if (VarIndicesToRemove.Num() != 0)
-// 		{
-// 			if (bEditsMade == false)
-// 			{
-// 				bEditsMade = true;
-// 				Modify();
-// 			}
-// 
-// 			for (int32 removeIdx = VarIndicesToRemove.Num() - 1; removeIdx >= 0; removeIdx--)
-// 			{
-// 				UE_LOG(LogNiagara, Warning, TEXT("Removing local override '%s'."), *SystemDataInterfaceLocalOverrides[VarIndicesToRemove[removeIdx]].Name.ToString());
-// 
-// 				SystemDataInterfaceLocalOverrides.RemoveAt(VarIndicesToRemove[removeIdx]);
-// 			}
-// 		}
-// 	}
-// 
-// 	return bEditsMade;
-	return false;
+#if WITH_EDITORONLY_DATA
+	OnSynchronizedWithAssetParametersDelegate.Broadcast();
+#endif
 }
 
-UNiagaraComponent::EAgeUpdateMode UNiagaraComponent::GetAgeUpdateMode() const
+void UNiagaraComponent::AssetExposedParametersChanged()
+{
+	SynchronizeWithSourceSystem();
+}
+#endif
+
+ENiagaraAgeUpdateMode UNiagaraComponent::GetAgeUpdateMode() const
 {
 	return AgeUpdateMode;
 }
 
-void UNiagaraComponent::SetAgeUpdateMode(EAgeUpdateMode InAgeUpdateMode)
+void UNiagaraComponent::SetAgeUpdateMode(ENiagaraAgeUpdateMode InAgeUpdateMode)
 {
 	AgeUpdateMode = InAgeUpdateMode;
 }
@@ -976,6 +1100,18 @@ float UNiagaraComponent::GetDesiredAge() const
 void UNiagaraComponent::SetDesiredAge(float InDesiredAge)
 {
 	DesiredAge = InDesiredAge;
+	bIsSeeking = false;
+}
+
+void UNiagaraComponent::SeekToDesiredAge(float InDesiredAge)
+{
+	DesiredAge = InDesiredAge;
+	bIsSeeking = true;
+}
+
+void UNiagaraComponent::SetCanRenderWhileSeeking(bool bInCanRenderWhileSeeking)
+{
+	bCanRenderWhileSeeking = bInCanRenderWhileSeeking;
 }
 
 float UNiagaraComponent::GetSeekDelta() const
@@ -988,7 +1124,17 @@ void UNiagaraComponent::SetSeekDelta(float InSeekDelta)
 	SeekDelta = InSeekDelta;
 }
 
+float UNiagaraComponent::GetMaxSimTime() const
+{
+	return MaxSimTime;
+}
+
+void UNiagaraComponent::SetMaxSimTime(float InMaxTime)
+{
+	MaxSimTime = InMaxTime;
+}
 
+#if WITH_EDITOR
 bool UNiagaraComponent::IsParameterValueOverriddenLocally(const FName& InParamName)
 {
 	bool* FoundVar = EditorOverridesValue.Find(InParamName);
@@ -1026,16 +1172,35 @@ void UNiagaraComponent::SetParameterValueOverriddenLocally(const FNiagaraVariabl
 
 void UNiagaraComponent::SetAsset(UNiagaraSystem* InAsset)
 {
-	Asset = InAsset;
+	if (Asset != InAsset)
+	{
+#if WITH_EDITOR
+		if (Asset != nullptr)
+		{
+			Asset->GetExposedParameters().RemoveOnChangedHandler(AssetExposedParametersChangedHandle);
+		}
+#endif
+		Asset = InAsset;
 
 #if WITH_EDITOR
-	if (GetWorld() != nullptr && GetWorld()->WorldType == EWorldType::Editor)
-	{
 		SynchronizeWithSourceSystem();
-	}
+		AssetExposedParametersChangedHandle = Asset->GetExposedParameters().AddOnChangedHandler(
+			FNiagaraParameterStore::FOnChanged::FDelegate::CreateUObject(this, &UNiagaraComponent::AssetExposedParametersChanged));
 #endif
 
-	ReinitializeSystem();
+		//Force a reinit.
+		DestroyInstance();
+	}
+}
+
+void UNiagaraComponent::SetForceSolo(bool bInForceSolo) 
+{ 
+	if (bForceSolo != bInForceSolo)
+	{
+		bForceSolo = bInForceSolo;
+		DestroyInstance();
+		SetComponentTickEnabled(bInForceSolo);
+	}
 }
 
 void UNiagaraComponent::SetAutoAttachmentParameters(USceneComponent* Parent, FName SocketName, EAttachmentRule LocationRule, EAttachmentRule RotationRule, EAttachmentRule ScaleRule)
@@ -1063,7 +1228,9 @@ void UNiagaraComponent::CancelAutoAttachment(bool bDetachFromParent)
 
 		if (bDetachFromParent)
 		{
+			//bIsChangingAutoAttachment = true;
 			DetachFromComponent(FDetachmentTransformRules::KeepRelativeTransform);
+			//bIsChangingAutoAttachment = false;
 		}
 	}
 }
diff --git a/Engine/Plugins/FX/Niagara/Source/Niagara/Private/NiagaraConstants.cpp b/Engine/Plugins/FX/Niagara/Source/Niagara/Private/NiagaraConstants.cpp
--- a/Engine/Plugins/FX/Niagara/Source/Niagara/Private/NiagaraConstants.cpp
+++ b/Engine/Plugins/FX/Niagara/Source/Niagara/Private/NiagaraConstants.cpp
@@ -1,18 +1,33 @@
 // Copyright 1998-2018 Epic Games, Inc. All Rights Reserved.
 
 #include "NiagaraConstants.h"
+#include "NiagaraModule.h"
 
 #include "Internationalization/Internationalization.h"
 
 #define LOCTEXT_NAMESPACE "FNiagaraConstants"
 
-const TArray<FNiagaraVariable>& FNiagaraConstants::GetEngineConstants()
+TArray<FNiagaraVariable> FNiagaraConstants::SystemParameters;
+TArray<FNiagaraVariable> FNiagaraConstants::TranslatorParameters;
+TMap<FName, FNiagaraVariable> FNiagaraConstants::UpdatedSystemParameters;
+TMap<FNiagaraVariable, FText> FNiagaraConstants::SystemStrMap;
+TArray<FNiagaraVariable> FNiagaraConstants::Attributes;
+TMap<FNiagaraVariable, FString> FNiagaraConstants::AttrDefaultsStrMap;
+TMap<FNiagaraVariable, FText> FNiagaraConstants::AttrDescStrMap;
+TMap<FNiagaraVariable, FNiagaraVariableMetaData> FNiagaraConstants::AttrMetaData;
+TMap<FNiagaraVariable, FNiagaraVariable> FNiagaraConstants::AttrDefaultsValueMap;
+TMap<FNiagaraVariable, FNiagaraVariable> FNiagaraConstants::AttrDataSetKeyMap;
+TArray<FNiagaraVariable> FNiagaraConstants::EngineManagedAttributes;
+
+void FNiagaraConstants::Init()
 {
-	static TArray<FNiagaraVariable> SystemParameters;
 	if (SystemParameters.Num() == 0)
 	{
 		SystemParameters.Add(SYS_PARAM_ENGINE_DELTA_TIME);
 		SystemParameters.Add(SYS_PARAM_ENGINE_INV_DELTA_TIME);
+		SystemParameters.Add(SYS_PARAM_ENGINE_TIME);
+		SystemParameters.Add(SYS_PARAM_ENGINE_REAL_TIME);
+
 		SystemParameters.Add(SYS_PARAM_ENGINE_POSITION);
 		SystemParameters.Add(SYS_PARAM_ENGINE_SCALE);
 		SystemParameters.Add(SYS_PARAM_ENGINE_VELOCITY);
@@ -30,16 +45,373 @@ const TArray<FNiagaraVariable>& FNiagaraConstants::GetEngineConstants()
 		SystemParameters.Add(SYS_PARAM_ENGINE_MIN_DIST_TO_CAMERA);
 		SystemParameters.Add(SYS_PARAM_ENGINE_TIME_SINCE_RENDERED);
 
+		SystemParameters.Add(SYS_PARAM_ENGINE_EXECUTION_STATE);
+
 		SystemParameters.Add(SYS_PARAM_ENGINE_EXEC_COUNT);
 		SystemParameters.Add(SYS_PARAM_ENGINE_EMITTER_NUM_PARTICLES);
 		SystemParameters.Add(SYS_PARAM_ENGINE_SYSTEM_NUM_EMITTERS_ALIVE);
 		SystemParameters.Add(SYS_PARAM_ENGINE_SYSTEM_NUM_EMITTERS);
+		SystemParameters.Add(SYS_PARAM_ENGINE_NUM_SYSTEM_INSTANCES);
+		SystemParameters.Add(SYS_PARAM_ENGINE_GLOBAL_SPAWN_COUNT_SCALE);
+		SystemParameters.Add(SYS_PARAM_ENGINE_GLOBAL_SYSTEM_COUNT_SCALE);
+		SystemParameters.Add(SYS_PARAM_ENGINE_SYSTEM_AGE);
 		SystemParameters.Add(SYS_PARAM_EMITTER_AGE);
 		SystemParameters.Add(SYS_PARAM_EMITTER_LOCALSPACE);
 	}
+
+	if (TranslatorParameters.Num() == 0)
+	{
+		TranslatorParameters.Add(TRANSLATOR_PARAM_BEGIN_DEFAULTS);
+	}
+
+	if (UpdatedSystemParameters.Num() == 0)
+	{
+		UpdatedSystemParameters.Add(FName(TEXT("System Delta Time")), SYS_PARAM_ENGINE_DELTA_TIME);
+		UpdatedSystemParameters.Add(FName(TEXT("System Inv Delta Time")), SYS_PARAM_ENGINE_INV_DELTA_TIME);
+		UpdatedSystemParameters.Add(FName(TEXT("System Position")), SYS_PARAM_ENGINE_POSITION);
+		UpdatedSystemParameters.Add(FName(TEXT("System Velocity")), SYS_PARAM_ENGINE_VELOCITY);
+		UpdatedSystemParameters.Add(FName(TEXT("System X Axis")), SYS_PARAM_ENGINE_X_AXIS);
+		UpdatedSystemParameters.Add(FName(TEXT("System Y Axis")), SYS_PARAM_ENGINE_Y_AXIS);
+		UpdatedSystemParameters.Add(FName(TEXT("System Z Axis")), SYS_PARAM_ENGINE_Z_AXIS);
+
+		UpdatedSystemParameters.Add(FName(TEXT("System Local To World")), SYS_PARAM_ENGINE_LOCAL_TO_WORLD);
+		UpdatedSystemParameters.Add(FName(TEXT("System World To Local")), SYS_PARAM_ENGINE_WORLD_TO_LOCAL);
+		UpdatedSystemParameters.Add(FName(TEXT("System Local To World Transposed")), SYS_PARAM_ENGINE_LOCAL_TO_WORLD_TRANSPOSED);
+		UpdatedSystemParameters.Add(FName(TEXT("System World To Local Transposed")), SYS_PARAM_ENGINE_WORLD_TO_LOCAL_TRANSPOSED);
+		UpdatedSystemParameters.Add(FName(TEXT("System Local To World No Scale")), SYS_PARAM_ENGINE_LOCAL_TO_WORLD_NO_SCALE);
+		UpdatedSystemParameters.Add(FName(TEXT("System World To Local No Scale")), SYS_PARAM_ENGINE_WORLD_TO_LOCAL_NO_SCALE);
+
+		UpdatedSystemParameters.Add(FName(TEXT("Emitter Execution Count")), SYS_PARAM_ENGINE_EXEC_COUNT);
+		UpdatedSystemParameters.Add(FName(TEXT("Emitter Age")), SYS_PARAM_EMITTER_AGE);
+		UpdatedSystemParameters.Add(FName(TEXT("Emitter Local Space")), SYS_PARAM_EMITTER_LOCALSPACE);
+		UpdatedSystemParameters.Add(FName(TEXT("Emitter Spawn Rate")), SYS_PARAM_EMITTER_SPAWNRATE);
+		UpdatedSystemParameters.Add(FName(TEXT("Emitter Spawn Interval")), SYS_PARAM_EMITTER_SPAWN_INTERVAL);
+		UpdatedSystemParameters.Add(FName(TEXT("Emitter Interp Spawn Start Dt")), SYS_PARAM_EMITTER_INTERP_SPAWN_START_DT);
+
+		UpdatedSystemParameters.Add(FName(TEXT("Delta Time")), SYS_PARAM_ENGINE_DELTA_TIME);
+		UpdatedSystemParameters.Add(FName(TEXT("Emitter Age")), SYS_PARAM_EMITTER_AGE);
+		UpdatedSystemParameters.Add(FName(TEXT("Emitter Local Space")), SYS_PARAM_EMITTER_LOCALSPACE);
+		UpdatedSystemParameters.Add(FName(TEXT("Effect Position")), SYS_PARAM_ENGINE_POSITION);
+		UpdatedSystemParameters.Add(FName(TEXT("Effect Velocity")), SYS_PARAM_ENGINE_VELOCITY);
+		UpdatedSystemParameters.Add(FName(TEXT("Effect X Axis")), SYS_PARAM_ENGINE_X_AXIS);
+		UpdatedSystemParameters.Add(FName(TEXT("Effect Y Axis")), SYS_PARAM_ENGINE_Y_AXIS);
+		UpdatedSystemParameters.Add(FName(TEXT("Effect Z Axis")), SYS_PARAM_ENGINE_Z_AXIS);
+
+		UpdatedSystemParameters.Add(FName(TEXT("Effect Local To World")), SYS_PARAM_ENGINE_LOCAL_TO_WORLD);
+		UpdatedSystemParameters.Add(FName(TEXT("Effect World To Local")), SYS_PARAM_ENGINE_WORLD_TO_LOCAL);
+		UpdatedSystemParameters.Add(FName(TEXT("Effect Local To World Transposed")), SYS_PARAM_ENGINE_LOCAL_TO_WORLD_TRANSPOSED);
+		UpdatedSystemParameters.Add(FName(TEXT("Effect World To Local Transposed")), SYS_PARAM_ENGINE_WORLD_TO_LOCAL_TRANSPOSED);
+		UpdatedSystemParameters.Add(FName(TEXT("Execution Count")), SYS_PARAM_ENGINE_EXEC_COUNT);
+		UpdatedSystemParameters.Add(FName(TEXT("Spawn Rate")), SYS_PARAM_EMITTER_SPAWNRATE);
+		UpdatedSystemParameters.Add(FName(TEXT("Spawn Interval")), SYS_PARAM_EMITTER_SPAWN_INTERVAL);
+		UpdatedSystemParameters.Add(FName(TEXT("Interp Spawn Start Dt")), SYS_PARAM_EMITTER_INTERP_SPAWN_START_DT);
+		UpdatedSystemParameters.Add(FName(TEXT("Inv Delta Time")), SYS_PARAM_ENGINE_INV_DELTA_TIME);
+	}
+
+	if (SystemStrMap.Num() == 0)
+	{
+		SystemStrMap.Add(SYS_PARAM_ENGINE_DELTA_TIME, LOCTEXT("EngineDeltaTimeDesc", "Time in seconds since the last tick."));
+		SystemStrMap.Add(SYS_PARAM_ENGINE_INV_DELTA_TIME, LOCTEXT("EngineInvDeltaTimeDesc", "One over Engine.DeltaTime"));
+		SystemStrMap.Add(SYS_PARAM_ENGINE_TIME, LOCTEXT("EngineTimeDesc", "Time in seconds since level began play, but IS paused when the game is paused, and IS dilated/clamped."));
+		SystemStrMap.Add(SYS_PARAM_ENGINE_REAL_TIME, LOCTEXT("EngineRealTimeDesc", "Time in seconds since level began play, but IS NOT paused when the game is paused, and IS NOT dilated/clamped."));
+		SystemStrMap.Add(SYS_PARAM_ENGINE_SYSTEM_AGE, LOCTEXT("EngineSystemTimeDesc", "Time in seconds since the system was first created. Managed by the NiagaraSystemInstance in code."));
+
+		SystemStrMap.Add(SYS_PARAM_ENGINE_POSITION, LOCTEXT("EnginePositionDesc", "The owning component's position in world space."));
+		SystemStrMap.Add(SYS_PARAM_ENGINE_SCALE, LOCTEXT("EngineScaleDesc", "The owning component's scale in world space."));
+		SystemStrMap.Add(SYS_PARAM_ENGINE_VELOCITY, LOCTEXT("EngineVelocityDesc", "The owning component's velocity in world space."));
+		SystemStrMap.Add(SYS_PARAM_ENGINE_X_AXIS, LOCTEXT("XAxisDesc", "The X-axis of the owning component."));
+		SystemStrMap.Add(SYS_PARAM_ENGINE_Y_AXIS, LOCTEXT("YAxisDesc", "The Y-axis of the owning component."));
+		SystemStrMap.Add(SYS_PARAM_ENGINE_Z_AXIS, LOCTEXT("ZAxisDesc", "The Z-axis of the owning component."));
+
+		SystemStrMap.Add(SYS_PARAM_ENGINE_LOCAL_TO_WORLD, LOCTEXT("LocalToWorldDesc", "Owning component's local space to world space transform matrix."));
+		SystemStrMap.Add(SYS_PARAM_ENGINE_WORLD_TO_LOCAL, LOCTEXT("WorldToLocalDesc", "Owning component's world space to local space transform matrix."));
+		SystemStrMap.Add(SYS_PARAM_ENGINE_LOCAL_TO_WORLD_TRANSPOSED, LOCTEXT("LocalToWorldTransposeDesc", "Owning component's local space to world space transform matrix transposed."));
+		SystemStrMap.Add(SYS_PARAM_ENGINE_WORLD_TO_LOCAL_TRANSPOSED, LOCTEXT("WorldToLocalTransposeDesc", "Owning component's world space to local space transform matrix transposed."));
+		SystemStrMap.Add(SYS_PARAM_ENGINE_LOCAL_TO_WORLD_NO_SCALE, LOCTEXT("LocalToWorldNoScaleDesc", "Owning component's local space to world space transform matrix with scaling removed."));
+		SystemStrMap.Add(SYS_PARAM_ENGINE_WORLD_TO_LOCAL_NO_SCALE, LOCTEXT("WorldToLocalNoScaleDesc", "Owning component's world space to local space transform matrix with scaling removed."));
+
+		SystemStrMap.Add(SYS_PARAM_ENGINE_TIME_SINCE_RENDERED, LOCTEXT("TimeSinceRendered", "The time in seconds that have passed since this system was last rendered."));
+		SystemStrMap.Add(SYS_PARAM_ENGINE_MIN_DIST_TO_CAMERA, LOCTEXT("MinDistanceToCamera", "The distance from the owner component to the nearest local player viewpoint."));
+		SystemStrMap.Add(SYS_PARAM_ENGINE_EXECUTION_STATE, LOCTEXT("ExecutionState", "The execution state of the systems owner. Takes precedence over the systems internal execution state."));
+
+		SystemStrMap.Add(SYS_PARAM_ENGINE_EXEC_COUNT, LOCTEXT("ExecCountDesc", "The index of this particle in the read buffer."));
+		SystemStrMap.Add(SYS_PARAM_ENGINE_EMITTER_NUM_PARTICLES, LOCTEXT("EmitterNumParticles", "The number of particles for this emitter at the beginning of simulation. Should only be used in Emitter scripts."));
+		SystemStrMap.Add(SYS_PARAM_ENGINE_SYSTEM_NUM_EMITTERS_ALIVE, LOCTEXT("SystemNumEmittersAlive", "The number of emitters still alive attached to this system. Should only be used in System scripts."));
+		SystemStrMap.Add(SYS_PARAM_ENGINE_SYSTEM_NUM_EMITTERS, LOCTEXT("SystemNumEmitters", "The number of emitters attached to this system. Should only be used in System scripts."));
+		SystemStrMap.Add(SYS_PARAM_ENGINE_NUM_SYSTEM_INSTANCES, LOCTEXT("SystemNumInstances", "The number of instances of the this system currently ticking. Should only be used in System scripts."));
+		SystemStrMap.Add(SYS_PARAM_ENGINE_GLOBAL_SPAWN_COUNT_SCALE, LOCTEXT("GlobalSpawnCountScale", "Global Spawn Count Scale. Should only be used in System scripts."));
+		SystemStrMap.Add(SYS_PARAM_ENGINE_GLOBAL_SYSTEM_COUNT_SCALE, LOCTEXT("GlobalSystemCountScale", "Global System Count Scale. Should only be used in System scripts."));
+	}
+
+	if (Attributes.Num() == 0)
+	{
+		Attributes.Add(SYS_PARAM_PARTICLES_ID);
+		Attributes.Add(SYS_PARAM_PARTICLES_POSITION);
+		Attributes.Add(SYS_PARAM_PARTICLES_VELOCITY);
+		Attributes.Add(SYS_PARAM_PARTICLES_COLOR);
+		Attributes.Add(SYS_PARAM_PARTICLES_SPRITE_ROTATION);
+		Attributes.Add(SYS_PARAM_PARTICLES_NORMALIZED_AGE);
+		Attributes.Add(SYS_PARAM_PARTICLES_SPRITE_SIZE);
+		Attributes.Add(SYS_PARAM_PARTICLES_SPRITE_FACING);
+		Attributes.Add(SYS_PARAM_PARTICLES_SPRITE_ALIGNMENT);
+		Attributes.Add(SYS_PARAM_PARTICLES_SUB_IMAGE_INDEX);
+		Attributes.Add(SYS_PARAM_PARTICLES_DYNAMIC_MATERIAL_PARAM);
+		Attributes.Add(SYS_PARAM_PARTICLES_DYNAMIC_MATERIAL_PARAM_1);
+		Attributes.Add(SYS_PARAM_PARTICLES_DYNAMIC_MATERIAL_PARAM_2);
+		Attributes.Add(SYS_PARAM_PARTICLES_DYNAMIC_MATERIAL_PARAM_3);
+		Attributes.Add(SYS_PARAM_PARTICLES_SCALE);
+		Attributes.Add(SYS_PARAM_PARTICLES_LIFETIME);
+		Attributes.Add(SYS_PARAM_PARTICLES_MESH_ORIENTATION);
+		Attributes.Add(SYS_PARAM_PARTICLES_CAMERA_OFFSET);
+		Attributes.Add(SYS_PARAM_PARTICLES_UV_SCALE);
+		Attributes.Add(SYS_PARAM_PARTICLES_MATERIAL_RANDOM);
+		Attributes.Add(SYS_PARAM_PARTICLES_LIGHT_RADIUS);
+		Attributes.Add(SYS_PARAM_PARTICLES_RIBBONID);
+		Attributes.Add(SYS_PARAM_PARTICLES_RIBBONWIDTH);
+		Attributes.Add(SYS_PARAM_PARTICLES_RIBBONTWIST);
+		Attributes.Add(SYS_PARAM_PARTICLES_RIBBONFACING);
+		Attributes.Add(SYS_PARAM_PARTICLES_RIBBONLINKORDER);
+
+
+		Attributes.Add(SYS_PARAM_INSTANCE_ALIVE);
+	}
+
+	if (AttrDataSetKeyMap.Num() == 0)
+	{
+		AttrDataSetKeyMap.Add(SYS_PARAM_PARTICLES_POSITION, GetAttributeAsDataSetKey(SYS_PARAM_PARTICLES_POSITION));
+		AttrDataSetKeyMap.Add(SYS_PARAM_PARTICLES_VELOCITY, GetAttributeAsDataSetKey(SYS_PARAM_PARTICLES_VELOCITY));
+		AttrDataSetKeyMap.Add(SYS_PARAM_PARTICLES_COLOR, GetAttributeAsDataSetKey(SYS_PARAM_PARTICLES_COLOR));
+		AttrDataSetKeyMap.Add(SYS_PARAM_PARTICLES_SPRITE_ROTATION, GetAttributeAsDataSetKey(SYS_PARAM_PARTICLES_SPRITE_ROTATION));
+		AttrDataSetKeyMap.Add(SYS_PARAM_PARTICLES_NORMALIZED_AGE, GetAttributeAsDataSetKey(SYS_PARAM_PARTICLES_NORMALIZED_AGE));
+		AttrDataSetKeyMap.Add(SYS_PARAM_PARTICLES_SPRITE_SIZE, GetAttributeAsDataSetKey(SYS_PARAM_PARTICLES_SPRITE_SIZE));
+		AttrDataSetKeyMap.Add(SYS_PARAM_PARTICLES_SPRITE_FACING, GetAttributeAsDataSetKey(SYS_PARAM_PARTICLES_SPRITE_FACING));
+		AttrDataSetKeyMap.Add(SYS_PARAM_PARTICLES_SPRITE_ALIGNMENT, GetAttributeAsDataSetKey(SYS_PARAM_PARTICLES_SPRITE_ALIGNMENT));
+		AttrDataSetKeyMap.Add(SYS_PARAM_PARTICLES_SUB_IMAGE_INDEX, GetAttributeAsDataSetKey(SYS_PARAM_PARTICLES_SUB_IMAGE_INDEX));
+		AttrDataSetKeyMap.Add(SYS_PARAM_PARTICLES_DYNAMIC_MATERIAL_PARAM, GetAttributeAsDataSetKey(SYS_PARAM_PARTICLES_DYNAMIC_MATERIAL_PARAM));
+		AttrDataSetKeyMap.Add(SYS_PARAM_PARTICLES_DYNAMIC_MATERIAL_PARAM_1, GetAttributeAsDataSetKey(SYS_PARAM_PARTICLES_DYNAMIC_MATERIAL_PARAM_1));
+		AttrDataSetKeyMap.Add(SYS_PARAM_PARTICLES_DYNAMIC_MATERIAL_PARAM_2, GetAttributeAsDataSetKey(SYS_PARAM_PARTICLES_DYNAMIC_MATERIAL_PARAM_2));
+		AttrDataSetKeyMap.Add(SYS_PARAM_PARTICLES_DYNAMIC_MATERIAL_PARAM_3, GetAttributeAsDataSetKey(SYS_PARAM_PARTICLES_DYNAMIC_MATERIAL_PARAM_3));
+		AttrDataSetKeyMap.Add(SYS_PARAM_PARTICLES_SCALE, GetAttributeAsDataSetKey(SYS_PARAM_PARTICLES_SCALE));
+		AttrDataSetKeyMap.Add(SYS_PARAM_PARTICLES_LIFETIME, GetAttributeAsDataSetKey(SYS_PARAM_PARTICLES_LIFETIME));
+		AttrDataSetKeyMap.Add(SYS_PARAM_PARTICLES_MESH_ORIENTATION, GetAttributeAsDataSetKey(SYS_PARAM_PARTICLES_MESH_ORIENTATION));
+		AttrDataSetKeyMap.Add(SYS_PARAM_PARTICLES_CAMERA_OFFSET, GetAttributeAsDataSetKey(SYS_PARAM_PARTICLES_CAMERA_OFFSET));
+		AttrDataSetKeyMap.Add(SYS_PARAM_PARTICLES_UV_SCALE, GetAttributeAsDataSetKey(SYS_PARAM_PARTICLES_UV_SCALE));
+		AttrDataSetKeyMap.Add(SYS_PARAM_PARTICLES_MATERIAL_RANDOM, GetAttributeAsDataSetKey(SYS_PARAM_PARTICLES_MATERIAL_RANDOM));
+		AttrDataSetKeyMap.Add(SYS_PARAM_PARTICLES_LIGHT_RADIUS, GetAttributeAsDataSetKey(SYS_PARAM_PARTICLES_LIGHT_RADIUS));
+		AttrDataSetKeyMap.Add(SYS_PARAM_PARTICLES_RIBBONID, GetAttributeAsDataSetKey(SYS_PARAM_PARTICLES_RIBBONID));
+		AttrDataSetKeyMap.Add(SYS_PARAM_PARTICLES_RIBBONWIDTH, GetAttributeAsDataSetKey(SYS_PARAM_PARTICLES_RIBBONWIDTH));
+		AttrDataSetKeyMap.Add(SYS_PARAM_PARTICLES_RIBBONTWIST, GetAttributeAsDataSetKey(SYS_PARAM_PARTICLES_RIBBONTWIST));
+		AttrDataSetKeyMap.Add(SYS_PARAM_PARTICLES_RIBBONFACING, GetAttributeAsDataSetKey(SYS_PARAM_PARTICLES_RIBBONFACING));
+		AttrDataSetKeyMap.Add(SYS_PARAM_PARTICLES_RIBBONLINKORDER, GetAttributeAsDataSetKey(SYS_PARAM_PARTICLES_RIBBONLINKORDER));
+	}
+
+	if (AttrDefaultsStrMap.Num() == 0)
+	{
+		FNiagaraVariable Var;
+		AttrDefaultsStrMap.Add(SYS_PARAM_PARTICLES_POSITION, SYS_PARAM_ENGINE_POSITION.GetName().ToString());
+		AttrDefaultsValueMap.Add(SYS_PARAM_PARTICLES_POSITION, SYS_PARAM_ENGINE_POSITION);
+		
+		AttrDefaultsStrMap.Add(SYS_PARAM_PARTICLES_VELOCITY, TEXT("0.0,0.0,0.0"));
+		Var = SYS_PARAM_PARTICLES_VELOCITY;
+		Var.SetValue<FVector>(FVector(0.0f, 0.0f, 0.0f));
+		AttrDefaultsValueMap.Add(SYS_PARAM_PARTICLES_VELOCITY, Var);
+
+		AttrDefaultsStrMap.Add(SYS_PARAM_PARTICLES_COLOR, FLinearColor(1.0f, 1.0f, 1.0f, 1.0f).ToString());
+		Var = SYS_PARAM_PARTICLES_COLOR;
+		Var.SetValue<FLinearColor>(FLinearColor(1.0f, 1.0f, 1.0f, 1.0f));
+		AttrDefaultsValueMap.Add(SYS_PARAM_PARTICLES_COLOR, Var);
+
+		AttrDefaultsStrMap.Add(SYS_PARAM_PARTICLES_SPRITE_ROTATION, TEXT("0.0"));
+		Var = SYS_PARAM_PARTICLES_SPRITE_ROTATION;
+		Var.SetValue<float>(0.0f);
+		AttrDefaultsValueMap.Add(SYS_PARAM_PARTICLES_SPRITE_ROTATION, Var);
+
+		AttrDefaultsStrMap.Add(SYS_PARAM_PARTICLES_NORMALIZED_AGE, TEXT("0.0"));
+		Var = SYS_PARAM_PARTICLES_NORMALIZED_AGE;
+		Var.SetValue<float>(0.0f);
+		AttrDefaultsValueMap.Add(SYS_PARAM_PARTICLES_NORMALIZED_AGE, Var);
+
+		AttrDefaultsStrMap.Add(SYS_PARAM_PARTICLES_SPRITE_SIZE, TEXT("X=50.0 Y=50.0"));
+		Var = SYS_PARAM_PARTICLES_SPRITE_SIZE;
+		Var.SetValue<FVector2D>(FVector2D(50.0f, 50.0f));
+		AttrDefaultsValueMap.Add(SYS_PARAM_PARTICLES_SPRITE_SIZE, Var);
+
+		AttrDefaultsStrMap.Add(SYS_PARAM_PARTICLES_SPRITE_FACING, TEXT("1.0,0.0,0.0"));
+		Var = SYS_PARAM_PARTICLES_SPRITE_FACING;
+		Var.SetValue<FVector>(FVector(1.0f, 0.0f, 0.0f));
+		AttrDefaultsValueMap.Add(SYS_PARAM_PARTICLES_SPRITE_FACING, Var);
+
+		AttrDefaultsStrMap.Add(SYS_PARAM_PARTICLES_SPRITE_ALIGNMENT, TEXT("1.0,0.0,0.0"));
+		Var = SYS_PARAM_PARTICLES_SPRITE_ALIGNMENT;
+		Var.SetValue<FVector>(FVector(1.0f, 0.0f, 0.0f));
+		AttrDefaultsValueMap.Add(SYS_PARAM_PARTICLES_SPRITE_ALIGNMENT, Var);
+		
+		AttrDefaultsStrMap.Add(SYS_PARAM_PARTICLES_SUB_IMAGE_INDEX, TEXT("0.0"));
+		Var = SYS_PARAM_PARTICLES_SUB_IMAGE_INDEX;
+		Var.SetValue<float>(0.0f);
+		AttrDefaultsValueMap.Add(SYS_PARAM_PARTICLES_SUB_IMAGE_INDEX, Var);
+
+		AttrDefaultsStrMap.Add(SYS_PARAM_PARTICLES_DYNAMIC_MATERIAL_PARAM, TEXT("1.0,1.0,1.0,1.0"));
+		Var = SYS_PARAM_PARTICLES_DYNAMIC_MATERIAL_PARAM;
+		Var.SetValue<FVector4>(FVector4(1.0f, 1.0f, 1.0f, 1.0f));
+		AttrDefaultsValueMap.Add(SYS_PARAM_PARTICLES_DYNAMIC_MATERIAL_PARAM, Var);
+
+		AttrDefaultsStrMap.Add(SYS_PARAM_PARTICLES_DYNAMIC_MATERIAL_PARAM_1, TEXT("1.0,1.0,1.0,1.0"));
+		Var = SYS_PARAM_PARTICLES_DYNAMIC_MATERIAL_PARAM_1;
+		Var.SetValue<FVector4>(FVector4(1.0f, 1.0f, 1.0f, 1.0f));
+		AttrDefaultsValueMap.Add(SYS_PARAM_PARTICLES_DYNAMIC_MATERIAL_PARAM_1, Var);
+
+		AttrDefaultsStrMap.Add(SYS_PARAM_PARTICLES_DYNAMIC_MATERIAL_PARAM_2, TEXT("1.0,1.0,1.0,1.0"));
+		Var = SYS_PARAM_PARTICLES_DYNAMIC_MATERIAL_PARAM_2;
+		Var.SetValue<FVector4>(FVector4(1.0f, 1.0f, 1.0f, 1.0f));
+		AttrDefaultsValueMap.Add(SYS_PARAM_PARTICLES_DYNAMIC_MATERIAL_PARAM_2, Var);
+
+		AttrDefaultsStrMap.Add(SYS_PARAM_PARTICLES_DYNAMIC_MATERIAL_PARAM_3, TEXT("1.0,1.0,1.0,1.0"));
+		Var = SYS_PARAM_PARTICLES_DYNAMIC_MATERIAL_PARAM_3;
+		Var.SetValue<FVector4>(FVector4(1.0f, 1.0f, 1.0f, 1.0f));
+		AttrDefaultsValueMap.Add(SYS_PARAM_PARTICLES_DYNAMIC_MATERIAL_PARAM_3, Var);
+
+		AttrDefaultsStrMap.Add(SYS_PARAM_PARTICLES_SCALE, TEXT("1.0,1.0,1.0"));
+		Var = SYS_PARAM_PARTICLES_SCALE;
+		Var.SetValue<FVector>(FVector(1.0f, 1.0f, 1.0f));
+		AttrDefaultsValueMap.Add(SYS_PARAM_PARTICLES_SCALE, Var);
+
+		AttrDefaultsStrMap.Add(SYS_PARAM_PARTICLES_LIFETIME, TEXT("5.0"));
+		Var = SYS_PARAM_PARTICLES_LIFETIME;
+		Var.SetValue<float>(5.0f);
+		AttrDefaultsValueMap.Add(SYS_PARAM_PARTICLES_LIFETIME, Var);
+
+		AttrDefaultsStrMap.Add(SYS_PARAM_PARTICLES_MESH_ORIENTATION, TEXT("0.0,0.0,0.0,1.0"));
+		Var = SYS_PARAM_PARTICLES_MESH_ORIENTATION;
+		Var.SetValue<FQuat>(FQuat::Identity);
+		AttrDefaultsValueMap.Add(SYS_PARAM_PARTICLES_MESH_ORIENTATION, Var);
+		
+		AttrDefaultsStrMap.Add(SYS_PARAM_PARTICLES_CAMERA_OFFSET, TEXT("0.0"));
+		Var = SYS_PARAM_PARTICLES_CAMERA_OFFSET;
+		Var.SetValue<float>(0.0f);
+		AttrDefaultsValueMap.Add(SYS_PARAM_PARTICLES_CAMERA_OFFSET, Var);
+
+		AttrDefaultsStrMap.Add(SYS_PARAM_PARTICLES_UV_SCALE, TEXT("X=1.0 Y=1.0"));
+		Var = SYS_PARAM_PARTICLES_UV_SCALE;
+		Var.SetValue<FVector2D>(FVector2D(1.0f, 1.0f));
+		AttrDefaultsValueMap.Add(SYS_PARAM_PARTICLES_UV_SCALE, Var);
+
+		AttrDefaultsStrMap.Add(SYS_PARAM_PARTICLES_MATERIAL_RANDOM, TEXT("0.0"));
+		Var = SYS_PARAM_PARTICLES_MATERIAL_RANDOM;
+		Var.SetValue<float>(0.0f);
+		AttrDefaultsValueMap.Add(SYS_PARAM_PARTICLES_MATERIAL_RANDOM, Var);
+
+		AttrDefaultsStrMap.Add(SYS_PARAM_PARTICLES_LIGHT_RADIUS, TEXT("100.0"));
+		Var = SYS_PARAM_PARTICLES_LIGHT_RADIUS;
+		Var.SetValue<float>(100.0f);
+		AttrDefaultsValueMap.Add(SYS_PARAM_PARTICLES_LIGHT_RADIUS, Var);
+
+		AttrDefaultsStrMap.Add(SYS_PARAM_PARTICLES_RIBBONID, TEXT("0"));
+		Var = SYS_PARAM_PARTICLES_RIBBONID;
+		Var.SetValue<FNiagaraID>(FNiagaraID());
+		AttrDefaultsValueMap.Add(SYS_PARAM_PARTICLES_RIBBONID, Var);
+
+		AttrDefaultsStrMap.Add(SYS_PARAM_PARTICLES_RIBBONWIDTH, TEXT("1.0"));
+		Var = SYS_PARAM_PARTICLES_RIBBONWIDTH;
+		Var.SetValue<float>(1.0f);
+		AttrDefaultsValueMap.Add(SYS_PARAM_PARTICLES_RIBBONWIDTH, Var);
+
+		AttrDefaultsStrMap.Add(SYS_PARAM_PARTICLES_RIBBONTWIST, TEXT("0.0"));
+		Var = SYS_PARAM_PARTICLES_RIBBONTWIST;
+		Var.SetValue<float>(0.0f);
+		AttrDefaultsValueMap.Add(SYS_PARAM_PARTICLES_RIBBONTWIST, Var);
+
+		AttrDefaultsStrMap.Add(SYS_PARAM_PARTICLES_RIBBONFACING, TEXT("0.0, 0.0, 1.0"));
+		Var = SYS_PARAM_PARTICLES_RIBBONFACING;
+		Var.SetValue<FVector>(FVector(0.0f, 0.0f, 1.0f));
+		AttrDefaultsValueMap.Add(SYS_PARAM_PARTICLES_RIBBONFACING, Var);
+
+		AttrDefaultsStrMap.Add(SYS_PARAM_PARTICLES_RIBBONLINKORDER, TEXT("0"));
+		Var = SYS_PARAM_PARTICLES_RIBBONLINKORDER;
+		Var.SetValue<float>(0.0f);
+		AttrDefaultsValueMap.Add(SYS_PARAM_PARTICLES_RIBBONLINKORDER, Var);
+	}
+
+	if (AttrDescStrMap.Num() == 0)
+	{
+		AttrDescStrMap.Add(SYS_PARAM_PARTICLES_POSITION, LOCTEXT("PositionDesc", "The position of the particle."));
+		AttrDescStrMap.Add(SYS_PARAM_PARTICLES_VELOCITY, LOCTEXT("VelocityDesc", "The velocity in cm/s of the particle."));
+		AttrDescStrMap.Add(SYS_PARAM_PARTICLES_COLOR, LOCTEXT("ColorDesc", "The color of the particle."));
+		AttrDescStrMap.Add(SYS_PARAM_PARTICLES_SPRITE_ROTATION, LOCTEXT("SpriteRotDesc", "The screen aligned roll of the particle in degrees."));
+		AttrDescStrMap.Add(SYS_PARAM_PARTICLES_NORMALIZED_AGE, LOCTEXT("NormalizedAgeDesc", "The age in seconds divided by lifetime in seconds. Useful for animation as the value is between 0 and 1."));
+		AttrDescStrMap.Add(SYS_PARAM_PARTICLES_SPRITE_SIZE, LOCTEXT("SpriteSizeDesc", "The size of the sprite quad."));
+		AttrDescStrMap.Add(SYS_PARAM_PARTICLES_SPRITE_FACING, LOCTEXT("FacingDesc", "Makes the surface of the sprite face towards a custom vector. Must be used with the SpriteRenderer's CustomFacingVector FacingMode and CustomFacingVectorMask options."));
+		AttrDescStrMap.Add(SYS_PARAM_PARTICLES_SPRITE_ALIGNMENT, LOCTEXT("AlignmentDesc", "Imagine the texture having an arrow pointing up, this attribute makes the arrow point towards the alignment axis. Must be used with the SpriteRenderer's CustomAlignment Alignment option."));
+		AttrDescStrMap.Add(SYS_PARAM_PARTICLES_SUB_IMAGE_INDEX, LOCTEXT("SubImageIndexDesc", "A value from 0 to the number of entries in the table of SubUV images."));
+		FText DynParamText = LOCTEXT("DynamicMaterialParameterDesc", "The 4-float vector used to send custom data to renderer.");
+		AttrDescStrMap.Add(SYS_PARAM_PARTICLES_DYNAMIC_MATERIAL_PARAM, DynParamText);
+		AttrDescStrMap.Add(SYS_PARAM_PARTICLES_DYNAMIC_MATERIAL_PARAM_1, DynParamText);
+		AttrDescStrMap.Add(SYS_PARAM_PARTICLES_DYNAMIC_MATERIAL_PARAM_2, DynParamText);
+		AttrDescStrMap.Add(SYS_PARAM_PARTICLES_DYNAMIC_MATERIAL_PARAM_3, DynParamText);
+		AttrDescStrMap.Add(SYS_PARAM_PARTICLES_SCALE, LOCTEXT("ScaleParamDesc", "The XYZ scale of the non-sprite based particle."));
+		AttrDescStrMap.Add(SYS_PARAM_PARTICLES_LIFETIME, LOCTEXT("LifetimeParamDesc", "The lifetime of a particle in seconds."));
+		AttrDescStrMap.Add(SYS_PARAM_PARTICLES_MESH_ORIENTATION, LOCTEXT("MeshOrientParamDesc", "The axis-angle rotation to be applied to the mesh particle."));
+		AttrDescStrMap.Add(SYS_PARAM_PARTICLES_CAMERA_OFFSET, LOCTEXT("CamOffsetParamDesc", "Used to offset position in the direction of the camera. The value is multiplied by the direction vector from the camera to the particle."));
+		AttrDescStrMap.Add(SYS_PARAM_PARTICLES_UV_SCALE, LOCTEXT("UVScalerParamDesc", "Used to multiply the generated UVs for Sprite renderers."));
+		AttrDescStrMap.Add(SYS_PARAM_PARTICLES_MATERIAL_RANDOM, LOCTEXT("MaterialRandomParamDesc", "Used to drive the Particle Random node in the Material Editor. Without this set, any Particle Randoms will get 0.0."));
+		AttrDescStrMap.Add(SYS_PARAM_PARTICLES_LIGHT_RADIUS, LOCTEXT("LightRadiusParamDesc", "Used to drive the radius of the light when using a Light renderer."));
+		AttrDescStrMap.Add(SYS_PARAM_INSTANCE_ALIVE, LOCTEXT("AliveParamDesc", "Used to determine whether or not this particle instance is still valid or if it can be deleted."));
+		AttrDescStrMap.Add(SYS_PARAM_PARTICLES_RIBBONID, LOCTEXT("RibbonIDDesc", "Sets the ribbon id for a particle. Particles with the same ribbon id will be connected into a ribbon."));
+		AttrDescStrMap.Add(SYS_PARAM_PARTICLES_RIBBONWIDTH, LOCTEXT("RibbonWidthDesc", "Sets the ribbon width for a particle, in UE4 units."));
+		AttrDescStrMap.Add(SYS_PARAM_PARTICLES_RIBBONTWIST, LOCTEXT("RibbonTwistDesc", "Sets the ribbon twist for a particle, in degrees."));
+		AttrDescStrMap.Add(SYS_PARAM_PARTICLES_RIBBONFACING, LOCTEXT("RibbonFacingDesc", "Sets the facing vector of the ribbon at the particle position."));
+		AttrDescStrMap.Add(SYS_PARAM_PARTICLES_RIBBONLINKORDER, LOCTEXT("RibbonLinkOrderDesc", "Explicit order for linking particles within a ribbon. Particles of the same ribbon id will be connected into a ribbon in incrementing order of this attribute value."));
+		AttrDescStrMap.Add(SYS_PARAM_PARTICLES_ID, LOCTEXT("IDDesc", "Engine managed particle attribute that is a persistent ID for each particle."));
+	}
+
+	if (AttrMetaData.Num() == 0)
+	{
+		// Add the engine default attributes..
+		{
+			for (FNiagaraVariable Var : Attributes)
+			{
+				FNiagaraVariableMetaData MetaData;
+				MetaData.Description = GetAttributeDescription(Var);
+				AttrMetaData.Add(Var, MetaData);
+			}
+		}
+
+		// Add the engine constants..
+		{
+			for (FNiagaraVariable Var : SystemParameters)
+			{
+				FNiagaraVariableMetaData MetaData;
+				MetaData.Description = GetEngineConstantDescription(Var);
+				AttrMetaData.Add(Var, MetaData);
+			}
+		}
+
+	}
+
+	if (EngineManagedAttributes.Num() == 0)
+	{
+		EngineManagedAttributes.Add(SYS_PARAM_PARTICLES_ID);
+	}
+
+}
+
+const TArray<FNiagaraVariable>& FNiagaraConstants::GetEngineConstants()
+{
+	check(SystemParameters.Num() != 0);
 	return SystemParameters;
 }
 
+const TArray<FNiagaraVariable>& FNiagaraConstants::GetTranslatorConstants()
+{
+	check(TranslatorParameters.Num() != 0);
+	return TranslatorParameters;
+}
+
+bool FNiagaraConstants::IsEngineManagedAttribute(const FNiagaraVariable& Var)
+{
+	return EngineManagedAttributes.Contains(Var);
+}
+
 FNiagaraVariable FNiagaraConstants::UpdateEngineConstant(const FNiagaraVariable& InVar)
 {
 	const FNiagaraVariable* FoundSystemVar = FindEngineConstant(InVar);
@@ -49,51 +421,7 @@ FNiagaraVariable FNiagaraConstants::UpdateEngineConstant(const FNiagaraVariable&
 	}
 	else
 	{
-		static TMap<FName, FNiagaraVariable> UpdatedSystemParameters;
-		if (UpdatedSystemParameters.Num() == 0)
-		{
-			UpdatedSystemParameters.Add(FName(TEXT("System Delta Time")), SYS_PARAM_ENGINE_DELTA_TIME);
-			UpdatedSystemParameters.Add(FName(TEXT("System Inv Delta Time")), SYS_PARAM_ENGINE_INV_DELTA_TIME);
-			UpdatedSystemParameters.Add(FName(TEXT("System Position")), SYS_PARAM_ENGINE_POSITION);
-			UpdatedSystemParameters.Add(FName(TEXT("System Velocity")), SYS_PARAM_ENGINE_VELOCITY);
-			UpdatedSystemParameters.Add(FName(TEXT("System X Axis")), SYS_PARAM_ENGINE_X_AXIS);
-			UpdatedSystemParameters.Add(FName(TEXT("System Y Axis")), SYS_PARAM_ENGINE_Y_AXIS);
-			UpdatedSystemParameters.Add(FName(TEXT("System Z Axis")), SYS_PARAM_ENGINE_Z_AXIS);
-
-			UpdatedSystemParameters.Add(FName(TEXT("System Local To World")), SYS_PARAM_ENGINE_LOCAL_TO_WORLD);
-			UpdatedSystemParameters.Add(FName(TEXT("System World To Local")), SYS_PARAM_ENGINE_WORLD_TO_LOCAL);
-			UpdatedSystemParameters.Add(FName(TEXT("System Local To World Transposed")), SYS_PARAM_ENGINE_LOCAL_TO_WORLD_TRANSPOSED);
-			UpdatedSystemParameters.Add(FName(TEXT("System World To Local Transposed")), SYS_PARAM_ENGINE_WORLD_TO_LOCAL_TRANSPOSED);
-			UpdatedSystemParameters.Add(FName(TEXT("System Local To World No Scale")), SYS_PARAM_ENGINE_LOCAL_TO_WORLD_NO_SCALE);
-			UpdatedSystemParameters.Add(FName(TEXT("System World To Local No Scale")), SYS_PARAM_ENGINE_WORLD_TO_LOCAL_NO_SCALE);
-
-			UpdatedSystemParameters.Add(FName(TEXT("Emitter Execution Count")), SYS_PARAM_ENGINE_EXEC_COUNT);
-			UpdatedSystemParameters.Add(FName(TEXT("Emitter Age")), SYS_PARAM_EMITTER_AGE);
-			UpdatedSystemParameters.Add(FName(TEXT("Emitter Local Space")), SYS_PARAM_EMITTER_LOCALSPACE);
-			UpdatedSystemParameters.Add(FName(TEXT("Emitter Spawn Rate")), SYS_PARAM_EMITTER_SPAWNRATE);
-			UpdatedSystemParameters.Add(FName(TEXT("Emitter Spawn Interval")), SYS_PARAM_EMITTER_SPAWN_INTERVAL);
-			UpdatedSystemParameters.Add(FName(TEXT("Emitter Interp Spawn Start Dt")), SYS_PARAM_EMITTER_INTERP_SPAWN_START_DT);
-
-			UpdatedSystemParameters.Add(FName(TEXT("Delta Time")), SYS_PARAM_ENGINE_DELTA_TIME);
-			UpdatedSystemParameters.Add(FName(TEXT("Emitter Age")), SYS_PARAM_EMITTER_AGE);
-			UpdatedSystemParameters.Add(FName(TEXT("Emitter Local Space")), SYS_PARAM_EMITTER_LOCALSPACE);
-			UpdatedSystemParameters.Add(FName(TEXT("Effect Position")), SYS_PARAM_ENGINE_POSITION);
-			UpdatedSystemParameters.Add(FName(TEXT("Effect Velocity")), SYS_PARAM_ENGINE_VELOCITY);
-			UpdatedSystemParameters.Add(FName(TEXT("Effect X Axis")), SYS_PARAM_ENGINE_X_AXIS);
-			UpdatedSystemParameters.Add(FName(TEXT("Effect Y Axis")), SYS_PARAM_ENGINE_Y_AXIS);
-			UpdatedSystemParameters.Add(FName(TEXT("Effect Z Axis")), SYS_PARAM_ENGINE_Z_AXIS);
-
-			UpdatedSystemParameters.Add(FName(TEXT("Effect Local To World")), SYS_PARAM_ENGINE_LOCAL_TO_WORLD);
-			UpdatedSystemParameters.Add(FName(TEXT("Effect World To Local")), SYS_PARAM_ENGINE_WORLD_TO_LOCAL);
-			UpdatedSystemParameters.Add(FName(TEXT("Effect Local To World Transposed")), SYS_PARAM_ENGINE_LOCAL_TO_WORLD_TRANSPOSED);
-			UpdatedSystemParameters.Add(FName(TEXT("Effect World To Local Transposed")), SYS_PARAM_ENGINE_WORLD_TO_LOCAL_TRANSPOSED);
-			UpdatedSystemParameters.Add(FName(TEXT("Execution Count")), SYS_PARAM_ENGINE_EXEC_COUNT);
-			UpdatedSystemParameters.Add(FName(TEXT("Spawn Rate")), SYS_PARAM_EMITTER_SPAWNRATE);
-			UpdatedSystemParameters.Add(FName(TEXT("Spawn Interval")), SYS_PARAM_EMITTER_SPAWN_INTERVAL);
-			UpdatedSystemParameters.Add(FName(TEXT("Interp Spawn Start Dt")), SYS_PARAM_EMITTER_INTERP_SPAWN_START_DT);
-			UpdatedSystemParameters.Add(FName(TEXT("Inv Delta Time")), SYS_PARAM_ENGINE_INV_DELTA_TIME);
-		}
-
+		check(UpdatedSystemParameters.Num() != 0);
 		const FNiagaraVariable* FoundSystemVarUpdate = UpdatedSystemParameters.Find(InVar.GetName());
 		if (FoundSystemVarUpdate != nullptr)
 		{
@@ -107,8 +435,8 @@ FNiagaraVariable FNiagaraConstants::UpdateEngineConstant(const FNiagaraVariable&
 
 const FNiagaraVariable* FNiagaraConstants::FindEngineConstant(const FNiagaraVariable& InVar)
 {
-	const TArray<FNiagaraVariable>& SystemParameters = GetEngineConstants();
-	const FNiagaraVariable* FoundSystemVar = SystemParameters.FindByPredicate([&](const FNiagaraVariable& Var)
+	const TArray<FNiagaraVariable>& LocalSystemParameters = GetEngineConstants();
+	const FNiagaraVariable* FoundSystemVar = LocalSystemParameters.FindByPredicate([&](const FNiagaraVariable& Var)
 	{
 		return Var.GetName() == InVar.GetName();
 	});
@@ -117,34 +445,8 @@ const FNiagaraVariable* FNiagaraConstants::FindEngineConstant(const FNiagaraVari
 
 FText FNiagaraConstants::GetEngineConstantDescription(const FNiagaraVariable& InAttribute)
 {
-	static TMap<FNiagaraVariable, FText> AttrStrMap;
-	if (AttrStrMap.Num() == 0)
-	{
-		AttrStrMap.Add(SYS_PARAM_ENGINE_DELTA_TIME, LOCTEXT("EngineDeltaTimeDesc", "Time in seconds since the last tick."));
-		AttrStrMap.Add(SYS_PARAM_ENGINE_INV_DELTA_TIME, LOCTEXT("EngineInvDeltaTimeDesc", "One over Engine.DeltaTime"));
-		AttrStrMap.Add(SYS_PARAM_ENGINE_POSITION, LOCTEXT("EnginePositionDesc", "The owning component's position in world space."));
-		AttrStrMap.Add(SYS_PARAM_ENGINE_SCALE, LOCTEXT("EnginePositionDesc", "The owning component's scale in world space."));
-		AttrStrMap.Add(SYS_PARAM_ENGINE_VELOCITY, LOCTEXT("EngineVelocityDesc", "The owning component's velocity in world space."));
-		AttrStrMap.Add(SYS_PARAM_ENGINE_X_AXIS, LOCTEXT("XAxisDesc", "The X-axis of the owning component."));
-		AttrStrMap.Add(SYS_PARAM_ENGINE_Y_AXIS, LOCTEXT("YAxisDesc", "The Y-axis of the owning component."));
-		AttrStrMap.Add(SYS_PARAM_ENGINE_Z_AXIS, LOCTEXT("ZAxisDesc", "The Z-axis of the owning component."));
-
-		AttrStrMap.Add(SYS_PARAM_ENGINE_LOCAL_TO_WORLD, LOCTEXT("LocalToWorldDesc", "Owning component's local space to world space transform matrix."));
-		AttrStrMap.Add(SYS_PARAM_ENGINE_WORLD_TO_LOCAL, LOCTEXT("WorldToLocalDesc", "Owning component's world space to local space transform matrix."));
-		AttrStrMap.Add(SYS_PARAM_ENGINE_LOCAL_TO_WORLD_TRANSPOSED, LOCTEXT("LocalToWorldTransposeDesc", "Owning component's local space to world space transform matrix transposed."));
-		AttrStrMap.Add(SYS_PARAM_ENGINE_WORLD_TO_LOCAL_TRANSPOSED, LOCTEXT("WorldToLocalTransposeDesc", "Owning component's world space to local space transform matrix transposed."));
-		AttrStrMap.Add(SYS_PARAM_ENGINE_LOCAL_TO_WORLD_NO_SCALE, LOCTEXT("LocalToWorldNoScaleDesc", "Owning component's local space to world space transform matrix with scaling removed."));
-		AttrStrMap.Add(SYS_PARAM_ENGINE_WORLD_TO_LOCAL_NO_SCALE, LOCTEXT("WorldToLocalNoScaleDesc", "Owning component's world space to local space transform matrix with scaling removed."));
-
-		AttrStrMap.Add(SYS_PARAM_ENGINE_TIME_SINCE_RENDERED, LOCTEXT("TimeSinceRendered", "The time in seconds that have passed since this system was last rendered."));
-		AttrStrMap.Add(SYS_PARAM_ENGINE_MIN_DIST_TO_CAMERA, LOCTEXT("MinDistanceToCamera", "The distance from the owner component to the nearest local player viewpoint."));
-		AttrStrMap.Add(SYS_PARAM_ENGINE_EXEC_COUNT, LOCTEXT("ExecCountDesc", "The index of this particle in the read buffer."));
-		AttrStrMap.Add(SYS_PARAM_ENGINE_EMITTER_NUM_PARTICLES, LOCTEXT("EmitterNumParticles", "The number of particles for this emitter at the beginning of simulation. Should only be used in Emitter scripts."));
-		AttrStrMap.Add(SYS_PARAM_ENGINE_SYSTEM_NUM_EMITTERS_ALIVE, LOCTEXT("SystemNumEmittersAlive", "The number of emitters still alive attached to this system. Should only be used in System scripts."));
-		AttrStrMap.Add(SYS_PARAM_ENGINE_SYSTEM_NUM_EMITTERS, LOCTEXT("SystemNumEmitters", "The number of emitters attached to this system. Should only be used in System scripts."));
-	}
-
-	FText* FoundStr = AttrStrMap.Find(InAttribute);
+	check(SystemStrMap.Num() != 0);
+	FText* FoundStr = SystemStrMap.Find(InAttribute);
 	if (FoundStr != nullptr && !FoundStr->IsEmpty())
 	{
 		return *FoundStr;
@@ -154,64 +456,14 @@ FText FNiagaraConstants::GetEngineConstantDescription(const FNiagaraVariable& In
 
 const TArray<FNiagaraVariable>& FNiagaraConstants::GetCommonParticleAttributes()
 {
-	static TArray<FNiagaraVariable> Attributes;
-	if (Attributes.Num() == 0)
-	{
-		Attributes.Add(SYS_PARAM_PARTICLES_POSITION);
-		Attributes.Add(SYS_PARAM_PARTICLES_VELOCITY);
-		Attributes.Add(SYS_PARAM_PARTICLES_COLOR);
-		Attributes.Add(SYS_PARAM_PARTICLES_SPRITE_ROTATION);
-		Attributes.Add(SYS_PARAM_PARTICLES_NORMALIZED_AGE);
-		Attributes.Add(SYS_PARAM_PARTICLES_SPRITE_SIZE);
-		Attributes.Add(SYS_PARAM_PARTICLES_SPRITE_FACING);
-		Attributes.Add(SYS_PARAM_PARTICLES_SPRITE_ALIGNMENT);
-		Attributes.Add(SYS_PARAM_PARTICLES_SUB_IMAGE_INDEX);
-		Attributes.Add(SYS_PARAM_PARTICLES_DYNAMIC_MATERIAL_PARAM);
-		Attributes.Add(SYS_PARAM_PARTICLES_SCALE);
-		Attributes.Add(SYS_PARAM_PARTICLES_LIFETIME);
-		Attributes.Add(SYS_PARAM_PARTICLES_MESH_ORIENTATION);
-		Attributes.Add(SYS_PARAM_PARTICLES_CAMERA_OFFSET);
-		Attributes.Add(SYS_PARAM_PARTICLES_UV_SCALE);
-		Attributes.Add(SYS_PARAM_PARTICLES_MATERIAL_RANDOM);
-		Attributes.Add(SYS_PARAM_PARTICLES_RIBBONID);
-		Attributes.Add(SYS_PARAM_PARTICLES_RIBBONWIDTH);
-		Attributes.Add(SYS_PARAM_PARTICLES_RIBBONTWIST);
-		Attributes.Add(SYS_PARAM_PARTICLES_RIBBONFACING);
-
-
-		Attributes.Add(SYS_PARAM_INSTANCE_ALIVE);
-	}
+	check(Attributes.Num() != 0);
 	return Attributes;
 }
 
 FString FNiagaraConstants::GetAttributeDefaultValue(const FNiagaraVariable& InAttribute)
 {
-	static TMap<FNiagaraVariable, FString> AttrStrMap;
-	if (AttrStrMap.Num() == 0)
-	{
-		AttrStrMap.Add(SYS_PARAM_PARTICLES_POSITION, SYS_PARAM_ENGINE_POSITION.GetName().ToString());
-		AttrStrMap.Add(SYS_PARAM_PARTICLES_VELOCITY, TEXT("0.0,0.0,100.0"));
-		AttrStrMap.Add(SYS_PARAM_PARTICLES_COLOR, FLinearColor(1.0f, 1.0f, 1.0f, 1.0f).ToString());
-		AttrStrMap.Add(SYS_PARAM_PARTICLES_SPRITE_ROTATION, TEXT("0.0"));
-		AttrStrMap.Add(SYS_PARAM_PARTICLES_NORMALIZED_AGE, TEXT("0.0"));
-		AttrStrMap.Add(SYS_PARAM_PARTICLES_SPRITE_SIZE, TEXT("X=50.0 Y=50.0"));
-		AttrStrMap.Add(SYS_PARAM_PARTICLES_SPRITE_FACING, TEXT("1.0,0.0,0.0"));
-		AttrStrMap.Add(SYS_PARAM_PARTICLES_SPRITE_ALIGNMENT, TEXT("1.0,0.0,0.0"));
-		AttrStrMap.Add(SYS_PARAM_PARTICLES_SUB_IMAGE_INDEX, TEXT("0.0"));
-		AttrStrMap.Add(SYS_PARAM_PARTICLES_DYNAMIC_MATERIAL_PARAM, TEXT("1.0,1.0,1.0,1.0"));
-		AttrStrMap.Add(SYS_PARAM_PARTICLES_SCALE, TEXT("1.0,1.0,1.0"));
-		AttrStrMap.Add(SYS_PARAM_PARTICLES_LIFETIME, TEXT("5.0"));
-		AttrStrMap.Add(SYS_PARAM_PARTICLES_MESH_ORIENTATION, TEXT("1.0,0.0,0.0,0.0"));
-		AttrStrMap.Add(SYS_PARAM_PARTICLES_CAMERA_OFFSET, TEXT("0.0"));
-		AttrStrMap.Add(SYS_PARAM_PARTICLES_UV_SCALE, TEXT("X=1.0 Y=1.0"));
-		AttrStrMap.Add(SYS_PARAM_PARTICLES_MATERIAL_RANDOM, TEXT("0.0"));
-		AttrStrMap.Add(SYS_PARAM_PARTICLES_RIBBONID, TEXT("0"));
-		AttrStrMap.Add(SYS_PARAM_PARTICLES_RIBBONWIDTH, TEXT("1.0"));	
-		AttrStrMap.Add(SYS_PARAM_PARTICLES_RIBBONTWIST, TEXT("0.0"));
-		AttrStrMap.Add(SYS_PARAM_PARTICLES_RIBBONFACING, TEXT("0.0, 0.0, 1.0"));
-	}
-
-	FString* FoundStr = AttrStrMap.Find(InAttribute);
+	check(AttrDefaultsStrMap.Num() != 0);
+	FString* FoundStr = AttrDefaultsStrMap.Find(InAttribute);
 	if (FoundStr != nullptr && !FoundStr->IsEmpty())
 	{
 		return *FoundStr;
@@ -221,32 +473,8 @@ FString FNiagaraConstants::GetAttributeDefaultValue(const FNiagaraVariable& InAt
 
 FText FNiagaraConstants::GetAttributeDescription(const FNiagaraVariable& InAttribute)
 {
-	static TMap<FNiagaraVariable, FText> AttrStrMap;
-	if (AttrStrMap.Num() == 0)
-	{
-		AttrStrMap.Add(SYS_PARAM_PARTICLES_POSITION, LOCTEXT("PositionDesc", "The position of the particle."));
-		AttrStrMap.Add(SYS_PARAM_PARTICLES_VELOCITY, LOCTEXT("VelocityDesc", "The velocity in cm/s of the particle."));
-		AttrStrMap.Add(SYS_PARAM_PARTICLES_COLOR, LOCTEXT("ColorDesc", "The color of the particle."));
-		AttrStrMap.Add(SYS_PARAM_PARTICLES_SPRITE_ROTATION, LOCTEXT("SpriteRotDesc", "The screen aligned roll of the particle in degrees."));
-		AttrStrMap.Add(SYS_PARAM_PARTICLES_NORMALIZED_AGE, LOCTEXT("NormalizedAgeDesc", "The age in seconds divided by lifetime in seconds. Useful for animation as the value is between 0 and 1."));
-		AttrStrMap.Add(SYS_PARAM_PARTICLES_SPRITE_SIZE, LOCTEXT("SpriteSizeDesc", "The size of the sprite quad."));
-		AttrStrMap.Add(SYS_PARAM_PARTICLES_SPRITE_FACING, LOCTEXT("FacingDesc", "Makes the surface of the sprite face towards a custom vector. Must be used with the SpriteRenderer's CustomFacingVector FacingMode and CustomFacingVectorMask options."));
-		AttrStrMap.Add(SYS_PARAM_PARTICLES_SPRITE_ALIGNMENT, LOCTEXT("AlignmentDesc", "Imagine the texture having an arrow pointing up, this attribute makes the arrow point towards the alignment axis. Must be used with the SpriteRenderer's CustomAlignment Alignment option."));
-		AttrStrMap.Add(SYS_PARAM_PARTICLES_SUB_IMAGE_INDEX, LOCTEXT("SubImageIndexDesc", "A value from 0 to the number of entries in the table of SubUV images."));
-		AttrStrMap.Add(SYS_PARAM_PARTICLES_DYNAMIC_MATERIAL_PARAM, LOCTEXT("DynamicMaterialParameterDesc", "The 4-float vector used to send custom data to renderer."));
-		AttrStrMap.Add(SYS_PARAM_PARTICLES_SCALE, LOCTEXT("ScaleParamDesc", "The XYZ scale of the non-sprite based particle."));
-		AttrStrMap.Add(SYS_PARAM_PARTICLES_LIFETIME, LOCTEXT("LifetimeParamDesc", "The lifetime of a particle in seconds."));
-		AttrStrMap.Add(SYS_PARAM_PARTICLES_MESH_ORIENTATION, LOCTEXT("MeshOrientParamDesc", "The axis-angle rotation to be applied to the mesh particle."));
-		AttrStrMap.Add(SYS_PARAM_PARTICLES_CAMERA_OFFSET, LOCTEXT("CamOffsetParamDesc", "Used to offset position in the direction of the camera. The value is multiplied by the direction vector from the camera to the particle."));
-		AttrStrMap.Add(SYS_PARAM_PARTICLES_UV_SCALE, LOCTEXT("UVScalerParamDesc", "Used to multiply the generated UVs for Sprite renderers."));
-		AttrStrMap.Add(SYS_PARAM_PARTICLES_MATERIAL_RANDOM, LOCTEXT("MaterialRandomParamDesc", "Used to drive the Particle Random node in the Material Editor. Without this set, any Particle Randoms will get 0.0."));
-		AttrStrMap.Add(SYS_PARAM_INSTANCE_ALIVE, LOCTEXT("AliveParamDesc", "Used to determine whether or not this particle instance is still valid or if it can be deleted."));
-		AttrStrMap.Add(SYS_PARAM_PARTICLES_RIBBONID, LOCTEXT("RibbonIDDesc", "Sets the ribbon id for a particle. Particles with the same ribbon id will be connected into a ribbon."));
-		AttrStrMap.Add(SYS_PARAM_PARTICLES_RIBBONWIDTH, LOCTEXT("RibbonWidthDesc", "Sets the ribbon width for a particle, in UE4 units."));
-		AttrStrMap.Add(SYS_PARAM_PARTICLES_RIBBONTWIST, LOCTEXT("RibbonTwistDesc", "Sets the ribbon twist for a particle, in degrees."));
-		AttrStrMap.Add(SYS_PARAM_PARTICLES_RIBBONFACING, LOCTEXT("RibbonFacingDesc", "Sets the facing vector of the ribbon at the particle position."));
-	}
-	FText* FoundStr = AttrStrMap.Find(InAttribute);
+	check(AttrDescStrMap.Num() != 0);
+	FText* FoundStr = AttrDescStrMap.Find(InAttribute);
 	if (FoundStr != nullptr && !FoundStr->IsEmpty())
 	{
 		return *FoundStr;
@@ -265,34 +493,97 @@ bool FNiagaraConstants::IsNiagaraConstant(const FNiagaraVariable& InVar)
 
 const FNiagaraVariableMetaData* FNiagaraConstants::GetConstantMetaData(const FNiagaraVariable& InVar)
 {
-	static TMap<FNiagaraVariable, FNiagaraVariableMetaData> AttrMetaData;
-	if (AttrMetaData.Num() == 0)
+	check(AttrMetaData.Num() != 0);
+	return AttrMetaData.Find(InVar);
+}
+
+FNiagaraVariable FNiagaraConstants::GetAttributeWithDefaultValue(const FNiagaraVariable& InAttribute)
+{
+	check(AttrDefaultsValueMap.Num() != 0);
+	FNiagaraVariable* FoundValue = AttrDefaultsValueMap.Find(InAttribute);
+	if (FoundValue != nullptr)
 	{
-		// Add the engine default attributes..
+		return *FoundValue;
+	}
+	return FNiagaraVariable();
+}
+
+FNiagaraVariable FNiagaraConstants::GetAttributeAsDataSetKey(const FNiagaraVariable& InVar)
+{
+	FNiagaraVariable OutVar = InVar;
+	FString DataSetName = InVar.GetName().ToString();
+	DataSetName.RemoveFromStart(TEXT("Particles."));
+	OutVar.SetName(*DataSetName);
+	return OutVar;
+}
+
+FNiagaraVariableAttributeBinding FNiagaraConstants::GetAttributeDefaultBinding(const FNiagaraVariable& InVar)
+{
+	if (AttrDefaultsValueMap.Num() == 0)
+	{
+		Init();
+	}
+
+	FNiagaraVariableAttributeBinding Binding;
+	Binding.BoundVariable = InVar;
+	Binding.DataSetVariable = InVar;
+	const FNiagaraVariable* FoundVar = AttrDataSetKeyMap.Find(InVar);
+	if (FoundVar)
+	{
+		Binding.DataSetVariable = *FoundVar;
+	}
+
+	Binding.DefaultValueIfNonExistent = GetAttributeWithDefaultValue(InVar);
+	return Binding;
+}
+
+const FNiagaraVariable* FNiagaraConstants::GetKnownConstant(const FName& InName, bool bAllowPartialNameMatch)
+{
+	const TArray<FNiagaraVariable>& EngineConstants = GetEngineConstants();
+
+	if (!bAllowPartialNameMatch)
+	{
+		const FNiagaraVariable* FoundSystemVar = EngineConstants.FindByPredicate([&](const FNiagaraVariable& Var)
 		{
-			TArray<FNiagaraVariable> Attributes = GetCommonParticleAttributes();
-			for (FNiagaraVariable Var : Attributes)
-			{
-				FNiagaraVariableMetaData MetaData;
-				MetaData.Description = GetAttributeDescription(Var);
-				AttrMetaData.Add(Var, MetaData);
-			}
-		}
+			return Var.GetName() == InName;
+		});
 
-		// Add the engine constants..
+		if (FoundSystemVar)
 		{
-			TArray<FNiagaraVariable> EngineConstants = GetEngineConstants();
-			for (FNiagaraVariable Var : EngineConstants)
-			{
-				FNiagaraVariableMetaData MetaData;
-				MetaData.Description = GetEngineConstantDescription(Var);
-				AttrMetaData.Add(Var, MetaData);
-			}
+			return FoundSystemVar;
 		}
+	}
+	else
+	{
+		int32 FoundIdx = FNiagaraVariable::SearchArrayForPartialNameMatch(EngineConstants, InName);
+		if (FoundIdx != INDEX_NONE)
+		{
+			return &EngineConstants[FoundIdx];
+		}
+	}
+
+	const TArray<FNiagaraVariable>& LocalAttributes = GetCommonParticleAttributes();
+	if (!bAllowPartialNameMatch)
+	{
+		const FNiagaraVariable* FoundAttribVar = LocalAttributes.FindByPredicate([&](const FNiagaraVariable& Var)
+		{
+			return Var.GetName() == InName;
+		});
 
+		if (FoundAttribVar)
+		{
+			return FoundAttribVar;
+		}
 	}
-	return AttrMetaData.Find(InVar);
+	else
+	{
+		int32 FoundIdx = FNiagaraVariable::SearchArrayForPartialNameMatch(Attributes, InName);
+		if (FoundIdx != INDEX_NONE)
+		{
+			return &Attributes[FoundIdx];
+		}
+	}
+	return nullptr;
 }
 
-
 #undef LOCTEXT_NAMESPACE
diff --git a/Engine/Plugins/FX/Niagara/Source/Niagara/Private/NiagaraDataInterface.cpp b/Engine/Plugins/FX/Niagara/Source/Niagara/Private/NiagaraDataInterface.cpp
--- a/Engine/Plugins/FX/Niagara/Source/Niagara/Private/NiagaraDataInterface.cpp
+++ b/Engine/Plugins/FX/Niagara/Source/Niagara/Private/NiagaraDataInterface.cpp
@@ -5,11 +5,19 @@
 #include "Curves/CurveLinearColor.h"
 #include "Curves/CurveFloat.h"
 #include "NiagaraTypes.h"
+#include "ShaderParameterUtils.h"
+#include "NiagaraShader.h"
 
 UNiagaraDataInterface::UNiagaraDataInterface(FObjectInitializer const& ObjectInitializer)
 {
 }
 
+void UNiagaraDataInterface::PostLoad()
+{
+	Super::PostLoad();
+	SetFlags(RF_Public);
+}
+
 bool UNiagaraDataInterface::CopyTo(UNiagaraDataInterface* Destination) const 
 {
 	bool result = CopyToInternal(Destination);
@@ -91,4 +99,88 @@ bool UNiagaraDataInterfaceCurveBase::Equals(const UNiagaraDataInterface* Other)
 	const UNiagaraDataInterfaceCurveBase* OtherTyped = CastChecked<UNiagaraDataInterfaceCurveBase>(Other);
 	bool bEqual = OtherTyped->bUseLUT == bUseLUT;
 	return bEqual;
-}
\ No newline at end of file
+}
+
+
+void UNiagaraDataInterfaceCurveBase::GetParameterDefinitionHLSL(FNiagaraDataInterfaceGPUParamInfo& ParamInfo, FString& OutHLSL)
+{
+	OutHLSL += TEXT("\n");
+	FString MinTimeStr = TEXT("MinTime_") + ParamInfo.DataInterfaceHLSLSymbol;
+	FString MaxTimeStr = TEXT("MaxTime_") + ParamInfo.DataInterfaceHLSLSymbol;
+	FString InvTimeRangeStr = TEXT("InvTimeRange_") + ParamInfo.DataInterfaceHLSLSymbol;
+
+	FString BufferName = "CurveLUT_" + ParamInfo.DataInterfaceHLSLSymbol;
+	OutHLSL += TEXT("Buffer<float> ") + BufferName + TEXT(";\n");
+
+	OutHLSL += TEXT("float ") + MinTimeStr + TEXT(";\n");
+	OutHLSL += TEXT("float ") + MaxTimeStr + TEXT(";\n");
+	OutHLSL += TEXT("float ") + InvTimeRangeStr + TEXT(";\n");
+	OutHLSL += TEXT("\n");
+
+	//TODO: Create a Unitiliy/Common funcitons hlsl def shared between all instances of the same data interface class for these.
+	OutHLSL += FString::Printf(TEXT("float TimeToLUTFraction_%s(float T)\n{\n\treturn saturate((T - %s) * %s);\n}\n"), *ParamInfo.DataInterfaceHLSLSymbol, *MinTimeStr, *InvTimeRangeStr);
+	OutHLSL += FString::Printf(TEXT("float SampleCurve_%s(float T)\n{\n\treturn %s[(uint)T];\n}\n"), *ParamInfo.DataInterfaceHLSLSymbol, *BufferName);
+	OutHLSL += TEXT("\n");
+}
+
+FRWBuffer& UNiagaraDataInterfaceCurveBase::GetCurveLUTGPUBuffer()
+{
+	//TODO: This isn't really very thread safe. Need to move to a proxy like system where DIs can push data to the RT safely.
+	if (GPUBufferDirty)
+	{
+		int32 ElemSize = GetCurveNumElems();
+		CurveLUT.Release();
+		CurveLUT.Initialize(sizeof(float), CurveLUTWidth * ElemSize, EPixelFormat::PF_R32_FLOAT, BUF_Static);
+		uint32 BufferSize = ShaderLUT.Num() * sizeof(float);
+		int32 *BufferData = static_cast<int32*>(RHILockVertexBuffer(CurveLUT.Buffer, 0, BufferSize, EResourceLockMode::RLM_WriteOnly));
+		FPlatformMemory::Memcpy(BufferData, ShaderLUT.GetData(), BufferSize);
+		RHIUnlockVertexBuffer(CurveLUT.Buffer);
+		GPUBufferDirty = false;
+	}
+
+	return CurveLUT;
+}
+
+struct FNiagaraDataInterfaceParametersCS_Curve : public FNiagaraDataInterfaceParametersCS
+{
+	virtual ~FNiagaraDataInterfaceParametersCS_Curve() {}
+	virtual void Bind(const FNiagaraDataInterfaceParamRef& ParamRef, const class FShaderParameterMap& ParameterMap) override
+	{
+		MinTime.Bind(ParameterMap, *(TEXT("MinTime_") + ParamRef.ParameterInfo.DataInterfaceHLSLSymbol));
+		MaxTime.Bind(ParameterMap, *(TEXT("MaxTime_") + ParamRef.ParameterInfo.DataInterfaceHLSLSymbol));
+		InvTimeRange.Bind(ParameterMap, *(TEXT("InvTimeRange_") + ParamRef.ParameterInfo.DataInterfaceHLSLSymbol));
+		CurveLUT.Bind(ParameterMap, *(TEXT("CurveLUT_") + ParamRef.ParameterInfo.DataInterfaceHLSLSymbol));
+	}
+
+	virtual void Serialize(FArchive& Ar) override
+	{
+		Ar << MinTime;
+		Ar << MaxTime;
+		Ar << InvTimeRange;
+		Ar << CurveLUT;
+	}
+
+	virtual void Set(FRHICommandList& RHICmdList, FNiagaraShader* Shader, class UNiagaraDataInterface* DataInterface) const override
+	{
+		check(IsInRenderingThread());
+
+		const FComputeShaderRHIParamRef ComputeShaderRHI = Shader->GetComputeShader();
+		UNiagaraDataInterfaceCurveBase* CurveDI = CastChecked<UNiagaraDataInterfaceCurveBase>(DataInterface);
+		FRWBuffer& CurveLUTBuffer = CurveDI->GetCurveLUTGPUBuffer();
+
+		SetShaderValue(RHICmdList, ComputeShaderRHI, MinTime, CurveDI->GetMinTime());
+		SetShaderValue(RHICmdList, ComputeShaderRHI, MaxTime, CurveDI->GetMaxTime());
+		SetShaderValue(RHICmdList, ComputeShaderRHI, InvTimeRange, CurveDI->GetInvTimeRange());
+		RHICmdList.SetShaderResourceViewParameter(ComputeShaderRHI, CurveLUT.GetBaseIndex(), CurveLUTBuffer.SRV);
+	}
+
+	FShaderParameter MinTime;
+	FShaderParameter MaxTime;
+	FShaderParameter InvTimeRange;
+	FShaderResourceParameter CurveLUT;
+};
+
+FNiagaraDataInterfaceParametersCS* UNiagaraDataInterfaceCurveBase::ConstructComputeParameters()const
+{
+	return new FNiagaraDataInterfaceParametersCS_Curve();
+}
