Copying //UE4/Dev-Niagara to //UE4/Dev-Main (Source: //UE4/Dev-Niagara @ 4074996)

diff --git a/Engine/Plugins/FX/Niagara/Source/Niagara/Private/NiagaraDataInterfaceCollisionQuery.cpp b/Engine/Plugins/FX/Niagara/Source/Niagara/Private/NiagaraDataInterfaceCollisionQuery.cpp
--- a/Engine/Plugins/FX/Niagara/Source/Niagara/Private/NiagaraDataInterfaceCollisionQuery.cpp
+++ b/Engine/Plugins/FX/Niagara/Source/Niagara/Private/NiagaraDataInterfaceCollisionQuery.cpp
@@ -3,10 +3,14 @@
 #include "NiagaraDataInterfaceCollisionQuery.h"
 #include "NiagaraTypes.h"
 #include "NiagaraCustomVersion.h"
+#include "NiagaraWorldManager.h"
+#include "ShaderParameterUtils.h"
+#include "Shader.h"
 
 //////////////////////////////////////////////////////////////////////////
 //Color Curve
 
+FCriticalSection UNiagaraDataInterfaceCollisionQuery::CriticalSection;
 
 UNiagaraDataInterfaceCollisionQuery::UNiagaraDataInterfaceCollisionQuery(FObjectInitializer const& ObjectInitializer)
 	: Super(ObjectInitializer)
@@ -82,13 +86,13 @@ void UNiagaraDataInterfaceCollisionQuery::GetFunctions(TArray<FNiagaraFunctionSi
 	//Sig3.Inputs.Add(FNiagaraVariable(FNiagaraTypeDefinition::GetBoolDef(), TEXT("PerformQuery")));
 	Sig3.Inputs.Add(FNiagaraVariable(FNiagaraTypeDefinition::GetIntDef(), TEXT("ReturnQueryID")));
 	Sig3.Inputs.Add(FNiagaraVariable(FNiagaraTypeDefinition::GetVec3Def(), TEXT("ParticlePosition")));
-	Sig3.Inputs.Add(FNiagaraVariable(FNiagaraTypeDefinition::GetVec3Def(), TEXT("ParticleVelocity")));
+	Sig3.Inputs.Add(FNiagaraVariable(FNiagaraTypeDefinition::GetVec3Def(), TEXT("Direction")));
 	Sig3.Inputs.Add(FNiagaraVariable(FNiagaraTypeDefinition::GetFloatDef(), TEXT("DeltaTime")));
 	Sig3.Inputs.Add(FNiagaraVariable(FNiagaraTypeDefinition::GetFloatDef(), TEXT("CollisionSize")));
+	Sig3.Inputs.Add(FNiagaraVariable(FNiagaraTypeDefinition::GetFloatDef(), TEXT("DepthBounds")));
 	Sig3.Outputs.Add(FNiagaraVariable(FNiagaraTypeDefinition::GetIntDef(), TEXT("QueryID")));
 	Sig3.Outputs.Add(FNiagaraVariable(FNiagaraTypeDefinition::GetBoolDef(), TEXT("CollisionValid")));
 	Sig3.Outputs.Add(FNiagaraVariable(FNiagaraTypeDefinition::GetVec3Def(), TEXT("CollisionPos")));
-	Sig3.Outputs.Add(FNiagaraVariable(FNiagaraTypeDefinition::GetVec3Def(), TEXT("CollisionVelocity")));
 	Sig3.Outputs.Add(FNiagaraVariable(FNiagaraTypeDefinition::GetVec3Def(), TEXT("CollisionNormal")));
 	Sig3.Outputs.Add(FNiagaraVariable(FNiagaraTypeDefinition::GetFloatDef(), TEXT("Friction")));
 	Sig3.Outputs.Add(FNiagaraVariable(FNiagaraTypeDefinition::GetFloatDef(), TEXT("Restitution")));
@@ -123,109 +127,139 @@ void UNiagaraDataInterfaceCollisionQuery::GetFunctions(TArray<FNiagaraFunctionSi
 // the HLSL in the spirit of a static switch
 // TODO: need a way to identify each specific function here
 // 
-bool UNiagaraDataInterfaceCollisionQuery::GetFunctionHLSL(const FName& DefinitionFunctionName, FString InstanceFunctionName, TArray<FDIGPUBufferParamDescriptor> &Descriptors, FString &HLSLInterfaceID, FString &OutHLSL)
+bool UNiagaraDataInterfaceCollisionQuery::GetFunctionHLSL(const FName& DefinitionFunctionName, FString InstanceFunctionName, FNiagaraDataInterfaceGPUParamInfo& ParamInfo, FString& OutHLSL)
 {
 	// a little tricky, since we've got two functions for submitting and retrieving a query; we store submitted queries per thread group,
 	// assuming it'll usually be the same thread trying to call ReadQuery for a particular QueryID, that submitted it in the first place.
-	if (InstanceFunctionName == TEXT("PerformQuery"))
+	if (DefinitionFunctionName == TEXT("PerformCollisionQuery"))
 	{
-		FString TextureName = Descriptors[0].BufferParamName;
-		OutHLSL += TEXT("groupshared int CollisionQueryID = 0;\n");
-		OutHLSL += TEXT("\t\n");
-		OutHLSL += TEXT("void ") + InstanceFunctionName + TEXT("(in int InQueryID, in float3 In_ParticlePos, in float3 In_ParticleVel, in float In_Deltatime, out int OutQueryID, out bool OutCollisionValid, out float3 Out_CollisionPos, out float3 Out_CollisionVel, out float3 Out_CollisionNormal) \n{\n");
-		OutHLSL += TEXT("\tInterlockedAdd(CollisionQueryID, 1);\n");
+		OutHLSL += TEXT("float3 WorldPositionFromSceneDepth(float2 ScreenPosition, float SceneDepth)\n{\n\tfloat4 HomogeneousWorldPosition = mul(float4(ScreenPosition * SceneDepth, SceneDepth, 1), View.ScreenToWorld);\n\treturn HomogeneousWorldPosition.xyz / HomogeneousWorldPosition.w;\n}\n");
+		OutHLSL += TEXT("\n");
+		OutHLSL += TEXT("// For retrieving the world-space normal.\n\
+							//#if FORWARD_SHADING\n\
+							Texture2D GBufferATexture;\n\
+							SamplerState GBufferATextureSampler;\n\
+							//#endif\n\n\
+							");
+		OutHLSL += TEXT("void ") + InstanceFunctionName + TEXT("(in int InQueryID, in float3 In_ParticlePos, in float3 In_ParticleVel, in float In_DeltaSeconds, float CollisionRadius, in float CollisionDepthBounds, int In_InstanceData, \
+			out int Out_QueryID, out bool OutCollisionValid, out float3 Out_CollisionPos, out float3 Out_CollisionNormal, out float Out_Friction, out float Out_Restitution) \n{\n");
 		// get the screen position
-		OutHLSL += TEXT("float3 CollisionPosition = InPosition + CollisionOffset;\
-							float4 SamplePosition = float4(CollisionPosition + View.PreViewTranslation, 1);\
-							float4 ClipPosition = mul(SamplePosition, View.TranslatedWorldToClip);\
-							float2 ScreenPosition = ClipPosition.xy / ClipPosition.w;\
-							// Don't try to collide if the particle falls outside the view.\
-							if (all(abs(ScreenPosition.xy) <= float2(1, 1)))\
-							{\
-								// Sample the depth buffer to get a world position near the particle.\
-								float2 ScreenUV = ScreenPosition * View.ScreenPositionScaleBias.xy + View.ScreenPositionScaleBias.wz;\
-								float SceneDepth = CalcSceneDepth(ScreenUV);\
-								if (abs(ClipPosition.w - SceneDepth) < CollisionDepthBounds)");
-		OutHLSL += TEXT("\n}\n");
+		OutHLSL += TEXT("\
+		OutCollisionValid = false;\n\
+		Out_QueryID = InQueryID;\n\
+		Out_CollisionPos = In_ParticlePos;\n\
+		Out_CollisionNormal = float3(0.0, 0.0, 1.0);\n\
+		Out_Friction = 0.0;\n\
+		Out_Restitution = 1.0;\n\
+		float3 DeltaPosition = In_DeltaSeconds * In_ParticleVel; \
+		float3 CollisionOffset = normalize(DeltaPosition) * CollisionRadius;\
+		float3 CollisionPosition = In_ParticlePos + CollisionOffset; \n\
+		float3 NewPosition = In_ParticlePos.xyz + DeltaPosition; \
+		float4 SamplePosition = float4(CollisionPosition + View.PreViewTranslation, 1); \n\
+		float4 ClipPosition = mul(SamplePosition, View.TranslatedWorldToClip);\n\
+		float2 ScreenPosition = ClipPosition.xy / ClipPosition.w;\n\
+		// Don't try to collide if the particle falls outside the view.\n\
+		if (all(abs(ScreenPosition.xy) <= float2(1, 1)))\n\
+		{\n\
+			// Sample the depth buffer to get a world position near the particle.\n\
+			float2 ScreenUV = ScreenPosition * View.ScreenPositionScaleBias.xy + View.ScreenPositionScaleBias.wz;\n\
+			float SceneDepth = CalcSceneDepth(ScreenUV);\n\
+			if (abs(ClipPosition.w - SceneDepth) < CollisionDepthBounds)\n\
+			{\n\
+			/*#if FORWARD_SHADING_ACCURATE_NORMAL && FORWARD_SHADING\n\
+				float SceneDepth0 = CalcSceneDepth(ScreenUV + float2(View.BufferSizeAndInvSize.z, 0.0));\n\
+				float SceneDepth1 = CalcSceneDepth(ScreenUV + float2(0.0, View.BufferSizeAndInvSize.w));\n\
+				// When using the forward shading, the normal of the pixel is approximated by the derivative of the world position\n\
+				// of the pixel. But in on the visible edge this derivative can become very high, making CollisionPlane almost\n\
+				// perpendicular to the view plane. In these case the particle may collide the visible edges of the diferent meshes\n\
+				// in the view frustum. To avoid this, we disable the collision test if one of the derivate is above a threshold.\n\
+				if (max(abs(SceneDepth - SceneDepth0), abs(SceneDepth - SceneDepth1)) > CollisionDepthBounds)\n\
+				{\n\
+					OutCollisionValid = false;\n\
+					return;\n\
+				}\n\
+			#endif*/\n\
+				// Reconstruct world position.\n\
+				float3 WorldPosition = WorldPositionFromSceneDepth(ScreenPosition.xy, SceneDepth);\n\
+			/*#if FORWARD_SHADING && FORWARD_SHADING_ACCURATE_NORMAL\n\
+				// Compute normal using the two additional neighbooring pixel depth buffer fetches.\n\
+				float3 WorldPosition0 = WorldPositionFromSceneDepth(ScreenPosition.xy + float2(2 * View.ViewSizeAndInvSize.z, 0.0), SceneDepth0);\n\
+				float3 WorldPosition1 = WorldPositionFromSceneDepth(ScreenPosition.xy - float2(0.0, 2 * View.ViewSizeAndInvSize.w), SceneDepth1);\n\
+				float3 WorldNormal = normalize(cross(WorldPosition0 - WorldPosition, WorldPosition1 - WorldPosition));\n\
+			#elif FORWARD_SHADING // && !FORWARD_SHADING_ACCURATE_NORMAL\n\
+				// Compute normal using a ddx/ddy hack, hopefully the neighbooring particules will be close enough...\n\
+				float3 WorldNormal = normalize(cross(ddx(WorldPosition), ddy(WorldPosition)));\n\
+				WorldNormal *= sign(dot(View.WorldCameraOrigin - NewPosition, WorldNormal));\n\
+			#else*/ //!FORWARD_SHADING\n\
+				// Sample the normal buffer to create a plane to collide against.\n\
+				float3 WorldNormal = Texture2DSampleLevel(GBufferATexture, GBufferATextureSampler, ScreenUV, 0).xyz * 2.0 - 1.0;\n\
+			//#endif\n\
+				float4 CollisionPlane = float4(WorldNormal, dot(WorldPosition.xyz, WorldNormal));\n\
+				// Compute the portion of velocity normal to the collision plane.\n\
+				float VelocityDot = dot(CollisionPlane.xyz, DeltaPosition.xyz);\n\
+				float d_back = (dot(CollisionPlane.xyz, In_ParticlePos.xyz) + CollisionRadius - CollisionPlane.w);\n\
+				float d_front = (dot(CollisionPlane.xyz, NewPosition.xyz) - CollisionRadius - CollisionPlane.w);\n\
+				// distance to the plane from current and predicted position\n\
+				if (d_back >= 0.0f && d_front <= 0.0f && VelocityDot < 0.0f)\n\
+				{\n\
+					OutCollisionValid = true;\n\
+					Out_CollisionPos = In_ParticlePos + (WorldNormal*d_back);\n\
+					Out_CollisionNormal = WorldNormal;\n\
+					Out_Friction = 0.0f;\n\
+					Out_Restitution = 1.0f;\n\
+					Out_QueryID = 0.0f;\
+				}\n\
+				else\n\
+				{\n\
+					OutCollisionValid = false; \n\
+				}\n\
+			}\n\
+		}\n\
+		\n\
+		\n");
+
+		OutHLSL += TEXT("}\n");
 	}
 
 
 	return true;
 }
 
-// build buffer definition hlsl
-// 1. Choose a buffer name, add the data interface ID (important!)
-// 2. add a DIGPUBufferParamDescriptor to the array argument; that'll be passed on to the FNiagaraShader for binding to a shader param, that can
-// then later be found by name via FindDIBufferParam for setting; 
-// 3. store buffer declaration hlsl in OutHLSL
-// multiple buffers can be defined at once here
-//
-void UNiagaraDataInterfaceCollisionQuery::GetBufferDefinitionHLSL(FString DataInterfaceID, TArray<FDIGPUBufferParamDescriptor> &BufferDescriptors, FString &OutHLSL)
+void UNiagaraDataInterfaceCollisionQuery::GetParameterDefinitionHLSL(FNiagaraDataInterfaceGPUParamInfo& ParamInfo, FString& OutHLSL)
 {
-	FString TextureName = "DepthBuffer" + DataInterfaceID;
-	OutHLSL += TEXT("Texture2D ") + TextureName + TEXT(";\n");
-	BufferDescriptors.Add(FDIGPUBufferParamDescriptor(TextureName, 0));		// add a descriptor for shader parameter binding
-
-	FString CounterName = "DepthBuffer" + DataInterfaceID;
-	OutHLSL += TEXT("Buffer<uint> ") + TextureName + TEXT(";\n");
-	BufferDescriptors.Add(FDIGPUBufferParamDescriptor(TextureName, 0));		// add a descriptor for shader parameter binding
-}
-
-// called after translate, to setup buffers matching the buffer descriptors generated during hlsl translation
-// need to do this because the script used during translate is a clone, including its DIs
-//
-void UNiagaraDataInterfaceCollisionQuery::SetupBuffers(FDIBufferDescriptorStore &BufferDescriptors)
-{
-	GPUBuffers.Empty();
-	/*
-	for (FDIGPUBufferParamDescriptor &Desc : BufferDescriptors.Descriptors)
-	{
-		FNiagaraDataInterfaceBufferData BufferData(*Desc.BufferParamName);	// store off the data for later use
-		GPUBuffers.Add(BufferData);
-	}
-	GPUBufferDirty = true;
-	*/
-}
-
-// return the GPU buffer array (called from NiagaraInstanceBatcher to get the buffers for setting to the shader)
-// we lazily update the buffer with a new LUT here if necessary
-//
-TArray<FNiagaraDataInterfaceBufferData> &UNiagaraDataInterfaceCollisionQuery::GetBufferDataArray()
-{
-	check(IsInRenderingThread());
-	return GPUBuffers;
+	// we don't need to add these to hlsl, as they're already in common.ush
+	//OutHLSL = "Texture2D GBufferATexture;\nSamplerState GBufferATextureSampler;";	
 }
 
+DEFINE_NDI_RAW_FUNC_BINDER(UNiagaraDataInterfaceCollisionQuery, SubmitQuery);
+DEFINE_NDI_RAW_FUNC_BINDER(UNiagaraDataInterfaceCollisionQuery, ReadQuery);
+DEFINE_NDI_RAW_FUNC_BINDER(UNiagaraDataInterfaceCollisionQuery, PerformQuery);
 
-DEFINE_NDI_FUNC_BINDER(UNiagaraDataInterfaceCollisionQuery, SubmitQuery);
-DEFINE_NDI_FUNC_BINDER(UNiagaraDataInterfaceCollisionQuery, ReadQuery);
-DEFINE_NDI_FUNC_BINDER(UNiagaraDataInterfaceCollisionQuery, PerformQuery);
-
-FVMExternalFunction UNiagaraDataInterfaceCollisionQuery::GetVMExternalFunction(const FVMExternalFunctionBindingInfo& BindingInfo, void* InstanceData)
+void UNiagaraDataInterfaceCollisionQuery::GetVMExternalFunction(const FVMExternalFunctionBindingInfo& BindingInfo, void* InstanceData, FVMExternalFunction &OutFunc)
 {
 	CQDIPerInstanceData *InstData = (CQDIPerInstanceData *)InstanceData;
 	if (BindingInfo.Name == TEXT("SubmitQuery") /*&& BindingInfo.GetNumInputs() == 3 && BindingInfo.GetNumOutputs() == 1*/)
 	{
-		return TNDIParamBinder<0, float, TNDIParamBinder<1, float, TNDIParamBinder<2, float, TNDIParamBinder<3, float, TNDIParamBinder<4, float, TNDIParamBinder<5, float, TNDIParamBinder<6, float, NDI_FUNC_BINDER(UNiagaraDataInterfaceCollisionQuery, SubmitQuery)>>>>>>>::Bind(this, BindingInfo, InstData);
+		TNDIParamBinder<0, float, TNDIParamBinder<1, float, TNDIParamBinder<2, float, TNDIParamBinder<3, float, TNDIParamBinder<4, float, TNDIParamBinder<5, float, TNDIParamBinder<6, float, NDI_RAW_FUNC_BINDER(UNiagaraDataInterfaceCollisionQuery, SubmitQuery)>>>>>>>::Bind(this, BindingInfo, InstData, OutFunc);
 	}
 	else if (BindingInfo.Name == TEXT("ReadQuery") /*&& BindingInfo.GetNumInputs() == 1 && BindingInfo.GetNumOutputs() == 4*/)
 	{
-		return TNDIParamBinder<0, int32, NDI_FUNC_BINDER(UNiagaraDataInterfaceCollisionQuery, ReadQuery)>::Bind(this, BindingInfo, InstData);
+		TNDIParamBinder<0, int32, NDI_RAW_FUNC_BINDER(UNiagaraDataInterfaceCollisionQuery, ReadQuery)>::Bind(this, BindingInfo, InstData, OutFunc);
 	}
 	else if (BindingInfo.Name == TEXT("PerformCollisionQuery") /*&& BindingInfo.GetNumInputs() == 1 && BindingInfo.GetNumOutputs() == 4*/)
 	{
-		return TNDIParamBinder<0, int, TNDIParamBinder<1, float, TNDIParamBinder<2, float, TNDIParamBinder<3, float, TNDIParamBinder<4, float, TNDIParamBinder<5, float, TNDIParamBinder<6, float, TNDIParamBinder<7, float, TNDIParamBinder<8, float, NDI_FUNC_BINDER(UNiagaraDataInterfaceCollisionQuery, PerformQuery)>>>>>>>>>::Bind(this, BindingInfo, InstData);
+		//				  id     |---------------------------------- position ---------------------------|    |------------------------------- direction -------------------------------|     |-------- dT --------|   |------- size -------|    |---- depthBounds ----|
+		TNDIParamBinder<0, int, TNDIParamBinder<1, float, TNDIParamBinder<2, float, TNDIParamBinder<3, float, TNDIParamBinder<4, float, TNDIParamBinder<5, float, TNDIParamBinder<6, float, TNDIParamBinder<7, float, TNDIParamBinder<8, float, TNDIParamBinder<9, float, NDI_RAW_FUNC_BINDER(UNiagaraDataInterfaceCollisionQuery, PerformQuery)>>>>>>>>>>::Bind(this, BindingInfo, InstData, OutFunc);
 	}
 	else
 	{
 		UE_LOG(LogNiagara, Error, TEXT("Could not find data interface external function. %s\n"),
 			*BindingInfo.Name.ToString());
-		return FVMExternalFunction();
 	}
 }
 
 
-template<typename InQueryIDType, typename PosTypeX, typename PosTypeY, typename PosTypeZ, typename VelTypeX, typename VelTypeY, typename VelTypeZ, typename DTType, typename SizeType>
+template<typename InQueryIDType, typename PosTypeX, typename PosTypeY, typename PosTypeZ, typename DirTypeX, typename DirTypeY, typename DirTypeZ, typename DTType, typename SizeType, typename DepthBoundsType>
 void UNiagaraDataInterfaceCollisionQuery::PerformQuery(FVectorVMContext& Context)
 {
 	//PerformType PerformParam(Context);
@@ -233,11 +267,14 @@ void UNiagaraDataInterfaceCollisionQuery::PerformQuery(FVectorVMContext& Context
 	PosTypeX PosParamX(Context);
 	PosTypeY PosParamY(Context);
 	PosTypeZ PosParamZ(Context);
-	VelTypeX VelParamX(Context);
-	VelTypeY VelParamY(Context);
-	VelTypeZ VelParamZ(Context);
+
+	DirTypeX DirParamX(Context);
+	DirTypeY DirParamY(Context);
+	DirTypeZ DirParamZ(Context);
+
 	DTType DTParam(Context);
 	SizeType SizeParam(Context);
+	DepthBoundsType DepthBoundsParam(Context);
 
 	FUserPtrHandler<CQDIPerInstanceData> InstanceData(Context);
 
@@ -247,27 +284,24 @@ void UNiagaraDataInterfaceCollisionQuery::PerformQuery(FVectorVMContext& Context
 	FRegisterHandler<float> OutCollisionPosX(Context);
 	FRegisterHandler<float> OutCollisionPosY(Context);
 	FRegisterHandler<float> OutCollisionPosZ(Context);
-	FRegisterHandler<float> OutCollisionVelX(Context);
-	FRegisterHandler<float> OutCollisionVelY(Context);
-	FRegisterHandler<float> OutCollisionVelZ(Context);
 	FRegisterHandler<float> OutCollisionNormX(Context);
 	FRegisterHandler<float> OutCollisionNormY(Context);
 	FRegisterHandler<float> OutCollisionNormZ(Context);
 	FRegisterHandler<float> Friction(Context);
 	FRegisterHandler<float> Restitution(Context);
 
+	FScopeLock ScopeLock(&CriticalSection);
 	for (int32 i = 0; i < Context.NumInstances; ++i)
 	{
 		// submit a query, get query id in return
 		//if (PerformParam.Get())
 		{
 			FVector Pos(PosParamX.Get(), PosParamY.Get(), PosParamZ.Get());
-			FVector Vel(VelParamX.Get(), VelParamY.Get(), VelParamZ.Get());
+			FVector Dir(DirParamX.Get(), DirParamY.Get(), DirParamZ.Get());
 			ensure(!Pos.ContainsNaN());
-			ensure(!Vel.ContainsNaN());
 			float Dt = DTParam.Get();
 			float Size = SizeParam.Get();
-			*OutQueryID.GetDest() = InstanceData->CollisionBatch.SubmitQuery(Pos, Vel, Size, Dt);
+			*OutQueryID.GetDest() = InstanceData->CollisionBatch.SubmitQuery(Pos, Dir, Size, Dt);
 		}
 
 		// try to retrieve a query with the supplied query ID
@@ -280,9 +314,6 @@ void UNiagaraDataInterfaceCollisionQuery::PerformQuery(FVectorVMContext& Context
 			*OutCollisionPosX.GetDest() = Res.CollisionPos.X;
 			*OutCollisionPosY.GetDest() = Res.CollisionPos.Y;
 			*OutCollisionPosZ.GetDest() = Res.CollisionPos.Z;
-			*OutCollisionVelX.GetDest() = Res.CollisionVelocity.X;
-			*OutCollisionVelY.GetDest() = Res.CollisionVelocity.Y;
-			*OutCollisionVelZ.GetDest() = Res.CollisionVelocity.Z;
 			*OutCollisionNormX.GetDest() = Res.CollisionNormal.X;
 			*OutCollisionNormY.GetDest() = Res.CollisionNormal.Y;
 			*OutCollisionNormZ.GetDest() = Res.CollisionNormal.Z;
@@ -292,6 +323,14 @@ void UNiagaraDataInterfaceCollisionQuery::PerformQuery(FVectorVMContext& Context
 		else
 		{
 			*OutQueryValid.GetDest() = 0;// ->SetValue(false);
+			*OutCollisionPosX.GetDest() =  0.0f;
+			*OutCollisionPosY.GetDest() =  0.0f;
+			*OutCollisionPosZ.GetDest() =  0.0f;
+			*OutCollisionNormX.GetDest() = 0.0f;
+			*OutCollisionNormY.GetDest() = 0.0f;
+			*OutCollisionNormZ.GetDest() = 0.0f;
+			*Friction.GetDest() = 0.0f;
+			*Restitution.GetDest() = 0.0f;
 		}
 
 		//PerformParam.Advance();
@@ -300,9 +339,6 @@ void UNiagaraDataInterfaceCollisionQuery::PerformQuery(FVectorVMContext& Context
 		OutCollisionPosX.Advance();
 		OutCollisionPosY.Advance();
 		OutCollisionPosZ.Advance();
-		OutCollisionVelX.Advance();
-		OutCollisionVelY.Advance();
-		OutCollisionVelZ.Advance();
 		OutCollisionNormX.Advance();
 		OutCollisionNormY.Advance();
 		OutCollisionNormZ.Advance();
@@ -312,11 +348,12 @@ void UNiagaraDataInterfaceCollisionQuery::PerformQuery(FVectorVMContext& Context
 		PosParamX.Advance();
 		PosParamY.Advance();
 		PosParamZ.Advance();
-		VelParamX.Advance();
-		VelParamY.Advance();
-		VelParamZ.Advance();
+		DirParamX.Advance();
+		DirParamY.Advance();
+		DirParamZ.Advance();
 		DTParam.Advance();
 		SizeParam.Advance();
+		DepthBoundsParam.Advance();
 		OutQueryID.Advance();
 	}
 
@@ -335,6 +372,9 @@ void UNiagaraDataInterfaceCollisionQuery::SubmitQuery(FVectorVMContext& Context)
 
 	DTType DTParam(Context);
 	FUserPtrHandler<CQDIPerInstanceData> InstanceData(Context);
+
+	FScopeLock ScopeLock(&CriticalSection);
+
 	FRegisterHandler<int32> OutQueryID(Context);
 	for (int32 i = 0; i < Context.NumInstances; ++i)
 	{
@@ -375,6 +415,7 @@ void UNiagaraDataInterfaceCollisionQuery::ReadQuery(FVectorVMContext& Context)
 	FRegisterHandler<float> OutCollisionNormY(Context);
 	FRegisterHandler<float> OutCollisionNormZ(Context);
 
+	FScopeLock ScopeLock(&CriticalSection);
 	for (int32 i = 0; i < Context.NumInstances; ++i)
 	{
 		FNiagaraDICollsionQueryResult Res;
@@ -413,7 +454,6 @@ void UNiagaraDataInterfaceCollisionQuery::ReadQuery(FVectorVMContext& Context)
 	}
 }
 
-
 bool UNiagaraDataInterfaceCollisionQuery::PerInstanceTick(void* PerInstanceData, FNiagaraSystemInstance* InSystemInstance, float DeltaSeconds)
 {
 	return false;
@@ -426,3 +466,69 @@ bool UNiagaraDataInterfaceCollisionQuery::PerInstanceTickPostSimulate(void* PerI
 	PIData->CollisionBatch.ClearWrite();
 	return false;
 }
+
+//////////////////////////////////////////////////////////////////////////
+
+struct FNiagaraDataInterfaceParametersCS_CollisionQuery : public FNiagaraDataInterfaceParametersCS
+{
+	virtual void Bind(const FNiagaraDataInterfaceParamRef& ParamRef, const class FShaderParameterMap& ParameterMap) override
+	{
+		SceneDepthTextureParameter.Bind(ParameterMap, TEXT("SceneDepthTexture"));
+		SceneDepthTextureParameterSampler.Bind(ParameterMap, TEXT("SceneDepthTextureSampler"));
+		GBufferATextureParameter.Bind(ParameterMap, TEXT("GBufferATexture"));
+		GBufferATextureParameterSampler.Bind(ParameterMap, TEXT("GBufferATextureSampler"));
+	}
+
+	virtual void Serialize(FArchive& Ar)override
+	{
+		Ar << SceneDepthTextureParameter;
+		Ar << SceneDepthTextureParameterSampler;
+		Ar << GBufferATextureParameter;
+		Ar << GBufferATextureParameterSampler;
+	}
+
+	virtual void Set(FRHICommandList& RHICmdList, FNiagaraShader* Shader, class UNiagaraDataInterface* DataInterface) const override
+	{
+		check(IsInRenderingThread());
+
+		const FComputeShaderRHIParamRef ComputeShaderRHI = Shader->GetComputeShader();
+		//UNiagaraDataInterfaceCollisionQuery* CollisionDI = CastChecked<UNiagaraDataInterfaceCollisionQuery>(DataInterface);
+
+		FTexture2DRHIParamRef SceneDepthTexture = GNiagaraViewDataManager.GetSceneDepthTexture();
+		FTexture2DRHIParamRef GBufferATexture = GNiagaraViewDataManager.GetSceneNormalTexture() ? GNiagaraViewDataManager.GetSceneNormalTexture() : nullptr;
+		RHICmdList.TransitionResource(EResourceTransitionAccess::EReadable, SceneDepthTexture);
+		RHICmdList.TransitionResource(EResourceTransitionAccess::EReadable, GBufferATexture);
+
+		SetTextureParameter(
+			RHICmdList,
+			ComputeShaderRHI,
+			SceneDepthTextureParameter,
+			SceneDepthTextureParameterSampler,
+			TStaticSamplerState<SF_Bilinear, AM_Clamp, AM_Clamp, AM_Clamp>::GetRHI(),
+			SceneDepthTexture//TODO: Does this work for multiple views?
+		);
+
+		SetTextureParameter(
+			RHICmdList,
+			ComputeShaderRHI,
+			GBufferATextureParameter,
+			GBufferATextureParameterSampler,
+			TStaticSamplerState<SF_Bilinear, AM_Clamp, AM_Clamp, AM_Clamp>::GetRHI(),
+			GBufferATexture//TODO: Does this work for multiple views?
+		);
+	}
+
+private:
+
+	/** The SceneDepthTexture parameter for depth buffer collision. */
+	FShaderResourceParameter SceneDepthTextureParameter;
+	FShaderResourceParameter SceneDepthTextureParameterSampler;
+	/** The GBufferATexture parameter for depth buffer collision. */
+	FShaderResourceParameter GBufferATextureParameter;
+	FShaderResourceParameter GBufferATextureParameterSampler;
+};
+
+FNiagaraDataInterfaceParametersCS* UNiagaraDataInterfaceCollisionQuery::ConstructComputeParameters()const
+{
+	return new FNiagaraDataInterfaceParametersCS_CollisionQuery();
+}
\ No newline at end of file
diff --git a/Engine/Plugins/FX/Niagara/Source/Niagara/Private/NiagaraDataInterfaceColorCurve.cpp b/Engine/Plugins/FX/Niagara/Source/Niagara/Private/NiagaraDataInterfaceColorCurve.cpp
--- a/Engine/Plugins/FX/Niagara/Source/Niagara/Private/NiagaraDataInterfaceColorCurve.cpp
+++ b/Engine/Plugins/FX/Niagara/Source/Niagara/Private/NiagaraDataInterfaceColorCurve.cpp
@@ -10,6 +10,8 @@
 //////////////////////////////////////////////////////////////////////////
 //Color Curve
 
+const FName UNiagaraDataInterfaceColorCurve::SampleCurveName(TEXT("SampleColorCurve"));
+
 UNiagaraDataInterfaceColorCurve::UNiagaraDataInterfaceColorCurve(FObjectInitializer const& ObjectInitializer)
 	: Super(ObjectInitializer)
 {
@@ -35,48 +37,29 @@ void UNiagaraDataInterfaceColorCurve::PostLoad()
 
 	const int32 NiagaraVer = GetLinkerCustomVersion(FNiagaraCustomVersion::GUID);
 
-	if (NiagaraVer < FNiagaraCustomVersion::CurveLUTNowOnByDefault)
+	if (NiagaraVer < FNiagaraCustomVersion::LatestVersion)
 	{
 		UpdateLUT();
 	}
-	TArray<float> OldLUT = ShaderLUT;
-	UpdateLUT();
-	if (!CompareLUTS(OldLUT))
+#if !UE_BUILD_SHIPPING
+	else
 	{
-		UE_LOG(LogNiagara, Log, TEXT("PostLoad LUT generation is out of sync. Please investigate. %s"), *GetPathName());
-	}
-
-}
-
-#if WITH_EDITOR
-
-void UNiagaraDataInterfaceColorCurve::PostEditChangeProperty(struct FPropertyChangedEvent& PropertyChangedEvent)
-{
-	Super::PostEditChangeProperty(PropertyChangedEvent);
+		TArray<float> OldLUT = ShaderLUT;
+		UpdateLUT();
 
-	if (PropertyChangedEvent.Property && PropertyChangedEvent.Property->GetFName() == GET_MEMBER_NAME_CHECKED(UNiagaraDataInterfaceColorCurve, CurveToCopy))
-	{
-		UCurveLinearColor* ColorCurveAsset = Cast<UCurveLinearColor>(CurveToCopy.TryLoad());
-		if(ColorCurveAsset)
+		if (!CompareLUTS(OldLUT))
 		{
-			Modify();
-			RedCurve = ColorCurveAsset->FloatCurves[0];
-			GreenCurve = ColorCurveAsset->FloatCurves[1];
-			BlueCurve = ColorCurveAsset->FloatCurves[2];
-			AlphaCurve = ColorCurveAsset->FloatCurves[3];
+			UE_LOG(LogNiagara, Log, TEXT("PostLoad LUT generation is out of sync. Please investigate. %s"), *GetPathName());
 		}
 	}
-
-	UpdateLUT();
-}
-
 #endif
+}
 
 void UNiagaraDataInterfaceColorCurve::UpdateLUT()
 {
 	ShaderLUT.Empty();
 
-	if (bAllowUnnormalizedLUT && (RedCurve.GetNumKeys() > 0 || GreenCurve.GetNumKeys() > 0 || BlueCurve.GetNumKeys() > 0 || AlphaCurve.GetNumKeys() > 0))
+	if ((RedCurve.GetNumKeys() > 0 || GreenCurve.GetNumKeys() > 0 || BlueCurve.GetNumKeys() > 0 || AlphaCurve.GetNumKeys() > 0))
 	{
 		LUTMinTime = FLT_MAX;
 		LUTMinTime = FMath::Min(RedCurve.GetNumKeys() > 0 ? RedCurve.GetFirstKey().Time : LUTMinTime, LUTMinTime);
@@ -100,7 +83,7 @@ void UNiagaraDataInterfaceColorCurve::UpdateLUT()
 
 	for (uint32 i = 0; i < CurveLUTWidth; i++)
 	{
-		float X = UnnormalizeTime(i / (float)CurveLUTWidth);
+		float X = UnnormalizeTime(i / (float)CurveLUTWidthMinusOne);
 		FLinearColor C(RedCurve.Eval(X), GreenCurve.Eval(X), BlueCurve.Eval(X), AlphaCurve.Eval(X));
 		ShaderLUT.Add(C.R);
 		ShaderLUT.Add(C.G);
@@ -154,7 +137,7 @@ void UNiagaraDataInterfaceColorCurve::GetCurveData(TArray<FCurveData>& OutCurveD
 void UNiagaraDataInterfaceColorCurve::GetFunctions(TArray<FNiagaraFunctionSignature>& OutFunctions)
 {
 	FNiagaraFunctionSignature Sig;
-	Sig.Name = TEXT("SampleColorCurve");
+	Sig.Name = SampleCurveName;
 	Sig.bMemberFunction = true;
 	Sig.bRequiresContext = false;
 	Sig.Inputs.Add(FNiagaraVariable(FNiagaraTypeDefinition(GetClass()), TEXT("ColorCurve")));
@@ -169,101 +152,57 @@ void UNiagaraDataInterfaceColorCurve::GetFunctions(TArray<FNiagaraFunctionSignat
 // the HLSL in the spirit of a static switch
 // TODO: need a way to identify each specific function here
 // 
-bool UNiagaraDataInterfaceColorCurve::GetFunctionHLSL(const FName& DefinitionFunctionName, FString InstanceFunctionName, TArray<FDIGPUBufferParamDescriptor> &Descriptors, FString &HLSLInterfaceID, FString &OutHLSL)
-{
-	FString BufferName = Descriptors[0].BufferParamName;
-	OutHLSL += TEXT("void ") + InstanceFunctionName + TEXT("(in float In_X, out float4 Out_Value) \n{\n");
-	OutHLSL += TEXT("\t Out_Value.x = ") + BufferName + TEXT("[(int)(In_X *") + FString::FromInt(CurveLUTWidth) + TEXT(")* 4 ];");
-	OutHLSL += TEXT("\t Out_Value.y = ") + BufferName + TEXT("[1+ (int)(In_X *") + FString::FromInt(CurveLUTWidth) + TEXT(")* 4 ];");
-	OutHLSL += TEXT("\t Out_Value.z = ") + BufferName + TEXT("[2+ (int)(In_X *") + FString::FromInt(CurveLUTWidth) + TEXT(")* 4 ];");
-	OutHLSL += TEXT("\t Out_Value.w = ") + BufferName + TEXT("[3+ (int)(In_X *") + FString::FromInt(CurveLUTWidth) + TEXT(")* 4 ];");
-	OutHLSL += TEXT("\n}\n");
-	return true;
-}
-
-// build buffer definition hlsl
-// 1. Choose a buffer name, add the data interface ID (important!)
-// 2. add a DIGPUBufferParamDescriptor to the array argument; that'll be passed on to the FNiagaraShader for binding to a shader param, that can
-// then later be found by name via FindDIBufferParam for setting; 
-// 3. store buffer declaration hlsl in OutHLSL
-// multiple buffers can be defined at once here
-//
-void UNiagaraDataInterfaceColorCurve::GetBufferDefinitionHLSL(FString DataInterfaceID, TArray<FDIGPUBufferParamDescriptor> &BufferDescriptors, FString &OutHLSL)
-{
-	FString BufferName = "CurveLUT" + DataInterfaceID;
-	OutHLSL += TEXT("Buffer<float> ") + BufferName + TEXT(";\n");
+bool UNiagaraDataInterfaceColorCurve::GetFunctionHLSL(const FName& DefinitionFunctionName, FString InstanceFunctionName, FNiagaraDataInterfaceGPUParamInfo& ParamInfo, FString& OutHLSL)
+{
+	FString TimeToLUTFrac = TEXT("TimeToLUTFraction_") + ParamInfo.DataInterfaceHLSLSymbol;
+	FString Sample = TEXT("SampleCurve_") + ParamInfo.DataInterfaceHLSLSymbol;
+	OutHLSL += FString::Printf(TEXT("\
+void %s(in float In_X, out float4 Out_Value) \n\
+{ \n\
+	float RemappedX = %s(In_X) * %u; \n\
+	float Prev = floor(RemappedX); \n\
+	float Next = Prev < %u ? Prev + 1.0 : Prev; \n\
+	float Interp = RemappedX - Prev; \n\
+	Prev *= %u; \n\
+	Next *= %u; \n\
+	float4 A = float4(%s(Prev), %s(Prev + 1), %s(Prev + 2), %s(Prev + 3)); \n\
+	float4 B = float4(%s(Next), %s(Next + 1), %s(Next + 2), %s(Next + 3)); \n\
+	Out_Value = lerp(A, B, Interp); \n\
+}\n")
+, *InstanceFunctionName, *TimeToLUTFrac, CurveLUTWidthMinusOne, CurveLUTWidthMinusOne, CurveLUTNumElems, CurveLUTNumElems
+, *Sample, *Sample, *Sample, *Sample, *Sample, *Sample, *Sample, *Sample);
 
-	BufferDescriptors.Add( FDIGPUBufferParamDescriptor(BufferName, 0) );		// add a descriptor for shader parameter binding
-
-}
-
-// called after translate, to setup buffers matching the buffer descriptors generated during hlsl translation
-// need to do this because the script used during translate is a clone, including its DIs
-//
-void UNiagaraDataInterfaceColorCurve::SetupBuffers(FDIBufferDescriptorStore &BufferDescriptors)
-{
-	GPUBuffers.Empty();
-	for (FDIGPUBufferParamDescriptor &Desc : BufferDescriptors.Descriptors)
-	{
-		FNiagaraDataInterfaceBufferData BufferData(*Desc.BufferParamName);	// store off the data for later use
-		GPUBuffers.Add(BufferData);
-	}
-	GPUBufferDirty = true;
-}
-
-// return the GPU buffer array (called from NiagaraInstanceBatcher to get the buffers for setting to the shader)
-// we lazily update the buffer with a new LUT here if necessary
-//
-TArray<FNiagaraDataInterfaceBufferData> &UNiagaraDataInterfaceColorCurve::GetBufferDataArray()
-{
-	check(IsInRenderingThread());
-	if (GPUBufferDirty)
-	{
-		check(GPUBuffers.Num() > 0);
-
-		FNiagaraDataInterfaceBufferData &GPUBuffer = GPUBuffers[0];
-		GPUBuffer.Buffer.Release();
-
-/*		TResourceArray<float> FloatRes;
-		FloatRes.AddZeroed(ShaderLUT.Num());
-		uint32 BufferSize = ShaderLUT.Num() * sizeof(float);
-		FPlatformMemory::Memcpy(FloatRes.GetData(), ShaderLUT.GetData(), BufferSize);*/
-
-		GPUBuffer.Buffer.Initialize(sizeof(float), CurveLUTWidth*4 , EPixelFormat::PF_R32_FLOAT, BUF_Static);	// always allocate for up to 64 data sets
-		uint32 BufferSize = ShaderLUT.Num() * sizeof(float);
-		int32 *BufferData = static_cast<int32*>(RHILockVertexBuffer(GPUBuffer.Buffer.Buffer, 0, BufferSize, EResourceLockMode::RLM_WriteOnly));
-		FPlatformMemory::Memcpy(BufferData, ShaderLUT.GetData(), BufferSize);
-		RHIUnlockVertexBuffer(GPUBuffer.Buffer.Buffer);
-		GPUBufferDirty = false;
-	}
-
-	return GPUBuffers;
+	return true;
 }
 
-
-
-
-DEFINE_NDI_FUNC_BINDER(UNiagaraDataInterfaceColorCurve, SampleCurve);
-FVMExternalFunction UNiagaraDataInterfaceColorCurve::GetVMExternalFunction(const FVMExternalFunctionBindingInfo& BindingInfo, void* InstanceData)
+DEFINE_NDI_RAW_FUNC_BINDER(UNiagaraDataInterfaceColorCurve, SampleCurve);
+void UNiagaraDataInterfaceColorCurve::GetVMExternalFunction(const FVMExternalFunctionBindingInfo& BindingInfo, void* InstanceData, FVMExternalFunction &OutFunc)
 {
-	if (BindingInfo.Name == TEXT("SampleColorCurve") && BindingInfo.GetNumInputs() == 1 && BindingInfo.GetNumOutputs() == 4)
+	if (BindingInfo.Name == SampleCurveName && BindingInfo.GetNumInputs() == 1 && BindingInfo.GetNumOutputs() == 4)
 	{
-		return TCurveUseLUTBinder<TNDIParamBinder<0, float, NDI_FUNC_BINDER(UNiagaraDataInterfaceColorCurve, SampleCurve)>>::Bind(this, BindingInfo, InstanceData);
+		TCurveUseLUTBinder<TNDIParamBinder<0, float, NDI_RAW_FUNC_BINDER(UNiagaraDataInterfaceColorCurve, SampleCurve)>>::Bind(this, BindingInfo, InstanceData, OutFunc);
 	}
 	else
 	{
 		UE_LOG(LogNiagara, Error, TEXT("Could not find data interface external function.\n\tExpected Name: SampleColorCurve  Actual Name: %s\n\tExpected Inputs: 1  Actual Inputs: %i\n\tExpected Outputs: 4  Actual Outputs: %i"),
 			*BindingInfo.Name.ToString(), BindingInfo.GetNumInputs(), BindingInfo.GetNumOutputs());
-		return FVMExternalFunction();
+		OutFunc = FVMExternalFunction();
 	}
 }
 
 template<>
 FORCEINLINE_DEBUGGABLE FLinearColor UNiagaraDataInterfaceColorCurve::SampleCurveInternal<TIntegralConstant<bool, true>>(float X)
 {
-	float NormalizedX = NormalizeTime(X);
-	int32 AccessIdx = FMath::Clamp<int32>(FMath::TruncToInt(NormalizedX * CurveLUTWidthMinusOne), 0, CurveLUTWidthMinusOne) * CurveLUTNumElems;
-	return FLinearColor(ShaderLUT[AccessIdx], ShaderLUT[AccessIdx + 1], ShaderLUT[AccessIdx + 2], ShaderLUT[AccessIdx + 3]);
+	float RemappedX = FMath::Clamp(NormalizeTime(X) * CurveLUTWidthMinusOne, 0.0f, (float)CurveLUTWidthMinusOne);
+	float PrevEntry = FMath::TruncToFloat(RemappedX);
+	float NextEntry = PrevEntry < (float)CurveLUTWidthMinusOne ? PrevEntry + 1.0f : PrevEntry;
+	float Interp = RemappedX - PrevEntry;
+	
+	int32 AIndex = PrevEntry * CurveLUTNumElems;
+	int32 BIndex = NextEntry * CurveLUTNumElems;
+	FLinearColor A = FLinearColor(ShaderLUT[AIndex], ShaderLUT[AIndex + 1], ShaderLUT[AIndex + 2], ShaderLUT[AIndex + 3]);
+	FLinearColor B = FLinearColor(ShaderLUT[BIndex], ShaderLUT[BIndex + 1], ShaderLUT[BIndex + 2], ShaderLUT[BIndex + 3]);
+	return FMath::Lerp(A, B, Interp);
 }
 
 template<>
@@ -284,16 +223,11 @@ void UNiagaraDataInterfaceColorCurve::SampleCurve(FVectorVMContext& Context)
 
 	for (int32 i = 0; i < Context.NumInstances; ++i)
 	{
-		float X = XParam.Get();
+		float X = XParam.GetAndAdvance();
 		FLinearColor C = SampleCurveInternal<UseLUT>(X);
-		*SamplePtrR.GetDest() = C.R;
-		*SamplePtrG.GetDest() = C.G;
-		*SamplePtrB.GetDest() = C.B;
-		*SamplePtrA.GetDest() = C.A;
-		XParam.Advance();
-		SamplePtrR.Advance();
-		SamplePtrG.Advance();
-		SamplePtrB.Advance();
-		SamplePtrA.Advance();
+		*SamplePtrR.GetDestAndAdvance() = C.R;
+		*SamplePtrG.GetDestAndAdvance() = C.G;
+		*SamplePtrB.GetDestAndAdvance() = C.B;
+		*SamplePtrA.GetDestAndAdvance() = C.A;
 	}
 }
diff --git a/Engine/Plugins/FX/Niagara/Source/Niagara/Private/NiagaraDataInterfaceCurlNoise.cpp b/Engine/Plugins/FX/Niagara/Source/Niagara/Private/NiagaraDataInterfaceCurlNoise.cpp
--- a/Engine/Plugins/FX/Niagara/Source/Niagara/Private/NiagaraDataInterfaceCurlNoise.cpp
+++ b/Engine/Plugins/FX/Niagara/Source/Niagara/Private/NiagaraDataInterfaceCurlNoise.cpp
@@ -1,6 +1,11 @@
 // Copyright 1998-2018 Epic Games, Inc. All Rights Reserved.
 
 #include "NiagaraDataInterfaceCurlNoise.h"
+#include "NiagaraShader.h"
+#include "ShaderParameterUtils.h"
+
+const FName UNiagaraDataInterfaceCurlNoise::SampleNoiseFieldName(TEXT("SampleNoiseField"));
+const FString UNiagaraDataInterfaceCurlNoise::CurlNoiseBufferName(TEXT("CurlNoiseBuffer_"));
 
 UNiagaraDataInterfaceCurlNoise::UNiagaraDataInterfaceCurlNoise(FObjectInitializer const& ObjectInitializer)
 	: Super(ObjectInitializer)
@@ -66,7 +71,7 @@ bool UNiagaraDataInterfaceCurlNoise::Equals(const UNiagaraDataInterface* Other)
 void UNiagaraDataInterfaceCurlNoise::GetFunctions(TArray<FNiagaraFunctionSignature>& OutFunctions)
 {
 	FNiagaraFunctionSignature Sig;
-	Sig.Name = TEXT("SampleNoiseField");
+	Sig.Name = SampleNoiseFieldName;
 	Sig.bMemberFunction = true;
 	Sig.bRequiresContext = false;
 	Sig.Inputs.Add(FNiagaraVariable(FNiagaraTypeDefinition(GetClass()), TEXT("NoiseField")));
@@ -77,12 +82,12 @@ void UNiagaraDataInterfaceCurlNoise::GetFunctions(TArray<FNiagaraFunctionSignatu
 	OutFunctions.Add(Sig);
 }
 
-DEFINE_NDI_FUNC_BINDER(UNiagaraDataInterfaceCurlNoise, SampleNoiseField);
-FVMExternalFunction UNiagaraDataInterfaceCurlNoise::GetVMExternalFunction(const FVMExternalFunctionBindingInfo& BindingInfo, void* InstanceData)
+DEFINE_NDI_RAW_FUNC_BINDER(UNiagaraDataInterfaceCurlNoise, SampleNoiseField);
+void UNiagaraDataInterfaceCurlNoise::GetVMExternalFunction(const FVMExternalFunctionBindingInfo& BindingInfo, void* InstanceData, FVMExternalFunction &OutFunc)
 {
-	check(BindingInfo.Name == TEXT("SampleNoiseField"));
+	check(BindingInfo.Name == SampleNoiseFieldName);
 	check(BindingInfo.GetNumInputs() == 3 && BindingInfo.GetNumOutputs() == 3);
-	return TNDIParamBinder<0, float, TNDIParamBinder<1, float, TNDIParamBinder<2, float, NDI_FUNC_BINDER(UNiagaraDataInterfaceCurlNoise, SampleNoiseField)>>>::Bind(this, BindingInfo, InstanceData);
+	TNDIParamBinder<0, float, TNDIParamBinder<1, float, TNDIParamBinder<2, float, NDI_RAW_FUNC_BINDER(UNiagaraDataInterfaceCurlNoise, SampleNoiseField)>>>::Bind(this, BindingInfo, InstanceData, OutFunc);
 }
 
 template<typename XType, typename YType, typename ZType>
@@ -94,22 +99,22 @@ void UNiagaraDataInterfaceCurlNoise::SampleNoiseField(FVectorVMContext& Context)
 	FRegisterHandler<float> OutSampleX(Context);
 	FRegisterHandler<float> OutSampleY(Context);
 	FRegisterHandler<float> OutSampleZ(Context);
-	
-	for (int32 InstanceIdx = 0; InstanceIdx < Context.NumInstances; ++InstanceIdx)
-	{
-		VectorRegister InCoords = MakeVectorRegister(XParam.Get(), YParam.Get(), ZParam.Get(), 0.0f);
 
-		const VectorRegister One = MakeVectorRegister(1.0f, 1.0f, 1.0f, 1.0f);
-		const VectorRegister Zero = MakeVectorRegister(0.0f, 0.0f, 0.0f, 0.0f);
-		// We use 15,15,15 here because the value will be between 0 and 1. 1 should indicate the maximal index of the 
-		// real table and we pad that out to handle border work. So our values for the Cxyz below should not
-		// ever be 17 in any dimension as that would overflow the array.
-		const VectorRegister VecSize = MakeVectorRegister(15.0f, 15.0f, 15.0f, 15.0f);
+	const VectorRegister One = MakeVectorRegister(1.0f, 1.0f, 1.0f, 1.0f);
+	const VectorRegister Zero = MakeVectorRegister(0.0f, 0.0f, 0.0f, 0.0f);
+	// We use 15,15,15 here because the value will be between 0 and 1. 1 should indicate the maximal index of the 
+	// real table and we pad that out to handle border work. So our values for the Cxyz below should not
+	// ever be 17 in any dimension as that would overflow the array.
+	const VectorRegister VecSize = MakeVectorRegister(15.0f, 15.0f, 15.0f, 15.0f);
 
+	float Di = 0.2f; // Hard-coded scale TODO remove!
+	const VectorRegister Div = MakeVectorRegister(Di, Di, Di, Di);
+
+	for (int32 InstanceIdx = 0; InstanceIdx < Context.NumInstances; ++InstanceIdx)
+	{
+		VectorRegister InCoords = MakeVectorRegister(XParam.GetAndAdvance(), YParam.GetAndAdvance(), ZParam.GetAndAdvance(), 0.0f);
 		VectorRegister Dst = MakeVectorRegister(0.0f, 0.0f, 0.0f, 0.0f);
 
-		float Di = 0.2f; // Hard-coded scale TODO remove!
-		VectorRegister Div = MakeVectorRegister(Di, Di, Di, Di);
 		VectorRegister Coords = VectorMod(VectorAbs(VectorMultiply(InCoords, Div)), VecSize);
 		Coords = VectorMin(Coords, VecSize);
 		Coords = VectorMax(Coords, Zero);
@@ -143,29 +148,15 @@ void UNiagaraDataInterfaceCurlNoise::SampleNoiseField(FVectorVMContext& Context)
 		Dst = VectorAdd(Dst, ZV);
 
 		float *RegPtr = reinterpret_cast<float*>(&Dst);
-		*OutSampleX.GetDest() = RegPtr[0];
-		*OutSampleY.GetDest() = RegPtr[1];
-		*OutSampleZ.GetDest() = RegPtr[2];
-
-		XParam.Advance();
-		YParam.Advance();
-		ZParam.Advance();
-		OutSampleX.Advance();
-		OutSampleY.Advance();
-		OutSampleZ.Advance();
+		*OutSampleX.GetDestAndAdvance() = RegPtr[0];
+		*OutSampleY.GetDestAndAdvance() = RegPtr[1];
+		*OutSampleZ.GetDestAndAdvance() = RegPtr[2];
 	}
 }
 
-
-
-// build the shader function HLSL; function name is passed in, as it's defined per-DI; that way, configuration could change
-// the HLSL in the spirit of a static switch
-// TODO: need a way to identify each specific function here
-// 
-bool UNiagaraDataInterfaceCurlNoise::GetFunctionHLSL(const FName& DefinitionFunctionName, FString InstanceFunctionName, TArray<FDIGPUBufferParamDescriptor> &Descriptors, FString &HLSLInterfaceID, FString &OutHLSL)
+bool UNiagaraDataInterfaceCurlNoise::GetFunctionHLSL(const FName& DefinitionFunctionName, FString InstanceFunctionName, FNiagaraDataInterfaceGPUParamInfo& ParamInfo, FString& OutHLSL)
 {
-
-	FString BufferName = Descriptors[0].BufferParamName;
+	FString BufferName = CurlNoiseBufferName + ParamInfo.DataInterfaceHLSLSymbol;
 	OutHLSL += TEXT("void ") + InstanceFunctionName + TEXT("(in float3 In_XYZ, out float3 Out_Value) \n{\n");
 	OutHLSL += TEXT("\t float3 a = trunc((In_XYZ*0.2) / 15.0);\n");
 	OutHLSL += TEXT("\t float3 ModXYZ = (In_XYZ*0.2) - a*15.0;\n");
@@ -192,48 +183,19 @@ bool UNiagaraDataInterfaceCurlNoise::GetFunctionHLSL(const FName& DefinitionFunc
 	return true;
 }
 
-// build buffer definition hlsl
-// 1. Choose a buffer name, add the data interface ID (important!)
-// 2. add a DIGPUBufferParamDescriptor to the array argument; that'll be passed on to the FNiagaraShader for binding to a shader param, that can
-// then later be found by name via FindDIBufferParam for setting; 
-// 3. store buffer declaration hlsl in OutHLSL
-// multiple buffers can be defined at once here
-//
-void UNiagaraDataInterfaceCurlNoise::GetBufferDefinitionHLSL(FString DataInterfaceID, TArray<FDIGPUBufferParamDescriptor> &BufferDescriptors, FString &OutHLSL)
+void UNiagaraDataInterfaceCurlNoise::GetParameterDefinitionHLSL(FNiagaraDataInterfaceGPUParamInfo& ParamInfo, FString& OutHLSL)
 {
-	FString BufferName = "CurlNoiseLUT" + DataInterfaceID;
+	FString BufferName = CurlNoiseBufferName + ParamInfo.DataInterfaceHLSLSymbol;
 	OutHLSL += TEXT("Buffer<float4> ") + BufferName + TEXT(";\n");
-
-	BufferDescriptors.Add(FDIGPUBufferParamDescriptor(BufferName, 0));		// add a descriptor for shader parameter binding
-
 }
 
-// called after translate, to setup buffers matching the buffer descriptors generated during hlsl translation
-// need to do this because the script used during translate is a clone, including its DIs
-//
-void UNiagaraDataInterfaceCurlNoise::SetupBuffers(FDIBufferDescriptorStore &BufferDescriptors)
-{
-	GPUBuffers.Empty();
-	for (FDIGPUBufferParamDescriptor &Desc : BufferDescriptors.Descriptors)
-	{
-		FNiagaraDataInterfaceBufferData BufferData(*Desc.BufferParamName);	// store off the data for later use
-		GPUBuffers.Add(BufferData);
-	}
-	bGPUBufferDirty = true;
-}
-
-// return the GPU buffer array (called from NiagaraInstanceBatcher to get the buffers for setting to the shader)
-// we lazily update the buffer with a new LUT here if necessary
-//
-TArray<FNiagaraDataInterfaceBufferData> &UNiagaraDataInterfaceCurlNoise::GetBufferDataArray()
+FRWBuffer& UNiagaraDataInterfaceCurlNoise::GetGPUBuffer()
 {
 	check(IsInRenderingThread());
 
 	if (bGPUBufferDirty)
 	{
-		check(GPUBuffers.Num() > 0);
-		FNiagaraDataInterfaceBufferData &GPUBuffer = GPUBuffers[0];
-		GPUBuffer.Buffer.Release();
+		GPUBuffer.Release();
 		uint32 BufferSize = 17 * 17 * 17 * sizeof(float) * 4;
 		//int32 *BufferData = static_cast<int32*>(RHILockVertexBuffer(GPUBuffer.Buffer.Buffer, 0, BufferSize, EResourceLockMode::RLM_WriteOnly));
 		TResourceArray<FVector4> TempTable;
@@ -249,14 +211,13 @@ TArray<FNiagaraDataInterfaceBufferData> &UNiagaraDataInterfaceCurlNoise::GetBuff
 				}
 			}
 		}
-		GPUBuffer.Buffer.Initialize(sizeof(float) * 4, 17 * 17 * 17, EPixelFormat::PF_A32B32G32R32F, BUF_Static, TEXT("CurlnoiseTable"), &TempTable);
+		GPUBuffer.Initialize(sizeof(float) * 4, 17 * 17 * 17, EPixelFormat::PF_A32B32G32R32F, BUF_Static, TEXT("CurlnoiseTable"), &TempTable);
 		//FPlatformMemory::Memcpy(BufferData, TempTable, BufferSize);
 		//RHIUnlockVertexBuffer(GPUBuffer.Buffer.Buffer);
 		bGPUBufferDirty = false;
 	}
 
-
-	return GPUBuffers;
+	return GPUBuffer;
 }
 
 
@@ -392,3 +353,37 @@ void UNiagaraDataInterfaceCurlNoise::InitNoiseLUT()
 
 	bGPUBufferDirty = true;
 }
+
+
+struct FNiagaraDataInterfaceParametersCS_CurlNoise : public FNiagaraDataInterfaceParametersCS
+{
+	virtual void Bind(const FNiagaraDataInterfaceParamRef& ParamRef, const class FShaderParameterMap& ParameterMap) override
+	{
+		CurlNoiseBuffer.Bind(ParameterMap, *(UNiagaraDataInterfaceCurlNoise::CurlNoiseBufferName + ParamRef.ParameterInfo.DataInterfaceHLSLSymbol));
+	}
+
+	virtual void Serialize(FArchive& Ar)override
+	{
+		Ar << CurlNoiseBuffer;
+	}
+
+	virtual void Set(FRHICommandList& RHICmdList, FNiagaraShader* Shader, class UNiagaraDataInterface* DataInterface) const override
+	{
+		check(IsInRenderingThread());
+
+		const FComputeShaderRHIParamRef ComputeShaderRHI = Shader->GetComputeShader();
+		UNiagaraDataInterfaceCurlNoise* CurlNoiseDI = CastChecked<UNiagaraDataInterfaceCurlNoise>(DataInterface);
+		FRWBuffer& GPUBuffer = CurlNoiseDI->GetGPUBuffer();
+
+		RHICmdList.SetShaderResourceViewParameter(ComputeShaderRHI, CurlNoiseBuffer.GetBaseIndex(), GPUBuffer.SRV);
+	}
+
+private:
+
+	FShaderResourceParameter CurlNoiseBuffer;
+};
+
+FNiagaraDataInterfaceParametersCS* UNiagaraDataInterfaceCurlNoise::ConstructComputeParameters()const
+{
+	return new FNiagaraDataInterfaceParametersCS_CurlNoise();
+}
\ No newline at end of file
diff --git a/Engine/Plugins/FX/Niagara/Source/Niagara/Private/NiagaraDataInterfaceCurve.cpp b/Engine/Plugins/FX/Niagara/Source/Niagara/Private/NiagaraDataInterfaceCurve.cpp
--- a/Engine/Plugins/FX/Niagara/Source/Niagara/Private/NiagaraDataInterfaceCurve.cpp
+++ b/Engine/Plugins/FX/Niagara/Source/Niagara/Private/NiagaraDataInterfaceCurve.cpp
@@ -7,6 +7,7 @@
 #include "NiagaraTypes.h"
 #include "NiagaraCustomVersion.h"
 
+const FName UNiagaraDataInterfaceCurve::SampleCurveName(TEXT("SampleCurve"));
 
 UNiagaraDataInterfaceCurve::UNiagaraDataInterfaceCurve(FObjectInitializer const& ObjectInitializer)
 	: Super(ObjectInitializer)
@@ -26,26 +27,6 @@ void UNiagaraDataInterfaceCurve::PostInitProperties()
 	UpdateLUT();
 }
 
-#if WITH_EDITOR
-
-void UNiagaraDataInterfaceCurve::PostEditChangeProperty(struct FPropertyChangedEvent& PropertyChangedEvent)
-{
-	Super::PostEditChangeProperty(PropertyChangedEvent);
-
-	if (PropertyChangedEvent.Property && PropertyChangedEvent.Property->GetFName() == GET_MEMBER_NAME_CHECKED(UNiagaraDataInterfaceCurve, CurveToCopy))
-	{
-		UCurveFloat* CurveAsset = Cast<UCurveFloat>(CurveToCopy.TryLoad());
-		if (CurveAsset)
-		{
-			Modify();
-			Curve = CurveAsset->FloatCurve;
-		}
-	}
-	UpdateLUT();
-}
-
-#endif
-
 bool UNiagaraDataInterfaceCurve::CopyToInternal(UNiagaraDataInterface* Destination) const 
 {
 	if (!Super::CopyToInternal(Destination))
@@ -54,7 +35,10 @@ bool UNiagaraDataInterfaceCurve::CopyToInternal(UNiagaraDataInterface* Destinati
 	}
 	CastChecked<UNiagaraDataInterfaceCurve>(Destination)->Curve = Curve;
 	CastChecked<UNiagaraDataInterfaceCurve>(Destination)->UpdateLUT();
-	ensure(CompareLUTS(CastChecked<UNiagaraDataInterfaceCurve>(Destination)->ShaderLUT));
+	if (!CompareLUTS(CastChecked<UNiagaraDataInterfaceCurve>(Destination)->ShaderLUT))
+	{
+		UE_LOG(LogNiagara, Log, TEXT("Post CopyToInternal LUT generation is out of sync. Please investigate. %s"), *GetPathName());
+	}
 	return true;
 }
 
@@ -76,7 +60,7 @@ void UNiagaraDataInterfaceCurve::GetCurveData(TArray<FCurveData>& OutCurveData)
 void UNiagaraDataInterfaceCurve::GetFunctions(TArray<FNiagaraFunctionSignature>& OutFunctions)
 {
 	FNiagaraFunctionSignature Sig;
-	Sig.Name = TEXT("SampleCurve");
+	Sig.Name = SampleCurveName;
 	Sig.bMemberFunction = true;
 	Sig.bRequiresContext = false;
 	Sig.Inputs.Add(FNiagaraVariable(FNiagaraTypeDefinition(GetClass()), TEXT("Curve")));
@@ -92,25 +76,29 @@ void UNiagaraDataInterfaceCurve::PostLoad()
 	Super::PostLoad();
 	const int32 NiagaraVer = GetLinkerCustomVersion(FNiagaraCustomVersion::GUID);
 
-	if (NiagaraVer < FNiagaraCustomVersion::CurveLUTNowOnByDefault)
+	if (NiagaraVer < FNiagaraCustomVersion::LatestVersion)
 	{
 		UpdateLUT();
 	}
-
-	TArray<float> OldLUT = ShaderLUT;
-	UpdateLUT();
-
-	if (!CompareLUTS(OldLUT))
+#if !UE_BUILD_SHIPPING
+	else
 	{
-		UE_LOG(LogNiagara, Log, TEXT("PostLoad LUT generation is out of sync. Please investigate. %s"), *GetPathName());
+		TArray<float> OldLUT = ShaderLUT;
+		UpdateLUT();
+
+		if (!CompareLUTS(OldLUT))
+		{
+			UE_LOG(LogNiagara, Log, TEXT("PostLoad LUT generation is out of sync. Please investigate. %s"), *GetPathName());
+		}
 	}
+#endif
 }
 
 
 void UNiagaraDataInterfaceCurve::UpdateLUT()
 {
 	ShaderLUT.Empty();
-	if (bAllowUnnormalizedLUT && Curve.GetNumKeys() > 0)
+	if (Curve.GetNumKeys() > 0)
 	{
 		LUTMinTime = Curve.GetFirstKey().Time;
 		LUTMaxTime = Curve.GetLastKey().Time;
@@ -125,7 +113,7 @@ void UNiagaraDataInterfaceCurve::UpdateLUT()
 
 	for (uint32 i = 0; i < CurveLUTWidth; i++)
 	{
-		float X = UnnormalizeTime(i / (float)CurveLUTWidth);
+		float X = UnnormalizeTime(i / (float)CurveLUTWidthMinusOne);
 		float C = Curve.Eval(X);
 		ShaderLUT.Add(C);
 	}
@@ -136,93 +124,57 @@ void UNiagaraDataInterfaceCurve::UpdateLUT()
 // the HLSL in the spirit of a static switch
 // TODO: need a way to identify each specific function here
 // 
-bool UNiagaraDataInterfaceCurve::GetFunctionHLSL(const FName& DefinitionFunctionName, FString InstanceFunctionName, TArray<FDIGPUBufferParamDescriptor> &Descriptors, FString &HLSLInterfaceID, FString &OutHLSL)
-{
-	FString BufferName = Descriptors[0].BufferParamName;
-	OutHLSL += TEXT("void ") + InstanceFunctionName + TEXT("(in float In_X, out float Out_Value) \n{\n");
-	OutHLSL += TEXT("\t Out_Value = ") + BufferName + TEXT("[(int)(clamp(In_X, 0.0, 1.0) *") + FString::FromInt(CurveLUTWidth) + TEXT(") ];");
-	OutHLSL += TEXT("\n}\n");
-	return true;
-}
-
-// build buffer definition hlsl
-// 1. Choose a buffer name, add the data interface ID (important!)
-// 2. add a DIGPUBufferParamDescriptor to the array argument; that'll be passed on to the FNiagaraShader for binding to a shader param, that can
-// then later be found by name via FindDIBufferParam for setting; 
-// 3. store buffer declaration hlsl in OutHLSL
-// multiple buffers can be defined at once here
-//
-void UNiagaraDataInterfaceCurve::GetBufferDefinitionHLSL(FString DataInterfaceID, TArray<FDIGPUBufferParamDescriptor> &BufferDescriptors, FString &OutHLSL)
-{
-	FString BufferName = "CurveLUT" + DataInterfaceID;
-	OutHLSL += TEXT("Buffer<float> ") + BufferName + TEXT(";\n");
-
-	BufferDescriptors.Add(FDIGPUBufferParamDescriptor(BufferName, 0));		// add a descriptor for shader parameter binding
-
-}
-
-// called after translate, to setup buffers matching the buffer descriptors generated during hlsl translation
-// need to do this because the script used during translate is a clone, including its DIs
-//
-void UNiagaraDataInterfaceCurve::SetupBuffers(FDIBufferDescriptorStore &BufferDescriptors)
-{
-	for (FDIGPUBufferParamDescriptor &Desc : BufferDescriptors.Descriptors)
-	{
-		FNiagaraDataInterfaceBufferData BufferData(*Desc.BufferParamName);	// store off the data for later use
-		GPUBuffers.Add(BufferData);
-	}
-}
-
-// return the GPU buffer array (called from NiagaraInstanceBatcher to get the buffers for setting to the shader)
-// we lazily update the buffer with a new LUT here if necessary
-//
-TArray<FNiagaraDataInterfaceBufferData> &UNiagaraDataInterfaceCurve::GetBufferDataArray()
-{
-	check(IsInRenderingThread());
-	if (GPUBufferDirty)
-	{
-		check(GPUBuffers.Num() > 0);
-
-		FNiagaraDataInterfaceBufferData &GPUBuffer = GPUBuffers[0];
-		GPUBuffer.Buffer.Release();
-		GPUBuffer.Buffer.Initialize(sizeof(float), CurveLUTWidth * 4, EPixelFormat::PF_R32_FLOAT, BUF_Static);	// always allocate for up to 64 data sets
-		uint32 BufferSize = ShaderLUT.Num() * sizeof(float);
-		int32 *BufferData = static_cast<int32*>(RHILockVertexBuffer(GPUBuffer.Buffer.Buffer, 0, BufferSize, EResourceLockMode::RLM_WriteOnly));
-		FPlatformMemory::Memcpy(BufferData, ShaderLUT.GetData(), BufferSize);
-		RHIUnlockVertexBuffer(GPUBuffer.Buffer.Buffer);
-		GPUBufferDirty = false;
-	}
+bool UNiagaraDataInterfaceCurve::GetFunctionHLSL(const FName& DefinitionFunctionName, FString InstanceFunctionName, FNiagaraDataInterfaceGPUParamInfo& ParamInfo, FString& OutHLSL)
+{
+	FString TimeToLUTFrac = TEXT("TimeToLUTFraction_") + ParamInfo.DataInterfaceHLSLSymbol;
+	FString Sample = TEXT("SampleCurve_") + ParamInfo.DataInterfaceHLSLSymbol;
+	OutHLSL += FString::Printf(TEXT("\
+void %s(in float In_X, out float Out_Value) \n\
+{ \n\
+	float RemappedX = %s(In_X) * %u; \n\
+	float Prev = floor(RemappedX); \n\
+	float Next = Prev < %u ? Prev + 1.0 : Prev; \n\
+	float Interp = RemappedX - Prev; \n\
+	float A = %s(Prev); \n\
+	float B = %s(Next); \n\
+	Out_Value = lerp(A, B, Interp); \n\
+}\n")
+, *InstanceFunctionName, *TimeToLUTFrac, CurveLUTWidthMinusOne, CurveLUTWidthMinusOne, *Sample, *Sample);
 
-	return GPUBuffers;
+	return true;
 }
 
-
-
-DEFINE_NDI_FUNC_BINDER(UNiagaraDataInterfaceCurve, SampleCurve);
-FVMExternalFunction UNiagaraDataInterfaceCurve::GetVMExternalFunction(const FVMExternalFunctionBindingInfo& BindingInfo, void* InstanceData)
+DEFINE_NDI_RAW_FUNC_BINDER(UNiagaraDataInterfaceCurve, SampleCurve);
+void UNiagaraDataInterfaceCurve::GetVMExternalFunction(const FVMExternalFunctionBindingInfo& BindingInfo, void* InstanceData, FVMExternalFunction &OutFunc)
 {
-	if (BindingInfo.Name == TEXT("SampleCurve") && BindingInfo.GetNumInputs() == 1 && BindingInfo.GetNumOutputs() == 1)
+	if (BindingInfo.Name == SampleCurveName && BindingInfo.GetNumInputs() == 1 && BindingInfo.GetNumOutputs() == 1)
 	{
-		return TCurveUseLUTBinder<TNDIParamBinder<0, float, NDI_FUNC_BINDER(UNiagaraDataInterfaceCurve, SampleCurve)>>::Bind(this, BindingInfo, InstanceData);
+		TCurveUseLUTBinder<TNDIParamBinder<0, float, NDI_RAW_FUNC_BINDER(UNiagaraDataInterfaceCurve, SampleCurve)>>::Bind(this, BindingInfo, InstanceData, OutFunc);
 	}
 	else
 	{
 		UE_LOG(LogNiagara, Error, TEXT("Could not find data interface external function.\n\tExpected Name: SampleCurve  Actual Name: %s\n\tExpected Inputs: 1  Actual Inputs: %i\n\tExpected Outputs: 1  Actual Outputs: %i"),
 			*BindingInfo.Name.ToString(), BindingInfo.GetNumInputs(), BindingInfo.GetNumOutputs());
-		return FVMExternalFunction();
 	}
 }
 
 template<>
-float UNiagaraDataInterfaceCurve::SampleCurveInternal<TIntegralConstant<bool, true>>(float X)
+FORCEINLINE_DEBUGGABLE float UNiagaraDataInterfaceCurve::SampleCurveInternal<TIntegralConstant<bool, true>>(float X)
 {
-	float NormalizedX = NormalizeTime(X);
-	int32 AccessIdx = FMath::Clamp<int32>(FMath::TruncToInt(NormalizedX * CurveLUTWidthMinusOne), 0, CurveLUTWidthMinusOne) * CurveLUTNumElems;
-	return ShaderLUT[AccessIdx];
+	float RemappedX = FMath::Clamp(NormalizeTime(X) * CurveLUTWidthMinusOne, 0.0f, (float)CurveLUTWidthMinusOne);
+	float PrevEntry = FMath::TruncToFloat(RemappedX);
+	float NextEntry = PrevEntry < (float)CurveLUTWidthMinusOne ? PrevEntry + 1.0f : PrevEntry;
+	float Interp = RemappedX - PrevEntry;
+
+	int32 AIndex = PrevEntry * CurveLUTNumElems;
+	int32 BIndex = NextEntry * CurveLUTNumElems;
+	float A = ShaderLUT[AIndex];
+	float B = ShaderLUT[BIndex];
+	return FMath::Lerp(A, B, Interp);
 }
 
 template<>
-float UNiagaraDataInterfaceCurve::SampleCurveInternal<TIntegralConstant<bool, false>>(float X)
+FORCEINLINE_DEBUGGABLE float UNiagaraDataInterfaceCurve::SampleCurveInternal<TIntegralConstant<bool, false>>(float X)
 {
 	return Curve.Eval(X);
 }
@@ -240,4 +192,3 @@ void UNiagaraDataInterfaceCurve::SampleCurve(FVectorVMContext& Context)
 		OutSample.Advance();
 	}
 }
-
diff --git a/Engine/Plugins/FX/Niagara/Source/Niagara/Private/NiagaraDataInterfaceSimpleCounter.cpp b/Engine/Plugins/FX/Niagara/Source/Niagara/Private/NiagaraDataInterfaceSimpleCounter.cpp
new file mode 100644
--- /dev/null
+++ b/Engine/Plugins/FX/Niagara/Source/Niagara/Private/NiagaraDataInterfaceSimpleCounter.cpp
@@ -0,0 +1,135 @@
+// Copyright 1998-2016 Epic Games, Inc. All Rights Reserved.
+
+#include "NiagaraDataInterfaceSimpleCounter.h"
+#include "NiagaraTypes.h"
+#include "NiagaraCustomVersion.h"
+#include "Internationalization/Internationalization.h"
+
+#define LOCTEXT_NAMESPACE "NiagaraDataInterfaceSimpleCounter"
+
+
+//////////////////////////////////////////////////////////////////////////
+//Color Curve
+
+
+UNiagaraDataInterfaceSimpleCounter::UNiagaraDataInterfaceSimpleCounter(FObjectInitializer const& ObjectInitializer)
+	: Super(ObjectInitializer)
+{
+}
+
+bool UNiagaraDataInterfaceSimpleCounter::InitPerInstanceData(void* PerInstanceData, FNiagaraSystemInstance* InSystemInstance)
+{
+	CounterInstanceData *PIData = new (PerInstanceData) CounterInstanceData;
+	PIData->Counter = 0;
+	return true;
+}
+
+
+void UNiagaraDataInterfaceSimpleCounter::PostInitProperties()
+{
+	Super::PostInitProperties();
+
+	//Can we register data interfaces as regular types and fold them into the FNiagaraVariable framework for UI and function calls etc?
+	if (HasAnyFlags(RF_ClassDefaultObject))
+	{
+		FNiagaraTypeRegistry::Register(FNiagaraTypeDefinition(GetClass()), true, false, false);
+	}
+}
+
+void UNiagaraDataInterfaceSimpleCounter::PostLoad()
+{
+	Super::PostLoad();
+
+	const int32 NiagaraVer = GetLinkerCustomVersion(FNiagaraCustomVersion::GUID);
+}
+
+#if WITH_EDITOR
+
+void UNiagaraDataInterfaceSimpleCounter::PostEditChangeProperty(struct FPropertyChangedEvent& PropertyChangedEvent)
+{
+	Super::PostEditChangeProperty(PropertyChangedEvent);
+}
+
+#endif
+
+
+bool UNiagaraDataInterfaceSimpleCounter::CopyToInternal(UNiagaraDataInterface* Destination) const
+{
+	if (!Super::CopyToInternal(Destination))
+	{
+		return false;
+	}
+	return true;
+}
+
+bool UNiagaraDataInterfaceSimpleCounter::Equals(const UNiagaraDataInterface* Other) const
+{
+	if (!Super::Equals(Other))
+	{
+		return false;
+	}
+	return true;
+}
+
+
+void UNiagaraDataInterfaceSimpleCounter::GetFunctions(TArray<FNiagaraFunctionSignature>& OutFunctions)
+{
+	FNiagaraFunctionSignature Sig2;
+	Sig2.Name = TEXT("GetNextValue");
+	Sig2.bMemberFunction = true;
+	Sig2.bRequiresContext = false;
+	Sig2.Inputs.Add(FNiagaraVariable(FNiagaraTypeDefinition(GetClass()), TEXT("Counter")));
+	Sig2.Outputs.Add(FNiagaraVariable(FNiagaraTypeDefinition::GetIntDef(), TEXT("Value")));
+	Sig2.SetDescription(LOCTEXT("UNiagaraDataInterfaceSimpleCounter_GetNextValue", "Increment the internal counter. Note that it is possible for this counter to roll over eventually, so make sure that your particles do not live extremely long lifetimes."));
+	OutFunctions.Add(Sig2);
+}
+
+bool UNiagaraDataInterfaceSimpleCounter::GetFunctionHLSL(const FName& DefinitionFunctionName, FString InstanceFunctionName, FNiagaraDataInterfaceGPUParamInfo& ParamInfo, FString& OutHLSL)
+{
+	return false;
+}
+
+void UNiagaraDataInterfaceSimpleCounter::GetParameterDefinitionHLSL(FNiagaraDataInterfaceGPUParamInfo& ParamInfo, FString& OutHLSL)
+{
+}
+
+
+void UNiagaraDataInterfaceSimpleCounter::GetVMExternalFunction(const FVMExternalFunctionBindingInfo& BindingInfo, void* InstanceData, FVMExternalFunction &OutFunc)
+{
+	CounterInstanceData *InstData = (CounterInstanceData *)InstanceData;
+	if (BindingInfo.Name == TEXT("GetNextValue"))
+	{
+		OutFunc = FVMExternalFunction::CreateUObject(this, &UNiagaraDataInterfaceSimpleCounter::GetNextValue);
+	}
+	else
+	{
+		UE_LOG(LogNiagara, Error, TEXT("Could not find data interface external function. %s\n"),
+			*BindingInfo.Name.ToString());
+	}
+}
+
+void UNiagaraDataInterfaceSimpleCounter::GetNextValue(FVectorVMContext& Context)
+{	
+	FUserPtrHandler<CounterInstanceData> InstanceData(Context);
+	FRegisterHandler<int32> OutValue(Context);
+
+	for (int32 i = 0; i < Context.NumInstances; ++i)
+	{
+		*OutValue.GetDest() = FPlatformAtomics::InterlockedIncrement(&InstanceData->Counter);
+		OutValue.Advance();
+	}
+}
+
+
+bool UNiagaraDataInterfaceSimpleCounter::PerInstanceTick(void* PerInstanceData, FNiagaraSystemInstance* InSystemInstance, float DeltaSeconds)
+{
+	return false;
+}
+
+bool UNiagaraDataInterfaceSimpleCounter::PerInstanceTickPostSimulate(void* PerInstanceData, FNiagaraSystemInstance* InSystemInstance, float DeltaSeconds)
+{
+	CounterInstanceData *PIData = static_cast<CounterInstanceData*>(PerInstanceData);
+	return false;
+}
+
+#undef LOCTEXT_NAMESPACE
\ No newline at end of file
