Copying //UE4/Dev-Niagara to //UE4/Dev-Main (Source: //UE4/Dev-Niagara @ 4074996)

diff --git a/Engine/Plugins/FX/Niagara/Source/Niagara/Private/NiagaraDataInterfaceSkeletalMesh.cpp b/Engine/Plugins/FX/Niagara/Source/Niagara/Private/NiagaraDataInterfaceSkeletalMesh.cpp
--- a/Engine/Plugins/FX/Niagara/Source/Niagara/Private/NiagaraDataInterfaceSkeletalMesh.cpp
+++ b/Engine/Plugins/FX/Niagara/Source/Niagara/Private/NiagaraDataInterfaceSkeletalMesh.cpp
@@ -12,10 +12,12 @@
 #include "NiagaraWorldManager.h"
 #include "Async/ParallelFor.h"
 #include "NiagaraStats.h"
+#include "Templates/AlignmentTemplates.h"
 
 #define LOCTEXT_NAMESPACE "NiagaraDataInterfaceSkeletalMesh"
 
 DECLARE_CYCLE_STAT(TEXT("PreSkin"), STAT_NiagaraSkel_PreSkin, STATGROUP_Niagara);
+DECLARE_CYCLE_STAT(TEXT("Skel Mesh Sampling"), STAT_NiagaraSkel_Sample, STATGROUP_Niagara);
 
 //////////////////////////////////////////////////////////////////////////
 
@@ -81,6 +83,11 @@ FSkeletalMeshSkinningDataHandle::~FSkeletalMeshSkinningDataHandle()
 }
 
 //////////////////////////////////////////////////////////////////////////
+void FSkeletalMeshSkinningData::ForceDataRefresh()
+{
+	FScopeLock Lock(&CriticalSection);
+	bForceDataRefresh = true;
+}
 
 void FSkeletalMeshSkinningData::RegisterUser(FSkeletalMeshSkinningDataUsage Usage)
 {
@@ -91,7 +98,7 @@ void FSkeletalMeshSkinningData::RegisterUser(FSkeletalMeshSkinningDataUsage Usag
 	check(LODIndex != INDEX_NONE);
 	check(SkelComp);
 
-	LODData.SetNum(LODIndex + 1);
+	LODData.SetNum(SkelComp->SkeletalMesh->GetLODInfoArray().Num());
 
 	if (Usage.NeedBoneMatrices())
 	{
@@ -181,7 +188,7 @@ bool FSkeletalMeshSkinningData::Tick(float InDeltaSeconds)
 	}
 
 	//Prime the prev matrices if they're missing.
-	if (PrevBoneRefToLocals().Num() != CurrBoneRefToLocals().Num())
+	if (PrevBoneRefToLocals().Num() != CurrBoneRefToLocals().Num() || bForceDataRefresh)
 	{
 		PrevBoneRefToLocals() = CurrBoneRefToLocals();
 	}
@@ -195,7 +202,7 @@ bool FSkeletalMeshSkinningData::Tick(float InDeltaSeconds)
 			FSkeletalMeshLODRenderData& SkelMeshLODData = SkelComp->SkeletalMesh->GetResourceForRendering()->LODRenderData[LODIndex];
 			FSkinWeightVertexBuffer* SkinWeightBuffer = SkelComp->GetSkinWeightBuffer(LODIndex);
 			USkeletalMeshComponent::ComputeSkinnedPositions(SkelComp, CurrSkinnedPositions(LODIndex), CurrBoneRefToLocals(), SkelMeshLODData, *SkinWeightBuffer);
-
+			//check(CurrSkinnedPositions(LODIndex).Num() == SkelMeshLODData.NumVertices);
 			//Prime the previous positions if they're missing
 			if (PrevSkinnedPositions(LODIndex).Num() != CurrSkinnedPositions(LODIndex).Num())
 			{
@@ -204,6 +211,7 @@ bool FSkeletalMeshSkinningData::Tick(float InDeltaSeconds)
 		}
 	}
 	
+	bForceDataRefresh = false;
 	return true;
 }
 
@@ -272,19 +280,9 @@ void FNDI_SkeletalMesh_GeneratedData::TickGeneratedData(float DeltaSeconds)
 //////////////////////////////////////////////////////////////////////////
 //FNDISkeletalMesh_InstanceData
 
-bool FNDISkeletalMesh_InstanceData::Init(UNiagaraDataInterfaceSkeletalMesh* Interface, FNiagaraSystemInstance* SystemInstance)
+USkeletalMesh* UNiagaraDataInterfaceSkeletalMesh::GetSkeletalMeshHelper(UNiagaraDataInterfaceSkeletalMesh* Interface, UNiagaraComponent* OwningComponent, TWeakObjectPtr<USceneComponent>& SceneComponent, USkeletalMeshComponent*& FoundSkelComp)
 {
-	check(SystemInstance);
-	USkeletalMesh* PrevMesh = Mesh;
-	Component = nullptr;
-	Mesh = nullptr;
-	Transform = FMatrix::Identity;
-	TransformInverseTransposed = FMatrix::Identity;
-	PrevTransform = FMatrix::Identity;
-	PrevTransformInverseTransposed = FMatrix::Identity;
-	DeltaSeconds = 0.0f;
-
-	USkeletalMeshComponent* NewSkelComp = nullptr;
+	USkeletalMesh* Mesh = nullptr;
 	if (Interface->Source)
 	{
 		ASkeletalMeshActor* MeshActor = Cast<ASkeletalMeshActor>(Interface->Source);
@@ -301,25 +299,25 @@ bool FNDISkeletalMesh_InstanceData::Init(UNiagaraDataInterfaceSkeletalMesh* Inte
 		if (SourceComp)
 		{
 			Mesh = SourceComp->SkeletalMesh;
-			NewSkelComp = SourceComp;
+			FoundSkelComp = SourceComp;
 		}
 		else
 		{
-			Component = Interface->Source->GetRootComponent();
+			SceneComponent = Interface->Source->GetRootComponent();
 		}
 	}
 	else
 	{
-		if (UNiagaraComponent* SimComp = SystemInstance->GetComponent())
+		if (UNiagaraComponent* SimComp = OwningComponent)
 		{
 			if (USkeletalMeshComponent* ParentComp = Cast<USkeletalMeshComponent>(SimComp->GetAttachParent()))
 			{
-				NewSkelComp = ParentComp;
+				FoundSkelComp = ParentComp;
 				Mesh = ParentComp->SkeletalMesh;
 			}
 			else if (USkeletalMeshComponent* OuterComp = SimComp->GetTypedOuter<USkeletalMeshComponent>())
 			{
-				NewSkelComp = OuterComp;
+				FoundSkelComp = OuterComp;
 				Mesh = OuterComp->SkeletalMesh;
 			}
 			else if (AActor* Owner = SimComp->GetAttachmentRootActor())
@@ -334,32 +332,52 @@ bool FNDISkeletalMesh_InstanceData::Init(UNiagaraDataInterfaceSkeletalMesh* Inte
 						if (PossibleMesh != nullptr/* && PossibleMesh->bAllowCPUAccess*/)
 						{
 							Mesh = PossibleMesh;
-							NewSkelComp = SourceComp;
+							FoundSkelComp = SourceComp;
 							break;
 						}
 					}
 				}
 			}
 
-			if (!Component.IsValid())
+			if (!SceneComponent.IsValid())
 			{
-				Component = SimComp;
+				SceneComponent = SimComp;
 			}
 		}
 	}
 
-	if (NewSkelComp)
+	if (FoundSkelComp)
 	{
-		Component = NewSkelComp;
+		SceneComponent = FoundSkelComp;
 	}
-
-	check(Component.IsValid());
-
+	
 	if (!Mesh && Interface->DefaultMesh)
 	{
 		Mesh = Interface->DefaultMesh;
 	}
 
+	return Mesh;
+}
+
+
+bool FNDISkeletalMesh_InstanceData::Init(UNiagaraDataInterfaceSkeletalMesh* Interface, FNiagaraSystemInstance* SystemInstance)
+{
+	check(SystemInstance);
+	ChangeId = Interface->ChangeId;
+	USkeletalMesh* PrevMesh = Mesh;
+	Component = nullptr;
+	Mesh = nullptr;
+	Transform = FMatrix::Identity;
+	TransformInverseTransposed = FMatrix::Identity;
+	PrevTransform = FMatrix::Identity;
+	PrevTransformInverseTransposed = FMatrix::Identity;
+	DeltaSeconds = 0.0f;
+
+	USkeletalMeshComponent* NewSkelComp = nullptr;
+	Mesh = UNiagaraDataInterfaceSkeletalMesh::GetSkeletalMeshHelper(Interface, SystemInstance->GetComponent(), Component, NewSkelComp);
+	
+	MeshSafe = Mesh;
+
 	if (Component.IsValid() && Mesh)
 	{
 		PrevTransform = Transform;
@@ -370,12 +388,16 @@ bool FNDISkeletalMesh_InstanceData::Init(UNiagaraDataInterfaceSkeletalMesh* Inte
 
 	if (!Mesh)
 	{
-		UE_LOG(LogNiagara, Log, TEXT("SkeletalMesh data interface has no valid mesh. Failed InitPerInstanceData - %s"), *Interface->GetFullName());
+		/*USceneComponent* Comp = Component.Get();
+		UE_LOG(LogNiagara, Log, TEXT("SkeletalMesh data interface has no valid mesh. Failed InitPerInstanceData!\nInterface: %s\nComponent: %s\nActor: %s\n")
+			, *Interface->GetFullName()
+			, Comp ? *Component->GetFullName() : TEXT("Null Component!")
+			, Comp ? *Comp->GetOwner()->GetFullName() : TEXT("NA"));*/
 		return false;
 	}
 
 #if WITH_EDITOR
-	Mesh->GetOnMeshChanged().AddUObject(SystemInstance->GetComponent(), &UNiagaraComponent::ReinitializeSystem);
+	MeshSafe->GetOnMeshChanged().AddUObject(SystemInstance->GetComponent(), &UNiagaraComponent::ReinitializeSystem);
 #endif
 
 
@@ -429,22 +451,30 @@ bool FNDISkeletalMesh_InstanceData::Init(UNiagaraDataInterfaceSkeletalMesh* Inte
 			{
 				const FSkeletalMeshSamplingRegion& Region = SamplingInfo.GetRegion(RegionIdx);
 				const FSkeletalMeshSamplingRegionBuiltData& RegionBuiltData = SamplingInfo.GetRegionBuiltData(RegionIdx);
-				if (LODIndex == INDEX_NONE)
+				int32 RegionLODIndex = Region.LODIndex;
+				if (RegionLODIndex == INDEX_NONE)
 				{
-					LODIndex = Region.LODIndex;
+					RegionLODIndex = Mesh->GetLODInfoArray().Num() - 1;
 				}
 				else
 				{
-					//ensure we don't try to use two regions from different LODs.
-					if (LODIndex != Region.LODIndex)
-					{
-						UE_LOG(LogNiagara, Warning, TEXT("Skeletal Mesh Data Interface is trying to use regions on different LODs of the mesh. This is currently unsupported.\nInterface: %s\nMesh: %s\nRegion: %s"),
-							*Interface->GetFullName(),
-							*Mesh->GetFullName(),
-							*RegionName.ToString());
+					RegionLODIndex = FMath::Clamp(RegionLODIndex, 0, Mesh->GetLODInfoArray().Num() - 1);
+				}
 
-						return false;
-					}
+				if (LODIndex == INDEX_NONE)
+				{
+					LODIndex = RegionLODIndex;
+				}
+
+				//ensure we don't try to use two regions from different LODs.
+				if (LODIndex != RegionLODIndex)
+				{
+					UE_LOG(LogNiagara, Warning, TEXT("Skeletal Mesh Data Interface is trying to use regions on different LODs of the mesh. This is currently unsupported.\nInterface: %s\nMesh: %s\nRegion: %s"),
+						*Interface->GetFullName(),
+						*Mesh->GetFullName(),
+						*RegionName.ToString());
+
+					return false;
 				}
 
 				if (RegionBuiltData.TriangleIndices.Num() > 0)
@@ -474,6 +504,11 @@ bool FNDISkeletalMesh_InstanceData::Init(UNiagaraDataInterfaceSkeletalMesh* Inte
 			}
 		}
 	}
+
+	// TODO: This change is temporary to work around a crash that happens when you change the
+	// source mesh on a system which is running in the level from the details panel.
+	// bool bNeedDataImmediately = SystemInstance->IsSolo();
+	bool bNeedDataImmediately = true;
 		
 	//Grab a handle to the skinning data if we have a component to skin.
 	ENDISkeletalMesh_SkinningMode SkinningMode = Interface->SkinningMode;
@@ -481,7 +516,7 @@ bool FNDISkeletalMesh_InstanceData::Init(UNiagaraDataInterfaceSkeletalMesh* Inte
 		LODIndex,
 		SkinningMode == ENDISkeletalMesh_SkinningMode::SkinOnTheFly || SkinningMode == ENDISkeletalMesh_SkinningMode::PreSkin,
 		SkinningMode == ENDISkeletalMesh_SkinningMode::PreSkin,
-		SystemInstance->IsSolo());
+		bNeedDataImmediately);
 
 	if (NewSkelComp)
 	{
@@ -506,21 +541,38 @@ bool FNDISkeletalMesh_InstanceData::Init(UNiagaraDataInterfaceSkeletalMesh* Inte
 	FSkeletalMeshLODRenderData& LODData = GetLODRenderDataAndSkinWeights(SkinWeightBuffer);
 
 	//Check for the validity of the Mesh's cpu data.
-	bool bMeshCPUDataValid =
-		LODData.GetNumVertices() > 0 &&
-		LODData.StaticVertexBuffers.PositionVertexBuffer.GetNumVertices() &&
-		SkinWeightBuffer &&
-		SkinWeightBuffer->GetNumVertices() > 0 &&
-		LODData.MultiSizeIndexContainer.IsIndexBufferValid() &&
-		LODData.MultiSizeIndexContainer.GetIndexBuffer() &&
-		LODData.MultiSizeIndexContainer.GetIndexBuffer()->Num() > 0;
+
+	bool LODDataNumVerticesCorrect = LODData.GetNumVertices() > 0;
+	bool LODDataPositonNumVerticesCorrect = LODData.StaticVertexBuffers.PositionVertexBuffer.GetNumVertices() > 0;
+	bool bSkinWeightBuffer = SkinWeightBuffer != nullptr;
+	bool SkinWeightBufferNumVerticesCorrect = bSkinWeightBuffer && (SkinWeightBuffer->GetNumVertices() > 0);
+	bool bIndexBufferValid = LODData.MultiSizeIndexContainer.IsIndexBufferValid();
+	bool bIndexBufferFound = bIndexBufferValid && (LODData.MultiSizeIndexContainer.GetIndexBuffer() != nullptr);
+	bool bIndexBufferNumCorrect = bIndexBufferFound && (LODData.MultiSizeIndexContainer.GetIndexBuffer()->Num() > 0);
+
+	bool bMeshCPUDataValid = LODDataNumVerticesCorrect &&
+		LODDataPositonNumVerticesCorrect &&
+		bSkinWeightBuffer &&
+		SkinWeightBufferNumVerticesCorrect && 
+		bIndexBufferValid &&
+		bIndexBufferFound &&
+		bIndexBufferNumCorrect;
 
 	if (!bMeshCPUDataValid)
 	{
-		UE_LOG(LogNiagara, Warning, TEXT("Skeletal Mesh Data Interface is trying to sample from a mesh with missing CPU vertex or index data.\nInterface: %s\nMesh: %s\nLOD: %d"),
+		UE_LOG(LogNiagara, Warning, TEXT("Skeletal Mesh Data Interface is trying to sample from a mesh with missing CPU vertex or index data.\nInterface: %s\nMesh: %s\nLOD: %d\n"
+			"LODDataNumVerticesCorrect: %d  LODDataPositonNumVerticesCorrect : %d  bSkinWeightBuffer : %d  SkinWeightBufferNumVerticesCorrect : %d bIndexBufferValid : %d  bIndexBufferFound : %d  bIndexBufferNumCorrect : %d"),
 			*Interface->GetFullName(),
 			*Mesh->GetFullName(),
-			LODIndex);
+			LODIndex,
+			LODDataNumVerticesCorrect ? 1 : 0,
+			LODDataPositonNumVerticesCorrect ? 1 : 0,
+			bSkinWeightBuffer ? 1 : 0,
+			SkinWeightBufferNumVerticesCorrect ? 1 : 0,
+			bIndexBufferValid ? 1 : 0,
+			bIndexBufferFound ? 1 : 0,
+			bIndexBufferNumCorrect ? 1 : 0
+			);
 
 		return false;
 	}
@@ -536,13 +588,24 @@ bool FNDISkeletalMesh_InstanceData::ResetRequired(UNiagaraDataInterfaceSkeletalM
 		//The component we were bound to is no longer valid so we have to trigger a reset.
 		return true;
 	}
-	
+		
 	if (USkeletalMeshComponent* SkelComp = Cast<USkeletalMeshComponent>(Comp))
 	{
 		if (!SkelComp->SkeletalMesh)
 		{
 			return true;
 		}
+		
+		// Handle the case where they've procedurally swapped out the skeletal mesh from
+		// the one we previously cached data for.
+		if (SkelComp->SkeletalMesh != Mesh && Mesh != nullptr && SkelComp->SkeletalMesh != nullptr)
+		{
+			if (SkinningData.SkinningData.IsValid())
+			{
+				SkinningData.SkinningData.Get()->ForceDataRefresh();
+			}
+			return true;
+		}
 	}
 	else
 	{
@@ -551,6 +614,12 @@ bool FNDISkeletalMesh_InstanceData::ResetRequired(UNiagaraDataInterfaceSkeletalM
 			return true;
 		}
 	}
+
+	if (Interface->ChangeId != ChangeId)
+	{
+		return true;
+	}
+
 	
 	return false;
 }
@@ -608,6 +677,7 @@ struct FSkeletalMeshAccessorHelper
 
 	USkeletalMeshComponent* Comp;
 	USkeletalMesh* Mesh;
+	TWeakObjectPtr<USkeletalMesh> MeshSafe;
 	FSkeletalMeshLODRenderData* LODData;
 	FSkinWeightVertexBuffer* SkinWeightBuffer;
 	FRawStaticIndexBuffer16or32Interface* IndexBuffer;
@@ -620,6 +690,7 @@ struct FSkeletalMeshAccessorHelper
 template<typename FilterMode, typename AreaWeightingMode>
 void FSkeletalMeshAccessorHelper::Init(FNDISkeletalMesh_InstanceData* InstData)
 {
+	Comp = Cast<USkeletalMeshComponent>(InstData->Component.Get());
 	Mesh = InstData->Mesh;
 	LODData = &InstData->GetLODRenderDataAndSkinWeights(SkinWeightBuffer);
 	IndexBuffer = LODData->MultiSizeIndexContainer.GetIndexBuffer();
@@ -633,6 +704,7 @@ void FSkeletalMeshAccessorHelper::Init<
 	TIntegralConstant<ENDISkelMesh_AreaWeightingMode, ENDISkelMesh_AreaWeightingMode::None>>
 	(FNDISkeletalMesh_InstanceData* InstData)
 {
+	Comp = Cast<USkeletalMeshComponent>(InstData->Component.Get());
 	Mesh = InstData->Mesh;
 	LODData = &InstData->GetLODRenderDataAndSkinWeights(SkinWeightBuffer);
 	IndexBuffer = LODData->MultiSizeIndexContainer.GetIndexBuffer();
@@ -650,6 +722,7 @@ void FSkeletalMeshAccessorHelper::Init<
 	TIntegralConstant<ENDISkelMesh_AreaWeightingMode, ENDISkelMesh_AreaWeightingMode::AreaWeighted>>
 	(FNDISkeletalMesh_InstanceData* InstData)
 {
+	Comp = Cast<USkeletalMeshComponent>(InstData->Component.Get());
 	Mesh = InstData->Mesh;
 	LODData = &InstData->GetLODRenderDataAndSkinWeights(SkinWeightBuffer);
 	IndexBuffer = LODData->MultiSizeIndexContainer.GetIndexBuffer();
@@ -667,13 +740,13 @@ template<typename SkinningMode>
 struct FSkinnedPositionAccessorHelper
 {
 	FORCEINLINE void GetSkinnedTrianglePositions(FSkeletalMeshAccessorHelper& Accessor,
-		int32 Tri, FVector& OutPos0, FVector& OutPos1, FVector& OutPos2, FVector& OutPrev0, FVector& OutPrev1, FVector& OutPrev2)
+		int32 Tri, FVector& OutPos0, FVector& OutPos1, FVector& OutPos2, FVector& OutPrev0, FVector& OutPrev1, FVector& OutPrev2, int32& Idx0, int32& Idx1, int32& Idx2)
 	{
 		checkf(false, TEXT("Must provide a specialization for this template type"));
 	}
 
 	FORCEINLINE void GetSkinnedTrianglePositions(FSkeletalMeshAccessorHelper& Accessor,
-		int32 Tri, FVector& OutPos0, FVector& OutPos1, FVector& OutPos2)
+		int32 Tri, FVector& OutPos0, FVector& OutPos1, FVector& OutPos2, int32& Idx0, int32& Idx1, int32& Idx2)
 	{
 		checkf(false, TEXT("Must provide a specialization for this template type"));
 	}
@@ -683,11 +756,12 @@ template<>
 struct FSkinnedPositionAccessorHelper<TIntegralConstant<ENDISkeletalMesh_SkinningMode, ENDISkeletalMesh_SkinningMode::None>>
 {
 	FORCEINLINE void GetSkinnedTrianglePositions(FSkeletalMeshAccessorHelper& Accessor,
-		int32 Tri, FVector& OutPos0, FVector& OutPos1, FVector& OutPos2)
+		int32 Tri, FVector& OutPos0, FVector& OutPos1, FVector& OutPos2, int32& Idx0, int32& Idx1, int32& Idx2)
 	{
-		int32 Idx0 = Accessor.IndexBuffer->Get(Tri);
-		int32 Idx1 = Accessor.IndexBuffer->Get(Tri + 1);
-		int32 Idx2 = Accessor.IndexBuffer->Get(Tri + 2);
+		checkSlow(Tri + 2 < Accessor.IndexBuffer->Num());
+		Idx0 = Accessor.IndexBuffer->Get(Tri);
+		Idx1 = Accessor.IndexBuffer->Get(Tri + 1);
+		Idx2 = Accessor.IndexBuffer->Get(Tri + 2);
 		
 		OutPos0 = GetSkeletalMeshRefVertLocation(Accessor.Mesh, *Accessor.LODData, *Accessor.SkinWeightBuffer, Idx0);
 		OutPos1 = GetSkeletalMeshRefVertLocation(Accessor.Mesh, *Accessor.LODData, *Accessor.SkinWeightBuffer, Idx1);
@@ -695,11 +769,12 @@ struct FSkinnedPositionAccessorHelper<TIntegralConstant<ENDISkeletalMesh_Skinnin
 	}
 
 	FORCEINLINE void GetSkinnedTrianglePositions(FSkeletalMeshAccessorHelper& Accessor,
-		int32 Tri, FVector& OutPos0, FVector& OutPos1, FVector& OutPos2, FVector& OutPrev0, FVector& OutPrev1, FVector& OutPrev2)
+		int32 Tri, FVector& OutPos0, FVector& OutPos1, FVector& OutPos2, FVector& OutPrev0, FVector& OutPrev1, FVector& OutPrev2, int32& Idx0, int32& Idx1, int32& Idx2)
 	{
-		int32 Idx0 = Accessor.IndexBuffer->Get(Tri);
-		int32 Idx1 = Accessor.IndexBuffer->Get(Tri + 1);
-		int32 Idx2 = Accessor.IndexBuffer->Get(Tri + 2);
+		checkSlow(Tri + 2 < Accessor.IndexBuffer->Num());
+		Idx0 = Accessor.IndexBuffer->Get(Tri);
+		Idx1 = Accessor.IndexBuffer->Get(Tri + 1);
+		Idx2 = Accessor.IndexBuffer->Get(Tri + 2);
 		OutPos0 = GetSkeletalMeshRefVertLocation(Accessor.Mesh, *Accessor.LODData, *Accessor.SkinWeightBuffer, Idx0);
 		OutPos1 = GetSkeletalMeshRefVertLocation(Accessor.Mesh, *Accessor.LODData, *Accessor.SkinWeightBuffer, Idx1);
 		OutPos2 = GetSkeletalMeshRefVertLocation(Accessor.Mesh, *Accessor.LODData, *Accessor.SkinWeightBuffer, Idx2);
@@ -714,22 +789,24 @@ template<>
 struct FSkinnedPositionAccessorHelper<TIntegralConstant<ENDISkeletalMesh_SkinningMode, ENDISkeletalMesh_SkinningMode::SkinOnTheFly>>
 {
 	FORCEINLINE void GetSkinnedTrianglePositions(FSkeletalMeshAccessorHelper& Accessor,
-		int32 Tri, FVector& OutPos0, FVector& OutPos1, FVector& OutPos2)
+		int32 Tri, FVector& OutPos0, FVector& OutPos1, FVector& OutPos2, int32& Idx0, int32& Idx1, int32& Idx2)
 	{
-		int32 Idx0 = Accessor.IndexBuffer->Get(Tri);
-		int32 Idx1 = Accessor.IndexBuffer->Get(Tri + 1);
-		int32 Idx2 = Accessor.IndexBuffer->Get(Tri + 2);
+		checkSlow(Tri + 2 < Accessor.IndexBuffer->Num());
+		Idx0 = Accessor.IndexBuffer->Get(Tri);
+		Idx1 = Accessor.IndexBuffer->Get(Tri + 1);
+		Idx2 = Accessor.IndexBuffer->Get(Tri + 2);
 		OutPos0 = USkeletalMeshComponent::GetSkinnedVertexPosition(Accessor.Comp, Idx0, *Accessor.LODData, *Accessor.SkinWeightBuffer, Accessor.SkinningData->CurrBoneRefToLocals());
 		OutPos1 = USkeletalMeshComponent::GetSkinnedVertexPosition(Accessor.Comp, Idx1, *Accessor.LODData, *Accessor.SkinWeightBuffer, Accessor.SkinningData->CurrBoneRefToLocals());
 		OutPos2 = USkeletalMeshComponent::GetSkinnedVertexPosition(Accessor.Comp, Idx2, *Accessor.LODData, *Accessor.SkinWeightBuffer, Accessor.SkinningData->CurrBoneRefToLocals());
 	}
 
 	FORCEINLINE void GetSkinnedTrianglePositions(FSkeletalMeshAccessorHelper& Accessor,
-		int32 Tri, FVector& OutPos0, FVector& OutPos1, FVector& OutPos2, FVector& OutPrev0, FVector& OutPrev1, FVector& OutPrev2)
+		int32 Tri, FVector& OutPos0, FVector& OutPos1, FVector& OutPos2, FVector& OutPrev0, FVector& OutPrev1, FVector& OutPrev2, int32& Idx0, int32& Idx1, int32& Idx2)
 	{
-		int32 Idx0 = Accessor.IndexBuffer->Get(Tri);
-		int32 Idx1 = Accessor.IndexBuffer->Get(Tri + 1);
-		int32 Idx2 = Accessor.IndexBuffer->Get(Tri + 2);
+		checkSlow(Tri + 2 < Accessor.IndexBuffer->Num());
+		Idx0 = Accessor.IndexBuffer->Get(Tri);
+		Idx1 = Accessor.IndexBuffer->Get(Tri + 1);
+		Idx2 = Accessor.IndexBuffer->Get(Tri + 2);
 		OutPos0 = USkeletalMeshComponent::GetSkinnedVertexPosition(Accessor.Comp, Idx0, *Accessor.LODData, *Accessor.SkinWeightBuffer, Accessor.SkinningData->CurrBoneRefToLocals());
 		OutPos1 = USkeletalMeshComponent::GetSkinnedVertexPosition(Accessor.Comp, Idx1, *Accessor.LODData, *Accessor.SkinWeightBuffer, Accessor.SkinningData->CurrBoneRefToLocals());
 		OutPos2 = USkeletalMeshComponent::GetSkinnedVertexPosition(Accessor.Comp, Idx2, *Accessor.LODData, *Accessor.SkinWeightBuffer, Accessor.SkinningData->CurrBoneRefToLocals());
@@ -743,22 +820,24 @@ template<>
 struct FSkinnedPositionAccessorHelper<TIntegralConstant<ENDISkeletalMesh_SkinningMode, ENDISkeletalMesh_SkinningMode::PreSkin>>
 {
 	FORCEINLINE void GetSkinnedTrianglePositions(FSkeletalMeshAccessorHelper& Accessor,
-		int32 Tri, FVector& OutPos0, FVector& OutPos1, FVector& OutPos2)
+		int32 Tri, FVector& OutPos0, FVector& OutPos1, FVector& OutPos2, int32& Idx0, int32& Idx1, int32& Idx2)
 	{
-		int32 Idx0 = Accessor.IndexBuffer->Get(Tri);
-		int32 Idx1 = Accessor.IndexBuffer->Get(Tri + 1);
-		int32 Idx2 = Accessor.IndexBuffer->Get(Tri + 2);
+		checkSlow(Tri + 2 < Accessor.IndexBuffer->Num());
+		Idx0 = Accessor.IndexBuffer->Get(Tri);
+		Idx1 = Accessor.IndexBuffer->Get(Tri + 1);
+		Idx2 = Accessor.IndexBuffer->Get(Tri + 2);
 		OutPos0 = Accessor.SkinningData->GetPosition(Accessor.Usage.GetLODIndex(), Idx0);
 		OutPos1 = Accessor.SkinningData->GetPosition(Accessor.Usage.GetLODIndex(), Idx1);
 		OutPos2 = Accessor.SkinningData->GetPosition(Accessor.Usage.GetLODIndex(), Idx2);
 	}
 
 	FORCEINLINE void GetSkinnedTrianglePositions(FSkeletalMeshAccessorHelper& Accessor,
-		int32 Tri, FVector& OutPos0, FVector& OutPos1, FVector& OutPos2, FVector& OutPrev0, FVector& OutPrev1, FVector& OutPrev2)
+		int32 Tri, FVector& OutPos0, FVector& OutPos1, FVector& OutPos2, FVector& OutPrev0, FVector& OutPrev1, FVector& OutPrev2, int32& Idx0, int32& Idx1, int32& Idx2)
 	{
-		int32 Idx0 = Accessor.IndexBuffer->Get(Tri);
-		int32 Idx1 = Accessor.IndexBuffer->Get(Tri + 1);
-		int32 Idx2 = Accessor.IndexBuffer->Get(Tri + 2);
+		checkSlow(Tri + 2 < Accessor.IndexBuffer->Num());
+		Idx0 = Accessor.IndexBuffer->Get(Tri);
+		Idx1 = Accessor.IndexBuffer->Get(Tri + 1);
+		Idx2 = Accessor.IndexBuffer->Get(Tri + 2);
 		OutPos0 = Accessor.SkinningData->GetPosition(Accessor.Usage.GetLODIndex(), Idx0);
 		OutPos1 = Accessor.SkinningData->GetPosition(Accessor.Usage.GetLODIndex(), Idx1);
 		OutPos2 = Accessor.SkinningData->GetPosition(Accessor.Usage.GetLODIndex(), Idx2);
@@ -799,7 +878,7 @@ template<typename NextBinder>
 struct TAreaWeightingModeBinder
 {
 	template<typename... ParamTypes>
-	static FVMExternalFunction Bind(UNiagaraDataInterface* Interface, const FVMExternalFunctionBindingInfo& BindingInfo, void* InstanceData)
+	static void Bind(UNiagaraDataInterface* Interface, const FVMExternalFunctionBindingInfo& BindingInfo, void* InstanceData, FVMExternalFunction &OutFunc)
 	{
 		FNDISkeletalMesh_InstanceData* InstData = (FNDISkeletalMesh_InstanceData*)InstanceData;
 		check(InstData);
@@ -825,10 +904,12 @@ struct TAreaWeightingModeBinder
 
 		if (bAreaWeighting)
 		{
-			return NextBinder::template Bind<ParamTypes..., TIntegralConstant<ENDISkelMesh_AreaWeightingMode, ENDISkelMesh_AreaWeightingMode::AreaWeighted>>(Interface, BindingInfo, InstanceData);
+			NextBinder::template Bind<ParamTypes..., TIntegralConstant<ENDISkelMesh_AreaWeightingMode, ENDISkelMesh_AreaWeightingMode::AreaWeighted>>(Interface, BindingInfo, InstanceData, OutFunc);
+		}
+		else
+		{
+			NextBinder::template Bind<ParamTypes..., TIntegralConstant<ENDISkelMesh_AreaWeightingMode, ENDISkelMesh_AreaWeightingMode::None>>(Interface, BindingInfo, InstanceData, OutFunc);
 		}
-			
-		return NextBinder::template Bind<ParamTypes..., TIntegralConstant<ENDISkelMesh_AreaWeightingMode, ENDISkelMesh_AreaWeightingMode::None>>(Interface, BindingInfo, InstanceData);
 	}
 };
 
@@ -837,7 +918,7 @@ template<typename NextBinder>
 struct TFilterModeBinder
 {
 	template<typename... ParamTypes>
-	static FVMExternalFunction Bind(UNiagaraDataInterface* Interface, const FVMExternalFunctionBindingInfo& BindingInfo, void* InstanceData)
+	static void Bind(UNiagaraDataInterface* Interface, const FVMExternalFunctionBindingInfo& BindingInfo, void* InstanceData, FVMExternalFunction &OutFunc)
 	{
 		FNDISkeletalMesh_InstanceData* InstData = (FNDISkeletalMesh_InstanceData*)InstanceData;
 		check(InstData);
@@ -845,14 +926,16 @@ struct TFilterModeBinder
 
 		if (InstData->SamplingRegionIndices.Num() == 1)
 		{
-			return NextBinder::template Bind<ParamTypes..., TIntegralConstant<ENDISkeletalMesh_FilterMode, ENDISkeletalMesh_FilterMode::SingleRegion>>(Interface, BindingInfo, InstanceData);
+			NextBinder::template Bind<ParamTypes..., TIntegralConstant<ENDISkeletalMesh_FilterMode, ENDISkeletalMesh_FilterMode::SingleRegion>>(Interface, BindingInfo, InstanceData, OutFunc);
 		}
 		else if (InstData->SamplingRegionIndices.Num() > 1)
 		{
-			return NextBinder::template Bind<ParamTypes..., TIntegralConstant<ENDISkeletalMesh_FilterMode, ENDISkeletalMesh_FilterMode::MultiRegion>>(Interface, BindingInfo, InstanceData);
+			NextBinder::template Bind<ParamTypes..., TIntegralConstant<ENDISkeletalMesh_FilterMode, ENDISkeletalMesh_FilterMode::MultiRegion>>(Interface, BindingInfo, InstanceData, OutFunc);
+		}
+		else
+		{
+			NextBinder::template Bind<ParamTypes..., TIntegralConstant<ENDISkeletalMesh_FilterMode, ENDISkeletalMesh_FilterMode::None>>(Interface, BindingInfo, InstanceData, OutFunc);
 		}
-
-		return NextBinder::template Bind<ParamTypes..., TIntegralConstant<ENDISkeletalMesh_FilterMode, ENDISkeletalMesh_FilterMode::None>>(Interface, BindingInfo, InstanceData);
 	}
 };
 
@@ -861,7 +944,7 @@ template<typename NextBinder>
 struct TVertexAccessorBinder
 {
 	template<typename... ParamTypes>
-	static FVMExternalFunction Bind(UNiagaraDataInterface* Interface, const FVMExternalFunctionBindingInfo& BindingInfo, void* InstanceData)
+	static void Bind(UNiagaraDataInterface* Interface, const FVMExternalFunctionBindingInfo& BindingInfo, void* InstanceData, FVMExternalFunction &OutFunc)
 	{
 		FNDISkeletalMesh_InstanceData* InstData = (FNDISkeletalMesh_InstanceData*)InstanceData;
 		UNiagaraDataInterfaceSkeletalMesh* MeshInterface = CastChecked<UNiagaraDataInterfaceSkeletalMesh>(Interface);
@@ -871,11 +954,11 @@ struct TVertexAccessorBinder
 
 		if (LODData.StaticVertexBuffers.StaticMeshVertexBuffer.GetUseFullPrecisionUVs())
 		{
-			return NextBinder::template Bind<ParamTypes..., FSkelMeshVertexAccessor<true>>(Interface, BindingInfo, InstanceData);
+			NextBinder::template Bind<ParamTypes..., FSkelMeshVertexAccessor<true>>(Interface, BindingInfo, InstanceData, OutFunc);
 		}
 		else
 		{
-			return NextBinder::template Bind<ParamTypes..., FSkelMeshVertexAccessor<false>>(Interface, BindingInfo, InstanceData);
+			NextBinder::template Bind<ParamTypes..., FSkelMeshVertexAccessor<false>>(Interface, BindingInfo, InstanceData, OutFunc);
 		}
 	}
 };
@@ -885,37 +968,42 @@ template<typename NextBinder>
 struct TSkinningModeBinder
 {
 	template<typename... ParamTypes>
-	static FVMExternalFunction Bind(UNiagaraDataInterface* Interface, const FVMExternalFunctionBindingInfo& BindingInfo, void* InstanceData)
+	static void Bind(UNiagaraDataInterface* Interface, const FVMExternalFunctionBindingInfo& BindingInfo, void* InstanceData, FVMExternalFunction &OutFunc)
 	{
 		FNDISkeletalMesh_InstanceData* InstData = (FNDISkeletalMesh_InstanceData*)InstanceData;
 		UNiagaraDataInterfaceSkeletalMesh* MeshInterface = CastChecked<UNiagaraDataInterfaceSkeletalMesh>(Interface);
 		USkeletalMeshComponent* Component = Cast<USkeletalMeshComponent>(InstData->Component.Get());
 		if (MeshInterface->SkinningMode == ENDISkeletalMesh_SkinningMode::None || !Component)//Can't skin if we have no component.
 		{
-			return NextBinder::template Bind<ParamTypes..., FSkinnedPositionAccessorHelper<TIntegralConstant<ENDISkeletalMesh_SkinningMode, ENDISkeletalMesh_SkinningMode::None>>>(Interface, BindingInfo, InstanceData);
+			NextBinder::template Bind<ParamTypes..., FSkinnedPositionAccessorHelper<TIntegralConstant<ENDISkeletalMesh_SkinningMode, ENDISkeletalMesh_SkinningMode::None>>>(Interface, BindingInfo, InstanceData, OutFunc);
 		}
 		else if (MeshInterface->SkinningMode == ENDISkeletalMesh_SkinningMode::SkinOnTheFly)
 		{
 			check(Component);
-			return NextBinder::template Bind<ParamTypes..., FSkinnedPositionAccessorHelper<TIntegralConstant<ENDISkeletalMesh_SkinningMode, ENDISkeletalMesh_SkinningMode::SkinOnTheFly>>>(Interface, BindingInfo, InstanceData);
+			NextBinder::template Bind<ParamTypes..., FSkinnedPositionAccessorHelper<TIntegralConstant<ENDISkeletalMesh_SkinningMode, ENDISkeletalMesh_SkinningMode::SkinOnTheFly>>>(Interface, BindingInfo, InstanceData, OutFunc);
 		}
 		else if (MeshInterface->SkinningMode == ENDISkeletalMesh_SkinningMode::PreSkin)
 		{
 			check(Component);
-			return NextBinder::template Bind<ParamTypes..., FSkinnedPositionAccessorHelper<TIntegralConstant<ENDISkeletalMesh_SkinningMode, ENDISkeletalMesh_SkinningMode::PreSkin>>>(Interface, BindingInfo, InstanceData);
+			NextBinder::template Bind<ParamTypes..., FSkinnedPositionAccessorHelper<TIntegralConstant<ENDISkeletalMesh_SkinningMode, ENDISkeletalMesh_SkinningMode::PreSkin>>>(Interface, BindingInfo, InstanceData, OutFunc);
+		}
+		else
+		{
+			checkf(false, TEXT("Invalid skinning mode in %s"), *Interface->GetPathName());
 		}
-
-		checkf(false, TEXT("Invalid skinning mode in %s"), *Interface->GetPathName());
-		return FVMExternalFunction();
 	}
 };
 
 //Final binders for all static mesh interface functions.
-DEFINE_NDI_FUNC_BINDER(UNiagaraDataInterfaceSkeletalMesh, RandomTriCoord);
-DEFINE_NDI_FUNC_BINDER(UNiagaraDataInterfaceSkeletalMesh, GetTriCoordPosition);
-DEFINE_NDI_FUNC_BINDER(UNiagaraDataInterfaceSkeletalMesh, GetTriCoordPositionVelocityAndNormal);
-DEFINE_NDI_FUNC_BINDER(UNiagaraDataInterfaceSkeletalMesh, GetTriCoordColor);
-DEFINE_NDI_FUNC_BINDER(UNiagaraDataInterfaceSkeletalMesh, GetTriCoordUV);
+DEFINE_NDI_RAW_FUNC_BINDER(UNiagaraDataInterfaceSkeletalMesh, RandomTriCoord);
+DEFINE_NDI_RAW_FUNC_BINDER(UNiagaraDataInterfaceSkeletalMesh, GetTriCoordPosition);
+DEFINE_NDI_RAW_FUNC_BINDER(UNiagaraDataInterfaceSkeletalMesh, GetTriCoordPositionVelocityAndNormal);
+DEFINE_NDI_RAW_FUNC_BINDER(UNiagaraDataInterfaceSkeletalMesh, GetTriCoordPositionVelocityAndNormalBinormalTangent);
+DEFINE_NDI_RAW_FUNC_BINDER(UNiagaraDataInterfaceSkeletalMesh, GetTriCoordColor);
+DEFINE_NDI_RAW_FUNC_BINDER(UNiagaraDataInterfaceSkeletalMesh, GetTriCoordUV);
+DEFINE_NDI_RAW_FUNC_BINDER(UNiagaraDataInterfaceSkeletalMesh, IsValidTriCoord);
+DEFINE_NDI_RAW_FUNC_BINDER(UNiagaraDataInterfaceSkeletalMesh, GetFilteredTriangleCount);
+DEFINE_NDI_RAW_FUNC_BINDER(UNiagaraDataInterfaceSkeletalMesh, GetFilteredTriangleAt);
 
 //Function Binders END
 //////////////////////////////////////////////////////////////////////////
@@ -929,6 +1017,7 @@ UNiagaraDataInterfaceSkeletalMesh::UNiagaraDataInterfaceSkeletalMesh(FObjectInit
 	, Source(nullptr)
 	, SkinningMode(ENDISkeletalMesh_SkinningMode::PreSkin)
 	, WholeMeshLOD(INDEX_NONE)
+	, ChangeId(0)
 {
 
 }
@@ -948,6 +1037,11 @@ void UNiagaraDataInterfaceSkeletalMesh::PostInitProperties()
 		FNiagaraTypeRegistry::Register(FMeshTriCoordinate::StaticStruct(), true, true, false);
 	}
 }
+void UNiagaraDataInterfaceSkeletalMesh::PostEditChangeProperty(FPropertyChangedEvent& PropertyChangedEvent)
+{
+	Super::PostEditChangeProperty(PropertyChangedEvent);
+	ChangeId++;
+}
 
 #endif //WITH_EDITOR
 
@@ -958,13 +1052,20 @@ static const FName GetTriPositionWSName("GetTriPositionWS");
 
 static const FName GetTriNormalName("GetTriNormal");
 static const FName GetTriNormalWSName("GetTriNormalWS");
+static const FName IsValidTriCoordName("IsValidTriCoord");
 
 static const FName GetTriColorName("GetTriColor");
 static const FName GetTriUVName("GetTriUV");
+static const FName GetTriangleCountName("GetFilteredTriangleCount");
+static const FName GetTriangleAtName("GetFilteredTriangle");
 
 static const FName GetTriPositionVelocityAndNormalName("GetTriPositionVelocityAndNormal");
 static const FName GetTriPositionVelocityAndNormalWSName("GetTriPositionVelocityAndNormalWS");
 
+static const FName GetTriPositionVelocityAndNormalBinormalTangentName("GetTriPositionVelocityAndNormalBinormalTangent");
+static const FName GetTriPositionVelocityAndNormalBinormalTangentWSName("GetTriPositionVelocityAndNormalBinormalTangentWS");
+
+
 void UNiagaraDataInterfaceSkeletalMesh::GetFunctions(TArray<FNiagaraFunctionSignature>& OutFunctions)
 {
 	{
@@ -976,6 +1077,20 @@ void UNiagaraDataInterfaceSkeletalMesh::GetFunctions(TArray<FNiagaraFunctionSign
 		Sig.bRequiresContext = false;
 		OutFunctions.Add(Sig);
 	}
+
+	{
+		FNiagaraFunctionSignature Sig;
+		Sig.Name = IsValidTriCoordName;
+		Sig.Inputs.Add(FNiagaraVariable(FNiagaraTypeDefinition(GetClass()), TEXT("SkeletalMesh")));
+		Sig.Inputs.Add(FNiagaraVariable(FNiagaraTypeDefinition(FMeshTriCoordinate::StaticStruct()), TEXT("Coord")));
+		Sig.Outputs.Add(FNiagaraVariable(FNiagaraTypeDefinition::GetBoolDef(), TEXT("IsValid")));
+		Sig.bMemberFunction = true;
+		Sig.bRequiresContext = false;
+#if WITH_EDITORONLY_DATA
+		Sig.Description = LOCTEXT("IsValidDesc","Determine if this tri coordinate's triangle index is valid for this mesh. Note that this only checks the mesh index buffer size and does not include any filtering settings." );
+#endif
+		OutFunctions.Add(Sig);
+	}
 	
 	{
 		FNiagaraFunctionSignature Sig;
@@ -1024,7 +1139,41 @@ void UNiagaraDataInterfaceSkeletalMesh::GetFunctions(TArray<FNiagaraFunctionSign
 		Sig.bRequiresContext = false;
 		OutFunctions.Add(Sig);
 	}
-	
+
+	{
+		FNiagaraFunctionSignature Sig;
+		Sig.Name = GetTriPositionVelocityAndNormalBinormalTangentName;
+		Sig.Inputs.Add(FNiagaraVariable(FNiagaraTypeDefinition(GetClass()), TEXT("SkeletalMesh")));
+		Sig.Inputs.Add(FNiagaraVariable(FNiagaraTypeDefinition(FMeshTriCoordinate::StaticStruct()), TEXT("Coord")));
+		Sig.Inputs.Add(FNiagaraVariable(FNiagaraTypeDefinition::GetIntDef(), TEXT("UV Set")));
+		Sig.Outputs.Add(FNiagaraVariable(FNiagaraTypeDefinition::GetVec3Def(), TEXT("Position")));
+		Sig.Outputs.Add(FNiagaraVariable(FNiagaraTypeDefinition::GetVec3Def(), TEXT("Velocity")));
+		Sig.Outputs.Add(FNiagaraVariable(FNiagaraTypeDefinition::GetVec3Def(), TEXT("Normal")));
+		Sig.Outputs.Add(FNiagaraVariable(FNiagaraTypeDefinition::GetVec3Def(), TEXT("Binormal")));
+		Sig.Outputs.Add(FNiagaraVariable(FNiagaraTypeDefinition::GetVec3Def(), TEXT("Tangent")));
+		Sig.Outputs.Add(FNiagaraVariable(FNiagaraTypeDefinition::GetVec2Def(), TEXT("UV")));
+		Sig.bMemberFunction = true;
+		Sig.bRequiresContext = false;
+		OutFunctions.Add(Sig);
+	}
+
+	{
+		FNiagaraFunctionSignature Sig;
+		Sig.Name = GetTriPositionVelocityAndNormalBinormalTangentWSName;
+		Sig.Inputs.Add(FNiagaraVariable(FNiagaraTypeDefinition(GetClass()), TEXT("SkeletalMesh")));
+		Sig.Inputs.Add(FNiagaraVariable(FNiagaraTypeDefinition(FMeshTriCoordinate::StaticStruct()), TEXT("Coord")));	
+		Sig.Inputs.Add(FNiagaraVariable(FNiagaraTypeDefinition::GetIntDef(), TEXT("UV Set")));
+		Sig.Outputs.Add(FNiagaraVariable(FNiagaraTypeDefinition::GetVec3Def(), TEXT("Position")));
+		Sig.Outputs.Add(FNiagaraVariable(FNiagaraTypeDefinition::GetVec3Def(), TEXT("Velocity")));
+		Sig.Outputs.Add(FNiagaraVariable(FNiagaraTypeDefinition::GetVec3Def(), TEXT("Normal")));
+		Sig.Outputs.Add(FNiagaraVariable(FNiagaraTypeDefinition::GetVec3Def(), TEXT("Binormal")));
+		Sig.Outputs.Add(FNiagaraVariable(FNiagaraTypeDefinition::GetVec3Def(), TEXT("Tangent")));
+		Sig.Outputs.Add(FNiagaraVariable(FNiagaraTypeDefinition::GetVec2Def(), TEXT("UV")));
+		Sig.bMemberFunction = true;
+		Sig.bRequiresContext = false;
+		OutFunctions.Add(Sig);
+	}
+		
 	{
 		FNiagaraFunctionSignature Sig;
 		Sig.Name = GetTriColorName;
@@ -1048,52 +1197,102 @@ void UNiagaraDataInterfaceSkeletalMesh::GetFunctions(TArray<FNiagaraFunctionSign
 		Sig.bRequiresContext = false;
 		OutFunctions.Add(Sig);
 	}
+
+	{
+		FNiagaraFunctionSignature Sig;
+		Sig.Name = GetTriangleCountName;
+		Sig.Inputs.Add(FNiagaraVariable(FNiagaraTypeDefinition(GetClass()), TEXT("SkeletalMesh")));
+		Sig.Outputs.Add(FNiagaraVariable(FNiagaraTypeDefinition::GetIntDef(), TEXT("Count")));
+		Sig.bMemberFunction = true;
+		Sig.bRequiresContext = false;
+		OutFunctions.Add(Sig);
+	}
+
+	{
+		FNiagaraFunctionSignature Sig;
+		Sig.Name = GetTriangleAtName;
+		Sig.Inputs.Add(FNiagaraVariable(FNiagaraTypeDefinition(GetClass()), TEXT("SkeletalMesh")));
+		Sig.Inputs.Add(FNiagaraVariable(FNiagaraTypeDefinition::GetIntDef(), TEXT("Index")));
+		Sig.Outputs.Add(FNiagaraVariable(FNiagaraTypeDefinition(FMeshTriCoordinate::StaticStruct()), TEXT("Coord")));
+		Sig.bMemberFunction = true;
+		Sig.bRequiresContext = false;
+		OutFunctions.Add(Sig);
+	}
 }
 
-FVMExternalFunction UNiagaraDataInterfaceSkeletalMesh::GetVMExternalFunction(const FVMExternalFunctionBindingInfo& BindingInfo, void* InstanceData)
+void UNiagaraDataInterfaceSkeletalMesh::GetVMExternalFunction(const FVMExternalFunctionBindingInfo& BindingInfo, void* InstanceData, FVMExternalFunction &OutFunc)
 {
 	FNDISkeletalMesh_InstanceData* InstData = (FNDISkeletalMesh_InstanceData*)InstanceData;
-	USkeletalMeshComponent* SkelComp = Cast<USkeletalMeshComponent>(InstData->Component.Get());
-	check(InstData && InstData->Mesh);
+	USkeletalMeshComponent* SkelComp = InstData != nullptr ? Cast<USkeletalMeshComponent>(InstData->Component.Get()) : nullptr;
+	
+	if (!InstData || !InstData->Mesh)
+	{
+		OutFunc = FVMExternalFunction();
+		return;
+	}
 
 	bool bNeedsVertexColors = false;
 
-	FVMExternalFunction Function;
 	if (BindingInfo.Name == RandomTriCoordName)
 	{
 		check(BindingInfo.GetNumInputs() == 1 && BindingInfo.GetNumOutputs() == 4);
-		Function = TFilterModeBinder<TAreaWeightingModeBinder<NDI_FUNC_BINDER(UNiagaraDataInterfaceSkeletalMesh, RandomTriCoord)>>::Bind(this, BindingInfo, InstanceData);
+		TFilterModeBinder<TAreaWeightingModeBinder<NDI_RAW_FUNC_BINDER(UNiagaraDataInterfaceSkeletalMesh, RandomTriCoord)>>::Bind(this, BindingInfo, InstanceData, OutFunc);
+	}
+	else if (BindingInfo.Name == IsValidTriCoordName)
+	{
+		check(BindingInfo.GetNumInputs() == 5 && BindingInfo.GetNumOutputs() == 1);
+		TFilterModeBinder<TAreaWeightingModeBinder<TNDIParamBinder<0, int32, TNDIParamBinder<1, float, TNDIParamBinder<2, float, TNDIParamBinder<3, float, NDI_RAW_FUNC_BINDER(UNiagaraDataInterfaceSkeletalMesh, IsValidTriCoord)>>>>>>::Bind(this, BindingInfo, InstanceData, OutFunc);
 	}
 	else if (BindingInfo.Name == GetTriPositionName)
 	{
 		check(BindingInfo.GetNumInputs() == 5 && BindingInfo.GetNumOutputs() == 3);
-		Function = TSkinningModeBinder<TNDIExplicitBinder<FNDITransformHandlerNoop, TNDIParamBinder<0, int32, TNDIParamBinder<1, float, TNDIParamBinder<2, float, TNDIParamBinder<3, float, NDI_FUNC_BINDER(UNiagaraDataInterfaceSkeletalMesh, GetTriCoordPosition)>>>>>>::Bind(this, BindingInfo, InstanceData);
+		TSkinningModeBinder<TNDIExplicitBinder<FNDITransformHandlerNoop, TNDIParamBinder<0, int32, TNDIParamBinder<1, float, TNDIParamBinder<2, float, TNDIParamBinder<3, float, NDI_RAW_FUNC_BINDER(UNiagaraDataInterfaceSkeletalMesh, GetTriCoordPosition)>>>>>>::Bind(this, BindingInfo, InstanceData, OutFunc);
 	}
 	else if (BindingInfo.Name == GetTriPositionWSName)
 	{
 		check(BindingInfo.GetNumInputs() == 5 && BindingInfo.GetNumOutputs() == 3);
-		Function = TSkinningModeBinder<TNDIExplicitBinder<FNDITransformHandler, TNDIParamBinder<0, int32, TNDIParamBinder<1, float, TNDIParamBinder<2, float, TNDIParamBinder<3, float, NDI_FUNC_BINDER(UNiagaraDataInterfaceSkeletalMesh, GetTriCoordPosition)>>>>>>::Bind(this, BindingInfo, InstanceData);
+		TSkinningModeBinder<TNDIExplicitBinder<FNDITransformHandler, TNDIParamBinder<0, int32, TNDIParamBinder<1, float, TNDIParamBinder<2, float, TNDIParamBinder<3, float, NDI_RAW_FUNC_BINDER(UNiagaraDataInterfaceSkeletalMesh, GetTriCoordPosition)>>>>>>::Bind(this, BindingInfo, InstanceData, OutFunc);
 	}
 	else if (BindingInfo.Name == GetTriPositionVelocityAndNormalName)
 	{
 		check(BindingInfo.GetNumInputs() == 5 && BindingInfo.GetNumOutputs() == 9);
-		Function = TSkinningModeBinder<TNDIExplicitBinder<FNDITransformHandlerNoop, TNDIParamBinder<0, int32, TNDIParamBinder<1, float, TNDIParamBinder<2, float, TNDIParamBinder<3, float, NDI_FUNC_BINDER(UNiagaraDataInterfaceSkeletalMesh, GetTriCoordPositionVelocityAndNormal)>>>>>>::Bind(this, BindingInfo, InstanceData);
+		TSkinningModeBinder<TNDIExplicitBinder<FNDITransformHandlerNoop, TNDIParamBinder<0, int32, TNDIParamBinder<1, float, TNDIParamBinder<2, float, TNDIParamBinder<3, float, NDI_RAW_FUNC_BINDER(UNiagaraDataInterfaceSkeletalMesh, GetTriCoordPositionVelocityAndNormal)>>>>>>::Bind(this, BindingInfo, InstanceData, OutFunc);
 	}
 	else if (BindingInfo.Name == GetTriPositionVelocityAndNormalWSName)
 	{
 		check(BindingInfo.GetNumInputs() == 5 && BindingInfo.GetNumOutputs() == 9);
-		Function = TSkinningModeBinder<TNDIExplicitBinder<FNDITransformHandler, TNDIParamBinder<0, int32, TNDIParamBinder<1, float, TNDIParamBinder<2, float, TNDIParamBinder<3, float, NDI_FUNC_BINDER(UNiagaraDataInterfaceSkeletalMesh, GetTriCoordPositionVelocityAndNormal)>>>>>>::Bind(this, BindingInfo, InstanceData);
+		TSkinningModeBinder<TNDIExplicitBinder<FNDITransformHandler, TNDIParamBinder<0, int32, TNDIParamBinder<1, float, TNDIParamBinder<2, float, TNDIParamBinder<3, float, NDI_RAW_FUNC_BINDER(UNiagaraDataInterfaceSkeletalMesh, GetTriCoordPositionVelocityAndNormal)>>>>>>::Bind(this, BindingInfo, InstanceData, OutFunc);
+	}
+	else if (BindingInfo.Name == GetTriPositionVelocityAndNormalBinormalTangentName)
+	{
+		check(BindingInfo.GetNumInputs() == 6 && BindingInfo.GetNumOutputs() == 17);
+		TSkinningModeBinder<TNDIExplicitBinder<FNDITransformHandlerNoop, TVertexAccessorBinder<TNDIParamBinder<0, int32, TNDIParamBinder<1, float, TNDIParamBinder<2, float, TNDIParamBinder<3, float, TNDIParamBinder<4, int32, NDI_RAW_FUNC_BINDER(UNiagaraDataInterfaceSkeletalMesh, GetTriCoordPositionVelocityAndNormalBinormalTangent)>>>>>>>>::Bind(this, BindingInfo, InstanceData, OutFunc);
+	}
+	else if (BindingInfo.Name == GetTriPositionVelocityAndNormalBinormalTangentWSName)
+	{
+		check(BindingInfo.GetNumInputs() == 6 && BindingInfo.GetNumOutputs() == 17);
+		TSkinningModeBinder<TNDIExplicitBinder<FNDITransformHandler, TVertexAccessorBinder<TNDIParamBinder<0, int32, TNDIParamBinder<1, float, TNDIParamBinder<2, float, TNDIParamBinder<3, float, TNDIParamBinder<4, int32, NDI_RAW_FUNC_BINDER(UNiagaraDataInterfaceSkeletalMesh, GetTriCoordPositionVelocityAndNormalBinormalTangent)>>>>>>>>::Bind(this, BindingInfo, InstanceData, OutFunc);
 	}
 	else if (BindingInfo.Name == GetTriColorName)
 	{
 		check(BindingInfo.GetNumInputs() == 5 && BindingInfo.GetNumOutputs() == 4);
 		bNeedsVertexColors = true;
-		Function = TNDIParamBinder<0, int32, TNDIParamBinder<1, float, TNDIParamBinder<2, float, TNDIParamBinder<3, float, NDI_FUNC_BINDER(UNiagaraDataInterfaceSkeletalMesh, GetTriCoordColor)>>>>::Bind(this, BindingInfo, InstanceData);
+		TNDIParamBinder<0, int32, TNDIParamBinder<1, float, TNDIParamBinder<2, float, TNDIParamBinder<3, float, NDI_RAW_FUNC_BINDER(UNiagaraDataInterfaceSkeletalMesh, GetTriCoordColor)>>>>::Bind(this, BindingInfo, InstanceData, OutFunc);
 	}
 	else if (BindingInfo.Name == GetTriUVName)
 	{
 		check(BindingInfo.GetNumInputs() == 6 && BindingInfo.GetNumOutputs() == 2);
-		Function = TVertexAccessorBinder<TNDIParamBinder<0, int32, TNDIParamBinder<1, float, TNDIParamBinder<2, float, TNDIParamBinder<3, float, TNDIParamBinder<4, int32, NDI_FUNC_BINDER(UNiagaraDataInterfaceSkeletalMesh, GetTriCoordUV)>>>>>>::Bind(this, BindingInfo, InstanceData);
+		TVertexAccessorBinder<TNDIParamBinder<0, int32, TNDIParamBinder<1, float, TNDIParamBinder<2, float, TNDIParamBinder<3, float, TNDIParamBinder<4, int32, NDI_RAW_FUNC_BINDER(UNiagaraDataInterfaceSkeletalMesh, GetTriCoordUV)>>>>>>::Bind(this, BindingInfo, InstanceData, OutFunc);
+	}
+	else if (BindingInfo.Name == GetTriangleCountName)
+	{
+		check(BindingInfo.GetNumInputs() == 1 && BindingInfo.GetNumOutputs() == 1);
+		TFilterModeBinder<TAreaWeightingModeBinder<NDI_RAW_FUNC_BINDER(UNiagaraDataInterfaceSkeletalMesh, GetFilteredTriangleCount)>>::Bind(this, BindingInfo, InstanceData, OutFunc);
+	}
+	else if (BindingInfo.Name == GetTriangleAtName)
+	{
+		check(BindingInfo.GetNumInputs() == 2 && BindingInfo.GetNumOutputs() == 4);
+		TFilterModeBinder<TAreaWeightingModeBinder<TNDIParamBinder<0, int32, NDI_RAW_FUNC_BINDER(UNiagaraDataInterfaceSkeletalMesh, GetFilteredTriangleAt)>>>::Bind(this, BindingInfo, InstanceData, OutFunc);
 	}
 	
 	check(InstData->Mesh);
@@ -1103,10 +1302,9 @@ FVMExternalFunction UNiagaraDataInterfaceSkeletalMesh::GetVMExternalFunction(con
 	if (bNeedsVertexColors && LODData.StaticVertexBuffers.ColorVertexBuffer.GetNumVertices() == 0)
 	{
 		UE_LOG(LogNiagara, Log, TEXT("Skeletal Mesh data interface is cannot run as it's reading color data on a mesh that does not provide it. - Mesh:%s  "), *InstData->Mesh->GetFullName());
-		Function = FVMExternalFunction();
+		OutFunc = FVMExternalFunction();
 	}
-	
-	return Function;
+
 }
 
 
@@ -1143,6 +1341,7 @@ bool UNiagaraDataInterfaceSkeletalMesh::Equals(const UNiagaraDataInterface* Othe
 bool UNiagaraDataInterfaceSkeletalMesh::InitPerInstanceData(void* PerInstanceData, FNiagaraSystemInstance* SystemInstance)
 {
 	FNDISkeletalMesh_InstanceData* Inst = new (PerInstanceData) FNDISkeletalMesh_InstanceData();
+	check(IsAligned(PerInstanceData, 16));
 	return Inst->Init(this, SystemInstance);
 }
 
@@ -1151,9 +1350,9 @@ void UNiagaraDataInterfaceSkeletalMesh::DestroyPerInstanceData(void* PerInstance
 	FNDISkeletalMesh_InstanceData* Inst = (FNDISkeletalMesh_InstanceData*)PerInstanceData;
 
 #if WITH_EDITOR
-	if(Inst->Mesh)
+	if(Inst->MeshSafe.IsValid())
 	{
-		Inst->Mesh->GetOnMeshChanged().RemoveAll(SystemInstance->GetComponent());
+		Inst->MeshSafe.Get()->GetOnMeshChanged().RemoveAll(SystemInstance->GetComponent());
 	}
 #endif
 
@@ -1166,10 +1365,46 @@ bool UNiagaraDataInterfaceSkeletalMesh::PerInstanceTick(void* PerInstanceData, F
 	return Inst->Tick(this, SystemInstance, InDeltaSeconds);
 }
 
+#if WITH_EDITOR	
+TArray<FNiagaraDataInterfaceError> UNiagaraDataInterfaceSkeletalMesh::GetErrors()
+{
+	TArray<FNiagaraDataInterfaceError> Errors;
+	bool bHasError= false;
+	if (DefaultMesh != nullptr)
+	{
+		for (auto info : DefaultMesh->GetLODInfoArray())
+		{
+			if (!info.bAllowCPUAccess)
+				bHasError = true;
+		}
+	}
+
+	if (Source == nullptr && bHasError)
+	{
+		FNiagaraDataInterfaceError CPUAccessNotAllowedError(FText::Format(LOCTEXT("CPUAccessNotAllowedError", "This mesh needs CPU access in order to be used properly.({0})"), FText::FromString(DefaultMesh->GetName())),
+			LOCTEXT("CPUAccessNotAllowedErrorSummary", "CPU access error"),
+			FNiagaraDataInterfaceFix::CreateLambda([=]()
+		{
+			DefaultMesh->Modify();
+			for (int i = 0; i < DefaultMesh->GetLODInfoArray().Num(); i++)
+			{
+				FSkeletalMeshLODInfo* info = &DefaultMesh->GetLODInfoArray()[i];
+				DefaultMesh->Modify();
+				info->bAllowCPUAccess = true;
+			}
+			return true;
+		}));
+
+		Errors.Add(CPUAccessNotAllowedError);
+	}
+	return Errors;
+}
+#endif
+
 //////////////////////////////////////////////////////////////////////////
 
 template<typename FilterMode, typename AreaWeightingMode>
-FORCEINLINE int32 UNiagaraDataInterfaceSkeletalMesh::RandomTriIndex(FSkeletalMeshAccessorHelper& Accessor, FNDISkeletalMesh_InstanceData* InstData)
+FORCEINLINE int32 UNiagaraDataInterfaceSkeletalMesh::RandomTriIndex(FRandomStream& RandStream, FSkeletalMeshAccessorHelper& Accessor, FNDISkeletalMesh_InstanceData* InstData)
 {
 	checkf(false, TEXT("Invalid template call for RandomTriIndex. Bug in Filter binding or Area Weighting binding. Contact code team."));
 	return 0;
@@ -1179,11 +1414,11 @@ template<>
 FORCEINLINE int32 UNiagaraDataInterfaceSkeletalMesh::RandomTriIndex<
 	TIntegralConstant<ENDISkeletalMesh_FilterMode, ENDISkeletalMesh_FilterMode::None>,
 	TIntegralConstant<ENDISkelMesh_AreaWeightingMode, ENDISkelMesh_AreaWeightingMode::None>>
-	(FSkeletalMeshAccessorHelper& Accessor, FNDISkeletalMesh_InstanceData* InstData)
+	(FRandomStream& RandStream, FSkeletalMeshAccessorHelper& Accessor, FNDISkeletalMesh_InstanceData* InstData)
 {
-	int32 SecIdx = InstData->RandStream.RandRange(0, Accessor.LODData->RenderSections.Num() - 1);
+	int32 SecIdx = RandStream.RandRange(0, Accessor.LODData->RenderSections.Num() - 1);
 	FSkelMeshRenderSection& Sec = Accessor.LODData->RenderSections[SecIdx];
-	int32 Tri = InstData->RandStream.RandRange(0, Sec.NumTriangles - 1);
+	int32 Tri = RandStream.RandRange(0, Sec.NumTriangles - 1);
 	return Sec.BaseIndex + Tri * 3;
 }
 
@@ -1191,11 +1426,11 @@ template<>
 FORCEINLINE int32 UNiagaraDataInterfaceSkeletalMesh::RandomTriIndex<
 	TIntegralConstant<ENDISkeletalMesh_FilterMode, ENDISkeletalMesh_FilterMode::None>,
 	TIntegralConstant<ENDISkelMesh_AreaWeightingMode, ENDISkelMesh_AreaWeightingMode::AreaWeighted>>
-	(FSkeletalMeshAccessorHelper& Accessor, FNDISkeletalMesh_InstanceData* InstData)
+	(FRandomStream& RandStream, FSkeletalMeshAccessorHelper& Accessor, FNDISkeletalMesh_InstanceData* InstData)
 {
 	const FSkeletalMeshSamplingInfo& SamplingInfo = InstData->Mesh->GetSamplingInfo();
 	const FSkeletalMeshSamplingLODBuiltData& WholeMeshBuiltData = SamplingInfo.GetWholeMeshLODBuiltData(InstData->GetLODIndex());
-	int32 TriIdx = WholeMeshBuiltData.AreaWeightedTriangleSampler.GetEntryIndex(InstData->RandStream.GetFraction(), InstData->RandStream.GetFraction());
+	int32 TriIdx = WholeMeshBuiltData.AreaWeightedTriangleSampler.GetEntryIndex(RandStream.GetFraction(), RandStream.GetFraction());
 	return TriIdx * 3;
 }
 
@@ -1203,9 +1438,9 @@ template<>
 FORCEINLINE int32 UNiagaraDataInterfaceSkeletalMesh::RandomTriIndex<
 	TIntegralConstant<ENDISkeletalMesh_FilterMode, ENDISkeletalMesh_FilterMode::SingleRegion>,
 	TIntegralConstant<ENDISkelMesh_AreaWeightingMode, ENDISkelMesh_AreaWeightingMode::None>>
-	(FSkeletalMeshAccessorHelper& Accessor, FNDISkeletalMesh_InstanceData* InstData)
+	(FRandomStream& RandStream, FSkeletalMeshAccessorHelper& Accessor, FNDISkeletalMesh_InstanceData* InstData)
 {
-	int32 Idx = InstData->RandStream.RandRange(0, Accessor.SamplingRegionBuiltData->TriangleIndices.Num() - 1);
+	int32 Idx = RandStream.RandRange(0, Accessor.SamplingRegionBuiltData->TriangleIndices.Num() - 1);
 	return Accessor.SamplingRegionBuiltData->TriangleIndices[Idx];
 }
 
@@ -1213,9 +1448,9 @@ template<>
 FORCEINLINE int32 UNiagaraDataInterfaceSkeletalMesh::RandomTriIndex<
 	TIntegralConstant<ENDISkeletalMesh_FilterMode, ENDISkeletalMesh_FilterMode::SingleRegion>,
 	TIntegralConstant<ENDISkelMesh_AreaWeightingMode, ENDISkelMesh_AreaWeightingMode::AreaWeighted>>
-	(FSkeletalMeshAccessorHelper& Accessor, FNDISkeletalMesh_InstanceData* InstData)
+	(FRandomStream& RandStream, FSkeletalMeshAccessorHelper& Accessor, FNDISkeletalMesh_InstanceData* InstData)
 {
-	int32 Idx = Accessor.SamplingRegionBuiltData->AreaWeightedSampler.GetEntryIndex(InstData->RandStream.GetFraction(), InstData->RandStream.GetFraction());
+	int32 Idx = Accessor.SamplingRegionBuiltData->AreaWeightedSampler.GetEntryIndex(RandStream.GetFraction(), RandStream.GetFraction());
 	return Accessor.SamplingRegionBuiltData->TriangleIndices[Idx];
 }
 
@@ -1223,13 +1458,13 @@ template<>
 FORCEINLINE int32 UNiagaraDataInterfaceSkeletalMesh::RandomTriIndex<
 	TIntegralConstant<ENDISkeletalMesh_FilterMode, ENDISkeletalMesh_FilterMode::MultiRegion>,
 	TIntegralConstant<ENDISkelMesh_AreaWeightingMode, ENDISkelMesh_AreaWeightingMode::None>>
-	(FSkeletalMeshAccessorHelper& Accessor, FNDISkeletalMesh_InstanceData* InstData)
+	(FRandomStream& RandStream, FSkeletalMeshAccessorHelper& Accessor, FNDISkeletalMesh_InstanceData* InstData)
 {
-	int32 RegionIdx = InstData->RandStream.RandRange(0, InstData->SamplingRegionIndices.Num() - 1);
+	int32 RegionIdx = RandStream.RandRange(0, InstData->SamplingRegionIndices.Num() - 1);
 	const FSkeletalMeshSamplingInfo& SamplingInfo = InstData->Mesh->GetSamplingInfo();
 	const FSkeletalMeshSamplingRegion& Region = SamplingInfo.GetRegion(InstData->SamplingRegionIndices[RegionIdx]);
 	const FSkeletalMeshSamplingRegionBuiltData& RegionBuiltData = SamplingInfo.GetRegionBuiltData(InstData->SamplingRegionIndices[RegionIdx]);
-	int32 Idx = InstData->RandStream.RandRange(0, RegionBuiltData.TriangleIndices.Num() - 1);
+	int32 Idx = RandStream.RandRange(0, RegionBuiltData.TriangleIndices.Num() - 1);
 	return RegionBuiltData.TriangleIndices[Idx];
 }
 
@@ -1237,19 +1472,20 @@ template<>
 FORCEINLINE int32 UNiagaraDataInterfaceSkeletalMesh::RandomTriIndex<
 	TIntegralConstant<ENDISkeletalMesh_FilterMode, ENDISkeletalMesh_FilterMode::MultiRegion>,
 	TIntegralConstant<ENDISkelMesh_AreaWeightingMode, ENDISkelMesh_AreaWeightingMode::AreaWeighted>>
-	(FSkeletalMeshAccessorHelper& Accessor, FNDISkeletalMesh_InstanceData* InstData)
+	(FRandomStream& RandStream, FSkeletalMeshAccessorHelper& Accessor, FNDISkeletalMesh_InstanceData* InstData)
 {
-	int32 RegionIdx = InstData->SamplingRegionAreaWeightedSampler.GetEntryIndex(InstData->RandStream.GetFraction(), InstData->RandStream.GetFraction());
+	int32 RegionIdx = InstData->SamplingRegionAreaWeightedSampler.GetEntryIndex(RandStream.GetFraction(), RandStream.GetFraction());
 	const FSkeletalMeshSamplingInfo& SamplingInfo = InstData->Mesh->GetSamplingInfo();
 	const FSkeletalMeshSamplingRegion& Region = SamplingInfo.GetRegion(InstData->SamplingRegionIndices[RegionIdx]);
 	const FSkeletalMeshSamplingRegionBuiltData& RegionBuiltData = SamplingInfo.GetRegionBuiltData(InstData->SamplingRegionIndices[RegionIdx]);
-	int32 Idx = RegionBuiltData.AreaWeightedSampler.GetEntryIndex(InstData->RandStream.GetFraction(), InstData->RandStream.GetFraction());
+	int32 Idx = RegionBuiltData.AreaWeightedSampler.GetEntryIndex(RandStream.GetFraction(), RandStream.GetFraction());
 	return RegionBuiltData.TriangleIndices[Idx];
 }
 
 template<typename FilterMode, typename AreaWeightingMode>
 void UNiagaraDataInterfaceSkeletalMesh::RandomTriCoord(FVectorVMContext& Context)
 {
+	SCOPE_CYCLE_COUNTER(STAT_NiagaraSkel_Sample);
 	FUserPtrHandler<FNDISkeletalMesh_InstanceData> InstData(Context);
 	checkfSlow(InstData.Get(), TEXT("Skeletal Mesh Interface has invalid instance data. %s"), *GetPathName());
 	checkfSlow(InstData->Mesh, TEXT("Skeletal Mesh Interface has invalid mesh. %s"), *GetPathName());
@@ -1262,8 +1498,8 @@ void UNiagaraDataInterfaceSkeletalMesh::RandomTriCoord(FVectorVMContext& Context
 	
 	for (int32 i = 0; i < Context.NumInstances; ++i)
 	{
-		*OutTri.GetDest() = RandomTriIndex<FilterMode, AreaWeightingMode>(MeshAccessor, InstData);
-		FVector Bary = RandomBarycentricCoord(InstData->RandStream);
+		*OutTri.GetDest() = RandomTriIndex<FilterMode, AreaWeightingMode>(Context.RandStream, MeshAccessor, InstData);
+		FVector Bary = RandomBarycentricCoord(Context.RandStream);
 		*OutBaryX.GetDest() = Bary.X;		*OutBaryY.GetDest() = Bary.Y;		*OutBaryZ.GetDest() = Bary.Z;
 		
 		OutTri.Advance();
@@ -1271,9 +1507,286 @@ void UNiagaraDataInterfaceSkeletalMesh::RandomTriCoord(FVectorVMContext& Context
 	}
 }
 
+template<typename FilterMode, typename AreaWeightingMode, typename TriType, typename BaryXType, typename BaryYType, typename BaryZType >
+void UNiagaraDataInterfaceSkeletalMesh::IsValidTriCoord(FVectorVMContext& Context)
+{
+	SCOPE_CYCLE_COUNTER(STAT_NiagaraSkel_Sample);
+	
+	TriType TriParam(Context);
+	BaryXType BaryXParam(Context);	BaryYType BaryYParam(Context);	BaryZType BaryZParam(Context);
+	
+	FUserPtrHandler<FNDISkeletalMesh_InstanceData> InstData(Context);
+	checkfSlow(InstData.Get(), TEXT("Skeletal Mesh Interface has invalid instance data. %s"), *GetPathName());
+	checkfSlow(InstData->Mesh, TEXT("Skeletal Mesh Interface has invalid mesh. %s"), *GetPathName());
+	
+	FRegisterHandler<FNiagaraBool> OutValid(Context);
+
+	FSkeletalMeshAccessorHelper MeshAccessor;
+	MeshAccessor.Init<FilterMode, AreaWeightingMode>(InstData);
+
+	for (int32 i = 0; i < Context.NumInstances; ++i)
+	{
+		int32 RequestedIndex = TriParam.Get() + 2; // Get the last triangle index of the set
+
+		FNiagaraBool Value;
+		Value.SetValue(MeshAccessor.IndexBuffer != nullptr && MeshAccessor.IndexBuffer->Num() > RequestedIndex);
+		*OutValid.GetDest() = Value;
+
+		OutValid.Advance();
+		BaryXParam.Advance();		BaryYParam.Advance();		BaryZParam.Advance(); TriParam.Advance();
+	}
+}
+
+//////////////////////////////////////////////////////////////////////////
+
+template<typename FilterMode, typename AreaWeightingMode>
+FORCEINLINE int32 UNiagaraDataInterfaceSkeletalMesh::GetSpecificTriangleCount(FSkeletalMeshAccessorHelper& Accessor, FNDISkeletalMesh_InstanceData* InstData)
+{
+	checkf(false, TEXT("Invalid template call for GetSpecificTriangleCount. Bug in Filter binding or Area Weighting binding. Contact code team."));
+	return 0;
+}
+
+template<>
+FORCEINLINE int32 UNiagaraDataInterfaceSkeletalMesh::GetSpecificTriangleCount<
+	TIntegralConstant<ENDISkeletalMesh_FilterMode, ENDISkeletalMesh_FilterMode::None>,
+	TIntegralConstant<ENDISkelMesh_AreaWeightingMode, ENDISkelMesh_AreaWeightingMode::None>>
+	(FSkeletalMeshAccessorHelper& Accessor, FNDISkeletalMesh_InstanceData* InstData)
+{
+	int32 NumTris = 0;
+	for (int32 i = 0; i < Accessor.LODData->RenderSections.Num(); i++)
+	{
+		NumTris += Accessor.LODData->RenderSections[i].NumTriangles;
+	}
+	return NumTris;
+}
+
+template<>
+FORCEINLINE int32 UNiagaraDataInterfaceSkeletalMesh::GetSpecificTriangleCount<
+	TIntegralConstant<ENDISkeletalMesh_FilterMode, ENDISkeletalMesh_FilterMode::None>,
+	TIntegralConstant<ENDISkelMesh_AreaWeightingMode, ENDISkelMesh_AreaWeightingMode::AreaWeighted>>
+	(FSkeletalMeshAccessorHelper& Accessor, FNDISkeletalMesh_InstanceData* InstData)
+{
+	const FSkeletalMeshSamplingInfo& SamplingInfo = InstData->Mesh->GetSamplingInfo();
+	const FSkeletalMeshSamplingLODBuiltData& WholeMeshBuiltData = SamplingInfo.GetWholeMeshLODBuiltData(InstData->GetLODIndex());
+	return WholeMeshBuiltData.AreaWeightedTriangleSampler.GetNumEntries();
+}
+
+template<>
+FORCEINLINE int32 UNiagaraDataInterfaceSkeletalMesh::GetSpecificTriangleCount<
+	TIntegralConstant<ENDISkeletalMesh_FilterMode, ENDISkeletalMesh_FilterMode::SingleRegion>,
+	TIntegralConstant<ENDISkelMesh_AreaWeightingMode, ENDISkelMesh_AreaWeightingMode::None>>
+	(FSkeletalMeshAccessorHelper& Accessor, FNDISkeletalMesh_InstanceData* InstData)
+{
+	return Accessor.SamplingRegionBuiltData->TriangleIndices.Num();
+}
+
+template<>
+FORCEINLINE int32 UNiagaraDataInterfaceSkeletalMesh::GetSpecificTriangleCount<
+	TIntegralConstant<ENDISkeletalMesh_FilterMode, ENDISkeletalMesh_FilterMode::SingleRegion>,
+	TIntegralConstant<ENDISkelMesh_AreaWeightingMode, ENDISkelMesh_AreaWeightingMode::AreaWeighted>>
+	(FSkeletalMeshAccessorHelper& Accessor, FNDISkeletalMesh_InstanceData* InstData)
+{
+	return Accessor.SamplingRegionBuiltData->AreaWeightedSampler.GetNumEntries();
+}
+
+template<>
+FORCEINLINE int32 UNiagaraDataInterfaceSkeletalMesh::GetSpecificTriangleCount<
+	TIntegralConstant<ENDISkeletalMesh_FilterMode, ENDISkeletalMesh_FilterMode::MultiRegion>,
+	TIntegralConstant<ENDISkelMesh_AreaWeightingMode, ENDISkelMesh_AreaWeightingMode::None>>
+	(FSkeletalMeshAccessorHelper& Accessor, FNDISkeletalMesh_InstanceData* InstData)
+{
+	int32 NumTris = 0;
+
+	for (int32 RegionIdx = 0; RegionIdx < InstData->SamplingRegionIndices.Num(); RegionIdx++)
+	{
+		const FSkeletalMeshSamplingInfo& SamplingInfo = InstData->Mesh->GetSamplingInfo();
+		const FSkeletalMeshSamplingRegion& Region = SamplingInfo.GetRegion(InstData->SamplingRegionIndices[RegionIdx]);
+		const FSkeletalMeshSamplingRegionBuiltData& RegionBuiltData = SamplingInfo.GetRegionBuiltData(InstData->SamplingRegionIndices[RegionIdx]);
+		NumTris += RegionBuiltData.TriangleIndices.Num();
+	}
+	return NumTris;
+}
+
+template<>
+FORCEINLINE int32 UNiagaraDataInterfaceSkeletalMesh::GetSpecificTriangleCount<
+	TIntegralConstant<ENDISkeletalMesh_FilterMode, ENDISkeletalMesh_FilterMode::MultiRegion>,
+	TIntegralConstant<ENDISkelMesh_AreaWeightingMode, ENDISkelMesh_AreaWeightingMode::AreaWeighted>>
+	(FSkeletalMeshAccessorHelper& Accessor, FNDISkeletalMesh_InstanceData* InstData)
+{
+	int32 NumTris = 0;
+
+	for (int32 RegionIdx = 0; RegionIdx < InstData->SamplingRegionIndices.Num(); RegionIdx++)
+	{
+		const FSkeletalMeshSamplingInfo& SamplingInfo = InstData->Mesh->GetSamplingInfo();
+		const FSkeletalMeshSamplingRegion& Region = SamplingInfo.GetRegion(InstData->SamplingRegionIndices[RegionIdx]);
+		const FSkeletalMeshSamplingRegionBuiltData& RegionBuiltData = SamplingInfo.GetRegionBuiltData(InstData->SamplingRegionIndices[RegionIdx]);
+		NumTris += RegionBuiltData.TriangleIndices.Num();
+	}
+	return NumTris;
+}
+
+template<typename FilterMode, typename AreaWeightingMode>
+void UNiagaraDataInterfaceSkeletalMesh::GetFilteredTriangleCount(FVectorVMContext& Context)
+{
+	SCOPE_CYCLE_COUNTER(STAT_NiagaraSkel_Sample);
+	FUserPtrHandler<FNDISkeletalMesh_InstanceData> InstData(Context);
+	checkfSlow(InstData.Get(), TEXT("Skeletal Mesh Interface has invalid instance data. %s"), *GetPathName());
+	checkfSlow(InstData->Mesh, TEXT("Skeletal Mesh Interface has invalid mesh. %s"), *GetPathName());
+
+	FRegisterHandler<int32> OutTri(Context);
+
+	FSkeletalMeshAccessorHelper MeshAccessor;
+	MeshAccessor.Init<FilterMode, AreaWeightingMode>(InstData);
+
+	int32 Count = GetSpecificTriangleCount<FilterMode, AreaWeightingMode>(MeshAccessor, InstData);
+	for (int32 i = 0; i < Context.NumInstances; ++i)
+	{
+		*OutTri.GetDest() = Count;
+		OutTri.Advance();
+	}
+}
+
+
+//////////////////////////////////////////////////////////////////////////
+
+template<typename FilterMode, typename AreaWeightingMode>
+FORCEINLINE int32 UNiagaraDataInterfaceSkeletalMesh::GetSpecificTriangleAt(FSkeletalMeshAccessorHelper& Accessor, FNDISkeletalMesh_InstanceData* InstData, int32 FilteredIndex)
+{
+	checkf(false, TEXT("Invalid template call for GetSpecificTriangleAt. Bug in Filter binding or Area Weighting binding. Contact code team."));
+	return 0;
+}
+
+template<>
+FORCEINLINE int32 UNiagaraDataInterfaceSkeletalMesh::GetSpecificTriangleAt<
+	TIntegralConstant<ENDISkeletalMesh_FilterMode, ENDISkeletalMesh_FilterMode::None>,
+	TIntegralConstant<ENDISkelMesh_AreaWeightingMode, ENDISkelMesh_AreaWeightingMode::None>>
+	(FSkeletalMeshAccessorHelper& Accessor, FNDISkeletalMesh_InstanceData* InstData, int32 FilteredIndex)
+{
+	for (int32 i = 0; i < Accessor.LODData->RenderSections.Num(); i++)
+	{
+		if (Accessor.LODData->RenderSections[i].NumTriangles > (uint32)FilteredIndex)
+		{
+			FSkelMeshRenderSection& Sec = Accessor.LODData->RenderSections[i];
+			return Sec.BaseIndex + FilteredIndex * 3;
+		}
+		FilteredIndex -= Accessor.LODData->RenderSections[i].NumTriangles;
+	}
+	return 0;
+}
+
+template<>
+FORCEINLINE int32 UNiagaraDataInterfaceSkeletalMesh::GetSpecificTriangleAt<
+	TIntegralConstant<ENDISkeletalMesh_FilterMode, ENDISkeletalMesh_FilterMode::None>,
+	TIntegralConstant<ENDISkelMesh_AreaWeightingMode, ENDISkelMesh_AreaWeightingMode::AreaWeighted>>
+	(FSkeletalMeshAccessorHelper& Accessor, FNDISkeletalMesh_InstanceData* InstData, int32 FilteredIndex)
+{
+	int32 TriIdx = FilteredIndex;
+	return TriIdx * 3;
+}
+
+template<>
+FORCEINLINE int32 UNiagaraDataInterfaceSkeletalMesh::GetSpecificTriangleAt<
+	TIntegralConstant<ENDISkeletalMesh_FilterMode, ENDISkeletalMesh_FilterMode::SingleRegion>,
+	TIntegralConstant<ENDISkelMesh_AreaWeightingMode, ENDISkelMesh_AreaWeightingMode::None>>
+	(FSkeletalMeshAccessorHelper& Accessor, FNDISkeletalMesh_InstanceData* InstData, int32 FilteredIndex)
+{
+	int32 MaxIdx = Accessor.SamplingRegionBuiltData->TriangleIndices.Num() - 1;
+	FilteredIndex = FMath::Min(FilteredIndex, MaxIdx);
+	return Accessor.SamplingRegionBuiltData->TriangleIndices[FilteredIndex];
+}
+
+template<>
+FORCEINLINE int32 UNiagaraDataInterfaceSkeletalMesh::GetSpecificTriangleAt<
+	TIntegralConstant<ENDISkeletalMesh_FilterMode, ENDISkeletalMesh_FilterMode::SingleRegion>,
+	TIntegralConstant<ENDISkelMesh_AreaWeightingMode, ENDISkelMesh_AreaWeightingMode::AreaWeighted>>
+	(FSkeletalMeshAccessorHelper& Accessor, FNDISkeletalMesh_InstanceData* InstData, int32 FilteredIndex)
+{
+	int32 Idx = FilteredIndex;
+	int32 MaxIdx = Accessor.SamplingRegionBuiltData->TriangleIndices.Num() - 1;
+	Idx = FMath::Min(Idx, MaxIdx);
+
+	return Accessor.SamplingRegionBuiltData->TriangleIndices[Idx];
+}
+
+template<>
+FORCEINLINE int32 UNiagaraDataInterfaceSkeletalMesh::GetSpecificTriangleAt<
+	TIntegralConstant<ENDISkeletalMesh_FilterMode, ENDISkeletalMesh_FilterMode::MultiRegion>,
+	TIntegralConstant<ENDISkelMesh_AreaWeightingMode, ENDISkelMesh_AreaWeightingMode::None>>
+	(FSkeletalMeshAccessorHelper& Accessor, FNDISkeletalMesh_InstanceData* InstData, int32 FilteredIndex)
+{
+	for (int32 RegionIdx = 0; RegionIdx < InstData->SamplingRegionIndices.Num(); RegionIdx++)
+	{
+		const FSkeletalMeshSamplingInfo& SamplingInfo = InstData->Mesh->GetSamplingInfo();
+		const FSkeletalMeshSamplingRegion& Region = SamplingInfo.GetRegion(InstData->SamplingRegionIndices[RegionIdx]);
+		const FSkeletalMeshSamplingRegionBuiltData& RegionBuiltData = SamplingInfo.GetRegionBuiltData(InstData->SamplingRegionIndices[RegionIdx]);
+		if (FilteredIndex < RegionBuiltData.TriangleIndices.Num())
+		{
+			return RegionBuiltData.TriangleIndices[FilteredIndex];
+		}
+
+		FilteredIndex -= RegionBuiltData.TriangleIndices.Num();
+	}
+	return 0;
+}
+
+template<>
+FORCEINLINE int32 UNiagaraDataInterfaceSkeletalMesh::GetSpecificTriangleAt<
+	TIntegralConstant<ENDISkeletalMesh_FilterMode, ENDISkeletalMesh_FilterMode::MultiRegion>,
+	TIntegralConstant<ENDISkelMesh_AreaWeightingMode, ENDISkelMesh_AreaWeightingMode::AreaWeighted>>
+	(FSkeletalMeshAccessorHelper& Accessor, FNDISkeletalMesh_InstanceData* InstData, int32 FilteredIndex)
+{
+	for (int32 RegionIdx = 0; RegionIdx < InstData->SamplingRegionIndices.Num(); RegionIdx++)
+	{
+		const FSkeletalMeshSamplingInfo& SamplingInfo = InstData->Mesh->GetSamplingInfo();
+		const FSkeletalMeshSamplingRegion& Region = SamplingInfo.GetRegion(InstData->SamplingRegionIndices[RegionIdx]);
+		const FSkeletalMeshSamplingRegionBuiltData& RegionBuiltData = SamplingInfo.GetRegionBuiltData(InstData->SamplingRegionIndices[RegionIdx]);
+		if (FilteredIndex < RegionBuiltData.TriangleIndices.Num())
+		{
+			return RegionBuiltData.TriangleIndices[FilteredIndex];
+		}
+		FilteredIndex -= RegionBuiltData.TriangleIndices.Num();
+	}
+	return 0;
+}
+
+template<typename FilterMode, typename AreaWeightingMode, typename TriType>
+void UNiagaraDataInterfaceSkeletalMesh::GetFilteredTriangleAt(FVectorVMContext& Context)
+{
+	SCOPE_CYCLE_COUNTER(STAT_NiagaraSkel_Sample);
+
+	TriType TriParam(Context);	
+	FUserPtrHandler<FNDISkeletalMesh_InstanceData> InstData(Context);
+	checkfSlow(InstData.Get(), TEXT("Skeletal Mesh Interface has invalid instance data. %s"), *GetPathName());
+	checkfSlow(InstData->Mesh, TEXT("Skeletal Mesh Interface has invalid mesh. %s"), *GetPathName());
+	FRegisterHandler<int32> OutTri(Context);
+	FRegisterHandler<float> OutBaryX(Context);	FRegisterHandler<float> OutBaryY(Context);	FRegisterHandler<float> OutBaryZ(Context);
+
+	FSkeletalMeshAccessorHelper Accessor;
+	Accessor.Init<FilterMode, AreaWeightingMode>(InstData);
+	
+	for (int32 i = 0; i < Context.NumInstances; ++i)
+	{
+		int32 Tri = TriParam.Get();
+		int32 RealIdx = 0;
+		RealIdx = GetSpecificTriangleAt<FilterMode, AreaWeightingMode>(Accessor, InstData, Tri);
+
+		int32 TriMax = Accessor.IndexBuffer->Num() - 3;
+		RealIdx = FMath::Min(RealIdx, TriMax);
+		
+		*OutTri.GetDest() = RealIdx;
+		float Coord = 1.0f / 3.0f;
+		*OutBaryX.GetDest() = Coord;		*OutBaryY.GetDest() = Coord;		*OutBaryZ.GetDest() = Coord;
+
+		TriParam.Advance();
+		OutTri.Advance();
+		OutBaryX.Advance();		OutBaryY.Advance();		OutBaryZ.Advance();
+	}
+}
+
 template<typename SkinningHandlerType, typename TransformHandlerType, typename TriType, typename BaryXType, typename BaryYType, typename BaryZType>
 void UNiagaraDataInterfaceSkeletalMesh::GetTriCoordPosition(FVectorVMContext& Context)
 {
+	SCOPE_CYCLE_COUNTER(STAT_NiagaraSkel_Sample);
 	SkinningHandlerType SkinningHandler;
 	TransformHandlerType TransformHandler;
 	TriType TriParam(Context);
@@ -1287,11 +1800,16 @@ void UNiagaraDataInterfaceSkeletalMesh::GetTriCoordPosition(FVectorVMContext& Co
 
 	FSkeletalMeshAccessorHelper Accessor;
 	Accessor.Init<TIntegralConstant<int32, 0>,TIntegralConstant<int32, 0>>(InstData);
+	int32 TriMax = Accessor.IndexBuffer->Num() - 3;
 	for (int32 i = 0; i < Context.NumInstances; ++i)
 	{
 		int32 Tri = TriParam.Get();
 		FVector V0;		FVector V1;		FVector V2;
-		SkinningHandler.GetSkinnedTrianglePositions(Accessor, Tri, V0, V1, V2);
+		int32 Idx0; int32 Idx1; int32 Idx2;
+
+		Tri = FMath::Min(Tri, TriMax);
+		
+		SkinningHandler.GetSkinnedTrianglePositions(Accessor, Tri, V0, V1, V2, Idx0, Idx1, Idx2);
 		
 		FVector Pos = BarycentricInterpolate(BaryXParam.Get(), BaryYParam.Get(), BaryZParam.Get(), V0, V1, V2);
 		TransformHandler.TransformPosition(Pos, InstData->Transform);
@@ -1307,6 +1825,7 @@ void UNiagaraDataInterfaceSkeletalMesh::GetTriCoordPosition(FVectorVMContext& Co
 template<typename SkinningHandlerType, typename TransformHandlerType, typename TriType, typename BaryXType, typename BaryYType, typename BaryZType>
 void UNiagaraDataInterfaceSkeletalMesh::GetTriCoordPositionVelocityAndNormal(FVectorVMContext& Context)
 {
+	SCOPE_CYCLE_COUNTER(STAT_NiagaraSkel_Sample);
 	SkinningHandlerType SkinningHandler;
 	TransformHandlerType TransformHandler;
 	TriType TriParam(Context);
@@ -1323,12 +1842,15 @@ void UNiagaraDataInterfaceSkeletalMesh::GetTriCoordPositionVelocityAndNormal(FVe
 	FSkeletalMeshAccessorHelper Accessor;
 	Accessor.Init<TIntegralConstant<int32, 0>, TIntegralConstant<int32, 0>>(InstData);
 	float InvDt = 1.0f / InstData->DeltaSeconds;
+	int32 TriMax = Accessor.IndexBuffer->Num() - 3;
 	for (int32 i = 0; i < Context.NumInstances; ++i)
 	{
 		int32 Tri = TriParam.Get();
 		FVector Pos0;		FVector Pos1;		FVector Pos2;
-		FVector Prev0;		FVector Prev1;		FVector Prev2;
-		SkinningHandler.GetSkinnedTrianglePositions(Accessor, Tri, Pos0, Pos1, Pos2, Prev0, Prev1, Prev2);
+		FVector Prev0;		FVector Prev1;		FVector Prev2;		
+		int32 Idx0; int32 Idx1; int32 Idx2;
+		Tri = FMath::Min(Tri, TriMax);
+		SkinningHandler.GetSkinnedTrianglePositions(Accessor, Tri, Pos0, Pos1, Pos2, Prev0, Prev1, Prev2, Idx0, Idx1, Idx2);
 
 		FVector Pos = BarycentricInterpolate(BaryXParam.Get(), BaryYParam.Get(), BaryZParam.Get(), Pos0, Pos1, Pos2);
 		FVector Prev = BarycentricInterpolate(BaryXParam.Get(), BaryYParam.Get(), BaryZParam.Get(), Prev0, Prev1, Prev2);
@@ -1337,7 +1859,10 @@ void UNiagaraDataInterfaceSkeletalMesh::GetTriCoordPositionVelocityAndNormal(FVe
 
 		FVector Vel = (Pos - Prev) * InvDt;
 
-		FVector Normal = ((Pos1 - Pos0) ^ (Pos2 - Pos0)).GetUnsafeNormal();//Temporarily having to get a dirty normal here until the newer pre skinning goodness comes online.
+		// Normal calculation code based on tools code found at:
+		// \Engine\Source\Developer\MeshUtilities\Private\MeshUtilities.cpp
+		FVector Normal = ((Pos1 - Pos2) ^ (Pos0 - Pos2)).GetUnsafeNormal();//Temporarily having to get a dirty normal here until the newer pre skinning goodness comes online.
+		TransformHandler.TransformVector(Normal, InstData->Transform);
 
 		*OutPosX.GetDest() = Pos.X;		*OutPosY.GetDest() = Pos.Y;		*OutPosZ.GetDest() = Pos.Z;
 		*OutVelX.GetDest() = Vel.X;		*OutVelY.GetDest() = Vel.Y;		*OutVelZ.GetDest() = Vel.Z;
@@ -1353,9 +1878,121 @@ void UNiagaraDataInterfaceSkeletalMesh::GetTriCoordPositionVelocityAndNormal(FVe
 	}
 }
 
+template<typename SkinningHandlerType, typename TransformHandlerType, typename VertexAccessorType, typename TriType, typename BaryXType, typename BaryYType, typename BaryZType, typename UVSetType>
+void UNiagaraDataInterfaceSkeletalMesh::GetTriCoordPositionVelocityAndNormalBinormalTangent(FVectorVMContext& Context)
+{
+	SCOPE_CYCLE_COUNTER(STAT_NiagaraSkel_Sample);
+	SkinningHandlerType SkinningHandler;
+	TransformHandlerType TransformHandler;
+	VertexAccessorType VertAccessor;
+	TriType TriParam(Context);
+	BaryXType BaryXParam(Context);	BaryYType BaryYParam(Context);	BaryZType BaryZParam(Context);
+	UVSetType UVSetParam(Context);
+
+	FUserPtrHandler<FNDISkeletalMesh_InstanceData> InstData(Context);
+
+	checkfSlow(InstData.Get(), TEXT("Skeletal Mesh Interface has invalid instance data. %s"), *GetPathName());
+	checkfSlow(InstData->Mesh, TEXT("Skeletal Mesh Interface has invalid mesh. %s"), *GetPathName());
+
+	FRegisterHandler<float> OutPosX(Context);	FRegisterHandler<float> OutPosY(Context);	FRegisterHandler<float> OutPosZ(Context);
+	FRegisterHandler<float> OutVelX(Context);	FRegisterHandler<float> OutVelY(Context);	FRegisterHandler<float> OutVelZ(Context);
+	FRegisterHandler<float> OutNormX(Context);	FRegisterHandler<float> OutNormY(Context);	FRegisterHandler<float> OutNormZ(Context);
+	FRegisterHandler<float> OutBiNormX(Context);	FRegisterHandler<float> OutBiNormY(Context);	FRegisterHandler<float> OutBiNormZ(Context);
+	FRegisterHandler<float> OutTangentX(Context);	FRegisterHandler<float> OutTangentY(Context);	FRegisterHandler<float> OutTangentZ(Context);
+	FRegisterHandler<float> OutUVX(Context);	FRegisterHandler<float> OutUVY(Context);
+
+	FSkinWeightVertexBuffer* SkinWeightBuffer;
+	FSkeletalMeshLODRenderData& LODData = InstData->GetLODRenderDataAndSkinWeights(SkinWeightBuffer);
+
+	FSkeletalMeshAccessorHelper Accessor;
+	Accessor.Init<TIntegralConstant<int32, 0>, TIntegralConstant<int32, 0>>(InstData);
+	int32 TriMax = Accessor.IndexBuffer->Num() - 3;
+	float InvDt = 1.0f / InstData->DeltaSeconds;
+	for (int32 i = 0; i < Context.NumInstances; ++i)
+	{
+		int32 Tri = FMath::Min(TriParam.Get(), TriMax);
+		FVector Pos0;		FVector Pos1;		FVector Pos2;
+		FVector Prev0;		FVector Prev1;		FVector Prev2;
+		FVector Normal;
+		FVector Binormal;
+		FVector Tangent;
+		int32 Idx0; int32 Idx1; int32 Idx2;
+		SkinningHandler.GetSkinnedTrianglePositions(Accessor, Tri, Pos0, Pos1, Pos2, Prev0, Prev1, Prev2, Idx0, Idx1, Idx2);
+		int32 UVSet = UVSetParam.Get();
+
+		FVector Pos = BarycentricInterpolate(BaryXParam.Get(), BaryYParam.Get(), BaryZParam.Get(), Pos0, Pos1, Pos2);
+		FVector Prev = BarycentricInterpolate(BaryXParam.Get(), BaryYParam.Get(), BaryZParam.Get(), Prev0, Prev1, Prev2);
+
+		FMatrix Transform = InstData->Transform;
+		FMatrix PrevTransform = InstData->PrevTransform;
+		TransformHandler.TransformPosition(Pos, Transform);
+		TransformHandler.TransformPosition(Prev, PrevTransform);
+
+		FVector2D UV0 = VertAccessor.GetVertexUV(LODData, Idx0, UVSet);
+		FVector2D UV1 = VertAccessor.GetVertexUV(LODData, Idx1, UVSet);
+		FVector2D UV2 = VertAccessor.GetVertexUV(LODData, Idx2, UVSet);
+
+		Normal = ((Pos1 - Pos2) ^ (Pos0 - Pos2)).GetSafeNormal();
+		
+		// Normal binormal tangent calculation code based on tools code found at:
+		// \Engine\Source\Developer\MeshUtilities\Private\MeshUtilities.cpp
+		// Skeletal_ComputeTriangleTangents
+		FMatrix	ParameterToLocal(
+			FPlane(Pos1.X - Pos0.X, Pos1.Y - Pos0.Y, Pos1.Z - Pos0.Z, 0),
+			FPlane(Pos2.X - Pos0.X, Pos2.Y - Pos0.Y, Pos2.Z - Pos0.Z, 0),
+			FPlane(Pos0.X, Pos0.Y, Pos0.Z, 0),
+			FPlane(0, 0, 0, 1)
+		);
+
+		FMatrix ParameterToTexture(
+			FPlane(UV1.X - UV0.X, UV1.Y - UV0.Y, 0, 0),
+			FPlane(UV2.X - UV0.X, UV2.Y - UV0.Y, 0, 0),
+			FPlane(UV0.X, UV0.Y, 1, 0),
+			FPlane(0, 0, 0, 1)
+		);
+
+		// Use InverseSlow to catch singular matrices.  Inverse can miss this sometimes.
+		const FMatrix TextureToLocal = ParameterToTexture.Inverse() * ParameterToLocal;
+
+		Binormal = (TextureToLocal.TransformVector(FVector(1, 0, 0)).GetSafeNormal());
+		Tangent = (TextureToLocal.TransformVector(FVector(0, 1, 0)).GetSafeNormal());
+
+		TransformHandler.TransformVector(Normal, Transform);
+		TransformHandler.TransformVector(Binormal, Transform);
+		TransformHandler.TransformVector(Tangent, Transform);
+		
+		FVector Vel = (Pos - Prev) * InvDt;
+
+		*OutPosX.GetDest() = Pos.X;		*OutPosY.GetDest() = Pos.Y;		*OutPosZ.GetDest() = Pos.Z;
+		*OutVelX.GetDest() = Vel.X;		*OutVelY.GetDest() = Vel.Y;		*OutVelZ.GetDest() = Vel.Z;
+		*OutNormX.GetDest() = Normal.X;	*OutNormY.GetDest() = Normal.Y;	*OutNormZ.GetDest() = Normal.Z;
+		*OutBiNormX.GetDest() = Binormal.X;	*OutBiNormY.GetDest() = Binormal.Y;	*OutBiNormZ.GetDest() = Binormal.Z;
+		*OutTangentX.GetDest() = Tangent.X;	*OutTangentY.GetDest() = Tangent.Y;	*OutTangentZ.GetDest() = Tangent.Z;
+		
+		// We already got it to compute NBT's, let's just compute the UV to save time later.
+		FVector2D UV = BarycentricInterpolate(BaryXParam.Get(), BaryYParam.Get(), BaryZParam.Get(), UV0, UV1, UV2);
+		*OutUVX.GetDest() = UV.X;
+		*OutUVY.GetDest() = UV.Y;
+
+		// 		//TODO: Use the velocity temp buffer as storage for the prev position so I can SIMD the velocity calculation.
+		// 		*OutVelX.GetDest() = Prev.X;		*OutVelY.GetDest() = Prev.Y;		*OutVelZ.GetDest() = Prev.Z;
+
+		UVSetParam.Advance();
+		TriParam.Advance();
+		BaryXParam.Advance();	BaryYParam.Advance();		BaryZParam.Advance();
+		OutPosX.Advance();		OutPosY.Advance();			OutPosZ.Advance();
+		OutVelX.Advance();		OutVelY.Advance();			OutVelZ.Advance();
+		OutNormX.Advance();		OutNormY.Advance();			OutNormZ.Advance();
+		OutBiNormX.Advance();		OutBiNormY.Advance();			OutBiNormZ.Advance();
+		OutTangentX.Advance();		OutTangentY.Advance();			OutTangentZ.Advance();
+		OutUVX.Advance();		OutUVY.Advance();
+	}
+}
+
 template<typename TriType, typename BaryXType, typename BaryYType, typename BaryZType>
 void UNiagaraDataInterfaceSkeletalMesh::GetTriCoordColor(FVectorVMContext& Context)
 {
+	SCOPE_CYCLE_COUNTER(STAT_NiagaraSkel_Sample);
 	TriType TriParam(Context);
 	BaryXType BaryXParam(Context);
 	BaryYType BaryYParam(Context);
@@ -1375,9 +2012,12 @@ void UNiagaraDataInterfaceSkeletalMesh::GetTriCoordColor(FVectorVMContext& Conte
 
 	FMultiSizeIndexContainer& Indices = LODData.MultiSizeIndexContainer;
 	const FRawStaticIndexBuffer16or32Interface* IndexBuffer = Indices.GetIndexBuffer();
+	int32 TriMax = IndexBuffer->Num() - 3;
 	for (int32 i = 0; i < Context.NumInstances; ++i)
 	{
 		int32 Tri = TriParam.Get();
+		Tri = FMath::Min(Tri, TriMax);
+
 		int32 Idx0 = IndexBuffer->Get(Tri);
 		int32 Idx1 = IndexBuffer->Get(Tri + 1);
 		int32 Idx2 = IndexBuffer->Get(Tri + 2);
@@ -1403,6 +2043,7 @@ void UNiagaraDataInterfaceSkeletalMesh::GetTriCoordColor(FVectorVMContext& Conte
 template<typename VertexAccessorType, typename TriType, typename BaryXType, typename BaryYType, typename BaryZType, typename UVSetType>
 void UNiagaraDataInterfaceSkeletalMesh::GetTriCoordUV(FVectorVMContext& Context)
 {
+	SCOPE_CYCLE_COUNTER(STAT_NiagaraSkel_Sample);
 	VertexAccessorType VertAccessor;
 	TriType TriParam(Context);
 	BaryXType BaryXParam(Context);	BaryYType BaryYParam(Context);	BaryZType BaryZParam(Context);
@@ -1420,10 +2061,13 @@ void UNiagaraDataInterfaceSkeletalMesh::GetTriCoordUV(FVectorVMContext& Context)
 
 	FMultiSizeIndexContainer& Indices = LODData.MultiSizeIndexContainer;
 	FRawStaticIndexBuffer16or32Interface* IndexBuffer = Indices.GetIndexBuffer();
+	int32 TriMax = IndexBuffer->Num() - 3;
 	float InvDt = 1.0f / InstData->DeltaSeconds;
 	for (int32 i = 0; i < Context.NumInstances; ++i)
 	{
 		int32 Tri = TriParam.Get();
+		Tri = FMath::Min(Tri, TriMax);
+
 		int32 Idx0 = IndexBuffer->Get(Tri);
 		int32 Idx1 = IndexBuffer->Get(Tri + 1);
 		int32 Idx2 = IndexBuffer->Get(Tri + 2);
diff --git a/Engine/Plugins/FX/Niagara/Source/Niagara/Private/NiagaraDataInterfaceSpline.cpp b/Engine/Plugins/FX/Niagara/Source/Niagara/Private/NiagaraDataInterfaceSpline.cpp
--- a/Engine/Plugins/FX/Niagara/Source/Niagara/Private/NiagaraDataInterfaceSpline.cpp
+++ b/Engine/Plugins/FX/Niagara/Source/Niagara/Private/NiagaraDataInterfaceSpline.cpp
@@ -220,83 +220,82 @@ void UNiagaraDataInterfaceSpline::GetFunctions(TArray<FNiagaraFunctionSignature>
 	}
 }
 
-DEFINE_NDI_FUNC_BINDER(UNiagaraDataInterfaceSpline, SampleSplinePositionByUnitDistance);
-DEFINE_NDI_FUNC_BINDER(UNiagaraDataInterfaceSpline, SampleSplineUpVectorByUnitDistance);
-DEFINE_NDI_FUNC_BINDER(UNiagaraDataInterfaceSpline, SampleSplineRightVectorByUnitDistance);
-DEFINE_NDI_FUNC_BINDER(UNiagaraDataInterfaceSpline, SampleSplineDirectionByUnitDistance);
-DEFINE_NDI_FUNC_BINDER(UNiagaraDataInterfaceSpline, SampleSplineTangentByUnitDistance);
-DEFINE_NDI_FUNC_BINDER(UNiagaraDataInterfaceSpline, FindClosestUnitDistanceFromPositionWS);
-
-FVMExternalFunction UNiagaraDataInterfaceSpline::GetVMExternalFunction(const FVMExternalFunctionBindingInfo& BindingInfo, void* InstanceData)
+DEFINE_NDI_RAW_FUNC_BINDER(UNiagaraDataInterfaceSpline, SampleSplinePositionByUnitDistance);
+DEFINE_NDI_RAW_FUNC_BINDER(UNiagaraDataInterfaceSpline, SampleSplineUpVectorByUnitDistance);
+DEFINE_NDI_RAW_FUNC_BINDER(UNiagaraDataInterfaceSpline, SampleSplineRightVectorByUnitDistance);
+DEFINE_NDI_RAW_FUNC_BINDER(UNiagaraDataInterfaceSpline, SampleSplineDirectionByUnitDistance);
+DEFINE_NDI_RAW_FUNC_BINDER(UNiagaraDataInterfaceSpline, SampleSplineTangentByUnitDistance);
+DEFINE_NDI_RAW_FUNC_BINDER(UNiagaraDataInterfaceSpline, FindClosestUnitDistanceFromPositionWS);
+
+void UNiagaraDataInterfaceSpline::GetVMExternalFunction(const FVMExternalFunctionBindingInfo& BindingInfo, void* InstanceData, FVMExternalFunction &OutFunc)
 {
 	if (BindingInfo.Name == SampleSplinePositionByUnitDistanceName)
 	{
 		check(BindingInfo.GetNumInputs() == 2 && BindingInfo.GetNumOutputs() == 3);
-		return TNDIExplicitBinder<FNDITransformHandlerNoop, TNDIParamBinder<0, float, NDI_FUNC_BINDER(UNiagaraDataInterfaceSpline, SampleSplinePositionByUnitDistance)>>::Bind(this, BindingInfo, InstanceData);
+		TNDIExplicitBinder<FNDITransformHandlerNoop, TNDIParamBinder<0, float, NDI_RAW_FUNC_BINDER(UNiagaraDataInterfaceSpline, SampleSplinePositionByUnitDistance)>>::Bind(this, BindingInfo, InstanceData, OutFunc);
 	}
 	else if (BindingInfo.Name == SampleSplinePositionByUnitDistanceWSName)
 	{
 		check(BindingInfo.GetNumInputs() == 2 && BindingInfo.GetNumOutputs() == 3);
-		return TNDIExplicitBinder<FNDITransformHandler, TNDIParamBinder<0, float, NDI_FUNC_BINDER(UNiagaraDataInterfaceSpline, SampleSplinePositionByUnitDistance)>>::Bind(this, BindingInfo, InstanceData);
+		TNDIExplicitBinder<FNDITransformHandler, TNDIParamBinder<0, float, NDI_RAW_FUNC_BINDER(UNiagaraDataInterfaceSpline, SampleSplinePositionByUnitDistance)>>::Bind(this, BindingInfo, InstanceData, OutFunc);
 	}
 	else if (BindingInfo.Name == SampleSplineUpVectorByUnitDistanceName)
 	{
 		check(BindingInfo.GetNumInputs() == 2 && BindingInfo.GetNumOutputs() == 3);
-		return TNDIExplicitBinder<FNDITransformHandlerNoop, TNDIParamBinder<0, float, NDI_FUNC_BINDER(UNiagaraDataInterfaceSpline, SampleSplineUpVectorByUnitDistance)>>::Bind(this, BindingInfo, InstanceData);
+		TNDIExplicitBinder<FNDITransformHandlerNoop, TNDIParamBinder<0, float, NDI_RAW_FUNC_BINDER(UNiagaraDataInterfaceSpline, SampleSplineUpVectorByUnitDistance)>>::Bind(this, BindingInfo, InstanceData, OutFunc);
 	}
 	else if (BindingInfo.Name == SampleSplineUpVectorByUnitDistanceWSName)
 	{
 		check(BindingInfo.GetNumInputs() == 2 && BindingInfo.GetNumOutputs() == 3);
-		return TNDIExplicitBinder<FNDITransformHandler, TNDIParamBinder<0, float, NDI_FUNC_BINDER(UNiagaraDataInterfaceSpline, SampleSplineUpVectorByUnitDistance)>>::Bind(this, BindingInfo, InstanceData);
+		TNDIExplicitBinder<FNDITransformHandler, TNDIParamBinder<0, float, NDI_RAW_FUNC_BINDER(UNiagaraDataInterfaceSpline, SampleSplineUpVectorByUnitDistance)>>::Bind(this, BindingInfo, InstanceData, OutFunc);
 	}
 	else if (BindingInfo.Name == SampleSplineDirectionByUnitDistanceName)
 	{
 		check(BindingInfo.GetNumInputs() == 2 && BindingInfo.GetNumOutputs() == 3);
-		return TNDIExplicitBinder<FNDITransformHandlerNoop, TNDIParamBinder<0, float, NDI_FUNC_BINDER(UNiagaraDataInterfaceSpline, SampleSplineDirectionByUnitDistance)>>::Bind(this, BindingInfo, InstanceData);
+		TNDIExplicitBinder<FNDITransformHandlerNoop, TNDIParamBinder<0, float, NDI_RAW_FUNC_BINDER(UNiagaraDataInterfaceSpline, SampleSplineDirectionByUnitDistance)>>::Bind(this, BindingInfo, InstanceData, OutFunc);
 	}
 	else if (BindingInfo.Name == SampleSplineDirectionByUnitDistanceWSName)
 	{
 		check(BindingInfo.GetNumInputs() == 2 && BindingInfo.GetNumOutputs() == 3);
-		return TNDIExplicitBinder<FNDITransformHandler, TNDIParamBinder<0, float, NDI_FUNC_BINDER(UNiagaraDataInterfaceSpline, SampleSplineDirectionByUnitDistance)>>::Bind(this, BindingInfo, InstanceData);
+		TNDIExplicitBinder<FNDITransformHandler, TNDIParamBinder<0, float, NDI_RAW_FUNC_BINDER(UNiagaraDataInterfaceSpline, SampleSplineDirectionByUnitDistance)>>::Bind(this, BindingInfo, InstanceData, OutFunc);
 	}
 	else if (BindingInfo.Name == SampleSplineRightVectorByUnitDistanceName)
 	{
 		check(BindingInfo.GetNumInputs() == 2 && BindingInfo.GetNumOutputs() == 3);
-		return TNDIExplicitBinder<FNDITransformHandlerNoop, TNDIParamBinder<0, float, NDI_FUNC_BINDER(UNiagaraDataInterfaceSpline, SampleSplineRightVectorByUnitDistance)>>::Bind(this, BindingInfo, InstanceData);
+		TNDIExplicitBinder<FNDITransformHandlerNoop, TNDIParamBinder<0, float, NDI_RAW_FUNC_BINDER(UNiagaraDataInterfaceSpline, SampleSplineRightVectorByUnitDistance)>>::Bind(this, BindingInfo, InstanceData, OutFunc);
 	}
 	else if (BindingInfo.Name == SampleSplineRightVectorByUnitDistanceWSName)
 	{
 		check(BindingInfo.GetNumInputs() == 2 && BindingInfo.GetNumOutputs() == 3);
-		return TNDIExplicitBinder<FNDITransformHandler, TNDIParamBinder<0, float, NDI_FUNC_BINDER(UNiagaraDataInterfaceSpline, SampleSplineRightVectorByUnitDistance)>>::Bind(this, BindingInfo, InstanceData);
+		TNDIExplicitBinder<FNDITransformHandler, TNDIParamBinder<0, float, NDI_RAW_FUNC_BINDER(UNiagaraDataInterfaceSpline, SampleSplineRightVectorByUnitDistance)>>::Bind(this, BindingInfo, InstanceData, OutFunc);
 	}
 	else if (BindingInfo.Name == SampleSplineTangentByUnitDistanceName)
 	{
 		check(BindingInfo.GetNumInputs() == 2 && BindingInfo.GetNumOutputs() == 3);
-		return TNDIExplicitBinder<FNDITransformHandlerNoop, TNDIParamBinder<0, float, NDI_FUNC_BINDER(UNiagaraDataInterfaceSpline, SampleSplineTangentByUnitDistance)>>::Bind(this, BindingInfo, InstanceData);
+		TNDIExplicitBinder<FNDITransformHandlerNoop, TNDIParamBinder<0, float, NDI_RAW_FUNC_BINDER(UNiagaraDataInterfaceSpline, SampleSplineTangentByUnitDistance)>>::Bind(this, BindingInfo, InstanceData, OutFunc);
 	}
 	else if (BindingInfo.Name == SampleSplineTangentByUnitDistanceWSName)
 	{
 		check(BindingInfo.GetNumInputs() == 2 && BindingInfo.GetNumOutputs() == 3);
-		return TNDIExplicitBinder<FNDITransformHandler, TNDIParamBinder<0, float, NDI_FUNC_BINDER(UNiagaraDataInterfaceSpline, SampleSplineTangentByUnitDistance)>>::Bind(this, BindingInfo, InstanceData);
+		TNDIExplicitBinder<FNDITransformHandler, TNDIParamBinder<0, float, NDI_RAW_FUNC_BINDER(UNiagaraDataInterfaceSpline, SampleSplineTangentByUnitDistance)>>::Bind(this, BindingInfo, InstanceData, OutFunc);
 	}
 	else if (BindingInfo.Name == FindClosestUnitDistanceFromPositionWSName)
 	{
 		check(BindingInfo.GetNumInputs() == 4 && BindingInfo.GetNumOutputs() == 1);
-		return TNDIParamBinder<0, float, TNDIParamBinder<1, float, TNDIParamBinder<2, float, NDI_FUNC_BINDER(UNiagaraDataInterfaceSpline, FindClosestUnitDistanceFromPositionWS)>>>::Bind(this, BindingInfo, InstanceData);
+		TNDIParamBinder<0, float, TNDIParamBinder<1, float, TNDIParamBinder<2, float, NDI_RAW_FUNC_BINDER(UNiagaraDataInterfaceSpline, FindClosestUnitDistanceFromPositionWS)>>>::Bind(this, BindingInfo, InstanceData, OutFunc);
 	}
 	else if (BindingInfo.Name == GetSplineLocalToWorldName)
 	{
 		check(BindingInfo.GetNumInputs() == 1 && BindingInfo.GetNumOutputs() == 16);
-		return FVMExternalFunction::CreateUObject(this, &UNiagaraDataInterfaceSpline::GetLocalToWorld);
+		OutFunc = FVMExternalFunction::CreateUObject(this, &UNiagaraDataInterfaceSpline::GetLocalToWorld);
 	}
 	else if (BindingInfo.Name == GetSplineLocalToWorldInverseTransposedName)
 	{
 		check(BindingInfo.GetNumInputs() == 1 && BindingInfo.GetNumOutputs() == 16);
-		return FVMExternalFunction::CreateUObject(this, &UNiagaraDataInterfaceSpline::GetLocalToWorldInverseTransposed);
+		OutFunc = FVMExternalFunction::CreateUObject(this, &UNiagaraDataInterfaceSpline::GetLocalToWorldInverseTransposed);
 	}
 
 	//check(0);
-	return FVMExternalFunction();
 }
 
 bool UNiagaraDataInterfaceSpline::CopyToInternal(UNiagaraDataInterface* Destination) const
