Copying //UE4/Dev-Niagara to //UE4/Dev-Main (Source: //UE4/Dev-Niagara @ 4074996)

diff --git a/Engine/Plugins/FX/Niagara/Source/Niagara/Private/NiagaraDataInterfaceStaticMesh.cpp b/Engine/Plugins/FX/Niagara/Source/Niagara/Private/NiagaraDataInterfaceStaticMesh.cpp
--- a/Engine/Plugins/FX/Niagara/Source/Niagara/Private/NiagaraDataInterfaceStaticMesh.cpp
+++ b/Engine/Plugins/FX/Niagara/Source/Niagara/Private/NiagaraDataInterfaceStaticMesh.cpp
@@ -59,6 +59,8 @@ bool FNDIStaticMesh_InstanceData::Init(UNiagaraDataInterfaceStaticMesh* Interfac
 	PrevTransform = FMatrix::Identity;
 	PrevTransformInverseTransposed = FMatrix::Identity;
 	DeltaSeconds = 0.0f;
+	ChangeId = Interface->ChangeId;
+
 	if (Interface->Source)
 	{
 		AStaticMeshActor* MeshActor = Cast<AStaticMeshActor>(Interface->Source);
@@ -183,7 +185,7 @@ bool FNDIStaticMesh_InstanceData::Init(UNiagaraDataInterfaceStaticMesh* Interfac
 	return true;
 }
 
-bool FNDIStaticMesh_InstanceData::ResetRequired()const
+bool FNDIStaticMesh_InstanceData::ResetRequired(UNiagaraDataInterfaceStaticMesh* Interface)const
 {
 	check(GetActualMesh());
 
@@ -193,6 +195,11 @@ bool FNDIStaticMesh_InstanceData::ResetRequired()const
 		return true;
 	}
 
+	if (Interface != nullptr && ChangeId != Interface->ChangeId)
+	{
+		return true;
+	}
+
 	bool bPrevAreaWeighted = bIsAreaWeightedSampling;
 	//bool bPrevVCSampling = bSupportingVertexColorSampling;//TODO: Vertex color filtering needs more work.
 	bool bReset = false;
@@ -206,7 +213,7 @@ bool FNDIStaticMesh_InstanceData::ResetRequired()const
 
 bool FNDIStaticMesh_InstanceData::Tick(UNiagaraDataInterfaceStaticMesh* Interface, FNiagaraSystemInstance* SystemInstance, float InDeltaSeconds)
 {
-	if (ResetRequired())
+	if (ResetRequired(Interface))
 	{
 		return true;
 	}
@@ -291,6 +298,7 @@ UNiagaraDataInterfaceStaticMesh::UNiagaraDataInterfaceStaticMesh(FObjectInitiali
 	: Super(ObjectInitializer)
 	, DefaultMesh(nullptr)
 	, Source(nullptr)	
+	, ChangeId(0)
 	//, bSupportingVertexColorSampling(0)//Vertex color filtering needs some more work.
 	//, bFilterInitialized(false)
 {
@@ -314,6 +322,12 @@ void UNiagaraDataInterfaceStaticMesh::PostInitProperties()
 	}
 }
 
+void UNiagaraDataInterfaceStaticMesh::PostEditChangeProperty(FPropertyChangedEvent& PropertyChangedEvent)
+{
+	Super::PostEditChangeProperty(PropertyChangedEvent);
+	ChangeId++;
+}
+
 #endif //WITH_EDITOR
 
 namespace StaticMeshHelpers
@@ -340,6 +354,9 @@ namespace StaticMeshHelpers
 	static const FName GetMeshLocalToWorldName("GetLocalToWorld");
 	static const FName GetMeshLocalToWorldInverseTransposedName("GetMeshLocalToWorldInverseTransposed");
 	static const FName GetMeshWorldVelocityName("GetWorldVelocity");
+
+	static const FName GetVertexPositionName("GetVertexPosition"); 
+	static const FName GetVertexPositionWSName("GetVertexPositionWS"); 
 };
 
 void UNiagaraDataInterfaceStaticMesh::GetFunctions(TArray<FNiagaraFunctionSignature>& OutFunctions)
@@ -574,6 +591,36 @@ void UNiagaraDataInterfaceStaticMesh::GetFunctions(TArray<FNiagaraFunctionSignat
 		//Sig.Owner = *GetFullName();
 		OutFunctions.Add(Sig);
 	}
+
+	{
+		FNiagaraFunctionSignature Sig;
+		Sig.Name = StaticMeshHelpers::GetVertexPositionName;
+		Sig.Inputs.Add(FNiagaraVariable(FNiagaraTypeDefinition(GetClass()), TEXT("StaticMesh")));
+		Sig.Inputs.Add(FNiagaraVariable(FNiagaraTypeDefinition::GetIntDef(), TEXT("Vertex")));
+		Sig.Outputs.Add(FNiagaraVariable(FNiagaraTypeDefinition::GetVec3Def(), TEXT("Position")));
+		Sig.bMemberFunction = true;
+		Sig.bRequiresContext = false;
+#if WITH_EDITORONLY_DATA
+		Sig.Description = LOCTEXT("GetVertexPositionDesc", "Returns the local space vertex position for the passed vertex.");
+#endif
+		//Sig.Owner = *GetFullName();
+		OutFunctions.Add(Sig);
+	} 
+	
+	{
+		FNiagaraFunctionSignature Sig;
+		Sig.Name = StaticMeshHelpers::GetVertexPositionWSName;
+		Sig.Inputs.Add(FNiagaraVariable(FNiagaraTypeDefinition(GetClass()), TEXT("StaticMesh")));
+		Sig.Inputs.Add(FNiagaraVariable(FNiagaraTypeDefinition::GetIntDef(), TEXT("Vertex")));
+		Sig.Outputs.Add(FNiagaraVariable(FNiagaraTypeDefinition::GetVec3Def(), TEXT("Position")));
+		Sig.bMemberFunction = true;
+		Sig.bRequiresContext = false;
+#if WITH_EDITORONLY_DATA
+		Sig.Description = LOCTEXT("GetVertexPositionWSDesc", "Returns the world space vertex position for the passed vertex.");
+#endif
+		//Sig.Owner = *GetFullName();
+		OutFunctions.Add(Sig);
+	}
 }
 
 //External function binder choosing between template specializations based on UsesAreaWeighting
@@ -581,17 +628,17 @@ template<typename NextBinder>
 struct TUsesAreaWeightingBinder
 {
 	template<typename... ParamTypes>
-	static FVMExternalFunction Bind(UNiagaraDataInterface* Interface, const FVMExternalFunctionBindingInfo& BindingInfo, void* InstanceData)
+	static void Bind(UNiagaraDataInterface* Interface, const FVMExternalFunctionBindingInfo& BindingInfo, void* InstanceData, FVMExternalFunction &OutFunc)
 	{
 		FNDIStaticMesh_InstanceData* InstData = (FNDIStaticMesh_InstanceData*)InstanceData;
 		UNiagaraDataInterfaceStaticMesh* MeshInterface = CastChecked<UNiagaraDataInterfaceStaticMesh>(Interface);
 		if (InstData->UsesAreaWeighting())
 		{
-			return NextBinder::template Bind<ParamTypes..., TIntegralConstant<bool, true>>(Interface, BindingInfo, InstanceData);
+			NextBinder::template Bind<ParamTypes..., TIntegralConstant<bool, true>>(Interface, BindingInfo, InstanceData, OutFunc);
 		}
 		else
 		{
-			return NextBinder::template Bind<ParamTypes..., TIntegralConstant<bool, false>>(Interface, BindingInfo, InstanceData);
+			NextBinder::template Bind<ParamTypes..., TIntegralConstant<bool, false>>(Interface, BindingInfo, InstanceData, OutFunc);
 		}
 	}
 };
@@ -616,7 +663,7 @@ template<typename NextBinder>
 struct TTypedMeshAccessorBinder
 {
 	template<typename... ParamTypes>
-	static FVMExternalFunction Bind(UNiagaraDataInterface* Interface, const FVMExternalFunctionBindingInfo& BindingInfo, void* InstanceData)
+	static void Bind(UNiagaraDataInterface* Interface, const FVMExternalFunctionBindingInfo& BindingInfo, void* InstanceData, FVMExternalFunction &OutFunc)
 	{
 		FNDIStaticMesh_InstanceData* InstData = (FNDIStaticMesh_InstanceData*)InstanceData;
 		UNiagaraDataInterfaceStaticMesh* MeshInterface = CastChecked<UNiagaraDataInterfaceStaticMesh>(Interface);
@@ -626,40 +673,42 @@ struct TTypedMeshAccessorBinder
 		{
 			if (Res.VertexBuffers.StaticMeshVertexBuffer.GetUseFullPrecisionUVs())
 			{
-				return NextBinder::template Bind<ParamTypes..., TTypedMeshVertexAccessor<EStaticMeshVertexTangentBasisType::HighPrecision, EStaticMeshVertexUVType::HighPrecision>>(Interface, BindingInfo, InstanceData);
+				NextBinder::template Bind<ParamTypes..., TTypedMeshVertexAccessor<EStaticMeshVertexTangentBasisType::HighPrecision, EStaticMeshVertexUVType::HighPrecision>>(Interface, BindingInfo, InstanceData, OutFunc);
 			}
 			else
 			{
-				return NextBinder::template Bind<ParamTypes..., TTypedMeshVertexAccessor<EStaticMeshVertexTangentBasisType::HighPrecision, EStaticMeshVertexUVType::Default>>(Interface, BindingInfo, InstanceData);
+				NextBinder::template Bind<ParamTypes..., TTypedMeshVertexAccessor<EStaticMeshVertexTangentBasisType::HighPrecision, EStaticMeshVertexUVType::Default>>(Interface, BindingInfo, InstanceData, OutFunc);
 			}
 		}
 		else
 		{
 			if (Res.VertexBuffers.StaticMeshVertexBuffer.GetUseFullPrecisionUVs())
 			{
-				return NextBinder::template Bind<ParamTypes..., TTypedMeshVertexAccessor<EStaticMeshVertexTangentBasisType::Default, EStaticMeshVertexUVType::HighPrecision>>(Interface, BindingInfo, InstanceData);
+				NextBinder::template Bind<ParamTypes..., TTypedMeshVertexAccessor<EStaticMeshVertexTangentBasisType::Default, EStaticMeshVertexUVType::HighPrecision>>(Interface, BindingInfo, InstanceData, OutFunc);
 			}
 			else
 			{
-				return NextBinder::template Bind<ParamTypes..., TTypedMeshVertexAccessor<EStaticMeshVertexTangentBasisType::Default, EStaticMeshVertexUVType::Default>>(Interface, BindingInfo, InstanceData);
+				NextBinder::template Bind<ParamTypes..., TTypedMeshVertexAccessor<EStaticMeshVertexTangentBasisType::Default, EStaticMeshVertexUVType::Default>>(Interface, BindingInfo, InstanceData, OutFunc);
 			}
 		}
 	}
 };
 
 //Final binders for all static mesh interface functions.
-DEFINE_NDI_FUNC_BINDER(UNiagaraDataInterfaceStaticMesh, RandomSection);
-DEFINE_NDI_FUNC_BINDER(UNiagaraDataInterfaceStaticMesh, RandomTriCoord);
-DEFINE_NDI_FUNC_BINDER(UNiagaraDataInterfaceStaticMesh, RandomTriCoordVertexColorFiltered);
-DEFINE_NDI_FUNC_BINDER(UNiagaraDataInterfaceStaticMesh, RandomTriCoordOnSection);
-DEFINE_NDI_FUNC_BINDER(UNiagaraDataInterfaceStaticMesh, GetTriCoordPosition);
-DEFINE_NDI_FUNC_BINDER(UNiagaraDataInterfaceStaticMesh, GetTriCoordNormal);
-DEFINE_NDI_FUNC_BINDER(UNiagaraDataInterfaceStaticMesh, GetTriCoordTangents);
-DEFINE_NDI_FUNC_BINDER(UNiagaraDataInterfaceStaticMesh, GetTriCoordColor);
-DEFINE_NDI_FUNC_BINDER(UNiagaraDataInterfaceStaticMesh, GetTriCoordUV);
-DEFINE_NDI_FUNC_BINDER(UNiagaraDataInterfaceStaticMesh, GetTriCoordPositionAndVelocity);
-
-FVMExternalFunction UNiagaraDataInterfaceStaticMesh::GetVMExternalFunction(const FVMExternalFunctionBindingInfo& BindingInfo, void* InstanceData)
+DEFINE_NDI_RAW_FUNC_BINDER(UNiagaraDataInterfaceStaticMesh, RandomSection);
+DEFINE_NDI_RAW_FUNC_BINDER(UNiagaraDataInterfaceStaticMesh, RandomTriCoord);
+DEFINE_NDI_RAW_FUNC_BINDER(UNiagaraDataInterfaceStaticMesh, RandomTriCoordVertexColorFiltered);
+DEFINE_NDI_RAW_FUNC_BINDER(UNiagaraDataInterfaceStaticMesh, RandomTriCoordOnSection);
+DEFINE_NDI_RAW_FUNC_BINDER(UNiagaraDataInterfaceStaticMesh, GetTriCoordPosition);
+DEFINE_NDI_RAW_FUNC_BINDER(UNiagaraDataInterfaceStaticMesh, GetTriCoordNormal);
+DEFINE_NDI_RAW_FUNC_BINDER(UNiagaraDataInterfaceStaticMesh, GetTriCoordTangents);
+DEFINE_NDI_RAW_FUNC_BINDER(UNiagaraDataInterfaceStaticMesh, GetTriCoordColor);
+DEFINE_NDI_RAW_FUNC_BINDER(UNiagaraDataInterfaceStaticMesh, GetTriCoordUV);
+DEFINE_NDI_RAW_FUNC_BINDER(UNiagaraDataInterfaceStaticMesh, GetTriCoordPositionAndVelocity);
+
+DEFINE_NDI_RAW_FUNC_BINDER(UNiagaraDataInterfaceStaticMesh, GetVertexPosition);
+
+void UNiagaraDataInterfaceStaticMesh::GetVMExternalFunction(const FVMExternalFunctionBindingInfo& BindingInfo, void* InstanceData, FVMExternalFunction &OutFunc)
 {
 	FNDIStaticMesh_InstanceData* InstData = (FNDIStaticMesh_InstanceData*)InstanceData;
 	check(InstData && InstData->Mesh && InstData->Component.IsValid());
@@ -668,117 +717,123 @@ FVMExternalFunction UNiagaraDataInterfaceStaticMesh::GetVMExternalFunction(const
 	bool bNeedsVertexColors = false;
 	bool bNeedsVertMain = true;//Assuming we always need this?
 	
-	FVMExternalFunction Function;
 	if (BindingInfo.Name == StaticMeshHelpers::RandomSectionName)
 	{
 		check(BindingInfo.GetNumInputs() == 1 && BindingInfo.GetNumOutputs() == 1);
-		Function = TUsesAreaWeightingBinder<NDI_FUNC_BINDER(UNiagaraDataInterfaceStaticMesh, RandomSection)>::Bind(this, BindingInfo, InstanceData);
+		TUsesAreaWeightingBinder<NDI_RAW_FUNC_BINDER(UNiagaraDataInterfaceStaticMesh, RandomSection)>::Bind(this, BindingInfo, InstanceData, OutFunc);
 	}
 	else if (BindingInfo.Name == StaticMeshHelpers::RandomTriCoordName)
 	{
 		check(BindingInfo.GetNumInputs() == 1 && BindingInfo.GetNumOutputs() == 4);
-		Function = TUsesAreaWeightingBinder<NDI_FUNC_BINDER(UNiagaraDataInterfaceStaticMesh, RandomTriCoord)>::Bind(this, BindingInfo, InstanceData);
+		TUsesAreaWeightingBinder<NDI_RAW_FUNC_BINDER(UNiagaraDataInterfaceStaticMesh, RandomTriCoord)>::Bind(this, BindingInfo, InstanceData, OutFunc);
 	}
 	//TODO: Vertex color filtering needs more work.
 	else if (BindingInfo.Name == StaticMeshHelpers::RandomTriCoordVCFilteredName)
 	{
 		InstData->InitVertexColorFiltering();
 		check(BindingInfo.GetNumInputs() == 3 && BindingInfo.GetNumOutputs() == 4);
-		Function = TNDIParamBinder<0, float, TNDIParamBinder<1, float, NDI_FUNC_BINDER(UNiagaraDataInterfaceStaticMesh, RandomTriCoordVertexColorFiltered)>>::Bind(this, BindingInfo, InstanceData);
+		TNDIParamBinder<0, float, TNDIParamBinder<1, float, NDI_RAW_FUNC_BINDER(UNiagaraDataInterfaceStaticMesh, RandomTriCoordVertexColorFiltered)>>::Bind(this, BindingInfo, InstanceData, OutFunc);
 	}	
 	else if (BindingInfo.Name == StaticMeshHelpers::RandomTriCoordOnSectionName)
 	{
 		check(BindingInfo.GetNumInputs() == 2 && BindingInfo.GetNumOutputs() == 4);
-		Function = TUsesAreaWeightingBinder<TNDIParamBinder<0, int32, NDI_FUNC_BINDER(UNiagaraDataInterfaceStaticMesh, RandomTriCoordOnSection)>>::Bind(this, BindingInfo, InstanceData);
+		TUsesAreaWeightingBinder<TNDIParamBinder<0, int32, NDI_RAW_FUNC_BINDER(UNiagaraDataInterfaceStaticMesh, RandomTriCoordOnSection)>>::Bind(this, BindingInfo, InstanceData, OutFunc);
 	}
-	if (BindingInfo.Name == StaticMeshHelpers::GetTriPositionName)
+	else if (BindingInfo.Name == StaticMeshHelpers::GetTriPositionName)
 	{
 		check(BindingInfo.GetNumInputs() == 5 && BindingInfo.GetNumOutputs() == 3);
 		bNeedsVertexPositions = true;
-		Function = TNDIExplicitBinder<FNDITransformHandlerNoop, TNDIParamBinder<0, int32, TNDIParamBinder<1, float, TNDIParamBinder<2, float, TNDIParamBinder<3, float, NDI_FUNC_BINDER(UNiagaraDataInterfaceStaticMesh, GetTriCoordPosition)>>>>>::Bind(this, BindingInfo, InstanceData);
+		TNDIExplicitBinder<FNDITransformHandlerNoop, TNDIParamBinder<0, int32, TNDIParamBinder<1, float, TNDIParamBinder<2, float, TNDIParamBinder<3, float, NDI_RAW_FUNC_BINDER(UNiagaraDataInterfaceStaticMesh, GetTriCoordPosition)>>>>>::Bind(this, BindingInfo, InstanceData, OutFunc);
 	}
 	else if (BindingInfo.Name == StaticMeshHelpers::GetTriPositionWSName)
 	{
 		check(BindingInfo.GetNumInputs() == 5 && BindingInfo.GetNumOutputs() == 3);
 		bNeedsVertexPositions = true;
-		Function = TNDIExplicitBinder<FNDITransformHandler, TNDIParamBinder<0, int32, TNDIParamBinder<1, float, TNDIParamBinder<2, float, TNDIParamBinder<3, float, NDI_FUNC_BINDER(UNiagaraDataInterfaceStaticMesh, GetTriCoordPosition)>>>>>::Bind(this, BindingInfo, InstanceData);
+		TNDIExplicitBinder<FNDITransformHandler, TNDIParamBinder<0, int32, TNDIParamBinder<1, float, TNDIParamBinder<2, float, TNDIParamBinder<3, float, NDI_RAW_FUNC_BINDER(UNiagaraDataInterfaceStaticMesh, GetTriCoordPosition)>>>>>::Bind(this, BindingInfo, InstanceData, OutFunc);
 	}
 	else if (BindingInfo.Name == StaticMeshHelpers::GetTriNormalName)
 	{
 		check(BindingInfo.GetNumInputs() == 6 && BindingInfo.GetNumOutputs() == 3);
 		bNeedsVertMain = true;
-		Function = TNDIExplicitBinder<FNDITransformHandlerNoop, TNDIParamBinder<0, int32, TNDIParamBinder<1, float, TNDIParamBinder<2, float, TNDIParamBinder<3, float, NDI_FUNC_BINDER(UNiagaraDataInterfaceStaticMesh, GetTriCoordNormal)>>>>>::Bind(this, BindingInfo, InstanceData);
+		TNDIExplicitBinder<FNDITransformHandlerNoop, TNDIParamBinder<0, int32, TNDIParamBinder<1, float, TNDIParamBinder<2, float, TNDIParamBinder<3, float, NDI_RAW_FUNC_BINDER(UNiagaraDataInterfaceStaticMesh, GetTriCoordNormal)>>>>>::Bind(this, BindingInfo, InstanceData, OutFunc);
 	}
 	else if (BindingInfo.Name == StaticMeshHelpers::GetTriNormalWSName)
 	{
 		check(BindingInfo.GetNumInputs() == 5 && BindingInfo.GetNumOutputs() == 3);
 		bNeedsVertMain = true;
-		Function = TNDIExplicitBinder<FNDITransformHandler, TNDIParamBinder<0, int32, TNDIParamBinder<1, float, TNDIParamBinder<2, float, TNDIParamBinder<3, float, NDI_FUNC_BINDER(UNiagaraDataInterfaceStaticMesh, GetTriCoordNormal)>>>>>::Bind(this, BindingInfo, InstanceData);
+		TNDIExplicitBinder<FNDITransformHandler, TNDIParamBinder<0, int32, TNDIParamBinder<1, float, TNDIParamBinder<2, float, TNDIParamBinder<3, float, NDI_RAW_FUNC_BINDER(UNiagaraDataInterfaceStaticMesh, GetTriCoordNormal)>>>>>::Bind(this, BindingInfo, InstanceData, OutFunc);
 	}
 	else if (BindingInfo.Name == StaticMeshHelpers::GetTriTangentsName)
 	{
 		check(BindingInfo.GetNumInputs() == 5 && BindingInfo.GetNumOutputs() == 9);
 		bNeedsVertMain = true;
-		Function = TTypedMeshAccessorBinder<TNDIExplicitBinder<FNDITransformHandlerNoop, TNDIParamBinder<0, int32, TNDIParamBinder<1, float, TNDIParamBinder<2, float, TNDIParamBinder<3, float, NDI_FUNC_BINDER(UNiagaraDataInterfaceStaticMesh, GetTriCoordTangents)>>>>>>::Bind(this, BindingInfo, InstanceData);
+		TTypedMeshAccessorBinder<TNDIExplicitBinder<FNDITransformHandlerNoop, TNDIParamBinder<0, int32, TNDIParamBinder<1, float, TNDIParamBinder<2, float, TNDIParamBinder<3, float, NDI_RAW_FUNC_BINDER(UNiagaraDataInterfaceStaticMesh, GetTriCoordTangents)>>>>>>::Bind(this, BindingInfo, InstanceData, OutFunc);
 	}
 	else if (BindingInfo.Name == StaticMeshHelpers::GetTriTangentsWSName)
 	{
 		check(BindingInfo.GetNumInputs() == 5 && BindingInfo.GetNumOutputs() == 9);
 		bNeedsVertMain = true;
-		Function = TTypedMeshAccessorBinder<TNDIExplicitBinder<FNDITransformHandler, TNDIParamBinder<0, int32, TNDIParamBinder<1, float, TNDIParamBinder<2, float, TNDIParamBinder<3, float, NDI_FUNC_BINDER(UNiagaraDataInterfaceStaticMesh, GetTriCoordTangents)>>>>>>::Bind(this, BindingInfo, InstanceData);
+		TTypedMeshAccessorBinder<TNDIExplicitBinder<FNDITransformHandler, TNDIParamBinder<0, int32, TNDIParamBinder<1, float, TNDIParamBinder<2, float, TNDIParamBinder<3, float, NDI_RAW_FUNC_BINDER(UNiagaraDataInterfaceStaticMesh, GetTriCoordTangents)>>>>>>::Bind(this, BindingInfo, InstanceData, OutFunc);
 	}
 	else if (BindingInfo.Name == StaticMeshHelpers::GetTriColorName)
 	{
 		check(BindingInfo.GetNumInputs() == 5 && BindingInfo.GetNumOutputs() == 4);
 		bNeedsVertexColors = true;
-		Function = TNDIParamBinder<0, int32, TNDIParamBinder<1, float, TNDIParamBinder<2, float, TNDIParamBinder<3, float, NDI_FUNC_BINDER(UNiagaraDataInterfaceStaticMesh, GetTriCoordColor)>>>>::Bind(this, BindingInfo, InstanceData);
+		TNDIParamBinder<0, int32, TNDIParamBinder<1, float, TNDIParamBinder<2, float, TNDIParamBinder<3, float, NDI_RAW_FUNC_BINDER(UNiagaraDataInterfaceStaticMesh, GetTriCoordColor)>>>>::Bind(this, BindingInfo, InstanceData, OutFunc);
 	}
 	else if (BindingInfo.Name == StaticMeshHelpers::GetTriUVName)
 	{
 		check(BindingInfo.GetNumInputs() == 6 && BindingInfo.GetNumOutputs() == 2);
 		bNeedsVertMain = true;
-		Function = TTypedMeshAccessorBinder<TNDIParamBinder<0, int32, TNDIParamBinder<1, float, TNDIParamBinder<2, float, TNDIParamBinder<3, float, TNDIParamBinder<4, int32, NDI_FUNC_BINDER(UNiagaraDataInterfaceStaticMesh, GetTriCoordUV)>>>>>>::Bind(this, BindingInfo, InstanceData);
+		TTypedMeshAccessorBinder<TNDIParamBinder<0, int32, TNDIParamBinder<1, float, TNDIParamBinder<2, float, TNDIParamBinder<3, float, TNDIParamBinder<4, int32, NDI_RAW_FUNC_BINDER(UNiagaraDataInterfaceStaticMesh, GetTriCoordUV)>>>>>>::Bind(this, BindingInfo, InstanceData, OutFunc);
 	}
 	else if (BindingInfo.Name == StaticMeshHelpers::GetTriPositionAndVelocityName)
 	{
 		check(BindingInfo.GetNumInputs() == 5 && BindingInfo.GetNumOutputs() == 6);
 		bNeedsVertMain = true;
 		bNeedsVertexPositions = true;
-		Function = TNDIParamBinder<0, int32, TNDIParamBinder<1, float, TNDIParamBinder<2, float, TNDIParamBinder<3, float, NDI_FUNC_BINDER(UNiagaraDataInterfaceStaticMesh, GetTriCoordPositionAndVelocity)>>>>::Bind(this, BindingInfo, InstanceData);
+		TNDIParamBinder<0, int32, TNDIParamBinder<1, float, TNDIParamBinder<2, float, TNDIParamBinder<3, float, NDI_RAW_FUNC_BINDER(UNiagaraDataInterfaceStaticMesh, GetTriCoordPositionAndVelocity)>>>>::Bind(this, BindingInfo, InstanceData, OutFunc);
 	}
 	else if (BindingInfo.Name == StaticMeshHelpers::GetMeshLocalToWorldName)
 	{
 		check(BindingInfo.GetNumInputs() == 1 && BindingInfo.GetNumOutputs() == 16);
-		Function = FVMExternalFunction::CreateUObject(this, &UNiagaraDataInterfaceStaticMesh::GetLocalToWorld);
+		OutFunc = FVMExternalFunction::CreateUObject(this, &UNiagaraDataInterfaceStaticMesh::GetLocalToWorld);
 	}
 	else if (BindingInfo.Name == StaticMeshHelpers::GetMeshLocalToWorldInverseTransposedName)
 	{
 		check(BindingInfo.GetNumInputs() == 1 && BindingInfo.GetNumOutputs() == 16);
-		Function = FVMExternalFunction::CreateUObject(this, &UNiagaraDataInterfaceStaticMesh::GetLocalToWorldInverseTransposed);
+		OutFunc = FVMExternalFunction::CreateUObject(this, &UNiagaraDataInterfaceStaticMesh::GetLocalToWorldInverseTransposed);
 	}
 	else if (BindingInfo.Name == StaticMeshHelpers::GetMeshWorldVelocityName)
 	{
 		check(BindingInfo.GetNumInputs() == 1 && BindingInfo.GetNumOutputs() == 3);
-		Function = FVMExternalFunction::CreateUObject(this, &UNiagaraDataInterfaceStaticMesh::GetWorldVelocity);
+		OutFunc = FVMExternalFunction::CreateUObject(this, &UNiagaraDataInterfaceStaticMesh::GetWorldVelocity);
+	}
+	else if (BindingInfo.Name == StaticMeshHelpers::GetVertexPositionName)
+	{
+		check(BindingInfo.GetNumInputs() == 2 && BindingInfo.GetNumOutputs() == 3);
+		bNeedsVertexPositions = true;
+		TNDIExplicitBinder<FNDITransformHandlerNoop, TNDIParamBinder<0, int32, NDI_RAW_FUNC_BINDER(UNiagaraDataInterfaceStaticMesh, GetVertexPosition)>>::Bind(this, BindingInfo, InstanceData, OutFunc);
+	}
+	else if (BindingInfo.Name == StaticMeshHelpers::GetVertexPositionWSName)
+	{
+		check(BindingInfo.GetNumInputs() == 2 && BindingInfo.GetNumOutputs() == 3);
+		bNeedsVertexPositions = true;
+		TNDIExplicitBinder<FNDITransformHandler, TNDIParamBinder<0, int32, NDI_RAW_FUNC_BINDER(UNiagaraDataInterfaceStaticMesh, GetVertexPosition)>>::Bind(this, BindingInfo, InstanceData, OutFunc);
 	}
 
 	if (bNeedsVertexPositions && !InstData->MeshHasPositions())
 	{
 		UE_LOG(LogNiagara, Log, TEXT("Static Mesh data interface is cannot run as it's reading position data on a mesh that does not provide it. - Mesh:%s  "), *InstData->Mesh->GetFullName());
-		Function = FVMExternalFunction();
 	}
 	if (bNeedsVertexColors && !InstData->MeshHasColors())
 	{
 		UE_LOG(LogNiagara, Log, TEXT("Static Mesh data interface is cannot run as it's reading color data on a mesh that does not provide it. - Mesh:%s  "), *InstData->Mesh->GetFullName());
-		Function = FVMExternalFunction();
 	}
 	if (bNeedsVertMain && !InstData->MeshHasVerts())
 	{
 		UE_LOG(LogNiagara, Log, TEXT("Static Mesh data interface is cannot run as it's reading vertex data on a mesh with no vertex data. - Mesh:%s  "), *InstData->Mesh->GetFullName());
-		Function = FVMExternalFunction();
 	}
-
-	return Function;
 }
 
 bool UNiagaraDataInterfaceStaticMesh::CopyToInternal(UNiagaraDataInterface* Destination) const
@@ -835,34 +890,55 @@ bool UNiagaraDataInterfaceStaticMesh::PerInstanceTick(void* PerInstanceData, FNi
 	return Inst->Tick(this, SystemInstance, InDeltaSeconds);
 }
 
+#if WITH_EDITOR	
+TArray<FNiagaraDataInterfaceError> UNiagaraDataInterfaceStaticMesh::GetErrors()
+{
+	TArray<FNiagaraDataInterfaceError> Errors;
+	if (Source == nullptr && DefaultMesh != nullptr && !DefaultMesh->bAllowCPUAccess)
+	{
+		FNiagaraDataInterfaceError CPUAccessNotAllowedError(FText::Format(LOCTEXT("CPUAccessNotAllowedError", "This mesh needs CPU access in order to be used properly.({0})"), FText::FromString(DefaultMesh->GetName())),
+			LOCTEXT("CPUAccessNotAllowedErrorSummary", "CPU access error"),
+			FNiagaraDataInterfaceFix::CreateLambda([=]()
+		{
+			DefaultMesh->Modify();
+			DefaultMesh->bAllowCPUAccess = true;
+			return true;
+		}));
+
+		Errors.Add(CPUAccessNotAllowedError);
+	}
+	return Errors;
+}
+#endif
+
 //RandomSection specializations.
 //Each combination for AreaWeighted and Section filtered.
 template<>
-FORCEINLINE int32 UNiagaraDataInterfaceStaticMesh::RandomSection<TIntegralConstant<bool, true>, true>(FStaticMeshLODResources& Res, FNDIStaticMesh_InstanceData* InstData)
+FORCEINLINE int32 UNiagaraDataInterfaceStaticMesh::RandomSection<TIntegralConstant<bool, true>, true>(FRandomStream& RandStream, FStaticMeshLODResources& Res, FNDIStaticMesh_InstanceData* InstData)
 {
 	checkSlow(InstData->GetValidSections().Num() > 0);
-	int32 Idx = InstData->GetAreaWeigtedSampler().GetEntryIndex(InstData->RandStream.GetFraction(), InstData->RandStream.GetFraction());
+	int32 Idx = InstData->GetAreaWeigtedSampler().GetEntryIndex(RandStream.GetFraction(), RandStream.GetFraction());
 	return InstData->GetValidSections()[Idx];
 }
 
 template<>
-FORCEINLINE int32 UNiagaraDataInterfaceStaticMesh::RandomSection<TIntegralConstant<bool, true>, false>(FStaticMeshLODResources& Res, FNDIStaticMesh_InstanceData* InstData)
+FORCEINLINE int32 UNiagaraDataInterfaceStaticMesh::RandomSection<TIntegralConstant<bool, true>, false>(FRandomStream& RandStream, FStaticMeshLODResources& Res, FNDIStaticMesh_InstanceData* InstData)
 {
-	return Res.AreaWeightedSampler.GetEntryIndex(InstData->RandStream.GetFraction(), InstData->RandStream.GetFraction());
+	return Res.AreaWeightedSampler.GetEntryIndex(RandStream.GetFraction(), RandStream.GetFraction());
 }
 
 template<>
-FORCEINLINE int32 UNiagaraDataInterfaceStaticMesh::RandomSection<TIntegralConstant<bool, false>, true>(FStaticMeshLODResources& Res, FNDIStaticMesh_InstanceData* InstData)
+FORCEINLINE int32 UNiagaraDataInterfaceStaticMesh::RandomSection<TIntegralConstant<bool, false>, true>(FRandomStream& RandStream, FStaticMeshLODResources& Res, FNDIStaticMesh_InstanceData* InstData)
 {
 	checkSlow(InstData->GetValidSections().Num() > 0);
-	int32 Idx = InstData->RandStream.RandRange(0, InstData->GetValidSections().Num() - 1);
+	int32 Idx = RandStream.RandRange(0, InstData->GetValidSections().Num() - 1);
 	return InstData->GetValidSections()[Idx];
 }
 
 template<>
-FORCEINLINE int32 UNiagaraDataInterfaceStaticMesh::RandomSection<TIntegralConstant<bool, false>, false>(FStaticMeshLODResources& Res, FNDIStaticMesh_InstanceData* InstData)
+FORCEINLINE int32 UNiagaraDataInterfaceStaticMesh::RandomSection<TIntegralConstant<bool, false>, false>(FRandomStream& RandStream, FStaticMeshLODResources& Res, FNDIStaticMesh_InstanceData* InstData)
 {
-	return InstData->RandStream.RandRange(0, Res.Sections.Num() - 1);
+	return RandStream.RandRange(0, Res.Sections.Num() - 1);
 }
 
 template<typename TAreaWeighted>
@@ -874,7 +950,7 @@ void UNiagaraDataInterfaceStaticMesh::RandomSection(FVectorVMContext& Context)
 	FStaticMeshLODResources& Res = InstData->Mesh->RenderData->LODResources[0];
 	for (int32 i = 0; i < Context.NumInstances; ++i)
 	{
-		*OutSection.GetDest() = RandomSection<TAreaWeighted, true>(Res, InstData);
+		*OutSection.GetDest() = RandomSection<TAreaWeighted, true>(Context.RandStream, Res, InstData);
 		OutSection.Advance();
 	}
 }
@@ -882,38 +958,38 @@ void UNiagaraDataInterfaceStaticMesh::RandomSection(FVectorVMContext& Context)
 //RandomTriIndex specializations.
 //Each combination for AreaWeighted and Section filtered.
 template<>
-FORCEINLINE int32 UNiagaraDataInterfaceStaticMesh::RandomTriIndex<TIntegralConstant<bool, true>, true>(FStaticMeshLODResources& Res, FNDIStaticMesh_InstanceData* InstData)
+FORCEINLINE int32 UNiagaraDataInterfaceStaticMesh::RandomTriIndex<TIntegralConstant<bool, true>, true>(FRandomStream& RandStream, FStaticMeshLODResources& Res, FNDIStaticMesh_InstanceData* InstData)
 {
-	int32 SecIdx = RandomSection<TIntegralConstant<bool, true>, true>(Res, InstData);
+	int32 SecIdx = RandomSection<TIntegralConstant<bool, true>, true>(RandStream, Res, InstData);
 	FStaticMeshSection&  Sec = Res.Sections[SecIdx];
-	int32 Tri = Res.AreaWeightedSectionSamplers[SecIdx].GetEntryIndex(InstData->RandStream.GetFraction(), InstData->RandStream.GetFraction());
+	int32 Tri = Res.AreaWeightedSectionSamplers[SecIdx].GetEntryIndex(RandStream.GetFraction(), RandStream.GetFraction());
 	return Sec.FirstIndex + Tri * 3;
 }
 
 template<>
-FORCEINLINE int32 UNiagaraDataInterfaceStaticMesh::RandomTriIndex<TIntegralConstant<bool, true>, false>(FStaticMeshLODResources& Res, FNDIStaticMesh_InstanceData* InstData)
+FORCEINLINE int32 UNiagaraDataInterfaceStaticMesh::RandomTriIndex<TIntegralConstant<bool, true>, false>(FRandomStream& RandStream, FStaticMeshLODResources& Res, FNDIStaticMesh_InstanceData* InstData)
 {
-	int32 SecIdx = RandomSection<TIntegralConstant<bool, true>, false>(Res, InstData);
+	int32 SecIdx = RandomSection<TIntegralConstant<bool, true>, false>(RandStream, Res, InstData);
 	FStaticMeshSection&  Sec = Res.Sections[SecIdx];
-	int32 Tri = Res.AreaWeightedSectionSamplers[SecIdx].GetEntryIndex(InstData->RandStream.GetFraction(), InstData->RandStream.GetFraction());
+	int32 Tri = Res.AreaWeightedSectionSamplers[SecIdx].GetEntryIndex(RandStream.GetFraction(), RandStream.GetFraction());
 	return Sec.FirstIndex + Tri * 3;
 }
 
 template<>
-FORCEINLINE int32 UNiagaraDataInterfaceStaticMesh::RandomTriIndex<TIntegralConstant<bool, false>, true>(FStaticMeshLODResources& Res, FNDIStaticMesh_InstanceData* InstData)
+FORCEINLINE int32 UNiagaraDataInterfaceStaticMesh::RandomTriIndex<TIntegralConstant<bool, false>, true>(FRandomStream& RandStream, FStaticMeshLODResources& Res, FNDIStaticMesh_InstanceData* InstData)
 {
-	int32 SecIdx = RandomSection<TIntegralConstant<bool, false>, true>(Res, InstData);
+	int32 SecIdx = RandomSection<TIntegralConstant<bool, false>, true>(RandStream, Res, InstData);
 	FStaticMeshSection&  Sec = Res.Sections[SecIdx];
-	int32 Tri = InstData->RandStream.RandRange(0, Sec.NumTriangles - 1);
+	int32 Tri = RandStream.RandRange(0, Sec.NumTriangles - 1);
 	return Sec.FirstIndex + Tri * 3;	
 }
 
 template<>
-FORCEINLINE int32 UNiagaraDataInterfaceStaticMesh::RandomTriIndex<TIntegralConstant<bool, false>, false>(FStaticMeshLODResources& Res, FNDIStaticMesh_InstanceData* InstData)
+FORCEINLINE int32 UNiagaraDataInterfaceStaticMesh::RandomTriIndex<TIntegralConstant<bool, false>, false>(FRandomStream& RandStream, FStaticMeshLODResources& Res, FNDIStaticMesh_InstanceData* InstData)
 {
-	int32 SecIdx = RandomSection<TIntegralConstant<bool, false>, false>(Res, InstData);
+	int32 SecIdx = RandomSection<TIntegralConstant<bool, false>, false>(RandStream, Res, InstData);
 	FStaticMeshSection&  Sec = Res.Sections[SecIdx];
-	int32 Tri = InstData->RandStream.RandRange(0, Sec.NumTriangles - 1);
+	int32 Tri = RandStream.RandRange(0, Sec.NumTriangles - 1);
 	return Sec.FirstIndex + Tri * 3;
 }
 
@@ -931,8 +1007,8 @@ void UNiagaraDataInterfaceStaticMesh::RandomTriCoord(FVectorVMContext& Context)
 	FIndexArrayView Indices = Res.IndexBuffer.GetArrayView();
 	for (int32 i = 0; i < Context.NumInstances; ++i)
 	{
-		*OutTri.GetDest() = RandomTriIndex<TAreaWeighted, true>(Res, InstData);
-		FVector Bary = RandomBarycentricCoord(InstData->RandStream);
+		*OutTri.GetDest() = RandomTriIndex<TAreaWeighted, true>(Context.RandStream, Res, InstData);
+		FVector Bary = RandomBarycentricCoord(Context.RandStream);
 		*OutBaryX.GetDest() = Bary.X;
 		*OutBaryY.GetDest() = Bary.Y;
 		*OutBaryZ.GetDest() = Bary.Z;
@@ -990,12 +1066,12 @@ void UNiagaraDataInterfaceStaticMesh::RandomTriCoordVertexColorFiltered(FVectorV
 		}
 
 		// Select a random triangle from the list.
-		uint32 RandomTri = InstData->RandStream.GetFraction()*NumTris;
+		uint32 RandomTri = Context.RandStream.GetFraction()*NumTris;
 
 		// Now emit that triangle...
 		*OutTri.GetDest() = VCFData->TrianglesSortedByVertexColor[VCFData->VertexColorToTriangleStart[StartIdx] + RandomTri];
 
-		FVector Bary = RandomBarycentricCoord(InstData->RandStream);
+		FVector Bary = RandomBarycentricCoord(Context.RandStream);
 		*OutBaryX.GetDest() = Bary.X;
 		*OutBaryY.GetDest() = Bary.Y;
 		*OutBaryZ.GetDest() = Bary.Z;
@@ -1010,16 +1086,16 @@ void UNiagaraDataInterfaceStaticMesh::RandomTriCoordVertexColorFiltered(FVectorV
 }
 
 template<>
-FORCEINLINE int32 UNiagaraDataInterfaceStaticMesh::RandomTriIndexOnSection<TIntegralConstant<bool, true>>(FStaticMeshLODResources& Res, int32 SecIdx, FNDIStaticMesh_InstanceData* InstData)
+FORCEINLINE int32 UNiagaraDataInterfaceStaticMesh::RandomTriIndexOnSection<TIntegralConstant<bool, true>>(FRandomStream& RandStream, FStaticMeshLODResources& Res, int32 SecIdx, FNDIStaticMesh_InstanceData* InstData)
 {
-	return Res.AreaWeightedSectionSamplers[SecIdx].GetEntryIndex(InstData->RandStream.GetFraction(), InstData->RandStream.GetFraction());
+	return Res.AreaWeightedSectionSamplers[SecIdx].GetEntryIndex(RandStream.GetFraction(), RandStream.GetFraction());
 }
 
 template<>
-FORCEINLINE int32 UNiagaraDataInterfaceStaticMesh::RandomTriIndexOnSection<TIntegralConstant<bool, false>>(FStaticMeshLODResources& Res, int32 SecIdx, FNDIStaticMesh_InstanceData* InstData)
+FORCEINLINE int32 UNiagaraDataInterfaceStaticMesh::RandomTriIndexOnSection<TIntegralConstant<bool, false>>(FRandomStream& RandStream, FStaticMeshLODResources& Res, int32 SecIdx, FNDIStaticMesh_InstanceData* InstData)
 {
 	FStaticMeshSection&  Sec = Res.Sections[SecIdx];
-	int32 Tri = InstData->RandStream.RandRange(0, Sec.NumTriangles - 1);
+	int32 Tri = RandStream.RandRange(0, Sec.NumTriangles - 1);
 	return Sec.FirstIndex + Tri * 3;
 }
 
@@ -1039,8 +1115,8 @@ void UNiagaraDataInterfaceStaticMesh::RandomTriCoordOnSection(FVectorVMContext&
 	for (int32 i = 0; i < Context.NumInstances; ++i)
 	{
 		int32 SecIdx = SectionIdxParam.Get();
-		*OutTri.GetDest() = RandomTriIndexOnSection<TAreaWeighted>(Res, SecIdx, InstData);
-		FVector Bary = RandomBarycentricCoord(InstData->RandStream);
+		*OutTri.GetDest() = RandomTriIndexOnSection<TAreaWeighted>(Context.RandStream, Res, SecIdx, InstData);
+		FVector Bary = RandomBarycentricCoord(Context.RandStream);
 		*OutBaryX.GetDest() = Bary.X;
 		*OutBaryY.GetDest() = Bary.Y;
 		*OutBaryZ.GetDest() = Bary.Z;
@@ -1414,6 +1490,34 @@ void UNiagaraDataInterfaceStaticMesh::GetWorldVelocity(FVectorVMContext& Context
 	}
 }
 
+template<typename TransformHandlerType, typename VertexIndexType>
+void UNiagaraDataInterfaceStaticMesh::GetVertexPosition(FVectorVMContext& Context)
+{
+	TransformHandlerType TransformHandler;
+	VertexIndexType VertexIndexParam(Context);
+	FUserPtrHandler<FNDIStaticMesh_InstanceData> InstData(Context);
+
+	FRegisterHandler<float> OutPosX(Context);
+	FRegisterHandler<float> OutPosY(Context);
+	FRegisterHandler<float> OutPosZ(Context);
+
+	FStaticMeshLODResources& Res = InstData->Mesh->RenderData->LODResources[0];
+	const FPositionVertexBuffer& Positions = Res.VertexBuffers.PositionVertexBuffer;
+
+	const int32 NumVerts = Positions.GetNumVertices();
+	FVector Pos;
+	for (int32 i = 0; i < Context.NumInstances; i++)
+	{
+		int32 VertexIndex = VertexIndexParam.Get() % NumVerts;
+		Pos = Positions.VertexPosition(VertexIndex);		
+		TransformHandler.TransformPosition(Pos, InstData->Transform);
+		VertexIndexParam.Advance();
+		*OutPosX.GetDestAndAdvance() = Pos.X;
+		*OutPosY.GetDestAndAdvance() = Pos.Y;
+		*OutPosZ.GetDestAndAdvance() = Pos.Z;
+	}
+}
+
 //////////////////////////////////////////////////////////////////////////
 
 bool FDynamicVertexColorFilterData::Init(FNDIStaticMesh_InstanceData* Owner)
diff --git a/Engine/Plugins/FX/Niagara/Source/Niagara/Private/NiagaraDataInterfaceVector2DCurve.cpp b/Engine/Plugins/FX/Niagara/Source/Niagara/Private/NiagaraDataInterfaceVector2DCurve.cpp
--- a/Engine/Plugins/FX/Niagara/Source/Niagara/Private/NiagaraDataInterfaceVector2DCurve.cpp
+++ b/Engine/Plugins/FX/Niagara/Source/Niagara/Private/NiagaraDataInterfaceVector2DCurve.cpp
@@ -7,10 +7,11 @@
 #include "NiagaraTypes.h"
 #include "NiagaraCustomVersion.h"
 
-
 //////////////////////////////////////////////////////////////////////////
 //Vector2D Curve
 
+const FName UNiagaraDataInterfaceVector2DCurve::SampleCurveName("SampleVector2DCurve");
+
 UNiagaraDataInterfaceVector2DCurve::UNiagaraDataInterfaceVector2DCurve(FObjectInitializer const& ObjectInitializer)
 	: Super(ObjectInitializer)
 {
@@ -34,24 +35,28 @@ void UNiagaraDataInterfaceVector2DCurve::PostLoad()
 
 	const int32 NiagaraVer = GetLinkerCustomVersion(FNiagaraCustomVersion::GUID);
 
-	if (NiagaraVer < FNiagaraCustomVersion::CurveLUTNowOnByDefault)
+	if (NiagaraVer < FNiagaraCustomVersion::LatestVersion)
 	{
 		UpdateLUT();
 	}
-	TArray<float> OldLUT = ShaderLUT;
-	UpdateLUT();
-
-	if (!CompareLUTS(OldLUT))
+#if !UE_BUILD_SHIPPING
+	else
 	{
-		UE_LOG(LogNiagara, Log, TEXT("PostLoad LUT generation is out of sync. Please investigate. %s"), *GetPathName());
-	}
+		TArray<float> OldLUT = ShaderLUT;
+		UpdateLUT();
 
+		if (!CompareLUTS(OldLUT))
+		{
+			UE_LOG(LogNiagara, Log, TEXT("PostLoad LUT generation is out of sync. Please investigate. %s"), *GetPathName());
+		}
+	}
+#endif
 }
 
 void UNiagaraDataInterfaceVector2DCurve::UpdateLUT()
 {
 	ShaderLUT.Empty();
-	if (bAllowUnnormalizedLUT && (XCurve.GetNumKeys() > 0 || YCurve.GetNumKeys() > 0))
+	if ((XCurve.GetNumKeys() > 0 || YCurve.GetNumKeys() > 0))
 	{
 		LUTMinTime = FLT_MAX;
 		LUTMinTime = FMath::Min(XCurve.GetNumKeys() > 0 ? XCurve.GetFirstKey().Time : LUTMinTime, LUTMinTime);
@@ -71,7 +76,7 @@ void UNiagaraDataInterfaceVector2DCurve::UpdateLUT()
 
 	for (uint32 i = 0; i < CurveLUTWidth; i++)
 	{
-		float X = UnnormalizeTime(i / (float)CurveLUTWidth);
+		float X = UnnormalizeTime(i / (float)CurveLUTWidthMinusOne);
 		FVector2D C(XCurve.Eval(X), YCurve.Eval(X));
 		ShaderLUT.Add(C.X);
 		ShaderLUT.Add(C.Y);
@@ -79,26 +84,6 @@ void UNiagaraDataInterfaceVector2DCurve::UpdateLUT()
 	GPUBufferDirty = true;
 }
 
-#if WITH_EDITOR
-void UNiagaraDataInterfaceVector2DCurve::PostEditChangeProperty(struct FPropertyChangedEvent& PropertyChangedEvent)
-{
-	Super::PostEditChangeProperty(PropertyChangedEvent);
-
-	if (PropertyChangedEvent.Property && PropertyChangedEvent.Property->GetFName() == GET_MEMBER_NAME_CHECKED(UNiagaraDataInterfaceVector2DCurve, CurveToCopy))
-	{
-		UCurveVector* VectorCurveAsset = Cast<UCurveVector>(CurveToCopy.TryLoad());
-		if (VectorCurveAsset)
-		{
-			Modify();
-			XCurve = VectorCurveAsset->FloatCurves[0];
-			YCurve = VectorCurveAsset->FloatCurves[1];
-		}
-	}
-
-	UpdateLUT();
-}
-#endif
-
 bool UNiagaraDataInterfaceVector2DCurve::CopyToInternal(UNiagaraDataInterface* Destination) const
 {
 	if (!Super::CopyToInternal(Destination))
@@ -109,7 +94,10 @@ bool UNiagaraDataInterfaceVector2DCurve::CopyToInternal(UNiagaraDataInterface* D
 	DestinationVector2DCurve->XCurve = XCurve;
 	DestinationVector2DCurve->YCurve = YCurve;
 	DestinationVector2DCurve->UpdateLUT();
-	ensure(CompareLUTS(DestinationVector2DCurve->ShaderLUT));
+	if (!CompareLUTS(DestinationVector2DCurve->ShaderLUT))
+	{
+		UE_LOG(LogNiagara, Log, TEXT("Post CopyToInternal LUT generation is out of sync. Please investigate. %s"), *GetPathName());
+	}
 	return true;
 }
 
@@ -132,7 +120,7 @@ void UNiagaraDataInterfaceVector2DCurve::GetCurveData(TArray<FCurveData>& OutCur
 void UNiagaraDataInterfaceVector2DCurve::GetFunctions(TArray<FNiagaraFunctionSignature>& OutFunctions)
 {
 	FNiagaraFunctionSignature Sig;
-	Sig.Name = TEXT("SampleVector2DCurve");
+	Sig.Name = SampleCurveName;
 	Sig.bMemberFunction = true;
 	Sig.bRequiresContext = false;
 	Sig.Inputs.Add(FNiagaraVariable(FNiagaraTypeDefinition(GetClass()), TEXT("Vector2DCurve")));
@@ -143,93 +131,59 @@ void UNiagaraDataInterfaceVector2DCurve::GetFunctions(TArray<FNiagaraFunctionSig
 	OutFunctions.Add(Sig);
 }
 
-
-
 // build the shader function HLSL; function name is passed in, as it's defined per-DI; that way, configuration could change
 // the HLSL in the spirit of a static switch
 // TODO: need a way to identify each specific function here
 // 
-bool UNiagaraDataInterfaceVector2DCurve::GetFunctionHLSL(const FName& DefinitionFunctionName, FString InstanceFunctionName, TArray<FDIGPUBufferParamDescriptor> &Descriptors, FString &HLSLInterfaceID, FString &OutHLSL)
-{
-	FString BufferName = Descriptors[0].BufferParamName;
-	OutHLSL += TEXT("void ") + InstanceFunctionName + TEXT("(in float In_X, out float2 Out_Value) \n{\n");
-	OutHLSL += TEXT("\t Out_Value.x = ") + BufferName + TEXT("[(int)(In_X *") + FString::FromInt(CurveLUTWidth) + TEXT(")* 2 ];");
-	OutHLSL += TEXT("\t Out_Value.y = ") + BufferName + TEXT("[1+ (int)(In_X *") + FString::FromInt(CurveLUTWidth) + TEXT(")* 2 ];");
-	OutHLSL += TEXT("\n}\n");
-	return true;
-}
-
-// build buffer definition hlsl
-// 1. Choose a buffer name, add the data interface ID (important!)
-// 2. add a DIGPUBufferParamDescriptor to the array argument; that'll be passed on to the FNiagaraShader for binding to a shader param, that can
-// then later be found by name via FindDIBufferParam for setting; 
-// 3. store buffer declaration hlsl in OutHLSL
-// multiple buffers can be defined at once here
-//
-void UNiagaraDataInterfaceVector2DCurve::GetBufferDefinitionHLSL(FString DataInterfaceID, TArray<FDIGPUBufferParamDescriptor> &BufferDescriptors, FString &OutHLSL)
-{
-	FString BufferName = "CurveLUT" + DataInterfaceID;
-	OutHLSL += TEXT("Buffer<float> ") + BufferName + TEXT(";\n");
-
-	BufferDescriptors.Add(FDIGPUBufferParamDescriptor(BufferName, 0));		// add a descriptor for shader parameter binding
-
-}
+bool UNiagaraDataInterfaceVector2DCurve::GetFunctionHLSL(const FName& DefinitionFunctionName, FString InstanceFunctionName, FNiagaraDataInterfaceGPUParamInfo& ParamInfo, FString& OutHLSL)
+{
+	FString TimeToLUTFrac = TEXT("TimeToLUTFraction_") + ParamInfo.DataInterfaceHLSLSymbol;
+	FString Sample = TEXT("SampleCurve_") + ParamInfo.DataInterfaceHLSLSymbol;
+	OutHLSL += FString::Printf(TEXT("\
+void %s(in float In_X, out float2 Out_Value) \n\
+{ \n\
+	float RemappedX = %s(In_X) * %u; \n\
+	float Prev = floor(RemappedX); \n\
+	float Next = Prev < %u ? Prev + 1.0 : Prev; \n\
+	float Interp = RemappedX - Prev; \n\
+	Prev *= %u; \n\
+	Next *= %u; \n\
+	float2 A = float2(%s(Prev), %s(Prev + 1)); \n\
+	float2 B = float2(%s(Next), %s(Next + 1)); \n\
+	Out_Value = lerp(A, B, Interp); \n\
+}\n")
+, *InstanceFunctionName, *TimeToLUTFrac, CurveLUTWidthMinusOne, CurveLUTWidthMinusOne, CurveLUTNumElems, CurveLUTNumElems, *Sample, *Sample, *Sample, *Sample);
 
-// called after translate, to setup buffers matching the buffer descriptors generated during hlsl translation
-// need to do this because the script used during translate is a clone, including its DIs
-//
-void UNiagaraDataInterfaceVector2DCurve::SetupBuffers(FDIBufferDescriptorStore &BufferDescriptors)
-{
-	for (FDIGPUBufferParamDescriptor &Desc : BufferDescriptors.Descriptors)
-	{
-		FNiagaraDataInterfaceBufferData BufferData(*Desc.BufferParamName);	// store off the data for later use
-		GPUBuffers.Add(BufferData);
-	}
-}
-// return the GPU buffer array (called from NiagaraInstanceBatcher to get the buffers for setting to the shader)
-// we lazily update the buffer with a new LUT here if necessary
-//
-TArray<FNiagaraDataInterfaceBufferData> &UNiagaraDataInterfaceVector2DCurve::GetBufferDataArray()
-{
-	check(IsInRenderingThread());
-	if (GPUBufferDirty)
-	{
-		check(GPUBuffers.Num() > 0);
-
-		FNiagaraDataInterfaceBufferData &GPUBuffer = GPUBuffers[0];
-		GPUBuffer.Buffer.Release();
-		GPUBuffer.Buffer.Initialize(sizeof(float), CurveLUTWidth * 2, EPixelFormat::PF_R32_FLOAT, BUF_Static);	// always allocate for up to 64 data sets
-		uint32 BufferSize = ShaderLUT.Num() * sizeof(float);
-		int32 *BufferData = static_cast<int32*>(RHILockVertexBuffer(GPUBuffer.Buffer.Buffer, 0, BufferSize, EResourceLockMode::RLM_WriteOnly));
-		FPlatformMemory::Memcpy(BufferData, ShaderLUT.GetData(), BufferSize);
-		RHIUnlockVertexBuffer(GPUBuffer.Buffer.Buffer);
-		GPUBufferDirty = false;
-	}
-
-	return GPUBuffers;
+	return true;
 }
 
-DEFINE_NDI_FUNC_BINDER(UNiagaraDataInterfaceVector2DCurve, SampleCurve);
-FVMExternalFunction UNiagaraDataInterfaceVector2DCurve::GetVMExternalFunction(const FVMExternalFunctionBindingInfo& BindingInfo, void* InstanceData)
+DEFINE_NDI_RAW_FUNC_BINDER(UNiagaraDataInterfaceVector2DCurve, SampleCurve);
+void UNiagaraDataInterfaceVector2DCurve::GetVMExternalFunction(const FVMExternalFunctionBindingInfo& BindingInfo, void* InstanceData, FVMExternalFunction &OutFunc)
 {
-	if (BindingInfo.Name == TEXT("SampleVector2DCurve") && BindingInfo.GetNumInputs() == 1 && BindingInfo.GetNumOutputs() == 2)
+	if (BindingInfo.Name == SampleCurveName && BindingInfo.GetNumInputs() == 1 && BindingInfo.GetNumOutputs() == 2)
 	{
-		return TCurveUseLUTBinder<TNDIParamBinder<0, float, NDI_FUNC_BINDER(UNiagaraDataInterfaceVector2DCurve, SampleCurve)>>::Bind(this, BindingInfo, InstanceData);
+		TCurveUseLUTBinder<TNDIParamBinder<0, float, NDI_RAW_FUNC_BINDER(UNiagaraDataInterfaceVector2DCurve, SampleCurve)>>::Bind(this, BindingInfo, InstanceData, OutFunc);
 	}
 	else
 	{
 		UE_LOG(LogNiagara, Error, TEXT("Could not find data interface external function.\n\tExpected Name: SampleVector2DCurve  Actual Name: %s\n\tExpected Inputs: 1  Actual Inputs: %i\n\tExpected Outputs: 2  Actual Outputs: %i"),
 			*BindingInfo.Name.ToString(), BindingInfo.GetNumInputs(), BindingInfo.GetNumOutputs());
-		return FVMExternalFunction();
 	}
 }
 
 template<>
 FORCEINLINE_DEBUGGABLE FVector2D UNiagaraDataInterfaceVector2DCurve::SampleCurveInternal<TIntegralConstant<bool, true>>(float X)
 {
-	float NormalizedX = NormalizeTime(X);
-	int32 AccessIdx = FMath::Clamp<int32>(FMath::TruncToInt(NormalizedX * CurveLUTWidthMinusOne), 0, CurveLUTWidthMinusOne) * CurveLUTNumElems;
-	return FVector2D(ShaderLUT[AccessIdx], ShaderLUT[AccessIdx + 1]);
+	float RemappedX = FMath::Clamp(NormalizeTime(X) * CurveLUTWidthMinusOne, 0.0f, (float)CurveLUTWidthMinusOne);
+	float PrevEntry = FMath::TruncToFloat(RemappedX);
+	float NextEntry = PrevEntry < (float)CurveLUTWidthMinusOne ? PrevEntry + 1.0f : PrevEntry;
+	float Interp = RemappedX - PrevEntry;
+
+	int32 AIndex = PrevEntry * CurveLUTNumElems;
+	int32 BIndex = NextEntry * CurveLUTNumElems;
+	FVector2D A = FVector2D(ShaderLUT[AIndex], ShaderLUT[AIndex + 1]);
+	FVector2D B = FVector2D(ShaderLUT[BIndex], ShaderLUT[BIndex + 1]);
+	return FMath::Lerp(A, B, Interp);
 }
 
 template<>
@@ -248,12 +202,9 @@ void UNiagaraDataInterfaceVector2DCurve::SampleCurve(FVectorVMContext& Context)
 
 	for (int32 i = 0; i < Context.NumInstances; ++i)
 	{
-		float X = XParam.Get();
+		float X = XParam.GetAndAdvance();
 		FVector2D V = SampleCurveInternal<UseLUT>(X);
-		*OutSampleX.GetDest() = V.X;
-		*OutSampleY.GetDest() = V.Y;
-		XParam.Advance();
-		OutSampleX.Advance();
-		OutSampleY.Advance();
+		*OutSampleX.GetDestAndAdvance() = V.X;
+		*OutSampleY.GetDestAndAdvance() = V.Y;
 	}
 }
\ No newline at end of file
diff --git a/Engine/Plugins/FX/Niagara/Source/Niagara/Private/NiagaraDataInterfaceVector4Curve.cpp b/Engine/Plugins/FX/Niagara/Source/Niagara/Private/NiagaraDataInterfaceVector4Curve.cpp
--- a/Engine/Plugins/FX/Niagara/Source/Niagara/Private/NiagaraDataInterfaceVector4Curve.cpp
+++ b/Engine/Plugins/FX/Niagara/Source/Niagara/Private/NiagaraDataInterfaceVector4Curve.cpp
@@ -10,6 +10,8 @@
 //////////////////////////////////////////////////////////////////////////
 //Color Curve
 
+const FName UNiagaraDataInterfaceVector4Curve::SampleCurveName(TEXT("SampleColorCurve"));
+
 UNiagaraDataInterfaceVector4Curve::UNiagaraDataInterfaceVector4Curve(FObjectInitializer const& ObjectInitializer)
 	: Super(ObjectInitializer)
 {
@@ -35,49 +37,29 @@ void UNiagaraDataInterfaceVector4Curve::PostLoad()
 
 	const int32 NiagaraVer = GetLinkerCustomVersion(FNiagaraCustomVersion::GUID);
 
-	if (NiagaraVer < FNiagaraCustomVersion::CurveLUTNowOnByDefault)
+	if (NiagaraVer < FNiagaraCustomVersion::LatestVersion)
 	{
 		UpdateLUT();
 	}
-	TArray<float> OldLUT = ShaderLUT;
-	UpdateLUT();
-
-	if (!CompareLUTS(OldLUT))
+#if !UE_BUILD_SHIPPING
+	else
 	{
-		UE_LOG(LogNiagara, Log, TEXT("PostLoad LUT generation is out of sync. Please investigate. %s"), *GetPathName());
-	}
-
-}
-
-#if WITH_EDITOR
-
-void UNiagaraDataInterfaceVector4Curve::PostEditChangeProperty(struct FPropertyChangedEvent& PropertyChangedEvent)
-{
-	Super::PostEditChangeProperty(PropertyChangedEvent);
+		TArray<float> OldLUT = ShaderLUT;
+		UpdateLUT();
 
-	if (PropertyChangedEvent.Property && PropertyChangedEvent.Property->GetFName() == GET_MEMBER_NAME_CHECKED(UNiagaraDataInterfaceVector4Curve, CurveToCopy))
-	{
-		UCurveLinearColor* ColorCurveAsset = Cast<UCurveLinearColor>(CurveToCopy.TryLoad());
-		if(ColorCurveAsset)
+		if (!CompareLUTS(OldLUT))
 		{
-			Modify();
-			XCurve = ColorCurveAsset->FloatCurves[0];
-			YCurve = ColorCurveAsset->FloatCurves[1];
-			ZCurve = ColorCurveAsset->FloatCurves[2];
-			WCurve = ColorCurveAsset->FloatCurves[3];
+			UE_LOG(LogNiagara, Log, TEXT("PostLoad LUT generation is out of sync. Please investigate. %s"), *GetPathName());
 		}
 	}
-
-	UpdateLUT();
-}
-
 #endif
+}
 
 void UNiagaraDataInterfaceVector4Curve::UpdateLUT()
 {
 	ShaderLUT.Empty();
 
-	if (bAllowUnnormalizedLUT && (XCurve.GetNumKeys() > 0 || YCurve.GetNumKeys() > 0 || ZCurve.GetNumKeys() > 0 || WCurve.GetNumKeys() > 0))
+	if ((XCurve.GetNumKeys() > 0 || YCurve.GetNumKeys() > 0 || ZCurve.GetNumKeys() > 0 || WCurve.GetNumKeys() > 0))
 	{
 		LUTMinTime = FLT_MAX;
 		LUTMinTime = FMath::Min(XCurve.GetNumKeys() > 0 ? XCurve.GetFirstKey().Time : LUTMinTime, LUTMinTime);
@@ -101,7 +83,7 @@ void UNiagaraDataInterfaceVector4Curve::UpdateLUT()
 
 	for (uint32 i = 0; i < CurveLUTWidth; i++)
 	{
-		float X = UnnormalizeTime(i / (float)CurveLUTWidth);
+		float X = UnnormalizeTime(i / (float)CurveLUTWidthMinusOne);
 		FLinearColor C(XCurve.Eval(X), YCurve.Eval(X), ZCurve.Eval(X), WCurve.Eval(X));
 		ShaderLUT.Add(C.R);
 		ShaderLUT.Add(C.G);
@@ -152,7 +134,7 @@ void UNiagaraDataInterfaceVector4Curve::GetCurveData(TArray<FCurveData>& OutCurv
 void UNiagaraDataInterfaceVector4Curve::GetFunctions(TArray<FNiagaraFunctionSignature>& OutFunctions)
 {
 	FNiagaraFunctionSignature Sig;
-	Sig.Name = TEXT("SampleColorCurve");
+	Sig.Name = SampleCurveName;
 	Sig.bMemberFunction = true;
 	Sig.bRequiresContext = false;
 	Sig.Inputs.Add(FNiagaraVariable(FNiagaraTypeDefinition(GetClass()), TEXT("Vector4Curve")));
@@ -167,93 +149,56 @@ void UNiagaraDataInterfaceVector4Curve::GetFunctions(TArray<FNiagaraFunctionSign
 // the HLSL in the spirit of a static switch
 // TODO: need a way to identify each specific function here
 // 
-bool UNiagaraDataInterfaceVector4Curve::GetFunctionHLSL(const FName& DefinitionFunctionName, FString InstanceFunctionName, TArray<FDIGPUBufferParamDescriptor> &Descriptors, FString &HLSLInterfaceID, FString &OutHLSL)
-{
-	FString BufferName = Descriptors[0].BufferParamName;
-	OutHLSL += TEXT("void ") + InstanceFunctionName + TEXT("(in float In_X, out float4 Out_Value) \n{\n");
-	OutHLSL += TEXT("\t Out_Value.x = ") + BufferName + TEXT("[(int)(In_X *") + FString::FromInt(CurveLUTWidth) + TEXT(")* 4 ];");
-	OutHLSL += TEXT("\t Out_Value.y = ") + BufferName + TEXT("[1+ (int)(In_X *") + FString::FromInt(CurveLUTWidth) + TEXT(")* 4 ];");
-	OutHLSL += TEXT("\t Out_Value.z = ") + BufferName + TEXT("[2+ (int)(In_X *") + FString::FromInt(CurveLUTWidth) + TEXT(")* 4 ];");
-	OutHLSL += TEXT("\t Out_Value.w = ") + BufferName + TEXT("[3+ (int)(In_X *") + FString::FromInt(CurveLUTWidth) + TEXT(")* 4 ];");
-	OutHLSL += TEXT("\n}\n");
-	return true;
-}
+bool UNiagaraDataInterfaceVector4Curve::GetFunctionHLSL(const FName& DefinitionFunctionName, FString InstanceFunctionName, FNiagaraDataInterfaceGPUParamInfo& ParamInfo, FString& OutHLSL)
+{
+	FString TimeToLUTFrac = TEXT("TimeToLUTFraction_") + ParamInfo.DataInterfaceHLSLSymbol;
+	FString Sample = TEXT("SampleCurve_") + ParamInfo.DataInterfaceHLSLSymbol;
+	OutHLSL += FString::Printf(TEXT("\
+void %s(in float In_X, out float4 Out_Value) \n\
+{ \n\
+	float RemappedX = %s(In_X) * %u; \n\
+	float Prev = floor(RemappedX); \n\
+	float Next = Prev < %u ? Prev + 1.0 : Prev; \n\
+	float Interp = RemappedX - Prev; \n\
+	Prev *= %u; \n\
+	Next *= %u; \n\
+	float4 A = float4(%s(Prev), %s(Prev + 1), %s(Prev + 2), %s(Prev + 3)); \n\
+	float4 B = float4(%s(Next), %s(Next + 1), %s(Next + 2), %s(Next + 3)); \n\
+	Out_Value = lerp(A, B, Interp); \n\
+}\n")
+, *InstanceFunctionName, *TimeToLUTFrac, CurveLUTWidthMinusOne, CurveLUTWidthMinusOne, CurveLUTNumElems, CurveLUTNumElems
+, *Sample, *Sample, *Sample, *Sample, *Sample, *Sample, *Sample, *Sample);
 
-// build buffer definition hlsl
-// 1. Choose a buffer name, add the data interface ID (important!)
-// 2. add a DIGPUBufferParamDescriptor to the array argument; that'll be passed on to the FNiagaraShader for binding to a shader param, that can
-// then later be found by name via FindDIBufferParam for setting; 
-// 3. store buffer declaration hlsl in OutHLSL
-// multiple buffers can be defined at once here
-//
-void UNiagaraDataInterfaceVector4Curve::GetBufferDefinitionHLSL(FString DataInterfaceID, TArray<FDIGPUBufferParamDescriptor> &BufferDescriptors, FString &OutHLSL)
-{
-	FString BufferName = "CurveLUT" + DataInterfaceID;
-	OutHLSL += TEXT("Buffer<float> ") + BufferName + TEXT(";\n");
-
-	BufferDescriptors.Add( FDIGPUBufferParamDescriptor(BufferName, 0) );		// add a descriptor for shader parameter binding
-
-}
-
-// called after translate, to setup buffers matching the buffer descriptors generated during hlsl translation
-// need to do this because the script used during translate is a clone, including its DIs
-//
-void UNiagaraDataInterfaceVector4Curve::SetupBuffers(FDIBufferDescriptorStore &BufferDescriptors)
-{
-	for (FDIGPUBufferParamDescriptor &Desc : BufferDescriptors.Descriptors)
-	{
-		FNiagaraDataInterfaceBufferData BufferData(*Desc.BufferParamName);	// store off the data for later use
-		GPUBuffers.Add(BufferData);
-	}
-}
-
-// return the GPU buffer array (called from NiagaraInstanceBatcher to get the buffers for setting to the shader)
-// we lazily update the buffer with a new LUT here if necessary
-//
-TArray<FNiagaraDataInterfaceBufferData> &UNiagaraDataInterfaceVector4Curve::GetBufferDataArray()
-{
-	check(IsInRenderingThread());
-	if (GPUBufferDirty)
-	{
-		check(GPUBuffers.Num() > 0);
-
-		FNiagaraDataInterfaceBufferData &GPUBuffer = GPUBuffers[0];
-		GPUBuffer.Buffer.Release();
-		GPUBuffer.Buffer.Initialize(sizeof(float), CurveLUTWidth*4 , EPixelFormat::PF_R32_FLOAT, BUF_Static);	// always allocate for up to 64 data sets
-		uint32 BufferSize = ShaderLUT.Num() * sizeof(float);
-		int32 *BufferData = static_cast<int32*>(RHILockVertexBuffer(GPUBuffer.Buffer.Buffer, 0, BufferSize, EResourceLockMode::RLM_WriteOnly));
-		FPlatformMemory::Memcpy(BufferData, ShaderLUT.GetData(), BufferSize);
-		RHIUnlockVertexBuffer(GPUBuffer.Buffer.Buffer);
-		GPUBufferDirty = false;
-	}
-
-	return GPUBuffers;
+	return true;
 }
 
-
-
-
-DEFINE_NDI_FUNC_BINDER(UNiagaraDataInterfaceVector4Curve, SampleCurve);
-FVMExternalFunction UNiagaraDataInterfaceVector4Curve::GetVMExternalFunction(const FVMExternalFunctionBindingInfo& BindingInfo, void* InstanceData)
+DEFINE_NDI_RAW_FUNC_BINDER(UNiagaraDataInterfaceVector4Curve, SampleCurve);
+void UNiagaraDataInterfaceVector4Curve::GetVMExternalFunction(const FVMExternalFunctionBindingInfo& BindingInfo, void* InstanceData, FVMExternalFunction &OutFunc)
 {
-	if (BindingInfo.Name == TEXT("SampleColorCurve") && BindingInfo.GetNumInputs() == 1 && BindingInfo.GetNumOutputs() == 4)
+	if (BindingInfo.Name == SampleCurveName && BindingInfo.GetNumInputs() == 1 && BindingInfo.GetNumOutputs() == 4)
 	{
-		return TCurveUseLUTBinder<TNDIParamBinder<0, float, NDI_FUNC_BINDER(UNiagaraDataInterfaceVector4Curve, SampleCurve)>>::Bind(this, BindingInfo, InstanceData);
+		TCurveUseLUTBinder<TNDIParamBinder<0, float, NDI_RAW_FUNC_BINDER(UNiagaraDataInterfaceVector4Curve, SampleCurve)>>::Bind(this, BindingInfo, InstanceData, OutFunc);
 	}
 	else
 	{
 		UE_LOG(LogNiagara, Error, TEXT("Could not find data interface external function.\n\tExpected Name: SampleColorCurve  Actual Name: %s\n\tExpected Inputs: 1  Actual Inputs: %i\n\tExpected Outputs: 4  Actual Outputs: %i"),
 			*BindingInfo.Name.ToString(), BindingInfo.GetNumInputs(), BindingInfo.GetNumOutputs());
-		return FVMExternalFunction();
 	}
 }
 
 template<>
 FORCEINLINE_DEBUGGABLE FVector4 UNiagaraDataInterfaceVector4Curve::SampleCurveInternal<TIntegralConstant<bool, true>>(float X)
 {
-	float NormalizedX = NormalizeTime(X);
-	int32 AccessIdx = FMath::Clamp<int32>(FMath::TruncToInt(NormalizedX * CurveLUTWidthMinusOne), 0, CurveLUTWidthMinusOne) * CurveLUTNumElems;
-	return FVector4(ShaderLUT[AccessIdx], ShaderLUT[AccessIdx + 1], ShaderLUT[AccessIdx + 2], ShaderLUT[AccessIdx + 3]);
+	float RemappedX = FMath::Clamp(NormalizeTime(X) * CurveLUTWidthMinusOne, 0.0f, (float)CurveLUTWidthMinusOne);
+	float PrevEntry = FMath::TruncToFloat(RemappedX);
+	float NextEntry = PrevEntry < (float)CurveLUTWidthMinusOne ? PrevEntry + 1.0f : PrevEntry;
+	float Interp = RemappedX - PrevEntry;
+
+	int32 AIndex = PrevEntry * CurveLUTNumElems;
+	int32 BIndex = NextEntry * CurveLUTNumElems;
+	FVector4 A = FVector4(ShaderLUT[AIndex], ShaderLUT[AIndex + 1], ShaderLUT[AIndex + 2], ShaderLUT[AIndex + 3]);
+	FVector4 B = FVector4(ShaderLUT[BIndex], ShaderLUT[BIndex + 1], ShaderLUT[BIndex + 2], ShaderLUT[BIndex + 3]);
+	return FMath::Lerp(A, B, Interp);
 }
 
 template<>
@@ -275,17 +220,11 @@ void UNiagaraDataInterfaceVector4Curve::SampleCurve(FVectorVMContext& Context)
 
 	for (int32 i = 0; i < Context.NumInstances; ++i)
 	{
-		float X = XParam.Get();
+		float X = XParam.GetAndAdvance();
 		FVector4 Sample = SampleCurveInternal<UseLUT>(X);
-		FLinearColor C(Sample.X, Sample.Y, Sample.Z, Sample.W);
-		*SamplePtrR.GetDest() = C.R;
-		*SamplePtrG.GetDest() = C.G;
-		*SamplePtrB.GetDest() = C.B;
-		*SamplePtrA.GetDest() = C.A;
-		XParam.Advance();
-		SamplePtrR.Advance();
-		SamplePtrG.Advance();
-		SamplePtrB.Advance();
-		SamplePtrA.Advance();
+		*SamplePtrR.GetDestAndAdvance() = Sample.X;
+		*SamplePtrG.GetDestAndAdvance() = Sample.Y;
+		*SamplePtrB.GetDestAndAdvance() = Sample.Z;
+		*SamplePtrA.GetDestAndAdvance() = Sample.W;
 	}
 }
diff --git a/Engine/Plugins/FX/Niagara/Source/Niagara/Private/NiagaraDataInterfaceVectorCurve.cpp b/Engine/Plugins/FX/Niagara/Source/Niagara/Private/NiagaraDataInterfaceVectorCurve.cpp
--- a/Engine/Plugins/FX/Niagara/Source/Niagara/Private/NiagaraDataInterfaceVectorCurve.cpp
+++ b/Engine/Plugins/FX/Niagara/Source/Niagara/Private/NiagaraDataInterfaceVectorCurve.cpp
@@ -10,6 +10,8 @@
 //////////////////////////////////////////////////////////////////////////
 //Vector Curve
 
+const FName UNiagaraDataInterfaceVectorCurve::SampleCurveName(TEXT("SampleVectorCurve"));
+
 UNiagaraDataInterfaceVectorCurve::UNiagaraDataInterfaceVectorCurve(FObjectInitializer const& ObjectInitializer)
 	: Super(ObjectInitializer)
 {
@@ -35,47 +37,27 @@ void UNiagaraDataInterfaceVectorCurve::PostLoad()
 
 	const int32 NiagaraVer = GetLinkerCustomVersion(FNiagaraCustomVersion::GUID);
 
-	if (NiagaraVer < FNiagaraCustomVersion::CurveLUTNowOnByDefault)
+	if (NiagaraVer < FNiagaraCustomVersion::LatestVersion)
 	{
 		UpdateLUT();
 	}
-	TArray<float> OldLUT = ShaderLUT;
-	UpdateLUT();
-	if (!CompareLUTS(OldLUT))
-	{
-		UE_LOG(LogNiagara, Log, TEXT("PostLoad LUT generation is out of sync. Please investigate. %s"), *GetPathName());
-	}
-
-}
-
-#if WITH_EDITOR
-
-void UNiagaraDataInterfaceVectorCurve::PostEditChangeProperty(struct FPropertyChangedEvent& PropertyChangedEvent)
-{
-	Super::PostEditChangeProperty(PropertyChangedEvent);
-
-	if (PropertyChangedEvent.Property && PropertyChangedEvent.Property->GetFName() == GET_MEMBER_NAME_CHECKED(UNiagaraDataInterfaceVectorCurve, CurveToCopy))
+#if !UE_BUILD_SHIPPING
+	else
 	{
-		UCurveVector* VectorCurveAsset = Cast<UCurveVector>(CurveToCopy.TryLoad());
-		if (VectorCurveAsset)
+		TArray<float> OldLUT = ShaderLUT;
+		UpdateLUT();
+		if (!CompareLUTS(OldLUT))
 		{
-			Modify();
-			XCurve = VectorCurveAsset->FloatCurves[0];
-			YCurve = VectorCurveAsset->FloatCurves[1];
-			ZCurve = VectorCurveAsset->FloatCurves[2];
+			UE_LOG(LogNiagara, Log, TEXT("PostLoad LUT generation is out of sync. Please investigate. %s"), *GetPathName());
 		}
 	}
-	
-	UpdateLUT();
-}
-
 #endif
-
+}
 
 void UNiagaraDataInterfaceVectorCurve::UpdateLUT()
 {
 	ShaderLUT.Empty();
-	if (bAllowUnnormalizedLUT && (XCurve.GetNumKeys() > 0 || YCurve.GetNumKeys() > 0 || ZCurve.GetNumKeys() > 0))
+	if ((XCurve.GetNumKeys() > 0 || YCurve.GetNumKeys() > 0 || ZCurve.GetNumKeys() > 0))
 	{
 		LUTMinTime = FLT_MAX;
 		LUTMinTime = FMath::Min(XCurve.GetNumKeys() > 0 ? XCurve.GetFirstKey().Time : LUTMinTime, LUTMinTime);
@@ -97,7 +79,7 @@ void UNiagaraDataInterfaceVectorCurve::UpdateLUT()
 
 	for (uint32 i = 0; i < CurveLUTWidth; i++)
 	{
-		float X = UnnormalizeTime(i / (float)CurveLUTWidth);
+		float X = UnnormalizeTime(i / (float)CurveLUTWidthMinusOne);
 		FVector C(XCurve.Eval(X), YCurve.Eval(X), ZCurve.Eval(X));
 		ShaderLUT.Add(C.X);
 		ShaderLUT.Add(C.Y);
@@ -118,7 +100,10 @@ bool UNiagaraDataInterfaceVectorCurve::CopyToInternal(UNiagaraDataInterface* Des
 	DestinationVectorCurve->ZCurve = ZCurve;
 	DestinationVectorCurve->UpdateLUT();
 
-	ensure(CompareLUTS(DestinationVectorCurve->ShaderLUT));
+	if (!CompareLUTS(DestinationVectorCurve->ShaderLUT))
+	{
+		UE_LOG(LogNiagara, Log, TEXT("Post CopyToInternal LUT generation is out of sync. Please investigate. %s"), *GetPathName());
+	}
 	return true;
 }
 
@@ -144,7 +129,7 @@ void UNiagaraDataInterfaceVectorCurve::GetCurveData(TArray<FCurveData>& OutCurve
 void UNiagaraDataInterfaceVectorCurve::GetFunctions(TArray<FNiagaraFunctionSignature>& OutFunctions)
 {
 	FNiagaraFunctionSignature Sig;
-	Sig.Name = TEXT("SampleVectorCurve");
+	Sig.Name = SampleCurveName;
 	Sig.bMemberFunction = true;
 	Sig.bRequiresContext = false;
 	Sig.Inputs.Add(FNiagaraVariable(FNiagaraTypeDefinition(GetClass()), TEXT("VectorCurve")));
@@ -155,96 +140,60 @@ void UNiagaraDataInterfaceVectorCurve::GetFunctions(TArray<FNiagaraFunctionSigna
 	OutFunctions.Add(Sig);
 }
 
-
-
 // build the shader function HLSL; function name is passed in, as it's defined per-DI; that way, configuration could change
 // the HLSL in the spirit of a static switch
 // TODO: need a way to identify each specific function here
 // 
-bool UNiagaraDataInterfaceVectorCurve::GetFunctionHLSL(const FName& DefinitionFunctionName, FString InstanceFunctionName, TArray<FDIGPUBufferParamDescriptor> &Descriptors, FString &HLSLInterfaceID, FString &OutHLSL)
-{
-	FString BufferName = Descriptors[0].BufferParamName;
-	OutHLSL += TEXT("void ") + InstanceFunctionName + TEXT("(in float In_X, out float3 Out_Value) \n{\n");
-	OutHLSL += TEXT("\t Out_Value.x = ") + BufferName + TEXT("[(int)(In_X *") + FString::FromInt(CurveLUTWidth) + TEXT(")* 3 ];");
-	OutHLSL += TEXT("\t Out_Value.y = ") + BufferName + TEXT("[1+ (int)(In_X *") + FString::FromInt(CurveLUTWidth) + TEXT(")* 3 ];");
-	OutHLSL += TEXT("\t Out_Value.z = ") + BufferName + TEXT("[2+ (int)(In_X *") + FString::FromInt(CurveLUTWidth) + TEXT(")* 3 ];");
-	OutHLSL += TEXT("\n}\n");
-	return true;
-}
-
-// build buffer definition hlsl
-// 1. Choose a buffer name, add the data interface ID (important!)
-// 2. add a DIGPUBufferParamDescriptor to the array argument; that'll be passed on to the FNiagaraShader for binding to a shader param, that can
-// then later be found by name via FindDIBufferParam for setting; 
-// 3. store buffer declaration hlsl in OutHLSL
-// multiple buffers can be defined at once here
-//
-void UNiagaraDataInterfaceVectorCurve::GetBufferDefinitionHLSL(FString DataInterfaceID, TArray<FDIGPUBufferParamDescriptor> &BufferDescriptors, FString &OutHLSL)
-{
-	FString BufferName = "CurveLUT" + DataInterfaceID;
-	OutHLSL += TEXT("Buffer<float> ") + BufferName + TEXT(";\n");
-
-	BufferDescriptors.Add(FDIGPUBufferParamDescriptor(BufferName, 0));		// add a descriptor for shader parameter binding
+bool UNiagaraDataInterfaceVectorCurve::GetFunctionHLSL(const FName& DefinitionFunctionName, FString InstanceFunctionName, FNiagaraDataInterfaceGPUParamInfo& ParamInfo, FString& OutHLSL)
+{
+	FString TimeToLUTFrac = TEXT("TimeToLUTFraction_") + ParamInfo.DataInterfaceHLSLSymbol;
+	FString Sample = TEXT("SampleCurve_") + ParamInfo.DataInterfaceHLSLSymbol;
+	OutHLSL += FString::Printf(TEXT("\
+void %s(in float In_X, out float3 Out_Value) \n\
+{ \n\
+	float RemappedX = %s(In_X) * %u; \n\
+	float Prev = floor(RemappedX); \n\
+	float Next = Prev < %u ? Prev + 1.0 : Prev; \n\
+	float Interp = RemappedX - Prev; \n\
+	Prev *= %u; \n\
+	Next *= %u; \n\
+	float3 A = float3(%s(Prev), %s(Prev + 1), %s(Prev + 2)); \n\
+	float3 B = float3(%s(Next), %s(Next + 1), %s(Next + 2)); \n\
+	Out_Value = lerp(A, B, Interp); \n\
+}\n") 
+, *InstanceFunctionName, *TimeToLUTFrac, CurveLUTWidthMinusOne, CurveLUTWidthMinusOne, CurveLUTNumElems, CurveLUTNumElems
+, *Sample, *Sample, *Sample, *Sample, *Sample, *Sample);
 
+	return true;
 }
 
-// called after translate, to setup buffers matching the buffer descriptors generated during hlsl translation
-// need to do this because the script used during translate is a clone, including its DIs
-//
-void UNiagaraDataInterfaceVectorCurve::SetupBuffers(FDIBufferDescriptorStore &BufferDescriptors)
-{
-	for (FDIGPUBufferParamDescriptor &Desc : BufferDescriptors.Descriptors)
-	{
-		FNiagaraDataInterfaceBufferData BufferData(*Desc.BufferParamName);	// store off the data for later use
-		GPUBuffers.Add(BufferData);
-	}
-}
-
-// return the GPU buffer array (called from NiagaraInstanceBatcher to get the buffers for setting to the shader)
-// we lazily update the buffer with a new LUT here if necessary
-//
-TArray<FNiagaraDataInterfaceBufferData> &UNiagaraDataInterfaceVectorCurve::GetBufferDataArray()
-{
-	check(IsInRenderingThread());
-	if (GPUBufferDirty)
-	{
-		check(GPUBuffers.Num() > 0);
-
-		FNiagaraDataInterfaceBufferData &GPUBuffer = GPUBuffers[0];
-		GPUBuffer.Buffer.Release();
-		GPUBuffer.Buffer.Initialize(sizeof(float), CurveLUTWidth * 3, EPixelFormat::PF_R32_FLOAT, BUF_Static);	// always allocate for up to 64 data sets
-		uint32 BufferSize = ShaderLUT.Num() * sizeof(float);
-		int32 *BufferData = static_cast<int32*>(RHILockVertexBuffer(GPUBuffer.Buffer.Buffer, 0, BufferSize, EResourceLockMode::RLM_WriteOnly));
-		FPlatformMemory::Memcpy(BufferData, ShaderLUT.GetData(), BufferSize);
-		RHIUnlockVertexBuffer(GPUBuffer.Buffer.Buffer);
-		GPUBufferDirty = false;
-	}
-
-	return GPUBuffers;
-}
-
-
-DEFINE_NDI_FUNC_BINDER(UNiagaraDataInterfaceVectorCurve, SampleCurve);
-FVMExternalFunction UNiagaraDataInterfaceVectorCurve::GetVMExternalFunction(const FVMExternalFunctionBindingInfo& BindingInfo, void* InstanceData)
+DEFINE_NDI_RAW_FUNC_BINDER(UNiagaraDataInterfaceVectorCurve, SampleCurve);
+void UNiagaraDataInterfaceVectorCurve::GetVMExternalFunction(const FVMExternalFunctionBindingInfo& BindingInfo, void* InstanceData, FVMExternalFunction &OutFunc)
 {
-	if (BindingInfo.Name == TEXT("SampleVectorCurve") && BindingInfo.GetNumInputs() == 1 && BindingInfo.GetNumOutputs() == 3)
+	if (BindingInfo.Name == SampleCurveName && BindingInfo.GetNumInputs() == 1 && BindingInfo.GetNumOutputs() == 3)
 	{
-		return TCurveUseLUTBinder<TNDIParamBinder<0, float, NDI_FUNC_BINDER(UNiagaraDataInterfaceVectorCurve, SampleCurve)>>::Bind(this, BindingInfo, InstanceData);
+		TCurveUseLUTBinder<TNDIParamBinder<0, float, NDI_RAW_FUNC_BINDER(UNiagaraDataInterfaceVectorCurve, SampleCurve)>>::Bind(this, BindingInfo, InstanceData, OutFunc);
 	}
 	else
 	{
 		UE_LOG(LogNiagara, Error, TEXT("Could not find data interface external function.\n\tExpected Name: SampleVectorCurve  Actual Name: %s\n\tExpected Inputs: 1  Actual Inputs: %i\n\tExpected Outputs: 3  Actual Outputs: %i"),
 			*BindingInfo.Name.ToString(), BindingInfo.GetNumInputs(), BindingInfo.GetNumOutputs());
-		return FVMExternalFunction();
 	}
 }
 
 template<>
 FORCEINLINE_DEBUGGABLE FVector UNiagaraDataInterfaceVectorCurve::SampleCurveInternal<TIntegralConstant<bool, true>>(float X)
 {
-	float NormalizedX = NormalizeTime(X);
-	int32 AccessIdx = FMath::Clamp<int32>(FMath::TruncToInt(NormalizedX * CurveLUTWidthMinusOne), 0, CurveLUTWidthMinusOne) * CurveLUTNumElems;
-	return FVector(ShaderLUT[AccessIdx], ShaderLUT[AccessIdx + 1], ShaderLUT[AccessIdx + 2]);
+	float RemappedX = FMath::Clamp(NormalizeTime(X) * CurveLUTWidthMinusOne, 0.0f, (float)CurveLUTWidthMinusOne);
+	float PrevEntry = FMath::TruncToFloat(RemappedX);
+	float NextEntry = PrevEntry < (float)CurveLUTWidthMinusOne ? PrevEntry + 1.0f : PrevEntry;
+	float Interp = RemappedX - PrevEntry;
+
+	int32 AIndex = PrevEntry * CurveLUTNumElems;
+	int32 BIndex = NextEntry * CurveLUTNumElems;
+	FVector A = FVector(ShaderLUT[AIndex], ShaderLUT[AIndex + 1], ShaderLUT[AIndex + 2]);
+	FVector B = FVector(ShaderLUT[BIndex], ShaderLUT[BIndex + 1], ShaderLUT[BIndex + 2]);
+	return FMath::Lerp(A, B, Interp);
 }
 
 template<>
@@ -264,14 +213,10 @@ void UNiagaraDataInterfaceVectorCurve::SampleCurve(FVectorVMContext& Context)
 
 	for (int32 i = 0; i < Context.NumInstances; ++i)
 	{
-		float X = XParam.Get();
+		float X = XParam.GetAndAdvance();
 		FVector V = SampleCurveInternal<UseLUT>(X);
-		*OutSampleX.GetDest() = V.X;
-		*OutSampleY.GetDest() = V.Y;
-		*OutSampleZ.GetDest() = V.Z;
-		XParam.Advance();
-		OutSampleX.Advance();
-		OutSampleY.Advance();
-		OutSampleZ.Advance();
+		*OutSampleX.GetDestAndAdvance() = V.X;
+		*OutSampleY.GetDestAndAdvance() = V.Y;
+		*OutSampleZ.GetDestAndAdvance() = V.Z;
 	}
 }
diff --git a/Engine/Plugins/FX/Niagara/Source/Niagara/Private/NiagaraDataInterfaceVectorField.cpp b/Engine/Plugins/FX/Niagara/Source/Niagara/Private/NiagaraDataInterfaceVectorField.cpp
--- a/Engine/Plugins/FX/Niagara/Source/Niagara/Private/NiagaraDataInterfaceVectorField.cpp
+++ b/Engine/Plugins/FX/Niagara/Source/Niagara/Private/NiagaraDataInterfaceVectorField.cpp
@@ -5,9 +5,17 @@
 #include "VectorField/VectorFieldStatic.h"
 #include "VectorField/VectorFieldAnimated.h"
 #include "Math/Float16Color.h"
+#include "NiagaraShader.h"
+#include "ShaderParameterUtils.h"
+
 #define LOCTEXT_NAMESPACE "NiagaraDataInterfaceVectorField"
 
 
+const FString UNiagaraDataInterfaceVectorField::BufferBaseName(TEXT("VectorFieldBuffer_"));
+const FString UNiagaraDataInterfaceVectorField::DimentionsBaseName(TEXT("Dimentions_"));
+const FString UNiagaraDataInterfaceVectorField::BoundsMinBaseName(TEXT("BoundsMin_"));
+const FString UNiagaraDataInterfaceVectorField::BoundsMaxBaseName(TEXT("BoundsMax_"));
+
 #if 0
 //////////////////////////////////////////////////////////////////////////
 //FNDIVectorField_InstanceData
@@ -370,23 +378,21 @@ void UNiagaraDataInterfaceVectorField::GetFunctions(TArray<FNiagaraFunctionSigna
 	}
 }
 
-DEFINE_NDI_FUNC_BINDER(UNiagaraDataInterfaceVectorField, SampleVectorField);
-FVMExternalFunction UNiagaraDataInterfaceVectorField::GetVMExternalFunction(const FVMExternalFunctionBindingInfo& BindingInfo, void* InstanceData)
+DEFINE_NDI_RAW_FUNC_BINDER(UNiagaraDataInterfaceVectorField, SampleVectorField);
+void UNiagaraDataInterfaceVectorField::GetVMExternalFunction(const FVMExternalFunctionBindingInfo& BindingInfo, void* InstanceData, FVMExternalFunction &OutFunc)
 {
-	FVMExternalFunction Function;
 	if (BindingInfo.Name == SampleVectorFieldName && BindingInfo.GetNumInputs() == 3 && BindingInfo.GetNumOutputs() == 3)
 	{
-		Function = TNDIParamBinder<0, float, TNDIParamBinder<1, float, TNDIParamBinder<2, float, NDI_FUNC_BINDER(UNiagaraDataInterfaceVectorField, SampleVectorField)>>>::Bind(this, BindingInfo, InstanceData);
+		TNDIParamBinder<0, float, TNDIParamBinder<1, float, TNDIParamBinder<2, float, NDI_RAW_FUNC_BINDER(UNiagaraDataInterfaceVectorField, SampleVectorField)>>>::Bind(this, BindingInfo, InstanceData, OutFunc);
 	}
 	else if (BindingInfo.Name == GetVectorDimsName && BindingInfo.GetNumInputs() == 0 && BindingInfo.GetNumOutputs() == 3)
 	{
-		Function = FVMExternalFunction::CreateUObject(this, &UNiagaraDataInterfaceVectorField::GetFieldDimensions);
+		OutFunc = FVMExternalFunction::CreateUObject(this, &UNiagaraDataInterfaceVectorField::GetFieldDimensions);
 	}
 	else if (BindingInfo.Name == GetVectorFieldBoundsName && BindingInfo.GetNumInputs() == 0 && BindingInfo.GetNumOutputs() == 6)
 	{
-		Function = FVMExternalFunction::CreateUObject(this, &UNiagaraDataInterfaceVectorField::GetFieldBounds);
+		OutFunc = FVMExternalFunction::CreateUObject(this, &UNiagaraDataInterfaceVectorField::GetFieldBounds);
 	}
-	return Function;
 }
 
 void UNiagaraDataInterfaceVectorField::GetFieldDimensions(FVectorVMContext& Context)
@@ -785,12 +791,11 @@ bool UNiagaraDataInterfaceVectorField::PerInstanceTick(void* PerInstanceData, FN
 	return false;
 }
 
-bool UNiagaraDataInterfaceVectorField::GetFunctionHLSL(const FName& DefinitionFunctionName, FString InstanceFunctionName, TArray<FDIGPUBufferParamDescriptor> &Descriptors, FString &HLSLInterfaceID, FString &OutHLSL)
+bool UNiagaraDataInterfaceVectorField::GetFunctionHLSL(const FName& DefinitionFunctionName, FString InstanceFunctionName, FNiagaraDataInterfaceGPUParamInfo& ParamInfo, FString& OutHLSL)
 {
-	//FString BufferName = Descriptors[0].BufferParamName;
 	if (DefinitionFunctionName == SampleVectorFieldName)
 	{
-		FString BufferName = Descriptors[0].BufferParamName;
+		FString BufferName = "VectorFieldLUT_" + ParamInfo.DataInterfaceHLSLSymbol;
 		OutHLSL += TEXT("void ") + InstanceFunctionName + TEXT("(float3 In_SamplePoint,  out float3 Out_Value) \n{\n");
 		OutHLSL += TEXT("\t int3 IntCoord = int3(0,0,0);\n");
 		OutHLSL += TEXT("\t int3 Dimensions = int3(1,1,1);\n");
@@ -800,57 +805,36 @@ bool UNiagaraDataInterfaceVectorField::GetFunctionHLSL(const FName& DefinitionFu
 	}
 	else if (DefinitionFunctionName == GetVectorDimsName)
 	{
+		FString DimsVar = DimentionsBaseName + ParamInfo.DataInterfaceHLSLSymbol;
 		OutHLSL += TEXT("void ") + InstanceFunctionName + TEXT("(out float3 Out_Value) \n{\n");
-		OutHLSL += TEXT("\t Out_Value = float3(1.0f, 1.0f, 1.0f);\n");
+		OutHLSL += TEXT("\t Out_Value = ") + DimsVar + TEXT(";\n");
 		OutHLSL += TEXT("\n}\n");
 		return true;
 	}
 	else if (DefinitionFunctionName == GetVectorFieldBoundsName)
 	{
+		FString BoundsMinVar = BoundsMinBaseName + ParamInfo.DataInterfaceHLSLSymbol;
+		FString BoundsMaxVar = BoundsMaxBaseName + ParamInfo.DataInterfaceHLSLSymbol;
 		OutHLSL += TEXT("void ") + InstanceFunctionName + TEXT("(out float3 Out_MinBounds, out float3 Out_MaxBounds) \n{\n");
-		OutHLSL += TEXT("\t Out_MinBounds = Out_MaxBounds = float3(1.0f, 1.0f, 1.0f);\n");
+		OutHLSL += TEXT("\t Out_MinBounds = ") + BoundsMinVar + TEXT(";\n");
+		OutHLSL += TEXT("\t Out_MaxBounds = ") + BoundsMaxVar + TEXT(";\n");
 		OutHLSL += TEXT("\n}\n");
 		return true;
 	}
 	return false;
 }
 
-// build buffer definition hlsl
-// 1. Choose a buffer name, add the data interface ID (important!)
-// 2. add a DIGPUBufferParamDescriptor to the array argument; that'll be passed on to the FNiagaraShader for binding to a shader param, that can
-// then later be found by name via FindDIBufferParam for setting; 
-// 3. store buffer declaration hlsl in OutHLSL
-// multiple buffers can be defined at once here
-//
-void UNiagaraDataInterfaceVectorField::GetBufferDefinitionHLSL(FString DataInterfaceID, TArray<FDIGPUBufferParamDescriptor> &BufferDescriptors, FString &OutHLSL)
-{
-	FString BufferName = "VectorFieldLUT" + DataInterfaceID;
-	OutHLSL += TEXT("Buffer<float4> ") + BufferName + TEXT(";\n");
-
-	BufferDescriptors.Add(FDIGPUBufferParamDescriptor(BufferName, 0));		// add a descriptor for shader parameter binding
-}
-
-// called after translate, to setup buffers matching the buffer descriptors generated during hlsl translation
-// need to do this because the script used during translate is a clone, including its DIs
-//
-void UNiagaraDataInterfaceVectorField::SetupBuffers(FDIBufferDescriptorStore &BufferDescriptors)
+void UNiagaraDataInterfaceVectorField::GetParameterDefinitionHLSL(FNiagaraDataInterfaceGPUParamInfo& ParamInfo, FString& OutHLSL)
 {
-	GPUBuffers.Empty();
-	for (FDIGPUBufferParamDescriptor &Desc : BufferDescriptors.Descriptors)
-	{
-		FNiagaraDataInterfaceBufferData BufferData(*Desc.BufferParamName);	// store off the data for later use
-		GPUBuffers.Add(BufferData);
-	}
-	bGPUBufferDirty = true;
+	OutHLSL += TEXT("Buffer<float4> ") + BufferBaseName + ParamInfo.DataInterfaceHLSLSymbol + TEXT(";\n");
+	OutHLSL += TEXT("float3 ") + DimentionsBaseName + ParamInfo.DataInterfaceHLSLSymbol + TEXT(";\n");
 }
 
-// return the GPU buffer array (called from NiagaraInstanceBatcher to get the buffers for setting to the shader)
-// we lazily update the buffer with a new LUT here if necessary
-//
-TArray<FNiagaraDataInterfaceBufferData> &UNiagaraDataInterfaceVectorField::GetBufferDataArray()
+FRWBuffer& UNiagaraDataInterfaceVectorField::GetGPUBuffer()
 {
 	check(IsInRenderingThread());
 
+	//TODO: Doing this here is not really threadsafe. Need to move to an RT proxy style model?
 	if (bGPUBufferDirty)
 	{
 		UVectorFieldStatic* StaticVectorField = Cast<UVectorFieldStatic>(Field);
@@ -868,9 +852,7 @@ TArray<FNiagaraDataInterfaceBufferData> &UNiagaraDataInterfaceVectorField::GetBu
 			Depth = StaticVectorField->SizeZ;
 		}
 		
-		check(GPUBuffers.Num() > 0);
-		FNiagaraDataInterfaceBufferData &GPUBuffer = GPUBuffers[0];
-		GPUBuffer.Buffer.Release();
+		GPUBuffer.Release();
 		uint32 BufferSize = Width * Height * Depth * sizeof(float) * 4;
 		TResourceArray<FVector4> TempTable;
 		TempTable.AddUninitialized(Width*Height*Depth);
@@ -909,14 +891,58 @@ TArray<FNiagaraDataInterfaceBufferData> &UNiagaraDataInterfaceVectorField::GetBu
 			}
 		}
 
-		GPUBuffer.Buffer.Initialize(sizeof(float) * 4, Width*Height*Depth, EPixelFormat::PF_A32B32G32R32F, BUF_Static, TEXT("CurlnoiseTable"), &TempTable);
+		GPUBuffer.Initialize(sizeof(float) * 4, Width*Height*Depth, EPixelFormat::PF_A32B32G32R32F, BUF_Static, TEXT("CurlnoiseTable"), &TempTable);
 		//FPlatformMemory::Memcpy(BufferData, TempTable, BufferSize);
 		//RHIUnlockVertexBuffer(GPUBuffer.Buffer.Buffer);
 		bGPUBufferDirty = false;
 	}
-	return GPUBuffers;
+
+	return GPUBuffer;
 }
 
+struct FNiagaraDataInterfaceParametersCS_VectorField : public FNiagaraDataInterfaceParametersCS
+{
+	virtual void Bind(const FNiagaraDataInterfaceParamRef& ParamRef, const class FShaderParameterMap& ParameterMap) override
+	{
+		VectorFieldBuffer.Bind(ParameterMap, *(UNiagaraDataInterfaceVectorField::BufferBaseName + ParamRef.ParameterInfo.DataInterfaceHLSLSymbol));
+		Dimentions.Bind(ParameterMap, *(UNiagaraDataInterfaceVectorField::DimentionsBaseName + ParamRef.ParameterInfo.DataInterfaceHLSLSymbol));
+		BoundsMin.Bind(ParameterMap, *(UNiagaraDataInterfaceVectorField::BoundsMinBaseName + ParamRef.ParameterInfo.DataInterfaceHLSLSymbol));
+		BoundsMax.Bind(ParameterMap, *(UNiagaraDataInterfaceVectorField::BoundsMaxBaseName + ParamRef.ParameterInfo.DataInterfaceHLSLSymbol));
+	}
+
+	virtual void Serialize(FArchive& Ar)override
+	{
+		Ar << VectorFieldBuffer;
+		Ar << Dimentions;
+		Ar << BoundsMin;
+		Ar << BoundsMax;
+	}
+
+	virtual void Set(FRHICommandList& RHICmdList, FNiagaraShader* Shader, class UNiagaraDataInterface* DataInterface) const override 
+	{
+		check(IsInRenderingThread());
+
+		const FComputeShaderRHIParamRef ComputeShaderRHI = Shader->GetComputeShader();
+		UNiagaraDataInterfaceVectorField* VFDI = CastChecked<UNiagaraDataInterfaceVectorField>(DataInterface);
+		FRWBuffer& VFBuffer = VFDI->GetGPUBuffer();
 
+		RHICmdList.SetShaderResourceViewParameter(ComputeShaderRHI, VectorFieldBuffer.GetBaseIndex(), VFBuffer.SRV);
+		SetShaderValue(RHICmdList, ComputeShaderRHI, Dimentions, VFDI->GetDimentions());
+		SetShaderValue(RHICmdList, ComputeShaderRHI, BoundsMin, VFDI->GetBoundsMin());
+		SetShaderValue(RHICmdList, ComputeShaderRHI, BoundsMax, VFDI->GetBoundsMax());
+	}
+
+private:
+
+	FShaderResourceParameter VectorFieldBuffer;
+	FShaderParameter Dimentions;
+	FShaderParameter BoundsMin;
+	FShaderParameter BoundsMax;
+};
+
+FNiagaraDataInterfaceParametersCS* UNiagaraDataInterfaceVectorField::ConstructComputeParameters()const
+{
+	return new FNiagaraDataInterfaceParametersCS_VectorField();
+}
 
 #undef LOCTEXT_NAMESPACE
diff --git a/Engine/Plugins/FX/Niagara/Source/Niagara/Private/NiagaraDataSet.cpp b/Engine/Plugins/FX/Niagara/Source/Niagara/Private/NiagaraDataSet.cpp
--- a/Engine/Plugins/FX/Niagara/Source/Niagara/Private/NiagaraDataSet.cpp
+++ b/Engine/Plugins/FX/Niagara/Source/Niagara/Private/NiagaraDataSet.cpp
@@ -7,9 +7,11 @@
 #include "UpdateTextureShaders.h"
 #include "ShaderParameterUtils.h"
 #include "NiagaraStats.h"
+#include "NiagaraRenderer.h"
 
 DECLARE_CYCLE_STAT(TEXT("InitRenderData"), STAT_InitRenderData, STATGROUP_Niagara);
 
+
 //////////////////////////////////////////////////////////////////////////
 void FNiagaraDataSet::SetShaderParams(FNiagaraShader *Shader, FRHICommandList &CommandList)
 {
@@ -17,35 +19,49 @@ void FNiagaraDataSet::SetShaderParams(FNiagaraShader *Shader, FRHICommandList &C
 
 	if (Shader->FloatInputBufferParam.IsBound())
 	{
-		CommandList.TransitionResource(EResourceTransitionAccess::EReadable, EResourceTransitionPipeline::EComputeToCompute, PrevDataRender().GetGPUBufferFloat()->UAV);
-		CommandList.SetShaderResourceViewParameter(Shader->GetComputeShader(), Shader->FloatInputBufferParam.GetBaseIndex(), PrevDataRender().GetGPUBufferFloat()->SRV);
+		CommandList.TransitionResource(EResourceTransitionAccess::EReadable, EResourceTransitionPipeline::EComputeToCompute, PrevData().GetGPUBufferFloat()->UAV);
+		if (PrevData().GetNumInstancesAllocated() > 0)
+		{
+			CommandList.SetShaderResourceViewParameter(Shader->GetComputeShader(), Shader->FloatInputBufferParam.GetBaseIndex(), PrevData().GetGPUBufferFloat()->SRV);
+		}
+		else
+		{
+			CommandList.SetShaderResourceViewParameter(Shader->GetComputeShader(), Shader->FloatInputBufferParam.GetBaseIndex(), NiagaraRenderer::GetDummyFloatBuffer().SRV);
+		}
 	}
 	if (Shader->IntInputBufferParam.IsBound())
 	{
-		CommandList.TransitionResource(EResourceTransitionAccess::EReadable, EResourceTransitionPipeline::EComputeToCompute, PrevDataRender().GetGPUBufferInt()->UAV);
-		CommandList.SetShaderResourceViewParameter(Shader->GetComputeShader(), Shader->IntInputBufferParam.GetBaseIndex(), PrevDataRender().GetGPUBufferInt()->SRV);
+		CommandList.TransitionResource(EResourceTransitionAccess::EReadable, EResourceTransitionPipeline::EComputeToCompute, PrevData().GetGPUBufferInt()->UAV);
+		if (PrevData().GetNumInstancesAllocated() > 0)
+		{
+			CommandList.SetShaderResourceViewParameter(Shader->GetComputeShader(), Shader->FloatInputBufferParam.GetBaseIndex(), PrevData().GetGPUBufferInt()->SRV);
+		}
+		else
+		{
+			CommandList.SetShaderResourceViewParameter(Shader->GetComputeShader(), Shader->FloatInputBufferParam.GetBaseIndex(), NiagaraRenderer::GetDummyIntBuffer().SRV);
+		}
 	}
 	if (Shader->FloatOutputBufferParam.IsUAVBound())
 	{
-		CommandList.TransitionResource(EResourceTransitionAccess::EWritable, EResourceTransitionPipeline::EGfxToCompute, CurrDataRender().GetGPUBufferFloat()->UAV);
-		CommandList.SetUAVParameter(Shader->GetComputeShader(), Shader->FloatOutputBufferParam.GetUAVIndex(), CurrDataRender().GetGPUBufferFloat()->UAV);
+		CommandList.TransitionResource(EResourceTransitionAccess::EWritable, EResourceTransitionPipeline::EGfxToCompute, CurrData().GetGPUBufferFloat()->UAV);
+		CommandList.SetUAVParameter(Shader->GetComputeShader(), Shader->FloatOutputBufferParam.GetUAVIndex(), CurrData().GetGPUBufferFloat()->UAV);
 	}
 	if (Shader->IntOutputBufferParam.IsUAVBound())
 	{
-		CommandList.TransitionResource(EResourceTransitionAccess::EWritable, EResourceTransitionPipeline::EGfxToCompute, CurrDataRender().GetGPUBufferInt()->UAV);
-		CommandList.SetUAVParameter(Shader->GetComputeShader(), Shader->IntOutputBufferParam.GetUAVIndex(), CurrDataRender().GetGPUBufferInt()->UAV);
+		CommandList.TransitionResource(EResourceTransitionAccess::EWritable, EResourceTransitionPipeline::EGfxToCompute, CurrData().GetGPUBufferInt()->UAV);
+		CommandList.SetUAVParameter(Shader->GetComputeShader(), Shader->IntOutputBufferParam.GetUAVIndex(), CurrData().GetGPUBufferInt()->UAV);
 	}
 
 	if (Shader->ComponentBufferSizeWriteParam.IsBound())
 	{
-		uint32 SafeBufferSize = CurrDataRender().GetFloatStride() / sizeof(float);
-		CommandList.SetShaderParameter(Shader->GetComputeShader(), 0, Shader->ComponentBufferSizeWriteParam.GetBaseIndex(), Shader->ComponentBufferSizeWriteParam.GetNumBytes(), &SafeBufferSize);
+		uint32 SafeBufferSize = CurrData().GetFloatStride() / sizeof(float);
+		CommandList.SetShaderParameter(Shader->GetComputeShader(), Shader->ComponentBufferSizeWriteParam.GetBufferIndex(), Shader->ComponentBufferSizeWriteParam.GetBaseIndex(), Shader->ComponentBufferSizeWriteParam.GetNumBytes(), &SafeBufferSize);
 	}
 
 	if (Shader->ComponentBufferSizeReadParam.IsBound())
 	{
-		uint32 SafeBufferSize = PrevDataRender().GetFloatStride() / sizeof(float);
-		CommandList.SetShaderParameter(Shader->GetComputeShader(), 0, Shader->ComponentBufferSizeReadParam.GetBaseIndex(), Shader->ComponentBufferSizeReadParam.GetNumBytes(), &SafeBufferSize);
+		uint32 SafeBufferSize = PrevData().GetFloatStride() / sizeof(float);
+		CommandList.SetShaderParameter(Shader->GetComputeShader(), Shader->ComponentBufferSizeReadParam.GetBufferIndex(), Shader->ComponentBufferSizeReadParam.GetBaseIndex(), Shader->ComponentBufferSizeReadParam.GetNumBytes(), &SafeBufferSize);
 	}
 }
 
@@ -89,102 +105,30 @@ static FAutoConsoleVariableRef CVarNiagaraRenderBufferShrinkFactor(
 	ECVF_Default
 );
 
-void FNiagaraDataSet::InitGPUFromCPU_RenderThread()
+FNiagaraDataSet::~FNiagaraDataSet()
 {
-	SCOPE_CYCLE_COUNTER(STAT_InitRenderData);
-	check(IsInRenderingThread());
-	
-	//Some very basic size handling so we don't have to recreate these buffers every frame. TODO: Improve.
-	auto CalcNewSize = [](int32 ReqNum, int32 CurrNum)
-	{
-		int32 NewSize = AlignArbitrary(ReqNum, GRenderDataBlockSize);
-		int32 ShrinkSize = AlignArbitrary(int32(CurrNum * GGPUBufferShrinkFactor), GRenderDataBlockSize);
-		if (ReqNum < ShrinkSize)
-		{
-			NewSize = ShrinkSize;
-		}
-
-		check(NewSize >= ReqNum);
-		return NewSize;
-	};
-
-	if (GetNumFloatComponents())
-	{
-		int32 CurrBytes = RenderDataFloat.NumBytes;
-		int32 ValidBytes = CurrDataRender().GetFloatBuffer().Num();
-		int32 NewBytes = CalcNewSize(ValidBytes, CurrBytes);
-		if (NewBytes != CurrBytes)
- 		{
-			RenderDataFloat.Release();
-
-			if (NewBytes > 0)
-			{
-				RenderDataFloat.Initialize(sizeof(float), NewBytes / sizeof(float), EPixelFormat::PF_R32_FLOAT, BUF_Dynamic);
-			}
-		}
-
-		if (NewBytes > 0 && ValidBytes > 0)
-		{
-			void* BufferData = RHILockVertexBuffer(RenderDataFloat.Buffer, 0, ValidBytes, RLM_WriteOnly);
-			FPlatformMemory::Memcpy(BufferData, CurrDataRender().GetFloatBuffer().GetData(), ValidBytes);
-			RHIUnlockVertexBuffer(RenderDataFloat.Buffer);
-		}
-	}
-
-	if (GetNumInt32Components())
-	{
-		int32 CurrBytes = RenderDataInt.NumBytes;
-		int32 ValidBytes = CurrDataRender().GetInt32Buffer().Num();
-		int32 NewBytes = CalcNewSize(ValidBytes, CurrBytes);
- 		if (NewBytes != CurrBytes)
- 		{
-			RenderDataInt.Release();
-
-			if (NewBytes > 0)
-			{
-				RenderDataInt.Initialize(sizeof(int32), NewBytes / sizeof(int32), EPixelFormat::PF_R32_SINT, BUF_Dynamic);
-			}
-		}
+// 	int32 CurrBytes = RenderDataFloat.NumBytes + RenderDataInt.NumBytes;
+// 	DEC_MEMORY_STAT_BY(STAT_NiagaraVBMemory, CurrBytes);
 
-		if (NewBytes > 0 && ValidBytes > 0)
-		{
-			void*BufferData = RHILockVertexBuffer(RenderDataInt.Buffer, 0, ValidBytes, RLM_WriteOnly);
-			FPlatformMemory::Memcpy(BufferData, CurrDataRender().GetInt32Buffer().GetData(), ValidBytes);
-			RHIUnlockVertexBuffer(RenderDataInt.Buffer);
-		}
-	}
-
-	CurrentFloatDataSRV = RenderDataFloat.SRV;
-	CurrentIntDataSRV = RenderDataInt.SRV;
 }
 
-
-void FNiagaraDataSet::InitGPUSimSRVs_RenderThread()
-{
-	CurrentFloatDataSRV = CurrDataRender().GetGPUBufferFloat()->SRV;
-	CurrentIntDataSRV = CurrDataRender().GetGPUBufferInt()->SRV;
-}
-
-
-void FNiagaraDataSet::Dump(FNiagaraDataSet& Other, bool bCurr)
+void FNiagaraDataSet::Dump(FNiagaraDataSet& Other, bool bCurr, int32 StartIdx , int32 NumInstances)
 {
 	Other.Reset();
 	Other.Variables = Variables;
 	Other.VariableLayouts = VariableLayouts;
 
-	int32 IndexRead = CurrBuffer > 0 ? CurrBuffer - 1 : MaxBufferIdx;
-	if (bCurr)
-	{
-		IndexRead = CurrBuffer;
-	}
+	FNiagaraDataBuffer& DataBuffer = bCurr ? CurrData() : PrevData();
+	FNiagaraDataBuffer& OtherDataBuffer = Other.CurrData();
 
-	if (Other.Data[Other.CurrBuffer].GetNumInstancesAllocated() != Data[IndexRead].GetNumInstancesAllocated())
+	
+	if (OtherDataBuffer.GetNumInstancesAllocated() != DataBuffer.GetNumInstancesAllocated())
 	{
 		Other.Finalize();
-		Other.Data[Other.CurrBuffer].Allocate(Data[IndexRead].GetNumInstancesAllocated());
+		OtherDataBuffer.Allocate(DataBuffer.GetNumInstancesAllocated());
 	}
 
-	Data[CurrBuffer].CopyTo(Other.Data[Other.CurrBuffer]);
+	DataBuffer.CopyTo(OtherDataBuffer, StartIdx, NumInstances);
 }
 
 
@@ -246,14 +190,66 @@ void FNiagaraDataBuffer::Init(FNiagaraDataSet* InOwner)
 	Owner = InOwner;
 }
 
-void FNiagaraDataBuffer::Allocate(uint32 InNumInstances, ENiagaraSimTarget Target, bool bMaintainExisting)
+FNiagaraDataBuffer::~FNiagaraDataBuffer()
+{
+	Reset();
+}
+
+void FNiagaraDataBuffer::Reset()
+{
+	DEC_MEMORY_STAT_BY(STAT_NiagaraParticleMemory, FloatData.Max() + Int32Data.Max());
+	FloatData.Empty();
+	Int32Data.Empty();
+	FloatStride = 0;
+	Int32Stride = 0;
+	NumInstances = 0;
+	NumInstancesAllocated = 0;
+	NumChunksAllocatedForGPU = 0;
+}
+
+int32 FNiagaraDataBuffer::TransferInstance(FNiagaraDataBuffer& SourceBuffer, int32 InstanceIndex)
+{
+	if (SourceBuffer.GetNumInstances() > (uint32)InstanceIndex)
+	{
+		int32 OldNumInstances = NumInstances;
+		if (NumInstances == NumInstancesAllocated)
+		{
+			//Have to allocate some more space.
+			Allocate(NumInstancesAllocated + 1, true);
+		}
+
+		SetNumInstances(OldNumInstances + 1);
+
+		/** Copy the instance data. */
+		for (int32 CompIdx = (int32)Owner->TotalFloatComponents - 1; CompIdx >= 0; --CompIdx)
+		{
+			float* Src = SourceBuffer.GetInstancePtrFloat(CompIdx, InstanceIndex);
+			float* Dst = GetInstancePtrFloat(CompIdx, OldNumInstances);
+			*Dst = *Src;
+		}
+		for (int32 CompIdx = (int32)Owner->TotalInt32Components - 1; CompIdx >= 0; --CompIdx)
+		{
+			int32* Src = SourceBuffer.GetInstancePtrInt32(CompIdx, InstanceIndex);
+			int32* Dst = GetInstancePtrInt32(CompIdx, OldNumInstances);
+			*Dst = *Src;
+		}
+
+		return OldNumInstances;
+	}
+
+	return INDEX_NONE;
+}
+
+void FNiagaraDataBuffer::Allocate(uint32 InNumInstances, bool bMaintainExisting)
 {
 	check(Owner);
-	if (Target == ENiagaraSimTarget::CPUSim)
+	if (Owner->SimTarget == ENiagaraSimTarget::CPUSim)
 	{
 		NumInstancesAllocated = InNumInstances;
 		NumInstances = 0;
 
+		DEC_MEMORY_STAT_BY(STAT_NiagaraParticleMemory, FloatData.Max() + Int32Data.Max());
+
 		int32 OldFloatStride = FloatStride;
 		FloatStride = GetSafeComponentBufferSize(NumInstancesAllocated * sizeof(float));
 		FloatData.SetNum(FloatStride * Owner->GetNumFloatComponents(), false);
@@ -262,6 +258,8 @@ void FNiagaraDataBuffer::Allocate(uint32 InNumInstances, ENiagaraSimTarget Targe
 		Int32Stride = GetSafeComponentBufferSize(NumInstancesAllocated * sizeof(int32));
 		Int32Data.SetNum(Int32Stride * Owner->GetNumInt32Components(), false);
 
+		INC_MEMORY_STAT_BY(STAT_NiagaraParticleMemory, FloatData.Max() + Int32Data.Max());
+
 		//In some cases we want the existing data in the buffer to be maintained which due to the data layout requires some fix up.
 		if (bMaintainExisting)
 		{
@@ -294,7 +292,7 @@ void FNiagaraDataBuffer::AllocateGPU(uint32 InNumInstances, FRHICommandList &RHI
 	{
 		return;
 	}
-	check(IsInRenderingThread())
+	check(IsInRenderingThread());
 	const uint32 ALLOC_CHUNKSIZE = 4096;
 
 	NumInstancesAllocated = InNumInstances;
@@ -372,6 +370,60 @@ void FNiagaraDataBuffer::KillInstance(uint32 InstanceIdx)
 	}
 }
 
+void FNiagaraDataBuffer::CopyTo(FNiagaraDataBuffer& DestBuffer, int32 InStartIdx, int32 InNumInstances)
+{
+	if (InStartIdx < 0 || (uint32)InStartIdx > NumInstances)
+	{
+		InStartIdx = NumInstances;
+	}
+	if (InNumInstances < 0 || ((uint32)InNumInstances + (uint32)InStartIdx) > NumInstances)
+	{
+		InNumInstances = NumInstances - InStartIdx;
+	}
+
+	if (InNumInstances != 0)
+	{
+		if (DestBuffer.NumInstancesAllocated != NumInstancesAllocated)
+		{
+			DestBuffer.Allocate(NumInstancesAllocated);
+		}
+
+		for (uint32 CompIdx = 0; CompIdx < Owner->TotalFloatComponents; ++CompIdx)
+		{
+			float* SrcStart = GetInstancePtrFloat(CompIdx, InStartIdx);
+			float* SrcEnd = GetInstancePtrFloat(CompIdx, InStartIdx + InNumInstances);
+			float* Dst = DestBuffer.GetInstancePtrFloat(CompIdx, 0);
+			size_t Count = SrcEnd - SrcStart;
+			FMemory::Memcpy(Dst, SrcStart, Count*sizeof(float));
+
+			if (Count > 0)
+			{
+				for (size_t i = 0; i < Count; i++)
+				{
+					check(SrcStart[i] == Dst[i]);
+				}
+			}
+		}
+		for (uint32 CompIdx = 0; CompIdx < Owner->TotalInt32Components; ++CompIdx)
+		{
+			int32* SrcStart = GetInstancePtrInt32(CompIdx, InStartIdx);
+			int32* SrcEnd = GetInstancePtrInt32(CompIdx, InStartIdx + InNumInstances);
+			int32* Dst = DestBuffer.GetInstancePtrInt32(CompIdx, 0);
+			size_t Count = SrcEnd - SrcStart;
+			FMemory::Memcpy(Dst, SrcStart, Count * sizeof(int32));
+
+			if (Count > 0)
+			{
+				for (size_t i = 0; i < Count; i++)
+				{
+					check(SrcStart[i] == Dst[i]);
+				}
+			}
+		}
+		DestBuffer.SetNumInstances(InNumInstances);
+	}
+}
+
 void FNiagaraDataBuffer::CopyTo(FNiagaraDataBuffer& DestBuffer)
 {
 	DestBuffer.FloatStride = FloatStride;
