Copying //UE4/Dev-Niagara to //UE4/Dev-Main (Source: //UE4/Dev-Niagara @ 4074996)

diff --git a/Engine/Plugins/FX/Niagara/Source/Niagara/Private/NiagaraEmitter.cpp b/Engine/Plugins/FX/Niagara/Source/Niagara/Private/NiagaraEmitter.cpp
--- a/Engine/Plugins/FX/Niagara/Source/Niagara/Private/NiagaraEmitter.cpp
+++ b/Engine/Plugins/FX/Niagara/Source/Niagara/Private/NiagaraEmitter.cpp
@@ -9,6 +9,11 @@
 #include "UObject/Package.h"
 #include "UObject/Linker.h"
 #include "NiagaraModule.h"
+#include "NiagaraSystem.h"
+
+#if WITH_EDITOR
+const FName UNiagaraEmitter::PrivateMemberNames::EventHandlerScriptProps = GET_MEMBER_NAME_CHECKED(UNiagaraEmitter, EventHandlerScriptProps);
+#endif
 
 static int32 GbForceNiagaraCompileOnLoad = 0;
 static FAutoConsoleVariableRef CVarForceNiagaraCompileOnLoad(
@@ -18,22 +23,30 @@ static FAutoConsoleVariableRef CVarForceNiagaraCompileOnLoad(
 	ECVF_Default
 	);
 
+static int32 GbForceNiagaraFailToCompile = 0;
+static FAutoConsoleVariableRef CVarForceNiagaraCompileToFail(
+	TEXT("fx.ForceNiagaraCompileToFail"),
+	GbForceNiagaraFailToCompile,
+	TEXT("If > 0 emitters will go through the motions of a compile, but will never set valid bytecode. \n"),
+	ECVF_Default
+);
+
 void FNiagaraEmitterScriptProperties::InitDataSetAccess()
 {
 	EventReceivers.Empty();
 	EventGenerators.Empty();
 
-	if (Script)
+	if (Script && Script->IsReadyToRun(ENiagaraSimTarget::CPUSim))
 	{
 		//UE_LOG(LogNiagara, Log, TEXT("InitDataSetAccess: %s %d %d"), *Script->GetPathName(), Script->ReadDataSets.Num(), Script->WriteDataSets.Num());
 		// TODO: add event receiver and generator lists to the script properties here
 		//
-		for (FNiagaraDataSetID &ReadID : Script->ReadDataSets)
+		for (FNiagaraDataSetID &ReadID : Script->GetVMExecutableData().ReadDataSets)
 		{
 			EventReceivers.Add( FNiagaraEventReceiverProperties(ReadID.Name, "", "") );
 		}
 
-		for (FNiagaraDataSetProperties &WriteID : Script->WriteDataSets)
+		for (FNiagaraDataSetProperties &WriteID : Script->GetVMExecutableData().WriteDataSets)
 		{
 			FNiagaraEventGeneratorProperties Props(WriteID, "", "");
 			EventGenerators.Add(Props);
@@ -43,13 +56,13 @@ void FNiagaraEmitterScriptProperties::InitDataSetAccess()
 
 bool FNiagaraEmitterScriptProperties::DataSetAccessSynchronized() const
 {
-	if (Script)
+	if (Script && Script->IsReadyToRun(ENiagaraSimTarget::CPUSim))
 	{
-		if (Script->ReadDataSets.Num() != EventReceivers.Num())
+		if (Script->GetVMExecutableData().ReadDataSets.Num() != EventReceivers.Num())
 		{
 			return false;
 		}
-		if (Script->WriteDataSets.Num() != EventGenerators.Num())
+		if (Script->GetVMExecutableData().WriteDataSets.Num() != EventGenerators.Num())
 		{
 			return false;
 		}
@@ -65,7 +78,6 @@ bool FNiagaraEmitterScriptProperties::DataSetAccessSynchronized() const
 
 UNiagaraEmitter::UNiagaraEmitter(const FObjectInitializer& Initializer)
 : Super(Initializer)
-, CollisionMode(ENiagaraCollisionMode::None)
 , FixedBounds(FBox(FVector(-100), FVector(100)))
 , MinDetailLevel(0)
 , MaxDetailLevel(4)
@@ -73,6 +85,7 @@ UNiagaraEmitter::UNiagaraEmitter(const FObjectInitializer& Initializer)
 , bFixedBounds(false)
 , bUseMinDetailLevel(false)
 , bUseMaxDetailLevel(false)
+, bRequiresPersistentIDs(false)
 #if WITH_EDITORONLY_DATA
 , ThumbnailImageOutOfDate(true)
 #endif
@@ -98,10 +111,20 @@ void UNiagaraEmitter::PostInitProperties()
 		EmitterUpdateScriptProps.Script = NewObject<UNiagaraScript>(this, "EmitterUpdateScript", EObjectFlags::RF_Transactional);
 		EmitterUpdateScriptProps.Script->SetUsage(ENiagaraScriptUsage::EmitterUpdateScript);
 
+		GPUComputeScript = NewObject<UNiagaraScript>(this, "GPUComputeScript", EObjectFlags::RF_Transactional);
+		GPUComputeScript->SetUsage(ENiagaraScriptUsage::ParticleGPUComputeScript);
+
 	}
 	UniqueEmitterName = TEXT("Emitter");
 }
 
+#if WITH_EDITORONLY_DATA
+bool UNiagaraEmitter::GetForceCompileOnLoad()
+{
+	return GbForceNiagaraCompileOnLoad > 0;
+}
+#endif
+
 void UNiagaraEmitter::Serialize(FArchive& Ar)
 {
 	Super::Serialize(Ar);
@@ -118,18 +141,14 @@ void UNiagaraEmitter::PostLoad()
 		SetFlags(RF_Transactional);
 	}
 
-// 	//Have to force the usage to be correct.
-// 	if (SpawnScriptProps.Script)
-// 	{
-// 		SpawnScriptProps.Script->ConditionalPostLoad();
-// 		ENiagaraScriptUsage ActualUsage = bInterpolatedSpawning ? ENiagaraScriptUsage::ParticleSpawnScriptInterpolated : ENiagaraScriptUsage::ParticleSpawnScript;
-// 		if (SpawnScriptProps.Script->Usage != ActualUsage)
-// 		{
-// 			SpawnScriptProps.Script->Usage = ActualUsage;
-// 			SpawnScriptProps.Script->ByteCode.Empty();//Force them to recompile.
-// 			UE_LOG(LogNiagara, Warning, TEXT("Niagara Emitter had mis matched bInterpolatedSpawning values. You must recompile this emitter. %s"), *GetFullName());
-// 		}
-// 	}
+	if (!GPUComputeScript)
+	{
+		GPUComputeScript = NewObject<UNiagaraScript>(this, "GPUComputeScript", EObjectFlags::RF_Transactional);
+		GPUComputeScript->SetUsage(ENiagaraScriptUsage::ParticleGPUComputeScript);
+#if WITH_EDITORONLY_DATA
+		GPUComputeScript->SetSource(SpawnScriptProps.Script ? SpawnScriptProps.Script->GetSource() : nullptr);
+#endif
+	}
 
 	if (EmitterSpawnScriptProps.Script == nullptr || EmitterUpdateScriptProps.Script == nullptr)
 	{
@@ -147,7 +166,7 @@ void UNiagaraEmitter::PostLoad()
 		}
 #endif
 	}
-		
+
 	//Temporarily disabling interpolated spawn if the script type and flag don't match.
 	if (SpawnScriptProps.Script)
 	{
@@ -158,106 +177,65 @@ void UNiagaraEmitter::PostLoad()
 			bInterpolatedSpawning = false;
 			if (bActualInterpolatedSpawning)
 			{
-				SpawnScriptProps.Script->InvalidateScript();//clear out the script as it was compiled with interpolated spawn.
-
+#if WITH_EDITORONLY_DATA
+				SpawnScriptProps.Script->InvalidateCompileResults();//clear out the script as it was compiled with interpolated spawn.
+#endif
 				SpawnScriptProps.Script->SetUsage(ENiagaraScriptUsage::ParticleSpawnScript);
 			}
 			UE_LOG(LogNiagara, Warning, TEXT("Disabling interpolated spawn because emitter flag and script type don't match. Did you adjust this value in the UI? Emitter may need recompile.. %s"), *GetFullName());
 		}
 	}
 
-	// Referenced scripts may need to be invalidated due to version mismatch or 
-	// other issues. This is determined in PostLoad for UNiagaraScript, so we need
-	// to make sure that PostLoad has happened for each of them first.
-	bool bNeedsRecompile = false;
-	if (SpawnScriptProps.Script)
+#if WITH_EDITORONLY_DATA
+	GraphSource->ConditionalPostLoad();
+	GraphSource->PostLoadFromEmitter(*this);
+#endif
+
+	TArray<UNiagaraScript*> AllScripts;
+	GetScripts(AllScripts, true);
+
+	// Post load scripts for use below.
+	for (UNiagaraScript* Script : AllScripts)
 	{
-		SpawnScriptProps.Script->ConditionalPostLoad();
-		if (SpawnScriptProps.Script->GetByteCode().Num() == 0 || GbForceNiagaraCompileOnLoad > 0)
-		{
-			bNeedsRecompile = true;
-		}
-	}
-	if (UpdateScriptProps.Script)
-	{
-		UpdateScriptProps.Script->ConditionalPostLoad();
-		if (UpdateScriptProps.Script->GetByteCode().Num() == 0 || GbForceNiagaraCompileOnLoad > 0)
-		{
-			bNeedsRecompile = true;
-		}
+		Script->ConditionalPostLoad();
 	}
 
-	for (int32 i = 0; i < EventHandlerScriptProps.Num(); i++)
+	// Reset scripts if recompile is forced.
+#if WITH_EDITORONLY_DATA
+	bool bGenerateNewChangeId = false;
+	if (GetForceCompileOnLoad())
 	{
-		if (EventHandlerScriptProps[i].Script)
+		// If we are a standalone emitter, then we invalidate id's, which should cause systems dependent on us to regenerate.
+		UObject* OuterObj = GetOuter();
+		if (OuterObj == GetOutermost())
 		{
-			EventHandlerScriptProps[i].Script->ConditionalPostLoad();
-			if (EventHandlerScriptProps[i].Script->GetByteCode().Num() == 0 || GbForceNiagaraCompileOnLoad > 0)
-			{
-				bNeedsRecompile = true;
-			}
+			GraphSource->InvalidateCachedCompileIds();
+			bGenerateNewChangeId = true;
+			UE_LOG(LogNiagara, Log, TEXT("InvalidateCachedCompileIds for %s because GbForceNiagaraCompileOnLoad = %d"), *GetPathName(), GbForceNiagaraCompileOnLoad);
 		}
 	}
-
-#if WITH_EDITORONLY_DATA
-	GraphSource->ConditionalPostLoad();
-	bNeedsRecompile |= GraphSource->PostLoadFromEmitter(*this);
-#endif
-
-	// Check ourselves against the current script rebuild version. Force our children
-	// to need to update.
-	const int32 NiagaraVer = GetLinkerCustomVersion(FNiagaraCustomVersion::GUID);
-	if (NiagaraVer < FNiagaraCustomVersion::LatestVersion)
+	
+	if (ChangeId.IsValid() == false)
 	{
-		bNeedsRecompile = true;
+		// If the change id is already invalid we need to generate a new one, and can skip checking the owned scripts.
+		bGenerateNewChangeId = true;
+		UE_LOG(LogNiagara, Log, TEXT("Change ID updated for emitter %s because the ID was invalid."), *GetPathName());
 	}
-	
-	if (bNeedsRecompile)
+	else
 	{
-		if (SpawnScriptProps.Script)
-		{
-			SpawnScriptProps.Script->InvalidateScript();//clear out the script.
-		}
-		if (UpdateScriptProps.Script)
-		{
-			UpdateScriptProps.Script->InvalidateScript();//clear out the script.
-		}
-
-		for (int32 i = 0; i < EventHandlerScriptProps.Num(); i++)
+		for (UNiagaraScript* Script : AllScripts)
 		{
-			if (EventHandlerScriptProps[i].Script)
+			if (Script->AreScriptAndSourceSynchronized() == false)
 			{
-				EventHandlerScriptProps[i].Script->InvalidateScript();//clear out the script.
+				bGenerateNewChangeId = true;
+				//UE_LOG(LogNiagara, Log, TEXT("Change ID updated for emitter %s because of a change to its script %s"), *GetPathName(), *Script->GetPathName());
 			}
 		}
 	}
 
-#if WITH_EDITORONLY_DATA
-	// Go ahead and do a recompile for all referenced scripts. Note that the compile will 
-	// cause the emitter's change id to be regenerated.
-	if (bNeedsRecompile)
+	if (bGenerateNewChangeId)
 	{
-		UObject* Outer = GetOuter();
-		if (Outer != nullptr && Outer->IsA<UPackage>())
-		{
-			TArray<ENiagaraScriptCompileStatus> ScriptStatuses;
-			TArray<FString> ScriptErrors;
-			TArray<FString> PathNames;
-			TArray<UNiagaraScript*> Scripts;
-			CompileScripts(ScriptStatuses, ScriptErrors, PathNames, Scripts, false);
-
-			for (int32 i = 0; i < ScriptStatuses.Num(); i++)
-			{
-				if (ScriptErrors[i].Len() != 0 || ScriptStatuses[i] != ENiagaraScriptCompileStatus::NCS_UpToDate)
-				{
-					UE_LOG(LogNiagara, Warning, TEXT("Script '%s', compile status: %d  errors: %s"), *PathNames[i], (int32)ScriptStatuses[i], *ScriptErrors[i]);
-				}
-				else
-				{
-					UE_LOG(LogNiagara, Log, TEXT("Script '%s', compile status: Success!"), *PathNames[i]);
-				}
-			}
-		}
+		ChangeId = FGuid::NewGuid();
 	}
 
 	GraphSource->OnChanged().AddUObject(this, &UNiagaraEmitter::GraphSourceChanged);
@@ -313,16 +291,34 @@ void UNiagaraEmitter::PostEditChangeProperty(struct FPropertyChangedEvent& Prope
 			UE_LOG(LogNiagara, Log, TEXT("Updating script usage: Script->IsInterpolatdSpawn %d Emitter->bInterpolatedSpawning %d"), (int32)SpawnScriptProps.Script->IsInterpolatedParticleSpawnScript(), bInterpolatedSpawning);
 			if (GraphSource != nullptr)
 			{
-				GraphSource->MarkNotSynchronized();
+				GraphSource->MarkNotSynchronized(TEXT("Emitter interpolated spawn changed"));
 			}
+#if WITH_EDITORONLY_DATA
+			UNiagaraSystem::RequestCompileForEmitter(this);
+#endif
 		}
 	}
 	if (PropertyName == GET_MEMBER_NAME_CHECKED(UNiagaraEmitter, SimTarget))
 	{
 		if (GraphSource != nullptr)
 		{
-			GraphSource->MarkNotSynchronized();
+			GraphSource->MarkNotSynchronized(TEXT("Emitter simulation target changed."));
+		}
+
+#if WITH_EDITORONLY_DATA
+		UNiagaraSystem::RequestCompileForEmitter(this);
+#endif
+	}
+	if (PropertyName == GET_MEMBER_NAME_CHECKED(UNiagaraEmitter, bRequiresPersistentIDs))
+	{
+		if (GraphSource != nullptr)
+		{
+			GraphSource->MarkNotSynchronized(TEXT("Emitter Requires Persistent IDs changed."));
 		}
+
+#if WITH_EDITORONLY_DATA
+		UNiagaraSystem::RequestCompileForEmitter(this);
+#endif
 	}
 	ThumbnailImageOutOfDate = true;
 	ChangeId = FGuid::NewGuid();
@@ -369,8 +365,8 @@ bool UNiagaraEmitter::IsValid()const
 
 	if (SimTarget == ENiagaraSimTarget::GPUComputeSim || SimTarget == ENiagaraSimTarget::DynamicLoadBalancedSim)
 	{
-		if (!SpawnScriptProps.Script->IsScriptCompilationPending(true) && 
-			!SpawnScriptProps.Script->DidScriptCompilationSucceed(true))
+		if (!GPUComputeScript->IsScriptCompilationPending(true) && 
+			!GPUComputeScript->DidScriptCompilationSucceed(true))
 		{
 			return false;
 		}
@@ -410,7 +406,7 @@ bool UNiagaraEmitter::IsReadyToRun() const
 
 	if (SimTarget == ENiagaraSimTarget::GPUComputeSim || SimTarget == ENiagaraSimTarget::DynamicLoadBalancedSim)
 	{
-		if (SpawnScriptProps.Script->IsScriptCompilationPending(true))
+		if (GPUComputeScript->IsScriptCompilationPending(true))
 		{
 			return false;
 		}
@@ -436,6 +432,11 @@ void UNiagaraEmitter::GetScripts(TArray<UNiagaraScript*>& OutScripts, bool bComp
 			OutScripts.Add(EventHandlerScriptProps[i].Script);
 		}
 	}
+
+	if (SimTarget == ENiagaraSimTarget::DynamicLoadBalancedSim || SimTarget == ENiagaraSimTarget::GPUComputeSim)
+	{
+		OutScripts.Add(GPUComputeScript);
+	}
 }
 
 UNiagaraScript* UNiagaraEmitter::GetScript(ENiagaraScriptUsage Usage, FGuid UsageId)
@@ -463,6 +464,11 @@ bool UNiagaraEmitter::IsAllowedByDetailLevel()const
 	return true;
 }
 
+bool UNiagaraEmitter::RequiresPersistantIDs()const
+{
+	return bRequiresPersistentIDs;
+}
+
 #if WITH_EDITORONLY_DATA
 
 FGuid UNiagaraEmitter::GetChangeId() const
@@ -503,169 +509,38 @@ bool UNiagaraEmitter::AreAllScriptAndSourcesSynchronized() const
 	return true;
 }
 
-void UNiagaraEmitter::CompileScripts(TArray<ENiagaraScriptCompileStatus>& OutScriptStatuses, TArray<FString>& OutGraphLevelErrorMessages, TArray<FString>& PathNames,
-	TArray<UNiagaraScript*>& Scripts, bool bForce)
-{
-	PathNames.AddDefaulted(2);
-	OutScriptStatuses.Empty();
-	OutGraphLevelErrorMessages.Empty();
-	bool bPostCompile = false;
-
-	// Force the end-user to reopen the file to resolve the problems.
-	if (GraphSource == nullptr)
-	{
-		OutScriptStatuses.AddDefaulted(2);
-		OutGraphLevelErrorMessages.AddDefaulted(2);
 
-		OutGraphLevelErrorMessages[0] = OutGraphLevelErrorMessages[1] = TEXT("Please reopen asset in editor.");
-		OutScriptStatuses[0] = OutScriptStatuses[1] = ENiagaraScriptCompileStatus::NCS_Error;
-		return;
-	}
-	else if (!GraphSource->IsPreCompiled()) // If we haven't already precompiled this source, do so..
-	{
-		GraphSource->PreCompile(this);
-		bPostCompile = true;
-	}
+UNiagaraEmitter::FOnEmitterCompiled& UNiagaraEmitter::OnEmitterVMCompiled()
+{
+	return OnVMScriptCompiledDelegate;
+}
 
-	TArray<int32> OutputScriptIndices;
-	TArray<int32> OutputStatusIndices;
+void UNiagaraEmitter::OnPostCompile()
+{
+	SyncEmitterAlias(TEXT("Emitter"), UniqueEmitterName);
 
-	FString ErrorMsg;
-	check((int32)EScriptCompileIndices::SpawnScript == 0);
-	int32 StatusIdx = OutScriptStatuses.Add(SpawnScriptProps.Script->Compile(ErrorMsg, bForce));
-	OutputStatusIndices.Add(StatusIdx);
-	OutGraphLevelErrorMessages.Add(ErrorMsg); // EScriptCompileIndices::SpawnScript
-	int32 ScriptAddIdx = Scripts.Add(SpawnScriptProps.Script);
-	OutputScriptIndices.Add(ScriptAddIdx);
 	SpawnScriptProps.InitDataSetAccess();
-	PathNames[(int32)EScriptCompileIndices::SpawnScript] = SpawnScriptProps.Script->GetPathName();
-
-	ErrorMsg.Empty();
-	check((int32)EScriptCompileIndices::UpdateScript == 1);
-	StatusIdx = OutScriptStatuses.Add(UpdateScriptProps.Script->Compile(ErrorMsg, bForce));
-	OutputStatusIndices.Add(StatusIdx);
-	ScriptAddIdx = Scripts.Add(UpdateScriptProps.Script);
-	OutputScriptIndices.Add(ScriptAddIdx);
-	OutGraphLevelErrorMessages.Add(ErrorMsg); // EScriptCompileIndices::UpdateScript
 	UpdateScriptProps.InitDataSetAccess();
-	PathNames[(int32)EScriptCompileIndices::UpdateScript] = UpdateScriptProps.Script->GetPathName();
 
 	for (int32 i = 0; i < EventHandlerScriptProps.Num(); i++)
 	{
 		if (EventHandlerScriptProps[i].Script)
 		{
-			check((int32)EScriptCompileIndices::EventScript == 2);
-			ErrorMsg.Empty();
-			StatusIdx = OutScriptStatuses.Add(EventHandlerScriptProps[i].Script->Compile(ErrorMsg, bForce));
-			OutputStatusIndices.Add(StatusIdx);
-			ScriptAddIdx = Scripts.Add(EventHandlerScriptProps[i].Script);
-			OutputScriptIndices.Add(ScriptAddIdx);
-			OutGraphLevelErrorMessages.Add(ErrorMsg); // EScriptCompileIndices::EventScript
 			EventHandlerScriptProps[i].InitDataSetAccess();
-			PathNames.Add(EventHandlerScriptProps[i].Script->GetPathName());
 		}
 	}
 
-	check(OutputScriptIndices.Num() == OutputStatusIndices.Num());
-	check(OutGraphLevelErrorMessages.Num() == OutputScriptIndices.Num());
-
-	// Go through all attached renderers and make sure that the scripts work properly for those renderers.
-	for (int32 OutputScriptIdx = 0; OutputScriptIdx < OutputScriptIndices.Num(); OutputScriptIdx++)
+	if (GbForceNiagaraFailToCompile != 0)
 	{
-		int32 LoopScriptIdx = OutputScriptIndices[OutputScriptIdx];
-		int32 LoopStatusIdx = OutputStatusIndices[OutputScriptIdx];
-		UNiagaraScript* CurrentScript = Scripts[LoopScriptIdx];
-
-		// Only care about scripts that successfully compiled so far.
-		if (OutScriptStatuses[LoopStatusIdx] != ENiagaraScriptCompileStatus::NCS_UpToDate)
-		{
-			continue;
-		}
-
-		// Ignore missing scripts.
-		if (CurrentScript == nullptr)
-		{
-			continue;
-		}
-
-		// Only care about particle scripts
-		if (!(CurrentScript->IsParticleSpawnScript() ||
-			CurrentScript->IsParticleUpdateScript()))
-		{
-			continue;
-		}
-
-		for (int32 i = 0; i < RendererProperties.Num(); i++)
+		TArray<UNiagaraScript*> Scripts;
+		GetScripts(Scripts, false);
+		for (int32 i = 0; i < Scripts.Num(); i++)
 		{
-			if (RendererProperties[i] != nullptr)
-			{
-				const TArray<FNiagaraVariable>& RequiredAttrs = RendererProperties[i]->GetRequiredAttributes();
-
-				for (FNiagaraVariable Attr : RequiredAttrs)
-				{
-					// TODO .. should we always be namespaced?
-					FString AttrName = Attr.GetName().ToString();
-					if (AttrName.RemoveFromStart(TEXT("Particles.")))
-					{
-						Attr.SetName(*AttrName);
-					}
-
-					if (false == Scripts[LoopScriptIdx]->Attributes.ContainsByPredicate([&Attr](const FNiagaraVariable& Var) { return Var.GetName() == Attr.GetName(); }))
-					{
-						OutGraphLevelErrorMessages[LoopStatusIdx] += FString::Printf(TEXT("\nCannot bind to renderer %s because it does not define attribute %s %s."), *RendererProperties[i]->GetName(), *Attr.GetType().GetNameText().ToString(), *Attr.GetName().ToString());
-						OutScriptStatuses[LoopStatusIdx] = ENiagaraScriptCompileStatus::NCS_Error;
-					}
-				}
-			}
+			Scripts[i]->InvalidateCompileResults(); 
 		}
 	}
 
-	if (GraphSource && bPostCompile) // Cleanup if we precompiled inside this function
-	{
-		GraphSource->PostCompile();
-	}
-}
-
-ENiagaraScriptCompileStatus UNiagaraEmitter::CompileScript(EScriptCompileIndices InScriptToCompile, FString& OutGraphLevelErrorMessages, bool bForce, int32 InSubScriptIdx)
-{
-	ENiagaraScriptCompileStatus CompileStatus = ENiagaraScriptCompileStatus::NCS_Unknown;
-	bool bPostCompile = false;
-	if (GraphSource && !GraphSource->IsPreCompiled()) // If we haven't already precompiled this source, do so..
-	{
-		bPostCompile = true;
-		GraphSource->PreCompile(this);
-	}
-
-	switch (InScriptToCompile)
-	{
-	case EScriptCompileIndices::SpawnScript:
-		CompileStatus = SpawnScriptProps.Script->Compile(OutGraphLevelErrorMessages, bForce);
-		SpawnScriptProps.InitDataSetAccess();
-		break;
-	case EScriptCompileIndices::UpdateScript:
-		CompileStatus = UpdateScriptProps.Script->Compile(OutGraphLevelErrorMessages, bForce);
-		UpdateScriptProps.InitDataSetAccess();
-		break;
-	case EScriptCompileIndices::EventScript:
-		if (EventHandlerScriptProps[InSubScriptIdx].Script)
-		{
-			CompileStatus = EventHandlerScriptProps[InSubScriptIdx].Script->Compile(OutGraphLevelErrorMessages, bForce);
-			EventHandlerScriptProps[InSubScriptIdx].InitDataSetAccess();
-		}
-		else
-		{
-			return CompileStatus;
-		}
-		break;
-	default:
-		return CompileStatus;
-	}
-
-	if (GraphSource && bPostCompile) // Cleanup if we precompiled inside this function
-	{
-		GraphSource->PostCompile();
-	}
-	return CompileStatus;
+	OnEmitterVMCompiled().Broadcast(this);
 }
 
 UNiagaraEmitter* UNiagaraEmitter::MakeRecursiveDeepCopy(UObject* DestOuter) const
@@ -775,7 +650,23 @@ FString UNiagaraEmitter::GetUniqueEmitterName()const
 	return UniqueEmitterName;
 }
 
-void UNiagaraEmitter::SetUniqueEmitterName(const FString& InName)
+#if WITH_EDITORONLY_DATA
+void UNiagaraEmitter::SyncEmitterAlias(const FString& InOldName, const FString& InNewName)
+{
+	TMap<FString, FString> RenameMap;
+	RenameMap.Add(InOldName, InNewName);
+
+	TArray<UNiagaraScript*> Scripts;
+	GetScripts(Scripts, false); // Get all the scripts...
+
+	for (UNiagaraScript* Script : Scripts)
+	{
+		Script->Modify();
+		Script->SyncAliases(RenameMap);
+	}
+}
+#endif
+bool UNiagaraEmitter::SetUniqueEmitterName(const FString& InName)
 {
 	if (InName != UniqueEmitterName)
 	{
@@ -783,66 +674,16 @@ void UNiagaraEmitter::SetUniqueEmitterName(const FString& InName)
 		FString OldName = UniqueEmitterName;
 		UniqueEmitterName = InName;
 
-		TMap<FString, FString> RenameMap;
-		RenameMap.Add(OldName, InName);
-
-		TArray<UNiagaraScript*> Scripts;
-		GetScripts(Scripts, false);
-
-		for (UNiagaraScript* Script : Scripts)
-		{
-			Script->Modify();
-
-			// First handle any rapid iteration parameters...
-			TArray<FNiagaraVariable> Params;
-			Script->RapidIterationParameters.GetParameters(Params);
-			for (FNiagaraVariable Var : Params)
-			{
-				FNiagaraVariable NewVar = FNiagaraVariable::ResolveAliases(Var, RenameMap);
-				if (NewVar.GetName() != Var.GetName())
-				{
-					Script->RapidIterationParameters.RenameParameter(Var, NewVar.GetName());
-				}
-			}
-
-			// Now handle any Parameters overall..
-			for (int32 i = 0; i < Script->Parameters.Parameters.Num(); i++)
-			{
-				FNiagaraVariable Var = Script->Parameters.Parameters[i];
-				FNiagaraVariable NewVar = FNiagaraVariable::ResolveAliases(Var, RenameMap);
-				if (NewVar.GetName() != Var.GetName())
-				{
-					Script->Parameters.Parameters[i] = NewVar;
-				}
-			}
-
-			// Also handle any data set mappings...
-			auto Iterator = Script->DataSetToParameters.CreateIterator();
-			while (Iterator)
-			{
-				for (int32 i = 0; i < Iterator.Value().Parameters.Num(); i++)
-				{
-					FNiagaraVariable Var = Iterator.Value().Parameters[i];
-					FNiagaraVariable NewVar = FNiagaraVariable::ResolveAliases(Var, RenameMap);
-					if (NewVar.GetName() != Var.GetName())
-					{
-						Iterator.Value().Parameters[i] = NewVar;
-					}
-				}
-				++Iterator;
-			}
-		}
-
-		// Since we use the emitter name as part of the script, we now need to recompile to work properly.
 #if WITH_EDITORONLY_DATA
-		if (GraphSource != nullptr)
-		{
-			GraphSource->MarkNotSynchronized();
-		}
+		SyncEmitterAlias(OldName, UniqueEmitterName);
 #endif
+		return true;
 	}
+
+	return false;
 }
 
+
 FNiagaraVariable UNiagaraEmitter::ToEmitterParameter(const FNiagaraVariable& EmitterVar)const
 {
 	FNiagaraVariable Var = EmitterVar;
@@ -941,4 +782,4 @@ void UNiagaraEmitter::GraphSourceChanged()
 {
 	UpdateChangeId();
 }
-#endif
+#endif
\ No newline at end of file
diff --git a/Engine/Plugins/FX/Niagara/Source/Niagara/Private/NiagaraEmitterHandle.cpp b/Engine/Plugins/FX/Niagara/Source/Niagara/Private/NiagaraEmitterHandle.cpp
--- a/Engine/Plugins/FX/Niagara/Source/Niagara/Private/NiagaraEmitterHandle.cpp
+++ b/Engine/Plugins/FX/Niagara/Source/Niagara/Private/NiagaraEmitterHandle.cpp
@@ -11,6 +11,7 @@
 #include "Modules/ModuleManager.h"
 
 const FNiagaraEmitterHandle FNiagaraEmitterHandle::InvalidHandle;
+const FString InitialNotSynchronizedReason("Emitter handle constructed");
 
 FNiagaraEmitterHandle::FNiagaraEmitterHandle() :
 #if WITH_EDITORONLY_DATA
@@ -34,7 +35,9 @@ FNiagaraEmitterHandle::FNiagaraEmitterHandle(UNiagaraEmitter& Emitter)
 #endif
 	, Instance(&Emitter)
 {
-	
+#if WITH_EDITORONLY_DATA
+	Instance->GraphSource->MarkNotSynchronized(InitialNotSynchronizedReason);
+#endif
 }
 
 #if WITH_EDITORONLY_DATA
@@ -49,6 +52,7 @@ FNiagaraEmitterHandle::FNiagaraEmitterHandle(UNiagaraEmitter& InSourceEmitter, F
 	, Instance(Cast<UNiagaraEmitter>(StaticDuplicateObject(Source, &InOuterSystem)))
 {
 	Instance->SetUniqueEmitterName(Name.ToString());
+	Instance->GraphSource->MarkNotSynchronized(InitialNotSynchronizedReason);
 }
 
 FNiagaraEmitterHandle::FNiagaraEmitterHandle(const FNiagaraEmitterHandle& InHandleToDuplicate, FName InDuplicateName, UNiagaraSystem& InDuplicateOwnerSystem)
@@ -64,6 +68,7 @@ FNiagaraEmitterHandle::FNiagaraEmitterHandle(const FNiagaraEmitterHandle& InHand
 	// Clear stand alone and public flags from the referenced emitters since they are not assets.
 	Instance->ClearFlags(RF_Standalone | RF_Public);
 	Instance->SetUniqueEmitterName(Name.ToString());
+	Instance->GraphSource->MarkNotSynchronized(InitialNotSynchronizedReason);
 	LastMergedSource->ClearFlags(RF_Standalone | RF_Public);
 }
 #endif
@@ -111,7 +116,17 @@ void FNiagaraEmitterHandle::SetName(FName InName, UNiagaraSystem& InOwnerSystem)
 	FName UniqueName = FNiagaraUtilities::GetUniqueName(InName, OtherEmitterNames);
 
 	Name = UniqueName;
-	Instance->SetUniqueEmitterName(Name.ToString());
+	if (Instance->SetUniqueEmitterName(Name.ToString()))
+	{
+#if WITH_EDITOR
+		if (InOwnerSystem.GetSystemSpawnScript() && InOwnerSystem.GetSystemSpawnScript()->GetSource())
+		{
+			// Just invalidate the system scripts here. The emitter scripts have their important variables 
+			// changed in the SetUniqueEmitterName method above.
+			InOwnerSystem.GetSystemSpawnScript()->GetSource()->InvalidateCachedCompileIds();
+		}
+#endif
+	}
 }
 
 bool FNiagaraEmitterHandle::GetIsEnabled() const
@@ -194,6 +209,9 @@ void FNiagaraEmitterHandle::ConditionalPostLoad()
 
 INiagaraModule::FMergeEmitterResults FNiagaraEmitterHandle::MergeSourceChanges()
 {
+	UE_LOG(LogNiagara, Log, TEXT("Emitter %s-%s is merging changes from source %s because its Change ID was updated."), *Instance->GetPathName(), *Name.ToString(),
+		Source != nullptr ? *Source->GetPathName() : TEXT("(null)"));
+
 	if (Source == nullptr)
 	{
 		// If we don't have a copy of the source emitter, this emitter can't safely be merged.
diff --git a/Engine/Plugins/FX/Niagara/Source/Niagara/Private/NiagaraEmitterInstance.cpp b/Engine/Plugins/FX/Niagara/Source/Niagara/Private/NiagaraEmitterInstance.cpp
--- a/Engine/Plugins/FX/Niagara/Source/Niagara/Private/NiagaraEmitterInstance.cpp
+++ b/Engine/Plugins/FX/Niagara/Source/Niagara/Private/NiagaraEmitterInstance.cpp
@@ -22,6 +22,7 @@ DECLARE_CYCLE_STAT(TEXT("Spawn"), STAT_NiagaraSpawn, STATGROUP_Niagara);
 DECLARE_CYCLE_STAT(TEXT("Spawn"), STAT_NiagaraEvents, STATGROUP_Niagara);
 DECLARE_CYCLE_STAT(TEXT("Kill"), STAT_NiagaraKill, STATGROUP_Niagara);
 DECLARE_CYCLE_STAT(TEXT("Event Handling"), STAT_NiagaraEventHandle, STATGROUP_Niagara);
+DECLARE_CYCLE_STAT(TEXT("Error Check"), STAT_NiagaraEmitterErrorCheck, STATGROUP_Niagara);
 
 static int32 GbDumpParticleData = 0;
 static FAutoConsoleVariableRef CVarNiagaraDumpParticleData(
@@ -33,12 +34,19 @@ static FAutoConsoleVariableRef CVarNiagaraDumpParticleData(
 
 //////////////////////////////////////////////////////////////////////////
 
+const FName FNiagaraEmitterInstance::PositionName(TEXT("Position"));
+const FName FNiagaraEmitterInstance::SizeName(TEXT("SpriteSize"));
+const FName FNiagaraEmitterInstance::MeshScaleName(TEXT("Scale"));
+
 FNiagaraEmitterInstance::FNiagaraEmitterInstance(FNiagaraSystemInstance* InParentSystemInstance)
-: bError(false)
-, CPUTimeMS(0.0f)
+: CPUTimeMS(0.0f)
 , ExecutionState(ENiagaraExecutionState::Inactive)
 , CachedBounds(ForceInit)
 , ParentSystemInstance(InParentSystemInstance)
+, CachedEmitter(nullptr)
+#if !UE_BUILD_SHIPPING
+, bEncounteredNaNs(false)
+#endif
 {
 	bDumpAfterEvent = false;
 	ParticleDataSet = new FNiagaraDataSet();
@@ -75,11 +83,11 @@ void FNiagaraEmitterInstance::ClearRenderer()
 
 FBox FNiagaraEmitterInstance::GetBounds()
 {
-	if (UNiagaraEmitter* Emitter = GetEmitterHandle().GetInstance())
+	checkSlow(CachedEmitter);
 	{
-		if (Emitter->bFixedBounds || Emitter->SimTarget == ENiagaraSimTarget::GPUComputeSim)//We must use fixed bounds on GPU sim.
+		if (CachedEmitter->bFixedBounds || CachedEmitter->SimTarget == ENiagaraSimTarget::GPUComputeSim)//We must use fixed bounds on GPU sim.
 		{
-			return Emitter->FixedBounds.TransformBy(ParentSystemInstance->GetComponent()->GetComponentToWorld());
+			return CachedEmitter->FixedBounds.TransformBy(ParentSystemInstance->GetComponent()->GetComponentToWorld());
 		}
 	}
 	return CachedBounds;
@@ -87,9 +95,7 @@ FBox FNiagaraEmitterInstance::GetBounds()
 
 bool FNiagaraEmitterInstance::IsReadyToRun() const
 {
-	const UNiagaraEmitter* PinnedProps = GetEmitterHandle().GetInstance();
-
-	if (!PinnedProps || !PinnedProps->IsReadyToRun())
+	if (!CachedEmitter->IsReadyToRun())
 	{
 		return false;
 	}
@@ -97,15 +103,31 @@ bool FNiagaraEmitterInstance::IsReadyToRun() const
 	return true;
 }
 
-
 void FNiagaraEmitterInstance::Init(int32 InEmitterIdx, FName InSystemInstanceName)
 {
 	check(ParticleDataSet);
 	FNiagaraDataSet& Data = *ParticleDataSet;
-
 	EmitterIdx = InEmitterIdx;
 	OwnerSystemInstanceName = InSystemInstanceName;
-	Data.Init(FNiagaraDataSetID(GetEmitterHandle().GetIdName(), ENiagaraDataSetType::ParticleData));
+	const FNiagaraEmitterHandle& EmitterHandle = GetEmitterHandle();
+	CachedEmitter = EmitterHandle.GetInstance();
+	checkSlow(CachedEmitter);
+	CachedIDName = EmitterHandle.GetIdName();
+
+	if (!EmitterHandle.GetIsEnabled()
+		|| !CachedEmitter->IsAllowedByDetailLevel()
+		|| (GMaxRHIFeatureLevel != ERHIFeatureLevel::SM5 && GMaxRHIFeatureLevel != ERHIFeatureLevel::ES3_1 && CachedEmitter->SimTarget == ENiagaraSimTarget::GPUComputeSim)  // skip if GPU sim and <SM5. TODO: fall back to CPU sim instead once we have scalability functionality to do so
+		)
+	{
+		ExecutionState = ENiagaraExecutionState::Disabled;
+		return;
+	}
+
+#if !UE_BUILD_SHIPPING
+	bEncounteredNaNs = false;
+#endif
+
+	Data.Init(FNiagaraDataSetID(CachedIDName, ENiagaraDataSetType::ParticleData), CachedEmitter->SimTarget);
 
 	//Init the spawn infos to the correct number for this system.
 	const TArray<FNiagaraEmitterSpawnAttributes>& EmitterSpawnInfoAttrs = ParentSystemInstance->GetSystem()->GetEmitterSpawnAttributes();
@@ -113,19 +135,17 @@ void FNiagaraEmitterInstance::Init(int32 InEmitterIdx, FName InSystemInstanceNam
 	{
 		SpawnInfos.SetNum(EmitterSpawnInfoAttrs[EmitterIdx].SpawnAttributes.Num());
 	}
-	const FNiagaraEmitterHandle& EmitterHandle = GetEmitterHandle();
 
-	ResetSimulation();
+	CheckForErrors();
 
 	if (IsDisabled())
 	{
 		return;
 	}
 
-	Data.Reset();
-	DataSetMap.Empty();
+	ResetSimulation();
 
-	const UNiagaraEmitter* PinnedProps = EmitterHandle.GetInstance();
+	DataSetMap.Empty();
 
 	//Add the particle data to the data set map.
 	//Currently just used for the tick loop but will also allow access directly to the particle data from other emitters.
@@ -137,78 +157,89 @@ void FNiagaraEmitterInstance::Init(int32 InEmitterIdx, FName InSystemInstanceNam
 
 	if (bVerboseAttributeLogging)
 	{
-		for (FNiagaraVariable& Attr : PinnedProps->UpdateScriptProps.Script->Attributes)
+		for (FNiagaraVariable& Attr : CachedEmitter->UpdateScriptProps.Script->GetVMExecutableData().Attributes)
 		{
 			int32 FoundIdx;
-			if (!PinnedProps->SpawnScriptProps.Script->Attributes.Find(Attr, FoundIdx))
+			if (!CachedEmitter->SpawnScriptProps.Script->GetVMExecutableData().Attributes.Find(Attr, FoundIdx))
 			{
 				UE_LOG(LogNiagara, Warning, TEXT("Attribute %s is used in the Update script for %s but it is not initialised in the Spawn script!"), *Attr.GetName().ToString(), *EmitterHandle.GetName().ToString());
 			}
-			for (int32 i = 0; i < PinnedProps->GetEventHandlers().Num(); i++)
+			for (int32 i = 0; i < CachedEmitter->GetEventHandlers().Num(); i++)
 			{
-				if (PinnedProps->GetEventHandlers()[i].Script && !PinnedProps->GetEventHandlers()[i].Script->Attributes.Find(Attr, FoundIdx))
+				if (CachedEmitter->GetEventHandlers()[i].Script && !CachedEmitter->GetEventHandlers()[i].Script->GetVMExecutableData().Attributes.Find(Attr, FoundIdx))
 				{
 					UE_LOG(LogNiagara, Warning, TEXT("Attribute %s is used in the event handler script for %s but it is not initialised in the Spawn script!"), *Attr.GetName().ToString(), *EmitterHandle.GetName().ToString());
 				}
 			}
 		}
 	}
-	Data.AddVariables(PinnedProps->UpdateScriptProps.Script->Attributes);
-	Data.AddVariables(PinnedProps->SpawnScriptProps.Script->Attributes);
+	Data.AddVariables(CachedEmitter->UpdateScriptProps.Script->GetVMExecutableData().Attributes);
+	Data.AddVariables(CachedEmitter->SpawnScriptProps.Script->GetVMExecutableData().Attributes);
+
+	//if we use persistent IDs then add that here too.
+	if (RequiredPersistentID())
+	{
+		Data.SetNeedsPersistentIDs(true);
+	}
+
 	Data.Finalize();
 
-	ensure(PinnedProps->UpdateScriptProps.DataSetAccessSynchronized());
+	ensure(CachedEmitter->UpdateScriptProps.DataSetAccessSynchronized());
 	UpdateScriptEventDataSets.Empty();
-	for (const FNiagaraEventGeneratorProperties &GeneratorProps : PinnedProps->UpdateScriptProps.EventGenerators)
+	for (const FNiagaraEventGeneratorProperties &GeneratorProps : CachedEmitter->UpdateScriptProps.EventGenerators)
 	{
 		FNiagaraDataSet *Set = FNiagaraEventDataSetMgr::CreateEventDataSet(ParentSystemInstance->GetIDName(), EmitterHandle.GetIdName(), GeneratorProps.SetProps.ID.Name);
-		Set->Reset();
+		Set->Init(FNiagaraDataSetID(), ENiagaraSimTarget::CPUSim);
 		Set->AddVariables(GeneratorProps.SetProps.Variables);
 		Set->Finalize();
 		UpdateScriptEventDataSets.Add(Set);
 	}
 
-	ensure(PinnedProps->SpawnScriptProps.DataSetAccessSynchronized());
+	ensure(CachedEmitter->SpawnScriptProps.DataSetAccessSynchronized());
 	SpawnScriptEventDataSets.Empty();
-	for (const FNiagaraEventGeneratorProperties &GeneratorProps : PinnedProps->SpawnScriptProps.EventGenerators)
+	for (const FNiagaraEventGeneratorProperties &GeneratorProps : CachedEmitter->SpawnScriptProps.EventGenerators)
 	{
 		FNiagaraDataSet *Set = FNiagaraEventDataSetMgr::CreateEventDataSet(ParentSystemInstance->GetIDName(), EmitterHandle.GetIdName(), GeneratorProps.SetProps.ID.Name);
-		Set->Reset();
+		Set->Init(FNiagaraDataSetID(), ENiagaraSimTarget::CPUSim);
 		Set->AddVariables(GeneratorProps.SetProps.Variables);
 		Set->Finalize();
 		SpawnScriptEventDataSets.Add(Set);
 	}
 
-	SpawnExecContext.Init(PinnedProps->SpawnScriptProps.Script, PinnedProps->SimTarget);
-	UpdateExecContext.Init(PinnedProps->UpdateScriptProps.Script, PinnedProps->SimTarget);
-	
+	SpawnExecContext.Init(CachedEmitter->SpawnScriptProps.Script, CachedEmitter->SimTarget);
+	UpdateExecContext.Init(CachedEmitter->UpdateScriptProps.Script, CachedEmitter->SimTarget);
+
 	// setup the parameer store for the GPU execution context; since spawn and update are combined here, we build one with params from both script props
-	if (PinnedProps->SimTarget == ENiagaraSimTarget::GPUComputeSim)
+	if (CachedEmitter->SimTarget == ENiagaraSimTarget::GPUComputeSim)
 	{
-		GPUExecContext.InitParams(PinnedProps->SpawnScriptProps.Script, PinnedProps->UpdateScriptProps.Script, PinnedProps->SimTarget);
+		GPUExecContext.InitParams(CachedEmitter->GetGPUComputeScript(), CachedEmitter->SpawnScriptProps.Script, CachedEmitter->UpdateScriptProps.Script, CachedEmitter->SimTarget);
 		SpawnExecContext.Parameters.Bind(&GPUExecContext.CombinedParamStore);
 		UpdateExecContext.Parameters.Bind(&GPUExecContext.CombinedParamStore);
 	}
 
-	EventExecContexts.SetNum(PinnedProps->GetEventHandlers().Num());
-	int32 NumEvents = PinnedProps->GetEventHandlers().Num();
+	EventExecContexts.SetNum(CachedEmitter->GetEventHandlers().Num());
+	int32 NumEvents = CachedEmitter->GetEventHandlers().Num();
 	for (int32 i = 0; i < NumEvents; i++)
 	{
-		ensure(PinnedProps->GetEventHandlers()[i].DataSetAccessSynchronized());
+		ensure(CachedEmitter->GetEventHandlers()[i].DataSetAccessSynchronized());
 
-		UNiagaraScript* EventScript = PinnedProps->GetEventHandlers()[i].Script;
+		UNiagaraScript* EventScript = CachedEmitter->GetEventHandlers()[i].Script;
 
 		//This is cpu explicitly? Are we doing event handlers on GPU?
 		EventExecContexts[i].Init(EventScript, ENiagaraSimTarget::CPUSim);
 	}
 
-	UNiagaraEmitter* Emitter = GetEmitterHandle().GetInstance();
-
 	//Setup direct bindings for setting parameter values.
-	SpawnIntervalBinding.Init(SpawnExecContext.Parameters, Emitter->ToEmitterParameter(SYS_PARAM_EMITTER_SPAWN_INTERVAL));
-	InterpSpawnStartBinding.Init(SpawnExecContext.Parameters, Emitter->ToEmitterParameter(SYS_PARAM_EMITTER_INTERP_SPAWN_START_DT));
+	SpawnIntervalBinding.Init(SpawnExecContext.Parameters, CachedEmitter->ToEmitterParameter(SYS_PARAM_EMITTER_SPAWN_INTERVAL));
+	InterpSpawnStartBinding.Init(SpawnExecContext.Parameters, CachedEmitter->ToEmitterParameter(SYS_PARAM_EMITTER_INTERP_SPAWN_START_DT));
+
+	if (CachedEmitter->SimTarget == ENiagaraSimTarget::GPUComputeSim)
+	{
+		SpawnIntervalBindingGPU.Init(GPUExecContext.CombinedParamStore, CachedEmitter->ToEmitterParameter(SYS_PARAM_EMITTER_SPAWN_INTERVAL));
+		InterpSpawnStartBindingGPU.Init(GPUExecContext.CombinedParamStore, CachedEmitter->ToEmitterParameter(SYS_PARAM_EMITTER_INTERP_SPAWN_START_DT));
+	}
 
-	FNiagaraVariable EmitterAgeParam = Emitter->ToEmitterParameter(SYS_PARAM_EMITTER_AGE);
+	FNiagaraVariable EmitterAgeParam = CachedEmitter->ToEmitterParameter(SYS_PARAM_EMITTER_AGE);
 	SpawnEmitterAgeBinding.Init(SpawnExecContext.Parameters, EmitterAgeParam);
 	UpdateEmitterAgeBinding.Init(UpdateExecContext.Parameters, EmitterAgeParam);
 	EventEmitterAgeBindings.SetNum(NumEvents);
@@ -217,7 +248,12 @@ void FNiagaraEmitterInstance::Init(int32 InEmitterIdx, FName InSystemInstanceNam
 		EventEmitterAgeBindings[i].Init(EventExecContexts[i].Parameters, EmitterAgeParam);
 	}
 
-	FNiagaraVariable EmitterLocalSpaceParam = Emitter->ToEmitterParameter(SYS_PARAM_EMITTER_LOCALSPACE);
+	if (CachedEmitter->SimTarget == ENiagaraSimTarget::GPUComputeSim)
+	{
+		EmitterAgeBindingGPU.Init(GPUExecContext.CombinedParamStore, EmitterAgeParam);
+	}
+
+	FNiagaraVariable EmitterLocalSpaceParam = CachedEmitter->ToEmitterParameter(SYS_PARAM_EMITTER_LOCALSPACE);
 	SpawnEmitterLocalSpaceBinding.Init(SpawnExecContext.Parameters, EmitterLocalSpaceParam);
 	UpdateEmitterLocalSpaceBinding.Init(UpdateExecContext.Parameters, EmitterLocalSpaceParam);
 	EventEmitterLocalSpaceBindings.SetNum(NumEvents);
@@ -225,6 +261,10 @@ void FNiagaraEmitterInstance::Init(int32 InEmitterIdx, FName InSystemInstanceNam
 	{
 		EventEmitterLocalSpaceBindings[i].Init(EventExecContexts[i].Parameters, EmitterLocalSpaceParam);
 	}
+	if (CachedEmitter->SimTarget == ENiagaraSimTarget::GPUComputeSim)
+	{
+		EmitterLocalSpaceBindingGPU.Init(GPUExecContext.CombinedParamStore, EmitterLocalSpaceParam);
+	}
 
 	SpawnExecCountBinding.Init(SpawnExecContext.Parameters, SYS_PARAM_ENGINE_EXEC_COUNT);
 	UpdateExecCountBinding.Init(UpdateExecContext.Parameters, SYS_PARAM_ENGINE_EXEC_COUNT);
@@ -234,16 +274,28 @@ void FNiagaraEmitterInstance::Init(int32 InEmitterIdx, FName InSystemInstanceNam
 		EventExecCountBindings[i].Init(EventExecContexts[i].Parameters, SYS_PARAM_ENGINE_EXEC_COUNT);
 	}
 
-	if (PinnedProps->SimTarget == ENiagaraSimTarget::GPUComputeSim)
+	if (CachedEmitter->SimTarget == ENiagaraSimTarget::GPUComputeSim)
 	{
 		//Just ensure we've generated the singleton here on the GT as it throws a wobbler if we do this later in parallel.
 		NiagaraEmitterInstanceBatcher::Get();
 	}
+	else
+	{
+		//Init accessors for PostProcessParticles
+		PositionAccessor = FNiagaraDataSetAccessor<FVector>(Data, FNiagaraVariable(FNiagaraTypeDefinition::GetVec3Def(), PositionName));
+		SizeAccessor = FNiagaraDataSetAccessor<FVector2D>(Data, FNiagaraVariable(FNiagaraTypeDefinition::GetVec2Def(), SizeName));
+		MeshScaleAccessor = FNiagaraDataSetAccessor<FVector>(Data, FNiagaraVariable(FNiagaraTypeDefinition::GetVec3Def(), MeshScaleName));
+	}
 
-	//Init accessors for PostProcessParticles
-	PositionAccessor = FNiagaraDataSetAccessor<FVector>(Data, FNiagaraVariable(FNiagaraTypeDefinition::GetVec3Def(), TEXT("Position")));
-	SizeAccessor = FNiagaraDataSetAccessor<FVector2D>(Data, FNiagaraVariable(FNiagaraTypeDefinition::GetVec2Def(), TEXT("SpriteSize")));
-	MeshScaleAccessor = FNiagaraDataSetAccessor<FVector>(Data, FNiagaraVariable(FNiagaraTypeDefinition::GetVec3Def(), TEXT("Scale")));
+	// Collect script defined data interface parameters.
+	TArray<UNiagaraScript*> Scripts;
+	Scripts.Add(CachedEmitter->SpawnScriptProps.Script);
+	Scripts.Add(CachedEmitter->UpdateScriptProps.Script);
+	for (const FNiagaraEventScriptProperties& EventHandler : CachedEmitter->GetEventHandlers())
+	{
+		Scripts.Add(EventHandler.Script);
+	}
+	FNiagaraUtilities::CollectScriptDataInterfaceParameters(*CachedEmitter, Scripts, ScriptDefinedDataInterfaceParameters);
 }
 
 void FNiagaraEmitterInstance::ResetSimulation()
@@ -251,70 +303,80 @@ void FNiagaraEmitterInstance::ResetSimulation()
 	bResetPending = true;
 	Age = 0;
 	Loops = 0;
-	bError = false;
-	
-	SetExecutionState(ENiagaraExecutionState::Active);
+	TickCount = 0;
+	CachedBounds.Init();
 
-	if (IsDisabled())
+	ParticleDataSet->ResetBuffers();
+	for (FNiagaraDataSet* SpawnScriptEventDataSet : SpawnScriptEventDataSets)
 	{
-		return;
+		SpawnScriptEventDataSet->ResetBuffers();
 	}
+	for (FNiagaraDataSet* UpdateScriptEventDataSet : UpdateScriptEventDataSets)
+	{
+		UpdateScriptEventDataSet->ResetBuffers();
+	}
+	
+	GPUExecContext.Reset();
 
-	CheckForErrors();
+	SetExecutionState(ENiagaraExecutionState::Active);
 }
 
 void FNiagaraEmitterInstance::CheckForErrors()
 {
-	const UNiagaraEmitter* PinnedProps = GetEmitterHandle().GetInstance();
-	if (!PinnedProps)
+	SCOPE_CYCLE_COUNTER(STAT_NiagaraEmitterErrorCheck);
+	
+	checkSlow(CachedEmitter);
+
+	//Check for various failure conditions and bail.
+	if (!CachedEmitter->UpdateScriptProps.Script || !CachedEmitter->SpawnScriptProps.Script )
 	{
-		UE_LOG(LogNiagara, Error, TEXT("Unknown Error creating Niagara Simulation. Properties were null."));
-		bError = true;
+		//TODO - Arbitrary named scripts. Would need some base functionality for Spawn/Udpate to be called that can be overriden in BPs for emitters with custom scripts.
+		UE_LOG(LogNiagara, Error, TEXT("Emitter cannot be enabled because it's doesn't have both an update and spawn script."), *CachedEmitter->GetFullName());
+		SetExecutionState(ENiagaraExecutionState::Disabled);
 		return;
 	}
 
-	//Check for various failure conditions and bail.
-	if (!PinnedProps->UpdateScriptProps.Script || !PinnedProps->SpawnScriptProps.Script)
+	if (!CachedEmitter->UpdateScriptProps.Script->IsReadyToRun(ENiagaraSimTarget::CPUSim) || !CachedEmitter->SpawnScriptProps.Script->IsReadyToRun(ENiagaraSimTarget::CPUSim))
 	{
 		//TODO - Arbitrary named scripts. Would need some base functionality for Spawn/Udpate to be called that can be overriden in BPs for emitters with custom scripts.
-		UE_LOG(LogNiagara, Error, TEXT("Emitter cannot be enabled because it's doesn't have both an update and spawn script."), *PinnedProps->GetFullName());
-		bError = true;
+		UE_LOG(LogNiagara, Error, TEXT("Emitter cannot be enabled because it's doesn't have both an update and spawn script ready to run CPU scripts."), *CachedEmitter->GetFullName());
+		SetExecutionState(ENiagaraExecutionState::Disabled);
 		return;
 	}
 
-	if (PinnedProps->SpawnScriptProps.Script->DataUsage.bReadsAttriubteData)
+	if (CachedEmitter->SpawnScriptProps.Script->GetVMExecutableData().DataUsage.bReadsAttributeData)
 	{
-		UE_LOG(LogNiagara, Error, TEXT("%s reads attribute data and so cannot be used as a spawn script. The data being read would be invalid."), *PinnedProps->SpawnScriptProps.Script->GetName());
-		bError = true;
+		UE_LOG(LogNiagara, Error, TEXT("%s reads attribute data and so cannot be used as a spawn script. The data being read would be invalid."), *CachedEmitter->SpawnScriptProps.Script->GetName());
+		SetExecutionState(ENiagaraExecutionState::Disabled);
 		return;
 	}
-	if (PinnedProps->UpdateScriptProps.Script->Attributes.Num() == 0 || PinnedProps->SpawnScriptProps.Script->Attributes.Num() == 0)
+	if (CachedEmitter->UpdateScriptProps.Script->GetVMExecutableData().Attributes.Num() == 0 || CachedEmitter->SpawnScriptProps.Script->GetVMExecutableData().Attributes.Num() == 0)
 	{
 		UE_LOG(LogNiagara, Error, TEXT("This emitter cannot be enabled because it's spawn or update script doesn't have any attriubtes.."));
-		bError = true;
+		SetExecutionState(ENiagaraExecutionState::Disabled);
 		return;
 	}
 
-	if (PinnedProps->SimTarget == ENiagaraSimTarget::CPUSim || PinnedProps->SimTarget == ENiagaraSimTarget::DynamicLoadBalancedSim)
+	if (CachedEmitter->SimTarget == ENiagaraSimTarget::CPUSim || CachedEmitter->SimTarget == ENiagaraSimTarget::DynamicLoadBalancedSim)
 	{
 		bool bFailed = false;
-		if (!PinnedProps->SpawnScriptProps.Script->DidScriptCompilationSucceed(false))
+		if (!CachedEmitter->SpawnScriptProps.Script->DidScriptCompilationSucceed(false))
 		{
 			bFailed = true;
 			UE_LOG(LogNiagara, Error, TEXT("This emitter cannot be enabled because it's CPU Spawn script failed to compile."));
 		}
 
-		if (!PinnedProps->UpdateScriptProps.Script->DidScriptCompilationSucceed(false))
+		if (!CachedEmitter->UpdateScriptProps.Script->DidScriptCompilationSucceed(false))
 		{
 			bFailed = true;
 			UE_LOG(LogNiagara, Error, TEXT("This emitter cannot be enabled because it's CPU Update script failed to compile."));
 		}
 
-		if (PinnedProps->GetEventHandlers().Num() != 0)
+		if (CachedEmitter->GetEventHandlers().Num() != 0)
 		{
-			for (int32 i = 0; i < PinnedProps->GetEventHandlers().Num(); i++)
+			for (int32 i = 0; i < CachedEmitter->GetEventHandlers().Num(); i++)
 			{
-				if (!PinnedProps->GetEventHandlers()[i].Script->DidScriptCompilationSucceed(false))
+				if (!CachedEmitter->GetEventHandlers()[i].Script->DidScriptCompilationSucceed(false))
 				{
 					bFailed = true;
 					UE_LOG(LogNiagara, Error, TEXT("This emitter cannot be enabled because one of it's CPU Event scripts failed to compile."));
@@ -324,23 +386,23 @@ void FNiagaraEmitterInstance::CheckForErrors()
 
 		if (bFailed)
 		{
-			bError = true;
+			SetExecutionState(ENiagaraExecutionState::Disabled);
 			return;
 		}
 	}
 
-	if (PinnedProps->SimTarget == ENiagaraSimTarget::GPUComputeSim || PinnedProps->SimTarget == ENiagaraSimTarget::DynamicLoadBalancedSim)
+	if (CachedEmitter->SimTarget == ENiagaraSimTarget::GPUComputeSim || CachedEmitter->SimTarget == ENiagaraSimTarget::DynamicLoadBalancedSim)
 	{
-		if (PinnedProps->SpawnScriptProps.Script->IsScriptCompilationPending(true))
+		if (CachedEmitter->GetGPUComputeScript()->IsScriptCompilationPending(true))
 		{
 			UE_LOG(LogNiagara, Error, TEXT("This emitter cannot be enabled because it's GPU script hasn't been compiled.."));
-			bError = true;
+			SetExecutionState(ENiagaraExecutionState::Disabled);
 			return;
 		}
-		if (!PinnedProps->SpawnScriptProps.Script->DidScriptCompilationSucceed(true))
+		if (!CachedEmitter->GetGPUComputeScript()->DidScriptCompilationSucceed(true))
 		{
 			UE_LOG(LogNiagara, Error, TEXT("This emitter cannot be enabled because it's GPU script failed to compile."));
-			bError = true;
+			SetExecutionState(ENiagaraExecutionState::Disabled);
 			return;
 		}
 	}
@@ -396,18 +458,18 @@ void FNiagaraEmitterInstance::BindParameters()
 	FNiagaraWorldManager* WorldMan = ParentSystemInstance->GetWorldManager();
 	check(WorldMan);
 
-	for (UNiagaraParameterCollection* Collection : SpawnExecContext.Script->ParameterCollections)
+	for (UNiagaraParameterCollection* Collection : SpawnExecContext.Script->GetCachedParameterCollectionReferences())
 	{
 		ParentSystemInstance->GetParameterCollectionInstance(Collection)->GetParameterStore().Bind(&SpawnExecContext.Parameters);
 	}
-	for (UNiagaraParameterCollection* Collection : UpdateExecContext.Script->ParameterCollections)
+	for (UNiagaraParameterCollection* Collection : UpdateExecContext.Script->GetCachedParameterCollectionReferences())
 	{
 		ParentSystemInstance->GetParameterCollectionInstance(Collection)->GetParameterStore().Bind(&UpdateExecContext.Parameters);
 	}
 
 	for (int32 EventIdx = 0; EventIdx < EventExecContexts.Num(); ++EventIdx)
 	{
-		for (UNiagaraParameterCollection* Collection : EventExecContexts[EventIdx].Script->ParameterCollections)
+		for (UNiagaraParameterCollection* Collection : EventExecContexts[EventIdx].Script->GetCachedParameterCollectionReferences())
 		{
 			ParentSystemInstance->GetParameterCollectionInstance(Collection)->GetParameterStore().Bind(&EventExecContexts[EventIdx].Parameters);
 		}
@@ -415,41 +477,49 @@ void FNiagaraEmitterInstance::BindParameters()
 
 	//Now bind parameters from the component and system.
 	FNiagaraParameterStore& InstanceParams = ParentSystemInstance->GetParameters();
+	FNiagaraParameterStore& SystemScriptDefinedDataInterfaceParameters = ParentSystemInstance->GetSystemSimulation()->GetScriptDefinedDataInterfaceParameters();
+	
 	InstanceParams.Bind(&SpawnExecContext.Parameters);
+	SystemScriptDefinedDataInterfaceParameters.Bind(&SpawnExecContext.Parameters);
+	ScriptDefinedDataInterfaceParameters.Bind(&SpawnExecContext.Parameters);
+
 	InstanceParams.Bind(&UpdateExecContext.Parameters);
+	SystemScriptDefinedDataInterfaceParameters.Bind(&UpdateExecContext.Parameters);
+	ScriptDefinedDataInterfaceParameters.Bind(&UpdateExecContext.Parameters);
+
 	for (FNiagaraScriptExecutionContext& EventContext : EventExecContexts)
 	{
 		InstanceParams.Bind(&EventContext.Parameters);
+		SystemScriptDefinedDataInterfaceParameters.Bind(&EventContext.Parameters);
+		ScriptDefinedDataInterfaceParameters.Bind(&EventContext.Parameters);
 	}
 
-	const FNiagaraEmitterHandle& EmitterHandle = GetEmitterHandle();
-	const UNiagaraEmitter* Props = EmitterHandle.GetInstance();
-	Props->SpawnScriptProps.Script->RapidIterationParameters.Bind(&SpawnExecContext.Parameters);
-	Props->UpdateScriptProps.Script->RapidIterationParameters.Bind(&UpdateExecContext.Parameters);
-	ensure(Props->GetEventHandlers().Num() == EventExecContexts.Num());
-	for (int32 i = 0; i < Props->GetEventHandlers().Num(); i++)
+#if WITH_EDITORONLY_DATA
+	CachedEmitter->SpawnScriptProps.Script->RapidIterationParameters.Bind(&SpawnExecContext.Parameters);
+	CachedEmitter->UpdateScriptProps.Script->RapidIterationParameters.Bind(&UpdateExecContext.Parameters);
+	ensure(CachedEmitter->GetEventHandlers().Num() == EventExecContexts.Num());
+	for (int32 i = 0; i < CachedEmitter->GetEventHandlers().Num(); i++)
 	{
-		Props->GetEventHandlers()[i].Script->RapidIterationParameters.Bind(&EventExecContexts[i].Parameters);
+		CachedEmitter->GetEventHandlers()[i].Script->RapidIterationParameters.Bind(&EventExecContexts[i].Parameters);
 	}
+#endif
 }
 
 void FNiagaraEmitterInstance::PostInitSimulation()
 {
-	const FNiagaraEmitterHandle& EmitterHandle = GetEmitterHandle();
-	if (!bError)
+	if (!IsDisabled())
 	{
 		check(ParentSystemInstance);
-		const UNiagaraEmitter* Props = EmitterHandle.GetInstance();
 
 		//Go through all our receivers and grab their generator sets so that the source emitters can do any init work they need to do.
-		for (const FNiagaraEventReceiverProperties& Receiver : Props->SpawnScriptProps.EventReceivers)
+		for (const FNiagaraEventReceiverProperties& Receiver : CachedEmitter->SpawnScriptProps.EventReceivers)
 		{
 			//FNiagaraDataSet* ReceiverSet = ParentSystemInstance->GetDataSet(FNiagaraDataSetID(Receiver.SourceEventGenerator, ENiagaraDataSetType::Event), Receiver.SourceEmitter);
 			const FNiagaraDataSet* ReceiverSet = FNiagaraEventDataSetMgr::GetEventDataSet(ParentSystemInstance->GetIDName(), Receiver.SourceEmitter, Receiver.SourceEventGenerator);
 
 		}
 
-		for (const FNiagaraEventReceiverProperties& Receiver : Props->UpdateScriptProps.EventReceivers)
+		for (const FNiagaraEventReceiverProperties& Receiver : CachedEmitter->UpdateScriptProps.EventReceivers)
 		{
 			//FNiagaraDataSet* ReceiverSet = ParentSystemInstance->GetDataSet(FNiagaraDataSetID(Receiver.SourceEventGenerator, ENiagaraDataSetType::Event), Receiver.SourceEmitter);
 			const FNiagaraDataSet* ReceiverSet = FNiagaraEventDataSetMgr::GetEventDataSet(ParentSystemInstance->GetIDName(), Receiver.SourceEmitter, Receiver.SourceEventGenerator);
@@ -475,7 +545,9 @@ FNiagaraDataSet* FNiagaraEmitterInstance::GetDataSet(FNiagaraDataSetID SetID)
 
 const FNiagaraEmitterHandle& FNiagaraEmitterInstance::GetEmitterHandle() const
 {
-	return ParentSystemInstance->GetSystem()->GetEmitterHandles()[EmitterIdx];
+	UNiagaraSystem* Sys = ParentSystemInstance->GetSystem();
+	checkSlow(Sys->GetEmitterHandles().Num() > EmitterIdx);
+	return Sys->GetEmitterHandles()[EmitterIdx];
 }
 
 float FNiagaraEmitterInstance::GetTotalCPUTime()
@@ -508,29 +580,28 @@ int FNiagaraEmitterInstance::GetTotalBytesUsed()
 
 FBox FNiagaraEmitterInstance::CalculateDynamicBounds()
 {
-	check(ParticleDataSet);
+	checkSlow(ParticleDataSet);
 	FNiagaraDataSet& Data = *ParticleDataSet;
-	int32 NumParticles = Data.GetNumInstances();
-	const FNiagaraEmitterHandle& EmitterHandle = GetEmitterHandle();
-	if (NumParticles == 0 || EmitterHandle.GetInstance()->SimTarget == ENiagaraSimTarget::GPUComputeSim)//TODO: Pull data back from gpu buffers to get bounds for GPU sims.
-	{
-		return FBox();
-	}
-
-	const UNiagaraEmitter* Emitter = EmitterHandle.GetInstance();
+	int32 NumInstances = Data.GetNumInstances();
 	FBox Ret;
 	Ret.Init();
+	if (IsComplete() || NumInstances == 0 || CachedEmitter->SimTarget == ENiagaraSimTarget::GPUComputeSim)//TODO: Pull data back from gpu buffers to get bounds for GPU sims.
+	{
+		return Ret;
+	}
 
 	PositionAccessor.InitForAccess(true);
 	SizeAccessor.InitForAccess(true);
 	MeshScaleAccessor.InitForAccess(true);
-	int32 NumInstances = Data.GetNumInstances();
 
 	FVector MaxSize(EForceInit::ForceInitToZero);
 
-	FMatrix SystemLocalToWorld = GetParentSystemInstance()->GetComponent()->GetComponentTransform().ToMatrixNoScale();
+	if (SizeAccessor.IsValid() == false && MeshScaleAccessor.IsValid() == false)
+	{
+		MaxSize = FVector(50.0f, 50.0f, 50.0f);
+	}
 
-	for (int32 InstIdx = 0; InstIdx < NumInstances; ++InstIdx)
+	for (int32 InstIdx = 0; InstIdx < NumInstances && PositionAccessor.IsValid(); ++InstIdx)
 	{
 		FVector Position;
 		PositionAccessor.Get(InstIdx, Position);
@@ -553,25 +624,38 @@ FBox FNiagaraEmitterInstance::CalculateDynamicBounds()
 		}
 		else
 		{
-			UE_LOG(LogNiagara, Warning, TEXT("Particle position data contains NaNs. Likely a divide by zero somewhere in your modules."));
+#if !UE_BUILD_SHIPPING
+			if (bEncounteredNaNs == false && ParentSystemInstance != nullptr && CachedEmitter != nullptr && ParentSystemInstance->GetSystem() != nullptr)
+			{
+				UE_LOG(LogNiagara, Warning, TEXT("Particle position data contains NaNs. Likely a divide by zero somewhere in your modules. Emitter \"%s\" in System \"%s\""),
+					*CachedEmitter->GetName(), *ParentSystemInstance->GetSystem()->GetName());
+				bEncounteredNaNs = true;
+			}
+#endif
 		}
 	}
 
-	FVector MaxBaseSize = FVector(0.0001f, 0.0001f, 0.0001f);
+	float MaxBaseSize = 0.0001f;
+	if (MaxSize.IsNearlyZero())
+	{
+		MaxSize = FVector(1.0f, 1.0f, 1.0f);
+	}
+
 	for (int32 i = 0; i < EmitterRenderer.Num(); i++)
 	{
 		if (EmitterRenderer[i])
 		{
 			FVector BaseExtents = EmitterRenderer[i]->GetBaseExtents();
 			FVector ComponentMax;
-			MaxBaseSize = BaseExtents.ComponentMax(MaxBaseSize);
+			MaxBaseSize = BaseExtents.ComponentMax(FVector(MaxBaseSize, MaxBaseSize, MaxBaseSize)).GetMax();
 		}
 	}
 
 	Ret = Ret.ExpandBy(MaxSize*MaxBaseSize);
 
-	if (Emitter->bLocalSpace)
+	if (CachedEmitter->bLocalSpace || PositionAccessor.IsValid() == false)
 	{
+		FMatrix SystemLocalToWorld = GetParentSystemInstance()->GetComponent()->GetComponentTransform().ToMatrixNoScale();
 		Ret = Ret.TransformBy(SystemLocalToWorld);
 	}
 
@@ -584,25 +668,17 @@ FBox FNiagaraEmitterInstance::CalculateDynamicBounds()
 void FNiagaraEmitterInstance::PostProcessParticles()
 {
 	SCOPE_CYCLE_COUNTER(STAT_NiagaraKill);
-	check(ParticleDataSet);
-	FNiagaraDataSet& Data = *ParticleDataSet;
-	int32 NumParticles = Data.GetNumInstances();
 
-	CachedBounds.Init();
-
-	const FNiagaraEmitterHandle& EmitterHandle = GetEmitterHandle();
-	if (NumParticles == 0 || EmitterHandle.GetInstance()->SimTarget == ENiagaraSimTarget::GPUComputeSim)
+	checkSlow(CachedEmitter);
+	if (CachedEmitter->bFixedBounds)
 	{
-		return;
-	}
-	
-	const UNiagaraEmitter* Emitter = EmitterHandle.GetInstance();
-	if (Emitter->bFixedBounds)
-	{
-		CachedBounds = Emitter->FixedBounds;
+		CachedBounds = CachedEmitter->FixedBounds;
 	}
 	else
 	{
+		checkSlow(ParticleDataSet);
+		FNiagaraDataSet& Data = *ParticleDataSet;
+		int32 NumParticles = Data.GetNumInstances();
 		CachedBounds = CalculateDynamicBounds();
 	}
 }
@@ -619,7 +695,7 @@ bool FNiagaraEmitterInstance::HandleCompletion(bool bForce)
 		//If we have any particles then clear out the buffers.
 		if (ParticleDataSet->GetNumInstances() > 0 || ParticleDataSet->GetPrevNumInstances() > 0)
 		{
-			ParticleDataSet->ResetBuffers(true);
+			ParticleDataSet->ResetBuffers();
 		}
 		return true;
 	}
@@ -627,34 +703,40 @@ bool FNiagaraEmitterInstance::HandleCompletion(bool bForce)
 	return false;
 }
 
+bool FNiagaraEmitterInstance::RequiredPersistentID()const
+{
+	//TODO: can we have this be enabled at runtime from outside the system?
+	return GetEmitterHandle().GetInstance()->RequiresPersistantIDs();
+}
+
 /** 
   * PreTick - handles killing dead particles, emitter death, and buffer swaps
   */
 void FNiagaraEmitterInstance::PreTick()
 {
-	const FNiagaraEmitterHandle& EmitterHandle = GetEmitterHandle();
-	const UNiagaraEmitter* PinnedProps = EmitterHandle.GetInstance();
-
-	if (!PinnedProps || bError || IsComplete())
+	if (IsComplete())
 	{
 		return;
 	}
 
-	check(ParticleDataSet);
+	checkSlow(ParticleDataSet);
 	FNiagaraDataSet& Data = *ParticleDataSet;
 
-	PinnedProps->SpawnScriptProps.Script->RapidIterationParameters.Tick();
-	PinnedProps->UpdateScriptProps.Script->RapidIterationParameters.Tick();
-	ensure(PinnedProps->GetEventHandlers().Num() == EventExecContexts.Num());
-	for (int32 i = 0; i < PinnedProps->GetEventHandlers().Num(); i++)
+#if WITH_EDITOR
+	CachedEmitter->SpawnScriptProps.Script->RapidIterationParameters.Tick();
+	CachedEmitter->UpdateScriptProps.Script->RapidIterationParameters.Tick();
+	ensure(CachedEmitter->GetEventHandlers().Num() == EventExecContexts.Num());
+	for (int32 i = 0; i < CachedEmitter->GetEventHandlers().Num(); i++)
 	{
-		PinnedProps->GetEventHandlers()[i].Script->RapidIterationParameters.Tick();
+		CachedEmitter->GetEventHandlers()[i].Script->RapidIterationParameters.Tick();
 	}
+#endif
+
 
 	bool bOk = true;
 	bOk &= SpawnExecContext.Tick(ParentSystemInstance);
 	bOk &= UpdateExecContext.Tick(ParentSystemInstance);
-	if (PinnedProps->SimTarget == ENiagaraSimTarget::GPUComputeSim)
+	if (CachedEmitter->SimTarget == ENiagaraSimTarget::GPUComputeSim)
 	{
 		bOk &= GPUExecContext.Tick(ParentSystemInstance);
 	}
@@ -666,38 +748,49 @@ void FNiagaraEmitterInstance::PreTick()
 	if (!bOk)
 	{
 		ResetSimulation();
-		bError = true;
+		SetExecutionState(ENiagaraExecutionState::Disabled);
 		return;
 	}
 
-	check(Data.GetNumVariables() > 0);
-	check(PinnedProps->SpawnScriptProps.Script);
-	check(PinnedProps->UpdateScriptProps.Script);
+	checkSlow(Data.GetNumVariables() > 0);
+	checkSlow(CachedEmitter->SpawnScriptProps.Script);
+	checkSlow(CachedEmitter->UpdateScriptProps.Script);
 
 	if (bResetPending)
 	{
-		Data.SetNumInstances(0);
+		Data.ResetCurrentBuffers();
+		for (FNiagaraDataSet* SpawnScriptEventDataSet : SpawnScriptEventDataSets)
+		{
+			SpawnScriptEventDataSet->ResetCurrentBuffers();
+		}
+		for (FNiagaraDataSet* UpdateScriptEventDataSet : UpdateScriptEventDataSets)
+		{
+			UpdateScriptEventDataSet->ResetCurrentBuffers();
+		}
 		bResetPending = false;
 	}
 
 	//Swap all data set buffers before doing the main tick on any simulation.
-	if (PinnedProps->SimTarget == ENiagaraSimTarget::CPUSim)
+	if (CachedEmitter->SimTarget == ENiagaraSimTarget::CPUSim)
 	{
 		for (TPair<FNiagaraDataSetID, FNiagaraDataSet*> SetPair : DataSetMap)
 		{
-			SetPair.Value->Tick(PinnedProps->SimTarget);
+			SetPair.Value->Tick();
 		}
 
 		for (FNiagaraDataSet* Set : UpdateScriptEventDataSets)
 		{
-			Set->Tick(PinnedProps->SimTarget);
+			Set->Tick();
 		}
 
 		for (FNiagaraDataSet* Set : SpawnScriptEventDataSets)
 		{
-			Set->Tick(PinnedProps->SimTarget);
+			Set->Tick();
 		}
 	}
+
+	++TickCount;
+	ParticleDataSet->SetIDAcquireTag(TickCount);
 }
 
 
@@ -706,27 +799,21 @@ void FNiagaraEmitterInstance::Tick(float DeltaSeconds)
 	SCOPE_CYCLE_COUNTER(STAT_NiagaraTick);
 	SimpleTimer TickTime;
 
-	const FNiagaraEmitterHandle& EmitterHandle = GetEmitterHandle();
-	const UNiagaraEmitter* PinnedProps = EmitterHandle.GetInstance();
-	if (!PinnedProps || bError)
-	{
-		return;
-	}
-
-	check(ParticleDataSet);
-	FNiagaraDataSet& Data = *ParticleDataSet;
-
 	if (HandleCompletion())
 	{
+		CPUTimeMS = TickTime.GetElapsedMilliseconds();
 		return;
 	}
 
+	checkSlow(ParticleDataSet);
+	FNiagaraDataSet& Data = *ParticleDataSet;
 	Age += DeltaSeconds;
 
 	if (ExecutionState == ENiagaraExecutionState::InactiveClear)
 	{
-		Data.ResetBuffers(true);
+		Data.ResetBuffers();
 		ExecutionState = ENiagaraExecutionState::Inactive;
+		CPUTimeMS = TickTime.GetElapsedMilliseconds();
 		return;
 	}
 
@@ -734,13 +821,16 @@ void FNiagaraEmitterInstance::Tick(float DeltaSeconds)
 	if (OrigNumParticles == 0 && ExecutionState != ENiagaraExecutionState::Active)
 	{
 		//Clear out curr buffer in case it had some data in previously.
-		Data.Allocate(0, PinnedProps->SimTarget);
+		Data.Allocate(0);
+		CPUTimeMS = TickTime.GetElapsedMilliseconds();
 		return;
 	}
 
+	UNiagaraSystem* System = ParentSystemInstance->GetSystem();
+
 	check(Data.GetNumVariables() > 0);
-	check(PinnedProps->SpawnScriptProps.Script);
-	check(PinnedProps->UpdateScriptProps.Script);
+	check(CachedEmitter->SpawnScriptProps.Script);
+	check(CachedEmitter->UpdateScriptProps.Script);
 	
 	//TickEvents(DeltaSeconds);
 
@@ -757,11 +847,17 @@ void FNiagaraEmitterInstance::Tick(float DeltaSeconds)
 			Binding.SetValue(Age);
 		}
 		
-		SpawnEmitterLocalSpaceBinding.SetValue(PinnedProps->bLocalSpace);
-		UpdateEmitterLocalSpaceBinding.SetValue(PinnedProps->bLocalSpace);
+		SpawnEmitterLocalSpaceBinding.SetValue(CachedEmitter->bLocalSpace);
+		UpdateEmitterLocalSpaceBinding.SetValue(CachedEmitter->bLocalSpace);
 		for (FNiagaraParameterDirectBinding<FNiagaraBool>& Binding : EventEmitterLocalSpaceBindings)
 		{
-			Binding.SetValue(PinnedProps->bLocalSpace);
+			Binding.SetValue(CachedEmitter->bLocalSpace);
+		}
+
+		if (CachedEmitter->SimTarget == ENiagaraSimTarget::GPUComputeSim)
+		{
+			EmitterAgeBindingGPU.SetValue(Age);
+			EmitterLocalSpaceBindingGPU.SetValue(CachedEmitter->bLocalSpace);
 		}
 	}
 	
@@ -779,25 +875,25 @@ void FNiagaraEmitterInstance::Tick(float DeltaSeconds)
 	}
 
 	// Calculate number of new particles from all event related spawns
-	TArray<TArray<int32>> EventSpawnCounts;
-	EventSpawnCounts.AddDefaulted(PinnedProps->GetEventHandlers().Num());
-	TArray<int32> EventHandlerSpawnCounts;
-	EventHandlerSpawnCounts.AddDefaulted(PinnedProps->GetEventHandlers().Num());
+	TArray<TArray<int32, TInlineAllocator<16>>, TInlineAllocator<16>> EventSpawnCounts;
+	EventSpawnCounts.AddDefaulted(CachedEmitter->GetEventHandlers().Num());
+	TArray<int32, TInlineAllocator<16>> EventHandlerSpawnCounts;
+	EventHandlerSpawnCounts.AddDefaulted(CachedEmitter->GetEventHandlers().Num());
 	uint32 EventSpawnTotal = 0;
-	TArray<FNiagaraDataSet*> EventSet;
-	EventSet.AddZeroed(PinnedProps->GetEventHandlers().Num());
-	TArray<FGuid> SourceEmitterGuid;
-	SourceEmitterGuid.AddDefaulted(PinnedProps->GetEventHandlers().Num());
-	TArray<FName> SourceEmitterName;
-	SourceEmitterName.AddDefaulted(PinnedProps->GetEventHandlers().Num());
-	TArray<bool> bPerformEventSpawning;
-	bPerformEventSpawning.AddDefaulted(PinnedProps->GetEventHandlers().Num());
-
-	for (int32 i = 0; i < PinnedProps->GetEventHandlers().Num(); i++)
-	{
-		const FNiagaraEventScriptProperties &EventHandlerProps = PinnedProps->GetEventHandlers()[i];
+	TArray<FNiagaraDataSet*, TInlineAllocator<16>> EventSet;
+	EventSet.AddZeroed(CachedEmitter->GetEventHandlers().Num());
+	TArray<FGuid, TInlineAllocator<16>> SourceEmitterGuid;
+	SourceEmitterGuid.AddDefaulted(CachedEmitter->GetEventHandlers().Num());
+	TArray<FName, TInlineAllocator<16>> SourceEmitterName;
+	SourceEmitterName.AddDefaulted(CachedEmitter->GetEventHandlers().Num());
+	TArray<bool, TInlineAllocator<16>> bPerformEventSpawning;
+	bPerformEventSpawning.AddDefaulted(CachedEmitter->GetEventHandlers().Num());
+
+	for (int32 i = 0; i < CachedEmitter->GetEventHandlers().Num(); i++)
+	{
+		const FNiagaraEventScriptProperties &EventHandlerProps = CachedEmitter->GetEventHandlers()[i];
 		SourceEmitterGuid[i] = EventHandlerProps.SourceEmitterID;
-		SourceEmitterName[i] = SourceEmitterGuid[i].IsValid() ? *SourceEmitterGuid[i].ToString() : EmitterHandle.GetIdName();
+		SourceEmitterName[i] = SourceEmitterGuid[i].IsValid() ? *SourceEmitterGuid[i].ToString() : CachedIDName;
 		EventSet[i] = FNiagaraEventDataSetMgr::GetEventDataSet(ParentSystemInstance->GetIDName(), SourceEmitterName[i], EventHandlerProps.SourceEventName);
 		bPerformEventSpawning[i] = (ExecutionState == ENiagaraExecutionState::Active && EventHandlerProps.Script && EventHandlerProps.ExecutionMode == EScriptExecutionMode::SpawnedParticles);
 		if (bPerformEventSpawning[i])
@@ -811,35 +907,59 @@ void FNiagaraEmitterInstance::Tick(float DeltaSeconds)
 
 	/* GPU simulation -  we just create an FNiagaraComputeExecutionContext, queue it, and let the batcher take care of the rest
 	 */
-	if (PinnedProps->SimTarget == ENiagaraSimTarget::GPUComputeSim)	
+	if (CachedEmitter->SimTarget == ENiagaraSimTarget::GPUComputeSim)	
 	{
 		//FNiagaraComputeExecutionContext *ComputeContext = new FNiagaraComputeExecutionContext();
 		GPUExecContext.MainDataSet = &Data;
-		GPUExecContext.RTSpawnScript = PinnedProps->SpawnScriptProps.Script->GetRenderThreadScript();
-		GPUExecContext.RTUpdateScript = PinnedProps->UpdateScriptProps.Script->GetRenderThreadScript();
+		GPUExecContext.RTGPUScript = CachedEmitter->GetGPUComputeScript()->GetRenderThreadScript();
+		GPUExecContext.RTSpawnScript = CachedEmitter->SpawnScriptProps.Script->GetRenderThreadScript();
+		GPUExecContext.RTUpdateScript = CachedEmitter->UpdateScriptProps.Script->GetRenderThreadScript();
 		GPUExecContext.SpawnRateInstances = SpawnTotal;
-		GPUExecContext.BurstInstances = 0;
 		GPUExecContext.EventSpawnTotal = EventSpawnTotal;
+		GPUExecContext.NumIndicesPerInstance = CachedEmitter->GetRenderers()[0]->GetNumIndicesPerInstance();
 
-		GPUExecContext.UpdateInterfaces = PinnedProps->UpdateScriptProps.Script->DataInterfaceInfo;
+		bool bOnlySetOnce = false;
+		for (FNiagaraSpawnInfo& Info : SpawnInfos)
+		{
+			if (Info.Count > 0 && !bOnlySetOnce)
+			{
+				SpawnIntervalBindingGPU.SetValue(Info.IntervalDt);
+				InterpSpawnStartBindingGPU.SetValue(Info.InterpStartDt);
+				bOnlySetOnce = true;
+			}
+			else if (Info.Count > 0)
+			{
+				UE_LOG(LogNiagara, Log, TEXT("Multiple spawns are happening this frame. Only doing the first!"));
+				break;
+			}
+		}
+
+		//GPUExecContext.UpdateInterfaces = CachedEmitter->UpdateScriptProps.Script->GetCachedDefaultDataInterfaces();
 
 		// copy over the constants for the render thread
 		//
-		int32 ParmSize = GPUExecContext.CombinedParamStore.GetPaddedParameterSizeInBytes();
-		if (ParmSize)
+		//UE_LOG(LogNiagara, Log, TEXT(".................Spawn................."));
+		//SpawnExecContext.Parameters.DumpParameters();
+		//UE_LOG(LogNiagara, Log, TEXT(".................Update................."));
+		//UpdateExecContext.Parameters.DumpParameters();
+		//UE_LOG(LogNiagara, Log, TEXT(".................GPU................."));
+		//GPUExecContext.CombinedParamStore.DumpParameters();
+		//UE_LOG(LogNiagara, Log, TEXT(".................END.................")); 
+		if (GbDumpParticleData || System->bDumpDebugEmitterInfo)
 		{
-#if WITH_EDITORONLY_DATA
-			if (SpawnExecContext.Script->GetDebuggerInfo().bRequestDebugFrame)
-			{
-				SpawnExecContext.Script->GetDebuggerInfo().DebugFrame.Reset();
-				SpawnExecContext.Script->GetDebuggerInfo().bRequestDebugFrame = false;
-				SpawnExecContext.Script->GetDebuggerInfo().DebugFrameLastWriteId = SpawnExecContext.TickCounter;
-				SpawnExecContext.Script->GetDebuggerInfo().DebugParameters = GPUExecContext.CombinedParamStore;
-			}
-#endif
+			UE_LOG(LogNiagara, Log, TEXT(".................Spawn................."));
+			SpawnExecContext.Parameters.DumpParameters(true);
+			UE_LOG(LogNiagara, Log, TEXT(".................Update................."));
+			UpdateExecContext.Parameters.DumpParameters(true);
+			UE_LOG(LogNiagara, Log, TEXT("................. %s Combined Parameters (%d Spawned )................."), TEXT("GPU Script"), SpawnTotal);
+			GPUExecContext.CombinedParamStore.DumpParameters();
 		}
-		GPUExecContext.SpawnParams.SetNumZeroed(ParmSize);
-		GPUExecContext.CombinedParamStore.CopyParameterDataToPaddedBuffer(GPUExecContext.SpawnParams.GetData(), ParmSize);
+
+		int32 ParmSize = GPUExecContext.CombinedParamStore.GetPaddedParameterSizeInBytes();
+		
+
+		GPUExecContext.ParamData_RT.SetNumZeroed(ParmSize);
+		GPUExecContext.CombinedParamStore.CopyParameterDataToPaddedBuffer(GPUExecContext.ParamData_RT.GetData(), ParmSize);
 
 		// push event data sets to the context
 		for (FNiagaraDataSet *Set : UpdateScriptEventDataSets)
@@ -847,22 +967,44 @@ void FNiagaraEmitterInstance::Tick(float DeltaSeconds)
 			GPUExecContext.UpdateEventWriteDataSets.Add(Set);
 		}
 
-		GPUExecContext.EventHandlerScriptProps = PinnedProps->GetEventHandlers();
+		GPUExecContext.EventHandlerScriptProps = CachedEmitter->GetEventHandlers();
 		GPUExecContext.EventSets = EventSet;
 		GPUExecContext.EventSpawnCounts = EventHandlerSpawnCounts;
 		NiagaraEmitterInstanceBatcher::Get()->Queue(&GPUExecContext);
 
+		// Need to call post-tick, which calls the copy to previous for interpolated spawning
+		SpawnExecContext.PostTick();
+		UpdateExecContext.PostTick();
+		for (FNiagaraScriptExecutionContext& EventContext : EventExecContexts)
+		{
+			EventContext.PostTick();
+		}
+
 		CachedBounds.Init();
 		CachedBounds = CachedBounds.ExpandBy(FVector(20.0, 20.0, 20.0));	// temp until GPU sims update bounds
+
+		CPUTimeMS = TickTime.GetElapsedMilliseconds();
+
+		/*if (CachedEmitter->SpawnScriptProps.Script->GetComputedVMCompilationId().HasInterpolatedParameters())
+		{
+			GPUExecContext.CombinedParamStore.CopyCurrToPrev();
+		}*/
+
 		return;
 	}
 
 	int32 AllocationSize = OrigNumParticles + SpawnTotal + EventSpawnTotal;
 	//Allocate space for prev frames particles and any new one's we're going to spawn.
-	Data.Allocate(AllocationSize, PinnedProps->SimTarget);
-
-
-	TArray<FNiagaraDataSetExecutionInfo> DataSetExecInfos;
+	Data.Allocate(AllocationSize);
+	for (FNiagaraDataSet* SpawnEventDataSet : SpawnScriptEventDataSets)
+	{
+		SpawnEventDataSet->Allocate(SpawnTotal + EventSpawnTotal);
+	}
+	for (FNiagaraDataSet* UpdateEventDataSet : UpdateScriptEventDataSets)
+	{
+		UpdateEventDataSet->Allocate(OrigNumParticles);
+	}
+	TArray<FNiagaraDataSetExecutionInfo, TInlineAllocator<8>> DataSetExecInfos;
 	DataSetExecInfos.Emplace(&Data, 0, false, true);
 
 	// Simulate existing particles forward by DeltaSeconds.
@@ -884,14 +1026,15 @@ void FNiagaraEmitterInstance::Tick(float DeltaSeconds)
 		DataSetExecInfos[0].StartInstance = 0;
 		for (FNiagaraDataSet* EventDataSet : UpdateScriptEventDataSets)
 		{
-			DataSetExecInfos.Emplace(EventDataSet, 0, true, false);
+			DataSetExecInfos.Emplace(EventDataSet, 0, false, true);
+			EventDataSet->SetNumInstances(OrigNumParticles);
 		}
 		UpdateExecContext.Execute(OrigNumParticles, DataSetExecInfos);
 		int32 DeltaParticles = Data.GetNumInstances() - OrigNumParticles;
 
 		ensure(DeltaParticles <= 0); // We either lose particles or stay the same, we should never add particles in update!
 
-		if (GbDumpParticleData)
+		if (GbDumpParticleData || System->bDumpDebugEmitterInfo)
 		{
 			UE_LOG(LogNiagara, Log, TEXT("=== Update Parameters ===") );
 			UpdateExecContext.Parameters.Dump();
@@ -900,8 +1043,22 @@ void FNiagaraEmitterInstance::Tick(float DeltaSeconds)
 			Data.Dump(true, 0, OrigNumParticles);
 		}
 	}
+	
+#if WITH_EDITORONLY_DATA
+	if (ParentSystemInstance->ShouldCaptureThisFrame())
+	{
+		FNiagaraScriptDebuggerInfo* DebugInfo = ParentSystemInstance->GetActiveCaptureWrite(CachedIDName, ENiagaraScriptUsage::ParticleUpdateScript, FGuid());
+		if (DebugInfo)
+		{
+			Data.Dump(DebugInfo->Frame, true, 0, OrigNumParticles);
+			//DebugInfo->Frame.Dump(true, 0, OrigNumParticles);
+			DebugInfo->Parameters = UpdateExecContext.Parameters;
+		}
+	}
+#endif
 
 	uint32 EventSpawnStart = Data.GetNumInstances();
+	int32 NumBeforeSpawn = Data.GetNumInstances();
 
 	//Init new particles with the spawn script.
 	if (SpawnTotal + EventSpawnTotal > 0)
@@ -919,16 +1076,19 @@ void FNiagaraEmitterInstance::Tick(float DeltaSeconds)
 				SpawnExecCountBinding.SetValue(Num);
 				DataSetExecInfos.SetNum(1, false);
 				DataSetExecInfos[0].StartInstance = OrigNum;
+
 				//UE_LOG(LogNiagara, Log, TEXT("SpawnScriptEventDataSets: %d"), SpawnScriptEventDataSets.Num());
 				for (FNiagaraDataSet* EventDataSet : SpawnScriptEventDataSets)
 				{
 					//UE_LOG(LogNiagara, Log, TEXT("SpawnScriptEventDataSets.. %d"), EventDataSet->GetNumVariables());
-					DataSetExecInfos.Emplace(EventDataSet, OrigNum, true, false);
+					int32 EventOrigNum = EventDataSet->GetNumInstances();
+					EventDataSet->SetNumInstances(EventOrigNum + Num);
+					DataSetExecInfos.Emplace(EventDataSet, EventOrigNum, false, true);
 				}
 				
 				SpawnExecContext.Execute(Num, DataSetExecInfos);
 
-				if (GbDumpParticleData)
+				if (GbDumpParticleData || System->bDumpDebugEmitterInfo)
 				{
 					UE_LOG(LogNiagara, Log, TEXT("=== %s Spawn Parameters ==="), *DumpLabel);
 					SpawnExecContext.Parameters.Dump();
@@ -949,7 +1109,7 @@ void FNiagaraEmitterInstance::Tick(float DeltaSeconds)
 
 		EventSpawnStart = Data.GetNumInstances();
 
-		for (int32 EventScriptIdx = 0; EventScriptIdx < PinnedProps->GetEventHandlers().Num(); EventScriptIdx++)
+		for (int32 EventScriptIdx = 0; EventScriptIdx < CachedEmitter->GetEventHandlers().Num(); EventScriptIdx++)
 		{
 			//Spawn particles coming from events.
 			for (int32 i = 0; i < EventSpawnCounts[EventScriptIdx].Num(); i++)
@@ -964,50 +1124,93 @@ void FNiagaraEmitterInstance::Tick(float DeltaSeconds)
 			}
 		}
 	}
+
+
+#if WITH_EDITORONLY_DATA
+	int32 NumAfterSpawn = Data.GetNumInstances();
+	int32 TotalNumSpawned = NumAfterSpawn - NumBeforeSpawn;
+	if (ParentSystemInstance->ShouldCaptureThisFrame())
+	{
+		FNiagaraScriptDebuggerInfo* DebugInfo = ParentSystemInstance->GetActiveCaptureWrite(CachedIDName, ENiagaraScriptUsage::ParticleSpawnScript, FGuid());
+		if (DebugInfo)
+		{
+			Data.Dump(DebugInfo->Frame, true, NumBeforeSpawn, TotalNumSpawned);
+			//DebugInfo->Frame.Dump(true, 0, Num);
+			DebugInfo->Parameters = SpawnExecContext.Parameters;
+		}
+	}
+#endif
 	/*else if (SpawnTotal + EventSpawnTotal > 0)
 	{
 		UE_LOG(LogNiagara, Log, TEXT("Skipping spawning due to execution state! %d"), (uint32)ExecutionState)
 	}*/
 
-	// handle event based spawning
-	for (int32 EventScriptIdx = 0;  EventScriptIdx < PinnedProps->GetEventHandlers().Num(); EventScriptIdx++)
+	// Events are all working from the same set of data generated during spawn they they only need 1 copy to have updated data.
+	if (CachedEmitter->GetEventHandlers().Num())
 	{
-		const FNiagaraEventScriptProperties &EventHandlerProps = PinnedProps->GetEventHandlers()[EventScriptIdx];
+		Data.CopyCurToPrev();
+	}
+	int32 SpawnEventScriptStartIndex = EventSpawnStart;
+	for (int32 EventScriptIdx = 0;  EventScriptIdx < CachedEmitter->GetEventHandlers().Num(); EventScriptIdx++)
+	{
+		const FNiagaraEventScriptProperties &EventHandlerProps = CachedEmitter->GetEventHandlers()[EventScriptIdx];
 
 		if (bPerformEventSpawning[EventScriptIdx] && EventSet[EventScriptIdx] && EventSpawnCounts[EventScriptIdx].Num())
 		{
 			uint32 NumParticles = Data.GetNumInstances();
 			SCOPE_CYCLE_COUNTER(STAT_NiagaraEventHandle);
 
-			Data.Tick();
-			Data.CopyPrevToCur();
-
 			for (int32 i = 0; i < EventSpawnCounts[EventScriptIdx].Num(); i++)
 			{
+				// When using interpolated spawn it's possible for the interpolated update script to kill a particle the same frame that it's spawned.
+				// In this case we have to decrease the number of instances to run the event script on.
 				int32 EventNumToSpawn = EventSpawnCounts[EventScriptIdx][i];
-				ensure(EventNumToSpawn + EventSpawnStart <= Data.GetNumInstances());
-				EventExecCountBindings[EventScriptIdx].SetValue(EventNumToSpawn);
-
-				DataSetExecInfos.SetNum(1, false);
-				DataSetExecInfos[0].StartInstance = EventSpawnStart;
-				DataSetExecInfos[0].bUpdateInstanceCount = false;
-				DataSetExecInfos.Emplace(EventSet[EventScriptIdx], i, false, false);
-				EventExecContexts[EventScriptIdx].Execute(EventNumToSpawn, DataSetExecInfos);
+				int32 ActualEventNumToSpawn = FMath::Min(EventNumToSpawn, (int32)Data.GetNumInstances() - (int32)EventSpawnStart);
 				
-				if (GbDumpParticleData)
+				if (ActualEventNumToSpawn > 0)
 				{
-					UE_LOG(LogNiagara, Log, TEXT("=== Event %d Parameters ==="), EventScriptIdx);
-					EventExecContexts[EventScriptIdx].Parameters.Dump();
-					UE_LOG(LogNiagara, Log, TEXT("=== Event %d %d Particles ==="), EventScriptIdx, EventNumToSpawn);
-					Data.Dump(true, EventSpawnStart, EventNumToSpawn);
-				}
+					EventExecCountBindings[EventScriptIdx].SetValue(EventNumToSpawn);
+
+					DataSetExecInfos.SetNum(1, false);
+					DataSetExecInfos[0].StartInstance = EventSpawnStart;
+					DataSetExecInfos[0].bUpdateInstanceCount = false;
+					DataSetExecInfos.Emplace(EventSet[EventScriptIdx], i, false, false);
+					EventExecContexts[EventScriptIdx].Execute(ActualEventNumToSpawn, DataSetExecInfos);
 
-				ensure(Data.GetNumInstances() == NumParticles);
+					if (GbDumpParticleData)
+					{
+						UE_LOG(LogNiagara, Log, TEXT("=== Event %d Parameters ==="), EventScriptIdx);
+						EventExecContexts[EventScriptIdx].Parameters.Dump();
+						UE_LOG(LogNiagara, Log, TEXT("=== Event %d %d Particles ==="), EventScriptIdx, ActualEventNumToSpawn);
+						Data.Dump(true, EventSpawnStart, ActualEventNumToSpawn);
+					}
 
-				EventSpawnStart += EventNumToSpawn;
+#if WITH_EDITORONLY_DATA
+					if (ParentSystemInstance->ShouldCaptureThisFrame())
+					{
+						FGuid EventGuid = EventExecContexts[EventScriptIdx].Script->GetUsageId();
+						FNiagaraScriptDebuggerInfo* DebugInfo = ParentSystemInstance->GetActiveCaptureWrite(CachedIDName, ENiagaraScriptUsage::ParticleEventScript, EventGuid);
+						if (DebugInfo)
+						{
+							Data.Dump(DebugInfo->Frame, true, EventSpawnStart, ActualEventNumToSpawn);		
+							//DebugInfo->Frame.Dump(true, 0, ActualEventNumToSpawn);
+							DebugInfo->Parameters = EventExecContexts[EventScriptIdx].Parameters;
+						}
+					}
+#endif
+
+					ensure(Data.GetNumInstances() == NumParticles);
+
+					EventSpawnStart += ActualEventNumToSpawn;
+				}
 			}
 		}
+	}
 
+	// Update events need a copy per event so that the previous event's data can be used.
+	for (int32 EventScriptIdx = 0; EventScriptIdx < CachedEmitter->GetEventHandlers().Num(); EventScriptIdx++)
+	{
+		const FNiagaraEventScriptProperties &EventHandlerProps = CachedEmitter->GetEventHandlers()[EventScriptIdx];
 
 		// handle all-particle events
 		if (EventHandlerProps.Script && EventHandlerProps.ExecutionMode == EScriptExecutionMode::EveryParticle && EventSet[EventScriptIdx])
@@ -1019,18 +1222,16 @@ void FNiagaraEmitterInstance::Tick(float DeltaSeconds)
 				
 				for (uint32 i = 0; i < EventSet[EventScriptIdx]->GetPrevNumInstances(); i++)
 				{
-					// If we have events, Swap buffers, to make sure we don't overwrite previous script results
-					// and copy prev to cur, because the event script isn't likely to write all attributes
-					// TODO: copy only event script's unused attributes
-					Data.Tick();
-					Data.CopyPrevToCur();
+					// Copy the current to previous so that the event script has access to the new values from the update
+					// script and any values updated in previous events.
+					Data.CopyCurToPrev();
 
 					uint32 NumInstancesPrev = Data.GetPrevNumInstances();
 					EventExecCountBindings[EventScriptIdx].SetValue(NumInstancesPrev);
 					DataSetExecInfos.SetNum(1, false);
 					DataSetExecInfos[0].StartInstance = 0;
 					DataSetExecInfos.Emplace(EventSet[EventScriptIdx], i, false, false);
-					/*if (GbDumpParticleData)
+					/*if (GbDumpParticleData || System->bDumpDebugEmitterInfo)
 					{
 						UE_LOG(LogNiagara, Log, TEXT("=== Event %d [%d] Payload ==="), EventScriptIdx, i);
 						DataSetExecInfos[1].DataSet->Dump(true, 0, 1);
@@ -1038,7 +1239,7 @@ void FNiagaraEmitterInstance::Tick(float DeltaSeconds)
 
 					EventExecContexts[EventScriptIdx].Execute(NumInstancesPrev, DataSetExecInfos);
 
-					if (GbDumpParticleData)
+					if (GbDumpParticleData || System->bDumpDebugEmitterInfo)
 					{
 						UE_LOG(LogNiagara, Log, TEXT("=== Event %d [%d] Parameters ==="), EventScriptIdx, i);
 						EventExecContexts[EventScriptIdx].Parameters.Dump();
@@ -1046,6 +1247,21 @@ void FNiagaraEmitterInstance::Tick(float DeltaSeconds)
 						Data.Dump(true, 0, NumInstancesPrev);
 					}
 
+
+#if WITH_EDITORONLY_DATA
+					if (ParentSystemInstance->ShouldCaptureThisFrame())
+					{
+						FGuid EventGuid = EventExecContexts[EventScriptIdx].Script->GetUsageId();
+						FNiagaraScriptDebuggerInfo* DebugInfo = ParentSystemInstance->GetActiveCaptureWrite(CachedIDName, ENiagaraScriptUsage::ParticleEventScript, EventGuid);
+						if (DebugInfo)
+						{
+							Data.Dump(DebugInfo->Frame, true, 0, NumInstancesPrev);
+							//DebugInfo->Frame.Dump(true, 0, NumInstancesPrev);
+							DebugInfo->Parameters = EventExecContexts[EventScriptIdx].Parameters;
+						}
+					}
+#endif
+
 					ensure(NumParticles == Data.GetNumInstances());
 				}
 			}
@@ -1065,8 +1281,7 @@ void FNiagaraEmitterInstance::Tick(float DeltaSeconds)
 			{
 				EventExecCountBindings[EventScriptIdx].SetValue(1);
 
-				Data.Tick();
-				Data.CopyPrevToCur();
+				Data.CopyCurToPrev();
 				uint32 NumParticles = Data.GetNumInstances();
 
 				for (uint32 i = 0; i < EventSet[EventScriptIdx]->GetPrevNumInstances(); i++)
@@ -1079,7 +1294,7 @@ void FNiagaraEmitterInstance::Tick(float DeltaSeconds)
 					DataSetExecInfos.Emplace(EventSet[EventScriptIdx], i, false, false);
 					EventExecContexts[EventScriptIdx].Execute(1, DataSetExecInfos);
 
-					if (GbDumpParticleData)
+					if (GbDumpParticleData || System->bDumpDebugEmitterInfo)
 					{
 						ensure(EventHandlerProps.Script->RapidIterationParameters.VerifyBinding(&EventExecContexts[EventScriptIdx].Parameters));
 						UE_LOG(LogNiagara, Log, TEXT("=== Event %d Src Parameters ==="), EventScriptIdx);
@@ -1089,6 +1304,21 @@ void FNiagaraEmitterInstance::Tick(float DeltaSeconds)
 						UE_LOG(LogNiagara, Log, TEXT("=== Event %d Particles (%d index written, %d total) ==="), EventScriptIdx, Index, Data.GetNumInstances());
 						Data.Dump(true, Index, 1);
 					}
+
+
+#if WITH_EDITORONLY_DATA
+					if (ParentSystemInstance->ShouldCaptureThisFrame())
+					{
+						FGuid EventGuid = EventExecContexts[EventScriptIdx].Script->GetUsageId();
+						FNiagaraScriptDebuggerInfo* DebugInfo = ParentSystemInstance->GetActiveCaptureWrite(CachedIDName, ENiagaraScriptUsage::ParticleEventScript, EventGuid);
+						if (DebugInfo)
+						{
+							Data.Dump(DebugInfo->Frame, true, Index, 1);
+							//DebugInfo->Frame.Dump(true, 0, 1);
+							DebugInfo->Parameters = EventExecContexts[EventScriptIdx].Parameters;
+						}
+					}
+#endif
 					ensure(NumParticles == Data.GetNumInstances());
 				}
 			}
@@ -1114,7 +1344,7 @@ void FNiagaraEmitterInstance::Tick(float DeltaSeconds)
  *  We get the counts ahead of event processing time so we only have to allocate new particles once
  *  TODO: augment for multiple spawning event scripts
  */
-uint32 FNiagaraEmitterInstance::CalculateEventSpawnCount(const FNiagaraEventScriptProperties &EventHandlerProps, TArray<int32> &EventSpawnCounts, FNiagaraDataSet *EventSet)
+uint32 FNiagaraEmitterInstance::CalculateEventSpawnCount(const FNiagaraEventScriptProperties &EventHandlerProps, TArray<int32, TInlineAllocator<16>>& EventSpawnCounts, FNiagaraDataSet *EventSet)
 {
 	uint32 EventSpawnTotal = 0;
 	int32 NumEventsToProcess = 0;
@@ -1127,12 +1357,14 @@ uint32 FNiagaraEmitterInstance::CalculateEventSpawnCount(const FNiagaraEventScri
 			NumEventsToProcess = FMath::Min<int32>(EventSet->GetPrevNumInstances(), EventHandlerProps.MaxEventsPerFrame);
 		}
 
+		const bool bUseRandom = EventHandlerProps.bRandomSpawnNumber && EventHandlerProps.MinSpawnNumber < EventHandlerProps.SpawnNumber;
 		for (int32 i = 0; i < NumEventsToProcess; i++)
 		{
-			if (ExecutionState == ENiagaraExecutionState::Active && EventHandlerProps.SpawnNumber > 0)
+			const uint32 SpawnNumber = bUseRandom ? FMath::RandRange((int32)EventHandlerProps.MinSpawnNumber, (int32)EventHandlerProps.SpawnNumber) : EventHandlerProps.SpawnNumber;
+			if (ExecutionState == ENiagaraExecutionState::Active && SpawnNumber > 0)
 			{
-				EventSpawnCounts.Add(EventHandlerProps.SpawnNumber);
-				EventSpawnTotal += EventHandlerProps.SpawnNumber;
+				EventSpawnCounts.Add(SpawnNumber);
+				EventSpawnTotal += SpawnNumber;
 			}
 		}
 	}
@@ -1153,15 +1385,12 @@ void FNiagaraEmitterInstance::SetExecutionState(ENiagaraExecutionState InState)
 	{
 		UE_LOG(LogNiagara, Log, TEXT("Emitter \"%s\" change state N O O O O O "), *GetEmitterHandle().GetName().ToString());
 	}*/
-
-	if (InState != ENiagaraExecutionState::Disabled && GetEmitterHandle().GetInstance()->IsAllowedByDetailLevel() && GetEmitterHandle().GetIsEnabled())
+	check(InState >= ENiagaraExecutionState::Active && InState < ENiagaraExecutionState::Num);
+	//We can't move out of disabled without a proper reinit.
+	if (ExecutionState != ENiagaraExecutionState::Disabled)
 	{
 		ExecutionState = InState;
 	}
-	else
-	{
-		ExecutionState = ENiagaraExecutionState::Disabled;
-	}
 }
 
 
@@ -1169,17 +1398,17 @@ void FNiagaraEmitterInstance::SetExecutionState(ENiagaraExecutionState InState)
 
 bool FNiagaraEmitterInstance::CheckAttributesForRenderer(int32 Index)
 {
-	const FNiagaraEmitterHandle& EmitterHandle = GetEmitterHandle();
 	if (Index > EmitterRenderer.Num())
 	{
 		return false;
 	}
 
-	check(ParticleDataSet);
-	FNiagaraDataSet& Data = *ParticleDataSet;
 	bool bOk = true;
+	checkSlow(ParticleDataSet);
+	FNiagaraDataSet& Data = *ParticleDataSet;
 	if (EmitterRenderer[Index])
 	{
+		
 		const TArray<FNiagaraVariable>& RequiredAttrs = EmitterRenderer[Index]->GetRequiredAttributes();
 
 		for (FNiagaraVariable Attr : RequiredAttrs)
@@ -1194,10 +1423,16 @@ bool FNiagaraEmitterInstance::CheckAttributesForRenderer(int32 Index)
 			if (!Data.HasVariable(Attr))
 			{
 				bOk = false;
-				UE_LOG(LogNiagara, Error, TEXT("Cannot render %s because it does not define attribute %s %s."), *EmitterHandle.GetName().ToString(), *Attr.GetType().GetNameText().ToString() , *Attr.GetName().ToString());
+				UE_LOG(LogNiagara, Error, TEXT("Cannot render %s because it does not define attribute %s %s."), *GetEmitterHandle().GetName().ToString(), *Attr.GetType().GetNameText().ToString() , *Attr.GetName().ToString());
 			}
 		}
 
+		if (bOk && !EmitterRenderer[Index]->GetRendererProperties()->IsSimTargetSupported(CachedEmitter->SimTarget))
+		{
+			UE_LOG(LogNiagara, Error, TEXT("Cannot render %s because it is not compatible with this SimTarget mode."), *GetEmitterHandle().GetName().ToString());
+			bOk = false;
+		}
+
 		EmitterRenderer[Index]->SetEnabled(bOk);
 	}
 	return bOk;
@@ -1210,60 +1445,56 @@ Don't forget to call RenderModuleUpdate on the SceneProxy after calling this!
  */
 void FNiagaraEmitterInstance::UpdateEmitterRenderer(ERHIFeatureLevel::Type FeatureLevel, TArray<NiagaraRenderer*>& ToBeAddedList, TArray<NiagaraRenderer*>& ToBeRemovedList)
 {
-	const FNiagaraEmitterHandle& EmitterHandle = GetEmitterHandle();
-	const UNiagaraEmitter* EmitterProperties = EmitterHandle.GetInstance();
+	checkSlow(CachedEmitter);
 
-	if (EmitterProperties != nullptr)
+	// Add all the old to be purged..
+	for (int32 SubIdx = 0; SubIdx < EmitterRenderer.Num(); SubIdx++)
 	{
-		// Add all the old to be purged..
-		for (int32 SubIdx = 0; SubIdx < EmitterRenderer.Num(); SubIdx++)
+		if (EmitterRenderer[SubIdx] != nullptr)
 		{
-			if (EmitterRenderer[SubIdx] != nullptr)
-			{
-				ToBeRemovedList.Add(EmitterRenderer[SubIdx]);
-				EmitterRenderer[SubIdx] = nullptr;
-			}
+			ToBeRemovedList.Add(EmitterRenderer[SubIdx]);
+			EmitterRenderer[SubIdx] = nullptr;
 		}
+	}
 
-		if (!IsComplete() && !bError)
+	if (!IsComplete())
+	{
+		EmitterRenderer.Empty();
+		EmitterRenderer.AddZeroed(CachedEmitter->GetRenderers().Num());
+		for (int32 SubIdx = 0; SubIdx < CachedEmitter->GetRenderers().Num(); SubIdx++)
 		{
-			EmitterRenderer.Empty();
-			EmitterRenderer.AddZeroed(EmitterProperties->GetRenderers().Num());
-			for (int32 SubIdx = 0; SubIdx < EmitterProperties->GetRenderers().Num(); SubIdx++)
-			{
-				UMaterialInterface *Material = nullptr;
+			UMaterialInterface *Material = nullptr;
 
-				TArray<UMaterialInterface*> UsedMats;
-				if (EmitterProperties->GetRenderers()[SubIdx] != nullptr)
+			TArray<UMaterialInterface*> UsedMats;
+			if (CachedEmitter->GetRenderers()[SubIdx] != nullptr)
+			{
+				CachedEmitter->GetRenderers()[SubIdx]->GetUsedMaterials(UsedMats);
+				if (UsedMats.Num() != 0)
 				{
-					EmitterProperties->GetRenderers()[SubIdx]->GetUsedMaterials(UsedMats);
-					if (UsedMats.Num() != 0)
-					{
-						Material = UsedMats[0];
-					}
+					Material = UsedMats[0];
 				}
+			}
 
-				if (Material == nullptr)
-				{
-					Material = UMaterial::GetDefaultMaterial(MD_Surface);
-				}
+			if (Material == nullptr)
+			{
+				Material = UMaterial::GetDefaultMaterial(MD_Surface);
+			}
 
-				if (EmitterProperties->GetRenderers()[SubIdx] != nullptr)
-				{
-					EmitterRenderer[SubIdx] = EmitterProperties->GetRenderers()[SubIdx]->CreateEmitterRenderer(FeatureLevel);
-					EmitterRenderer[SubIdx]->SetMaterial(Material, FeatureLevel);
-					EmitterRenderer[SubIdx]->SetLocalSpace(EmitterProperties->bLocalSpace);
-					ToBeAddedList.Add(EmitterRenderer[SubIdx]);
+			if (CachedEmitter->GetRenderers()[SubIdx] != nullptr)
+			{
+				EmitterRenderer[SubIdx] = CachedEmitter->GetRenderers()[SubIdx]->CreateEmitterRenderer(FeatureLevel);
+				EmitterRenderer[SubIdx]->SetMaterial(Material, FeatureLevel);
+				EmitterRenderer[SubIdx]->SetLocalSpace(CachedEmitter->bLocalSpace);
+				ToBeAddedList.Add(EmitterRenderer[SubIdx]);
 
-					//UE_LOG(LogNiagara, Warning, TEXT("CreateRenderer %p"), EmitterRenderer);
+				//UE_LOG(LogNiagara, Warning, TEXT("CreateRenderer %p"), EmitterRenderer);
 #if WITH_EDITORONLY_DATA
-					CheckAttributesForRenderer(SubIdx);
+				CheckAttributesForRenderer(SubIdx);
 #endif
-				}
-				else
-				{
-					EmitterRenderer[SubIdx] = nullptr;
-				}
+			}
+			else
+			{
+				EmitterRenderer[SubIdx] = nullptr;
 			}
 		}
 	}
