Copying //UE4/Dev-Niagara to //UE4/Dev-Main (Source: //UE4/Dev-Niagara @ 4074996)

diff --git a/Engine/Plugins/FX/Niagara/Source/Niagara/Private/NiagaraEmitterInstanceBatcher.cpp b/Engine/Plugins/FX/Niagara/Source/Niagara/Private/NiagaraEmitterInstanceBatcher.cpp
--- a/Engine/Plugins/FX/Niagara/Source/Niagara/Private/NiagaraEmitterInstanceBatcher.cpp
+++ b/Engine/Plugins/FX/Niagara/Source/Niagara/Private/NiagaraEmitterInstanceBatcher.cpp
@@ -3,19 +3,26 @@
 #include "NiagaraEmitterInstanceBatcher.h"
 #include "NiagaraScriptExecutionContext.h"
 #include "RHI.h"
+#include "RHIGPUReadback.h"
 #include "NiagaraStats.h"
 #include "NiagaraShader.h"
 #include "ShaderParameterUtils.h"
 #include "SceneUtils.h"
+#include "ShaderParameterUtils.h"
 #include "ClearQuad.h"
 
-DECLARE_CYCLE_STAT(TEXT("Batching"), STAT_NiagaraBatching, STATGROUP_Niagara);
-DECLARE_GPU_STAT(NiagaraSim);
+DECLARE_CYCLE_STAT(TEXT("Batching"), STAT_NiagaraGPUSimTick_RT, STATGROUP_Niagara);
+DECLARE_CYCLE_STAT(TEXT("GPU Data Readback"), STAT_NiagaraGPUReadback_RT, STATGROUP_Niagara);
+DECLARE_FLOAT_COUNTER_STAT(TEXT("Niagara GPU Sim"), STAT_GPU_NiagaraSim, STATGROUP_GPU);
+DECLARE_DWORD_COUNTER_STAT(TEXT("# GPU Particles"), STAT_NiagaraGPUParticles, STATGROUP_Niagara);
+DECLARE_DWORD_COUNTER_STAT(TEXT("Readback latency (frames)"), STAT_NiagaraReadbackLatency, STATGROUP_Niagara);
+
+DECLARE_GPU_STAT_NAMED(NiagaraGPUSimulation, TEXT("Niagara GPU Simulation"));
+DECLARE_GPU_STAT_NAMED(NiagaraIndexBufferClear, TEXT("Niagara index buffer clear"));
 
 
 NiagaraEmitterInstanceBatcher* NiagaraEmitterInstanceBatcher::BatcherSingleton = nullptr;
 uint32 FNiagaraComputeExecutionContext::TickCounter = 0;
-FRWBuffer NiagaraEmitterInstanceBatcher::DummyWriteIndexBuffer;
 
 void NiagaraEmitterInstanceBatcher::Queue(FNiagaraComputeExecutionContext *InContext)
 {
@@ -31,159 +38,228 @@ void NiagaraEmitterInstanceBatcher::Queue(FNiagaraComputeExecutionContext *InCon
 }
 
 
-void NiagaraEmitterInstanceBatcher::ExecuteAll(FRHICommandList &RHICmdList)
+void NiagaraEmitterInstanceBatcher::ExecuteAll(FRHICommandList &RHICmdList, FUniformBufferRHIParamRef ViewUniformBuffer)
 {
 	TArray<FNiagaraComputeExecutionContext*> &WorkQueue = SimulationQueue[CurQueueIndex ^ 0x1];
 	for (FNiagaraComputeExecutionContext *Context : WorkQueue)
 	{
+		// need to call RenderThreadInit on data interfaces
 		//ExecuteSingle(Context, RHICmdList);
-		TickSingle(Context, RHICmdList);
+		TickSingle(Context, RHICmdList, ViewUniformBuffer);
 	}
 	WorkQueue.Empty();
 }
 
-void NiagaraEmitterInstanceBatcher::TickSingle(const FNiagaraComputeExecutionContext *Context, FRHICommandList &RHICmdList) const
+void NiagaraEmitterInstanceBatcher::TickSingle(FNiagaraComputeExecutionContext *Context, FRHICommandList &RHICmdList, FUniformBufferRHIParamRef ViewUniformBuffer) const
 {
-	SCOPE_CYCLE_COUNTER(STAT_NiagaraTick);
+	SCOPE_CYCLE_COUNTER(STAT_NiagaraGPUSimTick_RT);
 
 	check(IsInRenderingThread());
-	Context->MainDataSet->TickRenderThread(ENiagaraSimTarget::GPUComputeSim);
+	Context->MainDataSet->Tick();
 
 	FNiagaraComputeExecutionContext::TickCounter++;
 
-	FNiagaraShader *SpawnShader = Context->RTSpawnScript->GetShader();
-	if (!SpawnShader)
+	FNiagaraShader* ComputeShader = Context->RTGPUScript->GetShader();
+	if (!ComputeShader)
 	{
 		return;
 	}
 
-	uint32 PrevNumInstances = Context->MainDataSet->PrevDataRender().GetNumInstances();
+	uint32 PrevNumInstances = Context->MainDataSet->PrevData().GetNumInstances();
+	uint32 NewNumInstances = Context->SpawnRateInstances + Context->EventSpawnTotal + PrevNumInstances;
+	Context->EventSpawnTotal = GetEventSpawnTotal(Context);
 
-	uint32 EventSpawnTotal = 0;
-	/*
-	for (int32 i = 0; i < Context->GetEventHandlers().Num(); i++)
-	{
-		const FNiagaraEventScriptProperties &EventHandlerProps = Context->GetEventHandlers()[i];
-		if (EventHandlerProps.ExecutionMode == EScriptExecutionMode::SpawnedParticles && Context->EventSets[i])
-		{
-			uint32 NumEventsToProcess = Context->EventSets[i]->PrevDataRender().GetNumInstances();
-			uint32 EventSpawnNum = NumEventsToProcess * EventHandlerProps.SpawnNumber;
-			EventSpawnTotal += EventSpawnNum;
-		}
-	}
-	*/
+	ResizeCurrentBuffer(RHICmdList, Context, NewNumInstances, PrevNumInstances);
 
-
-	// allocate for additional instances spawned and set the new number in the data set, if the new number is greater (meaning if we're spawning in this run)
-	// TODO: interpolated spawning
-	//
-	uint32 NewNumInstances = Context->SpawnRateInstances + Context->BurstInstances + EventSpawnTotal + PrevNumInstances;
-	if (NewNumInstances > PrevNumInstances)
-	{
-		Context->MainDataSet->CurrDataRender().AllocateGPU(NewNumInstances, RHICmdList);
-		Context->MainDataSet->CurrDataRender().SetNumInstances(NewNumInstances);
-	}
-	// if we're not spawning, we need to make sure that the current buffer alloc size and number of instances matches the last one
-	// we may have spawned in the last tick, so the buffers may be different sizes
+	// set up a data set index buffer, if we don't have one yet
 	//
-	else if (Context->MainDataSet->CurrDataRender().GetNumInstances() < Context->MainDataSet->PrevDataRender().GetNumInstances())
+	if (!Context->MainDataSet->HasDatasetIndices())
 	{
-		Context->MainDataSet->CurrDataRender().AllocateGPU(PrevNumInstances, RHICmdList);
-		Context->MainDataSet->CurrDataRender().SetNumInstances(PrevNumInstances);
+		Context->MainDataSet->SetupCurDatasetIndices();
 	}
 
-	//UE_LOG(LogNiagara, Log, TEXT("Pre sim instances %i"), PrevNumInstances);
-
-	// Calc shader parameters
-	uint32 TotalNumInstances = NewNumInstances;
-	uint32 SimulationStartInstance = 0;
-	uint32 SpawnStartInstance = PrevNumInstances;
-
-	// Recreate a clear data set index buffer for the simulation shader to write number of written instances to
+	// clear data set index buffer for the simulation shader to write number of written instances to
 	//
-	{
-		SCOPED_DRAW_EVENT(RHICmdList, NiagaraIndexBufferClear);
-		ClearUAV(RHICmdList, Context->MainDataSet->GetDataSetIndices(), 0);
-	}
-	RHICmdList.TransitionResource(EResourceTransitionAccess::EWritable, EResourceTransitionPipeline::EComputeToCompute, Context->MainDataSet->GetDataSetIndices().UAV);
-
+	ClearIndexBufferCur(RHICmdList, Context);
 
 	// run shader, sim and spawn in a single dispatch
-	RHICmdList.SetComputeShader(SpawnShader->GetComputeShader());
-	SetupDataInterfaceBuffers(Context->UpdateInterfaces, SpawnShader, RHICmdList);
-	Run(Context->MainDataSet, SimulationStartInstance, SpawnStartInstance, TotalNumInstances, SpawnShader, Context->SpawnParams, RHICmdList, Context->MainDataSet->GetDataSetIndices());
+	uint32 UpdateStartInstance = 0;
+	Run(Context, UpdateStartInstance, NewNumInstances, ComputeShader, RHICmdList, ViewUniformBuffer);
 
-	// resolve data set writes
-	// grabs the number of instances written from the index set during the simulation run
-	//
-	uint32 NumInstancesAfterSim[64];
-	NumInstancesAfterSim[0] = PrevNumInstances;
-	//we may be able to remove this transition if each step isn't dependent on the previous one.
-	RHICmdList.TransitionResource(EResourceTransitionAccess::EReadable, EResourceTransitionPipeline::EComputeToCompute, Context->MainDataSet->GetDataSetIndices().UAV);
-	ResolveDatasetWrites(NumInstancesAfterSim, Context, Context->MainDataSet->GetDataSetIndices());		// this is going to cause a flush; need to place it differently
-	Context->MainDataSet->CurrDataRender().SetNumInstances(NumInstancesAfterSim[0]);
-	RHICmdList.TransitionResource(EResourceTransitionAccess::EReadable, EResourceTransitionPipeline::EComputeToGfx, Context->MainDataSet->GetDataSetIndices().UAV);
+	// assume all instances survived; ResolveDataSetWrites will change this if the deferred readback was successful; that data may be several frames old
+	Context->MainDataSet->CurrData().SetNumInstances(NewNumInstances);
+
+	// resolve data set writes - grabs the number of instances written from the index set during the simulation run
+	ResolveDatasetWrites(Context);
+	RHICmdList.TransitionResource(EResourceTransitionAccess::EReadable, EResourceTransitionPipeline::EComputeToGfx, Context->MainDataSet->GetCurDataSetIndices().UAV);		// transition to readable; we'll be using this next frame
 
-	// TODO: hack - only updating event set 0 on update scripts now; need to match them to their indices and update them all
 	/*
+	// TODO: hack - only updating event set 0 on update scripts now; need to match them to their indices and update them all
 	if (Context->UpdateEventWriteDataSets.Num())
 	{ 
 		Context->UpdateEventWriteDataSets[0]->CurrDataRender().SetNumInstances(NumInstancesAfterSim[1]);
 	}
+	RunEventHandlers(Context, NumInstancesAfterSim[0], NumInstancesAfterSpawn, NumInstancesAfterNonEventSpawn, RHICmdList);
 	*/
 
-	//UE_LOG(LogNiagara, Log, TEXT("After sim instances %i"), NumInstancesAfterSim[0]);
-
-	//RunEventHandlers(Context, NumInstancesAfterSim[0], NumInstancesAfterSpawn, NumInstancesAfterNonEventSpawn, RHICmdList);
-
 	// the VF grabs PrevDataRender for drawing, so need to transition
-	RHICmdList.TransitionResource(EResourceTransitionAccess::EReadable, EResourceTransitionPipeline::EComputeToGfx, Context->MainDataSet->PrevDataRender().GetGPUBufferFloat()->UAV);
-	RHICmdList.TransitionResource(EResourceTransitionAccess::EReadable, EResourceTransitionPipeline::EComputeToGfx, Context->MainDataSet->PrevDataRender().GetGPUBufferInt()->UAV);
+	RHICmdList.TransitionResource(EResourceTransitionAccess::EReadable, EResourceTransitionPipeline::EComputeToGfx, Context->MainDataSet->PrevData().GetGPUBufferFloat()->UAV);
+	RHICmdList.TransitionResource(EResourceTransitionAccess::EReadable, EResourceTransitionPipeline::EComputeToGfx, Context->MainDataSet->PrevData().GetGPUBufferInt()->UAV);
+
+	check(Context->MainDataSet->HasDatasetIndices());
 }
 
 
-void NiagaraEmitterInstanceBatcher::ResolveDatasetWrites(uint32 *OutArray, const FNiagaraComputeExecutionContext *Context, FRWBuffer &WriteIndexBuffer) const
+/* Calculate total number of spawned instances from all spawn events
+*/
+uint32 NiagaraEmitterInstanceBatcher::GetEventSpawnTotal(const FNiagaraComputeExecutionContext *InContext) const
 {
-	if (WriteIndexBuffer.NumBytes > 0)
+	uint32 EventSpawnTotal = 0;
+	for (int32 i = 0; i < InContext->GetEventHandlers().Num(); i++)
 	{
-		int32 *DataSetInstanceIndexBuffer = static_cast<int32*>(RHILockVertexBuffer(WriteIndexBuffer.Buffer, 0, 64 * sizeof(int32), EResourceLockMode::RLM_ReadOnly));
-		for (uint32 SetIdx = 0; SetIdx < 64; SetIdx++)
+		const FNiagaraEventScriptProperties &EventHandlerProps = InContext->GetEventHandlers()[i];
+		if (EventHandlerProps.ExecutionMode == EScriptExecutionMode::SpawnedParticles && InContext->EventSets[i])
 		{
-			OutArray[SetIdx] = DataSetInstanceIndexBuffer[SetIdx*4+1];
+			uint32 NumEventsToProcess = InContext->EventSets[i]->PrevData().GetNumInstances();
+			uint32 EventSpawnNum = NumEventsToProcess * EventHandlerProps.SpawnNumber;
+			EventSpawnTotal += EventSpawnNum;
 		}
-		RHIUnlockVertexBuffer(WriteIndexBuffer.Buffer);
 	}
+	return EventSpawnTotal;
+}
+
+
+/* Clear the data set index buffer; needs to be called before a sim run
+ */
+void NiagaraEmitterInstanceBatcher::ClearIndexBufferCur(FRHICommandList &RHICmdList, FNiagaraComputeExecutionContext *Context) const
+{
+	FRWBuffer &DatasetIndexBufferWrite = Context->MainDataSet->GetCurDataSetIndices();
+	SCOPED_DRAW_EVENTF(RHICmdList, NiagaraIndexBufferClear, TEXT("Niagara index buffer clear"));
+	SCOPED_GPU_STAT(RHICmdList, NiagaraIndexBufferClear);
+
+	ClearUAV(RHICmdList, DatasetIndexBufferWrite, 0);
+	RHICmdList.TransitionResource(EResourceTransitionAccess::EWritable, EResourceTransitionPipeline::EComputeToCompute, DatasetIndexBufferWrite.UAV);
 }
 
-void NiagaraEmitterInstanceBatcher::SetupDataInterfaceBuffers(const TArray<FNiagaraScriptDataInterfaceInfo> &DIInfos, FNiagaraShader *Shader, FRHICommandList &RHICmdList) const
+/* Attempt to read back simulation results (number of live instances) from the GPU via an async readback request; 
+ *	If the readback isn't ready to be performed, we accumulate spawn rates and assume all instances have survived, until
+ *	the GPU can tell us how many are actually alive; since that data may be several frames old, we'll always end up
+ *	overallocating a bit, and the CPU might think we have more particles alive than we actually do;
+ *	since we use DrawIndirect with the GPU determining draw call parameters, that's not an issue
+ */
+void NiagaraEmitterInstanceBatcher::ResolveDatasetWrites(FNiagaraComputeExecutionContext *Context) const
+{
+	FRWBuffer &DatasetIndexBufferWrite = Context->MainDataSet->GetCurDataSetIndices();
+	uint32 SpawnedThisFrame = Context->SpawnRateInstances + Context->EventSpawnTotal;
+	Context->AccumulatedSpawnRate += SpawnedThisFrame;
+	if (!Context->GPUDataReadback)
+	{
+		Context->GPUDataReadback = new FRHIGPUMemoryReadback(DatasetIndexBufferWrite.Buffer, TEXT("Niagara GPU Emitter Readback"));
+		INC_DWORD_STAT(STAT_NiagaraReadbackLatency);
+	}
+	else if (Context->GPUDataReadback->IsReady())
+	{
+		SCOPE_CYCLE_COUNTER(STAT_NiagaraGPUReadback_RT);
+		int32 *NumInstancesAfterSim = static_cast<int32*>(Context->GPUDataReadback->RetrieveData(64 * sizeof(int32)));
+		Context->MainDataSet->CurrData().SetNumInstances(NumInstancesAfterSim[1] + Context->AccumulatedSpawnRate);	// index 1 is always the count
+		SET_DWORD_STAT(STAT_NiagaraGPUParticles, NumInstancesAfterSim[1] + Context->AccumulatedSpawnRate);
+		SET_DWORD_STAT(STAT_NiagaraReadbackLatency, 0);
+
+		Context->GPUDataReadback->Finish();
+
+		Context->AccumulatedSpawnRate = 0;
+		delete Context->GPUDataReadback;
+		Context->GPUDataReadback = new FRHIGPUMemoryReadback(DatasetIndexBufferWrite.Buffer, TEXT("Niagara GPU Emitter Readback"));
+	}
+}
+
+
+/* Resize data set buffers and set number of instances 
+ *	Allocates one additional instance at the end, which is a scratch instance; by setting the default index from AcquireIndex in the shader
+ *	to that scratch index, we can avoid branching in every single OutputData function
+ */
+void NiagaraEmitterInstanceBatcher::ResizeCurrentBuffer(FRHICommandList &RHICmdList, FNiagaraComputeExecutionContext *Context, uint32 NewNumInstances, uint32 PrevNumInstances) const
+{
+
+	// allocate for additional instances spawned and set the new number in the data set, if the new number is greater (meaning if we're spawning in this run)
+	// TODO: interpolated spawning
+	//
+	if (NewNumInstances > PrevNumInstances)
+	{
+		Context->MainDataSet->CurrData().AllocateGPU(NewNumInstances + 1, RHICmdList);
+		Context->MainDataSet->CurrData().SetNumInstances(NewNumInstances);
+	}
+	// if we're not spawning, we need to make sure that the current buffer alloc size and number of instances matches the last one
+	// we may have spawned in the last tick, so the buffers may be different sizes
+	//
+	else if (Context->MainDataSet->CurrData().GetNumInstances() < Context->MainDataSet->PrevData().GetNumInstances())
+	{
+		Context->MainDataSet->CurrData().AllocateGPU(PrevNumInstances + 1, RHICmdList);
+		Context->MainDataSet->CurrData().SetNumInstances(PrevNumInstances);
+	}
+}
+
+
+/* Set shader parameters for data interfaces
+ */
+void NiagaraEmitterInstanceBatcher::SetDataInterfaceParameters(const TArray<UNiagaraDataInterface*> &DataInterfaces, FNiagaraShader* Shader, FRHICommandList &RHICmdList) const
 {
 	// set up data interface buffers, as defined by the DIs during compilation
 	//
 	uint32 InterfaceIndex = 0;
-	for (const FNiagaraScriptDataInterfaceInfo &InterfaceInfo : DIInfos)
+	for (UNiagaraDataInterface* Interface : DataInterfaces)
 	{
-		TArray<FNiagaraDataInterfaceBufferData> &BufferDatas = InterfaceInfo.DataInterface->GetBufferDataArray();
-		for (FNiagaraDataInterfaceBufferData &BufferData : BufferDatas)
+		FNiagaraDataInterfaceParamRef& DIParam = Shader->GetDIParameters()[InterfaceIndex];
+		if (DIParam.Parameters)
 		{
-			FShaderResourceParameter *Param = Shader->FindDIBufferParam(InterfaceIndex, BufferData.UniformName);
-			if (Param)
-			{
-				RHICmdList.SetShaderResourceViewParameter(Shader->GetComputeShader(), Param->GetBaseIndex(), BufferData.Buffer.SRV);
-			}
+			DIParam.Parameters->Set(RHICmdList, Shader, Interface);
 		}
+
 		InterfaceIndex++;
 	}
 }
 
-void NiagaraEmitterInstanceBatcher::Run(FNiagaraDataSet *DataSet, uint32 Phase0StartInstance, uint32 Phase1StartInstance, const uint32 TotalNumInstances, FNiagaraShader *Shader, const TArray<uint8, TAlignedHeapAllocator<16>>  &Params, FRHICommandList &RHICmdList, const FRWBuffer &WriteIndexBuffer, bool bCopyBeforeStart) const
+
+/* Kick off a simulation/spawn run
+ */
+void NiagaraEmitterInstanceBatcher::Run(const FNiagaraComputeExecutionContext *Context, uint32 UpdateStartInstance, const uint32 TotalNumInstances, FNiagaraShader* Shader,
+	FRHICommandList &RHICmdList, FUniformBufferRHIParamRef ViewUniformBuffer, bool bCopyBeforeStart) const
 {
 	if (TotalNumInstances == 0)
 	{
 		return;
 	}
 
+	FNiagaraDataSet *DataSet = Context->MainDataSet;
+	const FNiagaraParameterStore& ParameterStore = Context->CombinedParamStore;
+	const TArray<uint8, TAlignedHeapAllocator<16>> &Params = Context->ParamData_RT;
+	const FRWBuffer &WriteIndexBuffer = DataSet->GetCurDataSetIndices();
+	FRWBuffer &ReadIndexBuffer = DataSet->GetPrevDataSetIndices();
+
+	// if we don't have a previous index buffer, we need to prep one using the maximum number of instances; this should only happen on the first frame
+	//		the data set index buffer is really the param buffer for the indirect draw call; it contains the number of live instances at index 1, and the simulation
+	//		CS uses this to determine the current number of active instances in the buffer
+	//
+	if (ReadIndexBuffer.Buffer == nullptr)
+	{
+		TResourceArray<int32> InitIndexBuffer;
+		InitIndexBuffer.AddUninitialized(64);
+		InitIndexBuffer[1] = 0;		// number of instances 
+		ReadIndexBuffer.Initialize(sizeof(int32), 64, EPixelFormat::PF_R32_UINT, BUF_DrawIndirect | BUF_Static, nullptr, &InitIndexBuffer);
+	}
+
 	RHICmdList.SetComputeShader(Shader->GetComputeShader());
 
+	RHICmdList.SetShaderResourceViewParameter(Shader->GetComputeShader(), Shader->InputIndexBufferParam.GetBaseIndex(), ReadIndexBuffer.SRV);
+
+	// set the view uniform buffer param
+	if (Shader->ViewUniformBufferParam.IsBound() && ViewUniformBuffer)
+	{
+		RHICmdList.SetShaderUniformBuffer(Shader->GetComputeShader(), Shader->ViewUniformBufferParam.GetBaseIndex(), ViewUniformBuffer);
+	}
+
+	SetDataInterfaceParameters(ParameterStore.GetDataInterfaces(), Shader, RHICmdList);
+
 	// set the shader and data set params 
 	//
 	DataSet->SetShaderParams(Shader, RHICmdList);
@@ -199,14 +275,15 @@ void NiagaraEmitterInstanceBatcher::Run(FNiagaraDataSet *DataSet, uint32 Phase0S
 	//
 	if (Shader->EmitterTickCounterParam.IsBound())
 	{
-		RHICmdList.SetShaderParameter(Shader->GetComputeShader(), 0, Shader->EmitterTickCounterParam.GetBaseIndex(), Shader->EmitterTickCounterParam.GetNumBytes(), &FNiagaraComputeExecutionContext::TickCounter);
+		RHICmdList.SetShaderParameter(Shader->GetComputeShader(), Shader->EmitterTickCounterParam.GetBufferIndex(), Shader->EmitterTickCounterParam.GetBaseIndex(), Shader->EmitterTickCounterParam.GetNumBytes(), &FNiagaraComputeExecutionContext::TickCounter);
 	}
 
 	uint32 Copy = bCopyBeforeStart ? 1 : 0;
 
-	RHICmdList.SetShaderParameter(Shader->GetComputeShader(), 0, Shader->StartInstancePhase0Param.GetBaseIndex(), Shader->StartInstancePhase0Param.GetNumBytes(), &Phase0StartInstance);
-	RHICmdList.SetShaderParameter(Shader->GetComputeShader(), 0, Shader->StartInstancePhase1Param.GetBaseIndex(), Shader->StartInstancePhase1Param.GetNumBytes(), &Phase1StartInstance);
-	RHICmdList.SetShaderParameter(Shader->GetComputeShader(), 0, Shader->TotalNumInstancesParam.GetBaseIndex(), Shader->TotalNumInstancesParam.GetNumBytes(), &TotalNumInstances);
+	RHICmdList.SetShaderParameter(Shader->GetComputeShader(), Shader->UpdateStartInstanceParam.GetBufferIndex(), Shader->UpdateStartInstanceParam.GetBaseIndex(), Shader->UpdateStartInstanceParam.GetNumBytes(), &UpdateStartInstance);					// 0, except for event handler runs
+	RHICmdList.SetShaderParameter(Shader->GetComputeShader(), Shader->NumIndicesPerInstanceParam.GetBufferIndex(), Shader->NumIndicesPerInstanceParam.GetBaseIndex(), Shader->NumIndicesPerInstanceParam.GetNumBytes(), &Context->NumIndicesPerInstance);		// set from the renderer in FNiagaraEmitterInstance::Tick
+	int32 InstancesToSpawnThisFrame = Context->SpawnRateInstances + Context->EventSpawnTotal;
+	RHICmdList.SetShaderParameter(Shader->GetComputeShader(), Shader->NumSpawnedInstancesParam.GetBufferIndex(), Shader->NumSpawnedInstancesParam.GetBaseIndex(), Shader->NumSpawnedInstancesParam.GetNumBytes(), &InstancesToSpawnThisFrame);				// number of instances in the spawn run
 
 	uint32 NumThreadGroups = 1;
 	if (TotalNumInstances > NIAGARA_COMPUTE_THREADGROUP_SIZE)
@@ -229,8 +306,8 @@ void NiagaraEmitterInstanceBatcher::Run(FNiagaraDataSet *DataSet, uint32 Phase0S
 	//
 	if (TotalNumInstances)
 	{
-		SCOPED_DRAW_EVENT(RHICmdList, NiagaraGPUSimulationCS);
-		SCOPED_GPU_STAT(RHICmdList, NiagaraSim);
+		SCOPED_DRAW_EVENTF(RHICmdList, NiagaraGPUSimulationCS, TEXT("Niagara GPU Simulation"));
+		SCOPED_GPU_STAT(RHICmdList, NiagaraGPUSimulation);
 		DispatchComputeShader(RHICmdList, Shader, NumThreadGroups, 1, 1);
 	}
 
@@ -240,7 +317,9 @@ void NiagaraEmitterInstanceBatcher::Run(FNiagaraDataSet *DataSet, uint32 Phase0S
 	Shader->OutputIndexBufferParam.UnsetUAV(RHICmdList, Shader->GetComputeShader());
 }
 
-
+/* Kick off event handler runs
+ * TODO: compile event handlers into the sim/spawn compute shader, so we can do everything within one dispatch
+ */
 void NiagaraEmitterInstanceBatcher::RunEventHandlers(const FNiagaraComputeExecutionContext *Context, uint32 NumInstancesAfterSim, uint32 NumInstancesAfterSpawn, uint32 NumInstancesAfterNonEventSpawn, FRHICommandList &RHICmdList) const
 {
 	// Event handler run
@@ -257,7 +336,7 @@ void NiagaraEmitterInstanceBatcher::RunEventHandlers(const FNiagaraComputeExecut
 			// handle all-particle events
 			if (NumEvents && EventHandlerProps.Script && EventHandlerProps.ExecutionMode == EScriptExecutionMode::EveryParticle && EventSet)
 			{
-				FNiagaraShader *EventHandlerShader = EventHandlerProps.Script->GetRenderThreadScript()->GetShader();
+				FNiagaraShader* EventHandlerShader = EventHandlerProps.Script->GetRenderThreadScript()->GetShader();
 				if (EventHandlerShader)
 				{
 					SetupDataInterfaceBuffers(EventHandlerProps.Script->DataInterfaceInfo, EventHandlerShader, RHICmdList);
@@ -290,7 +369,7 @@ void NiagaraEmitterInstanceBatcher::RunEventHandlers(const FNiagaraComputeExecut
 				uint32 EventSpawnNum = NumEventsToProcess * EventHandlerProps.SpawnNumber;
 
 				//int32 SpawnNum = Context->EventSpawnCounts[EventScriptIdx];
-				FNiagaraShader *EventHandlerShader = EventHandlerProps.Script->GetRenderThreadScript()->GetShader();
+				FNiagaraShader* EventHandlerShader = EventHandlerProps.Script->GetRenderThreadScript()->GetShader();
 				if (NumEvents && EventSpawnNum && EventHandlerShader)
 				{
 					//SetupDataInterfaceBuffers(EventHandlerProps.Script->DataInterfaceInfo, EventHandlerShader, RHICmdList);
@@ -324,41 +403,45 @@ void NiagaraEmitterInstanceBatcher::RunEventHandlers(const FNiagaraComputeExecut
 	*/
 }
 
-void NiagaraEmitterInstanceBatcher::SetPrevDataStrideParams(const FNiagaraDataSet *Set, FNiagaraShader *Shader, FRHICommandList &RHICmdList) const
+/* Set read strides for the sim shader
+ */
+void NiagaraEmitterInstanceBatcher::SetPrevDataStrideParams(const FNiagaraDataSet *Set, FNiagaraShader* Shader, FRHICommandList &RHICmdList) const
 {
-	int32 FloatStride = Set->PrevDataRender().GetFloatStride() / sizeof(float);
-	int32 IntStride = Set->PrevDataRender().GetInt32Stride() / sizeof(int32);
-	RHICmdList.SetShaderParameter(Shader->GetComputeShader(), 0, Shader->EventReadFloatStrideParams[0].GetBaseIndex(), sizeof(int32), &FloatStride);
-	RHICmdList.SetShaderParameter(Shader->GetComputeShader(), 0, Shader->EventReadIntStrideParams[0].GetBaseIndex(), sizeof(int32), &IntStride);
+	int32 FloatStride = Set->PrevData().GetFloatStride() / sizeof(float);
+	int32 IntStride = Set->PrevData().GetInt32Stride() / sizeof(int32);
+	RHICmdList.SetShaderParameter(Shader->GetComputeShader(), Shader->EventReadFloatStrideParams[0].GetBufferIndex(), Shader->EventReadFloatStrideParams[0].GetBaseIndex(), sizeof(int32), &FloatStride);
+	RHICmdList.SetShaderParameter(Shader->GetComputeShader(), Shader->EventReadIntStrideParams[0].GetBufferIndex(), Shader->EventReadIntStrideParams[0].GetBaseIndex(), sizeof(int32), &IntStride);
 }
 
 
+/* Set up UAVs for the event data sets
+*/
 void NiagaraEmitterInstanceBatcher::SetupEventUAVs(const FNiagaraComputeExecutionContext *Context, uint32 NewNumInstances, FRHICommandList &RHICmdList) const
 {
-	FNiagaraShader *UpdateShader = Context->RTUpdateScript->GetShader();
+	FNiagaraShader* UpdateShader = Context->RTUpdateScript->GetShader();
 
 	uint32 SetIndex = 0;
 	for (FNiagaraDataSet *Set : Context->UpdateEventWriteDataSets)
 	{
 		if (NewNumInstances)
 		{
-			Set->CurrDataRender().AllocateGPU(NewNumInstances, RHICmdList);
-			Set->CurrDataRender().SetNumInstances(NewNumInstances);
+			Set->CurrData().AllocateGPU(NewNumInstances, RHICmdList);
+			Set->CurrData().SetNumInstances(NewNumInstances);
 			FRWShaderParameter &FloatParam = UpdateShader->EventFloatUAVParams[SetIndex];
 			FRWShaderParameter &IntParam = UpdateShader->EventIntUAVParams[SetIndex];
 			if (FloatParam.IsUAVBound())
 			{
-				RHICmdList.SetUAVParameter(UpdateShader->GetComputeShader(), FloatParam.GetUAVIndex(), Set->CurrDataRender().GetGPUBufferFloat()->UAV);
+				RHICmdList.SetUAVParameter(UpdateShader->GetComputeShader(), FloatParam.GetUAVIndex(), Set->CurrData().GetGPUBufferFloat()->UAV);
 			}
 			if (IntParam.IsUAVBound())
 			{
-				RHICmdList.SetUAVParameter(UpdateShader->GetComputeShader(), IntParam.GetUAVIndex(), Set->CurrDataRender().GetGPUBufferInt()->UAV);
+				RHICmdList.SetUAVParameter(UpdateShader->GetComputeShader(), IntParam.GetUAVIndex(), Set->CurrData().GetGPUBufferInt()->UAV);
 			}
 
-			uint32 FloatStride = Set->CurrDataRender().GetFloatStride() / sizeof(float);
-			uint32 IntStride = Set->CurrDataRender().GetInt32Stride() / sizeof(int32);
-			RHICmdList.SetShaderParameter(UpdateShader->GetComputeShader(), 0, UpdateShader->EventWriteFloatStrideParams[SetIndex].GetBaseIndex(), sizeof(int32), &FloatStride);
-			RHICmdList.SetShaderParameter(UpdateShader->GetComputeShader(), 0, UpdateShader->EventWriteIntStrideParams[SetIndex].GetBaseIndex(), sizeof(int32), &IntStride);
+			uint32 FloatStride = Set->CurrData().GetFloatStride() / sizeof(float);
+			uint32 IntStride = Set->CurrData().GetInt32Stride() / sizeof(int32);
+			RHICmdList.SetShaderParameter(UpdateShader->GetComputeShader(), UpdateShader->EventWriteFloatStrideParams[SetIndex].GetBufferIndex(), UpdateShader->EventWriteFloatStrideParams[SetIndex].GetBaseIndex(), sizeof(int32), &FloatStride);
+			RHICmdList.SetShaderParameter(UpdateShader->GetComputeShader(), UpdateShader->EventWriteIntStrideParams[SetIndex].GetBufferIndex(), UpdateShader->EventWriteIntStrideParams[SetIndex].GetBaseIndex(), sizeof(int32), &IntStride);
 		}
 
 		SetIndex++;
@@ -370,7 +453,7 @@ void NiagaraEmitterInstanceBatcher::SetupEventUAVs(const FNiagaraComputeExecutio
 
 void NiagaraEmitterInstanceBatcher::UnsetEventUAVs(const FNiagaraComputeExecutionContext *Context, FRHICommandList &RHICmdList) const
 {
-	FNiagaraShader *UpdateShader = Context->RTUpdateScript->GetShader();
+	FNiagaraShader* UpdateShader = Context->RTUpdateScript->GetShader();
 
 	for (int32 SetIndex = 0; SetIndex < Context->UpdateEventWriteDataSets.Num(); SetIndex++)
 	{
diff --git a/Engine/Plugins/FX/Niagara/Source/Niagara/Private/NiagaraFunctionLibrary.cpp b/Engine/Plugins/FX/Niagara/Source/Niagara/Private/NiagaraFunctionLibrary.cpp
--- a/Engine/Plugins/FX/Niagara/Source/Niagara/Private/NiagaraFunctionLibrary.cpp
+++ b/Engine/Plugins/FX/Niagara/Source/Niagara/Private/NiagaraFunctionLibrary.cpp
@@ -31,12 +31,13 @@ UNiagaraComponent* UNiagaraFunctionLibrary::SpawnSystemAtLocation(UObject* World
 			PSC = NewObject<UNiagaraComponent>((Actor ? Actor : (UObject*)World));
 			PSC->SetAsset(SystemTemplate);
 			PSC->SetAutoDestroy(bAutoDestroy);
+			PSC->bAutoActivate = false;
 			PSC->RegisterComponentWithWorld(World);
 
 			PSC->SetAbsolute(true, true, true);
 			PSC->SetWorldLocationAndRotation(SpawnLocation, SpawnRotation);
 			PSC->SetRelativeScale3D(FVector(1.f));
-			//PSC->ActivateSystem(true);
+			PSC->Activate(true);
 		}
 	}
 	return PSC;
diff --git a/Engine/Plugins/FX/Niagara/Source/Niagara/Private/NiagaraLightRendererProperties.cpp b/Engine/Plugins/FX/Niagara/Source/Niagara/Private/NiagaraLightRendererProperties.cpp
--- a/Engine/Plugins/FX/Niagara/Source/Niagara/Private/NiagaraLightRendererProperties.cpp
+++ b/Engine/Plugins/FX/Niagara/Source/Niagara/Private/NiagaraLightRendererProperties.cpp
@@ -3,12 +3,32 @@
 #include "NiagaraLightRendererProperties.h"
 #include "NiagaraRenderer.h"
 #include "NiagaraConstants.h"
-
 UNiagaraLightRendererProperties::UNiagaraLightRendererProperties()
 	: RadiusScale(1.0f), ColorAdd(FVector(0.0f, 0.0f, 0.0f))
 {
 }
 
+void UNiagaraLightRendererProperties::PostInitProperties()
+{
+	Super::PostInitProperties();
+	if (HasAnyFlags(RF_ClassDefaultObject) == false && PositionBinding.BoundVariable.GetName() == NAME_None)
+	{
+		PositionBinding = FNiagaraConstants::GetAttributeDefaultBinding(SYS_PARAM_PARTICLES_POSITION);
+		ColorBinding = FNiagaraConstants::GetAttributeDefaultBinding(SYS_PARAM_PARTICLES_COLOR);
+		RadiusBinding = FNiagaraConstants::GetAttributeDefaultBinding(SYS_PARAM_PARTICLES_LIGHT_RADIUS);
+	}
+}
+
+/** The bindings depend on variables that are created during the NiagaraModule startup. However, the CDO's are build prior to this being initialized, so we defer setting these values until later.*/
+void UNiagaraLightRendererProperties::InitCDOPropertiesAfterModuleStartup()
+{
+	UNiagaraLightRendererProperties* CDO = CastChecked<UNiagaraLightRendererProperties>(UNiagaraLightRendererProperties::StaticClass()->GetDefaultObject());
+
+	CDO->PositionBinding = FNiagaraConstants::GetAttributeDefaultBinding(SYS_PARAM_PARTICLES_POSITION);
+	CDO->ColorBinding = FNiagaraConstants::GetAttributeDefaultBinding(SYS_PARAM_PARTICLES_COLOR);
+	CDO->RadiusBinding = FNiagaraConstants::GetAttributeDefaultBinding(SYS_PARAM_PARTICLES_LIGHT_RADIUS);
+}
+
 NiagaraRenderer* UNiagaraLightRendererProperties::CreateEmitterRenderer(ERHIFeatureLevel::Type FeatureLevel)
 {
 	return new NiagaraRendererLights(FeatureLevel, this);
@@ -25,33 +45,18 @@ void UNiagaraLightRendererProperties::GetUsedMaterials(TArray<UMaterialInterface
 const TArray<FNiagaraVariable>& UNiagaraLightRendererProperties::GetRequiredAttributes()
 {
 	static TArray<FNiagaraVariable> Attrs;
-
-	if (Attrs.Num() == 0)
-	{
-		Attrs.Add(SYS_PARAM_PARTICLES_POSITION);
-		Attrs.Add(SYS_PARAM_PARTICLES_COLOR);
-		Attrs.Add(SYS_PARAM_PARTICLES_SCALE);
-		/*Attrs.Add(SYS_PARAM_PARTICLES_VELOCITY);
-		Attrs.Add(SYS_PARAM_PARTICLES_SPRITE_ROTATION);
-		Attrs.Add(SYS_PARAM_PARTICLES_NORMALIZED_AGE);
-		Attrs.Add(SYS_PARAM_PARTICLES_SPRITE_SIZE);*/
-	}
-
 	return Attrs;
 }
 
 const TArray<FNiagaraVariable>& UNiagaraLightRendererProperties::GetOptionalAttributes()
 {
 	static TArray<FNiagaraVariable> Attrs;
-
 	if (Attrs.Num() == 0)
 	{
-		//Attrs.Add(SYS_PARAM_PARTICLES_SPRITE_FACING);
-		//Attrs.Add(SYS_PARAM_PARTICLES_SPRITE_ALIGNMENT);
-		//Attrs.Add(SYS_PARAM_PARTICLES_SUB_IMAGE_INDEX);
-		//Attrs.Add(SYS_PARAM_PARTICLES_DYNAMIC_MATERIAL_PARAM);
+		Attrs.Add(SYS_PARAM_PARTICLES_POSITION);
+		Attrs.Add(SYS_PARAM_PARTICLES_COLOR);
+		Attrs.Add(SYS_PARAM_PARTICLES_LIGHT_RADIUS);
 	}
-
 	return Attrs;
 }
 
diff --git a/Engine/Plugins/FX/Niagara/Source/Niagara/Private/NiagaraMeshRendererProperties.cpp b/Engine/Plugins/FX/Niagara/Source/Niagara/Private/NiagaraMeshRendererProperties.cpp
--- a/Engine/Plugins/FX/Niagara/Source/Niagara/Private/NiagaraMeshRendererProperties.cpp
+++ b/Engine/Plugins/FX/Niagara/Source/Niagara/Private/NiagaraMeshRendererProperties.cpp
@@ -7,6 +7,8 @@
 
 UNiagaraMeshRendererProperties::UNiagaraMeshRendererProperties()
 	: ParticleMesh(nullptr)
+	, SortMode(ENiagaraSortMode::ViewDistance)
+	, bSortOnlyWhenTranslucent(true)
 {
 }
 
@@ -15,6 +17,43 @@ NiagaraRenderer* UNiagaraMeshRendererProperties::CreateEmitterRenderer(ERHIFeatu
 	return new NiagaraRendererMeshes(FeatureLevel, this);
 }
 
+
+
+void UNiagaraMeshRendererProperties::PostInitProperties()
+{
+	Super::PostInitProperties();
+	SyncId = 0;
+	if (HasAnyFlags(RF_ClassDefaultObject) == false)
+	{
+		InitBindings();
+	}
+}
+
+/** The bindings depend on variables that are created during the NiagaraModule startup. However, the CDO's are build prior to this being initialized, so we defer setting these values until later.*/
+void UNiagaraMeshRendererProperties::InitCDOPropertiesAfterModuleStartup()
+{
+	UNiagaraMeshRendererProperties* CDO = CastChecked<UNiagaraMeshRendererProperties>(UNiagaraMeshRendererProperties::StaticClass()->GetDefaultObject());
+	CDO->InitBindings();
+}
+
+void UNiagaraMeshRendererProperties::InitBindings()
+{
+	if (PositionBinding.BoundVariable.GetName() == NAME_None)
+	{
+		PositionBinding = FNiagaraConstants::GetAttributeDefaultBinding(SYS_PARAM_PARTICLES_POSITION);
+		ColorBinding = FNiagaraConstants::GetAttributeDefaultBinding(SYS_PARAM_PARTICLES_COLOR);
+		VelocityBinding = FNiagaraConstants::GetAttributeDefaultBinding(SYS_PARAM_PARTICLES_VELOCITY);
+		DynamicMaterialBinding = FNiagaraConstants::GetAttributeDefaultBinding(SYS_PARAM_PARTICLES_DYNAMIC_MATERIAL_PARAM);
+		DynamicMaterial1Binding = FNiagaraConstants::GetAttributeDefaultBinding(SYS_PARAM_PARTICLES_DYNAMIC_MATERIAL_PARAM_1);
+		DynamicMaterial2Binding = FNiagaraConstants::GetAttributeDefaultBinding(SYS_PARAM_PARTICLES_DYNAMIC_MATERIAL_PARAM_2);
+		DynamicMaterial3Binding = FNiagaraConstants::GetAttributeDefaultBinding(SYS_PARAM_PARTICLES_DYNAMIC_MATERIAL_PARAM_3);
+		MeshOrientationBinding = FNiagaraConstants::GetAttributeDefaultBinding(SYS_PARAM_PARTICLES_MESH_ORIENTATION);
+		ScaleBinding = FNiagaraConstants::GetAttributeDefaultBinding(SYS_PARAM_PARTICLES_SCALE);
+		//Default custom sorting to age
+		CustomSortingBinding = FNiagaraConstants::GetAttributeDefaultBinding(SYS_PARAM_PARTICLES_NORMALIZED_AGE);
+	}
+}
+
 void UNiagaraMeshRendererProperties::GetUsedMaterials(TArray<UMaterialInterface*>& OutMaterials) const
 {
 	if (ParticleMesh)
@@ -49,8 +88,13 @@ void UNiagaraMeshRendererProperties::GetUsedMaterials(TArray<UMaterialInterface*
 	}
 }
 
-#if WITH_EDITORONLY_DATA
+uint32 UNiagaraMeshRendererProperties::GetNumIndicesPerInstance() 
+{
+	return ParticleMesh ? ParticleMesh->RenderData->LODResources[0].IndexBuffer.GetNumIndices() : 0;
+}
 
+
+#if WITH_EDITORONLY_DATA
 bool UNiagaraMeshRendererProperties::IsMaterialValidForRenderer(UMaterial* Material, FText& InvalidMessage)
 {
 	if (Material->bUsedWithNiagaraMeshParticles == false)
@@ -74,10 +118,6 @@ const TArray<FNiagaraVariable>& UNiagaraMeshRendererProperties::GetRequiredAttri
 
 	if (Attrs.Num() == 0)
 	{
-		Attrs.Add(SYS_PARAM_PARTICLES_POSITION);
-		Attrs.Add(SYS_PARAM_PARTICLES_VELOCITY);
-		Attrs.Add(SYS_PARAM_PARTICLES_COLOR);
-		Attrs.Add(SYS_PARAM_PARTICLES_NORMALIZED_AGE);
 	}
 
 	return Attrs;
@@ -90,18 +130,21 @@ const TArray<FNiagaraVariable>& UNiagaraMeshRendererProperties::GetOptionalAttri
 
 	if (Attrs.Num() == 0)
 	{
+		Attrs.Add(SYS_PARAM_PARTICLES_POSITION);
+		Attrs.Add(SYS_PARAM_PARTICLES_VELOCITY);
+		Attrs.Add(SYS_PARAM_PARTICLES_COLOR);
+		Attrs.Add(SYS_PARAM_PARTICLES_NORMALIZED_AGE);
 		Attrs.Add(SYS_PARAM_PARTICLES_SCALE);
-		//Attrs.Add(SYS_PARAM_PARTICLES_SPRITE_FACING);
-		//Attrs.Add(SYS_PARAM_PARTICLES_SPRITE_ALIGNMENT);
-		//Attrs.Add(SYS_PARAM_PARTICLES_SUB_IMAGE_INDEX);
+		Attrs.Add(SYS_PARAM_PARTICLES_MESH_ORIENTATION);
 		Attrs.Add(SYS_PARAM_PARTICLES_DYNAMIC_MATERIAL_PARAM);
+		Attrs.Add(SYS_PARAM_PARTICLES_DYNAMIC_MATERIAL_PARAM_1);
+		Attrs.Add(SYS_PARAM_PARTICLES_DYNAMIC_MATERIAL_PARAM_2);
+		Attrs.Add(SYS_PARAM_PARTICLES_DYNAMIC_MATERIAL_PARAM_3);
 	}
 
 	return Attrs;
 }
 
-#endif // WITH_EDITORONLY_DATA
-
 void UNiagaraMeshRendererProperties::PostEditChangeProperty(FPropertyChangedEvent& PropertyChangedEvent)
 {
 	if (ParticleMesh && PropertyChangedEvent.Property && PropertyChangedEvent.Property->GetName() == "ParticleMesh")
@@ -118,4 +161,11 @@ void UNiagaraMeshRendererProperties::PostEditChangeProperty(FPropertyChangedEven
 			}
 		}
 	}
+	
+	if (PropertyChangedEvent.GetPropertyName() != TEXT("SyncId"))
+	{
+		SyncId++;
+	}
 }
+
+#endif // WITH_EDITORONLY_DATA
\ No newline at end of file
diff --git a/Engine/Plugins/FX/Niagara/Source/Niagara/Private/NiagaraModule.cpp b/Engine/Plugins/FX/Niagara/Source/Niagara/Private/NiagaraModule.cpp
--- a/Engine/Plugins/FX/Niagara/Source/Niagara/Private/NiagaraModule.cpp
+++ b/Engine/Plugins/FX/Niagara/Source/Niagara/Private/NiagaraModule.cpp
@@ -10,11 +10,23 @@
 #include "UObject/Package.h"
 #include "NiagaraWorldManager.h"
 #include "VectorVM.h"
+#include "NiagaraConstants.h"
+#include "NiagaraLightRendererProperties.h"
+#include "NiagaraSpriteRendererProperties.h"
+#include "NiagaraMeshRendererProperties.h"
+#include "NiagaraRibbonRendererProperties.h"
+#include "NiagaraRenderer.h"
 
 IMPLEMENT_MODULE(INiagaraModule, Niagara);
 
 #define LOCTEXT_NAMESPACE "NiagaraModule"
 
+TMap<class UWorld*, class FNiagaraWorldManager*> INiagaraModule::WorldManagers;
+
+float INiagaraModule::EngineGlobalSpawnCountScale = 1.0f;
+float INiagaraModule::EngineGlobalSystemCountScale = 1.0f;
+int32 INiagaraModule::EngineDetailLevel = 4;
+
 /**
 Detail Level CVar.
 Effectively replaces the DetaiMode feature but allows for a rolling range of new hardware and emitters to target them.
@@ -29,10 +41,95 @@ static TAutoConsoleVariable<float> CVarDetailLevel(
 	TEXT("Default = 4"),
 	ECVF_Scalability);
 
+static FAutoConsoleVariableRef CVarNiaraGlobalSpawnCountScale(
+	TEXT("fx.NiagaraGlobalSpawnCountScale"),
+	INiagaraModule::EngineGlobalSpawnCountScale,
+	TEXT("A global scale on spawn counts in Niagara. \n"),
+	ECVF_Scalability
+);
+
+static FAutoConsoleVariableRef CVarNiaraGlobalSystemCountScale(
+	TEXT("fx.NiagaraGlobalSystemCountScale"),
+	INiagaraModule::EngineGlobalSystemCountScale,
+	TEXT("A global scale on system count thresholds for culling in Niagara. \n"),
+	ECVF_Scalability
+);
+
+FNiagaraVariable INiagaraModule::Engine_DeltaTime;
+FNiagaraVariable INiagaraModule::Engine_InvDeltaTime;
+FNiagaraVariable INiagaraModule::Engine_Time;
+FNiagaraVariable INiagaraModule::Engine_RealTime;
+
+FNiagaraVariable INiagaraModule::Engine_Owner_Position;
+FNiagaraVariable INiagaraModule::Engine_Owner_Velocity;
+FNiagaraVariable INiagaraModule::Engine_Owner_XAxis;
+FNiagaraVariable INiagaraModule::Engine_Owner_YAxis;
+FNiagaraVariable INiagaraModule::Engine_Owner_ZAxis;
+FNiagaraVariable INiagaraModule::Engine_Owner_Scale;
+
+FNiagaraVariable INiagaraModule::Engine_Owner_SystemLocalToWorld;
+FNiagaraVariable INiagaraModule::Engine_Owner_SystemWorldToLocal;
+FNiagaraVariable INiagaraModule::Engine_Owner_SystemLocalToWorldTransposed;
+FNiagaraVariable INiagaraModule::Engine_Owner_SystemWorldToLocalTransposed;
+FNiagaraVariable INiagaraModule::Engine_Owner_SystemLocalToWorldNoScale;
+FNiagaraVariable INiagaraModule::Engine_Owner_SystemWorldToLocalNoScale;
+
+FNiagaraVariable INiagaraModule::Engine_Owner_TimeSinceRendered;
+FNiagaraVariable INiagaraModule::Engine_Owner_MinDistanceToCamera;
+
+FNiagaraVariable INiagaraModule::Engine_Owner_ExecutionState;
+
+FNiagaraVariable INiagaraModule::Engine_ExecutionCount;
+FNiagaraVariable INiagaraModule::Engine_Emitter_NumParticles;
+FNiagaraVariable INiagaraModule::Engine_System_NumEmittersAlive;
+FNiagaraVariable INiagaraModule::Engine_System_NumEmitters;
+FNiagaraVariable INiagaraModule::Engine_NumSystemInstances;
+
+FNiagaraVariable INiagaraModule::Engine_GlobalSpawnCountScale;
+FNiagaraVariable INiagaraModule::Engine_GlobalSystemScale;
+
+FNiagaraVariable INiagaraModule::Engine_System_Age;
+
+FNiagaraVariable INiagaraModule::Emitter_Age;
+FNiagaraVariable INiagaraModule::Emitter_LocalSpace;
+FNiagaraVariable INiagaraModule::Emitter_SpawnRate;
+FNiagaraVariable INiagaraModule::Emitter_SpawnInterval;
+FNiagaraVariable INiagaraModule::Emitter_InterpSpawnStartDt;
+
+FNiagaraVariable INiagaraModule::Particles_ID;
+FNiagaraVariable INiagaraModule::Particles_Position;
+FNiagaraVariable INiagaraModule::Particles_Velocity;
+FNiagaraVariable INiagaraModule::Particles_Color;
+FNiagaraVariable INiagaraModule::Particles_SpriteRotation;
+FNiagaraVariable INiagaraModule::Particles_NormalizedAge;
+FNiagaraVariable INiagaraModule::Particles_SpriteSize;
+FNiagaraVariable INiagaraModule::Particles_SpriteFacing;
+FNiagaraVariable INiagaraModule::Particles_SpriteAlignment;
+FNiagaraVariable INiagaraModule::Particles_SubImageIndex;
+FNiagaraVariable INiagaraModule::Particles_DynamicMaterialParameter;
+FNiagaraVariable INiagaraModule::Particles_DynamicMaterialParameter1;
+FNiagaraVariable INiagaraModule::Particles_DynamicMaterialParameter2;
+FNiagaraVariable INiagaraModule::Particles_DynamicMaterialParameter3;
+FNiagaraVariable INiagaraModule::Particles_Scale;
+FNiagaraVariable INiagaraModule::Particles_Lifetime;
+FNiagaraVariable INiagaraModule::Particles_MeshOrientation;
+FNiagaraVariable INiagaraModule::Particles_UVScale;
+FNiagaraVariable INiagaraModule::Particles_CameraOffset;
+FNiagaraVariable INiagaraModule::Particles_MaterialRandom;
+FNiagaraVariable INiagaraModule::Particles_LightRadius;
+FNiagaraVariable INiagaraModule::Particles_RibbonID;
+FNiagaraVariable INiagaraModule::Particles_RibbonWidth;
+FNiagaraVariable INiagaraModule::Particles_RibbonTwist;
+FNiagaraVariable INiagaraModule::Particles_RibbonFacing;
+FNiagaraVariable INiagaraModule::Particles_RibbonLinkOrder;
+FNiagaraVariable INiagaraModule::DataInstance_Alive;
+FNiagaraVariable INiagaraModule::Translator_BeginDefaults;
+
 void INiagaraModule::StartupModule()
 {
 	VectorVM::Init();
 	FNiagaraTypeDefinition::Init();
+	FNiagaraViewDataMgr::Init();
 
 	FWorldDelegates::OnPreWorldInitialization.AddRaw(this, &INiagaraModule::OnWorldInit);
 	FWorldDelegates::OnWorldCleanup.AddRaw(this, &INiagaraModule::OnWorldCleanup);
@@ -48,6 +145,86 @@ void INiagaraModule::StartupModule()
 
 	CVarDetailLevel.AsVariable()->SetOnChangedCallback(FConsoleVariableDelegate::CreateRaw(this, &INiagaraModule::OnChangeDetailLevel));
 	OnChangeDetailLevel(CVarDetailLevel.AsVariable());
+
+	//Init commonly used FNiagaraVariables
+
+	Engine_DeltaTime = FNiagaraVariable(FNiagaraTypeDefinition::GetFloatDef(), TEXT("Engine.DeltaTime"));
+	Engine_InvDeltaTime = FNiagaraVariable(FNiagaraTypeDefinition::GetFloatDef(), TEXT("Engine.InverseDeltaTime"));
+	
+	Engine_Time = FNiagaraVariable(FNiagaraTypeDefinition::GetFloatDef(), TEXT("Engine.Time"));
+	Engine_RealTime = FNiagaraVariable(FNiagaraTypeDefinition::GetFloatDef(), TEXT("Engine.RealTime"));
+
+	Engine_Owner_Position = FNiagaraVariable(FNiagaraTypeDefinition::GetVec3Def(), TEXT("Engine.Owner.Position"));
+	Engine_Owner_Velocity = FNiagaraVariable(FNiagaraTypeDefinition::GetVec3Def(), TEXT("Engine.Owner.Velocity"));
+	Engine_Owner_XAxis = FNiagaraVariable(FNiagaraTypeDefinition::GetVec3Def(), TEXT("Engine.Owner.SystemXAxis"));
+	Engine_Owner_YAxis = FNiagaraVariable(FNiagaraTypeDefinition::GetVec3Def(), TEXT("Engine.Owner.SystemYAxis"));
+	Engine_Owner_ZAxis = FNiagaraVariable(FNiagaraTypeDefinition::GetVec3Def(), TEXT("Engine.Owner.SystemZAxis"));
+	Engine_Owner_Scale = FNiagaraVariable(FNiagaraTypeDefinition::GetVec3Def(), TEXT("Engine.Owner.Scale"));
+
+	Engine_Owner_SystemLocalToWorld = FNiagaraVariable(FNiagaraTypeDefinition::GetMatrix4Def(), TEXT("Engine.Owner.SystemLocalToWorld"));
+	Engine_Owner_SystemWorldToLocal = FNiagaraVariable(FNiagaraTypeDefinition::GetMatrix4Def(), TEXT("Engine.Owner.SystemWorldToLocal"));
+	Engine_Owner_SystemLocalToWorldTransposed = FNiagaraVariable(FNiagaraTypeDefinition::GetMatrix4Def(), TEXT("Engine.Owner.SystemLocalToWorldTransposed"));
+	Engine_Owner_SystemWorldToLocalTransposed = FNiagaraVariable(FNiagaraTypeDefinition::GetMatrix4Def(), TEXT("Engine.Owner.SystemWorldToLocalTransposed"));
+	Engine_Owner_SystemLocalToWorldNoScale = FNiagaraVariable(FNiagaraTypeDefinition::GetMatrix4Def(), TEXT("Engine.Owner.SystemLocalToWorldNoScale"));
+	Engine_Owner_SystemWorldToLocalNoScale = FNiagaraVariable(FNiagaraTypeDefinition::GetMatrix4Def(), TEXT("Engine.Owner.SystemWorldToLocalNoScale"));
+
+	Engine_Owner_TimeSinceRendered = FNiagaraVariable(FNiagaraTypeDefinition::GetFloatDef(), TEXT("Engine.Owner.TimeSinceRendered"));
+	Engine_Owner_MinDistanceToCamera = FNiagaraVariable(FNiagaraTypeDefinition::GetFloatDef(), TEXT("Engine.Owner.MinDistanceToCamera"));
+
+	Engine_Owner_ExecutionState = FNiagaraVariable(FNiagaraTypeDefinition::GetExecutionStateEnum(), TEXT("Engine.Owner.ExecutionState"));
+	
+	Engine_ExecutionCount = FNiagaraVariable(FNiagaraTypeDefinition::GetIntDef(), TEXT("Engine.ExecutionCount"));
+	Engine_Emitter_NumParticles = FNiagaraVariable(FNiagaraTypeDefinition::GetIntDef(), TEXT("Engine.Emitter.NumParticles"));
+	Engine_System_NumEmittersAlive = FNiagaraVariable(FNiagaraTypeDefinition::GetIntDef(), TEXT("Engine.System.NumEmittersAlive"));
+	Engine_System_NumEmitters = FNiagaraVariable(FNiagaraTypeDefinition::GetIntDef(), TEXT("Engine.System.NumEmitters"));
+	Engine_NumSystemInstances = FNiagaraVariable(FNiagaraTypeDefinition::GetIntDef(), TEXT("Engine.NumSystemInstances"));
+
+	Engine_GlobalSpawnCountScale = FNiagaraVariable(FNiagaraTypeDefinition::GetFloatDef(), TEXT("Engine.GlobalSpawnCountScale"));
+	Engine_GlobalSystemScale = FNiagaraVariable(FNiagaraTypeDefinition::GetFloatDef(), TEXT("Engine.GlobalSystemCountScale"));
+
+	Engine_System_Age = FNiagaraVariable(FNiagaraTypeDefinition::GetFloatDef(), TEXT("Engine.System.Age"));
+	Emitter_Age = FNiagaraVariable(FNiagaraTypeDefinition::GetFloatDef(), TEXT("Emitter.Age"));
+	Emitter_LocalSpace = FNiagaraVariable(FNiagaraTypeDefinition::GetBoolDef(), TEXT("Emitter.LocalSpace"));
+	Emitter_SpawnRate = FNiagaraVariable(FNiagaraTypeDefinition::GetFloatDef(), TEXT("Emitter.SpawnRate"));
+	Emitter_SpawnInterval = FNiagaraVariable(FNiagaraTypeDefinition::GetFloatDef(), TEXT("Emitter.SpawnInterval"));
+	Emitter_InterpSpawnStartDt = FNiagaraVariable(FNiagaraTypeDefinition::GetFloatDef(), TEXT("Emitter.InterpSpawnStartDt"));
+
+	Particles_ID = FNiagaraVariable(FNiagaraTypeDefinition::GetIDDef(), TEXT("Particles.ID"));
+	Particles_Position = FNiagaraVariable(FNiagaraTypeDefinition::GetVec3Def(), TEXT("Particles.Position"));
+	Particles_Velocity = FNiagaraVariable(FNiagaraTypeDefinition::GetVec3Def(), TEXT("Particles.Velocity"));
+	Particles_Color = FNiagaraVariable(FNiagaraTypeDefinition::GetColorDef(), TEXT("Particles.Color"));
+	Particles_SpriteRotation = FNiagaraVariable(FNiagaraTypeDefinition::GetFloatDef(), TEXT("Particles.SpriteRotation"));
+	Particles_NormalizedAge = FNiagaraVariable(FNiagaraTypeDefinition::GetFloatDef(), TEXT("Particles.NormalizedAge"));
+	Particles_SpriteSize = FNiagaraVariable(FNiagaraTypeDefinition::GetVec2Def(), TEXT("Particles.SpriteSize"));
+	Particles_SpriteFacing = FNiagaraVariable(FNiagaraTypeDefinition::GetVec3Def(), TEXT("Particles.SpriteFacing"));
+	Particles_SpriteAlignment = FNiagaraVariable(FNiagaraTypeDefinition::GetVec3Def(), TEXT("Particles.SpriteAlignment"));
+	Particles_SubImageIndex = FNiagaraVariable(FNiagaraTypeDefinition::GetFloatDef(), TEXT("Particles.SubImageIndex"));
+	Particles_DynamicMaterialParameter = FNiagaraVariable(FNiagaraTypeDefinition::GetVec4Def(), TEXT("Particles.DynamicMaterialParameter"));
+	Particles_DynamicMaterialParameter1 = FNiagaraVariable(FNiagaraTypeDefinition::GetVec4Def(), TEXT("Particles.DynamicMaterialParameter1"));
+	Particles_DynamicMaterialParameter2 = FNiagaraVariable(FNiagaraTypeDefinition::GetVec4Def(), TEXT("Particles.DynamicMaterialParameter2"));
+	Particles_DynamicMaterialParameter3 = FNiagaraVariable(FNiagaraTypeDefinition::GetVec4Def(), TEXT("Particles.DynamicMaterialParameter3"));
+	Particles_Scale = FNiagaraVariable(FNiagaraTypeDefinition::GetVec3Def(), TEXT("Particles.Scale"));
+	Particles_Lifetime = FNiagaraVariable(FNiagaraTypeDefinition::GetFloatDef(), TEXT("Particles.Lifetime"));
+	Particles_MeshOrientation = FNiagaraVariable(FNiagaraTypeDefinition::GetQuatDef(), TEXT("Particles.MeshOrientation"));
+	Particles_UVScale = FNiagaraVariable(FNiagaraTypeDefinition::GetVec2Def(), TEXT("Particles.UVScale"));
+	Particles_CameraOffset = FNiagaraVariable(FNiagaraTypeDefinition::GetFloatDef(), TEXT("Particles.CameraOffset"));
+	Particles_MaterialRandom = FNiagaraVariable(FNiagaraTypeDefinition::GetFloatDef(), TEXT("Particles.MaterialRandom"));
+	Particles_LightRadius = FNiagaraVariable(FNiagaraTypeDefinition::GetFloatDef(), TEXT("Particles.LightRadius"));
+	Particles_RibbonID = FNiagaraVariable(FNiagaraTypeDefinition::GetIDDef(), TEXT("Particles.RibbonID"));
+	Particles_RibbonWidth = FNiagaraVariable(FNiagaraTypeDefinition::GetFloatDef(), TEXT("Particles.RibbonWidth"));
+	Particles_RibbonTwist = FNiagaraVariable(FNiagaraTypeDefinition::GetFloatDef(), TEXT("Particles.RibbonTwist"));
+	Particles_RibbonFacing = FNiagaraVariable(FNiagaraTypeDefinition::GetVec3Def(), TEXT("Particles.RibbonFacing"));
+	Particles_RibbonLinkOrder = FNiagaraVariable(FNiagaraTypeDefinition::GetFloatDef(), TEXT("Particles.RibbonLinkOrder"));
+
+	DataInstance_Alive = FNiagaraVariable(FNiagaraTypeDefinition::GetBoolDef(), TEXT("DataInstance.Alive"));
+
+	Translator_BeginDefaults = FNiagaraVariable(FNiagaraTypeDefinition::GetParameterMapDef(), TEXT("Begin Defaults"));
+
+	FNiagaraConstants::Init();
+	UNiagaraLightRendererProperties::InitCDOPropertiesAfterModuleStartup();
+	UNiagaraSpriteRendererProperties::InitCDOPropertiesAfterModuleStartup();
+	UNiagaraRibbonRendererProperties::InitCDOPropertiesAfterModuleStartup();
+	UNiagaraMeshRendererProperties::InitCDOPropertiesAfterModuleStartup();
 }
 
 void INiagaraModule::ShutdownModule()
@@ -61,6 +238,7 @@ void INiagaraModule::ShutdownModule()
 	}
 
 	CVarDetailLevel.AsVariable()->SetOnChangedCallback(FConsoleVariableDelegate());
+	FNiagaraViewDataMgr::Shutdown();
 }
 
 FNiagaraWorldManager* INiagaraModule::GetWorldManager(UWorld* World)
@@ -133,9 +311,70 @@ void INiagaraModule::UnregisterOnMergeEmitter(FDelegateHandle DelegateHandle)
 	checkf(OnMergeEmitterDelegate.GetHandle() == DelegateHandle, TEXT("Can only unregister the OnMergeEmitter delegate with the handle it was registered with."));
 	OnMergeEmitterDelegate.Unbind();
 }
+
+UNiagaraScriptSourceBase* INiagaraModule::CreateDefaultScriptSource(UObject* Outer)
+{
+	checkf(OnCreateDefaultScriptSourceDelegate.IsBound(), TEXT("Create default script source delegate not bound."));
+	return OnCreateDefaultScriptSourceDelegate.Execute(Outer);
+}
+
+FDelegateHandle INiagaraModule::RegisterOnCreateDefaultScriptSource(FOnCreateDefaultScriptSource OnCreateDefaultScriptSource)
+{
+	checkf(OnCreateDefaultScriptSourceDelegate.IsBound() == false, TEXT("Only one handler is allowed for the OnCreateDefaultScriptSource delegate"));
+	OnCreateDefaultScriptSourceDelegate = OnCreateDefaultScriptSource;
+	return OnCreateDefaultScriptSourceDelegate.GetHandle();
+}
+
+void INiagaraModule::UnregisterOnCreateDefaultScriptSource(FDelegateHandle DelegateHandle)
+{
+	checkf(OnCreateDefaultScriptSourceDelegate.IsBound(), TEXT("OnCreateDefaultScriptSource is not registered"));
+	checkf(OnCreateDefaultScriptSourceDelegate.GetHandle() == DelegateHandle, TEXT("Can only unregister the OnCreateDefaultScriptSource delegate with the handle it was registered with."));
+	OnCreateDefaultScriptSourceDelegate.Unbind();
+}
+
+TSharedPtr<FNiagaraVMExecutableData> INiagaraModule::CompileScript(const FNiagaraCompileRequestDataBase* InCompileData, const FNiagaraCompileOptions& InCompileOptions)
+{
+	checkf(ScriptCompilerDelegate.IsBound(), TEXT("Create default script compiler delegate not bound."));
+	return ScriptCompilerDelegate.Execute(InCompileData, InCompileOptions);
+}
+
+FDelegateHandle INiagaraModule::RegisterScriptCompiler(FScriptCompiler ScriptCompiler)
+{
+	checkf(ScriptCompilerDelegate.IsBound() == false, TEXT("Only one handler is allowed for the ScriptCompiler delegate"));
+	ScriptCompilerDelegate = ScriptCompiler;
+	return ScriptCompilerDelegate.GetHandle();
+}
+
+void INiagaraModule::UnregisterScriptCompiler(FDelegateHandle DelegateHandle)
+{
+	checkf(ScriptCompilerDelegate.IsBound(), TEXT("ScriptCompiler is not registered"));
+	checkf(ScriptCompilerDelegate.GetHandle() == DelegateHandle, TEXT("Can only unregister the ScriptCompiler delegate with the handle it was registered with."));
+	ScriptCompilerDelegate.Unbind();
+}
+
+
+TSharedPtr<FNiagaraCompileRequestDataBase, ESPMode::ThreadSafe> INiagaraModule::Precompile(UObject* Obj)
+{
+	checkf(ObjectPrecompilerDelegate.IsBound(), TEXT("ObjectPrecompiler delegate not bound."));
+	return ObjectPrecompilerDelegate.Execute(Obj);
+}
+
+FDelegateHandle INiagaraModule::RegisterPrecompiler(FOnPrecompile PreCompiler)
+{
+	checkf(ObjectPrecompilerDelegate.IsBound() == false, TEXT("Only one handler is allowed for the ObjectPrecompiler delegate"));
+	ObjectPrecompilerDelegate = PreCompiler;
+	return ObjectPrecompilerDelegate.GetHandle();
+}
+
+void INiagaraModule::UnregisterPrecompiler(FDelegateHandle DelegateHandle)
+{
+	checkf(ObjectPrecompilerDelegate.IsBound(), TEXT("ObjectPrecompiler is not registered"));
+	checkf(ObjectPrecompilerDelegate.GetHandle() == DelegateHandle, TEXT("Can only unregister the ObjectPrecompiler delegate with the handle it was registered with."));
+	ObjectPrecompilerDelegate.Unbind();
+}
+
 #endif
 
-int32 INiagaraModule::EngineDetailLevel = 4;
 void INiagaraModule::OnChangeDetailLevel(class IConsoleVariable* CVar)
 {
 	//Can only change the detail level at runtime on when not cooked.
@@ -153,20 +392,24 @@ void INiagaraModule::OnChangeDetailLevel(class IConsoleVariable* CVar)
 
 //////////////////////////////////////////////////////////////////////////
 
-const UScriptStruct* FNiagaraTypeDefinition::ParameterMapStruct;
-const UScriptStruct* FNiagaraTypeDefinition::NumericStruct;
-const UScriptStruct* FNiagaraTypeDefinition::FloatStruct;
-const UScriptStruct* FNiagaraTypeDefinition::BoolStruct;
-const UScriptStruct* FNiagaraTypeDefinition::IntStruct;
-const UScriptStruct* FNiagaraTypeDefinition::Matrix4Struct;
-const UScriptStruct* FNiagaraTypeDefinition::Vec4Struct;
-const UScriptStruct* FNiagaraTypeDefinition::Vec3Struct;
-const UScriptStruct* FNiagaraTypeDefinition::Vec2Struct;
-const UScriptStruct* FNiagaraTypeDefinition::ColorStruct;
-
-const UEnum* FNiagaraTypeDefinition::ExecutionStateEnum;
+UScriptStruct* FNiagaraTypeDefinition::ParameterMapStruct;
+UScriptStruct* FNiagaraTypeDefinition::IDStruct;
+UScriptStruct* FNiagaraTypeDefinition::NumericStruct;
+UScriptStruct* FNiagaraTypeDefinition::FloatStruct;
+UScriptStruct* FNiagaraTypeDefinition::BoolStruct;
+UScriptStruct* FNiagaraTypeDefinition::IntStruct;
+UScriptStruct* FNiagaraTypeDefinition::Matrix4Struct;
+UScriptStruct* FNiagaraTypeDefinition::Vec4Struct;
+UScriptStruct* FNiagaraTypeDefinition::Vec3Struct;
+UScriptStruct* FNiagaraTypeDefinition::Vec2Struct;
+UScriptStruct* FNiagaraTypeDefinition::ColorStruct;
+UScriptStruct* FNiagaraTypeDefinition::QuatStruct;
+
+UEnum* FNiagaraTypeDefinition::ExecutionStateEnum;
+UEnum* FNiagaraTypeDefinition::ExecutionStateSourceEnum;
 
 FNiagaraTypeDefinition FNiagaraTypeDefinition::ParameterMapDef;
+FNiagaraTypeDefinition FNiagaraTypeDefinition::IDDef;
 FNiagaraTypeDefinition FNiagaraTypeDefinition::NumericDef;
 FNiagaraTypeDefinition FNiagaraTypeDefinition::FloatDef;
 FNiagaraTypeDefinition FNiagaraTypeDefinition::BoolDef;
@@ -176,15 +419,16 @@ FNiagaraTypeDefinition FNiagaraTypeDefinition::Vec4Def;
 FNiagaraTypeDefinition FNiagaraTypeDefinition::Vec3Def;
 FNiagaraTypeDefinition FNiagaraTypeDefinition::Vec2Def;
 FNiagaraTypeDefinition FNiagaraTypeDefinition::ColorDef;
+FNiagaraTypeDefinition FNiagaraTypeDefinition::QuatDef;
 
-TSet<const UScriptStruct*> FNiagaraTypeDefinition::NumericStructs;
+TSet<UScriptStruct*> FNiagaraTypeDefinition::NumericStructs;
 TArray<FNiagaraTypeDefinition> FNiagaraTypeDefinition::OrderedNumericTypes;
 
-TSet<const UScriptStruct*> FNiagaraTypeDefinition::ScalarStructs;
+TSet<UScriptStruct*> FNiagaraTypeDefinition::ScalarStructs;
 
-TSet<const UStruct*> FNiagaraTypeDefinition::FloatStructs;
-TSet<const UStruct*> FNiagaraTypeDefinition::IntStructs;
-TSet<const UStruct*> FNiagaraTypeDefinition::BoolStructs;
+TSet<UStruct*> FNiagaraTypeDefinition::FloatStructs;
+TSet<UStruct*> FNiagaraTypeDefinition::IntStructs;
+TSet<UStruct*> FNiagaraTypeDefinition::BoolStructs;
 
 FNiagaraTypeDefinition FNiagaraTypeDefinition::CollisionEventDef;
 
@@ -195,6 +439,7 @@ TArray<FNiagaraTypeDefinition> FNiagaraTypeRegistry::RegisteredPayloadTypes;
 TArray<FNiagaraTypeDefinition> FNiagaraTypeRegistry::RegisteredUserDefinedTypes;
 TArray<FNiagaraTypeDefinition> FNiagaraTypeRegistry::RegisteredNumericTypes;
 
+
 bool FNiagaraTypeDefinition::IsDataInterface()const
 {
 	return Struct->IsChildOf(UNiagaraDataInterface::StaticClass());
@@ -205,6 +450,7 @@ void FNiagaraTypeDefinition::Init()
 	static auto* CoreUObjectPkg = FindObjectChecked<UPackage>(nullptr, TEXT("/Script/CoreUObject"));
 	static auto* NiagaraPkg = FindObjectChecked<UPackage>(nullptr, TEXT("/Script/Niagara"));
 	FNiagaraTypeDefinition::ParameterMapStruct = FindObjectChecked<UScriptStruct>(NiagaraPkg, TEXT("NiagaraParameterMap"));
+	FNiagaraTypeDefinition::IDStruct = FindObjectChecked<UScriptStruct>(NiagaraPkg, TEXT("NiagaraID"));
 	FNiagaraTypeDefinition::NumericStruct = FindObjectChecked<UScriptStruct>(NiagaraPkg, TEXT("NiagaraNumeric"));
 	FNiagaraTypeDefinition::FloatStruct = FindObjectChecked<UScriptStruct>(NiagaraPkg, TEXT("NiagaraFloat"));
 	FNiagaraTypeDefinition::BoolStruct = FindObjectChecked<UScriptStruct>(NiagaraPkg, TEXT("NiagaraBool"));
@@ -215,8 +461,10 @@ void FNiagaraTypeDefinition::Init()
 	FNiagaraTypeDefinition::Vec3Struct = FindObjectChecked<UScriptStruct>(CoreUObjectPkg, TEXT("Vector"));
 	FNiagaraTypeDefinition::Vec4Struct = FindObjectChecked<UScriptStruct>(CoreUObjectPkg, TEXT("Vector4"));
 	FNiagaraTypeDefinition::ColorStruct = FindObjectChecked<UScriptStruct>(CoreUObjectPkg, TEXT("LinearColor"));
+	FNiagaraTypeDefinition::QuatStruct = FindObjectChecked<UScriptStruct>(CoreUObjectPkg, TEXT("Quat"));
 	
 	ParameterMapDef = FNiagaraTypeDefinition(ParameterMapStruct);
+	IDDef = FNiagaraTypeDefinition(IDStruct);
 	NumericDef = FNiagaraTypeDefinition(NumericStruct);
 	FloatDef = FNiagaraTypeDefinition(FloatStruct);
 	BoolDef = FNiagaraTypeDefinition(BoolStruct);
@@ -225,6 +473,7 @@ void FNiagaraTypeDefinition::Init()
 	Vec3Def = FNiagaraTypeDefinition(Vec3Struct);
 	Vec4Def = FNiagaraTypeDefinition(Vec4Struct);
 	ColorDef = FNiagaraTypeDefinition(ColorStruct);
+	QuatDef = FNiagaraTypeDefinition(QuatStruct);
 	Matrix4Def = FNiagaraTypeDefinition(Matrix4Struct);
 
 	CollisionEventDef = FNiagaraTypeDefinition(FNiagaraCollisionEventPayload::StaticStruct());
@@ -235,6 +484,7 @@ void FNiagaraTypeDefinition::Init()
 	NumericStructs.Add(Vec3Struct);
 	NumericStructs.Add(Vec4Struct);
 	NumericStructs.Add(ColorStruct);
+	NumericStructs.Add(QuatStruct);
 	//Make matrix a numeric type?
 
 	FloatStructs.Add(FloatStruct);
@@ -243,6 +493,7 @@ void FNiagaraTypeDefinition::Init()
 	FloatStructs.Add(Vec4Struct);
 	//FloatStructs.Add(Matrix4Struct)??
 	FloatStructs.Add(ColorStruct);
+	FloatStructs.Add(QuatStruct);
 
 	IntStructs.Add(IntStruct);
 
@@ -254,12 +505,14 @@ void FNiagaraTypeDefinition::Init()
 	OrderedNumericTypes.Add(Vec3Struct);
 	OrderedNumericTypes.Add(Vec4Struct);
 	OrderedNumericTypes.Add(ColorStruct);
+	OrderedNumericTypes.Add(QuatStruct);
 
 	ScalarStructs.Add(BoolStruct);
 	ScalarStructs.Add(IntStruct);
 	ScalarStructs.Add(FloatStruct);
 
 	ExecutionStateEnum = FindObject<UEnum>(ANY_PACKAGE, TEXT("ENiagaraExecutionState"), true);
+	ExecutionStateSourceEnum = FindObject<UEnum>(ANY_PACKAGE, TEXT("ENiagaraExecutionStateSource"), true);
 	
 	RecreateUserDefinedTypeRegistry();
 }
@@ -283,6 +536,7 @@ void FNiagaraTypeDefinition::RecreateUserDefinedTypeRegistry()
 	FNiagaraTypeRegistry::Register(CollisionEventDef, false, true, false);
 
 	FNiagaraTypeRegistry::Register(ParameterMapDef, true, false, false);
+	FNiagaraTypeRegistry::Register(IDDef, true, true, false);
 	FNiagaraTypeRegistry::Register(NumericDef, true, false, false);
 	FNiagaraTypeRegistry::Register(FloatDef, true, true, false);
 	FNiagaraTypeRegistry::Register(IntDef, true, true, false);
@@ -291,9 +545,11 @@ void FNiagaraTypeDefinition::RecreateUserDefinedTypeRegistry()
 	FNiagaraTypeRegistry::Register(Vec3Def, true, true, false);
 	FNiagaraTypeRegistry::Register(Vec4Def, true, true, false);
 	FNiagaraTypeRegistry::Register(ColorDef, true, true, false);
+	FNiagaraTypeRegistry::Register(QuatDef, true, true, false);
 	FNiagaraTypeRegistry::Register(Matrix4Def, true, false, false);
 
 	FNiagaraTypeRegistry::Register(FNiagaraTypeDefinition(ExecutionStateEnum), true, true, false);
+	FNiagaraTypeRegistry::Register(FNiagaraTypeDefinition(ExecutionStateSourceEnum), true, true, false);
 
 	UScriptStruct* TestStruct = FindObjectChecked<UScriptStruct>(NiagaraPkg, TEXT("NiagaraTestStruct"));
 	FNiagaraTypeDefinition TestDefinition(TestStruct);
@@ -480,13 +736,25 @@ FNiagaraTypeDefinition FNiagaraTypeDefinition::GetNumericOutputType(const TArray
 
 //////////////////////////////////////////////////////////////////////////
 
-void FNiagaraScriptDataInterfaceInfo::CopyTo(FNiagaraScriptDataInterfaceInfo* Destination, UObject* Outer) const
+template<>
+void FNiagaraVariable::SetValue(const bool& Data)
 {
-	Destination->Name = Name;
-	Destination->DataInterface = Cast<UNiagaraDataInterface>(StaticDuplicateObject(DataInterface, Outer, NAME_None, ~RF_Transient));
-	Destination->UserPtrIdx = UserPtrIdx;
+	check(TypeDef == FNiagaraTypeDefinition::GetBoolDef());
+	AllocateData();
+	FNiagaraBool* BoolStruct = (FNiagaraBool*)GetData();
+	BoolStruct->SetValue(Data);
 }
 
+template<>
+bool FNiagaraVariable::GetValue() const
+{
+	check(TypeDef == FNiagaraTypeDefinition::GetBoolDef());
+	check(IsDataAllocated());
+	FNiagaraBool* BoolStruct = (FNiagaraBool*)GetData();
+	return BoolStruct->GetValue();
+}
+
+//////////////////////////////////////////////////////////////////////////
 
 
 FDelegateHandle INiagaraModule::SetOnProcessShaderCompilationQueue(FOnProcessQueue InOnProcessQueue)
diff --git a/Engine/Plugins/FX/Niagara/Source/Niagara/Private/NiagaraParameterCollection.cpp b/Engine/Plugins/FX/Niagara/Source/Niagara/Private/NiagaraParameterCollection.cpp
--- a/Engine/Plugins/FX/Niagara/Source/Niagara/Private/NiagaraParameterCollection.cpp
+++ b/Engine/Plugins/FX/Niagara/Source/Niagara/Private/NiagaraParameterCollection.cpp
@@ -48,13 +48,13 @@ bool UNiagaraParameterCollectionInstance::AddParameter(const FNiagaraVariable& P
 bool UNiagaraParameterCollectionInstance::RemoveParameter(const FNiagaraVariable& Parameter)
 {
 	Modify();
-	return ParameterStorage.RemoveParameter(Parameter);
+	return ParameterStorage.RemoveParameter(Parameter); 
 }
 
 void UNiagaraParameterCollectionInstance::RenameParameter(const FNiagaraVariable& Parameter, FName NewName)
 {
 	Modify();
-	ParameterStorage.RenameParameter(Parameter, NewName);
+	ParameterStorage.RenameParameter(Parameter, NewName); 
 }
 
 void UNiagaraParameterCollectionInstance::Empty()
@@ -84,7 +84,7 @@ void UNiagaraParameterCollectionInstance::SyncWithCollection()
 		int32 Offset = OldStore.IndexOf(Param);
 		if (Offset != INDEX_NONE && OverridesParameter(Param))
 		{
-			//If this paramter is in the old store and we're overriding it, use the existing value in the store.
+			//If this parameter is in the old store and we're overriding it, use the existing value in the store.
 			ParameterStorage.AddParameter(Param, false);
 			if (Param.IsDataInterface())
 			{
@@ -165,6 +165,12 @@ FVector4 UNiagaraParameterCollectionInstance::GetVector4Parameter(const FString&
 	return ParameterStorage.GetParameterValue<FVector4>(FNiagaraVariable(FNiagaraTypeDefinition::GetVec4Def(), *Collection->ParameterNameFromFriendlyName(InVariableName)));
 }
 
+
+FQuat UNiagaraParameterCollectionInstance::GetQuatParameter(const FString& InVariableName)
+{
+	return ParameterStorage.GetParameterValue<FQuat>(FNiagaraVariable(FNiagaraTypeDefinition::GetQuatDef(), *Collection->ParameterNameFromFriendlyName(InVariableName)));
+}
+
 FLinearColor UNiagaraParameterCollectionInstance::GetColorParameter(const FString& InVariableName)
 {
 	return ParameterStorage.GetParameterValue<FLinearColor>(FNiagaraVariable(FNiagaraTypeDefinition::GetColorDef(), *Collection->ParameterNameFromFriendlyName(InVariableName)));
@@ -205,6 +211,11 @@ void UNiagaraParameterCollectionInstance::SetColorParameter(const FString& InVar
 	ParameterStorage.SetParameterValue(InValue, FNiagaraVariable(FNiagaraTypeDefinition::GetColorDef(), *Collection->ParameterNameFromFriendlyName(InVariableName)));
 }
 
+void UNiagaraParameterCollectionInstance::SetQuatParameter(const FString& InVariableName, const FQuat& InValue)
+{
+	ParameterStorage.SetParameterValue(InValue, FNiagaraVariable(FNiagaraTypeDefinition::GetQuatDef(), *Collection->ParameterNameFromFriendlyName(InVariableName)));
+}
+
 //////////////////////////////////////////////////////////////////////////
 
 UNiagaraParameterCollection::UNiagaraParameterCollection(const FObjectInitializer& ObjectInitializer)
@@ -238,6 +249,8 @@ int32 UNiagaraParameterCollection::AddParameter(FName Name, FNiagaraTypeDefiniti
 
 	int32 Idx = Parameters.AddDefaulted(1);
 
+	// We don't need to adjust the change id here as no one will be referencing the new parameter yet.
+
 	Parameters[Idx].SetName(Name);
 	Parameters[Idx].SetType(Type);
 
@@ -250,6 +263,7 @@ int32 UNiagaraParameterCollection::AddParameter(FName Name, FNiagaraTypeDefiniti
 void UNiagaraParameterCollection::RemoveParameter(const FNiagaraVariable& Parameter)
 {
 	Modify();
+	CompileId = FGuid::NewGuid();  // Any scripts depending on this parameter name will likely need to be changed.
 	DefaultInstance->RemoveParameter(Parameter);
 	Parameters.Remove(Parameter);
 }
@@ -257,7 +271,7 @@ void UNiagaraParameterCollection::RemoveParameter(const FNiagaraVariable& Parame
 void UNiagaraParameterCollection::RenameParameter(FNiagaraVariable& Parameter, FName NewName)
 {
 	Modify();
-
+	CompileId = FGuid::NewGuid(); // Any scripts depending on this parameter name will likely need to be changed.
 	int32 ParamIdx = Parameters.IndexOfByKey(Parameter);
 	check(ParamIdx != INDEX_NONE);
 
@@ -269,6 +283,12 @@ FString UNiagaraParameterCollection::GetFullNamespace()const
 {
 	return TEXT("NPC.") + Namespace.ToString() + TEXT(".");
 }
+
+void UNiagaraParameterCollection::RefreshCompileId()
+{
+	CompileId = FGuid::NewGuid();
+}
+
 FNiagaraVariable UNiagaraParameterCollection::CollectionParameterFromFriendlyParameter(const FNiagaraVariable& FriendlyParameter)const
 {
 	return FNiagaraVariable(FriendlyParameter.GetType(), *ParameterNameFromFriendlyName(FriendlyParameter.GetName().ToString()));
@@ -287,7 +307,7 @@ FString UNiagaraParameterCollection::FriendlyNameFromParameterName(FString Param
 
 FString UNiagaraParameterCollection::ParameterNameFromFriendlyName(const FString& FriendlyName)const
 {
-	return FString::Printf(TEXT("%s_%s"), *GetFullNamespace(), *FriendlyName);
+	return FString::Printf(TEXT("%s%s"), *GetFullNamespace(), *FriendlyName);
 }
 
 void UNiagaraParameterCollection::MakeNamespaceNameUnique()
@@ -324,4 +344,14 @@ void UNiagaraParameterCollection::MakeNamespaceNameUnique()
 		UE_LOG(LogNiagara, Warning, TEXT("Parameter collection namespace conflict found. \"%s\" is already in use!"), *Namespace.ToString());
 		Namespace = UniqueName;
 	}
-}
\ No newline at end of file
+}
+
+
+void UNiagaraParameterCollection::PostLoad()
+{
+	Super::PostLoad();
+	if (CompileId.IsValid() == false)
+	{
+		CompileId = FGuid::NewGuid();
+	}
+}
diff --git a/Engine/Plugins/FX/Niagara/Source/Niagara/Private/NiagaraParameterStore.cpp b/Engine/Plugins/FX/Niagara/Source/Niagara/Private/NiagaraParameterStore.cpp
--- a/Engine/Plugins/FX/Niagara/Source/Niagara/Private/NiagaraParameterStore.cpp
+++ b/Engine/Plugins/FX/Niagara/Source/Niagara/Private/NiagaraParameterStore.cpp
@@ -6,7 +6,13 @@
 #include "NiagaraComponent.h"
 #include "NiagaraSystemInstance.h"
 #include "NiagaraParameterCollection.h"
+#include "NiagaraStats.h"
 
+DECLARE_CYCLE_STAT(TEXT("Parameter store bind"), STAT_NiagaraParameterStoreBind, STATGROUP_Niagara);
+DECLARE_CYCLE_STAT(TEXT("Parameter store rebind"), STAT_NiagaraParameterStoreRebind, STATGROUP_Niagara);
+DECLARE_CYCLE_STAT(TEXT("Parameter store tick"), STAT_NiagaraParameterStoreTick, STATGROUP_Niagara);
+DECLARE_CYCLE_STAT(TEXT("Parameter store find var"), STAT_NiagaraParameterStoreFindVar, STATGROUP_Niagara);
+DECLARE_MEMORY_STAT(TEXT("Niagara parameter store memory"), STAT_NiagaraParamStoreMemory, STATGROUP_Niagara);
 
 #if WITH_EDITORONLY_DATA
 static int32 GbDumpParticleParameterStores = 0;
@@ -24,6 +30,7 @@ FNiagaraParameterStore::FNiagaraParameterStore()
 	: Owner(nullptr)
 	, bParametersDirty(true)
 	, bInterfacesDirty(true)
+	, LayoutVersion(0)
 {
 }
 
@@ -31,11 +38,12 @@ FNiagaraParameterStore::FNiagaraParameterStore(UObject* InOwner)
 	: Owner(InOwner)
 	, bParametersDirty(true)
 	, bInterfacesDirty(true)
+	, LayoutVersion(0)
 {
 #if WITH_EDITORONLY_DATA
 	if (InOwner != nullptr)
 	{
-		Name = *InOwner->GetFullName();
+		DebugName = *InOwner->GetFullName();
 	}
 #endif
 }
@@ -49,8 +57,11 @@ FNiagaraParameterStore& FNiagaraParameterStore::operator=(const FNiagaraParamete
 {
 	Owner = Other.Owner;
 	ParameterOffsets = Other.ParameterOffsets;
+	DEC_MEMORY_STAT_BY(STAT_NiagaraParamStoreMemory, ParameterData.Num());
 	ParameterData = Other.ParameterData;
+	INC_MEMORY_STAT_BY(STAT_NiagaraParamStoreMemory, ParameterData.Num());
 	DataInterfaces = Other.DataInterfaces;
+	++LayoutVersion;
 	//Don't copy bindings. We just want the data.
 	return *this;
 }
@@ -59,6 +70,7 @@ FNiagaraParameterStore::~FNiagaraParameterStore()
 {
 	//Ensure that any stores bound to drive this one are unbound.
 	UnbindFromSourceStores();
+	DEC_MEMORY_STAT_BY(STAT_NiagaraParamStoreMemory, ParameterData.Num());
 
 	//Also unbind from any stores we're feeding.
 	for (TPair<FNiagaraParameterStore*, FNiagaraParameterStoreBinding>& Binding : Bindings)
@@ -70,7 +82,7 @@ FNiagaraParameterStore::~FNiagaraParameterStore()
 void FNiagaraParameterStore::Bind(FNiagaraParameterStore* DestStore)
 {
 	check(DestStore);
-
+	SCOPE_CYCLE_COUNTER(STAT_NiagaraParameterStoreBind);
 	FNiagaraParameterStoreBinding& Binding = Bindings.FindOrAdd(DestStore);
 	Binding.Initialize(DestStore, this);
 }
@@ -87,6 +99,7 @@ void FNiagaraParameterStore::Unbind(FNiagaraParameterStore* DestStore)
 
 void FNiagaraParameterStore::Rebind()
 {
+	SCOPE_CYCLE_COUNTER(STAT_NiagaraParameterStoreRebind);
 	for (TPair<FNiagaraParameterStore*, FNiagaraParameterStoreBinding>& Binding : Bindings)
 	{
 		Binding.Value.Initialize(Binding.Key, this);
@@ -113,7 +126,7 @@ bool FNiagaraParameterStore::VerifyBinding(const FNiagaraParameterStore* DestSto
 	}
 	else
 	{
-		UE_LOG(LogNiagara, Warning, TEXT("Invalid ParameterStore Binding: % was not bound to %s."), *Name.ToString(), *DestStore->Name.ToString());
+		UE_LOG(LogNiagara, Warning, TEXT("Invalid ParameterStore Binding: % was not bound to %s."), *DebugName, *DestStore->DebugName);
 	}
 
 	return false;
@@ -124,16 +137,20 @@ bool FNiagaraParameterStore::VerifyBinding(const FNiagaraParameterStore* DestSto
 
 void FNiagaraParameterStore::Tick()
 {
-	for (TPair<FNiagaraParameterStore*, FNiagaraParameterStoreBinding>& Binding : Bindings)
+	SCOPE_CYCLE_COUNTER(STAT_NiagaraParameterStoreTick);
+	if (bParametersDirty || bInterfacesDirty)
 	{
-		Binding.Value.Tick(Binding.Key, this);
-	}
+		for (TPair<FNiagaraParameterStore*, FNiagaraParameterStoreBinding>& Binding : Bindings)
+		{
+			Binding.Value.Tick(Binding.Key, this);
+		}
 
-	Dump();
+		Dump();
 
-	//We have to have ticked all our source stores before now.
-	bParametersDirty = false;
-	bInterfacesDirty = false;
+		//We have to have ticked all our source stores before now.
+		bParametersDirty = false;
+		bInterfacesDirty = false;
+	}
 }
 
 void FNiagaraParameterStore::UnbindFromSourceStores()
@@ -147,7 +164,7 @@ void FNiagaraParameterStore::UnbindFromSourceStores()
 	ensureMsgf(SourceStores.Num() == 0, TEXT("Parameter store source array was not empty after unbinding all sources. Something seriously wrong."));
 }
 
-void FNiagaraParameterStore::DumpParameters()
+void FNiagaraParameterStore::DumpParameters(bool bDumpBindings)
 {
 	TArray<FNiagaraVariable> Vars;
 	GetParameters(Vars);
@@ -156,21 +173,37 @@ void FNiagaraParameterStore::DumpParameters()
 		Var.SetData(GetParameterData_Internal(IndexOf(Var))); // Need to keep the data in sync
 		UE_LOG(LogNiagara, Log, TEXT("Param: %s Offset: %d Type : %s"), *Var.ToString(), IndexOf(Var), *Var.GetType().GetName());
 	}
-}
 
-void FNiagaraParameterStore::Dump()
-{
-#if WITH_EDITORONLY_DATA
-	if (GbDumpParticleParameterStores && GetParametersDirty())
+	if (bDumpBindings)
 	{
-		UE_LOG(LogNiagara, Log, TEXT("\nSource Store: %s\n========================\n"), *Name.ToString());
-		DumpParameters();
-
 		for (TPair<FNiagaraParameterStore*, FNiagaraParameterStoreBinding>& Binding : Bindings)
 		{
 			Binding.Value.Dump(Binding.Key, this);
 		}
+	}
+}
+
+FString FNiagaraParameterStore::ToString() const
+{
+	FString Value;
+	TArray<FNiagaraVariable> Vars;
+	GetParameters(Vars);
+	for (FNiagaraVariable& Var : Vars)
+	{
+		Var.SetData(GetParameterData_Internal(IndexOf(Var))); // Need to keep the data in sync
+		Value += FString::Printf(TEXT("Param: %s Offset: %d Type : %s\n"), *Var.ToString(), IndexOf(Var), *Var.GetType().GetName());
+	}
+	return Value;
+}
 
+void FNiagaraParameterStore::Dump()
+{
+#if WITH_EDITORONLY_DATA
+	if (GbDumpParticleParameterStores && GetParametersDirty())
+	{
+		UE_LOG(LogNiagara, Log, TEXT("\nSource Store: %s\n========================\n"), *DebugName);
+		DumpParameters(true);
+		
 		UE_LOG(LogNiagara, Log, TEXT("\n========================\n"));
 	}
 #endif
@@ -181,7 +214,7 @@ Adds the passed parameter to this store.
 Does nothing if this parameter is already present.
 Returns true if we added a new parameter.
 */
-bool FNiagaraParameterStore::AddParameter(const FNiagaraVariable& Param, bool bInitInterfaces)
+bool FNiagaraParameterStore::AddParameter(const FNiagaraVariable& Param, bool bInitInterfaces, bool bTriggerRebind)
 {
 	if (int32* Existing = ParameterOffsets.Find(Param))
 	{
@@ -192,7 +225,7 @@ bool FNiagaraParameterStore::AddParameter(const FNiagaraVariable& Param, bool bI
 	{
 		int32 Offset = DataInterfaces.AddZeroed();
 		ParameterOffsets.Add(Param) = Offset;
-		DataInterfaces[Offset] = bInitInterfaces ? NewObject<UNiagaraDataInterface>(Owner, const_cast<UClass*>(Param.GetType().GetClass()), NAME_None, RF_Transactional) : nullptr;
+		DataInterfaces[Offset] = bInitInterfaces ? NewObject<UNiagaraDataInterface>(Owner, const_cast<UClass*>(Param.GetType().GetClass()), NAME_None, RF_Transactional | RF_Public) : nullptr;
 	}
 	else
 	{
@@ -200,8 +233,11 @@ bool FNiagaraParameterStore::AddParameter(const FNiagaraVariable& Param, bool bI
 		int32 ParamAlignment = Param.GetAlignment();
 		//int32 Offset = AlignArbitrary(ParameterData.Num(), ParamAlignment);//TODO: We need to handle alignment better here. Need to both satisfy CPU and GPU alignment concerns. VM doesn't care but the VM complier needs to be aware. Probably best to have everything adhere to GPU alignment rules.
 		int32 Offset = ParameterData.Num();
-		ParameterOffsets.Add(Param) = Offset;
+		FNiagaraVariable ParamWithNoAllocatedData(Param.GetType(), Param.GetName());
+		ParameterOffsets.Add(ParamWithNoAllocatedData) = Offset; // We don't need the default value saved in here.
 		ParameterData.AddUninitialized(ParamSize);
+				
+		INC_MEMORY_STAT_BY(STAT_NiagaraParamStoreMemory, ParamSize);
 
 		//Temporary to init param data from FNiagaraVariable storage. This will be removed when we change the UNiagaraScript to use a parameter store too.
 		if (Param.IsDataAllocated())
@@ -210,7 +246,15 @@ bool FNiagaraParameterStore::AddParameter(const FNiagaraVariable& Param, bool bI
 		}
 	}
 
-	OnLayoutChange();
+	if (bTriggerRebind)
+	{
+		OnLayoutChange();
+	}
+	else
+	{
+		++LayoutVersion;
+	}
+	
 	return true;
 }
 
@@ -246,6 +290,10 @@ bool FNiagaraParameterStore::RemoveParameter(const FNiagaraVariable& ToRemove)
 					FMemory::Memcpy(NewData.GetData() + Offset, ParameterData.GetData() + ExistingOffset, ParamSize);
 				}
 			}
+			else
+			{
+				DEC_MEMORY_STAT_BY(STAT_NiagaraParamStoreMemory, ExistingVar.GetSizeInBytes());
+			}
 		}
 
 		ParameterOffsets = NewOffsets;
@@ -266,7 +314,10 @@ void FNiagaraParameterStore::RenameParameter(const FNiagaraVariable& Param, FNam
 	{
 		FNiagaraVariable NewParam = Param;
 		NewParam.SetName(NewName);
-		AddParameter(NewParam);
+
+		bool bInitInterfaces = false;
+		bool bTriggerRebind = false;
+		AddParameter(NewParam, bInitInterfaces, bTriggerRebind);
 
 		int32 NewIdx = IndexOf(NewParam);
 		if (Param.IsDataInterface())
@@ -283,7 +334,7 @@ void FNiagaraParameterStore::RenameParameter(const FNiagaraVariable& Param, FNam
 	}
 }
 
-void FNiagaraParameterStore::CopyParametersTo(FNiagaraParameterStore& DestStore, bool bOnlyAdd)
+void FNiagaraParameterStore::CopyParametersTo(FNiagaraParameterStore& DestStore, bool bOnlyAdd, EDataInterfaceCopyMethod DataInterfaceCopyMethod)
 {
 	TMap<FNiagaraVariable, int32>::TConstIterator It = ParameterOffsets.CreateConstIterator();
 	while (It)
@@ -294,7 +345,9 @@ void FNiagaraParameterStore::CopyParametersTo(FNiagaraParameterStore& DestStore,
 		bool bWrite = false;
 		if (DestIndex == INDEX_NONE)
 		{
-			DestStore.AddParameter(Parameter);
+			bool bInitInterfaces = bOnlyAdd == false && Parameter.IsDataInterface() && DataInterfaceCopyMethod == EDataInterfaceCopyMethod::Value;
+			bool bTriggerRebind = false;
+			DestStore.AddParameter(Parameter, bInitInterfaces, bTriggerRebind);
 			bWrite = !bOnlyAdd;
 			DestIndex = DestStore.IndexOf(Parameter);
 		}
@@ -305,10 +358,67 @@ void FNiagaraParameterStore::CopyParametersTo(FNiagaraParameterStore& DestStore,
 
 		if (bWrite && DestIndex != INDEX_NONE && SrcIndex != INDEX_NONE)
 		{
-			DestStore.SetParameterData(GetParameterData_Internal(SrcIndex), DestIndex, Parameter.GetSizeInBytes());
+			if (Parameter.IsDataInterface())
+			{
+				if (DataInterfaceCopyMethod == EDataInterfaceCopyMethod::Reference)
+				{
+					DestStore.SetDataInterface(DataInterfaces[SrcIndex], DestIndex);
+				}
+				else if(DataInterfaceCopyMethod == EDataInterfaceCopyMethod::Value)
+				{
+					UNiagaraDataInterface* SourceInterface = DataInterfaces[SrcIndex];
+					SourceInterface->CopyTo(DestStore.GetDataInterface(DestIndex));
+				}
+				else
+				{
+					checkf(false, TEXT("A data interface copy method must be specified if the parameter store has data interfaces."));
+				}
+			}
+			else
+			{
+				if (ParameterData.Num() != 0)
+				{
+					DestStore.SetParameterData(GetParameterData_Internal(SrcIndex), DestIndex, Parameter.GetSizeInBytes());
+				}
+			}
 		}
 		++It;
 	}
+	DestStore.OnLayoutChange();
+}
+
+
+FORCEINLINE void FNiagaraParameterStore::SetParameterDataArray(const TArray<uint8>& InParameterDataArray)
+{
+	DEC_MEMORY_STAT_BY(STAT_NiagaraParamStoreMemory, ParameterData.Num());
+	ParameterData = InParameterDataArray;
+	INC_MEMORY_STAT_BY(STAT_NiagaraParamStoreMemory, ParameterData.Num());
+
+	OnParameterChange();
+}
+
+void FNiagaraParameterStore::InitFromSource(const FNiagaraParameterStore* SrcStore, bool bNotifyAsDirty)
+{
+	Empty(false);
+	if (SrcStore == nullptr)
+	{
+		return;
+	}
+
+	ParameterOffsets = SrcStore->ParameterOffsets;
+
+	DEC_MEMORY_STAT_BY(STAT_NiagaraParamStoreMemory, ParameterData.Num());
+	ParameterData = SrcStore->ParameterData;
+	INC_MEMORY_STAT_BY(STAT_NiagaraParamStoreMemory, ParameterData.Num());
+
+	DataInterfaces = SrcStore->DataInterfaces;
+
+	if (bNotifyAsDirty)
+	{
+		MarkParametersDirty();
+		MarkInterfacesDirty();
+		OnLayoutChange();
+	}
 }
 
 void FNiagaraParameterStore::RemoveParameters(FNiagaraParameterStore& DestStore)
@@ -318,6 +428,7 @@ void FNiagaraParameterStore::RemoveParameters(FNiagaraParameterStore& DestStore)
 	{
 		FNiagaraVariable Parameter = It->Key;
 		DestStore.RemoveParameter(Parameter);
+		DEC_MEMORY_STAT_BY(STAT_NiagaraParamStoreMemory, Parameter.GetSizeInBytes());
 		++It;
 	}
 }
@@ -325,28 +436,44 @@ void FNiagaraParameterStore::RemoveParameters(FNiagaraParameterStore& DestStore)
 void FNiagaraParameterStore::Empty(bool bClearBindings)
 {
 	ParameterOffsets.Empty();
+	DEC_MEMORY_STAT_BY(STAT_NiagaraParamStoreMemory, ParameterData.Num());
 	ParameterData.Empty();
+	DEC_MEMORY_STAT_BY(STAT_NiagaraParamStoreMemory, ParameterData.Num());
+
 	DataInterfaces.Empty();
 	if (bClearBindings)
 	{
 		Bindings.Empty();
 	}
-	OnParameterChange();
-	OnInterfaceChange();
-	OnLayoutChange();
+}
+
+void FNiagaraParameterStore::Reset(bool bClearBindings)
+{
+	ParameterOffsets.Reset();
+	DEC_MEMORY_STAT_BY(STAT_NiagaraParamStoreMemory, ParameterData.Num());
+	ParameterData.Reset();
+	DEC_MEMORY_STAT_BY(STAT_NiagaraParamStoreMemory, ParameterData.Num());
+
+	DataInterfaces.Reset();
+	if (bClearBindings)
+	{
+		Bindings.Reset();
+	}
 }
 
 void FNiagaraParameterStore::OnLayoutChange()
 {
 	Rebind();
+	++LayoutVersion;
 
-#ifdef WITH_EDITOR
+#if WITH_EDITOR
 	OnChangedDelegate.Broadcast();
 #endif
 }
 
 const FNiagaraVariable* FNiagaraParameterStore::FindVariable(UNiagaraDataInterface* Interface)const
 {
+	SCOPE_CYCLE_COUNTER(STAT_NiagaraParameterStoreFindVar);
 	int32 Idx = DataInterfaces.IndexOfByKey(Interface);
 	if (Idx != INDEX_NONE)
 	{
@@ -364,7 +491,7 @@ const FNiagaraVariable* FNiagaraParameterStore::FindVariable(UNiagaraDataInterfa
 	return nullptr;
 }
 
-#ifdef WITH_EDITOR
+#if WITH_EDITOR
 FDelegateHandle FNiagaraParameterStore::AddOnChangedHandler(FOnChanged::FDelegate InOnChanged)
 {
 	return OnChangedDelegate.Add(InOnChanged);
diff --git a/Engine/Plugins/FX/Niagara/Source/Niagara/Private/NiagaraParameters.cpp b/Engine/Plugins/FX/Niagara/Source/Niagara/Private/NiagaraParameters.cpp
--- a/Engine/Plugins/FX/Niagara/Source/Niagara/Private/NiagaraParameters.cpp
+++ b/Engine/Plugins/FX/Niagara/Source/Niagara/Private/NiagaraParameters.cpp
@@ -11,6 +11,15 @@ void FNiagaraParameters::Empty()
 	Parameters.Empty();
 }
 
+void FNiagaraParameters::DumpParameters()
+{
+	for (FNiagaraVariable& Var : Parameters)
+	{
+		UE_LOG(LogNiagara, Log, TEXT("Param: %s Type : %s"), *Var.ToString(), *Var.GetType().GetName());
+	}
+}
+
+
 void FNiagaraParameters::AppendToConstantsTable(uint8* ConstantsTable, const FNiagaraParameters& Externals)const
 {
 	uint8* Curr = ConstantsTable;
