Copying //UE4/Dev-Niagara to //UE4/Dev-Main (Source: //UE4/Dev-Niagara @ 4074996)

diff --git a/Engine/Plugins/FX/Niagara/Source/Niagara/Private/NiagaraRenderer.cpp b/Engine/Plugins/FX/Niagara/Source/Niagara/Private/NiagaraRenderer.cpp
--- a/Engine/Plugins/FX/Niagara/Source/Niagara/Private/NiagaraRenderer.cpp
+++ b/Engine/Plugins/FX/Niagara/Source/Niagara/Private/NiagaraRenderer.cpp
@@ -6,8 +6,10 @@
 #include "NiagaraDataSet.h"
 #include "NiagaraStats.h"
 #include "NiagaraVertexFactory.h"
+#include "Engine/Engine.h"
 
 DECLARE_CYCLE_STAT(TEXT("Generate Particle Lights"), STAT_NiagaraGenLights, STATGROUP_Niagara);
+DECLARE_CYCLE_STAT(TEXT("Sort Particles"), STAT_NiagaraSortParticles, STATGROUP_Niagara);
 
 
 /** Enable/disable parallelized System renderers */
@@ -21,6 +23,46 @@ static FAutoConsoleVariableRef CVarParallelEmitterRenderers(
 
 
 
+void FNiagaraDummyRWBufferFloat::InitRHI() 
+{
+	UE_LOG(LogNiagara, Log, TEXT("FNiagaraDummyRWBufferFloat InitRHI %s"), *DebugId);
+	Buffer.Initialize(sizeof(float), 1, EPixelFormat::PF_R32_FLOAT, BUF_Static);
+}
+
+void FNiagaraDummyRWBufferFloat::ReleaseRHI() 
+{
+	UE_LOG(LogNiagara, Log, TEXT("FNiagaraDummyRWBufferFloat ReleaseRHI %s"), *DebugId);
+	Buffer.Release();
+}
+
+void FNiagaraDummyRWBufferInt::InitRHI() 
+{
+	UE_LOG(LogNiagara, Log, TEXT("FNiagaraDummyRWBufferInt InitRHI %s"), *DebugId);
+	Buffer.Initialize(sizeof(int32), 1, EPixelFormat::PF_R32_SINT, BUF_Static);
+}
+void FNiagaraDummyRWBufferInt::ReleaseRHI() 
+{
+	UE_LOG(LogNiagara, Log, TEXT("FNiagaraDummyRWBufferInt ReleaseRHI %s"), *DebugId);
+	Buffer.Release();
+}
+
+
+FRWBuffer& NiagaraRenderer::GetDummyFloatBuffer()
+{
+	check(IsInRenderingThread());
+	static TGlobalResource<FNiagaraDummyRWBufferFloat> DummyFloatBuffer(TEXT("NiagaraRenderer::DummyFloat"));
+	return DummyFloatBuffer.Buffer;
+}
+
+FRWBuffer& NiagaraRenderer::GetDummyIntBuffer()
+{
+	check(IsInRenderingThread());
+	static TGlobalResource<FNiagaraDummyRWBufferInt> DummyIntBuffer(TEXT("NiagaraRenderer::DummyInt"));
+	return DummyIntBuffer.Buffer;
+}
+
+
+
 NiagaraRenderer::NiagaraRenderer()
 	: CPUTimeMS(0.0f)
 	, bLocalSpace(false)
@@ -29,7 +71,6 @@ NiagaraRenderer::NiagaraRenderer()
 	, BaseExtents(1.0f, 1.0f, 1.0f)
 {
 	Material = UMaterial::GetDefaultMaterial(MD_Surface);
-	FNiagaraVertexFactoryBase::Init();
 }
 
 
@@ -49,8 +90,113 @@ void NiagaraRenderer::Release()
 	);
 }
 
-	
+void NiagaraRenderer::SortIndices(ENiagaraSortMode SortMode, int32 SortAttributeOffset, const FNiagaraDataBuffer& Buffer, const FMatrix& LocalToWorld, const FSceneView* View, FNiagaraGlobalReadBuffer::FAllocation& OutIndices)const
+{
+	SCOPE_CYCLE_COUNTER(STAT_NiagaraSortParticles);
+
+	uint32 NumInstances = Buffer.GetNumInstances();
+	check(OutIndices.ReadBuffer->NumBytes >= OutIndices.FirstIndex + NumInstances * sizeof(int32));
+	check(SortMode != ENiagaraSortMode::None);
+	check(SortAttributeOffset != INDEX_NONE);
+
+	int32* RESTRICT IndexBuffer = (int32*)(OutIndices.Buffer);
+
+	struct FParticleOrder
+	{
+		int32 Index;
+		float Order;
+	};
+	FMemMark Mark(FMemStack::Get());
+	FParticleOrder* RESTRICT ParticleOrder = (FParticleOrder*)FMemStack::Get().Alloc(sizeof(FParticleOrder) * NumInstances, alignof(FParticleOrder));
+
+	auto SortAscending = [](const FParticleOrder& A, const FParticleOrder& B) { return A.Order < B.Order; };
+	auto SortDescending = [](const FParticleOrder& A, const FParticleOrder& B) { return A.Order > B.Order; };
+
+	if (SortMode == ENiagaraSortMode::ViewDepth || SortMode == ENiagaraSortMode::ViewDistance)
+	{
+		float* RESTRICT PositionX = (float*)Buffer.GetComponentPtrFloat(SortAttributeOffset);
+		float* RESTRICT PositionY = (float*)Buffer.GetComponentPtrFloat(SortAttributeOffset + 1);
+		float* RESTRICT PositionZ = (float*)Buffer.GetComponentPtrFloat(SortAttributeOffset + 2);
+		auto GetPos = [&PositionX, &PositionY, &PositionZ](int32 Idx)
+		{
+			return FVector(PositionX[Idx], PositionY[Idx], PositionZ[Idx]);
+		};
+
+		//TODO Parallelize in batches? Move to GPU for large emitters?
+		if (SortMode == ENiagaraSortMode::ViewDepth)
+		{
+			FMatrix ViewProjMatrix = View->ViewMatrices.GetViewProjectionMatrix();
+			if (bLocalSpace)
+			{
+				for (uint32 i = 0; i < NumInstances; ++i)
+				{
+					ParticleOrder[i].Index = i;
+					ParticleOrder[i].Order = ViewProjMatrix.TransformPosition(LocalToWorld.TransformPosition(GetPos(i))).W;
+				}
+			}
+			else
+			{
+				for (uint32 i = 0; i < NumInstances; ++i)
+				{
+					ParticleOrder[i].Index = i;
+					ParticleOrder[i].Order = ViewProjMatrix.TransformPosition(GetPos(i)).W;
+				}
+			}
+
+			Sort(ParticleOrder, NumInstances, SortDescending);
+		}
+		else 
+		{
+			// check(SortMode == ENiagaraSortMode::ViewDistance); Should always be true because we are within this block. If code is moved or copied, please change.
+			FVector ViewOrigin = View->ViewMatrices.GetViewOrigin();
+			if (bLocalSpace)
+			{
+				for (uint32 i = 0; i < NumInstances; ++i)
+				{
+					ParticleOrder[i].Index = i;
+					ParticleOrder[i].Order = (ViewOrigin - LocalToWorld.TransformPosition(GetPos(i))).SizeSquared();
+				}
+			}
+			else
+			{
+				for (uint32 i = 0; i < NumInstances; ++i)
+				{
+					ParticleOrder[i].Index = i;
+					ParticleOrder[i].Order = (ViewOrigin - GetPos(i)).SizeSquared();
+				}
+			}
+
+			Sort(ParticleOrder, NumInstances, SortDescending);
+		}
+	}
+	else
+	{
+		float* RESTRICT CustomSorting = (float*)Buffer.GetComponentPtrFloat(SortAttributeOffset);
+		for (uint32 i = 0; i < NumInstances; ++i)
+		{
+			ParticleOrder[i].Index = i;
+			ParticleOrder[i].Order = CustomSorting[i];
+		}
+		if (SortMode == ENiagaraSortMode::CustomAscending)
+		{
+			Sort(ParticleOrder, NumInstances, SortAscending);
+		}
+		else if (SortMode == ENiagaraSortMode::CustomDecending)
+		{
+			Sort(ParticleOrder, NumInstances, SortDescending);
+		}
+	}
 
+	//Now transfer to the real index buffer.
+	for (uint32 i = 0; i < NumInstances; ++i)
+	{
+		IndexBuffer[i] = ParticleOrder[i].Index;
+	}
+}
+
+	
+//////////////////////////////////////////////////////////////////////////
+// Light renderer
 
 NiagaraRendererLights::NiagaraRendererLights(ERHIFeatureLevel::Type FeatureLevel, UNiagaraRendererProperties *InProps) :
 	NiagaraRenderer()
@@ -76,36 +222,48 @@ FNiagaraDynamicDataBase *NiagaraRendererLights::GenerateVertexData(const FNiagar
 
 	SimpleTimer VertexDataTimer;
 
-	//I'm not a great fan of pulling scalar components out to a structured vert buffer like this.
-	//TODO: Experiment with a new VF that reads the data directly from the scalar layout.
-	FNiagaraDataSetIterator<FVector> PosItr(Data, FNiagaraVariable(FNiagaraTypeDefinition::GetVec3Def(), TEXT("Position")));
-	FNiagaraDataSetIterator<FLinearColor> ColItr(Data, FNiagaraVariable(FNiagaraTypeDefinition::GetColorDef(), TEXT("Color")));
-	FNiagaraDataSetIterator<FVector2D> SizeItr(Data, FNiagaraVariable(FNiagaraTypeDefinition::GetVec2Def(), TEXT("Size")));
-
 	//Bail if we don't have the required attributes to render this emitter.
-	if (!PosItr.IsValid() || !ColItr.IsValid() || !SizeItr.IsValid() || !bEnabled)
+	if (!bEnabled || !Properties)
 	{
 		return nullptr;
 	}
 
+	//I'm not a great fan of pulling scalar components out to a structured vert buffer like this.
+	//TODO: Experiment with a new VF that reads the data directly from the scalar layout.
+	FNiagaraDataSetIterator<FVector> PosItr(Data, Properties->PositionBinding.DataSetVariable);
+	FNiagaraDataSetIterator<FLinearColor> ColItr(Data, Properties->ColorBinding.DataSetVariable);
+	FNiagaraDataSetIterator<float> RadiusItr(Data, Properties->RadiusBinding.DataSetVariable);
+
 	FNiagaraDynamicDataLights *DynamicData = new FNiagaraDynamicDataLights;
+	FVector DefaultColor = FVector(Properties->ColorBinding.DefaultValueIfNonExistent.GetValue<FLinearColor>());
+	FVector DefaultPos = FVector4(Proxy->GetLocalToWorld().GetOrigin());
+	float DefaultRadius = Properties->RadiusBinding.DefaultValueIfNonExistent.GetValue<float>();
 
 	DynamicData->LightArray.Empty();
 
 	for (uint32 ParticleIndex = 0; ParticleIndex < Data.GetNumInstances(); ParticleIndex++)
 	{
 		SimpleLightData LightData;
-		LightData.LightEntry.Radius = (*SizeItr).X;	//LightPayload->RadiusScale * (Size.X + Size.Y) / 2.0f;
-		LightData.LightEntry.Color = FVector((*ColItr));				//FVector(Particle.Color) * Particle.Color.A * LightPayload->ColorScale;
+		LightData.LightEntry.Radius = (RadiusItr.IsValid() ? (*RadiusItr) :  DefaultRadius) * Properties->RadiusScale;	//LightPayload->RadiusScale * (Size.X + Size.Y) / 2.0f;
+		LightData.LightEntry.Color = (ColItr.IsValid() ? FVector((*ColItr)) : DefaultColor) + Properties->ColorAdd;				//FVector(Particle.Color) * Particle.Color.A * LightPayload->ColorScale;
 		LightData.LightEntry.Exponent = 1.0;
 		LightData.LightEntry.bAffectTranslucency = true;
-		LightData.PerViewEntry.Position = (*PosItr);
+		LightData.PerViewEntry.Position = PosItr.IsValid() ? (*PosItr) : DefaultPos;
 
 		DynamicData->LightArray.Add(LightData);
 
 		PosItr.Advance();
 		ColItr.Advance();
-		SizeItr.Advance();
+		RadiusItr.Advance();
+	}
+
+	if (bLocalSpace)
+	{
+		FMatrix Mat = Proxy->GetLocalToWorld();
+		for (uint32 ParticleIndex = 0; ParticleIndex < Data.GetNumInstances(); ParticleIndex++)
+		{
+			DynamicData->LightArray[ParticleIndex].PerViewEntry.Position = Mat.TransformPosition(DynamicData->LightArray[ParticleIndex].PerViewEntry.Position);
+		}
 	}
 
 	CPUTimeMS = VertexDataTimer.GetElapsedMilliseconds();
@@ -160,11 +318,236 @@ const TArray<FNiagaraVariable>& NiagaraRendererLights::GetOptionalAttributes()
 
 
 
+//////////////////////////////////////////////////////////////////////////
+// FNiagaraGlobalReadBuffer
+//////////////////////////////////////////////////////////////////////////
+
+int32 GMaxNiagaraRenderingBytesAllocatedPerFrame = 32 * 1024 * 1024;
+
+FAutoConsoleVariableRef CVarMaxNiagaraRenderingBytesAllocatedPerFrame(
+	TEXT("fx.MaxRenderingBytesAllocatedPerFrame"),
+	GMaxNiagaraRenderingBytesAllocatedPerFrame,
+	TEXT("The maximum number of transient rendering buffer bytes to allocate before we start panic logging who is doing the allocations"));
+
+
+int32 GMinNiagaraRenderingBufferSize = 8 * 1024;
+FAutoConsoleVariableRef CVarMinNiagaraRenderingBufferSize(
+	TEXT("fx.MinNiagaraRenderingBufferSize"),
+	GMinNiagaraRenderingBufferSize,
+	TEXT("The minimum size (in instances) to allocate in blocks for niagara rendering buffers."));
+
 
+struct FDynamicReadBufferPool
+{
+	/** List of vertex buffers. */
+	TIndirectArray<FDynamicAllocReadBuffer> Buffers;
+	/** The current buffer from which allocations are being made. */
+	FDynamicAllocReadBuffer* CurrentBuffer;
+
+	/** Default constructor. */
+	FDynamicReadBufferPool()
+		: CurrentBuffer(NULL)
+	{
+	}
 
+	/** Destructor. */
+	~FDynamicReadBufferPool()
+	{
+		int32 NumVertexBuffers = Buffers.Num();
+		for (int32 BufferIndex = 0; BufferIndex < NumVertexBuffers; ++BufferIndex)
+		{
+			Buffers[BufferIndex].Release();
+		}
+	}
 
+	FCriticalSection CriticalSection;
+};
 
 
 
+FNiagaraGlobalReadBuffer::FNiagaraGlobalReadBuffer()
+	: TotalAllocatedSinceLastCommit(0)
+{
+	FloatBufferPool = new FDynamicReadBufferPool();
+	Int32BufferPool = new FDynamicReadBufferPool();
 
+	CommitCallbackHandle = GEngine->GetPreRenderDelegate().AddRaw(this, &FNiagaraGlobalReadBuffer::Commit);
+}
 
+FNiagaraGlobalReadBuffer::~FNiagaraGlobalReadBuffer()
+{
+	Cleanup();
+
+	//GEngine->GetPostInitViewsDelegate().Remove(CommitCallbackHandle);
+}
+
+void FNiagaraGlobalReadBuffer::Cleanup()
+{
+	if (FloatBufferPool)
+	{
+		UE_LOG(LogNiagara, Log, TEXT("FNiagaraGlobalReadBuffer::Cleanup()"));
+		delete FloatBufferPool;
+		FloatBufferPool = nullptr;
+	}
+
+	if (Int32BufferPool)
+	{
+		delete Int32BufferPool;
+		Int32BufferPool = nullptr;
+	}
+}
+void FNiagaraGlobalReadBuffer::InitRHI()
+{
+	UE_LOG(LogNiagara, Log, TEXT("FNiagaraGlobalReadBuffer::InitRHI"));
+}
+
+void FNiagaraGlobalReadBuffer::ReleaseRHI() 
+{
+	UE_LOG(LogNiagara, Log, TEXT("FNiagaraGlobalReadBuffer::ReleaseRHI"));
+	Cleanup();
+}
+
+FNiagaraGlobalReadBuffer::FAllocation FNiagaraGlobalReadBuffer::AllocateFloat(uint32 Num)
+{
+	FScopeLock ScopeLock(&FloatBufferPool->CriticalSection);
+
+	FAllocation Allocation;
+
+	TotalAllocatedSinceLastCommit += Num;
+	if (IsRenderAlarmLoggingEnabled())
+	{
+		UE_LOG(LogNiagara, Warning, TEXT("FNiagaraGlobalReadBuffer::AllocateFloat(%u), will have allocated %u total this frame"), Num, TotalAllocatedSinceLastCommit);
+	}
+	uint32 SizeInBytes = sizeof(float) * Num;
+	FDynamicAllocReadBuffer* Buffer = FloatBufferPool->CurrentBuffer;
+	if (Buffer == NULL || Buffer->AllocatedByteCount + SizeInBytes > Buffer->NumBytes)
+	{
+		// Find a buffer in the pool big enough to service the request.
+		Buffer = NULL;
+		for (int32 BufferIndex = 0, NumBuffers = FloatBufferPool->Buffers.Num(); BufferIndex < NumBuffers; ++BufferIndex)
+		{
+			FDynamicAllocReadBuffer& BufferToCheck = FloatBufferPool->Buffers[BufferIndex];
+			if (BufferToCheck.AllocatedByteCount + SizeInBytes <= BufferToCheck.NumBytes)
+			{
+				Buffer = &BufferToCheck;
+				break;
+			}
+		}
+
+		// Create a new vertex buffer if needed.
+		if (Buffer == NULL)
+		{
+			uint32 NewBufferSize = FMath::Max(Num, (uint32)GMinNiagaraRenderingBufferSize);
+			Buffer = new(FloatBufferPool->Buffers) FDynamicAllocReadBuffer();
+			Buffer->Initialize(sizeof(float), NewBufferSize, PF_R32_FLOAT, BUF_Dynamic);
+		}
+
+		// Lock the buffer if needed.
+		if (Buffer->MappedBuffer == NULL)
+		{
+			Buffer->Lock();
+		}
+
+		// Remember this buffer, we'll try to allocate out of it in the future.
+		FloatBufferPool->CurrentBuffer = Buffer;
+	}
+
+	check(Buffer != NULL);
+	checkf(Buffer->AllocatedByteCount + SizeInBytes <= Buffer->NumBytes, TEXT("Niagara global float buffer allocation failed: BufferSize=%d AllocatedByteCount=%d SizeInBytes=%d"), Buffer->NumBytes, Buffer->AllocatedByteCount, SizeInBytes);
+	Allocation.Buffer = Buffer->MappedBuffer + Buffer->AllocatedByteCount;
+	Allocation.ReadBuffer = Buffer;
+	Allocation.FirstIndex = Buffer->AllocatedByteCount;
+	Buffer->AllocatedByteCount += SizeInBytes;
+
+	return Allocation;
+}
+
+FNiagaraGlobalReadBuffer::FAllocation FNiagaraGlobalReadBuffer::AllocateInt32(uint32 Num)
+{
+	FScopeLock ScopeLock(&Int32BufferPool->CriticalSection);
+	FAllocation Allocation;
+
+	TotalAllocatedSinceLastCommit += Num;
+	if (IsRenderAlarmLoggingEnabled())
+	{
+		UE_LOG(LogNiagara, Warning, TEXT("FNiagaraGlobalReadBuffer::AllocateInt32(%u), will have allocated %u total this frame"), Num, TotalAllocatedSinceLastCommit);
+	}
+	uint32 SizeInBytes = sizeof(int32) * Num;
+	FDynamicAllocReadBuffer* Buffer = Int32BufferPool->CurrentBuffer;
+	if (Buffer == NULL || Buffer->AllocatedByteCount + SizeInBytes > Buffer->NumBytes)
+	{
+		// Find a buffer in the pool big enough to service the request.
+		Buffer = NULL;
+		for (int32 BufferIndex = 0, NumBuffers = Int32BufferPool->Buffers.Num(); BufferIndex < NumBuffers; ++BufferIndex)
+		{
+			FDynamicAllocReadBuffer& BufferToCheck = Int32BufferPool->Buffers[BufferIndex];
+			if (BufferToCheck.AllocatedByteCount + SizeInBytes <= BufferToCheck.NumBytes)
+			{
+				Buffer = &BufferToCheck;
+				break;
+			}
+		}
+
+		// Create a new vertex buffer if needed.
+		if (Buffer == NULL)
+		{
+			uint32 NewBufferSize = FMath::Max(Num, (uint32)GMinNiagaraRenderingBufferSize);
+			Buffer = new(Int32BufferPool->Buffers) FDynamicAllocReadBuffer();
+			Buffer->Initialize(sizeof(int32), NewBufferSize, PF_R32_SINT, BUF_Dynamic);
+		}
+
+		// Lock the buffer if needed.
+		if (Buffer->MappedBuffer == NULL)
+		{
+			Buffer->Lock();
+		}
+
+		// Remember this buffer, we'll try to allocate out of it in the future.
+		Int32BufferPool->CurrentBuffer = Buffer;
+	}
+
+	check(Buffer != NULL);
+	checkf(Buffer->AllocatedByteCount + SizeInBytes <= Buffer->NumBytes, TEXT("Niagara global int32 buffer allocation failed: BufferSize=%d AllocatedByteCount=%d SizeInBytes=%d"), Buffer->NumBytes, Buffer->AllocatedByteCount, SizeInBytes);
+	Allocation.Buffer = Buffer->MappedBuffer + Buffer->AllocatedByteCount;
+	Allocation.ReadBuffer = Buffer;
+	Allocation.FirstIndex = Buffer->AllocatedByteCount;
+	Buffer->AllocatedByteCount += SizeInBytes;
+
+	return Allocation;
+}
+
+bool FNiagaraGlobalReadBuffer::IsRenderAlarmLoggingEnabled() const
+{
+	return GMaxNiagaraRenderingBytesAllocatedPerFrame > 0 && TotalAllocatedSinceLastCommit >= (size_t)GMaxNiagaraRenderingBytesAllocatedPerFrame;
+}
+
+void FNiagaraGlobalReadBuffer::Commit()
+{
+	for (int32 BufferIndex = 0, NumBuffers = FloatBufferPool->Buffers.Num(); BufferIndex < NumBuffers; ++BufferIndex)
+	{
+		FDynamicAllocReadBuffer& Buffer = FloatBufferPool->Buffers[BufferIndex];
+		if (Buffer.MappedBuffer != NULL)
+		{
+			Buffer.Unlock();
+		}
+	}
+	FloatBufferPool->CurrentBuffer = NULL;
+	for (int32 BufferIndex = 0, NumBuffers = Int32BufferPool->Buffers.Num(); BufferIndex < NumBuffers; ++BufferIndex)
+	{
+		FDynamicAllocReadBuffer& Buffer = Int32BufferPool->Buffers[BufferIndex];
+		if (Buffer.MappedBuffer != NULL)
+		{
+			Buffer.Unlock();
+		}
+	}
+	Int32BufferPool->CurrentBuffer = NULL;
+	TotalAllocatedSinceLastCommit = 0;
+}
+
+
+FNiagaraGlobalReadBuffer& FNiagaraGlobalReadBuffer::Get()
+{
+	check(IsInRenderingThread());
+	static TGlobalResource<FNiagaraGlobalReadBuffer> GlobalDynamicReadBuffer;
+	return GlobalDynamicReadBuffer;
+}
diff --git a/Engine/Plugins/FX/Niagara/Source/Niagara/Private/NiagaraRendererMeshes.cpp b/Engine/Plugins/FX/Niagara/Source/Niagara/Private/NiagaraRendererMeshes.cpp
--- a/Engine/Plugins/FX/Niagara/Source/Niagara/Private/NiagaraRendererMeshes.cpp
+++ b/Engine/Plugins/FX/Niagara/Source/Niagara/Private/NiagaraRendererMeshes.cpp
@@ -33,9 +33,21 @@ public:
 
 NiagaraRendererMeshes::NiagaraRendererMeshes(ERHIFeatureLevel::Type FeatureLevel, UNiagaraRendererProperties *InProps) :
 	NiagaraRenderer()
+	, PositionOffset(INDEX_NONE)
+	, VelocityOffset(INDEX_NONE)
+	, ColorOffset(INDEX_NONE)
+	, ScaleOffset(INDEX_NONE)
+	, SizeOffset(INDEX_NONE)
+	, MaterialParamOffset(INDEX_NONE)
+	, MaterialParamOffset1(INDEX_NONE)
+	, MaterialParamOffset2(INDEX_NONE)
+	, MaterialParamOffset3(INDEX_NONE)
+	, TransformOffset(INDEX_NONE)
+	, CustomSortingOffset(INDEX_NONE)
+	, LastSyncedId(INDEX_NONE)
 {
 	//check(InProps);
-	VertexFactory = ConstructNiagaraMeshVertexFactory(NVFT_Mesh, FeatureLevel, sizeof(FNiagaraMeshInstanceVertex), 0);
+	VertexFactory = ConstructNiagaraMeshVertexFactory(NVFT_Mesh, FeatureLevel);
 	Properties = Cast<UNiagaraMeshRendererProperties>(InProps);
 
 
@@ -74,61 +86,12 @@ NiagaraRendererMeshes::NiagaraRendererMeshes(ERHIFeatureLevel::Type FeatureLevel
 
 void NiagaraRendererMeshes::SetupVertexFactory(FNiagaraMeshVertexFactory *InVertexFactory, const FStaticMeshLODResources& LODResources) const
 {
-	FNiagaraMeshVertexFactory::FDataType Data;
+	FStaticMeshDataType Data;
 
 	LODResources.VertexBuffers.PositionVertexBuffer.BindPositionVertexBuffer(InVertexFactory, Data);
 	LODResources.VertexBuffers.StaticMeshVertexBuffer.BindTangentVertexBuffer(InVertexFactory, Data);
 	LODResources.VertexBuffers.StaticMeshVertexBuffer.BindTexCoordVertexBuffer(InVertexFactory, Data, MAX_TEXCOORDS);
 	LODResources.VertexBuffers.ColorVertexBuffer.BindColorVertexBuffer(InVertexFactory, Data);
-
-	// Initialize instanced data. Vertex buffer and stride are set before render.
-	// Particle color
-	Data.ParticleColorComponent = FVertexStreamComponent(
-		NULL,
-		STRUCT_OFFSET(FNiagaraMeshInstanceVertex, Color),
-		0,
-		VET_Float4,
-		EVertexStreamUsage::Instancing
-		);
-
-	// Particle transform matrix
-	for (int32 MatrixRow = 0; MatrixRow < 3; MatrixRow++)
-	{
-		Data.TransformComponent[MatrixRow] = FVertexStreamComponent(
-			NULL,
-			STRUCT_OFFSET(FNiagaraMeshInstanceVertex, Transform) + sizeof(FVector4)* MatrixRow,
-			0,
-			VET_Float4,
-			EVertexStreamUsage::Instancing
-			);
-	}
-
-	Data.VelocityComponent = FVertexStreamComponent(
-		NULL,
-		STRUCT_OFFSET(FNiagaraMeshInstanceVertex, Velocity),
-		0,
-		VET_Float4,
-		EVertexStreamUsage::Instancing
-		);
-	// SubUVs.
-	Data.SubUVs = FVertexStreamComponent(
-		NULL,
-		STRUCT_OFFSET(FNiagaraMeshInstanceVertex, SubUVParams),
-		0,
-		VET_Short4,
-		EVertexStreamUsage::Instancing
-		);
-
-	// Pack SubUV Lerp and the particle's relative time
-	Data.SubUVLerpAndRelTime = FVertexStreamComponent(
-		NULL,
-		STRUCT_OFFSET(FNiagaraMeshInstanceVertex, SubUVLerp),
-		0,
-		VET_Float2,
-		EVertexStreamUsage::Instancing
-		);
-
-	Data.bInitialized = true;
 	InVertexFactory->SetData(Data);
 }
 
@@ -156,14 +119,25 @@ void NiagaraRendererMeshes::GetDynamicMeshElements(const TArray<const FSceneView
 	FNiagaraDynamicDataMesh *DynamicDataMesh = (static_cast<FNiagaraDynamicDataMesh*>(DynamicDataRender));
 	//if (!DynamicDataMesh || !DynamicDataMesh->DataSet || DynamicDataMesh->DataSet->GetNumInstances()==0)
 	if (!DynamicDataMesh 
-		|| DynamicDataMesh->DataSet->CurrDataRender().GetNumInstancesAllocated() == 0
-		|| DynamicDataMesh->DataSet->CurrDataRender().GetNumInstances() == 0
-		|| DynamicDataMesh->DataSet->CurrDataRender().GetGPUBufferFloat() == nullptr
+		|| DynamicDataMesh->RTParticleData.GetNumInstancesAllocated() == 0
+		|| DynamicDataMesh->RTParticleData.GetNumInstances() == 0
 		|| nullptr == Properties)
 	{
 		return;
 	}
 
+	int32 NumInstances = DynamicDataMesh->RTParticleData.GetNumInstances();
+
+	int32 TotalFloatSize = DynamicDataMesh->RTParticleData.GetFloatBuffer().Num() / sizeof(float);
+	FNiagaraGlobalReadBuffer::FAllocation ParticleData;
+	
+	//For cpu sims we allocate render buffers from the global pool. GPU sims own their own.
+	if (DynamicDataMesh->DataSet->GetSimTarget() == ENiagaraSimTarget::CPUSim)
+	{
+		ParticleData = FNiagaraGlobalReadBuffer::Get().AllocateFloat(TotalFloatSize);
+		FMemory::Memcpy(ParticleData.Buffer, DynamicDataMesh->RTParticleData.GetFloatBuffer().GetData(), DynamicDataMesh->RTParticleData.GetFloatBuffer().Num());
+	}
+
 	//check(DynamicDataMesh->DataSet->PrevDataRender().GetNumInstances() > 0);
 	//check(DynamicDataMesh->DataSet->PrevDataRender().GetNumInstancesAllocated() > 0);
 	//check(DynamicDataMesh->DataSet->PrevData().GetGPUBufferFloat()->NumBytes > 0 || DynamicDataMesh->DataSet->PrevData().GetGPUBufferInt()->NumBytes > 0)
@@ -204,60 +178,102 @@ void NiagaraRendererMeshes::GetDynamicMeshElements(const TArray<const FSceneView
 				PerViewUniformParameters.LocalToWorldInverseTransposed = bLocalSpace ? SceneProxy->GetLocalToWorld().Inverse().GetTransposed() : FMatrix::Identity;
 				PerViewUniformParameters.PrevTransformAvailable = false;
 				PerViewUniformParameters.DeltaSeconds = ViewFamily.DeltaWorldTime;
-				PerViewUniformParameters.PositionDataOffset = DynamicDataMesh->PositionDataOffset;
-				PerViewUniformParameters.VelocityDataOffset = DynamicDataMesh->VelocityDataOffset;
-				PerViewUniformParameters.ColorDataOffset = DynamicDataMesh->ColorDataOffset;
-				PerViewUniformParameters.TransformDataOffset = DynamicDataMesh->TransformDataOffset;
-				PerViewUniformParameters.ScaleDataOffset = DynamicDataMesh->ScaleDataOffset;
-				PerViewUniformParameters.SizeDataOffset = DynamicDataMesh->SizeDataOffset;
-				PerViewUniformParameters.MaterialParamDataOffset = DynamicDataMesh->MaterialParamDataOffset;
+				PerViewUniformParameters.PositionDataOffset = PositionOffset;
+				PerViewUniformParameters.VelocityDataOffset = VelocityOffset;
+				PerViewUniformParameters.ColorDataOffset = ColorOffset;
+				PerViewUniformParameters.TransformDataOffset = TransformOffset;
+				PerViewUniformParameters.ScaleDataOffset = ScaleOffset;
+				PerViewUniformParameters.SizeDataOffset = SizeOffset;
+				PerViewUniformParameters.MaterialParamDataOffset = MaterialParamOffset;
+				PerViewUniformParameters.MaterialParam1DataOffset = MaterialParamOffset1;
+				PerViewUniformParameters.MaterialParam2DataOffset = MaterialParamOffset2;
+				PerViewUniformParameters.MaterialParam3DataOffset = MaterialParamOffset3;
+				PerViewUniformParameters.DefaultPos = bLocalSpace ? FVector4(0.0f, 0.0f, 0.0f, 1.0f) : FVector4(SceneProxy->GetLocalToWorld().GetOrigin());
 				/*
 				if (Properties)
 				{
 				PerViewUniformParameters.SubImageSize = FVector4(Properties->SubImageInfo.X, Properties->SubImageInfo.Y, 1.0f / Properties->SubImageInfo.X, 1.0f / Properties->SubImageInfo.Y);
 				}
 				*/
-				CollectorResources.VertexFactory.SetParticleData(DynamicDataMesh->DataSet);
-
-				// Collector.AllocateOneFrameResource uses default ctor, initialize the vertex factory
-				CollectorResources.VertexFactory.SetParticleFactoryType(NVFT_Mesh);
-				CollectorResources.VertexFactory.SetMeshFacingMode((uint32)Properties->FacingMode);
-				CollectorResources.UniformBuffer = FNiagaraMeshUniformBufferRef::CreateUniformBufferImmediate(PerViewUniformParameters, UniformBuffer_SingleFrame);
-				CollectorResources.VertexFactory.SetStrides(0, 0);
-
-				CollectorResources.VertexFactory.InitResource();
-				CollectorResources.VertexFactory.SetUniformBuffer(CollectorResources.UniformBuffer);
-			
-				const bool bIsWireframe = AllowDebugViewmodes() && View->Family->EngineShowFlags.Wireframe;
 
+				//Grab the material proxies we'll be using for each section and check them for translucency.
+				TArray<FMaterialRenderProxy*, TInlineAllocator<32>> MaterialProxies;
+				bool bHasTranslucentMaterials = false;
 				for (int32 SectionIndex = 0; SectionIndex < LODModel.Sections.Num(); SectionIndex++)
 				{
 					const FStaticMeshSection& Section = LODModel.Sections[SectionIndex];
 					UMaterialInterface* ParticleMeshMaterial = Properties->ParticleMesh->GetMaterial(Section.MaterialIndex);
-					if (!ParticleMeshMaterial)
-					{
-						continue;
-					}
 					FMaterialRenderProxy* MaterialProxy = nullptr;
-					
-					if (Properties->bOverrideMaterials && Properties->OverrideMaterials.Num() > Section.MaterialIndex && 
+
+					if (Properties->bOverrideMaterials && Properties->OverrideMaterials.Num() > Section.MaterialIndex &&
 						Properties->OverrideMaterials[Section.MaterialIndex] != nullptr)
 					{
 						MaterialProxy = Properties->OverrideMaterials[Section.MaterialIndex]->GetRenderProxy(false, false);
 					}
 
-					if (MaterialProxy == nullptr)
+					if (MaterialProxy == nullptr && ParticleMeshMaterial)
 					{
 						MaterialProxy = ParticleMeshMaterial->GetRenderProxy(false, false);
 					}
 
+					if (MaterialProxy == nullptr)
+					{
+						MaterialProxy = UMaterial::GetDefaultMaterial(MD_Surface)->GetRenderProxy(SceneProxy->IsSelected(), SceneProxy->IsHovered());
+					}
+
+					MaterialProxies.Add(MaterialProxy);
+					if (MaterialProxy)
+					{
+						EBlendMode BlendMode = MaterialProxy->GetMaterial(VertexFactory->GetFeatureLevel())->GetBlendMode();
+						bHasTranslucentMaterials |= BlendMode == BLEND_AlphaComposite || BlendMode == BLEND_Translucent;
+					}
+				}
+
+				//Sort particles if needed.
+				FNiagaraGlobalReadBuffer::FAllocation SortedIndices;
+				CollectorResources.VertexFactory.SetSortedIndices(nullptr, 0xFFFFFFFF);
+				if (DynamicDataMesh->DataSet->GetSimTarget() == ENiagaraSimTarget::CPUSim)//TODO: Compute shader for sorting gpu sims and larger cpu sims.
+				{
+					check(ParticleData.IsValid());
+					if (bHasTranslucentMaterials || !Properties->bSortOnlyWhenTranslucent)
+					{
+						ENiagaraSortMode SortMode = Properties->SortMode;
+						bool bCustomSortMode = SortMode == ENiagaraSortMode::CustomAscending || SortMode == ENiagaraSortMode::CustomDecending;
+						int32 SortAttributeOffest = bCustomSortMode ? CustomSortingOffset : PositionOffset;
+						if (SortMode != ENiagaraSortMode::None && SortAttributeOffest != INDEX_NONE)
+						{
+							SortedIndices = FNiagaraGlobalReadBuffer::Get().AllocateInt32(NumInstances);
+							SortIndices(SortMode, SortAttributeOffest, DynamicDataMesh->RTParticleData, SceneProxy->GetLocalToWorld(), View, SortedIndices);
+							CollectorResources.VertexFactory.SetSortedIndices(SortedIndices.ReadBuffer->SRV, SortedIndices.FirstIndex / sizeof(float));
+						}
+					}
+					CollectorResources.VertexFactory.SetParticleData(ParticleData.ReadBuffer->SRV, ParticleData.FirstIndex / sizeof(float), DynamicDataMesh->RTParticleData.GetFloatStride() / sizeof(float));
+				}
+				else
+				{
+					CollectorResources.VertexFactory.SetParticleData(DynamicDataMesh->DataSet->CurrData().GetGPUBufferFloat()->SRV, 0, DynamicDataMesh->DataSet->CurrData().GetFloatStride() / sizeof(float));
+				}
+
+				// Collector.AllocateOneFrameResource uses default ctor, initialize the vertex factory
+				CollectorResources.VertexFactory.SetParticleFactoryType(NVFT_Mesh);
+				CollectorResources.VertexFactory.SetMeshFacingMode((uint32)Properties->FacingMode);
+				CollectorResources.UniformBuffer = FNiagaraMeshUniformBufferRef::CreateUniformBufferImmediate(PerViewUniformParameters, UniformBuffer_SingleFrame);
+
+				CollectorResources.VertexFactory.InitResource();
+				CollectorResources.VertexFactory.SetUniformBuffer(CollectorResources.UniformBuffer);
+			
+				const bool bIsWireframe = AllowDebugViewmodes() && View->Family->EngineShowFlags.Wireframe;
+
+				for (int32 SectionIndex = 0; SectionIndex < LODModel.Sections.Num(); SectionIndex++)
+				{
+					const FStaticMeshSection& Section = LODModel.Sections[SectionIndex];
+					FMaterialRenderProxy* MaterialProxy = MaterialProxies[SectionIndex];
 					if ((Section.NumTriangles == 0) || (MaterialProxy == NULL))
 					{
 						//@todo. This should never occur, but it does occasionally.
 						continue;
 					}
 
-
 					FMeshBatch& Mesh = Collector.AllocateMesh();
 					Mesh.VertexFactory = &CollectorResources.VertexFactory;
 					Mesh.LCI = NULL;
@@ -270,8 +286,11 @@ void NiagaraRendererMeshes::GetDynamicMeshElements(const TArray<const FSceneView
 					BatchElement.FirstIndex = 0;
 					BatchElement.MinVertexIndex = 0;
 					BatchElement.MaxVertexIndex = 0;
-					BatchElement.NumInstances = DynamicDataMesh->DataSet->CurrDataRender().GetNumInstances();
-					BatchElement.IndirectArgsBuffer = DynamicDataMesh->DataSet->GetDataSetIndices().Buffer;
+					BatchElement.NumInstances = NumInstances;
+					if (DynamicDataMesh->DataSet->GetSimTarget() == ENiagaraSimTarget::GPUComputeSim)
+					{
+						BatchElement.IndirectArgsBuffer = DynamicDataMesh->DataSet->GetCurDataSetIndices().Buffer;
+					}
 
 					if (bIsWireframe)
 					{
@@ -305,6 +324,8 @@ void NiagaraRendererMeshes::GetDynamicMeshElements(const TArray<const FSceneView
 
 					Mesh.bCanApplyViewModeOverrides = true;
 					Mesh.bUseWireframeSelectionColoring = SceneProxy->IsSelected();
+
+					check(BatchElement.NumPrimitives > 0);
 					Collector.AddMesh(ViewIndex, Mesh);
 				}
 			}
@@ -327,9 +348,10 @@ bool NiagaraRendererMeshes::SetMaterialUsage()
 /** Update render data buffer from attributes */
 FNiagaraDynamicDataBase *NiagaraRendererMeshes::GenerateVertexData(const FNiagaraSceneProxy* Proxy, FNiagaraDataSet &Data, const ENiagaraSimTarget Target)
 {
+	SCOPE_CYCLE_COUNTER(STAT_NiagaraRenderGT);
 	SCOPE_CYCLE_COUNTER(STAT_NiagaraGenMeshVertexData);
 
-	if (!Properties || Properties->ParticleMesh == nullptr || !bEnabled || Data.CurrData().GetNumInstances() == 0)
+	if (!Properties || Properties->ParticleMesh == nullptr || !bEnabled)
 	{
 		return nullptr;
 	}
@@ -338,54 +360,42 @@ FNiagaraDynamicDataBase *NiagaraRendererMeshes::GenerateVertexData(const FNiagar
 //	TArray<FNiagaraMeshInstanceVertex>& RenderData = DynamicData->VertexData;
 	//TArray< FNiagaraMeshInstanceVertexDynamicParameter>& RenderMaterialVertexData = DynamicData->MaterialParameterVertexData;
 
-
-	const FNiagaraVariableLayoutInfo* PositionLayout = Data.GetVariableLayout(FNiagaraVariable(FNiagaraTypeDefinition::GetVec3Def(), TEXT("Position")));
-	const FNiagaraVariableLayoutInfo* VelocityLayout = Data.GetVariableLayout(FNiagaraVariable(FNiagaraTypeDefinition::GetVec3Def(), TEXT("Velocity")));
-	const FNiagaraVariableLayoutInfo* ColorLayout = Data.GetVariableLayout(FNiagaraVariable(FNiagaraTypeDefinition::GetColorDef(), TEXT("Color")));
-	const FNiagaraVariableLayoutInfo* XformLayout = Data.GetVariableLayout(FNiagaraVariable(FNiagaraTypeDefinition::GetVec4Def(), TEXT("MeshOrientation")));
-	const FNiagaraVariableLayoutInfo* ScaleLayout = Data.GetVariableLayout(FNiagaraVariable(FNiagaraTypeDefinition::GetVec3Def(), TEXT("Scale")));
-	const FNiagaraVariableLayoutInfo* SizeLayout = Data.GetVariableLayout(FNiagaraVariable(FNiagaraTypeDefinition::GetVec2Def(), TEXT("SpriteSize")));
-	const FNiagaraVariableLayoutInfo* MatParamLayout = Data.GetVariableLayout(FNiagaraVariable(FNiagaraTypeDefinition::GetVec4Def(), TEXT("DynamicMaterialParameter")));
+	if (PositionOffset == INDEX_NONE || LastSyncedId != Properties->SyncId)
+	{
+		// optional attributes
+		int32 IntDummy;
+		SizeOffset = -1;
+		Data.GetVariableComponentOffsets(Properties->PositionBinding.DataSetVariable, PositionOffset, IntDummy);
+		Data.GetVariableComponentOffsets(Properties->VelocityBinding.DataSetVariable, VelocityOffset, IntDummy);
+		Data.GetVariableComponentOffsets(Properties->ColorBinding.DataSetVariable, ColorOffset, IntDummy);
+		Data.GetVariableComponentOffsets(Properties->ScaleBinding.DataSetVariable, ScaleOffset, IntDummy);
+		Data.GetVariableComponentOffsets(Properties->DynamicMaterialBinding.DataSetVariable, MaterialParamOffset, IntDummy);
+		Data.GetVariableComponentOffsets(Properties->DynamicMaterial1Binding.DataSetVariable, MaterialParamOffset1, IntDummy);
+		Data.GetVariableComponentOffsets(Properties->DynamicMaterial2Binding.DataSetVariable, MaterialParamOffset2, IntDummy);
+		Data.GetVariableComponentOffsets(Properties->DynamicMaterial3Binding.DataSetVariable, MaterialParamOffset3, IntDummy);
+		Data.GetVariableComponentOffsets(Properties->MeshOrientationBinding.DataSetVariable, TransformOffset, IntDummy);
+		Data.GetVariableComponentOffsets(Properties->CustomSortingBinding.DataSetVariable, CustomSortingOffset, IntDummy);
+		LastSyncedId = Properties->SyncId;
+	}
 
 	//Bail if we don't have the required attributes to render this emitter.
-	if (!bEnabled || !PositionLayout /*|| !XformLayout*/ || !VelocityLayout)
+	if (!bEnabled)
 	{
 		return nullptr;
 	}
 
+	FNiagaraDynamicDataMesh *DynamicData = nullptr;
 
-	UStaticMesh *Mesh = nullptr;
-	if (Properties)
+	if (Data.CurrData().GetNumInstances() > 0)
 	{
-		Mesh = Properties->ParticleMesh;
-	}
+		DynamicData = new FNiagaraDynamicDataMesh;
 
-	// required attributes
-	FNiagaraDynamicDataMesh *DynamicData = new FNiagaraDynamicDataMesh;
-	DynamicData->PositionDataOffset = PositionLayout->FloatComponentStart;
-	DynamicData->VelocityDataOffset = VelocityLayout->FloatComponentStart;
-
-	// optional attributes
-	int32 IntDummy;
-	Data.GetVariableComponentOffsets(FNiagaraVariable(FNiagaraTypeDefinition::GetColorDef(), TEXT("Color")), DynamicData->ColorDataOffset, IntDummy);
-	Data.GetVariableComponentOffsets(FNiagaraVariable(FNiagaraTypeDefinition::GetVec3Def(), TEXT("Scale")), DynamicData->ScaleDataOffset, IntDummy);
-	Data.GetVariableComponentOffsets(FNiagaraVariable(FNiagaraTypeDefinition::GetVec2Def(), TEXT("SpriteSize")), DynamicData->SizeDataOffset, IntDummy);
-	Data.GetVariableComponentOffsets(FNiagaraVariable(FNiagaraTypeDefinition::GetVec4Def(), TEXT("DynamicMaterialParameter")), DynamicData->MaterialParamDataOffset, IntDummy);
-	Data.GetVariableComponentOffsets(FNiagaraVariable(FNiagaraTypeDefinition::GetVec4Def(), TEXT("MeshOrientation")), DynamicData->TransformDataOffset, IntDummy);
-
-	// if we're CPU simulating, need to init the GPU buffers for the vertex factory
-	if (Target == ENiagaraSimTarget::CPUSim)
-	{
-		//Data.ValidateBufferIndices();
-		Data.InitGPUFromCPU();
-	}
-	else
-	{
-		check(Data.GetDataSetIndices().Buffer != nullptr);
-		Data.InitGPUSimSRVs();
+		//TODO: This buffer is far fatter than needed. Just pull out the data needed for rendering.
+		Data.CurrData().CopyTo(DynamicData->RTParticleData);
+
+		DynamicData->DataSet = &Data;
 	}
 
-	DynamicData->DataSet = &Data;
 	CPUTimeMS = VertexDataTimer.GetElapsedMilliseconds();
 	return DynamicData;  
 }
diff --git a/Engine/Plugins/FX/Niagara/Source/Niagara/Private/NiagaraRendererRibbons.cpp b/Engine/Plugins/FX/Niagara/Source/Niagara/Private/NiagaraRendererRibbons.cpp
--- a/Engine/Plugins/FX/Niagara/Source/Niagara/Private/NiagaraRendererRibbons.cpp
+++ b/Engine/Plugins/FX/Niagara/Source/Niagara/Private/NiagaraRendererRibbons.cpp
@@ -16,13 +16,11 @@ struct FNiagaraDynamicDataRibbon : public FNiagaraDynamicDataBase
 {
 	TArray<FNiagaraRibbonVertex> VertexData;
 	TArray<int16> IndexData;
-	TArray<FNiagaraRibbonVertexDynamicParameter> MaterialParameterVertexData;
+	TArray<FNiagaraRibbonVertexDynamicParameter> MaterialParameterVertexData[4];
 
+	//Direct ptr to the dataset. ONLY FOR USE BE GPU EMITTERS.
+	//TODO: Even this needs to go soon.
 	const FNiagaraDataSet *DataSet;
-	int32 PositionDataOffset;
-	int32 WidthDataOffset;
-	int32 TwistDataOffset;
-	int32 ColorDataOffset;
 };
 
 class FNiagaraMeshCollectorResourcesRibbon : public FOneFrameResource
@@ -40,6 +38,11 @@ public:
 
 NiagaraRendererRibbons::NiagaraRendererRibbons(ERHIFeatureLevel::Type FeatureLevel, UNiagaraRendererProperties *InProps) :
 	NiagaraRenderer()
+	, PositionDataOffset(INDEX_NONE)
+	, WidthDataOffset(INDEX_NONE)
+	, TwistDataOffset(INDEX_NONE)
+	, ColorDataOffset(INDEX_NONE)
+	, LastSyncedId(INDEX_NONE)
 {
 	VertexFactory = new FNiagaraRibbonVertexFactory(NVFT_Ribbon, FeatureLevel);
 	Properties = Cast<UNiagaraRibbonRendererProperties>(InProps);
@@ -68,7 +71,11 @@ void NiagaraRendererRibbons::GetDynamicMeshElements(const TArray<const FSceneVie
 
 	SimpleTimer MeshElementsTimer;
 	FNiagaraDynamicDataRibbon *DynamicDataRibbon = static_cast<FNiagaraDynamicDataRibbon*>(DynamicDataRender);
-	if (!DynamicDataRibbon || DynamicDataRibbon->VertexData.Num() < 4)
+	if (!DynamicDataRibbon 
+		|| DynamicDataRibbon->VertexData.Num() < 4
+		|| DynamicDataRibbon->IndexData.Num() == 0
+		|| nullptr == Properties
+		)
 	{
 		return;
 	}
@@ -78,20 +85,27 @@ void NiagaraRendererRibbons::GetDynamicMeshElements(const TArray<const FSceneVie
 
 	int32 SizeInBytes = DynamicDataRibbon->VertexData.GetTypeSize() * DynamicDataRibbon->VertexData.Num();
 	FGlobalDynamicVertexBuffer::FAllocation LocalDynamicVertexAllocation = FGlobalDynamicVertexBuffer::Get().Allocate(SizeInBytes);
-	FGlobalDynamicVertexBuffer::FAllocation LocalDynamicVertexMaterialParamsAllocation;
+	FGlobalDynamicVertexBuffer::FAllocation LocalDynamicVertexMaterialParamsAllocation[4];
 	FGlobalDynamicIndexBuffer::FAllocation DynamicIndexAllocation = FGlobalDynamicIndexBuffer::Get().Allocate(DynamicDataRibbon->IndexData.Num(), sizeof(int16));
 
-	if (DynamicDataRibbon->MaterialParameterVertexData.Num() > 0)
+	auto AllocDynamicParameterBuffer = [&](int32 DynamicParamIndex)
 	{
-		int32 MatParamSizeInBytes = DynamicDataRibbon->MaterialParameterVertexData.GetTypeSize() * DynamicDataRibbon->MaterialParameterVertexData.Num();
-		LocalDynamicVertexMaterialParamsAllocation = FGlobalDynamicVertexBuffer::Get().Allocate(MatParamSizeInBytes);
-
-		if (LocalDynamicVertexMaterialParamsAllocation.IsValid())
+		if (DynamicDataRibbon->MaterialParameterVertexData[DynamicParamIndex].Num() > 0)
 		{
-			// Copy the extra material vertex data over.
-			FMemory::Memcpy(LocalDynamicVertexMaterialParamsAllocation.Buffer, DynamicDataRibbon->MaterialParameterVertexData.GetData(), MatParamSizeInBytes);
+			int32 MatParamSizeInBytes = DynamicDataRibbon->MaterialParameterVertexData[DynamicParamIndex].GetTypeSize() * DynamicDataRibbon->MaterialParameterVertexData[DynamicParamIndex].Num();
+			LocalDynamicVertexMaterialParamsAllocation[DynamicParamIndex] = FGlobalDynamicVertexBuffer::Get().Allocate(MatParamSizeInBytes);
+
+			if (LocalDynamicVertexMaterialParamsAllocation[DynamicParamIndex].IsValid())
+			{
+				// Copy the extra material vertex data over.
+				FMemory::Memcpy(LocalDynamicVertexMaterialParamsAllocation[DynamicParamIndex].Buffer, DynamicDataRibbon->MaterialParameterVertexData[DynamicParamIndex].GetData(), MatParamSizeInBytes);
+			}
 		}
-	}
+	};
+	AllocDynamicParameterBuffer(0);
+	AllocDynamicParameterBuffer(1);
+	AllocDynamicParameterBuffer(2);
+	AllocDynamicParameterBuffer(3);
 
 
 	if (LocalDynamicVertexAllocation.IsValid())
@@ -133,10 +147,11 @@ void NiagaraRendererRibbons::GetDynamicMeshElements(const TArray<const FSceneVie
 				PerViewUniformParameters.CameraRight = View->GetViewRight();//	FVector4(1.0f, 0.0f, 0.0f, 0.0f);
 				PerViewUniformParameters.ScreenAlignment = FVector4(0.0f, 0.0f, 0.0f, 0.0f);
 				PerViewUniformParameters.UseCustomFacing = Properties->FacingMode == ENiagaraRibbonFacingMode::Custom;
-				PerViewUniformParameters.PositionDataOffset = DynamicDataRibbon->PositionDataOffset;
-				PerViewUniformParameters.ColorDataOffset = DynamicDataRibbon->ColorDataOffset;
-				PerViewUniformParameters.WidthDataOffset = DynamicDataRibbon->WidthDataOffset;
-				PerViewUniformParameters.TwistDataOffset = DynamicDataRibbon->TwistDataOffset;
+
+				PerViewUniformParameters.PositionDataOffset = PositionDataOffset;
+				PerViewUniformParameters.ColorDataOffset = ColorDataOffset;
+				PerViewUniformParameters.WidthDataOffset = WidthDataOffset;
+				PerViewUniformParameters.TwistDataOffset = TwistDataOffset;
 				CollectorResources.VertexFactory.SetParticleData(DynamicDataRibbon->DataSet);
 
 				// Collector.AllocateOneFrameResource uses default ctor, initialize the vertex factory
@@ -147,14 +162,22 @@ void NiagaraRendererRibbons::GetDynamicMeshElements(const TArray<const FSceneVie
 				CollectorResources.VertexFactory.InitResource();
 				CollectorResources.VertexFactory.SetBeamTrailUniformBuffer(CollectorResources.UniformBuffer);
 				CollectorResources.VertexFactory.SetVertexBuffer(LocalDynamicVertexAllocation.VertexBuffer, LocalDynamicVertexAllocation.VertexOffset, sizeof(FNiagaraRibbonVertex));
-				if (DynamicDataRibbon->MaterialParameterVertexData.Num() > 0 && LocalDynamicVertexMaterialParamsAllocation.IsValid())
-				{
-					CollectorResources.VertexFactory.SetDynamicParameterBuffer(LocalDynamicVertexMaterialParamsAllocation.VertexBuffer, LocalDynamicVertexMaterialParamsAllocation.VertexOffset, sizeof(FNiagaraRibbonVertexDynamicParameter));
-				}
-				else
+
+				auto SetDynamicParameterBuffer = [&](int32 DynamicParameterIndex)
 				{
-					CollectorResources.VertexFactory.SetDynamicParameterBuffer(NULL, 0, 0);
-				}
+					if (DynamicDataRibbon->MaterialParameterVertexData[DynamicParameterIndex].Num() > 0 && LocalDynamicVertexMaterialParamsAllocation[DynamicParameterIndex].IsValid())
+					{
+						CollectorResources.VertexFactory.SetDynamicParameterBuffer(LocalDynamicVertexMaterialParamsAllocation[DynamicParameterIndex].VertexBuffer, DynamicParameterIndex, LocalDynamicVertexMaterialParamsAllocation[DynamicParameterIndex].VertexOffset, sizeof(FNiagaraRibbonVertexDynamicParameter));
+					}
+					else
+					{
+						CollectorResources.VertexFactory.SetDynamicParameterBuffer(NULL, DynamicParameterIndex, 0, 0);
+					}
+				};
+				SetDynamicParameterBuffer(0);
+				SetDynamicParameterBuffer(1);
+				SetDynamicParameterBuffer(2);
+				SetDynamicParameterBuffer(3);
 
 				FMeshBatch& MeshBatch = Collector.AllocateMesh();
 				MeshBatch.VertexFactory = &CollectorResources.VertexFactory;
@@ -180,6 +203,7 @@ void NiagaraRendererRibbons::GetDynamicMeshElements(const TArray<const FSceneVie
 				MeshElement.IndexBuffer = DynamicIndexAllocation.IndexBuffer;
 				MeshElement.FirstIndex = DynamicIndexAllocation.FirstIndex;
 				MeshElement.NumPrimitives = DynamicDataRibbon->IndexData.Num() / 3;
+				check(MeshElement.NumPrimitives > 0);
 				MeshElement.NumInstances = 1;
 				MeshElement.MinVertexIndex = 0;
 				MeshElement.MaxVertexIndex = DynamicDataRibbon->VertexData.Num() - 1;
@@ -212,7 +236,10 @@ int NiagaraRendererRibbons::GetDynamicDataSize()
 	{
 		Size += (static_cast<FNiagaraDynamicDataRibbon*>(DynamicDataRender))->VertexData.GetAllocatedSize();
 		Size += (static_cast<FNiagaraDynamicDataRibbon*>(DynamicDataRender))->IndexData.GetAllocatedSize();
-		Size += (static_cast<FNiagaraDynamicDataRibbon*>(DynamicDataRender))->MaterialParameterVertexData.GetAllocatedSize();
+		Size += (static_cast<FNiagaraDynamicDataRibbon*>(DynamicDataRender))->MaterialParameterVertexData[0].GetAllocatedSize();
+		Size += (static_cast<FNiagaraDynamicDataRibbon*>(DynamicDataRender))->MaterialParameterVertexData[1].GetAllocatedSize();
+		Size += (static_cast<FNiagaraDynamicDataRibbon*>(DynamicDataRender))->MaterialParameterVertexData[2].GetAllocatedSize();
+		Size += (static_cast<FNiagaraDynamicDataRibbon*>(DynamicDataRender))->MaterialParameterVertexData[3].GetAllocatedSize();
 	}
 
 	return Size;
@@ -248,7 +275,7 @@ FNiagaraDynamicDataBase *NiagaraRendererRibbons::GenerateVertexData(const FNiaga
 	SCOPE_CYCLE_COUNTER(STAT_NiagaraGenRibbonVertexData);
 
 	SimpleTimer VertexDataTimer;
-	if (Data.GetNumInstances() < 2 || !bEnabled)
+	if (!bEnabled)
 	{
 		return nullptr;
 	}
@@ -256,163 +283,117 @@ FNiagaraDynamicDataBase *NiagaraRendererRibbons::GenerateVertexData(const FNiaga
 	TArray<FNiagaraRibbonVertex>& RenderData = DynamicData->VertexData;
 	TArray<int16>& IndexData = DynamicData->IndexData;
 
-	//I'm not a great fan of pulling scalar components out to a structured vert buffer like this.
-	//TODO: Experiment with a new VF that reads the data directly from the scalar layout.
-	FNiagaraDataSetIterator<FVector> PosItr(Data, FNiagaraVariable(FNiagaraTypeDefinition::GetVec3Def(), TEXT("Position")));
-	FNiagaraDataSetIterator<FVector> VelItr(Data, FNiagaraVariable(FNiagaraTypeDefinition::GetVec3Def(), TEXT("Velocity")));
-	FNiagaraDataSetIterator<FLinearColor> ColItr(Data, FNiagaraVariable(FNiagaraTypeDefinition::GetColorDef(), TEXT("Color")));
-	FNiagaraDataSetIterator<float> NormAgeItr(Data, FNiagaraVariable(FNiagaraTypeDefinition::GetFloatDef(), TEXT("NormalizedAge")));
-	FNiagaraDataSetIterator<float> AgeItr(Data, FNiagaraVariable(FNiagaraTypeDefinition::GetFloatDef(), TEXT("NormalizedAge")));
-	FNiagaraDataSetIterator<int32> RibbonIdItr(Data, FNiagaraVariable(FNiagaraTypeDefinition::GetIntDef(), TEXT("RibbonId")));
-	FNiagaraDataSetIterator<float> SizeItr(Data, FNiagaraVariable(FNiagaraTypeDefinition::GetFloatDef(), TEXT("RibbonWidth")));
-	FNiagaraDataSetIterator<float> TwistItr(Data, FNiagaraVariable(FNiagaraTypeDefinition::GetFloatDef(), TEXT("RibbonTwist")));
-	FNiagaraDataSetIterator<FVector> AlignItr(Data, FNiagaraVariable(FNiagaraTypeDefinition::GetVec3Def(), TEXT("RibbonFacing")));
-	FNiagaraDataSetIterator<FVector4> MaterialParamItr(Data, FNiagaraVariable(FNiagaraTypeDefinition::GetVec4Def(), TEXT("DynamicMaterialParameter")));
+	RenderData.Empty();
+	IndexData.Empty();
+
+	// TODO : deal with the dynamic vertex material parameter should the user have specified it as an output...
+
+	FVector2D UVs[4] = { FVector2D(1.0f, 0.0f), FVector2D(1.0f, 1.0f) };
+	int32 NumTotalVerts = 0;
+
+	FNiagaraDataSetAccessor<FVector> PosData(Data, Properties->PositionBinding.DataSetVariable);
+	FNiagaraDataSetAccessor<FVector> VelData(Data, Properties->VelocityBinding.DataSetVariable);
+	FNiagaraDataSetAccessor<FLinearColor> ColData(Data, Properties->ColorBinding.DataSetVariable);
+
+	bool bSortKeyIsAge = false;
+	FNiagaraDataSetAccessor<float> SortKeyData(Data, Properties->RibbonLinkOrderBinding.DataSetVariable);
+	if (SortKeyData.IsValid() == false)
+	{
+		SortKeyData = FNiagaraDataSetAccessor<float>(Data, Properties->NormalizedAgeBinding.DataSetVariable);
+		bSortKeyIsAge = true;
+	}
 
 	//Bail if we don't have the required attributes to render this emitter.
-	if (!PosItr.IsValid() || !ColItr.IsValid() || !NormAgeItr.IsValid() || !AgeItr.IsValid() )
+	if (Data.GetNumInstances() < 2 || !PosData.IsValid() || !SortKeyData.IsValid())
 	{
 		return DynamicData;
 	}
 
-	const FNiagaraVariableLayoutInfo* PositionLayout = Data.GetVariableLayout(FNiagaraVariable(FNiagaraTypeDefinition::GetVec3Def(), TEXT("Position")));
-	const FNiagaraVariableLayoutInfo* TwistLayout = Data.GetVariableLayout(FNiagaraVariable(FNiagaraTypeDefinition::GetFloatDef(), TEXT("RibbonTwist")));
-	const FNiagaraVariableLayoutInfo* WidthLayout = Data.GetVariableLayout(FNiagaraVariable(FNiagaraTypeDefinition::GetFloatDef(), TEXT("RibbonWidth")));
-	const FNiagaraVariableLayoutInfo* ColorLayout = Data.GetVariableLayout(FNiagaraVariable(FNiagaraTypeDefinition::GetColorDef(), TEXT("Color")));
+	if (PositionDataOffset == INDEX_NONE || LastSyncedId != Properties->SyncId)
+	{
 
-	// required attributes
-	DynamicData->PositionDataOffset = PositionLayout->FloatComponentStart;
-	DynamicData->ColorDataOffset = ColorLayout->FloatComponentStart;
+		const FNiagaraVariableLayoutInfo* PositionLayout = Data.GetVariableLayout(Properties->PositionBinding.DataSetVariable);
+		const FNiagaraVariableLayoutInfo* TwistLayout = Data.GetVariableLayout(Properties->RibbonTwistBinding.DataSetVariable);
+		const FNiagaraVariableLayoutInfo* WidthLayout = Data.GetVariableLayout(Properties->RibbonWidthBinding.DataSetVariable);
+		const FNiagaraVariableLayoutInfo* ColorLayout = Data.GetVariableLayout(Properties->ColorBinding.DataSetVariable);
 
-	// optional attributes
-	int32 IntDummy;
-	Data.GetVariableComponentOffsets(FNiagaraVariable(FNiagaraTypeDefinition::GetFloatDef(), TEXT("RibbonWidth")), DynamicData->WidthDataOffset, IntDummy);
-	Data.GetVariableComponentOffsets(FNiagaraVariable(FNiagaraTypeDefinition::GetFloatDef(), TEXT("RibbonTwist")), DynamicData->TwistDataOffset, IntDummy);
+		// required attributes
+		PositionDataOffset = PositionLayout ? PositionLayout->FloatComponentStart : INDEX_NONE;
+		ColorDataOffset = ColorLayout ? ColorLayout->FloatComponentStart : INDEX_NONE;
 
-	DynamicData->DataSet = &Data;
-	
-	////////
-
-	bool bMultiRibbons = RibbonIdItr.IsValid();
-
-	// build a sorted list by age, so we always get particles in order regardless of them being moved around due to dieing and spawning
-	// TODO: need to pull the data out because we currently can't index into the data set; terrible, but works
-	TArray<int32> SortedIndices;
-	TArray<int32> RibbonIDs;
-	TArray<int32> UniqueRibbonIDs;
-	TArray<FVector> Positions;
-	TArray<FVector> FacingVectors;
-	TArray<FVector4> DynamicParams;
-	TArray<float> Sizes;
-	TArray<FLinearColor> Colors;
-	TArray<float> Rotations;
-	TArray<float> Ages;
-	uint32 CurIdx = 0;
-	for (uint32 Idx = 0; Idx < Data.GetNumInstances(); Idx++)
-	{
-		float Age = NormAgeItr.GetAdvance();
-		if (Age >= 0.0f && Age < 1.0f)
-		{
-			SortedIndices.Add(CurIdx++);
-			Positions.Add(PosItr.GetAdvance());
-			Sizes.Add(SizeItr.GetAdvanceWithDefault(1.0));
-			Rotations.Add(TwistItr.GetAdvanceWithDefault(0.0f));
-			FacingVectors.Add(AlignItr.GetAdvanceWithDefault(FVector(0.0f, 0.0f, 1.0f)));
-			DynamicParams.Add(MaterialParamItr.GetAdvanceWithDefault(FVector4(0.0f, 0.0f, 0.0f, 0.0f)));
-
-			Ages.Add(Age);
-			Colors.Add(ColItr.GetAdvance());
-			if (bMultiRibbons)
-			{
-				int32 RibbonID = RibbonIdItr.GetAdvance();
-				RibbonIDs.Add(RibbonID);
-				UniqueRibbonIDs.AddUnique(RibbonID);
-			}
-		}
+		// optional attributes
+		int32 IntDummy;
+		Data.GetVariableComponentOffsets(Properties->RibbonWidthBinding.DataSetVariable, WidthDataOffset, IntDummy);
+		Data.GetVariableComponentOffsets(Properties->RibbonTwistBinding.DataSetVariable, TwistDataOffset, IntDummy);
+		LastSyncedId = Properties->SyncId;
 	}
 
-	if (!bMultiRibbons)
+	DynamicData->DataSet = &Data;
+
+	////////
+	FNiagaraDataSetAccessor<float> SizeData(Data, Properties->RibbonWidthBinding.DataSetVariable);
+	FNiagaraDataSetAccessor<float> TwistData(Data, Properties->RibbonTwistBinding.DataSetVariable);
+	FNiagaraDataSetAccessor<FVector> AlignData(Data, Properties->RibbonFacingBinding.DataSetVariable);
+	FNiagaraDataSetAccessor<FVector4> MaterialParamData(Data, Properties->DynamicMaterialBinding.DataSetVariable);
+	FNiagaraDataSetAccessor<FVector4> MaterialParam1Data(Data, Properties->DynamicMaterial1Binding.DataSetVariable);
+	FNiagaraDataSetAccessor<FVector4> MaterialParam2Data(Data, Properties->DynamicMaterial2Binding.DataSetVariable);
+	FNiagaraDataSetAccessor<FVector4> MaterialParam3Data(Data, Properties->DynamicMaterial3Binding.DataSetVariable);
+
+	FNiagaraDataSetAccessor<int32> RibbonIdData;
+	FNiagaraDataSetAccessor<FNiagaraID> RibbonFullIDData;
+	if (Properties->RibbonIdBinding.DataSetVariable.GetType() == FNiagaraTypeDefinition::GetIDDef())
 	{
-		SortedIndices.Sort(
-			[&Ages](const int32& A, const int32& B) {
-			return (Ages[A] < Ages[B]);
-		}
-		);
-		UniqueRibbonIDs.Add(0);
+		RibbonFullIDData.Create(&Data, Properties->RibbonIdBinding.DataSetVariable);
+		RibbonFullIDData.InitForAccess(true);
 	}
 	else
 	{
-		SortedIndices.Sort(
-			[&RibbonIDs](const int32& A, const int32& B) {
-			return (RibbonIDs[A] < RibbonIDs[B]);
-		}
-		);
+		RibbonIdData.Create(&Data, Properties->RibbonIdBinding.DataSetVariable);
+		RibbonIdData.InitForAccess(true);
 	}
 
+	bool bFullIDs = RibbonFullIDData.IsValid();
+	bool bSimpleIDs = !bFullIDs && RibbonIdData.IsValid();
+	bool bMultiRibbons = bFullIDs || bSimpleIDs;
 
-	if (SortedIndices.Num() == 0)
+	auto AddRibbonVerts = [&](TArray<int32>& RibbonIndices)
 	{
-		return DynamicData;
-	}
-
-	RenderData.Empty();
-	IndexData.Empty();
-
-	// TODO : deal with the dynamic vertex material parameter should the user have specified it as an output...
-
-	FVector2D UVs[4] = { FVector2D(1.0f, 0.0f), FVector2D(1.0f, 1.0f) };
-	int32 StartIndex = 0;
-	int32 NumIndices = SortedIndices.Num();
-	int32 *IndexPtr = &SortedIndices[0];
-	int32 NumTotalVerts = 0;
-
-	for (int32 RibbonIdx = 0; RibbonIdx < UniqueRibbonIDs.Num(); RibbonIdx++)
-	{
-		TArray<int32> SortedPartialIndices;
-		if (bMultiRibbons)
-		{
-			int32 CurID = RibbonIDs[SortedIndices[StartIndex]];
-			NumIndices = 0;
-			while ((StartIndex + NumIndices)<RibbonIDs.Num() && RibbonIDs[SortedIndices[StartIndex + NumIndices]] == CurID)
-			{
-				SortedPartialIndices.Add(SortedIndices[StartIndex + NumIndices]);
-				NumIndices++;
-			}
-			SortedPartialIndices.Sort(
-				[&Ages](const int32& A, const int32& B) {
-				return (Ages[A] < Ages[B]);
-			}
-			);
-			IndexPtr = &SortedPartialIndices[0];
-		}
-
+		int32 NumIndices = RibbonIndices.Num();
 		if (NumIndices > 1)
 		{
 			FVector PrevPos, PrevPos2, PrevDir(0.0f, 0.0f, 0.1f);
 			float TotalDistance = 0.0f;
-			FVector2D AgeOffset(Ages.Last(), 0.0f);
 
-			for (int32 i = 0; i < NumIndices; i++)
+			FVector2D AgeOffset;
+			if (bSortKeyIsAge)
 			{
-				uint32 Index1 = IndexPtr[i];
-				uint32 Index2 = IndexPtr[i + 1];
+				AgeOffset = FVector2D(SortKeyData[RibbonIndices.Last()], 0.0f);
+			}
+			else
+			{
+				AgeOffset = FVector2D::ZeroVector;
+			}
 
-				const FVector ParticlePos = Positions[Index1];
+			for (int32 i = 0; i < NumIndices; i++)
+			{
+				uint32 Index1 = RibbonIndices[i];
+				uint32 Index2 = 0;
+				const FVector ParticlePos = PosData[Index1];
 				FVector ParticleDir;
 				if (i < NumIndices - 1)
 				{
-					ParticleDir = Positions[Index2] - ParticlePos;
+					Index2 = RibbonIndices[i + 1];
+					ParticleDir = PosData[Index2] - ParticlePos;
 				}
 				else
 				{
-					Index2 = IndexPtr[i - 1];
-					ParticleDir = ParticlePos - Positions[Index2];
+					Index2 = RibbonIndices[i - 1];
+					ParticleDir = ParticlePos - PosData[Index2];
 				}
 
-
 				// if two ribbon particles were spawned too close together, we skip one
 				// but never skip the last, because that will result in invalid indices from the prev loop
-				if (ParticleDir.SizeSquared() > 0.002 )
+				if (ParticleDir.SizeSquared() > 0.002 || i == NumIndices - 1)
 				{
 					FVector NormDir = ParticleDir.GetSafeNormal();
 					PrevDir = NormDir;
@@ -428,17 +409,37 @@ FNiagaraDynamicDataBase *NiagaraRendererRibbons::GenerateVertexData(const FNiaga
 						UV1Mult = FVector2D(TotalDistance / Properties->UV1TilingDistance, 1.0f);
 					}
 
+					//Todo. Possibly pull out to templated func to avoid safety code in loop here.
+					FLinearColor Color = ColData.GetSafe(Index1, FLinearColor::White);
+					float Twist = TwistData.GetSafe(Index1, 0.0f);
+					float Size = SizeData.GetSafe(Index1, 1.0f);
+					FVector Align = AlignData.GetSafe(Index1, FVector::UpVector);
+
 					TotalDistance += ParticleDir.Size();
-					AddRibbonVert(RenderData, ParticlePos, UVs[0] * UV0Mult + AgeOffset, UVs[0] * UV1Mult + AgeOffset, Colors[Index1], Ages[Index1], Rotations[Index1], Sizes[Index1], NormDir, FacingVectors[Index1]);
-					AddRibbonVert(RenderData, ParticlePos, UVs[1] * UV0Mult + AgeOffset, UVs[1] * UV1Mult + AgeOffset, Colors[Index1], Ages[Index1], Rotations[Index1], Sizes[Index1], NormDir, FacingVectors[Index1]);
-					if (DynamicParams.Num())
+					AddRibbonVert(RenderData, ParticlePos, UVs[0] * UV0Mult + AgeOffset, UVs[0] * UV1Mult + AgeOffset, Color, Twist, Size, NormDir, Align);
+					AddRibbonVert(RenderData, ParticlePos, UVs[1] * UV0Mult + AgeOffset, UVs[1] * UV1Mult + AgeOffset, Color, Twist, Size, NormDir, Align);
+					if (MaterialParamData.IsValid())
 					{
-						AddDynamicParam(DynamicData->MaterialParameterVertexData, DynamicParams[Index1]);
-						AddDynamicParam(DynamicData->MaterialParameterVertexData, DynamicParams[Index1]);
+						AddDynamicParam(DynamicData->MaterialParameterVertexData[0], MaterialParamData[Index1]);
+						AddDynamicParam(DynamicData->MaterialParameterVertexData[0], MaterialParamData[Index1]);
 					}
-
-					if (i < NumIndices - 1)
+					if (MaterialParam1Data.IsValid())
+					{
+						AddDynamicParam(DynamicData->MaterialParameterVertexData[1], MaterialParam1Data[Index1]);
+						AddDynamicParam(DynamicData->MaterialParameterVertexData[1], MaterialParam1Data[Index1]);
+					}
+					if (MaterialParam2Data.IsValid())
+					{
+						AddDynamicParam(DynamicData->MaterialParameterVertexData[2], MaterialParam2Data[Index1]);
+						AddDynamicParam(DynamicData->MaterialParameterVertexData[2], MaterialParam2Data[Index1]);
+					}
+					if (MaterialParam3Data.IsValid())
 					{
+						AddDynamicParam(DynamicData->MaterialParameterVertexData[3], MaterialParam3Data[Index1]);
+						AddDynamicParam(DynamicData->MaterialParameterVertexData[3], MaterialParam3Data[Index1]);
+					}
+ 					if (i < NumIndices - 1)
+ 					{
 						IndexData.Add(NumTotalVerts);
 						IndexData.Add(NumTotalVerts + 1);
 						IndexData.Add(NumTotalVerts + 2);
@@ -450,7 +451,68 @@ FNiagaraDynamicDataBase *NiagaraRendererRibbons::GenerateVertexData(const FNiaga
 				}
 			}
 		}
-		StartIndex += NumIndices;
+	};
+	
+	//TODO: Move sorting to share code with sprite and mesh sorting and support the custom sorting key.
+	int32 TotalIndices = Data.GetNumInstances();
+
+	if (!bMultiRibbons)
+	{
+		int32 CurIdx = 0;
+		TArray<int32> SortedIndices;
+		for (int32 i = 0; i < TotalIndices; ++i)
+		{
+			SortedIndices.Add(i);
+		}
+
+		SortedIndices.Sort([&SortKeyData](const int32& A, const int32& B) {	return (SortKeyData[A] < SortKeyData[B]); });
+
+		AddRibbonVerts(SortedIndices);
+	}
+	else
+	{
+		int32 CurIdx = 0;
+		if (bFullIDs)
+		{
+			TMap<FNiagaraID, TArray<int32>> MultiRibbonSortedIndices;
+			
+			for (int32 i = 0; i < TotalIndices; ++i)
+			{
+				TArray<int32>& Indices = MultiRibbonSortedIndices.FindOrAdd(RibbonFullIDData[i]);
+				Indices.Add(i);
+			}
+
+			for (TPair<FNiagaraID, TArray<int32>>& Pair : MultiRibbonSortedIndices)
+			{
+				FNiagaraID& ID = Pair.Key;
+				TArray<int32>& SortedIndices = Pair.Value;
+
+				SortedIndices.Sort([&SortKeyData](const int32& A, const int32& B) {	return (SortKeyData[A] < SortKeyData[B]); });
+				AddRibbonVerts(SortedIndices);			
+			};
+		}
+		else
+		{
+			//TODO: Remove simple ID path
+			check(bSimpleIDs);
+
+			TMap<int32, TArray<int32>> MultiRibbonSortedIndices;
+
+			for (int32 i = 0; i < TotalIndices; ++i)
+			{
+				TArray<int32>& Indices = MultiRibbonSortedIndices.FindOrAdd(RibbonIdData[i]);
+				Indices.Add(i);
+			}
+
+			for (TPair<int32, TArray<int32>>& Pair : MultiRibbonSortedIndices)
+			{
+				int32 ID = Pair.Key;
+				TArray<int32>& SortedIndices = Pair.Value;
+
+				SortedIndices.Sort([&SortKeyData](const int32& A, const int32& B) {	return (SortKeyData[A] < SortKeyData[B]); });
+				AddRibbonVerts(SortedIndices);
+			};
+		}
 	}
 
 	CPUTimeMS = VertexDataTimer.GetElapsedMilliseconds();
diff --git a/Engine/Plugins/FX/Niagara/Source/Niagara/Private/NiagaraRendererSprites.cpp b/Engine/Plugins/FX/Niagara/Source/Niagara/Private/NiagaraRendererSprites.cpp
--- a/Engine/Plugins/FX/Niagara/Source/Niagara/Private/NiagaraRendererSprites.cpp
+++ b/Engine/Plugins/FX/Niagara/Source/Niagara/Private/NiagaraRendererSprites.cpp
@@ -13,6 +13,8 @@ DECLARE_CYCLE_STAT(TEXT("Genereate GPU Buffers"), STAT_NiagaraGenSpriteGpuBuffer
 
 struct FNiagaraDynamicDataSprites : public FNiagaraDynamicDataBase
 {
+	//Direct ptr to the dataset. ONLY FOR USE BE GPU EMITTERS.
+	//TODO: Even this needs to go soon.
 	const FNiagaraDataSet *DataSet;
 };
 
@@ -44,9 +46,14 @@ NiagaraRendererSprites::NiagaraRendererSprites(ERHIFeatureLevel::Type FeatureLev
 	, AlignmentOffset(INDEX_NONE)
 	, SubImageOffset(INDEX_NONE)
 	, MaterialParamOffset(INDEX_NONE)
+	, MaterialParamOffset1(INDEX_NONE)
+	, MaterialParamOffset2(INDEX_NONE)
+	, MaterialParamOffset3(INDEX_NONE)
 	, CameraOffsetOffset(INDEX_NONE)
 	, UVScaleOffset(INDEX_NONE)
 	, ParticleRandomOffset(INDEX_NONE)
+	, CustomSortingOffset(INDEX_NONE)
+	, LastSyncId(INDEX_NONE)
 {
 	//check(InProps);
 	VertexFactory = new FNiagaraSpriteVertexFactory(NVFT_Sprite, FeatureLevel);
@@ -79,16 +86,24 @@ void NiagaraRendererSprites::GetDynamicMeshElements(const TArray<const FSceneVie
 	FNiagaraDynamicDataSprites *DynamicDataSprites = static_cast<FNiagaraDynamicDataSprites*>(DynamicDataRender);
 
 	if (!DynamicDataSprites
-		|| DynamicDataSprites->DataSet->CurrDataRender().GetNumInstancesAllocated() == 0
-		|| DynamicDataSprites->DataSet->CurrDataRender().GetNumInstances() == 0
-		|| DynamicDataSprites->DataSet->CurrDataRender().GetGPUBufferFloat() == nullptr
+		|| DynamicDataSprites->RTParticleData.GetNumInstancesAllocated() == 0
+		|| DynamicDataSprites->RTParticleData.GetNumInstances() == 0
 		|| nullptr == Properties
 		)
 	{
 		return;
 	}
 
-	int32 NumInstances = DynamicDataSprites->DataSet->CurrDataRender().GetNumInstances();
+	int32 NumInstances = DynamicDataSprites->RTParticleData.GetNumInstances();
+
+	int32 TotalFloatSize = DynamicDataSprites->RTParticleData.GetFloatBuffer().Num() / sizeof(float);
+	FNiagaraGlobalReadBuffer::FAllocation ParticleData;
+	
+	if (DynamicDataSprites->DataSet->GetSimTarget() == ENiagaraSimTarget::CPUSim)
+	{
+		ParticleData = FNiagaraGlobalReadBuffer::Get().AllocateFloat(TotalFloatSize);
+		FMemory::Memcpy(ParticleData.Buffer, DynamicDataSprites->RTParticleData.GetFloatBuffer().GetData(), DynamicDataSprites->RTParticleData.GetFloatBuffer().Num());
+	}
 
 	const bool bIsWireframe = ViewFamily.EngineShowFlags.Wireframe;
 	FMaterialRenderProxy* MaterialRenderProxy = Material->GetRenderProxy(SceneProxy->IsSelected(), SceneProxy->IsHovered());
@@ -129,7 +144,7 @@ void NiagaraRendererSprites::GetDynamicMeshElements(const TArray<const FSceneVie
 				PerViewUniformParameters.RotationBias = 0.0f;
 				PerViewUniformParameters.RotationScale = 1.0f;
 				PerViewUniformParameters.TangentSelector = FVector4(0.0f, 0.0f, 0.0f, 1.0f);
-				PerViewUniformParameters.InvDeltaSeconds = 30.0f;
+				PerViewUniformParameters.DeltaSeconds = ViewFamily.DeltaWorldTime;
 				PerViewUniformParameters.NormalsType = 0.0f;
 				PerViewUniformParameters.NormalsSphereCenter = FVector4(0.0f, 0.0f, 0.0f, 1.0f);
 				PerViewUniformParameters.NormalsCylinderUnitDirection = FVector4(0.0f, 0.0f, 1.0f, 0.0f);
@@ -145,6 +160,9 @@ void NiagaraRendererSprites::GetDynamicMeshElements(const TArray<const FSceneVie
 				PerViewUniformParameters.SizeDataOffset = SizeOffset;
 				PerViewUniformParameters.ColorDataOffset = ColorOffset;
 				PerViewUniformParameters.MaterialParamDataOffset = MaterialParamOffset;
+				PerViewUniformParameters.MaterialParam1DataOffset = MaterialParamOffset1;
+				PerViewUniformParameters.MaterialParam2DataOffset = MaterialParamOffset2;
+				PerViewUniformParameters.MaterialParam3DataOffset = MaterialParamOffset3;
 				PerViewUniformParameters.SubimageDataOffset = SubImageOffset;
 				PerViewUniformParameters.FacingOffset = FacingOffset;
 				PerViewUniformParameters.AlignmentOffset = AlignmentOffset;
@@ -152,6 +170,9 @@ void NiagaraRendererSprites::GetDynamicMeshElements(const TArray<const FSceneVie
 				PerViewUniformParameters.CameraOffsetOffset = CameraOffsetOffset;
 				PerViewUniformParameters.UVScaleOffset = UVScaleOffset;
 				PerViewUniformParameters.ParticleRandomOffset = ParticleRandomOffset;
+				PerViewUniformParameters.DefaultPos = bLocalSpace ? FVector4(0.0f, 0.0f, 0.0f, 1.0f) : FVector4(SceneProxy->GetLocalToWorld().GetOrigin());
+
+				//UE_LOG(LogNiagara, Log, TEXT("SubImageSize: %f %f offset: %d"), Properties->SubImageSize.X, Properties->SubImageSize.Y, SubImageOffset);
 
 				// Collector.AllocateOneFrameResource uses default ctor, initialize the vertex factory
 				ENiagaraSpriteFacingMode FacingMode = Properties->FacingMode;
@@ -191,7 +212,32 @@ void NiagaraRendererSprites::GetDynamicMeshElements(const TArray<const FSceneVie
 					PerViewUniformParameters.CustomFacingVectorMask = Properties->CustomFacingVectorMask;
 				}
 
-				CollectorResources.VertexFactory.SetParticleData(DynamicDataSprites->DataSet);
+				//Sort particles if needed.
+				EBlendMode BlendMode = MaterialRenderProxy->GetMaterial(VertexFactory->GetFeatureLevel())->GetBlendMode();
+				FNiagaraGlobalReadBuffer::FAllocation SortedIndices;
+				CollectorResources.VertexFactory.SetSortedIndices(nullptr, 0xFFFFFFFF);
+				if (DynamicDataSprites->DataSet->GetSimTarget() == ENiagaraSimTarget::CPUSim)//TODO: Compute shader for sorting gpu sims and larger cpu sims.
+				{
+					check(ParticleData.IsValid());
+					if (BlendMode == BLEND_AlphaComposite || BlendMode == BLEND_Translucent || !Properties->bSortOnlyWhenTranslucent)
+					{
+						ENiagaraSortMode SortMode = Properties->SortMode;
+						bool bCustomSortMode = SortMode == ENiagaraSortMode::CustomAscending || SortMode == ENiagaraSortMode::CustomDecending;
+						int32 SortAttributeOffest = bCustomSortMode ? CustomSortingOffset : PositionOffset;
+						if (SortMode != ENiagaraSortMode::None && SortAttributeOffest != INDEX_NONE)
+						{
+							SortedIndices = FNiagaraGlobalReadBuffer::Get().AllocateInt32(NumInstances);
+							SortIndices(SortMode, SortAttributeOffest, DynamicDataSprites->RTParticleData, SceneProxy->GetLocalToWorld(), View, SortedIndices);
+							CollectorResources.VertexFactory.SetSortedIndices(SortedIndices.ReadBuffer->SRV, SortedIndices.FirstIndex / sizeof(float));
+						}
+					}
+					CollectorResources.VertexFactory.SetParticleData(ParticleData.ReadBuffer->SRV, ParticleData.FirstIndex / sizeof(float), DynamicDataSprites->RTParticleData.GetFloatStride() / sizeof(float));
+				}
+				else
+				{
+					CollectorResources.VertexFactory.SetParticleData(DynamicDataSprites->DataSet->CurrData().GetGPUBufferFloat()->SRV, 0, DynamicDataSprites->DataSet->CurrData().GetFloatStride() / sizeof(float));
+				}
+
 				CollectorResources.VertexFactory.SetAlignmentMode((uint32)AlignmentMode);
 				CollectorResources.VertexFactory.SetFacingMode((uint32)FacingMode);
 
@@ -231,7 +277,10 @@ void NiagaraRendererSprites::GetDynamicMeshElements(const TArray<const FSceneVie
 				MeshElement.MinVertexIndex = 0;
 				MeshElement.MaxVertexIndex = 0;// MeshElement.NumInstances * 4 - 1;
 				MeshElement.PrimitiveUniformBufferResource = &WorldSpacePrimitiveUniformBuffer;
-				MeshElement.IndirectArgsBuffer = DynamicDataSprites->DataSet->GetDataSetIndices().Buffer;
+				if (DynamicDataSprites->DataSet->GetSimTarget() == ENiagaraSimTarget::GPUComputeSim)
+				{
+					MeshElement.IndirectArgsBuffer = DynamicDataSprites->DataSet->GetCurDataSetIndices().Buffer;
+				}
 				
 				Collector.AddMesh(ViewIndex, MeshBatch);				
 			}
@@ -252,57 +301,51 @@ FNiagaraDynamicDataBase *NiagaraRendererSprites::GenerateVertexData(const FNiaga
 {
 	FNiagaraDynamicDataSprites *DynamicData = nullptr;
 
-	if (Data.CurrData().GetNumInstances() > 0 && (Target != ENiagaraSimTarget::CPUSim || (Target == ENiagaraSimTarget::CPUSim && Data.CurrData().GetSizeBytes() != 0)))
+	if (bEnabled && Properties)
 	{
 		SimpleTimer VertexDataTimer;
 
 		SCOPE_CYCLE_COUNTER(STAT_NiagaraGenSpriteVertexData);
 
-		if (PositionOffset == INDEX_NONE)
-		{
-			const FNiagaraVariableLayoutInfo* PositionLayout = Data.GetVariableLayout(FNiagaraVariable(FNiagaraTypeDefinition::GetVec3Def(), TEXT("Position")));
-			const FNiagaraVariableLayoutInfo* VelocityLayout = Data.GetVariableLayout(FNiagaraVariable(FNiagaraTypeDefinition::GetVec3Def(), TEXT("Velocity")));
-			const FNiagaraVariableLayoutInfo* RotationLayout = Data.GetVariableLayout(FNiagaraVariable(FNiagaraTypeDefinition::GetFloatDef(), TEXT("SpriteRotation")));
-			const FNiagaraVariableLayoutInfo* SizeLayout = Data.GetVariableLayout(FNiagaraVariable(FNiagaraTypeDefinition::GetVec2Def(), TEXT("SpriteSize")));
-			const FNiagaraVariableLayoutInfo* ColorLayout = Data.GetVariableLayout(FNiagaraVariable(FNiagaraTypeDefinition::GetColorDef(), TEXT("Color")));
-			PositionOffset = PositionLayout->FloatComponentStart;
-			VelocityOffset = VelocityLayout->FloatComponentStart;
-			RotationOffset = RotationLayout->FloatComponentStart;
-			SizeOffset = SizeLayout->FloatComponentStart;
-			ColorOffset = ColorLayout->FloatComponentStart;
-
+		if (PositionOffset == INDEX_NONE || LastSyncId != Properties->SyncId)
+		{			
 			// optional attributes; we pass -1 as the offset so the VF can branch
 			int32 IntDummy;
-			Data.GetVariableComponentOffsets(FNiagaraVariable(FNiagaraTypeDefinition::GetVec3Def(), TEXT("SpriteFacing")), FacingOffset, IntDummy);
-			Data.GetVariableComponentOffsets(FNiagaraVariable(FNiagaraTypeDefinition::GetVec3Def(), TEXT("SpriteAlignment")), AlignmentOffset, IntDummy);
-			Data.GetVariableComponentOffsets(FNiagaraVariable(FNiagaraTypeDefinition::GetFloatDef(), TEXT("SubImageIndex")), SubImageOffset, IntDummy);
-			Data.GetVariableComponentOffsets(FNiagaraVariable(FNiagaraTypeDefinition::GetVec4Def(), TEXT("DynamicMaterialParameter")), MaterialParamOffset, IntDummy);
-			Data.GetVariableComponentOffsets(FNiagaraVariable(FNiagaraTypeDefinition::GetFloatDef(), TEXT("CameraOffset")), CameraOffsetOffset, IntDummy);
-			Data.GetVariableComponentOffsets(FNiagaraVariable(FNiagaraTypeDefinition::GetVec2Def(), TEXT("UVScale")), UVScaleOffset, IntDummy);
-			Data.GetVariableComponentOffsets(FNiagaraVariable(FNiagaraTypeDefinition::GetFloatDef(), TEXT("MaterialRandom")), ParticleRandomOffset, IntDummy);
-		}
-		
-		if (!bEnabled || PositionOffset == INDEX_NONE || VelocityOffset == INDEX_NONE || RotationOffset == INDEX_NONE || SizeOffset == INDEX_NONE || ColorOffset == INDEX_NONE)
-		{
-			return nullptr;
+			Data.GetVariableComponentOffsets(Properties->PositionBinding.DataSetVariable, PositionOffset, IntDummy);
+			Data.GetVariableComponentOffsets(Properties->VelocityBinding.DataSetVariable, VelocityOffset, IntDummy);
+			Data.GetVariableComponentOffsets(Properties->SpriteRotationBinding.DataSetVariable, RotationOffset, IntDummy);
+			Data.GetVariableComponentOffsets(Properties->SpriteSizeBinding.DataSetVariable, SizeOffset, IntDummy);
+			Data.GetVariableComponentOffsets(Properties->ColorBinding.DataSetVariable, ColorOffset, IntDummy);
+			Data.GetVariableComponentOffsets(Properties->SpriteFacingBinding.DataSetVariable, FacingOffset, IntDummy);
+			Data.GetVariableComponentOffsets(Properties->SpriteAlignmentBinding.DataSetVariable, AlignmentOffset, IntDummy);
+			Data.GetVariableComponentOffsets(Properties->SubImageIndexBinding.DataSetVariable, SubImageOffset, IntDummy);
+			Data.GetVariableComponentOffsets(Properties->DynamicMaterialBinding.DataSetVariable, MaterialParamOffset, IntDummy);
+			Data.GetVariableComponentOffsets(Properties->DynamicMaterial1Binding.DataSetVariable, MaterialParamOffset1, IntDummy);
+			Data.GetVariableComponentOffsets(Properties->DynamicMaterial2Binding.DataSetVariable, MaterialParamOffset2, IntDummy);
+			Data.GetVariableComponentOffsets(Properties->DynamicMaterial3Binding.DataSetVariable, MaterialParamOffset3, IntDummy);
+			Data.GetVariableComponentOffsets(Properties->CameraOffsetBinding.DataSetVariable, CameraOffsetOffset, IntDummy);
+			Data.GetVariableComponentOffsets(Properties->UVScaleBinding.DataSetVariable, UVScaleOffset, IntDummy);
+			Data.GetVariableComponentOffsets(Properties->MaterialRandomBinding.DataSetVariable, ParticleRandomOffset, IntDummy);
+			Data.GetVariableComponentOffsets(Properties->CustomSortingBinding.DataSetVariable, CustomSortingOffset, IntDummy);
+
+			if (CustomSortingOffset == INDEX_NONE && (Properties->SortMode == ENiagaraSortMode::CustomAscending || Properties->SortMode == ENiagaraSortMode::CustomDecending))
+			{
+				UE_LOG(LogNiagara, Warning, TEXT("Niagara Sprite Emitter using custom sorting but does not have a valid custom sorting attribute binding."));
+			}
+
+			LastSyncId = Properties->SyncId;
 		}
 		
-		DynamicData = new FNiagaraDynamicDataSprites;
-		// if we're CPU simulating, need to init the GPU buffers for the vertex factory
-			//This is a killer for perf. Assume its all the GPU buffer allocation.
-			//Need to allocate big shared buffer and allocate ranges out of those.
-		SCOPE_CYCLE_COUNTER(STAT_NiagaraGenSpriteGpuBuffers);
-		if (Target == ENiagaraSimTarget::CPUSim)
-		{
-			Data.InitGPUFromCPU();
-		}
-		else
+		if(Data.CurrData().GetNumInstances() > 0)
 		{
-			check(Data.GetDataSetIndices().Buffer != nullptr);
-			Data.InitGPUSimSRVs();
+			DynamicData = new FNiagaraDynamicDataSprites;
+
+			//TODO: This buffer is far fatter than needed. Just pull out the data needed for rendering.
+			Data.CurrData().CopyTo(DynamicData->RTParticleData);
+
+			DynamicData->DataSet = &Data;
 		}
 
-		DynamicData->DataSet = &Data;
 		CPUTimeMS = VertexDataTimer.GetElapsedMilliseconds();
 	}
 
diff --git a/Engine/Plugins/FX/Niagara/Source/Niagara/Private/NiagaraRibbonRendererProperties.cpp b/Engine/Plugins/FX/Niagara/Source/Niagara/Private/NiagaraRibbonRendererProperties.cpp
--- a/Engine/Plugins/FX/Niagara/Source/Niagara/Private/NiagaraRibbonRendererProperties.cpp
+++ b/Engine/Plugins/FX/Niagara/Source/Niagara/Private/NiagaraRibbonRendererProperties.cpp
@@ -14,7 +14,52 @@ void UNiagaraRibbonRendererProperties::GetUsedMaterials(TArray<UMaterialInterfac
 	OutMaterials.Add(Material);
 }
 
+void UNiagaraRibbonRendererProperties::PostInitProperties()
+{
+	Super::PostInitProperties();
+	SyncId = 0;
+	if (HasAnyFlags(RF_ClassDefaultObject) == false)
+	{
+		InitBindings();
+	}
+}
+
+/** The bindings depend on variables that are created during the NiagaraModule startup. However, the CDO's are build prior to this being initialized, so we defer setting these values until later.*/
+void UNiagaraRibbonRendererProperties::InitCDOPropertiesAfterModuleStartup()
+{
+	UNiagaraRibbonRendererProperties* CDO = CastChecked<UNiagaraRibbonRendererProperties>(UNiagaraRibbonRendererProperties::StaticClass()->GetDefaultObject());
+	CDO->InitBindings();
+}
+
+void UNiagaraRibbonRendererProperties::InitBindings()
+{
+	if (PositionBinding.BoundVariable.GetName() == NAME_None)
+	{
+		PositionBinding = FNiagaraConstants::GetAttributeDefaultBinding(SYS_PARAM_PARTICLES_POSITION);
+		ColorBinding = FNiagaraConstants::GetAttributeDefaultBinding(SYS_PARAM_PARTICLES_COLOR);
+		VelocityBinding = FNiagaraConstants::GetAttributeDefaultBinding(SYS_PARAM_PARTICLES_VELOCITY);
+		DynamicMaterialBinding = FNiagaraConstants::GetAttributeDefaultBinding(SYS_PARAM_PARTICLES_DYNAMIC_MATERIAL_PARAM);
+		DynamicMaterial1Binding = FNiagaraConstants::GetAttributeDefaultBinding(SYS_PARAM_PARTICLES_DYNAMIC_MATERIAL_PARAM_1);
+		DynamicMaterial2Binding = FNiagaraConstants::GetAttributeDefaultBinding(SYS_PARAM_PARTICLES_DYNAMIC_MATERIAL_PARAM_2);
+		DynamicMaterial3Binding = FNiagaraConstants::GetAttributeDefaultBinding(SYS_PARAM_PARTICLES_DYNAMIC_MATERIAL_PARAM_3);
+		NormalizedAgeBinding = FNiagaraConstants::GetAttributeDefaultBinding(SYS_PARAM_PARTICLES_NORMALIZED_AGE);
+		RibbonTwistBinding = FNiagaraConstants::GetAttributeDefaultBinding(SYS_PARAM_PARTICLES_RIBBONTWIST);
+		RibbonWidthBinding = FNiagaraConstants::GetAttributeDefaultBinding(SYS_PARAM_PARTICLES_RIBBONWIDTH);
+		RibbonFacingBinding = FNiagaraConstants::GetAttributeDefaultBinding(SYS_PARAM_PARTICLES_RIBBONFACING);
+		RibbonIdBinding = FNiagaraConstants::GetAttributeDefaultBinding(SYS_PARAM_PARTICLES_RIBBONID);
+		RibbonLinkOrderBinding = FNiagaraConstants::GetAttributeDefaultBinding(SYS_PARAM_PARTICLES_RIBBONLINKORDER);
+	}
+}
+
 #if WITH_EDITORONLY_DATA
+void UNiagaraRibbonRendererProperties::PostEditChangeProperty(struct FPropertyChangedEvent& PropertyChangedEvent)
+{
+	Super::PostEditChangeProperty(PropertyChangedEvent);
+	if (PropertyChangedEvent.GetPropertyName() != TEXT("SyncId"))
+	{
+		SyncId++;
+	}
+}
 
 const TArray<FNiagaraVariable>& UNiagaraRibbonRendererProperties::GetRequiredAttributes()
 {
@@ -22,9 +67,6 @@ const TArray<FNiagaraVariable>& UNiagaraRibbonRendererProperties::GetRequiredAtt
 
 	if (Attrs.Num() == 0)
 	{
-		Attrs.Add(SYS_PARAM_PARTICLES_POSITION);
-		Attrs.Add(SYS_PARAM_PARTICLES_NORMALIZED_AGE);
-		Attrs.Add(SYS_PARAM_PARTICLES_COLOR);
 	}
 
 	return Attrs;
@@ -37,10 +79,14 @@ const TArray<FNiagaraVariable>& UNiagaraRibbonRendererProperties::GetOptionalAtt
 
 	if (Attrs.Num() == 0)
 	{
+		Attrs.Add(SYS_PARAM_PARTICLES_POSITION);
+		Attrs.Add(SYS_PARAM_PARTICLES_NORMALIZED_AGE);
+		Attrs.Add(SYS_PARAM_PARTICLES_COLOR);
 		Attrs.Add(SYS_PARAM_PARTICLES_RIBBONID);
 		Attrs.Add(SYS_PARAM_PARTICLES_RIBBONTWIST);
 		Attrs.Add(SYS_PARAM_PARTICLES_RIBBONWIDTH);
 		Attrs.Add(SYS_PARAM_PARTICLES_RIBBONFACING);
+		Attrs.Add(SYS_PARAM_PARTICLES_RIBBONLINKORDER);
 	}
 
 	return Attrs;
