Copying //UE4/Dev-Niagara to //UE4/Dev-Main (Source: //UE4/Dev-Niagara @ 4074996)

diff --git a/Engine/Plugins/FX/Niagara/Source/Niagara/Private/NiagaraScript.cpp b/Engine/Plugins/FX/Niagara/Source/Niagara/Private/NiagaraScript.cpp
--- a/Engine/Plugins/FX/Niagara/Source/Niagara/Private/NiagaraScript.cpp
+++ b/Engine/Plugins/FX/Niagara/Source/Niagara/Private/NiagaraScript.cpp
@@ -10,18 +10,37 @@
 #include "NiagaraModule.h"
 #include "NiagaraCustomVersion.h"
 #include "NiagaraShaderCompilationManager.h"
-
+#include "Serialization/MemoryReader.h"
 
 #include "Stats/Stats.h"
 #include "UObject/Linker.h"
+#include "HAL/PlatformFilemanager.h"
+#include "Misc/FileHelper.h"
 
 #if WITH_EDITOR
+	#include "NiagaraScriptDerivedData.h"
+	#include "DerivedDataCacheInterface.h"
 	#include "Interfaces/ITargetPlatform.h"
 #endif
 
+
+
+static int32 GbLogNiagaraSystemCompileResults = 1;
+static FAutoConsoleVariableRef CVarLogNiagaraSystemCompileResults(
+	TEXT("fx.LogNiagaraSystemCompileResults"),
+	GbLogNiagaraSystemCompileResults,
+	TEXT("If > 0 Niagara System scripts hlsl will be written to a text format when opened and closed in the editor. \n"),
+	ECVF_Default
+);
+
 DECLARE_STATS_GROUP(TEXT("Niagara Detailed"), STATGROUP_NiagaraDetailed, STATCAT_Advanced);
 
-FNiagaraScriptDebuggerInfo::FNiagaraScriptDebuggerInfo() : bRequestDebugFrame(false), DebugFrameLastWriteId(-1)
+FNiagaraScriptDebuggerInfo::FNiagaraScriptDebuggerInfo() : FrameLastWriteId(-1)
+{
+}
+
+
+FNiagaraScriptDebuggerInfo::FNiagaraScriptDebuggerInfo(FName InName, ENiagaraScriptUsage InUsage, const FGuid& InUsageId) : HandleName(InName), Usage(InUsage), UsageId(InUsageId), FrameLastWriteId(-1)
 {
 }
 
@@ -33,6 +52,133 @@ UNiagaraScriptSourceBase::UNiagaraScriptSourceBase(const FObjectInitializer& Obj
 }
 
 
+FNiagaraVMExecutableData::FNiagaraVMExecutableData() 
+	: NumUserPtrs(0)
+	, LastOpCount(0)
+	, LastCompileStatus(ENiagaraScriptCompileStatus::NCS_Unknown)
+	, bReadsAttributeData(false)
+	, CompileTime(0.0f)
+{
+}
+
+
+bool FNiagaraVMExecutableData::IsValid() const
+{
+	return LastCompileStatus != ENiagaraScriptCompileStatus::NCS_Unknown;
+}
+
+void FNiagaraVMExecutableData::Reset() 
+{
+	*this = FNiagaraVMExecutableData();
+}
+
+void FNiagaraVMExecutableData::SerializeData(FArchive& Ar, bool bDDCData)
+{
+	FNiagaraVMExecutableData::StaticStruct()->SerializeBin(Ar, this);
+}
+
+bool FNiagaraVMExecutableDataId::IsValid() const
+{
+	return BaseScriptID.IsValid();
+}
+
+void FNiagaraVMExecutableDataId::Invalidate()
+{
+	*this = FNiagaraVMExecutableDataId();
+}
+
+bool FNiagaraVMExecutableDataId::HasInterpolatedParameters() const
+{
+	return AdditionalDefines.Contains("InterpolatedSpawn");
+}
+
+bool FNiagaraVMExecutableDataId::RequiresPersistentIDs() const
+{
+	return AdditionalDefines.Contains("RequiresPersistentIDs");
+}
+
+/**
+* Tests this set against another for equality, disregarding override settings.
+*
+* @param ReferenceSet	The set to compare against
+* @return				true if the sets are equal
+*/
+bool FNiagaraVMExecutableDataId::operator==(const FNiagaraVMExecutableDataId& ReferenceSet) const
+{
+	if (CompilerVersionID != ReferenceSet.CompilerVersionID ||
+		ScriptUsageType != ReferenceSet.ScriptUsageType || 
+		ScriptUsageTypeID != ReferenceSet.ScriptUsageTypeID ||
+		BaseScriptID != ReferenceSet.BaseScriptID)
+	{
+		return false;
+	}
+	
+	if (ReferencedDependencyIds.Num() != ReferenceSet.ReferencedDependencyIds.Num())
+	{
+		return false;
+	}
+
+	for (int32 RefFunctionIndex = 0; RefFunctionIndex < ReferenceSet.ReferencedDependencyIds.Num(); RefFunctionIndex++)
+	{
+		const FGuid& ReferenceGuid = ReferenceSet.ReferencedDependencyIds[RefFunctionIndex];
+
+		if (ReferencedDependencyIds[RefFunctionIndex] != ReferenceGuid)
+		{
+			return false;
+		}
+	}
+
+	if (AdditionalDefines.Num() != ReferenceSet.AdditionalDefines.Num())
+	{
+		return false;
+	}
+
+	for (int32 Idx = 0; Idx < ReferenceSet.AdditionalDefines.Num(); Idx++)
+	{
+		const FString& ReferenceStr = ReferenceSet.AdditionalDefines[Idx];
+
+		if (AdditionalDefines[Idx] != ReferenceStr)
+		{
+			return false;
+		}
+	}
+
+
+	return true;
+}
+
+void FNiagaraVMExecutableDataId::AppendKeyString(FString& KeyString) const
+{
+	KeyString += FString::Printf(TEXT("%d_"), (int32)ScriptUsageType);
+	KeyString += ScriptUsageTypeID.ToString();
+	KeyString += TEXT("_");
+	KeyString += CompilerVersionID.ToString();
+	KeyString += TEXT("_");
+	KeyString += BaseScriptID.ToString();
+	KeyString += TEXT("_");
+
+	for (int32 Idx = 0; Idx < AdditionalDefines.Num(); Idx++)
+	{
+		KeyString += AdditionalDefines[Idx];
+
+		if (Idx < AdditionalDefines.Num() - 1)
+		{
+			KeyString += TEXT("_");
+		}
+	}
+	
+	// Add any referenced functions to the key so that we will recompile when they are changed
+	for (int32 FunctionIndex = 0; FunctionIndex < ReferencedDependencyIds.Num(); FunctionIndex++)
+	{
+		KeyString += ReferencedDependencyIds[FunctionIndex].ToString();
+
+		if (FunctionIndex < ReferencedDependencyIds.Num() - 1)
+		{
+			KeyString += TEXT("_");
+		}
+	}	
+}
+
 UNiagaraScript::UNiagaraScript(const FObjectInitializer& ObjectInitializer)
 	: Super(ObjectInitializer)
 	, Usage(ENiagaraScriptUsage::Function)
@@ -40,14 +186,12 @@ UNiagaraScript::UNiagaraScript(const FObjectInitializer& ObjectInitializer)
 	, UsageIndex_DEPRECATED(0)
 #endif
 	, ModuleUsageBitmask( (1 << (int32)ENiagaraScriptUsage::ParticleSpawnScript) | (1 << (int32)ENiagaraScriptUsage::ParticleSpawnScriptInterpolated) | (1 << (int32)ENiagaraScriptUsage::ParticleUpdateScript) | (1 << (int32)ENiagaraScriptUsage::ParticleEventScript) )
-	, NumUserPtrs(0)
 	, NumericOutputTypeSelectionMode(ENiagaraNumericOutputTypeSelectionMode::Largest)
-	, LastCompileStatus(ENiagaraScriptCompileStatus::NCS_Unknown)
 {
 #if WITH_EDITORONLY_DATA
 	ScriptResource.OnCompilationComplete().AddUniqueDynamic(this, &UNiagaraScript::OnCompilationComplete);
 
-	RapidIterationParameters.Name = *GetFullName();
+	RapidIterationParameters.DebugName = *GetFullName();
 #endif	
 }
 
@@ -55,6 +199,35 @@ UNiagaraScript::~UNiagaraScript()
 {
 }
 
+#if WITH_EDITORONLY_DATA
+void UNiagaraScript::ComputeVMCompilationId(FNiagaraVMExecutableDataId& Id) const
+{
+	Id = FNiagaraVMExecutableDataId();
+	
+	// Ideally we wouldn't want to do this but rather than push the data down
+	// from the emitter.
+	UObject* Obj = GetOuter();
+	UNiagaraEmitter* Emitter = Cast<UNiagaraEmitter>(Obj);
+	if (Emitter)
+	{
+		if ((Emitter->bInterpolatedSpawning && Usage == ENiagaraScriptUsage::ParticleGPUComputeScript) || 
+			(Emitter->bInterpolatedSpawning && Usage == ENiagaraScriptUsage::ParticleSpawnScript) ||
+			Usage == ENiagaraScriptUsage::ParticleSpawnScriptInterpolated)
+		{
+			Id.AdditionalDefines.Add(TEXT("InterpolatedSpawn"));
+		}
+		if (Emitter->RequiresPersistantIDs())
+		{
+			Id.AdditionalDefines.Add(TEXT("RequiresPersistentIDs"));
+		}
+	}
+
+	Source->ComputeVMCompilationId(Id, Usage, UsageId);
+	
+	LastGeneratedVMId = Id;
+}
+#endif
+
 bool UNiagaraScript::ContainsUsage(ENiagaraScriptUsage InUsage) const
 {
 	if (IsEquivalentUsage(InUsage))
@@ -62,6 +235,11 @@ bool UNiagaraScript::ContainsUsage(ENiagaraScriptUsage InUsage) const
 		return true;
 	}
 
+	if (Usage == ENiagaraScriptUsage::ParticleGPUComputeScript && IsParticleScript(InUsage))
+	{
+		return true;
+	}
+
 	if (InUsage == ENiagaraScriptUsage::ParticleUpdateScript && Usage == ENiagaraScriptUsage::ParticleSpawnScriptInterpolated)
 	{
 		return true;
@@ -80,13 +258,38 @@ bool UNiagaraScript::ContainsUsage(ENiagaraScriptUsage InUsage) const
 	return false;
 }
 
+FNiagaraScriptExecutionParameterStore* UNiagaraScript::GetExecutionReadyParameterStore(ENiagaraSimTarget SimTarget)
+{
+	if (SimTarget == ENiagaraSimTarget::CPUSim && IsReadyToRun(ENiagaraSimTarget::CPUSim))
+	{
+		if (ScriptExecutionParamStoreCPU.IsInitialized() == false)
+		{
+			ScriptExecutionParamStoreCPU.InitFromOwningScript(this, SimTarget, false);
+		}
+		return &ScriptExecutionParamStoreCPU;
+	}
+	else if (SimTarget == ENiagaraSimTarget::GPUComputeSim)
+	{
+		if (ScriptExecutionParamStoreGPU.IsInitialized() == false)
+		{
+			ScriptExecutionParamStoreGPU.InitFromOwningScript(this, SimTarget, false);
+		}
+		return &ScriptExecutionParamStoreGPU;
+	}
+	else
+	{
+		return nullptr;
+	}
+}
+
+
 void UNiagaraScript::Serialize(FArchive& Ar)
 {
 	Super::Serialize(Ar);
 
 	Ar.UsingCustomVersion(FNiagaraCustomVersion::GUID);		// only changes version if not loading
 	const int32 NiagaraVer = Ar.CustomVer(FNiagaraCustomVersion::GUID);
-
+	
 	bool IsValidShaderScript = false;
 	if (NiagaraVer < FNiagaraCustomVersion::DontCompileGPUWhenNotNeeded)
 	{
@@ -94,6 +297,10 @@ void UNiagaraScript::Serialize(FArchive& Ar)
 			&& (NiagaraVer < FNiagaraCustomVersion::NiagaraShaderMapCooking2 || (Usage != ENiagaraScriptUsage::SystemSpawnScript && Usage != ENiagaraScriptUsage::SystemUpdateScript))
 			&& (NiagaraVer < FNiagaraCustomVersion::NiagaraCombinedGPUSpawnUpdate || (Usage != ENiagaraScriptUsage::ParticleUpdateScript && Usage != ENiagaraScriptUsage::EmitterSpawnScript && Usage != ENiagaraScriptUsage::EmitterUpdateScript));
 	}
+	else if (NiagaraVer < FNiagaraCustomVersion::MovedToDerivedDataCache)
+	{
+		IsValidShaderScript = LegacyCanBeRunOnGpu();
+	}
 	else
 	{
 		IsValidShaderScript = CanBeRunOnGpu();
@@ -111,6 +318,66 @@ void UNiagaraScript::Serialize(FArchive& Ar)
 	}
 }
 
+/** Is usage A dependent on Usage B?*/
+bool UNiagaraScript::IsUsageDependentOn(ENiagaraScriptUsage InUsageA, ENiagaraScriptUsage InUsageB)
+{
+	if (InUsageA == InUsageB)
+	{
+		return false;
+	}
+
+	// Usages of the same phase are interdependent because we copy the attributes from one to the other and if those got 
+	// out of sync, there could be problems.
+
+	if ((InUsageA == ENiagaraScriptUsage::ParticleSpawnScript || InUsageA == ENiagaraScriptUsage::ParticleSpawnScriptInterpolated || InUsageA == ENiagaraScriptUsage::ParticleUpdateScript || InUsageA == ENiagaraScriptUsage::ParticleEventScript)
+		&& (InUsageB == ENiagaraScriptUsage::ParticleSpawnScript || InUsageB == ENiagaraScriptUsage::ParticleSpawnScriptInterpolated || InUsageB == ENiagaraScriptUsage::ParticleUpdateScript || InUsageB == ENiagaraScriptUsage::ParticleEventScript))
+	{
+		return true;
+	}
+
+	// The GPU compute script is always dependent on the other particle scripts.
+	if ((InUsageA == ENiagaraScriptUsage::ParticleGPUComputeScript)
+		&& (InUsageB == ENiagaraScriptUsage::ParticleSpawnScript || InUsageB == ENiagaraScriptUsage::ParticleSpawnScriptInterpolated || InUsageB == ENiagaraScriptUsage::ParticleUpdateScript || InUsageB == ENiagaraScriptUsage::ParticleEventScript))
+	{
+		return true;
+	}
+
+	if ((InUsageA == ENiagaraScriptUsage::EmitterSpawnScript || InUsageA == ENiagaraScriptUsage::EmitterUpdateScript)
+		&& (InUsageB == ENiagaraScriptUsage::EmitterSpawnScript || InUsageB == ENiagaraScriptUsage::EmitterUpdateScript))
+	{
+		return true;
+	}
+
+	if ((InUsageA == ENiagaraScriptUsage::SystemSpawnScript || InUsageA == ENiagaraScriptUsage::SystemUpdateScript)
+		&& (InUsageB == ENiagaraScriptUsage::SystemSpawnScript || InUsageB == ENiagaraScriptUsage::SystemUpdateScript))
+	{
+		return true;
+	}
+
+	return false;
+}
+
+bool UNiagaraScript::ConvertUsageToGroup(ENiagaraScriptUsage InUsage, ENiagaraScriptGroup& OutGroup)
+{
+	if (IsParticleScript(InUsage) || IsStandaloneScript(InUsage))
+	{
+		OutGroup = ENiagaraScriptGroup::Particle;
+		return true;
+	}
+	else if (IsEmitterSpawnScript(InUsage) || IsEmitterUpdateScript(InUsage))
+	{
+		OutGroup = ENiagaraScriptGroup::Emitter;
+		return true;
+	}
+	else if (IsSystemSpawnScript(InUsage) || IsSystemUpdateScript(InUsage))
+	{
+		OutGroup = ENiagaraScriptGroup::System;
+		return true;
+	}
+
+	return false;
+}
+
 void UNiagaraScript::PostLoad()
 {
 	Super::PostLoad();
@@ -118,22 +385,9 @@ void UNiagaraScript::PostLoad()
 	bool bNeedsRecompile = false;
 	const int32 NiagaraVer = GetLinkerCustomVersion(FNiagaraCustomVersion::GUID);
 	
-	//Recompile is before the latest version that requires a script recompile. Some may not.
-	if (NiagaraVer < FNiagaraCustomVersion::RemovalOfNiagaraVariableIDs)
-	{
-		if (Usage != ENiagaraScriptUsage::Module && Usage != ENiagaraScriptUsage::Function && Usage != ENiagaraScriptUsage::DynamicInput)
-		{
-			bNeedsRecompile = true;
-			InvalidateScript();
-			UE_LOG(LogNiagara, Log, TEXT("Niagara script is out of date and requires recompile to be used! %s"), *GetFullName());
-		}
-	}
-
-
 	// Resources can be processed / registered now that we're back on the main thread
 	ProcessSerializedShaderMaps(this, ScriptResource, ScriptResourcesByFeatureLevel);
 
-
 	if (GIsEditor)
 	{
 		
@@ -145,22 +399,11 @@ void UNiagaraScript::PostLoad()
 			Source->ConditionalPostLoad();
 		}
 
-		// If we've never computed a change id or the source graph is different than the compiled script change id,
-		// we're out of sync and should recompile.
-		if (Usage != ENiagaraScriptUsage::Module && Usage != ENiagaraScriptUsage::Function && Usage != ENiagaraScriptUsage::DynamicInput)
-		{
-			if (ChangeId.IsValid() == false || (Source != nullptr && !Source->IsSynchronized(ChangeId)))
-			{
-				bNeedsRecompile = true;
-				InvalidateScript();//clear out the script.
-				UE_LOG(LogNiagara, Log, TEXT("Niagara script is out of date with source graph requires recompile to be used! %s"), *GetFullName());
-			}
-		}
 #endif
 	}
 
 	// for now, force recompile until we can be sure everything is working
-//	bNeedsRecompile = true;
+	//bNeedsRecompile = true;
 #if WITH_EDITORONLY_DATA
 	CacheResourceShadersForRendering(false, bNeedsRecompile);
 #endif
@@ -171,14 +414,35 @@ void UNiagaraScript::PostLoad()
 }
 
 
+bool UNiagaraScript::IsReadyToRun(ENiagaraSimTarget SimTarget) const
+{
+	if (SimTarget == ENiagaraSimTarget::CPUSim)
+	{
+		if (CachedScriptVM.IsValid())
+		{
+			return true;
+		}
+	}
+	else if (SimTarget == ENiagaraSimTarget::GPUComputeSim)
+	{
+		return CanBeRunOnGpu();
+	}
+
+	return false;
+}
+
+
 
 #if STATS
 void UNiagaraScript::GenerateStatScopeIDs()
 {
 	StatScopesIDs.Empty();
-	for (FNiagaraStatScope& StatScope : StatScopes)
+	if (IsReadyToRun(ENiagaraSimTarget::CPUSim))
 	{
-		StatScopesIDs.Add(FDynamicStats::CreateStatId<FStatGroup_STATGROUP_NiagaraDetailed>(StatScope.FriendlyName.ToString()));
+		for (FNiagaraStatScope& StatScope : CachedScriptVM.StatScopes)
+		{
+			StatScopesIDs.Add(FDynamicStats::CreateStatId<FStatGroup_STATGROUP_NiagaraDetailed>(StatScope.FriendlyName.ToString()));
+		}
 	}
 }
 #endif
@@ -197,7 +461,36 @@ bool UNiagaraScript::AreScriptAndSourceSynchronized() const
 {
 	if (Source)
 	{
-		return Source->IsSynchronized(ChangeId);
+		FNiagaraVMExecutableDataId NewId;
+		ComputeVMCompilationId(NewId);
+		bool bSynchronized = (NewId.IsValid() && NewId == CachedScriptVMId);
+		if (!bSynchronized && NewId.IsValid() && CachedScriptVMId.IsValid() && CachedScriptVM.IsValid())
+		{
+			if (NewId != LastReportedVMId)
+			{
+				if (NewId.BaseScriptID != CachedScriptVMId.BaseScriptID)
+				{
+					UE_LOG(LogNiagara, Log, TEXT("AreScriptAndSourceSynchronized base script id's don't match. %s != %s"), *NewId.BaseScriptID.ToString(), *CachedScriptVMId.BaseScriptID.ToString());
+				}
+				if (NewId.ReferencedDependencyIds.Num() != CachedScriptVMId.ReferencedDependencyIds.Num())
+				{
+					UE_LOG(LogNiagara, Log, TEXT("AreScriptAndSourceSynchronized num dependencies don't match. %d != %d"), NewId.ReferencedDependencyIds.Num(), CachedScriptVMId.ReferencedDependencyIds.Num());
+				}
+				else
+				{
+					for (int32 i = 0; i < NewId.ReferencedDependencyIds.Num(); i++)
+					{
+						if (NewId.ReferencedDependencyIds[i] != CachedScriptVMId.ReferencedDependencyIds[i])
+						{
+							UE_LOG(LogNiagara, Log, TEXT("AreScriptAndSourceSynchronized reference id %d doesn't match. %s != %s, source %s"), i, *NewId.ReferencedDependencyIds[i].ToString(), *CachedScriptVMId.ReferencedDependencyIds[i].ToString(),
+								NewId.ReferencedObjects[i] != nullptr ? *NewId.ReferencedObjects[i]->GetPathName() : TEXT("nullptr"));
+						}
+					}
+				}
+				LastReportedVMId = NewId;
+			}
+		}
+		return bSynchronized;
 	}
 	else
 	{
@@ -205,27 +498,92 @@ bool UNiagaraScript::AreScriptAndSourceSynchronized() const
 	}
 }
 
-void UNiagaraScript::MarkScriptAndSourceDesynchronized()
+void UNiagaraScript::MarkScriptAndSourceDesynchronized(FString Reason)
 {
 	if (Source)
 	{
-		Source->MarkNotSynchronized();
+		Source->MarkNotSynchronized(Reason);
 	}
 }
 
-void UNiagaraScript::SetByteCode(const TArray<uint8>& InByteCode)
+bool UNiagaraScript::HandleVariableRenames(const TMap<FNiagaraVariable, FNiagaraVariable>& OldToNewVars, const FString& UniqueEmitterName)
 {
-	ByteCode = InByteCode;
-	if (InByteCode.Num() == 0)
+	bool bConvertedAnything = false;
+	auto Iter = OldToNewVars.CreateConstIterator();
+	while (Iter)
 	{
-		LastCompileStatus = ENiagaraScriptCompileStatus::NCS_Error;
+		// Sometimes the script is under the generic name, other times it has been converted to the unique emitter name. Handle both cases below...
+		FNiagaraVariable RISrcVarA = FNiagaraUtilities::ConvertVariableToRapidIterationConstantName(Iter->Key, !UniqueEmitterName.IsEmpty() ? TEXT("Emitter") : nullptr , GetUsage());
+		FNiagaraVariable RISrcVarB = FNiagaraUtilities::ConvertVariableToRapidIterationConstantName(Iter->Key, !UniqueEmitterName.IsEmpty() ? *UniqueEmitterName : nullptr, GetUsage());
+		FNiagaraVariable RIDestVarA = FNiagaraUtilities::ConvertVariableToRapidIterationConstantName(Iter->Value, !UniqueEmitterName.IsEmpty() ? TEXT("Emitter") : nullptr, GetUsage());
+		FNiagaraVariable RIDestVarB = FNiagaraUtilities::ConvertVariableToRapidIterationConstantName(Iter->Value, !UniqueEmitterName.IsEmpty() ? *UniqueEmitterName : nullptr, GetUsage());
+
+		{
+			if (nullptr != RapidIterationParameters.FindParameterOffset(RISrcVarA))
+			{
+				RapidIterationParameters.RenameParameter(RISrcVarA, RIDestVarA.GetName());
+				UE_LOG(LogNiagara, Log, TEXT("Converted RI variable \"%s\" to \"%s\" in Script \"%s\""), *RISrcVarA.GetName().ToString(), *RIDestVarA.GetName().ToString(), *GetFullName());
+				bConvertedAnything = true;
+			}
+			else if (nullptr != RapidIterationParameters.FindParameterOffset(RISrcVarB))
+			{
+				RapidIterationParameters.RenameParameter(RISrcVarB, RIDestVarB.GetName());
+				UE_LOG(LogNiagara, Log, TEXT("Converted RI variable \"%s\" to \"%s\" in Script \"%s\""), *RISrcVarB.GetName().ToString(), *RIDestVarB.GetName().ToString(), *GetFullName());
+				bConvertedAnything = true;
+			}
+		}
+
+		{
+			// Go ahead and convert the stored VM executable data too. I'm not 100% sure why this is necessary, since we should be recompiling.
+			int32 VarIdx = GetVMExecutableData().Parameters.Parameters.IndexOfByKey(RISrcVarA);
+			if (VarIdx != INDEX_NONE)
+			{
+				GetVMExecutableData().Parameters.Parameters[VarIdx].SetName(RIDestVarA.GetName());
+				UE_LOG(LogNiagara, Log, TEXT("Converted exec param variable \"%s\" to \"%s\" in Script \"%s\""), *RISrcVarA.GetName().ToString(), *RIDestVarA.GetName().ToString(), *GetFullName());
+				bConvertedAnything = true;
+			}
+
+			VarIdx = GetVMExecutableData().Parameters.Parameters.IndexOfByKey(RISrcVarB);
+			if (VarIdx != INDEX_NONE)
+			{
+				GetVMExecutableData().Parameters.Parameters[VarIdx].SetName(RIDestVarB.GetName());
+				UE_LOG(LogNiagara, Log, TEXT("Converted exec param  variable \"%s\" to \"%s\" in Script \"%s\""), *RISrcVarB.GetName().ToString(), *RIDestVarB.GetName().ToString(), *GetFullName());
+				bConvertedAnything = true;
+			}
+		}
+
+		{
+			// Also handle any data set mappings...
+			auto DS2PIterator = GetVMExecutableData().DataSetToParameters.CreateIterator();
+			while (DS2PIterator)
+			{
+				for (int32 i = 0; i < DS2PIterator.Value().Parameters.Num(); i++)
+				{
+					FNiagaraVariable Var = DS2PIterator.Value().Parameters[i];
+					if (Var == RISrcVarA)
+					{
+						DS2PIterator.Value().Parameters[i].SetName(RIDestVarA.GetName());
+						bConvertedAnything = true;
+					}
+					else if (Var == RISrcVarB)
+					{
+						DS2PIterator.Value().Parameters[i].SetName(RIDestVarB.GetName());
+						bConvertedAnything = true;
+					}
+				}
+				++DS2PIterator;
+			}
+		}
+		++Iter;
 	}
-	else
+
+	if (bConvertedAnything)
 	{
-		LastCompileStatus = ENiagaraScriptCompileStatus::NCS_UpToDate;
+		InvalidateExecutionReadyParameterStores();
 	}
-}
 
+	return bConvertedAnything;
+}
 
 UNiagaraScript* UNiagaraScript::MakeRecursiveDeepCopy(UObject* DestOuter, TMap<const UObject*, UObject*>& ExistingConversions) const
 {
@@ -258,7 +616,7 @@ UNiagaraScript* UNiagaraScript::MakeRecursiveDeepCopy(UObject* DestOuter, TMap<c
 	ObjParameters.InternalFlagMask = InternalFlagsMask;
 	ObjParameters.DuplicateMode = DuplicateMode;
 	
-	// Make sure that we don't duplicate objects that we've alredy converted...
+	// Make sure that we don't duplicate objects that we've already converted...
 	TMap<const UObject*, UObject*>::TConstIterator It = ExistingConversions.CreateConstIterator();
 	while (It)
 	{
@@ -298,28 +656,248 @@ void UNiagaraScript::SubsumeExternalDependencies(TMap<const UObject*, UObject*>&
 	Source->SubsumeExternalDependencies(ExistingConversions);
 }
 
-ENiagaraScriptCompileStatus UNiagaraScript::Compile(FString& OutGraphLevelErrorMessages, bool bForce)
+void WriteTextFileToDisk(FString SaveDirectory, FString FileName, FString TextToSave, bool bAllowOverwriting)
+{
+	IPlatformFile& PlatformFile = FPlatformFileManager::Get().GetPlatformFile();
+
+	// CreateDirectoryTree returns true if the destination
+	// directory existed prior to call or has been created
+	// during the call.
+	if (PlatformFile.CreateDirectoryTree(*SaveDirectory))
+	{
+		// Get absolute file path
+		FString AbsoluteFilePath = SaveDirectory + "/" + FileName;
+
+		// Allow overwriting or file doesn't already exist
+		if (bAllowOverwriting || !PlatformFile.FileExists(*AbsoluteFilePath))
+		{
+			if (FFileHelper::SaveStringToFile(TextToSave, *AbsoluteFilePath))
+			{
+				UE_LOG(LogNiagara, Log, TEXT("Wrote file to %s"), *AbsoluteFilePath);
+				return;
+			}
+
+		}
+	}
+}
+
+void UNiagaraScript::SetVMCompilationResults(const FNiagaraVMExecutableDataId& InCompileId, FNiagaraVMExecutableData& InScriptVM, FNiagaraCompileRequestDataBase* InRequestData)
 {
-	if (Source == nullptr)
-		return ENiagaraScriptCompileStatus::NCS_Error;
+	check(InRequestData != nullptr);
 
-	if (Source->GetChangeID() == ChangeId && ChangeId.IsValid() && !bForce)
+	CachedScriptVMId = InCompileId;
+	CachedScriptVM = InScriptVM;
+	CachedParameterCollectionReferences.Empty();
+	
+	if (CachedScriptVM.LastCompileStatus == ENiagaraScriptCompileStatus::NCS_Error)
+	{
+		UE_LOG(LogNiagara, Error, TEXT("%s"), *CachedScriptVM.ErrorMsg);
+	}
+	else if (CachedScriptVM.LastCompileStatus == ENiagaraScriptCompileStatus::NCS_UpToDateWithWarnings)
 	{
-		UE_LOG(LogNiagara, Log, TEXT("Skipping compilation due to changeid's being in sync already.. path: %s change-id: %s"), *GetFullName(), *ChangeId.ToString())
-		return LastCompileStatus;
+		UE_LOG(LogNiagara, Warning, TEXT("%s"), *CachedScriptVM.ErrorMsg);
+	}
+
+	// The compilation process only references via soft references any parameter collections. This resolves those 
+	// soft references to real references.
+	for (FString& Path : CachedScriptVM.ParameterCollectionPaths)
+	{
+		FSoftObjectPath SoftPath(Path);
+		UObject* Obj = SoftPath.TryLoad();
+		UNiagaraParameterCollection* ParamCollection = Cast<UNiagaraParameterCollection>(Obj);
+		if (ParamCollection != nullptr)
+		{
+			CachedParameterCollectionReferences.Add(ParamCollection);
+		}
 	}
 
-	float StartTime = FPlatformTime::Seconds();
-	ENiagaraScriptCompileStatus CPUCompileStatus = Source->Compile(this, OutGraphLevelErrorMessages);
-	CacheResourceShadersForRendering(false, true);
-	float DeltaTime = FPlatformTime::Seconds() - StartTime;
-	UE_LOG(LogNiagara, Log, TEXT("Compiling %s took %f sec."), *GetFullName(), DeltaTime);
+	CachedDefaultDataInterfaces.Empty(CachedScriptVM.DataInterfaceInfo.Num());
+	for (FNiagaraScriptDataInterfaceCompileInfo Info : CachedScriptVM.DataInterfaceInfo)
+	{
+		int32 Idx = CachedDefaultDataInterfaces.AddDefaulted();
+		CachedDefaultDataInterfaces[Idx].UserPtrIdx = Info.UserPtrIdx;
+		CachedDefaultDataInterfaces[Idx].Name = Info.Name;
+		CachedDefaultDataInterfaces[Idx].Type = Info.Type;
+		CachedDefaultDataInterfaces[Idx].RegisteredParameterMapRead = InRequestData->ResolveEmitterAlias(Info.RegisteredParameterMapRead);
+		CachedDefaultDataInterfaces[Idx].RegisteredParameterMapWrite = InRequestData->ResolveEmitterAlias(Info.RegisteredParameterMapWrite);
+
+		// We compiled it just a bit ago, so we should be able to resolve it from the table that we passed in.
+		UNiagaraDataInterface*const* FindDIById = InRequestData->GetObjectNameMap().Find(Info.Name);
+		if (FindDIById != nullptr && *(FindDIById) != nullptr)
+		{
+			CachedDefaultDataInterfaces[Idx].DataInterface = DuplicateObject<UNiagaraDataInterface>(*(FindDIById), this);
+			check(CachedDefaultDataInterfaces[Idx].DataInterface != nullptr);
+		}			
+		
+		if (CachedDefaultDataInterfaces[Idx].DataInterface == nullptr)
+		{
+			// Use the CDO since we didn't have a default..
+			UObject* Obj = const_cast<UClass*>(Info.Type.GetClass())->GetDefaultObject(true);
+			CachedDefaultDataInterfaces[Idx].DataInterface = Cast<UNiagaraDataInterface>(DuplicateObject(Obj, this));
 
-	return CPUCompileStatus;
+			if (Info.bIsPlaceholder == false)
+			{
+				UE_LOG(LogNiagara, Error, TEXT("We somehow ended up with a data interface that we couldn't match post compile. This shouldn't happen. Creating a dummy to prevent crashes. %s"), *Info.Name.ToString());
+			}
+		}
+		check(CachedDefaultDataInterfaces[Idx].DataInterface != nullptr);
+	}
+
+	GenerateStatScopeIDs();
+
+	// Now go ahead and trigger the GPU script compile now that we have a compiled GPU hlsl script.
+	if (Usage == ENiagaraScriptUsage::ParticleGPUComputeScript)
+	{
+		CacheResourceShadersForRendering(false, true);
+	}
+
+	InvalidateExecutionReadyParameterStores();
+
+	if (GbLogNiagaraSystemCompileResults > 0)
+	{
+		FString ExportText = CachedScriptVM.LastHlslTranslation;
+		FString ExportTextAsm = CachedScriptVM.LastAssemblyTranslation;
+		if (IsGPUScript(GetUsage()))
+		{
+			ExportText = CachedScriptVM.LastHlslTranslationGPU;
+			ExportTextAsm = "";
+		}
+		FString ExportTextParams;
+		for (const FNiagaraVariable& Var : CachedScriptVM.Parameters.Parameters)
+		{
+			ExportTextParams += Var.ToString();
+			ExportTextParams += "\n";
+		}
+
+		FString ExportRIParams = RapidIterationParameters.ToString();
+		
+		UEnum* UsageEnum = Cast<UEnum>(FindObject<UEnum>(ANY_PACKAGE, TEXT("ENiagaraScriptUsage"), true));
+		FString TypeName = UsageEnum->GetNameStringByValue((int64)GetUsage());
+		FString FilePath = GetOutermost()->FileName.ToString();
+		FString PathPart, FilenamePart, ExtensionPart;
+		FPaths::Split(FilePath, PathPart, FilenamePart, ExtensionPart);
+		WriteTextFileToDisk(FPaths::ProjectLogDir(), FilenamePart + TEXT(".")+ TypeName + TEXT(".hlsl"), ExportText, true);
+		WriteTextFileToDisk(FPaths::ProjectLogDir(), FilenamePart + TEXT(".")+ TypeName + TEXT(".asm"), ExportTextAsm, true);
+		WriteTextFileToDisk(FPaths::ProjectLogDir(), FilenamePart + TEXT(".") + TypeName + TEXT(".params"), ExportTextParams, true);
+		WriteTextFileToDisk(FPaths::ProjectLogDir(), FilenamePart + TEXT(".") + TypeName + TEXT(".riparams"), ExportRIParams, true);
+	}
+
+	OnVMScriptCompiled().Broadcast(this);
 }
 
-#endif
+void UNiagaraScript::InvalidateExecutionReadyParameterStores()
+{
+	// Make sure that we regenerate any parameter stores, since they must be kept in sync with the layout from script compilation.
+	ScriptExecutionParamStoreCPU.Clear();
+	ScriptExecutionParamStoreGPU.Clear();
+}
+
+void UNiagaraScript::InvalidateCachedCompileIds()
+{
+	GetSource()->InvalidateCachedCompileIds();
+}
+
+void UNiagaraScript::RequestCompile()
+{
+	if (!AreScriptAndSourceSynchronized())
+	{
+		if (IsCompilable() == false)
+		{
+			CachedScriptVM.LastCompileStatus = ENiagaraScriptCompileStatus::NCS_Unknown;
+			CachedScriptVMId = LastGeneratedVMId;
+			return;
+		}
+
+		CachedScriptVM.LastCompileStatus = ENiagaraScriptCompileStatus::NCS_BeingCreated;
+
+		TArray<TSharedPtr<FNiagaraCompileRequestDataBase, ESPMode::ThreadSafe>> DependentRequests;
+		TArray<uint8> OutData;
+		INiagaraModule& NiagaraModule = FModuleManager::Get().LoadModuleChecked<INiagaraModule>(TEXT("Niagara"));
+		TSharedPtr<FNiagaraCompileRequestDataBase, ESPMode::ThreadSafe> RequestData = NiagaraModule.Precompile(this);
+
+		ActiveCompileRoots.Empty();
+		RequestData->GetReferencedObjects(ActiveCompileRoots);
+
+		FNiagaraCompileOptions Options(GetUsage(), GetUsageId(), ModuleUsageBitmask, GetPathName(), GetFullName(), GetName());
+
+		FNiagaraScriptDerivedData* CompileTask = new FNiagaraScriptDerivedData(GetFullName(), RequestData, Options, LastGeneratedVMId, false);
+
+		// For debugging DDC/Compression issues		
+		const bool bSkipDDC = false;
+		if (bSkipDDC)
+		{
+			CompileTask->Build(OutData);
+
+			delete CompileTask;
+			CompileTask = nullptr;
+		}
+		else
+		{
+			if (CompileTask->CanBuild())
+			{
+				GetDerivedDataCacheRef().GetSynchronous(CompileTask, OutData);
+				// Assume that once given over to the derived cache, the compile task is going to be killed by it.
+				CompileTask = nullptr;
+			}
+			else
+			{
+				delete CompileTask;
+				CompileTask = nullptr;
+			}
+		}
+
+		if (OutData.Num() > 0)
+		{
+			FNiagaraVMExecutableData ExeData;
+			FNiagaraScriptDerivedData::BinaryToExecData(OutData, ExeData);
+			SetVMCompilationResults(LastGeneratedVMId, ExeData, RequestData.Get());
+		}
+		else
+		{
+			check(false);
+		}
 
+		ActiveCompileRoots.Empty();
+	}
+	else
+	{
+		UE_LOG(LogNiagara, Log, TEXT("Script '%s' is in-sync skipping compile.."), *GetFullName());
+	}
+}
+
+bool UNiagaraScript::RequestExternallyManagedAsyncCompile(const TSharedPtr<FNiagaraCompileRequestDataBase, ESPMode::ThreadSafe>& RequestData, FNiagaraVMExecutableDataId& OutCompileId, uint32& OutAsyncHandle, bool bTrulyAsync)
+{
+	if (!AreScriptAndSourceSynchronized())
+	{
+		if (IsCompilable() == false)
+		{
+			OutCompileId = LastGeneratedVMId;
+			OutAsyncHandle = (uint32)INDEX_NONE;
+			CachedScriptVM.LastCompileStatus = ENiagaraScriptCompileStatus::NCS_Unknown;
+			CachedScriptVMId = LastGeneratedVMId;
+			return false;
+		}
+
+		CachedScriptVM.LastCompileStatus = ENiagaraScriptCompileStatus::NCS_BeingCreated;
+
+		OutCompileId = LastGeneratedVMId;
+		FNiagaraCompileOptions Options(GetUsage(), GetUsageId(), ModuleUsageBitmask, GetPathName(), GetFullName(), GetName());
+		FNiagaraScriptDerivedData* CompileTask = new FNiagaraScriptDerivedData(GetFullName(), RequestData, Options, LastGeneratedVMId, bTrulyAsync);
+	
+		check(CompileTask->CanBuild());
+		OutAsyncHandle = GetDerivedDataCacheRef().GetAsynchronous(CompileTask);
+
+		return true;
+	}
+	else
+	{
+		OutCompileId = LastGeneratedVMId;
+		OutAsyncHandle = (uint32)INDEX_NONE;
+		UE_LOG(LogNiagara, Log, TEXT("Script '%s' is in-sync skipping compile.."), *GetFullName());
+		return false;
+	}
+}
+#endif
 
 void UNiagaraScript::OnCompilationComplete()
 {
@@ -328,6 +906,23 @@ void UNiagaraScript::OnCompilationComplete()
 #endif
 }
 
+void UNiagaraScript::GetAssetRegistryTags(TArray<FAssetRegistryTag>& OutTags) const
+{
+	Super::GetAssetRegistryTags(OutTags);
+#if WITH_EDITORONLY_DATA
+	FName TagName = GET_MEMBER_NAME_CHECKED(UNiagaraScript, ProvidedDependencies);
+	FString DependenciesProvidedString;
+	for (FName DependencyProvided : ProvidedDependencies)
+	{
+		DependenciesProvidedString.Append(DependencyProvided.ToString() + ",");
+	}
+	if (ProvidedDependencies.Num() > 0)
+	{
+		OutTags.Add(FAssetRegistryTag(TagName, DependenciesProvidedString, UObject::FAssetRegistryTag::TT_Hidden));
+	}
+#endif
+}
+
 #if WITH_EDITOR
 void UNiagaraScript::BeginCacheForCookedPlatformData(const ITargetPlatform *TargetPlatform)
 {
@@ -336,12 +931,12 @@ void UNiagaraScript::BeginCacheForCookedPlatformData(const ITargetPlatform *Targ
 		TArray<FName> DesiredShaderFormats;
 		TargetPlatform->GetAllTargetedShaderFormats(DesiredShaderFormats);
 
-		FNiagaraScript** CachedScriptResourceForPlatformPtr = CachedScriptResourcesForCooking.Find(TargetPlatform);
+		FNiagaraShaderScript** CachedScriptResourceForPlatformPtr = CachedScriptResourcesForCooking.Find(TargetPlatform);
 
 		if (CachedScriptResourceForPlatformPtr == nullptr)
 		{
 			CachedScriptResourcesForCooking.Add(TargetPlatform);
-			FNiagaraScript* &CachedScriptResourceForPlatform = *CachedScriptResourcesForCooking.Find(TargetPlatform);
+			FNiagaraShaderScript* &CachedScriptResourceForPlatform = *CachedScriptResourcesForCooking.Find(TargetPlatform);
 
 			if (DesiredShaderFormats.Num())
 			{
@@ -352,7 +947,7 @@ void UNiagaraScript::BeginCacheForCookedPlatformData(const ITargetPlatform *Targ
 					{
 						const EShaderPlatform LegacyShaderPlatform = ShaderFormatToLegacyShaderPlatform(DesiredShaderFormats[FormatIndex]);
 
-						// Begin caching shaders for the target platform and store the FNiagaraScript being compiled into CachedScriptResourcesForCooking
+						// Begin caching shaders for the target platform and store the FNiagaraShaderScript being compiled into CachedScriptResourcesForCooking
 						CacheResourceShadersForCooking(LegacyShaderPlatform, CachedScriptResourceForPlatform);
 					}
 				}
@@ -361,19 +956,21 @@ void UNiagaraScript::BeginCacheForCookedPlatformData(const ITargetPlatform *Targ
 	}
 }
 
-
-void UNiagaraScript::CacheResourceShadersForCooking(EShaderPlatform ShaderPlatform, FNiagaraScript* &OutCachedResource)
+void UNiagaraScript::CacheResourceShadersForCooking(EShaderPlatform ShaderPlatform, FNiagaraShaderScript* &OutCachedResource)
 {
 	if (CanBeRunOnGpu())
 	{
 		// spawn and update are combined on GPU, so we only compile spawn scripts
-		if (Usage == ENiagaraScriptUsage::ParticleSpawnScript || Usage == ENiagaraScriptUsage::ParticleSpawnScriptInterpolated)
+		if (Usage == ENiagaraScriptUsage::ParticleGPUComputeScript)
 		{
-			FNiagaraScript *ResourceToCache = nullptr;
+			FNiagaraShaderScript *ResourceToCache = nullptr;
 			ERHIFeatureLevel::Type TargetFeatureLevel = GetMaxSupportedFeatureLevel(ShaderPlatform);
 
-			FNiagaraScript* NewResource = AllocateResource();
-			NewResource->SetScript(this, (ERHIFeatureLevel::Type)TargetFeatureLevel, UniqueID, GetName());
+			FNiagaraShaderScript* NewResource = AllocateResource();
+			check(CachedScriptVMId.CompilerVersionID != FGuid());
+			check(CachedScriptVMId.BaseScriptID != FGuid());
+
+			NewResource->SetScript(this, (ERHIFeatureLevel::Type)TargetFeatureLevel, CachedScriptVMId.CompilerVersionID, CachedScriptVMId.BaseScriptID, CachedScriptVMId.ReferencedDependencyIds, GetName());
 			ResourceToCache = NewResource;
 
 			check(ResourceToCache);
@@ -390,7 +987,7 @@ void UNiagaraScript::CacheResourceShadersForCooking(EShaderPlatform ShaderPlatfo
 
 
 
-void UNiagaraScript::CacheShadersForResources(EShaderPlatform ShaderPlatform, FNiagaraScript *ResourceToCache, bool bApplyCompletedShaderMapForRendering, bool bForceRecompile, bool bCooking)
+void UNiagaraScript::CacheShadersForResources(EShaderPlatform ShaderPlatform, FNiagaraShaderScript *ResourceToCache, bool bApplyCompletedShaderMapForRendering, bool bForceRecompile, bool bCooking)
 {
 	if (CanBeRunOnGpu())
 	{
@@ -434,21 +1031,125 @@ void UNiagaraScript::CacheResourceShadersForRendering(bool bRegenerateId, bool b
 	{
 		if (Source)
 		{
-			FNiagaraScript* ResourceToCache;
+			FNiagaraShaderScript* ResourceToCache;
 			ERHIFeatureLevel::Type CacheFeatureLevel = ERHIFeatureLevel::SM5;
-			ScriptResource.SetScript(this, FeatureLevel, UniqueID, GetName());
+			ScriptResource.SetScript(this, FeatureLevel, CachedScriptVMId.CompilerVersionID, CachedScriptVMId.BaseScriptID, CachedScriptVMId.ReferencedDependencyIds, GetName());
 
 			//if (ScriptResourcesByFeatureLevel[FeatureLevel])
 			{
 				EShaderPlatform ShaderPlatform = GShaderPlatformForFeatureLevel[CacheFeatureLevel];
 				ResourceToCache = ScriptResourcesByFeatureLevel[CacheFeatureLevel];
-				CacheShadersForResources(ShaderPlatform, &ScriptResource, true, bForceRecompile);
+				CacheShadersForResources(ShaderPlatform, &ScriptResource, true);
 				ScriptResourcesByFeatureLevel[CacheFeatureLevel] = &ScriptResource;
 			}
 		}
 	}
 }
 
+void UNiagaraScript::SyncAliases(const TMap<FString, FString>& RenameMap)
+{
+	// First handle any rapid iteration parameters...
+	{
+		TArray<FNiagaraVariable> Params;
+		RapidIterationParameters.GetParameters(Params);
+		for (FNiagaraVariable Var : Params)
+		{
+			FNiagaraVariable NewVar = FNiagaraVariable::ResolveAliases(Var, RenameMap);
+			if (NewVar.GetName() != Var.GetName())
+			{
+				RapidIterationParameters.RenameParameter(Var, NewVar.GetName());
+			}
+		}
+	}
+
+	InvalidateExecutionReadyParameterStores();
+
+	// Now handle any Parameters overall..
+	for (int32 i = 0; i < GetVMExecutableData().Parameters.Parameters.Num(); i++)
+	{
+		FNiagaraVariable Var = GetVMExecutableData().Parameters.Parameters[i];
+		FNiagaraVariable NewVar = FNiagaraVariable::ResolveAliases(Var, RenameMap);
+		if (NewVar.GetName() != Var.GetName())
+		{
+			GetVMExecutableData().Parameters.Parameters[i] = NewVar;
+		}
+	}
+
+	// Also handle any data set mappings...
+	auto Iterator = GetVMExecutableData().DataSetToParameters.CreateIterator();
+	while (Iterator)
+	{
+		for (int32 i = 0; i < Iterator.Value().Parameters.Num(); i++)
+		{
+			FNiagaraVariable Var = Iterator.Value().Parameters[i];
+			FNiagaraVariable NewVar = FNiagaraVariable::ResolveAliases(Var, RenameMap);
+			if (NewVar.GetName() != Var.GetName())
+			{
+				Iterator.Value().Parameters[i] = NewVar;
+			}
+		}
+		++Iterator;
+	}
+}
+
+bool UNiagaraScript::SynchronizeExecutablesWithMaster(const UNiagaraScript* Script, const TMap<FString, FString>& RenameMap)
+{
+	FNiagaraVMExecutableDataId Id;
+	ComputeVMCompilationId(Id);
+
+#if 1 // TODO Shaun... turn this on...
+	if (Id == Script->GetVMExecutableDataCompilationId())
+	{
+		CachedScriptVM.Reset();
+		ScriptResource.Invalidate();
+
+		CachedScriptVM = Script->CachedScriptVM;
+		CachedScriptVMId = Script->CachedScriptVMId;
+		CachedParameterCollectionReferences = Script->CachedParameterCollectionReferences;
+		CachedDefaultDataInterfaces.Empty();
+		for (const FNiagaraScriptDataInterfaceInfo& Info : Script->CachedDefaultDataInterfaces)
+		{
+			FNiagaraScriptDataInterfaceInfo AddInfo;
+			AddInfo = Info;
+			AddInfo.DataInterface = DuplicateObject<UNiagaraDataInterface>(Info.DataInterface, this);
+			CachedDefaultDataInterfaces.Add(AddInfo);
+		}
+
+		GenerateStatScopeIDs();
+
+		//SyncAliases(RenameMap);
+
+		// Now go ahead and trigger the GPU script compile now that we have a compiled GPU hlsl script.
+		if (Usage == ENiagaraScriptUsage::ParticleGPUComputeScript)
+		{
+			CacheResourceShadersForRendering(false, true);
+		}
+
+		OnVMScriptCompiled().Broadcast(this);
+		return true;
+	}
+#endif
+
+	return false;
+}
+
+void UNiagaraScript::InvalidateCompileResults()
+{
+	UE_LOG(LogNiagara, Log, TEXT("InvalidateCompileResults %s"), *GetPathName());
+	CachedScriptVM.Reset();
+	ScriptResource.Invalidate();
+	CachedScriptVMId.Invalidate();
+	LastGeneratedVMId.Invalidate();
+}
+
+
+UNiagaraScript::FOnScriptCompiled& UNiagaraScript::OnVMScriptCompiled()
+{
+	return OnVMScriptCompiledDelegate;
+}
+
+
+
 #endif
 
 
@@ -465,21 +1166,14 @@ NIAGARA_API bool UNiagaraScript::IsScriptCompilationPending(bool bGPUScript) con
 	}
 	else
 	{
-		return ByteCode.Num() == 0 && (LastCompileStatus == ENiagaraScriptCompileStatus::NCS_BeingCreated || LastCompileStatus == ENiagaraScriptCompileStatus::NCS_Unknown);
+		if (CachedScriptVM.IsValid())
+		{
+			return CachedScriptVM.ByteCode.Num() == 0 && (CachedScriptVM.LastCompileStatus == ENiagaraScriptCompileStatus::NCS_BeingCreated || CachedScriptVM.LastCompileStatus == ENiagaraScriptCompileStatus::NCS_Unknown);
+		}
+		return false;
 	}
 }
 
-NIAGARA_API void UNiagaraScript::InvalidateScript()
-{
-	ByteCode.Empty();
-	ScriptResource.Invalidate();
-	LastCompileStatus = ENiagaraScriptCompileStatus::NCS_Unknown;
-#if WITH_EDITORONLY_DATA
-	InvalidateChangeID(); // Identify that we have no idea if we're synchronized or not.
-#endif
-}
-
-
 NIAGARA_API bool UNiagaraScript::DidScriptCompilationSucceed(bool bGPUScript) const
 {
 	if (bGPUScript)
@@ -498,20 +1192,23 @@ NIAGARA_API bool UNiagaraScript::DidScriptCompilationSucceed(bool bGPUScript) co
 	}
 	else
 	{
-		return ByteCode.Num() != 0;
+		if (CachedScriptVM.IsValid())
+		{
+			return CachedScriptVM.ByteCode.Num() != 0;
+		}
 	}
 
 	return false;
 }
 
-void SerializeNiagaraShaderMaps(const TMap<const ITargetPlatform*, FNiagaraScript*>* PlatformScriptResourcesToSavePtr, FArchive& Ar, FNiagaraScript &OutLoadedResource)
+void SerializeNiagaraShaderMaps(const TMap<const ITargetPlatform*, FNiagaraShaderScript*>* PlatformScriptResourcesToSavePtr, FArchive& Ar, FNiagaraShaderScript &OutLoadedResource)
 {
 //	SCOPED_LOADTIMER(SerializeInlineShaderMaps);
 	if (Ar.IsSaving())
 	{
 		int32 NumResourcesToSave = 0;
-		FNiagaraScript *ScriptResourceToSave = nullptr;
-		FNiagaraScript* const * ScriptResourceToSavePtr = nullptr;
+		FNiagaraShaderScript *ScriptResourceToSave = nullptr;
+		FNiagaraShaderScript* const * ScriptResourceToSavePtr = nullptr;
 		if (Ar.IsCooking())
 		{
 			check(PlatformScriptResourcesToSavePtr);
@@ -545,9 +1242,7 @@ void SerializeNiagaraShaderMaps(const TMap<const ITargetPlatform*, FNiagaraScrip
 	}
 }
 
-
-
-void ProcessSerializedShaderMaps(UNiagaraScript* Owner, FNiagaraScript LoadedResource, FNiagaraScript* (&OutScriptResourcesLoaded)[ERHIFeatureLevel::Num])
+void ProcessSerializedShaderMaps(UNiagaraScript* Owner, FNiagaraShaderScript LoadedResource, FNiagaraShaderScript* (&OutScriptResourcesLoaded)[ERHIFeatureLevel::Num])
 {
 	check(IsInGameThread());
 
@@ -563,24 +1258,27 @@ void ProcessSerializedShaderMaps(UNiagaraScript* Owner, FNiagaraScript LoadedRes
 		}
 
 		OutScriptResourcesLoaded[LoadedFeatureLevel]->SetShaderMap(LoadedShaderMap);
-		LoadedResource.SetDatainterfaceBufferDescriptors(LoadedResource.GetShaderGameThread()->GetDIBufferDescriptors());
+		LoadedResource.SetDataInterfaceParamInfo(LoadedResource.GetShaderGameThread()->GetDIParameters());
 	}
 
 }
 
-
-
-FNiagaraScript* UNiagaraScript::AllocateResource()
+FNiagaraShaderScript* UNiagaraScript::AllocateResource()
 {
-	return new FNiagaraScript();
+	return new FNiagaraShaderScript();
 }
 
 TArray<ENiagaraScriptUsage> UNiagaraScript::GetSupportedUsageContexts() const
+{
+	return GetSupportedUsageContextsForBitmask(ModuleUsageBitmask);
+}
+
+TArray<ENiagaraScriptUsage> UNiagaraScript::GetSupportedUsageContextsForBitmask(int32 InModuleUsageBitmask)
 {
 	TArray<ENiagaraScriptUsage> Supported;
 	for (int32 i = 0; i <= (int32)ENiagaraScriptUsage::SystemUpdateScript; i++)
 	{
-		int32 TargetBit = (ModuleUsageBitmask >> (int32)i) & 1;
+		int32 TargetBit = (InModuleUsageBitmask >> (int32)i) & 1;
 		if (TargetBit == 1)
 		{
 			Supported.Add((ENiagaraScriptUsage)i);
@@ -591,36 +1289,75 @@ TArray<ENiagaraScriptUsage> UNiagaraScript::GetSupportedUsageContexts() const
 
 bool UNiagaraScript::CanBeRunOnGpu()const
 {
-	if (UNiagaraEmitter* Emitter = GetTypedOuter<UNiagaraEmitter>())
+
+	if (Usage != ENiagaraScriptUsage::ParticleGPUComputeScript)
 	{
-		if (Emitter->SimTarget == ENiagaraSimTarget::CPUSim)
+		return false;
+	}
+	if (!CachedScriptVM.IsValid())
+	{
+		return false;
+	}
+	for (const FNiagaraScriptDataInterfaceCompileInfo& InterfaceInfo : CachedScriptVM.DataInterfaceInfo)
+	{
+		if (!InterfaceInfo.CanExecuteOnTarget(ENiagaraSimTarget::GPUComputeSim))
 		{
 			return false;
 		}
+	}
+	return true;
+}
 
-		if (!IsParticleSpawnScript())
+
+bool UNiagaraScript::LegacyCanBeRunOnGpu() const
+{
+	if (UNiagaraEmitter* Emitter = GetTypedOuter<UNiagaraEmitter>())
+	{
+		if (Emitter->SimTarget == ENiagaraSimTarget::CPUSim)
 		{
 			return false;
 		}
 
-		for (const FNiagaraScriptDataInterfaceInfo& InterfaceInfo : DataInterfaceInfo)
+		if (!IsParticleSpawnScript())
 		{
-			if (!InterfaceInfo.DataInterface || !InterfaceInfo.DataInterface->CanExecuteOnTarget(ENiagaraSimTarget::GPUComputeSim))
-			{
-				return false;
-			}
+			return false;
 		}
 
 		return true;
 	}
-
 	return false;
 }
 
+
+#if WITH_EDITORONLY_DATA
+void UNiagaraScript::CleanUpOldAndInitializeNewRapidIterationParameters(FString UniqueEmitterName)
+{
+	Source->CleanUpOldAndInitializeNewRapidIterationParameters(UniqueEmitterName, Usage, UsageId, RapidIterationParameters);
+}
+
+FGuid UNiagaraScript::GetBaseChangeID() const
+{
+	return Source->GetChangeID(); 
+}
+
+ENiagaraScriptCompileStatus UNiagaraScript::GetLastCompileStatus() const
+{
+	if (CachedScriptVM.IsValid())
+	{
+		return CachedScriptVM.LastCompileStatus;
+	}
+	return ENiagaraScriptCompileStatus::NCS_Unknown;
+}
+#endif
+
 bool UNiagaraScript::UsesCollection(const UNiagaraParameterCollection* Collection)const
 {
-	return ParameterCollections.FindByPredicate([&](const UNiagaraParameterCollection* CheckCollection)
+	if (CachedScriptVM.IsValid())
 	{
-		return CheckCollection == Collection;
-	}) != NULL;
+		return CachedParameterCollectionReferences.FindByPredicate([&](const UNiagaraParameterCollection* CheckCollection)
+		{
+			return CheckCollection == Collection;
+		}) != NULL;
+	}
+	return false;
 }
diff --git a/Engine/Plugins/FX/Niagara/Source/Niagara/Private/NiagaraScriptDerivedData.cpp b/Engine/Plugins/FX/Niagara/Source/Niagara/Private/NiagaraScriptDerivedData.cpp
new file mode 100644
--- /dev/null
+++ b/Engine/Plugins/FX/Niagara/Source/Niagara/Private/NiagaraScriptDerivedData.cpp
@@ -0,0 +1,97 @@
+// Copyright 1998-2018 Epic Games, Inc. All Rights Reserved.
+#include "NiagaraScriptDerivedData.h"
+#include "NiagaraStats.h"
+#include "Serialization/MemoryWriter.h"
+#include "NiagaraScript.h"
+#include "UObject/Package.h"
+#include "NiagaraCustomVersion.h"
+#include "NiagaraScriptSourceBase.h"
+#include "Serialization/ObjectAndNameAsStringProxyArchive.h"
+#include "Serialization/MemoryReader.h"
+#include "UObject/PropertyPortFlags.h"
+#include "Modules/ModuleManager.h"
+
+#if WITH_EDITOR
+
+DECLARE_CYCLE_STAT(TEXT("Niagara - DerivedData - Compile"), STAT_Niagara_VM_Compile, STATGROUP_Niagara);
+
+FNiagaraScriptDerivedData::FNiagaraScriptDerivedData(const FString& InScriptFullName, const TSharedPtr<FNiagaraCompileRequestDataBase, ESPMode::ThreadSafe>& InPrecompiledData, const FNiagaraCompileOptions& InCompileOptions, const FNiagaraVMExecutableDataId& InGeneratedVMId, bool bInAsync)
+	: ScriptFullName(InScriptFullName),  CompileOptions(InCompileOptions), PrecompiledData(InPrecompiledData), GeneratedVMId(InGeneratedVMId), bAsynchronous(bInAsync)
+{
+	check(InPrecompiledData.IsValid());
+
+	CompileOptions.AdditionalDefines = InGeneratedVMId.AdditionalDefines;
+	NiagaraModule = &FModuleManager::Get().LoadModuleChecked<INiagaraModule>(TEXT("Niagara"));
+}
+
+FNiagaraScriptDerivedData::~FNiagaraScriptDerivedData()
+{
+}
+
+FString FNiagaraScriptDerivedData::GetPluginSpecificCacheKeySuffix() const
+{
+	enum { UE_NIAGARA_COMPILATION_DERIVEDDATA_VER = 1 };
+
+	FString Ret = FString::Printf(TEXT("%i"),
+		(int32)UE_NIAGARA_COMPILATION_DERIVEDDATA_VER);
+	GeneratedVMId.AppendKeyString(Ret);
+
+	//UE_LOG(LogNiagara, Log, TEXT("Cache Key Gen Niagara VM Script DDC data for %s key %s"), *ScriptFullName, *Ret);
+
+	return Ret;
+}
+
+bool FNiagaraScriptDerivedData::BinaryToExecData(const TArray<uint8>& InBinaryData, FNiagaraVMExecutableData& OutExecData)
+{
+	// Do not call this on anything other than the game thread as it depends on the FObjectAndNameAsStringProxyArchive, which 
+	// calls FindStaticObject which can fail when used in any other thread!
+	check(IsInGameThread());
+	if (InBinaryData.Num() == 0)
+	{
+		return false;
+	}
+	FMemoryReader Ar(InBinaryData);
+	FObjectAndNameAsStringProxyArchive SafeAr(Ar, false);
+	OutExecData.SerializeData(SafeAr, true);
+	SafeAr.Close();
+	Ar.Close();
+
+	return !SafeAr.IsError();
+}
+
+
+bool FNiagaraScriptDerivedData::Build(TArray<uint8>& OutData)
+{
+	SCOPE_CYCLE_COUNTER(STAT_Niagara_VM_Compile);
+	//UE_LOG(LogNiagara, Log, TEXT("Building Niagara VM Script DDC data for %s"), *ScriptFullName);
+
+	bool bCompilationSuccessful = false;
+	{
+		FString OutGraphLevelErrorMessages;
+		TSharedPtr<FNiagaraVMExecutableData> ExeData = NiagaraModule->CompileScript(PrecompiledData.Get(), CompileOptions);
+		bCompilationSuccessful = ExeData.IsValid();
+
+		if (bCompilationSuccessful)
+		{
+			//const bool bTestSuccessfulReadWrite = false; // Set to true to debug serialization issues.
+			{
+				FMemoryWriter Ar(OutData, true);
+				FObjectAndNameAsStringProxyArchive SafeAr(Ar, false);
+				ExeData->SerializeData(SafeAr, true);
+			}
+
+			/*if (bTestSuccessfulReadWrite && IsInGameThread())
+			{
+				FNiagaraVMExecutableData TestExeData;
+				bool bGenerated = BinaryToExecData(OutData, TestExeData);
+				check(bGenerated);
+
+				bool bIdentical = FNiagaraVMExecutableData::StaticStruct()->CompareScriptStruct(&TestExeData, ExeData.Get(), PPF_None);
+				check(bIdentical);
+			}*/
+		}
+	}
+
+	return bCompilationSuccessful;
+}
+#endif	//WITH_EDITOR
diff --git a/Engine/Plugins/FX/Niagara/Source/Niagara/Private/NiagaraScriptDerivedData.h b/Engine/Plugins/FX/Niagara/Source/Niagara/Private/NiagaraScriptDerivedData.h
new file mode 100644
--- /dev/null
+++ b/Engine/Plugins/FX/Niagara/Source/Niagara/Private/NiagaraScriptDerivedData.h
@@ -0,0 +1,87 @@
+// Copyright 1998-2018 Epic Games, Inc. All Rights Reserved.
+
+#pragma once
+
+#include "CoreMinimal.h"
+#include "NiagaraScript.h"
+#include "NiagaraModule.h"
+#include "NiagaraScriptSourceBase.h"
+
+#if WITH_EDITOR
+#include "DerivedDataPluginInterface.h"
+#endif
+
+
+#if WITH_EDITOR
+
+//////////////////////////////////////////////////////////////////////////
+// FNiagaraScriptDerivedData
+class FNiagaraScriptDerivedData : public FDerivedDataPluginInterface
+{
+private:
+	FString ScriptFullName;
+
+	// The compile options to use in compiling this script.
+	FNiagaraCompileOptions CompileOptions;
+
+	TSharedPtr<FNiagaraCompileRequestDataBase, ESPMode::ThreadSafe> PrecompiledData;
+
+	// Niagara script id requested.
+	FNiagaraVMExecutableDataId GeneratedVMId;
+
+	// Was this created to be async?
+	bool bAsynchronous;
+
+	// The compiler
+	INiagaraModule* NiagaraModule;
+
+
+public:
+	FNiagaraScriptDerivedData(const FString& InScriptFullName,const TSharedPtr<FNiagaraCompileRequestDataBase, ESPMode::ThreadSafe>& InPrecompiledData, const FNiagaraCompileOptions& InCompileOptions, const FNiagaraVMExecutableDataId& InGeneratedVMId, bool bInAsync);
+	virtual ~FNiagaraScriptDerivedData();
+
+	const FNiagaraVMExecutableDataId& GetGeneratedId() const {
+		return GeneratedVMId;
+	}
+
+	virtual const TCHAR* GetPluginName() const override
+	{
+		return TEXT("NiagaraScriptDerivedData");
+	}
+
+	virtual const TCHAR* GetVersionString() const override
+	{
+		// This is a version string that mimics the old versioning scheme. If you
+		// want to bump this version, generate a new guid using VS->Tools->Create GUID and
+		// return it here. Ex.
+		return TEXT("F7CCEA20D9A5451F8658F37985856110");
+	}
+
+	virtual FString GetPluginSpecificCacheKeySuffix() const override;
+
+
+	virtual bool IsBuildThreadsafe() const override
+	{
+		return bAsynchronous;
+	}
+
+	/** Helper to convert the struct from its binary data out of the DDC to it's actual in-memory version.
+		Do not call this on anything other than the game thread as it depends on the FObjectAndNameAsStringProxyArchive, 
+		which calls FindStaticObject which can fail when used in any other thread!*/
+	static bool BinaryToExecData(const TArray<uint8>& InBinaryData, FNiagaraVMExecutableData& OutExecData);
+
+	/** Indicated that this plugin generates deterministic data. This is used for DDC verification */
+	virtual bool IsDeterministic() const override { return true; }
+
+	virtual FString GetDebugContextString() const override { return TEXT("Unknown Context"); }
+
+	virtual bool Build( TArray<uint8>& OutData ) override;
+
+	/** Return true if we can build **/
+	bool CanBuild()
+	{
+		return PrecompiledData.IsValid();
+	}
+};
+
+#endif	//WITH_EDITOR
diff --git a/Engine/Plugins/FX/Niagara/Source/Niagara/Private/NiagaraScriptExecutionContext.cpp b/Engine/Plugins/FX/Niagara/Source/Niagara/Private/NiagaraScriptExecutionContext.cpp
--- a/Engine/Plugins/FX/Niagara/Source/Niagara/Private/NiagaraScriptExecutionContext.cpp
+++ b/Engine/Plugins/FX/Niagara/Source/Niagara/Private/NiagaraScriptExecutionContext.cpp
@@ -5,233 +5,14 @@
 #include "NiagaraEmitterInstanceBatcher.h"
 #include "NiagaraDataInterface.h"
 #include "NiagaraSystemInstance.h"
+#include "NiagaraWorldManager.h"
 
 DECLARE_CYCLE_STAT(TEXT("Register Setup"), STAT_NiagaraSimRegisterSetup, STATGROUP_Niagara);
-
-
-//////////////////////////////////////////////////////////////////////////
-FNiagaraScriptExecutionParameterStore::FNiagaraScriptExecutionParameterStore()
-	: FNiagaraParameterStore()
-{
-
-}
-
-FNiagaraScriptExecutionParameterStore::FNiagaraScriptExecutionParameterStore(const FNiagaraParameterStore& Other)
-{
-	*this = Other;
-}
-
-FNiagaraScriptExecutionParameterStore& FNiagaraScriptExecutionParameterStore::operator=(const FNiagaraParameterStore& Other)
-{
-	FNiagaraParameterStore::operator=(Other);
-	return *this;
-}
-
-uint32 OffsetAlign(uint32 SrcOffset, uint32 Size)
-{
-	uint32 OffsetRemaining = UNIFORM_BUFFER_STRUCT_ALIGNMENT - (SrcOffset % UNIFORM_BUFFER_STRUCT_ALIGNMENT);
-	if (Size <= OffsetRemaining)
-	{
-		return SrcOffset;
-	}
-	else
-	{
-		return Align(SrcOffset, UNIFORM_BUFFER_STRUCT_ALIGNMENT);
-	}
-}
-
-void FNiagaraScriptExecutionParameterStore::GenerateLayoutInfoInternal(TArray<FNiagaraScriptExecutionParameterStore::FPaddingInfo>& Members, uint32& NextMemberOffset, const UStruct* InSrcStruct, uint32 InSrcOffset)
-{
-	// Now insert an appropriate data member into the mix...
-	if (InSrcStruct == FNiagaraTypeDefinition::GetBoolStruct() || InSrcStruct == FNiagaraTypeDefinition::GetIntStruct())
-	{
-		//Members.Emplace(*Property->GetName(), TEXT(""), InSrcOffset, (EUniformBufferBaseType)TUniformBufferTypeInfo<uint32>::BaseType, EShaderPrecisionModifier::Float, TUniformBufferTypeInfo<uint32>::NumRows, TUniformBufferTypeInfo<uint32>::NumColumns, TUniformBufferTypeInfo<uint32>::NumElements, TUniformBufferTypeInfo<uint32>::GetStruct());
-		Members.Emplace(InSrcOffset, Align(NextMemberOffset, TUniformBufferTypeInfo<uint32, false>::Alignment), (TUniformBufferTypeInfo<uint32, false>::NumRows * TUniformBufferTypeInfo<uint32, false>::NumColumns) * sizeof(uint32));
-		InSrcOffset += sizeof(uint32);
-		NextMemberOffset = Members[Members.Num() - 1].DestOffset + Members[Members.Num() - 1].Size;
-	}
-	else if (InSrcStruct == FNiagaraTypeDefinition::GetFloatStruct())
-	{
-		//Members.Emplace(*Property->GetName(), TEXT(""), InSrcOffset, (EUniformBufferBaseType)TUniformBufferTypeInfo < float > ::BaseType, EShaderPrecisionModifier::Float, TUniformBufferTypeInfo<float>::NumRows, TUniformBufferTypeInfo<float>::NumColumns, TUniformBufferTypeInfo<float>::NumElements, TUniformBufferTypeInfo<float>::GetStruct());
-		Members.Emplace(InSrcOffset, Align(NextMemberOffset, TUniformBufferTypeInfo<float, false>::Alignment), (TUniformBufferTypeInfo<float, false>::NumRows * TUniformBufferTypeInfo<float, false>::NumColumns) * sizeof(float));
-		InSrcOffset += sizeof(float);
-		NextMemberOffset = Members[Members.Num() - 1].DestOffset + Members[Members.Num() - 1].Size;
-	}
-	else if (InSrcStruct == FNiagaraTypeDefinition::GetVec2Struct())
-	{
-		//Members.Emplace(*Property->GetName(), TEXT(""), InSrcOffset, (EUniformBufferBaseType)TUniformBufferTypeInfo < FVector2D > ::BaseType, EShaderPrecisionModifier::Float, TUniformBufferTypeInfo<FVector2D>::NumRows, TUniformBufferTypeInfo<FVector2D>::NumColumns, TUniformBufferTypeInfo<FVector2D>::NumElements, TUniformBufferTypeInfo<FVector2D>::GetStruct());
-		uint32 StructFinalSize = (TUniformBufferTypeInfo<FVector2D, false>::NumRows * TUniformBufferTypeInfo<FVector2D, false>::NumColumns) * sizeof(float);
-		Members.Emplace(InSrcOffset, OffsetAlign(NextMemberOffset, StructFinalSize),StructFinalSize);
-		InSrcOffset += sizeof(FVector2D);
-		NextMemberOffset = Members[Members.Num() - 1].DestOffset + Members[Members.Num() - 1].Size;
-	}
-	else if (InSrcStruct == FNiagaraTypeDefinition::GetVec3Struct())
-	{
-		//Members.Emplace(*Property->GetName(), TEXT(""), InSrcOffset, (EUniformBufferBaseType)TUniformBufferTypeInfo < FVector > ::BaseType, EShaderPrecisionModifier::Float, TUniformBufferTypeInfo<FVector>::NumRows, TUniformBufferTypeInfo<FVector>::NumColumns, TUniformBufferTypeInfo<FVector>::NumElements, TUniformBufferTypeInfo<FVector>::GetStruct());
-		uint32 StructFinalSize = (TUniformBufferTypeInfo<FVector, false>::NumRows * TUniformBufferTypeInfo<FVector, false>::NumColumns) * sizeof(float);
-		Members.Emplace(InSrcOffset, OffsetAlign(NextMemberOffset, StructFinalSize), StructFinalSize);
-		InSrcOffset += sizeof(FVector);
-		NextMemberOffset = Members[Members.Num() - 1].DestOffset + Members[Members.Num() - 1].Size;
-	}
-	else if (InSrcStruct == FNiagaraTypeDefinition::GetVec4Struct() || InSrcStruct == FNiagaraTypeDefinition::GetColorStruct())
-	{
-		//Members.Emplace(*Property->GetName(), TEXT(""), InSrcOffset, (EUniformBufferBaseType)TUniformBufferTypeInfo < FVector4 > ::BaseType, EShaderPrecisionModifier::Float, TUniformBufferTypeInfo<FVector4>::NumRows, TUniformBufferTypeInfo<FVector4>::NumColumns, TUniformBufferTypeInfo<FVector4>::NumElements, TUniformBufferTypeInfo<FVector4>::GetStruct());
-		Members.Emplace(InSrcOffset, Align(NextMemberOffset, TUniformBufferTypeInfo<FVector4, false>::Alignment), (TUniformBufferTypeInfo<FVector4, false>::NumRows * TUniformBufferTypeInfo<FVector4, false>::NumColumns) * sizeof(float));
-		InSrcOffset += sizeof(FVector4);
-		NextMemberOffset = Members[Members.Num() - 1].DestOffset + Members[Members.Num() - 1].Size;
-	}
-	else if (InSrcStruct == FNiagaraTypeDefinition::GetMatrix4Struct())
-	{
-		//Members.Emplace(*Property->GetName(), TEXT(""), InSrcOffset, (EUniformBufferBaseType)TUniformBufferTypeInfo < FMatrix > ::BaseType, EShaderPrecisionModifier::Float, TUniformBufferTypeInfo<FMatrix>::NumRows, TUniformBufferTypeInfo<FMatrix>::NumColumns, TUniformBufferTypeInfo<FMatrix>::NumElements, TUniformBufferTypeInfo<FMatrix>::GetStruct());
-		Members.Emplace(InSrcOffset, Align(NextMemberOffset, TUniformBufferTypeInfo<FMatrix, false>::Alignment), (TUniformBufferTypeInfo<FMatrix, false>::NumRows * TUniformBufferTypeInfo<FMatrix, false>::NumColumns) * sizeof(float));
-		InSrcOffset += sizeof(FMatrix);
-		NextMemberOffset = Members[Members.Num() - 1].DestOffset + Members[Members.Num() - 1].Size;
-	}
-	else
-	{
-		NextMemberOffset = Align(NextMemberOffset, UNIFORM_BUFFER_STRUCT_ALIGNMENT); // New structs should be aligned to the head..
-
-		for (TFieldIterator<UProperty> PropertyIt(InSrcStruct, EFieldIteratorFlags::IncludeSuper); PropertyIt; ++PropertyIt)
-		{
-			UProperty* Property = *PropertyIt;
-			const UStruct* Struct = nullptr;
-
-			// First determine what struct type we're dealing with...
-			if (Property->IsA(UFloatProperty::StaticClass()))
-			{
-				Struct = FNiagaraTypeDefinition::GetFloatStruct();
-			}
-			else if (Property->IsA(UIntProperty::StaticClass()))
-			{
-				Struct = FNiagaraTypeDefinition::GetIntStruct();
-			}
-			else if (Property->IsA(UBoolProperty::StaticClass()))
-			{
-				Struct = FNiagaraTypeDefinition::GetBoolStruct();
-			}
-			//Should be able to support double easily enough
-			else if (UStructProperty* StructProp = CastChecked<UStructProperty>(Property))
-			{
-				Struct = StructProp->Struct;
-			}
-			else
-			{
-				check(false);
-			}
-
-			GenerateLayoutInfoInternal(Members, NextMemberOffset, Struct, InSrcOffset);
-		}
-	}
-
-}
-
-void FNiagaraScriptExecutionParameterStore::AddPaddedParamSize(const FNiagaraTypeDefinition& InParamType, uint32 InOffset)
-{
-	if (InParamType.IsDataInterface())
-	{
-		return;
-	}
-
-	uint32 NextMemberOffset = 0;
-	if (PaddingInfo.Num() != 0)
-	{
-		NextMemberOffset = PaddingInfo[PaddingInfo.Num() - 1].DestOffset + PaddingInfo[PaddingInfo.Num() - 1].Size;
-	}
-	GenerateLayoutInfoInternal(PaddingInfo, NextMemberOffset, InParamType.GetScriptStruct(), InOffset);
-
-	if (PaddingInfo.Num() != 0)
-	{
-		NextMemberOffset = PaddingInfo[PaddingInfo.Num() - 1].DestOffset + PaddingInfo[PaddingInfo.Num() - 1].Size;
-		PaddedParameterSize = Align(NextMemberOffset, UNIFORM_BUFFER_STRUCT_ALIGNMENT);
-	}
-	else
-	{
-		PaddedParameterSize = 0;
-	}
-}
-
-
-void FNiagaraScriptExecutionParameterStore::Init(UNiagaraScript* Script, ENiagaraSimTarget SimTarget)
-{
-	//TEMPORARTY:
-	//We should replace the storage on the script with an FNiagaraParameterStore also so we can just copy that over here. Though that is an even bigger refactor job so this is a convenient place to break that work up.
-
-	Empty();
-	PaddedParameterSize = 0;
-	AddScriptParams(Script, SimTarget);
-}
-
-void FNiagaraScriptExecutionParameterStore::AddScriptParams(UNiagaraScript* Script, ENiagaraSimTarget SimTarget)
-{
-	if (Script != nullptr)
-	{
-		//Here we add the current frame parameters.
-		for (FNiagaraVariable& Param : Script->Parameters.Parameters)
-		{
-			AddParameter(Param, false);
-		}
-
-#if WITH_EDITORONLY_DATA
-		Name = *FString::Printf(TEXT("ScriptExecParamStore %s %p"), *Script->GetFullName(), this);
-#endif
-
-		//Add previous frame values if we're interpolated spawn.
-		bool bIsInterpolatedSpawn = Script->IsInterpolatedParticleSpawnScript();
-
-		if (bIsInterpolatedSpawn)
-		{
-			for (FNiagaraVariable& Param : Script->Parameters.Parameters)
-			{
-				FNiagaraVariable PrevParam(Param.GetType(), FName(*(TEXT("PREV__") + Param.GetName().ToString())));
-				AddParameter(PrevParam, false);
-			}
-		}
-
-		ParameterSize = GetParameterDataArray().Num();
-		if (bIsInterpolatedSpawn)
-		{
-			CopyCurrToPrev();
-		}
-
-		//Internal constants - only needed for non-GPU sim
-		if (SimTarget != ENiagaraSimTarget::GPUComputeSim)
-		{
-			for (FNiagaraVariable& InternalVar : Script->InternalParameters.Parameters)
-			{
-				AddParameter(InternalVar, false);
-			}
-		}
-
-		for (FNiagaraScriptDataInterfaceInfo& Info : Script->DataInterfaceInfo)
-		{
-			FNiagaraVariable Var(Info.Type, Info.Name);
-			AddParameter(Var, false);
-			SetDataInterface(Info.DataInterface, IndexOf(Var));
-		}
-	}
-}
-
-void FNiagaraScriptExecutionParameterStore::CopyCurrToPrev()
-{
-	int32 ParamStart = ParameterSize / 2;
-	checkSlow(FMath::Frac((float)ParameterSize / 2) == 0.0f);
-
-	FMemory::Memcpy(GetParameterData_Internal(ParamStart), GetParameterData(0), ParamStart);
-}
-
-
-void FNiagaraScriptExecutionParameterStore::CopyParameterDataToPaddedBuffer(uint8* InTargetBuffer, uint32 InTargetBufferSizeInBytes)
-{
-	check(InTargetBufferSizeInBytes >= PaddedParameterSize);
-	FMemory::Memzero(InTargetBuffer, InTargetBufferSizeInBytes);
-	const uint8* SrcData = GetParameterDataArray().GetData();
-	for (int32 i = 0; i < PaddingInfo.Num(); i++)
-	{
-		check((PaddingInfo[i].DestOffset + PaddingInfo[i].Size) <= InTargetBufferSizeInBytes);
-		check((PaddingInfo[i].SrcOffset + PaddingInfo[i].Size) <= (uint32)GetParameterDataArray().Num());
-		FMemory::Memcpy(InTargetBuffer + PaddingInfo[i].DestOffset, SrcData + PaddingInfo[i].SrcOffset, PaddingInfo[i].Size);
-	}
-}
+DECLARE_CYCLE_STAT(TEXT("Context Ticking"), STAT_NiagaraScriptExecContextTick, STATGROUP_Niagara);
+DECLARE_CYCLE_STAT(TEXT("Rebind DInterface Func Table"), STAT_NiagaraRebindDataInterfaceFunctionTable, STATGROUP_Niagara);
+	//Add previous frame values if we're interpolated spawn.
+	
+	//Internal constants - only needed for non-GPU sim
 
 uint32 FNiagaraScriptExecutionContext::TickCounter = 0;
 
@@ -249,25 +30,27 @@ bool FNiagaraScriptExecutionContext::Init(UNiagaraScript* InScript, ENiagaraSimT
 {
 	Script = InScript;
 
-	Parameters.Init(Script, InTarget);
-	
+	Parameters.InitFromOwningContext(Script, InTarget, true);
 
 	return true;//TODO: Error cases?
 }
 
 bool FNiagaraScriptExecutionContext::Tick(FNiagaraSystemInstance* ParentSystemInstance, ENiagaraSimTarget SimTarget)
 {
-	if (Script)//TODO: Remove. Script can only be null for system instances that currently don't have their script exec context set up correctly.
+	SCOPE_CYCLE_COUNTER(STAT_NiagaraScriptExecContextTick);
+
+	if (Script && Script->IsReadyToRun(ENiagaraSimTarget::CPUSim))//TODO: Remove. Script can only be null for system instances that currently don't have their script exec context set up correctly.
 	{
+		const TArray<UNiagaraDataInterface*>& DataInterfaces = GetDataInterfaces();
+		
 		//Bind data interfaces if needed.
 		if (Parameters.GetInterfacesDirty())
 		{
+			SCOPE_CYCLE_COUNTER(STAT_NiagaraRebindDataInterfaceFunctionTable);
 			// UE_LOG(LogNiagara, Log, TEXT("Updating data interfaces for script %s"), *Script->GetFullName());
 
-			FunctionTable.Empty();
-			const TArray<UNiagaraDataInterface*>& DataInterfaces = GetDataInterfaces();
 			// We must make sure that the data interfaces match up between the original script values and our overrides...
-			if (Script->DataInterfaceInfo.Num() != DataInterfaces.Num())
+			if (Script->GetVMExecutableData().DataInterfaceInfo.Num() != DataInterfaces.Num())
 			{
 				UE_LOG(LogNiagara, Warning, TEXT("Mismatch between Niagara Exectuion Context data interfaces and those in it's script!"));
 				return false;
@@ -276,11 +59,12 @@ bool FNiagaraScriptExecutionContext::Tick(FNiagaraSystemInstance* ParentSystemIn
 			//Fill the instance data table.
 			if (ParentSystemInstance)
 			{
-				DataInterfaceInstDataTable.SetNumZeroed(Script->NumUserPtrs);
-				for (int32 i = 0; i < Script->DataInterfaceInfo.Num(); i++)
+				DataInterfaceInstDataTable.SetNumZeroed(Script->GetVMExecutableData().NumUserPtrs, false);
+				for (int32 i = 0; i < DataInterfaces.Num(); i++)
 				{
 					UNiagaraDataInterface* Interface = DataInterfaces[i];
-					int32 UserPtrIdx = Script->DataInterfaceInfo[i].UserPtrIdx;
+
+					int32 UserPtrIdx = Script->GetVMExecutableData().DataInterfaceInfo[i].UserPtrIdx;
 					if (UserPtrIdx != INDEX_NONE)
 					{
 						void* InstData = ParentSystemInstance->FindDataInterfaceInstanceData(Interface);
@@ -290,26 +74,28 @@ bool FNiagaraScriptExecutionContext::Tick(FNiagaraSystemInstance* ParentSystemIn
 			}
 			else
 			{
-				check(Script->NumUserPtrs == 0);//Can't have user ptrs if we have no parent instance.
+				check(Script->GetVMExecutableData().NumUserPtrs == 0);//Can't have user ptrs if we have no parent instance.
 			}
 
+			FunctionTable.Reset(Script->GetVMExecutableData().CalledVMExternalFunctions.Num());
+
 			bool bSuccessfullyMapped = true;
-			for (FVMExternalFunctionBindingInfo& BindingInfo : Script->CalledVMExternalFunctions)
+			for (FVMExternalFunctionBindingInfo& BindingInfo : Script->GetVMExecutableData().CalledVMExternalFunctions)
 			{
-				for (int32 i = 0; i < Script->DataInterfaceInfo.Num(); i++)
+				for (int32 i = 0; i < Script->GetVMExecutableData().DataInterfaceInfo.Num(); i++)
 				{
-					FNiagaraScriptDataInterfaceInfo& ScriptInfo = Script->DataInterfaceInfo[i];
+					FNiagaraScriptDataInterfaceCompileInfo& ScriptInfo = Script->GetVMExecutableData().DataInterfaceInfo[i];
 					UNiagaraDataInterface* ExternalInterface = DataInterfaces[i];
-
 					if (ScriptInfo.Name == BindingInfo.OwnerName)
 					{
 						void* InstData = ScriptInfo.UserPtrIdx == INDEX_NONE ? nullptr : DataInterfaceInstDataTable[ScriptInfo.UserPtrIdx];
-						FVMExternalFunction Func = ExternalInterface->GetVMExternalFunction(BindingInfo, InstData);
-						if (Func.IsBound())
+						int32 AddedIdx = FunctionTable.Add(FVMExternalFunction());
+						if (ExternalInterface != nullptr)
 						{
-							FunctionTable.Add(Func);
+							ExternalInterface->GetVMExternalFunction(BindingInfo, InstData, FunctionTable[AddedIdx]);
 						}
-						else
+
+						if (AddedIdx != INDEX_NONE && !FunctionTable[AddedIdx].IsBound())
 						{
 							UE_LOG(LogNiagara, Error, TEXT("Could not Get VMExternalFunction '%s'.. emitter will not run!"), *BindingInfo.Name.ToString());
 							bSuccessfullyMapped = false;
@@ -335,13 +121,13 @@ bool FNiagaraScriptExecutionContext::Tick(FNiagaraSystemInstance* ParentSystemIn
 void FNiagaraScriptExecutionContext::PostTick()
 {
 	//If we're for interpolated spawn, copy over the previous frame's parameters into the Prev parameters.
-	if (Script && Script->IsInterpolatedParticleSpawnScript())
+	if (Script && Script->GetComputedVMCompilationId().HasInterpolatedParameters())
 	{
 		Parameters.CopyCurrToPrev();
 	}
 }
 
-bool FNiagaraScriptExecutionContext::Execute(uint32 NumInstances, TArray<FNiagaraDataSetExecutionInfo>& DataSetInfos)
+bool FNiagaraScriptExecutionContext::Execute(uint32 NumInstances, TArray<FNiagaraDataSetExecutionInfo, TInlineAllocator<8>>& DataSetInfos)
 {
 	if (NumInstances == 0)
 	{
@@ -363,7 +149,8 @@ bool FNiagaraScriptExecutionContext::Execute(uint32 NumInstances, TArray<FNiagar
 		for (FNiagaraDataSetExecutionInfo& DataSetInfo : DataSetInfos)
 		{
 			check(DataSetInfo.DataSet);
-			FDataSetMeta SetMeta(DataSetInfo.DataSet->GetSizeBytes(), &InputRegisters[NumInputRegisters], NumInputRegisters);
+			FDataSetMeta SetMeta(DataSetInfo.DataSet->GetSizeBytes(), &InputRegisters[NumInputRegisters], NumInputRegisters, DataSetInfo.StartInstance, 
+				&DataSetInfo.DataSet->CurrIDTable(), &DataSetInfo.DataSet->GetFreeIDTable(), &DataSetInfo.DataSet->GetNumFreeIDs(), &DataSetInfo.DataSet->GetMaxUsedID(), DataSetInfo.DataSet->GetIDAcquireTag());
 			DataSetMetaTable.Add(SetMeta);
 			if (DataSetInfo.bAllocate)
 			{
@@ -378,7 +165,7 @@ bool FNiagaraScriptExecutionContext::Execute(uint32 NumInstances, TArray<FNiagar
 	if (bRegisterSetupCompleted)
 	{
 		VectorVM::Exec(
-			Script->GetByteCode().GetData(),
+			Script->GetVMExecutableData().ByteCode.GetData(),
 			InputRegisters,
 			NumInputRegisters,
 			OutputRegisters,
@@ -400,19 +187,10 @@ bool FNiagaraScriptExecutionContext::Execute(uint32 NumInstances, TArray<FNiagar
 		FNiagaraDataSetExecutionInfo& Info = DataSetInfos[Idx];
 		if (Info.bUpdateInstanceCount)
 		{
-			Info.DataSet->SetNumInstances(Info.StartInstance + DataSetMetaTable[Idx].DataSetAccessIndex);
+			Info.DataSet->SetNumInstances(Info.StartInstance + DataSetMetaTable[Idx].DataSetAccessIndex + 1);
 		}
 	}
 
-#if WITH_EDITORONLY_DATA
-	if (Script->GetDebuggerInfo().bRequestDebugFrame)
-	{
-		DataSetInfos[0].DataSet->Dump(Script->GetDebuggerInfo().DebugFrame, true);
-		Script->GetDebuggerInfo().bRequestDebugFrame = false;
-		Script->GetDebuggerInfo().DebugFrameLastWriteId = TickCounter;
-		Script->GetDebuggerInfo().DebugParameters = Parameters;
-	}
-#endif
 	return true;//TODO: Error cases?
 }
 
@@ -423,5 +201,5 @@ void FNiagaraScriptExecutionContext::DirtyDataInterfaces()
 
 bool FNiagaraScriptExecutionContext::CanExecute()const
 {
-	return Script && Script->GetByteCode().Num() > 0;
-}
\ No newline at end of file
+	return Script && Script->GetVMExecutableData().IsValid() && Script->GetVMExecutableData().ByteCode.Num() > 0;
+}
diff --git a/Engine/Plugins/FX/Niagara/Source/Niagara/Private/NiagaraScriptExecutionParameterStore.cpp b/Engine/Plugins/FX/Niagara/Source/Niagara/Private/NiagaraScriptExecutionParameterStore.cpp
new file mode 100644
--- /dev/null
+++ b/Engine/Plugins/FX/Niagara/Source/Niagara/Private/NiagaraScriptExecutionParameterStore.cpp
@@ -0,0 +1,301 @@
+// Copyright 1998-2016 Epic Games, Inc. All Rights Reserved.
+
+#include "NiagaraScriptExecutionParameterStore.h"
+#include "NiagaraStats.h"
+#include "NiagaraEmitterInstanceBatcher.h"
+#include "NiagaraDataInterface.h"
+#include "NiagaraSystemInstance.h"
+
+FNiagaraScriptExecutionParameterStore::FNiagaraScriptExecutionParameterStore()
+	: FNiagaraParameterStore(), bInitialized(false)
+{
+
+}
+
+FNiagaraScriptExecutionParameterStore::FNiagaraScriptExecutionParameterStore(const FNiagaraParameterStore& Other)
+{
+	*this = Other;
+}
+
+FNiagaraScriptExecutionParameterStore& FNiagaraScriptExecutionParameterStore::operator=(const FNiagaraParameterStore& Other)
+{
+	FNiagaraParameterStore::operator=(Other);
+	return *this;
+}
+
+uint32 OffsetAlign(uint32 SrcOffset, uint32 Size)
+{
+	uint32 OffsetRemaining = UNIFORM_BUFFER_STRUCT_ALIGNMENT - (SrcOffset % UNIFORM_BUFFER_STRUCT_ALIGNMENT);
+	if (Size <= OffsetRemaining)
+	{
+		return SrcOffset;
+	}
+	else
+	{
+		return Align(SrcOffset, UNIFORM_BUFFER_STRUCT_ALIGNMENT);
+	}
+}
+
+void FNiagaraScriptExecutionParameterStore::GenerateLayoutInfoInternal(TArray<FNiagaraScriptExecutionPaddingInfo>& Members, uint32& NextMemberOffset, const UStruct* InSrcStruct, uint32 InSrcOffset)
+{
+	// Now insert an appropriate data member into the mix...
+	if (InSrcStruct == FNiagaraTypeDefinition::GetBoolStruct() || InSrcStruct == FNiagaraTypeDefinition::GetIntStruct())
+	{
+		//Members.Emplace(*Property->GetName(), TEXT(""), InSrcOffset, (EUniformBufferBaseType)TUniformBufferTypeInfo<uint32>::BaseType, EShaderPrecisionModifier::Float, TUniformBufferTypeInfo<uint32>::NumRows, TUniformBufferTypeInfo<uint32>::NumColumns, TUniformBufferTypeInfo<uint32>::NumElements, TUniformBufferTypeInfo<uint32>::GetStruct());
+		Members.Emplace(InSrcOffset, Align(NextMemberOffset, TUniformBufferTypeInfo<uint32, false>::Alignment), (TUniformBufferTypeInfo<uint32, false>::NumRows * TUniformBufferTypeInfo<uint32, false>::NumColumns) * sizeof(uint32));
+		InSrcOffset += sizeof(uint32);
+		NextMemberOffset = Members[Members.Num() - 1].DestOffset + Members[Members.Num() - 1].Size;
+	}
+	else if (InSrcStruct == FNiagaraTypeDefinition::GetFloatStruct())
+	{
+		//Members.Emplace(*Property->GetName(), TEXT(""), InSrcOffset, (EUniformBufferBaseType)TUniformBufferTypeInfo < float > ::BaseType, EShaderPrecisionModifier::Float, TUniformBufferTypeInfo<float>::NumRows, TUniformBufferTypeInfo<float>::NumColumns, TUniformBufferTypeInfo<float>::NumElements, TUniformBufferTypeInfo<float>::GetStruct());
+		Members.Emplace(InSrcOffset, Align(NextMemberOffset, TUniformBufferTypeInfo<float, false>::Alignment), (TUniformBufferTypeInfo<float, false>::NumRows * TUniformBufferTypeInfo<float, false>::NumColumns) * sizeof(float));
+		InSrcOffset += sizeof(float);
+		NextMemberOffset = Members[Members.Num() - 1].DestOffset + Members[Members.Num() - 1].Size;
+	}
+	else if (InSrcStruct == FNiagaraTypeDefinition::GetVec2Struct())
+	{
+		//Members.Emplace(*Property->GetName(), TEXT(""), InSrcOffset, (EUniformBufferBaseType)TUniformBufferTypeInfo < FVector2D > ::BaseType, EShaderPrecisionModifier::Float, TUniformBufferTypeInfo<FVector2D>::NumRows, TUniformBufferTypeInfo<FVector2D>::NumColumns, TUniformBufferTypeInfo<FVector2D>::NumElements, TUniformBufferTypeInfo<FVector2D>::GetStruct());
+		uint32 StructFinalSize = (TUniformBufferTypeInfo<FVector2D, false>::NumRows * TUniformBufferTypeInfo<FVector2D, false>::NumColumns) * sizeof(float);
+		Members.Emplace(InSrcOffset, OffsetAlign(NextMemberOffset, StructFinalSize),StructFinalSize);
+		InSrcOffset += sizeof(FVector2D);
+		NextMemberOffset = Members[Members.Num() - 1].DestOffset + Members[Members.Num() - 1].Size;
+	}
+	else if (InSrcStruct == FNiagaraTypeDefinition::GetVec3Struct())
+	{
+		//Members.Emplace(*Property->GetName(), TEXT(""), InSrcOffset, (EUniformBufferBaseType)TUniformBufferTypeInfo < FVector > ::BaseType, EShaderPrecisionModifier::Float, TUniformBufferTypeInfo<FVector>::NumRows, TUniformBufferTypeInfo<FVector>::NumColumns, TUniformBufferTypeInfo<FVector>::NumElements, TUniformBufferTypeInfo<FVector>::GetStruct());
+		uint32 StructFinalSize = (TUniformBufferTypeInfo<FVector, false>::NumRows * TUniformBufferTypeInfo<FVector, false>::NumColumns) * sizeof(float);
+		Members.Emplace(InSrcOffset, OffsetAlign(NextMemberOffset, StructFinalSize), StructFinalSize);
+		InSrcOffset += sizeof(FVector);
+		NextMemberOffset = Members[Members.Num() - 1].DestOffset + Members[Members.Num() - 1].Size;
+	}
+	else if (InSrcStruct == FNiagaraTypeDefinition::GetVec4Struct() || InSrcStruct == FNiagaraTypeDefinition::GetColorStruct() || InSrcStruct == FNiagaraTypeDefinition::GetQuatStruct())
+	{
+		//Members.Emplace(*Property->GetName(), TEXT(""), InSrcOffset, (EUniformBufferBaseType)TUniformBufferTypeInfo < FVector4 > ::BaseType, EShaderPrecisionModifier::Float, TUniformBufferTypeInfo<FVector4>::NumRows, TUniformBufferTypeInfo<FVector4>::NumColumns, TUniformBufferTypeInfo<FVector4>::NumElements, TUniformBufferTypeInfo<FVector4>::GetStruct());
+		Members.Emplace(InSrcOffset, Align(NextMemberOffset, TUniformBufferTypeInfo<FVector4, false>::Alignment), (TUniformBufferTypeInfo<FVector4, false>::NumRows * TUniformBufferTypeInfo<FVector4, false>::NumColumns) * sizeof(float));
+		InSrcOffset += sizeof(FVector4);
+		NextMemberOffset = Members[Members.Num() - 1].DestOffset + Members[Members.Num() - 1].Size;
+	}
+	else if (InSrcStruct == FNiagaraTypeDefinition::GetMatrix4Struct())
+	{
+		//Members.Emplace(*Property->GetName(), TEXT(""), InSrcOffset, (EUniformBufferBaseType)TUniformBufferTypeInfo < FMatrix > ::BaseType, EShaderPrecisionModifier::Float, TUniformBufferTypeInfo<FMatrix>::NumRows, TUniformBufferTypeInfo<FMatrix>::NumColumns, TUniformBufferTypeInfo<FMatrix>::NumElements, TUniformBufferTypeInfo<FMatrix>::GetStruct());
+		Members.Emplace(InSrcOffset, Align(NextMemberOffset, TUniformBufferTypeInfo<FMatrix, false>::Alignment), (TUniformBufferTypeInfo<FMatrix, false>::NumRows * TUniformBufferTypeInfo<FMatrix, false>::NumColumns) * sizeof(float));
+		InSrcOffset += sizeof(FMatrix);
+		NextMemberOffset = Members[Members.Num() - 1].DestOffset + Members[Members.Num() - 1].Size;
+	}
+	else
+	{
+		NextMemberOffset = Align(NextMemberOffset, UNIFORM_BUFFER_STRUCT_ALIGNMENT); // New structs should be aligned to the head..
+
+		for (TFieldIterator<UProperty> PropertyIt(InSrcStruct, EFieldIteratorFlags::IncludeSuper); PropertyIt; ++PropertyIt)
+		{
+			UProperty* Property = *PropertyIt;
+			const UStruct* Struct = nullptr;
+
+			// First determine what struct type we're dealing with...
+			if (Property->IsA(UFloatProperty::StaticClass()))
+			{
+				Struct = FNiagaraTypeDefinition::GetFloatStruct();
+			}
+			else if (Property->IsA(UIntProperty::StaticClass()))
+			{
+				Struct = FNiagaraTypeDefinition::GetIntStruct();
+			}
+			else if (Property->IsA(UBoolProperty::StaticClass()))
+			{
+				Struct = FNiagaraTypeDefinition::GetBoolStruct();
+			}
+			//Should be able to support double easily enough
+			else if (UStructProperty* StructProp = CastChecked<UStructProperty>(Property))
+			{
+				Struct = StructProp->Struct;
+			}
+			else
+			{
+				check(false);
+			}
+
+			GenerateLayoutInfoInternal(Members, NextMemberOffset, Struct, InSrcOffset);
+		}
+	}
+
+}
+
+void FNiagaraScriptExecutionParameterStore::AddPaddedParamSize(const FNiagaraTypeDefinition& InParamType, uint32 InOffset)
+{
+	if (InParamType.IsDataInterface())
+	{
+		return;
+	}
+
+	uint32 NextMemberOffset = 0;
+	if (PaddingInfo.Num() != 0)
+	{
+		NextMemberOffset = PaddingInfo[PaddingInfo.Num() - 1].DestOffset + PaddingInfo[PaddingInfo.Num() - 1].Size;
+	}
+	GenerateLayoutInfoInternal(PaddingInfo, NextMemberOffset, InParamType.GetScriptStruct(), InOffset);
+
+	if (PaddingInfo.Num() != 0)
+	{
+		NextMemberOffset = PaddingInfo[PaddingInfo.Num() - 1].DestOffset + PaddingInfo[PaddingInfo.Num() - 1].Size;
+		PaddedParameterSize = Align(NextMemberOffset, UNIFORM_BUFFER_STRUCT_ALIGNMENT);
+	}
+	else
+	{
+		PaddedParameterSize = 0;
+	}
+}
+
+
+void FNiagaraScriptExecutionParameterStore::InitFromOwningScript(UNiagaraScript* Script, ENiagaraSimTarget SimTarget, bool bNotifyAsDirty)
+{
+	//TEMPORARTY:
+	//We should replace the storage on the script with an FNiagaraParameterStore also so we can just copy that over here. Though that is an even bigger refactor job so this is a convenient place to break that work up.
+
+	Empty();
+	PaddedParameterSize = 0;
+	PaddingInfo.Empty();
+
+	if (Script)
+	{
+		AddScriptParams(Script, SimTarget, false);
+
+		Script->RapidIterationParameters.Bind(this);
+
+		if (bNotifyAsDirty)
+		{
+			MarkParametersDirty();
+			MarkInterfacesDirty();
+			OnLayoutChange();
+		}
+	}
+	bInitialized = true;
+}
+
+void FNiagaraScriptExecutionParameterStore::InitFromOwningContext(UNiagaraScript* Script, ENiagaraSimTarget SimTarget, bool bNotifyAsDirty)
+{
+	Clear();
+	PaddingInfo.Empty();
+
+#if WITH_EDITORONLY_DATA
+	if (Script != nullptr)
+	{
+		DebugName = FString::Printf(TEXT("ScriptExecParamStore %s %p"), *Script->GetFullName(), this);
+	}
+#endif
+
+	const FNiagaraScriptExecutionParameterStore* SrcStore = Script->GetExecutionReadyParameterStore(SimTarget);
+	if (SrcStore)
+	{
+		InitFromSource(SrcStore, false);
+		ParameterSize = SrcStore->ParameterSize;
+		PaddedParameterSize = SrcStore->PaddedParameterSize;
+		PaddingInfo = SrcStore->PaddingInfo;
+
+		if (bNotifyAsDirty)
+		{
+			MarkParametersDirty();
+			MarkInterfacesDirty();
+			OnLayoutChange();
+		}
+	}
+
+
+	bInitialized = true;
+}
+
+void FNiagaraScriptExecutionParameterStore::AddScriptParams(UNiagaraScript* Script, ENiagaraSimTarget SimTarget, bool bTriggerRebind)
+{
+	if (Script == nullptr)
+	{
+		return;
+	}
+	PaddingInfo.Empty();
+
+	//Here we add the current frame parameters.
+	bool bAdded = false;
+	for (FNiagaraVariable& Param : Script->GetVMExecutableData().Parameters.Parameters)
+	{
+		bAdded |= AddParameter(Param, false, false);
+	}
+
+#if WITH_EDITORONLY_DATA
+	DebugName = FString::Printf(TEXT("ScriptExecParamStore %s %p"), *Script->GetFullName(), this);
+#endif
+	
+	//Add previous frame values if we're interpolated spawn.
+	bool bIsInterpolatedSpawn = Script->GetVMExecutableDataCompilationId().HasInterpolatedParameters();
+
+	if (bIsInterpolatedSpawn)
+	{
+		for (FNiagaraVariable& Param : Script->GetVMExecutableData().Parameters.Parameters)
+		{
+			FNiagaraVariable PrevParam(Param.GetType(), FName(*(INTERPOLATED_PARAMETER_PREFIX + Param.GetName().ToString())));
+			bAdded |= AddParameter(PrevParam, false, false);
+		}
+	}
+	
+	ParameterSize = GetParameterDataArray().Num();
+	if (bIsInterpolatedSpawn)
+	{
+		CopyCurrToPrev();
+		bAdded = true;
+	}
+
+	//Internal constants - only needed for non-GPU sim
+	if (SimTarget != ENiagaraSimTarget::GPUComputeSim)
+	{
+		for (FNiagaraVariable& InternalVar : Script->GetVMExecutableData().InternalParameters.Parameters)
+		{
+			bAdded |= AddParameter(InternalVar, false, false);
+		}
+	}
+
+	check(Script->GetVMExecutableData().DataInterfaceInfo.Num() == Script->GetCachedDefaultDataInterfaces().Num());
+	for (FNiagaraScriptDataInterfaceInfo& Info : Script->GetCachedDefaultDataInterfaces())
+	{
+		FName ParameterName;
+		if (Info.RegisteredParameterMapRead != NAME_None)
+		{
+			ParameterName = Info.RegisteredParameterMapRead;
+		}
+		else
+		{
+			// If the data interface wasn't used in a parameter map, mangle the name so that it doesn't accidentally bind to
+			// a valid parameter.
+			ParameterName = *(TEXT("__INTERNAL__.") + Info.Name.ToString());
+		}
+
+		FNiagaraVariable Var(Info.Type, ParameterName);
+		bAdded |= AddParameter(Var, false, false);
+		SetDataInterface(Info.DataInterface, IndexOf(Var));
+	}
+
+	if (bAdded && bTriggerRebind)
+	{
+		OnLayoutChange();
+	}
+}
+
+void FNiagaraScriptExecutionParameterStore::CopyCurrToPrev()
+{
+	int32 ParamStart = ParameterSize / 2;
+	checkSlow(FMath::Frac((float)ParameterSize / 2) == 0.0f);
+
+	FMemory::Memcpy(GetParameterData_Internal(ParamStart), GetParameterData(0), ParamStart);
+}
+
+
+void FNiagaraScriptExecutionParameterStore::CopyParameterDataToPaddedBuffer(uint8* InTargetBuffer, uint32 InTargetBufferSizeInBytes)
+{
+	check((uint32)ParameterSize <= PaddedParameterSize);
+	check(InTargetBufferSizeInBytes >= PaddedParameterSize);
+	FMemory::Memzero(InTargetBuffer, InTargetBufferSizeInBytes);
+	const uint8* SrcData = GetParameterDataArray().GetData();
+	for (int32 i = 0; i < PaddingInfo.Num(); i++)
+	{
+		check((PaddingInfo[i].DestOffset + PaddingInfo[i].Size) <= InTargetBufferSizeInBytes);
+		check((PaddingInfo[i].SrcOffset + PaddingInfo[i].Size) <= (uint32)GetParameterDataArray().Num());
+		FMemory::Memcpy(InTargetBuffer + PaddingInfo[i].DestOffset, SrcData + PaddingInfo[i].SrcOffset, PaddingInfo[i].Size);
+	}
+}
diff --git a/Engine/Plugins/FX/Niagara/Source/Niagara/Private/NiagaraSpriteRendererProperties.cpp b/Engine/Plugins/FX/Niagara/Source/Niagara/Private/NiagaraSpriteRendererProperties.cpp
--- a/Engine/Plugins/FX/Niagara/Source/Niagara/Private/NiagaraSpriteRendererProperties.cpp
+++ b/Engine/Plugins/FX/Niagara/Source/Niagara/Private/NiagaraSpriteRendererProperties.cpp
@@ -12,12 +12,14 @@ UNiagaraSpriteRendererProperties::UNiagaraSpriteRendererProperties()
 	, FacingMode(ENiagaraSpriteFacingMode::FaceCamera)
 	, CustomFacingVectorMask(EForceInit::ForceInitToZero)
 	, PivotInUVSpace(0.5f, 0.5f)
-	, SortMode(ENiagaraSortMode::SortNone)
+	, SortMode(ENiagaraSortMode::ViewDistance)
 	, SubImageSize(1.0f, 1.0f)
 	, bSubImageBlend(false)
 	, bRemoveHMDRollInVR(false)
+	, bSortOnlyWhenTranslucent(true)
 	, MinFacingCameraBlendDistance(0.0f)
 	, MaxFacingCameraBlendDistance(0.0f)
+	, SyncId(0)
 {
 }
 
@@ -31,8 +33,59 @@ void UNiagaraSpriteRendererProperties::GetUsedMaterials(TArray<UMaterialInterfac
 	OutMaterials.Add(Material);
 }
 
+void UNiagaraSpriteRendererProperties::PostInitProperties()
+{
+	Super::PostInitProperties();
+	SyncId = 0;
+	if (HasAnyFlags(RF_ClassDefaultObject) == false)
+	{
+		InitBindings();
+	}
+}
+
+/** The bindings depend on variables that are created during the NiagaraModule startup. However, the CDO's are build prior to this being initialized, so we defer setting these values until later.*/
+void UNiagaraSpriteRendererProperties::InitCDOPropertiesAfterModuleStartup()
+{
+	UNiagaraSpriteRendererProperties* CDO = CastChecked<UNiagaraSpriteRendererProperties>(UNiagaraSpriteRendererProperties::StaticClass()->GetDefaultObject());
+	CDO->InitBindings();
+}
+
+void UNiagaraSpriteRendererProperties::InitBindings()
+{
+	if (PositionBinding.BoundVariable.GetName() == NAME_None)
+	{
+		PositionBinding = FNiagaraConstants::GetAttributeDefaultBinding(SYS_PARAM_PARTICLES_POSITION);
+		ColorBinding = FNiagaraConstants::GetAttributeDefaultBinding(SYS_PARAM_PARTICLES_COLOR);
+		VelocityBinding = FNiagaraConstants::GetAttributeDefaultBinding(SYS_PARAM_PARTICLES_VELOCITY);
+		SpriteRotationBinding = FNiagaraConstants::GetAttributeDefaultBinding(SYS_PARAM_PARTICLES_SPRITE_ROTATION);
+		SpriteSizeBinding = FNiagaraConstants::GetAttributeDefaultBinding(SYS_PARAM_PARTICLES_SPRITE_SIZE);
+		SpriteFacingBinding = FNiagaraConstants::GetAttributeDefaultBinding(SYS_PARAM_PARTICLES_SPRITE_FACING);
+		SpriteAlignmentBinding = FNiagaraConstants::GetAttributeDefaultBinding(SYS_PARAM_PARTICLES_SPRITE_ALIGNMENT);
+		SubImageIndexBinding = FNiagaraConstants::GetAttributeDefaultBinding(SYS_PARAM_PARTICLES_SUB_IMAGE_INDEX);
+		DynamicMaterialBinding = FNiagaraConstants::GetAttributeDefaultBinding(SYS_PARAM_PARTICLES_DYNAMIC_MATERIAL_PARAM);
+		DynamicMaterial1Binding = FNiagaraConstants::GetAttributeDefaultBinding(SYS_PARAM_PARTICLES_DYNAMIC_MATERIAL_PARAM_1);
+		DynamicMaterial2Binding = FNiagaraConstants::GetAttributeDefaultBinding(SYS_PARAM_PARTICLES_DYNAMIC_MATERIAL_PARAM_2);
+		DynamicMaterial3Binding = FNiagaraConstants::GetAttributeDefaultBinding(SYS_PARAM_PARTICLES_DYNAMIC_MATERIAL_PARAM_3);
+		CameraOffsetBinding = FNiagaraConstants::GetAttributeDefaultBinding(SYS_PARAM_PARTICLES_CAMERA_OFFSET);
+		UVScaleBinding = FNiagaraConstants::GetAttributeDefaultBinding(SYS_PARAM_PARTICLES_UV_SCALE);
+		MaterialRandomBinding = FNiagaraConstants::GetAttributeDefaultBinding(SYS_PARAM_PARTICLES_MATERIAL_RANDOM);
+		
+		//Default custom sorting to age
+		CustomSortingBinding = FNiagaraConstants::GetAttributeDefaultBinding(SYS_PARAM_PARTICLES_NORMALIZED_AGE);
+	}
+}
+
 
 #if WITH_EDITORONLY_DATA
+void UNiagaraSpriteRendererProperties::PostEditChangeProperty(struct FPropertyChangedEvent& PropertyChangedEvent) 
+{
+	Super::PostEditChangeProperty(PropertyChangedEvent);
+	if (PropertyChangedEvent.GetPropertyName() != TEXT("SyncId"))
+	{
+		SyncId++;
+	}
+}
+
 
 const TArray<FNiagaraVariable>& UNiagaraSpriteRendererProperties::GetRequiredAttributes()
 {
@@ -40,12 +93,6 @@ const TArray<FNiagaraVariable>& UNiagaraSpriteRendererProperties::GetRequiredAtt
 
 	if (Attrs.Num() == 0)
 	{
-		Attrs.Add(SYS_PARAM_PARTICLES_POSITION);
-		Attrs.Add(SYS_PARAM_PARTICLES_VELOCITY);
-		Attrs.Add(SYS_PARAM_PARTICLES_COLOR);
-		Attrs.Add(SYS_PARAM_PARTICLES_SPRITE_ROTATION);
-		Attrs.Add(SYS_PARAM_PARTICLES_NORMALIZED_AGE);
-		Attrs.Add(SYS_PARAM_PARTICLES_SPRITE_SIZE);
 	}
 
 	return Attrs;
@@ -58,10 +105,19 @@ const TArray<FNiagaraVariable>& UNiagaraSpriteRendererProperties::GetOptionalAtt
 
 	if (Attrs.Num() == 0)
 	{
+		Attrs.Add(SYS_PARAM_PARTICLES_POSITION);
+		Attrs.Add(SYS_PARAM_PARTICLES_VELOCITY);
+		Attrs.Add(SYS_PARAM_PARTICLES_COLOR);
+		Attrs.Add(SYS_PARAM_PARTICLES_SPRITE_ROTATION);
+		Attrs.Add(SYS_PARAM_PARTICLES_NORMALIZED_AGE);
+		Attrs.Add(SYS_PARAM_PARTICLES_SPRITE_SIZE);
 		Attrs.Add(SYS_PARAM_PARTICLES_SPRITE_FACING);
 		Attrs.Add(SYS_PARAM_PARTICLES_SPRITE_ALIGNMENT);
 		Attrs.Add(SYS_PARAM_PARTICLES_SUB_IMAGE_INDEX);
 		Attrs.Add(SYS_PARAM_PARTICLES_DYNAMIC_MATERIAL_PARAM);
+		Attrs.Add(SYS_PARAM_PARTICLES_DYNAMIC_MATERIAL_PARAM_1);
+		Attrs.Add(SYS_PARAM_PARTICLES_DYNAMIC_MATERIAL_PARAM_2);
+		Attrs.Add(SYS_PARAM_PARTICLES_DYNAMIC_MATERIAL_PARAM_3);
 		Attrs.Add(SYS_PARAM_PARTICLES_CAMERA_OFFSET);
 		Attrs.Add(SYS_PARAM_PARTICLES_UV_SCALE);
 		Attrs.Add(SYS_PARAM_PARTICLES_MATERIAL_RANDOM);
diff --git a/Engine/Plugins/FX/Niagara/Source/Niagara/Private/NiagaraStats.h b/Engine/Plugins/FX/Niagara/Source/Niagara/Private/NiagaraStats.h
--- a/Engine/Plugins/FX/Niagara/Source/Niagara/Private/NiagaraStats.h
+++ b/Engine/Plugins/FX/Niagara/Source/Niagara/Private/NiagaraStats.h
@@ -9,6 +9,10 @@
 
 DECLARE_STATS_GROUP(TEXT("Niagara"), STATGROUP_Niagara, STATCAT_Advanced);
 DECLARE_CYCLE_STAT(TEXT("Render Total"), STAT_NiagaraRender, STATGROUP_Niagara);
+DECLARE_CYCLE_STAT(TEXT("Render Total (GT)"), STAT_NiagaraRenderGT, STATGROUP_Niagara);
 DECLARE_DWORD_COUNTER_STAT(TEXT("NumParticles"), STAT_NiagaraNumParticles, STATGROUP_Niagara);
 DECLARE_CYCLE_STAT(TEXT("Constant Setup"), STAT_NiagaraConstants, STATGROUP_Niagara);
 DECLARE_CYCLE_STAT(TEXT("Tick"), STAT_NiagaraTick, STATGROUP_Niagara);
+DECLARE_DWORD_COUNTER_STAT(TEXT("NumSystems"), STAT_NiagaraNumSystems, STATGROUP_Niagara);
+DECLARE_MEMORY_STAT(TEXT("Niagara particle data memory"), STAT_NiagaraParticleMemory, STATGROUP_Niagara);
+DECLARE_MEMORY_STAT(TEXT("Niagara vertex buffer memory"), STAT_NiagaraVBMemory, STATGROUP_Niagara);
