Copying //UE4/Dev-Niagara to //UE4/Dev-Main (Source: //UE4/Dev-Niagara @ 4074996)

diff --git a/Engine/Plugins/FX/Niagara/Source/Niagara/Private/NiagaraSystem.cpp b/Engine/Plugins/FX/Niagara/Source/Niagara/Private/NiagaraSystem.cpp
--- a/Engine/Plugins/FX/Niagara/Source/Niagara/Private/NiagaraSystem.cpp
+++ b/Engine/Plugins/FX/Niagara/Source/Niagara/Private/NiagaraSystem.cpp
@@ -14,7 +14,14 @@
 #include "NiagaraEmitterHandle.h"
 #include "AssetData.h"
 #include "NiagaraStats.h"
+#include "NiagaraEditorDataBase.h"
 
+#if WITH_EDITOR
+#include "NiagaraScriptDerivedData.h"
+#include "DerivedDataCacheInterface.h"
+#endif
+
+DECLARE_CYCLE_STAT(TEXT("Niagara - System - Precompile"), STAT_Niagara_System_Precompile, STATGROUP_Niagara);
 DECLARE_CYCLE_STAT(TEXT("Niagara - System - CompileScript"), STAT_Niagara_System_CompileScript, STATGROUP_Niagara);
 DECLARE_CYCLE_STAT(TEXT("Niagara - System - CompileScript_ResetAfter"), STAT_Niagara_System_CompileScriptResetAfter, STATGROUP_Niagara);
 
@@ -32,9 +39,41 @@ UNiagaraSystem::UNiagaraSystem(const FObjectInitializer& ObjectInitializer)
 , bIsolateEnabled(false)
 #endif
 , bAutoDeactivate(true)
+, WarmupTime(0.0f)
+, WarmupTickCount(0)
+, WarmupTickDelta(1.0f / 15.0f)
+{
+}
+
+void UNiagaraSystem::BeginDestroy()
 {
+	Super::BeginDestroy();
+#if WITH_EDITORONLY_DATA
+	while (ActiveCompilations.Num() > 0)
+	{
+		QueryCompileComplete(true, false, true);
+	}
+#endif
 }
 
+void UNiagaraSystem::PreSave(const class ITargetPlatform * TargetPlatform)
+{
+	Super::PreSave(TargetPlatform);
+#if WITH_EDITORONLY_DATA
+	WaitForCompilationComplete();
+#endif
+}
+
+#if WITH_EDITOR
+void UNiagaraSystem::BeginCacheForCookedPlatformData(const ITargetPlatform *TargetPlatform)
+{
+	Super::BeginCacheForCookedPlatformData(TargetPlatform);
+#if WITH_EDITORONLY_DATA
+	WaitForCompilationComplete();
+#endif
+}
+#endif
+
 void UNiagaraSystem::PostInitProperties()
 {
 	Super::PostInitProperties();
@@ -48,12 +87,6 @@ void UNiagaraSystem::PostInitProperties()
 
 		SystemUpdateScript = NewObject<UNiagaraScript>(this, "SystemUpdateScript", RF_Transactional);
 		SystemUpdateScript->SetUsage(ENiagaraScriptUsage::SystemUpdateScript);
-
-		SystemSpawnScriptSolo = NewObject<UNiagaraScript>(this, "SystemSpawnScriptSolo", RF_Transactional);
-		SystemSpawnScriptSolo->SetUsage(ENiagaraScriptUsage::SystemSpawnScript);
-
-		SystemUpdateScriptSolo = NewObject<UNiagaraScript>(this, "SystemUpdateScriptSolo", RF_Transactional);
-		SystemUpdateScriptSolo->SetUsage(ENiagaraScriptUsage::SystemUpdateScript);
 	}
 }
 
@@ -65,9 +98,7 @@ bool UNiagaraSystem::IsLooping() const
 bool UNiagaraSystem::UsesCollection(const UNiagaraParameterCollection* Collection)const
 {
 	if (SystemSpawnScript->UsesCollection(Collection) ||
-		SystemSpawnScriptSolo->UsesCollection(Collection) ||
-		SystemUpdateScript->UsesCollection(Collection) ||
-		SystemUpdateScriptSolo->UsesCollection(Collection))
+		SystemUpdateScript->UsesCollection(Collection))
 	{
 		return true;
 	}
@@ -87,9 +118,7 @@ bool UNiagaraSystem::UsesCollection(const UNiagaraParameterCollection* Collectio
 bool UNiagaraSystem::UsesScript(const UNiagaraScript* Script)const
 {
 	if (SystemSpawnScript == Script ||
-		SystemSpawnScriptSolo == Script ||
-		SystemUpdateScript == Script ||
-		SystemUpdateScriptSolo == Script)
+		SystemUpdateScript == Script)
 	{
 		return true;
 	}
@@ -104,6 +133,31 @@ bool UNiagaraSystem::UsesScript(const UNiagaraScript* Script)const
 	
 	return false;
 }
+
+bool UNiagaraSystem::UsesEmitter(const UNiagaraEmitter* Emitter) const
+{
+	for (FNiagaraEmitterHandle EmitterHandle : GetEmitterHandles())
+	{
+		if (Emitter == EmitterHandle.GetSource() || Emitter == EmitterHandle.GetInstance())
+		{
+			return true;
+		}
+	}
+	return false;
+}
+
+void UNiagaraSystem::RequestCompileForEmitter(UNiagaraEmitter* InEmitter)
+{
+	for (TObjectIterator<UNiagaraSystem> It; It; ++It)
+	{
+		UNiagaraSystem* Sys = *It;
+		if (Sys && Sys->UsesEmitter(InEmitter))
+		{
+			Sys->RequestCompile(false);
+		}
+	}
+}
+
 #endif
 
 void UNiagaraSystem::Serialize(FArchive& Ar)
@@ -120,6 +174,30 @@ void UNiagaraSystem::PostEditChangeProperty(struct FPropertyChangedEvent& Proper
 	FNiagaraSystemUpdateContext(this, true);
 
 	ThumbnailImageOutOfDate = true;
+
+	DetermineIfSolo();
+	
+	if (PropertyChangedEvent.Property != nullptr)
+	{
+		if (PropertyChangedEvent.Property->GetFName() == GET_MEMBER_NAME_CHECKED(UNiagaraSystem, WarmupTickCount))
+		{
+			//Set the WarmupTime to feed back to the user.
+			WarmupTime = WarmupTickCount * WarmupTickDelta;
+		}
+		else if (PropertyChangedEvent.Property->GetFName() == GET_MEMBER_NAME_CHECKED(UNiagaraSystem, WarmupTime))
+		{
+			//Set the WarmupTickCount to feed back to the user.
+			if (FMath::IsNearlyZero(WarmupTickDelta))
+			{
+				WarmupTickDelta = 0.0f;
+			}
+			else
+			{
+				WarmupTickCount = WarmupTime / WarmupTickDelta;
+				WarmupTime = WarmupTickDelta * WarmupTickCount;
+			}
+		}
+	}
 }
 #endif 
 
@@ -132,14 +210,6 @@ void UNiagaraSystem::PostLoad()
 		SetFlags(RF_Transactional);
 	}
 
-	// Check to see if our version is out of date. If so, we'll finally need to recompile.
-	bool bNeedsRecompile = false;
-	const int32 NiagaraVer = GetLinkerCustomVersion(FNiagaraCustomVersion::GUID);
-	if (NiagaraVer < FNiagaraCustomVersion::LatestVersion)
-	{
-		bNeedsRecompile = true;
-	}
-
 	// Previously added emitters didn't have their stand alone and public flags cleared so
 	// they 'leak' into the system package.  Clear the flags here so they can be collected
 	// during the next save.
@@ -158,78 +228,81 @@ void UNiagaraSystem::PostLoad()
 			}
 		}
 	}
-	
-	// We will be using this later potentially, so make sure that it's postload is already up-to-date.
-	if (SystemSpawnScript)
-	{
-		SystemSpawnScript->ConditionalPostLoad();
-	}
-	if (SystemUpdateScript)
-	{
-		SystemUpdateScript->ConditionalPostLoad();
-	}
-
-	if (NiagaraVer < FNiagaraCustomVersion::SystemEmitterScriptSimulations)
-	{
-		bNeedsRecompile = true;
-		SystemSpawnScriptSolo = NewObject<UNiagaraScript>(this, "SystemSpawnScriptSolo", RF_Transactional);
-		SystemSpawnScriptSolo->SetUsage(ENiagaraScriptUsage::SystemSpawnScript);
-
-		SystemUpdateScriptSolo = NewObject<UNiagaraScript>(this, "SystemUpdateScriptSolo", RF_Transactional);
-		SystemUpdateScriptSolo->SetUsage(ENiagaraScriptUsage::SystemUpdateScript);
-#if WITH_EDITORONLY_DATA
-		if (SystemSpawnScript)
-		{
-			SystemSpawnScriptSolo->SetSource(SystemSpawnScript->GetSource());
-		}
-
-		if (SystemUpdateScript)
-		{
-			SystemUpdateScriptSolo->SetSource(SystemUpdateScript->GetSource());
-		}
-#endif
-	}
-
-	if (bNeedsRecompile)
-	{
-		for (FNiagaraEmitterHandle& Handle : EmitterHandles)
-		{
-			UNiagaraEmitter* Props = Handle.GetInstance();
-			if (Props)
-			{
-				// We will be refreshing later potentially, so make sure that it's postload is already up-to-date.
-				Props->ConditionalPostLoad();
-			}
-		}
-	}
 
 #if WITH_EDITORONLY_DATA
+	TArray<UNiagaraScript*> AllSystemScripts;
+	
+	UNiagaraScriptSourceBase* SystemScriptSource = nullptr;
 	if (SystemSpawnScript == nullptr)
 	{
 		SystemSpawnScript = NewObject<UNiagaraScript>(this, "SystemSpawnScript", RF_Transactional);
 		SystemSpawnScript->SetUsage(ENiagaraScriptUsage::SystemSpawnScript);
+		INiagaraModule& NiagaraModule = FModuleManager::GetModuleChecked<INiagaraModule>("Niagara");
+		SystemScriptSource = NiagaraModule.CreateDefaultScriptSource(this);
+		SystemSpawnScript->SetSource(SystemScriptSource);
+	}
+	else
+	{
+		SystemSpawnScript->ConditionalPostLoad();
+		SystemScriptSource = SystemSpawnScript->GetSource();
 	}
+	AllSystemScripts.Add(SystemSpawnScript);
 
 	if (SystemUpdateScript == nullptr)
 	{
 		SystemUpdateScript = NewObject<UNiagaraScript>(this, "SystemUpdateScript", RF_Transactional);
 		SystemUpdateScript->SetUsage(ENiagaraScriptUsage::SystemUpdateScript);
+		SystemUpdateScript->SetSource(SystemScriptSource);
+	}
+	else
+	{
+		SystemUpdateScript->ConditionalPostLoad();
 	}
+	AllSystemScripts.Add(SystemUpdateScript);
+
+	const int32 NiagaraVer = GetLinkerCustomVersion(FNiagaraCustomVersion::GUID);
 
 	//TODO: This causes a crash becuase the script source ptr is null? Fix
 	//For existing emitters before the lifecylce rework, ensure they have the system lifecycle module.
- 	if (NiagaraVer < FNiagaraCustomVersion::LifeCycleRework)
- 	{
- 		/*UNiagaraScriptSourceBase* SystemScriptSource = SystemUpdateScript->GetSource();
- 		if (SystemScriptSource)
- 		{
+	if (NiagaraVer < FNiagaraCustomVersion::LifeCycleRework)
+	{
+		/*UNiagaraScriptSourceBase* SystemScriptSource = SystemUpdateScript->GetSource();
+		if (SystemScriptSource)
+		{
 			bool bFoundModule;
 			if (SystemScriptSource->AddModuleIfMissing(TEXT("/Niagara/Modules/System/SystemLifeCycle.SystemLifeCycle"), ENiagaraScriptUsage::SystemUpdateScript, bFoundModule))
 			{
 				bNeedsRecompile = true;
 			}
- 		}*/
- 	}
+		}*/
+	}
+
+	bool bSystemScriptsAreSynchronized = true;
+	for (UNiagaraScript* SystemScript : AllSystemScripts)
+	{
+		bSystemScriptsAreSynchronized &= SystemScript->AreScriptAndSourceSynchronized();
+	}
+
+	bool bEmitterGraphChangedFromMerge = false;
+	bool bEmitterScriptsAreSynchronized = true;
+
+#if 0
+	UE_LOG(LogNiagara, Log, TEXT("PreMerger"));
+	for (FNiagaraEmitterHandle& EmitterHandle : EmitterHandles)
+	{
+		UE_LOG(LogNiagara, Log, TEXT("Emitter Handle: %s"), *EmitterHandle.GetUniqueInstanceName());
+		UNiagaraScript* UpdateScript = EmitterHandle.GetInstance()->GetScript(ENiagaraScriptUsage::ParticleUpdateScript, FGuid());
+		UNiagaraScript* SpawnScript = EmitterHandle.GetInstance()->GetScript(ENiagaraScriptUsage::ParticleSpawnScript, FGuid());
+		UE_LOG(LogNiagara, Log, TEXT("Spawn Parameters"));
+		SpawnScript->GetVMExecutableData().Parameters.DumpParameters();
+		UE_LOG(LogNiagara, Log, TEXT("Spawn RI Parameters"));
+		SpawnScript->RapidIterationParameters.DumpParameters();
+		UE_LOG(LogNiagara, Log, TEXT("Update Parameters"));
+		UpdateScript->GetVMExecutableData().Parameters.DumpParameters();
+		UE_LOG(LogNiagara, Log, TEXT("Update RI Parameters"));
+		UpdateScript->RapidIterationParameters.DumpParameters();
+	}
+#endif
 
 	for (FNiagaraEmitterHandle& EmitterHandle : EmitterHandles)
 	{
@@ -237,30 +310,104 @@ void UNiagaraSystem::PostLoad()
 		if (EmitterHandle.IsSynchronizedWithSource() == false)
 		{
 			INiagaraModule::FMergeEmitterResults Results = MergeChangesForEmitterHandle(EmitterHandle);
-			bNeedsRecompile |= Results.bSucceeded && Results.bModifiedGraph;
+			if (Results.bSucceeded)
+			{
+				bEmitterGraphChangedFromMerge |= Results.bModifiedGraph;
+			}
+		}
+		if (bEmitterScriptsAreSynchronized)
+		{
+			if (!EmitterHandle.GetInstance()->AreAllScriptAndSourcesSynchronized())
+			{
+				bEmitterScriptsAreSynchronized = false;
+			}
 		}
 	}
 
-	if (bNeedsRecompile)
+	if (EditorData != nullptr)
+	{
+		EditorData->PostLoadFromOwner(this);
+	}
+
+	if (UNiagaraEmitter::GetForceCompileOnLoad())
 	{
-		Compile(false);
+		InvalidateCachedCompileIds();
+		UE_LOG(LogNiagara, Log, TEXT("System %s being rebuilt because UNiagaraEmitter::GetForceCompileOnLoad() == true."), *GetPathName());
+	}
+
+	if (bSystemScriptsAreSynchronized == false)
+	{
+		UE_LOG(LogNiagara, Log, TEXT("System %s being compiled because there were changes to a system script Change ID."), *GetPathName());
+	}
+
+	if (bEmitterScriptsAreSynchronized == false)
+	{
+		UE_LOG(LogNiagara, Log, TEXT("System %s being compiled because there were changes to an emitter script Change ID."), *GetPathName());
+	}
+
+	if (bEmitterGraphChangedFromMerge)
+	{
+		UE_LOG(LogNiagara, Log, TEXT("System %s being compiled because graph changes were merged for a base emitter."), *GetPathName());
+	}
+
+#if 0
+	UE_LOG(LogNiagara, Log, TEXT("Before"));
+	for (FNiagaraEmitterHandle& EmitterHandle : EmitterHandles)
+	{
+		UE_LOG(LogNiagara, Log, TEXT("Emitter Handle: %s"), *EmitterHandle.GetUniqueInstanceName());
+		UNiagaraScript* UpdateScript = EmitterHandle.GetInstance()->GetScript(ENiagaraScriptUsage::ParticleUpdateScript, FGuid());
+		UNiagaraScript* SpawnScript = EmitterHandle.GetInstance()->GetScript(ENiagaraScriptUsage::ParticleSpawnScript, FGuid());
+		UE_LOG(LogNiagara, Log, TEXT("Spawn Parameters"));
+		SpawnScript->GetVMExecutableData().Parameters.DumpParameters();
+		UE_LOG(LogNiagara, Log, TEXT("Spawn RI Parameters"));
+		SpawnScript->RapidIterationParameters.DumpParameters();
+		UE_LOG(LogNiagara, Log, TEXT("Update Parameters"));
+		UpdateScript->GetVMExecutableData().Parameters.DumpParameters();
+		UE_LOG(LogNiagara, Log, TEXT("Update RI Parameters"));
+		UpdateScript->RapidIterationParameters.DumpParameters();
 	}
 #endif
+
+	if (bSystemScriptsAreSynchronized == false || bEmitterScriptsAreSynchronized == false || bEmitterGraphChangedFromMerge)
+	{
+		RequestCompile(false);
+	}
+
+#if 0
+	UE_LOG(LogNiagara, Log, TEXT("After"));
+	for (FNiagaraEmitterHandle& EmitterHandle : EmitterHandles)
+	{
+		UE_LOG(LogNiagara, Log, TEXT("Emitter Handle: %s"), *EmitterHandle.GetUniqueInstanceName());
+		UNiagaraScript* UpdateScript = EmitterHandle.GetInstance()->GetScript(ENiagaraScriptUsage::ParticleUpdateScript, FGuid());
+		UNiagaraScript* SpawnScript = EmitterHandle.GetInstance()->GetScript(ENiagaraScriptUsage::ParticleSpawnScript, FGuid());
+		UE_LOG(LogNiagara, Log, TEXT("Spawn Parameters"));
+		SpawnScript->GetVMExecutableData().Parameters.DumpParameters();
+		UE_LOG(LogNiagara, Log, TEXT("Spawn RI Parameters"));
+		SpawnScript->RapidIterationParameters.DumpParameters();
+		UE_LOG(LogNiagara, Log, TEXT("Update Parameters"));
+		UpdateScript->GetVMExecutableData().Parameters.DumpParameters();
+		UE_LOG(LogNiagara, Log, TEXT("Update RI Parameters"));
+		UpdateScript->RapidIterationParameters.DumpParameters();
+	}
+#endif
+#endif
+
+	DetermineIfSolo();
 }
 
 #if WITH_EDITORONLY_DATA
 
-UObject* UNiagaraSystem::GetEditorData()
+UNiagaraEditorDataBase* UNiagaraSystem::GetEditorData()
 {
 	return EditorData;
 }
 
-const UObject* UNiagaraSystem::GetEditorData() const
+const UNiagaraEditorDataBase* UNiagaraSystem::GetEditorData() const
 {
 	return EditorData;
 }
 
-void UNiagaraSystem::SetEditorData(UObject* InEditorData)
+void UNiagaraSystem::SetEditorData(UNiagaraEditorDataBase* InEditorData)
 {
 	EditorData = InEditorData;
 }
@@ -275,7 +422,7 @@ INiagaraModule::FMergeEmitterResults UNiagaraSystem::MergeChangesForEmitterHandl
 		if (Instance->bInterpolatedSpawning)
 		{
 			Instance->UpdateScriptProps.Script->RapidIterationParameters.CopyParametersTo(
-				Instance->SpawnScriptProps.Script->RapidIterationParameters, false);
+				Instance->SpawnScriptProps.Script->RapidIterationParameters, false, FNiagaraParameterStore::EDataInterfaceCopyMethod::None);
 		}
 	}
 	else
@@ -312,31 +459,31 @@ void UNiagaraSystem::UpdateFromEmitterChanges(UNiagaraEmitter& ChangedSourceEmit
 
 	if (bNeedsCompile)
 	{
-		Compile(false);
+		RequestCompile(false);
 	}
 }
 
 void UNiagaraSystem::RefreshSystemParametersFromEmitter(const FNiagaraEmitterHandle& EmitterHandle)
 {
+	InitEmitterSpawnAttributes();
 	if (ensureMsgf(EmitterHandles.ContainsByPredicate([=](const FNiagaraEmitterHandle& OwnedEmitterHandle) { return OwnedEmitterHandle.GetId() == EmitterHandle.GetId(); }),
 		TEXT("Can't refresh parameters from an emitter handle this system doesn't own.")))
 	{
-		EmitterHandle.GetInstance()->EmitterSpawnScriptProps.Script->RapidIterationParameters.CopyParametersTo(SystemSpawnScript->RapidIterationParameters, false);
-		EmitterHandle.GetInstance()->EmitterSpawnScriptProps.Script->RapidIterationParameters.CopyParametersTo(SystemSpawnScriptSolo->RapidIterationParameters, false);
-		EmitterHandle.GetInstance()->EmitterUpdateScriptProps.Script->RapidIterationParameters.CopyParametersTo(SystemUpdateScript->RapidIterationParameters, false);
-		EmitterHandle.GetInstance()->EmitterUpdateScriptProps.Script->RapidIterationParameters.CopyParametersTo(SystemUpdateScriptSolo->RapidIterationParameters, false);
+		EmitterHandle.GetInstance()->EmitterSpawnScriptProps.Script->RapidIterationParameters.CopyParametersTo(
+			SystemSpawnScript->RapidIterationParameters, false, FNiagaraParameterStore::EDataInterfaceCopyMethod::None);
+		EmitterHandle.GetInstance()->EmitterUpdateScriptProps.Script->RapidIterationParameters.CopyParametersTo(
+			SystemUpdateScript->RapidIterationParameters, false, FNiagaraParameterStore::EDataInterfaceCopyMethod::None);
 	}
 }
 
 void UNiagaraSystem::RemoveSystemParametersForEmitter(const FNiagaraEmitterHandle& EmitterHandle)
 {
+	InitEmitterSpawnAttributes();
 	if (ensureMsgf(EmitterHandles.ContainsByPredicate([=](const FNiagaraEmitterHandle& OwnedEmitterHandle) { return OwnedEmitterHandle.GetId() == EmitterHandle.GetId(); }),
 		TEXT("Can't remove parameters for an emitter handle this system doesn't own.")))
 	{
 		EmitterHandle.GetInstance()->EmitterSpawnScriptProps.Script->RapidIterationParameters.RemoveParameters(SystemSpawnScript->RapidIterationParameters);
-		EmitterHandle.GetInstance()->EmitterSpawnScriptProps.Script->RapidIterationParameters.RemoveParameters(SystemSpawnScriptSolo->RapidIterationParameters);
 		EmitterHandle.GetInstance()->EmitterUpdateScriptProps.Script->RapidIterationParameters.RemoveParameters(SystemUpdateScript->RapidIterationParameters);
-		EmitterHandle.GetInstance()->EmitterUpdateScriptProps.Script->RapidIterationParameters.RemoveParameters(SystemUpdateScriptSolo->RapidIterationParameters);
 	}
 }
 #endif
@@ -354,13 +501,20 @@ const TArray<FNiagaraEmitterHandle>& UNiagaraSystem::GetEmitterHandles()const
 
 bool UNiagaraSystem::IsReadyToRun() const
 {
-	if (!SystemSpawnScript || !SystemSpawnScriptSolo || !SystemUpdateScript || !SystemUpdateScriptSolo)
+	if (!SystemSpawnScript || !SystemUpdateScript)
 	{
 		return false;
 	}
 
-	if (SystemSpawnScript->IsScriptCompilationPending(false) || SystemSpawnScriptSolo->IsScriptCompilationPending(false) || 
-		SystemUpdateScript->IsScriptCompilationPending(false) || SystemUpdateScriptSolo->IsScriptCompilationPending(false))
+#if WITH_EDITORONLY_DATA
+	if (HasOutstandingCompilationRequests())
+	{
+		return false;
+	}
+#endif
+
+	if (SystemSpawnScript->IsScriptCompilationPending(false) || 
+		SystemUpdateScript->IsScriptCompilationPending(false))
 	{
 		return false;
 	}
@@ -375,17 +529,60 @@ bool UNiagaraSystem::IsReadyToRun() const
 	return true;
 }
 
+bool UNiagaraSystem::HasOutstandingCompilationRequests() const
+{
+	return ActiveCompilations.Num() > 0;
+}
+
+
+bool UNiagaraSystem::IsSolo()const
+{
+	return bSolo;
+}
+
+void UNiagaraSystem::DetermineIfSolo()
+{
+	//Determine if we can update normally or have to update solo.
+	bSolo = false;
+	//If our scripts have any interfaces that require instance data.
+	UNiagaraScript* SystemSpawn = GetSystemSpawnScript();
+	if (SystemSpawn->GetVMExecutableData().IsValid())
+	{
+		for (int32 i = 0; !bSolo && i < SystemSpawn->GetVMExecutableData().DataInterfaceInfo.Num(); ++i)
+		{
+			FNiagaraScriptDataInterfaceCompileInfo& Info = SystemSpawn->GetVMExecutableData().DataInterfaceInfo[i];
+			if (Info.IsSystemSolo())//Temp hack to force solo on any systems with system scrips needing user (aka per instance) interfaces.
+			{
+				bSolo = true;
+				break;
+			}
+		}
+	}
+
+	UNiagaraScript* SystemUpdate = GetSystemUpdateScript();
+	if (SystemUpdate->GetVMExecutableData().IsValid())
+	{
+		for (int32 i = 0; !bSolo && i < SystemUpdate->GetVMExecutableData().DataInterfaceInfo.Num(); ++i)
+		{
+			FNiagaraScriptDataInterfaceCompileInfo& Info = SystemUpdate->GetVMExecutableData().DataInterfaceInfo[i];
+			if (Info.IsSystemSolo())//Temp hack to force solo on any systems with system scrips needing user (aka per instance) interfaces.
+			{
+				bSolo = true;
+				break;
+			}
+		}
+	}
+}
+
 bool UNiagaraSystem::IsValid()const
 {
-	if (!SystemSpawnScript || !SystemSpawnScriptSolo || !SystemUpdateScript || !SystemUpdateScriptSolo)
+	if (!SystemSpawnScript || !SystemUpdateScript)
 	{
 		return false;
 	}
 
 	if ((!SystemSpawnScript->IsScriptCompilationPending(false) && !SystemSpawnScript->DidScriptCompilationSucceed(false)) ||
-		(!SystemSpawnScriptSolo->IsScriptCompilationPending(false) && !SystemSpawnScriptSolo->DidScriptCompilationSucceed(false)) ||
-		(!SystemUpdateScript->IsScriptCompilationPending(false) && !SystemUpdateScript->DidScriptCompilationSucceed(false)) ||
-		(!SystemUpdateScriptSolo->IsScriptCompilationPending(false) && !SystemUpdateScriptSolo->DidScriptCompilationSucceed(false)))
+		(!SystemUpdateScript->IsScriptCompilationPending(false) && !SystemUpdateScript->DidScriptCompilationSucceed(false)))
 	{
 		return false;
 	}
@@ -442,18 +639,20 @@ void UNiagaraSystem::RemoveEmitterHandlesById(const TSet<FGuid>& HandlesToRemove
 		return HandlesToRemove.Contains(EmitterHandle.GetId());
 	};
 	EmitterHandles.RemoveAll(RemovePredicate);
+
+	InitEmitterSpawnAttributes();
 }
 #endif
 
 
-UNiagaraScript* UNiagaraSystem::GetSystemSpawnScript(bool bSolo)
+UNiagaraScript* UNiagaraSystem::GetSystemSpawnScript()
 {
-	return bSolo ? SystemSpawnScriptSolo : SystemSpawnScript;
+	return SystemSpawnScript;
 }
 
-UNiagaraScript* UNiagaraSystem::GetSystemUpdateScript(bool bSolo)
+UNiagaraScript* UNiagaraSystem::GetSystemUpdateScript()
 {
-	return bSolo ? SystemUpdateScriptSolo : SystemUpdateScript;
+	return SystemUpdateScript;
 }
 
 #if WITH_EDITORONLY_DATA
@@ -468,225 +667,320 @@ void UNiagaraSystem::SetIsolateEnabled(bool bIsolate)
 	bIsolateEnabled = bIsolate;
 }
 
-
-void UNiagaraSystem::CompileScripts(TArray<ENiagaraScriptCompileStatus>& OutScriptStatuses, TArray<FString>& OutGraphLevelErrorMessages, TArray<FString>& PathNames, TArray<UNiagaraScript*>& Scripts, bool bForce)
+UNiagaraSystem::FOnSystemCompiled& UNiagaraSystem::OnSystemCompiled()
 {
-	float StartTime = FPlatformTime::Seconds();
-
-	SCOPE_CYCLE_COUNTER(STAT_Niagara_System_CompileScript);
-	PathNames.Reset();
-	OutScriptStatuses.Reset();
-	OutGraphLevelErrorMessages.Reset();
+	return OnSystemCompiledDelegate;
+}
 
-	TArray<FNiagaraVariable> ExposedVars;
+void UNiagaraSystem::InvalidateCachedCompileIds()
+{
+	check(SystemSpawnScript->GetSource() == SystemUpdateScript->GetSource());
+	SystemSpawnScript->GetSource()->InvalidateCachedCompileIds();
 
-	//Compile all emitters
-	bool bAnyUnsynchronized = false;
 	for (FNiagaraEmitterHandle Handle : EmitterHandles)
 	{
-		TArray<ENiagaraScriptCompileStatus> ScriptStatuses;
-		TArray<FString> ScriptErrors;
-
 		UNiagaraScriptSourceBase* GraphSource = Handle.GetInstance()->GraphSource;
-		check(!GraphSource->IsPreCompiled());
-		GraphSource->PreCompile(Handle.GetInstance());
-
-		if (!Handle.GetInstance()->AreAllScriptAndSourcesSynchronized())
-		{
-			bAnyUnsynchronized = true;
-		}
-		Handle.GetInstance()->CompileScripts(ScriptStatuses, ScriptErrors, PathNames, Scripts, bForce);
+		GraphSource->InvalidateCachedCompileIds();
+	}
+}
 
-		GraphSource->GatherPreCompiledVariables(TEXT("User"), ExposedVars);
-		GraphSource->PostCompile();
+void UNiagaraSystem::WaitForCompilationComplete()
+{
+	while (ActiveCompilations.Num() > 0)
+	{
+		QueryCompileComplete(true, ActiveCompilations.Num() == 1);
+	}
+}
 
-		OutScriptStatuses.Append(ScriptStatuses);
-		OutGraphLevelErrorMessages.Append(ScriptErrors);
+bool UNiagaraSystem::PollForCompilationComplete()
+{
+	if (ActiveCompilations.Num() > 0)
+	{
+		return QueryCompileComplete(false, true);
 	}
+	return true;
+}
 
-	bool bForceSystems = bForce || bAnyUnsynchronized;
+bool UNiagaraSystem::QueryCompileComplete(bool bWait, bool bDoPost, bool bDoNotApply)
+{
+	if (ActiveCompilations.Num() > 0)
+	{
+		int32 ActiveCompileIdx = 0;
 
-	check(SystemSpawnScript->GetSource() == SystemUpdateScript->GetSource());
-	check(SystemSpawnScript->GetSource()->IsPreCompiled() == false);
-
-	SystemSpawnScript->GetSource()->PreCompile(nullptr);
-	SystemSpawnScript->GetSource()->GatherPreCompiledVariables(TEXT("User"), ExposedVars);
-
-	FString ErrorMsg;
-	check((int32)EScriptCompileIndices::SpawnScript == 0);
-	OutScriptStatuses.Add(SystemSpawnScript->Compile(ErrorMsg, bForceSystems));
-	OutGraphLevelErrorMessages.Add(ErrorMsg); // EScriptCompileIndices::SpawnScript
-	Scripts.Add(SystemSpawnScript);
-	PathNames.Add(SystemSpawnScript->GetPathName());
-
-	ErrorMsg.Empty();
-	check((int32)EScriptCompileIndices::UpdateScript == 1);
-	OutScriptStatuses.Add(SystemUpdateScript->Compile(ErrorMsg, bForceSystems));
-	OutGraphLevelErrorMessages.Add(ErrorMsg); // EScriptCompileIndices::SpawnScript
-	Scripts.Add(SystemUpdateScript);
-	PathNames.Add(SystemUpdateScript->GetPathName());
-
-	ErrorMsg.Empty();
-	check((int32)EScriptCompileIndices::SpawnScript == 0);
-	OutScriptStatuses.Add(SystemSpawnScriptSolo->Compile(ErrorMsg, bForceSystems));
-	OutGraphLevelErrorMessages.Add(ErrorMsg); // EScriptCompileIndices::SpawnScript
-	Scripts.Add(SystemSpawnScriptSolo);
-	PathNames.Add(SystemSpawnScriptSolo->GetPathName());
-
-	ErrorMsg.Empty();
-	check((int32)EScriptCompileIndices::UpdateScript == 1);
-	OutScriptStatuses.Add(SystemUpdateScriptSolo->Compile(ErrorMsg, bForceSystems));
-	OutGraphLevelErrorMessages.Add(ErrorMsg); // EScriptCompileIndices::SpawnScript
-	Scripts.Add(SystemUpdateScriptSolo);
-	PathNames.Add(SystemUpdateScriptSolo->GetPathName());
-
-	SystemSpawnScript->GetSource()->PostCompile();
+		bool bAreWeWaitingForAnyResults = false;
 
-	InitEmitterSpawnAttributes();
-	
-	// Now let's synchronize the variables that we expose to the end user.
-	TArray<FNiagaraVariable> OriginalVars;
-	ExposedParameters.GetParameters(OriginalVars);
-	for (int32 i = 0; i < ExposedVars.Num(); i++)
-	{
-		if (OriginalVars.Contains(ExposedVars[i]) == false)
+		// Check to see if ALL of the sub-requests have resolved. 
+		for (FEmitterCompiledScriptPair& EmitterCompiledScriptPair : ActiveCompilations[ActiveCompileIdx].EmitterCompiledScriptPairs)
 		{
-			// Just in case it wasn't added previously..
-			ExposedParameters.AddParameter(ExposedVars[i]);
+			if ((uint32)INDEX_NONE == EmitterCompiledScriptPair.PendingDDCID || EmitterCompiledScriptPair.bResultsReady)
+			{
+				continue;
+			}
+			if (bWait)
+			{
+				GetDerivedDataCacheRef().WaitAsynchronousCompletion(EmitterCompiledScriptPair.PendingDDCID);
+				EmitterCompiledScriptPair.bResultsReady = true;
+			}
+			else
+			{
+				EmitterCompiledScriptPair.bResultsReady = GetDerivedDataCacheRef().PollAsynchronousCompletion(EmitterCompiledScriptPair.PendingDDCID);
+				if (!EmitterCompiledScriptPair.bResultsReady)
+				{
+					bAreWeWaitingForAnyResults = true;
+				}
+			}
+
+			// If the results are ready, go ahead and cache them so that the pending task isn't removed prematurely...
+			if (EmitterCompiledScriptPair.bResultsReady)
+			{
+				TArray<uint8> OutData;
+				bool bBuiltLocally = false;
+				if (GetDerivedDataCacheRef().GetAsynchronousResults(EmitterCompiledScriptPair.PendingDDCID, OutData, &bBuiltLocally))
+				{
+					if (bBuiltLocally)
+					{
+						UE_LOG(LogNiagara, Log, TEXT("UNiagraScript \'%s\' was built locally.."), *EmitterCompiledScriptPair.CompiledScript->GetFullName());
+					}
+					else
+					{
+						UE_LOG(LogNiagara, Log, TEXT("UNiagraScript \'%s\' was pulled from DDC."), *EmitterCompiledScriptPair.CompiledScript->GetFullName());
+					}
+
+					TSharedPtr<FNiagaraVMExecutableData> ExeData = MakeShared<FNiagaraVMExecutableData>();
+					EmitterCompiledScriptPair.CompileResults = ExeData;
+					if (!bDoNotApply)
+					{
+						FNiagaraScriptDerivedData::BinaryToExecData(OutData, *(ExeData.Get()));
+					}
+				}
+			}
 		}
-	}
 
-	{
-		SCOPE_CYCLE_COUNTER(STAT_Niagara_System_CompileScriptResetAfter);
+		check(bWait ? (bAreWeWaitingForAnyResults == false) : true);
 
-		FNiagaraSystemUpdateContext UpdateCtx(this, true);
-	}
-	//I assume this was added to recompile emitters too but I'm doing that above. We need to rework this compilation flow.
-	//Compile();
+		// Make sure that we aren't waiting for any results to come back.
+		if (bAreWeWaitingForAnyResults && !bWait)
+		{
+			return false;
+		}
 
+		// In the world of do not apply, we're exiting the system completely so let's just kill any active compilations altogether.
+		if (bDoNotApply)
+		{
+			ActiveCompilations[ActiveCompileIdx].RootObjects.Empty();
+			ActiveCompilations.RemoveAt(ActiveCompileIdx);
+			return true;
+		}
 
-	UE_LOG(LogNiagara, Log, TEXT("Compiling System %s took %f sec."), *GetFullName(), FPlatformTime::Seconds() - StartTime);
-}
 
-void UNiagaraSystem::Compile(bool bForce)
-{
-	float StartTime = FPlatformTime::Seconds();
+		SCOPE_CYCLE_COUNTER(STAT_Niagara_System_CompileScript);
 
-	bool bAnyUnsynchronized = false;
-	TArray<FNiagaraVariable> ExposedVars;
-	for (FNiagaraEmitterHandle Handle: EmitterHandles)
-	{
-		TArray<ENiagaraScriptCompileStatus> ScriptStatuses;
-		TArray<FString> ScriptErrors;
-		TArray<FString> PathNames;
-		TArray<UNiagaraScript*> Scripts;
-		
-		UNiagaraScriptSourceBase* GraphSource = Handle.GetInstance()->GraphSource;
-		check(!GraphSource->IsPreCompiled());
-		GraphSource->PreCompile(Handle.GetInstance());
+		// Now that the above code says they are all complete, go ahead and resolve them all at once.
+		float CombinedCompileTime = 0.0f;
+		for (FEmitterCompiledScriptPair& EmitterCompiledScriptPair : ActiveCompilations[ActiveCompileIdx].EmitterCompiledScriptPairs)
+		{
+			if ((uint32)INDEX_NONE == EmitterCompiledScriptPair.PendingDDCID )
+			{
+				continue;
+			}
+			CombinedCompileTime += EmitterCompiledScriptPair.CompileResults->CompileTime;
+			check(EmitterCompiledScriptPair.bResultsReady);
+
+			TSharedPtr<FNiagaraVMExecutableData> ExeData = EmitterCompiledScriptPair.CompileResults;
+			TSharedPtr<FNiagaraCompileRequestDataBase, ESPMode::ThreadSafe> PrecompData = ActiveCompilations[ActiveCompileIdx].MappedData.FindChecked(EmitterCompiledScriptPair.CompiledScript);
+			EmitterCompiledScriptPair.CompiledScript->SetVMCompilationResults(EmitterCompiledScriptPair.CompileId, *(ExeData.Get()), PrecompData.Get());	
+		}
 
-		if (!Handle.GetInstance()->AreAllScriptAndSourcesSynchronized())
+		if (bDoPost)
 		{
-			bAnyUnsynchronized = true;
+			for (FNiagaraEmitterHandle Handle : EmitterHandles)
+			{
+				Handle.GetInstance()->OnPostCompile();
+			}
 		}
-		Handle.GetInstance()->CompileScripts(ScriptStatuses, ScriptErrors, PathNames, Scripts, bForce);
 
-		GraphSource->GatherPreCompiledVariables(TEXT("User"), ExposedVars);
+		InitEmitterSpawnAttributes();
 
-		GraphSource->PostCompile();
+		// Clean up old, and initialize new rapid iterations parameters due to graph changes.
+		for (FEmitterCompiledScriptPair& EmitterCompiledScriptPair : ActiveCompilations[ActiveCompileIdx].EmitterCompiledScriptPairs)
+		{
+			if (EmitterCompiledScriptPair.Emitter == nullptr)
+			{
+				continue;
+			}
+
+			EmitterCompiledScriptPair.CompiledScript->CleanUpOldAndInitializeNewRapidIterationParameters(EmitterCompiledScriptPair.Emitter->GetUniqueEmitterName());
+		}
 
-		for (int32 i = 0; i < ScriptStatuses.Num(); i++)
+		// Push rapid iteration parameters for dependent scripts.
+		for (FEmitterCompiledScriptPair& EmitterCompiledScriptPair : ActiveCompilations[ActiveCompileIdx].EmitterCompiledScriptPairs)
 		{
-			if (ScriptErrors[i].Len() != 0 || ScriptStatuses[i] != ENiagaraScriptCompileStatus::NCS_UpToDate)
+			if (EmitterCompiledScriptPair.Emitter == nullptr)
 			{
-				UE_LOG(LogNiagara, Warning, TEXT("Script '%s', compile status: %d  errors: %s"), *PathNames[i], (int32)ScriptStatuses[i], *ScriptErrors[i]);
+				continue;
 			}
-			else
+
+			const bool bOnlyAddParametersOnCopy = false;
+			if (EmitterCompiledScriptPair.CompiledScript->GetUsage() == ENiagaraScriptUsage::ParticleUpdateScript && EmitterCompiledScriptPair.Emitter->bInterpolatedSpawning)
+			{
+				EmitterCompiledScriptPair.CompiledScript->RapidIterationParameters.CopyParametersTo(
+					EmitterCompiledScriptPair.Emitter->SpawnScriptProps.Script->RapidIterationParameters, bOnlyAddParametersOnCopy, FNiagaraParameterStore::EDataInterfaceCopyMethod::None);
+			}
+			else if (EmitterCompiledScriptPair.CompiledScript->GetUsage() == ENiagaraScriptUsage::EmitterSpawnScript)
+			{
+				EmitterCompiledScriptPair.CompiledScript->RapidIterationParameters.CopyParametersTo(
+					SystemSpawnScript->RapidIterationParameters, bOnlyAddParametersOnCopy, FNiagaraParameterStore::EDataInterfaceCopyMethod::None);
+			}
+			else if (EmitterCompiledScriptPair.CompiledScript->GetUsage() == ENiagaraScriptUsage::EmitterUpdateScript)
 			{
-				UE_LOG(LogNiagara, Log, TEXT("Script '%s', compile status: Success!"), *PathNames[i]);
+				EmitterCompiledScriptPair.CompiledScript->RapidIterationParameters.CopyParametersTo(
+					SystemUpdateScript->RapidIterationParameters, bOnlyAddParametersOnCopy, FNiagaraParameterStore::EDataInterfaceCopyMethod::None);
+			}
+			else if (EmitterCompiledScriptPair.CompiledScript->GetUsage() == ENiagaraScriptUsage::ParticleGPUComputeScript && EmitterCompiledScriptPair.Emitter->SimTarget != ENiagaraSimTarget::CPUSim)
+			{
+				EmitterCompiledScriptPair.Emitter->SpawnScriptProps.Script->RapidIterationParameters.CopyParametersTo(
+					EmitterCompiledScriptPair.CompiledScript->RapidIterationParameters, bOnlyAddParametersOnCopy, FNiagaraParameterStore::EDataInterfaceCopyMethod::None);
+				EmitterCompiledScriptPair.Emitter->UpdateScriptProps.Script->RapidIterationParameters.CopyParametersTo(
+					EmitterCompiledScriptPair.CompiledScript->RapidIterationParameters, bOnlyAddParametersOnCopy, FNiagaraParameterStore::EDataInterfaceCopyMethod::None);
 			}
 		}
-	}
-	bool bForceSystems = bForce || bAnyUnsynchronized;
 
-	check(SystemSpawnScript->GetSource() == SystemUpdateScript->GetSource());
-	check(SystemSpawnScript->GetSource()->IsPreCompiled() == false);
+		ActiveCompilations[ActiveCompileIdx].RootObjects.Empty();
 
-	SystemSpawnScript->GetSource()->PreCompile(nullptr);
-	SystemSpawnScript->GetSource()->GatherPreCompiledVariables(TEXT("User"), ExposedVars);
+		DetermineIfSolo();
 
-	FString CompileErrors;
-	ENiagaraScriptCompileStatus CompileStatus = SystemSpawnScript->Compile(CompileErrors, bForceSystems);
-	if (CompileErrors.Len() != 0 || CompileStatus != ENiagaraScriptCompileStatus::NCS_UpToDate)
-	{
-		UE_LOG(LogNiagara, Warning, TEXT("System Spawn Script '%s', compile status: %d  errors: %s"), *SystemSpawnScript->GetPathName(), (int32)CompileStatus, *CompileErrors);
+		UE_LOG(LogNiagara, Log, TEXT("Compiling System %s took %f sec (wall time), %f sec (combined time)."), *GetFullName(), (float)(FPlatformTime::Seconds() - ActiveCompilations[ActiveCompileIdx].StartTime),
+			CombinedCompileTime);
+
+		ActiveCompilations.RemoveAt(ActiveCompileIdx);
+
+		if (bDoPost)
+		{
+			SCOPE_CYCLE_COUNTER(STAT_Niagara_System_CompileScriptResetAfter);
+
+			OnSystemCompiled().Broadcast(this);
+		}
+
+		return true;
 	}
-	else
+
+	return false;
+}
+
+bool UNiagaraSystem::RequestCompile(bool bForce)
+{
+	if (bForce)
 	{
-		UE_LOG(LogNiagara, Log, TEXT("Script Spawn '%s', compile status: Success!"), *SystemSpawnScript->GetPathName());
+		InvalidateCachedCompileIds();
+		bForce = false;
 	}
 
-	CompileStatus = SystemUpdateScript->Compile(CompileErrors, bForceSystems);
-	if (CompileErrors.Len() != 0 || CompileStatus != ENiagaraScriptCompileStatus::NCS_UpToDate)
+	if (ActiveCompilations.Num() > 0)
 	{
-		UE_LOG(LogNiagara, Warning, TEXT("System Update Script '%s', compile status: %d  errors: %s"), *SystemUpdateScript->GetPathName(), (int32)CompileStatus, *CompileErrors);
+		PollForCompilationComplete();
 	}
-	else
+	
+	int32 ActiveCompileIdx = ActiveCompilations.AddDefaulted();
+	ActiveCompilations[ActiveCompileIdx].StartTime = FPlatformTime::Seconds();
+
+	SCOPE_CYCLE_COUNTER(STAT_Niagara_System_Precompile);
+	
+	check(SystemSpawnScript->GetSource() == SystemUpdateScript->GetSource());
+	TArray<FNiagaraVariable> OriginalExposedParams;
+	GetExposedParameters().GetParameters(OriginalExposedParams);
+
+	INiagaraModule& NiagaraModule = FModuleManager::Get().LoadModuleChecked<INiagaraModule>(TEXT("Niagara"));
+	TSharedPtr<FNiagaraCompileRequestDataBase, ESPMode::ThreadSafe> SystemPrecompiledData = NiagaraModule.Precompile(this);
+
+	SystemPrecompiledData->GetReferencedObjects(ActiveCompilations[ActiveCompileIdx].RootObjects);
+
+	//Compile all emitters
+	bool bTrulyAsync = true;
+	bool bAnyUnsynchronized = false;	
+
+	ActiveCompilations[ActiveCompileIdx].MappedData.Add(SystemSpawnScript, SystemPrecompiledData);
+	ActiveCompilations[ActiveCompileIdx].MappedData.Add(SystemUpdateScript, SystemPrecompiledData);
+
+	check(EmitterHandles.Num() == SystemPrecompiledData->GetDependentRequestCount());
+
+	// Grab the list of user variables that were actually encountered so that we can add to them later.
+	TArray<FNiagaraVariable> EncounteredExposedVars;
+	SystemPrecompiledData->GatherPreCompiledVariables(TEXT("User"), EncounteredExposedVars);
+
+	for (int32 i = 0; i < EmitterHandles.Num(); i++)
 	{
-		UE_LOG(LogNiagara, Log, TEXT("Script Update '%s', compile status: Success!"), *SystemUpdateScript->GetPathName());
+		FNiagaraEmitterHandle Handle = EmitterHandles[i];
+
+		UNiagaraScriptSourceBase* GraphSource = Handle.GetInstance()->GraphSource;
+		TSharedPtr<FNiagaraCompileRequestDataBase, ESPMode::ThreadSafe> EmitterPrecompiledData = SystemPrecompiledData->GetDependentRequest(i);
+		EmitterPrecompiledData->GetReferencedObjects(ActiveCompilations[ActiveCompileIdx].RootObjects);
+
+		TArray<UNiagaraScript*> EmitterScripts;
+		Handle.GetInstance()->GetScripts(EmitterScripts, false);
+		check(EmitterScripts.Num() > 0);
+		for (UNiagaraScript* EmitterScript : EmitterScripts)
+		{
+			ActiveCompilations[ActiveCompileIdx].MappedData.Add(EmitterScript, EmitterPrecompiledData);
+
+			FEmitterCompiledScriptPair Pair;
+			Pair.bResultsReady = false;
+			Pair.Emitter = Handle.GetInstance();
+			Pair.CompiledScript = EmitterScript;			
+			if (EmitterScript->RequestExternallyManagedAsyncCompile(EmitterPrecompiledData, Pair.CompileId, Pair.PendingDDCID, bTrulyAsync))
+			{
+				bAnyUnsynchronized = true;
+			}
+			ActiveCompilations[ActiveCompileIdx].EmitterCompiledScriptPairs.Add(Pair);
+		}	
+
+		// Add the emitter's User variables to the encountered list to expose for later.
+		EmitterPrecompiledData->GatherPreCompiledVariables(TEXT("User"), EncounteredExposedVars);
 	}
 
-	if(SystemSpawnScriptSolo)
+	bool bForceSystems = bForce || bAnyUnsynchronized;
+	bool bAnyCompiled = bAnyUnsynchronized || bForce;
+
+	// Now add the system scripts for compilation...
 	{
-		CompileStatus = SystemSpawnScriptSolo->Compile(CompileErrors, bForceSystems);
-		if (CompileErrors.Len() != 0 || CompileStatus != ENiagaraScriptCompileStatus::NCS_UpToDate)
+		FEmitterCompiledScriptPair Pair;
+		Pair.bResultsReady = false;
+		Pair.Emitter = nullptr;
+		Pair.CompiledScript = SystemSpawnScript;
+		if (SystemSpawnScript->RequestExternallyManagedAsyncCompile(SystemPrecompiledData, Pair.CompileId, Pair.PendingDDCID, bTrulyAsync))
 		{
-			UE_LOG(LogNiagara, Warning, TEXT("System Spawn Solo Script '%s', compile status: %d  errors: %s"), *SystemSpawnScriptSolo->GetPathName(), (int32)CompileStatus, *CompileErrors);
-		}
-		else
-		{
-			UE_LOG(LogNiagara, Log, TEXT("Script Spawn Solo '%s', compile status: Success!"), *SystemSpawnScriptSolo->GetPathName());
+			bAnyCompiled = true;
 		}
+		ActiveCompilations[ActiveCompileIdx].EmitterCompiledScriptPairs.Add(Pair);
 	}
 
-	if (SystemUpdateScriptSolo)
 	{
-		CompileStatus = SystemUpdateScriptSolo->Compile(CompileErrors, bForceSystems);
-		if (CompileErrors.Len() != 0 || CompileStatus != ENiagaraScriptCompileStatus::NCS_UpToDate)
-		{
-			UE_LOG(LogNiagara, Warning, TEXT("System Update Solo Script '%s', compile status: %d  errors: %s"), *SystemUpdateScriptSolo->GetPathName(), (int32)CompileStatus, *CompileErrors);
-		}
-		else
+		FEmitterCompiledScriptPair Pair;
+		Pair.bResultsReady = false;
+		Pair.Emitter = nullptr;
+		Pair.CompiledScript = SystemUpdateScript;
+		if (SystemUpdateScript->RequestExternallyManagedAsyncCompile(SystemPrecompiledData, Pair.CompileId, Pair.PendingDDCID, bTrulyAsync))
 		{
-			UE_LOG(LogNiagara, Log, TEXT("Script Update Solo '%s', compile status: Success!"), *SystemUpdateScriptSolo->GetPathName());
+			bAnyCompiled = true;
 		}
+		ActiveCompilations[ActiveCompileIdx].EmitterCompiledScriptPairs.Add(Pair);
 	}
-	SystemSpawnScript->GetSource()->PostCompile();
 
-	InitEmitterSpawnAttributes();
-	TArray<FNiagaraVariable> OriginalVars;
-	ExposedParameters.GetParameters(OriginalVars);
-	for (int32 i = 0; i < ExposedVars.Num(); i++)
-	{
-		ExposedParameters.AddParameter(ExposedVars[i]);
-	}
-	for (int32 i = 0; i < OriginalVars.Num(); i++)
+	// Now let's synchronize the variables that we actually encountered during compile so that we can expose them to the end user.
+	for (int32 i = 0; i < EncounteredExposedVars.Num(); i++)
 	{
-		if (ExposedVars.Contains(OriginalVars[i]) == false)
+		if (OriginalExposedParams.Contains(EncounteredExposedVars[i]) == false)
 		{
-			ExposedParameters.RemoveParameter(OriginalVars[i]);
+			// Just in case it wasn't added previously..
+			ExposedParameters.AddParameter(EncounteredExposedVars[i]);
 		}
 	}
 
-	{
+	FNiagaraSystemUpdateContext UpdateCtx(this, true);
 
-		SCOPE_CYCLE_COUNTER(STAT_Niagara_System_CompileScriptResetAfter);
-		FNiagaraSystemUpdateContext UpdateCtx(this, true);
-	}
+	return bAnyCompiled;
+}
 
-	UE_LOG(LogNiagara, Log, TEXT("Compiling System %s took %f sec."), *GetFullName(), FPlatformTime::Seconds() - StartTime);
 
-}
 #endif
 
 void UNiagaraSystem::InitEmitterSpawnAttributes()
@@ -694,28 +988,34 @@ void UNiagaraSystem::InitEmitterSpawnAttributes()
 	EmitterSpawnAttributes.Empty();
 	EmitterSpawnAttributes.SetNum(EmitterHandles.Num());
 	FNiagaraTypeDefinition SpawnInfoDef = FNiagaraTypeDefinition(FNiagaraSpawnInfo::StaticStruct());
-	for (FNiagaraVariable& Var : SystemSpawnScript->Attributes)
+	if (SystemSpawnScript->GetVMExecutableData().IsValid())
 	{
-		for (int32 EmitterIdx = 0; EmitterIdx < EmitterHandles.Num(); ++EmitterIdx)
+		for (FNiagaraVariable& Var : SystemSpawnScript->GetVMExecutableData().Attributes)
 		{
-			UNiagaraEmitter* Emitter = EmitterHandles[EmitterIdx].GetInstance();
-			FString EmitterName = Emitter->GetUniqueEmitterName() + TEXT(".");
-			if (Var.GetType() == SpawnInfoDef && Var.GetName().ToString().StartsWith(EmitterName))
+			for (int32 EmitterIdx = 0; EmitterIdx < EmitterHandles.Num(); ++EmitterIdx)
 			{
-				EmitterSpawnAttributes[EmitterIdx].SpawnAttributes.AddUnique(Var.GetName());
+				UNiagaraEmitter* Emitter = EmitterHandles[EmitterIdx].GetInstance();
+				FString EmitterName = Emitter->GetUniqueEmitterName() + TEXT(".");
+				if (Var.GetType() == SpawnInfoDef && Var.GetName().ToString().StartsWith(EmitterName))
+				{
+					EmitterSpawnAttributes[EmitterIdx].SpawnAttributes.AddUnique(Var.GetName());
+				}
 			}
 		}
 	}
-	for (FNiagaraVariable& Var : SystemUpdateScript->Attributes)
+	if (SystemUpdateScript->GetVMExecutableData().IsValid())
 	{
-		for (int32 EmitterIdx = 0; EmitterIdx < EmitterHandles.Num(); ++EmitterIdx)
+		for (FNiagaraVariable& Var : SystemUpdateScript->GetVMExecutableData().Attributes)
 		{
-			UNiagaraEmitter* Emitter = EmitterHandles[EmitterIdx].GetInstance();
-			FString EmitterName = Emitter->GetUniqueEmitterName() + TEXT(".");
-			if (Var.GetType() == SpawnInfoDef && Var.GetName().ToString().StartsWith(EmitterName))
+			for (int32 EmitterIdx = 0; EmitterIdx < EmitterHandles.Num(); ++EmitterIdx)
 			{
-				EmitterSpawnAttributes[EmitterIdx].SpawnAttributes.AddUnique(Var.GetName());
+				UNiagaraEmitter* Emitter = EmitterHandles[EmitterIdx].GetInstance();
+				FString EmitterName = Emitter->GetUniqueEmitterName() + TEXT(".");
+				if (Var.GetType() == SpawnInfoDef && Var.GetName().ToString().StartsWith(EmitterName))
+				{
+					EmitterSpawnAttributes[EmitterIdx].SpawnAttributes.AddUnique(Var.GetName());
+				}
 			}
 		}
 	}
-}
+}
\ No newline at end of file
diff --git a/Engine/Plugins/FX/Niagara/Source/Niagara/Private/NiagaraSystemInstance.cpp b/Engine/Plugins/FX/Niagara/Source/Niagara/Private/NiagaraSystemInstance.cpp
--- a/Engine/Plugins/FX/Niagara/Source/Niagara/Private/NiagaraSystemInstance.cpp
+++ b/Engine/Plugins/FX/Niagara/Source/Niagara/Private/NiagaraSystemInstance.cpp
@@ -11,10 +11,19 @@
 #include "NiagaraComponent.h"
 #include "NiagaraRenderer.h"
 #include "GameFramework/PlayerController.h"
+#include "Templates/AlignmentTemplates.h"
 
+
+DECLARE_CYCLE_STAT(TEXT("System Activate (GT)"), STAT_NiagaraSystemActivate, STATGROUP_Niagara);
+DECLARE_CYCLE_STAT(TEXT("System Deactivate (GT)"), STAT_NiagaraSystemDeactivate, STATGROUP_Niagara);
+DECLARE_CYCLE_STAT(TEXT("System Complete (GT)"), STAT_NiagaraSystemComplete, STATGROUP_Niagara);
 DECLARE_CYCLE_STAT(TEXT("Parallel Tick"), STAT_NiagaraParallelTick, STATGROUP_Niagara);
 DECLARE_CYCLE_STAT(TEXT("System Reset (GT)"), STAT_NiagaraSystemReset, STATGROUP_Niagara);
 DECLARE_CYCLE_STAT(TEXT("System Reinit (GT)"), STAT_NiagaraSystemReinit, STATGROUP_Niagara);
+DECLARE_CYCLE_STAT(TEXT("System Init Emitters (GT)"), STAT_NiagaraSystemInitEmitters, STATGROUP_Niagara);
+DECLARE_CYCLE_STAT(TEXT("System Advance Simulation "), STAT_NiagaraSystemAdvanceSim, STATGROUP_Niagara);
+DECLARE_CYCLE_STAT(TEXT("System SetSolo "), STAT_NiagaraSystemSetSolo, STATGROUP_Niagara); 
+DECLARE_CYCLE_STAT(TEXT("System PreSimulateTick "), STAT_NiagaraSystemPreSimulateTick, STATGROUP_Niagara); 
 
  
 /** Safety time to allow for the LastRenderTime coming back from the RT. */
@@ -35,67 +44,229 @@ FNiagaraSystemInstance::FNiagaraSystemInstance(UNiagaraComponent* InComponent)
 	, InstanceParameters(Component)
 	, bSolo(false)
 	, bForceSolo(false)
-	, bError(true)
-	, bDataInterfacesNeedInit(true)
+	, bPendingSpawn(false)
 	, bHasTickingEmitters(true)
-	, ExecutionState(ENiagaraExecutionState::Complete)
+	, RequestedExecutionState(ENiagaraExecutionState::Complete)
+	, ActualExecutionState(ENiagaraExecutionState::Complete)
 {
 	SystemBounds.Init();
 }
 
-void FNiagaraSystemInstance::Init(TSharedRef<FNiagaraSystemSimulation, ESPMode::ThreadSafe> InSystemSimulation, bool bInForceSolo)
+void FNiagaraSystemInstance::Init(UNiagaraSystem* InSystem, bool bInForceSolo)
 {
-	SystemSimulation = InSystemSimulation;
-
-	bError = false;
 	bForceSolo = bInForceSolo;
-	ExecutionState = ENiagaraExecutionState::Inactive;
+	ActualExecutionState = ENiagaraExecutionState::Inactive;
+	RequestedExecutionState = ENiagaraExecutionState::Inactive;
 
 	//InstanceParameters = GetSystem()->GetInstanceParameters();
-	InitEmitters();
-	Reset(EResetMode::ReInit);
+	// In order to get user data interface parameters in the component to work properly,
+	// we need to bind here, otherwise the instances when we init data interfaces during reset will potentially
+	// be the defaults (i.e. null) for things like static mesh data interfaces.
+	Reset(EResetMode::ReInit, true);
 
 #if WITH_EDITORONLY_DATA
-	InstanceParameters.Name = *FString::Printf(TEXT("SystemInstance %p"), this);
+	InstanceParameters.DebugName = *FString::Printf(TEXT("SystemInstance %p"), this);
 #endif
 	OnInitializedDelegate.Broadcast();
 }
 
-void FNiagaraSystemInstance::SetExecutionState(ENiagaraExecutionState InState)
+void FNiagaraSystemInstance::SetRequestedExecutionState(ENiagaraExecutionState InState)
+{
+	//Once in disabled state we can never get out except on Reinit.
+	if (RequestedExecutionState != InState && RequestedExecutionState != ENiagaraExecutionState::Disabled)
+	{
+		/*const UEnum* EnumPtr = FNiagaraTypeDefinition::GetExecutionStateEnum();
+		UE_LOG(LogNiagara, Log, TEXT("Component \"%s\" System \"%s\" requested change state: %s to %s, actual %s"), *GetComponent()->GetName(), *GetSystem()->GetName(), *EnumPtr->GetNameStringByValue((int64)RequestedExecutionState),
+			*EnumPtr->GetNameStringByValue((int64)InState), *EnumPtr->GetNameStringByValue((int64)ActualExecutionState));
+		*/
+		if (InState == ENiagaraExecutionState::Disabled)
+		{
+			//Really move to disabled straight away.
+			ActualExecutionState = ENiagaraExecutionState::Disabled;
+			Cleanup();
+		}
+		RequestedExecutionState = InState;
+	}
+}
+
+void FNiagaraSystemInstance::SetActualExecutionState(ENiagaraExecutionState InState)
 {
-	if (ExecutionState != InState)
+
+	//Once in disabled state we can never get out except on Reinit.
+	if (ActualExecutionState != InState && ActualExecutionState != ENiagaraExecutionState::Disabled)
 	{
-		const UEnum* EnumPtr = FNiagaraTypeDefinition::GetExecutionStateEnum();
-		UE_LOG(LogNiagara, Log, TEXT("System \"%s\" change state: %s to %s"), *GetSystem()->GetName(), *EnumPtr->GetNameStringByValue((int64)ExecutionState),
+		/*const UEnum* EnumPtr = FNiagaraTypeDefinition::GetExecutionStateEnum();
+		UE_LOG(LogNiagara, Log, TEXT("Component \"%s\" System \"%s\" actual change state: %s to %s"), *GetComponent()->GetName(), *GetSystem()->GetName(), *EnumPtr->GetNameStringByValue((int64)ActualExecutionState),
 			*EnumPtr->GetNameStringByValue((int64)InState));
+		*/
+		ActualExecutionState = InState;
 
+		if (ActualExecutionState == ENiagaraExecutionState::Active)
+		{
+			// We only need to notify completion once after each successful active.
+			// Here's when we know that we just became active.
+			bNotifyOnCompletion = true;
+
+			// We may also end up calling HandleCompletion on each emitter.
+			// This may happen *before* we've successfully pulled data off of a 
+			// simulation run. This means that we need to synchronize the execution
+			// states upon activation.
+			for (int32 EmitterIdx = 0; EmitterIdx < Emitters.Num(); ++EmitterIdx)
+			{
+				FNiagaraEmitterInstance& EmitterInst = Emitters[EmitterIdx].Get();
+				EmitterInst.SetExecutionState(ENiagaraExecutionState::Active);
+			}
+		}
 	}
-	ExecutionState = InState;
 }
 
-void FNiagaraSystemInstance::Activate(bool bReset /* = false */)
+#if WITH_EDITORONLY_DATA
+bool FNiagaraSystemInstance::RequestCapture(const FGuid& RequestId)
 {
-	if (GetSystem() && GetSystem()->IsValid() && IsReadyToRun())
+	if (IsComplete() || CurrentCapture.IsValid())
+	{
+		return false;
+	}
+
+	bWasSoloPriorToCaptureRequest = bSolo;
+	SetSolo(true);
+
+	// Go ahead and populate the shared array so that we don't have to do this on the game thread and potentially race.
+	TSharedRef<TArray<TSharedPtr<struct FNiagaraScriptDebuggerInfo, ESPMode::ThreadSafe>>, ESPMode::ThreadSafe> TempCaptureHolder = 
+		MakeShared<TArray<TSharedPtr<struct FNiagaraScriptDebuggerInfo, ESPMode::ThreadSafe>>, ESPMode::ThreadSafe>();
+	
+	TempCaptureHolder->Add(MakeShared<FNiagaraScriptDebuggerInfo, ESPMode::ThreadSafe>(NAME_None, ENiagaraScriptUsage::SystemSpawnScript, FGuid()));
+	TempCaptureHolder->Add(MakeShared<FNiagaraScriptDebuggerInfo, ESPMode::ThreadSafe>(NAME_None, ENiagaraScriptUsage::SystemUpdateScript, FGuid()));
+
+	for (const FNiagaraEmitterHandle& Handle : GetSystem()->GetEmitterHandles())
 	{
-		if (IsComplete() || bReset)
+		TArray<UNiagaraScript*> Scripts;
+		Handle.GetInstance()->GetScripts(Scripts, false);
+
+		for (UNiagaraScript* Script : Scripts)
 		{
-			Reset(EResetMode::ResetAll);
+			TempCaptureHolder->Add(MakeShared<FNiagaraScriptDebuggerInfo, ESPMode::ThreadSafe>(Handle.GetIdName(), Script->GetUsage(), Script->GetUsageId()));
 		}
-		else
+	}
+	CapturedFrames.Add(RequestId, TempCaptureHolder);
+	CurrentCapture = TempCaptureHolder;
+	CurrentCaptureGuid = MakeShared<FGuid, ESPMode::ThreadSafe>(RequestId);
+	return true;
+}
+
+void FNiagaraSystemInstance::FinishCapture()
+{
+	if (!CurrentCapture.IsValid())
+	{
+		return;
+	}
+
+	SetSolo(bWasSoloPriorToCaptureRequest);
+	CurrentCapture.Reset();
+	CurrentCaptureGuid.Reset();
+}
+
+bool FNiagaraSystemInstance::QueryCaptureResults(const FGuid& RequestId, TArray<TSharedPtr<struct FNiagaraScriptDebuggerInfo, ESPMode::ThreadSafe>>& OutCaptureResults)
+{
+	if (CurrentCaptureGuid.IsValid() && RequestId == *CurrentCaptureGuid.Get())
+	{
+		return false;
+	}
+
+	const TSharedPtr<TArray<TSharedPtr<struct FNiagaraScriptDebuggerInfo, ESPMode::ThreadSafe>>, ESPMode::ThreadSafe>* FoundEntry = CapturedFrames.Find(RequestId);
+	if (FoundEntry != nullptr)
+	{
+		TArray<TSharedPtr<struct FNiagaraScriptDebuggerInfo, ESPMode::ThreadSafe>>* Array = FoundEntry->Get();
+		OutCaptureResults.SetNum(Array->Num());
+
+		for (int32 i = 0; i < FoundEntry->Get()->Num(); i++)
+		{
+			OutCaptureResults[i] = (*Array)[i];
+		}
+		CapturedFrames.Remove(RequestId);
+		return true;
+	}
+
+	return false;
+}
+
+TArray<TSharedPtr<struct FNiagaraScriptDebuggerInfo, ESPMode::ThreadSafe>>* FNiagaraSystemInstance::GetActiveCaptureResults()
+{
+	return CurrentCapture.Get();
+}
+
+FNiagaraScriptDebuggerInfo* FNiagaraSystemInstance::GetActiveCaptureWrite(const FName& InHandleName, ENiagaraScriptUsage InUsage, const FGuid& InUsageId)
+{
+	if (CurrentCapture.IsValid())
+	{
+		TSharedPtr<struct FNiagaraScriptDebuggerInfo, ESPMode::ThreadSafe>* FoundEntry = CurrentCapture->FindByPredicate([&](const TSharedPtr<struct FNiagaraScriptDebuggerInfo, ESPMode::ThreadSafe>& Entry)
+		{
+			return Entry->HandleName == InHandleName && UNiagaraScript::IsEquivalentUsage(Entry->Usage, InUsage) && Entry->UsageId == InUsageId;
+		});
+
+		if (FoundEntry != nullptr)
 		{
-			Reset(EResetMode::ResetSystem);
+			return FoundEntry->Get();
 		}
+	}
+	return nullptr;
+}
 
-		ExecutionState = ENiagaraExecutionState::Active;
+bool FNiagaraSystemInstance::ShouldCaptureThisFrame() const
+{
+	return CurrentCapture.IsValid();
+}
+#endif
+
+void FNiagaraSystemInstance::SetSolo(bool bInSolo)
+{
+	SCOPE_CYCLE_COUNTER(STAT_NiagaraSystemSetSolo);
+	if (bSolo == bInSolo)
+	{
+		return;
+	}
+
+	UNiagaraSystem* System = GetSystem();
+	if (bInSolo)
+	{
+		TSharedPtr<FNiagaraSystemSimulation, ESPMode::ThreadSafe> NewSoloSim = MakeShared<FNiagaraSystemSimulation, ESPMode::ThreadSafe>();
+		NewSoloSim->Init(System, Component->GetWorld(), this);
+
+		NewSoloSim->TransferInstance(SystemSimulation.Get(), this);	
+
+		SystemSimulation = NewSoloSim;
+		bSolo = true;
 	}
 	else
 	{
-		ExecutionState = ENiagaraExecutionState::Disabled;
+		TSharedPtr<FNiagaraSystemSimulation, ESPMode::ThreadSafe> NewSim = GetWorldManager()->GetSystemSimulation(System);
+
+		NewSim->TransferInstance(SystemSimulation.Get(), this);
+		
+		SystemSimulation = NewSim;
+		bSolo = false;
+	}
+}
+
+void FNiagaraSystemInstance::Activate(EResetMode InResetMode)
+{
+	SCOPE_CYCLE_COUNTER(STAT_NiagaraSystemActivate);
+	
+	UNiagaraSystem* System = GetSystem();
+	if (System && System->IsValid() && IsReadyToRun())
+	{
+		Reset(InResetMode, true);
+	}
+	else
+	{
+		SetRequestedExecutionState(ENiagaraExecutionState::Disabled);
 	}
 }
 
 void FNiagaraSystemInstance::Deactivate(bool bImmediate)
 {
+
+	SCOPE_CYCLE_COUNTER(STAT_NiagaraSystemDeactivate);
 	if (IsComplete())
 	{
 		return;
@@ -107,82 +278,142 @@ void FNiagaraSystemInstance::Deactivate(bool bImmediate)
 	}
 	else
 	{
-		ExecutionState = ENiagaraExecutionState::Inactive;
+		SetRequestedExecutionState(ENiagaraExecutionState::Inactive);
 	}
 }
 
 void FNiagaraSystemInstance::Complete()
 {
-	TSharedPtr<FNiagaraSystemSimulation, ESPMode::ThreadSafe> SystemSim = GetSystemSimulation();
-	SystemSim->RemoveInstance(this);
+	SCOPE_CYCLE_COUNTER(STAT_NiagaraSystemComplete);
+	
+	// Only notify others if have yet to complete
+	bool bNeedToNotifyOthers = bNotifyOnCompletion;
 
-	ExecutionState = ENiagaraExecutionState::Complete;
+	if (SystemInstanceIndex != INDEX_NONE)
+	{
+		TSharedPtr<FNiagaraSystemSimulation, ESPMode::ThreadSafe> SystemSim = GetSystemSimulation();
+		SystemSim->RemoveInstance(this);
 
-	for (TSharedRef<FNiagaraEmitterInstance> Simulation : Emitters)
+		SetActualExecutionState(ENiagaraExecutionState::Complete);
+		SetRequestedExecutionState(ENiagaraExecutionState::Complete);
+
+		for (TSharedRef<FNiagaraEmitterInstance> Simulation : Emitters)
+		{
+			Simulation->HandleCompletion(true);
+		}
+	}
+	else
 	{
-		Simulation->HandleCompletion(true);
+		SetActualExecutionState(ENiagaraExecutionState::Complete);
+		SetRequestedExecutionState(ENiagaraExecutionState::Complete);
 	}
 
-	if (Component)
+	DestroyDataInterfaceInstanceData();
+
+	UnbindParameters();
+
+	if (bNeedToNotifyOthers)
 	{
-		Component->OnSystemComplete();
-	}
+		OnCompleteDelegate.Broadcast(this);
 
-	OnCompleteDelegate.Broadcast(this);
+		if (Component)
+		{
+			Component->OnSystemComplete();
+		}
+		
+		// We've already notified once, no need to do so again.
+		bNotifyOnCompletion = false;
+	}
 }
 
-void FNiagaraSystemInstance::Reset(FNiagaraSystemInstance::EResetMode Mode)
+void FNiagaraSystemInstance::Reset(FNiagaraSystemInstance::EResetMode Mode, bool bBindParams)
 {
 	SCOPE_CYCLE_COUNTER(STAT_NiagaraSystemReset);
 
-	TSharedPtr<FNiagaraSystemSimulation, ESPMode::ThreadSafe> SystemSim = GetSystemSimulation();
-	if (!SystemSim.IsValid())
+	FNiagaraSystemSimulation* SystemSim = GetSystemSimulation().Get();
+	if (Mode == EResetMode::None)
 	{
+		// Right now we don't support binding with reset mode none.
+		/*if (Mode == EResetMode::None && bBindParams)
+		{
+			BindParameters();
+		}*/
 		return;
 	}
 
 	Component->LastRenderTime = Component->GetWorld()->GetTimeSeconds();
 
-	SystemSim->RemoveInstance(this);
-	bPendingSpawn = true;
+	if (SystemSim)
+	{
+		SystemSim->RemoveInstance(this);
+	}
+	else
+	{
+		Mode = EResetMode::ReInit;
+	}
 
-	if (bError)
+	//If we were disabled, try to reinit on reset.
+	if (IsDisabled())
 	{
 		Mode = EResetMode::ReInit;
 	}
 		
 	if (Mode == EResetMode::ResetSystem)
 	{
+		//UE_LOG(LogNiagara, Log, TEXT("FNiagaraSystemInstance::Reset false"));
 		ResetInternal(false);
 	}
 	else if (Mode == EResetMode::ResetAll)
 	{
+		//UE_LOG(LogNiagara, Log, TEXT("FNiagaraSystemInstance::Reset true"));
 		ResetInternal(true);
 	}
 	else if (Mode == EResetMode::ReInit)
 	{
-		DestroyDataInterfaceInstanceData();
+		//UE_LOG(LogNiagara, Log, TEXT("FNiagaraSystemInstance::ReInit"));
 		ReInitInternal();
 	}
 	
-	if (!bError)
+	if (bBindParams)
 	{
-		if (bSolo)
-		{
-			SystemSim->ResetSolo(this);
-		}
-		else
+		BindParameters();
+	}
+
+	SystemSim = GetSystemSimulation().Get();
+	SetRequestedExecutionState(ENiagaraExecutionState::Active);
+	SetActualExecutionState(ENiagaraExecutionState::Active);
+
+	InitDataInterfaces();
+
+	//Interface init can disable the system.
+	if (!IsComplete())
+	{
+		bPendingSpawn = true;
+		SystemSim->AddInstance(this);
+
+		UNiagaraSystem* System = GetSystem();
+		if (System->NeedsWarmup())
 		{
-			SystemSim->AddInstance(this);
+			int32 WarmupTicks = System->GetWarmupTickCount();
+			float WarmupDt = System->GetWarmupTickDelta();
+			
+			AdvanceSimulation(WarmupTicks, WarmupDt);
+
+			//Reset age to zero.
+			Age = 0.0f;
 		}
 	}
+
+	// This system may not tick again immediately so we mark the render state dirty here so that
+	// the renderers will be reset this frame.
+	Component->MarkRenderDynamicDataDirty();
 }
 
 void FNiagaraSystemInstance::ResetInternal(bool bResetSimulations)
 {
 	Age = 0;
 	UNiagaraSystem* System = GetSystem();
-	if (System == nullptr || Component == nullptr || bError)
+	if (System == nullptr || Component == nullptr || IsDisabled())
 	{
 		return;
 	}
@@ -190,7 +421,6 @@ void FNiagaraSystemInstance::ResetInternal(bool bResetSimulations)
 #if WITH_EDITOR
 	if (Component->GetWorld() != nullptr && Component->GetWorld()->WorldType == EWorldType::Editor)
 	{
-		Component->SynchronizeWithSourceSystem();
 		Component->GetOverrideParameters().Tick();
 	}
 #endif
@@ -204,7 +434,7 @@ void FNiagaraSystemInstance::ResetInternal(bool bResetSimulations)
 
 	if (!System->IsValid())
 	{
-		bError = true;		
+		SetRequestedExecutionState(ENiagaraExecutionState::Disabled);
 		UE_LOG(LogNiagara, Warning, TEXT("Failed to activate Niagara System due to invalid asset!"));
 		return;
 	}
@@ -228,6 +458,22 @@ UNiagaraParameterCollectionInstance* FNiagaraSystemInstance::GetParameterCollect
 	return SystemSimulation->GetParameterCollectionInstance(Collection);
 }
 
+void FNiagaraSystemInstance::AdvanceSimulation(int32 TickCount, float TickDeltaSeconds)
+{
+	if (TickCount > 0)
+	{
+		SCOPE_CYCLE_COUNTER(STAT_NiagaraSystemAdvanceSim);
+		bool bWasSolo = bSolo;
+		SetSolo(true);
+
+		for (int32 TickIdx = 0; TickIdx < TickCount; ++TickIdx)
+		{
+			ComponentTick(TickDeltaSeconds);
+		}
+		SetSolo(bWasSolo);
+	}
+}
+
 bool FNiagaraSystemInstance::IsReadyToRun() const
 {
 	bool bAllReadyToRun = true;
@@ -252,21 +498,16 @@ bool FNiagaraSystemInstance::IsReadyToRun() const
 void FNiagaraSystemInstance::ReInitInternal()
 {
 	SCOPE_CYCLE_COUNTER(STAT_NiagaraSystemReinit);
-	bool bPrevError = bError;
 	Age = 0;
 	UNiagaraSystem* System = GetSystem();
 	if (System == nullptr || Component == nullptr)
 	{
 		return;
 	}
-	bError = false;
 
-#if WITH_EDITOR
-	if (Component->GetWorld() != nullptr && Component->GetWorld()->WorldType == EWorldType::Editor)
-	{
-		Component->SynchronizeWithSourceSystem();
-	}
-#endif
+	//Bypass the SetExecutionState() and it's check for disabled.
+	RequestedExecutionState = ENiagaraExecutionState::Inactive;
+	ActualExecutionState = ENiagaraExecutionState::Inactive;
 
 	bool bAllReadyToRun = IsReadyToRun();
 
@@ -277,35 +518,72 @@ void FNiagaraSystemInstance::ReInitInternal()
 	
 	if (!System->IsValid())
 	{
-		bError = true;
+		SetRequestedExecutionState(ENiagaraExecutionState::Disabled);
 		UE_LOG(LogNiagara, Warning, TEXT("Failed to activate Niagara System due to invalid asset!"));
 		return;
 	}
 
+	/** Do we need to run in solo mode? */
+	bSolo = bForceSolo || System->IsSolo();
+	if (bSolo)
+	{
+		if (!SystemSimulation.IsValid())
+		{
+			SystemSimulation = MakeShared<FNiagaraSystemSimulation, ESPMode::ThreadSafe>();
+			SystemSimulation->Init(System, Component->GetWorld(), this);
+		}
+	}
+	else
+	{
+		SystemSimulation = GetWorldManager()->GetSystemSimulation(System);
+	}
+
 	//When re initializing, throw away old emitters and init new ones.
 	Emitters.Reset();
 	InitEmitters();
 	
-	InstanceParameters.Empty();
-	InstanceParameters.AddParameter(SYS_PARAM_ENGINE_POSITION);
-	InstanceParameters.AddParameter(SYS_PARAM_ENGINE_SCALE);
-	InstanceParameters.AddParameter(SYS_PARAM_ENGINE_VELOCITY);
-	InstanceParameters.AddParameter(SYS_PARAM_ENGINE_X_AXIS);
-	InstanceParameters.AddParameter(SYS_PARAM_ENGINE_Y_AXIS);
-	InstanceParameters.AddParameter(SYS_PARAM_ENGINE_Z_AXIS);
-	InstanceParameters.AddParameter(SYS_PARAM_ENGINE_LOCAL_TO_WORLD);
-	InstanceParameters.AddParameter(SYS_PARAM_ENGINE_WORLD_TO_LOCAL);
-	InstanceParameters.AddParameter(SYS_PARAM_ENGINE_LOCAL_TO_WORLD_TRANSPOSED);
-	InstanceParameters.AddParameter(SYS_PARAM_ENGINE_WORLD_TO_LOCAL_TRANSPOSED);
-	InstanceParameters.AddParameter(SYS_PARAM_ENGINE_LOCAL_TO_WORLD_NO_SCALE);
-	InstanceParameters.AddParameter(SYS_PARAM_ENGINE_WORLD_TO_LOCAL_NO_SCALE);
-	InstanceParameters.AddParameter(SYS_PARAM_ENGINE_DELTA_TIME);
-	InstanceParameters.AddParameter(SYS_PARAM_ENGINE_INV_DELTA_TIME);
-	InstanceParameters.AddParameter(SYS_PARAM_ENGINE_TIME_SINCE_RENDERED);
-	InstanceParameters.AddParameter(SYS_PARAM_ENGINE_MIN_DIST_TO_CAMERA);
-	InstanceParameters.AddParameter(SYS_PARAM_ENGINE_SYSTEM_NUM_EMITTERS);
-	InstanceParameters.AddParameter(SYS_PARAM_ENGINE_SYSTEM_NUM_EMITTERS_ALIVE);
+	InstanceParameters.Reset();
+	InstanceParameters.AddParameter(SYS_PARAM_ENGINE_POSITION, true, false);
+	InstanceParameters.AddParameter(SYS_PARAM_ENGINE_SCALE, true, false);
+	InstanceParameters.AddParameter(SYS_PARAM_ENGINE_VELOCITY, true, false);
+	InstanceParameters.AddParameter(SYS_PARAM_ENGINE_X_AXIS, true, false);
+	InstanceParameters.AddParameter(SYS_PARAM_ENGINE_Y_AXIS, true, false);
+	InstanceParameters.AddParameter(SYS_PARAM_ENGINE_Z_AXIS, true, false);
+	InstanceParameters.AddParameter(SYS_PARAM_ENGINE_LOCAL_TO_WORLD, true, false);
+	InstanceParameters.AddParameter(SYS_PARAM_ENGINE_WORLD_TO_LOCAL, true, false);
+	InstanceParameters.AddParameter(SYS_PARAM_ENGINE_LOCAL_TO_WORLD_TRANSPOSED, true, false);
+	InstanceParameters.AddParameter(SYS_PARAM_ENGINE_WORLD_TO_LOCAL_TRANSPOSED, true, false);
+	InstanceParameters.AddParameter(SYS_PARAM_ENGINE_LOCAL_TO_WORLD_NO_SCALE, true, false);
+	InstanceParameters.AddParameter(SYS_PARAM_ENGINE_WORLD_TO_LOCAL_NO_SCALE, true, false);
+	InstanceParameters.AddParameter(SYS_PARAM_ENGINE_DELTA_TIME, true, false);
+	InstanceParameters.AddParameter(SYS_PARAM_ENGINE_TIME, true, false);
+	InstanceParameters.AddParameter(SYS_PARAM_ENGINE_REAL_TIME, true, false);
+	InstanceParameters.AddParameter(SYS_PARAM_ENGINE_INV_DELTA_TIME, true, false);
+	InstanceParameters.AddParameter(SYS_PARAM_ENGINE_TIME_SINCE_RENDERED, true, false);
+	InstanceParameters.AddParameter(SYS_PARAM_ENGINE_EXECUTION_STATE, true, false);
+	InstanceParameters.AddParameter(SYS_PARAM_ENGINE_MIN_DIST_TO_CAMERA, true, false);
+	InstanceParameters.AddParameter(SYS_PARAM_ENGINE_SYSTEM_NUM_EMITTERS, true, false);
+	InstanceParameters.AddParameter(SYS_PARAM_ENGINE_SYSTEM_NUM_EMITTERS_ALIVE, true, false);
+	InstanceParameters.AddParameter(SYS_PARAM_ENGINE_SYSTEM_AGE);
+
+	// This is required for user default data interface's (like say static meshes) to be set up properly.
+	// Additionally, it must happen here for data to be properly found below.
+	const bool bOnlyAdd = false;
+	System->GetExposedParameters().CopyParametersTo(InstanceParameters, bOnlyAdd, FNiagaraParameterStore::EDataInterfaceCopyMethod::Reference);
+
+	TArray<FNiagaraVariable> NumParticleVars;
+	for (int32 i = 0; i < Emitters.Num(); i++)
+	{
+		TSharedRef<FNiagaraEmitterInstance> Simulation = Emitters[i];
+		FString EmitterName = Simulation->GetEmitterHandle().GetInstance()->GetUniqueEmitterName();
+		FNiagaraVariable Var = SYS_PARAM_ENGINE_EMITTER_NUM_PARTICLES;
+		FString ParamName = Var.GetName().ToString().Replace(TEXT("Emitter"), *EmitterName);
+		Var.SetName(*ParamName);
+		InstanceParameters.AddParameter(Var, true, false);
+		NumParticleVars.Add(Var);
+	}
 
+	// Make sure all parameters are added before initializing the bindings, otherwise parameter store layout changes might invalidate the bindings.
 	OwnerPositionParam.Init(InstanceParameters, SYS_PARAM_ENGINE_POSITION);
 	OwnerScaleParam.Init(InstanceParameters, SYS_PARAM_ENGINE_SCALE);
 	OwnerVelocityParam.Init(InstanceParameters, SYS_PARAM_ENGINE_VELOCITY);
@@ -322,60 +600,29 @@ void FNiagaraSystemInstance::ReInitInternal()
 
 	OwnerDeltaSecondsParam.Init(InstanceParameters, SYS_PARAM_ENGINE_DELTA_TIME);
 	OwnerInverseDeltaSecondsParam.Init(InstanceParameters, SYS_PARAM_ENGINE_INV_DELTA_TIME);
+
+	SystemAgeParam.Init(InstanceParameters, SYS_PARAM_ENGINE_SYSTEM_AGE);
+	OwnerEngineTimeParam.Init(InstanceParameters, SYS_PARAM_ENGINE_TIME);
+	OwnerEngineRealtimeParam.Init(InstanceParameters, SYS_PARAM_ENGINE_REAL_TIME);
+
 	OwnerMinDistanceToCameraParam.Init(InstanceParameters, SYS_PARAM_ENGINE_MIN_DIST_TO_CAMERA);
 	SystemNumEmittersParam.Init(InstanceParameters, SYS_PARAM_ENGINE_SYSTEM_NUM_EMITTERS);
 	SystemNumEmittersAliveParam.Init(InstanceParameters, SYS_PARAM_ENGINE_SYSTEM_NUM_EMITTERS_ALIVE);
 
 	SystemTimeSinceRenderedParam.Init(InstanceParameters, SYS_PARAM_ENGINE_TIME_SINCE_RENDERED);
 
-	ParameterNumParticleBindings.SetNum(Emitters.Num());
-	for (int32 i = 0; i < Emitters.Num(); i++)
-	{
-		TSharedRef<FNiagaraEmitterInstance> Simulation = Emitters[i];
-		FString EmitterName = Simulation->GetEmitterHandle().GetInstance()->GetUniqueEmitterName();
-		FNiagaraVariable Var = SYS_PARAM_ENGINE_EMITTER_NUM_PARTICLES;
-		FString ParamName = Var.GetName().ToString().Replace(TEXT("Emitter"), *EmitterName);
-		Var.SetName(*ParamName);
-		InstanceParameters.AddParameter(Var);
-		ParameterNumParticleBindings[i].Init(InstanceParameters, Var);
-	}
-
-	const bool bOnlyAdd = false;
-	System->GetExposedParameters().CopyParametersTo(InstanceParameters, bOnlyAdd);
-
-	TickInstanceParameters(0.01f);
+	OwnerExecutionStateParam.Init(InstanceParameters, SYS_PARAM_ENGINE_EXECUTION_STATE);
 
-	//Determine if we can update normally or have to update solo.
-	bSolo = bForceSolo;
-	//If our scripts have any interfaces that require instance data.
-	UNiagaraScript* SystemSpawn = GetSystem()->GetSystemSpawnScript();
-	for (int32 i = 0; !bSolo && i < SystemSpawn->DataInterfaceInfo.Num(); ++i)
-	{
-		FNiagaraScriptDataInterfaceInfo& Info = SystemSpawn->DataInterfaceInfo[i];
-		if ( (Info.DataInterface && Info.DataInterface->PerInstanceDataSize() > 0) ||
-			Info.Name.ToString().StartsWith("User."))//Temp hack to force solo on any systems with system scrips needing user (aka per instance) interfaces.
-		{
-			bSolo = true;
-			break;
-		}
-		
-	}	
-	UNiagaraScript* SystemUpdate = GetSystem()->GetSystemUpdateScript();
-	for (int32 i = 0; !bSolo && i < SystemUpdate->DataInterfaceInfo.Num(); ++i)
+	ParameterNumParticleBindings.SetNum(NumParticleVars.Num());
+	for (int32 i = 0; i < NumParticleVars.Num(); i++)
 	{
-		FNiagaraScriptDataInterfaceInfo& Info = SystemUpdate->DataInterfaceInfo[i];
-		if ( (Info.DataInterface && Info.DataInterface->PerInstanceDataSize() > 0) ||
-			Info.Name.ToString().StartsWith("User."))//Temp hack to force solo on any systems with system scrips needing user (aka per instance) interfaces.
-		{
-			bSolo = true;
-			break;
-		}
+		ParameterNumParticleBindings[i].Init(InstanceParameters, NumParticleVars[i]);
 	}
 
-	BindParameters();
+	// rebind now after all parameters have been added
+	InstanceParameters.Rebind();
 
-	InitDataInterfaces(); // We can't wait until tick b/c the data interfaces need to be around when we update render modules.
-	//FlushRenderingCommands();//Ah! This causes massive hitch on activation of new fx. 
+	TickInstanceParameters(0.01f);
 
 	// This gets a little tricky, but we want to delete any renderers that are no longer in use on the rendering thread, but
 	// first (to be safe), we want to update the proxy to point to the new renderer objects.
@@ -412,6 +659,15 @@ FNiagaraSystemInstance::~FNiagaraSystemInstance()
 
 	//FlushRenderingCommands();
 
+	Cleanup();
+
+// #if WITH_EDITOR
+// 	OnDestroyedDelegate.Broadcast();
+// #endif
+}
+
+void FNiagaraSystemInstance::Cleanup()
+{
 	if (SystemInstanceIndex != INDEX_NONE)
 	{
 		TSharedPtr<FNiagaraSystemSimulation, ESPMode::ThreadSafe> SystemSim = GetSystemSimulation();
@@ -434,10 +690,6 @@ FNiagaraSystemInstance::~FNiagaraSystemInstance()
 
 	// Clear out the emitters.
 	Emitters.Empty(0);
-
-// #if WITH_EDITOR
-// 	OnDestroyedDelegate.Broadcast();
-// #endif
 }
 
 //Unsure on usage of this atm. Possibly useful in future.
@@ -486,7 +738,7 @@ void FNiagaraSystemInstance::InitDataInterfaces()
 	// If either the System or the component is invalid, it is possible that our cached data interfaces
 	// are now bogus and could point to invalid memory. Only the UNiagaraComponent or UNiagaraSystem
 	// can hold onto GC references to the DataInterfaces.
-	if (GetSystem() == nullptr || bError)
+	if (GetSystem() == nullptr || IsDisabled())
 	{
 		return;
 	}
@@ -498,8 +750,6 @@ void FNiagaraSystemInstance::InitDataInterfaces()
 
 	Component->GetOverrideParameters().Tick();
 	
-	bDataInterfacesNeedInit = false;
-
 	DestroyDataInterfaceInstanceData();
 
 	//Now the interfaces in the simulations are all correct, we can build the per instance data table.
@@ -520,20 +770,15 @@ void FNiagaraSystemInstance::InitDataInterfaces()
 				if (!ExistingInstanceDataOffset)//Don't add instance data for interfaces we've seen before.
 				{
 					DataInterfaceInstanceDataOffsets.Add(Interface) = InstanceDataSize;
-					InstanceDataSize += Size;
+					// Assume that some of our data is going to be 16 byte aligned, so enforce that 
+					// all per-instance data is aligned that way.
+					InstanceDataSize += Align(Size, 16);
 				}
 			}
 		}
 	};
 
 	CalcInstDataSize(InstanceParameters.GetDataInterfaces());//This probably should be a proper exec context. 
-	if (bSolo)
-	{
-		TSharedPtr<FNiagaraSystemSimulation, ESPMode::ThreadSafe> SystemSim = GetSystemSimulation();
-		
-		CalcInstDataSize(SystemSim->GetSoloDataInterfacesSpawn());
-		CalcInstDataSize(SystemSim->GetSoloDataInterfacesUpdate());
-	}
 
 	//Iterate over interfaces to get size for table and clear their interface bindings.
 	for (TSharedRef<FNiagaraEmitterInstance> Simulation : Emitters)
@@ -557,8 +802,15 @@ void FNiagaraSystemInstance::InitDataInterfaces()
 	{
 		if (UNiagaraDataInterface* Interface = Pair.Key.Get())
 		{
+			check(IsAligned(&DataInterfaceInstanceData[Pair.Value], 16));
+
 			//Ideally when we make the batching changes, we can keep the instance data in big single type blocks that can all be updated together with a single virtual call.
-			bOk &= Pair.Key->InitPerInstanceData(&DataInterfaceInstanceData[Pair.Value], this);
+			bool bResult = Pair.Key->InitPerInstanceData(&DataInterfaceInstanceData[Pair.Value], this);
+			bOk &= bResult;
+			if (!bResult)
+			{
+				UE_LOG(LogNiagara, Error, TEXT("Error initializing data interface \"%s\" for system. %u | %s"), *Interface->GetPathName(), Component, *Component->GetAsset()->GetName());
+			}
 		}
 		else
 		{
@@ -567,42 +819,48 @@ void FNiagaraSystemInstance::InitDataInterfaces()
 		}
 	}
 
-	if (!bOk)
+	if (!bOk && (!IsComplete() && !IsPendingSpawn()))
 	{
-		UE_LOG(LogNiagara, Error, TEXT("Error initializing data interfaces."));
-		bError = true;
-		Component->MarkRenderStateDirty();
+		//Some error initializing the data interfaces so disable until we're explicitly reinitialized.
+		UE_LOG(LogNiagara, Error, TEXT("Error initializing data interfaces. Completing system. %u | %s"), Component, *Component->GetAsset()->GetName());
+		Complete();
 	}
 }
 
 void FNiagaraSystemInstance::TickDataInterfaces(float DeltaSeconds, bool bPostSimulate)
 {
-	if (!GetSystem() || !Component || bError)
+	if (!GetSystem() || !Component || IsDisabled())
 	{
 		return;
 	}
 
 	bool bReInitDataInterfaces = false;
-	for (TPair<TWeakObjectPtr<UNiagaraDataInterface>, int32>& Pair : DataInterfaceInstanceDataOffsets)
+	if (bPostSimulate)
 	{
-		if(UNiagaraDataInterface* Interface = Pair.Key.Get())
+		for (TPair<TWeakObjectPtr<UNiagaraDataInterface>, int32>& Pair : DataInterfaceInstanceDataOffsets)
 		{
-			//Ideally when we make the batching changes, we can keep the instance data in big single type blocks that can all be updated together with a single virtual call.
-			if (bPostSimulate)
+			if (UNiagaraDataInterface* Interface = Pair.Key.Get())
 			{
+				//Ideally when we make the batching changes, we can keep the instance data in big single type blocks that can all be updated together with a single virtual call.
 				bReInitDataInterfaces |= Interface->PerInstanceTickPostSimulate(&DataInterfaceInstanceData[Pair.Value], this, DeltaSeconds);
 			}
-			else
+		}
+	}
+	else
+	{
+		for (TPair<TWeakObjectPtr<UNiagaraDataInterface>, int32>& Pair : DataInterfaceInstanceDataOffsets)
+		{
+			if (UNiagaraDataInterface* Interface = Pair.Key.Get())
 			{
+				//Ideally when we make the batching changes, we can keep the instance data in big single type blocks that can all be updated together with a single virtual call.
 				bReInitDataInterfaces |= Interface->PerInstanceTick(&DataInterfaceInstanceData[Pair.Value], this, DeltaSeconds);
 			}
 		}
 	}
-	
-	//This should ideally really only happen at edit time.
+
 	if (bReInitDataInterfaces)
 	{
-		InitDataInterfaces();	
+		InitDataInterfaces();
 	}
 }
 
@@ -670,7 +928,16 @@ void FNiagaraSystemInstance::TickInstanceParameters(float DeltaSeconds)
 			}
 		}
 		OwnerMinDistanceToCameraParam.SetValue(FMath::Sqrt(LODDistanceSqr));
+
+		OwnerEngineTimeParam.SetValue(World->TimeSeconds);
+		OwnerEngineRealtimeParam.SetValue(World->RealTimeSeconds);
+	}
+	else
+	{
+		OwnerEngineTimeParam.SetValue(Age);
+		OwnerEngineRealtimeParam.SetValue(Age);
 	}
+	SystemAgeParam.SetValue(Age);
 
 	int32 NumAlive = 0;
 	for (int32 i = 0; i < Emitters.Num(); i++)
@@ -688,7 +955,8 @@ void FNiagaraSystemInstance::TickInstanceParameters(float DeltaSeconds)
 	check(World);
 	float SafeTimeSinceRendererd = FMath::Max(0.0f, World->GetTimeSeconds() - Component->LastRenderTime - GLastRenderTimeSafetyBias);
 	SystemTimeSinceRenderedParam.SetValue(SafeTimeSinceRendererd);
-
+	
+	OwnerExecutionStateParam.SetValue((int32)RequestedExecutionState);
 	
 	Component->GetOverrideParameters().Tick();
 	InstanceParameters.Tick();
@@ -701,13 +969,7 @@ bool FNiagaraSystemInstance::UsesEmitter(const UNiagaraEmitter* Emitter)const
 {
 	if (GetSystem())
 	{
-		for (FNiagaraEmitterHandle EmitterHandle : GetSystem()->GetEmitterHandles())
-		{
-			if (Emitter == EmitterHandle.GetSource() || Emitter == EmitterHandle.GetInstance())
-			{
-				return true;
-			}
-		}
+		return GetSystem()->UsesEmitter(Emitter);
 	}
 	return false;
 }
@@ -748,21 +1010,23 @@ bool FNiagaraSystemInstance::UsesCollection(const UNiagaraParameterCollection* C
 
 void FNiagaraSystemInstance::InitEmitters()
 {
+	SCOPE_CYCLE_COUNTER(STAT_NiagaraSystemInitEmitters);
 	if (Component)
 	{
 		Component->MarkRenderStateDirty();
 	}
 
-	// Just in case this ends up being called more than in Init, we need to 
-	// clear out the update proxy of any renderers that will be destroyed when Emitters.Empty occurs..
-	TArray<NiagaraRenderer*> NewRenderers;
-	UpdateProxy(NewRenderers);
-
-	// Clear out the System renderer from the simulation.
-	for (TSharedRef<FNiagaraEmitterInstance> Simulation : Emitters)
-	{
-		Simulation->ClearRenderer();
-	}
+	//STovey: This is done after InitEmitters so removing this duplication.
+// 	// Just in case this ends up being called more than in Init, we need to 
+// 	// clear out the update proxy of any renderers that will be destroyed when Emitters.Empty occurs..
+// 	TArray<NiagaraRenderer*> NewRenderers;
+// 	UpdateProxy(NewRenderers);
+// 
+// 	// Clear out the System renderer from the simulation.
+// 	for (TSharedRef<FNiagaraEmitterInstance> Simulation : Emitters)
+// 	{
+// 		Simulation->ClearRenderer();
+// 	}
 
 	Emitters.Empty();
 	if (GetSystem() != nullptr)
@@ -780,8 +1044,6 @@ void FNiagaraSystemInstance::InitEmitters()
 		{
 			Simulation->PostInitSimulation();
 		}
-
-		bDataInterfacesNeedInit = true;
 	}
 }
 
@@ -795,7 +1057,7 @@ void FNiagaraSystemInstance::UpdateRenderModules(ERHIFeatureLevel::Type InFeatur
 
 void FNiagaraSystemInstance::UpdateProxy(TArray<NiagaraRenderer*>& InRenderers)
 {
-	if (!GetSystem() || !Component)
+	if (!Component)
 	{
 		return;
 	}
@@ -803,7 +1065,7 @@ void FNiagaraSystemInstance::UpdateProxy(TArray<NiagaraRenderer*>& InRenderers)
 	FNiagaraSceneProxy *NiagaraProxy = static_cast<FNiagaraSceneProxy*>(Component->SceneProxy);
 	if (NiagaraProxy)
 	{
-		if (Component != nullptr && Component->GetWorld() != nullptr)
+		if (Component->GetWorld() != nullptr)
 		{
 			// Tell the scene proxy on the render thread to update its System renderers.
 			ENQUEUE_UNIQUE_RENDER_COMMAND_TWOPARAMETER(
@@ -820,7 +1082,7 @@ void FNiagaraSystemInstance::UpdateProxy(TArray<NiagaraRenderer*>& InRenderers)
 
 void FNiagaraSystemInstance::ComponentTick(float DeltaSeconds)
 {
-	if (bError)
+	if (IsDisabled())
 	{
 		return;
 	}
@@ -830,15 +1092,20 @@ void FNiagaraSystemInstance::ComponentTick(float DeltaSeconds)
 	check(IsInGameThread());
 	check(bSolo);
 	check(Component);
-
-	PreSimulateTick(DeltaSeconds);
-
+	
 	TSharedPtr<FNiagaraSystemSimulation, ESPMode::ThreadSafe> SystemSim = GetSystemSimulation();
-	SystemSim->TickSolo(this);
+	SystemSim->Tick(DeltaSeconds);
+}
 
-	if(!HandleCompletion())
+void FNiagaraSystemInstance::FinalizeTick(float DeltaSeconds)
+{
+	//Post tick our interfaces.
+	TickDataInterfaces(DeltaSeconds, true);
+
+	if (HasTickingEmitters())
 	{
-		PostSimulateTick(DeltaSeconds);
+		Component->UpdateComponentToWorld();//Needed for bounds updates. Can probably skip if using fixed bounds
+		Component->MarkRenderDynamicDataDirty();
 	}
 }
 
@@ -853,8 +1120,10 @@ bool FNiagaraSystemInstance::HandleCompletion()
 		bHasTickingEmitters |= Inst.ShouldTick();
 	}
 
-	if (IsComplete() || bEmittersCompleteOrDisabled)
+	bool bCompletedAlready = IsComplete();
+	if (bCompletedAlready || bEmittersCompleteOrDisabled)
 	{
+		//UE_LOG(LogNiagara, Log, TEXT("Completion Achieved"));
 		Complete();
 		return true;
 	}
@@ -863,12 +1132,13 @@ bool FNiagaraSystemInstance::HandleCompletion()
 
 void FNiagaraSystemInstance::PreSimulateTick(float DeltaSeconds)
 {
+	SCOPE_CYCLE_COUNTER(STAT_NiagaraSystemPreSimulateTick);
 	TickInstanceParameters(DeltaSeconds);
 }
 
 void FNiagaraSystemInstance::PostSimulateTick(float DeltaSeconds)
 {
-	if (IsComplete() || !bHasTickingEmitters || GetSystem() == nullptr || Component == nullptr || DeltaSeconds < SMALL_NUMBER || bError)
+	if (IsComplete() || !bHasTickingEmitters || GetSystem() == nullptr || Component == nullptr || DeltaSeconds < SMALL_NUMBER)
 	{
 		return;
 	}
@@ -880,31 +1150,9 @@ void FNiagaraSystemInstance::PostSimulateTick(float DeltaSeconds)
 		EventSetPair.Value.Tick();
 	}
 
-	if (bDataInterfacesNeedInit)
-	{
-		InitDataInterfaces();
-		// We may have entered into an error state as a result of InitDataInterfaces.
-		if (bError)
-		{
-			return;
-		}
-	}
-
-	//TODO: Now we're batching in the world manager we can store these in big blocks across systems and update all instance data in one virtual call?
-	TickDataInterfaces(DeltaSeconds, false);
-
-	// We may have entered into an error state as a result of TickDataInterfaces.
-	if (bError)
-	{
-		return;
-	}
-
-	// need to run this serially, as PreTick will initialize data interfaces, which can't be done in parallel
 	for (int32 EmitterIdx = 0; EmitterIdx < Emitters.Num(); EmitterIdx++)
 	{
-		FNiagaraEmitterInstance& Inst = Emitters[EmitterIdx].Get();		
-		const UNiagaraEmitter *Props = Inst.GetEmitterHandle().GetInstance();
-		check(Props);
+		FNiagaraEmitterInstance& Inst = Emitters[EmitterIdx].Get();
 		Inst.PreTick();
 	}
 
@@ -915,8 +1163,6 @@ void FNiagaraSystemInstance::PostSimulateTick(float DeltaSeconds)
 		Inst.Tick(DeltaSeconds);
 	}
 
-	TickDataInterfaces(DeltaSeconds, true);
-
 	Age += DeltaSeconds;
 }
 
@@ -932,12 +1178,6 @@ bool FNiagaraSystemInstance::GetIsolateEnabled() const
 }
 #endif
 
-void FNiagaraSystemInstance::ReinitializeDataInterfaces()
-{
-	// Update next tick
-	bDataInterfacesNeedInit = true;
-}
-
 void FNiagaraSystemInstance::DestroyDataInterfaceInstanceData()
 {
 	for (TPair<TWeakObjectPtr<UNiagaraDataInterface>, int32>& Pair : DataInterfaceInstanceDataOffsets)
@@ -965,7 +1205,19 @@ TSharedPtr<FNiagaraEmitterInstance> FNiagaraSystemInstance::GetSimulationForHand
 
 UNiagaraSystem* FNiagaraSystemInstance::GetSystem()const
 {
-	return GetSystemSimulation()->GetSystem(); 
+	return Component->GetAsset();
+}
+
+FNiagaraEmitterInstance* FNiagaraSystemInstance::GetEmitterByID(FGuid InID)
+{
+	for (TSharedRef<FNiagaraEmitterInstance>& Emitter : Emitters)
+	{
+		if (Emitter->GetEmitterHandle().GetId() == InID)
+		{
+			return &Emitter.Get();
+		}
+	}
+	return nullptr;
 }
 
 FNiagaraDataSet* FNiagaraSystemInstance::GetDataSet(FNiagaraDataSetID SetID, FName EmitterName)
@@ -982,7 +1234,7 @@ FNiagaraDataSet* FNiagaraSystemInstance::GetDataSet(FNiagaraDataSetID SetID, FNa
 		check(Emitter.IsValid());
 		if (!Emitter->IsComplete())
 		{
-			if (Emitter->GetEmitterHandle().GetIdName() == EmitterName)
+			if (Emitter->GetCachedIDName() == EmitterName)
 			{
 				return Emitter->GetDataSet(SetID);
 			}
