Copying //UE4/Dev-Niagara to //UE4/Dev-Main (Source: //UE4/Dev-Niagara @ 4074996)

diff --git a/Engine/Plugins/FX/Niagara/Source/Niagara/Private/NiagaraSystemSimulation.cpp b/Engine/Plugins/FX/Niagara/Source/Niagara/Private/NiagaraSystemSimulation.cpp
--- a/Engine/Plugins/FX/Niagara/Source/Niagara/Private/NiagaraSystemSimulation.cpp
+++ b/Engine/Plugins/FX/Niagara/Source/Niagara/Private/NiagaraSystemSimulation.cpp
@@ -14,16 +14,15 @@
 #include "NiagaraComponent.h"
 #include "NiagaraWorldManager.h"
 
-DECLARE_CYCLE_STAT(TEXT("System Simulation (Batched)"), STAT_NiagaraSystemSim, STATGROUP_Niagara);
-DECLARE_CYCLE_STAT(TEXT("System Pre Simulate (Batched)"), STAT_NiagaraSystemSim_PreSimulate, STATGROUP_Niagara);
-DECLARE_CYCLE_STAT(TEXT("System Update (Batched)"), STAT_NiagaraSystemSim_Update, STATGROUP_Niagara);
-DECLARE_CYCLE_STAT(TEXT("System Spawn (Batched)"), STAT_NiagaraSystemSim_Spawn, STATGROUP_Niagara);
-DECLARE_CYCLE_STAT(TEXT("System Transfer Parameters (Batched)"), STAT_NiagaraSystemSim_TransferParameters, STATGROUP_Niagara);
-DECLARE_CYCLE_STAT(TEXT("System Post Simulate (Batched)"), STAT_NiagaraSystemSim_PostSimulate, STATGROUP_Niagara);
+DECLARE_CYCLE_STAT(TEXT("System Simulation"), STAT_NiagaraSystemSim, STATGROUP_Niagara);
+DECLARE_CYCLE_STAT(TEXT("System Pre Simulate"), STAT_NiagaraSystemSim_PreSimulate, STATGROUP_Niagara);
+DECLARE_CYCLE_STAT(TEXT("System Prepare For Simulate"), STAT_NiagaraSystemSim_PrepareForSimulate, STATGROUP_Niagara);
+DECLARE_CYCLE_STAT(TEXT("System Update"), STAT_NiagaraSystemSim_Update, STATGROUP_Niagara);
+DECLARE_CYCLE_STAT(TEXT("System Spawn"), STAT_NiagaraSystemSim_Spawn, STATGROUP_Niagara);
+DECLARE_CYCLE_STAT(TEXT("System Transfer Parameters"), STAT_NiagaraSystemSim_TransferParameters, STATGROUP_Niagara);
+DECLARE_CYCLE_STAT(TEXT("System Post Simulate"), STAT_NiagaraSystemSim_PostSimulate, STATGROUP_Niagara);
 DECLARE_CYCLE_STAT(TEXT("System Mark Component Dirty"), STAT_NiagaraSystemSim_MarkComponentDirty, STATGROUP_Niagara);
 
-DECLARE_CYCLE_STAT(TEXT("System Simulation (Solo)"), STAT_NiagaraSystemSimSolo, STATGROUP_Niagara);
-
 
 static int32 GbDumpSystemData = 0;
 static FAutoConsoleVariableRef CVarNiagaraDumpSystemData(
@@ -33,15 +32,16 @@ static FAutoConsoleVariableRef CVarNiagaraDumpSystemData(
 	ECVF_Default
 );
 
-static int32 GbParallelSystemPreTick = 0;
-static FAutoConsoleVariableRef CVarNiagaraParallelSystemPreTick(
-	TEXT("fx.ParallelSystemPreTick"),
-	GbParallelSystemPreTick,
-	TEXT("If > 0, system pre tick is parallelized. \n"),
-	ECVF_Default
-);
+//Pretick can no longer be run in parallel. Will likely remain this way.
+// static int32 GbParallelSystemPreTick = 0;
+// static FAutoConsoleVariableRef CVarNiagaraParallelSystemPreTick(
+// 	TEXT("fx.ParallelSystemPreTick"),
+// 	GbParallelSystemPreTick,
+// 	TEXT("If > 0, system pre tick is parallelized. \n"),
+// 	ECVF_Default
+// );
 
-static int32 GbParallelSystemPostTick = 0;
+static int32 GbParallelSystemPostTick = 1;
 static FAutoConsoleVariableRef CVarNiagaraParallelSystemPostTick(
 	TEXT("fx.ParallelSystemPostTick"),
 	GbParallelSystemPostTick,
@@ -65,37 +65,37 @@ FNiagaraSystemSimulation::~FNiagaraSystemSimulation()
 	Destroy();
 }
 
-bool FNiagaraSystemSimulation::Init(UNiagaraSystem* InSystem, UWorld* InWorld)
+bool FNiagaraSystemSimulation::Init(UNiagaraSystem* InSystem, UWorld* InWorld, FNiagaraSystemInstance* InSoloSystemInstance)
 {
 	UNiagaraSystem* System = InSystem;
 	WeakSystem = System;
+	SoloSystemInstance = InSoloSystemInstance;
 
 	World = InWorld;
 
 	FNiagaraWorldManager* WorldMan = FNiagaraWorldManager::Get(InWorld);
 	check(WorldMan);
 
-	bCanExecute = System->GetSystemSpawnScript()->IsValid() && System->GetSystemUpdateScript()->IsValid();
-	bCanExecuteSolo = System->GetSystemSpawnScript(true)->IsValid() && System->GetSystemUpdateScript(true)->IsValid();
-	const UEnum* EnumPtr = FNiagaraTypeDefinition::GetExecutionStateEnum();
+	bCanExecute = System->GetSystemSpawnScript()->GetVMExecutableData().IsValid() && System->GetSystemUpdateScript()->GetVMExecutableData().IsValid();
+	UEnum* EnumPtr = FNiagaraTypeDefinition::GetExecutionStateEnum();
 
 	if (bCanExecute)
 	{
-		DataSet.Reset();
-		DataSet.AddVariables(System->GetSystemSpawnScript()->Attributes);
-		DataSet.AddVariables(System->GetSystemUpdateScript()->Attributes);
+		DataSet.Init(FNiagaraDataSetID(), ENiagaraSimTarget::CPUSim);
+		DataSet.AddVariables(System->GetSystemSpawnScript()->GetVMExecutableData().Attributes);
+		DataSet.AddVariables(System->GetSystemUpdateScript()->GetVMExecutableData().Attributes);
 		DataSet.Finalize();
 
 		{
-			SpawnInstanceParameterDataSet.Reset();
-			FNiagaraParameters* EngineParamsSpawn = System->GetSystemSpawnScript()->DataSetToParameters.Find(TEXT("Engine"));
+			SpawnInstanceParameterDataSet.Init(FNiagaraDataSetID(), ENiagaraSimTarget::CPUSim);
+			FNiagaraParameters* EngineParamsSpawn = System->GetSystemSpawnScript()->GetVMExecutableData().DataSetToParameters.Find(TEXT("Engine"));
 			if (EngineParamsSpawn != nullptr)
 			{
 				SpawnInstanceParameterDataSet.AddVariables(EngineParamsSpawn->Parameters);
 			}
 			SpawnInstanceParameterDataSet.Finalize();
-			UpdateInstanceParameterDataSet.Reset();
-			FNiagaraParameters* EngineParamsUpdate = System->GetSystemUpdateScript()->DataSetToParameters.Find(TEXT("Engine"));
+			UpdateInstanceParameterDataSet.Init(FNiagaraDataSetID(), ENiagaraSimTarget::CPUSim);
+			FNiagaraParameters* EngineParamsUpdate = System->GetSystemUpdateScript()->GetVMExecutableData().DataSetToParameters.Find(TEXT("Engine"));
 			if (EngineParamsUpdate != nullptr)
 			{
 				UpdateInstanceParameterDataSet.AddVariables(EngineParamsUpdate->Parameters);
@@ -103,22 +103,30 @@ bool FNiagaraSystemSimulation::Init(UNiagaraSystem* InSystem, UWorld* InWorld)
 			UpdateInstanceParameterDataSet.Finalize();
 		}
 
-		UNiagaraScript* SpawnScript = System->GetSystemSpawnScript(false);
-		UNiagaraScript* UpdateScript = System->GetSystemUpdateScript(false);
+		UNiagaraScript* SpawnScript = System->GetSystemSpawnScript();
+		UNiagaraScript* UpdateScript = System->GetSystemUpdateScript();
 
 		SpawnExecContext.Init(SpawnScript, ENiagaraSimTarget::CPUSim);
 		UpdateExecContext.Init(UpdateScript, ENiagaraSimTarget::CPUSim);
 
 		//Bind parameter collections.
-		for (UNiagaraParameterCollection* Collection : SpawnScript->ParameterCollections)
+		for (UNiagaraParameterCollection* Collection : SpawnScript->GetCachedParameterCollectionReferences())
 		{
 			GetParameterCollectionInstance(Collection)->GetParameterStore().Bind(&SpawnExecContext.Parameters);
 		}
-		for (UNiagaraParameterCollection* Collection : UpdateScript->ParameterCollections)
+		for (UNiagaraParameterCollection* Collection : UpdateScript->GetCachedParameterCollectionReferences())
 		{
 			GetParameterCollectionInstance(Collection)->GetParameterStore().Bind(&UpdateExecContext.Parameters);
 		}
 
+		TArray<UNiagaraScript*> Scripts;
+		Scripts.Add(SpawnScript);
+		Scripts.Add(UpdateScript);
+		FNiagaraUtilities::CollectScriptDataInterfaceParameters(*System, Scripts, ScriptDefinedDataInterfaceParameters);
+
+		ScriptDefinedDataInterfaceParameters.Bind(&SpawnExecContext.Parameters);
+		ScriptDefinedDataInterfaceParameters.Bind(&UpdateExecContext.Parameters);
+
 		SpawnScript->RapidIterationParameters.Bind(&SpawnExecContext.Parameters);
 		UpdateScript->RapidIterationParameters.Bind(&UpdateExecContext.Parameters);
 
@@ -142,53 +150,17 @@ bool FNiagaraSystemSimulation::Init(UNiagaraSystem* InSystem, UWorld* InWorld)
 				EmitterSpawnInfoAccessors[EmitterIdx].Emplace(DataSet, FNiagaraVariable(FNiagaraTypeDefinition(FNiagaraSpawnInfo::StaticStruct()), AttrName));
 			}
 		}
-	}
-
-	if (bCanExecuteSolo)
-	{
-		UNiagaraScript* SpawnScript = System->GetSystemSpawnScript(true);
-		UNiagaraScript* UpdateScript = System->GetSystemUpdateScript(true);
 
-		//Init solo simulation.
-		DataSetSolo.Reset();
-		DataSetSolo.AddVariables(SpawnScript->Attributes);
-		DataSetSolo.AddVariables(UpdateScript->Attributes);
-		DataSetSolo.Finalize();
-
-		SpawnExecContextSolo.Init(SpawnScript, ENiagaraSimTarget::CPUSim);
-		UpdateExecContextSolo.Init(UpdateScript, ENiagaraSimTarget::CPUSim);
-
-		SoloSystemExecutionStateAccessor.Create(&DataSetSolo, FNiagaraVariable(EnumPtr, TEXT("System.ExecutionState")));
-		SoloEmitterExecutionStateAccessors.Reset();
-		SoloEmitterSpawnInfoAccessors.Reset();
-		SoloEmitterSpawnInfoAccessors.SetNum(System->GetNumEmitters());
-
-		for (int32 EmitterIdx = 0; EmitterIdx < System->GetNumEmitters(); ++EmitterIdx)
-		{
-			FNiagaraEmitterHandle& EmitterHandle = System->GetEmitterHandle(EmitterIdx);
-			UNiagaraEmitter* Emitter = EmitterHandle.GetInstance();
-			FString EmitterName = Emitter->GetUniqueEmitterName();
-			check(Emitter);
-			SoloEmitterExecutionStateAccessors.Emplace(DataSetSolo, FNiagaraVariable(EnumPtr, *(EmitterName + TEXT(".ExecutionState"))));
-			
-			for (FName AttrName : System->GetEmitterSpawnAttributes()[EmitterIdx].SpawnAttributes)
-			{
-				SoloEmitterSpawnInfoAccessors[EmitterIdx].Emplace(DataSetSolo, FNiagaraVariable(FNiagaraTypeDefinition(FNiagaraSpawnInfo::StaticStruct()), AttrName));
-			}
-		}
-
-		//Bind the parameter collections and RI parameters into the Solo Execution contexts
-		for (UNiagaraParameterCollection* Collection : SpawnScript->ParameterCollections)
-		{
-			GetParameterCollectionInstance(Collection)->GetParameterStore().Bind(&SpawnExecContextSolo.Parameters);
-		}
-		SpawnScript->RapidIterationParameters.Bind(&SpawnExecContextSolo.Parameters);
-
-		for (UNiagaraParameterCollection* Collection : UpdateScript->ParameterCollections)
-		{
-			GetParameterCollectionInstance(Collection)->GetParameterStore().Bind(&UpdateExecContextSolo.Parameters);
-		}
-		UpdateScript->RapidIterationParameters.Bind(&UpdateExecContextSolo.Parameters);
+		SpawnDeltaTimeParam.Init(SpawnExecContext.Parameters, SYS_PARAM_ENGINE_DELTA_TIME);
+		UpdateDeltaTimeParam.Init(UpdateExecContext.Parameters, SYS_PARAM_ENGINE_DELTA_TIME);
+		SpawnInvDeltaTimeParam.Init(SpawnExecContext.Parameters, SYS_PARAM_ENGINE_INV_DELTA_TIME);
+		UpdateInvDeltaTimeParam.Init(UpdateExecContext.Parameters, SYS_PARAM_ENGINE_INV_DELTA_TIME);
+		SpawnNumSystemInstancesParam.Init(SpawnExecContext.Parameters, SYS_PARAM_ENGINE_NUM_SYSTEM_INSTANCES);
+		UpdateNumSystemInstancesParam.Init(UpdateExecContext.Parameters, SYS_PARAM_ENGINE_NUM_SYSTEM_INSTANCES);
+		SpawnGlobalSpawnCountScaleParam.Init(SpawnExecContext.Parameters, SYS_PARAM_ENGINE_GLOBAL_SPAWN_COUNT_SCALE);
+		UpdateGlobalSpawnCountScaleParam.Init(UpdateExecContext.Parameters, SYS_PARAM_ENGINE_GLOBAL_SPAWN_COUNT_SCALE);
+		SpawnGlobalSystemCountScaleParam.Init(SpawnExecContext.Parameters, SYS_PARAM_ENGINE_GLOBAL_SYSTEM_COUNT_SCALE);
+		UpdateGlobalSystemCountScaleParam.Init(UpdateExecContext.Parameters, SYS_PARAM_ENGINE_GLOBAL_SYSTEM_COUNT_SCALE);
 	}
 
 	return true;
@@ -204,17 +176,11 @@ void FNiagaraSystemSimulation::Destroy()
 	{
 		PendingSystemInstances.Last()->Deactivate(true);
 	}
-	while (SoloSystemInstances.Num())
-	{
-		SoloSystemInstances.Last()->Deactivate(true);
-	}
 
 	FNiagaraWorldManager* WorldMan = FNiagaraWorldManager::Get(World);
 	check(WorldMan);
 	SpawnExecContext.Parameters.UnbindFromSourceStores();
 	UpdateExecContext.Parameters.UnbindFromSourceStores();
-	SpawnExecContextSolo.Parameters.UnbindFromSourceStores();
-	UpdateExecContextSolo.Parameters.UnbindFromSourceStores();
 }
 
 UNiagaraParameterCollectionInstance* FNiagaraSystemSimulation::GetParameterCollectionInstance(UNiagaraParameterCollection* Collection)
@@ -233,140 +199,191 @@ UNiagaraParameterCollectionInstance* FNiagaraSystemSimulation::GetParameterColle
 	return Ret;
 }
 
-void FNiagaraSystemSimulation::TickSoloDataSet()
+FNiagaraParameterStore& FNiagaraSystemSimulation::GetScriptDefinedDataInterfaceParameters()
 {
-	//Flip buffers and shrink the solo data set allocation.
-	if (bCanExecuteSolo)
+	return ScriptDefinedDataInterfaceParameters;
+}
+
+void FNiagaraSystemSimulation::TransferInstance(FNiagaraSystemSimulation* SourceSimulation, FNiagaraSystemInstance* SystemInst)
+{
+	check(SourceSimulation->GetSystem() == GetSystem());
+	check(SystemInst);
+
+	int32 SystemInstIdx = SystemInst->SystemInstanceIndex;
+	int32 NewDataIndex = INDEX_NONE;
+	if (!SystemInst->IsPendingSpawn())
 	{
-		DataSetSolo.Tick();
-		DataSetSolo.Allocate(SoloSystemInstances.Num());
-		DataSetSolo.SetNumInstances(SoloSystemInstances.Num());
+// 		UE_LOG(LogNiagara, Log, TEXT("== Dataset Transfer ========================"));
+// 		UE_LOG(LogNiagara, Log, TEXT(" ----- Existing values in src. Idx: %d -----"), SystemInstIdx);
+// 		SourceSimulation->DataSet.Dump(true, SystemInstIdx, 1);
+
+		//If we're not pending then the system actually has data to pull over.
+		NewDataIndex = DataSet.TransferInstance(SourceSimulation->DataSet, SystemInstIdx);
+
+// 		UE_LOG(LogNiagara, Log, TEXT(" ----- Transfered values in dest. Idx: %d -----"), NewDataIndex);
+// 		DataSet.Dump(true, NewDataIndex, 1);
+	
+		SourceSimulation->RemoveInstance(SystemInst);
+	
+		//Move the system direct to the new sim's 
+		SystemInst->SystemInstanceIndex = SystemInstances.Add(SystemInst);
+
+		check(NewDataIndex == SystemInst->SystemInstanceIndex);
+	}
+	else
+	{
+		SourceSimulation->RemoveInstance(SystemInst);
+
+		AddInstance(SystemInst);			
 	}
 }
 
 bool FNiagaraSystemSimulation::Tick(float DeltaSeconds)
 {
 	UNiagaraSystem* System = WeakSystem.Get();
-	if (System == nullptr)
+	if (System == nullptr || bCanExecute == false)
 	{
 		// TODO: evaluate whether or not we should have removed this from the world manager instead?
 		return false;
 	}
 
-	UNiagaraScript* SystemSpawnScript = System->GetSystemSpawnScript(false);
-	UNiagaraScript* SystemUpdateScript = System->GetSystemUpdateScript(false);
-	UNiagaraScript* SystemSpawnScriptSolo = System->GetSystemSpawnScript(true);
-	UNiagaraScript* SystemUpdateScriptSolo = System->GetSystemUpdateScript(true);
+	UNiagaraScript* SystemSpawnScript = System->GetSystemSpawnScript();
+	UNiagaraScript* SystemUpdateScript = System->GetSystemUpdateScript();
+#if WITH_EDITOR
 	SystemSpawnScript->RapidIterationParameters.Tick();
 	SystemUpdateScript->RapidIterationParameters.Tick();
-	SystemSpawnScriptSolo->RapidIterationParameters.Tick();
-	SystemUpdateScriptSolo->RapidIterationParameters.Tick();
-
-	TickSoloDataSet();
+#endif
 
 	SCOPE_CYCLE_COUNTER(STAT_NiagaraSystemSim);
 
 	int32 OrigNum = SystemInstances.Num();
 	int32 SpawnNum = PendingSystemInstances.Num();
 	int32 NewNum = OrigNum + SpawnNum;
-	
-	SystemInstances.Reserve(NewNum);
-	for (FNiagaraSystemInstance* Inst : PendingSystemInstances)
-	{
-		Inst->SetPendingSpawn(false);
-		Inst->SystemInstanceIndex = SystemInstances.Add(Inst);
-	}
-	PendingSystemInstances.Reset();
-
-	// Go ahead an initialize the data set accessors that we will potentially write since they 
-	// may have changed internally from the previous frame. We write to current here because we
-	// call tick after pretick, which will make these values the previous values during execution.
-	SystemExecutionStateAccessor.InitForAccess(true);
-	for (int32 EmitterIdx = 0; EmitterIdx < System->GetNumEmitters(); ++EmitterIdx)
-	{
-		EmitterExecutionStateAccessors[EmitterIdx].InitForAccess(true);
-	}
-	
 
 	{
 		SCOPE_CYCLE_COUNTER(STAT_NiagaraSystemSim_PreSimulate);
 
-		auto PreTick = [&](int32 SystemIndex)
 		{
-			FNiagaraSystemInstance* Inst = SystemInstances[SystemIndex];
-			Inst->PreSimulateTick(DeltaSeconds);
-
-			//System exec state can be driven externally so we must feed it in as well as take it out.
-			if (SystemExecutionStateAccessor.IsValid())
+			int32 SystemIndex = 0;
+			while (SystemIndex < SystemInstances.Num())
 			{
-				SystemExecutionStateAccessor.Set(SystemIndex, (int32)Inst->GetExecutionState());
-			}
+				FNiagaraSystemInstance* Inst = SystemInstances[SystemIndex];
 
-			if (Inst->GetParameters().GetParametersDirty() && bCanExecute)
-			{
-				SpawnInstanceParameterToDataSetBinding.ParameterStoreToDataSet(Inst->GetParameters(), SpawnInstanceParameterDataSet, SystemIndex);
-				UpdateInstanceParameterToDataSetBinding.ParameterStoreToDataSet(Inst->GetParameters(), UpdateInstanceParameterDataSet, SystemIndex);
+				Inst->TickDataInterfaces(DeltaSeconds, false);
+				if (!Inst->IsComplete())
+				{
+					++SystemIndex;
+				}
+				else
+				{
+					checkSlow(Inst->SystemInstanceIndex == INDEX_NONE);
+				}
 			}
+			OrigNum = SystemIndex;
 
-			//TODO: Find good way to check that we're not using any instance parameter data interfaces in the system scripts here.
-			//In that case we need to solo and will never get here.
-
-			// Update the individual parameters that may have been manually set. Note that we check for a valid base in multiple situations
-			// as in certain situations (like spawn) these values may not yet have been set or created.
-			if (SystemExecutionStateAccessor.BaseIsValid())
+			//Pre tick and gather any still valid pending instances for spawn.
+			SystemInstances.Reserve(NewNum);
+			SpawnNum = 0;
+			while (PendingSystemInstances.Num())
 			{
-				SystemExecutionStateAccessor.Set(SystemIndex, (int32)Inst->GetExecutionState());
-			}
+				FNiagaraSystemInstance* Inst = PendingSystemInstances.Pop();
 
-			TArray<TSharedRef<FNiagaraEmitterInstance>>& Emitters = Inst->GetEmitters();
-			for (int32 EmitterIdx = 0; EmitterIdx < Emitters.Num(); ++EmitterIdx)
-			{
-				FNiagaraEmitterInstance& EmitterInst = Emitters[EmitterIdx].Get();
-				if (EmitterExecutionStateAccessors.Num() > EmitterIdx && EmitterExecutionStateAccessors[EmitterIdx].BaseIsValid())
+				Inst->TickDataInterfaces(DeltaSeconds, false);
+				Inst->SetPendingSpawn(false);
+				Inst->SystemInstanceIndex = SystemInstances.Add(Inst);
+				if (!Inst->IsComplete())
+				{
+					++SpawnNum;
+				}
+				else
 				{
-					EmitterExecutionStateAccessors[EmitterIdx].Set(SystemIndex, (int32)EmitterInst.GetExecutionState());
+					checkSlow(Inst->SystemInstanceIndex == INDEX_NONE);
 				}
 			}
-		};
+		}
 
-		SpawnInstanceParameterDataSet.Allocate(NewNum);
-		UpdateInstanceParameterDataSet.Allocate(NewNum);
+		NewNum = OrigNum + SpawnNum;
+		check(NewNum == SystemInstances.Num());
+	}
 
-		if (GbParallelSystemPreTick)
+	if (bCanExecute && NewNum > 0)
+	{
+		if (GbDumpSystemData || System->bDumpDebugSystemInfo)
 		{
-			ParallelFor(SystemInstances.Num(), PreTick);
+			UE_LOG(LogNiagara, Log, TEXT("=========================================================="));
+			UE_LOG(LogNiagara, Log, TEXT("Niagara System Sim Tick: %s"), *System->GetName());
+			UE_LOG(LogNiagara, Log, TEXT("=========================================================="));
 		}
-		else
+
+		TArray<FNiagaraDataSetExecutionInfo, TInlineAllocator<8>> DataSetExecInfos;
 		{
-			//Tick instances and pull out any per instance parameters we need.
-			for (int32 i = 0; i < SystemInstances.Num(); ++i)
+			SCOPE_CYCLE_COUNTER(STAT_NiagaraSystemSim_PrepareForSimulate);
+			auto PreSimulateAndTransferParams = [&](int32 SystemIndex)
 			{
-				PreTick(i);
-			}
-		}
+				FNiagaraSystemInstance* Inst = SystemInstances[SystemIndex];
+				Inst->PreSimulateTick(DeltaSeconds);
 
-		SpawnInstanceParameterDataSet.Tick();
-		UpdateInstanceParameterDataSet.Tick();
-	}
+				if (Inst->GetParameters().GetParametersDirty() && bCanExecute)
+				{
+					SpawnInstanceParameterToDataSetBinding.ParameterStoreToDataSet(Inst->GetParameters(), SpawnInstanceParameterDataSet, SystemIndex);
+					UpdateInstanceParameterToDataSetBinding.ParameterStoreToDataSet(Inst->GetParameters(), UpdateInstanceParameterDataSet, SystemIndex);
+				}
 
-	if (bCanExecute && NewNum > 0)
-	{
-		InitBindings(SystemInstances[0]);
+				//TODO: Find good way to check that we're not using any instance parameter data interfaces in the system scripts here.
+				//In that case we need to solo and will never get here.
+
+				TArray<TSharedRef<FNiagaraEmitterInstance>>& Emitters = Inst->GetEmitters();
+				for (int32 EmitterIdx = 0; EmitterIdx < Emitters.Num(); ++EmitterIdx)
+				{
+					FNiagaraEmitterInstance& EmitterInst = Emitters[EmitterIdx].Get();
+					if (EmitterExecutionStateAccessors.Num() > EmitterIdx && EmitterExecutionStateAccessors[EmitterIdx].BaseIsValid())
+					{
+						EmitterExecutionStateAccessors[EmitterIdx].Set(SystemIndex, (int32)EmitterInst.GetExecutionState());
+					}
+				}
+			};
+
+			InitBindings(SystemInstances[0]);
 
-		DataSet.Tick();
-		DataSet.Allocate(NewNum);
-		
-		//Setup the few real constants like delta time.
-		float InvDt = 1.0f / DeltaSeconds;
-		SpawnExecContext.Parameters.SetParameterValue(DeltaSeconds, SYS_PARAM_ENGINE_DELTA_TIME);
-		SpawnExecContext.Parameters.SetParameterValue(InvDt, SYS_PARAM_ENGINE_INV_DELTA_TIME);
+			SpawnInstanceParameterDataSet.Allocate(NewNum);
+			UpdateInstanceParameterDataSet.Allocate(NewNum);
+
+			for (int32 EmitterIdx = 0; EmitterIdx < System->GetNumEmitters(); ++EmitterIdx)
+			{
+				EmitterExecutionStateAccessors[EmitterIdx].InitForAccess(true);
+			}
 
-		UpdateExecContext.Parameters.SetParameterValue(DeltaSeconds, SYS_PARAM_ENGINE_DELTA_TIME);
-		UpdateExecContext.Parameters.SetParameterValue(InvDt, SYS_PARAM_ENGINE_INV_DELTA_TIME);
+			//Transfer any values like execution state from the system instance into the dataset for simulation.
+			ParallelFor(SystemInstances.Num(), [&](int32 SystemIndex)
+			{
+				PreSimulateAndTransferParams(SystemIndex);
+			});
 
-		TArray<FNiagaraDataSetExecutionInfo> DataSetExecInfos;
-		DataSetExecInfos.Emplace(&DataSet, 0, false, true);
-		DataSetExecInfos.Emplace(&UpdateInstanceParameterDataSet, 0, false, false);
+			SpawnInstanceParameterDataSet.SetNumInstances(NewNum);
+			UpdateInstanceParameterDataSet.SetNumInstances(NewNum);
+			SpawnInstanceParameterDataSet.Tick();
+			UpdateInstanceParameterDataSet.Tick();
+
+			DataSet.Tick();
+			DataSet.Allocate(NewNum);
+
+			//Setup the few real constants like delta time.
+			float InvDt = 1.0f / DeltaSeconds;
+			float GlobalSpawnCountScale = INiagaraModule::GetGlobalSpawnCountScale();
+			float GlobalSystemCountScale = INiagaraModule::GetGlobalSystemCountScale();
+			SpawnDeltaTimeParam.SetValue(DeltaSeconds);
+			UpdateDeltaTimeParam.SetValue(DeltaSeconds);
+			SpawnInvDeltaTimeParam.SetValue(InvDt);
+			UpdateInvDeltaTimeParam.SetValue(InvDt);
+			SpawnNumSystemInstancesParam.SetValue(NewNum);
+			UpdateNumSystemInstancesParam.SetValue(NewNum);
+			SpawnGlobalSpawnCountScaleParam.SetValue(GlobalSpawnCountScale);
+			UpdateGlobalSpawnCountScaleParam.SetValue(GlobalSpawnCountScale);
+			SpawnGlobalSystemCountScaleParam.SetValue(GlobalSystemCountScale);
+			UpdateGlobalSystemCountScaleParam.SetValue(GlobalSystemCountScale);
+
+			DataSetExecInfos.Emplace(&DataSet, 0, false, true);
+		}
 
 		{
 			SCOPE_CYCLE_COUNTER(STAT_NiagaraSystemSim_Update);
@@ -374,23 +391,46 @@ bool FNiagaraSystemSimulation::Tick(float DeltaSeconds)
 			UpdateInstanceParameterDataSet.SetNumInstances(OrigNum);
 
 			//Run update.
-			UpdateExecContext.Tick(nullptr);//We can't require a specfic instance here as these are for all instances.
+			UpdateExecContext.Tick(SystemInstances[0]);
 			DataSetExecInfos[0].StartInstance = 0;
+			DataSetExecInfos.Emplace(&UpdateInstanceParameterDataSet, 0, false, false);
 
-			if (GbDumpSystemData)
-			{
-				UE_LOG(LogNiagara, Log, TEXT("=== PreUpdated %d Systems ==="), OrigNum);
-				DataSet.Dump(true, 0, OrigNum);
-				UpdateInstanceParameterDataSet.Dump(true, 0, OrigNum);
-			}
+// 			if (GbDumpSystemData || System->bDumpDebugSystemInfo)
+// 			{
+// 				UE_LOG(LogNiagara, Log, TEXT("=== PreUpdated %d Systems ==="), OrigNum);
+// 				DataSet.Dump(false, 0, OrigNum);
+// 				UpdateInstanceParameterDataSet.Dump(false, 0, OrigNum);
+// 			}
 
 			UpdateExecContext.Execute(OrigNum, DataSetExecInfos);
 
-			if (GbDumpSystemData)
+			if (GbDumpSystemData || System->bDumpDebugSystemInfo)
 			{
 				UE_LOG(LogNiagara, Log, TEXT("=== Updated %d Systems ==="), OrigNum);
 				DataSet.Dump(true, 0, OrigNum);
-				UpdateInstanceParameterDataSet.Dump(false, 0, OrigNum);
+			}
+
+#if WITH_EDITORONLY_DATA
+			if (SoloSystemInstance && SoloSystemInstance->ShouldCaptureThisFrame())
+			{
+				FNiagaraScriptDebuggerInfo* DebugInfo = SoloSystemInstance->GetActiveCaptureWrite(NAME_None, ENiagaraScriptUsage::SystemUpdateScript, FGuid());
+				if (DebugInfo)
+				{
+					DataSet.Dump(DebugInfo->Frame, true, 0, OrigNum);
+					//DebugInfo->Frame.Dump(true, 0, OrigNum);
+					DebugInfo->Parameters = UpdateExecContext.Parameters;
+				}
+			}
+#endif
+		}
+
+		SystemExecutionStateAccessor.InitForAccess(true);
+		for (int32 EmitterIdx = 0; EmitterIdx < System->GetNumEmitters(); ++EmitterIdx)
+		{
+			EmitterExecutionStateAccessors[EmitterIdx].InitForAccess(true);
+			for (int32 SpawnInfoIdx = 0; SpawnInfoIdx < EmitterSpawnInfoAccessors[EmitterIdx].Num(); ++SpawnInfoIdx)
+			{
+				EmitterSpawnInfoAccessors[EmitterIdx][SpawnInfoIdx].InitForAccess(true);
 			}
 		}
 
@@ -401,39 +441,45 @@ bool FNiagaraSystemSimulation::Tick(float DeltaSeconds)
 			SpawnInstanceParameterDataSet.SetNumInstances(NewNum);
 
 			//Run Spawn
-			SpawnExecContext.Tick(nullptr);//We can't require a specific instance here as these are for all instances.
+			SpawnExecContext.Tick(SoloSystemInstance);//We can't require a specific instance here as these are for all instances.
 			DataSetExecInfos.SetNum(1, false);
 			DataSetExecInfos[0].StartInstance = OrigNum;
 			DataSetExecInfos.Emplace(&SpawnInstanceParameterDataSet, OrigNum, false, false);
 			SpawnExecContext.Execute(SpawnNum, DataSetExecInfos);
 
-			if (GbDumpSystemData)
+			if (GbDumpSystemData || System->bDumpDebugSystemInfo)
 			{
 				UE_LOG(LogNiagara, Log, TEXT("=== Spwaned %d Systems ==="), SpawnNum);
 				DataSet.Dump(true, OrigNum, SpawnNum);
 				SpawnInstanceParameterDataSet.Dump(false, OrigNum, SpawnNum);
 			}
-		}
 
-		{
-			SCOPE_CYCLE_COUNTER(STAT_NiagaraSystemSim_TransferParameters);
-			SystemExecutionStateAccessor.InitForAccess(true);
-			for (int32 EmitterIdx = 0; EmitterIdx < System->GetNumEmitters(); ++EmitterIdx)
+#if WITH_EDITORONLY_DATA
+			if (SoloSystemInstance && SoloSystemInstance->ShouldCaptureThisFrame())
 			{
-				EmitterExecutionStateAccessors[EmitterIdx].InitForAccess(true);
-				for (int32 SpawnInfoIdx = 0; SpawnInfoIdx < EmitterSpawnInfoAccessors[EmitterIdx].Num(); ++SpawnInfoIdx)
+				FNiagaraScriptDebuggerInfo* DebugInfo = SoloSystemInstance->GetActiveCaptureWrite(NAME_None, ENiagaraScriptUsage::SystemSpawnScript, FGuid());
+				if (DebugInfo)
 				{
-					EmitterSpawnInfoAccessors[EmitterIdx][SpawnInfoIdx].InitForAccess(true);
-				}				
+					DataSet.Dump(DebugInfo->Frame, true, OrigNum, SpawnNum);
+					//DebugInfo->Frame.Dump(true, 0, SpawnNum);
+					DebugInfo->Parameters = SpawnExecContext.Parameters;
+				}
 			}
+#endif
+		}
 
+		{
+			SCOPE_CYCLE_COUNTER(STAT_NiagaraSystemSim_TransferParameters);
 			int32 SystemIndex = 0;
 			while (SystemIndex < SystemInstances.Num())
 			{
 				ENiagaraExecutionState ExecutionState = (ENiagaraExecutionState)SystemExecutionStateAccessor.GetSafe(SystemIndex, (int32)ENiagaraExecutionState::Disabled);
 				FNiagaraSystemInstance* SystemInst = SystemInstances[SystemIndex];
-				SystemInst->SetExecutionState(ExecutionState);
-				if (!SystemInst->HandleCompletion())
+
+				//Apply the systems requested execution state to it's actual execution state.
+				SystemInst->SetActualExecutionState(ExecutionState);
+
+				if (!SystemInst->IsDisabled() && !SystemInst->HandleCompletion())
 				{
 					//Now pull data out of the simulation and drive the emitters with it.
 					TArray<TSharedRef<FNiagaraEmitterInstance>>& Emitters = SystemInst->GetEmitters();
@@ -527,18 +573,23 @@ bool FNiagaraSystemSimulation::Tick(float DeltaSeconds)
 		{
 			FNiagaraSystemInstance* SystemInst = SystemInstances[SystemIndex];
 			++SystemIndex;
-			if (SystemInst->HasTickingEmitters())
+			if (SystemIndex < SystemInstances.Num())
 			{
-				if (SystemIndex < SystemInstances.Num())
-				{
-					FPlatformMisc::Prefetch(SystemInstances[SystemIndex]->GetComponent());
-				}
-				SystemInst->GetComponent()->UpdateComponentToWorld();//Needed for bounds updates. Can probably skip if using fixed bounds
-				SystemInst->GetComponent()->MarkRenderDynamicDataDirty();
+				FPlatformMisc::Prefetch(SystemInstances[SystemIndex]->GetComponent());
 			}
+			SystemInst->FinalizeTick(DeltaSeconds);
 		}
 	}
 
+#if WITH_EDITORONLY_DATA
+	if (SoloSystemInstance != nullptr)
+	{
+		SoloSystemInstance->FinishCapture();
+	}
+#endif
+
+	INC_DWORD_STAT_BY(STAT_NiagaraNumSystems, SystemInstances.Num());
+
 	return true;
 }
 
@@ -549,34 +600,8 @@ void FNiagaraSystemSimulation::RemoveInstance(FNiagaraSystemInstance* Instance)
 		return;
 	}
 
-	if (Instance->IsSolo())
-	{
-		if (SoloSystemInstances.IsValidIndex(Instance->SystemInstanceIndex))
-		{
-			if (GbDumpSystemData)
-			{
-				UE_LOG(LogNiagara, Log, TEXT("=== Removing System Solo %d ==="), Instance->SystemInstanceIndex);
-				DataSetSolo.Dump(true, Instance->SystemInstanceIndex, 1);
-			}
-
-			int32 NumInstances = DataSetSolo.GetNumInstances();
-			check(SoloSystemInstances.Num() == NumInstances);
-
-			int32 SystemIndex = Instance->SystemInstanceIndex;
-			check(Instance == SoloSystemInstances[SystemIndex]);
-			check(SoloSystemInstances.IsValidIndex(SystemIndex));
-			DataSetSolo.KillInstance(SystemIndex);
-			SoloSystemInstances.RemoveAtSwap(SystemIndex);
-			Instance->SystemInstanceIndex = INDEX_NONE;
-			DataSetSolo.SetNumInstances(SoloSystemInstances.Num());
-
-			if (SoloSystemInstances.IsValidIndex(SystemIndex))
-			{
-				SoloSystemInstances[SystemIndex]->SystemInstanceIndex = SystemIndex;
-			}
-		}
-	}
-	else if (Instance->IsPendingSpawn())
+	UNiagaraSystem* System = WeakSystem.Get();
+	if (Instance->IsPendingSpawn())
 	{
 		int32 SystemIndex = Instance->SystemInstanceIndex;
 		check(Instance == PendingSystemInstances[SystemIndex]);
@@ -590,7 +615,7 @@ void FNiagaraSystemSimulation::RemoveInstance(FNiagaraSystemInstance* Instance)
 	}
 	else if (SystemInstances.IsValidIndex(Instance->SystemInstanceIndex))
 	{
-		if (GbDumpSystemData)
+		if (GbDumpSystemData || (System && System->bDumpDebugSystemInfo))
 		{
 			UE_LOG(LogNiagara, Log, TEXT("=== Removing System %d ==="), Instance->SystemInstanceIndex);
 			DataSet.Dump(true, Instance->SystemInstanceIndex, 1);
@@ -619,208 +644,13 @@ void FNiagaraSystemSimulation::AddInstance(FNiagaraSystemInstance* Instance)
 	Instance->SystemInstanceIndex = PendingSystemInstances.Add(Instance);
 }
 
-void FNiagaraSystemSimulation::ResetSolo(FNiagaraSystemInstance* Instance)
-{
-	Instance->SetPendingSpawn(true);
-	Instance->SystemInstanceIndex = INDEX_NONE;
-}
-
-const TArray<UNiagaraDataInterface*>& FNiagaraSystemSimulation::GetSoloDataInterfacesSpawn()
-{
-	return SpawnExecContextSolo.GetDataInterfaces();
-}
-
-const TArray<UNiagaraDataInterface*>& FNiagaraSystemSimulation::GetSoloDataInterfacesUpdate()
-{
-	return UpdateExecContextSolo.GetDataInterfaces();
-}
-
-void FNiagaraSystemSimulation::TickSolo(FNiagaraSystemInstance* SystemInst)
-{
-	SCOPE_CYCLE_COUNTER(STAT_NiagaraSystemSimSolo);
-
-	UNiagaraSystem* System = WeakSystem.Get();
-	check(System != nullptr);
-
-	if (!System->GetSystemSpawnScript(true)->IsValid() || !System->GetSystemUpdateScript(true)->IsValid())
-	{
-		return;
-	}
-
-	auto SetExecStateFromExternal = [&](int32 SystemIndex, bool bCurr)
-	{
-		SoloSystemExecutionStateAccessor.InitForAccess(bCurr);
-		for (int32 EmitterIdx = 0; EmitterIdx < System->GetNumEmitters(); ++EmitterIdx)
-		{
-			SoloEmitterExecutionStateAccessors[EmitterIdx].InitForAccess(bCurr);
-		}
-		
-		// Update the individual parameters that may have been manually set.
-		if (SoloSystemExecutionStateAccessor.BaseIsValid())
-		{
-			SoloSystemExecutionStateAccessor.Set(SystemIndex, (int32)SystemInst->GetExecutionState());
-		}
-		TArray<TSharedRef<FNiagaraEmitterInstance>>& Emitters = SystemInst->GetEmitters();
-		for (int32 EmitterIdx = 0; EmitterIdx < Emitters.Num(); ++EmitterIdx)
-		{
-			FNiagaraEmitterInstance& EmitterInst = Emitters[EmitterIdx].Get();
-			if (SoloEmitterExecutionStateAccessors.Num() > EmitterIdx && SoloEmitterExecutionStateAccessors[EmitterIdx].BaseIsValid())
-			{
-				SoloEmitterExecutionStateAccessors[EmitterIdx].Set(SystemIndex, (int32)EmitterInst.GetExecutionState());
-			}
-		}
-	};
-
-	FNiagaraWorldManager* WorldMan = FNiagaraWorldManager::Get(World);
-	check(WorldMan);
-
-	int32 SystemIndex = SystemInst->SystemInstanceIndex;
-
-	TArray<FNiagaraDataSetExecutionInfo> DataSetExecInfos;
-	DataSetExecInfos.Emplace(&DataSetSolo, SystemIndex, false, true);
-
-	if (SystemInst->SystemInstanceIndex == INDEX_NONE || SystemInst->IsPendingSpawn())
-	{
-		SystemInst->SetPendingSpawn(false);
-
-		//Spawn/respawn this system this frame rather than updating.
-		if (SystemIndex == INDEX_NONE)
-		{
-			SystemIndex = SoloSystemInstances.Num();
-			SystemInst->SystemInstanceIndex = SystemIndex;
-			verify(SoloSystemInstances.Add(SystemInst) == SystemIndex);
-			DataSetSolo.Allocate(SoloSystemInstances.Num(), ENiagaraSimTarget::CPUSim, true);
-			DataSetSolo.SetNumInstances(SoloSystemInstances.Num());
-			DataSetExecInfos[0].StartInstance = SystemIndex;
-		}
-
-		check(SoloSystemInstances.IsValidIndex(SystemIndex));
-
-		SetExecStateFromExternal(SystemIndex, true);
-
-		//////////////////////////////////////////////////////////////////////////
-		//TODO: This can be optimized as all the offests are the same. Can skip all the searching involved with a full bind().
-		SystemInst->GetParameters().Bind(&SpawnExecContextSolo.Parameters);
-
-		//We must tick the RI parameters here to be sure we have up to date data. NPC data can afford to be a frame delayed but not RI data.
-		SpawnExecContextSolo.Script->RapidIterationParameters.Tick();
-
-		SpawnExecContextSolo.Tick(SystemInst);
-		//////////////////////////////////////////////////////////////////////////
-		
-		SpawnExecContextSolo.Execute(1, DataSetExecInfos);
-
-		//////////////////////////////////////////////////////////////////////////
-		//TODO: This can be optimized as all the offests are the same. Can skip all the searching involved with a full bind().
-		SystemInst->GetParameters().Unbind(&SpawnExecContextSolo.Parameters);
-		//////////////////////////////////////////////////////////////////////////
-		
-		DataSetSolo.SetNumInstances(SoloSystemInstances.Num());
-
-		if (GbDumpSystemData)
-		{
-			UE_LOG(LogNiagara, Log, TEXT("=== Reset/spawn System Solo %d ==="), SystemIndex);
-			DataSetSolo.Dump(true, SystemIndex, 1);
-		}
-	}
-	else
-	{
-		check(SoloSystemInstances.IsValidIndex(SystemIndex));
-
-		SetExecStateFromExternal(SystemIndex, true);
-
-		//////////////////////////////////////////////////////////////////////////
-		//TODO: This can be optimized as all the offests are the same. Can skip all the searching involved with a full bind().
-		SystemInst->GetParameters().Bind(&UpdateExecContextSolo.Parameters);
-		//We must tick the RI parameters here to be sure we have up to date data. NPC data can afford to be a frame delayed but not RI data.
-		UpdateExecContextSolo.Script->RapidIterationParameters.Tick();
-		UpdateExecContextSolo.Tick(SystemInst);
-		//////////////////////////////////////////////////////////////////////////
-
-		UpdateExecContextSolo.Execute(1, DataSetExecInfos);
-
-		//////////////////////////////////////////////////////////////////////////
-		//TODO: This can be optimized as all the offests are the same. Can skip all the searching involved with a full bind().
-		SystemInst->GetParameters().Unbind(&UpdateExecContextSolo.Parameters);
-		//////////////////////////////////////////////////////////////////////////
-
-		DataSetSolo.SetNumInstances(SoloSystemInstances.Num());
-
-		if (GbDumpSystemData)
-		{
-			UE_LOG(LogNiagara, Log, TEXT("=== Updated System Solo %d ==="), SystemIndex);
-			DataSetSolo.Dump(true, SystemIndex, 1);
-		}
-	}
-
-	InitBindings(SystemInst);
-
-	SoloSystemExecutionStateAccessor.InitForAccess(true);
-	for (int32 EmitterIdx = 0; EmitterIdx < System->GetNumEmitters(); ++EmitterIdx)
-	{
-		SoloEmitterExecutionStateAccessors[EmitterIdx].InitForAccess(true);
-		for (int32 SpawnInfoIdx = 0; SpawnInfoIdx < SoloEmitterSpawnInfoAccessors[EmitterIdx].Num(); ++SpawnInfoIdx)
-		{
-			SoloEmitterSpawnInfoAccessors[EmitterIdx][SpawnInfoIdx].InitForAccess(true);
-		}
-	}		
-
-	ensure(SystemIndex != INDEX_NONE);
-	ENiagaraExecutionState SoloExecutionState = (ENiagaraExecutionState)SoloSystemExecutionStateAccessor.GetSafe(SystemIndex, (int32)ENiagaraExecutionState::Disabled);
-	SystemInst->SetExecutionState(SoloExecutionState);
-
-	//Now pull data out of the simulation and drive the emitters with it.
-	TArray<TSharedRef<FNiagaraEmitterInstance>>& Emitters = SystemInst->GetEmitters();
-	for (int32 EmitterIdx = 0; EmitterIdx < Emitters.Num(); ++EmitterIdx)
-	{
-		FNiagaraEmitterInstance& EmitterInst = Emitters[EmitterIdx].Get();
-
-		//Early exit before we set the state as if we're complete or disabled we should never let the emitter turn itself back. It needs to be reset/reinited manually.
-		if (EmitterInst.IsComplete())
-		{
-			continue;
-		}
-
-		ENiagaraExecutionState State = (ENiagaraExecutionState)SoloEmitterExecutionStateAccessors[EmitterIdx].GetSafe(SystemIndex, (int32)ENiagaraExecutionState::Disabled);
-		EmitterInst.SetExecutionState(State);
-
-		for (int32 SpawnInfoIdx = 0; SpawnInfoIdx < SoloEmitterSpawnInfoAccessors[EmitterIdx].Num(); ++SpawnInfoIdx)
-		{
-			EmitterInst.GetSpawnInfo()[SpawnInfoIdx] = SoloEmitterSpawnInfoAccessors[EmitterIdx][SpawnInfoIdx].Get(SystemIndex);
-		}
-
-		//TODO: Any other fixed function stuff like this?
-
-		FNiagaraScriptExecutionContext& SpawnContext = EmitterInst.GetSpawnExecutionContext();
-		DataSetToEmitterSpawnParameters[EmitterIdx].DataSetToParameterStore(SpawnContext.Parameters, DataSetSolo, SystemIndex);
-
-		FNiagaraScriptExecutionContext& UpdateContext = EmitterInst.GetUpdateExecutionContext();
-		DataSetToEmitterUpdateParameters[EmitterIdx].DataSetToParameterStore(UpdateContext.Parameters, DataSetSolo, SystemIndex);
-
-		TArray<FNiagaraScriptExecutionContext>& EventContexts = EmitterInst.GetEventExecutionContexts();
-		for (int32 EventIdx = 0; EventIdx < EventContexts.Num(); ++EventIdx)
-		{
-			FNiagaraScriptExecutionContext& EventContext = EventContexts[EventIdx];
-			if (DataSetToEmitterEventParameters[EmitterIdx].Num() > EventIdx)
-			{
-				DataSetToEmitterEventParameters[EmitterIdx][EventIdx].DataSetToParameterStore(EventContext.Parameters, DataSetSolo, SystemIndex);
-			}
-			else
-			{
-				UE_LOG(LogNiagara, Log, TEXT("Skipping DataSetToEmitterEventParameters because EventIdx is out-of-bounds. %d of %d"), EventIdx, DataSetToEmitterEventParameters[EmitterIdx].Num());
-			}
-
-		}
-	}
-}
-
 void FNiagaraSystemSimulation::InitBindings(FNiagaraSystemInstance* SystemInst)
 {
 	//Have to init here as we need an actual parameter store to pull the layout info from.
 	//TODO: Pull the layout stuff out of each data set and store. So much duplicated data.
 	//This assumes that all layouts for all emitters is the same. Which it should be.
 	//Ideally we can store all this layout info in the systm/emitter assets so we can just generate this in Init()
-	if (DataSetToEmitterSpawnParameters.Num() == 0)
+	if (DataSetToEmitterSpawnParameters.Num() == 0 && SystemInst != nullptr)
 	{
 		SpawnInstanceParameterToDataSetBinding.Init(SpawnInstanceParameterDataSet, SystemInst->GetInstanceParameters());
 		UpdateInstanceParameterToDataSetBinding.Init(UpdateInstanceParameterDataSet, SystemInst->GetInstanceParameters());
@@ -849,4 +679,13 @@ void FNiagaraSystemSimulation::InitBindings(FNiagaraSystemInstance* SystemInst)
 			}
 		}
 	}
+
+	// If we have data interfaces (say user ones) that need to be pushed to the 
+	// system spawn and update scripts, it needs to happen here.
+	if (SystemInst)
+	{
+		SystemInst->GetInstanceParameters().Bind(&SpawnExecContext.Parameters);
+		SystemInst->GetInstanceParameters().Bind(&UpdateExecContext.Parameters);
+
+	}
 }
diff --git a/Engine/Plugins/FX/Niagara/Source/Niagara/Private/NiagaraWorldManager.cpp b/Engine/Plugins/FX/Niagara/Source/Niagara/Private/NiagaraWorldManager.cpp
--- a/Engine/Plugins/FX/Niagara/Source/Niagara/Private/NiagaraWorldManager.cpp
+++ b/Engine/Plugins/FX/Niagara/Source/Niagara/Private/NiagaraWorldManager.cpp
@@ -11,12 +11,43 @@
 #include "Scalability.h"
 #include "Misc/ConfigCacheIni.h"
 #include "NiagaraDataInterfaceSkeletalMesh.h"
+#include "EngineModule.h"
+
+FNiagaraViewDataMgr GNiagaraViewDataManager;
+
+FNiagaraViewDataMgr::FNiagaraViewDataMgr()
+	: SceneDepthTexture(nullptr)
+	, SceneNormalTexture(nullptr)
+	, ViewUniformBuffer(nullptr)
+{
+
+}
+
+void FNiagaraViewDataMgr::Init()
+{
+	IRendererModule& RendererModule = GetRendererModule();
+	GNiagaraViewDataManager.PostOpaqueDelegate.BindRaw(&GNiagaraViewDataManager, &FNiagaraViewDataMgr::PostOpaqueRender);
+	RendererModule.RegisterPostOpaqueRenderDelegate(GNiagaraViewDataManager.PostOpaqueDelegate);
+}
+
+void FNiagaraViewDataMgr::Shutdown()
+{
+
+}
+
+
+FNiagaraWorldManager::FNiagaraWorldManager(UWorld* InWorld)
+	: World(InWorld)
+	, CachedEffectsQuality(INDEX_NONE)
+{
+}
+
 
 
 FNiagaraWorldManager* FNiagaraWorldManager::Get(UWorld* World)
 {
-	INiagaraModule& NiagaraModule = FModuleManager::LoadModuleChecked<INiagaraModule>("Niagara");
-	return NiagaraModule.GetWorldManager(World);
+	//INiagaraModule& NiagaraModule = FModuleManager::LoadModuleChecked<INiagaraModule>("Niagara");
+	return INiagaraModule::GetWorldManager(World);
 }
 
 void FNiagaraWorldManager::AddReferencedObjects(FReferenceCollector& Collector)
diff --git a/Engine/Plugins/FX/Niagara/Source/Niagara/Public/MovieScene/MovieSceneNiagaraSystemTrack.h b/Engine/Plugins/FX/Niagara/Source/Niagara/Public/MovieScene/MovieSceneNiagaraSystemTrack.h
new file mode 100644
--- /dev/null
+++ b/Engine/Plugins/FX/Niagara/Source/Niagara/Public/MovieScene/MovieSceneNiagaraSystemTrack.h
@@ -0,0 +1,17 @@
+// Copyright 1998-2017 Epic Games, Inc. All Rights Reserved.
+
+#pragma once
+#include "MovieScene/MovieSceneNiagaraTrack.h"
+#include "MovieSceneNiagaraSystemTrack.generated.h"
+
+UCLASS(MinimalAPI)
+class UMovieSceneNiagaraSystemTrack : public UMovieSceneNiagaraTrack
+{
+	GENERATED_BODY()
+
+public:
+	virtual UMovieSceneSection* CreateNewSection() override;
+	virtual FMovieSceneEvalTemplatePtr CreateTemplateForSection(const UMovieSceneSection& InSection) const override;
+	virtual void PostCompile(FMovieSceneEvaluationTrack& OutTrack, const FMovieSceneTrackCompilerArgs& Args) const override;
+	virtual FMovieSceneTrackSegmentBlenderPtr GetTrackSegmentBlender() const override;
+};
\ No newline at end of file
diff --git a/Engine/Plugins/FX/Niagara/Source/Niagara/Public/NiagaraCommon.h b/Engine/Plugins/FX/Niagara/Source/Niagara/Public/NiagaraCommon.h
--- a/Engine/Plugins/FX/Niagara/Source/Niagara/Public/NiagaraCommon.h
+++ b/Engine/Plugins/FX/Niagara/Source/Niagara/Public/NiagaraCommon.h
@@ -5,6 +5,7 @@
 #include "CoreMinimal.h"
 #include "UObject/ObjectMacros.h"
 #include "NiagaraTypes.h"
+#include "UObject/SoftObjectPath.h"
 #include "NiagaraCommon.generated.h"
 
 class UNiagaraSystem;
@@ -13,9 +14,12 @@ class UNiagaraDataInterface;
 class UNiagaraEmitter;
 class FNiagaraSystemInstance;
 class UNiagaraParameterCollection;
+struct FNiagaraParameterStore;
 
 const uint32 NIAGARA_COMPUTE_THREADGROUP_SIZE = 64;
-const uint32 NIAGARA_MAX_COMPUTE_THREADGROUPS = 16384;
+const uint32 NIAGARA_MAX_COMPUTE_THREADGROUPS = 65536;
+
+const FString INTERPOLATED_PARAMETER_PREFIX = TEXT("PREV_");
 
 enum ENiagaraBaseTypes
 {
@@ -34,6 +38,16 @@ enum class ENiagaraSimTarget : uint8
 };
 
 
+/** Defines modes for updating the component's age. */
+UENUM()
+enum class ENiagaraAgeUpdateMode : uint8
+{
+	/** Update the age using the delta time supplied to the tick function. */
+	TickDeltaTime,
+	/** Update the age by seeking to the DesiredAge. To prevent major perf loss, we clamp to MaxClampTime*/
+	DesiredAge
+};
+
 
 UENUM()
 enum class ENiagaraDataSetType : uint8
@@ -50,7 +64,8 @@ enum class ENiagaraInputNodeUsage : uint8
 	Undefined = 0,
 	Parameter,
 	Attribute,
-	SystemConstant
+	SystemConstant,
+	TranslatorConstant,
 };
 
 /**
@@ -163,12 +178,12 @@ struct FNiagaraScriptDataUsageInfo
 	GENERATED_BODY()
 
 		FNiagaraScriptDataUsageInfo()
-		: bReadsAttriubteData(false)
+		: bReadsAttributeData(false)
 	{}
 
 	/** If true, this script reads attribute data. */
 	UPROPERTY()
-		bool bReadsAttriubteData;
+		bool bReadsAttributeData;
 };
 
 
@@ -257,7 +272,7 @@ public:
 
 	}
 
-	UPROPERTY(Instanced)
+	UPROPERTY()
 	class UNiagaraDataInterface* DataInterface;
 	
 	UPROPERTY()
@@ -270,12 +285,61 @@ public:
 	UPROPERTY()
 	FNiagaraTypeDefinition Type;
 
-	TArray<FNiagaraFunctionSignature> RegisteredFunctions;
+	UPROPERTY()
+	FName RegisteredParameterMapRead;
+
+	UPROPERTY()
+	FName RegisteredParameterMapWrite;
 
 	//TODO: Allow data interfaces to own datasets
 	void CopyTo(FNiagaraScriptDataInterfaceInfo* Destination, UObject* Outer) const;
 };
 
+USTRUCT()
+struct NIAGARA_API FNiagaraScriptDataInterfaceCompileInfo
+{
+	GENERATED_USTRUCT_BODY()
+public:
+	FNiagaraScriptDataInterfaceCompileInfo()
+		: Name(NAME_None)
+		, UserPtrIdx(INDEX_NONE)
+		, bIsPlaceholder(false)
+	{
+
+	}
+
+	UPROPERTY()
+	FName Name;
+
+	/** Index of the user pointer for this data interface. */
+	UPROPERTY()
+	int32 UserPtrIdx;
+
+	UPROPERTY()
+	FNiagaraTypeDefinition Type;
+
+	UPROPERTY()
+	TArray<FNiagaraFunctionSignature> RegisteredFunctions;
+
+	UPROPERTY()
+	FName RegisteredParameterMapRead;
+
+	UPROPERTY()
+	FName RegisteredParameterMapWrite;
+
+	UPROPERTY()
+	bool bIsPlaceholder;
+
+	/** Would this data interface work on the target execution type? Only call this on the game thread.*/
+	bool CanExecuteOnTarget(ENiagaraSimTarget SimTarget) const;
+
+	/** Would this data interface force a system script to be solo? Only call this on the game thread.*/
+	bool IsSystemSolo() const;
+
+	/** Note that this is the CDO for this type of data interface, as we often cannot guarantee that the same instance of the data interface we compiled with is the one the user ultimately executes.  Only call this on the game thread.*/
+	UNiagaraDataInterface* GetDefaultDataInterface() const;
+};
+
 USTRUCT()
 struct FNiagaraStatScope
 {
@@ -337,9 +401,9 @@ struct NIAGARA_API FNiagaraSystemUpdateContext
 
 	/** Adds all currently active systems.*/
 	void AddAll(bool bReInit);
+
 private:
 	void AddInternal(class UNiagaraComponent* Comp, bool bReInit);
-
 	FNiagaraSystemUpdateContext(FNiagaraSystemUpdateContext& Other) { }
 
 	TArray<UNiagaraComponent*> ComponentsToReset;
@@ -352,6 +416,8 @@ private:
 
 
 
+
+
 /** Defines different usages for a niagara script. */
 UENUM()
 enum class ENiagaraScriptUsage : uint8
@@ -370,6 +436,8 @@ enum class ENiagaraScriptUsage : uint8
 	ParticleUpdateScript UMETA(Hidden),
 	/** The script is called to update particles in response to an event. */
 	ParticleEventScript UMETA(Hidden),
+	/** The script is called to update particles on the GPU. */
+	ParticleGPUComputeScript UMETA(Hidden),
 	/** The script is called once when the emitter spawns. */
 	EmitterSpawnScript UMETA(Hidden),
 	/** The script is called every frame to tick the emitter. */
@@ -380,6 +448,15 @@ enum class ENiagaraScriptUsage : uint8
 	SystemUpdateScript UMETA(Hidden),
 };
 
+UENUM()
+enum class ENiagaraScriptGroup : uint8
+{
+	Particle = 0,
+	Emitter,
+	System,
+	Max
+};
+
 
 /** Defines all you need to know about a variable.*/
 USTRUCT()
@@ -399,9 +476,40 @@ struct FNiagaraVariableInfo
 	UNiagaraDataInterface* DataInterface;
 };
 
+USTRUCT()
+struct FNiagaraVariableAttributeBinding
+{
+	GENERATED_USTRUCT_BODY();
+
+	FNiagaraVariableAttributeBinding() {}
+	FNiagaraVariableAttributeBinding(const FNiagaraVariable& InVar, const FNiagaraVariable& InAttrVar) : BoundVariable(InVar), DataSetVariable(InAttrVar), DefaultValueIfNonExistent(InAttrVar)
+	{
+		check(InVar.GetType() == InAttrVar.GetType());
+	}
+	FNiagaraVariableAttributeBinding(const FNiagaraVariable& InVar, const FNiagaraVariable& InAttrVar, const FNiagaraVariable& InNonExistentValue) : BoundVariable(InVar), DataSetVariable(InAttrVar), DefaultValueIfNonExistent(InNonExistentValue)
+	{
+		check(InVar.GetType() == InAttrVar.GetType() && InNonExistentValue.GetType() == InAttrVar.GetType());
+	}
+
+
+	UPROPERTY()
+	FNiagaraVariable BoundVariable;
+
+	UPROPERTY()
+	FNiagaraVariable DataSetVariable;
+
+	UPROPERTY()
+	FNiagaraVariable DefaultValueIfNonExistent;
+};
+
+
 namespace FNiagaraUtilities
 {
 	/** Builds a unique name from a candidate name and a set of existing names.  The candidate name will be made unique
 	if necessary by adding a 3 digit index to the end. */
 	FName NIAGARA_API GetUniqueName(FName CandidateName, const TSet<FName>& ExistingNames);
+
+	FNiagaraVariable NIAGARA_API ConvertVariableToRapidIterationConstantName(FNiagaraVariable InVar, const TCHAR* InEmitterName, ENiagaraScriptUsage InUsage);
+
+	void CollectScriptDataInterfaceParameters(const UObject& Owner, const TArray<UNiagaraScript*>& Scripts, FNiagaraParameterStore& OutDataInterfaceParameters);
 };
\ No newline at end of file
diff --git a/Engine/Plugins/FX/Niagara/Source/Niagara/Public/NiagaraComponent.h b/Engine/Plugins/FX/Niagara/Source/Niagara/Public/NiagaraComponent.h
--- a/Engine/Plugins/FX/Niagara/Source/Niagara/Public/NiagaraComponent.h
+++ b/Engine/Plugins/FX/Niagara/Source/Niagara/Public/NiagaraComponent.h
@@ -21,7 +21,7 @@ class UNiagaraParameterCollectionInstance;
 class FNiagaraSystemSimulation;
 
 // Called when the particle system is done
-DECLARE_DYNAMIC_MULTICAST_DELEGATE_OneParam(FOnNiagaraSystemFinished, class UNiagaraComponent*, PSystem);
+DECLARE_MULTICAST_DELEGATE_OneParam(FOnNiagaraSystemFinished, class UNiagaraComponent*);
 
 /**
 * UNiagaraComponent is the primitive component for a Niagara System.
@@ -32,18 +32,10 @@ UCLASS(ClassGroup = (Rendering, Common), hidecategories = Object, hidecategories
 class NIAGARA_API UNiagaraComponent : public UPrimitiveComponent
 {
 	GENERATED_UCLASS_BODY()
-		
-	/** Defines modes for updating the component's age. */
-	enum class EAgeUpdateMode
-	{
-		/** Update the age using the delta time supplied to the tick function. */
-		TickDeltaTime,
-		/** Update the age by seeking to the DesiredAge. */
-		DesiredAge
-	};
 
 #if WITH_EDITORONLY_DATA
 	DECLARE_MULTICAST_DELEGATE(FOnSystemInstanceChanged);
+	DECLARE_MULTICAST_DELEGATE(FOnSynchronizedWithAssetParameters);
 #endif
 
 private:
@@ -62,6 +54,8 @@ private:
 	TMap<FName, bool> EditorOverridesValue;
 
 	FOnSystemInstanceChanged OnSystemInstanceChangedDelegate;
+
+	FOnSynchronizedWithAssetParameters OnSynchronizedWithAssetParametersDelegate;
 #endif
 
 	/**
@@ -73,14 +67,23 @@ private:
 	TUniquePtr<FNiagaraSystemInstance> SystemInstance;
 
 	/** Defines the mode use when updating the System age. */
-	EAgeUpdateMode AgeUpdateMode;
+	ENiagaraAgeUpdateMode AgeUpdateMode;
 	
 	/** The desired age of the System instance.  This is only relevant when using the DesiredAge age update mode. */
 	float DesiredAge;
 
+	/** Whether or not the component can render while seeking to the desired age. */
+	bool bCanRenderWhileSeeking;
+
 	/** The delta time used when seeking to the desired age.  This is only relevant when using the DesiredAge age update mode. */
 	float SeekDelta;
 
+	/** The maximum amount of time in seconds to spend seeking to the desired age in a single frame. */
+	float MaxSimTime;
+
+	/** Whether or not the component is currently seeking to the desired time. */
+	bool bIsSeeking;
+
 	UPROPERTY()
 	uint32 bAutoDestroy : 1;
 
@@ -93,6 +96,7 @@ protected:
 	virtual void OnUnregister() override;
 	virtual void SendRenderDynamicData_Concurrent() override;
 	virtual void BeginDestroy() override;
+	//virtual void OnAttachmentChanged() override;
 public:
 	/**
 	* True if we should automatically attach to AutoAttachParent when activated, and detach from our parent when completed.
@@ -107,6 +111,7 @@ public:
 
 	virtual void Activate(bool bReset = false)override;
 	virtual void Deactivate()override;
+	void DeactivateImmediate();
 
 	virtual void TickComponent(float DeltaTime, enum ELevelTick TickType, FActorComponentTickFunction* ThisTickFunction) override;
 	virtual const UObject* AdditionalStatObject() const override;
@@ -122,36 +127,69 @@ public:
 
 	TSharedPtr<FNiagaraSystemSimulation, ESPMode::ThreadSafe> GetSystemSimulation();
 
+	bool InitializeSystem();
 	void OnSystemComplete();
 	void DestroyInstance();
 
-		UFUNCTION(BlueprintCallable, Category = Niagara, meta = (DisplayName = "Set Niagara System Asset"))
+	UFUNCTION(BlueprintCallable, Category = Niagara, meta = (DisplayName = "Set Niagara System Asset"))
 	void SetAsset(UNiagaraSystem* InAsset);
+
+	UFUNCTION(BlueprintCallable, Category = Niagara, meta = (DisplayName = "Get Niagara System Asset"))
 	UNiagaraSystem* GetAsset() const { return Asset; }
 
-	void SetForceSolo(bool bInForceSolo) { bForceSolo = bInForceSolo; }
+	UFUNCTION(BlueprintCallable, Category = Niagara, meta = (DisplayName = "Set Forced Solo Mode"))
+	void SetForceSolo(bool bInForceSolo);
+
+	UFUNCTION(BlueprintCallable, Category = Niagara, meta = (DisplayName = "Is In Forced Solo Mode"))
 	bool GetForceSolo()const { return bForceSolo; }
 
-	EAgeUpdateMode GetAgeUpdateMode() const;
+	UFUNCTION(BlueprintCallable, Category = Niagara, meta = (DisplayName = "Get Age Update Mode"))
+	ENiagaraAgeUpdateMode GetAgeUpdateMode() const;
 
 	/** Sets the age update mode for the System instance. */
-	void SetAgeUpdateMode(EAgeUpdateMode InAgeUpdateMode);
+	UFUNCTION(BlueprintCallable, Category = Niagara, meta = (DisplayName = "Set Age Update Mode"))
+	void SetAgeUpdateMode(ENiagaraAgeUpdateMode InAgeUpdateMode);
 
 	/** Gets the desired age of the System instance.  This is only relevant when using the DesiredAge age update mode. */
+	UFUNCTION(BlueprintCallable, Category = Niagara, meta = (DisplayName = "Get Desired Age"))
 	float GetDesiredAge() const;
 
 	/** Sets the desired age of the System instance.  This is only relevant when using the DesiredAge age update mode. */
+	UFUNCTION(BlueprintCallable, Category = Niagara, meta = (DisplayName = "Set Desired Age"))
 	void SetDesiredAge(float InDesiredAge);
 
+	/** Sets the desired age of the System instance and designates that this change is a seek.  When seeking to a desired age the
+	    The component can optionally prevent rendering. */
+	UFUNCTION(BlueprintCallable, Category = Niagara, meta = (DisplayName = "Seek to Desired Age"))
+	void SeekToDesiredAge(float InDesiredAge);
+
+	/** Sets whether or not the system can render while seeking. */
+	UFUNCTION(BlueprintCallable, Category = Niagara, meta = (DisplayName = "Set Can Render While Seeking"))
+	void SetCanRenderWhileSeeking(bool bInCanRenderWhileSeeking);
+
 	/** Gets the delta value which is used when seeking from the current age, to the desired age.  This is only relevant
 	when using the DesiredAge age update mode. */
+	UFUNCTION(BlueprintCallable, Category = Niagara, meta = (DisplayName = "Get Desired Age Seek Delta"))
 	float GetSeekDelta() const;
 
 	/** Sets the delta value which is used when seeking from the current age, to the desired age.  This is only relevant
 	when using the DesiredAge age update mode. */
+	UFUNCTION(BlueprintCallable, Category = Niagara, meta = (DisplayName = "Set Desired Age Seek Delta"))
 	void SetSeekDelta(float InSeekDelta);
 
+	/** Sets the maximum time that you can jump within a tick which is used when seeking from the current age, to the desired age.  This is only relevant
+	when using the DesiredAge age update mode. */
+	UFUNCTION(BlueprintCallable, Category = Niagara, meta = (DisplayName = "Get Max Desired Age Tick Delta"))
+	float GetMaxSimTime() const;
+
+	/** Sets the maximum time that you can jump within a tick which is used when seeking from the current age, to the desired age.  This is only relevant
+	when using the DesiredAge age update mode. */
+	UFUNCTION(BlueprintCallable, Category = Niagara, meta = (DisplayName = "Set Max Desired Age Tick Delta"))
+	void SetMaxSimTime(float InMaxTime);
+
+	UFUNCTION(BlueprintCallable, Category = Niagara, meta = (DisplayName = "Set Auto Destroy"))
 	void SetAutoDestroy(bool bInAutoDestroy) { bAutoDestroy = bInAutoDestroy; }
+
 	FNiagaraSystemInstance* GetSystemInstance() const;
 
 	/** Returns true if this component forces it's instances to run in "Solo" mode. A sub optimal path required in some situations. */
@@ -161,10 +199,14 @@ public:
 	UFUNCTION(BlueprintCallable, Category = Niagara, meta = (DisplayName = "Set Niagara Variable (LinearColor)"))
 	void SetNiagaraVariableLinearColor(const FString& InVariableName, const FLinearColor& InValue);
 
-	/** Sets a Niagara Vector3 parameter by name, overriding locally if necessary.*/
+	/** Sets a Niagara Vector4 parameter by name, overriding locally if necessary.*/
 	UFUNCTION(BlueprintCallable, Category = Niagara, meta = (DisplayName = "Set Niagara Variable (Vector4)"))
 	void SetNiagaraVariableVec4(const FString& InVariableName, const FVector4& InValue);
 
+	/** Sets a Niagara Vector3 parameter by name, overriding locally if necessary.*/
+	UFUNCTION(BlueprintCallable, Category = Niagara, meta = (DisplayName = "Set Niagara Variable (Quaternion)"))
+	void SetNiagaraVariableQuat(const FString& InVariableName, const FQuat& InValue);
+
 	/** Sets a Niagara Vector3 parameter by name, overriding locally if necessary.*/
 	UFUNCTION(BlueprintCallable, Category = Niagara, meta = (DisplayName = "Set Niagara Variable (Vector3)"))
 	void SetNiagaraVariableVec3(const FString& InVariableName, FVector InValue);
@@ -211,19 +253,27 @@ public:
 	UFUNCTION(BlueprintCallable, Category = Niagara, meta = (DisplayName = "Set Rendering Enabled"))
 	void SetRenderingEnabled(bool bInRenderingEnabled);
 
+	/** Advances this system's simulation by the specified number of ticks and delta time. */
+	UFUNCTION(BlueprintCallable, Category = Niagara)
+	void AdvanceSimulation(int32 TickCount, float TickDeltaSeconds);
+
+	/** Advances this system's simulation by the specified time in seconds and delta time. Advancement is done in whole ticks of TickDeltaSeconds so actual simulated time will be the nearest lower multiple of TickDeltaSeconds. */
+	UFUNCTION(BlueprintCallable, Category = Niagara)
+	void AdvanceSimulationByTime(float SimulateTime, float TickDeltaSeconds);
+
 	//~ Begin UObject Interface.
 	virtual void PostLoad();
 #if WITH_EDITOR
-	
-	virtual void PostEditChangeProperty(FPropertyChangedEvent& PropertyChangedEvent) override;
+	virtual void PreEditChange(UProperty* PropertyAboutToChange) override;
 
-	/** Compare local overrides with the source System. Remove any that have mismatched types or no longer exist on the System. Returns whether or not any changes occurred.*/
-	virtual bool SynchronizeWithSourceSystem();
+	virtual void PostEditChangeProperty(FPropertyChangedEvent& PropertyChangedEvent) override;
 
 	bool IsParameterValueOverriddenLocally(const FName& InParamName);
 	void SetParameterValueOverriddenLocally(const FNiagaraVariable& InParam, bool bInOverridden);
 	
 	FOnSystemInstanceChanged& OnSystemInstanceChanged() { return OnSystemInstanceChangedDelegate; }
+
+	FOnSynchronizedWithAssetParameters& OnSynchronizedWithAssetParameters() { return OnSynchronizedWithAssetParametersDelegate; }
 #endif
 
 	FNiagaraParameterStore& GetOverrideParameters() { return OverrideParameters; }
@@ -231,9 +281,14 @@ public:
 	//~ End UObject Interface.
 
 	// Called when the particle system is done
-	UPROPERTY(BlueprintAssignable)
 	FOnNiagaraSystemFinished OnSystemFinished;
 
+private:
+	/** Compare local overrides with the source System. Remove any that have mismatched types or no longer exist on the System. Returns whether or not any changes occurred.*/
+	void SynchronizeWithSourceSystem();
+
+	void AssetExposedParametersChanged();
+
 public:
 	/**
 	 * Component we automatically attach to when activated, if bAutoManageAttachment is true.
@@ -284,11 +339,23 @@ public:
 	UFUNCTION(BlueprintCallable, Category = Niagara)
 	void SetAutoAttachmentParameters(USceneComponent* Parent, FName SocketName, EAttachmentRule LocationRule, EAttachmentRule RotationRule, EAttachmentRule ScaleRule);
 
+#if WITH_EDITORONLY_DATA
+	UPROPERTY(EditAnywhere, Category = Compilation)
+	uint32 bWaitForCompilationOnActivate : 1;
+#endif
+
 private:
+	/** Did we try and activate but fail due to the asset being not yet ready. Keep looping.*/
+	uint32 bAwaitingActivationDueToNotReady : 1;
+	/** Should we try and reset when ready? */
+	uint32 bActivateShouldResetWhenReady : 1;
 
 	/** Did we auto attach during activation? Used to determine if we should restore the relative transform during detachment. */
 	uint32 bDidAutoAttach : 1;
 
+	/** Flag to mark us as currently changing auto attachment as part of Activate/Deactivate so we don't reset in the OnAttachmentChanged() callback. */
+	//uint32 bIsChangingAutoAttachment : 1;
+
 	/** Restore relative transform from auto attachment and optionally detach from parent (regardless of whether it was an auto attachment). */
 	void CancelAutoAttachment(bool bDetachFromParent);
 
@@ -297,6 +364,8 @@ private:
 	FVector SavedAutoAttachRelativeLocation;
 	FRotator SavedAutoAttachRelativeRotation;
 	FVector SavedAutoAttachRelativeScale3D;
+
+	FDelegateHandle AssetExposedParametersChangedHandle;
 };
 
 
@@ -318,7 +387,6 @@ public:
 	/** Called on render thread to assign new dynamic data */
 	void SetDynamicData_RenderThread(struct FNiagaraDynamicDataBase* NewDynamicData);
 	TArray<class NiagaraRenderer*>& GetEmitterRenderers() { return EmitterRenderers; }
-	void AddEmitterRenderer(NiagaraRenderer* Renderer);
 	void UpdateEmitterRenderers(TArray<NiagaraRenderer*>& InRenderers);
 
 	/** Gets whether or not this scene proxy should be rendered. */
diff --git a/Engine/Plugins/FX/Niagara/Source/Niagara/Public/NiagaraEditorDataBase.h b/Engine/Plugins/FX/Niagara/Source/Niagara/Public/NiagaraEditorDataBase.h
new file mode 100644
--- /dev/null
+++ b/Engine/Plugins/FX/Niagara/Source/Niagara/Public/NiagaraEditorDataBase.h
@@ -0,0 +1,20 @@
+// Copyright 1998-2017 Epic Games, Inc. All Rights Reserved.
+
+#pragma once
+
+#include "CoreTypes.h"
+
+
+#include "NiagaraEditorDataBase.generated.h"
+
+/** A base class for editor only data which supports post loading from the runtime owner object. */
+UCLASS(MinimalAPI)
+class UNiagaraEditorDataBase : public UObject
+{
+	GENERATED_BODY()
+public:
+#if WITH_EDITORONLY_DATA
+	virtual void PostLoadFromOwner(UObject* InOwner) PURE_VIRTUAL(UNiagaraEditorDataBase::PostLoadFromOwner, );
+#endif
+};
+
diff --git a/Engine/Plugins/FX/Niagara/Source/Niagara/Public/NiagaraGlobalReadBuffer.h b/Engine/Plugins/FX/Niagara/Source/Niagara/Public/NiagaraGlobalReadBuffer.h
new file mode 100644
--- /dev/null
+++ b/Engine/Plugins/FX/Niagara/Source/Niagara/Public/NiagaraGlobalReadBuffer.h
@@ -0,0 +1,99 @@
+// Copyright 1998-2018 Epic Games, Inc. All Rights Reserved.
+
+/*==============================================================================
+NiagaraGlobalReadBuffer.h: Class for allocating transient rendering data.
+==============================================================================*/
+
+#pragma once
+
+struct FDynamicReadBufferPool;
+
+struct FDynamicAllocReadBuffer : public FDynamicReadBuffer
+{
+	FDynamicAllocReadBuffer()
+		: FDynamicReadBuffer()
+		, AllocatedByteCount(0)
+	{}
+
+	int32 AllocatedByteCount;
+
+	/**
+	* Unocks the buffer so the GPU may read from it.
+	*/
+	void Unlock()
+	{
+		FDynamicReadBuffer::Unlock();
+		AllocatedByteCount = 0;
+	}
+};
+
+/**
+* A system for dynamically allocating GPU memory for Niagara rendering. Note that this must derive from FRenderResource 
+  so that we can safely free the shader resource views for OpenGL and other platforms. If we wait until the module is shutdown,
+  the renderer RHI will have already been destroyed and we can execute code on invalid data. By making ourself a render resource, we
+  clean up immediately before the renderer dies.
+*/
+class FNiagaraGlobalReadBuffer : public FRenderResource
+{
+public:
+	/**
+	* Information regarding an allocation from this buffer.
+	*/
+	struct FAllocation
+	{
+		/** The location of the buffer in main memory. */
+		uint8* Buffer;
+		/** The read buffer to bind for draw calls. */
+		FDynamicAllocReadBuffer* ReadBuffer;
+		/** The offset in to the read buffer. */
+		uint32 FirstIndex;
+
+		/** Default constructor. */
+		FAllocation()
+			: Buffer(NULL)
+			, ReadBuffer(NULL)
+			, FirstIndex(0)
+		{
+		}
+
+		/** Returns true if the allocation is valid. */
+		FORCEINLINE bool IsValid() const
+		{
+			return Buffer != NULL;
+		}
+	};
+
+	FNiagaraGlobalReadBuffer();
+	~FNiagaraGlobalReadBuffer();
+	
+	FAllocation AllocateFloat(uint32 Num);
+	FAllocation AllocateInt32(uint32 Num);
+
+	/**
+	* Commits allocated memory to the GPU.
+	*		WARNING: Once this buffer has been committed to the GPU, allocations
+	*		remain valid only until the next call to Allocate!
+	*/
+	void Commit();
+
+
+	/** Returns true if log statements should be made because we exceeded GMaxVertexBytesAllocatedPerFrame */
+	bool IsRenderAlarmLoggingEnabled() const;
+
+	static FNiagaraGlobalReadBuffer& Get();
+
+protected:
+	virtual void InitRHI() override;
+	virtual void ReleaseRHI() override;
+	void Cleanup();
+
+	/** The pools of read buffers from which allocations are made. */
+	FDynamicReadBufferPool* FloatBufferPool;
+	FDynamicReadBufferPool* Int32BufferPool;
+
+	/** A total of all allocations made since the last commit. Used to alert about spikes in memory usage. */
+	size_t TotalAllocatedSinceLastCommit;
+
+	FDelegateHandle CommitCallbackHandle;
+};
+
diff --git a/Engine/Plugins/FX/Niagara/Source/Niagara/Public/NiagaraLightRendererProperties.h b/Engine/Plugins/FX/Niagara/Source/Niagara/Public/NiagaraLightRendererProperties.h
--- a/Engine/Plugins/FX/Niagara/Source/Niagara/Public/NiagaraLightRendererProperties.h
+++ b/Engine/Plugins/FX/Niagara/Source/Niagara/Public/NiagaraLightRendererProperties.h
@@ -5,21 +5,25 @@
 #include "CoreMinimal.h"
 #include "UObject/ObjectMacros.h"
 #include "NiagaraRendererProperties.h"
+#include "NiagaraCommon.h"
 #include "NiagaraLightRendererProperties.generated.h"
 
-
-
 UCLASS(editinlinenew)
 class UNiagaraLightRendererProperties : public UNiagaraRendererProperties
 {
 public:
 	GENERATED_BODY()
 
-		UNiagaraLightRendererProperties();
+	UNiagaraLightRendererProperties();
+
+	virtual void PostInitProperties() override;
+
+	static void InitCDOPropertiesAfterModuleStartup();
 
 	//~ UNiagaraRendererProperties interface
 	virtual NiagaraRenderer* CreateEmitterRenderer(ERHIFeatureLevel::Type FeatureLevel) override;
 	virtual void GetUsedMaterials(TArray<UMaterialInterface*>& OutMaterials) const override;
+	virtual bool IsSimTargetSupported(ENiagaraSimTarget InSimTarget) const override { return (InSimTarget == ENiagaraSimTarget::CPUSim); };
 #if WITH_EDITORONLY_DATA
 	virtual bool IsMaterialValidForRenderer(UMaterial* Material, FText& InvalidMessage) override;
 	virtual void FixMaterial(UMaterial* Material) override;
@@ -28,8 +32,20 @@ public:
 #endif // WITH_EDITORONLY_DATA
 
 	UPROPERTY(EditAnywhere, Category = "Light Rendering")
-		float RadiusScale;
+	float RadiusScale;
 
 	UPROPERTY(EditAnywhere, Category = "Light Rendering")
-		FVector ColorAdd;
+	FVector ColorAdd;
+
+	/** Which attribute should we use for position when generating lights?*/
+	UPROPERTY(EditAnywhere, AdvancedDisplay, Category = "Light Rendering")
+	FNiagaraVariableAttributeBinding PositionBinding;
+
+	/** Which attribute should we use for light color when generating lights?*/
+	UPROPERTY(EditAnywhere, AdvancedDisplay, Category = "Light Rendering")
+	FNiagaraVariableAttributeBinding ColorBinding;
+
+	/** Which attribute should we use for light radius when generating lights?*/
+	UPROPERTY(EditAnywhere, AdvancedDisplay, Category = "Light Rendering")
+	FNiagaraVariableAttributeBinding RadiusBinding;
 };
diff --git a/Engine/Plugins/FX/Niagara/Source/Niagara/Public/NiagaraMeshRendererProperties.h b/Engine/Plugins/FX/Niagara/Source/Niagara/Public/NiagaraMeshRendererProperties.h
--- a/Engine/Plugins/FX/Niagara/Source/Niagara/Public/NiagaraMeshRendererProperties.h
+++ b/Engine/Plugins/FX/Niagara/Source/Niagara/Public/NiagaraMeshRendererProperties.h
@@ -6,6 +6,7 @@
 #include "UObject/ObjectMacros.h"
 #include "NiagaraRendererProperties.h"
 #include "StaticMeshResources.h"
+#include "NiagaraCommon.h"
 #include "NiagaraMeshRendererProperties.generated.h"
 
 /** This enum decides how a mesh particle will orient its "facing" axis relative to camera. Must keep these in sync with NiagaraMeshVertexFactory.ush*/
@@ -33,24 +34,34 @@ public:
 	//~ UNiagaraRendererProperties interface
 	virtual NiagaraRenderer* CreateEmitterRenderer(ERHIFeatureLevel::Type FeatureLevel) override;
 	virtual void GetUsedMaterials(TArray<UMaterialInterface*>& OutMaterials) const override;
+	virtual bool IsSimTargetSupported(ENiagaraSimTarget InSimTarget) const override { return true; };
 
 #if WITH_EDITORONLY_DATA
+	virtual void PostEditChangeProperty(struct FPropertyChangedEvent& PropertyChangedEvent) override;
 	virtual bool IsMaterialValidForRenderer(UMaterial* Material, FText& InvalidMessage) override;
 	virtual void FixMaterial(UMaterial* Material) override;
 	virtual const TArray<FNiagaraVariable>& GetRequiredAttributes() override;
 	virtual const TArray<FNiagaraVariable>& GetOptionalAttributes() override;
 #endif // WITH_EDITORONLY_DATA
 
-	void PostEditChangeProperty(FPropertyChangedEvent& PropertyChangedEvent);
+	virtual uint32 GetNumIndicesPerInstance() final override;
 
 	/** The static mesh to be instanced when rendering mesh particles. If OverrideMaterial is not specified, the mesh's material is used. Note that the material must have the Niagara Mesh Particles flag checked.*/
 	UPROPERTY(EditAnywhere, Category = "Mesh Rendering")
 	UStaticMesh *ParticleMesh;
 
+	/** Determines how we sort the particles prior to rendering.*/
+	UPROPERTY(EditAnywhere, Category = "Sorting")
+	ENiagaraSortMode SortMode;
+
 	/** Whether or not to use the OverrideMaterials array instead of the mesh's existing materials.*/
 	UPROPERTY(EditAnywhere, Category = "Mesh Rendering", meta = (InlineEditConditionToggle))
 	uint32 bOverrideMaterials : 1;
 
+	/** If true, the particles are only sorted when using a translucent material. */
+	UPROPERTY(EditAnywhere, Category = "Sorting")
+	uint32 bSortOnlyWhenTranslucent : 1;
+
 	/** The materials to be used instead of the StaticMesh's materials. Note that each material must have the Niagara Mesh Particles flag checked. If the ParticleMesh 
 	requires more materials than exist in this array or any entry in this array is set to None, we will use the ParticleMesh's existing Material instead.*/
 	UPROPERTY(EditAnywhere, Category = "Mesh Rendering", meta = (EditCondition = "bOverrideMaterials"))
@@ -58,5 +69,53 @@ public:
 
 	/** Determines how the mesh orients itself relative to the camera.*/
 	UPROPERTY(EditAnywhere, Category = "Mesh Rendering")
-	ENiagaraMeshFacingMode FacingMode;
+	ENiagaraMeshFacingMode FacingMode; 
+	
+	virtual void PostInitProperties() override;
+	static void InitCDOPropertiesAfterModuleStartup();
+	
+	/** Which attribute should we use for position when generating instanced meshes?*/
+	UPROPERTY(EditAnywhere, Category = "Bindings")
+	FNiagaraVariableAttributeBinding PositionBinding;
+
+	/** Which attribute should we use for color when generating instanced meshes?*/
+	UPROPERTY(EditAnywhere, Category = "Bindings")
+	FNiagaraVariableAttributeBinding ColorBinding;
+
+	/** Which attribute should we use for velocity when generating instanced meshes?*/
+	UPROPERTY(EditAnywhere, Category = "Bindings")
+	FNiagaraVariableAttributeBinding VelocityBinding;
+
+	/** Which attribute should we use for orienting meshes when generating instanced meshes?*/
+	UPROPERTY(EditAnywhere, Category = "Bindings")
+	FNiagaraVariableAttributeBinding MeshOrientationBinding;
+
+	/** Which attribute should we use for scale when generating instanced meshes?*/
+	UPROPERTY(EditAnywhere, Category = "Bindings")
+	FNiagaraVariableAttributeBinding ScaleBinding;
+	
+	/** Which attribute should we use for dynamic material parameters when generating instanced meshes?*/
+	UPROPERTY(EditAnywhere, Category = "Bindings")
+	FNiagaraVariableAttributeBinding DynamicMaterialBinding;
+
+	/** Which attribute should we use for dynamic material parameters when generating instanced meshes?*/
+	UPROPERTY(EditAnywhere, Category = "Bindings")
+	FNiagaraVariableAttributeBinding DynamicMaterial1Binding;
+
+	/** Which attribute should we use for dynamic material parameters when generating instanced meshes?*/
+	UPROPERTY(EditAnywhere, Category = "Bindings")
+	FNiagaraVariableAttributeBinding DynamicMaterial2Binding;
+	
+	/** Which attribute should we use for dynamic material parameters when generating instanced meshes?*/
+	UPROPERTY(EditAnywhere, Category = "Bindings")
+	FNiagaraVariableAttributeBinding DynamicMaterial3Binding;
+
+	/** Which attribute should we use custom sorting of particles in this emitter. */
+	UPROPERTY(EditAnywhere, Category = "Bindings")
+	FNiagaraVariableAttributeBinding CustomSortingBinding;
+
+	UPROPERTY(Transient)
+	int32 SyncId;
+protected:
+	void InitBindings();
 };
diff --git a/Engine/Plugins/FX/Niagara/Source/Niagara/Public/NiagaraModule.h b/Engine/Plugins/FX/Niagara/Source/Niagara/Public/NiagaraModule.h
--- a/Engine/Plugins/FX/Niagara/Source/Niagara/Public/NiagaraModule.h
+++ b/Engine/Plugins/FX/Niagara/Source/Niagara/Public/NiagaraModule.h
@@ -6,14 +6,20 @@
 #include "Stats/Stats.h"
 #include "Modules/ModuleInterface.h"
 #include "Engine/World.h"
+#include "NiagaraTypes.h"
+#include "Templates/SharedPointer.h"
 
 class FNiagaraWorldManager;
 class UNiagaraEmitter;
+struct FNiagaraVMExecutableData;
+class UNiagaraScript;
+class FNiagaraCompileOptions;
+class FNiagaraCompileRequestDataBase;
 
 /**
 * Niagara module interface
 */
-class INiagaraModule : public IModuleInterface
+class NIAGARA_API INiagaraModule : public IModuleInterface
 {
 public:
 #if WITH_EDITOR
@@ -40,11 +46,16 @@ public:
 			return FString::Join(ErrorMessageStrings, TEXT("\n"));
 		}
 	};
+	
 #endif
 
 public:
 #if WITH_EDITOR
+	typedef TSharedPtr<FNiagaraCompileRequestDataBase, ESPMode::ThreadSafe> CompileRequestPtr;
 	DECLARE_DELEGATE_RetVal_ThreeParams(FMergeEmitterResults, FOnMergeEmitter, UNiagaraEmitter&, UNiagaraEmitter&, UNiagaraEmitter&);
+	DECLARE_DELEGATE_RetVal_OneParam(class UNiagaraScriptSourceBase*, FOnCreateDefaultScriptSource, UObject*);
+	DECLARE_DELEGATE_RetVal_TwoParams(TSharedPtr<FNiagaraVMExecutableData>, FScriptCompiler,const FNiagaraCompileRequestDataBase*, const FNiagaraCompileOptions&);
+	DECLARE_DELEGATE_RetVal_OneParam(CompileRequestPtr, FOnPrecompile, UObject*);
 #endif
 	DECLARE_DELEGATE_RetVal(void, FOnProcessQueue);
 
@@ -52,12 +63,12 @@ public:
 	virtual void StartupModule()override;
 	virtual void ShutdownModule()override;
 
-	FDelegateHandle NIAGARA_API SetOnProcessShaderCompilationQueue(FOnProcessQueue InOnProcessQueue);
-	void NIAGARA_API ResetOnProcessShaderCompilationQueue(FDelegateHandle DelegateHandle);
+	FDelegateHandle SetOnProcessShaderCompilationQueue(FOnProcessQueue InOnProcessQueue);
+	void ResetOnProcessShaderCompilationQueue(FDelegateHandle DelegateHandle);
 	void ProcessShaderCompilationQueue();
 
 
-	FNiagaraWorldManager* GetWorldManager(UWorld* World);
+	static FNiagaraWorldManager* GetWorldManager(UWorld* World);
 
 	void DestroyAllSystemSimulations(class UNiagaraSystem* System);
 
@@ -75,21 +86,187 @@ public:
 #if WITH_EDITOR
 	FMergeEmitterResults MergeEmitter(UNiagaraEmitter& Source, UNiagaraEmitter& LastMergedSource, UNiagaraEmitter& Instance);
 
-	FDelegateHandle NIAGARA_API RegisterOnMergeEmitter(FOnMergeEmitter OnMergeEmitter);
+	FDelegateHandle RegisterOnMergeEmitter(FOnMergeEmitter OnMergeEmitter);
+
+	void UnregisterOnMergeEmitter(FDelegateHandle DelegateHandle);
+
+	UNiagaraScriptSourceBase* CreateDefaultScriptSource(UObject* Outer);
+
+	FDelegateHandle RegisterOnCreateDefaultScriptSource(FOnCreateDefaultScriptSource OnCreateDefaultScriptSource);
+
+	void UnregisterOnCreateDefaultScriptSource(FDelegateHandle DelegateHandle);
+
+	TSharedPtr<FNiagaraVMExecutableData> CompileScript(const FNiagaraCompileRequestDataBase* InCompileData, const FNiagaraCompileOptions& InCompileOptions);
+	FDelegateHandle RegisterScriptCompiler(FScriptCompiler ScriptCompiler);
+	void UnregisterScriptCompiler(FDelegateHandle DelegateHandle);
+
+	TSharedPtr<FNiagaraCompileRequestDataBase, ESPMode::ThreadSafe> Precompile(UObject* InObj);
+	FDelegateHandle RegisterPrecompiler(FOnPrecompile PreCompiler);
+	void UnregisterPrecompiler(FDelegateHandle DelegateHandle);
 
-	void NIAGARA_API UnregisterOnMergeEmitter(FDelegateHandle DelegateHandle);
 #endif
 
 	FORCEINLINE static int32 GetDetailLevel() { return EngineDetailLevel; }
+	FORCEINLINE static float GetGlobalSpawnCountScale() { return EngineGlobalSpawnCountScale; }
+	FORCEINLINE static float GetGlobalSystemCountScale() { return EngineGlobalSystemCountScale; }
+
+	static float EngineGlobalSpawnCountScale;
+	static float EngineGlobalSystemCountScale;
+
+	FORCEINLINE static const FNiagaraVariable&  GetVar_Engine_DeltaTime() { return Engine_DeltaTime; }
+	FORCEINLINE static const FNiagaraVariable&  GetVar_Engine_InvDeltaTime() { return Engine_InvDeltaTime; }
+	FORCEINLINE static const FNiagaraVariable&  GetVar_Engine_Time() { return Engine_Time; }
+	FORCEINLINE static const FNiagaraVariable&  GetVar_Engine_RealTime() { return Engine_RealTime; }
+
+	FORCEINLINE static const FNiagaraVariable&  GetVar_Engine_Owner_Position() { return Engine_Owner_Position; }
+	FORCEINLINE static const FNiagaraVariable&  GetVar_Engine_Owner_Velocity() { return Engine_Owner_Velocity; }
+	FORCEINLINE static const FNiagaraVariable&  GetVar_Engine_Owner_XAxis() { return Engine_Owner_XAxis; }
+	FORCEINLINE static const FNiagaraVariable&  GetVar_Engine_Owner_YAxis() { return Engine_Owner_YAxis; }
+	FORCEINLINE static const FNiagaraVariable&  GetVar_Engine_Owner_ZAxis() { return Engine_Owner_ZAxis; }
+	FORCEINLINE static const FNiagaraVariable&  GetVar_Engine_Owner_Scale() { return Engine_Owner_Scale; }
+
+	FORCEINLINE static const FNiagaraVariable&  GetVar_Engine_Owner_SystemLocalToWorld() { return Engine_Owner_SystemLocalToWorld; }
+	FORCEINLINE static const FNiagaraVariable&  GetVar_Engine_Owner_SystemWorldToLocal() { return Engine_Owner_SystemWorldToLocal; }
+	FORCEINLINE static const FNiagaraVariable&  GetVar_Engine_Owner_SystemLocalToWorldTransposed() { return Engine_Owner_SystemLocalToWorldTransposed; }
+	FORCEINLINE static const FNiagaraVariable&  GetVar_Engine_Owner_SystemWorldToLocalTransposed() { return Engine_Owner_SystemWorldToLocalTransposed; }
+	FORCEINLINE static const FNiagaraVariable&  GetVar_Engine_Owner_SystemLocalToWorldNoScale() { return Engine_Owner_SystemLocalToWorldNoScale; }
+	FORCEINLINE static const FNiagaraVariable&  GetVar_Engine_Owner_SystemWorldToLocalNoScale() { return Engine_Owner_SystemWorldToLocalNoScale; }
+
+	FORCEINLINE static const FNiagaraVariable&  GetVar_Engine_Owner_TimeSinceRendered() { return Engine_Owner_TimeSinceRendered; }
+	FORCEINLINE static const FNiagaraVariable&  GetVar_Engine_Owner_MinDistanceToCamera() { return Engine_Owner_MinDistanceToCamera; }
+
+	FORCEINLINE static const FNiagaraVariable&  GetVar_Engine_Owner_ExecutionState() { return Engine_Owner_ExecutionState; }
+
+	FORCEINLINE static const FNiagaraVariable&  GetVar_Engine_ExecutionCount() { return Engine_ExecutionCount; }
+	FORCEINLINE static const FNiagaraVariable&  GetVar_Engine_Emitter_NumParticles() { return Engine_Emitter_NumParticles; }
+	FORCEINLINE static const FNiagaraVariable&  GetVar_Engine_System_NumEmittersAlive() { return Engine_System_NumEmittersAlive; }
+	FORCEINLINE static const FNiagaraVariable&  GetVar_Engine_System_NumEmitters() { return Engine_System_NumEmitters; }
+	FORCEINLINE static const FNiagaraVariable&  GetVar_Engine_NumSystemInstances() { return Engine_NumSystemInstances; }
+
+	FORCEINLINE static const FNiagaraVariable&  GetVar_Engine_GlobalSpawnCountScale() { return Engine_GlobalSpawnCountScale; }
+	FORCEINLINE static const FNiagaraVariable&  GetVar_Engine_GlobalSystemScale() { return Engine_GlobalSystemScale; }
+
+	FORCEINLINE static const FNiagaraVariable&  GetVar_Engine_System_Age() { return Engine_System_Age; }
+	FORCEINLINE static const FNiagaraVariable&  GetVar_Emitter_Age() { return Emitter_Age; }
+	FORCEINLINE static const FNiagaraVariable&  GetVar_Emitter_LocalSpace() { return Emitter_LocalSpace; }
+	FORCEINLINE static const FNiagaraVariable&  GetVar_Emitter_SpawnRate() { return Emitter_SpawnRate; }
+	FORCEINLINE static const FNiagaraVariable&  GetVar_Emitter_SpawnInterval() { return Emitter_SpawnInterval; }
+	FORCEINLINE static const FNiagaraVariable&  GetVar_Emitter_InterpSpawnStartDt() { return Emitter_InterpSpawnStartDt; }
+
+	FORCEINLINE static const FNiagaraVariable&  GetVar_Particles_ID() { return Particles_ID; }
+	FORCEINLINE static const FNiagaraVariable&  GetVar_Particles_Position() { return Particles_Position; }
+	FORCEINLINE static const FNiagaraVariable&  GetVar_Particles_Velocity() { return Particles_Velocity; }
+	FORCEINLINE static const FNiagaraVariable&  GetVar_Particles_Color() { return Particles_Color; }
+	FORCEINLINE static const FNiagaraVariable&  GetVar_Particles_SpriteRotation() { return Particles_SpriteRotation; }
+	FORCEINLINE static const FNiagaraVariable&  GetVar_Particles_NormalizedAge() { return Particles_NormalizedAge; }
+	FORCEINLINE static const FNiagaraVariable&  GetVar_Particles_SpriteSize() { return Particles_SpriteSize; }
+	FORCEINLINE static const FNiagaraVariable&  GetVar_Particles_SpriteFacing() { return Particles_SpriteFacing; }
+	FORCEINLINE static const FNiagaraVariable&  GetVar_Particles_SpriteAlignment() { return Particles_SpriteAlignment; }
+	FORCEINLINE static const FNiagaraVariable&  GetVar_Particles_SubImageIndex() { return Particles_SubImageIndex; }
+	FORCEINLINE static const FNiagaraVariable&  GetVar_Particles_DynamicMaterialParameter() { return Particles_DynamicMaterialParameter; }
+	FORCEINLINE static const FNiagaraVariable&  GetVar_Particles_DynamicMaterialParameter1() { return Particles_DynamicMaterialParameter1; }
+	FORCEINLINE static const FNiagaraVariable&  GetVar_Particles_DynamicMaterialParameter2() { return Particles_DynamicMaterialParameter2; }
+	FORCEINLINE static const FNiagaraVariable&  GetVar_Particles_DynamicMaterialParameter3() { return Particles_DynamicMaterialParameter3; }
+	FORCEINLINE static const FNiagaraVariable&  GetVar_Particles_Scale() { return Particles_Scale; }
+	FORCEINLINE static const FNiagaraVariable&  GetVar_Particles_Lifetime() { return Particles_Lifetime; }
+	FORCEINLINE static const FNiagaraVariable&  GetVar_Particles_MeshOrientation() { return Particles_MeshOrientation; }
+	FORCEINLINE static const FNiagaraVariable&  GetVar_Particles_UVScale() { return Particles_UVScale; }
+	FORCEINLINE static const FNiagaraVariable&  GetVar_Particles_CameraOffset() { return Particles_CameraOffset; }
+	FORCEINLINE static const FNiagaraVariable&  GetVar_Particles_MaterialRandom() { return Particles_MaterialRandom; }
+	FORCEINLINE static const FNiagaraVariable&  GetVar_Particles_LightRadius() { return Particles_LightRadius; }
+	FORCEINLINE static const FNiagaraVariable&  GetVar_Particles_RibbonID() { return Particles_RibbonID; }
+	FORCEINLINE static const FNiagaraVariable&  GetVar_Particles_RibbonWidth() { return Particles_RibbonWidth; }
+	FORCEINLINE static const FNiagaraVariable&  GetVar_Particles_RibbonTwist() { return Particles_RibbonTwist; }
+	FORCEINLINE static const FNiagaraVariable&  GetVar_Particles_RibbonFacing() { return Particles_RibbonFacing; }
+	FORCEINLINE static const FNiagaraVariable&  GetVar_Particles_RibbonLinkOrder() { return Particles_RibbonLinkOrder; }
+	
+	FORCEINLINE static const FNiagaraVariable&  GetVar_DataInstance_Alive() { return DataInstance_Alive; }
+	FORCEINLINE static const FNiagaraVariable&  GetVar_BeginDefaults() { return Translator_BeginDefaults; }
+
 private:
-	TMap<class UWorld*, class FNiagaraWorldManager*> WorldManagers;
 	FOnProcessQueue OnProcessQueue;
 
 #if WITH_EDITORONLY_DATA
 	FOnMergeEmitter OnMergeEmitterDelegate;
+	FOnCreateDefaultScriptSource OnCreateDefaultScriptSourceDelegate;
+	FScriptCompiler ScriptCompilerDelegate;
+	FOnPrecompile ObjectPrecompilerDelegate;
 #endif
 
+	static TMap<class UWorld*, class FNiagaraWorldManager*> WorldManagers;
+
 	void OnChangeDetailLevel(class IConsoleVariable* CVar);
 	static int32 EngineDetailLevel;
+
+private:
+	static FNiagaraVariable Engine_DeltaTime;
+	static FNiagaraVariable Engine_InvDeltaTime;
+	static FNiagaraVariable Engine_Time; 
+	static FNiagaraVariable Engine_RealTime; 
+
+	static FNiagaraVariable Engine_Owner_Position;
+	static FNiagaraVariable Engine_Owner_Velocity;
+	static FNiagaraVariable Engine_Owner_XAxis;
+	static FNiagaraVariable Engine_Owner_YAxis;
+	static FNiagaraVariable Engine_Owner_ZAxis;
+	static FNiagaraVariable Engine_Owner_Scale;
+
+	static FNiagaraVariable Engine_Owner_SystemLocalToWorld;
+	static FNiagaraVariable Engine_Owner_SystemWorldToLocal;
+	static FNiagaraVariable Engine_Owner_SystemLocalToWorldTransposed;
+	static FNiagaraVariable Engine_Owner_SystemWorldToLocalTransposed;
+	static FNiagaraVariable Engine_Owner_SystemLocalToWorldNoScale;
+	static FNiagaraVariable Engine_Owner_SystemWorldToLocalNoScale;
+
+	static FNiagaraVariable Engine_Owner_TimeSinceRendered;
+	static FNiagaraVariable Engine_Owner_MinDistanceToCamera;
+	
+	static FNiagaraVariable Engine_Owner_ExecutionState;
+
+	static FNiagaraVariable Engine_ExecutionCount;
+	static FNiagaraVariable Engine_Emitter_NumParticles;
+	static FNiagaraVariable Engine_System_NumEmittersAlive;
+	static FNiagaraVariable Engine_System_NumEmitters;
+	static FNiagaraVariable Engine_NumSystemInstances;
+
+	static FNiagaraVariable Engine_GlobalSpawnCountScale;
+	static FNiagaraVariable Engine_GlobalSystemScale;
+
+	static FNiagaraVariable Engine_System_Age;
+	static FNiagaraVariable Emitter_Age;
+	static FNiagaraVariable Emitter_LocalSpace;
+	static FNiagaraVariable Emitter_SpawnRate;
+	static FNiagaraVariable Emitter_SpawnInterval;
+	static FNiagaraVariable Emitter_InterpSpawnStartDt;
+
+	static FNiagaraVariable Particles_ID;
+	static FNiagaraVariable Particles_Position;
+	static FNiagaraVariable Particles_Velocity;
+	static FNiagaraVariable Particles_Color;
+	static FNiagaraVariable Particles_SpriteRotation;
+	static FNiagaraVariable Particles_NormalizedAge;
+	static FNiagaraVariable Particles_SpriteSize;
+	static FNiagaraVariable Particles_SpriteFacing;
+	static FNiagaraVariable Particles_SpriteAlignment;
+	static FNiagaraVariable Particles_SubImageIndex;
+	static FNiagaraVariable Particles_DynamicMaterialParameter;
+	static FNiagaraVariable Particles_DynamicMaterialParameter1;
+	static FNiagaraVariable Particles_DynamicMaterialParameter2;
+	static FNiagaraVariable Particles_DynamicMaterialParameter3;
+	static FNiagaraVariable Particles_Scale;
+	static FNiagaraVariable Particles_Lifetime;
+	static FNiagaraVariable Particles_MeshOrientation;
+	static FNiagaraVariable Particles_UVScale;
+	static FNiagaraVariable Particles_CameraOffset;
+	static FNiagaraVariable Particles_MaterialRandom;
+	static FNiagaraVariable Particles_LightRadius;
+	static FNiagaraVariable Particles_RibbonID;
+	static FNiagaraVariable Particles_RibbonWidth;
+	static FNiagaraVariable Particles_RibbonTwist;
+	static FNiagaraVariable Particles_RibbonFacing;
+	static FNiagaraVariable Particles_RibbonLinkOrder;
+
+	static FNiagaraVariable DataInstance_Alive;
+	static FNiagaraVariable Translator_BeginDefaults;
 };
 
