Copying //UE4/Dev-Niagara to //UE4/Dev-Main (Source: //UE4/Dev-Niagara @ 4074996)

diff --git a/Engine/Plugins/FX/Niagara/Source/Niagara/Public/NiagaraParameterStore.h b/Engine/Plugins/FX/Niagara/Source/Niagara/Public/NiagaraParameterStore.h
--- a/Engine/Plugins/FX/Niagara/Source/Niagara/Public/NiagaraParameterStore.h
+++ b/Engine/Plugins/FX/Niagara/Source/Niagara/Public/NiagaraParameterStore.h
@@ -72,7 +72,7 @@ private:
 USTRUCT()
 struct NIAGARA_API FNiagaraParameterStore
 {
-#ifdef WITH_EDITOR
+#if WITH_EDITOR
 	DECLARE_MULTICAST_DELEGATE(FOnChanged);
 #endif
 
@@ -106,7 +106,10 @@ private:
 	/** Marks our interfaces as dirty. They will be pushed to any bound stores on tick if true. */
 	uint32 bInterfacesDirty : 1;
 
-#ifdef WITH_EDITOR
+	/** Uniquely identifies the current layout of this parameter store for detecting layout changes. */
+	uint32 LayoutVersion;
+
+#if WITH_EDITOR
 	FOnChanged OnChangedDelegate;
 #endif
 
@@ -120,11 +123,13 @@ public:
 	
 #if WITH_EDITORONLY_DATA
 	UPROPERTY()
-	FName Name;
+	FString DebugName;
 #endif
 
+	UObject* GetOwner()const { return Owner; }
+
 	void Dump();
-	void DumpParameters();
+	void DumpParameters(bool bDumpBindings = false);
 
 	FORCEINLINE uint32 GetParametersDirty() const { return bParametersDirty; }
 	FORCEINLINE uint32 GetInterfacesDirty() const { return bInterfacesDirty; }
@@ -132,6 +137,8 @@ public:
 	FORCEINLINE void MarkParametersDirty() { bParametersDirty = true; }
 	FORCEINLINE void MarkInterfacesDirty() { bInterfacesDirty = true; }
 
+	FORCEINLINE uint32 GetLayoutVersion() const { return LayoutVersion; }
+
 	/** Binds this parameter store to another. During Tick the values of this store will be pushed into all bound stores */
 	void Bind(FNiagaraParameterStore* DestStore);
 	/** Unbinds this store form one it's bound to. */
@@ -152,7 +159,7 @@ public:
 	Does nothing if this parameter is already present.
 	Returns true if we added a new parameter.
 	*/
-	bool AddParameter(const FNiagaraVariable& Param, bool bInitialize=true);
+	bool AddParameter(const FNiagaraVariable& Param, bool bInitialize=true, bool bTriggerRebind = true);
 
 	/** Removes the passed parameter if it exists in the store. */
 	bool RemoveParameter(const FNiagaraVariable& Param);
@@ -161,7 +168,10 @@ public:
 	void RenameParameter(const FNiagaraVariable& Param, FName NewName);
 
 	/** Removes all parameters from this store and releases any data. */
-	void Empty(bool bClearBindings=true);
+	void Empty(bool bClearBindings = true);
+
+	/** Removes all parameters from this store but does't change memory allocations. */
+	void Reset(bool bClearBindings = true);
 
 	FORCEINLINE TArray<FNiagaraParameterStore*>& GetSourceParameterStores() { return SourceStores; }
 	FORCEINLINE const TMap<FNiagaraVariable, int32>& GetParameterOffests()const { return ParameterOffsets; }
@@ -172,11 +182,10 @@ public:
 	FORCEINLINE const TArray<UNiagaraDataInterface*>& GetDataInterfaces()const { return DataInterfaces; }
 	FORCEINLINE const TArray<uint8>& GetParameterDataArray()const { return ParameterData; }
 
-	FORCEINLINE void SetParameterDataArray(const TArray<uint8>& InParameterDataArray)
-	{ 
-		ParameterData = InParameterDataArray;
-		OnParameterChange();
-	}
+	FORCEINLINE void SetParameterDataArray(const TArray<uint8>& InParameterDataArray);
+
+	// Called to initially set up the parameter store to *exactly* match the input store (other than any bindings and the internal name of it).
+	void InitFromSource(const FNiagaraParameterStore* SrcStore, bool bNotifyAsDirty);
 
 	/** Gets the index of the passed parameter. If it is a data interface, this is an offset into the data interface table, otherwise a byte offset into he parameter data buffer. */
 	int32 IndexOf(const FNiagaraVariable& Parameter)const
@@ -276,12 +285,24 @@ public:
 		}
 	}
 
+	enum class EDataInterfaceCopyMethod
+	{
+		/** A new data interface will be created and it will be synchronized using the CopyTo method. */
+		Value,
+		/** A reference to the source data interface will be added to the destination. */
+		Reference,
+		/** Do not copy data interfaces.  This will cause an assert if there are data interfaces in the source
+		  * store, and bOnlyAdd is false. */
+		None
+	};
+
 	/** Copies all parameters from this parameter store into another.*/
-	void CopyParametersTo(FNiagaraParameterStore& DestStore, bool bOnlyAdd);
+	void CopyParametersTo(FNiagaraParameterStore& DestStore, bool bOnlyAdd, EDataInterfaceCopyMethod DataInterfaceCopyMethod);
 
 	/** Remove all parameters from this parameter store from another.*/
 	void RemoveParameters(FNiagaraParameterStore& DestStore);
 
+	FString ToString() const;
 
 	template<typename T>
 	FORCEINLINE_DEBUGGABLE bool SetParameterValue(const T& InValue, const FNiagaraVariable& Param, bool bAdd=false)
@@ -292,7 +313,8 @@ public:
 		{
 			//*(T*)(GetParameterData_Internal(Offset)) = InValue;
 			//Until we solve our alignment issues, temporarily just doing a memcpy here.
-			FMemory::Memcpy(GetParameterData_Internal(Offset), &InValue, sizeof(T));
+			T* ParamData = reinterpret_cast<T*>( GetParameterData_Internal(Offset) );
+			FMemory::Memcpy(ParamData, &InValue, sizeof(T));
 			OnParameterChange();
 			return true;
 		}
@@ -300,10 +322,13 @@ public:
 		{
 			if (bAdd)
 			{
-				AddParameter(Param);
+				bool bInitInterfaces = false;
+				bool bTriggerRebind = false;
+				AddParameter(Param, bInitInterfaces, bTriggerRebind);
 				Offset = IndexOf(Param);
 				check(Offset != INDEX_NONE);
 				*(T*)(GetParameterData_Internal(Offset)) = InValue;
+				OnLayoutChange();
 				return true;
 			}
 		}
@@ -312,19 +337,37 @@ public:
 
 	FORCEINLINE_DEBUGGABLE void SetParameterData(const uint8* Data, int32 Offset, int32 Size)
 	{
+		checkSlow(Data != nullptr);
 		FMemory::Memcpy(GetParameterData_Internal(Offset), Data, Size);
 		OnParameterChange();
 	}
 
-	FORCEINLINE_DEBUGGABLE void SetParameterData(const uint8* Data, const FNiagaraVariable& Param)
+	FORCEINLINE_DEBUGGABLE bool SetParameterData(const uint8* Data, const FNiagaraVariable& Param, bool bAdd = false)
 	{
+		checkSlow(Data != nullptr);
 		int32 Offset = IndexOf(Param);
 		if (Offset != INDEX_NONE)
 		{
 			checkSlow(!Param.IsDataInterface());
 			FMemory::Memcpy(GetParameterData_Internal(Offset), Data, Param.GetSizeInBytes());
 			OnParameterChange();
+			return true;
+		}
+		else
+		{
+			if (bAdd)
+			{
+				bool bInitInterfaces = false;
+				bool bTriggerRebind = false;
+				AddParameter(Param, bInitInterfaces, bTriggerRebind);
+				Offset = IndexOf(Param);
+				check(Offset != INDEX_NONE);
+				FMemory::Memcpy(GetParameterData_Internal(Offset), Data, Param.GetSizeInBytes());
+				OnLayoutChange();
+				return true;
+			}
 		}
+		return false;
 	}
 
 	/**
@@ -361,7 +404,7 @@ public:
 	FORCEINLINE void OnParameterChange() 
 	{ 
 		bParametersDirty = true;
-#ifdef WITH_EDITOR
+#if WITH_EDITOR
 		OnChangedDelegate.Broadcast();
 #endif
 	}
@@ -369,17 +412,19 @@ public:
 	FORCEINLINE void OnInterfaceChange() 
 	{ 
 		bInterfacesDirty = true;
-#ifdef WITH_EDITOR
+#if WITH_EDITOR
 		OnChangedDelegate.Broadcast();
 #endif
 	}
 
-#ifdef WITH_EDITOR
+#if WITH_EDITOR
 	FDelegateHandle AddOnChangedHandler(FOnChanged::FDelegate InOnChanged);
 	void RemoveOnChangedHandler(FDelegateHandle DelegateHandle);
 	void RemoveAllOnChangedHandlers(const void* InUserObject);
 #endif
 
+	void TriggerOnLayoutChanged() { OnLayoutChange(); }
+
 protected:
 	void OnLayoutChange();
 
@@ -393,23 +438,34 @@ protected:
 	{
 		return ParameterData.GetData() + Offset;
 	}
+
+
+	template<typename ParamType>
+	void SetParameterByOffset(uint32 ParamOffset, const ParamType& Param)
+	{
+		ParamType* ParamPtr = (ParamType*)(GetParameterDataArray().GetData() + ParamOffset);
+		*ParamPtr = Param;
+		//SetParameterData((const uint8*)&Param, ParamOffset, sizeof(ParamType)); // TODO why aren't we using this path instead of SetParametersByOffset?
+	}
+
+	friend struct FNiagaraParameterStoreToDataSetBinding;    // this should be the only class calling SetParameterByOffset
 };
 
 FORCEINLINE_DEBUGGABLE void FNiagaraParameterStoreBinding::Empty(FNiagaraParameterStore* DestStore, FNiagaraParameterStore* SrcStore)
 {
 	if (DestStore)
 	{
-		DestStore->GetSourceParameterStores().RemoveSingleSwap(SrcStore);
+		DestStore->GetSourceParameterStores().RemoveSingleSwap(SrcStore, false);
 	}
 	DestStore = nullptr;
-	ParameterBindings.Empty();
-	InterfaceBindings.Empty();
+	ParameterBindings.Reset();
+	InterfaceBindings.Reset();
 }
 
 FORCEINLINE_DEBUGGABLE void FNiagaraParameterStoreBinding::BindParameters(FNiagaraParameterStore* DestStore, FNiagaraParameterStore* SrcStore)
 {
-	InterfaceBindings.Empty();
-	ParameterBindings.Empty();
+	InterfaceBindings.Reset();
+	ParameterBindings.Reset();
 	for (const TPair<FNiagaraVariable, int32>& ParamOffsetPair : DestStore->GetParameterOffests())
 	{
 		const FNiagaraVariable& Parameter = ParamOffsetPair.Key;
@@ -471,7 +527,7 @@ FORCEINLINE_DEBUGGABLE bool FNiagaraParameterStoreBinding::VerifyBinding(const F
 		}
 	}
 
-	UE_LOG(LogNiagara, Warning, TEXT("Invalid ParameterStore Binding: Parameters missing from binding between %s and %s. Stores should have been rebound when one changed!"), *SrcStore->Name.ToString(), *DestStore->Name.ToString());
+	UE_LOG(LogNiagara, Warning, TEXT("Invalid ParameterStore Binding: Parameters missing from binding between %s and %s. Stores should have been rebound when one changed!"), *SrcStore->DebugName, *DestStore->DebugName);
 	for (FName MissingParam : MissingParameterNames)
 	{
 		UE_LOG(LogNiagara, Warning, TEXT("%s"), *MissingParam.ToString());
@@ -508,7 +564,7 @@ FORCEINLINE_DEBUGGABLE void FNiagaraParameterStoreBinding::Tick(FNiagaraParamete
 FORCEINLINE_DEBUGGABLE void FNiagaraParameterStoreBinding::Dump(FNiagaraParameterStore* DestStore, FNiagaraParameterStore* SrcStore)const
 {
 #if WITH_EDITORONLY_DATA
-	UE_LOG(LogNiagara, Log, TEXT("\n\nDest Store: %s\n"), *DestStore->Name.ToString());
+	UE_LOG(LogNiagara, Log, TEXT("\n\nDest Store: %s\n"), *DestStore->DebugName);
 	for (const FParameterBinding& Binding : ParameterBindings)
 	{
 		ensure(Binding.Size != 0);
@@ -588,6 +644,7 @@ struct FNiagaraParameterDirectBinding
 	mutable T* ValuePtr;
 	FNiagaraParameterStore* BoundStore;
 	FNiagaraVariable BoundVariable;
+	uint32 LayoutVersion;
 
 	FNiagaraParameterDirectBinding()
 		: ValuePtr(nullptr), BoundStore(nullptr)
@@ -597,6 +654,8 @@ struct FNiagaraParameterDirectBinding
 	{
 		BoundStore = &InStore;
 		BoundVariable = DestVariable;
+		LayoutVersion = BoundStore->GetLayoutVersion();
+
 		check(DestVariable.GetSizeInBytes() == sizeof(T));
 		ValuePtr = (T*)InStore.GetParameterData(DestVariable);
 		return ValuePtr;
@@ -605,7 +664,7 @@ struct FNiagaraParameterDirectBinding
 	FORCEINLINE void SetValue(const T& InValue)
 	{
 		check(BoundVariable.GetSizeInBytes() == sizeof(T));
-		ValuePtr = (T*)BoundStore->GetParameterData(BoundVariable);
+		checkf(LayoutVersion == BoundStore->GetLayoutVersion(), TEXT("This binding is invalid, its bound parameter store's layout was changed since it was created"));
 
 		if (ValuePtr)
 		{
@@ -616,7 +675,7 @@ struct FNiagaraParameterDirectBinding
 	FORCEINLINE T GetValue()const 
 	{
 		check(BoundVariable.GetSizeInBytes() == sizeof(T));
-		ValuePtr = (T*)BoundStore->GetParameterData(BoundVariable);
+		checkf(LayoutVersion == BoundStore->GetLayoutVersion(), TEXT("This binding is invalid, its bound parameter store's layout was changed since it was created"));
 
 		if (ValuePtr)
 		{
@@ -632,6 +691,7 @@ struct FNiagaraParameterDirectBinding<FMatrix>
 	mutable FMatrix* ValuePtr;
 	FNiagaraParameterStore* BoundStore;
 	FNiagaraVariable BoundVariable;
+	uint32 LayoutVersion;
 
 	FNiagaraParameterDirectBinding():ValuePtr(nullptr), BoundStore(nullptr) {}
 
@@ -639,6 +699,7 @@ struct FNiagaraParameterDirectBinding<FMatrix>
 	{
 		BoundStore = &InStore;
 		BoundVariable = DestVariable;
+		LayoutVersion = BoundStore->GetLayoutVersion();
 
 		check(DestVariable.GetSizeInBytes() == sizeof(FMatrix));
 		ValuePtr = (FMatrix*)InStore.GetParameterData(DestVariable);
@@ -648,7 +709,7 @@ struct FNiagaraParameterDirectBinding<FMatrix>
 	FORCEINLINE void SetValue(const FMatrix& InValue)
 	{
 		check(BoundVariable.GetSizeInBytes() == sizeof(FMatrix));
-		ValuePtr = (FMatrix*)BoundStore->GetParameterData(BoundVariable);
+		checkf(LayoutVersion == BoundStore->GetLayoutVersion(), TEXT("This binding is invalid, its bound parameter store's layout was changed since it was created"));
 
 		if (ValuePtr)
 		{
@@ -659,7 +720,7 @@ struct FNiagaraParameterDirectBinding<FMatrix>
 	FORCEINLINE FMatrix GetValue()const
 	{
 		check(BoundVariable.GetSizeInBytes() == sizeof(FMatrix));
-		ValuePtr = (FMatrix*)BoundStore->GetParameterData(BoundVariable);
+		checkf(LayoutVersion == BoundStore->GetLayoutVersion(), TEXT("This binding is invalid, its bound parameter store's layout was changed since it was created"));
 
 		FMatrix Ret;
 		if (ValuePtr)
@@ -676,6 +737,7 @@ struct FNiagaraParameterDirectBinding<FVector4>
 	mutable FVector4* ValuePtr;
 	FNiagaraParameterStore* BoundStore;
 	FNiagaraVariable BoundVariable;
+	uint32 LayoutVersion;
 
 	FNiagaraParameterDirectBinding() :ValuePtr(nullptr), BoundStore(nullptr){}
 
@@ -683,6 +745,7 @@ struct FNiagaraParameterDirectBinding<FVector4>
 	{
 		BoundStore = &InStore;
 		BoundVariable = DestVariable;
+		LayoutVersion = BoundStore->GetLayoutVersion();
 
 		check(DestVariable.GetSizeInBytes() == sizeof(FVector4));
 		ValuePtr = (FVector4*)InStore.GetParameterData(DestVariable);
@@ -692,7 +755,7 @@ struct FNiagaraParameterDirectBinding<FVector4>
 	FORCEINLINE void SetValue(const FVector4& InValue)
 	{
 		check(BoundVariable.GetSizeInBytes() == sizeof(FVector4));
-		ValuePtr = (FVector4*)BoundStore->GetParameterData(BoundVariable);
+		checkf(LayoutVersion == BoundStore->GetLayoutVersion(), TEXT("This binding is invalid, its bound parameter store's layout was changed since it was created"));
 
 		if (ValuePtr)
 		{
@@ -703,7 +766,7 @@ struct FNiagaraParameterDirectBinding<FVector4>
 	FORCEINLINE FVector4 GetValue()const
 	{
 		check(BoundVariable.GetSizeInBytes() == sizeof(FVector4));
-		ValuePtr = (FVector4*)BoundStore->GetParameterData(BoundVariable);
+		checkf(LayoutVersion == BoundStore->GetLayoutVersion(), TEXT("This binding is invalid, its bound parameter store's layout was changed since it was created"));
 
 		FVector4 Ret;
 		if (ValuePtr)
@@ -712,4 +775,123 @@ struct FNiagaraParameterDirectBinding<FVector4>
 		}
 		return Ret;
 	}
+};
+
+template<>
+struct FNiagaraParameterDirectBinding<FQuat>
+{
+	mutable FQuat* ValuePtr;
+	FNiagaraParameterStore* BoundStore;
+	FNiagaraVariable BoundVariable;
+	uint32 LayoutVersion;
+
+	FNiagaraParameterDirectBinding() :ValuePtr(nullptr), BoundStore(nullptr) {}
+
+	FQuat* Init(FNiagaraParameterStore& InStore, const FNiagaraVariable& DestVariable)
+	{
+		BoundStore = &InStore;
+		BoundVariable = DestVariable;
+		LayoutVersion = BoundStore->GetLayoutVersion();
+
+		check(DestVariable.GetSizeInBytes() == sizeof(FQuat));
+		ValuePtr = (FQuat*)InStore.GetParameterData(DestVariable);
+		return ValuePtr;
+	}
+
+	FORCEINLINE void SetValue(const FQuat& InValue)
+	{
+		check(BoundVariable.GetSizeInBytes() == sizeof(FQuat));
+		checkf(LayoutVersion == BoundStore->GetLayoutVersion(), TEXT("This binding is invalid, its bound parameter store's layout was changed since it was created"));
+
+		if (ValuePtr)
+		{
+			FMemory::Memcpy(ValuePtr, &InValue, sizeof(FQuat));//Temp annoyance until we fix the alignment issues with parameter stores.
+		}
+	}
+
+	FORCEINLINE FQuat GetValue()const
+	{
+		check(BoundVariable.GetSizeInBytes() == sizeof(FQuat));
+		checkf(LayoutVersion == BoundStore->GetLayoutVersion(), TEXT("This binding is invalid, its bound parameter store's layout was changed since it was created"));
+
+		FQuat Ret;
+		if (ValuePtr)
+		{
+			FMemory::Memcpy(&Ret, ValuePtr, sizeof(FQuat));//Temp annoyance until we fix the alignment issues with parameter stores.
+		}
+		return Ret;
+	}
+};
+
+template<>
+struct FNiagaraParameterDirectBinding<FNiagaraBool>
+{
+	mutable uint32* ValuePtr;
+	FNiagaraParameterStore* BoundStore;
+	FNiagaraVariable BoundVariable;
+	uint32 LayoutVersion;
+
+	FNiagaraParameterDirectBinding() :ValuePtr(nullptr), BoundStore(nullptr) {}
+
+	uint32* Init(FNiagaraParameterStore& InStore, const FNiagaraVariable& DestVariable)
+	{
+		BoundStore = &InStore;
+		BoundVariable = DestVariable;
+		LayoutVersion = BoundStore->GetLayoutVersion();
+
+		check(DestVariable.GetSizeInBytes() == sizeof(FNiagaraBool));
+		check(sizeof(uint32) == sizeof(FNiagaraBool));
+		ValuePtr = (uint32*)InStore.GetParameterData(DestVariable);
+		return ValuePtr;
+	}
+
+	FORCEINLINE void SetValue(const FNiagaraBool& InValue)
+	{
+		check(BoundVariable.GetSizeInBytes() == sizeof(FNiagaraBool));
+		check(sizeof(uint32) == sizeof(FNiagaraBool));
+		checkf(LayoutVersion == BoundStore->GetLayoutVersion(), TEXT("This binding is invalid, its bound parameter store's layout was changed since it was created"));
+
+		if (ValuePtr)
+		{
+			FMemory::Memcpy(ValuePtr, &InValue, sizeof(FNiagaraBool));
+		}
+	}
+
+	FORCEINLINE void SetValue(const bool& InValue)
+	{
+		check(BoundVariable.GetSizeInBytes() == sizeof(FNiagaraBool));
+		checkf(LayoutVersion == BoundStore->GetLayoutVersion(), TEXT("This binding is invalid, its bound parameter store's layout was changed since it was created"));
+
+		if (ValuePtr)
+		{
+			if (!InValue)
+			{
+				*ValuePtr = FNiagaraBool::False;
+			}
+			else
+			{
+				*ValuePtr = FNiagaraBool::True;
+			}
+		}
+	}
+
+	FORCEINLINE FNiagaraBool GetValue()const
+	{
+		check(BoundVariable.GetSizeInBytes() == sizeof(FNiagaraBool));
+		checkf(LayoutVersion == BoundStore->GetLayoutVersion(), TEXT("This binding is invalid, its bound parameter store's layout was changed since it was created"));
+
+		FNiagaraBool Ret(false);
+		if (ValuePtr)
+		{
+			if (FNiagaraBool::False == *ValuePtr)
+			{
+				Ret = FNiagaraBool(false);
+			}
+			else
+			{
+				Ret = FNiagaraBool(true);
+			}
+		}
+		return Ret;
+	}
 };
\ No newline at end of file
diff --git a/Engine/Plugins/FX/Niagara/Source/Niagara/Public/NiagaraParameters.h b/Engine/Plugins/FX/Niagara/Source/Niagara/Public/NiagaraParameters.h
--- a/Engine/Plugins/FX/Niagara/Source/Niagara/Public/NiagaraParameters.h
+++ b/Engine/Plugins/FX/Niagara/Source/Niagara/Public/NiagaraParameters.h
@@ -73,4 +73,6 @@ public:
 		Empty();
 		Parameters = InParameters.Parameters;
 	}
+
+	NIAGARA_API void DumpParameters();
 };
diff --git a/Engine/Plugins/FX/Niagara/Source/Niagara/Public/NiagaraRenderer.h b/Engine/Plugins/FX/Niagara/Source/Niagara/Public/NiagaraRenderer.h
--- a/Engine/Plugins/FX/Niagara/Source/Niagara/Public/NiagaraRenderer.h
+++ b/Engine/Plugins/FX/Niagara/Source/Niagara/Public/NiagaraRenderer.h
@@ -19,12 +19,16 @@ NiagaraRenderer.h: Base class for Niagara render modules
 #include "RenderingThread.h"
 #include "SceneView.h"
 #include "NiagaraComponent.h"
+#include "NiagaraGlobalReadBuffer.h"
 
 class FNiagaraDataSet;
 
 /** Struct used to pass dynamic data from game thread to render thread */
 struct FNiagaraDynamicDataBase
 {
+	//Ugh. Temporarily copying the whole buffer over.
+	//After GDC i have a shelf that does just the data we need for rendering. For now, this is safer.
+	FNiagaraDataBuffer RTParticleData;
 };
 
 
@@ -50,6 +54,25 @@ private:
 	double StartTime;
 };
 
+class FNiagaraDummyRWBufferInt : public FRenderResource
+{
+public:
+	FNiagaraDummyRWBufferInt(const FString InDebugId) : DebugId(InDebugId) {}
+	FString DebugId;
+	FRWBuffer Buffer;
+	virtual void InitRHI() override;
+	virtual void ReleaseRHI() override;
+};
+
+class FNiagaraDummyRWBufferFloat : public FRenderResource
+{
+public:
+	FNiagaraDummyRWBufferFloat(const FString InDebugId) : DebugId(InDebugId) {}
+	FString DebugId;
+	FRWBuffer Buffer;
+	virtual void InitRHI() override;
+	virtual void ReleaseRHI() override;
+};
 
 /**
 * Base class for Niagara System renderers. System renderers handle generating vertex data for and
@@ -81,7 +104,7 @@ public:
 		bool bHasDynamicData = HasDynamicData();
 
 		//Always draw so our LastRenderTime is updated. We may not have dynamic data if we're disabled from visibility culling.
-		Result.bDrawRelevance = /*bHasDynamicData &&*/ SceneProxy->IsShown(View) && View->Family->EngineShowFlags.Particles;
+		Result.bDrawRelevance = bHasDynamicData && SceneProxy->IsShown(View) && View->Family->EngineShowFlags.Particles;
 		Result.bShadowRelevance = bHasDynamicData && SceneProxy->IsShadowCast(View);
 		Result.bDynamicRelevance = bHasDynamicData;
 		if (bHasDynamicData && View->Family->EngineShowFlags.Bounds)
@@ -111,6 +134,7 @@ public:
 	
 	virtual UClass *GetPropertiesClass() = 0;
 	virtual void SetRendererProperties(UNiagaraRendererProperties *Props) = 0;
+	virtual UNiagaraRendererProperties* GetRendererProperties() const = 0;
 
 	float GetCPUTimeMS() { return CPUTimeMS; }
 
@@ -126,6 +150,12 @@ public:
 	void SetEnabled(bool bInEnabled) { bEnabled = bInEnabled; }
 	
 	const FVector& GetBaseExtents() const {	return BaseExtents; }
+
+	void SortIndices(ENiagaraSortMode SortMode, int32 SortAttributeOffset, const FNiagaraDataBuffer& Buffer, const FMatrix& LocalToWorld, const FSceneView* View, FNiagaraGlobalReadBuffer::FAllocation& OutIndices)const;
+
+	static FRWBuffer& GetDummyFloatBuffer(); 
+	static FRWBuffer& GetDummyIntBuffer();
+	
 protected:
 	NiagaraRenderer();
 	virtual ~NiagaraRenderer();
@@ -175,6 +205,9 @@ public:
 
 	UClass *GetPropertiesClass() override { return UNiagaraSpriteRendererProperties::StaticClass(); }
 	void SetRendererProperties(UNiagaraRendererProperties *Props) override { Properties = Cast<UNiagaraSpriteRendererProperties>(Props); }
+	virtual UNiagaraRendererProperties* GetRendererProperties() const override {
+		return Properties;
+	}
 
 #if WITH_EDITORONLY_DATA
 	virtual const TArray<FNiagaraVariable>& GetRequiredAttributes() override;
@@ -196,9 +229,15 @@ private:
 	int32 AlignmentOffset;
 	int32 SubImageOffset;
 	int32 MaterialParamOffset;
+	int32 MaterialParamOffset1;
+	int32 MaterialParamOffset2;
+	int32 MaterialParamOffset3;
 	int32 CameraOffsetOffset;
 	int32 UVScaleOffset;
 	int32 ParticleRandomOffset;
+	int32 CustomSortingOffset;
+
+	int32 LastSyncId;
 };
 
 
@@ -250,6 +289,10 @@ public:
 
 	UClass *GetPropertiesClass() override { return UNiagaraLightRendererProperties::StaticClass(); }
 	void SetRendererProperties(UNiagaraRendererProperties *Props) override { Properties = Cast<UNiagaraLightRendererProperties>(Props); }
+	virtual UNiagaraRendererProperties* GetRendererProperties() const override {
+		return Properties;
+	}
+
 #if WITH_EDITORONLY_DATA
 	virtual const TArray<FNiagaraVariable>& GetRequiredAttributes() override;
 	virtual const TArray<FNiagaraVariable>& GetOptionalAttributes() override;
diff --git a/Engine/Plugins/FX/Niagara/Source/Niagara/Public/NiagaraRendererMeshes.h b/Engine/Plugins/FX/Niagara/Source/Niagara/Public/NiagaraRendererMeshes.h
--- a/Engine/Plugins/FX/Niagara/Source/Niagara/Public/NiagaraRendererMeshes.h
+++ b/Engine/Plugins/FX/Niagara/Source/Niagara/Public/NiagaraRendererMeshes.h
@@ -14,14 +14,9 @@ class FNiagaraDataSet;
 
 struct FNiagaraDynamicDataMesh : public FNiagaraDynamicDataBase
 {
+	//Direct ptr to the dataset. ONLY FOR USE BE GPU EMITTERS.
+	//TODO: Even this needs to go soon.
 	const FNiagaraDataSet *DataSet;
-	int32 PositionDataOffset;
-	int32 VelocityDataOffset;
-	int32 ColorDataOffset;
-	int32 TransformDataOffset;
-	int32 ScaleDataOffset;
-	int32 SizeDataOffset;
-	int32 MaterialParamDataOffset;
 };
 
 
@@ -78,6 +73,9 @@ public:
 
 	UClass *GetPropertiesClass() override { return UNiagaraMeshRendererProperties::StaticClass(); }
 	void SetRendererProperties(UNiagaraRendererProperties *Props) override { Properties = Cast<UNiagaraMeshRendererProperties>(Props); }
+	virtual UNiagaraRendererProperties* GetRendererProperties()  const override {
+		return Properties;
+	}
 #if WITH_EDITORONLY_DATA
 	virtual const TArray<FNiagaraVariable>& GetRequiredAttributes() override;
 	virtual const TArray<FNiagaraVariable>& GetOptionalAttributes() override;
@@ -88,4 +86,17 @@ private:
 	UNiagaraMeshRendererProperties *Properties;
 	mutable TUniformBuffer<FPrimitiveUniformShaderParameters> WorldSpacePrimitiveUniformBuffer;
 	class FNiagaraMeshVertexFactory* VertexFactory;
+
+	int32 PositionOffset;
+	int32 VelocityOffset;
+	int32 ColorOffset;
+	int32 ScaleOffset;
+	int32 SizeOffset;
+	int32 MaterialParamOffset;
+	int32 MaterialParamOffset1;
+	int32 MaterialParamOffset2;
+	int32 MaterialParamOffset3;
+	int32 TransformOffset;
+	int32 CustomSortingOffset;
+	int32 LastSyncedId;
 };
diff --git a/Engine/Plugins/FX/Niagara/Source/Niagara/Public/NiagaraRendererProperties.h b/Engine/Plugins/FX/Niagara/Source/Niagara/Public/NiagaraRendererProperties.h
--- a/Engine/Plugins/FX/Niagara/Source/Niagara/Public/NiagaraRendererProperties.h
+++ b/Engine/Plugins/FX/Niagara/Source/Niagara/Public/NiagaraRendererProperties.h
@@ -7,8 +7,27 @@
 #include "UObject/Object.h"
 #include "RHIDefinitions.h"
 #include "NiagaraTypes.h"
+#include "NiagaraCommon.h"
 #include "NiagaraMergeable.h"
 #include "NiagaraRendererProperties.generated.h"
+
+
+UENUM()
+enum class ENiagaraSortMode : uint8
+{
+	/** Perform no additional sorting prior to rendering.*/
+	None,
+	/** Sort by depth to the camera's near plane.*/
+	ViewDepth,
+	/** Sort by distance to the camera's origin.*/
+	ViewDistance,
+	/** Custom sorting according to a per particle attribute. Lower values are rendered before higher values. */
+	CustomAscending,
+	/** Custom sorting according to a per particle attribute. Higher values are rendered before lower values. */
+	CustomDecending,
+};
+
+
 /**
 * Emitter properties base class
 * Each EmitterRenderer derives from this with its own class, and returns it in GetProperties; a copy
@@ -26,9 +45,15 @@ class NIAGARA_API UNiagaraRendererProperties : public UNiagaraMergeable
 	GENERATED_BODY()
 
 public:
+	UNiagaraRendererProperties()
+	: bIsEnabled(true)
+	{
+	}
 	virtual NiagaraRenderer* CreateEmitterRenderer(ERHIFeatureLevel::Type FeatureLevel) PURE_VIRTUAL ( UNiagaraRendererProperties::CreateEmitterRenderer, return nullptr;);
 	virtual void GetUsedMaterials(TArray<UMaterialInterface*>& OutMaterials) const PURE_VIRTUAL(UNiagaraRendererProperties::GetUsedMaterials,);
 
+	virtual bool IsSimTargetSupported(ENiagaraSimTarget InSimTarget) const { return false; };
+
 #if WITH_EDITORONLY_DATA
 	virtual bool IsMaterialValidForRenderer(UMaterial* Material, FText& InvalidMessage) { return true; }
 
@@ -37,6 +62,27 @@ public:
 	virtual const TArray<FNiagaraVariable>& GetRequiredAttributes() { static TArray<FNiagaraVariable> Vars; return Vars; };
 	virtual const TArray<FNiagaraVariable>& GetOptionalAttributes() { static TArray<FNiagaraVariable> Vars; return Vars; };
 #endif // WITH_EDITORONLY_DATA
+
+
+	// GPU simulation uses DrawIndirect, so the sim step needs to know indices per instance in order to prepare the draw call parameters
+	virtual uint32 GetNumIndicesPerInstance() { return 0; }
+
+	virtual bool GetIsEnabled() const
+	{
+		return bIsEnabled;
+	}
+
+	virtual void SetIsEnabled(bool bInIsEnabled)
+	{
+		bIsEnabled = bInIsEnabled;
+	}
+	/** By default, emitters are drawn in the order that they are added to the system. This value will allow you to control the order in a more fine-grained manner. 
+	Materials of the same type (i.e. Transparent) will draw in order from lowest to highest within the system. The default value is 0.*/
+	UPROPERTY(EditAnywhere, Category = "Sort Order")
+	int32 SortOrderHint;
+	
+	UPROPERTY()
+	bool bIsEnabled;
 };
 
 
diff --git a/Engine/Plugins/FX/Niagara/Source/Niagara/Public/NiagaraRendererRibbons.h b/Engine/Plugins/FX/Niagara/Source/Niagara/Public/NiagaraRendererRibbons.h
--- a/Engine/Plugins/FX/Niagara/Source/Niagara/Public/NiagaraRendererRibbons.h
+++ b/Engine/Plugins/FX/Niagara/Source/Niagara/Public/NiagaraRendererRibbons.h
@@ -35,7 +35,7 @@ public:
 	FNiagaraDynamicDataBase *GenerateVertexData(const FNiagaraSceneProxy* Proxy, FNiagaraDataSet &Data, const ENiagaraSimTarget Target) override;
 
 	void AddRibbonVert(TArray<FNiagaraRibbonVertex>& RenderData, FVector ParticlePos, FVector2D UV1, FVector2D UV2,
-		const FLinearColor &Color, const float Age, const float Rotation, const float Size, const FVector NormDir, const FVector CustomFacing)
+		const FLinearColor &Color, const float Rotation, const float Size, const FVector NormDir, const FVector CustomFacing)
 	{
 		FNiagaraRibbonVertex NewVertex;
 		NewVertex.Position = ParticlePos;
@@ -48,7 +48,6 @@ public:
 		NewVertex.Tex_U2 = UV2.X * Properties->UV1Scale.X;
 		NewVertex.Tex_V2 = UV2.Y * Properties->UV1Scale.Y;
 		NewVertex.CustomFacingVector = CustomFacing;
-		//NewVertex.NormalizedAge = Age;
 		RenderData.Add(NewVertex);
 	}
 
@@ -75,11 +74,19 @@ public:
 
 	UClass *GetPropertiesClass() override { return UNiagaraRibbonRendererProperties::StaticClass(); }
 	void SetRendererProperties(UNiagaraRendererProperties *Props) override { Properties = Cast<UNiagaraRibbonRendererProperties>(Props); }
+	virtual UNiagaraRendererProperties* GetRendererProperties() const override {
+		return Properties;
+	}
 
 private:
 	class FNiagaraRibbonVertexFactory *VertexFactory;
 	UNiagaraRibbonRendererProperties *Properties;
 	mutable TUniformBuffer<FPrimitiveUniformShaderParameters> WorldSpacePrimitiveUniformBuffer;
+	int32 PositionDataOffset;
+	int32 WidthDataOffset;
+	int32 TwistDataOffset;
+	int32 ColorDataOffset;
+	int32 LastSyncedId;
 };
 
 
diff --git a/Engine/Plugins/FX/Niagara/Source/Niagara/Public/NiagaraRibbonRendererProperties.h b/Engine/Plugins/FX/Niagara/Source/Niagara/Public/NiagaraRibbonRendererProperties.h
--- a/Engine/Plugins/FX/Niagara/Source/Niagara/Public/NiagaraRibbonRendererProperties.h
+++ b/Engine/Plugins/FX/Niagara/Source/Niagara/Public/NiagaraRibbonRendererProperties.h
@@ -5,6 +5,7 @@
 #include "CoreMinimal.h"
 #include "UObject/ObjectMacros.h"
 #include "NiagaraRendererProperties.h"
+#include "NiagaraCommon.h"
 #include "NiagaraRibbonRendererProperties.generated.h"
 
 UENUM()
@@ -49,10 +50,73 @@ public:
 	//~ UNiagaraRendererProperties interface
 	virtual NiagaraRenderer* CreateEmitterRenderer(ERHIFeatureLevel::Type FeatureLevel) override;
 	virtual void GetUsedMaterials(TArray<UMaterialInterface*>& OutMaterials) const override;
+	virtual bool IsSimTargetSupported(ENiagaraSimTarget InSimTarget) const override { return (InSimTarget == ENiagaraSimTarget::CPUSim); };
 #if WITH_EDITORONLY_DATA
+	virtual void PostEditChangeProperty(struct FPropertyChangedEvent& PropertyChangedEvent) override;
 	virtual const TArray<FNiagaraVariable>& GetRequiredAttributes() override;
 	virtual const TArray<FNiagaraVariable>& GetOptionalAttributes() override;
     virtual bool IsMaterialValidForRenderer(UMaterial* Material, FText& InvalidMessage) override;
 	virtual void FixMaterial(UMaterial* Material);
 #endif
+
+	virtual void PostInitProperties() override;
+	static void InitCDOPropertiesAfterModuleStartup();
+
+
+	/** Which attribute should we use for position when generating ribbons?*/
+	UPROPERTY(EditAnywhere, AdvancedDisplay, Category = "Ribbon Rendering")
+	FNiagaraVariableAttributeBinding PositionBinding;
+
+	/** Which attribute should we use for color when generating ribbons?*/
+	UPROPERTY(EditAnywhere, AdvancedDisplay, Category = "Ribbon Rendering")
+	FNiagaraVariableAttributeBinding ColorBinding;
+
+	/** Which attribute should we use for velocity when generating ribbons?*/
+	UPROPERTY(EditAnywhere, AdvancedDisplay, Category = "Ribbon Rendering")
+	FNiagaraVariableAttributeBinding VelocityBinding;
+
+	/** Which attribute should we use for normalized age when generating ribbons?*/
+	UPROPERTY(EditAnywhere, AdvancedDisplay, Category = "Ribbon Rendering")
+	FNiagaraVariableAttributeBinding NormalizedAgeBinding;
+
+	/** Which attribute should we use for ribbon twist when generating ribbons?*/
+	UPROPERTY(EditAnywhere, AdvancedDisplay, Category = "Ribbon Rendering")
+	FNiagaraVariableAttributeBinding RibbonTwistBinding;
+
+	/** Which attribute should we use for ribbon width when generating ribbons?*/
+	UPROPERTY(EditAnywhere, AdvancedDisplay, Category = "Ribbon Rendering")
+	FNiagaraVariableAttributeBinding RibbonWidthBinding;
+
+	/** Which attribute should we use for ribbon facing when generating ribbons?*/
+	UPROPERTY(EditAnywhere, AdvancedDisplay, Category = "Ribbon Rendering")
+	FNiagaraVariableAttributeBinding RibbonFacingBinding;
+	
+	/** Which attribute should we use for ribbon id when generating ribbons?*/
+	UPROPERTY(EditAnywhere, AdvancedDisplay, Category = "Ribbon Rendering")
+	FNiagaraVariableAttributeBinding RibbonIdBinding;
+
+	/** Which attribute should we use for RibbonLinkOrder when generating ribbons?*/
+	UPROPERTY(EditAnywhere, AdvancedDisplay, Category = "Ribbon Rendering")
+	FNiagaraVariableAttributeBinding RibbonLinkOrderBinding;
+
+	/** Which attribute should we use for dynamic material parameters when generating ribbons?*/
+	UPROPERTY(EditAnywhere, AdvancedDisplay, Category = "Ribbon Rendering")
+	FNiagaraVariableAttributeBinding DynamicMaterialBinding;
+
+	/** Which attribute should we use for dynamic material parameters when generating ribbons?*/
+	UPROPERTY(EditAnywhere, AdvancedDisplay, Category = "Ribbon Rendering")
+	FNiagaraVariableAttributeBinding DynamicMaterial1Binding;
+
+	/** Which attribute should we use for dynamic material parameters when generating ribbons?*/
+	UPROPERTY(EditAnywhere, AdvancedDisplay, Category = "Ribbon Rendering")
+	FNiagaraVariableAttributeBinding DynamicMaterial2Binding;
+
+	/** Which attribute should we use for dynamic material parameters when generating ribbons?*/
+	UPROPERTY(EditAnywhere, AdvancedDisplay, Category = "Ribbon Rendering")
+	FNiagaraVariableAttributeBinding DynamicMaterial3Binding;
+
+	UPROPERTY(Transient)
+	int32 SyncId;
+protected:
+	void InitBindings();
 };
diff --git a/Engine/Plugins/FX/Niagara/Source/Niagara/Public/NiagaraScriptExecutionParameterStore.h b/Engine/Plugins/FX/Niagara/Source/Niagara/Public/NiagaraScriptExecutionParameterStore.h
new file mode 100644
--- /dev/null
+++ b/Engine/Plugins/FX/Niagara/Source/Niagara/Public/NiagaraScriptExecutionParameterStore.h
@@ -0,0 +1,101 @@
+#pragma once
+
+#include "NiagaraCommon.h"
+#include "NiagaraDataInterface.h"
+#include "NiagaraParameterStore.h"
+#include "NiagaraScriptExecutionParameterStore.generated.h"
+
+USTRUCT()
+struct FNiagaraScriptExecutionPaddingInfo
+{
+	GENERATED_USTRUCT_BODY()
+public:
+	FNiagaraScriptExecutionPaddingInfo() : SrcOffset(0), DestOffset(0), Size(0) {}
+	FNiagaraScriptExecutionPaddingInfo(uint32 InSrcOffset, uint32 InDestOffset, uint32 InSize) : SrcOffset(InSrcOffset), DestOffset(InDestOffset), Size(InSize) {}
+	uint32 SrcOffset;
+	uint32 DestOffset;
+	uint32 Size;
+};
+
+/**
+Storage class containing actual runtime buffers to be used by the VM and the GPU.
+Is not the actual source for any parameter data, rather just the final place it's gathered from various other places ready for execution.
+*/
+USTRUCT()
+struct FNiagaraScriptExecutionParameterStore : public FNiagaraParameterStore
+{
+	GENERATED_USTRUCT_BODY()
+public:
+	FNiagaraScriptExecutionParameterStore();
+	FNiagaraScriptExecutionParameterStore(const FNiagaraParameterStore& Other);
+	FNiagaraScriptExecutionParameterStore& operator=(const FNiagaraParameterStore& Other);
+
+	//TODO: These function can probably go away entirely when we replace the FNiagaraParameters and DataInterface info in the script with an FNiagaraParameterStore.
+	//Special care with prev params and internal params will have to be taken there.
+	/** Call this init function if you are using a Niagara parameter store within a UNiagaraScript.*/
+	void InitFromOwningScript(UNiagaraScript* Script, ENiagaraSimTarget SimTarget, bool bNotifyAsDirty);
+	/** Call this init function if you are using a Niagara parameter store within an FNiagaraScriptExecutionContext.*/
+	void InitFromOwningContext(UNiagaraScript* Script, ENiagaraSimTarget SimTarget, bool bNotifyAsDirty);
+	void AddScriptParams(UNiagaraScript* Script, ENiagaraSimTarget SimTarget, bool bTriggerRebind);
+	void CopyCurrToPrev();
+
+	bool AddParameter(const FNiagaraVariable& Param, bool bInitInterfaces = true, bool bTriggerRebind = true)
+	{
+		if (FNiagaraParameterStore::AddParameter(Param, bInitInterfaces, bTriggerRebind))
+		{
+			AddPaddedParamSize(Param.GetType(), IndexOf(Param));
+			return true;
+		}
+		return false;
+	}
+
+	bool RemoveParameter(FNiagaraVariable& Param)
+	{
+		check(0);//Not allowed to remove parameters from an execution store as it will adjust the table layout mess up the 
+		return false;
+	}
+
+	void RenameParameter(FNiagaraVariable& Param, FName NewName)
+	{
+		check(0);//Can't rename parameters for an execution store.
+	}
+
+	// Just the external parameters, not previous or internal...
+	uint32 GetExternalParameterSize() { return ParameterSize; }
+
+	// The entire buffer padded out by the required alignment of the types..
+	uint32 GetPaddedParameterSizeInBytes() { return PaddedParameterSize; }
+
+	// Helper that converts the data from the base type array internally into the padded out renderer-ready format.
+	void CopyParameterDataToPaddedBuffer(uint8* InTargetBuffer, uint32 InTargetBufferSizeInBytes);
+
+	void Clear()
+	{
+		Empty();
+		PaddingInfo.Empty();
+		PaddedParameterSize = 0;
+		bInitialized = false;
+	}
+
+	bool IsInitialized() const {return bInitialized;}
+
+protected:
+	void AddPaddedParamSize(const FNiagaraTypeDefinition& InParamType, uint32 InOffset);
+
+private:
+
+	/** Size of the parameter data not including prev frame values or internal constants. Allows copying into previous parameter values for interpolated spawn scripts. */
+	UPROPERTY()
+	int32 ParameterSize;
+
+	UPROPERTY()
+	uint32 PaddedParameterSize;
+
+	static void GenerateLayoutInfoInternal(TArray<FNiagaraScriptExecutionPaddingInfo>& Members, uint32& NextMemberOffset, const UStruct* InSrcStruct, uint32 InSrcOffset);
+
+	UPROPERTY()
+	TArray<FNiagaraScriptExecutionPaddingInfo> PaddingInfo;
+
+	UPROPERTY()
+	uint8 bInitialized : 1;
+};
\ No newline at end of file
diff --git a/Engine/Plugins/FX/Niagara/Source/Niagara/Public/NiagaraSpriteRendererProperties.h b/Engine/Plugins/FX/Niagara/Source/Niagara/Public/NiagaraSpriteRendererProperties.h
--- a/Engine/Plugins/FX/Niagara/Source/Niagara/Public/NiagaraSpriteRendererProperties.h
+++ b/Engine/Plugins/FX/Niagara/Source/Niagara/Public/NiagaraSpriteRendererProperties.h
@@ -4,21 +4,11 @@
 
 #include "CoreMinimal.h"
 #include "UObject/ObjectMacros.h"
+#include "NiagaraCommon.h"
 #include "NiagaraRendererProperties.h"
 #include "NiagaraSpriteRendererProperties.generated.h"
 
 
-UENUM()
-enum class ENiagaraSortMode : uint8
-{
-	/** Perform no additional sorting prior to rendering.*/
-	SortNone,
-	/** Sort by depth to the camera's near plane.*/
-	SortViewDepth,
-	/** Sort by distance to the camera's origin.*/
-	SortViewDistance
-};
-
 /** This enum decides how a sprite particle will orient its "up" axis. Must keep these in sync with NiagaraSpriteVertexFactory.ush*/
 UENUM()
 enum class ENiagaraSpriteAlignment : uint8
@@ -57,16 +47,25 @@ public:
 
 	UNiagaraSpriteRendererProperties();
 
+	virtual void PostInitProperties() override;
+
+	static void InitCDOPropertiesAfterModuleStartup();
+
 	//~ UNiagaraRendererProperties interface
 	virtual NiagaraRenderer* CreateEmitterRenderer(ERHIFeatureLevel::Type FeatureLevel) override;
 	virtual void GetUsedMaterials(TArray<UMaterialInterface*>& OutMaterials) const override;
+	virtual bool IsSimTargetSupported(ENiagaraSimTarget InSimTarget) const override { return true; };
 #if WITH_EDITORONLY_DATA
+	virtual void PostEditChangeProperty(struct FPropertyChangedEvent& PropertyChangedEvent) override;
 	virtual bool IsMaterialValidForRenderer(UMaterial* Material, FText& InvalidMessage) override;
 	virtual void FixMaterial(UMaterial* Material) override;
 	virtual const TArray<FNiagaraVariable>& GetRequiredAttributes() override;
 	virtual const TArray<FNiagaraVariable>& GetOptionalAttributes() override;
 #endif // WITH_EDITORONLY_DATA
 
+	// TODO: once we support cutouts, need to change this
+	virtual uint32 GetNumIndicesPerInstance() { return 6; }
+
 	/** The material used to render the particle. Note that it must have the Use with Niagara Sprites flag checked.*/
 	UPROPERTY(EditAnywhere, Category = "Sprite Rendering")
 	UMaterialInterface* Material;
@@ -88,21 +87,25 @@ public:
 	FVector2D PivotInUVSpace;
 
 	/** Determines how we sort the particles prior to rendering.*/
-	UPROPERTY(EditAnywhere, Category = "Sprite Rendering")
+	UPROPERTY(EditAnywhere, Category = "Sorting")
 	ENiagaraSortMode SortMode;
 	
 	/** When using SubImage lookups for particles, this variable contains the number of columns in X and the number of rows in Y.*/
-	UPROPERTY(EditAnywhere, Category = "Sprite Rendering")
+	UPROPERTY(EditAnywhere, Category = "SubUV")
 	FVector2D SubImageSize;
 
 	/** If true, blends the sub-image UV lookup with its next adjacent member using the fractional part of the SubImageIndex float value as the linear interpolation factor.*/
-	UPROPERTY(EditAnywhere, Category = "Sprite Rendering", meta = (DisplayName = "Sub UV Blending Enabled"))
+	UPROPERTY(EditAnywhere, Category = "SubUV", meta = (DisplayName = "Sub UV Blending Enabled"))
 	uint32 bSubImageBlend : 1;
 
 	/** If true, removes the HMD view roll (e.g. in VR) */
 	UPROPERTY(EditAnywhere, Category = "Sprite Rendering", meta = (DisplayName = "Remove HMD Roll"))
 	uint32 bRemoveHMDRollInVR : 1;
 
+	/** If true, the particles are only sorted when using a translucent material. */
+	UPROPERTY(EditAnywhere, Category = "Sorting")
+	uint32 bSortOnlyWhenTranslucent : 1;
+
 	/** The distance at which FacingCameraDistanceBlend	is fully FacingCamera */
 	UPROPERTY(EditAnywhere, Category = "Sprite Rendering", meta = (UIMin = "0"))
 	float MinFacingCameraBlendDistance;
@@ -110,6 +113,75 @@ public:
 	/** The distance at which FacingCameraDistanceBlend	is fully FacingCameraPosition */
 	UPROPERTY(EditAnywhere, Category = "Sprite Rendering", meta = (UIMin = "0"))
 	float MaxFacingCameraBlendDistance;
+
+	/** Which attribute should we use for position when generating sprites?*/
+	UPROPERTY(EditAnywhere, Category = "Bindings")
+	FNiagaraVariableAttributeBinding PositionBinding;
+
+	/** Which attribute should we use for color when generating sprites?*/
+	UPROPERTY(EditAnywhere, Category = "Bindings")
+	FNiagaraVariableAttributeBinding ColorBinding;
+
+	/** Which attribute should we use for velocity when generating sprites?*/
+	UPROPERTY(EditAnywhere, Category = "Bindings")
+	FNiagaraVariableAttributeBinding VelocityBinding;
+
+	/** Which attribute should we use for sprite rotation (in degrees) when generating sprites?*/
+	UPROPERTY(EditAnywhere, Category = "Bindings")
+	FNiagaraVariableAttributeBinding SpriteRotationBinding;
+
+	/** Which attribute should we use for sprite size when generating sprites?*/
+	UPROPERTY(EditAnywhere, Category = "Bindings")
+	FNiagaraVariableAttributeBinding SpriteSizeBinding;
+
+	/** Which attribute should we use for sprite facing when generating sprites?*/
+	UPROPERTY(EditAnywhere, Category = "Bindings")
+	FNiagaraVariableAttributeBinding SpriteFacingBinding;
+
+	/** Which attribute should we use for sprite alignment when generating sprites?*/
+	UPROPERTY(EditAnywhere, Category = "Bindings")
+	FNiagaraVariableAttributeBinding SpriteAlignmentBinding;
+
+	/** Which attribute should we use for sprite sub-image indexing when generating sprites?*/
+	UPROPERTY(EditAnywhere, Category = "Bindings")
+	FNiagaraVariableAttributeBinding SubImageIndexBinding;
+
+	/** Which attribute should we use for dynamic material parameters when generating sprites?*/
+	UPROPERTY(EditAnywhere, Category = "Bindings")
+	FNiagaraVariableAttributeBinding DynamicMaterialBinding;
+
+	/** Which attribute should we use for dynamic material parameters when generating sprites?*/
+	UPROPERTY(EditAnywhere, Category = "Bindings")
+	FNiagaraVariableAttributeBinding DynamicMaterial1Binding;
+
+	/** Which attribute should we use for dynamic material parameters when generating sprites?*/
+	UPROPERTY(EditAnywhere, Category = "Bindings")
+	FNiagaraVariableAttributeBinding DynamicMaterial2Binding;
+
+	/** Which attribute should we use for dynamic material parameters when generating sprites?*/
+	UPROPERTY(EditAnywhere, Category = "Bindings")
+	FNiagaraVariableAttributeBinding DynamicMaterial3Binding;
+
+	/** Which attribute should we use for camera offset when generating sprites?*/
+	UPROPERTY(EditAnywhere, Category = "Bindings")
+	FNiagaraVariableAttributeBinding CameraOffsetBinding;
+
+	/** Which attribute should we use for UV scale when generating sprites?*/
+	UPROPERTY(EditAnywhere, Category = "Bindings")
+	FNiagaraVariableAttributeBinding UVScaleBinding;
+
+	/** Which attribute should we use for material randoms when generating sprites?*/
+	UPROPERTY(EditAnywhere, Category = "Bindings")
+	FNiagaraVariableAttributeBinding MaterialRandomBinding;
+
+	/** Which attribute should we use for custom sorting? */
+	UPROPERTY(EditAnywhere, Category = "Bindings")
+	FNiagaraVariableAttributeBinding CustomSortingBinding;
+
+	UPROPERTY(Transient)
+	int32 SyncId;
+
+	void InitBindings();
 };
 
 
diff --git a/Engine/Plugins/FX/Niagara/Source/Niagara/Public/NiagaraSystemInstance.h b/Engine/Plugins/FX/Niagara/Source/Niagara/Public/NiagaraSystemInstance.h
--- a/Engine/Plugins/FX/Niagara/Source/Niagara/Public/NiagaraSystemInstance.h
+++ b/Engine/Plugins/FX/Niagara/Source/Niagara/Public/NiagaraSystemInstance.h
@@ -15,6 +15,7 @@
 class FNiagaraWorldManager;
 class UNiagaraComponent;
 class FNiagaraSystemInstance;
+class FNiagaraSystemSimulation;
 
 class NIAGARA_API FNiagaraSystemInstance 
 {
@@ -28,19 +29,38 @@ public:
 #endif
 
 public:
+
+	/** Defines modes for resetting the System instance. */
+	enum class EResetMode
+	{
+		/** Resets the System instance and simulations. */
+		ResetAll,
+		/** Resets the System instance but not the simualtions */
+		ResetSystem,
+		/** Full reinitialization of the system and emitters.  */
+		ReInit,
+		/** No reset */
+		None
+	};
+
+
 	/** Creates a new niagara System instance with the supplied component. */
 	explicit FNiagaraSystemInstance(UNiagaraComponent* InComponent);
 
 	/** Cleanup*/
 	virtual ~FNiagaraSystemInstance();
 
+	void Cleanup();
+
 	/** Initializes this System instance to simulate the supplied System. */
-	void Init(TSharedRef<class FNiagaraSystemSimulation, ESPMode::ThreadSafe> InSystemSimulation, bool bInForceSolo=false);
+	void Init(UNiagaraSystem* InSystem, bool bInForceSolo=false);
 
-	void Activate(bool bReset = false);
+	void Activate(EResetMode InResetMode = EResetMode::ResetAll);
 	void Deactivate(bool bImmediate = false);
 	void Complete();
 
+	void SetSolo(bool bInSolo);
+
 	//void RebindParameterCollection(UNiagaraParameterCollectionInstance* OldInstance, UNiagaraParameterCollectionInstance* NewInstance);
 	void BindParameters();
 	void UnbindParameters();
@@ -49,30 +69,27 @@ public:
 	
 	FNiagaraWorldManager* GetWorldManager()const;
 
-	/** Defines modes for resetting the System instance. */
-	enum class EResetMode
-	{
-		/** Resets the System instance and simulations. */
-		ResetAll,
-		/** Resets the System instance but not the simualtions */
-		ResetSystem,
-		/** Full reinitialization of the system and emitters.  */
-		ReInit
-	};
-	
 	/** Requests the the simulation be reset on the next tick. */
-	void Reset(EResetMode Mode);
+	void Reset(EResetMode Mode, bool bBindParams = false);
 
 	void ComponentTick(float DeltaSeconds);
 	void PreSimulateTick(float DeltaSeconds);
 	void PostSimulateTick(float DeltaSeconds);
+	void FinalizeTick(float DeltaSeconds);
 	/** Handles completion of the system and returns true if the system is complete. */
 	bool HandleCompletion();
 
-	ENiagaraExecutionState GetExecutionState() { return ExecutionState; }
-	void SetExecutionState(ENiagaraExecutionState InState);
+	/** Perform per-tick updates on data interfaces that need it. This can cause systems to complete so cannot be parallelized. */
+	void TickDataInterfaces(float DeltaSeconds, bool bPostSimulate);
+
+	ENiagaraExecutionState GetRequestedExecutionState() { return RequestedExecutionState; }
+	void SetRequestedExecutionState(ENiagaraExecutionState InState);
 
-	FORCEINLINE bool IsComplete()const { return ExecutionState == ENiagaraExecutionState::Complete || ExecutionState == ENiagaraExecutionState::Disabled; }
+	ENiagaraExecutionState GetActualExecutionState() { return ActualExecutionState; }
+	void SetActualExecutionState(ENiagaraExecutionState InState);
+
+	FORCEINLINE bool IsComplete()const { return ActualExecutionState == ENiagaraExecutionState::Complete || ActualExecutionState == ENiagaraExecutionState::Disabled; }
+	FORCEINLINE bool IsDisabled()const { return ActualExecutionState == ENiagaraExecutionState::Disabled; }
 
 	/** Gets the simulation for the supplied emitter handle. */
 	TSharedPtr<FNiagaraEmitterInstance> GetSimulationForHandle(const FNiagaraEmitterHandle& EmitterHandle);
@@ -82,6 +99,8 @@ public:
 	FORCEINLINE TArray<TSharedRef<FNiagaraEmitterInstance> > &GetEmitters()	{ return Emitters; }
 	FORCEINLINE FBox &GetSystemBounds()	{ return SystemBounds;  }
 
+	FNiagaraEmitterInstance* GetEmitterByID(FGuid InID);
+
 	FORCEINLINE bool IsSolo()const { return bSolo; }
 
 	//TEMPORARY. We wont have a single set of parameters when we're executing system scripts.
@@ -110,9 +129,6 @@ public:
 
 	FName GetIDName() { return IDName; }
 
-	/** Queue up the data sources to have PrepareForSimulation called on them next tick.*/
-	void ReinitializeDataInterfaces();
-
 	/** Returns the instance data for a particular interface for this System. */
 	FORCEINLINE void* FindDataInterfaceInstanceData(UNiagaraDataInterface* Interface) 
 	{
@@ -137,22 +153,45 @@ public:
 	
 	FORCEINLINE TSharedPtr<FNiagaraSystemSimulation, ESPMode::ThreadSafe> GetSystemSimulation()const
 	{
-		check(SystemSimulation.IsValid());
 		return SystemSimulation; 
 	}
 
-	FORCEINLINE bool IsActivated() const { return ExecutionState == ENiagaraExecutionState::Active; }
-
 	bool IsReadyToRun() const;
 
 	/** Index of this instance in the system simulation. */
 	int32 SystemInstanceIndex;
 
-	FORCEINLINE bool IsDeactivatedDueToError() const { return bError; }
-
 	FORCEINLINE bool HasTickingEmitters()const { return bHasTickingEmitters; }
 
 	UNiagaraParameterCollectionInstance* GetParameterCollectionInstance(UNiagaraParameterCollection* Collection);
+
+	/** 
+	Manually advances this system's simulation by the specified number of ticks and tick delta. 
+	To be advanced in this way a system must be in solo mode or moved into solo mode which will add additional overhead.
+	*/
+	void AdvanceSimulation(int32 TickCount, float TickDeltaSeconds);
+
+#if WITH_EDITORONLY_DATA
+	/** Request that this simulation capture a frame. Cannot capture if disabled or already completed.*/
+	bool RequestCapture(const FGuid& RequestId);
+
+	/** Poll for previous frame capture requests. Once queried and bool is returned, the results are cleared from this system instance.*/
+	bool QueryCaptureResults(const FGuid& RequestId, TArray<TSharedPtr<struct FNiagaraScriptDebuggerInfo, ESPMode::ThreadSafe>>& OutCaptureResults);
+
+	/** Only call from within the script execution states. Value is null if not capturing a frame.*/
+	TArray<TSharedPtr<struct FNiagaraScriptDebuggerInfo, ESPMode::ThreadSafe>>* GetActiveCaptureResults();
+
+	/** Only call from within the script execution states. Does nothing if not capturing a frame.*/
+	void FinishCapture();
+
+	/** Only call from within the script execution states. Value is false if not capturing a frame.*/
+	bool ShouldCaptureThisFrame() const;
+
+	/** Only call from within the script execution states. Value is nullptr if not capturing a frame.*/
+	FNiagaraScriptDebuggerInfo* GetActiveCaptureWrite(const FName& InHandleName, ENiagaraScriptUsage InUsage, const FGuid& InUsageId);
+
+#endif
+
 private:
 	/** Builds the emitter simulations. */
 	void InitEmitters();
@@ -169,10 +208,7 @@ private:
 	void UpdateProxy(TArray<NiagaraRenderer*>& InRenderers);
 
 	/** Call PrepareForSImulation on each data source from the simulations and determine which need per-tick updates.*/
-	void InitDataInterfaces();
-	
-	/** Perform per-tick updates on data interfaces that need it.*/
-	void TickDataInterfaces(float DeltaSeconds, bool bPostSimulate);
+	void InitDataInterfaces();	
 
 	void TickInstanceParameters(float DeltaSeconds);
 
@@ -197,11 +233,18 @@ private:
 	FOnDestroyed OnDestroyedDelegate;
 #endif
 
+#if WITH_EDITORONLY_DATA
+	TSharedPtr<TArray<TSharedPtr<struct FNiagaraScriptDebuggerInfo, ESPMode::ThreadSafe>>, ESPMode::ThreadSafe> CurrentCapture;
+	TSharedPtr<FGuid, ESPMode::ThreadSafe> CurrentCaptureGuid;
+	bool bWasSoloPriorToCaptureRequest;
+	TMap<FGuid, TSharedPtr<TArray<TSharedPtr<struct FNiagaraScriptDebuggerInfo, ESPMode::ThreadSafe>>, ESPMode::ThreadSafe> > CapturedFrames;
+#endif
+
 	FGuid ID;
 	FName IDName;
 	
 	/** Per instance data for any data interfaces requiring it. */
-	TArray<uint8> DataInterfaceInstanceData;
+	TArray<uint8, TAlignedHeapAllocator<16>> DataInterfaceInstanceData;
 
 	/** Map of data interfaces to their instance data. */
 	TMap<TWeakObjectPtr<UNiagaraDataInterface>, int32> DataInterfaceInstanceDataOffsets;
@@ -225,12 +268,18 @@ private:
 
 	FNiagaraParameterDirectBinding<float> OwnerDeltaSecondsParam;
 	FNiagaraParameterDirectBinding<float> OwnerInverseDeltaSecondsParam;
+	FNiagaraParameterDirectBinding<float> OwnerEngineTimeParam;
+	FNiagaraParameterDirectBinding<float> OwnerEngineRealtimeParam;
+	FNiagaraParameterDirectBinding<float> SystemAgeParam;
+
 	FNiagaraParameterDirectBinding<float> OwnerMinDistanceToCameraParam;
 	FNiagaraParameterDirectBinding<int32> SystemNumEmittersParam;
 	FNiagaraParameterDirectBinding<int32> SystemNumEmittersAliveParam;
 
 	FNiagaraParameterDirectBinding<float> SystemTimeSinceRenderedParam;
 
+	FNiagaraParameterDirectBinding<int32> OwnerExecutionStateParam;
+
 	TArray<FNiagaraParameterDirectBinding<int32>> ParameterNumParticleBindings;
 
 	/** Indicates whether this instance must update itself rather than being batched up as most instances are. */
@@ -238,16 +287,14 @@ private:
 	uint32 bForceSolo : 1;
 
 	uint32 bPendingSpawn : 1;
-
-	/** Disable ticking and rendering if there was some serious error. */
-	uint32 bError : 1;
-
-	/** Notifier that data interfaces need reinitialization next tick.*/
-	uint32 bDataInterfacesNeedInit : 1;
+	uint32 bNotifyOnCompletion : 1;
 
 	/** If this instance has any currently ticking emitters. If false, allows us to skip some work. */
 	uint32 bHasTickingEmitters : 1;
 
-	/* System tick state */
-	ENiagaraExecutionState ExecutionState;
+	/* Execution state requested by external code/BPs calling Activate/Deactivate. */
+	ENiagaraExecutionState RequestedExecutionState;
+
+	/** Copy of simulations internal state so that it can be passed to emitters etc. */
+	ENiagaraExecutionState ActualExecutionState;
 };
diff --git a/Engine/Plugins/FX/Niagara/Source/Niagara/Public/NiagaraSystemSimulation.h b/Engine/Plugins/FX/Niagara/Source/Niagara/Public/NiagaraSystemSimulation.h
--- a/Engine/Plugins/FX/Niagara/Source/Niagara/Public/NiagaraSystemSimulation.h
+++ b/Engine/Plugins/FX/Niagara/Source/Niagara/Public/NiagaraSystemSimulation.h
@@ -65,23 +65,19 @@ struct FNiagaraParameterStoreToDataSetBinding
 	FORCEINLINE_DEBUGGABLE void DataSetToParameterStore(FNiagaraParameterStore& ParameterStore, FNiagaraDataSet& DataSet, int32 DataSetInstanceIndex)
 	{
 		FNiagaraDataBuffer& CurrBuffer = DataSet.CurrData();
-		TArray<uint8> ParameterDataArrayCopy = ParameterStore.GetParameterDataArray();
-		uint8* ParameterData = ParameterDataArrayCopy.GetData();
 
 		for (const FDataOffsets& DataOffsets : FloatOffsets)
 		{
-			float* ParamPtr = (float*)(ParameterData + DataOffsets.ParameterOffset);
 			float* DataSetPtr = CurrBuffer.GetInstancePtrFloat(DataOffsets.DataSetComponentOffset, DataSetInstanceIndex);
-			*ParamPtr = *DataSetPtr;
+			ParameterStore.SetParameterByOffset(DataOffsets.ParameterOffset, *DataSetPtr);
 		}
 		for (const FDataOffsets& DataOffsets : Int32Offsets)
 		{
-			int32* ParamPtr = (int32*)(ParameterData + DataOffsets.ParameterOffset);
 			int32* DataSetPtr = CurrBuffer.GetInstancePtrInt32(DataOffsets.DataSetComponentOffset, DataSetInstanceIndex);
-			*ParamPtr = *DataSetPtr;
+			ParameterStore.SetParameterByOffset(DataOffsets.ParameterOffset, *DataSetPtr);
 		}
 
-		ParameterStore.SetParameterDataArray(ParameterDataArrayCopy);
+		ParameterStore.OnParameterChange();
 	}
 
 	FORCEINLINE_DEBUGGABLE void ParameterStoreToDataSet(FNiagaraParameterStore& ParameterStore, FNiagaraDataSet& DataSet, int32 DataSetInstanceIndex)
@@ -109,42 +105,35 @@ class FNiagaraSystemSimulation
 {
 public:
 	~FNiagaraSystemSimulation();
-	bool Init(UNiagaraSystem* InSystem, UWorld* InWorld);
+	bool Init(UNiagaraSystem* InSystem, UWorld* InWorld, FNiagaraSystemInstance* InSoloSystemInstance=nullptr);
 	void Destroy();
 	bool Tick(float DeltaSeconds);
 
 	void RemoveInstance(FNiagaraSystemInstance* Instance);
 	void AddInstance(FNiagaraSystemInstance* Instance);
 
-	void ResetSolo(FNiagaraSystemInstance* Instance);
-	void TickSolo(FNiagaraSystemInstance* Instance);
-	void RemoveSolo(FNiagaraSystemInstance* Instance);
-
-	void TickSoloDataSet();
-
 	FORCEINLINE UNiagaraSystem* GetSystem()const { return WeakSystem.Get(); }
 
-	const TArray<UNiagaraDataInterface*>& GetSoloDataInterfacesSpawn();
-	const TArray<UNiagaraDataInterface*>& GetSoloDataInterfacesUpdate();
-
 	UNiagaraParameterCollectionInstance* GetParameterCollectionInstance(UNiagaraParameterCollection* Collection);
+
+	FNiagaraParameterStore& GetScriptDefinedDataInterfaceParameters();
+
+	/** Transfers a system instance from SourceSimulation. */
+	void TransferInstance(FNiagaraSystemSimulation* SourceSimulation, FNiagaraSystemInstance* SystemInst);
 protected:
 
 	/** System of instances being simulated.  We use a weak object ptr here because once the last referencing object goes away this system may be come invalid at runtime. */
 	TWeakObjectPtr<UNiagaraSystem> WeakSystem;
 
+	/** The parent system instance if this simulation is  */
+	FNiagaraSystemInstance* SoloSystemInstance;
+
 	/** World this system simulation belongs to. */
 	UWorld* World;
 
 	/** Main dataset containing system instance attribute data. */
 	FNiagaraDataSet DataSet;
-
-	/**
-	Dataset for system instances doing solo simulation i.e. not batched like most systems.
-	This can be required if strict ordering is needed or the system/emitter scripts use a data interface overridden by the component.
-	*/
-	FNiagaraDataSet DataSetSolo;
-
+	
 	/**
 	As there's a 1 to 1 relationship between system instance and their execution in this simulation we must pull all that instances parameters into a dataset for simulation.
 	In some cases this might be a big waste of memory as there'll be duplicated data from a parameter store that's shared across all instances.
@@ -156,15 +145,10 @@ protected:
 	FNiagaraScriptExecutionContext SpawnExecContext;
 	FNiagaraScriptExecutionContext UpdateExecContext;
 
-	FNiagaraScriptExecutionContext SpawnExecContextSolo;
-	FNiagaraScriptExecutionContext UpdateExecContextSolo;
-
 	/** Bindings that pull per component parameters into the spawn parameter dataset. */
 	FNiagaraParameterStoreToDataSetBinding SpawnInstanceParameterToDataSetBinding;
 	/** Bindings that pull per component parameters into the update parameter dataset. */
 	FNiagaraParameterStoreToDataSetBinding UpdateInstanceParameterToDataSetBinding;
-	FNiagaraParameterStoreToDataSetBinding SpawnConstantParameterToDataSetBinding;
-	FNiagaraParameterStoreToDataSetBinding UpdateConstantParameterToDataSetBinding;
 
 	/** Binding to push system attributes into each emitter spawn parameters. */
 	TArray<FNiagaraParameterStoreToDataSetBinding> DataSetToEmitterSpawnParameters;
@@ -173,24 +157,37 @@ protected:
 	/** Binding to push system attributes into each emitter event parameters. */
 	TArray<TArray<FNiagaraParameterStoreToDataSetBinding>> DataSetToEmitterEventParameters;
 
+	/** Direct bindings for Engine variables in System Spawn and Update scripts. */
+	FNiagaraParameterDirectBinding<float> SpawnDeltaTimeParam;
+	FNiagaraParameterDirectBinding<float> UpdateDeltaTimeParam;
+
+	FNiagaraParameterDirectBinding<float> SpawnInvDeltaTimeParam;
+	FNiagaraParameterDirectBinding<float> UpdateInvDeltaTimeParam;
+	
+	FNiagaraParameterDirectBinding<int32> SpawnNumSystemInstancesParam;
+	FNiagaraParameterDirectBinding<int32> UpdateNumSystemInstancesParam;
+
+	FNiagaraParameterDirectBinding<float> SpawnGlobalSpawnCountScaleParam;
+	FNiagaraParameterDirectBinding<float> UpdateGlobalSpawnCountScaleParam;
+
+	FNiagaraParameterDirectBinding<float> SpawnGlobalSystemCountScaleParam;
+	FNiagaraParameterDirectBinding<float> UpdateGlobalSystemCountScaleParam;
+
+
 	/** System instances that have been spawned and are now simulating. */
 	TArray<FNiagaraSystemInstance*> SystemInstances;
 	/** System instances that are pending to be spawned. */
 	TArray<FNiagaraSystemInstance*> PendingSystemInstances;
 
-	TArray<FNiagaraSystemInstance*> SoloSystemInstances;
+	TArray<TArray<FNiagaraDataSetAccessor<FNiagaraSpawnInfo>>> EmitterSpawnInfoAccessors;
 
 	void InitBindings(FNiagaraSystemInstance* SystemInst);
 
 	FNiagaraDataSetAccessor<int32> SystemExecutionStateAccessor;
-	TArray<TArray<FNiagaraDataSetAccessor<FNiagaraSpawnInfo>>> EmitterSpawnInfoAccessors;
 	TArray<FNiagaraDataSetAccessor<int32>> EmitterExecutionStateAccessors;
 
-	//Annoying duplicates required because these access the solo data set. When I rework the where the layout data for parameters and data sets live then these can go away.
-	FNiagaraDataSetAccessor<int32> SoloSystemExecutionStateAccessor;
-	TArray<TArray<FNiagaraDataSetAccessor<FNiagaraSpawnInfo>>> SoloEmitterSpawnInfoAccessors;
-	TArray<FNiagaraDataSetAccessor<int32>> SoloEmitterExecutionStateAccessors;
-
 	uint32 bCanExecute : 1;
-	uint32 bCanExecuteSolo : 1;
+
+	/** A parameter store which contains the data interfaces parameters which were defined by the scripts. */
+	FNiagaraParameterStore ScriptDefinedDataInterfaceParameters;
 };
diff --git a/Engine/Plugins/FX/Niagara/Source/Niagara/Public/NiagaraTypes.h b/Engine/Plugins/FX/Niagara/Source/Niagara/Public/NiagaraTypes.h
--- a/Engine/Plugins/FX/Niagara/Source/Niagara/Public/NiagaraTypes.h
+++ b/Engine/Plugins/FX/Niagara/Source/Niagara/Public/NiagaraTypes.h
@@ -147,6 +147,34 @@ struct FNiagaraSpawnInfo
 	float IntervalDt;
 };
 
+USTRUCT(Blueprintable, meta = (DisplayName = "Niagara ID"))
+struct FNiagaraID
+{
+	GENERATED_USTRUCT_BODY()
+
+	/** 
+	Index in the indirection table for this particle. Allows fast access to this particles data.
+	Is always unique among currently living particles but will be reused after the particle dies.
+	*/
+	UPROPERTY(EditAnywhere, Category = ID)
+	int32 Index;
+
+	/** 
+	A unique tag for when this ID was acquired. 
+	Allows us to differentiate between particles when one dies and another reuses it's Index.
+	*/
+	UPROPERTY(EditAnywhere, Category = ID)
+	int32 AcquireTag;
+
+	bool operator==(const FNiagaraID& Other)const { return Index == Other.Index && AcquireTag == Other.AcquireTag; }
+	bool operator<(const FNiagaraID& Other)const { return Index < Other.Index && AcquireTag < Other.AcquireTag; }
+};
+
+FORCEINLINE uint32 GetTypeHash(const FNiagaraID& ID)
+{
+	return HashCombine(GetTypeHash(ID.Index), GetTypeHash(ID.AcquireTag));
+}
+
 /** Information about how this type should be laid out in an FNiagaraDataSet */
 struct FNiagaraTypeLayoutInfo
 {
@@ -211,6 +239,7 @@ class NIAGARA_API FNiagaraTypeHelper
 {
 public:
 	static FString ToString(const uint8* ValueData, const UScriptStruct* Struct);
+
 };
 
 /** Defines different modes for selecting the output numeric type of a function or operation based on the types of the inputs. */
@@ -227,6 +256,19 @@ enum class ENiagaraNumericOutputTypeSelectionMode : uint8
 	Scalar,
 };
 
+/** 
+The source from which a script execution state was set. Used to allow scalability etc to change the state but only if the state has not been defined by something with higher precedence. 
+If this changes, all scripts must be recompiled by bumping the NiagaraCustomVersion
+*/
+UENUM()
+enum class ENiagaraExecutionStateSource : uint32
+{
+	Scalability, //State set by Scalability logic. Lowest precedence.
+	Internal, //Misc internal state. For example becoming inactive after we finish our set loops.
+	Owner, //State requested by the owner. Takes precedence over everything but internal completion logic.
+	InternalCompletion, // Internal completion logic. Has to take highest precedence for completion to be ensured.
+};
+
 UENUM()
 enum class ENiagaraExecutionState : uint32
 {
@@ -240,12 +282,20 @@ enum class ENiagaraExecutionState : uint32
 	Complete,
 	/** Emitter only. Emitter is disabled. Will not tick or render again until a full re initialization of the system. */
 	Disabled,
+
+	// insert new states before
+	Num
 };
 
 USTRUCT()
 struct NIAGARA_API FNiagaraVariableMetaData
 {
 	GENERATED_USTRUCT_BODY()
+public:
+	FNiagaraVariableMetaData()
+		:EditorSortPriority(0)
+	{
+	}
 public:
 	UPROPERTY(EditAnywhere, Category = "Variable")
 	TMap<FName, FString> PropertyMetaData;
@@ -253,6 +303,12 @@ public:
 	UPROPERTY(EditAnywhere, Category = "Variable")
 	FText Description;
 
+	UPROPERTY(EditAnywhere, Category = "Variable")
+	FText CategoryName;
+
+	UPROPERTY(EditAnywhere, Category = "Variable", meta = (ToolTip = "Affects the sort order in the editor stacks. Use a smaller number to push it to the top. Defaults to zero."))
+	int32 EditorSortPriority;
+
 	UPROPERTY()
 	int32 CallSortPriority;
 
@@ -267,32 +323,32 @@ struct NIAGARA_API FNiagaraTypeDefinition
 public:
 
 	// Construct blank raw type definition 
-	FNiagaraTypeDefinition(const UClass *ClassDef)
-		: Struct(ClassDef), Enum(nullptr)
+	FNiagaraTypeDefinition(UClass *ClassDef)
+		: Struct(ClassDef), Enum(nullptr), Size(INDEX_NONE), Alignment(INDEX_NONE)
 	{
-		check(Struct != nullptr);
+		checkSlow(Struct != nullptr);
 	}
 
-	FNiagaraTypeDefinition(const UEnum *EnumDef)
-		: Struct(IntStruct), Enum(EnumDef)
+	FNiagaraTypeDefinition(UEnum *EnumDef)
+		: Struct(IntStruct), Enum(EnumDef), Size(INDEX_NONE), Alignment(INDEX_NONE)
 	{
-		check(Struct != nullptr);
+		checkSlow(Struct != nullptr);
 	}
 
-	FNiagaraTypeDefinition(const UScriptStruct *StructDef)
-		: Struct(StructDef), Enum(nullptr)
+	FNiagaraTypeDefinition(UScriptStruct *StructDef)
+		: Struct(StructDef), Enum(nullptr), Size(INDEX_NONE), Alignment(INDEX_NONE)
 	{
-		check(Struct != nullptr);
+		checkSlow(Struct != nullptr);
 	}
 
 	FNiagaraTypeDefinition(const FNiagaraTypeDefinition &Other)
-		: Struct(Other.Struct), Enum(Other.Enum)
+		: Struct(Other.Struct), Enum(Other.Enum), Size(INDEX_NONE), Alignment(INDEX_NONE)
 	{
 	}
 
 	// Construct a blank raw type definition
 	FNiagaraTypeDefinition()
-		: Struct(nullptr), Enum(nullptr)
+		: Struct(nullptr), Enum(nullptr), Size(INDEX_NONE), Alignment(INDEX_NONE)
 	{}
 
 	bool operator !=(const FNiagaraTypeDefinition &Other) const
@@ -338,23 +394,23 @@ public:
 		return GetStruct()->GetName();
 	}
 
-	const UStruct* GetStruct()const
+	UStruct* GetStruct()const
 	{
 		return Struct;
 	}
 
-	const UScriptStruct* GetScriptStruct()const
+	UScriptStruct* GetScriptStruct()const
 	{
 		return Cast<UScriptStruct>(Struct);
 	}
-		
+
 	/** Gets the class ptr for this type if it is a class. */
-	const UClass* GetClass()const
+	UClass* GetClass()const
 	{
 		return Cast<UClass>(Struct);
 	}
 
-	const UEnum* GetEnum() const
+	UEnum* GetEnum() const
 	{
 		return Enum;
 	}
@@ -365,34 +421,42 @@ public:
 	
 	int32 GetSize()const
 	{
-		checkf(IsValid(), TEXT("Type definition is not valid."));
-		if (GetClass())
-		{
-			return 0;//TODO: sizeof(void*);//If we're a class then we allocate space for the user to instantiate it. This and stopping it being GCd is up to the user.
-		}
-		else
+		if (Size == INDEX_NONE)
 		{
-			return CastChecked<UScriptStruct>(Struct)->GetStructureSize();	// TODO: cache this here?
+			checkfSlow(IsValid(), TEXT("Type definition is not valid."));
+			if (GetClass())
+			{
+				Size = 0;//TODO: sizeof(void*);//If we're a class then we allocate space for the user to instantiate it. This and stopping it being GCd is up to the user.
+			}
+			else
+			{
+				Size = CastChecked<UScriptStruct>(Struct)->GetStructureSize();	// TODO: cache this here?
+			}
 		}
+		return Size;
 	}
 
 	int32 GetAlignment()const
 	{
-		checkf(IsValid(), TEXT("Type definition is not valid."));
-		if (GetClass())
+		if (Alignment == INDEX_NONE)
 		{
-			return 0;//TODO: sizeof(void*);//If we're a class then we allocate space for the user to instantiate it. This and stopping it being GCd is up to the user.
-		}
-		else
-		{
-			return CastChecked<UScriptStruct>(Struct)->GetMinAlignment();
+			checkfSlow(IsValid(), TEXT("Type definition is not valid."));
+			if (GetClass())
+			{
+				Alignment = 0;//TODO: sizeof(void*);//If we're a class then we allocate space for the user to instantiate it. This and stopping it being GCd is up to the user.
+			}
+			else
+			{
+				Alignment = CastChecked<UScriptStruct>(Struct)->GetMinAlignment();
+			}
 		}
+		return Alignment;
 	}
 
 	bool IsFloatPrimitive() const
 	{
 		return Struct == FNiagaraTypeDefinition::GetFloatStruct() || Struct == FNiagaraTypeDefinition::GetVec2Struct() || Struct == FNiagaraTypeDefinition::GetVec3Struct() || Struct == FNiagaraTypeDefinition::GetVec4Struct() ||
-			Struct == FNiagaraTypeDefinition::GetMatrix4Struct() || Struct == FNiagaraTypeDefinition::GetColorStruct();
+			Struct == FNiagaraTypeDefinition::GetMatrix4Struct() || Struct == FNiagaraTypeDefinition::GetColorStruct() || Struct == FNiagaraTypeDefinition::GetQuatStruct();
  	}
 
 	bool IsValid() const 
@@ -406,10 +470,16 @@ public:
 	In occasional situations this may be a UClass when we're dealing with DataInterface etc.
 	*/
 	UPROPERTY()
-	const UStruct *Struct;
+	UStruct *Struct;
 
 	UPROPERTY()
-	const UEnum* Enum;
+	UEnum* Enum;
+
+private:
+	mutable int16 Size;
+	mutable int16 Alignment;
+
+public:
 
 	static void Init();
 	static void RecreateUserDefinedTypeRegistry();
@@ -420,22 +490,26 @@ public:
 	static const FNiagaraTypeDefinition& GetVec3Def() { return Vec3Def; }
 	static const FNiagaraTypeDefinition& GetVec4Def() { return Vec4Def; }
 	static const FNiagaraTypeDefinition& GetColorDef() { return ColorDef; }
+	static const FNiagaraTypeDefinition& GetQuatDef() { return QuatDef; }
 	static const FNiagaraTypeDefinition& GetMatrix4Def() { return Matrix4Def; }
 	static const FNiagaraTypeDefinition& GetGenericNumericDef() { return NumericDef; }
 	static const FNiagaraTypeDefinition& GetParameterMapDef() { return ParameterMapDef; }
-
-	static const UScriptStruct* GetFloatStruct() { return FloatStruct; }
-	static const UScriptStruct* GetBoolStruct() { return BoolStruct; }
-	static const UScriptStruct* GetIntStruct() { return IntStruct; }
-	static const UScriptStruct* GetVec2Struct() { return Vec2Struct; }
-	static const UScriptStruct* GetVec3Struct() { return Vec3Struct; }
-	static const UScriptStruct* GetVec4Struct() { return Vec4Struct; }
-	static const UScriptStruct* GetColorStruct() { return ColorStruct; }
-	static const UScriptStruct* GetMatrix4Struct() { return Matrix4Struct; }
-	static const UScriptStruct* GetGenericNumericStruct() { return NumericStruct; }
-	static const UScriptStruct* GetParameterMapStruct() { return ParameterMapStruct; }
-
-	static const UEnum* GetExecutionStateEnum() { return ExecutionStateEnum; }
+	static const FNiagaraTypeDefinition& GetIDDef() { return IDDef; }
+
+	static UScriptStruct* GetFloatStruct() { return FloatStruct; }
+	static UScriptStruct* GetBoolStruct() { return BoolStruct; }
+	static UScriptStruct* GetIntStruct() { return IntStruct; }
+	static UScriptStruct* GetVec2Struct() { return Vec2Struct; }
+	static UScriptStruct* GetVec3Struct() { return Vec3Struct; }
+	static UScriptStruct* GetVec4Struct() { return Vec4Struct; }
+	static UScriptStruct* GetColorStruct() { return ColorStruct; }
+	static UScriptStruct* GetQuatStruct() { return QuatStruct; }
+	static UScriptStruct* GetMatrix4Struct() { return Matrix4Struct; }
+	static UScriptStruct* GetGenericNumericStruct() { return NumericStruct; }
+	static UScriptStruct* GetParameterMapStruct() { return ParameterMapStruct; }
+	static UScriptStruct* GetIDStruct() { return IDStruct; }
+
+	static UEnum* GetExecutionStateEnum() { return ExecutionStateEnum; }
 
 	static const FNiagaraTypeDefinition& GetCollisionEventDef() { return CollisionEventDef; }
 
@@ -466,32 +540,37 @@ private:
 	static FNiagaraTypeDefinition Vec3Def;
 	static FNiagaraTypeDefinition Vec4Def;
 	static FNiagaraTypeDefinition ColorDef;
+	static FNiagaraTypeDefinition QuatDef;
 	static FNiagaraTypeDefinition Matrix4Def;
 	static FNiagaraTypeDefinition NumericDef;
 	static FNiagaraTypeDefinition ParameterMapDef;
-
-	static const UScriptStruct* FloatStruct;
-	static const UScriptStruct* BoolStruct;
-	static const UScriptStruct* IntStruct;
-	static const UScriptStruct* Vec2Struct;
-	static const UScriptStruct* Vec3Struct;
-	static const UScriptStruct* Vec4Struct;
-	static const UScriptStruct* ColorStruct;
-	static const UScriptStruct* Matrix4Struct;
-	static const UScriptStruct* NumericStruct;
-
-	static const UEnum* ExecutionStateEnum;
-
-	static const UScriptStruct* ParameterMapStruct;
-
-	static TSet<const UScriptStruct*> NumericStructs;
+	static FNiagaraTypeDefinition IDDef;
+
+	static UScriptStruct* FloatStruct;
+	static UScriptStruct* BoolStruct;
+	static UScriptStruct* IntStruct;
+	static UScriptStruct* Vec2Struct;
+	static UScriptStruct* Vec3Struct;
+	static UScriptStruct* Vec4Struct;
+	static UScriptStruct* QuatStruct;
+	static UScriptStruct* ColorStruct;
+	static UScriptStruct* Matrix4Struct;
+	static UScriptStruct* NumericStruct;
+
+	static UEnum* ExecutionStateEnum;
+	static UEnum* ExecutionStateSourceEnum;
+
+	static UScriptStruct* ParameterMapStruct;
+	static UScriptStruct* IDStruct;
+
+	static TSet<UScriptStruct*> NumericStructs;
 	static TArray<FNiagaraTypeDefinition> OrderedNumericTypes;
 
-	static TSet<const UScriptStruct*> ScalarStructs;
+	static TSet<UScriptStruct*> ScalarStructs;
 
-	static TSet<const UStruct*> FloatStructs;
-	static TSet<const UStruct*> IntStructs;
-	static TSet<const UStruct*> BoolStructs;
+	static TSet<UStruct*> FloatStructs;
+	static TSet<UStruct*> IntStructs;
+	static TSet<UStruct*> BoolStructs;
 
 	static FNiagaraTypeDefinition CollisionEventDef;
 
@@ -568,6 +647,19 @@ public:
 		}
 	}
 
+	FNiagaraTypeDefinition GetTypeDefFromStruct(UStruct* Struct)
+	{
+		for (FNiagaraTypeDefinition& TypeDef : RegisteredTypes)
+		{
+			if (Struct == TypeDef.GetStruct())
+			{
+				return TypeDef;
+			}
+		}
+
+		return FNiagaraTypeDefinition();
+	}
+
 private:
 	static TArray<FNiagaraTypeDefinition> RegisteredTypes;
 	static TArray<FNiagaraTypeDefinition> RegisteredParamTypes;
@@ -748,6 +840,30 @@ struct FNiagaraVariable
 		return OutVar;
 	}
 
+	static int32 SearchArrayForPartialNameMatch(const TArray<FNiagaraVariable>& Variables, const FName& VariableName)
+	{
+		FString VarNameStr = VariableName.ToString();
+		FString BestMatchSoFar;
+		int32 BestMatchIdx = INDEX_NONE;
+
+		for (int32 i = 0; i < Variables.Num(); i++)
+		{
+			const FNiagaraVariable& TestVar = Variables[i];
+			FString TestVarNameStr = TestVar.GetName().ToString();
+			if (TestVarNameStr == VarNameStr)
+			{
+				return i;
+			}
+			else if (VarNameStr.StartsWith(TestVarNameStr + TEXT(".")) && (BestMatchSoFar.Len() == 0 || TestVarNameStr.Len() > BestMatchSoFar.Len()))
+			{
+				BestMatchIdx = i;
+				BestMatchSoFar = TestVarNameStr;
+			}
+		}
+
+		return BestMatchIdx;
+	}
+
 private:
 	UPROPERTY(EditAnywhere, Category = "Variable")
 	FName Name;
diff --git a/Engine/Plugins/FX/Niagara/Source/Niagara/Public/NiagaraWorldManager.h b/Engine/Plugins/FX/Niagara/Source/Niagara/Public/NiagaraWorldManager.h
--- a/Engine/Plugins/FX/Niagara/Source/Niagara/Public/NiagaraWorldManager.h
+++ b/Engine/Plugins/FX/Niagara/Source/Niagara/Public/NiagaraWorldManager.h
@@ -16,6 +16,38 @@ class UWorld;
 class UNiagaraParameterCollection;
 class UNiagaraParameterCollectionInstance;
 
+
+
+class FNiagaraViewDataMgr
+{
+public:
+	FNiagaraViewDataMgr();
+
+	static void Init();
+	static void Shutdown();
+
+	void PostOpaqueRender(FPostOpaqueRenderParameters& Params)
+	{
+		SceneDepthTexture = Params.DepthTexture;
+		ViewUniformBuffer = Params.ViewUniformBuffer;
+		SceneNormalTexture = Params.NormalTexture;
+	}
+
+	FTexture2DRHIParamRef GetSceneDepthTexture() { return SceneDepthTexture; }
+	FTexture2DRHIParamRef GetSceneNormalTexture() { return SceneNormalTexture; }
+	FUniformBufferRHIParamRef GetViewUniformBuffer() { return ViewUniformBuffer; }
+
+private:
+	FTexture2DRHIParamRef SceneDepthTexture;
+	FTexture2DRHIParamRef SceneNormalTexture;
+	FUniformBufferRHIParamRef ViewUniformBuffer;
+
+	FPostOpaqueRenderDelegate PostOpaqueDelegate;
+};
+
+extern FNiagaraViewDataMgr GNiagaraViewDataManager;
+
+
 /**
 * Manager class for any data relating to a particular world.
 */
@@ -23,11 +55,7 @@ class FNiagaraWorldManager : public FGCObject
 {
 public:
 	
-	FNiagaraWorldManager(UWorld* InWorld)
-		: World(InWorld)
-		, CachedEffectsQuality(INDEX_NONE)
-	{}
-
+	FNiagaraWorldManager(UWorld* InWorld);
 	static FNiagaraWorldManager* Get(UWorld* World);
 
 	//~ GCObject Interface
@@ -43,7 +71,7 @@ public:
 	void Tick(float DeltaSeconds);
 
 	void OnWorldCleanup(bool bSessionEnded, bool bCleanupResources);
-
+	
 	FORCEINLINE FNDI_SkeletalMesh_GeneratedData& GetSkeletalMeshGeneratedData() { return SkeletalMeshGeneratedData; }
 private:
 	UWorld* World;
diff --git a/Engine/Plugins/FX/Niagara/Source/NiagaraCore/Private/NiagaraCustomVersion.cpp b/Engine/Plugins/FX/Niagara/Source/NiagaraCore/Private/NiagaraCustomVersion.cpp
new file mode 100644
--- /dev/null
+++ b/Engine/Plugins/FX/Niagara/Source/NiagaraCore/Private/NiagaraCustomVersion.cpp
@@ -0,0 +1,11 @@
+// Copyright 1998-2018 Epic Games, Inc. All Rights Reserved.
+
+#include "NiagaraCustomVersion.h"
+#include "Serialization/CustomVersion.h"
+
+const FGuid FNiagaraCustomVersion::GUID(0xFCF57AFA, 0x50764283, 0xB9A9E658, 0xFFA02D32);
+
+// Register the custom version with core
+FCustomVersionRegistration GRegisterNiagaraCustomVersion(FNiagaraCustomVersion::GUID, FNiagaraCustomVersion::LatestVersion, TEXT("NiagaraVer"));
+
+const FGuid FNiagaraCustomVersion::LatestScriptCompileVersion(0x375471F4, 0xD5DA466F, 0x8E78C8A5, 0xA257BCA2);
\ No newline at end of file
diff --git a/Engine/Plugins/FX/Niagara/Source/NiagaraCore/Private/NiagaraMergeable.cpp b/Engine/Plugins/FX/Niagara/Source/NiagaraCore/Private/NiagaraMergeable.cpp
new file mode 100644
--- /dev/null
+++ b/Engine/Plugins/FX/Niagara/Source/NiagaraCore/Private/NiagaraMergeable.cpp
@@ -0,0 +1,49 @@
+// Copyright 1998-2018 Epic Games, Inc. All Rights Reserved.
+
+#include "NiagaraMergeable.h"
+#include "UObject/PropertyPortFlags.h"
+
+UNiagaraMergeable::UNiagaraMergeable()
+#if WITH_EDITORONLY_DATA
+	: MergeId(FGuid::NewGuid())
+#endif
+{
+}
+
+#if WITH_EDITOR
+void UNiagaraMergeable::PostEditChangeProperty(struct FPropertyChangedEvent& PropertyChangedEvent)
+{
+	Super::PostEditChangeProperty(PropertyChangedEvent);
+	OnChangedDelegate.Broadcast();
+}
+
+bool UNiagaraMergeable::Equals(const UNiagaraMergeable* Other)
+{
+	if (Other == nullptr || this->GetClass() != Other->GetClass())
+	{
+		return false;
+	}
+
+	for (TFieldIterator<UProperty> PropertyIterator(this->GetClass()); PropertyIterator; ++PropertyIterator)
+	{
+		if (PropertyIterator->Identical(
+			PropertyIterator->ContainerPtrToValuePtr<void>(this),
+			PropertyIterator->ContainerPtrToValuePtr<void>(Other), PPF_DeepComparison) == false)
+		{
+			return false;
+		}
+	}
+
+	return true;
+}
+
+UNiagaraMergeable::FOnChanged& UNiagaraMergeable::OnChanged()
+{
+	return OnChangedDelegate;
+}
+
+FGuid UNiagaraMergeable::GetMergeId()
+{
+	return MergeId;
+}
+#endif
\ No newline at end of file
diff --git a/Engine/Plugins/FX/Niagara/Source/NiagaraCore/Public/NiagaraCoreModule.h b/Engine/Plugins/FX/Niagara/Source/NiagaraCore/Public/NiagaraCoreModule.h
new file mode 100644
--- /dev/null
+++ b/Engine/Plugins/FX/Niagara/Source/NiagaraCore/Public/NiagaraCoreModule.h
@@ -0,0 +1,14 @@
+// Copyright 1998-2018 Epic Games, Inc. All Rights Reserved.
+
+#pragma once
+
+#include "CoreMinimal.h"
+#include "Modules/ModuleInterface.h"
+
+/**
+* Niagara core module interface
+*/
+class INiagaraCoreModule : public IModuleInterface
+{
+};
+
diff --git a/Engine/Plugins/FX/Niagara/Source/NiagaraCore/Public/NiagaraDataInterfaceBase.h b/Engine/Plugins/FX/Niagara/Source/NiagaraCore/Public/NiagaraDataInterfaceBase.h
new file mode 100644
--- /dev/null
+++ b/Engine/Plugins/FX/Niagara/Source/NiagaraCore/Public/NiagaraDataInterfaceBase.h
@@ -0,0 +1,43 @@
+// Copyright 1998-2018 Epic Games, Inc. All Rights Reserved.
+#pragma once
+
+#include "NiagaraMergeable.h"
+#include "NiagaraDataInterfaceBase.generated.h"
+
+class INiagaraCompiler;
+class UCurveVector;
+class UCurveLinearColor;
+class UCurveFloat;
+class FNiagaraSystemInstance;
+class FNiagaraShader;
+struct FNiagaraDataInterfaceParamRef;
+class FRHICommandList;
+
+/**
+* An interface to the parameter bindings for the data interface used by a Niagara compute shader.
+*/
+struct FNiagaraDataInterfaceParametersCS
+{
+	virtual ~FNiagaraDataInterfaceParametersCS() {}
+	virtual void Bind(const FNiagaraDataInterfaceParamRef& ParamRef, const class FShaderParameterMap& ParameterMap) {}
+	virtual void Serialize(FArchive& Ar) { }
+	virtual void Set(FRHICommandList& RHICmdList, FNiagaraShader* Shader, class UNiagaraDataInterface* DataInterface) const {}
+};
+
+//////////////////////////////////////////////////////////////////////////
+
+/** Base class for all Niagara data interfaces. */
+UCLASS(abstract, EditInlineNew)
+class NIAGARACORE_API UNiagaraDataInterfaceBase : public UNiagaraMergeable
+{
+	GENERATED_UCLASS_BODY()
+
+public:
+	
+	/** Constructs the correct CS parameter type for this DI (if any). The object type returned by this can only vary by class and not per object data. */
+	virtual FNiagaraDataInterfaceParametersCS* ConstructComputeParameters()const
+	{
+		return nullptr;
+	}
+};
+
diff --git a/Engine/Plugins/FX/Niagara/Source/NiagaraCore/Public/NiagaraMergeable.h b/Engine/Plugins/FX/Niagara/Source/NiagaraCore/Public/NiagaraMergeable.h
new file mode 100644
--- /dev/null
+++ b/Engine/Plugins/FX/Niagara/Source/NiagaraCore/Public/NiagaraMergeable.h
@@ -0,0 +1,39 @@
+// Copyright 1998-2018 Epic Games, Inc. All Rights Reserved.
+
+#pragma once
+
+#include "NiagaraMergeable.generated.h"
+
+UCLASS()
+class NIAGARACORE_API UNiagaraMergeable : public UObject
+{
+	GENERATED_BODY()
+
+public:
+#if WITH_EDITOR
+	DECLARE_MULTICAST_DELEGATE(FOnChanged);
+#endif
+
+public:
+	UNiagaraMergeable();
+
+#if WITH_EDITOR
+	virtual void PostEditChangeProperty(struct FPropertyChangedEvent& PropertyChangedEvent);
+
+	bool Equals(const UNiagaraMergeable* Other);
+
+	FOnChanged& OnChanged();
+
+	FGuid GetMergeId();
+#endif
+
+private:
+#if WITH_EDITOR
+	FOnChanged OnChangedDelegate;
+#endif
+
+#if WITH_EDITORONLY_DATA
+	UPROPERTY()
+	FGuid MergeId;
+#endif
+};
\ No newline at end of file
diff --git a/Engine/Plugins/FX/Niagara/Source/NiagaraEditor/NiagaraEditor.Build.cs b/Engine/Plugins/FX/Niagara/Source/NiagaraEditor/NiagaraEditor.Build.cs
--- a/Engine/Plugins/FX/Niagara/Source/NiagaraEditor/NiagaraEditor.Build.cs
+++ b/Engine/Plugins/FX/Niagara/Source/NiagaraEditor/NiagaraEditor.Build.cs
@@ -31,6 +31,7 @@ public class NiagaraEditor : ModuleRules
                 "EditorStyle",
 				"UnrealEd", 
 				"VectorVM",
+                "NiagaraCore",
                 "Niagara",
                 "NiagaraShader",
                 "MovieScene",
@@ -44,7 +45,8 @@ public class NiagaraEditor : ModuleRules
 				"MovieSceneTools",
                 "MovieSceneTracks",
                 "AdvancedPreviewScene",
-            }
+				"Projects",
+			}
         );
 
 		PrivateIncludePathModuleNames.AddRange(
@@ -61,8 +63,10 @@ public class NiagaraEditor : ModuleRules
 
 		PublicDependencyModuleNames.AddRange(
             new string[] {
+                "NiagaraCore",
                 "NiagaraShader",
                 "Engine",
+                "NiagaraCore",
                 "Niagara",
                 "UnrealEd",
             }
diff --git a/Engine/Plugins/FX/Niagara/Source/NiagaraEditor/Private/AssetTypeActions/AssetTypeActions_NiagaraEmitter.cpp b/Engine/Plugins/FX/Niagara/Source/NiagaraEditor/Private/AssetTypeActions/AssetTypeActions_NiagaraEmitter.cpp
--- a/Engine/Plugins/FX/Niagara/Source/NiagaraEditor/Private/AssetTypeActions/AssetTypeActions_NiagaraEmitter.cpp
+++ b/Engine/Plugins/FX/Niagara/Source/NiagaraEditor/Private/AssetTypeActions/AssetTypeActions_NiagaraEmitter.cpp
@@ -6,16 +6,26 @@
 #include "NiagaraSystemToolkit.h"
 #include "NiagaraEditorStyle.h"
 #include "NiagaraSystemFactoryNew.h"
-#include "NiagaraSystemViewModel.h"
+#include "ViewModels/NiagaraSystemViewModel.h"
 #include "NiagaraSystemScriptViewModel.h"
 #include "Framework/MultiBox/MultiBoxBuilder.h"
 #include "EditorStyleSet.h"
 #include "Toolkits/AssetEditorToolkit.h"
 #include "IContentBrowserSingleton.h"
 #include "ContentBrowserModule.h"
+#include "NiagaraEditorUtilities.h"
+#include "Modules/ModuleManager.h"
 
 #define LOCTEXT_NAMESPACE "AssetTypeActions"
 
+FAssetTypeActions_NiagaraEmitter::FAssetTypeActions_NiagaraEmitter()
+{
+}
+
+FAssetTypeActions_NiagaraEmitter::~FAssetTypeActions_NiagaraEmitter()
+{
+}
+
 FColor FAssetTypeActions_NiagaraEmitter::GetTypeColor() const
 {
 	return FNiagaraEditorStyle::Get().GetColor("NiagaraEditor.AssetColors.Emitter").ToFColor(true);
@@ -53,6 +63,21 @@ void FAssetTypeActions_NiagaraEmitter::GetActions(const TArray<UObject*>& InObje
 			FExecuteAction::CreateSP(this, &FAssetTypeActions_NiagaraEmitter::ExecuteNewNiagaraSystem, NiagaraEmitters)
 		)
 	);
+
+	bool bInObjectsAreCompilable = true;
+	for (UObject* InObject : InObjects)
+	{
+		if (FNiagaraEditorUtilities::IsCompilableAssetClass(InObject->GetClass()) == false)
+		{
+			bInObjectsAreCompilable = false;
+		}
+	}
+
+	MenuBuilder.AddMenuEntry(
+		LOCTEXT("MarkDependentCompilableAssetsDirtyLabel", "Mark dependent compilable assets dirty"),
+		LOCTEXT("MarkDependentCompilableAssetsDirtyToolTip", "Finds all niagara assets which depend on this asset either directly or indirectly,\n and marks them dirty so they can be saved with the latest version."),
+		FSlateIcon(),
+		FUIAction(FExecuteAction::CreateStatic(&FNiagaraEditorUtilities::MarkDependentCompilableAssetsDirty, InObjects)));
 }
 
 void FAssetTypeActions_NiagaraEmitter::ExecuteNewNiagaraSystem(TArray<TWeakObjectPtr<UNiagaraEmitter>> Objects)
diff --git a/Engine/Plugins/FX/Niagara/Source/NiagaraEditor/Private/AssetTypeActions/AssetTypeActions_NiagaraEmitter.h b/Engine/Plugins/FX/Niagara/Source/NiagaraEditor/Private/AssetTypeActions/AssetTypeActions_NiagaraEmitter.h
--- a/Engine/Plugins/FX/Niagara/Source/NiagaraEditor/Private/AssetTypeActions/AssetTypeActions_NiagaraEmitter.h
+++ b/Engine/Plugins/FX/Niagara/Source/NiagaraEditor/Private/AssetTypeActions/AssetTypeActions_NiagaraEmitter.h
@@ -9,6 +9,9 @@
 class FAssetTypeActions_NiagaraEmitter: public FAssetTypeActions_Base
 {
 public:
+	FAssetTypeActions_NiagaraEmitter();
+	virtual ~FAssetTypeActions_NiagaraEmitter();
+
 	// IAssetTypeActions Implementation
 	virtual FText GetName() const override { return NSLOCTEXT("AssetTypeActions", "AssetTypeActions_NiagaraEmitter", "Niagara Emitter"); }
 	virtual FColor GetTypeColor() const override;
diff --git a/Engine/Plugins/FX/Niagara/Source/NiagaraEditor/Private/AssetTypeActions/AssetTypeActions_NiagaraScript.cpp b/Engine/Plugins/FX/Niagara/Source/NiagaraEditor/Private/AssetTypeActions/AssetTypeActions_NiagaraScript.cpp
--- a/Engine/Plugins/FX/Niagara/Source/NiagaraEditor/Private/AssetTypeActions/AssetTypeActions_NiagaraScript.cpp
+++ b/Engine/Plugins/FX/Niagara/Source/NiagaraEditor/Private/AssetTypeActions/AssetTypeActions_NiagaraScript.cpp
@@ -5,6 +5,10 @@
 #include "NiagaraScriptToolkit.h"
 #include "NiagaraEditorStyle.h"
 #include "AssetData.h"
+#include "NiagaraEditorUtilities.h"
+#include "Framework/MultiBox/MultiBoxBuilder.h"
+
+#define LOCTEXT_NAMESPACE "NiagaraScriptAssetTypeActions"
 
 const FName FAssetTypeActions_NiagaraScriptFunctions::NiagaraFunctionScriptName = FName("Niagara Function Script");
 const FName FAssetTypeActions_NiagaraScriptModules::NiagaraModuleScriptName = FName("Niagara Module Script");
@@ -56,7 +60,30 @@ FText FAssetTypeActions_NiagaraScript::GetDisplayNameFromAssetData(const FAssetD
 	return FAssetTypeActions_NiagaraScript::GetName();
 }
 
+bool FAssetTypeActions_NiagaraScript::HasActions(const TArray<UObject*>& InObjects) const
+{
+	for (UObject* ActionObject : InObjects)
+	{
+		if (FNiagaraEditorUtilities::IsCompilableAssetClass(ActionObject->GetClass()) == false)
+		{
+			return false;
+		}
+	}
+	return true;
+}
+
+void FAssetTypeActions_NiagaraScript::GetActions(const TArray<UObject*>& InObjects, FMenuBuilder& MenuBuilder)
+{
+	MenuBuilder.AddMenuEntry(
+		LOCTEXT("MarkDependentCompilableAssetsDirtyLabel", "Mark dependent compilable assets dirty"),
+		LOCTEXT("MarkDependentCompilableAssetsDirtyToolTip", "Finds all niagara assets which depend on this asset either directly or indirectly,\n and marks them dirty so they can be saved with the latest version."),
+		FSlateIcon(),
+		FUIAction(FExecuteAction::CreateStatic(&FNiagaraEditorUtilities::MarkDependentCompilableAssetsDirty, InObjects)));
+}
+
 UClass* FAssetTypeActions_NiagaraScript::GetSupportedClass() const
 { 
 	return UNiagaraScript::StaticClass(); 
 }
+
+#undef LOCTEXT_NAMESPACE
diff --git a/Engine/Plugins/FX/Niagara/Source/NiagaraEditor/Private/AssetTypeActions/AssetTypeActions_NiagaraScript.h b/Engine/Plugins/FX/Niagara/Source/NiagaraEditor/Private/AssetTypeActions/AssetTypeActions_NiagaraScript.h
--- a/Engine/Plugins/FX/Niagara/Source/NiagaraEditor/Private/AssetTypeActions/AssetTypeActions_NiagaraScript.h
+++ b/Engine/Plugins/FX/Niagara/Source/NiagaraEditor/Private/AssetTypeActions/AssetTypeActions_NiagaraScript.h
@@ -17,6 +17,8 @@ public:
 	virtual void OpenAssetEditor(const TArray<UObject*>& InObjects, TSharedPtr<class IToolkitHost> EditWithinLevelEditor = TSharedPtr<IToolkitHost>()) override;
 	virtual uint32 GetCategories() override { return FNiagaraEditorModule::GetAssetCategory(); }
 	virtual FText GetDisplayNameFromAssetData(const FAssetData& AssetData) const override;
+	virtual bool HasActions(const TArray<UObject*>& InObjects) const override;
+	virtual void GetActions(const TArray<UObject*>& InObjects, FMenuBuilder& MenuBuilder) override;
 	// End IAssetTypeActions Interface
 
 	/** Get the action type name. */
