Copying //UE4/Dev-Niagara to //UE4/Dev-Main (Source: //UE4/Dev-Niagara @ 4074996)

diff --git a/Engine/Plugins/FX/Niagara/Source/NiagaraEditor/Private/Customizations/NiagaraComponentDetails.cpp b/Engine/Plugins/FX/Niagara/Source/NiagaraEditor/Private/Customizations/NiagaraComponentDetails.cpp
--- a/Engine/Plugins/FX/Niagara/Source/NiagaraEditor/Private/Customizations/NiagaraComponentDetails.cpp
+++ b/Engine/Plugins/FX/Niagara/Source/NiagaraEditor/Private/Customizations/NiagaraComponentDetails.cpp
@@ -45,6 +45,7 @@
 #include "Modules/ModuleManager.h"
 #include "INiagaraEditorTypeUtilities.h"
 #include "Widgets/Layout/SBox.h"
+#include "UObject/WeakObjectPtr.h"
 #define LOCTEXT_NAMESPACE "NiagaraComponentDetails"
 
 class FNiagaraComponentNodeBuilder : public IDetailCustomNodeBuilder
@@ -53,12 +54,19 @@ public:
 	FNiagaraComponentNodeBuilder(UNiagaraComponent* InComponent, UNiagaraScript* SourceSpawn, UNiagaraScript* SourceUpdate)						   
 	{
 		Component = InComponent;
+		Component->OnSynchronizedWithAssetParameters().AddRaw(this, &FNiagaraComponentNodeBuilder::ComponentSynchronizedWithAssetParameters);
 		OriginalScripts.Add(SourceSpawn);
 		OriginalScripts.Add(SourceUpdate);
+		//UE_LOG(LogNiagaraEditor, Log, TEXT("FNiagaraComponentNodeBuilder %p Component %p"), this, Component.Get());
 	}
 
 	~FNiagaraComponentNodeBuilder()
 	{
+		if (Component.IsValid())
+		{
+			Component->OnSynchronizedWithAssetParameters().RemoveAll(this);
+		}
+		//UE_LOG(LogNiagaraEditor, Log, TEXT("~FNiagaraComponentNodeBuilder %p Component %p"), this, Component.Get());
 	}
 
 	virtual void SetOnRebuildChildren(FSimpleDelegate InOnRegenerateChildren) override
@@ -78,6 +86,7 @@ public:
 
 	virtual void GenerateChildContent(IDetailChildrenBuilder& ChildrenBuilder) override
 	{
+		check(Component.IsValid());
 		TArray<FNiagaraVariable> Parameters;
 		FNiagaraParameterStore& ParamStore = Component->GetOverrideParameters();
 		ParamStore.GetParameters(Parameters);
@@ -208,23 +217,21 @@ public:
 	}
 
 private:
-	void OnCollectionViewModelChanged()
-	{
-		OnRebuildChildren.ExecuteIfBound();
-	}
-
 	void OnParameterPreChange(FNiagaraVariable Var)
 	{
+		check(Component.IsValid());
 		Component->Modify();
 	}
 
 	void OnDataInterfacePreChange(FNiagaraVariable Var)
 	{
+		check(Component.IsValid());
 		Component->Modify();
 	}
 
 	void OnParameterChanged(FNiagaraVariable Var)
 	{
+		check(Component.IsValid());
 		Component->GetOverrideParameters().OnParameterChange();
 		Component->SetParameterValueOverriddenLocally(Var, true);
 	}
@@ -232,16 +239,19 @@ private:
 	void OnDataInterfaceChanged(FNiagaraVariable Var)
 	{
 		Component->GetOverrideParameters().OnInterfaceChange();
+		check(Component.IsValid());
 		Component->SetParameterValueOverriddenLocally(Var, true);
 	}
 
 	bool DoesParameterDifferFromDefault(const FNiagaraVariable& Var)
 	{
+		check(Component.IsValid());
 		return Component->IsParameterValueOverriddenLocally(Var.GetName());
 	}
 
 	FReply OnLocationResetClicked(FNiagaraVariable Parameter)
 	{
+		check(Component.IsValid());
 		FScopedTransaction ScopedTransaction(LOCTEXT("ResetParameterValue", "Reset parameter value to system defaults."));
 		Component->Modify();
 		Component->SetParameterValueOverriddenLocally(Parameter, false);
@@ -250,11 +260,16 @@ private:
 
 	EVisibility GetLocationResetVisibility(FNiagaraVariable Parameter) const
 	{
-		return Component->IsParameterValueOverriddenLocally(Parameter.GetName()) ? EVisibility::Visible : EVisibility::Collapsed;
+		return Component.IsValid() && Component->IsParameterValueOverriddenLocally(Parameter.GetName()) ? EVisibility::Visible : EVisibility::Collapsed;
+	}
+
+	void ComponentSynchronizedWithAssetParameters()
+	{
+		OnRebuildChildren.ExecuteIfBound();
 	}
 
 private:
-	UNiagaraComponent* Component;
+	TWeakObjectPtr<UNiagaraComponent> Component;
 	FSimpleDelegate OnRebuildChildren;
 	TArray<TSharedRef<FStructOnScope>> CreatedStructOnScopes;
 	TArray<UNiagaraScript*> OriginalScripts;
@@ -265,132 +280,32 @@ TSharedRef<IDetailCustomization> FNiagaraComponentDetails::MakeInstance()
 	return MakeShareable(new FNiagaraComponentDetails);
 }
 
-FNiagaraComponentDetails::FNiagaraComponentDetails() : Builder(nullptr), bQueueForDetailsRefresh(false)
+FNiagaraComponentDetails::FNiagaraComponentDetails() : Builder(nullptr)
 {
-	GEditor->RegisterForUndo(this);
 }
 
-
 FNiagaraComponentDetails::~FNiagaraComponentDetails()
 {
-	GEditor->UnregisterForUndo(this);
-
 	if (GEngine)
 	{
 		GEngine->OnWorldDestroyed().RemoveAll(this);
 	}
 
 	FGameDelegates::Get().GetEndPlayMapDelegate().RemoveAll(this);
-
-	if (ObjectsCustomized.Num() == 1)
-	{
-		UNiagaraComponent* Component = Cast<UNiagaraComponent>(ObjectsCustomized[0].Get());
-
-		if (Component)
-		{
-			if (FNiagaraSystemInstance* SystemInstance = Component->GetSystemInstance())
-			{
-				if (OnResetDelegateHandle.IsValid())
-				{
-					SystemInstance->OnReset().Remove(OnResetDelegateHandle);
-				}
-
-				if (OnInitDelegateHandle.IsValid())
-				{
-					SystemInstance->OnInitialized().Remove(OnInitDelegateHandle);
-				}
-
-				SystemInstance->OnDestroyed().RemoveAll(this);
-			}
-		}
-	}
-}
-
-bool FNiagaraComponentDetails::IsTickable() const
-{
-	return bQueueForDetailsRefresh;
-}
-
-void FNiagaraComponentDetails::Tick(float DeltaTime)
-{
-	if (bQueueForDetailsRefresh)
-	{
-		if (Builder)
-		{
-			Builder->ForceRefreshDetails();
-			bQueueForDetailsRefresh = false;
-		}
-	}
-}
-
-void FNiagaraComponentDetails::PostUndo(bool bSuccess)
-{
-	// We may have queued up as a result of an Undo of adding the System itself. The objects we were 
-	// referencing may therefore have been removed. If so, we'll have to take a different path later on in the code.
-	int32 NumValid = 0;
-	for (TWeakObjectPtr<UObject> WeakPtr : ObjectsCustomized)
-	{
-		if (!WeakPtr.IsStale(true) && WeakPtr.IsValid())
-		{
-			NumValid++;
-		}
-	}
-
-	if (Builder)
-	{
-		// This is tricky. There is a high chance that if the component changed, then
-		// any cached FNiagaraVariable that we're holding on to may have been changed out from underneath us.
-		// So we essentially must tear down and start again in the UI.
-		// HOWEVER, a refresh will invoke a new constructor of the FNiagaraComponentDetails, which puts us 
-		// in the queue to handle PostUndo events. This turns quickly into an infinitely loop. Therefore,
-		// we circumvent this by telling the editor that we need to queue up an event
-		// that we will handle in the subsequent frame's Tick event. Not the cleanest approach, 
-		// but because we are doing things like CopyOnWrite, the normal Unreal property editing
-		// path is not available to us.
-		if (NumValid != 0)
-		{
-			bQueueForDetailsRefresh = true;
-		}
-		else
-		{
-			// If we no longer have any valid pointers, but previously had a builder, that means 
-			// that the builder is probably dead or dying soon. We shouldn't trust it any more and
-			// we should make sure that we aren't queueing for ticks to refresh either.
-			Builder = nullptr;
-			bQueueForDetailsRefresh = false;
-		}
-
-#if 0
-		for (TSharedPtr<FNiagaraParameterViewModelCustomDetails>& ViewModel : ViewModels)
-		{
-			ViewModel->Reset();
-		}
-#endif
-	}
-}
-
-TStatId FNiagaraComponentDetails::GetStatId() const
-{
-	RETURN_QUICK_DECLARE_CYCLE_STAT(FNiagaraComponentDetails, STATGROUP_Tickables);
 }
 
 void FNiagaraComponentDetails::OnPiEEnd()
 {
 	UE_LOG(LogNiagaraEditor, Log, TEXT("onPieEnd"));
-	if (ObjectsCustomized.Num() == 1)
+	if (Component.IsValid())
 	{
-		UNiagaraComponent* Component = Cast<UNiagaraComponent>(ObjectsCustomized[0].Get());
-
-		if (Component)
+		if (Component->GetOutermost()->HasAnyPackageFlags(PKG_PlayInEditor))
 		{
-			if (Component->GetOutermost()->HasAnyPackageFlags(PKG_PlayInEditor))
+			UE_LOG(LogNiagaraEditor, Log, TEXT("onPieEnd - has package flags"));
+			UWorld* TheWorld = UWorld::FindWorldInPackage(Component->GetOutermost());
+			if (TheWorld)
 			{
-				UE_LOG(LogNiagaraEditor, Log, TEXT("onPieEnd - has package flags"));
-				UWorld* TheWorld = UWorld::FindWorldInPackage(Component->GetOutermost());
-				if (TheWorld)
-				{
-					OnWorldDestroyed(TheWorld);
-				}
+				OnWorldDestroyed(TheWorld);
 			}
 		}
 	}
@@ -400,52 +315,35 @@ void FNiagaraComponentDetails::OnWorldDestroyed(class UWorld* InWorld)
 {
 	// We have to clear out any temp data interfaces that were bound to the component's package when the world goes away or otherwise
 	// we'll report GC leaks..
-	if (ObjectsCustomized.Num() == 1)
+	if (Component.IsValid())
 	{
-		UNiagaraComponent* Component = Cast<UNiagaraComponent>(ObjectsCustomized[0].Get());
-
-		if (Component)
+		if (Component->GetWorld() == InWorld)
 		{
-			if (Component->GetWorld() == InWorld)
-			{
-				UE_LOG(LogNiagaraEditor, Log, TEXT("OnWorldDestroyed - matched up"));
-				Builder = nullptr;
-			}
+			UE_LOG(LogNiagaraEditor, Log, TEXT("OnWorldDestroyed - matched up"));
+			Builder = nullptr;
 		}
 	}
 }
 
 void FNiagaraComponentDetails::CustomizeDetails(IDetailLayoutBuilder& DetailBuilder)
 {
- 	Builder = &DetailBuilder;
-// 
-// 	// Create a category so this is displayed early in the properties
-// 	IDetailCategoryBuilder& NiagaraCategory = DetailBuilder.EditCategory("Niagara", FText::GetEmpty(), ECategoryPriority::Important);
-// 
-// 	// Store off the properties for analysis in later function calls.
+	Builder = &DetailBuilder;
+
 	static const FName ParamCategoryName = TEXT("NiagaraComponent_Parameters");
-	static const FName ParamCollectionCategoryName = TEXT("NiagaraComponent_ParameterCollectionOverrides");
 	static const FName ScriptCategoryName = TEXT("Parameters");
 
-	DetailBuilder.EditCategory(ScriptCategoryName);
-
- 	TSharedPtr<IPropertyHandle> LocalOverridesPropertyHandle = DetailBuilder.GetProperty(TEXT("OverrideParameters"));
- 	if (LocalOverridesPropertyHandle.IsValid())
- 	{
- 		LocalOverridesPropertyHandle->MarkHiddenByCustomization();
-		UProperty* Property = LocalOverridesPropertyHandle->GetProperty();
- 	}
+	TSharedPtr<IPropertyHandle> LocalOverridesPropertyHandle = DetailBuilder.GetProperty(TEXT("OverrideParameters"));
+	if (LocalOverridesPropertyHandle.IsValid())
+	{
+		LocalOverridesPropertyHandle->MarkHiddenByCustomization();
+	}
 
+	TArray<TWeakObjectPtr<UObject>> ObjectsCustomized;
 	DetailBuilder.GetObjectsBeingCustomized(ObjectsCustomized);
 
-	if (ObjectsCustomized.Num() == 1)
+	if (ObjectsCustomized.Num() == 1 && ObjectsCustomized[0]->IsA<UNiagaraComponent>())
 	{
-		UNiagaraComponent* Component = Cast<UNiagaraComponent>(ObjectsCustomized[0].Get());
-
-		if (Component == nullptr)
-		{
-			return;
-		}
+		Component = CastChecked<UNiagaraComponent>(ObjectsCustomized[0].Get());
 
 		if (GEngine)
 		{
@@ -454,52 +352,15 @@ void FNiagaraComponentDetails::CustomizeDetails(IDetailLayoutBuilder& DetailBuil
 
 		FGameDelegates::Get().GetEndPlayMapDelegate().AddRaw(this, &FNiagaraComponentDetails::OnPiEEnd);
 
-		FNiagaraSystemInstance* SystemInstance = Component->GetSystemInstance();
-		if (SystemInstance == nullptr)
+		if (Component->GetAsset())
 		{
-			return;
-		}
-
-		// We'll want to monitor for any serious changes that might require a rebuild of the UI from scratch.
-		UE_LOG(LogNiagaraEditor, Log, TEXT("Registering with System instance %p"), SystemInstance);
-		SystemInstance->OnReset().RemoveAll(this);
-		SystemInstance->OnInitialized().RemoveAll(this);
-		SystemInstance->OnDestroyed().RemoveAll(this);
-		OnResetDelegateHandle = SystemInstance->OnReset().AddSP(this, &FNiagaraComponentDetails::OnSystemInstanceReset);
-		OnInitDelegateHandle = SystemInstance->OnInitialized().AddSP(this, &FNiagaraComponentDetails::OnSystemInstanceReset);
-		SystemInstance->OnDestroyed().AddSP(this, &FNiagaraComponentDetails::OnSystemInstanceDestroyed);
-		FNiagaraParameterStore* ParamStore = &Component->GetOverrideParameters();
-		UNiagaraScript* ScriptSpawn = Component->GetAsset()->GetSystemSpawnScript();
-		UNiagaraScript* ScriptUpdate = Component->GetAsset()->GetSystemUpdateScript();
-		
-		IDetailCategoryBuilder& InputParamCategory = DetailBuilder.EditCategory(ParamCategoryName, LOCTEXT("ParamCategoryName", "Override Parameters"));
-		//InputParamCategory.HeaderContent(SNew(SAddParameterButton, InputCollectionViewModel));
-		InputParamCategory.AddCustomBuilder(MakeShared<FNiagaraComponentNodeBuilder>(Component, ScriptSpawn, ScriptUpdate));
-	}
-
-}
+			UNiagaraScript* ScriptSpawn = Component->GetAsset()->GetSystemSpawnScript();
+			UNiagaraScript* ScriptUpdate = Component->GetAsset()->GetSystemUpdateScript();
 
-void FNiagaraComponentDetails::OnSystemInstanceReset()
-{
-	UE_LOG(LogNiagaraEditor, Log, TEXT("OnSystemInstanceReset()"));
-	bQueueForDetailsRefresh = true;
-}
-
-void FNiagaraComponentDetails::OnSystemInstanceDestroyed()
-{
-	UE_LOG(LogNiagaraEditor, Log, TEXT("OnSystemInstanceDestroyed()"));
-	bQueueForDetailsRefresh = true;
-}
-
-FReply FNiagaraComponentDetails::OnSystemOpenRequested(UNiagaraSystem* InSystem)
-{
-	if (InSystem)
-	{
-		FAssetEditorManager::Get().OpenEditorForAsset(InSystem);
+			IDetailCategoryBuilder& InputParamCategory = DetailBuilder.EditCategory(ParamCategoryName, LOCTEXT("ParamCategoryName", "Override Parameters"));
+			InputParamCategory.AddCustomBuilder(MakeShared<FNiagaraComponentNodeBuilder>(Component.Get(), ScriptSpawn, ScriptUpdate));
+		}
 	}
-	return FReply::Handled();
 }
 
-
-
 #undef LOCTEXT_NAMESPACE
diff --git a/Engine/Plugins/FX/Niagara/Source/NiagaraEditor/Private/Customizations/NiagaraComponentDetails.h b/Engine/Plugins/FX/Niagara/Source/NiagaraEditor/Private/Customizations/NiagaraComponentDetails.h
--- a/Engine/Plugins/FX/Niagara/Source/NiagaraEditor/Private/Customizations/NiagaraComponentDetails.h
+++ b/Engine/Plugins/FX/Niagara/Source/NiagaraEditor/Private/Customizations/NiagaraComponentDetails.h
@@ -9,6 +9,7 @@
 #include "Input/Reply.h"
 #include "Layout/Visibility.h"
 #include "Misc/NotifyHook.h"
+#include "NiagaraComponent.h"
 
 class IDetailLayoutBuilder;
 class IPropertyHandle;
@@ -18,7 +19,7 @@ class SNiagaraParameterEditor;
 class FNiagaraParameterViewModelCustomDetails;
 class UNiagaraSystem;
 
-class FNiagaraComponentDetails : public IDetailCustomization, public FEditorUndoClient, public FTickableEditorObject
+class FNiagaraComponentDetails : public IDetailCustomization
 {
 public:
 	virtual ~FNiagaraComponentDetails();
@@ -29,28 +30,13 @@ public:
 	/** IDetailCustomization interface */
 	virtual void CustomizeDetails( IDetailLayoutBuilder& DetailBuilder ) override;
 
-	//~ FEditorUndoClient interface
-	virtual void PostUndo(bool bSuccess) override;
-	virtual void PostRedo(bool bSuccess) override { PostUndo(bSuccess); }
-
-	// FTickableEditorObject
-	virtual bool IsTickable() const override;
-	virtual void Tick(float DeltaTime) override;
-	virtual TStatId GetStatId() const override;
-
 protected:
-	void OnSystemInstanceReset();
-	void OnSystemInstanceDestroyed();
-	FReply OnSystemOpenRequested(UNiagaraSystem* InSystem);
 	void OnWorldDestroyed(class UWorld* InWorld);
 	void OnPiEEnd();
 
 	FNiagaraComponentDetails();
+
 private:
-	TArray<TWeakObjectPtr<UObject>> ObjectsCustomized;
+	TWeakObjectPtr<UNiagaraComponent> Component;
 	IDetailLayoutBuilder* Builder;
-	bool bQueueForDetailsRefresh;
-
-	FDelegateHandle OnInitDelegateHandle;
-	FDelegateHandle OnResetDelegateHandle;
 };
diff --git a/Engine/Plugins/FX/Niagara/Source/NiagaraEditor/Private/Customizations/NiagaraEventScriptPropertiesCustomization.cpp b/Engine/Plugins/FX/Niagara/Source/NiagaraEditor/Private/Customizations/NiagaraEventScriptPropertiesCustomization.cpp
--- a/Engine/Plugins/FX/Niagara/Source/NiagaraEditor/Private/Customizations/NiagaraEventScriptPropertiesCustomization.cpp
+++ b/Engine/Plugins/FX/Niagara/Source/NiagaraEditor/Private/Customizations/NiagaraEventScriptPropertiesCustomization.cpp
@@ -25,10 +25,11 @@
 #include "SGraphActionMenu.h"
 #include "NiagaraEditorUtilities.h"
 #include "Editor.h"
+#include "IDetailChildrenBuilder.h"
 
 #define LOCTEXT_NAMESPACE "FNiagaraEventScriptPropertiesCustomization"
 
-TSharedRef<class IDetailCustomization> FNiagaraEventScriptPropertiesCustomization::MakeInstance(TWeakObjectPtr<UNiagaraSystem> InSystem,
+TSharedRef<IPropertyTypeCustomization> FNiagaraEventScriptPropertiesCustomization::MakeInstance(TWeakObjectPtr<UNiagaraSystem> InSystem,
 TWeakObjectPtr<UNiagaraEmitter> InEmitter)
 {
 	return MakeShareable(new FNiagaraEventScriptPropertiesCustomization(InSystem, InEmitter));
@@ -45,47 +46,34 @@ FNiagaraEventScriptPropertiesCustomization::~FNiagaraEventScriptPropertiesCustom
 	GEditor->UnregisterForUndo(this);
 }
 
-void FNiagaraEventScriptPropertiesCustomization::CustomizeDetails(IDetailLayoutBuilder& DetailBuilder)
+void FNiagaraEventScriptPropertiesCustomization::CustomizeHeader(TSharedRef<IPropertyHandle> StructPropertyHandle, class FDetailWidgetRow& HeaderRow, IPropertyTypeCustomizationUtils& StructCustomizationUtils)
 {
-	HandleSrcID = DetailBuilder.GetProperty(GET_MEMBER_NAME_CHECKED(FNiagaraEventScriptProperties, SourceEmitterID));
-	HandleEventName = DetailBuilder.GetProperty(GET_MEMBER_NAME_CHECKED(FNiagaraEventScriptProperties, SourceEventName));
-	HandleSpawnNumber = DetailBuilder.GetProperty(GET_MEMBER_NAME_CHECKED(FNiagaraEventScriptProperties, SpawnNumber));
-	HandleExecutionMode = DetailBuilder.GetProperty(GET_MEMBER_NAME_CHECKED(FNiagaraEventScriptProperties, ExecutionMode));
-	HandleMaxEvents = DetailBuilder.GetProperty(GET_MEMBER_NAME_CHECKED(FNiagaraEventScriptProperties, MaxEventsPerFrame));
-
-	IDetailCategoryBuilder& CategoryBuilder = DetailBuilder.EditCategory(TEXT("Event Handler Options"));
-
-	if (HandleSrcID.IsValid())
-	{
-		HandleSrcID->MarkHiddenByCustomization();
-	}
-
-	if (HandleEventName.IsValid())
-	{
-		HandleEventName->MarkHiddenByCustomization();
-	}
-
-	if (HandleSpawnNumber.IsValid())
-	{
-		HandleSpawnNumber->MarkHiddenByCustomization();
-	}
+	HeaderRow.NameContent()
+		[
+			StructPropertyHandle->CreatePropertyNameWidget()
+		];
+	HeaderRow.ValueContent()
+		[
+			StructPropertyHandle->CreatePropertyValueWidget()
+		];
+}
 
-	if (HandleExecutionMode.IsValid())
-	{
-		HandleExecutionMode->MarkHiddenByCustomization();
-	}
+void FNiagaraEventScriptPropertiesCustomization::CustomizeChildren(TSharedRef<IPropertyHandle> StructPropertyHandle, class IDetailChildrenBuilder& ChildBuilder, IPropertyTypeCustomizationUtils& StructCustomizationUtils)
+{
+	HandleSrcID = StructPropertyHandle->GetChildHandle(GET_MEMBER_NAME_CHECKED(FNiagaraEventScriptProperties, SourceEmitterID));
+	HandleEventName = StructPropertyHandle->GetChildHandle(GET_MEMBER_NAME_CHECKED(FNiagaraEventScriptProperties, SourceEventName));
+	HandleSpawnNumber = StructPropertyHandle->GetChildHandle(GET_MEMBER_NAME_CHECKED(FNiagaraEventScriptProperties, SpawnNumber));
+	HandleExecutionMode = StructPropertyHandle->GetChildHandle(GET_MEMBER_NAME_CHECKED(FNiagaraEventScriptProperties, ExecutionMode));
+	HandleMaxEvents = StructPropertyHandle->GetChildHandle(GET_MEMBER_NAME_CHECKED(FNiagaraEventScriptProperties, MaxEventsPerFrame));
+	HandleUseRandomSpawnNumber = StructPropertyHandle->GetChildHandle(GET_MEMBER_NAME_CHECKED(FNiagaraEventScriptProperties, bRandomSpawnNumber));
+	HandleMinSpawnNumber = StructPropertyHandle->GetChildHandle(GET_MEMBER_NAME_CHECKED(FNiagaraEventScriptProperties, MinSpawnNumber));
 
-	if (HandleMaxEvents.IsValid())
-	{
-		HandleMaxEvents->MarkHiddenByCustomization();
-	}
-	
 	ResolveEmitterName();
 
 	// Add Source property
 	{
 		FText EventSrcTxt = LOCTEXT("EventSource", "Source");
-		FDetailWidgetRow& Row = CategoryBuilder.AddCustomRow(EventSrcTxt);
+		FDetailWidgetRow& Row = ChildBuilder.AddCustomRow(EventSrcTxt);
 		FText TooltipText = LOCTEXT("ChooseProvider", "Choose the source emitter and event.");
 
 		TAttribute<FText> ErrorTextAttribute(this, &FNiagaraEventScriptPropertiesCustomization::GetErrorText);
@@ -132,7 +120,7 @@ void FNiagaraEventScriptPropertiesCustomization::CustomizeDetails(IDetailLayoutB
 	// Add Execution Mode property
 	{
 		FText EventSrcTxt = LOCTEXT("ExecutionMode", "ExecutionMode");
-		FDetailWidgetRow& ExecRow = CategoryBuilder.AddCustomRow(EventSrcTxt);
+		FDetailWidgetRow& ExecRow = ChildBuilder.AddCustomRow(EventSrcTxt);
 		ExecRow.NameWidget
 			[
 				HandleExecutionMode->CreatePropertyNameWidget()
@@ -145,14 +133,14 @@ void FNiagaraEventScriptPropertiesCustomization::CustomizeDetails(IDetailLayoutB
 
 	// Add Max Events property
 	{
-		IDetailPropertyRow& MaxRow = CategoryBuilder.AddProperty(HandleMaxEvents);
+		IDetailPropertyRow& MaxRow = ChildBuilder.AddProperty(HandleMaxEvents.ToSharedRef());
 	}
 	
 	// Add Spawn Number property
 	{
 		FText EventSrcTxt = LOCTEXT("SpawnNumber", "SpawnNumber");
-		FDetailWidgetRow& SpawnRow = CategoryBuilder.AddCustomRow(EventSrcTxt);
-		TSharedRef<SWidget> SpawnValueWidget = HandleSpawnNumber->CreatePropertyValueWidget();
+		FDetailWidgetRow& SpawnRow = ChildBuilder.AddCustomRow(EventSrcTxt);
+		TSharedRef<SWidget> SpawnNumberValueWidget = HandleSpawnNumber->CreatePropertyValueWidget();
 		TAttribute<bool> EnabledAttribute(this, &FNiagaraEventScriptPropertiesCustomization::GetSpawnNumberEnabled);
 		SpawnRow.IsEnabled(EnabledAttribute);
 
@@ -160,6 +148,44 @@ void FNiagaraEventScriptPropertiesCustomization::CustomizeDetails(IDetailLayoutB
 			[
 				HandleSpawnNumber->CreatePropertyNameWidget()
 			];
+		SpawnRow.ValueWidget
+			[
+				SpawnRow.ValueWidget.Widget = SpawnNumberValueWidget
+			];
+	}
+
+	// Add Spawn Number property
+	{
+		FText EventSrcTxt = LOCTEXT("SpawnNumber", "Minimum SpawnNumber");
+		FDetailWidgetRow& SpawnRow = ChildBuilder.AddCustomRow(EventSrcTxt);
+		TSharedRef<SWidget> MinSpawnNumberValueWidget = HandleMinSpawnNumber->CreatePropertyValueWidget();
+		TAttribute<bool> EnabledAttribute(this, &FNiagaraEventScriptPropertiesCustomization::GetUseRandomSpawnNumber);
+		SpawnRow.IsEnabled(EnabledAttribute);
+		TAttribute<EVisibility> VisibilityAttribute(this, &FNiagaraEventScriptPropertiesCustomization::GetMinSpawnNumberVisible);
+		SpawnRow.Visibility(VisibilityAttribute);
+
+		SpawnRow.NameWidget
+			[
+				HandleMinSpawnNumber->CreatePropertyNameWidget()
+			];
+		SpawnRow.ValueWidget
+			[
+				SpawnRow.ValueWidget.Widget = MinSpawnNumberValueWidget
+			];
+	}
+
+	// Add Random Spawn Number boolean property
+	{
+		FText SrcTxt = LOCTEXT("RandomSpawnNumber", "Random SpawnNumber");
+		FDetailWidgetRow& SpawnRow = ChildBuilder.AddCustomRow(SrcTxt);
+		TSharedRef<SWidget> SpawnValueWidget = HandleUseRandomSpawnNumber->CreatePropertyValueWidget();
+		TAttribute<bool> EnabledAttribute(this, &FNiagaraEventScriptPropertiesCustomization::GetSpawnNumberEnabled);
+		SpawnRow.IsEnabled(EnabledAttribute);
+		
+		SpawnRow.NameWidget
+			[
+				HandleUseRandomSpawnNumber->CreatePropertyNameWidget()
+			];
 		SpawnRow.ValueWidget
 			[
 				SpawnRow.ValueWidget.Widget = SpawnValueWidget
@@ -252,6 +278,22 @@ bool FNiagaraEventScriptPropertiesCustomization::GetSpawnNumberEnabled() const
 	return false;
 }
 
+bool FNiagaraEventScriptPropertiesCustomization::GetUseRandomSpawnNumber() const
+{
+	bool bRandom = false;
+	if (GetSpawnNumberEnabled() && HandleUseRandomSpawnNumber->GetValue(bRandom))
+	{
+		return bRandom;
+	}
+
+	return false;
+}
+
+EVisibility FNiagaraEventScriptPropertiesCustomization::GetMinSpawnNumberVisible() const
+{
+	return GetUseRandomSpawnNumber() ? EVisibility::Visible : EVisibility::Collapsed;
+}
+
 void FNiagaraEventScriptPropertiesCustomization::ResolveEmitterName()
 {
 	FString EventSrcIdStr;
@@ -309,17 +351,17 @@ TSharedRef<SWidget> FNiagaraEventScriptPropertiesCustomization::OnGetMenuContent
 TArray<FName> FNiagaraEventScriptPropertiesCustomization::GetEventNames(UNiagaraEmitter* InEmitter) const
 {
 	TArray<FName> EventNames;
-	if (InEmitter->CollisionMode != ENiagaraCollisionMode::None)
-	{
-		EventNames.Add(NIAGARA_BUILTIN_EVENTNAME_COLLISION);
-	}
-
 	TArray<UNiagaraScript*> Scripts;
 	InEmitter->GetScripts(Scripts);
 
 	for (UNiagaraScript* Script : Scripts)
 	{
-		for (FNiagaraDataSetProperties& Props : Script->WriteDataSets)
+		if (!Script->IsReadyToRun(ENiagaraSimTarget::CPUSim))
+		{
+			continue;
+		}
+
+		for (FNiagaraDataSetProperties& Props : Script->GetVMExecutableData().WriteDataSets)
 		{
 			if (Props.ID.Name.IsValid())
 			{
diff --git a/Engine/Plugins/FX/Niagara/Source/NiagaraEditor/Private/Customizations/NiagaraMetaDataCustomNodeBuilder.h b/Engine/Plugins/FX/Niagara/Source/NiagaraEditor/Private/Customizations/NiagaraMetaDataCustomNodeBuilder.h
--- a/Engine/Plugins/FX/Niagara/Source/NiagaraEditor/Private/Customizations/NiagaraMetaDataCustomNodeBuilder.h
+++ b/Engine/Plugins/FX/Niagara/Source/NiagaraEditor/Private/Customizations/NiagaraMetaDataCustomNodeBuilder.h
@@ -1,31 +1,14 @@
-// Copyright 1998-2018 Epic Games, Inc. All Rights Reserved.
+﻿// Copyright 1998-2018 Epic Games, Inc. All Rights Reserved.
 
 #pragma once
-#include "DetailLayoutBuilder.h"
-#include "DetailCategoryBuilder.h"
 #include "IDetailPropertyRow.h"
-#include "DetailWidgetRow.h"
-#include "Widgets/Text/SInlineEditableTextBlock.h"
-#include "NiagaraMetaDataCollectionViewModel.h"
-#include "NiagaraScriptInputCollectionViewModel.h"
-#include "NiagaraScriptOutputCollectionViewModel.h"
 #include "NiagaraMetaDataCollectionViewModel.h"
 #include "NiagaraMetaDataViewModel.h"
-#include "NiagaraScriptViewModel.h"
-#include "NiagaraParameterViewModel.h"
-#include "NiagaraEditorStyle.h"
 #include "IDetailChildrenBuilder.h"
-#include "Framework/MultiBox/MultiBoxBuilder.h"
-#include "Widgets/Input/SComboButton.h"
-#include "Widgets/Images/SImage.h"
-#include "NiagaraEditorModule.h"
-#include "Modules/ModuleManager.h"
-#include "INiagaraEditorTypeUtilities.h"
-#include "Widgets/Layout/SBox.h"
-#include "NiagaraScript.h"
 #include "UObject/StructOnScope.h"
-#include "Internationalization/Text.h"
 #include "IDetailGroup.h"
+#include "NiagaraEditorModule.h"
+#include "Modules/ModuleManager.h"
 
 
 #define LOCTEXT_NAMESPACE "NiagaraMetaDataCustomNodeBuilder"
@@ -82,6 +65,27 @@ public:
 				TSharedPtr<IPropertyHandle> PropertyHandle = DescriptionRow->GetPropertyHandle();
 				PropertyHandle->SetOnPropertyValueChanged(FSimpleDelegate::CreateSP(MetadataViewModel, &FNiagaraMetaDataViewModel::NotifyMetaDataChanged));
 			}
+			IDetailPropertyRow* CategoryRow = ChildrenBuilder.AddExternalStructureProperty(StructData, FName(TEXT("CategoryName")), "CategoryName");
+
+			if (CategoryRow)
+			{
+				CategoryRow->Visibility(EVisibility::Hidden); // hide it here, show it in groups only
+				MetaDataGroup.AddPropertyRow(CategoryRow->GetPropertyHandle()->AsShared());
+
+				TSharedPtr<IPropertyHandle> PropertyHandle = CategoryRow->GetPropertyHandle();
+				PropertyHandle->SetOnPropertyValueChanged(FSimpleDelegate::CreateSP(MetadataViewModel, &FNiagaraMetaDataViewModel::NotifyMetaDataChanged));
+			}
+
+			IDetailPropertyRow* SortingRow = ChildrenBuilder.AddExternalStructureProperty(StructData, FName(TEXT("EditorSortPriority")), "EditorSortPriority");
+
+			if (SortingRow)
+			{
+				SortingRow->Visibility(EVisibility::Hidden); // hide it here, show it in groups only
+				MetaDataGroup.AddPropertyRow(SortingRow->GetPropertyHandle()->AsShared());
+
+				TSharedPtr<IPropertyHandle> PropertyHandle = SortingRow->GetPropertyHandle();
+				PropertyHandle->SetOnPropertyValueChanged(FSimpleDelegate::CreateSP(MetadataViewModel, &FNiagaraMetaDataViewModel::NotifyMetaDataChanged));
+			}
 
 			IDetailPropertyRow* MetaDataRow = ChildrenBuilder.AddExternalStructureProperty(StructData, FName(TEXT("PropertyMetaData")), "PropertyMetaData");
 			
diff --git a/Engine/Plugins/FX/Niagara/Source/NiagaraEditor/Private/Customizations/NiagaraParameterCollectionCustomNodeBuilder.h b/Engine/Plugins/FX/Niagara/Source/NiagaraEditor/Private/Customizations/NiagaraParameterCollectionCustomNodeBuilder.h
--- a/Engine/Plugins/FX/Niagara/Source/NiagaraEditor/Private/Customizations/NiagaraParameterCollectionCustomNodeBuilder.h
+++ b/Engine/Plugins/FX/Niagara/Source/NiagaraEditor/Private/Customizations/NiagaraParameterCollectionCustomNodeBuilder.h
@@ -25,8 +25,8 @@
 class FNiagaraParameterCollectionCustomNodeBuilder : public IDetailCustomNodeBuilder
 {
 public:
-	FNiagaraParameterCollectionCustomNodeBuilder(TSharedRef<INiagaraParameterCollectionViewModel> InViewModel, bool bAllowMetaData = true)
-		: ViewModel(InViewModel)
+	FNiagaraParameterCollectionCustomNodeBuilder(TSharedRef<INiagaraParameterCollectionViewModel> InViewModel, bool bInAllowMetaData = true)
+		: ViewModel(InViewModel), bAllowMetaData(bInAllowMetaData)
 	{
 		ViewModel->OnCollectionChanged().AddRaw(this, &FNiagaraParameterCollectionCustomNodeBuilder::OnCollectionViewModelChanged);
 	}
diff --git a/Engine/Plugins/FX/Niagara/Source/NiagaraEditor/Private/Customizations/NiagaraScriptDetails.cpp b/Engine/Plugins/FX/Niagara/Source/NiagaraEditor/Private/Customizations/NiagaraScriptDetails.cpp
--- a/Engine/Plugins/FX/Niagara/Source/NiagaraEditor/Private/Customizations/NiagaraScriptDetails.cpp
+++ b/Engine/Plugins/FX/Niagara/Source/NiagaraEditor/Private/Customizations/NiagaraScriptDetails.cpp
@@ -15,6 +15,7 @@
 #include "IDetailChildrenBuilder.h"
 #include "Framework/MultiBox/MultiBoxBuilder.h"
 #include "Widgets/Input/SComboButton.h"
+#include "Widgets/Input/SButton.h"
 #include "Widgets/Images/SImage.h"
 #include "NiagaraEditorModule.h"
 #include "Modules/ModuleManager.h"
@@ -121,6 +122,12 @@ void FNiagaraScriptDetails::CustomizeDetails(IDetailLayoutBuilder& DetailBuilder
 	static const FName ScriptCategoryName = TEXT("Script");
 	static const FName MetadataCategoryName = TEXT("Metadata");
 
+	TSharedPtr<IPropertyHandle> NumericOutputPropertyHandle = DetailBuilder.GetProperty(TEXT("NumericOutputTypeSelectionMode"));
+	if (NumericOutputPropertyHandle.IsValid())
+	{
+		NumericOutputPropertyHandle->MarkHiddenByCustomization();
+	}
+
 	DetailBuilder.EditCategory(ScriptCategoryName);
 
 	UNiagaraScript* StandaloneScript = ScriptViewModel->GetStandaloneScript();
@@ -160,9 +167,36 @@ void FNiagaraScriptDetails::CustomizeDetails(IDetailLayoutBuilder& DetailBuilder
 	if (MetadataCollectionViewModel.IsValid())
 	{
 		IDetailCategoryBuilder& MetadataDetailCategory = DetailBuilder.EditCategory(MetadataCategoryName, LOCTEXT("MetadataParamCategoryName", "Variable Metadata"));
+		MetadataDetailCategory.HeaderContent(
+			// Checkbox for showing in curve editor
+			SNew(SHorizontalBox)
+			+ SHorizontalBox::Slot()
+			.Padding(5, 0, 5, 0)
+			.HAlign(EHorizontalAlignment::HAlign_Right)
+			[
+				SNew(SButton)
+				.ButtonStyle(FEditorStyle::Get(), "HoverHintOnly")
+				.HAlign(HAlign_Right)
+				.VAlign(VAlign_Center)
+				.ContentPadding(1)
+				.ToolTipText(LOCTEXT("RefreshMetadataToolTip", "Save metadata changes to graph and refresh the view"))
+				.OnClicked(this, &FNiagaraScriptDetails::OnRefreshMetadata)
+				.Content()
+				[
+					SNew(SImage)
+					.Image(FEditorStyle::GetBrush("Icons.Refresh"))
+				]
+			]
+		);
 		MetadataDetailCategory.AddCustomBuilder(MakeShared<FNiagaraMetaDataCustomNodeBuilder>(MetadataCollectionViewModel.ToSharedRef()));
 	}	
 }
 
+FReply FNiagaraScriptDetails::OnRefreshMetadata()
+{
+	ScriptViewModel->GetMetadataCollectionViewModel()->RequestRefresh(true);
+    return FReply::Handled();
+}
+
 
 #undef LOCTEXT_NAMESPACE
diff --git a/Engine/Plugins/FX/Niagara/Source/NiagaraEditor/Private/Customizations/NiagaraScriptDetails.h b/Engine/Plugins/FX/Niagara/Source/NiagaraEditor/Private/Customizations/NiagaraScriptDetails.h
--- a/Engine/Plugins/FX/Niagara/Source/NiagaraEditor/Private/Customizations/NiagaraScriptDetails.h
+++ b/Engine/Plugins/FX/Niagara/Source/NiagaraEditor/Private/Customizations/NiagaraScriptDetails.h
@@ -1,9 +1,10 @@
-// Copyright 1998-2018 Epic Games, Inc. All Rights Reserved.
+﻿// Copyright 1998-2018 Epic Games, Inc. All Rights Reserved.
 
 #pragma once
 
 #include "CoreMinimal.h"
 #include "IDetailCustomization.h"
+#include "Input/Reply.h"
 
 class FNiagaraScriptViewModel;
 class INiagaraParameterCollectionViewModel;
@@ -16,6 +17,9 @@ public:
 	FNiagaraScriptDetails(TSharedPtr<FNiagaraScriptViewModel> InScriptViewModel);
 
 	virtual void CustomizeDetails(IDetailLayoutBuilder& DetailBuilder);
+
+	FReply OnRefreshMetadata();
+
 private:
 	TSharedPtr<FNiagaraScriptViewModel> ScriptViewModel;
 
diff --git a/Engine/Plugins/FX/Niagara/Source/NiagaraEditor/Private/Customizations/NiagaraTypeCustomizations.cpp b/Engine/Plugins/FX/Niagara/Source/NiagaraEditor/Private/Customizations/NiagaraTypeCustomizations.cpp
--- a/Engine/Plugins/FX/Niagara/Source/NiagaraEditor/Private/Customizations/NiagaraTypeCustomizations.cpp
+++ b/Engine/Plugins/FX/Niagara/Source/NiagaraEditor/Private/Customizations/NiagaraTypeCustomizations.cpp
@@ -9,7 +9,19 @@
 #include "Widgets/Text/STextBlock.h"
 #include "Widgets/Input/SCheckBox.h"
 #include "IDetailChildrenBuilder.h"
+#include "NiagaraRendererProperties.h"
+#include "NiagaraEmitter.h"
+#include "SGraphActionMenu.h"
+#include "Framework/Application/SlateApplication.h"
+#include "ScopedTransaction.h"
+#include "NiagaraParameterMapHistory.h"
+#include "NiagaraScriptSource.h"
+#include "NiagaraNodeOutput.h"
+#include "NiagaraConstants.h"
+#include "NiagaraNodeParameterMapBase.h"
+#include "Widgets/Input/STextComboBox.h"
 
+#define LOCTEXT_NAMESPACE "FNiagaraVariableAttributeBindingCustomization"
 
 
 void FNiagaraNumericCustomization::CustomizeHeader(TSharedRef<IPropertyHandle> PropertyHandle, FDetailWidgetRow& HeaderRow, IPropertyTypeCustomizationUtils& CustomizationUtils)
@@ -90,3 +102,200 @@ void FNiagaraMatrixCustomization::CustomizeChildren(TSharedRef<IPropertyHandle>
 		ChildBuilder.AddProperty(PropertyHandle->GetChildHandle(ChildNum).ToSharedRef());
 	}
 }
+
+FText FNiagaraVariableAttributeBindingCustomization::GetCurrentText() const
+{
+	if (BaseEmitter && TargetVariableBinding)
+	{
+		return FText::FromName(TargetVariableBinding->BoundVariable.GetName());
+	}
+	return FText::FromString(TEXT("Missing"));
+}
+
+FText FNiagaraVariableAttributeBindingCustomization::GetTooltipText() const
+{
+	if (BaseEmitter && TargetVariableBinding)
+	{
+		FString DefaultValueStr = TargetVariableBinding->DefaultValueIfNonExistent.GetName().ToString();
+		
+		if (!TargetVariableBinding->DefaultValueIfNonExistent.GetName().IsValid() || TargetVariableBinding->DefaultValueIfNonExistent.IsDataAllocated() == true)
+		{
+			DefaultValueStr = TargetVariableBinding->DefaultValueIfNonExistent.GetType().ToString(TargetVariableBinding->DefaultValueIfNonExistent.GetData());
+			DefaultValueStr.TrimEndInline();
+		}
+
+		FText TooltipDesc = FText::Format(LOCTEXT("BindingTooltip", "Use the variable \"{0}\" if it exists, otherwise use the default \"{1}\" "), FText::FromName(TargetVariableBinding->BoundVariable.GetName()),
+			FText::FromString(DefaultValueStr));
+		return TooltipDesc;
+	}
+	return FText::FromString(TEXT("Missing"));
+}
+
+TSharedRef<SWidget> FNiagaraVariableAttributeBindingCustomization::OnGetMenuContent() const
+{
+	FGraphActionMenuBuilder MenuBuilder;
+
+	return SNew(SBorder)
+		.BorderImage(FEditorStyle::GetBrush("Menu.Background"))
+		.Padding(5)
+		[
+			SNew(SBox)
+			[
+				SNew(SGraphActionMenu)
+				.OnActionSelected(this, &FNiagaraVariableAttributeBindingCustomization::OnActionSelected)
+				.OnCreateWidgetForAction(SGraphActionMenu::FOnCreateWidgetForAction::CreateSP(this, &FNiagaraVariableAttributeBindingCustomization::OnCreateWidgetForAction))
+				.OnCollectAllActions(this, &FNiagaraVariableAttributeBindingCustomization::CollectAllActions)
+				.AutoExpandActionMenu(false)
+				.ShowFilterTextBox(true)
+			]
+		];
+}
+
+TArray<FName> FNiagaraVariableAttributeBindingCustomization::GetNames(UNiagaraEmitter* InEmitter) const
+{
+	TArray<FName> Names;
+
+	UNiagaraScriptSource* Source = Cast<UNiagaraScriptSource>(InEmitter->GraphSource);
+	if (Source)
+	{
+		TArray<FNiagaraParameterMapHistory> Histories =  UNiagaraNodeParameterMapBase::GetParameterMaps(Source->NodeGraph);
+		for (const FNiagaraParameterMapHistory& History : Histories)
+		{
+			for (const FNiagaraVariable& Var : History.Variables)
+			{
+				if (FNiagaraParameterMapHistory::IsAttribute(Var) && Var.GetType() == TargetVariableBinding->BoundVariable.GetType())
+				{
+					Names.AddUnique(Var.GetName());
+				}
+			}
+		}
+		
+	}
+	return Names;
+}
+
+void FNiagaraVariableAttributeBindingCustomization::CollectAllActions(FGraphActionListBuilderBase& OutAllActions)
+{
+	TArray<FName> EventNames = GetNames(BaseEmitter);
+	FName EmitterName = BaseEmitter->GetFName();
+	for (FName EventName : EventNames)
+	{
+		FText CategoryName = FText();
+		FString DisplayNameString = FName::NameToDisplayString(EventName.ToString(), false);
+		const FText NameText = FText::FromString(DisplayNameString);
+		const FText TooltipDesc = FText::Format(LOCTEXT("SetFunctionPopupTooltip", "Use the variable \"{0}\" "), FText::FromString(DisplayNameString));
+		TSharedPtr<FNiagaraStackAssetAction_VarBind> NewNodeAction(new FNiagaraStackAssetAction_VarBind(EventName, CategoryName, NameText,
+			TooltipDesc, 0, FText()));
+		OutAllActions.AddAction(NewNodeAction);
+	}
+}
+
+TSharedRef<SWidget> FNiagaraVariableAttributeBindingCustomization::OnCreateWidgetForAction(struct FCreateWidgetForActionData* const InCreateData)
+{
+	return SNew(SVerticalBox)
+		+ SVerticalBox::Slot()
+		.AutoHeight()
+		[
+			SNew(STextBlock)
+			.Text(InCreateData->Action->GetMenuDescription())
+			.ToolTipText(InCreateData->Action->GetTooltipDescription())
+		];
+}
+
+
+void FNiagaraVariableAttributeBindingCustomization::OnActionSelected(const TArray< TSharedPtr<FEdGraphSchemaAction> >& SelectedActions, ESelectInfo::Type InSelectionType)
+{
+	if (InSelectionType == ESelectInfo::OnMouseClick || InSelectionType == ESelectInfo::OnKeyPress || SelectedActions.Num() == 0)
+	{
+		for (int32 ActionIndex = 0; ActionIndex < SelectedActions.Num(); ActionIndex++)
+		{
+			TSharedPtr<FEdGraphSchemaAction> CurrentAction = SelectedActions[ActionIndex];
+
+			if (CurrentAction.IsValid())
+			{
+				FSlateApplication::Get().DismissAllMenus();
+				FNiagaraStackAssetAction_VarBind* EventSourceAction = (FNiagaraStackAssetAction_VarBind*)CurrentAction.Get();
+				ChangeSource(EventSourceAction->VarName);
+			}
+		}
+	}
+}
+
+void FNiagaraVariableAttributeBindingCustomization::ChangeSource(FName InVarName)
+{
+	FScopedTransaction Transaction(FText::Format(LOCTEXT("ChangeSource", " Change Variable Source to \"{0}\" "), FText::FromName(InVarName)));
+	TArray<UObject*> Objects;
+	PropertyHandle->GetOuterObjects(Objects);
+	for (UObject* Obj : Objects)
+	{
+		Obj->Modify();
+	}
+
+	PropertyHandle->NotifyPreChange();
+	TargetVariableBinding->BoundVariable.SetName(InVarName);
+	TargetVariableBinding->DataSetVariable = FNiagaraConstants::GetAttributeAsDataSetKey(TargetVariableBinding->BoundVariable);
+	PropertyHandle->NotifyPostChange();
+	PropertyHandle->NotifyFinishedChangingProperties();
+}
+
+void FNiagaraVariableAttributeBindingCustomization::CustomizeHeader(TSharedRef<IPropertyHandle> InPropertyHandle, FDetailWidgetRow& HeaderRow, IPropertyTypeCustomizationUtils& CustomizationUtils)
+{
+	PropertyHandle = InPropertyHandle;
+	TArray<UObject*> Objects;
+	PropertyHandle->GetOuterObjects(Objects);
+	bool bAddDefault = true;
+	if (Objects.Num() == 1)
+	{
+		UNiagaraRendererProperties* RenderProps = Cast<UNiagaraRendererProperties>(Objects[0]);
+		if (RenderProps)
+		{
+			BaseEmitter = Cast<UNiagaraEmitter>(RenderProps->GetOuter());
+
+			if (BaseEmitter)
+			{
+				TargetVariableBinding = (FNiagaraVariableAttributeBinding*)PropertyHandle->GetValueBaseAddress((uint8*)Objects[0]);
+				
+				HeaderRow
+					.NameContent()
+					[
+						PropertyHandle->CreatePropertyNameWidget()
+					]
+					.ValueContent()
+					.MaxDesiredWidth(200.f)
+					[
+						SNew(SComboButton)
+						.OnGetMenuContent(this, &FNiagaraVariableAttributeBindingCustomization::OnGetMenuContent)
+						.ContentPadding(1)
+						.ToolTipText(this, &FNiagaraVariableAttributeBindingCustomization::GetTooltipText)
+						.ButtonContent()
+						[
+							SNew(STextBlock)
+							.Text(this, &FNiagaraVariableAttributeBindingCustomization::GetCurrentText)
+							.Font(IDetailLayoutBuilder::GetDetailFont())
+						]
+					];
+				bAddDefault = false;
+			}
+		}
+	}
+	
+
+	if (bAddDefault)
+	{
+		HeaderRow
+			.NameContent()
+			[
+				PropertyHandle->CreatePropertyNameWidget()
+			]
+		.ValueContent()
+			.MaxDesiredWidth(200.f)
+			[
+				SNew(STextBlock)
+				.Text(FText::FromString(FName::NameToDisplayString(Cast<UStructProperty>(PropertyHandle->GetProperty())->Struct->GetName(), false)))
+				.Font(IDetailLayoutBuilder::GetDetailFont())
+			];
+	}
+}
+
+
+#undef LOCTEXT_NAMESPACE
diff --git a/Engine/Plugins/FX/Niagara/Source/NiagaraEditor/Private/Customizations/NiagaraTypeCustomizations.h b/Engine/Plugins/FX/Niagara/Source/NiagaraEditor/Private/Customizations/NiagaraTypeCustomizations.h
--- a/Engine/Plugins/FX/Niagara/Source/NiagaraEditor/Private/Customizations/NiagaraTypeCustomizations.h
+++ b/Engine/Plugins/FX/Niagara/Source/NiagaraEditor/Private/Customizations/NiagaraTypeCustomizations.h
@@ -6,6 +6,8 @@
 #include "IDetailCustomization.h"
 #include "IPropertyTypeCustomization.h"
 #include "Styling/SlateTypes.h"
+#include "EdGraph/EdGraphSchema.h"
+#include "NiagaraTypeCustomizations.generated.h"
 
 class FNiagaraScriptViewModel;
 class IPropertyHandle;
@@ -65,3 +67,63 @@ public:
 
 private:
 };
+
+
+USTRUCT()
+struct FNiagaraStackAssetAction_VarBind : public FEdGraphSchemaAction
+{
+	GENERATED_USTRUCT_BODY();
+
+	FName VarName;
+
+	// Simple type info
+	static FName StaticGetTypeId() { static FName Type("FNiagaraStackAssetAction_VarBind"); return Type; }
+	virtual FName GetTypeId() const override { return StaticGetTypeId(); }
+
+	FNiagaraStackAssetAction_VarBind()
+		: FEdGraphSchemaAction()
+	{}
+
+	FNiagaraStackAssetAction_VarBind(FName InVarName,
+		FText InNodeCategory, FText InMenuDesc, FText InToolTip, const int32 InGrouping, FText InKeywords)
+		: FEdGraphSchemaAction(MoveTemp(InNodeCategory), MoveTemp(InMenuDesc), MoveTemp(InToolTip), InGrouping, MoveTemp(InKeywords)),
+		VarName(InVarName)
+	{}
+
+	//~ Begin FEdGraphSchemaAction Interface
+	virtual UEdGraphNode* PerformAction(class UEdGraph* ParentGraph, UEdGraphPin* FromPin, const FVector2D Location, bool bSelectNewNode = true) override
+	{
+		return nullptr;
+	}
+	//~ End FEdGraphSchemaAction Interface
+
+};
+
+class FNiagaraVariableAttributeBindingCustomization : public IPropertyTypeCustomization
+{
+public:
+	/** @return A new instance of this class */
+	static TSharedRef<IPropertyTypeCustomization> MakeInstance()
+	{
+		return MakeShared<FNiagaraVariableAttributeBindingCustomization>();
+	}
+
+	/** IPropertyTypeCustomization interface begin */
+	virtual void CustomizeHeader(TSharedRef<IPropertyHandle> StructPropertyHandle, class FDetailWidgetRow& HeaderRow, IPropertyTypeCustomizationUtils& StructCustomizationUtils) override;
+	virtual void CustomizeChildren(TSharedRef<IPropertyHandle> StructPropertyHandle, class IDetailChildrenBuilder& ChildBuilder, IPropertyTypeCustomizationUtils& StructCustomizationUtils) override {};
+	/** IPropertyTypeCustomization interface end */
+private:
+	FText GetCurrentText() const;
+	FText GetTooltipText() const;
+	TSharedRef<SWidget> OnGetMenuContent() const;
+	TArray<FName> GetNames(class UNiagaraEmitter* InEmitter) const;
+	void ChangeSource(FName InVarName);
+	void CollectAllActions(FGraphActionListBuilderBase& OutAllActions);
+	TSharedRef<SWidget> OnCreateWidgetForAction(struct FCreateWidgetForActionData* const InCreateData);
+	void OnActionSelected(const TArray< TSharedPtr<FEdGraphSchemaAction> >& SelectedActions, ESelectInfo::Type InSelectionType);
+
+	TSharedPtr<IPropertyHandle> PropertyHandle;
+	class UNiagaraEmitter* BaseEmitter;
+	struct FNiagaraVariableAttributeBinding* TargetVariableBinding;
+
+};
\ No newline at end of file
diff --git a/Engine/Plugins/FX/Niagara/Source/NiagaraEditor/Private/EdGraphSchema_Niagara.cpp b/Engine/Plugins/FX/Niagara/Source/NiagaraEditor/Private/EdGraphSchema_Niagara.cpp
--- a/Engine/Plugins/FX/Niagara/Source/NiagaraEditor/Private/EdGraphSchema_Niagara.cpp
+++ b/Engine/Plugins/FX/Niagara/Source/NiagaraEditor/Private/EdGraphSchema_Niagara.cpp
@@ -36,6 +36,8 @@
 #include "NiagaraEmitter.h"
 #include "NiagaraNodeEmitter.h"
 #include "NiagaraParameterCollection.h"
+#include "NiagaraNodeReroute.h"
+#include "NiagaraNodeUsageSelector.h"
 #include "EdGraphNode_Comment.h"
 
 #include "Modules/ModuleManager.h"
@@ -51,6 +53,7 @@ const FLinearColor UEdGraphSchema_Niagara::NodeTitleColor_SystemConstant = FLine
 const FLinearColor UEdGraphSchema_Niagara::NodeTitleColor_FunctionCall = FLinearColor::Blue;
 const FLinearColor UEdGraphSchema_Niagara::NodeTitleColor_CustomHlsl = FLinearColor::Yellow;
 const FLinearColor UEdGraphSchema_Niagara::NodeTitleColor_Event = FLinearColor::Red;
+const FLinearColor UEdGraphSchema_Niagara::NodeTitleColor_TranslatorConstant = FLinearColor::Gray;
 
 const FName UEdGraphSchema_Niagara::PinCategoryType("Type");
 const FName UEdGraphSchema_Niagara::PinCategoryMisc("Misc");
@@ -117,7 +120,7 @@ UEdGraphNode* FNiagaraSchemaAction_NewNode::PerformAction(class UEdGraph* Parent
 
 		ResultNode = NodeTemplate;
 
-		ParentGraph->NotifyGraphChanged();
+		//ParentGraph->NotifyGraphChanged();
 	}
 
 	return ResultNode;
@@ -299,6 +302,7 @@ TArray<TSharedPtr<FNiagaraSchemaAction_NewNode> > UEdGraphSchema_Niagara::GetGra
 	NiagaraGraph->GetNodesOfClass<UNiagaraNodeEmitter>(Emitters);
 	bool bSystemGraph = Emitters.Num() != 0 || NiagaraGraph->FindOutputNode(ENiagaraScriptUsage::SystemSpawnScript) != nullptr || NiagaraGraph->FindOutputNode(ENiagaraScriptUsage::SystemUpdateScript) != nullptr;
 	bool bModuleGraph = NiagaraGraph->FindOutputNode(ENiagaraScriptUsage::Module) != nullptr;
+	bool bDynamicInputGraph = NiagaraGraph->FindOutputNode(ENiagaraScriptUsage::DynamicInput) != nullptr;
 	bool bFunctionGraph = NiagaraGraph->FindOutputNode(ENiagaraScriptUsage::Function) != nullptr;
 	bool bUpdateGraph = NiagaraGraph->FindOutputNode(ENiagaraScriptUsage::ParticleUpdateScript) != nullptr;
 
@@ -409,6 +413,36 @@ TArray<TSharedPtr<FNiagaraSchemaAction_NewNode> > UEdGraphSchema_Niagara::GetGra
 		}
 	}
 
+	// Add dynamic inputs for default usage in module and dynamic input graphs
+	if (bModuleGraph || bDynamicInputGraph)
+	{
+		for (const FAssetData& ScriptAsset : ScriptAssets)
+		{
+			FName UsageName;
+			ScriptAsset.GetTagValue(GET_MEMBER_NAME_CHECKED(UNiagaraScript, Usage), UsageName);
+			FString QualifiedUsageName = "ENiagaraScriptUsage::" + UsageName.ToString();
+			int32 UsageIndex = NiagaraScriptUsageEnum->GetIndexByNameString(QualifiedUsageName);
+			if (UsageIndex != INDEX_NONE)
+			{
+				ENiagaraScriptUsage Usage = static_cast<ENiagaraScriptUsage>(NiagaraScriptUsageEnum->GetValueByIndex(UsageIndex));
+				if (Usage == ENiagaraScriptUsage::DynamicInput)
+				{
+					FString DisplayNameString = FName::NameToDisplayString(ScriptAsset.AssetName.ToString(), false);
+					const FText MenuDesc = FText::FromString(DisplayNameString);
+					FText AssetDesc;
+					ScriptAsset.GetTagValue(GET_MEMBER_NAME_CHECKED(UNiagaraScript, Description), AssetDesc);
+					const FText TooltipDesc = FText::Format(LOCTEXT("DynamicInputPopupTooltip", "Path: {0}\nDescription: {1}"), FText::FromString(ScriptAsset.ObjectPath.ToString()), AssetDesc);
+
+					TSharedPtr<FNiagaraSchemaAction_NewNode> FunctionCallAction = AddNewNodeAction(NewActions, LOCTEXT("DynamicInputMenu", "Dynamic Inputs"), MenuDesc, *DisplayNameString, FText::FromName(ScriptAsset.ObjectPath));
+
+					UNiagaraNodeFunctionCall* FunctionCallNode = NewObject<UNiagaraNodeFunctionCall>(OwnerOfTemporaries);
+					FunctionCallNode->FunctionScriptAssetObjectPath = ScriptAsset.ObjectPath;
+					FunctionCallAction->NodeTemplate = FunctionCallNode;
+				}
+			}
+		}
+	}
+
 	//Add event read and writes nodes
 	if (bModuleGraph)
 	{
@@ -711,6 +745,26 @@ TArray<TSharedPtr<FNiagaraSchemaAction_NewNode> > UEdGraphSchema_Niagara::GetGra
 			GetConstAction->NodeTemplate = InputNode;
 		}
 
+		//Emitter constants managed by the Translator.
+		const TArray<FNiagaraVariable>& TranslatorConstants = FNiagaraConstants::GetTranslatorConstants();
+		for (const FNiagaraVariable& TransConst : TranslatorConstants)
+		{
+			FFormatNamedArguments Args;
+			Args.Add(TEXT("Constant"), FText::FromName(TransConst.GetName()));
+			const FText MenuDesc = FText::Format(LOCTEXT("GetTranslatorConstant", "{Constant}"), Args);
+
+			TSharedPtr<FNiagaraSchemaAction_NewNode> GetConstAction = AddNewNodeAction(NewActions, LOCTEXT("Translator Parameters Menu Title", "Special Purpose Parameters"), MenuDesc, TransConst.GetName(), FText::GetEmpty());
+
+			UNiagaraNodeInput* InputNode = NewObject<UNiagaraNodeInput>(OwnerOfTemporaries);
+			InputNode->Usage = ENiagaraInputNodeUsage::TranslatorConstant;
+			InputNode->ExposureOptions.bCanAutoBind = true;
+			InputNode->ExposureOptions.bHidden = true;
+			InputNode->ExposureOptions.bRequired = false;
+			InputNode->ExposureOptions.bExposed = false;
+			InputNode->Input = TransConst;
+			GetConstAction->NodeTemplate = InputNode;
+		}
+
 		AddParameterMenuOptions(NewActions, NiagaraGraph, OwnerOfTemporaries, NiagaraGraph);
 
 		//Add a generic Parameter node to allow easy creation of parameters.
@@ -768,6 +822,24 @@ TArray<TSharedPtr<FNiagaraSchemaAction_NewNode> > UEdGraphSchema_Niagara::GetGra
 	}
 	//TODO: Add quick commands for certain UNiagaraStructs and UNiagaraScripts to be added as functions
 
+	// Add reroute node
+	{
+		const FText UtilMenuCat = LOCTEXT("NiagaraRerouteMenuCat", "Util");
+		const FText RerouteMenuDesc = LOCTEXT("NiagaraRerouteMenuDesc", "Reroute ");
+		TSharedPtr<FNiagaraSchemaAction_NewNode> Action = AddNewNodeAction(NewActions, UtilMenuCat, RerouteMenuDesc, TEXT("Reroute"), FText::GetEmpty());
+		UNiagaraNodeReroute* RerouteNode = NewObject<UNiagaraNodeReroute>(OwnerOfTemporaries);
+		Action->NodeTemplate = RerouteNode;
+	}
+
+	// Add usage selector node
+	{
+		const FText UtilMenuCat = LOCTEXT("NiagaraUsageSelectorMenuCat", "Util");
+		const FText UsageSelectorMenuDesc = LOCTEXT("NiagaraUsageSelectorMenuDesc", "Select By Use");
+		TSharedPtr<FNiagaraSchemaAction_NewNode> Action = AddNewNodeAction(NewActions, UtilMenuCat, UsageSelectorMenuDesc, TEXT("Select By Use"), FText::GetEmpty());
+		UNiagaraNodeUsageSelector* Node = NewObject<UNiagaraNodeUsageSelector>(OwnerOfTemporaries);
+		Action->NodeTemplate = Node;
+	}
+
 	return NewActions;
 }
 
@@ -874,9 +946,16 @@ void UEdGraphSchema_Niagara::BreakSinglePinLink(UEdGraphPin* SourcePin, UEdGraph
 	Super::BreakSinglePinLink(SourcePin, TargetPin);
 }
 
+void UEdGraphSchema_Niagara::BreakPinLinks(UEdGraphPin& TargetPin, bool bSendsNodeNotification) const
+{
+	const FScopedTransaction Transaction(NSLOCTEXT("UnrealEd", "NiagaraEditorBreakPinLinks", "Niagara Editor: Break Pin Links"));
+
+	Super::BreakPinLinks(TargetPin, bSendsNodeNotification);
+}
+
 FConnectionDrawingPolicy* UEdGraphSchema_Niagara::CreateConnectionDrawingPolicy(int32 InBackLayerID, int32 InFrontLayerID, float InZoomFactor, const FSlateRect& InClippingRect, class FSlateWindowElementList& InDrawElements, class UEdGraph* InGraphObj) const
 {
-	return new FNiagaraConnectionDrawingPolicy(InBackLayerID, InFrontLayerID, InZoomFactor, InClippingRect, InDrawElements);
+	return new FNiagaraConnectionDrawingPolicy(InBackLayerID, InFrontLayerID, InZoomFactor, InClippingRect, InDrawElements, InGraphObj);
 }
 
 void UEdGraphSchema_Niagara::ResetPinToAutogeneratedDefaultValue(UEdGraphPin* Pin, bool bCallModifyCallbacks) const
@@ -888,6 +967,25 @@ void UEdGraphSchema_Niagara::ResetPinToAutogeneratedDefaultValue(UEdGraphPin* Pi
 	}
 }
 
+void UEdGraphSchema_Niagara::OnPinConnectionDoubleCicked(UEdGraphPin* PinA, UEdGraphPin* PinB, const FVector2D& GraphPosition) const
+{
+	const FScopedTransaction Transaction(LOCTEXT("CreateRerouteNodeOnWire", "Create Reroute Node"));
+
+	//@TODO: This constant is duplicated from inside of SGraphNodeKnot
+	const FVector2D NodeSpacerSize(42.0f, 24.0f);
+	const FVector2D KnotTopLeft = GraphPosition - (NodeSpacerSize * 0.5f);
+
+	// Create a new knot
+	UEdGraph* ParentGraph = PinA->GetOwningNode()->GetGraph();
+	UNiagaraNodeReroute* NewReroute = FNiagaraSchemaAction_NewNode::SpawnNodeFromTemplate<UNiagaraNodeReroute>(ParentGraph, NewObject<UNiagaraNodeReroute>(), KnotTopLeft);
+
+	// Move the connections across (only notifying the knot, as the other two didn't really change)
+	PinA->BreakLinkTo(PinB);
+	PinA->MakeLinkTo((PinA->Direction == EGPD_Output) ? NewReroute->GetInputPin(0) : NewReroute->GetOutputPin(0));
+	PinB->MakeLinkTo((PinB->Direction == EGPD_Output) ? NewReroute->GetInputPin(0) : NewReroute->GetOutputPin(0));
+	NewReroute->PropagatePinType();
+}
+
 bool UEdGraphSchema_Niagara::TryCreateConnection(UEdGraphPin* PinA, UEdGraphPin* PinB) const
 {
 	const FScopedTransaction Transaction(NSLOCTEXT("UnrealEd", "NiagaraEditorCreateConnection", "Niagara Editor: Create Connection"));
@@ -988,37 +1086,42 @@ bool UEdGraphSchema_Niagara::TryCreateConnection(UEdGraphPin* PinA, UEdGraphPin*
 
 FLinearColor UEdGraphSchema_Niagara::GetPinTypeColor(const FEdGraphPinType& PinType) const
 {
-	const UGraphEditorSettings* Settings = GetDefault<UGraphEditorSettings>();
 	if (PinType.PinCategory == PinCategoryType)
 	{
 		FNiagaraTypeDefinition Type(CastChecked<UScriptStruct>(PinType.PinSubCategoryObject.Get()));
+		return GetTypeColor(Type);
+	}
+		
+	const UGraphEditorSettings* Settings = GetDefault<UGraphEditorSettings>();
+	return Settings->WildcardPinTypeColor;
+}
 
-		if (Type == FNiagaraTypeDefinition::GetFloatDef())
-		{
-			return Settings->FloatPinTypeColor;
-		}
-		else if (Type == FNiagaraTypeDefinition::GetIntDef())
-		{
-			return Settings->IntPinTypeColor;
-		}
-		else if (Type == FNiagaraTypeDefinition::GetBoolDef())
-		{
-			return Settings->BooleanPinTypeColor;
-		}
-		else if (Type == FNiagaraTypeDefinition::GetVec2Def()
-			|| Type == FNiagaraTypeDefinition::GetVec3Def()
-			|| Type == FNiagaraTypeDefinition::GetVec4Def())
-		{
-			return Settings->VectorPinTypeColor;
-		}
-		else
-		{
-			return Settings->StructPinTypeColor;
-		}
+FLinearColor UEdGraphSchema_Niagara::GetTypeColor(const FNiagaraTypeDefinition& Type)
+{
+	const UGraphEditorSettings* Settings = GetDefault<UGraphEditorSettings>();
+	if (Type == FNiagaraTypeDefinition::GetFloatDef())
+	{
+		return Settings->FloatPinTypeColor;
+	}
+	else if (Type == FNiagaraTypeDefinition::GetIntDef())
+	{
+		return Settings->IntPinTypeColor;
+	}
+	else if (Type == FNiagaraTypeDefinition::GetBoolDef())
+	{
+		return Settings->BooleanPinTypeColor;
+	}
+	else if (Type == FNiagaraTypeDefinition::GetVec3Def())
+	{
+		return Settings->VectorPinTypeColor;
+	}
+	else if (Type == FNiagaraTypeDefinition::GetParameterMapDef())
+	{
+		return Settings->ExecutionPinTypeColor;
 	}
 	else
 	{
-		return Settings->WildcardPinTypeColor;
+		return Settings->StructPinTypeColor;
 	}
 }
 
@@ -1041,7 +1144,7 @@ FNiagaraVariable UEdGraphSchema_Niagara::PinToNiagaraVariable(const UEdGraphPin*
 	if (Pin->bDefaultValueIsIgnored == false && Pin->DefaultValue.IsEmpty() == false)
 	{
 		FNiagaraEditorModule& NiagaraEditorModule = FModuleManager::GetModuleChecked<FNiagaraEditorModule>("NiagaraEditor");
-		TSharedPtr<INiagaraEditorTypeUtilities> TypeEditorUtilities = NiagaraEditorModule.GetTypeUtilities(Var.GetType());
+		TSharedPtr<INiagaraEditorTypeUtilities, ESPMode::ThreadSafe> TypeEditorUtilities = NiagaraEditorModule.GetTypeUtilities(Var.GetType());
 		if (TypeEditorUtilities.IsValid() && TypeEditorUtilities->CanHandlePinDefaults())
 		{
 			bHasValue = TypeEditorUtilities->SetValueFromPinDefaultString(Pin->DefaultValue, Var);
@@ -1053,14 +1156,21 @@ FNiagaraVariable UEdGraphSchema_Niagara::PinToNiagaraVariable(const UEdGraphPin*
 		}
 		else
 		{
-			FString OwningNodePath = Pin->GetOwningNode() != nullptr ? Pin->GetOwningNode()->GetPathName() : TEXT("Unknown");
-			UE_LOG(LogNiagaraEditor, Error, TEXT("Pin had default value string, but default values aren't supported for variables of type {%s}. Owning node path: %s"), *Var.GetType().GetName(), *OwningNodePath);
+			if (Pin->GetOwningNode() != nullptr && nullptr == Cast<UNiagaraNodeOp>(Pin->GetOwningNode()))
+			{
+				FString OwningNodePath = Pin->GetOwningNode() != nullptr ? Pin->GetOwningNode()->GetPathName() : TEXT("Unknown");
+				UE_LOG(LogNiagaraEditor, Error, TEXT("Pin had default value string, but default values aren't supported for variables of type {%s}. Owning node path: %s"), *Var.GetType().GetName(), *OwningNodePath);
+			}
 		}
 	}
 
 	if (bNeedsValue && bHasValue == false)
 	{
 		FNiagaraEditorUtilities::ResetVariableToDefaultValue(Var);
+		if (Var.GetData() == nullptr)
+		{
+			UE_LOG(LogNiagaraEditor, Error, TEXT("ResetVariableToDefaultValue called, but failed on var %s type %s. "), *Var.GetName().ToString(), *Var.GetType().GetName());
+		}
 	}
 
 	return Var;
@@ -1076,7 +1186,7 @@ bool UEdGraphSchema_Niagara::TryGetPinDefaultValueFromNiagaraVariable(const FNia
 	}
 
 	FNiagaraEditorModule& NiagaraEditorModule = FModuleManager::GetModuleChecked<FNiagaraEditorModule>("NiagaraEditor");
-	TSharedPtr<INiagaraEditorTypeUtilities> TypeEditorUtilities = NiagaraEditorModule.GetTypeUtilities(PinDefaultVariable.GetType());
+	TSharedPtr<INiagaraEditorTypeUtilities, ESPMode::ThreadSafe> TypeEditorUtilities = NiagaraEditorModule.GetTypeUtilities(PinDefaultVariable.GetType());
 	if (TypeEditorUtilities.IsValid() && TypeEditorUtilities->CanHandlePinDefaults())
 	{
 		OutPinDefaultValue = TypeEditorUtilities->GetPinDefaultStringFromValue(PinDefaultVariable);
@@ -1091,7 +1201,7 @@ FNiagaraTypeDefinition UEdGraphSchema_Niagara::PinToTypeDefinition(const UEdGrap
 {
 	if (Pin->PinType.PinCategory == PinCategoryType && Pin->PinType.PinSubCategoryObject != nullptr)
 	{
-		const UScriptStruct* Struct = Cast<const UScriptStruct>(Pin->PinType.PinSubCategoryObject.Get());
+		UScriptStruct* Struct = Cast<UScriptStruct>(Pin->PinType.PinSubCategoryObject.Get());
 		if (Struct == nullptr)
 		{
 			UE_LOG(LogNiagaraEditor, Error, TEXT("Pin states that it is of struct type, but is missing its struct object. This is usually the result of a registered type going away. Pin Name '%s' Owning Node '%s'."),
@@ -1102,7 +1212,7 @@ FNiagaraTypeDefinition UEdGraphSchema_Niagara::PinToTypeDefinition(const UEdGrap
 	}
 	else if (Pin->PinType.PinCategory == PinCategoryClass)
 	{
-		const UClass* Class = Cast<const UClass>(Pin->PinType.PinSubCategoryObject.Get());
+		UClass* Class = Cast<UClass>(Pin->PinType.PinSubCategoryObject.Get());
 		if (Class == nullptr)
 		{
 			UE_LOG(LogNiagaraEditor, Error, TEXT("Pin states that it is of class type, but is missing its class object. This is usually the result of a registered type going away. Pin Name '%s' Owning Node '%s'."),
@@ -1113,7 +1223,7 @@ FNiagaraTypeDefinition UEdGraphSchema_Niagara::PinToTypeDefinition(const UEdGrap
 	}
 	else if (Pin->PinType.PinCategory == PinCategoryEnum)
 	{
-		const UEnum* Enum = Cast<const UEnum>(Pin->PinType.PinSubCategoryObject.Get());
+		UEnum* Enum = Cast<UEnum>(Pin->PinType.PinSubCategoryObject.Get());
 		if (Enum == nullptr)
 		{
 			UE_LOG(LogNiagaraEditor, Error, TEXT("Pin states that it is of Enum type, but is missing its Enum! Pin Name '%s' Owning Node '%s'. Turning into standard int definition!"), *Pin->PinName.ToString(),
@@ -1170,6 +1280,47 @@ UNiagaraParameterCollection* UEdGraphSchema_Niagara::VariableIsFromParameterColl
 	return nullptr;
 }
 
+UNiagaraParameterCollection* UEdGraphSchema_Niagara::VariableIsFromParameterCollection(const FString& VarName, bool bAllowPartialMatch, FNiagaraVariable& OutVar)const
+{
+	OutVar = FNiagaraVariable();
+
+	if (VarName.StartsWith(TEXT("NPC.")))
+	{
+		FAssetRegistryModule& AssetRegistryModule = FModuleManager::LoadModuleChecked<FAssetRegistryModule>("AssetRegistry");
+		TArray<FAssetData> CollectionAssets;
+		AssetRegistryModule.Get().GetAssetsByClass(UNiagaraParameterCollection::StaticClass()->GetFName(), CollectionAssets);
+		TArray<FName> ExistingNames;
+		for (FAssetData& CollectionAsset : CollectionAssets)
+		{
+			if (UNiagaraParameterCollection* Collection = CastChecked<UNiagaraParameterCollection>(CollectionAsset.GetAsset()))
+			{
+				if (VarName.StartsWith(Collection->GetFullNamespace()))
+				{
+					const TArray<FNiagaraVariable>& CollectionVariables = Collection->GetParameters();
+					FString BestMatchSoFar;
+
+					for (const FNiagaraVariable& CollVar : CollectionVariables)
+					{
+						FString CollVarName = CollVar.GetName().ToString();
+						if (CollVarName == VarName)
+						{
+							OutVar = CollVar;
+							break;
+						}
+						else if (bAllowPartialMatch && VarName.StartsWith(CollVarName + TEXT(".")) && (BestMatchSoFar.Len() == 0 || CollVarName.Len() > BestMatchSoFar.Len()))
+						{
+							OutVar = CollVar;
+							BestMatchSoFar = CollVarName;
+						}
+					}
+					return Collection;
+				}
+			}
+		}
+	}
+	return nullptr;
+}
+
 FNiagaraTypeDefinition UEdGraphSchema_Niagara::GetTypeDefForProperty(const UProperty* Property)const
 {
 	if (Property->IsA(UFloatProperty::StaticClass()))
@@ -1282,14 +1433,14 @@ void UEdGraphSchema_Niagara::ToggleNodeEnabledState(UNiagaraNode* InNode) const
 			const FScopedTransaction Transaction(NSLOCTEXT("UnrealEd", "NiagaraEditorSetNodeEnabled", "Enabled Node"));
 			InNode->Modify();
 			InNode->SetEnabledState(ENodeEnabledState::Enabled, true);
-			InNode->GetNiagaraGraph()->NotifyGraphNeedsRecompile();
+			InNode->MarkNodeRequiresSynchronization(__FUNCTION__, true);
 		}
 		else if (InNode->GetDesiredEnabledState() == ENodeEnabledState::Enabled)
 		{
 			const FScopedTransaction Transaction(NSLOCTEXT("UnrealEd", "NiagaraEditorSetNodeDisabled", "Disabled Node"));
 			InNode->Modify();
 			InNode->SetEnabledState(ENodeEnabledState::Disabled, true);
-			InNode->GetNiagaraGraph()->NotifyGraphNeedsRecompile();
+			InNode->MarkNodeRequiresSynchronization(__FUNCTION__, true);
 		}
 	}
 }
@@ -1383,9 +1534,12 @@ void UEdGraphSchema_Niagara::GetContextMenuActions(const UEdGraph* CurrentGraph,
 	Super::GetContextMenuActions(CurrentGraph, InGraphNode, InGraphPin, MenuBuilder, bIsDebugging);
 }
 
-FNiagaraConnectionDrawingPolicy::FNiagaraConnectionDrawingPolicy(int32 InBackLayerID, int32 InFrontLayerID, float InZoomFactor, const FSlateRect& InClippingRect, FSlateWindowElementList& InDrawElements)
+FNiagaraConnectionDrawingPolicy::FNiagaraConnectionDrawingPolicy(int32 InBackLayerID, int32 InFrontLayerID, float InZoomFactor, const FSlateRect& InClippingRect, FSlateWindowElementList& InDrawElements, UEdGraph* InGraph)
 	: FConnectionDrawingPolicy(InBackLayerID, InFrontLayerID, InZoomFactor, InClippingRect, InDrawElements)
+	, Graph(InGraph)
 {
+	ArrowImage = nullptr;
+	ArrowRadius = FVector2D::ZeroVector;
 }
 
 void FNiagaraConnectionDrawingPolicy::DetermineWiringStyle(UEdGraphPin* OutputPin, UEdGraphPin* InputPin, /*inout*/ FConnectionParams& Params)
@@ -1395,6 +1549,12 @@ void FNiagaraConnectionDrawingPolicy::DetermineWiringStyle(UEdGraphPin* OutputPi
 	{
 		Params.WireThickness = Params.WireThickness * 5;
 	}
+
+	if (Graph)
+	{
+		const UEdGraphSchema* Schema = Graph->GetSchema();
+		Params.WireColor = Schema->GetPinTypeColor(OutputPin->PinType);
+	}
 }
 
 #undef LOCTEXT_NAMESPACE
