Copying //UE4/Dev-Niagara to //UE4/Dev-Main (Source: //UE4/Dev-Niagara @ 4074996)

diff --git a/Engine/Plugins/FX/Niagara/Source/NiagaraEditor/Private/NiagaraCompiler.cpp b/Engine/Plugins/FX/Niagara/Source/NiagaraEditor/Private/NiagaraCompiler.cpp
--- a/Engine/Plugins/FX/Niagara/Source/NiagaraEditor/Private/NiagaraCompiler.cpp
+++ b/Engine/Plugins/FX/Niagara/Source/NiagaraEditor/Private/NiagaraCompiler.cpp
@@ -13,12 +13,20 @@
 #include "NiagaraSystem.h"
 #include "NiagaraNodeEmitter.h"
 #include "NiagaraNodeInput.h"
-
+#include "NiagaraScriptSource.h"
 #include "NiagaraDataInterface.h"
 #include "NiagaraDataInterfaceStaticMesh.h"
 #include "NiagaraDataInterfaceCurlNoise.h"
-
+#include "ViewModels/Stack/NiagaraStackGraphUtilities.h"
+#include "NiagaraNodeFunctionCall.h"
+#include "NiagaraNodeParameterMapSet.h"
+#include "INiagaraEditorTypeUtilities.h"
+#include "NiagaraEditorUtilities.h"
+#include "NiagaraNodeEmitter.h"
+#include "NiagaraNodeOutput.h"
 #include "ShaderCore.h"
+#include "EdGraphSchema_Niagara.h"
+#include "Misc/FileHelper.h"
 
 #define LOCTEXT_NAMESPACE "NiagaraCompiler"
 
@@ -28,63 +36,481 @@ DECLARE_CYCLE_STAT(TEXT("Niagara - Module - CompileScript"), STAT_NiagaraEditor_
 DECLARE_CYCLE_STAT(TEXT("Niagara - HlslCompiler - CompileScript"), STAT_NiagaraEditor_HlslCompiler_CompileScript, STATGROUP_NiagaraEditor);
 DECLARE_CYCLE_STAT(TEXT("Niagara - HlslCompiler - CompileShader_VectorVM"), STAT_NiagaraEditor_HlslCompiler_CompileShader_VectorVM, STATGROUP_NiagaraEditor);
 DECLARE_CYCLE_STAT(TEXT("Niagara - Module - CompileShader_VectorVMSucceeded"), STAT_NiagaraEditor_HlslCompiler_CompileShader_VectorVMSucceeded, STATGROUP_NiagaraEditor);
+DECLARE_CYCLE_STAT(TEXT("Niagara - ScriptSource - PreCompile"), STAT_NiagaraEditor_ScriptSource_PreCompile, STATGROUP_NiagaraEditor);
+DECLARE_CYCLE_STAT(TEXT("Niagara - HlslCompiler - TestCompileShader_VectorVM"), STAT_NiagaraEditor_HlslCompiler_TestCompileShader_VectorVM, STATGROUP_NiagaraEditor);
 
-ENiagaraScriptCompileStatus FNiagaraEditorModule::CompileScript(UNiagaraScript* ScriptToCompile, FString& OutGraphLevelErrorMessages)
+static int32 GbForceNiagaraTranslatorSingleThreaded = 0;
+static FAutoConsoleVariableRef CVarForceNiagaraTranslatorSingleThreaded(
+	TEXT("fx.ForceNiagaraTranslatorSingleThreaded"),
+	GbForceNiagaraTranslatorSingleThreaded,
+	TEXT("If > 0 all translation will occur one at a time, useful for debugging. \n"),
+	ECVF_Default
+);
+
+
+static FCriticalSection TranslationCritSec;
+
+void FNiagaraCompileRequestData::VisitReferencedGraphs(UNiagaraGraph* InSrcGraph, UNiagaraGraph* InDupeGraph, ENiagaraScriptUsage InUsage)
 {
-	SCOPE_CYCLE_COUNTER(STAT_NiagaraEditor_Module_CompileScript);
-	check(ScriptToCompile != NULL);
+	if (!InDupeGraph && !InSrcGraph)
+	{
+		return;
+	}
+	FunctionData Data;
+	Data.ClonedScript = nullptr;
+	Data.ClonedGraph = InDupeGraph;
+	Data.Usage = InUsage;
+	Data.bHasNumericInputs = false;
+	TArray<FunctionData> MadeArray;
+	MadeArray.Add(Data);
+	PreprocessedFunctions.Add(InSrcGraph, MadeArray);
+
+	bool bStandaloneScript = false;
+
+	TArray<UNiagaraNodeOutput*> OutputNodes;
+	InDupeGraph->FindOutputNodes(OutputNodes);
+	if (OutputNodes.Num() == 1 && UNiagaraScript::IsStandaloneScript(OutputNodes[0]->GetUsage()))
+	{
+		bStandaloneScript = true;
+	}
+	FNiagaraEditorUtilities::ResolveNumerics(InDupeGraph, bStandaloneScript, ChangedFromNumericVars);
+	ClonedGraphs.AddUnique(InDupeGraph);
 
-	UNiagaraGraph* Graph = Cast<UNiagaraScriptSource>(ScriptToCompile->GetSource())->NodeGraph;
+	VisitReferencedGraphsRecursive(InDupeGraph);
+}
 
-	OutGraphLevelErrorMessages.Empty();
-	FNiagaraCompileResults Results;
-	FHlslNiagaraCompiler Compiler;
-	FNiagaraTranslateResults TranslateResults;
-	FHlslNiagaraTranslator Translator;
+void FNiagaraCompileRequestData::VisitReferencedGraphsRecursive(UNiagaraGraph* InGraph)
+{
+	if (!InGraph)
+	{
+		return;
+	}
+	UPackage* OwningPackage = InGraph->GetOutermost();
+
+	TArray<UNiagaraNode*> Nodes;
+	InGraph->GetNodesOfClass(Nodes);
+	const UEdGraphSchema_Niagara* Schema = GetDefault<UEdGraphSchema_Niagara>();
 
-	/*
-	if (ScriptToCompile && ScriptToCompile->IsNonParticleScript())
+	for (UEdGraphNode* Node : Nodes)
 	{
-		ENiagaraScriptCompileStatus Status = ENiagaraScriptCompileStatus::NCS_UpToDate;
-		if (ScriptToCompile->IsSystemSpawnScript())
+		if (UNiagaraNode* InNode = Cast<UNiagaraNode>(Node))
 		{
-			Status = CompileSystemScript(ScriptToCompile, OutGraphLevelErrorMessages);
+			UNiagaraNodeInput* InputNode = Cast<UNiagaraNodeInput>(InNode);
+			if (InputNode)
+			{
+				if (InputNode->Input.IsDataInterface())
+				{
+					UNiagaraDataInterface* DataInterface = InputNode->GetDataInterface();
+					FName DIName = InputNode->Input.GetName();
+					UNiagaraDataInterface* Dupe = DuplicateObject<UNiagaraDataInterface>(DataInterface, GetTransientPackage());
+					CopiedDataInterfacesByName.Add(DIName, Dupe);
+				}
+				continue;
+			}
+
+			ENiagaraScriptUsage ScriptUsage = ENiagaraScriptUsage::Function;
+
+			UNiagaraNodeFunctionCall* FunctionCallNode = Cast<UNiagaraNodeFunctionCall>(InNode);
+			if (FunctionCallNode)
+			{
+				UNiagaraScript* FunctionScript = FunctionCallNode->FunctionScript;
+				ScriptUsage = FunctionCallNode->GetCalledUsage();
+
+				if (FunctionScript != nullptr)
+				{
+					UNiagaraGraph* FunctionGraph = FunctionCallNode->GetCalledGraph();
+					{
+						bool bHasNumericParams = FunctionGraph->HasNumericParameters();
+						bool bHasNumericInputs = false;
+
+						UPackage* FunctionPackage = FunctionGraph->GetOutermost();
+						bool bFromDifferentPackage = OwningPackage != FunctionPackage;
+
+						TArray<UEdGraphPin*> CallOutputs;
+						TArray<UEdGraphPin*> CallInputs;
+						InNode->GetOutputPins(CallOutputs);
+						InNode->GetInputPins(CallInputs);
+
+						TArray<UNiagaraNodeInput*> InputNodes;
+						UNiagaraGraph::FFindInputNodeOptions Options;
+						Options.bFilterDuplicates = true;
+						Options.bIncludeParameters = true;
+						Options.bIncludeAttributes = false;
+						Options.bIncludeSystemConstants = false;
+						Options.bIncludeTranslatorConstants = false;
+						FunctionGraph->FindInputNodes(InputNodes, Options);
+
+						for (UNiagaraNodeInput* Input : InputNodes)
+						{
+							if (Input->Input.GetType() == FNiagaraTypeDefinition::GetGenericNumericDef())
+							{
+								bHasNumericInputs = true;
+							}
+						}
+
+						/*UE_LOG(LogNiagaraEditor, Log, TEXT("FunctionGraph = %p SrcScriptName = %s DiffPackage? %s Numerics? %s NumericInputs? %s"), FunctionGraph, *FunctionGraph->GetPathName(),
+							bFromDifferentPackage ? TEXT("yes") : TEXT("no"), bHasNumericParams ? TEXT("yes") : TEXT("no"), bHasNumericInputs ? TEXT("yes") : TEXT("no"));*/
+
+						UNiagaraGraph* ProcessedGraph = nullptr;
+						// We only need to clone a non-numeric graph once.
+
+						if (!PreprocessedFunctions.Contains(FunctionGraph))
+						{
+							UNiagaraScript* DupeScript = nullptr;
+							if (!bFromDifferentPackage && !bHasNumericInputs && !bHasNumericParams)
+							{
+								DupeScript = FunctionScript;
+								ProcessedGraph = FunctionGraph;
+							}
+							else
+							{
+								DupeScript = DuplicateObject<UNiagaraScript>(FunctionScript, InNode, FunctionScript->GetFName());
+								ProcessedGraph = Cast<UNiagaraScriptSource>(DupeScript->GetSource())->NodeGraph;
+								FEdGraphUtilities::MergeChildrenGraphsIn(ProcessedGraph, ProcessedGraph, /*bRequireSchemaMatch=*/ true);
+								FNiagaraEditorUtilities::PreprocessFunctionGraph(Schema, ProcessedGraph, CallInputs, CallOutputs, ScriptUsage);
+								FunctionCallNode->FunctionScript = DupeScript;
+
+								/*UE_LOG(LogNiagaraEditor, Log, TEXT("Duplicating script %s OriginalGraph %p NewGraph %p"), *FunctionScript->GetFName().ToString(),
+									Cast<UNiagaraScriptSource>(FunctionScript->GetSource())->NodeGraph,
+									Cast<UNiagaraScriptSource>(DupeScript->GetSource())->NodeGraph);*/
+							}
+
+							FunctionData Data;
+							Data.ClonedScript = DupeScript;
+							Data.ClonedGraph = ProcessedGraph;
+							Data.CallInputs = CallInputs;
+							Data.CallOutputs = CallOutputs;
+							Data.Usage = ScriptUsage;
+							Data.bHasNumericInputs = bHasNumericInputs;
+							TArray<FunctionData> MadeArray;
+							MadeArray.Add(Data);
+							PreprocessedFunctions.Add(FunctionGraph, MadeArray);
+							VisitReferencedGraphsRecursive(ProcessedGraph);		
+							ClonedGraphs.AddUnique(ProcessedGraph);
+
+						}
+						else if (bHasNumericParams)
+						{
+							UNiagaraScript* DupeScript = DuplicateObject<UNiagaraScript>(FunctionScript, InNode, FunctionScript->GetFName());
+							ProcessedGraph = Cast<UNiagaraScriptSource>(DupeScript->GetSource())->NodeGraph;
+							FEdGraphUtilities::MergeChildrenGraphsIn(ProcessedGraph, ProcessedGraph, /*bRequireSchemaMatch=*/ true);
+							FNiagaraEditorUtilities::PreprocessFunctionGraph(Schema, ProcessedGraph, CallInputs, CallOutputs, ScriptUsage);
+							FunctionCallNode->FunctionScript = DupeScript;
+
+							/*UE_LOG(LogNiagaraEditor, Log, TEXT("Duplicating script %s OriginalGraph %p NewGraph %p"), *FunctionScript->GetFName().ToString(),
+								Cast<UNiagaraScriptSource>(FunctionScript->GetSource())->NodeGraph,
+								Cast<UNiagaraScriptSource>(DupeScript->GetSource())->NodeGraph);*/
+
+							TArray<FunctionData>* FoundArray = PreprocessedFunctions.Find(FunctionGraph);
+							ClonedGraphs.AddUnique(ProcessedGraph);
+
+							FunctionData Data;
+							Data.ClonedScript = DupeScript;
+							Data.ClonedGraph = ProcessedGraph;
+							Data.CallInputs = CallInputs;
+							Data.CallOutputs = CallOutputs;
+							Data.Usage = ScriptUsage;
+							Data.bHasNumericInputs = bHasNumericInputs;
+
+							FoundArray->Add(Data);
+							VisitReferencedGraphsRecursive(ProcessedGraph);
+						}
+						else if (bFromDifferentPackage)
+						{
+							TArray<FunctionData>* FoundArray = PreprocessedFunctions.Find(FunctionGraph);
+							check(FoundArray != nullptr && FoundArray->Num() != 0);
+							FunctionCallNode->FunctionScript = (*FoundArray)[0].ClonedScript;
+						}
+					}
+				}
+			}
+
+			UNiagaraNodeEmitter* EmitterNode = Cast<UNiagaraNodeEmitter>(InNode);
+			if (EmitterNode)
+			{
+				for (TSharedPtr<FNiagaraCompileRequestData, ESPMode::ThreadSafe>& Ptr : EmitterData)
+				{
+					if (Ptr->EmitterUniqueName == EmitterNode->GetEmitterUniqueName())
+					{
+						EmitterNode->SetOwnerSystem(nullptr);
+						EmitterNode->SetCachedVariablesForCompilation(*Ptr->EmitterUniqueName, Ptr->NodeGraphDeepCopy, Ptr->Source);
+					}
+				}
+			}
 		}
-		ScriptToCompile->SetChangeID(Graph->GetChangeID());
-		ScriptToCompile->SetLastCompileStatus(Status);
-		return Status;
 	}
-	else
-	*/
+}
+
+const TMap<FName, UNiagaraDataInterface*>& FNiagaraCompileRequestData::GetObjectNameMap()
+{
+	return CopiedDataInterfacesByName;
+}
+
+void FNiagaraCompileRequestData::MergeInEmitterPrecompiledData(FNiagaraCompileRequestDataBase* InEmitterDataBase)
+{
+//	check(EmitterData.Contains(InEmitterDataBase));
+
+	FNiagaraCompileRequestData* InEmitterData = (FNiagaraCompileRequestData*)InEmitterDataBase;
+	if (InEmitterData)
 	{
-		FHlslNiagaraTranslatorOptions Options; 
-		Options.SimTarget = ENiagaraSimTarget::CPUSim;
-		TranslateResults = Translator.Translate(ScriptToCompile, Options);
-		Results = Compiler.CompileScript(ScriptToCompile, &Translator.GetTranslateOutput(), Translator.GetTranslatedHLSL());
-		ScriptToCompile->SetChangeID(Graph->GetChangeID());
-		ScriptToCompile->GenerateStatScopeIDs();
+		{
+			auto It = InEmitterData->CopiedDataInterfacesByName.CreateIterator();
+			while (It)
+			{
+				FName Name = It.Key();
+				Name = FNiagaraParameterMapHistory::ResolveEmitterAlias(Name, InEmitterData->GetUniqueEmitterName());
+				CopiedDataInterfacesByName.Add(Name, It.Value());
+				++It;
+			}
+		}
 	}
+}
+
+FName FNiagaraCompileRequestData::ResolveEmitterAlias(FName VariableName) const
+{
+	return FNiagaraParameterMapHistory::ResolveEmitterAlias(VariableName, EmitterUniqueName);
+}
 
-#if defined(NIAGARA_SCRIPT_COMPILE_LOGGING_MEDIUM)
-	if (TranslateResults.bHLSLGenSucceeded)
+void FNiagaraCompileRequestData::GetReferencedObjects(TArray<UObject*>& Objects)
+{
+	Objects.Add(NodeGraphDeepCopy);
+	TArray<UNiagaraDataInterface*> DIs;
+	CopiedDataInterfacesByName.GenerateValueArray(DIs);
+	for (UNiagaraDataInterface* DI : DIs)
 	{
-		UE_LOG(LogNiagaraCompiler, Log, TEXT("HLSL generation succeeded: %s"), *ScriptToCompile->GetPathName());
+		Objects.Add(DI);
 	}
-	else
+
+	{
+		auto Iter = CDOs.CreateIterator();
+		while (Iter)
+		{
+			Objects.Add(Iter.Value());
+			++Iter;
+		}
+	}
+
 	{
-		UE_LOG(LogNiagaraCompiler, Error, TEXT("HLSL generation failed: %s"), *ScriptToCompile->GetPathName());
+		auto Iter = PreprocessedFunctions.CreateIterator();
+		while (Iter)
+		{
+			for (int32 i = 0; i < Iter.Value().Num(); i++)
+			{
+				Objects.Add(Iter.Value()[i].ClonedScript);
+			}
+			++Iter;
+		}
 	}
+}
 
-	if (Results.bVMSucceeded)
+bool FNiagaraCompileRequestData::GatherPreCompiledVariables(const FString& InNamespaceFilter, TArray<FNiagaraVariable>& OutVars)
+{
+	if (PrecompiledHistories.Num() == 0)
 	{
-		UE_LOG(LogNiagaraCompiler, Log, TEXT("CPU Compile succeeded: %s"), *ScriptToCompile->GetPathName());
+		return false;
 	}
-	else
+
+	for (const FNiagaraParameterMapHistory& History : PrecompiledHistories)
+	{
+		for (const FNiagaraVariable& Var : History.Variables)
+		{
+			if (FNiagaraParameterMapHistory::IsInNamespace(Var, InNamespaceFilter))
+			{
+				FNiagaraVariable NewVar = Var;
+				if (NewVar.IsDataAllocated() == false && !Var.IsDataInterface())
+				{
+					FNiagaraEditorUtilities::ResetVariableToDefaultValue(NewVar);
+				}
+				OutVars.AddUnique(NewVar);
+			}
+		}
+	}
+	return true;
+}
+
+void FNiagaraCompileRequestData::DeepCopyGraphs(UNiagaraScriptSource* ScriptSource, ENiagaraScriptUsage InUsage)
+{
+	// Clone the source graph so we can modify it as needed; merging in the child graphs
+	//NodeGraphDeepCopy = CastChecked<UNiagaraGraph>(FEdGraphUtilities::CloneGraph(ScriptSource->NodeGraph, GetTransientPackage(), 0));
+	Source = DuplicateObject<UNiagaraScriptSource>(ScriptSource, GetTransientPackage());
+	NodeGraphDeepCopy = Source->NodeGraph;
+	FEdGraphUtilities::MergeChildrenGraphsIn(NodeGraphDeepCopy, NodeGraphDeepCopy, /*bRequireSchemaMatch=*/ true);
+	VisitReferencedGraphs(ScriptSource->NodeGraph, NodeGraphDeepCopy, InUsage /*InEmitter ? ENiagaraScriptUsage::EmitterSpawnScript : ENiagaraScriptUsage::SystemSpawnScript*/);
+}
+
+
+void FNiagaraCompileRequestData::FinishPrecompile(UNiagaraScriptSource* ScriptSource, const TArray<FNiagaraVariable>& EncounterableVariables, ENiagaraScriptUsage InUsage)
+{
+	{
+		ENiagaraScriptCompileStatusEnum = FindObject<UEnum>(ANY_PACKAGE, TEXT("ENiagaraScriptCompileStatus"), true);
+
+		PrecompiledHistories.Empty();
+
+		TArray<UNiagaraNodeOutput*> OutputNodes;
+		NodeGraphDeepCopy->FindOutputNodes(OutputNodes);
+		PrecompiledHistories.Empty();
+
+		for (UNiagaraNodeOutput* FoundOutputNode : OutputNodes)
+		{
+			// Map all for this output node
+			FNiagaraParameterMapHistoryBuilder Builder;
+			Builder.RegisterEncounterableVariables(EncounterableVariables);
+
+			FString TranslationName = TEXT("Emitter");
+			Builder.BeginTranslation(TranslationName);
+			Builder.EnableScriptWhitelist(true, FoundOutputNode->GetUsage());
+			Builder.BuildParameterMaps(FoundOutputNode, true);
+			TArray<FNiagaraParameterMapHistory> Histories = Builder.Histories;
+			ensure(Histories.Num() <= 1);
+
+			for (FNiagaraParameterMapHistory& History : Histories)
+			{
+				for (FNiagaraVariable& Var : History.Variables)
+				{
+					check(Var.GetType() != FNiagaraTypeDefinition::GetGenericNumericDef());
+				}
+			}
+
+			PrecompiledHistories.Append(Histories);
+			Builder.EndTranslation(TranslationName);
+		}
+
+		// Generate CDO's for any referenced data interfaces...
+		for (int32 i = 0; i < PrecompiledHistories.Num(); i++)
+		{
+			for (const FNiagaraVariable& Var : PrecompiledHistories[i].Variables)
+			{
+				if (Var.IsDataInterface())
+				{
+					UClass* Class = const_cast<UClass*>(Var.GetType().GetClass());
+					UObject* Obj = DuplicateObject(Class->GetDefaultObject(true), GetTransientPackage());
+					CDOs.Add(Class, Obj);
+				}
+			}
+		}
+
+	}
+}
+
+TSharedPtr<FNiagaraCompileRequestDataBase, ESPMode::ThreadSafe> FNiagaraEditorModule::Precompile(UObject* InObj)
+{
+	UNiagaraScript* Script = Cast<UNiagaraScript>(InObj);
+	UNiagaraSystem* System = Cast<UNiagaraSystem>(InObj);
+
+	if (!Script && !System)
+	{
+		TSharedPtr<FNiagaraCompileRequestDataBase, ESPMode::ThreadSafe> InvalidPtr;
+		return InvalidPtr;
+	}
+
+	SCOPE_CYCLE_COUNTER(STAT_NiagaraEditor_ScriptSource_PreCompile);
+	double StartTime = FPlatformTime::Seconds();
+
+	TSharedPtr<FNiagaraCompileRequestData, ESPMode::ThreadSafe> BasePtr = MakeShared<FNiagaraCompileRequestData, ESPMode::ThreadSafe>();
+	TArray<TSharedPtr<FNiagaraCompileRequestData, ESPMode::ThreadSafe>> DependentRequests;
+
+	if (Script)
 	{
-		UE_LOG(LogNiagaraCompiler, Error, TEXT("CPU Compile failed: %s"), *ScriptToCompile->GetPathName());
+		UNiagaraScriptSource* Source = Cast<UNiagaraScriptSource>(Script->GetSource());
+		BasePtr->DeepCopyGraphs(Source, Script->GetUsage());
+		const TArray<FNiagaraVariable> EncounterableVariables;
+		BasePtr->FinishPrecompile(Source, EncounterableVariables, Script->GetUsage());
+	}
+	else if (System)
+	{
+		// Store off the current variables in the exposed parameters list.
+		TArray<FNiagaraVariable> OriginalExposedParams;
+		System->GetExposedParameters().GetParameters(OriginalExposedParams);
+
+		// Create an array of variables that we might encounter when traversing the graphs (include the originally exposed vars above)
+		TArray<FNiagaraVariable> EncounterableVars(OriginalExposedParams);
+
+		// Create an array of variables that we *did* encounter when traversing the graphs.
+		TArray<FNiagaraVariable> EncounteredExposedVars;
+		check(System->GetSystemSpawnScript()->GetSource() == System->GetSystemUpdateScript()->GetSource());
+
+		// First deep copy all the emitter graphs referenced by the system so that we can later hook up emitter handles in the system traversal.
+		BasePtr->EmitterData.Empty();
+		for (int32 i = 0; i < System->GetEmitterHandles().Num(); i++)
+		{
+			const FNiagaraEmitterHandle& Handle = System->GetEmitterHandle(i);
+			TSharedPtr<FNiagaraCompileRequestData, ESPMode::ThreadSafe> EmitterPtr = MakeShared<FNiagaraCompileRequestData, ESPMode::ThreadSafe>();
+			EmitterPtr->DeepCopyGraphs(Cast<UNiagaraScriptSource>(Handle.GetInstance()->GraphSource), ENiagaraScriptUsage::EmitterSpawnScript);
+			EmitterPtr->EmitterUniqueName = Handle.GetInstance()->GetUniqueEmitterName();
+			BasePtr->EmitterData.Add(EmitterPtr);
+		}
+
+		// Now deep copy the system graphs, skipping traversal into any emitter references.
+		UNiagaraScriptSource* Source = Cast<UNiagaraScriptSource>(System->GetSystemSpawnScript()->GetSource());
+		BasePtr->DeepCopyGraphs(Source, ENiagaraScriptUsage::SystemSpawnScript);
+		BasePtr->FinishPrecompile(Source, EncounterableVars, ENiagaraScriptUsage::SystemSpawnScript);
+
+		// Add the User and System variables that we did encounter to the list that emitters might also encounter.
+		BasePtr->GatherPreCompiledVariables(TEXT("User"), EncounterableVars);
+		BasePtr->GatherPreCompiledVariables(TEXT("System"), EncounterableVars);
+
+		// Now we can finish off the emitters.
+		for (int32 i = 0; i < System->GetEmitterHandles().Num(); i++)
+		{
+			const FNiagaraEmitterHandle& Handle = System->GetEmitterHandle(i);
+			BasePtr->EmitterData[i]->FinishPrecompile(Cast<UNiagaraScriptSource>(Handle.GetInstance()->GraphSource), EncounterableVars, ENiagaraScriptUsage::EmitterSpawnScript);
+			BasePtr->MergeInEmitterPrecompiledData(BasePtr->EmitterData[i].Get());
+		}
+
 	}
-#endif
 
+	UE_LOG(LogNiagaraEditor, Log, TEXT("'%s' Precompile took %f sec."), *InObj->GetOutermost()->GetName(),
+		(float)(FPlatformTime::Seconds() - StartTime));
+
+	return BasePtr;
+}
+
+TSharedPtr<FNiagaraVMExecutableData> FNiagaraEditorModule::CompileScript(const FNiagaraCompileRequestDataBase* InCompileRequest, const FNiagaraCompileOptions& InCompileOptions)
+{
+	SCOPE_CYCLE_COUNTER(STAT_NiagaraEditor_Module_CompileScript);
+
+	double StartTime = FPlatformTime::Seconds();
+
+	check(InCompileRequest != NULL);
+	const FNiagaraCompileRequestData* CompileRequest = (const FNiagaraCompileRequestData*)InCompileRequest;
+
+	UE_LOG(LogNiagaraEditor, Log, TEXT("Compiling System %s ..................................................................."), *InCompileOptions.FullName);
+
+	FNiagaraCompileResults Results;
+	FHlslNiagaraCompiler Compiler;
+	FNiagaraTranslateResults TranslateResults;
+	FHlslNiagaraTranslator Translator;
+
+	float TranslationTime = 0.0f;
+	float VMCompilationTime = 0.0f;
+	{
+		FHlslNiagaraTranslatorOptions TranslateOptions;
+
+		if (InCompileOptions.TargetUsage == ENiagaraScriptUsage::ParticleGPUComputeScript)
+		{
+			TranslateOptions.SimTarget = ENiagaraSimTarget::GPUComputeSim;
+		}
+		else
+		{
+			TranslateOptions.SimTarget = ENiagaraSimTarget::CPUSim;
+		}
+
+		double TranslationStartTime = FPlatformTime::Seconds();
+		if (GbForceNiagaraTranslatorSingleThreaded > 0)
+		{
+			FScopeLock Lock(&TranslationCritSec);
+			TranslateResults = Translator.Translate(CompileRequest, InCompileOptions, TranslateOptions);
+		}
+		else
+		{
+			TranslateResults = Translator.Translate(CompileRequest, InCompileOptions, TranslateOptions);
+		}
+		TranslationTime = (float)(FPlatformTime::Seconds() - TranslationStartTime);
+
+		Results = Compiler.CompileScript(CompileRequest, InCompileOptions, &Translator.GetTranslateOutput(), Translator.GetTranslatedHLSL());
+		VMCompilationTime = Results.CompileTime;
+	}
+	
 	TArray<TSharedRef<FTokenizedMessage>> Messages;
 	if (TranslateResults.MessageLog)
 	{
@@ -95,6 +521,7 @@ ENiagaraScriptCompileStatus FNiagaraEditorModule::CompileScript(UNiagaraScript*
 		Messages.Append(Results.MessageLog->Messages);
 	}
 
+	FString OutGraphLevelErrorMessages;
 	for (TSharedRef<FTokenizedMessage> Message : Messages)
 	{
 		if (Message->GetSeverity() == EMessageSeverity::Info)
@@ -122,21 +549,63 @@ ENiagaraScriptCompileStatus FNiagaraEditorModule::CompileScript(UNiagaraScript*
 			OutGraphLevelErrorMessages += Message->ToText().ToString();
 		}
 	}
-
-#if defined(NIAGARA_SCRIPT_COMPILE_LOGGING_MEDIUM)
-	UE_LOG(LogNiagaraCompiler, Log, TEXT("Compile output as text:"));
-	UE_LOG(LogNiagaraCompiler, Log, TEXT("==================================================================================="));
-	TArray<FString> OutputByLines;
-	Results.OutputHLSL.ParseIntoArrayLines(OutputByLines, false);
-	for (int32 i = 0; i < OutputByLines.Num(); i++)
+	
+	Results.Data->ErrorMsg = OutGraphLevelErrorMessages;
+	Results.Data->LastCompileStatus = (FNiagaraCompileResults::CompileResultsToSummary(&Results));
+	
+	UEnum* FoundEnum = CompileRequest->ENiagaraScriptCompileStatusEnum;
+	FString ResultsEnum = TEXT("??");
+	if (FoundEnum)
 	{
-		UE_LOG(LogNiagaraCompiler, Log, TEXT("/*%04d*/\t\t%s"), i, *OutputByLines[i]);
+		ResultsEnum = FoundEnum->GetNameStringByValue((int64)Results.Data->LastCompileStatus);
 	}
-	UE_LOG(LogNiagaraCompiler, Log, TEXT("==================================================================================="));
-#endif
+	Results.Data->CompileTime = (float)(FPlatformTime::Seconds() - StartTime);
 
-	ScriptToCompile->SetLastCompileStatus(FNiagaraCompileResults::CompileResultsToSummary(&Results));
-	return ScriptToCompile->GetLastCompileStatus();
+	UE_LOG(LogNiagaraEditor, Log, TEXT("Compiling System %s took %f sec (%f/%f)... Status %s"), *InCompileOptions.FullName, Results.Data->CompileTime,
+		TranslationTime, VMCompilationTime, *ResultsEnum);
+	return Results.Data;
+}
+
+
+void FNiagaraEditorModule::TestCompileScriptFromConsole(const TArray<FString>& Arguments)
+{
+	if (Arguments.Num() == 1)
+	{
+		FString TranslatedHLSL;
+		FFileHelper::LoadFileToString(TranslatedHLSL, *Arguments[0]);
+		if (TranslatedHLSL.Len() != 0)
+		{
+			SCOPE_CYCLE_COUNTER(STAT_NiagaraEditor_HlslCompiler_TestCompileShader_VectorVM);
+			FShaderCompilerInput Input;
+			Input.VirtualSourceFilePath = TEXT("/Engine/Private/NiagaraEmitterInstanceShader.usf");
+			Input.EntryPointName = TEXT("SimulateMain");
+			Input.Environment.SetDefine(TEXT("VM_SIMULATION"), 1);
+			Input.Environment.IncludeVirtualPathToContentsMap.Add(TEXT("/Engine/Generated/NiagaraEmitterInstance.ush"), TranslatedHLSL);
+
+			FShaderCompilerOutput Output;
+			FVectorVMCompilationOutput CompilationOutput;
+			double StartTime = FPlatformTime::Seconds();
+			bool bSucceeded = CompileShader_VectorVM(Input, Output, FString(FPlatformProcess::ShaderDir()), 0, CompilationOutput);
+			float DeltaTime = (float)(FPlatformTime::Seconds() - StartTime);
+
+			if (bSucceeded)
+			{
+				UE_LOG(LogNiagaraCompiler, Log, TEXT("Test compile of %s took %f seconds and succeeded."), *Arguments[0], DeltaTime);
+			}
+			else
+			{
+				UE_LOG(LogNiagaraCompiler, Error, TEXT("Test compile of %s took %f seconds and failed.  Errors: %s"), *Arguments[0], DeltaTime, *CompilationOutput.Errors);
+			}
+		}
+		else
+		{
+			UE_LOG(LogNiagaraCompiler, Error, TEXT("Test compile of %s failed, the file could not be loaded or it was empty."), *Arguments[0]);
+		}
+	}
+	else
+	{
+		UE_LOG(LogNiagaraCompiler, Error, TEXT("Test compile failed, file name argument was missing."));
+	}
 }
 
 
@@ -179,157 +648,167 @@ ENiagaraScriptCompileStatus FNiagaraCompileResults::CompileResultsToSummary(cons
 	return SummaryStatus;
 }
 
-const FNiagaraCompileResults &FHlslNiagaraCompiler::CompileScript(UNiagaraScript* InScript, FNiagaraTranslatorOutput *TranslatorOutput, FString &TranslatedHLSL)
+FNiagaraCompileResults FHlslNiagaraCompiler::CompileScript(const FNiagaraCompileRequestData* InCompileRequest, const FNiagaraCompileOptions& InOptions, FNiagaraTranslatorOutput *TranslatorOutput, FString &TranslatedHLSL)
 {
 	SCOPE_CYCLE_COUNTER(STAT_NiagaraEditor_HlslCompiler_CompileScript);
 
+	CompileResults.Data = MakeShared<FNiagaraVMExecutableData>();
+
 	//TODO: This should probably be done via the same route that other shaders take through the shader compiler etc.
 	//But that adds the complexity of a new shader type, new shader class and a new shader map to contain them etc.
 	//Can do things simply for now.
-	Script = InScript;
-	Script->LastHlslTranslation = TEXT("");
-	FShaderCompilerInput Input;
-	Input.VirtualSourceFilePath = TEXT("/Engine/Private/NiagaraEmitterInstanceShader.usf");
-	Input.EntryPointName = TEXT("SimulateMain");
-	Input.Environment.SetDefine(TEXT("VM_SIMULATION"), 1);
-	Input.Environment.IncludeVirtualPathToContentsMap.Add(TEXT("/Engine/Generated/NiagaraEmitterInstance.usf"), TranslatedHLSL);
-	FShaderCompilerOutput Output;
-
-	FVectorVMCompilationOutput CompilationOutput;
-
-	if (TranslatorOutput != nullptr && TranslatorOutput->Errors.Len() > 0)
-	{
-		//TODO: Map Lines of HLSL to their source Nodes and flag those nodes with errors associated with their lines.
-		Error(FText::Format(LOCTEXT("HlslTranslateErrorMessageFormat", "The HLSL Translator failed.  Errors:\n{0}"), FText::FromString(TranslatorOutput->Errors)));
-		CompileResults.bVMSucceeded = false;
-	}
-	else if (TranslatedHLSL.Len() == 0)
+	
+	CompileResults.Data->LastHlslTranslation = TEXT("");
+
+	if (InOptions.TargetUsage == ENiagaraScriptUsage::ParticleGPUComputeScript)
 	{
-		Error(LOCTEXT("HlslTranslateErrorMessageFormat", "The HLSL Translator failed to generate HLSL!"));
-		CompileResults.bVMSucceeded = false;
+		CompileResults.bComputeSucceeded = false;
+		if (TranslatorOutput != nullptr && TranslatorOutput->Errors.Len() > 0)
+		{
+			Error(FText::Format(LOCTEXT("HlslTranslateErrorMessageFormat", "The HLSL Translator failed.  Errors:\n{0}"), FText::FromString(TranslatorOutput->Errors)));
+			CompileResults.bVMSucceeded = false;
+		}
+		else if (TranslatedHLSL.Len() == 0)
+		{
+			Error(LOCTEXT("HlslTranslateErrorMessageFailed", "The HLSL Translator failed to generate HLSL!"));
+			CompileResults.bVMSucceeded = false;
+		}
+		else
+		{
+			*(CompileResults.Data) = TranslatorOutput->ScriptData;
+			CompileResults.Data->ByteCode.Empty();
+			CompileResults.bComputeSucceeded = true;
+		}
+		CompileResults.Data->LastHlslTranslationGPU = TranslatedHLSL;
 	}
 	else
 	{
-		SCOPE_CYCLE_COUNTER(STAT_NiagaraEditor_HlslCompiler_CompileShader_VectorVM);
-		CompileResults.bVMSucceeded = CompileShader_VectorVM(Input, Output, FString(FPlatformProcess::ShaderDir()), 0, CompilationOutput);
-	}
+		FShaderCompilerInput Input;
+		Input.VirtualSourceFilePath = TEXT("/Engine/Private/NiagaraEmitterInstanceShader.usf");
+		Input.EntryPointName = TEXT("SimulateMain");
+		Input.Environment.SetDefine(TEXT("VM_SIMULATION"), 1);
+		Input.Environment.IncludeVirtualPathToContentsMap.Add(TEXT("/Engine/Generated/NiagaraEmitterInstance.ush"), TranslatedHLSL);
+		FShaderCompilerOutput Output;
 
-	if (CompilationOutput.Errors.Len() > 0)
-	{
-		//TODO: Map Lines of HLSL to their source Nodes and flag those nodes with errors associated with their lines.
-		Error(FText::Format(LOCTEXT("VectorVMCompileErrorMessageFormat", "The Vector VM compile failed.  Errors:\n{0}"), FText::FromString(CompilationOutput.Errors)));
-		CompileResults.bVMSucceeded = false;
-	}
+		FVectorVMCompilationOutput CompilationOutput;
 
-	//For now we just copy the shader code over into the script. 
-	//Eventually Niagara will have all the shader plumbing and do things like materials.
-	if (CompileResults.bVMSucceeded)
-	{
-		SCOPE_CYCLE_COUNTER(STAT_NiagaraEditor_HlslCompiler_CompileShader_VectorVMSucceeded);
-		check(TranslatorOutput);
-		Script->SetByteCode(CompilationOutput.ByteCode);
-		Script->LastHlslTranslation = TranslatedHLSL;
-		Script->LastAssemblyTranslation = CompilationOutput.AssemblyAsString;
-		Script->LastOpCount = CompilationOutput.NumOps;
-		Script->Attributes = TranslatorOutput->Attributes;
-		Script->Parameters = TranslatorOutput->Parameters;
-		Script->DataUsage.bReadsAttriubteData = TranslatorOutput->bReadsAttributeData;
-		//Build internal parameters
-		Script->InternalParameters.Empty();
-		for (int32 i = 0; i < CompilationOutput.InternalConstantOffsets.Num(); ++i)
+		if (TranslatorOutput != nullptr && TranslatorOutput->Errors.Len() > 0)
 		{
-			EVectorVMBaseTypes Type = CompilationOutput.InternalConstantTypes[i];
-			int32 Offset = CompilationOutput.InternalConstantOffsets[i];
-			switch (Type)
-			{
-			case EVectorVMBaseTypes::Float:
-			{
-				float Val = *(float*)(CompilationOutput.InternalConstantData.GetData() + Offset);
-				Script->InternalParameters.SetOrAdd(FNiagaraVariable(FNiagaraTypeDefinition::GetFloatDef(), *LexicalConversion::ToString(Val)))->SetValue(Val);
-			}
-			break;
-			case EVectorVMBaseTypes::Int:
-			{
-				int32 Val = *(int32*)(CompilationOutput.InternalConstantData.GetData() + Offset);
-				Script->InternalParameters.SetOrAdd(FNiagaraVariable(FNiagaraTypeDefinition::GetIntDef(), *LexicalConversion::ToString(Val)))->SetValue(Val);
-			}
-			break;
-			case EVectorVMBaseTypes::Bool:
-			{
-				int32 Val = *(int32*)(CompilationOutput.InternalConstantData.GetData() + Offset);
-				Script->InternalParameters.SetOrAdd(FNiagaraVariable(FNiagaraTypeDefinition::GetIntDef(), Val == 0 ? TEXT("FALSE") : TEXT("TRUE")))->SetValue(Val);
-			}
-			break;
-			}
+			//TODO: Map Lines of HLSL to their source Nodes and flag those nodes with errors associated with their lines.
+			Error(FText::Format(LOCTEXT("HlslTranslateErrorMessageFormat", "The HLSL Translator failed.  Errors:\n{0}"), FText::FromString(TranslatorOutput->Errors)));
+			CompileResults.bVMSucceeded = false;
 		}
-
-		Script->DataInterfaceInfo.Empty();
-		for (FNiagaraScriptDataInterfaceInfo& DataInterfaceInfo : TranslatorOutput->DataInterfaceInfo)
+		else if (TranslatedHLSL.Len() == 0)
 		{
-			int32 Idx = Script->DataInterfaceInfo.AddDefaulted();
-			if (DataInterfaceInfo.DataInterface)
-			{
-				Script->DataInterfaceInfo[Idx].DataInterface = Cast<UNiagaraDataInterface>(StaticDuplicateObject(DataInterfaceInfo.DataInterface, Script, NAME_None, ~RF_Transient));
-			}
-			else
-			{
-				Script->DataInterfaceInfo[Idx].DataInterface = nullptr;
-			}
-			Script->DataInterfaceInfo[Idx].Type = DataInterfaceInfo.Type;
-			Script->DataInterfaceInfo[Idx].Name = DataInterfaceInfo.Name;
-			Script->DataInterfaceInfo[Idx].UserPtrIdx = DataInterfaceInfo.UserPtrIdx;
+			Error(LOCTEXT("HlslTranslateErrorMessageFailed", "The HLSL Translator failed to generate HLSL!"));
+			CompileResults.bVMSucceeded = false;
+		}
+		else
+		{
+			SCOPE_CYCLE_COUNTER(STAT_NiagaraEditor_HlslCompiler_CompileShader_VectorVM);
+			static FCriticalSection CritSec;
+
+			CritSec.Lock();
+			double StartTime = FPlatformTime::Seconds();
+			CompileResults.bVMSucceeded = CompileShader_VectorVM(Input, Output, FString(FPlatformProcess::ShaderDir()), 0, CompilationOutput);
+			CompileResults.CompileTime = (float)(FPlatformTime::Seconds() - StartTime);
+			CritSec.Unlock();
 		}
 
-		Script->NumUserPtrs = TranslatorOutput->NumUserPtrs;
+		if (CompilationOutput.Errors.Len() > 0)
+		{
+			//TODO: Map Lines of HLSL to their source Nodes and flag those nodes with errors associated with their lines.
+			Error(FText::Format(LOCTEXT("VectorVMCompileErrorMessageFormat", "The Vector VM compile failed.  Errors:\n{0}"), FText::FromString(CompilationOutput.Errors)));
+			CompileResults.bVMSucceeded = false;
+		}
 
-		//Extract the external function call table binding info.
-		Script->CalledVMExternalFunctions.Empty(CompilationOutput.CalledVMFunctionTable.Num());
-		for (FVectorVMCompilationOutput::FCalledVMFunction& FuncInfo : CompilationOutput.CalledVMFunctionTable)
+		//For now we just copy the shader code over into the script. 
+		//Eventually Niagara will have all the shader plumbing and do things like materials.
+		if (CompileResults.bVMSucceeded)
 		{
-			//Find the interface corresponding to this call.
-			const FNiagaraFunctionSignature* Sig = nullptr;
-			for (FNiagaraScriptDataInterfaceInfo& NDIInfo : TranslatorOutput->DataInterfaceInfo)
-			{			
-				Sig = NDIInfo.RegisteredFunctions.FindByPredicate([&](const FNiagaraFunctionSignature& CheckSig)
+			SCOPE_CYCLE_COUNTER(STAT_NiagaraEditor_HlslCompiler_CompileShader_VectorVMSucceeded);
+			check(TranslatorOutput);
+			*CompileResults.Data.Get() = TranslatorOutput->ScriptData;
+			CompileResults.Data->ByteCode = CompilationOutput.ByteCode;
+			CompileResults.Data->LastAssemblyTranslation = CompilationOutput.AssemblyAsString;
+			CompileResults.Data->LastOpCount = CompilationOutput.NumOps;
+			//Build internal parameters
+			CompileResults.Data->InternalParameters.Empty();
+			for (int32 i = 0; i < CompilationOutput.InternalConstantOffsets.Num(); ++i)
+			{
+				EVectorVMBaseTypes Type = CompilationOutput.InternalConstantTypes[i];
+				int32 Offset = CompilationOutput.InternalConstantOffsets[i];
+				switch (Type)
 				{
-					FString SigSymbol = FHlslNiagaraTranslator::GetFunctionSignatureSymbol(CheckSig);
-					return SigSymbol == FuncInfo.Name;
-				});
-				if (Sig)
+				case EVectorVMBaseTypes::Float:
 				{
-					break;
+					float Val = *(float*)(CompilationOutput.InternalConstantData.GetData() + Offset);
+					CompileResults.Data->InternalParameters.SetOrAdd(FNiagaraVariable(FNiagaraTypeDefinition::GetFloatDef(), *LexicalConversion::ToString(Val)))->SetValue(Val);
+				}
+				break;
+				case EVectorVMBaseTypes::Int:
+				{
+					int32 Val = *(int32*)(CompilationOutput.InternalConstantData.GetData() + Offset);
+					CompileResults.Data->InternalParameters.SetOrAdd(FNiagaraVariable(FNiagaraTypeDefinition::GetIntDef(), *LexicalConversion::ToString(Val)))->SetValue(Val);
+				}
+				break;
+				case EVectorVMBaseTypes::Bool:
+				{
+					int32 Val = *(int32*)(CompilationOutput.InternalConstantData.GetData() + Offset);
+					CompileResults.Data->InternalParameters.SetOrAdd(FNiagaraVariable(FNiagaraTypeDefinition::GetIntDef(), Val == 0 ? TEXT("FALSE") : TEXT("TRUE")))->SetValue(Val);
+				}
+				break;
 				}
 			}
 
-			if (Sig)
+			//Extract the external function call table binding info.
+			CompileResults.Data->CalledVMExternalFunctions.Empty(CompilationOutput.CalledVMFunctionTable.Num());
+			for (FVectorVMCompilationOutput::FCalledVMFunction& FuncInfo : CompilationOutput.CalledVMFunctionTable)
 			{
-				int32 NewBindingIdx = Script->CalledVMExternalFunctions.AddDefaulted();
-				Script->CalledVMExternalFunctions[NewBindingIdx].Name = *Sig->GetName();
-				Script->CalledVMExternalFunctions[NewBindingIdx].OwnerName = Sig->OwnerName;
+				//Find the interface corresponding to this call.
+				const FNiagaraFunctionSignature* Sig = nullptr;
+				for (FNiagaraScriptDataInterfaceCompileInfo& NDIInfo : TranslatorOutput->ScriptData.DataInterfaceInfo)
+				{
+					Sig = NDIInfo.RegisteredFunctions.FindByPredicate([&](const FNiagaraFunctionSignature& CheckSig)
+					{
+						FString SigSymbol = FHlslNiagaraTranslator::GetFunctionSignatureSymbol(CheckSig);
+						return SigSymbol == FuncInfo.Name;
+					});
+					if (Sig)
+					{
+						break;
+					}
+				}
 
-				Script->CalledVMExternalFunctions[NewBindingIdx].InputParamLocations = FuncInfo.InputParamLocations;
-				Script->CalledVMExternalFunctions[NewBindingIdx].NumOutputs = FuncInfo.NumOutputs;
-			}
-			else
-			{
-				Error(FText::Format(LOCTEXT("VectorVMExternalFunctionBindingError", "Failed to bind the exernal function call:  {0}"), FText::FromString(FuncInfo.Name)));
-				CompileResults.bVMSucceeded = false;
+				if (Sig)
+				{
+					int32 NewBindingIdx = CompileResults.Data->CalledVMExternalFunctions.AddDefaulted();
+					CompileResults.Data->CalledVMExternalFunctions[NewBindingIdx].Name = *Sig->GetName();
+					CompileResults.Data->CalledVMExternalFunctions[NewBindingIdx].OwnerName = Sig->OwnerName;
+
+					CompileResults.Data->CalledVMExternalFunctions[NewBindingIdx].InputParamLocations = FuncInfo.InputParamLocations;
+					CompileResults.Data->CalledVMExternalFunctions[NewBindingIdx].NumOutputs = FuncInfo.NumOutputs;
+				}
+				else
+				{
+					Error(FText::Format(LOCTEXT("VectorVMExternalFunctionBindingError", "Failed to bind the exernal function call:  {0}"), FText::FromString(FuncInfo.Name)));
+					CompileResults.bVMSucceeded = false;
+				}
 			}
 		}
-	}
 
-	CompileResults.OutputHLSL = TranslatedHLSL;
+		CompileResults.Data->LastHlslTranslation = TranslatedHLSL;
 
-	if (CompileResults.bVMSucceeded == false)
-	{
-		//Some error. Clear script and exit.
-		Script->LastHlslTranslation = TranslatedHLSL;
-		Script->SetByteCode(CompilationOutput.ByteCode);
-		Script->Attributes.Empty();
-		Script->Parameters.Empty();
-		Script->InternalParameters.Empty();
-		Script->DataInterfaceInfo.Empty();
-//		Script->NumUserPtrs = 0;
+		if (CompileResults.bVMSucceeded == false)
+		{
+			//Some error. Clear script and exit.
+			CompileResults.Data->ByteCode.Empty();
+			CompileResults.Data->Attributes.Empty();
+			CompileResults.Data->Parameters.Empty();
+			CompileResults.Data->InternalParameters.Empty();
+			CompileResults.Data->DataInterfaceInfo.Empty();
+			//		Script->NumUserPtrs = 0;
+		}
 	}
 
 	return CompileResults;
@@ -337,8 +816,7 @@ const FNiagaraCompileResults &FHlslNiagaraCompiler::CompileScript(UNiagaraScript
 
 
 FHlslNiagaraCompiler::FHlslNiagaraCompiler()
-	: Script(nullptr)
-	, CompileResults(&MessageLog)
+	: CompileResults(&MessageLog)
 {
 	// Make the message log silent so we're not spamming the blueprint log.
 	MessageLog.bSilentMode = true;
diff --git a/Engine/Plugins/FX/Niagara/Source/NiagaraEditor/Private/NiagaraCompiler.h b/Engine/Plugins/FX/Niagara/Source/NiagaraEditor/Private/NiagaraCompiler.h
--- a/Engine/Plugins/FX/Niagara/Source/NiagaraEditor/Private/NiagaraCompiler.h
+++ b/Engine/Plugins/FX/Niagara/Source/NiagaraEditor/Private/NiagaraCompiler.h
@@ -18,10 +18,7 @@ struct FNiagaraTranslatorOutput;
 
 class NIAGARAEDITOR_API FHlslNiagaraCompiler : public INiagaraCompiler
 {
-protected:
-	/** The script we are compiling. */
-	UNiagaraScript* Script;
-
+protected:	
 	/** Message log. Automatically handles marking the NodeGraph with errors. */
 	FCompilerResultsLog MessageLog;
 
@@ -34,7 +31,7 @@ public:
 	virtual ~FHlslNiagaraCompiler() {}
 
 	//Begin INiagaraCompiler Interface
-	virtual const FNiagaraCompileResults& CompileScript(UNiagaraScript* InScript, FNiagaraTranslatorOutput *TranslatorOutput, FString& TranslatedHLSL) override;
+	virtual FNiagaraCompileResults CompileScript(const FNiagaraCompileRequestData* InCompileRequest, const FNiagaraCompileOptions& InOptions, FNiagaraTranslatorOutput *TranslatorOutput, FString& TranslatedHLSL) override;
 
 	virtual void Error(FText ErrorText) override;
 	virtual void Warning(FText WarningText) override;
diff --git a/Engine/Plugins/FX/Niagara/Source/NiagaraEditor/Private/NiagaraEditorCommands.cpp b/Engine/Plugins/FX/Niagara/Source/NiagaraEditor/Private/NiagaraEditorCommands.cpp
--- a/Engine/Plugins/FX/Niagara/Source/NiagaraEditor/Private/NiagaraEditorCommands.cpp
+++ b/Engine/Plugins/FX/Niagara/Source/NiagaraEditor/Private/NiagaraEditorCommands.cpp
@@ -18,8 +18,11 @@ void FNiagaraEditorCommands::RegisterCommands()
 	UI_COMMAND(ToggleBounds, "Bounds", "Display Bounds", EUserInterfaceActionType::ToggleButton, FInputChord());
 	UI_COMMAND(ToggleBounds_SetFixedBounds, "Set Fixed Bounds", "Set Fixed Bounds", EUserInterfaceActionType::Button, FInputChord());
 	UI_COMMAND(ToggleOrbit, "Orbit Mode", "Toggle Orbit Navigation", EUserInterfaceActionType::ToggleButton, FInputChord());
-
 	UI_COMMAND(SaveThumbnailImage, "Thumbnail", "Generate Thumbnail", EUserInterfaceActionType::Button, FInputChord());
+	UI_COMMAND(ToggleAutoPlay, "Auto-play", "Toggles whether or not simulations auto-play when their asset editor is opened, and when the asset is modified.", EUserInterfaceActionType::ToggleButton, FInputChord());
+	UI_COMMAND(ToggleResetSimulationOnChange, "Reset on change", "Toggles whether or not the simulation is reset whenever a change is made in the asset editor.", EUserInterfaceActionType::ToggleButton, FInputChord());
+	UI_COMMAND(ToggleResimulateOnChangeWhilePaused, "Resimulate when paused", "Toggles whether or not the simulation is rerun to the current time when making changes while paused.", EUserInterfaceActionType::ToggleButton, FInputChord());
+	UI_COMMAND(CollapseStackToHeaders, "Collapse to Headers", "Expands all headsers and collapse all items.", EUserInterfaceActionType::Button, FInputChord(EModifierKey::Control, EKeys::O));
 
 }
 
diff --git a/Engine/Plugins/FX/Niagara/Source/NiagaraEditor/Private/NiagaraEditorCommon.cpp b/Engine/Plugins/FX/Niagara/Source/NiagaraEditor/Private/NiagaraEditorCommon.cpp
--- a/Engine/Plugins/FX/Niagara/Source/NiagaraEditor/Private/NiagaraEditorCommon.cpp
+++ b/Engine/Plugins/FX/Niagara/Source/NiagaraEditor/Private/NiagaraEditorCommon.cpp
@@ -253,81 +253,292 @@ void FNiagaraOpInfo::Init()
 		Op->BuildName(TEXT("Log2"), CategoryName);
 		OpInfoMap.Add(Op->Name) = Idx;
 
+		FText AngleFriendlyNameText = NSLOCTEXT("NiagaraOpInfo", "Angle Name", "Angle");
+		FText AngleDescText = NSLOCTEXT("NiagaraOpInfo", "Angle Desc", "Angle as specified by the period range.");
+		FText PeriodFriendlyNameText = NSLOCTEXT("NiagaraOpInfo", "Period Name", "Period");
+		FText PeriodDescText = NSLOCTEXT("NiagaraOpInfo", "Period Desc", "Value in which a complete rotation has occurred.");
+		FNiagaraTypeDefinition PeriodType = FNiagaraTypeDefinition::GetFloatDef();
+		FNiagaraTypeDefinition AngleType = FNiagaraTypeDefinition::GetFloatDef();
+
 		Idx = OpInfos.AddDefaulted();
 		Op = &OpInfos[Idx];
 		Op->Category = CategoryText;
-		Op->FriendlyName = NSLOCTEXT("NiagaraOpInfo", "Sin Name", "Sin");
-		Op->Description = NSLOCTEXT("NiagaraOpInfo", "Sin Desc", "Result = sin(A)");
-		Op->Keywords = FText::FromString(TEXT("sin"));
+		Op->FriendlyName = NSLOCTEXT("NiagaraOpInfo", "Sine Name", "Sine");
+		Op->Description = NSLOCTEXT("NiagaraOpInfo", "Sine Desc", "Result = sin(Angle*(TWO_PI/Period))");
+		Op->Keywords = FText::FromString(TEXT("sine"));
+		Op->Inputs.Add(FNiagaraOpInOutInfo(TEXT("Angle"), AngleType, AngleFriendlyNameText, AngleDescText, DefaultStr_One));
+		Op->Inputs.Add(FNiagaraOpInOutInfo(TEXT("Period"), PeriodType, PeriodFriendlyNameText, PeriodDescText, DefaultStr_One));
+		Op->Outputs.Add(FNiagaraOpInOutInfo(Result, Type, ResultText, ResultText, DefaultStr_One, TEXT("sin({0}*(TWO_PI/{1}))")));
+		Op->BuildName(TEXT("Sine"), CategoryName);
+		OpInfoMap.Add(Op->Name) = Idx;
+
+		Idx = OpInfos.AddDefaulted();
+		Op = &OpInfos[Idx];
+		Op->Category = CategoryText;
+		Op->FriendlyName = NSLOCTEXT("NiagaraOpInfo", "SinRad Name", "Sine(Radians)");
+		Op->Description = NSLOCTEXT("NiagaraOpInfo", "SinRad Desc", "Result = sin(AngleInRadians)");
+		Op->Keywords = FText::FromString(TEXT("sine"));
 		Op->Inputs.Add(FNiagaraOpInOutInfo(A, Type, AText, AText, DefaultStr_One));
 		Op->Outputs.Add(FNiagaraOpInOutInfo(Result, Type, ResultText, ResultText, DefaultStr_One, TEXT("sin({0})")));
-		Op->BuildName(TEXT("Sin"), CategoryName);
+		Op->BuildName(TEXT("Sine(Radians)"), CategoryName);
 		OpInfoMap.Add(Op->Name) = Idx;
 
 		Idx = OpInfos.AddDefaulted();
 		Op = &OpInfos[Idx];
 		Op->Category = CategoryText;
-		Op->FriendlyName = NSLOCTEXT("NiagaraOpInfo", "Cos Name", "Cos");
-		Op->Description = NSLOCTEXT("NiagaraOpInfo", "Cos Desc", "Result = cos(A)");
-		Op->Keywords = FText::FromString(TEXT("cos"));
+		Op->FriendlyName = NSLOCTEXT("NiagaraOpInfo", "SinDeg Name", "Sine(Degrees)");
+		Op->Description = NSLOCTEXT("NiagaraOpInfo", "SinDeg Desc", "Result = sin(AngleInDegrees*DegreesToRadians)");
+		Op->Keywords = FText::FromString(TEXT("sine"));
+		Op->Inputs.Add(FNiagaraOpInOutInfo(A, Type, AText, AText, DefaultStr_One));
+		Op->Outputs.Add(FNiagaraOpInOutInfo(Result, Type, ResultText, ResultText, DefaultStr_One, TEXT("sin((PI/180.0f)*{0})")));
+		Op->BuildName(TEXT("Sine(Degrees)"), CategoryName);
+		OpInfoMap.Add(Op->Name) = Idx;
+
+		Idx = OpInfos.AddDefaulted();
+		Op = &OpInfos[Idx];
+		Op->Category = CategoryText;
+		Op->FriendlyName = NSLOCTEXT("NiagaraOpInfo", "Cosine Name", "Cosine");
+		Op->Description = NSLOCTEXT("NiagaraOpInfo", "Cosine Desc", "Result = cos(Angle*(TWO_PI/Period))");
+		Op->Keywords = FText::FromString(TEXT("Cosine"));
+		Op->Inputs.Add(FNiagaraOpInOutInfo(TEXT("Angle"), AngleType, AngleFriendlyNameText, AngleDescText, DefaultStr_One));
+		Op->Inputs.Add(FNiagaraOpInOutInfo(TEXT("Period"), PeriodType, PeriodFriendlyNameText, PeriodDescText, DefaultStr_One));
+		Op->Outputs.Add(FNiagaraOpInOutInfo(Result, Type, ResultText, ResultText, DefaultStr_One, TEXT("cos({0}*(TWO_PI/{1}))")));
+		Op->BuildName(TEXT("Cosine"), CategoryName);
+		OpInfoMap.Add(Op->Name) = Idx;
+
+		Idx = OpInfos.AddDefaulted();
+		Op = &OpInfos[Idx];
+		Op->Category = CategoryText;
+		Op->FriendlyName = NSLOCTEXT("NiagaraOpInfo", "CosRad Name", "Cosine(Radians)");
+		Op->Description = NSLOCTEXT("NiagaraOpInfo", "CosRad Desc", "Result = cos(AngleInRadians)");
+		Op->Keywords = FText::FromString(TEXT("Cosine"));
 		Op->Inputs.Add(FNiagaraOpInOutInfo(A, Type, AText, AText, DefaultStr_One));
 		Op->Outputs.Add(FNiagaraOpInOutInfo(Result, Type, ResultText, ResultText, DefaultStr_One, TEXT("cos({0})")));
-		Op->BuildName(TEXT("Cos"), CategoryName);
+		Op->BuildName(TEXT("Cosine(Radians)"), CategoryName);
+		OpInfoMap.Add(Op->Name) = Idx;
+
+		Idx = OpInfos.AddDefaulted();
+		Op = &OpInfos[Idx];
+		Op->Category = CategoryText;
+		Op->FriendlyName = NSLOCTEXT("NiagaraOpInfo", "CosDeg Name", "Cosine(Degrees)");
+		Op->Description = NSLOCTEXT("NiagaraOpInfo", "CosDeg Desc", "Result = cos(AngleInDegrees*DegreesToRadians)");
+		Op->Keywords = FText::FromString(TEXT("Cosine"));
+		Op->Inputs.Add(FNiagaraOpInOutInfo(A, Type, AText, AText, DefaultStr_One));
+		Op->Outputs.Add(FNiagaraOpInOutInfo(Result, Type, ResultText, ResultText, DefaultStr_One, TEXT("cos((PI/180.0f)*{0})")));
+		Op->BuildName(TEXT("Cosine(Degrees)"), CategoryName);
+		OpInfoMap.Add(Op->Name) = Idx;
+
+		Idx = OpInfos.AddDefaulted();
+		Op = &OpInfos[Idx];
+		Op->Category = CategoryText;
+		Op->FriendlyName = NSLOCTEXT("NiagaraOpInfo", "Tangent Name", "Tangent");
+		Op->Description = NSLOCTEXT("NiagaraOpInfo", "Tangent Desc", "Result = tan(Angle*(TWO_PI/Period))");
+		Op->Keywords = FText::FromString(TEXT("Tangent"));
+		Op->Inputs.Add(FNiagaraOpInOutInfo(TEXT("Angle"), AngleType, AngleFriendlyNameText, AngleDescText, DefaultStr_One));
+		Op->Inputs.Add(FNiagaraOpInOutInfo(TEXT("Period"), PeriodType, PeriodFriendlyNameText, PeriodDescText, DefaultStr_One));
+		Op->Outputs.Add(FNiagaraOpInOutInfo(Result, Type, ResultText, ResultText, DefaultStr_One, TEXT("tan({0}*(TWO_PI/{1}))")));
+		Op->BuildName(TEXT("Tangent"), CategoryName);
 		OpInfoMap.Add(Op->Name) = Idx;
 
 		Idx = OpInfos.AddDefaulted();
 		Op = &OpInfos[Idx];
 		Op->Category = CategoryText;
-		Op->FriendlyName = NSLOCTEXT("NiagaraOpInfo", "Tan Name", "Tan");
-		Op->Description = NSLOCTEXT("NiagaraOpInfo", "Tan Desc", "Result = tan(A)");
-		Op->Keywords = FText::FromString(TEXT("tan"));
+		Op->FriendlyName = NSLOCTEXT("NiagaraOpInfo", "TanRad Name", "Tangent(Radians)");
+		Op->Description = NSLOCTEXT("NiagaraOpInfo", "TanRad Desc", "Result = tan(AngleInRadians)");
+		Op->Keywords = FText::FromString(TEXT("Tangent"));
 		Op->Inputs.Add(FNiagaraOpInOutInfo(A, Type, AText, AText, DefaultStr_One));
 		Op->Outputs.Add(FNiagaraOpInOutInfo(Result, Type, ResultText, ResultText, DefaultStr_One, TEXT("tan({0})")));
-		Op->BuildName(TEXT("Tan"), CategoryName);
+		Op->BuildName(TEXT("Tangent(Radians)"), CategoryName);
+		OpInfoMap.Add(Op->Name) = Idx;
+
+		Idx = OpInfos.AddDefaulted();
+		Op = &OpInfos[Idx];
+		Op->Category = CategoryText;
+		Op->FriendlyName = NSLOCTEXT("NiagaraOpInfo", "TanDeg Name", "Tangent(Degrees)");
+		Op->Description = NSLOCTEXT("NiagaraOpInfo", "TanDeg Desc", "Result = tan(AngleInDegrees*DegreesToRadians)");
+		Op->Keywords = FText::FromString(TEXT("Tangent"));
+		Op->Inputs.Add(FNiagaraOpInOutInfo(A, Type, AText, AText, DefaultStr_One));
+		Op->Outputs.Add(FNiagaraOpInOutInfo(Result, Type, ResultText, ResultText, DefaultStr_One, TEXT("tan((PI/180.0f)*{0})")));
+		Op->BuildName(TEXT("Tangent(Degrees)"), CategoryName);
+		OpInfoMap.Add(Op->Name) = Idx;
+
+
+		Idx = OpInfos.AddDefaulted();
+		Op = &OpInfos[Idx];
+		Op->Category = CategoryText;
+		Op->FriendlyName = NSLOCTEXT("NiagaraOpInfo", "ArcSine Name", "ArcSine");
+		Op->Description = NSLOCTEXT("NiagaraOpInfo", "ArcSine Desc", "Result = asin(A)*(Period/TWO_PI)");
+		Op->Keywords = FText::FromString(TEXT("ArcSine"));
+		Op->Inputs.Add(FNiagaraOpInOutInfo(A, Type, AText, AText, DefaultStr_One));
+		Op->Inputs.Add(FNiagaraOpInOutInfo(TEXT("Period"), PeriodType, PeriodFriendlyNameText, PeriodDescText, DefaultStr_One));
+		Op->Outputs.Add(FNiagaraOpInOutInfo(Result, Type, ResultText, ResultText, DefaultStr_One, TEXT("asin({0})*({1}/TWO_PI)")));
+		Op->BuildName(TEXT("ArcSine"), CategoryName);
 		OpInfoMap.Add(Op->Name) = Idx;
 
 		Idx = OpInfos.AddDefaulted();
 		Op = &OpInfos[Idx];
 		Op->Category = CategoryText;
-		Op->FriendlyName = NSLOCTEXT("NiagaraOpInfo", "ASin Name", "ASin");
-		Op->Description = NSLOCTEXT("NiagaraOpInfo", "ASin Desc", "Result = asin(A)");
-		Op->Keywords = FText::FromString(TEXT("asin"));
+		Op->FriendlyName = NSLOCTEXT("NiagaraOpInfo", "ArcSineRad Name", "ArcSine(Radians)");
+		Op->Description = NSLOCTEXT("NiagaraOpInfo", "ArcSineRad Desc", "Result = asin(A)");
+		Op->Keywords = FText::FromString(TEXT("ArcSine"));
 		Op->Inputs.Add(FNiagaraOpInOutInfo(A, Type, AText, AText, DefaultStr_One));
 		Op->Outputs.Add(FNiagaraOpInOutInfo(Result, Type, ResultText, ResultText, DefaultStr_One, TEXT("asin({0})")));
-		Op->BuildName(TEXT("ASin"), CategoryName);
+		Op->BuildName(TEXT("ArcSine(Radians)"), CategoryName);
 		OpInfoMap.Add(Op->Name) = Idx;
 
 		Idx = OpInfos.AddDefaulted();
 		Op = &OpInfos[Idx];
 		Op->Category = CategoryText;
-		Op->FriendlyName = NSLOCTEXT("NiagaraOpInfo", "ACos Name", "ACos");
-		Op->Description = NSLOCTEXT("NiagaraOpInfo", "ACos Desc", "Result = acos(A)");
-		Op->Keywords = FText::FromString(TEXT("acos"));
+		Op->FriendlyName = NSLOCTEXT("NiagaraOpInfo", "ArcSineDeg Name", "ArcSine(Degrees)");
+		Op->Description = NSLOCTEXT("NiagaraOpInfo", "ArcSineDeg Desc", "Result = asin(A)*RadiansToDegrees");
+		Op->Keywords = FText::FromString(TEXT("ArcSine"));
+		Op->Inputs.Add(FNiagaraOpInOutInfo(A, Type, AText, AText, DefaultStr_One));
+		Op->Outputs.Add(FNiagaraOpInOutInfo(Result, Type, ResultText, ResultText, DefaultStr_One, TEXT("(180.0f/PI)*asin({0})")));
+		Op->BuildName(TEXT("ArcSine(Degrees)"), CategoryName);
+		OpInfoMap.Add(Op->Name) = Idx;
+
+		Idx = OpInfos.AddDefaulted();
+		Op = &OpInfos[Idx];
+		Op->Category = CategoryText;
+		Op->FriendlyName = NSLOCTEXT("NiagaraOpInfo", "Pi Name", "PI");
+		Op->Description = NSLOCTEXT("NiagaraOpInfo", "Pi Desc", "The constant PI");
+		Op->Keywords = FText::FromString(TEXT("pi"));
+		Op->Outputs.Add(FNiagaraOpInOutInfo(Result, FNiagaraTypeDefinition::GetFloatDef(), ResultText, ResultText, DefaultStr_One, TEXT("PI")));
+		Op->BuildName(TEXT("PI"), CategoryName);
+		OpInfoMap.Add(Op->Name) = Idx;
+
+		Idx = OpInfos.AddDefaulted();
+		Op = &OpInfos[Idx];
+		Op->Category = CategoryText;
+		Op->FriendlyName = NSLOCTEXT("NiagaraOpInfo", "Two Pi Name", "TWO_PI");
+		Op->Description = NSLOCTEXT("NiagaraOpInfo", "Two Pi Desc", "The constant PI * 2");
+		Op->Keywords = FText::FromString(TEXT("pi"));
+		Op->Outputs.Add(FNiagaraOpInOutInfo(Result, FNiagaraTypeDefinition::GetFloatDef(), ResultText, ResultText, DefaultStr_One, TEXT("TWO_PI")));
+		Op->BuildName(TEXT("TWO_PI"), CategoryName);
+		OpInfoMap.Add(Op->Name) = Idx;
+
+
+		Idx = OpInfos.AddDefaulted();
+		Op = &OpInfos[Idx];
+		Op->Category = CategoryText;
+		Op->FriendlyName = NSLOCTEXT("NiagaraOpInfo", "ArcCosine Name", "ArcCosine");
+		Op->Description = NSLOCTEXT("NiagaraOpInfo", "ArcCosine Desc", "Result = acos(A)*(Period/TWO_PI)");
+		Op->Keywords = FText::FromString(TEXT("ArcCosine"));
+		Op->Inputs.Add(FNiagaraOpInOutInfo(A, Type, AText, AText, DefaultStr_One));
+		Op->Inputs.Add(FNiagaraOpInOutInfo(TEXT("Period"), PeriodType, PeriodFriendlyNameText, PeriodDescText, DefaultStr_One));
+		Op->Outputs.Add(FNiagaraOpInOutInfo(Result, Type, ResultText, ResultText, DefaultStr_One, TEXT("acos({0})*({1}/TWO_PI)")));
+		Op->BuildName(TEXT("ArcCosine"), CategoryName);
+		OpInfoMap.Add(Op->Name) = Idx;
+
+		Idx = OpInfos.AddDefaulted();
+		Op = &OpInfos[Idx];
+		Op->Category = CategoryText;
+		Op->FriendlyName = NSLOCTEXT("NiagaraOpInfo", "ArcCosineRad Name", "ArcCosine(Radians)");
+		Op->Description = NSLOCTEXT("NiagaraOpInfo", "ArcCosineRad Desc", "Result = acos(A)");
+		Op->Keywords = FText::FromString(TEXT("ArcCosine"));
 		Op->Inputs.Add(FNiagaraOpInOutInfo(A, Type, AText, AText, DefaultStr_One));
 		Op->Outputs.Add(FNiagaraOpInOutInfo(Result, Type, ResultText, ResultText, DefaultStr_One, TEXT("acos({0})")));
-		Op->BuildName(TEXT("ACos"), CategoryName);
+		Op->BuildName(TEXT("ArcCosine(Radians)"), CategoryName);
 		OpInfoMap.Add(Op->Name) = Idx;
 
 		Idx = OpInfos.AddDefaulted();
 		Op = &OpInfos[Idx];
 		Op->Category = CategoryText;
-		Op->FriendlyName = NSLOCTEXT("NiagaraOpInfo", "ATan Name", "ATan");
-		Op->Description = NSLOCTEXT("NiagaraOpInfo", "ATan Desc", "Result = atan(A)");
-		Op->Keywords = FText::FromString(TEXT("atan"));
+		Op->FriendlyName = NSLOCTEXT("NiagaraOpInfo", "ArcCosineDeg Name", "ArcCosine(Degrees)");
+		Op->Description = NSLOCTEXT("NiagaraOpInfo", "ArcCosineDeg Desc", "Result = acos(A)*RadiansToDegrees");
+		Op->Keywords = FText::FromString(TEXT("ArcCosine"));
+		Op->Inputs.Add(FNiagaraOpInOutInfo(A, Type, AText, AText, DefaultStr_One));
+		Op->Outputs.Add(FNiagaraOpInOutInfo(Result, Type, ResultText, ResultText, DefaultStr_One, TEXT("(180.0f/PI)*acos({0})")));
+		Op->BuildName(TEXT("ArcCosine(Degrees)"), CategoryName);
+		OpInfoMap.Add(Op->Name) = Idx;
+
+
+		Idx = OpInfos.AddDefaulted();
+		Op = &OpInfos[Idx];
+		Op->Category = CategoryText;
+		Op->FriendlyName = NSLOCTEXT("NiagaraOpInfo", "ArcTangent Name", "ArcTangent");
+		Op->Description = NSLOCTEXT("NiagaraOpInfo", "ArcTangent Desc", "Result = atan(A)*(Period/TWO_PI)");
+		Op->Keywords = FText::FromString(TEXT("ArcTangent"));
+		Op->Inputs.Add(FNiagaraOpInOutInfo(A, Type, AText, AText, DefaultStr_One));
+		Op->Inputs.Add(FNiagaraOpInOutInfo(TEXT("Period"), PeriodType, PeriodFriendlyNameText, PeriodDescText, DefaultStr_One));
+		Op->Outputs.Add(FNiagaraOpInOutInfo(Result, Type, ResultText, ResultText, DefaultStr_One, TEXT("atan({0})*({1}/TWO_PI)")));
+		Op->BuildName(TEXT("ArcTangent"), CategoryName);
+		OpInfoMap.Add(Op->Name) = Idx;
+
+		Idx = OpInfos.AddDefaulted();
+		Op = &OpInfos[Idx];
+		Op->Category = CategoryText;
+		Op->FriendlyName = NSLOCTEXT("NiagaraOpInfo", "ArcTangentRad Name", "ArcTangent(Radians)");
+		Op->Description = NSLOCTEXT("NiagaraOpInfo", "ArcTangentRad Desc", "Result = atan(A)");
+		Op->Keywords = FText::FromString(TEXT("ArcTangent"));
 		Op->Inputs.Add(FNiagaraOpInOutInfo(A, Type, AText, AText, DefaultStr_One));
 		Op->Outputs.Add(FNiagaraOpInOutInfo(Result, Type, ResultText, ResultText, DefaultStr_One, TEXT("atan({0})")));
-		Op->BuildName(TEXT("ATan"), CategoryName);
+		Op->BuildName(TEXT("ArcTangent(Radians)"), CategoryName);
 		OpInfoMap.Add(Op->Name) = Idx;
 
 		Idx = OpInfos.AddDefaulted();
 		Op = &OpInfos[Idx];
 		Op->Category = CategoryText;
-		Op->FriendlyName = NSLOCTEXT("NiagaraOpInfo", "ATan2 Name", "ATan2");
-		Op->Description = NSLOCTEXT("NiagaraOpInfo", "ATan2 Desc", "Result = atan2(A, B)");
+		Op->FriendlyName = NSLOCTEXT("NiagaraOpInfo", "ArcTangentDeg Name", "ArcTangent(Degrees)");
+		Op->Description = NSLOCTEXT("NiagaraOpInfo", "ArcTangentDeg Desc", "Result = atan(A)*RadiansToDegrees");
+		Op->Keywords = FText::FromString(TEXT("ArcTangent"));
+		Op->Inputs.Add(FNiagaraOpInOutInfo(A, Type, AText, AText, DefaultStr_One));
+		Op->Outputs.Add(FNiagaraOpInOutInfo(Result, Type, ResultText, ResultText, DefaultStr_One, TEXT("(180.0f/PI)*atan({0})")));
+		Op->BuildName(TEXT("ArcTangent(Degrees)"), CategoryName);
+		OpInfoMap.Add(Op->Name) = Idx;
+
+		Idx = OpInfos.AddDefaulted();
+		Op = &OpInfos[Idx];
+		Op->Category = CategoryText;
+		Op->FriendlyName = NSLOCTEXT("NiagaraOpInfo", "ATan2 Name", "ArcTangent2");
+		Op->Description = NSLOCTEXT("NiagaraOpInfo", "ATan2 Desc", "ResultInPeriod = Period * atan2(A, B) / 2PI");
+		Op->Inputs.Add(FNiagaraOpInOutInfo(A, Type, AText, AText, DefaultStr_One));
+		Op->Inputs.Add(FNiagaraOpInOutInfo(B, Type, BText, BText, DefaultStr_One));
+		Op->Inputs.Add(FNiagaraOpInOutInfo(TEXT("Period"), PeriodType, PeriodFriendlyNameText, PeriodDescText, DefaultStr_One));
+		Op->Outputs.Add(FNiagaraOpInOutInfo(Result, Type, ResultText, ResultText, DefaultStr_One, TEXT("atan2({0},{1})*({2}/TWO_PI)")));
+		Op->BuildName(TEXT("ArcTangent2"), CategoryName);
+		OpInfoMap.Add(Op->Name) = Idx;
+
+		Idx = OpInfos.AddDefaulted();
+		Op = &OpInfos[Idx];
+		Op->Category = CategoryText;
+		Op->FriendlyName = NSLOCTEXT("NiagaraOpInfo", "ATan2Rad Name", "ArcTangent2(Radians)");
+		Op->Description = NSLOCTEXT("NiagaraOpInfo", "ATan2Rad Desc", "ResultInRadians = atan2(A, B)");
 		Op->Inputs.Add(FNiagaraOpInOutInfo(A, Type, AText, AText, DefaultStr_One));
 		Op->Inputs.Add(FNiagaraOpInOutInfo(B, Type, BText, BText, DefaultStr_One));
 		Op->Outputs.Add(FNiagaraOpInOutInfo(Result, Type, ResultText, ResultText, DefaultStr_One, TEXT("atan2({0},{1})")));
-		Op->BuildName(TEXT("ATan2"), CategoryName);
+		Op->BuildName(TEXT("ArcTangent2(Radians)"), CategoryName);
+		OpInfoMap.Add(Op->Name) = Idx;
+
+		Idx = OpInfos.AddDefaulted();
+		Op = &OpInfos[Idx];
+		Op->Category = CategoryText;
+		Op->FriendlyName = NSLOCTEXT("NiagaraOpInfo", "ATan2Deg Name", "ArcTangent2(Degrees)");
+		Op->Description = NSLOCTEXT("NiagaraOpInfo", "ATan2Deg Desc", "ResultInPeriod = 180 * atan2(A, B) / PI");
+		Op->Inputs.Add(FNiagaraOpInOutInfo(A, Type, AText, AText, DefaultStr_One));
+		Op->Inputs.Add(FNiagaraOpInOutInfo(B, Type, BText, BText, DefaultStr_One));
+		Op->Outputs.Add(FNiagaraOpInOutInfo(Result, Type, ResultText, ResultText, DefaultStr_One, TEXT("atan2({0},{1})*(180.0f/PI)")));
+		Op->BuildName(TEXT("ArcTangent2(Degrees)"), CategoryName);
+		OpInfoMap.Add(Op->Name) = Idx;
+
+		Idx = OpInfos.AddDefaulted();
+		Op = &OpInfos[Idx];
+		Op->Category = CategoryText;
+		Op->FriendlyName = NSLOCTEXT("NiagaraOpInfo", "Degrees To Radians", "DegreesToRadians");
+		Op->Description = NSLOCTEXT("NiagaraOpInfo", "Degrees To Radians Desc", "DegreesToRadians(A)");
+		Op->Inputs.Add(FNiagaraOpInOutInfo(A, FNiagaraTypeDefinition::GetFloatDef(), AText, AText, DefaultStr_One));
+		Op->Outputs.Add(FNiagaraOpInOutInfo(Result, FNiagaraTypeDefinition::GetFloatDef(), ResultText, ResultText, DefaultStr_One, TEXT("(PI/180.0f)*({0})")));
+		Op->BuildName(TEXT("DegreesToRadians"), CategoryName);
+		OpInfoMap.Add(Op->Name) = Idx;
+
+		Idx = OpInfos.AddDefaulted();
+		Op = &OpInfos[Idx];
+		Op->Category = CategoryText;
+		Op->FriendlyName = NSLOCTEXT("NiagaraOpInfo", "Radians To Degrees", "RadiansToDegrees");
+		Op->Description = NSLOCTEXT("NiagaraOpInfo", "RadiansToDegrees Desc", "RadiansToDegrees(A)");
+		Op->Inputs.Add(FNiagaraOpInOutInfo(A, FNiagaraTypeDefinition::GetFloatDef(), AText, AText, DefaultStr_One));
+		Op->Outputs.Add(FNiagaraOpInOutInfo(Result, FNiagaraTypeDefinition::GetFloatDef(), ResultText, ResultText, DefaultStr_One, TEXT("(180.0f/PI)*({0})")));
+		Op->BuildName(TEXT("RadiansToDegrees"), CategoryName);
 		OpInfoMap.Add(Op->Name) = Idx;
 
 		Idx = OpInfos.AddDefaulted();
@@ -506,7 +717,7 @@ void FNiagaraOpInfo::Init()
 		Op = &OpInfos[Idx];
 		Op->Category = CategoryText;
 		Op->FriendlyName = NSLOCTEXT("NiagaraOpInfo", "Rand Name", "Random");
-		Op->Description = NSLOCTEXT("NiagaraOpInfo", "Rand Desc", "Result = rand(A)");
+		Op->Description = NSLOCTEXT("NiagaraOpInfo", "Rand Desc", "Returns a random value between 0 and A.");
 		Op->Inputs.Add(FNiagaraOpInOutInfo(A, Type, AText, AText, DefaultStr_One));
 		Op->Outputs.Add(FNiagaraOpInOutInfo(Result, Type, ResultText, ResultText, DefaultStr_One, TEXT("rand({0})")));
 		Op->BuildName(TEXT("Rand"), CategoryName);
diff --git a/Engine/Plugins/FX/Niagara/Source/NiagaraEditor/Private/NiagaraEditorModule.cpp b/Engine/Plugins/FX/Niagara/Source/NiagaraEditor/Private/NiagaraEditorModule.cpp
--- a/Engine/Plugins/FX/Niagara/Source/NiagaraEditor/Private/NiagaraEditorModule.cpp
+++ b/Engine/Plugins/FX/Niagara/Source/NiagaraEditor/Private/NiagaraEditorModule.cpp
@@ -11,6 +11,7 @@
 #include "ISettingsModule.h"
 #include "SequencerChannelInterface.h"
 #include "SequencerSettings.h"
+#include "AssetRegistryModule.h"
 
 #include "AssetTypeActions/AssetTypeActions_NiagaraSystem.h"
 #include "AssetTypeActions/AssetTypeActions_NiagaraEmitter.h"
@@ -39,10 +40,12 @@
 #include "TypeEditorUtilities/NiagaraFloatTypeEditorUtilities.h"
 #include "TypeEditorUtilities/NiagaraVectorTypeEditorUtilities.h"
 #include "TypeEditorUtilities/NiagaraColorTypeEditorUtilities.h"
+#include "TypeEditorUtilities/NiagaraMatrixTypeEditorUtilities.h"
 #include "TypeEditorUtilities/NiagaraDataInterfaceCurveTypeEditorUtilities.h"
 #include "NiagaraEditorStyle.h"
 #include "NiagaraEditorCommands.h"
-#include "NiagaraEmitterTrackEditor.h"
+#include "Sequencer/NiagaraSequence/NiagaraEmitterTrackEditor.h"
+#include "Sequencer/LevelSequence/NiagaraSystemTrackEditor.h"
 #include "PropertyEditorModule.h"
 #include "NiagaraSettings.h"
 #include "NiagaraModule.h"
@@ -54,17 +57,37 @@
 #include "NiagaraDataInterfaceVector4Curve.h"
 #include "NiagaraDataInterfaceColorCurve.h"
 #include "NiagaraScriptViewModel.h"
-#include "NiagaraSystemViewModel.h"
-#include "NiagaraEmitterViewModel.h"
+#include "ViewModels/NiagaraSystemViewModel.h"
+#include "ViewModels/NiagaraEmitterViewModel.h"
 #include "TNiagaraGraphPinEditableName.h"
-#include "MovieSceneNiagaraEmitterSection.h"
+#include "Sequencer/NiagaraSequence/Sections/MovieSceneNiagaraEmitterSection.h"
 #include "UObject/Class.h"
 #include "NiagaraScriptMergeManager.h"
 #include "NiagaraEmitter.h"
+#include "NiagaraScriptSource.h"
+#include "NiagaraTypes.h"
+
+#include "MovieScene/Parameters/MovieSceneNiagaraBoolParameterTrack.h"
+#include "MovieScene/Parameters/MovieSceneNiagaraFloatParameterTrack.h"
+#include "MovieScene/Parameters/MovieSceneNiagaraIntegerParameterTrack.h"
+#include "MovieScene/Parameters/MovieSceneNiagaraVectorParameterTrack.h"
+#include "MovieScene/Parameters/MovieSceneNiagaraColorParameterTrack.h"
+
+#include "Sections/MovieSceneBoolSection.h"
+#include "Sections/MovieSceneFloatSection.h"
+#include "Sections/MovieSceneIntegerSection.h"
+#include "Sections/MovieSceneVectorSection.h"
+#include "Sections/MovieSceneColorSection.h"
+
+#include "ISequencerSection.h"
+#include "Sections/BoolPropertySection.h"
+#include "Sections/ColorPropertySection.h"
 
 #include "Customizations/NiagaraComponentDetails.h"
 #include "Customizations/NiagaraTypeCustomizations.h"
 #include "Customizations/NiagaraEventScriptPropertiesCustomization.h"
+#include "HAL/IConsoleManager.h"
+
 
 IMPLEMENT_MODULE( FNiagaraEditorModule, NiagaraEditor );
 
@@ -73,6 +96,11 @@ IMPLEMENT_MODULE( FNiagaraEditorModule, NiagaraEditor );
 const FName FNiagaraEditorModule::NiagaraEditorAppIdentifier( TEXT( "NiagaraEditorApp" ) );
 const FLinearColor FNiagaraEditorModule::WorldCentricTabColorScale(0.0f, 0.0f, 0.2f, 0.5f);
 
+const FName FNiagaraEditorModule::FInputMetaDataKeys::AdvancedDisplay = "AdvancedDisplay";
+const FName FNiagaraEditorModule::FInputMetaDataKeys::EditCondition = "EditCondition";
+const FName FNiagaraEditorModule::FInputMetaDataKeys::VisibleCondition = "VisibleCondition";
+const FName FNiagaraEditorModule::FInputMetaDataKeys::InlineEditConditionToggle = "InlineEditConditionToggle";
+
 EAssetTypeCategories::Type FNiagaraEditorModule::NiagaraAssetCategory;
 
 //////////////////////////////////////////////////////////////////////////
@@ -143,9 +171,114 @@ private:
 
 FNiagaraEditorModule::FNiagaraEditorModule() 
 	: SequencerSettings(nullptr)
+	, TestCompileScriptCommand(nullptr)
+{
+}
+
+void DumpParameterStore(const FNiagaraParameterStore& ParameterStore)
+{
+	FNiagaraEditorModule& NiagaraEditorModule = FModuleManager::Get().GetModuleChecked<FNiagaraEditorModule>("NiagaraEditor");
+	TArray<FNiagaraVariable> ParameterVariables;
+	ParameterStore.GetParameters(ParameterVariables);
+	for (const FNiagaraVariable& ParameterVariable : ParameterVariables)
+	{
+		FString Name = ParameterVariable.GetName().ToString();
+		FString Type = ParameterVariable.GetType().GetName();
+		FString Value;
+		TSharedPtr<INiagaraEditorTypeUtilities, ESPMode::ThreadSafe> ParameterTypeUtilities = NiagaraEditorModule.GetTypeUtilities(ParameterVariable.GetType());
+		if (ParameterTypeUtilities.IsValid() && ParameterTypeUtilities->CanHandlePinDefaults())
+		{
+			FNiagaraVariable ParameterVariableWithValue = ParameterVariable;
+			ParameterVariableWithValue.SetData(ParameterStore.GetParameterData(ParameterVariable));
+			Value = ParameterTypeUtilities->GetPinDefaultStringFromValue(ParameterVariableWithValue);
+		}
+		else
+		{
+			Value = "(unsupported)";
+		}
+		UE_LOG(LogNiagaraEditor, Log, TEXT("%s\t%s\t%s"), *Name, *Type, *Value);
+	}
+}
+
+void DumpRapidIterationParametersForScript(UNiagaraScript* Script, const FString& HeaderName)
+{
+	UEnum* NiagaraScriptUsageEnum = FindObjectChecked<UEnum>(ANY_PACKAGE, TEXT("ENiagaraScriptUsage"), true);
+	FString UsageName = NiagaraScriptUsageEnum->GetNameByValue((int64)Script->GetUsage()).ToString();
+	UE_LOG(LogNiagaraEditor, Log, TEXT("%s - %s - %s"), *Script->GetPathName(), *HeaderName, *UsageName);
+	DumpParameterStore(Script->RapidIterationParameters);
+}
+
+void DumpRapidIterationParametersForEmitter(UNiagaraEmitter* Emitter, const FString& EmitterName)
 {
+	TArray<UNiagaraScript*> Scripts;
+	Emitter->GetScripts(Scripts, false);
+	for (UNiagaraScript* Script : Scripts)
+	{
+		DumpRapidIterationParametersForScript(Script, EmitterName);
+	}
+}
+
+void DumpRapidIterationParamersForAsset(const TArray<FString>& Arguments)
+{
+	if (Arguments.Num() == 1)
+	{
+		const FAssetRegistryModule& AssetRegistryModule = FModuleManager::LoadModuleChecked<FAssetRegistryModule>(TEXT("AssetRegistry"));
+		const FAssetData AssetData = AssetRegistryModule.Get().GetAssetByObjectPath(*Arguments[0]);
+		UObject* Asset = AssetData.GetAsset();
+		if (Asset != nullptr)
+		{
+			UNiagaraSystem* SystemAsset = Cast<UNiagaraSystem>(Asset);
+			if (SystemAsset != nullptr)
+			{
+				DumpRapidIterationParametersForScript(SystemAsset->GetSystemSpawnScript(), SystemAsset->GetName());
+				DumpRapidIterationParametersForScript(SystemAsset->GetSystemUpdateScript(), SystemAsset->GetName());
+				for (const FNiagaraEmitterHandle& EmitterHandle : SystemAsset->GetEmitterHandles())
+				{
+					DumpRapidIterationParametersForEmitter(EmitterHandle.GetInstance(), EmitterHandle.GetName().ToString());
+				}
+			}
+			else
+			{
+				UNiagaraEmitter* EmitterAsset = Cast<UNiagaraEmitter>(Asset);
+				if (EmitterAsset != nullptr)
+				{
+					DumpRapidIterationParametersForEmitter(EmitterAsset, EmitterAsset->GetName());
+				}
+				else
+				{
+					UE_LOG(LogNiagaraEditor, Warning, TEXT("DumpRapidIterationParameters - Only niagara system and niagara emitter assets are supported"));
+				}
+			}
+		}
+		else
+		{
+			UE_LOG(LogNiagaraEditor, Warning, TEXT("DumpRapidIterationParameters - Asset not found"));
+		}
+	}
+	else
+	{
+		UE_LOG(LogNiagaraEditor, Warning, TEXT("DumpRapidIterationParameters - Must supply an asset path to dump"));
+	}
 }
 
+class FNiagaraSystemBoolParameterTrackEditor : public FNiagaraSystemParameterTrackEditor<UMovieSceneNiagaraBoolParameterTrack, UMovieSceneBoolSection>
+{
+	virtual TSharedRef<ISequencerSection> MakeSectionInterface(UMovieSceneSection& SectionObject, UMovieSceneTrack& Track, FGuid ObjectBinding) override
+	{
+		checkf(SectionObject.GetClass()->IsChildOf<UMovieSceneBoolSection>(), TEXT("Unsupported section."));
+		return MakeShareable(new FBoolPropertySection(SectionObject));
+	}
+};
+
+class FNiagaraSystemColorParameterTrackEditor : public FNiagaraSystemParameterTrackEditor<UMovieSceneNiagaraColorParameterTrack, UMovieSceneColorSection>
+{
+	virtual TSharedRef<ISequencerSection> MakeSectionInterface(UMovieSceneSection& SectionObject, UMovieSceneTrack& Track, FGuid ObjectBinding) override
+	{
+		checkf(SectionObject.GetClass()->IsChildOf<UMovieSceneColorSection>(), TEXT("Unsupported section."));
+		return MakeShareable(new FColorPropertySection(*Cast<UMovieSceneColorSection>(&SectionObject), ObjectBinding, GetSequencer()));
+	}
+};
+
 void FNiagaraEditorModule::StartupModule()
 {
 	MenuExtensibilityManager = MakeShareable(new FExtensibilityManager);
@@ -192,6 +325,10 @@ void FNiagaraEditorModule::StartupModule()
 		FOnGetPropertyTypeCustomizationInstance::CreateStatic(&FNiagaraMatrixCustomization::MakeInstance)
 	);
 
+	PropertyModule.RegisterCustomPropertyTypeLayout("NiagaraVariableAttributeBinding",
+		FOnGetPropertyTypeCustomizationInstance::CreateStatic(&FNiagaraVariableAttributeBindingCustomization::MakeInstance)
+	);
+
 	FNiagaraEditorStyle::Initialize();
 	FNiagaraEditorCommands::Register();
 
@@ -232,13 +369,15 @@ void FNiagaraEditorModule::StartupModule()
 	RegisterTypeUtilities(FNiagaraTypeDefinition::GetVec2Def(), MakeShareable(new FNiagaraEditorVector2TypeUtilities()));
 	RegisterTypeUtilities(FNiagaraTypeDefinition::GetVec3Def(), MakeShareable(new FNiagaraEditorVector3TypeUtilities()));
 	RegisterTypeUtilities(FNiagaraTypeDefinition::GetVec4Def(), MakeShareable(new FNiagaraEditorVector4TypeUtilities()));
+	RegisterTypeUtilities(FNiagaraTypeDefinition::GetQuatDef(), MakeShareable(new FNiagaraEditorQuatTypeUtilities()));
 	RegisterTypeUtilities(FNiagaraTypeDefinition::GetColorDef(), MakeShareable(new FNiagaraEditorColorTypeUtilities()));
+	RegisterTypeUtilities(FNiagaraTypeDefinition::GetMatrix4Def(), MakeShareable(new FNiagaraEditorMatrixTypeUtilities()));
 
-	RegisterTypeUtilities(FNiagaraTypeDefinition(UNiagaraDataInterfaceCurve::StaticClass()), MakeShared<FNiagaraDataInterfaceCurveTypeEditorUtilities>());
-	RegisterTypeUtilities(FNiagaraTypeDefinition(UNiagaraDataInterfaceVector2DCurve::StaticClass()), MakeShared<FNiagaraDataInterfaceCurveTypeEditorUtilities>());
-	RegisterTypeUtilities(FNiagaraTypeDefinition(UNiagaraDataInterfaceVectorCurve::StaticClass()), MakeShared<FNiagaraDataInterfaceVectorCurveTypeEditorUtilities>());
-	RegisterTypeUtilities(FNiagaraTypeDefinition(UNiagaraDataInterfaceVector4Curve::StaticClass()), MakeShared<FNiagaraDataInterfaceVectorCurveTypeEditorUtilities>());
-	RegisterTypeUtilities(FNiagaraTypeDefinition(UNiagaraDataInterfaceColorCurve::StaticClass()), MakeShared<FNiagaraDataInterfaceColorCurveTypeEditorUtilities>());
+	RegisterTypeUtilities(FNiagaraTypeDefinition(UNiagaraDataInterfaceCurve::StaticClass()), MakeShared<FNiagaraDataInterfaceCurveTypeEditorUtilities, ESPMode::ThreadSafe>());
+	RegisterTypeUtilities(FNiagaraTypeDefinition(UNiagaraDataInterfaceVector2DCurve::StaticClass()), MakeShared<FNiagaraDataInterfaceCurveTypeEditorUtilities, ESPMode::ThreadSafe>());
+	RegisterTypeUtilities(FNiagaraTypeDefinition(UNiagaraDataInterfaceVectorCurve::StaticClass()), MakeShared<FNiagaraDataInterfaceVectorCurveTypeEditorUtilities, ESPMode::ThreadSafe>());
+	RegisterTypeUtilities(FNiagaraTypeDefinition(UNiagaraDataInterfaceVector4Curve::StaticClass()), MakeShared<FNiagaraDataInterfaceVectorCurveTypeEditorUtilities, ESPMode::ThreadSafe>());
+	RegisterTypeUtilities(FNiagaraTypeDefinition(UNiagaraDataInterfaceColorCurve::StaticClass()), MakeShared<FNiagaraDataInterfaceColorCurveTypeEditorUtilities, ESPMode::ThreadSafe>());
 
 	FEdGraphUtilities::RegisterVisualPinFactory(GraphPanelPinFactory);
 
@@ -246,12 +385,51 @@ void FNiagaraEditorModule::StartupModule()
 
 	RegisterSettings();
 
-	// Register sequencer track editor
+	// Register sequencer track editors
 	ISequencerModule &SequencerModule = FModuleManager::LoadModuleChecked<ISequencerModule>("Sequencer");
 	CreateEmitterTrackEditorHandle = SequencerModule.RegisterTrackEditor(FOnCreateTrackEditor::CreateStatic(&FNiagaraEmitterTrackEditor::CreateTrackEditor));
+	CreateSystemTrackEditorHandle = SequencerModule.RegisterTrackEditor(FOnCreateTrackEditor::CreateStatic(&FNiagaraSystemTrackEditor::CreateTrackEditor));
 
 	SequencerModule.RegisterChannelInterface<FMovieSceneNiagaraEmitterChannel>();
 
+	CreateBoolParameterTrackEditorHandle = SequencerModule.RegisterTrackEditor(FOnCreateTrackEditor::CreateStatic(
+		&FNiagaraSystemBoolParameterTrackEditor::CreateTrackEditor));
+	CreateFloatParameterTrackEditorHandle = SequencerModule.RegisterTrackEditor(FOnCreateTrackEditor::CreateStatic(
+		&FNiagaraSystemParameterTrackEditor<UMovieSceneNiagaraFloatParameterTrack, UMovieSceneFloatSection>::CreateTrackEditor));
+	CreateIntegerParameterTrackEditorHandle = SequencerModule.RegisterTrackEditor(FOnCreateTrackEditor::CreateStatic(
+		&FNiagaraSystemParameterTrackEditor<UMovieSceneNiagaraIntegerParameterTrack, UMovieSceneIntegerSection>::CreateTrackEditor));
+	CreateVectorParameterTrackEditorHandle = SequencerModule.RegisterTrackEditor(FOnCreateTrackEditor::CreateStatic(
+		&FNiagaraSystemParameterTrackEditor<UMovieSceneNiagaraVectorParameterTrack, UMovieSceneVectorSection>::CreateTrackEditor));
+	CreateColorParameterTrackEditorHandle = SequencerModule.RegisterTrackEditor(FOnCreateTrackEditor::CreateStatic(
+		&FNiagaraSystemColorParameterTrackEditor::CreateTrackEditor));
+
+	RegisterParameterTrackCreatorForType(*FNiagaraBool::StaticStruct(), FOnCreateMovieSceneTrackForParameter::CreateLambda([](FNiagaraVariable InParameter) {
+		return NewObject<UMovieSceneNiagaraBoolParameterTrack>(); }));
+	RegisterParameterTrackCreatorForType(*FNiagaraFloat::StaticStruct(), FOnCreateMovieSceneTrackForParameter::CreateLambda([](FNiagaraVariable InParameter) {
+		return NewObject<UMovieSceneNiagaraFloatParameterTrack>(); }));
+	RegisterParameterTrackCreatorForType(*FNiagaraInt32::StaticStruct(), FOnCreateMovieSceneTrackForParameter::CreateLambda([](FNiagaraVariable InParameter) {
+		return NewObject<UMovieSceneNiagaraIntegerParameterTrack>(); }));
+	RegisterParameterTrackCreatorForType(*FNiagaraTypeDefinition::GetVec2Struct(), FOnCreateMovieSceneTrackForParameter::CreateLambda([](FNiagaraVariable InParameter) 
+	{
+		UMovieSceneNiagaraVectorParameterTrack* VectorTrack = NewObject<UMovieSceneNiagaraVectorParameterTrack>();
+		VectorTrack->SetChannelsUsed(2);
+		return VectorTrack;
+	}));
+	RegisterParameterTrackCreatorForType(*FNiagaraTypeDefinition::GetVec3Struct(), FOnCreateMovieSceneTrackForParameter::CreateLambda([](FNiagaraVariable InParameter)
+	{
+		UMovieSceneNiagaraVectorParameterTrack* VectorTrack = NewObject<UMovieSceneNiagaraVectorParameterTrack>();
+		VectorTrack->SetChannelsUsed(3);
+		return VectorTrack;
+	}));
+	RegisterParameterTrackCreatorForType(*FNiagaraTypeDefinition::GetVec4Struct(), FOnCreateMovieSceneTrackForParameter::CreateLambda([](FNiagaraVariable InParameter)
+	{
+		UMovieSceneNiagaraVectorParameterTrack* VectorTrack = NewObject<UMovieSceneNiagaraVectorParameterTrack>();
+		VectorTrack->SetChannelsUsed(4);
+		return VectorTrack;
+	}));
+	RegisterParameterTrackCreatorForType(*FNiagaraTypeDefinition::GetColorStruct(), FOnCreateMovieSceneTrackForParameter::CreateLambda([](FNiagaraVariable InParameter) {
+		return NewObject<UMovieSceneNiagaraColorParameterTrack>(); }));
+
 	// Register the shader queue processor (for cooking)
 	INiagaraModule& NiagaraModule = FModuleManager::LoadModuleChecked<INiagaraModule>("Niagara");
 	NiagaraModule.SetOnProcessShaderCompilationQueue(INiagaraModule::FOnProcessQueue::CreateLambda([]()
@@ -268,11 +446,47 @@ void FNiagaraEditorModule::StartupModule()
 	// Register the emitter merge handler.
 	ScriptMergeManager = MakeShared<FNiagaraScriptMergeManager>();
 	MergeEmitterHandle = NiagaraModule.RegisterOnMergeEmitter(INiagaraModule::FOnMergeEmitter::CreateSP(ScriptMergeManager.ToSharedRef(), &FNiagaraScriptMergeManager::MergeEmitter));
+
+	// Register the script compiler
+	ScriptCompilerHandle = NiagaraModule.RegisterScriptCompiler(INiagaraModule::FScriptCompiler::CreateLambda([this](const FNiagaraCompileRequestDataBase* CompileRequest, const FNiagaraCompileOptions& Options)
+	{
+		return CompileScript(CompileRequest, Options);
+	}));
+
+	PrecompilerHandle = NiagaraModule.RegisterPrecompiler(INiagaraModule::FOnPrecompile::CreateLambda([this](UObject* InObj)
+	{
+		return Precompile(InObj);
+	}));
+
+	// Register the create default script source handler.
+	CreateDefaultScriptSourceHandle = NiagaraModule.RegisterOnCreateDefaultScriptSource(
+		INiagaraModule::FOnCreateDefaultScriptSource::CreateLambda([](UObject* Outer) { return NewObject<UNiagaraScriptSource>(Outer); }));
+
+
+	TestCompileScriptCommand = IConsoleManager::Get().RegisterConsoleCommand(
+		TEXT("fx.TestCompileNiagaraScript"),
+		TEXT("Compiles the specified script on disk for the niagara vector vm"),
+		FConsoleCommandWithArgsDelegate::CreateRaw(this, &FNiagaraEditorModule::TestCompileScriptFromConsole));
+
+	DumpRapidIterationParametersForAsset = IConsoleManager::Get().RegisterConsoleCommand(
+		TEXT("fx.DumpRapidIterationParametersForAsset"),
+		TEXT("Dumps the values of the rapid iteration parameters for the specified asset by path."),
+		FConsoleCommandWithArgsDelegate::CreateStatic(&DumpRapidIterationParamersForAsset));
 }
 
 
 void FNiagaraEditorModule::ShutdownModule()
 {
+	// Ensure that we don't have any lingering compiles laying around that will explode after this module shuts down.
+	for (TObjectIterator<UNiagaraSystem> It; It; ++It)
+	{
+		UNiagaraSystem* Sys = *It;
+		if (Sys)
+		{
+			Sys->WaitForCompilationComplete();
+		}
+	}
+
 	MenuExtensibilityManager.Reset();
 	ToolBarExtensibilityManager.Reset();
 
@@ -303,18 +517,42 @@ void FNiagaraEditorModule::ShutdownModule()
 	if (SequencerModule != nullptr)
 	{
 		SequencerModule->UnRegisterTrackEditor(CreateEmitterTrackEditorHandle);
+		SequencerModule->UnRegisterTrackEditor(CreateSystemTrackEditorHandle);
+		SequencerModule->UnRegisterTrackEditor(CreateBoolParameterTrackEditorHandle);
+		SequencerModule->UnRegisterTrackEditor(CreateFloatParameterTrackEditorHandle);
+		SequencerModule->UnRegisterTrackEditor(CreateIntegerParameterTrackEditorHandle);
+		SequencerModule->UnRegisterTrackEditor(CreateVectorParameterTrackEditorHandle);
+		SequencerModule->UnRegisterTrackEditor(CreateColorParameterTrackEditorHandle);
 	}
 
 	INiagaraModule* NiagaraModule = FModuleManager::GetModulePtr<INiagaraModule>("Niagara");
 	if (NiagaraModule != nullptr)
 	{
 		NiagaraModule->UnregisterOnMergeEmitter(MergeEmitterHandle);
+		NiagaraModule->UnregisterOnCreateDefaultScriptSource(CreateDefaultScriptSourceHandle);
+		NiagaraModule->UnregisterScriptCompiler(ScriptCompilerHandle);
+		NiagaraModule->UnregisterPrecompiler(PrecompilerHandle);
 	}
 
 	// Verify that we've cleaned up all the view models in the world.
-	FNiagaraScriptViewModel::CleanAll();
 	FNiagaraSystemViewModel::CleanAll();
 	FNiagaraEmitterViewModel::CleanAll();
+	FNiagaraScriptViewModel::CleanAll();
+
+	if (TestCompileScriptCommand != nullptr)
+	{
+		IConsoleManager::Get().UnregisterConsoleObject(TestCompileScriptCommand);
+	}
+
+	if (DumpRapidIterationParametersForAsset != nullptr)
+	{
+		IConsoleManager::Get().UnregisterConsoleObject(DumpRapidIterationParametersForAsset);
+	}
+}
+
+FNiagaraEditorModule& FNiagaraEditorModule::Get()
+{
+	return FModuleManager::LoadModuleChecked<FNiagaraEditorModule>("NiagaraEditor");
 }
 
 void FNiagaraEditorModule::OnNiagaraSettingsChangedEvent(const FString& PropertyName, const UNiagaraSettings* Settings)
@@ -325,15 +563,20 @@ void FNiagaraEditorModule::OnNiagaraSettingsChangedEvent(const FString& Property
 	}
 }
 
-void FNiagaraEditorModule::RegisterTypeUtilities(FNiagaraTypeDefinition Type, TSharedRef<INiagaraEditorTypeUtilities> EditorUtilities)
+void FNiagaraEditorModule::RegisterTypeUtilities(FNiagaraTypeDefinition Type, TSharedRef<INiagaraEditorTypeUtilities, ESPMode::ThreadSafe> EditorUtilities)
 {
+	TypeEditorsCS.Lock();
 	TypeToEditorUtilitiesMap.Add(Type, EditorUtilities);
+	TypeEditorsCS.Unlock();
 }
 
 
-TSharedPtr<INiagaraEditorTypeUtilities> FNiagaraEditorModule::GetTypeUtilities(const FNiagaraTypeDefinition& Type)
+TSharedPtr<INiagaraEditorTypeUtilities, ESPMode::ThreadSafe> FNiagaraEditorModule::GetTypeUtilities(const FNiagaraTypeDefinition& Type)
 {
-	TSharedRef<INiagaraEditorTypeUtilities>* EditorUtilities = TypeToEditorUtilitiesMap.Find(Type);
+	TypeEditorsCS.Lock();
+	TSharedRef<INiagaraEditorTypeUtilities, ESPMode::ThreadSafe>* EditorUtilities = TypeToEditorUtilitiesMap.Find(Type);
+	TypeEditorsCS.Unlock();
+
 	if(EditorUtilities != nullptr)
 	{
 		return *EditorUtilities;
@@ -344,7 +587,7 @@ TSharedPtr<INiagaraEditorTypeUtilities> FNiagaraEditorModule::GetTypeUtilities(c
 		return EnumTypeUtilities;
 	}
 
-	return TSharedPtr<INiagaraEditorTypeUtilities>();
+	return TSharedPtr<INiagaraEditorTypeUtilities, ESPMode::ThreadSafe>();
 }
 
 TSharedRef<SWidget> FNiagaraEditorModule::CreateStackWidget(UNiagaraStackViewModel* StackViewModel) const
@@ -371,6 +614,36 @@ TSharedRef<FNiagaraScriptMergeManager> FNiagaraEditorModule::GetScriptMergeManag
 	return ScriptMergeManager.ToSharedRef();
 }
 
+void FNiagaraEditorModule::RegisterParameterTrackCreatorForType(const UScriptStruct& StructType, FOnCreateMovieSceneTrackForParameter CreateTrack)
+{
+	checkf(TypeToParameterTrackCreatorMap.Contains(&StructType) == false, TEXT("Type already registered"));
+	TypeToParameterTrackCreatorMap.Add(&StructType, CreateTrack);
+}
+
+void FNiagaraEditorModule::UnregisterParameterTrackCreatorForType(const UScriptStruct& StructType)
+{
+	TypeToParameterTrackCreatorMap.Remove(&StructType);
+}
+
+bool FNiagaraEditorModule::CanCreateParameterTrackForType(const UScriptStruct& StructType)
+{
+	return TypeToParameterTrackCreatorMap.Contains(&StructType);
+}
+
+UMovieSceneNiagaraParameterTrack* FNiagaraEditorModule::CreateParameterTrackForType(const UScriptStruct& StructType, FNiagaraVariable Parameter)
+{
+	FOnCreateMovieSceneTrackForParameter* CreateTrack = TypeToParameterTrackCreatorMap.Find(&StructType);
+	checkf(CreateTrack != nullptr, TEXT("Type not supported"));
+	UMovieSceneNiagaraParameterTrack* ParameterTrack = CreateTrack->Execute(Parameter);
+	ParameterTrack->SetParameter(Parameter);
+	return ParameterTrack;
+}
+
+const FNiagaraEditorCommands& FNiagaraEditorModule::Commands()
+{
+	return FNiagaraEditorCommands::Get();
+}
+
 void FNiagaraEditorModule::RegisterAssetTypeAction(IAssetTools& AssetTools, TSharedRef<IAssetTypeActions> Action)
 {
 	AssetTools.RegisterAssetTypeActions(Action);
diff --git a/Engine/Plugins/FX/Niagara/Source/NiagaraEditor/Private/NiagaraEditorSettings.cpp b/Engine/Plugins/FX/Niagara/Source/NiagaraEditor/Private/NiagaraEditorSettings.cpp
--- a/Engine/Plugins/FX/Niagara/Source/NiagaraEditor/Private/NiagaraEditorSettings.cpp
+++ b/Engine/Plugins/FX/Niagara/Source/NiagaraEditor/Private/NiagaraEditorSettings.cpp
@@ -6,6 +6,65 @@ UNiagaraEditorSettings::UNiagaraEditorSettings(const FObjectInitializer& ObjectI
 	: Super(ObjectInitlaizer)
 {
 	bAutoCompile = true;
+	bAutoPlay = true;
+	bResetSimulationOnChange = true;
+	bResimulateOnChangeWhilePaused = true;
+}
+
+bool UNiagaraEditorSettings::GetAutoCompile() const
+{
+	return bAutoCompile;
+}
+
+void UNiagaraEditorSettings::SetAutoCompile(bool bInAutoCompile)
+{
+	if (bAutoCompile != bInAutoCompile)
+	{
+		bAutoCompile = bInAutoCompile;
+		SaveConfig();
+	}
+}
+
+bool UNiagaraEditorSettings::GetAutoPlay() const
+{
+	return bAutoPlay;
+}
+
+void UNiagaraEditorSettings::SetAutoPlay(bool bInAutoPlay)
+{
+	if (bAutoPlay != bInAutoPlay)
+	{
+		bAutoPlay = bInAutoPlay;
+		SaveConfig();
+	}
+}
+
+bool UNiagaraEditorSettings::GetResetSimulationOnChange() const
+{
+	return bResetSimulationOnChange;
+}
+
+void UNiagaraEditorSettings::SetResetSimulationOnChange(bool bInResetSimulationOnChange)
+{
+	if (bResetSimulationOnChange != bInResetSimulationOnChange)
+	{
+		bResetSimulationOnChange = bInResetSimulationOnChange;
+		SaveConfig();
+	}
+}
+
+bool UNiagaraEditorSettings::GetResimulateOnChangeWhilePaused() const
+{
+	return bResimulateOnChangeWhilePaused;
+}
+
+void UNiagaraEditorSettings::SetResimulateOnChangeWhilePaused(bool bInResimulateOnChangeWhilePaused)
+{
+	if (bResimulateOnChangeWhilePaused != bInResimulateOnChangeWhilePaused)
+	{
+		bResimulateOnChangeWhilePaused = bInResimulateOnChangeWhilePaused;
+		SaveConfig();
+	}
 }
 
 FName UNiagaraEditorSettings::GetCategoryName() const
@@ -13,14 +72,11 @@ FName UNiagaraEditorSettings::GetCategoryName() const
 	return TEXT("Plugins");
 }
 
-#if WITH_EDITOR
 FText UNiagaraEditorSettings::GetSectionText() const
 {
 	return NSLOCTEXT("NiagaraEditorPlugin", "NiagaraEditorSettingsSection", "Niagara Editor");
 }
-#endif
 
-#if WITH_EDITOR
 void UNiagaraEditorSettings::PostEditChangeProperty(FPropertyChangedEvent& PropertyChangedEvent)
 {
 	if (PropertyChangedEvent.Property != nullptr)
@@ -34,7 +90,4 @@ UNiagaraEditorSettings::FOnNiagaraEditorSettingsChanged& UNiagaraEditorSettings:
 	return SettingsChangedDelegate;
 }
 
-UNiagaraEditorSettings::FOnNiagaraEditorSettingsChanged UNiagaraEditorSettings::SettingsChangedDelegate;
-#endif
-
-
+UNiagaraEditorSettings::FOnNiagaraEditorSettingsChanged UNiagaraEditorSettings::SettingsChangedDelegate;
\ No newline at end of file
diff --git a/Engine/Plugins/FX/Niagara/Source/NiagaraEditor/Private/NiagaraEditorStyle.cpp b/Engine/Plugins/FX/Niagara/Source/NiagaraEditor/Private/NiagaraEditorStyle.cpp
--- a/Engine/Plugins/FX/Niagara/Source/NiagaraEditor/Private/NiagaraEditorStyle.cpp
+++ b/Engine/Plugins/FX/Niagara/Source/NiagaraEditor/Private/NiagaraEditorStyle.cpp
@@ -8,6 +8,7 @@
 #include "Styling/SlateStyleRegistry.h"
 #include "Styling/SlateTypes.h"
 #include "Styling/CoreStyle.h"
+#include "Interfaces/IPluginManager.h"
 
 TSharedPtr< FSlateStyleSet > FNiagaraEditorStyle::NiagaraEditorStyleInstance = NULL;
 
@@ -33,6 +34,13 @@ FName FNiagaraEditorStyle::GetStyleSetName()
 	return StyleSetName;
 }
 
+FString RelativePathToPluginPath(const FString& RelativePath, const ANSICHAR* Extension)
+{
+	static FString ContentDir = IPluginManager::Get().FindPlugin(TEXT("Niagara"))->GetContentDir();
+	return (ContentDir / RelativePath) + Extension;
+}
+
+#define IMAGE_PLUGIN_BRUSH( RelativePath, ... ) FSlateImageBrush( RelativePathToPluginPath( RelativePath, ".png" ), __VA_ARGS__ )
 #define IMAGE_BRUSH( RelativePath, ... ) FSlateImageBrush( Style->RootToContentDir( RelativePath, TEXT(".png") ), __VA_ARGS__ )
 #define BOX_BRUSH( RelativePath, ... ) FSlateBoxBrush( Style->RootToContentDir( RelativePath, TEXT(".png") ), __VA_ARGS__ )
 #define BORDER_BRUSH( RelativePath, ... ) FSlateBorderBrush( Style->RootToContentDir( RelativePath, TEXT(".png") ), __VA_ARGS__ )
@@ -133,7 +141,7 @@ TSharedRef< FSlateStyleSet > FNiagaraEditorStyle::Create()
 		.SetShadowColorAndOpacity(FLinearColor(0, 0, 0, 0.9f));
 	Style->Set("NiagaraEditor.SelectedEmitter.UnsupportedSelectionText", SelectedEmitterUnsupportedSelectionText);
 
-	// Icons
+	// Toolbar Icons
 	Style->Set("NiagaraEditor.Apply", new IMAGE_BRUSH("Icons/icon_Niagara_Apply_40x", Icon40x40));
 	Style->Set("NiagaraEditor.Apply.Small", new IMAGE_BRUSH("Icons/icon_Niagara_Apply_40x", Icon20x20));
 	Style->Set("NiagaraEditor.Compile", new IMAGE_BRUSH("Icons/icon_compile_40x", Icon40x40));
@@ -144,6 +152,8 @@ TSharedRef< FSlateStyleSet > FNiagaraEditorStyle::Create()
 	Style->Set("NiagaraEditor.UnlockToChanges.Small", new IMAGE_BRUSH("Icons/icon_levels_unlocked_40x", Icon20x20));
 	Style->Set("NiagaraEditor.LockToChanges", new IMAGE_BRUSH("Icons/icon_levels_LockedReadOnly_40x", Icon40x40));
 	Style->Set("NiagaraEditor.LockToChanges.Small", new IMAGE_BRUSH("Icons/icon_levels_LockedReadOnly_40x", Icon20x20));
+	Style->Set("NiagaraEditor.SimulationOptions", new IMAGE_PLUGIN_BRUSH("Icons/Commands/icon_simulationOptions_40x", Icon40x40));
+	Style->Set("NiagaraEditor.SimulationOptions.Small", new IMAGE_PLUGIN_BRUSH("Icons/Commands/icon_simulationOptions_40x", Icon20x20));
 
 	Style->Set("Niagara.CompileStatus.Unknown", new IMAGE_BRUSH("Icons/CompileStatus_Working", Icon40x40));
 	Style->Set("Niagara.CompileStatus.Error",   new IMAGE_BRUSH("Icons/CompileStatus_Fail", Icon40x40));
@@ -151,6 +161,9 @@ TSharedRef< FSlateStyleSet > FNiagaraEditorStyle::Create()
 	Style->Set("Niagara.CompileStatus.Warning", new IMAGE_BRUSH("Icons/CompileStatus_Warning", Icon40x40));
 	Style->Set("Niagara.Asset.ReimportAsset.Needed", new IMAGE_BRUSH("Icons/icon_Reimport_Needed_40x", Icon40x40));
 	Style->Set("Niagara.Asset.ReimportAsset.Default", new IMAGE_BRUSH("Icons/icon_Reimport_40x", Icon40x40));
+
+	// Icons
+	Style->Set("NiagaraEditor.Isolate", new IMAGE_PLUGIN_BRUSH("Icons/Isolate", Icon16x16));
 	
 	// Emitter details customization
 	Style->Set("NiagaraEditor.MaterialWarningBorder", new BOX_CORE_BRUSH("Common/GroupBorderLight", FMargin(4.0f / 16.0f)));
@@ -167,6 +180,16 @@ TSharedRef< FSlateStyleSet > FNiagaraEditorStyle::Create()
 	Style->Set("NiagaraEditor.Thumbnails.Functions", new IMAGE_BRUSH("Icons/NiagaraScriptFunction_64x", Icon64x64));
 	Style->Set("NiagaraEditor.Thumbnails.Modules", new IMAGE_BRUSH("Icons/NiagaraScriptModules_64x", Icon64x64));
 
+	// Renderer class icons
+	Style->Set("ClassIcon.NiagaraSpriteRendererProperties", new IMAGE_PLUGIN_BRUSH("Icons/Renderers/renderer_sprite", Icon16x16));
+	Style->Set("ClassIcon.NiagaraMeshRendererProperties", new IMAGE_PLUGIN_BRUSH("Icons/Renderers/renderer_mesh", Icon16x16));
+	Style->Set("ClassIcon.NiagaraRibbonRendererProperties", new IMAGE_PLUGIN_BRUSH("Icons/Renderers/renderer_ribbon", Icon16x16));
+	Style->Set("ClassIcon.NiagaraLightRendererProperties", new IMAGE_PLUGIN_BRUSH("Icons/Renderers/renderer_light", Icon16x16));
+	Style->Set("ClassIcon.NiagaraRendererProperties", new IMAGE_PLUGIN_BRUSH("Icons/Renderers/renderer_default", Icon16x16));
+
+	// Niagara sequence
+	Style->Set("NiagaraEditor.NiagaraSequence.DefaultTrackColor", FLinearColor(0, .25f, 0));
+
 	return Style;
 }
 
