Copying //UE4/Dev-Niagara to //UE4/Dev-Main (Source: //UE4/Dev-Niagara @ 4074996)

diff --git a/Engine/Plugins/FX/Niagara/Source/NiagaraEditor/Private/NiagaraEditorUtilities.cpp b/Engine/Plugins/FX/Niagara/Source/NiagaraEditor/Private/NiagaraEditorUtilities.cpp
--- a/Engine/Plugins/FX/Niagara/Source/NiagaraEditor/Private/NiagaraEditorUtilities.cpp
+++ b/Engine/Plugins/FX/Niagara/Source/NiagaraEditor/Private/NiagaraEditorUtilities.cpp
@@ -22,14 +22,19 @@
 #include "HAL/PlatformApplicationMisc.h"
 #include "NiagaraEditorStyle.h"
 #include "EditorStyleSet.h"
-#include "NiagaraSystemViewModel.h"
-#include "NiagaraEmitterViewModel.h"
+#include "ViewModels/NiagaraSystemViewModel.h"
+#include "ViewModels/NiagaraEmitterViewModel.h"
 #include "HAL/PlatformApplicationMisc.h"
+#include "AssetRegistryModule.h"
+#include "Misc/FeedbackContext.h"
+#include "EdGraphSchema_Niagara.h"
+#include "HAL/PlatformFilemanager.h"
+#include "Misc/FileHelper.h"
+#include "EdGraph/EdGraphPin.h"
+#include "NiagaraNodeWriteDataSet.h"
 
 #define LOCTEXT_NAMESPACE "FNiagaraEditorUtilities"
 
-
-
 TSet<FName> FNiagaraEditorUtilities::GetSystemConstantNames()
 {
 	TSet<FName> SystemConstantNames;
@@ -57,7 +62,7 @@ void FNiagaraEditorUtilities::ResetVariableToDefaultValue(FNiagaraVariable& Vari
 	if (const UScriptStruct* ScriptStruct = Variable.GetType().GetScriptStruct())
 	{
 		FNiagaraEditorModule& NiagaraEditorModule = FModuleManager::GetModuleChecked<FNiagaraEditorModule>("NiagaraEditor");
-		TSharedPtr<INiagaraEditorTypeUtilities> TypeEditorUtilities = NiagaraEditorModule.GetTypeUtilities(Variable.GetType());
+		TSharedPtr<INiagaraEditorTypeUtilities, ESPMode::ThreadSafe> TypeEditorUtilities = NiagaraEditorModule.GetTypeUtilities(Variable.GetType());
 		if (TypeEditorUtilities.IsValid() && TypeEditorUtilities->CanProvideDefaultValue())
 		{
 			TypeEditorUtilities->UpdateVariableWithDefaultValue(Variable);
@@ -108,6 +113,7 @@ void FNiagaraEditorUtilities::GetParameterVariablesFromSystem(UNiagaraSystem& Sy
 				UNiagaraGraph::FFindInputNodeOptions FindOptions;
 				FindOptions.bIncludeAttributes = false;
 				FindOptions.bIncludeSystemConstants = false;
+				FindOptions.bIncludeTranslatorConstants = false;
 				FindOptions.bFilterDuplicates = true;
 
 				TArray<UNiagaraNodeInput*> InputNodes;
@@ -126,136 +132,6 @@ void FNiagaraEditorUtilities::GetParameterVariablesFromSystem(UNiagaraSystem& Sy
 	}
 }
 
-bool FNiagaraEditorUtilities::ConvertToMergedGraph(UNiagaraEmitter* InEmitter)
-{
-	if (InEmitter->GraphSource == nullptr)
-	{
-		UNiagaraScriptSource* Source = NewObject<UNiagaraScriptSource>(InEmitter, NAME_None, RF_Transactional);
-		if (Source)
-		{
-			UNiagaraGraph* CreatedGraph = NewObject<UNiagaraGraph>(Source, NAME_None, RF_Transactional);
-			Source->NodeGraph = CreatedGraph;
-
-			TArray<UNiagaraGraph*> GraphsToConvert;
-			TArray<ENiagaraScriptUsage> GraphUsages;
-
-			int32 yMaxPrevious = 0;
-
-			GraphsToConvert.Add(CastChecked<UNiagaraScriptSource>(InEmitter->SpawnScriptProps.Script->GetSource())->NodeGraph);
-			GraphUsages.Add(ENiagaraScriptUsage::ParticleSpawnScript);
-			GraphsToConvert.Add(CastChecked<UNiagaraScriptSource>(InEmitter->UpdateScriptProps.Script->GetSource())->NodeGraph);
-			GraphUsages.Add(ENiagaraScriptUsage::ParticleUpdateScript);
-
-			for (int32 i = 0; i < InEmitter->GetEventHandlers().Num(); i++)
-			{
-				if (InEmitter->GetEventHandlers()[i].Script != nullptr)
-				{
-					GraphsToConvert.Add(CastChecked<UNiagaraScriptSource>(InEmitter->GetEventHandlers()[i].Script->GetSource())->NodeGraph);
-					GraphUsages.Add(ENiagaraScriptUsage::ParticleEventScript);
-				}
-			}
-
-			for (int32 i = 0; i < GraphsToConvert.Num(); i++)
-			{
-				UNiagaraGraph* Graph = GraphsToConvert[i];
-				ENiagaraScriptUsage GraphUsage = GraphUsages[i];
-
-				TSet<UObject*> NodesToCopy;
-				TArray<UNiagaraNode*> SourceNodes;
-				Graph->GetNodesOfClass<UNiagaraNode>(SourceNodes);
-
-				int32 HighestY = -INT_MAX;
-				int32 EstimatedHeight = 300;
-				for (UNiagaraNode* SelectedGraphNode : SourceNodes)
-				{
-					if (SelectedGraphNode != nullptr)
-					{
-						if (SelectedGraphNode->NodePosY + EstimatedHeight > HighestY)
-						{
-							HighestY = SelectedGraphNode->NodePosY + EstimatedHeight;
-						}
-
-						if (SelectedGraphNode->CanDuplicateNode())
-						{
-							SelectedGraphNode->PrepareForCopying();
-							NodesToCopy.Add(SelectedGraphNode);
-						}
-						else
-						{
-							UE_LOG(LogNiagaraEditor, Error, TEXT("Could not clone node! %s"), *SelectedGraphNode->GetName());
-						}
-					}
-				}
-
-				FString ExportedText;
-				FEdGraphUtilities::ExportNodesToText(NodesToCopy, ExportedText);
-				FPlatformApplicationMisc::ClipboardCopy(*ExportedText);
-
-				// Grab the text to paste from the clipboard.
-				FPlatformApplicationMisc::ClipboardPaste(ExportedText);
-
-				// Import the nodes
-				TSet<UEdGraphNode*> PastedNodes;
-				FEdGraphUtilities::ImportNodesFromText(CreatedGraph, ExportedText, PastedNodes);
-
-
-				for (UEdGraphNode* PastedNode : PastedNodes)
-				{
-					PastedNode->CreateNewGuid();
-					PastedNode->NodePosY += yMaxPrevious;
-
-					UNiagaraNodeOutput* Output = Cast<UNiagaraNodeOutput>(PastedNode);
-					if (Output != nullptr)
-					{
-						Output->SetUsage(GraphUsage);
-					}
-				}
-
-				FixUpPastedInputNodes(CreatedGraph, PastedNodes);
-				yMaxPrevious = yMaxPrevious + HighestY;
-			}
-
-			InEmitter->GraphSource = Source;
-			InEmitter->SpawnScriptProps.Script->SetSource(Source);
-			InEmitter->UpdateScriptProps.Script->SetSource(Source);
-			for (int32 i = 0; i < InEmitter->GetEventHandlers().Num(); i++)
-			{
-				if (InEmitter->GetEventHandlers()[i].Script != nullptr)
-				{
-					InEmitter->GetEventHandlers()[i].Script->SetSource(Source);
-				}
-			}
-
-			// Also fix up any dependencies' referenced script type..
-			TArray<const UNiagaraGraph*> ReferencedGraphs;
-			CreatedGraph->GetAllReferencedGraphs(ReferencedGraphs);
-			for (const UNiagaraGraph* Graph : ReferencedGraphs)
-			{
-				TArray<UNiagaraNodeOutput*> Outputs;
-				Graph->FindOutputNodes(Outputs);
-				UNiagaraScript* Script = Cast<UNiagaraScript>(Graph->GetOuter());
-				if (Script)
-				{
-					for (UNiagaraNodeOutput* OutputNode : Outputs)
-					{
-						OutputNode->SetUsage(Script->GetUsage());
-					}
-				}
-			}
-			
-			// Now make sure that anyone referencing these graphs knows that they are out-of-date.
-			Source->MarkNotSynchronized();
-
-			//FString ErrorMessages;
-			//FNiagaraEditorModule& NiagaraEditorModule = FModuleManager::Get().LoadModuleChecked<FNiagaraEditorModule>(TEXT("NiagaraEditor"));
-			//NiagaraEditorModule.CompileScript(NewScript, ErrorMessages);
-			return true;
-		}
-	}
-	return false;
-}
-
-
 // TODO: This is overly complicated.
 void FNiagaraEditorUtilities::FixUpPastedInputNodes(UEdGraph* Graph, TSet<UEdGraphNode*> PastedNodes)
 {
@@ -360,6 +236,123 @@ void FNiagaraEditorUtilities::FixUpPastedInputNodes(UEdGraph* Graph, TSet<UEdGra
 	}
 }
 
+void FNiagaraEditorUtilities::WriteTextFileToDisk(FString SaveDirectory, FString FileName, FString TextToSave, bool bAllowOverwriting)
+{
+	IPlatformFile& PlatformFile = FPlatformFileManager::Get().GetPlatformFile();
+
+	// CreateDirectoryTree returns true if the destination
+	// directory existed prior to call or has been created
+	// during the call.
+	if (PlatformFile.CreateDirectoryTree(*SaveDirectory))
+	{
+		// Get absolute file path
+		FString AbsoluteFilePath = SaveDirectory + "/" + FileName;
+
+		// Allow overwriting or file doesn't already exist
+		if (bAllowOverwriting || !PlatformFile.FileExists(*AbsoluteFilePath))
+		{
+			if (FFileHelper::SaveStringToFile(TextToSave, *AbsoluteFilePath))
+			{
+				UE_LOG(LogNiagaraEditor, Log, TEXT("Wrote file to %s"), *AbsoluteFilePath);
+				return;
+			}
+
+		}
+	}
+}
+
+void FNiagaraEditorUtilities::GatherChangeIds(UNiagaraEmitter& Emitter, TMap<FGuid, FGuid>& ChangeIds, const FString& InDebugName, bool bWriteToLogDir)
+{
+	FString ExportText;
+	ChangeIds.Empty();
+	TArray<UNiagaraGraph*> Graphs;
+	TArray<UNiagaraScript*> Scripts;
+	Emitter.GetScripts(Scripts);
+
+	// First gather all the graphs used by this emitter..
+	for (UNiagaraScript* Script : Scripts)
+	{
+		if (Script != nullptr && Script->GetSource() != nullptr)
+		{
+			UNiagaraScriptSource* ScriptSource = Cast<UNiagaraScriptSource>(Script->GetSource());
+			if (ScriptSource != nullptr)
+			{
+				Graphs.AddUnique(ScriptSource->NodeGraph);
+			}
+
+			if (bWriteToLogDir)
+			{
+				FNiagaraVMExecutableDataId Id;
+				Script->ComputeVMCompilationId(Id);
+				FString KeyString;
+				Id.AppendKeyString(KeyString);
+
+				UEnum* FoundEnum = FindObject<UEnum>(ANY_PACKAGE, TEXT("ENiagaraScriptUsage"), true);
+
+				FString ResultsEnum = TEXT("??");
+				if (FoundEnum)
+				{
+					ResultsEnum = FoundEnum->GetNameStringByValue((int64)Script->Usage);
+				}
+
+				ExportText += FString::Printf(TEXT("Usage: %s CompileKey: %s\n"), *ResultsEnum, *KeyString );
+			}
+		}
+	}
+
+	
+	// Now gather all the node change id's within these graphs.
+	for (UNiagaraGraph* Graph : Graphs)
+	{
+		TArray<UNiagaraNode*> Nodes;
+		Graph->GetNodesOfClass(Nodes);
+
+		for (UNiagaraNode* Node : Nodes)
+		{
+			ChangeIds.Add(Node->NodeGuid, Node->GetChangeId());
+
+			if (bWriteToLogDir)
+			{
+				ExportText += FString::Printf(TEXT("%40s    guid: %25s    changeId: %25s\n"), *Node->GetName(), *Node->NodeGuid.ToString(), *Node->GetChangeId().ToString());
+
+			}
+		}
+	}
+
+	if (bWriteToLogDir)
+	{
+		FNiagaraEditorUtilities::WriteTextFileToDisk(FPaths::ProjectLogDir(), InDebugName + TEXT(".txt"), ExportText, true);
+	}
+	
+}
+
+void FNiagaraEditorUtilities::GatherChangeIds(UNiagaraGraph& Graph, TMap<FGuid, FGuid>& ChangeIds, const FString& InDebugName, bool bWriteToLogDir)
+{
+	ChangeIds.Empty();
+	TArray<UNiagaraGraph*> Graphs;
+	
+	FString ExportText;
+	// Now gather all the node change id's within these graphs.
+	{
+		TArray<UNiagaraNode*> Nodes;
+		Graph.GetNodesOfClass(Nodes);
+
+		for (UNiagaraNode* Node : Nodes)
+		{
+			ChangeIds.Add(Node->NodeGuid, Node->GetChangeId());
+			if (bWriteToLogDir)
+			{
+				ExportText += FString::Printf(TEXT("%40s    guid: %25s    changeId: %25s\n"), *Node->GetName(), *Node->NodeGuid.ToString(), *Node->GetChangeId().ToString());
+			}
+		}
+	}
+
+	if (bWriteToLogDir)
+	{
+		FNiagaraEditorUtilities::WriteTextFileToDisk(FPaths::ProjectLogDir(), InDebugName + TEXT(".txt"), ExportText, true);
+	}
+}
+
 FText FNiagaraEditorUtilities::StatusToText(ENiagaraScriptCompileStatus Status)
 {
 	switch (Status)
@@ -511,7 +504,7 @@ void FNiagaraEditorUtilities::CompileExistingEmitters(const TArray<UNiagaraEmitt
 		{
 			if (SystemIterator->ReferencesSourceEmitter(*Emitter))
 			{
-				SystemIterator->Compile(false);
+				SystemIterator->RequestCompile(false);
 
 				TArray<TSharedPtr<FNiagaraSystemViewModel>> ExistingSystemViewModels;
 				FNiagaraSystemViewModel::GetAllViewModelsForObject(*SystemIterator, ExistingSystemViewModels);
@@ -545,4 +538,284 @@ bool FNiagaraEditorUtilities::TryGetEventDisplayName(UNiagaraEmitter* Emitter, F
 	return false;
 }
 
+bool FNiagaraEditorUtilities::IsCompilableAssetClass(UClass* AssetClass)
+{
+	static const TSet<UClass*> CompilableClasses = { UNiagaraScript::StaticClass(), UNiagaraEmitter::StaticClass(), UNiagaraSystem::StaticClass() };
+	return CompilableClasses.Contains(AssetClass);
+}
+
+void FNiagaraEditorUtilities::MarkDependentCompilableAssetsDirty(TArray<UObject*> InObjects)
+{
+	const FText LoadAndMarkDirtyDisplayName = NSLOCTEXT("NiagaraEditor", "MarkDependentAssetsDirtySlowTask", "Loading and marking dependent assets dirty.");
+	GWarn->BeginSlowTask(LoadAndMarkDirtyDisplayName, true, true);
+
+	FAssetRegistryModule& AssetRegistryModule = FModuleManager::LoadModuleChecked<FAssetRegistryModule>("AssetRegistry");
+	TArray<FAssetIdentifier> ReferenceNames;
+
+	TArray<FAssetData> AssetsToLoadAndMarkDirty;
+	TArray<FAssetData> AssetsToCheck;
+
+	for (UObject* InObject : InObjects)
+	{	
+		AssetsToCheck.Add(FAssetData(InObject));
+	}
+
+	while (AssetsToCheck.Num() > 0)
+	{
+		FAssetData AssetToCheck = AssetsToCheck[0];
+		AssetsToCheck.RemoveAtSwap(0);
+		if (IsCompilableAssetClass(AssetToCheck.GetClass()))
+		{
+			if (AssetsToLoadAndMarkDirty.Contains(AssetToCheck) == false)
+			{
+				AssetsToLoadAndMarkDirty.Add(AssetToCheck);
+				TArray<FName> Referencers;
+				AssetRegistryModule.Get().GetReferencers(AssetToCheck.PackageName, Referencers);
+				for (FName& Referencer : Referencers)
+				{
+					AssetRegistryModule.Get().GetAssetsByPackageName(Referencer, AssetsToCheck);
+				}
+			}
+		}
+	}
+
+	int32 ItemIndex = 0;
+	for (FAssetData AssetDataToLoadAndMarkDirty : AssetsToLoadAndMarkDirty)
+	{	
+		if (GWarn->ReceivedUserCancel())
+		{
+			break;
+		}
+		GWarn->StatusUpdate(ItemIndex++, AssetsToLoadAndMarkDirty.Num(), LoadAndMarkDirtyDisplayName);
+		UObject* AssetToMarkDirty = AssetDataToLoadAndMarkDirty.GetAsset();
+		AssetToMarkDirty->Modify(true);
+	}
+
+	GWarn->EndSlowTask();
+}
+
+template<typename Action>
+void TraverseGraphFromOutputDepthFirst(const UEdGraphSchema_Niagara* Schema, UNiagaraNode* Node, Action& VisitAction)
+{
+	UNiagaraGraph* Graph = Node->GetNiagaraGraph();
+	TArray<UNiagaraNode*> Nodes;
+	Graph->BuildTraversal(Nodes, Node);
+	for (UNiagaraNode* GraphNode : Nodes)
+	{
+		VisitAction(Schema, GraphNode);
+	}
+}
+
+void FixUpNumericPinsVisitor(const UEdGraphSchema_Niagara* Schema, UNiagaraNode* Node)
+{
+	// Fix up numeric input pins and keep track of numeric types to decide the output type.
+	TArray<FNiagaraTypeDefinition> InputTypes;
+	TArray<UEdGraphPin*> InputPins;
+	Node->GetInputPins(InputPins);
+	for (UEdGraphPin* InputPin : InputPins)
+	{
+		if (InputPin->PinType.PinCategory == UEdGraphSchema_Niagara::PinCategoryType)
+		{
+			FNiagaraTypeDefinition InputPinType = Schema->PinToTypeDefinition(InputPin);
+
+			// If the input pin is the generic numeric type set it to the type of the linked output pin which should have been processed already.
+			if (InputPinType == FNiagaraTypeDefinition::GetGenericNumericDef() && InputPin->LinkedTo.Num() == 1)
+			{
+				UEdGraphPin* InputPinLinkedPin = UNiagaraNode::TraceOutputPin(InputPin->LinkedTo[0]);
+				FNiagaraTypeDefinition InputPinLinkedPinType = Schema->PinToTypeDefinition(InputPinLinkedPin);
+				if (InputPinLinkedPinType.IsValid())
+				{
+					// Only update the input pin type if the linked pin type is valid.
+					InputPin->PinType = Schema->TypeDefinitionToPinType(InputPinLinkedPinType);
+					InputPinType = InputPinLinkedPinType;
+				}
+			}
+
+			if (InputPinType == FNiagaraTypeDefinition::GetGenericNumericDef())
+			{
+				UE_LOG(LogNiagaraEditor, Error, TEXT("Unable to deduce type for numeric input pin. Node: %s Pin: %s"), *Node->GetName(), *InputPin->GetName());
+			}
+
+			InputTypes.Add(InputPinType);
+		}
+	}
+
+	// Fix up numeric outputs based on the inputs.
+	if (InputTypes.Num() > 0 && Node->GetNumericOutputTypeSelectionMode() != ENiagaraNumericOutputTypeSelectionMode::None)
+	{
+		FNiagaraTypeDefinition OutputNumericType = FNiagaraTypeDefinition::GetNumericOutputType(InputTypes, Node->GetNumericOutputTypeSelectionMode());
+		if (OutputNumericType != FNiagaraTypeDefinition::GetGenericNumericDef())
+		{
+			TArray<UEdGraphPin*> OutputPins;
+			Node->GetOutputPins(OutputPins);
+			for (UEdGraphPin* OutputPin : OutputPins)
+			{
+				FNiagaraTypeDefinition OutputPinType = Schema->PinToTypeDefinition(OutputPin);
+				if (OutputPinType == FNiagaraTypeDefinition::GetGenericNumericDef())
+				{
+					OutputPin->PinType = Schema->TypeDefinitionToPinType(OutputNumericType);
+				}
+			}
+		}
+	}
+}
+
+void FNiagaraEditorUtilities::FixUpNumericPins(const UEdGraphSchema_Niagara* Schema, UNiagaraNode* Node)
+{
+	auto FixUpVisitor = [&](const UEdGraphSchema_Niagara* LSchema, UNiagaraNode* LNode) { FixUpNumericPinsVisitor(LSchema, LNode); };
+	TraverseGraphFromOutputDepthFirst(Schema, Node, FixUpVisitor);
+}
+
+/* Go through the graph and attempt to auto-detect the type of any numeric pins by working back from the leaves of the graph. Only change the types of pins, not FNiagaraVariables.*/
+void PreprocessGraph(const UEdGraphSchema_Niagara* Schema, UNiagaraGraph* Graph, UNiagaraNodeOutput* OutputNode)
+{
+	check(OutputNode);
+	{
+		FNiagaraEditorUtilities::FixUpNumericPins(Schema, OutputNode);
+	}
+}
+
+/* Go through the graph and force any input nodes with Numeric types to a hard-coded type of float. This will allow modules and functions to compile properly.*/
+void PreProcessGraphForInputNumerics(const UEdGraphSchema_Niagara* Schema, UNiagaraGraph* Graph, TArray<FNiagaraVariable>& OutChangedNumericParams)
+{
+	// Visit all input nodes
+	TArray<UNiagaraNodeInput*> InputNodes;
+	Graph->FindInputNodes(InputNodes);
+	for (UNiagaraNodeInput* InputNode : InputNodes)
+	{
+		// See if any of the output pins are of Numeric type. If so, force to floats.
+		TArray<UEdGraphPin*> OutputPins;
+		InputNode->GetOutputPins(OutputPins);
+		for (UEdGraphPin* OutputPin : OutputPins)
+		{
+			FNiagaraTypeDefinition OutputPinType = Schema->PinToTypeDefinition(OutputPin);
+			if (OutputPinType == FNiagaraTypeDefinition::GetGenericNumericDef())
+			{
+				OutputPin->PinType = Schema->TypeDefinitionToPinType(FNiagaraTypeDefinition::GetFloatDef());
+			}
+		}
+
+		// Record that we touched this variable for later cleanup and make sure that the 
+		// variable's type now matches the pin.
+		if (InputNode->Input.GetType() == FNiagaraTypeDefinition::GetGenericNumericDef())
+		{
+			OutChangedNumericParams.Add(InputNode->Input);
+			InputNode->Input.SetType(FNiagaraTypeDefinition::GetFloatDef());
+		}
+	}
+}
+
+/* Should be called after all pins have been successfully auto-detected for type. This goes through and synchronizes any Numeric FNiagaraVarible outputs with the deduced pin type. This will allow modules and functions to compile properly.*/
+void PreProcessGraphForAttributeNumerics(const UEdGraphSchema_Niagara* Schema, UNiagaraGraph* Graph, UNiagaraNodeOutput* OutputNode, TArray<FNiagaraVariable>& OutChangedNumericParams)
+{
+	// Visit the output node
+	if (OutputNode != nullptr)
+	{
+		// For each pin, make sure that if it has a valid type, but the associated variable is still Numeric,
+		// force the variable to match the pin's new type. Record that we touched this variable for later cleanup.
+		TArray<UEdGraphPin*> InputPins;
+		OutputNode->GetInputPins(InputPins);
+		check(OutputNode->Outputs.Num() == InputPins.Num());
+		for (int32 i = 0; i < InputPins.Num(); i++)
+		{
+			FNiagaraVariable& Param = OutputNode->Outputs[i];
+			UEdGraphPin* InputPin = InputPins[i];
+
+			FNiagaraTypeDefinition InputPinType = Schema->PinToTypeDefinition(InputPin);
+			if (Param.GetType() == FNiagaraTypeDefinition::GetGenericNumericDef() &&
+				InputPinType != FNiagaraTypeDefinition::GetGenericNumericDef())
+			{
+				OutChangedNumericParams.Add(Param);
+				Param.SetType(InputPinType);
+			}
+		}
+	}
+}
+
+void FNiagaraEditorUtilities::ResolveNumerics(UNiagaraGraph* SourceGraph, bool bForceParametersToResolveNumerics, TArray<FNiagaraVariable>& ChangedNumericParams)
+{
+	const UEdGraphSchema_Niagara* Schema = CastChecked<UEdGraphSchema_Niagara>(SourceGraph->GetSchema());
+
+	// In the case of functions or modules, we may not have enough information at this time to fully resolve the type. In that case,
+	// we circumvent the resulting errors by forcing a type. This gives the user an appropriate level of type checking. We will, however need to clean this up in
+	// the parameters that we output.
+	//bool bForceParametersToResolveNumerics = InScript->IsStandaloneScript();
+	if (bForceParametersToResolveNumerics)
+	{
+		PreProcessGraphForInputNumerics(Schema, SourceGraph, ChangedNumericParams);
+	}
+
+	// Auto-deduce the input types for numerics in the graph and overwrite the types on the pins. If PreProcessGraphForInputNumerics occurred, then
+	// we will have pre-populated the inputs with valid types.
+	TArray<UNiagaraNodeOutput*> OutputNodes;
+	SourceGraph->FindOutputNodes(OutputNodes);
+
+	for (UNiagaraNodeOutput* OutputNode : OutputNodes)
+	{
+		PreprocessGraph(Schema, SourceGraph, OutputNode);
+
+		// Now that we've auto-deduced the types, we need to handle any lingering Numerics in the Output's FNiagaraVariable outputs. 
+		// We use the pin's deduced type to temporarily overwrite the variable's type.
+		if (bForceParametersToResolveNumerics)
+		{
+			PreProcessGraphForAttributeNumerics(Schema, SourceGraph, OutputNode, ChangedNumericParams);
+		}
+	}
+}
+
+void FNiagaraEditorUtilities::PreprocessFunctionGraph(const UEdGraphSchema_Niagara* Schema, UNiagaraGraph* Graph, const TArray<UEdGraphPin*>& CallInputs, const TArray<UEdGraphPin*>& CallOutputs, ENiagaraScriptUsage ScriptUsage)
+{
+	// Change any numeric inputs or outputs to match the types from the call node.
+	TArray<UNiagaraNodeInput*> InputNodes;
+
+	// Only handle nodes connected to the correct output node in the event of multiple output nodes in the graph.
+	UNiagaraGraph::FFindInputNodeOptions Options;
+	Options.bFilterByScriptUsage = true;
+	Options.TargetScriptUsage = ScriptUsage;
+
+	Graph->FindInputNodes(InputNodes, Options);
+
+	for (UNiagaraNodeInput* InputNode : InputNodes)
+	{
+		FNiagaraVariable& Input = InputNode->Input;
+		if (Input.GetType() == FNiagaraTypeDefinition::GetGenericNumericDef())
+		{
+			UEdGraphPin* const* MatchingPin = CallInputs.FindByPredicate([&](const UEdGraphPin* Pin) { return (Pin->PinName == Input.GetName()); });
+
+			if (MatchingPin != nullptr)
+			{
+				FNiagaraTypeDefinition PinType = Schema->PinToTypeDefinition(*MatchingPin);
+				Input.SetType(PinType);
+				TArray<UEdGraphPin*> OutputPins;
+				InputNode->GetOutputPins(OutputPins);
+				check(OutputPins.Num() == 1);
+				OutputPins[0]->PinType = (*MatchingPin)->PinType;
+			}
+		}
+	}
+
+	UNiagaraNodeOutput* OutputNode = Graph->FindOutputNode(ScriptUsage);
+	check(OutputNode);
+
+	TArray<UEdGraphPin*> InputPins;
+	OutputNode->GetInputPins(InputPins);
+
+	for (FNiagaraVariable& Output : OutputNode->Outputs)
+	{
+		if (Output.GetType() == FNiagaraTypeDefinition::GetGenericNumericDef())
+		{
+			UEdGraphPin* const* MatchingPin = CallOutputs.FindByPredicate([&](const UEdGraphPin* Pin) { return (Pin->PinName == Output.GetName()); });
+
+			if (MatchingPin != nullptr)
+			{
+				FNiagaraTypeDefinition PinType = Schema->PinToTypeDefinition(*MatchingPin);
+				Output.SetType(PinType);
+			}
+		}
+	}
+
+	FNiagaraEditorUtilities::FixUpNumericPins(Schema, OutputNode);
+
+}
+
 #undef LOCTEXT_NAMESPACE
diff --git a/Engine/Plugins/FX/Niagara/Source/NiagaraEditor/Private/NiagaraEmitterDetailsCustomization.cpp b/Engine/Plugins/FX/Niagara/Source/NiagaraEditor/Private/NiagaraEmitterDetailsCustomization.cpp
--- a/Engine/Plugins/FX/Niagara/Source/NiagaraEditor/Private/NiagaraEmitterDetailsCustomization.cpp
+++ b/Engine/Plugins/FX/Niagara/Source/NiagaraEditor/Private/NiagaraEmitterDetailsCustomization.cpp
@@ -1,272 +1,18 @@
 // Copyright 1998-2018 Epic Games, Inc. All Rights Reserved.
 
 #include "NiagaraEmitterDetailsCustomization.h"
-#include "NiagaraCommon.h"
-#include "NiagaraEmitter.h"
-#include "NiagaraRendererProperties.h"
-#include "NiagaraComponent.h"
-#include "NiagaraEditorStyle.h"
-
-#include "DetailWidgetRow.h"
 #include "PropertyHandle.h"
 #include "DetailLayoutBuilder.h"
-#include "IDetailGroup.h"
-#include "IDetailChildrenBuilder.h"
-#include "IDetailPropertyRow.h"
-#include "DetailCategoryBuilder.h"
-#include "PropertyCustomizationHelpers.h"
-#include "ScopedTransaction.h"
-
-#include "Widgets/Layout/SBorder.h"
-#include "Widgets/Text/STextBlock.h"
-#include "Widgets/Input/SButton.h"
-
-
-#define LOCTEXT_NAMESPACE "NiagaraSystemEditor"
-#if 0
-class SMaterialUsageWarning : public SCompoundWidget
-{
-public:
-	SLATE_BEGIN_ARGS(SMaterialUsageWarning)	{}
-	SLATE_END_ARGS();
-
-	void Construct(const FArguments& InArgs, UNiagaraEmitter* InEmitterProperties)
-	{
-		EmitterProperties = InEmitterProperties;
-		ChildSlot
-		[
-			SNew(SBorder)
-			.BorderImage(FNiagaraEditorStyle::Get().GetBrush("NiagaraEditor.MaterialWarningBorder"))
-			.BorderBackgroundColor(FLinearColor(.6f, 0.0f, 0.0f, 1.0f))
-			[
-				SNew(SHorizontalBox)
-				+ SHorizontalBox::Slot()
-				.VAlign(VAlign_Center)
-				.Padding(3, 0, 0, 0)
-				[
-					SNew(STextBlock)
-					.Text(this, &SMaterialUsageWarning::GetInvalidMessage)
-					.AutoWrapText(true)
-				]
-				+ SHorizontalBox::Slot()
-				.VAlign(VAlign_Top)
-				.Padding(5, 0, 0, 0)
-				.AutoWidth()
-				[
-					SNew(SButton)
-					.OnClicked(this, &SMaterialUsageWarning::OnFixNowClicked)
-					.ToolTipText(NSLOCTEXT("MaterialUsageWarning", "FixButtonLabelToolTip", "Fix the material issue now by modifying the material asset."))
-					.Content()
-					[
-						SNew(STextBlock)
-						.Text(NSLOCTEXT("MaterialUsageWarning", "FixButtonLabel", "Fix Now"))
-					]
-				]
-			]
-		];
-	}
-
-	EVisibility GetVisibility() const
-	{
-		bool bMaterialIsValid =
-			EmitterProperties->RendererProperties == nullptr ||
-			EmitterProperties->Material == nullptr ||
-			EmitterProperties->RendererProperties->IsMaterialValidForRenderer(EmitterProperties->Material, InvalidMessage);
-
-		return bMaterialIsValid
-			? EVisibility::Collapsed
-			: EVisibility::Visible;
-	}
-
-private:
-	FText GetInvalidMessage() const
-	{
-		return InvalidMessage;
-	}
-
-	FReply OnFixNowClicked()
-	{
-		FScopedTransaction ScopedTransaction(NSLOCTEXT("EmitterDetailsCustomization", "FixMaterialTransaction", "Fix material for emitter"));
-		EmitterProperties->RendererProperties->FixMaterial(EmitterProperties->Material);
-		return FReply::Handled();
-	}
-
-private:
-	UNiagaraEmitter* EmitterProperties;
-	mutable FText InvalidMessage;
-};
-#endif
-
-TSharedRef<IDetailCustomization> FNiagaraEmitterDetails::MakeInstance(TWeakObjectPtr<UNiagaraEmitter> EmitterProperties)
-{
-	return MakeShareable(new FNiagaraEmitterDetails(EmitterProperties));
-}
+#include "NiagaraEmitter.h"
 
-FNiagaraEmitterDetails::FNiagaraEmitterDetails(TWeakObjectPtr<UNiagaraEmitter> EmitterProperties)
-: EmitterProps(EmitterProperties)
+TSharedRef<IDetailCustomization> FNiagaraEmitterDetails::MakeInstance()
 {
+	return MakeShared<FNiagaraEmitterDetails>();
 }
 
 void FNiagaraEmitterDetails::CustomizeDetails(IDetailLayoutBuilder& InDetailLayout)
 {
-	DetailLayout = &InDetailLayout;
-	FName EmitterCategoryName = TEXT("Emitter");
-	IDetailCategoryBuilder& EmitterCategory = DetailLayout->EditCategory(EmitterCategoryName);
-	
-	TArray<TWeakObjectPtr<UObject>> ObjectsBeingCustomized;
-	InDetailLayout.GetObjectsBeingCustomized(ObjectsBeingCustomized);
-	UNiagaraEmitter* EmitterProperties = CastChecked<UNiagaraEmitter>(ObjectsBeingCustomized[0].Get());
-
-#if 0
-	IDetailCategoryBuilder& RenderCategory = DetailLayout->EditCategory(TEXT("Render"));
-	TSharedRef<SMaterialUsageWarning> MaterialWarning = SNew(SMaterialUsageWarning, CastChecked<UNiagaraEmitter>(EmitterProperties));
-	RenderCategory.AddCustomRow(FText())
-		.Visibility(TAttribute<EVisibility>::Create(TAttribute<EVisibility>::FGetter::CreateSP(MaterialWarning, &SMaterialUsageWarning::GetVisibility)))
-		[
-			MaterialWarning
-		];
-#endif
-	/*
-	TSharedRef<IPropertyHandle> UpdateScriptProps = DetailLayout->GetProperty(GET_MEMBER_NAME_CHECKED(UNiagaraEmitter, UpdateScriptProps));
-	TSharedRef<IPropertyHandle> SpawnScriptProps = DetailLayout->GetProperty(GET_MEMBER_NAME_CHECKED(UNiagaraEmitter, SpawnScriptProps));
-
-	BuildScriptProperties(SpawnScriptProps, TEXT("Spawn Script"), LOCTEXT("SpawnScriptDisplayName", "Spawn Script"));
-	BuildScriptProperties(UpdateScriptProps, TEXT("Update Script"), LOCTEXT("UpdateScriptDisplayName", "Update Script"));
-	*/
+	TSharedPtr<IPropertyHandle> EventHandlersPropertyHandle = InDetailLayout.GetProperty(UNiagaraEmitter::PrivateMemberNames::EventHandlerScriptProps);
+	EventHandlersPropertyHandle->MarkHiddenByCustomization();
 }
 
-void FNiagaraEmitterDetails::BuildScriptProperties(TSharedRef<IPropertyHandle> ScriptPropsHandle, FName Name, FText DisplayName)
-{
-//	/*
-//	DetailLayout->HideProperty(ScriptPropsHandle);
-//
-//	IDetailCategoryBuilder& ScriptCategory = DetailLayout->EditCategory(Name, DisplayName);
-//
-//	//Script
-//	TSharedPtr<IPropertyHandle> ScriptHandle = ScriptPropsHandle->GetChildHandle(GET_MEMBER_NAME_CHECKED(FNiagaraEmitterScriptProperties, Script));
-//	ScriptCategory.AddProperty(ScriptHandle);
-//
-//	//Constants
-//	/*
-//	bool bGenerateHeader = false;
-//	bool bDisplayResetToDefault = false;
-//	TSharedPtr<IPropertyHandle> Constants = ScriptPropsHandle->GetChildHandle(GET_MEMBER_NAME_CHECKED(FNiagaraEmitterScriptProperties, ExternalConstants));
-//	DetailLayout->HideProperty(Constants);
-//
-//	TSharedPtr<IPropertyHandle> Uniforms = Constants->GetChildHandle(GET_MEMBER_NAME_CHECKED(FNiagaraConstants, Uniforms));
-//	TSharedRef<FDetailArrayBuilder> ConstantsBuilder = MakeShareable(new FDetailArrayBuilder(Uniforms.ToSharedRef(), bGenerateHeader, bDisplayResetToDefault));
-//	ConstantsBuilder->OnGenerateArrayElementWidget(FOnGenerateArrayElementWidget::CreateSP(this, &FNiagaraEmitterDetails::OnGenerateConstantEntry));
-//	ScriptCategory.AddCustomBuilder(ConstantsBuilder);
-//	*/
-//
-//	//Scalar Constants.
-//	/*
-//	TSharedPtr<IPropertyHandle> ScalarConstants = Constants->GetChildHandle(GET_MEMBER_NAME_CHECKED(FNiagaraConstants, ScalarConstants));
-//	TSharedRef<FDetailArrayBuilder> ScalarConstantsBuilder = MakeShareable(new FDetailArrayBuilder(ScalarConstants.ToSharedRef(), bGenerateHeader, bDisplayResetToDefault));
-//	ScalarConstantsBuilder->OnGenerateArrayElementWidget(FOnGenerateArrayElementWidget::CreateSP(this, &FNiagaraEmitterDetails::OnGenerateScalarConstantEntry));
-//	ScriptCategory.AddCustomBuilder(ScalarConstantsBuilder);
-//	//Vector Constants.
-//	TSharedPtr<IPropertyHandle> VectorConstants = Constants->GetChildHandle(GET_MEMBER_NAME_CHECKED(FNiagaraConstants, VectorConstants));
-//	TSharedRef<FDetailArrayBuilder> VectorConstantsBuilder = MakeShareable(new FDetailArrayBuilder(VectorConstants.ToSharedRef(), bGenerateHeader, bDisplayResetToDefault));
-//	VectorConstantsBuilder->OnGenerateArrayElementWidget(FOnGenerateArrayElementWidget::CreateSP(this, &FNiagaraEmitterDetails::OnGenerateVectorConstantEntry));
-//	ScriptCategory.AddCustomBuilder(VectorConstantsBuilder);
-//	*/
-//	//Events
-//	//bGenerateHeader = true;
-//	//Generators
-//// 	TSharedPtr<IPropertyHandle> EventGenerators = ScriptPropsHandle->GetChildHandle(GET_MEMBER_NAME_CHECKED(FNiagaraEmitterScriptProperties, EventGenerators));
-//// 	TSharedRef<FDetailArrayBuilder> EventGeneratorsBuilder = MakeShareable(new FDetailArrayBuilder(EventGenerators.ToSharedRef(), bGenerateHeader, bDisplayResetToDefault));
-//// 	EventGeneratorsBuilder->OnGenerateArrayElementWidget(FOnGenerateArrayElementWidget::CreateSP(this, &FNiagaraEmitterDetails::OnGenerateEventGeneratorEntry));
-//// 	ScriptCategory.AddCustomBuilder(EventGeneratorsBuilder);
-//// 	//Receivers
-//// 	TSharedPtr<IPropertyHandle> EventReceivers = ScriptPropsHandle->GetChildHandle(GET_MEMBER_NAME_CHECKED(FNiagaraEmitterScriptProperties, EventReceivers));
-//// 	TSharedRef<FDetailArrayBuilder> EventReceiversBuilder = MakeShareable(new FDetailArrayBuilder(EventReceivers.ToSharedRef(), bGenerateHeader, bDisplayResetToDefault));
-//// 	EventReceiversBuilder->OnGenerateArrayElementWidget(FOnGenerateArrayElementWidget::CreateSP(this, &FNiagaraEmitterDetails::OnGenerateEventReceiverEntry));
-//// 	ScriptCategory.AddCustomBuilder(EventReceiversBuilder);
-//	
-}
-
-
-
-void FNiagaraEmitterDetails::OnGenerateConstantEntry(TSharedRef<IPropertyHandle> ElementProperty, int32 ElementIndex, IDetailChildrenBuilder& ChildrenBuilder)
-{
-// 	TSharedPtr<IPropertyHandle> ValueProperty = ElementProperty->GetChildHandle(GET_MEMBER_NAME_CHECKED(FNiagaraVariable, VarData));
-// 	TSharedPtr<IPropertyHandle> NameProperty = ElementProperty->GetChildHandle(GET_MEMBER_NAME_CHECKED(FNiagaraVariable, Name));
- 	
-// 	FName DisplayName;
-// 	NameProperty->GetValue(DisplayName);
-
-	//STOVEY. Not sure this is needed. System constants don't show up in the param UI anyway.
-	//Don't display system constants
-// 	if (UNiagaraComponent::FindSystemConstant(Var) == nullptr)
-// 	{
-//		ChildrenBuilder.AddChildProperty(ValueProperty.ToSharedRef()).DisplayName(FText::FromName(DisplayName));
-//	}
-}
-
-
-void FNiagaraEmitterDetails::OnGenerateEventGeneratorEntry(TSharedRef<IPropertyHandle> ElementProperty, int32 ElementIndex, IDetailChildrenBuilder& ChildrenBuilder)
-{
-	TSharedPtr<IPropertyHandle> IDProperty = ElementProperty->GetChildHandle(GET_MEMBER_NAME_CHECKED(FNiagaraEventGeneratorProperties, ID));
-	TSharedPtr<IPropertyHandle> NameProperty = IDProperty->GetChildHandle(GET_MEMBER_NAME_CHECKED(FNiagaraDataSetID, Name));
-
-	FName DisplayName;
-	NameProperty->GetValue(DisplayName);
-	//IDetailPropertyRow& Row = ChildrenBuilder.AddChildProperty(ElementProperty).DisplayName(FText::FromName(DisplayName));
-
-	IDetailGroup& GenGroup = ChildrenBuilder.AddGroup(DisplayName, FText::FromName(DisplayName));
-	uint32 NumChildren = 0;
-	if (ElementProperty->GetNumChildren(NumChildren) == FPropertyAccess::Success)
-	{
-		for (uint32 i = 0; i < NumChildren; ++i)
-		{
-			TSharedPtr<IPropertyHandle> Child = ElementProperty->GetChildHandle(i);
-			//Dont add the ID. We just grab it's name for the name region of this property.
-			if (Child.IsValid() && Child->GetProperty()->GetName() != GET_MEMBER_NAME_CHECKED(FNiagaraEventGeneratorProperties, ID).ToString())
-			{
-				GenGroup.AddPropertyRow(Child.ToSharedRef());
-			}
-		}
-	}
-}
-
-void FNiagaraEmitterDetails::OnGenerateEventReceiverEntry(TSharedRef<IPropertyHandle> ElementProperty, int32 ElementIndex, IDetailChildrenBuilder& ChildrenBuilder)
-{
-	TSharedPtr<IPropertyHandle> NameProperty = ElementProperty->GetChildHandle(GET_MEMBER_NAME_CHECKED(FNiagaraEventReceiverProperties, Name));
-
- 	FName DisplayName;
- 	NameProperty->GetValue(DisplayName);
-// 	ChildrenBuilder.AddChildProperty(ElementProperty).DisplayName(FText::FromName(DisplayName));
-
-	IDetailGroup& Group = ChildrenBuilder.AddGroup(DisplayName, FText::FromName(DisplayName));
-	uint32 NumChildren = 0;
-	if (ElementProperty->GetNumChildren(NumChildren) == FPropertyAccess::Success)
-	{
-		for (uint32 i = 0; i < NumChildren; ++i)
-		{
-			TSharedPtr<IPropertyHandle> Child = ElementProperty->GetChildHandle(i);
-			//Dont add the ID. We just grab it's name for the name region of this property.
-			if (Child.IsValid() && Child->GetProperty()->GetName() != GET_MEMBER_NAME_CHECKED(FNiagaraEventReceiverProperties, Name).ToString())
-			{
-				TSharedPtr<SWidget> NameWidget;
-				TSharedPtr<SWidget> ValueWidget;
-				FDetailWidgetRow DefaultDetailRow;
-			
-				IDetailPropertyRow& Row = Group.AddPropertyRow(Child.ToSharedRef());
-				Row.GetDefaultWidgets(NameWidget, ValueWidget, DefaultDetailRow);
-				Row.CustomWidget(true)
-					.NameContent()
-					[
-						NameWidget.ToSharedRef()
-					]
-					.ValueContent()
-					[
-						ValueWidget.ToSharedRef()
-					];
-			}
-		}
-	}
-}
-
-#undef LOCTEXT_NAMESPACE
-
diff --git a/Engine/Plugins/FX/Niagara/Source/NiagaraEditor/Private/NiagaraEmitterDetailsCustomization.h b/Engine/Plugins/FX/Niagara/Source/NiagaraEditor/Private/NiagaraEmitterDetailsCustomization.h
--- a/Engine/Plugins/FX/Niagara/Source/NiagaraEditor/Private/NiagaraEmitterDetailsCustomization.h
+++ b/Engine/Plugins/FX/Niagara/Source/NiagaraEditor/Private/NiagaraEmitterDetailsCustomization.h
@@ -2,14 +2,9 @@
 
 #pragma once
 
-#include "CoreMinimal.h"
-#include "UObject/WeakObjectPtr.h"
 #include "IDetailCustomization.h"
 
-class IDetailChildrenBuilder;
 class IDetailLayoutBuilder;
-class IPropertyHandle;
-class UNiagaraEmitter;
 
 /*-----------------------------------------------------------------------------
 FNiagaraEmitterDetails
@@ -19,26 +14,9 @@ class FNiagaraEmitterDetails : public IDetailCustomization
 {
 public:
 	/** Makes a new instance of this detail layout class for a specific detail view requesting it */
-	static TSharedRef<class IDetailCustomization> MakeInstance(TWeakObjectPtr<UNiagaraEmitter> EmitterProperties);
-
-	/** Constructor */
-	FNiagaraEmitterDetails(TWeakObjectPtr<UNiagaraEmitter> EmitterProperties);
+	static TSharedRef<IDetailCustomization> MakeInstance();
 
 	/** IDetailCustomization interface */
 	virtual void CustomizeDetails(IDetailLayoutBuilder& InDetailLayout) override;
-
-	void OnGenerateConstantEntry(TSharedRef<IPropertyHandle> ElementProperty, int32 ElementIndex, IDetailChildrenBuilder& ChildrenBuilder);
-
-	void OnGenerateScalarConstantEntry(TSharedRef<IPropertyHandle> ElementProperty, int32 ElementIndex, IDetailChildrenBuilder& ChildrenBuilder);
-	void OnGenerateVectorConstantEntry(TSharedRef<IPropertyHandle> ElementProperty, int32 ElementIndex, IDetailChildrenBuilder& ChildrenBuilder);
-	void OnGenerateEventGeneratorEntry(TSharedRef<IPropertyHandle> ElementProperty, int32 ElementIndex, IDetailChildrenBuilder& ChildrenBuilder);
-	void OnGenerateEventReceiverEntry(TSharedRef<IPropertyHandle> ElementProperty, int32 ElementIndex, IDetailChildrenBuilder& ChildrenBuilder);
-
-	void BuildScriptProperties(TSharedRef<IPropertyHandle> ScriptPropsHandle, FName Name, FText DisplayName);
-private:
-	/** Object that stores all of the possible parameters we can edit */
-	TWeakObjectPtr<UNiagaraEmitter> EmitterProps;
-
-	IDetailLayoutBuilder* DetailLayout;
 };
 
diff --git a/Engine/Plugins/FX/Niagara/Source/NiagaraEditor/Private/NiagaraEmitterEditorData.cpp b/Engine/Plugins/FX/Niagara/Source/NiagaraEditor/Private/NiagaraEmitterEditorData.cpp
--- a/Engine/Plugins/FX/Niagara/Source/NiagaraEditor/Private/NiagaraEmitterEditorData.cpp
+++ b/Engine/Plugins/FX/Niagara/Source/NiagaraEditor/Private/NiagaraEmitterEditorData.cpp
@@ -6,6 +6,8 @@
 UNiagaraEmitterEditorData::UNiagaraEmitterEditorData(const FObjectInitializer& ObjectInitializer)
 {
 	StackEditorData = ObjectInitializer.CreateDefaultSubobject<UNiagaraStackEditorData>(this, TEXT("StackEditorData"));
+	PlaybackRangeMin = 0;
+	PlaybackRangeMax = 10;
 }
 
 void UNiagaraEmitterEditorData::PostLoad()
@@ -20,4 +22,15 @@ void UNiagaraEmitterEditorData::PostLoad()
 UNiagaraStackEditorData& UNiagaraEmitterEditorData::GetStackEditorData() const
 {
 	return *StackEditorData;
+}
+
+TRange<float> UNiagaraEmitterEditorData::GetPlaybackRange() const
+{
+	return TRange<float>(PlaybackRangeMin, PlaybackRangeMax);
+}
+
+void UNiagaraEmitterEditorData::SetPlaybackRange(TRange<float> InPlaybackRange)
+{
+	PlaybackRangeMin = InPlaybackRange.GetLowerBoundValue();
+	PlaybackRangeMax = InPlaybackRange.GetUpperBoundValue();
 }
\ No newline at end of file
diff --git a/Engine/Plugins/FX/Niagara/Source/NiagaraEditor/Private/NiagaraEmitterEditorData.h b/Engine/Plugins/FX/Niagara/Source/NiagaraEditor/Private/NiagaraEmitterEditorData.h
--- a/Engine/Plugins/FX/Niagara/Source/NiagaraEditor/Private/NiagaraEmitterEditorData.h
+++ b/Engine/Plugins/FX/Niagara/Source/NiagaraEditor/Private/NiagaraEmitterEditorData.h
@@ -19,7 +19,17 @@ public:
 
 	UNiagaraStackEditorData& GetStackEditorData() const;
 
+	TRange<float> GetPlaybackRange() const;
+
+	void SetPlaybackRange(TRange<float> InPlaybackRange);
+
 private:
 	UPROPERTY(Instanced)
 	UNiagaraStackEditorData* StackEditorData;
+
+	UPROPERTY()
+	float PlaybackRangeMin;
+
+	UPROPERTY()
+	float PlaybackRangeMax;
 };
\ No newline at end of file
diff --git a/Engine/Plugins/FX/Niagara/Source/NiagaraEditor/Private/NiagaraEmitterFactoryNew.cpp b/Engine/Plugins/FX/Niagara/Source/NiagaraEditor/Private/NiagaraEmitterFactoryNew.cpp
--- a/Engine/Plugins/FX/Niagara/Source/NiagaraEditor/Private/NiagaraEmitterFactoryNew.cpp
+++ b/Engine/Plugins/FX/Niagara/Source/NiagaraEditor/Private/NiagaraEmitterFactoryNew.cpp
@@ -57,6 +57,7 @@ UObject* UNiagaraEmitterFactoryNew::FactoryCreateNew(UClass* Class, UObject* InP
 			NewEmitter->UpdateScriptProps.Script->SetSource(Source);
 			NewEmitter->EmitterSpawnScriptProps.Script->SetSource(Source);
 			NewEmitter->EmitterUpdateScriptProps.Script->SetSource(Source);
+			NewEmitter->GetGPUComputeScript()->SetSource(Source);
 			NewEmitter->SimTarget = ENiagaraSimTarget::CPUSim;
 
 			bool bCreateDefaultNodes = true;
@@ -135,23 +136,29 @@ UObject* UNiagaraEmitterFactoryNew::FactoryCreateNew(UClass* Class, UObject* InP
 					{
 						if (ParticleSpawnOutputNode)
 						{
+							TArray<FNiagaraVariable> Vars;
+							TArray<FString> Defaults;
 							{
 								FNiagaraVariable Var = SYS_PARAM_PARTICLES_SPRITE_SIZE;
 								FString DefaultValue = FNiagaraConstants::GetAttributeDefaultValue(Var);
-								FNiagaraStackGraphUtilities::AddParameterModuleToStack(Var, *ParticleSpawnOutputNode, INDEX_NONE,&DefaultValue);
+								Vars.Add(Var);
+								Defaults.Add(DefaultValue);
 							}
 
 							{
 								FNiagaraVariable Var = SYS_PARAM_PARTICLES_SPRITE_ROTATION;
 								FString DefaultValue = FNiagaraConstants::GetAttributeDefaultValue(Var);
-								FNiagaraStackGraphUtilities::AddParameterModuleToStack(Var, *ParticleSpawnOutputNode, INDEX_NONE, &DefaultValue);
+								Vars.Add(Var);
+								Defaults.Add(DefaultValue);
 							}
 
 							{
 								FNiagaraVariable Var = SYS_PARAM_PARTICLES_LIFETIME;
 								FString DefaultValue = FNiagaraConstants::GetAttributeDefaultValue(Var);
-								FNiagaraStackGraphUtilities::AddParameterModuleToStack(Var, *ParticleSpawnOutputNode, INDEX_NONE, &DefaultValue);
+								Vars.Add(Var);
+								Defaults.Add(DefaultValue);
 							}
+							FNiagaraStackGraphUtilities::AddParameterModuleToStack(Vars, *ParticleSpawnOutputNode, INDEX_NONE, Defaults);
 
 						}
 					}
diff --git a/Engine/Plugins/FX/Niagara/Source/NiagaraEditor/Private/NiagaraGraph.cpp b/Engine/Plugins/FX/Niagara/Source/NiagaraEditor/Private/NiagaraGraph.cpp
--- a/Engine/Plugins/FX/Niagara/Source/NiagaraEditor/Private/NiagaraGraph.cpp
+++ b/Engine/Plugins/FX/Niagara/Source/NiagaraEditor/Private/NiagaraGraph.cpp
@@ -25,6 +25,11 @@
 #include "INiagaraEditorTypeUtilities.h"
 #include "NiagaraConstants.h"
 #include "NiagaraEditorModule.h"
+#include "NiagaraEditorUtilities.h"
+#include "NiagaraNode.h"
+#include "EdGraphSchema_Niagara.h"
+#include "ViewModels/Stack/NiagaraParameterHandle.h"
+#include "ScopedTransaction.h"
 
 DECLARE_CYCLE_STAT(TEXT("NiagaraEditor - Graph - FindInputNodes"), STAT_NiagaraEditor_Graph_FindInputNodes, STATGROUP_NiagaraEditor);
 DECLARE_CYCLE_STAT(TEXT("NiagaraEditor - Graph - FindInputNodes_NotFilterUsage"), STAT_NiagaraEditor_Graph_FindInputNodes_NotFilterUsage, STATGROUP_NiagaraEditor);
@@ -34,24 +39,63 @@ DECLARE_CYCLE_STAT(TEXT("NiagaraEditor - Graph - FindInputNodes_FindInputNodes_S
 DECLARE_CYCLE_STAT(TEXT("NiagaraEditor - Graph - FindOutputNode"), STAT_NiagaraEditor_Graph_FindOutputNode, STATGROUP_NiagaraEditor);
 DECLARE_CYCLE_STAT(TEXT("NiagaraEditor - Graph - BuildTraversalHelper"), STAT_NiagaraEditor_Graph_BuildTraversalHelper, STATGROUP_NiagaraEditor);
 
+bool bWriteToLog = false;
+
+#define LOCTEXT_NAMESPACE "NiagaraGraph"
+
+FNiagaraGraphParameterReferenceCollection::FNiagaraGraphParameterReferenceCollection(const bool bInCreated)
+	: bCreated(bInCreated)
+{
+}
+
+bool FNiagaraGraphParameterReferenceCollection::WasCreated() const
+{
+	return bCreated;
+}
+
+FNiagaraGraphScriptUsageInfo::FNiagaraGraphScriptUsageInfo() 
+{
+	DataHash.AddZeroed(sizeof(FSHAHash));
+}
+
 UNiagaraGraph::UNiagaraGraph(const FObjectInitializer& ObjectInitializer)
 	: Super(ObjectInitializer)
+	, bFindParametersAllowed(true)
+	, bIsRenamingParameter(false)
 {
 	Schema = UEdGraphSchema_Niagara::StaticClass();
 	ChangeId = FGuid::NewGuid();
 }
+
+FDelegateHandle UNiagaraGraph::AddOnGraphNeedsRecompileHandler(const FOnGraphChanged::FDelegate& InHandler)
+{
+	return OnGraphNeedsRecompile.Add(InHandler);
+}
+
+void UNiagaraGraph::RemoveOnGraphNeedsRecompileHandler(FDelegateHandle Handle)
+{
+	OnGraphNeedsRecompile.Remove(Handle);
+}
+
 void UNiagaraGraph::NotifyGraphChanged(const FEdGraphEditAction& InAction)
 {
+	FindParameters();
 	if ((InAction.Action & GRAPHACTION_AddNode) != 0 || (InAction.Action & GRAPHACTION_RemoveNode) != 0 ||
 		(InAction.Action & GRAPHACTION_GenericNeedsRecompile) != 0)
 	{
-		MarkGraphRequiresSynchronization();
+		MarkGraphRequiresSynchronization(TEXT("Graph Changed"));
+	}
+	if ((InAction.Action & GRAPHACTION_GenericNeedsRecompile) != 0)
+	{
+		OnGraphNeedsRecompile.Broadcast(InAction);
+		return;
 	}
 	Super::NotifyGraphChanged(InAction);
 }
 
 void UNiagaraGraph::NotifyGraphChanged()
 {
+	FindParameters();
 	Super::NotifyGraphChanged();
 }
 
@@ -104,7 +148,7 @@ void UNiagaraGraph::PostLoad()
 	// If this is from a prior version, enforce a valid Change Id!
 	if (ChangeId.IsValid() == false)
 	{
-		MarkGraphRequiresSynchronization();
+		MarkGraphRequiresSynchronization(TEXT("Graph change id was invalid"));
 	}
 
 	// Assume that all externally referenced assets have changed, so update to match. They will return true if they have changed.
@@ -118,7 +162,11 @@ void UNiagaraGraph::PostLoad()
 		{
 			ReferencedAsset->ConditionalPostLoad();
 			NiagaraNode->ConditionalPostLoad();
-			bAnyExternalChanges |= (NiagaraNode->RefreshFromExternalChanges());
+			if (NiagaraNode->RefreshFromExternalChanges())
+			{
+				bAnyExternalChanges = true;
+				
+			}
 		}
 		else
 		{
@@ -126,17 +174,15 @@ void UNiagaraGraph::PostLoad()
 		}
 	}
 
-	if (bAnyExternalChanges)
-	{
-		MarkGraphRequiresSynchronization();
-		NotifyGraphNeedsRecompile();
-	}
+	RebuildCachedData();
 
 	if (GIsEditor)
 	{
 		SetFlags(RF_Transactional);
 	}
 
+	Parameters.Empty();
+	FindParameters();
 }
 
 void UNiagaraGraph::PostEditChangeProperty(FPropertyChangedEvent& PropertyChangedEvent)
@@ -149,26 +195,95 @@ class UNiagaraScriptSource* UNiagaraGraph::GetSource() const
 	return CastChecked<UNiagaraScriptSource>(GetOuter());
 }
 
-UEdGraphPin* UNiagaraGraph::FindParameterMapDefaultValuePin(const FName VariableName)
+FGuid UNiagaraGraph::GetCompileID(ENiagaraScriptUsage InUsage, const FGuid& InUsageId)
+{
+	RebuildCachedData();
+
+	// Since there gpu compute script contains spawn, update, and emitter logic, and we can only return one,
+	// just return the particle spawn script here.
+	if (InUsage == ENiagaraScriptUsage::ParticleGPUComputeScript)
+	{
+		InUsage = ENiagaraScriptUsage::ParticleSpawnScript;
+	}
+
+	for (int32 j = 0; j < CachedUsageInfo.Num(); j++)
+	{
+		if (UNiagaraScript::IsEquivalentUsage(CachedUsageInfo[j].UsageType, InUsage) && CachedUsageInfo[j].UsageId == InUsageId)
+		{
+			return CachedUsageInfo[j].GeneratedCompileId;
+		}
+	}
+
+	return FGuid();
+
+}
+
+UEdGraphPin* UNiagaraGraph::FindParameterMapDefaultValuePin(const FName VariableName, ENiagaraScriptUsage InUsage, ENiagaraScriptUsage InParentUsage) const
 {
 	TArray<UEdGraphPin*> MatchingDefaultPins;
-	TArray<UNiagaraNodeParameterMapGet*> GetNodes;
-	GetNodesOfClass<UNiagaraNodeParameterMapGet>(GetNodes);
+	
+	TArray<UNiagaraNode*> NodesTraversed;
+	BuildTraversal(NodesTraversed, InUsage, FGuid());
 
-	for (UNiagaraNodeParameterMapGet* GetNode : GetNodes)
+	UEdGraphPin* DefaultInputPin = nullptr;
+	for (UNiagaraNode* Node : NodesTraversed)
 	{
-		TArray<UEdGraphPin*> OutputPins;
-		GetNode->GetOutputPins(OutputPins);
-		for (UEdGraphPin* OutputPin : OutputPins)
+		UNiagaraNodeParameterMapGet* GetNode = Cast<UNiagaraNodeParameterMapGet>(Node);
+	
+		if (GetNode)
 		{
-			if (VariableName == OutputPin->PinName)
+			TArray<UEdGraphPin*> OutputPins;
+			GetNode->GetOutputPins(OutputPins);
+			for (UEdGraphPin* OutputPin : OutputPins)
 			{
-				if (UEdGraphPin* Pin = GetNode->GetDefaultPin(OutputPin))
+				if (VariableName == OutputPin->PinName)
 				{
-					return Pin;
+					UEdGraphPin* Pin = GetNode->GetDefaultPin(OutputPin);
+					if (Pin)
+					{
+						DefaultInputPin = Pin;
+						break;
+					}
 				}
 			}
 		}
+
+		if (DefaultInputPin != nullptr)
+		{
+			break;
+		}
+	}
+
+
+	// There are some pins 
+	if (DefaultInputPin && DefaultInputPin->LinkedTo.Num() != 0 && DefaultInputPin->LinkedTo[0] != nullptr)
+	{
+		UNiagaraNode* Owner = Cast<UNiagaraNode>(DefaultInputPin->LinkedTo[0]->GetOwningNode());
+		UEdGraphPin* PreviousInput = DefaultInputPin;
+		int32 NumIters = 0;
+		while (Owner)
+		{
+			// Check to see if there are any reroute or choose by usage nodes involved in this..
+			UEdGraphPin* InputPin = Owner->GetPassThroughPin(PreviousInput->LinkedTo[0], InParentUsage);
+			if (InputPin == nullptr)
+			{
+				return PreviousInput;
+			}
+			else if (InputPin->LinkedTo.Num() == 0)
+			{
+				return InputPin;
+			}
+
+			check(InputPin->LinkedTo[0] != nullptr);
+			Owner = Cast<UNiagaraNode>(InputPin->LinkedTo[0]->GetOwningNode());
+			PreviousInput = InputPin;
+			++NumIters;
+			check(NumIters < Nodes.Num()); // If you hit this assert then we have a cycle in our graph somewhere.
+		}
+	}
+	else
+	{
+		return DefaultInputPin;
 	}
 
 	return nullptr;
@@ -201,6 +316,21 @@ void UNiagaraGraph::FindOutputNodes(ENiagaraScriptUsage TargetUsageType, TArray<
 	OutputNodes = NodesFound;
 }
 
+void UNiagaraGraph::FindEquivalentOutputNodes(ENiagaraScriptUsage TargetUsageType, TArray<UNiagaraNodeOutput*>& OutputNodes) const
+{
+	TArray<UNiagaraNodeOutput*> NodesFound;
+	for (UEdGraphNode* Node : Nodes)
+	{
+		UNiagaraNodeOutput* OutNode = Cast<UNiagaraNodeOutput>(Node);
+		if (OutNode && UNiagaraScript::IsEquivalentUsage(OutNode->GetUsage(), TargetUsageType))
+		{
+			NodesFound.Add(OutNode);
+		}
+	}
+
+	OutputNodes = NodesFound;
+}
+
 UNiagaraNodeOutput* UNiagaraGraph::FindOutputNode(ENiagaraScriptUsage TargetUsageType, FGuid TargetUsageId) const
 {
 	SCOPE_CYCLE_COUNTER(STAT_NiagaraEditor_Graph_FindOutputNode);
@@ -217,6 +347,23 @@ UNiagaraNodeOutput* UNiagaraGraph::FindOutputNode(ENiagaraScriptUsage TargetUsag
 	return nullptr;
 }
 
+UNiagaraNodeOutput* UNiagaraGraph::FindEquivalentOutputNode(ENiagaraScriptUsage TargetUsageType, FGuid TargetUsageId) const
+{
+	SCOPE_CYCLE_COUNTER(STAT_NiagaraEditor_Graph_FindOutputNode);
+	for (UEdGraphNode* Node : Nodes)
+	{
+		if (UNiagaraNodeOutput* OutNode = Cast<UNiagaraNodeOutput>(Node))
+		{
+			if (UNiagaraScript::IsEquivalentUsage(OutNode->GetUsage(), TargetUsageType) && OutNode->GetUsageId() == TargetUsageId)
+			{
+				return OutNode;
+			}
+		}
+	}
+	return nullptr;
+}
+
+
 void BuildTraversalHelper(TArray<class UNiagaraNode*>& OutNodesTraversed, UNiagaraNode* CurrentNode)
 {
 	if (CurrentNode == nullptr)
@@ -252,7 +399,7 @@ void UNiagaraGraph::BuildTraversal(TArray<class UNiagaraNode*>& OutNodesTraverse
 	}
 }
 
-void UNiagaraGraph::BuildTraversal(TArray<class UNiagaraNode*>& OutNodesTraversed, UNiagaraNode* FinalNode) const
+void UNiagaraGraph::BuildTraversal(TArray<class UNiagaraNode*>& OutNodesTraversed, UNiagaraNode* FinalNode) 
 {
 	if (FinalNode)
 	{
@@ -275,8 +422,9 @@ void UNiagaraGraph::FindInputNodes(TArray<UNiagaraNodeInput*>& OutInputNodes, UN
 			UNiagaraNodeInput* NiagaraInputNode = Cast<UNiagaraNodeInput>(Node);
 			if (NiagaraInputNode != nullptr &&
 				((NiagaraInputNode->Usage == ENiagaraInputNodeUsage::Parameter && Options.bIncludeParameters) ||
-					(NiagaraInputNode->Usage == ENiagaraInputNodeUsage::Attribute && Options.bIncludeAttributes) ||
-					(NiagaraInputNode->Usage == ENiagaraInputNodeUsage::SystemConstant && Options.bIncludeSystemConstants)))
+				(NiagaraInputNode->Usage == ENiagaraInputNodeUsage::Attribute && Options.bIncludeAttributes) ||
+				(NiagaraInputNode->Usage == ENiagaraInputNodeUsage::SystemConstant && Options.bIncludeSystemConstants) || 
+				(NiagaraInputNode->Usage == ENiagaraInputNodeUsage::TranslatorConstant && Options.bIncludeTranslatorConstants)))
 			{
 				InputNodes.Add(NiagaraInputNode);
 			}
@@ -372,26 +520,113 @@ void UNiagaraGraph::GetParameters(TArray<FNiagaraVariable>& Inputs, TArray<FNiag
 // 	Outputs.Sort(SortVars);
 }
 
-void UNiagaraGraph::FindReadDataSetNodes(TArray<class UNiagaraNodeReadDataSet*>& OutReadNodes) const
+const TMap<FNiagaraVariable, FNiagaraGraphParameterReferenceCollection>& UNiagaraGraph::GetParameterMap()
 {
-	for (UEdGraphNode* Node : Nodes)
+	return Parameters;
+}
+
+void UNiagaraGraph::AddParameter(const FNiagaraVariable& Parameter)
+{
+	FNiagaraGraphParameterReferenceCollection* FoundParameterReferenceCollection = Parameters.Find(Parameter);
+	if (!FoundParameterReferenceCollection)
+	{
+		FNiagaraGraphParameterReferenceCollection NewReferenceCollection = FNiagaraGraphParameterReferenceCollection(true /*bCreated*/);
+		NewReferenceCollection.Graph = this;
+		Parameters.Add(Parameter, NewReferenceCollection);
+	}
+}
+
+void UNiagaraGraph::RemoveParameter(const FNiagaraVariable& Parameter, const bool bNotifyGraphChanged /*= true*/)
+{
+	FScopedTransaction RemoveParametersWithPins(LOCTEXT("RemoveParametersWithPins", "Remove parameter and referenced pins"));
+	FNiagaraGraphParameterReferenceCollection* ReferenceCollection = Parameters.Find(Parameter);
+	if (ReferenceCollection)
 	{
-		if (UNiagaraNodeReadDataSet* InNode = Cast<UNiagaraNodeReadDataSet>(Node))
+		// Prevent finding all parameters and metadata when renaming each pin.
+		SetFindParametersAllowed(false);
+
+		for (int32 Index = 0; Index < ReferenceCollection->ParameterReferences.Num(); Index++)
+		{
+			const FNiagaraGraphParameterReference& Reference = ReferenceCollection->ParameterReferences[Index];
+			UNiagaraNode* Node = Reference.Value.Get();
+			if (Node)
+			{
+				UEdGraphPin* Pin = Node->GetPinByPersistentGuid(Reference.Key);
+				if (Pin)
+				{
+					Node->RemovePin(Pin);
+				}
+			}
+		}
+
+		Parameters.Remove(Parameter);
+
+		SetFindParametersAllowed(true);
+
+		if (bNotifyGraphChanged)
 		{
-			OutReadNodes.Add(InNode);
+			NotifyGraphChanged();
 		}
 	}
 }
 
-void UNiagaraGraph::FindWriteDataSetNodes(TArray<class UNiagaraNodeWriteDataSet*>& OutWriteNodes) const
+bool UNiagaraGraph::RenameParameter(const FNiagaraVariable& Parameter, FName NewName, const bool bInNotifyGraphChanged /*= true*/)
 {
-	for (UEdGraphNode* Node : Nodes)
+	// Block rename when already renaming. This prevents recursion when CommitEditablePinName is called on referenced nodes. 
+	if (bIsRenamingParameter)
+	{
+		return false;
+	}
+	bIsRenamingParameter = true;
+
+	FScopedTransaction RemoveParametersWithPins(LOCTEXT("RenameParameter", "Rename parameter, referenced pins and metadata"));
+
+	// Prevent finding all parameters and metadata when renaming each pin.
+	SetFindParametersAllowed(false);
+	
+	// Create the new parameter
+	FNiagaraVariable NewParameter = Parameter;
+	NewParameter.SetName(NewName);
+	
+	FNiagaraGraphParameterReferenceCollection* ReferenceCollection = Parameters.Find(Parameter);
+	if (ReferenceCollection)
 	{
-		if (UNiagaraNodeWriteDataSet* InNode = Cast<UNiagaraNodeWriteDataSet>(Node))
+		const FText NewNameText = FText::FromName(NewName);
+		FNiagaraGraphParameterReferenceCollection NewReferences = *ReferenceCollection;
+		for (FNiagaraGraphParameterReference& Reference : NewReferences.ParameterReferences)
 		{
-			OutWriteNodes.Add(InNode);
+			UNiagaraNode* Node = Reference.Value.Get();
+			if (Node)
+			{
+				UEdGraphPin* Pin = Node->GetPinByPersistentGuid(Reference.Key);
+				if (Pin)
+				{
+					Node->CommitEditablePinName(NewNameText, Pin);
+				}
+			}
 		}
+
+		Parameters.Remove(Parameter);
+		Parameters.Add(NewParameter, NewReferences);
+	}
+
+	// Swap metadata to the new parameter
+	FNiagaraVariableMetaData* Metadata = GetMetaData(Parameter);
+	if (Metadata)
+	{
+		FNiagaraVariableMetaData MetadataCopy = *Metadata;
+		VariableToMetaData.Remove(Parameter);
+		VariableToMetaData.Add(NewParameter, MetadataCopy);
 	}
+
+	SetFindParametersAllowed(true);
+	bIsRenamingParameter = false;
+
+	if (bInNotifyGraphChanged)
+	{
+		NotifyGraphChanged();
+	}
+	return true;
 }
 
 int32 UNiagaraGraph::GetOutputNodeVariableIndex(const FNiagaraVariable& Variable)const
@@ -484,6 +719,7 @@ void UNiagaraGraph::NotifyGraphNeedsRecompile()
 	NotifyGraphChanged(Action);
 }
 
+
 void UNiagaraGraph::NotifyGraphDataInterfaceChanged()
 {
 	OnDataInterfaceChangedDelegate.Broadcast();
@@ -499,6 +735,229 @@ void UNiagaraGraph::SubsumeExternalDependencies(TMap<const UObject*, UObject*>&
 	}
 }
 
+void UNiagaraGraph::RebuildCachedData(bool bForce)
+{
+	// If the graph hasn't changed since last rebuild, then do nothing.
+	if (!bForce && ChangeId == LastBuiltTraversalDataChangeId && LastBuiltTraversalDataChangeId.IsValid())
+	{
+		return;
+	}
+
+	// First find all the output nodes
+	TArray<UNiagaraNodeOutput*> NiagaraOutputNodes;
+	GetNodesOfClass<UNiagaraNodeOutput>(NiagaraOutputNodes);
+
+	// Now build the new cache..
+	TArray<FNiagaraGraphScriptUsageInfo> NewUsageCache;
+	NewUsageCache.AddDefaulted(NiagaraOutputNodes.Num());
+
+	UEnum* FoundEnum = nullptr;
+	bool bNeedsAnyNewCompileIds = false;
+
+	for (int32 i = 0; i < NiagaraOutputNodes.Num(); i++)
+	{
+		UNiagaraNodeOutput* OutputNode = NiagaraOutputNodes[i];
+		NewUsageCache[i].UsageType = OutputNode->GetUsage();
+		NewUsageCache[i].UsageId = OutputNode->GetUsageId();
+
+		BuildTraversal(NewUsageCache[i].Traversal, OutputNode);
+
+		int32 FoundMatchIdx = INDEX_NONE;
+		for (int32 j = 0; j < CachedUsageInfo.Num(); j++)
+		{
+			if (UNiagaraScript::IsEquivalentUsage(CachedUsageInfo[j].UsageType, NewUsageCache[i].UsageType) && CachedUsageInfo[j].UsageId == NewUsageCache[i].UsageId)
+			{
+				FoundMatchIdx = j;
+				break;
+			}
+		}
+
+		// Now compare the change id's of all the nodes in the traversal by hashing them up and comparing the hash
+		// now with the hash from previous runs.
+		FSHA1 HashState;
+		for (UNiagaraNode* Node : NewUsageCache[i].Traversal)
+		{
+			FGuid Guid = Node->GetChangeId();
+			HashState.Update((const uint8*)&Guid, sizeof(FGuid));
+		}
+		HashState.Final();
+
+		// We can't store in a FShaHash struct directly because you can't UProperty it. Using a standin of the same size.
+		check(sizeof(uint8)*NewUsageCache[i].DataHash.Num() == sizeof(FSHAHash));
+		HashState.GetHash(&NewUsageCache[i].DataHash[0]);
+
+		bool bNeedsNewCompileId = true;
+
+		// Now compare the hashed data. If it is the same as before, then leave the compile ID as-is. If it is different, generate a new guid.
+		if (FoundMatchIdx != INDEX_NONE)
+		{
+			if (NewUsageCache[i].DataHash == CachedUsageInfo[FoundMatchIdx].DataHash)
+			{
+				NewUsageCache[i].GeneratedCompileId = CachedUsageInfo[FoundMatchIdx].GeneratedCompileId;
+				bNeedsNewCompileId = false;
+			}
+		}
+
+		if (bNeedsNewCompileId)
+		{
+			NewUsageCache[i].GeneratedCompileId = FGuid::NewGuid();
+			bNeedsAnyNewCompileIds = true;
+		}
+
+		// TODO sckime debug logic... should be disabled or put under a cvar in the future
+		{
+
+			if (FoundEnum == nullptr)
+			{
+				FoundEnum = FindObject<UEnum>(ANY_PACKAGE, TEXT("ENiagaraScriptUsage"), true);
+			}
+
+			FString ResultsEnum = TEXT("??");
+			if (FoundEnum)
+			{
+				ResultsEnum = FoundEnum->GetNameStringByValue((int64)NewUsageCache[i].UsageType);
+			}
+
+			if (bNeedsNewCompileId)
+			{
+				//UE_LOG(LogNiagaraEditor, Log, TEXT("'%s' changes detected in %s .. new guid: %s"), *GetFullName(), *ResultsEnum, *NewUsageCache[i].GeneratedCompileId.ToString());
+			}
+			else
+			{
+				//UE_LOG(LogNiagaraEditor, Log, TEXT("'%s' changes NOT detected in %s .. keeping guid: %s"), *GetFullName(), *ResultsEnum, *NewUsageCache[i].GeneratedCompileId.ToString());
+			}
+		}
+	}
+
+	// Debug logic, usually disabled at top of file.
+	if (bNeedsAnyNewCompileIds && bWriteToLog)
+	{
+		TMap<FGuid, FGuid> ComputeChangeIds;
+		FNiagaraEditorUtilities::GatherChangeIds(*this, ComputeChangeIds, GetName());
+	}
+
+	// Now update the cache with the newly computed results.
+	CachedUsageInfo = NewUsageCache;
+	LastBuiltTraversalDataChangeId = ChangeId;
+
+}
+
+void UNiagaraGraph::SynchronizeInternalCacheWithGraph(UNiagaraGraph* Other)
+{
+	// Force us to rebuild the cache, note that this builds traversals and everything else, keeping it in sync if nothing changed from the current version.
+	RebuildCachedData(true);
+	
+	UEnum* FoundEnum = nullptr;
+
+	// Now go through all of the other graph's usage info. If we find a match for its usage and our data hashes match, use the generated compile id from
+	// the other graph.
+	for (int32 i = 0; i < CachedUsageInfo.Num(); i++)
+	{
+		int32 FoundMatchIdx = INDEX_NONE;
+		for (int32 j = 0; j < Other->CachedUsageInfo.Num(); j++)
+		{
+			if (UNiagaraScript::IsEquivalentUsage(Other->CachedUsageInfo[j].UsageType, CachedUsageInfo[i].UsageType) && Other->CachedUsageInfo[j].UsageId == CachedUsageInfo[i].UsageId)
+			{
+				FoundMatchIdx = j;
+				break;
+			}
+		}
+
+		if (FoundMatchIdx != INDEX_NONE)
+		{
+			if (CachedUsageInfo[i].DataHash == Other->CachedUsageInfo[FoundMatchIdx].DataHash)
+			{
+				CachedUsageInfo[i].GeneratedCompileId = Other->CachedUsageInfo[FoundMatchIdx].GeneratedCompileId;		
+
+				// TODO sckime debug logic... should be disabled or put under a cvar in the future
+				{
+					if (FoundEnum == nullptr)
+					{
+						FoundEnum = FindObject<UEnum>(ANY_PACKAGE, TEXT("ENiagaraScriptUsage"), true);
+					}
+
+					FString ResultsEnum = TEXT("??");
+					if (FoundEnum)
+					{
+						ResultsEnum = FoundEnum->GetNameStringByValue((int64)CachedUsageInfo[i].UsageType);
+					}
+					UE_LOG(LogNiagaraEditor, Log, TEXT("'%s' changes synchronized with master script in %s .. synced guid: %s"), *GetFullName(), *ResultsEnum, *CachedUsageInfo[i].GeneratedCompileId.ToString());
+				}
+			}
+		}
+	}
+
+	if (bWriteToLog)
+	{
+		TMap<FGuid, FGuid> ComputeChangeIds;
+		FNiagaraEditorUtilities::GatherChangeIds(*this, ComputeChangeIds, GetName() + TEXT(".Synced"));
+	}
+}
+
+
+void UNiagaraGraph::InvalidateCachedCompileIds()
+{
+	Modify();
+	CachedUsageInfo.Empty();
+	MarkGraphRequiresSynchronization(__FUNCTION__);
+}
+
+void UNiagaraGraph::GatherExternalDependencyIDs(ENiagaraScriptUsage InUsage, const FGuid& InUsageId, TArray<FGuid>& InReferencedIDs, TArray<UObject*>& InReferencedObjs)
+{
+	RebuildCachedData();
+	
+	// Particle compute scripts get all particle scripts baked into their dependency chain. 
+	if (InUsage == ENiagaraScriptUsage::ParticleGPUComputeScript)
+	{
+		for (int32 i = 0; i < CachedUsageInfo.Num(); i++)
+		{
+			// Add all chains that we depend on.
+			if (UNiagaraScript::IsUsageDependentOn(InUsage, CachedUsageInfo[i].UsageType)) 
+			{
+				// Skip adding to list because we already did it in GetCompileId above if spawn script.
+				if (CachedUsageInfo[i].UsageType != ENiagaraScriptUsage::ParticleSpawnScript)
+				{
+					InReferencedIDs.Add(CachedUsageInfo[i].GeneratedCompileId);
+					InReferencedObjs.Add(CachedUsageInfo[i].Traversal.Last());
+				}
+
+				for (UNiagaraNode* Node : CachedUsageInfo[i].Traversal)
+				{
+					Node->GatherExternalDependencyIDs(InUsage, InUsageId, InReferencedIDs, InReferencedObjs);
+				}
+			}
+		}
+	}
+	// Otherwise, just add downstream dependencies for the specific usage type we're on.
+	else
+	{
+		for (int32 i = 0; i < CachedUsageInfo.Num(); i++)
+		{
+			// First add our direct dependency chain...
+			if (UNiagaraScript::IsEquivalentUsage(CachedUsageInfo[i].UsageType, InUsage) && CachedUsageInfo[i].UsageId == InUsageId)
+			{
+				// Skip adding to list because we already did it in GetCompileId above.
+				for (UNiagaraNode* Node : CachedUsageInfo[i].Traversal)
+				{
+					Node->GatherExternalDependencyIDs(InUsage, InUsageId, InReferencedIDs, InReferencedObjs);
+				}
+			}
+			// Now add any other dependency chains that we might have...
+			else if (UNiagaraScript::IsUsageDependentOn(InUsage, CachedUsageInfo[i].UsageType))
+			{
+				InReferencedIDs.Add(CachedUsageInfo[i].GeneratedCompileId);
+				InReferencedObjs.Add(CachedUsageInfo[i].Traversal.Last());
+
+				for (UNiagaraNode* Node : CachedUsageInfo[i].Traversal)
+				{
+					Node->GatherExternalDependencyIDs(InUsage, InUsageId, InReferencedIDs, InReferencedObjs);
+				}
+			}
+		}
+	}
+}
+
+
 void UNiagaraGraph::GetAllReferencedGraphs(TArray<const UNiagaraGraph*>& Graphs) const
 {
 	Graphs.AddUnique(this);
@@ -550,11 +1009,11 @@ bool UNiagaraGraph::IsOtherSynchronized(const FGuid& InChangeId) const
 }
 
 /** Identify that this graph has undergone changes that will require synchronization with a compiled script.*/
-void UNiagaraGraph::MarkGraphRequiresSynchronization()
+void UNiagaraGraph::MarkGraphRequiresSynchronization(FString Reason)
 {
 	Modify();
 	ChangeId = FGuid::NewGuid();
-	UE_LOG(LogNiagaraEditor, Log, TEXT("MarkGraphRequiresSynchronization %s"), *ChangeId.ToString());
+	//UE_LOG(LogNiagaraEditor, Verbose, TEXT("Graph %s was marked requires synchronization.  Reason: %s"), *GetPathName(), *Reason);
 }
 
 /** Get the meta-data associated with this variable, if it exists.*/
@@ -607,7 +1066,146 @@ void UNiagaraGraph::PurgeUnreferencedMetaData()
 	}
 }
 
+void UNiagaraGraph::PurgeUnreferencedParameters()
+{
+	TArray<FNiagaraVariable> VarsToRemove;
+	for (auto& ParameterEntry : Parameters)
+	{
+		if (!ParameterEntry.Value.WasCreated() && ParameterEntry.Value.ParameterReferences.Num() == 0)
+		{
+			VarsToRemove.Add(ParameterEntry.Key);
+		}
+	}
+
+	for (FNiagaraVariable& Var : VarsToRemove)
+	{
+		Parameters.Remove(Var);
+	}
+}
+
 UNiagaraGraph::FOnDataInterfaceChanged& UNiagaraGraph::OnDataInterfaceChanged()
 {
 	return OnDataInterfaceChangedDelegate;
 }
+
+void UNiagaraGraph::FindParameters()
+{
+	if (!bFindParametersAllowed)
+	{
+		return;
+	}
+
+	for (auto& ParameterEntry : Parameters)
+	{
+		ParameterEntry.Value.ParameterReferences.Empty();
+	}
+
+	for (auto& MetadataEntry : VariableToMetaData)
+	{
+		MetadataEntry.Value.ReferencerNodes.Empty();
+	}
+
+	auto AddParameterReference = [&](const FNiagaraVariable& Parameter, const UEdGraphPin* Pin, FNiagaraGraphParameterReferenceCollection*& ReferenceCollection)
+	{
+		const FNiagaraGraphParameterReference Reference(Pin->PersistentGuid, Cast<UNiagaraNode>(Pin->GetOwningNode()));
+		bool bNewReference = true;
+		if (ReferenceCollection)
+		{
+			ReferenceCollection->ParameterReferences.AddUnique(Reference);
+			bNewReference = false;
+		}
+		else
+		{
+			FNiagaraGraphParameterReferenceCollection* FoundReferenceCollection = Parameters.Find(Parameter);
+			if (FoundReferenceCollection)
+			{
+				ReferenceCollection = FoundReferenceCollection;
+				FoundReferenceCollection->ParameterReferences.AddUnique(Reference);
+				bNewReference = false;
+			}
+		}
+
+		if (bNewReference)
+		{
+			FNiagaraGraphParameterReferenceCollection NewReferenceCollection;
+			NewReferenceCollection.ParameterReferences.Add(Reference);
+			NewReferenceCollection.Graph = this;
+			Parameters.Add(Parameter, NewReferenceCollection);
+		}
+	};
+
+	const TArray<FNiagaraParameterMapHistory> Histories = UNiagaraNodeParameterMapBase::GetParameterMaps(this);
+	for (const FNiagaraParameterMapHistory& History : Histories)
+	{
+		for (int32 Index = 0; Index < History.VariablesWithOriginalAliasesIntact.Num(); Index++)
+		{
+			const FNiagaraVariable& Parameter = History.VariablesWithOriginalAliasesIntact[Index];
+
+			FNiagaraGraphParameterReferenceCollection* FoundReferences = nullptr;
+			for (const UEdGraphPin* WritePin : History.PerVariableWriteHistory[Index])
+			{
+				AddParameterReference(Parameter, WritePin, FoundReferences);
+			}
+
+			for (const TTuple<const UEdGraphPin*, const UEdGraphPin*>& ReadPinTuple : History.PerVariableReadHistory[Index])
+			{
+				AddParameterReference(Parameter, ReadPinTuple.Key, FoundReferences);
+			}
+		}
+	}
+
+	// Find all the parameters in the graph that have no connection and won't be picked up by the parameter map history.
+	const UEdGraphSchema_Niagara* NiagaraSchema = Cast<UEdGraphSchema_Niagara>(Schema);
+	for (UEdGraphNode* Node : Nodes)
+	{
+		for (UEdGraphPin* Pin : Node->Pins)
+		{
+			const FNiagaraVariable Parameter = NiagaraSchema->PinToNiagaraVariable(Pin, false);
+			const FNiagaraParameterHandle Handle = FNiagaraParameterHandle(Parameter.GetName());
+
+			if (Handle.IsModuleHandle() && !FNiagaraConstants::IsNiagaraConstant(Parameter))
+			{
+				FNiagaraVariableMetaData* MetaData = VariableToMetaData.Find(Parameter);
+				if (MetaData)
+				{
+					MetaData->ReferencerNodes.AddUnique(Node);
+				}
+				else
+				{
+					FNiagaraVariableMetaData NewVariableMetadata;
+					NewVariableMetadata.ReferencerNodes.Add(Node);
+					VariableToMetaData.Add(Parameter, NewVariableMetadata);
+				}
+			}
+			
+			if (!Handle.GetNamespace().IsNone())
+			{
+				const FNiagaraGraphParameterReference Reference(Pin->PersistentGuid, Cast<UNiagaraNode>(Pin->GetOwningNode()));
+				FNiagaraGraphParameterReferenceCollection* FoundParameterReferenceCollection = Parameters.Find(Parameter);
+				if (FoundParameterReferenceCollection)
+				{
+					FoundParameterReferenceCollection->ParameterReferences.AddUnique(Reference);
+				}
+				else
+				{
+					FNiagaraGraphParameterReferenceCollection NewReferenceCollection;
+					NewReferenceCollection.ParameterReferences.Add(Reference);
+					NewReferenceCollection.Graph = this;
+					Parameters.Add(Parameter, NewReferenceCollection);
+				}
+			}
+		}
+	}
+
+	// Clean up all parameters and metadata that do not have a reference
+	PurgeUnreferencedParameters();
+	PurgeUnreferencedMetaData();
+}
+
+
+void UNiagaraGraph::SetFindParametersAllowed(const bool bAllowed)
+{
+	bFindParametersAllowed = bAllowed;
+}
+
+#undef LOCTEXT_NAMESPACE
\ No newline at end of file
