Copying //UE4/Dev-Niagara to //UE4/Dev-Main (Source: //UE4/Dev-Niagara @ 4074996)

diff --git a/Engine/Plugins/FX/Niagara/Source/NiagaraEditor/Private/NiagaraHlslTranslator.cpp b/Engine/Plugins/FX/Niagara/Source/NiagaraEditor/Private/NiagaraHlslTranslator.cpp
--- a/Engine/Plugins/FX/Niagara/Source/NiagaraEditor/Private/NiagaraHlslTranslator.cpp
+++ b/Engine/Plugins/FX/Niagara/Source/NiagaraEditor/Private/NiagaraHlslTranslator.cpp
@@ -27,6 +27,7 @@
 #include "INiagaraEditorTypeUtilities.h"
 #include "NiagaraEditorUtilities.h"
 #include "NiagaraEditorModule.h"
+#include "NiagaraNodeReroute.h"
 
 #include "NiagaraDataInterface.h"
 #include "NiagaraDataInterfaceCurve.h"
@@ -74,9 +75,6 @@ DECLARE_CYCLE_STAT(TEXT("Niagara - HlslTranslator - GenerateFunctionSignature_In
 DECLARE_CYCLE_STAT(TEXT("Niagara - HlslTranslator - GenerateFunctionSignature_FindInputNodes"), STAT_NiagaraEditor_Module_NiagaraHLSLTranslator_GenerateFunctionSignature_FindInputNodes, STATGROUP_NiagaraEditor);
 
 
-
-
-
 static FNiagaraShaderQueueTickable NiagaraShaderQueueProcessor;
 FNiagaraShaderProcessorTickable NiagaraShaderProcessor;
 
@@ -88,103 +86,51 @@ void FNiagaraShaderQueueTickable::ProcessQueue()
 {
 	for (FNiagaraCompilationQueue::NiagaraCompilationQueueItem &Item : FNiagaraCompilationQueue::Get()->GetQueue())
 	{
-		FNiagaraScript *Script = Item.Script;
+		FNiagaraShaderScript* ShaderScript = Item.Script;
 		TRefCountPtr<FNiagaraShaderMap>NewShaderMap = Item.ShaderMap;
 
-		if (Script == nullptr) // This script has been removed from the pending queue post submission... just skip it.
+		if (ShaderScript == nullptr) // This script has been removed from the pending queue post submission... just skip it.
 		{
 			FNiagaraShaderMap::RemovePendingMap(NewShaderMap);
 			NewShaderMap->SetCompiledSuccessfully(false);
 			UE_LOG(LogNiagaraEditor, Log, TEXT("GPU shader compile skipped. Id %d"), NewShaderMap->GetCompilingId());
 			continue;
 		}
-		UNiagaraScript* CompilableScript = Script->GetGTScript();
+		UNiagaraScript* CompilableScript = ShaderScript->GetBaseVMScript();
 
 		// For now System scripts don't generate HLSL and go through a special pass...
 		// [OP] thinking they'll likely never run on GPU anyways
-		if (CompilableScript->IsValidLowLevel() == false || CompilableScript->CanBeRunOnGpu() == false)
+		if (CompilableScript->IsValidLowLevel() == false || CompilableScript->CanBeRunOnGpu() == false || !CompilableScript->GetVMExecutableData().IsValid() ||
+			CompilableScript->GetVMExecutableData().LastHlslTranslationGPU.Len() == 0)
 		{
 			NewShaderMap->SetCompiledSuccessfully(false);
 			FNiagaraShaderMap::RemovePendingMap(NewShaderMap);
-			Script->RemoveOutstandingCompileId(NewShaderMap->GetCompilingId());
+			ShaderScript->RemoveOutstandingCompileId(NewShaderMap->GetCompilingId());
 			UE_LOG(LogNiagaraEditor, Log, TEXT("GPU shader compile skipped. Id %d"), NewShaderMap->GetCompilingId());
 			continue;
 		}
 
 		FNiagaraComputeShaderCompilationOutput NewCompilationOutput;
-		FHlslNiagaraTranslator NiagaraTranslator;
-		FHlslNiagaraTranslatorOptions Options;
-		Options.SimTarget = ENiagaraSimTarget::GPUComputeSim;
 
-		bool bNeedsPostCompile = false;
-		if (CompilableScript->GetSource()->IsPreCompiled() == false)
+		ShaderScript->SetDataInterfaceParamInfo(CompilableScript->GetVMExecutableData().DIParamInfo);
+		ShaderScript->SourceName = TEXT("NiagaraComputeShader");
+		UNiagaraEmitter* Emitter = Cast<UNiagaraEmitter>(CompilableScript->GetOuter());
+		if (Emitter && Emitter->GetUniqueEmitterName().Len() > 0)
 		{
-			CompilableScript->GetSource()->PreCompile(Cast<UNiagaraEmitter>(CompilableScript->GetOuter()));
-			bNeedsPostCompile = true;
+			ShaderScript->SourceName = Emitter->GetUniqueEmitterName();
 		}
-		FNiagaraTranslateResults Results = NiagaraTranslator.Translate(CompilableScript, Options);
-		CompilableScript->LastHlslTranslationGPU = Results.OutputHLSL;
-
-		if (bNeedsPostCompile)
-		{
-			CompilableScript->GetSource()->PostCompile();
-		}
-		Script->SetDatainterfaceBufferDescriptors(NiagaraTranslator.GetDatainterfaceBufferDescriptors());
-		CompilableScript->SetDatainterfaceBufferDescriptors(NiagaraTranslator.GetDatainterfaceBufferDescriptors());
-		// make sure buffers are set up on the original script
-		for (int32 i = 0; i < NiagaraTranslator.GetDatainterfaceBufferDescriptors().Num(); i++)
-		{
-			if (CompilableScript->DataInterfaceInfo.Num() > i)
-			{
-				if (CompilableScript->DataInterfaceInfo[i].DataInterface)
-				{
-					CompilableScript->DataInterfaceInfo[i].DataInterface->SetupBuffers(NiagaraTranslator.GetDatainterfaceBufferDescriptors()[i]);
-				}
-			}
-		}
-
-		Script->HlslOutput = Results.OutputHLSL;
-
-		if (!Results.bHLSLGenSucceeded || Results.MessageLog->NumErrors > 0)
-		{
-			UE_LOG(LogNiagaraEditor, Warning, TEXT("%s"), TEXT("Niagara GPU Translation Failed!"));
-			TArray<TSharedRef<FTokenizedMessage>> Messages;
-			if (Results.MessageLog)
-			{
-				Messages.Append(Results.MessageLog->Messages);
-			}
-
-			for (TSharedRef<FTokenizedMessage> Message : Messages)
-			{
-				if (Message->GetSeverity() == EMessageSeverity::Info)
-				{
-					UE_LOG(LogNiagaraEditor, Log, TEXT("%s"), *Message->ToText().ToString());
-				}
-				else if (Message->GetSeverity() == EMessageSeverity::Warning || Message->GetSeverity() == EMessageSeverity::PerformanceWarning)
-				{
-					UE_LOG(LogNiagaraEditor, Warning, TEXT("%s"), *Message->ToText().ToString());
-				}
-				else if (Message->GetSeverity() == EMessageSeverity::Error || Message->GetSeverity() == EMessageSeverity::CriticalError)
-				{
-					UE_LOG(LogNiagaraEditor, Error, TEXT("%s"), *Message->ToText().ToString());
-				}
-			}
-			NewShaderMap->SetCompiledSuccessfully(false);
-			FNiagaraShaderMap::RemovePendingMap(NewShaderMap);
-
-		}
-		else
+		ShaderScript->HlslOutput = CompilableScript->GetVMExecutableData().LastHlslTranslationGPU;
+		
 		{
 			// Create a shader compiler environment for the script that will be shared by all jobs from this script
 			TRefCountPtr<FShaderCompilerEnvironment> CompilerEnvironment = new FShaderCompilerEnvironment();
 
-			FString ShaderCode = NiagaraTranslator.GetTranslatedHLSL();
+			FString ShaderCode = CompilableScript->GetVMExecutableData().LastHlslTranslationGPU;
 			const bool bSynchronousCompile = false;
 
 			// Compile the shaders for the script.
-			NewShaderMap->Compile(Script, Item.ShaderMapId, CompilerEnvironment, NewCompilationOutput, Item.Platform, bSynchronousCompile, Item.bApply);
+			NewShaderMap->Compile(ShaderScript, Item.ShaderMapId, CompilerEnvironment, NewCompilationOutput, Item.Platform, bSynchronousCompile, Item.bApply);
 		}
-
 	}
 
 	FNiagaraCompilationQueue::Get()->GetQueue().Empty();
@@ -195,258 +141,6 @@ void FNiagaraShaderQueueTickable::Tick(float DeltaSeconds)
 	ProcessQueue();
 }
 
-
-template<typename Action>
-void TraverseGraphFromOutputDepthFirst(FHlslNiagaraTranslator& Compiler, const UEdGraphSchema_Niagara* Schema, UNiagaraNode* Node, Action& VisitAction, int32 VisitID)
-{
-	if (Node->ShouldVisit(VisitID))
-	{
-		TArray<UEdGraphPin*> InputPins;
-		Node->GetInputPins(InputPins);
-		for (UEdGraphPin* InputPin : InputPins)
-		{
-			// TODO: Error here if there are multiple links or non-niagara nodes?
-			if (InputPin->LinkedTo.Num() == 1)
-			{
-				UNiagaraNode* LinkedNiagaraNode = Cast<UNiagaraNode>(InputPin->LinkedTo[0]->GetOwningNode());
-				if (LinkedNiagaraNode != nullptr)
-				{
-					TraverseGraphFromOutputDepthFirst(Compiler, Schema, LinkedNiagaraNode, VisitAction, VisitID);
-				}
-			}
-		}
-		VisitAction(Schema, Node);
-	}
-}
-
-void FixUpNumericPinsVisitor(FHlslNiagaraTranslator& Compiler, const UEdGraphSchema_Niagara* Schema, UNiagaraNode* Node)
-{
-	// Fix up numeric input pins and keep track of numeric types to decide the output type.
-	TArray<FNiagaraTypeDefinition> InputTypes;
-	TArray<UEdGraphPin*> InputPins;
-	Node->GetInputPins(InputPins);
-	for (UEdGraphPin* InputPin : InputPins)
-	{
-		if (InputPin->PinType.PinCategory == UEdGraphSchema_Niagara::PinCategoryType)
-		{
-			FNiagaraTypeDefinition InputPinType = Schema->PinToTypeDefinition(InputPin);
-
-			// If the input pin is the generic numeric type set it to the type of the linked output pin which should have been processed already.
-			if (InputPinType == FNiagaraTypeDefinition::GetGenericNumericDef() && InputPin->LinkedTo.Num() == 1)
-			{
-				UEdGraphPin* InputPinLinkedPin = InputPin->LinkedTo[0];
-				FNiagaraTypeDefinition InputPinLinkedPinType = Schema->PinToTypeDefinition(InputPinLinkedPin);
-				if (InputPinLinkedPinType.IsValid())
-				{
-					// Only update the input pin type if the linked pin type is valid.
-					InputPin->PinType = Schema->TypeDefinitionToPinType(InputPinLinkedPinType);
-					InputPinType = InputPinLinkedPinType;
-				}
-			}
-
-			if (InputPinType == FNiagaraTypeDefinition::GetGenericNumericDef())
-			{
-				Compiler.Error(LOCTEXT("NumericPinError", "Unable to deduce type for numeric input pin."), Node, InputPin);
-			}
-			InputTypes.Add(InputPinType);
-		}
-	}
-
-	// Fix up numeric outputs based on the inputs.
-	if (InputTypes.Num() > 0 && Node->GetNumericOutputTypeSelectionMode() != ENiagaraNumericOutputTypeSelectionMode::None)
-	{
-		FNiagaraTypeDefinition OutputNumericType = FNiagaraTypeDefinition::GetNumericOutputType(InputTypes, Node->GetNumericOutputTypeSelectionMode());
-		if (OutputNumericType != FNiagaraTypeDefinition::GetGenericNumericDef())
-		{
-			TArray<UEdGraphPin*> OutputPins;
-			Node->GetOutputPins(OutputPins);
-			for (UEdGraphPin* OutputPin : OutputPins)
-			{
-				FNiagaraTypeDefinition OutputPinType = Schema->PinToTypeDefinition(OutputPin);
-				if (OutputPinType == FNiagaraTypeDefinition::GetGenericNumericDef())
-				{
-					OutputPin->PinType = Schema->TypeDefinitionToPinType(OutputNumericType);
-				}
-			}
-		}
-	}
-}
-
-void FixUpNumericPins(FHlslNiagaraTranslator& Compiler, const UEdGraphSchema_Niagara* Schema, UNiagaraNode* Node)
-{
-	auto FixUpVisitor = [&](const UEdGraphSchema_Niagara* LSchema, UNiagaraNode* LNode) { FixUpNumericPinsVisitor(Compiler, LSchema, LNode); };
-	TraverseGraphFromOutputDepthFirst(Compiler, Schema, Node, FixUpVisitor, ++Compiler.VisitID);
-}
-
-/* Go through the graph and attempt to auto-detect the type of any numeric pins by working back from the leaves of the graph. Only change the types of pins, not FNiagaraVariables.*/
-void PreprocessGraph(FHlslNiagaraTranslator& Compiler, const UEdGraphSchema_Niagara* Schema, UNiagaraGraph* Graph, UNiagaraNodeOutput* OutputNode)
-{
-	if (OutputNode)
-	{
-		FixUpNumericPins(Compiler, Schema, OutputNode);
-	}
-	else
-	{
-		// This should never happen!
-		TArray<FStringFormatArg> Args;
-		Args.Add(FString::FromInt((int32)Compiler.GetTargetUsage()));
-		FString ErrorText = FString::Format(TEXT("Unable to preprocess graph due to missing output node of type '{0}'!"), Args);
-		Compiler.Error(FText::FromString(ErrorText), nullptr, nullptr);
-	}
-}
-
-/* Go through the graph and force any input nodes with Numeric types to a hard-coded type of float. This will allow modules and functions to compile properly.*/
-void PreProcessGraphForInputNumerics(FHlslNiagaraTranslator& Compiler, const UEdGraphSchema_Niagara* Schema, UNiagaraGraph* Graph, TArray<FNiagaraVariable>& OutChangedNumericParams)
-{
-	// Visit all input nodes
-	TArray<UNiagaraNodeInput*> InputNodes;
-	Graph->FindInputNodes(InputNodes);
-	for (UNiagaraNodeInput* InputNode : InputNodes)
-	{
-		// See if any of the output pins are of Numeric type. If so, force to floats.
-		TArray<UEdGraphPin*> OutputPins;
-		InputNode->GetOutputPins(OutputPins);
-		for (UEdGraphPin* OutputPin : OutputPins)
-		{
-			FNiagaraTypeDefinition OutputPinType = Schema->PinToTypeDefinition(OutputPin);
-			if (OutputPinType == FNiagaraTypeDefinition::GetGenericNumericDef())
-			{
-				OutputPin->PinType = Schema->TypeDefinitionToPinType(FNiagaraTypeDefinition::GetFloatDef());
-			}
-		}
-
-		// Record that we touched this variable for later cleanup and make sure that the 
-		// variable's type now matches the pin.
-		if (InputNode->Input.GetType() == FNiagaraTypeDefinition::GetGenericNumericDef())
-		{
-			OutChangedNumericParams.Add(InputNode->Input);
-			InputNode->Input.SetType(FNiagaraTypeDefinition::GetFloatDef());
-		}
-	}
-}
-
-/* Should be called after all pins have been successfully auto-detected for type. This goes through and synchronizes any Numeric FNiagaraVarible outputs with the deduced pin type. This will allow modules and functions to compile properly.*/
-void PreProcessGraphForAttributeNumerics(FHlslNiagaraTranslator& Compiler, const UEdGraphSchema_Niagara* Schema, UNiagaraGraph* Graph, UNiagaraNodeOutput* OutputNode, TArray<FNiagaraVariable>& OutChangedNumericParams)
-{
-	// Visit the output node
-	if (OutputNode != nullptr)
-	{
-		// For each pin, make sure that if it has a valid type, but the associated variable is still Numeric,
-		// force the variable to match the pin's new type. Record that we touched this variable for later cleanup.
-		TArray<UEdGraphPin*> InputPins;
-		OutputNode->GetInputPins(InputPins);
-		check(OutputNode->Outputs.Num() == InputPins.Num());
-		for (int32 i = 0; i < InputPins.Num(); i++)
-		{
-			FNiagaraVariable& Param = OutputNode->Outputs[i];
-			UEdGraphPin* InputPin = InputPins[i];
-
-			FNiagaraTypeDefinition InputPinType = Schema->PinToTypeDefinition(InputPin);
-			if (Param.GetType() == FNiagaraTypeDefinition::GetGenericNumericDef() && 
-				InputPinType != FNiagaraTypeDefinition::GetGenericNumericDef())
-			{
-				OutChangedNumericParams.Add(Param);
-				Param.SetType(InputPinType);
-			}
-		}
-	}
-}
-
-/* Clean up the lingering Systems of PreProcessGraphForInputNumerics and PreProcessGraphForAttributeNumerics by resetting the FNiagaraVariables back to their original types.*/
-void RevertParametersToNumerics(FHlslNiagaraTranslator& Compiler, UNiagaraScript* Script, const TArray<FNiagaraVariable>& ChangedNumericParams)
-{
-	// We either changed an input node variable or an output node variable in the prior functions, let's
-	// check where those ended up and fixup any discrepencies.
-	for (const FNiagaraVariable& ChangedVariable : ChangedNumericParams)
-	{
-		// Check input variables... 
-		FNiagaraVariable* CorrespondingVariable = Script->Parameters.FindParameter(ChangedVariable);
-		if (CorrespondingVariable == nullptr)
-		{
-			// Check output variables... 
-			CorrespondingVariable = Script->Attributes.FindByPredicate([&](FNiagaraVariable& Attribute) { return Attribute.GetName() == ChangedVariable.GetName(); });
-		}
-
-		// Convert back to Numeric so that we maintain consistency externally.
-		if (CorrespondingVariable != nullptr)
-		{
-			check(ChangedVariable.GetType() == FNiagaraTypeDefinition::GetGenericNumericDef());
-			CorrespondingVariable->SetType(ChangedVariable.GetType());
-		}
-		else
-		{
-			// This should never happen!
-			TArray<FStringFormatArg> Args;
-			Args.Add(ChangedVariable.ToString());
-			FString ErrorText = FString::Format(TEXT("Unable to find parameter '{0}' in outputs!"), Args);
-			Compiler.Error(FText::FromString(ErrorText), nullptr, nullptr);
-		}
-	}
-}
-
-
-void PreprocessFunctionGraph(FHlslNiagaraTranslator& Compiler, const UEdGraphSchema_Niagara* Schema, UNiagaraGraph* Graph, const TArray<UEdGraphPin*>& CallInputs, const TArray<UEdGraphPin*>& CallOutputs, ENiagaraScriptUsage ScriptUsage)
-{
-	// Change any numeric inputs or outputs to match the types from the call node.
-	TArray<UNiagaraNodeInput*> InputNodes;
-	
-	// Only handle nodes connected to the correct output node in the event of multiple output nodes in the graph.
-	UNiagaraGraph::FFindInputNodeOptions Options;
-	Options.bFilterByScriptUsage = true;
-	Options.TargetScriptUsage = ScriptUsage;
-
-	Graph->FindInputNodes(InputNodes, Options);
-
-	for (UNiagaraNodeInput* InputNode : InputNodes)
-	{
-		FNiagaraVariable& Input = InputNode->Input;
-		if (Input.GetType() == FNiagaraTypeDefinition::GetGenericNumericDef())
-		{
-			UEdGraphPin* const* MatchingPin = CallInputs.FindByPredicate([&](const UEdGraphPin* Pin) { return (Pin->PinName == Input.GetName()); });
-
-			if (MatchingPin != nullptr)
-			{
-				FNiagaraTypeDefinition PinType = Schema->PinToTypeDefinition(*MatchingPin);
-				Input.SetType(PinType);
-				TArray<UEdGraphPin*> OutputPins;
-				InputNode->GetOutputPins(OutputPins);
-				check(OutputPins.Num() == 1);
-				OutputPins[0]->PinType = (*MatchingPin)->PinType;
-			}
-		}
-	}
-
-	UNiagaraNodeOutput* OutputNode = Graph->FindOutputNode(ScriptUsage);
-	check(OutputNode);
-
-	TArray<UEdGraphPin*> InputPins;
-	OutputNode->GetInputPins(InputPins);
-
-	for (FNiagaraVariable& Output : OutputNode->Outputs)
-	{
-		if (Output.GetType() == FNiagaraTypeDefinition::GetGenericNumericDef())
-		{
-			UEdGraphPin* const* MatchingPin = CallOutputs.FindByPredicate([&](const UEdGraphPin* Pin) { return (Pin->PinName == Output.GetName()); });
-
-			if (MatchingPin != nullptr)
-			{
-				FNiagaraTypeDefinition PinType = Schema->PinToTypeDefinition(*MatchingPin);
-				Output.SetType(PinType);
-			}
-		}
-	}
-
-	FixUpNumericPins(Compiler, Schema, OutputNode);
-
-	// Fix up numeric pins for data set writes since they may have numerics which were missed traversing from the output node.
-	TArray<UNiagaraNodeWriteDataSet*> DataSetWriteNodes;
-	Graph->GetNodesOfClass<UNiagaraNodeWriteDataSet>(DataSetWriteNodes);
-	for (UNiagaraNodeWriteDataSet* DataSetWriteNode : DataSetWriteNodes)
-	{
-		FixUpNumericPins(Compiler, Schema, DataSetWriteNode);
-	}
-}
-
 ENiagaraScriptCompileStatus FNiagaraTranslateResults::TranslateResultsToSummary(const FNiagaraTranslateResults* TranslateResults)
 {
 	ENiagaraScriptCompileStatus SummaryStatus = ENiagaraScriptCompileStatus::NCS_Unknown;
@@ -573,6 +267,7 @@ void FHlslNiagaraTranslator::GenerateFunctionSignature(ENiagaraScriptUsage Scrip
 		UNiagaraGraph::FFindInputNodeOptions Options;
 		Options.bSort = true;
 		Options.bFilterDuplicates = true;
+		Options.bIncludeTranslatorConstants = false;
 		Options.bFilterByScriptUsage = true;
 		Options.TargetScriptUsage = ScriptUsage;
 		FuncGraph->FindInputNodes(InputsNodes, Options);
@@ -659,12 +354,11 @@ void FHlslNiagaraTranslator::GenerateFunctionSignature(ENiagaraScriptUsage Scrip
 //////////////////////////////////////////////////////////////////////////
 
 FHlslNiagaraTranslator::FHlslNiagaraTranslator()
-	: Script(nullptr)
-	, Schema(nullptr)
+	: Schema(nullptr)
 	, TranslateResults(&MessageLog)
-	, VisitID(INDEX_NONE)
-	, bInsideInterpolatedSpawnScript(false)
 	, CurrentBodyChunkMode(ENiagaraCodeChunkMode::Body)
+	, ActiveStageIdx(-1)
+	, bInitializedDefaults(false)
 {
 	// Make the message log silent so we're not spamming the blueprint log.
 	MessageLog.bSilentMode = true;
@@ -691,57 +385,14 @@ FString FHlslNiagaraTranslator::GetFunctionDefinitions()
 	return FwdDeclString + TEXT("\n") + DefinitionsString;
 }
 
-UNiagaraGraph* FHlslNiagaraTranslator::CloneGraphAndPrepareForCompilation(UNiagaraScript* InScript, UNiagaraScriptSource* InSource, bool bClearErrors)
-{
-	SCOPE_CYCLE_COUNTER(STAT_NiagaraEditor_HlslTranslator_CloneGraphAndPrepareForCompilation);
-	if (InSource == nullptr)
-	{
-		Error(LOCTEXT("NoSourceErrorMessage", "Script has no source."), nullptr, nullptr);
-		return nullptr;
-	}
-
-	UNiagaraGraph* SourceGraph = InSource->GetPrecomputedNodeGraph();
-	check(SourceGraph != nullptr);
 	
-	TArray<FNiagaraVariable> ChangedNumericParams;
-	// In the case of functions or modules, we may not have enough information at this time to fully resolve the type. In that case,
-	// we circumvent the resulting errors by forcing a type. This gives the user an appropriate level of type checking. We will, however need to clean this up in
-	// the parameters that we output.
-	bool bForceParametersToResolveNumerics = InScript->IsStandaloneScript();
-	if (bForceParametersToResolveNumerics)
-	{
-		PreProcessGraphForInputNumerics(*this, Schema, SourceGraph, ChangedNumericParams);
-	}
-
-	// Auto-deduce the input types for numerics in the graph and overwrite the types on the pins. If PreProcessGraphForInputNumerics occurred, then
-	// we will have pre-populated the inputs with valid types.
-	TArray<UNiagaraNodeOutput*> OutputNodes;
-	SourceGraph->FindOutputNodes(OutputNodes);
-
-	for (UNiagaraNodeOutput* OutputNode : OutputNodes)
-	{
-		PreprocessGraph(*this, Schema, SourceGraph, OutputNode);
 
-		// Now that we've auto-deduced the types, we need to handle any lingering Numerics in the Output's FNiagaraVariable outputs. 
-		// We use the pin's deduced type to temporarily overwrite the variable's type.
-		if (bForceParametersToResolveNumerics)
-		{
-			PreProcessGraphForAttributeNumerics(*this, Schema, SourceGraph, OutputNode, ChangedNumericParams);
-		}
-	}
-
-	return SourceGraph;
-}
-
-FString FHlslNiagaraTranslator::BuildMissingDefaults(TArray<TPair<FNiagaraVariable, const UEdGraphPin*> >& MissingDefaults)
+FString FHlslNiagaraTranslator::BuildMissingDefaults()
 {
 	FString HlslOutputString = TEXT("void HandleMissingDefaultValues(inout FSimulationContext Context)\n{\n");
-	for (TPair<FNiagaraVariable, const UEdGraphPin*>& Item : MissingDefaults)
+	if (UNiagaraScript::IsSpawnScript(TranslationStages[ActiveStageIdx].ScriptUsage))
 	{
-		FNiagaraVariable Var = Item.Key;
-		const UEdGraphPin* DefaultPin = Item.Value;
-
-		if (Script->IsParticleSpawnScript() || Script->IsEmitterSpawnScript() || Script->IsSystemSpawnScript())
+		for (FNiagaraVariable& Var : InitialNamespaceVariablesMissingDefault)
 		{
 			if (FNiagaraParameterMapHistory::IsInitialValue(Var))
 			{
@@ -751,29 +402,15 @@ FString FHlslNiagaraTranslator::BuildMissingDefaults(TArray<TPair<FNiagaraVariab
 					ParameterMapInstanceName + TEXT(".") + GetSanitizedSymbolName(SourceForInitialValue.GetName().ToString()) + TEXT(";"));
 				continue;
 			}
-
-			if (DefaultPin && DefaultPin->LinkedTo.Num() == 0)
-			{
-				FNiagaraVariable PinVar = Schema->PinToNiagaraVariable(DefaultPin, true);
-				FString ConstantValue = GenerateConstantString(PinVar);
-				if (Var.GetType().GetClass() == nullptr) // Only need to do this wiring for things that aren't data interfaces.
-				{
-					FString ParameterMapInstanceName = GetParameterMapInstanceName(0);
-					HlslOutputString += TEXT("\t") + ParameterMapInstanceName + TEXT(".") + GetSanitizedSymbolName(Var.GetName().ToString()) + TEXT(" = ") + ConstantValue + TEXT(";");
-				}
-			}
-			else if (DefaultPin)
-			{
-				Error(FText::Format(LOCTEXT("OnlySimpleDefaults", "Only simple constants are supported for defaults of primary values: {0}"), FText::FromName(Var.GetName())), nullptr, nullptr);
-			}
 		}
 		HlslOutputString += TEXT("\n");
 	}
+	
 	HlslOutputString += TEXT("\n}\n");
 	return HlslOutputString;
 }
 
-FString FHlslNiagaraTranslator::BuildParameterMapHlslDefinitions(TArray<FNiagaraVariable>& PrimaryDataSetOutputEntries, TArray<TPair<FNiagaraVariable, const UEdGraphPin*> >& MissingDefaults)
+FString FHlslNiagaraTranslator::BuildParameterMapHlslDefinitions(TArray<FNiagaraVariable>& PrimaryDataSetOutputEntries)
 {
 	SCOPE_CYCLE_COUNTER(STAT_NiagaraEditor_HlslTranslator_BuildParameterMapHlslDefinitions);
 	FString HlslOutputString;
@@ -836,42 +473,25 @@ FString FHlslNiagaraTranslator::BuildParameterMapHlslDefinitions(TArray<FNiagara
 		UniqueVariables.AddUnique(Var);
 	}
 
+	bool bIsSpawnScript = IsSpawnScript();
+
 	// For now we only care about attributes from the other output parameter map histories.
 	for (int32 ParamMapIdx = 0; ParamMapIdx < OtherOutputParamMapHistories.Num(); ParamMapIdx++)
 	{
 		for (int32 VarIdx = 0; VarIdx < OtherOutputParamMapHistories[ParamMapIdx].Variables.Num(); VarIdx++)
 		{
 			FNiagaraVariable& Var = OtherOutputParamMapHistories[ParamMapIdx].Variables[VarIdx];
-			if (OtherOutputParamMapHistories[ParamMapIdx].IsPrimaryDataSetOutput(Var, Script))
+			if (OtherOutputParamMapHistories[ParamMapIdx].IsPrimaryDataSetOutput(Var, CompileOptions.TargetUsage))
 			{
 				int32 PreviousMax = UniqueVariables.Num();
 				if (UniqueVariables.AddUnique(Var) == PreviousMax) // i.e. we didn't find it previously, so we added to the end.
 				{
-					if (Script->IsParticleSpawnScript() || Script->IsEmitterSpawnScript() || Script->IsSystemSpawnScript())
+					if (bIsSpawnScript)
 					{
 						if (!AddStructToDefinitionSet(Var.GetType()))
 						{
 							Error(FText::Format(LOCTEXT("ParameterMapTypeError", "Cannot handle type {0}! Variable: {1}"), Var.GetType().GetNameText(), FText::FromName(Var.GetName())), nullptr, nullptr);
 						}
-						if (FNiagaraParameterMapHistory::IsInitialValue(Var))
-						{
-							FNiagaraVariable SourceForInitialValue = FNiagaraParameterMapHistory::GetSourceForInitialValue(Var);
-							if (!UniqueVariables.Contains(SourceForInitialValue))
-							{
-								Error(FText::Format(LOCTEXT("MissingInitialValueSource", "Variable {0} is used, but its source variable {1} is not set!"), FText::FromName(Var.GetName()), FText::FromName(SourceForInitialValue.GetName())), nullptr, nullptr);
-							}
-							TPair<FNiagaraVariable, const UEdGraphPin*> Value;
-							Value.Key = Var;
-							Value.Value = nullptr;
-							MissingDefaults.Add(Value);
-							continue;
-						}
-
-						const UEdGraphPin* DefaultPin = OtherOutputParamMapHistories[ParamMapIdx].GetDefaultValuePin(VarIdx);
-						TPair<FNiagaraVariable, const UEdGraphPin*> Value;
-						Value.Key = Var;
-						Value.Value = DefaultPin;
-						MissingDefaults.Add(Value);
 					}
 				}
 			}
@@ -923,7 +543,7 @@ FString FHlslNiagaraTranslator::BuildParameterMapHlslDefinitions(TArray<FNiagara
 		// In order 
 		for (int32 ParamMapIdx = 0; ParamMapIdx < OtherOutputParamMapHistories.Num(); ParamMapIdx++)
 		{
-			if (OtherOutputParamMapHistories[ParamMapIdx].IsPrimaryDataSetOutput(Variable, Script))
+			if (OtherOutputParamMapHistories[ParamMapIdx].IsPrimaryDataSetOutput(Variable, CompileOptions.TargetUsage))
 			{
 				PrimaryDataSetOutputEntries.AddUnique(Variable);
 				break;
@@ -1027,21 +647,23 @@ void FHlslNiagaraTranslator::HandleNamespacedExternalVariablesToDataSetRead(TArr
 	}
 }
 
-const FNiagaraTranslateResults &FHlslNiagaraTranslator::Translate(UNiagaraScript* InScript, FHlslNiagaraTranslatorOptions Options)
+const FNiagaraTranslateResults &FHlslNiagaraTranslator::Translate(const FNiagaraCompileRequestData* InCompileData, const FNiagaraCompileOptions& InCompileOptions, FHlslNiagaraTranslatorOptions InTranslateOptions)
 {
 	SCOPE_CYCLE_COUNTER(STAT_NiagaraEditor_HlslTranslator_Translate);
-	check(InScript);
+	check(InCompileData);
 
-	TranslationOptions = Options;
+	CompileOptions = InCompileOptions;
+	CompileData = InCompileData;
+	TranslationOptions = InTranslateOptions;
 	CompilationTarget = TranslationOptions.SimTarget;
 	TranslateResults.bHLSLGenSucceeded = false;
 	TranslateResults.OutputHLSL = "";
 
-	Script = InScript;
 	//Should we roll our own message/error log and put it in a window somewhere?
-	MessageLog.SetSourcePath(InScript->GetPathName());
+	MessageLog.SetSourcePath(CompileOptions.GetPathName());
+
+	UNiagaraGraph* SourceGraph = CompileData->NodeGraphDeepCopy;
 
-	UNiagaraGraph* SourceGraph = CloneGraphAndPrepareForCompilation(InScript, Cast<UNiagaraScriptSource>(InScript->GetSource()), true);
 	if (!SourceGraph)
 	{
 		Error(LOCTEXT("GetGraphFail", "Cannot find graph node!"), nullptr, nullptr);
@@ -1050,7 +672,7 @@ const FNiagaraTranslateResults &FHlslNiagaraTranslator::Translate(UNiagaraScript
 
 	if (SourceGraph->IsEmpty())
 	{
-		if (Script->IsSystemSpawnScript() || Script->IsSystemUpdateScript())
+		if (UNiagaraScript::IsSystemScript(CompileOptions.TargetUsage))
 		{
 			Error(LOCTEXT("GetNoNodeSystemFail", "Graph contains no nodes! Please add an emitter."), nullptr, nullptr);
 		}
@@ -1061,52 +683,80 @@ const FNiagaraTranslateResults &FHlslNiagaraTranslator::Translate(UNiagaraScript
 		return TranslateResults;
 	}
 
-	// Find the output node and compile it.
-	UNiagaraNodeOutput* OutputNode = SourceGraph->FindOutputNode(GetTargetUsage(), GetTargetUsageId());
-	if (OutputNode == nullptr)
-	{
-		Error(FText::Format(LOCTEXT("GetOutputNodeFail", "Cannot find output node of type {0}!"), FText::AsNumber((int32)GetTargetUsage())), nullptr, nullptr);
-		return TranslateResults;
-	}
-	ValidateTypePins(OutputNode);
-	{
-		bool bHasAnyConnections = false;
-		for (int32 PinIdx = 0; PinIdx < OutputNode->Pins.Num(); PinIdx++)
-		{
-			if (OutputNode->Pins[PinIdx]->Direction == EEdGraphPinDirection::EGPD_Input && OutputNode->Pins[PinIdx]->LinkedTo.Num() != 0)
-			{
-				bHasAnyConnections = true;
-			}
-		}
-		if (!bHasAnyConnections)
-		{
-			Error(FText::Format(LOCTEXT("GetOutputNodeConnectivityFail", "Cannot find any connections to output node of type {0}!"), FText::AsNumber((int32)GetTargetUsage())), nullptr, nullptr);
-			return TranslateResults;
-		}
-	}
+	TranslationStages.Empty();
+	ActiveStageIdx = 0;
 
-	// Build up a traversal from the output pin that touches all the parameter maps we might encounter, using the histories we build below.
+	bool bHasInterpolatedSpawn = InCompileOptions.AdditionalDefines.Contains(TEXT("InterpolatedSpawn"));
 	ParamMapHistories.Empty();
 	ParamMapSetVariablesToChunks.Empty();
-	TArray<UNiagaraNodeOutput*> UsageOutputs;
-	if (Script->IsInterpolatedParticleSpawnScript() || Options.SimTarget == ENiagaraSimTarget::GPUComputeSim)
-	{
-		UsageOutputs.Add(OutputNode);
-		UsageOutputs.Add(SourceGraph->FindOutputNode(ENiagaraScriptUsage::ParticleUpdateScript));
+
+	switch (CompileOptions.TargetUsage)
+	{
+	case ENiagaraScriptUsage::ParticleSpawnScriptInterpolated:
+		TranslationStages.Add(FHlslNiagaraTranslationStage(CompileOptions.TargetUsage, CompileOptions.TargetUsageId));
+		TranslationStages.Add(FHlslNiagaraTranslationStage(ENiagaraScriptUsage::ParticleUpdateScript, FGuid()));
+		TranslationStages[0].PassNamespace = TEXT("MapSpawn");
+		TranslationStages[1].PassNamespace = TEXT("MapUpdate");
+		TranslationStages[0].ChunkModeIndex = ENiagaraCodeChunkMode::SpawnBody;
+		TranslationStages[1].ChunkModeIndex = ENiagaraCodeChunkMode::UpdateBody;
+		TranslationStages[0].OutputNode = SourceGraph->FindOutputNode(ENiagaraScriptUsage::ParticleSpawnScript, TranslationStages[0].UsageId);
+		TranslationStages[1].OutputNode = SourceGraph->FindOutputNode(TranslationStages[1].ScriptUsage, TranslationStages[1].UsageId);
+		TranslationStages[1].bInterpolatePreviousParams = true;
 		ParamMapHistories.AddDefaulted(2);
 		ParamMapSetVariablesToChunks.AddDefaulted(2);
-	}
-	else
-	{
-		UsageOutputs.Add(OutputNode);
+		break;
+	case ENiagaraScriptUsage::ParticleGPUComputeScript:
+		TranslationStages.Add(FHlslNiagaraTranslationStage((bHasInterpolatedSpawn ? ENiagaraScriptUsage::ParticleSpawnScriptInterpolated : ENiagaraScriptUsage::ParticleSpawnScript), FGuid()));
+		TranslationStages.Add(FHlslNiagaraTranslationStage(ENiagaraScriptUsage::ParticleUpdateScript, FGuid()));
+		TranslationStages[0].PassNamespace = TEXT("MapSpawn");
+		TranslationStages[1].PassNamespace = TEXT("MapUpdate");
+		TranslationStages[0].ChunkModeIndex = ENiagaraCodeChunkMode::SpawnBody;
+		TranslationStages[1].ChunkModeIndex = ENiagaraCodeChunkMode::UpdateBody;
+		TranslationStages[0].OutputNode = SourceGraph->FindOutputNode(ENiagaraScriptUsage::ParticleSpawnScript, TranslationStages[0].UsageId);
+		TranslationStages[1].OutputNode = SourceGraph->FindOutputNode(TranslationStages[1].ScriptUsage, TranslationStages[1].UsageId);
+		TranslationStages[1].bInterpolatePreviousParams = bHasInterpolatedSpawn;
+		ParamMapHistories.AddDefaulted(2);
+		ParamMapSetVariablesToChunks.AddDefaulted(2);
+		break;
+	default:
+		TranslationStages.Add(FHlslNiagaraTranslationStage(CompileOptions.TargetUsage, CompileOptions.TargetUsageId));
+		TranslationStages[0].PassNamespace = TEXT("Map");		
+		TranslationStages[0].OutputNode = SourceGraph->FindOutputNode(TranslationStages[0].ScriptUsage, TranslationStages[0].UsageId);
+		TranslationStages[0].ChunkModeIndex = ENiagaraCodeChunkMode::Body;
 		ParamMapHistories.AddDefaulted(1);
 		ParamMapSetVariablesToChunks.AddDefaulted(1);
+		break;
+	}
+
+	for (int32 i = 0; i < TranslationStages.Num(); i++)
+	{
+		if (TranslationStages[i].OutputNode == nullptr)
+		{
+			Error(FText::Format(LOCTEXT("GetOutputNodeFail", "Cannot find output node of type {0}!"), FText::AsNumber((int32)TranslationStages[i].ScriptUsage)), nullptr, nullptr);
+			return TranslateResults;
+		}
+
+		ValidateTypePins(TranslationStages[i].OutputNode);
+		{
+			bool bHasAnyConnections = false;
+			for (int32 PinIdx = 0; PinIdx < TranslationStages[i].OutputNode->Pins.Num(); PinIdx++)
+			{
+				if (TranslationStages[i].OutputNode->Pins[PinIdx]->Direction == EEdGraphPinDirection::EGPD_Input && TranslationStages[i].OutputNode->Pins[PinIdx]->LinkedTo.Num() != 0)
+				{
+					bHasAnyConnections = true;
+				}
+			}
+			if (!bHasAnyConnections)
+			{
+				Error(FText::Format(LOCTEXT("GetOutputNodeConnectivityFail", "Cannot find any connections to output node of type {0}!"), FText::AsNumber((int32)TranslationStages[i].ScriptUsage)), nullptr, nullptr);
+				return TranslateResults;
+			}
+		}
 	}
 
+
 	// Get all the parameter map histories traced to this graph from output nodes. We'll revisit this shortly in order to build out just the ones we care about for this translation.
-	check(Script->GetSource()->IsPreCompiled());
-	UNiagaraScriptSource* ScriptSource = Cast<UNiagaraScriptSource>(Script->GetSource());
-	OtherOutputParamMapHistories = ScriptSource->GetPrecomputedHistories();
+	OtherOutputParamMapHistories = CompileData->GetPrecomputedHistories();
 	for (FNiagaraParameterMapHistory& FoundHistory : OtherOutputParamMapHistories)
 	{
 		const UNiagaraNodeOutput* HistoryOutputNode = FoundHistory.GetFinalOutputNode();
@@ -1116,12 +766,48 @@ const FNiagaraTranslateResults &FHlslNiagaraTranslator::Translate(UNiagaraScript
 		}
 		
 		// Now see if we want to use any of these specifically..
-		for (int32 ParamMapIdx = 0; ParamMapIdx < UsageOutputs.Num(); ParamMapIdx++)
+		for (int32 ParamMapIdx = 0; ParamMapIdx < TranslationStages.Num(); ParamMapIdx++)
 		{
-			UNiagaraNodeOutput* TargetOutputNode = UsageOutputs[ParamMapIdx];
+			UNiagaraNodeOutput* TargetOutputNode = TranslationStages[ParamMapIdx].OutputNode;
 			if (FoundHistory.GetFinalOutputNode() == TargetOutputNode)
 			{
+				bool bNeedsPersistentIDs = InCompileOptions.AdditionalDefines.Contains(TEXT("RequiresPersistentIDs"));
+				if (bNeedsPersistentIDs)
+				{
+					//TODO: Setup alias for current level to decouple from "Particles". Would we ever want emitter or system persistent IDs?
+					FNiagaraVariable Var = FNiagaraVariable(FNiagaraTypeDefinition::GetIDDef(), TEXT("Particles.ID"));
+					FoundHistory.AddVariable(Var, Var, nullptr);
+				}
+
+				if (RequiresInterpolation())
+				{
+					FNiagaraVariable Var = FNiagaraVariable(FNiagaraTypeDefinition::GetIntDef(), TEXT("Interpolation.InterpSpawn_Index"));
+					FoundHistory.AddVariable(Var, Var, nullptr);
+
+					Var = FNiagaraVariable(FNiagaraTypeDefinition::GetFloatDef(), TEXT("Interpolation.InterpSpawn_SpawnTime"));
+					FoundHistory.AddVariable(Var, Var, nullptr);
+
+					Var = FNiagaraVariable(FNiagaraTypeDefinition::GetFloatDef(), TEXT("Interpolation.InterpSpawn_UpdateTime"));
+					FoundHistory.AddVariable(Var, Var, nullptr);
+
+					Var = FNiagaraVariable(FNiagaraTypeDefinition::GetFloatDef(), TEXT("Interpolation.InterpSpawn_InvSpawnTime"));
+					FoundHistory.AddVariable(Var, Var, nullptr);
+
+					Var = FNiagaraVariable(FNiagaraTypeDefinition::GetFloatDef(), TEXT("Interpolation.InterpSpawn_InvUpdateTime"));
+					FoundHistory.AddVariable(Var, Var, nullptr);
+
+					Var = FNiagaraVariable(FNiagaraTypeDefinition::GetFloatDef(), TEXT("Interpolation.SpawnInterp"));
+					FoundHistory.AddVariable(Var, Var, nullptr);
+
+					Var = FNiagaraVariable(FNiagaraTypeDefinition::GetFloatDef(), TEXT("Interpolation.Emitter_SpawnInterval"));
+					FoundHistory.AddVariable(Var, Var, nullptr);
+
+					Var = FNiagaraVariable(FNiagaraTypeDefinition::GetFloatDef(), TEXT("Interpolation.Emitter_InterpSpawnStartDt"));
+					FoundHistory.AddVariable(Var, Var, nullptr);
+				}
+
 				ParamMapHistories[ParamMapIdx] = (FoundHistory);
+
 				TArray<int32> Entries;
 				Entries.AddZeroed(FoundHistory.Variables.Num());
 				for (int32 i = 0; i < Entries.Num(); i++)
@@ -1133,80 +819,87 @@ const FNiagaraTranslateResults &FHlslNiagaraTranslator::Translate(UNiagaraScript
 		}
 	}
 
-	Script->ParameterCollections.Empty();
+	CompilationOutput.ScriptData.ParameterCollectionPaths.Empty();
 	for (FNiagaraParameterMapHistory& History : ParamMapHistories)
 	{
 		for (UNiagaraParameterCollection* Collection : History.ParameterCollections)
 		{
-			Script->ParameterCollections.AddUnique(Collection);
+			CompilationOutput.ScriptData.ParameterCollectionPaths.AddUnique(FSoftObjectPath(Collection).ToString());
 		}
 	}
 	
 	//Create main scope pin cache.
 	PinToCodeChunks.AddDefaulted(1);
 
-	ActiveHistoryForFunctionCalls.BeginTranslation(InScript);
+	ActiveHistoryForFunctionCalls.BeginTranslation(GetUniqueEmitterName());
 
-	Script->StatScopes.Empty();
-	EnterStatsScope(FNiagaraStatScope(*Script->GetFullName(), *Script->GetFullName()));
+	CompilationOutput.ScriptData.StatScopes.Empty();
+	EnterStatsScope(FNiagaraStatScope(*CompileOptions.GetFullName(), *CompileOptions.GetFullName()));
 
 	FHlslNiagaraTranslator* ThisTranslator = this;
 	TArray<int32> OutputChunks;
-	if (Script->IsInterpolatedParticleSpawnScript() || Options.SimTarget == ENiagaraSimTarget::GPUComputeSim)
+
+	bool bInterpolateParams = false;
+
+	if (TranslationStages.Num() > 1)
 	{
-		CurrentBodyChunkMode = ENiagaraCodeChunkMode::SpawnBody;
-		//Here we compile the spawn script but write to temporary outputs in the context.
-		if (Script->IsInterpolatedParticleSpawnScript())
-		{
-			AddBodyChunk(TEXT(""), TEXT("//Begin Interpolated Spawn Script!"), FNiagaraTypeDefinition::GetIntDef(), false, false);
-		}
-		else
+		for (int32 i = 0; i < TranslationStages.Num(); i++)
 		{
-			AddBodyChunk(TEXT(""), TEXT("//Begin Spawn Script!"), FNiagaraTypeDefinition::GetIntDef(), false, false);
-		}
-		bInsideInterpolatedSpawnScript = true;
-		CurrentParamMapIndices.Empty();
-		CurrentParamMapIndices.Add(0);
-		OutputNode->Compile(ThisTranslator, OutputChunks);
-		bInsideInterpolatedSpawnScript = false;
-		InstanceWrite = FDataSetAccessInfo(); // Reset after building the output..
-		AddBodyChunk(TEXT(""), TEXT("//End Spawn Script!\n\n"), FNiagaraTypeDefinition::GetIntDef(), false, false);
+			ActiveStageIdx = i;
+			CurrentBodyChunkMode = TranslationStages[i].ChunkModeIndex;
+			if (UNiagaraScript::IsParticleSpawnScript(TranslationStages[i].ScriptUsage))
+			{
+				AddBodyChunk(TEXT(""), bHasInterpolatedSpawn ? TEXT("//Begin Interpolated Spawn Script!") : TEXT("//Begin Spawn Script!"), FNiagaraTypeDefinition::GetIntDef(), false, false);
+				CurrentParamMapIndices.Empty();
+				CurrentParamMapIndices.Add(0);
+				TranslationStages[i].OutputNode->Compile(ThisTranslator, OutputChunks);
+				InstanceWrite = FDataSetAccessInfo(); // Reset after building the output..
+				AddBodyChunk(TEXT(""), TEXT("//End Spawn Script!\n\n"), FNiagaraTypeDefinition::GetIntDef(), false, false);
+				AddBodyChunk(TEXT(""), TEXT("HandleMissingDefaultValues(Context);"), FNiagaraTypeDefinition::GetIntDef(), false, false);
+			}
+			
+			if (TranslationStages[i].bInterpolatePreviousParams)
+			{
+				bInterpolateParams = true;
+			}
 
-		if (Script->IsInterpolatedParticleSpawnScript())
-		{
-			AddBodyChunk(TEXT(""), TEXT("//Begin Transfer of Attributes!"), FNiagaraTypeDefinition::GetIntDef(), false, false);
-			AddBodyChunk(TEXT(""), TEXT("Context.MapUpdate.Particles = Context.MapSpawn.Particles"), FNiagaraTypeDefinition::GetIntDef(), false, true);
-			AddBodyChunk(TEXT(""), TEXT("//End Transfer of Attributes!\n\n"), FNiagaraTypeDefinition::GetIntDef(), false, false);
-			AddBodyChunk(TEXT(""), TEXT("HandleMissingDefaultValues(Context);"), FNiagaraTypeDefinition::GetIntDef(), false, false);
+			if (UNiagaraScript::IsParticleUpdateScript(TranslationStages[i].ScriptUsage))
+			{
+				AddBodyChunk(TEXT(""), TEXT("//Begin Update Script!"), FNiagaraTypeDefinition::GetIntDef(), false, false);
+				//Now we compile the update script (with partial dt) and read from the temp values written above.
+				CurrentParamMapIndices.Empty();
+				CurrentParamMapIndices.Add(1);
+				TranslationStages[i].OutputNode->Compile(ThisTranslator, OutputChunks);
+				AddBodyChunk(TEXT(""), TEXT("//End Update Script!\n\n"), FNiagaraTypeDefinition::GetIntDef(), false, false);
+			}
 		}
-
-		CurrentBodyChunkMode = ENiagaraCodeChunkMode::UpdateBody;
-		AddBodyChunk(TEXT(""), TEXT("//Begin Update Script!"), FNiagaraTypeDefinition::GetIntDef(), false, false);
-		//Now we compile the update script (with partial dt) and read from the temp values written above.
-		UNiagaraNodeOutput* UpdateOutputNode = SourceGraph->FindOutputNode(ENiagaraScriptUsage::ParticleUpdateScript);
-		check(UpdateOutputNode);
-		CurrentParamMapIndices.Empty();
-		CurrentParamMapIndices.Add(1);
-		UpdateOutputNode->Compile(ThisTranslator, OutputChunks);
-
-		AddBodyChunk(TEXT(""), TEXT("//End Update Script!\n\n"), FNiagaraTypeDefinition::GetIntDef(), false, false);
 		CurrentBodyChunkMode = ENiagaraCodeChunkMode::Body;
 	}
-	else
+	else if (TranslationStages.Num() == 1)
 	{
+		CurrentBodyChunkMode = TranslationStages[0].ChunkModeIndex;
+		ActiveStageIdx = 0;
+		check(CompileOptions.TargetUsage == TranslationStages[0].ScriptUsage);
 		CurrentParamMapIndices.Empty();
 		CurrentParamMapIndices.Add(0);
-		OutputNode->Compile(ThisTranslator, OutputChunks);
 
-		if (Script->IsParticleSpawnScript() || Script->IsSystemSpawnScript())
+		TranslationStages[0].OutputNode->Compile(ThisTranslator, OutputChunks);
+
+		if (IsSpawnScript())
 		{
 			AddBodyChunk(TEXT(""), TEXT("HandleMissingDefaultValues(Context);"), FNiagaraTypeDefinition::GetIntDef(), false, false);
 		}
 	}
+	else
+	{
+		Error(LOCTEXT("NoTranslationStages", "Cannot find any translation stages!"), nullptr, nullptr);
+		return TranslateResults;
+	}
+
 	CurrentParamMapIndices.Empty();
 	ExitStatsScope();
 
-	ActiveHistoryForFunctionCalls.EndTranslation(Script);
+	ActiveHistoryForFunctionCalls.EndTranslation(GetUniqueEmitterName());
 
 	TranslateResults.bHLSLGenSucceeded = MessageLog.NumErrors == 0;
 
@@ -1219,14 +912,14 @@ const FNiagaraTranslateResults &FHlslNiagaraTranslator::Translate(UNiagaraScript
 	if (TranslateResults.bHLSLGenSucceeded)
 	{
 		//TODO: Declare all used structures up here too.
-		InScript->ReadDataSets.Empty();
-		InScript->WriteDataSets.Empty();
+		CompilationOutput.ScriptData.ReadDataSets.Empty();
+		CompilationOutput.ScriptData.WriteDataSets.Empty();
 
 		//Generate function definitions
 		FString FunctionDefinitionString = GetFunctionDefinitions();
 		FunctionDefinitionString += TEXT("\n");
 
-		if (Script->IsInterpolatedParticleSpawnScript())
+		if (TranslationStages.Num() > 1 && RequiresInterpolation())
 		{
 			//ensure the interpolated spawn constants are part of the parameter set.
 			int32 OutputIdx = 0;
@@ -1241,8 +934,7 @@ const FNiagaraTranslateResults &FHlslNiagaraTranslator::Translate(UNiagaraScript
 		// Generate the Parameter Map HLSL definitions. We don't add to the final HLSL output here. We just build up the strings and tables
 		// that are needed later.
 		TArray<FNiagaraVariable> PrimaryDataSetOutputEntries;
-		TArray<TPair<FNiagaraVariable, const UEdGraphPin*> > MissingDefaultValues;
-		FString ParameterMapDefinitionStr = BuildParameterMapHlslDefinitions(PrimaryDataSetOutputEntries, MissingDefaultValues);
+		FString ParameterMapDefinitionStr = BuildParameterMapHlslDefinitions(PrimaryDataSetOutputEntries);
 
 		for (FNiagaraTypeDefinition Type : StructsToDefine)
 		{
@@ -1260,14 +952,14 @@ const FNiagaraTranslateResults &FHlslNiagaraTranslator::Translate(UNiagaraScript
 				HlslOutput += TEXT("\t") + GetCode(ChunksByMode[(int32)ENiagaraCodeChunkMode::Uniform][i]);
 			}
 
-			if (Script->IsInterpolatedParticleSpawnScript())
+			if (bInterpolateParams)
 			{
 				//Define the params from the previous frame after the main parameters.
 				for (int32 i = 0; i < ChunksByMode[(int32)ENiagaraCodeChunkMode::Uniform].Num(); ++i)
 				{
 					//Copy the chunk so we can fiddle it's symbol name.
 					FNiagaraCodeChunk Chunk = CodeChunks[ChunksByMode[(int32)ENiagaraCodeChunkMode::Uniform][i]];
-					Chunk.SymbolName = TEXT("PREV__") + Chunk.SymbolName;
+					Chunk.SymbolName = INTERPOLATED_PARAMETER_PREFIX + Chunk.SymbolName;
 					HlslOutput += TEXT("\t") + GetCode(Chunk);
 				}
 			}
@@ -1319,7 +1011,7 @@ const FNiagaraTranslateResults &FHlslNiagaraTranslator::Translate(UNiagaraScript
 			{
 				FNiagaraParameters ExternalParams;
 				ExternalParams.Parameters = SystemEngineReadVars;
-				Script->DataSetToParameters.Add(GetSystemEngineDataSetID().Name, ExternalParams);
+				CompilationOutput.ScriptData.DataSetToParameters.Add(GetSystemEngineDataSetID().Name, ExternalParams);
 			}
 		}
 
@@ -1377,14 +1069,9 @@ const FNiagaraTranslateResults &FHlslNiagaraTranslator::Translate(UNiagaraScript
 			// We need to reserve a place in the simulation context for the base Parameter Map.
 			if (PrimaryDataSetOutputEntries.Num() != 0 || ParamMapDefinedSystemToNamespaceVars.Num() != 0 || ParamMapDefinedEmitterParameterToNamespaceVars.Num() != 0 || (ParamMapSetVariablesToChunks.Num() != 0 && ParamMapSetVariablesToChunks[0].Num() > 0))
 			{
-				if (Script->IsInterpolatedParticleSpawnScript())
-				{
-					HlslOutput += TEXT("\tFParamMap0 MapSpawn;\n");
-					HlslOutput += TEXT("\tFParamMap0 MapUpdate;\n");
-				}
-				else
+				for (int32 i = 0; i < TranslationStages.Num(); i++)
 				{
-					HlslOutput += TEXT("\tFParamMap0 Map;\n");
+					HlslOutput += TEXT("\tFParamMap0 ") + TranslationStages[i].PassNamespace + TEXT(";\n");
 				}
 			}
 
@@ -1402,7 +1089,7 @@ const FNiagaraTranslateResults &FHlslNiagaraTranslator::Translate(UNiagaraScript
 		// copy the accessed data sets over to the script, so we can grab them during sim
 		for (TPair <FNiagaraDataSetID, TMap<int32, FDataSetAccessInfo>> InfoPair : DataSetReadInfo[0])
 		{
-			Script->ReadDataSets.Add(InfoPair.Key);
+			CompilationOutput.ScriptData.ReadDataSets.Add(InfoPair.Key);
 		}
 
 		for (TPair <FNiagaraDataSetID, TMap<int32, FDataSetAccessInfo>> InfoPair : DataSetWriteInfo[0])
@@ -1414,7 +1101,7 @@ const FNiagaraTranslateResults &FHlslNiagaraTranslator::Translate(UNiagaraScript
 				SetProps.Variables = IndexPair.Value.Variables;
 			}
 
-			Script->WriteDataSets.Add(SetProps);
+			CompilationOutput.ScriptData.WriteDataSets.Add(SetProps);
 
 			int32* ConditionalWriteChunkIdxPtr = DataSetWriteConditionalInfo[0].Find(InfoPair.Key);
 			if (ConditionalWriteChunkIdxPtr == nullptr)
@@ -1429,40 +1116,32 @@ const FNiagaraTranslateResults &FHlslNiagaraTranslator::Translate(UNiagaraScript
 
 		DefineInterpolatedParametersFunction(HlslOutput);
 
-		if (Script->IsParticleSpawnScript() || Script->IsInterpolatedParticleSpawnScript() || Script->IsSystemSpawnScript())
+		if (IsSpawnScript())
 		{
-			HlslOutput += BuildMissingDefaults(MissingDefaultValues);
+			// The missing defaults function assumes the active stage index is the 0th value in the array.
+			int32 CurStageIdx = ActiveStageIdx;
+			ActiveStageIdx = 0;
+			HlslOutput += BuildMissingDefaults();
+			ActiveStageIdx = CurStageIdx;
 		}
 
 		// define functions for reading and writing all secondary data sets
-		DefineDataSetReadFunction(HlslOutput, InScript->ReadDataSets);
-		DefineDataSetWriteFunction(HlslOutput, InScript->WriteDataSets, WriteConditionVars);
+		DefineDataSetReadFunction(HlslOutput, CompilationOutput.ScriptData.ReadDataSets);
+		DefineDataSetWriteFunction(HlslOutput, CompilationOutput.ScriptData.WriteDataSets, WriteConditionVars);
 
 
 		//Define the shared per instance simulation function
 		// for interpolated scripts AND GPU sim, define spawn and sim in separate functions
-		if (Script->IsInterpolatedParticleSpawnScript() || Options.SimTarget == ENiagaraSimTarget::GPUComputeSim)
+		if (TranslationStages.Num() > 1)
 		{
-			HlslOutput += TEXT("void SimulateSpawn(inout FSimulationContext Context)\n{\n");
-			for (int32 i = 0; i < ChunksByMode[(int32)ENiagaraCodeChunkMode::SpawnBody].Num(); ++i)
-			{
-				HlslOutput += TEXT("\t") + GetCode(ChunksByMode[(int32)ENiagaraCodeChunkMode::SpawnBody][i]);
-			}
-			HlslOutput += TEXT("}\n");
-
-			HlslOutput += TEXT("void SimulateUpdate(inout FSimulationContext Context)\n{\n");
-			for (int32 i = 0; i < ChunksByMode[(int32)ENiagaraCodeChunkMode::UpdateBody].Num(); ++i)
-			{
-				HlslOutput += TEXT("\t") + GetCode(ChunksByMode[(int32)ENiagaraCodeChunkMode::UpdateBody][i]);
-			}
-			HlslOutput += TEXT("}\n");
-
-			// add simulate function 
-			if (Options.SimTarget == ENiagaraSimTarget::CPUSim)
+			for (int32 StageIdx = 0; StageIdx < TranslationStages.Num(); StageIdx++)
 			{
-				HlslOutput += TEXT("void Simulate(inout FSimulationContext Context)\n{\n");
-				HlslOutput += TEXT("\tSimulateSpawn(Context);\n");
-				HlslOutput += TEXT("\tSimulateUpdate(Context);\n");
+				HlslOutput += TEXT("void Simulate") + TranslationStages[StageIdx].PassNamespace + TEXT("(inout FSimulationContext Context)\n{\n");
+				int32 ChunkMode = (int32)TranslationStages[StageIdx].ChunkModeIndex;
+				for (int32 i = 0; i < ChunksByMode[ChunkMode].Num(); ++i)
+				{
+					HlslOutput += TEXT("\t") + GetCode(ChunksByMode[ChunkMode][i]);
+				}
 				HlslOutput += TEXT("}\n");
 			}
 		}
@@ -1476,7 +1155,7 @@ const FNiagaraTranslateResults &FHlslNiagaraTranslator::Translate(UNiagaraScript
 			HlslOutput += TEXT("}\n");
 		}
 
-		if (Options.SimTarget == ENiagaraSimTarget::GPUComputeSim)
+		if (TranslationOptions.SimTarget == ENiagaraSimTarget::GPUComputeSim)
 		{
 			FString DataInterfaceHLSL;
 			DefineDataInterfaceHLSL(DataInterfaceHLSL);
@@ -1494,34 +1173,29 @@ const FNiagaraTranslateResults &FHlslNiagaraTranslator::Translate(UNiagaraScript
 			if (FNiagaraParameterMapHistory::IsAttribute(Var))
 			{
 				FNiagaraVariable BasicAttribVar = FNiagaraParameterMapHistory::ResolveAsBasicAttribute(Var);
-				CompilationOutput.Attributes.AddUnique(BasicAttribVar);
+				CompilationOutput.ScriptData.Attributes.AddUnique(BasicAttribVar);
 			}
 			else
 			{
-				CompilationOutput.Attributes.AddUnique(Var);
+				CompilationOutput.ScriptData.Attributes.AddUnique(Var);
 			}
 		}
 
 		// We may have created some transient data interfaces. This cleans up the ones that we created.
-		for (int32 Idx = 0; Idx < CompilationOutput.DataInterfaceInfo.Num(); Idx++)
-		{
-			if (CompilationOutput.DataInterfaceInfo[Idx].DataInterface != nullptr &&
-				PlaceholderDataInterfaces.Contains(CompilationOutput.DataInterfaceInfo[Idx].DataInterface))
-			{
-				CompilationOutput.DataInterfaceInfo[Idx].DataInterface = nullptr;
-			}
-		}
-		CompilationOutput.bReadsAttributeData = InstanceRead.Variables.Num() != 0;
+		CompilationOutput.ScriptData.DIParamInfo = DIParamInfo;
+		CompilationOutput.ScriptData.bReadsAttributeData = InstanceRead.Variables.Num() != 0;
 		TranslateResults.OutputHLSL = HlslOutput;
+
 	}
 
 	return TranslateResults;
 }
 
-void FHlslNiagaraTranslator::GatherVariableForDataSetAccess(const FNiagaraVariable& Var, FString Format, int32& Counter, int32 DataSetIndex, FString InstanceIdxSymbol, FString &HlslOutputString)
+
+void FHlslNiagaraTranslator::GatherVariableForDataSetAccess(const FNiagaraVariable& Var, FString Format, int32& IntCounter, int32 &FloatCounter, int32 DataSetIndex, FString InstanceIdxSymbol, FString &HlslOutputString)
 {
 	TArray<FString> Components;
-	const UScriptStruct* Struct = Var.GetType().GetScriptStruct();
+	UScriptStruct* Struct = Var.GetType().GetScriptStruct();
 	check(Struct);
 
 	TArray<ENiagaraBaseTypes> Types;
@@ -1551,25 +1225,34 @@ void FHlslNiagaraTranslator::GatherVariableForDataSetAccess(const FNiagaraVariab
 		{
 			FormatArgs[1] = TEXT("Float");
 			FormatArgs[DefaultIdx] = TEXT("0.0f");
+			FormatArgs[RegIdx] = FloatCounter++;
 		}
 		else if (Types[CompIdx] == NBT_Int32)
 		{
 			FormatArgs[1] = TEXT("Int");
 			FormatArgs[DefaultIdx] = TEXT("0");
+			if (CompilationTarget == ENiagaraSimTarget::GPUComputeSim)
+			{
+				FormatArgs[RegIdx] = IntCounter++;
+			}
+			else
+			{
+				FormatArgs[RegIdx] = FloatCounter++;
+			}
 		}
 		else
 		{
 			check(Types[CompIdx] == NBT_Bool);
 			FormatArgs[1] = TEXT("Bool");
 			FormatArgs[DefaultIdx] = TEXT("false");
+			FormatArgs[RegIdx] = CompilationTarget == ENiagaraSimTarget::GPUComputeSim ? IntCounter++ : FloatCounter++;
 		}
 		FormatArgs[0] = Components[CompIdx];
-		FormatArgs[RegIdx] = Counter++;
 		HlslOutputString += FString::Format(*Format, FormatArgs);
 	}
 }
 
-void FHlslNiagaraTranslator::GatherComponentsForDataSetAccess(const UScriptStruct* Struct, FString VariableSymbol, bool bMatrixRoot, TArray<FString>& Components, TArray<ENiagaraBaseTypes>& Types)
+void FHlslNiagaraTranslator::GatherComponentsForDataSetAccess(UScriptStruct* Struct, FString VariableSymbol, bool bMatrixRoot, TArray<FString>& Components, TArray<ENiagaraBaseTypes>& Types)
 {
 	bool bIsVector = IsHlslBuiltinVector(FNiagaraTypeDefinition(Struct));
 	bool bIsScalar = FNiagaraTypeDefinition::IsScalarDefinition(Struct);
@@ -1643,9 +1326,13 @@ void FHlslNiagaraTranslator::GatherComponentsForDataSetAccess(const UScriptStruc
 
 void FHlslNiagaraTranslator::DefineInterpolatedParametersFunction(FString &HlslOutputString)
 {
-	// temporarily not doing this for GPU sim; interpolated spawn doesn't work there
-	if (Script->IsInterpolatedParticleSpawnScript())
+	for (int32 i = 0; i < TranslationStages.Num(); i++)
 	{
+		if (!TranslationStages[i].bInterpolatePreviousParams)
+		{
+			continue;
+		}
+
 		FString Emitter_InterpSpawnStartDt = GetSanitizedSymbolName(ActiveHistoryForFunctionCalls.ResolveAliases(SYS_PARAM_EMITTER_INTERP_SPAWN_START_DT).GetName().ToString());
 		Emitter_InterpSpawnStartDt = Emitter_InterpSpawnStartDt.Replace(TEXT("."), TEXT("_"));//TODO: This should be rolled into GetSanitisedSymbolName but currently some usages require the '.' intact. Fix those up!.
 		FString Emitter_SpawnInterval = GetSanitizedSymbolName(ActiveHistoryForFunctionCalls.ResolveAliases(SYS_PARAM_EMITTER_SPAWN_INTERVAL).GetName().ToString());
@@ -1653,8 +1340,8 @@ void FHlslNiagaraTranslator::DefineInterpolatedParametersFunction(FString &HlslO
 		
 		HlslOutputString += TEXT("void InterpolateParameters(inout FSimulationContext Context)\n{\n");
 
-		// TODO: GPU interpolated spawning
-		if (CompilationTarget != ENiagaraSimTarget::GPUComputeSim)
+		FString PrevMap = TranslationStages[i-1].PassNamespace;
+		FString CurMap = TranslationStages[i].PassNamespace;
 		{
 			HlslOutputString += TEXT("\tint InterpSpawn_Index = ExecIndex();\n");
 			HlslOutputString += TEXT("\tfloat InterpSpawn_SpawnTime = ") + Emitter_InterpSpawnStartDt + TEXT(" + (") + Emitter_SpawnInterval + TEXT(" * InterpSpawn_Index);\n");
@@ -1662,6 +1349,16 @@ void FHlslNiagaraTranslator::DefineInterpolatedParametersFunction(FString &HlslO
 			HlslOutputString += TEXT("\tfloat InterpSpawn_InvSpawnTime = 1.0 / InterpSpawn_SpawnTime;\n");
 			HlslOutputString += TEXT("\tfloat InterpSpawn_InvUpdateTime = 1.0 / InterpSpawn_UpdateTime;\n");
 			HlslOutputString += TEXT("\tfloat SpawnInterp = InterpSpawn_SpawnTime * Engine_InverseDeltaTime ;\n");
+
+			HlslOutputString += TEXT("\tContext.") + PrevMap + TEXT(".Interpolation.InterpSpawn_Index = InterpSpawn_Index;\n");
+			HlslOutputString += TEXT("\tContext.") + PrevMap + TEXT(".Interpolation.InterpSpawn_SpawnTime = InterpSpawn_SpawnTime;\n");
+			HlslOutputString += TEXT("\tContext.") + PrevMap + TEXT(".Interpolation.InterpSpawn_UpdateTime = InterpSpawn_UpdateTime;\n");
+			HlslOutputString += TEXT("\tContext.") + PrevMap + TEXT(".Interpolation.InterpSpawn_InvSpawnTime = InterpSpawn_InvSpawnTime;\n");
+			HlslOutputString += TEXT("\tContext.") + PrevMap + TEXT(".Interpolation.InterpSpawn_InvUpdateTime = InterpSpawn_InvUpdateTime;\n");
+			HlslOutputString += TEXT("\tContext.") + PrevMap + TEXT(".Interpolation.SpawnInterp = SpawnInterp;\n");
+			HlslOutputString += TEXT("\tContext.") + PrevMap + TEXT(".Interpolation.Emitter_SpawnInterval = Emitter_SpawnInterval;\n");
+			HlslOutputString += TEXT("\tContext.") + PrevMap + TEXT(".Interpolation.Emitter_InterpSpawnStartDt = Emitter_InterpSpawnStartDt;\n");
+
 			for (int32 UniformIdx = 0; UniformIdx < ChunksByMode[(int32)ENiagaraCodeChunkMode::Uniform].Num(); ++UniformIdx)
 			{
 				int32 ChunkIdx = ChunksByMode[(int32)ENiagaraCodeChunkMode::Uniform][UniformIdx];
@@ -1691,20 +1388,19 @@ void FHlslNiagaraTranslator::DefineInterpolatedParametersFunction(FString &HlslO
 						FNiagaraCodeChunk& Chunk = CodeChunks[ChunkIdx];
 						if (ShouldInterpolateParameter(*FoundNamespacedVar))
 						{
-							HlslOutputString += TEXT("\tContext.MapSpawn.") + FoundName + TEXT(" = lerp(PREV__") + Chunk.SymbolName + TEXT(", ") + Chunk.SymbolName + TEXT(", SpawnInterp);\n");
+							HlslOutputString += TEXT("\tContext.") + PrevMap + TEXT(".") + FoundName + TEXT(" = lerp(") + INTERPOLATED_PARAMETER_PREFIX + Chunk.SymbolName + TEXT(", ") + Chunk.SymbolName + TEXT(", ") + TEXT("SpawnInterp);\n");
 						}
 						else
 						{
 							// For now, we do nothing for non-floating point variables..
-							//HlslOutputString += TEXT("\tContext.MapSpawn.") + FoundName + TEXT(" = ") + Chunk.SymbolName + TEXT(";\n");
 						}
 					}
 				}
 			}
-			HlslOutputString += TEXT("\tContext.MapSpawn.Engine.DeltaTime = 0.0f;\n");
-			HlslOutputString += TEXT("\tContext.MapSpawn.Engine.InverseDeltaTime = 0.0f;\n");
-			HlslOutputString += TEXT("\tContext.MapUpdate.Engine.DeltaTime = InterpSpawn_UpdateTime;\n");
-			HlslOutputString += TEXT("\tContext.MapUpdate.Engine.InverseDeltaTime = InterpSpawn_InvUpdateTime;\n");
+			HlslOutputString += TEXT("\tContext.") + PrevMap + TEXT(".Engine.DeltaTime = 0.0f;\n");
+			HlslOutputString += TEXT("\tContext.") + PrevMap + TEXT(".Engine.InverseDeltaTime = 0.0f;\n");
+			HlslOutputString += TEXT("\tContext.") + CurMap + TEXT(".Engine.DeltaTime = InterpSpawn_UpdateTime;\n");
+			HlslOutputString += TEXT("\tContext.") + CurMap + TEXT(".Engine.InverseDeltaTime = InterpSpawn_InvUpdateTime;\n");
 		}
 
 		HlslOutputString += TEXT("}\n\n");
@@ -1713,7 +1409,7 @@ void FHlslNiagaraTranslator::DefineInterpolatedParametersFunction(FString &HlslO
 
 void FHlslNiagaraTranslator::DefineDataSetReadFunction(FString &HlslOutputString, TArray<FNiagaraDataSetID> &ReadDataSets)
 {
-	if (Script->IsParticleEventScript() && CompilationTarget == ENiagaraSimTarget::GPUComputeSim)
+	if (UNiagaraScript::IsParticleEventScript(CompileOptions.TargetUsage) && CompilationTarget == ENiagaraSimTarget::GPUComputeSim)
 	{
 		HlslOutputString += TEXT("void ReadDataSets(inout FSimulationContext Context, int SetInstanceIndex)\n{\n");
 	}
@@ -1723,7 +1419,7 @@ void FHlslNiagaraTranslator::DefineDataSetReadFunction(FString &HlslOutputString
 	}
 
 	// We shouldn't read anything in a Spawn Script!
-	if (Script->IsParticleSpawnScript())
+	if (UNiagaraScript::IsParticleSpawnScript(CompileOptions.TargetUsage) || UNiagaraScript::IsGPUScript(CompileOptions.TargetUsage))
 	{
 		HlslOutputString += TEXT("}\n\n");
 		return;
@@ -1732,7 +1428,10 @@ void FHlslNiagaraTranslator::DefineDataSetReadFunction(FString &HlslOutputString
 	for (TPair<FNiagaraDataSetID, TMap<int32, FDataSetAccessInfo>>& DataSetInfoPair : DataSetReadInfo[0])
 	{
 		FNiagaraDataSetID DataSet = DataSetInfoPair.Key;
-		int32 OffsetCounter = 0;
+		int32 OffsetCounterInt = 0;
+		int32 OffsetCounterFloat = 0;
+		int32 &FloatCounter = OffsetCounterFloat;
+		int32 &IntCounter = CompilationTarget == ENiagaraSimTarget::GPUComputeSim ? OffsetCounterInt : OffsetCounterFloat;
 		int32 DataSetIndex = 1;
 		for (TPair<int32, FDataSetAccessInfo>& IndexInfoPair : DataSetInfoPair.Value)
 		{
@@ -1745,7 +1444,7 @@ void FHlslNiagaraTranslator::DefineDataSetReadFunction(FString &HlslOutputString
 				{
 					// TODO: temp = should really generate output functions for each set
 					FString Fmt = Symbol + Var.GetName().ToString() + FString(TEXT("{0} = ReadDataSet{1}")) + SetIdx + TEXT("[{2}*") + DataSetComponentBufferSize + " + SetInstanceIndex];\n";
-					GatherVariableForDataSetAccess(Var, Fmt, OffsetCounter, -1, TEXT(""), HlslOutputString);
+					GatherVariableForDataSetAccess(Var, Fmt, IntCounter, FloatCounter, -1, TEXT(""), HlslOutputString);
 				}
 			}
 			else
@@ -1754,7 +1453,7 @@ void FHlslNiagaraTranslator::DefineDataSetReadFunction(FString &HlslOutputString
 				{
 					// TODO: currently always emitting a non-advancing read, needs to be changed for some of the use cases
 					FString Fmt = TEXT("\tContext.") + DataSet.Name.ToString() + "Read." + Var.GetName().ToString() + TEXT("{0} = InputDataNoadvance{1}({2}, {3});\n");
-					GatherVariableForDataSetAccess(Var, Fmt, OffsetCounter, DataSetIndex, TEXT(""), HlslOutputString);
+					GatherVariableForDataSetAccess(Var, Fmt, IntCounter, FloatCounter, DataSetIndex, TEXT(""), HlslOutputString);
 				}
 			}
 		}
@@ -1768,12 +1467,13 @@ void FHlslNiagaraTranslator::DefineDataSetWriteFunction(FString &HlslOutputStrin
 {
 	HlslOutputString += TEXT("void WriteDataSets(inout FSimulationContext Context)\n{\n");
 
+	int32 DataSetIndex = 1;
 	for (TPair<FNiagaraDataSetID, TMap<int32, FDataSetAccessInfo>>& DataSetInfoPair : DataSetWriteInfo[0])
 	{
 		FNiagaraDataSetID DataSet = DataSetInfoPair.Key;
 		int32 OffsetCounter = 0;
-		int32 DataSetIndex = 1;
 
+		HlslOutputString += "\t{\n";
 		HlslOutputString += "\tint TmpWriteIndex;\n";
 		int32* ConditionalWriteIdxPtr = DataSetWriteConditionalInfo[0].Find(DataSet);
 		if (nullptr == ConditionalWriteIdxPtr || INDEX_NONE == *ConditionalWriteIdxPtr)
@@ -1784,10 +1484,16 @@ void FHlslNiagaraTranslator::DefineDataSetWriteFunction(FString &HlslOutputStrin
 		{
 			HlslOutputString += "\tbool bValid = " + FString("Context.") + DataSet.Name.ToString() + "Write_Valid;\n";
 		}
-		int32 WriteOffset = 0;
+		int32 WriteOffsetInt = 0;
+		int32 WriteOffsetFloat = 0;
+		int32 &FloatCounter = WriteOffsetFloat;
+		int32 &IntCounter = CompilationTarget == ENiagaraSimTarget::GPUComputeSim ? WriteOffsetInt : WriteOffsetFloat;
+
 		// grab the current ouput index; currently pass true, but should use an arbitrary bool to determine whether write should happen or not
 
-		HlslOutputString += "\tTmpWriteIndex = AcquireIndex(1, bValid);\n";
+		HlslOutputString += "\tTmpWriteIndex = AcquireIndex(";
+		HlslOutputString += FString::FromInt(DataSetIndex);
+		HlslOutputString +=", bValid);\n";
 
 		HlslOutputString += CompilationTarget==ENiagaraSimTarget::GPUComputeSim ? "\tif(TmpWriteIndex>=0)\n\t{\n" : "";
 
@@ -1802,7 +1508,7 @@ void FHlslNiagaraTranslator::DefineDataSetWriteFunction(FString &HlslOutputStrin
 				{
 					// TODO: temp = should really generate output functions for each set
 					FString Fmt = FString(TEXT("\t\tRWWriteDataSet{1}")) + SetIdx + TEXT("[{2}*") + DataSetComponentBufferSize + TEXT(" + {3}] = ") + Symbol + TEXT(".") + Var.GetName().ToString() + TEXT("{0};\n");
-					GatherVariableForDataSetAccess(Var, Fmt, WriteOffset, -1, TEXT("TmpWriteIndex"), HlslOutputString);
+					GatherVariableForDataSetAccess(Var, Fmt, IntCounter, FloatCounter, -1, TEXT("TmpWriteIndex"), HlslOutputString);
 				}
 			}
 			else
@@ -1810,14 +1516,15 @@ void FHlslNiagaraTranslator::DefineDataSetWriteFunction(FString &HlslOutputStrin
 				for (FNiagaraVariable &Var : IndexInfoPair.Value.Variables)
 				{
 					// TODO: data set index is always 1; need to increase each set
-					FString Fmt = TEXT("\t\tOutputData{1}(1, {2}, {3}, ") + Symbol + "." + Var.GetName().ToString() + TEXT("{0});\n");
-					GatherVariableForDataSetAccess(Var, Fmt, WriteOffset, -1, TEXT("TmpWriteIndex"), HlslOutputString);
+					FString Fmt = TEXT("\t\tOutputData{1}(") + FString::FromInt(DataSetIndex) + (", {2}, {3}, ") + Symbol + "." + Var.GetName().ToString() + TEXT("{0});\n");
+					GatherVariableForDataSetAccess(Var, Fmt, FloatCounter, FloatCounter, -1, TEXT("TmpWriteIndex"), HlslOutputString);
 				}
 			}
 		}
 
 		HlslOutputString += CompilationTarget == ENiagaraSimTarget::GPUComputeSim ? "\t}\n" : "";
 		DataSetIndex++;
+		HlslOutputString += "\t}\n";
 	}
 
 	HlslOutput += TEXT("}\n\n");
@@ -1836,27 +1543,26 @@ void FHlslNiagaraTranslator::DefineDataInterfaceHLSL(FString &InHlslOutput)
 	}
 
 	uint32 CurBufferIndex = 0;
-	for (int32 i = 0; i < CompilationOutput.DataInterfaceInfo.Num(); i++)
+	for (int32 i = 0; i < CompilationOutput.ScriptData.DataInterfaceInfo.Num(); i++)
 	{
-		FNiagaraScriptDataInterfaceInfo &Info = CompilationOutput.DataInterfaceInfo[i];
-		
-		if (!Info.DataInterface)
-		{
-			Error(LOCTEXT("InvalidDataInterfacePtr", "Null DataInterface pointer!"), nullptr, nullptr);
-			return;
-		}
+		FNiagaraScriptDataInterfaceCompileInfo& Info = CompilationOutput.ScriptData.DataInterfaceInfo[i];		
 
-		if (Info.DataInterface->CanExecuteOnTarget(ENiagaraSimTarget::GPUComputeSim))
+		UObject* const* FoundCDO = CompileData->CDOs.Find(Info.Type.GetClass());
+		check(FoundCDO != nullptr);
+		UNiagaraDataInterface* CDO = Cast<UNiagaraDataInterface>(*FoundCDO);
+		if (CDO && CDO->CanExecuteOnTarget(ENiagaraSimTarget::GPUComputeSim))
 		{
 			TArray<FNiagaraFunctionSignature> DataInterfaceFunctions;
-			Info.DataInterface->GetFunctions(DataInterfaceFunctions);
+			CDO->GetFunctions(DataInterfaceFunctions);
 			FString OwnerIDString = Info.Name.ToString();
 			FString SanitizedOwnerIDString = GetSanitizedSymbolName(OwnerIDString, true);
 
 			// grab the buffer definition from the interface
 			//
-			int32 NewIdx = DIBufferDescriptors.AddDefaulted(1);
-			Info.DataInterface->GetBufferDefinitionHLSL(SanitizedOwnerIDString, DIBufferDescriptors[NewIdx].Descriptors, InterfaceUniformHLSL);
+			int32 NewIdx = DIParamInfo.AddDefaulted(1);
+			DIParamInfo[NewIdx].DataInterfaceHLSLSymbol = SanitizedOwnerIDString;
+			DIParamInfo[NewIdx].DIClassName = Info.Type.GetClass()->GetName();
+			CDO->GetParameterDefinitionHLSL(DIParamInfo[NewIdx], InterfaceUniformHLSL);
 
 			// grab the function hlsl from the interface
 			//
@@ -1866,13 +1572,13 @@ void FHlslNiagaraTranslator::DefineDataInterfaceHLSL(FString &InHlslOutput)
 				Sig.OwnerName = Info.Name;
 				FString DefStr = GetFunctionSignatureSymbol(Sig);
 
-				bool HlslOK = Info.DataInterface->GetFunctionHLSL(Sig.Name, DefStr, DIBufferDescriptors[NewIdx].Descriptors, SanitizedOwnerIDString, InterfaceFunctionHLSL);
+				bool HlslOK = CDO->GetFunctionHLSL(Sig.Name, DefStr, DIParamInfo[NewIdx], InterfaceFunctionHLSL);
 				ensure(HlslOK == true);
 			}
 		}
 		else
 		{
-			Error(FText::Format(LOCTEXT("NonGPUDataInterfaceError", "DataInterface {0} ({1}) cannot run on the GPU."), FText::FromName(Info.Name), FText::FromString(Info.DataInterface->GetClass()->GetName()))
+			Error(FText::Format(LOCTEXT("NonGPUDataInterfaceError", "DataInterface {0} ({1}) cannot run on the GPU."), FText::FromName(Info.Name), FText::FromString(CDO ? CDO->GetClass()->GetName() : TEXT("")))
 				, nullptr, nullptr);
 		}
 	}
@@ -1893,7 +1599,7 @@ void FHlslNiagaraTranslator::DefineMain(FString &OutHlslOutput,
 		OutHlslOutput += TEXT("void SimulateMain()\n{\n");
 	}
 
-	EnterStatsScope(FNiagaraStatScope(*(Script->GetName() + TEXT("_Main")), TEXT("Main")), OutHlslOutput);
+	EnterStatsScope(FNiagaraStatScope(*(CompileOptions.GetName() + TEXT("_Main")), TEXT("Main")), OutHlslOutput);
 
 
 	OutHlslOutput += TEXT("\n\tFSimulationContext Context = (FSimulationContext)0;\n");
@@ -1913,51 +1619,103 @@ void FHlslNiagaraTranslator::DefineMain(FString &OutHlslOutput,
 		DefineDataSetVariableReads(HlslOutput, ReadIds[VarArrayIdx], VarArrayIdx, *ArrayRef);
 	}
 
+	bool bNeedsPersistentIDs = CompileOptions.AdditionalDefines.Contains(TEXT("RequiresPersistentIDs"));
+	if (bNeedsPersistentIDs && UNiagaraScript::IsSpawnScript(CompileOptions.TargetUsage))
+	{
+		FString MapName = UNiagaraScript::IsInterpolatedParticleSpawnScript(CompileOptions.TargetUsage) ? TEXT("Context.MapSpawn") : TEXT("Context.Map");
+		//Add code to handle persistent IDs.
+		OutHlslOutput += TEXT("\tint TempIDIndex;\n\tint TempIDTag;\n");
+		OutHlslOutput += TEXT("\tAcquireID(0, TempIDIndex, TempIDTag);\n");
+		OutHlslOutput += FString::Printf(TEXT("\t%s.Particles.ID.Index = TempIDIndex;\n\t%s.Particles.ID.AcquireTag = TempIDTag;\n"), *MapName, *MapName);
+	}
+
 	// Fill in the defaults for parameters.
 	for (int32 i = 0; i < MainPreSimulateChunks.Num(); ++i)
 	{
 		OutHlslOutput += TEXT("\t") + MainPreSimulateChunks[i] + TEXT("\n");
 	}
 	
-	// if we're in an event handler script on GPU, sim needs to loop over all events as well
-	// TODO: once looping works on CPU, we can do this for both targets
-	//
-	if (Script->IsParticleEventScript() && CompilationTarget == ENiagaraSimTarget::GPUComputeSim)
-	{
-		OutHlslOutput += TEXT("\tfor(int EventIdx=InEventIndex; EventIdx<InEventIndex+NumEventsPerParticle; EventIdx++)\n\t{\n");
-		OutHlslOutput += TEXT("\t\tReadDataSets(Context, EventIdx);\n");
-		OutHlslOutput += TEXT("\t\tSimulate(Context);\n");
-		OutHlslOutput += TEXT("\t}");
-	}
-	else
 	{
 		// call the read data set function
 		OutHlslOutput += TEXT("\tReadDataSets(Context);\n");
-
-		//Interpolate between prev and current parameters for interpolated spawning.
-		if (Script->IsInterpolatedParticleSpawnScript())
-		{
-			OutHlslOutput += TEXT("\tInterpolateParameters(Context);\n");
-		}
-
+		
 		// branch between spawn and update
 		if (CompilationTarget == ENiagaraSimTarget::GPUComputeSim)
 		{
-			//Call simulate.
-			if (Script->IsInterpolatedParticleSpawnScript())
-			{
-				OutHlslOutput += TEXT("\tif(Phase==0)\n\t{\n\t\tSimulateSpawn(Context);\n\t\tContext.MapSpawn.DataInstance.Alive=true;\n\t}\n");
-			}
-			else
+			OutHlslOutput += TEXT("\tint StartingPhase = Phase;\n");
+			OutHlslOutput += TEXT("\tGCurrentPhase = Phase;\n");
+
+			for (int32 StageIdx = 0; StageIdx < TranslationStages.Num(); StageIdx++)
 			{
-				OutHlslOutput += TEXT("\tif(Phase==0)\n\t{\n\t\tSimulateSpawn(Context);\n\t\tContext.Map.DataInstance.Alive=true;\n\t}\n");
+				OutHlslOutput += FString::Printf(TEXT("\tif(Phase==%d)\n\t{\n"), StageIdx);
+				
+				OutHlslOutput += FString::Printf(TEXT("\t\tif (StartingPhase == %d)\n\t\t{\n\t\t\tContext.%s.DataInstance.Alive=true;\n\t\t}\n"), StageIdx, *TranslationStages[StageIdx].PassNamespace);
+
+				if (StageIdx == 0)
+				{
+					// Either go on to the next phase, or write to the final output context.
+					if (RequiresInterpolation())
+					{
+						OutHlslOutput += TEXT("\t\tInterpolateParameters(Context);\n"); // Requires ExecIndex, which needs to be in a stage.
+					}
+				}
+
+				OutHlslOutput += FString::Printf(TEXT("\t\tSimulate%s(Context);\n"), TranslationStages.Num() > 1 ? *TranslationStages[StageIdx].PassNamespace : TEXT(""));
+				
+				if (StageIdx + 1 < TranslationStages.Num() && TranslationStages[StageIdx + 1].bCopyPreviousParams)
+				{
+					OutHlslOutput += FString::Printf(TEXT("\t\tContext.%s.DataInstance.Alive = Context.%s.DataInstance.Alive;\n"), *TranslationStages[StageIdx + 1].PassNamespace, *TranslationStages[StageIdx].PassNamespace);
+					OutHlslOutput += TEXT("\t\t//Begin Transfer of Attributes!\n");
+					if (ParamMapDefinedAttributesToNamespaceVars.Num() != 0)
+					{
+						FString CopyStr = TEXT("\t\tContext.") + TranslationStages[StageIdx + 1].PassNamespace + TEXT(".Particles = Context.") + TranslationStages[StageIdx].PassNamespace + TEXT(".Particles;\n");
+						OutHlslOutput += CopyStr;
+					}
+					OutHlslOutput += TEXT("\t\t//End Transfer of Attributes!\n\n");
+				}
+
+				// Either go on to the next phase, or write to the final output context.
+				if (StageIdx + 1 < TranslationStages.Num() && TranslationStages[StageIdx + 1].bInterpolatePreviousParams)
+				{
+					OutHlslOutput += FString::Printf(TEXT("\t\tPhase = %d;\n"), StageIdx + 1);
+					OutHlslOutput += FString::Printf(TEXT("\t\tGCurrentPhase = %d;\n"), StageIdx + 1);
+				}
+
+				else if (StageIdx + 1 < TranslationStages.Num())
+				{
+					OutHlslOutput += FString::Printf(TEXT("\t\tContext.%s.Particles = Context.%s.Particles;\n"), *TranslationStages[TranslationStages.Num() - 1].PassNamespace, *TranslationStages[StageIdx].PassNamespace);
+				}
+				
+				OutHlslOutput += TEXT("\t}\n");
 			}
-			OutHlslOutput += TEXT("\telse if(Phase==1)\n\t{\n\t\tSimulateUpdate(Context);\n\t}\n");
 		}
 		else
 		{
-			//Call simulate.
-			OutHlslOutput += TEXT("\tSimulate(Context);\n");
+			for (int32 StageIdx = 0; StageIdx < TranslationStages.Num(); StageIdx++)
+			{
+				if (StageIdx == 0)
+				{
+					// Either go on to the next phase, or write to the final output context.
+					if (RequiresInterpolation())
+					{
+						OutHlslOutput += TEXT("\tInterpolateParameters(Context);\n"); // Requires ExecIndex, which needs to be in a stage.
+					}
+				}
+
+				OutHlslOutput += FString::Printf(TEXT("\tSimulate%s(Context);\n"), TranslationStages.Num() > 1 ? *TranslationStages[StageIdx].PassNamespace : TEXT(""));
+
+				if (StageIdx + 1 < TranslationStages.Num() && TranslationStages[StageIdx + 1].bCopyPreviousParams)
+				{
+					OutHlslOutput += FString::Printf(TEXT("\tContext.%s.DataInstance.Alive = Context.%s.DataInstance.Alive;\n"), *TranslationStages[StageIdx + 1].PassNamespace, *TranslationStages[StageIdx].PassNamespace);
+					OutHlslOutput += TEXT("\t//Begin Transfer of Attributes!\n");
+					if (ParamMapDefinedAttributesToNamespaceVars.Num() != 0)
+					{
+						FString CopyStr = TEXT("\tContext.") + TranslationStages[StageIdx + 1].PassNamespace + TEXT(".Particles = Context.") + TranslationStages[StageIdx].PassNamespace + TEXT(".Particles;\n");
+						OutHlslOutput += CopyStr;
+					}
+					OutHlslOutput += TEXT("\t//End Transfer of Attributes!\n\n");
+				}
+			}
 		}
 	}
 
@@ -1984,7 +1742,7 @@ void FHlslNiagaraTranslator::DefineMain(FString &OutHlslOutput,
 	if (CompilationTarget == ENiagaraSimTarget::GPUComputeSim)
 	{
 		OutHlslOutput += TEXT("void CopyInstance(in int InstanceIdx)\n{\n\tFSimulationContext Context = (FSimulationContext)0;\n");
-		if (Script->IsParticleEventScript())
+		if (UNiagaraScript::IsParticleEventScript(CompileOptions.TargetUsage))
 		{
 			for (int32 VarArrayIdx = 0; VarArrayIdx < InstanceReadVars.Num(); VarArrayIdx++)
 			{
@@ -2007,9 +1765,10 @@ void FHlslNiagaraTranslator::DefineMain(FString &OutHlslOutput,
 
 void FHlslNiagaraTranslator::DefineDataSetVariableWrites(FString &OutHlslOutput, FNiagaraDataSetID& Id, int32 DataSetIndex, TArray<FNiagaraVariable>& WriteVars)
 {
-	//TODO Grab indices for data set writes (inc output) and do the write
+	//TODO Grab indices for data set writes (inc output) and do the write. Need to rewrite this for events interleaved..
+	OutHlslOutput += "\t{\n";
 	OutHlslOutput += "\tint TmpWriteIndex;\n";
-	if (Script->IsNonParticleScript())
+	if (UNiagaraScript::IsNonParticleScript(CompileOptions.TargetUsage))
 	{
 		OutHlslOutput += "\tbool bValid = true;\n";
 	}
@@ -2037,7 +1796,7 @@ void FHlslNiagaraTranslator::DefineDataSetVariableWrites(FString &OutHlslOutput,
 			}
 		}
 
-		if ((bHasPerParticleAliveSpawn || bHasPerParticleAliveUpdate) && Script->IsInterpolatedParticleSpawnScript())
+		if ((bHasPerParticleAliveSpawn || bHasPerParticleAliveUpdate) && TranslationStages.Num() > 1)
 		{
 			if (bHasPerParticleAliveSpawn && bHasPerParticleAliveUpdate)
 			{
@@ -2052,10 +1811,11 @@ void FHlslNiagaraTranslator::DefineDataSetVariableWrites(FString &OutHlslOutput,
 				OutHlslOutput += TEXT("\tbool bValid = Context.MapUpdate.") + DataSetName + TEXT(".Alive;\n");
 			}
 		}
-		else if ((Script->IsParticleSpawnScript() && bHasPerParticleAliveSpawn) 
-			|| (Script->IsParticleUpdateScript() && bHasPerParticleAliveUpdate) 
-			|| (Script->IsParticleEventScript() && bHasPerParticleAliveEvent)
-			|| (Script->IsParticleSpawnScript() && bHasPerParticleAliveUpdate && TranslationOptions.SimTarget == ENiagaraSimTarget::GPUComputeSim))
+		else if ((UNiagaraScript::IsParticleSpawnScript(CompileOptions.TargetUsage) && bHasPerParticleAliveSpawn) 
+			|| (UNiagaraScript::IsGPUScript(CompileOptions.TargetUsage) && bHasPerParticleAliveSpawn)
+			|| (UNiagaraScript::IsParticleUpdateScript(CompileOptions.TargetUsage) && bHasPerParticleAliveUpdate) 
+			|| (UNiagaraScript::IsParticleEventScript(CompileOptions.TargetUsage) && bHasPerParticleAliveEvent)
+			|| (UNiagaraScript::IsGPUScript(CompileOptions.TargetUsage) && bHasPerParticleAliveUpdate && CompilationTarget == ENiagaraSimTarget::GPUComputeSim))
 		{
 			OutHlslOutput += TEXT("\tbool bValid = Context.Map.") + DataSetName + TEXT(".Alive;\n");
 		}
@@ -2064,32 +1824,48 @@ void FHlslNiagaraTranslator::DefineDataSetVariableWrites(FString &OutHlslOutput,
 			OutHlslOutput += "\tbool bValid = true;\n";
 		}
 	}
-	int32 WriteOffset = 0;
+	int32 WriteOffsetInt = 0;
+	int32 WriteOffsetFloat = 0;
+	int32 &FloatCounter = WriteOffsetFloat;
+	int32 &IntCounter = CompilationTarget == ENiagaraSimTarget::GPUComputeSim ? WriteOffsetInt : WriteOffsetFloat;
 
 	// grab the current ouput index; currently pass true, but should use an arbitrary bool to determine whether write should happen or not
 	OutHlslOutput += "\tTmpWriteIndex = AcquireIndex(0, bValid);\n";
+
+	bool bNeedsPersistentIDs = CompileOptions.AdditionalDefines.Contains(TEXT("RequiresPersistentIDs"));
+	if (bNeedsPersistentIDs && DataSetIndex == 0)
+	{
+		FString MapName = GetParameterMapInstanceName(0);
+		OutHlslOutput += FString::Printf(TEXT("\tUpdateID(0, %s.Particles.ID.Index, TmpWriteIndex);\n"), *MapName);
+	}
+
 	for (FNiagaraVariable &Var : WriteVars)
 	{
 		// If coming from a parameter map, use the one on the context, otherwise use the output.
 		FString Fmt;
 
 		{
-			if (Script->IsInterpolatedParticleSpawnScript())
+			if (TranslationStages.Num() > 1)
 			{
-				Fmt = TEXT("\tOutputData{1}(0, {2}, {3}, Context.MapUpdate.") + GetSanitizedSymbolName(Var.GetName().ToString()) + TEXT("{0});\n");
+				Fmt = TEXT("\tOutputData{1}(0, {2}, {3}, Context.") + TranslationStages[TranslationStages.Num() - 1].PassNamespace + TEXT(".") + GetSanitizedSymbolName(Var.GetName().ToString()) + TEXT("{0});\n");
 			}
 			else
 			{
 				Fmt = TEXT("\tOutputData{1}(0, {2}, {3}, Context.Map.") + GetSanitizedSymbolName(Var.GetName().ToString()) + TEXT("{0});\n");
 			}
-			GatherVariableForDataSetAccess(Var, Fmt, WriteOffset, -1, TEXT("TmpWriteIndex"), OutHlslOutput);
+			GatherVariableForDataSetAccess(Var, Fmt, IntCounter, FloatCounter, -1, TEXT("TmpWriteIndex"), OutHlslOutput);
 		}
 	}
+	OutHlslOutput += "\t}\n";
 }
 
 void FHlslNiagaraTranslator::DefineDataSetVariableReads(FString &OutHlslOutput, FNiagaraDataSetID& Id, int32 DataSetIndex, TArray<FNiagaraVariable>& ReadVars)
 {
-	int32 ReadOffset = 0;
+	int32 ReadOffsetInt = 0;
+	int32 ReadOffsetFloat = 0;
+	int32 &FloatCounter = ReadOffsetFloat;
+	int32 &IntCounter = CompilationTarget == ENiagaraSimTarget::GPUComputeSim ? ReadOffsetInt : ReadOffsetFloat;
+
 	FString DataSetName = Id.Name.ToString();
 	FString Fmt;
 	FString SpawnCondition = "\t{\n", UpdateCondition = "\t{\n";
@@ -2100,102 +1876,86 @@ void FHlslNiagaraTranslator::DefineDataSetVariableReads(FString &OutHlslOutput,
 		UpdateCondition = "\tif(Phase==1)\n\t{\n";
 	}
 
-	// interpolated spawn or particle spawn on CPU
-	FString VarReads;
-	if (Script->IsInterpolatedParticleSpawnScript())
+	bool bIsGPUScript = UNiagaraScript::IsGPUScript(CompileOptions.TargetUsage);
+	bool bIsSpawnScript = 
+		UNiagaraScript::IsParticleSpawnScript(CompileOptions.TargetUsage) ||UNiagaraScript::IsInterpolatedParticleSpawnScript(CompileOptions.TargetUsage) ||
+		UNiagaraScript::IsEmitterSpawnScript(CompileOptions.TargetUsage) ||	UNiagaraScript::IsSystemSpawnScript(CompileOptions.TargetUsage);
+	bool bIsUpdateScript = 
+		UNiagaraScript::IsParticleUpdateScript(CompileOptions.TargetUsage) || UNiagaraScript::IsEmitterUpdateScript(CompileOptions.TargetUsage) ||
+		UNiagaraScript::IsSystemUpdateScript(CompileOptions.TargetUsage);
+	bool bIsEventScript = UNiagaraScript::IsParticleEventScript(CompileOptions.TargetUsage);
+	bool bIsSystemOrEmitterScript = 
+		UNiagaraScript::IsEmitterSpawnScript(CompileOptions.TargetUsage) ||	UNiagaraScript::IsSystemSpawnScript(CompileOptions.TargetUsage) ||
+		UNiagaraScript::IsEmitterUpdateScript(CompileOptions.TargetUsage) || UNiagaraScript::IsSystemUpdateScript(CompileOptions.TargetUsage);
+	bool bIsPrimaryDataSet = DataSetIndex == 0;
+
+	// This will initialize parameters to 0 for spawning.  For the system and emitter combined spawn script we want to do this on the
+	// primary data set which contains the particle data, but we do not want to do this for the secondary data set since it has
+	// external user and engine parameters which must be read.
+	if (bIsGPUScript || (bIsSpawnScript && (bIsPrimaryDataSet || bIsSystemOrEmitterScript == false)))
 	{
-		for (FNiagaraVariable &Var : ReadVars)
+		FString ContextName = TEXT("\tContext.Map.");
+		if (TranslationStages.Num() > 1)
 		{
-			Fmt = TEXT("\tContext.MapSpawn.") + GetSanitizedSymbolName(Var.GetName().ToString()) + TEXT("{0} = {4};\n");
-			GatherVariableForDataSetAccess(Var, Fmt, ReadOffset, DataSetIndex, TEXT(""), VarReads);
+			ContextName = FString::Printf(TEXT("\tContext.%s."),*TranslationStages[0].PassNamespace);
 		}
-		OutHlslOutput += SpawnCondition;
-		OutHlslOutput += VarReads + "\t}\n";
-	}
 
-	// regular emitter or system spawn script
-	else if ((Script->IsParticleSpawnScript()) || Script->IsEmitterSpawnScript() || Script->IsSystemSpawnScript()) 	// We shouldn't read anything in a Spawn Script! Set to defaults.
-	{
+		FString VarReads;
+
 		for (FNiagaraVariable &Var : ReadVars)
 		{
-			Fmt = TEXT("\tContext.Map.") + GetSanitizedSymbolName(Var.GetName().ToString()) + TEXT("{0} = {4};\n");
-			GatherVariableForDataSetAccess(Var, Fmt, ReadOffset, DataSetIndex, TEXT(""), VarReads);
+			Fmt = ContextName + GetSanitizedSymbolName(Var.GetName().ToString()) + TEXT("{0} = {4};\n");
+			GatherVariableForDataSetAccess(Var, Fmt, IntCounter, FloatCounter, DataSetIndex, TEXT(""), VarReads);
 		}
+
 		OutHlslOutput += SpawnCondition;
 		OutHlslOutput += VarReads + "\t}\n";
 	}
 
-	VarReads = "";
-
-	if (Script->IsParticleUpdateScript() || Script->IsParticleEventScript() || Script->IsEmitterUpdateScript() || Script->IsSystemUpdateScript()
-		|| (Script->IsParticleSpawnScript() && TranslationOptions.SimTarget == ENiagaraSimTarget::GPUComputeSim))
+	// This will initialize parameters to their correct initial values from constants or data sets for update, and will also initialize parameters
+	// for spawn if this is a combined system and emitter spawn script and we're reading from a secondary data set for engine and user parameters.
+	if (bIsGPUScript || bIsEventScript || bIsUpdateScript || (bIsSpawnScript && bIsPrimaryDataSet == false && bIsSystemOrEmitterScript))
 	{
 		FString ContextName = TEXT("\tContext.Map.");
-		if (Script->IsInterpolatedParticleSpawnScript())
-		{
-			ContextName = TEXT("\tContext.MapUpdate.");
-		}
-
-		// if we're a GPU spawn script (meaning a combined spawn/update script), we need to reset register index counter
-		if (Script->IsParticleSpawnScript() && TranslationOptions.SimTarget == ENiagaraSimTarget::GPUComputeSim)
-		{
-			ReadOffset = 0;
-		}
-		for (FNiagaraVariable &Var : ReadVars)
+		if (TranslationStages.Num() > 1)
 		{
-			// If the NiagaraClearEachFrame value is set on the data set, we don't bother reading it in each frame as we know that it is is invalid. However,
-			// this is only used for the base data set. Other reads are potentially from events and are therefore perfectly valid.
-			if (DataSetIndex == 0 && Var.GetType().GetScriptStruct() != nullptr && Var.GetType().GetScriptStruct()->GetMetaData(TEXT("NiagaraClearEachFrame")).Equals(TEXT("true"), ESearchCase::IgnoreCase))
-			{
-				Fmt = ContextName + GetSanitizedSymbolName(Var.GetName().ToString()) + TEXT("{0} = {4};\n");
-			}
-			else if (CompilationTarget == ENiagaraSimTarget::GPUComputeSim)
-			{
-				Fmt = ContextName + GetSanitizedSymbolName(Var.GetName().ToString()) + TEXT("{0} = InputData{1}({2}, {3}, InstanceIdx);\n");
-			}
-			else
-			{
-				Fmt = ContextName + GetSanitizedSymbolName(Var.GetName().ToString()) + TEXT("{0} = InputData{1}({2}, {3});\n");
-			}
-			GatherVariableForDataSetAccess(Var, Fmt, ReadOffset, DataSetIndex, TEXT(""), VarReads);
+			ContextName = FString::Printf(TEXT("\tContext.%s."), *TranslationStages[TranslationStages.Num() - 1].PassNamespace);
 		}
 
-		OutHlslOutput += UpdateCondition;
-		OutHlslOutput += VarReads + "\n\t}\n";
-	}
-
-	/*
-	for (FNiagaraVariable &Var : ReadVars)
-	{
-		// If coming from a parameter map, use the one on the context, otherwise use the input.
-
-		// In the case of interpolated spawn scripts, we read into the MapSpawn and write out of the Map.
-		if (Script->IsInterpolatedParticleSpawnScript())
-		{
-			Fmt = TEXT("\tContext.MapSpawn.") + GetSanitizedSymbolName(Var.GetName().ToString()) + TEXT("{0} = {4};\n");
-		}
-		else if ( (Script->IsParticleSpawnScript()) || Script->IsEmitterSpawnScript() || Script->IsSystemSpawnScript()) 	// We shouldn't read anything in a Spawn Script! Set to defaults.
+		// if we're a GPU spawn script (meaning a combined spawn/update script), we need to reset register index counter
+		if (UNiagaraScript::IsGPUScript(CompileOptions.TargetUsage))
 		{
-			Fmt = TEXT("\tContext.Map.") + GetSanitizedSymbolName(Var.GetName().ToString()) + TEXT("{0} = {4};\n");
+			ReadOffsetInt = 0;
+			ReadOffsetFloat = 0;
 		}
-		else
+
+		FString VarReads;
+
+		for (FNiagaraVariable &Var : ReadVars)
 		{
-			if (CompilationTarget == ENiagaraSimTarget::GPUComputeSim)
+			// If the NiagaraClearEachFrame value is set on the data set, we don't bother reading it in each frame as we know that it is is invalid. However,
+			// this is only used for the base data set. Other reads are potentially from events and are therefore perfectly valid.
+			if (DataSetIndex == 0 && Var.GetType().GetScriptStruct() != nullptr && Var.GetType().GetScriptStruct()->GetMetaData(TEXT("NiagaraClearEachFrame")).Equals(TEXT("true"), ESearchCase::IgnoreCase))
+			{
+				Fmt = ContextName + GetSanitizedSymbolName(Var.GetName().ToString()) + TEXT("{0} = {4};\n");
+			}
+			else if (CompilationTarget == ENiagaraSimTarget::GPUComputeSim)
 			{
-				Fmt = TEXT("\tContext.Map.") + GetSanitizedSymbolName(Var.GetName().ToString()) + TEXT("{0} = InputData{1}({2}, {3}, InstanceIdx);\n");
+				Fmt = ContextName + GetSanitizedSymbolName(Var.GetName().ToString()) + TEXT("{0} = InputData{1}({2}, {3}, InstanceIdx);\n");
 			}
 			else
 			{
-				Fmt = TEXT("\tContext.Map.") + GetSanitizedSymbolName(Var.GetName().ToString()) + TEXT("{0} = InputData{1}({2}, {3});\n");
+				Fmt = ContextName + GetSanitizedSymbolName(Var.GetName().ToString()) + TEXT("{0} = InputData{1}({2}, {3});\n");
 			}
+			GatherVariableForDataSetAccess(Var, Fmt, IntCounter, FloatCounter, DataSetIndex, TEXT(""), VarReads);
 		}
 
-		GatherVariableForDataSetAccess(Var, Fmt, ReadOffset, DataSetIndex, TEXT(""), OutHlslOutput);
+		OutHlslOutput += UpdateCondition;
+		OutHlslOutput += VarReads + "\n\t}\n";
 	}
-	*/
 }
 
-void FHlslNiagaraTranslator::WriteDataSetContextVars(TMap<FNiagaraDataSetID, TMap<int32, FDataSetAccessInfo>>& DataSetAccessInfo, bool bRead, FString &OutHlslOutput)
+void FHlslNiagaraTranslator::WriteDataSetContextVars(TMap<FNiagaraDataSetID, TMap<int32, FDataSetAccessInfo>>& DataSetAccessInfo, bool bRead, FString& OutHlslOutput)
 {
 	//Now the intermediate storage for the data set reads and writes.
 	uint32 DataSetIndex = 0;
@@ -2213,7 +1973,7 @@ void FHlslNiagaraTranslator::WriteDataSetContextVars(TMap<FNiagaraDataSetID, TMa
 };
 
 
-void FHlslNiagaraTranslator::WriteDataSetStructDeclarations(TMap<FNiagaraDataSetID, TMap<int32, FDataSetAccessInfo>>& DataSetAccessInfo, bool bRead, FString &OutHlslOutput)
+void FHlslNiagaraTranslator::WriteDataSetStructDeclarations(TMap<FNiagaraDataSetID, TMap<int32, FDataSetAccessInfo>>& DataSetAccessInfo, bool bRead, FString& OutHlslOutput)
 {
 	uint32 DataSetIndex = 1;
 	for (TPair<FNiagaraDataSetID, TMap<int32, FDataSetAccessInfo>>& DataSetInfoPair : DataSetAccessInfo)
@@ -2226,7 +1986,7 @@ void FHlslNiagaraTranslator::WriteDataSetStructDeclarations(TMap<FNiagaraDataSet
 		{
 			for (FNiagaraVariable Var : IndexInfoPair.Value.Variables)
 			{
-				OutHlslOutput += TEXT("\t") + GetStructHlslTypeName(Var.GetType().GetScriptStruct()) + TEXT(" ") + Var.GetName().ToString() + ";\n";
+				OutHlslOutput += TEXT("\t") + GetStructHlslTypeName(Var.GetType()) + TEXT(" ") + Var.GetName().ToString() + ";\n";
 			}
 		}
 
@@ -2261,7 +2021,7 @@ void FHlslNiagaraTranslator::WriteDataSetStructDeclarations(TMap<FNiagaraDataSet
 
 
 //Decomposes each variable into its constituent register accesses.
-void FHlslNiagaraTranslator::DecomposeVariableAccess(const UStruct* Struct, bool bRead, FString IndexSymbol, FString HLSLString)
+void FHlslNiagaraTranslator::DecomposeVariableAccess(UStruct* Struct, bool bRead, FString IndexSymbol, FString HLSLString)
 {
 	FString AccessStr;
 
@@ -2372,7 +2132,7 @@ void FHlslNiagaraTranslator::EnterStatsScope(FNiagaraStatScope StatScope)
 {
 	if (GbNiagaraScriptStatTracking)
 	{
-		int32 ScopeIdx = Script->StatScopes.AddUnique(StatScope);
+		int32 ScopeIdx = CompilationOutput.ScriptData.StatScopes.AddUnique(StatScope);
 		AddBodyChunk(TEXT(""), FString::Printf(TEXT("EnterStatScope(%d /**%s*/)"), ScopeIdx, *StatScope.FullName.ToString()), FNiagaraTypeDefinition::GetFloatDef(), false);
 		StatScopeStack.Push(ScopeIdx);
 	}
@@ -2383,7 +2143,7 @@ void FHlslNiagaraTranslator::ExitStatsScope()
 	if (GbNiagaraScriptStatTracking)
 	{
 		int32 ScopeIdx = StatScopeStack.Pop();
-		AddBodyChunk(TEXT(""), FString::Printf(TEXT("ExitStatScope(/**%s*/)"), *Script->StatScopes[ScopeIdx].FullName.ToString()), FNiagaraTypeDefinition::GetFloatDef(), false);
+		AddBodyChunk(TEXT(""), FString::Printf(TEXT("ExitStatScope(/**%s*/)"), *CompilationOutput.ScriptData.StatScopes[ScopeIdx].FullName.ToString()), FNiagaraTypeDefinition::GetFloatDef(), false);
 	}
 }
 
@@ -2391,7 +2151,7 @@ void FHlslNiagaraTranslator::EnterStatsScope(FNiagaraStatScope StatScope, FStrin
 {
 	if (GbNiagaraScriptStatTracking)
 	{
-		int32 ScopeIdx = Script->StatScopes.AddUnique(StatScope);
+		int32 ScopeIdx = CompilationOutput.ScriptData.StatScopes.AddUnique(StatScope);
 		OutHlsl += FString::Printf(TEXT("EnterStatScope(%d /**%s*/);\n"), ScopeIdx, *StatScope.FullName.ToString());
 		StatScopeStack.Push(ScopeIdx);
 	}
@@ -2402,13 +2162,13 @@ void FHlslNiagaraTranslator::ExitStatsScope(FString& OutHlsl)
 	if (GbNiagaraScriptStatTracking)
 	{
 		int32 ScopeIdx = StatScopeStack.Pop();
-		OutHlsl += FString::Printf(TEXT("ExitStatScope(/**%s*/);\n"), *Script->StatScopes[ScopeIdx].FullName.ToString());
+		OutHlsl += FString::Printf(TEXT("ExitStatScope(/**%s*/);\n"), *CompilationOutput.ScriptData.StatScopes[ScopeIdx].FullName.ToString());
 	}
 }
 
 FString FHlslNiagaraTranslator::GetCallstack()
 {
-	FString Callstack = Script->GetName();
+	FString Callstack = CompileOptions.GetName();
 
 	for (FFunctionContext& Ctx : FunctionContextStack)
 	{
@@ -2573,6 +2333,19 @@ int32 FHlslNiagaraTranslator::GetParameter(const FNiagaraVariable& Parameter)
 		Error(FText::Format(LOCTEXT("GetParameterFail", "Cannot handle type {0}! Variable: {1}"), Parameter.GetType().GetNameText(), FText::FromName(Parameter.GetName())), nullptr, nullptr);
 	}
 
+	if (Parameter == TRANSLATOR_PARAM_BEGIN_DEFAULTS)
+	{
+		if (CurrentDefaultPinTraversal.Num() != 0)
+		{
+			return ActiveStageIdx;
+		}
+		else
+		{
+			Error(FText::Format(LOCTEXT("InitializingDefaults", "Cannot have a {0} node if you are not tracing a default value from a Get node."), FText::FromName(Parameter.GetName())), nullptr, nullptr);
+			return INDEX_NONE;
+		}
+	}
+
 	int32 FuncParam = INDEX_NONE;
 	if (GetFunctionParameter(Parameter, FuncParam))
 	{
@@ -2617,6 +2390,14 @@ int32 FHlslNiagaraTranslator::GetParameter(const FNiagaraVariable& Parameter)
 					UNiagaraNodeInput* InputNode = Cast<UNiagaraNodeInput>(Pin->GetOwningNode());
 					if (InputNode != nullptr && InputNode->Input == Parameter)
 					{
+						if (CurrentDefaultPinTraversal.Num() == 0)
+						{
+							if (!bInitializedDefaults)
+							{
+								InitializeParameterMapDefaults(i);
+							}
+						}
+						
 						return i;
 					}
 				}
@@ -2628,7 +2409,7 @@ int32 FHlslNiagaraTranslator::GetParameter(const FNiagaraVariable& Parameter)
 	//Not a in a function or not a valid function parameter so grab from the main uniforms.
 	int32 OutputChunkIdx = INDEX_NONE;
 	FNiagaraVariable OutputVariable = Parameter;
-	if (FNiagaraParameterMapHistory::IsExternalConstantNamespace(OutputVariable, Script))
+	if (FNiagaraParameterMapHistory::IsExternalConstantNamespace(OutputVariable, CompileOptions.TargetUsage, CompileOptions.GetTargetUsageBitmask()))
 	{
 		if (!ParameterMapRegisterExternalConstantNamespaceVariable(OutputVariable, nullptr, 0, OutputChunkIdx, nullptr))
 		{
@@ -2637,7 +2418,7 @@ int32 FHlslNiagaraTranslator::GetParameter(const FNiagaraVariable& Parameter)
 	}
 	else
 	{
-		OutputVariable = FNiagaraParameterMapHistory::MoveToExternalConstantNamespaceVariable(OutputVariable, Script);
+		OutputVariable = FNiagaraParameterMapHistory::MoveToExternalConstantNamespaceVariable(OutputVariable, CompileOptions.TargetUsage);
 		if (!ParameterMapRegisterExternalConstantNamespaceVariable(OutputVariable, nullptr, 0, OutputChunkIdx, nullptr))
 		{
 			OutputChunkIdx = INDEX_NONE;
@@ -2654,6 +2435,11 @@ int32 FHlslNiagaraTranslator::GetParameter(const FNiagaraVariable& Parameter)
 
 int32 FHlslNiagaraTranslator::GetConstant(const FNiagaraVariable& Constant)
 {
+	if (Constant.IsDataInterface())
+	{
+		return INDEX_NONE;
+	}
+
 	FString ConstantStr = GenerateConstantString(Constant);
 	if (ConstantStr.IsEmpty())
 	{
@@ -2677,7 +2463,7 @@ FString FHlslNiagaraTranslator::GenerateConstantString(const FNiagaraVariable& C
 	{
 		Error(FText::Format(LOCTEXT("GetConstantFail", "Cannot handle type {0}! Variable: {1}"), Type.GetNameText(), FText::FromName(Constant.GetName())), nullptr, nullptr);
 	}
-	FString ConstantStr = GetStructHlslTypeName(Type) + GetHlslDefaultForType(Type);
+	FString ConstantStr = GetHlslDefaultForType(Type);
 	if (Constant.IsDataAllocated())
 	{
 		if (Type == FNiagaraTypeDefinition::GetFloatDef())
@@ -2705,6 +2491,11 @@ FString FHlslNiagaraTranslator::GenerateConstantString(const FNiagaraVariable& C
 			float* ValuePtr = (float*)Constant.GetData();
 			ConstantStr = FString::Printf(TEXT("float4(%g,%g,%g,%g)"), *ValuePtr, *(ValuePtr + 1), *(ValuePtr + 2), *(ValuePtr + 3));
 		}
+		else if (Type == FNiagaraTypeDefinition::GetQuatDef())
+		{
+			float* ValuePtr = (float*)Constant.GetData();
+			ConstantStr = FString::Printf(TEXT("float4(%g,%g,%g,%g)"), *ValuePtr, *(ValuePtr + 1), *(ValuePtr + 2), *(ValuePtr + 3));
+		}
 		else if (Type == FNiagaraTypeDefinition::GetIntDef() || Type.GetStruct() == FNiagaraTypeDefinition::GetIntStruct())
 		{
 			int32* ValuePtr = (int32*)Constant.GetData();
@@ -2727,13 +2518,92 @@ FString FHlslNiagaraTranslator::GenerateConstantString(const FNiagaraVariable& C
 		{
 			//This is easily doable, just need to keep track of all structs used and define them as well as a ctor function signature with all values decomposed into float1/2/3/4 etc
 			//Then call said function here with the same decomposition literal values.
-			Error(LOCTEXT("StructContantsUnsupportedError", "Constants of struct types are currently unsupported."), nullptr, nullptr);
-			return FString();
+
+			//For now lets allow this but just ignore the value and take the default ctor.
+// 			Error(LOCTEXT("StructContantsUnsupportedError", "Constants of struct types are currently unsupported."), nullptr, nullptr);
+// 			return FString();
+			return ConstantStr;
 		}
 	}
 	return ConstantStr;
 }
 
+void FHlslNiagaraTranslator::InitializeParameterMapDefaults(int32 ParamMapHistoryIdx)
+{
+	bInitializedDefaults = true;
+	AddBodyChunk(TEXT(""), TEXT("//Begin Initialize Parameter Map Defaults"), FNiagaraTypeDefinition::GetIntDef(), false, false);
+	TArray<FNiagaraVariable> UniqueVars;
+	check(ParamMapHistories.Num() == TranslationStages.Num());
+
+	TMap<FNiagaraVariable, const UEdGraphPin*> UniqueVarToDefaultPin;
+	TMap<FNiagaraVariable, bool> UniqueVarToWriteToParamMap;
+
+	// First pass just use the current parameter map.
+	{
+		const FNiagaraParameterMapHistory& History = ParamMapHistories[ParamMapHistoryIdx];
+		for (int32 i = 0; i < History.Variables.Num(); i++)
+		{
+			const FNiagaraVariable& Var = History.Variables[i];
+			const FNiagaraVariable& AliasedVar = History.VariablesWithOriginalAliasesIntact[i];
+			// Only add primary data set outputs at the top of the script if in a spawn script, otherwise they should be left alone.
+			if (UNiagaraScript::IsSpawnScript(TranslationStages[ActiveStageIdx].ScriptUsage))
+			{
+				if (History.IsPrimaryDataSetOutput(AliasedVar, TranslationStages[ActiveStageIdx].ScriptUsage) &&
+					!UniqueVars.Contains(Var))
+				{
+					UniqueVars.Add(Var);
+					const UEdGraphPin* DefaultPin = History.GetDefaultValuePin(i);
+					UniqueVarToDefaultPin.Add(Var, DefaultPin);
+					UniqueVarToWriteToParamMap.Add(Var, true);
+				}
+			}
+		}
+	}
+
+	// Only add primary data set outputs at the top of the script if in a spawn script, otherwise they should be left alone.
+	// Above we added all the known from the spawn script, now let's add for all the others.
+	if (UNiagaraScript::IsSpawnScript(TranslationStages[ActiveStageIdx].ScriptUsage))
+	{
+		for (int32 ParamMapIdx = 0; ParamMapIdx < OtherOutputParamMapHistories.Num(); ParamMapIdx++)
+		{
+			const FNiagaraParameterMapHistory& History = OtherOutputParamMapHistories[ParamMapIdx];
+			for (int32 i = 0; i < History.Variables.Num(); i++)
+			{
+				const FNiagaraVariable& Var = History.Variables[i];
+				const FNiagaraVariable& AliasedVar = History.VariablesWithOriginalAliasesIntact[i];
+				if (History.IsPrimaryDataSetOutput(AliasedVar, TranslationStages[ActiveStageIdx].ScriptUsage) &&
+					!UniqueVars.Contains(Var))
+				{
+					UniqueVars.Add(Var);
+					const UEdGraphPin* DefaultPin = History.GetDefaultValuePin(i);
+					UniqueVarToDefaultPin.Add(Var, DefaultPin);
+					UniqueVarToWriteToParamMap.Add(Var, false);
+				}
+			}
+		}
+
+		for (FNiagaraVariable& Var : UniqueVars)
+		{
+			const UEdGraphPin* DefaultPin = UniqueVarToDefaultPin.FindChecked(Var);
+			bool bWriteToParamMapEntries = UniqueVarToWriteToParamMap.FindChecked(Var);
+			int32 OutputChunkId = INDEX_NONE;
+			HandleParameterRead(ParamMapHistoryIdx, Var, DefaultPin, DefaultPin != nullptr ? Cast<UNiagaraNode>(DefaultPin->GetOwningNode()) : nullptr, OutputChunkId, !bWriteToParamMapEntries);
+
+			if (FNiagaraParameterMapHistory::IsInitialValue(Var))
+			{
+				FNiagaraVariable SourceForInitialValue = FNiagaraParameterMapHistory::GetSourceForInitialValue(Var);
+				if (!UniqueVars.Contains(SourceForInitialValue))
+				{
+					Error(FText::Format(LOCTEXT("MissingInitialValueSource", "Variable {0} is used, but its source variable {1} is not set!"), FText::FromName(Var.GetName()), FText::FromName(SourceForInitialValue.GetName())), nullptr, nullptr);
+				}
+				InitialNamespaceVariablesMissingDefault.Add(Var);
+			}
+		}
+	}
+
+	AddBodyChunk(TEXT(""), TEXT("//End Initialize Parameter Map Defaults"), FNiagaraTypeDefinition::GetIntDef(), false, false);
+}
+
 void FHlslNiagaraTranslator::Output(UNiagaraNodeOutput* OutputNode, const TArray<int32>& ComputedInputs)
 {
 	SCOPE_CYCLE_COUNTER(STAT_NiagaraEditor_HlslTranslator_Output);
@@ -2822,9 +2692,9 @@ int32 FHlslNiagaraTranslator::GetAttribute(const FNiagaraVariable& Attribute)
 		Error(FText::Format(LOCTEXT("GetConstantFail", "Cannot handle type {0}! Variable: {1}"), Attribute.GetType().GetNameText(), FText::FromName(Attribute.GetName())), nullptr, nullptr);
 	}
 
-	if (Script->IsParticleSpawnScript())
+	if (TranslationStages.Num() > 1 && UNiagaraScript::IsParticleSpawnScript(TranslationStages[0].ScriptUsage))
 	{
-		if (Script->IsInterpolatedParticleSpawnScript() && !bInsideInterpolatedSpawnScript)
+		if (ActiveStageIdx > 0)
 		{
 			//This is a special case where we allow the grabbing of attributes in the update section of an interpolated spawn script.
 			//But we return the results of the previously ran spawn script.
@@ -2842,7 +2712,7 @@ int32 FHlslNiagaraTranslator::GetAttribute(const FNiagaraVariable& Attribute)
 	}
 	else
 	{
-		CompilationOutput.DataUsage.bReadsAttriubteData = true;
+		CompilationOutput.ScriptData.DataUsage.bReadsAttributeData = true;
 		int32 Chunk = INDEX_NONE;
 		if (!ParameterMapRegisterUniformAttributeVariable(Attribute, nullptr, 0, Chunk))
 		{
@@ -2923,13 +2793,25 @@ void FHlslNiagaraTranslator::ParameterMapSet(UNiagaraNodeParameterMapSet* SetNod
 			}
 
 			FString VarName = Var.GetName().ToString();
-			if (FNiagaraParameterMapHistory::IsExternalConstantNamespace(Var, Script))
+			if (FNiagaraParameterMapHistory::IsExternalConstantNamespace(Var, CompileOptions.TargetUsage, CompileOptions.GetTargetUsageBitmask()))
 			{
 				Error(FText::Format(LOCTEXT("SetSystemConstantFail", "Cannot Set external constant, Type: {0} Variable: {1}"), Var.GetType().GetNameText(), FText::FromName(Var.GetName())), SetNode, nullptr);
 				continue;
 			}
 
 			Var = ActiveHistoryForFunctionCalls.ResolveAliases(Var);
+			const FNiagaraVariable* ConstantVar = FNiagaraConstants::GetKnownConstant(Var.GetName(), false);
+			if (ConstantVar != nullptr && ConstantVar->GetType() != Var.GetType())
+			{
+				Error(FText::Format(LOCTEXT("MismatchedConstantTypes", "Variable {0} is a system constant, but its type is different! {1} != {2}"), FText::FromName(Var.GetName()),
+					ConstantVar->GetType().GetNameText(), Var.GetType().GetNameText()), nullptr, nullptr);
+			}
+
+			if (FNiagaraConstants::IsEngineManagedAttribute(Var))
+			{
+				Error(FText::Format(LOCTEXT("SettingSystemAttr", "Variable {0} is an engine managed particle attribute and cannot be set directly."), FText::FromName(Var.GetName())), nullptr, nullptr);
+				continue;
+			}
 
 			if (ParamMapHistoryIdx < ParamMapHistories.Num())
 			{
@@ -2941,9 +2823,48 @@ void FHlslNiagaraTranslator::ParameterMapSet(UNiagaraNodeParameterMapSet* SetNod
 				}
 			}
 
-			if (Var.GetType().GetClass() != nullptr)
+			if (Var.IsDataInterface())
 			{
-				// do nothing for now, we've recorded the value for the future.
+				if (CompileOptions.TargetUsage == ENiagaraScriptUsage::ParticleSpawnScriptInterpolated && TranslationStages[ActiveStageIdx].ScriptUsage == ENiagaraScriptUsage::ParticleUpdateScript)
+				{
+					// We don't want to add writes for particle update data interface parameters in both interpolated spawn and update, so skip them when processing the update stage of the 
+					// interpolated spawn script.  We don't skip the writes when compiling the particle update script because it's not recompiled when the interpolated spawn flag is changed
+					// and this would result in missing data interfaces if interpolated spawn was turned off.
+					continue;
+				}
+
+				bool bAllowDataInterfaces = true;
+				if(ParamMapHistoryIdx < ParamMapHistories.Num() && ParamMapHistories[ParamMapHistoryIdx].IsPrimaryDataSetOutput(Var, CompileOptions.TargetUsage, bAllowDataInterfaces))
+				{
+					if (Input < 0 || Input >= CompilationOutput.ScriptData.DataInterfaceInfo.Num())
+					{
+						Error(FText::Format(LOCTEXT("ParameterMapDataInterfaceNotFoundErrorFormat", "Data interface could not be found for parameter map set.  Paramter: {0}"), FText::FromName(Var.GetName())), SetNode, InputPins[i]);
+						continue;
+					}
+
+					FName UsageName;
+					if (FNiagaraParameterMapHistory::IsAliasedEmitterParameter(Var))
+					{
+						FNiagaraVariable AliasedVar = ActiveHistoryForFunctionCalls.ResolveAliases(Var);
+						UsageName = AliasedVar.GetName();
+					}
+					else
+					{
+						UsageName = Var.GetName();
+					}
+
+					FNiagaraScriptDataInterfaceCompileInfo& Info = CompilationOutput.ScriptData.DataInterfaceInfo[Input];
+					if (Info.RegisteredParameterMapWrite == NAME_None)
+					{
+						Info.RegisteredParameterMapWrite = UsageName;
+					}
+					else
+					{
+						Error(FText::Format(LOCTEXT("ExternalDataInterfaceAssignedToMultipleParameters", "The data interface named {0} was added to a parameter map multiple times which isn't supported.  First usage: {1} Invalid usage:{2}"),
+							FText::FromName(Info.Name), FText::FromName(Info.RegisteredParameterMapWrite), FText::FromName(UsageName)), SetNode, InputPins[i]);
+						continue;
+					}
+				}
 			}
 			else
 			{
@@ -2954,14 +2875,50 @@ void FHlslNiagaraTranslator::ParameterMapSet(UNiagaraNodeParameterMapSet* SetNod
 
 }
 
+FString FHlslNiagaraTranslator::GetUniqueEmitterName() const
+{
+	if (CompileOptions.TargetUsage == ENiagaraScriptUsage::SystemSpawnScript || CompileOptions.TargetUsage == ENiagaraScriptUsage::SystemUpdateScript)
+	{
+		return CompileData->GetUniqueEmitterName();
+	}
+	else
+	{
+		return TEXT("Emitter");
+	}
+}
+
 bool FHlslNiagaraTranslator::IsBulkSystemScript() const
 {
-	return (Script->Usage == ENiagaraScriptUsage::SystemSpawnScript || Script->Usage == ENiagaraScriptUsage::SystemUpdateScript) &&
-		!Script->GetName().Contains(TEXT("Solo"));
+	return (CompileOptions.TargetUsage == ENiagaraScriptUsage::SystemSpawnScript || CompileOptions.TargetUsage == ENiagaraScriptUsage::SystemUpdateScript);
+}
+
+bool FHlslNiagaraTranslator::IsSpawnScript() const
+{
+	for (int32 i = 0; i < TranslationStages.Num(); i++)
+	{
+		if (UNiagaraScript::IsSpawnScript(TranslationStages[i].ScriptUsage))
+		{
+			return true;
+		}
+	}
+	return false;
+}
+
+
+bool FHlslNiagaraTranslator::RequiresInterpolation() const
+{
+	for (int32 i = 0; i < TranslationStages.Num(); i++)
+	{
+		if (TranslationStages[i].bInterpolatePreviousParams)
+		{
+			return true;
+		}
+	}
+	return false;
 }
 
 
-bool FHlslNiagaraTranslator::ParameterMapRegisterExternalConstantNamespaceVariable(FNiagaraVariable InVariable, UNiagaraNode* InNode, int32 InParamMapHistoryIdx, int32& Output, UEdGraphPin* InDefaultPin)
+bool FHlslNiagaraTranslator::ParameterMapRegisterExternalConstantNamespaceVariable(FNiagaraVariable InVariable, UNiagaraNode* InNode, int32 InParamMapHistoryIdx, int32& Output, const UEdGraphPin* InDefaultPin)
 {
 	InVariable = ActiveHistoryForFunctionCalls.ResolveAliases(InVariable);
 	FString VarName = InVariable.GetName().ToString();
@@ -2973,9 +2930,7 @@ bool FHlslNiagaraTranslator::ParameterMapRegisterExternalConstantNamespaceVariab
 	if (InVariable.IsValid())
 	{
 		// We don't really want system delta time or inverse system delta time in a spawn script. It leads to trouble.
-		if ((Script->IsInterpolatedParticleSpawnScript() && this->bInsideInterpolatedSpawnScript) 
-			|| (Script->IsParticleSpawnScript() && !Script->IsInterpolatedParticleSpawnScript() && TranslationOptions.SimTarget != ENiagaraSimTarget::GPUComputeSim)
-			)
+		if (TranslationStages.Num() > 0 && UNiagaraScript::IsParticleSpawnScript(TranslationStages[ActiveStageIdx].ScriptUsage))
 		{
 			if (InVariable == SYS_PARAM_ENGINE_DELTA_TIME || InVariable == SYS_PARAM_ENGINE_INV_DELTA_TIME)
 			{
@@ -2993,7 +2948,8 @@ bool FHlslNiagaraTranslator::ParameterMapRegisterExternalConstantNamespaceVariab
 		}
 
 		bool bIsDataInterface = InVariable.GetType().IsDataInterface();
-		bool bIsPerInstanceBulkSystemParam = IsBulkSystemScript() && !bIsDataInterface && (FNiagaraParameterMapHistory::IsUserParameter(InVariable) || FNiagaraParameterMapHistory::IsEngineParameter(InVariable));
+		const FString* EmitterAlias = ActiveHistoryForFunctionCalls.GetEmitterAlias();
+		bool bIsPerInstanceBulkSystemParam = IsBulkSystemScript() && !bIsDataInterface && (FNiagaraParameterMapHistory::IsUserParameter(InVariable) || FNiagaraParameterMapHistory::IsPerInstanceEngineParameter(InVariable, EmitterAlias != nullptr ? *EmitterAlias : TEXT("Emitter")));
 
 		if (!bIsPerInstanceBulkSystemParam)
 		{
@@ -3013,12 +2969,13 @@ bool FHlslNiagaraTranslator::ParameterMapRegisterExternalConstantNamespaceVariab
 					}
 					else
 					{
-						UObject* Obj = const_cast<UClass*>(InVariable.GetType().GetClass())->GetDefaultObject(true);
-						DataInterface = Cast<UNiagaraDataInterface>(DuplicateObject(Obj, GetTransientPackage()));
+						UObject*const* Obj = CompileData->CDOs.Find(const_cast<UClass*>(InVariable.GetType().GetClass()));
+						check(Obj != nullptr);
+						DataInterface = CastChecked<UNiagaraDataInterface>(*Obj);
 					}
 					if (ensure(DataInterface))
 					{
-						Output = RegisterDataInterface(InVariable, DataInterface, true);
+						Output = RegisterDataInterface(InVariable, DataInterface, true, true);
 						return true;
 					}
 				}
@@ -3041,7 +2998,7 @@ bool FHlslNiagaraTranslator::ParameterMapRegisterExternalConstantNamespaceVariab
 					Error(FText::Format(LOCTEXT("GetParameterUnsetParam", "Variable {0} hasn't had its default value set. Required Bytes: {1} vs Allocated Bytes: {2}"), FText::FromName(InVariable.GetName()), FText::AsNumber(InVariable.GetType().GetSize()), FText::AsNumber(InVariable.GetSizeInBytes())), nullptr, nullptr);
 				}
 
-				CompilationOutput.Parameters.SetOrAdd(InVariable);
+				CompilationOutput.ScriptData.Parameters.SetOrAdd(InVariable);
 				UniformIdx = ChunksByMode[(int32)ENiagaraCodeChunkMode::Uniform].Num();
 				UniformChunk = AddUniformChunk(SymbolNameDefined, InVariable.GetType());
 				ParamMapDefinedSystemVarsToUniformChunks.Add(InVariable.GetName(), UniformIdx);
@@ -3125,17 +3082,11 @@ bool FHlslNiagaraTranslator::ParameterMapRegisterNamespaceAttributeVariable(cons
 
 FString FHlslNiagaraTranslator::GetParameterMapInstanceName(int32 ParamMapHistoryIdx)
 {
-	FString ParameterMapInstanceName = TEXT("Context.Map");
-	if (Script->IsInterpolatedParticleSpawnScript())
+	FString ParameterMapInstanceName;
+
+	if (TranslationStages.Num() > ActiveStageIdx)
 	{
-		if (bInsideInterpolatedSpawnScript)
-		{
-			ParameterMapInstanceName = TEXT("Context.MapSpawn");
-		}
-		else
-		{
-			ParameterMapInstanceName = TEXT("Context.MapUpdate");
-		}
+		ParameterMapInstanceName = FString::Printf(TEXT("Context.%s"), *TranslationStages[ActiveStageIdx].PassNamespace);
 	}
 	return ParameterMapInstanceName;
 }
@@ -3148,7 +3099,7 @@ void FHlslNiagaraTranslator::Emitter(class UNiagaraNodeEmitter* EmitterNode, TAr
 	UNiagaraScriptSource* Source = EmitterNode->GetScriptSource();
 	if (Source == nullptr)
 	{
-		Error(LOCTEXT("FunctionCallNonexistant", "Emitter call missing ScriptSource"), EmitterNode, nullptr);
+		Error(LOCTEXT("FunctionCallNonexistantScriptSource", "Emitter call missing ScriptSource"), EmitterNode, nullptr);
 		return;
 	}
 
@@ -3270,221 +3221,307 @@ void FHlslNiagaraTranslator::ParameterMapGet(UNiagaraNodeParameterMapGet* GetNod
 		}
 		else // These are the pins that we are getting off the parameter map.
 		{
-			FNiagaraVariable Var = Schema->PinToNiagaraVariable(OutputPins[i], true);
+			FNiagaraTypeDefinition OutputTypeDefinition = Schema->PinToTypeDefinition(OutputPins[i]);
+			bool bNeedsValue =
+				OutputTypeDefinition != FNiagaraTypeDefinition::GetParameterMapDef() &&
+				OutputTypeDefinition.IsDataInterface() == false;
+			FNiagaraVariable Var = Schema->PinToNiagaraVariable(OutputPins[i], bNeedsValue);
 
-			if (!AddStructToDefinitionSet(Var.GetType()))
-			{
-				Error(FText::Format(LOCTEXT("ParameterMapGetTypeError", "Cannot handle type {0}! Variable: {1}"), Var.GetType().GetNameText(), FText::FromName(Var.GetName())), nullptr, nullptr);
-			}
+			HandleParameterRead(ParamMapHistoryIdx, Var, GetNode->GetDefaultPin(OutputPins[i]), GetNode, Outputs[i]);
+		}
+	}
+}
 
-			// If this is a System parameter, just wire in the system appropriate system attribute.
-			FString VarName = Var.GetName().ToString();
-			FString SymbolName = GetSanitizedSymbolName(VarName);
+void FHlslNiagaraTranslator::HandleParameterRead(int32 ParamMapHistoryIdx, const FNiagaraVariable& InVar, const UEdGraphPin* DefaultPin, UNiagaraNode* ErrorNode, int32& OutputChunkId, bool bTreatAsUnknownParameterMap)
+{
+	FString ParameterMapInstanceName = GetParameterMapInstanceName(ParamMapHistoryIdx);
+	FNiagaraVariable Var = InVar;
+	if (!AddStructToDefinitionSet(Var.GetType()))
+	{
+		Error(FText::Format(LOCTEXT("ParameterMapGetTypeError", "Cannot handle type {0}! Variable: {1}"), Var.GetType().GetNameText(), FText::FromName(Var.GetName())), nullptr, nullptr);
+	}
+
+	// If this is a System parameter, just wire in the system appropriate system attribute.
+	FString VarName = Var.GetName().ToString();
+	FString SymbolName = GetSanitizedSymbolName(VarName);
+
+	bool bIsPerInstanceAttribute = false;
+	bool bIsCandidateForRapidIteration = false;
+	const UEdGraphPin* InputPin = DefaultPin;
+
+	if (FNiagaraParameterMapHistory::IsExternalConstantNamespace(Var, CompileOptions.TargetUsage, CompileOptions.GetTargetUsageBitmask()))
+	{
+		if (ParameterMapRegisterExternalConstantNamespaceVariable(Var, ErrorNode, ParamMapHistoryIdx, OutputChunkId, DefaultPin))
+		{
+			return;
+		}
+	}
+	else if (FNiagaraParameterMapHistory::IsAliasedModuleParameter(Var) && ActiveHistoryForFunctionCalls.InTopLevelFunctionCall(CompileOptions.TargetUsage))
+	{
+		if (InputPin != nullptr && InputPin->LinkedTo.Num() == 0 && Var.GetType() != FNiagaraTypeDefinition::GetBoolDef() && !Var.GetType().IsEnum() && !Var.GetType().IsDataInterface())
+		{
+			bIsCandidateForRapidIteration = true;
+		}
+	}
+
+	FNiagaraParameterMapHistory& History = ParamMapHistories[ParamMapHistoryIdx];
+	bool bWasEmitterAliased = FNiagaraParameterMapHistory::IsAliasedEmitterParameter(Var);
+	Var = ActiveHistoryForFunctionCalls.ResolveAliases(Var);
+
+	const FNiagaraVariable* ConstantVar = FNiagaraConstants::GetKnownConstant(Var.GetName(), false);
+	if (ConstantVar != nullptr && ConstantVar->GetType() != Var.GetType())
+	{
+		Error(FText::Format(LOCTEXT("MismatchedConstantTypes", "Variable {0} is a system constant, but its type is different! {1} != {2}"), FText::FromName(Var.GetName()),
+			ConstantVar->GetType().GetNameText(), Var.GetType().GetNameText()), ErrorNode, nullptr);
+	}
+
+	if (History.IsPrimaryDataSetOutput(Var, GetTargetUsage())) // Note that data interfaces aren't ever in the primary data set even if the namespace matches.
+	{
+		bIsPerInstanceAttribute = true;
+	}
 
-			bool bIsPerInstanceAttribute = false;
-			bool bIsCandidateForRapidIteration = false;
-			UEdGraphPin* InputPin = GetNode->GetDefaultPin(OutputPins[i]);
+	int32 LastSetChunkIdx = INDEX_NONE;
+	if (ParamMapHistoryIdx < ParamMapHistories.Num())
+	{
+		// See if we've written this variable before, if so we can reuse the index
+		int32 VarIdx = ParamMapHistories[ParamMapHistoryIdx].FindVariableByName(Var.GetName());
+		if (VarIdx != INDEX_NONE && VarIdx < ParamMapSetVariablesToChunks[ParamMapHistoryIdx].Num())
+		{
+			LastSetChunkIdx = ParamMapSetVariablesToChunks[ParamMapHistoryIdx][VarIdx];
+		}
+
+		// Check to see if this is the first time we've encountered this node and it is a viable candidate for rapid iteration
+		if (LastSetChunkIdx == INDEX_NONE && bIsCandidateForRapidIteration && TranslationOptions.bParameterRapidIteration)
+		{
+			bool bVarChanged = false;
+			if (!bWasEmitterAliased && ActiveHistoryForFunctionCalls.GetEmitterAlias() != nullptr)
+			{
+				Var = FNiagaraParameterMapHistory::ConvertVariableToRapidIterationConstantName(Var, *(*ActiveHistoryForFunctionCalls.GetEmitterAlias()), GetTargetUsage());
+				bVarChanged = true;
+			}
+			else if (UNiagaraScript::IsSystemSpawnScript(CompileOptions.TargetUsage) || UNiagaraScript::IsSystemUpdateScript(CompileOptions.TargetUsage))
+			{
+				Var = FNiagaraParameterMapHistory::ConvertVariableToRapidIterationConstantName(Var, nullptr, GetTargetUsage());
+				bVarChanged = true;
+			}
 
-			if (FNiagaraParameterMapHistory::IsExternalConstantNamespace(Var, Script))
+			// Now try to look up with the new name.. we may have already made this an external variable before..
+			if (bVarChanged)
 			{
-				if (ParameterMapRegisterExternalConstantNamespaceVariable(Var, GetNode, ParamMapHistoryIdx, Outputs[i], GetNode->GetDefaultPin(OutputPins[i])))
+				VarIdx = ParamMapHistories[ParamMapHistoryIdx].FindVariableByName(Var.GetName());
+				if (VarIdx != INDEX_NONE && VarIdx < ParamMapSetVariablesToChunks[ParamMapHistoryIdx].Num())
 				{
-					continue;
+					LastSetChunkIdx = ParamMapSetVariablesToChunks[ParamMapHistoryIdx][VarIdx];
 				}
 			}
-			else if (FNiagaraParameterMapHistory::IsAliasedModuleParameter(Var) && ActiveHistoryForFunctionCalls.InTopLevelFunctionCall(Script->Usage))
+
+			// If it isn't found yet.. go ahead and make it into a constant variable..
+			if (LastSetChunkIdx == INDEX_NONE && ParameterMapRegisterExternalConstantNamespaceVariable(Var, ErrorNode, ParamMapHistoryIdx, OutputChunkId, InputPin))
 			{
-				if (InputPin != nullptr && InputPin->LinkedTo.Num() == 0 && Var.GetType() != FNiagaraTypeDefinition::GetBoolDef() && !Var.GetType().IsEnum() && !Var.GetType().IsDataInterface())
+				LastSetChunkIdx = OutputChunkId;
+				if (VarIdx != INDEX_NONE && VarIdx < ParamMapSetVariablesToChunks[ParamMapHistoryIdx].Num())
 				{
-					bIsCandidateForRapidIteration = true;
+					// Record that we wrote to it.
+					ParamMapSetVariablesToChunks[ParamMapHistoryIdx][VarIdx] = LastSetChunkIdx;
+					ParamMapDefinedAttributesToNamespaceVars.FindOrAdd(Var.GetName()) = Var;
 				}
+				return;
 			}
+		}
 
-			FNiagaraParameterMapHistory& History = ParamMapHistories[ParamMapHistoryIdx];
-			bool bWasEmitterAliased = FNiagaraParameterMapHistory::IsAliasedEmitterParameter(Var);
-			Var = ActiveHistoryForFunctionCalls.ResolveAliases(Var);
-			if (History.IsPrimaryDataSetOutput(Var, GetTargetUsage()))
+		// We have yet to write to this parameter, use the default value if specified and the parameter 
+		// isn't a per-particle value.
+		bool bIgnoreDefaultValue = ParamMapHistories[ParamMapHistoryIdx].ShouldIgnoreVariableDefault(Var);
+		if (bIsPerInstanceAttribute)
+		{
+			FNiagaraVariable* ExistingVar = ParamMapDefinedAttributesToNamespaceVars.Find(Var.GetName());
+			bool ExistsInAttribArrayAlready = ExistingVar != nullptr;
+			if (ExistsInAttribArrayAlready&& ExistingVar->GetType() != Var.GetType())
 			{
-				bIsPerInstanceAttribute = true;
+				Error(FText::Format(LOCTEXT("Mismatched Types", "Variable {0} was defined earlier, but its type is different! {1} != {2}"), FText::FromName(Var.GetName()),
+					ExistingVar->GetType().GetNameText(), Var.GetType().GetNameText()), ErrorNode, nullptr);
 			}
 
-			int32 LastSetChunkIdx = INDEX_NONE;
-			if (ParamMapHistoryIdx < ParamMapHistories.Num())
+			if ((TranslationStages.Num() > 1 && !UNiagaraScript::IsParticleSpawnScript(TranslationStages[ActiveStageIdx].ScriptUsage) && ExistsInAttribArrayAlready) ||
+				!UNiagaraScript::IsSpawnScript(TranslationStages[ActiveStageIdx].ScriptUsage))
 			{
-				// See if we've written this variable before, if so we can reuse the index
-				int32 VarIdx = ParamMapHistories[ParamMapHistoryIdx].FindVariableByName(Var.GetName());
-				if (VarIdx != INDEX_NONE && VarIdx < ParamMapSetVariablesToChunks[ParamMapHistoryIdx].Num())
+				bIgnoreDefaultValue = true;
+			}
+		}
+
+		if (LastSetChunkIdx == INDEX_NONE && (UNiagaraScript::IsSpawnScript(TranslationStages[ActiveStageIdx].ScriptUsage)))
+		{
+			if (FNiagaraParameterMapHistory::IsInitialValue(Var))
+			{
+				FNiagaraVariable SourceForInitialValue = FNiagaraParameterMapHistory::GetSourceForInitialValue(Var);
+				bool bFoundExistingSet = false;
+				for (int32 OtherParamIdx = 0; OtherParamIdx < OtherOutputParamMapHistories.Num(); OtherParamIdx++)
 				{
-					LastSetChunkIdx = ParamMapSetVariablesToChunks[ParamMapHistoryIdx][VarIdx];
+					if (INDEX_NONE != OtherOutputParamMapHistories[OtherParamIdx].FindVariableByName(SourceForInitialValue.GetName()))
+					{
+						bFoundExistingSet = true;
+					}
 				}
 
-				// Check to see if this is the first time we've encountered this node and it is a viable candidate for rapid iteration
-				if (LastSetChunkIdx == INDEX_NONE && bIsCandidateForRapidIteration && TranslationOptions.bParameterRapidIteration)
+				if (bFoundExistingSet)
+				{
+					LastSetChunkIdx = AddBodyChunk(ParameterMapInstanceName + TEXT(".") + GetSanitizedSymbolName(Var.GetName().ToString()),
+						ParameterMapInstanceName + TEXT(".") + GetSanitizedSymbolName(SourceForInitialValue.GetName().ToString()), Var.GetType(), false);
+					ParamMapDefinedAttributesToNamespaceVars.FindOrAdd(Var.GetName()) = Var;
+				}
+				else
+				{
+					Error(FText::Format(LOCTEXT("MissingInitialValueSource", "Variable {0} is used, but its source variable {1} is not set!"), FText::FromName(Var.GetName()), FText::FromName(SourceForInitialValue.GetName())), nullptr, nullptr);
+				}
+			}
+		}
+
+		if (LastSetChunkIdx == INDEX_NONE && !bIgnoreDefaultValue)
+		{
+			// Default was found, trace back its inputs.
+			if (InputPin != nullptr)
+			{
+				// Check to see if there are any overrides passed in to the translator. This allows us to bake in rapid iteration variables for performance.
+				if (InputPin->LinkedTo.Num() == 0 && bIsCandidateForRapidIteration && !TranslationOptions.bParameterRapidIteration)
 				{
 					bool bVarChanged = false;
+					FNiagaraVariable RapidIterationConstantVar;
 					if (!bWasEmitterAliased && ActiveHistoryForFunctionCalls.GetEmitterAlias() != nullptr)
 					{
-						Var = FNiagaraParameterMapHistory::ConvertVariableToRapidIterationConstantName(Var, *(*ActiveHistoryForFunctionCalls.GetEmitterAlias()));
+						RapidIterationConstantVar = FNiagaraParameterMapHistory::ConvertVariableToRapidIterationConstantName(Var, *(*ActiveHistoryForFunctionCalls.GetEmitterAlias()), GetTargetUsage());
 						bVarChanged = true;
 					}
-					else if (Script->IsSystemSpawnScript() || Script->IsSystemUpdateScript())
+					else if (UNiagaraScript::IsSystemSpawnScript(CompileOptions.TargetUsage) || UNiagaraScript::IsSystemUpdateScript(CompileOptions.TargetUsage))
 					{
-						Var = FNiagaraParameterMapHistory::ConvertVariableToRapidIterationConstantName(Var, nullptr);
+						RapidIterationConstantVar = FNiagaraParameterMapHistory::ConvertVariableToRapidIterationConstantName(Var, nullptr, GetTargetUsage());
 						bVarChanged = true;
 					}
 
-					// Now try to look up with the new name.. we may have already made this an external variable before..
-					if (bVarChanged)
-					{
-						VarIdx = ParamMapHistories[ParamMapHistoryIdx].FindVariableByName(Var.GetName());
-						if (VarIdx != INDEX_NONE && VarIdx < ParamMapSetVariablesToChunks[ParamMapHistoryIdx].Num())
-						{
-							LastSetChunkIdx = ParamMapSetVariablesToChunks[ParamMapHistoryIdx][VarIdx];
-						}
-					}
-
-					// If it isn't found yet.. go ahead and make it into a constant variable..
-					if (LastSetChunkIdx == INDEX_NONE && ParameterMapRegisterExternalConstantNamespaceVariable(Var, GetNode, ParamMapHistoryIdx, Outputs[i], InputPin))
-					{
-						LastSetChunkIdx = Outputs[i];
-						if (VarIdx != INDEX_NONE && VarIdx < ParamMapSetVariablesToChunks[ParamMapHistoryIdx].Num())
-						{
-							// Record that we wrote to it.
-							ParamMapSetVariablesToChunks[ParamMapHistoryIdx][VarIdx] = LastSetChunkIdx;
-							ParamMapDefinedAttributesToNamespaceVars.FindOrAdd(Var.GetName()) = Var;
-						}
-						continue;
-					}
-				}
-
-				// We have yet to write to this parameter, use the default value if specified and the parameter 
-				// isn't a per-particle value.
-				bool bIgnoreDefaultValue = false;
-				if (bIsPerInstanceAttribute)
-				{
-					if ((Script->IsInterpolatedParticleSpawnScript() && bInsideInterpolatedSpawnScript == false) 
-						|| ((TranslationOptions.SimTarget == ENiagaraSimTarget::GPUComputeSim && bInsideInterpolatedSpawnScript == false) || !Script->IsSpawnScript()))
-					{
-						bIgnoreDefaultValue = true;
-					}
-				}
-
-				if (LastSetChunkIdx == INDEX_NONE && (Script->IsInterpolatedParticleSpawnScript() || Script->IsParticleSpawnScript() || Script->IsEmitterSpawnScript() || Script->IsSystemSpawnScript()))
-				{
-					if (FNiagaraParameterMapHistory::IsInitialValue(Var))
+					int32 FoundIdx = TranslationOptions.OverrideModuleConstants.Find(RapidIterationConstantVar);
+					if (FoundIdx != INDEX_NONE)
 					{
-						FNiagaraVariable SourceForInitialValue = FNiagaraParameterMapHistory::GetSourceForInitialValue(Var);
-						bool bFoundExistingSet = false;
-						for (int32 OtherParamIdx = 0; OtherParamIdx < OtherOutputParamMapHistories.Num(); OtherParamIdx++)
-						{
-							if (INDEX_NONE != OtherOutputParamMapHistories[OtherParamIdx].FindVariableByName(SourceForInitialValue.GetName()))
-							{
-								bFoundExistingSet = true;
-							}
-						}
-
-						if (bFoundExistingSet)
-						{
-							LastSetChunkIdx = AddBodyChunk(ParameterMapInstanceName + TEXT(".") + GetSanitizedSymbolName(Var.GetName().ToString()),
-								ParameterMapInstanceName + TEXT(".") + GetSanitizedSymbolName(SourceForInitialValue.GetName().ToString()), Var.GetType(), false);
-							ParamMapDefinedAttributesToNamespaceVars.FindOrAdd(Var.GetName()) = Var;
-						}
-						else
-						{
-							Error(FText::Format(LOCTEXT("MissingInitialValueSource", "Variable {0} is used, but its source variable {1} is not set!"), FText::FromName(Var.GetName()), FText::FromName(SourceForInitialValue.GetName())), nullptr, nullptr);
-						}
+						OutputChunkId = GetConstant(TranslationOptions.OverrideModuleConstants[FoundIdx]);
+						return;
 					}
 				}
 
-				if (LastSetChunkIdx == INDEX_NONE && !bIgnoreDefaultValue)
+				CurrentDefaultPinTraversal.Push(InputPin);
+				if (InputPin->LinkedTo.Num() != 0 && InputPin->LinkedTo[0] != nullptr)
 				{
-					// Default was found, trace back its inputs.
-					if (InputPin != nullptr )
-					{						
-						// Check to see if there are any overrides passed in to the translator. This allows us to bake in rapid iteration variables for performance.
-						if (InputPin->LinkedTo.Num() == 0 && bIsCandidateForRapidIteration && !TranslationOptions.bParameterRapidIteration)
-						{
-							bool bVarChanged = false;
-							FNiagaraVariable RapidIterationConstantVar;
-							if (!bWasEmitterAliased && ActiveHistoryForFunctionCalls.GetEmitterAlias() != nullptr)
-							{
-								RapidIterationConstantVar = FNiagaraParameterMapHistory::ConvertVariableToRapidIterationConstantName(Var, *(*ActiveHistoryForFunctionCalls.GetEmitterAlias()));
-								bVarChanged = true;
-							}
-							else if (Script->IsSystemSpawnScript() || Script->IsSystemUpdateScript())
-							{
-								RapidIterationConstantVar = FNiagaraParameterMapHistory::ConvertVariableToRapidIterationConstantName(Var, nullptr);
-								bVarChanged = true;
-							}
-
-							int32 FoundIdx = TranslationOptions.OverrideModuleConstants.Find(RapidIterationConstantVar);
-							if (FoundIdx != INDEX_NONE)
-							{
-								Outputs[i] = GetConstant(TranslationOptions.OverrideModuleConstants[FoundIdx]);
-								continue;
-							}
-						}
-
-						LastSetChunkIdx = CompilePin(InputPin);
-
-						if (!Var.IsDataInterface())
-						{
-							if (VarIdx != INDEX_NONE && VarIdx < ParamMapSetVariablesToChunks[ParamMapHistoryIdx].Num())
-							{
-								// Record that we wrote to it.
-								ParamMapSetVariablesToChunks[ParamMapHistoryIdx][VarIdx] = LastSetChunkIdx;	
-								ParamMapDefinedAttributesToNamespaceVars.FindOrAdd(Var.GetName()) = Var;
-							}
-							else
-							{
-								Error(FText::Format(LOCTEXT("NoVarDefaultFound", "Default found for {0}, but not found in ParameterMap traversal"), FText::FromName(Var.GetName())), GetNode, nullptr);
-							}
-
-							// Actually insert the text that sets the default value
-							if (LastSetChunkIdx != INDEX_NONE)
+					// Double-check to make sure that we are connected to a TRANSLATOR_PARAM_BEGIN_DEFAULTS input node rather than
+					// a normal parameter-based parameter map input node to ensure that we don't get into weird traversals.
+					TArray<UNiagaraNode*> Nodes;
+					UNiagaraGraph::BuildTraversal(Nodes, Cast<UNiagaraNode>(InputPin->LinkedTo[0]->GetOwningNode()));
+					for (UNiagaraNode* Node : Nodes)
+					{
+						UNiagaraNodeInput* InputNode = Cast<UNiagaraNodeInput>(Node);
+						if (InputNode && InputNode->Input.GetType() == FNiagaraTypeDefinition::GetParameterMapDef())
+						{
+							if (InputNode->Usage != ENiagaraInputNodeUsage::TranslatorConstant)
 							{
-								if (Var.GetType().GetClass() == nullptr) // Only need to do this wiring for things that aren't data interfaces.
-								{
-									AddBodyChunk(ParameterMapInstanceName + TEXT(".") + GetSanitizedSymbolName(Var.GetName().ToString()), TEXT("{0}"), Var.GetType(), LastSetChunkIdx, false);
-								}
+								Error(FText::Format(LOCTEXT("InvalidParamMapStartForDefaultPin", "Default found for {0}, but the parameter map source for default pins needs to be a {1} node, not a generic input node."), 
+									FText::FromName(Var.GetName()),
+									FText::FromName(TRANSLATOR_PARAM_BEGIN_DEFAULTS.GetName())), ErrorNode, nullptr);
 							}
 						}
 					}
 				}
+				LastSetChunkIdx = CompilePin(InputPin);
+				CurrentDefaultPinTraversal.Pop();
+			}
+			else 
+			{
+				LastSetChunkIdx = GetConstant(Var);
 			}
 
-			// If we are of a data interface, we should output the data interface registration index, otherwise output
-			// the map namespace that we're writing to.
-			if (Var.IsDataInterface())
+			if (!Var.IsDataInterface() && LastSetChunkIdx != INDEX_NONE)
 			{
-				// In order for a module to compile successfully, we potentially need to generate default values
-				// for variables encountered without ever being set. We do this by creating an instance of the CDO.
-				if (Script->IsStandaloneScript() && LastSetChunkIdx == INDEX_NONE)
+				if (bTreatAsUnknownParameterMap == false)
 				{
-					UObject* Obj = const_cast<UClass*>(Var.GetType().GetClass())->GetDefaultObject(true);
-					UNiagaraDataInterface* DataInterface = Cast<UNiagaraDataInterface>(DuplicateObject(Obj, GetTransientPackage()));
-					if (DataInterface)
+					if (VarIdx != INDEX_NONE && VarIdx < ParamMapSetVariablesToChunks[ParamMapHistoryIdx].Num())
 					{
-						LastSetChunkIdx = RegisterDataInterface(Var, DataInterface, true);
+						// Record that we wrote to it.
+						ParamMapSetVariablesToChunks[ParamMapHistoryIdx][VarIdx] = LastSetChunkIdx;
+						ParamMapDefinedAttributesToNamespaceVars.FindOrAdd(Var.GetName()) = Var;
+					}
+					else
+					{
+						Error(FText::Format(LOCTEXT("NoVarDefaultFound", "Default found for {0}, but not found in ParameterMap traversal"), FText::FromName(Var.GetName())), ErrorNode, nullptr);
+					}
+				}
+
+				// Actually insert the text that sets the default value
+				if (LastSetChunkIdx != INDEX_NONE)
+				{
+					if (Var.GetType().GetClass() == nullptr) // Only need to do this wiring for things that aren't data interfaces.
+					{
+						AddBodyChunk(ParameterMapInstanceName + TEXT(".") + GetSanitizedSymbolName(Var.GetName().ToString()), TEXT("{0}"), Var.GetType(), LastSetChunkIdx, false);
 					}
 				}
-				
-				Outputs[i] = LastSetChunkIdx;
 			}
-			else
+
+			if (LastSetChunkIdx == INDEX_NONE && VarIdx != INDEX_NONE && Var.IsDataInterface())
+			{
+				// If this variable is a data interface and it's in the parameter map, but hasn't been set yet, then is is an external data interface so try to register it.
+				if (ParameterMapRegisterExternalConstantNamespaceVariable(Var, ErrorNode, ParamMapHistoryIdx, OutputChunkId, DefaultPin))
+				{
+					return;
+				}
+			}
+		}
+	}
+
+	// If we are of a data interface, we should output the data interface registration index, otherwise output
+	// the map namespace that we're writing to.
+	if (Var.IsDataInterface())
+	{
+		// In order for a module to compile successfully, we potentially need to generate default values
+		// for variables encountered without ever being set. We do this by creating an instance of the CDO.
+		if (UNiagaraScript::IsStandaloneScript(CompileOptions.TargetUsage) && LastSetChunkIdx == INDEX_NONE)
+		{
+			UObject*const* Obj = CompileData->CDOs.Find(const_cast<UClass*>(Var.GetType().GetClass()));
+			check(Obj != nullptr);
+			UNiagaraDataInterface* DataInterface = CastChecked<UNiagaraDataInterface>(*Obj);
+			if (DataInterface)
 			{
-				Outputs[i] = AddSourceChunk(ParameterMapInstanceName + TEXT(".") + GetSanitizedSymbolName(Var.GetName().ToString()), Var.GetType());
-				ParamMapDefinedAttributesToNamespaceVars.FindOrAdd(Var.GetName()) = Var;
+				LastSetChunkIdx = RegisterDataInterface(Var, DataInterface, true, false);
 			}
 		}
+
+		OutputChunkId = LastSetChunkIdx;
+	}
+	else
+	{
+		OutputChunkId = AddSourceChunk(ParameterMapInstanceName + TEXT(".") + GetSanitizedSymbolName(Var.GetName().ToString()), Var.GetType());
+		ParamMapDefinedAttributesToNamespaceVars.FindOrAdd(Var.GetName()) = Var;
 	}
 }
 
 
+
 void FHlslNiagaraTranslator::ReadDataSet(const FNiagaraDataSetID DataSet, const TArray<FNiagaraVariable>& Variables, ENiagaraDataSetAccessMode AccessMode, int32 InputChunk, TArray<int32>& Outputs)
 {
 	//Eventually may allow events that take in a direct index or condition but for now we don't
-	check(InputChunk == INDEX_NONE);
+	int32 ParamMapHistoryIdx = InputChunk;
+
+	if (ParamMapHistoryIdx == -1)
+	{
+		Error(LOCTEXT("NoParamMapIdx", "Cannot find parameter map for input to ReadDataSet!"), nullptr, nullptr);
+		for (int32 i = 0; i < Outputs.Num(); i++)
+		{
+			Outputs[i] = INDEX_NONE;
+			return;
+		}
+	}
+	else if (ParamMapHistoryIdx >= ParamMapHistories.Num())
+	{
+		Error(FText::Format(LOCTEXT("InvalidParamMapIdx", "Invalid parameter map index for ReadDataSet input {0} of {1}!"), ParamMapHistoryIdx, ParamMapHistories.Num()), nullptr, nullptr);
+		for (int32 i = 0; i < Outputs.Num(); i++)
+		{
+			Outputs[i] = INDEX_NONE;
+			return;
+		}
+	}
 
 	TMap<int32, FDataSetAccessInfo>& Reads = DataSetReadInfo[(int32)AccessMode].FindOrAdd(DataSet);
 	FDataSetAccessInfo* DataSetReadForInput = Reads.Find(InputChunk);
@@ -3507,19 +3544,45 @@ void FHlslNiagaraTranslator::ReadDataSet(const FNiagaraDataSetID DataSet, const
 			}
 			DataSetReadForInput->CodeChunks.Add(AddSourceChunk(DataSetAccessSymbol + TEXT(".") + GetSanitizedSymbolName(Var.GetName().ToString()), Var.GetType()));
 		}
-		Outputs = DataSetReadForInput->CodeChunks;
+		Outputs.Add(ParamMapHistoryIdx);
+		Outputs.Append(DataSetReadForInput->CodeChunks);
 	}
 	else
 	{
 		check(Variables.Num() == DataSetReadForInput->Variables.Num());
-		Outputs = DataSetReadForInput->CodeChunks;
+		Outputs.Add(ParamMapHistoryIdx);
+		Outputs.Append(DataSetReadForInput->CodeChunks);
 	}
 }
 
-void FHlslNiagaraTranslator::WriteDataSet(const FNiagaraDataSetID DataSet, const TArray<FNiagaraVariable>& Variables, ENiagaraDataSetAccessMode AccessMode, const TArray<int32>& Inputs)
+void FHlslNiagaraTranslator::WriteDataSet(const FNiagaraDataSetID DataSet, const TArray<FNiagaraVariable>& Variables, ENiagaraDataSetAccessMode AccessMode, const TArray<int32>& Inputs, TArray<int32>& Outputs)
 {
-	int32 ConditionalChunk = Inputs[0];
-	int32 InputChunk = Inputs[1];
+	int32 ParamMapHistoryIdx = Inputs[0];
+	int32 ConditionalChunk = Inputs[1];
+	int32 InputChunk = Inputs[2];
+	Outputs.SetNum(1);
+	Outputs[0] = ParamMapHistoryIdx;
+
+	if (ParamMapHistoryIdx == -1)
+	{
+		Error(LOCTEXT("NoParamMapIdx", "Cannot find parameter map for input to WriteDataSet!"), nullptr, nullptr);
+		for (int32 i = 0; i < Outputs.Num(); i++)
+		{
+			Outputs[i] = INDEX_NONE;
+			return;
+		}
+	}
+	else if (ParamMapHistoryIdx >= ParamMapHistories.Num())
+	{
+		Error(FText::Format(LOCTEXT("InvalidParamMapIdx", "Invalid parameter map index for WriteDataSet input {0} of {1}!"), ParamMapHistoryIdx, ParamMapHistories.Num()), nullptr, nullptr);
+		for (int32 i = 0; i < Outputs.Num(); i++)
+		{
+			Outputs[i] = INDEX_NONE;
+			return;
+		}
+	}
+
+
 	TMap<int32, FDataSetAccessInfo>& Writes = DataSetWriteInfo[(int32)AccessMode].FindOrAdd(DataSet);
 	FDataSetAccessInfo* DataSetWriteForInput = Writes.Find(InputChunk);
 	int32& DataSetWriteConditionalIndex = DataSetWriteConditionalInfo[(int32)AccessMode].FindOrAdd(DataSet);
@@ -3544,57 +3607,84 @@ void FHlslNiagaraTranslator::WriteDataSet(const FNiagaraDataSetID DataSet, const
 	//First chunk for a write is always the condition pin.
 	//We always write the event payload into the temp storage but we can access this condition to pass to the final actual write to the buffer.
 	
-	DataSetWriteForInput->CodeChunks.Add(AddBodyChunk(DataSetAccessName + TEXT("_Valid"), TEXT("{0}"), FNiagaraTypeDefinition::GetBoolDef(), Inputs[0], false));
+	DataSetWriteForInput->CodeChunks.Add(AddBodyChunk(DataSetAccessName + TEXT("_Valid"), TEXT("{0}"), FNiagaraTypeDefinition::GetBoolDef(), Inputs[1], false));
 	for (int32 i = 0; i < Variables.Num(); ++i)
 	{
 		const FNiagaraVariable& Var = Variables[i];
-		int32 Input = Inputs[i + 1];//input 0 is the valid input (no entry in variables array), so we need of offset all other inputs by 1.
+		int32 Input = Inputs[i + 2];//input 0 is the valid input (no entry in variables array), so we need of offset all other inputs by 1.
 		DataSetWriteForInput->CodeChunks.Add(AddBodyChunk(DataSetAccessName + TEXT(".") + GetSanitizedSymbolName(Var.GetName().ToString()), TEXT("{0}"), Var.GetType(), Input, false));
 	}
 
 }
 
-int32 FHlslNiagaraTranslator::RegisterDataInterface(FNiagaraVariable& Var, UNiagaraDataInterface* DataInterface, bool bPlaceholder)
+int32 FHlslNiagaraTranslator::RegisterDataInterface(FNiagaraVariable& Var, UNiagaraDataInterface* DataInterface, bool bPlaceholder, bool bAddParameterMapRead)
 {
+	FString Id = DataInterface ? *DataInterface->GetMergeId().ToString() : TEXT("??");
+	FString PathName = DataInterface ? *DataInterface->GetPathName() : TEXT("XX");
+	//UE_LOG(LogNiagaraEditor, Log, TEXT("RegisterDataInterface %s %s %s '%s' %s"), *Var.GetName().ToString(), *Var.GetType().GetName(), *Id, *PathName,
+	//	bPlaceholder ? TEXT("bPlaceholder=true") : TEXT("bPlaceholder=false"));
+
 	int32 FuncParam;
 	if (GetFunctionParameter(Var, FuncParam))
 	{
 		if (FuncParam != INDEX_NONE)
 		{
 			//This data interface param has been overridden by the function call so use that index.	
+			UE_LOG(LogNiagaraEditor, Log, TEXT("RegisterDataInterface is funcParam"));
 			return FuncParam;
 		}
 	}
-	
+
 	//If we get here then this is a new data interface.
-	int32 Idx = CompilationOutput.DataInterfaceInfo.IndexOfByPredicate([&](const FNiagaraScriptDataInterfaceInfo& OtherInfo)
+	int32 Idx = CompilationOutput.ScriptData.DataInterfaceInfo.IndexOfByPredicate([&](const FNiagaraScriptDataInterfaceCompileInfo& OtherInfo)
 	{
 		return OtherInfo.Name == Var.GetName();
 	});
 	if (Idx == INDEX_NONE)
 	{
-		Idx = CompilationOutput.DataInterfaceInfo.AddDefaulted();
-		CompilationOutput.DataInterfaceInfo[Idx].DataInterface = DataInterface;
-		CompilationOutput.DataInterfaceInfo[Idx].Name = Var.GetName();
-		CompilationOutput.DataInterfaceInfo[Idx].Type = Var.GetType();
+		FName DataInterfaceName;
+		if (FNiagaraParameterMapHistory::IsAliasedEmitterParameter(Var.GetName().ToString()))
+		{
+			FNiagaraVariable AliasedVar = ActiveHistoryForFunctionCalls.ResolveAliases(Var);
+			DataInterfaceName = AliasedVar.GetName();
+		}
+		else
+		{
+			DataInterfaceName = Var.GetName();
+		}
+
+		Idx = CompilationOutput.ScriptData.DataInterfaceInfo.AddDefaulted();
+		CompilationOutput.ScriptData.DataInterfaceInfo[Idx].Name = DataInterfaceName;
+		CompilationOutput.ScriptData.DataInterfaceInfo[Idx].Type = Var.GetType();
+		CompilationOutput.ScriptData.DataInterfaceInfo[Idx].bIsPlaceholder = bPlaceholder;
 
 		//Interface requires per instance data so add a user pointer table entry.
-		if (DataInterface->PerInstanceDataSize() > 0)
+		if (DataInterface != nullptr && DataInterface->PerInstanceDataSize() > 0)
 		{
-			CompilationOutput.DataInterfaceInfo[Idx].UserPtrIdx = CompilationOutput.NumUserPtrs++;
+			CompilationOutput.ScriptData.DataInterfaceInfo[Idx].UserPtrIdx = CompilationOutput.ScriptData.NumUserPtrs++;
 		}
 	}
 	else
 	{
-		check(CompilationOutput.DataInterfaceInfo[Idx].Name == Var.GetName());
-		check(CompilationOutput.DataInterfaceInfo[Idx].Type == Var.GetType());
+		check(CompilationOutput.ScriptData.DataInterfaceInfo[Idx].Name == Var.GetName());
+		check(CompilationOutput.ScriptData.DataInterfaceInfo[Idx].Type == Var.GetType());
 	}
 
-	if (bPlaceholder)
+	if (bAddParameterMapRead)
 	{
-		PlaceholderDataInterfaces.AddUnique(DataInterface);
+		FName UsageName;
+		if (FNiagaraParameterMapHistory::IsAliasedEmitterParameter(Var.GetName().ToString()))
+		{
+			FNiagaraVariable AliasedVar = ActiveHistoryForFunctionCalls.ResolveAliases(Var);
+			UsageName = AliasedVar.GetName();
+		}
+		else
+		{
+			UsageName = Var.GetName();
+		}
+		CompilationOutput.ScriptData.DataInterfaceInfo[Idx].RegisteredParameterMapRead = UsageName;
 	}
-
+	
 	return Idx;
 }
 
@@ -3611,12 +3701,13 @@ void FHlslNiagaraTranslator::Operation(class UNiagaraNodeOp* Operation, TArray<i
 	Operation->GetOutputPins(OutputPins);
 	for (int32 OutputIndex = 0; OutputIndex < OutputPins.Num(); OutputIndex++)
 	{
-		FNiagaraTypeDefinition OutputType = Schema->PinToTypeDefinition(OutputPins[OutputIndex]);
+		UEdGraphPin* OutputPin = OutputPins[OutputIndex];
+		FNiagaraTypeDefinition OutputType = Schema->PinToTypeDefinition(OutputPin);
 
 		if (!AddStructToDefinitionSet(OutputType))
 		{
-			FText PinNameText = OutputPins[OutputIndex]->PinFriendlyName.IsEmpty() ? FText::FromName(OutputPins[OutputIndex]->PinName) : OutputPins[OutputIndex]->PinFriendlyName;
-			Error(FText::Format(LOCTEXT("GetConstantFailTypePin", "Cannot handle type {0}! Output Pin: {1}"), OutputType.GetNameText(), PinNameText), Operation, OutputPins[OutputIndex]);
+			FText PinNameText = OutputPin->PinFriendlyName.IsEmpty() ? FText::FromName(OutputPin->PinName) : OutputPin->PinFriendlyName;
+			Error(FText::Format(LOCTEXT("GetConstantFailTypePin", "Cannot handle type {0}! Output Pin: {1}"), OutputType.GetNameText(), PinNameText), Operation, OutputPin);
 		}
 
 		const FNiagaraOpInOutInfo& IOInfo = OpInfo->Outputs[OutputIndex];
@@ -3672,13 +3763,14 @@ void FHlslNiagaraTranslator::FunctionCall(UNiagaraNodeFunctionCall* FunctionNode
 	FNiagaraFunctionSignature OutputSignature;
 	if (FunctionNode->FunctionScript == nullptr && !FunctionNode->Signature.IsValid())
 	{
-		Error(LOCTEXT("FunctionCallNonexistant", "Function call missing FunctionScript and invalid signature"), FunctionNode, nullptr);
+		Error(LOCTEXT("FunctionCallNonexistantFunctionScript", "Function call missing FunctionScript and invalid signature"), FunctionNode, nullptr);
 		return;
 	}
 
 	// We need the generated string to generate the proper signature for now.
 	ActiveHistoryForFunctionCalls.EnterFunction(FunctionNode->GetFunctionName(), FunctionNode->FunctionScript, FunctionNode);
-	
+
+
 	// Remove input add pin if it exists
 	for (int32 i = 0; i < CallOutputs.Num(); i++)
 	{
@@ -3713,76 +3805,236 @@ void FHlslNiagaraTranslator::FunctionCall(UNiagaraNodeFunctionCall* FunctionNode
 		Name = FunctionNode->FunctionScript->GetName();
 		FullName = FunctionNode->FunctionScript->GetFullName();
 		Source = CastChecked<UNiagaraScriptSource>(FunctionNode->FunctionScript->GetSource());
+		check(Source->GetOutermost() == GetTransientPackage());
 	}
 	UNiagaraNodeCustomHlsl* CustomFunctionHlsl = Cast<UNiagaraNodeCustomHlsl>(FunctionNode);
 	if (CustomFunctionHlsl != nullptr)
 	{
-		SCOPE_CYCLE_COUNTER(STAT_NiagaraEditor_Module_NiagaraHLSLTranslator_CustomHLSL);
+		// All of the arguments here are resolved withing the HandleCustomHlsl function..
+		HandleCustomHlslNode(CustomFunctionHlsl, ScriptUsage, Name, FullName, bCustomHlsl, CustomHlsl, Signature, Inputs);
+	}
+	
+	RegisterFunctionCall(ScriptUsage, Name, FullName,  Source, Signature, bCustomHlsl, CustomHlsl, Inputs, CallInputs, CallOutputs, OutputSignature);
+	GenerateFunctionCall(OutputSignature, Inputs, Outputs);
+
+	if (bCustomHlsl)
+	{
+		// Re-add the add pins.
+		Inputs.Add(INDEX_NONE);
+		Outputs.Add(INDEX_NONE);
+	}
+	ActiveHistoryForFunctionCalls.ExitFunction(FunctionNode->GetFunctionName(), FunctionNode->FunctionScript, FunctionNode);
+}
 
-		ScriptUsage = CustomFunctionHlsl->ScriptUsage;
-		Name = GetSanitizedSymbolName(CustomFunctionHlsl->Signature.Name.ToString() + FunctionNode->NodeGuid.ToString());
-		Signature.Name = *Name; // Force the name to be set to include the node guid for safety...
-		bCustomHlsl = true;
-		CustomHlsl = CustomFunctionHlsl->CustomHlsl;
 
-		TArray<FNiagaraVariable> SigInputs;
-		for (FNiagaraVariable Input : Signature.Inputs)
+
+// From a valid list of namespaces, resolve any aliased tokens and promote namespaced variables without a master namespace to the input parameter map instance namespace
+void FHlslNiagaraTranslator::FinalResolveNamespacedTokens(const FString& ParameterMapInstanceNamespace, TArray<FString>& Tokens, TArray<FString>& ValidChildNamespaces, FNiagaraParameterMapHistoryBuilder& Builder, TArray<FNiagaraVariable>& UniqueParameterMapEntriesAliasesIntact, TArray<FNiagaraVariable>& UniqueParameterMapEntries, int32 ParamMapHistoryIdx)
+{
+	for (int32 i = 0; i < Tokens.Num(); i++)
+	{
+		if (Tokens[i].Contains(TEXT("."))) // Only check tokens with namespaces in them..
 		{
-			if (Input.GetType() == FNiagaraTypeDefinition::GetParameterMapDef())
-			{
-				FString ParameterMapInstanceName = GetParameterMapInstanceName(0);
-				FString ReplaceSrc = Input.GetName().ToString() + TEXT(".");
-				FString ReplaceDest = ParameterMapInstanceName + TEXT(".");
-				CustomHlsl = CustomHlsl.Replace(*ReplaceSrc, *ReplaceDest);
-			}
-			else
+			for (const FString& ValidNamespace : ValidChildNamespaces)
 			{
-				FString ReplaceSrc = Input.GetName().ToString();
-				FString ReplaceDest = TEXT("In_") + ReplaceSrc;
-				CustomHlsl = CustomHlsl.Replace(*ReplaceSrc, *ReplaceDest);
-				SigInputs.Add(Input);
+				FNiagaraVariable Var;
+
+				// There are two possible paths here, one where we're using the namespace as-is from the valid list and one where we've already
+				// prepended with the master parameter map instance namespace but may not have resolved any internal aliases yet.
+				if (Tokens[i].StartsWith(ValidNamespace, ESearchCase::CaseSensitive))
+				{
+					FNiagaraVariable TempVar(FNiagaraTypeDefinition::GetFloatDef(), *Tokens[i]); // Declare a dummy name here that we will convert aliases for and use later
+					Var = Builder.ResolveAliases(TempVar);
+				}
+				else if (Tokens[i].StartsWith(ParameterMapInstanceNamespace + ValidNamespace, ESearchCase::CaseSensitive))
+				{
+					FString BaseToken = Tokens[i].Mid(ParameterMapInstanceNamespace.Len());
+					FNiagaraVariable TempVar(FNiagaraTypeDefinition::GetFloatDef(), *BaseToken); // Declare a dummy name here that we will convert aliases for and use later
+					Var = Builder.ResolveAliases(TempVar);
+				}
+
+				if (Var.IsValid())
+				{
+					if (ParamMapHistoryIdx != INDEX_NONE)
+					{
+						bool bAdded = false;
+						for (int32 j = 0; j < OtherOutputParamMapHistories.Num(); j++)
+						{
+							int32 VarIdx = OtherOutputParamMapHistories[j].FindVariableByName(Var.GetName(), true);
+							if (VarIdx != INDEX_NONE)
+							{
+								if (OtherOutputParamMapHistories[j].VariablesWithOriginalAliasesIntact[VarIdx].IsValid())
+								{
+									UniqueParameterMapEntriesAliasesIntact.AddUnique(OtherOutputParamMapHistories[j].VariablesWithOriginalAliasesIntact[VarIdx]);
+								}
+								else
+								{
+									UniqueParameterMapEntriesAliasesIntact.AddUnique(OtherOutputParamMapHistories[j].Variables[VarIdx]);
+								}
+								UniqueParameterMapEntries.AddUnique(OtherOutputParamMapHistories[j].Variables[VarIdx]);
+								bAdded = true;
+								break;
+							}
+						}
+						if (!bAdded)
+						{
+							Error(FText::Format(LOCTEXT("GetCustomFail1", "Cannot use variable in custom expression, it hasn't been encountered yet: {0}"), FText::FromName(Var.GetName())), nullptr, nullptr);
+						}
+
+					}
+
+					
+					Tokens[i] = ParameterMapInstanceNamespace + Var.GetName().ToString();
+					break;
+				}
 			}
 		}
-		Signature.Inputs = SigInputs;
+	}
+}
+
+void FHlslNiagaraTranslator::HandleCustomHlslNode(UNiagaraNodeCustomHlsl* CustomFunctionHlsl, ENiagaraScriptUsage& OutScriptUsage, FString& OutName, FString& OutFullName, bool& bOutCustomHlsl, FString& OutCustomHlsl, 
+	FNiagaraFunctionSignature& OutSignature, TArray<int32>& Inputs)
+{
+	SCOPE_CYCLE_COUNTER(STAT_NiagaraEditor_Module_NiagaraHLSLTranslator_CustomHLSL);
+	if (!CustomFunctionHlsl)
+	{
+		return;
+	}
+
+	// Determine the important outputs
+	OutScriptUsage = CustomFunctionHlsl->ScriptUsage;
+	OutName = GetSanitizedSymbolName(CustomFunctionHlsl->Signature.Name.ToString() + CustomFunctionHlsl->NodeGuid.ToString());
+	OutSignature = CustomFunctionHlsl->Signature;
+	OutFullName = CustomFunctionHlsl->GetFullName();
+	OutSignature.Name = *OutName; // Force the name to be set to include the node guid for safety...
+	bOutCustomHlsl = true;
+	OutCustomHlsl = CustomFunctionHlsl->CustomHlsl;
+
+	// Split up the hlsl into constituent tokens
+	TArray<FString> Tokens;
+	CustomFunctionHlsl->GetTokens(Tokens);
+
+	int32 ParamMapHistoryIdx = INDEX_NONE;
+	bool bHasParamMapOutputs = false;
+	bool bHasParamMapInputs = false;
+
+	// Resolve the names of any internal variables from the input variables.
+	TArray<FNiagaraVariable> SigInputs;
+	for (int32 i = 0; i < OutSignature.Inputs.Num(); i++)
+	{
+		FNiagaraVariable Input = OutSignature.Inputs[i];
+		if (Input.GetType() == FNiagaraTypeDefinition::GetParameterMapDef())
+		{
+			FString ParameterMapInstanceName = GetParameterMapInstanceName(0);
+			FString ReplaceSrc = Input.GetName().ToString() + TEXT(".");
+			FString ReplaceDest = ParameterMapInstanceName + TEXT(".");
+			CustomFunctionHlsl->ReplaceExactMatchTokens(Tokens, ReplaceSrc, ReplaceDest, false);
+			SigInputs.Add(Input);
+			OutSignature.bRequiresContext = true;
+			ParamMapHistoryIdx = Inputs[i];
+			bHasParamMapInputs = true;
+		}
+		else
+		{
+			FString ReplaceSrc = Input.GetName().ToString();
+			FString ReplaceDest = TEXT("In_") + ReplaceSrc;
+			CustomFunctionHlsl->ReplaceExactMatchTokens(Tokens, ReplaceSrc, ReplaceDest, true);
+			SigInputs.Add(Input);
+		}
+	}
+	OutSignature.Inputs = SigInputs;
 
-		TArray<FNiagaraVariable> SigOutputs;
-		for (FNiagaraVariable Output : Signature.Outputs)
+	// Resolve the names of any internal variables from the output variables.
+	TArray<FNiagaraVariable> SigOutputs;
+	for (FNiagaraVariable Output : OutSignature.Outputs)
+	{
+		if (Output.GetType() == FNiagaraTypeDefinition::GetParameterMapDef())
 		{
-			if (Output.GetType() == FNiagaraTypeDefinition::GetParameterMapDef())
+			FString ParameterMapInstanceName = GetParameterMapInstanceName(0);
+			FString ReplaceSrc = Output.GetName().ToString() + TEXT(".");
+			FString ReplaceDest = ParameterMapInstanceName + TEXT(".");
+			CustomFunctionHlsl->ReplaceExactMatchTokens(Tokens, ReplaceSrc, ReplaceDest, false);
+			SigOutputs.Add(Output);
+			OutSignature.bRequiresContext = true;
+			bHasParamMapOutputs = true;
+		}
+		else
+		{
+			FString ReplaceSrc = Output.GetName().ToString();
+			FString ReplaceDest = TEXT("Out_") + ReplaceSrc;
+			CustomFunctionHlsl->ReplaceExactMatchTokens(Tokens, ReplaceSrc, ReplaceDest, true);
+			SigOutputs.Add(Output);
+		}
+	}
+
+	if (bHasParamMapOutputs || bHasParamMapInputs)
+	{
+		// Clean up any namespaced variables in the token list if they are aliased or promote any tokens that are namespaced to the parent 
+		// parameter map.
+		TArray<FString> PossibleNamespaces;
+		FNiagaraParameterMapHistory::GetValidNamespacesForReading(CompileOptions.TargetUsage, PossibleNamespaces);
+
+		for (FNiagaraParameterMapHistory& History : ParamMapHistories)
+		{
+			for (FNiagaraVariable& Var : History.Variables)
 			{
-				FString ParameterMapInstanceName = GetParameterMapInstanceName(0);
-				FString ReplaceSrc = Output.GetName().ToString() + TEXT(".");
-				FString ReplaceDest = ParameterMapInstanceName + TEXT(".");
-				CustomHlsl = CustomHlsl.Replace(*ReplaceSrc, *ReplaceDest);
+				FString Namespace = FNiagaraParameterMapHistory::GetNamespace(Var);
+				PossibleNamespaces.AddUnique(Namespace);
 			}
-			else
+		}
+
+		TArray<FNiagaraVariable> UniqueParamMapEntries;
+		TArray<FNiagaraVariable> UniqueParamMapEntriesAliasesIntact;
+		FinalResolveNamespacedTokens(GetParameterMapInstanceName(0) + TEXT("."), Tokens, PossibleNamespaces, ActiveHistoryForFunctionCalls, UniqueParamMapEntriesAliasesIntact, UniqueParamMapEntries, ParamMapHistoryIdx);
+
+		// We must register any external constant variables that we encountered.
+		for (FNiagaraVariable Var : UniqueParamMapEntriesAliasesIntact)
+		{
+			if (FNiagaraParameterMapHistory::IsExternalConstantNamespace(Var, CompileOptions.TargetUsage, CompileOptions.GetTargetUsageBitmask()))
 			{
-				FString ReplaceSrc = Output.GetName().ToString();
-				FString ReplaceDest = TEXT("Out_") + ReplaceSrc;
-				CustomHlsl = CustomHlsl.Replace(*ReplaceSrc, *ReplaceDest);
-				SigOutputs.Add(Output);
+				int32 TempOutput;
+				if (ParameterMapRegisterExternalConstantNamespaceVariable(Var, CustomFunctionHlsl, ParamMapHistoryIdx, TempOutput, nullptr))
+				{
+					continue;
+				}
 			}
-			Signature.Outputs = SigOutputs;
 		}
-
-		CustomHlsl = CustomHlsl.Replace(TEXT("\n"), TEXT("\n\t"));
-		CustomHlsl = TEXT("\n") + CustomHlsl + TEXT("\n");
 	}
-	
 
-	RegisterFunctionCall(ScriptUsage, Name, FullName,  Source, Signature, bCustomHlsl, CustomHlsl, Inputs, CallInputs, CallOutputs, OutputSignature);
-	GenerateFunctionCall(OutputSignature, Inputs, Outputs);
+	// Now reassemble the tokens into the final hlsl output
+	OutSignature.Outputs = SigOutputs;
+	OutCustomHlsl = FString::Join<FString>(Tokens, TEXT(""));
 
-	if (bCustomHlsl)
+	// Dynamic inputs are assumed to be of the form 
+	// "20.0f * Particles.Velocity.x + length(Particles.Velocity)", i.e. a mix of native functions, constants, operations, and variable names.
+	// This needs to be modified to match the following requirements:	
+	// 1) Write to the output variable of the dynamic input.
+	// 2) Terminate in valid HLSL (i.e. have a ; at the end)
+	// 3) Be guaranteed to write to the correct output type.
+	if (OutScriptUsage == ENiagaraScriptUsage::DynamicInput)
 	{
-		// Re-add the add pins.
-		Inputs.Add(INDEX_NONE);
-		Outputs.Add(INDEX_NONE);
+		if (CustomFunctionHlsl->Signature.Outputs.Num() != 1)
+		{
+			Error(LOCTEXT("CustomHlslDynamicInputMissingOutputs", "Custom hlsl dynamic input signature should have one and only one output."), CustomFunctionHlsl, nullptr);
+			return;
+		}
+		if (CustomFunctionHlsl->Signature.Inputs.Num() < 1 || CustomFunctionHlsl->Signature.Inputs[0].GetType() != FNiagaraTypeDefinition::GetParameterMapDef())
+		{
+			Error(LOCTEXT("CustomHlslDynamicInputMissingInputs", "Custom hlsl dynamic input signature should have at least one input (a parameter map)."), CustomFunctionHlsl, nullptr);
+			return;
+		}
+
+		OutSignature.bRequiresContext = true;
+		FString ReplaceSrc = CustomFunctionHlsl->Signature.Outputs[0].GetName().ToString();
+		FString ReplaceDest = TEXT("Out_") + ReplaceSrc;
+		OutCustomHlsl = ReplaceDest + TEXT(" = (") + GetStructHlslTypeName(CustomFunctionHlsl->Signature.Outputs[0].GetType()) + TEXT(")(") + OutCustomHlsl + TEXT(");\n");
 	}
-	ActiveHistoryForFunctionCalls.ExitFunction(FunctionNode->GetFunctionName(), FunctionNode->FunctionScript, FunctionNode);
+
+	OutCustomHlsl = OutCustomHlsl.Replace(TEXT("\n"), TEXT("\n\t"));
+	OutCustomHlsl = TEXT("\n") + OutCustomHlsl + TEXT("\n");
 }
 
-void FHlslNiagaraTranslator::RegisterFunctionCall(ENiagaraScriptUsage ScriptUsage, const FString& InName, const FString& InFullName, UNiagaraScriptSource* Source, FNiagaraFunctionSignature& InSignature, bool bIsCustomHlsl, const FString& InCustomHlsl, TArray<int32>& Inputs, const TArray<UEdGraphPin*>& CallInputs, const TArray<UEdGraphPin*>& CallOutputs,
+void FHlslNiagaraTranslator::RegisterFunctionCall(ENiagaraScriptUsage ScriptUsage, const FString& InName, const FString& InFullName, UNiagaraScriptSource* Source,  FNiagaraFunctionSignature& InSignature, bool bIsCustomHlsl, const FString& InCustomHlsl, TArray<int32>& Inputs, const TArray<UEdGraphPin*>& CallInputs, const TArray<UEdGraphPin*>& CallOutputs,
 	FNiagaraFunctionSignature& OutSignature)
 {
 	SCOPE_CYCLE_COUNTER(STAT_NiagaraEditor_Module_NiagaraHLSLTranslator_RegisterFunctionCall);
@@ -3790,49 +4042,21 @@ void FHlslNiagaraTranslator::RegisterFunctionCall(ENiagaraScriptUsage ScriptUsag
 	//////////////////////////////////////////////////////////////////////////
 	if (Source)
 	{
-		SCOPE_CYCLE_COUNTER(STAT_NiagaraEditor_Module_NiagaraHLSLTranslator_RegisterFunctionCall_Source); 
+		SCOPE_CYCLE_COUNTER(STAT_NiagaraEditor_Module_NiagaraHLSLTranslator_RegisterFunctionCall_Source);
 		UNiagaraGraph* SourceGraph = CastChecked<UNiagaraGraph>(Source->NodeGraph);
 
 		bool bHasNumericInputs = false;
 		if (SourceGraph->HasNumericParameters())
 		{
-			TArray<UNiagaraNodeInput*> InputNodes;
-			UNiagaraGraph::FFindInputNodeOptions Options;
-			Options.bFilterDuplicates = true;
-			Options.bIncludeParameters = true;
-			Options.bIncludeAttributes = false;
-			Options.bIncludeSystemConstants = false;
-			SourceGraph->FindInputNodes(InputNodes, Options);
-
-			for (UNiagaraNodeInput* Input : InputNodes)
+			for (int32 i = 0; i < CallInputs.Num(); i++)
 			{
-				if (Input->Input.GetType() == FNiagaraTypeDefinition::GetGenericNumericDef())
+				if (Schema->PinToTypeDefinition(CallInputs[i]) == FNiagaraTypeDefinition::GetGenericNumericDef())
 				{
 					bHasNumericInputs = true;
 				}
 			}
-
-			SCOPE_CYCLE_COUNTER(STAT_NiagaraEditor_Module_NiagaraHLSLTranslator_FunctionCallCloneGraphNumeric);
-			//We have to clone and preprocess the function graph to ensure it's numeric types have been fixed up to real types.
-			UNiagaraGraph* PreprocessedGraph = CastChecked<UNiagaraGraph>(FEdGraphUtilities::CloneGraph(SourceGraph, Source, &MessageLog));
-			FEdGraphUtilities::MergeChildrenGraphsIn(PreprocessedGraph, PreprocessedGraph, /*bRequireSchemaMatch=*/ true);
-			PreprocessFunctionGraph(*this, Schema, PreprocessedGraph, CallInputs, CallOutputs, ScriptUsage);
-			SourceGraph = PreprocessedGraph;
-		}
-		else
-		{
-			SCOPE_CYCLE_COUNTER(STAT_NiagaraEditor_Module_NiagaraHLSLTranslator_FunctionCallCloneGraphNonNumeric);
-			//If we don't have numeric inputs then we can cache the preprocessed graphs.
-			UNiagaraGraph** CachedGraph = PreprocessedFunctions.Find(SourceGraph);
-			if (!CachedGraph)
-			{
-				CachedGraph = &PreprocessedFunctions.Add(SourceGraph);
-				*CachedGraph = CastChecked<UNiagaraGraph>(FEdGraphUtilities::CloneGraph(Source->NodeGraph, Source, &MessageLog));
-				FEdGraphUtilities::MergeChildrenGraphsIn(*CachedGraph, *CachedGraph, /*bRequireSchemaMatch=*/ true);
-				PreprocessFunctionGraph(*this, Schema, *CachedGraph, CallInputs, CallOutputs, ScriptUsage);
-			}
-			SourceGraph = *CachedGraph;
 		}
+
 		
 		bool bHasParameterMapParameters = SourceGraph->HasParameterMapParameters();
 
@@ -3867,6 +4091,60 @@ void FHlslNiagaraTranslator::RegisterFunctionCall(ENiagaraScriptUsage ScriptUsag
 			UNiagaraNodeOutput* FuncOutput = SourceGraph->FindOutputNode(ScriptUsage);
 			check(FuncOutput);
 
+			// Go ahead and insert any defaulted values into the parameter map here at the top level.
+			if (ActiveHistoryForFunctionCalls.InTopLevelFunctionCall(CompileOptions.TargetUsage) && ActiveHistoryForFunctionCalls.GetModuleAlias() != nullptr)
+			{
+				UEdGraphPin* ParamMapPin = nullptr;
+				for (UEdGraphPin* Pin : CallInputs)
+				{
+					if (Schema->PinToTypeDefinition(Pin) == FNiagaraTypeDefinition::GetParameterMapDef())
+					{
+						ParamMapPin = Pin;
+						break;
+					}
+				}
+
+				if (ParamMapPin != nullptr)
+				{
+					UNiagaraNode* ParamNode = Cast<UNiagaraNode>(ParamMapPin->GetOwningNode());
+					if (ParamNode)
+					{
+						check(ParamMapHistories.Num() == TranslationStages.Num());
+						const FNiagaraParameterMapHistory& History = ParamMapHistories[ActiveStageIdx];
+						uint32 FoundIdx = History.MapNodeVisitations.Find(ParamNode);
+						if (FoundIdx != INDEX_NONE)
+						{
+							check((uint32)History.MapNodeVariableMetaData.Num() > FoundIdx);
+							check(INDEX_NONE != History.MapNodeVariableMetaData[FoundIdx].Key);
+							check(INDEX_NONE != History.MapNodeVariableMetaData[FoundIdx].Value);
+
+							for (uint32 VarIdx = History.MapNodeVariableMetaData[FoundIdx].Key; VarIdx < History.MapNodeVariableMetaData[FoundIdx].Value; VarIdx++)
+							{
+								const FNiagaraVariable& Var = History.Variables[VarIdx];
+								const FNiagaraVariable& AliasedVar = History.VariablesWithOriginalAliasesIntact[VarIdx];
+								int32 LastSetChunkIdx = ParamMapSetVariablesToChunks[ActiveStageIdx][VarIdx];
+								if (LastSetChunkIdx == INDEX_NONE)
+								{
+									const UEdGraphPin* DefaultPin = History.GetDefaultValuePin(VarIdx);
+									HandleParameterRead(ActiveStageIdx, AliasedVar, DefaultPin, ParamNode, LastSetChunkIdx);
+/*
+									if (History.IsPrimaryDataSetOutput(AliasedVar, CompileOptions.TargetUsage))
+									{
+										// insert defaults here..
+									}
+									else if (History.IsAliasedModuleParameter(AliasedVar))
+									{
+										// insert defaults here..
+										HandleParameterRead(ActiveStageIdx, AliasedVar, DefaultPin, ParamNode, LastSetChunkIdx);
+									}
+									*/
+								}
+							}
+						}
+					}
+				}
+			}
+
 			//Track the start of this funciton in the chunks so we can remove them after we grab the funcitons hlsl.
 			int32 ChunkStart = CodeChunks.Num();
 			int32 ChunkStartsByMode[(int32)ENiagaraCodeChunkMode::Num];
@@ -3886,38 +4164,6 @@ void FHlslNiagaraTranslator::RegisterFunctionCall(ENiagaraScriptUsage ScriptUsag
 			}
 			CurrentBodyChunkMode = OldMode;
 
-			// Find all of the data set writes that were connected to this particular graph.
-			TArray<UNiagaraNode*> OutputTraversal;
-			OutputTraversal.Reserve(1024);
-			SourceGraph->BuildTraversal(OutputTraversal, FuncOutput);
-
-			TArray<UNiagaraNodeWriteDataSet *>WriteNodes;
-			SourceGraph->FindWriteDataSetNodes(WriteNodes);
-
-			// If the write data set graph traversal intersects with our normal output graph traversal, then
-			// we consider them part of the same overall graph traversal and include the write nodes.
-			for (UNiagaraNodeWriteDataSet *WriteNode : WriteNodes)
-			{
-				TArray<UNiagaraNode*> WriteTraversal;
-				WriteTraversal.Reserve(1024);
-				SourceGraph->BuildTraversal(WriteTraversal, WriteNode);
-
-				bool bReferencesSameGraph = false;
-				for (UNiagaraNode* Node : WriteTraversal)
-				{
-					if (OutputTraversal.Contains(Node))
-					{
-						bReferencesSameGraph = true;
-					}
-				}
-
-				if (bReferencesSameGraph)
-				{
-					SCOPE_CYCLE_COUNTER(STAT_NiagaraEditor_Module_NiagaraHLSLTranslator_RegisterFunctionCall_Compile);
-					WriteNode->Compile(ThisTranslator, FuncOutputChunks);
-				}
-			}
-
 			{
 				SCOPE_CYCLE_COUNTER(STAT_NiagaraEditor_Module_NiagaraHLSLTranslator_RegisterFunctionCall_FunctionDefStr);
 				//Grab all the body chunks for this function.
@@ -4000,20 +4246,23 @@ void FHlslNiagaraTranslator::RegisterFunctionCall(ENiagaraScriptUsage ScriptUsag
 		else
 		{
 			int32 OwnerIdx = Inputs[0];
-			if (OwnerIdx < 0 || OwnerIdx >= CompilationOutput.DataInterfaceInfo.Num())
+			if (OwnerIdx < 0 || OwnerIdx >= CompilationOutput.ScriptData.DataInterfaceInfo.Num())
 			{
 				Error(LOCTEXT("FunctionCallDataInterfaceMissingRegistration", "Function call signature does not match to a registered DataInterface. Valid DataInterfaces should be wired into a DataInterface function call."), nullptr, nullptr);
 				return;
 			}
-			FNiagaraScriptDataInterfaceInfo& Info = CompilationOutput.DataInterfaceInfo[OwnerIdx];
+			FNiagaraScriptDataInterfaceCompileInfo& Info = CompilationOutput.ScriptData.DataInterfaceInfo[OwnerIdx];
 
 			// Double-check to make sure that the signature matches those specified by the data 
 			// interface. It could be that the existing node has been removed and the graph
 			// needs to be refactored. If that's the case, emit an error.
-			if (Info.DataInterface != nullptr && OutSignature.bMemberFunction)
+			UObject* const* FoundCDO = CompileData->CDOs.Find(Info.Type.GetClass());
+			check(FoundCDO != nullptr);
+			UNiagaraDataInterface* CDO = Cast<UNiagaraDataInterface>(*FoundCDO);
+			if (CDO != nullptr && OutSignature.bMemberFunction)
 			{
 				TArray<FNiagaraFunctionSignature> DataInterfaceFunctions;
-				Info.DataInterface->GetFunctions(DataInterfaceFunctions);
+				CDO->GetFunctions(DataInterfaceFunctions);
 				bool bFoundMatch = false;
 				for (const FNiagaraFunctionSignature& Sig : DataInterfaceFunctions)
 				{
@@ -4473,15 +4722,15 @@ void FHlslNiagaraTranslator::If(TArray<FNiagaraVariable>& Vars, int32 Condition,
 	AddBodyChunk(TEXT(""), TEXT("if({0})\n\t{"), FNiagaraTypeDefinition::GetFloatDef(), Condition, false, false);
 	for (int32 i = 0; i < NumVars; ++i)
 	{
-		FNiagaraCodeChunk& OutChunk = CodeChunks[Outputs[i]];		
-		FNiagaraCodeChunk& BranchChunk = CodeChunks[AddBodyChunk(OutSymbols[i], TEXT("{0}"), OutChunk.Type, false)];
+		FNiagaraTypeDefinition OutChunkType = CodeChunks[Outputs[i]].Type;
+		FNiagaraCodeChunk& BranchChunk = CodeChunks[AddBodyChunk(OutSymbols[i], TEXT("{0}"), OutChunkType, false)];
 		BranchChunk.AddSourceChunk(PathA[i]);
 	}	
 	AddBodyChunk(TEXT(""), TEXT("}\n\telse\n\t{"), FNiagaraTypeDefinition::GetFloatDef(), false, false);
 	for (int32 i = 0; i < NumVars; ++i)
 	{
-		FNiagaraCodeChunk& OutChunk = CodeChunks[Outputs[i]];
-		FNiagaraCodeChunk& BranchChunk = CodeChunks[AddBodyChunk(OutSymbols[i], TEXT("{0}"), OutChunk.Type, false)];
+		FNiagaraTypeDefinition OutChunkType = CodeChunks[Outputs[i]].Type;
+		FNiagaraCodeChunk& BranchChunk = CodeChunks[AddBodyChunk(OutSymbols[i], TEXT("{0}"), OutChunkType, false)];
 		BranchChunk.AddSourceChunk(PathB[i]);
 	}
 	AddBodyChunk(TEXT(""), TEXT("}"), FNiagaraTypeDefinition::GetFloatDef(), false, false);
@@ -4490,7 +4739,7 @@ void FHlslNiagaraTranslator::If(TArray<FNiagaraVariable>& Vars, int32 Condition,
 	Outputs.Add(INDEX_NONE);
 }
 
-int32 FHlslNiagaraTranslator::CompilePin(UEdGraphPin* Pin)
+int32 FHlslNiagaraTranslator::CompilePin(const UEdGraphPin* Pin)
 {
 	SCOPE_CYCLE_COUNTER(STAT_NiagaraEditor_HlslTranslator_CompilePin);
 
@@ -4507,7 +4756,7 @@ int32 FHlslNiagaraTranslator::CompilePin(UEdGraphPin* Pin)
 		{
 			if (TypeDef == FNiagaraTypeDefinition::GetParameterMapDef())
 			{
-				Error(FText::FromString(TEXT("Cannot create a constant ParameterMap!")), Cast<UNiagaraNode>(Pin->GetOwningNode()), Pin);
+				Error(FText::FromString(TEXT("Parameter Maps must be created via an Input Node, not the default value of a pin! Please connect to a valid input Parameter Map.")), Cast<UNiagaraNode>(Pin->GetOwningNode()), nullptr);
 				return INDEX_NONE;
 			}
 			else
@@ -4532,9 +4781,13 @@ int32 FHlslNiagaraTranslator::CompilePin(UEdGraphPin* Pin)
 	return Ret;
 }
 
-int32 FHlslNiagaraTranslator::CompileOutputPin(UEdGraphPin* Pin)
+int32 FHlslNiagaraTranslator::CompileOutputPin(const UEdGraphPin* InPin)
 {
 	SCOPE_CYCLE_COUNTER(STAT_NiagaraEditor_HlslTranslator_CompileOutputPin);
+	
+	// The incoming pin to compile may be pointing to a reroute node. If so, we just jump over it
+	// to where it really came from.
+	UEdGraphPin* Pin = UNiagaraNode::TraceOutputPin(const_cast<UEdGraphPin*>(InPin));
 
 	check(Pin->Direction == EGPD_Output);
 
@@ -4573,7 +4826,7 @@ int32 FHlslNiagaraTranslator::CompileOutputPin(UEdGraphPin* Pin)
 			}
 			else
 			{
-				Error(LOCTEXT("IncorrectNumOutputsError", "Incorrect number of outputs. Can possibly be fixed with a graph refresh."), Node, Pin);
+				Error(LOCTEXT("IncorrectNumOutputsError", "Incorect number of outputs. Can possibly be fixed with a graph refresh."), Node, nullptr);
 			}
 		}
 	}
@@ -4611,26 +4864,59 @@ bool FHlslNiagaraTranslator::GetFunctionParameter(const FNiagaraVariable& Parame
 
 bool FHlslNiagaraTranslator::CanReadAttributes()const
 {
-	if (Script->IsParticleUpdateScript() || (Script->IsInterpolatedParticleSpawnScript() && !bInsideInterpolatedSpawnScript))
+	if (UNiagaraScript::IsParticleUpdateScript(TranslationStages[ActiveStageIdx].ScriptUsage))
 	{
 		return true;
 	}
 	return false;
 }
 
+ENiagaraScriptUsage FHlslNiagaraTranslator::GetCurrentUsage() const
+{
+	if (UNiagaraScript::IsParticleScript(CompileOptions.TargetUsage))
+	{
+		return CompileOptions.TargetUsage;
+	}
+	else if (UNiagaraScript::IsSystemSpawnScript(CompileOptions.TargetUsage) || UNiagaraScript::IsSystemUpdateScript(CompileOptions.TargetUsage))
+	{
+		if (ActiveHistoryForFunctionCalls.ContextContains(ENiagaraScriptUsage::EmitterSpawnScript))
+		{
+			return ENiagaraScriptUsage::EmitterSpawnScript;
+		}
+		else if (ActiveHistoryForFunctionCalls.ContextContains(ENiagaraScriptUsage::EmitterUpdateScript))
+		{
+			return ENiagaraScriptUsage::EmitterUpdateScript;
+		}
+		return CompileOptions.TargetUsage;
+	}
+	else if (UNiagaraScript::IsStandaloneScript(CompileOptions.TargetUsage))
+	{
+		// Since we never use the results of a standalone script directly, just choose one by default.
+		return ENiagaraScriptUsage::ParticleSpawnScript;
+	}
+	else
+	{
+		check(false);
+		return ENiagaraScriptUsage::ParticleSpawnScript;
+	}
+}
 
 ENiagaraScriptUsage FHlslNiagaraTranslator::GetTargetUsage() const
 {
-	if (Script->IsInterpolatedParticleSpawnScript())
+	if (CompileOptions.TargetUsage == ENiagaraScriptUsage::ParticleGPUComputeScript) // Act as if building spawn script.
+	{
+		return ENiagaraScriptUsage::ParticleSpawnScript;
+	}
+	if (UNiagaraScript::IsInterpolatedParticleSpawnScript(CompileOptions.TargetUsage))
 	{
 		return ENiagaraScriptUsage::ParticleSpawnScript;
 	}
-	return Script->GetUsage();
+	return CompileOptions.TargetUsage;
 }
 
 FGuid FHlslNiagaraTranslator::GetTargetUsageId() const
 {
-	return Script->GetUsageId();
+	return CompileOptions.TargetUsageId;
 }
 //////////////////////////////////////////////////////////////////////////
 
@@ -4642,19 +4928,23 @@ FString FHlslNiagaraTranslator::GetHlslDefaultForType(FNiagaraTypeDefinition Typ
 	}
 	else if (Type == FNiagaraTypeDefinition::GetVec2Def())
 	{
-		return "(0.0,0.0)";
+		return "float2(0.0,0.0)";
 	}
 	else if (Type == FNiagaraTypeDefinition::GetVec3Def())
 	{
-		return "(0.0,0.0,0.0)";
+		return "float3(0.0,0.0,0.0)";
 	}
 	else if (Type == FNiagaraTypeDefinition::GetVec4Def())
 	{
-		return "(0.0,0.0,0.0,0.0)";
+		return "float4(0.0,0.0,0.0,0.0)";
+	}
+	else if (Type == FNiagaraTypeDefinition::GetQuatDef())
+	{
+		return "float4(0.0,0.0,0.0,1.0)";
 	}
 	else if (Type == FNiagaraTypeDefinition::GetColorDef())
 	{
-		return "(1.0,1.0,1.0,1.0)";
+		return "float4(1.0,1.0,1.0,1.0)";
 	}
 	else if (Type == FNiagaraTypeDefinition::GetIntDef())
 	{
@@ -4666,7 +4956,7 @@ FString FHlslNiagaraTranslator::GetHlslDefaultForType(FNiagaraTypeDefinition Typ
 	}
 	else
 	{
-		return Type.GetName();
+		return TEXT("(") + GetStructHlslTypeName(Type) + TEXT(")0");
 	}
 }
 
@@ -4678,6 +4968,7 @@ bool FHlslNiagaraTranslator::IsBuiltInHlslType(FNiagaraTypeDefinition Type)
 		Type == FNiagaraTypeDefinition::GetVec3Def() ||
 		Type == FNiagaraTypeDefinition::GetVec4Def() ||
 		Type == FNiagaraTypeDefinition::GetColorDef() ||
+		Type == FNiagaraTypeDefinition::GetQuatDef() ||
 		Type == FNiagaraTypeDefinition::GetMatrix4Def() ||
 		Type == FNiagaraTypeDefinition::GetIntDef() ||
 		Type.GetStruct() == FNiagaraTypeDefinition::GetIntStruct() ||
@@ -4702,7 +4993,7 @@ FString FHlslNiagaraTranslator::GetStructHlslTypeName(FNiagaraTypeDefinition Typ
 	{
 		return "float3";
 	}
-	else if (Type == FNiagaraTypeDefinition::GetVec4Def() || Type == FNiagaraTypeDefinition::GetColorDef())
+	else if (Type == FNiagaraTypeDefinition::GetVec4Def() || Type == FNiagaraTypeDefinition::GetColorDef() || Type == FNiagaraTypeDefinition::GetQuatDef())
 	{
 		return "float4";
 	}
@@ -4783,6 +5074,7 @@ bool FHlslNiagaraTranslator::IsHlslBuiltinVector(FNiagaraTypeDefinition Type)
 	if ((Type == FNiagaraTypeDefinition::GetVec2Def()) ||
 		(Type == FNiagaraTypeDefinition::GetVec3Def()) ||
 		(Type == FNiagaraTypeDefinition::GetVec4Def()) ||
+		(Type == FNiagaraTypeDefinition::GetQuatDef()) ||
 		(Type == FNiagaraTypeDefinition::GetColorDef()))
 	{
 		return true;
@@ -4840,7 +5132,7 @@ bool FHlslNiagaraTranslator::AddStructToDefinitionSet(const FNiagaraTypeDefiniti
 TArray<FName> FHlslNiagaraTranslator::ConditionPropertyPath(const FNiagaraTypeDefinition& Type, const TArray<FName>& InPath)
 {
 	// TODO: Build something more extensible and less hard coded for path conditioning.
-	const UScriptStruct* Struct = Type.GetScriptStruct();
+	UScriptStruct* Struct = Type.GetScriptStruct();
 	if (InPath.Num() == 0) // Pointing to the root
 	{
 		return TArray<FName>();
