Copying //UE4/Dev-Niagara to //UE4/Dev-Main (Source: //UE4/Dev-Niagara @ 4074996)

diff --git a/Engine/Plugins/FX/Niagara/Source/NiagaraEditor/Private/NiagaraHlslTranslator.h b/Engine/Plugins/FX/Niagara/Source/NiagaraEditor/Private/NiagaraHlslTranslator.h
--- a/Engine/Plugins/FX/Niagara/Source/NiagaraEditor/Private/NiagaraHlslTranslator.h
+++ b/Engine/Plugins/FX/Niagara/Source/NiagaraEditor/Private/NiagaraHlslTranslator.h
@@ -12,11 +12,13 @@
 #include "NiagaraShaderCompilationManager.h"
 #include "NiagaraDataInterface.h"
 #include "TickableEditorObject.h"
+#include "NiagaraScriptSource.h"
 
 class Error;
 class UNiagaraGraph;
 class UNiagaraNode;
 class UNiagaraNodeFunctionCall;
+class UNiagaraNodeCustomHlsl;
 class UNiagaraNodeOutput;
 class UNiagaraScriptSource;
 class UNiagaraNodeEmitter;
@@ -76,30 +78,68 @@ struct FNiagaraTranslateResults
 	static ENiagaraScriptCompileStatus TranslateResultsToSummary(const FNiagaraTranslateResults* CompileResults);
 };
 
-
-/** Data which is generated from the hlsl by the VectorVMBackend and fed back into the */
-struct FNiagaraTranslatorOutput
+class FNiagaraCompileRequestData : public FNiagaraCompileRequestDataBase
 {
-	FNiagaraTranslatorOutput() : NumUserPtrs(0), bReadsAttributeData(false){}
+public:
+	virtual ~FNiagaraCompileRequestData() {}
+	virtual bool GatherPreCompiledVariables(const FString& InNamespaceFilter, TArray<FNiagaraVariable>& OutVars) override;
+	virtual void GetReferencedObjects(TArray<UObject*>& Objects) override;
+	virtual const TMap<FName, UNiagaraDataInterface*>& GetObjectNameMap() override;
+	void MergeInEmitterPrecompiledData(FNiagaraCompileRequestDataBase* InEmitterDataBase);
+	virtual FName ResolveEmitterAlias(FName VariableName) const override;
+
+	TArray<FNiagaraParameterMapHistory>& GetPrecomputedHistories() { return PrecompiledHistories; }
+	const TArray<FNiagaraParameterMapHistory>& GetPrecomputedHistories() const { return PrecompiledHistories; }
+	class UNiagaraGraph* GetPrecomputedNodeGraph() { return NodeGraphDeepCopy; }
+	const class UNiagaraGraph* GetPrecomputedNodeGraph() const { return NodeGraphDeepCopy; }
+	const FString& GetUniqueEmitterName() const { return EmitterUniqueName; }
+	void VisitReferencedGraphs(UNiagaraGraph* InSrcGraph, UNiagaraGraph* InDupeGraph, ENiagaraScriptUsage InUsage);
+	void DeepCopyGraphs(UNiagaraScriptSource* ScriptSource, ENiagaraScriptUsage InUsage);
+	void FinishPrecompile(UNiagaraScriptSource* ScriptSource, const TArray<FNiagaraVariable>& EncounterableVariables, ENiagaraScriptUsage InUsage);
+	virtual int32 GetDependentRequestCount() const override {
+		return EmitterData.Num();
+	};
+	virtual TSharedPtr<FNiagaraCompileRequestDataBase, ESPMode::ThreadSafe> GetDependentRequest(int32 Index) override {
+		return EmitterData[Index];
+	}
 
-	/** All external parameters used in the graph. */
-	FNiagaraParameters Parameters;
 
-	/** All internal constants used in the graph. */
-	FNiagaraParameters InternalParameters;
+	// If this is being held onto for any length of time, make sure to hold onto it in a gc-aware object. Right now in this information-passing struct,
+	// we could have a leaked garbage collected pointer if not held onto by someone capable of registering a reference.
+	UNiagaraGraph* NodeGraphDeepCopy;
+	TArray<FNiagaraParameterMapHistory> PrecompiledHistories;
+	TArray<FNiagaraVariable> ChangedFromNumericVars;
+	TMap<FName, UNiagaraDataInterface*> CopiedDataInterfacesByName;
+	TMap<UClass*, UObject*> CDOs;
+	FString EmitterUniqueName;
+	TArray<TSharedPtr<FNiagaraCompileRequestData, ESPMode::ThreadSafe>> EmitterData;
+	UNiagaraScriptSource* Source;
 
-	/** Data sets this script reads. */
-	TArray<FNiagaraDataSetProperties> DataSetReads;
+	UEnum* ENiagaraScriptCompileStatusEnum;
 
-	/** Data sets this script writes. */
-	TArray<FNiagaraDataSetProperties> DataSetWrites;
+	struct FunctionData
+	{
+		UNiagaraScript* ClonedScript;
+		UNiagaraGraph* ClonedGraph;
+		TArray<UEdGraphPin*> CallInputs;
+		TArray<UEdGraphPin*> CallOutputs;
+		ENiagaraScriptUsage Usage;
+		bool bHasNumericInputs;
+	};
+	TMap<const UNiagaraGraph*, TArray<FunctionData>> PreprocessedFunctions;
+	TArray<UNiagaraGraph*> ClonedGraphs;
+protected:
+	void VisitReferencedGraphsRecursive(UNiagaraGraph* InGraph);
+};
 
-	/* Per instance attributes */
-	TArray<FNiagaraVariable> Attributes;
 
-	FNiagaraScriptDataUsageInfo DataUsage;
 
-	TArray<FNiagaraScriptDataInterfaceInfo> DataInterfaceInfo;
+/** Data which is generated from the hlsl by the VectorVMBackend and fed back into the */
+struct FNiagaraTranslatorOutput
+{
+	FNiagaraTranslatorOutput() {}
+
+	FNiagaraVMExecutableData ScriptData;
 
 	/** Ordered table of functions actually called by the VM script. */
 	struct FCalledVMFunction
@@ -113,10 +153,6 @@ struct FNiagaraTranslatorOutput
 
 	FString Errors;
 
-	/** Number of user pointers. For passing things like data interface instance data etc. */
-	int32 NumUserPtrs;
-
-	bool bReadsAttributeData;
 };
 
 
@@ -128,6 +164,7 @@ enum class ENiagaraCodeChunkMode : uint8
 	Body,
 	SpawnBody,
 	UpdateBody,
+	InitializerBody,
 	Num,
 };
 
@@ -219,6 +256,20 @@ public:
 	TArray<FNiagaraVariable> OverrideModuleConstants;
 };
 
+class NIAGARAEDITOR_API FHlslNiagaraTranslationStage
+{
+public:
+	FHlslNiagaraTranslationStage(ENiagaraScriptUsage InScriptUsage, FGuid InUsageId) : ScriptUsage(InScriptUsage), UsageId(InUsageId), OutputNode(nullptr), bInterpolatePreviousParams(false), bCopyPreviousParams(true), ChunkModeIndex((ENiagaraCodeChunkMode)-1){}
+
+	ENiagaraScriptUsage ScriptUsage;
+	FGuid UsageId;
+	UNiagaraNodeOutput* OutputNode;
+	FString PassNamespace;
+	bool bInterpolatePreviousParams;
+	bool bCopyPreviousParams;
+	ENiagaraCodeChunkMode ChunkModeIndex;
+};
+
 class NIAGARAEDITOR_API FHlslNiagaraTranslator
 {
 public:
@@ -232,9 +283,8 @@ public:
 	};
 
 protected:
-
-	/** The script we are compiling. */
-	UNiagaraScript* Script;
+	const FNiagaraCompileRequestData* CompileData;
+	FNiagaraCompileOptions CompileOptions;
 
 	FHlslNiagaraTranslatorOptions TranslationOptions;
 
@@ -287,11 +337,11 @@ protected:
 	FString GetFunctionSignature(const FNiagaraFunctionSignature& Sig);
 
 	/** Compiles an output Pin on a graph node. Caches the result for any future inputs connected to it. */
-	int32 CompileOutputPin(UEdGraphPin* Pin);
+	int32 CompileOutputPin(const UEdGraphPin* Pin);
 
-	void WriteDataSetContextVars(TMap<FNiagaraDataSetID, TMap<int32, FDataSetAccessInfo>>& DataSetAccessInfo, bool bRead, FString &HlslOutput);
-	void WriteDataSetStructDeclarations(TMap<FNiagaraDataSetID, TMap<int32, FDataSetAccessInfo>>& DataSetAccessInfo, bool bRead, FString &HlslOutput);
-	void DecomposeVariableAccess(const UStruct* Struct, bool bRead, FString IndexSymbol, FString HLSLString);
+	void WriteDataSetContextVars(TMap<FNiagaraDataSetID, TMap<int32, FDataSetAccessInfo>>& DataSetAccessInfo, bool bRead, FString &OutHLSLOutput);
+	void WriteDataSetStructDeclarations(TMap<FNiagaraDataSetID, TMap<int32, FDataSetAccessInfo>>& DataSetAccessInfo, bool bRead, FString &OutHLSLOutput);
+	void DecomposeVariableAccess(UStruct* Struct, bool bRead, FString IndexSymbol, FString HLSLString);
 
 	FString GetUniqueSymbolName(FName BaseName);
 
@@ -333,17 +383,19 @@ protected:
 	int32 AddBodyChunk(FString SymbolName, FString Definition, const FNiagaraTypeDefinition& Type, bool bDecl = true, bool bIsTerminated = true);
 
 	FString GetFunctionDefinitions();
+
+	FString GetUniqueEmitterName() const;
 public:
 
 	FHlslNiagaraTranslator();
 	virtual ~FHlslNiagaraTranslator() {}
-
-	virtual const FNiagaraTranslateResults &Translate(UNiagaraScript* InScript, FHlslNiagaraTranslatorOptions Options);
+	
+	virtual const FNiagaraTranslateResults &Translate(const FNiagaraCompileRequestData* InCompileData, const FNiagaraCompileOptions& InCompileOptions, FHlslNiagaraTranslatorOptions Options);
 	FNiagaraTranslatorOutput &GetTranslateOutput() { return CompilationOutput; }
 
-	virtual int32 CompilePin(UEdGraphPin* Pin);
+	virtual int32 CompilePin(const UEdGraphPin* Pin);
 
-	virtual int32 RegisterDataInterface(FNiagaraVariable& Var, UNiagaraDataInterface* DataInterface, bool bPlaceholder);
+	virtual int32 RegisterDataInterface(FNiagaraVariable& Var, UNiagaraDataInterface* DataInterface, bool bPlaceholder, bool bAddParameterMapRead);
 
 	virtual void Operation(class UNiagaraNodeOp* Operation, TArray<int32>& Inputs, TArray<int32>& Outputs);
 	virtual void Output(UNiagaraNodeOutput* OutputNode, const TArray<int32>& ComputedInputs);
@@ -355,7 +407,7 @@ public:
 	virtual int32 GetConstant(const FNiagaraVariable& Constant);
 	
 	virtual void ReadDataSet(const FNiagaraDataSetID DataSet, const TArray<FNiagaraVariable>& Variable, ENiagaraDataSetAccessMode AccessMode, int32 InputChunk, TArray<int32>& Outputs);
-	virtual void WriteDataSet(const FNiagaraDataSetID DataSet, const TArray<FNiagaraVariable>& Variable, ENiagaraDataSetAccessMode AccessMode, const TArray<int32>& Inputs);
+	virtual void WriteDataSet(const FNiagaraDataSetID DataSet, const TArray<FNiagaraVariable>& Variable, ENiagaraDataSetAccessMode AccessMode, const TArray<int32>& Inputs, TArray<int32>& Outputs);
 	virtual void ParameterMapSet(class UNiagaraNodeParameterMapSet* SetNode, TArray<int32>& Inputs, TArray<int32>& Outputs);
 	virtual void ParameterMapGet(class UNiagaraNodeParameterMapGet* GetNode, TArray<int32>& Inputs, TArray<int32>& Outputs);
 	virtual void Emitter(class UNiagaraNodeEmitter* GetNode, TArray<int32>& Inputs, TArray<int32>& Outputs);
@@ -367,11 +419,11 @@ public:
 	void DefineDataSetVariableReads(FString &HLSLOutput, FNiagaraDataSetID& Id, int32 DataSetIndex, TArray<FNiagaraVariable> & ReadVars);
 	void DefineDataSetVariableWrites(FString &HlslOutput, FNiagaraDataSetID& Id, int32 DataSetIndex, TArray<FNiagaraVariable>& WriteVars);
 	void DefineDataInterfaceHLSL(FString &HlslOutput);
-	TArray< FDIBufferDescriptorStore >&GetDatainterfaceBufferDescriptors() { return DIBufferDescriptors; }
+	TArray<FNiagaraDataInterfaceGPUParamInfo>& GetDataInterfaceParameters() { return DIParamInfo; }
 
 	// Format string should have up to 5 entries, {{0} = Computed Variable Suffix, {1} = Float or Int, {2} = Data Set Index, {3} = Register Index, {4} Default value for that type.
-	void GatherVariableForDataSetAccess(const FNiagaraVariable& Variable, FString Format, int32& RegisterIdx, int32 DataSetIndex, FString InstanceIdxSymbol, FString &HlslOutput);
-	void GatherComponentsForDataSetAccess(const UScriptStruct* Struct, FString VariableSymbol, bool bMatrixRoot, TArray<FString>& Components, TArray<ENiagaraBaseTypes>& Types);
+	void GatherVariableForDataSetAccess(const FNiagaraVariable& Variable, FString Format, int32& RegisterIdxInt, int32& RegisterIdxFloat, int32 DataSetIndex, FString InstanceIdxSymbol, FString &HlslOutput);
+	void GatherComponentsForDataSetAccess(UScriptStruct* Struct, FString VariableSymbol, bool bMatrixRoot, TArray<FString>& Components, TArray<ENiagaraBaseTypes>& Types);
 
 	FString CompileDataInterfaceFunction(UNiagaraDataInterface* DataInterface, FNiagaraFunctionSignature& Signature);
 
@@ -388,6 +440,7 @@ public:
 	virtual bool CanReadAttributes()const;
 	virtual ENiagaraScriptUsage GetTargetUsage() const;
 	FGuid GetTargetUsageId() const;
+	virtual ENiagaraScriptUsage GetCurrentUsage() const;
 
 	static bool IsBuiltInHlslType(FNiagaraTypeDefinition Type);
 	static FString GetStructHlslTypeName(FNiagaraTypeDefinition Type);
@@ -409,18 +462,22 @@ public:
 
 	static FString GetFunctionSignatureSymbol(const FNiagaraFunctionSignature& Sig);
 
-	/** ID to allow visitors over the node graph to mark nodes as visited. */
-	int32 VisitID;
 
 private:
+	void InitializeParameterMapDefaults(int32 ParamMapHistoryIdx);
+	void HandleParameterRead(int32 ParamMapHistoryIdx, const FNiagaraVariable& Var, const UEdGraphPin* DefaultPin, UNiagaraNode* ErrorNode, int32& OutputChunkId, bool bTreatAsUnknownParameterMap = false);
 	bool ShouldConsiderTargetParameterMap(ENiagaraScriptUsage InUsage) const;
-	FString BuildParameterMapHlslDefinitions(TArray<FNiagaraVariable>& PrimaryDataSetOutputEntries, TArray<TPair<FNiagaraVariable, const UEdGraphPin*> >& MissingDefaults);
-	FString BuildMissingDefaults(TArray<TPair<FNiagaraVariable, const UEdGraphPin*> >& MissingDefaults);
+	FString BuildParameterMapHlslDefinitions(TArray<FNiagaraVariable>& PrimaryDataSetOutputEntries);
+	FString BuildMissingDefaults();
+	void FinalResolveNamespacedTokens(const FString& ParameterMapInstanceNamespace, TArray<FString>& Tokens, TArray<FString>& ValidChildNamespaces, FNiagaraParameterMapHistoryBuilder& Builder, TArray<FNiagaraVariable>& UniqueParameterMapEntriesAliasesIntact, TArray<FNiagaraVariable>& UniqueParameterMapEntries, int32 ParamMapHistoryIdx);
 
 	void HandleNamespacedExternalVariablesToDataSetRead(TArray<FNiagaraVariable>& InDataSetVars, FString InNamespaceStr);
 
 	FString ComputeMatrixColumnAccess(const FString& Name);
 	FString ComputeMatrixRowAccess(const FString& Name);
+
+	void HandleCustomHlslNode(UNiagaraNodeCustomHlsl* CustomFunctionHlsl, ENiagaraScriptUsage& OutScriptUsage, FString& OutName, FString& OutFullName, bool& bOutCustomHlsl, FString& OutCustomHlsl,
+		FNiagaraFunctionSignature& OutSignature, TArray<int32>& Inputs);
 	
 	// Add a raw float constant chunk
 	int32 GetConstantFloat(float InValue);
@@ -442,7 +499,7 @@ private:
 	FString GetParameterMapInstanceName(int32 ParamMapHistoryIdx);
 	
 	// Register a System/Engine/read-only variable in its namespaced form
-	bool ParameterMapRegisterExternalConstantNamespaceVariable(FNiagaraVariable InVariable, UNiagaraNode* InNode, int32 InParamMapHistoryIdx, int32& Output, UEdGraphPin* InDefaultPin);
+	bool ParameterMapRegisterExternalConstantNamespaceVariable(FNiagaraVariable InVariable, UNiagaraNode* InNode, int32 InParamMapHistoryIdx, int32& Output, const UEdGraphPin* InDefaultPin);
 
 	// Register an attribute in its non namespaced form
 	bool ParameterMapRegisterNamespaceAttributeVariable(const FNiagaraVariable& InVariable, UNiagaraNode* InNode, int32 InParamMapHistoryIdx, int32& Output);
@@ -453,11 +510,13 @@ private:
 	void GenerateFunctionSignature(ENiagaraScriptUsage ScriptUsage, FString InName, const FString& InFullName, UNiagaraGraph* FuncGraph, TArray<int32>& Inputs, 
 		bool bHadNumericInputs, bool bHasParameterMapParameters, FNiagaraFunctionSignature& OutSig)const;
 
-	UNiagaraGraph* CloneGraphAndPrepareForCompilation(UNiagaraScript* InScript, UNiagaraScriptSource* InSource, bool bClearErrors);
+	UNiagaraGraph* CloneGraphAndPrepareForCompilation(const UNiagaraScript* InScript, const UNiagaraScriptSource* InSource, bool bClearErrors);
 
 	bool ShouldInterpolateParameter(const FNiagaraVariable& Parameter);
 
 	bool IsBulkSystemScript() const;
+	bool IsSpawnScript() const;
+	bool RequiresInterpolation() const;
 
 	/** Map of symbol names to count of times it's been used. Used for generating unique symbol names. */
 	TMap<FName, uint32> SymbolCounts;
@@ -500,15 +559,9 @@ private:
 	int32 ReadIdx;
 	int32 WriteIdx;
 
-	// Data interface buffer descriptors per data interface
-	TArray< FDIBufferDescriptorStore > DIBufferDescriptors;
-
-	/**
-	True if we're compiling the spawn section of an interpolated spawn script.
-	This will cause us to read from interpolated parameters and write to temp outputs that are fed back into the update script as inputs.
-	*/
-	bool bInsideInterpolatedSpawnScript;
-
+	// Parameter data per data interface.
+	TArray< FNiagaraDataInterfaceGPUParamInfo > DIParamInfo;
+	
 	/** Stack of currently tracked stats scopes. */
 	TArray<int32> StatScopeStack;
 
@@ -522,5 +575,10 @@ private:
 
 	ENiagaraCodeChunkMode CurrentBodyChunkMode;
 
-	TArray<UNiagaraDataInterface*> PlaceholderDataInterfaces;
+	TArray<FHlslNiagaraTranslationStage> TranslationStages;
+	int32 ActiveStageIdx;
+	bool bInitializedDefaults;
+
+	TArray<const UEdGraphPin*> CurrentDefaultPinTraversal;
+	TArray<FNiagaraVariable> InitialNamespaceVariablesMissingDefault;
 };
diff --git a/Engine/Plugins/FX/Niagara/Source/NiagaraEditor/Private/NiagaraNode.cpp b/Engine/Plugins/FX/Niagara/Source/NiagaraEditor/Private/NiagaraNode.cpp
--- a/Engine/Plugins/FX/Niagara/Source/NiagaraEditor/Private/NiagaraNode.cpp
+++ b/Engine/Plugins/FX/Niagara/Source/NiagaraEditor/Private/NiagaraNode.cpp
@@ -5,12 +5,12 @@
 #include "EdGraphSchema_Niagara.h"
 #include "NiagaraHlslTranslator.h"
 #include "GraphEditAction.h"
+#include "SNiagaraGraphNode.h"
 
 #define LOCTEXT_NAMESPACE "NiagaraNode"
 
 UNiagaraNode::UNiagaraNode(const FObjectInitializer& ObjectInitializer)
 	: Super(ObjectInitializer)
-	, VisitID(INDEX_NONE)
 {
 }
 
@@ -18,6 +18,11 @@ void UNiagaraNode::PostLoad()
 {
 	Super::PostLoad();
 
+	if (ChangeId.IsValid() == false)
+	{
+		ChangeId = FGuid::NewGuid();
+	}
+
 	if (GIsEditor && HasAllFlags(RF_Transactional) == false)
 	{
 		SetFlags(RF_Transactional);
@@ -103,10 +108,14 @@ bool UNiagaraNode::ReallocatePins()
 		OldPin->MarkPendingKill();
 	}
 
-	GetGraph()->NotifyGraphChanged();
+	//GetGraph()->NotifyGraphChanged();
 	if (!bAllSame)
 	{
-		CastChecked<UNiagaraGraph>(GetGraph())->NotifyGraphNeedsRecompile();
+		MarkNodeRequiresSynchronization(__FUNCTION__, true);
+	}
+	else
+	{
+		VisualsChangedDelegate.Broadcast(this);
 	}
 	return bAllSame;
 }
@@ -145,6 +154,14 @@ bool UNiagaraNode::CompileInputPins(FHlslNiagaraTranslator *Translator, TArray<i
 	return bError;
 }
 
+void UNiagaraNode::PostPlacedNewNode()
+{
+	if (ChangeId.IsValid() == false)
+	{
+		MarkNodeRequiresSynchronization(__FUNCTION__, false); // The add will have notified us anyway
+	}
+}
+
 void UNiagaraNode::AutowireNewNode(UEdGraphPin* FromPin)
 {
 	Super::AutowireNewNode(FromPin);
@@ -200,40 +217,73 @@ bool UNiagaraNode::ConvertNumericPinToType(UEdGraphPin* InGraphPin, FNiagaraType
 	return true;
 }
 
-void UNiagaraNode::PinDefaultValueChanged(UEdGraphPin* Pin) 
+TSharedPtr<SGraphNode> UNiagaraNode::CreateVisualWidget() 
 {
-	UNiagaraGraph* Graph = GetNiagaraGraph();
-	Graph->NotifyGraphNeedsRecompile();
+	return SNew(SNiagaraGraphNode, this);
 }
 
-void UNiagaraNode::PinConnectionListChanged(UEdGraphPin* Pin) 
+
+void UNiagaraNode::MarkNodeRequiresSynchronization(FString Reason, bool bRaiseGraphNeedsRecompile)
 {
-	UNiagaraGraph* Graph = GetNiagaraGraph();
-	Graph->NotifyGraphNeedsRecompile();
+	Modify();
+	ChangeId = FGuid::NewGuid();
+	//UE_LOG(LogNiagaraEditor, Verbose, TEXT("Node %s was marked requires synchronization.  Reason: %s"), *GetPathName(), *Reason);
 
-	const UEdGraphSchema_Niagara* Schema = CastChecked<UEdGraphSchema_Niagara>(GetSchema());
-	check(Schema);
-	FNiagaraTypeDefinition TypeDef = Schema->PinToTypeDefinition(Pin);
-	if (TypeDef == FNiagaraTypeDefinition::GetGenericNumericDef())
+	UNiagaraGraph* Graph = GetNiagaraGraph();
+	if (bRaiseGraphNeedsRecompile)
 	{
-		for (UEdGraphPin* LinkedPin : Pin->LinkedTo)
-		{
-			if (LinkedPin != nullptr)
-			{
-				FNiagaraTypeDefinition LinkedPinTypeDef = Schema->PinToTypeDefinition(LinkedPin);
-				if (LinkedPinTypeDef != FNiagaraTypeDefinition::GetGenericNumericDef())
-				{
+		Graph->NotifyGraphNeedsRecompile();
+	}
 
-				}
-			}
-		}
+	VisualsChangedDelegate.Broadcast(this);
+}
+
+void UNiagaraNode::ForceChangeId(const FGuid& InId, bool bRaiseGraphNeedsRecompile)
+{
+	Modify();
+	ChangeId = InId;
+
+	if (bRaiseGraphNeedsRecompile)
+	{
+		UNiagaraGraph* Graph = GetNiagaraGraph();
+		Graph->NotifyGraphNeedsRecompile();
 	}
 }
 
+void UNiagaraNode::PinDefaultValueChanged(UEdGraphPin* Pin) 
+{
+	MarkNodeRequiresSynchronization(__FUNCTION__, true);
+	Super::PinDefaultValueChanged(Pin);
+}
+
+void UNiagaraNode::OnRenameNode(const FString& NewName)
+{
+	MarkNodeRequiresSynchronization(__FUNCTION__, true);
+	Super::OnRenameNode(NewName);
+}
+
+void UNiagaraNode::OnPinRemoved(UEdGraphPin* InRemovedPin)
+{
+	MarkNodeRequiresSynchronization(__FUNCTION__, true);
+	Super::OnPinRemoved(InRemovedPin);
+}
+
+void UNiagaraNode::NodeConnectionListChanged()
+{
+	MarkNodeRequiresSynchronization(__FUNCTION__, true);
+	Super::NodeConnectionListChanged();
+}
+
+
+void UNiagaraNode::PinConnectionListChanged(UEdGraphPin* Pin) 
+{
+	MarkNodeRequiresSynchronization(__FUNCTION__, true);
+	Super::PinConnectionListChanged(Pin);
+}
+
 void UNiagaraNode::PinTypeChanged(UEdGraphPin* Pin) 
 {
-	UNiagaraGraph* Graph = GetNiagaraGraph();
-	Graph->NotifyGraphNeedsRecompile();
+	MarkNodeRequiresSynchronization(__FUNCTION__, true);
 }
 
 const UNiagaraGraph* UNiagaraNode::GetNiagaraGraph()const
@@ -341,11 +391,36 @@ void UNiagaraNode::GetOutputPins(TArray<class UEdGraphPin*>& OutOutputPins) cons
 	}
 }
 
+UEdGraphPin* UNiagaraNode::GetPinByPersistentGuid(const FGuid& InPersistentGuid) const
+{
+	for (UEdGraphPin* Pin : Pins)
+	{
+		if (InPersistentGuid == Pin->PersistentGuid)
+		{
+			return Pin;
+		}
+	}
+
+	return nullptr;
+}
+
 ENiagaraNumericOutputTypeSelectionMode UNiagaraNode::GetNumericOutputTypeSelectionMode() const
 {
 	return ENiagaraNumericOutputTypeSelectionMode::None;
 }
 
+
+UEdGraphPin* UNiagaraNode::TraceOutputPin(UEdGraphPin* LocallyOwnedOutputPin)
+{
+	if (LocallyOwnedOutputPin == nullptr)
+	{
+		return nullptr;
+	}
+
+	UNiagaraNode* LinkedNode = CastChecked<UNiagaraNode>(LocallyOwnedOutputPin->GetOwningNode());
+	return LinkedNode->GetTracedOutputPin(LocallyOwnedOutputPin);
+}
+
 void UNiagaraNode::RouteParameterMapAroundMe(FNiagaraParameterMapHistoryBuilder& OutHistory, bool bRecursive)
 {
 	const UEdGraphSchema_Niagara* Schema = CastChecked<UEdGraphSchema_Niagara>(GetSchema());
@@ -382,5 +457,10 @@ void UNiagaraNode::RouteParameterMapAroundMe(FNiagaraParameterMapHistoryBuilder&
 		OutHistory.RegisterParameterMapPin(PMIdx, OutputPin);
 	}
 }
+UNiagaraNode::FOnNodeVisualsChanged& UNiagaraNode::OnVisualsChanged()
+{
+	return VisualsChangedDelegate;
+}
+
 
 #undef LOCTEXT_NAMESPACE
\ No newline at end of file
diff --git a/Engine/Plugins/FX/Niagara/Source/NiagaraEditor/Private/NiagaraNodeAssignment.cpp b/Engine/Plugins/FX/Niagara/Source/NiagaraEditor/Private/NiagaraNodeAssignment.cpp
--- a/Engine/Plugins/FX/Niagara/Source/NiagaraEditor/Private/NiagaraNodeAssignment.cpp
+++ b/Engine/Plugins/FX/Niagara/Source/NiagaraEditor/Private/NiagaraNodeAssignment.cpp
@@ -17,6 +17,9 @@
 #include "NiagaraNodeParameterMapSet.h"
 #include "NiagaraConstants.h"
 #include "ViewModels/Stack/NiagaraParameterHandle.h"
+#include "Framework/MultiBox/MultiBoxBuilder.h"
+#include "ScopedTransaction.h"
+#include "ViewModels/Stack/NiagaraStackGraphUtilities.h"
 
 #define LOCTEXT_NAMESPACE "NiagaraNodeAssigment"
 
@@ -28,7 +31,7 @@ void UNiagaraNodeAssignment::AllocateDefaultPins()
 
 FText UNiagaraNodeAssignment::GetNodeTitle(ENodeTitleType::Type TitleType) const
 {
-	return FText::Format(LOCTEXT("NodeTitle", "Set {0}"), FText::FromName(AssignmentTarget.GetName()));
+	return LOCTEXT("NodeTitle", "Set Variables");
 }
 
 bool UNiagaraNodeAssignment::RefreshFromExternalChanges()
@@ -39,23 +42,230 @@ bool UNiagaraNodeAssignment::RefreshFromExternalChanges()
 	return true;
 }
 
-void UNiagaraNodeAssignment::GenerateScript()
+void UNiagaraNodeAssignment::PostLoad()
 {
-	FString NodeName = TEXT("Set_") + AssignmentTarget.GetName().ToString();
-	NodeName = NodeName.Replace(TEXT("."), TEXT("_"));
+	Super::PostLoad();
+
+	// Handle the case where we moved towards an array of assignment targets...
+	if (AssignmentTarget_DEPRECATED.IsValid() && AssignmentTargets.Num() == 0)
+	{
+		AssignmentTargets.Add(AssignmentTarget_DEPRECATED);
+		AssignmentDefaultValues.Add(AssignmentDefaultValue_DEPRECATED);
+		OldFunctionCallName = FunctionDisplayName;
+		FunctionDisplayName.Empty();
+		RefreshFromExternalChanges();
+
+		UE_LOG(LogNiagaraEditor, Log, TEXT("Found old Assignment Node, converting variable \"%s\" in \"%s\""), *AssignmentTarget_DEPRECATED.GetName().ToString(), *GetFullName());
+
+		MarkNodeRequiresSynchronization(__FUNCTION__, true);
+		
+		// Deduce what rapid iteration variable we would have previously been and prepare to change
+		// any instances of it.
+		TMap<FNiagaraVariable, FNiagaraVariable> Converted;
+		FNiagaraParameterHandle TargetHandle(AssignmentTarget_DEPRECATED.GetName());
+		FString VarNamespace = TargetHandle.GetNamespace().ToString();
+		TMap<FString, FString> AliasMap;
+		AliasMap.Add(OldFunctionCallName, FunctionDisplayName + TEXT(".") + VarNamespace);
+		FNiagaraVariable RemapVar = FNiagaraVariable(AssignmentTarget_DEPRECATED.GetType(), *(OldFunctionCallName + TEXT(".") + TargetHandle.GetName().ToString()));
+		FNiagaraVariable NewVar = FNiagaraParameterMapHistory::ResolveAliases(RemapVar, AliasMap, TEXT("."));
+		Converted.Add(RemapVar, NewVar);
+
+		bool bConvertedAnything = false;
+
+		// Now clean up the input set node going into us...
+		UEdGraphPin* Pin = GetInputPin(0);
+		if (Pin != nullptr && Pin->LinkedTo.Num() == 1)
+		{
+			// Likely we have a set node going into us, check to see if it has any variables that need to be cleaned up.
+			UNiagaraNodeParameterMapSet* SetNode = Cast<UNiagaraNodeParameterMapSet>(Pin->LinkedTo[0]->GetOwningNode());
+			if (SetNode)
+			{
+				SetNode->ConditionalPostLoad();
 
+				TArray<UEdGraphPin*> InputPins;
+				SetNode->GetInputPins(InputPins);
+
+				const UEdGraphSchema_Niagara* NiagaraSchema = GetDefault<UEdGraphSchema_Niagara>();
+				for (UEdGraphPin* InputPin : InputPins)
+				{
+					FNiagaraVariable Var = NiagaraSchema->PinToNiagaraVariable(InputPin, false);
+					if (Var.GetName() == RemapVar.GetName())
+					{
+						if (NewVar != Var)
+						{
+							SetNode->SetPinName(InputPin, NewVar.GetName());
+							UE_LOG(LogNiagaraEditor, Log, TEXT("Converted Set pin variable \"%s\" to \"%s\" in \"%s\""), *Var.GetName().ToString(), *NewVar.GetName().ToString(), *GetFullName());
+							bConvertedAnything = true;
+						}
+					}
+				}
+			}
+			else if (Pin->LinkedTo[0]->GetOwningNode())
+			{
+				// Sometimes we don't automatically have set nodes between modules in the stack... just skip over these.
+				UE_LOG(LogNiagaraEditor, Log, TEXT("Found node \"%s\" attached to assignment \"%s\" variable %s"), *Pin->LinkedTo[0]->GetOwningNode()->GetFullName(), *GetFullName(), *NewVar.GetName().ToString());
+			}
+		}
+
+		// Now we need to find the scripts affecting this node... we cheat and walk up our ownership hierarchy until we find a system or emitter.
+		if (Converted.Num() != 0)
+		{
+			UNiagaraEmitter* Emitter = nullptr;
+			UNiagaraSystem* System = nullptr;
+			UObject* OuterObj = GetOuter();
+			while (OuterObj != nullptr)
+			{
+				if (Emitter == nullptr)
+				{
+					Emitter = Cast<UNiagaraEmitter>(OuterObj);
+				}
+				if (System == nullptr)
+				{
+					System = Cast<UNiagaraSystem>(OuterObj);
+				}
+
+				OuterObj = OuterObj->GetOuter();
+			}
+
+			// Gather up the affected scripts from the relevant owner...
+			TArray<UNiagaraScript*> Scripts;
+			if (Emitter)
+			{
+				Emitter->GetScripts(Scripts, false);
+			}
+			if (System)
+			{
+				Scripts.Add(System->GetSystemSpawnScript());
+				Scripts.Add(System->GetSystemUpdateScript());
+			}
+
+			for (UNiagaraScript* Script : Scripts)
+			{
+				if (Script->HandleVariableRenames(Converted, Emitter ? Emitter->GetUniqueEmitterName() : FString()))
+				{
+					bConvertedAnything = true;
+				}
+			}
+		}
+
+		if (!bConvertedAnything)
+		{
+			UE_LOG(LogNiagaraEditor, Log, TEXT("Found old Assignment Node, nothing was attached???? variable \"%s\" in \"%s\""), *AssignmentTarget_DEPRECATED.GetName().ToString(), *GetFullName());
+		}
+	}
+}
+
+void UNiagaraNodeAssignment::BuildParameterMapHistory(FNiagaraParameterMapHistoryBuilder& OutHistory, bool bRecursive)
+{
+	Super::BuildParameterMapHistory(OutHistory, bRecursive);
+}
+
+void UNiagaraNodeAssignment::GenerateScript()
+{
 	if (FunctionScript == nullptr)
 	{
-		FunctionScript = NewObject<UNiagaraScript>(this, FName(*NodeName), RF_Transactional);
+		FunctionScript = NewObject<UNiagaraScript>(this, FName(*(TEXT("SetVariables_") + NodeGuid.ToString())), RF_Transactional);
 		FunctionScript->SetUsage(ENiagaraScriptUsage::Module);
+		FunctionScript->Description = LOCTEXT("AssignmentNodeDesc", "Sets one or more variables in the stack.");
 		InitializeScript(FunctionScript);
 		ComputeNodeName();
 	}
-	if (FunctionScript->GetName() != NodeName)
+}
+
+void UNiagaraNodeAssignment::MergeUp()
+{
+	//NiagaraStackUtilities::
+}
+
+void UNiagaraNodeAssignment::BuildParameterMenu(FMenuBuilder& MenuBuilder, ENiagaraScriptUsage InUsage, UNiagaraNodeOutput* InGraphOutputNode)
+{
+	TArray<FNiagaraParameterMapHistory> Histories = UNiagaraNodeParameterMapBase::GetParameterMaps(InGraphOutputNode->GetNiagaraGraph());
+	TArray<FNiagaraVariable> AvailableParameters;
+
+	if (InGraphOutputNode->GetUsage() == ENiagaraScriptUsage::ParticleSpawnScript ||
+		InGraphOutputNode->GetUsage() == ENiagaraScriptUsage::ParticleSpawnScriptInterpolated ||
+		InGraphOutputNode->GetUsage() == ENiagaraScriptUsage::ParticleUpdateScript ||
+		InGraphOutputNode->GetUsage() == ENiagaraScriptUsage::ParticleEventScript)
 	{
-		FunctionScript->Rename(*NodeName);
-		ComputeNodeName();
+		AvailableParameters.Append(FNiagaraConstants::GetCommonParticleAttributes());
+	}
+
+	for (FNiagaraParameterMapHistory& History : Histories)
+	{
+		for (FNiagaraVariable& Variable : History.Variables)
+		{
+			if (History.IsPrimaryDataSetOutput(Variable, InGraphOutputNode->GetUsage()))
+			{
+				AvailableParameters.AddUnique(Variable);
+			}
+		}
+	}
+
+	for (const FNiagaraVariable& AvailableParameter : AvailableParameters)
+	{
+		FString DisplayNameString = FName::NameToDisplayString(AvailableParameter.GetName().ToString(), false);
+		const FText NameText = FText::FromString(DisplayNameString);
+		FText VarDesc = FNiagaraConstants::GetAttributeDescription(AvailableParameter);
+		FString VarDefaultValue = FNiagaraConstants::GetAttributeDefaultValue(AvailableParameter);
+		const FText TooltipDesc = FText::Format(LOCTEXT("SetFunctionPopupTooltip", "Description: Set the parameter {0}. {1}"), FText::FromName(AvailableParameter.GetName()), VarDesc);
+		FText CategoryName = LOCTEXT("ModuleSetCategory", "Set Specific Parameters");
+
+		MenuBuilder.AddMenuEntry(
+			NameText,
+			TooltipDesc,
+			FSlateIcon(),
+			FUIAction(FExecuteAction::CreateUObject(this, &UNiagaraNodeAssignment::AddParameter, AvailableParameter, VarDefaultValue)));
+	}
+
+}
+
+void UNiagaraNodeAssignment::AddParameter(FNiagaraVariable InVar, FString InDefaultValue)
+{
+	const FText TransactionDesc = FText::Format(LOCTEXT("SetFunctionTransactionDesc", "Add the parameter {0}."), FText::FromName(InVar.GetName()));
+	FScopedTransaction ScopedTransaction(TransactionDesc);
+
+	// Since we blow away the graph, we need to cache *everything* we create potentially.
+	Modify();
+	FunctionScript->Modify();
+	UNiagaraScriptSource* Source = Cast<UNiagaraScriptSource>(FunctionScript->GetSource());
+	Source->Modify();
+	Source->NodeGraph->Modify();
+	for (UEdGraphNode* Node : Source->NodeGraph->Nodes)
+	{
+		Node->Modify();
 	}
+
+	AddAssignmentTarget(InVar, &InDefaultValue);
+
+	RefreshFromExternalChanges();
+	OnInputsChangedDelegate.Broadcast();
+}
+
+void UNiagaraNodeAssignment::RemoveParameter(const FNiagaraVariable& InVar)
+{
+	const FText TransactionDesc = FText::Format(LOCTEXT("RemoveFunctionTransactionDesc", "Remove the parameter {0}."), FText::FromName(InVar.GetName()));
+	FScopedTransaction ScopedTransaction(TransactionDesc);
+
+	// Since we blow away the graph, we need to cache *everything* we create potentially.
+	Modify();
+	FunctionScript->Modify();
+	UNiagaraScriptSource* Source = Cast<UNiagaraScriptSource>(FunctionScript->GetSource());
+	Source->Modify();
+	Source->NodeGraph->Modify();
+	for (UEdGraphNode* Node : Source->NodeGraph->Nodes)
+	{
+		Node->Modify();
+	}
+
+	int32 Index = INDEX_NONE;
+	if (AssignmentTargets.Find(InVar, Index))
+	{
+		AssignmentTargets.RemoveAt(Index);
+		AssignmentDefaultValues.RemoveAt(Index);
+	}
+
+	RefreshFromExternalChanges();
+	OnInputsChangedDelegate.Broadcast();
 }
 
 void UNiagaraNodeAssignment::InitializeScript(UNiagaraScript* NewScript)
@@ -154,37 +364,134 @@ void UNiagaraNodeAssignment::InitializeScript(UNiagaraScript* NewScript)
 			GetNodes[1]->RemovePin(GetNodes[1]->GetInputPin(0));
 		}
 
-		// Now create the proper new pins and connect them.
-		if (AssignmentTarget.GetName() != NAME_None)
-		{
-			FNiagaraParameterHandle TargetHandle(AssignmentTarget.GetName());
-			UEdGraphPin* SetPin = SetNodes[0]->RequestNewTypedPin(EGPD_Input, AssignmentTarget.GetType(), AssignmentTarget.GetName());
-			UEdGraphPin* GetPin = GetNodes[0]->RequestNewTypedPin(EGPD_Output, AssignmentTarget.GetType(), *FString::Printf(TEXT("Module.%s"), *TargetHandle.GetName().ToString()));
+		const UEdGraphSchema_Niagara* NiagaraSchema = GetDefault<UEdGraphSchema_Niagara>();
 
-			GetPin->MakeLinkTo(SetPin);
+		for (int32 i = 0; i < AssignmentTargets.Num(); i++)
+		{
+			// Now create the proper new pins and connect them.
+			FName Name = AssignmentTargets[i].GetName();
+			FNiagaraTypeDefinition Type = AssignmentTargets[i].GetType();
+			FString DefaultValue = AssignmentDefaultValues[i];
 
-			if (!AssignmentDefaultValue.IsEmpty())
+			if (Name != NAME_None)
 			{
-				UEdGraphPin* DefaultInputPin = GetNodes[0]->GetDefaultPin(GetPin);
+				FNiagaraParameterHandle TargetHandle(Name);
+				UEdGraphPin* SetPin = SetNodes[0]->RequestNewTypedPin(EGPD_Input, Type, Name);
+				FString ModuleVarName = FString::Printf(TEXT("Module.%s"), *TargetHandle.GetParameterHandleString().ToString());
+				UEdGraphPin* GetPin = GetNodes[0]->RequestNewTypedPin(EGPD_Output, Type, *ModuleVarName);
+				FNiagaraVariable TargetVar = NiagaraSchema->PinToNiagaraVariable(GetPin, false);
+				GetPin->MakeLinkTo(SetPin);
 
-				bool isEngineConstant = false;
-				FNiagaraVariable SeekVar = FNiagaraVariable(AssignmentTarget.GetType(), FName(*AssignmentDefaultValue));
-				const FNiagaraVariable* FoundVar = FNiagaraConstants::FindEngineConstant(SeekVar);
-				if (FoundVar != nullptr)
+				if (!DefaultValue.IsEmpty())
 				{
-					UEdGraphPin* DefaultGetPin = GetNodes[1]->RequestNewTypedPin(EGPD_Output, AssignmentTarget.GetType(), FoundVar->GetName());
-					DefaultGetPin->MakeLinkTo(DefaultInputPin);
+					UEdGraphPin* DefaultInputPin = GetNodes[0]->GetDefaultPin(GetPin);
+
+					bool isEngineConstant = false;
+					FNiagaraVariable SeekVar = FNiagaraVariable(Type, FName(*DefaultValue));
+					const FNiagaraVariable* FoundVar = FNiagaraConstants::FindEngineConstant(SeekVar);
+					if (FoundVar != nullptr)
+					{
+						UEdGraphPin* DefaultGetPin = GetNodes[1]->RequestNewTypedPin(EGPD_Output, Type, FoundVar->GetName());
+						DefaultGetPin->MakeLinkTo(DefaultInputPin);
+					}
+					else
+					{
+						DefaultInputPin->bDefaultValueIsIgnored = false;
+						DefaultInputPin->DefaultValue = DefaultValue;
+					}
 				}
-				else
+
+				if (FNiagaraConstants::IsNiagaraConstant(AssignmentTargets[i]))
 				{
-					DefaultInputPin->bDefaultValueIsIgnored = false;
-					DefaultInputPin->DefaultValue = AssignmentDefaultValue;
+					const FNiagaraVariableMetaData* FoundMetaData = FNiagaraConstants::GetConstantMetaData(AssignmentTargets[i]);
+					if (FoundMetaData)
+					{
+						FNiagaraVariableMetaData& MetaData = CreatedGraph->FindOrAddMetaData(TargetVar);
+						MetaData.Description = FoundMetaData->Description;
+						MetaData.ReferencerNodes.Empty();
+						MetaData.ReferencerNodes.Add(GetNodes[0]);
+					}
 				}
 			}
 		}
 
+		CreatedGraph->PurgeUnreferencedMetaData();
+	}
+}
 
+int32 UNiagaraNodeAssignment::FindAssignmentTarget(const FName& InName, const FNiagaraTypeDefinition& InType)
+{
+	for (int32 i = 0; i < AssignmentTargets.Num(); i++)
+	{
+		FName Name = AssignmentTargets[i].GetName();
+		FNiagaraTypeDefinition Type = AssignmentTargets[i].GetType();
+	
+		if (InName == Name && InType == Type)
+		{
+			return i;
+		}
+	}
+
+	return INDEX_NONE;
+}
+
+int32 UNiagaraNodeAssignment::FindAssignmentTarget(const FName& InName)
+{
+	for (int32 i = 0; i < AssignmentTargets.Num(); i++)
+	{
+		FName Name = AssignmentTargets[i].GetName();
+		FNiagaraTypeDefinition Type = AssignmentTargets[i].GetType();
+
+		if (InName == Name)
+		{
+			return i;
+		}
+	}
+
+	return INDEX_NONE;
+}
+
+int32 UNiagaraNodeAssignment::AddAssignmentTarget(const FNiagaraVariable& InVar, const FString* InDefaultValue)
+{
+	int32 IdxA = AssignmentTargets.AddDefaulted();
+	int32 IdxB = AssignmentDefaultValues.AddDefaulted();
+	check(IdxA == IdxB);
+	SetAssignmentTarget(IdxA, InVar, InDefaultValue);
+	return IdxA;
+}
+
+bool UNiagaraNodeAssignment::SetAssignmentTarget(int32 Idx, const FNiagaraVariable& InVar, const FString* InDefaultValue)
+{
+	check(Idx < AssignmentTargets.Num());
+
+	bool bRetValue = false;
+	if (InVar != AssignmentTargets[Idx])
+	{
+		AssignmentTargets[Idx] = InVar;
+		MarkNodeRequiresSynchronization(__FUNCTION__, true);
+		bRetValue = true;
 	}
+	
+	if (InDefaultValue != nullptr && AssignmentDefaultValues[Idx] != *InDefaultValue)
+	{ 
+		AssignmentDefaultValues[Idx] = *InDefaultValue;
+		MarkNodeRequiresSynchronization(__FUNCTION__, true);
+		bRetValue = true;
+	}
+	return bRetValue;
 }
 
+bool UNiagaraNodeAssignment::SetAssignmentTargetName(int32 Idx, const FName& InName)
+{
+	check(Idx < AssignmentTargets.Num());
+	if (AssignmentTargets[Idx].GetName() != InName)
+	{
+		AssignmentTargets[Idx].SetName(InName);
+		MarkNodeRequiresSynchronization(__FUNCTION__, true);
+		return true;
+	}
+	return false;
+}
+
+
 #undef LOCTEXT_NAMESPACE
diff --git a/Engine/Plugins/FX/Niagara/Source/NiagaraEditor/Private/NiagaraNodeConvert.cpp b/Engine/Plugins/FX/Niagara/Source/NiagaraEditor/Private/NiagaraNodeConvert.cpp
--- a/Engine/Plugins/FX/Niagara/Source/NiagaraEditor/Private/NiagaraNodeConvert.cpp
+++ b/Engine/Plugins/FX/Niagara/Source/NiagaraEditor/Private/NiagaraNodeConvert.cpp
@@ -196,7 +196,9 @@ void UNiagaraNodeConvert::AutowireNewNode(UEdGraphPin* FromPin)
 			}
 		}
 	}
-	GetGraph()->NotifyGraphChanged();
+
+	MarkNodeRequiresSynchronization(__FUNCTION__, true);
+	//GetGraph()->NotifyGraphChanged();
 }
 
 FText UNiagaraNodeConvert::GetNodeTitle(ENodeTitleType::Type TitleType)const
diff --git a/Engine/Plugins/FX/Niagara/Source/NiagaraEditor/Private/NiagaraNodeCustomHlsl.cpp b/Engine/Plugins/FX/Niagara/Source/NiagaraEditor/Private/NiagaraNodeCustomHlsl.cpp
--- a/Engine/Plugins/FX/Niagara/Source/NiagaraEditor/Private/NiagaraNodeCustomHlsl.cpp
+++ b/Engine/Plugins/FX/Niagara/Source/NiagaraEditor/Private/NiagaraNodeCustomHlsl.cpp
@@ -38,11 +38,15 @@ FText UNiagaraNodeCustomHlsl::GetHlslText() const
 
 void UNiagaraNodeCustomHlsl::OnCustomHlslTextCommitted(const FText& InText, ETextCommit::Type InType)
 {
-	FScopedTransaction Transaction(LOCTEXT("CustomHlslCommit", "Edited Custom Hlsl"));
-	Modify();
-	CustomHlsl = InText.ToString();
-	RefreshFromExternalChanges();
-	GetNiagaraGraph()->NotifyGraphNeedsRecompile();
+	FString NewValue = InText.ToString();
+	if (NewValue != CustomHlsl)
+	{
+		FScopedTransaction Transaction(LOCTEXT("CustomHlslCommit", "Edited Custom Hlsl"));
+		Modify();
+		CustomHlsl = NewValue;
+		RefreshFromExternalChanges();			
+		MarkNodeRequiresSynchronization(__FUNCTION__, true);
+	}
 }
 
 FLinearColor UNiagaraNodeCustomHlsl::GetNodeTitleColor() const
@@ -50,6 +54,92 @@ FLinearColor UNiagaraNodeCustomHlsl::GetNodeTitleColor() const
 	return UEdGraphSchema_Niagara::NodeTitleColor_CustomHlsl;
 }
 
+bool UNiagaraNodeCustomHlsl::GetTokens(TArray<FString>& OutTokens) const
+{
+	FString HlslData = *CustomHlsl;
+
+	if (HlslData.Len() == 0)
+	{
+		return false;
+	}
+	
+	FString InputVars = TEXT(";/*+-)(?:, \t\n");
+	int32 LastValidIdx = INDEX_NONE;
+	bool bComment = false;
+	int32 TargetLength = HlslData.Len();
+	for (int32 i = 0; i < TargetLength; )
+	{
+		int32 Index = INDEX_NONE;
+
+		// Determine if we are a splitter character or a regular character.
+		if (InputVars.FindChar(HlslData[i], Index) && Index != INDEX_NONE)
+		{
+			if (HlslData[i] == '/' && (i + 1 != TargetLength) && HlslData[i+1] == '/') // We are a comment, go to end of line.
+			{
+				int32 FoundEndIdx = HlslData.Find("\n", ESearchCase::IgnoreCase, ESearchDir::FromStart, i + 2);
+				if (FoundEndIdx != INDEX_NONE)
+				{
+					// do nothing // FoundEndIdx = FoundEndIdx;
+				}
+				else
+				{
+					FoundEndIdx = TargetLength - 1;
+				}
+
+				OutTokens.Add(HlslData.Mid(i, FoundEndIdx - i));
+				OutTokens.Add("\n");
+				i = FoundEndIdx + 1;
+			}
+			else if (HlslData[i] == '/' && (i + 1 != TargetLength) && HlslData[i+1] == '*') // We are a multiline comment, go to end comment. Nested comments are not supported currently.
+			{
+				int32 FoundEndIdx = HlslData.Find("*/", ESearchCase::IgnoreCase, ESearchDir::FromStart, i + 2);
+				if (FoundEndIdx != INDEX_NONE)
+				{
+					FoundEndIdx = FoundEndIdx + 1;
+				}
+				else
+				{
+					FoundEndIdx = TargetLength - 1;
+				}
+
+				OutTokens.Add(HlslData.Mid(i, FoundEndIdx - i));
+				i = FoundEndIdx + 1;
+			}
+			else if (LastValidIdx != INDEX_NONE) // We encountered a non-splitter character in the past that hasn't been recorded yet but we are a splitter.
+			{
+				OutTokens.Add(HlslData.Mid(LastValidIdx, i - LastValidIdx));
+				OutTokens.Add(FString(1, &HlslData[i]));
+				i++;
+			}
+			else //if (LastValidIdx == INDEX_NONE) // We are a splitter with no known unrecorded tokens prior.
+			{
+				OutTokens.Add(FString(1, &HlslData[i]));
+				i++;
+			}
+
+			LastValidIdx = INDEX_NONE;
+		}
+		else if (LastValidIdx == INDEX_NONE)
+		{
+			LastValidIdx = i; // Record that this is where we encountered the first non-splitter character that has yet to be recorded.
+			i++;
+		}		
+		else
+		{
+			i++;
+		}
+	}
+
+	// We may need to pull in the last chars from the end.
+	if (LastValidIdx != INDEX_NONE)
+	{
+		OutTokens.Add(HlslData.Mid(LastValidIdx));
+	}
+
+	return true;
+}
+
+
 #if WITH_EDITOR
 
 void UNiagaraNodeCustomHlsl::PostEditChangeProperty(struct FPropertyChangedEvent& PropertyChangedEvent)
@@ -62,6 +152,16 @@ void UNiagaraNodeCustomHlsl::PostEditChangeProperty(struct FPropertyChangedEvent
 		GetNiagaraGraph()->NotifyGraphNeedsRecompile();
 	}
 }
+
+void UNiagaraNodeCustomHlsl::InitAsCustomHlslDynamicInput(const FNiagaraTypeDefinition& OutputType)
+{
+	Modify();
+	ReallocatePins();
+	RequestNewTypedPin(EGPD_Input, FNiagaraTypeDefinition::GetParameterMapDef(), FName("Map"));
+	RequestNewTypedPin(EGPD_Output, OutputType, FName("Output"));
+	ScriptUsage = ENiagaraScriptUsage::DynamicInput;
+}
+
 #endif
 
 /** Called when a new typed pin is added by the user. */
@@ -95,6 +195,121 @@ void UNiagaraNodeCustomHlsl::MoveDynamicPin(UEdGraphPin* Pin, int32 DirectionToM
 	RebuildSignatureFromPins();
 }
 
+void UNiagaraNodeCustomHlsl::BuildParameterMapHistory(FNiagaraParameterMapHistoryBuilder& OutHistory, bool bRecursive)
+{
+	Super::BuildParameterMapHistory(OutHistory, bRecursive);
+	if (!IsNodeEnabled() && OutHistory.GetIgnoreDisabled())
+	{
+		RouteParameterMapAroundMe(OutHistory, bRecursive);
+		return;
+	}
+
+	TArray<FString> Tokens;
+	GetTokens(Tokens);
+
+	TArray<UEdGraphPin*> InputPins;
+	GetInputPins(InputPins);
+
+	TArray<UEdGraphPin*> OutputPins;
+	GetOutputPins(OutputPins);
+
+	int32 ParamMapIdx = INDEX_NONE;
+	TArray<FNiagaraVariable> LocalVars;
+	// This only works currently if the input pins are in the same order as the signature pins.
+	if (InputPins.Num() == Signature.Inputs.Num() + 1)// the add pin is extra
+	{
+		bool bHasParamMapInput = false;
+		bool bHasParamMapOutput = false;
+		for (int32 i = 0; i < InputPins.Num(); i++)
+		{
+			if (IsAddPin(InputPins[i]))
+				continue;
+
+			FNiagaraVariable Input = Signature.Inputs[i];
+			if (Input.GetType() == FNiagaraTypeDefinition::GetParameterMapDef())
+			{
+				bHasParamMapInput = true;
+				FString ReplaceSrc = Input.GetName().ToString() + TEXT(".");
+				ReplaceExactMatchTokens(Tokens, ReplaceSrc, TEXT(""), false);
+				if (InputPins[i]->LinkedTo.Num() != 0)
+				{
+					ParamMapIdx = OutHistory.TraceParameterMapOutputPin(InputPins[i]->LinkedTo[0]);
+				}
+			}
+			else
+			{
+				LocalVars.Add(Input);
+			}
+		}
+
+		for (int32 i = 0; i < OutputPins.Num(); i++)
+		{
+			if (IsAddPin(OutputPins[i]))
+				continue;
+
+			FNiagaraVariable Output = Signature.Outputs[i];
+			if (Output.GetType() == FNiagaraTypeDefinition::GetParameterMapDef())
+			{
+				bHasParamMapOutput = true;
+				FString ReplaceSrc = Output.GetName().ToString() + TEXT(".");
+				ReplaceExactMatchTokens(Tokens, ReplaceSrc, TEXT(""), false);
+			}
+			else
+			{
+				LocalVars.Add(Output);
+			}
+		}
+
+		TArray<FString> PossibleNamespaces;
+		FNiagaraParameterMapHistory::GetValidNamespacesForReading(OutHistory.GetBaseUsageContext(), PossibleNamespaces);
+
+		if ((bHasParamMapOutput || bHasParamMapInput) && ParamMapIdx != INDEX_NONE)
+		{
+			for (int32 i = 0; i < Tokens.Num(); i++)
+			{
+				bool bFoundLocal = false;
+				if (INDEX_NONE != FNiagaraVariable::SearchArrayForPartialNameMatch(LocalVars, *Tokens[i]))
+				{
+					bFoundLocal = true;
+				}
+
+				if (!bFoundLocal && Tokens[i].Contains(TEXT("."))) // Only check tokens with namespaces in them..
+				{
+					for (const FString& ValidNamespace : PossibleNamespaces)
+					{
+						// There is one possible path here, one where we're using the namespace as-is from the valid list.
+						if (Tokens[i].StartsWith(ValidNamespace, ESearchCase::CaseSensitive))
+						{
+							bool bUsedDefault = false;
+							OutHistory.HandleExternalVariableRead(ParamMapIdx, *Tokens[i]);
+						}
+					}
+				}
+			}
+		}
+	}
+}
+
+// Replace items in the tokens array if they start with the src string or optionally src string and a namespace delimiter
+void UNiagaraNodeCustomHlsl::ReplaceExactMatchTokens(TArray<FString>& Tokens, const FString& SrcString, const FString& ReplaceString, bool bAllowNamespaceSeparation)
+{
+	for (int32 i = 0; i < Tokens.Num(); i++)
+	{
+		if (Tokens[i].StartsWith(SrcString, ESearchCase::CaseSensitive))
+		{
+			if (Tokens[i].Len() > SrcString.Len() && Tokens[i][SrcString.Len()] == '.' && bAllowNamespaceSeparation)
+			{
+				Tokens[i] = ReplaceString + Tokens[i].Mid(SrcString.Len());
+			}
+			else if (Tokens[i].Len() == SrcString.Len())
+			{
+				Tokens[i] = ReplaceString;
+			}
+		}
+	}
+}
+
+
 void UNiagaraNodeCustomHlsl::RebuildSignatureFromPins()
 {
 	Modify();
@@ -130,7 +345,7 @@ void UNiagaraNodeCustomHlsl::RebuildSignatureFromPins()
 	Signature = Sig;
 
 	RefreshFromExternalChanges();
-	GetNiagaraGraph()->NotifyGraphNeedsRecompile();
+	MarkNodeRequiresSynchronization(__FUNCTION__, true);
 }
 
 #undef LOCTEXT_NAMESPACE
diff --git a/Engine/Plugins/FX/Niagara/Source/NiagaraEditor/Private/NiagaraNodeDataSetBase.cpp b/Engine/Plugins/FX/Niagara/Source/NiagaraEditor/Private/NiagaraNodeDataSetBase.cpp
--- a/Engine/Plugins/FX/Niagara/Source/NiagaraEditor/Private/NiagaraNodeDataSetBase.cpp
+++ b/Engine/Plugins/FX/Niagara/Source/NiagaraEditor/Private/NiagaraNodeDataSetBase.cpp
@@ -9,6 +9,8 @@
 #define LOCTEXT_NAMESPACE "UNiagaraNodeDataSetBase"
 
 const FName UNiagaraNodeDataSetBase::ConditionVarName(TEXT("__Condition"));
+const FName UNiagaraNodeDataSetBase::ParamMapInVarName(TEXT("__ParamMapIn"));
+const FName UNiagaraNodeDataSetBase::ParamMapOutVarName(TEXT("__ParamMapOut"));
 
 UNiagaraNodeDataSetBase::UNiagaraNodeDataSetBase(const FObjectInitializer& ObjectInitializer)
 : Super(ObjectInitializer), ExternalStructAsset(nullptr)
@@ -54,6 +56,17 @@ FLinearColor UNiagaraNodeDataSetBase::GetNodeTitleColor() const
 	return CastChecked<UEdGraphSchema_Niagara>(GetSchema())->NodeTitleColor_Event;
 }
 
+void UNiagaraNodeDataSetBase::AddParameterMapPins()
+{
+	const UEdGraphSchema_Niagara* Schema = GetDefault<UEdGraphSchema_Niagara>();
+	UEdGraphPin* ParamMapInPin = CreatePin(EGPD_Input, Schema->TypeDefinitionToPinType(FNiagaraTypeDefinition::GetParameterMapDef()), ParamMapInVarName, 0);
+	ParamMapInPin->bDefaultValueIsIgnored = true;
+	ParamMapInPin->PinFriendlyName = LOCTEXT("UNiagaraNodeWriteDataSetParamMapInPin", "ParamMap");
+	UEdGraphPin* ParamMapOutPin = CreatePin(EGPD_Output, Schema->TypeDefinitionToPinType(FNiagaraTypeDefinition::GetParameterMapDef()), ParamMapOutVarName, 1);
+	ParamMapOutPin->bDefaultValueIsIgnored = true;
+	ParamMapOutPin->PinFriendlyName = LOCTEXT("UNiagaraNodeWriteDataSetParamMapOutPin", "ParamMap");
+}
+
 void UNiagaraNodeDataSetBase::PostLoad()
 {
 	Super::PostLoad();
@@ -73,6 +86,13 @@ void UNiagaraNodeDataSetBase::PostLoad()
 	}*/
 	
 	IsSynchronizedWithStruct(true, nullptr, true);
+
+	const int32 NiagaraVer = GetLinkerCustomVersion(FNiagaraCustomVersion::GUID);
+
+	if (NiagaraVer < FNiagaraCustomVersion::AddingParamMapToDataSetBaseNode)
+	{
+		ReallocatePins();
+	}
 }
 
 bool UNiagaraNodeDataSetBase::IsSynchronizedWithStruct(bool bIgnoreConditionVariable, FString* Issues, bool bLogIssues)
@@ -234,6 +254,16 @@ bool UNiagaraNodeDataSetBase::GetSupportedNiagaraTypeDef(const UProperty* Proper
 		TypeDef = FNiagaraTypeDefinition::GetColorDef();
 		return true;
 	}
+	else if (StructProp && StructProp->Struct == FNiagaraTypeDefinition::GetQuatStruct())
+	{
+		TypeDef = FNiagaraTypeDefinition::GetQuatDef();
+		return true;
+	}
+	else if (StructProp && StructProp->Struct)
+	{
+		TypeDef = FNiagaraTypeDefinition(StructProp->Struct);
+		return FNiagaraTypeRegistry::GetRegisteredPayloadTypes().Contains(TypeDef);
+	}
 	return false;
 }
 
diff --git a/Engine/Plugins/FX/Niagara/Source/NiagaraEditor/Private/NiagaraNodeEmitter.cpp b/Engine/Plugins/FX/Niagara/Source/NiagaraEditor/Private/NiagaraNodeEmitter.cpp
--- a/Engine/Plugins/FX/Niagara/Source/NiagaraEditor/Private/NiagaraNodeEmitter.cpp
+++ b/Engine/Plugins/FX/Niagara/Source/NiagaraEditor/Private/NiagaraNodeEmitter.cpp
@@ -29,6 +29,8 @@ void UNiagaraNodeEmitter::PostInitProperties()
 {
 	Super::PostInitProperties();
 	PinPendingRename = nullptr;
+	CachedGraph = nullptr;
+	CachedScriptSource = nullptr;
 }
 
 UNiagaraSystem* UNiagaraNodeEmitter::GetOwnerSystem() const
@@ -144,7 +146,8 @@ FLinearColor UNiagaraNodeEmitter::GetNodeTitleColor() const
 
 void UNiagaraNodeEmitter::NodeConnectionListChanged()
 {
-	GetGraph()->NotifyGraphChanged();
+	MarkNodeRequiresSynchronization(__FUNCTION__, true);
+	//GetGraph()->NotifyGraphChanged();
 }
 
 FString UNiagaraNodeEmitter::GetEmitterUniqueName() const
@@ -159,7 +162,8 @@ FString UNiagaraNodeEmitter::GetEmitterUniqueName() const
 			}
 		}
 	}
-	return FString();
+
+	return CachedUniqueName.ToString();
 }
 
 UNiagaraScriptSource* UNiagaraNodeEmitter::GetScriptSource() const
@@ -184,7 +188,7 @@ UNiagaraScriptSource* UNiagaraNodeEmitter::GetScriptSource() const
 		return Source;
 	}
 
-	return nullptr;
+	return Cast<UNiagaraScriptSource>(CachedScriptSource);
 }
 
 UNiagaraGraph* UNiagaraNodeEmitter::GetCalledGraph() const
@@ -211,6 +215,11 @@ UNiagaraGraph* UNiagaraNodeEmitter::GetCalledGraph() const
 			return Source->NodeGraph;
 		}
 	}
+
+	if (CachedGraph != nullptr)
+	{
+		return CachedGraph;
+	}
 	return nullptr;
 }
 
@@ -220,6 +229,14 @@ bool UNiagaraNodeEmitter::RefreshFromExternalChanges()
 	return true;
 }
 
+void UNiagaraNodeEmitter::SetCachedVariablesForCompilation(const FName& InUniqueName, UNiagaraGraph* InGraph, UNiagaraScriptSourceBase* InSource)
+{
+	CachedUniqueName = InUniqueName;
+	CachedGraph = InGraph;
+	CachedScriptSource = InSource;
+}
+
+
 FText UNiagaraNodeEmitter::GetNameFromEmitter()
 {
 	if (OwnerSystem != nullptr && EmitterHandleId.IsValid())
@@ -232,6 +249,10 @@ FText UNiagaraNodeEmitter::GetNameFromEmitter()
 			}
 		}
 	}
+	else if (CachedUniqueName.IsValid())
+	{
+		return FText::FromName(CachedUniqueName);
+	}
 	return FText();
 }
 
@@ -256,7 +277,7 @@ void UNiagaraNodeEmitter::BuildParameterMapHistory(FNiagaraParameterMapHistoryBu
 	{
 		if (bRecursive)
 		{
-			ParamMapIdx = OutHistory.TraceParameterMapOutputPin(GetInputPin(0)->LinkedTo[0]);
+			ParamMapIdx = OutHistory.TraceParameterMapOutputPin(UNiagaraNode::TraceOutputPin(GetInputPin(0)->LinkedTo[0]));
 		}
 		else
 		{
@@ -277,7 +298,8 @@ void UNiagaraNodeEmitter::BuildParameterMapHistory(FNiagaraParameterMapHistoryBu
 		Usages.Add(ENiagaraScriptUsage::ParticleSpawnScriptInterpolated);
 		Usages.Add(ENiagaraScriptUsage::ParticleUpdateScript);
 		Usages.Add(ENiagaraScriptUsage::ParticleEventScript);
-		
+	
+		uint32 NodeIdx = OutHistory.BeginNodeVisitation(ParamMapIdx, this);
 		for (ENiagaraScriptUsage OutputNodeUsage : Usages)
 		{
 			TArray<UNiagaraNodeOutput*> OutputNodes;
@@ -286,21 +308,26 @@ void UNiagaraNodeEmitter::BuildParameterMapHistory(FNiagaraParameterMapHistoryBu
 
 			// Build up a new parameter map history with all the child graph nodes..
 			FNiagaraParameterMapHistoryBuilder ChildBuilder;
-
+			ChildBuilder.RegisterEncounterableVariables(OutHistory.GetEncounterableVariables());
 			ChildBuilder.EnableScriptWhitelist(true, GetUsage());
-			ChildBuilder.EnterEmitter(EmitterUniqueName, this);
+			FString LocalEmitterName = TEXT("Emitter");
+			ChildBuilder.EnterEmitter(LocalEmitterName, this);
 			for (UNiagaraNodeOutput* OutputNode : OutputNodes)
 			{
 				ChildBuilder.BuildParameterMaps(OutputNode, true);
 			}
-			ChildBuilder.ExitEmitter(EmitterUniqueName, this);
+			ChildBuilder.ExitEmitter(LocalEmitterName, this);
 			 
+			TMap<FString, FString> RenameMap;
+			RenameMap.Add(LocalEmitterName, EmitterUniqueName);
 			for (FNiagaraParameterMapHistory& History : ChildBuilder.Histories)
 			{
 				OutHistory.Histories[ParamMapIdx].MapPinHistory.Append(History.MapPinHistory);
 				for (int32 SrcVarIdx = 0; SrcVarIdx < History.Variables.Num(); SrcVarIdx++)
 				{
 					FNiagaraVariable& Var = History.Variables[SrcVarIdx];
+					Var = FNiagaraParameterMapHistory::ResolveAliases(Var, RenameMap, TEXT("."));
+
 					int32 ExistingIdx = OutHistory.Histories[ParamMapIdx].FindVariable(Var.GetName(), Var.GetType());
 					if (ExistingIdx == INDEX_NONE)
 					{
@@ -323,6 +350,7 @@ void UNiagaraNodeEmitter::BuildParameterMapHistory(FNiagaraParameterMapHistoryBu
 			}
 		}
 
+		OutHistory.EndNodeVisitation(ParamMapIdx, NodeIdx);
 		OutHistory.ExitEmitter(EmitterUniqueName, this);
 	}
 
@@ -394,59 +422,17 @@ void UNiagaraNodeEmitter::Compile(FHlslNiagaraTranslator *Translator, TArray<int
 	}
 }
 
-TSharedRef<SWidget> UNiagaraNodeEmitter::GenerateAddPinMenu(const FString& InWorkingPinName, SNiagaraGraphPinAdd* InPin)
+void UNiagaraNodeEmitter::GatherExternalDependencyIDs(ENiagaraScriptUsage InMasterUsage, const FGuid& InMasterUsageId, TArray<FGuid>& InReferencedIDs, TArray<UObject*>& InReferencedObjs) const
 {
-	FMenuBuilder MenuBuilder(true, nullptr);
-	
-	UNiagaraGraph* Graph = GetCalledGraph();
-	const UEdGraphSchema_Niagara* Schema = GetDefault<UEdGraphSchema_Niagara>();
-
-	TArray<UEdGraphPin*> InputPins;
-	GetInputPins(InputPins);
-
-	FString EmitterNamespace = FNiagaraParameterMapHistory::MakeSafeNamespaceString(GetNameFromEmitter().ToString());
-
-	FNiagaraParameterMapHistoryBuilder HistoryBuilder;
-	BuildParameterMapHistory(HistoryBuilder, false);
-
-	if (HistoryBuilder.Histories.Num() == 0)
-	{
-		return SNullWidget::NullWidget;
-	}
-
-	TArray<FNiagaraVariable> ExistingVariables;
-	for (UEdGraphPin* InputPin : InputPins)
-	{
-		ExistingVariables.Add(Schema->PinToNiagaraVariable(InputPin));
-	}
-
-	TArray<FNiagaraVariable> Variables;
-	for (const FNiagaraVariable& Var : HistoryBuilder.Histories[0].Variables)
-	{
-		if (ExistingVariables.ContainsByPredicate([Var](const FNiagaraVariable& A) { return A.GetName().ToString().Equals(Var.GetName().ToString()); }))
-		{
-			continue;
-		}
-
-		if (FNiagaraParameterMapHistory::IsInNamespace(Var, TEXT("EmitterUniforms")) || FNiagaraParameterMapHistory::IsInNamespace(Var, EmitterNamespace))
-		{
-			Variables.AddUnique(Var);
-		}
-	}
-
-	Variables.Sort([](const FNiagaraVariable& A, const FNiagaraVariable& B) { return (A.GetName() < B.GetName()); });
+	UNiagaraGraph* CalledGraph = GetCalledGraph();
 
-	for (FNiagaraVariable& NamespacedVar : Variables)
+	if (CalledGraph)
 	{
-		FString NamespacedName = NamespacedVar.GetName().ToString();
-		MenuBuilder.AddMenuEntry(
-			FText::FromString(NamespacedName),
-			FText::Format(LOCTEXT("AddButtonTypeEntryToolTipFormatSystem", "Add a reference to {0}"), FText::FromName(NamespacedVar.GetName())),
-			FSlateIcon(),
-			FUIAction(FExecuteAction::CreateSP(InPin, &SNiagaraGraphPinAdd::OnAddType, NamespacedVar)));
+		ENiagaraScriptUsage TargetUsage = InMasterUsage == ENiagaraScriptUsage::SystemSpawnScript ? ENiagaraScriptUsage::EmitterSpawnScript : ENiagaraScriptUsage::EmitterUpdateScript;
+		InReferencedIDs.Add(CalledGraph->GetCompileID(TargetUsage, FGuid(0,0,0,0)));
+		InReferencedObjs.Add(CalledGraph);
+		CalledGraph->GatherExternalDependencyIDs(TargetUsage, FGuid(0, 0, 0, 0), InReferencedIDs, InReferencedObjs);
 	}
-	
-	return MenuBuilder.MakeWidget();
 }
 
 #undef LOCTEXT_NAMESPACE // NiagaraNodeEmitter
diff --git a/Engine/Plugins/FX/Niagara/Source/NiagaraEditor/Private/NiagaraNodeEmitter.h b/Engine/Plugins/FX/Niagara/Source/NiagaraEditor/Private/NiagaraNodeEmitter.h
--- a/Engine/Plugins/FX/Niagara/Source/NiagaraEditor/Private/NiagaraNodeEmitter.h
+++ b/Engine/Plugins/FX/Niagara/Source/NiagaraEditor/Private/NiagaraNodeEmitter.h
@@ -44,7 +44,6 @@ public:
 	virtual bool IsPinNameEditableUponCreation(const UEdGraphPin* GraphPinObj) const override;
 	virtual bool VerifyEditablePinName(const FText& InName, FText& OutErrorMessage, const UEdGraphPin* InGraphPinObj) const override;
 	virtual bool CommitEditablePinName(const FText& InName, UEdGraphPin* InGraphPinObj)  override;
-	virtual TSharedRef<SWidget> GenerateAddPinMenu(const FString& InWorkingPinName, SNiagaraGraphPinAdd* InPin) override;
 
 	virtual void BuildParameterMapHistory(FNiagaraParameterMapHistoryBuilder& OutHistory, bool bRecursive = true) override;
 
@@ -58,6 +57,9 @@ public:
 	UNiagaraGraph* GetCalledGraph() const;
 	
 	virtual void Compile(FHlslNiagaraTranslator *Translator, TArray<int32>& Outputs) override;
+	virtual void GatherExternalDependencyIDs(ENiagaraScriptUsage InMasterUsage, const FGuid& InMasterUsageId, TArray<FGuid>& InReferencedIDs, TArray<UObject*>& InReferencedObjs) const override;
+
+	void SetCachedVariablesForCompilation(const FName& InUniqueName, UNiagaraGraph* InGraph, UNiagaraScriptSourceBase* InSource);
 
 protected:
 	UEdGraphPin* PinPendingRename;
@@ -83,4 +85,8 @@ private:
 
 	UPROPERTY()
 	ENiagaraScriptUsage ScriptType;
+
+	FName CachedUniqueName;
+	UNiagaraGraph* CachedGraph;
+	UNiagaraScriptSourceBase* CachedScriptSource;
 };
\ No newline at end of file
diff --git a/Engine/Plugins/FX/Niagara/Source/NiagaraEditor/Private/NiagaraNodeFunctionCall.cpp b/Engine/Plugins/FX/Niagara/Source/NiagaraEditor/Private/NiagaraNodeFunctionCall.cpp
--- a/Engine/Plugins/FX/Niagara/Source/NiagaraEditor/Private/NiagaraNodeFunctionCall.cpp
+++ b/Engine/Plugins/FX/Niagara/Source/NiagaraEditor/Private/NiagaraNodeFunctionCall.cpp
@@ -83,7 +83,8 @@ void UNiagaraNodeFunctionCall::PostEditChangeProperty(struct FPropertyChangedEve
 	}
 	Super::PostEditChangeProperty(PropertyChangedEvent);
 
-	GetGraph()->NotifyGraphChanged();
+	MarkNodeRequiresSynchronization(__FUNCTION__, true);
+	//GetGraph()->NotifyGraphChanged();
 }
 
 void UNiagaraNodeFunctionCall::AllocateDefaultPins()
@@ -280,6 +281,31 @@ bool UNiagaraNodeFunctionCall::CanAddToGraph(UNiagaraGraph* TargetGraph, FString
 	return true;
 }
 
+UNiagaraGraph* UNiagaraNodeFunctionCall::GetCalledGraph() const
+{
+	if (FunctionScript)
+	{
+		UNiagaraScriptSource* Source = CastChecked<UNiagaraScriptSource>(FunctionScript->GetSource());
+		if (Source)
+		{
+			UNiagaraGraph* FunctionGraph = Source->NodeGraph;
+			return FunctionGraph;
+		}
+	}
+
+	return nullptr;
+}
+
+
+ENiagaraScriptUsage UNiagaraNodeFunctionCall::GetCalledUsage() const
+{
+	if (FunctionScript)
+	{
+		return FunctionScript->GetUsage();
+	}
+	return ENiagaraScriptUsage::Function;
+}
+
 void UNiagaraNodeFunctionCall::Compile(class FHlslNiagaraTranslator* Translator, TArray<int32>& Outputs)
 {
 	TArray<int32> Inputs;
@@ -327,7 +353,7 @@ void UNiagaraNodeFunctionCall::Compile(class FHlslNiagaraTranslator* Translator,
 			}
 
 			UEdGraphPin* CallerPin = *PinPtr;
-			UEdGraphPin* CallerLinkedTo = CallerPin->LinkedTo.Num() > 0 ? CallerPin->LinkedTo[0] : nullptr;			
+			UEdGraphPin* CallerLinkedTo = CallerPin->LinkedTo.Num() > 0 ? UNiagaraNode::TraceOutputPin(CallerPin->LinkedTo[0]) : nullptr;			
 			FNiagaraVariable PinVar = Schema->PinToNiagaraVariable(CallerPin);
 			if (!CallerLinkedTo)
 			{
@@ -427,7 +453,6 @@ bool UNiagaraNodeFunctionCall::RefreshFromExternalChanges()
 			if (bReload)
 			{
 				bReload = true;
-				UE_LOG(LogNiagaraEditor, Log, TEXT("RefreshFromExternalChanges %s"), *(FunctionScript->GetPathName()));
 			}
 		}
 	}
@@ -468,6 +493,31 @@ void UNiagaraNodeFunctionCall::SubsumeExternalDependencies(TMap<const UObject*,
 	}
 }
 
+void UNiagaraNodeFunctionCall::GatherExternalDependencyIDs(ENiagaraScriptUsage InMasterUsage, const FGuid& InMasterUsageId, TArray<FGuid>& InReferencedIDs, TArray<UObject*>& InReferencedObjs) const
+{
+	if (FunctionScript && FunctionScript->GetOutermost() != this->GetOutermost())
+	{
+		UNiagaraScriptSource* Source = CastChecked<UNiagaraScriptSource>(FunctionScript->GetSource());
+		UNiagaraGraph* FunctionGraph = CastChecked<UNiagaraGraph>(Source->NodeGraph);
+		
+		// We don't know which graph type we're referencing, so we try them all... may need to replace this with something faster in the future.
+		if (FunctionGraph)
+		{
+			for (int32 i = (int32)ENiagaraScriptUsage::Function; i <= (int32)ENiagaraScriptUsage::Module; i++)
+			{
+				FGuid FoundGuid = FunctionGraph->GetCompileID((ENiagaraScriptUsage)i, FGuid(0, 0, 0, 0));
+				if (FoundGuid.IsValid())
+				{
+					InReferencedIDs.Add(FoundGuid);
+					InReferencedObjs.Add(FunctionGraph);
+
+					FunctionGraph->GatherExternalDependencyIDs((ENiagaraScriptUsage)i, FGuid(0, 0, 0, 0), InReferencedIDs, InReferencedObjs);
+				}
+			}
+		}
+	}
+}
+
 bool UNiagaraNodeFunctionCall::ScriptIsValid() const
 {
 	if (FunctionScript != nullptr)
@@ -505,8 +555,23 @@ void UNiagaraNodeFunctionCall::BuildParameterMapHistory(FNiagaraParameterMapHist
 		{
 			OutputNode = FunctionGraph->FindOutputNode(ENiagaraScriptUsage::DynamicInput);
 		}
+
+		int32 ParamMapIdx = INDEX_NONE;
+		uint32 NodeIdx = INDEX_NONE;
+		const UEdGraphPin* CandidateParamMapPin = GetInputPin(0);
+		if (CandidateParamMapPin && CandidateParamMapPin->LinkedTo.Num() != 0 && Schema->PinToTypeDefinition(CandidateParamMapPin) == FNiagaraTypeDefinition::GetParameterMapDef())
+		{
+			if (bRecursive)
+			{
+				ParamMapIdx = OutHistory.TraceParameterMapOutputPin(UNiagaraNode::TraceOutputPin(GetInputPin(0)->LinkedTo[0]));
+			}
+		}
 		
 		OutHistory.EnterFunction(GetFunctionName(), FunctionScript, this);
+		if (ParamMapIdx != INDEX_NONE)
+		{
+			NodeIdx = OutHistory.BeginNodeVisitation(ParamMapIdx, this);
+		}
 		OutputNode->BuildParameterMapHistory(OutHistory, true);
 
 		// Since we're about to lose the pin calling context, we finish up the function call parameter map pin wiring
@@ -530,12 +595,18 @@ void UNiagaraNodeFunctionCall::BuildParameterMapHistory(FNiagaraParameterMapHist
 					{
 						TPair<UEdGraphPin*, int32> Pair;
 						Pair.Key = OutputPins[i];
-						Pair.Value = OutHistory.TraceParameterMapOutputPin(ChildOutputNodePin->LinkedTo[0]);
+						Pair.Value = OutHistory.TraceParameterMapOutputPin(UNiagaraNode::TraceOutputPin(ChildOutputNodePin->LinkedTo[0]));
 						MatchedPairs.Add(Pair);
 					}
 				}
 			}
 		}
+
+		if (ParamMapIdx != INDEX_NONE)
+		{
+			OutHistory.EndNodeVisitation(ParamMapIdx, NodeIdx);
+		}
+
 		OutHistory.ExitFunction(GetFunctionName(), FunctionScript, this);
 
 		for (int32 i = 0; i < MatchedPairs.Num(); i++)
@@ -549,19 +620,29 @@ void UNiagaraNodeFunctionCall::BuildParameterMapHistory(FNiagaraParameterMapHist
 	}
 }
 
-UEdGraphPin* UNiagaraNodeFunctionCall::FindParameterMapDefaultValuePin(const FName VariableName)
+UEdGraphPin* UNiagaraNodeFunctionCall::FindParameterMapDefaultValuePin(const FName VariableName, ENiagaraScriptUsage InParentUsage) const
 {
 	if (FunctionScript)
 	{
 		UNiagaraScriptSource* ScriptSource = Cast<UNiagaraScriptSource>(FunctionScript->GetSource());
-		if (ScriptSource != nullptr)
+		if (ScriptSource != nullptr && ScriptSource->NodeGraph != nullptr)
 		{
-			return ScriptSource->NodeGraph->FindParameterMapDefaultValuePin(VariableName);
+			return ScriptSource->NodeGraph->FindParameterMapDefaultValuePin(VariableName, FunctionScript->GetUsage(), InParentUsage);
 		}
 	}
 	return nullptr;
 }
 
+void UNiagaraNodeFunctionCall::SuggestName(FString SuggestedName)
+{
+	ComputeNodeName(SuggestedName);
+}
+
+UNiagaraNodeFunctionCall::FOnInputsChanged& UNiagaraNodeFunctionCall::OnInputsChanged()
+{
+	return OnInputsChangedDelegate;
+}
+
 ENiagaraNumericOutputTypeSelectionMode UNiagaraNodeFunctionCall::GetNumericOutputTypeSelectionMode() const
 {
 	if (FunctionScript)
@@ -577,14 +658,31 @@ void UNiagaraNodeFunctionCall::AutowireNewNode(UEdGraphPin* FromPin)
 	ComputeNodeName();
 }
 
-void UNiagaraNodeFunctionCall::ComputeNodeName()
+void UNiagaraNodeFunctionCall::ComputeNodeName(FString SuggestedName)
 {
-	const FString CurrentName = FunctionDisplayName;
-	FName ProposedName = *CurrentName;
 	FString FunctionName = FunctionScript ? FunctionScript->GetName() : Signature.GetName();
-	if (!FunctionName.IsEmpty())
+	FName ProposedName;
+	if (SuggestedName.IsEmpty() == false)
+	{ 
+		// If we have a suggested name and and either there is no function name, or it is a permutation of the function name
+		// it can be used as the proposed name.
+		if (FunctionName.IsEmpty() || SuggestedName == FunctionName || (SuggestedName.StartsWith(FunctionName) && SuggestedName.RightChop(FunctionName.Len()).IsNumeric()))
+		{
+			ProposedName = *SuggestedName;
+		}
+	}
+
+	if(ProposedName == NAME_None)
 	{
-		ProposedName = *FunctionName;
+		const FString CurrentName = FunctionDisplayName;
+		if (FunctionName.IsEmpty() == false)
+		{
+			ProposedName = *FunctionName;
+		}
+		else
+		{
+			ProposedName = *CurrentName;
+		}
 	}
 
 	UNiagaraGraph* Graph = GetNiagaraGraph();
diff --git a/Engine/Plugins/FX/Niagara/Source/NiagaraEditor/Private/NiagaraNodeIf.cpp b/Engine/Plugins/FX/Niagara/Source/NiagaraEditor/Private/NiagaraNodeIf.cpp
--- a/Engine/Plugins/FX/Niagara/Source/NiagaraEditor/Private/NiagaraNodeIf.cpp
+++ b/Engine/Plugins/FX/Niagara/Source/NiagaraEditor/Private/NiagaraNodeIf.cpp
@@ -6,6 +6,9 @@
 
 #define LOCTEXT_NAMESPACE "NiagaraNodeIf"
 
+const FString UNiagaraNodeIf::InputAPinSuffix(" A");
+const FString UNiagaraNodeIf::InputBPinSuffix(" B");
+
 UNiagaraNodeIf::UNiagaraNodeIf(const FObjectInitializer& ObjectInitializer)
 	: Super(ObjectInitializer)
 {
@@ -29,18 +32,43 @@ void UNiagaraNodeIf::PostLoad()
 	if (OutputVars.Num() != OutputVarGuids.Num())
 	{
 		OutputVarGuids.SetNum(OutputVars.Num());
-		for (int32 i = 0; i < OutputVars.Num(); i++)
+	}
+
+	if (OutputVars.Num() != InputAVarGuids.Num())
+	{
+		InputAVarGuids.SetNum(OutputVars.Num());
+	}
+
+	if (OutputVars.Num() != InputBVarGuids.Num())
+	{
+		InputBVarGuids.SetNum(OutputVars.Num());
+	}
+
+	auto LoadGuid = [&](FGuid& Guid, const FString& Name, const EEdGraphPinDirection Direction)
+	{
+		UEdGraphPin* Pin = FindPin(Name, Direction);
+		if (Pin)
 		{
-			UEdGraphPin* Pin = FindPin(OutputVars[i].GetName().ToString(), EGPD_Output);
-			if (Pin != nullptr)
+			if (!Pin->PersistentGuid.IsValid())
 			{
-				OutputVarGuids[i] = Pin->PersistentGuid;
-			}
-			else
-			{
-				UE_LOG(LogNiagaraEditor, Error, TEXT("Unable to find pin named %s"), *OutputVars[i].GetName().ToString());
+				Pin->PersistentGuid = FGuid::NewGuid();
 			}
+			Guid = Pin->PersistentGuid;
+		}
+		else
+		{
+			UE_LOG(LogNiagaraEditor, Error, TEXT("Unable to find output pin named %s"), *Name);
 		}
+	};
+
+	for (int32 i = 0; i < OutputVars.Num(); i++)
+	{
+		const FString VarName = OutputVars[i].GetName().ToString();
+		LoadGuid(OutputVarGuids[i], VarName, EGPD_Output);
+		const FString InputAName = VarName + InputAPinSuffix;
+		LoadGuid(InputAVarGuids[i], InputAName, EGPD_Input);
+		const FString InputBName = VarName + InputBPinSuffix;
+		LoadGuid(InputBVarGuids[i], InputBName, EGPD_Input);
 	}
 }
 
@@ -57,24 +85,25 @@ void UNiagaraNodeIf::AllocateDefaultPins()
 	CreatePin(EGPD_Input, Schema->TypeDefinitionToPinType(FNiagaraTypeDefinition::GetBoolDef()), TEXT("Condition"));
 
 	//Create the inputs for each path.
-	for (FNiagaraVariable& Var : OutputVars)
+	for (int32 Index = 0; Index < OutputVars.Num(); Index++)
 	{
-		const TCHAR* PathSuffix = TEXT(" A");
-		CreatePin(EGPD_Input, Schema->TypeDefinitionToPinType(Var.GetType()), *(Var.GetName().ToString() + PathSuffix));	
+		const FNiagaraVariable& Var = OutputVars[Index];
+		UEdGraphPin* NewPin = CreatePin(EGPD_Input, Schema->TypeDefinitionToPinType(Var.GetType()), *(Var.GetName().ToString() + InputAPinSuffix));
+		NewPin->PersistentGuid = InputAVarGuids[Index];
 	}
 
-	for (FNiagaraVariable& Var : OutputVars)
+	for (int32 Index = 0; Index < OutputVars.Num(); Index++)
 	{
-		const TCHAR* PathSuffix = TEXT(" B");
-		CreatePin(EGPD_Input, Schema->TypeDefinitionToPinType(Var.GetType()), *(Var.GetName().ToString() + PathSuffix));
+		const FNiagaraVariable& Var = OutputVars[Index];
+		UEdGraphPin* NewPin = CreatePin(EGPD_Input, Schema->TypeDefinitionToPinType(Var.GetType()), *(Var.GetName().ToString() + InputBPinSuffix));
+		NewPin->PersistentGuid = InputBVarGuids[Index];
 	}
 
 	for (int32 Index = 0; Index < OutputVars.Num(); Index++)
 	{
 		const FNiagaraVariable& Var = OutputVars[Index];
-		const FGuid& Guid = OutputVarGuids[Index];
 		UEdGraphPin* NewPin = CreatePin(EGPD_Output, Schema->TypeDefinitionToPinType(Var.GetType()), Var.GetName());
-		NewPin->PersistentGuid = Guid;
+		NewPin->PersistentGuid = OutputVarGuids[Index];
 	}
 
 	CreateAddPin(EGPD_Output);
@@ -121,9 +150,16 @@ bool UNiagaraNodeIf::RefreshFromExternalChanges()
 FGuid UNiagaraNodeIf::AddOutput(FNiagaraTypeDefinition Type, const FName& Name)
 {
 	FNiagaraVariable NewOutput(Type, Name);
-	FGuid Guid = FGuid::NewGuid();
 	OutputVars.Add(NewOutput);
+	FGuid Guid = FGuid::NewGuid();
 	OutputVarGuids.Add(Guid);
+	InputAVarGuids.Add(FGuid::NewGuid());
+	InputBVarGuids.Add(FGuid::NewGuid());
+
+	const UEdGraphSchema_Niagara* Schema = CastChecked<UEdGraphSchema_Niagara>(GetSchema());
+	CreatePin(EGPD_Input, Schema->TypeDefinitionToPinType(Type), *(Name.ToString() + InputAPinSuffix), InputAVarGuids.Num());
+	CreatePin(EGPD_Input, Schema->TypeDefinitionToPinType(Type), *(Name.ToString() + InputBPinSuffix), InputAVarGuids.Num() + InputBVarGuids.Num());
+
 	return Guid;
 }
 
@@ -156,9 +192,8 @@ void UNiagaraNodeIf::OnNewTypedPinAdded(UEdGraphPin* NewPin)
 	FGuid Guid = AddOutput(OutputType, OutputName);
 
 	// Update the pin's data too so that it's connection is maintained after reallocating.
+	NewPin->PinName = OutputName;
 	NewPin->PersistentGuid = Guid;
-
-	ReallocatePins();
 }
 
 void UNiagaraNodeIf::OnPinRenamed(UEdGraphPin* RenamedPin, const FString& OldName)
diff --git a/Engine/Plugins/FX/Niagara/Source/NiagaraEditor/Private/NiagaraNodeIf.h b/Engine/Plugins/FX/Niagara/Source/NiagaraEditor/Private/NiagaraNodeIf.h
--- a/Engine/Plugins/FX/Niagara/Source/NiagaraEditor/Private/NiagaraNodeIf.h
+++ b/Engine/Plugins/FX/Niagara/Source/NiagaraEditor/Private/NiagaraNodeIf.h
@@ -19,6 +19,12 @@ public:
 	UPROPERTY(EditAnywhere, Category = If)
 	TArray<FGuid> OutputVarGuids;
 
+	UPROPERTY(EditAnywhere, Category = If)
+	TArray<FGuid> InputAVarGuids;
+
+	UPROPERTY(EditAnywhere, Category = If)
+	TArray<FGuid> InputBVarGuids;
+
 	//~ Begin UObject Interface
 	virtual void PostEditChangeProperty(FPropertyChangedEvent& PropertyChangedEvent) override;
 	virtual void PostLoad() override;
@@ -35,10 +41,11 @@ public:
 	virtual bool RefreshFromExternalChanges() override;
 	//~ End UNiagaraNode Interface
 
+protected:
+
 	/** Helper function to create a variable to add to the OutputVars and FGuid to add to OutputVarGuids. */
 	FGuid AddOutput(FNiagaraTypeDefinition Type, const FName& Name);
 
-protected:
 	//~ Begin EdGraphNode Interface
 	virtual void OnPinRemoved(UEdGraphPin* PinToRemove) override;
 	//~ End EdGraphNode Interface
@@ -51,4 +58,8 @@ protected:
 	virtual bool CanMovePin(const UEdGraphPin* Pin) const override { return false; }
 	virtual bool AllowNiagaraTypeForAddPin(const FNiagaraTypeDefinition& InType) override;
 	//~ End UNiagaraNodeWithDynamicPins Interface
+
+private:
+	static const FString InputAPinSuffix;
+	static const FString InputBPinSuffix;
 };
diff --git a/Engine/Plugins/FX/Niagara/Source/NiagaraEditor/Private/NiagaraNodeInput.cpp b/Engine/Plugins/FX/Niagara/Source/NiagaraEditor/Private/NiagaraNodeInput.cpp
--- a/Engine/Plugins/FX/Niagara/Source/NiagaraEditor/Private/NiagaraNodeInput.cpp
+++ b/Engine/Plugins/FX/Niagara/Source/NiagaraEditor/Private/NiagaraNodeInput.cpp
@@ -84,11 +84,22 @@ void UNiagaraNodeInput::BuildParameterMapHistory(FNiagaraParameterMapHistoryBuil
 	{
 		int32 ParamMapIdx = OutHistory.FindMatchingParameterMapFromContextInputs(Input);
 
-		if (ParamMapIdx == INDEX_NONE)
+		if (ParamMapIdx == INDEX_NONE && Usage != ENiagaraInputNodeUsage::TranslatorConstant)
 		{
 			ParamMapIdx = OutHistory.CreateParameterMap();
 		}
-		OutHistory.RegisterParameterMapPin(ParamMapIdx, GetOutputPin(0));
+		else if (ParamMapIdx == INDEX_NONE && OutHistory.Histories.Num() != 0)
+		{
+			ParamMapIdx = 0;
+		}
+
+		if (ParamMapIdx != INDEX_NONE)
+		{
+			uint32 NodeIdx = OutHistory.BeginNodeVisitation(ParamMapIdx, this);
+			OutHistory.EndNodeVisitation(ParamMapIdx, NodeIdx);
+
+			OutHistory.RegisterParameterMapPin(ParamMapIdx, GetOutputPin(0));
+		}
 	}
 }
 
@@ -100,7 +111,7 @@ void UNiagaraNodeInput::AllocateDefaultPins()
 		check(Class->IsChildOf(UNiagaraDataInterface::StaticClass()));
 		if (!DataInterface)
 		{
-			DataInterface = NewObject<UNiagaraDataInterface>(this, Class);
+			DataInterface = NewObject<UNiagaraDataInterface>(this, Class, NAME_None, RF_Transactional | RF_Public);
 		}
 	}
 
@@ -287,7 +298,7 @@ void UNiagaraNodeInput::OnRenameNode(const FString& NewName)
 		Node->ReallocatePins();
 	}
 
-	Graph->MarkGraphRequiresSynchronization();
+	Graph->MarkGraphRequiresSynchronization(TEXT("Input node renamed"));
 }
 
 TSharedPtr<SGraphNode> UNiagaraNodeInput::CreateVisualWidget()
@@ -306,6 +317,8 @@ FLinearColor UNiagaraNodeInput::GetNodeTitleColor() const
 		return Schema->NodeTitleColor_SystemConstant;
 	case ENiagaraInputNodeUsage::Attribute:
 		return Schema->NodeTitleColor_Attribute;
+	case ENiagaraInputNodeUsage::TranslatorConstant:
+		return Schema->NodeTitleColor_TranslatorConstant;
 	default:
 		// TODO: Do something better here.
 		return FLinearColor::Black;
@@ -441,6 +454,9 @@ void UNiagaraNodeInput::Compile(class FHlslNiagaraTranslator* Translator, TArray
 		return;
 	}
 
+	UNiagaraGraph* Graph = GetNiagaraGraph();
+	check(Input.GetType() != FNiagaraTypeDefinition::GetGenericNumericDef());
+
 	int32 FunctionParam = INDEX_NONE;
 	if (IsExposed() && Translator->GetFunctionParameter(Input, FunctionParam))
 	{
@@ -456,7 +472,7 @@ void UNiagaraNodeInput::Compile(class FHlslNiagaraTranslator* Translator, TArray
 				if (Usage == ENiagaraInputNodeUsage::Parameter && DataInterface != nullptr)
 				{
 					check(Input.GetType().GetClass());
-					Outputs.Add(Translator->RegisterDataInterface(Input, DataInterface, false));
+					Outputs.Add(Translator->RegisterDataInterface(Input, DataInterface, false, false));
 					return;
 				}
 				else
@@ -475,7 +491,7 @@ void UNiagaraNodeInput::Compile(class FHlslNiagaraTranslator* Translator, TArray
 		if (DataInterface)
 		{
 			check(Input.GetType().GetClass());
-			Outputs.Add(Translator->RegisterDataInterface(Input, DataInterface, false)); break;
+			Outputs.Add(Translator->RegisterDataInterface(Input, DataInterface, false, false)); break;
 		}
 		else
 		{
@@ -485,6 +501,8 @@ void UNiagaraNodeInput::Compile(class FHlslNiagaraTranslator* Translator, TArray
 		Outputs.Add(Translator->GetParameter(Input)); break;
 	case ENiagaraInputNodeUsage::Attribute:
 		Outputs.Add(Translator->GetAttribute(Input)); break;
+	case ENiagaraInputNodeUsage::TranslatorConstant:
+		Outputs.Add(Translator->GetParameter(Input)); break;
 	default:
 		check(false);
 	}
diff --git a/Engine/Plugins/FX/Niagara/Source/NiagaraEditor/Private/NiagaraNodeOp.cpp b/Engine/Plugins/FX/Niagara/Source/NiagaraEditor/Private/NiagaraNodeOp.cpp
--- a/Engine/Plugins/FX/Niagara/Source/NiagaraEditor/Private/NiagaraNodeOp.cpp
+++ b/Engine/Plugins/FX/Niagara/Source/NiagaraEditor/Private/NiagaraNodeOp.cpp
@@ -101,6 +101,48 @@ void UNiagaraNodeOp::Compile(class FHlslNiagaraTranslator* Translator, TArray<in
 	Translator->Operation(this, Inputs, Outputs);
 }
 
+void UNiagaraNodeOp::PostLoad()
+{
+	Super::PostLoad();
+
+	FName OriginalOpName = OpName;
+
+	if (OpName == TEXT("Numeric::Cos"))
+	{
+		OpName = TEXT("Numeric::Cosine(Radians)");
+	}
+	else if (OpName == TEXT("Numeric::Sin"))
+	{
+		OpName = TEXT("Numeric::Sine(Radians)");
+	}
+	else if (OpName == TEXT("Numeric::Tan"))
+	{
+		OpName = TEXT("Numeric::Tangent(Radians)");
+	}
+	else if (OpName == TEXT("Numeric::ASin"))
+	{
+		OpName = TEXT("Numeric::ArcSine(Radians)");
+	}
+	else if (OpName == TEXT("Numeric::ACos"))
+	{
+		OpName = TEXT("Numeric::ArcCosine(Radians)");
+	}
+	else if (OpName == TEXT("Numeric::ATan"))
+	{
+		OpName = TEXT("Numeric::ArcTangent(Radians)");
+	}
+	else if (OpName == TEXT("Numeric::ATan2"))
+	{
+		OpName = TEXT("Numeric::ArcTangent2(Radians)");
+	}
+
+	if (OpName != OriginalOpName)
+	{
+		UE_LOG(LogNiagaraEditor, Log, TEXT("OpNode: Converted %s to %s, Package: %s"), *OriginalOpName.ToString(), *OpName.ToString(), *GetOutermost()->GetName());
+	}
+
+}
+
 bool UNiagaraNodeOp::RefreshFromExternalChanges()
 {
 	// TODO - Leverage code in reallocate pins to determine if any pins have changed...
diff --git a/Engine/Plugins/FX/Niagara/Source/NiagaraEditor/Private/NiagaraNodeOutput.cpp b/Engine/Plugins/FX/Niagara/Source/NiagaraEditor/Private/NiagaraNodeOutput.cpp
--- a/Engine/Plugins/FX/Niagara/Source/NiagaraEditor/Private/NiagaraNodeOutput.cpp
+++ b/Engine/Plugins/FX/Niagara/Source/NiagaraEditor/Private/NiagaraNodeOutput.cpp
@@ -49,7 +49,7 @@ void UNiagaraNodeOutput::RemoveOutputPin(UEdGraphPin* Pin)
 		Modify();
 		Outputs.RemoveAt(Index);
 		ReallocatePins();
-		GetNiagaraGraph()->NotifyGraphNeedsRecompile();
+		MarkNodeRequiresSynchronization(__FUNCTION__, true);
 	}
 }
 
@@ -69,8 +69,8 @@ void UNiagaraNodeOutput::PinNameTextCommitted(const FText& Text, ETextCommit::Ty
 		FNiagaraVariable* Var = Outputs.FindByPredicate([&](const FNiagaraVariable& InVar) {return Pin->PinName == InVar.GetName(); });
 		check(Var != nullptr);
 		Pin->PinName = *Text.ToString();
-		Var->SetName(Pin->PinName);
-		GetNiagaraGraph()->NotifyGraphNeedsRecompile();
+		Var->SetName(Pin->PinName);		
+		MarkNodeRequiresSynchronization(__FUNCTION__, true);
 	}
 }
 
diff --git a/Engine/Plugins/FX/Niagara/Source/NiagaraEditor/Private/NiagaraNodeParameterMapBase.cpp b/Engine/Plugins/FX/Niagara/Source/NiagaraEditor/Private/NiagaraNodeParameterMapBase.cpp
--- a/Engine/Plugins/FX/Niagara/Source/NiagaraEditor/Private/NiagaraNodeParameterMapBase.cpp
+++ b/Engine/Plugins/FX/Niagara/Source/NiagaraEditor/Private/NiagaraNodeParameterMapBase.cpp
@@ -15,6 +15,8 @@
 #include "NiagaraScriptSource.h"
 #include "NiagaraConstants.h"
 #include "NiagaraParameterCollection.h"
+#include "Widgets/SNiagaraParameterMapView.h"
+#include "Widgets/Input/SEditableTextBox.h"
 
 #include "IAssetTools.h"
 #include "AssetRegistryModule.h"
@@ -23,24 +25,26 @@
 
 #define LOCTEXT_NAMESPACE "NiagaraNodeParameterMapBase"
 
-UNiagaraNodeParameterMapBase::UNiagaraNodeParameterMapBase() : UNiagaraNodeWithDynamicPins()
+UNiagaraNodeParameterMapBase::UNiagaraNodeParameterMapBase() 
+	: UNiagaraNodeWithDynamicPins()
+	, PinPendingRename(nullptr)
 {
 
 }
 
-TArray<FNiagaraParameterMapHistory> UNiagaraNodeParameterMapBase::GetParameterMaps(UNiagaraScriptSourceBase* InSource, FString EmitterNameOverride)
+TArray<FNiagaraParameterMapHistory> UNiagaraNodeParameterMapBase::GetParameterMaps(UNiagaraScriptSourceBase* InSource, FString EmitterNameOverride, const TArray<FNiagaraVariable>& EncounterableVariables)
 {
 	TArray<FNiagaraParameterMapHistory> OutputParameterMapHistories;
 	UNiagaraScriptSource* Base = Cast<UNiagaraScriptSource>(InSource);
 	if (Base != nullptr)
 	{
-		OutputParameterMapHistories = GetParameterMaps(Base->NodeGraph, EmitterNameOverride);
+		OutputParameterMapHistories = GetParameterMaps(Base->NodeGraph, EmitterNameOverride, EncounterableVariables);
 	}
 	return OutputParameterMapHistories;
 
 }
 
-TArray<FNiagaraParameterMapHistory> UNiagaraNodeParameterMapBase::GetParameterMaps(UNiagaraGraph* InGraph, FString EmitterNameOverride)
+TArray<FNiagaraParameterMapHistory> UNiagaraNodeParameterMapBase::GetParameterMaps(UNiagaraGraph* InGraph, FString EmitterNameOverride, const TArray<FNiagaraVariable>& EncounterableVariables)
 {
 	TArray<UNiagaraNodeOutput*> OutputNodes;
 	InGraph->FindOutputNodes(OutputNodes);
@@ -48,17 +52,18 @@ TArray<FNiagaraParameterMapHistory> UNiagaraNodeParameterMapBase::GetParameterMa
 
 	for (UNiagaraNodeOutput* FoundOutputNode : OutputNodes)
 	{
-		OutputParameterMapHistories.Append(GetParameterMaps(FoundOutputNode, false, EmitterNameOverride));
+		OutputParameterMapHistories.Append(GetParameterMaps(FoundOutputNode, false, EmitterNameOverride,EncounterableVariables));
 	}
 
 	return OutputParameterMapHistories;
 }
 
-TArray<FNiagaraParameterMapHistory> UNiagaraNodeParameterMapBase::GetParameterMaps(UNiagaraNodeOutput* InGraphEnd, bool bLimitToOutputScriptType, FString EmitterNameOverride)
+TArray<FNiagaraParameterMapHistory> UNiagaraNodeParameterMapBase::GetParameterMaps(UNiagaraNodeOutput* InGraphEnd, bool bLimitToOutputScriptType, FString EmitterNameOverride, const TArray<FNiagaraVariable>& EncounterableVariables)
 {
 	const UEdGraphSchema_Niagara* Schema = Cast<UEdGraphSchema_Niagara>(InGraphEnd->GetSchema());
-	
+
 	FNiagaraParameterMapHistoryBuilder Builder;
+	Builder.RegisterEncounterableVariables(EncounterableVariables);
 	if (!EmitterNameOverride.IsEmpty())
 	{
 		Builder.EnterEmitter(EmitterNameOverride, nullptr);
@@ -82,158 +87,7 @@ TArray<FNiagaraParameterMapHistory> UNiagaraNodeParameterMapBase::GetParameterMa
 
 bool UNiagaraNodeParameterMapBase::AllowNiagaraTypeForAddPin(const FNiagaraTypeDefinition& InType)
 {
-	return  InType != FNiagaraTypeDefinition::GetGenericNumericDef();
-}
-
-TSharedRef<SWidget> UNiagaraNodeParameterMapBase::GenerateAddPinMenu(const FString& InWorkingPinName, SNiagaraGraphPinAdd* InPin)
-{
-	UNiagaraGraph* Graph = GetNiagaraGraph();
-	bool IsModule = Graph->FindOutputNode(ENiagaraScriptUsage::Module) != nullptr || Graph->FindOutputNode(ENiagaraScriptUsage::DynamicInput) != nullptr 
-		|| Graph->FindOutputNode(ENiagaraScriptUsage::Function) != nullptr;
-
-	bool bSupportsAttributes = true;
-	UNiagaraScriptSource* Source = Cast<UNiagaraScriptSource>(Graph->GetOuter());
-	if (Source && IsModule)
-	{
-		UNiagaraScript* Script = Cast<UNiagaraScript>(Source->GetOuter());
-		if (Script)
-		{
-			TArray<ENiagaraScriptUsage> Usages = Script->GetSupportedUsageContexts();
-			if (Usages.Contains(ENiagaraScriptUsage::ParticleEventScript) || Usages.Contains(ENiagaraScriptUsage::ParticleSpawnScript) || Usages.Contains(ENiagaraScriptUsage::ParticleUpdateScript))
-			{
-				bSupportsAttributes = true;
-			}
-			else
-			{
-				bSupportsAttributes = false;
-			}
-		}
-	}
-
-	FString WorkingCustomName;
-	FMenuBuilder MenuBuilder(true, nullptr);
-	MenuBuilder.AddSubMenu(LOCTEXT("CommonEngine", "Common Engine-Provided Variables"),
-		LOCTEXT("CommonSystemTooltip", "Create an entry using one of the common engine variables."),
-		FNewMenuDelegate::CreateLambda([InWorkingPinName, InPin, this](FMenuBuilder& SubMenuBuilder) { BuildEngineMenu(SubMenuBuilder, InWorkingPinName, InPin); }));
-
-	MenuBuilder.AddSubMenu(LOCTEXT("LocalVars", "Existing Graph Variables"),
-		LOCTEXT("LocalVarsTooltip", "Create an entry using existing graph variables."),
-		FNewMenuDelegate::CreateLambda([InWorkingPinName, InPin, this](FMenuBuilder& SubMenuBuilder) { BuildLocalMenu(SubMenuBuilder, InWorkingPinName, InPin); }));
-
-	if (bSupportsAttributes)
-	{
-		MenuBuilder.AddSubMenu(LOCTEXT("CommonAttributes", "Common Attributes"),
-			LOCTEXT("CommonAttributesTooltip", "Create an entry using one of the common attributes."),
-			FNewMenuDelegate::CreateLambda([InWorkingPinName, InPin, this](FMenuBuilder& SubMenuBuilder) { BuildCommonMenu(SubMenuBuilder, InWorkingPinName, InPin); }));
-
-		WorkingCustomName = TEXT("Particles.") + InWorkingPinName;
-		MenuBuilder.AddSubMenu(LOCTEXT("DefineAttribute", "Define Attribute"),
-			LOCTEXT("SupportedTypesTooltip", "Create an entry in the particles namespace that you will name immediately after."),
-			FNewMenuDelegate::CreateLambda([WorkingCustomName, InPin, this](FMenuBuilder& SubMenuBuilder) { BuildTypeMenu(SubMenuBuilder, WorkingCustomName, InPin); }));
-	}
-
-	WorkingCustomName = TEXT("Module.") + InWorkingPinName;
-	MenuBuilder.AddSubMenu(LOCTEXT("DefineModuleLocal", "Define Module Local"),
-		LOCTEXT("SupportedTypesTooltip", "Create an entry in the module namespace that you will name immediately after."),
-		FNewMenuDelegate::CreateLambda([WorkingCustomName, InPin, this](FMenuBuilder& SubMenuBuilder) { BuildTypeMenu(SubMenuBuilder, WorkingCustomName, InPin); }));
-
-	WorkingCustomName = TEXT("ParameterCollections.") + InWorkingPinName;
-	MenuBuilder.AddSubMenu(LOCTEXT("ParameterCollectionsMenuTitle", "ParameterCollections"),
-		LOCTEXT("ParameterCollectionsMenuTooltip", "Create a reference to a parameter in a Parameter Collection."),
-		FNewMenuDelegate::CreateLambda([WorkingCustomName, InPin, this](FMenuBuilder& SubMenuBuilder) { BuildParameterCollectionsMenu(SubMenuBuilder, WorkingCustomName, InPin); }));
-
-	return MenuBuilder.MakeWidget();
-}
-
-void UNiagaraNodeParameterMapBase::BuildCommonMenu(FMenuBuilder& InMenuBuilder, const FString& InWorkingName, SNiagaraGraphPinAdd* InPin)
-{
-	TArray<FNiagaraVariable> Variables = FNiagaraConstants::GetCommonParticleAttributes();
-
-	for (FNiagaraVariable& NamespacedVar : Variables)
-	{
-		FText Desc = FNiagaraConstants::GetAttributeDescription(NamespacedVar);
-		FString NamespacedName = NamespacedVar.GetName().ToString();
-		InMenuBuilder.AddMenuEntry(
-			FText::FromString(NamespacedName),
-			FText::Format(LOCTEXT("AddButtonTypeEntryToolTipFormatCommon", "Add a reference to {0}. {1}"), FText::FromName(NamespacedVar.GetName()), Desc),
-			FSlateIcon(),
-			FUIAction(FExecuteAction::CreateSP(InPin, &SNiagaraGraphPinAdd::OnAddType, NamespacedVar)));
-	}
-}
-
-void UNiagaraNodeParameterMapBase::BuildParameterCollectionsMenu(FMenuBuilder& InMenuBuilder, const FString& InWorkingName, SNiagaraGraphPinAdd* InPin)
-{
-	//Create sub menus for parameter collections.
-	FAssetRegistryModule& AssetRegistryModule = FModuleManager::LoadModuleChecked<FAssetRegistryModule>("AssetRegistry");
-	TArray<FAssetData> CollectionAssets;
-	AssetRegistryModule.Get().GetAssetsByClass(UNiagaraParameterCollection::StaticClass()->GetFName(), CollectionAssets);
-	for (FAssetData& CollectionAsset : CollectionAssets)
-	{
-		UNiagaraParameterCollection* Collection = CastChecked<UNiagaraParameterCollection>(CollectionAsset.GetAsset());
-
-		InMenuBuilder.AddSubMenu(FText::FromName(CollectionAsset.AssetName),
-			FText::FromString(CollectionAsset.GetFullName()),
-			FNewMenuDelegate::CreateLambda([Collection, InPin, this](FMenuBuilder& SubMenuBuilder) { BuildParameterCollectionMenu(SubMenuBuilder, Collection, InPin); }));
-	}
-}
-
-void UNiagaraNodeParameterMapBase::BuildParameterCollectionMenu(FMenuBuilder& InMenuBuilder, UNiagaraParameterCollection* Collection, SNiagaraGraphPinAdd* InPin)
-{
-	for (const FNiagaraVariable& Parameter : Collection->GetParameters())
-	{
-		FText Desc = FText::GetEmpty();//FNiagaraConstants::GetAttributeDescription(Parameter);
-		FString FriendlyName = Collection->FriendlyNameFromParameterName(Parameter.GetName().ToString());
-		InMenuBuilder.AddMenuEntry(
-			FText::FromString(FriendlyName),
-			FText::Format(LOCTEXT("AddButtonTypeEntryToolTipFormatNPC", "Add a reference to Collection Parameter {0}. {1}"), FText::FromName(Parameter.GetName()), Desc),
-			FSlateIcon(),
-			FUIAction(FExecuteAction::CreateSP(InPin, &SNiagaraGraphPinAdd::OnAddType, Parameter)));
-	}
-}
-
-void UNiagaraNodeParameterMapBase::BuildLocalMenu(FMenuBuilder& InMenuBuilder, const FString& InWorkingName, SNiagaraGraphPinAdd* InPin)
-{
-	UNiagaraGraph* Graph = GetNiagaraGraph();
-	const TArray<FNiagaraParameterMapHistory> Histories = GetParameterMaps(Graph);
-
-	TArray<FNiagaraVariable> Variables;
-	for (const FNiagaraParameterMapHistory& History : Histories)
-	{
-		for (const FNiagaraVariable& Var : History.Variables)
-		{
-			Variables.AddUnique(Var);
-		}
-	}
-
-	Variables.Sort([](const FNiagaraVariable& A, const FNiagaraVariable& B) { return (A.GetName() < B.GetName()); });
-
-	for (FNiagaraVariable& NamespacedVar : Variables)
-	{
-		FString NamespacedName = NamespacedVar.GetName().ToString();
-		InMenuBuilder.AddMenuEntry(
-			FText::FromString(NamespacedName),
-			FText::Format(LOCTEXT("AddButtonTypeEntryToolTipFormatSystem", "Add a reference to {0}"), FText::FromName(NamespacedVar.GetName())),
-			FSlateIcon(),
-			FUIAction(FExecuteAction::CreateSP(InPin, &SNiagaraGraphPinAdd::OnAddType, NamespacedVar)));
-	}
-}
-
-void UNiagaraNodeParameterMapBase::BuildEngineMenu(FMenuBuilder& InMenuBuilder, const FString& InWorkingName, SNiagaraGraphPinAdd* InPin)
-{
-	TArray<FNiagaraVariable> Variables = FNiagaraConstants::GetEngineConstants();
-	Variables.Sort([](const FNiagaraVariable& A, const FNiagaraVariable& B) { return (A.GetName() < B.GetName()); });
-
-	for (FNiagaraVariable& Var : Variables)
-	{
-		FText Desc = FNiagaraConstants::GetEngineConstantDescription(Var);
-		FNiagaraVariable NamespacedVar = Var;
-		FString NamespacedName = NamespacedVar.GetName().ToString();
-		InMenuBuilder.AddMenuEntry(
-			FText::FromString(NamespacedName),
-			FText::Format(LOCTEXT("AddButtonTypeEntryToolTipFormatSystem", "Add a reference to {0}. {1}"), FText::FromName(Var.GetName()), Desc),
-			FSlateIcon(),
-			FUIAction(FExecuteAction::CreateSP(InPin, &SNiagaraGraphPinAdd::OnAddType, NamespacedVar)));
-	}
+	return InType != FNiagaraTypeDefinition::GetGenericNumericDef();
 }
 
 FText UNiagaraNodeParameterMapBase::GetPinDescriptionText(UEdGraphPin* Pin) const
@@ -286,5 +140,45 @@ void UNiagaraNodeParameterMapBase::PinDescriptionTextCommitted(const FText& Text
 	}
 }
 
+void UNiagaraNodeParameterMapBase::CollectAddPinActions(FGraphActionListBuilderBase& OutActions, bool& bOutCreateRemainingActions, UEdGraphPin* Pin)
+{
+	bOutCreateRemainingActions = true;
+}
+
+void UNiagaraNodeParameterMapBase::GetPinHoverText(const UEdGraphPin& Pin, FString& HoverTextOut) const
+{
+	// Get hover text from metadata description.
+	const UNiagaraGraph* NiagaraGraph = GetNiagaraGraph();
+	if (NiagaraGraph)
+	{
+		const UEdGraphSchema_Niagara* Schema = Cast<UEdGraphSchema_Niagara>(NiagaraGraph->GetSchema());
+		if (Schema)
+		{
+			FNiagaraVariable Var = FNiagaraVariable(Schema->PinToTypeDefinition(&Pin), Pin.PinName);
+			const FNiagaraVariableMetaData* Metadata = NiagaraGraph->GetMetaData(Var);
+			if (Metadata)
+			{
+				HoverTextOut = Metadata->Description.ToString();
+			}
+		}
+	}
+}
+
+void UNiagaraNodeParameterMapBase::OnPinRenamed(UEdGraphPin* RenamedPin, const FString& OldName)
+{
+	RenamedPin->PinFriendlyName = FText::FromName(RenamedPin->PinName);
+
+	FNiagaraTypeDefinition VarType = CastChecked<UEdGraphSchema_Niagara>(GetSchema())->PinToTypeDefinition(RenamedPin);
+	FNiagaraVariable Var(VarType, *OldName);
+
+	UNiagaraGraph* Graph = GetNiagaraGraph();
+	Graph->RenameParameter(Var, RenamedPin->GetFName(), false /*bInNotifyGraphChanged*/); // Notify graph changed in child.
+
+	if (RenamedPin == PinPendingRename)
+	{
+		PinPendingRename = nullptr;
+	}
+
+}
 
 #undef LOCTEXT_NAMESPACE
\ No newline at end of file
diff --git a/Engine/Plugins/FX/Niagara/Source/NiagaraEditor/Private/NiagaraNodeParameterMapBase.h b/Engine/Plugins/FX/Niagara/Source/NiagaraEditor/Private/NiagaraNodeParameterMapBase.h
--- a/Engine/Plugins/FX/Niagara/Source/NiagaraEditor/Private/NiagaraNodeParameterMapBase.h
+++ b/Engine/Plugins/FX/Niagara/Source/NiagaraEditor/Private/NiagaraNodeParameterMapBase.h
@@ -20,24 +20,23 @@ public:
 	/** Traverse the graph looking for the history of the parameter map specified by the input pin. This will return the list of variables discovered, any per-variable warnings (type mismatches, etc)
 		encountered per variable, and an array of pins encountered in order of traversal outward from the input pin.
 	*/
-	static TArray<FNiagaraParameterMapHistory> GetParameterMaps(UNiagaraNodeOutput* InGraphEnd, bool bLimitToOutputScriptType = false, FString EmitterNameOverride = TEXT(""));
-	static TArray<FNiagaraParameterMapHistory> GetParameterMaps(UNiagaraGraph* InGraph, FString EmitterNameOverride = TEXT(""));
-	static TArray<FNiagaraParameterMapHistory> GetParameterMaps(class UNiagaraScriptSourceBase* InSource, FString EmitterNameOverride = TEXT(""));
+	static TArray<FNiagaraParameterMapHistory> GetParameterMaps(UNiagaraNodeOutput* InGraphEnd, bool bLimitToOutputScriptType = false, FString EmitterNameOverride = TEXT(""), const TArray<FNiagaraVariable>& EncounterableVariables = TArray<FNiagaraVariable>());
+	static TArray<FNiagaraParameterMapHistory> GetParameterMaps(UNiagaraGraph* InGraph, FString EmitterNameOverride = TEXT(""), const TArray<FNiagaraVariable>& EncounterableVariables = TArray<FNiagaraVariable>());
+	static TArray<FNiagaraParameterMapHistory> GetParameterMaps(class UNiagaraScriptSourceBase* InSource, FString EmitterNameOverride = TEXT(""), const TArray<FNiagaraVariable>& EncounterableVariables = TArray<FNiagaraVariable>());
 
 	virtual bool AllowNiagaraTypeForAddPin(const FNiagaraTypeDefinition& InType) override;
 	
-	virtual TSharedRef<SWidget> GenerateAddPinMenu(const FString& InWorkingPinName, SNiagaraGraphPinAdd* InPin) override;
-
 	/** Gets the description text for a pin. */
 	FText GetPinDescriptionText(UEdGraphPin* Pin) const;
 
 	/** Called when a pin's description text is committed. */
 	void PinDescriptionTextCommitted(const FText& Text, ETextCommit::Type CommitType, UEdGraphPin* Pin);
+
+	virtual void CollectAddPinActions(FGraphActionListBuilderBase& OutActions, bool& bOutCreateRemainingActions, UEdGraphPin* Pin) override;
+	virtual void GetPinHoverText(const UEdGraphPin& Pin, FString& HoverTextOut) const override;
+
 protected:
-	virtual void BuildCommonMenu(FMenuBuilder& InMenuBuilder, const FString& InWorkingName, SNiagaraGraphPinAdd* InPin);
-	virtual void BuildLocalMenu(FMenuBuilder& InMenuBuilder, const FString& InWorkingName, SNiagaraGraphPinAdd* InPin);
-	virtual void BuildEngineMenu(FMenuBuilder& InMenuBuilder, const FString& InWorkingName, SNiagaraGraphPinAdd* InPin);
-	virtual void BuildParameterCollectionsMenu(FMenuBuilder& InMenuBuilder, const FString& InWorkingName, SNiagaraGraphPinAdd* InPin);
-	virtual void BuildParameterCollectionMenu(FMenuBuilder& InMenuBuilder, UNiagaraParameterCollection* Collection, SNiagaraGraphPinAdd* InPin);
+	virtual void OnPinRenamed(UEdGraphPin* RenamedPin, const FString& OldName) override;
 
+	UEdGraphPin* PinPendingRename;
 };
