Copying //UE4/Dev-Niagara to //UE4/Dev-Main (Source: //UE4/Dev-Niagara @ 4074996)

diff --git a/Engine/Plugins/FX/Niagara/Source/NiagaraEditor/Private/NiagaraNodeParameterMapGet.cpp b/Engine/Plugins/FX/Niagara/Source/NiagaraEditor/Private/NiagaraNodeParameterMapGet.cpp
--- a/Engine/Plugins/FX/Niagara/Source/NiagaraEditor/Private/NiagaraNodeParameterMapGet.cpp
+++ b/Engine/Plugins/FX/Niagara/Source/NiagaraEditor/Private/NiagaraNodeParameterMapGet.cpp
@@ -17,10 +17,11 @@
 #include "Widgets/Layout/SBox.h"
 #include "Widgets/Input/SEditableTextBox.h"
 #include "EdGraph/EdGraphNode.h"
+#include "NiagaraParameterCollection.h"
 
 #define LOCTEXT_NAMESPACE "NiagaraNodeParameterMapGet"
 
-UNiagaraNodeParameterMapGet::UNiagaraNodeParameterMapGet() : UNiagaraNodeParameterMapBase(), PinPendingRename(nullptr)
+UNiagaraNodeParameterMapGet::UNiagaraNodeParameterMapGet() : UNiagaraNodeParameterMapBase()
 {
 
 }
@@ -87,7 +88,10 @@ UEdGraphPin* UNiagaraNodeParameterMapGet::CreateDefaultPin(UEdGraphPin* OutputPi
 	UEdGraphPin* DefaultPin = CreatePin(EEdGraphPinDirection::EGPD_Input, OutputPin->PinType, TEXT(""));
 
 	const UEdGraphSchema_Niagara* Schema = GetDefault<UEdGraphSchema_Niagara>();
-	FNiagaraVariable Var = Schema->PinToNiagaraVariable(OutputPin, true);
+	FNiagaraTypeDefinition NiagaraType = Schema->PinToTypeDefinition(OutputPin);
+	bool bNeedsValue = NiagaraType.IsDataInterface() == false;
+	FNiagaraVariable Var = Schema->PinToNiagaraVariable(OutputPin, bNeedsValue);
+
 	FString PinDefaultValue;
 	if (Schema->TryGetPinDefaultValueFromNiagaraVariable(Var, PinDefaultValue))
 	{
@@ -110,52 +114,16 @@ UEdGraphPin* UNiagaraNodeParameterMapGet::CreateDefaultPin(UEdGraphPin* OutputPi
 
 void UNiagaraNodeParameterMapGet::OnPinRenamed(UEdGraphPin* RenamedPin, const FString& OldName)
 {
-	RenamedPin->PinFriendlyName = FText::FromName(RenamedPin->PinName);
-
-	FNiagaraTypeDefinition VarType = CastChecked<UEdGraphSchema_Niagara>(GetSchema())->PinToTypeDefinition(RenamedPin);
-	FNiagaraVariable Var(VarType, *OldName);
-
-	UNiagaraGraph* Graph = GetNiagaraGraph();
-	FNiagaraVariableMetaData* OldMetaData = Graph->GetMetaData(Var);
-
-	if (OldMetaData)
-	{
-		Graph->Modify();
-		OldMetaData->ReferencerNodes.Remove(TWeakObjectPtr<UObject>(this));
-	}
-
-	FNiagaraVariable NewVar(VarType, *RenamedPin->GetName());
-	FNiagaraVariableMetaData* NewMetaData = Graph->GetMetaData(NewVar);
-
-	// If no variable has already defined the meta-data for this entry and it isn't being redefined to one of our constants, copy over the old values that make sense (not display name).
-	if (NewMetaData == nullptr && OldMetaData != nullptr && !FNiagaraConstants::IsNiagaraConstant(NewVar))
-	{
-		FNiagaraVariableMetaData& NewVarMetaData = Graph->FindOrAddMetaData(NewVar);
-		NewVarMetaData.PropertyMetaData = OldMetaData->PropertyMetaData;
-		NewVarMetaData.PropertyMetaData.Remove(TEXT("DisplayName")); // DisplayName is probably incorrect, so drop it.
-		NewVarMetaData.ReferencerNodes.Add(TWeakObjectPtr<UObject>(this));
-	}
-
-	if (OldMetaData)
-	{
-		Graph->PurgeUnreferencedMetaData();
-	}
-
-	if (RenamedPin == PinPendingRename)
-	{
-		PinPendingRename = nullptr;
-	}
-
+	UNiagaraNodeParameterMapBase::OnPinRenamed(RenamedPin, OldName);
 
 	UEdGraphPin* DefaultPin = GetDefaultPin(RenamedPin);
-	
 	if (DefaultPin)
 	{
 		DefaultPin->Modify();
 		SynchronizeDefaultInputPin(DefaultPin, RenamedPin);
 	}
 
-	Graph->NotifyGraphNeedsRecompile();
+	MarkNodeRequiresSynchronization(__FUNCTION__, true);
 }
 
 
@@ -213,10 +181,11 @@ void UNiagaraNodeParameterMapGet::RemoveDynamicPin(UEdGraphPin* Pin)
 		RemovePin(DefaultPin);
 	}
 
-	GetGraph()->NotifyGraphChanged();
+	MarkNodeRequiresSynchronization(__FUNCTION__, true);
+	//GetGraph()->NotifyGraphChanged();
 }
 
-UEdGraphPin* UNiagaraNodeParameterMapGet::GetDefaultPin(UEdGraphPin* OutputPin)
+UEdGraphPin* UNiagaraNodeParameterMapGet::GetDefaultPin(UEdGraphPin* OutputPin) const
 {
 	TArray<UEdGraphPin*> InputPins;
 	GetInputPins(InputPins);
@@ -329,11 +298,12 @@ void UNiagaraNodeParameterMapGet::BuildParameterMapHistory(FNiagaraParameterMapH
 	int32 ParamMapIdx = INDEX_NONE;
 	if (GetInputPin(0)->LinkedTo.Num() != 0)
 	{
-		ParamMapIdx = OutHistory.TraceParameterMapOutputPin(GetInputPin(0)->LinkedTo[0]);
+		ParamMapIdx = OutHistory.TraceParameterMapOutputPin(UNiagaraNode::TraceOutputPin(GetInputPin(0)->LinkedTo[0]));
 	}
 
 	if (ParamMapIdx != INDEX_NONE)
 	{
+		uint32 NodeIdx = OutHistory.BeginNodeVisitation(ParamMapIdx, this);
 		TArray<UEdGraphPin*> OutputPins;
 		GetOutputPins(OutputPins);
 		for (int32 i = 0; i < OutputPins.Num(); i++)
@@ -353,6 +323,7 @@ void UNiagaraNodeParameterMapGet::BuildParameterMapHistory(FNiagaraParameterMapH
 				OutHistory.HandleVariableRead(ParamMapIdx, OutputPins[i], true, nullptr, bUsedDefaults);
 			}
 		}
+		OutHistory.EndNodeVisitation(ParamMapIdx, NodeIdx);
 	}
 }
 
@@ -407,6 +378,14 @@ void UNiagaraNodeParameterMapGet::Compile(class FHlslNiagaraTranslator* Translat
 	}
 }
 
+bool UNiagaraNodeParameterMapGet::CancelEditablePinName(const FText& InName, UEdGraphPin* InGraphPinObj)
+{
+	if (InGraphPinObj == PinPendingRename)
+	{
+		PinPendingRename = nullptr;
+	}
+	return true;
+}
 
 bool UNiagaraNodeParameterMapGet::CommitEditablePinName(const FText& InName, UEdGraphPin* InGraphPinObj) 
 {
@@ -426,6 +405,32 @@ bool UNiagaraNodeParameterMapGet::CommitEditablePinName(const FText& InName, UEd
 	return false;
 }
 
+void UNiagaraNodeParameterMapGet::GatherExternalDependencyIDs(ENiagaraScriptUsage InMasterUsage, const FGuid& InMasterUsageId, TArray<FGuid>& InReferencedIDs, TArray<UObject*>& InReferencedObjs) const
+{
+	// If we are referencing any parameter collections, we need to register them here... might want to speeed this up in the future 
+	// by caching any parameter collections locally.
+	TArray<UEdGraphPin*> OutputPins;
+	GetOutputPins(OutputPins);
+	const UEdGraphSchema_Niagara* Schema = CastChecked<UEdGraphSchema_Niagara>(GetSchema());
+
+	for (int32 i = 0; i < OutputPins.Num(); i++)
+	{
+		if (IsAddPin(OutputPins[i]))
+		{
+			continue;
+		}
+
+		FNiagaraVariable Var = CastChecked<UEdGraphSchema_Niagara>(GetSchema())->PinToNiagaraVariable(OutputPins[i]);
+		UNiagaraParameterCollection* Collection = Schema->VariableIsFromParameterCollection(Var);
+		if (Collection)
+		{
+			InReferencedIDs.Add(Collection->GetCompileId());
+			InReferencedObjs.Add(Collection);
+		}
+
+	}
+}
+
 void UNiagaraNodeParameterMapGet::GetContextMenuActions(const FGraphNodeContextMenuBuilder& Context) const
 {
 	UNiagaraNodeParameterMapBase::GetContextMenuActions(Context);
diff --git a/Engine/Plugins/FX/Niagara/Source/NiagaraEditor/Private/NiagaraNodeParameterMapGet.h b/Engine/Plugins/FX/Niagara/Source/NiagaraEditor/Private/NiagaraNodeParameterMapGet.h
--- a/Engine/Plugins/FX/Niagara/Source/NiagaraEditor/Private/NiagaraNodeParameterMapGet.h
+++ b/Engine/Plugins/FX/Niagara/Source/NiagaraEditor/Private/NiagaraNodeParameterMapGet.h
@@ -26,6 +26,7 @@ public:
 	virtual bool IsPinNameEditableUponCreation(const UEdGraphPin* GraphPinObj) const override;
 	virtual bool VerifyEditablePinName(const FText& InName, FText& OutErrorMessage, const UEdGraphPin* InGraphPinObj) const override;
 	virtual bool CommitEditablePinName(const FText& InName, UEdGraphPin* InGraphPinObj)  override;
+	virtual bool CancelEditablePinName(const FText& InName, UEdGraphPin* InGraphPinObj) override;
 
 	virtual void Compile(class FHlslNiagaraTranslator* Translator, TArray<int32>& Outputs);
 
@@ -34,10 +35,12 @@ public:
 	virtual FText GetNodeTitle(ENodeTitleType::Type TitleType) const;
 
 	/** Get the default value input pin for one of the output pins specified.*/
-	UEdGraphPin* GetDefaultPin(UEdGraphPin* OutputPin);
+	UEdGraphPin* GetDefaultPin(UEdGraphPin* OutputPin) const;
 
 	virtual void PostLoad() override;
 
+	void GatherExternalDependencyIDs(ENiagaraScriptUsage InMasterUsage, const FGuid& InMasterUsageId, TArray<FGuid>& InReferencedIDs, TArray<UObject*>& InReferencedObjs) const override;
+
 protected:
 	virtual void OnNewTypedPinAdded(UEdGraphPin* NewPin) override;
 	virtual void OnPinRenamed(UEdGraphPin* RenamedPin, const FString& OldName) override;
@@ -54,8 +57,6 @@ protected:
 	/** Properly set up the default input pin for an output pin.*/
 	UEdGraphPin* CreateDefaultPin(UEdGraphPin* OutputPin);
 
-	UEdGraphPin* PinPendingRename;
-
 	UPROPERTY()
 	TMap<FGuid, FGuid> PinOutputToPinDefaultPersistentId;
 };
diff --git a/Engine/Plugins/FX/Niagara/Source/NiagaraEditor/Private/NiagaraNodeParameterMapSet.cpp b/Engine/Plugins/FX/Niagara/Source/NiagaraEditor/Private/NiagaraNodeParameterMapSet.cpp
--- a/Engine/Plugins/FX/Niagara/Source/NiagaraEditor/Private/NiagaraNodeParameterMapSet.cpp
+++ b/Engine/Plugins/FX/Niagara/Source/NiagaraEditor/Private/NiagaraNodeParameterMapSet.cpp
@@ -17,7 +17,7 @@
 
 #define LOCTEXT_NAMESPACE "NiagaraNodeParameterMapSet"
 
-UNiagaraNodeParameterMapSet::UNiagaraNodeParameterMapSet() : UNiagaraNodeParameterMapBase(), PinPendingRename(nullptr)
+UNiagaraNodeParameterMapSet::UNiagaraNodeParameterMapSet() : UNiagaraNodeParameterMapBase()
 {
 
 }
@@ -75,48 +75,34 @@ void UNiagaraNodeParameterMapSet::OnNewTypedPinAdded(UEdGraphPin* NewPin)
 		return;
 	}
 
+	if (!NewPin->PersistentGuid.IsValid())
+	{
+		NewPin->PersistentGuid = FGuid::NewGuid();
+	}
+
 	PinPendingRename = NewPin;
 }
 
 void UNiagaraNodeParameterMapSet::OnPinRenamed(UEdGraphPin* RenamedPin, const FString& OldName)
 {
-	RenamedPin->PinFriendlyName = FText::FromName(RenamedPin->PinName);
-
-	FNiagaraTypeDefinition VarType = CastChecked<UEdGraphSchema_Niagara>(GetSchema())->PinToTypeDefinition(RenamedPin);
-	FNiagaraVariable Var(VarType, *OldName);
-
-	UNiagaraGraph* Graph = GetNiagaraGraph();
-	FNiagaraVariableMetaData* OldMetaData = Graph->GetMetaData(Var);
-
-	if (OldMetaData)
-	{
-		Graph->Modify();
-		OldMetaData->ReferencerNodes.Remove(TWeakObjectPtr<UObject>(this));
-	}
-
-	FNiagaraVariable NewVar(VarType, *RenamedPin->GetName());
-	FNiagaraVariableMetaData* NewMetaData = Graph->GetMetaData(NewVar);
-
-	// If no variable has already defined the meta-data for this entry and it isn't being redefined to one of our constants, copy over the old values that make sense (not display name).
-	if (NewMetaData == nullptr && OldMetaData != nullptr && !FNiagaraConstants::IsNiagaraConstant(NewVar))
-	{
-		FNiagaraVariableMetaData& NewVarMetaData = Graph->FindOrAddMetaData(NewVar);
-		NewVarMetaData.PropertyMetaData = OldMetaData->PropertyMetaData;
-		NewVarMetaData.PropertyMetaData.Remove(TEXT("DisplayName")); // DisplayName is probably incorrect, so drop it.
-		NewVarMetaData.ReferencerNodes.Add(TWeakObjectPtr<UObject>(this));
-	}
+	UNiagaraNodeParameterMapBase::OnPinRenamed(RenamedPin, OldName);
+	MarkNodeRequiresSynchronization(__FUNCTION__, true);
+}
 
-	if (OldMetaData)
-	{
-		Graph->PurgeUnreferencedMetaData();
-	}
+void UNiagaraNodeParameterMapSet::SetPinName(UEdGraphPin* InPin, const FName& InName)
+{
+	FName OldName = InPin->PinName;
+	InPin->PinName = InName;
+	OnPinRenamed(InPin, OldName.ToString());
+}
 
-	if (RenamedPin == PinPendingRename)
+bool UNiagaraNodeParameterMapSet::CancelEditablePinName(const FText& InName, UEdGraphPin* InGraphPinObj)
+{
+	if (InGraphPinObj == PinPendingRename)
 	{
 		PinPendingRename = nullptr;
 	}
-	
-	Graph->NotifyGraphNeedsRecompile();
+	return true;
 }
 
 bool UNiagaraNodeParameterMapSet::CommitEditablePinName(const FText& InName, UEdGraphPin* InGraphPinObj) 
@@ -170,7 +156,7 @@ void UNiagaraNodeParameterMapSet::Compile(class FHlslNiagaraTranslator* Translat
 		int32 CompiledInput = Translator->CompilePin(InputPin);
 		if (CompiledInput == INDEX_NONE)
 		{
-			Translator->Error(LOCTEXT("InputError", "Error compiling input for get node."), this, InputPin);
+			Translator->Error(LOCTEXT("InputError", "Error compiling input for set node."), this, InputPin);
 		}
 		CompileInputs.Add(CompiledInput);
 	}
@@ -194,6 +180,8 @@ void UNiagaraNodeParameterMapSet::BuildParameterMapHistory(FNiagaraParameterMapH
 	GetInputPins(InputPins);
 	
 	int32 ParamMapIdx = INDEX_NONE;
+	uint32 NodeIdx = INDEX_NONE;
+
 	for (int32 i = 0; i < InputPins.Num(); i++)
 	{
 		if (IsAddPin(InputPins[i]))
@@ -222,6 +210,7 @@ void UNiagaraNodeParameterMapSet::BuildParameterMapHistory(FNiagaraParameterMapH
 			if (PriorParamPin)
 			{
 				ParamMapIdx = OutHistory.TraceParameterMapOutputPin(PriorParamPin);
+				NodeIdx = OutHistory.BeginNodeVisitation(ParamMapIdx, this);
 			}
 		}
 		else if (i > 0 && InputPins[i] != nullptr && ParamMapIdx != INDEX_NONE)
@@ -236,6 +225,11 @@ void UNiagaraNodeParameterMapSet::BuildParameterMapHistory(FNiagaraParameterMapH
 		return;
 	}
 
+	if (ParamMapIdx != INDEX_NONE)
+	{
+		OutHistory.EndNodeVisitation(ParamMapIdx, NodeIdx);
+	}
+
 	OutHistory.RegisterParameterMapPin(ParamMapIdx, GetOutputPin(0));
 }
 
@@ -268,4 +262,16 @@ void UNiagaraNodeParameterMapSet::GetContextMenuActions(const FGraphNodeContextM
 	}
 }
 
+void UNiagaraNodeParameterMapSet::PostLoad()
+{
+	for (UEdGraphPin* Pin : Pins)
+	{
+		if (!Pin->PersistentGuid.IsValid())
+		{
+			Pin->PersistentGuid = FGuid::NewGuid();
+		}
+	}
+	Super::PostLoad();
+}
+
 #undef LOCTEXT_NAMESPACE
diff --git a/Engine/Plugins/FX/Niagara/Source/NiagaraEditor/Private/NiagaraNodeParameterMapSet.h b/Engine/Plugins/FX/Niagara/Source/NiagaraEditor/Private/NiagaraNodeParameterMapSet.h
--- a/Engine/Plugins/FX/Niagara/Source/NiagaraEditor/Private/NiagaraNodeParameterMapSet.h
+++ b/Engine/Plugins/FX/Niagara/Source/NiagaraEditor/Private/NiagaraNodeParameterMapSet.h
@@ -24,6 +24,7 @@ public:
 	virtual bool IsPinNameEditableUponCreation(const UEdGraphPin* GraphPinObj) const override;
 	virtual bool VerifyEditablePinName(const FText& InName, FText& OutErrorMessage, const UEdGraphPin* InGraphPinObj) const override;
 	virtual bool CommitEditablePinName(const FText& InName, UEdGraphPin* InGraphPinObj)  override;
+	virtual bool CancelEditablePinName(const FText& InName, UEdGraphPin* InGraphPinObj) override;
 
 	virtual void BuildParameterMapHistory(FNiagaraParameterMapHistoryBuilder& OutHistory, bool bRecursive = true) override;
 
@@ -32,10 +33,11 @@ public:
 	virtual FText GetNodeTitle(ENodeTitleType::Type TitleType) const;
 
 	virtual void GetContextMenuActions(const FGraphNodeContextMenuBuilder& Context) const override;
+	virtual void PostLoad() override;
+
+	void SetPinName(UEdGraphPin* InPin, const FName& InName);
 
 protected:
 	virtual void OnNewTypedPinAdded(UEdGraphPin* NewPin); 
 	virtual void OnPinRenamed(UEdGraphPin* RenamedPin, const FString& OldName) override;
-	UEdGraphPin* PinPendingRename;
-
 };
diff --git a/Engine/Plugins/FX/Niagara/Source/NiagaraEditor/Private/NiagaraNodeReadDataSet.cpp b/Engine/Plugins/FX/Niagara/Source/NiagaraEditor/Private/NiagaraNodeReadDataSet.cpp
--- a/Engine/Plugins/FX/Niagara/Source/NiagaraEditor/Private/NiagaraNodeReadDataSet.cpp
+++ b/Engine/Plugins/FX/Niagara/Source/NiagaraEditor/Private/NiagaraNodeReadDataSet.cpp
@@ -25,6 +25,8 @@ void UNiagaraNodeReadDataSet::AllocateDefaultPins()
 		//Pin->bDefaultValueIsIgnored = true;
 	}
 
+	AddParameterMapPins();
+
 	bool useFriendlyNames = (VariableFriendlyNames.Num() == Variables.Num());
 	for (int32 i = 0; i < Variables.Num(); i++)
 	{
@@ -104,7 +106,35 @@ void UNiagaraNodeReadDataSet::Compile(class FHlslNiagaraTranslator* Translator,
 	{
 		Translator->Error(FText::FromString(IssuesWithStruct), this, nullptr);
 	}
-	Translator->ReadDataSet(DataSet, Variables, ENiagaraDataSetAccessMode::AppendConsume, INDEX_NONE, Outputs);
+	Translator->ReadDataSet(DataSet, Variables, ENiagaraDataSetAccessMode::AppendConsume, Inputs[0], Outputs);
+}
+
+void UNiagaraNodeReadDataSet::BuildParameterMapHistory(FNiagaraParameterMapHistoryBuilder& OutHistory, bool bRecursive)
+{
+	if (bRecursive)
+	{
+		OutHistory.VisitInputPins(this);
+	}
+
+	if (!IsNodeEnabled() && OutHistory.GetIgnoreDisabled())
+	{
+		RouteParameterMapAroundMe(OutHistory, bRecursive);
+		return;
+	}
+
+	int32 ParamMapIdx = INDEX_NONE;
+	if (GetInputPin(0)->LinkedTo.Num() != 0)
+	{
+		ParamMapIdx = OutHistory.TraceParameterMapOutputPin(UNiagaraNode::TraceOutputPin(GetInputPin(0)->LinkedTo[0]));
+	}
+
+	if (ParamMapIdx != INDEX_NONE)
+	{
+		uint32 NodeIdx = OutHistory.BeginNodeVisitation(ParamMapIdx, this);
+		OutHistory.EndNodeVisitation(ParamMapIdx, NodeIdx);
+	}
+
+	OutHistory.RegisterParameterMapPin(ParamMapIdx, GetOutputPin(0));
 }
 
 #undef LOCTEXT_NAMESPACE
diff --git a/Engine/Plugins/FX/Niagara/Source/NiagaraEditor/Private/NiagaraNodeUsageSelector.h b/Engine/Plugins/FX/Niagara/Source/NiagaraEditor/Private/NiagaraNodeUsageSelector.h
new file mode 100644
--- /dev/null
+++ b/Engine/Plugins/FX/Niagara/Source/NiagaraEditor/Private/NiagaraNodeUsageSelector.h
@@ -0,0 +1,56 @@
+// Copyright 1998-2018 Epic Games, Inc. All Rights Reserved.
+
+#pragma once
+#include "NiagaraNodeWithDynamicPins.h"
+#include "EdGraphSchema_Niagara.h"
+#include "NiagaraNodeUsageSelector.generated.h"
+
+UCLASS(MinimalAPI)
+class UNiagaraNodeUsageSelector: public UNiagaraNodeWithDynamicPins
+{
+	GENERATED_UCLASS_BODY()
+
+public:
+	UPROPERTY(EditAnywhere, Category=If)
+	TArray<FNiagaraVariable> OutputVars;
+	
+	UPROPERTY(EditAnywhere, Category=If)
+	TArray<FGuid> OutputVarGuids;
+	
+	//~ Begin UObject Interface
+	virtual void PostEditChangeProperty(FPropertyChangedEvent& PropertyChangedEvent) override;
+	virtual void PostLoad() override;
+	//~ End UObject Interface
+
+	//~ Begin EdGraphNode Interface
+	virtual void AllocateDefaultPins() override;
+	virtual FText GetTooltipText() const override;
+	virtual FText GetNodeTitle(ENodeTitleType::Type TitleType) const override;
+	//~ End EdGraphNode Interface
+
+	//~ Begin UNiagaraNode Interface
+	virtual void Compile(class FHlslNiagaraTranslator* Translator, TArray<int32>& Outputs) override;
+	virtual bool RefreshFromExternalChanges() override;
+
+	void BuildParameterMapHistory(FNiagaraParameterMapHistoryBuilder& OutHistory, bool bRecursive) override;
+
+	//~ End UNiagaraNode Interface
+
+	/** Helper function to create a variable to add to the OutputVars and FGuid to add to OutputVarGuids. */
+	FGuid AddOutput(FNiagaraTypeDefinition Type, const FName& Name);
+	virtual UEdGraphPin* GetPassThroughPin(const UEdGraphPin* LocallyOwnedOutputPin, ENiagaraScriptUsage MasterUsage) const override; 
+
+protected:
+	//~ Begin EdGraphNode Interface
+	virtual void OnPinRemoved(UEdGraphPin* PinToRemove) override;
+	//~ End EdGraphNode Interface
+
+	//~ Begin UNiagaraNodeWithDynamicPins Interface
+	virtual void OnNewTypedPinAdded(UEdGraphPin* NewPin) override;
+	virtual void OnPinRenamed(UEdGraphPin* RenamedPin, const FString& OldName) override;
+	virtual bool CanRenamePin(const UEdGraphPin* Pin) const override;
+	virtual bool CanRemovePin(const UEdGraphPin* Pin) const override;
+	virtual bool CanMovePin(const UEdGraphPin* Pin) const override { return false; }
+	virtual bool AllowNiagaraTypeForAddPin(const FNiagaraTypeDefinition& InType) override;
+	//~ End UNiagaraNodeWithDynamicPins Interface
+};
diff --git a/Engine/Plugins/FX/Niagara/Source/NiagaraEditor/Private/NiagaraNodeWithDynamicPins.cpp b/Engine/Plugins/FX/Niagara/Source/NiagaraEditor/Private/NiagaraNodeWithDynamicPins.cpp
--- a/Engine/Plugins/FX/Niagara/Source/NiagaraEditor/Private/NiagaraNodeWithDynamicPins.cpp
+++ b/Engine/Plugins/FX/Niagara/Source/NiagaraEditor/Private/NiagaraNodeWithDynamicPins.cpp
@@ -12,6 +12,8 @@
 #include "SNiagaraGraphPinAdd.h"
 #include "INiagaraEditorTypeUtilities.h"
 #include "NiagaraEditorUtilities.h"
+#include "SNiagaraParameterMapView.h"
+#include "Framework/Application/SlateApplication.h"
 
 #define LOCTEXT_NAMESPACE "NiagaraNodeWithDynamicPins"
 
@@ -31,7 +33,8 @@ void UNiagaraNodeWithDynamicPins::PinConnectionListChanged(UEdGraphPin* Pin)
 
 		CreateAddPin(Pin->Direction);
 		OnNewTypedPinAdded(Pin);
-		GetGraph()->NotifyGraphChanged();
+		MarkNodeRequiresSynchronization(__FUNCTION__, true);
+		//GetGraph()->NotifyGraphChanged();
 	}
 }
 
@@ -83,8 +86,7 @@ UEdGraphPin* UNiagaraNodeWithDynamicPins::RequestNewTypedPin(EEdGraphPinDirectio
 
 	CreateAddPin(Direction);
 	OnNewTypedPinAdded(AddPin);
-	UNiagaraGraph* Graph = GetNiagaraGraph();
-	Graph->NotifyGraphNeedsRecompile();
+	MarkNodeRequiresSynchronization(__FUNCTION__, true);
 
 	return AddPin;
 }
@@ -151,14 +153,15 @@ void UNiagaraNodeWithDynamicPins::MoveDynamicPin(UEdGraphPin* Pin, int32 Directi
 				
 				Pins[SwapRealPinIdx] = Pin;
 				Pins[RealPinIdx] = PinOld;
-				GetGraph()->NotifyGraphChanged();
+				//GetGraph()->NotifyGraphChanged();
+
+				MarkNodeRequiresSynchronization(__FUNCTION__, true);
 				break;
 			}
 		}
 	}
 }
 
-
 void UNiagaraNodeWithDynamicPins::GetContextMenuActions(const FGraphNodeContextMenuBuilder& Context) const
 {
 	Super::GetContextMenuActions(Context);
@@ -222,14 +225,7 @@ void UNiagaraNodeWithDynamicPins::GetContextMenuActions(const FGraphNodeContextM
 	}
 }
 
-TSharedRef<SWidget> UNiagaraNodeWithDynamicPins::GenerateAddPinMenu(const FString& InWorkingPinName, SNiagaraGraphPinAdd* InPin)
-{
-	FMenuBuilder MenuBuilder(true, nullptr);
-	BuildTypeMenu(MenuBuilder, InWorkingPinName, InPin);
-	return MenuBuilder.MakeWidget();
-}
-
-void UNiagaraNodeWithDynamicPins::BuildTypeMenu(FMenuBuilder& InMenuBuilder, const FString& InWorkingName, SNiagaraGraphPinAdd* InPin)
+void UNiagaraNodeWithDynamicPins::CollectAddPinActions(FGraphActionListBuilderBase& OutActions, bool& bOutCreateRemainingActions, UEdGraphPin* Pin)
 {
 	TArray<FNiagaraTypeDefinition> Types = FNiagaraTypeRegistry::GetRegisteredTypes();
 	Types.Sort([](const FNiagaraTypeDefinition& A, const FNiagaraTypeDefinition& B) { return (A.GetNameText().ToLower().ToString() < B.GetNameText().ToLower().ToString()); });
@@ -241,23 +237,35 @@ void UNiagaraNodeWithDynamicPins::BuildTypeMenu(FMenuBuilder& InMenuBuilder, con
 
 		if (bAllowType)
 		{
-			FNiagaraVariable Var(RegisteredType, *InWorkingName);
+			FNiagaraVariable Var(RegisteredType, FName(*RegisteredType.GetName()));
 			FNiagaraEditorUtilities::ResetVariableToDefaultValue(Var);
 
-			InMenuBuilder.AddMenuEntry(
-				RegisteredType.GetNameText(),
-				FText::Format(LOCTEXT("AddButtonTypeEntryToolTipFormat", "Add a new {0} pin"), RegisteredType.GetNameText()),
-				FSlateIcon(),
-				FUIAction(FExecuteAction::CreateSP(InPin, &SNiagaraGraphPinAdd::OnAddType, Var)));
+			const FText DisplayName = RegisteredType.GetNameText();
+			const FText Tooltip = FText::Format(LOCTEXT("AddButtonTypeEntryToolTipFormat", "Add a new {0} pin"), RegisteredType.GetNameText());
+			TSharedPtr<FNiagaraMenuAction> Action(new FNiagaraMenuAction(
+				FText::GetEmpty(), DisplayName, Tooltip, 0, FText::GetEmpty(),
+				FNiagaraMenuAction::FOnExecuteStackAction::CreateUObject(this, &UNiagaraNodeWithDynamicPins::AddParameter, Var, Pin)));
+
+			OutActions.AddAction(Action);
 		}
 	}
+
+	bOutCreateRemainingActions = false;
+}
+
+void UNiagaraNodeWithDynamicPins::AddParameter(FNiagaraVariable Parameter, UEdGraphPin* AddPin)
+{
+	FScopedTransaction AddNewPinTransaction(LOCTEXT("AddNewPinTransaction", "Add pin to node"));
+	this->RequestNewTypedPin(AddPin->Direction, Parameter.GetType(), Parameter.GetName());
 }
 
 void UNiagaraNodeWithDynamicPins::RemoveDynamicPin(UEdGraphPin* Pin)
 {
 	FScopedTransaction RemovePinTransaction(LOCTEXT("RemovePinTransaction", "Remove pin"));
 	RemovePin(Pin);
-	GetGraph()->NotifyGraphChanged();
+	//GetGraph()->NotifyGraphChanged();
+
+	MarkNodeRequiresSynchronization(__FUNCTION__, true);
 }
 
 FText UNiagaraNodeWithDynamicPins::GetPinNameText(UEdGraphPin* Pin) const
diff --git a/Engine/Plugins/FX/Niagara/Source/NiagaraEditor/Private/NiagaraNodeWriteDataSet.cpp b/Engine/Plugins/FX/Niagara/Source/NiagaraEditor/Private/NiagaraNodeWriteDataSet.cpp
--- a/Engine/Plugins/FX/Niagara/Source/NiagaraEditor/Private/NiagaraNodeWriteDataSet.cpp
+++ b/Engine/Plugins/FX/Niagara/Source/NiagaraEditor/Private/NiagaraNodeWriteDataSet.cpp
@@ -37,6 +37,7 @@ void UNiagaraNodeWriteDataSet::AllocateDefaultPins()
 		//Pin->bDefaultValueIsIgnored = true;
 	}
 
+	AddParameterMapPins();
 	AddConditionPin();
 	
 	bool useFriendlyNames = (VariableFriendlyNames.Num() == Variables.Num());
@@ -83,7 +84,7 @@ void UNiagaraNodeWriteDataSet::Compile(class FHlslNiagaraTranslator* Translator,
 
 	FNiagaraDataSetID AlteredDataSet = DataSet;
 	AlteredDataSet.Name = EventName;
-	Translator->WriteDataSet(AlteredDataSet, Variables, ENiagaraDataSetAccessMode::AppendConsume, Inputs);
+	Translator->WriteDataSet(AlteredDataSet, Variables, ENiagaraDataSetAccessMode::AppendConsume, Inputs, Outputs);
 
 }
 
@@ -103,7 +104,7 @@ void UNiagaraNodeWriteDataSet::PostLoad()
 
 	if (!bFoundMatchingPin)
 	{
-		AddConditionPin(0);
+		AddConditionPin(1);
 	}
 
 	if (EventName.IsNone())
@@ -112,6 +113,35 @@ void UNiagaraNodeWriteDataSet::PostLoad()
 	}
 }
 
+void UNiagaraNodeWriteDataSet::BuildParameterMapHistory(FNiagaraParameterMapHistoryBuilder& OutHistory, bool bRecursive)
+{
+	if (bRecursive)
+	{
+		OutHistory.VisitInputPins(this);
+	}
+
+	if (!IsNodeEnabled() && OutHistory.GetIgnoreDisabled())
+	{
+		RouteParameterMapAroundMe(OutHistory, bRecursive);
+		return;
+	}
+
+	int32 ParamMapIdx = INDEX_NONE;
+	if (GetInputPin(0)->LinkedTo.Num() != 0)
+	{
+		ParamMapIdx = OutHistory.TraceParameterMapOutputPin(UNiagaraNode::TraceOutputPin(GetInputPin(0)->LinkedTo[0]));
+	}
+
+	if (ParamMapIdx != INDEX_NONE)
+	{
+		uint32 NodeIdx = OutHistory.BeginNodeVisitation(ParamMapIdx, this);
+		OutHistory.EndNodeVisitation(ParamMapIdx, NodeIdx);
+	}
+
+	OutHistory.RegisterParameterMapPin(ParamMapIdx, GetOutputPin(0));
+}
+
+
 #undef LOCTEXT_NAMESPACE
 
 
diff --git a/Engine/Plugins/FX/Niagara/Source/NiagaraEditor/Private/NiagaraParameterCollectionFactoryNew.cpp b/Engine/Plugins/FX/Niagara/Source/NiagaraEditor/Private/NiagaraParameterCollectionFactoryNew.cpp
--- a/Engine/Plugins/FX/Niagara/Source/NiagaraEditor/Private/NiagaraParameterCollectionFactoryNew.cpp
+++ b/Engine/Plugins/FX/Niagara/Source/NiagaraEditor/Private/NiagaraParameterCollectionFactoryNew.cpp
@@ -23,6 +23,7 @@ UObject* UNiagaraParameterCollectionFactoryNew::FactoryCreateNew(UClass* Class,
 	check(Class->IsChildOf(UNiagaraParameterCollection::StaticClass()));
 
 	UNiagaraParameterCollection* NewCollection = NewObject<UNiagaraParameterCollection>(InParent, Class, Name, Flags | RF_Transactional);
+	NewCollection->RefreshCompileId();
 
 	return NewCollection;
 }
diff --git a/Engine/Plugins/FX/Niagara/Source/NiagaraEditor/Private/NiagaraParameterMapHistory.cpp b/Engine/Plugins/FX/Niagara/Source/NiagaraEditor/Private/NiagaraParameterMapHistory.cpp
--- a/Engine/Plugins/FX/Niagara/Source/NiagaraEditor/Private/NiagaraParameterMapHistory.cpp
+++ b/Engine/Plugins/FX/Niagara/Source/NiagaraEditor/Private/NiagaraParameterMapHistory.cpp
@@ -18,11 +18,15 @@
 #include "ViewModels/Stack/NiagaraStackEntry.h"
 #include "ViewModels/Stack/NiagaraStackGraphUtilities.h"
 #include "NiagaraParameterCollection.h"
+#include "NiagaraConstants.h"
 
 #define LOCTEXT_NAMESPACE "NiagaraEditor"
 
 #define PARAM_MAP_NPC_STR TEXT("NPC.")
 #define PARAM_MAP_ENGINE_STR TEXT("Engine.")
+#define PARAM_MAP_ENGINE_OWNER_STR TEXT("Engine.Owner.")
+#define PARAM_MAP_ENGINE_SYSTEM_STR TEXT("Engine.System.")
+#define PARAM_MAP_ENGINE_EMITTER_STR TEXT("Engine.Emitter.")
 #define PARAM_MAP_USER_STR TEXT("User.")
 #define PARAM_MAP_SYSTEM_STR TEXT("System.")
 #define PARAM_MAP_EMITTER_STR TEXT("Emitter.")
@@ -31,6 +35,7 @@
 #define PARAM_MAP_INITIAL_STR TEXT("Initial.")
 #define PARAM_MAP_INITIAL_BASE_STR TEXT("Initial")
 #define PARAM_MAP_RAPID_ITERATION_STR TEXT("Constants.")
+#define PARAM_MAP_RAPID_ITERATION_BASE_STR TEXT("Constants")
 
 
 FNiagaraParameterMapHistory::FNiagaraParameterMapHistory() 
@@ -38,20 +43,101 @@ FNiagaraParameterMapHistory::FNiagaraParameterMapHistory()
 	
 }
 
+void FNiagaraParameterMapHistory::GetValidNamespacesForReading(const UNiagaraScript* InScript, TArray<FString>& OutputNamespaces)
+{
+	GetValidNamespacesForReading(InScript->GetUsage(), OutputNamespaces);
+}
+
+void FNiagaraParameterMapHistory::GetValidNamespacesForReading(ENiagaraScriptUsage InScriptUsage, TArray<FString>& OutputNamespaces)
+{
+	OutputNamespaces.Add(PARAM_MAP_MODULE_STR);
+	OutputNamespaces.Add(PARAM_MAP_ENGINE_STR);
+	OutputNamespaces.Add(PARAM_MAP_NPC_STR);
+	OutputNamespaces.Add(PARAM_MAP_USER_STR);
+	OutputNamespaces.Add(PARAM_MAP_SYSTEM_STR);
+	OutputNamespaces.Add(PARAM_MAP_EMITTER_STR);
+	if (UNiagaraScript::IsParticleScript(InScriptUsage))
+	{
+		OutputNamespaces.Add(PARAM_MAP_ATTRIBUTE_STR);
+	}
+}
+
+FString FNiagaraParameterMapHistory::GetNamespace(const FNiagaraVariable& InVar, bool bIncludeDelimiter)
+{
+	TArray<FString> SplitName;
+	InVar.GetName().ToString().ParseIntoArray(SplitName, TEXT("."));
+
+	check(SplitName.Num() > 0);
+
+	if (bIncludeDelimiter)
+	{
+		return SplitName[0] + TEXT(".");
+	}
+	else
+	{
+		return SplitName[0];
+	}
+}
+
+bool FNiagaraParameterMapHistory::IsValidNamespaceForReading(ENiagaraScriptUsage InScriptUsage, FString Namespace)
+{
+	TArray<FString> OutputNamespaces;
+	GetValidNamespacesForReading(InScriptUsage, OutputNamespaces);
+	
+	if (!Namespace.EndsWith(TEXT(".")))
+	{
+		Namespace.Append(TEXT("."));
+	}
+
+	for (const FString& ValidNamespace : OutputNamespaces)
+	{
+		if (Namespace.StartsWith(ValidNamespace))
+		{
+			return true;
+		}
+	}
+
+	return false;
+}
+
+
 int32 FNiagaraParameterMapHistory::RegisterParameterMapPin(const UEdGraphPin* Pin)
 {
 	int32 RetIdx =  MapPinHistory.Add(Pin);
 	return RetIdx;
 }
 
-int32 FNiagaraParameterMapHistory::FindVariableByName(const FName& VariableName)
+uint32 FNiagaraParameterMapHistory::BeginNodeVisitation(UNiagaraNode* Node)
 {
-	int32 FoundIdx = Variables.IndexOfByPredicate([&](const FNiagaraVariable& InObj) -> bool
+	uint32 AddedIndex = MapNodeVisitations.Add(Node);
+	MapNodeVariableMetaData.Add(TTuple<uint32, uint32>(Variables.Num(), 0));
+	check(MapNodeVisitations.Num() == MapNodeVariableMetaData.Num());
+	return AddedIndex;
+}
+
+void FNiagaraParameterMapHistory::EndNodeVisitation(uint32 IndexFromBeginNode)
+{
+	check(IndexFromBeginNode < (uint32)MapNodeVisitations.Num());
+	check(MapNodeVisitations.Num() == MapNodeVariableMetaData.Num());
+	MapNodeVariableMetaData[IndexFromBeginNode].Value = Variables.Num();
+}
+
+
+int32 FNiagaraParameterMapHistory::FindVariableByName(const FName& VariableName, bool bAllowPartialMatch)
+{
+	if (!bAllowPartialMatch)
 	{
-		return (InObj.GetName() == VariableName);
-	});
+		int32 FoundIdx = Variables.IndexOfByPredicate([&](const FNiagaraVariable& InObj) -> bool
+		{
+			return (InObj.GetName() == VariableName);
+		});
 
-	return FoundIdx;
+		return FoundIdx;
+	}
+	else
+	{
+		return FNiagaraVariable::SearchArrayForPartialNameMatch(Variables, VariableName);
+	}
 }
 
 
@@ -129,6 +215,17 @@ FNiagaraVariable FNiagaraParameterMapHistory::ResolveAliases(const FNiagaraVaria
 	return FNiagaraVariable::ResolveAliases(InVar, InAliases, InJoinSeparator);
 }
 
+
+FName FNiagaraParameterMapHistory::ResolveEmitterAlias(const FName& InName, const FString& InAlias)
+{
+	FNiagaraVariable Var(FNiagaraTypeDefinition::GetFloatDef(), InName);
+	TMap<FString, FString> ResolveMap;
+	ResolveMap.Add(TEXT("Emitter"), InAlias);
+	Var = FNiagaraParameterMapHistory::ResolveAliases(Var, ResolveMap, TEXT("."));
+	return Var.GetName();
+}
+
+
 FString FNiagaraParameterMapHistory::MakeSafeNamespaceString(const FString& InStr)
 {
 	FString  Sanitized = FHlslNiagaraTranslator::GetSanitizedSymbolName(InStr);
@@ -207,6 +304,12 @@ bool FNiagaraParameterMapHistory::IsAliasedEmitterParameter(const FNiagaraVariab
 	return IsInNamespace(InVar, PARAM_MAP_EMITTER_STR);
 }
 
+bool FNiagaraParameterMapHistory::IsAliasedEmitterParameter(const FString& InVarName)
+{
+	return IsAliasedEmitterParameter(FNiagaraVariable(FNiagaraTypeDefinition::GetFloatDef(), *InVarName));
+}
+
+
 bool FNiagaraParameterMapHistory::IsSystemParameter(const FNiagaraVariable& InVar)
 {
 	return IsInNamespace(InVar, PARAM_MAP_SYSTEM_STR);
@@ -217,6 +320,13 @@ bool FNiagaraParameterMapHistory::IsEngineParameter(const FNiagaraVariable& InVa
 	return IsInNamespace(InVar, PARAM_MAP_ENGINE_STR);
 }
 
+bool FNiagaraParameterMapHistory::IsPerInstanceEngineParameter(const FNiagaraVariable& InVar, const FString& EmitterAlias)
+{
+	FString EmitterEngineNamespaceAlias = TEXT("Engine.") + EmitterAlias + TEXT(".");
+	return IsInNamespace(InVar, PARAM_MAP_ENGINE_OWNER_STR) || IsInNamespace(InVar, PARAM_MAP_ENGINE_SYSTEM_STR) || IsInNamespace(InVar, PARAM_MAP_ENGINE_EMITTER_STR) || 
+		IsInNamespace(InVar, EmitterEngineNamespaceAlias);
+}
+
 bool FNiagaraParameterMapHistory::IsUserParameter(const FNiagaraVariable& InVar)
 {
 	return IsInNamespace(InVar, PARAM_MAP_USER_STR);
@@ -245,7 +355,7 @@ bool FNiagaraParameterMapHistory::IsAttribute(const FNiagaraVariable& InVar)
 	return IsInNamespace(InVar, PARAM_MAP_ATTRIBUTE_STR);
 }
 
-bool FNiagaraParameterMapHistory::IsSystemNamespaceReadOnly(UNiagaraScript* InScript)
+bool FNiagaraParameterMapHistory::IsSystemNamespaceReadOnly(const UNiagaraScript* InScript)
 {
 	if (InScript->IsSystemSpawnScript() || InScript->IsSystemUpdateScript())
 	{
@@ -263,7 +373,7 @@ bool FNiagaraParameterMapHistory::IsSystemNamespaceReadOnly(UNiagaraScript* InSc
 	return true;
 }
 
-bool FNiagaraParameterMapHistory::IsEmitterNamespaceReadOnly(UNiagaraScript* InScript)
+bool FNiagaraParameterMapHistory::IsEmitterNamespaceReadOnly(const UNiagaraScript* InScript)
 {
 	if (InScript->IsEmitterSpawnScript() || InScript->IsEmitterUpdateScript() || InScript->IsSystemSpawnScript() || InScript->IsSystemUpdateScript())
 	{
@@ -378,14 +488,14 @@ FNiagaraVariableMetaData* FNiagaraParameterMapHistory::GetMetaData(int32 VarIdx)
 	return nullptr;
 }
 
-bool FNiagaraParameterMapHistory::IsPrimaryDataSetOutput(const FNiagaraVariable& InVar, UNiagaraScript* InScript)
+bool FNiagaraParameterMapHistory::IsPrimaryDataSetOutput(const FNiagaraVariable& InVar, const UNiagaraScript* InScript, bool bAllowDataInterfaces) const
 {
-	return IsPrimaryDataSetOutput(InVar, InScript->GetUsage());
+	return IsPrimaryDataSetOutput(InVar, InScript->GetUsage(), bAllowDataInterfaces);
 }
 
-bool FNiagaraParameterMapHistory::IsPrimaryDataSetOutput(const FNiagaraVariable& InVar, ENiagaraScriptUsage Usage)
+bool FNiagaraParameterMapHistory::IsPrimaryDataSetOutput(const FNiagaraVariable& InVar, ENiagaraScriptUsage Usage, bool bAllowDataInterfaces) const
 {
-	if (InVar.GetType().GetClass() != nullptr)
+	if (bAllowDataInterfaces == false && InVar.GetType().GetClass() != nullptr)
 	{
 		return false;
 	}
@@ -411,25 +521,29 @@ bool FNiagaraParameterMapHistory::IsPrimaryDataSetOutput(const FNiagaraVariable&
 	return IsInNamespace(InVar, PARAM_MAP_ATTRIBUTE_STR);
 }
 
-
-FNiagaraVariable FNiagaraParameterMapHistory::MoveToExternalConstantNamespaceVariable(const FNiagaraVariable& InVar, UNiagaraScript* InScript)
+FNiagaraVariable FNiagaraParameterMapHistory::MoveToExternalConstantNamespaceVariable(const FNiagaraVariable& InVar, ENiagaraScriptUsage InUsage)
 {
-	if (InScript->IsParticleEventScript() || InScript->IsParticleSpawnScript() || InScript->IsParticleUpdateScript())
+	if (UNiagaraScript::IsParticleScript(InUsage))
 	{
 		return VariableToNamespacedVariable(InVar, PARAM_MAP_EMITTER_STR);
 	}
-	else if (InScript->IsStandaloneScript())
+	else if (UNiagaraScript::IsStandaloneScript(InUsage))
 	{
 		return VariableToNamespacedVariable(InVar, PARAM_MAP_MODULE_STR);
 	}
-	else if (InScript->IsEmitterSpawnScript() || InScript->IsEmitterUpdateScript() || InScript->IsSystemSpawnScript() || InScript->IsSystemUpdateScript())
+	else if (UNiagaraScript::IsEmitterSpawnScript(InUsage) || UNiagaraScript::IsEmitterUpdateScript(InUsage) || UNiagaraScript::IsSystemSpawnScript(InUsage) || UNiagaraScript::IsSystemUpdateScript(InUsage))
 	{
 		return VariableToNamespacedVariable(InVar, PARAM_MAP_USER_STR);
 	}
 	return InVar;
 }
 
-bool FNiagaraParameterMapHistory::IsExportableExternalConstant(const FNiagaraVariable& InVar, UNiagaraScript* InScript)
+FNiagaraVariable FNiagaraParameterMapHistory::MoveToExternalConstantNamespaceVariable(const FNiagaraVariable& InVar, const UNiagaraScript* InScript)
+{
+	return MoveToExternalConstantNamespaceVariable(InVar, InScript->GetUsage());
+}
+
+bool FNiagaraParameterMapHistory::IsExportableExternalConstant(const FNiagaraVariable& InVar, const UNiagaraScript* InScript)
 {
 	if (InScript->IsEquivalentUsage(ENiagaraScriptUsage::SystemSpawnScript))
 	{
@@ -441,7 +555,7 @@ bool FNiagaraParameterMapHistory::IsExportableExternalConstant(const FNiagaraVar
 	}
 }
 
-bool FNiagaraParameterMapHistory::IsExternalConstantNamespace(const FNiagaraVariable& InVar, UNiagaraScript* InScript)
+bool FNiagaraParameterMapHistory::IsExternalConstantNamespace(const FNiagaraVariable& InVar, ENiagaraScriptUsage InUsage, int32 InUsageBitmask)
 {
 	// Parameter collections are always constants
 	if (IsInNamespace(InVar, PARAM_MAP_NPC_STR))
@@ -462,28 +576,33 @@ bool FNiagaraParameterMapHistory::IsExternalConstantNamespace(const FNiagaraVari
 
 	// Modules and functions need to act as if they are within the script types that they 
 	// say that they support rather than using their exact script type.
-	if (InScript->IsStandaloneScript())
+	if (UNiagaraScript::IsStandaloneScript(InUsage))
 	{
-		TArray<ENiagaraScriptUsage> SupportedContexts = InScript->GetSupportedUsageContexts();
+		TArray<ENiagaraScriptUsage> SupportedContexts = UNiagaraScript::GetSupportedUsageContextsForBitmask(InUsageBitmask);
 		if (((!SupportedContexts.Contains(ENiagaraScriptUsage::EmitterSpawnScript) && !SupportedContexts.Contains(ENiagaraScriptUsage::EmitterUpdateScript)) && IsInNamespace(InVar, PARAM_MAP_EMITTER_STR))
 			|| ((!SupportedContexts.Contains(ENiagaraScriptUsage::SystemSpawnScript) && !SupportedContexts.Contains(ENiagaraScriptUsage::SystemUpdateScript)) && IsInNamespace(InVar, PARAM_MAP_SYSTEM_STR)))
 		{
 			return true;
 		}
 	}
-	
+
 	// Particle scripts cannot write to the emitter or system namespace.
-	if (InScript->IsParticleEventScript() || InScript->IsParticleSpawnScript() || InScript->IsParticleUpdateScript())
+	if (UNiagaraScript::IsParticleScript(InUsage))
 	{
 		if (IsInNamespace(InVar, PARAM_MAP_EMITTER_STR) || IsInNamespace(InVar, PARAM_MAP_SYSTEM_STR))
 		{
 			return true;
 		}
 	}
-	
+
 	return false;
 }
 
+bool FNiagaraParameterMapHistory::IsExternalConstantNamespace(const FNiagaraVariable& InVar, const UNiagaraScript* InScript)
+{
+	return IsExternalConstantNamespace(InVar, InScript->GetUsage(), InScript->ModuleUsageBitmask);
+}
+
 const UNiagaraNodeOutput* FNiagaraParameterMapHistory::GetFinalOutputNode() const
 {
 	const UEdGraphPin* Pin = GetFinalPin();
@@ -499,17 +618,9 @@ const UNiagaraNodeOutput* FNiagaraParameterMapHistory::GetFinalOutputNode() cons
 	return nullptr;
 }
 
-FNiagaraVariable FNiagaraParameterMapHistory::ConvertVariableToRapidIterationConstantName(FNiagaraVariable InVar, const TCHAR* InEmitterName)
+FNiagaraVariable FNiagaraParameterMapHistory::ConvertVariableToRapidIterationConstantName(FNiagaraVariable InVar, const TCHAR* InEmitterName, ENiagaraScriptUsage InUsage)
 {
-	FNiagaraVariable Var = InVar;
-	FString NewName = PARAM_MAP_RAPID_ITERATION_STR + Var.GetName().ToString();
-	FString EmitterName = FString(InEmitterName);
-	if (InEmitterName != nullptr && false == InVar.GetName().ToString().StartsWith(EmitterName + TEXT(".")))
-	{
-		NewName = PARAM_MAP_RAPID_ITERATION_STR + EmitterName + TEXT(".") + Var.GetName().ToString();
-	}
-	Var.SetName(*NewName);
-	return Var;
+	return FNiagaraUtilities::ConvertVariableToRapidIterationConstantName(InVar, InEmitterName, InUsage);
 }
 
 UNiagaraParameterCollection* FNiagaraParameterMapHistory::IsParameterCollectionParameter(FNiagaraVariable& InVar, bool& bMissingParameter)
@@ -527,6 +638,12 @@ UNiagaraParameterCollection* FNiagaraParameterMapHistory::IsParameterCollectionP
 	return nullptr;
 }
 
+bool FNiagaraParameterMapHistory::ShouldIgnoreVariableDefault(const FNiagaraVariable& Var)const
+{
+	//For now just skip default for ID but maybe other cases/reasons?
+	return Var == FNiagaraVariable(FNiagaraTypeDefinition::GetIDDef(), TEXT("Particles.ID"));
+}
+
 FNiagaraParameterMapHistoryBuilder::FNiagaraParameterMapHistoryBuilder()
 {
 	ContextuallyVisitedNodes.AddDefaulted(1);
@@ -549,17 +666,61 @@ void FNiagaraParameterMapHistoryBuilder::EnableScriptWhitelist(bool bInEnable, E
 	FilterScriptType = InScriptType;
 }
 
+
+bool FNiagaraParameterMapHistoryBuilder::ContextContains(ENiagaraScriptUsage InUsage) const
+{
+	if (RelevantScriptUsageContext.Num() == 0)
+	{
+		return false;
+	}
+	else
+	{
+		return RelevantScriptUsageContext.Contains(InUsage);
+	}
+}
+
+
 ENiagaraScriptUsage FNiagaraParameterMapHistoryBuilder::GetCurrentUsageContext()const
 {
+	if (RelevantScriptUsageContext.Num() == 0)
+		return ENiagaraScriptUsage::Function;
 	return RelevantScriptUsageContext.Last();
 }
 
+ENiagaraScriptUsage FNiagaraParameterMapHistoryBuilder::GetBaseUsageContext()const
+{
+	if (RelevantScriptUsageContext.Num() == 0)
+		return ENiagaraScriptUsage::Function;
+	return RelevantScriptUsageContext[0];
+}
+
 int32 FNiagaraParameterMapHistoryBuilder::CreateParameterMap()
 {
 	int32 RetValue = Histories.AddDefaulted(1);
 	return RetValue;
 }
 
+uint32 FNiagaraParameterMapHistoryBuilder::BeginNodeVisitation(int32 WhichParameterMap, class UNiagaraNode* Node)
+{
+	if (WhichParameterMap != INDEX_NONE)
+	{
+		return Histories[WhichParameterMap].BeginNodeVisitation(Node);
+	}
+	else
+	{
+		return INDEX_NONE;
+	}
+}
+
+void FNiagaraParameterMapHistoryBuilder::EndNodeVisitation(int32 WhichParameterMap, uint32 IndexFromBeginNode)
+{
+	if (WhichParameterMap != INDEX_NONE)
+	{
+		return Histories[WhichParameterMap].EndNodeVisitation(IndexFromBeginNode);
+	}
+}
+
+
 int32 FNiagaraParameterMapHistoryBuilder::RegisterParameterMapPin(int32 WhichParameterMap, const UEdGraphPin* Pin)
 {
 	if (WhichParameterMap != INDEX_NONE)
@@ -581,10 +742,14 @@ int32 FNiagaraParameterMapHistoryBuilder::TraceParameterMapOutputPin(const UEdGr
 {
 	if (OutputPin->Direction == EEdGraphPinDirection::EGPD_Output)
 	{
-		const int32* IdxPtr = PinToParameterMapIndices.Last().Find(OutputPin);
-		if (IdxPtr != nullptr)
+		OutputPin = UNiagaraNode::TraceOutputPin(const_cast<UEdGraphPin*>(OutputPin));
+		if (OutputPin)
 		{
-			return *IdxPtr;
+			const int32* IdxPtr = PinToParameterMapIndices.Last().Find(OutputPin);
+			if (IdxPtr != nullptr)
+			{
+				return *IdxPtr;
+			}
 		}
 	}
 	return INDEX_NONE;
@@ -622,7 +787,8 @@ int32 FNiagaraParameterMapHistoryBuilder::FindMatchingParameterMapFromContextInp
 		{
 			if (Inputs[i]->LinkedTo.Num() != 0 && PinToParameterMapIndices.Num() >= 2)
 			{
-				const int32* ParamMapIdxPtr = PinToParameterMapIndices[PinToParameterMapIndices.Num() - 2].Find(Inputs[i]->LinkedTo[0]);
+				const UEdGraphPin* OutputPin = UNiagaraNode::TraceOutputPin(const_cast<UEdGraphPin*>(Inputs[i]->LinkedTo[0]));
+				const int32* ParamMapIdxPtr = PinToParameterMapIndices[PinToParameterMapIndices.Num() - 2].Find(OutputPin);
 				if (ParamMapIdxPtr != nullptr)
 				{
 					return *ParamMapIdxPtr;
@@ -663,7 +829,7 @@ int32 FNiagaraParameterMapHistoryBuilder::FindMatchingParameterMapFromContextInp
 	return INDEX_NONE;
 }
 
-void FNiagaraParameterMapHistoryBuilder::BeginTranslation(UNiagaraScript* Script)
+void FNiagaraParameterMapHistoryBuilder::BeginTranslation(const UNiagaraScript* Script)
 {
 	//For now this will just tell particle scripts what emitter they're being compiled as part of but maybe we want to do more here.
 	//This is mainly so that parameter names match up between System/Emitter scripts and the parameters they drive within particle scripts.
@@ -673,12 +839,12 @@ void FNiagaraParameterMapHistoryBuilder::BeginTranslation(UNiagaraScript* Script
 	BeginTranslation(Emitter);
 }
 
-void FNiagaraParameterMapHistoryBuilder::EndTranslation(UNiagaraScript* Script)
+void FNiagaraParameterMapHistoryBuilder::EndTranslation(const UNiagaraScript* Script)
 {
 	EmitterNameContextStack.Reset();
 }
 
-void FNiagaraParameterMapHistoryBuilder::BeginTranslation(UNiagaraEmitter* Emitter)
+void FNiagaraParameterMapHistoryBuilder::BeginTranslation(const UNiagaraEmitter* Emitter)
 {
 	//For now this will just tell particle scripts what emitter they're being compiled as part of but maybe we want to do more here.
 	//This is mainly so that parameter names match up between System/Emitter scripts and the parameters they drive within particle scripts.
@@ -692,7 +858,24 @@ void FNiagaraParameterMapHistoryBuilder::BeginTranslation(UNiagaraEmitter* Emitt
 	BuildCurrentAliases();
 }
 
-void FNiagaraParameterMapHistoryBuilder::EndTranslation(UNiagaraEmitter* Emitter)
+void FNiagaraParameterMapHistoryBuilder::EndTranslation(const UNiagaraEmitter* Emitter)
+{
+	EmitterNameContextStack.Reset();
+}
+
+void FNiagaraParameterMapHistoryBuilder::BeginTranslation(const FString& EmitterUniqueName)
+{
+	//For now this will just tell particle scripts what emitter they're being compiled as part of but maybe we want to do more here.
+	//This is mainly so that parameter names match up between System/Emitter scripts and the parameters they drive within particle scripts.
+	//I dislike this coupling of the translator to emitters but for now it'll have to do.
+	//Will refactor in the future.
+	if (EmitterUniqueName.IsEmpty() == false)
+	{
+		EmitterNameContextStack.Add(*EmitterUniqueName);
+	}
+	BuildCurrentAliases();
+}
+void FNiagaraParameterMapHistoryBuilder::EndTranslation(const FString& EmitterUniqueName)
 {
 	EmitterNameContextStack.Reset();
 }
@@ -723,7 +906,7 @@ bool FNiagaraParameterMapHistoryBuilder::InTopLevelFunctionCall(ENiagaraScriptUs
 		}
 
 	}
-	else if (InFilterScriptType == ENiagaraScriptUsage::ParticleSpawnScript || InFilterScriptType == ENiagaraScriptUsage::ParticleEventScript || InFilterScriptType == ENiagaraScriptUsage::ParticleSpawnScriptInterpolated || InFilterScriptType == ENiagaraScriptUsage::ParticleUpdateScript)
+	else if (UNiagaraScript::IsParticleScript(InFilterScriptType))
 	{
 		if (CallingContext.Num() <= 1) // Handle a function call
 		{
@@ -735,7 +918,7 @@ bool FNiagaraParameterMapHistoryBuilder::InTopLevelFunctionCall(ENiagaraScriptUs
 }
 
 
-void FNiagaraParameterMapHistoryBuilder::EnterFunction(const FString& InNodeName, UNiagaraScript* InScript, UNiagaraNode* Node)
+void FNiagaraParameterMapHistoryBuilder::EnterFunction(const FString& InNodeName, const UNiagaraScript* InScript, UNiagaraNode* Node)
 {
 	if (InScript != nullptr )
 	{
@@ -752,7 +935,7 @@ void FNiagaraParameterMapHistoryBuilder::EnterFunction(const FString& InNodeName
 	}
 }
 
-void FNiagaraParameterMapHistoryBuilder::ExitFunction(const FString& InNodeName, UNiagaraScript* InScript, UNiagaraNode* Node)
+void FNiagaraParameterMapHistoryBuilder::ExitFunction(const FString& InNodeName, const UNiagaraScript* InScript, UNiagaraNode* Node)
 {
 	if (InScript != nullptr)
 	{
@@ -839,7 +1022,6 @@ FNiagaraVariable FNiagaraParameterMapHistoryBuilder::ResolveAliases(const FNiaga
 	return FNiagaraParameterMapHistory::ResolveAliases(InVar, AliasMap, TEXT("."));
 }
 
-
 void FNiagaraParameterMapHistoryBuilder::RegisterNodeVisitation(const UEdGraphNode* Node)
 {
 	ContextuallyVisitedNodes.Last().AddUnique(CastChecked<UNiagaraNode>(Node));
@@ -864,9 +1046,10 @@ void FNiagaraParameterMapHistoryBuilder::VisitInputPin(const UEdGraphPin* Pin, c
 	{
 		for (int32 j = 0; j < Pin->LinkedTo.Num(); j++)
 		{
-			if (Pin->LinkedTo[j])
+			const UEdGraphPin* OutputPin = UNiagaraNode::TraceOutputPin(Pin->LinkedTo[j]);
+			if (OutputPin)
 			{
-				UNiagaraNode* Node = CastChecked<UNiagaraNode>(Pin->LinkedTo[j]->GetOwningNode());
+				UNiagaraNode* Node = CastChecked<UNiagaraNode>(OutputPin->GetOwningNode());
 
 				if (!GetNodePreviouslyVisited(Node))
 				{
@@ -876,7 +1059,7 @@ void FNiagaraParameterMapHistoryBuilder::VisitInputPin(const UEdGraphPin* Pin, c
 
 				if (Schema->PinToTypeDefinition(Pin) == FNiagaraTypeDefinition::GetParameterMapDef())
 				{
-					int32 ParamMapIdx = TraceParameterMapOutputPin(Pin->LinkedTo[j]);
+					int32 ParamMapIdx = TraceParameterMapOutputPin(OutputPin);
 					RegisterParameterMapPin(ParamMapIdx, Pin);
 				}
 			}
@@ -928,6 +1111,18 @@ int32 FNiagaraParameterMapHistoryBuilder::HandleVariableWrite(int32 ParamMapIdx,
 	return Histories[ParamMapIdx].AddVariable(Var,AliasedVar, InPin);
 }
 
+int32 FNiagaraParameterMapHistoryBuilder::HandleVariableWrite(int32 ParameterMapIndex, const FNiagaraVariable& Var)
+{
+	const UEdGraphSchema_Niagara* Schema = GetDefault<UEdGraphSchema_Niagara>();
+
+	if (!ShouldTrackVariable(Var))
+	{
+		return INDEX_NONE;
+	}
+	FNiagaraVariable ResolvedVar = ResolveAliases(Var);
+
+	return Histories[ParameterMapIndex].AddVariable(ResolvedVar, Var, nullptr);
+}
 
 int32 FNiagaraParameterMapHistoryBuilder::HandleVariableRead(int32 ParamMapIdx, const UEdGraphPin* InPin, bool RegisterReadsAsVariables, const UEdGraphPin* InDefaultPin, bool& OutUsedDefault)
 {
@@ -1008,6 +1203,124 @@ int32 FNiagaraParameterMapHistoryBuilder::HandleVariableRead(int32 ParamMapIdx,
 	return FoundIdx;
 }
 
+void FNiagaraParameterMapHistoryBuilder::RegisterEncounterableVariables(const TArray<FNiagaraVariable>& Variables)
+{
+	EncounterableExternalVariables.Append(Variables);
+}
+
+void UpdateAliasedVariable(FNiagaraVariable& AliasedVar, const FNiagaraVariable& OriginalUnaliasedVar, const FNiagaraVariable& UpdatedUnaliasedVar)
+{
+	AliasedVar.SetType(UpdatedUnaliasedVar.GetType());
+
+	TArray<FString> AliasedSplitName;
+	AliasedVar.GetName().ToString().ParseIntoArray(AliasedSplitName, TEXT("."));
+
+	TArray<FString> OriginalUnaliasedSplitName;
+	OriginalUnaliasedVar.GetName().ToString().ParseIntoArray(OriginalUnaliasedSplitName, TEXT("."));
+
+	TArray<FString> UpdatedUnaliasedSplitName;
+	UpdatedUnaliasedVar.GetName().ToString().ParseIntoArray(UpdatedUnaliasedSplitName, TEXT("."));
+
+	TArray<FString> JoinName;
+	for (int32 i = 0; i < AliasedSplitName.Num(); i++)
+	{
+		if (i >= OriginalUnaliasedSplitName.Num() || i >= UpdatedUnaliasedSplitName.Num())
+		{
+			continue;
+		}
+
+		//if (UpdatedUnaliasedSplitName[i] == OriginalUnaliasedSplitName[i])
+		{
+			JoinName.Add(AliasedSplitName[i]);
+		}
+		//else
+		//{
+		//	JoinName.Add(AliasedSplitName[i]);
+		//}
+	}
+
+	FString OutVarStrName = FString::Join<FString>(JoinName, TEXT("."));
+	AliasedVar.SetName(*OutVarStrName);
+}
+
+int32 FNiagaraParameterMapHistoryBuilder::HandleExternalVariableRead(int32 ParamMapIdx, const FName& Name)
+{
+	const UEdGraphSchema_Niagara* Schema = GetDefault<UEdGraphSchema_Niagara>();
+
+	FNiagaraParameterMapHistory& History = Histories[ParamMapIdx];
+
+	FNiagaraVariable Var(FNiagaraTypeDefinition(), Name);
+
+	if (!ShouldTrackVariable(Var))
+	{
+		return INDEX_NONE;
+	}
+	FNiagaraVariable AliasedVar = Var;
+	Var = ResolveAliases(Var);
+	FNiagaraVariable OriginalUnaliasedVar = Var;
+
+	//Track any parameter collections we're referencing.
+	FNiagaraVariable FoundTempVar;
+	if (UNiagaraParameterCollection* Collection = Schema->VariableIsFromParameterCollection(Name.ToString(), true, FoundTempVar))
+	{
+		int32 Index = History.ParameterCollections.AddUnique(Collection);
+		History.ParameterCollectionNamespaces.SetNum(History.ParameterCollections.Num());
+		History.ParameterCollectionVariables.SetNum(History.ParameterCollections.Num());
+		History.ParameterCollectionNamespaces[Index] = Collection->GetFullNamespace();
+		History.ParameterCollectionVariables[Index] = Collection->GetParameters();
+		Var = FoundTempVar;
+		UpdateAliasedVariable(AliasedVar, OriginalUnaliasedVar, Var);
+	}
+
+	int32 FoundIdx = Histories[ParamMapIdx].FindVariableByName(Name, true);
+	
+	if (FoundIdx == -1)
+	{
+		const FNiagaraVariable* TempKnownConstant = FNiagaraConstants::GetKnownConstant(Name, true);
+
+		if (!Var.IsValid() && TempKnownConstant != nullptr)
+		{
+			Var = *TempKnownConstant;
+			UpdateAliasedVariable(AliasedVar, OriginalUnaliasedVar, Var);
+		}
+
+		if (!Var.IsValid())
+		{
+			int32 EncounterableFoundIdx = FNiagaraVariable::SearchArrayForPartialNameMatch(EncounterableExternalVariables, Name);
+
+			if (EncounterableFoundIdx != INDEX_NONE)
+			{
+				Var = EncounterableExternalVariables[EncounterableFoundIdx];
+				UpdateAliasedVariable(AliasedVar, OriginalUnaliasedVar, Var);
+			}
+		}
+
+		if (Var.IsValid())
+		{
+			FoundIdx = Histories[ParamMapIdx].Variables.Add(Var);
+			Histories[ParamMapIdx].VariablesWithOriginalAliasesIntact.Add(AliasedVar);
+			Histories[ParamMapIdx].PerVariableWarnings.AddDefaulted(1);
+			Histories[ParamMapIdx].PerVariableWriteHistory.AddDefaulted(1);
+			Histories[ParamMapIdx].PerVariableReadHistory.AddDefaulted(1);
+			check(Histories[ParamMapIdx].Variables.Num() == Histories[ParamMapIdx].PerVariableWarnings.Num());
+			check(Histories[ParamMapIdx].Variables.Num() == Histories[ParamMapIdx].PerVariableWriteHistory.Num());
+			check(Histories[ParamMapIdx].Variables.Num() == Histories[ParamMapIdx].PerVariableReadHistory.Num());
+		}
+		else
+		{
+			UE_LOG(LogNiagaraEditor, Log, TEXT("Could not resolve variable: %s"), *Name.ToString());
+		}
+		
+	}
+	else
+	{
+		// Do nothing here
+	}
+
+	return FoundIdx;
+}
+
+
 void FNiagaraParameterMapHistoryBuilder::BuildCurrentAliases()
 {
 	AliasMap = TMap<FString, FString>();
diff --git a/Engine/Plugins/FX/Niagara/Source/NiagaraEditor/Private/NiagaraScriptFactory.cpp b/Engine/Plugins/FX/Niagara/Source/NiagaraEditor/Private/NiagaraScriptFactory.cpp
--- a/Engine/Plugins/FX/Niagara/Source/NiagaraEditor/Private/NiagaraScriptFactory.cpp
+++ b/Engine/Plugins/FX/Niagara/Source/NiagaraEditor/Private/NiagaraScriptFactory.cpp
@@ -211,9 +211,8 @@ void UNiagaraScriptFactoryNew::InitializeScript(UNiagaraScript* NewScript)
 			// Set pointer in script to source
 			NewScript->SetSource(Source);
 
-			FString OutGraphLevelErrorMessages;
-			bool bForce = true;
-			NewScript->Compile(OutGraphLevelErrorMessages, bForce);
+
+			NewScript->RequestCompile();
 		}
 	}
 }
diff --git a/Engine/Plugins/FX/Niagara/Source/NiagaraEditor/Private/NiagaraScriptMergeManager.cpp b/Engine/Plugins/FX/Niagara/Source/NiagaraEditor/Private/NiagaraScriptMergeManager.cpp
--- a/Engine/Plugins/FX/Niagara/Source/NiagaraEditor/Private/NiagaraScriptMergeManager.cpp
+++ b/Engine/Plugins/FX/Niagara/Source/NiagaraEditor/Private/NiagaraScriptMergeManager.cpp
@@ -11,12 +11,15 @@
 #include "NiagaraNodeInput.h"
 #include "NiagaraNodeOutput.h"
 #include "NiagaraNodeFunctionCall.h"
+#include "NiagaraNodeCustomHlsl.h"
 #include "NiagaraNodeAssignment.h"
 #include "NiagaraNodeParameterMapGet.h"
 #include "NiagaraNodeParameterMapSet.h"
 #include "NiagaraEditorUtilities.h"
 #include "ViewModels/Stack/NiagaraStackGraphUtilities.h"
 #include "NiagaraEditorModule.h"
+#include "NiagaraEmitterEditorData.h"
+
 #include "UObject/PropertyPortFlags.h"
 
 #include "EdGraph/EdGraphNode.h"
@@ -44,6 +47,8 @@ FNiagaraStackFunctionInputOverrideMergeAdapter::FNiagaraStackFunctionInputOverri
 	
 	InputName = FNiagaraParameterHandle(OverridePin->PinName).GetName().ToString();
 	OverrideNode = CastChecked<UNiagaraNodeParameterMapSet>(OverridePin->GetOwningNode());
+	const UEdGraphSchema_Niagara* NiagaraSchema = GetDefault<UEdGraphSchema_Niagara>();
+	Type = NiagaraSchema->PinToTypeDefinition(OverridePin);
 
 	if (OverridePin->LinkedTo.Num() == 0)
 	{
@@ -51,6 +56,8 @@ FNiagaraStackFunctionInputOverrideMergeAdapter::FNiagaraStackFunctionInputOverri
 	}
 	else if (OverridePin->LinkedTo.Num() == 1)
 	{
+		OverrideValueNodePersistentId = OverridePin->LinkedTo[0]->GetOwningNode()->NodeGuid;
+
 		if (OverridePin->LinkedTo[0]->GetOwningNode()->IsA<UNiagaraNodeParameterMapGet>())
 		{
 			LinkedValueHandle = FNiagaraParameterHandle(OverridePin->LinkedTo[0]->PinName);
@@ -86,6 +93,7 @@ FNiagaraStackFunctionInputOverrideMergeAdapter::FNiagaraStackFunctionInputOverri
 	, OwningScript(&InOwningScript)
 	, OwningFunctionCallNode(&InOwningFunctionCallNode)
 	, InputName(InInputName)
+	, Type(InRapidIterationParameter.GetType())
 	, OverridePin(nullptr)
 	, LocalValueRapidIterationParameter(InRapidIterationParameter)
 	, DataValueObject(nullptr)
@@ -112,11 +120,21 @@ UNiagaraNodeParameterMapSet* FNiagaraStackFunctionInputOverrideMergeAdapter::Get
 	return OverrideNode.Get();
 }
 
+const FNiagaraTypeDefinition& FNiagaraStackFunctionInputOverrideMergeAdapter::GetType() const
+{
+	return Type;
+}
+
 UEdGraphPin* FNiagaraStackFunctionInputOverrideMergeAdapter::GetOverridePin() const
 {
 	return OverridePin;
 }
 
+const FGuid& FNiagaraStackFunctionInputOverrideMergeAdapter::GetOverrideNodeId() const
+{
+	return OverrideValueNodePersistentId;
+}
+
 TOptional<FString> FNiagaraStackFunctionInputOverrideMergeAdapter::GetLocalValueString() const
 {
 	return LocalValueString;
@@ -213,12 +231,18 @@ TSharedPtr<FNiagaraStackFunctionInputOverrideMergeAdapter> FNiagaraStackFunction
 FNiagaraScriptStackMergeAdapter::FNiagaraScriptStackMergeAdapter(UNiagaraNodeOutput& InOutputNode, UNiagaraScript& InScript, FString InUniqueEmitterName)
 {
 	OutputNode = &InOutputNode;
+	InputNode.Reset();
 	Script = &InScript;
 	UniqueEmitterName = InUniqueEmitterName;
 
 	TArray<FNiagaraStackGraphUtilities::FStackNodeGroup> StackGroups;
 	FNiagaraStackGraphUtilities::GetStackNodeGroups(*OutputNode, StackGroups);
 
+	if (StackGroups.Num() >= 2 && StackGroups[0].EndNode->IsA<UNiagaraNodeInput>())
+	{
+		InputNode = Cast<UNiagaraNodeInput>(StackGroups[0].EndNode);
+	}
+
 	if (StackGroups.Num() > 2 && StackGroups[0].EndNode->IsA<UNiagaraNodeInput>() && StackGroups.Last().EndNode->IsA<UNiagaraNodeOutput>())
 	{
 		for (int i = 1; i < StackGroups.Num() - 1; i++)
@@ -234,6 +258,11 @@ FNiagaraScriptStackMergeAdapter::FNiagaraScriptStackMergeAdapter(UNiagaraNodeOut
 	}
 }
 
+UNiagaraNodeInput* FNiagaraScriptStackMergeAdapter::GetInputNode() const
+{
+	return InputNode.Get();
+}
+
 UNiagaraNodeOutput* FNiagaraScriptStackMergeAdapter::GetOutputNode() const
 {
 	return OutputNode.Get();
@@ -293,6 +322,7 @@ void FNiagaraEventHandlerMergeAdapter::Initialize(const UNiagaraEmitter& InEmitt
 	if (EventScriptProperties != nullptr && OutputNode != nullptr)
 	{
 		EventStack = MakeShared<FNiagaraScriptStackMergeAdapter>(*OutputNode.Get(), *EventScriptProperties->Script, Emitter->GetUniqueEmitterName());
+		InputNode = EventStack->GetInputNode();
 	}
 }
 
@@ -328,6 +358,11 @@ UNiagaraNodeOutput* FNiagaraEventHandlerMergeAdapter::GetOutputNode() const
 	return OutputNode.Get();
 }
 
+UNiagaraNodeInput* FNiagaraEventHandlerMergeAdapter::GetInputNode() const
+{
+	return InputNode.Get();
+}
+
 TSharedPtr<FNiagaraScriptStackMergeAdapter> FNiagaraEventHandlerMergeAdapter::GetEventStack() const
 {
 	return EventStack;
@@ -525,7 +560,9 @@ bool FNiagaraScriptStackDiffResults::IsEmpty() const
 		EnabledChangedOtherModules.Num() == 0 &&
 		RemovedBaseInputOverrides.Num() == 0 &&
 		AddedOtherInputOverrides.Num() == 0 &&
-		ModifiedOtherInputOverrides.Num() == 0;
+		ModifiedOtherInputOverrides.Num() == 0 &&
+		ChangedBaseUsage.IsSet() == false &&
+		ChangedOtherUsage.IsSet() == false;
 }
 
 bool FNiagaraScriptStackDiffResults::IsValid() const
@@ -601,16 +638,173 @@ TSharedRef<FNiagaraScriptMergeManager> FNiagaraScriptMergeManager::Get()
 	return NiagaraEditorModule.GetScriptMergeManager();
 }
 
+void FNiagaraScriptMergeManager::DiffChangeIds(const TMap<FGuid, FGuid>& InSourceChangeIds, const TMap<FGuid, FGuid>& InLastMergedChangeIds, const TMap<FGuid, FGuid>& InInstanceChangeIds, TMap<FGuid, FGuid>& OutChangeIdsToKeepOnInstance)
+{
+	auto It = InInstanceChangeIds.CreateConstIterator();
+	while (It)
+	{
+		const FGuid* MatchingSourceChangeId = InSourceChangeIds.Find(It.Key());
+		const FGuid* MatchingLastMergedChangeId = InLastMergedChangeIds.Find(It.Key());
+
+		// If we don't have source from the original or from the last merged version of the original, then the instance originated the node and needs to keep it.
+		if (MatchingSourceChangeId == nullptr && MatchingLastMergedChangeId == nullptr)
+		{
+			OutChangeIdsToKeepOnInstance.Add(It.Key(), It.Value());
+		}
+		else if (MatchingSourceChangeId != nullptr && MatchingLastMergedChangeId != nullptr)
+		{
+			// If both had a copy of the node and both agree on the change id, then we should keep the change id of the instance as it will be the most accurate.
+			if (MatchingSourceChangeId->IsValid() && MatchingLastMergedChangeId->IsValid() && (*MatchingSourceChangeId) == (*MatchingLastMergedChangeId))
+			{
+				OutChangeIdsToKeepOnInstance.Add(It.Key(), It.Value());
+			}
+		}
+		else if (MatchingLastMergedChangeId != nullptr)
+		{
+			// If only the previous version had the matching key, then we may possibly keep this node around as a local node, in which case, we should apply the override. 
+			OutChangeIdsToKeepOnInstance.Add(It.Key(), It.Value());
+		}
+		else if (MatchingSourceChangeId != nullptr)
+		{
+			// I'm not sure that there's a way for us to reach this situation, where the node exists on the source and the instance, but not the 
+			// last merged version.
+			check(false);
+		}
+
+		++It;
+	}
+}
+
+FNiagaraScriptMergeManager::FApplyDiffResults FNiagaraScriptMergeManager::ResolveChangeIds(TSharedRef<FNiagaraEmitterMergeAdapter> MergedInstanceAdapter, UNiagaraEmitter& OriginalEmitterInstance, const TMap<FGuid, FGuid>& ChangeIdsThatNeedToBeReset)
+{
+	FApplyDiffResults DiffResults;
+
+	if (ChangeIdsThatNeedToBeReset.Num() != 0)
+	{
+		auto It = ChangeIdsThatNeedToBeReset.CreateConstIterator();
+		UNiagaraEmitter* Emitter = MergedInstanceAdapter->GetEditableEmitter();
+
+		TArray<UNiagaraGraph*> Graphs;
+		TArray<UNiagaraScript*> Scripts;
+		Emitter->GetScripts(Scripts);
+
+		TArray<UNiagaraScript*> OriginalScripts;
+		OriginalEmitterInstance.GetScripts(OriginalScripts);
+
+		// First gather all the graphs used by this emitter..
+		for (UNiagaraScript* Script : Scripts)
+		{
+			if (Script != nullptr && Script->GetSource() != nullptr)
+			{
+				UNiagaraScriptSource* ScriptSource = Cast<UNiagaraScriptSource>(Script->GetSource());
+				if (ScriptSource != nullptr)
+				{
+					Graphs.AddUnique(ScriptSource->NodeGraph);
+				}
+			}
+		}
+
+		// Now gather up all the nodes
+		TArray<UNiagaraNode*> Nodes;
+		for (UNiagaraGraph* Graph : Graphs)
+		{
+			Graph->GetNodesOfClass(Nodes);
+		}
+
+		// Now go through all the nodes and set the persistent change ids if we encounter a node that needs it's change id kept.
+		bool bAnySet = false;
+		while (It)
+		{			
+			bool bFound = false;
+			for (UNiagaraNode* Node : Nodes)
+			{
+				if (Node->NodeGuid == It.Key())
+				{
+					Node->ForceChangeId(It.Value(), false);
+					bAnySet = true;
+					bFound = true;
+					break;
+				}
+			}
+
+			if (!bFound)
+			{
+				UE_LOG(LogNiagaraEditor, Log, TEXT("Looks like %s was removed.. is that expected?"), *It.Key().ToString())
+			}
+			++It;
+		}
+
+		if (bAnySet)
+		{
+			for (UNiagaraGraph* Graph : Graphs)
+			{
+				Graph->MarkGraphRequiresSynchronization(TEXT("Overwrote change id's within graph."));
+			}
+		}
+
+		DiffResults.bModifiedGraph = bAnySet;
+
+		UNiagaraScriptSource* ScriptSource = Cast<UNiagaraScriptSource>(OriginalEmitterInstance.GraphSource);
+		if (ScriptSource && bAnySet)
+		{
+			UNiagaraGraph* OtherGraph = ScriptSource->NodeGraph;
+			if (OtherGraph)
+			{
+				for (UNiagaraGraph* Graph : Graphs)
+				{
+					Graph->SynchronizeInternalCacheWithGraph(OtherGraph);
+				}
+			}
+		}
+
+		if (bAnySet)
+		{
+			bool bAnyUpdated = false;
+			TMap<FString, FString> RenameMap;
+			RenameMap.Add(TEXT("Emitter"), TEXT("Emitter"));
+			for (UNiagaraScript* Script : Scripts)
+			{
+				for (UNiagaraScript* OriginalScript : OriginalScripts)
+				{
+					if (Script->Usage == OriginalScript->Usage && Script->GetUsageId() == OriginalScript->GetUsageId())
+					{
+						bAnyUpdated |= Script->SynchronizeExecutablesWithMaster(OriginalScript, RenameMap);
+					}
+				}
+			}
+
+			if (bAnyUpdated)
+			{
+				//Emitter->OnPostCompile();
+			}
+		}
+	}
+
+	DiffResults.bSucceeded = true;
+	return DiffResults;
+}
+
 INiagaraModule::FMergeEmitterResults FNiagaraScriptMergeManager::MergeEmitter(UNiagaraEmitter& Source, UNiagaraEmitter& LastMergedSource, UNiagaraEmitter& Instance)
 {
 	SCOPE_CYCLE_COUNTER(STAT_NiagaraEditor_ScriptMergeManager_MergeEmitter);
 	INiagaraModule::FMergeEmitterResults MergeResults;
+
 	FNiagaraEmitterDiffResults DiffResults = DiffEmitters(LastMergedSource, Instance);
+
 	if (DiffResults.IsValid())
 	{
 		UNiagaraEmitter* MergedInstance = CastChecked<UNiagaraEmitter>(StaticDuplicateObject(&Source, (UObject*)GetTransientPackage()));
 		TSharedRef<FNiagaraEmitterMergeAdapter> MergedInstanceAdapter = MakeShared<FNiagaraEmitterMergeAdapter>(*MergedInstance);
 
+		TMap<FGuid, FGuid> SourceChangeIds;
+		TMap<FGuid, FGuid> PreviousSourceChangeIds;
+		TMap<FGuid, FGuid> LastChangeIds;
+		TMap<FGuid, FGuid> ChangeIdsThatNeedToBeReset;
+		FNiagaraEditorUtilities::GatherChangeIds(Source, SourceChangeIds, TEXT("Source"));
+		FNiagaraEditorUtilities::GatherChangeIds(LastMergedSource, PreviousSourceChangeIds, TEXT("MergeLast"));
+		FNiagaraEditorUtilities::GatherChangeIds(Instance, LastChangeIds, TEXT("Instance"));
+		DiffChangeIds(SourceChangeIds, PreviousSourceChangeIds, LastChangeIds, ChangeIdsThatNeedToBeReset);
+
 		FApplyDiffResults EmitterSpawnResults = ApplyScriptStackDiff(MergedInstanceAdapter->GetEmitterSpawnStack().ToSharedRef(), DiffResults.EmitterSpawnDiffResults);
 		MergeResults.bSucceeded &= EmitterSpawnResults.bSucceeded;
 		MergeResults.bModifiedGraph |= EmitterSpawnResults.bModifiedGraph;
@@ -642,15 +836,81 @@ INiagaraModule::FMergeEmitterResults FNiagaraScriptMergeManager::MergeEmitter(UN
 		MergeResults.ErrorMessages.Append(RendererResults.ErrorMessages);
 
 		CopyPropertiesToBase(MergedInstance, &Instance, DiffResults.DifferentEmitterProperties);
+		if (Instance.EditorData != nullptr)
+		{
+			// We keep the instance editor data here so that the UI state in the system stays more consistent.
+			MergedInstance->EditorData = Cast<UNiagaraEmitterEditorData>(StaticDuplicateObject(Instance.EditorData, MergedInstance));
+		}
+
+#if 0
+		UE_LOG(LogNiagaraEditor, Log, TEXT("A"));
+		//for (FNiagaraEmitterHandle& EmitterHandle : EmitterHandles)
+		{
+			//UE_LOG(LogNiagara, Log, TEXT("Emitter Handle: %s"), *EmitterHandle.GetUniqueInstanceName());
+			UNiagaraScript* UpdateScript = MergedInstance->GetScript(ENiagaraScriptUsage::ParticleUpdateScript, FGuid());
+			UNiagaraScript* SpawnScript = MergedInstance->GetScript(ENiagaraScriptUsage::ParticleSpawnScript, FGuid());
+			UE_LOG(LogNiagaraEditor, Log, TEXT("Spawn Parameters"));
+			SpawnScript->GetVMExecutableData().Parameters.DumpParameters();
+			UE_LOG(LogNiagaraEditor, Log, TEXT("Spawn RI Parameters"));
+			SpawnScript->RapidIterationParameters.DumpParameters();
+			UE_LOG(LogNiagaraEditor, Log, TEXT("Update Parameters"));
+			UpdateScript->GetVMExecutableData().Parameters.DumpParameters();
+			UE_LOG(LogNiagaraEditor, Log, TEXT("Update RI Parameters"));
+			UpdateScript->RapidIterationParameters.DumpParameters();
+		}
+#endif
+		
+		FApplyDiffResults ChangeIdResults = ResolveChangeIds(MergedInstanceAdapter, Instance, ChangeIdsThatNeedToBeReset);
+		MergeResults.bSucceeded &= ChangeIdResults.bSucceeded;
+		MergeResults.bModifiedGraph |= ChangeIdResults.bModifiedGraph;
+		MergeResults.ErrorMessages.Append(ChangeIdResults.ErrorMessages);
+		
+#if 0
+		UE_LOG(LogNiagaraEditor, Log, TEXT("B"));
+		//for (FNiagaraEmitterHandle& EmitterHandle : EmitterHandles)
+		{
+			//UE_LOG(LogNiagara, Log, TEXT("Emitter Handle: %s"), *EmitterHandle.GetUniqueInstanceName());
+			UNiagaraScript* UpdateScript = MergedInstance->GetScript(ENiagaraScriptUsage::ParticleUpdateScript, FGuid());
+			UNiagaraScript* SpawnScript = MergedInstance->GetScript(ENiagaraScriptUsage::ParticleSpawnScript, FGuid());
+			UE_LOG(LogNiagaraEditor, Log, TEXT("Spawn Parameters"));
+			SpawnScript->GetVMExecutableData().Parameters.DumpParameters();
+			UE_LOG(LogNiagaraEditor, Log, TEXT("Spawn RI Parameters"));
+			SpawnScript->RapidIterationParameters.DumpParameters();
+			UE_LOG(LogNiagaraEditor, Log, TEXT("Update Parameters"));
+			UpdateScript->GetVMExecutableData().Parameters.DumpParameters();
+			UE_LOG(LogNiagaraEditor, Log, TEXT("Update RI Parameters"));
+			UpdateScript->RapidIterationParameters.DumpParameters();
+		}
+#endif
 
 		FNiagaraStackGraphUtilities::CleanUpStaleRapidIterationParameters(*MergedInstance);
 
+#if 0
+		UE_LOG(LogNiagaraEditor, Log, TEXT("C"));
+		//for (FNiagaraEmitterHandle& EmitterHandle : EmitterHandles)
+		{
+			//UE_LOG(LogNiagara, Log, TEXT("Emitter Handle: %s"), *EmitterHandle.GetUniqueInstanceName());
+			UNiagaraScript* UpdateScript = MergedInstance->GetScript(ENiagaraScriptUsage::ParticleUpdateScript, FGuid());
+			UNiagaraScript* SpawnScript = MergedInstance->GetScript(ENiagaraScriptUsage::ParticleSpawnScript, FGuid());
+			UE_LOG(LogNiagaraEditor, Log, TEXT("Spawn Parameters"));
+			SpawnScript->GetVMExecutableData().Parameters.DumpParameters();
+			UE_LOG(LogNiagaraEditor, Log, TEXT("Spawn RI Parameters"));
+			SpawnScript->RapidIterationParameters.DumpParameters();
+			UE_LOG(LogNiagaraEditor, Log, TEXT("Update Parameters"));
+			UpdateScript->GetVMExecutableData().Parameters.DumpParameters();
+			UE_LOG(LogNiagaraEditor, Log, TEXT("Update RI Parameters"));
+			UpdateScript->RapidIterationParameters.DumpParameters();
+		}
+#endif
 		if (MergeResults.bSucceeded)
 		{
 			UNiagaraScriptSource* ScriptSource = Cast<UNiagaraScriptSource>(MergedInstance->GraphSource);
 			FNiagaraStackGraphUtilities::RelayoutGraph(*ScriptSource->NodeGraph);
 			MergeResults.MergedInstance = MergedInstance;
 		}
+
+		TMap<FGuid, FGuid> FinalChangeIds;
+		FNiagaraEditorUtilities::GatherChangeIds(*MergedInstance, FinalChangeIds, TEXT("Final"));
 	}
 	else
 	{
@@ -1117,6 +1377,12 @@ void FNiagaraScriptMergeManager::DiffScriptStacks(TSharedRef<FNiagaraScriptStack
 
 		DiffFunctionInputs(CommonValuePair.BaseValue, CommonValuePair.OtherValue, DiffResults);
 	}
+
+	if (BaseScriptStackAdapter->GetScript()->GetUsage() != OtherScriptStackAdapter->GetScript()->GetUsage())
+	{
+		DiffResults.ChangedBaseUsage = BaseScriptStackAdapter->GetScript()->GetUsage();
+		DiffResults.ChangedOtherUsage = OtherScriptStackAdapter->GetScript()->GetUsage();
+	}
 }
 
 void FNiagaraScriptMergeManager::DiffFunctionInputs(TSharedRef<FNiagaraStackFunctionMergeAdapter> BaseFunctionAdapter, TSharedRef<FNiagaraStackFunctionMergeAdapter> OtherFunctionAdapter, FNiagaraScriptStackDiffResults& DiffResults)
@@ -1153,7 +1419,7 @@ void FNiagaraScriptMergeManager::DiffEditableProperties(const void* BaseDataAddr
 {
 	for (TFieldIterator<UProperty> PropertyIterator(&Struct); PropertyIterator; ++PropertyIterator)
 	{
-		if (PropertyIterator->HasAllPropertyFlags(CPF_Edit))
+		if (PropertyIterator->HasAllPropertyFlags(CPF_Edit) && PropertyIterator->HasMetaData("NiagaraNoMerge") == false)
 		{
 			if (PropertyIterator->Identical(
 				PropertyIterator->ContainerPtrToValuePtr<void>(BaseDataAddress),
@@ -1231,7 +1497,17 @@ TOptional<bool> FNiagaraScriptMergeManager::DoFunctionInputOverridesMatch(TShare
 
 	if (BaseFunctionInputAdapter->GetDynamicValueFunction().IsValid() && OtherFunctionInputAdapter->GetDynamicValueFunction().IsValid())
 	{
-		if (BaseFunctionInputAdapter->GetDynamicValueFunction()->GetFunctionCallNode()->FunctionScript != OtherFunctionInputAdapter->GetDynamicValueFunction()->GetFunctionCallNode()->FunctionScript)
+		UNiagaraNodeCustomHlsl* BaseCustomHlsl = Cast<UNiagaraNodeCustomHlsl>(BaseFunctionInputAdapter->GetDynamicValueFunction()->GetFunctionCallNode());
+		UNiagaraNodeCustomHlsl* OtherCustomHlsl = Cast<UNiagaraNodeCustomHlsl>(OtherFunctionInputAdapter->GetDynamicValueFunction()->GetFunctionCallNode());
+
+		if (BaseCustomHlsl && OtherCustomHlsl)
+		{
+			if (BaseCustomHlsl->CustomHlsl != OtherCustomHlsl->CustomHlsl || BaseCustomHlsl->ScriptUsage != OtherCustomHlsl->ScriptUsage)
+			{
+				return false;
+			}
+		}
+		else if (BaseFunctionInputAdapter->GetDynamicValueFunction()->GetFunctionCallNode()->FunctionScript != OtherFunctionInputAdapter->GetDynamicValueFunction()->GetFunctionCallNode()->FunctionScript)
 		{
 			return false;
 		}
@@ -1254,7 +1530,12 @@ FNiagaraScriptMergeManager::FApplyDiffResults FNiagaraScriptMergeManager::AddMod
 	UNiagaraNodeFunctionCall* AddedModuleNode = nullptr;
 	if (AddModule->GetFunctionCallNode()->IsA<UNiagaraNodeAssignment>())
 	{
-		AddedModuleNode = FNiagaraStackGraphUtilities::AddParameterModuleToStack(CastChecked<UNiagaraNodeAssignment>(AddModule->GetFunctionCallNode())->AssignmentTarget, TargetOutputNode, AddModule->GetStackIndex());
+		UNiagaraNodeAssignment* AssignmentNode = CastChecked<UNiagaraNodeAssignment>(AddModule->GetFunctionCallNode());
+		const TArray<FNiagaraVariable>& Targets = AssignmentNode->GetAssignmentTargets();
+		const TArray<FString>& Defaults = AssignmentNode->GetAssignmentDefaults();
+		AddedModuleNode = FNiagaraStackGraphUtilities::AddParameterModuleToStack(Targets, TargetOutputNode, AddModule->GetStackIndex(),Defaults);
+		AddedModuleNode->NodeGuid = AddModule->GetFunctionCallNode()->NodeGuid;
+		AddedModuleNode->RefreshFromExternalChanges();
 		Results.bModifiedGraph = true;
 	}
 	else
@@ -1262,6 +1543,7 @@ FNiagaraScriptMergeManager::FApplyDiffResults FNiagaraScriptMergeManager::AddMod
 		if (AddModule->GetFunctionCallNode()->FunctionScript != nullptr)
 		{
 			AddedModuleNode = FNiagaraStackGraphUtilities::AddScriptModuleToStack(AddModule->GetFunctionCallNode()->FunctionScript, TargetOutputNode, AddModule->GetStackIndex());
+			AddedModuleNode->NodeGuid = AddModule->GetFunctionCallNode()->NodeGuid; // Synchronize the node Guid across runs so that the compile id's synch up.
 			Results.bModifiedGraph = true;
 		}
 		else
@@ -1275,6 +1557,9 @@ FNiagaraScriptMergeManager::FApplyDiffResults FNiagaraScriptMergeManager::AddMod
 
 	if (AddedModuleNode != nullptr)
 	{
+		AddedModuleNode->NodeGuid = AddModule->GetFunctionCallNode()->NodeGuid; // Synchronize the node Guid across runs so that the compile id's synch up.
+
+		AddedModuleNode->SetEnabledState(AddModule->GetFunctionCallNode()->GetDesiredEnabledState(), AddModule->GetFunctionCallNode()->HasUserSetTheEnabledState());
 		for (TSharedRef<FNiagaraStackFunctionInputOverrideMergeAdapter> InputOverride : AddModule->GetInputOverrides())
 		{
 			FApplyDiffResults AddInputResults = AddInputOverride(UniqueEmitterName, OwningScript, *AddedModuleNode, InputOverride);
@@ -1322,6 +1607,31 @@ FNiagaraScriptMergeManager::FApplyDiffResults FNiagaraScriptMergeManager::AddInp
 {
 	FApplyDiffResults Results;
 
+	// If an assignment node, make sure that we have an assignment target for the input override.
+	UNiagaraNodeAssignment* AssignmentNode = Cast<UNiagaraNodeAssignment>(&TargetFunctionCall);
+	if (AssignmentNode)
+	{
+		FNiagaraParameterHandle FunctionInputHandle(FNiagaraParameterHandle::ModuleNamespace, *OverrideToAdd->GetInputName());
+		UNiagaraNodeAssignment* PreviousVersionAssignmentNode = Cast<UNiagaraNodeAssignment>(OverrideToAdd->GetOwningFunctionCall());
+		bool bAnyAdded = false;
+		for (int32 i = 0; i < PreviousVersionAssignmentNode->NumTargets(); i++)
+		{
+			const FNiagaraVariable& Var = PreviousVersionAssignmentNode->GetAssignmentTarget(i);
+
+			int32 FoundVarIdx = AssignmentNode->FindAssignmentTarget(Var.GetName());
+			if (FoundVarIdx == INDEX_NONE)
+			{
+				AssignmentNode->AddAssignmentTarget(Var, &PreviousVersionAssignmentNode->GetAssignmentDefaults()[i]);
+				bAnyAdded = true;
+			}
+		}
+
+		if (bAnyAdded)
+		{
+			AssignmentNode->RefreshFromExternalChanges();
+		}
+	}
+
 	if (OverrideToAdd->GetOverridePin() != nullptr)
 	{
 		FNiagaraParameterHandle FunctionInputHandle(FNiagaraParameterHandle::ModuleNamespace, *OverrideToAdd->GetInputName());
@@ -1330,51 +1640,77 @@ FNiagaraScriptMergeManager::FApplyDiffResults FNiagaraScriptMergeManager::AddInp
 		const UEdGraphSchema_Niagara* NiagaraSchema = GetDefault<UEdGraphSchema_Niagara>();
 		FNiagaraTypeDefinition InputType = NiagaraSchema->PinToTypeDefinition(OverrideToAdd->GetOverridePin());
 
-		UEdGraphPin& InputOverridePin = FNiagaraStackGraphUtilities::GetOrCreateStackFunctionInputOverridePin(TargetFunctionCall, AliasedFunctionInputHandle, InputType);
-
-		if (OverrideToAdd->GetLocalValueString().IsSet())
+		UEdGraphPin& InputOverridePin = FNiagaraStackGraphUtilities::GetOrCreateStackFunctionInputOverridePin(TargetFunctionCall, AliasedFunctionInputHandle, InputType, OverrideToAdd->GetOverrideNode()->NodeGuid);
+		if (InputOverridePin.LinkedTo.Num() != 0)
 		{
-			InputOverridePin.DefaultValue = OverrideToAdd->GetLocalValueString().GetValue();
-			Results.bSucceeded = true;
-		}
-		else if (OverrideToAdd->GetLinkedValueHandle().IsSet())
-		{
-			FNiagaraStackGraphUtilities::SetLinkedValueHandleForFunctionInput(InputOverridePin, OverrideToAdd->GetLinkedValueHandle().GetValue());
-			Results.bSucceeded = true;
-		}
-		else if (OverrideToAdd->GetDataValueObject() != nullptr)
-		{
-			UNiagaraDataInterface* OverrideValueObject = OverrideToAdd->GetDataValueObject();
-			UNiagaraDataInterface* NewOverrideValueObject;
-			FNiagaraStackGraphUtilities::SetDataValueObjectForFunctionInput(InputOverridePin, OverrideToAdd->GetDataValueObject()->GetClass(), OverrideValueObject->GetName(), NewOverrideValueObject);
-			OverrideValueObject->CopyTo(NewOverrideValueObject);
-			Results.bSucceeded = true;
+			Results.bSucceeded = false;
+			Results.ErrorMessages.Add(LOCTEXT("AddPinBasedInputOverrideFailedOverridePinStillLinked", "Failed to add input override because the target override pin was still linked to other nodes."));
 		}
-		else if (OverrideToAdd->GetDynamicValueFunction().IsValid())
+		else
 		{
-			if (OverrideToAdd->GetDynamicValueFunction()->GetFunctionCallNode()->FunctionScript != nullptr)
+			if (OverrideToAdd->GetLocalValueString().IsSet())
+			{
+				InputOverridePin.DefaultValue = OverrideToAdd->GetLocalValueString().GetValue();
+				Results.bSucceeded = true;
+			}
+			else if (OverrideToAdd->GetLinkedValueHandle().IsSet())
+			{
+				FNiagaraStackGraphUtilities::SetLinkedValueHandleForFunctionInput(InputOverridePin, OverrideToAdd->GetLinkedValueHandle().GetValue(), OverrideToAdd->GetOverrideNodeId());
+				Results.bSucceeded = true;
+			}
+			else if (OverrideToAdd->GetDataValueObject() != nullptr)
 			{
-				UNiagaraNodeFunctionCall* DynamicInputFunctionCall;
-				FNiagaraStackGraphUtilities::SetDynamicInputForFunctionInput(InputOverridePin, OverrideToAdd->GetDynamicValueFunction()->GetFunctionCallNode()->FunctionScript, DynamicInputFunctionCall);
-				for (TSharedRef<FNiagaraStackFunctionInputOverrideMergeAdapter> DynamicInputInputOverride : OverrideToAdd->GetDynamicValueFunction()->GetInputOverrides())
+				UNiagaraDataInterface* OverrideValueObject = OverrideToAdd->GetDataValueObject();
+				UNiagaraDataInterface* NewOverrideValueObject;
+				FNiagaraStackGraphUtilities::SetDataValueObjectForFunctionInput(InputOverridePin, OverrideToAdd->GetDataValueObject()->GetClass(), OverrideValueObject->GetName(), NewOverrideValueObject, OverrideToAdd->GetOverrideNodeId());
+				OverrideValueObject->CopyTo(NewOverrideValueObject);
+				Results.bSucceeded = true;
+			}
+			else if (OverrideToAdd->GetDynamicValueFunction().IsValid())
+			{
+				UNiagaraNodeCustomHlsl* CustomHlslFunction = Cast<UNiagaraNodeCustomHlsl>(OverrideToAdd->GetDynamicValueFunction()->GetFunctionCallNode());
+				if (CustomHlslFunction != nullptr)
+				{
+					UNiagaraNodeCustomHlsl* DynamicInputFunctionCall;
+					FNiagaraStackGraphUtilities::SetCustomExpressionForFunctionInput(InputOverridePin, DynamicInputFunctionCall, OverrideToAdd->GetOverrideNodeId());
+					DynamicInputFunctionCall->CustomHlsl = CustomHlslFunction->CustomHlsl;
+					for (TSharedRef<FNiagaraStackFunctionInputOverrideMergeAdapter> DynamicInputInputOverride : OverrideToAdd->GetDynamicValueFunction()->GetInputOverrides())
+					{
+						FApplyDiffResults AddResults = AddInputOverride(UniqueEmitterName, OwningScript, *((UNiagaraNodeFunctionCall*)DynamicInputFunctionCall), DynamicInputInputOverride);
+						Results.bSucceeded &= AddResults.bSucceeded;
+						Results.bModifiedGraph |= AddResults.bModifiedGraph;
+						Results.ErrorMessages.Append(AddResults.ErrorMessages);
+					}
+				}
+				else if (OverrideToAdd->GetDynamicValueFunction()->GetFunctionCallNode()->FunctionScript != nullptr)
 				{
-					FApplyDiffResults AddResults = AddInputOverride(UniqueEmitterName, OwningScript, *DynamicInputFunctionCall, DynamicInputInputOverride);
-					Results.bSucceeded &= AddResults.bSucceeded;
-					Results.bModifiedGraph |= AddResults.bModifiedGraph;
-					Results.ErrorMessages.Append(AddResults.ErrorMessages);
+					UNiagaraNodeFunctionCall* DynamicInputFunctionCall;
+					FNiagaraStackGraphUtilities::SetDynamicInputForFunctionInput(InputOverridePin, OverrideToAdd->GetDynamicValueFunction()->GetFunctionCallNode()->FunctionScript,
+						DynamicInputFunctionCall, OverrideToAdd->GetOverrideNodeId(), OverrideToAdd->GetDynamicValueFunction()->GetFunctionCallNode()->GetFunctionName());
+					if (DynamicInputFunctionCall->GetFunctionName() != OverrideToAdd->GetDynamicValueFunction()->GetFunctionCallNode()->GetFunctionName())
+					{
+						DynamicInputFunctionCall->NodeGuid = FGuid::NewGuid();
+					}
+					for (TSharedRef<FNiagaraStackFunctionInputOverrideMergeAdapter> DynamicInputInputOverride : OverrideToAdd->GetDynamicValueFunction()->GetInputOverrides())
+					{
+						FApplyDiffResults AddResults = AddInputOverride(UniqueEmitterName, OwningScript, *DynamicInputFunctionCall, DynamicInputInputOverride);
+						Results.bSucceeded &= AddResults.bSucceeded;
+						Results.bModifiedGraph |= AddResults.bModifiedGraph;
+						Results.ErrorMessages.Append(AddResults.ErrorMessages);
+					}
+				}
+				else
+				{
+					Results.bSucceeded = false;
+					Results.ErrorMessages.Add(LOCTEXT("AddPinBasedInputOverrideFailedInvalidDynamicInput", "Failed to add input override because it's dynamic function call's function script was null."));
 				}
 			}
 			else
 			{
 				Results.bSucceeded = false;
-				Results.ErrorMessages.Add(LOCTEXT("AddPinBasedInputOverrideFailedInvalidDynamicInput", "Failed to add input override because it's dynamic function call's function script was null."));
+				Results.ErrorMessages.Add(LOCTEXT("AddPinBasedInputOverrideFailed", "Failed to add input override because it was invalid."));
 			}
 		}
-		else
-		{
-			Results.bSucceeded = false;
-			Results.ErrorMessages.Add(LOCTEXT("AddPinBasedInputOverrideFailed", "Failed to add input override because it was invalid."));
-		}
 		Results.bModifiedGraph = true;
 	}
 	else
@@ -1508,6 +1844,13 @@ FNiagaraScriptMergeManager::FApplyDiffResults FNiagaraScriptMergeManager::ApplyS
 		}
 	}
 
+	// Update the usage if different
+	if (DiffResults.ChangedOtherUsage.IsSet())
+	{
+		BaseScriptStackAdapter->GetScript()->SetUsage(DiffResults.ChangedOtherUsage.GetValue());
+		BaseScriptStackAdapter->GetOutputNode()->SetUsage(DiffResults.ChangedOtherUsage.GetValue());
+	}
+
 	// Apply the graph actions.
 	for (TSharedRef<FNiagaraStackFunctionMergeAdapter> RemoveModule : RemoveModules)
 	{
@@ -1624,11 +1967,13 @@ FNiagaraScriptMergeManager::FApplyDiffResults FNiagaraScriptMergeManager::ApplyE
 			FNiagaraEventScriptProperties AddedEventScriptProperties = *AddedEventHandler->GetEventScriptProperties();
 			AddedEventScriptProperties.Script = NewObject<UNiagaraScript>(BaseEmitter, MakeUniqueObjectName(BaseEmitter, UNiagaraScript::StaticClass(), "EventScript"), EObjectFlags::RF_Transactional);
 			AddedEventScriptProperties.Script->SetUsage(ENiagaraScriptUsage::ParticleEventScript);
-			AddedEventScriptProperties.Script->SetUsageId(FGuid::NewGuid());
+			AddedEventScriptProperties.Script->SetUsageId(AddedEventHandler->GetUsageId());
 			AddedEventScriptProperties.Script->SetSource(EmitterSource);
 			BaseEmitter->AddEventHandler(AddedEventScriptProperties);
 
-			UNiagaraNodeOutput* EventOutputNode = FNiagaraStackGraphUtilities::ResetGraphForOutput(*EmitterGraph, ENiagaraScriptUsage::ParticleEventScript, AddedEventScriptProperties.Script->GetUsageId());
+			FGuid PreferredOutputNodeGuid = AddedEventHandler->GetOutputNode()->NodeGuid;
+			FGuid PreferredInputNodeGuid = AddedEventHandler->GetInputNode()->NodeGuid;
+			UNiagaraNodeOutput* EventOutputNode = FNiagaraStackGraphUtilities::ResetGraphForOutput(*EmitterGraph, ENiagaraScriptUsage::ParticleEventScript, AddedEventScriptProperties.Script->GetUsageId(), PreferredOutputNodeGuid, PreferredInputNodeGuid);
 			for (TSharedRef<FNiagaraStackFunctionMergeAdapter> ModuleAdapter : AddedEventHandler->GetEventStack()->GetModuleFunctions())
 			{
 				FApplyDiffResults AddModuleResults = AddModule(BaseEmitter->GetUniqueEmitterName(), *AddedEventScriptProperties.Script, *EventOutputNode, ModuleAdapter);
diff --git a/Engine/Plugins/FX/Niagara/Source/NiagaraEditor/Private/NiagaraScriptMergeManager.h b/Engine/Plugins/FX/Niagara/Source/NiagaraEditor/Private/NiagaraScriptMergeManager.h
--- a/Engine/Plugins/FX/Niagara/Source/NiagaraEditor/Private/NiagaraScriptMergeManager.h
+++ b/Engine/Plugins/FX/Niagara/Source/NiagaraEditor/Private/NiagaraScriptMergeManager.h
@@ -15,6 +15,7 @@
 class UNiagaraEmitter;
 class UNiagaraScript;
 class UNiagaraNodeOutput;
+class UNiagaraNodeInput;
 class UNiagaraNodeFunctionCall;
 class UNiagaraNodeParameterMapSet;
 class UEdGraphNode;
@@ -46,6 +47,8 @@ public:
 	FString GetInputName() const;
 	UNiagaraNodeParameterMapSet* GetOverrideNode() const;
 	UEdGraphPin* GetOverridePin() const;
+	const FGuid& GetOverrideNodeId() const;
+	const FNiagaraTypeDefinition& GetType() const;
 
 	TOptional<FString> GetLocalValueString() const;
 	TOptional<FNiagaraVariable> GetLocalValueRapidIterationParameter() const;
@@ -58,6 +61,7 @@ private:
 	TWeakObjectPtr<UNiagaraScript> OwningScript;
 	TWeakObjectPtr<UNiagaraNodeFunctionCall> OwningFunctionCallNode;
 	FString InputName;
+	FNiagaraTypeDefinition Type;
 
 	TWeakObjectPtr<UNiagaraNodeParameterMapSet> OverrideNode;
 	UEdGraphPin* OverridePin;
@@ -67,6 +71,8 @@ private:
 	TOptional<FNiagaraParameterHandle> LinkedValueHandle;
 	UNiagaraDataInterface* DataValueObject;
 	TSharedPtr<FNiagaraStackFunctionMergeAdapter> DynamicValueFunction;
+
+	FGuid OverrideValueNodePersistentId;
 };
 
 class FNiagaraStackFunctionMergeAdapter
@@ -97,6 +103,7 @@ public:
 	FNiagaraScriptStackMergeAdapter(UNiagaraNodeOutput& InOutputNode, UNiagaraScript& InScript, FString InUniqueEmitterName);
 
 	UNiagaraNodeOutput* GetOutputNode() const;
+	UNiagaraNodeInput* GetInputNode() const;
 
 	UNiagaraScript* GetScript() const;
 
@@ -107,6 +114,7 @@ public:
 	TSharedPtr<FNiagaraStackFunctionMergeAdapter> GetModuleFunctionById(FGuid FunctionCallNodeId) const;
 
 private:
+	TWeakObjectPtr<UNiagaraNodeInput> InputNode;
 	TWeakObjectPtr<UNiagaraNodeOutput> OutputNode;
 	TWeakObjectPtr<UNiagaraScript> Script;
 	FString UniqueEmitterName;
@@ -125,6 +133,7 @@ public:
 	const FNiagaraEventScriptProperties* GetEventScriptProperties() const;
 	FNiagaraEventScriptProperties* GetEditableEventScriptProperties() const;
 	UNiagaraNodeOutput* GetOutputNode() const;
+	UNiagaraNodeInput* GetInputNode() const;
 	TSharedPtr<FNiagaraScriptStackMergeAdapter> GetEventStack() const;
 
 private:
@@ -135,6 +144,7 @@ private:
 	const FNiagaraEventScriptProperties* EventScriptProperties;
 	FNiagaraEventScriptProperties* EditableEventScriptProperties;
 	TWeakObjectPtr<UNiagaraNodeOutput> OutputNode;
+	TWeakObjectPtr<UNiagaraNodeInput> InputNode;
 
 	TSharedPtr<FNiagaraScriptStackMergeAdapter> EventStack;
 };
@@ -211,6 +221,9 @@ struct FNiagaraScriptStackDiffResults
 	TArray<TSharedRef<FNiagaraStackFunctionInputOverrideMergeAdapter>> ModifiedBaseInputOverrides;
 	TArray<TSharedRef<FNiagaraStackFunctionInputOverrideMergeAdapter>> ModifiedOtherInputOverrides;
 
+	TOptional<ENiagaraScriptUsage> ChangedBaseUsage;
+	TOptional<ENiagaraScriptUsage> ChangedOtherUsage;
+
 private:
 	bool bIsValid;
 	TArray<FText> ErrorMessages;
@@ -338,6 +351,10 @@ private:
 
 	TSharedRef<FNiagaraEmitterMergeAdapter> GetEmitterMergeAdapterUsingCache(UNiagaraEmitter& Emitter);
 
+	void DiffChangeIds(const TMap<FGuid, FGuid>& InSourceChangeIds, const TMap<FGuid, FGuid>& InLastMergedChangeIds, const TMap<FGuid, FGuid>& InInstanceChangeIds, TMap<FGuid, FGuid>& OutChangeIdsToKeepOnInstance);
+	FApplyDiffResults ResolveChangeIds(TSharedRef<FNiagaraEmitterMergeAdapter> MergedInstanceAdapter, UNiagaraEmitter& OriginalEmitterInstance, const TMap<FGuid, FGuid>& ChangeIdsThatNeedToBeReset);
+
+
 private:
 	struct FCachedMergeAdapter
 	{
