Copying //UE4/Dev-Niagara to //UE4/Dev-Main (Source: //UE4/Dev-Niagara @ 4074996)

diff --git a/Engine/Plugins/FX/Niagara/Source/NiagaraEditor/Private/NiagaraScriptSource.cpp b/Engine/Plugins/FX/Niagara/Source/NiagaraEditor/Private/NiagaraScriptSource.cpp
--- a/Engine/Plugins/FX/Niagara/Source/NiagaraEditor/Private/NiagaraScriptSource.cpp
+++ b/Engine/Plugins/FX/Niagara/Source/NiagaraEditor/Private/NiagaraScriptSource.cpp
@@ -20,6 +20,7 @@
 #include "NiagaraDataInterface.h"
 #include "GraphEditAction.h"
 #include "NiagaraNodeParameterMapBase.h"
+#include "NiagaraNodeInput.h"
 #include "NiagaraNodeOutput.h"
 #include "EdGraphUtilities.h"
 #include "EdGraphSchema_Niagara.h"
@@ -30,15 +31,36 @@
 #include "NiagaraNodeFunctionCall.h"
 #include "NiagaraNodeParameterMapSet.h"
 #include "INiagaraEditorTypeUtilities.h"
+#include "NiagaraEditorUtilities.h"
+#include "NiagaraNodeEmitter.h"
+#include "Logging/LogMacros.h"
 
 DECLARE_CYCLE_STAT(TEXT("Niagara - ScriptSource - Compile"), STAT_NiagaraEditor_ScriptSource_Compile, STATGROUP_NiagaraEditor);
-DECLARE_CYCLE_STAT(TEXT("Niagara - ScriptSource - PreCompile"), STAT_NiagaraEditor_ScriptSource_PreCompile, STATGROUP_NiagaraEditor);
+DECLARE_CYCLE_STAT(TEXT("Niagara - ScriptSource - InitializeNewRapidIterationParameters"), STAT_NiagaraEditor_ScriptSource_InitializeNewRapidIterationParameters, STATGROUP_NiagaraEditor);
 
 UNiagaraScriptSource::UNiagaraScriptSource(const FObjectInitializer& ObjectInitializer)
-	: Super(ObjectInitializer), bIsPrecompiled(false)
+	: Super(ObjectInitializer)
+{
+}
+
+void UNiagaraScriptSource::ComputeVMCompilationId(FNiagaraVMExecutableDataId& Id, ENiagaraScriptUsage InUsage, const FGuid& InUsageId) const
+{
+	Id.ScriptUsageType = InUsage;
+	Id.ScriptUsageTypeID = InUsageId;
+	Id.CompilerVersionID = FNiagaraCustomVersion::LatestScriptCompileVersion;
+	if (NodeGraph)
+	{
+		Id.BaseScriptID = NodeGraph->GetCompileID(InUsage, InUsageId);
+		NodeGraph->GatherExternalDependencyIDs(InUsage, InUsageId, Id.ReferencedDependencyIds, Id.ReferencedObjects);
+	}
+}
+
+void UNiagaraScriptSource::InvalidateCachedCompileIds() 
 {
+	NodeGraph->InvalidateCachedCompileIds();
 }
 
+
 void UNiagaraScriptSource::PostLoad()
 {
 	Super::PostLoad();
@@ -50,208 +72,9 @@ void UNiagaraScriptSource::PostLoad()
 
 		// Hook up event handlers so the on changed handler can be called correctly.
 		NodeGraph->AddOnGraphChangedHandler(FOnGraphChanged::FDelegate::CreateUObject(this, &UNiagaraScriptSource::OnGraphChanged));
+		NodeGraph->AddOnGraphNeedsRecompileHandler(FOnGraphChanged::FDelegate::CreateUObject(this, &UNiagaraScriptSource::OnGraphChanged));
 		NodeGraph->OnDataInterfaceChanged().AddUObject(this, &UNiagaraScriptSource::OnGraphDataInterfaceChanged);
 	}
-
-	const int32 NiagaraVer = GetLinkerCustomVersion(FNiagaraCustomVersion::GUID);
-	if (NiagaraVer < FNiagaraCustomVersion::AddedScriptRapidIterationVariables)
-	{
-		UpgradeForScriptRapidIterationVariables();
-	}
-}
-
-void UNiagaraScriptSource::UpgradeForScriptRapidIterationVariables()
-{
-	// In this version we added the ability to rapidly iterate on top-level module parameters. This means that there are 
-	// loose set nodes that were just setting constants that need to be removed from the graph as well as needing to build the script
-	// rapid iteration parameters list.
-	TArray<UNiagaraNodeOutput*> OutputNodes;
-	NodeGraph->GetNodesOfClass(OutputNodes);
-	UNiagaraEmitter* ParentEmitter = Cast<UNiagaraEmitter>(GetOuter());
-	UNiagaraSystem* ParentSystem = Cast<UNiagaraSystem>(GetOuter());
-	FNiagaraEditorModule& NiagaraEditorModule = FModuleManager::GetModuleChecked<FNiagaraEditorModule>("NiagaraEditor");
-
-	if (ParentEmitter != nullptr || ParentSystem != nullptr)
-	{
-		TArray<UNiagaraScript*> Scripts;
-		if (ParentEmitter != nullptr)
-		{
-			ParentEmitter->ConditionalPostLoad();
-			ParentEmitter->GetScripts(Scripts, false);
-		}
-		else if (ParentSystem != nullptr)
-		{
-			ParentSystem->ConditionalPostLoad();
-			Scripts.Add(ParentSystem->GetSystemSpawnScript(true));
-			Scripts.Add(ParentSystem->GetSystemUpdateScript(true));
-			Scripts.Add(ParentSystem->GetSystemSpawnScript(false));
-			Scripts.Add(ParentSystem->GetSystemUpdateScript(false));
-		}
-
-		for (UNiagaraScript* Script : Scripts)
-		{
-			ensure(Script->RapidIterationParameters.GetParameterDataArray().Num() == 0);
-		}
-
-		bool bNeedsRecompile = false;
-
-		TArray<UEdGraphPin*> PinsToRemove;
-
-		for (UNiagaraNodeOutput* Output : OutputNodes)
-		{
-			FNiagaraParameterMapHistoryBuilder Builder;
-			Builder.BuildParameterMaps(Output);
-				
-			TArray<UNiagaraNodeFunctionCall*> FunctionCallNodes;
-			TArray<UNiagaraNode*> Nodes;
-			NodeGraph->BuildTraversal(Nodes, Output);
-			for (UNiagaraNode* Node : Nodes)
-			{
-				UNiagaraNodeFunctionCall* FuncCallNode = Cast<UNiagaraNodeFunctionCall>(Node);
-				if (FuncCallNode != nullptr)
-				{
-					FunctionCallNodes.Add(FuncCallNode);
-				}
-			}
-
-			TArray<const UEdGraphPin*> CandidateInputPins;
-			TArray<FNiagaraVariable> CandidateVariables;
-			for (UNiagaraNodeFunctionCall* FunctionCallNode : FunctionCallNodes)
-			{
-				TArray<const UEdGraphPin*> InputPins;
-				FNiagaraStackGraphUtilities::GetStackFunctionInputPins(*FunctionCallNode, InputPins, FNiagaraStackGraphUtilities::ENiagaraGetStackFunctionInputPinsOptions::ModuleInputsOnly);
-
-				for (const UEdGraphPin* InputPin : InputPins)
-				{
-					CandidateInputPins.Add(InputPin);
-					const UEdGraphSchema_Niagara* NiagaraSchema = CastChecked<UEdGraphSchema_Niagara>(InputPin->GetSchema());
-					FNiagaraVariable FunctionVar = NiagaraSchema->PinToNiagaraVariable(InputPin);
-					TMap<FString, FString> AliasMap = TMap<FString, FString>();
-					AliasMap.Add(TEXT("Module"), FunctionCallNode->GetFunctionName());
-					if (ParentEmitter)
-					{
-						AliasMap.Add(TEXT("Emitter"), ParentEmitter->GetUniqueEmitterName());
-					}
-					FunctionVar = FNiagaraParameterMapHistory::ResolveAliases(FunctionVar, AliasMap, TEXT("."));
-					CandidateVariables.Add(FunctionVar);
-				}
-			}
-
-			UE_LOG(LogNiagaraEditor, Log, TEXT("Converting OutputNode: %s"), *Output->GetNodeTitle(ENodeTitleType::FullTitle).ToString());
-
-			// As of this version there aren't multiple instances of Event graphs, so we ignore the occurrence parameter.
-			// We do however have the solo and multi versions of the system scripts, so we need to be able to set multiple scripts.
-			TArray<UNiagaraScript*> ActiveScripts;
-			for (UNiagaraScript* Script : Scripts)
-			{
-				if (Script->ContainsUsage(Output->GetUsage()))
-				{
-					ActiveScripts.Add(Script);
-				}
-			}
-
-			if (ActiveScripts.Num() == 0)
-			{
-				continue;
-			}
-
-			for (int32 i = 0; i < CandidateInputPins.Num(); i++)
-			{
-				const UEdGraphPin* Pin = CandidateInputPins[i];
-				FNiagaraVariable CandidateVar = CandidateVariables[i];
-
-				check(Builder.Histories.Num() == 1);
-				const UEdGraphSchema_Niagara* NiagaraSchema = CastChecked<UEdGraphSchema_Niagara>(Pin->GetSchema());
-				FNiagaraVariable FunctionVar = NiagaraSchema->PinToNiagaraVariable(Pin);
-				FNiagaraTypeDefinition PinDefinition = FunctionVar.GetType();
-
-				// Now check to see if this is one of the appropriate types:
-				if (PinDefinition == FNiagaraTypeDefinition::GetBoolDef() || PinDefinition.IsEnum() || PinDefinition == FNiagaraTypeDefinition::GetParameterMapDef())
-				{
-					continue;
-				}
-				
-				// Is this pin even in the subgraph we care about right now?
-				int32 ParamIdx = Builder.Histories[0].FindVariable(CandidateVar.GetName(), CandidateVar.GetType());
-				if (ParamIdx == INDEX_NONE)
-				{
-					continue;
-				}
-
-				FString DefaultValue;
-				// Now check to see if anyone writes to this variable before reading in the graph.
-				if (Builder.Histories[0].PerVariableReadHistory[ParamIdx].Num() != 0)
-				{
-					const UEdGraphPin* FirstReadPin = Builder.Histories[0].PerVariableReadHistory[ParamIdx][0].Key;
-					const UEdGraphPin* FirstWritePin = Builder.Histories[0].PerVariableReadHistory[ParamIdx][0].Value;
-					const UEdGraphPin* DefaultValuePin = Builder.Histories[0].GetDefaultValuePin(ParamIdx);
-
-					UNiagaraNodeParameterMapSet* MapSetNode = nullptr;
-					if (FirstWritePin != nullptr)
-					{
-						MapSetNode = Cast<UNiagaraNodeParameterMapSet>(FirstWritePin->GetOwningNode());
-					}
-
-					if (FirstWritePin != nullptr && FirstWritePin->LinkedTo.Num() == 0 && nullptr != MapSetNode) // Set by a ParameterMapSet pin linked to a constant...
-					{
-						DefaultValue = FirstWritePin->DefaultValue;
-						PinsToRemove.Add((UEdGraphPin*)FirstWritePin);
-						UE_LOG(LogNiagaraEditor, Log, TEXT("Removing Set pin: %s"), *Builder.Histories[0].Variables[ParamIdx].GetName().ToString());
-					}
-					else if (FirstWritePin != nullptr && FirstWritePin->LinkedTo.Num() > 0) // The Set node is set by something other than a constant, so skip it.
-					{
-						continue;
-					}
-					else if (FirstWritePin == nullptr && DefaultValuePin != nullptr) // Get the default value of a pin...
-					{
-						DefaultValue = DefaultValuePin->DefaultValue;
-					}
-
-					FNiagaraVariable FinalVar = Builder.Histories[0].Variables[ParamIdx];
-					FinalVar = FNiagaraParameterMapHistory::ConvertVariableToRapidIterationConstantName(FinalVar, ParentEmitter ? *ParentEmitter->GetUniqueEmitterName() : nullptr);
-
-					TSharedPtr<INiagaraEditorTypeUtilities> TypeEditorUtilities = NiagaraEditorModule.GetTypeUtilities(FinalVar.GetType());
-					if (TypeEditorUtilities.IsValid() && TypeEditorUtilities->CanHandlePinDefaults())
-					{
-						bool bValueSet = TypeEditorUtilities->SetValueFromPinDefaultString(DefaultValue, FinalVar);
-						if (bValueSet)
-						{
-							UE_LOG(LogNiagaraEditor, Log, TEXT("Adding Rapid Iteration Variable: %s  Default Value: %s"), *FinalVar.GetName().ToString(), *DefaultValue);
-						}
-						else
-						{
-							UE_LOG(LogNiagaraEditor, Warning, TEXT("Adding Rapid Iteration Variable: %s  Default Value '%s' could not be set for type %s"),
-								*FinalVar.GetName().ToString(), *DefaultValue, *FinalVar.GetType().GetName());
-						}
-					}
-
-					for (UNiagaraScript* ActiveScript : ActiveScripts)
-					{
-						ActiveScript->RapidIterationParameters.AddParameter(FinalVar);
-					}
-				}					
-			}
-		}
-
-		for (UEdGraphPin* PinToRemove : PinsToRemove)
-		{
-			if (PinToRemove->IsPendingKill())
-			{
-				continue;
-			}
-			UNiagaraNodeParameterMapSet* SetNode = Cast<UNiagaraNodeParameterMapSet>(PinToRemove->GetOwningNode());
-			if (SetNode)
-			{
-				SetNode->RemovePin(PinToRemove);
-				bNeedsRecompile = true;
-			}
-		}
-
-		if (bNeedsRecompile)
-		{
-			NodeGraph->NotifyGraphNeedsRecompile();
-		}
-	}
 }
 
 UNiagaraScriptSourceBase* UNiagaraScriptSource::MakeRecursiveDeepCopy(UObject* DestOuter, TMap<const UObject*, UObject*>& ExistingConversions) const
@@ -293,107 +116,16 @@ bool UNiagaraScriptSource::IsSynchronized(const FGuid& InChangeId)
 	}
 }
 
-void UNiagaraScriptSource::MarkNotSynchronized()
+void UNiagaraScriptSource::MarkNotSynchronized(FString Reason)
 {
 	if (NodeGraph)
 	{
-		NodeGraph->MarkGraphRequiresSynchronization();
-	}
-}
-
-bool UNiagaraScriptSource::IsPreCompiled() const
-{
-	return bIsPrecompiled;
-}
-
-void UNiagaraScriptSource::PreCompile(UNiagaraEmitter* InEmitter, bool bClearErrors)
-{
-	SCOPE_CYCLE_COUNTER(STAT_NiagaraEditor_ScriptSource_PreCompile);
-	if (!bIsPrecompiled)
-	{
-		bIsPrecompiled = true;
-
-		if (bClearErrors)
-		{
-			//Clear previous graph errors.
-			bool bHasClearedGraphErrors = false;
-			for (UEdGraphNode* Node : NodeGraph->Nodes)
-			{
-				if (Node->bHasCompilerMessage)
-				{
-					Node->ErrorMsg.Empty();
-					Node->ErrorType = EMessageSeverity::Info;
-					Node->bHasCompilerMessage = false;
-					Node->Modify(true);
-					bHasClearedGraphErrors = true;
-				}
-			}
-			if (bHasClearedGraphErrors)
-			{
-				NodeGraph->NotifyGraphChanged();
-			}
-		}
-
-		// Clone the source graph so we can modify it as needed; merging in the child graphs
-		NodeGraphDeepCopy = CastChecked<UNiagaraGraph>(FEdGraphUtilities::CloneGraph(NodeGraph, this, 0));
-		FEdGraphUtilities::MergeChildrenGraphsIn(NodeGraphDeepCopy, NodeGraphDeepCopy, /*bRequireSchemaMatch=*/ true);
-		
-		PrecompiledHistories.Empty();
-
-		TArray<UNiagaraNodeOutput*> OutputNodes;
-		NodeGraphDeepCopy->FindOutputNodes(OutputNodes);
-		PrecompiledHistories.Empty();
-		for (UNiagaraNodeOutput* FoundOutputNode : OutputNodes)
-		{
-			// Map all for this output node
-			FNiagaraParameterMapHistoryBuilder Builder;
-			Builder.BeginTranslation(InEmitter);
-			Builder.EnableScriptWhitelist(true, FoundOutputNode->GetUsage());
-			Builder.BuildParameterMaps(FoundOutputNode, true);
-			TArray<FNiagaraParameterMapHistory> Histories = Builder.Histories;
-			ensure(Histories.Num() <= 1);
-			PrecompiledHistories.Append(Histories);
-			Builder.EndTranslation(InEmitter);
-		}
-	}
-}
-
-
-bool UNiagaraScriptSource::GatherPreCompiledVariables(const FString& InNamespaceFilter, TArray<FNiagaraVariable>& OutVars)
-{
-	if (!bIsPrecompiled || PrecompiledHistories.Num() == 0)
-	{
-		return false;
-	}
-
-	for (const FNiagaraParameterMapHistory& History : PrecompiledHistories)
-	{
-		for (const FNiagaraVariable& Var : History.Variables)
-		{
-			if (FNiagaraParameterMapHistory::IsInNamespace(Var, InNamespaceFilter))
-			{
-				FNiagaraVariable NewVar = Var;
-				if (NewVar.IsDataAllocated() == false && !Var.IsDataInterface())
-				{
-					FNiagaraEditorUtilities::ResetVariableToDefaultValue(NewVar);
-				}
-				OutVars.AddUnique(NewVar);
-			}
-		}
+		NodeGraph->MarkGraphRequiresSynchronization(Reason);
 	}
-	return true;
-}
-
-void UNiagaraScriptSource::PostCompile() 
-{
-	bIsPrecompiled = false;
-	PrecompiledHistories.Empty();
-	NodeGraphDeepCopy = nullptr;
 }
 
-bool UNiagaraScriptSource::PostLoadFromEmitter(UNiagaraEmitter& OwningEmitter)
+void UNiagaraScriptSource::PostLoadFromEmitter(UNiagaraEmitter& OwningEmitter)
 {
-	bool bNeedsCompile = false;
 	const int32 NiagaraCustomVersion = GetLinkerCustomVersion(FNiagaraCustomVersion::GUID);
 	if (NiagaraCustomVersion < FNiagaraCustomVersion::ScriptsNowUseAGuidForIdentificationInsteadOfAnIndex)
 	{
@@ -415,9 +147,8 @@ bool UNiagaraScriptSource::PostLoadFromEmitter(UNiagaraEmitter& OwningEmitter)
 				MatchingOutputNode->SetUsageId(EventScriptProperties.Script->GetUsageId());
 			}
 		}
-		bNeedsCompile = true;
+		NodeGraph->MarkGraphRequiresSynchronization("Modified while handling a change to the niagara custom version.");
 	}
-	return bNeedsCompile;
 }
 
 bool UNiagaraScriptSource::AddModuleIfMissing(FString ModulePath, ENiagaraScriptUsage Usage, bool& bOutFoundModule)
@@ -444,13 +175,103 @@ bool UNiagaraScriptSource::AddModuleIfMissing(FString ModulePath, ENiagaraScript
 	return false;
 }
 
+void InitializeNewRapidIterationParametersForNode(const UEdGraphSchema_Niagara* Schema, UNiagaraNode* Node, const FString& UniqueEmitterName, ENiagaraScriptUsage ScriptUsage, FNiagaraParameterStore& RapidIterationParameters, TSet<FName>& ValidRapidIterationParameterNames)
+{
+	UNiagaraNodeFunctionCall* FunctionCallNode = Cast<UNiagaraNodeFunctionCall>(Node);
+	if (FunctionCallNode != nullptr)
+	{
+		TArray<const UEdGraphPin*> FunctionInputPins;
+		FNiagaraStackGraphUtilities::GetStackFunctionInputPins(*FunctionCallNode, FunctionInputPins, FNiagaraStackGraphUtilities::ENiagaraGetStackFunctionInputPinsOptions::ModuleInputsOnly, false);
+		for (const UEdGraphPin* FunctionInputPin : FunctionInputPins)
+		{
+			FNiagaraTypeDefinition InputType = Schema->PinToTypeDefinition(FunctionInputPin);
+			if (FNiagaraStackGraphUtilities::IsRapidIterationType(InputType))
+			{
+				FNiagaraParameterHandle AliasedFunctionInputHandle = FNiagaraParameterHandle::CreateAliasedModuleParameterHandle(FNiagaraParameterHandle(FunctionInputPin->PinName), FunctionCallNode);
+				FNiagaraVariable RapidIterationParameter = FNiagaraStackGraphUtilities::CreateRapidIterationParameter(UniqueEmitterName, ScriptUsage, AliasedFunctionInputHandle.GetParameterHandleString(), InputType);
+				ValidRapidIterationParameterNames.Add(RapidIterationParameter.GetName());
+				int32 ParameterIndex = RapidIterationParameters.IndexOf(RapidIterationParameter);
+				// Only set a value for the parameter if it's not already set.
+				if (ParameterIndex == INDEX_NONE)
+				{
+					UEdGraphPin* DefaultPin = FunctionCallNode->FindParameterMapDefaultValuePin(FunctionInputPin->PinName, ScriptUsage);
+					// Only set values for inputs which don't have a default wired in the script graph, since inputs with wired defaults can't currently use rapid iteration parameters.
+					if (DefaultPin != nullptr && DefaultPin->LinkedTo.Num() == 0)
+					{
+						// Only set values for inputs without override pins, since and override pin means it's being read from a different value.
+						UEdGraphPin* OverridePin = FNiagaraStackGraphUtilities::GetStackFunctionInputOverridePin(*FunctionCallNode, AliasedFunctionInputHandle);
+						if (OverridePin == nullptr)
+						{
+							FNiagaraVariable DefaultVariable = Schema->PinToNiagaraVariable(DefaultPin, true);
+							check(DefaultVariable.GetData() != nullptr);
+							RapidIterationParameters.AddParameter(RapidIterationParameter, false);
+							RapidIterationParameters.SetParameterData(DefaultVariable.GetData(), RapidIterationParameter);
+						}
+					}
+				}
+			}
+		}
+	}
+}
+
+void UNiagaraScriptSource::CleanUpOldAndInitializeNewRapidIterationParameters(FString UniqueEmitterName, ENiagaraScriptUsage ScriptUsage, FGuid ScriptUsageId, FNiagaraParameterStore& RapidIterationParameters) const
+{
+	SCOPE_CYCLE_COUNTER(STAT_NiagaraEditor_ScriptSource_InitializeNewRapidIterationParameters);
+	TArray<UNiagaraNodeOutput*> OutputNodes;
+	if (ScriptUsage == ENiagaraScriptUsage::ParticleGPUComputeScript)
+	{
+		TArray<UNiagaraNodeOutput*> TempOutputNodes;
+		NodeGraph->FindOutputNodes(TempOutputNodes);
+		for (UNiagaraNodeOutput* OutputNode : TempOutputNodes)
+		{
+			if (UNiagaraScript::IsParticleScript(OutputNode->GetUsage()))
+			{
+				OutputNodes.AddUnique(OutputNode);
+			}
+		}
+	}
+	else
+	{
+		UNiagaraNodeOutput* OutputNode = NodeGraph->FindEquivalentOutputNode(ScriptUsage, ScriptUsageId);
+		OutputNodes.Add(OutputNode);
+	}
+
+	TSet<FName> ValidRapidIterationParameterNames;
+	for (UNiagaraNodeOutput* OutputNode : OutputNodes)
+	{
+		if (OutputNode != nullptr)
+		{
+			TArray<UNiagaraNode*> Nodes;
+			NodeGraph->BuildTraversal(Nodes, OutputNode);
+			const UEdGraphSchema_Niagara* Schema = GetDefault<UEdGraphSchema_Niagara>();
+
+			for(UNiagaraNode* Node : Nodes)
+			{	
+				InitializeNewRapidIterationParametersForNode(Schema, Node, UniqueEmitterName, OutputNode->GetUsage(), RapidIterationParameters, ValidRapidIterationParameterNames);
+			}
+		}
+	}
+
+	// Clean up old rapid iteration parameters.
+	TArray<FNiagaraVariable> CurrentRapidIterationVariables;
+	RapidIterationParameters.GetParameters(CurrentRapidIterationVariables);
+	for (const FNiagaraVariable& CurrentRapidIterationVariable : CurrentRapidIterationVariables)
+	{
+		if (ValidRapidIterationParameterNames.Contains(CurrentRapidIterationVariable.GetName()) == false)
+		{
+			RapidIterationParameters.RemoveParameter(CurrentRapidIterationVariable);
+		}
+	}
+}
+
+/*
 ENiagaraScriptCompileStatus UNiagaraScriptSource::Compile(UNiagaraScript* ScriptOwner, FString& OutGraphLevelErrorMessages)
 {
 	SCOPE_CYCLE_COUNTER(STAT_NiagaraEditor_ScriptSource_Compile);
 	bool bDoPostCompile = false;
 	if (!bIsPrecompiled)
 	{
-		PreCompile(nullptr);
+		PreCompile(nullptr, TArray<FNiagaraVariable>());
 		bDoPostCompile = true;
 	}
 
@@ -487,6 +308,7 @@ ENiagaraScriptCompileStatus UNiagaraScriptSource::Compile(UNiagaraScript* Script
 // 	}
 
 }
+*/
 
 void UNiagaraScriptSource::OnGraphChanged(const FEdGraphEditAction &Action)
 {
diff --git a/Engine/Plugins/FX/Niagara/Source/NiagaraEditor/Private/NiagaraStackEditorData.cpp b/Engine/Plugins/FX/Niagara/Source/NiagaraEditor/Private/NiagaraStackEditorData.cpp
--- a/Engine/Plugins/FX/Niagara/Source/NiagaraEditor/Private/NiagaraStackEditorData.cpp
+++ b/Engine/Plugins/FX/Niagara/Source/NiagaraEditor/Private/NiagaraStackEditorData.cpp
@@ -2,17 +2,6 @@
 
 #include "NiagaraStackEditorData.h"
 
-bool UNiagaraStackEditorData::GetModuleInputIsPinned(const FString& ModuleInputKey) const
-{
-	const bool* bIsPinnedPtr = ModuleInputKeyToPinnedMap.Find(ModuleInputKey);
-	return bIsPinnedPtr != nullptr && *bIsPinnedPtr;
-}
-
-void UNiagaraStackEditorData::SetModuleInputIsPinned(const FString& ModuleInputKey, bool bIsPinned)
-{
-	ModuleInputKeyToPinnedMap.FindOrAdd(ModuleInputKey) = bIsPinned;
-}
-
 bool UNiagaraStackEditorData::GetModuleInputIsRenamePending(const FString& ModuleInputKey) const
 {
 	const bool* bIsRenamePendingPtr = ModuleInputKeyToRenamePendingMap.Find(ModuleInputKey);
@@ -32,5 +21,67 @@ bool UNiagaraStackEditorData::GetStackEntryIsExpanded(const FString& StackEntryK
 
 void UNiagaraStackEditorData::SetStackEntryIsExpanded(const FString& StackEntryKey, bool bIsExpanded)
 {
-	StackEntryKeyToExpandedMap.FindOrAdd(StackEntryKey) = bIsExpanded;
-}
\ No newline at end of file
+	if (ensureMsgf(StackEntryKey.IsEmpty() == false, TEXT("Can not set the expanded state with an empty key")))
+	{
+		StackEntryKeyToExpandedMap.FindOrAdd(StackEntryKey) = bIsExpanded;
+	}
+}
+
+bool UNiagaraStackEditorData::GetStackItemShowAdvanced(const FString& StackEntryKey, bool bShowAdvancedDefault) const
+{
+	const bool* bShowAdvancedPtr = StackItemKeyToShowAdvancedMap.Find(StackEntryKey);
+	return bShowAdvancedPtr != nullptr ? *bShowAdvancedPtr : bShowAdvancedDefault;
+}
+
+void UNiagaraStackEditorData::SetStackItemShowAdvanced(const FString& StackEntryKey, bool bShowAdvanced)
+{
+	if (ensureMsgf(StackEntryKey.IsEmpty() == false, TEXT("Can not set the show advanced state with an empty key")))
+	{
+		StackItemKeyToShowAdvancedMap.FindOrAdd(StackEntryKey) = bShowAdvanced;
+	}
+}
+
+bool UNiagaraStackEditorData::GetShowAllAdvanced() const
+{
+	return bShowAllAdvanced;
+}
+
+void UNiagaraStackEditorData::SetShowAllAdvanced(bool bInShowAllAdvanced)
+{
+	bShowAllAdvanced = bInShowAllAdvanced;
+}
+
+bool UNiagaraStackEditorData::GetShowOutputs() const
+{
+	return bShowOutputs;
+}
+
+void UNiagaraStackEditorData::SetShowOutputs(bool bInShowOutputs)
+{
+	bShowOutputs = bInShowOutputs;
+}
+
+double UNiagaraStackEditorData::GetLastScrollPosition() const
+{
+	return LastScrollPosition;
+}
+
+void UNiagaraStackEditorData::SetLastScrollPosition(double InLastScrollPosition)
+{
+	LastScrollPosition = InLastScrollPosition;
+}
+
+void UNiagaraStackEditorData::DismissStackIssue(FName IssueId)
+{
+	DismissedStackIssueIds.AddUnique(IssueId);
+}
+
+void UNiagaraStackEditorData::UndismissAllIssues()
+{
+	DismissedStackIssueIds.Empty();
+}
+
+const TArray<FName>& UNiagaraStackEditorData::GetDismissedStackIssueIds()
+{
+	return DismissedStackIssueIds;
+}
diff --git a/Engine/Plugins/FX/Niagara/Source/NiagaraEditor/Private/NiagaraSystemEditorData.cpp b/Engine/Plugins/FX/Niagara/Source/NiagaraEditor/Private/NiagaraSystemEditorData.cpp
--- a/Engine/Plugins/FX/Niagara/Source/NiagaraEditor/Private/NiagaraSystemEditorData.cpp
+++ b/Engine/Plugins/FX/Niagara/Source/NiagaraEditor/Private/NiagaraSystemEditorData.cpp
@@ -2,6 +2,11 @@
 
 #include "NiagaraSystemEditorData.h"
 #include "NiagaraStackEditorData.h"
+#include "NiagaraCustomVersion.h"
+#include "NiagaraSystem.h"
+#include "NiagaraEmitterHandle.h"
+#include "NiagaraEmitter.h"
+#include "NiagaraEmitterEditorData.h"
 
 const FName UNiagaraSystemEditorFolder::GetFolderName() const
 {
@@ -52,11 +57,14 @@ UNiagaraSystemEditorData::UNiagaraSystemEditorData(const FObjectInitializer& Obj
 	RootFolder = ObjectInitializer.CreateDefaultSubobject<UNiagaraSystemEditorFolder>(this, TEXT("RootFolder"));
 	StackEditorData = ObjectInitializer.CreateDefaultSubobject<UNiagaraStackEditorData>(this, TEXT("StackEditorData"));
 	OwnerTransform.SetLocation(FVector(0.0f, 0.0f, 100.0f));
+	PlaybackRangeMin = 0;
+	PlaybackRangeMax = 10;
 }
 
-void UNiagaraSystemEditorData::PostLoad()
+void UNiagaraSystemEditorData::PostLoadFromOwner(UObject* InOwner)
 {
-	Super::PostLoad();
+	UNiagaraSystem* OwnerSystem = CastChecked<UNiagaraSystem>(InOwner);
+
 	if (RootFolder == nullptr)
 	{
 		RootFolder = NewObject<UNiagaraSystemEditorFolder>(this, TEXT("RootFolder"), RF_Transactional);
@@ -65,6 +73,13 @@ void UNiagaraSystemEditorData::PostLoad()
 	{
 		StackEditorData = NewObject<UNiagaraStackEditorData>(this, TEXT("StackEditorData"), RF_Transactional);
 	}
+
+	const int32 NiagaraVer = GetLinkerCustomVersion(FNiagaraCustomVersion::GUID);
+
+	if (NiagaraVer < FNiagaraCustomVersion::PlaybackRangeStoredOnSystem)
+	{
+		UpdatePlaybackRangeFromEmitters(OwnerSystem);
+	}
 }
 
 UNiagaraSystemEditorFolder& UNiagaraSystemEditorData::GetRootFolder() const
@@ -75,4 +90,37 @@ UNiagaraSystemEditorFolder& UNiagaraSystemEditorData::GetRootFolder() const
 UNiagaraStackEditorData& UNiagaraSystemEditorData::GetStackEditorData() const
 {
 	return *StackEditorData;
+}
+
+TRange<float> UNiagaraSystemEditorData::GetPlaybackRange() const
+{
+	return TRange<float>(PlaybackRangeMin, PlaybackRangeMax);
+}
+
+void UNiagaraSystemEditorData::SetPlaybackRange(TRange<float> InPlaybackRange)
+{
+	PlaybackRangeMin = InPlaybackRange.GetLowerBoundValue();
+	PlaybackRangeMax = InPlaybackRange.GetUpperBoundValue();
+}
+
+void UNiagaraSystemEditorData::UpdatePlaybackRangeFromEmitters(UNiagaraSystem* OwnerSystem)
+{
+	if (OwnerSystem->GetEmitterHandles().Num() > 0)
+	{
+		float EmitterPlaybackRangeMin = TNumericLimits<float>::Max();
+		float EmitterPlaybackRangeMax = TNumericLimits<float>::Lowest();
+
+		for (const FNiagaraEmitterHandle& EmitterHandle : OwnerSystem->GetEmitterHandles())
+		{
+			UNiagaraEmitterEditorData* EmitterEditorData = Cast<UNiagaraEmitterEditorData>(EmitterHandle.GetInstance()->EditorData);
+			if (EmitterEditorData != nullptr)
+			{
+				EmitterPlaybackRangeMin = FMath::Min(PlaybackRangeMin, EmitterEditorData->GetPlaybackRange().GetLowerBoundValue());
+				EmitterPlaybackRangeMax = FMath::Max(PlaybackRangeMax, EmitterEditorData->GetPlaybackRange().GetUpperBoundValue());
+			}
+		}
+
+		PlaybackRangeMin = EmitterPlaybackRangeMin;
+		PlaybackRangeMax = EmitterPlaybackRangeMax;
+	}
 }
\ No newline at end of file
diff --git a/Engine/Plugins/FX/Niagara/Source/NiagaraEditor/Private/NiagaraSystemEditorData.h b/Engine/Plugins/FX/Niagara/Source/NiagaraEditor/Private/NiagaraSystemEditorData.h
--- a/Engine/Plugins/FX/Niagara/Source/NiagaraEditor/Private/NiagaraSystemEditorData.h
+++ b/Engine/Plugins/FX/Niagara/Source/NiagaraEditor/Private/NiagaraSystemEditorData.h
@@ -2,9 +2,11 @@
 
 #pragma once
 
+#include "NiagaraEditorDataBase.h"
 #include "NiagaraSystemEditorData.generated.h"
 
 class UNiagaraStackEditorData;
+class UNiagaraSystem;
 
 /** Editor only folder data for emitters in a system. */
 UCLASS()
@@ -42,14 +44,14 @@ private:
 
 /** Editor only UI data for systems. */
 UCLASS()
-class UNiagaraSystemEditorData : public UObject
+class UNiagaraSystemEditorData : public UNiagaraEditorDataBase
 {
 	GENERATED_BODY()
 
 public:
 	UNiagaraSystemEditorData(const FObjectInitializer& ObjectInitializer);
 
-	virtual void PostLoad() override;
+	virtual void PostLoadFromOwner(UObject* InOwner) override;
 
 	/** Gets the root folder for UI folders for emitters. */
 	UNiagaraSystemEditorFolder& GetRootFolder() const;
@@ -65,6 +67,13 @@ public:
 		OwnerTransform = InTransform;
 	}
 
+	TRange<float> GetPlaybackRange() const;
+
+	void SetPlaybackRange(TRange<float> InPlaybackRange);
+
+private:
+	void UpdatePlaybackRangeFromEmitters(UNiagaraSystem* OwnerSystem);
+
 private:
 	UPROPERTY(Instanced)
 	UNiagaraSystemEditorFolder* RootFolder;
@@ -74,4 +83,10 @@ private:
 
 	UPROPERTY()
 	FTransform OwnerTransform;
+
+	UPROPERTY()
+	float PlaybackRangeMin;
+
+	UPROPERTY()
+	float PlaybackRangeMax;
 };
\ No newline at end of file
diff --git a/Engine/Plugins/FX/Niagara/Source/NiagaraEditor/Private/NiagaraSystemFactoryNew.cpp b/Engine/Plugins/FX/Niagara/Source/NiagaraEditor/Private/NiagaraSystemFactoryNew.cpp
--- a/Engine/Plugins/FX/Niagara/Source/NiagaraEditor/Private/NiagaraSystemFactoryNew.cpp
+++ b/Engine/Plugins/FX/Niagara/Source/NiagaraEditor/Private/NiagaraSystemFactoryNew.cpp
@@ -50,8 +50,6 @@ void UNiagaraSystemFactoryNew::InitializeSystem(UNiagaraSystem* System, bool bCr
 {
 	UNiagaraScript* SystemSpawnScript = System->GetSystemSpawnScript();
 	UNiagaraScript* SystemUpdateScript = System->GetSystemUpdateScript();
-	UNiagaraScript* SystemSpawnScriptSolo = System->GetSystemSpawnScript(true);
-	UNiagaraScript* SystemUpdateScriptSolo = System->GetSystemUpdateScript(true);
 
 	UNiagaraScriptSource* SystemScriptSource = NewObject<UNiagaraScriptSource>(SystemSpawnScript, "SystemScriptSource", RF_Transactional);
 
@@ -62,8 +60,6 @@ void UNiagaraSystemFactoryNew::InitializeSystem(UNiagaraSystem* System, bool bCr
 
 	SystemSpawnScript->SetSource(SystemScriptSource);
 	SystemUpdateScript->SetSource(SystemScriptSource);
-	SystemSpawnScriptSolo->SetSource(SystemScriptSource);
-	SystemUpdateScriptSolo->SetSource(SystemScriptSource);
 
 	if (bCreateDefaultNodes)
 	{
diff --git a/Engine/Plugins/FX/Niagara/Source/NiagaraEditor/Private/SNiagaraParameterEditor.cpp b/Engine/Plugins/FX/Niagara/Source/NiagaraEditor/Private/SNiagaraParameterEditor.cpp
--- a/Engine/Plugins/FX/Niagara/Source/NiagaraEditor/Private/SNiagaraParameterEditor.cpp
+++ b/Engine/Plugins/FX/Niagara/Source/NiagaraEditor/Private/SNiagaraParameterEditor.cpp
@@ -2,6 +2,16 @@
 
 #include "SNiagaraParameterEditor.h"
 
+const float SNiagaraParameterEditor::DefaultInputSize = 110.0f;
+
+void SNiagaraParameterEditor::Construct(const FArguments& InArgs)
+{
+	HorizontalAlignment = InArgs._HAlign;
+	VerticalAlignment = InArgs._VAlign;
+	MinimumDesiredWidth = InArgs._MinimumDesiredWidth;
+	MaximumDesiredWidth = InArgs._MaximumDesiredWidth;
+}
+
 void SNiagaraParameterEditor::SetOnBeginValueChange(FOnValueChange InOnBeginValueChange)
 {
 	OnBeginValueChange = InOnBeginValueChange;
@@ -17,6 +27,43 @@ void SNiagaraParameterEditor::SetOnValueChanged(FOnValueChange InOnValueChanged)
 	OnValueChanged = InOnValueChanged;
 }
 
+const TOptional<float>& SNiagaraParameterEditor::GetMinimumDesiredWidth() const
+{
+	return MinimumDesiredWidth;
+}
+
+const TOptional<float>& SNiagaraParameterEditor::GetMaximumDesiredWidth() const
+{
+	return MaximumDesiredWidth;
+}
+
+EHorizontalAlignment SNiagaraParameterEditor::GetHorizontalAlignment() const
+{
+	return HorizontalAlignment;
+}
+
+EVerticalAlignment SNiagaraParameterEditor::GetVerticalAlignment() const
+{
+	return VerticalAlignment;
+}
+
+FVector2D SNiagaraParameterEditor::ComputeDesiredSize(float LayoutScaleMultiplier) const
+{
+	FVector2D ComputedDesiredSize = SCompoundWidget::ComputeDesiredSize(LayoutScaleMultiplier);
+	
+	if (MinimumDesiredWidth.IsSet())
+	{
+		ComputedDesiredSize.X = FMath::Max(MinimumDesiredWidth.GetValue(), ComputedDesiredSize.X);
+	}
+
+	if (MaximumDesiredWidth.IsSet())
+	{
+		ComputedDesiredSize.X = FMath::Min(MaximumDesiredWidth.GetValue(), ComputedDesiredSize.X);
+	}
+
+	return ComputedDesiredSize;
+}
+
 bool SNiagaraParameterEditor::GetIsEditingExclusively()
 {
 	return bIsEditingExclusively;
diff --git a/Engine/Plugins/FX/Niagara/Source/NiagaraEditor/Private/Sequencer/NiagaraSequence/MovieSceneNiagaraEmitterTrack.cpp b/Engine/Plugins/FX/Niagara/Source/NiagaraEditor/Private/Sequencer/NiagaraSequence/MovieSceneNiagaraEmitterTrack.cpp
--- a/Engine/Plugins/FX/Niagara/Source/NiagaraEditor/Private/Sequencer/NiagaraSequence/MovieSceneNiagaraEmitterTrack.cpp
+++ b/Engine/Plugins/FX/Niagara/Source/NiagaraEditor/Private/Sequencer/NiagaraSequence/MovieSceneNiagaraEmitterTrack.cpp
@@ -2,29 +2,146 @@
 
 #include "MovieSceneNiagaraEmitterTrack.h"
 #include "NiagaraEmitterHandle.h"
-#include "NiagaraEmitterHandleViewModel.h"
-#include "NiagaraEmitterViewModel.h"
-#include "MovieSceneNiagaraEmitterSection.h"
+#include "NiagaraSystem.h"
+#include "NiagaraEmitter.h"
+#include "ViewModels/NiagaraSystemViewModel.h"
+#include "ViewModels/NiagaraEmitterHandleViewModel.h"
+#include "ViewModels/NiagaraEmitterViewModel.h"
 #include "MovieSceneNiagaraEmitterTrackInstance.h"
+#include "NiagaraGraph.h"
+#include "NiagaraNodeOutput.h"
+#include "NiagaraNodeFunctionCall.h"
+#include "NiagaraScriptSource.h"
+#include "ViewModels/Stack/NiagaraStackGraphUtilities.h"
+#include "NiagaraEditorStyle.h"
 
-UMovieSceneNiagaraEmitterTrack::UMovieSceneNiagaraEmitterTrack(const FObjectInitializer& ObjectInitializer)
-: Super(ObjectInitializer)
+#include "ISequencerSection.h"
+#include "SequencerSectionPainter.h"
+
+#define LOCTEXT_NAMESPACE "MovieSceneNiagaraEmitterTrack"
+
+class FNiagaraInvalidSequencerSection : public ISequencerSection
+{
+public:
+	virtual UMovieSceneSection* GetSectionObject()  override
+	{
+		unimplemented();
+		return nullptr;
+	}
+
+	virtual int32 OnPaintSection(FSequencerSectionPainter& InPainter) const override
+	{
+		unimplemented();
+		return InPainter.LayerId;
+	}
+};
+
+void UMovieSceneNiagaraEmitterSectionBase::Initialize(FNiagaraSystemViewModel& InSystemViewModel, TSharedRef<FNiagaraEmitterHandleViewModel> InEmitterHandleViewModel)
+{
+	ClearFlags(RF_Transactional);
+	SystemViewModel = &InSystemViewModel;
+	EmitterHandleViewModel = InEmitterHandleViewModel;
+}
+
+FNiagaraSystemViewModel& UMovieSceneNiagaraEmitterSectionBase::GetSystemViewModel() const
 {
+	return *SystemViewModel;
 }
 
-TSharedPtr<FNiagaraEmitterHandleViewModel> UMovieSceneNiagaraEmitterTrack::GetEmitterHandle()
+TSharedPtr<FNiagaraEmitterHandleViewModel> UMovieSceneNiagaraEmitterSectionBase::GetEmitterHandleViewModel() const
 {
 	return EmitterHandleViewModel.Pin();
 }
 
-void UMovieSceneNiagaraEmitterTrack::SetEmitterHandle(TSharedRef<FNiagaraEmitterHandleViewModel> InEmitterHandleViewModel)
+FName UMovieSceneNiagaraEmitterSectionBase::GetInstanceName() const
 {
+	return InstanceName;
+}
+
+void UMovieSceneNiagaraEmitterSectionBase::SetInstanceName(FName InInstanceName)
+{
+	InstanceName = InInstanceName;
+}
+
+TSharedRef<ISequencerSection> UMovieSceneNiagaraEmitterSectionBase::MakeInvalidSectionInterface()
+{
+	checkf(false, TEXT("Invalid section interface can't be used."));
+	return MakeShared<FNiagaraInvalidSequencerSection>();
+}
+
+UMovieSceneNiagaraEmitterTrack::UMovieSceneNiagaraEmitterTrack(const FObjectInitializer& ObjectInitializer)
+	: Super(ObjectInitializer)
+{
+}
+
+void UMovieSceneNiagaraEmitterTrack::Initialize(FNiagaraSystemViewModel& InSystemViewModel, TSharedRef<FNiagaraEmitterHandleViewModel> InEmitterHandleViewModel, const FFrameRate& InFrameResolution)
+{
+	ClearFlags(RF_Transactional);
+	SystemViewModel = &InSystemViewModel;
 	EmitterHandleViewModel = InEmitterHandleViewModel;
+	SetDisplayName(InEmitterHandleViewModel->GetNameText());
 	EmitterHandleId = InEmitterHandleViewModel->GetId();
+	SetColorTint(FNiagaraEditorStyle::Get().GetColor("NiagaraEditor.NiagaraSequence.DefaultTrackColor").ToFColor(true));
+	CreateSections(InFrameResolution);
+}
+
+FNiagaraSystemViewModel& UMovieSceneNiagaraEmitterTrack::GetSystemViewModel() const
+{
+	return *SystemViewModel;
+}
 
-	if (Sections.Num() == 0)
+TSharedPtr<FNiagaraEmitterHandleViewModel> UMovieSceneNiagaraEmitterTrack::GetEmitterHandleViewModel() const
+{
+	return EmitterHandleViewModel.Pin();
+}
+
+void UMovieSceneNiagaraEmitterTrack::UpdateTrackFromEmitterGraphChange(const FFrameRate& InFrameResolution)
+{
+	SetDisplayName(EmitterHandleViewModel.Pin()->GetNameText());
+	Sections.Empty();
+	CreateSections(InFrameResolution);
+}
+
+void UMovieSceneNiagaraEmitterTrack::UpdateTrackFromEmitterParameterChange(const FFrameRate& InFrameResolution)
+{
+	for (UMovieSceneSection* Section : Sections)
 	{
-		Sections.Add(NewObject<UMovieSceneNiagaraEmitterSection>(this));
+		CastChecked<UMovieSceneNiagaraEmitterSectionBase>(Section)->UpdateSectionFromModules(InFrameResolution);
+		Section->SetIsActive(EmitterHandleViewModel.Pin()->GetIsEnabled());
+	}
+}
+
+void UMovieSceneNiagaraEmitterTrack::UpdateEmitterHandleFromTrackChange(const FFrameRate& InFrameResolution)
+{
+	if (Sections.Num() > 0)
+	{
+		TSharedPtr<FNiagaraEmitterHandleViewModel> EmitterHandleViewModelPinned = EmitterHandleViewModel.Pin();
+		checkf(EmitterHandleViewModelPinned.IsValid(), TEXT("Emitter handle view model is invalid while it's track is still active."));
+
+		TOptional<bool> SectionsActive = Sections[0]->IsActive();
+		for (UMovieSceneSection* Section : Sections)
+		{
+			CastChecked<UMovieSceneNiagaraEmitterSectionBase>(Section)->UpdateModulesFromSection(InFrameResolution);
+			if (SectionsActive.IsSet() && Section->IsActive() != SectionsActive.GetValue())
+			{
+				SectionsActive.Reset();
+			}
+		}
+
+		if (SectionsActive.IsSet())
+		{
+			EmitterHandleViewModelPinned->SetIsEnabled(SectionsActive.GetValue());
+		}
+		else
+		{
+			// If there were mixed active states on the sections then one was toggled from the emitter handle state so we need
+			// to update the emitter is active and the other sections.
+			EmitterHandleViewModelPinned->SetIsEnabled(!EmitterHandleViewModelPinned->GetIsEnabled());
+			for (UMovieSceneSection* Section : Sections)
+			{
+				Section->SetIsActive(EmitterHandleViewModelPinned->GetIsEnabled());
+			}
+		}
 	}
 }
 
@@ -48,12 +165,9 @@ const TArray<UMovieSceneSection*>& UMovieSceneNiagaraEmitterTrack::GetAllSection
 	return Sections;
 }
 
-FName UMovieSceneNiagaraEmitterTrack::GetTrackName() const
+bool UMovieSceneNiagaraEmitterTrack::SupportsMultipleRows() const
 {
-	TSharedPtr<FNiagaraEmitterHandleViewModel> EmitterHandleViewModelPinned = EmitterHandleViewModel.Pin();
-	return EmitterHandleViewModelPinned.IsValid()
-		? EmitterHandleViewModelPinned->GetName()
-		: NAME_None;
+	return true;
 }
 
 FGuid UMovieSceneNiagaraEmitterTrack::GetEmitterHandleId() const
@@ -61,3 +175,77 @@ FGuid UMovieSceneNiagaraEmitterTrack::GetEmitterHandleId() const
 	return EmitterHandleId;
 }
 
+const TArray<FText>& UMovieSceneNiagaraEmitterTrack::GetSectionInitializationErrors() const
+{
+	return SectionInitializationErrors;
+}
+
+void UMovieSceneNiagaraEmitterTrack::CreateSections(const FFrameRate& InFrameResolution)
+{
+	SectionInitializationErrors.Empty();
+
+	UNiagaraScript* EmitterUpdateScript = GetEmitterHandleViewModel()->GetEmitterViewModel()->GetEmitter()->GetScript(ENiagaraScriptUsage::EmitterUpdateScript, FGuid());
+	UNiagaraScriptSource* ScriptSource = CastChecked<UNiagaraScriptSource>(EmitterUpdateScript->GetSource());
+	UNiagaraNodeOutput* OutputNode = ScriptSource->NodeGraph->FindOutputNode(ENiagaraScriptUsage::EmitterUpdateScript);
+
+	TArray<FNiagaraStackGraphUtilities::FStackNodeGroup> StackGroups;
+	FNiagaraStackGraphUtilities::GetStackNodeGroups(*OutputNode, StackGroups);
+	for (FNiagaraStackGraphUtilities::FStackNodeGroup& StackGroup : StackGroups)
+	{
+		UNiagaraNodeFunctionCall* FunctionNode = Cast<UNiagaraNodeFunctionCall>(StackGroup.EndNode);
+		if (FunctionNode != nullptr && FunctionNode->FunctionScript != nullptr)
+		{
+			FString* SectionClassName = FunctionNode->FunctionScript->ScriptMetaData.Find("NiagaraTimelineSectionClass");
+			if (SectionClassName != nullptr)
+			{
+				UClass* SectionClass = FindObject<UClass>(ANY_PACKAGE, **SectionClassName);
+				if (SectionClass != nullptr)
+				{
+					UMovieSceneNiagaraEmitterSectionBase* EmitterSection = nullptr;
+					bool bSectionCreated = false;
+
+					UMovieSceneSection** SharedEmitterSection = Sections.FindByPredicate([=](UMovieSceneSection* Section) {	return Section->IsA(SectionClass); });
+					if (SharedEmitterSection != nullptr)
+					{
+						EmitterSection = CastChecked<UMovieSceneNiagaraEmitterSectionBase>(*SharedEmitterSection);
+					}
+
+					if(EmitterSection == nullptr)
+					{
+						EmitterSection = NewObject<UMovieSceneNiagaraEmitterSectionBase>(this, SectionClass);
+						EmitterSection->Initialize(*SystemViewModel, EmitterHandleViewModel.Pin().ToSharedRef());
+						bSectionCreated = true;
+					}
+
+					FText AddError;
+					if (EmitterSection->TryAddModule(*FunctionNode, AddError))
+					{
+						if (bSectionCreated)
+						{
+							EmitterSection->SetRowIndex(Sections.Num());
+							Sections.Add(EmitterSection);
+						}
+					}
+					else
+					{
+						FText SectionInitializationError = FText::Format(LOCTEXT("AddModuleErrorFormat", "Failed to add module {0} to section of type {1}.\nMessage: {2}"),
+							FText::FromString(FunctionNode->GetFunctionName()), FText::FromString(*SectionClassName), AddError);
+						SectionInitializationErrors.Add(SectionInitializationError);
+					}
+				}
+				else
+				{
+					SectionInitializationErrors.Add(
+						FText::Format(LOCTEXT("SectionClassNotFoundErrorFormat", "Module script {0} tried to use {1} as it's timeline section class, but it could not be found."),
+							FText::FromString(FunctionNode->FunctionScript->GetPathName()),
+							FText::FromString(*SectionClassName)));
+				}
+			}
+		}
+	}
+
+	UpdateTrackFromEmitterParameterChange(InFrameResolution);
+}
+
+#undef LOCTEXT_NAMESPACE
+
diff --git a/Engine/Plugins/FX/Niagara/Source/NiagaraEditor/Private/Sequencer/NiagaraSequence/MovieSceneNiagaraEmitterTrack.h b/Engine/Plugins/FX/Niagara/Source/NiagaraEditor/Private/Sequencer/NiagaraSequence/MovieSceneNiagaraEmitterTrack.h
--- a/Engine/Plugins/FX/Niagara/Source/NiagaraEditor/Private/Sequencer/NiagaraSequence/MovieSceneNiagaraEmitterTrack.h
+++ b/Engine/Plugins/FX/Niagara/Source/NiagaraEditor/Private/Sequencer/NiagaraSequence/MovieSceneNiagaraEmitterTrack.h
@@ -3,26 +3,70 @@
 #pragma once
 
 #include "MovieSceneNameableTrack.h"
+#include "MovieSceneSection.h"
 #include "MovieSceneNiagaraEmitterTrack.generated.h"
 
-
+class UNiagaraSystem;
+class FNiagaraSystemViewModel;
 class FNiagaraEmitterHandleViewModel;
+class UNiagaraNodeFunctionCall;
+class ISequencerSection;
+
+UCLASS(MinimalAPI)
+class UMovieSceneNiagaraEmitterSectionBase : public UMovieSceneSection
+{
+	GENERATED_BODY()
+
+public:
+	void Initialize(FNiagaraSystemViewModel& InSystemViewModel, TSharedRef<FNiagaraEmitterHandleViewModel> InEmitterHandleViewModel);
+
+	FNiagaraSystemViewModel& GetSystemViewModel() const;
+
+	TSharedPtr<FNiagaraEmitterHandleViewModel> GetEmitterHandleViewModel() const;
+
+	FName GetInstanceName() const;
+
+	void SetInstanceName(FName InInstanceName);
+
+	virtual bool TryAddModule(UNiagaraNodeFunctionCall& InModule, FText& OutErrorMessage) PURE_VIRTUAL(UMovieSceneNiagaraEmitterSectionBase::TryAddModule, return false;);
+
+	virtual void UpdateSectionFromModules(const FFrameRate& InFrameResolution) PURE_VIRTUAL(UMovieSceneNiagaraEmitterSectionBase::UpdateSectionFromModules, );
+
+	virtual void UpdateModulesFromSection(const FFrameRate& InFrameResolution) PURE_VIRTUAL(UMovieSceneNiagaraEmitterSectionBase::UpdateModulesFromSection, );
+
+	virtual TSharedRef<ISequencerSection> MakeSectionInterface() PURE_VIRTUAL(UMovieSceneNiagaraEmitterSectionBase::UpdateModulesFromSection, return MakeInvalidSectionInterface(););
+
+private:
+	TSharedRef<ISequencerSection> MakeInvalidSectionInterface();
+
+	FNiagaraSystemViewModel* SystemViewModel;
+
+	TWeakPtr<FNiagaraEmitterHandleViewModel> EmitterHandleViewModel;
+
+	FName InstanceName;
+};
 
 /**
-*	Single track containing exactly one UNiagaraMovieSceneSection, representing one emitter
+*	A track representing an emitter in the niagara effect editor timeline.
 */
 UCLASS(MinimalAPI)
 class UMovieSceneNiagaraEmitterTrack
 	: public UMovieSceneNameableTrack
 {
 	GENERATED_UCLASS_BODY()
+
 public:
+	void Initialize(FNiagaraSystemViewModel& SystemViewModel, TSharedRef<FNiagaraEmitterHandleViewModel> InEmitterHandleViewModel, const FFrameRate& InFrameResolution);
 
-	TSharedPtr<FNiagaraEmitterHandleViewModel> GetEmitterHandle();
+	FNiagaraSystemViewModel& GetSystemViewModel() const;
 
-	void SetEmitterHandle(TSharedRef<FNiagaraEmitterHandleViewModel> InEmitterHandleViewModel);
+	TSharedPtr<FNiagaraEmitterHandleViewModel> GetEmitterHandleViewModel() const;
 
-public:
+	void UpdateTrackFromEmitterGraphChange(const FFrameRate& InFrameResolution);
+
+	void UpdateTrackFromEmitterParameterChange(const FFrameRate& InFrameResolution);
+
+	void UpdateEmitterHandleFromTrackChange(const FFrameRate& InFrameResolution);
 
 	//~ UMovieSceneTrack interface
 	virtual void RemoveAllAnimationData() override { }
@@ -30,16 +74,19 @@ public:
 	virtual void RemoveSection(UMovieSceneSection& Section) override;
 	virtual bool IsEmpty() const override;
 	virtual const TArray<UMovieSceneSection*>& GetAllSections() const override;
-	virtual FName GetTrackName() const override;
+	virtual bool SupportsMultipleRows() const override;
+
 	/** Gets the unique id for the emitter handle that was associated with this track; used for copy/paste detection */
 	FGuid GetEmitterHandleId() const;
 
+	const TArray<FText>& GetSectionInitializationErrors() const;
+
 private:
-	void EmitterPropertyChanged();
-	void EmitterHandlePropertyChanged();
-	void SynchronizeWithEmitterHandle();
+	void CreateSections(const FFrameRate& InFrameResolution);
 
 private:
+	FNiagaraSystemViewModel* SystemViewModel;
+
 	TWeakPtr<FNiagaraEmitterHandleViewModel> EmitterHandleViewModel;
 
 	UPROPERTY()
@@ -48,4 +95,6 @@ private:
 	// Used for detecting copy/paste 
 	UPROPERTY()
 	FGuid EmitterHandleId;
+
+	TArray<FText> SectionInitializationErrors;
 };
\ No newline at end of file
diff --git a/Engine/Plugins/FX/Niagara/Source/NiagaraEditor/Private/Sequencer/NiagaraSequence/NiagaraEmitterTrackEditor.cpp b/Engine/Plugins/FX/Niagara/Source/NiagaraEditor/Private/Sequencer/NiagaraSequence/NiagaraEmitterTrackEditor.cpp
--- a/Engine/Plugins/FX/Niagara/Source/NiagaraEditor/Private/Sequencer/NiagaraSequence/NiagaraEmitterTrackEditor.cpp
+++ b/Engine/Plugins/FX/Niagara/Source/NiagaraEditor/Private/Sequencer/NiagaraSequence/NiagaraEmitterTrackEditor.cpp
@@ -1,14 +1,180 @@
 // Copyright 1998-2018 Epic Games, Inc. All Rights Reserved.
 
 #include "NiagaraEmitterTrackEditor.h"
-#include "NiagaraEmitterSection.h"
 #include "NiagaraEmitter.h"
+#include "NiagaraRendererProperties.h"
 #include "ViewModels/NiagaraSystemViewModel.h"
+#include "ViewModels/NiagaraEmitterHandleViewModel.h"
+#include "ViewModels/NiagaraEmitterViewModel.h"
 #include "Sequencer/NiagaraSequence/NiagaraSequence.h"
 #include "Framework/MultiBox/MultiBoxBuilder.h"
+#include "NiagaraEditorStyle.h"
+
+#include "EditorStyleSet.h"
+#include "Styling/SlateIconFinder.h"
+#include "Widgets/SBoxPanel.h"
+#include "Widgets/Input/SButton.h"
+#include "Widgets/Images/SImage.h"
+#include "Widgets/Input/SCheckBox.h"
 
 #define LOCTEXT_NAMESPACE "NiagaraEmitterTrackEditor"
 
+class SEmitterTrackWidget : public SCompoundWidget
+{
+public:
+	SLATE_BEGIN_ARGS(SEmitterTrackWidget) {}
+	SLATE_END_ARGS();
+
+	void Construct(const FArguments& InArgs, UMovieSceneNiagaraEmitterTrack& InEmitterTrack)
+	{
+		EmitterTrack = &InEmitterTrack;
+
+		TSharedRef<SHorizontalBox> TrackBox = SNew(SHorizontalBox)
+			// Track initialization error icon.
+			+ SHorizontalBox::Slot()
+			.AutoWidth()
+			.VAlign(VAlign_Center)
+			.Padding(3, 0, 0, 0)
+			[
+				SNew(SImage)
+				.Visibility(this, &SEmitterTrackWidget::GetTrackErrorIconVisibility)
+				.Image(FEditorStyle::GetBrush("Icons.Info"))
+				.ToolTipText(this, &SEmitterTrackWidget::GetTrackErrorIconToolTip)
+			]
+			// Enabled checkbox.
+			+ SHorizontalBox::Slot()
+			.AutoWidth()
+			.VAlign(VAlign_Center)
+			.Padding(3, 0, 0, 0)
+			[
+				SNew(SCheckBox)
+				.ToolTipText(LOCTEXT("EnabledTooltip", "Toggle whether or not this emitter is enabled."))
+				.IsChecked(this, &SEmitterTrackWidget::GetEnabledCheckState)
+				.OnCheckStateChanged(this, &SEmitterTrackWidget::OnEnabledCheckStateChanged)
+			]
+			// Isolate toggle
+			+ SHorizontalBox::Slot()
+			.AutoWidth()
+			.VAlign(VAlign_Center)
+			.Padding(3, 0, 0, 0)
+			[
+				SNew(SButton)
+				.ButtonStyle(FEditorStyle::Get(), "HoverHintOnly")
+				.HAlign(HAlign_Center)
+				.ContentPadding(1)
+				.ToolTipText(this, &SEmitterTrackWidget::GetToggleIsolateToolTip)
+				.OnClicked(this, &SEmitterTrackWidget::OnToggleIsolateButtonClicked)
+				.Content()
+				[
+					SNew(SImage)
+					.Image(FNiagaraEditorStyle::Get().GetBrush("NiagaraEditor.Isolate"))
+					.ColorAndOpacity(this, &SEmitterTrackWidget::GetToggleIsolateImageColor)
+				]
+			];
+
+		// Renderer buttons
+		for (UNiagaraRendererProperties* Renderer : EmitterTrack->GetEmitterHandleViewModel()->GetEmitterViewModel()->GetEmitter()->GetRenderers())
+		{
+			TrackBox->AddSlot()
+				.AutoWidth()
+				.VAlign(VAlign_Center)
+				.Padding(3, 0, 0, 0)
+				[
+					SNew(SButton)
+					.ButtonStyle(FEditorStyle::Get(), "HoverHintOnly")
+					.IsFocusable(false)
+					.ToolTipText(FText::Format(LOCTEXT("RenderButtonToolTip", "{0} - Press to select."), FText::FromString(FName::NameToDisplayString(Renderer->GetName(), false))))
+					.OnClicked(this, &SEmitterTrackWidget::OnRenderButtonClicked)
+					[
+						SNew(SImage)
+						.Image(FSlateIconFinder::FindIconBrushForClass(Renderer->GetClass()))
+					]
+				];
+		}
+
+		ChildSlot
+		[
+			TrackBox
+		];
+	}
+
+private:
+	EVisibility GetTrackErrorIconVisibility() const 
+	{
+		return EmitterTrack.IsValid() && EmitterTrack.Get()->GetSectionInitializationErrors().Num() > 0
+			? EVisibility::Visible
+			: EVisibility::Collapsed;
+	}
+
+	FText GetTrackErrorIconToolTip() const
+	{
+		if(TrackErrorIconToolTip.IsSet() == false && EmitterTrack.IsValid())
+		{
+			FString TrackErrorIconToolTipBuilder;
+			for (const FText& SectionInitializationError : EmitterTrack.Get()->GetSectionInitializationErrors())
+			{
+				if (TrackErrorIconToolTipBuilder.IsEmpty() == false)
+				{
+					TrackErrorIconToolTipBuilder.Append(TEXT("\n"));
+				}
+				TrackErrorIconToolTipBuilder.Append(SectionInitializationError.ToString());
+			}
+			TrackErrorIconToolTip = FText::FromString(TrackErrorIconToolTipBuilder);
+		}
+		return TrackErrorIconToolTip.GetValue();
+	}
+
+	ECheckBoxState GetEnabledCheckState() const
+	{
+		return EmitterTrack.IsValid() && EmitterTrack->GetEmitterHandleViewModel()->GetIsEnabled() ? ECheckBoxState::Checked : ECheckBoxState::Unchecked;
+	}
+
+	void OnEnabledCheckStateChanged(ECheckBoxState InCheckState)
+	{
+		if (EmitterTrack.IsValid())
+		{
+			EmitterTrack->GetEmitterHandleViewModel()->SetIsEnabled(InCheckState == ECheckBoxState::Checked);
+		}
+	}
+
+	FReply OnToggleIsolateButtonClicked()
+	{
+		TArray<TSharedRef<FNiagaraEmitterHandleViewModel>> EmittersToIsolate;
+		if (EmitterTrack.IsValid())
+		{
+			if (EmitterTrack->GetSystemViewModel().IsEmitterIsolated(EmitterTrack->GetEmitterHandleViewModel().ToSharedRef()) == false)
+			{
+				EmittersToIsolate.Add(EmitterTrack->GetEmitterHandleViewModel().ToSharedRef());
+			}
+			EmitterTrack->GetSystemViewModel().IsolateEmitters(EmittersToIsolate);
+		}
+		return FReply::Handled();
+	}
+
+	FText GetToggleIsolateToolTip() const
+	{
+		return EmitterTrack.IsValid() && EmitterTrack->GetSystemViewModel().IsEmitterIsolated(EmitterTrack->GetEmitterHandleViewModel().ToSharedRef())
+			? LOCTEXT("TurnOffEmitterIsolation", "Disable emitter isolation.")
+			: LOCTEXT("IsolateThisEmitter", "Enable isolation for this emitter.");
+	}
+
+	FSlateColor GetToggleIsolateImageColor() const
+	{
+		return EmitterTrack.IsValid() && EmitterTrack->GetSystemViewModel().IsEmitterIsolated(EmitterTrack->GetEmitterHandleViewModel().ToSharedRef())
+			? FEditorStyle::GetSlateColor("SelectionColor")
+			: FLinearColor::Gray;
+	}
+
+	FReply OnRenderButtonClicked()
+	{
+		return FReply::Handled();
+	}
+
+private:
+	TWeakObjectPtr<UMovieSceneNiagaraEmitterTrack> EmitterTrack;
+	mutable TOptional<FText> TrackErrorIconToolTip;
+};
+
 FNiagaraEmitterTrackEditor::FNiagaraEmitterTrackEditor(TSharedPtr<ISequencer> Sequencer) 
 	: FMovieSceneTrackEditor(Sequencer.ToSharedRef())
 {
@@ -30,7 +196,8 @@ bool FNiagaraEmitterTrackEditor::SupportsType(TSubclassOf<UMovieSceneTrack> Trac
 
 TSharedRef<ISequencerSection> FNiagaraEmitterTrackEditor::MakeSectionInterface(UMovieSceneSection& SectionObject, UMovieSceneTrack& Track, FGuid ObjectBinding)
 {
-	return MakeShareable(new FNiagaraEmitterSection(SectionObject));
+	UMovieSceneNiagaraEmitterSectionBase* EmitterSection = CastChecked<UMovieSceneNiagaraEmitterSectionBase>(&SectionObject);
+	return EmitterSection->MakeSectionInterface();
 }
 
 bool FNiagaraEmitterTrackEditor::HandleAssetAdded(UObject* Asset, const FGuid& TargetObjectGuid)
@@ -46,39 +213,38 @@ bool FNiagaraEmitterTrackEditor::HandleAssetAdded(UObject* Asset, const FGuid& T
 
 void FNiagaraEmitterTrackEditor::BuildTrackContextMenu( FMenuBuilder& MenuBuilder, UMovieSceneTrack* Track )
 {
-	UNiagaraSequence* NiagaraSequencer = Cast<UNiagaraSequence>(GetSequencer()->GetRootMovieSceneSequence());
-	if (NiagaraSequencer != nullptr)
+	UMovieSceneNiagaraEmitterTrack* EmitterTrack = CastChecked<UMovieSceneNiagaraEmitterTrack>(Track);
+	FNiagaraSystemViewModel& SystemViewModel = EmitterTrack->GetSystemViewModel();
+
+	if (SystemViewModel.GetEditMode() == ENiagaraSystemViewModelEditMode::SystemAsset)
 	{
-		FNiagaraSystemViewModel* SystemViewModel = &NiagaraSequencer->GetSystemViewModel();
-		if (SystemViewModel->GetEditMode() == ENiagaraSystemViewModelEditMode::SystemAsset)
+		MenuBuilder.BeginSection("Niagara", LOCTEXT("NiagaraContextMenuSectionName", "Niagara"));
 		{
-			MenuBuilder.BeginSection("Niagara", LOCTEXT("NiagaraContextMenuSectionName", "Niagara"));
-			{
-				MenuBuilder.AddMenuEntry(
-					LOCTEXT("IsolateSelectedEmitters", "Isolate Selected"),
-					LOCTEXT("IsolateSelectedEmittersTooltip", "Force selected emitters to only be isolated."),
-					FSlateIcon(),
-					FUIAction(
-						FExecuteAction::CreateSP(SystemViewModel, &FNiagaraSystemViewModel::IsolateSelectedEmitters),
-						FCanExecuteAction::CreateSP(SystemViewModel, &FNiagaraSystemViewModel::CanExecuteToggleIsolation)
-					)
-				);
-				MenuBuilder.AddMenuEntry(
-					LOCTEXT("IsolateEmitter", "Isolated"),
-					LOCTEXT("IsolateEmitterTooltip", "Toggle the isolation state of the selected emitters."),
-					FSlateIcon(),
-					FUIAction(
-						FExecuteAction::CreateSP(SystemViewModel, &FNiagaraSystemViewModel::IsolateEmitterToggle),
-						FCanExecuteAction::CreateSP(SystemViewModel, &FNiagaraSystemViewModel::CanExecuteToggleIsolation),
-						FIsActionChecked::CreateSP(SystemViewModel, &FNiagaraSystemViewModel::IsEmitterIsolationActive)
-					),
-					NAME_None,
-					EUserInterfaceActionType::ToggleButton
-				);
-			}
-			MenuBuilder.EndSection();
+			MenuBuilder.AddMenuEntry(
+				SystemViewModel.IsEmitterIsolated(EmitterTrack->GetEmitterHandleViewModel().ToSharedRef())
+					? LOCTEXT("RemoveFromIsolation", "Remove this from isolation.")
+					: LOCTEXT("AddToIsolation", "Add this to isolation"),
+				SystemViewModel.IsEmitterIsolated(EmitterTrack->GetEmitterHandleViewModel().ToSharedRef())
+					? LOCTEXT("RemoveFromIsolation", "Remove this emitter from isolation, without changing other emitters.")
+					: LOCTEXT("AddToIsolation", "Add this emitter to isolation, without changing other emitters."),
+				FSlateIcon(),
+				FUIAction(FExecuteAction::CreateRaw(&SystemViewModel, &FNiagaraSystemViewModel::ToggleEmitterIsolation, EmitterTrack->GetEmitterHandleViewModel().ToSharedRef())));
+			
+			TArray<TSharedRef<FNiagaraEmitterHandleViewModel>> SelectedEmitters;
+			SystemViewModel.GetSelectedEmitterHandles(SelectedEmitters);
+			MenuBuilder.AddMenuEntry(
+				LOCTEXT("IsolateSelected", "Isolate all selected"),
+				LOCTEXT("IsolateSelectedToolTip", "Add all of the selected emitters to isloation"),
+				FSlateIcon(),
+				FUIAction(FExecuteAction::CreateRaw(&SystemViewModel, &FNiagaraSystemViewModel::IsolateEmitters, SelectedEmitters)));
 		}
+		MenuBuilder.EndSection();
 	}
 }
 
+TSharedPtr<SWidget> FNiagaraEmitterTrackEditor::BuildOutlinerEditWidget(const FGuid& ObjectBinding, UMovieSceneTrack* Track, const FBuildEditWidgetParams& Params)
+{
+	return SNew(SEmitterTrackWidget, *CastChecked<UMovieSceneNiagaraEmitterTrack>(Track));
+}
+
 #undef LOCTEXT_NAMESPACE
\ No newline at end of file
diff --git a/Engine/Plugins/FX/Niagara/Source/NiagaraEditor/Private/Sequencer/NiagaraSequence/NiagaraEmitterTrackEditor.h b/Engine/Plugins/FX/Niagara/Source/NiagaraEditor/Private/Sequencer/NiagaraSequence/NiagaraEmitterTrackEditor.h
--- a/Engine/Plugins/FX/Niagara/Source/NiagaraEditor/Private/Sequencer/NiagaraSequence/NiagaraEmitterTrackEditor.h
+++ b/Engine/Plugins/FX/Niagara/Source/NiagaraEditor/Private/Sequencer/NiagaraSequence/NiagaraEmitterTrackEditor.h
@@ -23,4 +23,12 @@ public:
 	virtual TSharedRef<ISequencerSection> MakeSectionInterface(UMovieSceneSection& SectionObject, UMovieSceneTrack& Track, FGuid ObjectBinding) override;
 	virtual bool HandleAssetAdded(UObject* Asset, const FGuid& TargetObjectGuid) override;
 	virtual void BuildTrackContextMenu( FMenuBuilder& MenuBuilder, UMovieSceneTrack* Track ) override; 
+	virtual TSharedPtr<SWidget> BuildOutlinerEditWidget(const FGuid& ObjectBinding, UMovieSceneTrack* Track, const FBuildEditWidgetParams& Params) override;
+
+private:
+	void ToggleEmitterIsolation(TSharedRef<FNiagaraEmitterHandleViewModel> EmitterToIsolate);
+
+	bool CanIsolateSelectedEmitters() const;
+
+	void IsolateSelectedEmitters();
 };
diff --git a/Engine/Plugins/FX/Niagara/Source/NiagaraEditor/Private/TNiagaraGraphPinEditableName.h b/Engine/Plugins/FX/Niagara/Source/NiagaraEditor/Private/TNiagaraGraphPinEditableName.h
--- a/Engine/Plugins/FX/Niagara/Source/NiagaraEditor/Private/TNiagaraGraphPinEditableName.h
+++ b/Engine/Plugins/FX/Niagara/Source/NiagaraEditor/Private/TNiagaraGraphPinEditableName.h
@@ -48,10 +48,21 @@ protected:
 
 	void OnTextCommitted(const FText& InText, ETextCommit::Type InCommitType)
 	{
-		UNiagaraNode* ParentNode = Cast<UNiagaraNode>(this->GraphPinObj->GetOwningNode());
-		if (ParentNode != nullptr)
+		if (this->GraphPinObj->PinName.ToString() != InText.ToString())
+		{
+			UNiagaraNode* ParentNode = Cast<UNiagaraNode>(this->GraphPinObj->GetOwningNode());
+			if (ParentNode != nullptr)
+			{
+				ParentNode->CommitEditablePinName(InText, this->GraphPinObj);
+			}
+		}
+		else
 		{
-			ParentNode->CommitEditablePinName(InText, this->GraphPinObj);
+			UNiagaraNode* ParentNode = Cast<UNiagaraNode>(this->GraphPinObj->GetOwningNode());
+			if (ParentNode != nullptr)
+			{
+				ParentNode->CancelEditablePinName(InText, this->GraphPinObj);
+			}
 		}
 	}
 
diff --git a/Engine/Plugins/FX/Niagara/Source/NiagaraEditor/Private/Tests/NiagaraEmitterInheritanceTests.cpp b/Engine/Plugins/FX/Niagara/Source/NiagaraEditor/Private/Tests/NiagaraEmitterInheritanceTests.cpp
--- a/Engine/Plugins/FX/Niagara/Source/NiagaraEditor/Private/Tests/NiagaraEmitterInheritanceTests.cpp
+++ b/Engine/Plugins/FX/Niagara/Source/NiagaraEditor/Private/Tests/NiagaraEmitterInheritanceTests.cpp
@@ -87,11 +87,20 @@ bool CheckModuleListsMatch(TArray<TSharedRef<FNiagaraStackFunctionMergeAdapter>>
 		if (AssignmentNodeA != nullptr)
 		{
 			// Assignment nodes have generated scripts based on the assignment target so check that instead of the script directly.
-			if (AssignmentNodeA->AssignmentTarget != AssignmentNodeB->AssignmentTarget)
+			if (AssignmentNodeA->GetAssignmentTargets().Num() != AssignmentNodeB->GetAssignmentTargets().Num())
 			{
-				ErrorMessage = FString::Printf(TEXT("Modules at list index %i were assignment nodes but had different assignment targets."), i);
+				ErrorMessage = FString::Printf(TEXT("Modules at list index %i were assignment nodes but had different assignment target counts."), i);
 				return false;
 			}
+
+			for (int32 k = 0; k < AssignmentNodeA->GetAssignmentTargets().Num(); k++)
+			{
+				if (AssignmentNodeA->GetAssignmentTarget(k) != AssignmentNodeB->GetAssignmentTarget(k))
+				{
+					ErrorMessage = FString::Printf(TEXT("Modules at list index %i were assignment nodes but had different assignment targets at %d."), i, k);
+					return false;
+				}
+			}
 		}
 		else
 		{
diff --git a/Engine/Plugins/FX/Niagara/Source/NiagaraEditor/Private/Tests/NiagaraStackViewModelTests.cpp b/Engine/Plugins/FX/Niagara/Source/NiagaraEditor/Private/Tests/NiagaraStackViewModelTests.cpp
--- a/Engine/Plugins/FX/Niagara/Source/NiagaraEditor/Private/Tests/NiagaraStackViewModelTests.cpp
+++ b/Engine/Plugins/FX/Niagara/Source/NiagaraEditor/Private/Tests/NiagaraStackViewModelTests.cpp
@@ -8,7 +8,6 @@
 #include "ViewModels/Stack/NiagaraStackGraphUtilities.h"
 #include "ViewModels/Stack/NiagaraStackRoot.h"
 #include "ViewModels/Stack/NiagaraStackScriptItemGroup.h"
-#include "ViewModels/Stack/NiagaraStackAddModuleItem.h"
 #include "ViewModels/Stack/NiagaraStackModuleItem.h"
 #include "ViewModels/Stack/NiagaraStackFunctionInputCollection.h"
 #include "ViewModels/Stack/NiagaraStackFunctionInput.h"
@@ -21,9 +20,9 @@
 UNiagaraStackViewModel* CreateStackViewModelForEmitter(UNiagaraEmitter* Emitter)
 {
 	TSharedRef<FNiagaraSystemViewModel> SystemViewModel = FNiagaraEditorTestUtilities::CreateTestSystemViewModelForEmitter(Emitter);
-	TSharedRef<FNiagaraEmitterViewModel> EmitterViewModel = SystemViewModel->GetEmitterHandleViewModels()[0]->GetEmitterViewModel();
+	TSharedRef<FNiagaraEmitterHandleViewModel> EmitterHandleViewModel = SystemViewModel->GetEmitterHandleViewModels()[0];
 	UNiagaraStackViewModel* StackViewModel = NewObject<UNiagaraStackViewModel>();
-	StackViewModel->Initialize(SystemViewModel, EmitterViewModel);
+	StackViewModel->Initialize(SystemViewModel, EmitterHandleViewModel);
 	return StackViewModel;
 }
 
diff --git a/Engine/Plugins/FX/Niagara/Source/NiagaraEditor/Private/Tests/NiagaraStackViewModelTests.h b/Engine/Plugins/FX/Niagara/Source/NiagaraEditor/Private/Tests/NiagaraStackViewModelTests.h
--- a/Engine/Plugins/FX/Niagara/Source/NiagaraEditor/Private/Tests/NiagaraStackViewModelTests.h
+++ b/Engine/Plugins/FX/Niagara/Source/NiagaraEditor/Private/Tests/NiagaraStackViewModelTests.h
@@ -1 +1,2 @@
-// Copyright 1998-2018 Epic Games, Inc. All Rights Reserved.
\ No newline at end of file
+// Copyright 1998-2018 Epic Games, Inc. All Rights Reserved.
+#pragma once
\ No newline at end of file
diff --git a/Engine/Plugins/FX/Niagara/Source/NiagaraEditor/Private/Toolkits/NiagaraScriptToolkit.cpp b/Engine/Plugins/FX/Niagara/Source/NiagaraEditor/Private/Toolkits/NiagaraScriptToolkit.cpp
--- a/Engine/Plugins/FX/Niagara/Source/NiagaraEditor/Private/Toolkits/NiagaraScriptToolkit.cpp
+++ b/Engine/Plugins/FX/Niagara/Source/NiagaraEditor/Private/Toolkits/NiagaraScriptToolkit.cpp
@@ -14,7 +14,9 @@
 #include "SNiagaraSelectedObjectsDetails.h"
 #include "SNiagaraParameterCollection.h"
 #include "UObject/Package.h"
+#include "SNiagaraParameterMapView.h"
 #include "NiagaraEditorStyle.h"
+#include "NiagaraSystem.h"
 #include "BusyCursor.h"
 #include "EngineGlobals.h"
 #include "Misc/FeedbackContext.h"
@@ -37,6 +39,10 @@
 #include "Customizations/NiagaraScriptDetails.h"
 #include "PropertyEditorModule.h"
 
+#include "Developer/MessageLog/Public/IMessageLogListing.h"
+#include "Developer/MessageLog/Public/MessageLogInitializationOptions.h"
+#include "Developer/MessageLog/Public/MessageLogModule.h"
+
 #define LOCTEXT_NAMESPACE "NiagaraScriptToolkit"
 
 DECLARE_CYCLE_STAT(TEXT("Niagara - ScriptToolkit - OnApply"), STAT_NiagaraEditor_ScriptToolkit_OnApply, STATGROUP_NiagaraEditor);
@@ -50,6 +56,7 @@ static TAutoConsoleVariable<int32> CVarDevDetails(
 const FName FNiagaraScriptToolkit::NodeGraphTabId(TEXT("NiagaraEditor_NodeGraph")); 
 const FName FNiagaraScriptToolkit::DetailsTabId(TEXT("NiagaraEditor_Details"));
 const FName FNiagaraScriptToolkit::ParametersTabId(TEXT("NiagaraEditor_Parameters"));
+const FName FNiagaraScriptToolkit::StatsTabId(TEXT("NiagaraEditor_Stats"));
 
 FNiagaraScriptToolkit::FNiagaraScriptToolkit()
 {
@@ -57,9 +64,9 @@ FNiagaraScriptToolkit::FNiagaraScriptToolkit()
 
 FNiagaraScriptToolkit::~FNiagaraScriptToolkit()
 {
+	EditedNiagaraScript->OnVMScriptCompiled().RemoveAll(this);
 }
 
-
 void FNiagaraScriptToolkit::RegisterTabSpawners(const TSharedRef<class FTabManager>& InTabManager)
 {
 	WorkspaceMenuCategory = InTabManager->AddLocalWorkspaceMenuCategory(LOCTEXT("WorkspaceMenu_NiagaraEditor", "Niagara"));
@@ -76,6 +83,15 @@ void FNiagaraScriptToolkit::RegisterTabSpawners(const TSharedRef<class FTabManag
 		.SetDisplayName(LOCTEXT("DetailsTab", "Details"))
 		.SetGroup(WorkspaceMenuCategoryRef)
 		.SetIcon(FSlateIcon(FEditorStyle::GetStyleSetName(), "LevelEditor.Tabs.Details"));
+
+	InTabManager->RegisterTabSpawner(ParametersTabId, FOnSpawnTab::CreateSP(this, &FNiagaraScriptToolkit::SpawnTabScriptParameters))
+		.SetDisplayName(LOCTEXT("ParametersTab", "Parameters"))
+		.SetGroup(WorkspaceMenuCategoryRef);
+		
+	InTabManager->RegisterTabSpawner(StatsTabId, FOnSpawnTab::CreateSP(this, &FNiagaraScriptToolkit::SpawnTabStats))
+		.SetDisplayName(LOCTEXT("StatsTab", "Stats"))
+		.SetGroup(WorkspaceMenuCategoryRef)
+		.SetIcon(FSlateIcon(FEditorStyle::GetStyleSetName(), "LevelEditor.Tabs.Details"));
 }
 
 void FNiagaraScriptToolkit::UnregisterTabSpawners(const TSharedRef<class FTabManager>& InTabManager)
@@ -84,6 +100,8 @@ void FNiagaraScriptToolkit::UnregisterTabSpawners(const TSharedRef<class FTabMan
 
 	InTabManager->UnregisterTabSpawner(NodeGraphTabId);
 	InTabManager->UnregisterTabSpawner(DetailsTabId);
+	InTabManager->UnregisterTabSpawner(ParametersTabId);
+	InTabManager->UnregisterTabSpawner(StatsTabId);
 }
 
 void FNiagaraScriptToolkit::Initialize( const EToolkitMode::Type Mode, const TSharedPtr< class IToolkitHost >& InitToolkitHost, UNiagaraScript* InputScript )
@@ -93,31 +111,60 @@ void FNiagaraScriptToolkit::Initialize( const EToolkitMode::Type Mode, const TSh
 	ResetLoaders(GetTransientPackage()); // Make sure that we're not going to get invalid version number linkers into the package we are going into. 
 	GetTransientPackage()->LinkerCustomVersion.Empty();
 	EditedNiagaraScript = (UNiagaraScript*)StaticDuplicateObject(OriginalNiagaraScript, GetTransientPackage(), NAME_None, ~RF_Standalone, UNiagaraScript::StaticClass());
-
+	EditedNiagaraScript->OnVMScriptCompiled().AddSP(this, &FNiagaraScriptToolkit::OnVMScriptCompiled);
+	
 	ScriptViewModel = MakeShareable(new FNiagaraScriptViewModel(EditedNiagaraScript, LOCTEXT("NiagaraScriptDisplayName", "Niagara Script"), ENiagaraParameterEditMode::EditAll));
 	DetailsSelection = MakeShareable(new FNiagaraObjectSelection());
 	DetailsSelection->SetSelectedObject(EditedNiagaraScript);
-		
-	TSharedRef<FTabManager::FLayout> StandaloneDefaultLayout = FTabManager::NewLayout( "Standalone_Niagara_Layout_v6" )
+	
+	FMessageLogModule& MessageLogModule = FModuleManager::LoadModuleChecked<FMessageLogModule>("MessageLog");
+	FMessageLogInitializationOptions LogOptions;
+	// Show Pages so that user is never allowed to clear log messages
+	LogOptions.bShowPages = false;
+	LogOptions.bShowFilters = false;
+	LogOptions.bAllowClear = false;
+	LogOptions.MaxPageCount = 1;
+	StatsListing = MessageLogModule.CreateLogListing("MaterialEditorStats", LogOptions);
+	Stats = MessageLogModule.CreateLogListingWidget(StatsListing.ToSharedRef());
+
+	TSharedRef<FTabManager::FLayout> StandaloneDefaultLayout = FTabManager::NewLayout("Standalone_Niagara_Layout_v7")
 	->AddArea
 	(
-		FTabManager::NewPrimaryArea() ->SetOrientation(Orient_Vertical)
+		FTabManager::NewPrimaryArea()->SetOrientation(Orient_Vertical)
 		->Split
 		(
 			FTabManager::NewStack()
 			->SetSizeCoefficient(0.1f)
-			->AddTab(GetToolbarTabId(), ETabState::OpenedTab) 
-			->SetHideTabWell( true )
+			->AddTab(GetToolbarTabId(), ETabState::OpenedTab)
+			->SetHideTabWell(true)
 		)
 		->Split
 		(
 			FTabManager::NewSplitter()->SetOrientation(Orient_Horizontal)->SetSizeCoefficient(0.9f)
 			->Split
 			(
-				FTabManager::NewStack()
+				FTabManager::NewSplitter()->SetOrientation(Orient_Vertical)
 				->SetSizeCoefficient(0.2f)
-				->AddTab(DetailsTabId, ETabState::OpenedTab)
-				->SetForegroundTab(DetailsTabId)
+				->Split
+				(
+					FTabManager::NewStack()
+					->SetSizeCoefficient(0.5f)
+					->AddTab(DetailsTabId, ETabState::OpenedTab)
+					->SetForegroundTab(DetailsTabId)
+				)
+				->Split
+				(
+					FTabManager::NewStack()
+					->SetSizeCoefficient(0.4f)
+					->AddTab(ParametersTabId, ETabState::OpenedTab)
+					->SetForegroundTab(ParametersTabId)
+				)
+				->Split
+				(
+					FTabManager::NewStack()
+					->SetSizeCoefficient(0.1f)
+					->AddTab(StatsTabId, ETabState::OpenedTab)
+				)
 			)
 			->Split
 			(
@@ -138,7 +185,7 @@ void FNiagaraScriptToolkit::Initialize( const EToolkitMode::Type Mode, const TSh
 	SetupCommands();
 	ExtendToolbar();
 	RegenerateMenusAndToolbars();
-
+	UpdateModuleStats();
 	// @todo toolkit world centric editing
 	/*// Setup our tool's layout
 	if( IsWorldCentricAssetEditor() )
@@ -216,6 +263,11 @@ void FNiagaraScriptToolkit::OnDetailsSelectionPropertyFinishedChanging(const FPr
 	}
 }
 
+void FNiagaraScriptToolkit::OnVMScriptCompiled(UNiagaraScript* InScript)
+{
+	UpdateModuleStats();
+}
+
 TSharedRef<SDockTab> FNiagaraScriptToolkit::SpawnTabNodeDetails(const FSpawnTabArgs& Args)
 {
 	checkf(Args.GetTabId().TabType == DetailsTabId, TEXT("Wrong tab ID in NiagaraScriptToolkit"));
@@ -267,6 +319,36 @@ TSharedRef<SDockTab> FNiagaraScriptToolkit::SpawnTabNodeDetails(const FSpawnTabA
 	}
 }
 
+TSharedRef<SDockTab> FNiagaraScriptToolkit::SpawnTabScriptParameters(const FSpawnTabArgs& Args)
+{
+	checkf(Args.GetTabId().TabType == ParametersTabId, TEXT("Wrong tab ID in NiagaraScriptToolkit"));
+
+	TSharedRef<SDockTab> SpawnedTab =
+		SNew(SDockTab)
+		[
+			SNew(SNiagaraParameterMapView, DetailsSelection.ToSharedRef(), SNiagaraParameterMapView::EToolkitType::SCRIPT, GetToolkitCommands())
+		];
+
+	return SpawnedTab;
+}
+
+TSharedRef<SDockTab> FNiagaraScriptToolkit::SpawnTabStats(const FSpawnTabArgs& Args)
+{
+	check(Args.GetTabId() == StatsTabId);
+
+	TSharedRef<SDockTab> SpawnedTab = SNew(SDockTab)
+		.Label(LOCTEXT("ModuleStatsTitle", "Stats"))
+		[
+			SNew(SBox)
+			.AddMetaData<FTagMetaData>(FTagMetaData(TEXT("ModuleStats")))
+			[
+				Stats.ToSharedRef()
+			]
+		];
+
+	return SpawnedTab;
+}
+
 void FNiagaraScriptToolkit::SetupCommands()
 {
 	GetToolkitCommands()->MapAction(
@@ -301,12 +383,13 @@ void FNiagaraScriptToolkit::ExtendToolbar()
 				TAttribute<FText>(ScriptToolkit, &FNiagaraScriptToolkit::GetCompileStatusTooltip),
 				TAttribute<FSlateIcon>(ScriptToolkit, &FNiagaraScriptToolkit::GetCompileStatusImage),
 				FName(TEXT("CompileNiagaraScript")));
-			ToolbarBuilder.AddToolBarButton(FNiagaraEditorCommands::Get().RefreshNodes,
+			// removed this for UE-58554 ahead of 4.20. Functionality code should also be removed if this becomes permanent.
+			/*ToolbarBuilder.AddToolBarButton(FNiagaraEditorCommands::Get().RefreshNodes,
 				NAME_None,
 				TAttribute<FText>(),
 				TAttribute<FText>(ScriptToolkit, &FNiagaraScriptToolkit::GetRefreshStatusTooltip),
 				TAttribute<FSlateIcon>(ScriptToolkit, &FNiagaraScriptToolkit::GetRefreshStatusImage),
-				FName(TEXT("RefreshScriptReferences")));
+				FName(TEXT("RefreshScriptReferences")));*/
 			ToolbarBuilder.EndSection();
 		}
 	};
@@ -362,7 +445,7 @@ FText FNiagaraScriptToolkit::GetRefreshStatusTooltip() const
 
 void FNiagaraScriptToolkit::CompileScript(bool bForce)
 {
-	ScriptViewModel->CompileStandaloneScript(bForce);
+	ScriptViewModel->CompileStandaloneScript();
 }
 
 void FNiagaraScriptToolkit::RefreshNodes()
@@ -370,6 +453,13 @@ void FNiagaraScriptToolkit::RefreshNodes()
 	ScriptViewModel->RefreshNodes();
 }
 
+bool FNiagaraScriptToolkit::IsEditScriptDifferentFromOriginalScript() const
+{
+	UNiagaraScript* EditScript = ScriptViewModel->GetScript(OriginalNiagaraScript->GetUsage(), OriginalNiagaraScript->GetUsageId());
+	return EditScript->AreScriptAndSourceSynchronized() == false ||
+		EditScript->GetBaseChangeID() != OriginalNiagaraScript->GetBaseChangeID();
+}
+
 void FNiagaraScriptToolkit::OnApply()
 {
 	SCOPE_CYCLE_COUNTER(STAT_NiagaraEditor_ScriptToolkit_OnApply);
@@ -379,7 +469,7 @@ void FNiagaraScriptToolkit::OnApply()
 
 bool FNiagaraScriptToolkit::OnApplyEnabled() const
 {
-	return ScriptViewModel->GetScriptDirty();
+	return IsEditScriptDifferentFromOriginalScript();
 }
 
 void FNiagaraScriptToolkit::AddReferencedObjects(FReferenceCollector& Collector)
@@ -390,6 +480,19 @@ void FNiagaraScriptToolkit::AddReferencedObjects(FReferenceCollector& Collector)
 
 
 
+void FNiagaraScriptToolkit::UpdateModuleStats()
+{
+	TArray< TSharedRef<class FTokenizedMessage> > Messages;
+
+	TSharedRef<FTokenizedMessage> Line = FTokenizedMessage::Create(EMessageSeverity::Info);
+	uint32 LastOpCount = EditedNiagaraScript->GetVMExecutableData().LastOpCount;
+	Line->AddToken(FTextToken::Create(FText::Format(FText::FromString("LastOpCount: {0}"), LastOpCount)));
+	Messages.Add(Line);
+	
+	StatsListing->ClearMessages();
+	StatsListing->AddMessages(Messages);
+}
+
 void FNiagaraScriptToolkit::GetSaveableObjects(TArray<UObject*>& OutObjects) const
 {
 	OutObjects.Add(OriginalNiagaraScript);
@@ -430,7 +533,8 @@ void FNiagaraScriptToolkit::UpdateOriginalNiagaraScript()
 	ResetLoaders(OriginalNiagaraScript->GetOutermost()); // Make sure that we're not going to get invalid version number linkers into the package we are going into. 
 	OriginalNiagaraScript->GetOutermost()->LinkerCustomVersion.Empty();
 
-	// overwrite the original script in place by constructing a new one with the same name
+	// Compile and then overwrite the original script in place by constructing a new one with the same name
+	ScriptViewModel->CompileStandaloneScript();
 	OriginalNiagaraScript = (UNiagaraScript*)StaticDuplicateObject(EditedNiagaraScript, OriginalNiagaraScript->GetOuter(), OriginalNiagaraScript->GetFName(),
 		RF_AllFlags,
 		OriginalNiagaraScript->GetClass());
@@ -467,7 +571,7 @@ void FNiagaraScriptToolkit::UpdateOriginalNiagaraScript()
 
 		if (bRefreshed)
 		{
-			Source->NodeGraph->NotifyGraphNeedsRecompile();
+			//Source->NodeGraph->NotifyGraphNeedsRecompile();
 			AffectedScripts.AddUnique(*It);
 		}
 		else
@@ -480,7 +584,7 @@ void FNiagaraScriptToolkit::UpdateOriginalNiagaraScript()
 			{
 				if (Graph == OriginalGraph)
 				{
-					Source->NodeGraph->NotifyGraphNeedsRecompile();
+					//Source->NodeGraph->NotifyGraphNeedsRecompile();
 					AffectedScripts.AddUnique(*It);
 					break;
 				}
@@ -493,7 +597,7 @@ void FNiagaraScriptToolkit::UpdateOriginalNiagaraScript()
 	TArray<UNiagaraEmitter*> AffectedEmitters;
 	for (UNiagaraScript* Script : AffectedScripts)
 	{
-		if (Script->IsParticleUpdateScript() || Script->IsParticleSpawnScript() || Script->IsParticleEventScript())
+		if (Script->IsParticleScript() || Script->IsEmitterSpawnScript() || Script->IsEmitterUpdateScript())
 		{
 			UNiagaraEmitter* Emitter = Cast<UNiagaraEmitter>(Script->GetOuter());
 			if (Emitter)
@@ -501,6 +605,18 @@ void FNiagaraScriptToolkit::UpdateOriginalNiagaraScript()
 				AffectedEmitters.AddUnique(Emitter);
 			}
 		}
+		else if (Script->IsSystemSpawnScript() || Script->IsSystemUpdateScript())
+		{
+			UNiagaraSystem* System = Cast<UNiagaraSystem>(Script->GetOuter());
+			if (System)
+			{
+				for (int32 i = 0; i < System->GetNumEmitters(); i++)
+				{
+					AffectedEmitters.AddUnique(System->GetEmitterHandle(i).GetInstance());
+					AffectedEmitters.AddUnique((UNiagaraEmitter*)System->GetEmitterHandle(i).GetSource());
+				}
+			}
+		}
 		else
 		{
 			TSharedPtr<FNiagaraScriptViewModel> AffectedScriptViewModel = FNiagaraScriptViewModel::GetExistingViewModelForObject(Script);
@@ -508,20 +624,19 @@ void FNiagaraScriptToolkit::UpdateOriginalNiagaraScript()
 			{
 				AffectedScriptViewModel = MakeShareable(new FNiagaraScriptViewModel(Script, FText::FromString(Script->GetName()), ENiagaraParameterEditMode::EditValueOnly));
 			}
-			AffectedScriptViewModel->CompileStandaloneScript(false);
+			AffectedScriptViewModel->CompileStandaloneScript();
 		}
 	}
 
 	FNiagaraEditorUtilities::CompileExistingEmitters(AffectedEmitters);
 
 	GWarn->EndSlowTask();
-	ScriptViewModel->SetScriptDirty(false);
 }
 
 
 bool FNiagaraScriptToolkit::OnRequestClose()
 {
-	if (ScriptViewModel->GetScriptDirty())
+	if (IsEditScriptDifferentFromOriginalScript())
 	{
 		// find out the user wants to do with this dirty NiagaraScript
 		EAppReturnType::Type YesNoCancelReply = FMessageDialog::Open(EAppMsgType::YesNoCancel,
diff --git a/Engine/Plugins/FX/Niagara/Source/NiagaraEditor/Private/Toolkits/NiagaraScriptToolkit.h b/Engine/Plugins/FX/Niagara/Source/NiagaraEditor/Private/Toolkits/NiagaraScriptToolkit.h
--- a/Engine/Plugins/FX/Niagara/Source/NiagaraEditor/Private/Toolkits/NiagaraScriptToolkit.h
+++ b/Engine/Plugins/FX/Niagara/Source/NiagaraEditor/Private/Toolkits/NiagaraScriptToolkit.h
@@ -49,6 +49,10 @@ public:
 	/** FGCObject interface */
 	virtual void AddReferencedObjects(FReferenceCollector& Collector) override;
 
+	/**
+	* Updates list of module info used to show stats
+	*/
+	void UpdateModuleStats();
 
 protected:
 	//~ FAssetEditorToolkit interface
@@ -58,6 +62,7 @@ protected:
 	virtual bool OnRequestClose() override;
 
 	void OnDetailsSelectionPropertyFinishedChanging(const FPropertyChangedEvent& InEvent);
+	void OnVMScriptCompiled(UNiagaraScript* InScript);
 
 private:
 
@@ -67,6 +72,10 @@ private:
 	/** Spawns the tab with the script details inside. */
 	TSharedRef<SDockTab> SpawnTabNodeDetails(const FSpawnTabArgs& Args);
 
+	/** Spawns the tab with the script details inside. */
+	TSharedRef<SDockTab> SpawnTabScriptParameters(const FSpawnTabArgs& Args);
+	TSharedRef<SDockTab> SpawnTabStats(const FSpawnTabArgs& Args);
+
 	/** Sets up commands for the toolkit toolbar. */
 	void SetupCommands();
 
@@ -86,6 +95,8 @@ private:
 	FText GetRefreshStatusTooltip() const;
 	
 private:
+	bool IsEditScriptDifferentFromOriginalScript() const;
+
 	/** Command for the apply button */
 	void OnApply();
 	bool OnApplyEnabled() const;
@@ -102,4 +113,8 @@ private:
 	static const FName NodeGraphTabId; 
 	static const FName DetailsTabId;
 	static const FName ParametersTabId;
+	static const FName StatsTabId;
+	/** Stats log, with the log listing that it reflects */
+	TSharedPtr<class SWidget> Stats;
+	TSharedPtr<class IMessageLogListing> StatsListing;
 };
