Copying //UE4/Dev-Niagara to //UE4/Dev-Main (Source: //UE4/Dev-Niagara @ 4074996)

diff --git a/Engine/Plugins/FX/Niagara/Source/NiagaraEditor/Private/Toolkits/NiagaraSystemToolkit.cpp b/Engine/Plugins/FX/Niagara/Source/NiagaraEditor/Private/Toolkits/NiagaraSystemToolkit.cpp
--- a/Engine/Plugins/FX/Niagara/Source/NiagaraEditor/Private/Toolkits/NiagaraSystemToolkit.cpp
+++ b/Engine/Plugins/FX/Niagara/Source/NiagaraEditor/Private/Toolkits/NiagaraSystemToolkit.cpp
@@ -5,15 +5,16 @@
 #include "NiagaraSystem.h"
 #include "NiagaraEmitter.h"
 #include "NiagaraObjectSelection.h"
-#include "NiagaraSystemViewModel.h"
-#include "NiagaraEmitterHandleviewModel.h"
-#include "NiagaraEmitterViewModel.h"
+#include "ViewModels/NiagaraSystemViewModel.h"
+#include "ViewModels/NiagaraEmitterHandleViewModel.h"
+#include "ViewModels/NiagaraEmitterViewModel.h"
 #include "NiagaraSystemScriptViewModel.h"
 #include "Widgets/SNiagaraCurveEditor.h"
 #include "Widgets/SNiagaraSystemScript.h"
 #include "Widgets/SNiagaraSystemViewport.h"
 #include "Widgets/SNiagaraSelectedObjectsDetails.h"
-#include "Widgets/SNiagaraSelectedEmitterHandle.h"
+#include "Widgets/SNiagaraParameterMapView.h"
+#include "Widgets/SNiagaraSelectedEmitterHandles.h"
 #include "Widgets/SNiagaraSpreadsheetView.h"
 #include "Widgets/SNiagaraGeneratedCodeView.h"
 #include "Widgets/SNiagaraScriptGraph.h"
@@ -45,6 +46,8 @@
 #include "Modules/ModuleManager.h"
 #include "AssetRegistryModule.h"
 #include "IAssetRegistry.h"
+#include "HAL/PlatformFilemanager.h"
+#include "Misc/FileHelper.h"
 
 #define LOCTEXT_NAMESPACE "NiagaraSystemEditor"
 
@@ -55,12 +58,21 @@ const FName FNiagaraSystemToolkit::CurveEditorTabID(TEXT("NiagaraSystemEditor_Cu
 const FName FNiagaraSystemToolkit::SequencerTabID(TEXT("NiagaraSystemEditor_Sequencer"));
 const FName FNiagaraSystemToolkit::SystemScriptTabID(TEXT("NiagaraSystemEditor_SystemScript"));
 const FName FNiagaraSystemToolkit::SystemDetailsTabID(TEXT("NiagaraSystemEditor_SystemDetails"));
+const FName FNiagaraSystemToolkit::SystemParametersTabID(TEXT("NiagaraSystemEditor_SystemParameters"));
 const FName FNiagaraSystemToolkit::SelectedEmitterStackTabID(TEXT("NiagaraSystemEditor_SelectedEmitterStack"));
 const FName FNiagaraSystemToolkit::SelectedEmitterGraphTabID(TEXT("NiagaraSystemEditor_SelectedEmitterGraph"));
 const FName FNiagaraSystemToolkit::DebugSpreadsheetTabID(TEXT("NiagaraSystemEditor_DebugAttributeSpreadsheet"));
 const FName FNiagaraSystemToolkit::PreviewSettingsTabId(TEXT("NiagaraSystemEditor_PreviewSettings"));
 const FName FNiagaraSystemToolkit::GeneratedCodeTabID(TEXT("NiagaraSystemEditor_GeneratedCode"));
 
+static int32 GbLogNiagaraSystemChanges = 0;
+static FAutoConsoleVariableRef CVarSuppressNiagaraSystems(
+	TEXT("fx.LogNiagaraSystemChanges"),
+	GbLogNiagaraSystemChanges,
+	TEXT("If > 0 Niagara Systems will be written to a text format when opened and closed in the editor. \n"),
+	ECVF_Default
+);
+
 void FNiagaraSystemToolkit::RegisterTabSpawners(const TSharedRef<class FTabManager>& InTabManager)
 {
 	WorkspaceMenuCategory = InTabManager->AddLocalWorkspaceMenuCategory(LOCTEXT("WorkspaceMenu_NiagaraSystemEditor", "Niagara System"));
@@ -88,8 +100,12 @@ void FNiagaraSystemToolkit::RegisterTabSpawners(const TSharedRef<class FTabManag
 		.SetDisplayName(LOCTEXT("SystemDetails", "System Details"))
 		.SetGroup(WorkspaceMenuCategory.ToSharedRef());
 
+	InTabManager->RegisterTabSpawner(SystemParametersTabID, FOnSpawnTab::CreateSP(this, &FNiagaraSystemToolkit::SpawnTab_SystemParameters))
+		.SetDisplayName(LOCTEXT("SystemParameters", "Parameters"))
+		.SetGroup(WorkspaceMenuCategory.ToSharedRef());
+
 	InTabManager->RegisterTabSpawner(SelectedEmitterStackTabID, FOnSpawnTab::CreateSP(this, &FNiagaraSystemToolkit::SpawnTab_SelectedEmitterStack))
-		.SetDisplayName(LOCTEXT("SelectedEmitterStack", "Selected Emitter"))
+		.SetDisplayName(LOCTEXT("SelectedEmitterStacks", "Selected Emitters"))
 		.SetGroup(WorkspaceMenuCategory.ToSharedRef());
 
 	InTabManager->RegisterTabSpawner(SelectedEmitterGraphTabID, FOnSpawnTab::CreateSP(this, &FNiagaraSystemToolkit::SpawnTab_SelectedEmitterGraph))
@@ -119,6 +135,7 @@ void FNiagaraSystemToolkit::UnregisterTabSpawners(const TSharedRef<class FTabMan
 	InTabManager->UnregisterTabSpawner(SequencerTabID);
 	InTabManager->UnregisterTabSpawner(SystemScriptTabID);
 	InTabManager->UnregisterTabSpawner(SystemDetailsTabID);
+	InTabManager->UnregisterTabSpawner(SystemParametersTabID);
 	InTabManager->UnregisterTabSpawner(SelectedEmitterStackTabID);
 	InTabManager->UnregisterTabSpawner(SelectedEmitterGraphTabID);
 	InTabManager->UnregisterTabSpawner(DebugSpreadsheetTabID);
@@ -132,6 +149,7 @@ FNiagaraSystemToolkit::~FNiagaraSystemToolkit()
 	if (SystemViewModel.IsValid())
 	{
 		SystemViewModel->Cleanup();
+		SystemViewModel->GetOnPinnedCurvesChanged().RemoveAll(this);
 	}
 	SystemViewModel.Reset();
 }
@@ -141,35 +159,13 @@ void FNiagaraSystemToolkit::AddReferencedObjects(FReferenceCollector& Collector)
 	Collector.AddReferencedObject(System);
 }
 
+
+
 void FNiagaraSystemToolkit::InitializeWithSystem(const EToolkitMode::Type Mode, const TSharedPtr< class IToolkitHost >& InitToolkitHost, UNiagaraSystem& InSystem)
 {
 	System = &InSystem;
 	Emitter = nullptr;
 
-	// In the FNiagaraCustomVersion::UpdateSpawnEventGraphCombination we merged graphs. We update the graph source here because there isn't a good place to do it in the postload pipeline.
-	bool bConverted = false;
-	for (int32 i = 0; i < System->GetNumEmitters(); i++)
-	{
-		FNiagaraEmitterHandle& Handle = System->GetEmitterHandle(i);
-		if (Handle.GetSource() == nullptr)
-		{
-			UE_LOG(LogNiagaraEditor, Error, TEXT("Missing source emitter!"));
-			break;
-		}
-
-		if (Handle.GetSource()->GraphSource == nullptr)
-		{
-			if (FNiagaraEditorUtilities::ConvertToMergedGraph((UNiagaraEmitter*)Handle.GetSource()))
-			{
-				bConverted = true;
-			}
-			else
-			{
-				UE_LOG(LogNiagaraEditor, Error, TEXT("Failed to merge emitter!"));
-			}
-		}
-	}
-
 	FNiagaraSystemViewModelOptions SystemOptions;
 	SystemOptions.bCanModifyEmittersFromTimeline = true;
 	SystemOptions.bUseSystemExecStateForTimelineReset = true;
@@ -177,22 +173,26 @@ void FNiagaraSystemToolkit::InitializeWithSystem(const EToolkitMode::Type Mode,
 	SystemOptions.OnGetSequencerAddMenuContent.BindSP(this, &FNiagaraSystemToolkit::GetSequencerAddMenuContent);
 
 	SystemViewModel = MakeShareable(new FNiagaraSystemViewModel(*System, SystemOptions));
+	SystemViewModel->SetToolkitCommands(GetToolkitCommands());
 	SystemToolkitMode = ESystemToolkitMode::System;
+
+	if (GbLogNiagaraSystemChanges > 0)
+	{
+		FString ExportText;
+		SystemViewModel->DumpToText(ExportText);
+		FString FilePath = System->GetOutermost()->FileName.ToString();
+		FString PathPart, FilenamePart, ExtensionPart;
+		FPaths::Split(FilePath, PathPart, FilenamePart, ExtensionPart);
+
+		FNiagaraEditorUtilities::WriteTextFileToDisk(FPaths::ProjectLogDir(), FilenamePart + TEXT(".onLoad.txt"), ExportText, true);
+	}
+
 	InitializeInternal(Mode, InitToolkitHost);
 }
 
 void FNiagaraSystemToolkit::InitializeWithEmitter(const EToolkitMode::Type Mode, const TSharedPtr< class IToolkitHost >& InitToolkitHost, UNiagaraEmitter& InEmitter)
 {
-	// In the FNiagaraCustomVersion::UpdateSpawnEventGraphCombination we merged graphs. We update the graph source here because there isn't a good place to do it in the postload pipeline.
-	if (InEmitter.GraphSource == nullptr)
-	{
-		if (false == FNiagaraEditorUtilities::ConvertToMergedGraph(&InEmitter))
-		{
-			UE_LOG(LogNiagaraEditor, Error, TEXT("Failed to merge emitter!"));
-		}
-	}
-
-	System = NewObject<UNiagaraSystem>(GetTransientPackage(), NAME_None, RF_Transient);
+	System = NewObject<UNiagaraSystem>(GetTransientPackage(), NAME_None, RF_Transient | RF_Transactional);
 	UNiagaraSystemFactoryNew::InitializeSystem(System, true);
 
 	Emitter = &InEmitter;
@@ -201,7 +201,6 @@ void FNiagaraSystemToolkit::InitializeWithEmitter(const EToolkitMode::Type Mode,
 	GetTransientPackage()->LinkerCustomVersion.Empty();
 
 	UNiagaraEmitter* EditableEmitter = (UNiagaraEmitter*)StaticDuplicateObject(Emitter, GetTransientPackage(), NAME_None, ~RF_Standalone, UNiagaraEmitter::StaticClass());
-	System->AddEmitterHandleWithoutCopying(*EditableEmitter);
 
 	FNiagaraSystemViewModelOptions SystemOptions;
 	SystemOptions.bCanModifyEmittersFromTimeline = false;
@@ -209,9 +208,22 @@ void FNiagaraSystemToolkit::InitializeWithEmitter(const EToolkitMode::Type Mode,
 	SystemOptions.EditMode = ENiagaraSystemViewModelEditMode::EmitterAsset;
 
 	SystemViewModel = MakeShareable(new FNiagaraSystemViewModel(*System, SystemOptions));
+	SystemViewModel->SetToolkitCommands(GetToolkitCommands());
+	SystemViewModel->AddEmitter(*EditableEmitter);
 	SystemViewModel->GetSystemScriptViewModel()->RebuildEmitterNodes();
-	SystemViewModel->GetSystemScriptViewModel()->CompileSystem(false);
 	SystemToolkitMode = ESystemToolkitMode::Emitter;
+
+	if (GbLogNiagaraSystemChanges > 0)
+	{
+		FString ExportText;
+		SystemViewModel->DumpToText(ExportText);
+		FString FilePath = Emitter->GetOutermost()->FileName.ToString();
+		FString PathPart, FilenamePart, ExtensionPart;
+		FPaths::Split(FilePath, PathPart, FilenamePart, ExtensionPart);
+
+		FNiagaraEditorUtilities::WriteTextFileToDisk(FPaths::ProjectLogDir(), FilenamePart + TEXT(".onLoad.txt"), ExportText, true);
+	}
+
 	InitializeInternal(Mode, InitToolkitHost);
 }
 
@@ -222,6 +234,11 @@ void FNiagaraSystemToolkit::InitializeInternal(const EToolkitMode::Type Mode, co
 		SystemViewModel->SetSelectedEmitterHandleById(SystemViewModel->GetEmitterHandleViewModels()[0]->GetId());
 	}
 
+	SystemViewModel->OnEmitterHandleViewModelsChanged().AddSP(this, &FNiagaraSystemToolkit::OnRefresh);
+	SystemViewModel->OnSelectedEmitterHandlesChanged().AddSP(this, &FNiagaraSystemToolkit::OnRefresh);
+	SystemViewModel->GetOnPinnedEmittersChanged().AddSP(this, &FNiagaraSystemToolkit::OnRefresh);
+	SystemViewModel->GetOnPinnedCurvesChanged().AddSP(this, &FNiagaraSystemToolkit::OnPinnedCurvesChanged);
+
 	const float InTime = -0.02f;
 	const float OutTime = 3.2f;
 
@@ -235,32 +252,43 @@ void FNiagaraSystemToolkit::InitializeInternal(const EToolkitMode::Type Mode, co
 				->SetSizeCoefficient(0.1f)
 				->AddTab(GetToolbarTabId(), ETabState::OpenedTab)
 				->SetHideTabWell(true)
-				)
+			)
 			->Split
 			(
 				FTabManager::NewSplitter()->SetOrientation(Orient_Horizontal)
 				->Split
 				(
 					FTabManager::NewSplitter()->SetOrientation(Orient_Vertical)
-					->SetSizeCoefficient(.75f)
+					->SetSizeCoefficient(.60f)
 					->Split
 					(
-						FTabManager::NewStack()
-						->SetSizeCoefficient(.75f)
-						->AddTab(ViewportTabID, ETabState::OpenedTab)
+						FTabManager::NewSplitter()->SetOrientation(Orient_Horizontal)
+						->SetSizeCoefficient(0.75f)
+						->Split
+						(
+							FTabManager::NewStack()
+							->SetSizeCoefficient(.80f)
+							->AddTab(ViewportTabID, ETabState::OpenedTab)
 						)
+						->Split
+						(
+							FTabManager::NewStack()
+							->SetSizeCoefficient(0.20f)
+							->AddTab(SystemParametersTabID, ETabState::OpenedTab)
+						)
+					)
 					->Split
 					(
 						FTabManager::NewStack()
 						->SetSizeCoefficient(0.25f)
 						->AddTab(CurveEditorTabID, ETabState::OpenedTab)
 						->AddTab(SequencerTabID, ETabState::OpenedTab)
-						)
 					)
+				)
 				->Split
 				(
 					FTabManager::NewStack()
-					->SetSizeCoefficient(0.25f)
+					->SetSizeCoefficient(0.40f)
 					->AddTab(SelectedEmitterStackTabID, ETabState::OpenedTab)
 					->AddTab(SelectedEmitterGraphTabID, ETabState::ClosedTab)
 					->AddTab(SystemScriptTabID, ETabState::ClosedTab)
@@ -406,6 +434,31 @@ TSharedRef<SDockTab> FNiagaraSystemToolkit::SpawnTab_SystemDetails(const FSpawnT
 	return SpawnedTab;
 }
 
+TSharedRef<SDockTab> FNiagaraSystemToolkit::SpawnTab_SystemParameters(const FSpawnTabArgs& Args)
+{
+	check(Args.GetTabId().TabType == SystemParametersTabID);
+
+	TSharedRef<FNiagaraObjectSelection> ObjectSelection = MakeShareable(new FNiagaraObjectSelection());
+	if (SystemToolkitMode == ESystemToolkitMode::Emitter)
+	{
+		TSharedPtr<FNiagaraEmitterViewModel> EditableEmitterViewModel = SystemViewModel->GetEmitterHandleViewModels()[0]->GetEmitterViewModel();
+		UNiagaraEmitter* EditableEmitter = EditableEmitterViewModel->GetEmitter();
+		ObjectSelection->SetSelectedObject(EditableEmitter);
+	}
+	else if (SystemToolkitMode == ESystemToolkitMode::System)
+	{
+		ObjectSelection->SetSelectedObject(System);
+	}
+
+	TSharedRef<SDockTab> SpawnedTab =
+		SNew(SDockTab)
+		[
+			SAssignNew(ParameterMapView, SNiagaraParameterMapView, ObjectSelection, SNiagaraParameterMapView::EToolkitType::SYSTEM, GetToolkitCommands())
+		];
+
+	return SpawnedTab;
+}
+
 TSharedRef<SDockTab> FNiagaraSystemToolkit::SpawnTab_SelectedEmitterStack(const FSpawnTabArgs& Args)
 {
 	check(Args.GetTabId().TabType == SelectedEmitterStackTabID);
@@ -413,7 +466,7 @@ TSharedRef<SDockTab> FNiagaraSystemToolkit::SpawnTab_SelectedEmitterStack(const
 	TSharedRef<SDockTab> SpawnedTab =
 		SNew(SDockTab)
 		[
-			SNew(SNiagaraSelectedEmitterHandle, SystemViewModel.ToSharedRef())
+			SNew(SNiagaraSelectedEmitterHandles, SystemViewModel.ToSharedRef())
 		];
 
 	return SpawnedTab;
@@ -441,7 +494,10 @@ public:
 	{
 		if (SystemViewModel.IsValid())
 		{
-			SystemViewModel->OnCurveOwnerChanged().RemoveAll(this);
+			SystemViewModel->OnEmitterHandleViewModelsChanged().RemoveAll(this);
+			SystemViewModel->OnSelectedEmitterHandlesChanged().RemoveAll(this);
+			SystemViewModel->GetOnPinnedEmittersChanged().RemoveAll(this);
+			SystemViewModel->OnSelectedEmitterHandlesChanged().RemoveAll(this);
 		}
 	}
 
@@ -514,7 +570,7 @@ void FNiagaraSystemToolkit::SetupCommands()
 {
 	GetToolkitCommands()->MapAction(
 		FNiagaraEditorCommands::Get().Compile,
-		FExecuteAction::CreateRaw(this, &FNiagaraSystemToolkit::CompileSystem, true));
+		FExecuteAction::CreateRaw(this, &FNiagaraSystemToolkit::CompileSystem, false));
 	GetToolkitCommands()->MapAction(
 		FNiagaraEditorCommands::Get().ResetSimulation,
 		FExecuteAction::CreateRaw(this, &FNiagaraSystemToolkit::ResetSimulation));
@@ -537,6 +593,36 @@ void FNiagaraSystemToolkit::SetupCommands()
 		FNiagaraEditorCommands::Get().Apply,
 		FExecuteAction::CreateSP(this, &FNiagaraSystemToolkit::OnApply),
 		FCanExecuteAction::CreateSP(this, &FNiagaraSystemToolkit::OnApplyEnabled));
+
+	GetToolkitCommands()->MapAction(
+		FNiagaraEditorCommands::Get().ToggleAutoPlay,
+		FExecuteAction::CreateLambda([]()
+		{
+			UNiagaraEditorSettings* Settings = GetMutableDefault<UNiagaraEditorSettings>();
+			Settings->SetAutoPlay(!Settings->GetAutoPlay());
+		}),
+		FCanExecuteAction(),
+		FIsActionChecked::CreateLambda([]() { return GetDefault<UNiagaraEditorSettings>()->GetAutoPlay(); }));
+
+	GetToolkitCommands()->MapAction(
+		FNiagaraEditorCommands::Get().ToggleResetSimulationOnChange,
+		FExecuteAction::CreateLambda([]()
+		{
+			UNiagaraEditorSettings* Settings = GetMutableDefault<UNiagaraEditorSettings>();
+			Settings->SetResetSimulationOnChange(!Settings->GetResetSimulationOnChange());
+		}),
+		FCanExecuteAction(),
+		FIsActionChecked::CreateLambda([]() { return GetDefault<UNiagaraEditorSettings>()->GetResetSimulationOnChange(); }));
+
+	GetToolkitCommands()->MapAction(
+		FNiagaraEditorCommands::Get().ToggleResimulateOnChangeWhilePaused,
+		FExecuteAction::CreateLambda([]()
+		{
+			UNiagaraEditorSettings* Settings = GetMutableDefault<UNiagaraEditorSettings>();
+			Settings->SetResimulateOnChangeWhilePaused(!Settings->GetResimulateOnChangeWhilePaused());
+		}),
+		FCanExecuteAction(),
+		FIsActionChecked::CreateLambda([]() { return GetDefault<UNiagaraEditorSettings>()->GetResimulateOnChangeWhilePaused(); }));
 }
 
 void FNiagaraSystemToolkit::OnSaveThumbnailImage()
@@ -556,6 +642,15 @@ void FNiagaraSystemToolkit::ExtendToolbar()
 {
 	struct Local
 	{
+		static TSharedRef<SWidget> FillSimulationOptionsMenu(FNiagaraSystemToolkit* Toolkit)
+		{
+			FMenuBuilder MenuBuilder(true, Toolkit->GetToolkitCommands());
+			MenuBuilder.AddMenuEntry(FNiagaraEditorCommands::Get().ToggleAutoPlay);
+			MenuBuilder.AddMenuEntry(FNiagaraEditorCommands::Get().ToggleResetSimulationOnChange);
+			MenuBuilder.AddMenuEntry(FNiagaraEditorCommands::Get().ToggleResimulateOnChangeWhilePaused);
+			return MenuBuilder.MakeWidget();
+		}
+
 		static void FillToolbar(FToolBarBuilder& ToolbarBuilder, FNiagaraSystemToolkit* Toolkit)
 		{
 			if (Toolkit->Emitter != nullptr)
@@ -579,7 +674,7 @@ void FNiagaraSystemToolkit::ExtendToolbar()
 					FName(TEXT("CompileNiagaraSystem")));
 				ToolbarBuilder.AddComboButton(
 					FUIAction(),
-					FOnGetContent::CreateStatic(&FNiagaraSystemToolkit::GenerateCompileMenuContent),
+					FOnGetContent::CreateRaw(Toolkit, &FNiagaraSystemToolkit::GenerateCompileMenuContent),
 					LOCTEXT("BuildCombo_Label", "Auto-Compile Options"),
 					LOCTEXT("BuildComboToolTip", "Auto-Compile options menu"),
 					FSlateIcon(FEditorStyle::GetStyleSetName(), "LevelEditor.Build"),
@@ -612,6 +707,18 @@ void FNiagaraSystemToolkit::ExtendToolbar()
 				);
 			}
 			ToolbarBuilder.EndSection();
+
+			ToolbarBuilder.BeginSection("PlaybackOptions");
+			{
+				ToolbarBuilder.AddComboButton(
+					FUIAction(),
+					FOnGetContent::CreateStatic(Local::FillSimulationOptionsMenu, Toolkit),
+					LOCTEXT("SimulationOptions", "Simulation"),
+					LOCTEXT("SimulationOptionsTooltip", "Simulation options"),
+					FSlateIcon(FNiagaraEditorStyle::GetStyleSetName(), "NiagaraEditor.SimulationOptions")
+				);
+			}
+			ToolbarBuilder.EndSection();
 		}
 	};
 
@@ -645,9 +752,9 @@ void FNiagaraSystemToolkit::GetSequencerAddMenuContent(FMenuBuilder& MenuBuilder
 	MenuBuilder.AddSubMenu(
 		LOCTEXT("EmittersLabel", "Emitters..."),
 		LOCTEXT("EmittersToolTip", "Add an existing emitter..."),
-		FNewMenuDelegate::CreateLambda([&](FMenuBuilder& MenuBuilder)
+		FNewMenuDelegate::CreateLambda([&](FMenuBuilder& InMenuBuilder)
 		{
-			MenuBuilder.AddWidget(CreateAddEmitterMenuContent(), FText());
+			InMenuBuilder.AddWidget(CreateAddEmitterMenuContent(), FText());
 		}));
 }
 
@@ -681,7 +788,15 @@ TSharedRef<SWidget> FNiagaraSystemToolkit::GenerateCompileMenuContent()
 		FCanExecuteAction(),
 		FIsActionChecked::CreateStatic(&FNiagaraSystemToolkit::IsAutoCompileEnabled));
 
-	MenuBuilder.AddMenuEntry(LOCTEXT("AutoCompile", "Automatically compile when graph changes"), FText(), FSlateIcon(), Action, NAME_None, EUserInterfaceActionType::ToggleButton);
+	FUIAction FullRebuildAction(
+		FExecuteAction::CreateRaw(this, &FNiagaraSystemToolkit::CompileSystem, true));
+
+	MenuBuilder.AddMenuEntry(LOCTEXT("FullRebuild", "Full Rebuild"),
+		LOCTEXT("FullRebuildTooltip", "Triggers a full rebuild of this system, ignoring the change tracking."),
+		FSlateIcon(FNiagaraEditorStyle::GetStyleSetName(), "Niagara.CompileStatus.Unknown"),
+		FullRebuildAction, NAME_None, EUserInterfaceActionType::Button);
+	MenuBuilder.AddMenuEntry(LOCTEXT("AutoCompile", "Automatically compile when graph changes"),
+		FText(), FSlateIcon(), Action, NAME_None, EUserInterfaceActionType::ToggleButton);
 
 	return MenuBuilder.MakeWidget();
 }
@@ -712,9 +827,14 @@ FText FNiagaraSystemToolkit::GetCompileStatusTooltip() const
 }
 
 
-void FNiagaraSystemToolkit::CompileSystem(bool bForce)
+void FNiagaraSystemToolkit::CompileSystem(bool bFullRebuild)
 {
-	SystemViewModel->CompileSystem(bForce);
+	SystemViewModel->CompileSystem(bFullRebuild);
+}
+
+TSharedPtr<FNiagaraSystemViewModel> FNiagaraSystemToolkit::GetSystemViewModel()
+{
+	return SystemViewModel;
 }
 
 void FNiagaraSystemToolkit::OnToggleBounds()
@@ -800,19 +920,30 @@ void FNiagaraSystemToolkit::UpdateOriginalEmitter()
 	TSharedPtr<FNiagaraEmitterViewModel> EditableEmitterViewModel = SystemViewModel->GetEmitterHandleViewModels()[0]->GetEmitterViewModel();
 	UNiagaraEmitter* EditableEmitter = EditableEmitterViewModel->GetEmitter();
 
+	TArray<UNiagaraScript*> AllScripts;
+	EditableEmitter->GetScripts(AllScripts, true);
+	for (UNiagaraScript* Script : AllScripts)
+	{
+		checkSlow(Script->AreScriptAndSourceSynchronized());
+	}
+	checkSlow(EditableEmitter->AreAllScriptAndSourcesSynchronized());
+
 	// overwrite the original script in place by constructing a new one with the same name
 	Emitter = (UNiagaraEmitter*)StaticDuplicateObject(EditableEmitter, Emitter->GetOuter(), Emitter->GetFName(),
 		RF_AllFlags,
 		Emitter->GetClass());
 
+	checkSlow (UNiagaraEmitter::GetForceCompileOnLoad() || Emitter->GetChangeId() == EditableEmitter->GetChangeId());
+
 	// Restore RF_Standalone on the original emitter, as it had been removed from the preview emitter so that it could be GC'd.
 	Emitter->SetFlags(RF_Standalone);
 
-
 	TArray<UNiagaraEmitter*> AffectedEmitters;
 	AffectedEmitters.Add(Emitter);
 	UpdateExistingEmitters();
 
+	checkSlow(UNiagaraEmitter::GetForceCompileOnLoad() || Emitter->GetChangeId() == EditableEmitter->GetChangeId());
+
 	GWarn->EndSlowTask();
 }
 
@@ -841,7 +972,6 @@ void FNiagaraSystemToolkit::UpdateExistingEmitters()
 					UNiagaraComponent* Component = *ComponentIterator;
 					if (Component->GetAsset() == LoadedSystem)
 					{
-						Component->SynchronizeWithSourceSystem();
 						Component->ReinitializeSystem();
 					}
 				}
@@ -869,7 +999,7 @@ void FNiagaraSystemToolkit::SaveAsset_Execute()
 		UE_LOG(LogNiagaraEditor, Log, TEXT("Saving and Compiling NiagaraEmitter %s"), *GetEditingObjects()[0]->GetName());
 		UpdateOriginalEmitter();
 	}
-
+	SystemViewModel->OnPreSave();
 	FAssetEditorToolkit::SaveAsset_Execute();
 }
 
@@ -880,12 +1010,35 @@ void FNiagaraSystemToolkit::SaveAssetAs_Execute()
 		UE_LOG(LogNiagaraEditor, Log, TEXT("Saving and Compiling NiagaraEmitter %s"), *GetEditingObjects()[0]->GetName());
 		UpdateOriginalEmitter();
 	}
-
+	SystemViewModel->OnPreSave();
 	FAssetEditorToolkit::SaveAssetAs_Execute();
 }
 
 bool FNiagaraSystemToolkit::OnRequestClose()
 {
+	if (GbLogNiagaraSystemChanges > 0)
+	{
+		FString ExportText;
+		SystemViewModel->DumpToText(ExportText);
+		FString FilePath;
+
+		if (SystemToolkitMode == ESystemToolkitMode::System)
+		{
+			FilePath = System->GetOutermost()->FileName.ToString();
+		}
+		else if (SystemToolkitMode == ESystemToolkitMode::Emitter)
+		{
+			FilePath = Emitter->GetOutermost()->FileName.ToString();
+		}
+
+		FString PathPart, FilenamePart, ExtensionPart;
+		FPaths::Split(FilePath, PathPart, FilenamePart, ExtensionPart);
+
+		FNiagaraEditorUtilities::WriteTextFileToDisk(FPaths::ProjectLogDir(), FilenamePart + TEXT(".onClose.txt"), ExportText, true);
+	}
+
+	SystemViewModel->OnPreClose();
+
 	if (SystemToolkitMode == ESystemToolkitMode::Emitter)
 	{
 		TSharedPtr<FNiagaraEmitterViewModel> EmitterViewModel = SystemViewModel->GetEmitterHandleViewModels()[0]->GetEmitterViewModel();
@@ -917,7 +1070,7 @@ bool FNiagaraSystemToolkit::OnRequestClose()
 		}
 		return true;
 	}
-
+	
 	return FAssetEditorToolkit::OnRequestClose();
 }
 
@@ -930,12 +1083,12 @@ void FNiagaraSystemToolkit::EmitterAssetSelected(const FAssetData& AssetData)
 void FNiagaraSystemToolkit::ToggleCompileEnabled()
 {
 	UNiagaraEditorSettings* Settings = GetMutableDefault<UNiagaraEditorSettings>();
-	Settings->bAutoCompile = !Settings->bAutoCompile;
+	Settings->SetAutoCompile(!Settings->GetAutoCompile());
 }
 
 bool FNiagaraSystemToolkit::IsAutoCompileEnabled()
 {
-	return GetDefault<UNiagaraEditorSettings>()->bAutoCompile;
+	return GetDefault<UNiagaraEditorSettings>()->GetAutoCompile();
 }
 
 void FNiagaraSystemToolkit::OnApply()
@@ -954,4 +1107,26 @@ bool FNiagaraSystemToolkit::OnApplyEnabled() const
 	return false;
 }
 
+void FNiagaraSystemToolkit::OnPinnedCurvesChanged()
+{
+	TabManager->InvokeTab(CurveEditorTabID);
+}
+
+void FNiagaraSystemToolkit::OnRefresh()
+{
+	if (ParameterMapView.IsValid())
+	{
+		TArray<TSharedPtr<FNiagaraEmitterHandleViewModel>> EmitterHandlesToDisplay;
+		EmitterHandlesToDisplay.Append(SystemViewModel->GetPinnedEmitterHandles());
+		TArray<TSharedRef<FNiagaraEmitterHandleViewModel>> SelectedEmitterHandles;
+		SystemViewModel->GetSelectedEmitterHandles(SelectedEmitterHandles);
+		for (auto Handle : SelectedEmitterHandles)
+		{
+			EmitterHandlesToDisplay.AddUnique(Handle);
+		}
+
+		ParameterMapView->RefreshEmitterHandles(EmitterHandlesToDisplay);
+	}
+}
+
 #undef LOCTEXT_NAMESPACE
diff --git a/Engine/Plugins/FX/Niagara/Source/NiagaraEditor/Private/Toolkits/NiagaraSystemToolkit.h b/Engine/Plugins/FX/Niagara/Source/NiagaraEditor/Private/Toolkits/NiagaraSystemToolkit.h
--- a/Engine/Plugins/FX/Niagara/Source/NiagaraEditor/Private/Toolkits/NiagaraSystemToolkit.h
+++ b/Engine/Plugins/FX/Niagara/Source/NiagaraEditor/Private/Toolkits/NiagaraSystemToolkit.h
@@ -63,7 +63,9 @@ public:
 	FText GetCompileStatusTooltip() const;
 
 	/** Compiles the system script. */
-	void CompileSystem(bool bForce);
+	void CompileSystem(bool bFullRebuild);
+
+	TSharedPtr<FNiagaraSystemViewModel> GetSystemViewModel();
 
 protected:
 	void OnToggleBounds();
@@ -78,7 +80,7 @@ protected:
 	virtual void SaveAsset_Execute() override;
 	virtual void SaveAssetAs_Execute() override;
 	virtual bool OnRequestClose() override;
-
+	
 private:
 	void InitializeInternal(const EToolkitMode::Type Mode, const TSharedPtr<class IToolkitHost>& InitToolkitHost);
 
@@ -90,6 +92,7 @@ private:
 	TSharedRef<SDockTab> SpawnTab_Sequencer(const FSpawnTabArgs& Args);
 	TSharedRef<SDockTab> SpawnTab_SystemScript(const FSpawnTabArgs& Args);
 	TSharedRef<SDockTab> SpawnTab_SystemDetails(const FSpawnTabArgs& Args);
+	TSharedRef<SDockTab> SpawnTab_SystemParameters(const FSpawnTabArgs& Args);
 	TSharedRef<SDockTab> SpawnTab_SelectedEmitterStack(const FSpawnTabArgs& Args);
 	TSharedRef<SDockTab> SpawnTab_SelectedEmitterGraph(const FSpawnTabArgs& Args);
 	TSharedRef<SDockTab> SpawnTab_DebugSpreadsheet(const FSpawnTabArgs& Args);
@@ -104,16 +107,19 @@ private:
 
 	void GetSequencerAddMenuContent(FMenuBuilder& MenuBuilder, TSharedRef<ISequencer> Sequencer);
 	TSharedRef<SWidget> CreateAddEmitterMenuContent();
-	static TSharedRef<SWidget> GenerateCompileMenuContent();
+	TSharedRef<SWidget> GenerateCompileMenuContent();
 
 	void EmitterAssetSelected(const FAssetData& AssetData);
 
 	static void ToggleCompileEnabled();
 	static bool IsAutoCompileEnabled();
-
+	
 	void OnApply();
 	bool OnApplyEnabled() const;
 
+	void OnPinnedCurvesChanged();
+	void OnRefresh();
+
 private:
 	TSharedRef<SWidget> GenerateBoundsMenuContent(TSharedRef<FUICommandList> InCommandList);
 	void OnSaveThumbnailImage();
@@ -134,11 +140,14 @@ private:
 	/** The command list for this editor */
 	TSharedPtr<FUICommandList> EditorCommands;
 
+	TSharedPtr<class SNiagaraParameterMapView> ParameterMapView;
+
 	static const FName ViewportTabID;
 	static const FName CurveEditorTabID;
 	static const FName SequencerTabID;
 	static const FName SystemScriptTabID;
 	static const FName SystemDetailsTabID;
+	static const FName SystemParametersTabID;
 	static const FName SelectedEmitterStackTabID;
 	static const FName SelectedEmitterGraphTabID;
 	static const FName DebugSpreadsheetTabID;
diff --git a/Engine/Plugins/FX/Niagara/Source/NiagaraEditor/Private/TypeEditorUtilities/NiagaraColorTypeEditorUtilities.cpp b/Engine/Plugins/FX/Niagara/Source/NiagaraEditor/Private/TypeEditorUtilities/NiagaraColorTypeEditorUtilities.cpp
--- a/Engine/Plugins/FX/Niagara/Source/NiagaraEditor/Private/TypeEditorUtilities/NiagaraColorTypeEditorUtilities.cpp
+++ b/Engine/Plugins/FX/Niagara/Source/NiagaraEditor/Private/TypeEditorUtilities/NiagaraColorTypeEditorUtilities.cpp
@@ -19,6 +19,10 @@ public:
 
 	void Construct(const FArguments& InArgs)
 	{
+		SNiagaraParameterEditor::Construct(SNiagaraParameterEditor::FArguments()
+			.MinimumDesiredWidth(3 * DefaultInputSize)
+			.MaximumDesiredWidth(3 * DefaultInputSize));
+
 		ChildSlot
 		[
 			SNew(SGridPanel)
diff --git a/Engine/Plugins/FX/Niagara/Source/NiagaraEditor/Private/TypeEditorUtilities/NiagaraDataInterfaceCurveTypeEditorUtilities.cpp b/Engine/Plugins/FX/Niagara/Source/NiagaraEditor/Private/TypeEditorUtilities/NiagaraDataInterfaceCurveTypeEditorUtilities.cpp
--- a/Engine/Plugins/FX/Niagara/Source/NiagaraEditor/Private/TypeEditorUtilities/NiagaraDataInterfaceCurveTypeEditorUtilities.cpp
+++ b/Engine/Plugins/FX/Niagara/Source/NiagaraEditor/Private/TypeEditorUtilities/NiagaraDataInterfaceCurveTypeEditorUtilities.cpp
@@ -36,7 +36,7 @@ TSharedPtr<SWidget> FNiagaraDataInterfaceCurveTypeEditorUtilitiesBase::CreateDat
 		return SNew(SComboButton)
 			.ButtonStyle(FEditorStyle::Get(), "HoverHintOnly")
 			.ForegroundColor(FSlateColor::UseForeground())
-			.OnGetMenuContent(this, &FNiagaraDataInterfaceCurveTypeEditorUtilitiesBase::GetImportMenuContent, MakeWeakObjectPtr(CurveDataInterface), DataInterfaceChangedHandler)
+			.OnGetMenuContent_Raw(this, &FNiagaraDataInterfaceCurveTypeEditorUtilitiesBase::GetImportMenuContent, MakeWeakObjectPtr(CurveDataInterface), DataInterfaceChangedHandler)
 			.ButtonContent()
 			[
 				SNew(STextBlock)
@@ -51,7 +51,7 @@ TSharedRef<SWidget> FNiagaraDataInterfaceCurveTypeEditorUtilitiesBase::GetImport
 {
 	FAssetPickerConfig AssetPickerConfig;
 	{
-		AssetPickerConfig.OnAssetSelected = FOnAssetSelected::CreateSP(this, &FNiagaraDataInterfaceCurveTypeEditorUtilitiesBase::CurveAssetSelected, CurveDataInterface, DataInterfaceChangedHandler);
+		AssetPickerConfig.OnAssetSelected = FOnAssetSelected::CreateRaw(this, &FNiagaraDataInterfaceCurveTypeEditorUtilitiesBase::CurveAssetSelected, CurveDataInterface, DataInterfaceChangedHandler);
 		AssetPickerConfig.bAllowNullSelection = false;
 		AssetPickerConfig.InitialAssetViewType = EAssetViewType::List;
 		AssetPickerConfig.Filter.ClassNames.Add(GetSupportedAssetClassName());
diff --git a/Engine/Plugins/FX/Niagara/Source/NiagaraEditor/Private/TypeEditorUtilities/NiagaraEnumTypeEditorUtilities.cpp b/Engine/Plugins/FX/Niagara/Source/NiagaraEditor/Private/TypeEditorUtilities/NiagaraEnumTypeEditorUtilities.cpp
--- a/Engine/Plugins/FX/Niagara/Source/NiagaraEditor/Private/TypeEditorUtilities/NiagaraEnumTypeEditorUtilities.cpp
+++ b/Engine/Plugins/FX/Niagara/Source/NiagaraEditor/Private/TypeEditorUtilities/NiagaraEnumTypeEditorUtilities.cpp
@@ -112,6 +112,10 @@ public:
 
 	void Construct(const FArguments& InArgs, const UEnum* Enum)
 	{
+		SNiagaraParameterEditor::Construct(SNiagaraParameterEditor::FArguments()
+			.MinimumDesiredWidth(DefaultInputSize)
+			.MaximumDesiredWidth(2 * DefaultInputSize));
+
 		ChildSlot
 		[
 			SNew(SEnumCombobox, Enum)
@@ -209,4 +213,35 @@ bool FNiagaraEditorEnumTypeUtilities::SetValueFromPinDefaultString(const FString
 		return true;
 	}
 	return false;
+}
+
+bool FNiagaraEditorEnumTypeUtilities::CanSetValueFromDisplayName() const
+{
+	return true;
+}
+
+bool FNiagaraEditorEnumTypeUtilities::SetValueFromDisplayName(const FText& TextValue, FNiagaraVariable& Variable) const
+{
+	const UEnum* Enum = Variable.GetType().GetEnum();
+	checkf(Enum != nullptr, TEXT("Variable is not an enum type."));
+
+	int32 FoundIndex = INDEX_NONE;
+	for (int32 i = 0; i < Enum->NumEnums(); i++)
+	{
+		if (TextValue.CompareTo(Enum->GetDisplayNameTextByIndex(i)) == 0)
+		{
+			FoundIndex = i;
+			break;
+		}
+	}
+
+	if (FoundIndex != INDEX_NONE)
+	{
+		FNiagaraInt32 EnumValue;
+		EnumValue.Value = (int32)Enum->GetValueByIndex(FoundIndex);
+		Variable.AllocateData();
+		Variable.SetValue<FNiagaraInt32>(EnumValue);
+		return true;
+	}
+	return false;
 }
\ No newline at end of file
diff --git a/Engine/Plugins/FX/Niagara/Source/NiagaraEditor/Private/TypeEditorUtilities/NiagaraEnumTypeEditorUtilities.h b/Engine/Plugins/FX/Niagara/Source/NiagaraEditor/Private/TypeEditorUtilities/NiagaraEnumTypeEditorUtilities.h
--- a/Engine/Plugins/FX/Niagara/Source/NiagaraEditor/Private/TypeEditorUtilities/NiagaraEnumTypeEditorUtilities.h
+++ b/Engine/Plugins/FX/Niagara/Source/NiagaraEditor/Private/TypeEditorUtilities/NiagaraEnumTypeEditorUtilities.h
@@ -17,4 +17,6 @@ public:
 	virtual bool CanHandlePinDefaults() const override;
 	virtual FString GetPinDefaultStringFromValue(const FNiagaraVariable& AllocatedVariable) const override;
 	virtual bool SetValueFromPinDefaultString(const FString& StringValue, FNiagaraVariable& Variable) const override;
+	virtual bool CanSetValueFromDisplayName() const override;
+	virtual bool SetValueFromDisplayName(const FText& TextValue, FNiagaraVariable& Variable) const override;
 };
\ No newline at end of file
diff --git a/Engine/Plugins/FX/Niagara/Source/NiagaraEditor/Private/TypeEditorUtilities/NiagaraFloatTypeEditorUtilities.cpp b/Engine/Plugins/FX/Niagara/Source/NiagaraEditor/Private/TypeEditorUtilities/NiagaraFloatTypeEditorUtilities.cpp
--- a/Engine/Plugins/FX/Niagara/Source/NiagaraEditor/Private/TypeEditorUtilities/NiagaraFloatTypeEditorUtilities.cpp
+++ b/Engine/Plugins/FX/Niagara/Source/NiagaraEditor/Private/TypeEditorUtilities/NiagaraFloatTypeEditorUtilities.cpp
@@ -15,27 +15,25 @@ public:
 
 	void Construct(const FArguments& InArgs)
 	{
+		SNiagaraParameterEditor::Construct(SNiagaraParameterEditor::FArguments()
+			.MinimumDesiredWidth(DefaultInputSize)
+			.MaximumDesiredWidth(DefaultInputSize));
+
 		ChildSlot
 		[
-			SNew(SHorizontalBox)
-			+ SHorizontalBox::Slot()
-			.AutoWidth()
-			[
-				SNew(SSpinBox<float>)
-				.Style(FNiagaraEditorStyle::Get(), "NiagaraEditor.ParameterSpinBox")
-				.Font(FNiagaraEditorStyle::Get().GetFontStyle("NiagaraEditor.ParameterFont"))
-				.MinValue(TOptional<float>())
-				.MaxValue(TOptional<float>())
-				.MaxSliderValue(TOptional<float>())
-				.MinSliderValue(TOptional<float>())
-				.Delta(0.0f)
-				.Value(this, &SNiagaraFloatParameterEditor::GetValue)
-				.OnValueChanged(this, &SNiagaraFloatParameterEditor::ValueChanged)
-				.OnValueCommitted(this, &SNiagaraFloatParameterEditor::ValueCommitted)
-				.OnBeginSliderMovement(this, &SNiagaraFloatParameterEditor::BeginSliderMovement)
-				.OnEndSliderMovement(this, &SNiagaraFloatParameterEditor::EndSliderMovement)
-				.MinDesiredWidth(100)
-			]
+			SNew(SSpinBox<float>)
+			.Style(FNiagaraEditorStyle::Get(), "NiagaraEditor.ParameterSpinBox")
+			.Font(FNiagaraEditorStyle::Get().GetFontStyle("NiagaraEditor.ParameterFont"))
+			.MinValue(TOptional<float>())
+			.MaxValue(TOptional<float>())
+			.MaxSliderValue(TOptional<float>())
+			.MinSliderValue(TOptional<float>())
+			.Delta(0.0f)
+			.Value(this, &SNiagaraFloatParameterEditor::GetValue)
+			.OnValueChanged(this, &SNiagaraFloatParameterEditor::ValueChanged)
+			.OnValueCommitted(this, &SNiagaraFloatParameterEditor::ValueCommitted)
+			.OnBeginSliderMovement(this, &SNiagaraFloatParameterEditor::BeginSliderMovement)
+			.OnEndSliderMovement(this, &SNiagaraFloatParameterEditor::EndSliderMovement)
 		];
 	}
 
diff --git a/Engine/Plugins/FX/Niagara/Source/NiagaraEditor/Private/TypeEditorUtilities/NiagaraIntegerTypeEditorUtilities.cpp b/Engine/Plugins/FX/Niagara/Source/NiagaraEditor/Private/TypeEditorUtilities/NiagaraIntegerTypeEditorUtilities.cpp
--- a/Engine/Plugins/FX/Niagara/Source/NiagaraEditor/Private/TypeEditorUtilities/NiagaraIntegerTypeEditorUtilities.cpp
+++ b/Engine/Plugins/FX/Niagara/Source/NiagaraEditor/Private/TypeEditorUtilities/NiagaraIntegerTypeEditorUtilities.cpp
@@ -16,27 +16,25 @@ public:
 
 	void Construct(const FArguments& InArgs)
 	{
+		SNiagaraParameterEditor::Construct(SNiagaraParameterEditor::FArguments()
+			.MinimumDesiredWidth(DefaultInputSize)
+			.MaximumDesiredWidth(DefaultInputSize));
+
 		ChildSlot
 		[
-			SNew(SHorizontalBox)
-			+ SHorizontalBox::Slot()
-			.AutoWidth()
-			[
-				SNew(SSpinBox<int32>)
-				.Style(FNiagaraEditorStyle::Get(), "NiagaraEditor.ParameterSpinBox")
-				.Font(FNiagaraEditorStyle::Get().GetFontStyle("NiagaraEditor.ParameterFont"))
-				.MinValue(TOptional<int32>())
-				.MaxValue(TOptional<int32>())
-				.MaxSliderValue(TOptional<int32>())
-				.MinSliderValue(TOptional<int32>())
-				.Delta(0.001f)
-				.Value(this, &SNiagaraIntegerParameterEditor::GetValue)
-				.OnValueChanged(this, &SNiagaraIntegerParameterEditor::ValueChanged)
-				.OnValueCommitted(this, &SNiagaraIntegerParameterEditor::ValueCommitted)
-				.OnBeginSliderMovement(this, &SNiagaraIntegerParameterEditor::BeginSliderMovement)
-				.OnEndSliderMovement(this, &SNiagaraIntegerParameterEditor::EndSliderMovement)
-				.MinDesiredWidth(100)
-			]
+			SNew(SSpinBox<int32>)
+			.Style(FNiagaraEditorStyle::Get(), "NiagaraEditor.ParameterSpinBox")
+			.Font(FNiagaraEditorStyle::Get().GetFontStyle("NiagaraEditor.ParameterFont"))
+			.MinValue(TOptional<int32>())
+			.MaxValue(TOptional<int32>())
+			.MaxSliderValue(TOptional<int32>())
+			.MinSliderValue(TOptional<int32>())
+			.Delta(0.001f)
+			.Value(this, &SNiagaraIntegerParameterEditor::GetValue)
+			.OnValueChanged(this, &SNiagaraIntegerParameterEditor::ValueChanged)
+			.OnValueCommitted(this, &SNiagaraIntegerParameterEditor::ValueCommitted)
+			.OnBeginSliderMovement(this, &SNiagaraIntegerParameterEditor::BeginSliderMovement)
+			.OnEndSliderMovement(this, &SNiagaraIntegerParameterEditor::EndSliderMovement)
 		];
 	}
 
diff --git a/Engine/Plugins/FX/Niagara/Source/NiagaraEditor/Private/TypeEditorUtilities/NiagaraVectorTypeEditorUtilities.cpp b/Engine/Plugins/FX/Niagara/Source/NiagaraEditor/Private/TypeEditorUtilities/NiagaraVectorTypeEditorUtilities.cpp
--- a/Engine/Plugins/FX/Niagara/Source/NiagaraEditor/Private/TypeEditorUtilities/NiagaraVectorTypeEditorUtilities.cpp
+++ b/Engine/Plugins/FX/Niagara/Source/NiagaraEditor/Private/TypeEditorUtilities/NiagaraVectorTypeEditorUtilities.cpp
@@ -17,6 +17,10 @@ public:
 
 	void Construct(const FArguments& InArgs)
 	{
+		SNiagaraParameterEditor::Construct(SNiagaraParameterEditor::FArguments()
+			.MinimumDesiredWidth(DefaultInputSize * InArgs._ComponentCount)
+			.MaximumDesiredWidth(DefaultInputSize * InArgs._ComponentCount));
+
 		ComponentLabels.Add(NSLOCTEXT("VectorParameterEditor", "XLabel", "X"));
 		ComponentLabels.Add(NSLOCTEXT("VectorParameterEditor", "YLabel", "Y"));
 		ComponentLabels.Add(NSLOCTEXT("VectorParameterEditor", "ZLabel", "Z"));
@@ -310,3 +314,110 @@ bool FNiagaraEditorVector4TypeUtilities::SetValueFromPinDefaultString(const FStr
 	}
 	return false;
 }
+
+
+class SNiagaraQuatParameterEditor : public SNiagaraVectorParameterEditorBase
+{
+public:
+	SLATE_BEGIN_ARGS(SNiagaraQuatParameterEditor) { }
+	SLATE_END_ARGS();
+
+	void Construct(const FArguments& InArgs)
+	{
+		SNiagaraVectorParameterEditorBase::Construct(
+			SNiagaraVectorParameterEditorBase::FArguments()
+			.ComponentCount(4));
+	}
+
+	virtual void UpdateInternalValueFromStruct(TSharedRef<FStructOnScope> Struct) override
+	{
+		checkf(Struct->GetStruct() == FNiagaraTypeDefinition::GetQuatStruct(), TEXT("Struct type not supported."));
+		VectorValue = *((FQuat*)Struct->GetStructMemory());
+	}
+
+	virtual void UpdateStructFromInternalValue(TSharedRef<FStructOnScope> Struct) override
+	{
+		checkf(Struct->GetStruct() == FNiagaraTypeDefinition::GetQuatStruct(), TEXT("Struct type not supported."));
+		*((FQuat*)Struct->GetStructMemory()) = VectorValue;
+	}
+
+protected:
+	virtual float GetValue(int32 Index) const override
+	{
+		switch (Index)
+		{
+		case 0:
+			return VectorValue.X;
+			break;
+		case 1:
+			return VectorValue.Y;
+			break;
+		case 2:
+			return VectorValue.Z;
+			break;
+		case 3:
+			return VectorValue.W;
+			break;
+		}
+		return 0.0f;
+	}
+
+	virtual void SetValue(int32 Index, float Value) override
+	{
+		switch (Index)
+		{
+		case 0:
+			VectorValue.X = Value;
+			break;
+		case 1:
+			VectorValue.Y = Value;
+			break;
+		case 2:
+			VectorValue.Z = Value;
+			break;
+		case 3:
+			VectorValue.W = Value;
+			break;
+		}
+	}
+
+private:
+	FQuat VectorValue;
+};
+
+TSharedPtr<SNiagaraParameterEditor> FNiagaraEditorQuatTypeUtilities::CreateParameterEditor(const FNiagaraTypeDefinition& ParameterType) const
+{
+	return SNew(SNiagaraQuatParameterEditor);
+}
+
+bool FNiagaraEditorQuatTypeUtilities::CanHandlePinDefaults() const
+{
+	return true;
+}
+
+FString FNiagaraEditorQuatTypeUtilities::GetPinDefaultStringFromValue(const FNiagaraVariable& AllocatedVariable) const
+{
+	checkf(AllocatedVariable.IsDataAllocated(), TEXT("Can not generate a default value string for an unallocated variable."));
+	// NOTE: We can not use ToString() here since the vector pin control doesn't use the standard 'X=0,Y=0,Z=0,W=0' syntax.
+	FQuat Value = AllocatedVariable.GetValue<FQuat>();
+	return FString::Printf(TEXT("%3.3f,%3.3f,%3.3f,%3.3f"), Value.X, Value.Y, Value.Z, Value.W);
+}
+
+bool FNiagaraEditorQuatTypeUtilities::SetValueFromPinDefaultString(const FString& StringValue, FNiagaraVariable& Variable) const
+{
+	// NOTE: We can not use InitFromString() here since the vector pin control doesn't use the standard 'X=0,Y=0,Z=0,W=0' syntax.
+	FVector4 Value;
+	if (FDefaultValueHelper::ParseVector4(StringValue, Value))
+	{
+		FQuat Quat(Value.X, Value.Y, Value.Z, Value.W);
+		Variable.SetValue<FQuat>(Quat);
+		return true;
+	}
+	return false;
+}
+
+void FNiagaraEditorQuatTypeUtilities::UpdateVariableWithDefaultValue(FNiagaraVariable& Variable) const
+{
+	checkf(Variable.GetType().GetStruct() == FNiagaraTypeDefinition::GetQuatStruct(), TEXT("Struct type not supported."));
+	Variable.SetValue<FQuat>(FQuat(0, 0, 0, 1));
+}
diff --git a/Engine/Plugins/FX/Niagara/Source/NiagaraEditor/Private/TypeEditorUtilities/NiagaraVectorTypeEditorUtilities.h b/Engine/Plugins/FX/Niagara/Source/NiagaraEditor/Private/TypeEditorUtilities/NiagaraVectorTypeEditorUtilities.h
--- a/Engine/Plugins/FX/Niagara/Source/NiagaraEditor/Private/TypeEditorUtilities/NiagaraVectorTypeEditorUtilities.h
+++ b/Engine/Plugins/FX/Niagara/Source/NiagaraEditor/Private/TypeEditorUtilities/NiagaraVectorTypeEditorUtilities.h
@@ -40,4 +40,18 @@ public:
 	virtual bool CanHandlePinDefaults() const override;
 	virtual FString GetPinDefaultStringFromValue(const FNiagaraVariable& AllocatedVariable) const override;
 	virtual bool SetValueFromPinDefaultString(const FString& StringValue, FNiagaraVariable& Variable) const override;
+};
+
+/** Niagara editor utilities for the FVector4 type. */
+class FNiagaraEditorQuatTypeUtilities : public FNiagaraEditorTypeUtilities
+{
+public:
+	//~ INiagaraEditorTypeUtilities interface.
+	virtual bool CanProvideDefaultValue() const override { return true; }
+	virtual void UpdateVariableWithDefaultValue(FNiagaraVariable& Variable) const override;
+	virtual bool CanCreateParameterEditor() const override { return true; }
+	virtual TSharedPtr<SNiagaraParameterEditor> CreateParameterEditor(const FNiagaraTypeDefinition& ParameterType) const override;
+	virtual bool CanHandlePinDefaults() const override;
+	virtual FString GetPinDefaultStringFromValue(const FNiagaraVariable& AllocatedVariable) const override;
+	virtual bool SetValueFromPinDefaultString(const FString& StringValue, FNiagaraVariable& Variable) const override;
 };
\ No newline at end of file
diff --git a/Engine/Plugins/FX/Niagara/Source/NiagaraEditor/Private/ViewModels/NiagaraEmitterHandleViewModel.cpp b/Engine/Plugins/FX/Niagara/Source/NiagaraEditor/Private/ViewModels/NiagaraEmitterHandleViewModel.cpp
--- a/Engine/Plugins/FX/Niagara/Source/NiagaraEditor/Private/ViewModels/NiagaraEmitterHandleViewModel.cpp
+++ b/Engine/Plugins/FX/Niagara/Source/NiagaraEditor/Private/ViewModels/NiagaraEmitterHandleViewModel.cpp
@@ -1,9 +1,9 @@
 // Copyright 1998-2018 Epic Games, Inc. All Rights Reserved.
 
-#include "NiagaraEmitterHandleViewModel.h"
+#include "ViewModels/NiagaraEmitterHandleViewModel.h"
 #include "NiagaraSystem.h"
 #include "NiagaraEmitterHandle.h"
-#include "NiagaraEmitterViewModel.h"
+#include "ViewModels/NiagaraEmitterViewModel.h"
 #include "NiagaraScriptViewModel.h"
 #include "NiagaraScriptGraphViewModel.h"
 #include "NiagaraObjectSelection.h"
@@ -27,8 +27,19 @@ FNiagaraEmitterHandleViewModel::FNiagaraEmitterHandleViewModel(FNiagaraEmitterHa
 {
 }
 
+void FNiagaraEmitterHandleViewModel::Cleanup()
+{
+	if (EmitterViewModel.IsValid())
+	{
+		EmitterViewModel->Cleanup();
+	}
+	EmitterViewModel.Reset();
+}
+
+
 FNiagaraEmitterHandleViewModel::~FNiagaraEmitterHandleViewModel()
 {
+	Cleanup();
 }
 
 bool FNiagaraEmitterHandleViewModel::Set(FNiagaraEmitterHandle* InEmitterHandle, TWeakPtr<FNiagaraEmitterInstance> InSimulation, UNiagaraSystem& InOwningSystem)
@@ -46,6 +57,7 @@ bool FNiagaraEmitterHandleViewModel::Set(FNiagaraEmitterHandle* InEmitterHandle,
 	{
 		EmitterProperties = InEmitterHandle->GetInstance();
 	}
+	check(EmitterViewModel.IsValid());
 	return EmitterViewModel->Set(EmitterProperties, InSimulation);
 }
 
@@ -57,6 +69,7 @@ void FNiagaraEmitterHandleViewModel::SetEmitterHandle(FNiagaraEmitterHandle* InE
 
 void FNiagaraEmitterHandleViewModel::SetSimulation(TWeakPtr<FNiagaraEmitterInstance> InSimulation)
 {
+	check(EmitterViewModel.IsValid());
 	EmitterViewModel->SetSimulation(InSimulation);
 }
 
@@ -77,6 +90,8 @@ FText FNiagaraEmitterHandleViewModel::GetIdText() const
 
 FText FNiagaraEmitterHandleViewModel::GetErrorText() const
 {
+	check(EmitterViewModel.IsValid());
+
 	switch (EmitterViewModel->GetLatestCompileStatus())
 	{
 	case ENiagaraScriptCompileStatus::NCS_Unknown:
@@ -91,6 +106,7 @@ FText FNiagaraEmitterHandleViewModel::GetErrorText() const
 
 FSlateColor FNiagaraEmitterHandleViewModel::GetErrorTextColor() const
 {
+	check(EmitterViewModel.IsValid());
 	switch (EmitterViewModel->GetLatestCompileStatus())
 	{
 	case ENiagaraScriptCompileStatus::NCS_Unknown:
@@ -105,6 +121,7 @@ FSlateColor FNiagaraEmitterHandleViewModel::GetErrorTextColor() const
 
 EVisibility FNiagaraEmitterHandleViewModel::GetErrorTextVisibility() const
 {
+	check(EmitterViewModel.IsValid());
 	return EmitterViewModel->GetLatestCompileStatus() != ENiagaraScriptCompileStatus::NCS_UpToDate ? EVisibility::Visible : EVisibility::Collapsed;
 }
 
@@ -200,15 +217,11 @@ FNiagaraEmitterHandle* FNiagaraEmitterHandleViewModel::GetEmitterHandle()
 	return EmitterHandle;
 }
 
-TSharedRef<FNiagaraEmitterViewModel> FNiagaraEmitterHandleViewModel::GetEmitterViewModel()
+TSharedPtr<FNiagaraEmitterViewModel> FNiagaraEmitterHandleViewModel::GetEmitterViewModel()
 {
 	return EmitterViewModel;
 }
 
-void FNiagaraEmitterHandleViewModel::CompileScripts(bool bForce)
-{
-	EmitterViewModel->CompileScripts(bForce);
-}
 
 void FNiagaraEmitterHandleViewModel::OpenSourceEmitter()
 {
diff --git a/Engine/Plugins/FX/Niagara/Source/NiagaraEditor/Private/ViewModels/NiagaraEmitterViewModel.cpp b/Engine/Plugins/FX/Niagara/Source/NiagaraEditor/Private/ViewModels/NiagaraEmitterViewModel.cpp
--- a/Engine/Plugins/FX/Niagara/Source/NiagaraEditor/Private/ViewModels/NiagaraEmitterViewModel.cpp
+++ b/Engine/Plugins/FX/Niagara/Source/NiagaraEditor/Private/ViewModels/NiagaraEmitterViewModel.cpp
@@ -1,6 +1,6 @@
 // Copyright 1998-2018 Epic Games, Inc. All Rights Reserved.
 
-#include "NiagaraEmitterViewModel.h"
+#include "ViewModels/NiagaraEmitterViewModel.h"
 #include "NiagaraEmitter.h"
 #include "NiagaraEmitterInstance.h"
 #include "NiagaraScriptSourceBase.h"
@@ -11,6 +11,9 @@
 #include "NiagaraScriptSource.h"
 #include "NiagaraEmitterEditorData.h"
 #include "NiagaraParameterEditMode.h"
+#include "NiagaraGraph.h"
+#include "NiagaraSystem.h"
+#include "NiagaraSystemInstance.h"
 
 #include "ScopedTransaction.h"
 
@@ -19,6 +22,13 @@
 template<> TMap<UNiagaraEmitter*, TArray<FNiagaraEmitterViewModel*>> TNiagaraViewModelManager<UNiagaraEmitter, FNiagaraEmitterViewModel>::ObjectsToViewModels{};
 
 const FText FNiagaraEmitterViewModel::StatsFormat = NSLOCTEXT("NiagaraEmitterViewModel", "StatsFormat", "{0} Particles | {1} ms | {2} MB | {3}");
+const FText FNiagaraEmitterViewModel::StatsParticleCountFormat = NSLOCTEXT("NiagaraEmitterViewModel", "StatsParticleCountFormat", "{0} Particles");
+
+namespace NiagaraCommands
+{
+	static FAutoConsoleVariable EmitterStatsFormat(TEXT("Niagara.EmitterStatsFormat"), 1, TEXT("0 shows the particles count, ms, mb and state. 1 shows particles count."));
+}
+
 const float Megabyte = 1024.0f * 1024.0f;
 
 FNiagaraEmitterViewModel::FNiagaraEmitterViewModel(UNiagaraEmitter* InEmitter, TWeakPtr<FNiagaraEmitterInstance> InSimulation)
@@ -28,16 +38,41 @@ FNiagaraEmitterViewModel::FNiagaraEmitterViewModel(UNiagaraEmitter* InEmitter, T
 	, bUpdatingSelectionInternally(false)
 	, LastEventScriptStatus(ENiagaraScriptCompileStatus::NCS_Unknown)
 {
-	ExecutionStateEnum = FindObject<UEnum>(ANY_PACKAGE, TEXT("ENiagaraExecutionState")); 
+	//UE_LOG(LogNiagaraEditor, Warning, TEXT("FNiagaraEmitterViewModel %p"), this);
 
-	if (Emitter.IsValid() && Emitter->GetEventHandlers().Num() != 0 && Emitter->GetEventHandlers()[0].Script && Emitter->GetEventHandlers()[0].Script->GetByteCode().Num() != 0)
+	ExecutionStateEnum = FindObject<UEnum>(ANY_PACKAGE, TEXT("ENiagaraExecutionState"));
+
+	if (Emitter.IsValid() && Emitter->GetEventHandlers().Num() != 0 && Emitter->GetEventHandlers()[0].Script && Emitter->GetEventHandlers()[0].Script->GetVMExecutableData().IsValid()
+		&& Emitter->GetEventHandlers()[0].Script->GetVMExecutableData().ByteCode.Num() != 0)
 	{
 		LastEventScriptStatus = ENiagaraScriptCompileStatus::NCS_UpToDate;
 	}
+	
+	if (Emitter.IsValid())
+	{
+		Emitter->OnEmitterVMCompiled().AddRaw(this, &FNiagaraEmitterViewModel::OnVMCompiled);
+	}
 
 	RegisteredHandle = RegisterViewModelWithMap(Emitter.Get(), this);
+
+	AddScriptEventHandlers();
 }
 
+void FNiagaraEmitterViewModel::Cleanup()
+{
+	if (Emitter.IsValid())
+	{
+		Emitter->OnEmitterVMCompiled().RemoveAll(this);
+	}
+
+	if (SharedScriptViewModel.IsValid())
+	{
+		SharedScriptViewModel->GetGraphViewModel()->GetSelection()->OnSelectedObjectsChanged().RemoveAll(this);
+		SharedScriptViewModel.Reset();
+	}
+
+	RemoveScriptEventHandlers();
+}
 
 bool FNiagaraEmitterViewModel::Set(UNiagaraEmitter* InEmitter, TWeakPtr<FNiagaraEmitterInstance> InSimulation)
 {
@@ -48,19 +83,36 @@ bool FNiagaraEmitterViewModel::Set(UNiagaraEmitter* InEmitter, TWeakPtr<FNiagara
 
 FNiagaraEmitterViewModel::~FNiagaraEmitterViewModel()
 {
-	SharedScriptViewModel->GetGraphViewModel()->GetSelection()->OnSelectedObjectsChanged().RemoveAll(this);
+	Cleanup();
 	UnregisterViewModelWithMap(RegisteredHandle);
+
+	//UE_LOG(LogNiagaraEditor, Warning, TEXT("Deleting Emitter view model %p"), this);
 }
 
 
 void FNiagaraEmitterViewModel::SetEmitter(UNiagaraEmitter* InEmitter)
 {
+	if (Emitter.IsValid())
+	{
+		Emitter->OnEmitterVMCompiled().RemoveAll(this);
+	}
+
 	UnregisterViewModelWithMap(RegisteredHandle);
 
+	RemoveScriptEventHandlers();
+
 	Emitter = InEmitter;
 
+	if (Emitter.IsValid())
+	{
+		Emitter->OnEmitterVMCompiled().AddRaw(this, &FNiagaraEmitterViewModel::OnVMCompiled);
+	}
+
+	AddScriptEventHandlers();
+
 	RegisteredHandle = RegisterViewModelWithMap(InEmitter, this);
 
+	check(SharedScriptViewModel.IsValid());
 	SharedScriptViewModel->SetScripts(InEmitter);
 
 	OnEmitterChanged().Broadcast();
@@ -73,38 +125,6 @@ void FNiagaraEmitterViewModel::SetSimulation(TWeakPtr<FNiagaraEmitterInstance> I
 }
 
 
-float FNiagaraEmitterViewModel::GetStartTime() const
-{
-	// TODO, reference values off the parameter store...
-	return 0.0f;
-}
-
-
-void FNiagaraEmitterViewModel::SetStartTime(float InStartTime)
-{
-	// TODO, reference values off the parameter store...
-}
-
-
-float FNiagaraEmitterViewModel::GetEndTime() const
-{
-	// TODO, reference values off the parameter store...
-	return 0.0f;
-}
-
-
-void FNiagaraEmitterViewModel::SetEndTime(float InEndTime)
-{
-	// TODO, reference values off the parameter store...
-}
-
-
-int32 FNiagaraEmitterViewModel::GetNumLoops() const
-{
-	// TODO, reference values off the parameter store...
-	return 0;
-}
-
 UNiagaraEmitter* FNiagaraEmitterViewModel::GetEmitter()
 {
 	return Emitter.Get();
@@ -118,7 +138,11 @@ FText FNiagaraEmitterViewModel::GetStatsText() const
 		TSharedPtr<FNiagaraEmitterInstance> SimInstance = Simulation.Pin();
 		if (SimInstance.IsValid())
 		{
-			if (!SimInstance->IsReadyToRun())
+			static const FNumberFormattingOptions FractionalFormatOptions = FNumberFormattingOptions()
+				.SetMinimumFractionalDigits(3)
+				.SetMaximumFractionalDigits(3);
+
+			if (!SimInstance->IsReadyToRun() || SimInstance->GetParentSystemInstance()->GetSystem()->HasOutstandingCompilationRequests())
 			{
 				return LOCTEXT("PendingCompile", "Compilation in progress...");
 			}
@@ -147,11 +171,18 @@ FText FNiagaraEmitterViewModel::GetStatsText() const
 					return LOCTEXT("DisabledSimulation", "Simulation is not enabled.");
 				}
 
-				return FText::Format(StatsFormat,
-					FText::AsNumber(SimInstance->GetNumParticles()),
-					FText::AsNumber(SimInstance->GetTotalCPUTime()),
-					FText::AsNumber(SimInstance->GetTotalBytesUsed() / Megabyte),
-					ExecutionStateEnum->GetDisplayNameTextByValue((int32)SimInstance->GetExecutionState()));
+				if (NiagaraCommands::EmitterStatsFormat->GetInt() == 1)
+				{
+					return FText::Format(StatsParticleCountFormat, FText::AsNumber(SimInstance->GetNumParticles()));
+				}
+				else
+				{
+					return FText::Format(StatsFormat,
+						FText::AsNumber(SimInstance->GetNumParticles()),
+						FText::AsNumber(SimInstance->GetTotalCPUTime(), &FractionalFormatOptions),
+						FText::AsNumber(SimInstance->GetTotalBytesUsed() / Megabyte, &FractionalFormatOptions),
+						ExecutionStateEnum->GetDisplayNameTextByValue((int32)SimInstance->GetExecutionState()));
+				}
 			}
 		}
 	}
@@ -165,7 +196,7 @@ FText FNiagaraEmitterViewModel::GetStatsText() const
 
 TSharedRef<FNiagaraScriptViewModel> FNiagaraEmitterViewModel::GetSharedScriptViewModel()
 {
-	return SharedScriptViewModel;
+	return SharedScriptViewModel.ToSharedRef();
 }
 
 const UNiagaraEmitterEditorData& FNiagaraEmitterViewModel::GetEditorData() const
@@ -193,36 +224,78 @@ UNiagaraEmitterEditorData& FNiagaraEmitterViewModel::GetOrCreateEditorData()
 	return *EditorData;
 }
 
-void FNiagaraEmitterViewModel::CompileScripts(bool bForce)
+void FNiagaraEmitterViewModel::OnVMCompiled(UNiagaraEmitter* InEmitter)
 {
-	if (Emitter.IsValid())
+	if (Emitter.IsValid() && InEmitter == Emitter.Get())
 	{
 		TArray<ENiagaraScriptCompileStatus> CompileStatuses;
 		TArray<FString> CompileErrors;
 		TArray<FString> CompilePaths;
-		TArray<UNiagaraScript*> Scripts;
 		TArray<TPair<ENiagaraScriptUsage, int32> > Usages;
-		Emitter->CompileScripts(CompileStatuses, CompileErrors, CompilePaths, Scripts, bForce);
 
 		ENiagaraScriptCompileStatus AggregateStatus = ENiagaraScriptCompileStatus::NCS_UpToDate;
 		FString AggregateErrors;
 
+		TArray<UNiagaraScript*> Scripts;
+		Emitter->GetScripts(Scripts, true);
+
+		int32 EventsFound = 0;
+		for (int32 i = 0; i < Scripts.Num(); i++)
+		{
+			UNiagaraScript* Script = Scripts[i];
+			if (Script != nullptr && Script->GetVMExecutableData().IsValid())
+			{
+				CompileStatuses.Add(Script->GetVMExecutableData().LastCompileStatus);
+				CompileErrors.Add(Script->GetVMExecutableData().ErrorMsg);
+				CompilePaths.Add(Script->GetPathName());
+
+				if (Script->GetUsage() == ENiagaraScriptUsage::ParticleEventScript)
+				{
+					Usages.Add(TPair<ENiagaraScriptUsage, int32>(Script->GetUsage(), EventsFound));
+					EventsFound++;
+				}
+				else
+				{
+					Usages.Add(TPair<ENiagaraScriptUsage, int32>(Script->GetUsage(), 0));
+				}
+			}
+			else
+			{
+				CompileStatuses.Add(ENiagaraScriptCompileStatus::NCS_Unknown);
+				CompileErrors.Add(TEXT("Invalid script pointer!"));
+				CompilePaths.Add(TEXT("Unknown..."));
+				Usages.Add(TPair<ENiagaraScriptUsage, int32>(ENiagaraScriptUsage::Function, 0));
+			}
+
+		}
+
 		for (int32 i = 0; i < CompileStatuses.Num(); i++)
 		{
 			AggregateStatus = FNiagaraEditorUtilities::UnionCompileStatus(AggregateStatus, CompileStatuses[i]);
 			AggregateErrors += CompilePaths[i] + TEXT(" ") + FNiagaraEditorUtilities::StatusToText(CompileStatuses[i]).ToString() + TEXT("\n");
 			AggregateErrors += CompileErrors[i] + TEXT("\n");
 		}
+		check(SharedScriptViewModel.IsValid());
 		SharedScriptViewModel->UpdateCompileStatus(AggregateStatus, AggregateErrors, CompileStatuses, CompileErrors, CompilePaths, Scripts);
 	}
 	OnScriptCompiled().Broadcast();
 }
 
-
-
 ENiagaraScriptCompileStatus FNiagaraEmitterViewModel::GetLatestCompileStatus()
 {
+	check(SharedScriptViewModel.IsValid());
 	ENiagaraScriptCompileStatus UnionStatus = SharedScriptViewModel->GetLatestCompileStatus();
+
+	if (Emitter->SimTarget == ENiagaraSimTarget::GPUComputeSim || Emitter->SimTarget == ENiagaraSimTarget::DynamicLoadBalancedSim)
+	{
+		if (UnionStatus != ENiagaraScriptCompileStatus::NCS_Dirty)
+		{
+			if (!Emitter->GetGPUComputeScript()->AreScriptAndSourceSynchronized())
+			{
+				UnionStatus = ENiagaraScriptCompileStatus::NCS_Dirty;
+			}
+		}
+	}
 	return UnionStatus;
 }
 
@@ -242,4 +315,84 @@ FNiagaraEmitterViewModel::FOnScriptCompiled& FNiagaraEmitterViewModel::OnScriptC
 	return OnScriptCompiledDelegate;
 }
 
+FNiagaraEmitterViewModel::FOnScriptGraphChanged& FNiagaraEmitterViewModel::OnScriptGraphChanged()
+{
+	return OnScriptGraphChangedDelegate;
+}
+
+FNiagaraEmitterViewModel::FOnScriptParameterStoreChanged& FNiagaraEmitterViewModel::OnScriptParameterStoreChanged()
+{
+	return OnScriptParameterStoreChangedDelegate;
+}
+
+void FNiagaraEmitterViewModel::AddScriptEventHandlers()
+{
+	if (Emitter.IsValid())
+	{
+		TArray<UNiagaraScript*> Scripts;
+		Emitter->GetScripts(Scripts, false);
+		for (UNiagaraScript* Script : Scripts)
+		{
+			UNiagaraScriptSource* ScriptSource = CastChecked<UNiagaraScriptSource>(Script->GetSource());
+			FDelegateHandle OnGraphChangedHandle = ScriptSource->NodeGraph->AddOnGraphChangedHandler(
+				FOnGraphChanged::FDelegate::CreateRaw<FNiagaraEmitterViewModel, const UNiagaraScript&>(this, &FNiagaraEmitterViewModel::ScriptGraphChanged, *Script));
+			FDelegateHandle OnGraphNeedRecompileHandle = ScriptSource->NodeGraph->AddOnGraphNeedsRecompileHandler(
+				FOnGraphChanged::FDelegate::CreateRaw<FNiagaraEmitterViewModel, const UNiagaraScript&>(this, &FNiagaraEmitterViewModel::ScriptGraphChanged, *Script));
+
+			ScriptToOnGraphChangedHandleMap.Add(FObjectKey(Script), OnGraphChangedHandle);
+			ScriptToRecompileHandleMap.Add(FObjectKey(Script), OnGraphNeedRecompileHandle);
+
+			FDelegateHandle OnParameterStoreChangedHandle = Script->RapidIterationParameters.AddOnChangedHandler(
+				FNiagaraParameterStore::FOnChanged::FDelegate::CreateRaw<FNiagaraEmitterViewModel, const FNiagaraParameterStore&, const UNiagaraScript&>(
+					this, &FNiagaraEmitterViewModel::ScriptParameterStoreChanged, Script->RapidIterationParameters, *Script));
+			ScriptToOnParameterStoreChangedHandleMap.Add(FObjectKey(Script), OnParameterStoreChangedHandle);
+		}
+	}
+}
+
+void FNiagaraEmitterViewModel::RemoveScriptEventHandlers()
+{
+	if (Emitter.IsValid())
+	{
+		TArray<UNiagaraScript*> Scripts;
+		Emitter->GetScripts(Scripts, false);
+		for (UNiagaraScript* Script : Scripts)
+		{
+			FDelegateHandle* OnGraphChangedHandle = ScriptToOnGraphChangedHandleMap.Find(FObjectKey(Script));
+			if (OnGraphChangedHandle != nullptr)
+			{
+				UNiagaraScriptSource* ScriptSource = CastChecked<UNiagaraScriptSource>(Script->GetSource());
+				ScriptSource->NodeGraph->RemoveOnGraphChangedHandler(*OnGraphChangedHandle);
+			}
+
+			FDelegateHandle* OnGraphRecompileHandle = ScriptToRecompileHandleMap.Find(FObjectKey(Script));
+			if (OnGraphRecompileHandle != nullptr)
+			{
+				UNiagaraScriptSource* ScriptSource = CastChecked<UNiagaraScriptSource>(Script->GetSource());
+				ScriptSource->NodeGraph->RemoveOnGraphNeedsRecompileHandler(*OnGraphRecompileHandle);
+			}
+
+			FDelegateHandle* OnParameterStoreChangedHandle = ScriptToOnParameterStoreChangedHandleMap.Find(FObjectKey(Script));
+			if (OnParameterStoreChangedHandle != nullptr)
+			{
+				Script->RapidIterationParameters.RemoveOnChangedHandler(*OnParameterStoreChangedHandle);
+			}
+		}
+	}
+
+	ScriptToOnGraphChangedHandleMap.Empty();
+	ScriptToRecompileHandleMap.Empty();
+	ScriptToOnParameterStoreChangedHandleMap.Empty();
+}
+
+void FNiagaraEmitterViewModel::ScriptGraphChanged(const FEdGraphEditAction& InAction, const UNiagaraScript& OwningScript)
+{
+	OnScriptGraphChangedDelegate.Broadcast(InAction, OwningScript);
+}
+
+void FNiagaraEmitterViewModel::ScriptParameterStoreChanged(const FNiagaraParameterStore& ChangedParameterStore, const UNiagaraScript& OwningScript)
+{
+	OnScriptParameterStoreChangedDelegate.Broadcast(ChangedParameterStore, OwningScript);
+}
+
 #undef LOCTEXT_NAMESPACE
diff --git a/Engine/Plugins/FX/Niagara/Source/NiagaraEditor/Private/ViewModels/NiagaraMetaDataCollectionViewModel.cpp b/Engine/Plugins/FX/Niagara/Source/NiagaraEditor/Private/ViewModels/NiagaraMetaDataCollectionViewModel.cpp
--- a/Engine/Plugins/FX/Niagara/Source/NiagaraEditor/Private/ViewModels/NiagaraMetaDataCollectionViewModel.cpp
+++ b/Engine/Plugins/FX/Niagara/Source/NiagaraEditor/Private/ViewModels/NiagaraMetaDataCollectionViewModel.cpp
@@ -1,4 +1,4 @@
-// Copyright 1998-2018 Epic Games, Inc. All Rights Reserved.
+﻿// Copyright 1998-2018 Epic Games, Inc. All Rights Reserved.
 
 #include "NiagaraMetaDataCollectionViewModel.h"
 #include "NiagaraTypes.h"
@@ -16,12 +16,39 @@
 #define LOCTEXT_NAMESPACE "NiagaraMetaDataCollectionViewModel"
 
 FNiagaraMetaDataCollectionViewModel::FNiagaraMetaDataCollectionViewModel()
+	: ModuleGraph(nullptr)
+	, bNeedsRefresh(false)
+	, bNeedsNotify(false)
 {
-	ModuleGraph = nullptr;
+}
+
+void FNiagaraMetaDataCollectionViewModel::Tick(float DeltaTime)
+{
+	if (bNeedsRefresh)
+	{
+		Refresh(bNeedsNotify);
+		bNeedsRefresh = false;
+		bNeedsNotify = false;
+	}
+}
+
+bool FNiagaraMetaDataCollectionViewModel::IsTickable() const
+{
+	return true;
+}
+
+TStatId FNiagaraMetaDataCollectionViewModel::GetStatId() const
+{
+	RETURN_QUICK_DECLARE_CYCLE_STAT(FNiagaraMetaDataCollectionViewModel, STATGROUP_Tickables);
 }
 
 void FNiagaraMetaDataCollectionViewModel::SetGraph(UNiagaraGraph* InGraph)
 {
+	if (ModuleGraph != nullptr)
+	{
+		Cleanup();
+	}
+
 	if (InGraph == nullptr)
 	{
 		return;
@@ -29,36 +56,10 @@ void FNiagaraMetaDataCollectionViewModel::SetGraph(UNiagaraGraph* InGraph)
 	// now build variables
 	ModuleGraph = InGraph;
 	
-	// Get the variable metadata from the graph
-	const UEdGraphSchema_Niagara* Schema = GetDefault<UEdGraphSchema_Niagara>();
-	TArray<UNiagaraNode*> NiagaraNodes;
-	ModuleGraph->GetNodesOfClass<UNiagaraNode>(NiagaraNodes);
-	for (UEdGraphNode* GraphNode : NiagaraNodes)
-	{
-		for (UEdGraphPin* CurrentPin : GraphNode->Pins)
-		{
-			// parse all pins, add all valid metadata to object array.
-			FNiagaraVariable NiagaraVariable = Schema->PinToNiagaraVariable(CurrentPin, true);
-			if (NiagaraVariable.IsInNameSpace(FNiagaraParameterHandle::ModuleNamespace.ToString()) && !FNiagaraConstants::IsNiagaraConstant(NiagaraVariable))
-			{
-				TSharedPtr<FNiagaraMetaDataViewModel> MetadataViewModel = GetMetadataViewModelForVariable(NiagaraVariable);
-				bool bNeedsToAdd = false;
-				if (!MetadataViewModel.IsValid())
-				{
-					MetadataViewModel = TSharedPtr<FNiagaraMetaDataViewModel>(new FNiagaraMetaDataViewModel(NiagaraVariable, *ModuleGraph));
-					bNeedsToAdd = true;
-				}
-				MetadataViewModel->AssociateNode(CurrentPin->GetOwningNode());
-				MetadataViewModel->OnMetadataChanged().AddRaw(this, &FNiagaraMetaDataCollectionViewModel::BroadcastChanged);
-				if (bNeedsToAdd)
-				{
-					MetaDataViewModels.Add(MetadataViewModel.ToSharedRef());
-				}
-			}
-		}
-	}
+	Refresh(false);
 
 	OnGraphChangedHandle = ModuleGraph->AddOnGraphChangedHandler(FOnGraphChanged::FDelegate::CreateRaw(this, &FNiagaraMetaDataCollectionViewModel::OnGraphChanged));
+	OnRecompileHandle = ModuleGraph->AddOnGraphNeedsRecompileHandler(FOnGraphChanged::FDelegate::CreateRaw(this, &FNiagaraMetaDataCollectionViewModel::OnGraphChanged));
 }
 
 TSharedPtr<FNiagaraMetaDataViewModel> FNiagaraMetaDataCollectionViewModel::GetMetadataViewModelForVariable(FNiagaraVariable InVariable)
@@ -76,9 +77,15 @@ TArray<TSharedRef<FNiagaraMetaDataViewModel>>& FNiagaraMetaDataCollectionViewMod
 	return MetaDataViewModels;
 }
 
+void FNiagaraMetaDataCollectionViewModel::RequestRefresh(bool bNotifyGraph)
+{
+	bNeedsRefresh = true;
+	bNeedsNotify = bNotifyGraph;
+}
+
 void FNiagaraMetaDataCollectionViewModel::OnGraphChanged(const struct FEdGraphEditAction& InAction)
 {
-	Reload(); 
+	RequestRefresh(false);
 } 
 
 void FNiagaraMetaDataCollectionViewModel::BroadcastChanged()
@@ -86,25 +93,86 @@ void FNiagaraMetaDataCollectionViewModel::BroadcastChanged()
 	OnCollectionChangedDelegate.Broadcast();
 }
 
-void FNiagaraMetaDataCollectionViewModel::Reload()
+void FNiagaraMetaDataCollectionViewModel::Refresh(bool bNotifyGraph)
 {
-	Cleanup();
-	SetGraph(ModuleGraph);
+	if (!ModuleGraph)
+	{
+		return;
+	}
+
+	// Get the variable metadata from the graph
+	MetaDataViewModels.Empty();
+	for (auto& MetadataElement : ModuleGraph->GetAllMetaData())
+	{
+		TSharedPtr<FNiagaraMetaDataViewModel> MetadataViewModel = GetMetadataViewModelForVariable(MetadataElement.Key);
+		if (!MetadataViewModel.IsValid())
+		{
+			MetadataViewModel = MakeShared<FNiagaraMetaDataViewModel>(MetadataElement.Key, *ModuleGraph);
+			MetaDataViewModels.Add(MetadataViewModel.ToSharedRef());
+		}
+
+		FNiagaraVariableMetaData& Metadata = MetadataElement.Value;
+		if (Metadata.ReferencerNodes.Num() > 0)
+		{
+			UNiagaraNode* Node = Cast<UNiagaraNode>(Metadata.ReferencerNodes[0].Get());
+			if (Node)
+			{
+				MetadataViewModel->AssociateNode(Node);
+			}
+		}
+	}
+	SortViewModels();
+	if (bNotifyGraph)
+	{
+		ModuleGraph->NotifyGraphNeedsRecompile();
+	}
 	BroadcastChanged();
 }
 
-void FNiagaraMetaDataCollectionViewModel::Cleanup()
+void FNiagaraMetaDataCollectionViewModel::SortViewModels()
 {
-	for (int32 i = 0; i < MetaDataViewModels.Num(); i++)
+	TMap<FString, int32> CategoryPriorityMap;
+	for (auto Metadata : MetaDataViewModels)
 	{
-		FNiagaraMetaDataViewModel* MetaDataViewModel = (FNiagaraMetaDataViewModel*)(&MetaDataViewModels[i].Get());
-		if (MetaDataViewModel != nullptr)
+		Metadata->RefreshMetaDataValue();
+		FString CategoryName = Metadata->GetGraphMetaData()->CategoryName.ToString();
+		if (!CategoryName.IsEmpty() && 
+			(!CategoryPriorityMap.Contains(CategoryName) || (CategoryPriorityMap.Contains(CategoryName) && CategoryPriorityMap[CategoryName] > Metadata->GetEditorSortPriority())))
 		{
-			MetaDataViewModel->OnMetadataChanged().RemoveAll(this);
+			CategoryPriorityMap.Add(CategoryName, Metadata->GetEditorSortPriority());
 		}
 	}
+	auto SortVars = [&](const TSharedRef<FNiagaraMetaDataViewModel>& A, const TSharedRef<FNiagaraMetaDataViewModel>& B)
+	{
+		int32 CategoryPriorityA = CategoryPriorityMap.Contains(A->GetGraphMetaData()->CategoryName.ToString()) ? CategoryPriorityMap[A->GetGraphMetaData()->CategoryName.ToString()] : MIN_int32;
+		int32 CategoryPriorityB = CategoryPriorityMap.Contains(B->GetGraphMetaData()->CategoryName.ToString()) ? CategoryPriorityMap[B->GetGraphMetaData()->CategoryName.ToString()] : MIN_int32;
+		if (CategoryPriorityA < CategoryPriorityB)
+		{
+			return true;
+		}
+		else if (CategoryPriorityA > CategoryPriorityB)
+		{
+			return false;
+		}
+		if (A->GetEditorSortPriority() < B->GetEditorSortPriority())
+		{
+			return true;
+		}
+		else if (A->GetEditorSortPriority() > B->GetEditorSortPriority())
+		{
+			return false;
+		}
+		//If equal priority, sort alphabetically.
+		return A->GetName().ToString() < B->GetName().ToString();
+	};
+	MetaDataViewModels.Sort(SortVars);
+}
+
+void FNiagaraMetaDataCollectionViewModel::Cleanup()
+{
 	MetaDataViewModels.Empty();
 	ModuleGraph->RemoveOnGraphChangedHandler(OnGraphChangedHandle);
+	ModuleGraph->RemoveOnGraphNeedsRecompileHandler(OnRecompileHandle);
 }
 
 FNiagaraMetaDataCollectionViewModel::~FNiagaraMetaDataCollectionViewModel()
diff --git a/Engine/Plugins/FX/Niagara/Source/NiagaraEditor/Private/ViewModels/NiagaraMetaDataCollectionViewModel.h b/Engine/Plugins/FX/Niagara/Source/NiagaraEditor/Private/ViewModels/NiagaraMetaDataCollectionViewModel.h
--- a/Engine/Plugins/FX/Niagara/Source/NiagaraEditor/Private/ViewModels/NiagaraMetaDataCollectionViewModel.h
+++ b/Engine/Plugins/FX/Niagara/Source/NiagaraEditor/Private/ViewModels/NiagaraMetaDataCollectionViewModel.h
@@ -1,13 +1,15 @@
-// Copyright 1998-2018 Epic Games, Inc. All Rights Reserved.
+﻿// Copyright 1998-2018 Epic Games, Inc. All Rights Reserved.
 
 #pragma once
 
 #include "NiagaraTypes.h"
+#include "TickableEditorObject.h"
+
 class UNiagaraScript;
 class UNiagaraGraph;
 class FNiagaraMetaDataViewModel;
 
-class FNiagaraMetaDataCollectionViewModel 
+class FNiagaraMetaDataCollectionViewModel : public FTickableEditorObject
 {
 public:
 	DECLARE_MULTICAST_DELEGATE(FOnCollectionChanged);
@@ -17,19 +19,26 @@ public:
 
 	~FNiagaraMetaDataCollectionViewModel();
 
+	//~ FTickableEditorObject interface
+	virtual void Tick(float DeltaTime) override;
+	virtual bool IsTickable() const override;
+	virtual TStatId GetStatId() const override;
+
 	/** Sets the view model to a new script. */
 	void SetGraph(UNiagaraGraph* InGraph);
 
 	/** Gets the metadata view models. */
 	TArray<TSharedRef<FNiagaraMetaDataViewModel>>& GetVariableModels();
 	
-	/** Reloads the data from the graph and builds the viewmodel*/
-	void Reload();
+	/** Request refresh the data from the graph and builds the viewmodels next frame. */
+	void RequestRefresh(bool bNotifyGraph);
 
 	/** returns the delegate to be called when the collection changes*/
 	FOnCollectionChanged& OnCollectionChanged() { return OnCollectionChangedDelegate; }
-	
+
 private:
+	void Refresh(bool bNotifyGraph);
+	void SortViewModels();
 	/** Callback for when the graph changes and the collection viewmodel needs to react */
 	void OnGraphChanged(const struct FEdGraphEditAction& InAction);
 	void Cleanup();
@@ -44,7 +53,14 @@ private:
 
 	/** The handle to the graph changed delegate. */
 	FDelegateHandle OnGraphChangedHandle;
+	FDelegateHandle OnRecompileHandle;
 
 	/** A multicast delegate which is called whenever the parameter collection is changed. */
 	FOnCollectionChanged OnCollectionChangedDelegate;
+	
+	/** Refresh the UI next frames. */
+	bool bNeedsRefresh;
+
+	/** Needs to notify the graph on the next refesh. */
+	bool bNeedsNotify;
 };
\ No newline at end of file
diff --git a/Engine/Plugins/FX/Niagara/Source/NiagaraEditor/Private/ViewModels/NiagaraMetaDataViewModel.cpp b/Engine/Plugins/FX/Niagara/Source/NiagaraEditor/Private/ViewModels/NiagaraMetaDataViewModel.cpp
--- a/Engine/Plugins/FX/Niagara/Source/NiagaraEditor/Private/ViewModels/NiagaraMetaDataViewModel.cpp
+++ b/Engine/Plugins/FX/Niagara/Source/NiagaraEditor/Private/ViewModels/NiagaraMetaDataViewModel.cpp
@@ -1,4 +1,4 @@
-// Copyright 1998-2018 Epic Games, Inc. All Rights Reserved.
+﻿// Copyright 1998-2018 Epic Games, Inc. All Rights Reserved.
 
 #include "NiagaraMetaDataViewModel.h"
 #include "NiagaraEditorUtilities.h"
@@ -17,6 +17,7 @@ FNiagaraMetaDataViewModel::FNiagaraMetaDataViewModel(FNiagaraVariable& InGraphVa
 	: GraphVariable(InGraphVariable)
 	, CurrentGraph(&InGraph)
 {
+	GraphMetaData = nullptr;
 	RefreshMetaDataValue();
 }
 
@@ -37,6 +38,15 @@ void FNiagaraMetaDataViewModel::AssociateNode(UEdGraphNode* InNode)
 	AssociatedNodes.AddUnique(TWeakObjectPtr<UObject>(InNode));
 }
 
+int32 FNiagaraMetaDataViewModel::GetEditorSortPriority()
+{
+	if (GraphMetaData != nullptr)
+	{
+		return GraphMetaData->EditorSortPriority;
+	}
+	return 0;
+}
+
 void FNiagaraMetaDataViewModel::NotifyMetaDataChanged()
 {
 	FScopedTransaction ScopedTransaction(LOCTEXT("EditVariableMetadata", "Edit variable metadata"));
@@ -55,9 +65,6 @@ void FNiagaraMetaDataViewModel::NotifyMetaDataChanged()
 	}
 	
 	OnMetadataChangedDelegate.Broadcast();
-	
-	// Trigger save prompt on close; it's sufficiently cheap
-	CurrentGraph->NotifyGraphNeedsRecompile();
 }
 
 FNiagaraVariableMetaData*  FNiagaraMetaDataViewModel::GetGraphMetaData()
diff --git a/Engine/Plugins/FX/Niagara/Source/NiagaraEditor/Private/ViewModels/NiagaraMetaDataViewModel.h b/Engine/Plugins/FX/Niagara/Source/NiagaraEditor/Private/ViewModels/NiagaraMetaDataViewModel.h
--- a/Engine/Plugins/FX/Niagara/Source/NiagaraEditor/Private/ViewModels/NiagaraMetaDataViewModel.h
+++ b/Engine/Plugins/FX/Niagara/Source/NiagaraEditor/Private/ViewModels/NiagaraMetaDataViewModel.h
@@ -1,4 +1,4 @@
-// Copyright 1998-2018 Epic Games, Inc. All Rights Reserved.
+﻿// Copyright 1998-2018 Epic Games, Inc. All Rights Reserved.
 
 #pragma once
 
@@ -45,8 +45,10 @@ public:
 
 	/** associates a node to this metadata*/
 	void AssociateNode(UEdGraphNode* InNode);
+
+	/** gets the sort order prioriy from metadata */
+	int32 GetEditorSortPriority();
 	
-private:
 	/** Refreshes the parameter value struct from the variable data. */
 	void RefreshMetaDataValue();
 
diff --git a/Engine/Plugins/FX/Niagara/Source/NiagaraEditor/Private/ViewModels/NiagaraParameterCollectionAssetViewModel.cpp b/Engine/Plugins/FX/Niagara/Source/NiagaraEditor/Private/ViewModels/NiagaraParameterCollectionAssetViewModel.cpp
--- a/Engine/Plugins/FX/Niagara/Source/NiagaraEditor/Private/ViewModels/NiagaraParameterCollectionAssetViewModel.cpp
+++ b/Engine/Plugins/FX/Niagara/Source/NiagaraEditor/Private/ViewModels/NiagaraParameterCollectionAssetViewModel.cpp
@@ -10,7 +10,7 @@
 #include "ScopedTransaction.h"
 #include "EditorSupportDelegates.h"
 #include "Modules/ModuleManager.h"
-#include "TNiagaraViewModelManager.h"
+#include "ViewModels/TNiagaraViewModelManager.h"
 #include "Math/NumericLimits.h"
 
 #include "NiagaraDataInterface.h"
diff --git a/Engine/Plugins/FX/Niagara/Source/NiagaraEditor/Private/ViewModels/NiagaraParameterCollectionAssetViewModel.h b/Engine/Plugins/FX/Niagara/Source/NiagaraEditor/Private/ViewModels/NiagaraParameterCollectionAssetViewModel.h
--- a/Engine/Plugins/FX/Niagara/Source/NiagaraEditor/Private/ViewModels/NiagaraParameterCollectionAssetViewModel.h
+++ b/Engine/Plugins/FX/Niagara/Source/NiagaraEditor/Private/ViewModels/NiagaraParameterCollectionAssetViewModel.h
@@ -3,7 +3,7 @@
 #pragma once
 
 #include "NiagaraParameterCollectionViewModel.h"
-#include "TNiagaraViewModelManager.h"
+#include "ViewModels/TNiagaraViewModelManager.h"
 #include "EditorUndoClient.h"
 #include "Misc/NotifyHook.h"
 
diff --git a/Engine/Plugins/FX/Niagara/Source/NiagaraEditor/Private/ViewModels/NiagaraParameterCollectionViewModel.cpp b/Engine/Plugins/FX/Niagara/Source/NiagaraEditor/Private/ViewModels/NiagaraParameterCollectionViewModel.cpp
--- a/Engine/Plugins/FX/Niagara/Source/NiagaraEditor/Private/ViewModels/NiagaraParameterCollectionViewModel.cpp
+++ b/Engine/Plugins/FX/Niagara/Source/NiagaraEditor/Private/ViewModels/NiagaraParameterCollectionViewModel.cpp
@@ -28,8 +28,8 @@ void INiagaraParameterCollectionViewModel::SortViewModels(TArray<TSharedRef<INia
 
 FNiagaraParameterCollectionViewModel::FNiagaraParameterCollectionViewModel(ENiagaraParameterEditMode InParameterEditMode)
 	: ParameterEditMode(InParameterEditMode)
+	, bNeedsRefresh(false)
 	, bIsExpanded(true)
-	
 {
 }
 
@@ -62,6 +62,25 @@ bool FNiagaraParameterCollectionViewModel::CanDeleteParameters() const
 	return ParameterEditMode == ENiagaraParameterEditMode::EditAll;
 }
 
+void FNiagaraParameterCollectionViewModel::Tick(float DeltaTime)
+{
+	if (bNeedsRefresh)
+	{
+		RefreshParameterViewModels();
+		bNeedsRefresh = false;
+	}
+}
+
+bool FNiagaraParameterCollectionViewModel::IsTickable() const
+{
+	return true;
+}
+
+TStatId FNiagaraParameterCollectionViewModel::GetStatId() const
+{
+	RETURN_QUICK_DECLARE_CYCLE_STAT(FNiagaraParameterCollectionViewModel, STATGROUP_Tickables);
+}
+
 TSet<FName> FNiagaraParameterCollectionViewModel::GetParameterNames()
 {
 	TSet<FName> ParameterNames;
diff --git a/Engine/Plugins/FX/Niagara/Source/NiagaraEditor/Private/ViewModels/NiagaraParameterCollectionViewModel.h b/Engine/Plugins/FX/Niagara/Source/NiagaraEditor/Private/ViewModels/NiagaraParameterCollectionViewModel.h
--- a/Engine/Plugins/FX/Niagara/Source/NiagaraEditor/Private/ViewModels/NiagaraParameterCollectionViewModel.h
+++ b/Engine/Plugins/FX/Niagara/Source/NiagaraEditor/Private/ViewModels/NiagaraParameterCollectionViewModel.h
@@ -5,6 +5,7 @@
 #include "NiagaraObjectSelection.h"
 #include "NiagaraParameterViewModel.h"
 #include "Layout/Visibility.h"
+#include "TickableEditorObject.h"
 
 struct FNiagaraTypeDefinition;
 
@@ -84,7 +85,7 @@ public:
 
 /** Base class for parameter collection view models.  Partially implements the parameter collection interface with
 	behavior common to all view models. */
-class FNiagaraParameterCollectionViewModel : public INiagaraParameterCollectionViewModel
+class FNiagaraParameterCollectionViewModel : public INiagaraParameterCollectionViewModel, public FTickableEditorObject
 {
 public:
 	FNiagaraParameterCollectionViewModel(ENiagaraParameterEditMode InParameterEditMode);
@@ -106,6 +107,11 @@ public:
 	virtual FOnParameterValueChanged& OnParameterValueChanged() override { return OnParameterValueChangedDelegate; }
 	virtual void NotifyParameterChangedExternally(FName ParameterName) override;
 
+	//~ FTickableEditorObject interface
+	virtual void Tick(float DeltaTime) override;
+	virtual bool IsTickable() const override;
+	virtual TStatId GetStatId() const override;
+
 protected:
 	/** Gets a set containing the names of the parameters. */
 	TSet<FName> GetParameterNames();
@@ -133,6 +139,9 @@ protected:
 	/** Gets the edit mode for parameters in this collection. */
 	const ENiagaraParameterEditMode ParameterEditMode;
 
+	/** Refresh the UI next frames. */
+	bool bNeedsRefresh;
+
 private:
 	/** An array of available types for parameters. */
 	TSharedPtr<TArray<TSharedPtr<FNiagaraTypeDefinition>>> AvailableTypes;
diff --git a/Engine/Plugins/FX/Niagara/Source/NiagaraEditor/Private/ViewModels/NiagaraScriptGraphViewModel.cpp b/Engine/Plugins/FX/Niagara/Source/NiagaraEditor/Private/ViewModels/NiagaraScriptGraphViewModel.cpp
--- a/Engine/Plugins/FX/Niagara/Source/NiagaraEditor/Private/ViewModels/NiagaraScriptGraphViewModel.cpp
+++ b/Engine/Plugins/FX/Niagara/Source/NiagaraEditor/Private/ViewModels/NiagaraScriptGraphViewModel.cpp
@@ -293,6 +293,9 @@ void FNiagaraScriptGraphViewModel::PasteNodes()
 		for (UEdGraphNode* PastedNode : PastedNodes)
 		{
 			PastedNode->CreateNewGuid();
+			UNiagaraNode* Node = Cast<UNiagaraNode>(PastedNode);
+			if (Node)
+				Node->MarkNodeRequiresSynchronization(__FUNCTION__, false);
 		}
 
 		FNiagaraEditorUtilities::FixUpPastedInputNodes(Graph, PastedNodes);
diff --git a/Engine/Plugins/FX/Niagara/Source/NiagaraEditor/Private/ViewModels/NiagaraScriptInputCollectionViewModel.cpp b/Engine/Plugins/FX/Niagara/Source/NiagaraEditor/Private/ViewModels/NiagaraScriptInputCollectionViewModel.cpp
--- a/Engine/Plugins/FX/Niagara/Source/NiagaraEditor/Private/ViewModels/NiagaraScriptInputCollectionViewModel.cpp
+++ b/Engine/Plugins/FX/Niagara/Source/NiagaraEditor/Private/ViewModels/NiagaraScriptInputCollectionViewModel.cpp
@@ -44,6 +44,8 @@ FNiagaraScriptInputCollectionViewModel::FNiagaraScriptInputCollectionViewModel(U
 	{
 		OnGraphChangedHandle = Graph->AddOnGraphChangedHandler(
 			FOnGraphChanged::FDelegate::CreateRaw(this, &FNiagaraScriptInputCollectionViewModel::OnGraphChanged));
+		OnRecompileHandle = Graph->AddOnGraphNeedsRecompileHandler(
+			FOnGraphChanged::FDelegate::CreateRaw(this, &FNiagaraScriptInputCollectionViewModel::OnGraphChanged));
 	}
 }
 
@@ -79,6 +81,9 @@ FNiagaraScriptInputCollectionViewModel::FNiagaraScriptInputCollectionViewModel(U
 	{
 		OnGraphChangedHandle = Graph->AddOnGraphChangedHandler(
 			FOnGraphChanged::FDelegate::CreateRaw(this, &FNiagaraScriptInputCollectionViewModel::OnGraphChanged));
+		OnRecompileHandle = Graph->AddOnGraphNeedsRecompileHandler(
+			FOnGraphChanged::FDelegate::CreateRaw(this, &FNiagaraScriptInputCollectionViewModel::OnGraphChanged));
+
 	}
 }
 
@@ -100,6 +105,7 @@ FNiagaraScriptInputCollectionViewModel::~FNiagaraScriptInputCollectionViewModel(
 	if (Graph.IsValid())
 	{
 		Graph->RemoveOnGraphChangedHandler(OnGraphChangedHandle);
+		Graph->RemoveOnGraphNeedsRecompileHandler(OnRecompileHandle);
 	}
 }
 
@@ -108,6 +114,7 @@ void FNiagaraScriptInputCollectionViewModel::SetScripts(TArray<UNiagaraScript*>
 	if (Graph.IsValid())
 	{
 		Graph->RemoveOnGraphChangedHandler(OnGraphChangedHandle);
+		Graph->RemoveOnGraphNeedsRecompileHandler(OnRecompileHandle);
 	}
 
 	Scripts.Empty();
@@ -122,6 +129,8 @@ void FNiagaraScriptInputCollectionViewModel::SetScripts(TArray<UNiagaraScript*>
 		Graph = Cast<UNiagaraScriptSource>(Scripts[0]->GetSource())->NodeGraph;
 		OnGraphChangedHandle = Graph->AddOnGraphChangedHandler(
 			FOnGraphChanged::FDelegate::CreateRaw(this, &FNiagaraScriptInputCollectionViewModel::OnGraphChanged));
+		OnRecompileHandle = Graph->AddOnGraphNeedsRecompileHandler(
+			FOnGraphChanged::FDelegate::CreateRaw(this, &FNiagaraScriptInputCollectionViewModel::OnGraphChanged));
 		bCanHaveNumericParameters = Scripts[0]->IsStandaloneScript();
 	}
 	else
@@ -302,7 +311,7 @@ void FNiagaraScriptInputCollectionViewModel::RefreshParameterViewModels()
 				UNiagaraScript* Script = nullptr;
 				for (TWeakObjectPtr<UNiagaraScript> ScriptWeakPtr : Scripts)
 				{
-					if (!ScriptWeakPtr.IsValid())
+					if (!ScriptWeakPtr.IsValid() || !ScriptWeakPtr->GetVMExecutableData().IsValid())
 					{
 						continue;
 					}
@@ -311,7 +320,7 @@ void FNiagaraScriptInputCollectionViewModel::RefreshParameterViewModels()
 						break;
 					}
 
-					for (FNiagaraVariable& EmitterVariableToCheck : ScriptWeakPtr->Parameters.Parameters)
+					for (FNiagaraVariable& EmitterVariableToCheck : ScriptWeakPtr->GetVMExecutableData().Parameters.Parameters)
 					{
 
 						// @TODO We should check ID's here, but its possible that the 
@@ -339,7 +348,7 @@ void FNiagaraScriptInputCollectionViewModel::RefreshParameterViewModels()
 						break;
 					}
 
-					for (FNiagaraScriptDataInterfaceInfo& DataInterfaceInfoItem : ScriptWeakPtr->DataInterfaceInfo)
+					for (FNiagaraScriptDataInterfaceInfo& DataInterfaceInfoItem : ScriptWeakPtr->GetCachedDefaultDataInterfaces())
 					{
 						if (DataInterfaceInfoItem.Name == InputNode->Input.GetName())
 						{
@@ -370,7 +379,7 @@ bool FNiagaraScriptInputCollectionViewModel::SupportsType(const FNiagaraTypeDefi
 
 void FNiagaraScriptInputCollectionViewModel::OnGraphChanged(const struct FEdGraphEditAction& InAction)
 {
-	RefreshParameterViewModels();
+	bNeedsRefresh = true;
 }
 
 void FNiagaraScriptInputCollectionViewModel::OnParameterNameChanged(FName OldName, FName NewName, FNiagaraVariable* ParameterVariable)
@@ -417,7 +426,12 @@ void FNiagaraScriptInputCollectionViewModel::OnParameterNameChanged(FName OldNam
 	// Synchronize script variables...
 	for (TWeakObjectPtr<UNiagaraScript> Script : Scripts)
 	{
-		for (FNiagaraVariable& EmitterVariableToCheck : Script->Parameters.Parameters)
+		if (!Script.IsValid() || !Script->GetVMExecutableData().IsValid())
+		{
+			continue;
+		}
+
+		for (FNiagaraVariable& EmitterVariableToCheck : Script->GetVMExecutableData().Parameters.Parameters)
 		{
 			if (EmitterVariableToCheck.GetName() == ParameterVariable->GetName() && ParameterVariable != &EmitterVariableToCheck)
 			{
@@ -503,7 +517,11 @@ void FNiagaraScriptInputCollectionViewModel::OnParameterTypeChanged(FNiagaraVari
 	// Synchronize script variables...
 	for (TWeakObjectPtr<UNiagaraScript> Script : Scripts)
 	{
-		for (FNiagaraVariable& EmitterVariableToCheck : Script->Parameters.Parameters)
+		if (!Script.IsValid() || !Script->GetVMExecutableData().IsValid())
+		{
+			continue;
+		}
+		for (FNiagaraVariable& EmitterVariableToCheck : Script->GetVMExecutableData().Parameters.Parameters)
 		{
 			if (EmitterVariableToCheck.GetName() == ParameterVariable->GetName() && ParameterVariable != &EmitterVariableToCheck)
 			{
@@ -555,7 +573,12 @@ void FNiagaraScriptInputCollectionViewModel::OnParameterValueChangedInternal(TSh
 			// Synchronize script variables...
 			for (TWeakObjectPtr<UNiagaraScript> Script : Scripts)
 			{
-				for (FNiagaraVariable& EmitterVariableToCheck : Script->Parameters.Parameters)
+				if (!Script.IsValid() || !Script->GetVMExecutableData().IsValid())
+				{
+					continue;
+				}
+
+				for (FNiagaraVariable& EmitterVariableToCheck : Script->GetVMExecutableData().Parameters.Parameters)
 				{
 					if (EmitterVariableToCheck.GetName() == ChangedParameter->GetName())
 					{
@@ -588,7 +611,11 @@ void FNiagaraScriptInputCollectionViewModel::OnParameterValueChangedInternal(TSh
 				// Synchronize script variables...
 				for (TWeakObjectPtr<UNiagaraScript> Script : Scripts)
 				{
-					for (FNiagaraScriptDataInterfaceInfo& Info : Script->DataInterfaceInfo)
+					if (!Script.IsValid() || !Script->GetVMExecutableData().IsValid())
+					{
+						continue;
+					}
+					for (FNiagaraScriptDataInterfaceInfo& Info : Script->GetCachedDefaultDataInterfaces())
 					{
 						if (Info.Name == ChangedParameter->GetName() && DataInterface != Info.DataInterface)
 						{
diff --git a/Engine/Plugins/FX/Niagara/Source/NiagaraEditor/Private/ViewModels/NiagaraScriptInputCollectionViewModel.h b/Engine/Plugins/FX/Niagara/Source/NiagaraEditor/Private/ViewModels/NiagaraScriptInputCollectionViewModel.h
--- a/Engine/Plugins/FX/Niagara/Source/NiagaraEditor/Private/ViewModels/NiagaraScriptInputCollectionViewModel.h
+++ b/Engine/Plugins/FX/Niagara/Source/NiagaraEditor/Private/ViewModels/NiagaraScriptInputCollectionViewModel.h
@@ -74,6 +74,7 @@ private:
 
 	/** The handle to the graph changed delegate. */
 	FDelegateHandle OnGraphChangedHandle;
+	FDelegateHandle OnRecompileHandle;
 
 	/** Whether or not generic numeric type parameters are supported as inputs and outputs for this script. */
 	bool bCanHaveNumericParameters;
diff --git a/Engine/Plugins/FX/Niagara/Source/NiagaraEditor/Private/ViewModels/NiagaraScriptOutputCollectionViewModel.cpp b/Engine/Plugins/FX/Niagara/Source/NiagaraEditor/Private/ViewModels/NiagaraScriptOutputCollectionViewModel.cpp
--- a/Engine/Plugins/FX/Niagara/Source/NiagaraEditor/Private/ViewModels/NiagaraScriptOutputCollectionViewModel.cpp
+++ b/Engine/Plugins/FX/Niagara/Source/NiagaraEditor/Private/ViewModels/NiagaraScriptOutputCollectionViewModel.cpp
@@ -48,6 +48,8 @@ FNiagaraScriptOutputCollectionViewModel::FNiagaraScriptOutputCollectionViewModel
 	{
 		OnGraphChangedHandle = Graph->AddOnGraphChangedHandler(
 			FOnGraphChanged::FDelegate::CreateRaw(this, &FNiagaraScriptOutputCollectionViewModel::OnGraphChanged));
+		OnRecompileHandle = Graph->AddOnGraphNeedsRecompileHandler(
+			FOnGraphChanged::FDelegate::CreateRaw(this, &FNiagaraScriptOutputCollectionViewModel::OnGraphChanged));
 	}
 }
 
@@ -84,6 +86,8 @@ FNiagaraScriptOutputCollectionViewModel::FNiagaraScriptOutputCollectionViewModel
 	{
 		OnGraphChangedHandle = Graph->AddOnGraphChangedHandler(
 			FOnGraphChanged::FDelegate::CreateRaw(this, &FNiagaraScriptOutputCollectionViewModel::OnGraphChanged));
+		OnRecompileHandle = Graph->AddOnGraphNeedsRecompileHandler(
+			FOnGraphChanged::FDelegate::CreateRaw(this, &FNiagaraScriptOutputCollectionViewModel::OnGraphChanged));
 	}
 }
 
@@ -103,6 +107,7 @@ FNiagaraScriptOutputCollectionViewModel::~FNiagaraScriptOutputCollectionViewMode
 	if (Graph.IsValid())
 	{
 		Graph->RemoveOnGraphChangedHandler(OnGraphChangedHandle);
+		Graph->RemoveOnGraphNeedsRecompileHandler(OnRecompileHandle);
 	}
 }
 
@@ -112,6 +117,7 @@ void FNiagaraScriptOutputCollectionViewModel::SetScripts(TArray<UNiagaraScript*>
 	if (Graph.IsValid())
 	{
 		Graph->RemoveOnGraphChangedHandler(OnGraphChangedHandle);
+		Graph->RemoveOnGraphNeedsRecompileHandler(OnRecompileHandle);
 	}
 
 	Scripts.Empty();
@@ -127,6 +133,8 @@ void FNiagaraScriptOutputCollectionViewModel::SetScripts(TArray<UNiagaraScript*>
 		Graph = Cast<UNiagaraScriptSource>(Scripts[0]->GetSource())->NodeGraph;
 		OnGraphChangedHandle = Graph->AddOnGraphChangedHandler(
 			FOnGraphChanged::FDelegate::CreateRaw(this, &FNiagaraScriptOutputCollectionViewModel::OnGraphChanged));
+		OnRecompileHandle = Graph->AddOnGraphNeedsRecompileHandler(
+			FOnGraphChanged::FDelegate::CreateRaw(this, &FNiagaraScriptOutputCollectionViewModel::OnGraphChanged));
 		bCanHaveNumericParameters = Scripts[0]->IsStandaloneScript();
 
 		if (Scripts.Num() == 1)
@@ -181,7 +189,7 @@ void FNiagaraScriptOutputCollectionViewModel::AddParameter(TSharedPtr<FNiagaraTy
 
 		OtherOutputNode->NotifyOutputVariablesChanged();
 	}
-	Graph->MarkGraphRequiresSynchronization();
+	Graph->MarkGraphRequiresSynchronization(TEXT("Ouptut parameter added"));
 	OnOutputParametersChangedDelegate.Broadcast();
 }
 
@@ -274,7 +282,7 @@ FNiagaraScriptOutputCollectionViewModel::FOnOutputParametersChanged& FNiagaraScr
 
 void FNiagaraScriptOutputCollectionViewModel::OnGraphChanged(const struct FEdGraphEditAction& InAction)
 {
-	RefreshParameterViewModels();
+	bNeedsRefresh = true;
 }
 
 void FNiagaraScriptOutputCollectionViewModel::OnParameterNameChanged(FName OldName, FName NewName, FNiagaraVariable* ParameterVariable)
diff --git a/Engine/Plugins/FX/Niagara/Source/NiagaraEditor/Private/ViewModels/NiagaraScriptOutputCollectionViewModel.h b/Engine/Plugins/FX/Niagara/Source/NiagaraEditor/Private/ViewModels/NiagaraScriptOutputCollectionViewModel.h
--- a/Engine/Plugins/FX/Niagara/Source/NiagaraEditor/Private/ViewModels/NiagaraScriptOutputCollectionViewModel.h
+++ b/Engine/Plugins/FX/Niagara/Source/NiagaraEditor/Private/ViewModels/NiagaraScriptOutputCollectionViewModel.h
@@ -75,6 +75,7 @@ private:
 
 	/** The handle to the graph changed delegate needed for removing. */
 	FDelegateHandle OnGraphChangedHandle;
+	FDelegateHandle OnRecompileHandle;
 
 	/** A multicast delegate which is called whenever the output parameter collection is changed, or when
 	a paramter in the collection is changed. */
