Copying //UE4/Dev-Niagara to //UE4/Dev-Main (Source: //UE4/Dev-Niagara @ 4074996)

diff --git a/Engine/Plugins/FX/Niagara/Source/NiagaraEditor/Private/ViewModels/NiagaraScriptViewModel.cpp b/Engine/Plugins/FX/Niagara/Source/NiagaraEditor/Private/ViewModels/NiagaraScriptViewModel.cpp
--- a/Engine/Plugins/FX/Niagara/Source/NiagaraEditor/Private/ViewModels/NiagaraScriptViewModel.cpp
+++ b/Engine/Plugins/FX/Niagara/Source/NiagaraEditor/Private/ViewModels/NiagaraScriptViewModel.cpp
@@ -15,7 +15,7 @@
 #include "GraphEditAction.h"
 #include "UObject/Package.h"
 #include "Editor.h"
-#include "TNiagaraViewModelManager.h"
+#include "ViewModels/TNiagaraViewModelManager.h"
 #include "NiagaraDataInterface.h"
 
 template<> TMap<UNiagaraScript*, TArray<FNiagaraScriptViewModel*>> TNiagaraViewModelManager<UNiagaraScript, FNiagaraScriptViewModel>::ObjectsToViewModels{};
@@ -27,12 +27,12 @@ FNiagaraScriptViewModel::FNiagaraScriptViewModel(UNiagaraScript* InScript, FText
 	, GraphViewModel(MakeShareable(new FNiagaraScriptGraphViewModel(Cast<UNiagaraScriptSource>(InScript->GetSource()), DisplayName)))
 	, bUpdatingSelectionInternally(false)
 	, LastCompileStatus(ENiagaraScriptCompileStatus::NCS_Unknown)
-	, bNeedsSave(true)
 {
 	if (InScript != nullptr)
 	{
 		Scripts.Add(InScript);
 		Source = Cast<UNiagaraScriptSource>(InScript->GetSource());
+		Scripts[0]->OnVMScriptCompiled().AddRaw(this, &FNiagaraScriptViewModel::OnVMScriptCompiled);
 	}
 
 	InputCollectionViewModel->GetSelection().OnSelectedObjectsChanged().AddRaw(this, &FNiagaraScriptViewModel::InputViewModelSelectionChanged);
@@ -44,18 +44,12 @@ FNiagaraScriptViewModel::FNiagaraScriptViewModel(UNiagaraScript* InScript, FText
 	if (Source != nullptr)
 	{
 		UNiagaraGraph* Graph = Source->NodeGraph;
-		if (Graph != nullptr)
-		{
-			OnGraphChangedHandle = Graph->AddOnGraphChangedHandler(
-				FOnGraphChanged::FDelegate::CreateRaw(this, &FNiagaraScriptViewModel::OnGraphChanged));
-		}
 		MetaDataCollectionViewModel->SetGraph(Graph);
 	}
 
 	// Guess at initial compile status
-	if (InScript != nullptr && InScript->GetByteCode().Num() == 0) // This is either a brand new script or failed in the past. Since we create a default working script, assume invalid.
+	if (InScript != nullptr && InScript->GetVMExecutableData().IsValid() && InScript->GetVMExecutableData().ByteCode.Num() == 0) // This is either a brand new script or failed in the past. Since we create a default working script, assume invalid.
 	{
-		bNeedsSave = true;
 		LastCompileStatus = InScript->GetLastCompileStatus();
 		GraphViewModel->SetErrorTextToolTip("Please recompile for full error stack.");
 	}
@@ -63,18 +57,18 @@ FNiagaraScriptViewModel::FNiagaraScriptViewModel(UNiagaraScript* InScript, FText
 	{
 		if (InScript != nullptr && InScript->AreScriptAndSourceSynchronized())
 		{
-			bNeedsSave = false;
 			LastCompileStatus = InScript->GetLastCompileStatus();
 		}
 		else
 		{
-			bNeedsSave = true;
-			LastCompileStatus = ENiagaraScriptCompileStatus::NCS_Error;
+			LastCompileStatus = ENiagaraScriptCompileStatus::NCS_Unknown;
 			GraphViewModel->SetErrorTextToolTip("Please recompile for full error stack.");
 		}
 	}
 
 	RegisteredHandles.Add(RegisterViewModelWithMap(InScript, this));
+	//UE_LOG(LogNiagaraEditor, Warning, TEXT("FNiagaraScriptViewModel %p"), this);
+
 }
 
 FNiagaraScriptViewModel::FNiagaraScriptViewModel(UNiagaraEmitter* InEmitter, FText DisplayName, ENiagaraParameterEditMode InParameterEditMode)
@@ -84,15 +78,15 @@ FNiagaraScriptViewModel::FNiagaraScriptViewModel(UNiagaraEmitter* InEmitter, FTe
 	, GraphViewModel(MakeShareable(new FNiagaraScriptGraphViewModel(Cast<UNiagaraScriptSource>(InEmitter->GraphSource), DisplayName)))
 	, bUpdatingSelectionInternally(false)
 	, LastCompileStatus(ENiagaraScriptCompileStatus::NCS_Unknown)
-	, bNeedsSave(true)
 {
 	TArray<UNiagaraScript*> InScripts;
 	InEmitter->GetScripts(InScripts);
 	// Because of weak pointers, we need to copy ourselves..
 	for (UNiagaraScript* Script : InScripts)
 	{
-		Scripts.Add(Script);
+		int32 i = Scripts.Add(Script);
 		check(Script->GetSource() == InEmitter->GraphSource);
+		Scripts[i]->OnVMScriptCompiled().AddRaw(this, &FNiagaraScriptViewModel::OnVMScriptCompiled);
 	}
 	Source = Cast<UNiagaraScriptSource>(InEmitter->GraphSource);
 
@@ -105,15 +99,9 @@ FNiagaraScriptViewModel::FNiagaraScriptViewModel(UNiagaraEmitter* InEmitter, FTe
 	UNiagaraGraph* Graph = Source->NodeGraph;
 
 	MetaDataCollectionViewModel->SetGraph(Graph);
-	if (Graph != nullptr)
-	{
-		OnGraphChangedHandle = Graph->AddOnGraphChangedHandler(
-			FOnGraphChanged::FDelegate::CreateRaw(this, &FNiagaraScriptViewModel::OnGraphChanged));
-	}
 
 	// Guess at initial compile status
 	LastCompileStatus = ENiagaraScriptCompileStatus::NCS_UpToDate;
-	bNeedsSave = false;
 
 	CompileStatuses.Empty();
 	CompileErrors.Empty();
@@ -123,9 +111,8 @@ FNiagaraScriptViewModel::FNiagaraScriptViewModel(UNiagaraEmitter* InEmitter, FTe
 	{
 		FString Message;
 		ENiagaraScriptCompileStatus ScriptStatus = Scripts[i]->GetLastCompileStatus();
-		if (Scripts[i].IsValid() && Scripts[i]->IsCompilable() && Scripts[i]->GetByteCode().Num() == 0) // This is either a brand new script or failed in the past. Since we create a default working script, assume invalid.
+		if (Scripts[i].IsValid() && Scripts[i]->IsCompilable() && Scripts[i]->GetVMExecutableData().IsValid() && Scripts[i]->GetVMExecutableData().ByteCode.Num() == 0) // This is either a brand new script or failed in the past. Since we create a default working script, assume invalid.
 		{
-			bNeedsSave = true;
 			LastCompileStatus = ScriptStatus;
 			Message = TEXT("Please recompile for full error stack.");
 			GraphViewModel->SetErrorTextToolTip(Message);
@@ -143,7 +130,6 @@ FNiagaraScriptViewModel::FNiagaraScriptViewModel(UNiagaraEmitter* InEmitter, FTe
 			}
 			else
 			{
-				bNeedsSave = true;
 				LastCompileStatus = FNiagaraEditorUtilities::UnionCompileStatus(LastCompileStatus, ENiagaraScriptCompileStatus::NCS_Error);
 				ScriptStatus = ENiagaraScriptCompileStatus::NCS_Error;
 				Message = TEXT("Please recompile for full error stack.");
@@ -164,6 +150,37 @@ FNiagaraScriptViewModel::FNiagaraScriptViewModel(UNiagaraEmitter* InEmitter, FTe
 		CompileTypes[i].Key = Scripts[i]->GetUsage();
 		CompileTypes[i].Value = Scripts[i]->GetUsageId();
 	}
+	//UE_LOG(LogNiagaraEditor, Warning, TEXT("FNiagaraScriptViewModel %p"), this);
+}
+
+void FNiagaraScriptViewModel::OnVMScriptCompiled(UNiagaraScript* InScript)
+{
+	for (int32 i = 0; i < Scripts.Num(); i++)
+	{
+		UNiagaraScript* Script = Scripts[i].Get();
+		if (Script == InScript && InScript->IsStandaloneScript())
+		{
+			if (InScript->GetVMExecutableData().IsValid() == false)
+			{
+				continue;
+			}
+
+			LastCompileStatus = InScript->GetVMExecutableData().LastCompileStatus;
+
+			if (InScript->GetVMExecutableData().LastCompileStatus == ENiagaraScriptCompileStatus::NCS_Error)
+			{
+				const FString& ErrorMsg = InScript->GetVMExecutableData().ErrorMsg;
+				GraphViewModel->SetErrorTextToolTip(ErrorMsg + "\n(These same errors are also in the log)");
+			}
+			else
+			{
+				GraphViewModel->SetErrorTextToolTip("");
+			}
+
+			InputCollectionViewModel->RefreshParameterViewModels();
+			OutputCollectionViewModel->RefreshParameterViewModels();
+		}
+	}
 }
 
 bool FNiagaraScriptViewModel::IsGraphDirty() const
@@ -198,12 +215,22 @@ FNiagaraScriptViewModel::~FNiagaraScriptViewModel()
 		}
 	}
 
+	for (int32 i = 0; i < Scripts.Num(); i++)
+	{
+		if (Scripts[i].IsValid())
+		{
+			Scripts[i]->OnVMScriptCompiled().RemoveAll(this);
+		}
+	}
+
 	GEditor->UnregisterForUndo(this);
 
 	for (TNiagaraViewModelManager<UNiagaraScript, FNiagaraScriptViewModel>::Handle RegisteredHandle : RegisteredHandles)
 	{
 		UnregisterViewModelWithMap(RegisteredHandle);
 	}
+	//UE_LOG(LogNiagaraEditor, Warning, TEXT("Deleting script view model %p"), this);
+
 }
 
 
@@ -223,6 +250,14 @@ void FNiagaraScriptViewModel::SetScripts(UNiagaraScriptSource* InScriptSource, T
 		Source = nullptr;
 	}
 
+	for (int32 i = 0; i < Scripts.Num(); i++)
+	{
+		if (Scripts[i].IsValid())
+		{
+			Scripts[i]->OnVMScriptCompiled().RemoveAll(this);
+		}
+	}
+
 	for (TNiagaraViewModelManager<UNiagaraScript, FNiagaraScriptViewModel>::Handle RegisteredHandle : RegisteredHandles)
 	{
 		UnregisterViewModelWithMap(RegisteredHandle);
@@ -232,8 +267,9 @@ void FNiagaraScriptViewModel::SetScripts(UNiagaraScriptSource* InScriptSource, T
 	Scripts.Empty();
 	for (UNiagaraScript* Script : InScripts)
 	{
-		Scripts.Add(Script);
+		int32 i = Scripts.Add(Script);
 		check(Script->GetSource() == InScriptSource);
+		Scripts[i]->OnVMScriptCompiled().AddRaw(this, &FNiagaraScriptViewModel::OnVMScriptCompiled);
 	}
 	Source = InScriptSource;
 
@@ -246,17 +282,11 @@ void FNiagaraScriptViewModel::SetScripts(UNiagaraScriptSource* InScriptSource, T
 	{
 		// The underlying graph may have changed after the previous call.
 		Graph = Source->NodeGraph;
-		if (Graph != nullptr)
-		{
-			OnGraphChangedHandle = Graph->AddOnGraphChangedHandler(
-				FOnGraphChanged::FDelegate::CreateRaw(this, &FNiagaraScriptViewModel::OnGraphChanged));
-		}
 	}
 	MetaDataCollectionViewModel->SetGraph(Graph);
 
 	// Guess at initial compile status
 	LastCompileStatus = ENiagaraScriptCompileStatus::NCS_UpToDate;
-	bNeedsSave = false;
 
 	CompileStatuses.Empty();
 	CompileErrors.Empty();
@@ -266,9 +296,8 @@ void FNiagaraScriptViewModel::SetScripts(UNiagaraScriptSource* InScriptSource, T
 	{
 		FString Message;
 		ENiagaraScriptCompileStatus ScriptStatus = Scripts[i]->GetLastCompileStatus();
-		if (Scripts[i].IsValid() && Scripts[i]->IsCompilable() && Scripts[i]->GetByteCode().Num() == 0) // This is either a brand new script or failed in the past. Since we create a default working script, assume invalid.
+		if (Scripts[i].IsValid() && Scripts[i]->IsCompilable() && Scripts[i]->GetVMExecutableData().IsValid() && Scripts[i]->GetVMExecutableData().ByteCode.Num() == 0) // This is either a brand new script or failed in the past. Since we create a default working script, assume invalid.
 		{
-			bNeedsSave = true;
 			LastCompileStatus = LastCompileStatus;
 			Message = TEXT("Please recompile for full error stack.");
 			GraphViewModel->SetErrorTextToolTip(Message);
@@ -286,7 +315,6 @@ void FNiagaraScriptViewModel::SetScripts(UNiagaraScriptSource* InScriptSource, T
 			}
 			else
 			{
-				bNeedsSave = true;
 				LastCompileStatus = FNiagaraEditorUtilities::UnionCompileStatus(LastCompileStatus, ENiagaraScriptCompileStatus::NCS_Error);
 				ScriptStatus = ENiagaraScriptCompileStatus::NCS_Error;
 				Message = TEXT("Please recompile for full error stack.");
@@ -336,25 +364,15 @@ void FNiagaraScriptViewModel::SetScript(UNiagaraScript* InScript)
 	SetScripts(InSource, InScripts);
 }
 
-void FNiagaraScriptViewModel::MarkAllDirty()
+void FNiagaraScriptViewModel::MarkAllDirty(FString Reason)
 {
 	for (int32 i = 0; i < Scripts.Num(); i++)
 	{
 		if (Scripts[i].IsValid() && Scripts[i]->IsCompilable())
 		{
-			Scripts[i]->MarkScriptAndSourceDesynchronized();
+			Scripts[i]->MarkScriptAndSourceDesynchronized(Reason);
 		}
 	}
-	bNeedsSave = true;
-}
-
-void FNiagaraScriptViewModel::OnGraphChanged(const struct FEdGraphEditAction& InAction)
-{
-	if ((InAction.Action & GRAPHACTION_AddNode) != 0 || (InAction.Action & GRAPHACTION_RemoveNode) != 0 ||
-		(InAction.Action & GRAPHACTION_GenericNeedsRecompile) != 0)
-	{
-		MarkAllDirty();
-	}
 }
 
 TSharedRef<FNiagaraScriptInputCollectionViewModel> FNiagaraScriptViewModel::GetInputCollectionViewModel()
@@ -410,7 +428,6 @@ void FNiagaraScriptViewModel::UpdateCompileStatus(ENiagaraScriptCompileStatus In
 		}
 
 		LastCompileStatus = InAggregateCompileStatus;
-		bNeedsSave = true;
 		
 		if (LastCompileStatus == ENiagaraScriptCompileStatus::NCS_Error)
 		{
@@ -452,38 +469,6 @@ FText FNiagaraScriptViewModel::GetScriptErrors(ENiagaraScriptUsage InUsage, FGui
 	return Text;
 }
 
-void FNiagaraScriptViewModel::UpdateCompiledDataInterfaces(UNiagaraDataInterface* ChangedDataInterface)
-{
-	FName ChangedInputName = NAME_None;
-	TArray<UNiagaraNodeInput*> InputNodes;
-	GetGraphViewModel()->GetGraph()->GetNodesOfClass<UNiagaraNodeInput>(InputNodes);
-	for (UNiagaraNodeInput* InputNode : InputNodes)
-	{
-		if (InputNode->GetDataInterface() == ChangedDataInterface)
-		{
-			ChangedInputName = InputNode->Input.GetName();
-			break;
-		}
-	}
-
-	if (ChangedInputName != NAME_None)
-	{
-		for (TWeakObjectPtr<UNiagaraScript> Script : Scripts)
-		{
-			if (Script.IsValid())
-			{
-				for (FNiagaraScriptDataInterfaceInfo& DataInterfaceInfo : Script->DataInterfaceInfo)
-				{
-					if (DataInterfaceInfo.Name == ChangedInputName)
-					{
-						ChangedDataInterface->CopyTo(DataInterfaceInfo.DataInterface);
-					}
-				}
-			}
-		}
-	}
-}
-
 UNiagaraScript* FNiagaraScriptViewModel::GetContainerScript(ENiagaraScriptUsage InUsage, FGuid InUsageId) 
 {
 	for (int32 i = 0; i < Scripts.Num(); i++)
@@ -508,28 +493,13 @@ UNiagaraScript* FNiagaraScriptViewModel::GetScript(ENiagaraScriptUsage InUsage,
 	return nullptr;
 }
 
-void FNiagaraScriptViewModel::CompileStandaloneScript(bool bForce)
+void FNiagaraScriptViewModel::CompileStandaloneScript()
 {
 	if (Source.IsValid() && Source != nullptr && Scripts.Num() == 1 && Scripts[0].IsValid ())
 	{
 		if (Scripts[0]->IsStandaloneScript() && Scripts[0]->IsCompilable())
 		{
-			FString ErrorMsg;
-			LastCompileStatus = Scripts[0]->Compile(ErrorMsg, bForce);
-			Scripts[0]->CacheResourceShadersForRendering(true);
-			bNeedsSave = true;
-
-			if (LastCompileStatus == ENiagaraScriptCompileStatus::NCS_Error)
-			{
-				GraphViewModel->SetErrorTextToolTip(ErrorMsg + "\n(These same errors are also in the log)");
-			}
-			else
-			{
-				GraphViewModel->SetErrorTextToolTip("");
-			}
-
-			InputCollectionViewModel->RefreshParameterViewModels();
-			OutputCollectionViewModel->RefreshParameterViewModels();
+			Scripts[0]->RequestCompile();
 		}
 		else if (!Scripts[0]->IsCompilable())
 		{
@@ -568,11 +538,10 @@ void FNiagaraScriptViewModel::RefreshNodes()
 				{
 					if (Scripts[i].IsValid() && Scripts[i]->IsCompilable())
 					{
-						Scripts[i]->MarkScriptAndSourceDesynchronized();
+						Scripts[i]->MarkScriptAndSourceDesynchronized(TEXT("Nodes manually refreshed"));
 					}
 				}
 			}
-			bNeedsSave = true;
 		}
 	}
 }
@@ -648,10 +617,10 @@ void FNiagaraScriptViewModel::InputViewModelSelectionChanged()
 
 void FNiagaraScriptViewModel::InputParameterValueChanged(FName ParameterName)
 {
-	MarkAllDirty();
+	MarkAllDirty(TEXT("Input parameter value changed"));
 }
 
 void FNiagaraScriptViewModel::OutputParameterValueChanged(FName ParameterName)
 {
-	MarkAllDirty();
+	MarkAllDirty(TEXT("Output parameter value changed"));
 }
diff --git a/Engine/Plugins/FX/Niagara/Source/NiagaraEditor/Private/ViewModels/NiagaraScriptViewModel.h b/Engine/Plugins/FX/Niagara/Source/NiagaraEditor/Private/ViewModels/NiagaraScriptViewModel.h
--- a/Engine/Plugins/FX/Niagara/Source/NiagaraEditor/Private/ViewModels/NiagaraScriptViewModel.h
+++ b/Engine/Plugins/FX/Niagara/Source/NiagaraEditor/Private/ViewModels/NiagaraScriptViewModel.h
@@ -4,7 +4,7 @@
 
 #include "NiagaraParameterEditMode.h"
 #include "INiagaraCompiler.h"
-#include "TNiagaraViewModelManager.h"
+#include "ViewModels/TNiagaraViewModelManager.h"
 #include "EditorUndoClient.h"
 #include "Misc/NotifyHook.h"
 
@@ -53,7 +53,7 @@ public:
 		const TArray<UNiagaraScript*>& InCompileSources);
 
 	/** Compiles a script that isn't part of an emitter or System. */
-	void CompileStandaloneScript(bool bForce);
+	void CompileStandaloneScript();
 
 	/** Get the latest status of this view-model's script compilation.*/
 	ENiagaraScriptCompileStatus GetLatestCompileStatus();
@@ -65,19 +65,12 @@ public:
 	virtual void PostUndo(bool bSuccess) override;
 	virtual void PostRedo(bool bSuccess) override { PostUndo(bSuccess); }
 
-	bool GetScriptDirty() const { return bNeedsSave; }
-
-	void SetScriptDirty(bool InNeedsSave) { bNeedsSave = InNeedsSave; }
-
 	UNiagaraScript* GetContainerScript(ENiagaraScriptUsage InUsage, FGuid InUsageId = FGuid());
 	UNiagaraScript* GetScript(ENiagaraScriptUsage InUsage, FGuid InUsageId = FGuid());
 
 	ENiagaraScriptCompileStatus GetScriptCompileStatus(ENiagaraScriptUsage InUsage, FGuid InUsageId) const;
 	FText GetScriptErrors(ENiagaraScriptUsage InUsage, FGuid InUsageId) const;
 
-	/** Updates the compiled versions of data interfaces from changes to their source. */
-	void UpdateCompiledDataInterfaces(UNiagaraDataInterface* ChangedDataInterface);
-
 	/** If this is editing a standalone script, returns the script being edited.*/
 	UNiagaraScript* GetStandaloneScript();
 
@@ -89,10 +82,7 @@ private:
 	void InputViewModelSelectionChanged();
 
 	/** Marks this script view model as dirty and marks the scripts as needing synchrnozation. */
-	void MarkAllDirty();
-
-    /** Handle the graph being changed by the UI notifications to see if we need to mark as needing recompile.*/
-	void OnGraphChanged(const struct FEdGraphEditAction& InAction);
+	void MarkAllDirty(FString Reason);
 
 	void SetScripts(UNiagaraScriptSource* InScriptSource, TArray<UNiagaraScript*>& InScripts);
 
@@ -106,6 +96,8 @@ protected:
 	/** The script which provides the data for this view model. */
 	TArray<TWeakObjectPtr<UNiagaraScript>> Scripts;
 
+	void OnVMScriptCompiled(UNiagaraScript* InScript);
+
 	TWeakObjectPtr<UNiagaraScriptSource> Source;
 
 	/** The view model for the input parameter collection. */
@@ -129,9 +121,6 @@ protected:
 	/** The handle to the graph changed delegate needed for removing. */
 	FDelegateHandle OnGraphChangedHandle;
 
-	/** An edit has been made since the last save.*/
-	bool bNeedsSave;
-
 	TArray<TNiagaraViewModelManager<UNiagaraScript, FNiagaraScriptViewModel>::Handle> RegisteredHandles;
 
 	bool IsGraphDirty() const;
diff --git a/Engine/Plugins/FX/Niagara/Source/NiagaraEditor/Private/ViewModels/NiagaraSystemScriptViewModel.cpp b/Engine/Plugins/FX/Niagara/Source/NiagaraEditor/Private/ViewModels/NiagaraSystemScriptViewModel.cpp
--- a/Engine/Plugins/FX/Niagara/Source/NiagaraEditor/Private/ViewModels/NiagaraSystemScriptViewModel.cpp
+++ b/Engine/Plugins/FX/Niagara/Source/NiagaraEditor/Private/ViewModels/NiagaraSystemScriptViewModel.cpp
@@ -18,9 +18,9 @@
 #include "GraphEditAction.h"
 #include "ViewModels/Stack/NiagaraStackEntry.h"
 #include "ViewModels/Stack/NiagaraStackGraphUtilities.h"
-#include "NiagaraSystemViewModel.h"
-#include "NiagaraEmitterHandleViewModel.h"
-#include "NiagaraEmitterViewModel.h"
+#include "ViewModels/NiagaraSystemViewModel.h"
+#include "ViewModels/NiagaraEmitterHandleViewModel.h"
+#include "ViewModels/NiagaraEmitterViewModel.h"
 
 FNiagaraSystemScriptViewModel::FNiagaraSystemScriptViewModel(UNiagaraSystem& InSystem, FNiagaraSystemViewModel* InParent)
 	: FNiagaraScriptViewModel(InSystem.GetSystemSpawnScript(), NSLOCTEXT("SystemScriptViewModel", "GraphName", "System"), ENiagaraParameterEditMode::EditAll)
@@ -33,19 +33,95 @@ FNiagaraSystemScriptViewModel::FNiagaraSystemScriptViewModel(UNiagaraSystem& InS
 	{
 		OnGraphChangedHandle = GetGraphViewModel()->GetGraph()->AddOnGraphChangedHandler(
 			FOnGraphChanged::FDelegate::CreateRaw(this, &FNiagaraSystemScriptViewModel::OnGraphChanged));
+
+		OnRecompileHandle = GetGraphViewModel()->GetGraph()->AddOnGraphNeedsRecompileHandler(
+			FOnGraphChanged::FDelegate::CreateRaw(this, &FNiagaraSystemScriptViewModel::OnGraphChanged));
 		
 		GetGraphViewModel()->SetErrorTextToolTip("");
 	}
+
+	System.OnSystemCompiled().AddRaw(this, &FNiagaraSystemScriptViewModel::OnSystemVMCompiled);
 }
 
 FNiagaraSystemScriptViewModel::~FNiagaraSystemScriptViewModel()
 {
+	System.OnSystemCompiled().RemoveAll(this);
 	if (GetGraphViewModel()->GetGraph())
 	{
 		GetGraphViewModel()->GetGraph()->RemoveOnGraphChangedHandler(OnGraphChangedHandle);
+		GetGraphViewModel()->GetGraph()->RemoveOnGraphNeedsRecompileHandler(OnRecompileHandle);
+		
 	}
 }
 
+void FNiagaraSystemScriptViewModel::OnSystemVMCompiled(UNiagaraSystem* InSystem)
+{
+	if (InSystem != &System)
+	{
+		return;
+	}
+
+	TArray<ENiagaraScriptCompileStatus> InCompileStatuses;
+	TArray<FString> InCompileErrors;
+	TArray<FString> InCompilePaths;
+	TArray<TPair<ENiagaraScriptUsage, int32> > InUsages;
+
+	ENiagaraScriptCompileStatus AggregateStatus = ENiagaraScriptCompileStatus::NCS_UpToDate;
+	FString AggregateErrors;
+
+	TArray<UNiagaraScript*> SystemScripts;
+	SystemScripts.Add(InSystem->GetSystemSpawnScript());
+	SystemScripts.Add(InSystem->GetSystemUpdateScript());
+	for (const FNiagaraEmitterHandle& Handle : InSystem->GetEmitterHandles())
+	{
+		Handle.GetInstance()->GetScripts(SystemScripts, true);
+	}
+
+	int32 EventsFound = 0;
+	for (int32 i = 0; i < SystemScripts.Num(); i++)
+	{
+		UNiagaraScript* Script = SystemScripts[i];
+		if (Script != nullptr && Script->GetVMExecutableData().IsValid())
+		{
+			InCompileStatuses.Add(Script->GetVMExecutableData().LastCompileStatus);
+			InCompileErrors.Add(Script->GetVMExecutableData().ErrorMsg);
+			InCompilePaths.Add(Script->GetPathName());
+
+			if (Script->GetUsage() == ENiagaraScriptUsage::ParticleEventScript)
+			{
+				InUsages.Add(TPair<ENiagaraScriptUsage, int32>(Script->GetUsage(), EventsFound));
+				EventsFound++;
+			}
+			else
+			{
+				InUsages.Add(TPair<ENiagaraScriptUsage, int32>(Script->GetUsage(), 0));
+			}
+		}
+		else
+		{
+			InCompileStatuses.Add(ENiagaraScriptCompileStatus::NCS_Unknown);
+			InCompileErrors.Add(TEXT("Invalid script pointer!"));
+			InCompilePaths.Add(TEXT("Unknown..."));
+			InUsages.Add(TPair<ENiagaraScriptUsage, int32>(ENiagaraScriptUsage::Function, 0));
+		}
+	}
+
+	for (int32 i = 0; i < InCompileStatuses.Num(); i++)
+	{
+		AggregateStatus = FNiagaraEditorUtilities::UnionCompileStatus(AggregateStatus, InCompileStatuses[i]);
+		AggregateErrors += InCompilePaths[i] + TEXT(" ") + FNiagaraEditorUtilities::StatusToText(InCompileStatuses[i]).ToString() + TEXT("\n");
+		AggregateErrors += InCompileErrors[i] + TEXT("\n");
+	}
+
+	UpdateCompileStatus(AggregateStatus, AggregateErrors, InCompileStatuses, InCompileErrors, InCompilePaths, SystemScripts);
+
+	LastCompileStatus = AggregateStatus;
+
+	if (OnSystemCompiledDelegate.IsBound())
+	{
+		OnSystemCompiledDelegate.Broadcast();
+	}
+}
 
 float EmitterNodeVerticalOffset = 150.0f;
 
@@ -201,61 +277,7 @@ FNiagaraSystemScriptViewModel::FOnSystemCompiled& FNiagaraSystemScriptViewModel:
 
 void FNiagaraSystemScriptViewModel::CompileSystem(bool bForce)
 {
-
-	TArray<ENiagaraScriptCompileStatus> CompileSystemStatuses;
-	TArray<FString> CompileSystemErrors;
-	TArray<FString> CompileSystemPaths;
-	TArray<UNiagaraScript*> CompileSystemScripts;
-	System.CompileScripts(CompileSystemStatuses, CompileSystemErrors, CompileSystemPaths, CompileSystemScripts, bForce);
-
-	// First we need to find all the possibly affected view models..
-	TArray<FNiagaraScriptViewModel*> ScriptViewModels;
-	ScriptViewModels.Add(this);
-	if (Parent)
-	{
-		for (const TSharedRef<FNiagaraEmitterHandleViewModel>& EmitterHandleVM : Parent->GetEmitterHandleViewModels())
-		{
-			TSharedRef<FNiagaraEmitterViewModel> EmitterVM = EmitterHandleVM->GetEmitterViewModel();
-			ScriptViewModels.Add(&EmitterVM->GetSharedScriptViewModel().Get());
-		}
-	}
-
-	// Now we need to match up the results with the view models that need to know about them..
-	for (FNiagaraScriptViewModel* ScriptVM : ScriptViewModels)
-	{
-		TArray<ENiagaraScriptCompileStatus> CompileChildStatuses;
-		TArray<FString> CompileChildErrors;
-		TArray<FString> CompileChildPaths;
-		TArray<UNiagaraScript*> CompileChildScripts;
-
-		for (int32 i = 0; i < CompileSystemScripts.Num(); i++)
-		{
-			UNiagaraScript* Script = CompileSystemScripts[i];
-			if (ScriptVM->GetScript(Script->GetUsage(), Script->GetUsageId()) == Script) // Check to see if it contains this script
-			{
-				CompileChildScripts.Add(Script);
-				CompileChildStatuses.Add(CompileSystemStatuses[i]);
-				CompileChildErrors.Add(CompileSystemErrors[i]);
-				CompileChildPaths.Add(CompileSystemPaths[i]);
-			}
-		}
-		ENiagaraScriptCompileStatus AggregateStatus = ENiagaraScriptCompileStatus::NCS_UpToDate;
-		FString AggregateErrors;
-
-		for (int32 i = 0; i < CompileChildStatuses.Num(); i++)
-		{
-			AggregateStatus = FNiagaraEditorUtilities::UnionCompileStatus(AggregateStatus, CompileChildStatuses[i]);
-			AggregateErrors += CompileChildPaths[i] + TEXT(" ") + FNiagaraEditorUtilities::StatusToText(CompileChildStatuses[i]).ToString() + TEXT("\n");
-			AggregateErrors += CompileChildErrors[i] + TEXT("\n");
-		}
-		
-		ScriptVM->UpdateCompileStatus(AggregateStatus, AggregateErrors, CompileChildStatuses, CompileChildErrors, CompileChildPaths, CompileChildScripts);
-	}
-	
-	if (OnSystemCompiledDelegate.IsBound())
-	{
-		OnSystemCompiledDelegate.Broadcast();
-	}
+	System.RequestCompile(bForce);
 }
 
 void FNiagaraSystemScriptViewModel::OnGraphChanged(const struct FEdGraphEditAction& InAction)
diff --git a/Engine/Plugins/FX/Niagara/Source/NiagaraEditor/Private/ViewModels/NiagaraSystemScriptViewModel.h b/Engine/Plugins/FX/Niagara/Source/NiagaraEditor/Private/ViewModels/NiagaraSystemScriptViewModel.h
--- a/Engine/Plugins/FX/Niagara/Source/NiagaraEditor/Private/ViewModels/NiagaraSystemScriptViewModel.h
+++ b/Engine/Plugins/FX/Niagara/Source/NiagaraEditor/Private/ViewModels/NiagaraSystemScriptViewModel.h
@@ -29,6 +29,8 @@ private:
 	/** Called whenever the graph for the System script changes. */
 	void OnGraphChanged(const struct FEdGraphEditAction& InAction);
 
+	void OnSystemVMCompiled(UNiagaraSystem* InSystem);
+
 	FNiagaraSystemViewModel* Parent;
 
 	/** The System who's script is getting viewed and edited by this view model. */
@@ -36,6 +38,7 @@ private:
 
 	/** A handle to the on graph changed delegate. */
 	FDelegateHandle OnGraphChangedHandle;
+	FDelegateHandle OnRecompileHandle;
 
 	FOnSystemCompiled OnSystemCompiledDelegate;
 };
\ No newline at end of file
diff --git a/Engine/Plugins/FX/Niagara/Source/NiagaraEditor/Private/ViewModels/NiagaraSystemViewModel.cpp b/Engine/Plugins/FX/Niagara/Source/NiagaraEditor/Private/ViewModels/NiagaraSystemViewModel.cpp
--- a/Engine/Plugins/FX/Niagara/Source/NiagaraEditor/Private/ViewModels/NiagaraSystemViewModel.cpp
+++ b/Engine/Plugins/FX/Niagara/Source/NiagaraEditor/Private/ViewModels/NiagaraSystemViewModel.cpp
@@ -1,18 +1,19 @@
 // Copyright 1998-2018 Epic Games, Inc. All Rights Reserved.
 
-#include "NiagaraSystemViewModel.h"
+#include "ViewModels/NiagaraSystemViewModel.h"
 #include "NiagaraSystem.h"
 #include "NiagaraEmitterHandle.h"
 #include "NiagaraEmitter.h"
+#include "NiagaraEmitterEditorData.h"
 #include "NiagaraScriptSource.h"
 #include "NiagaraEditorUtilities.h"
-#include "NiagaraEmitterHandleViewModel.h"
-#include "NiagaraEmitterViewModel.h"
+#include "ViewModels/NiagaraEmitterHandleViewModel.h"
+#include "ViewModels/NiagaraEmitterViewModel.h"
+#include "ViewModels/Stack/NiagaraStackGraphUtilities.h"
 #include "NiagaraSystemScriptViewModel.h"
 #include "NiagaraScriptGraphViewModel.h"
 #include "NiagaraSequence.h"
 #include "MovieSceneNiagaraEmitterTrack.h"
-#include "MovieSceneNiagaraEmitterSection.h"
 #include "NiagaraSystemInstance.h"
 #include "NiagaraEditorSettings.h"
 #include "NiagaraSystemEditorData.h"
@@ -27,11 +28,12 @@
 #include "ISequencerModule.h"
 #include "EditorSupportDelegates.h"
 #include "Modules/ModuleManager.h"
-#include "TNiagaraViewModelManager.h"
+#include "ViewModels/TNiagaraViewModelManager.h"
 #include "Math/NumericLimits.h"
 #include "NiagaraComponent.h"
 #include "MovieSceneFolder.h"
 #include "Channels/MovieSceneChannelProxy.h"
+#include "Framework/Commands/UICommandList.h"
 
 DECLARE_CYCLE_STAT(TEXT("Niagara - SystemViewModel - CompileSystem"), STAT_NiagaraEditor_SystemViewModel_CompileSystem, STATGROUP_NiagaraEditor);
 
@@ -59,15 +61,27 @@ FNiagaraSystemViewModel::FNiagaraSystemViewModel(UNiagaraSystem& InSystem, FNiag
 	, bCanSimulate(InOptions.bCanSimulate)
 	, EditMode(InOptions.EditMode)
 	, OnGetSequencerAddMenuContent(InOptions.OnGetSequencerAddMenuContent)
-	, bUpdatingFromSequencerDataChange(false)
+	, bUpdatingEmittersFromSequencerDataChange(false)
+	, bUpdatingSequencerFromEmitterDataChange(false)
 	, bUpdatingSystemSelectionFromSequencer(false)
 	, bUpdatingSequencerSelectionFromSystem(false)
+	, EditorSettings(GetMutableDefault<UNiagaraEditorSettings>())
+	, bResetRequestPending(false)
+	, bCompilePendingCompletion(false)
 {
 	SetupPreviewComponentAndInstance();
 	SetupSequencer();
 	RefreshAll();
 	GEditor->RegisterForUndo(this);
 	RegisteredHandle = RegisterViewModelWithMap(&InSystem, this);
+	AddSystemEventHandlers();
+}
+
+void FNiagaraSystemViewModel::DumpToText(FString& ExportText)
+{
+	TSet<UObject*> ExportObjs;
+	ExportObjs.Add(&System);
+	FEdGraphUtilities::ExportNodesToText(ExportObjs, ExportText);
 }
 
 void FNiagaraSystemViewModel::Cleanup()
@@ -79,6 +93,12 @@ void FNiagaraSystemViewModel::Cleanup()
 		PreviewComponent->OnSystemInstanceChanged().RemoveAll(this);
 	}
 
+	if (SystemInstance)
+	{
+		SystemInstance->OnInitialized().RemoveAll(this);
+		SystemInstance->OnReset().RemoveAll(this);
+	}
+
 	CurveOwner.EmptyCurves();
 
 	GEditor->UnregisterForUndo(this);
@@ -91,6 +111,7 @@ void FNiagaraSystemViewModel::Cleanup()
 		HandleRef->OnPropertyChanged().RemoveAll(this);
 		HandleRef->GetEmitterViewModel()->OnPropertyChanged().RemoveAll(this);
 		HandleRef->GetEmitterViewModel()->OnScriptCompiled().RemoveAll(this);
+		HandleRef->Cleanup();
 	}
 	EmitterHandleViewModels.Empty();
 
@@ -104,6 +125,9 @@ void FNiagaraSystemViewModel::Cleanup()
 	}
 
 	PreviewComponent = nullptr;
+	SystemScriptViewModel.Reset();
+	
+	RemoveSystemEventHandlers();
 }
 
 FNiagaraSystemViewModel::~FNiagaraSystemViewModel()
@@ -118,7 +142,19 @@ const TArray<TSharedRef<FNiagaraEmitterHandleViewModel>>& FNiagaraSystemViewMode
 	return EmitterHandleViewModels;
 }
 
-TSharedRef<FNiagaraSystemScriptViewModel> FNiagaraSystemViewModel::GetSystemScriptViewModel()
+TSharedPtr<FNiagaraEmitterHandleViewModel> FNiagaraSystemViewModel::GetEmitterHandleViewModelById(FGuid InEmitterHandleId)
+{
+	for (TSharedRef<FNiagaraEmitterHandleViewModel> EmitterHandleViewModel : EmitterHandleViewModels)
+	{
+		if (EmitterHandleViewModel->GetId() == InEmitterHandleId)
+		{
+			return EmitterHandleViewModel;
+		}
+	}
+	return TSharedPtr<FNiagaraEmitterHandleViewModel>();
+}
+
+TSharedPtr<FNiagaraSystemScriptViewModel> FNiagaraSystemViewModel::GetSystemScriptViewModel()
 {
 	return SystemScriptViewModel;
 }
@@ -127,12 +163,14 @@ void FNiagaraSystemViewModel::CompileSystem(bool bForce)
 {
 	SCOPE_CYCLE_COUNTER(STAT_NiagaraEditor_SystemViewModel_CompileSystem);
 	KillSystemInstances();
+	check(SystemScriptViewModel.IsValid());
 	SystemScriptViewModel->CompileSystem(bForce);
-	OnSystemCompiledDelegate.Broadcast();
+	bCompilePendingCompletion = true;
 }
 
 ENiagaraScriptCompileStatus FNiagaraSystemViewModel::GetLatestCompileStatus()
 {
+	check(SystemScriptViewModel.IsValid());
 	return SystemScriptViewModel->GetLatestCompileStatus();
 }
 
@@ -260,11 +298,21 @@ void FNiagaraSystemViewModel::AddEmitter(UNiagaraEmitter& Emitter)
 		EmitterHandleNames.Add(EmitterHandle.GetName());
 	}
 
-	System.Modify();
-	FNiagaraEmitterHandle EmitterHandle = System.AddEmitterHandle(Emitter, FNiagaraUtilities::GetUniqueName(Emitter.GetFName(), EmitterHandleNames));
+	FNiagaraEmitterHandle EmitterHandle;
+	if (EditMode == ENiagaraSystemViewModelEditMode::SystemAsset)
+	{
+		System.Modify();
+		EmitterHandle = System.AddEmitterHandle(Emitter, FNiagaraUtilities::GetUniqueName(Emitter.GetFName(), EmitterHandleNames));
+	}
+	else
+	{
+		EmitterHandle = System.AddEmitterHandleWithoutCopying(Emitter);
+	}
+
+	check(SystemScriptViewModel.IsValid());
 	SystemScriptViewModel->RebuildEmitterNodes();
 
-	if (System.GetNumEmitters() == 1)
+	if (System.GetNumEmitters() == 1 && EditorSettings->GetAutoPlay())
 	{
 		// When adding a new emitter to an empty system start playing.
 		Sequencer->SetPlaybackStatus(EMovieScenePlayerStatus::Playing);
@@ -272,6 +320,22 @@ void FNiagaraSystemViewModel::AddEmitter(UNiagaraEmitter& Emitter)
 		
 	RefreshAll();
 
+	TRange<float> SystemPlaybackRange = GetEditorData().GetPlaybackRange();
+	TRange<float> EmitterPlaybackRange = GetEmitterHandleViewModelById(EmitterHandle.GetId())->GetEmitterViewModel()->GetEditorData().GetPlaybackRange();
+	TRange<float> NewSystemPlaybackRange = TRange<float>(
+		FMath::Min(SystemPlaybackRange.GetLowerBoundValue(), EmitterPlaybackRange.GetLowerBoundValue()),
+		FMath::Max(SystemPlaybackRange.GetUpperBoundValue(), EmitterPlaybackRange.GetUpperBoundValue()));
+
+	GetOrCreateEditorData().Modify();
+	GetOrCreateEditorData().SetPlaybackRange(NewSystemPlaybackRange);
+
+	TGuardValue<bool> UpdateGuard(bUpdatingSequencerFromEmitterDataChange, true);
+
+	FFrameTime NewStartFrame = NewSystemPlaybackRange.GetLowerBoundValue() * NiagaraSequence->GetMovieScene()->GetTickResolution();
+	int32 NewDuration = (NewSystemPlaybackRange.Size<float>() * NiagaraSequence->GetMovieScene()->GetTickResolution()).FrameNumber.Value;
+
+	NiagaraSequence->GetMovieScene()->SetPlaybackRange(NewStartFrame.RoundToFrame(), NewDuration);
+		
 	SetSelectedEmitterHandleById(EmitterHandle.GetId());
 
 	bForceAutoCompileOnce = true;
@@ -326,6 +390,7 @@ void FNiagaraSystemViewModel::DuplicateEmitters(TSet<FGuid> EmitterHandleIdsToDu
 		}
 	}
 
+	check(SystemScriptViewModel.IsValid());
 	SystemScriptViewModel->RebuildEmitterNodes();
 	RefreshAll();
 	bForceAutoCompileOnce = true;
@@ -352,6 +417,8 @@ void FNiagaraSystemViewModel::DeleteEmitters(TSet<FGuid> EmitterHandleIdsToDelet
 
 		System.Modify();
 		System.RemoveEmitterHandlesById(EmitterHandleIdsToDelete);
+
+		check(SystemScriptViewModel.IsValid());
 		SystemScriptViewModel->RebuildEmitterNodes();
 
 		RefreshAll();
@@ -359,6 +426,35 @@ void FNiagaraSystemViewModel::DeleteEmitters(TSet<FGuid> EmitterHandleIdsToDelet
 	}
 }
 
+TArray<TSharedRef<FNiagaraEmitterHandleViewModel>> FNiagaraSystemViewModel::GetPinnedEmitterHandles()
+{
+	return PinnedEmitterHandles;
+}
+
+void FNiagaraSystemViewModel::SetEmitterPinnedState(TSharedRef<FNiagaraEmitterHandleViewModel> EmitterHandleModel, bool bPinnedState)
+{
+	if (bPinnedState)
+	{
+		PinnedEmitterHandles.AddUnique(EmitterHandleModel);
+	}
+	else
+	{
+		PinnedEmitterHandles.Remove(EmitterHandleModel);
+	}
+	OnPinnedChangedDelegate.Broadcast();
+}
+
+FNiagaraSystemViewModel::FOnPinnedEmittersChanged& FNiagaraSystemViewModel::GetOnPinnedEmittersChanged()
+{
+	return OnPinnedChangedDelegate;
+}
+
+bool FNiagaraSystemViewModel::GetIsEmitterPinned(TSharedRef<FNiagaraEmitterHandleViewModel> EmitterHandleModel)
+{
+	return PinnedEmitterHandles.ContainsByPredicate([=](TSharedRef<FNiagaraEmitterHandleViewModel> Model) {return Model==EmitterHandleModel;});
+}
+
+
 FNiagaraSystemViewModel::FOnEmitterHandleViewModelsChanged& FNiagaraSystemViewModel::OnEmitterHandleViewModelsChanged()
 {
 	return OnEmitterHandleViewModelsChangedDelegate;
@@ -404,6 +500,8 @@ void FNiagaraSystemViewModel::PostUndo(bool bSuccess)
 void FNiagaraSystemViewModel::Tick(float DeltaTime)
 {
 	bool bRecompile = false;
+
+	check(SystemScriptViewModel.IsValid());
 	if (SystemScriptViewModel->GetLatestCompileStatus() == ENiagaraScriptCompileStatus::NCS_Dirty)
 	{
 		//SystemScriptViewModel->CompileSystem();
@@ -421,17 +519,63 @@ void FNiagaraSystemViewModel::Tick(float DeltaTime)
 		}
 	}
 
-	if (bRecompile || bForceAutoCompileOnce )
+	if (System.HasOutstandingCompilationRequests() == false)
 	{
-		CompileSystem(false);
+		if (bCompilePendingCompletion)
+		{
+			bCompilePendingCompletion = false;
+			OnSystemCompiled().Broadcast();
+		}
+
+		if (bRecompile || bForceAutoCompileOnce)
+		{
+			CompileSystem(false);
+			bForceAutoCompileOnce = false;
+		}
 
-		bForceAutoCompileOnce = false;
+		if (bResetRequestPending)
+		{
+			ResetSystem();
+		}
 	}
 }
 
+void FNiagaraSystemViewModel::OnPreSave()
+{
+	if (System.HasOutstandingCompilationRequests())
+	{
+		UE_LOG(LogNiagaraEditor, Log, TEXT("System %s has pending compile jobs. Waiting for that code to complete before Saving.."), *System.GetName());
+		System.WaitForCompilationComplete();
+	}
+}
+
+void FNiagaraSystemViewModel::OnPreClose()
+{
+	if (System.HasOutstandingCompilationRequests())
+	{
+		UE_LOG(LogNiagaraEditor, Log, TEXT("System %s has pending compile jobs. Waiting for that code to complete before Closing.."), *System.GetName());
+		System.WaitForCompilationComplete();
+	}
+}
+
+TSharedPtr<FUICommandList> FNiagaraSystemViewModel::GetToolkitCommands()
+{
+	return ToolkitCommands.Pin();
+}
+
+FNiagaraSystemViewModel::FOnPinnedCurvesChanged& FNiagaraSystemViewModel::GetOnPinnedCurvesChanged()
+{
+	return OnPinnedCurvesChangedDelegate;
+}
+
+void FNiagaraSystemViewModel::SetToolkitCommands(const TSharedRef<FUICommandList>& InToolkitCommands)
+{
+	ToolkitCommands = InToolkitCommands;
+}
+
 bool FNiagaraSystemViewModel::IsTickable() const
 {
-	return bForceAutoCompileOnce || (GetDefault<UNiagaraEditorSettings>()->bAutoCompile && bCanAutoCompile);
+	return bForceAutoCompileOnce || (GetDefault<UNiagaraEditorSettings>()->GetAutoCompile() && bCanAutoCompile);
 }
 
 TStatId FNiagaraSystemViewModel::GetStatId() const
@@ -448,7 +592,8 @@ void FNiagaraSystemViewModel::SetupPreviewComponentAndInstance()
 		PreviewComponent->bCastDynamicShadow = 1;
 		PreviewComponent->SetAsset(&System);
 		PreviewComponent->SetForceSolo(true);
-		PreviewComponent->SetAgeUpdateMode(UNiagaraComponent::EAgeUpdateMode::DesiredAge);
+		PreviewComponent->SetAgeUpdateMode(ENiagaraAgeUpdateMode::DesiredAge);
+		PreviewComponent->SetCanRenderWhileSeeking(false);
 		PreviewComponent->Activate(true);
 
 		UNiagaraSystemEditorData& EditorData = GetOrCreateEditorData();
@@ -479,17 +624,18 @@ void FNiagaraSystemViewModel::NotifyDataObjectChanged(UObject* ChangedObject)
 	UNiagaraDataInterfaceCurveBase* ChangedDataInterfaceCurve = Cast<UNiagaraDataInterfaceCurveBase>(ChangedDataInterface);
 	if (ChangedDataInterfaceCurve || ChangedObject == nullptr)
 	{
+		TArray<UNiagaraDataInterfaceCurveBase*> OldShownCurveDataInterfaces = ShownCurveDataInterfaces;
 		ResetCurveData();
+		if(ChangedDataInterfaceCurve != nullptr && ChangedDataInterfaceCurve->ShowInCurveEditor && OldShownCurveDataInterfaces.Contains(ChangedDataInterfaceCurve) == false)
+		{
+			NotifyPinnedCurvesChanged();
+		}
 	}
-	ReInitializeSystemInstances();
-}
 
-bool FNiagaraSystemViewModel::CanExecuteToggleIsolation()
-{
-	return SelectedEmitterHandleIds.Num() > 0;
+	ReInitializeSystemInstances();
 }
 
-void FNiagaraSystemViewModel::IsolateSelectedEmitters()
+void FNiagaraSystemViewModel::IsolateEmitters(TArray<TSharedRef<FNiagaraEmitterHandleViewModel>> EmitterHandlesToIsolate)
 {
 	for (TSharedRef<FNiagaraEmitterHandleViewModel> EmitterHandle : EmitterHandleViewModels)
 	{
@@ -497,66 +643,37 @@ void FNiagaraSystemViewModel::IsolateSelectedEmitters()
 	}
 
 	bool bAnyEmitterIsolated = false;
-	TArray<TSharedRef<FNiagaraEmitterHandleViewModel>> SelectedEmitterHandleViewModels;
-	GetSelectedEmitterHandles(SelectedEmitterHandleViewModels);
-	for (TSharedRef<FNiagaraEmitterHandleViewModel> SelectedEmitterHandle : SelectedEmitterHandleViewModels)
+	for (TSharedRef<FNiagaraEmitterHandleViewModel> EmitterHandleToIsolate : EmitterHandlesToIsolate)
 	{
 		bAnyEmitterIsolated = true;
-		SelectedEmitterHandle->GetEmitterHandle()->SetIsolated(true);
+		EmitterHandleToIsolate->GetEmitterHandle()->SetIsolated(true);
 	}
 
 	System.SetIsolateEnabled(bAnyEmitterIsolated);
 }
 
-void FNiagaraSystemViewModel::IsolateEmitterToggle()
+void FNiagaraSystemViewModel::ToggleEmitterIsolation(TSharedRef<FNiagaraEmitterHandleViewModel> InEmitterHandle)
 {
+	InEmitterHandle->GetEmitterHandle()->SetIsolated(!InEmitterHandle->GetEmitterHandle()->IsIsolated());
+	
 	bool bAnyEmitterIsolated = false;
-	TArray<TSharedRef<FNiagaraEmitterHandleViewModel>> SelectedEmitterHandleViewModels;
-	GetSelectedEmitterHandles(SelectedEmitterHandleViewModels);
-	for (TSharedRef<FNiagaraEmitterHandleViewModel> SelectedEmitterHandle : SelectedEmitterHandleViewModels)
+	for (TSharedRef<FNiagaraEmitterHandleViewModel> EmitterHandle : EmitterHandleViewModels)
 	{
-		const bool bShouldIsolate = !SelectedEmitterHandle->GetEmitterHandle()->IsIsolated();
-		SelectedEmitterHandle->GetEmitterHandle()->SetIsolated(bShouldIsolate);
-		if (!bAnyEmitterIsolated && bShouldIsolate)
+		if (EmitterHandle->GetEmitterHandle()->IsIsolated())
 		{
 			bAnyEmitterIsolated = true;
-		}
-	}
-
-	// If all the selected emitters are set to be not isolated we have to check if there is any emitter that is isolated.
-	if (!bAnyEmitterIsolated)
-	{
-		for (TSharedRef<FNiagaraEmitterHandleViewModel> EmitterHandle : EmitterHandleViewModels)
-		{
-			if (EmitterHandle->GetEmitterHandle()->IsIsolated())
-			{
-				bAnyEmitterIsolated = true;
-				break;
-			}
+			break;
 		}
 	}
 
 	System.SetIsolateEnabled(bAnyEmitterIsolated);
 }
 
-bool FNiagaraSystemViewModel::IsEmitterIsolationActive()
+bool FNiagaraSystemViewModel::IsEmitterIsolated(TSharedRef<FNiagaraEmitterHandleViewModel> InEmitterHandle)
 {
-	bool bShowActive = true;
-	TArray<TSharedRef<FNiagaraEmitterHandleViewModel>> SelectedEmitterHandleViewModels;
-	GetSelectedEmitterHandles(SelectedEmitterHandleViewModels);
-	for (TSharedRef<FNiagaraEmitterHandleViewModel> SelectedEmitterHandle : SelectedEmitterHandleViewModels)
-	{
-		if (!SelectedEmitterHandle->GetEmitterHandle()->IsIsolated())
-		{
-			bShowActive = false;
-			break;
-		}
-	}
-
-	return bShowActive;
+	return InEmitterHandle->GetEmitterHandle()->IsIsolated();
 }
 
-
 void FNiagaraSystemViewModel::RefreshEmitterHandleViewModels()
 {
 	TArray<TSharedRef<FNiagaraEmitterHandleViewModel>> OldViewModels = EmitterHandleViewModels;
@@ -580,6 +697,8 @@ void FNiagaraSystemViewModel::RefreshEmitterHandleViewModels()
 			ViewModel->OnNameChanged().AddRaw(this, &FNiagaraSystemViewModel::EmitterHandleNameChanged, ViewModel);
 			ViewModel->GetEmitterViewModel()->OnPropertyChanged().AddRaw(this, &FNiagaraSystemViewModel::EmitterPropertyChanged, ViewModel);
 			ViewModel->GetEmitterViewModel()->OnScriptCompiled().AddRaw(this, &FNiagaraSystemViewModel::ScriptCompiled);
+			ViewModel->GetEmitterViewModel()->OnScriptGraphChanged().AddRaw(this, &FNiagaraSystemViewModel::EmitterScriptGraphChanged, ViewModel);
+			ViewModel->GetEmitterViewModel()->OnScriptParameterStoreChanged().AddRaw(this, &FNiagaraSystemViewModel::EmitterParameterStoreChanged, ViewModel);
 			EmitterHandleViewModels.Add(ViewModel);
 		}
 		else
@@ -600,6 +719,8 @@ void FNiagaraSystemViewModel::RefreshEmitterHandleViewModels()
 		ViewModel->OnPropertyChanged().RemoveAll(this);
 		ViewModel->GetEmitterViewModel()->OnPropertyChanged().RemoveAll(this);
 		ViewModel->GetEmitterViewModel()->OnScriptCompiled().RemoveAll(this);
+		ViewModel->GetEmitterViewModel()->OnScriptGraphChanged().RemoveAll(this);
+		ViewModel->GetEmitterViewModel()->OnScriptParameterStoreChanged().RemoveAll(this);
 		ViewModel->Set(nullptr, nullptr, System);
 	}
 
@@ -614,9 +735,6 @@ void FNiagaraSystemViewModel::RefreshEmitterHandleViewModels()
 	}
 }
 
-TRange<float> SequencerDefaultPlaybackRange(0, 1000.0f);
-TRange<float> SequencerDefaultViewRange(0, 10.0f);
-
 void PopulateChildMovieSceneFoldersFromNiagaraFolders(const UNiagaraSystemEditorFolder* NiagaraFolder, UMovieSceneFolder* MovieSceneFolder, const TMap<FGuid, UMovieSceneNiagaraEmitterTrack*>& EmitterHandleIdToTrackMap)
 {
 	for (const UNiagaraSystemEditorFolder* ChildNiagaraFolder : NiagaraFolder->GetChildFolders())
@@ -652,6 +770,8 @@ void PopulateChildMovieSceneFoldersFromNiagaraFolders(const UNiagaraSystemEditor
 
 void FNiagaraSystemViewModel::RefreshSequencerTracks()
 {
+	TGuardValue<bool> UpdateGuard(bUpdatingSequencerFromEmitterDataChange, true);
+
 	TArray<UMovieSceneTrack*> MasterTracks = NiagaraSequence->GetMovieScene()->GetMasterTracks();
 	for (UMovieSceneTrack* MasterTrack : MasterTracks)
 	{
@@ -661,16 +781,11 @@ void FNiagaraSystemViewModel::RefreshSequencerTracks()
 		}
 	}
 
-	float MinEmitterTime = 0;
-	float MaxEmitterTime = 0;
 	TMap<FGuid, UMovieSceneNiagaraEmitterTrack*> EmitterHandleIdToTrackMap;
 	for (TSharedRef<FNiagaraEmitterHandleViewModel> EmitterHandleViewModel : EmitterHandleViewModels)
 	{
 		UMovieSceneNiagaraEmitterTrack* EmitterTrack = Cast<UMovieSceneNiagaraEmitterTrack>(NiagaraSequence->GetMovieScene()->AddMasterTrack(UMovieSceneNiagaraEmitterTrack::StaticClass()));
-		EmitterTrack->SetEmitterHandle(EmitterHandleViewModel);
-		RefreshSequencerTrack(EmitterTrack);
-		MinEmitterTime = FMath::Min(MinEmitterTime, EmitterHandleViewModel->GetEmitterViewModel()->GetStartTime());
-		MaxEmitterTime = FMath::Max(MaxEmitterTime, EmitterHandleViewModel->GetEmitterViewModel()->GetEndTime());
+		EmitterTrack->Initialize(*this, EmitterHandleViewModel, NiagaraSequence->GetMovieScene()->GetTickResolution());
 		EmitterHandleIdToTrackMap.Add(EmitterHandleViewModel->GetId(), EmitterTrack);
 	}
 
@@ -687,11 +802,6 @@ void FNiagaraSystemViewModel::RefreshSequencerTracks()
 		PopulateChildMovieSceneFoldersFromNiagaraFolders(RootChildFolder, MovieSceneRootFolder, EmitterHandleIdToTrackMap);
 	}
 
-	// Expand the view range to show all emitters.
-	FMovieSceneEditorData& EditorData = NiagaraSequence->GetMovieScene()->GetEditorData();
-	EditorData.ViewStart = FMath::Min<double>(MinEmitterTime, EditorData.ViewStart);
-	EditorData.ViewEnd   = FMath::Max<double>(MaxEmitterTime, EditorData.ViewEnd);
-
 	Sequencer->NotifyMovieSceneDataChanged(EMovieSceneDataChangeType::MovieSceneStructureItemsChanged);
 	Sequencer->SetGlobalTime(0);
 }
@@ -701,7 +811,7 @@ UMovieSceneNiagaraEmitterTrack* FNiagaraSystemViewModel::GetTrackForHandleViewMo
 	for (UMovieSceneTrack* Track : NiagaraSequence->GetMovieScene()->GetMasterTracks())
 	{
 		UMovieSceneNiagaraEmitterTrack* EmitterTrack = CastChecked<UMovieSceneNiagaraEmitterTrack>(Track);
-		if (EmitterTrack->GetEmitterHandle() == EmitterHandleViewModel)
+		if (EmitterTrack->GetEmitterHandleViewModel() == EmitterHandleViewModel)
 		{
 			return EmitterTrack;
 		}
@@ -709,84 +819,25 @@ UMovieSceneNiagaraEmitterTrack* FNiagaraSystemViewModel::GetTrackForHandleViewMo
 	return nullptr;
 }
 
-void FNiagaraSystemViewModel::RefreshSequencerTrack(UMovieSceneNiagaraEmitterTrack* EmitterTrack)
-{
-	if(EmitterTrack != nullptr)
-	{
-		TSharedPtr<FNiagaraEmitterHandleViewModel> EmitterHandleViewModel = EmitterTrack->GetEmitterHandle();
-		UNiagaraEmitter* Emitter = EmitterHandleViewModel->GetEmitterViewModel()->GetEmitter();
-
-		TArray<UMovieSceneSection*> Sections = EmitterTrack->GetAllSections();
-		UMovieSceneNiagaraEmitterSection* EmitterSection = Cast<UMovieSceneNiagaraEmitterSection>(Sections.Num() == 1 ? Sections[0] : nullptr);
-
-		if (EmitterSection != nullptr && Emitter != nullptr)
-		{
-			EmitterTrack->SetDisplayName(EmitterHandleViewModel->GetNameText());
-
-			bool bIsInfinite = EmitterHandleViewModel->GetEmitterViewModel()->GetStartTime() == 0
-				&& EmitterHandleViewModel->GetEmitterViewModel()->GetEndTime() == 0;
-			float StartTime;
-			float EndTime;
-
-// 			if (bIsInfinite && Emitter->Bursts.Num() > 0)
-// 			{
-// 				StartTime = TNumericLimits<float>::Max();
-// 				EndTime = TNumericLimits<float>::Lowest();
-// 				for (const FNiagaraEmitterBurst& Burst : Emitter->Bursts)
-// 				{
-// 					StartTime = FMath::Min(StartTime, Burst.Time);
-// 					EndTime = FMath::Max(EndTime, Burst.Time);
-// 				}
-// 			}
-// 			else
-			{
-				StartTime = EmitterHandleViewModel->GetEmitterViewModel()->GetStartTime();
-				EndTime = EmitterHandleViewModel->GetEmitterViewModel()->GetEndTime();
-			}
-
-			FFrameRate TickResolution = EmitterSection->GetTypedOuter<UMovieScene>()->GetTickResolution();
-
-			EmitterSection->SetEmitterHandle(EmitterHandleViewModel.ToSharedRef());
-			EmitterSection->SetIsActive(EmitterHandleViewModel->GetIsEnabled());
-			if (bIsInfinite)
-			{
-				EmitterSection->SetRange(TRange<FFrameNumber>((StartTime * TickResolution).FloorToFrame(), (EndTime * TickResolution).FloorToFrame()));
-			}
-			else
-			{
-				EmitterSection->SetRange(TRange<FFrameNumber>::All());
-			}
-
-			TMovieSceneChannelData<FMovieSceneBurstKey> Channel = EmitterSection->GetChannelProxy().GetChannel<FMovieSceneNiagaraEmitterChannel>(0)->GetData();
-
-// 			Channel.Reset();
-// 			for (const FNiagaraEmitterBurst& Burst : Emitter->Bursts)
-// 			{
-// 				FMovieSceneBurstKey BurstKey;
-// 				BurstKey.TimeRange = Burst.TimeRange;
-// 				BurstKey.SpawnMinimum = Burst.SpawnMinimum;
-// 				BurstKey.SpawnMaximum = Burst.SpawnMaximum;
-// 				Channel.AddKey(Burst.Time, BurstKey);
-// 			}
-		}
-	}
-}
-
 void FNiagaraSystemViewModel::SetupSequencer()
 {
 	NiagaraSequence = NewObject<UNiagaraSequence>(GetTransientPackage());
 	UMovieScene* MovieScene = NewObject<UMovieScene>(NiagaraSequence, FName("Niagara System MovieScene"), RF_Transactional);
+	MovieScene->SetDisplayRate(FFrameRate(240, 1));
 
 	NiagaraSequence->Initialize(this, MovieScene);
 
-	FFrameTime StartTime = SequencerDefaultPlaybackRange.GetLowerBoundValue() * MovieScene->GetTickResolution();
-	int32      Duration  = (SequencerDefaultPlaybackRange.Size<float>() * MovieScene->GetTickResolution()).FrameNumber.Value;
+	FFrameTime StartTime = GetEditorData().GetPlaybackRange().GetLowerBoundValue() * MovieScene->GetTickResolution();
+	int32      Duration  = (GetEditorData().GetPlaybackRange().Size<float>() * MovieScene->GetTickResolution()).FrameNumber.Value;
 
 	MovieScene->SetPlaybackRange(StartTime.RoundToFrame(), Duration);
 
 	FMovieSceneEditorData& EditorData = NiagaraSequence->GetMovieScene()->GetEditorData();
-	EditorData.ViewStart = EditorData.WorkStart = SequencerDefaultViewRange.GetLowerBoundValue();
-	EditorData.ViewEnd   = EditorData.WorkEnd   = SequencerDefaultViewRange.GetUpperBoundValue();
+	float ViewTimeOffset = .1f;
+	EditorData.WorkStart = GetEditorData().GetPlaybackRange().GetLowerBoundValue() - ViewTimeOffset;
+	EditorData.WorkEnd = GetEditorData().GetPlaybackRange().GetUpperBoundValue() + ViewTimeOffset;
+	EditorData.ViewStart = EditorData.WorkStart;
+	EditorData.ViewEnd = EditorData.WorkEnd;
 
 	FSequencerViewParams ViewParams(TEXT("NiagaraSequencerSettings"));
 	{
@@ -809,26 +860,34 @@ void FNiagaraSystemViewModel::SetupSequencer()
 	Sequencer->OnGlobalTimeChanged().AddRaw(this, &FNiagaraSystemViewModel::SequencerTimeChanged);
 	Sequencer->GetSelectionChangedTracks().AddRaw(this, &FNiagaraSystemViewModel::SequencerTrackSelectionChanged);
 	Sequencer->GetSelectionChangedSections().AddRaw(this, &FNiagaraSystemViewModel::SequencerSectionSelectionChanged);
-	Sequencer->SetPlaybackStatus(System.GetNumEmitters() > 0 
+	Sequencer->SetPlaybackStatus(System.GetNumEmitters() > 0 && EditorSettings->GetAutoPlay()
 		? EMovieScenePlayerStatus::Playing
 		: EMovieScenePlayerStatus::Stopped);
 }
 
 void FNiagaraSystemViewModel::ResetSystem()
 {
-	Sequencer->SetGlobalTime(0);
-	Sequencer->SetPlaybackStatus(EMovieScenePlayerStatus::Playing);
+	if (Sequencer->GetPlaybackStatus() == EMovieScenePlayerStatus::Playing || EditorSettings->GetResimulateOnChangeWhilePaused() == false)
+	{
+		Sequencer->SetGlobalTime(0);
+	}
 	
 	for (TObjectIterator<UNiagaraComponent> ComponentIt; ComponentIt; ++ComponentIt)
 	{
 		UNiagaraComponent* Component = *ComponentIt;
 		if (Component->GetAsset() == &System)
 		{
-			Component->SynchronizeWithSourceSystem();
 			Component->ResetSystem();
 		}
 	}
+
 	FEditorSupportDelegates::RedrawAllViewports.Broadcast();
+	bResetRequestPending = false;
+}
+
+void FNiagaraSystemViewModel::RequestResetSystem()
+{
+	bResetRequestPending = true;
 }
 
 void FNiagaraSystemViewModel::KillSystemInstances()
@@ -855,7 +914,6 @@ void FNiagaraSystemViewModel::ReInitializeSystemInstances()
 		UNiagaraComponent* Component = *ComponentIt;
 		if (Component->GetAsset() == &System)
 		{
-			Component->SynchronizeWithSourceSystem();
 			Component->ReinitializeSystem();
 		}
 	}
@@ -870,7 +928,34 @@ struct FNiagaraSystemCurveData
 	UObject* Owner;
 };
 
-void GetCurveData(FString CurveSource, UNiagaraGraph* SourceGraph, TArray<FNiagaraSystemCurveData>& OutCurveData)
+void GetCurveDataFromInterface(UNiagaraDataInterfaceCurveBase* CurveDataInterface, FString CurveSource, FString DefaultName,  
+	TArray<FNiagaraSystemCurveData>& OutCurveData, TArray<UNiagaraDataInterfaceCurveBase*>& OutCurveDataInterfaces)
+{
+	if (!CurveDataInterface->ShowInCurveEditor)
+	{
+		return;
+	}
+	OutCurveDataInterfaces.Add(CurveDataInterface);
+	TArray<UNiagaraDataInterfaceCurveBase::FCurveData> CurveData;
+	CurveDataInterface->GetCurveData(CurveData);
+	for (UNiagaraDataInterfaceCurveBase::FCurveData CurveDataItem : CurveData)
+	{
+		FNiagaraSystemCurveData SystemCurveData;
+		SystemCurveData.Curve = CurveDataItem.Curve;
+		SystemCurveData.Color = CurveDataItem.Color;
+		SystemCurveData.Owner = CurveDataInterface;
+		FString ParameterName = CurveDataItem.Name == NAME_None
+			? DefaultName
+			: DefaultName + ".";
+		FString DataName = CurveDataItem.Name == NAME_None
+			? FString()
+			: CurveDataItem.Name.ToString();
+		SystemCurveData.Name = *(CurveSource + ParameterName + DataName);
+		OutCurveData.Add(SystemCurveData);
+	}
+}
+
+void GetCurveData(FString CurveSource, UNiagaraGraph* SourceGraph, TArray<FNiagaraSystemCurveData>& OutCurveData, TArray<UNiagaraDataInterfaceCurveBase*>& OutCurveDataInterfaces)
 {
 	TArray<UNiagaraNodeInput*> InputNodes;
 	SourceGraph->GetNodesOfClass<UNiagaraNodeInput>(InputNodes);
@@ -884,23 +969,8 @@ void GetCurveData(FString CurveSource, UNiagaraGraph* SourceGraph, TArray<FNiaga
 				UNiagaraDataInterfaceCurveBase* CurveDataInterface = Cast<UNiagaraDataInterfaceCurveBase>(InputNode->GetDataInterface());
 				if (CurveDataInterface != nullptr)
 				{
-					TArray<UNiagaraDataInterfaceCurveBase::FCurveData> CurveData;
-					CurveDataInterface->GetCurveData(CurveData);
-					for (UNiagaraDataInterfaceCurveBase::FCurveData CurveDataItem : CurveData)
-					{
-						FNiagaraSystemCurveData SystemCurveData;
-						SystemCurveData.Curve = CurveDataItem.Curve;
-						SystemCurveData.Color = CurveDataItem.Color;
-						SystemCurveData.Owner = CurveDataInterface;
-						FString ParameterName = CurveDataItem.Name == NAME_None
-							? InputNode->Input.GetName().ToString()
-							: InputNode->Input.GetName().ToString() + ".";
-						FString DataName = CurveDataItem.Name == NAME_None
-							? FString()
-							: CurveDataItem.Name.ToString();
-						SystemCurveData.Name = *(CurveSource + ParameterName + DataName);
-						OutCurveData.Add(SystemCurveData);
-					}
+					FString DefaultName = InputNode->Input.GetName().ToString();
+					GetCurveDataFromInterface(CurveDataInterface, CurveSource, DefaultName, OutCurveData, OutCurveDataInterfaces);
 				}
 			}
 			HandledInputs.Add(InputNode->Input.GetName());
@@ -911,20 +981,34 @@ void GetCurveData(FString CurveSource, UNiagaraGraph* SourceGraph, TArray<FNiaga
 void FNiagaraSystemViewModel::ResetCurveData()
 {
 	CurveOwner.EmptyCurves();
+	ShownCurveDataInterfaces.Empty();
 
 	TArray<FNiagaraSystemCurveData> CurveData;
 
+
+	check(SystemScriptViewModel.IsValid()); 
 	GetCurveData(
 		TEXT("System"),
 		SystemScriptViewModel->GetGraphViewModel()->GetGraph(),
-		CurveData);
-
+		CurveData,
+		ShownCurveDataInterfaces);
+	// Get curves from user variables
+	for (UNiagaraDataInterface* DataInterface : System.GetExposedParameters().GetDataInterfaces())
+	{
+		UNiagaraDataInterfaceCurveBase* CurveDataInterface = Cast<UNiagaraDataInterfaceCurveBase>(DataInterface);
+		if (CurveDataInterface != nullptr)
+		{
+			GetCurveDataFromInterface(CurveDataInterface, TEXT("System"), TEXT("User"), CurveData, ShownCurveDataInterfaces);
+		}
+	}
+	
 	for (TSharedRef<FNiagaraEmitterHandleViewModel> EmitterHandleViewModel : EmitterHandleViewModels)
 	{
 		GetCurveData(
 			EmitterHandleViewModel->GetName().ToString(),
 			EmitterHandleViewModel->GetEmitterViewModel()->GetSharedScriptViewModel()->GetGraphViewModel()->GetGraph(),
-			CurveData);
+			CurveData, 
+			ShownCurveDataInterfaces);
 	}
 
 	for (FNiagaraSystemCurveData& CurveDataItem : CurveData)
@@ -940,22 +1024,85 @@ void FNiagaraSystemViewModel::ResetCurveData()
 	OnCurveOwnerChangedDelegate.Broadcast();
 }
 
+void UpdateCompiledDataInterfacesForScript(UNiagaraScript& TargetScript, FName TargetDataInterfaceName, UNiagaraDataInterface& SourceDataInterface)
+{
+	for (FNiagaraScriptDataInterfaceInfo& DataInterfaceInfo : TargetScript.GetCachedDefaultDataInterfaces())
+	{
+		if (DataInterfaceInfo.Name == TargetDataInterfaceName)
+		{
+			SourceDataInterface.CopyTo(DataInterfaceInfo.DataInterface);
+			break;
+		}
+	}
+}
+
 void FNiagaraSystemViewModel::UpdateCompiledDataInterfaces(UNiagaraDataInterface* ChangedDataInterface)
 {
-	SystemScriptViewModel->UpdateCompiledDataInterfaces(ChangedDataInterface);
-	for (TSharedRef<FNiagaraEmitterHandleViewModel> EmitterHandleViewModel : EmitterHandleViewModels)
+	UNiagaraNodeInput* OuterInputNode = ChangedDataInterface->GetTypedOuter<UNiagaraNodeInput>();
+	if (OuterInputNode != nullptr)
 	{
-		EmitterHandleViewModel->GetEmitterViewModel()->GetSharedScriptViewModel()->UpdateCompiledDataInterfaces(ChangedDataInterface);
+		// If the data interface's owning node has been removed from it's graph then it's not valid so early out here.
+		bool bIsValidInputNode = OuterInputNode->GetGraph()->Nodes.Contains(OuterInputNode);
+		if (bIsValidInputNode == false)
+		{
+			return;
+		}
+
+		// If the data interface was owned by an input node, then we need to try to update the compiled version.
+		UNiagaraEmitter* OwningEmitter;
+		UNiagaraScript* OwningScript;
+		FNiagaraStackGraphUtilities::GetOwningEmitterAndScriptForStackNode(*OuterInputNode, System, OwningEmitter, OwningScript);
+		if (ensureMsgf(OwningScript != nullptr, TEXT("Could not find owning script for data interface input node.")))
+		{
+			switch (OwningScript->GetUsage())
+			{
+			case ENiagaraScriptUsage::SystemSpawnScript:
+			case ENiagaraScriptUsage::SystemUpdateScript:
+			case ENiagaraScriptUsage::ParticleSpawnScript:
+			case ENiagaraScriptUsage::ParticleSpawnScriptInterpolated:
+			case ENiagaraScriptUsage::ParticleUpdateScript:
+			case ENiagaraScriptUsage::ParticleEventScript:
+				UpdateCompiledDataInterfacesForScript(*OwningScript, OuterInputNode->Input.GetName(), *ChangedDataInterface);
+				break;
+			case ENiagaraScriptUsage::EmitterSpawnScript:
+			case ENiagaraScriptUsage::EmitterUpdateScript:
+				if (ensureMsgf(OwningEmitter != nullptr, TEXT("Could not find owning emitter for data interface input node.")))
+				{
+					UNiagaraScript& TargetScript = OwningScript->GetUsage() == ENiagaraScriptUsage::EmitterSpawnScript ? *System.GetSystemSpawnScript() : *System.GetSystemUpdateScript();
+					FName AliasedInputNodeName = FNiagaraParameterMapHistory::ResolveEmitterAlias(OuterInputNode->Input.GetName(), OwningEmitter->GetUniqueEmitterName());
+					UpdateCompiledDataInterfacesForScript(TargetScript, AliasedInputNodeName, *ChangedDataInterface);
+				}
+				break;
+			}
+		}
+	}
+	else
+	{
+		// If the data interface wasn't owned by a script, try to find it in the exposed parameter data interfaces.
+		const FNiagaraVariable* FoundExposedDataInterface = System.GetExposedParameters().FindVariable(ChangedDataInterface);
+		if (FoundExposedDataInterface != nullptr)
+		{
+			System.GetExposedParameters().OnInterfaceChange();
+		}
 	}
 }
 
 void FNiagaraSystemViewModel::EmitterHandlePropertyChanged(TSharedRef<FNiagaraEmitterHandleViewModel> EmitterHandleViewModel)
 {
-	// When the emitter handle changes, refresh the System scripts emitter nodes just in cast the
+	// When the emitter handle changes, refresh the System scripts emitter nodes and the sequencer tracks just in case the
 	// property that changed was the handles emitter.
-	if (bUpdatingFromSequencerDataChange == false)
+	if (bUpdatingEmittersFromSequencerDataChange == false)
 	{
-		RefreshSequencerTrack(GetTrackForHandleViewModel(EmitterHandleViewModel));
+		TGuardValue<bool> UpdateGuard(bUpdatingSequencerFromEmitterDataChange, true);
+		for (UMovieSceneTrack* Track : NiagaraSequence->GetMovieScene()->GetMasterTracks())
+		{
+			UMovieSceneNiagaraEmitterTrack* EmitterTrack = CastChecked<UMovieSceneNiagaraEmitterTrack>(Track);
+			if (EmitterTrack->GetEmitterHandleViewModel() == EmitterHandleViewModel)
+			{
+				EmitterTrack->UpdateTrackFromEmitterGraphChange(NiagaraSequence->GetMovieScene()->GetTickResolution());
+			}
+		}
+		Sequencer->NotifyMovieSceneDataChanged(EMovieSceneDataChangeType::MovieSceneStructureItemsChanged);
 	}
 	ReInitializeSystemInstances();
 }
@@ -967,24 +1114,73 @@ void FNiagaraSystemViewModel::EmitterHandleNameChanged(TSharedRef<FNiagaraEmitte
 
 void FNiagaraSystemViewModel::EmitterPropertyChanged(TSharedRef<FNiagaraEmitterHandleViewModel> EmitterHandleViewModel)
 {
-	if (bUpdatingFromSequencerDataChange == false)
-	{
-		RefreshSequencerTrack(GetTrackForHandleViewModel(EmitterHandleViewModel));
-	}
 	ReInitializeSystemInstances();
 }
 
 void FNiagaraSystemViewModel::ScriptCompiled()
 {
-	ReInitializeSystemInstances();
+	//ReInitializeSystemInstances();
+}
+
+void FNiagaraSystemViewModel::SystemParameterStoreChanged(const FNiagaraParameterStore& ChangedParameterStore, const UNiagaraScript* OwningScript)
+{
+	UpdateSimulationFromParameterChange();
+}
+
+void FNiagaraSystemViewModel::EmitterScriptGraphChanged(const FEdGraphEditAction& InAction, const UNiagaraScript& OwningScript, const TSharedRef<FNiagaraEmitterHandleViewModel> OwningEmitterHandleViewModel)
+{
+	if (bUpdatingEmittersFromSequencerDataChange == false)
+	{
+		TGuardValue<bool> UpdateGuard(bUpdatingSequencerFromEmitterDataChange, true);
+		for (UMovieSceneTrack* Track : NiagaraSequence->GetMovieScene()->GetMasterTracks())
+		{
+			UMovieSceneNiagaraEmitterTrack* EmitterTrack = CastChecked<UMovieSceneNiagaraEmitterTrack>(Track);
+			if (EmitterTrack->GetEmitterHandleViewModel() == OwningEmitterHandleViewModel)
+			{
+				EmitterTrack->UpdateTrackFromEmitterGraphChange(NiagaraSequence->GetMovieScene()->GetTickResolution());
+			}
+		}
+		Sequencer->NotifyMovieSceneDataChanged(EMovieSceneDataChangeType::MovieSceneStructureItemsChanged);
+	}
+}
+
+void FNiagaraSystemViewModel::EmitterParameterStoreChanged(const FNiagaraParameterStore& ChangedParameterStore, const UNiagaraScript& OwningScript, const TSharedRef<FNiagaraEmitterHandleViewModel> OwningEmitterHandleViewModel)
+{
+	if (bUpdatingEmittersFromSequencerDataChange == false)
+	{
+		TGuardValue<bool> UpdateGuard(bUpdatingSequencerFromEmitterDataChange, true);
+		for (UMovieSceneTrack* Track : NiagaraSequence->GetMovieScene()->GetMasterTracks())
+		{
+			UMovieSceneNiagaraEmitterTrack* EmitterTrack = CastChecked<UMovieSceneNiagaraEmitterTrack>(Track);
+			EmitterTrack->UpdateTrackFromEmitterParameterChange(NiagaraSequence->GetMovieScene()->GetTickResolution());
+		}
+		Sequencer->NotifyMovieSceneDataChanged(EMovieSceneDataChangeType::TrackValueChanged);
+	}
+	UpdateSimulationFromParameterChange();
+}
+
+void FNiagaraSystemViewModel::UpdateSimulationFromParameterChange()
+{
+	if (EditorSettings->GetResetSimulationOnChange())
+	{
+		RequestResetSystem();
+	}
+	else
+	{
+		if (Sequencer->GetPlaybackStatus() == EMovieScenePlayerStatus::Stopped)
+		{
+			// TODO: Update the view when paused and reset on change is turned off.
+		}
+	}
 }
 
 void FNiagaraSystemViewModel::CurveChanged(FRichCurve* ChangedCurve, UObject* InCurveOwner)
 {
-	UNiagaraDataInterface* ChangedDataInterface = Cast<UNiagaraDataInterface>(InCurveOwner);
-	if (ChangedDataInterface != nullptr)
+	UNiagaraDataInterfaceCurveBase* CurveDataInterface = Cast<UNiagaraDataInterfaceCurveBase>(InCurveOwner);
+	if (CurveDataInterface != nullptr)
 	{
-		UpdateCompiledDataInterfaces(ChangedDataInterface);
+		CurveDataInterface->UpdateLUT();
+		UpdateCompiledDataInterfaces(CurveDataInterface);
 	}
 	ResetSystem();
 }
@@ -1031,7 +1227,7 @@ void PopulateNiagaraFoldersFromMovieSceneFolders(const TArray<UMovieSceneFolder*
 		UMovieSceneNiagaraEmitterTrack* NiagaraEmitterTrack = Cast<UMovieSceneNiagaraEmitterTrack>(MovieSceneTrack);
 		if (NiagaraEmitterTrack != nullptr)
 		{
-			FGuid EmitterHandleId = NiagaraEmitterTrack->GetEmitterHandle()->GetId();
+			FGuid EmitterHandleId = NiagaraEmitterTrack->GetEmitterHandleViewModel()->GetId();
 			ValidEmitterHandleIds.Add(EmitterHandleId);
 			if (ParentFolder->GetChildEmitterHandleIds().Contains(EmitterHandleId) == false)
 			{
@@ -1052,121 +1248,99 @@ void PopulateNiagaraFoldersFromMovieSceneFolders(const TArray<UMovieSceneFolder*
 
 void FNiagaraSystemViewModel::SequencerDataChanged(EMovieSceneDataChangeType DataChangeType)
 {
-	bUpdatingFromSequencerDataChange = true;
-	TSet<FGuid> VaildTrackEmitterHandleIds;
-	TSet<FGuid> EmittersToDuplicate;
-	for (UMovieSceneTrack* Track : NiagaraSequence->GetMovieScene()->GetMasterTracks())
+	if (bUpdatingSequencerFromEmitterDataChange == false && GIsTransacting == false)
 	{
-		UMovieSceneNiagaraEmitterTrack* EmitterTrack = CastChecked<UMovieSceneNiagaraEmitterTrack>(Track);
-		TArray<UMovieSceneSection*> Sections = EmitterTrack->GetAllSections();
-		UMovieSceneNiagaraEmitterSection* EmitterSection = Cast<UMovieSceneNiagaraEmitterSection>(Sections.Num() == 1 ? Sections[0] : nullptr);
-		if (EmitterSection != nullptr)
+		TGuardValue<bool> UpdateGuard(bUpdatingEmittersFromSequencerDataChange, true);
+
+		GetOrCreateEditorData().Modify();
+		TRange<FFrameNumber> FramePlaybackRange = NiagaraSequence->GetMovieScene()->GetPlaybackRange();
+		float StartTimeSeconds = NiagaraSequence->GetMovieScene()->GetTickResolution().AsSeconds(FramePlaybackRange.GetLowerBoundValue());
+		float EndTimeSeconds = NiagaraSequence->GetMovieScene()->GetTickResolution().AsSeconds(FramePlaybackRange.GetUpperBoundValue());
+		GetOrCreateEditorData().SetPlaybackRange(TRange<float>(StartTimeSeconds, EndTimeSeconds));
+
+		TSet<FGuid> VaildTrackEmitterHandleIds;
+		TSet<FGuid> EmittersToDuplicate;
+		TArray<TTuple<TSharedPtr<FNiagaraEmitterHandleViewModel>, FName>> EmitterHandlesToRename;
+
+		for (UMovieSceneTrack* Track : NiagaraSequence->GetMovieScene()->GetMasterTracks())
 		{
-			if (!EmitterTrack->GetEmitterHandle().IsValid())
+			UMovieSceneNiagaraEmitterTrack* EmitterTrack = CastChecked<UMovieSceneNiagaraEmitterTrack>(Track);
+			if (EmitterTrack->GetEmitterHandleViewModel().IsValid())
+			{
+				VaildTrackEmitterHandleIds.Add(EmitterTrack->GetEmitterHandleViewModel()->GetId());
+				EmitterTrack->UpdateEmitterHandleFromTrackChange(NiagaraSequence->GetMovieScene()->GetTickResolution());
+				EmitterTrack->GetEmitterHandleViewModel()->GetEmitterViewModel()->GetOrCreateEditorData().Modify();
+				EmitterTrack->GetEmitterHandleViewModel()->GetEmitterViewModel()->GetOrCreateEditorData().SetPlaybackRange(GetEditorData().GetPlaybackRange());
+				if (EmitterTrack->GetDisplayName().ToString() != EmitterTrack->GetEmitterHandleViewModel()->GetNameText().ToString())
+				{
+					EmitterHandlesToRename.Add(TTuple<TSharedPtr<FNiagaraEmitterHandleViewModel>, FName>(EmitterTrack->GetEmitterHandleViewModel(), *EmitterTrack->GetDisplayName().ToString()));
+				}
+			}
+			else
 			{
 				if (EmitterTrack->GetEmitterHandleId().IsValid())
 				{
 					// The emitter handle is invalid, but the track has a valid Id, most probably because of a copy/paste event
 					EmittersToDuplicate.Add(EmitterTrack->GetEmitterHandleId());
 				}
-				continue;
 			}
-			VaildTrackEmitterHandleIds.Add(EmitterTrack->GetEmitterHandle()->GetId());
-
-			TSharedPtr<FNiagaraEmitterHandleViewModel> EmitterHandleViewModel = EmitterTrack->GetEmitterHandle();
-			if (bCanModifyEmittersFromTimeline)
-			{
-				EmitterHandleViewModel->SetName(*EmitterTrack->GetDisplayName().ToString());
-			}
-			
-			EmitterTrack->SetDisplayName(EmitterHandleViewModel->GetNameText());
-			
-			EmitterHandleViewModel->SetIsEnabled(EmitterSection->IsActive());
-
-			TSharedPtr<FNiagaraEmitterViewModel> EmitterViewModel = EmitterTrack->GetEmitterHandle()->GetEmitterViewModel();
-
-			// @todo: sequencer-timecode: these don't seem to actually do anything???
-			// if (EmitterSection->GetRange() != TRange<FFrameNumber>::All())
-			// {
-			// 	EmitterViewModel->SetStartTime(EmitterSection->GetStartTime());
-			// 	EmitterViewModel->SetEndTime(EmitterSection->GetEndTime());
-			// }
+		}
 
-			const FMovieSceneNiagaraEmitterChannel* Channel = EmitterSection->GetChannelProxy().GetChannel<FMovieSceneNiagaraEmitterChannel>(0);
+		bool bRefreshAllTracks = EmitterHandlesToRename.Num() > 0;
 
-			TArrayView<const FFrameNumber>        Times  = Channel->GetTimes();
-			TArrayView<const FMovieSceneBurstKey> Values = Channel->GetValues();
+		for (TTuple<TSharedPtr<FNiagaraEmitterHandleViewModel>, FName>& EmitterHandletoRename : EmitterHandlesToRename)
+		{
+			EmitterHandletoRename.Get<0>()->SetName(EmitterHandletoRename.Get<1>());
+		}
 
-// 			UNiagaraEmitter* Emitter = EmitterViewModel->GetEmitter();
-// 			if (Emitter)
-// 			{
-// 				Emitter->Bursts.Empty();
-// 			}
+		TSet<FGuid> AllEmitterHandleIds;
+		for (TSharedRef<FNiagaraEmitterHandleViewModel> EmitterHandleViewModel : EmitterHandleViewModels)
+		{
+			AllEmitterHandleIds.Add(EmitterHandleViewModel->GetId());
+		}
 
-			FFrameRate TickResolution = EmitterSection->GetTypedOuter<UMovieScene>()->GetTickResolution();
-			for (int32 Index = 0; Index < Times.Num(); ++Index)
+		TSet<FGuid> RemovedEmitterHandleIds = AllEmitterHandleIds.Difference(VaildTrackEmitterHandleIds);
+		if (RemovedEmitterHandleIds.Num() > 0)
+		{
+			if (bCanModifyEmittersFromTimeline)
 			{
-				FNiagaraEmitterBurst Burst;
-				Burst.Time = Times[Index] / TickResolution;
-				Burst.TimeRange    = Values[Index].TimeRange / TickResolution;
-				Burst.SpawnMinimum = Values[Index].SpawnMinimum;
-				Burst.SpawnMaximum = Values[Index].SpawnMaximum;
-// 				if (Emitter)
-// 				{
-// 					Emitter->Bursts.Add(Burst);
-// 				}
+				DeleteEmitters(RemovedEmitterHandleIds);
+				// When deleting emitters from sequencer, select a new one if one is available.
+				if (SelectedEmitterHandleIds.Num() == 0 && EmitterHandleViewModels.Num() > 0)
+				{
+					SetSelectedEmitterHandleById(EmitterHandleViewModels[0]->GetId());
+				}
+			}
+			else
+			{
+				bRefreshAllTracks = true;
 			}
 		}
-	}
 
-	TSet<FGuid> AllEmitterHandleIds;
-	for (TSharedRef<FNiagaraEmitterHandleViewModel> EmitterHandleViewModel : EmitterHandleViewModels)
-	{
-		AllEmitterHandleIds.Add(EmitterHandleViewModel->GetId());
-	}
-
-	TSet<FGuid> RemovedEmitterHandleIds = AllEmitterHandleIds.Difference(VaildTrackEmitterHandleIds);
-	if (RemovedEmitterHandleIds.Num() > 0)
-	{
-		if (bCanModifyEmittersFromTimeline)
+		if (EmittersToDuplicate.Num() > 0)
 		{
-			DeleteEmitters(RemovedEmitterHandleIds);
-			// When deleting emitters from sequencer, select a new one if one is available.
-			if (SelectedEmitterHandleIds.Num() == 0 && EmitterHandleViewModels.Num() > 0)
+			if (bCanModifyEmittersFromTimeline)
 			{
-				SetSelectedEmitterHandleById(EmitterHandleViewModels[0]->GetId());
+				DuplicateEmitters(EmittersToDuplicate);
+			}
+			else
+			{
+				bRefreshAllTracks = true;
 			}
 		}
-		else
-		{
-			RefreshSequencerTracks();
-		}
-	}
-	
-	if (EmittersToDuplicate.Num() > 0)
-	{
-		if (bCanModifyEmittersFromTimeline)
+
+		TArray<UMovieSceneTrack*> RootTracks;
+		TArray<UMovieSceneFolder*> RootFolders = NiagaraSequence->GetMovieScene()->GetRootFolders();
+		if (RootFolders.Num() != 0 || GetEditorData().GetRootFolder().GetChildFolders().Num() != 0)
 		{
-			DuplicateEmitters(EmittersToDuplicate);
+			PopulateNiagaraFoldersFromMovieSceneFolders(RootFolders, RootTracks, &GetOrCreateEditorData().GetRootFolder());
 		}
-		else
+
+		if (bRefreshAllTracks)
 		{
 			RefreshSequencerTracks();
 		}
 	}
-
-	TArray<UMovieSceneTrack*> RootTracks;
-	TArray<UMovieSceneFolder*> RootFolders = NiagaraSequence->GetMovieScene()->GetRootFolders();
-	if (RootFolders.Num() != 0 || GetEditorData().GetRootFolder().GetChildFolders().Num() != 0)
-	{
-		PopulateNiagaraFoldersFromMovieSceneFolders(RootFolders, RootTracks, &GetOrCreateEditorData().GetRootFolder());
-	}
-
-	if (SystemInstance != nullptr)
-	{
-		SystemInstance->Reset(FNiagaraSystemInstance::EResetMode::ResetAll);
-	}
-	bUpdatingFromSequencerDataChange = false;
 }
 
 void FNiagaraSystemViewModel::SequencerTimeChanged()
@@ -1209,23 +1383,12 @@ void FNiagaraSystemViewModel::SequencerTimeChanged()
 						{
 							bSystemIsAlive |= true;
 						}
-
-						/*for (const TSharedRef<FNiagaraEmitterHandleViewModel> HandleViewModel : GetEmitterHandleViewModels())
-						{
-							TSharedRef<FNiagaraEmitterViewModel> EmitterViewModel = HandleViewModel->GetEmitterViewModel();
-							if (EmitterViewModel->GetEmitter() == Simulation->GetEmitterHandle().GetInstance())
-							{
-								float EmitterEndTime = EmitterViewModel->GetNumLoops() == 0
-								? EmitterViewModel->GetEndTime()
-								: (EmitterViewModel->GetEndTime() - EmitterViewModel->GetStartTime()) * EmitterViewModel->GetNumLoops();
-								MaxEmitterTime = FMath::Max(MaxEmitterTime, EmitterEndTime);
-
-							}
-						}*/
 					}
 				}
 
-				bool bIsUpdatingOrCanSpawn = bSystemIsAlive /*|| (MaxEmitterTime == 0.0f || CurrentSequencerTime <= MaxEmitterTime)*/;
+				float PlaybackRangeEnd = NiagaraSequence->GetMovieScene()->GetTickResolution().AsSeconds(
+					NiagaraSequence->GetMovieScene()->GetPlaybackRange().GetUpperBoundValue());
+				bool bIsUpdatingOrCanSpawn = bSystemIsAlive && CurrentSequencerTime < PlaybackRangeEnd;
 				if (bIsUpdatingOrCanSpawn)
 				{
 					// Skip the first update after going from stopped to playing because snapping in sequencer may have made the time
@@ -1247,28 +1410,21 @@ void FNiagaraSystemViewModel::SequencerTimeChanged()
 				bUpdateDesiredAge = !bStoppedPlaying;
 			}
 
-			if (CurrentSequencerTime < PreviousSequencerTime)
-			{
-				//Always reset if we've gone backwards.
-				bResetSystemInstance = true;
-			}
-
 			if (bUpdateDesiredAge)
 			{
-				PreviewComponent->SetDesiredAge(FMath::Max(CurrentSequencerTime, 0.0f));
-				if (CurrentSequencerTime < PreviewComponent->GetSystemInstance()->GetAge())
+				if (CurrentStatus == EMovieScenePlayerStatus::Playing)
 				{
-					bResetSystemInstance = true;
+					PreviewComponent->SetDesiredAge(FMath::Max(CurrentSequencerTime, 0.0f));
+				}
+				else
+				{
+					PreviewComponent->SeekToDesiredAge(FMath::Max(CurrentSequencerTime, 0.0f));
 				}
-			//SystemInstance->Tick(0);TODO: What was this for? Possibly reinstate but I don't like use dealing with the instance directly.
 			}
 
-			if (bResetSystemInstance || !PreviewComponent->IsActive())
+			if (bResetSystemInstance)
 			{
-				PreviewComponent->Activate(true);
-
-				TGuardValue<bool>(bSettingSequencerTimeDirectly, true);
-				Sequencer->SetLocalTime(0);
+				ResetSystem();
 			}
 		}
 	}
@@ -1304,9 +1460,9 @@ void FNiagaraSystemViewModel::UpdateEmitterHandleSelectionFromSequencer()
 	for (UMovieSceneTrack* SelectedTrack : SelectedTracks)
 	{
 		UMovieSceneNiagaraEmitterTrack* SelectedEmitterTrack = Cast<UMovieSceneNiagaraEmitterTrack>(SelectedTrack);
-		if (SelectedEmitterTrack != nullptr && SelectedEmitterTrack->GetEmitterHandle().IsValid())
+		if (SelectedEmitterTrack != nullptr && SelectedEmitterTrack->GetEmitterHandleViewModel().IsValid())
 		{
-			NewSelectedEmitterHandleIds.AddUnique(SelectedEmitterTrack->GetEmitterHandle()->GetId());
+			NewSelectedEmitterHandleIds.AddUnique(SelectedEmitterTrack->GetEmitterHandleViewModel()->GetId());
 		}
 	}
 
@@ -1314,10 +1470,10 @@ void FNiagaraSystemViewModel::UpdateEmitterHandleSelectionFromSequencer()
 	Sequencer->GetSelectedSections(SelectedSections);
 	for (UMovieSceneSection* SelectedSection : SelectedSections)
 	{
-		UMovieSceneNiagaraEmitterSection* SelectedEmitterSection = Cast<UMovieSceneNiagaraEmitterSection>(SelectedSection);
-		if (SelectedEmitterSection != nullptr && SelectedEmitterSection->GetEmitterHandle().IsValid())
+		UMovieSceneNiagaraEmitterSectionBase* SelectedEmitterSection = Cast<UMovieSceneNiagaraEmitterSectionBase>(SelectedSection);
+		if (SelectedEmitterSection != nullptr && SelectedEmitterSection->GetEmitterHandleViewModel().IsValid())
 		{
-			NewSelectedEmitterHandleIds.AddUnique(SelectedEmitterSection->GetEmitterHandle()->GetId());
+			NewSelectedEmitterHandleIds.AddUnique(SelectedEmitterSection->GetEmitterHandleViewModel()->GetId());
 		}
 	}
 
@@ -1334,7 +1490,7 @@ void FNiagaraSystemViewModel::UpdateSequencerFromEmitterHandleSelection()
 		for (UMovieSceneTrack* MasterTrack : NiagaraSequence->GetMovieScene()->GetMasterTracks())
 		{
 			UMovieSceneNiagaraEmitterTrack* EmitterTrack = Cast<UMovieSceneNiagaraEmitterTrack>(MasterTrack);
-			if (EmitterTrack != nullptr && EmitterTrack->GetEmitterHandle()->GetId() == SelectedEmitterHandleId)
+			if (EmitterTrack != nullptr && EmitterTrack->GetEmitterHandleViewModel()->GetId() == SelectedEmitterHandleId)
 			{
 				Sequencer->SelectTrack(EmitterTrack);
 			}
@@ -1406,4 +1562,55 @@ void FNiagaraSystemViewModel::UpdateEmitterFixedBounds()
 	}
 }
 
-#undef LOCTEXT_NAMESPACE // NiagaraSystemViewModel
+void FNiagaraSystemViewModel::AddSystemEventHandlers()
+{
+	if (System.IsValid())
+	{
+		TArray<UNiagaraScript*> Scripts;
+		Scripts.Add(System.GetSystemSpawnScript());
+		Scripts.Add(System.GetSystemUpdateScript());
+		
+		for (UNiagaraScript* Script : Scripts)
+		{
+			FDelegateHandle OnParameterStoreChangedHandle = Script->RapidIterationParameters.AddOnChangedHandler(
+				FNiagaraParameterStore::FOnChanged::FDelegate::CreateRaw<FNiagaraSystemViewModel, const FNiagaraParameterStore&, const UNiagaraScript*>(
+					this, &FNiagaraSystemViewModel::SystemParameterStoreChanged, Script->RapidIterationParameters, Script));
+			ScriptToOnParameterStoreChangedHandleMap.Add(FObjectKey(Script), OnParameterStoreChangedHandle);
+		}
+
+		UserParameterStoreChangedHandle = System.GetExposedParameters().AddOnChangedHandler(
+			FNiagaraParameterStore::FOnChanged::FDelegate::CreateRaw<FNiagaraSystemViewModel, const FNiagaraParameterStore&, const UNiagaraScript*>(
+				this, &FNiagaraSystemViewModel::SystemParameterStoreChanged, System.GetExposedParameters(), nullptr));
+	}
+}
+
+void FNiagaraSystemViewModel::RemoveSystemEventHandlers()
+{
+	if (System.IsValid())
+	{
+		TArray<UNiagaraScript*> Scripts;
+		Scripts.Add(System.GetSystemSpawnScript());
+		Scripts.Add(System.GetSystemUpdateScript());
+
+		for (UNiagaraScript* Script : Scripts)
+		{
+			FDelegateHandle* OnParameterStoreChangedHandle = ScriptToOnParameterStoreChangedHandleMap.Find(FObjectKey(Script));
+			if (OnParameterStoreChangedHandle != nullptr)
+			{
+				Script->RapidIterationParameters.RemoveOnChangedHandler(*OnParameterStoreChangedHandle);
+			}
+		}
+
+		System.GetExposedParameters().RemoveOnChangedHandler(UserParameterStoreChangedHandle);
+	}
+
+	ScriptToOnParameterStoreChangedHandleMap.Empty();
+	UserParameterStoreChangedHandle.Reset();
+}
+
+void FNiagaraSystemViewModel::NotifyPinnedCurvesChanged()
+{
+	OnPinnedCurvesChangedDelegate.Broadcast();
+}
+
+#undef LOCTEXT_NAMESPACE // NiagaraSystemViewModel
\ No newline at end of file
diff --git a/Engine/Plugins/FX/Niagara/Source/NiagaraEditor/Private/ViewModels/Stack/NiagaraParameterHandle.cpp b/Engine/Plugins/FX/Niagara/Source/NiagaraEditor/Private/ViewModels/Stack/NiagaraParameterHandle.cpp
--- a/Engine/Plugins/FX/Niagara/Source/NiagaraEditor/Private/ViewModels/Stack/NiagaraParameterHandle.cpp
+++ b/Engine/Plugins/FX/Niagara/Source/NiagaraEditor/Private/ViewModels/Stack/NiagaraParameterHandle.cpp
@@ -11,6 +11,7 @@ const FName FNiagaraParameterHandle::SystemNamespace(TEXT("System"));
 const FName FNiagaraParameterHandle::EmitterNamespace(TEXT("Emitter"));
 const FName FNiagaraParameterHandle::ParticleAttributeNamespace(TEXT("Particles"));
 const FName FNiagaraParameterHandle::ModuleNamespace(TEXT("Module"));
+const FName FNiagaraParameterHandle::ParameterCollectionNamespace(TEXT("NPC"));
 const FString FNiagaraParameterHandle::InitialPrefix(TEXT("Initial"));
 
 FNiagaraParameterHandle::FNiagaraParameterHandle() 
@@ -45,7 +46,7 @@ bool FNiagaraParameterHandle::operator==(const FNiagaraParameterHandle& Other) c
 	return ParameterHandleName == Other.ParameterHandleName;
 }
 
-FNiagaraParameterHandle FNiagaraParameterHandle::CreateAliasedModuleParameterHandle(const FNiagaraParameterHandle& ModuleParameterHandle, UNiagaraNodeFunctionCall* ModuleNode)
+FNiagaraParameterHandle FNiagaraParameterHandle::CreateAliasedModuleParameterHandle(const FNiagaraParameterHandle& ModuleParameterHandle, const UNiagaraNodeFunctionCall* ModuleNode)
 {
 	return ModuleParameterHandle.IsModuleHandle()
 		? FNiagaraParameterHandle(*ModuleNode->GetFunctionName(), ModuleParameterHandle.GetName())
@@ -98,6 +99,11 @@ const FName FNiagaraParameterHandle::GetNamespace() const
 	return Namespace;
 }
 
+bool FNiagaraParameterHandle::IsUserHandle() const
+{
+	return Namespace == UserNamespace;
+}
+
 bool FNiagaraParameterHandle::IsEngineHandle() const
 {
 	return Namespace == EngineNamespace;
@@ -122,3 +128,8 @@ bool FNiagaraParameterHandle::IsModuleHandle() const
 {
 	return Namespace == ModuleNamespace;
 }
+
+bool FNiagaraParameterHandle::IsParameterCollectionHandle() const
+{
+	return Namespace == ParameterCollectionNamespace;
+}
\ No newline at end of file
diff --git a/Engine/Plugins/FX/Niagara/Source/NiagaraEditor/Private/ViewModels/Stack/NiagaraStackAdvancedExpander.cpp b/Engine/Plugins/FX/Niagara/Source/NiagaraEditor/Private/ViewModels/Stack/NiagaraStackAdvancedExpander.cpp
new file mode 100644
--- /dev/null
+++ b/Engine/Plugins/FX/Niagara/Source/NiagaraEditor/Private/ViewModels/Stack/NiagaraStackAdvancedExpander.cpp
@@ -0,0 +1,46 @@
+// Copyright 1998-2018 Epic Games, Inc. All Rights Reserved.
+
+#include "ViewModels/Stack/NiagaraStackAdvancedExpander.h"
+#include "NiagaraStackEditorData.h"
+#include "ViewModels/NiagaraEmitterViewModel.h"
+#include "NiagaraNode.h"
+
+void UNiagaraStackAdvancedExpander::Initialize(
+	FRequiredEntryData InRequiredEntryData,
+	FString InOwnerStackItemEditorDataKey,
+	UNiagaraNode* InOwningNiagaraNode)
+{
+	Super::Initialize(InRequiredEntryData, FString());
+	OwnerStackItemEditorDataKey = InOwnerStackItemEditorDataKey;
+	OwningNiagaraNode = InOwningNiagaraNode;
+}
+
+bool UNiagaraStackAdvancedExpander::GetCanExpand() const
+{
+	return false;
+}
+
+bool UNiagaraStackAdvancedExpander::GetIsEnabled() const
+{
+	return OwningNiagaraNode == nullptr || OwningNiagaraNode->GetDesiredEnabledState() == ENodeEnabledState::Enabled;
+}
+
+UNiagaraStackEntry::EStackRowStyle UNiagaraStackAdvancedExpander::GetStackRowStyle() const
+{
+	return UNiagaraStackEntry::EStackRowStyle::ItemFooter;
+}
+
+void UNiagaraStackAdvancedExpander::SetOnToggleShowAdvanced(FOnToggleShowAdvanced OnExpandedChanged)
+{
+	ToggleShowAdvancedDelegate = OnExpandedChanged;
+}
+
+bool UNiagaraStackAdvancedExpander::GetShowAdvanced() const
+{
+	return GetStackEditorData().GetStackItemShowAdvanced(OwnerStackItemEditorDataKey, false);
+}
+
+void UNiagaraStackAdvancedExpander::ToggleShowAdvanced()
+{
+	ToggleShowAdvancedDelegate.ExecuteIfBound();
+}
diff --git a/Engine/Plugins/FX/Niagara/Source/NiagaraEditor/Private/ViewModels/Stack/NiagaraStackEmitterSpawnScriptItemGroup.cpp b/Engine/Plugins/FX/Niagara/Source/NiagaraEditor/Private/ViewModels/Stack/NiagaraStackEmitterSpawnScriptItemGroup.cpp
--- a/Engine/Plugins/FX/Niagara/Source/NiagaraEditor/Private/ViewModels/Stack/NiagaraStackEmitterSpawnScriptItemGroup.cpp
+++ b/Engine/Plugins/FX/Niagara/Source/NiagaraEditor/Private/ViewModels/Stack/NiagaraStackEmitterSpawnScriptItemGroup.cpp
@@ -3,19 +3,19 @@
 #include "ViewModels/Stack/NiagaraStackEmitterSpawnScriptItemGroup.h"
 #include "ViewModels/Stack/NiagaraStackObject.h"
 #include "ViewModels/Stack/NiagaraStackSpacer.h"
-#include "ViewModels/Stack/NiagaraStackItemExpander.h"
 #include "ViewModels/NiagaraSystemViewModel.h"
 #include "ViewModels/NiagaraEmitterViewModel.h"
 #include "NiagaraEmitter.h"
 #include "NiagaraStackEditorData.h"
 #include "ViewModels/Stack/NiagaraStackGraphUtilities.h"
 #include "NiagaraScriptMergeManager.h"
+#include "NiagaraEmitterDetailsCustomization.h"
 
 #define LOCTEXT_NAMESPACE "UNiagaraStackScriptItemGroup"
 
-void UNiagaraStackEmitterPropertiesItem::Initialize(TSharedRef<FNiagaraSystemViewModel> InSystemViewModel, TSharedRef<FNiagaraEmitterViewModel> InEmitterViewModel, UNiagaraStackEditorData& InStackEditorData)
+void UNiagaraStackEmitterPropertiesItem::Initialize(FRequiredEntryData InRequiredEntryData)
 {
-	Super::Initialize(InSystemViewModel, InEmitterViewModel, InStackEditorData);
+	Super::Initialize(InRequiredEntryData, TEXT("EmitterProperties"));
 	Emitter = GetEmitterViewModel()->GetEmitter();
 	Emitter->OnPropertiesChanged().AddUObject(this, &UNiagaraStackEmitterPropertiesItem::EmitterPropertiesChanged);
 }
@@ -62,33 +62,19 @@ void UNiagaraStackEmitterPropertiesItem::BeginDestroy()
 	Super::BeginDestroy();
 }
 
-void UNiagaraStackEmitterPropertiesItem::RefreshChildrenInternal(const TArray<UNiagaraStackEntry*>& CurrentChildren, TArray<UNiagaraStackEntry*>& NewChildren)
+void UNiagaraStackEmitterPropertiesItem::RefreshChildrenInternal(const TArray<UNiagaraStackEntry*>& CurrentChildren, TArray<UNiagaraStackEntry*>& NewChildren, TArray<FStackIssue>& NewIssues)
 {
 	if (EmitterObject == nullptr)
 	{
 		EmitterObject = NewObject<UNiagaraStackObject>(this);
-		EmitterObject->Initialize(GetSystemViewModel(), GetEmitterViewModel(), Emitter.Get());
+		FRequiredEntryData RequiredEntryData(GetSystemViewModel(), GetEmitterViewModel(), FExecutionCategoryNames::Emitter, NAME_None, GetStackEditorData());
+		EmitterObject->Initialize(RequiredEntryData, Emitter.Get(), GetStackEditorDataKey());
+		EmitterObject->RegisterInstancedCustomPropertyLayout(UNiagaraEmitter::StaticClass(), FOnGetDetailCustomizationInstance::CreateStatic(&FNiagaraEmitterDetails::MakeInstance));
 	}
 
-	if (EmitterExpander == nullptr)
-	{
-		EmitterExpander = NewObject<UNiagaraStackItemExpander>(this);
-		EmitterExpander->Initialize(GetSystemViewModel(), GetEmitterViewModel(), GetStackEditorData(), TEXT("Emitter"), false);
-		EmitterExpander->SetOnExpnadedChanged(UNiagaraStackItemExpander::FOnExpandedChanged::CreateUObject(this, &UNiagaraStackEmitterPropertiesItem::EmitterExpandedChanged));
-	}
-
-	if (GetStackEditorData().GetStackEntryIsExpanded(TEXT("Emitter"), false))
-	{
-		NewChildren.Add(EmitterObject);
-	}
-
-	NewChildren.Add(EmitterExpander);
+	NewChildren.Add(EmitterObject);
 	bCanResetToBase.Reset();
-}
-
-void UNiagaraStackEmitterPropertiesItem::EmitterExpandedChanged()
-{
-	RefreshChildren();
+	Super::RefreshChildrenInternal(CurrentChildren, NewChildren, NewIssues);
 }
 
 void UNiagaraStackEmitterPropertiesItem::EmitterPropertiesChanged()
@@ -101,24 +87,29 @@ UNiagaraStackEmitterSpawnScriptItemGroup::UNiagaraStackEmitterSpawnScriptItemGro
 {
 }
 
-void UNiagaraStackEmitterSpawnScriptItemGroup::RefreshChildrenInternal(const TArray<UNiagaraStackEntry*>& CurrentChildren, TArray<UNiagaraStackEntry*>& NewChildren)
+void UNiagaraStackEmitterSpawnScriptItemGroup::RefreshChildrenInternal(const TArray<UNiagaraStackEntry*>& CurrentChildren, TArray<UNiagaraStackEntry*>& NewChildren, TArray<FStackIssue>& NewIssues)
 {
-	if (PropertiesItem == nullptr)
-	{
-		PropertiesItem = NewObject<UNiagaraStackEmitterPropertiesItem>(this);
-		PropertiesItem->Initialize(GetSystemViewModel(), GetEmitterViewModel(), GetStackEditorData());
-	}
+	FName PropertiesSpacerKey = "PropertiesSpacer";
+	UNiagaraStackSpacer* PropertiesSpacer = FindCurrentChildOfTypeByPredicate<UNiagaraStackSpacer>(CurrentChildren,
+		[=](UNiagaraStackSpacer* CurrentPropertiesSpacer) { return CurrentPropertiesSpacer->GetSpacerKey() == PropertiesSpacerKey; });
 
 	if (PropertiesSpacer == nullptr)
 	{
 		PropertiesSpacer = NewObject<UNiagaraStackSpacer>(this);
-		PropertiesSpacer->Initialize(GetSystemViewModel(), GetEmitterViewModel(), "EmitterProperties");
+		PropertiesSpacer->Initialize(CreateDefaultChildRequiredData(), PropertiesSpacerKey);
 	}
 
-	NewChildren.Add(PropertiesItem);
 	NewChildren.Add(PropertiesSpacer);
 
-	Super::RefreshChildrenInternal(CurrentChildren, NewChildren);
+	if (PropertiesItem == nullptr)
+	{
+		PropertiesItem = NewObject<UNiagaraStackEmitterPropertiesItem>(this);
+		PropertiesItem->Initialize(CreateDefaultChildRequiredData());
+	}
+
+	NewChildren.Add(PropertiesItem);
+
+	Super::RefreshChildrenInternal(CurrentChildren, NewChildren, NewIssues);
 }
 
 #undef LOCTEXT_NAMESPACE
