Copying //UE4/Dev-Niagara to //UE4/Dev-Main (Source: //UE4/Dev-Niagara @ 4074996)

diff --git a/Engine/Plugins/FX/Niagara/Source/NiagaraEditor/Private/ViewModels/Stack/NiagaraStackEntry.cpp b/Engine/Plugins/FX/Niagara/Source/NiagaraEditor/Private/ViewModels/Stack/NiagaraStackEntry.cpp
--- a/Engine/Plugins/FX/Niagara/Source/NiagaraEditor/Private/ViewModels/Stack/NiagaraStackEntry.cpp
+++ b/Engine/Plugins/FX/Niagara/Source/NiagaraEditor/Private/ViewModels/Stack/NiagaraStackEntry.cpp
@@ -2,39 +2,61 @@
 
 #include "ViewModels/Stack/NiagaraStackEntry.h"
 #include "ViewModels/Stack/NiagaraStackErrorItem.h"
-#include "NiagaraUnmergeableStackEntryFilter.h"
+#include "NiagaraStackEditorData.h"
+
+const FName UNiagaraStackEntry::FExecutionCategoryNames::System = TEXT("System");
+const FName UNiagaraStackEntry::FExecutionCategoryNames::Emitter = TEXT("Emitter");
+const FName UNiagaraStackEntry::FExecutionCategoryNames::Particle = TEXT("Particle");
+const FName UNiagaraStackEntry::FExecutionCategoryNames::Render = TEXT("Render");
+
+const FName UNiagaraStackEntry::FExecutionSubcategoryNames::Parameters = TEXT("Parameters");
+const FName UNiagaraStackEntry::FExecutionSubcategoryNames::Spawn = TEXT("Spawn");
+const FName UNiagaraStackEntry::FExecutionSubcategoryNames::Update = TEXT("Update");
+const FName UNiagaraStackEntry::FExecutionSubcategoryNames::Event = TEXT("Event");
 
 UNiagaraStackEntry::UNiagaraStackEntry()
+	: IndentLevel(0)
 {
 }
 
-void UNiagaraStackEntry::Initialize(TSharedRef<FNiagaraSystemViewModel> InSystemViewModel, TSharedRef<FNiagaraEmitterViewModel> InEmitterViewModel)
+void UNiagaraStackEntry::Initialize(FRequiredEntryData InRequiredEntryData, FString InStackEditorDataKey)
 {
-	SystemViewModel = InSystemViewModel;
-	EmitterViewModel = InEmitterViewModel;
-	bIsExpanded = IsExpandedByDefault();
-	AddChildFilter(FNiagaraUnmergeableStackEntryFilter::CreateFilter());
+	SystemViewModel = InRequiredEntryData.SystemViewModel;
+	EmitterViewModel = InRequiredEntryData.EmitterViewModel;
+	ExecutionCategoryName = InRequiredEntryData.ExecutionCategoryName;
+	ExecutionSubcategoryName = InRequiredEntryData.ExecutionSubcategoryName;
+	StackEditorData = InRequiredEntryData.StackEditorData;
+	StackEditorDataKey = InStackEditorDataKey;
 }
 
-FText UNiagaraStackEntry::GetDisplayName() const
+bool UNiagaraStackEntry::IsValid() const
 {
-	ensureMsgf(false, TEXT("GetDisplayName not implemented"));
-	return FText::FromName(NAME_None);
+	return SystemViewModel.IsValid();
 }
 
-FText UNiagaraStackEntry::GetTooltipText() const
+FText UNiagaraStackEntry::GetDisplayName() const
 {
 	return FText();
 }
 
-FName UNiagaraStackEntry::GetTextStyleName() const
+UNiagaraStackEditorData& UNiagaraStackEntry::GetStackEditorData() const
+{
+	return *StackEditorData;
+}
+
+FString UNiagaraStackEntry::GetStackEditorDataKey() const
 {
-	return "NiagaraEditor.Stack.DefaultText";
+	return StackEditorDataKey;
+}
+
+FText UNiagaraStackEntry::GetTooltipText() const
+{
+	return FText();
 }
 
 bool UNiagaraStackEntry::GetCanExpand() const
 {
-	return false;
+	return true;
 }
 
 bool UNiagaraStackEntry::IsExpandedByDefault() const
@@ -44,12 +66,43 @@ bool UNiagaraStackEntry::IsExpandedByDefault() const
 
 bool UNiagaraStackEntry::GetIsExpanded() const
 {
-	return bIsExpanded;
+	if (GetShouldShowInStack() == false || GetCanExpand() == false)
+	{
+		// Entries that aren't visible, or can't expand are always expanded.
+		return true;
+	}
+
+	if (bIsExpandedCache.IsSet() == false)
+	{
+		bIsExpandedCache = StackEditorData->GetStackEntryIsExpanded(GetStackEditorDataKey(), IsExpandedByDefault());
+	}
+	return bIsExpandedCache.GetValue();
 }
 
 void UNiagaraStackEntry::SetIsExpanded(bool bInExpanded)
 {
-	bIsExpanded = bInExpanded;
+	StackEditorData->SetStackEntryIsExpanded(GetStackEditorDataKey(), bInExpanded);
+	bIsExpandedCache.Reset();
+}
+
+bool UNiagaraStackEntry::GetIsEnabled() const
+{
+	return true;
+}
+
+FName UNiagaraStackEntry::GetExecutionCategoryName() const
+{
+	return ExecutionCategoryName;
+}
+
+FName UNiagaraStackEntry::GetExecutionSubcategoryName() const
+{
+	return ExecutionSubcategoryName;
+}
+
+UNiagaraStackEntry::EStackRowStyle UNiagaraStackEntry::GetStackRowStyle() const
+{
+	return EStackRowStyle::None;
 }
 
 bool UNiagaraStackEntry::GetShouldShowInStack() const
@@ -85,18 +138,10 @@ void UNiagaraStackEntry::GetUnfilteredChildren(TArray<UNiagaraStackEntry*>& OutU
 	OutUnfilteredChildren.Append(Children);
 }
 
-void UNiagaraStackEntry::RefreshErrors()
+void UNiagaraStackEntry::GetUnfilteredChildren(TArray<UNiagaraStackEntry*>& OutUnfilteredChildren) const
 {
-	ErrorChildren.Empty();
-	if (GetErrorCount() != 0)
-	{
-		for (int32 i = 0; i < GetErrorCount(); i++)
-		{
-			UNiagaraStackErrorItem* ErrorObject = NewObject<UNiagaraStackErrorItem>(this);
-			ErrorObject->Initialize(GetSystemViewModel(), GetEmitterViewModel(), this, i);
-			ErrorChildren.Add(ErrorObject);
-		}
-	}
+	OutUnfilteredChildren.Append(ErrorChildren);
+	OutUnfilteredChildren.Append(Children);
 }
 
 FDelegateHandle UNiagaraStackEntry::AddChildFilter(FOnFilterChild ChildFilter)
@@ -136,29 +181,93 @@ UNiagaraStackEntry::FOnDataObjectModified& UNiagaraStackEntry::OnDataObjectModif
 	return DataObjectModifiedDelegate;
 }
 
-FName UNiagaraStackEntry::GetGroupBackgroundName() const
+int32 UNiagaraStackEntry::GetIndentLevel() const
+{
+	return IndentLevel;
+}
+
+void UNiagaraStackEntry::GetSearchItems(TArray<UNiagaraStackEntry::FStackSearchItem>& SearchItems) const
+{
+	SearchItems.Add({FName("DisplayName"), GetDisplayName()}); 
+	GetAdditionalSearchItemsInternal(SearchItems);
+}
+
+UObject* UNiagaraStackEntry::GetExternalAsset() const
+{
+	return nullptr;
+}
+
+bool UNiagaraStackEntry::CanDrag() const
+{
+	return false;
+}
+
+TOptional<UNiagaraStackEntry::FDropResult> UNiagaraStackEntry::CanDrop(const TArray<UNiagaraStackEntry*>& DraggedEntries)
+{
+	TOptional<FDropResult> Result = CanDropInternal(DraggedEntries);
+	if (Result.IsSet())
+	{
+		return Result;
+	}
+	else
+	{
+		return OnRequestCanDropDelegate.IsBound() 
+			? OnRequestCanDropDelegate.Execute(*this, DraggedEntries)
+			: TOptional<FDropResult>();
+	}
+}
+
+TOptional<UNiagaraStackEntry::FDropResult> UNiagaraStackEntry::Drop(const TArray<UNiagaraStackEntry*>& DraggedEntries)
+{
+	TOptional<FDropResult> Result = DropInternal(DraggedEntries);
+	if (Result.IsSet())
+	{
+		return Result;
+	}
+	else
+	{
+		return OnRequestDropDelegate.IsBound()
+			? OnRequestDropDelegate.Execute(*this, DraggedEntries)
+			: TOptional<FDropResult>();
+	}
+}
+
+void UNiagaraStackEntry::SetOnRequestCanDrop(FOnRequestDrop InOnRequestCanDrop)
+{
+	OnRequestCanDropDelegate = InOnRequestCanDrop;
+}
+
+void UNiagaraStackEntry::SetOnRequestDrop(FOnRequestDrop InOnRequestDrop)
+{
+	OnRequestDropDelegate = InOnRequestDrop;
+}
+
+void UNiagaraStackEntry::GetAdditionalSearchItemsInternal(TArray<FStackSearchItem>& SearchItems) const
 {
-	return "NiagaraEditor.Stack.Group.BackgroundColor";
 }
 
-FName UNiagaraStackEntry::GetGroupForegroundName() const
+TOptional<UNiagaraStackEntry::FDropResult> UNiagaraStackEntry::CanDropInternal(const TArray<UNiagaraStackEntry*>& DraggedEntries)
 {
-	return "NiagaraEditor.Stack.ForegroundColor";
+	return TOptional<FDropResult>();
 }
 
-FName UNiagaraStackEntry::GetItemBackgroundName() const
+TOptional<UNiagaraStackEntry::FDropResult> UNiagaraStackEntry::DropInternal(const TArray<UNiagaraStackEntry*>& DraggedEntries)
 {
-	return "NiagaraEditor.Stack.Item.BackgroundColor";
+	return TOptional<FDropResult>();
 }
 
-FName UNiagaraStackEntry::GetItemForegroundName() const
+TOptional<UNiagaraStackEntry::FDropResult> UNiagaraStackEntry::ChildRequestCanDropInternal(const UNiagaraStackEntry& TargetChild, const TArray<UNiagaraStackEntry*>& DraggedEntries)
 {
-	return "NiagaraEditor.Stack.ForegroundColor";
+	return TOptional<FDropResult>();
 }
 
-int32 UNiagaraStackEntry::GetItemIndentLevel() const
+TOptional<UNiagaraStackEntry::FDropResult> UNiagaraStackEntry::ChildRequestDropInternal(const UNiagaraStackEntry& TargetChild, const TArray<UNiagaraStackEntry*>& DraggedEntries)
+{
+	return TOptional<FDropResult>();
+}
+
+void UNiagaraStackEntry::ChlildStructureChangedInternal()
 {
-	return 0;
 }
 
 void UNiagaraStackEntry::RefreshChildren()
@@ -169,31 +278,102 @@ void UNiagaraStackEntry::RefreshChildren()
 	{
 		Child->OnStructureChanged().RemoveAll(this);
 		Child->OnDataObjectModified().RemoveAll(this);
+		Child->SetOnRequestCanDrop(FOnRequestDrop());
+		Child->SetOnRequestDrop(FOnRequestDrop());
+	}
+	for (UNiagaraStackErrorItem* ErrorChild : ErrorChildren)
+	{
+		ErrorChild->OnStructureChanged().RemoveAll(this);
+		ErrorChild->OnDataObjectModified().RemoveAll(this);
+		ErrorChild->OnIssueModified().RemoveAll(this);
 	}
 
 	TArray<UNiagaraStackEntry*> NewChildren;
-	RefreshChildrenInternal(Children, NewChildren);
+	TArray<FStackIssue> NewStackIssues;
+	RefreshChildrenInternal(Children, NewChildren, NewStackIssues);
 	Children.Empty();
 	Children.Append(NewChildren);
 
 	for (UNiagaraStackEntry* Child : Children)
 	{
+		Child->IndentLevel = GetChildIndentLevel();
 		Child->RefreshChildren();
 		Child->OnStructureChanged().AddUObject(this, &UNiagaraStackEntry::ChildStructureChanged);
 		Child->OnDataObjectModified().AddUObject(this, &UNiagaraStackEntry::ChildDataObjectModified);
+		Child->SetOnRequestCanDrop(FOnRequestDrop::CreateUObject(this, &UNiagaraStackEntry::ChildRequestCanDrop));
+		Child->SetOnRequestDrop(FOnRequestDrop::CreateUObject(this, &UNiagaraStackEntry::ChildRequestDrop));
+	}
+	
+	// Stack issues refresh
+	NewStackIssues.RemoveAll([=](const FStackIssue& Issue) { return GetStackEditorData().GetDismissedStackIssueIds().Contains(Issue.UniqueIdentifier); }); 
+
+	StackIssues.Empty();
+	StackIssues.Append(NewStackIssues);
+	RefreshStackErrorChildren();
+	for (UNiagaraStackErrorItem* ErrorChild : ErrorChildren)
+	{
+		ErrorChild->IndentLevel = GetChildIndentLevel();
+		ErrorChild->RefreshChildren();
+		ErrorChild->OnStructureChanged().AddUObject(this, &UNiagaraStackEntry::ChildStructureChanged);
+		ErrorChild->OnDataObjectModified().AddUObject(this, &UNiagaraStackEntry::ChildDataObjectModified);
+		ErrorChild->OnIssueModified().AddUObject(this, &UNiagaraStackEntry::IssueModified);
 	}
 
-	RefreshErrors();
+	PostRefreshChildrenInternal();
 
 	StructureChangedDelegate.Broadcast();
 }
 
-void UNiagaraStackEntry::RefreshChildrenInternal(const TArray<UNiagaraStackEntry*>& CurrentChildren, TArray<UNiagaraStackEntry*>& NewChildren)
+void UNiagaraStackEntry::RefreshStackErrorChildren()
+{
+    // keep the error entries that are already built
+	TArray<UNiagaraStackErrorItem*> NewErrorChildren;
+	for (FStackIssue Issue : StackIssues)
+	{
+		UNiagaraStackErrorItem* ErrorEntry = nullptr;
+		UNiagaraStackErrorItem** Found = ErrorChildren.FindByPredicate(
+			[&](UNiagaraStackErrorItem* CurrentChild) { return CurrentChild->GetStackIssue().UniqueIdentifier == Issue.UniqueIdentifier; });
+		if (Found == nullptr)
+		{
+			ErrorEntry = NewObject<UNiagaraStackErrorItem>(this);
+			ErrorEntry->Initialize(CreateDefaultChildRequiredData(), Issue, GetStackEditorDataKey());
+		}
+		else
+		{
+			ErrorEntry = *Found;
+		}
+		NewErrorChildren.Add(ErrorEntry);
+	}
+	ErrorChildren.Empty();
+	ErrorChildren.Append(NewErrorChildren);
+}
+
+void UNiagaraStackEntry::IssueModified()
+{
+	RefreshChildren();
+}
+
+void UNiagaraStackEntry::RefreshChildrenInternal(const TArray<UNiagaraStackEntry*>& CurrentChildren, TArray<UNiagaraStackEntry*>& NewChildren, TArray<FStackIssue>& NewIssues)
+{
+}
+
+void UNiagaraStackEntry::PostRefreshChildrenInternal()
+{
+}
+
+UNiagaraStackEntry::FRequiredEntryData UNiagaraStackEntry::CreateDefaultChildRequiredData() const
+{
+	return FRequiredEntryData(GetSystemViewModel(), GetEmitterViewModel(), GetExecutionCategoryName(), GetExecutionSubcategoryName(), GetStackEditorData());
+}
+
+int32 UNiagaraStackEntry::GetChildIndentLevel() const
 {
+	return GetShouldShowInStack() ? GetIndentLevel() + 1 : GetIndentLevel();
 }
 
 void UNiagaraStackEntry::ChildStructureChanged()
 {
+	ChlildStructureChangedInternal();
 	StructureChangedDelegate.Broadcast();
 }
 
@@ -201,3 +381,33 @@ void UNiagaraStackEntry::ChildDataObjectModified(UObject* ChangedObject)
 {
 	DataObjectModifiedDelegate.Broadcast(ChangedObject);
 }
+
+TOptional<UNiagaraStackEntry::FDropResult> UNiagaraStackEntry::ChildRequestCanDrop(const UNiagaraStackEntry& TargetChild, const TArray<UNiagaraStackEntry*>& DraggedEntries)
+{
+	TOptional<FDropResult> Result = ChildRequestCanDropInternal(TargetChild, DraggedEntries);
+	if (Result.IsSet())
+	{
+		return Result;
+	}
+	else
+	{
+		return OnRequestCanDropDelegate.IsBound()
+			? OnRequestCanDropDelegate.Execute(TargetChild, DraggedEntries)
+			: TOptional<FDropResult>();
+	}
+}
+
+TOptional<UNiagaraStackEntry::FDropResult> UNiagaraStackEntry::ChildRequestDrop(const UNiagaraStackEntry& TargetChild, const TArray<UNiagaraStackEntry*>& DraggedEntries)
+{
+	TOptional<FDropResult> Result = ChildRequestDropInternal(TargetChild, DraggedEntries);
+	if (Result.IsSet())
+	{
+		return Result;
+	}
+	else
+	{
+		return OnRequestDropDelegate.IsBound()
+			? OnRequestDropDelegate.Execute(TargetChild, DraggedEntries)
+			: TOptional<FDropResult>();
+	}
+}
diff --git a/Engine/Plugins/FX/Niagara/Source/NiagaraEditor/Private/ViewModels/Stack/NiagaraStackErrorItem.cpp b/Engine/Plugins/FX/Niagara/Source/NiagaraEditor/Private/ViewModels/Stack/NiagaraStackErrorItem.cpp
--- a/Engine/Plugins/FX/Niagara/Source/NiagaraEditor/Private/ViewModels/Stack/NiagaraStackErrorItem.cpp
+++ b/Engine/Plugins/FX/Niagara/Source/NiagaraEditor/Private/ViewModels/Stack/NiagaraStackErrorItem.cpp
@@ -1,70 +1,182 @@
 // Copyright 1998-2018 Epic Games, Inc. All Rights Reserved.
 
 #include "ViewModels/Stack/NiagaraStackErrorItem.h"
+#include "NiagaraStackEditorData.h"
+#include "ScopedTransaction.h"
 
-UNiagaraStackErrorItem::UNiagaraStackErrorItem() : ErrorSource(nullptr), ErrorIdx(INDEX_NONE)
+#define LOCTEXT_NAMESPACE "NiagaraStackErrorItem"
+//UNiagaraStackErrorItem
+
+void UNiagaraStackErrorItem::Initialize(FRequiredEntryData InRequiredEntryData, FStackIssue InStackIssue, FString InStackEditorDataKey)
 {
+	FString ErrorStackEditorDataKey = FString::Printf(TEXT("%s-Error-%s"), *InStackEditorDataKey, InStackIssue.UniqueIdentifier.GetPlainANSIString());
+	Super::Initialize(InRequiredEntryData, ErrorStackEditorDataKey);
+	StackIssue = InStackIssue;
+	EntryStackEditorDataKey = InStackEditorDataKey;
 }
 
-void UNiagaraStackErrorItem::Initialize(TSharedRef<FNiagaraSystemViewModel> InSystemViewModel, TSharedRef<FNiagaraEmitterViewModel> InEmitterViewModel, UNiagaraStackEntry* InErrorSource, int32 InErrorIndex)
+FText UNiagaraStackErrorItem::GetDisplayName() const
 {
-	checkf(ErrorSource == nullptr, TEXT("Can only initialize once."));
-	Super::Initialize(InSystemViewModel, InEmitterViewModel);
-	ErrorSource = InErrorSource;
-	ErrorIdx = InErrorIndex;
+	return StackIssue.ShortDescription;
 }
 
-int32 UNiagaraStackErrorItem::GetItemIndentLevel() const
+UNiagaraStackEntry::EStackRowStyle UNiagaraStackErrorItem::GetStackRowStyle() const
 {
-	return ItemIndentLevel;
+	return EStackRowStyle::StackIssue;
 }
 
-void UNiagaraStackErrorItem::SetItemIndentLevel(int32 InItemIndentLevel)
+UNiagaraStackErrorItem::FOnIssueNotify& UNiagaraStackErrorItem::OnIssueModified()
 {
-	ItemIndentLevel = InItemIndentLevel;
+	return IssueModifiedDelegate;
 }
 
-FText UNiagaraStackErrorItem::ErrorText() const
+void UNiagaraStackErrorItem::RefreshChildrenInternal(const TArray<UNiagaraStackEntry*>& CurrentChildren, TArray<UNiagaraStackEntry*>& NewChildren, TArray<FStackIssue>& NewIssues)
 {
-	FText SummaryText = ErrorSource->GetErrorSummaryText(ErrorIdx);
-	if (SummaryText.IsEmpty())
+	for (UNiagaraStackEntry* Child : CurrentChildren)
+	{
+		UNiagaraStackErrorItemFix* FixChild = Cast<UNiagaraStackErrorItemFix>(Child);
+		if (FixChild != nullptr)
+		{
+			FixChild->OnIssueFixed().RemoveAll(this);
+		}
+	}
+	// long description
+	UNiagaraStackErrorItemLongDescription* ErrorEntryLongDescription = FindCurrentChildOfTypeByPredicate<UNiagaraStackErrorItemLongDescription>(CurrentChildren,
+		[&](UNiagaraStackErrorItemLongDescription* CurrentChild) { return true; });
+	if (ErrorEntryLongDescription == nullptr)
+	{
+		ErrorEntryLongDescription = NewObject<UNiagaraStackErrorItemLongDescription>(this);
+		ErrorEntryLongDescription->Initialize(CreateDefaultChildRequiredData(), StackIssue, GetStackEditorDataKey());
+	}
+	NewChildren.Add(ErrorEntryLongDescription);
+
+	// fixes
+	for (int i = 0; i < StackIssue.Fixes.Num(); i++)
 	{
-		return ErrorSource->GetErrorText(ErrorIdx);
+		UNiagaraStackErrorItemFix* ErrorEntryFix = FindCurrentChildOfTypeByPredicate<UNiagaraStackErrorItemFix>(CurrentChildren,
+			[&](UNiagaraStackErrorItemFix* CurrentChild) { return CurrentChild->GetStackIssueFix() == StackIssue.Fixes[i]; });
+		if (ErrorEntryFix == nullptr)
+		{
+			ErrorEntryFix = NewObject<UNiagaraStackErrorItemFix>(this);
+			ErrorEntryFix->Initialize(CreateDefaultChildRequiredData(), StackIssue, StackIssue.Fixes[i], EntryStackEditorDataKey);
+		}
+		ErrorEntryFix->OnIssueFixed().AddUObject(this, &UNiagaraStackErrorItem::IssueFixed);
+		NewChildren.Add(ErrorEntryFix);
 	}
-	else
+	// dismiss button
+	if (StackIssue.bCanBeDismissed)
 	{
-		return SummaryText;
+		UNiagaraStackErrorItemDismiss* ErrorEntryDismiss = FindCurrentChildOfTypeByPredicate<UNiagaraStackErrorItemDismiss>(CurrentChildren,
+			[&](UNiagaraStackErrorItemDismiss* CurrentChild) { return true; });
+		if (ErrorEntryDismiss == nullptr)
+		{
+			ErrorEntryDismiss = NewObject<UNiagaraStackErrorItemDismiss>(this);
+			ErrorEntryDismiss->Initialize(CreateDefaultChildRequiredData(), StackIssue, EntryStackEditorDataKey);
+		}
+		ErrorEntryDismiss->OnIssueFixed().AddUObject(this, &UNiagaraStackErrorItem::IssueFixed);
+		NewChildren.Add(ErrorEntryDismiss);
 	}
 }
 
-FText UNiagaraStackErrorItem::ErrorTextTooltip() const
+void UNiagaraStackErrorItem::IssueFixed()
+{
+	OnIssueModified().Broadcast();
+}
+
+//UNiagaraStackErrorItemLongDescription
+void UNiagaraStackErrorItemLongDescription::Initialize(FRequiredEntryData InRequiredEntryData, UNiagaraStackEntry::FStackIssue InStackIssue, FString InStackEditorDataKey)
+{
+	FString ErrorStackEditorDataKey = FString::Printf(TEXT("Long-%s"), *InStackEditorDataKey);
+	Super::Initialize(InRequiredEntryData, ErrorStackEditorDataKey);
+	StackIssue = InStackIssue;
+}
+
+FText UNiagaraStackErrorItemLongDescription::GetDisplayName() const
 {
-	return ErrorSource->GetErrorText(ErrorIdx);
+	return StackIssue.LongDescription;
 }
 
-FReply UNiagaraStackErrorItem::OnTryFixError()
+UNiagaraStackEntry::EStackRowStyle UNiagaraStackErrorItemLongDescription::GetStackRowStyle() const
 {
-	if (ErrorSource->TryFixError(ErrorIdx))
+	return EStackRowStyle::StackIssue;
+}
+
+//UNiagaraStackErrorItemFix
+
+void UNiagaraStackErrorItemFix::Initialize(FRequiredEntryData InRequiredEntryData, FStackIssue InStackIssue, FStackIssueFix InIssueFix, FString InStackEditorDataKey)
+{
+	FString ErrorStackEditorDataKey = FString::Printf(TEXT("Fix-%s"), *InStackEditorDataKey);
+	Super::Initialize(InRequiredEntryData, ErrorStackEditorDataKey);
+	StackIssue = InStackIssue;
+	IssueFix = InIssueFix;
+}
+
+FText UNiagaraStackErrorItemFix::FixDescription() const
+{
+	if (IssueFix.Description.IsEmptyOrWhitespace())
 	{
-		return FReply::Handled();
+		ensureMsgf(false, TEXT("Fix description is empty or whitespace."));
+		return LOCTEXT("FixIssue", "Fix the issue");
 	}
+	return IssueFix.Description;
+}
+
+FReply UNiagaraStackErrorItemFix::OnTryFixError()
+{
+	ensureMsgf(IssueFix.FixDelegate.IsBound(), TEXT("Fix delegate is not bound!"));
+	IssueFix.FixDelegate.ExecuteIfBound();
+	OnIssueFixed().Broadcast();
 	return FReply::Handled();
 }
 
-EVisibility UNiagaraStackErrorItem::CanFixVisibility() const
+UNiagaraStackEntry::EStackRowStyle UNiagaraStackErrorItemFix::GetStackRowStyle() const
 {
-	if (ErrorSource->GetErrorFixable(ErrorIdx))
-	{
-		return EVisibility::Visible;
-	}
-	else
-	{
-		return EVisibility::Hidden;
-	}
+	return EStackRowStyle::StackIssue;
+}
+
+FText UNiagaraStackErrorItemFix::GetFixButtonText() const
+{
+	return LOCTEXT("FixIssue", "Fix issue");
 }
 
+UNiagaraStackErrorItem::FOnIssueNotify& UNiagaraStackErrorItemFix::OnIssueFixed()
+{
+	return IssueFixedDelegate;
+}
+
+//UNiagaraStackErrorItemDismiss
+
+void UNiagaraStackErrorItemDismiss::Initialize(FRequiredEntryData InRequiredEntryData, UNiagaraStackEntry::FStackIssue InStackIssue, FString InStackEditorDataKey)
+{
+	FString ErrorStackEditorDataKey = FString::Printf(TEXT("Dismiss-%s"), *InStackEditorDataKey);
+	UNiagaraStackEntry::Initialize(InRequiredEntryData, ErrorStackEditorDataKey);
+	StackIssue = InStackIssue;
+	IssueFix.Description = LOCTEXT("DismissError", "Dismiss the issue without fixing (I know what I'm doing).");
+}
 
-FName UNiagaraStackErrorItem::GetItemBackgroundName() const 
+void UNiagaraStackErrorItemDismiss::DismissIssue()
 {
-	return "NiagaraEditor.Stack.Item.ErrorBackgroundColor";
+	GetStackEditorData().Modify();
+	GetStackEditorData().DismissStackIssue(StackIssue.UniqueIdentifier);
 }
+
+FReply UNiagaraStackErrorItemDismiss::OnTryFixError()
+{
+	ensureMsgf(IssueFixedDelegate.IsBound(), TEXT("Issue fixed delegate is not bound on the dismiss entry!"));
+	FScopedTransaction ScopedTransaction(LOCTEXT("DismissIssue", "Dismiss issue"));
+	DismissIssue();
+	OnIssueFixed().Broadcast();
+	return FReply::Handled();
+}
+
+UNiagaraStackEntry::EStackRowStyle UNiagaraStackErrorItemDismiss::GetStackRowStyle() const
+{
+	return EStackRowStyle::StackIssue;
+}
+
+FText UNiagaraStackErrorItemDismiss::GetFixButtonText() const
+{
+	return LOCTEXT("DismissIssue", "Dismiss issue");
+}
+
+#undef LOCTEXT_NAMESPACE
\ No newline at end of file
diff --git a/Engine/Plugins/FX/Niagara/Source/NiagaraEditor/Private/ViewModels/Stack/NiagaraStackEventHandlerGroup.cpp b/Engine/Plugins/FX/Niagara/Source/NiagaraEditor/Private/ViewModels/Stack/NiagaraStackEventHandlerGroup.cpp
--- a/Engine/Plugins/FX/Niagara/Source/NiagaraEditor/Private/ViewModels/Stack/NiagaraStackEventHandlerGroup.cpp
+++ b/Engine/Plugins/FX/Niagara/Source/NiagaraEditor/Private/ViewModels/Stack/NiagaraStackEventHandlerGroup.cpp
@@ -3,49 +3,84 @@
 #include "ViewModels/Stack/NiagaraStackEventHandlerGroup.h"
 #include "ViewModels/Stack/NiagaraStackRendererItem.h"
 #include "ViewModels/Stack/NiagaraStackEventScriptItemGroup.h"
-#include "NiagaraEmitterViewModel.h"
+#include "ViewModels/NiagaraEmitterViewModel.h"
+#include "NiagaraScriptViewModel.h"
+#include "NiagaraScriptGraphViewModel.h"
 #include "NiagaraEmitter.h"
-#include "ViewModels/Stack/NiagaraStackAddEventScriptItem.h"
 #include "NiagaraGraph.h"
 #include "NiagaraScriptSource.h"
+#include "ViewModels/Stack/NiagaraStackSpacer.h"
+#include "ViewModels/Stack/INiagaraStackItemGroupAddUtilities.h"
+#include "ViewModels/Stack/NiagaraStackGraphUtilities.h"
+
+#include "ScopedTransaction.h"
 
 #define LOCTEXT_NAMESPACE "NiagaraStackEventHandlerGroup"
 
+class FEventHandlerGroupAddUtilities : public FNiagaraStackItemGroupAddUtilities
+{
+public:
+	FEventHandlerGroupAddUtilities(TSharedRef<FNiagaraEmitterViewModel> InEmitterViewModel, FOnItemAdded InOnItemAdded)
+		: FNiagaraStackItemGroupAddUtilities(LOCTEXT("ScriptGroupAddItemName", "Event Handler"), EAddMode::AddDirectly, false, InOnItemAdded)
+		, EmitterViewModel(InEmitterViewModel)
+	{
+	}
 
+	virtual void AddItemDirectly() override
+	{
+		TSharedPtr<FNiagaraEmitterViewModel> EmitterViewModelPinned = EmitterViewModel.Pin();
+		if (EmitterViewModelPinned.IsValid() == false)
+		{
+			return;
+		}
 
-FText UNiagaraStackEventHandlerGroup::GetDisplayName() const
-{
-	return DisplayName;
-}
+		UNiagaraEmitter* Emitter = EmitterViewModelPinned->GetEmitter();
+		UNiagaraScriptSource* Source = EmitterViewModelPinned->GetSharedScriptViewModel()->GetGraphViewModel()->GetScriptSource();
+		UNiagaraGraph* Graph = EmitterViewModelPinned->GetSharedScriptViewModel()->GetGraphViewModel()->GetGraph();
 
-void UNiagaraStackEventHandlerGroup::SetDisplayName(FText InDisplayName)
-{
-	DisplayName = InDisplayName;
-}
+		// The stack should not have been created if any of these are null, so bail out if it happens somehow rather than try to handle all of these cases.
+		checkf(Emitter != nullptr && Source != nullptr && Graph != nullptr, TEXT("Stack created for invalid emitter or graph."));
 
-void UNiagaraStackEventHandlerGroup::RefreshChildrenInternal(const TArray<UNiagaraStackEntry*>& CurrentChildren, TArray<UNiagaraStackEntry*>& NewChildren)
-{
-	if (AddScriptItem == nullptr)
-	{
-		AddScriptItem = NewObject<UNiagaraStackAddEventScriptItem>(this);
-		AddScriptItem->Initialize(GetSystemViewModel(), GetEmitterViewModel(), GetStackEditorData());
-		AddScriptItem->SetOnItemAdded(UNiagaraStackAddEventScriptItem::FOnItemAdded::CreateUObject(this, &UNiagaraStackEventHandlerGroup::ChildModifiedGroupItems));
+		FScopedTransaction ScopedTransaction(LOCTEXT("AddNewEventHandlerTransaction", "Add new event handler"));
+
+		Emitter->Modify();
+		FNiagaraEventScriptProperties EventScriptProperties;
+		EventScriptProperties.Script = NewObject<UNiagaraScript>(Emitter, MakeUniqueObjectName(Emitter, UNiagaraScript::StaticClass(), "EventScript"), EObjectFlags::RF_Transactional);
+		EventScriptProperties.Script->SetUsage(ENiagaraScriptUsage::ParticleEventScript);
+		EventScriptProperties.Script->SetUsageId(FGuid::NewGuid());
+		EventScriptProperties.Script->SetSource(Source);
+		Emitter->AddEventHandler(EventScriptProperties);
+		FNiagaraStackGraphUtilities::ResetGraphForOutput(*Graph, ENiagaraScriptUsage::ParticleEventScript, EventScriptProperties.Script->GetUsageId());
+
+		// Set the emitter here so that the internal state of the view model is updated.
+		// TODO: Move the logic for managing event handlers into the emitter view model or script view model.
+		EmitterViewModelPinned->SetEmitter(Emitter);
+
+		OnItemAdded.ExecuteIfBound();
 	}
 
-	NewChildren.Add(AddScriptItem);
-}
+	virtual void GenerateAddActions(TArray<TSharedRef<INiagaraStackItemGroupAddAction>>& OutAddActions) const override { unimplemented(); }
+	virtual void ExecuteAddAction(TSharedRef<INiagaraStackItemGroupAddAction> AddAction, int32 TargetIndex) override { unimplemented(); }
 
-void UNiagaraStackEventHandlerGroup::SetOnItemAdded(FOnItemAdded InOnItemAdded)
+private:
+	TWeakPtr<FNiagaraEmitterViewModel> EmitterViewModel;
+};
+
+void UNiagaraStackEventHandlerGroup::Initialize(FRequiredEntryData InRequiredEntryData)
 {
-	ItemAddedDelegate = InOnItemAdded;
+	FText DisplayName = LOCTEXT("EventGroupName", "Add Event Handler");
+	FText ToolTip = LOCTEXT("EventGroupTooltip", "Determines how this Emitter responds to incoming events. There can be more than one event handler script stack per Emitter.");
+	AddUtilities = MakeShared<FEventHandlerGroupAddUtilities>(InRequiredEntryData.EmitterViewModel, 
+		FNiagaraStackItemGroupAddUtilities::FOnItemAdded::CreateUObject(this, &UNiagaraStackEventHandlerGroup::ItemAddedFromUtilties));
+	Super::Initialize(InRequiredEntryData, DisplayName, ToolTip, AddUtilities.Get());
 }
 
-bool UNiagaraStackEventHandlerGroup::GetShouldShowInStack() const
+void UNiagaraStackEventHandlerGroup::SetOnItemAdded(FOnItemAdded InOnItemAdded)
 {
-	return true;
+	ItemAddedDelegate = InOnItemAdded;
 }
 
-void UNiagaraStackEventHandlerGroup::ChildModifiedGroupItems()
+void UNiagaraStackEventHandlerGroup::ItemAddedFromUtilties()
 {
 	ItemAddedDelegate.ExecuteIfBound();
 }
diff --git a/Engine/Plugins/FX/Niagara/Source/NiagaraEditor/Private/ViewModels/Stack/NiagaraStackEventScriptItemGroup.cpp b/Engine/Plugins/FX/Niagara/Source/NiagaraEditor/Private/ViewModels/Stack/NiagaraStackEventScriptItemGroup.cpp
--- a/Engine/Plugins/FX/Niagara/Source/NiagaraEditor/Private/ViewModels/Stack/NiagaraStackEventScriptItemGroup.cpp
+++ b/Engine/Plugins/FX/Niagara/Source/NiagaraEditor/Private/ViewModels/Stack/NiagaraStackEventScriptItemGroup.cpp
@@ -1,37 +1,40 @@
 // Copyright 1998-2018 Epic Games, Inc. All Rights Reserved.
 
 #include "ViewModels/Stack/NiagaraStackEventScriptItemGroup.h"
+#include "NiagaraEmitter.h"
+#include "ViewModels/Stack/NiagaraStackObject.h"
 #include "ViewModels/Stack/NiagaraStackModuleItem.h"
-#include "ViewModels/Stack/NiagaraStackAddModuleItem.h"
 #include "ViewModels/Stack/NiagaraStackSpacer.h"
 #include "NiagaraScriptViewModel.h"
 #include "NiagaraScriptGraphViewModel.h"
-#include "NiagaraEmitterViewModel.h"
+#include "ViewModels/NiagaraEmitterViewModel.h"
 #include "NiagaraGraph.h"
 #include "NiagaraNodeOutput.h"
 #include "NiagaraNodeFunctionCall.h"
 #include "EdGraphSchema_Niagara.h"
 #include "ViewModels/Stack/NiagaraStackErrorItem.h"
-#include "ViewModels/Stack/NiagaraStackStruct.h"
 #include "NiagaraScriptSource.h"
 #include "ViewModels/Stack/NiagaraStackGraphUtilities.h"
 #include "NiagaraScriptMergeManager.h"
 #include "NiagaraStackEditorData.h"
+#include "ViewModels/NiagaraSystemViewModel.h"
+#include "Customizations/NiagaraEventScriptPropertiesCustomization.h"
 
 #include "Internationalization/Internationalization.h"
 #include "ScopedTransaction.h"
-#include "NiagaraSystemViewModel.h"
-#include "Customizations/NiagaraEventScriptPropertiesCustomization.h"
+#include "IDetailTreeNode.h"
 
 #define LOCTEXT_NAMESPACE "UNiagaraStackEventScriptItemGroup"
 
-void UNiagaraStackEventHandlerPropertiesItem::Initialize(TSharedRef<FNiagaraSystemViewModel> InSystemViewModel, TSharedRef<FNiagaraEmitterViewModel> InEmitterViewModel, FGuid InEventScriptUsageId, UNiagaraStackEditorData& InStackEditorData)
+void UNiagaraStackEventHandlerPropertiesItem::Initialize(FRequiredEntryData InRequiredEntryData, FGuid InEventScriptUsageId)
 {
-	Super::Initialize(InSystemViewModel, InEmitterViewModel, InStackEditorData);
-	Emitter = GetEmitterViewModel()->GetEmitter();
+	FString EventStackEditorDataKey = FString::Printf(TEXT("Event-%s-Properties"), *InEventScriptUsageId.ToString(EGuidFormats::DigitsWithHyphens));
+	Super::Initialize(InRequiredEntryData, EventStackEditorDataKey);
+
 	EventScriptUsageId = InEventScriptUsageId;
+
+	Emitter = GetEmitterViewModel()->GetEmitter();
 	Emitter->OnPropertiesChanged().AddUObject(this, &UNiagaraStackEventHandlerPropertiesItem::EventHandlerPropertiesChanged);
-	bForceRebuildChildStruct = true;
 
 	const UNiagaraEmitter* BaseEmitter = FNiagaraStackGraphUtilities::GetBaseEmitter(*Emitter.Get(), GetSystemViewModel()->GetSystem());
 	if (BaseEmitter != nullptr && Emitter != BaseEmitter)
@@ -78,7 +81,6 @@ void UNiagaraStackEventHandlerPropertiesItem::ResetToBase()
 		const UNiagaraEmitter* BaseEmitter = FNiagaraStackGraphUtilities::GetBaseEmitter(*Emitter.Get(), GetSystemViewModel()->GetSystem());
 		TSharedRef<FNiagaraScriptMergeManager> MergeManager = FNiagaraScriptMergeManager::Get();
 		MergeManager->ResetEventHandlerPropertySetToBase(*Emitter, *BaseEmitter, EventScriptUsageId);
-		bForceRebuildChildStruct = true;
 		RefreshChildren();
 	}
 }
@@ -92,99 +94,125 @@ void UNiagaraStackEventHandlerPropertiesItem::BeginDestroy()
 	Super::BeginDestroy();
 }
 
-void UNiagaraStackEventHandlerPropertiesItem::RefreshChildrenInternal(const TArray<UNiagaraStackEntry*>& CurrentChildren, TArray<UNiagaraStackEntry*>& NewChildren)
+void UNiagaraStackEventHandlerPropertiesItem::RefreshChildrenInternal(const TArray<UNiagaraStackEntry*>& CurrentChildren, TArray<UNiagaraStackEntry*>& NewChildren, TArray<FStackIssue>& NewIssues)
 {
-	FNiagaraEventScriptProperties* EventScriptProperties = Emitter->GetEventHandlerByIdUnsafe(EventScriptUsageId);
-
-	UNiagaraStackStruct* EventHandlerStruct = nullptr;
-	if (EventScriptProperties != nullptr)
+	if (EmitterObject == nullptr)
 	{
-		uint8* EventStructMemory = (uint8*)EventScriptProperties;
-
-		if (bForceRebuildChildStruct == false)
-		{
-			EventHandlerStruct = FindCurrentChildOfTypeByPredicate<UNiagaraStackStruct>(CurrentChildren,
-				[=](UNiagaraStackStruct* CurrentChild) { return CurrentChild->GetOwningObject() == Emitter && CurrentChild->GetStructOnScope()->GetStructMemory() == EventStructMemory; });
-		}
-
-		if (EventHandlerStruct == nullptr)
-		{
-			UNiagaraSystem* System = &GetSystemViewModel()->GetSystem();
-			EventHandlerStruct = NewObject<UNiagaraStackStruct>(this);
-			EventHandlerStruct->Initialize(GetSystemViewModel(), GetEmitterViewModel(), GetEmitterViewModel()->GetEmitter(), FNiagaraEventScriptProperties::StaticStruct(), EventStructMemory);
-			EventHandlerStruct->SetDetailsCustomization(FOnGetDetailCustomizationInstance::CreateStatic(&FNiagaraEventScriptPropertiesCustomization::MakeInstance, MakeWeakObjectPtr(System), TWeakObjectPtr<UNiagaraEmitter>(Emitter)));
-			bForceRebuildChildStruct = false;
-		}
+		EmitterObject = NewObject<UNiagaraStackObject>(this);
+		EmitterObject->Initialize(CreateDefaultChildRequiredData(), Emitter.Get(), GetStackEditorDataKey());
+		EmitterObject->RegisterInstancedCustomPropertyTypeLayout(FNiagaraEventScriptProperties::StaticStruct()->GetFName(), 
+			FOnGetPropertyTypeCustomizationInstance::CreateStatic(&FNiagaraEventScriptPropertiesCustomization::MakeInstance, 
+			TWeakObjectPtr<UNiagaraSystem>(&GetSystemViewModel()->GetSystem()), TWeakObjectPtr<UNiagaraEmitter>(GetEmitterViewModel()->GetEmitter())));
+		EmitterObject->SetOnSelectRootNodes(UNiagaraStackObject::FOnSelectRootNodes::CreateUObject(this, &UNiagaraStackEventHandlerPropertiesItem::SelectEmitterStackObjectRootTreeNodes));
 	}
 
-	FString ExpandedKey = TEXT("EventHandler_") + EventScriptUsageId.ToString();
-
-	if (EventHandlerExpander == nullptr)
-	{
-		EventHandlerExpander = NewObject<UNiagaraStackItemExpander>(this);
-		EventHandlerExpander->Initialize(GetSystemViewModel(), GetEmitterViewModel(), GetStackEditorData(), ExpandedKey, false);
-		EventHandlerExpander->SetOnExpnadedChanged(UNiagaraStackItemExpander::FOnExpandedChanged::CreateUObject(this, &UNiagaraStackEventHandlerPropertiesItem::EventHandlerExpandedChanged));
-	}
+	NewChildren.Add(EmitterObject);
 
-	if (EventHandlerStruct != nullptr && GetStackEditorData().GetStackEntryIsExpanded(ExpandedKey, false))
-	{
-		NewChildren.Add(EventHandlerStruct);
-	}
+	bCanResetToBase.Reset();
+	Super::RefreshChildrenInternal(CurrentChildren, NewChildren, NewIssues);
+}
 
-	NewChildren.Add(EventHandlerExpander);
+void UNiagaraStackEventHandlerPropertiesItem::EventHandlerPropertiesChanged()
+{
 	bCanResetToBase.Reset();
 }
 
-void UNiagaraStackEventHandlerPropertiesItem::EventHandlerExpandedChanged()
+TSharedPtr<IDetailTreeNode> GetEventHandlerArrayPropertyNode(const TArray<TSharedRef<IDetailTreeNode>>& Nodes)
 {
-	RefreshChildren();
+	TArray<TSharedRef<IDetailTreeNode>> ChildrenToCheck;
+	for (TSharedRef<IDetailTreeNode> Node : Nodes)
+	{
+		if (Node->GetNodeType() == EDetailNodeType::Item)
+		{
+			TSharedPtr<IPropertyHandle> NodePropertyHandle = Node->CreatePropertyHandle();
+			if (NodePropertyHandle.IsValid() && NodePropertyHandle->GetProperty()->GetFName() == UNiagaraEmitter::PrivateMemberNames::EventHandlerScriptProps)
+			{
+				return Node;
+			}
+		}
+
+		TArray<TSharedRef<IDetailTreeNode>> Children;
+		Node->GetChildren(Children);
+		ChildrenToCheck.Append(Children);
+	}
+	if (ChildrenToCheck.Num() == 0)
+	{
+		return nullptr;
+	}
+	return GetEventHandlerArrayPropertyNode(ChildrenToCheck);
 }
 
-void UNiagaraStackEventHandlerPropertiesItem::EventHandlerPropertiesChanged()
+void UNiagaraStackEventHandlerPropertiesItem::SelectEmitterStackObjectRootTreeNodes(TArray<TSharedRef<IDetailTreeNode>> Source, TArray<TSharedRef<IDetailTreeNode>>* Selected)
 {
-	bCanResetToBase.Reset();
+	TSharedPtr<IDetailTreeNode> EventHandlerArrayPropertyNode = GetEventHandlerArrayPropertyNode(Source);
+	if (EventHandlerArrayPropertyNode.IsValid())
+	{
+		TArray<TSharedRef<IDetailTreeNode>> EventHandlerArrayItemNodes;
+		EventHandlerArrayPropertyNode->GetChildren(EventHandlerArrayItemNodes);
+		for (TSharedRef<IDetailTreeNode> EventHandlerArrayItemNode : EventHandlerArrayItemNodes)
+		{
+			TSharedPtr<IPropertyHandle> EventHandlerArrayItemPropertyHandle = EventHandlerArrayItemNode->CreatePropertyHandle();
+			if (EventHandlerArrayItemPropertyHandle.IsValid())
+			{
+				UStructProperty* StructProperty = Cast<UStructProperty>(EventHandlerArrayItemPropertyHandle->GetProperty());
+				if (StructProperty != nullptr && StructProperty->Struct == FNiagaraEventScriptProperties::StaticStruct())
+				{
+					TArray<void*> RawData;
+					EventHandlerArrayItemPropertyHandle->AccessRawData(RawData);
+					if (RawData.Num() == 1)
+					{
+						FNiagaraEventScriptProperties* EventScriptProperties = static_cast<FNiagaraEventScriptProperties*>(RawData[0]);
+						if (EventScriptProperties->Script->GetUsageId() == EventScriptUsageId)
+						{
+							EventHandlerArrayItemNode->GetChildren(*Selected);
+							return;
+						}
+					}
+				}
+			}
+		}
+	}
 }
 
-FText UNiagaraStackEventScriptItemGroup::GetDisplayName() const 
+void UNiagaraStackEventScriptItemGroup::Initialize(
+	FRequiredEntryData InRequiredEntryData,
+	TSharedRef<FNiagaraScriptViewModel> InScriptViewModel,
+	ENiagaraScriptUsage InScriptUsage,
+	FGuid InScriptUsageId)
 {
-	return DisplayName;
+	FText ToolTip = LOCTEXT("EventGroupTooltip", "Determines how this Emitter responds to incoming events. There can be more than one event handler script stack per Emitter.");
+	FText TempDisplayName = FText::Format(LOCTEXT("TempDisplayNameFormat", "Event Handler - {0}"), FText::FromString(InScriptUsageId.ToString(EGuidFormats::DigitsWithHyphens)));
+	Super::Initialize(InRequiredEntryData, TempDisplayName, ToolTip, InScriptViewModel, InScriptUsage, InScriptUsageId);
 }
 
-void UNiagaraStackEventScriptItemGroup::RefreshChildrenInternal(const TArray<UNiagaraStackEntry*>& CurrentChildren, TArray<UNiagaraStackEntry*>& NewChildren)
+void UNiagaraStackEventScriptItemGroup::RefreshChildrenInternal(const TArray<UNiagaraStackEntry*>& CurrentChildren, TArray<UNiagaraStackEntry*>& NewChildren, TArray<FStackIssue>& NewIssues)
 {
 	UNiagaraEmitter* Emitter = GetEmitterViewModel()->GetEmitter();
 
 	const FNiagaraEventScriptProperties* EventScriptProperties = Emitter->GetEventHandlers().FindByPredicate(
-		[=](const FNiagaraEventScriptProperties& EventScriptProperties) { return EventScriptProperties.Script->GetUsageId() == GetScriptUsageId(); });
+		[=](const FNiagaraEventScriptProperties& InEventScriptProperties) { return InEventScriptProperties.Script->GetUsageId() == GetScriptUsageId(); });
 
 	if (EventScriptProperties != nullptr)
 	{
-		DisplayName = FText::Format(LOCTEXT("FormatEventScriptDisplayName", "Event Handler - Source: {0}"), FText::FromName(EventScriptProperties->SourceEventName));
+		SetDisplayName(FText::Format(LOCTEXT("FormatEventScriptDisplayName", "Event Handler - Source: {0}"), FText::FromName(EventScriptProperties->SourceEventName)));
 
 		const UNiagaraEmitter* BaseEmitter = FNiagaraStackGraphUtilities::GetBaseEmitter(*Emitter, GetSystemViewModel()->GetSystem());
 		bHasBaseEventHandler = BaseEmitter != nullptr && FNiagaraScriptMergeManager::Get()->HasBaseEventHandler(*BaseEmitter, GetScriptUsageId());
 	}
 	else
 	{
-		DisplayName = LOCTEXT("UnassignedEventDisplayName", "Unassigned Event");
+		SetDisplayName(LOCTEXT("UnassignedEventDisplayName", "Unassigned Event"));
 		bHasBaseEventHandler = false;
 	}
 
 	if (EventHandlerProperties == nullptr)
 	{
 		EventHandlerProperties = NewObject<UNiagaraStackEventHandlerPropertiesItem>(this);
-		EventHandlerProperties->Initialize(GetSystemViewModel(), GetEmitterViewModel(), GetScriptUsageId(), GetStackEditorData());
+		EventHandlerProperties->Initialize(CreateDefaultChildRequiredData(), GetScriptUsageId());
 	}
 	NewChildren.Add(EventHandlerProperties);
 
-	if (EventHandlerPropertiesSpacer == nullptr)
-	{
-		EventHandlerPropertiesSpacer = NewObject<UNiagaraStackSpacer>(this);
-		EventHandlerPropertiesSpacer->Initialize(GetSystemViewModel(), GetEmitterViewModel(), "EventHandlerProperties");
-	}
-	NewChildren.Add(EventHandlerPropertiesSpacer);
-
-	Super::RefreshChildrenInternal(CurrentChildren, NewChildren);
+	Super::RefreshChildrenInternal(CurrentChildren, NewChildren, NewIssues);
 }
 
 bool UNiagaraStackEventScriptItemGroup::CanDelete() const
diff --git a/Engine/Plugins/FX/Niagara/Source/NiagaraEditor/Private/ViewModels/Stack/NiagaraStackFunctionInput.cpp b/Engine/Plugins/FX/Niagara/Source/NiagaraEditor/Private/ViewModels/Stack/NiagaraStackFunctionInput.cpp
--- a/Engine/Plugins/FX/Niagara/Source/NiagaraEditor/Private/ViewModels/Stack/NiagaraStackFunctionInput.cpp
+++ b/Engine/Plugins/FX/Niagara/Source/NiagaraEditor/Private/ViewModels/Stack/NiagaraStackFunctionInput.cpp
@@ -3,6 +3,7 @@
 #include "ViewModels/Stack/NiagaraStackFunctionInput.h"
 #include "ViewModels/Stack/NiagaraStackFunctionInputCollection.h"
 #include "ViewModels/Stack/NiagaraStackObject.h"
+#include "ViewModels/Stack/NiagaraStackSpacer.h"
 #include "NiagaraNodeFunctionCall.h"
 #include "NiagaraNodeAssignment.h"
 #include "NiagaraNodeParameterMapGet.h"
@@ -13,10 +14,10 @@
 #include "NiagaraGraph.h"
 #include "EdGraphSchema_Niagara.h"
 #include "NiagaraEditorUtilities.h"
-#include "NiagaraSystemViewModel.h"
+#include "ViewModels/NiagaraSystemViewModel.h"
 #include "NiagaraSystemScriptViewModel.h"
-#include "NiagaraEmitterHandleViewModel.h"
-#include "NiagaraEmitterViewModel.h"
+#include "ViewModels/NiagaraEmitterHandleViewModel.h"
+#include "ViewModels/NiagaraEmitterViewModel.h"
 #include "NiagaraScriptGraphViewModel.h"
 #include "NiagaraStackEditorData.h"
 #include "NiagaraComponent.h"
@@ -24,7 +25,14 @@
 #include "NiagaraConstants.h"
 #include "ViewModels/Stack/NiagaraStackGraphUtilities.h"
 #include "NiagaraScriptMergeManager.h"
-
+#include "NiagaraNodeCustomHlsl.h"
+#include "NiagaraEditorModule.h"
+#include "NiagaraRendererProperties.h"
+#include "ViewModels/Stack/NiagaraParameterHandle.h"
+#include "NiagaraEmitter.h"
+
+#include "Materials/Material.h"
+#include "Materials/MaterialExpressionDynamicParameter.h"
 #include "ScopedTransaction.h"
 #include "Editor.h"
 #include "UObject/StructOnScope.h"
@@ -37,8 +45,11 @@
 UNiagaraStackFunctionInput::UNiagaraStackFunctionInput()
 	: OwningModuleNode(nullptr)
 	, OwningFunctionCallNode(nullptr)
-	, bCanBePinned(true)
+	, InputMetaData(nullptr)
 	, bUpdatingGraphDirectly(false)
+	, bUpdatingLocalValueDirectly(false)
+	, bShowEditConditionInline(false)
+	, bIsInlineEditConditionToggle(false)
 {
 }
 
@@ -48,6 +59,7 @@ void UNiagaraStackFunctionInput::BeginDestroy()
 	if (OwningFunctionCallNode.IsValid())
 	{
 		OwningFunctionCallNode->GetGraph()->RemoveOnGraphChangedHandler(GraphChangedHandle);
+		OwningFunctionCallNode->GetNiagaraGraph()->RemoveOnGraphNeedsRecompileHandler(OnRecompileHandle);
 
 		if (SourceScript.IsValid())
 		{
@@ -75,16 +87,6 @@ void UNiagaraStackFunctionInput::AddReferencedObjects(FReferenceCollector& Colle
 	}
 }
 
-int32 UNiagaraStackFunctionInput::GetItemIndentLevel() const
-{
-	return ItemIndentLevel;
-}
-
-void UNiagaraStackFunctionInput::SetItemIndentLevel(int32 InItemIndentLevel)
-{
-	ItemIndentLevel = InItemIndentLevel;
-}
-
 // Traverses the path between the owning module node and the function call node this input belongs too collecting up the input handles between them.
 void GenerateInputParameterHandlePath(UNiagaraNodeFunctionCall& ModuleNode, UNiagaraNodeFunctionCall& FunctionCallNode, TArray<FNiagaraParameterHandle>& OutHandlePath)
 {
@@ -101,21 +103,19 @@ void GenerateInputParameterHandlePath(UNiagaraNodeFunctionCall& ModuleNode, UNia
 			UNiagaraNodeParameterMapSet* NextOverrideNode = CastChecked<UNiagaraNodeParameterMapSet>(FunctionOutputPins[0]->LinkedTo[0]->GetOwningNode());
 			UEdGraphPin* NextOverrideNodeOutputPin = FNiagaraStackGraphUtilities::GetParameterMapOutputPin(*NextOverrideNode);
 			
-			TArray<UNiagaraNodeFunctionCall*> NextFunctionCallNodes;
+			CurrentFunctionCallNode = nullptr;
 			for (UEdGraphPin* NextOverrideNodeOutputPinLinkedPin : NextOverrideNodeOutputPin->LinkedTo)
 			{
 				UNiagaraNodeFunctionCall* NextFunctionCallNode = Cast<UNiagaraNodeFunctionCall>(NextOverrideNodeOutputPinLinkedPin->GetOwningNode());
-				if (NextFunctionCallNode != nullptr)
+				if (NextFunctionCallNode != nullptr && NextFunctionCallNode->GetFunctionName() == AliasedHandle.GetNamespace().ToString())
 				{
-					NextFunctionCallNodes.Add(NextFunctionCallNode);
+					CurrentFunctionCallNode = NextFunctionCallNode;
+					break;
 				}
 			}
 
-			if (ensureMsgf(NextFunctionCallNodes.Num() == 1, TEXT("Invalid Stack Graph - Override node not corrected to a single function call node.")))
-			{
-				CurrentFunctionCallNode = NextFunctionCallNodes[0];
-			}
-			else
+
+			if (ensureMsgf(CurrentFunctionCallNode != nullptr, TEXT("Invalid Stack Graph - Function call node for override pin %s could not be found."), *FunctionOutputPins[0]->PinName.ToString()) == false)
 			{
 				OutHandlePath.Empty();
 				return;
@@ -130,24 +130,24 @@ void GenerateInputParameterHandlePath(UNiagaraNodeFunctionCall& ModuleNode, UNia
 }
 
 void UNiagaraStackFunctionInput::Initialize(
-	TSharedRef<FNiagaraSystemViewModel> InSystemViewModel,
-	TSharedRef<FNiagaraEmitterViewModel> InEmitterViewModel,
-	UNiagaraStackEditorData& InStackEditorData,
+	FRequiredEntryData InRequiredEntryData,
 	UNiagaraNodeFunctionCall& InModuleNode,
 	UNiagaraNodeFunctionCall& InInputFunctionCallNode,
 	FName InInputParameterHandle,
-	FNiagaraTypeDefinition InInputType)
+	FNiagaraTypeDefinition InInputType,
+	FString InOwnerStackItemEditorDataKey)
 {
 	checkf(OwningModuleNode.IsValid() == false && OwningFunctionCallNode.IsValid() == false, TEXT("Can only initialize once."));
-	Super::Initialize(InSystemViewModel, InEmitterViewModel);
-	StackEditorData = &InStackEditorData;
+	bool bInputIsAdvanced = false;
+	FString InputStackEditorDataKey = FString::Printf(TEXT("%s-Input-%s"), *InInputFunctionCallNode.NodeGuid.ToString(EGuidFormats::DigitsWithHyphens), *InInputParameterHandle.ToString());
+	Super::Initialize(InRequiredEntryData, bInputIsAdvanced, InOwnerStackItemEditorDataKey, InputStackEditorDataKey);
 	OwningModuleNode = &InModuleNode;
 	OwningFunctionCallNode = &InInputFunctionCallNode;
 	OwningAssignmentNode = Cast<UNiagaraNodeAssignment>(OwningFunctionCallNode.Get());
 
 	UNiagaraNodeOutput* OutputNode = FNiagaraStackGraphUtilities::GetEmitterOutputNodeForStackNode(*OwningModuleNode.Get());
-	UNiagaraEmitter* ParentEmitter = InEmitterViewModel->GetEmitter();
-	UNiagaraSystem* ParentSystem = &InSystemViewModel->GetSystem();
+	UNiagaraEmitter* ParentEmitter = GetEmitterViewModel()->GetEmitter();
+	UNiagaraSystem* ParentSystem = &GetSystemViewModel()->GetSystem();
 	if (OutputNode)
 	{
 		TArray<UNiagaraScript*> Scripts;
@@ -157,10 +157,8 @@ void UNiagaraStackFunctionInput::Initialize(
 		}
 		if (ParentSystem != nullptr)
 		{
-			Scripts.Add(ParentSystem->GetSystemSpawnScript(true));
-			Scripts.Add(ParentSystem->GetSystemUpdateScript(true));
-			Scripts.Add(ParentSystem->GetSystemSpawnScript(false));
-			Scripts.Add(ParentSystem->GetSystemUpdateScript(false));
+			Scripts.Add(ParentSystem->GetSystemSpawnScript());
+			Scripts.Add(ParentSystem->GetSystemUpdateScript());
 		}
 
 		for (UNiagaraScript* Script : Scripts)
@@ -196,6 +194,8 @@ void UNiagaraStackFunctionInput::Initialize(
 
 	GraphChangedHandle = OwningFunctionCallNode->GetGraph()->AddOnGraphChangedHandler(
 		FOnGraphChanged::FDelegate::CreateUObject(this, &UNiagaraStackFunctionInput::OnGraphChanged));
+	OnRecompileHandle = OwningFunctionCallNode->GetNiagaraGraph()->AddOnGraphNeedsRecompileHandler(
+		FOnGraphChanged::FDelegate::CreateUObject(this, &UNiagaraStackFunctionInput::OnGraphChanged));
 
 	InputParameterHandle = FNiagaraParameterHandle(InInputParameterHandle);
 	GenerateInputParameterHandlePath(*OwningModuleNode, *OwningFunctionCallNode, InputParameterHandlePath);
@@ -207,14 +207,14 @@ void UNiagaraStackFunctionInput::Initialize(
 	InputType = InInputType;
 	StackEditorDataKey = FNiagaraStackGraphUtilities::GenerateStackFunctionInputEditorDataKey(*OwningFunctionCallNode.Get(), InputParameterHandle);
 
-	/*
-	UE_LOG(LogNiagaraEditor, Log, TEXT("Scripts for Parameter '%s'"), *RapidIterationParameterName);
-	for (UNiagaraScript* Script : AffectedScripts)
+	TArray<UNiagaraScript*> AffectedScriptsNotWeak;
+	for (TWeakObjectPtr<UNiagaraScript> AffectedScript : AffectedScripts)
 	{
-		UE_LOG(LogNiagaraEditor, Log, TEXT("Added... '%s'"), *Script->GetFullName());
+		AffectedScriptsNotWeak.Add(AffectedScript.Get());
 	}
-	*/
 
+	EditCondition.Initialize(SourceScript.Get(), AffectedScriptsNotWeak, GetEmitterViewModel()->GetEmitter()->GetUniqueEmitterName(), OwningFunctionCallNode.Get());
+	VisibleCondition.Initialize(SourceScript.Get(), AffectedScriptsNotWeak, GetEmitterViewModel()->GetEmitter()->GetUniqueEmitterName(), OwningFunctionCallNode.Get());
 }
 
 const UNiagaraNodeFunctionCall& UNiagaraStackFunctionInput::GetInputFunctionCallNode() const
@@ -227,11 +227,6 @@ UNiagaraStackFunctionInput::EValueMode UNiagaraStackFunctionInput::GetValueMode(
 	return InputValues.Mode;
 }
 
-bool UNiagaraStackFunctionInput::GetCanBePinned() const
-{
-	return OwningFunctionCallNode.IsValid() && OwningFunctionCallNode.Get()->FunctionScript->Usage == ENiagaraScriptUsage::Module;
-}
-
 const FNiagaraTypeDefinition& UNiagaraStackFunctionInput::GetInputType() const
 {
 	return InputType;
@@ -242,6 +237,20 @@ FText UNiagaraStackFunctionInput::GetTooltipText() const
 	return GetTooltipText(InputValues.Mode);
 }
 
+bool UNiagaraStackFunctionInput::GetIsEnabled() const
+{
+	return OwningFunctionCallNode->GetDesiredEnabledState() == ENodeEnabledState::Enabled;
+}
+
+UObject* UNiagaraStackFunctionInput::GetExternalAsset() const
+{
+	if (OwningFunctionCallNode.IsValid() && OwningFunctionCallNode->FunctionScript != nullptr && OwningFunctionCallNode->FunctionScript->IsAsset())
+	{
+		return OwningFunctionCallNode->FunctionScript;
+	}
+	return nullptr;
+}
+
 FText UNiagaraStackFunctionInput::GetTooltipText(EValueMode InValueMode) const
 {
 	FNiagaraVariable ValueVariable;
@@ -286,9 +295,11 @@ FText UNiagaraStackFunctionInput::GetTooltipText(EValueMode InValueMode) const
 	return FText::FromName(ValueVariable.GetName());
 }
 
-void UNiagaraStackFunctionInput::RefreshChildrenInternal(const TArray<UNiagaraStackEntry*>& CurrentChildren, TArray<UNiagaraStackEntry*>& NewChildren)
+void UNiagaraStackFunctionInput::RefreshChildrenInternal(const TArray<UNiagaraStackEntry*>& CurrentChildren, TArray<UNiagaraStackEntry*>& NewChildren, TArray<FStackIssue>& NewIssues)
 {
-	RapidIterationParameterName = CreateRapidIterationVariableName(AliasedInputParameterHandle.GetParameterHandleString().ToString());
+	RapidIterationParameter = CreateRapidIterationVariable(AliasedInputParameterHandle.GetParameterHandleString());
+
+	RefreshFromMetaData();
 	RefreshValues();
 
 	if (InputValues.Mode == EValueMode::Dynamic && InputValues.DynamicNode.IsValid())
@@ -302,14 +313,11 @@ void UNiagaraStackFunctionInput::RefreshChildrenInternal(const TArray<UNiagaraSt
 
 		if (DynamicInputEntry == nullptr)
 		{ 
-			UNiagaraStackFunctionInputCollection::FDisplayOptions DisplayOptions;
-			DisplayOptions.DisplayName = LOCTEXT("DynamicInputCollecitonDisplayName", "Dynamic Input Function Input");
-			DisplayOptions.bShouldShowInStack = false;
-			DisplayOptions.ChildItemIndentLevel = ItemIndentLevel + 1;
-
 			DynamicInputEntry = NewObject<UNiagaraStackFunctionInputCollection>(this);
-			DynamicInputEntry->Initialize(GetSystemViewModel(), GetEmitterViewModel(), *StackEditorData, *OwningModuleNode, *InputValues.DynamicNode.Get(), DisplayOptions);
+			DynamicInputEntry->Initialize(CreateDefaultChildRequiredData(), *OwningModuleNode, *InputValues.DynamicNode.Get(), GetOwnerStackItemEditorDataKey());
+			DynamicInputEntry->SetShouldShowInStack(false);
 		}
+
 		NewChildren.Add(DynamicInputEntry);
 	}
 
@@ -321,13 +329,160 @@ void UNiagaraStackFunctionInput::RefreshChildrenInternal(const TArray<UNiagaraSt
 		if(ValueObjectEntry == nullptr)
 		{
 			ValueObjectEntry = NewObject<UNiagaraStackObject>(this);
-			ValueObjectEntry->Initialize(GetSystemViewModel(), GetEmitterViewModel(), InputValues.DataObjects.GetValueObject());
-			ValueObjectEntry->SetItemIndentLevel(ItemIndentLevel + 1);
+			ValueObjectEntry->Initialize(CreateDefaultChildRequiredData(), InputValues.DataObjects.GetValueObject(), GetOwnerStackItemEditorDataKey(), OwningFunctionCallNode.Get());
 		}
 		NewChildren.Add(ValueObjectEntry);
 	}
+
+	if (NewChildren.Num() > 0)
+	{
+		UNiagaraStackSpacer* ChildrenSpacer = FindCurrentChildOfTypeByPredicate<UNiagaraStackSpacer>(CurrentChildren,
+			[=](UNiagaraStackSpacer* CurrentSpacer) { return true; });
+
+		if (ChildrenSpacer == nullptr)
+		{
+			ChildrenSpacer = NewObject<UNiagaraStackSpacer>(this);
+			ChildrenSpacer->Initialize(CreateDefaultChildRequiredData(), "ChildrenSpacer", .5f, UNiagaraStackEntry::EStackRowStyle::ItemContent);
+		}
+
+		NewChildren.Add(ChildrenSpacer);
+	}
+
+	DisplayNameOverride.Reset();
+
+	if (InputMetaData)
+	{
+		const FString* FoundDisplayName = InputMetaData->PropertyMetaData.Find(TEXT("DisplayName"));
+		const FString* FoundDisplayNameArg0 = InputMetaData->PropertyMetaData.Find(TEXT("DisplayNameArg0"));
+		if (FoundDisplayName != nullptr)
+		{
+			FString DisplayNameStr = *FoundDisplayName;
+			if (FoundDisplayNameArg0 != nullptr)
+			{
+				TArray<FStringFormatArg> Args;
+				Args.Add(FStringFormatArg(ResolveDisplayNameArgument(*FoundDisplayNameArg0)));
+				DisplayNameStr = FString::Format(*DisplayNameStr, Args);
+			}
+			DisplayNameOverride = FText::FromString(DisplayNameStr);
+		}	
+	}
+}
+
+class UNiagaraStackFunctionInputUtilities
+{
+public:
+	static bool GetMaterialExpressionDynamicParameter(UNiagaraEmitter* InEmitter, TArray<UMaterialExpressionDynamicParameter*>& OutDynamicParameterExpressions)
+	{
+		TArray<UMaterial*> Materials = GetMaterialFromEmitter(InEmitter);
+		
+		OutDynamicParameterExpressions.Reset();
+
+		// Find the dynamic parameters expression from the material.
+		// @YannickLange todo: Notify user that the material did not have any dynamic parameters. Maybe add them from code?
+		for (UMaterial* Material : Materials)
+		{
+			if (Material != nullptr)
+			{
+				for (UMaterialExpression* Expression : Material->Expressions)
+				{
+					UMaterialExpressionDynamicParameter* DynParamExpFound = Cast<UMaterialExpressionDynamicParameter>(Expression);
+
+					if (DynParamExpFound != nullptr)
+					{
+						OutDynamicParameterExpressions.Add(DynParamExpFound);
+					}
+				}
+			}
+		}
+
+		return OutDynamicParameterExpressions.Num() > 0;
+	}
+
+	static TArray<UMaterial*> GetMaterialFromEmitter(UNiagaraEmitter* InEmitter)
+	{
+		TArray<UMaterial*> ResultMaterials;
+		if (InEmitter->GetRenderers().Num() > 0)
+		{
+			for (UNiagaraRendererProperties* RenderProperties : InEmitter->GetRenderers())
+			{
+				TArray<UMaterialInterface*> UsedMaterialInteraces;
+				RenderProperties->GetUsedMaterials(UsedMaterialInteraces);
+				for (UMaterialInterface* UsedMaterialInterface : UsedMaterialInteraces)
+				{
+					if (UsedMaterialInterface != nullptr)
+					{
+						UMaterial* UsedMaterial = UsedMaterialInterface->GetBaseMaterial();
+						if (UsedMaterial != nullptr)
+						{
+							ResultMaterials.AddUnique(UsedMaterial);
+							break;
+						}
+					}
+				}
+			}
+		}
+		return ResultMaterials;
+	}
+};
+
+FString UNiagaraStackFunctionInput::ResolveDisplayNameArgument(const FString& InArg) const
+{
+	if (InArg.StartsWith(TEXT("MaterialDynamicParam")))
+	{
+		TSharedRef<FNiagaraEmitterViewModel> ThisEmitterViewModel = GetEmitterViewModel();
+		TArray<UMaterialExpressionDynamicParameter*> ExpressionParams;
+		if (!UNiagaraStackFunctionInputUtilities::GetMaterialExpressionDynamicParameter(ThisEmitterViewModel->GetEmitter(), ExpressionParams))
+		{
+			return InArg.Replace(TEXT("MaterialDynamic"), TEXT("")) + TEXT(" (No material found using dynamic params)");
+		}
+		
+		FString Suffix = InArg.Right(3);
+		int32 ParamIdx;
+		LexicalConversion::FromString(ParamIdx, *Suffix.Left(1));
+		int32 ParamSlotIdx;
+		LexicalConversion::FromString(ParamSlotIdx, *Suffix.Right(1));
+		
+		if (ParamIdx < 0 || ParamIdx > 3 || ParamSlotIdx < 0 || ParamSlotIdx > 3)
+		{
+			return InArg.Replace(TEXT("MaterialDynamic"), TEXT("")) + TEXT(" (error parsing parameter name)");
+		}
+
+		FName ParamName = NAME_None;
+		bool bAllSame = true;
+		for (UMaterialExpressionDynamicParameter* Expression : ExpressionParams)
+		{
+			const FExpressionOutput& Output = Expression->GetOutputs()[ParamIdx];
+			if (ParamSlotIdx == Expression->ParameterIndex)
+			{
+				if (ParamName == NAME_None)
+				{
+					ParamName = Output.OutputName;
+				}
+				else if (ParamName != Output.OutputName)
+				{
+					bAllSame = false;
+				}
+			}
+		}
+
+		if (ParamName != NAME_None)
+		{
+			if (bAllSame)
+			{
+				return ParamName.ToString();
+			}
+			else
+			{
+				return ParamName.ToString() + TEXT(" (Multiple Aliases Found)");
+			}
+		}
+
+		return InArg.Replace(TEXT("MaterialDynamic"), TEXT("")) + TEXT(" (Parameter not used in materials.)");
+	}
+	return FString();
 }
 
+
 TSharedPtr<FStructOnScope> UNiagaraStackFunctionInput::FInputValues::GetLocalStructToReuse()
 {
 	return Mode == EValueMode::Local ? LocalStruct : TSharedPtr<FStructOnScope>();
@@ -356,7 +511,7 @@ void UNiagaraStackFunctionInput::RefreshValues()
 		UEdGraphPin* OverridePin = GetOverridePin();
 		UEdGraphPin* ValuePin = OverridePin != nullptr ? OverridePin : DefaultPin;
 
-		if (TryGetCurrentLocalValue(InputValues.LocalStruct, *ValuePin, OldValues.GetLocalStructToReuse()))
+		if (TryGetCurrentLocalValue(InputValues.LocalStruct, *DefaultPin, *ValuePin, OldValues.GetLocalStructToReuse()))
 		{
 			InputValues.Mode = EValueMode::Local;
 		}
@@ -368,29 +523,58 @@ void UNiagaraStackFunctionInput::RefreshValues()
 		{
 			InputValues.Mode = EValueMode::Data;
 		}
+		else if (TryGetCurrentExpressionValue(InputValues.ExpressionNode, OverridePin))
+		{
+			InputValues.Mode = EValueMode::Expression;
+		}
 		else if (TryGetCurrentDynamicValue(InputValues.DynamicNode, OverridePin))
 		{
 			InputValues.Mode = EValueMode::Dynamic;
 		}
 	}
 
+	bCanReset.Reset();
 	bCanResetToBase.Reset();
 	ValueChangedDelegate.Broadcast();
 }
 
-FText UNiagaraStackFunctionInput::GetDisplayName() const
+void UNiagaraStackFunctionInput::RefreshFromMetaData()
 {
-	return DisplayName;
-}
+	if (OwningFunctionCallNode->FunctionScript != nullptr)
+	{
+		UNiagaraGraph* FunctionGraph = CastChecked<UNiagaraScriptSource>(OwningFunctionCallNode->FunctionScript->GetSource())->NodeGraph;
+		FNiagaraVariable InputVariable(InputType, InputParameterHandle.GetParameterHandleString());
+		InputMetaData = FunctionGraph->GetMetaData(InputVariable);
 
-FName UNiagaraStackFunctionInput::GetTextStyleName() const
-{
-	return "NiagaraEditor.Stack.ParameterText";
+		if (InputMetaData != nullptr)
+		{
+			SetIsAdvanced(InputMetaData->PropertyMetaData.Contains(FNiagaraEditorModule::FInputMetaDataKeys::AdvancedDisplay));
+
+			EditCondition.Refresh(InputMetaData->PropertyMetaData.Find(FNiagaraEditorModule::FInputMetaDataKeys::EditCondition));
+			if (EditCondition.IsValid() && EditCondition.GetConditionInputType() == FNiagaraTypeDefinition::GetBoolDef())
+			{
+				FNiagaraVariableMetaData* EditConditionInputMetadata = EditCondition.GetConditionInputMetaData();
+				if (EditConditionInputMetadata != nullptr)
+				{
+					bShowEditConditionInline = EditConditionInputMetadata->PropertyMetaData.Contains(FNiagaraEditorModule::FInputMetaDataKeys::InlineEditConditionToggle);
+				}
+			}
+			else
+			{
+				bShowEditConditionInline = false;
+			}
+
+			VisibleCondition.Refresh(InputMetaData->PropertyMetaData.Find(FNiagaraEditorModule::FInputMetaDataKeys::VisibleCondition));
+
+			bIsInlineEditConditionToggle = InputType == FNiagaraTypeDefinition::GetBoolDef() && 
+				InputMetaData->PropertyMetaData.Contains(FNiagaraEditorModule::FInputMetaDataKeys::InlineEditConditionToggle);
+		}
+	}
 }
 
-bool UNiagaraStackFunctionInput::GetCanExpand() const
+FText UNiagaraStackFunctionInput::GetDisplayName() const
 {
-	return true;
+	return DisplayNameOverride.IsSet() ? DisplayNameOverride.GetValue() : DisplayName;
 }
 
 const TArray<FNiagaraParameterHandle>& UNiagaraStackFunctionInput::GetInputParameterHandlePath() const
@@ -506,7 +690,7 @@ void UNiagaraStackFunctionInput::GetAvailableParameterHandles(TArray<FNiagaraPar
 
 	for (UNiagaraNodeOutput* OutputNode : AllOutputNodes)
 	{
-		if (OutputNode == CurrentOutputNode || UsageRunsBefore(OutputNode->GetUsage(), CurrentOutputNode->GetUsage()))
+		if (OutputNode == CurrentOutputNode || (CurrentOutputNode != nullptr && UsageRunsBefore(OutputNode->GetUsage(), CurrentOutputNode->GetUsage())) || (CurrentOutputNode != nullptr && IsSpawnUsage(CurrentOutputNode->GetUsage())))
 		{
 			TArray<FNiagaraParameterHandle> AvailableParameterHandlesForThisOutput;
 			TArray<FNiagaraStackGraphUtilities::FStackNodeGroup> StackGroups;
@@ -528,6 +712,7 @@ void UNiagaraStackFunctionInput::GetAvailableParameterHandles(TArray<FNiagaraPar
 					for (int32 j = 0; j < Builder.Histories[0].Variables.Num(); j++)
 					{
 						FNiagaraVariable& Variable = Builder.Histories[0].Variables[j];
+						FNiagaraParameterHandle AvailableHandle = FNiagaraParameterHandle(Variable.GetName());
 						if (Variable.GetType() == InputType)
 						{
 							TArray<const UEdGraphPin*>& WriteHistory = Builder.Histories[0].PerVariableWriteHistory[j];
@@ -535,7 +720,6 @@ void UNiagaraStackFunctionInput::GetAvailableParameterHandles(TArray<FNiagaraPar
 							{
 								if (Cast<UNiagaraNodeParameterMapSet>(WritePin->GetOwningNode()) != nullptr)
 								{
-									FNiagaraParameterHandle AvailableHandle = FNiagaraParameterHandle(Variable.GetName());
 									AvailableParameterHandles.AddUnique(AvailableHandle);
 									AvailableParameterHandlesForThisOutput.AddUnique(AvailableHandle);
 									break;
@@ -569,6 +753,11 @@ UNiagaraNodeFunctionCall* UNiagaraStackFunctionInput::GetDynamicInputNode() cons
 	return InputValues.DynamicNode.Get();
 }
 
+UNiagaraNodeCustomHlsl* UNiagaraStackFunctionInput::GetExpressionNode() const
+{
+	return InputValues.ExpressionNode.Get();
+}
+
 void UNiagaraStackFunctionInput::GetAvailableDynamicInputs(TArray<UNiagaraScript*>& AvailableDynamicInputs)
 {
 	UEnum* NiagaraScriptUsageEnum = FindObjectChecked<UEnum>(ANY_PACKAGE, TEXT("ENiagaraScriptUsage"), true);
@@ -624,31 +813,40 @@ void UNiagaraStackFunctionInput::SetDynamicInput(UNiagaraScript* DynamicInput)
 
 	UNiagaraNodeFunctionCall* FunctionCallNode;
 	FNiagaraStackGraphUtilities::SetDynamicInputForFunctionInput(OverridePin, DynamicInput, FunctionCallNode);
-	FNiagaraStackGraphUtilities::InitializeDataInterfaceInputs(GetSystemViewModel(), GetEmitterViewModel(), *StackEditorData, *OwningModuleNode, *FunctionCallNode);
+	FNiagaraStackGraphUtilities::InitializeStackFunctionInputs(GetSystemViewModel(), GetEmitterViewModel(), GetStackEditorData(), *OwningModuleNode, *FunctionCallNode);
 	FNiagaraStackGraphUtilities::RelayoutGraph(*OwningFunctionCallNode->GetGraph());
 
 	RefreshChildren();
 }
 
-TSharedPtr<FStructOnScope> UNiagaraStackFunctionInput::GetLocalValueStruct()
+void UNiagaraStackFunctionInput::SetCustomExpression(const FString& InputText)
 {
-	return InputValues.LocalStruct;
-}
+	FScopedTransaction ScopedTransaction(LOCTEXT("SetCustomExpressionInput", "Make custom expression input"));
 
-UNiagaraDataInterface* UNiagaraStackFunctionInput::GetDataValueObject()
-{
-	return InputValues.DataObjects.GetValueObject();
+
+	UEdGraphPin& OverridePin = GetOrCreateOverridePin();
+	RemoveNodesForOverridePin(OverridePin);
+	if (IsRapidIterationCandidate())
+	{
+		RemoveRapidIterationParametersForAffectedScripts();
+	}
+
+	UNiagaraNodeCustomHlsl* FunctionCallNode;
+	FNiagaraStackGraphUtilities::SetCustomExpressionForFunctionInput(OverridePin, FunctionCallNode);
+	FNiagaraStackGraphUtilities::InitializeStackFunctionInputs(GetSystemViewModel(), GetEmitterViewModel(), GetStackEditorData(), *OwningModuleNode, *FunctionCallNode);
+	FNiagaraStackGraphUtilities::RelayoutGraph(*OwningFunctionCallNode->GetGraph());
+
+	RefreshChildren();
 }
 
-bool UNiagaraStackFunctionInput::GetIsPinned() const
+TSharedPtr<FStructOnScope> UNiagaraStackFunctionInput::GetLocalValueStruct()
 {
-	return StackEditorData->GetModuleInputIsPinned(StackEditorDataKey);
+	return InputValues.LocalStruct;
 }
 
-void UNiagaraStackFunctionInput::SetIsPinned(bool bIsPinned)
+UNiagaraDataInterface* UNiagaraStackFunctionInput::GetDataValueObject()
 {
-	StackEditorData->SetModuleInputIsPinned(StackEditorDataKey, bIsPinned);
-	PinnedChangedDelegate.Broadcast();
+	return InputValues.DataObjects.GetValueObject();
 }
 
 void UNiagaraStackFunctionInput::NotifyBeginLocalValueChange()
@@ -666,12 +864,13 @@ void UNiagaraStackFunctionInput::NotifyEndLocalValueChange()
 
 bool UNiagaraStackFunctionInput::IsRapidIterationCandidate() const
 {
-	return InputType != FNiagaraTypeDefinition::GetBoolDef() && !InputType.IsEnum() &&
-		InputType != FNiagaraTypeDefinition::GetParameterMapDef() && !InputType.IsDataInterface();
+	return FNiagaraStackGraphUtilities::IsRapidIterationType(InputType);
 }
 
 void UNiagaraStackFunctionInput::SetLocalValue(TSharedRef<FStructOnScope> InLocalValue)
 {
+	TGuardValue<bool> UpdateGuard(bUpdatingLocalValueDirectly, true);
+
 	UEdGraphPin* DefaultPin = GetDefaultPin();
 	UEdGraphPin* OverridePin = GetOverridePin();
 	UEdGraphPin* ValuePin = DefaultPin;
@@ -686,7 +885,7 @@ void UNiagaraStackFunctionInput::SetLocalValue(TSharedRef<FStructOnScope> InLoca
 
 	TSharedPtr<FStructOnScope> CurrentValue;
 	bool bCanHaveLocalValue = ValuePin != nullptr;
-	bool bHasLocalValue = bCanHaveLocalValue && InputValues.Mode == EValueMode::Local && TryGetCurrentLocalValue(CurrentValue, *ValuePin, TSharedPtr<FStructOnScope>());
+	bool bHasLocalValue = bCanHaveLocalValue && InputValues.Mode == EValueMode::Local && TryGetCurrentLocalValue(CurrentValue, *DefaultPin, *ValuePin, TSharedPtr<FStructOnScope>());
 	bool bLocalValueMatchesSetValue = bHasLocalValue && FNiagaraEditorUtilities::DataMatches(*CurrentValue.Get(), InLocalValue.Get());
 
 	if (bCanHaveLocalValue == false || bLocalValueMatchesSetValue)
@@ -719,22 +918,15 @@ void UNiagaraStackFunctionInput::SetLocalValue(TSharedRef<FStructOnScope> InLoca
 			OverrideNode->Modify();
 			OverrideNode->RemovePin(OverridePin);
 		}
-	
-		FNiagaraVariable ValueVariable = NiagaraSchema->PinToNiagaraVariable(ValuePin, true);
-		ValueVariable.SetName(*RapidIterationParameterName);
 
 		for (TWeakObjectPtr<UNiagaraScript> Script : AffectedScripts)
 		{
-			int32 FoundOffset = Script->RapidIterationParameters.IndexOf(ValueVariable);
-			checkSlow(!ValueVariable.IsDataInterface());
-			if (FoundOffset == INDEX_NONE)
+			if(Script->RapidIterationParameters.AddParameter(RapidIterationParameter))
 			{
-				Script->RapidIterationParameters.AddParameter(ValueVariable);
-				UE_LOG(LogNiagaraEditor, Log, TEXT("Adding Parameter %s to Script %s"), *ValueVariable.GetName().ToString(), *Script->GetFullName());
+				UE_LOG(LogNiagaraEditor, Log, TEXT("Adding Parameter %s to Script %s"), *RapidIterationParameter.GetName().ToString(), *Script->GetFullName());
 			}
-			Script->RapidIterationParameters.SetParameterData(InLocalValue->GetStructMemory(), ValueVariable);
+			Script->RapidIterationParameters.SetParameterData(InLocalValue->GetStructMemory(), RapidIterationParameter);
 		}
-		GetSystemViewModel()->ReInitializeSystemInstances();
 	}
 	else 
 	{
@@ -752,7 +944,7 @@ void UNiagaraStackFunctionInput::SetLocalValue(TSharedRef<FStructOnScope> InLoca
 
 			OverridePin->Modify();
 			OverridePin->DefaultValue = PinDefaultValue;
-			EmitterGraph->NotifyGraphNeedsRecompile();
+			Cast<UNiagaraNode>(OverridePin->GetOwningNode())->MarkNodeRequiresSynchronization(TEXT("OverridePin Default Value Changed"), true);
 		}
 	}
 	
@@ -766,43 +958,61 @@ void UNiagaraStackFunctionInput::SetLocalValue(TSharedRef<FStructOnScope> InLoca
 
 bool UNiagaraStackFunctionInput::CanReset() const
 {
-	if(InputValues.Mode == EValueMode::Data)
-	{
-		// For data values a copy of the default object should have been created automatically and attached to the override pin for this input.  If a 
-		// copy of the default object wasn't created, the input can be reset to create one.  If a copy of the data object is available it can be
-		// reset if it's different from it's default value.
-		bool bHasDataValueObject = InputValues.DataObjects.GetValueObject() != nullptr;
-		bool bHasDefaultDataValueObject = InputValues.DataObjects.GetDefaultValueObject() != nullptr;
-		bool bIsDataValueDifferentFromDefaultDataValue = bHasDataValueObject && bHasDefaultDataValueObject 
-			&& InputValues.DataObjects.GetValueObject()->Equals(InputValues.DataObjects.GetDefaultValueObject()) == false;
-		return bHasDataValueObject == false || bHasDefaultDataValueObject == false || bIsDataValueDifferentFromDefaultDataValue;
-	}
-	else
+	if (bCanReset.IsSet() == false)
 	{
-		if (IsRapidIterationCandidate())
+		bool bNewCanReset;
+		if (InputValues.Mode == EValueMode::Data)
 		{
-			if (GetOverridePin() != nullptr)
-			{
-				return true;
-			}
-			else if (InputValues.LocalStruct.IsValid() && !FNiagaraEditorUtilities::DataMatches(GetDefaultValue(), *(InputValues.LocalStruct.Get())))
+			// For data values a copy of the default object should have been created automatically and attached to the override pin for this input.  If a 
+			// copy of the default object wasn't created, the input can be reset to create one.  If a copy of the data object is available it can be
+			// reset if it's different from it's default value.
+			bool bHasDataValueObject = InputValues.DataObjects.GetValueObject() != nullptr;
+			bool bHasDefaultDataValueObject = InputValues.DataObjects.GetDefaultValueObject() != nullptr;
+			bool bIsDataValueDifferentFromDefaultDataValue = bHasDataValueObject && bHasDefaultDataValueObject
+				&& InputValues.DataObjects.GetValueObject()->Equals(InputValues.DataObjects.GetDefaultValueObject()) == false;
+			bNewCanReset = bHasDataValueObject == false || bHasDefaultDataValueObject == false || bIsDataValueDifferentFromDefaultDataValue;
+		}
+		else
+		{
+			UEdGraphPin* DefaultPin = GetDefaultPin();
+			if(DefaultPin->LinkedTo.Num() == 0)
 			{
-				return true;
+				if (GetOverridePin() != nullptr)
+				{
+					bNewCanReset = true;
+				}
+				else if (IsRapidIterationCandidate())
+				{
+					FNiagaraVariable DefaultVar = GetDefaultVariableForRapidIterationParameter();
+					bool bHasValidLocalValue = InputValues.LocalStruct.IsValid();
+					bool bHasValidDefaultValue = DefaultVar.IsValid();
+					bNewCanReset = bHasValidLocalValue && bHasValidDefaultValue && FNiagaraEditorUtilities::DataMatches(DefaultVar, *InputValues.LocalStruct.Get()) == false;
+				}
+				else
+				{
+					bNewCanReset = false;
+				}
 			}
 			else
 			{
-				return false;
+				if (FNiagaraStackGraphUtilities::IsValidDefaultDynamicInput(*SourceScript, *DefaultPin))
+				{
+					UEdGraphPin* OverridePin = GetOverridePin();
+					bNewCanReset = OverridePin == nullptr || FNiagaraStackGraphUtilities::DoesDynamicInputMatchDefault(GetEmitterViewModel()->GetEmitter()->GetUniqueEmitterName(), *SourceScript,
+						*OwningFunctionCallNode, *OverridePin, InputParameterHandle.GetName(), *DefaultPin) == false;
+				}
+				else
+				{
+					bNewCanReset = GetOverridePin() != nullptr;
+				}
 			}
 		}
-		else
-		{
-			// All other input modes can be reset if there is an override pin available.
-			return GetOverridePin() != nullptr;
-		}
+		bCanReset = bNewCanReset;
 	}
+	return bCanReset.GetValue();
 }
 
-FNiagaraVariable UNiagaraStackFunctionInput::GetDefaultValue() const
+FNiagaraVariable UNiagaraStackFunctionInput::GetDefaultVariableForRapidIterationParameter() const
 {
 	FNiagaraVariable Var;
 	UEdGraphPin* DefaultPin = GetDefaultPin();
@@ -810,7 +1020,7 @@ FNiagaraVariable UNiagaraStackFunctionInput::GetDefaultValue() const
 	{
 		const UEdGraphSchema_Niagara* NiagaraSchema = GetDefault<UEdGraphSchema_Niagara>();
 		Var = NiagaraSchema->PinToNiagaraVariable(DefaultPin, true);
-		Var.SetName(*RapidIterationParameterName);
+		Var.SetName(*RapidIterationParameter.GetName().ToString());
 	}
 	return Var;
 }
@@ -822,19 +1032,13 @@ bool UNiagaraStackFunctionInput::UpdateRapidIterationParametersForAffectedScript
 		Script->Modify();
 	}
 
-	FNiagaraVariable Var(InputType, *RapidIterationParameterName);
-
 	for (TWeakObjectPtr<UNiagaraScript> Script : AffectedScripts)
 	{
-		int32 FoundOffset = Script->RapidIterationParameters.IndexOf(Var);
-		checkSlow(!Var.IsDataInterface());
-		if (FoundOffset == INDEX_NONE)
+		if (Script->RapidIterationParameters.AddParameter(RapidIterationParameter))
 		{
-			Script->RapidIterationParameters.AddParameter(Var);
-			UE_LOG(LogNiagaraEditor, Log, TEXT("Adding Parameter %s to Script %s"), *Var.GetName().ToString(), *Script->GetFullName());
-
+			UE_LOG(LogNiagaraEditor, Log, TEXT("Adding Parameter %s to Script %s"), *RapidIterationParameter.GetName().ToString(), *Script->GetFullName());
 		}
-		Script->RapidIterationParameters.SetParameterData(Data, Var);
+		Script->RapidIterationParameters.SetParameterData(Data, RapidIterationParameter);
 	}
 	GetSystemViewModel()->ResetSystem();
 	return true;
@@ -847,18 +1051,12 @@ bool UNiagaraStackFunctionInput::RemoveRapidIterationParametersForAffectedScript
 		Script->Modify();
 	}
 
-	FNiagaraVariable Var(InputType, *RapidIterationParameterName);
 	for (TWeakObjectPtr<UNiagaraScript> Script : AffectedScripts)
 	{
-		int32 FoundOffset = Script->RapidIterationParameters.IndexOf(Var);
-		if (FoundOffset == INDEX_NONE)
+		if (Script->RapidIterationParameters.RemoveParameter(RapidIterationParameter))
 		{
-			return false;
+			UE_LOG(LogNiagaraEditor, Log, TEXT("Removed Var '%s' from Script %s"), *RapidIterationParameter.GetName().ToString(), *Script->GetFullName());
 		}
-
-		UE_LOG(LogNiagaraEditor, Log, TEXT("Removed Var '%s' from Script %s"), *Var.GetName().ToString(), *Script->GetFullName());
-
-		Script->RapidIterationParameters.RemoveParameter(Var);
 	}
 	return true;
 }
@@ -900,11 +1098,12 @@ void UNiagaraStackFunctionInput::Reset()
 		// For all other value modes removing the nodes connected to the override pin resets them.
 		UNiagaraNodeParameterMapSet* OverrideNode = GetOverrideNode();
 		UEdGraphPin* OverridePin = GetOverridePin();
+		UEdGraphPin* DefaultPin = GetDefaultPin();
 		bool bGraphNeedsRecompile = false;
 		{
 			FScopedTransaction ScopedTransaction(LOCTEXT("ResetInputStructTransaction", "Reset the inputs value to default."));
 			
-			if (IsRapidIterationCandidate())
+			if(DefaultPin->LinkedTo.Num() == 0 && IsRapidIterationCandidate())
 			{
 				if (OverrideNode != nullptr && OverridePin != nullptr)
 				{
@@ -915,14 +1114,13 @@ void UNiagaraStackFunctionInput::Reset()
 				}
 
 				// Get the default value of the graph pin and use that to reset the rapid iteration variables...
-				UEdGraphPin* DefaultPin = GetDefaultPin();
-				if (DefaultPin != nullptr)
+				FNiagaraVariable DefaultVar = GetDefaultVariableForRapidIterationParameter();
+				if (DefaultVar.IsValid())
 				{
-					FNiagaraVariable DefaultVar = GetDefaultValue();
 					UpdateRapidIterationParametersForAffectedScripts(DefaultVar.GetData());
 				}
 			}
-			else
+			else if(DefaultPin->LinkedTo.Num() == 0 || FNiagaraStackGraphUtilities::IsValidDefaultDynamicInput(*SourceScript, *DefaultPin) == false)
 			{
 				if (ensureMsgf(OverrideNode != nullptr && OverridePin != nullptr, TEXT("Can not reset the value of an input that doesn't have a valid override node and override pin")))
 				{
@@ -932,6 +1130,16 @@ void UNiagaraStackFunctionInput::Reset()
 					bGraphNeedsRecompile =  true;
 				}
 			}
+			else
+			{
+				if (OverridePin != nullptr)
+				{
+					RemoveNodesForOverridePin(*OverridePin);
+				}
+				FNiagaraStackGraphUtilities::ResetToDefaultDynamicInput(GetSystemViewModel(), GetEmitterViewModel(), GetStackEditorData(),
+					*SourceScript, AffectedScripts, *OwningModuleNode, *OwningFunctionCallNode, InputParameterHandle.GetName(), *DefaultPin);
+				bGraphNeedsRecompile = true;
+			}
 
 			if (bGraphNeedsRecompile)
 			{
@@ -1044,25 +1252,13 @@ void UNiagaraStackFunctionInput::ResetToBase()
 	}
 }
 
-FString UNiagaraStackFunctionInput::CreateRapidIterationVariableName(const FString& InName)
+FNiagaraVariable UNiagaraStackFunctionInput::CreateRapidIterationVariable(const FName& InName)
 {
 	UNiagaraNodeOutput* OutputNode = FNiagaraStackGraphUtilities::GetEmitterOutputNodeForStackNode(*OwningModuleNode.Get());
 	UNiagaraEmitter* ParentEmitter = GetEmitterViewModel()->GetEmitter();
-	UNiagaraSystem* ParentSystem = &GetSystemViewModel()->GetSystem();
-	FNiagaraVariable Temp(InputType, *InName);
-	if (OutputNode->GetUsage() == ENiagaraScriptUsage::SystemSpawnScript || OutputNode->GetUsage() == ENiagaraScriptUsage::SystemUpdateScript)
-	{
-		Temp = FNiagaraParameterMapHistory::ConvertVariableToRapidIterationConstantName(Temp, nullptr); // These names *should* have the emitter baked in...
-	}
-	else
-	{
-		Temp = FNiagaraParameterMapHistory::ConvertVariableToRapidIterationConstantName(Temp, *ParentEmitter->GetUniqueEmitterName());
-	}
-
-	return Temp.GetName().ToString();
+	return FNiagaraStackGraphUtilities::CreateRapidIterationParameter(ParentEmitter->GetUniqueEmitterName(), OutputNode->GetUsage(), InName, InputType);
 }
 
-
 bool UNiagaraStackFunctionInput::CanRenameInput() const
 {
 	// Only module level assignment node inputs can be renamed.
@@ -1071,14 +1267,14 @@ bool UNiagaraStackFunctionInput::CanRenameInput() const
 
 bool UNiagaraStackFunctionInput::GetIsRenamePending() const
 {
-	return CanRenameInput() && StackEditorData->GetModuleInputIsRenamePending(StackEditorDataKey);
+	return CanRenameInput() && GetStackEditorData().GetModuleInputIsRenamePending(StackEditorDataKey);
 }
 
 void UNiagaraStackFunctionInput::SetIsRenamePending(bool bIsRenamePending)
 {
 	if (CanRenameInput())
 	{
-		StackEditorData->SetModuleInputIsRenamePending(StackEditorDataKey, bIsRenamePending);
+		GetStackEditorData().SetModuleInputIsRenamePending(StackEditorDataKey, bIsRenamePending);
 	}
 }
 
@@ -1086,15 +1282,26 @@ void UNiagaraStackFunctionInput::RenameInput(FName NewName)
 {
 	if (OwningAssignmentNode.IsValid() && InputParameterHandlePath.Num() == 1 && InputParameterHandle.GetName() != NewName)
 	{
-		bool bIsCurrentlyPinned = GetIsPinned();
-		bool bIsCurrentlyExpanded = StackEditorData->GetStackEntryIsExpanded(FNiagaraStackGraphUtilities::GenerateStackModuleEditorDataKey(*OwningAssignmentNode), false);
+		FInputValues OldInputValues = InputValues;
+		UEdGraphPin* OriginalOverridePin = GetOverridePin();
 
-		FNiagaraParameterHandle TargetHandle(OwningAssignmentNode->AssignmentTarget.GetName());
-		FNiagaraParameterHandle RenamedTargetHandle(TargetHandle.GetNamespace(), NewName);
+		// We'll be making changes, so go ahead and keep track of the override pointer if it exists.
+		if (OriginalOverridePin != nullptr)
+		{
+			OriginalOverridePin->GetOwningNode()->Modify();
+		}
 
-		OwningAssignmentNode->AssignmentTarget.SetName(RenamedTargetHandle.GetParameterHandleString());
-		OwningAssignmentNode->RefreshFromExternalChanges();
+		bool bIsCurrentlyExpanded = GetStackEditorData().GetStackEntryIsExpanded(FNiagaraStackGraphUtilities::GenerateStackModuleEditorDataKey(*OwningAssignmentNode), false);
 
+
+		int32 FoundIdx = OwningAssignmentNode->FindAssignmentTarget(InputParameterHandle.GetName());
+		check(FoundIdx != INDEX_NONE);
+		FNiagaraParameterHandle TargetHandle(OwningAssignmentNode->GetAssignmentTargetName(FoundIdx));
+
+		if (OwningAssignmentNode->SetAssignmentTargetName(FoundIdx, NewName))
+		{
+			OwningAssignmentNode->RefreshFromExternalChanges();
+		}
 		InputParameterHandle = FNiagaraParameterHandle(InputParameterHandle.GetNamespace(), NewName);
 		InputParameterHandlePath.Empty(1);
 		InputParameterHandlePath.Add(InputParameterHandle);
@@ -1103,32 +1310,54 @@ void UNiagaraStackFunctionInput::RenameInput(FName NewName)
 
 		if (IsRapidIterationCandidate())
 		{
-			FString OriginalName = RapidIterationParameterName;
-			FString NewNameForRIParam = CreateRapidIterationVariableName(AliasedInputParameterHandle.GetParameterHandleString().ToString());
+			FNiagaraVariable OldRapidIterationParameter = RapidIterationParameter;
+			RapidIterationParameter = CreateRapidIterationVariable(AliasedInputParameterHandle.GetParameterHandleString());
 
-			FNiagaraVariable Temp(InputType, *OriginalName);
 			for (TWeakObjectPtr<UNiagaraScript> Script : AffectedScripts)
 			{
-				Script->RapidIterationParameters.RenameParameter(Temp, *NewNameForRIParam);
+				Script->RapidIterationParameters.RenameParameter(OldRapidIterationParameter, *RapidIterationParameter.GetName().ToString());
 			}
-			UE_LOG(LogNiagaraEditor, Log, TEXT("Renaming %s to %s"), *RapidIterationParameterName, *NewNameForRIParam);
-			RapidIterationParameterName = NewNameForRIParam;
+
+			UE_LOG(LogNiagaraEditor, Log, TEXT("Renaming %s to %s"), *OldRapidIterationParameter.GetName().ToString(), *RapidIterationParameter.GetName().ToString());
 		}
 
-		UEdGraphPin* OverridePin = GetOverridePin();
-		if (OverridePin != nullptr)
+		// Go ahead and have the override pin point to the new name instead of the old name..
+		if (OriginalOverridePin != nullptr)
 		{
-			OverridePin->PinName = AliasedInputParameterHandle.GetParameterHandleString();
+			OriginalOverridePin->PinName = AliasedInputParameterHandle.GetParameterHandleString();
 		}
-
+		
 		StackEditorDataKey = FNiagaraStackGraphUtilities::GenerateStackFunctionInputEditorDataKey(*OwningFunctionCallNode.Get(), InputParameterHandle);
-		StackEditorData->SetModuleInputIsPinned(StackEditorDataKey, bIsCurrentlyPinned);
-		StackEditorData->SetStackEntryIsExpanded(FNiagaraStackGraphUtilities::GenerateStackModuleEditorDataKey(*OwningAssignmentNode), bIsCurrentlyExpanded);
+		GetStackEditorData().SetStackEntryIsExpanded(FNiagaraStackGraphUtilities::GenerateStackModuleEditorDataKey(*OwningAssignmentNode), bIsCurrentlyExpanded);
 
 		CastChecked<UNiagaraGraph>(OwningAssignmentNode->GetGraph())->NotifyGraphNeedsRecompile();
 	}
 }
 
+bool UNiagaraStackFunctionInput::CanDeleteInput() const
+{
+	return GetInputFunctionCallNode().IsA(UNiagaraNodeAssignment::StaticClass());
+}
+
+void UNiagaraStackFunctionInput::DeleteInput()
+{
+	if (UNiagaraNodeAssignment* NodeAssignment = Cast<UNiagaraNodeAssignment>(OwningFunctionCallNode.Get()))
+	{
+		FNiagaraVariable Var = FNiagaraVariable(GetInputType(), GetInputParameterHandle().GetName());
+		NodeAssignment->RemoveParameter(Var);
+
+		FNiagaraParameterHandle ParameterHandle(Var.GetName());
+		UEdGraphPin* OverridePin = FNiagaraStackGraphUtilities::GetStackFunctionInputOverridePin(*OwningFunctionCallNode.Get(), ParameterHandle);
+		if (OverridePin && OverridePin->GetOwningNode())
+		{
+			UEdGraphNode* Node = OverridePin->GetOwningNode();
+			FNiagaraStackGraphUtilities::RemoveNodesForStackFunctionInputOverridePin(*OverridePin);
+			OverridePin->GetOwningNode()->RemovePin(OverridePin);
+			OverridePin = nullptr;
+		}
+	}
+}
+
 void UNiagaraStackFunctionInput::GetNamespacesForNewParameters(TArray<FName>& OutNamespacesForNewParameters) const
 {
 	UNiagaraNodeOutput* OutputNode = FNiagaraStackGraphUtilities::GetEmitterOutputNodeForStackNode(*OwningFunctionCallNode);
@@ -1153,7 +1382,7 @@ void UNiagaraStackFunctionInput::GetNamespacesForNewParameters(TArray<FName>& Ou
 			OutNamespacesForNewParameters.Add(FNiagaraParameterHandle::UserNamespace);
 		}
 	}
-	else if (OutputNode->GetUsage() == ENiagaraScriptUsage::SystemSpawnScript || OutputNode->GetUsage() == ENiagaraScriptUsage::ParticleUpdateScript && bIsEditingSystem)
+	else if ((OutputNode->GetUsage() == ENiagaraScriptUsage::SystemSpawnScript || OutputNode->GetUsage() == ENiagaraScriptUsage::SystemUpdateScript) && bIsEditingSystem)
 	{
 		OutNamespacesForNewParameters.Add(FNiagaraParameterHandle::SystemNamespace);
 		OutNamespacesForNewParameters.Add(FNiagaraParameterHandle::UserNamespace);
@@ -1165,9 +1394,42 @@ UNiagaraStackFunctionInput::FOnValueChanged& UNiagaraStackFunctionInput::OnValue
 	return ValueChangedDelegate;
 }
 
-UNiagaraStackFunctionInput::FOnPinnedChanged& UNiagaraStackFunctionInput::OnPinnedChanged()
+bool UNiagaraStackFunctionInput::GetHasEditCondition() const
+{
+	return EditCondition.IsValid();
+}
+
+bool UNiagaraStackFunctionInput::GetShowEditConditionInline() const
+{
+	return bShowEditConditionInline;
+}
+
+bool UNiagaraStackFunctionInput::GetEditConditionEnabled() const
+{
+	return EditCondition.IsValid() && EditCondition.GetConditionIsEnabled();
+}
+
+void UNiagaraStackFunctionInput::SetEditConditionEnabled(bool bIsEnabled)
+{
+	if(EditCondition.CanSetConditionIsEnabled())
+	{ 
+		EditCondition.SetConditionIsEnabled(bIsEnabled);
+	}
+}
+
+bool UNiagaraStackFunctionInput::GetHasVisibleCondition() const
 {
-	return PinnedChangedDelegate;
+	return VisibleCondition.IsValid();
+}
+
+bool UNiagaraStackFunctionInput::GetVisibleConditionEnabled() const
+{
+	return VisibleCondition.IsValid() && VisibleCondition.GetConditionIsEnabled();
+}
+
+bool UNiagaraStackFunctionInput::GetIsInlineEditConditionToggle() const
+{
+	return bIsInlineEditConditionToggle;
 }
 
 void UNiagaraStackFunctionInput::OnGraphChanged(const struct FEdGraphEditAction& InAction)
@@ -1181,11 +1443,20 @@ void UNiagaraStackFunctionInput::OnGraphChanged(const struct FEdGraphEditAction&
 
 void UNiagaraStackFunctionInput::OnRapidIterationParametersChanged()
 {
+	bCanReset.Reset();
 	bCanResetToBase.Reset();
+	if (ensureMsgf(OwningModuleNode.IsValid() && OwningFunctionCallNode.IsValid(), TEXT("Stack entry with invalid module or function call not cleaned up.")))
+	{
+		if (bUpdatingLocalValueDirectly == false && IsRapidIterationCandidate() && (OverridePinCache.IsSet() == false || OverridePinCache.GetValue() == nullptr))
+		{
+			RefreshValues();
+		}
+	}
 }
 
 void UNiagaraStackFunctionInput::OnScriptSourceChanged()
 {
+	bCanReset.Reset();
 	bCanResetToBase.Reset();
 }
 
@@ -1217,7 +1488,7 @@ UNiagaraNodeParameterMapSet& UNiagaraStackFunctionInput::GetOrCreateOverrideNode
 
 UEdGraphPin* UNiagaraStackFunctionInput::GetDefaultPin() const
 {
-	return OwningFunctionCallNode->FindParameterMapDefaultValuePin(InputParameterHandle.GetParameterHandleString());
+	return OwningFunctionCallNode->FindParameterMapDefaultValuePin(InputParameterHandle.GetParameterHandleString(), SourceScript->GetUsage());
 }
 
 UEdGraphPin* UNiagaraStackFunctionInput::GetOverridePin() const
@@ -1241,7 +1512,7 @@ UEdGraphPin& UNiagaraStackFunctionInput::GetOrCreateOverridePin()
 	return *OverridePin;
 }
 
-bool UNiagaraStackFunctionInput::TryGetCurrentLocalValue(TSharedPtr<FStructOnScope>& LocalValue, UEdGraphPin& ValuePin, TSharedPtr<FStructOnScope> OldValueToReuse)
+bool UNiagaraStackFunctionInput::TryGetCurrentLocalValue(TSharedPtr<FStructOnScope>& LocalValue, UEdGraphPin& DefaultPin, UEdGraphPin& ValuePin, TSharedPtr<FStructOnScope> OldValueToReuse)
 {
 	if (InputType.IsDataInterface() == false && ValuePin.LinkedTo.Num() == 0)
 	{
@@ -1256,31 +1527,21 @@ bool UNiagaraStackFunctionInput::TryGetCurrentLocalValue(TSharedPtr<FStructOnSco
 			LocalValue = MakeShared<FStructOnScope>(ValueVariable.GetType().GetStruct());
 		}
 
-		if (IsRapidIterationCandidate())
+		// If the default pin in the function graph is connected internally, rapid iteration parameters can't be used since
+		// the compilation currently won't use them.
+		bool bCanUseRapidIterationParameter = IsRapidIterationCandidate() && DefaultPin.LinkedTo.Num() == 0;
+		bool bFoundRapidIterationParameter = false;
+		if (bCanUseRapidIterationParameter)
 		{
-			ValueVariable.SetName(*RapidIterationParameterName);
-
-			int32 FoundOffset = SourceScript->RapidIterationParameters.IndexOf(ValueVariable);
-			if (FoundOffset == INDEX_NONE)
-			{
-				// If the source script didn't have a value for this variable, update all affected scripts
-				// with the default value from the function.
-				UpdateRapidIterationParametersForAffectedScripts(ValueVariable.GetData());
-			}
-
-			FoundOffset = SourceScript->RapidIterationParameters.IndexOf(ValueVariable);
-			if (FoundOffset != INDEX_NONE)
-			{
-				FMemory::Memcpy(LocalValue->GetStructMemory(), SourceScript->RapidIterationParameters.GetParameterData(FoundOffset), ValueVariable.GetSizeInBytes());
-			}
-			else
+			const uint8* RapidIterationParameterData = SourceScript->RapidIterationParameters.GetParameterData(RapidIterationParameter);
+			if(RapidIterationParameterData != nullptr)
 			{
-				// How would we ever get here in practice, we just added above?
-				ensure(false);
-				return false;
+				FMemory::Memcpy(LocalValue->GetStructMemory(), RapidIterationParameterData, ValueVariable.GetSizeInBytes());
+				bFoundRapidIterationParameter = true;
 			}
 		}
-		else
+
+		if(bFoundRapidIterationParameter == false)
 		{
 			ValueVariable.CopyTo(LocalValue->GetStructMemory());
 		}
@@ -1325,7 +1586,7 @@ bool UNiagaraStackFunctionInput::TryGetCurrentDataValue(FDataValues& DataValues,
 		{
 			if (LocallyOwnedDefaultDataValueObjectToReuse == nullptr)
 			{
-				DefaultDataValueObject = NewObject<UNiagaraDataInterface>(this, const_cast<UClass*>(InputType.GetClass()));
+				DefaultDataValueObject = NewObject<UNiagaraDataInterface>(this, const_cast<UClass*>(InputType.GetClass()), NAME_None, RF_Transactional | RF_Public);
 			}
 			else
 			{
@@ -1395,6 +1656,20 @@ bool UNiagaraStackFunctionInput::TryGetCurrentLinkedValue(FNiagaraParameterHandl
 	return false;
 }
 
+bool UNiagaraStackFunctionInput::TryGetCurrentExpressionValue(TWeakObjectPtr<UNiagaraNodeCustomHlsl>& ExpressionValue, UEdGraphPin* OverridePin)
+{
+	if (OverridePin != nullptr && OverridePin->LinkedTo.Num() == 1)
+	{
+		UNiagaraNodeCustomHlsl* DynamicNode = Cast<UNiagaraNodeCustomHlsl>(OverridePin->LinkedTo[0]->GetOwningNode());
+		if (DynamicNode != nullptr)
+		{
+			ExpressionValue = DynamicNode;
+			return true;
+		}
+	}
+	return false;
+}
+
 bool UNiagaraStackFunctionInput::TryGetCurrentDynamicValue(TWeakObjectPtr<UNiagaraNodeFunctionCall>& DynamicValue, UEdGraphPin* OverridePin)
 {
 	if (OverridePin != nullptr && OverridePin->LinkedTo.Num() == 1)
