Copying //UE4/Dev-Niagara to //UE4/Dev-Main (Source: //UE4/Dev-Niagara @ 4074996)

diff --git a/Engine/Plugins/FX/Niagara/Source/NiagaraEditor/Private/ViewModels/Stack/NiagaraStackFunctionInputCollection.cpp b/Engine/Plugins/FX/Niagara/Source/NiagaraEditor/Private/ViewModels/Stack/NiagaraStackFunctionInputCollection.cpp
--- a/Engine/Plugins/FX/Niagara/Source/NiagaraEditor/Private/ViewModels/Stack/NiagaraStackFunctionInputCollection.cpp
+++ b/Engine/Plugins/FX/Niagara/Source/NiagaraEditor/Private/ViewModels/Stack/NiagaraStackFunctionInputCollection.cpp
@@ -2,12 +2,16 @@
 
 #include "ViewModels/Stack/NiagaraStackFunctionInputCollection.h"
 #include "ViewModels/Stack/NiagaraStackFunctionInput.h"
+#include "ViewModels/Stack/NiagaraStackInputCategory.h"
 #include "NiagaraNodeFunctionCall.h"
 #include "NiagaraEmitterEditorData.h"
 #include "EdGraphSchema_Niagara.h"
 #include "ViewModels/Stack/NiagaraStackGraphUtilities.h"
 #include "NiagaraNodeParameterMapSet.h"
 #include "NiagaraDataInterface.h"
+#include "NiagaraConstants.h"
+#include "NiagaraGraph.h"
+#include "NiagaraStackEditorData.h"
 
 #include "EdGraph/EdGraphPin.h"
 #include "ScopedTransaction.h"
@@ -17,6 +21,7 @@
 UNiagaraStackFunctionInputCollection::UNiagaraStackFunctionInputCollection()
 	: ModuleNode(nullptr)
 	, InputFunctionCallNode(nullptr)
+	, bShouldShowInStack(true)
 {
 }
 
@@ -31,115 +36,87 @@ UNiagaraNodeFunctionCall* UNiagaraStackFunctionInputCollection::GetInputFunction
 }
 
 void UNiagaraStackFunctionInputCollection::Initialize(
-	TSharedRef<FNiagaraSystemViewModel> InSystemViewModel,
-	TSharedRef<FNiagaraEmitterViewModel> InEmitterViewModel,
-	UNiagaraStackEditorData& InStackEditorData,
+	FRequiredEntryData InRequiredEntryData,
 	UNiagaraNodeFunctionCall& InModuleNode,
 	UNiagaraNodeFunctionCall& InInputFunctionCallNode,
-	FDisplayOptions InDisplayOptions)
+	FString InOwnerStackItemEditorDataKey)
 {
 	checkf(ModuleNode == nullptr && InputFunctionCallNode == nullptr, TEXT("Can not set the node more than once."));
-	Super::Initialize(InSystemViewModel, InEmitterViewModel);
-	StackEditorData = &InStackEditorData;
+	FString InputCollectionStackEditorDataKey = FString::Printf(TEXT("%s-Inputs"), *InInputFunctionCallNode.NodeGuid.ToString(EGuidFormats::DigitsWithHyphens));
+	Super::Initialize(InRequiredEntryData, false, InOwnerStackItemEditorDataKey, InputCollectionStackEditorDataKey);
 	ModuleNode = &InModuleNode;
 	InputFunctionCallNode = &InInputFunctionCallNode;
-	DisplayOptions = InDisplayOptions;
-	if (DisplayOptions.ChildFilter.IsBound())
-	{
-		AddChildFilter(DisplayOptions.ChildFilter);
-	}
+	InputFunctionCallNode->OnInputsChanged().AddUObject(this, &UNiagaraStackFunctionInputCollection::OnFunctionInputsChanged);
 }
 
 FText UNiagaraStackFunctionInputCollection::GetDisplayName() const
 {
-	return DisplayOptions.DisplayName;
-}
-
-FName UNiagaraStackFunctionInputCollection::GetTextStyleName() const
-{
-	return "NiagaraEditor.Stack.ParameterCollectionText";
-}
-
-bool UNiagaraStackFunctionInputCollection::GetCanExpand() const
-{
-	return true;
+	return LOCTEXT("InputCollectionDisplayName", "Inputs");
 }
 
 bool UNiagaraStackFunctionInputCollection::GetShouldShowInStack() const
 {
-	return DisplayOptions.bShouldShowInStack;
+	return bShouldShowInStack;
 }
 
-int32 UNiagaraStackFunctionInputCollection::GetErrorCount() const
+bool UNiagaraStackFunctionInputCollection::GetIsEnabled() const
 {
-	return Errors.Num();
+	return InputFunctionCallNode->GetDesiredEnabledState() == ENodeEnabledState::Enabled;
 }
 
-bool UNiagaraStackFunctionInputCollection::GetErrorFixable(int32 ErrorIdx) const
+void UNiagaraStackFunctionInputCollection::SetShouldShowInStack(bool bInShouldShowInStack)
 {
-	return Errors[ErrorIdx].Fix.IsBound();
+	bShouldShowInStack = bInShouldShowInStack;
 }
 
-bool UNiagaraStackFunctionInputCollection::TryFixError(int32 ErrorIdx)
-{
-	Errors[ErrorIdx].Fix.Execute();
-	return true;
-}
-
-FText UNiagaraStackFunctionInputCollection::GetErrorText(int32 ErrorIdx) const
-{
-	return Errors[ErrorIdx].ErrorText;
-}
-
-FText UNiagaraStackFunctionInputCollection::GetErrorSummaryText(int32 ErrorIdx) const
-{
-	return Errors[ErrorIdx].ErrorSummaryText;
-}
-
-UNiagaraStackFunctionInputCollection::FOnInputPinnedChanged& UNiagaraStackFunctionInputCollection::OnInputPinnedChanged()
-{
-	return InputPinnedChangedDelegate;
-}
-
-void UNiagaraStackFunctionInputCollection::RefreshChildrenInternal(const TArray<UNiagaraStackEntry*>& CurrentChildren, TArray<UNiagaraStackEntry*>& NewChildren)
+void UNiagaraStackFunctionInputCollection::RefreshChildrenInternal(const TArray<UNiagaraStackEntry*>& CurrentChildren, TArray<UNiagaraStackEntry*>& NewChildren, TArray<FStackIssue>& NewIssues)
 {
 	TArray<const UEdGraphPin*> InputPins;
 	FNiagaraStackGraphUtilities::GetStackFunctionInputPins(*InputFunctionCallNode, InputPins, FNiagaraStackGraphUtilities::ENiagaraGetStackFunctionInputPinsOptions::ModuleInputsOnly);
+
 	const UEdGraphSchema_Niagara* NiagaraSchema = GetDefault<UEdGraphSchema_Niagara>();
-	TArray<FName> ValidAliasedInputNames;
+
 	TArray<FName> ProcessedInputNames;
-	TArray<FName> DuplicateInputNames;
-	TArray<const UEdGraphPin*> PinsWithInvalidTypes;
-	for(const UEdGraphPin* InputPin : InputPins)
+
+	FText UncategorizedName = LOCTEXT("Uncategorized", "Uncategorized");
+
+	UNiagaraGraph* InputFunctionGraph = InputFunctionCallNode->FunctionScript != nullptr
+		? CastChecked<UNiagaraScriptSource>(InputFunctionCallNode->FunctionScript->GetSource())->NodeGraph
+		: nullptr;
+
+	TArray<FInputData> InputDataCollection;
+	
+	// Gather input data
+	for (const UEdGraphPin* InputPin : InputPins)
 	{
 		if (ProcessedInputNames.Contains(InputPin->PinName) == false)
 		{
-			UNiagaraStackFunctionInput* Input = FindCurrentChildOfTypeByPredicate<UNiagaraStackFunctionInput>(CurrentChildren,
-				[=](UNiagaraStackFunctionInput* CurrentFunctionInput) { return CurrentFunctionInput->GetInputParameterHandle().GetParameterHandleString() == InputPin->PinName; });
-			
-			if (Input == nullptr)
+			ProcessedInputNames.Add(InputPin->PinName);
+
+			FNiagaraVariable InputVariable = NiagaraSchema->PinToNiagaraVariable(InputPin);
+			if (InputVariable.GetType().IsValid() == false)
 			{
-				FNiagaraTypeDefinition InputType = NiagaraSchema->PinToTypeDefinition(InputPin);
-				if (InputType.IsValid())
-				{
-					Input = NewObject<UNiagaraStackFunctionInput>(this);
-					Input->Initialize(GetSystemViewModel(), GetEmitterViewModel(), *StackEditorData, *ModuleNode, *InputFunctionCallNode, InputPin->PinName, NiagaraSchema->PinToTypeDefinition(InputPin));
-					Input->SetItemIndentLevel(DisplayOptions.ChildItemIndentLevel);
-					Input->OnPinnedChanged().AddUObject(this, &UNiagaraStackFunctionInputCollection::ChildPinnedChanged);
-				}
-				else
-				{
-					PinsWithInvalidTypes.Add(InputPin);
-				}
+				PinsWithInvalidTypes.Add(InputPin);
+				continue;
+			}
+			else
+			{
+				ValidAliasedInputNames.Add(FNiagaraParameterHandle::CreateAliasedModuleParameterHandle(
+					FNiagaraParameterHandle(InputPin->PinName), InputFunctionCallNode).GetParameterHandleString());
 			}
 
-			if (Input != nullptr)
+			FNiagaraVariableMetaData* InputMetaData = nullptr;
+			if (InputFunctionGraph != nullptr)
 			{
-				NewChildren.Add(Input);
+				InputMetaData = InputFunctionGraph->GetMetaData(InputVariable);
 			}
-			
-			ValidAliasedInputNames.Add(FNiagaraParameterHandle::CreateAliasedModuleParameterHandle(FNiagaraParameterHandle(InputPin->PinName), InputFunctionCallNode).GetParameterHandleString());
-			ProcessedInputNames.Add(InputPin->PinName);
+
+			FText InputCategory = InputMetaData != nullptr && InputMetaData->CategoryName.IsEmptyOrWhitespace() == false
+				? InputMetaData->CategoryName
+				: UncategorizedName;
+
+			FInputData InputData = { InputPin, InputVariable.GetType(), InputMetaData ? InputMetaData->EditorSortPriority : 0, InputCategory };
+			InputDataCollection.Add(InputData);
 		}
 		else
 		{
@@ -147,8 +124,60 @@ void UNiagaraStackFunctionInputCollection::RefreshChildrenInternal(const TArray<
 		}
 	}
 
+	// Sort data and keep the uncategorized first
+	InputDataCollection.Sort([UncategorizedName](const FInputData& A, const FInputData& B)
+	{
+		if (A.Category.CompareTo(UncategorizedName) == 0 && B.Category.CompareTo(UncategorizedName) != 0)
+		{
+			return true;
+		}
+		if (A.Category.CompareTo(UncategorizedName) != 0 && B.Category.CompareTo(UncategorizedName) == 0)
+		{
+			return false;
+		}
+		if (A.SortKey != B.SortKey)
+		{
+			return A.SortKey < B.SortKey;
+		}
+		else
+		{
+			return A.Pin->PinName < B.Pin->PinName;
+		}
+	});
+
+	// Populate the category children
+	for (const FInputData InputData : InputDataCollection)
+	{
+		UNiagaraStackInputCategory* InputCategory = FindCurrentChildOfTypeByPredicate<UNiagaraStackInputCategory>(NewChildren,
+			[&](UNiagaraStackInputCategory* CurrentCategory) { return CurrentCategory->GetCategoryName().CompareTo(InputData.Category) == 0; });
+
+		if (InputCategory == nullptr)
+		{
+			InputCategory = FindCurrentChildOfTypeByPredicate<UNiagaraStackInputCategory>(CurrentChildren,
+				[&](UNiagaraStackInputCategory* CurrentCategory) { return CurrentCategory->GetCategoryName().CompareTo(InputData.Category) == 0; });
+			if (InputCategory == nullptr)
+			{
+				InputCategory = NewObject<UNiagaraStackInputCategory>(this);
+				InputCategory->Initialize(CreateDefaultChildRequiredData(), *ModuleNode, *InputFunctionCallNode, InputData.Category, GetOwnerStackItemEditorDataKey());
+			}
+			else
+			{
+				InputCategory->ResetInputs();
+			}
+			if (InputData.Category.CompareTo(UncategorizedName) == 0)
+			{
+				InputCategory->SetShouldShowInStack(false);
+			}
+			NewChildren.Add(InputCategory);
+		}
+		InputCategory->AddInput(InputData.Pin->PinName, InputData.Type);
+	}
+	RefreshIssues(NewIssues);
+}
+
+void UNiagaraStackFunctionInputCollection::RefreshIssues(TArray<FStackIssue>& NewIssues)
+{
 	// Try to find function input overrides which are no longer valid so we can generate errors for them.
-	Errors.Empty();
 	UNiagaraNodeParameterMapSet* OverrideNode = FNiagaraStackGraphUtilities::GetStackFunctionOverrideNode(*InputFunctionCallNode);
 	if (OverrideNode != nullptr)
 	{
@@ -163,13 +192,16 @@ void UNiagaraStackFunctionInputCollection::RefreshChildrenInternal(const TArray<
 				FNiagaraParameterHandle(OverridePin->PinName).GetNamespace().ToString() == InputFunctionCallNode->GetFunctionName() &&
 				ValidAliasedInputNames.Contains(OverridePin->PinName) == false)
 			{
-				FError InvalidInputError;
-				InvalidInputError.ErrorSummaryText = FText::Format(LOCTEXT("InvalidInputSummaryFormat", "Invalid Input: {0}"), FText::FromString(OverridePin->PinName.ToString()));
-				InvalidInputError.ErrorText = FText::Format(LOCTEXT("InvalidInputFormat", "The input {0} was previously overriden but is no longer exposed by the function {1}.\nPress the fix button to remove this unused override data,\nor check the function definition to see why this input is no longer exposed."),
+				UNiagaraStackEntry::FStackIssue InvalidInputError;
+				InvalidInputError.ShortDescription = FText::Format(LOCTEXT("InvalidInputSummaryFormat", "Invalid Input: {0}"), FText::FromString(OverridePin->PinName.ToString()));
+				InvalidInputError.LongDescription = FText::Format(LOCTEXT("InvalidInputFormat", "The input {0} was previously overriden but is no longer exposed by the function {1}.\nPress the fix button to remove this unused override data,\nor check the function definition to see why this input is no longer exposed."),
 					FText::FromString(OverridePin->PinName.ToString()), FText::FromString(InputFunctionCallNode->GetFunctionName()));
-				InvalidInputError.Fix.BindLambda([=]()
+				InvalidInputError.UniqueIdentifier = FName(*FString::Printf(TEXT("%s-InvalidInput-%s"), *GetStackEditorDataKey(), *OverridePin->PinName.ToString()));
+				UNiagaraStackEntry::FStackIssueFix Fix;
+				Fix.Description = LOCTEXT("RemoveInvalidInputTransaction", "Remove invalid input override.");
+				Fix.FixDelegate.BindLambda([=]()
 				{
-					FScopedTransaction ScopedTransaction(LOCTEXT("RemoveInvalidInputTransaction", "Remove invalid input override."));
+					FScopedTransaction ScopedTransaction(Fix.Description);
 					TArray<TWeakObjectPtr<UNiagaraDataInterface>> RemovedDataObjects;
 					FNiagaraStackGraphUtilities::RemoveNodesForStackFunctionInputOverridePin(*OverridePin, RemovedDataObjects);
 					for (TWeakObjectPtr<UNiagaraDataInterface> RemovedDataObject : RemovedDataObjects)
@@ -181,36 +213,41 @@ void UNiagaraStackFunctionInputCollection::RefreshChildrenInternal(const TArray<
 					}
 					OverridePin->GetOwningNode()->RemovePin(OverridePin);
 				});
-				Errors.Add(InvalidInputError);
+				InvalidInputError.Fixes.Add(Fix);
+				NewIssues.Add(InvalidInputError);
 			}
 		}
 	}
 
 	for (const FName& DuplicateInputName : DuplicateInputNames)
 	{
-		FError DuplicateInputError;
+		UNiagaraStackEntry::FStackIssue DuplicateInputError;
 
 		FString DuplicateInputNameString = DuplicateInputName.ToString();
 
-		DuplicateInputError.ErrorSummaryText = FText::Format(LOCTEXT("DuplicateInputSummaryFormat", "Duplicate Input: {0}"), FText::FromString(DuplicateInputNameString));
-		DuplicateInputError.ErrorText = FText::Format(LOCTEXT("DuplicateInputFormat", "There are multiple inputs with the same name {0}, but different types exposed by the function {1}.\nThis is not suppored and must be fixed in the script that defines this function."),
+		DuplicateInputError.ShortDescription = FText::Format(LOCTEXT("DuplicateInputSummaryFormat", "Duplicate Input: {0}"), FText::FromString(DuplicateInputNameString));
+		DuplicateInputError.LongDescription = FText::Format(LOCTEXT("DuplicateInputFormat", "There are multiple inputs with the same name {0}, but different types exposed by the function {1}.\nThis is not suppored and must be fixed in the script that defines this function."),
 			FText::FromString(DuplicateInputNameString), FText::FromString(InputFunctionCallNode->GetFunctionName()));
-		Errors.Add(DuplicateInputError);
+		DuplicateInputError.UniqueIdentifier = FName(*FString::Printf(TEXT("%s-DuplicateInput-%s"), *GetStackEditorDataKey(), *DuplicateInputNameString));
+		NewIssues.Add(DuplicateInputError);
 	}
 
 	for (const UEdGraphPin* PinWithInvalidType : PinsWithInvalidTypes)
 	{
-		FError InputWithInvalidTypeError;
-		InputWithInvalidTypeError.ErrorSummaryText = FText::Format(LOCTEXT("InputWithInvalidTypeSummaryFormat", "Input has an invalid type: {0}"), FText::FromString(PinWithInvalidType->PinName.ToString()));
-		InputWithInvalidTypeError.ErrorText = FText::Format(LOCTEXT("InputWithInvalidTypeFormat", "The input {0} on function {1} has a type which is invalid.\nThe type of this input likely doesn't exist anymore.\nThis input must be fixed in the script before this module can be used."),
+		UNiagaraStackEntry::FStackIssue InputWithInvalidTypeError;
+		InputWithInvalidTypeError.ShortDescription = FText::Format(LOCTEXT("InputWithInvalidTypeSummaryFormat", "Input has an invalid type: {0}"), FText::FromString(PinWithInvalidType->PinName.ToString()));
+		InputWithInvalidTypeError.LongDescription = FText::Format(LOCTEXT("InputWithInvalidTypeFormat", "The input {0} on function {1} has a type which is invalid.\nThe type of this input likely doesn't exist anymore.\nThis input must be fixed in the script before this module can be used."),
 			FText::FromString(PinWithInvalidType->PinName.ToString()), FText::FromString(InputFunctionCallNode->GetFunctionName()));
-		Errors.Add(InputWithInvalidTypeError);
+		NewIssues.Add(InputWithInvalidTypeError);
 	}
 }
 
-void UNiagaraStackFunctionInputCollection::ChildPinnedChanged()
+void UNiagaraStackFunctionInputCollection::OnFunctionInputsChanged()
 {
-	InputPinnedChangedDelegate.Broadcast();
+	if (IsValid())
+	{
+		RefreshChildren();
+	}
 }
 
 #undef LOCTEXT_NAMESPACE
diff --git a/Engine/Plugins/FX/Niagara/Source/NiagaraEditor/Private/ViewModels/Stack/NiagaraStackFunctionInputCondition.cpp b/Engine/Plugins/FX/Niagara/Source/NiagaraEditor/Private/ViewModels/Stack/NiagaraStackFunctionInputCondition.cpp
new file mode 100644
--- /dev/null
+++ b/Engine/Plugins/FX/Niagara/Source/NiagaraEditor/Private/ViewModels/Stack/NiagaraStackFunctionInputCondition.cpp
@@ -0,0 +1,130 @@
+// Copyright 1998-2018 Epic Games, Inc. All Rights Reserved.
+
+#include "ViewModels/Stack/NiagaraStackFunctionInputCondition.h"
+#include "NiagaraEditorModule.h"
+#include "NiagaraScript.h"
+#include "NiagaraScriptSource.h"
+#include "NiagaraNodeFunctionCall.h"
+#include "NiagaraTypes.h"
+#include "INiagaraEditorTypeUtilities.h"
+
+#include "Modules/ModuleManager.h"
+
+#define LOCTEXT_NAMESPACE "NiagaraStackFunctionInputCondition"
+
+void FNiagaraStackFunctionInputCondition::Initialize(UNiagaraScript* InScript,
+	TArray<UNiagaraScript*> InDependentScripts,
+	FString InOwningEmitterUniqueName,
+	UNiagaraNodeFunctionCall* InFunctionCallNode)
+{
+	Script = InScript;
+	DependentScripts = InDependentScripts;
+	OwningEmitterUniqueName = InOwningEmitterUniqueName;
+	FunctionCallNode = InFunctionCallNode;
+}
+
+void FNiagaraStackFunctionInputCondition::Refresh(const FString* Condition)
+{
+	TargetValueData.Empty();
+	InputBinder.Reset();
+
+	if (Condition == nullptr || Condition->IsEmpty())
+	{
+		return;
+	}
+
+	FString InputName;
+	FString TargetValue;
+	int32 EqualsIndex = Condition->Find("=");
+	if (EqualsIndex == INDEX_NONE)
+	{
+		// If the condition doesn't have an equals sign the target value is assumed to be a bool value of true.
+		InputName = *Condition;
+		TargetValue = "true";
+	}
+	else
+	{
+		InputName = Condition->Left(EqualsIndex);
+		TargetValue = Condition->RightChop(EqualsIndex + 1);
+	}
+
+	FText ErrorMessage;
+	if (InputBinder.TryBind(Script, DependentScripts, OwningEmitterUniqueName, FunctionCallNode, *InputName, TOptional<FNiagaraTypeDefinition>(), true, ErrorMessage))
+	{
+		FNiagaraEditorModule& NiagaraEditorModule = FModuleManager::GetModuleChecked<FNiagaraEditorModule>("NiagaraEditor");
+		TSharedPtr<INiagaraEditorTypeUtilities, ESPMode::ThreadSafe> TypeEditorUtilities = NiagaraEditorModule.GetTypeUtilities(InputBinder.GetInputType());
+		if (TypeEditorUtilities.IsValid())
+		{
+			FNiagaraVariable TempVariable(InputBinder.GetInputType(), "Temp");
+			bool bValueParsed = false;
+
+			if (TypeEditorUtilities->CanHandlePinDefaults())
+			{
+				bValueParsed = TypeEditorUtilities->SetValueFromPinDefaultString(TargetValue, TempVariable);
+			}
+			if (bValueParsed == false && TypeEditorUtilities->CanSetValueFromDisplayName())
+			{
+				bValueParsed = TypeEditorUtilities->SetValueFromDisplayName(FText::FromString(TargetValue), TempVariable);
+			}
+
+			if (bValueParsed)
+			{
+				TargetValueData.AddUninitialized(InputBinder.GetInputType().GetSize());
+				TempVariable.CopyTo(TargetValueData.GetData());
+			}
+			else
+			{
+				ErrorMessage = FText::Format(LOCTEXT("ParseValueError", "Target value {0} is not a valid for type {1}"), FText::FromString(TargetValue), InputBinder.GetInputType().GetNameText());
+				InputBinder.Reset();
+			}
+		}
+	}
+
+	if(ErrorMessage.IsEmpty() == false)
+	{
+		UE_LOG(LogNiagaraEditor, Warning, TEXT("Input condition failed to bind %s=%s.  Message: %s"), *InputName, *TargetValue, *ErrorMessage.ToString());
+	}
+}
+
+bool FNiagaraStackFunctionInputCondition::IsValid() const
+{
+	return InputBinder.IsValid() && (TargetValueData.Num() > 0);
+}
+
+bool FNiagaraStackFunctionInputCondition::GetConditionIsEnabled() const
+{
+	return IsValid() && FMemory::Memcmp(TargetValueData.GetData(), InputBinder.GetData().GetData(), TargetValueData.Num()) == 0;
+}
+
+bool FNiagaraStackFunctionInputCondition::CanSetConditionIsEnabled() const
+{
+	return IsValid() && InputBinder.GetInputType() == FNiagaraTypeDefinition::GetBoolDef();
+}
+
+void FNiagaraStackFunctionInputCondition::SetConditionIsEnabled(bool bInIsEnabled)
+{
+	checkf(CanSetConditionIsEnabled(), TEXT("Can not set this condition"));
+	InputBinder.SetValue(bInIsEnabled);
+}
+
+FName FNiagaraStackFunctionInputCondition::GetConditionInputName() const
+{
+	checkf(IsValid(), TEXT("Can not get the input name for an invalid input condition"));
+	return InputBinder.GetInputName();
+}
+
+FNiagaraTypeDefinition FNiagaraStackFunctionInputCondition::GetConditionInputType() const
+{
+	checkf(IsValid(), TEXT("Can not get the input type for an invalid input condition"));
+	return InputBinder.GetInputType();
+}
+
+FNiagaraVariableMetaData* FNiagaraStackFunctionInputCondition::GetConditionInputMetaData() const
+{
+	checkf(IsValid(), TEXT("Can not get the input metadata for an invalid input condition"));
+	FNiagaraVariable InputVariable(InputBinder.GetInputType(), InputBinder.GetInputName());
+	UNiagaraGraph* FunctionGraph = CastChecked<UNiagaraScriptSource>(FunctionCallNode->FunctionScript->GetSource())->NodeGraph;
+	return FunctionGraph->GetMetaData(InputVariable);
+}
+
+#undef LOCTEXT_NAMESPACE
diff --git a/Engine/Plugins/FX/Niagara/Source/NiagaraEditor/Private/ViewModels/Stack/NiagaraStackGraphUtilities.cpp b/Engine/Plugins/FX/Niagara/Source/NiagaraEditor/Private/ViewModels/Stack/NiagaraStackGraphUtilities.cpp
--- a/Engine/Plugins/FX/Niagara/Source/NiagaraEditor/Private/ViewModels/Stack/NiagaraStackGraphUtilities.cpp
+++ b/Engine/Plugins/FX/Niagara/Source/NiagaraEditor/Private/ViewModels/Stack/NiagaraStackGraphUtilities.cpp
@@ -2,7 +2,7 @@
 
 #include "ViewModels/Stack/NiagaraStackGraphUtilities.h"
 #include "NiagaraParameterMapHistory.h"
-#include "NiagaraSystemViewModel.h"
+#include "ViewModels/NiagaraSystemViewModel.h"
 #include "NiagaraSystemScriptViewModel.h"
 #include "NiagaraGraph.h"
 #include "NiagaraNode.h"
@@ -10,21 +10,25 @@
 #include "NiagaraNodeInput.h"
 #include "NiagaraNodeFunctionCall.h"
 #include "NiagaraNodeAssignment.h"
-#include "NiagaraMaterialParameterNode.h"
+#include "NiagaraNodeCustomHlsl.h"
 #include "NiagaraNodeParameterMapGet.h"
 #include "NiagaraNodeParameterMapSet.h"
 #include "NiagaraScriptSource.h"
 #include "EdGraphSchema_Niagara.h"
 #include "ViewModels/Stack/NiagaraStackEntry.h"
 #include "ViewModels/Stack/NiagaraStackFunctionInputCollection.h"
+#include "ViewModels/Stack/NiagaraStackInputCategory.h"
 #include "ViewModels/Stack/NiagaraStackFunctionInput.h"
+#include "ViewModels/Stack/NiagaraStackErrorItem.h"
 #include "NiagaraSystem.h"
 #include "NiagaraEditorModule.h"
 #include "NiagaraEditorUtilities.h"
+#include "NiagaraConstants.h"
 #include "EdGraph/EdGraph.h"
 #include "EdGraph/EdGraphNode.h"
 #include "EdGraph/EdGraphPin.h"
-#include "NiagaraEmitterViewModel.h"
+#include "ViewModels/NiagaraEmitterViewModel.h"
+#include "AssetRegistryModule.h"
 
 DECLARE_CYCLE_STAT(TEXT("Niagara - StackGraphUtilities - RelayoutGraph"), STAT_NiagaraEditor_StackGraphUtilities_RelayoutGraph, STATGROUP_NiagaraEditor);
 
@@ -256,11 +260,14 @@ UNiagaraNodeFunctionCall* FNiagaraStackGraphUtilities::GetNextModuleNode(UNiagar
 UNiagaraNodeOutput* FNiagaraStackGraphUtilities::GetEmitterOutputNodeForStackNode(UNiagaraNode& StackNode)
 {
 	TArray<UNiagaraNode*> NodesToCheck;
+	TSet<UNiagaraNode*> NodesChecked;
 	NodesToCheck.Add(&StackNode);
 	while (NodesToCheck.Num() > 0)
 	{
 		UNiagaraNode* NodeToCheck = NodesToCheck[0];
 		NodesToCheck.RemoveAt(0);
+		NodesChecked.Add(NodeToCheck);
+
 		if (NodeToCheck->GetClass() == UNiagaraNodeOutput::StaticClass())
 		{
 			return CastChecked<UNiagaraNodeOutput>(NodeToCheck);
@@ -273,7 +280,7 @@ UNiagaraNodeOutput* FNiagaraStackGraphUtilities::GetEmitterOutputNodeForStackNod
 			for (UEdGraphPin* LinkedPin : OutputPin->LinkedTo)
 			{
 				UNiagaraNode* LinkedNiagaraNode = Cast<UNiagaraNode>(LinkedPin->GetOwningNode());
-				if (LinkedNiagaraNode != nullptr)
+				if (LinkedNiagaraNode != nullptr && NodesChecked.Contains(LinkedNiagaraNode) == false)
 				{
 					NodesToCheck.Add(LinkedNiagaraNode);
 				}
@@ -419,20 +426,31 @@ void FNiagaraStackGraphUtilities::ConnectStackNodeGroup(const FStackNodeGroup& C
 	}
 }
 
-void FNiagaraStackGraphUtilities::InitializeDataInterfaceInputs(TSharedRef<FNiagaraSystemViewModel> SystemViewModel, TSharedRef<FNiagaraEmitterViewModel> EmitterViewModel, UNiagaraStackEditorData& StackEditorData, UNiagaraNodeFunctionCall& ModuleNode, UNiagaraNodeFunctionCall& InputFunctionCallNode)
+void FNiagaraStackGraphUtilities::InitializeStackFunctionInputs(TSharedRef<FNiagaraSystemViewModel> SystemViewModel, TSharedRef<FNiagaraEmitterViewModel> EmitterViewModel, UNiagaraStackEditorData& StackEditorData, UNiagaraNodeFunctionCall& ModuleNode, UNiagaraNodeFunctionCall& InputFunctionCallNode)
 {
 	UNiagaraStackFunctionInputCollection* FunctionInputCollection = NewObject<UNiagaraStackFunctionInputCollection>(GetTransientPackage()); 
-	FunctionInputCollection->Initialize(SystemViewModel, EmitterViewModel, StackEditorData, ModuleNode, InputFunctionCallNode, UNiagaraStackFunctionInputCollection::FDisplayOptions());
+	UNiagaraStackEntry::FRequiredEntryData RequiredEntryData(SystemViewModel, EmitterViewModel, NAME_None, NAME_None, StackEditorData);
+	FunctionInputCollection->Initialize(RequiredEntryData, ModuleNode, InputFunctionCallNode, FString());
 	FunctionInputCollection->RefreshChildren();
 
+	// Reset all direct inputs on this function to initialize data interfaces and default dynamic inputs.
 	TArray<UNiagaraStackEntry*> Children;
 	FunctionInputCollection->GetUnfilteredChildren(Children);
 	for (UNiagaraStackEntry* Child : Children)
 	{
-		UNiagaraStackFunctionInput* FunctionInput = Cast<UNiagaraStackFunctionInput>(Child);
-		if (FunctionInput != nullptr && FunctionInput->GetInputType().GetClass() != nullptr)
+		UNiagaraStackInputCategory* InputCategory = Cast<UNiagaraStackInputCategory>(Child);
+		if (InputCategory != nullptr)
 		{
-			FunctionInput->Reset();
+			TArray<UNiagaraStackEntry*> CategoryChildren;
+			InputCategory->GetUnfilteredChildren(CategoryChildren);
+			for (UNiagaraStackEntry* CategoryChild : CategoryChildren)
+			{
+				UNiagaraStackFunctionInput* FunctionInput = Cast<UNiagaraStackFunctionInput>(CategoryChild);
+				if (FunctionInput != nullptr && FunctionInput->CanReset())
+				{
+					FunctionInput->Reset();
+				}
+			}
 		}
 	}
 
@@ -466,7 +484,10 @@ void FNiagaraStackGraphUtilities::GetStackFunctionInputPins(UNiagaraNodeFunction
 			if (ReadHistory.Num() > 0 && ReadHistory[0].Get<1>() == nullptr)
 			{
 				const UEdGraphPin* InputPin = ReadHistory[0].Get<0>();
-				if (Options == ENiagaraGetStackFunctionInputPinsOptions::AllInputs || FNiagaraParameterHandle(InputPin->PinName).IsModuleHandle())
+
+				// Make sure that the module input is from the called graph, and not a nested graph.
+				if (InputPin->GetOwningNode()->GetGraph() == FunctionCallNode.GetCalledGraph() &&
+					(Options == ENiagaraGetStackFunctionInputPinsOptions::AllInputs || FNiagaraParameterHandle(InputPin->PinName).IsModuleHandle()))
 				{
 					OutInputPins.Add(InputPin);
 				}
@@ -485,7 +506,7 @@ UNiagaraNodeParameterMapSet* FNiagaraStackGraphUtilities::GetStackFunctionOverri
 	return nullptr;
 }
 
-UNiagaraNodeParameterMapSet& FNiagaraStackGraphUtilities::GetOrCreateStackFunctionOverrideNode(UNiagaraNodeFunctionCall& StackFunctionCall)
+UNiagaraNodeParameterMapSet& FNiagaraStackGraphUtilities::GetOrCreateStackFunctionOverrideNode(UNiagaraNodeFunctionCall& StackFunctionCall, const FGuid& PreferredOverrideNodeGuid)
 {
 	UNiagaraNodeParameterMapSet* OverrideNode = GetStackFunctionOverrideNode(StackFunctionCall);
 	if (OverrideNode == nullptr)
@@ -494,6 +515,10 @@ UNiagaraNodeParameterMapSet& FNiagaraStackGraphUtilities::GetOrCreateStackFuncti
 		FGraphNodeCreator<UNiagaraNodeParameterMapSet> ParameterMapSetNodeCreator(*Graph);
 		OverrideNode = ParameterMapSetNodeCreator.CreateNode();
 		ParameterMapSetNodeCreator.Finalize();
+		if (PreferredOverrideNodeGuid.IsValid())
+		{
+			OverrideNode->NodeGuid = PreferredOverrideNodeGuid;
+		}
 		OverrideNode->SetEnabledState(StackFunctionCall.GetDesiredEnabledState(), StackFunctionCall.HasUserSetTheEnabledState());
 
 		UEdGraphPin* OverrideNodeInputPin = FNiagaraStackGraphUtilities::GetParameterMapInputPin(*OverrideNode);
@@ -530,12 +555,12 @@ UEdGraphPin* FNiagaraStackGraphUtilities::GetStackFunctionInputOverridePin(UNiag
 	return nullptr;
 }
 
-UEdGraphPin& FNiagaraStackGraphUtilities::GetOrCreateStackFunctionInputOverridePin(UNiagaraNodeFunctionCall& StackFunctionCall, FNiagaraParameterHandle AliasedInputParameterHandle, FNiagaraTypeDefinition InputType)
+UEdGraphPin& FNiagaraStackGraphUtilities::GetOrCreateStackFunctionInputOverridePin(UNiagaraNodeFunctionCall& StackFunctionCall, FNiagaraParameterHandle AliasedInputParameterHandle, FNiagaraTypeDefinition InputType, const FGuid& PreferredOverrideNodeGuid)
 {
 	UEdGraphPin* OverridePin = GetStackFunctionInputOverridePin(StackFunctionCall, AliasedInputParameterHandle);
 	if (OverridePin == nullptr)
 	{
-		UNiagaraNodeParameterMapSet& OverrideNode = GetOrCreateStackFunctionOverrideNode(StackFunctionCall);
+		UNiagaraNodeParameterMapSet& OverrideNode = GetOrCreateStackFunctionOverrideNode(StackFunctionCall, PreferredOverrideNodeGuid);
 		OverrideNode.Modify();
 
 		TArray<UEdGraphPin*> OverrideInputPins;
@@ -591,25 +616,37 @@ void FNiagaraStackGraphUtilities::RemoveNodesForStackFunctionInputOverridePin(UE
 				DynamicInputNodeOverrideNode->GetInputPins(NewInputPins);
 				if (NewInputPins.Num() == 2)
 				{
-					// If there are only 2 pins, they are the parameter map input and the add pin, so the dynamic input's override node can be removed.  This
-					// not always be the case when removing dynamic input nodes because they share the same override node.
+					// If there are only 2 input pins left, they are the parameter map input and the add pin, so the dynamic input's override node 
+					// can be removed.  This not always be the case when removing dynamic input nodes because they share the same override node.
 					UEdGraphPin* InputPin = FNiagaraStackGraphUtilities::GetParameterMapInputPin(*DynamicInputNodeOverrideNode);
 					UEdGraphPin* OutputPin = FNiagaraStackGraphUtilities::GetParameterMapOutputPin(*DynamicInputNodeOverrideNode);
 
-					if (ensureMsgf(InputPin != nullptr && InputPin->LinkedTo.Num() == 1 &&
-						OutputPin != nullptr && OutputPin->LinkedTo.Num() == 2,
-						TEXT("Invalid Stack - Dynamic input node override node not connected correctly.")))
+					if (ensureMsgf(InputPin != nullptr && InputPin->LinkedTo.Num() == 1 && OutputPin != nullptr && 
+						OutputPin->LinkedTo.Num() >= 2,	TEXT("Invalid Stack - Dynamic input node override node not connected correctly.")))
 					{
 						// The DynamicInputOverrideNode will have a single input which is the previous module or override map set, and
-						// two output links, one to the dynamic input node and one to the next override map set.
+						// two or more output links, one to the dynamic input node, one to the next override map set, and 0 or more links
+						// to other dynamic inputs on sibling inputs.  Collect these linked pins to reconnect after removing the override node.
 						UEdGraphPin* LinkedInputPin = InputPin->LinkedTo[0];
-						UEdGraphPin* LinkedOutputPin = OutputPin->LinkedTo[0]->GetOwningNode() != DynamicInputNode
-							? OutputPin->LinkedTo[0]
-							: OutputPin->LinkedTo[1];
+						TArray<UEdGraphPin*> LinkedOutputPins;
+						for (UEdGraphPin* LinkedOutputPin : OutputPin->LinkedTo)
+						{
+							if (LinkedOutputPin->GetOwningNode() != DynamicInputNode)
+							{
+								LinkedOutputPins.Add(LinkedOutputPin);
+							}
+						}
+
+						// Disconnect the override node and remove it.
 						InputPin->BreakAllPinLinks();
 						OutputPin->BreakAllPinLinks();
-						LinkedInputPin->MakeLinkTo(LinkedOutputPin);
 						Graph->RemoveNode(DynamicInputNodeOverrideNode);
+
+						// Reconnect the pins which were connected to the removed override node.
+						for (UEdGraphPin* LinkedOutputPin : LinkedOutputPins)
+						{
+							LinkedInputPin->MakeLinkTo(LinkedOutputPin);
+						}
 					}
 				}
 			}
@@ -619,7 +656,7 @@ void FNiagaraStackGraphUtilities::RemoveNodesForStackFunctionInputOverridePin(UE
 	}
 }
 
-void FNiagaraStackGraphUtilities::SetLinkedValueHandleForFunctionInput(UEdGraphPin& OverridePin, FNiagaraParameterHandle LinkedParameterHandle)
+void FNiagaraStackGraphUtilities::SetLinkedValueHandleForFunctionInput(UEdGraphPin& OverridePin, FNiagaraParameterHandle LinkedParameterHandle, const FGuid& NewNodePersistentId)
 {
 	checkf(OverridePin.LinkedTo.Num() == 0, TEXT("Can't set a linked value handle when the override pin already has a value."));
 
@@ -642,9 +679,14 @@ void FNiagaraStackGraphUtilities::SetLinkedValueHandleForFunctionInput(UEdGraphP
 	UEdGraphPin* GetOutputPin = GetNode->RequestNewTypedPin(EGPD_Output, InputType, LinkedParameterHandle.GetParameterHandleString());
 	GetInputPin->MakeLinkTo(PreviousStackNodeOutputPin);
 	GetOutputPin->MakeLinkTo(&OverridePin);
+
+	if (NewNodePersistentId.IsValid())
+	{
+		GetNode->NodeGuid = NewNodePersistentId;
+	}
 }
 
-void FNiagaraStackGraphUtilities::SetDataValueObjectForFunctionInput(UEdGraphPin& OverridePin, UClass* DataObjectType, FString DataObjectName, UNiagaraDataInterface*& OutDataObject)
+void FNiagaraStackGraphUtilities::SetDataValueObjectForFunctionInput(UEdGraphPin& OverridePin, UClass* DataObjectType, FString DataObjectName, UNiagaraDataInterface*& OutDataObject, const FGuid& NewNodePersistentId)
 {
 	checkf(OverridePin.LinkedTo.Num() == 0, TEXT("Can't set a data value when the override pin already has a value."));
 	checkf(DataObjectType->IsChildOf<UNiagaraDataInterface>(), TEXT("Can only set a function input to a data interface value object"));
@@ -655,14 +697,19 @@ void FNiagaraStackGraphUtilities::SetDataValueObjectForFunctionInput(UEdGraphPin
 	UNiagaraNodeInput* InputNode = InputNodeCreator.CreateNode();
 	FNiagaraEditorUtilities::InitializeParameterInputNode(*InputNode, FNiagaraTypeDefinition(DataObjectType), CastChecked<UNiagaraGraph>(Graph), *DataObjectName);
 
-	OutDataObject = NewObject<UNiagaraDataInterface>(InputNode, DataObjectType, *DataObjectName, RF_Transactional);
+	OutDataObject = NewObject<UNiagaraDataInterface>(InputNode, DataObjectType, *DataObjectName, RF_Transactional | RF_Public);
 	InputNode->SetDataInterface(OutDataObject);
 
 	InputNodeCreator.Finalize();
 	FNiagaraStackGraphUtilities::ConnectPinToInputNode(OverridePin, *InputNode);
+
+	if (NewNodePersistentId.IsValid())
+	{
+		InputNode->NodeGuid = NewNodePersistentId;
+	}
 }
 
-void FNiagaraStackGraphUtilities::SetDynamicInputForFunctionInput(UEdGraphPin& OverridePin, UNiagaraScript* DynamicInput, UNiagaraNodeFunctionCall*& OutDynamicInputFunctionCall)
+void FNiagaraStackGraphUtilities::SetDynamicInputForFunctionInput(UEdGraphPin& OverridePin, UNiagaraScript* DynamicInput, UNiagaraNodeFunctionCall*& OutDynamicInputFunctionCall, const FGuid& NewNodePersistentId, FString SuggestedName)
 {
 	checkf(OverridePin.LinkedTo.Num() == 0, TEXT("Can't set a data value when the override pin already has a value."));
 
@@ -692,9 +739,64 @@ void FNiagaraStackGraphUtilities::SetDynamicInputForFunctionInput(UEdGraphPin& O
 	FunctionCallOutputPins[0]->MakeLinkTo(&OverridePin);
 
 	OutDynamicInputFunctionCall = FunctionCallNode;
+
+	if (NewNodePersistentId.IsValid())
+	{
+		FunctionCallNode->NodeGuid = NewNodePersistentId;
+	}
+
+	if (SuggestedName.IsEmpty() == false)
+	{
+		FunctionCallNode->SuggestName(SuggestedName);
+	}
+}
+
+void FNiagaraStackGraphUtilities::SetCustomExpressionForFunctionInput(UEdGraphPin& OverridePin, UNiagaraNodeCustomHlsl*& OutDynamicInputFunctionCall, const FGuid& NewNodePersistentId)
+{
+	checkf(OverridePin.LinkedTo.Num() == 0, TEXT("Can't set a data value when the override pin already has a value."));
+
+	UNiagaraNodeParameterMapSet* OverrideNode = CastChecked<UNiagaraNodeParameterMapSet>(OverridePin.GetOwningNode());
+	UEdGraph* Graph = OverrideNode->GetGraph();
+	const UEdGraphSchema_Niagara* Schema = CastChecked<UEdGraphSchema_Niagara>(OverrideNode->GetSchema());
+
+	FGraphNodeCreator<UNiagaraNodeCustomHlsl> FunctionCallNodeCreator(*Graph);
+	UNiagaraNodeCustomHlsl* FunctionCallNode = FunctionCallNodeCreator.CreateNode();
+	FunctionCallNode->InitAsCustomHlslDynamicInput(Schema->PinToTypeDefinition(&OverridePin));
+	FunctionCallNodeCreator.Finalize();
+	FunctionCallNode->SetEnabledState(OverrideNode->GetDesiredEnabledState(), OverrideNode->HasUserSetTheEnabledState());
+
+	UEdGraphPin* FunctionCallInputPin = FNiagaraStackGraphUtilities::GetParameterMapInputPin(*FunctionCallNode);
+	TArray<UEdGraphPin*> FunctionCallOutputPins;
+	FunctionCallNode->GetOutputPins(FunctionCallOutputPins);
+
+	const UEdGraphSchema_Niagara* NiagaraSchema = GetDefault<UEdGraphSchema_Niagara>();
+
+	FNiagaraTypeDefinition InputType = NiagaraSchema->PinToTypeDefinition(&OverridePin);
+	checkf(FunctionCallInputPin != nullptr, TEXT("Dynamic Input function call did not have a parameter map input pin."));
+	checkf(FunctionCallOutputPins.Num() == 2 && NiagaraSchema->PinToTypeDefinition(FunctionCallOutputPins[0]) == InputType, TEXT("Invalid Stack Graph - Dynamic Input function did not have the correct typed output pin"));
+
+	UEdGraphPin* OverrideNodeInputPin = FNiagaraStackGraphUtilities::GetParameterMapInputPin(*OverrideNode);
+	UEdGraphPin* PreviousStackNodeOutputPin = OverrideNodeInputPin->LinkedTo[0];
+	checkf(PreviousStackNodeOutputPin != nullptr, TEXT("Invalid Stack Graph - No previous stack node."));
+
+	FunctionCallInputPin->MakeLinkTo(PreviousStackNodeOutputPin);
+	FunctionCallOutputPins[0]->MakeLinkTo(&OverridePin);
+
+	OutDynamicInputFunctionCall = FunctionCallNode;
+
+	if (NewNodePersistentId.IsValid())
+	{
+		FunctionCallNode->NodeGuid = NewNodePersistentId;
+	}
 }
 
 bool FNiagaraStackGraphUtilities::RemoveModuleFromStack(UNiagaraNodeFunctionCall& ModuleNode)
+{
+	TArray<TWeakObjectPtr<UNiagaraNodeInput>> RemovedNodes;
+	return RemoveModuleFromStack(ModuleNode, RemovedNodes);
+}
+
+bool FNiagaraStackGraphUtilities::RemoveModuleFromStack(UNiagaraNodeFunctionCall& ModuleNode, TArray<TWeakObjectPtr<UNiagaraNodeInput>>& OutRemovedNodes)
 {
 	TArray<FNiagaraStackGraphUtilities::FStackNodeGroup> StackNodeGroups;
 	FNiagaraStackGraphUtilities::GetStackNodeGroups(ModuleNode, StackNodeGroups);
@@ -739,6 +841,11 @@ bool FNiagaraStackGraphUtilities::RemoveModuleFromStack(UNiagaraNodeFunctionCall
 	for (UNiagaraNode* NodeToRemove : NodesToRemove)
 	{
 		Graph->RemoveNode(NodeToRemove);
+		UNiagaraNodeInput* InputNode = Cast<UNiagaraNodeInput>(NodeToRemove);
+		if (InputNode != nullptr)
+		{
+			OutRemovedNodes.Add(InputNode);
+		}
 	}
 
 	return true;
@@ -808,17 +915,18 @@ UNiagaraNodeFunctionCall* FNiagaraStackGraphUtilities::AddScriptModuleToStack(FA
 	return NewModuleNode;
 }
 
-UNiagaraNodeAssignment* FNiagaraStackGraphUtilities::AddParameterModuleToStack(FNiagaraVariable ParameterVariable, UNiagaraNodeOutput& TargetOutputNode, int32 TargetIndex, const FString* InDefaultValue)
+UNiagaraNodeAssignment* FNiagaraStackGraphUtilities::AddParameterModuleToStack(const TArray<FNiagaraVariable>& ParameterVariables, UNiagaraNodeOutput& TargetOutputNode, int32 TargetIndex, const TArray<FString>& InDefaultValues)
 {
 	UEdGraph* Graph = TargetOutputNode.GetGraph();
 	Graph->Modify();
 
 	FGraphNodeCreator<UNiagaraNodeAssignment> AssignmentNodeCreator(*Graph);
 	UNiagaraNodeAssignment* NewAssignmentNode = AssignmentNodeCreator.CreateNode();
-	NewAssignmentNode->AssignmentTarget = ParameterVariable;
-	if (InDefaultValue)
+
+	check(ParameterVariables.Num() == InDefaultValues.Num());
+	for (int32 i = 0; i < ParameterVariables.Num(); i++)
 	{
-		NewAssignmentNode->AssignmentDefaultValue = *InDefaultValue;
+		NewAssignmentNode->AddAssignmentTarget(ParameterVariables[i], &InDefaultValues[i]);
 	}
 	AssignmentNodeCreator.Finalize();
 
@@ -826,20 +934,6 @@ UNiagaraNodeAssignment* FNiagaraStackGraphUtilities::AddParameterModuleToStack(F
 	return NewAssignmentNode;
 }
 
-UNiagaraMaterialParameterNode* FNiagaraStackGraphUtilities::AddMaterialParameterModuleToStack(TSharedRef<FNiagaraEmitterViewModel>& EmitterViewModel, UNiagaraNodeOutput& TargetOutputNode, int32 TargetIndex)
-{
-	UEdGraph* Graph = TargetOutputNode.GetGraph();
-	Graph->Modify();
-
-	FGraphNodeCreator<UNiagaraMaterialParameterNode> MaterialParameterNodeCreator(*Graph);
-	UNiagaraMaterialParameterNode* NewMaterialParameterNode = MaterialParameterNodeCreator.CreateNode();
-	NewMaterialParameterNode->SetEmitter(EmitterViewModel->GetEmitter());
-	MaterialParameterNodeCreator.Finalize();
-	
-	ConnectModuleNode(*NewMaterialParameterNode, TargetOutputNode, TargetIndex);
-	return NewMaterialParameterNode;
-}
-
 void GetAllNodesForModule(UNiagaraNodeFunctionCall& ModuleFunctionCall, TArray<UNiagaraNode*>& ModuleNodes)
 {
 	TArray<FNiagaraStackGraphUtilities::FStackNodeGroup> StackNodeGroups;
@@ -874,6 +968,7 @@ void FNiagaraStackGraphUtilities::SetModuleIsEnabled(UNiagaraNodeFunctionCall& F
 	for (UNiagaraNode* ModuleNode : ModuleNodes)
 	{
 		ModuleNode->SetEnabledState(bIsEnabled ? ENodeEnabledState::Enabled : ENodeEnabledState::Disabled, true);
+		ModuleNode->MarkNodeRequiresSynchronization(__FUNCTION__, false);
 	}
 	FunctionCallNode.GetNiagaraGraph()->NotifyGraphNeedsRecompile();
 }
@@ -899,7 +994,7 @@ bool FNiagaraStackGraphUtilities::ValidateGraphForOutput(UNiagaraGraph& NiagaraG
 	return true;
 }
 
-UNiagaraNodeOutput* FNiagaraStackGraphUtilities::ResetGraphForOutput(UNiagaraGraph& NiagaraGraph, ENiagaraScriptUsage ScriptUsage, FGuid ScriptUsageId)
+UNiagaraNodeOutput* FNiagaraStackGraphUtilities::ResetGraphForOutput(UNiagaraGraph& NiagaraGraph, ENiagaraScriptUsage ScriptUsage, FGuid ScriptUsageId, const FGuid& PreferredOutputNodeGuid, const FGuid& PreferredInputNodeGuid)
 {
 	NiagaraGraph.Modify();
 	UNiagaraNodeOutput* OutputNode = NiagaraGraph.FindOutputNode(ScriptUsage, ScriptUsageId);
@@ -919,6 +1014,11 @@ UNiagaraNodeOutput* FNiagaraStackGraphUtilities::ResetGraphForOutput(UNiagaraGra
 		OutputNode->Outputs.Add(FNiagaraVariable(FNiagaraTypeDefinition::GetParameterMapDef(), TEXT("Out")));
 		OutputNodeCreator.Finalize();
 
+		if (PreferredOutputNodeGuid.IsValid())
+		{
+			OutputNode->NodeGuid = PreferredOutputNodeGuid;
+		}
+
 		OutputNodeInputPin = GetParameterMapInputPin(*OutputNode);
 	}
 	else
@@ -933,6 +1033,11 @@ UNiagaraNodeOutput* FNiagaraStackGraphUtilities::ResetGraphForOutput(UNiagaraGra
 	InputNode->Usage = ENiagaraInputNodeUsage::Parameter;
 	InputNodeCreator.Finalize();
 
+	if (PreferredInputNodeGuid.IsValid())
+	{
+		InputNode->NodeGuid = PreferredInputNodeGuid;
+	}
+
 	UEdGraphPin* InputNodeOutputPin = GetParameterMapOutputPin(*InputNode);
 	OutputNodeInputPin->BreakAllPinLinks();
 	OutputNodeInputPin->MakeLinkTo(InputNodeOutputPin);
@@ -1004,6 +1109,28 @@ void GetFunctionNamesForOutputNode(UNiagaraNodeOutput& OutputNode, TArray<FStrin
 	GetFunctionNamesRecursive(&OutputNode, VisitedNodes, FunctionNames);
 }
 
+bool FNiagaraStackGraphUtilities::IsRapidIterationType(const FNiagaraTypeDefinition& InputType)
+{
+	return InputType != FNiagaraTypeDefinition::GetBoolDef() && !InputType.IsEnum() &&
+		InputType != FNiagaraTypeDefinition::GetParameterMapDef() && !InputType.IsDataInterface();
+}
+
+FNiagaraVariable FNiagaraStackGraphUtilities::CreateRapidIterationParameter(const FString& UniqueEmitterName, ENiagaraScriptUsage ScriptUsage, const FName& AliasedInputName, const FNiagaraTypeDefinition& InputType)
+{
+	FNiagaraVariable InputVariable(InputType, AliasedInputName);
+	FNiagaraVariable RapidIterationVariable;
+	if (ScriptUsage == ENiagaraScriptUsage::SystemSpawnScript || ScriptUsage == ENiagaraScriptUsage::SystemUpdateScript)
+	{
+		RapidIterationVariable = FNiagaraParameterMapHistory::ConvertVariableToRapidIterationConstantName(InputVariable, nullptr, ScriptUsage); // These names *should* have the emitter baked in...
+	}
+	else
+	{
+		RapidIterationVariable = FNiagaraParameterMapHistory::ConvertVariableToRapidIterationConstantName(InputVariable, *UniqueEmitterName, ScriptUsage);
+	}
+
+	return RapidIterationVariable;
+}
+
 void FNiagaraStackGraphUtilities::CleanUpStaleRapidIterationParameters(UNiagaraScript& Script, UNiagaraEmitter& OwningEmitter)
 {
 	UNiagaraScriptSource* Source = CastChecked<UNiagaraScriptSource>(Script.GetSource());
@@ -1039,4 +1166,468 @@ void FNiagaraStackGraphUtilities::CleanUpStaleRapidIterationParameters(UNiagaraE
 	}
 }
 
+void FNiagaraStackGraphUtilities::GetNewParameterAvailableTypes(TArray<FNiagaraTypeDefinition>& OutAvailableTypes)
+{
+	for (const FNiagaraTypeDefinition& RegisteredParameterType : FNiagaraTypeRegistry::GetRegisteredParameterTypes())
+	{
+		if (RegisteredParameterType != FNiagaraTypeDefinition::GetGenericNumericDef() && RegisteredParameterType != FNiagaraTypeDefinition::GetParameterMapDef())
+		{
+			OutAvailableTypes.Add(RegisteredParameterType);
+		}
+	}
+}
+
+void FNiagaraStackGraphUtilities::GetScriptAssetsByUsage(ENiagaraScriptUsage AssetUsage, ENiagaraScriptUsage TargetUsage, TArray<FAssetData>& OutAssets)
+{
+	FAssetRegistryModule& AssetRegistryModule = FModuleManager::LoadModuleChecked<FAssetRegistryModule>(TEXT("AssetRegistry"));
+	TArray<FAssetData> ScriptAssets;
+	AssetRegistryModule.Get().GetAssetsByClass(UNiagaraScript::StaticClass()->GetFName(), ScriptAssets);
+
+	UEnum* NiagaraScriptUsageEnum = FindObjectChecked<UEnum>(ANY_PACKAGE, TEXT("ENiagaraScriptUsage"), true);
+
+	for (const FAssetData& ScriptAsset : ScriptAssets)
+	{
+		FName CandidateUsageName;
+		ScriptAsset.GetTagValue(GET_MEMBER_NAME_CHECKED(UNiagaraScript, Usage), CandidateUsageName);
+		FString QualifiedCandidateUsageName = "ENiagaraScriptUsage::" + CandidateUsageName.ToString();
+		int32 CandidateUsageIndex = NiagaraScriptUsageEnum->GetIndexByNameString(QualifiedCandidateUsageName);
+		if (CandidateUsageIndex != INDEX_NONE)
+		{
+			ENiagaraScriptUsage CandidateAssetUsage = static_cast<ENiagaraScriptUsage>(NiagaraScriptUsageEnum->GetValueByIndex(CandidateUsageIndex));
+			if (CandidateAssetUsage == AssetUsage)
+			{
+				const FString TargetUsageBitfieldTagValue = ScriptAsset.GetTagValueRef<FString>(GET_MEMBER_NAME_CHECKED(UNiagaraScript, ModuleUsageBitmask));
+				int32 TargetUsageBitfieldValue = FCString::Atoi(*TargetUsageBitfieldTagValue);
+				bool bTargetUsageBitfieldHasTargetUsage = ((TargetUsageBitfieldValue >> (int32)TargetUsage) & 1) == 1;
+				if (bTargetUsageBitfieldHasTargetUsage)
+				{
+					OutAssets.Add(ScriptAsset);
+				}
+			}
+		}
+	}
+}
+
+void FNiagaraStackGraphUtilities::GetScriptAssetsByDependencyProvided(ENiagaraScriptUsage AssetUsage, FName DependencyName, TArray<FAssetData>& OutAssets)
+{
+	FAssetRegistryModule& AssetRegistryModule = FModuleManager::LoadModuleChecked<FAssetRegistryModule>(TEXT("AssetRegistry"));
+	TArray<FAssetData> ScriptAssets;
+	AssetRegistryModule.Get().GetAssetsByClass(UNiagaraScript::StaticClass()->GetFName(), ScriptAssets);
+	
+	for (const FAssetData& ScriptAsset : ScriptAssets)
+	{
+		auto TagName = GET_MEMBER_NAME_CHECKED(UNiagaraScript, ProvidedDependencies);
+
+		TArray<UObject::FAssetRegistryTag> Tags;
+		ScriptAsset.GetAsset()->GetAssetRegistryTags(Tags);
+		UObject::FAssetRegistryTag* FoundTag = Tags.FindByPredicate([=](const UObject::FAssetRegistryTag& InTag)
+		{
+			return InTag.Name == TagName;
+		});
+		if (FoundTag != nullptr)
+		{
+			// parse it and search through it%
+			TArray<FString> DependencyStrings;
+			FoundTag->Value.ParseIntoArray(DependencyStrings, TEXT(","));
+			for (FString DependencyString: DependencyStrings)
+			{
+				if (FName(*DependencyString) == DependencyName)
+				{
+					OutAssets.Add(ScriptAsset);
+					break;
+				}
+			}
+		}
+	}
+}
+
+void FNiagaraStackGraphUtilities::GetAvailableParametersForScript(UNiagaraNodeOutput& ScriptOutputNode, TArray<FNiagaraVariable>& OutAvailableParameters)
+{
+	TArray<FNiagaraParameterMapHistory> Histories = UNiagaraNodeParameterMapBase::GetParameterMaps(ScriptOutputNode.GetNiagaraGraph());
+
+	if (ScriptOutputNode.GetUsage() == ENiagaraScriptUsage::ParticleSpawnScript ||
+		ScriptOutputNode.GetUsage() == ENiagaraScriptUsage::ParticleSpawnScriptInterpolated ||
+		ScriptOutputNode.GetUsage() == ENiagaraScriptUsage::ParticleUpdateScript ||
+		ScriptOutputNode.GetUsage() == ENiagaraScriptUsage::ParticleEventScript)
+	{
+		OutAvailableParameters.Append(FNiagaraConstants::GetCommonParticleAttributes());
+	}
+
+	for (FNiagaraParameterMapHistory& History : Histories)
+	{
+		for (FNiagaraVariable& Variable : History.Variables)
+		{
+			if (History.IsPrimaryDataSetOutput(Variable, ScriptOutputNode.GetUsage()))
+			{
+				OutAvailableParameters.AddUnique(Variable);
+			}
+		}
+	}
+}
+
+TOptional<FName> FNiagaraStackGraphUtilities::GetNamespaceForScriptUsage(ENiagaraScriptUsage ScriptUsage)
+{
+	switch (ScriptUsage)
+	{
+	case ENiagaraScriptUsage::ParticleSpawnScript:
+	case ENiagaraScriptUsage::ParticleSpawnScriptInterpolated:
+	case ENiagaraScriptUsage::ParticleUpdateScript:
+		return FNiagaraParameterHandle::ParticleAttributeNamespace;
+	case ENiagaraScriptUsage::EmitterSpawnScript:
+	case ENiagaraScriptUsage::EmitterUpdateScript:
+		return FNiagaraParameterHandle::EmitterNamespace;
+	case ENiagaraScriptUsage::SystemSpawnScript:
+	case ENiagaraScriptUsage::SystemUpdateScript:
+		return FNiagaraParameterHandle::SystemNamespace;
+	default:
+		return TOptional<FName>();
+	}
+}
+
+void FNiagaraStackGraphUtilities::GetOwningEmitterAndScriptForStackNode(UNiagaraNode& StackNode, UNiagaraSystem& OwningSystem, UNiagaraEmitter*& OutEmitter, UNiagaraScript*& OutScript)
+{
+	OutEmitter = nullptr;
+	OutScript = nullptr;
+	UNiagaraNodeOutput* OutputNode = GetEmitterOutputNodeForStackNode(StackNode);
+	if (OutputNode != nullptr)
+	{
+		switch (OutputNode->GetUsage())
+		{
+			case ENiagaraScriptUsage::SystemSpawnScript:
+				OutScript = OwningSystem.GetSystemSpawnScript();
+				break;
+			case ENiagaraScriptUsage::SystemUpdateScript:
+				OutScript = OwningSystem.GetSystemUpdateScript();
+				break;
+			case ENiagaraScriptUsage::EmitterSpawnScript:
+			case ENiagaraScriptUsage::EmitterUpdateScript:
+			case ENiagaraScriptUsage::ParticleSpawnScript:
+			case ENiagaraScriptUsage::ParticleUpdateScript:
+			case ENiagaraScriptUsage::ParticleEventScript:
+				for (const FNiagaraEmitterHandle& EmitterHandle : OwningSystem.GetEmitterHandles())
+				{
+					UNiagaraScriptSource* EmitterSource = CastChecked<UNiagaraScriptSource>(EmitterHandle.GetInstance()->GraphSource);
+					if (EmitterSource->NodeGraph == StackNode.GetNiagaraGraph())
+					{
+						OutEmitter = EmitterHandle.GetInstance();
+						OutScript = OutEmitter->GetScript(OutputNode->GetUsage(), OutputNode->GetUsageId());
+						break;
+					}
+				}
+				break;
+		}
+	}
+}
+
+struct FRapidIterationParameterContext
+{
+	FRapidIterationParameterContext()
+		: UniqueEmitterName()
+		, OwningFunctionCall(nullptr)
+	{
+	}
+
+	FRapidIterationParameterContext(FString InUniqueEmitterName, UNiagaraNodeFunctionCall& InOwningFunctionCall)
+		: UniqueEmitterName(InUniqueEmitterName)
+		, OwningFunctionCall(&InOwningFunctionCall)
+	{
+	}
+
+	bool IsValid()
+	{
+		return UniqueEmitterName.IsEmpty() == false && OwningFunctionCall != nullptr;
+	}
+
+	FNiagaraVariable GetValue(UNiagaraScript& OwningScript, FName InputName, FNiagaraTypeDefinition Type)
+	{
+		FNiagaraParameterHandle ModuleHandle = FNiagaraParameterHandle::CreateModuleParameterHandle(InputName);
+		FNiagaraParameterHandle AliasedFunctionHandle = FNiagaraParameterHandle::CreateAliasedModuleParameterHandle(ModuleHandle, OwningFunctionCall);
+		FNiagaraVariable RapidIterationParameter = FNiagaraStackGraphUtilities::CreateRapidIterationParameter(UniqueEmitterName, OwningScript.GetUsage(), AliasedFunctionHandle.GetParameterHandleString(), Type);
+		const uint8* ValueData = OwningScript.RapidIterationParameters.GetParameterData(RapidIterationParameter);
+		if (ValueData != nullptr)
+		{
+			RapidIterationParameter.SetData(ValueData);
+			return RapidIterationParameter;
+		}
+		return FNiagaraVariable();
+	}
+
+	const FString UniqueEmitterName;
+	UNiagaraNodeFunctionCall* const OwningFunctionCall;
+};
+
+struct FStackFunctionInputValue
+{
+	FNiagaraTypeDefinition Type;
+	bool bIsOverride;
+	TOptional<FNiagaraVariable> LocalValue;
+	TOptional<FName> LinkedValue;
+	TOptional<UNiagaraDataInterface*> DataValue;
+	TOptional<UNiagaraNodeFunctionCall*> DynamicValue;
+	TMap<FName, TSharedRef<FStackFunctionInputValue>> DynamicValueInputs;
+
+	bool Matches(const FStackFunctionInputValue& Other) const
+	{
+		if (Type != Other.Type)
+		{
+			return false;
+		}
+		if (LocalValue.IsSet() && Other.LocalValue.IsSet())
+		{
+			return LocalValue.GetValue().GetType() == Other.LocalValue.GetValue().GetType() &&
+				FMemory::Memcmp(LocalValue.GetValue().GetData(), Other.LocalValue.GetValue().GetData(), LocalValue.GetValue().GetSizeInBytes()) == 0;
+		}
+		else if (LinkedValue.IsSet() && Other.LinkedValue.IsSet())
+		{
+			return LinkedValue.GetValue() == Other.LinkedValue.GetValue();
+		}
+		else if (DataValue.IsSet() && Other.DataValue.IsSet())
+		{
+			return (DataValue.GetValue() == nullptr && Other.DataValue.GetValue() == nullptr) ||
+				(DataValue.GetValue() != nullptr && Other.DataValue.GetValue() != nullptr && DataValue.GetValue()->Equals(Other.DataValue.GetValue()));
+		}
+		else if (DynamicValue.IsSet() && Other.DynamicValue.IsSet())
+		{
+			if (DynamicValue.GetValue()->FunctionScript == Other.DynamicValue.GetValue()->FunctionScript)
+			{
+				for (auto It = DynamicValueInputs.CreateConstIterator(); It; ++It)
+				{
+					FName InputName = It.Key();
+					TSharedRef<FStackFunctionInputValue> InputValue = It.Value();
+					const TSharedRef<FStackFunctionInputValue>* OtherInputValue = Other.DynamicValueInputs.Find(InputName);
+					if (OtherInputValue == nullptr || InputValue->Matches(**OtherInputValue) == false)
+					{
+						return false;
+					}
+				}
+				return true;
+			}
+			else
+			{
+				return false;
+			}
+		}
+		else
+		{
+			return false;
+		}
+	}
+};
+
+bool TryGetStackFunctionInputValue(UNiagaraScript& OwningScript, const UEdGraphPin* OverridePin, const UEdGraphPin& DefaultPin, FName InputName, FRapidIterationParameterContext RapidIterationParameterContext, FStackFunctionInputValue& OutStackFunctionInputValue)
+{
+	OutStackFunctionInputValue.Type = GetDefault<UEdGraphSchema_Niagara>()->PinToTypeDefinition(&DefaultPin);
+	OutStackFunctionInputValue.bIsOverride = OverridePin != nullptr;
+	const UEdGraphPin& InputPin = OverridePin != nullptr ? *OverridePin : DefaultPin;
+	if (RapidIterationParameterContext.IsValid() && DefaultPin.LinkedTo.Num() == 0 && OverridePin == nullptr)
+	{
+		OutStackFunctionInputValue.LocalValue = RapidIterationParameterContext.GetValue(OwningScript, InputName, OutStackFunctionInputValue.Type);
+	}
+	else if (InputPin.LinkedTo.Num() == 0)
+	{
+		const UEdGraphSchema_Niagara* NiagaraSchema = GetDefault<UEdGraphSchema_Niagara>();
+		if (NiagaraSchema->PinToTypeDefinition(&InputPin).IsDataInterface())
+		{
+			OutStackFunctionInputValue.DataValue = nullptr;
+		}
+		else
+		{
+			OutStackFunctionInputValue.LocalValue = NiagaraSchema->PinToNiagaraVariable(&InputPin, true);
+		}
+	}
+	else if (InputPin.LinkedTo.Num() == 1)
+	{
+		if (InputPin.LinkedTo[0]->GetOwningNode()->IsA<UNiagaraNodeParameterMapGet>())
+		{
+			OutStackFunctionInputValue.LinkedValue = InputPin.LinkedTo[0]->GetFName();
+		}
+		else if (InputPin.LinkedTo[0]->GetOwningNode()->IsA<UNiagaraNodeInput>())
+		{
+			OutStackFunctionInputValue.DataValue = CastChecked<UNiagaraNodeInput>(InputPin.LinkedTo[0]->GetOwningNode())->GetDataInterface();
+		}
+		else if (InputPin.LinkedTo[0]->GetOwningNode()->IsA<UNiagaraNodeFunctionCall>())
+		{
+			UNiagaraNodeFunctionCall* DynamicInputFunctionCall = CastChecked<UNiagaraNodeFunctionCall>(InputPin.LinkedTo[0]->GetOwningNode());
+			OutStackFunctionInputValue.DynamicValue = DynamicInputFunctionCall;
+			TArray<const UEdGraphPin*> DynamicValueInputPins;
+			FNiagaraStackGraphUtilities::GetStackFunctionInputPins(*DynamicInputFunctionCall,DynamicValueInputPins,
+				FNiagaraStackGraphUtilities::ENiagaraGetStackFunctionInputPinsOptions::ModuleInputsOnly);
+
+			FRapidIterationParameterContext InputRapidIterationParameterContext = RapidIterationParameterContext.IsValid() 
+				? FRapidIterationParameterContext(RapidIterationParameterContext.UniqueEmitterName, *DynamicInputFunctionCall)
+				: FRapidIterationParameterContext();
+
+			for (const UEdGraphPin* DynamicValueInputPin : DynamicValueInputPins)
+			{
+				FNiagaraParameterHandle ModuleHandle(DynamicValueInputPin->GetFName());
+				UEdGraphPin* DynamicValueOverridePin = FNiagaraStackGraphUtilities::GetStackFunctionInputOverridePin(*DynamicInputFunctionCall,
+					FNiagaraParameterHandle::CreateAliasedModuleParameterHandle(ModuleHandle, DynamicInputFunctionCall));
+
+				UEdGraphPin* DynamicValueInputDefaultPin = DynamicInputFunctionCall->FindParameterMapDefaultValuePin(DynamicValueInputPin->PinName, OwningScript.GetUsage());
+
+				FStackFunctionInputValue InputValue;
+				if (TryGetStackFunctionInputValue(OwningScript, DynamicValueOverridePin, *DynamicValueInputDefaultPin, ModuleHandle.GetName(), InputRapidIterationParameterContext, InputValue))
+				{
+					OutStackFunctionInputValue.DynamicValueInputs.Add(ModuleHandle.GetName(), MakeShared<FStackFunctionInputValue>(InputValue));
+				}
+				else
+				{
+					return false;
+				}
+			}
+		}
+		else
+		{
+			return false;
+		}
+	}
+	else
+	{
+		return false;
+	}
+	return true;
+}
+
+bool FNiagaraStackGraphUtilities::IsValidDefaultDynamicInput(UNiagaraScript& OwningScript, UEdGraphPin& DefaultPin)
+{
+	FStackFunctionInputValue InputValue;
+	return TryGetStackFunctionInputValue(OwningScript, nullptr, DefaultPin, NAME_None, FRapidIterationParameterContext(), InputValue) && InputValue.DynamicValue.IsSet();
+}
+
+bool FNiagaraStackGraphUtilities::DoesDynamicInputMatchDefault(
+	FString EmitterUniqueName, 
+	UNiagaraScript& OwningScript,
+	UNiagaraNodeFunctionCall& OwningFunctionCallNode,
+	UEdGraphPin& OverridePin,
+	FName InputName,
+	UEdGraphPin& DefaultPin)
+{
+	FStackFunctionInputValue CurrentValue;
+	FStackFunctionInputValue DefaultValue;
+	return
+		TryGetStackFunctionInputValue(OwningScript, &OverridePin, DefaultPin, InputName, FRapidIterationParameterContext(EmitterUniqueName, OwningFunctionCallNode), CurrentValue) &&
+		TryGetStackFunctionInputValue(OwningScript, nullptr, DefaultPin, NAME_None, FRapidIterationParameterContext(), DefaultValue) &&
+		CurrentValue.Matches(DefaultValue);
+}
+
+void SetInputValue(
+	TSharedRef<FNiagaraSystemViewModel> SystemViewModel,
+	TSharedRef<FNiagaraEmitterViewModel> EmitterViewModel,
+	UNiagaraStackEditorData& StackEditorData,
+	UNiagaraScript& SourceScript,
+	const TArray<TWeakObjectPtr<UNiagaraScript>> AffectedScripts,
+	UNiagaraNodeFunctionCall& ModuleNode,
+	UNiagaraNodeFunctionCall& InputFunctionCallNode,
+	FName InputName,
+	const FStackFunctionInputValue& Value)
+{
+	FNiagaraParameterHandle ModuleHandle = FNiagaraParameterHandle::CreateModuleParameterHandle(InputName);
+	FNiagaraParameterHandle AliasedFunctionHandle = FNiagaraParameterHandle::CreateAliasedModuleParameterHandle(ModuleHandle, &InputFunctionCallNode);
+	if (Value.LocalValue.IsSet())
+	{
+		bool bRapidIterationParameterSet = false;
+		if (FNiagaraStackGraphUtilities::IsRapidIterationType(Value.Type))
+		{
+			UEdGraphPin* DefaultPin = InputFunctionCallNode.FindParameterMapDefaultValuePin(ModuleHandle.GetParameterHandleString(), SourceScript.GetUsage());
+			if (DefaultPin->LinkedTo.Num() == 0)
+			{
+				FNiagaraVariable RapidIterationParameter = FNiagaraStackGraphUtilities::CreateRapidIterationParameter(
+					EmitterViewModel->GetEmitter()->GetUniqueEmitterName(), SourceScript.GetUsage(), AliasedFunctionHandle.GetParameterHandleString(), Value.Type);
+
+				for (TWeakObjectPtr<UNiagaraScript> AffectedScript : AffectedScripts)
+				{
+					AffectedScript->Modify();
+					AffectedScript->RapidIterationParameters.SetParameterData(Value.LocalValue.GetValue().GetData(), RapidIterationParameter, true);
+				}
+				bRapidIterationParameterSet = true;
+			}
+		}
+
+		if (bRapidIterationParameterSet == false)
+		{
+			const UEdGraphSchema_Niagara* NiagaraSchema = GetDefault<UEdGraphSchema_Niagara>();
+			FString PinDefaultValue;
+			if (ensureMsgf(NiagaraSchema->TryGetPinDefaultValueFromNiagaraVariable(Value.LocalValue.GetValue(), PinDefaultValue),
+				TEXT("Could not generate default value string for non-rapid iteration parameter.")))
+			{
+				UEdGraphPin& OverridePin = FNiagaraStackGraphUtilities::GetOrCreateStackFunctionInputOverridePin(InputFunctionCallNode, AliasedFunctionHandle, Value.Type);
+				OverridePin.Modify();
+				OverridePin.DefaultValue = PinDefaultValue;
+				Cast<UNiagaraNode>(OverridePin.GetOwningNode())->MarkNodeRequiresSynchronization(TEXT("OverridePin Default Value Changed"), true);
+			}
+		}
+	}
+	else if (Value.LinkedValue.IsSet())
+	{
+		UEdGraphPin& OverridePin = FNiagaraStackGraphUtilities::GetOrCreateStackFunctionInputOverridePin(InputFunctionCallNode, AliasedFunctionHandle, Value.Type);
+		FNiagaraStackGraphUtilities::SetLinkedValueHandleForFunctionInput(OverridePin, FNiagaraParameterHandle(Value.LinkedValue.GetValue()));
+	}
+	else if (Value.DataValue.IsSet())
+	{
+		UEdGraphPin& OverridePin = FNiagaraStackGraphUtilities::GetOrCreateStackFunctionInputOverridePin(InputFunctionCallNode, AliasedFunctionHandle, Value.Type);
+		FString DataObjectName = Value.DataValue.GetValue() != nullptr ? Value.DataValue.GetValue()->GetName() : Value.Type.GetName();
+		UNiagaraDataInterface* NewDataObject;
+		FNiagaraStackGraphUtilities::SetDataValueObjectForFunctionInput(OverridePin, Value.Type.GetClass(), DataObjectName, NewDataObject);
+		if (Value.DataValue.GetValue() != nullptr)
+		{
+			Value.DataValue.GetValue()->CopyTo(NewDataObject);
+		}
+	}
+	else if (Value.DynamicValue.IsSet())
+	{
+		UEdGraphPin& OverridePin = FNiagaraStackGraphUtilities::GetOrCreateStackFunctionInputOverridePin(InputFunctionCallNode, AliasedFunctionHandle, Value.Type);
+		UNiagaraNodeFunctionCall* NewDynamicInputFunctionCall;
+		FNiagaraStackGraphUtilities::SetDynamicInputForFunctionInput(OverridePin, Value.DynamicValue.GetValue()->FunctionScript, NewDynamicInputFunctionCall);
+		FNiagaraStackGraphUtilities::InitializeStackFunctionInputs(SystemViewModel, EmitterViewModel, StackEditorData, ModuleNode, *NewDynamicInputFunctionCall);
+		for (auto It = Value.DynamicValueInputs.CreateConstIterator(); It; ++It)
+		{
+			FName DynamicValueInputName = It.Key();
+			TSharedRef<FStackFunctionInputValue> DynamicValueInputValue = It.Value();
+			if (DynamicValueInputValue->bIsOverride)
+			{
+				SetInputValue(SystemViewModel, EmitterViewModel, StackEditorData, SourceScript, AffectedScripts, ModuleNode, *NewDynamicInputFunctionCall, DynamicValueInputName, *DynamicValueInputValue);
+			}
+		}
+	}
+}
+
+void FNiagaraStackGraphUtilities::ResetToDefaultDynamicInput(
+	TSharedRef<FNiagaraSystemViewModel> SystemViewModel,
+	TSharedRef<FNiagaraEmitterViewModel> EmitterViewModel,
+	UNiagaraStackEditorData& StackEditorData,
+	UNiagaraScript& SourceScript,
+	const TArray<TWeakObjectPtr<UNiagaraScript>> AffectedScripts,
+	UNiagaraNodeFunctionCall& ModuleNode,
+	UNiagaraNodeFunctionCall& InputFunctionCallNode,
+	FName InputName,
+	UEdGraphPin& DefaultPin)
+{
+	FStackFunctionInputValue DefaultValue;
+	if (ensureMsgf(TryGetStackFunctionInputValue(SourceScript, nullptr, DefaultPin, NAME_None, FRapidIterationParameterContext(), DefaultValue), TEXT("Default dynamic input was not valid")))
+	{
+		SetInputValue(SystemViewModel, EmitterViewModel, StackEditorData, SourceScript, AffectedScripts, ModuleNode, InputFunctionCallNode, InputName, DefaultValue);
+	}
+}
+
+bool FNiagaraStackGraphUtilities::GetStackIssuesRecursively(const UNiagaraStackEntry* const Entry, TArray<UNiagaraStackErrorItem*>& OutIssues)
+{
+	TArray<UNiagaraStackEntry*> Entries;
+	Entry->GetUnfilteredChildren(Entries);
+	while (Entries.Num() > 0)
+	{
+		UNiagaraStackEntry* EntryToProcess = Entries[0];
+		UNiagaraStackErrorItem* ErrorItem = Cast<UNiagaraStackErrorItem>(EntryToProcess);
+		if (ErrorItem != nullptr)
+		{
+			OutIssues.Add(ErrorItem);
+		}
+		else // don't process error items, errors don't have errors
+		{
+			EntryToProcess->GetUnfilteredChildren(Entries);
+		}
+		Entries.RemoveAtSwap(0); 
+	}
+	return OutIssues.Num() > 0;
+}
+
 #undef LOCTEXT_NAMESPACE
diff --git a/Engine/Plugins/FX/Niagara/Source/NiagaraEditor/Private/ViewModels/Stack/NiagaraStackInputCategory.cpp b/Engine/Plugins/FX/Niagara/Source/NiagaraEditor/Private/ViewModels/Stack/NiagaraStackInputCategory.cpp
new file mode 100644
--- /dev/null
+++ b/Engine/Plugins/FX/Niagara/Source/NiagaraEditor/Private/ViewModels/Stack/NiagaraStackInputCategory.cpp
@@ -0,0 +1,95 @@
+// Copyright 1998-2018 Epic Games, Inc. All Rights Reserved.
+
+#include "ViewModels/Stack/NiagaraStackInputCategory.h"
+#include "ViewModels/Stack/NiagaraStackFunctionInput.h"
+#include "NiagaraNodeFunctionCall.h"
+
+void UNiagaraStackInputCategory::Initialize(
+	FRequiredEntryData InRequiredEntryData,
+	UNiagaraNodeFunctionCall& InModuleNode,
+	UNiagaraNodeFunctionCall& InInputFunctionCallNode,
+	FText InCategoryName,
+	FString InOwnerStackItemEditorDataKey)
+{
+	bool bCategoryIsAdvanced = false;
+	FString InputCategoryStackEditorDataKey = FString::Printf(TEXT("%s-InputCategory-%s"), *InInputFunctionCallNode.NodeGuid.ToString(EGuidFormats::DigitsWithHyphens), *InCategoryName.ToString());
+	Super::Initialize(InRequiredEntryData, bCategoryIsAdvanced, InOwnerStackItemEditorDataKey, InputCategoryStackEditorDataKey);
+	ModuleNode = &InModuleNode;
+	InputFunctionCallNode = &InInputFunctionCallNode;
+	CategoryName = InCategoryName;
+	bShouldShowInStack = true;
+	
+	AddChildFilter(FOnFilterChild::CreateUObject(this, &UNiagaraStackInputCategory::FilterForVisibleCondition));
+	AddChildFilter(FOnFilterChild::CreateUObject(this, &UNiagaraStackInputCategory::FilterForIsInlineEditConditionToggle));
+}
+
+const FText& UNiagaraStackInputCategory::GetCategoryName() const
+{
+	return CategoryName;
+}
+
+void UNiagaraStackInputCategory::ResetInputs()
+{
+	Inputs.Empty();
+}
+
+void UNiagaraStackInputCategory::AddInput(FName InInputParameterHandle, FNiagaraTypeDefinition InInputType)
+{
+	Inputs.Add({ InInputParameterHandle, InInputType });
+}
+
+void UNiagaraStackInputCategory::RefreshChildrenInternal(const TArray<UNiagaraStackEntry*>& CurrentChildren, TArray<UNiagaraStackEntry*>& NewChildren, TArray<FStackIssue>& NewIssues)
+{
+	for (FInputParameterHandleAndType& Input : Inputs)
+	{
+		UNiagaraStackFunctionInput* InputChild = FindCurrentChildOfTypeByPredicate<UNiagaraStackFunctionInput>(CurrentChildren,
+			[&](UNiagaraStackFunctionInput* CurrentInput) { return CurrentInput->GetInputParameterHandle() == Input.ParameterHandle; });
+
+		if (InputChild == nullptr)
+		{
+			InputChild = NewObject<UNiagaraStackFunctionInput>(this);
+			InputChild->Initialize(CreateDefaultChildRequiredData(), *ModuleNode, *InputFunctionCallNode,
+				Input.ParameterHandle, Input.Type, GetOwnerStackItemEditorDataKey());
+		}
+
+		NewChildren.Add(InputChild);
+	}
+}
+
+FText UNiagaraStackInputCategory::GetDisplayName() const
+{
+	return CategoryName;
+}
+
+bool UNiagaraStackInputCategory::GetShouldShowInStack() const
+{
+	return bShouldShowInStack;
+}
+
+UNiagaraStackEntry::EStackRowStyle UNiagaraStackInputCategory::GetStackRowStyle() const
+{
+	return EStackRowStyle::ItemCategory;
+}
+
+bool UNiagaraStackInputCategory::GetIsEnabled() const
+{
+	return InputFunctionCallNode->GetDesiredEnabledState() == ENodeEnabledState::Enabled;
+}
+
+void UNiagaraStackInputCategory::SetShouldShowInStack(bool bInShouldShowInStack)
+{
+	bShouldShowInStack = bInShouldShowInStack;
+}
+
+bool UNiagaraStackInputCategory::FilterForVisibleCondition(const UNiagaraStackEntry& Child) const
+{
+	const UNiagaraStackFunctionInput* StackFunctionInputChild = Cast<UNiagaraStackFunctionInput>(&Child);
+	return StackFunctionInputChild == nullptr || StackFunctionInputChild->GetHasVisibleCondition() == false || StackFunctionInputChild->GetVisibleConditionEnabled();
+}
+
+bool UNiagaraStackInputCategory::FilterForIsInlineEditConditionToggle(const UNiagaraStackEntry& Child) const
+{
+	const UNiagaraStackFunctionInput* StackFunctionInputChild = Cast<UNiagaraStackFunctionInput>(&Child);
+	return StackFunctionInputChild == nullptr || StackFunctionInputChild->GetIsInlineEditConditionToggle() == false;
+}
+
diff --git a/Engine/Plugins/FX/Niagara/Source/NiagaraEditor/Private/ViewModels/Stack/NiagaraStackItem.cpp b/Engine/Plugins/FX/Niagara/Source/NiagaraEditor/Private/ViewModels/Stack/NiagaraStackItem.cpp
--- a/Engine/Plugins/FX/Niagara/Source/NiagaraEditor/Private/ViewModels/Stack/NiagaraStackItem.cpp
+++ b/Engine/Plugins/FX/Niagara/Source/NiagaraEditor/Private/ViewModels/Stack/NiagaraStackItem.cpp
@@ -1,11 +1,23 @@
 // Copyright 1998-2018 Epic Games, Inc. All Rights Reserved.
 
 #include "ViewModels/Stack/NiagaraStackItem.h"
+#include "ViewModels/Stack/NiagaraStackSpacer.h"
+#include "ViewModels/Stack/NiagaraStackAdvancedExpander.h"
+#include "ViewModels/Stack/NiagaraStackGraphUtilities.h"
+#include "ViewModels/Stack/NiagaraStackErrorItem.h"
+#include "ViewModels/Stack/NiagaraStackEntry.h"
+#include "NiagaraStackEditorData.h"
 
-void UNiagaraStackItem::Initialize(TSharedRef<FNiagaraSystemViewModel> InSystemViewModel, TSharedRef<FNiagaraEmitterViewModel> InEmitterViewModel, UNiagaraStackEditorData& InStackEditorData)
+void UNiagaraStackItem::Initialize(FRequiredEntryData InRequiredEntryData, FString InStackEditorDataKey)
 {
-	Super::Initialize(InSystemViewModel, InEmitterViewModel);
-	StackEditorData = &InStackEditorData;
+	Super::Initialize(InRequiredEntryData, InStackEditorDataKey);
+	AddChildFilter(FOnFilterChild::CreateUObject(this, &UNiagaraStackItem::FilterAdvancedChildren));
+	AddChildFilter(FOnFilterChild::CreateUObject(this, &UNiagaraStackItem::FilterShowAdvancedChild));
+}
+
+UNiagaraStackEntry::EStackRowStyle UNiagaraStackItem::GetStackRowStyle() const
+{
+	return UNiagaraStackEntry::EStackRowStyle::ItemHeader;
 }
 
 void UNiagaraStackItem::SetOnModifiedGroupItems(FOnModifiedGroupItems OnModifiedGroupItems)
@@ -13,7 +25,182 @@ void UNiagaraStackItem::SetOnModifiedGroupItems(FOnModifiedGroupItems OnModified
 	ModifiedGroupItemsDelegate = OnModifiedGroupItems;
 }
 
-UNiagaraStackEditorData& UNiagaraStackItem::GetStackEditorData() const
+uint32 UNiagaraStackItem::GetRecursiveStackIssuesCount() const
+{
+	if (RecursiveStackIssuesCount.IsSet() == false)
+	{
+		TArray<UNiagaraStackErrorItem*> RecursiveIssues;
+		FNiagaraStackGraphUtilities::GetStackIssuesRecursively(this, RecursiveIssues);
+		RecursiveStackIssuesCount = RecursiveIssues.Num();
+		EStackIssueSeverity MinSeverity = EStackIssueSeverity::Info;
+		for (auto Issue : RecursiveIssues)
+		{
+			if (Issue->GetStackIssue().Severity < MinSeverity)
+			{
+				MinSeverity = Issue->GetStackIssue().Severity;
+			}
+		}
+		HighestIssueSeverity = MinSeverity;
+	}
+	return RecursiveStackIssuesCount.GetValue();
+}
+
+EStackIssueSeverity UNiagaraStackItem::GetHighestStackIssueSeverity() const
+{
+	if (HighestIssueSeverity.IsSet() == false)
+	{
+		GetRecursiveStackIssuesCount();
+	}
+	return HighestIssueSeverity.GetValue();
+}
+
+void UNiagaraStackItem::RefreshChildrenInternal(const TArray<UNiagaraStackEntry*>& CurrentChildren, TArray<UNiagaraStackEntry*>& NewChildren, TArray<FStackIssue>& NewIssues)
+{
+	if (FooterSpacer == nullptr)
+	{
+		FooterSpacer = NewObject<UNiagaraStackSpacer>(this);
+		FooterSpacer->Initialize(CreateDefaultChildRequiredData(), "ItemFooterSpacer", 1, UNiagaraStackEntry::EStackRowStyle::ItemContent);
+	}
+
+	if (ShowAdvancedExpander == nullptr)
+	{
+		ShowAdvancedExpander = NewObject<UNiagaraStackAdvancedExpander>(this);
+		ShowAdvancedExpander->Initialize(CreateDefaultChildRequiredData(), GetStackEditorDataKey(), GetOwningNiagaraNode());
+		ShowAdvancedExpander->SetOnToggleShowAdvanced(UNiagaraStackAdvancedExpander::FOnToggleShowAdvanced::CreateUObject(this, &UNiagaraStackItem::ToggleShowAdvanced));
+	}
+
+	NewChildren.Add(FooterSpacer);
+	NewChildren.Add(ShowAdvancedExpander);
+	RecursiveStackIssuesCount.Reset();
+	HighestIssueSeverity.Reset();
+}
+
+void GetContentChildren(UNiagaraStackEntry& CurrentEntry, TArray<UNiagaraStackItemContent*>& ContentChildren)
+{
+	TArray<UNiagaraStackEntry*> Children;
+	CurrentEntry.GetUnfilteredChildren(Children);
+	for (UNiagaraStackEntry* Child : Children)
+	{
+		UNiagaraStackItemContent* ContentChild = Cast<UNiagaraStackItemContent>(Child);
+		if (ContentChild != nullptr)
+		{
+			ContentChildren.Add(ContentChild);
+		}
+		GetContentChildren(*Child, ContentChildren);
+	}
+}
+
+void UNiagaraStackItem::PostRefreshChildrenInternal()
+{
+	Super::PostRefreshChildrenInternal();
+	bHasAdvancedContent = false;
+	TArray<UNiagaraStackItemContent*> ContentChildren;
+	GetContentChildren(*this, ContentChildren);
+	for (UNiagaraStackItemContent* ContentChild : ContentChildren)
+	{
+		if (ContentChild->GetIsAdvanced())
+		{
+			bHasAdvancedContent = true;
+			break;
+		}
+	}
+}
+
+int32 UNiagaraStackItem::GetChildIndentLevel() const
+{
+	return GetIndentLevel();
+}
+
+UNiagaraNode* UNiagaraStackItem::GetOwningNiagaraNode() const
+{
+	return nullptr;
+}
+
+void UNiagaraStackItem::ChlildStructureChangedInternal()
+{
+	Super::ChlildStructureChangedInternal();
+	RecursiveStackIssuesCount.Reset();
+	HighestIssueSeverity.Reset();
+}
+
+bool UNiagaraStackItem::FilterAdvancedChildren(const UNiagaraStackEntry& Child) const
+{
+	const UNiagaraStackItemContent* ItemContent = Cast<UNiagaraStackItemContent>(&Child);
+	if (ItemContent == nullptr || ItemContent->GetIsAdvanced() == false)
+	{
+		return true;
+	}
+	else
+	{
+		return GetStackEditorData().GetShowAllAdvanced() || GetStackEditorData().GetStackItemShowAdvanced(GetStackEditorDataKey(), false);
+	}
+}
+
+bool UNiagaraStackItem::FilterShowAdvancedChild(const UNiagaraStackEntry& Child) const 
+{
+	if (&Child == ShowAdvancedExpander && bHasAdvancedContent == false)
+	{
+		return false;
+	}
+	else if (&Child == FooterSpacer && bHasAdvancedContent)
+	{
+		return false;
+	}
+	else
+	{
+		return true;
+	}
+}
+
+void UNiagaraStackItem::ToggleShowAdvanced()
+{
+	bool bCurrentShowAdvanced = GetStackEditorData().GetStackItemShowAdvanced(GetStackEditorDataKey(), false);
+	GetStackEditorData().SetStackItemShowAdvanced(GetStackEditorDataKey(), !bCurrentShowAdvanced);
+	OnStructureChanged().Broadcast();
+}
+
+void UNiagaraStackItemContent::Initialize(FRequiredEntryData InRequiredEntryData, bool bInIsAdvanced, FString InOwningStackItemEditorDataKey, FString InStackEditorDataKey)
+{
+	Super::Initialize(InRequiredEntryData, InStackEditorDataKey);
+	OwningStackItemEditorDataKey = InOwningStackItemEditorDataKey;
+	bIsAdvanced = bInIsAdvanced;
+	AddChildFilter(FOnFilterChild::CreateUObject(this, &UNiagaraStackItemContent::FilterAdvancedChildren));
+}
+
+UNiagaraStackEntry::EStackRowStyle UNiagaraStackItemContent::GetStackRowStyle() const
+{
+	return bIsAdvanced ? EStackRowStyle::ItemContentAdvanced : EStackRowStyle::ItemContent;
+}
+
+bool UNiagaraStackItemContent::GetIsAdvanced() const
+{
+	return bIsAdvanced;
+}
+
+FString UNiagaraStackItemContent::GetOwnerStackItemEditorDataKey() const
+{
+	return OwningStackItemEditorDataKey;
+}
+
+void UNiagaraStackItemContent::SetIsAdvanced(bool bInIsAdvanced)
+{
+	if (bIsAdvanced != bInIsAdvanced)
+	{
+		// When changing advanced, invalidate the structure so that the filters run again.
+		bIsAdvanced = bInIsAdvanced;
+		OnStructureChanged().Broadcast();
+	}
+}
+
+bool UNiagaraStackItemContent::FilterAdvancedChildren(const UNiagaraStackEntry& Child) const
 {
-	return *StackEditorData;
+	const UNiagaraStackItemContent* ItemContent = Cast<UNiagaraStackItemContent>(&Child);
+	if (ItemContent == nullptr || ItemContent->GetIsAdvanced() == false)
+	{
+		return true;
+	}
+	else
+	{
+		return GetStackEditorData().GetShowAllAdvanced() || GetStackEditorData().GetStackItemShowAdvanced(OwningStackItemEditorDataKey, false);
+	}
 }
diff --git a/Engine/Plugins/FX/Niagara/Source/NiagaraEditor/Private/ViewModels/Stack/NiagaraStackItemGroup.cpp b/Engine/Plugins/FX/Niagara/Source/NiagaraEditor/Private/ViewModels/Stack/NiagaraStackItemGroup.cpp
--- a/Engine/Plugins/FX/Niagara/Source/NiagaraEditor/Private/ViewModels/Stack/NiagaraStackItemGroup.cpp
+++ b/Engine/Plugins/FX/Niagara/Source/NiagaraEditor/Private/ViewModels/Stack/NiagaraStackItemGroup.cpp
@@ -2,50 +2,98 @@
 
 #include "ViewModels/Stack/NiagaraStackItemGroup.h"
 #include "ViewModels/Stack/NiagaraStackItem.h"
+#include "ViewModels/Stack/NiagaraStackErrorItem.h"
+#include "ViewModels/Stack/NiagaraStackEntry.h"
 #include "NiagaraStackEditorData.h"
+#include "ViewModels/Stack/NiagaraStackSpacer.h"
+#include "ViewModels/Stack/NiagaraStackGraphUtilities.h"
 
-void UNiagaraStackItemGroup::Initialize(TSharedRef<FNiagaraSystemViewModel> InSystemViewModel, TSharedRef<FNiagaraEmitterViewModel> InEmitterViewModel, UNiagaraStackEditorData& InStackEditorData)
+
+void UNiagaraStackItemGroup::Initialize(FRequiredEntryData InRequiredEntryData, FText InDisplayName, FText InToolTip, INiagaraStackItemGroupAddUtilities* InAddUtilities)
 {
-	Super::Initialize(InSystemViewModel, InEmitterViewModel);
-	StackEditorData = &InStackEditorData;
+	Super::Initialize(InRequiredEntryData, InDisplayName.ToString());
+	GroupDisplayName = InDisplayName;
+	GroupToolTip = InToolTip;
+	AddUtilities = InAddUtilities;
 }
 
-bool UNiagaraStackItemGroup::GetCanExpand() const
+FText UNiagaraStackItemGroup::GetDisplayName() const
 {
-	return true;
+	return GroupDisplayName;
 }
 
-bool UNiagaraStackItemGroup::GetIsExpanded() const
+UNiagaraStackEntry::EStackRowStyle UNiagaraStackItemGroup::GetStackRowStyle() const
 {
-	return StackEditorData->GetStackEntryIsExpanded(GetDisplayName().ToString(), IsExpandedByDefault());
+	return EStackRowStyle::GroupHeader;
 }
 
-void UNiagaraStackItemGroup::SetIsExpanded(bool bInExpanded)
+FText UNiagaraStackItemGroup::GetTooltipText() const 
 {
-	StackEditorData->SetStackEntryIsExpanded(GetDisplayName().ToString(), bInExpanded);
+	return GroupToolTip;
 }
 
-FName UNiagaraStackItemGroup::GetItemBackgroundName() const
+INiagaraStackItemGroupAddUtilities* UNiagaraStackItemGroup::GetAddUtilities() const
 {
-	return "NiagaraEditor.Stack.Group.BackgroundColor";
+	return AddUtilities;
 }
 
-FName UNiagaraStackItemGroup::GetTextStyleName() const
+uint32 UNiagaraStackItemGroup::GetRecursiveStackIssuesCount() const
 {
-	return "NiagaraEditor.Stack.GroupText";
+	if (RecursiveStackIssuesCount.IsSet() == false)
+	{
+		TArray<UNiagaraStackErrorItem*> RecursiveIssues;
+		FNiagaraStackGraphUtilities::GetStackIssuesRecursively(this, RecursiveIssues);
+		RecursiveStackIssuesCount = RecursiveIssues.Num();
+		EStackIssueSeverity MinSeverity = EStackIssueSeverity::Info;
+		for (auto Issue : RecursiveIssues)
+		{
+			if (Issue->GetStackIssue().Severity < MinSeverity)
+			{
+				MinSeverity = Issue->GetStackIssue().Severity;
+			}
+		}
+		HighestIssueSeverity = MinSeverity;
+	}
+	return RecursiveStackIssuesCount.GetValue();
 }
 
-FText UNiagaraStackItemGroup::GetTooltipText() const 
+EStackIssueSeverity UNiagaraStackItemGroup::GetHighestStackIssueSeverity() const
+{
+	if (HighestIssueSeverity.IsSet() == false)
+	{
+		GetRecursiveStackIssuesCount();
+	}
+	return HighestIssueSeverity.GetValue();
+}
+
+void UNiagaraStackItemGroup::SetDisplayName(FText InDisplayName)
+{
+	GroupDisplayName = InDisplayName;
+}
+
+void UNiagaraStackItemGroup::RefreshChildrenInternal(const TArray<UNiagaraStackEntry*>& CurrentChildren, TArray<UNiagaraStackEntry*>& NewChildren, TArray<FStackIssue>& NewIssues)
 {
-	return TooltipText;
+	UNiagaraStackSpacer* SeparatorSpacer = FindCurrentChildOfTypeByPredicate<UNiagaraStackSpacer>(CurrentChildren,
+		[=](UNiagaraStackSpacer* CurrentSpacer) { return CurrentSpacer->GetSpacerKey() == "SeparatorSpacer"; });
+	if (SeparatorSpacer == nullptr)
+	{
+		SeparatorSpacer = NewObject<UNiagaraStackSpacer>(this);
+		FRequiredEntryData RequiredEntryData(GetSystemViewModel(), GetEmitterViewModel(), GetExecutionCategoryName(), NAME_None, GetStackEditorData());
+		SeparatorSpacer->Initialize(RequiredEntryData, "SeparatorSpacer");
+	}
+	NewChildren.Add(SeparatorSpacer);
+	RecursiveStackIssuesCount.Reset();
+	HighestIssueSeverity.Reset();
 }
 
-void UNiagaraStackItemGroup::SetTooltipText(const FText& InText)
+int32 UNiagaraStackItemGroup::GetChildIndentLevel() const
 {
-	TooltipText = InText;
+	return GetIndentLevel();
 }
 
-UNiagaraStackEditorData& UNiagaraStackItemGroup::GetStackEditorData() const
+void UNiagaraStackItemGroup::ChlildStructureChangedInternal()
 {
-	return *StackEditorData;
+	Super::ChlildStructureChangedInternal();
+	RecursiveStackIssuesCount.Reset();
+	HighestIssueSeverity.Reset();
 }
