Copying //UE4/Dev-Niagara to //UE4/Dev-Main (Source: //UE4/Dev-Niagara @ 4074996)

diff --git a/Engine/Plugins/FX/Niagara/Source/NiagaraEditor/Private/ViewModels/Stack/NiagaraStackModuleItem.cpp b/Engine/Plugins/FX/Niagara/Source/NiagaraEditor/Private/ViewModels/Stack/NiagaraStackModuleItem.cpp
--- a/Engine/Plugins/FX/Niagara/Source/NiagaraEditor/Private/ViewModels/Stack/NiagaraStackModuleItem.cpp
+++ b/Engine/Plugins/FX/Niagara/Source/NiagaraEditor/Private/ViewModels/Stack/NiagaraStackModuleItem.cpp
@@ -2,13 +2,13 @@
 
 #include "ViewModels/Stack/NiagaraStackModuleItem.h"
 #include "NiagaraEditorModule.h"
-#include "NiagaraSystemViewModel.h"
-#include "NiagaraEmitterHandleViewModel.h"
-#include "NiagaraEmitterViewModel.h"
+#include "ViewModels/NiagaraSystemViewModel.h"
+#include "ViewModels/NiagaraEmitterHandleViewModel.h"
+#include "ViewModels/NiagaraEmitterViewModel.h"
 #include "ViewModels/Stack/NiagaraStackFunctionInputCollection.h"
 #include "ViewModels/Stack/NiagaraStackFunctionInput.h"
+#include "ViewModels/Stack/NiagaraStackInputCategory.h"
 #include "ViewModels/Stack/NiagaraStackModuleItemOutputCollection.h"
-#include "ViewModels/Stack/NiagaraStackItemExpander.h"
 #include "NiagaraNodeFunctionCall.h"
 #include "NiagaraNodeParameterMapSet.h"
 #include "NiagaraNodeOutput.h"
@@ -16,20 +16,24 @@
 #include "NiagaraStackEditorData.h"
 #include "NiagaraEditorStyle.h"
 #include "NiagaraGraph.h"
+#include "NiagaraScript.h"
 #include "ViewModels/Stack/NiagaraStackGraphUtilities.h"
 #include "NiagaraEmitterHandle.h"
 #include "NiagaraScriptMergeManager.h"
 #include "NiagaraStackEditorData.h"
+#include "NiagaraScriptSource.h"
+#include "NiagaraNodeAssignment.h"
+#include "NiagaraCommon.h"
+#include "Widgets/SWidget.h"
 
 #include "ScopedTransaction.h"
 
-#define LOCTEXT_NAMESPACE "NiagaraStackViewModel"
+#define LOCTEXT_NAMESPACE "NiagaraStackModuleItem"
 
 UNiagaraStackModuleItem::UNiagaraStackModuleItem()
 	: FunctionCallNode(nullptr)
-	, PinnedInputCollection(nullptr)
-	, UnpinnedInputCollection(nullptr)
-	, ModuleExpander(nullptr)
+	, bCanRefresh(false)
+	, InputCollection(nullptr)
 {
 }
 
@@ -38,20 +42,19 @@ const UNiagaraNodeFunctionCall& UNiagaraStackModuleItem::GetModuleNode() const
 	return *FunctionCallNode;
 }
 
-void UNiagaraStackModuleItem::Initialize(TSharedRef<FNiagaraSystemViewModel> InSystemViewModel, TSharedRef<FNiagaraEmitterViewModel> InEmitterViewModel, UNiagaraStackEditorData& InStackEditorData, UNiagaraNodeFunctionCall& InFunctionCallNode)
+UNiagaraNodeFunctionCall& UNiagaraStackModuleItem::GetModuleNode()
+{
+	return *FunctionCallNode;
+}
+
+void UNiagaraStackModuleItem::Initialize(FRequiredEntryData InRequiredEntryData, INiagaraStackItemGroupAddUtilities* InGroupAddUtilities, UNiagaraNodeFunctionCall& InFunctionCallNode)
 {
 	checkf(FunctionCallNode == nullptr, TEXT("Can not set the node more than once."));
-	Super::Initialize(InSystemViewModel, InEmitterViewModel, InStackEditorData);
+	FString ModuleStackEditorDataKey = InFunctionCallNode.NodeGuid.ToString(EGuidFormats::DigitsWithHyphens);
+	Super::Initialize(InRequiredEntryData, ModuleStackEditorDataKey);
+	GroupAddUtilities = InGroupAddUtilities;
 	FunctionCallNode = &InFunctionCallNode;
-	ModuleEditorDataKey = FNiagaraStackGraphUtilities::GenerateStackModuleEditorDataKey(*FunctionCallNode);
-
-	FOnFilterChild CollapsedFilter = FOnFilterChild::CreateLambda([this](const UNiagaraStackEntry& Child)
-	{
-		bool bChildShouldHiddenWhenNotExpanded = &Child == UnpinnedInputCollection || &Child == OutputCollection;
-		bool bIsExpanded = GetStackEditorData().GetStackEntryIsExpanded(ModuleEditorDataKey, true);
-		return bIsExpanded || bChildShouldHiddenWhenNotExpanded == false;
-	});
-	AddChildFilter(CollapsedFilter);
+	AddChildFilter(FOnFilterChild::CreateUObject(this, &UNiagaraStackModuleItem::FilterOutputCollection));
 
 	// Update bCanMoveAndDelete
 	if (GetSystemViewModel()->GetEditMode() == ENiagaraSystemViewModelEditMode::EmitterAsset)
@@ -98,132 +101,199 @@ FText UNiagaraStackModuleItem::GetTooltipText() const
 	}
 }
 
-void UNiagaraStackModuleItem::RefreshChildrenInternal(const TArray<UNiagaraStackEntry*>& CurrentChildren, TArray<UNiagaraStackEntry*>& NewChildren)
+INiagaraStackItemGroupAddUtilities* UNiagaraStackModuleItem::GetGroupAddUtilities()
+{
+	return GroupAddUtilities;
+}
+
+void UNiagaraStackModuleItem::RefreshChildrenInternal(const TArray<UNiagaraStackEntry*>& CurrentChildren, TArray<UNiagaraStackEntry*>& NewChildren, TArray<FStackIssue>& NewIssues)
 {
-	Errors.Empty();
-	if (FunctionCallNode->ScriptIsValid())
+	bCanRefresh = false;
+
+	if (FunctionCallNode != nullptr && FunctionCallNode->ScriptIsValid())
 	{
-		if (PinnedInputCollection == nullptr)
+		// Determine if meta-data requires that we add our own refresh button here.
+		if (FunctionCallNode->FunctionScript)
 		{
-			UNiagaraStackFunctionInputCollection::FDisplayOptions DisplayOptions;
-			DisplayOptions.DisplayName = LOCTEXT("PinnedInputsLabel", "Pinned Inputs");
-			DisplayOptions.ChildItemIndentLevel = 0;
-			DisplayOptions.bShouldShowInStack = false;
-			DisplayOptions.ChildFilter = FOnFilterChild::CreateLambda([](const UNiagaraStackEntry& Child)
+			UNiagaraScriptSource* Source = CastChecked<UNiagaraScriptSource>(FunctionCallNode->FunctionScript->GetSource());
+			UNiagaraGraph* Graph = CastChecked<UNiagaraGraph>(Source->NodeGraph);
+			const TMap<FNiagaraVariable, FNiagaraVariableMetaData>& MetaDataMap = Graph->GetAllMetaData();
+			auto Iter = MetaDataMap.CreateConstIterator();
+			while (Iter)
 			{
-				const UNiagaraStackFunctionInput* FunctionInput = CastChecked<const UNiagaraStackFunctionInput>(&Child);
-				return FunctionInput->GetIsPinned();
-			});
-
-			TArray<FString> InputParameterHandlePath;
-			PinnedInputCollection = NewObject<UNiagaraStackFunctionInputCollection>(this);
-			PinnedInputCollection->Initialize(GetSystemViewModel(), GetEmitterViewModel(), GetStackEditorData(), *FunctionCallNode, *FunctionCallNode, DisplayOptions);
-			PinnedInputCollection->OnInputPinnedChanged().AddUObject(this, &UNiagaraStackModuleItem::InputPinnedChanged);
+				auto PropertyIter = Iter.Value().PropertyMetaData.CreateConstIterator();
+				while (PropertyIter)
+				{
+					if (PropertyIter.Key() == (TEXT("DisplayNameArg0")))
+					{
+						bCanRefresh = true;
+					}
+					++PropertyIter;
+				}
+				++Iter;
+			}
 		}
 
-		if (UnpinnedInputCollection == nullptr)
+		if (InputCollection == nullptr)
 		{
-			UNiagaraStackFunctionInputCollection::FDisplayOptions DisplayOptions;
-			DisplayOptions.DisplayName = LOCTEXT("InputsLabel", "Inputs");
-			DisplayOptions.ChildItemIndentLevel = 1;
-			DisplayOptions.bShouldShowInStack = true;
-			DisplayOptions.ChildFilter = FOnFilterChild::CreateLambda([](const UNiagaraStackEntry& Child)
-			{
-				const UNiagaraStackFunctionInput* FunctionInput = CastChecked<const UNiagaraStackFunctionInput>(&Child);
-				return FunctionInput->GetIsPinned() == false;
-			});
-
 			TArray<FString> InputParameterHandlePath;
-			UnpinnedInputCollection = NewObject<UNiagaraStackFunctionInputCollection>(this);
-			UnpinnedInputCollection->Initialize(GetSystemViewModel(), GetEmitterViewModel(), GetStackEditorData(), *FunctionCallNode, *FunctionCallNode, DisplayOptions);
-			UnpinnedInputCollection->OnInputPinnedChanged().AddUObject(this, &UNiagaraStackModuleItem::InputPinnedChanged);
+			InputCollection = NewObject<UNiagaraStackFunctionInputCollection>(this);
+			InputCollection->Initialize(CreateDefaultChildRequiredData(), *FunctionCallNode, *FunctionCallNode, GetStackEditorDataKey());
 		}
 
+		InputCollection->SetShouldShowInStack(GetStackEditorData().GetShowOutputs());
+
 		if (OutputCollection == nullptr)
 		{
 			OutputCollection = NewObject<UNiagaraStackModuleItemOutputCollection>(this);
-			OutputCollection->Initialize(GetSystemViewModel(), GetEmitterViewModel(), *FunctionCallNode);
-			OutputCollection->SetDisplayName(LOCTEXT("OutputsLabel", "Outputs"));
+			OutputCollection->Initialize(CreateDefaultChildRequiredData(), *FunctionCallNode);
 		}
 
-		if (ModuleExpander == nullptr)
-		{
-			ModuleExpander = NewObject<UNiagaraStackItemExpander>(this);
-			ModuleExpander->Initialize(GetSystemViewModel(), GetEmitterViewModel(), GetStackEditorData(), ModuleEditorDataKey, true);
-			ModuleExpander->SetOnExpnadedChanged(UNiagaraStackItemExpander::FOnExpandedChanged::CreateUObject(this, &UNiagaraStackModuleItem::ModuleExpandedChanged));
-		}
-
-		NewChildren.Add(PinnedInputCollection);
-		NewChildren.Add(UnpinnedInputCollection);
+		NewChildren.Add(InputCollection);
 		NewChildren.Add(OutputCollection);
-		NewChildren.Add(ModuleExpander);
-	}
-	else
-	{
-		FError InvalidScriptError;
-		InvalidScriptError.ErrorSummaryText = LOCTEXT("MissingModule", "The referenced module's script asset is missing.");
-		InvalidScriptError.ErrorText = InvalidScriptError.ErrorSummaryText;
 	}
 
 	RefreshIsEnabled();
+
+	Super::RefreshChildrenInternal(CurrentChildren, NewChildren, NewIssues);
+	
+	RefreshIssues(NewIssues);
 }
 
-void UNiagaraStackModuleItem::RefreshIsEnabled()
+void UNiagaraStackModuleItem::RefreshIssues(TArray<FStackIssue>& NewIssues)
 {
-	TOptional<bool> IsEnabled = FNiagaraStackGraphUtilities::GetModuleIsEnabled(*FunctionCallNode);
-	if (IsEnabled.IsSet())
+	if (FunctionCallNode != nullptr)
 	{
-		bIsEnabled = IsEnabled.GetValue();
+		if (!FunctionCallNode->ScriptIsValid())
+		{
+			UNiagaraStackEntry::FStackIssue InvalidScriptError;
+			InvalidScriptError.ShortDescription = LOCTEXT("MissingModule", "The referenced module's script asset is missing.");
+			InvalidScriptError.LongDescription = InvalidScriptError.ShortDescription;
+			InvalidScriptError.UniqueIdentifier = FName(*FString::Printf(TEXT("MissingModule-%s"), *GetStackEditorDataKey()));
+			NewIssues.Add(InvalidScriptError);
+		}
+
+		TOptional<bool> IsEnabled = FNiagaraStackGraphUtilities::GetModuleIsEnabled(*FunctionCallNode);
+		if (!IsEnabled.IsSet())
+		{
+			bIsEnabled = false;
+			UNiagaraStackEntry::FStackIssue InconsistentEnabledError;
+			InconsistentEnabledError.ShortDescription = LOCTEXT("InconsistentEnabledErrorSummary", "The enabled state for module is inconsistent.");
+			InconsistentEnabledError.LongDescription = LOCTEXT("InconsistentEnabledError", "This module is using multiple functions and their enabled state is inconsistent.\nClick fix to make all of the functions for this module enabled.");
+			InconsistentEnabledError.UniqueIdentifier = FName(*FString::Printf(TEXT("InconsistentEnabled-%s"), *GetStackEditorDataKey()));
+			UNiagaraStackEntry::FStackIssueFix Fix;
+			Fix.Description = LOCTEXT("EnableModule", "Enable module");
+			Fix.FixDelegate.BindLambda([=]()
+			{
+				FScopedTransaction ScopedTransaction(Fix.Description);
+				SetIsEnabled(true);
+			});
+			InconsistentEnabledError.Fixes.Add(Fix);
+			NewIssues.Add(InconsistentEnabledError);
+		}
 	}
-	else
+	/*// Generate dependency errors with their fixes
+	TArray<UNiagaraNodeFunctionCall*> FoundCalls;
+	TArray<FNiagaraStackGraphUtilities::FStackNodeGroup> StackNodeGroups;
+	FNiagaraStackGraphUtilities::GetStackNodeGroups(*FunctionCallNode, StackNodeGroups);
+	TArray<FNiagaraModuleDependency> DependenciesNeeded;
+	UNiagaraNodeOutput* OutputNode = GetOutputNode();
+	for (FNiagaraModuleDependency Dependency: FunctionCallNode->FunctionScript->RequiredDependencies)
 	{
-		bIsEnabled = false;
-		FError InconsistentEnabledError;
-		InconsistentEnabledError.ErrorSummaryText = LOCTEXT("InconsistentEnabledErrorSummary", "The endabled state for module is inconsistent.");
-		InconsistentEnabledError.ErrorText = LOCTEXT("InconsistentEnabledError", "This module is using multiple functions and their enabled state is inconsistent.\nClick fix to make all of the functions for this module enabled.");
-		InconsistentEnabledError.Fix = FFixError::CreateLambda([=]()
+		if (Dependency.Id == NAME_None)
 		{
-			SetIsEnabled(true);
-			return true;
-		});
-		Errors.Add(InconsistentEnabledError);
-	}
-}
-
-int32 UNiagaraStackModuleItem::GetErrorCount() const 
-{
-	return Errors.Num();
-}
-
-bool UNiagaraStackModuleItem::GetErrorFixable(int32 ErrorIdx) const
-{
-	return Errors[ErrorIdx].Fix.IsBound();
-}
-
-bool UNiagaraStackModuleItem::TryFixError(int32 ErrorIdx)
-{
-	return Errors[ErrorIdx].Fix.IsBound() ? Errors[ErrorIdx].Fix.Execute() : false;
-}
-
-FText UNiagaraStackModuleItem::GetErrorText(int32 ErrorIdx) const
-{
-	return Errors[ErrorIdx].ErrorText;
-}
-
-FText UNiagaraStackModuleItem::GetErrorSummaryText(int32 ErrorIdx) const 
-{
-	return Errors[ErrorIdx].ErrorSummaryText;
+			continue;
+		}
+		bool bDependencyMet = false;
+		UNiagaraNodeFunctionCall* FunctionNode = nullptr;
+		TArray <UNiagaraNodeFunctionCall*> DisabledDependencies;
+		for (int i = 1; i < StackNodeGroups.Num() - 1; i++)
+		{
+			FNiagaraStackGraphUtilities::FStackNodeGroup& StackNodeGroup = StackNodeGroups[i];
+			FunctionNode = Cast<UNiagaraNodeFunctionCall> (StackNodeGroup.EndNode);
+			if (FunctionNode != nullptr && FunctionNode->FunctionScript->ProvidedDependencies.Contains(Dependency.Id))
+			// treat disabled dependencies separately because they are easier to fix
+			{
+				if (FunctionNode->IsNodeEnabled() == false)
+				{
+					DisabledDependencies.Add(FunctionNode);
+				}
+				else
+				{
+					bDependencyMet = true;
+					break;
+				}
+			}
+		}
+		if (bDependencyMet == false)
+		{
+			DependenciesNeeded.Add(Dependency);
+			UNiagaraStackEntry::FStackIssue Error;
+			Error.ShortDescription = LOCTEXT("DependencyWarning", "The module has unmet dependencies.");;
+			Error.LongDescription = FText::Format(LOCTEXT("DependencyWarningLong", "The following dependency is not met: {0}"), Dependency.Description);
+			Error.bCanBeDismissed = true;
+			Error.UniqueIdentifier = FName(*FString::Printf(TEXT("%s-dependency-%s"), *GetStackEditorDataKey(), *Dependency.Id.ToString()));
+
+			for ( UNiagaraNodeFunctionCall* DisabledNode : DisabledDependencies) // module exists but disabled
+			{
+				UNiagaraStackEntry::FStackIssueFix Fix;
+				Fix.Description = FText::Format(LOCTEXT("EnableDependency", "Enable dependency module {0}"), FText::FromString(DisabledNode->GetFunctionName()));
+				Fix.FixDelegate.BindLambda([=]()
+				{
+					FScopedTransaction ScopedTransaction(LOCTEXT("EnableDependencyModule", "Enable dependency module"));
+					DisabledNode->Modify();
+					FNiagaraStackGraphUtilities::SetModuleIsEnabled(*DisabledNode, true);
+				});
+				Error.Fixes.Add(Fix);
+			}
+
+			TArray<FAssetData> ModuleAssets;
+			FNiagaraStackGraphUtilities::GetScriptAssetsByDependencyProvided(ENiagaraScriptUsage::Module, Dependency.Id, ModuleAssets);
+			for (FAssetData ModuleAsset : ModuleAssets) 
+			{
+				auto DisabledAsset = DisabledDependencies.FindByPredicate([=](UNiagaraNodeFunctionCall* DisabledNode) {return DisabledNode->GetReferencedAsset() == ModuleAsset.GetAsset(); });
+				if (DisabledAsset != nullptr)
+				{
+					continue;
+				}
+				UNiagaraStackEntry::FStackIssueFix Fix;
+				Fix.Description = LOCTEXT("InsertNewModule", "Insert new module");
+				int32 InsertIndex = GetModuleIndex();
+				Fix.FixDelegate.BindLambda([=]()
+				{
+					FScopedTransaction ScopedTransaction(Fix.Description);
+					UNiagaraNodeFunctionCall* NewModuleNode = nullptr;
+					int32 TargetIndex = Dependency.Type == ENiagaraModuleDependencyType::PostDependency ? GetModuleIndex() + 1 : GetModuleIndex();
+					NewModuleNode = FNiagaraStackGraphUtilities::AddScriptModuleToStack(ModuleAsset, *OutputNode, TargetIndex);
+					checkf(NewModuleNode != nullptr, TEXT("Add module action failed"));
+					FNiagaraStackGraphUtilities::InitializeStackFunctionInputs(GetSystemViewModel(), GetEmitterViewModel(), GetStackEditorData(), *NewModuleNode, *NewModuleNode);
+					FNiagaraStackGraphUtilities::RelayoutGraph(*OutputNode->GetGraph());
+				});
+				Fix.Description = FText::Format(LOCTEXT("AddDependency", "Add dependency module {0}"), FText::FromName(ModuleAsset.AssetName));
+				Error.Fixes.Add(Fix);
+			}
+			
+			NewIssues.Add(Error);
+		}
+	}*/
 }
 
-void UNiagaraStackModuleItem::InputPinnedChanged()
+bool UNiagaraStackModuleItem::FilterOutputCollection(const UNiagaraStackEntry& Child) const
 {
-	PinnedInputCollection->RefreshChildren();
-	UnpinnedInputCollection->RefreshChildren();
+	if (Child.IsA<UNiagaraStackModuleItemOutputCollection>() && GetStackEditorData().GetShowOutputs() == false)
+	{
+		return false;
+	}
+	return true;
 }
 
-void UNiagaraStackModuleItem::ModuleExpandedChanged()
+void UNiagaraStackModuleItem::RefreshIsEnabled()
 {
-	OnStructureChanged().Broadcast();
+	TOptional<bool> IsEnabled = FNiagaraStackGraphUtilities::GetModuleIsEnabled(*FunctionCallNode);
+	if (IsEnabled.IsSet())
+	{
+		bIsEnabled = IsEnabled.GetValue();
+	}
 }
 
 bool UNiagaraStackModuleItem::CanMoveAndDelete() const
@@ -233,16 +303,19 @@ bool UNiagaraStackModuleItem::CanMoveAndDelete() const
 
 bool UNiagaraStackModuleItem::CanRefresh() const
 {
-	return FunctionCallNode != nullptr ? FunctionCallNode->CanRefreshFromModuleItem() : false;
+	return bCanRefresh;
 }
 
 void UNiagaraStackModuleItem::Refresh()
 {
 	if (CanRefresh())
 	{
-		FunctionCallNode->RefreshFromExternalChanges();
-		FunctionCallNode->GetNiagaraGraph()->NotifyGraphNeedsRecompile();
-		GetSystemViewModel()->ResetSystem();
+		if (FunctionCallNode->RefreshFromExternalChanges())
+		{
+			FunctionCallNode->GetNiagaraGraph()->NotifyGraphNeedsRecompile();
+			GetSystemViewModel()->ResetSystem();
+		}
+		RefreshChildren();
 	}
 }
 
@@ -257,62 +330,70 @@ void UNiagaraStackModuleItem::SetIsEnabled(bool bInIsEnabled)
 	bIsEnabled = bInIsEnabled;
 }
 
-void UNiagaraStackModuleItem::MoveUp()
+void UNiagaraStackModuleItem::Delete()
 {
-	checkf(CanMoveAndDelete(), TEXT("This module can't be moved or deleted"));
-
-	TArray<FNiagaraStackGraphUtilities::FStackNodeGroup> StackNodeGroups;
-	FNiagaraStackGraphUtilities::GetStackNodeGroups(*FunctionCallNode, StackNodeGroups);
-	int32 ModuleStackIndex = StackNodeGroups.IndexOfByPredicate([&](const FNiagaraStackGraphUtilities::FStackNodeGroup& StackNodeGroup) { return StackNodeGroup.EndNode == FunctionCallNode; });
+	checkf(CanMoveAndDelete(), TEXT("This module can't be deleted"));
 
-	if (ModuleStackIndex > 1)
+	FScopedTransaction ScopedTransaction(LOCTEXT("RemoveAModuleFromTheStack", "Remove a module from the stack"));
+	TArray<TWeakObjectPtr<UNiagaraNodeInput>> RemovedNodes;
+	bool bRemoved = FNiagaraStackGraphUtilities::RemoveModuleFromStack(*FunctionCallNode, RemovedNodes);
+	if (bRemoved)
 	{
-		FScopedTransaction ScopedTransaction(LOCTEXT("MoveModuleUpTheStackTransaction", "Move module up the stack"));
-
-		FNiagaraStackGraphUtilities::DisconnectStackNodeGroup(StackNodeGroups[ModuleStackIndex], StackNodeGroups[ModuleStackIndex - 1], StackNodeGroups[ModuleStackIndex + 1]);
-		FNiagaraStackGraphUtilities::ConnectStackNodeGroup(StackNodeGroups[ModuleStackIndex], StackNodeGroups[ModuleStackIndex - 2], StackNodeGroups[ModuleStackIndex - 1]);
-
-		FunctionCallNode->GetNiagaraGraph()->NotifyGraphNeedsRecompile();
+		UNiagaraGraph* Graph = FunctionCallNode->GetNiagaraGraph();
+		Graph->NotifyGraphNeedsRecompile();
 		FNiagaraStackGraphUtilities::RelayoutGraph(*FunctionCallNode->GetGraph());
 		ModifiedGroupItemsDelegate.ExecuteIfBound();
+		for (auto InputNode: RemovedNodes)
+		{
+			if (InputNode != nullptr && InputNode->Usage == ENiagaraInputNodeUsage::Parameter)
+			{
+				GetSystemViewModel()->NotifyDataObjectChanged(InputNode->GetDataInterface());
+			}
+		}
 	}
 }
 
-void UNiagaraStackModuleItem::MoveDown()
+int32 UNiagaraStackModuleItem::GetModuleIndex()
 {
-	checkf(CanMoveAndDelete(), TEXT("This module can't be moved or deleted"));
-
-	TArray<FNiagaraStackGraphUtilities::FStackNodeGroup> StackNodeGroups;
-	FNiagaraStackGraphUtilities::GetStackNodeGroups(*FunctionCallNode, StackNodeGroups);
-	int32 ModuleStackIndex = StackNodeGroups.IndexOfByPredicate([&](const FNiagaraStackGraphUtilities::FStackNodeGroup& StackNodeGroup) { return StackNodeGroup.EndNode == FunctionCallNode; });
-
-	if (ModuleStackIndex < StackNodeGroups.Num() - 2)
+	TArray<FNiagaraStackGraphUtilities::FStackNodeGroup> StackGroups;
+	FNiagaraStackGraphUtilities::GetStackNodeGroups(*FunctionCallNode, StackGroups);
+	int32 ModuleIndex = 0;
+	for (FNiagaraStackGraphUtilities::FStackNodeGroup& StackGroup : StackGroups)
 	{
-		FScopedTransaction ScopedTransaction(LOCTEXT("MoveModuleDownTheStackTransaction", "Move module down the stack"));
-
-		FNiagaraStackGraphUtilities::DisconnectStackNodeGroup(StackNodeGroups[ModuleStackIndex], StackNodeGroups[ModuleStackIndex - 1], StackNodeGroups[ModuleStackIndex + 1]);
-		FNiagaraStackGraphUtilities::ConnectStackNodeGroup(StackNodeGroups[ModuleStackIndex], StackNodeGroups[ModuleStackIndex + 1], StackNodeGroups[ModuleStackIndex + 2]);
-
-		FunctionCallNode->GetNiagaraGraph()->NotifyGraphNeedsRecompile();
-		FNiagaraStackGraphUtilities::RelayoutGraph(*FunctionCallNode->GetGraph());
-		ModifiedGroupItemsDelegate.ExecuteIfBound();
+		if (StackGroup.EndNode == FunctionCallNode)
+		{
+			return ModuleIndex;
+		}
+		if (StackGroup.EndNode->IsA<UNiagaraNodeFunctionCall>())
+		{
+			ModuleIndex++;
+		}
 	}
+	return INDEX_NONE;
 }
 
-void UNiagaraStackModuleItem::Delete()
+UNiagaraNodeOutput* UNiagaraStackModuleItem::GetOutputNode() const
 {
-	checkf(CanMoveAndDelete(), TEXT("This module can't be moved or deleted"));
+	return FNiagaraStackGraphUtilities::GetEmitterOutputNodeForStackNode(*FunctionCallNode);
+}
 
-	FScopedTransaction ScopedTransaction(LOCTEXT("RemoveAModuleFromTheStack", "Remove a module from the stack"));
+void UNiagaraStackModuleItem::NotifyModuleMoved()
+{
+	ModifiedGroupItemsDelegate.ExecuteIfBound();
+}
 
-	bool bRemoved = FNiagaraStackGraphUtilities::RemoveModuleFromStack(*FunctionCallNode);
-	if (bRemoved)
+UObject* UNiagaraStackModuleItem::GetExternalAsset() const
+{
+	if (GetModuleNode().FunctionScript != nullptr && GetModuleNode().FunctionScript->IsAsset())
 	{
-		UNiagaraGraph* Graph = FunctionCallNode->GetNiagaraGraph();
-		Graph->NotifyGraphNeedsRecompile();
-		FNiagaraStackGraphUtilities::RelayoutGraph(*FunctionCallNode->GetGraph());
-		ModifiedGroupItemsDelegate.ExecuteIfBound();
+		return GetModuleNode().FunctionScript;
 	}
+	return nullptr;
+}
+
+bool UNiagaraStackModuleItem::CanDrag() const
+{
+	return true;
 }
 
 #undef LOCTEXT_NAMESPACE
diff --git a/Engine/Plugins/FX/Niagara/Source/NiagaraEditor/Private/ViewModels/Stack/NiagaraStackModuleItemOutput.cpp b/Engine/Plugins/FX/Niagara/Source/NiagaraEditor/Private/ViewModels/Stack/NiagaraStackModuleItemOutput.cpp
--- a/Engine/Plugins/FX/Niagara/Source/NiagaraEditor/Private/ViewModels/Stack/NiagaraStackModuleItemOutput.cpp
+++ b/Engine/Plugins/FX/Niagara/Source/NiagaraEditor/Private/ViewModels/Stack/NiagaraStackModuleItemOutput.cpp
@@ -15,16 +15,12 @@ UNiagaraStackModuleItemOutput::UNiagaraStackModuleItemOutput()
 {
 }
 
-int32 UNiagaraStackModuleItemOutput::GetItemIndentLevel() const
-{
-	return 1;
-}
-
-void UNiagaraStackModuleItemOutput::Initialize(TSharedRef<FNiagaraSystemViewModel> InSystemViewModel, TSharedRef<FNiagaraEmitterViewModel> InEmitterViewModel, UNiagaraNodeFunctionCall& InFunctionCallNode, FName InOutputParameterHandle,
+void UNiagaraStackModuleItemOutput::Initialize(FRequiredEntryData InRequiredEntryData, UNiagaraNodeFunctionCall& InFunctionCallNode, FName InOutputParameterHandle,
 	FNiagaraTypeDefinition InOutputType)
 {
 	checkf(FunctionCallNode.Get() == nullptr, TEXT("Can only set the Output once."));
-	Super::Initialize(InSystemViewModel, InEmitterViewModel);
+	FString OutputStackEditorDataKey = FString::Printf(TEXT("%s-Output-%s"), *InFunctionCallNode.NodeGuid.ToString(EGuidFormats::DigitsWithHyphens), *InOutputParameterHandle.ToString());
+	Super::Initialize(InRequiredEntryData, OutputStackEditorDataKey);
 	FunctionCallNode = &InFunctionCallNode;
 	OutputType = InOutputType;
 	OutputParameterHandle = FNiagaraParameterHandle(InOutputParameterHandle);
@@ -36,16 +32,6 @@ FText UNiagaraStackModuleItemOutput::GetDisplayName() const
 	return DisplayName;
 }
 
-FName UNiagaraStackModuleItemOutput::GetTextStyleName() const
-{
-	return "NiagaraEditor.Stack.ParameterText";
-}
-
-bool UNiagaraStackModuleItemOutput::GetCanExpand() const
-{
-	return true;
-}
-
 FText UNiagaraStackModuleItemOutput::GetTooltipText() const
 {
 	FNiagaraVariable ValueVariable(OutputType, OutputParameterHandle.GetParameterHandleString());
@@ -71,6 +57,16 @@ FText UNiagaraStackModuleItemOutput::GetTooltipText() const
 	return FText::FromName(ValueVariable.GetName());
 }
 
+bool UNiagaraStackModuleItemOutput::GetIsEnabled() const
+{
+	return FunctionCallNode->GetDesiredEnabledState() == ENodeEnabledState::Enabled;
+}
+
+UNiagaraStackEntry::EStackRowStyle UNiagaraStackModuleItemOutput::GetStackRowStyle() const
+{
+	return EStackRowStyle::ItemContent;
+}
+
 const FNiagaraParameterHandle& UNiagaraStackModuleItemOutput::GetOutputParameterHandle() const
 {
 	return OutputParameterHandle;
diff --git a/Engine/Plugins/FX/Niagara/Source/NiagaraEditor/Private/ViewModels/Stack/NiagaraStackModuleItemOutputCollection.cpp b/Engine/Plugins/FX/Niagara/Source/NiagaraEditor/Private/ViewModels/Stack/NiagaraStackModuleItemOutputCollection.cpp
--- a/Engine/Plugins/FX/Niagara/Source/NiagaraEditor/Private/ViewModels/Stack/NiagaraStackModuleItemOutputCollection.cpp
+++ b/Engine/Plugins/FX/Niagara/Source/NiagaraEditor/Private/ViewModels/Stack/NiagaraStackModuleItemOutputCollection.cpp
@@ -14,39 +14,35 @@ UNiagaraStackModuleItemOutputCollection::UNiagaraStackModuleItemOutputCollection
 {
 }
 
-void UNiagaraStackModuleItemOutputCollection::Initialize(TSharedRef<FNiagaraSystemViewModel> InSystemViewModel, TSharedRef<FNiagaraEmitterViewModel> InEmitterViewModel, UNiagaraNodeFunctionCall& InFunctionCallNode)
+void UNiagaraStackModuleItemOutputCollection::Initialize(FRequiredEntryData InRequiredEntryData, UNiagaraNodeFunctionCall& InFunctionCallNode)
 {
 	checkf(FunctionCallNode == nullptr, TEXT("Can not set the node more than once."));
-	Super::Initialize(InSystemViewModel, InEmitterViewModel);
+	FString OutputCollectionStackEditorDataKey = FString::Printf(TEXT("%s-Outputs"), *InFunctionCallNode.NodeGuid.ToString(EGuidFormats::DigitsWithHyphens));
+	Super::Initialize(InRequiredEntryData, OutputCollectionStackEditorDataKey);
 	FunctionCallNode = &InFunctionCallNode;
 }
 
 FText UNiagaraStackModuleItemOutputCollection::GetDisplayName() const
 {
-	return DisplayName;
+	return NSLOCTEXT("StackModuleItemOutputCollection", "OutputsLabel", "Outputs");
 }
 
-FName UNiagaraStackModuleItemOutputCollection::GetTextStyleName() const
-{
-	return "NiagaraEditor.Stack.ParameterCollectionText";
-}
-
-void UNiagaraStackModuleItemOutputCollection::SetDisplayName(FText InDisplayName)
+bool UNiagaraStackModuleItemOutputCollection::IsExpandedByDefault() const
 {
-	DisplayName = InDisplayName;
+	return false;
 }
 
-bool UNiagaraStackModuleItemOutputCollection::GetCanExpand() const
+bool UNiagaraStackModuleItemOutputCollection::GetIsEnabled() const
 {
-	return true;
+	return FunctionCallNode->GetDesiredEnabledState() == ENodeEnabledState::Enabled;
 }
 
-bool UNiagaraStackModuleItemOutputCollection::IsExpandedByDefault() const
+UNiagaraStackEntry::EStackRowStyle UNiagaraStackModuleItemOutputCollection::GetStackRowStyle() const
 {
-	return false;
+	return EStackRowStyle::ItemContent;
 }
 
-void UNiagaraStackModuleItemOutputCollection::RefreshChildrenInternal(const TArray<UNiagaraStackEntry*>& CurrentChildren, TArray<UNiagaraStackEntry*>& NewChildren)
+void UNiagaraStackModuleItemOutputCollection::RefreshChildrenInternal(const TArray<UNiagaraStackEntry*>& CurrentChildren, TArray<UNiagaraStackEntry*>& NewChildren, TArray<FStackIssue>& NewIssues)
 {
 	UEdGraphPin* OutputParameterMapPin = FNiagaraStackGraphUtilities::GetParameterMapOutputPin(*FunctionCallNode);
 	if (ensureMsgf(OutputParameterMapPin != nullptr, TEXT("Invalid Stack Graph - Function call node has no output pin.")))
@@ -66,20 +62,13 @@ void UNiagaraStackModuleItemOutputCollection::RefreshChildrenInternal(const TArr
 				{
 					if (Cast<UNiagaraNodeParameterMapSet>(WritePin->GetOwningNode()) != nullptr)
 					{
-						UNiagaraStackModuleItemOutput* Output = nullptr;
-						for (UNiagaraStackEntry* CurrentChild : CurrentChildren)
-						{
-							UNiagaraStackModuleItemOutput* ChildOutput = CastChecked<UNiagaraStackModuleItemOutput>(CurrentChild);
-						if (ChildOutput->GetOutputParameterHandle().GetParameterHandleString() == Variable.GetName())
-							{
-								Output = ChildOutput;
-								break;
-							}
-						}
+						UNiagaraStackModuleItemOutput* Output = FindCurrentChildOfTypeByPredicate<UNiagaraStackModuleItemOutput>(CurrentChildren,
+							[&](UNiagaraStackModuleItemOutput* CurrentOutput) { return CurrentOutput->GetOutputParameterHandle().GetParameterHandleString() == Variable.GetName(); });
+
 						if (Output == nullptr)
 						{
 							Output = NewObject<UNiagaraStackModuleItemOutput>(this);
-							Output->Initialize(GetSystemViewModel(), GetEmitterViewModel(), *FunctionCallNode, Variable.GetName(), Variable.GetType());
+							Output->Initialize(CreateDefaultChildRequiredData(), *FunctionCallNode, Variable.GetName(), Variable.GetType());
 						}
 
 						NewChildren.Add(Output);
diff --git a/Engine/Plugins/FX/Niagara/Source/NiagaraEditor/Private/ViewModels/Stack/NiagaraStackObject.cpp b/Engine/Plugins/FX/Niagara/Source/NiagaraEditor/Private/ViewModels/Stack/NiagaraStackObject.cpp
--- a/Engine/Plugins/FX/Niagara/Source/NiagaraEditor/Private/ViewModels/Stack/NiagaraStackObject.cpp
+++ b/Engine/Plugins/FX/Niagara/Source/NiagaraEditor/Private/ViewModels/Stack/NiagaraStackObject.cpp
@@ -1,18 +1,43 @@
 // Copyright 1998-2018 Epic Games, Inc. All Rights Reserved.
 
 #include "ViewModels/Stack/NiagaraStackObject.h"
+#include "ViewModels/Stack/NiagaraStackPropertyRow.h"
+#include "NiagaraNode.h"
 
+#include "Modules/ModuleManager.h"
+#include "IPropertyRowGenerator.h"
+#include "PropertyEditorModule.h"
+#include "IDetailTreeNode.h"
 
 UNiagaraStackObject::UNiagaraStackObject()
 	: Object(nullptr)
 {
 }
 
-void UNiagaraStackObject::Initialize(TSharedRef<FNiagaraSystemViewModel> InSystemViewModel, TSharedRef<FNiagaraEmitterViewModel> InEmitterViewModel, UObject* InObject)
+void UNiagaraStackObject::Initialize(FRequiredEntryData InRequiredEntryData, UObject* InObject, FString InOwnerStackItemEditorDataKey, UNiagaraNode* InOwningNiagaraNode)
 {
 	checkf(Object == nullptr, TEXT("Can only initialize once."));
-	Super::Initialize(InSystemViewModel, InEmitterViewModel);
+	FString ObjectStackEditorDataKey = FString::Printf(TEXT("%s-%s"), *InOwnerStackItemEditorDataKey, *InObject->GetName());
+	Super::Initialize(InRequiredEntryData, false, InOwnerStackItemEditorDataKey, ObjectStackEditorDataKey);
 	Object = InObject;
+	OwningNiagaraNode = InOwningNiagaraNode;
+}
+
+void UNiagaraStackObject::SetOnSelectRootNodes(FOnSelectRootNodes OnSelectRootNodes)
+{
+	OnSelectRootNodesDelegate = OnSelectRootNodes;
+}
+
+void UNiagaraStackObject::RegisterInstancedCustomPropertyLayout(UStruct* Class, FOnGetDetailCustomizationInstance DetailLayoutDelegate)
+{
+	checkf(PropertyRowGenerator.IsValid() == false, TEXT("Can not add additional customizations after children have been refreshed."));
+	RegisteredClassCustomizations.Add({ Class, DetailLayoutDelegate });
+}
+
+void UNiagaraStackObject::RegisterInstancedCustomPropertyTypeLayout(FName PropertyTypeName, FOnGetPropertyTypeCustomizationInstance PropertyTypeLayoutDelegate, TSharedPtr<IPropertyTypeIdentifier> Identifier)
+{
+	checkf(PropertyRowGenerator.IsValid() == false, TEXT("Can not add additional customizations after children have been refreshed."));
+	RegisteredPropertyCustomizations.Add({ PropertyTypeName, PropertyTypeLayoutDelegate, Identifier });
 }
 
 UObject* UNiagaraStackObject::GetObject()
@@ -20,17 +45,84 @@ UObject* UNiagaraStackObject::GetObject()
 	return Object;
 }
 
-int32 UNiagaraStackObject::GetItemIndentLevel() const
+void UNiagaraStackObject::NotifyPostChange(const FPropertyChangedEvent& PropertyChangedEvent, UProperty* PropertyThatChanged)
 {
-	return ItemIndentLevel;
+	OnDataObjectModified().Broadcast(Object);
 }
 
-void UNiagaraStackObject::SetItemIndentLevel(int32 InItemIndentLevel)
+bool UNiagaraStackObject::GetIsEnabled() const
 {
-	ItemIndentLevel = InItemIndentLevel;
+	return OwningNiagaraNode == nullptr || OwningNiagaraNode->GetDesiredEnabledState() == ENodeEnabledState::Enabled;
 }
 
-void UNiagaraStackObject::NotifyObjectChanged()
+bool UNiagaraStackObject::GetShouldShowInStack() const
 {
-	OnDataObjectModified().Broadcast(Object);
+	return false;
+}
+
+void UNiagaraStackObject::RefreshChildrenInternal(const TArray<UNiagaraStackEntry*>& CurrentChildren, TArray<UNiagaraStackEntry*>& NewChildren, TArray<FStackIssue>& NewIssues)
+{
+	if (PropertyRowGenerator.IsValid() == false)
+	{
+		FPropertyEditorModule& PropertyEditorModule = FModuleManager::LoadModuleChecked<FPropertyEditorModule>("PropertyEditor");
+		FPropertyRowGeneratorArgs Args;
+		Args.NotifyHook = this;
+		PropertyRowGenerator = PropertyEditorModule.CreatePropertyRowGenerator(Args);
+
+		for (FRegisteredClassCustomization& RegisteredClassCustomization : RegisteredClassCustomizations)
+		{
+			PropertyRowGenerator->RegisterInstancedCustomPropertyLayout(RegisteredClassCustomization.Class, RegisteredClassCustomization.DetailLayoutDelegate);
+		}
+
+		for (FRegisteredPropertyCustomization& RegisteredPropertyCustomization : RegisteredPropertyCustomizations)
+		{
+			PropertyRowGenerator->RegisterInstancedCustomPropertyTypeLayout(RegisteredPropertyCustomization.PropertyTypeName,
+				RegisteredPropertyCustomization.PropertyTypeLayoutDelegate, RegisteredPropertyCustomization.Identifier);
+		}
+
+		TArray<UObject*> Objects;
+		Objects.Add(Object);
+		PropertyRowGenerator->SetObjects(Objects);
+
+		// Add the refresh delegate after setting the objects to prevent refreshing children immediately.
+		PropertyRowGenerator->OnRowsRefreshed().AddUObject(this, &UNiagaraStackObject::PropertyRowsRefreshed);
+	}
+
+	TArray<TSharedRef<IDetailTreeNode>> DefaultRootTreeNodes = PropertyRowGenerator->GetRootTreeNodes();
+	TArray<TSharedRef<IDetailTreeNode>> RootTreeNodes;
+	if (OnSelectRootNodesDelegate.IsBound())
+	{
+		OnSelectRootNodesDelegate.Execute(DefaultRootTreeNodes, &RootTreeNodes);
+	}
+	else
+	{
+		RootTreeNodes = DefaultRootTreeNodes;
+	}
+
+	for (TSharedRef<IDetailTreeNode> RootTreeNode : RootTreeNodes)
+	{
+		if (RootTreeNode->GetNodeType() == EDetailNodeType::Advanced)
+		{
+			continue;
+		}
+
+		UNiagaraStackPropertyRow* ChildRow = FindCurrentChildOfTypeByPredicate<UNiagaraStackPropertyRow>(CurrentChildren,
+			[=](UNiagaraStackPropertyRow* CurrentChild) { return CurrentChild->GetDetailTreeNode() == RootTreeNode; });
+
+		if (ChildRow == nullptr)
+		{
+			ChildRow = NewObject<UNiagaraStackPropertyRow>(this);
+			ChildRow->Initialize(CreateDefaultChildRequiredData(), RootTreeNode, GetOwnerStackItemEditorDataKey(), GetOwnerStackItemEditorDataKey(), OwningNiagaraNode);
+		}
+
+		NewChildren.Add(ChildRow);
+	}
+}
+
+void UNiagaraStackObject::PropertyRowsRefreshed()
+{
+	if (IsValid())
+	{
+		RefreshChildren();
+	}
 }
diff --git a/Engine/Plugins/FX/Niagara/Source/NiagaraEditor/Private/ViewModels/Stack/NiagaraStackParameterStoreEntry.cpp b/Engine/Plugins/FX/Niagara/Source/NiagaraEditor/Private/ViewModels/Stack/NiagaraStackParameterStoreEntry.cpp
--- a/Engine/Plugins/FX/Niagara/Source/NiagaraEditor/Private/ViewModels/Stack/NiagaraStackParameterStoreEntry.cpp
+++ b/Engine/Plugins/FX/Niagara/Source/NiagaraEditor/Private/ViewModels/Stack/NiagaraStackParameterStoreEntry.cpp
@@ -5,11 +5,16 @@
 #include "NiagaraGraph.h"
 #include "EdGraphSchema_Niagara.h"
 #include "NiagaraEditorUtilities.h"
-#include "NiagaraSystemViewModel.h"
+#include "ViewModels/NiagaraSystemViewModel.h"
+#include "NiagaraSystem.h"
 #include "NiagaraSystemScriptViewModel.h"
-#include "NiagaraEmitterViewModel.h"
+#include "ViewModels/NiagaraEmitterViewModel.h"
+#include "NiagaraEmitterHandle.h"
 #include "NiagaraScriptGraphViewModel.h"
 #include "ViewModels/Stack/NiagaraStackObject.h"
+#include "NiagaraNodeParameterMapGet.h"
+#include "ViewModels/Stack/NiagaraStackGraphUtilities.h"
+#include "NiagaraStackEditorData.h"
 
 #include "ScopedTransaction.h"
 #include "Editor.h"
@@ -19,10 +24,10 @@
 #include "ARFilter.h"
 #include "EdGraph/EdGraphPin.h"
 
+
 #define LOCTEXT_NAMESPACE "UNiagaraStackParameterStoreEntry"
 UNiagaraStackParameterStoreEntry::UNiagaraStackParameterStoreEntry()
-	: ItemIndentLevel(0)
-	, ValueObjectEntry(nullptr)
+	: ValueObjectEntry(nullptr)
 {
 }
 
@@ -30,27 +35,18 @@ void UNiagaraStackParameterStoreEntry::BeginDestroy()
 {
 	Super::BeginDestroy();
 }
-int32 UNiagaraStackParameterStoreEntry::GetItemIndentLevel() const
-{
-	return ItemIndentLevel;
-}
-
-void UNiagaraStackParameterStoreEntry::SetItemIndentLevel(int32 InItemIndentLevel)
-{
-	ItemIndentLevel = InItemIndentLevel;
-}
 
 void UNiagaraStackParameterStoreEntry::Initialize(
-	TSharedRef<FNiagaraSystemViewModel> InSystemViewModel,
-	TSharedRef<FNiagaraEmitterViewModel> InEmitterViewModel,
-	UNiagaraStackEditorData& InStackEditorData,
+	FRequiredEntryData InRequiredEntryData,
 	UObject* InOwner,
 	FNiagaraParameterStore* InParameterStore,
 	FString InInputParameterHandle,
-	FNiagaraTypeDefinition InInputType)
+	FNiagaraTypeDefinition InInputType,
+	FString InOwnerStackItemEditorDataKey)
 {
-	Super::Initialize(InSystemViewModel, InEmitterViewModel);
-	StackEditorData = &InStackEditorData;
+	bool bParameterIsAdvanced = false;
+	FString ParameterStackEditorDataKey = FString::Printf(TEXT("Parameter-%s"), *InInputParameterHandle);
+	Super::Initialize(InRequiredEntryData, bParameterIsAdvanced, InOwnerStackItemEditorDataKey, ParameterStackEditorDataKey);
 	DisplayName = FText::FromString(InInputParameterHandle);
 	ParameterName = *InInputParameterHandle;
 	InputType = InInputType;
@@ -63,18 +59,16 @@ const FNiagaraTypeDefinition& UNiagaraStackParameterStoreEntry::GetInputType() c
 	return InputType;
 }
 
-void UNiagaraStackParameterStoreEntry::RefreshChildrenInternal(const TArray<UNiagaraStackEntry*>& CurrentChildren, TArray<UNiagaraStackEntry*>& NewChildren)
+void UNiagaraStackParameterStoreEntry::RefreshChildrenInternal(const TArray<UNiagaraStackEntry*>& CurrentChildren, TArray<UNiagaraStackEntry*>& NewChildren, TArray<FStackIssue>& NewIssues)
 {
 	RefreshValueAndHandle();
 
 	if (ValueObject != nullptr)
 	{
-		if(ValueObjectEntry == nullptr ||
-			ValueObjectEntry->GetObject() != ValueObject)
+		if(ValueObjectEntry == nullptr || ValueObjectEntry->GetObject() != ValueObject)
 		{
 			ValueObjectEntry = NewObject<UNiagaraStackObject>(this);
-			ValueObjectEntry->Initialize(GetSystemViewModel(), GetEmitterViewModel(), ValueObject);
-			ValueObjectEntry->SetItemIndentLevel(ItemIndentLevel + 1);
+			ValueObjectEntry->Initialize(CreateDefaultChildRequiredData(), ValueObject, GetOwnerStackItemEditorDataKey());
 		}
 		NewChildren.Add(ValueObjectEntry);
 	}
@@ -110,16 +104,6 @@ FText UNiagaraStackParameterStoreEntry::GetDisplayName() const
 	return DisplayName;
 }
 
-FName UNiagaraStackParameterStoreEntry::GetTextStyleName() const
-{
-	return "NiagaraEditor.Stack.ParameterText";
-}
-
-bool UNiagaraStackParameterStoreEntry::GetCanExpand() const
-{
-	return true;
-}
-
 TSharedPtr<FStructOnScope> UNiagaraStackParameterStoreEntry::GetValueStruct()
 {
 	return LocalValueStruct;
@@ -132,7 +116,7 @@ UNiagaraDataInterface* UNiagaraStackParameterStoreEntry::GetValueObject()
 
 void UNiagaraStackParameterStoreEntry::NotifyBeginValueChange()
 {
-	GEditor->BeginTransaction(LOCTEXT("BeginEditModuleInputValue", "Edit module input value."));
+	GEditor->BeginTransaction(LOCTEXT("ModifyInputValue", "Modify input value."));
 	Owner->Modify();
 }
 
@@ -155,7 +139,6 @@ void UNiagaraStackParameterStoreEntry::NotifyValueChanged()
 	{
 		FNiagaraVariable DefaultVariable(InputType, ParameterName);
 		ParameterStore->SetParameterData(LocalValueStruct->GetStructMemory(), DefaultVariable);
-		GetSystemViewModel()->ResetSystem();
 	}
 }
 
@@ -166,69 +149,189 @@ bool UNiagaraStackParameterStoreEntry::CanReset() const
 
 void UNiagaraStackParameterStoreEntry::Reset()
 {
+	NotifyBeginValueChange();
+	FNiagaraVariable Var (InputType, ParameterName);
 	if (InputType.GetClass() == nullptr)
-	{		
-		// For struct inputs the override pin and anything attached to it should be removed.
-		FScopedTransaction ScopedTransaction(LOCTEXT("ResetInputStructTransaction", "Reset the inputs value to default."));
+	{
+		FNiagaraEditorUtilities::ResetVariableToDefaultValue(Var);
+		Var.CopyTo(LocalValueStruct->GetStructMemory()); 
+		ParameterStore->SetParameterData(LocalValueStruct->GetStructMemory(), Var);
 	}
 	else
 	{
-		// For object types make sure the override is setup to an input which matches the default object.
-		FScopedTransaction ScopedTransaction(LOCTEXT("ResetInputObjectTransaction", "Reset the inputs data interface object to default."));
+		UNiagaraDataInterface* DefaultObject = NewObject<UNiagaraDataInterface>(this, const_cast<UClass*>(InputType.GetClass()));
+		DefaultObject->CopyTo(ParameterStore->GetDataInterface(Var));
 	}
+	RefreshValueAndHandle();
 	RefreshChildren();
+	NotifyEndValueChange();
+	GetSystemViewModel()->ResetSystem();
 }
 
 bool UNiagaraStackParameterStoreEntry::CanRenameInput() const
 {
-	return false;
+	return true; 
 }
 
 bool UNiagaraStackParameterStoreEntry::GetIsRenamePending() const
 {
-	return CanRenameInput() /*&& StackEditorData->GetModuleInputIsRenamePending(StackEditorDataKey)*/;
+	return CanRenameInput() && GetStackEditorData().GetModuleInputIsRenamePending(ParameterName.ToString());
 }
 
 void UNiagaraStackParameterStoreEntry::SetIsRenamePending(bool bIsRenamePending)
 {
 	if (CanRenameInput())
 	{
-		//StackEditorData->SetModuleInputIsRenamePending(StackEditorDataKey, bIsRenamePending);
+		GetStackEditorData().SetModuleInputIsRenamePending(ParameterName.ToString(), bIsRenamePending);
 	}
 }
 
+TArray<UEdGraphPin*> UNiagaraStackParameterStoreEntry::GetOwningPins()
+{
+	TArray<UNiagaraGraph*> GraphsToCheck;
+	// search system graph
+	UNiagaraScript* SystemScript = GetSystemViewModel()->GetSystem().GetSystemSpawnScript();
+	if (SystemScript != nullptr)
+	{
+		UNiagaraScriptSource* ScriptSource = Cast<UNiagaraScriptSource>(SystemScript->GetSource());
+		if (ScriptSource != nullptr)
+		{
+			UNiagaraGraph* SystemGraph = ScriptSource->NodeGraph;
+			if (SystemGraph != nullptr)
+			{
+				GraphsToCheck.Add(SystemGraph);
+			}
+		}
+	}
+
+	// search emitter graphs
+	auto EmitterHandles = GetSystemViewModel()->GetSystem().GetEmitterHandles();
+	for (FNiagaraEmitterHandle Handle : EmitterHandles)
+	{
+		UNiagaraGraph* EmitterGraph = CastChecked<UNiagaraScriptSource>(Handle.GetInstance()->GraphSource)->NodeGraph;
+		GraphsToCheck.Add(EmitterGraph);
+	}
+	TArray<UEdGraphPin*> OwningPins;
+	for (UNiagaraGraph* Graph : GraphsToCheck)
+	{
+		TArray<UNiagaraNodeParameterMapGet*> MapReadNodes;
+		Graph->GetNodesOfClass<UNiagaraNodeParameterMapGet>(MapReadNodes);
+		for (UNiagaraNode* Node : MapReadNodes)
+		{
+			for (UEdGraphPin* GraphPin : Node->Pins)
+			{
+				if (GraphPin->GetName() == ParameterName.ToString())
+				{
+					OwningPins.Add(GraphPin);
+					break;
+				}
+			}
+		}
+	}
+	return OwningPins;
+}
+
 void UNiagaraStackParameterStoreEntry::RenameInput(FString NewName)
 {
-	/*
-	if (OwningAssignmentNode.IsValid() && InputParameterHandlePath.Num() == 1 && InputParameterHandle.GetName() != NewName)
+	FName NewFName = FName(*NewName);
+	FString ActualNameString = NewName;
+	FString NamespacePrefix = FNiagaraParameterHandle::UserNamespace.ToString() + ".";
+	if (NewName.Contains(NamespacePrefix))
 	{
-		bool bIsCurrentlyPinned = GetIsPinned();
-		bool bIsCurrentlyExpanded = StackEditorData->GetStackEntryIsExpanded(FNiagaraStackGraphUtilities::GenerateStackModuleEditorDataKey(*OwningAssignmentNode), false);
-
-		FNiagaraParameterHandle TargetHandle(OwningAssignmentNode->AssignmentTarget.GetName().ToString());
-		FNiagaraParameterHandle RenamedTargetHandle(TargetHandle.GetNamespace(), NewName);
-		OwningAssignmentNode->AssignmentTarget.SetName(RenamedTargetHandle.GetParameterHandleString());
-		OwningAssignmentNode->RefreshFromExternalChanges();
-
-		InputParameterHandle = FNiagaraParameterHandle(InputParameterHandle.GetNamespace(), NewName);
-		InputParameterHandlePath.Empty();
-		InputParameterHandlePath.Add(InputParameterHandle);
-		AliasedInputParameterHandle = FNiagaraParameterHandle::CreateAliasedModuleParameterHandle(InputParameterHandle, OwningAssignmentNode.Get());
-		DisplayName = FText::FromString(InputParameterHandle.GetName());
-
-		UEdGraphPin* OverridePin = GetOverridePin();
-		if (OverridePin != nullptr)
+		ActualNameString = NewName.Replace(*NamespacePrefix, TEXT(""));
+	}
+	FName ActualName = FName(*ActualNameString);
+	// what if it's not user namespace? dehardcode.
+	FNiagaraParameterHandle ParameterHandle(FNiagaraParameterHandle::UserNamespace, ActualName); 
+	FName VariableName = ParameterHandle.GetParameterHandleString();
+	if (VariableName != ParameterName)
+	{
+
+		// destroy links, rename parameter and rebuild links
+		TArray<UEdGraphPin*> OwningPins = GetOwningPins();
+		TArray<UEdGraphPin*> LinkedPins;
+		for (UEdGraphPin* GraphPin : OwningPins)
 		{
-			OverridePin->PinName = AliasedInputParameterHandle.GetParameterHandleString();
+			for (UEdGraphPin* OverridePin : GraphPin->LinkedTo)
+			{
+				LinkedPins.Add(OverridePin);
+			}
 		}
 
-		StackEditorDataKey = FNiagaraStackGraphUtilities::GenerateStackFunctionInputEditorDataKey(*OwningFunctionCallNode.Get(), InputParameterHandle);
-		StackEditorData->SetModuleInputIsPinned(StackEditorDataKey, bIsCurrentlyPinned);
-		StackEditorData->SetStackEntryIsExpanded(FNiagaraStackGraphUtilities::GenerateStackModuleEditorDataKey(*OwningAssignmentNode), bIsCurrentlyExpanded);
+		GEditor->BeginTransaction(LOCTEXT("RenameParameter", "Rename Parameter"));
+		ParameterStore->GetOwner()->Modify();
+		// remove old one (a bit overkill but it beats duplicating code)
+		RemovePins(OwningPins);
+		// TODO Would it be better to actually keep variable name being ActualName and ParameterName being ParameterHandle.GetParameterHandleString(), and rewrite the way the Entry is built?
+		ParameterStore->RenameParameter(FNiagaraVariable(InputType, ParameterName), VariableName); 
+		// rebuild all links
+		for (UEdGraphPin* LinkedPin : LinkedPins)
+		{
+			// remove links
+			TArray<TWeakObjectPtr<UNiagaraDataInterface>> RemovedDataObjects;
+			FNiagaraStackGraphUtilities::RemoveNodesForStackFunctionInputOverridePin(*LinkedPin, RemovedDataObjects); // no need to broadcast data objects modified here, the graph will recompile
+			// generate current link
+			FNiagaraStackGraphUtilities::SetLinkedValueHandleForFunctionInput(*LinkedPin, ParameterHandle);
+		}
+		GEditor->EndTransaction();
+	}
+}
 
-		CastChecked<UNiagaraGraph>(OwningAssignmentNode->GetGraph())->NotifyGraphNeedsRecompile();
+void UNiagaraStackParameterStoreEntry::Delete()
+{
+	// for delete, do a parameter map traversal to deduce all the usages and then  remove them 
+	TArray<UEdGraphPin*> OwningPins = GetOwningPins();
+	NotifyBeginValueChange();
+	RemovePins(OwningPins);
+	//remove from store
+	ParameterStore->RemoveParameter(FNiagaraVariable(InputType, ParameterName));
+	if (InputType.GetClass() != nullptr)
+	{
+		UNiagaraDataInterface* DataInterface = NewObject<UNiagaraDataInterface>(this, const_cast<UClass*>(InputType.GetClass()));
+		if (DataInterface != nullptr)
+		{
+			GetSystemViewModel()->NotifyDataObjectChanged(DataInterface);
+		}
+	}
+	NotifyEndValueChange();
+}
+
+void UNiagaraStackParameterStoreEntry::RemovePins(TArray<UEdGraphPin*> OwningPins /*, bool bSetPreviousValue*/)
+{
+	for (UEdGraphPin* GraphPin : OwningPins)
+	{
+		UNiagaraGraph* Graph = CastChecked<UNiagaraGraph>(GraphPin->GetOwningNode()->GetGraph());
+	
+		if (GraphPin->LinkedTo.Num() != 0)
+		{
+			// remember output of pin 
+			UEdGraphPin* OverridePin = GraphPin->LinkedTo[0];
+			//break old pin links
+			GraphPin->BreakAllPinLinks();
+			
+			//now set value of pin output to the value of the GetCurrentValueVariable()
+			const UEdGraphSchema_Niagara* Schema = GetDefault<UEdGraphSchema_Niagara>();
+			//FNiagaraVariable Var = Schema->PinToNiagaraVariable(GraphPin, true); // use this instead of GetCurrentValueVariable() for default value
+			FString PinDefaultValue;
+			if (InputType.GetClass() == nullptr)
+			{
+				if (Schema->TryGetPinDefaultValueFromNiagaraVariable(*GetCurrentValueVariable(), PinDefaultValue))
+				{
+					OverridePin->DefaultValue = PinDefaultValue;
+				}
+			}
+			else
+			{
+				UNiagaraDataInterface* OverrideObj = NewObject<UNiagaraDataInterface>(this, const_cast<UClass*>(InputType.GetClass()));
+				FNiagaraStackGraphUtilities::SetDataValueObjectForFunctionInput(*OverridePin, const_cast<UClass*>(InputType.GetClass()), GetCurrentValueObject()->GetName(), OverrideObj);
+				GetCurrentValueObject()->CopyTo(OverrideObj);
+			}
+			
+		}
+		// now also remove node
+		Graph->RemoveNode(GraphPin->GetOwningNode());
+		Graph->NotifyGraphNeedsRecompile();
 	}
-	*/
 }
 
 UNiagaraStackParameterStoreEntry::FOnValueChanged& UNiagaraStackParameterStoreEntry::OnValueChanged()
@@ -260,4 +363,37 @@ UNiagaraDataInterface* UNiagaraStackParameterStoreEntry::GetCurrentValueObject()
 	return nullptr;
 }
 
+bool UNiagaraStackParameterStoreEntry::IsUniqueName(FString NewName)
+{
+	FString NamespacePrefix = FNiagaraParameterHandle::UserNamespace.ToString() + "."; // correcting name of variable for comparison, all user variables start with "User."
+	if (!NewName.Contains(NamespacePrefix))
+	{
+		NewName = NamespacePrefix + NewName;
+	}
+	TArray<FNiagaraVariable> Variables;
+	ParameterStore->GetParameters(Variables);
+	// check for duplicates, but exclude self from search
+	for (auto parameter : Variables)
+	{
+		if (parameter.GetName().ToString() == NewName)
+		{
+			if (GetCurrentValueVariable().IsValid())
+			{
+				if (parameter != *GetCurrentValueVariable())
+				{
+					return false;
+				}
+			}
+			else if (GetCurrentValueObject() != nullptr)
+			{
+				if (ParameterStore->GetDataInterface(parameter) != GetCurrentValueObject())
+				{
+					return false;
+				}
+			}
+		}
+	}
+	return true;
+}
+
 #undef LOCTEXT_NAMESPACE
diff --git a/Engine/Plugins/FX/Niagara/Source/NiagaraEditor/Private/ViewModels/Stack/NiagaraStackParameterStoreGroup.cpp b/Engine/Plugins/FX/Niagara/Source/NiagaraEditor/Private/ViewModels/Stack/NiagaraStackParameterStoreGroup.cpp
--- a/Engine/Plugins/FX/Niagara/Source/NiagaraEditor/Private/ViewModels/Stack/NiagaraStackParameterStoreGroup.cpp
+++ b/Engine/Plugins/FX/Niagara/Source/NiagaraEditor/Private/ViewModels/Stack/NiagaraStackParameterStoreGroup.cpp
@@ -2,7 +2,6 @@
 
 #include "ViewModels/Stack/NiagaraStackParameterStoreGroup.h"
 #include "ViewModels/Stack/NiagaraStackModuleItem.h"
-#include "ViewModels/Stack/NiagaraStackAddScriptModuleItem.h"
 #include "NiagaraScriptViewModel.h"
 #include "NiagaraScriptGraphViewModel.h"
 #include "NiagaraGraph.h"
@@ -10,62 +9,194 @@
 #include "NiagaraNodeFunctionCall.h"
 #include "EdGraphSchema_Niagara.h"
 #include "ViewModels/Stack/NiagaraStackErrorItem.h"
-#include "Internationalization/Internationalization.h"
 #include "ViewModels/Stack/NiagaraStackParameterStoreEntry.h"
+#include "Internationalization/Internationalization.h"
+#include "ViewModels/Stack/NiagaraStackSpacer.h"
+#include "NiagaraStackEditorData.h"
+#include "ScopedTransaction.h"
+#include "ViewModels/Stack/NiagaraStackGraphUtilities.h"
 
 
 #define LOCTEXT_NAMESPACE "UNiagaraStackParameterStoreGroup"
 
+class FParameterStoreGroupAddAction : public INiagaraStackItemGroupAddAction
+{
+public:
+	FParameterStoreGroupAddAction(FNiagaraVariable InNewParameterVariable)
+		: NewParameterVariable(InNewParameterVariable)
+	{
+	}
+
+	FNiagaraVariable GetNewParameterVariable() const
+	{
+		return NewParameterVariable;
+	}
+
+	virtual FText GetCategory() const override
+	{
+		return LOCTEXT("CreateNewParameterCategory", "Parameter Types");
+	}
+
+	virtual FText GetDisplayName() const override
+	{
+		return NewParameterVariable.GetType().GetNameText();
+	}
+
+	virtual FText GetDescription() const override
+	{
+		return FText::Format(LOCTEXT("AddParameterActionDescriptionFormat", "Create a new {0} parameter."), GetDisplayName());
+	}
+
+private:
+	FNiagaraVariable NewParameterVariable;
+};
+
+class FParameterStoreGroupAddUtiliites : public FNiagaraStackItemGroupAddUtilities
+{
+public:
+	FParameterStoreGroupAddUtiliites(UObject& InParameterStoreOwner, FNiagaraParameterStore& InParameterStore, UNiagaraStackEditorData& InStackEditorData)
+		: FNiagaraStackItemGroupAddUtilities(LOCTEXT("ScriptGroupAddItemName", "Parameter"), EAddMode::AddFromAction, true, FOnItemAdded())
+		, ParameterStoreOwner(InParameterStoreOwner)
+		, ParameterStore(InParameterStore)
+		, StackEditorData(InStackEditorData)
+	{
+	}
+
+	virtual void AddItemDirectly() override { unimplemented(); }
+
+	virtual void GenerateAddActions(TArray<TSharedRef<INiagaraStackItemGroupAddAction>>& OutAddActions) const override
+	{
+		TArray<FNiagaraTypeDefinition> AvailableTypes;
+		FNiagaraStackGraphUtilities::GetNewParameterAvailableTypes(AvailableTypes);
+		for (const FNiagaraTypeDefinition& AvailableType : AvailableTypes)
+		{
+			FNiagaraParameterHandle NewParameterHandle(FNiagaraParameterHandle::UserNamespace, *(TEXT("New") + AvailableType.GetName()));
+			FNiagaraVariable NewParameterVariable(AvailableType, NewParameterHandle.GetParameterHandleString());
+			OutAddActions.Add(MakeShared<FParameterStoreGroupAddAction>(NewParameterVariable));
+		}
+	}
+
+	virtual void ExecuteAddAction(TSharedRef<INiagaraStackItemGroupAddAction> AddAction, int32 TargetIndex) override
+	{
+		TSharedRef<FParameterStoreGroupAddAction> ParameterAddAction = StaticCastSharedRef<FParameterStoreGroupAddAction>(AddAction);
+
+		FScopedTransaction AddTransaction(LOCTEXT("AddParameter", "Add Parameter"));
+		ParameterStoreOwner.Modify();
+
+		FNiagaraVariable ParameterVariable = ParameterAddAction->GetNewParameterVariable();
+		FNiagaraEditorUtilities::ResetVariableToDefaultValue(ParameterVariable);
+
+		ParameterStore.AddParameter(ParameterVariable);
+		StackEditorData.SetModuleInputIsRenamePending(ParameterVariable.GetName().ToString(), true);
+	}
+
+private:
+	UObject& ParameterStoreOwner;
+	FNiagaraParameterStore& ParameterStore;
+	UNiagaraStackEditorData& StackEditorData;
+};
 
 void UNiagaraStackParameterStoreGroup::Initialize(
-	TSharedRef<FNiagaraSystemViewModel> InSystemViewModel,
-	TSharedRef<FNiagaraEmitterViewModel> InEmitterViewModel,
-	UNiagaraStackEditorData& InStackEditorData,
-	TSharedRef<FNiagaraScriptViewModel> InScriptViewModel,
+	FRequiredEntryData InRequiredEntryData,
 	UObject* InOwner,
 	FNiagaraParameterStore* InParameterStore)
 {
-	Super::Initialize(InSystemViewModel, InEmitterViewModel, InStackEditorData);
+	AddUtilities = MakeShared<FParameterStoreGroupAddUtiliites>(*InOwner, *InParameterStore, *InRequiredEntryData.StackEditorData);
+	FText DisplayName = LOCTEXT("SystemExposedVariablesGroup", "System Exposed Parameters");
+	FText ToolTip = LOCTEXT("SystemExposedVariablesGroupToolTip", "Displays the variables created in the User namespace. These variables are exposed to owning UComponents, blueprints, etc.");
+	Super::Initialize(InRequiredEntryData, DisplayName, ToolTip, AddUtilities.Get());
+	
 	Owner = InOwner;
 	ParameterStore = InParameterStore;
 }
 
-FText UNiagaraStackParameterStoreGroup::GetDisplayName() const
+void UNiagaraStackParameterStoreGroup::RefreshChildrenInternal(const TArray<UNiagaraStackEntry*>& CurrentChildren, TArray<UNiagaraStackEntry*>& NewChildren, TArray<FStackIssue>& NewIssues)
 {
-	return DisplayName;
+	if (ParameterStore != nullptr)
+	{
+		FName ParameterStoreItemSpacerKey = "ParameterStoreSpacer";
+		UNiagaraStackSpacer* ParameterStoreItemSpacer = FindCurrentChildOfTypeByPredicate<UNiagaraStackSpacer>(CurrentChildren,
+			[=](UNiagaraStackSpacer* CurrentParameterStoreItemSpacer) { return CurrentParameterStoreItemSpacer->GetSpacerKey() == ParameterStoreItemSpacerKey; });
+
+		if (ParameterStoreItemSpacer == nullptr)
+		{
+			ParameterStoreItemSpacer = NewObject<UNiagaraStackSpacer>(this);
+			ParameterStoreItemSpacer->Initialize(CreateDefaultChildRequiredData(), ParameterStoreItemSpacerKey, 1.4f);
+		}
+
+		NewChildren.Add(ParameterStoreItemSpacer);
+
+		UNiagaraStackParameterStoreItem* ParameterStoreItem = FindCurrentChildOfTypeByPredicate<UNiagaraStackParameterStoreItem>(CurrentChildren,
+			[=](UNiagaraStackParameterStoreItem* CurrentItem) { return true; });
+
+		if (ParameterStoreItem == nullptr)
+		{
+			ParameterStoreItem = NewObject<UNiagaraStackParameterStoreItem>(this);
+			ParameterStoreItem->Initialize(CreateDefaultChildRequiredData(), Owner.Get(), ParameterStore);
+		}
+
+		NewChildren.Add(ParameterStoreItem);
+	}
+	Super::RefreshChildrenInternal(CurrentChildren, NewChildren, NewIssues);
+}
+
+void UNiagaraStackParameterStoreItem::Initialize(
+	FRequiredEntryData InRequiredEntryData,
+	UObject* InOwner,
+	FNiagaraParameterStore* InParameterStore)
+{
+	Super::Initialize(InRequiredEntryData, TEXT("ParameterStoreItem"));
+
+	Owner = InOwner;
+	ParameterStore = InParameterStore;
+	ParameterStoreChangedHandle = ParameterStore->AddOnChangedHandler(
+		FNiagaraParameterStore::FOnChanged::FDelegate::CreateUObject(this, &UNiagaraStackParameterStoreItem::ParameterStoreChanged));
 }
 
-void UNiagaraStackParameterStoreGroup::SetDisplayName(FText InDisplayName)
+FText UNiagaraStackParameterStoreItem::GetDisplayName() const
 {
-	DisplayName = InDisplayName;
+	return LOCTEXT("ParameterItemDisplayName", "Parameters");
 }
 
-void UNiagaraStackParameterStoreGroup::RefreshChildrenInternal(const TArray<UNiagaraStackEntry*>& CurrentChildren, TArray<UNiagaraStackEntry*>& NewChildren)
+void UNiagaraStackParameterStoreItem::BeginDestroy()
 {
-	if (ParameterStore != nullptr)
+	if (Owner != nullptr)
+	{
+		ParameterStore->RemoveOnChangedHandler(ParameterStoreChangedHandle);
+	}
+	Super::BeginDestroy();
+}
+
+void UNiagaraStackParameterStoreItem::RefreshChildrenInternal(const TArray<UNiagaraStackEntry*>& CurrentChildren, TArray<UNiagaraStackEntry*>& NewChildren, TArray<FStackIssue>& NewIssues)
+{
+	if (ParameterStore != nullptr && ParameterStore->GetNumParameters() > 0)
 	{
 		TArray<FNiagaraVariable> Variables;
 		ParameterStore->GetParameters(Variables);
 
 		for (FNiagaraVariable& Var : Variables)
 		{
-			UNiagaraStackParameterStoreEntry* ValueObjectEntry = NewObject<UNiagaraStackParameterStoreEntry>(this);
-			ValueObjectEntry->Initialize(GetSystemViewModel(), GetEmitterViewModel(), GetStackEditorData(), Owner, ParameterStore, Var.GetName().ToString(), Var.GetType());
-			ValueObjectEntry->SetItemIndentLevel(1);
+			UNiagaraStackParameterStoreEntry* ValueObjectEntry = FindCurrentChildOfTypeByPredicate<UNiagaraStackParameterStoreEntry>(CurrentChildren,
+				[=](UNiagaraStackParameterStoreEntry* CurrentEntry) { return CurrentEntry->GetDisplayName().ToString() == Var.GetName().ToString(); });
+
+			if (ValueObjectEntry == nullptr)
+			{
+				ValueObjectEntry = NewObject<UNiagaraStackParameterStoreEntry>(this);
+				ValueObjectEntry->Initialize(CreateDefaultChildRequiredData(), Owner.Get(), ParameterStore, Var.GetName().ToString(), Var.GetType(), GetStackEditorDataKey());
+			}
+
 			NewChildren.Add(ValueObjectEntry);
 		}
 	}
+	Super::RefreshChildrenInternal(CurrentChildren, NewChildren, NewIssues);
 }
 
-void UNiagaraStackParameterStoreGroup::ItemAdded()
-{
-	RefreshChildren();
-}
-
-void UNiagaraStackParameterStoreGroup::ChildModifiedGroupItems()
+void UNiagaraStackParameterStoreItem::ParameterStoreChanged()
 {
-	RefreshChildren();
+	if (IsValid())
+	{
+		RefreshChildren();
+	}
 }
 
-
 #undef LOCTEXT_NAMESPACE
diff --git a/Engine/Plugins/FX/Niagara/Source/NiagaraEditor/Private/ViewModels/Stack/NiagaraStackPropertyRow.cpp b/Engine/Plugins/FX/Niagara/Source/NiagaraEditor/Private/ViewModels/Stack/NiagaraStackPropertyRow.cpp
new file mode 100644
--- /dev/null
+++ b/Engine/Plugins/FX/Niagara/Source/NiagaraEditor/Private/ViewModels/Stack/NiagaraStackPropertyRow.cpp
@@ -0,0 +1,64 @@
+// Copyright 1998-2018 Epic Games, Inc. All Rights Reserved.
+
+#include "ViewModels/Stack/NiagaraStackPropertyRow.h"
+#include "NiagaraNode.h"
+
+#include "IDetailTreeNode.h"
+#include "PropertyHandle.h"
+
+void UNiagaraStackPropertyRow::Initialize(FRequiredEntryData InRequiredEntryData, TSharedRef<IDetailTreeNode> InDetailTreeNode, FString InOwnerStackItemEditorDataKey, FString InOwnerStackEditorDataKey, UNiagaraNode* InOwningNiagaraNode)
+{
+	TSharedPtr<IPropertyHandle> PropertyHandle = InDetailTreeNode->CreatePropertyHandle();
+	bool bRowIsAdvanced = PropertyHandle.IsValid() && PropertyHandle->GetProperty()->HasAnyPropertyFlags(CPF_AdvancedDisplay);
+	FString RowStackEditorDataKey = FString::Printf(TEXT("%s-%s"), *InOwnerStackEditorDataKey, *InDetailTreeNode->GetNodeName().ToString());
+	Super::Initialize(InRequiredEntryData, bRowIsAdvanced, InOwnerStackItemEditorDataKey, RowStackEditorDataKey);
+	DetailTreeNode = InDetailTreeNode;
+	OwningNiagaraNode = InOwningNiagaraNode;
+	RowStyle = DetailTreeNode->GetNodeType() == EDetailNodeType::Category
+		? EStackRowStyle::ItemCategory
+		: EStackRowStyle::ItemContent;
+}
+
+TSharedRef<IDetailTreeNode> UNiagaraStackPropertyRow::GetDetailTreeNode() const
+{
+	return DetailTreeNode.ToSharedRef();
+}
+
+bool UNiagaraStackPropertyRow::GetIsEnabled() const
+{
+	return OwningNiagaraNode == nullptr || OwningNiagaraNode->GetDesiredEnabledState() == ENodeEnabledState::Enabled;
+}
+
+void UNiagaraStackPropertyRow::RefreshChildrenInternal(const TArray<UNiagaraStackEntry*>& CurrentChildren, TArray<UNiagaraStackEntry*>& NewChildren, TArray<FStackIssue>& NewIssues)
+{
+	TArray<TSharedRef<IDetailTreeNode>> NodeChildren;
+	DetailTreeNode->GetChildren(NodeChildren);
+	for (TSharedRef<IDetailTreeNode> NodeChild : NodeChildren)
+	{
+		if (NodeChild->GetNodeType() == EDetailNodeType::Advanced)
+		{
+			continue;
+		}
+
+		UNiagaraStackPropertyRow* ChildRow = FindCurrentChildOfTypeByPredicate<UNiagaraStackPropertyRow>(CurrentChildren,
+			[=](UNiagaraStackPropertyRow* CurrentChild) { return CurrentChild->GetDetailTreeNode() == NodeChild; });
+
+		if (ChildRow == nullptr)
+		{
+			ChildRow = NewObject<UNiagaraStackPropertyRow>(this);
+			ChildRow->Initialize(CreateDefaultChildRequiredData(), NodeChild, GetOwnerStackItemEditorDataKey(), GetStackEditorDataKey(), OwningNiagaraNode);
+		}
+
+		NewChildren.Add(ChildRow);
+	}
+}
+
+void UNiagaraStackPropertyRow::GetAdditionalSearchItemsInternal(TArray<FStackSearchItem>& SearchItems) const
+{
+	TArray<FString> NodeFilterStrings;
+	DetailTreeNode->GetFilterStrings(NodeFilterStrings);
+	for (FString& FilterString : NodeFilterStrings)
+	{
+		SearchItems.Add({ "PropertyRowFilterString", FText::FromString(FilterString) });
+	}
+}
\ No newline at end of file
diff --git a/Engine/Plugins/FX/Niagara/Source/NiagaraEditor/Private/ViewModels/Stack/NiagaraStackRenderItemGroup.cpp b/Engine/Plugins/FX/Niagara/Source/NiagaraEditor/Private/ViewModels/Stack/NiagaraStackRenderItemGroup.cpp
--- a/Engine/Plugins/FX/Niagara/Source/NiagaraEditor/Private/ViewModels/Stack/NiagaraStackRenderItemGroup.cpp
+++ b/Engine/Plugins/FX/Niagara/Source/NiagaraEditor/Private/ViewModels/Stack/NiagaraStackRenderItemGroup.cpp
@@ -3,37 +3,129 @@
 #include "ViewModels/Stack/NiagaraStackRenderItemGroup.h"
 #include "ViewModels/Stack/NiagaraStackRendererItem.h"
 #include "ViewModels/Stack/NiagaraStackSpacer.h"
-#include "ViewModels/Stack/NiagaraStackAddRendererItem.h"
-#include "NiagaraEmitterViewModel.h"
+#include "ViewModels/NiagaraEmitterViewModel.h"
 #include "NiagaraEmitter.h"
 #include "NiagaraEmitterEditorData.h"
-#include "NiagaraSystemViewModel.h"
+#include "ViewModels/NiagaraSystemViewModel.h"
+#include "NiagaraRendererProperties.h"
 
-FText UNiagaraStackRenderItemGroup::GetDisplayName() const
-{
-	return DisplayName;
-}
+#include "ScopedTransaction.h"
+#include "Widgets/Notifications/SNotificationList.h"
+#include "Framework/Notifications/NotificationManager.h"
+#include "Styling/CoreStyle.h"
+
+#define LOCTEXT_NAMESPACE "UNiagaraStackRenderItemGroup"
 
-void UNiagaraStackRenderItemGroup::SetDisplayName(FText InDisplayName)
+class FRenderItemGroupAddAction : public INiagaraStackItemGroupAddAction
 {
-	DisplayName = InDisplayName;
-}
+public:
+	FRenderItemGroupAddAction(UClass* InRendererClass)
+		: RendererClass(InRendererClass)
+	{
+	}
+
+	virtual FText GetCategory() const override
+	{
+		return LOCTEXT("AddRendererCategory", "Add Renderer");
+	}
+
+	virtual FText GetDisplayName() const override
+	{
+		return RendererClass->GetDisplayNameText();
+	}
+
+	virtual FText GetDescription() const override
+	{
+		return FText::FromString(RendererClass->GetDescription());
+	}
+
+	UClass* GetRendererClass() const
+	{
+		return RendererClass;
+	}
 
-void UNiagaraStackRenderItemGroup::RefreshChildrenInternal(const TArray<UNiagaraStackEntry*>& CurrentChildren, TArray<UNiagaraStackEntry*>& NewChildren)
+private:
+	UClass* RendererClass;
+};
+
+class FRenderItemGroupAddUtilities : public FNiagaraStackItemGroupAddUtilities
 {
-	int32 RendererIndex = 0;
-	for (UNiagaraRendererProperties* RendererProperties : GetEmitterViewModel()->GetEmitter()->GetRenderers())
+public:
+	FRenderItemGroupAddUtilities(TSharedRef<FNiagaraEmitterViewModel> InEmitterViewModel, FOnItemAdded InOnItemAdded)
+		: FNiagaraStackItemGroupAddUtilities(LOCTEXT("RenderGroupAddItemName", "Renderer"), EAddMode::AddFromAction, true, InOnItemAdded)
+		, EmitterViewModel(InEmitterViewModel)
 	{
-		UNiagaraStackRendererItem* RendererItem = FindCurrentChildOfTypeByPredicate<UNiagaraStackRendererItem>(CurrentChildren,
-			[=](UNiagaraStackRendererItem* CurrentRendererItem) { return CurrentRendererItem->GetRendererProperties() == RendererProperties; });
+	}
 
-		if (RendererItem == nullptr)
+	virtual void AddItemDirectly() override { unimplemented(); }
+
+	virtual void GenerateAddActions(TArray<TSharedRef<INiagaraStackItemGroupAddAction>>& OutAddActions) const override
+	{
+		TArray<UClass*> RendererClasses;
+		GetDerivedClasses(UNiagaraRendererProperties::StaticClass(), RendererClasses);
+		for (UClass* RendererClass : RendererClasses)
 		{
-			RendererItem = NewObject<UNiagaraStackRendererItem>(this);
-			RendererItem->Initialize(GetSystemViewModel(), GetEmitterViewModel(), GetStackEditorData(), RendererProperties);
-			RendererItem->SetOnModifiedGroupItems(UNiagaraStackItem::FOnModifiedGroupItems::CreateUObject(this, &UNiagaraStackRenderItemGroup::ChildModifiedGroupItems));
+			OutAddActions.Add(MakeShared<FRenderItemGroupAddAction>(RendererClass));
+		}
+	}
+
+	virtual void ExecuteAddAction(TSharedRef<INiagaraStackItemGroupAddAction> AddAction, int32 TargetIndex) override
+	{
+		TSharedPtr<FNiagaraEmitterViewModel> EmitterViewModelPinned = EmitterViewModel.Pin();
+		if (EmitterViewModelPinned.IsValid() == false)
+		{
+			return;
+		}
+
+		TSharedRef<FRenderItemGroupAddAction> RenderAddAction = StaticCastSharedRef<FRenderItemGroupAddAction>(AddAction);
+
+		FScopedTransaction ScopedTransaction(LOCTEXT("AddNewRendererTransaction", "Add new renderer"));
+
+		UNiagaraEmitter* Emitter = EmitterViewModelPinned->GetEmitter();
+		Emitter->Modify();
+		UNiagaraRendererProperties* RendererProperties = NewObject<UNiagaraRendererProperties>(Emitter, RenderAddAction->GetRendererClass(), NAME_None, RF_Transactional);
+		Emitter->AddRenderer(RendererProperties);
+
+		bool bVarsAdded = false;
+		TArray<FNiagaraVariable> MissingAttributes = UNiagaraStackRendererItem::GetMissingVariables(RendererProperties, Emitter);
+		for (int32 i = 0; i < MissingAttributes.Num(); i++)
+		{
+			if (UNiagaraStackRendererItem::AddMissingVariable(Emitter, MissingAttributes[i]))
+			{
+				bVarsAdded = true;
+			}
 		}
 
+		if (bVarsAdded)
+		{
+			FNotificationInfo Info(LOCTEXT("AddedVariables", "One or more variables have been added to the Spawn script to support the added renderer."));
+			Info.ExpireDuration = 5.0f;
+			Info.bFireAndForget = true;
+			Info.Image = FCoreStyle::Get().GetBrush(TEXT("MessageLog.Info"));
+			FSlateNotificationManager::Get().AddNotification(Info);
+		}
+
+		OnItemAdded.ExecuteIfBound();
+	}
+
+private:
+	TWeakPtr<FNiagaraEmitterViewModel> EmitterViewModel;
+};
+
+void UNiagaraStackRenderItemGroup::Initialize(FRequiredEntryData InRequiredEntryData)
+{
+	FText DisplayName = LOCTEXT("RenderGroupName", "Render");
+	FText ToolTip = LOCTEXT("RendererGroupTooltip", "Describes how we should display/present each particle. Note that this doesn't have to be visual. Multiple renderers are supported. Order in this stack is not necessarily relevant to draw order.");
+	AddUtilities = MakeShared<FRenderItemGroupAddUtilities>(InRequiredEntryData.EmitterViewModel,
+		FNiagaraStackItemGroupAddUtilities::FOnItemAdded::CreateUObject(this, &UNiagaraStackRenderItemGroup::ItemAdded));
+	Super::Initialize(InRequiredEntryData, DisplayName, ToolTip, AddUtilities.Get());
+}
+
+void UNiagaraStackRenderItemGroup::RefreshChildrenInternal(const TArray<UNiagaraStackEntry*>& CurrentChildren, TArray<UNiagaraStackEntry*>& NewChildren, TArray<FStackIssue>& NewIssues)
+{
+	int32 RendererIndex = 0;
+	for (UNiagaraRendererProperties* RendererProperties : GetEmitterViewModel()->GetEmitter()->GetRenderers())
+	{
 		FName RendererSpacerKey = *FString::Printf(TEXT("Renderer%i"), RendererIndex);
 		UNiagaraStackSpacer* RendererSpacer = FindCurrentChildOfTypeByPredicate<UNiagaraStackSpacer>(CurrentChildren,
 			[=](UNiagaraStackSpacer* CurrentRendererSpacer) { return CurrentRendererSpacer->GetSpacerKey() == RendererSpacerKey; });
@@ -41,24 +133,38 @@ void UNiagaraStackRenderItemGroup::RefreshChildrenInternal(const TArray<UNiagara
 		if (RendererSpacer == nullptr)
 		{
 			RendererSpacer = NewObject<UNiagaraStackSpacer>(this);
-			RendererSpacer->Initialize(GetSystemViewModel(), GetEmitterViewModel(), RendererSpacerKey);
+			RendererSpacer->Initialize(CreateDefaultChildRequiredData(), RendererSpacerKey);
 		}
 
-		NewChildren.Add(RendererItem);
 		NewChildren.Add(RendererSpacer);
 
+		UNiagaraStackRendererItem* RendererItem = FindCurrentChildOfTypeByPredicate<UNiagaraStackRendererItem>(CurrentChildren,
+			[=](UNiagaraStackRendererItem* CurrentRendererItem) { return CurrentRendererItem->GetRendererProperties() == RendererProperties; });
+
+		if (RendererItem == nullptr)
+		{
+			RendererItem = NewObject<UNiagaraStackRendererItem>(this);
+			RendererItem->Initialize(CreateDefaultChildRequiredData(), RendererProperties);
+			RendererItem->SetOnModifiedGroupItems(UNiagaraStackItem::FOnModifiedGroupItems::CreateUObject(this, &UNiagaraStackRenderItemGroup::ChildModifiedGroupItems));
+		}
+
+		NewChildren.Add(RendererItem);
+
 		RendererIndex++;
 	}
 
-	UNiagaraStackAddRendererItem* AddRendererItem = NewObject<UNiagaraStackAddRendererItem>(this);
-	AddRendererItem->Initialize(GetSystemViewModel(), GetEmitterViewModel());
-	AddRendererItem->SetOnItemAdded(UNiagaraStackAddRendererItem::FOnItemAdded::CreateUObject(this, &UNiagaraStackRenderItemGroup::ChildModifiedGroupItems));
-	NewChildren.Add(AddRendererItem);
+	Super::RefreshChildrenInternal(CurrentChildren, NewChildren, NewIssues);
+}
+
+void UNiagaraStackRenderItemGroup::ItemAdded()
+{
+	RefreshChildren();
 }
 
 void UNiagaraStackRenderItemGroup::ChildModifiedGroupItems()
 {
-	GetSystemViewModel()->RefreshAll();
 	RefreshChildren();
 }
 
+#undef LOCTEXT_NAMESPACE
+
diff --git a/Engine/Plugins/FX/Niagara/Source/NiagaraEditor/Private/ViewModels/Stack/NiagaraStackRendererItem.cpp b/Engine/Plugins/FX/Niagara/Source/NiagaraEditor/Private/ViewModels/Stack/NiagaraStackRendererItem.cpp
--- a/Engine/Plugins/FX/Niagara/Source/NiagaraEditor/Private/ViewModels/Stack/NiagaraStackRendererItem.cpp
+++ b/Engine/Plugins/FX/Niagara/Source/NiagaraEditor/Private/ViewModels/Stack/NiagaraStackRendererItem.cpp
@@ -2,13 +2,12 @@
 
 #include "ViewModels/Stack/NiagaraStackRendererItem.h"
 #include "ViewModels/Stack/NiagaraStackObject.h"
-#include "ViewModels/Stack/NiagaraStackItemExpander.h"
 #include "NiagaraEmitter.h"
 #include "NiagaraStackEditorData.h"
 #include "NiagaraRendererProperties.h"
 #include "NiagaraScript.h"
-#include "NiagaraSystemViewModel.h"
-#include "NiagaraEmitterViewModel.h"
+#include "ViewModels/NiagaraSystemViewModel.h"
+#include "ViewModels/NiagaraEmitterViewModel.h"
 #include "NiagaraScriptViewModel.h"
 #include "Internationalization/Internationalization.h"
 #include "NiagaraNodeAssignment.h"
@@ -33,10 +32,11 @@ UNiagaraStackRendererItem::UNiagaraStackRendererItem()
 {
 }
 
-void UNiagaraStackRendererItem::Initialize(TSharedRef<FNiagaraSystemViewModel> InSystemViewModel, TSharedRef<FNiagaraEmitterViewModel> InEmitterViewModel, UNiagaraStackEditorData& InStackEditorData, UNiagaraRendererProperties* InRendererProperties)
+void UNiagaraStackRendererItem::Initialize(FRequiredEntryData InRequiredEntryData, UNiagaraRendererProperties* InRendererProperties)
 {
 	checkf(RendererProperties.IsValid() == false, TEXT("Can not initialize more than once."));
-	Super::Initialize(InSystemViewModel, InEmitterViewModel, InStackEditorData);
+	FString RendererStackEditorDataKey = FString::Printf(TEXT("Renderer-%s"), *InRendererProperties->GetName());
+	Super::Initialize(InRequiredEntryData, RendererStackEditorDataKey);
 	RendererProperties = InRendererProperties;
 	RendererProperties->OnChanged().AddUObject(this, &UNiagaraStackRendererItem::RendererChanged);
 
@@ -57,7 +57,7 @@ TArray<FNiagaraVariable> UNiagaraStackRendererItem::GetMissingVariables(UNiagara
 	TArray<FNiagaraVariable> MissingAttributes;
 	const TArray<FNiagaraVariable>& RequiredAttrs = RendererProperties->GetRequiredAttributes();
 	const UNiagaraScript* Script = Emitter->SpawnScriptProps.Script;
-	if (Script != nullptr && Script->GetByteCode().Num() != 0)
+	if (Script != nullptr && Script->IsReadyToRun(ENiagaraSimTarget::CPUSim))
 	{
 		MissingAttributes.Empty();
 		for (FNiagaraVariable Attr : RequiredAttrs)
@@ -70,7 +70,7 @@ TArray<FNiagaraVariable> UNiagaraStackRendererItem::GetMissingVariables(UNiagara
 				Attr.SetName(*AttrName);
 			}
 
-			bool ContainsVar = Script->Attributes.ContainsByPredicate([&Attr](const FNiagaraVariable& Var) { return Var.GetName() == Attr.GetName(); });
+			bool ContainsVar = Script->GetVMExecutableData().Attributes.ContainsByPredicate([&Attr](const FNiagaraVariable& Var) { return Var.GetName() == Attr.GetName(); });
 			if (!ContainsVar)
 			{
 				MissingAttributes.Add(OriginalAttr);
@@ -110,9 +110,8 @@ bool UNiagaraStackRendererItem::AddMissingVariable(UNiagaraEmitter* Emitter, con
 
 	FGraphNodeCreator<UNiagaraNodeAssignment> NodeBuilder(*Graph);
 	UNiagaraNodeAssignment* NewAssignmentNode = NodeBuilder.CreateNode();
-	NewAssignmentNode->AssignmentTarget = Variable;
-	FString VarDefaultValue = FNiagaraConstants::GetAttributeDefaultValue(NewAssignmentNode->AssignmentTarget);
-	NewAssignmentNode->AssignmentDefaultValue = VarDefaultValue;
+	FString VarDefaultValue = FNiagaraConstants::GetAttributeDefaultValue(Variable);
+	NewAssignmentNode->AddAssignmentTarget(Variable, &VarDefaultValue);
 	NodeBuilder.Finalize();
 
 	TArray<FNiagaraStackGraphUtilities::FStackNodeGroup> StackNodeGroups;
@@ -194,48 +193,16 @@ void UNiagaraStackRendererItem::ResetToBase()
 	}
 }
 
-FName UNiagaraStackRendererItem::GetItemBackgroundName() const
+bool UNiagaraStackRendererItem::GetIsEnabled() const
 {
-	return "NiagaraEditor.Stack.Item.BackgroundColor";
+	return RendererProperties->GetIsEnabled();
 }
 
-int32 UNiagaraStackRendererItem::GetErrorCount() const
+void UNiagaraStackRendererItem::SetIsEnabled(bool bInIsEnabled)
 {
-	return MissingAttributes.Num();
-}
-
-bool UNiagaraStackRendererItem::GetErrorFixable(int32 ErrorIdx) const
-{
-	if (ErrorIdx < MissingAttributes.Num())
-	{
-		return true;
-	}
-	return false;
-}
-
-bool UNiagaraStackRendererItem::TryFixError(int32 ErrorIdx)
-{
-	FNiagaraVariable MissingVar = MissingAttributes[ErrorIdx];
-	if (AddMissingVariable(GetEmitterViewModel()->GetEmitter(), MissingVar))
-	{
-		FNotificationInfo Info(FText::Format(LOCTEXT("AddedVariableForFix", "Added {0} to the Spawn script to support the renderer."), FText::FromName(MissingVar.GetName())));
-		Info.ExpireDuration = 5.0f;
-		Info.bFireAndForget = true;
-		Info.Image = FCoreStyle::Get().GetBrush(TEXT("MessageLog.Info"));
-		FSlateNotificationManager::Get().AddNotification(Info);
-		return true;
-	}
-	return false;
-}
-
-FText UNiagaraStackRendererItem::GetErrorText(int32 ErrorIdx) const
-{
-	if (ErrorIdx < MissingAttributes.Num())
-	{
-		const FNiagaraVariable& Attr = MissingAttributes[ErrorIdx];
-		return FText::Format(LOCTEXT("FailedRendererBind", "Missing attribute \"{0}\" of Type \"{1}\"."), FText::FromName(Attr.GetName()), Attr.GetType().GetNameText());
-	}
-	return FText();
+	FScopedTransaction ScopedTransaction(LOCTEXT("SetRendererEnabledState", "Set renderer enabled/disabled state."));
+	RendererObject->Modify();
+	RendererProperties->SetIsEnabled(bInIsEnabled);
 }
 
 void UNiagaraStackRendererItem::BeginDestroy()
@@ -247,35 +214,55 @@ void UNiagaraStackRendererItem::BeginDestroy()
 	Super::BeginDestroy();
 }
 
-void UNiagaraStackRendererItem::RefreshChildrenInternal(const TArray<UNiagaraStackEntry*>& CurrentChildren, TArray<UNiagaraStackEntry*>& NewChildren)
+void UNiagaraStackRendererItem::RefreshChildrenInternal(const TArray<UNiagaraStackEntry*>& CurrentChildren, TArray<UNiagaraStackEntry*>& NewChildren, TArray<FStackIssue>& NewIssues)
 {
 	if (RendererObject == nullptr)
 	{
 		RendererObject = NewObject<UNiagaraStackObject>(this);
-		RendererObject->Initialize(GetSystemViewModel(), GetEmitterViewModel(), RendererProperties.Get());
-	}
-
-	if (RendererExpander == nullptr)
-	{
-		RendererExpander = NewObject<UNiagaraStackItemExpander>(this);
-		RendererExpander->Initialize(GetSystemViewModel(), GetEmitterViewModel(), GetStackEditorData(), RendererProperties->GetName(), false);
-		RendererExpander->SetOnExpnadedChanged(UNiagaraStackItemExpander::FOnExpandedChanged::CreateUObject(this, &UNiagaraStackRendererItem::RendererExpandedChanged));
+		RendererObject->Initialize(CreateDefaultChildRequiredData(), RendererProperties.Get(), GetStackEditorDataKey());
 	}
-	
-	if (GetStackEditorData().GetStackEntryIsExpanded(RendererProperties->GetName(), false))
-	{
-		NewChildren.Add(RendererObject);
-	}
-
-	NewChildren.Add(RendererExpander);
 
+	NewChildren.Add(RendererObject);
 	MissingAttributes = GetMissingVariables(RendererProperties.Get(), GetEmitterViewModel()->GetEmitter());
 	bCanResetToBase.Reset();
+	Super::RefreshChildrenInternal(CurrentChildren, NewChildren, NewIssues);
+	
+	RefreshIssues(NewIssues);
 }
 
-void UNiagaraStackRendererItem::RendererExpandedChanged()
+void UNiagaraStackRendererItem::RefreshIssues(TArray<FStackIssue>& NewIssues)
 {
-	RefreshChildren();
+	for (FNiagaraVariable Attribute : MissingAttributes)
+	{
+		UNiagaraStackEntry::FStackIssue MissingAttributeError;
+		MissingAttributeError.ShortDescription = LOCTEXT("FailedRendererBindShort", "An attribute is missing.");
+		MissingAttributeError.LongDescription = FText::Format(LOCTEXT("FailedRendererBind", "Missing attribute \"{0}\" of Type \"{1}\"."), FText::FromName(Attribute.GetName()), Attribute.GetType().GetNameText());
+		MissingAttributeError.UniqueIdentifier = FName(*FString::Printf(TEXT("%s-AttributeMissing-%s"), *GetStackEditorDataKey(), *Attribute.GetName().ToString()));
+		UNiagaraStackEntry::FStackIssueFix Fix;
+		Fix.Description = LOCTEXT("AddMissingVariable", "Add missing variable");
+		Fix.FixDelegate.BindLambda([=]()
+		{
+			FScopedTransaction ScopedTransaction(Fix.Description);
+			if (AddMissingVariable(GetEmitterViewModel()->GetEmitter(), Attribute))
+			{
+				FNotificationInfo Info(FText::Format(LOCTEXT("AddedVariableForFix", "Added {0} to the Spawn script to support the renderer."), FText::FromName(Attribute.GetName())));
+				Info.ExpireDuration = 5.0f;
+				Info.bFireAndForget = true;
+				Info.Image = FCoreStyle::Get().GetBrush(TEXT("MessageLog.Info"));
+				FSlateNotificationManager::Get().AddNotification(Info);
+			}
+		});
+		MissingAttributeError.Fixes.Add(Fix);
+		NewIssues.Add(MissingAttributeError);
+	}
+	if (RendererProperties->GetIsEnabled() && !RendererProperties->IsSimTargetSupported(GetEmitterViewModel()->GetEmitter()->SimTarget))
+	{
+		UNiagaraStackEntry::FStackIssue TargetSupportError;
+		TargetSupportError.ShortDescription = LOCTEXT("FailedRendererDueToSimTarget", "Renderer incompatible with SimTarget mode.");
+		TargetSupportError.LongDescription = FText::Format(LOCTEXT("FailedRendererDueToSimTargetLong", "Renderer incompatible with SimTarget mode \"{0}\"."), (int32)GetEmitterViewModel()->GetEmitter()->SimTarget);
+		TargetSupportError.UniqueIdentifier = FName(*FString::Printf(TEXT("%s-SimTargetModeError-%d"), *GetStackEditorDataKey(), (int32)GetEmitterViewModel()->GetEmitter()->SimTarget));
+		NewIssues.Add(TargetSupportError);
+	}
 }
 
 void UNiagaraStackRendererItem::RendererChanged()
diff --git a/Engine/Plugins/FX/Niagara/Source/NiagaraEditor/Private/ViewModels/Stack/NiagaraStackRoot.cpp b/Engine/Plugins/FX/Niagara/Source/NiagaraEditor/Private/ViewModels/Stack/NiagaraStackRoot.cpp
--- a/Engine/Plugins/FX/Niagara/Source/NiagaraEditor/Private/ViewModels/Stack/NiagaraStackRoot.cpp
+++ b/Engine/Plugins/FX/Niagara/Source/NiagaraEditor/Private/ViewModels/Stack/NiagaraStackRoot.cpp
@@ -6,14 +6,15 @@
 #include "ViewModels/Stack/NiagaraStackRenderItemGroup.h"
 #include "ViewModels/Stack/NiagaraStackEventHandlerGroup.h"
 #include "ViewModels/Stack/NiagaraStackEventScriptItemGroup.h"
-#include "NiagaraSystemViewModel.h"
-#include "NiagaraEmitterViewModel.h"
+#include "ViewModels/NiagaraSystemViewModel.h"
+#include "ViewModels/NiagaraEmitterViewModel.h"
 #include "NiagaraSystemScriptViewModel.h"
 #include "NiagaraScriptViewModel.h"
 #include "NiagaraSystem.h"
 #include "NiagaraSystemEditorData.h"
 #include "NiagaraEmitterEditorData.h"
 #include "ViewModels/Stack/NiagaraStackParameterStoreGroup.h"
+#include "ViewModels/Stack/NiagaraStackSpacer.h"
 
 #define LOCTEXT_NAMESPACE "NiagaraStackViewModel"
 
@@ -28,9 +29,9 @@ UNiagaraStackRoot::UNiagaraStackRoot()
 {
 }
 
-void UNiagaraStackRoot::Initialize(TSharedRef<FNiagaraSystemViewModel> InSystemViewModel, TSharedRef<FNiagaraEmitterViewModel> InEmitterViewModel)
+void UNiagaraStackRoot::Initialize(FRequiredEntryData InRequiredEntryData)
 {
-	Super::Initialize(InSystemViewModel, InEmitterViewModel);
+	Super::Initialize(InRequiredEntryData, FString());
 	EmitterSpawnGroup = nullptr;
 	EmitterUpdateGroup = nullptr;
 	ParticleSpawnGroup = nullptr;
@@ -40,13 +41,17 @@ void UNiagaraStackRoot::Initialize(TSharedRef<FNiagaraSystemViewModel> InSystemV
 	SystemExposedVariablesGroup = nullptr;
 }
 
+bool UNiagaraStackRoot::GetCanExpand() const
+{
+	return false;
+}
 
 bool UNiagaraStackRoot::GetShouldShowInStack() const
 {
 	return false;
 }
 
-void UNiagaraStackRoot::RefreshChildrenInternal(const TArray<UNiagaraStackEntry*>& CurrentChildren, TArray<UNiagaraStackEntry*>& NewChildren)
+void UNiagaraStackRoot::RefreshChildrenInternal(const TArray<UNiagaraStackEntry*>& CurrentChildren, TArray<UNiagaraStackEntry*>& NewChildren, TArray<FStackIssue>& NewIssues)
 {
 	// We only allow displaying and editing system stacks if the system isn't transient which is the case in the emitter editor.
 	bool bShowSystemGroups = GetSystemViewModel()->GetEditMode() == ENiagaraSystemViewModelEditMode::SystemAsset;
@@ -55,85 +60,129 @@ void UNiagaraStackRoot::RefreshChildrenInternal(const TArray<UNiagaraStackEntry*
 	if (bShowSystemGroups && SystemExposedVariablesGroup == nullptr)
 	{
 		SystemExposedVariablesGroup = NewObject<UNiagaraStackParameterStoreGroup>(this);
-		SystemExposedVariablesGroup->Initialize(GetSystemViewModel(), GetEmitterViewModel(), GetSystemViewModel()->GetOrCreateEditorData().GetStackEditorData(), GetSystemViewModel()->GetSystemScriptViewModel(), &GetSystemViewModel()->GetSystem(), &GetSystemViewModel()->GetSystem().GetExposedParameters());
-		SystemExposedVariablesGroup->SetDisplayName(LOCTEXT("SystemExposedVariablesGroup", "System Exposed Variables"));
-		SystemExposedVariablesGroup->SetTooltipText(LOCTEXT("SystemExposedVariablesGroupToolTip", "Displays the variables created in the User namespace. These variables are exposed to owning UComponents, blueprints, etc."));
+		FRequiredEntryData RequiredEntryData(GetSystemViewModel(), GetEmitterViewModel(),
+			FExecutionCategoryNames::System, FExecutionSubcategoryNames::Parameters,
+			GetSystemViewModel()->GetOrCreateEditorData().GetStackEditorData());
+		SystemExposedVariablesGroup->Initialize(RequiredEntryData, &GetSystemViewModel()->GetSystem(), &GetSystemViewModel()->GetSystem().GetExposedParameters());
 	}
 
 	if (bShowSystemGroups && SystemSpawnGroup == nullptr)
 	{
 		SystemSpawnGroup = NewObject<UNiagaraStackScriptItemGroup>(this);
-		SystemSpawnGroup->Initialize(GetSystemViewModel(), GetEmitterViewModel(), GetSystemViewModel()->GetOrCreateEditorData().GetStackEditorData(), GetSystemViewModel()->GetSystemScriptViewModel(), ENiagaraScriptUsage::SystemSpawnScript);
-		SystemSpawnGroup->SetDisplayName(LOCTEXT("SystemSpawnGroupName", "System Spawn"));
-		SystemSpawnGroup->SetTooltipText(LOCTEXT("SystemSpawnGroupToolTip", "Occurs once at System creation on the CPU. Modules in this section should initialize defaults and/or do initial setup.\r\nModules are executed in order from top to bottom of the stack."));
+		FRequiredEntryData RequiredEntryData(GetSystemViewModel(), GetEmitterViewModel(),
+			FExecutionCategoryNames::System, FExecutionSubcategoryNames::Spawn,
+			GetSystemViewModel()->GetOrCreateEditorData().GetStackEditorData());
+		FText DisplayName = LOCTEXT("SystemSpawnGroupName", "System Spawn");
+		FText ToolTip = LOCTEXT("SystemSpawnGroupToolTip", "Occurs once at System creation on the CPU. Modules in this section should initialize defaults and/or do initial setup.\r\nModules are executed in order from top to bottom of the stack.");
+		SystemSpawnGroup->Initialize(RequiredEntryData, DisplayName, ToolTip, GetSystemViewModel()->GetSystemScriptViewModel().ToSharedRef(), ENiagaraScriptUsage::SystemSpawnScript);
 	}
 
 	if (bShowSystemGroups && SystemUpdateGroup == nullptr)
 	{
 		SystemUpdateGroup = NewObject<UNiagaraStackScriptItemGroup>(this);
-		SystemUpdateGroup->Initialize(GetSystemViewModel(), GetEmitterViewModel(), GetSystemViewModel()->GetOrCreateEditorData().GetStackEditorData(), GetSystemViewModel()->GetSystemScriptViewModel(), ENiagaraScriptUsage::SystemUpdateScript);
-		SystemUpdateGroup->SetDisplayName(LOCTEXT("SystemUpdateGroupName", "System Update"));
-		SystemUpdateGroup->SetTooltipText(LOCTEXT("SystemUpdateGroupToolTip", "Occurs every Emitter tick on the CPU.Modules in this section should compute values for parameters for emitter or particle update or spawning this frame.\r\nModules are executed in order from top to bottom of the stack."));
+		FRequiredEntryData RequiredEntryData(GetSystemViewModel(), GetEmitterViewModel(),
+			FExecutionCategoryNames::System, FExecutionSubcategoryNames::Update,
+			GetSystemViewModel()->GetOrCreateEditorData().GetStackEditorData());
+		FText DisplayName = LOCTEXT("SystemUpdateGroupName", "System Update");
+		FText ToolTip = LOCTEXT("SystemUpdateGroupToolTip", "Occurs every Emitter tick on the CPU.Modules in this section should compute values for parameters for emitter or particle update or spawning this frame.\r\nModules are executed in order from top to bottom of the stack.");
+		SystemUpdateGroup->Initialize(RequiredEntryData, DisplayName, ToolTip, GetSystemViewModel()->GetSystemScriptViewModel().ToSharedRef(), ENiagaraScriptUsage::SystemUpdateScript);
 	}
 
 	if (EmitterSpawnGroup == nullptr)
 	{
 		EmitterSpawnGroup = NewObject<UNiagaraStackEmitterSpawnScriptItemGroup>(this);
-		EmitterSpawnGroup->Initialize(GetSystemViewModel(), GetEmitterViewModel(), GetEmitterViewModel()->GetOrCreateEditorData().GetStackEditorData(), GetEmitterViewModel()->GetSharedScriptViewModel(), ENiagaraScriptUsage::EmitterSpawnScript);
-		EmitterSpawnGroup->SetDisplayName(LOCTEXT("EmitterSpawnGroupName", "Emitter Spawn"));
-		EmitterSpawnGroup->SetTooltipText(LOCTEXT("EmitterSpawnGroupTooltip", "Occurs once at Emitter creation on the CPU. Modules in this section should initialize defaults and/or do initial setup.\r\nModules are executed in order from top to bottom of the stack."));
+		FRequiredEntryData RequiredEntryData(GetSystemViewModel(), GetEmitterViewModel(),
+			FExecutionCategoryNames::Emitter, FExecutionSubcategoryNames::Spawn,
+			GetEmitterViewModel()->GetOrCreateEditorData().GetStackEditorData());
+		FText DisplayName = LOCTEXT("EmitterSpawnGroupName", "Emitter Spawn");
+		FText ToolTip = LOCTEXT("EmitterSpawnGroupTooltip", "Occurs once at Emitter creation on the CPU. Modules in this section should initialize defaults and/or do initial setup.\r\nModules are executed in order from top to bottom of the stack.");
+		EmitterSpawnGroup->Initialize(RequiredEntryData, DisplayName, ToolTip, GetEmitterViewModel()->GetSharedScriptViewModel(), ENiagaraScriptUsage::EmitterSpawnScript);
 	}
 
 	if (EmitterUpdateGroup == nullptr)
 	{
 		EmitterUpdateGroup = NewObject<UNiagaraStackScriptItemGroup>(this);
-		EmitterUpdateGroup->Initialize(GetSystemViewModel(), GetEmitterViewModel(), GetEmitterViewModel()->GetOrCreateEditorData().GetStackEditorData(), GetEmitterViewModel()->GetSharedScriptViewModel(), ENiagaraScriptUsage::EmitterUpdateScript);
-		EmitterUpdateGroup->SetDisplayName(LOCTEXT("EmitterUpdateGroupName", "Emitter Update"));
-		EmitterUpdateGroup->SetTooltipText(LOCTEXT("EmitterUpdateGroupTooltip", "Occurs every Emitter tick on the CPU. Modules in this section should compute values for parameters for Particle Update or Spawning this frame.\r\nModules are executed in order from top to bottom of the stack."));
+		FRequiredEntryData RequiredEntryData(GetSystemViewModel(), GetEmitterViewModel(),
+			FExecutionCategoryNames::Emitter, FExecutionSubcategoryNames::Update,
+			GetEmitterViewModel()->GetOrCreateEditorData().GetStackEditorData());
+		FText DisplayName = LOCTEXT("EmitterUpdateGroupName", "Emitter Update");
+		FText ToolTip = LOCTEXT("EmitterUpdateGroupTooltip", "Occurs every Emitter tick on the CPU. Modules in this section should compute values for parameters for Particle Update or Spawning this frame.\r\nModules are executed in order from top to bottom of the stack.");
+		EmitterUpdateGroup->Initialize(RequiredEntryData, DisplayName, ToolTip, GetEmitterViewModel()->GetSharedScriptViewModel(), ENiagaraScriptUsage::EmitterUpdateScript);
 	}
 
 	if (ParticleSpawnGroup == nullptr)
 	{
 		ParticleSpawnGroup = NewObject<UNiagaraStackScriptItemGroup>(this);
-		ParticleSpawnGroup->Initialize(GetSystemViewModel(), GetEmitterViewModel(), GetEmitterViewModel()->GetOrCreateEditorData().GetStackEditorData(), GetEmitterViewModel()->GetSharedScriptViewModel(), ENiagaraScriptUsage::ParticleSpawnScript);
-		ParticleSpawnGroup->SetDisplayName(LOCTEXT("ParticleSpawnGroupName", "Particle Spawn"));
-		ParticleSpawnGroup->SetTooltipText(LOCTEXT("ParticleSpawnGroupTooltip", "Called once per created particle. Modules in this section should set up initial values for each particle.\r\nIf \"Use Interpolated Spawning\" is set, we will also run the Particle Update script after the Particle Spawn script.\r\nModules are executed in order from top to bottom of the stack."));
+		FRequiredEntryData RequiredEntryData(GetSystemViewModel(), GetEmitterViewModel(),
+			FExecutionCategoryNames::Particle, FExecutionSubcategoryNames::Spawn,
+			GetEmitterViewModel()->GetOrCreateEditorData().GetStackEditorData());
+		FText DisplayName = LOCTEXT("ParticleSpawnGroupName", "Particle Spawn");
+		FText ToolTip = LOCTEXT("ParticleSpawnGroupTooltip", "Called once per created particle. Modules in this section should set up initial values for each particle.\r\nIf \"Use Interpolated Spawning\" is set, we will also run the Particle Update script after the Particle Spawn script.\r\nModules are executed in order from top to bottom of the stack.");
+		ParticleSpawnGroup->Initialize(RequiredEntryData, DisplayName, ToolTip, GetEmitterViewModel()->GetSharedScriptViewModel(), ENiagaraScriptUsage::ParticleSpawnScript);
 	}
 
 	if (ParticleUpdateGroup == nullptr)
 	{
 		ParticleUpdateGroup = NewObject<UNiagaraStackScriptItemGroup>(this);
-		ParticleUpdateGroup->Initialize(GetSystemViewModel(), GetEmitterViewModel(), GetEmitterViewModel()->GetOrCreateEditorData().GetStackEditorData(), GetEmitterViewModel()->GetSharedScriptViewModel(), ENiagaraScriptUsage::ParticleUpdateScript);
-		ParticleUpdateGroup->SetDisplayName(LOCTEXT("ParticleUpdateGroupName", "Particle Update"));
-		ParticleUpdateGroup->SetTooltipText(LOCTEXT("ParticleUpdateGroupTooltip", "Called every frame per particle. Modules in this section should update new values for this frame.\r\nModules are executed in order from top to bottom of the stack."));
+		FRequiredEntryData RequiredEntryData(GetSystemViewModel(), GetEmitterViewModel(),
+			FExecutionCategoryNames::Particle, FExecutionSubcategoryNames::Update,
+			GetEmitterViewModel()->GetOrCreateEditorData().GetStackEditorData());
+		FText DisplayName = LOCTEXT("ParticleUpdateGroupName", "Particle Update");
+		FText ToolTip = LOCTEXT("ParticleUpdateGroupTooltip", "Called every frame per particle. Modules in this section should update new values for this frame.\r\nModules are executed in order from top to bottom of the stack.");
+		ParticleUpdateGroup->Initialize(RequiredEntryData, DisplayName, ToolTip, GetEmitterViewModel()->GetSharedScriptViewModel(), ENiagaraScriptUsage::ParticleUpdateScript);
 	}
 
 	if (AddEventHandlerGroup == nullptr)
 	{
 		AddEventHandlerGroup = NewObject<UNiagaraStackEventHandlerGroup>(this);
-		AddEventHandlerGroup->Initialize(GetSystemViewModel(), GetEmitterViewModel(), GetEmitterViewModel()->GetOrCreateEditorData().GetStackEditorData());
-		AddEventHandlerGroup->SetDisplayName(LOCTEXT("EventGroupName", "Add Event Handler"));
-		AddEventHandlerGroup->SetTooltipText(LOCTEXT("EventGroupTooltip", "Determines how this Emitter responds to incoming events. There can be more than one event handler script stack per Emitter."));
+		FRequiredEntryData RequiredEntryData(GetSystemViewModel(), GetEmitterViewModel(),
+			FExecutionCategoryNames::Particle, FExecutionSubcategoryNames::Event,
+			GetEmitterViewModel()->GetOrCreateEditorData().GetStackEditorData());
+		AddEventHandlerGroup->Initialize(RequiredEntryData);
 		AddEventHandlerGroup->SetOnItemAdded(UNiagaraStackEventHandlerGroup::FOnItemAdded::CreateUObject(this, &UNiagaraStackRoot::EmitterEventArraysChanged));
 	}
 
 	if (RenderGroup == nullptr)
 	{
 		RenderGroup = NewObject<UNiagaraStackRenderItemGroup>(this);
-		RenderGroup->Initialize(GetSystemViewModel(), GetEmitterViewModel(), GetEmitterViewModel()->GetOrCreateEditorData().GetStackEditorData());
-		RenderGroup->SetDisplayName(LOCTEXT("RenderGroupName", "Render"));
-		RenderGroup->SetTooltipText(LOCTEXT("RendererGroupTooltip", "Describes how we should display/present each particle. Note that this doesn't have to be visual. Multiple renderers are supported. Order in this stack is not necessarily relevant to draw order."));
+		FRequiredEntryData RequiredEntryData(GetSystemViewModel(), GetEmitterViewModel(),
+			FExecutionCategoryNames::Render, NAME_None,
+			GetEmitterViewModel()->GetOrCreateEditorData().GetStackEditorData());
+		RenderGroup->Initialize(RequiredEntryData);
 	}
 
+	auto GetOrCreateSpacer = [&](FName SpacerExecutionCategory, FName SpacerKey)
+	{
+		UNiagaraStackSpacer* Spacer = FindCurrentChildOfTypeByPredicate<UNiagaraStackSpacer>(CurrentChildren,
+			[SpacerKey](UNiagaraStackSpacer* CurrentSpacer) { return CurrentSpacer->GetSpacerKey() == SpacerKey; });
+
+		if (Spacer == nullptr)
+		{
+			Spacer = NewObject<UNiagaraStackSpacer>(this);
+			FRequiredEntryData RequiredEntryData(GetSystemViewModel(), GetEmitterViewModel(),
+				SpacerExecutionCategory, NAME_None,
+				GetSystemViewModel()->GetEditorData().GetStackEditorData());
+			Spacer->Initialize(RequiredEntryData, SpacerKey, 1.0f);
+		}
+
+		return Spacer;
+	};
+
 	// Populate new children
 	if (bShowSystemGroups)
 	{
 		NewChildren.Add(SystemExposedVariablesGroup);
 		NewChildren.Add(SystemSpawnGroup);
 		NewChildren.Add(SystemUpdateGroup);
+		NewChildren.Add(GetOrCreateSpacer(FExecutionCategoryNames::System, "SystemFooter"));
+		NewChildren.Add(GetOrCreateSpacer(NAME_None, "SystemSpacer"));
 	}
+
 	NewChildren.Add(EmitterSpawnGroup);
 	NewChildren.Add(EmitterUpdateGroup);
+	NewChildren.Add(GetOrCreateSpacer(FExecutionCategoryNames::Emitter, "EmitterFooter"));
+	NewChildren.Add(GetOrCreateSpacer(NAME_None, "EmitterSpacer"));
+
 	NewChildren.Add(ParticleSpawnGroup);
 	NewChildren.Add(ParticleUpdateGroup);
 
@@ -145,16 +194,22 @@ void UNiagaraStackRoot::RefreshChildrenInternal(const TArray<UNiagaraStackEntry*
 		if (EventHandlerGroup == nullptr)
 		{
 			EventHandlerGroup = NewObject<UNiagaraStackEventScriptItemGroup>(this, NAME_None, RF_Transactional);
-			EventHandlerGroup->Initialize(GetSystemViewModel(), GetEmitterViewModel(), GetEmitterViewModel()->GetEditorData().GetStackEditorData(), GetEmitterViewModel()->GetSharedScriptViewModel(), ENiagaraScriptUsage::ParticleEventScript, EventScriptProperties.Script->GetUsageId());
+			FRequiredEntryData RequiredEntryData(GetSystemViewModel(), GetEmitterViewModel(),
+				FExecutionCategoryNames::Particle, FExecutionSubcategoryNames::Event,
+				GetEmitterViewModel()->GetEditorData().GetStackEditorData());
+			EventHandlerGroup->Initialize(RequiredEntryData, GetEmitterViewModel()->GetSharedScriptViewModel(), ENiagaraScriptUsage::ParticleEventScript, EventScriptProperties.Script->GetUsageId());
 			EventHandlerGroup->SetOnModifiedEventHandlers(UNiagaraStackEventScriptItemGroup::FOnModifiedEventHandlers::CreateUObject(this, &UNiagaraStackRoot::EmitterEventArraysChanged));
-			EventHandlerGroup->SetTooltipText(LOCTEXT("EventGroupTooltip", "Determines how this Emitter responds to incoming events. There can be more than one event handler script stack per Emitter."));
 		}
 
 		NewChildren.Add(EventHandlerGroup);
 	}
 
 	NewChildren.Add(AddEventHandlerGroup);
+	NewChildren.Add(GetOrCreateSpacer(FExecutionCategoryNames::Particle, "ParticleFooter"));
+	NewChildren.Add(GetOrCreateSpacer(NAME_None, "ParticleSpacer"));
+
 	NewChildren.Add(RenderGroup);
+	NewChildren.Add(GetOrCreateSpacer(FExecutionCategoryNames::Render, "RenderFooter"));
 }
 
 void UNiagaraStackRoot::EmitterEventArraysChanged()
