Copying //UE4/Dev-Niagara to //UE4/Dev-Main (Source: //UE4/Dev-Niagara @ 4074996)

diff --git a/Engine/Plugins/FX/Niagara/Source/NiagaraEditor/Private/ViewModels/Stack/NiagaraStackScriptItemGroup.cpp b/Engine/Plugins/FX/Niagara/Source/NiagaraEditor/Private/ViewModels/Stack/NiagaraStackScriptItemGroup.cpp
--- a/Engine/Plugins/FX/Niagara/Source/NiagaraEditor/Private/ViewModels/Stack/NiagaraStackScriptItemGroup.cpp
+++ b/Engine/Plugins/FX/Niagara/Source/NiagaraEditor/Private/ViewModels/Stack/NiagaraStackScriptItemGroup.cpp
@@ -3,98 +3,351 @@
 #include "ViewModels/Stack/NiagaraStackScriptItemGroup.h"
 #include "ViewModels/Stack/NiagaraStackModuleItem.h"
 #include "ViewModels/Stack/NiagaraStackSpacer.h"
-#include "ViewModels/Stack/NiagaraStackAddScriptModuleItem.h"
+#include "ViewModels/NiagaraEmitterViewModel.h"
 #include "NiagaraScriptViewModel.h"
 #include "NiagaraScriptGraphViewModel.h"
 #include "NiagaraGraph.h"
 #include "NiagaraNodeOutput.h"
 #include "NiagaraNodeFunctionCall.h"
+#include "NiagaraNodeAssignment.h"
 #include "EdGraphSchema_Niagara.h"
 #include "ViewModels/Stack/NiagaraStackErrorItem.h"
-#include "Internationalization/Internationalization.h"
 #include "ViewModels/Stack/NiagaraStackGraphUtilities.h"
+#include "NiagaraConstants.h"
+#include "NiagaraStackEditorData.h"
+
+#include "Internationalization/Internationalization.h"
 #include "ScopedTransaction.h"
+#include "Widgets/Notifications/SNotificationList.h"
+#include "Framework/Notifications/NotificationManager.h"
+#include "Internationalization/Internationalization.h"
 
 #define LOCTEXT_NAMESPACE "UNiagaraStackScriptItemGroup"
 
+class FScriptGroupAddAction : public INiagaraStackItemGroupAddAction
+{
+public:
+	static TSharedRef<FScriptGroupAddAction> CreateAssetModuleAction(FAssetData AssetData)
+	{
+		FText Category;
+		AssetData.GetTagValue(GET_MEMBER_NAME_CHECKED(UNiagaraScript, Category), Category);
+		if (Category.IsEmptyOrWhitespace())
+		{
+			Category = LOCTEXT("ModuleNotCategorized", "Uncategorized Modules");
+		}
+
+		FString DisplayNameString = FName::NameToDisplayString(AssetData.AssetName.ToString(), false);
+		FText DisplayName = FText::FromString(DisplayNameString);
+
+		FText AssetDescription;
+		AssetData.GetTagValue(GET_MEMBER_NAME_CHECKED(UNiagaraScript, Description), AssetDescription);
+		FText Description = FText::Format(LOCTEXT("AssetModuleDescriptionFormat", "Path: {0}\nDescription: {1}"), FText::FromString(AssetData.ObjectPath.ToString()), AssetDescription);
+
+		return MakeShareable(new FScriptGroupAddAction(Category, DisplayName, Description, FNiagaraVariable(), false, AssetData, false));
+	}
+
+	static TSharedRef<FScriptGroupAddAction> CreateExistingParameterModuleAction(FNiagaraVariable ParameterVariable)
+	{
+		FText Category = LOCTEXT("ExistingParameterModuleCategory", "Set Specific Parameters");
+
+		FString DisplayNameString = FName::NameToDisplayString(ParameterVariable.GetName().ToString(), false);
+		FText DisplayName = FText::FromString(DisplayNameString);
+
+		FText AttributeDescription = FNiagaraConstants::GetAttributeDescription(ParameterVariable);
+		FText Description = FText::Format(LOCTEXT("ExistingParameterModuleDescriptoinFormat", "Description: Set the parameter {0}. {1}"), FText::FromName(ParameterVariable.GetName()), AttributeDescription);
+
+		return MakeShareable(new FScriptGroupAddAction(Category, DisplayName, Description, ParameterVariable, false, FAssetData(), false));
+	}
+
+	static TSharedRef<FScriptGroupAddAction> CreateNewParameterModuleAction(FName NewParameterNamespace, FNiagaraTypeDefinition NewParameterType)
+	{
+		FText Category = LOCTEXT("NewParameterModuleCategory", "Create New Parameter");
+		FText DisplayName = NewParameterType.GetNameText();
+		FText Description = FText::Format(LOCTEXT("NewParameterModuleDescriptionFormat", "Description: Create a new {0} parameter."), DisplayName);
+
+		FNiagaraParameterHandle NewParameterHandle(NewParameterNamespace, *(TEXT("New") + NewParameterType.GetName()));
+		FNiagaraVariable NewParameter(NewParameterType, NewParameterHandle.GetParameterHandleString());
+
+		return MakeShareable(new FScriptGroupAddAction(Category, DisplayName, Description, NewParameter, true, FAssetData(), false));
+	}
+
+	virtual FText GetCategory() const override
+	{
+		return Category;
+	}
+
+	virtual FText GetDisplayName() const override
+	{
+		return DisplayName;
+	}
+
+	virtual FText GetDescription() const override
+	{
+		return Description;
+	}
+
+	const FNiagaraVariable& GetModuleParameterVariable() const
+	{
+		return ModuleParameterVariable;
+	}
+
+	bool GetRenameParameterOnAdd() const
+	{
+		return bRenameParameterOnAdd;
+	}
+
+	const FAssetData& GetModuleAssetData() const
+	{
+		return ModuleAssetData;
+	}
+
+	bool GetIsMaterialParameterModuleAction() const
+	{
+		return bIsMaterialParameterModuleAction;
+	}
+
+private:
+	FScriptGroupAddAction(FText InCategory, FText InDisplayName, FText InDescription, FNiagaraVariable InModuleParameterVariable, bool bInRenameParameterOnAdd, FAssetData InModuleAssetData, bool bInIsMaterialParameterModuleAction)
+		: Category(InCategory)
+		, DisplayName(InDisplayName)
+		, Description(InDescription)
+		, ModuleParameterVariable(InModuleParameterVariable)
+		, bRenameParameterOnAdd(bInRenameParameterOnAdd)
+		, ModuleAssetData(InModuleAssetData)
+		, bIsMaterialParameterModuleAction(bInIsMaterialParameterModuleAction)
+	{
+	}
+
+private:
+	FText Category;
+	FText DisplayName;
+	FText Description;
+	FNiagaraVariable ModuleParameterVariable;
+	bool bRenameParameterOnAdd;
+	FAssetData ModuleAssetData;
+	bool bIsMaterialParameterModuleAction;
+};
+
+class FScriptItemGroupAddUtilities : public FNiagaraStackItemGroupAddUtilities
+{
+public:
+	FScriptItemGroupAddUtilities(TSharedRef<FNiagaraSystemViewModel> InSystemViewModel, TSharedRef<FNiagaraEmitterViewModel> InEmitterViewModel, UNiagaraStackEditorData& InStackEditorData, FOnItemAdded InOnItemAdded)
+		: FNiagaraStackItemGroupAddUtilities(LOCTEXT("ScriptGroupAddItemName", "Module"), EAddMode::AddFromAction, false, InOnItemAdded)
+		, SystemViewModel(InSystemViewModel)
+		, EmitterViewModel(InEmitterViewModel)
+		, StackEditorData(InStackEditorData)
+	{
+	}
+
+	void SetOutputNode(UNiagaraNodeOutput* InOutputNode)
+	{
+		OutputNode = InOutputNode;
+	}
+
+	virtual void AddItemDirectly() override { unimplemented(); }
+
+	virtual void GenerateAddActions(TArray<TSharedRef<INiagaraStackItemGroupAddAction>>& OutAddActions) const override
+	{
+		if (SystemViewModel.IsValid() == false || EmitterViewModel.IsValid() == false || OutputNode == nullptr)
+		{
+			return;
+		}
+
+		// Generate actions for adding script asset modules.
+		TArray<FAssetData> ModuleAssets;
+		FNiagaraStackGraphUtilities::GetScriptAssetsByUsage(ENiagaraScriptUsage::Module, OutputNode->GetUsage(), ModuleAssets);
+		for(const FAssetData& ModuleAsset : ModuleAssets)
+		{
+			OutAddActions.Add(FScriptGroupAddAction::CreateAssetModuleAction(ModuleAsset));
+		}
+
+		// Generate actions for the available parameters to set.
+		TArray<FNiagaraVariable> AvailableParameters;
+		FNiagaraStackGraphUtilities::GetAvailableParametersForScript(*OutputNode, AvailableParameters);
+		for (const FNiagaraVariable& AvailableParameter : AvailableParameters)
+		{
+			OutAddActions.Add(FScriptGroupAddAction::CreateExistingParameterModuleAction(AvailableParameter));
+		}
+
+		// Generate actions for setting new typed parameters.
+		TOptional<FName> NewParameterNamespace = FNiagaraStackGraphUtilities::GetNamespaceForScriptUsage(OutputNode->GetUsage());
+		if (NewParameterNamespace.IsSet())
+		{
+			TArray<FNiagaraTypeDefinition> AvailableTypes;
+			FNiagaraStackGraphUtilities::GetNewParameterAvailableTypes(AvailableTypes);
+			for (const FNiagaraTypeDefinition& AvailableType : AvailableTypes)
+			{
+				OutAddActions.Add(FScriptGroupAddAction::CreateNewParameterModuleAction(NewParameterNamespace.GetValue(), AvailableType));
+			}
+		}
+	}
+
+	virtual void ExecuteAddAction(TSharedRef<INiagaraStackItemGroupAddAction> AddAction, int32 TargetIndex) override
+	{
+		TSharedRef<FScriptGroupAddAction> ScriptGroupAddAction = StaticCastSharedRef<FScriptGroupAddAction>(AddAction);
+		FScopedTransaction ScopedTransaction(LOCTEXT("InsertNewModule", "Insert new module"));
+		UNiagaraNodeFunctionCall* NewModuleNode = nullptr;
+		if (ScriptGroupAddAction->GetModuleAssetData().IsValid())
+		{
+			NewModuleNode = FNiagaraStackGraphUtilities::AddScriptModuleToStack(ScriptGroupAddAction->GetModuleAssetData(), *OutputNode, TargetIndex);
+		}
+		else if (ScriptGroupAddAction->GetModuleParameterVariable().IsValid())
+		{
+			NewModuleNode = AddParameterModule(ScriptGroupAddAction->GetModuleParameterVariable(), ScriptGroupAddAction->GetRenameParameterOnAdd(), TargetIndex);
+		}
+
+		checkf(NewModuleNode != nullptr, TEXT("Add module action failed"));
+		FNiagaraStackGraphUtilities::InitializeStackFunctionInputs(SystemViewModel.Pin().ToSharedRef(), EmitterViewModel.Pin().ToSharedRef(), StackEditorData, *NewModuleNode, *NewModuleNode);
+		FNiagaraStackGraphUtilities::RelayoutGraph(*OutputNode->GetGraph());
+		OnItemAdded.ExecuteIfBound();
+	}
+private:
+	UNiagaraNodeFunctionCall* AddScriptAssetModule(const FAssetData& AssetData, int32 TargetIndex)
+	{
+		return FNiagaraStackGraphUtilities::AddScriptModuleToStack(AssetData, *OutputNode, TargetIndex);
+	}
+
+	UNiagaraNodeFunctionCall* AddParameterModule(const FNiagaraVariable& ParameterVariable, bool bRenameParameterOnAdd, int32 TargetIndex)
+	{
+		TArray<FNiagaraVariable> Vars;
+		Vars.Add(ParameterVariable);
+		TArray<FString> DefaultVals;
+		DefaultVals.Add(FNiagaraConstants::GetAttributeDefaultValue(ParameterVariable));
+		UNiagaraNodeAssignment* NewAssignmentModule = FNiagaraStackGraphUtilities::AddParameterModuleToStack(Vars, *OutputNode, TargetIndex,DefaultVals );
+		
+		TArray<const UEdGraphPin*> InputPins;
+		FNiagaraStackGraphUtilities::GetStackFunctionInputPins(*NewAssignmentModule, InputPins);
+		if (InputPins.Num() == 1)
+		{
+			FString FunctionInputEditorDataKey = FNiagaraStackGraphUtilities::GenerateStackFunctionInputEditorDataKey(*NewAssignmentModule, InputPins[0]->PinName);
+			if (bRenameParameterOnAdd)
+			{
+				StackEditorData.SetModuleInputIsRenamePending(FunctionInputEditorDataKey, true);
+			}
+		}
+
+		return NewAssignmentModule;
+	}
+
+private:
+	UNiagaraNodeOutput* OutputNode;
+	TWeakPtr<FNiagaraSystemViewModel> SystemViewModel;
+	TWeakPtr<FNiagaraEmitterViewModel> EmitterViewModel;
+	UNiagaraStackEditorData& StackEditorData;
+};
+
 void UNiagaraStackScriptItemGroup::Initialize(
-	TSharedRef<FNiagaraSystemViewModel> InSystemViewModel,
-	TSharedRef<FNiagaraEmitterViewModel> InEmitterViewModel,
-	UNiagaraStackEditorData& InStackEditorData,
+	FRequiredEntryData InRequiredEntryData,
+	FText InDisplayName,
+	FText InToolTip,
 	TSharedRef<FNiagaraScriptViewModel> InScriptViewModel,
 	ENiagaraScriptUsage InScriptUsage,
 	FGuid InScriptUsageId)
 {
 	checkf(ScriptViewModel.IsValid() == false, TEXT("Can not set the script view model more than once."));
-	Super::Initialize(InSystemViewModel, InEmitterViewModel, InStackEditorData);
+	AddUtilities = MakeShared<FScriptItemGroupAddUtilities>(InRequiredEntryData.SystemViewModel, InRequiredEntryData.EmitterViewModel,
+		*InRequiredEntryData.StackEditorData, FNiagaraStackItemGroupAddUtilities::FOnItemAdded::CreateUObject(this, &UNiagaraStackScriptItemGroup::ItemAdded));
+	Super::Initialize(InRequiredEntryData, InDisplayName, InToolTip, AddUtilities.Get());
 	ScriptViewModel = InScriptViewModel;
 	ScriptUsage = InScriptUsage;
 	ScriptUsageId = InScriptUsageId;
 }
 
-FText UNiagaraStackScriptItemGroup::GetDisplayName() const
-{
-	return DisplayName;
-}
-
-void UNiagaraStackScriptItemGroup::SetDisplayName(FText InDisplayName)
-{
-	DisplayName = InDisplayName;
-}
-
-void UNiagaraStackScriptItemGroup::RefreshChildrenInternal(const TArray<UNiagaraStackEntry*>& CurrentChildren, TArray<UNiagaraStackEntry*>& NewChildren)
+void UNiagaraStackScriptItemGroup::RefreshChildrenInternal(const TArray<UNiagaraStackEntry*>& CurrentChildren, TArray<UNiagaraStackEntry*>& NewChildren, TArray<FStackIssue>& NewIssues)
 {
 	TSharedPtr<FNiagaraScriptViewModel> ScriptViewModelPinned = ScriptViewModel.Pin();
 	checkf(ScriptViewModelPinned.IsValid(), TEXT("Can not refresh children when the script view model has been deleted."));
 
+	StackSpacerToModuleItemMap.Empty();
+
 	UNiagaraGraph* Graph = ScriptViewModelPinned->GetGraphViewModel()->GetGraph();
 	FText ErrorMessage;
-	if (FNiagaraStackGraphUtilities::ValidateGraphForOutput(*Graph, ScriptUsage, ScriptUsageId, ErrorMessage) == false)
-	{
-		UE_LOG(LogNiagaraEditor, Error, TEXT("Failed to Create Stack.  Message: %s"), *ErrorMessage.ToString());
-		Error = MakeShared<FError>();
-		Error->ErrorText = LOCTEXT("InvalidErrorText", "The data used to generate the stack has been corrupted and can not be used.\nUsing the fix option will reset this part of the stack to its default empty state.");
-		Error->ErrorSummaryText = LOCTEXT("InvalidErrorSummaryText", "The stack data is invalid");
-		Error->Fix.BindLambda([=]()
-		{
-			FScopedTransaction ScopedTransaction(LOCTEXT("FixStackGraph", "Fix invalid stack graph"));
-			FNiagaraStackGraphUtilities::ResetGraphForOutput(*Graph, ScriptUsage, ScriptUsageId);
-			FNiagaraStackGraphUtilities::RelayoutGraph(*Graph);
-		});
-	}
-	else
+	bIsValidForOutput = false;
+	if (FNiagaraStackGraphUtilities::ValidateGraphForOutput(*Graph, ScriptUsage, ScriptUsageId, ErrorMessage) == true)
 	{
+		bIsValidForOutput = true;
+
 		UNiagaraNodeOutput* MatchingOutputNode = Graph->FindOutputNode(ScriptUsage, ScriptUsageId);
+		AddUtilities->SetOutputNode(MatchingOutputNode);
+
 		TArray<UNiagaraNodeFunctionCall*> ModuleNodes;
 		FNiagaraStackGraphUtilities::GetOrderedModuleNodes(*MatchingOutputNode, ModuleNodes);
 		int32 ModuleIndex = 0;
 		for (UNiagaraNodeFunctionCall* ModuleNode : ModuleNodes)
 		{
-			UNiagaraStackAddScriptModuleItem* AddModuleAtIndexItem = FindCurrentChildOfTypeByPredicate<UNiagaraStackAddScriptModuleItem>(CurrentChildren,
-				[=](UNiagaraStackAddScriptModuleItem* CurrentAddModuleItem) { return CurrentAddModuleItem->GetTargetIndex() == ModuleIndex; });
+			FName ModuleSpacerKey = *FString::Printf(TEXT("Module%i"), ModuleIndex);
+			UNiagaraStackSpacer* ModuleSpacer = FindCurrentChildOfTypeByPredicate<UNiagaraStackSpacer>(CurrentChildren,
+				[=](UNiagaraStackSpacer* CurrentModuleSpacer) { return CurrentModuleSpacer->GetSpacerKey() == ModuleSpacerKey; });
 
-			if (AddModuleAtIndexItem == nullptr)
+			if (ModuleSpacer == nullptr)
 			{
-				AddModuleAtIndexItem = NewObject<UNiagaraStackAddScriptModuleItem>(this);
-				AddModuleAtIndexItem->Initialize(GetSystemViewModel(), GetEmitterViewModel(), GetStackEditorData(), *MatchingOutputNode, ModuleIndex);
-				AddModuleAtIndexItem->SetOnItemAdded(UNiagaraStackAddModuleItem::FOnItemAdded::CreateUObject(this, &UNiagaraStackScriptItemGroup::ItemAdded));
+				ModuleSpacer = NewObject<UNiagaraStackSpacer>(this);
+				ModuleSpacer->Initialize(CreateDefaultChildRequiredData(), ModuleSpacerKey, 1.4f);
 			}
 
+			NewChildren.Add(ModuleSpacer);
+
 			UNiagaraStackModuleItem* ModuleItem = FindCurrentChildOfTypeByPredicate<UNiagaraStackModuleItem>(CurrentChildren, 
 				[=](UNiagaraStackModuleItem* CurrentModuleItem) { return &CurrentModuleItem->GetModuleNode() == ModuleNode; });
 
 			if (ModuleItem == nullptr)
 			{
 				ModuleItem = NewObject<UNiagaraStackModuleItem>(this);
-				ModuleItem->Initialize(GetSystemViewModel(), GetEmitterViewModel(), GetStackEditorData(), *ModuleNode);
+				ModuleItem->Initialize(CreateDefaultChildRequiredData(), GetAddUtilities(), *ModuleNode);
 				ModuleItem->SetOnModifiedGroupItems(UNiagaraStackModuleItem::FOnModifiedGroupItems::CreateUObject(this, &UNiagaraStackScriptItemGroup::ChildModifiedGroupItems));
 			}
 
-			NewChildren.Add(AddModuleAtIndexItem);
 			NewChildren.Add(ModuleItem);
+			StackSpacerToModuleItemMap.Add(FObjectKey(ModuleSpacer), ModuleItem);
+
 			ModuleIndex++;
 		}
 
+		// Add the post items spacer.
+		FName ModuleSpacerKey = *FString::Printf(TEXT("Module%i"), ModuleIndex);
+		UNiagaraStackSpacer* ModuleSpacer = FindCurrentChildOfTypeByPredicate<UNiagaraStackSpacer>(CurrentChildren,
+			[=](UNiagaraStackSpacer* CurrentModuleSpacer) { return CurrentModuleSpacer->GetSpacerKey() == ModuleSpacerKey; });
+
+		if (ModuleSpacer == nullptr)
+		{
+			ModuleSpacer = NewObject<UNiagaraStackSpacer>(this);
+			ModuleSpacer->Initialize(CreateDefaultChildRequiredData(), ModuleSpacerKey, 1.4f);
+		}
+
+		NewChildren.Add(ModuleSpacer);
+		StackSpacerToModuleItemMap.Add(FObjectKey(ModuleSpacer), nullptr);
+	}
+	Super::RefreshChildrenInternal(CurrentChildren, NewChildren, NewIssues);
+	RefreshIssues(NewIssues);
+}
+
+void UNiagaraStackScriptItemGroup::RefreshIssues(TArray<FStackIssue>& NewIssues)
+{
+	TSharedPtr<FNiagaraScriptViewModel> ScriptViewModelPinned = ScriptViewModel.Pin();
+	checkf(ScriptViewModelPinned.IsValid(), TEXT("Can not refresh issues when the script view model has been deleted."));
+	UNiagaraGraph* Graph = ScriptViewModelPinned->GetGraphViewModel()->GetGraph();
+	FText ErrorMessage;
+	
+	if (FNiagaraStackGraphUtilities::ValidateGraphForOutput(*Graph, ScriptUsage, ScriptUsageId, ErrorMessage) == false)
+	{
+		UE_LOG(LogNiagaraEditor, Error, TEXT("Failed to Create Stack.  Message: %s"), *ErrorMessage.ToString());
+		UNiagaraStackEntry::FStackIssue Error;
+		Error.LongDescription = LOCTEXT("InvalidErrorText", "The data used to generate the stack has been corrupted and can not be used.\nUsing the fix option will reset this part of the stack to its default empty state.");
+		Error.ShortDescription = LOCTEXT("InvalidErrorSummaryText", "The stack data is invalid");
+		Error.UniqueIdentifier = FName(*FString::Printf(TEXT("StackDataInvalid-%s"), *GetStackEditorDataKey()));
+		UNiagaraStackEntry::FStackIssueFix Fix;
+		Fix.Description = LOCTEXT("FixStackGraph", "Fix invalid stack graph");
+		Fix.FixDelegate.BindLambda([=]()
+		{
+			FScopedTransaction ScopedTransaction(Fix.Description);
+			FNiagaraStackGraphUtilities::ResetGraphForOutput(*Graph, ScriptUsage, ScriptUsageId);
+			FNiagaraStackGraphUtilities::RelayoutGraph(*Graph);
+		});
+		Error.Fixes.Add(Fix);
+		NewIssues.Add(Error); 
+	}
+	else
+	{
 		bool bForcedError = false;
 		if (ScriptUsage == ENiagaraScriptUsage::SystemUpdateScript)
 		{
@@ -107,121 +360,217 @@ void UNiagaraStackScriptItemGroup::RefreshChildrenInternal(const TArray<UNiagara
 			ModuleScriptAsset.ObjectPath = SystemUpdateScriptRef.GetAssetPathName();
 
 			TArray<UNiagaraNodeFunctionCall*> FoundCalls;
+			UNiagaraNodeOutput* MatchingOutputNode = Graph->FindOutputNode(ScriptUsage, ScriptUsageId);
 			if (!FNiagaraStackGraphUtilities::FindScriptModulesInStack(ModuleScriptAsset, *MatchingOutputNode, FoundCalls))
 			{
 				bForcedError = true;
-				Error = MakeShared<FError>();
-				Error->ErrorText = LOCTEXT("SystemLifeCycleWarning", "The stack needs a SystemLifeCycle module.");;
-				Error->ErrorSummaryText = LOCTEXT("MissingRequiredMode", "Missing required module.");
-				Error->Fix.BindLambda([=]()
+				UNiagaraStackEntry::FStackIssue Error;
+				Error.LongDescription = LOCTEXT("SystemLifeCycleWarning", "The stack needs a SystemLifeCycle module.");;
+				Error.ShortDescription = LOCTEXT("MissingRequiredMode", "Missing required module.");
+				Error.UniqueIdentifier = FName(*FString::Printf(TEXT("MissingLifecycleModule-%s"), *GetStackEditorDataKey()));
+				UNiagaraStackEntry::FStackIssueFix Fix;
+				Fix.Description = LOCTEXT("AddingSystemLifecycleModule", "Adding System Lifecycle Module.");
+				Fix.FixDelegate.BindLambda([=]()
 				{
-					FScopedTransaction ScopedTransaction(LOCTEXT("AddingSystemLifecycleModule", "Adding System Lifecycle Module."));
-					FNiagaraStackGraphUtilities::AddScriptModuleToStack(ModuleScriptAsset, *MatchingOutputNode);
+					FScopedTransaction ScopedTransaction(Fix.Description);
+					UNiagaraNodeFunctionCall* AddedModuleNode = FNiagaraStackGraphUtilities::AddScriptModuleToStack(ModuleScriptAsset, *MatchingOutputNode);
+					if (AddedModuleNode == nullptr)
+					{
+						FNotificationInfo Info(LOCTEXT("FailedToAddSystemLifecycle", "Failed to add system life cycle module.\nCheck the log for errors."));
+						Info.ExpireDuration = 5.0f;
+						Info.bFireAndForget = true;
+						Info.Image = FCoreStyle::Get().GetBrush(TEXT("MessageLog.Error"));
+						FSlateNotificationManager::Get().AddNotification(Info);
+					}
 				});
+				Error.Fixes.Add(Fix);
+				NewIssues.Add(Error);
 			}
 		}
 
-		UNiagaraStackSpacer* AddButtonSpacer = FindCurrentChildOfTypeByPredicate<UNiagaraStackSpacer>(CurrentChildren,
-			[=](UNiagaraStackSpacer* CurrentSpacer) { return CurrentSpacer->GetSpacerKey() == "AddButton"; });
-
-		if (AddButtonSpacer == nullptr)
+		ENiagaraScriptCompileStatus Status = ScriptViewModelPinned->GetScriptCompileStatus(GetScriptUsage(), GetScriptUsageId());
+		if (!bForcedError)
 		{
-			AddButtonSpacer = NewObject<UNiagaraStackSpacer>(this);
-			AddButtonSpacer->Initialize(GetSystemViewModel(), GetEmitterViewModel(), "AddButton");
+			if (Status == ENiagaraScriptCompileStatus::NCS_Error)
+			{
+				UNiagaraStackEntry::FStackIssue Error;
+				Error.LongDescription = ScriptViewModelPinned->GetScriptErrors(GetScriptUsage(), GetScriptUsageId());
+				Error.ShortDescription = LOCTEXT("ConpileErrorSummary", "The stack has compile errors.");
+				Error.UniqueIdentifier = FName(*FString::Printf(TEXT("CompileErrors-%s"), *GetStackEditorDataKey()));
+				NewIssues.Add(Error);
+			}
 		}
+	}
+}
+
+void GenerateDragDropData(
+	UNiagaraNodeFunctionCall* SourceModule,
+	UNiagaraNodeFunctionCall* TargetModule, const UNiagaraGraph* TargetGraph, ENiagaraScriptUsage TargetScriptUsage, FGuid TargetScriptUsageId,
+	TArray<FNiagaraStackGraphUtilities::FStackNodeGroup>& OutSourceStackGroups, int32& OutSourceGroupIndex,
+	TArray<FNiagaraStackGraphUtilities::FStackNodeGroup>& OutTargetStackGroups, int32& OutTargetGroupIndex)
+{
+	// Find the output nodes for the source and target
+	UNiagaraNodeOutput* SourceOutputNode = FNiagaraStackGraphUtilities::GetEmitterOutputNodeForStackNode(*SourceModule);
+	UNiagaraNodeOutput* TargetOutputNode = TargetModule != nullptr
+		? FNiagaraStackGraphUtilities::GetEmitterOutputNodeForStackNode(*TargetModule)
+		: TargetGraph->FindOutputNode(TargetScriptUsage, TargetScriptUsageId);
 
-		UNiagaraStackAddScriptModuleItem* AddModuleItem = FindCurrentChildOfTypeByPredicate<UNiagaraStackAddScriptModuleItem>(CurrentChildren,
-			[=](UNiagaraStackAddScriptModuleItem* CurrentAddModuleItem) { return CurrentAddModuleItem->GetTargetIndex() == INDEX_NONE; });
+	// Collect the stack node groups for the source and target.
+	FNiagaraStackGraphUtilities::GetStackNodeGroups(*SourceOutputNode, OutSourceStackGroups);
+	if (SourceOutputNode == TargetOutputNode)
+	{
+		OutTargetStackGroups.Append(OutSourceStackGroups);
+	}
+	else
+	{
+		FNiagaraStackGraphUtilities::GetStackNodeGroups(*TargetOutputNode, OutTargetStackGroups);
+	}
 
-		if (AddModuleItem == nullptr)
+	// Calculate the source and target groups indexes for the drag/drop
+	OutSourceGroupIndex = INDEX_NONE;
+	for (int32 GroupIndex = 0; GroupIndex < OutSourceStackGroups.Num(); GroupIndex++)
+	{
+		if (OutSourceStackGroups[GroupIndex].EndNode == SourceModule)
 		{
-			AddModuleItem = NewObject<UNiagaraStackAddScriptModuleItem>(this);
-			AddModuleItem->Initialize(GetSystemViewModel(), GetEmitterViewModel(), GetStackEditorData(), *MatchingOutputNode, INDEX_NONE);
-			AddModuleItem->SetOnItemAdded(UNiagaraStackAddModuleItem::FOnItemAdded::CreateUObject(this, &UNiagaraStackScriptItemGroup::ItemAdded));
+			OutSourceGroupIndex = GroupIndex;
+			break;
 		}
+	}
 
-		UNiagaraStackSpacer* BottomSpacer = FindCurrentChildOfTypeByPredicate<UNiagaraStackSpacer>(CurrentChildren,
-			[=](UNiagaraStackSpacer* CurrentSpacer) { return CurrentSpacer->GetSpacerKey() == "ScriptStackBottom"; });
-
-		if (BottomSpacer == nullptr)
+	OutTargetGroupIndex = INDEX_NONE;
+	if (TargetModule == SourceModule)
+	{
+		OutTargetGroupIndex = OutSourceGroupIndex;
+	}
+	else if (TargetModule != nullptr)
+	{
+		for (int32 GroupIndex = 0; GroupIndex < OutTargetStackGroups.Num(); GroupIndex++)
 		{
-			BottomSpacer = NewObject<UNiagaraStackSpacer>(this);
-			BottomSpacer->Initialize(GetSystemViewModel(), GetEmitterViewModel(), "ScriptStackBottom");
+			if (OutTargetStackGroups[GroupIndex].EndNode == TargetModule)
+			{
+				OutTargetGroupIndex = GroupIndex;
+				break;
+			}
 		}
+	}
+	else
+	{
+		// If there is no target module then we need to insert at the end.  The last group is the output node and we want to
+		// insert before that.
+		OutTargetGroupIndex = OutTargetStackGroups.Num() - 1;
+	}
+}
 
-		NewChildren.Add(AddButtonSpacer);
-		NewChildren.Add(AddModuleItem);
-		NewChildren.Add(BottomSpacer);
-
-		ENiagaraScriptCompileStatus Status = ScriptViewModelPinned->GetScriptCompileStatus(GetScriptUsage(), GetScriptUsageId());
-		if (!bForcedError)
+TOptional<UNiagaraStackEntry::FDropResult> UNiagaraStackScriptItemGroup::ChildRequestCanDropInternal(const UNiagaraStackEntry& TargetChild, const TArray<UNiagaraStackEntry*>& DraggedEntries)
+{
+	if (bIsValidForOutput && DraggedEntries.Num() == 1)
+	{
+		UNiagaraStackModuleItem* SourceModuleItem = Cast<UNiagaraStackModuleItem>(DraggedEntries[0]);
+		if (SourceModuleItem != nullptr)
 		{
-			if (Status == ENiagaraScriptCompileStatus::NCS_Error)
+			if (SourceModuleItem->CanMoveAndDelete() == false)
+			{
+				return FDropResult(false, LOCTEXT("CantMoveModuleError", "This inherited module can't be moved."));
+			}
+
+			if (SourceModuleItem->GetModuleNode().GetGraph() != ScriptViewModel.Pin()->GetGraphViewModel()->GetGraph())
 			{
-				Error = MakeShared<FError>();
-				Error->ErrorText = ScriptViewModelPinned->GetScriptErrors(GetScriptUsage(), GetScriptUsageId());
-				Error->ErrorSummaryText = LOCTEXT("ConpileErrorSummary", "The stack has compile errors.");
+				return FDropResult(false, LOCTEXT("CantMoveModuleBetweenGraphsError", "This module can not be moved to this section of the stack"));
 			}
-			else
+
+			const UNiagaraStackSpacer* TargetSpacer = Cast<UNiagaraStackSpacer>(&TargetChild);
+			if (TargetSpacer != nullptr)
 			{
-				Error.Reset();
+				UNiagaraStackModuleItem** TargetModuleItemPtr = StackSpacerToModuleItemMap.Find(FObjectKey(TargetSpacer));
+				if (TargetModuleItemPtr != nullptr)
+				{
+					UNiagaraStackModuleItem* TargetModuleItem = *TargetModuleItemPtr;
+					TArray<FNiagaraStackGraphUtilities::FStackNodeGroup> SourceStackGroups;
+					TArray<FNiagaraStackGraphUtilities::FStackNodeGroup> TargetStackGroups;
+					int32 SourceGroupIndex;
+					int32 TargetGroupIndex;
+					GenerateDragDropData(
+						&SourceModuleItem->GetModuleNode(), TargetModuleItem != nullptr ? &TargetModuleItem->GetModuleNode() : nullptr,
+						ScriptViewModel.Pin()->GetGraphViewModel()->GetGraph(), ScriptUsage, ScriptUsageId,
+						SourceStackGroups, SourceGroupIndex,
+						TargetStackGroups, TargetGroupIndex);
+
+					// Make sure the source and target indices are within safe ranges, and make sure that the insert target isn't the source target or the spot directly
+					// after the source target since that won't actually move the module.
+					if (SourceGroupIndex > 0 && SourceGroupIndex < SourceStackGroups.Num() - 1 && TargetGroupIndex > 0 && TargetGroupIndex < TargetStackGroups.Num() &&
+						SourceStackGroups[SourceGroupIndex].EndNode != TargetStackGroups[TargetGroupIndex].EndNode &&
+						SourceStackGroups[SourceGroupIndex].EndNode != TargetStackGroups[TargetGroupIndex - 1].EndNode)
+					{
+						return FDropResult(true, LOCTEXT("MoveModuleResult", "Move this module here."));
+					}
+				}
 			}
 		}
 	}
+	return TOptional<FDropResult>();
 }
 
-void UNiagaraStackScriptItemGroup::ItemAdded()
+TOptional<UNiagaraStackEntry::FDropResult> UNiagaraStackScriptItemGroup::ChildRequestDropInternal(const UNiagaraStackEntry& TargetChild, const TArray<UNiagaraStackEntry*>& DraggedEntries)
 {
-	RefreshChildren();
-}
+	if (bIsValidForOutput && DraggedEntries.Num() == 1)
+	{
+		UNiagaraStackModuleItem* SourceModuleItem = Cast<UNiagaraStackModuleItem>(DraggedEntries[0]);
+		UNiagaraGraph* TargetGraph = ScriptViewModel.Pin()->GetGraphViewModel()->GetGraph();
+		if (SourceModuleItem != nullptr &&
+			SourceModuleItem->CanMoveAndDelete() &&
+			SourceModuleItem->GetModuleNode().GetGraph() == TargetGraph)
+		{
+			const UNiagaraStackSpacer* TargetSpacer = Cast<UNiagaraStackSpacer>(&TargetChild);
+			if (TargetSpacer != nullptr)
+			{
+				UNiagaraStackModuleItem** TargetModuleItemPtr = StackSpacerToModuleItemMap.Find(FObjectKey(TargetSpacer));
+				if (TargetModuleItemPtr != nullptr)
+				{
+					UNiagaraStackModuleItem* TargetModuleItem = *TargetModuleItemPtr;
+					TArray<FNiagaraStackGraphUtilities::FStackNodeGroup> SourceStackGroups;
+					TArray<FNiagaraStackGraphUtilities::FStackNodeGroup> TargetStackGroups;
+					int32 SourceGroupIndex;
+					int32 TargetGroupIndex;
+					GenerateDragDropData(
+						&SourceModuleItem->GetModuleNode(), TargetModuleItem != nullptr ? &TargetModuleItem->GetModuleNode() : nullptr,
+						TargetGraph, ScriptUsage, ScriptUsageId,
+						SourceStackGroups, SourceGroupIndex,
+						TargetStackGroups, TargetGroupIndex);
 
-void UNiagaraStackScriptItemGroup::ChildModifiedGroupItems()
-{
-	RefreshChildren();
-}
+					// Make sure the source and target indices are within safe ranges, and make sure that the insert target isn't the source target or the spot directly
+					// after the source target since that won't actually move the module.
+					if (SourceGroupIndex > 0 && SourceGroupIndex < SourceStackGroups.Num() - 1 && TargetGroupIndex > 0 && TargetGroupIndex < TargetStackGroups.Num() &&
+						SourceStackGroups[SourceGroupIndex].EndNode != TargetStackGroups[TargetGroupIndex].EndNode &&
+						SourceStackGroups[SourceGroupIndex].EndNode != TargetStackGroups[TargetGroupIndex - 1].EndNode)
+					{
+						FScopedTransaction ScopedTransaction(LOCTEXT("DragAndDropModule", "Drag and drop module"));
+						FNiagaraStackGraphUtilities::DisconnectStackNodeGroup(SourceStackGroups[SourceGroupIndex], SourceStackGroups[SourceGroupIndex - 1], SourceStackGroups[SourceGroupIndex + 1]);
+						FNiagaraStackGraphUtilities::ConnectStackNodeGroup(SourceStackGroups[SourceGroupIndex], TargetStackGroups[TargetGroupIndex - 1], TargetStackGroups[TargetGroupIndex]);
 
-int32 UNiagaraStackScriptItemGroup::GetErrorCount() const
-{
-	return Error.IsValid() ? 1 : 0;
-}
+						FNiagaraStackGraphUtilities::RelayoutGraph(*TargetGraph);
+						TargetGraph->NotifyGraphNeedsRecompile();
 
-bool UNiagaraStackScriptItemGroup::GetErrorFixable(int32 ErrorIdx) const
-{
-	return Error.IsValid() && Error->Fix.IsBound();
-}
+						SourceModuleItem->NotifyModuleMoved();
+						RefreshChildren();
 
-bool UNiagaraStackScriptItemGroup::TryFixError(int32 ErrorIdx)
-{
-	if (Error.IsValid() && Error->Fix.IsBound())
-	{
-		Error->Fix.Execute();
-		return true;
+						return FDropResult(true);
+					}
+				}
+			}
+		}
 	}
-	return false;
+	return FDropResult(false);
 }
 
-FText UNiagaraStackScriptItemGroup::GetErrorText(int32 ErrorIdx) const
+void UNiagaraStackScriptItemGroup::ItemAdded()
 {
-	if (Error.IsValid())
-	{
-		return Error->ErrorText;
-	}
-	else
-	{
-		return FText::GetEmpty();
-	}
+	RefreshChildren();
 }
 
-FText UNiagaraStackScriptItemGroup::GetErrorSummaryText(int32 ErrorIdx) const
+void UNiagaraStackScriptItemGroup::ChildModifiedGroupItems()
 {
-	if (Error.IsValid())
-	{
-		return Error->ErrorSummaryText;
-	}
-	else
-	{
-		return FText::GetEmpty();
-	}
+	RefreshChildren();
 }
 
+
 #undef LOCTEXT_NAMESPACE
diff --git a/Engine/Plugins/FX/Niagara/Source/NiagaraEditor/Private/ViewModels/Stack/NiagaraStackSpacer.cpp b/Engine/Plugins/FX/Niagara/Source/NiagaraEditor/Private/ViewModels/Stack/NiagaraStackSpacer.cpp
--- a/Engine/Plugins/FX/Niagara/Source/NiagaraEditor/Private/ViewModels/Stack/NiagaraStackSpacer.cpp
+++ b/Engine/Plugins/FX/Niagara/Source/NiagaraEditor/Private/ViewModels/Stack/NiagaraStackSpacer.cpp
@@ -2,10 +2,12 @@
 
 #include "ViewModels/Stack/NiagaraStackSpacer.h"
 
-void UNiagaraStackSpacer::Initialize(TSharedRef<FNiagaraSystemViewModel> InSystemViewModel, TSharedRef<FNiagaraEmitterViewModel> InEmitterViewModel, FName InSpacerKey)
+void UNiagaraStackSpacer::Initialize(FRequiredEntryData InRequiredEntryData, FName InSpacerKey, float InSpacerScale, EStackRowStyle InRowStyle)
 {
-	Super::Initialize(InSystemViewModel, InEmitterViewModel);
+	Super::Initialize(InRequiredEntryData, FString());
 	SpacerKey = InSpacerKey;
+	SpacerScale = InSpacerScale;
+	RowStyle = InRowStyle;
 }
 
 FText UNiagaraStackSpacer::GetDisplayName() const
@@ -13,12 +15,22 @@ FText UNiagaraStackSpacer::GetDisplayName() const
 	return FText();
 }
 
-FName UNiagaraStackSpacer::GetItemBackgroundName() const
+bool UNiagaraStackSpacer::GetCanExpand() const
 {
-	return "NiagaraEditor.Stack.Group.BackgroundColor";
+	return false;
+}
+
+UNiagaraStackEntry::EStackRowStyle UNiagaraStackSpacer::GetStackRowStyle() const
+{
+	return RowStyle;
 }
 
 FName UNiagaraStackSpacer::GetSpacerKey() const
 {
 	return SpacerKey;
 }
+
+float UNiagaraStackSpacer::GetSpacerScale() const
+{
+	return SpacerScale;
+}
\ No newline at end of file
diff --git a/Engine/Plugins/FX/Niagara/Source/NiagaraEditor/Private/ViewModels/Stack/NiagaraStackViewModel.cpp b/Engine/Plugins/FX/Niagara/Source/NiagaraEditor/Private/ViewModels/Stack/NiagaraStackViewModel.cpp
--- a/Engine/Plugins/FX/Niagara/Source/NiagaraEditor/Private/ViewModels/Stack/NiagaraStackViewModel.cpp
+++ b/Engine/Plugins/FX/Niagara/Source/NiagaraEditor/Private/ViewModels/Stack/NiagaraStackViewModel.cpp
@@ -2,31 +2,47 @@
 
 #include "ViewModels/Stack/NiagaraStackViewModel.h"
 #include "ViewModels/Stack/NiagaraStackRoot.h"
-#include "NiagaraSystemViewModel.h"
-#include "NiagaraEmitterViewModel.h"
+#include "ViewModels/NiagaraSystemViewModel.h"
+#include "ViewModels/NiagaraEmitterViewModel.h"
+#include "ViewModels/NiagaraEmitterHandleViewModel.h"
 #include "NiagaraScriptGraphViewModel.h"
+#include "NiagaraScriptViewModel.h"
+#include "NiagaraSystemEditorData.h"
 #include "NiagaraEmitterEditorData.h"
+#include "NiagaraStackEditorData.h"
 #include "NiagaraEmitter.h"
+#include "ViewModels/Stack/NiagaraStackItemGroup.h"
+#include "ViewModels/Stack/NiagaraStackItem.h"
+#include "ScopedTransaction.h"
 
 #include "Editor.h"
 
 #define LOCTEXT_NAMESPACE "NiagaraStackViewModel"
+const double UNiagaraStackViewModel::MaxSearchTime = .02f; // search at 50 fps
+TSharedPtr<FNiagaraSystemViewModel> UNiagaraStackViewModel::GetSystemViewModel()
+{
+	return SystemViewModel;
+}
 
-TSharedPtr<FNiagaraEmitterViewModel> UNiagaraStackViewModel::GetEmitterViewModel()
+TSharedPtr<FNiagaraEmitterHandleViewModel> UNiagaraStackViewModel::GetEmitterHandleViewModel()
 {
-	return EmitterViewModel;
+	return EmitterHandleViewModel;
 }
 
-void UNiagaraStackViewModel::Initialize(TSharedPtr<FNiagaraSystemViewModel> InSystemViewModel, TSharedPtr<FNiagaraEmitterViewModel> InEmitterViewModel)
+void UNiagaraStackViewModel::Initialize(TSharedPtr<FNiagaraSystemViewModel> InSystemViewModel, TSharedPtr<FNiagaraEmitterHandleViewModel> InEmitterHandleViewModel)
 {
-	if (EmitterViewModel.IsValid())
+	if (RootEntry != nullptr)
 	{
-		GEditor->UnregisterForUndo(this);
-		EmitterViewModel->OnScriptCompiled().RemoveAll(this);
 		RootEntry->OnStructureChanged().RemoveAll(this);
 		RootEntry->OnDataObjectModified().RemoveAll(this);
 		RootEntries.Empty();
 		RootEntry = nullptr;
+		GEditor->UnregisterForUndo(this);
+	}
+
+	if (EmitterHandleViewModel.IsValid() && EmitterHandleViewModel->GetEmitterViewModel().IsValid())
+	{
+		EmitterHandleViewModel->GetEmitterViewModel()->OnScriptCompiled().RemoveAll(this);
 	}
 
 	if (SystemViewModel.IsValid())
@@ -35,7 +51,9 @@ void UNiagaraStackViewModel::Initialize(TSharedPtr<FNiagaraSystemViewModel> InSy
 	}
 
 	SystemViewModel = InSystemViewModel;
-	EmitterViewModel = InEmitterViewModel;
+	//UE_LOG(LogNiagaraEditor, Warning, TEXT("UNiagaraStackViewModel %p refs FNiagaraEmitterViewModel %p"), this, InEmitterViewModel.Get());
+	EmitterHandleViewModel = InEmitterHandleViewModel;
+	TSharedPtr<FNiagaraEmitterViewModel> EmitterViewModel = InEmitterHandleViewModel.IsValid() ? InEmitterHandleViewModel->GetEmitterViewModel() : TSharedPtr<FNiagaraEmitterViewModel>();
 
 	if (SystemViewModel.IsValid() && EmitterViewModel.IsValid() && EmitterViewModel->GetSharedScriptViewModel()->GetGraphViewModel()->GetGraph() != nullptr)
 	{
@@ -45,34 +63,242 @@ void UNiagaraStackViewModel::Initialize(TSharedPtr<FNiagaraSystemViewModel> InSy
 		SystemViewModel->OnSystemCompiled().AddUObject(this, &UNiagaraStackViewModel::OnSystemCompiled);
 		
 		RootEntry = NewObject<UNiagaraStackRoot>(this);
-		RootEntry->Initialize(SystemViewModel.ToSharedRef(), EmitterViewModel.ToSharedRef());
+		UNiagaraStackEntry::FRequiredEntryData RequiredEntryData(SystemViewModel.ToSharedRef(), EmitterViewModel.ToSharedRef(),
+			NAME_None, NAME_None,
+			SystemViewModel->GetOrCreateEditorData().GetStackEditorData());
+		RootEntry->Initialize(RequiredEntryData);
 		RootEntry->RefreshChildren();
 		RootEntry->OnStructureChanged().AddUObject(this, &UNiagaraStackViewModel::EntryStructureChanged);
 		RootEntry->OnDataObjectModified().AddUObject(this, &UNiagaraStackViewModel::EntryDataObjectModified);
 		RootEntries.Add(RootEntry);
 	}
-
+	CurrentFocusedSearchMatchIndex = -1;
 	StructureChangedDelegate.Broadcast();
+	bRestartSearch = false;
 }
 
 void UNiagaraStackViewModel::BeginDestroy()
 {
 	Super::BeginDestroy();
-	if (EmitterViewModel.IsValid())
+	if (EmitterHandleViewModel.IsValid() && EmitterHandleViewModel->GetEmitterViewModel().IsValid())
 	{
 		GEditor->UnregisterForUndo(this);
-		EmitterViewModel->OnScriptCompiled().RemoveAll(this);
+		EmitterHandleViewModel->GetEmitterViewModel()->OnScriptCompiled().RemoveAll(this);
+
+		//UE_LOG(LogNiagaraEditor, Log, TEXT("UNiagaraStackViewModel::BeginDestroy"));
+		EmitterHandleViewModel->GetEmitterViewModel().Reset();
+		EmitterHandleViewModel.Reset();
+		SystemViewModel.Reset();
+	} 
+	RootEntry = nullptr;
+}
+
+void UNiagaraStackViewModel::Tick()
+{
+	if (RootEntry)
+	{
+		SearchTick();
+	}
+}
+
+void UNiagaraStackViewModel::OnSearchTextChanged(const FText& SearchText)
+{
+	if (RootEntry)
+	{
+		CurrentSearchText = SearchText;
+		// postpone searching until next tick; protects against crashes from the GC
+		// also this can be triggered by multiple events, so better wait
+		bRestartSearch = true;
+	}
+}
+
+bool UNiagaraStackViewModel::IsSearching()
+{
+	return ItemsToSearch.Num() > 0;
+}
+
+const TArray<UNiagaraStackViewModel::FSearchResult>& UNiagaraStackViewModel::GetCurrentSearchResults()
+{
+	return CurrentSearchResults;
+}
+
+UNiagaraStackEntry* UNiagaraStackViewModel::GetCurrentFocusedEntry()
+{
+	if (CurrentFocusedSearchMatchIndex >= 0)
+	{
+		FSearchResult FocusedMatch = CurrentSearchResults[CurrentFocusedSearchMatchIndex];
+		return FocusedMatch.GetEntry();
 	}
+	return nullptr;
+}
+
+void UNiagaraStackViewModel::AddSearchScrollOffset(int NumberOfSteps)
+{
+	CurrentFocusedSearchMatchIndex += NumberOfSteps;
+	if (CurrentFocusedSearchMatchIndex >= CurrentSearchResults.Num())
+	{
+		CurrentFocusedSearchMatchIndex = 0;
+	}
+	if (CurrentFocusedSearchMatchIndex < 0)
+	{
+		CurrentFocusedSearchMatchIndex = CurrentSearchResults.Num() - 1;
+	}
+}
+
+void UNiagaraStackViewModel::CollapseToHeaders()
+{
+	CollapseToHeadersRecursive(GetRootEntries());
+	NotifyStructureChanged();
+}
+
+void UNiagaraStackViewModel::UndismissAllIssues()
+{
+	FScopedTransaction ScopedTransaction(LOCTEXT("UnDismissIssues", "Undismiss issues"));
+	UNiagaraStackEditorData& EmitterData = GetEmitterHandleViewModel()->GetEmitterViewModel()->GetEditorData().GetStackEditorData();
+	EmitterData.Modify();
+	EmitterData.UndismissAllIssues();
+
+	UNiagaraStackEditorData& SystemData = GetSystemViewModel()->GetEditorData().GetStackEditorData();
+	SystemData.Modify();
+	SystemData.UndismissAllIssues();
+	
+	RootEntry->RefreshChildren();
+}
+
+bool UNiagaraStackViewModel::HasDismissedStackIssues()
+{
+	return GetSystemViewModel()->GetEditorData().GetStackEditorData().GetDismissedStackIssueIds().Num() > 0
+		|| GetEmitterHandleViewModel()->GetEmitterViewModel()->GetEditorData().GetStackEditorData().GetDismissedStackIssueIds().Num() > 0;
+}
+
+void UNiagaraStackViewModel::CollapseToHeadersRecursive(TArray<UNiagaraStackEntry*> Entries)
+{
+	for (UNiagaraStackEntry* Entry : Entries)
+	{
+		if (Entry->GetCanExpand())
+		{
+			if (Entry->IsA<UNiagaraStackItemGroup>())
+			{
+				Entry->SetIsExpanded(true);
+			}
+			else if (Entry->IsA<UNiagaraStackItem>())
+			{
+				Entry->SetIsExpanded(false);
+			}
+		}
+
+		TArray<UNiagaraStackEntry*> Children;
+		Entry->GetUnfilteredChildren(Children);
+		CollapseToHeadersRecursive(Children);
+	}
+}
+
+void UNiagaraStackViewModel::GetPathForEntry(UNiagaraStackEntry* Entry, TArray<UNiagaraStackEntry*>& EntryPath)
+{
+	GeneratePathForEntry(RootEntry, Entry, TArray<UNiagaraStackEntry*>(), EntryPath);
 }
 
 void UNiagaraStackViewModel::OnSystemCompiled()
 {
 	RootEntry->RefreshChildren();
+	// when the entries are refreshed, make sure search results are still valid, the stack counts on them
+	OnSearchTextChanged(CurrentSearchText);
 }
 
 void UNiagaraStackViewModel::OnEmitterCompiled()
 {
 	RootEntry->RefreshChildren();
+	OnSearchTextChanged(CurrentSearchText);
+}
+
+void UNiagaraStackViewModel::SearchTick()
+{
+	// perform partial searches here, by processing a fixed number of entries (maybe more than one?)
+	if (bRestartSearch)
+	{
+		// clear the search results
+		CurrentSearchResults.Empty();
+		CurrentFocusedSearchMatchIndex = -1;
+		// generates ItemsToSearch, these will be processed on tick, in batches
+		if (CurrentSearchText.IsEmpty() == false)
+		{
+			GenerateTraversalEntries(RootEntry, TArray<UNiagaraStackEntry*>(), ItemsToSearch);
+		}
+		bRestartSearch = false;
+	}
+	else if (IsSearching())
+	{
+		double SearchStartTime= FPlatformTime::Seconds();
+		double CurrentSearchLoopTime = SearchStartTime;
+		// process at least one item, but don't go over MaxSearchTime for the rest
+		while (ItemsToSearch.Num() > 0 && CurrentSearchLoopTime - SearchStartTime < MaxSearchTime)
+		{
+			UNiagaraStackEntry* EntryToProcess = ItemsToSearch[0].GetEntry();
+			ensure(EntryToProcess != nullptr); // should never happen so something went wrong if this is hit
+			if (EntryToProcess != nullptr)
+			{
+				TArray<UNiagaraStackEntry::FStackSearchItem> SearchItems;
+				EntryToProcess->GetSearchItems(SearchItems);
+				TSet<FName> MatchedKeys;
+				for (UNiagaraStackEntry::FStackSearchItem SearchItem : SearchItems)
+				{
+					if (ItemMatchesSearchCriteria(SearchItem)) 
+					{
+						if (MatchedKeys.Contains(SearchItem.Key) == false)
+						{
+							CurrentSearchResults.Add({ ItemsToSearch[0].EntryPath, SearchItem });
+							MatchedKeys.Add(SearchItem.Key);
+						}
+					}
+				}
+			}
+			ItemsToSearch.RemoveAt(0); // can't use RemoveAtSwap because we need to preserve the order
+			CurrentSearchLoopTime = FPlatformTime::Seconds();
+		}
+		if (ItemsToSearch.Num() == 0)
+		{
+			SearchCompletedDelegate.Broadcast();
+		}
+	}
+}
+
+void UNiagaraStackViewModel::GenerateTraversalEntries(UNiagaraStackEntry* Root, TArray<UNiagaraStackEntry*> ParentChain,
+	TArray<FSearchWorkItem>& TraversedArray)
+{
+	TArray<UNiagaraStackEntry*> Children;
+	Root->GetFilteredChildren(Children);
+	ParentChain.Add(Root);
+	TraversedArray.Add(FSearchWorkItem{ParentChain});
+	for (auto Child : Children)
+	{
+		GenerateTraversalEntries(Child, ParentChain, TraversedArray);
+	}
+}
+
+bool UNiagaraStackViewModel::ItemMatchesSearchCriteria(UNiagaraStackEntry::FStackSearchItem SearchItem)
+{
+	// this is a simple text compare, we need to replace this with a complex search on future passes
+	return SearchItem.Value.ToString().Contains(CurrentSearchText.ToString());
+}
+
+void UNiagaraStackViewModel::GeneratePathForEntry(UNiagaraStackEntry* Root, UNiagaraStackEntry* Entry, TArray<UNiagaraStackEntry*> CurrentPath, TArray<UNiagaraStackEntry*>& EntryPath)
+{
+	if (EntryPath.Num() > 0)
+	{
+		return;
+	}
+	TArray<UNiagaraStackEntry*> Children;
+	Root->GetUnfilteredChildren(Children);
+	CurrentPath.Add(Root);
+	for (auto Child : Children)
+	{
+		if (Child == Entry)
+		{
+			EntryPath.Append(CurrentPath);
+			return;
+		}
+		GeneratePathForEntry(Child, Entry, CurrentPath, EntryPath);
+	}
 }
 
 TArray<UNiagaraStackEntry*>& UNiagaraStackViewModel::GetRootEntries()
@@ -85,19 +311,93 @@ UNiagaraStackViewModel::FOnStructureChanged& UNiagaraStackViewModel::OnStructure
 	return StructureChangedDelegate;
 }
 
+UNiagaraStackViewModel::FOnSearchCompleted& UNiagaraStackViewModel::OnSearchCompleted()
+{
+	return SearchCompletedDelegate;
+}
+
+bool UNiagaraStackViewModel::GetShowAllAdvanced() const
+{
+	if (SystemViewModel.IsValid() && EmitterHandleViewModel.IsValid())
+	{
+		return SystemViewModel->GetEditorData().GetStackEditorData().GetShowAllAdvanced() ||
+			EmitterHandleViewModel->GetEmitterViewModel()->GetEditorData().GetStackEditorData().GetShowAllAdvanced();
+	}
+	return false;
+}
+
+void UNiagaraStackViewModel::SetShowAllAdvanced(bool bInShowAllAdvanced)
+{
+	if (SystemViewModel.IsValid() && EmitterHandleViewModel.IsValid())
+	{
+		SystemViewModel->GetOrCreateEditorData().GetStackEditorData().SetShowAllAdvanced(bInShowAllAdvanced);
+		EmitterHandleViewModel->GetEmitterViewModel()->GetOrCreateEditorData().GetStackEditorData().SetShowAllAdvanced(bInShowAllAdvanced);
+		OnSearchTextChanged(CurrentSearchText);
+		StructureChangedDelegate.Broadcast();
+	}
+}
+
+bool UNiagaraStackViewModel::GetShowOutputs() const
+{
+	if (SystemViewModel.IsValid() && EmitterHandleViewModel.IsValid())
+	{
+		return SystemViewModel->GetEditorData().GetStackEditorData().GetShowOutputs() ||
+			EmitterHandleViewModel->GetEmitterViewModel()->GetEditorData().GetStackEditorData().GetShowOutputs();
+	}
+	return false;
+}
+
+void UNiagaraStackViewModel::SetShowOutputs(bool bInShowOutputs)
+{
+	if (SystemViewModel.IsValid() && EmitterHandleViewModel.IsValid())
+	{
+		SystemViewModel->GetOrCreateEditorData().GetStackEditorData().SetShowOutputs(bInShowOutputs);
+		EmitterHandleViewModel->GetEmitterViewModel()->GetOrCreateEditorData().GetStackEditorData().SetShowOutputs(bInShowOutputs);
+		OnSearchTextChanged(CurrentSearchText);
+
+		// Showing outputs changes indenting so a full refresh is needed.
+		RootEntry->RefreshChildren();
+	}
+}
+
+double UNiagaraStackViewModel::GetLastScrollPosition() const
+{
+	if (EmitterHandleViewModel.IsValid())
+	{
+		return EmitterHandleViewModel->GetEmitterViewModel()->GetEditorData().GetStackEditorData().GetLastScrollPosition();
+	}
+	return 0;
+}
+
+void UNiagaraStackViewModel::SetLastScrollPosition(double InLastScrollPosition)
+{
+	if (EmitterHandleViewModel.IsValid())
+	{
+		EmitterHandleViewModel->GetEmitterViewModel()->GetOrCreateEditorData().GetStackEditorData().SetLastScrollPosition(InLastScrollPosition);
+	}
+}
+
+void UNiagaraStackViewModel::NotifyStructureChanged()
+{
+	EntryStructureChanged();
+}
+
 void UNiagaraStackViewModel::PostUndo(bool bSuccess)
 {
 	RootEntry->RefreshChildren();
+	OnSearchTextChanged(CurrentSearchText);
 }
 
 void UNiagaraStackViewModel::EntryStructureChanged()
 {
 	StructureChangedDelegate.Broadcast();
+	OnSearchTextChanged(CurrentSearchText);
 }
 
 void UNiagaraStackViewModel::EntryDataObjectModified(UObject* ChangedObject)
 {
 	SystemViewModel->NotifyDataObjectChanged(ChangedObject);
+	OnSearchTextChanged(CurrentSearchText);
 }
 
 #undef LOCTEXT_NAMESPACE
diff --git a/Engine/Plugins/FX/Niagara/Source/NiagaraEditor/Private/Widgets/SNiagaraCurveEditor.cpp b/Engine/Plugins/FX/Niagara/Source/NiagaraEditor/Private/Widgets/SNiagaraCurveEditor.cpp
--- a/Engine/Plugins/FX/Niagara/Source/NiagaraEditor/Private/Widgets/SNiagaraCurveEditor.cpp
+++ b/Engine/Plugins/FX/Niagara/Source/NiagaraEditor/Private/Widgets/SNiagaraCurveEditor.cpp
@@ -1,7 +1,7 @@
 // Copyright 1998-2018 Epic Games, Inc. All Rights Reserved.
 
 #include "SNiagaraCurveEditor.h"
-#include "NiagaraSystemViewModel.h"
+#include "ViewModels/NiagaraSystemViewModel.h"
 #include "SCurveEditor.h"
 #include "Widgets/SOverlay.h"
 #include "Framework/MultiBox/MultiBoxBuilder.h"
@@ -20,7 +20,9 @@ void SNiagaraCurveEditor::Construct(const FArguments& InArgs, TSharedRef<FNiagar
 	OutputSnap = .1f;
 
 	SAssignNew(CurveEditor, SCurveEditor)
-		.ShowCurveSelector(true);
+		.ShowCurveSelector(true)
+		.InputSnap(this, &SNiagaraCurveEditor::GetInputSnap)
+		.OutputSnap(this, &SNiagaraCurveEditor::GetOutputSnap);
 	CurveEditor->SetCurveOwner(&SystemViewModel->GetCurveOwner());
 
 	TSharedPtr<SOverlay> OverlayWidget;
diff --git a/Engine/Plugins/FX/Niagara/Source/NiagaraEditor/Private/Widgets/SNiagaraGeneratedCodeView.cpp b/Engine/Plugins/FX/Niagara/Source/NiagaraEditor/Private/Widgets/SNiagaraGeneratedCodeView.cpp
--- a/Engine/Plugins/FX/Niagara/Source/NiagaraEditor/Private/Widgets/SNiagaraGeneratedCodeView.cpp
+++ b/Engine/Plugins/FX/Niagara/Source/NiagaraEditor/Private/Widgets/SNiagaraGeneratedCodeView.cpp
@@ -7,8 +7,8 @@
 #include "Widgets/Input/SSearchBox.h"
 #include "Widgets/Input/SButton.h"
 #include "ISequencer.h"
-#include "NiagaraSystemViewModel.h"
-#include "NiagaraEmitterHandleViewModel.h"
+#include "ViewModels/NiagaraSystemViewModel.h"
+#include "ViewModels/NiagaraEmitterHandleViewModel.h"
 #include "NiagaraEmitterHandle.h"
 #include "NiagaraEmitter.h"
 #include "NiagaraScript.h"
@@ -323,13 +323,9 @@ void SNiagaraGeneratedCodeView::UpdateUI()
 	TArray<UNiagaraScript*> Scripts;
 	TArray<uint32> ScriptDisplayTypes;
 	UNiagaraSystem& System = SystemViewModel->GetSystem();
-	Scripts.Add(System.GetSystemSpawnScript(false));
+	Scripts.Add(System.GetSystemSpawnScript());
 	ScriptDisplayTypes.Add(0);
-	Scripts.Add(System.GetSystemUpdateScript(false));
-	ScriptDisplayTypes.Add(0);
-	Scripts.Add(System.GetSystemSpawnScript(true));
-	ScriptDisplayTypes.Add(0);
-	Scripts.Add(System.GetSystemUpdateScript(true));
+	Scripts.Add(System.GetSystemUpdateScript());
 	ScriptDisplayTypes.Add(0);
 
 	TArray<TSharedRef<FNiagaraEmitterHandleViewModel>> SelectedEmitterHandles;
@@ -348,18 +344,18 @@ void SNiagaraGeneratedCodeView::UpdateUI()
 
 	// find the particle spawn script
 	TArray<UNiagaraScript*> DupeScriptsForAssembly;
-	UNiagaraScript *ParticleSpawnScript = nullptr;
+	UNiagaraScript* ParticleGPUScript = nullptr;
 	for (UNiagaraScript *Script : Scripts)
 	{
 		DupeScriptsForAssembly.Add(Script);
 		ScriptDisplayTypes.Add(2);
-		if (Script->Usage == ENiagaraScriptUsage::ParticleSpawnScript || Script->Usage == ENiagaraScriptUsage::ParticleSpawnScriptInterpolated)
+		if (Script->Usage == ENiagaraScriptUsage::ParticleGPUComputeScript)
 		{
-			ParticleSpawnScript = Script;
+			ParticleGPUScript = Script;
 		}
 	}
 	Scripts.Append(DupeScriptsForAssembly);
-	Scripts.Add(ParticleSpawnScript); // add for the GPU update/spawn script
+	Scripts.Add(ParticleGPUScript); // add for the GPU update/spawn script
 	ScriptDisplayTypes.Add(1);
 		
 	GeneratedCode.SetNum(Scripts.Num());
@@ -384,20 +380,29 @@ void SNiagaraGeneratedCodeView::UpdateUI()
 			// GPU combined spawn / update script
 			if (bIsGPU)
 			{
-				GeneratedCode[i].Usage = ENiagaraScriptUsage::ParticleSpawnScript;
-				ParticleSpawnScript->LastHlslTranslationGPU.ParseIntoArrayLines(OutputByLines, false);
+				GeneratedCode[i].Usage = Scripts[i]->Usage;
+				if (ParticleGPUScript->GetVMExecutableData().IsValid())
+				{
+					ParticleGPUScript->GetVMExecutableData().LastHlslTranslationGPU.ParseIntoArrayLines(OutputByLines, false);
+				}
 			}
 			else if (bIsAssembly)
 			{
 				GeneratedCode[i].Usage = Scripts[i]->Usage;
 				GeneratedCode[i].UsageId = Scripts[i]->GetUsageId();
-				Scripts[i]->LastAssemblyTranslation.ParseIntoArrayLines(OutputByLines, false);
+				if (Scripts[i]->GetVMExecutableData().IsValid())
+				{
+					Scripts[i]->GetVMExecutableData().LastAssemblyTranslation.ParseIntoArrayLines(OutputByLines, false);
+				}
 			}
 			else
 			{
 				GeneratedCode[i].Usage = Scripts[i]->Usage;
 				GeneratedCode[i].UsageId = Scripts[i]->GetUsageId();
-				Scripts[i]->LastHlslTranslation.ParseIntoArrayLines(OutputByLines, false);
+				if (Scripts[i]->GetVMExecutableData().IsValid())
+				{
+					Scripts[i]->GetVMExecutableData().LastHlslTranslation.ParseIntoArrayLines(OutputByLines, false);
+				}
 			}
 		}
 		else
@@ -408,9 +413,9 @@ void SNiagaraGeneratedCodeView::UpdateUI()
 		GeneratedCode[i].HlslByLines.SetNum(OutputByLines.Num());
 		if (bIsAssembly)
 		{
-			if (Scripts[i] != nullptr)
+			if (Scripts[i] != nullptr && Scripts[i]->GetVMExecutableData().IsValid())
 			{
-				GeneratedCode[i].Hlsl = FText::FromString(Scripts[i]->LastAssemblyTranslation);
+				GeneratedCode[i].Hlsl = FText::FromString(Scripts[i]->GetVMExecutableData().LastAssemblyTranslation);
 			}
 			GeneratedCode[i].HlslByLines = OutputByLines;
 		}
@@ -439,7 +444,7 @@ void SNiagaraGeneratedCodeView::UpdateUI()
 			GeneratedCode[i].UsageName = FText::Format(LOCTEXT("UsageNameEvent", "{0}-{1}{2}"), ScriptEnum->GetDisplayNameTextByValue((int64)Scripts[i]->Usage), EventName, bIsAssembly ? AssemblyIdText : FText::GetEmpty());
 		}
 		// GPU combined spawn / update script
-		else if (i == GeneratedCode.Num() - 1 && Scripts[i]->Usage == ENiagaraScriptUsage::ParticleSpawnScript)
+		else if (bIsGPU && i == GeneratedCode.Num() - 1 && Scripts[i]->IsParticleSpawnScript())
 		{
 			GeneratedCode[i].UsageName = LOCTEXT("UsageNameGPU", "GPU Spawn/Update");
 		}
@@ -518,7 +523,10 @@ SNiagaraGeneratedCodeView::~SNiagaraGeneratedCodeView()
 	if (SystemViewModel.IsValid())
 	{
 		SystemViewModel->OnSelectedEmitterHandlesChanged().RemoveAll(this);
-		SystemViewModel->GetSystemScriptViewModel()->OnSystemCompiled().RemoveAll(this);
+		if (SystemViewModel->GetSystemScriptViewModel().IsValid())
+		{
+			SystemViewModel->GetSystemScriptViewModel()->OnSystemCompiled().RemoveAll(this);
+		}
 	}
 	
 }
diff --git a/Engine/Plugins/FX/Niagara/Source/NiagaraEditor/Private/Widgets/SNiagaraGeneratedCodeView.h b/Engine/Plugins/FX/Niagara/Source/NiagaraEditor/Private/Widgets/SNiagaraGeneratedCodeView.h
--- a/Engine/Plugins/FX/Niagara/Source/NiagaraEditor/Private/Widgets/SNiagaraGeneratedCodeView.h
+++ b/Engine/Plugins/FX/Niagara/Source/NiagaraEditor/Private/Widgets/SNiagaraGeneratedCodeView.h
@@ -8,7 +8,7 @@
 #include "Widgets/Views/STableRow.h"
 #include "Widgets/Views/STableViewBase.h"
 #include "Widgets/Views/STreeView.h"
-#include "NiagaraSystemViewModel.h"
+#include "ViewModels/NiagaraSystemViewModel.h"
 #include "TickableEditorObject.h"
 #include "UObject/WeakObjectPtrTemplates.h"
 #include "NiagaraDataSet.h"
diff --git a/Engine/Plugins/FX/Niagara/Source/NiagaraEditor/Private/Widgets/SNiagaraGraphNodeConvert.cpp b/Engine/Plugins/FX/Niagara/Source/NiagaraEditor/Private/Widgets/SNiagaraGraphNodeConvert.cpp
--- a/Engine/Plugins/FX/Niagara/Source/NiagaraEditor/Private/Widgets/SNiagaraGraphNodeConvert.cpp
+++ b/Engine/Plugins/FX/Niagara/Source/NiagaraEditor/Private/Widgets/SNiagaraGraphNodeConvert.cpp
@@ -16,7 +16,9 @@
 
 void SNiagaraGraphNodeConvert::Construct(const FArguments& InArgs, UEdGraphNode* InGraphNode)
 {
-	GraphNode = InGraphNode;
+	GraphNode = InGraphNode; 
+	RegisterNiagaraGraphNode(InGraphNode);
+
 	UpdateGraphNode();
 }
 
@@ -98,7 +100,7 @@ void SNiagaraGraphNodeConvert::UpdateGraphNode()
 	{
 		ConvertNodeViewModel = MakeShareable(new FNiagaraConvertNodeViewModel(*ConvertNode));
 	}
-	SGraphNode::UpdateGraphNode();
+	SNiagaraGraphNode::UpdateGraphNode();
 }
 
 void SNiagaraGraphNodeConvert::AddPin(const TSharedRef<SGraphPin>& PinToAdd)
diff --git a/Engine/Plugins/FX/Niagara/Source/NiagaraEditor/Private/Widgets/SNiagaraGraphNodeConvert.h b/Engine/Plugins/FX/Niagara/Source/NiagaraEditor/Private/Widgets/SNiagaraGraphNodeConvert.h
--- a/Engine/Plugins/FX/Niagara/Source/NiagaraEditor/Private/Widgets/SNiagaraGraphNodeConvert.h
+++ b/Engine/Plugins/FX/Niagara/Source/NiagaraEditor/Private/Widgets/SNiagaraGraphNodeConvert.h
@@ -2,13 +2,13 @@
 
 #pragma once
 
-#include "SGraphNode.h"
+#include "SNiagaraGraphNode.h"
 
 class FNiagaraConvertNodeViewModel;
 class FNiagaraConvertPinViewModel;
 
 /** A graph node widget representing a niagara convert node. */
-class SNiagaraGraphNodeConvert : public SGraphNode
+class SNiagaraGraphNodeConvert : public SNiagaraGraphNode
 {
 public:
 	SLATE_BEGIN_ARGS(SNiagaraGraphNodeConvert) {}
diff --git a/Engine/Plugins/FX/Niagara/Source/NiagaraEditor/Private/Widgets/SNiagaraGraphNodeCustomHlsl.cpp b/Engine/Plugins/FX/Niagara/Source/NiagaraEditor/Private/Widgets/SNiagaraGraphNodeCustomHlsl.cpp
--- a/Engine/Plugins/FX/Niagara/Source/NiagaraEditor/Private/Widgets/SNiagaraGraphNodeCustomHlsl.cpp
+++ b/Engine/Plugins/FX/Niagara/Source/NiagaraEditor/Private/Widgets/SNiagaraGraphNodeCustomHlsl.cpp
@@ -17,7 +17,7 @@
 
 #define LOCTEXT_NAMESPACE "SNiagaraGraphNodeCustomHlsl"
 
-SNiagaraGraphNodeCustomHlsl::SNiagaraGraphNodeCustomHlsl() : SGraphNode()
+SNiagaraGraphNodeCustomHlsl::SNiagaraGraphNodeCustomHlsl() : SNiagaraGraphNode()
 {
 
 }
@@ -25,6 +25,7 @@ SNiagaraGraphNodeCustomHlsl::SNiagaraGraphNodeCustomHlsl() : SGraphNode()
 void SNiagaraGraphNodeCustomHlsl::Construct(const FArguments& InArgs, UEdGraphNode* InGraphNode)
 {
 	GraphNode = InGraphNode;
+	RegisterNiagaraGraphNode(InGraphNode);
 	UpdateGraphNode();
 }
 
diff --git a/Engine/Plugins/FX/Niagara/Source/NiagaraEditor/Private/Widgets/SNiagaraGraphNodeCustomHlsl.h b/Engine/Plugins/FX/Niagara/Source/NiagaraEditor/Private/Widgets/SNiagaraGraphNodeCustomHlsl.h
--- a/Engine/Plugins/FX/Niagara/Source/NiagaraEditor/Private/Widgets/SNiagaraGraphNodeCustomHlsl.h
+++ b/Engine/Plugins/FX/Niagara/Source/NiagaraEditor/Private/Widgets/SNiagaraGraphNodeCustomHlsl.h
@@ -2,12 +2,12 @@
 
 #pragma once
 
-#include "SGraphNode.h"
+#include "SNiagaraGraphNode.h"
 #include "Widgets/Input/SMenuAnchor.h"
 #include "Framework/Commands/UICommandList.h"
 
 /** A graph node widget representing a niagara input node. */
-class SNiagaraGraphNodeCustomHlsl : public SGraphNode
+class SNiagaraGraphNodeCustomHlsl : public SNiagaraGraphNode
 {
 public:
 	SLATE_BEGIN_ARGS(SNiagaraGraphNodeCustomHlsl) {}
diff --git a/Engine/Plugins/FX/Niagara/Source/NiagaraEditor/Private/Widgets/SNiagaraGraphNodeInput.cpp b/Engine/Plugins/FX/Niagara/Source/NiagaraEditor/Private/Widgets/SNiagaraGraphNodeInput.cpp
--- a/Engine/Plugins/FX/Niagara/Source/NiagaraEditor/Private/Widgets/SNiagaraGraphNodeInput.cpp
+++ b/Engine/Plugins/FX/Niagara/Source/NiagaraEditor/Private/Widgets/SNiagaraGraphNodeInput.cpp
@@ -44,7 +44,7 @@ public:
 };
 
 
-SNiagaraGraphNodeInput::SNiagaraGraphNodeInput() : SGraphNode(),
+SNiagaraGraphNodeInput::SNiagaraGraphNodeInput() : SNiagaraGraphNode(),
 	ToolkitCommands(new FUICommandList()),
 	bRequestedSyncExposureOptions(false)
 {
@@ -53,7 +53,9 @@ SNiagaraGraphNodeInput::SNiagaraGraphNodeInput() : SGraphNode(),
 
 void SNiagaraGraphNodeInput::Construct(const FArguments& InArgs, UEdGraphNode* InGraphNode)
 {
-	GraphNode = InGraphNode;
+	GraphNode = InGraphNode; 
+	RegisterNiagaraGraphNode(InGraphNode);
+
 	FNiagaraGraphNodeInputCommands::Register();
 	BindCommands();
 	UpdateGraphNode();
diff --git a/Engine/Plugins/FX/Niagara/Source/NiagaraEditor/Private/Widgets/SNiagaraGraphNodeInput.h b/Engine/Plugins/FX/Niagara/Source/NiagaraEditor/Private/Widgets/SNiagaraGraphNodeInput.h
--- a/Engine/Plugins/FX/Niagara/Source/NiagaraEditor/Private/Widgets/SNiagaraGraphNodeInput.h
+++ b/Engine/Plugins/FX/Niagara/Source/NiagaraEditor/Private/Widgets/SNiagaraGraphNodeInput.h
@@ -2,12 +2,12 @@
 
 #pragma once
 
-#include "SGraphNode.h"
+#include "SNiagaraGraphNode.h"
 #include "Widgets/Input/SMenuAnchor.h"
 #include "Framework/Commands/UICommandList.h"
 
 /** A graph node widget representing a niagara input node. */
-class SNiagaraGraphNodeInput : public SGraphNode
+class SNiagaraGraphNodeInput : public SNiagaraGraphNode
 {
 public:
 	SLATE_BEGIN_ARGS(SNiagaraGraphNodeInput) {}
diff --git a/Engine/Plugins/FX/Niagara/Source/NiagaraEditor/Private/Widgets/SNiagaraGraphNodeWriteDataSet.cpp b/Engine/Plugins/FX/Niagara/Source/NiagaraEditor/Private/Widgets/SNiagaraGraphNodeWriteDataSet.cpp
--- a/Engine/Plugins/FX/Niagara/Source/NiagaraEditor/Private/Widgets/SNiagaraGraphNodeWriteDataSet.cpp
+++ b/Engine/Plugins/FX/Niagara/Source/NiagaraEditor/Private/Widgets/SNiagaraGraphNodeWriteDataSet.cpp
@@ -20,7 +20,8 @@
 
 void SNiagaraGraphNodeWriteDataSet::Construct(const FArguments& InArgs, UEdGraphNode* InGraphNode)
 {
-	GraphNode = InGraphNode;
+	GraphNode = InGraphNode; 
+	RegisterNiagaraGraphNode(InGraphNode);
 
 	this->UpdateGraphNode();
 }
diff --git a/Engine/Plugins/FX/Niagara/Source/NiagaraEditor/Private/Widgets/SNiagaraGraphNodeWriteDataSet.h b/Engine/Plugins/FX/Niagara/Source/NiagaraEditor/Private/Widgets/SNiagaraGraphNodeWriteDataSet.h
--- a/Engine/Plugins/FX/Niagara/Source/NiagaraEditor/Private/Widgets/SNiagaraGraphNodeWriteDataSet.h
+++ b/Engine/Plugins/FX/Niagara/Source/NiagaraEditor/Private/Widgets/SNiagaraGraphNodeWriteDataSet.h
@@ -2,12 +2,12 @@
 
 #pragma once
 
-#include "SGraphNode.h"
+#include "SNiagaraGraphNode.h"
 #include "Widgets/Input/SMenuAnchor.h"
 #include "Framework/Commands/UICommandList.h"
 
 /** A graph node widget representing a niagara write data set node. */
-class SNiagaraGraphNodeWriteDataSet : public SGraphNode
+class SNiagaraGraphNodeWriteDataSet : public SNiagaraGraphNode
 {
 public:
 	SLATE_BEGIN_ARGS(SNiagaraGraphNodeWriteDataSet) {}
diff --git a/Engine/Plugins/FX/Niagara/Source/NiagaraEditor/Private/Widgets/SNiagaraGraphParameterMapGetNode.cpp b/Engine/Plugins/FX/Niagara/Source/NiagaraEditor/Private/Widgets/SNiagaraGraphParameterMapGetNode.cpp
--- a/Engine/Plugins/FX/Niagara/Source/NiagaraEditor/Private/Widgets/SNiagaraGraphParameterMapGetNode.cpp
+++ b/Engine/Plugins/FX/Niagara/Source/NiagaraEditor/Private/Widgets/SNiagaraGraphParameterMapGetNode.cpp
@@ -16,7 +16,9 @@ void SNiagaraGraphParameterMapGetNode::Construct(const FArguments& InArgs, UEdGr
 	BackgroundBrush = FEditorStyle::GetBrush("Graph.Pin.Background");
 	BackgroundHoveredBrush = FEditorStyle::GetBrush("PlainBorder");
 
-	GraphNode = InGraphNode;
+	GraphNode = InGraphNode; 
+	RegisterNiagaraGraphNode(InGraphNode);
+
 	UpdateGraphNode();
 }
 
diff --git a/Engine/Plugins/FX/Niagara/Source/NiagaraEditor/Private/Widgets/SNiagaraGraphParameterMapGetNode.h b/Engine/Plugins/FX/Niagara/Source/NiagaraEditor/Private/Widgets/SNiagaraGraphParameterMapGetNode.h
--- a/Engine/Plugins/FX/Niagara/Source/NiagaraEditor/Private/Widgets/SNiagaraGraphParameterMapGetNode.h
+++ b/Engine/Plugins/FX/Niagara/Source/NiagaraEditor/Private/Widgets/SNiagaraGraphParameterMapGetNode.h
@@ -1,9 +1,9 @@
 // Copyright 1998-2018 Epic Games, Inc. All Rights Reserved.
 
 #pragma once
-#include "SGraphNode.h"
+#include "SNiagaraGraphNode.h"
 
-class SNiagaraGraphParameterMapGetNode : public SGraphNode
+class SNiagaraGraphParameterMapGetNode : public SNiagaraGraphNode
 {
 public:
 	SLATE_BEGIN_ARGS(SNiagaraGraphParameterMapGetNode) {}
diff --git a/Engine/Plugins/FX/Niagara/Source/NiagaraEditor/Private/Widgets/SNiagaraGraphPinAdd.cpp b/Engine/Plugins/FX/Niagara/Source/NiagaraEditor/Private/Widgets/SNiagaraGraphPinAdd.cpp
--- a/Engine/Plugins/FX/Niagara/Source/NiagaraEditor/Private/Widgets/SNiagaraGraphPinAdd.cpp
+++ b/Engine/Plugins/FX/Niagara/Source/NiagaraEditor/Private/Widgets/SNiagaraGraphPinAdd.cpp
@@ -11,6 +11,8 @@
 #include "Widgets/Input/SEditableTextBox.h"
 #include "Widgets/Text/SInlineEditableTextBlock.h"
 #include "Widgets/Layout/SBox.h"
+#include "SNiagaraParameterMapView.h"
+#include "Framework/Application/SlateApplication.h"
 
 #define LOCTEXT_NAMESPACE "NiagaraGraphPinAdd"
 
@@ -60,30 +62,15 @@ TSharedRef<SWidget>	SNiagaraGraphPinAdd::ConstructAddButton()
 
 TSharedRef<SWidget> SNiagaraGraphPinAdd::OnGetAddButtonMenuContent()
 {
-	TArray<UEdGraphPin*> Pins = GetPinObj()->GetOwningNode()->GetAllPins();
-	EEdGraphPinDirection PinDir = GetPinObj()->Direction;
-
-	int32 FirstPinSameDir = Pins.IndexOfByPredicate([&](UEdGraphPin* InObj) -> bool
-	{
-		return InObj && (InObj->Direction == PinDir);
-	});
-
-	int32 IndexOfPin = Pins.IndexOfByPredicate([&](UEdGraphPin* InObj) -> bool
-	{
-		return InObj && (InObj == GetPinObj());
-	});
-
-	int PinIdx = 0;
-	if (FirstPinSameDir != -1 && IndexOfPin != -1)
-	{
-		PinIdx = IndexOfPin - FirstPinSameDir;
-	}
-
-	FString WorkingName = FString::Printf(TEXT("%s%d"), PinDir == EEdGraphPinDirection::EGPD_Input ? TEXT("Input") : TEXT("Output"), PinIdx);
-
 	if (OwningNode != nullptr)
 	{
-		return OwningNode->GenerateAddPinMenu(WorkingName, this);
+		TSharedRef<SNiagaraAddParameterMenu> MenuWidget = SNew(SNiagaraAddParameterMenu, OwningNode->GetNiagaraGraph())
+			.OnAddParameter_UObject(OwningNode, &UNiagaraNodeWithDynamicPins::AddParameter, GetPinObj()) // For non custom actions
+			.OnCollectCustomActions_UObject(OwningNode, &UNiagaraNodeWithDynamicPins::CollectAddPinActions, GetPinObj())
+			.OnAllowMakeType_UObject(OwningNode, &UNiagaraNodeWithDynamicPins::AllowNiagaraTypeForAddPin);
+		
+		FSlateApplication::Get().SetKeyboardFocus(MenuWidget->GetSearchBox());
+		return MenuWidget;
 	}
 	else
 	{
@@ -92,12 +79,4 @@ TSharedRef<SWidget> SNiagaraGraphPinAdd::OnGetAddButtonMenuContent()
 
 }
 
-void SNiagaraGraphPinAdd::OnAddType(FNiagaraVariable InAdd)
-{
-	if (OwningNode != nullptr)
-	{
-		FScopedTransaction AddNewPinTransaction(LOCTEXT("AddNewPinTransaction", "Add pin to node"));
-		OwningNode->RequestNewTypedPin(GetPinObj()->Direction, InAdd.GetType(), InAdd.GetName());
-	}
-}
 #undef LOCTEXT_NAMESPACE
diff --git a/Engine/Plugins/FX/Niagara/Source/NiagaraEditor/Private/Widgets/SNiagaraGraphPinAdd.h b/Engine/Plugins/FX/Niagara/Source/NiagaraEditor/Private/Widgets/SNiagaraGraphPinAdd.h
--- a/Engine/Plugins/FX/Niagara/Source/NiagaraEditor/Private/Widgets/SNiagaraGraphPinAdd.h
+++ b/Engine/Plugins/FX/Niagara/Source/NiagaraEditor/Private/Widgets/SNiagaraGraphPinAdd.h
@@ -16,7 +16,8 @@ class SNiagaraGraphPinAdd : public SGraphPin
 
 	void Construct(const FArguments& InArgs, UEdGraphPin* InGraphPinObj);
 	
-	void OnAddType(FNiagaraVariable InAdd);
+	/** Returns the owning node of this add pin. */
+	UNiagaraNodeWithDynamicPins* GetOwningNode() { return OwningNode; }
 
 private:
 	TSharedRef<SWidget>	ConstructAddButton();
diff --git a/Engine/Plugins/FX/Niagara/Source/NiagaraEditor/Private/Widgets/SNiagaraParameterCollection.cpp b/Engine/Plugins/FX/Niagara/Source/NiagaraEditor/Private/Widgets/SNiagaraParameterCollection.cpp
--- a/Engine/Plugins/FX/Niagara/Source/NiagaraEditor/Private/Widgets/SNiagaraParameterCollection.cpp
+++ b/Engine/Plugins/FX/Niagara/Source/NiagaraEditor/Private/Widgets/SNiagaraParameterCollection.cpp
@@ -435,41 +435,6 @@ TSharedRef<ITableRow> SNiagaraParameterCollection::OnGenerateRowForParameter(TSh
 			];
 	}
 
-	// Type Widget
-	TSharedPtr<SWidget> TypeWidget;
-	if (Item->CanChangeParameterType())
-	{
-		// Because we use templating behind the scenes to get the combo box to work, we will ultimately end up comparing the SharedPtr's by pointer value.
-		// Since the Available Types array comes from one location and the current type comes from another, their smart pointers are not guaranteed to 
-		// point to the same FNiagaraTypeDefinition in memory. We need to enforce that the current type is from the values in the AvailableTypes array, which is why 
-		// we look up the current type before creating the combo box.
-		const TArray<TSharedPtr<FNiagaraTypeDefinition>>& AvailableTypes = Collection->GetAvailableTypes();
-		TSharedPtr<FNiagaraTypeDefinition> CurrentType = Item->GetType();
-		if (CurrentType.IsValid())
-		{
-			const TSharedPtr<FNiagaraTypeDefinition>* CurrentTypeFromAvailableList = AvailableTypes.FindByPredicate([&](const TSharedPtr<FNiagaraTypeDefinition>& AvailableType) { return (*AvailableType.Get()) == (*CurrentType.Get()); });
-			if (CurrentTypeFromAvailableList != nullptr)
-				CurrentType = *CurrentTypeFromAvailableList;
-		}
-
-		SAssignNew(TypeWidget, SComboBox<TSharedPtr<FNiagaraTypeDefinition>>)
-			.OptionsSource(&AvailableTypes)
-			.OnGenerateWidget(this, &SNiagaraParameterCollection::OnGenerateWidgetForTypeComboBox)
-			.OnSelectionChanged(Item, &INiagaraParameterViewModel::SelectedTypeChanged)
-			.InitiallySelectedItem(CurrentType)
-			.Content()
-			[
-				SNew(STextBlock)
-				.TextStyle(FEditorStyle::Get(), "SmallText")
-				.Text(Item, &INiagaraParameterViewModel::GetTypeDisplayName)
-			];
-	}
-	else
-	{
-		SAssignNew(TypeWidget, STextBlock)
-			.Text(Item, &INiagaraParameterViewModel::GetTypeDisplayName);
-	}
-
 //#define DEBUG_SORT_ORDER
 #if defined(DEBUG_SORT_ORDER)
 	TSharedPtr<SWidget> SortOrderWidget;
@@ -486,12 +451,11 @@ TSharedRef<ITableRow> SNiagaraParameterCollection::OnGenerateRowForParameter(TSh
 	// Details and parameter editor widgets.
 	TSharedPtr<SWidget> CustomValueEditor;
 	TSharedPtr<SWidget> DetailsWidget;
-	bool bCanEdit = Item->IsEditingEnabled();//Can I veiw the details view but disable edits?
-	if (bCanEdit && Item->GetDefaultValueType() == INiagaraParameterViewModel::EDefaultValueType::Struct)
+	if (Item->GetDefaultValueType() == INiagaraParameterViewModel::EDefaultValueType::Struct)
 	{
 		FNiagaraEditorModule& NiagaraEditorModule = FModuleManager::GetModuleChecked<FNiagaraEditorModule>("NiagaraEditor");
 		FNiagaraTypeDefinition ParameterType = *Item->GetType().Get();
-		TSharedPtr<INiagaraEditorTypeUtilities> TypeEditorUtilities = NiagaraEditorModule.GetTypeUtilities(ParameterType);
+		TSharedPtr<INiagaraEditorTypeUtilities, ESPMode::ThreadSafe> TypeEditorUtilities = NiagaraEditorModule.GetTypeUtilities(ParameterType);
 		TSharedPtr<SNiagaraParameterEditor> ParameterEditor;
 		if (TypeEditorUtilities.IsValid() && TypeEditorUtilities->CanCreateParameterEditor())
 		{
@@ -514,17 +478,18 @@ TSharedRef<ITableRow> SNiagaraParameterCollection::OnGenerateRowForParameter(TSh
 
 		StructureDetailsView->SetStructureData(Item->GetDefaultValueStruct());
 		StructureDetailsView->GetOnFinishedChangingPropertiesDelegate().AddSP(Item, &INiagaraParameterViewModel::NotifyDefaultValuePropertyChanged);
-
+		
 		Item->OnDefaultValueChanged().AddSP(this, &SNiagaraParameterCollection::ParameterViewModelDefaultValueChanged, Item, ParameterEditor, StructureDetailsView);
 		Item->OnTypeChanged().AddSP(this, &SNiagaraParameterCollection::ParameterViewModelTypeChanged);
 
 		CustomValueEditor = ParameterEditor;
 		DetailsWidget = StructureDetailsView->GetWidget();
+		DetailsWidget->SetEnabled(Item->IsEditingEnabled());
 	}
-	else if (bCanEdit && Item->GetDefaultValueType() == INiagaraParameterViewModel::EDefaultValueType::Object)
+	else if (Item->GetDefaultValueType() == INiagaraParameterViewModel::EDefaultValueType::Object)
 	{
 		FNiagaraEditorModule& NiagaraEditorModule = FModuleManager::GetModuleChecked<FNiagaraEditorModule>("NiagaraEditor");
-		TSharedPtr<INiagaraEditorTypeUtilities> TypeEditorUtilities = NiagaraEditorModule.GetTypeUtilities(*Item->GetType().Get());
+		TSharedPtr<INiagaraEditorTypeUtilities, ESPMode::ThreadSafe> TypeEditorUtilities = NiagaraEditorModule.GetTypeUtilities(*Item->GetType().Get());
 		if (TypeEditorUtilities.IsValid() && TypeEditorUtilities->CanCreateDataInterfaceEditor())
 		{
 			CustomValueEditor = TypeEditorUtilities->CreateDataInterfaceEditor(Item->GetDefaultValueObject(),
@@ -545,6 +510,9 @@ TSharedRef<ITableRow> SNiagaraParameterCollection::OnGenerateRowForParameter(TSh
 
 	if (CustomValueEditor.IsValid())
 	{
+		// Should the value of the parameter be enabled?
+		CustomValueEditor->SetEnabled(Item->IsEditingEnabled());
+
 		return SNew(STableRow<TSharedRef<INiagaraParameterViewModel>>, OwnerTable)
 		//.IsEnabled(TAttribute<bool>(Item, &INiagaraParameterViewModel::IsEditingEnabled))
 		.ToolTipText(TAttribute<FText>(Item, &INiagaraParameterViewModel::GetTooltip))
@@ -595,12 +563,6 @@ TSharedRef<ITableRow> SNiagaraParameterCollection::OnGenerateRowForParameter(TSh
 				+ SVerticalBox::Slot()
 				.AutoHeight()
 				.Padding(25, 2, 0, 0)
-				[
-					TypeWidget.ToSharedRef()
-				]
-				+ SVerticalBox::Slot()
-				.AutoHeight()
-				.Padding(25, 2, 0, 0)
 				[
 					DetailsWidget.ToSharedRef()
 				]
@@ -636,12 +598,6 @@ TSharedRef<ITableRow> SNiagaraParameterCollection::OnGenerateRowForParameter(TSh
 				.OnSlotResized(SSplitter::FOnSlotResized::CreateSP(this, &SNiagaraParameterCollection::ParameterContentColumnWidthChanged))
 				[
 					SNew(SHorizontalBox)
-					+ SHorizontalBox::Slot()
-					.AutoWidth()
-					.Padding(FMargin(3, 0, 0, 0))
-					[
-						TypeWidget.ToSharedRef()
-					]
 #if defined(DEBUG_SORT_ORDER)
 					+ SHorizontalBox::Slot()
 					.AutoWidth()
diff --git a/Engine/Plugins/FX/Niagara/Source/NiagaraEditor/Private/Widgets/SNiagaraParameterMapPaletteItem.cpp b/Engine/Plugins/FX/Niagara/Source/NiagaraEditor/Private/Widgets/SNiagaraParameterMapPaletteItem.cpp
new file mode 100644
--- /dev/null
+++ b/Engine/Plugins/FX/Niagara/Source/NiagaraEditor/Private/Widgets/SNiagaraParameterMapPaletteItem.cpp
@@ -0,0 +1,108 @@
+// Copyright 1998-2018 Epic Games, Inc. All Rights Reserved.
+
+#include "SNiagaraParameterMapPaletteItem.h"
+#include "NiagaraActions.h"
+#include "EdGraphSchema_Niagara.h"
+#include "TutorialMetaData.h"
+#include "NiagaraGraph.h"
+#include "NiagaraEditorStyle.h"
+#include "Widgets/Input/SComboButton.h"
+#include "Widgets/Images/SImage.h"
+
+#define LOCTEXT_NAMESPACE "NiagaraParameterMapPalleteItem"
+
+void SNiagaraParameterMapPalleteItem::Construct(const FArguments& InArgs, FCreateWidgetForActionData* const InCreateData )
+{
+	this->OnItemRenamed = InArgs._OnItemRenamed;
+
+	TAttribute<bool> bIsReadOnly = false;
+	TAttribute<bool> bIsEditingEnabled = true;
+
+	check(InCreateData->Action.IsValid());
+	TSharedPtr<FNiagaraParameterAction> ParameterAction = StaticCastSharedPtr<FNiagaraParameterAction>(InCreateData->Action);
+	ActionPtr = InCreateData->Action;
+
+	FTutorialMetaData TagMeta("PaletteItem");
+
+	const FLinearColor TypeColor = UEdGraphSchema_Niagara::GetTypeColor(ParameterAction->GetParameter().GetType());
+	FSlateBrush const* IconBrush = FEditorStyle::GetBrush(TEXT("Kismet.AllClasses.VariableIcon"));
+	FSlateBrush const* SecondaryBrush = FEditorStyle::GetBrush(TEXT("NoBrush"));
+	FSlateColor        IconColor = FSlateColor(TypeColor);
+	FSlateColor        SecondaryIconColor = IconColor;
+	FText			   IconToolTip = FText::GetEmpty();
+	FString			   IconDocLink, IconDocExcerpt;
+	TSharedRef<SWidget> IconWidget = CreateIconWidget(IconToolTip, IconBrush, IconColor, IconDocLink, IconDocExcerpt, SecondaryBrush, SecondaryIconColor);
+	IconWidget->SetEnabled(bIsEditingEnabled);
+
+	static const FName BoldFontName = FName("Bold");
+	static const FName ItalicFontName = FName("Italic");
+	const FName FontType = ItalicFontName;
+	FSlateFontInfo Font = FCoreStyle::GetDefaultFontStyle(FontType, 10);
+	TSharedRef<SWidget> NameSlotWidget = CreateTextSlotWidget(Font, InCreateData, bIsReadOnly);
+
+	// now, create the actual widget
+	ChildSlot
+	[
+		SNew(SHorizontalBox)
+		.AddMetaData<FTutorialMetaData>(TagMeta)
+		.ToolTipText(ParameterAction->GetTooltipDescription())
+		// icon slot
+		+SHorizontalBox::Slot()
+		.AutoWidth()
+		[
+			IconWidget
+		]
+		// name slot
+		+SHorizontalBox::Slot()
+		.FillWidth(1.f)
+		.VAlign(VAlign_Center)
+		.Padding(3,0)
+		[
+			NameSlotWidget
+		]
+		// name slot
+		+ SHorizontalBox::Slot()
+		.AutoWidth()
+		.VAlign(VAlign_Center)
+		.Padding(3, 0)
+		[
+			SNew(SComboButton)
+			.HasDownArrow(false)
+			.ButtonStyle(FEditorStyle::Get(), "RoundButton")
+			.ForegroundColor(FSlateColor::UseForeground())
+			.ContentPadding(FMargin(2.0f))
+			.HAlign(HAlign_Right)
+			.VAlign(VAlign_Center)
+			.ButtonContent()
+			[
+				SNew(STextBlock)
+				.Text(this, &SNiagaraParameterMapPalleteItem::GetReferenceCount)
+				.Font(Font)
+			]
+		]
+
+	];
+}
+
+void SNiagaraParameterMapPalleteItem::OnNameTextCommitted(const FText& NewText, ETextCommit::Type InTextCommit)
+{
+	TSharedPtr<FNiagaraParameterAction> ParameterAction = StaticCastSharedPtr<FNiagaraParameterAction>(ActionPtr.Pin());
+	OnItemRenamed.ExecuteIfBound(NewText, *ParameterAction.Get());
+}
+
+FText SNiagaraParameterMapPalleteItem::GetReferenceCount() const
+{
+	TSharedPtr<FNiagaraParameterAction> ParameterAction = StaticCastSharedPtr<FNiagaraParameterAction>(ActionPtr.Pin());
+	if (ParameterAction.IsValid())
+	{
+		int32 TotalCount = 0;
+		for (const FNiagaraGraphParameterReferenceCollection& ReferenceCollection : ParameterAction->ReferenceCollection)
+		{
+			TotalCount += ReferenceCollection.ParameterReferences.Num();
+		}
+		return FText::AsNumber(TotalCount);
+	}
+	return FText::GetEmpty();
+}
+
+#undef LOCTEXT_NAMESPACE // "SNiagaraParameterMapPalleteItem"
\ No newline at end of file
diff --git a/Engine/Plugins/FX/Niagara/Source/NiagaraEditor/Private/Widgets/SNiagaraParameterMapPaletteItem.h b/Engine/Plugins/FX/Niagara/Source/NiagaraEditor/Private/Widgets/SNiagaraParameterMapPaletteItem.h
new file mode 100644
--- /dev/null
+++ b/Engine/Plugins/FX/Niagara/Source/NiagaraEditor/Private/Widgets/SNiagaraParameterMapPaletteItem.h
@@ -0,0 +1,33 @@
+// Copyright 1998-2018 Epic Games, Inc. All Rights Reserved.
+
+#pragma once
+
+#include "CoreMinimal.h"
+#include "Widgets/DeclarativeSyntaxSupport.h"
+#include "Widgets/SWidget.h"
+#include "Widgets/SCompoundWidget.h"
+#include "SGraphPalette.h"
+#include "NiagaraActions.h"
+#include "EdGraph/EdGraphNodeUtils.h"
+
+class SNiagaraParameterMapPalleteItem : public SGraphPaletteItem
+{
+public:
+
+	DECLARE_DELEGATE_TwoParams(FOnItemRenamed, const FText&, struct FNiagaraParameterAction&)
+
+	SLATE_BEGIN_ARGS(SNiagaraParameterMapPalleteItem)
+	{}
+		SLATE_EVENT(FOnItemRenamed, OnItemRenamed)
+	SLATE_END_ARGS();
+
+	void Construct(const FArguments& InArgs, FCreateWidgetForActionData* const InCreateData);
+
+protected:
+	/** Callback when rename text is committed */
+	virtual void OnNameTextCommitted(const FText& NewText, ETextCommit::Type InTextCommit) override;
+
+	FText GetReferenceCount() const;
+private:
+	FOnItemRenamed OnItemRenamed;
+};
\ No newline at end of file
