Copying //UE4/Dev-Niagara to //UE4/Dev-Main (Source: //UE4/Dev-Niagara @ 4074996)

diff --git a/Engine/Plugins/FX/Niagara/Source/NiagaraEditor/Private/Widgets/SNiagaraParameterMapView.cpp b/Engine/Plugins/FX/Niagara/Source/NiagaraEditor/Private/Widgets/SNiagaraParameterMapView.cpp
new file mode 100644
--- /dev/null
+++ b/Engine/Plugins/FX/Niagara/Source/NiagaraEditor/Private/Widgets/SNiagaraParameterMapView.cpp
@@ -0,0 +1,1009 @@
+// Copyright 1998-2018 Epic Games, Inc. All Rights Reserved.
+
+#include "SNiagaraParameterMapView.h"
+#include "SNiagaraParameterMapPaletteItem.h"
+#include "NiagaraObjectSelection.h"
+#include "Widgets/Images/SImage.h"
+#include "Widgets/Input/SButton.h"
+#include "Widgets/Input/SComboButton.h"
+#include "Widgets/Input/SSearchBox.h"
+#include "Widgets/SNiagaraGraphPinAdd.h"
+#include "NiagaraScript.h"
+#include "NiagaraScriptSource.h"
+#include "NiagaraEmitter.h"
+#include "NiagaraSystem.h"
+#include "NiagaraGraph.h"
+#include "NiagaraParameterStore.h"
+#include "NiagaraNodeWithDynamicPins.h"
+#include "NiagaraNodeParameterMapBase.h"
+#include "NiagaraActions.h"
+#include "SGraphActionMenu.h"
+#include "Framework/MultiBox/MultiBoxBuilder.h"
+#include "IDetailsView.h"
+#include "DetailLayoutBuilder.h"
+#include "NiagaraConstants.h"
+#include "EdGraph/EdGraphSchema.h"
+#include "Framework/Application/SlateApplication.h"
+#include "IAssetTools.h"
+#include "AssetRegistryModule.h"
+#include "Framework/Commands/GenericCommands.h"
+#include "EdGraph/EdGraphNode.h"
+#include "EdGraph/EdGraphPin.h"
+#include "EdGraphSchema_Niagara.h"
+#include "ViewModels/NiagaraEmitterHandleViewModel.h"
+
+#define LOCTEXT_NAMESPACE "NiagaraParameterMapView"
+
+DEFINE_LOG_CATEGORY(LogNiagara);
+
+FText NiagaraParameterMapSectionID::OnGetSectionTitle(const NiagaraParameterMapSectionID::Type InSection)
+{
+	/* Setup an appropriate name for the section for this node */
+	FText SectionTitle;
+	switch (InSection)
+	{
+	case NiagaraParameterMapSectionID::ENGINE:
+		SectionTitle = NSLOCTEXT("GraphActionNode", "Engine", "Engine");
+		break;
+	case NiagaraParameterMapSectionID::EMITTER:
+		SectionTitle = NSLOCTEXT("GraphActionNode", "Emitter", "Emitter");
+		break;
+	case NiagaraParameterMapSectionID::MODULE:
+		SectionTitle = NSLOCTEXT("GraphActionNode", "Module", "Module");
+		break;
+	case NiagaraParameterMapSectionID::SYSTEM:
+		SectionTitle = NSLOCTEXT("GraphActionNode", "System", "System");
+		break;
+	case NiagaraParameterMapSectionID::PARTICLE:
+		SectionTitle = NSLOCTEXT("GraphActionNode", "Particles", "Particles");
+		break;
+	case NiagaraParameterMapSectionID::USER:
+		SectionTitle = NSLOCTEXT("GraphActionNode", "User", "User");
+		break;
+	case NiagaraParameterMapSectionID::PARAMETERCOLLECTION:
+		SectionTitle = NSLOCTEXT("GraphActionNode", "ParameterCollection", "Parameter Collection");
+		break;
+	case NiagaraParameterMapSectionID::OTHER:
+		SectionTitle = NSLOCTEXT("GraphActionNode", "Other", "Other");
+		break;
+	case NiagaraParameterMapSectionID::NONE:
+		SectionTitle = FText::GetEmpty();
+		break;
+	}
+	return SectionTitle;
+}
+
+NiagaraParameterMapSectionID::Type NiagaraParameterMapSectionID::OnGetSectionFromVariable(const FNiagaraVariable& InVar, FNiagaraParameterHandle& OutParameterHandle, const NiagaraParameterMapSectionID::Type DefaultType)
+{
+	OutParameterHandle = FNiagaraParameterHandle(InVar.GetName());
+	Type SectionID = DefaultType;
+	if (OutParameterHandle.IsEmitterHandle())
+	{
+		SectionID = NiagaraParameterMapSectionID::EMITTER;
+	}
+	else if (OutParameterHandle.IsModuleHandle())
+	{
+		SectionID = NiagaraParameterMapSectionID::MODULE;
+	}
+	else if (OutParameterHandle.IsUserHandle())
+	{
+		SectionID = NiagaraParameterMapSectionID::USER;
+	}
+	else if (OutParameterHandle.IsEngineHandle())
+	{
+		SectionID = NiagaraParameterMapSectionID::ENGINE;
+	}
+	else if (OutParameterHandle.IsSystemHandle())
+	{
+		SectionID = NiagaraParameterMapSectionID::SYSTEM;
+	}
+	else if (OutParameterHandle.IsParticleAttributeHandle())
+	{
+		SectionID = NiagaraParameterMapSectionID::PARTICLE;
+	}
+	else if (OutParameterHandle.IsParameterCollectionHandle())
+	{
+		SectionID = NiagaraParameterMapSectionID::PARAMETERCOLLECTION;
+	}
+
+	return SectionID;
+}
+
+void FNiagaraParameterMapViewCommands::RegisterCommands()
+{
+	UI_COMMAND(DeleteEntry, "Delete", "Deletes this parameter.", EUserInterfaceActionType::Button, FInputChord(EKeys::Platform_Delete));
+}
+
+SNiagaraParameterMapView::~SNiagaraParameterMapView()
+{
+	TSet<UObject*> Objects = SelectedObjects->GetSelectedObjects();
+	for (UObject* Object : Objects)
+	{
+		if (UNiagaraSystem* System = Cast<UNiagaraSystem>(Object))
+		{
+			System->GetExposedParameters().RemoveAllOnChangedHandlers(this);
+			break;
+		}
+	}
+
+	Graphs.Empty();
+	SelectedObjects->OnSelectedObjectsChanged().RemoveAll(this);
+}
+
+void SNiagaraParameterMapView::Construct(const FArguments& InArgs, const TSharedRef<FNiagaraObjectSelection>& InSelectedObjects, const EToolkitType InToolkitType, const TSharedPtr<FUICommandList>& InToolkitCommands)
+{
+	bNeedsRefresh = false;
+	ToolkitType = InToolkitType;
+	ToolkitCommands = InToolkitCommands;
+
+	SelectedObjects = InSelectedObjects;
+	SelectedObjects->OnSelectedObjectsChanged().AddSP(this, &SNiagaraParameterMapView::SelectedObjectsChanged);
+	
+	// Register all commands for right click on action node
+	{
+		FNiagaraParameterMapViewCommands::Register();
+		TSharedPtr<FUICommandList> ToolKitCommandList = ToolkitCommands;
+		ToolKitCommandList->MapAction(FNiagaraParameterMapViewCommands::Get().DeleteEntry,
+			FExecuteAction::CreateSP(this, &SNiagaraParameterMapView::OnDeleteEntry),
+			FCanExecuteAction::CreateSP(this, &SNiagaraParameterMapView::CanDeleteEntry));
+		ToolKitCommandList->MapAction(FGenericCommands::Get().Rename,
+			FExecuteAction::CreateSP(this, &SNiagaraParameterMapView::OnRequestRenameOnActionNode),
+			FCanExecuteAction::CreateSP(this, &SNiagaraParameterMapView::CanRequestRenameOnActionNode));
+	}
+
+	Refresh(false);
+
+	SAssignNew(FilterBox, SSearchBox)
+		.OnTextChanged(this, &SNiagaraParameterMapView::OnFilterTextChanged);
+
+	// create the main action list piece of this widget
+	SAssignNew(GraphActionMenu, SGraphActionMenu, false)
+		.OnGetFilterText(this, &SNiagaraParameterMapView::GetFilterText)
+		.OnCreateWidgetForAction(this, &SNiagaraParameterMapView::OnCreateWidgetForAction)
+		.OnCollectAllActions(this, &SNiagaraParameterMapView::CollectAllActions)
+		.OnCollectStaticSections(this, &SNiagaraParameterMapView::CollectStaticSections)
+		.OnActionDragged(this, &SNiagaraParameterMapView::OnActionDragged)
+		.OnActionSelected(this, &SNiagaraParameterMapView::OnActionSelected)
+		.OnActionDoubleClicked(this, &SNiagaraParameterMapView::OnActionDoubleClicked)
+		.OnContextMenuOpening(this, &SNiagaraParameterMapView::OnContextMenuOpening)
+		.OnCanRenameSelectedAction(this, &SNiagaraParameterMapView::CanRequestRenameOnActionNode)
+		.OnGetSectionTitle(this, &SNiagaraParameterMapView::OnGetSectionTitle)
+		.OnGetSectionWidget(this, &SNiagaraParameterMapView::OnGetSectionWidget)
+		.OnCreateCustomRowExpander_Static(&SNiagaraParameterMapView::CreateCustomActionExpander)
+		.OnActionMatchesName(this, &SNiagaraParameterMapView::HandleActionMatchesName)
+		.AutoExpandActionMenu(false)
+		.AlphaSortItems(false)
+		.UseSectionStyling(true)
+		.ShowFilterTextBox(true);
+
+	ChildSlot
+	[
+		SNew(SBox)
+		.MinDesiredWidth(300)
+		[
+			SNew(SVerticalBox)
+
+			+ SVerticalBox::Slot()
+			.AutoHeight()
+			[
+				SNew(SBorder)
+				.Padding(4.0f)
+				.BorderImage(FEditorStyle::GetBrush("ToolPanel.GroupBorder"))
+				.AddMetaData<FTagMetaData>(FTagMetaData(TEXT("ParameterMapPanel")))
+				[
+					SNew(SVerticalBox)
+					+ SVerticalBox::Slot()
+					.AutoHeight()
+					[
+						SNew(SHorizontalBox)
+
+						+ SHorizontalBox::Slot()
+						.FillWidth(1.0f)
+						.VAlign(VAlign_Center)
+						[
+							FilterBox.ToSharedRef()
+						]
+					]
+				]
+			]
+		
+			+ SVerticalBox::Slot()
+			.FillHeight(1.0f)
+			[
+				GraphActionMenu.ToSharedRef()
+			]
+		]
+	];
+}
+
+void SNiagaraParameterMapView::Tick(const FGeometry& AllottedGeometry, const double InCurrentTime, const float InDeltaTime)
+{
+	if (bNeedsRefresh)
+	{
+		GraphActionMenu->RefreshAllActions(true);
+		bNeedsRefresh = false;
+	}
+}
+
+bool SNiagaraParameterMapView::IsEditingMode() const
+{
+	return true;
+}
+
+void SNiagaraParameterMapView::AddParameter(FNiagaraVariable NewVariable)
+{
+	bool bAddedParameter = false;
+	// Check whether we have to add this parameter to the user exposed system parameters.
+	FNiagaraParameterHandle ParameterHandle;
+	if (NiagaraParameterMapSectionID::OnGetSectionFromVariable(NewVariable, ParameterHandle) == NiagaraParameterMapSectionID::USER)
+	{
+		for (UObject* Object : SelectedObjects->GetSelectedObjects())
+		{
+			if (UNiagaraSystem* System = Cast<UNiagaraSystem>(Object))
+			{
+				bAddedParameter = System->GetExposedParameters().AddParameter(NewVariable);
+				break;
+			}
+		}
+	}
+
+	if (!bAddedParameter && Graphs.Num() > 0)
+	{
+		for (const TWeakObjectPtr<UNiagaraGraph>& GraphWeakPtr : Graphs)
+		{
+			if (GraphWeakPtr.IsValid())
+			{
+				UNiagaraGraph* Graph = GraphWeakPtr.Get();
+				Graph->AddParameter(NewVariable);
+			}
+		}
+		GraphActionMenu->RefreshAllActions(true);
+		GraphActionMenu->SelectItemByName(NewVariable.GetName());
+		GraphActionMenu->OnRequestRenameOnActionNode();
+	}
+}
+
+void SNiagaraParameterMapView::OnFilterTextChanged(const FText& InFilterText)
+{
+	GraphActionMenu->GenerateFilteredItems(false);
+}
+
+FText SNiagaraParameterMapView::GetFilterText() const
+{
+	return FilterBox->GetText();
+}
+
+TSharedRef<SWidget> SNiagaraParameterMapView::OnCreateWidgetForAction(struct FCreateWidgetForActionData* const InCreateData)
+{
+	return SNew(SNiagaraParameterMapPalleteItem, InCreateData)
+		.OnItemRenamed(this, &SNiagaraParameterMapView::OnPostRenameActionNode);
+}
+
+void SNiagaraParameterMapView::CollectAllActions(FGraphActionListBuilderBase& OutAllActions)
+{
+	if (Graphs.Num() == 0)
+	{
+		return;
+	}
+
+	TMap<FNiagaraVariable, TArray<FNiagaraGraphParameterReferenceCollection>> ParameterEntries;
+	for (auto& GraphWeakPtr : Graphs)
+	{
+		UNiagaraGraph* Graph = GraphWeakPtr.Get();
+		for (const auto& ParameterElement : Graph->GetParameterMap())
+		{
+			TArray<FNiagaraGraphParameterReferenceCollection>* Found = ParameterEntries.Find(ParameterElement.Key);
+			if (Found)
+			{
+				Found->Add(ParameterElement.Value);
+			}
+			else
+			{
+				TArray<FNiagaraGraphParameterReferenceCollection> Collection;
+				Collection.Add(ParameterElement.Value);
+				ParameterEntries.Add(ParameterElement.Key, Collection);
+			}
+		}
+	}
+
+	TSet<UObject*> Objects = SelectedObjects->GetSelectedObjects();
+	for (UObject* Object : Objects)
+	{
+		if (UNiagaraSystem* System = Cast<UNiagaraSystem>(Object))
+		{
+			TArray<FNiagaraVariable> ExposedVars;
+			System->GetExposedParameters().GetParameters(ExposedVars);
+			for (const FNiagaraVariable& ExposedVar : ExposedVars)
+			{
+				TArray<FNiagaraGraphParameterReferenceCollection>* Found = ParameterEntries.Find(ExposedVar);
+				if (!Found)
+				{
+					TArray<FNiagaraGraphParameterReferenceCollection> Collection;
+					ParameterEntries.Add(ExposedVar, Collection);
+				}
+			}
+		}
+	}
+
+	ParameterEntries.KeySort([](const FNiagaraVariable& A, const FNiagaraVariable& B) { return (A.GetName() < B.GetName()); });
+
+	const FText TooltipFormat = LOCTEXT("Parameters", "Name: {0} Type: {1}");
+	for (const auto& ParameterEntry : ParameterEntries)
+	{
+		const FNiagaraVariable& Parameter = ParameterEntry.Key;
+		FNiagaraParameterHandle Handle;
+		const NiagaraParameterMapSectionID::Type Section = NiagaraParameterMapSectionID::OnGetSectionFromVariable(Parameter, Handle, /*Default*/ NiagaraParameterMapSectionID::OTHER);
+		if (!IsSystemToolkit() || (IsSystemToolkit() && Section != NiagaraParameterMapSectionID::MODULE))
+		{
+			const FText Name = FText::FromName(Parameter.GetName());
+			const FText Tooltip = FText::Format(TooltipFormat, FText::FromName(Parameter.GetName()), Parameter.GetType().GetNameText());
+			TSharedPtr<FNiagaraParameterAction> ParameterAction(new FNiagaraParameterAction(Parameter, ParameterEntry.Value, FText::GetEmpty(), Name, Tooltip, 0, FText(), Section));
+			OutAllActions.AddAction(ParameterAction);
+		}
+	}
+}
+
+void SNiagaraParameterMapView::CollectStaticSections(TArray<int32>& StaticSectionIDs)
+{
+	if (!IsSystemToolkit())
+	{
+		StaticSectionIDs.Add(NiagaraParameterMapSectionID::MODULE);
+	}
+	StaticSectionIDs.Add(NiagaraParameterMapSectionID::ENGINE);
+	StaticSectionIDs.Add(NiagaraParameterMapSectionID::EMITTER);
+	StaticSectionIDs.Add(NiagaraParameterMapSectionID::SYSTEM);
+	StaticSectionIDs.Add(NiagaraParameterMapSectionID::PARTICLE);
+	StaticSectionIDs.Add(NiagaraParameterMapSectionID::USER);
+	StaticSectionIDs.Add(NiagaraParameterMapSectionID::PARAMETERCOLLECTION);
+	StaticSectionIDs.Add(NiagaraParameterMapSectionID::OTHER);
+}
+
+FReply SNiagaraParameterMapView::OnActionDragged(const TArray<TSharedPtr<FEdGraphSchemaAction>>& InActions, const FPointerEvent& MouseEvent)
+{
+	TSharedPtr<FEdGraphSchemaAction> InAction(InActions.Num() > 0 ? InActions[0] : NULL);
+	if (InAction.IsValid())
+	{
+		FNiagaraParameterAction* ParameterAction = (FNiagaraParameterAction*)InAction.Get();
+		if (ParameterAction)
+		{
+			if (IsScriptToolkit())
+			{
+				TSharedRef<FNiagaraParameterGraphDragOperation> DragOperation = FNiagaraParameterGraphDragOperation::New(InAction);
+				DragOperation->SetAltDrag(MouseEvent.IsAltDown());
+				DragOperation->SetCtrlDrag(MouseEvent.IsLeftControlDown() || MouseEvent.IsRightControlDown());
+				return FReply::Handled().BeginDragDrop(DragOperation);
+			}
+			else if (IsSystemToolkit())
+			{
+				TSharedRef<FNiagaraStackDragOperation> DragOperation = FNiagaraStackDragOperation::New(InAction);
+				return FReply::Handled().BeginDragDrop(DragOperation);
+			}
+		}
+	}
+
+	return FReply::Handled();
+}
+
+void SNiagaraParameterMapView::OnActionSelected(const TArray< TSharedPtr<FEdGraphSchemaAction> >& InActions, ESelectInfo::Type InSelectionType)
+{
+
+}
+
+void SNiagaraParameterMapView::OnActionDoubleClicked(const TArray< TSharedPtr<FEdGraphSchemaAction> >& InActions)
+{
+
+}
+
+TSharedPtr<SWidget> SNiagaraParameterMapView::OnContextMenuOpening()
+{
+	// Check if the selected action is valid for a context menu
+	if (SelectionHasContextMenu())
+	{
+		const bool bShouldCloseWindowAfterMenuSelection = true;
+		FMenuBuilder MenuBuilder(bShouldCloseWindowAfterMenuSelection, ToolkitCommands);
+		MenuBuilder.BeginSection("BasicOperations");
+		{
+			MenuBuilder.AddMenuEntry(FGenericCommands::Get().Rename, NAME_None, LOCTEXT("Rename", "Rename"), LOCTEXT("Rename_Tooltip", "Renames this parameter"));
+			MenuBuilder.AddMenuEntry(FNiagaraParameterMapViewCommands::Get().DeleteEntry);
+		}
+		MenuBuilder.EndSection();
+		return MenuBuilder.MakeWidget();
+	}
+
+	return SNullWidget::NullWidget;
+}
+
+FText SNiagaraParameterMapView::OnGetSectionTitle(int32 InSectionID)
+{
+	return NiagaraParameterMapSectionID::OnGetSectionTitle((NiagaraParameterMapSectionID::Type)InSectionID);
+}
+
+TSharedRef<SWidget> SNiagaraParameterMapView::OnGetSectionWidget(TSharedRef<SWidget> RowWidget, int32 InSectionID)
+{
+	TWeakPtr<SWidget> WeakRowWidget = RowWidget;
+	FText AddNewText = LOCTEXT("AddNewParameter", "Add Parameter");
+	FName MetaDataTag = TEXT("AddNewParameter");
+	return CreateAddToSectionButton((NiagaraParameterMapSectionID::Type) InSectionID, WeakRowWidget, AddNewText, MetaDataTag);
+}
+
+TSharedRef<SWidget> SNiagaraParameterMapView::CreateAddToSectionButton(const NiagaraParameterMapSectionID::Type InSection, TWeakPtr<SWidget> WeakRowWidget, FText AddNewText, FName MetaDataTag)
+{
+	TSharedPtr<SComboButton> Button;
+	SAssignNew(Button, SComboButton)
+	.ButtonStyle(FEditorStyle::Get(), "RoundButton")
+	.ForegroundColor(FEditorStyle::GetSlateColor("DefaultForeground"))
+	.ContentPadding(FMargin(2, 0))
+	.OnGetMenuContent(this, &SNiagaraParameterMapView::OnGetParameterMenu, InSection)
+	.IsEnabled(this, &SNiagaraParameterMapView::IsEditingMode)
+	.HAlign(HAlign_Center)
+	.VAlign(VAlign_Center)
+	.HasDownArrow(false)
+	.AddMetaData<FTagMetaData>(FTagMetaData(MetaDataTag))
+	.ButtonContent()
+	[
+		SNew(SHorizontalBox)
+
+		+ SHorizontalBox::Slot()
+		.AutoWidth()
+		.Padding(FMargin(0, 1))
+		[
+			SNew(SImage)
+			.Image(FEditorStyle::GetBrush("Plus"))
+		]
+
+		+ SHorizontalBox::Slot()
+		.VAlign(VAlign_Center)
+		.AutoWidth()
+		.Padding(FMargin(2,0,0,0))
+		[
+			SNew(STextBlock)
+			.Font(IDetailLayoutBuilder::GetDetailFontBold())
+			.Text(AddNewText)
+			.Visibility(this, &SNiagaraParameterMapView::OnAddButtonTextVisibility, WeakRowWidget, InSection)
+			.ShadowOffset(FVector2D(1,1))
+		]
+	];
+
+	return Button.ToSharedRef();
+}
+
+bool SNiagaraParameterMapView::SelectionHasContextMenu() const
+{
+	TArray<TSharedPtr<FEdGraphSchemaAction>> SelectedActions;
+	GraphActionMenu->GetSelectedActions(SelectedActions);
+	return SelectedActions.Num() > 0;
+}
+
+TSharedRef<SWidget> SNiagaraParameterMapView::OnGetParameterMenu(const NiagaraParameterMapSectionID::Type InSection)
+{
+	TSharedRef<SNiagaraAddParameterMenu> MenuWidget = SNew(SNiagaraAddParameterMenu, Graphs[0].Get())
+		.OnAddParameter(this, &SNiagaraParameterMapView::AddParameter)
+		.Section(InSection)
+		.ShowNamespaceCategory(false)
+		.ShowGraphParameters(false)
+		.AutoExpandMenu(true);
+
+	FSlateApplication::Get().SetKeyboardFocus(MenuWidget->GetSearchBox());
+	return MenuWidget;
+}
+
+EVisibility SNiagaraParameterMapView::OnAddButtonTextVisibility(TWeakPtr<SWidget> RowWidget, const NiagaraParameterMapSectionID::Type InSection) const
+{
+	return EVisibility::Collapsed; // RowWidget.Pin()->IsHovered() ? EVisibility::SelfHitTestInvisible : EVisibility::Collapsed;
+}
+
+void SNiagaraParameterMapView::Refresh(bool bRefreshMenu/* = true*/)
+{
+	Graphs.Empty();
+
+	TSet<UObject*> Objects = SelectedObjects->GetSelectedObjects();
+	for (UObject* Object : Objects)
+	{
+		if (UNiagaraScript* Script = Cast<UNiagaraScript>(Object))
+		{
+			AddGraph(Script->GetSource());
+			break;
+		}
+		else if (UNiagaraEmitter* Emitter = Cast<UNiagaraEmitter>(Object))
+		{
+			AddGraph(Emitter->GraphSource);
+			break;
+		}
+		else if (UNiagaraSystem* System = Cast<UNiagaraSystem>(Object))
+		{
+			for (const FNiagaraEmitterHandle& EmitterHandle : System->GetEmitterHandles())
+			{
+				AddGraph(EmitterHandle.GetInstance()->GraphSource);
+			}
+			System->GetExposedParameters().AddOnChangedHandler(FNiagaraParameterStore::FOnChanged::FDelegate::CreateSP(this, &SNiagaraParameterMapView::RefreshActions));
+			break;
+		}
+	}
+
+	if (bRefreshMenu)
+	{
+		GraphActionMenu->RefreshAllActions(true);
+	}
+}
+
+void SNiagaraParameterMapView::RefreshEmitterHandles(const TArray<TSharedPtr<FNiagaraEmitterHandleViewModel>>& EmitterHandles)
+{
+	Graphs.Empty();
+	for (const TSharedPtr<FNiagaraEmitterHandleViewModel>& Handle : EmitterHandles)
+	{
+		AddGraph(Handle->GetEmitterHandle()->GetInstance()->GraphSource);
+	}
+	GraphActionMenu->RefreshAllActions(true);
+}
+
+void SNiagaraParameterMapView::SelectedObjectsChanged()
+{
+	Refresh(true);
+}
+
+void SNiagaraParameterMapView::AddGraph(UNiagaraGraph* Graph)
+{
+	if (Graph)
+	{
+		Graphs.AddUnique(Graph);
+		OnGraphChangedHandle = Graph->AddOnGraphChangedHandler(
+			FOnGraphChanged::FDelegate::CreateRaw(this, &SNiagaraParameterMapView::OnGraphChanged));
+		OnRecompileHandle = Graph->AddOnGraphNeedsRecompileHandler(
+			FOnGraphChanged::FDelegate::CreateRaw(this, &SNiagaraParameterMapView::OnGraphChanged));
+	}
+}
+
+void SNiagaraParameterMapView::AddGraph(UNiagaraScriptSourceBase* SourceBase)
+{
+	UNiagaraScriptSource* Source = Cast<UNiagaraScriptSource>(SourceBase);
+	if (Source)
+	{
+		AddGraph(Source->NodeGraph);
+	}
+}
+
+void SNiagaraParameterMapView::OnGraphChanged(const FEdGraphEditAction& InAction)
+{
+	RefreshActions();
+}
+
+void SNiagaraParameterMapView::OnDeleteEntry()
+{
+	TArray<TSharedPtr<FEdGraphSchemaAction>> SelectedActions;
+	GraphActionMenu->GetSelectedActions(SelectedActions);
+	
+	TArray<UNiagaraGraph*> GraphsToNotifyChanged;
+	for (auto& Action : SelectedActions)
+	{
+		TSharedPtr<FNiagaraParameterAction> ParameterAction = StaticCastSharedPtr<FNiagaraParameterAction>(Action);
+		if (ParameterAction.Get())
+		{
+			FNiagaraParameterHandle ParameterHandle;
+			if (NiagaraParameterMapSectionID::OnGetSectionFromVariable(ParameterAction->GetParameter(), ParameterHandle) == NiagaraParameterMapSectionID::USER)
+			{
+				for (UObject* Object : SelectedObjects->GetSelectedObjects())
+				{
+					if (UNiagaraSystem* System = Cast<UNiagaraSystem>(Object))
+					{
+						System->GetExposedParameters().RemoveParameter(ParameterAction->GetParameter());
+					}
+				}
+			}
+
+			for (const TWeakObjectPtr<UNiagaraGraph>& GraphWeakPtr : Graphs)
+			{
+				if (GraphWeakPtr.IsValid())
+				{
+					UNiagaraGraph* Graph = GraphWeakPtr.Get();
+					Graph->RemoveParameter(ParameterAction->GetParameter(), false);
+					GraphsToNotifyChanged.Add(Graph);
+				}
+			}
+		}
+	}
+
+	for (UNiagaraGraph* Graph : GraphsToNotifyChanged)
+	{
+		Graph->NotifyGraphChanged();
+	}
+}
+
+bool SNiagaraParameterMapView::CanDeleteEntry() const
+{
+	return true;
+}
+
+void SNiagaraParameterMapView::OnRequestRenameOnActionNode()
+{
+	// Attempt to rename in both menus, only one of them will have anything selected
+	GraphActionMenu->OnRequestRenameOnActionNode();
+}
+
+
+bool SNiagaraParameterMapView::CanRequestRenameOnActionNode(TWeakPtr<struct FGraphActionNode> InSelectedNode) const
+{
+	return true;
+}
+
+bool SNiagaraParameterMapView::CanRequestRenameOnActionNode() const
+{
+	TArray<TSharedPtr<FEdGraphSchemaAction> > SelectedActions;
+	GraphActionMenu->GetSelectedActions(SelectedActions);
+
+	// If there is anything selected in the GraphActionMenu, check the item for if it can be renamed.
+	if (SelectedActions.Num())
+	{
+		return GraphActionMenu->CanRequestRenameOnActionNode();
+	}
+	return false;
+}
+
+void SNiagaraParameterMapView::OnPostRenameActionNode(const FText& InText, FNiagaraParameterAction& InAction)
+{
+	const FName NewName = FName(*InText.ToString());
+	if (InAction.Parameter.GetName() != NewName)
+	{
+		FNiagaraParameterHandle ParameterHandle;
+		if (NiagaraParameterMapSectionID::OnGetSectionFromVariable(InAction.Parameter, ParameterHandle) == NiagaraParameterMapSectionID::USER)
+		{
+			// Check if the new name is also an user variable.
+			const FNiagaraVariable NewParameterValidTest = FNiagaraVariable(InAction.Parameter.GetType(), NewName);
+			if (NiagaraParameterMapSectionID::OnGetSectionFromVariable(NewParameterValidTest, ParameterHandle) == NiagaraParameterMapSectionID::USER)
+			{
+				for (UObject* Object : SelectedObjects->GetSelectedObjects())
+				{
+					if (UNiagaraSystem* System = Cast<UNiagaraSystem>(Object))
+					{
+						System->GetExposedParameters().RenameParameter(InAction.Parameter, NewName);
+					}
+				}
+			}
+		}
+
+		if (Graphs.Num() > 0)
+		{
+			for (const TWeakObjectPtr<UNiagaraGraph>& Graph : Graphs)
+			{
+				if (Graph.IsValid())
+				{
+					Graph.Get()->RenameParameter(InAction.Parameter, NewName);
+				}
+			}
+		}
+	}
+}
+
+bool SNiagaraParameterMapView::IsSystemToolkit()
+{
+	return ToolkitType == EToolkitType::SYSTEM;
+}
+
+bool SNiagaraParameterMapView::IsScriptToolkit()
+{
+	return ToolkitType == EToolkitType::SCRIPT;
+}
+
+bool SNiagaraParameterMapView::HandleActionMatchesName(FEdGraphSchemaAction* InAction, const FName& InName) const
+{
+	return FName(*InAction->GetMenuDescription().ToString()) == InName;
+}
+
+void SNiagaraParameterMapView::RefreshActions()
+{
+	bNeedsRefresh = true;
+}
+
+/************************************************************************/
+/* SNiagaraAddParameterMenu                                             */
+/************************************************************************/
+void SNiagaraAddParameterMenu::Construct(const FArguments& InArgs, UNiagaraGraph* InGraph)
+{
+	this->OnAddParameter = InArgs._OnAddParameter;
+	this->OnCollectCustomActions = InArgs._OnCollectCustomActions;
+	this->OnAllowMakeType = InArgs._OnAllowMakeType;
+	this->Section = InArgs._Section;
+	this->AllowCreatingNew = InArgs._AllowCreatingNew;
+	this->ShowNamespaceCategory = InArgs._ShowNamespaceCategory;
+	this->ShowGraphParameters = InArgs._ShowGraphParameters;
+	this->AutoExpandMenu = InArgs._AutoExpandMenu;
+
+	Graph = InGraph;
+
+	ChildSlot
+	[
+		SNew(SBorder)
+		.BorderImage(FEditorStyle::GetBrush("Menu.Background"))
+		.Padding(5)
+		[
+			SNew(SBox)
+			.MinDesiredWidth(300)
+			[
+				SAssignNew(GraphMenu, SGraphActionMenu)
+				.OnActionSelected(this, &SNiagaraAddParameterMenu::OnActionSelected)
+				.OnCollectAllActions(this, &SNiagaraAddParameterMenu::CollectAllActions)
+				.AutoExpandActionMenu(AutoExpandMenu.Get())
+				.ShowFilterTextBox(true)
+				.OnCreateCustomRowExpander_Static(&SNiagaraParameterMapView::CreateCustomActionExpander)
+			]
+		]
+	];
+}
+
+TSharedRef<SEditableTextBox> SNiagaraAddParameterMenu::GetSearchBox()
+{
+	return GraphMenu->GetFilterTextBox();
+}
+
+void SNiagaraAddParameterMenu::OnActionSelected(const TArray<TSharedPtr<FEdGraphSchemaAction>>& SelectedActions, ESelectInfo::Type InSelectionType)
+{
+	if (InSelectionType == ESelectInfo::OnMouseClick || InSelectionType == ESelectInfo::OnKeyPress || SelectedActions.Num() == 0)
+	{
+		for (int32 ActionIndex = 0; ActionIndex < SelectedActions.Num(); ActionIndex++)
+		{
+			TSharedPtr<FNiagaraMenuAction> CurrentAction = StaticCastSharedPtr<FNiagaraMenuAction>(SelectedActions[ActionIndex]);
+
+			if (CurrentAction.IsValid())
+			{
+				FSlateApplication::Get().DismissAllMenus();
+				CurrentAction->ExecuteAction();
+			}
+		}
+	}
+}
+
+void SNiagaraAddParameterMenu::CollectAllActions(FGraphActionListBuilderBase& OutAllActions)
+{
+	if (OnCollectCustomActions.IsBound())
+	{
+		bool bCreateRemainingActions = true;
+		OnCollectCustomActions.Execute(OutAllActions, bCreateRemainingActions);
+		if (!bCreateRemainingActions)
+		{
+			return;
+		}
+	}
+
+	auto CanCollectSection = [&](const NiagaraParameterMapSectionID::Type GivenSectionID)
+	{
+		NiagaraParameterMapSectionID::Type ID = Section.Get();
+		return ID == NiagaraParameterMapSectionID::NONE || (ID != NiagaraParameterMapSectionID::NONE && ID == GivenSectionID);
+	};
+
+	// Particle
+	if (CanCollectSection(NiagaraParameterMapSectionID::PARTICLE))
+	{
+		bool bSupportsAttributes = true;
+		if (Graph.Get())
+		{
+			bool IsModule = Graph->FindOutputNode(ENiagaraScriptUsage::Module) != nullptr || Graph->FindOutputNode(ENiagaraScriptUsage::DynamicInput) != nullptr
+				|| Graph->FindOutputNode(ENiagaraScriptUsage::Function) != nullptr;
+
+			UNiagaraScriptSource* Source = Cast<UNiagaraScriptSource>(Graph->GetOuter());
+			if (Source && IsModule)
+			{
+				UNiagaraScript* Script = Cast<UNiagaraScript>(Source->GetOuter());
+				if (Script)
+				{
+					TArray<ENiagaraScriptUsage> Usages = Script->GetSupportedUsageContexts();
+					if (Usages.Contains(ENiagaraScriptUsage::ParticleEventScript) || Usages.Contains(ENiagaraScriptUsage::ParticleSpawnScript) || Usages.Contains(ENiagaraScriptUsage::ParticleUpdateScript))
+					{
+						bSupportsAttributes = true;
+					}
+					else
+					{
+						bSupportsAttributes = false;
+					}
+				}
+			}
+		}
+
+		if (bSupportsAttributes)
+		{
+			const FText Category = ShowNamespaceCategory.Get() ? NiagaraParameterMapSectionID::OnGetSectionTitle(NiagaraParameterMapSectionID::PARTICLE) : FText::GetEmpty();
+			TArray<FNiagaraVariable> Variables = FNiagaraConstants::GetCommonParticleAttributes();
+			AddParameterGroup(OutAllActions, Variables, NiagaraParameterMapSectionID::PARTICLE, Category, FString(), true, false);
+			CollectMakeNew(OutAllActions, NiagaraParameterMapSectionID::PARTICLE);
+		}
+	}
+
+	// Emitter
+	if (CanCollectSection(NiagaraParameterMapSectionID::EMITTER))
+	{
+		CollectMakeNew(OutAllActions, NiagaraParameterMapSectionID::EMITTER);
+	}
+
+	// Module
+	if (CanCollectSection(NiagaraParameterMapSectionID::MODULE))
+	{
+		CollectMakeNew(OutAllActions, NiagaraParameterMapSectionID::MODULE);
+	}
+
+	// System
+	if (CanCollectSection(NiagaraParameterMapSectionID::SYSTEM))
+	{
+		CollectMakeNew(OutAllActions, NiagaraParameterMapSectionID::SYSTEM);
+	}
+
+	// User
+	if (CanCollectSection(NiagaraParameterMapSectionID::USER))
+	{
+		CollectMakeNew(OutAllActions, NiagaraParameterMapSectionID::USER);
+	}
+
+	// Parameter collections
+	if (CanCollectSection(NiagaraParameterMapSectionID::PARAMETERCOLLECTION))
+	{
+		CollectParameterCollectionsActions(OutAllActions);
+	}
+
+	if (CanCollectSection(NiagaraParameterMapSectionID::OTHER))
+	{
+		CollectMakeNew(OutAllActions, NiagaraParameterMapSectionID::OTHER);
+	}
+
+	// Collect locals
+	if (Graph.Get())
+	{
+		TArray<FNiagaraVariable> Parameters;
+
+		if (ShowGraphParameters.Get())
+		{
+			for (const auto& Parameter : Graph.Get()->GetParameterMap())
+			{
+				Parameters.AddUnique(Parameter.Key);
+			}
+		}
+
+		Parameters.Sort([](const FNiagaraVariable& A, const FNiagaraVariable& B) { return (A.GetName() < B.GetName()); });
+
+		for (const FNiagaraVariable& Parameter : Parameters)
+		{
+			FNiagaraParameterHandle Handle;
+			const NiagaraParameterMapSectionID::Type ParameterSectionID = NiagaraParameterMapSectionID::OnGetSectionFromVariable(Parameter, Handle, /*Default*/ NiagaraParameterMapSectionID::OTHER);
+			if (CanCollectSection(ParameterSectionID))
+			{
+				const FText Category = ShowNamespaceCategory.Get() ? NiagaraParameterMapSectionID::OnGetSectionTitle(ParameterSectionID) : FText::GetEmpty();
+				const FText DisplayName = FText::FromName(Parameter.GetName());
+				const FText Tooltip = FText::GetEmpty();
+				TSharedPtr<FNiagaraMenuAction> Action(new FNiagaraMenuAction(
+					Category, DisplayName, Tooltip, 0, FText::GetEmpty(),
+					FNiagaraMenuAction::FOnExecuteStackAction::CreateSP(this, &SNiagaraAddParameterMenu::AddParameterSelected, Parameter, false, ParameterSectionID)));
+				OutAllActions.AddAction(Action);
+			}
+		}
+	}
+
+	// Engine
+	if (CanCollectSection(NiagaraParameterMapSectionID::ENGINE))
+	{
+		const FText Category = NiagaraParameterMapSectionID::OnGetSectionTitle(NiagaraParameterMapSectionID::ENGINE);
+		TArray<FNiagaraVariable> Variables = FNiagaraConstants::GetEngineConstants();
+		AddParameterGroup(OutAllActions,
+			Variables,
+			NiagaraParameterMapSectionID::ENGINE,
+			ShowNamespaceCategory.Get() ? NiagaraParameterMapSectionID::OnGetSectionTitle(NiagaraParameterMapSectionID::ENGINE) : FText::GetEmpty(),
+			FString(), true, false);
+	}
+
+}
+
+void SNiagaraAddParameterMenu::AddParameterGroup(FGraphActionListBuilderBase& OutActions, TArray<FNiagaraVariable>& Variables, const NiagaraParameterMapSectionID::Type InSection, const FText& Category, const FString& RootCategory, const bool bSort, const bool bCustomName)
+{
+	if (bSort)
+	{
+		Variables.Sort([](const FNiagaraVariable& A, const FNiagaraVariable& B) { return (A.GetName() < B.GetName()); });
+	}
+
+	for (FNiagaraVariable& Variable : Variables)
+	{
+		const FText DisplayName = FText::FromName(Variable.GetName());
+		const FText Tooltip = FText::GetEmpty();
+		TSharedPtr<FNiagaraMenuAction> Action(new FNiagaraMenuAction(Category, DisplayName, Tooltip, 0, FText(),
+			FNiagaraMenuAction::FOnExecuteStackAction::CreateSP(this, &SNiagaraAddParameterMenu::AddParameterSelected, Variable, bCustomName, InSection)));
+		OutActions.AddAction(Action, RootCategory);
+	}
+}
+
+void SNiagaraAddParameterMenu::CollectParameterCollectionsActions(FGraphActionListBuilderBase& OutActions)
+{
+	//Create sub menus for parameter collections.
+	FAssetRegistryModule& AssetRegistryModule = FModuleManager::LoadModuleChecked<FAssetRegistryModule>("AssetRegistry");
+	TArray<FAssetData> CollectionAssets;
+	AssetRegistryModule.Get().GetAssetsByClass(UNiagaraParameterCollection::StaticClass()->GetFName(), CollectionAssets);
+
+	const FText Category = NiagaraParameterMapSectionID::OnGetSectionTitle(NiagaraParameterMapSectionID::PARAMETERCOLLECTION);
+	for (FAssetData& CollectionAsset : CollectionAssets)
+	{
+		UNiagaraParameterCollection* Collection = CastChecked<UNiagaraParameterCollection>(CollectionAsset.GetAsset());
+		if (Collection)
+		{
+			AddParameterGroup(OutActions, Collection->GetParameters(), NiagaraParameterMapSectionID::PARAMETERCOLLECTION, Category, FString(), true, false);
+		}
+	}
+}
+
+void SNiagaraAddParameterMenu::CollectMakeNew(FGraphActionListBuilderBase& OutActions, const NiagaraParameterMapSectionID::Type InSection)
+{
+	if (!AllowCreatingNew.Get())
+	{
+		return;
+	}
+
+	TArray<FNiagaraVariable> Variables;
+	TArray<FNiagaraTypeDefinition> Types = FNiagaraTypeRegistry::GetRegisteredTypes();
+	for (const FNiagaraTypeDefinition& RegisteredType : Types)
+	{
+		bool bAllowType = true;
+		if (OnAllowMakeType.IsBound())
+		{
+			bAllowType = OnAllowMakeType.Execute(RegisteredType);
+		}
+
+		if (bAllowType)
+		{
+			FNiagaraVariable Var(RegisteredType, FName(*RegisteredType.GetNameText().ToString()));
+			FNiagaraEditorUtilities::ResetVariableToDefaultValue(Var);
+			Variables.Add(Var);
+		}
+	}
+
+	AddParameterGroup(OutActions, Variables, InSection,
+		LOCTEXT("MakeNewCat", "Make New"), 
+		ShowNamespaceCategory.Get() ? NiagaraParameterMapSectionID::OnGetSectionTitle(InSection).ToString() : FString(), 
+		true, true);
+}
+
+void SNiagaraAddParameterMenu::AddParameterSelected(FNiagaraVariable NewVariable, const bool bCreateCustomName, const NiagaraParameterMapSectionID::Type InSection)
+{
+	if (bCreateCustomName)
+	{
+		const FString ResultName = (InSection != NiagaraParameterMapSectionID::NONE ? NiagaraParameterMapSectionID::OnGetSectionTitle(InSection).ToString() + TEXT(".") : FString())
+			+ FString::Printf(TEXT("%s%d"), *NewVariable.GetName().ToString(), Graph->GetParameterMap().Num());
+		NewVariable.SetName(FName(*ResultName));
+	}
+
+	OnAddParameter.ExecuteIfBound(NewVariable);
+}
+
+class SNiagaraActionMenuExpander : public SExpanderArrow
+{
+	SLATE_BEGIN_ARGS(SNiagaraActionMenuExpander) {}
+		SLATE_ATTRIBUTE(float, IndentAmount)
+	SLATE_END_ARGS()
+
+public:
+	void Construct(const FArguments& InArgs, const FCustomExpanderData& ActionMenuData)
+	{
+		OwnerRowPtr = ActionMenuData.TableRow;
+		IndentAmount = InArgs._IndentAmount;
+		if (!ActionMenuData.RowAction.IsValid())
+		{
+			SExpanderArrow::FArguments SuperArgs;
+			SuperArgs._IndentAmount = InArgs._IndentAmount;
+
+			SExpanderArrow::Construct(SuperArgs, ActionMenuData.TableRow);
+		}
+		else
+		{
+			ChildSlot
+			.Padding(TAttribute<FMargin>(this, &SNiagaraActionMenuExpander::GetCustomIndentPadding))
+			[	
+				SNew(SBox)
+			];
+		}
+	}
+
+private:
+	FMargin GetCustomIndentPadding() const
+	{
+		return SExpanderArrow::GetExpanderPadding();
+	}
+};
+
+
+TSharedRef<SExpanderArrow> SNiagaraParameterMapView::CreateCustomActionExpander(const FCustomExpanderData& ActionMenuData)
+{
+	return SNew(SNiagaraActionMenuExpander, ActionMenuData);
+}
+
+#undef LOCTEXT_NAMESPACE // "NiagaraParameterMapView"
\ No newline at end of file
diff --git a/Engine/Plugins/FX/Niagara/Source/NiagaraEditor/Private/Widgets/SNiagaraScriptGraph.cpp b/Engine/Plugins/FX/Niagara/Source/NiagaraEditor/Private/Widgets/SNiagaraScriptGraph.cpp
--- a/Engine/Plugins/FX/Niagara/Source/NiagaraEditor/Private/Widgets/SNiagaraScriptGraph.cpp
+++ b/Engine/Plugins/FX/Niagara/Source/NiagaraEditor/Private/Widgets/SNiagaraScriptGraph.cpp
@@ -18,6 +18,7 @@
 #include "EdGraphSchema_Niagara.h"
 #include "ScopedTransaction.h"
 #include "NiagaraEditorUtilities.h"
+#include "NiagaraNodeFactory.h"
 
 #define LOCTEXT_NAMESPACE "NiagaraScriptGraph"
 
@@ -77,12 +78,17 @@ TSharedRef<SGraphEditor> SNiagaraScriptGraph::ConstructGraphEditor()
 	Events.OnVerifyTextCommit = FOnNodeVerifyTextCommit::CreateSP(this, &SNiagaraScriptGraph::OnVerifyNodeTextCommit);
 	Events.OnSpawnNodeByShortcut = SGraphEditor::FOnSpawnNodeByShortcut::CreateSP(this, &SNiagaraScriptGraph::OnSpawnGraphNodeByShortcut);
 
-	return SNew(SGraphEditor)
+	TSharedRef<SGraphEditor> CreatedGraphEditor = SNew(SGraphEditor)
 		.AdditionalCommands(ViewModel->GetCommands())
 		.Appearance(AppearanceInfo)
 		.TitleBar(TitleBarWidget)
 		.GraphToEdit(ViewModel->GetGraph())
 		.GraphEvents(Events);
+
+	// Set a niagara node factory.
+	CreatedGraphEditor->SetNodeFactory(MakeShareable(new FNiagaraNodeFactory()));
+
+	return CreatedGraphEditor;
 }
 
 void SNiagaraScriptGraph::ViewModelSelectedNodesChanged()
diff --git a/Engine/Plugins/FX/Niagara/Source/NiagaraEditor/Private/Widgets/SNiagaraSelectedEmitterHandles.cpp b/Engine/Plugins/FX/Niagara/Source/NiagaraEditor/Private/Widgets/SNiagaraSelectedEmitterHandles.cpp
new file mode 100644
--- /dev/null
+++ b/Engine/Plugins/FX/Niagara/Source/NiagaraEditor/Private/Widgets/SNiagaraSelectedEmitterHandles.cpp
@@ -0,0 +1,151 @@
+// Copyright 1998-2016 Epic Games, Inc. All Rights Reserved.
+
+#include "SNiagaraSelectedEmitterHandles.h"
+#include "NiagaraEditorStyle.h"
+#include "ViewModels/NiagaraSystemViewModel.h"
+#include "ViewModels/NiagaraEmitterHandleViewModel.h"
+#include "ViewModels/NiagaraEmitterViewModel.h"
+#include "ViewModels/Stack/NiagaraStackViewModel.h"
+#include "NiagaraEditorModule.h"
+#include "NiagaraEditorCommands.h"
+#include "Widgets/Layout/SSplitter.h"
+
+#include "Modules/ModuleManager.h"
+#include "Widgets/Layout/SBox.h"
+#include "Widgets/Text/STextBlock.h"
+
+#define LOCTEXT_NAMESPACE "NiagaraSelectedEmittersHandle"
+
+void SNiagaraSelectedEmitterHandles::Construct(const FArguments& InArgs, TSharedRef<FNiagaraSystemViewModel> InSystemViewModel)
+{
+	SystemViewModel = InSystemViewModel;
+	SystemViewModel->OnSelectedEmitterHandlesChanged().AddRaw(this, &SNiagaraSelectedEmitterHandles::SelectedEmitterHandlesChanged);
+	SystemViewModel->OnEmitterHandleViewModelsChanged().AddRaw(this, &SNiagaraSelectedEmitterHandles::EmitterHandleViewModelsChanged);
+	SystemViewModel->GetOnPinnedEmittersChanged().AddRaw(this, &SNiagaraSelectedEmitterHandles::OnEmitterPinnedChanged);
+
+	{
+		TSharedPtr<FUICommandList> ToolkitCommands = SystemViewModel->GetToolkitCommands();
+		ToolkitCommands->MapAction(
+			FNiagaraEditorCommands::Get().CollapseStackToHeaders,
+			FExecuteAction::CreateSP(this, &SNiagaraSelectedEmitterHandles::CollapseToHeaders));
+	}
+
+	ChildSlot
+	[
+		SNew(SOverlay)
+		+ SOverlay::Slot()
+		[
+			SAssignNew(EmitterSplitter, SSplitter)
+			.MinimumSlotHeight(150.0f)
+		]
+		+ SOverlay::Slot()
+		.Padding(0, 20, 0, 0)
+		.HAlign(HAlign_Center)
+		[
+			SNew(STextBlock)
+			.Text(this, &SNiagaraSelectedEmitterHandles::GetUnsupportedSelectionText)
+			.TextStyle(FNiagaraEditorStyle::Get(), "NiagaraEditor.SelectedEmitter.UnsupportedSelectionText")
+			.Visibility(this, &SNiagaraSelectedEmitterHandles::GetUnsupportedSelectionTextVisibility)
+		]
+	];
+	TArray<TSharedRef<FNiagaraEmitterHandleViewModel>> SelectedEmitterHandles;
+	SystemViewModel->GetSelectedEmitterHandles(SelectedEmitterHandles);
+	RefreshEmitterWidgets();
+}
+
+SNiagaraSelectedEmitterHandles::~SNiagaraSelectedEmitterHandles()
+{
+	SystemViewModel->OnEmitterHandleViewModelsChanged().RemoveAll(this);
+	SystemViewModel->OnSelectedEmitterHandlesChanged().RemoveAll(this);
+	SystemViewModel->GetOnPinnedEmittersChanged().RemoveAll(this);
+
+	for (auto Model : StackViewModels)
+	{
+		Model->Initialize(nullptr, nullptr);
+	}
+}
+
+void SNiagaraSelectedEmitterHandles::AddReferencedObjects(FReferenceCollector& Collector)
+{
+	Collector.AddReferencedObjects(StackViewModels);
+}
+
+void SNiagaraSelectedEmitterHandles::RefreshEmitterWidgets()
+{
+	ClearViewModels();
+	TArray<TSharedPtr<FNiagaraEmitterHandleViewModel>> EmitterHandlesToDisplay;
+	EmitterHandlesToDisplay.Append(SystemViewModel->GetPinnedEmitterHandles());
+	TArray<TSharedRef<FNiagaraEmitterHandleViewModel>> SelectedEmitterHandles;
+	SystemViewModel->GetSelectedEmitterHandles(SelectedEmitterHandles);
+	for (auto Handle : SelectedEmitterHandles)
+	{
+		EmitterHandlesToDisplay.AddUnique(Handle);
+	}
+
+	FNiagaraEditorModule& NiagaraEditorModule = FModuleManager::GetModuleChecked<FNiagaraEditorModule>(TEXT("NiagaraEditor"));
+
+	for (auto EmitterHandleViewModel : EmitterHandlesToDisplay)
+	{
+		UNiagaraStackViewModel* StackModel = NewObject<UNiagaraStackViewModel>();
+		StackModel->Initialize(SystemViewModel, EmitterHandleViewModel);
+		StackViewModels.Add(StackModel);
+		EmitterSplitter->AddSlot()
+		[
+			NiagaraEditorModule.CreateStackWidget(StackModel)
+		];
+	}
+	EmitterHandlesToDisplay.Empty();
+}
+
+void SNiagaraSelectedEmitterHandles::ClearViewModels()
+{
+	while (EmitterSplitter->GetChildren()->Num() > 0)
+	{
+		int i = EmitterSplitter->GetChildren()->Num() - 1;
+		EmitterSplitter->RemoveAt(i);
+		StackViewModels[i]->Initialize(nullptr, nullptr);
+	}
+	StackViewModels.Empty();
+}
+
+void SNiagaraSelectedEmitterHandles::SelectedEmitterHandlesChanged()
+{
+	RefreshEmitterWidgets();
+}
+
+void SNiagaraSelectedEmitterHandles::EmitterHandleViewModelsChanged()
+{
+	RefreshEmitterWidgets();
+}
+
+void SNiagaraSelectedEmitterHandles::OnEmitterPinnedChanged()
+{
+	RefreshEmitterWidgets();
+}
+
+EVisibility SNiagaraSelectedEmitterHandles::GetUnsupportedSelectionTextVisibility() const
+{
+	return SystemViewModel->GetSelectedEmitterHandleIds().Num() != 1 ? EVisibility::Visible : EVisibility::Collapsed;
+}
+
+FText SNiagaraSelectedEmitterHandles::GetUnsupportedSelectionText() const
+{
+	if (SystemViewModel->GetSelectedEmitterHandleIds().Num() == 0 && SystemViewModel->GetPinnedEmitterHandles().Num() == 0)
+	{
+		return LOCTEXT("NoSelectionMessage", "Select an emitter in the timeline.");
+	}
+	else
+	{
+		return FText();
+	}
+}
+
+void SNiagaraSelectedEmitterHandles::CollapseToHeaders()
+{
+	for (UNiagaraStackViewModel* ViewModel : StackViewModels)
+	{
+		ViewModel->CollapseToHeaders();
+	}
+}
+
+#undef LOCTEXT_NAMESPACE
\ No newline at end of file
diff --git a/Engine/Plugins/FX/Niagara/Source/NiagaraEditor/Private/Widgets/SNiagaraSelectedEmitterHandles.h b/Engine/Plugins/FX/Niagara/Source/NiagaraEditor/Private/Widgets/SNiagaraSelectedEmitterHandles.h
new file mode 100644
--- /dev/null
+++ b/Engine/Plugins/FX/Niagara/Source/NiagaraEditor/Private/Widgets/SNiagaraSelectedEmitterHandles.h
@@ -0,0 +1,53 @@
+// Copyright 1998-2016 Epic Games, Inc. All Rights Reserved.
+
+#pragma once
+
+#include "Widgets/SCompoundWidget.h"
+#include "Widgets/DeclarativeSyntaxSupport.h"
+#include "UObject/GCObject.h"
+#include "Framework/Commands/Commands.h"
+#include "Framework/Commands/UICommandList.h"
+#include "EditorStyleSet.h"
+
+class FNiagaraSystemViewModel;
+class UNiagaraStackViewModel;
+class SNiagaraStack;
+class SBox;
+class SSplitter;
+
+class SNiagaraSelectedEmitterHandles : public SCompoundWidget, public FGCObject
+{
+public:
+	SLATE_BEGIN_ARGS(SNiagaraSelectedEmitterHandles)
+	{}
+	SLATE_END_ARGS();
+
+	void Construct(const FArguments& InArgs, TSharedRef<FNiagaraSystemViewModel> InSystemViewModel);
+
+	~SNiagaraSelectedEmitterHandles();
+
+	virtual void AddReferencedObjects(FReferenceCollector& Collector) override;
+
+private:
+	void SelectedEmitterHandlesChanged();
+
+	void EmitterHandleViewModelsChanged();
+
+	void RefreshEmitterWidgets();
+
+	EVisibility GetUnsupportedSelectionTextVisibility() const;
+
+	FText GetUnsupportedSelectionText() const;
+
+	void OnEmitterPinnedChanged();
+
+	void ClearViewModels();
+
+	void CollapseToHeaders();
+private:
+	TSharedPtr<FNiagaraSystemViewModel> SystemViewModel;
+
+	TArray<UNiagaraStackViewModel*> StackViewModels;
+
+	TSharedPtr<SSplitter> EmitterSplitter;
+};
diff --git a/Engine/Plugins/FX/Niagara/Source/NiagaraEditor/Private/Widgets/SNiagaraSpreadsheetView.cpp b/Engine/Plugins/FX/Niagara/Source/NiagaraEditor/Private/Widgets/SNiagaraSpreadsheetView.cpp
--- a/Engine/Plugins/FX/Niagara/Source/NiagaraEditor/Private/Widgets/SNiagaraSpreadsheetView.cpp
+++ b/Engine/Plugins/FX/Niagara/Source/NiagaraEditor/Private/Widgets/SNiagaraSpreadsheetView.cpp
@@ -6,9 +6,10 @@
 #include "Framework/MultiBox/MultiBoxBuilder.h"
 #include "Widgets/Input/SSearchBox.h"
 #include "Widgets/Input/SButton.h"
+#include "Widgets/Input/SComboButton.h"
 #include "ISequencer.h"
-#include "NiagaraSystemViewModel.h"
-#include "NiagaraEmitterHandleViewModel.h"
+#include "ViewModels/NiagaraSystemViewModel.h"
+#include "ViewModels/NiagaraEmitterHandleViewModel.h"
 #include "NiagaraEmitterHandle.h"
 #include "NiagaraEmitter.h"
 #include "NiagaraScript.h"
@@ -22,6 +23,14 @@
 #include "NiagaraSystem.h"
 #include "NiagaraEditorStyle.h"
 #include "HAL/PlatformApplicationMisc.h"
+#include "UObject/UObjectIterator.h"
+#include "NiagaraComponent.h"
+#include "ViewModels/NiagaraEmitterHandleViewModel.h"
+#include "NiagaraEmitterHandle.h"
+#include "Kismet2/DebuggerCommands.h"
+#include "Editor/EditorEngine.h"
+#include "UnrealEdGlobals.h"
+#include "EngineGlobals.h"
 
 #define LOCTEXT_NAMESPACE "SNiagaraSpreadsheetView"
 #define ARRAY_INDEX_COLUMN_NAME TEXT("Array Index")
@@ -62,17 +71,22 @@ TSharedRef< SWidget > SNiagaraSpreadsheetRow::GenerateWidgetForColumn(const FNam
 	}
 	else if (!ColumnsAreAttributes && (ColumnName == INPUT_KEY_COLUMN_NAME || ColumnName == OUTPUT_KEY_COLUMN_NAME))
 	{
+		FText Text;
+		if (SupportedFields.IsValid())
+		{
+			Text = FText::FromName((*SupportedFields.Get())[RowIndex]);
+		}
 		EntryWidget = SNew(STextBlock)
-			.Text(FText::FromName((*SupportedFields.Get())[RowIndex]));
+			.Text(Text);
 	}
 	else if (ColumnsAreAttributes)
 	{
-		FieldInfo = FieldInfoMap->Find(ColumnName);
+		FieldInfo = FieldInfoMap.IsValid() ? FieldInfoMap->Find(ColumnName) : nullptr;
 		RealRowIdx = RowIndex;
 	}
 	else if (!ColumnsAreAttributes && ColumnName == VALUE_COLUMN_NAME)
 	{
-		FieldInfo = FieldInfoMap->Find((*SupportedFields.Get())[RowIndex]);
+		FieldInfo = FieldInfoMap.IsValid() ? FieldInfoMap->Find((*SupportedFields.Get())[RowIndex]) : nullptr;
 	}
 
 	if (FieldInfo != nullptr && (UseGlobalOffsets ? ParameterStore != nullptr : DataSet != nullptr) && !EntryWidget.IsValid())
@@ -81,19 +95,25 @@ TSharedRef< SWidget > SNiagaraSpreadsheetRow::GenerateWidgetForColumn(const FNam
 		if (FieldInfo->bFloat)
 		{
 			float* Src = nullptr;
-			if (UseGlobalOffsets)
+			if (UseGlobalOffsets && ParameterStore->GetParameterDataArray().GetData() != nullptr)
 			{
 				uint32 CompBufferOffset = FieldInfo->GlobalStartOffset;
 				Src = (float*)(ParameterStore->GetParameterDataArray().GetData() + CompBufferOffset);
 			}
-			else
+			else if (DataSet != nullptr && DataSet->PrevData().GetNumInstances() > 0)
 			{
 				uint32 CompBufferOffset = FieldInfo->FloatStartOffset;
 				Src = DataSet->PrevData().GetInstancePtrFloat(CompBufferOffset, RealRowIdx);
 			}
 
+			float Value = 0.0f;
+			if (Src != nullptr)
+			{
+				Value = Src[0];
+			}
+
 			EntryWidget = SNew(STextBlock)
-				.Text(FText::AsNumber(Src[0]));
+				.Text(FText::AsNumber(Value));
 
 		}
 		else if (FieldInfo->bBoolean)
@@ -104,17 +124,17 @@ TSharedRef< SWidget > SNiagaraSpreadsheetRow::GenerateWidgetForColumn(const FNam
 				uint32 CompBufferOffset = FieldInfo->GlobalStartOffset;
 				Src = (int32*)(ParameterStore->GetParameterDataArray().GetData() + CompBufferOffset);
 			}
-			else
+			else if (DataSet && DataSet->PrevData().GetNumInstances() > 0)
 			{
 				uint32 CompBufferOffset = FieldInfo->IntStartOffset;
 				Src = DataSet->PrevData().GetInstancePtrInt32(CompBufferOffset, RealRowIdx);
 			}
 			FText ValueText;
-			if (Src[0] == 0)
+			if (Src && Src[0] == 0)
 			{
 				ValueText = LOCTEXT("NiagaraFalse", "False(0)");
 			}
-			else if (Src[0] == -1)
+			else if (Src && Src[0] == -1)
 			{
 				ValueText = LOCTEXT("NiagaraTrue", "True(-1)");
 			}
@@ -134,13 +154,19 @@ TSharedRef< SWidget > SNiagaraSpreadsheetRow::GenerateWidgetForColumn(const FNam
 				uint32 CompBufferOffset = FieldInfo->GlobalStartOffset;
 				Src = (int32*)(ParameterStore->GetParameterDataArray().GetData() + CompBufferOffset);
 			}
-			else
+			else if (DataSet != nullptr && DataSet->PrevData().GetNumInstances() > 0)
 			{
 				uint32 CompBufferOffset = FieldInfo->IntStartOffset;
 				Src = DataSet->PrevData().GetInstancePtrInt32(CompBufferOffset, RealRowIdx);
 			}
+
+			int32 Value = 0;
+			if (Src)
+			{
+				Value = Src[0];
+			}
 			EntryWidget = SNew(STextBlock)
-				.Text(FText::Format(LOCTEXT("EnumValue","{0}({1})"), FieldInfo->Enum->GetDisplayNameTextByValue(Src[0]), FText::AsNumber(Src[0])));
+				.Text(FText::Format(LOCTEXT("EnumValue","{0}({1})"), FieldInfo->Enum->GetDisplayNameTextByValue(Value), FText::AsNumber(Value)));
 
 		}
 		else
@@ -151,13 +177,19 @@ TSharedRef< SWidget > SNiagaraSpreadsheetRow::GenerateWidgetForColumn(const FNam
 				uint32 CompBufferOffset = FieldInfo->GlobalStartOffset;
 				Src = (int32*)(ParameterStore->GetParameterDataArray().GetData() + CompBufferOffset);
 			}
-			else
+			else if (DataSet && DataSet->PrevData().GetNumInstances() > 0)
 			{
 				uint32 CompBufferOffset = FieldInfo->IntStartOffset;
 				Src = DataSet->PrevData().GetInstancePtrInt32(CompBufferOffset, RealRowIdx);
 			}
+
+			int32 Value = 0;
+			if (Src)
+			{
+				Value = Src[0];
+			}
 			EntryWidget = SNew(STextBlock)
-				.Text(FText::AsNumber(Src[0]));
+				.Text(FText::AsNumber(Value));
 
 		}
 	}
@@ -182,24 +214,65 @@ void SNiagaraSpreadsheetView::Construct(const FArguments& InArgs, TSharedRef<FNi
 {
 	TabState = UIPerParticleUpdate;
 	ScriptEnum = FindObject<UEnum>(ANY_PACKAGE, TEXT("ENiagaraScriptUsage"));
+	TargetComponent = InSystemViewModel->GetPreviewComponent();
 	ensure(ScriptEnum);
 	
-	CaptureData.SetNum(2);
-	CaptureData[(int)UIPerParticleUpdate].TargetUsage = ENiagaraScriptUsage::ParticleUpdateScript;
-	CaptureData[(int)UISystemUpdate].TargetUsage = ENiagaraScriptUsage::SystemUpdateScript;
-	CaptureData[(int)UIPerParticleUpdate].ColumnName = LOCTEXT("PerParticleUpdate", "Per-Particle Update");
-	CaptureData[(int)UISystemUpdate].ColumnName = LOCTEXT("SystemUpdate", "System Update");
-	CaptureData[(int)UIPerParticleUpdate].bOutputColumnsAreAttributes = true;
-	CaptureData[(int)UISystemUpdate].bOutputColumnsAreAttributes = false;
-	CaptureData[(int)UIPerParticleUpdate].bInputColumnsAreAttributes = false;
-	CaptureData[(int)UISystemUpdate].bInputColumnsAreAttributes = false;
+	CaptureData.SetNum(UIMax);
 
 	SystemViewModel = InSystemViewModel;
 	SystemViewModel->OnSelectedEmitterHandlesChanged().AddRaw(this, &SNiagaraSpreadsheetView::SelectedEmitterHandlesChanged);
 	SystemViewModel->OnPostSequencerTimeChanged().AddRaw(this, &SNiagaraSpreadsheetView::OnSequencerTimeChanged);
 
+	bInitialColumns = true;
+
 	for (int32 i = 0; i < (int32)UIMax; i++)
 	{
+		switch (i)
+		{
+		case UIPerParticleUpdate:
+			CaptureData[i].TargetUsage = ENiagaraScriptUsage::ParticleUpdateScript;
+			CaptureData[i].ColumnName = LOCTEXT("PerParticleUpdate", "Particle Update");
+			CaptureData[i].bOutputColumnsAreAttributes = true;
+			CaptureData[i].bInputColumnsAreAttributes = false;
+			break;
+		case UIPerParticleSpawn:
+			CaptureData[i].TargetUsage = ENiagaraScriptUsage::ParticleSpawnScript;
+			CaptureData[i].ColumnName = LOCTEXT("PerParticleSpawn", "Particle Spawn");
+			CaptureData[i].bOutputColumnsAreAttributes = true;
+			CaptureData[i].bInputColumnsAreAttributes = false;
+			break;
+		case UIPerParticleEvent0:
+			CaptureData[i].TargetUsage = ENiagaraScriptUsage::ParticleEventScript;
+			CaptureData[i].ColumnName = LOCTEXT("PerParticleEvent0", "Particle Event0");
+			CaptureData[i].bOutputColumnsAreAttributes = true;
+			CaptureData[i].bInputColumnsAreAttributes = false;
+			break;
+		case UIPerParticleEvent1:
+			CaptureData[i].TargetUsage = ENiagaraScriptUsage::ParticleEventScript;
+			CaptureData[i].ColumnName = LOCTEXT("PerParticleEvent1", "Particle Event1");
+			CaptureData[i].bOutputColumnsAreAttributes = true;
+			CaptureData[i].bInputColumnsAreAttributes = false;
+			break;
+		case UIPerParticleEvent2:
+			CaptureData[i].TargetUsage = ENiagaraScriptUsage::ParticleEventScript;
+			CaptureData[i].ColumnName = LOCTEXT("PerParticleEvent2", "Particle Event2");
+			CaptureData[i].bOutputColumnsAreAttributes = true;
+			CaptureData[i].bInputColumnsAreAttributes = false;
+			break;
+		case UISystemUpdate:
+			CaptureData[i].TargetUsage = ENiagaraScriptUsage::SystemUpdateScript;
+			CaptureData[i].ColumnName = LOCTEXT("SystemUpdate", "System Update");
+			CaptureData[i].bOutputColumnsAreAttributes = false;
+			CaptureData[i].bInputColumnsAreAttributes = false;
+			break;
+		default:
+			CaptureData[i].TargetUsage = ENiagaraScriptUsage::Function;
+			CaptureData[i].ColumnName = LOCTEXT("PerParticleUnknown", "Particle Unknown");
+			CaptureData[i].bOutputColumnsAreAttributes = true;
+			CaptureData[i].bInputColumnsAreAttributes = false;
+			break;
+		}
+
 		CaptureData[i].OutputHorizontalScrollBar = SNew(SScrollBar)
 			.Orientation(Orient_Horizontal)
 			.Thickness(FVector2D(8.0f, 8.0f));
@@ -238,6 +311,17 @@ void SNiagaraSpreadsheetView::Construct(const FArguments& InArgs, TSharedRef<FNi
 				SAssignNew(CaptureData[i].OutputHeaderRow, SHeaderRow)
 			);
 
+		SAssignNew(CaptureData[i].OutputFilterButton, SComboButton)
+			.HasDownArrow(false)
+			.OnGetMenuContent(this, &SNiagaraSpreadsheetView::GetOutputFilterMenu, (EUITab)i)
+			.ContentPadding(2)
+			.ButtonContent()
+			[
+				SNew(STextBlock)
+				.TextStyle(FNiagaraEditorStyle::Get(), "NiagaraEditor.ParameterText")
+				.Text(LOCTEXT("SpreadSheetOutputFilterBox", "Filter Attributes"))
+			];
+
 		SAssignNew(CaptureData[i].InputsListView, STreeView< TSharedPtr<int32> >)
 			.IsEnabled(this, &SNiagaraSpreadsheetView::IsPausedAtRightTimeOnRightHandle)
 			// List view items are this tall
@@ -264,6 +348,7 @@ void SNiagaraSpreadsheetView::Construct(const FArguments& InArgs, TSharedRef<FNi
 			//.Style(FEditorStyle::Get(), "PlacementBrowser.Tab")
 			.Style(FEditorStyle::Get(), i == 0 ? "Property.ToggleButton.Start" : (i < CaptureData.Num() - 1 ? "Property.ToggleButton.Middle" : "Property.ToggleButton.End"))
 			.OnCheckStateChanged(this, &SNiagaraSpreadsheetView::OnTabChanged, (EUITab)i)
+			.Visibility(this, &SNiagaraSpreadsheetView::GetTabVisibility, (EUITab)i)
 			.IsChecked(this, &SNiagaraSpreadsheetView::GetTabCheckedState, (EUITab)i)
 			[
 				SNew(SOverlay)
@@ -306,6 +391,16 @@ void SNiagaraSpreadsheetView::Construct(const FArguments& InArgs, TSharedRef<FNi
 				CaptureData[i].InputHorizontalScrollBar.ToSharedRef()
 			]
 			+ SVerticalBox::Slot()
+			.AutoHeight()
+			[
+				SNew(SHorizontalBox)
+				+ SHorizontalBox::Slot()
+				.AutoWidth()
+				[
+					CaptureData[i].OutputFilterButton.ToSharedRef()
+				]
+			]
+			+ SVerticalBox::Slot()
 			[
 				SNew(SHorizontalBox)
 				+ SHorizontalBox::Slot()
@@ -340,33 +435,71 @@ void SNiagaraSpreadsheetView::Construct(const FArguments& InArgs, TSharedRef<FNi
 				SNew(SBorder)
 				.BorderImage(FEditorStyle::GetBrush(TEXT("ToolPanel.GroupBorder")))
 				[
-					SNew(SHorizontalBox)
-
-					// Toolbar
-					+ SHorizontalBox::Slot()
-					.AutoWidth()
-					.VAlign(VAlign_Center)
-					.HAlign(HAlign_Left)
+					SNew(SVerticalBox)
+					+ SVerticalBox::Slot()
+					.AutoHeight()
 					[
-						SNew(SButton)
-						.HAlign(HAlign_Center)
-						.OnClicked(this, &SNiagaraSpreadsheetView::OnCaptureRequestPressed)
-						.Text(LOCTEXT("CaptureLabel", "Capture"))
-						.ToolTipText(LOCTEXT("CaptureToolitp", "Press this button to capture one frame's contents. Can only capture CPU systems."))
-						.IsEnabled(this, &SNiagaraSpreadsheetView::CanCapture)
+						SNew(SHorizontalBox)
+
+						+ SHorizontalBox::Slot()
+						.AutoWidth()
+						.VAlign(VAlign_Center)
+						.HAlign(HAlign_Left)
+						[
+							SNew(SButton)
+							.HAlign(HAlign_Center)
+							.OnClicked(this, &SNiagaraSpreadsheetView::OnCaptureRequestPressed)
+							.Text(LOCTEXT("CaptureLabel", "Capture"))
+							.ToolTipText(LOCTEXT("CaptureToolitp", "Press this button to capture one frame's contents. Can only capture CPU systems."))
+							.IsEnabled(this, &SNiagaraSpreadsheetView::CanCapture)
+						]
+
+						+ SHorizontalBox::Slot()
+						.AutoWidth()
+						.VAlign(VAlign_Center)
+						.HAlign(HAlign_Left)
+						[
+							SNew(SButton)
+							.HAlign(HAlign_Center)
+							.OnClicked(this, &SNiagaraSpreadsheetView::OnCSVOutputPressed)
+							.Text(LOCTEXT("CSVOutput", "Copy For Excel"))
+							.ToolTipText(LOCTEXT("CSVOutputToolitp", "Press this button to put the contents of this spreadsheet in the clipboard in an Excel-friendly format."))
+							.IsEnabled(this, &SNiagaraSpreadsheetView::IsPausedAtRightTimeOnRightHandle)
+						]
 					]
-
-					+ SHorizontalBox::Slot()
-					.AutoWidth()
-					.VAlign(VAlign_Center)
-					.HAlign(HAlign_Left)
+					+ SVerticalBox::Slot()
+					.AutoHeight()
 					[
-						SNew(SButton)
-						.HAlign(HAlign_Center)
-						.OnClicked(this, &SNiagaraSpreadsheetView::OnCSVOutputPressed)
-						.Text(LOCTEXT("CSVOutput", "Copy For Excel"))
-						.ToolTipText(LOCTEXT("CSVOutputToolitp", "Press this button to put the contents of this spreadsheet in the clipboard in an Excel-friendly format."))
-						.IsEnabled(this, &SNiagaraSpreadsheetView::IsPausedAtRightTimeOnRightHandle)
+						SNew(SHorizontalBox)
+
+						// Toolbar
+						+ SHorizontalBox::Slot()
+						.AutoWidth()
+						.VAlign(VAlign_Center)
+						.HAlign(HAlign_Left)
+						[
+							SNew(STextBlock)
+							.Text(LOCTEXT("CaptureTargetLabel", "Target: "))
+							.ToolTipText(LOCTEXT("TargetToolitp", "Select the actor that you wish to capture from."))
+							.IsEnabled(this, &SNiagaraSpreadsheetView::CanCapture)
+						]
+						+ SHorizontalBox::Slot()
+						.AutoWidth()
+						.VAlign(VAlign_Center)
+						.HAlign(HAlign_Left)
+						[
+							SNew(SComboButton)
+							.HAlign(HAlign_Center)
+							.OnGetMenuContent(this, &SNiagaraSpreadsheetView::OnGetTargetMenuContent)
+							.ContentPadding(1)
+							.ToolTipText(LOCTEXT("TargetToolitp", "Select the actor that you wish to capture from."))
+							.ButtonContent()
+							[
+								SNew(STextBlock)
+								.Text(this, &SNiagaraSpreadsheetView::OnGetTargetButtonText)
+							]
+							.IsEnabled(this, &SNiagaraSpreadsheetView::CanCapture)
+						]
 					]
 				]
 		]
@@ -378,6 +511,7 @@ void SNiagaraSpreadsheetView::Construct(const FArguments& InArgs, TSharedRef<FNi
 			[
 				SNew(SVerticalBox)
 				+ SVerticalBox::Slot()
+				.AutoHeight()
 				[
 					SNew(STextBlock)
 					.Text(LOCTEXT("NoDataText", "Please press capture to examine data from a particular frame."))
@@ -388,6 +522,7 @@ void SNiagaraSpreadsheetView::Construct(const FArguments& InArgs, TSharedRef<FNi
 					})
 				]
 				+ SVerticalBox::Slot()
+				.AutoHeight()
 				[
 					SNew(STextBlock)
 					.Text(this, &SNiagaraSpreadsheetView::LastCapturedInfoText)
@@ -407,6 +542,26 @@ void SNiagaraSpreadsheetView::Construct(const FArguments& InArgs, TSharedRef<FNi
 					]
 					+ SHorizontalBox::Slot()
 					.AutoWidth()
+					[
+						CaptureData[UIPerParticleSpawn].CheckBox.ToSharedRef()
+					]
+					+ SHorizontalBox::Slot()
+					.AutoWidth()
+					[
+						CaptureData[UIPerParticleEvent0].CheckBox.ToSharedRef()
+					]
+					+ SHorizontalBox::Slot()
+					.AutoWidth()
+					[
+						CaptureData[UIPerParticleEvent1].CheckBox.ToSharedRef()
+					]
+					+ SHorizontalBox::Slot()
+					.AutoWidth()
+					[
+						CaptureData[UIPerParticleEvent2].CheckBox.ToSharedRef()
+					]
+					+ SHorizontalBox::Slot()
+					.AutoWidth()
 					[
 						CaptureData[UISystemUpdate].CheckBox.ToSharedRef()
 					]
@@ -418,6 +573,22 @@ void SNiagaraSpreadsheetView::Construct(const FArguments& InArgs, TSharedRef<FNi
 			CaptureData[UIPerParticleUpdate].Container.ToSharedRef()
 		]
 		+ SVerticalBox::Slot()
+		[
+			CaptureData[UIPerParticleSpawn].Container.ToSharedRef()
+		]
+		+ SVerticalBox::Slot()
+		[
+			CaptureData[UIPerParticleEvent0].Container.ToSharedRef()
+		]
+		+ SVerticalBox::Slot()
+		[
+			CaptureData[UIPerParticleEvent1].Container.ToSharedRef()
+		]
+		+ SVerticalBox::Slot()
+		[
+			CaptureData[UIPerParticleEvent2].Container.ToSharedRef()
+		]
+		+ SVerticalBox::Slot()
 		[
 			CaptureData[UISystemUpdate].Container.ToSharedRef()
 		]
@@ -433,6 +604,104 @@ SNiagaraSpreadsheetView::~SNiagaraSpreadsheetView()
 	}
 }
 
+void SNiagaraSpreadsheetView::GetNameAndTooltip(const UNiagaraComponent* InComponent, FText& OutText, FText& OutTooltip) const
+{
+	const UNiagaraComponent* PreviewComponent = SystemViewModel->GetPreviewComponent();
+
+	if (InComponent == nullptr)
+	{
+		OutText = LOCTEXT("NullComponentLabel", "Unknown");
+		OutTooltip = LOCTEXT("NullComponentTooltip", "Unknown");
+	}
+	else if (PreviewComponent == InComponent)
+	{
+		OutText = LOCTEXT("PreviewComponentLabel", "Editor Viewport");
+		OutTooltip = LOCTEXT("PreviewComponentTooltip", "The instance of the Niagara Component in the Niagara editor viewport.");
+	}
+	else
+	{
+		const UWorld* World = InComponent->GetWorld();
+		const AActor* Actor = InComponent->GetOwner();
+		OutText = FText::Format(LOCTEXT("SourceComponentLabel","World: \"{0}\" Actor: \"{1}\""), World ? FText::FromString(World->GetName()) : FText::GetEmpty(), Actor ? FText::FromString(Actor->GetName()) : FText::GetEmpty());
+		OutTooltip = OutText;
+	}
+}
+
+TSharedRef<SWidget> SNiagaraSpreadsheetView::OnGetTargetMenuContent() const
+{
+	FMenuBuilder MenuBuilder(true, nullptr);
+	const UNiagaraComponent* PreviewComponent = SystemViewModel->GetPreviewComponent();
+
+	for (TObjectIterator<UNiagaraComponent> It; It; ++It)
+	{
+		// Ignore dying or CDO versions of data..
+		if (It->IsPendingKillOrUnreachable() || It->HasAnyFlags(EObjectFlags::RF_ClassDefaultObject))
+		{
+			continue;
+		}
+
+		// Ignore any component not referencing this system.
+		if (It->GetAsset() != &SystemViewModel->GetSystem())
+		{
+			continue;
+		}
+
+		// Ignore non-Niagara editor systems or non-PiE components
+		const UWorld* World = It->GetWorld();
+		bool bAdd = false;
+
+		if (World && World->IsPlayInEditor())
+		{
+			if (It->GetForceSolo() == true || (It->GetSystemInstance() && It->GetSystemInstance()->IsSolo()))
+			{
+				bAdd = true;
+			}
+		}
+		else if (*It == PreviewComponent)
+		{
+			bAdd = true;
+		}
+
+		if (bAdd)
+		{
+			FText ComponentName;
+			FText ComponentTooltip;
+			GetNameAndTooltip(*It, ComponentName, ComponentTooltip);
+
+			MenuBuilder.AddMenuEntry(
+				ComponentName,
+				ComponentTooltip,
+				FSlateIcon(),
+				FUIAction(FExecuteAction::CreateRaw(this, &SNiagaraSpreadsheetView::SetTarget, *It)));
+		}
+	}
+
+	return MenuBuilder.MakeWidget();
+}
+
+
+void SNiagaraSpreadsheetView::SetTarget(UNiagaraComponent* InComponent)
+{
+	TargetComponent = InComponent;
+	TargetRequestId.Invalidate();
+	TargetCaptureData.Empty();
+
+	for (int32 i = 0; i < (int32)UIMax; i++)
+	{
+		CaptureData[i].DataSet.Init(FNiagaraDataSetID(), ENiagaraSimTarget::CPUSim);
+		CaptureData[i].InputParams.Reset();
+		CaptureData[i].CaptureData.Reset();
+	}
+}
+
+FText SNiagaraSpreadsheetView::OnGetTargetButtonText() const
+{
+	FText Text;
+	FText Tooltip;
+	GetNameAndTooltip(TargetComponent.Get(), Text, Tooltip);
+	return Text;
+}
+
 void SNiagaraSpreadsheetView::OnTabChanged(ECheckBoxState State, EUITab Tab)
 {
 	if (State == ECheckBoxState::Checked)
@@ -451,9 +720,27 @@ EVisibility SNiagaraSpreadsheetView::GetViewVisibility(EUITab Tab) const
 	return TabState == Tab ? EVisibility::Visible : EVisibility::Collapsed;
 }
 
+EVisibility SNiagaraSpreadsheetView::GetTabVisibility(EUITab Tab) const
+{
+	return (CaptureData[(int32)Tab].CaptureData.IsValid()) ? EVisibility::Visible : EVisibility::Collapsed;
+}
+
 
 TSharedRef< ITableRow > SNiagaraSpreadsheetView::OnGenerateWidgetForList(TSharedPtr<int32> InItem, const TSharedRef<STableViewBase>& OwnerTable, EUITab Tab, bool bInputList)
 {
+	if (CaptureData[(int32)Tab].CaptureData.IsValid() == false)
+	{
+		return SNew(SNiagaraSpreadsheetRow, OwnerTable)
+			.RowIndex(*InItem)
+			.ColumnsAreAttributes(CaptureData[(int32)Tab].bOutputColumnsAreAttributes)
+			.DataSet(nullptr)
+			.SupportedFields(CaptureData[(int32)Tab].SupportedOutputFields)
+			.FieldInfoMap(CaptureData[(int32)Tab].OutputFieldInfoMap)
+			.UseGlobalOffsets(false)
+			.ParameterStore(nullptr);
+
+	}
+
 	if (bInputList)
 	{
 		return SNew(SNiagaraSpreadsheetRow, OwnerTable)
@@ -470,7 +757,7 @@ TSharedRef< ITableRow > SNiagaraSpreadsheetView::OnGenerateWidgetForList(TShared
 		return SNew(SNiagaraSpreadsheetRow, OwnerTable)
 			.RowIndex(*InItem)
 			.ColumnsAreAttributes(CaptureData[(int32)Tab].bOutputColumnsAreAttributes)
-			.DataSet(CaptureData[(int32)Tab].DataSet)
+			.DataSet(&CaptureData[(int32)Tab].DataSet)
 			.SupportedFields(CaptureData[(int32)Tab].SupportedOutputFields)
 			.FieldInfoMap(CaptureData[(int32)Tab].OutputFieldInfoMap)
 			.UseGlobalOffsets(false)
@@ -482,12 +769,11 @@ FText SNiagaraSpreadsheetView::LastCapturedInfoText() const
 {
 	TArray<TSharedRef<FNiagaraEmitterHandleViewModel>> SelectedEmitterHandles;
 	SystemViewModel->GetSelectedEmitterHandles(SelectedEmitterHandles);
-	if (SelectedEmitterHandles.Num() == 1 && SNiagaraSpreadsheetView::IsPausedAtRightTimeOnRightHandle() && CaptureData[(int32)TabState].DataSet != nullptr)
+	if (SelectedEmitterHandles.Num() == 1 && SNiagaraSpreadsheetView::IsPausedAtRightTimeOnRightHandle())
 	{
-		return FText::Format(LOCTEXT("LastCapturedInfoName", "Captured Emitter: \"{0}\"     # Particles: {1}     Local Time: {2}     Script Type: {3}"),
+		return FText::Format(LOCTEXT("LastCapturedInfoName", "Captured Emitter: \"{0}\"     # Particles: {1}    Script Type: {2}"),
 			SelectedEmitterHandles[0]->GetNameText(),
-			FText::AsNumber(CaptureData[(int32)TabState].DataSet->PrevData().GetNumInstances()),
-			FText::AsNumber(CaptureData[(int32)TabState].LastCaptureTime),
+			FText::AsNumber(CaptureData[(int32)TabState].DataSet.PrevData().GetNumInstances()),
 			ScriptEnum->GetDisplayNameTextByValue((int64)CaptureData[(int32)TabState].TargetUsage));
 	}
 
@@ -499,13 +785,102 @@ void SNiagaraSpreadsheetView::OnGetChildrenForList(TSharedPtr<int32> InItem, TAr
 	OutChildren.Empty();
 }
 
+
+TSharedRef<SWidget> SNiagaraSpreadsheetView::GetOutputFilterMenu(EUITab Tab)
+{
+	FMenuBuilder MenuBuilder(false, nullptr);
+	MenuBuilder.BeginSection("OutputAttributeActions", LOCTEXT("OutputAttributes", "Output Attributes"));
+	MenuBuilder.AddMenuEntry(
+		LOCTEXT("AllOutputAttribute", "Toggle All"),
+		LOCTEXT("AllOutputAttributeTooltip", "Toggle visibility for all attributes"),
+		FSlateIcon(),
+		FUIAction(FExecuteAction::CreateSP(this, &SNiagaraSpreadsheetView::ToggleAllOutputAttributes, Tab),
+			FCanExecuteAction::CreateLambda([=] { return true; }),
+			FIsActionChecked::CreateSP(this, &SNiagaraSpreadsheetView::AnyOutputAttributeEnabled, Tab)),
+		NAME_None,
+		EUserInterfaceActionType::ToggleButton);
+
+	TArray<FNiagaraVariable> Variables = CaptureData[(int32)Tab].DataSet.GetVariables();
+	for (const FNiagaraVariable& Var : Variables)
+	{
+		const FNiagaraTypeDefinition& TypeDef = Var.GetType();
+		const UScriptStruct* Struct = TypeDef.GetScriptStruct();
+		const UEnum* Enum = TypeDef.GetEnum();
+
+		FNiagaraTypeLayoutInfo Layout;
+		TArray<FName> PropertyNames;
+		TArray<SNiagaraSpreadsheetView::FieldInfo> FieldInfos;
+
+		GenerateLayoutInfo(Layout, Struct, Enum, Var.GetName(), PropertyNames, FieldInfos);
+
+		for (int32 VarIdx = 0; VarIdx < PropertyNames.Num(); VarIdx++)
+		{
+			const FText PropertyText = FText::FromName(PropertyNames[VarIdx]);
+			MenuBuilder.AddMenuEntry(
+				PropertyText,
+				FText::Format(LOCTEXT("OutputAttributeTooltip", "Toggle {0}"), PropertyText),
+				FSlateIcon(),
+				FUIAction(FExecuteAction::CreateSP(this, &SNiagaraSpreadsheetView::ToggleFilterOutputAttribute, Tab, PropertyNames[VarIdx]),
+					FCanExecuteAction::CreateLambda([=] { return true; }),
+					FIsActionChecked::CreateSP(this, &SNiagaraSpreadsheetView::IsOutputAttributeEnabled, Tab, PropertyNames[VarIdx])),
+				NAME_None,
+				EUserInterfaceActionType::ToggleButton);
+		}
+	}
+
+	MenuBuilder.EndSection();
+
+	return MenuBuilder.MakeWidget();
+}
+
+
+void SNiagaraSpreadsheetView::ToggleAllOutputAttributes(EUITab Tab)
+{
+	TArray<FName>& OutputFields = CaptureData[(int32)Tab].FilteredOutputFields;
+	if (OutputFields.Num() > 0)
+	{
+		OutputFields.Empty();
+	}
+	else if (CaptureData[(int32)Tab].SupportedOutputFields.IsValid())
+	{
+		OutputFields = *CaptureData[(int32)Tab].SupportedOutputFields.Get();
+	}
+
+	ResetColumns(Tab);
+}
+
+
+void SNiagaraSpreadsheetView::ToggleFilterOutputAttribute(EUITab Tab, FName Item)
+{
+	const int32 Index = CaptureData[(int32)Tab].FilteredOutputFields.Find(Item);
+	if (Index == INDEX_NONE)
+	{
+		CaptureData[(int32)Tab].FilteredOutputFields.AddUnique(Item);
+	}
+	else
+	{
+		CaptureData[(int32)Tab].FilteredOutputFields.RemoveAt(Index);
+	}
+
+	ResetColumns(Tab);
+}
+
+bool SNiagaraSpreadsheetView::AnyOutputAttributeEnabled(EUITab Tab)
+{
+	return CaptureData[(int32)Tab].FilteredOutputFields.Num() > 0;
+}
+
+bool SNiagaraSpreadsheetView::IsOutputAttributeEnabled(EUITab Tab, FName Item)
+{
+	return CaptureData[(int32)Tab].FilteredOutputFields.Find(Item) != INDEX_NONE;
+}
+
 void SNiagaraSpreadsheetView::SelectedEmitterHandlesChanged()
 {
 	// Need to reset the attributes list...
 	for (int32 i = 0; i < (int32)UIMax; i++)
 	{
-		CaptureData[i].LastReadWriteId = -1;
-		CaptureData[i].DataSet = nullptr;
+		CaptureData[i].DataSet.Init(FNiagaraDataSetID(), ENiagaraSimTarget::CPUSim);
 		CaptureData[i].SupportedInputIndices.SetNum(0);
 		CaptureData[i].SupportedOutputIndices.SetNum(0);
 		CaptureData[i].OutputsListView->RequestTreeRefresh();
@@ -546,7 +921,7 @@ FReply SNiagaraSpreadsheetView::OnCSVOutputPressed()
 
 		CSVOutput += "\r\n";
 
-		for (uint32 RowIndex = 0; RowIndex < CaptureData[(int32)TabState].DataSet->PrevData().GetNumInstances(); RowIndex++)
+		for (uint32 RowIndex = 0; RowIndex < CaptureData[(int32)TabState].DataSet.PrevData().GetNumInstances(); RowIndex++)
 		{
 			NumWritten = 0;
 			for (int32 i = 0; i < CaptureData[(int32)TabState].SupportedOutputFields->Num(); i++)
@@ -562,18 +937,18 @@ FReply SNiagaraSpreadsheetView::OnCSVOutputPressed()
 				}
 				const SNiagaraSpreadsheetView::FieldInfo* FieldInfo = FieldInfos[i];
 				
-				if (FieldInfo != nullptr && CaptureData[(int32)TabState].DataSet != nullptr)
+				if (FieldInfo != nullptr && CaptureData[(int32)TabState].DataSet.GetNumInstances() != 0)
 				{
 					if (FieldInfo->bFloat)
 					{
 						uint32 CompBufferOffset = FieldInfo->FloatStartOffset;
-						float* Src = CaptureData[(int32)TabState].DataSet->PrevData().GetInstancePtrFloat(CompBufferOffset, RowIndex);
+						float* Src = CaptureData[(int32)TabState].DataSet.PrevData().GetInstancePtrFloat(CompBufferOffset, RowIndex);
 						CSVOutput += FString::Printf(TEXT("%3.3f"), Src[0]);
 					}
 					else
 					{
 						uint32 CompBufferOffset = FieldInfo->IntStartOffset;
-						int32* Src = CaptureData[(int32)TabState].DataSet->PrevData().GetInstancePtrInt32(CompBufferOffset, RowIndex);
+						int32* Src = CaptureData[(int32)TabState].DataSet.PrevData().GetInstancePtrInt32(CompBufferOffset, RowIndex);
 						CSVOutput += FString::Printf(TEXT("%d"), Src[0]);
 					}
 				}
@@ -600,61 +975,60 @@ void SNiagaraSpreadsheetView::Tick(float DeltaTime)
 
 void SNiagaraSpreadsheetView::HandleTimeChange()
 {
-	for (int32 i = 0; i < (int32)UIMax; i++)
+	if (TargetComponent.IsValid() && TargetRequestId.IsValid() && TargetComponent->GetSystemInstance() != nullptr)
 	{
-		if (!CaptureData[i].DataSource.IsValid())
-		{
-			CaptureData[i].bAwaitingFrame = false;
-		}
-
-		if (CaptureData[i].bAwaitingFrame)
+		if (TargetComponent->GetSystemInstance()->QueryCaptureResults(TargetRequestId, TargetCaptureData))
 		{
-
-			TArray<TSharedRef<FNiagaraEmitterHandleViewModel>> SelectedEmitterHandles;
-			SystemViewModel->GetSelectedEmitterHandles(SelectedEmitterHandles);
-			if (SelectedEmitterHandles.Num() == 1)
+			for (int32 i = 0; i < (int32)UIMax; i++)
 			{
-				TArray<UNiagaraScript*> Scripts;
-				CaptureData[i].DataSource->GetScripts(Scripts);
-				Scripts.Add(SystemViewModel->GetSystem().GetSystemSpawnScript(true));
-				Scripts.Add(SystemViewModel->GetSystem().GetSystemUpdateScript(true));
-
-				UNiagaraScript* FoundScript = nullptr;
-				for (UNiagaraScript* Script : Scripts)
+				TArray<TSharedRef<FNiagaraEmitterHandleViewModel>> SelectedEmitterHandles;
+				SystemViewModel->GetSelectedEmitterHandles(SelectedEmitterHandles);
+				if (SelectedEmitterHandles.Num() == 1)
 				{
-					if (Script && Script->IsEquivalentUsage(CaptureData[i].TargetUsage))
+					FName EntryName = NAME_None;
+					if (i != UISystemUpdate)
 					{
-						FoundScript = Script;
-						break;
+						EntryName = SelectedEmitterHandles[0]->GetEmitterHandle()->GetIdName();
+					}
+					else //if (i == UISystemUpdate)
+					{
+						EntryName = NAME_None;
 					}
-				}
-
-				FQualifiedFrameTime LocalCaptureTime = SystemViewModel->GetSequencer()->GetLocalTime();
-
-				if (FoundScript && FoundScript->GetDebuggerInfo().bRequestDebugFrame == false && CaptureData[i].LastReadWriteId != FoundScript->GetDebuggerInfo().DebugFrameLastWriteId)
-				{
-					CaptureData[i].LastReadWriteId = FoundScript->GetDebuggerInfo().DebugFrameLastWriteId;
-					CaptureData[i].DataSet = &FoundScript->GetDebuggerInfo().DebugFrame;
-					CaptureData[i].InputParams = FoundScript->GetDebuggerInfo().DebugParameters;
-					CaptureData[i].DataSet->Tick(ENiagaraSimTarget::CPUSim); // Force a buffer swap, from here out we read from PrevData.
-
-					CaptureData[i].LastCaptureTime = LocalCaptureTime.AsSeconds();
-					ensure(CaptureData[i].LastCaptureTime == CaptureData[i].TargetCaptureTime);
-					CaptureData[i].LastCaptureHandleId = SelectedEmitterHandles[0]->GetId();
 
-					ResetColumns((EUITab)i);
-					ResetEntries((EUITab)i);
+					TSharedPtr<struct FNiagaraScriptDebuggerInfo, ESPMode::ThreadSafe>* FoundEntry = TargetCaptureData.FindByPredicate([&](const TSharedPtr<struct FNiagaraScriptDebuggerInfo, ESPMode::ThreadSafe>& Entry)
+					{
+						return Entry->HandleName == EntryName && UNiagaraScript::IsEquivalentUsage(Entry->Usage, CaptureData[i].TargetUsage) && Entry->UsageId == CaptureData[i].TargetUsageId;
+					});
+						
+					if (FoundEntry != nullptr)
+					{
+						CaptureData[i].CaptureData = *FoundEntry;
+						CaptureData[i].CaptureData->Frame.CopyCurToPrev();
+						CaptureData[i].DataSet = CaptureData[i].CaptureData->Frame;
+						CaptureData[i].InputParams = CaptureData[i].CaptureData->Parameters;
+						CaptureData[i].LastCaptureHandleId = SelectedEmitterHandles[0]->GetId();
+
+						ResetColumns((EUITab)i);
+						ResetEntries((EUITab)i);
+					}
+					else
+					{
+						CaptureData[i].CaptureData.Reset();
+						CaptureData[i].DataSet.Init(FNiagaraDataSetID(), ENiagaraSimTarget::CPUSim);
 
-					CaptureData[i].bAwaitingFrame = false;
+						ResetColumns((EUITab)i);
+						ResetEntries((EUITab)i);
+					}
 				}
 			}
+			TargetRequestId.Invalidate();
 		}
 	}
 }
 
 bool SNiagaraSpreadsheetView::IsTickable() const
 {
-	return CaptureData[(int32)TabState].bAwaitingFrame;
+	return TargetRequestId.IsValid();
 }
 
 TStatId SNiagaraSpreadsheetView::GetStatId() const
@@ -685,7 +1059,7 @@ bool SNiagaraSpreadsheetView::IsPausedAtRightTimeOnRightHandle() const
 	if (SelectedEmitterHandles.Num() == 1)
 	{
 		return SystemViewModel->GetSequencer()->GetPlaybackStatus() == EMovieScenePlayerStatus::Stopped &&
-			CaptureData[(int32)TabState].LastCaptureTime == SystemViewModel->GetSequencer()->GetLocalTime().AsSeconds() &&
+			CaptureData[(int32)TabState].CaptureData.IsValid() &&
 			CaptureData[(int32)TabState].LastCaptureHandleId == SelectedEmitterHandles[0]->GetId();
 	}
 	return false;
@@ -693,11 +1067,10 @@ bool SNiagaraSpreadsheetView::IsPausedAtRightTimeOnRightHandle() const
 
 void SNiagaraSpreadsheetView::ResetEntries(EUITab Tab)
 {
-	if (CaptureData[(int32)Tab].DataSet)
 	{
 		{
-			int32 NumInstances = CaptureData[(int32)Tab].DataSet->GetPrevNumInstances();
-			if (!CaptureData[(int32)Tab].bOutputColumnsAreAttributes)
+			int32 NumInstances = CaptureData[(int32)Tab].DataSet.GetPrevNumInstances();
+			if (!CaptureData[(int32)Tab].bOutputColumnsAreAttributes && CaptureData[(int32)Tab].SupportedOutputFields.IsValid())
 			{
 				NumInstances = CaptureData[(int32)Tab].SupportedOutputFields->Num();
 			}
@@ -714,7 +1087,7 @@ void SNiagaraSpreadsheetView::ResetEntries(EUITab Tab)
 	
 		{
 			int32 NumInstances = CaptureData[(int32)Tab].InputParams.GetNumParameters();
-			if (!CaptureData[(int32)Tab].bInputColumnsAreAttributes)
+			if (!CaptureData[(int32)Tab].bInputColumnsAreAttributes && CaptureData[(int32)Tab].SupportedInputFields.IsValid())
 			{
 				NumInstances = CaptureData[(int32)Tab].SupportedInputFields->Num();
 			}
@@ -794,9 +1167,9 @@ void SNiagaraSpreadsheetView::ResetColumns(EUITab Tab)
 {
 	int32 i = (int32)Tab;
 
-	if (CaptureData[(int32)i].DataSet)
+	if (CaptureData[(int32)i].DataSet.GetNumInstances() != 0)
 	{
-		float ManualWidth = 75.0f;
+		float ManualWidth = 125.0f;
 
 		// Handle output columns
 		{
@@ -807,7 +1180,7 @@ void SNiagaraSpreadsheetView::ResetColumns(EUITab Tab)
 			uint32 TotalFloatComponents = 0;
 			uint32 TotalInt32Components = 0;
 
-			TArray<FNiagaraVariable> Variables = CaptureData[(int32)i].DataSet->GetVariables();
+			TArray<FNiagaraVariable> Variables = CaptureData[(int32)i].DataSet.GetVariables();
 
 			TArray<FName> ColumnNames;
 
@@ -817,7 +1190,7 @@ void SNiagaraSpreadsheetView::ResetColumns(EUITab Tab)
 			}
 			else
 			{
-				ManualWidth = 300.0f;
+				ManualWidth = 125.0f;
 				ColumnNames.Add(OUTPUT_KEY_COLUMN_NAME);
 				ColumnNames.Add(VALUE_COLUMN_NAME);
 				ColumnNames.Add(FILLER_COLUMN_NAME);
@@ -840,30 +1213,33 @@ void SNiagaraSpreadsheetView::ResetColumns(EUITab Tab)
 
 				for (int32 VarIdx = 0; VarIdx < PropertyNames.Num(); VarIdx++)
 				{
-
+					if (FieldInfos[VarIdx].bFloat)
+					{
+						FieldInfos[VarIdx].FloatStartOffset += TotalFloatComponentsBeforeStruct;
+						TotalFloatComponents++;
+					}
+					else
 					{
-						if (FieldInfos[VarIdx].bFloat)
-						{
-							FieldInfos[VarIdx].FloatStartOffset += TotalFloatComponentsBeforeStruct;
-							TotalFloatComponents++;
-						}
-						else
-						{
-							FieldInfos[VarIdx].IntStartOffset += TotalInt32ComponentsBeforeStruct;
-							TotalInt32Components++;
-						}
-
-						CaptureData[(int32)i].SupportedOutputFields->Add(PropertyNames[VarIdx]);
-						CaptureData[(int32)i].OutputFieldInfoMap->Add(PropertyNames[VarIdx], FieldInfos[VarIdx]);
+						FieldInfos[VarIdx].IntStartOffset += TotalInt32ComponentsBeforeStruct;
+						TotalInt32Components++;
 					}
 
-					if (CaptureData[(int32)i].bOutputColumnsAreAttributes)
+					CaptureData[i].SupportedOutputFields->Add(PropertyNames[VarIdx]);
+					CaptureData[i].OutputFieldInfoMap->Add(PropertyNames[VarIdx], FieldInfos[VarIdx]);
+					
+					if (CaptureData[i].bOutputColumnsAreAttributes 
+						&& (bInitialColumns || IsOutputAttributeEnabled(Tab, PropertyNames[VarIdx])))
 					{
 						ColumnNames.Add(PropertyNames[VarIdx]);
 					}
 				}
 			}
 
+			if (bInitialColumns)
+			{
+				CaptureData[(int32)Tab].FilteredOutputFields = *CaptureData[(int32)Tab].SupportedOutputFields.Get();
+				bInitialColumns = false;
+			}
 
 			for (int32 ColIdx = 0; ColIdx < ColumnNames.Num(); ColIdx++)
 			{
@@ -915,7 +1291,7 @@ void SNiagaraSpreadsheetView::ResetColumns(EUITab Tab)
 			}
 			else
 			{
-				ManualWidth = 300.0f;
+				ManualWidth = 125.0f;
 				ColumnNames.Add(INPUT_KEY_COLUMN_NAME);
 				ColumnNames.Add(VALUE_COLUMN_NAME);
 				ColumnNames.Add(FILLER_COLUMN_NAME);
@@ -998,31 +1374,70 @@ FReply SNiagaraSpreadsheetView::OnCaptureRequestPressed()
 	SystemViewModel->GetSelectedEmitterHandles(SelectedEmitterHandles);
 	ensure(SelectedEmitterHandles.Num() == 1);
 
-	for (int32 i = 0; i < CaptureData.Num(); i++)
+	if (TargetComponent.IsValid() && TargetComponent->GetSystemInstance())
 	{
-		CaptureData[(int32)i].DataSource = SelectedEmitterHandles[0]->GetEmitterHandle()->GetInstance();
-		
-		TArray<UNiagaraScript*> Scripts;
-		CaptureData[(int32)i].DataSource->GetScripts(Scripts);
-		Scripts.Add(SystemViewModel->GetSystem().GetSystemSpawnScript(true));
-		Scripts.Add(SystemViewModel->GetSystem().GetSystemUpdateScript(true));
+		TargetRequestId = FGuid::NewGuid();
+		TargetComponent->GetSystemInstance()->RequestCapture(TargetRequestId);
 
-		UNiagaraScript* FoundScript = nullptr;
-		for (UNiagaraScript* Script : Scripts)
+		UNiagaraEmitter* Emitter = SelectedEmitterHandles[0]->GetEmitterHandle()->GetInstance();
+		
+		for (int32 i = 0; i < CaptureData.Num(); i++)
 		{
-			if (Script && Script->IsEquivalentUsage(CaptureData[(int32)i].TargetUsage))
+			CaptureData[(int32)i].DataSource = Emitter;
+			switch (i)
 			{
-				FoundScript = Script;
-				break;
+				case UIPerParticleUpdate:
+					CaptureData[i].TargetUsage = ENiagaraScriptUsage::ParticleUpdateScript;
+					CaptureData[i].TargetUsageId = FGuid();
+					break;
+				case UIPerParticleSpawn:
+					CaptureData[i].TargetUsage = ENiagaraScriptUsage::ParticleSpawnScript;
+					CaptureData[i].TargetUsageId = FGuid();
+					break;
+				case UIPerParticleEvent0:
+					CaptureData[i].TargetUsage = ENiagaraScriptUsage::ParticleEventScript;		
+					CaptureData[i].TargetUsageId = Emitter->GetEventHandlers().Num() >= 1 ? Emitter->GetEventHandlers() [0].Script->GetUsageId() : FGuid();
+					break;
+				case UIPerParticleEvent1:
+					CaptureData[i].TargetUsage = ENiagaraScriptUsage::ParticleEventScript;
+					CaptureData[i].TargetUsageId = Emitter->GetEventHandlers().Num() >= 2 ? Emitter->GetEventHandlers()[1].Script->GetUsageId() : FGuid();
+					break;
+				case UIPerParticleEvent2:
+					CaptureData[i].TargetUsage = ENiagaraScriptUsage::ParticleEventScript;					
+					CaptureData[i].TargetUsageId = Emitter->GetEventHandlers().Num() >= 3 ? Emitter->GetEventHandlers()[2].Script->GetUsageId() : FGuid();
+					break;
+				case UISystemUpdate:
+					CaptureData[i].TargetUsage = ENiagaraScriptUsage::SystemUpdateScript;
+					CaptureData[i].TargetUsageId = FGuid();
+					break;
+				default:
+					CaptureData[i].TargetUsage = ENiagaraScriptUsage::Function;
+					CaptureData[i].TargetUsageId = FGuid();
+					break;
 			}
 		}
 
-		if (FoundScript)
+		UWorld* World = TargetComponent->GetWorld();
+		if (World && World->IsPlayInEditor())
 		{
-			FoundScript->GetDebuggerInfo().bRequestDebugFrame = true;
-			CaptureData[(int32)i].bAwaitingFrame = true;
+			if (FPlayWorldCommandCallbacks::IsInPIE())
+			{
+				if (FPlayWorldCommandCallbacks::IsInPIE_AndRunning())
+				{
+					// Need to Pause
+					if (FPlayWorldCommandCallbacks::HasPlayWorld())
+					{
+						FPlayWorldCommandCallbacks::PausePlaySession_Clicked();
+					}
+				}
+
+				// Need to single-step once
+				if (FPlayWorldCommandCallbacks::HasPlayWorldAndPaused())
+				{
+					FPlayWorldCommandCallbacks::SingleFrameAdvance_Clicked();
+				}
 
-			CaptureData[(int32)i].TargetCaptureTime = TargetCaptureTime;
+			}
 		}
 	}
 
diff --git a/Engine/Plugins/FX/Niagara/Source/NiagaraEditor/Private/Widgets/SNiagaraSpreadsheetView.h b/Engine/Plugins/FX/Niagara/Source/NiagaraEditor/Private/Widgets/SNiagaraSpreadsheetView.h
--- a/Engine/Plugins/FX/Niagara/Source/NiagaraEditor/Private/Widgets/SNiagaraSpreadsheetView.h
+++ b/Engine/Plugins/FX/Niagara/Source/NiagaraEditor/Private/Widgets/SNiagaraSpreadsheetView.h
@@ -8,7 +8,8 @@
 #include "Widgets/Views/STableRow.h"
 #include "Widgets/Views/STableViewBase.h"
 #include "Widgets/Views/STreeView.h"
-#include "NiagaraSystemViewModel.h"
+#include "Widgets/Input/SComboBox.h"
+#include "ViewModels/NiagaraSystemViewModel.h"
 #include "TickableEditorObject.h"
 #include "UObject/WeakObjectPtrTemplates.h"
 #include "NiagaraDataSet.h"
@@ -48,6 +49,11 @@ protected:
 	enum EUITab : uint8
 	{
 		UIPerParticleUpdate = 0,
+		UIPerParticleSpawn,
+		UIPerParticleEvent0,
+		UIPerParticleEvent1,
+		UIPerParticleEvent2,
+		UIPerParticleEvent3,
 		UISystemUpdate,
 		UIMax
 	};
@@ -55,6 +61,12 @@ protected:
 	void OnTabChanged(ECheckBoxState State,	EUITab Tab);
 	ECheckBoxState GetTabCheckedState(EUITab Tab) const;
 	EVisibility GetViewVisibility(EUITab Tab) const;
+	EVisibility GetTabVisibility(EUITab Tab) const;
+
+	TSharedRef<SWidget> OnGetTargetMenuContent() const;
+	FText OnGetTargetButtonText() const;
+	void SetTarget(UNiagaraComponent* InComponent);
+	void GetNameAndTooltip(const UNiagaraComponent*, FText& OutText, FText& OutTooltip) const;
 
 	void GenerateLayoutInfo(FNiagaraTypeLayoutInfo& Layout, const UScriptStruct* Struct, const UEnum* Enum, FName BaseName, TArray<FName>& PropertyNames, TArray<FieldInfo>& FieldInfo);
 
@@ -62,7 +74,7 @@ protected:
 
 	struct CapturedUIData
 	{
-		CapturedUIData() : LastReadWriteId(-1), DataSet(nullptr), TargetUsage(ENiagaraScriptUsage::ParticleUpdateScript), bAwaitingFrame(false), LastCaptureTime(-FLT_MAX), bInputColumnsAreAttributes(true), bOutputColumnsAreAttributes(true)
+		CapturedUIData() : TargetUsage(ENiagaraScriptUsage::ParticleUpdateScript), bInputColumnsAreAttributes(true), bOutputColumnsAreAttributes(true)
 		{ }
 
 		TSharedPtr<SHeaderRow> OutputHeaderRow;
@@ -70,19 +82,21 @@ protected:
 		TSharedPtr < STreeView<TSharedPtr<int32> > > OutputsListView;
 		TSharedPtr < STreeView<TSharedPtr<int32> > > InputsListView;
 		TSharedPtr < SCheckBox > CheckBox;
+
+		TSharedPtr<SComboButton> OutputFilterButton;
+		TArray<FName> FilteredOutputFields;
+
 		TArray< TSharedPtr<int32> > SupportedInputIndices;
 		TArray< TSharedPtr<int32> > SupportedOutputIndices;
-		int32 LastReadWriteId;
-		FNiagaraDataSet* DataSet;
 		FNiagaraParameterStore InputParams;
+		FNiagaraDataSet DataSet;
+		TSharedPtr<struct FNiagaraScriptDebuggerInfo, ESPMode::ThreadSafe> CaptureData;
 		TSharedPtr<TArray<FName> > SupportedInputFields;
 		TSharedPtr<TArray<FName> > SupportedOutputFields;
 		TSharedPtr<TMap<FName, FieldInfo> > InputFieldInfoMap;
 		TSharedPtr<TMap<FName, FieldInfo> > OutputFieldInfoMap;
 		ENiagaraScriptUsage TargetUsage;
-		bool bAwaitingFrame;
-		float LastCaptureTime;
-		float TargetCaptureTime;
+		FGuid TargetUsageId;
 		FGuid LastCaptureHandleId;
 		TWeakObjectPtr<UNiagaraEmitter> DataSource;
 		TSharedPtr<SScrollBar> OutputHorizontalScrollBar;
@@ -97,7 +111,11 @@ protected:
 
 	TArray<CapturedUIData> CaptureData;
 
+	FGuid TargetRequestId;
+	TArray<TSharedPtr<struct FNiagaraScriptDebuggerInfo, ESPMode::ThreadSafe>> TargetCaptureData;
+
 	TSharedPtr<FNiagaraSystemViewModel> SystemViewModel;
+	TWeakObjectPtr<UNiagaraComponent> TargetComponent;
 
 	UEnum* ScriptEnum;
 
@@ -141,6 +159,25 @@ protected:
 	 * @param OutChildren Generated children
 	 */
 	void OnGetChildrenForList(TSharedPtr<int32> InItem, TArray<TSharedPtr<int32>>& OutChildren, EUITab Tab, bool bInputList);
+
+	/** Creates the attribute filter dropdown menu. */
+	TSharedRef<SWidget> GetOutputFilterMenu(EUITab Tab);
+
+	/** Toggles all attributes. */
+	void ToggleAllOutputAttributes(EUITab Tab);
+
+	/** Toggles the column visibility of the given attribute. */
+	void ToggleFilterOutputAttribute(EUITab Tab, FName Item);
+
+	/** True if any attribute is in the filtered list. */
+	bool AnyOutputAttributeEnabled(EUITab Tab);
+
+	/** True if the given attribute is in the filtered list. */
+	bool IsOutputAttributeEnabled(EUITab Tab, FName Item);
+
+private:
+	/** If this is the first time setting the columns, so we can enable them all in the filter by default. */
+	bool bInitialColumns;
 };
 
 
diff --git a/Engine/Plugins/FX/Niagara/Source/NiagaraEditor/Private/Widgets/SNiagaraSystemScript.cpp b/Engine/Plugins/FX/Niagara/Source/NiagaraEditor/Private/Widgets/SNiagaraSystemScript.cpp
--- a/Engine/Plugins/FX/Niagara/Source/NiagaraEditor/Private/Widgets/SNiagaraSystemScript.cpp
+++ b/Engine/Plugins/FX/Niagara/Source/NiagaraEditor/Private/Widgets/SNiagaraSystemScript.cpp
@@ -1,7 +1,7 @@
 // Copyright 1998-2018 Epic Games, Inc. All Rights Reserved.
 
 #include "SNiagaraSystemScript.h"
-#include "NiagaraSystemViewModel.h"
+#include "ViewModels/NiagaraSystemViewModel.h"
 #include "NiagaraSystemScriptViewModel.h"
 #include "NiagaraScriptViewModel.h"
 #include "NiagaraScriptInputCollectionViewModel.h"
diff --git a/Engine/Plugins/FX/Niagara/Source/NiagaraEditor/Private/Widgets/SNiagaraSystemViewport.cpp b/Engine/Plugins/FX/Niagara/Source/NiagaraEditor/Private/Widgets/SNiagaraSystemViewport.cpp
--- a/Engine/Plugins/FX/Niagara/Source/NiagaraEditor/Private/Widgets/SNiagaraSystemViewport.cpp
+++ b/Engine/Plugins/FX/Niagara/Source/NiagaraEditor/Private/Widgets/SNiagaraSystemViewport.cpp
@@ -18,6 +18,8 @@
 #include "AdvancedPreviewScene.h"
 #include "ImageUtils.h"
 
+#define LOCTEXT_NAMESPACE "SNiagaraSystemViewport"
+
 /** Viewport Client for the preview viewport */
 class FNiagaraSystemViewportClient : public FEditorViewportClient
 {
@@ -30,7 +32,9 @@ public:
 	virtual void Draw(FViewport* Viewport,FCanvas* Canvas) override;
 	virtual bool ShouldOrbitCamera() const override;
 	virtual FSceneView* CalcSceneView(FSceneViewFamily* ViewFamily, const EStereoscopicPass StereoPass = eSSP_FULL) override;
-	
+	virtual bool CanSetWidgetMode(FWidget::EWidgetMode NewMode) const override { return false; }
+	virtual bool CanCycleWidgetMode() const override { return false; }
+
 	void SetShowGrid(bool bShowGrid);
 
 	virtual void SetIsSimulateInEditorViewport(bool bInIsSimulateInEditorViewport)override;
@@ -53,7 +57,8 @@ FNiagaraSystemViewportClient::FNiagaraSystemViewportClient(FAdvancedPreviewScene
 	DrawHelper.GridColorMajor = FColor(72,72,72);
 	DrawHelper.GridColorMinor = FColor(64,64,64);
 	DrawHelper.PerspectiveGridSize = HALF_WORLD_MAX1;
-	
+	ShowWidget(false);
+
 	SetViewMode(VMI_Lit);
 	
 	EngineShowFlags.DisableAdvancedFeatures();
@@ -81,6 +86,8 @@ void FNiagaraSystemViewportClient::Tick(float DeltaSeconds)
 
 void FNiagaraSystemViewportClient::Draw(FViewport* InViewport,FCanvas* Canvas)
 {
+	UNiagaraSystem* ParticleSystem = NiagaraViewport.IsValid() ? NiagaraViewport->GetPreviewComponent()->GetAsset() : nullptr;
+
 	if (NiagaraViewport.IsValid() && NiagaraViewport->GetDrawElement(SNiagaraSystemViewport::EDrawElements::Bounds))
 	{
 		EngineShowFlags.SetBounds(true);
@@ -94,9 +101,9 @@ void FNiagaraSystemViewportClient::Draw(FViewport* InViewport,FCanvas* Canvas)
 
 	FEditorViewportClient::Draw(InViewport, Canvas);
 
-	if (bCaptureScreenShot)
+	if (bCaptureScreenShot && ParticleSystem != nullptr)
 	{
-		UNiagaraSystem* ParticleSystem = NiagaraViewport->GetPreviewComponent()->GetAsset();
+		
 		int32 SrcWidth = InViewport->GetSizeXY().X;
 		int32 SrcHeight = InViewport->GetSizeXY().Y;
 		// Read the contents of the viewport into an array.
@@ -347,7 +354,7 @@ TSharedRef<FEditorViewportClient> SNiagaraSystemViewport::MakeEditorViewportClie
 	SystemViewportClient->SetViewRotation( FRotator::ZeroRotator );
 	SystemViewportClient->SetViewLocationForOrbiting( FVector::ZeroVector );
 	SystemViewportClient->bSetListenerPosition = false;
-	
+
 	SystemViewportClient->SetRealtime( true );
 	SystemViewportClient->VisibilityDelegate.BindSP( this, &SNiagaraSystemViewport::IsVisible );
 	
@@ -371,6 +378,18 @@ EVisibility SNiagaraSystemViewport::OnGetViewportContentVisibility() const
 	return IsVisible() ? EVisibility::Visible : EVisibility::Collapsed;
 }
 
+EVisibility SNiagaraSystemViewport::OnGetViewportCompileTextVisibility() const
+{
+	if (PreviewComponent && PreviewComponent->GetAsset() && PreviewComponent->GetAsset()->HasOutstandingCompilationRequests())
+	{
+		return EVisibility::Visible;
+	}
+	else
+	{
+		return EVisibility::Collapsed;
+	}
+}
+
 void SNiagaraSystemViewport::PopulateViewportOverlays(TSharedRef<class SOverlay> Overlay)
 {
 	Overlay->AddSlot()
@@ -378,6 +397,15 @@ void SNiagaraSystemViewport::PopulateViewportOverlays(TSharedRef<class SOverlay>
 		[
 			SNew(SNiagaraSystemViewportToolBar, SharedThis(this))
 		];
+	Overlay->AddSlot()
+		.VAlign(VAlign_Center)
+		.HAlign(HAlign_Center)
+		[
+			SAssignNew(CompileText, STextBlock)
+			.Visibility_Raw(this, &SNiagaraSystemViewport::OnGetViewportCompileTextVisibility)
+		];
+
+	CompileText->SetText(LOCTEXT("Compiling","Compiling"));
 
 }
 
@@ -395,4 +423,7 @@ TSharedPtr<FExtender> SNiagaraSystemViewport::GetExtenders() const
 
 void SNiagaraSystemViewport::OnFloatingButtonClicked()
 {
-}
\ No newline at end of file
+}
+
+
+#undef LOCTEXT_NAMESPACE
\ No newline at end of file
diff --git a/Engine/Plugins/FX/Niagara/Source/NiagaraEditor/Private/Widgets/SNiagaraSystemViewport.h b/Engine/Plugins/FX/Niagara/Source/NiagaraEditor/Private/Widgets/SNiagaraSystemViewport.h
--- a/Engine/Plugins/FX/Niagara/Source/NiagaraEditor/Private/Widgets/SNiagaraSystemViewport.h
+++ b/Engine/Plugins/FX/Niagara/Source/NiagaraEditor/Private/Widgets/SNiagaraSystemViewport.h
@@ -89,6 +89,7 @@ protected:
 	virtual void BindCommands() override;
 	virtual void OnFocusViewportToSelection() override;
 	virtual void PopulateViewportOverlays(TSharedRef<class SOverlay> Overlay) override;
+	EVisibility OnGetViewportCompileTextVisibility() const;
 
 private:
 	/** The parent tab where this viewport resides */
@@ -96,6 +97,8 @@ private:
 	
 	/** Preview Scene - uses advanced preview settings */
 	TSharedPtr<class FAdvancedPreviewScene> AdvancedPreviewScene;
+
+	TSharedPtr<STextBlock> CompileText;
 	
 	bool IsVisible() const override;
 	
diff --git a/Engine/Plugins/FX/Niagara/Source/NiagaraEditor/Public/Customizations/NiagaraEventScriptPropertiesCustomization.h b/Engine/Plugins/FX/Niagara/Source/NiagaraEditor/Public/Customizations/NiagaraEventScriptPropertiesCustomization.h
--- a/Engine/Plugins/FX/Niagara/Source/NiagaraEditor/Public/Customizations/NiagaraEventScriptPropertiesCustomization.h
+++ b/Engine/Plugins/FX/Niagara/Source/NiagaraEditor/Public/Customizations/NiagaraEventScriptPropertiesCustomization.h
@@ -2,7 +2,7 @@
 
 #pragma once
 #include "CoreMinimal.h"
-#include "IDetailCustomization.h"
+#include "IPropertyTypeCustomization.h"
 #include "PropertyHandle.h"
 #include "EdGraph/EdGraphSchema.h"
 #include "EditorUndoClient.h"
@@ -44,11 +44,11 @@ struct FNiagaraStackAssetAction_EventSource : public FEdGraphSchemaAction
 
 };
 
-class FNiagaraEventScriptPropertiesCustomization : public IDetailCustomization, public FEditorUndoClient
+class FNiagaraEventScriptPropertiesCustomization : public IPropertyTypeCustomization, public FEditorUndoClient
 {
 public:
 	/** Makes a new instance of this detail layout class for a specific detail view requesting it */
-	static TSharedRef<class IDetailCustomization> MakeInstance(TWeakObjectPtr<UNiagaraSystem> InSystem,
+	static TSharedRef<class IPropertyTypeCustomization> MakeInstance(TWeakObjectPtr<UNiagaraSystem> InSystem,
 		TWeakObjectPtr<UNiagaraEmitter> InEmitter);
 
 	FNiagaraEventScriptPropertiesCustomization(TWeakObjectPtr<UNiagaraSystem> InSystem,
@@ -56,8 +56,9 @@ public:
 
 	~FNiagaraEventScriptPropertiesCustomization();
 
-	/** IDetailCustomization interface */
-	virtual void CustomizeDetails(IDetailLayoutBuilder& DetailBuilder) override;
+	/** IPropertyTypeCustomization interface */
+	virtual void CustomizeHeader(TSharedRef<IPropertyHandle> StructPropertyHandle, class FDetailWidgetRow& HeaderRow, IPropertyTypeCustomizationUtils& StructCustomizationUtils) override;
+	virtual void CustomizeChildren(TSharedRef<IPropertyHandle> StructPropertyHandle, class IDetailChildrenBuilder& ChildBuilder, IPropertyTypeCustomizationUtils& StructCustomizationUtils) override;
 
 	//~ FEditorUndoClient interface
 	virtual void PostUndo(bool bSuccess) override;
@@ -71,6 +72,8 @@ protected:
 	TSharedRef<SWidget> OnCreateWidgetForAction(struct FCreateWidgetForActionData* const InCreateData);
 	void OnActionSelected(const TArray< TSharedPtr<FEdGraphSchemaAction> >& SelectedActions, ESelectInfo::Type InSelectionType);
 	bool GetSpawnNumberEnabled() const;
+	bool GetUseRandomSpawnNumber() const;
+	EVisibility GetMinSpawnNumberVisible() const;
 	void ResolveEmitterName();
 	void ComputeErrorVisibility();
 	EVisibility GetErrorVisibility() const;
@@ -84,6 +87,8 @@ private:
 	TSharedPtr<IPropertyHandle> HandleSpawnNumber;
 	TSharedPtr<IPropertyHandle> HandleExecutionMode;
 	TSharedPtr<IPropertyHandle> HandleMaxEvents;
+	TSharedPtr<IPropertyHandle> HandleUseRandomSpawnNumber;
+	TSharedPtr<IPropertyHandle> HandleMinSpawnNumber;
 
 	TWeakObjectPtr<UNiagaraSystem> System;
 	TWeakObjectPtr<UNiagaraEmitter> Emitter;
