Copying //UE4/Dev-Niagara to //UE4/Dev-Main (Source: //UE4/Dev-Niagara @ 4074996)

diff --git a/Engine/Plugins/FX/Niagara/Source/NiagaraEditor/Public/EdGraphSchema_Niagara.h b/Engine/Plugins/FX/Niagara/Source/NiagaraEditor/Public/EdGraphSchema_Niagara.h
--- a/Engine/Plugins/FX/Niagara/Source/NiagaraEditor/Public/EdGraphSchema_Niagara.h
+++ b/Engine/Plugins/FX/Niagara/Source/NiagaraEditor/Public/EdGraphSchema_Niagara.h
@@ -72,10 +72,15 @@ class NIAGARAEDITOR_API UEdGraphSchema_Niagara : public UEdGraphSchema
 	virtual bool ShouldHidePinDefaultValue(UEdGraphPin* Pin) const override;
 	virtual bool TryCreateConnection(UEdGraphPin* A, UEdGraphPin* B) const override;
 	virtual void BreakSinglePinLink(UEdGraphPin* SourcePin, UEdGraphPin* TargetPin) const override;
+	virtual void BreakPinLinks(UEdGraphPin& TargetPin, bool bSendsNodeNotification) const override;
 	virtual FConnectionDrawingPolicy* CreateConnectionDrawingPolicy(int32 InBackLayerID, int32 InFrontLayerID, float InZoomFactor, const FSlateRect& InClippingRect, class FSlateWindowElementList& InDrawElements, class UEdGraph* InGraphObj) const override;
 	virtual void ResetPinToAutogeneratedDefaultValue(UEdGraphPin* Pin, bool bCallModifyCallbacks = true) const override;
+	virtual void OnPinConnectionDoubleCicked(UEdGraphPin* PinA, UEdGraphPin* PinB, const FVector2D& GraphPosition) const override; 
+	virtual bool ShouldAlwaysPurgeOnModification() const override { return false; }
 	//~ End EdGraphSchema Interface
 
+	static FLinearColor GetTypeColor(const FNiagaraTypeDefinition& Type);
+
 	TArray<TSharedPtr<FNiagaraSchemaAction_NewNode> > GetGraphContextActions(const UEdGraph* CurrentGraph, TArray<UObject*>& SelectedObjects, const UEdGraphPin* FromPin, UEdGraph* OwnerOfTemporaries) const;
 	void PromoteSinglePinToParameter(UEdGraphPin* SourcePin);
 	void ToggleNodeEnabledState(class UNiagaraNode* InNode) const;
@@ -102,6 +107,8 @@ class NIAGARAEDITOR_API UEdGraphSchema_Niagara : public UEdGraphSchema
 	bool IsSystemConstant(const FNiagaraVariable& Variable)const;
 
 	class UNiagaraParameterCollection* VariableIsFromParameterCollection(const FNiagaraVariable& Var)const;
+	class UNiagaraParameterCollection* VariableIsFromParameterCollection(const FString& VarName, bool bAllowPartialMatch, FNiagaraVariable& OutVar)const;
+	
 
 	FNiagaraTypeDefinition GetTypeDefForProperty(const UProperty* Property)const;
 
@@ -110,8 +117,9 @@ class NIAGARAEDITOR_API UEdGraphSchema_Niagara : public UEdGraphSchema
 	static const FLinearColor NodeTitleColor_SystemConstant;
 	static const FLinearColor NodeTitleColor_FunctionCall;
 	static const FLinearColor NodeTitleColor_CustomHlsl;
-	static const FLinearColor NodeTitleColor_Event;
-
+	static const FLinearColor NodeTitleColor_Event; 
+	static const FLinearColor NodeTitleColor_TranslatorConstant;
+	
 private:
 	void GetBreakLinkToSubMenuActions(class FMenuBuilder& MenuBuilder, UEdGraphPin* InGraphPin);
 	void GetNumericConversionToSubMenuActions(class FMenuBuilder& MenuBuilder, UEdGraphPin* InGraphPin);
@@ -122,7 +130,10 @@ private:
 class FNiagaraConnectionDrawingPolicy : public FConnectionDrawingPolicy
 {
 public:
-	FNiagaraConnectionDrawingPolicy(int32 InBackLayerID, int32 InFrontLayerID, float InZoomFactor, const FSlateRect& InClippingRect, FSlateWindowElementList& InDrawElements);
+	FNiagaraConnectionDrawingPolicy(int32 InBackLayerID, int32 InFrontLayerID, float InZoomFactor, const FSlateRect& InClippingRect, FSlateWindowElementList& InDrawElements, UEdGraph* InGraph);
 	virtual void DetermineWiringStyle(UEdGraphPin* OutputPin, UEdGraphPin* InputPin, /*inout*/ FConnectionParams& Params) override;
+
+private:
+	UEdGraph* Graph;
 };
 
diff --git a/Engine/Plugins/FX/Niagara/Source/NiagaraEditor/Public/INiagaraCompiler.h b/Engine/Plugins/FX/Niagara/Source/NiagaraEditor/Public/INiagaraCompiler.h
--- a/Engine/Plugins/FX/Niagara/Source/NiagaraEditor/Public/INiagaraCompiler.h
+++ b/Engine/Plugins/FX/Niagara/Source/NiagaraEditor/Public/INiagaraCompiler.h
@@ -14,6 +14,9 @@ class UEdGraphPin;
 class FCompilerResultsLog;
 class UNiagaraDataInterface;
 struct FNiagaraTranslatorOutput;
+struct FNiagaraVMExecutableData;
+class FNiagaraCompileRequestData;
+class FNiagaraCompileOptions;
 
 
 /** Defines information about the results of a Niagara script compile. */
@@ -28,18 +31,17 @@ struct FNiagaraCompileResults
 	/** A results log with messages, warnings, and errors which occurred during the compile. */
 	FCompilerResultsLog* MessageLog;
 
-	/** A string representation of the compilation output. */
-	FString OutputHLSL;
+	/** The actual final compiled data.*/
+	TSharedPtr<FNiagaraVMExecutableData> Data;
+
+	float CompileTime;
 
 	FNiagaraCompileResults(FCompilerResultsLog* InMessageLog = nullptr)
-		: MessageLog(InMessageLog)
+		: MessageLog(InMessageLog), CompileTime(0.0f)
 	{
 	}
 
 	static ENiagaraScriptCompileStatus CompileResultsToSummary(const FNiagaraCompileResults* CompileResults);
-
-	FNiagaraParameters Parameters;
-	TArray<FNiagaraVariable> Attributes;
 };
 
 //Interface for Niagara compilers.
@@ -49,7 +51,7 @@ class INiagaraCompiler
 {
 public:
 	/** Compiles a script. */
-	virtual const FNiagaraCompileResults& CompileScript(class UNiagaraScript* InScript, FNiagaraTranslatorOutput *TranslatorOutput, FString& TranslatedHLSL) = 0;
+	virtual FNiagaraCompileResults CompileScript(const FNiagaraCompileRequestData* InCompileRequest, const FNiagaraCompileOptions& InOptions, FNiagaraTranslatorOutput *TranslatorOutput, FString& TranslatedHLSL) = 0;
 
 	/** Adds an error to be reported to the user. Any error will lead to compilation failure. */
 	virtual void Error(FText ErrorText) = 0 ;
diff --git a/Engine/Plugins/FX/Niagara/Source/NiagaraEditor/Public/INiagaraEditorTypeUtilities.h b/Engine/Plugins/FX/Niagara/Source/NiagaraEditor/Public/INiagaraEditorTypeUtilities.h
--- a/Engine/Plugins/FX/Niagara/Source/NiagaraEditor/Public/INiagaraEditorTypeUtilities.h
+++ b/Engine/Plugins/FX/Niagara/Source/NiagaraEditor/Public/INiagaraEditorTypeUtilities.h
@@ -34,9 +34,13 @@ public:
 	virtual FString GetPinDefaultStringFromValue(const FNiagaraVariable& AllocatedVariable) const = 0;
 
 	virtual bool SetValueFromPinDefaultString(const FString& StringValue, FNiagaraVariable& Variable) const = 0;
+
+	virtual bool CanSetValueFromDisplayName() const = 0;
+
+	virtual bool SetValueFromDisplayName(const FText& TextValue, FNiagaraVariable& Variable) const = 0;
 };
 
-class FNiagaraEditorTypeUtilities : public INiagaraEditorTypeUtilities, public TSharedFromThis<FNiagaraEditorTypeUtilities>
+class FNiagaraEditorTypeUtilities : public INiagaraEditorTypeUtilities, public TSharedFromThis<FNiagaraEditorTypeUtilities, ESPMode::ThreadSafe>
 {
 public:
 	DECLARE_DELEGATE(FNotifyValueChanged);
@@ -51,4 +55,6 @@ public:
 	virtual bool CanHandlePinDefaults() const override { return false; }
 	virtual FString GetPinDefaultStringFromValue(const FNiagaraVariable& AllocatedVariable) const override { return FString(); }
 	virtual bool SetValueFromPinDefaultString(const FString& StringValue, FNiagaraVariable& Variable) const override { return false; }
-};
+	virtual bool CanSetValueFromDisplayName() const override { return false; }
+	virtual bool SetValueFromDisplayName(const FText& TextValue, FNiagaraVariable& Variable) const override { return false; }
+};
\ No newline at end of file
diff --git a/Engine/Plugins/FX/Niagara/Source/NiagaraEditor/Public/NiagaraActions.h b/Engine/Plugins/FX/Niagara/Source/NiagaraEditor/Public/NiagaraActions.h
new file mode 100644
--- /dev/null
+++ b/Engine/Plugins/FX/Niagara/Source/NiagaraEditor/Public/NiagaraActions.h
@@ -0,0 +1,140 @@
+// Copyright 1998-2018 Epic Games, Inc. All Rights Reserved.
+
+#pragma once
+
+#include "CoreMinimal.h"
+#include "NiagaraTypes.h"
+#include "NiagaraGraph.h"
+#include "EdGraph/EdGraphSchema.h"
+#include "GraphEditorDragDropAction.h"
+#include "NiagaraActions.generated.h"
+
+USTRUCT()
+struct NIAGARAEDITOR_API FNiagaraMenuAction : public FEdGraphSchemaAction
+{
+	GENERATED_USTRUCT_BODY();
+
+	DECLARE_DELEGATE(FOnExecuteStackAction);
+	DECLARE_DELEGATE_RetVal(bool, FCanExecuteStackAction);
+
+	FNiagaraMenuAction()
+	{
+	}
+
+	FNiagaraMenuAction(FText InNodeCategory, FText InMenuDesc, FText InToolTip, const int32 InGrouping, FText InKeywords, FOnExecuteStackAction InAction, int32 InSectionID = 0);
+	FNiagaraMenuAction(FText InNodeCategory, FText InMenuDesc, FText InToolTip, const int32 InGrouping, FText InKeywords, FOnExecuteStackAction InAction, FCanExecuteStackAction InCanPerformAction, int32 InSectionID = 0);
+
+	void ExecuteAction()
+	{
+		if (CanExecute())
+		{
+			Action.ExecuteIfBound();
+		}
+	}
+
+	bool CanExecute() const
+	{
+		// Fire the 'can execute' delegate if we have one, otherwise always return true
+		return CanPerformAction.IsBound() ? CanPerformAction.Execute() : true;
+	}
+
+private:
+	FOnExecuteStackAction Action;
+	FCanExecuteStackAction CanPerformAction;
+};
+
+USTRUCT()
+struct NIAGARAEDITOR_API FNiagaraParameterAction : public FEdGraphSchemaAction
+{
+	GENERATED_USTRUCT_BODY()
+
+	FNiagaraParameterAction() {}
+	FNiagaraParameterAction(const FNiagaraVariable& InParameter,
+		const TArray<FNiagaraGraphParameterReferenceCollection>& InReferenceCollection,
+		FText InNodeCategory, FText InMenuDesc, FText InToolTip, const int32 InGrouping, FText InKeywords, int32 InSectionID = 0);
+
+	FNiagaraVariable& GetParameter() { return Parameter; }
+
+	UPROPERTY()
+	FNiagaraVariable Parameter;
+
+	UPROPERTY()
+	TArray<FNiagaraGraphParameterReferenceCollection> ReferenceCollection;
+};
+
+class NIAGARAEDITOR_API FNiagaraParameterGraphDragOperation : public FGraphSchemaActionDragDropAction
+{
+public:
+	DRAG_DROP_OPERATOR_TYPE(FNiagaraParameterGraphDragOperation, FGraphSchemaActionDragDropAction)
+
+	static TSharedRef<FNiagaraParameterGraphDragOperation> New(const TSharedPtr<FEdGraphSchemaAction>& InActionNode);
+
+	// FGraphEditorDragDropAction interface
+	virtual void HoverTargetChanged() override;
+	virtual FReply DroppedOnNode(FVector2D ScreenPosition, FVector2D GraphPosition) override;
+	virtual FReply DroppedOnPanel(const TSharedRef< class SWidget >& Panel, FVector2D ScreenPosition, FVector2D GraphPosition, UEdGraph& Graph) override;
+	// End of FGraphEditorDragDropAction
+
+	/** Set if operation is modified by alt */
+	void SetAltDrag(bool InIsAltDrag) { bAltDrag = InIsAltDrag; }
+
+	/** Set if operation is modified by the ctrl key */
+	void SetCtrlDrag(bool InIsCtrlDrag) { bControlDrag = InIsCtrlDrag; }
+
+protected:
+	/** Constructor */
+	FNiagaraParameterGraphDragOperation();
+
+	/** Structure for required node construction parameters */
+	struct FNiagaraParameterNodeConstructionParams
+	{
+		FVector2D GraphPosition;
+		UEdGraph* Graph;
+		FNiagaraVariable Parameter;
+	};
+
+	static void MakeGetMap(FNiagaraParameterNodeConstructionParams InParams);
+	static void MakeSetMap(FNiagaraParameterNodeConstructionParams InParams);
+
+	virtual EVisibility GetIconVisible() const override;
+	virtual EVisibility GetErrorIconVisible() const override;
+
+	/** Was ctrl held down at start of drag */
+	bool bControlDrag;
+	/** Was alt held down at the start of drag */
+	bool bAltDrag;
+};
+
+class NIAGARAEDITOR_API FNiagaraStackDragOperation : public FDragDropOperation
+{
+public:
+	DRAG_DROP_OPERATOR_TYPE(FNiagaraStackDragOperation, FDragDropOperation)
+
+	static TSharedRef<FNiagaraStackDragOperation> New(TSharedPtr<FEdGraphSchemaAction> InActionNode);
+
+	void Construct();
+
+	/** Set if operation is modified by alt */
+	void SetAltDrag(bool InIsAltDrag) { bAltDrag = InIsAltDrag; }
+
+	/** Set if operation is modified by the ctrl key */
+	void SetCtrlDrag(bool InIsCtrlDrag) { bControlDrag = InIsCtrlDrag; }
+
+	TSharedPtr<FEdGraphSchemaAction> GetAction() { return SourceAction; }
+
+	virtual void OnDrop(bool bDropWasHandled, const FPointerEvent& MouseEvent) override;
+	virtual void OnDragged(const class FDragDropEvent& DragDropEvent) override;
+protected:
+	FNiagaraStackDragOperation();
+
+	bool HasFeedbackMessage();
+	void SetFeedbackMessage(const TSharedPtr<SWidget>& Message);
+	void SetSimpleFeedbackMessage(const FText& Message);
+	
+	TSharedPtr<FEdGraphSchemaAction> SourceAction;
+
+	/** Was ctrl held down at start of drag */
+	bool bControlDrag;
+	/** Was alt held down at the start of drag */
+	bool bAltDrag;
+};
\ No newline at end of file
diff --git a/Engine/Plugins/FX/Niagara/Source/NiagaraEditor/Public/NiagaraEditorModule.h b/Engine/Plugins/FX/Niagara/Source/NiagaraEditor/Public/NiagaraEditorModule.h
--- a/Engine/Plugins/FX/Niagara/Source/NiagaraEditor/Public/NiagaraEditorModule.h
+++ b/Engine/Plugins/FX/Niagara/Source/NiagaraEditor/Public/NiagaraEditorModule.h
@@ -17,6 +17,10 @@ class UNiagaraSettings;
 class USequencerSettings;
 class UNiagaraStackViewModel;
 class FNiagaraScriptMergeManager;
+class FNiagaraCompileOptions;
+class FNiagaraCompileRequestDataBase;
+class UMovieSceneNiagaraParameterTrack;
+struct IConsoleCommand;
 
 DECLARE_STATS_GROUP(TEXT("Niagara Editor"), STATGROUP_NiagaraEditor, STATCAT_Advanced);
 
@@ -26,6 +30,15 @@ class FNiagaraEditorModule : public IModuleInterface,
 {
 public:
 	DECLARE_DELEGATE_RetVal_OneParam(TSharedRef<SWidget>, FOnCreateStackWidget, UNiagaraStackViewModel*);
+	DECLARE_DELEGATE_RetVal_OneParam(UMovieSceneNiagaraParameterTrack*, FOnCreateMovieSceneTrackForParameter, FNiagaraVariable);
+
+	struct FInputMetaDataKeys
+	{
+		static const FName AdvancedDisplay;
+		static const FName EditCondition;
+		static const FName VisibleCondition;
+		static const FName InlineEditConditionToggle;
+	};
 
 public:
 	FNiagaraEditorModule();
@@ -34,22 +47,27 @@ public:
 	virtual void StartupModule() override;
 	virtual void ShutdownModule() override;
 
+	/** Get the instance of this module. */
+	NIAGARAEDITOR_API static FNiagaraEditorModule& Get();
+
 	/** Compile the specified script. */
-	virtual ENiagaraScriptCompileStatus CompileScript(class UNiagaraScript* ScriptToCompile, FString& OutGraphLevelErrorMessages);
+	virtual TSharedPtr<FNiagaraVMExecutableData> CompileScript(const FNiagaraCompileRequestDataBase* InCompileRequest, const FNiagaraCompileOptions& InCompileOptions);
+
+	TSharedPtr<FNiagaraCompileRequestDataBase, ESPMode::ThreadSafe> Precompile(UObject* Obj);
 
 	/** Gets the extensibility managers for outside entities to extend static mesh editor's menus and toolbars */
 	virtual TSharedPtr<FExtensibilityManager> GetMenuExtensibilityManager() override {return MenuExtensibilityManager;}
 	virtual TSharedPtr<FExtensibilityManager> GetToolBarExtensibilityManager() override {return ToolBarExtensibilityManager;}
 
 	/** Registers niagara editor type utilities for a specific type. */
-	void RegisterTypeUtilities(FNiagaraTypeDefinition Type, TSharedRef<INiagaraEditorTypeUtilities> EditorUtilities);
+	void RegisterTypeUtilities(FNiagaraTypeDefinition Type, TSharedRef<INiagaraEditorTypeUtilities, ESPMode::ThreadSafe> EditorUtilities);
 
 	/** Register/unregister niagara editor settings. */
 	void RegisterSettings();
 	void UnregisterSettings();
 
-	/** Gets niagara editor type utilities for a specific type if there are any registered. */
-	TSharedPtr<INiagaraEditorTypeUtilities> NIAGARAEDITOR_API GetTypeUtilities(const FNiagaraTypeDefinition& Type);
+	/** Gets Niagara editor type utilities for a specific type if there are any registered. */
+	TSharedPtr<INiagaraEditorTypeUtilities, ESPMode::ThreadSafe> NIAGARAEDITOR_API GetTypeUtilities(const FNiagaraTypeDefinition& Type);
 
 	static EAssetTypeCategories::Type GetAssetCategory() { return NiagaraAssetCategory; }
 
@@ -60,12 +78,20 @@ public:
 
 	TSharedRef<FNiagaraScriptMergeManager> GetScriptMergeManager() const;
 
+	void RegisterParameterTrackCreatorForType(const UScriptStruct& StructType, FOnCreateMovieSceneTrackForParameter CreateTrack);
+	void UnregisterParameterTrackCreatorForType(const UScriptStruct& StructType);
+	bool CanCreateParameterTrackForType(const UScriptStruct& StructType);
+	UMovieSceneNiagaraParameterTrack* CreateParameterTrackForType(const UScriptStruct& StructType, FNiagaraVariable Parameter);
+
 	/** Niagara Editor app identifier string */
 	static const FName NiagaraEditorAppIdentifier;
 
 	/** The tab color scale for niagara editors. */
 	static const FLinearColor WorldCentricTabColorScale;
 
+	/** Get the niagara UI commands. */
+	NIAGARAEDITOR_API const class FNiagaraEditorCommands& Commands();
+
 private:
 	void RegisterAssetTypeAction(IAssetTools& AssetTools, TSharedRef<IAssetTypeActions> Action);
 	void OnNiagaraSettingsChangedEvent(const FString& PropertyName, const UNiagaraSettings* Settings);
@@ -74,6 +100,8 @@ private:
 	/** FGCObject interface */
 	virtual void AddReferencedObjects( FReferenceCollector& Collector ) override;
 
+	void TestCompileScriptFromConsole(const TArray<FString>& Arguments);
+
 private:
 	TSharedPtr<FExtensibilityManager> MenuExtensibilityManager;
 	TSharedPtr<FExtensibilityManager> ToolBarExtensibilityManager;
@@ -81,21 +109,35 @@ private:
 	/** All created asset type actions.  Cached here so that we can unregister it during shutdown. */
 	TArray< TSharedPtr<IAssetTypeActions> > CreatedAssetTypeActions;
 
-	TMap<FNiagaraTypeDefinition, TSharedRef<INiagaraEditorTypeUtilities>> TypeToEditorUtilitiesMap;
-	TSharedPtr<INiagaraEditorTypeUtilities> EnumTypeUtilities;
+	FCriticalSection TypeEditorsCS;
+	TMap<FNiagaraTypeDefinition, TSharedRef<INiagaraEditorTypeUtilities, ESPMode::ThreadSafe>> TypeToEditorUtilitiesMap;
+	TSharedPtr<INiagaraEditorTypeUtilities, ESPMode::ThreadSafe> EnumTypeUtilities;
 
 	static EAssetTypeCategories::Type NiagaraAssetCategory;
 
 	FDelegateHandle CreateEmitterTrackEditorHandle;
 	FDelegateHandle CreateSystemTrackEditorHandle;
-	FDelegateHandle CreateFloatParameterTrackEditorHandle;
+
 	FDelegateHandle CreateBoolParameterTrackEditorHandle;
+	FDelegateHandle CreateFloatParameterTrackEditorHandle;
+	FDelegateHandle CreateIntegerParameterTrackEditorHandle;
+	FDelegateHandle CreateVectorParameterTrackEditorHandle;
+	FDelegateHandle CreateColorParameterTrackEditorHandle;
+
 	FDelegateHandle MergeEmitterHandle;
+	FDelegateHandle CreateDefaultScriptSourceHandle;
+	FDelegateHandle ScriptCompilerHandle;
+	FDelegateHandle PrecompilerHandle;
 
 	USequencerSettings* SequencerSettings;
 	FOnCreateStackWidget OnCreateStackWidget;
 
 	TSharedPtr<FNiagaraScriptMergeManager> ScriptMergeManager;
+
+	TMap<const UScriptStruct*, FOnCreateMovieSceneTrackForParameter> TypeToParameterTrackCreatorMap;
+
+	IConsoleCommand* TestCompileScriptCommand;
+	IConsoleCommand* DumpRapidIterationParametersForAsset;
 };
 
 
diff --git a/Engine/Plugins/FX/Niagara/Source/NiagaraEditor/Public/NiagaraEditorSettings.h b/Engine/Plugins/FX/Niagara/Source/NiagaraEditor/Public/NiagaraEditorSettings.h
--- a/Engine/Plugins/FX/Niagara/Source/NiagaraEditor/Public/NiagaraEditorSettings.h
+++ b/Engine/Plugins/FX/Niagara/Source/NiagaraEditor/Public/NiagaraEditorSettings.h
@@ -44,21 +44,37 @@ public:
 	UPROPERTY(config, EditAnywhere, Category = Niagara)
 	TArray<FNiagaraSpawnShortcut> GraphCreationShortcuts;
 
-	/** Whether or not auto-compile is enabled in the editors. */
-	UPROPERTY(config, EditAnywhere, Category = Niagara)
-	bool bAutoCompile;
+	/** Gets whether or not auto-compile is enabled in the editors. */
+	bool GetAutoCompile() const;
+
+	/** Sets whether or not auto-compile is enabled in the editors. */
+	void SetAutoCompile(bool bInAutoCompile);
+
+	/** Gets whether or not simulations should start playing automatically when the emitter or system editor is opened, or when the data is changed in the editor. */
+	bool GetAutoPlay() const;
+
+	/** Sets whether or not simulations should start playing automatically when the emitter or system editor is opened, or when the data is changed in the editor. */
+	void SetAutoPlay(bool bInAutoPlay);
+
+	/** Gets whether or not the simulation should reset when a value on the emitter or system is changed. */
+	bool GetResetSimulationOnChange() const;
+
+	/** Sets whether or not the simulation should reset when a value on the emitter or system is changed. */
+	void SetResetSimulationOnChange(bool bInResetSimulationOnChange);
+
+	/** Gets whether or not to rerun the simulation to the current time when making modifications while paused. */
+	bool GetResimulateOnChangeWhilePaused() const;
+
+	/** Sets whether or not to rerun the simulation to the current time when making modifications while paused. */
+	void SetResimulateOnChangeWhilePaused(bool bInResimulateOnChangeWhilePaused);
 	
 	// Begin UDeveloperSettings Interface
 	virtual FName GetCategoryName() const override;
-#if WITH_EDITOR
 	virtual FText GetSectionText() const override;
-#endif
 	// END UDeveloperSettings Interface
 
-#if WITH_EDITOR
 	virtual void PostEditChangeProperty(FPropertyChangedEvent& PropertyChangedEvent) override;
 
-public:
 	DECLARE_MULTICAST_DELEGATE_TwoParams(FOnNiagaraEditorSettingsChanged, const FString&, const UNiagaraEditorSettings*);
 
 	/** Gets a multicast delegate which is called whenever one of the parameters in this settings object changes. */
@@ -66,5 +82,21 @@ public:
 
 protected:
 	static FOnNiagaraEditorSettingsChanged SettingsChangedDelegate;
-#endif
+
+private:
+	/** Whether or not auto-compile is enabled in the editors. */
+	UPROPERTY(config, EditAnywhere, Category = Niagara)
+	bool bAutoCompile;
+
+	/** Whether or not simulations should start playing automatically when the emitter or system editor is opened, or when the data is changed in the editor. */
+	UPROPERTY(config, EditAnywhere, Category = SimulationOptions)
+	bool bAutoPlay;
+
+	/** Whether or not the simulation should reset when a value on the emitter or system is changed. */
+	UPROPERTY(config, EditAnywhere, Category = SimulationOptions)
+	bool bResetSimulationOnChange;
+
+	/** Whether or not to rerun the simulation to the current time when making modifications while paused. */
+	UPROPERTY(config, EditAnywhere, Category = SimulationOptions)
+	bool bResimulateOnChangeWhilePaused;
 };
diff --git a/Engine/Plugins/FX/Niagara/Source/NiagaraEditor/Public/NiagaraEditorUtilities.h b/Engine/Plugins/FX/Niagara/Source/NiagaraEditor/Public/NiagaraEditorUtilities.h
--- a/Engine/Plugins/FX/Niagara/Source/NiagaraEditor/Public/NiagaraEditorUtilities.h
+++ b/Engine/Plugins/FX/Niagara/Source/NiagaraEditor/Public/NiagaraEditorUtilities.h
@@ -18,6 +18,9 @@ class FStructOnScope;
 class UEdGraph;
 class UEdGraphNode;
 class SWidget;
+class UNiagaraNode;
+class UEdGraphSchema_Niagara;
+class UEdGraphPin;
 
 namespace FNiagaraEditorUtilities
 {
@@ -62,7 +65,7 @@ namespace FNiagaraEditorUtilities
 
 		return true;
 	}
-	
+
 	/** Gets a set of the system constant names. */
 	TSet<FName> GetSystemConstantNames();
 
@@ -75,6 +78,13 @@ namespace FNiagaraEditorUtilities
 	/** Sets up a niagara input node for parameter usage. */
 	void InitializeParameterInputNode(UNiagaraNodeInput& InputNode, const FNiagaraTypeDefinition& Type, const UNiagaraGraph* Graph, FName InputName = FName(TEXT("NewInput")));
 
+	/** Writes text to a specified location on disk.*/
+	void NIAGARAEDITOR_API WriteTextFileToDisk(FString SaveDirectory, FString FileName, FString TextToSave, bool bAllowOverwriting = false);
+
+	/** Gathers up the change Id's and optionally writes them to disk.*/
+	void GatherChangeIds(UNiagaraEmitter& Emitter, TMap<FGuid, FGuid>& ChangeIds, const FString& InDebugName, bool bWriteToLogDir = false);
+	void GatherChangeIds(UNiagaraGraph& Graph, TMap<FGuid, FGuid>& ChangeIds, const FString& InDebugName, bool bWriteToLogDir = false);
+
 	/** Options for the GetParameterVariablesFromSystem function. */
 	struct FGetParameterVariablesFromSystemOptions
 	{
@@ -91,9 +101,6 @@ namespace FNiagaraEditorUtilities
 	/** Gets the niagara variables for the input parameters on a niagara System. */
 	void GetParameterVariablesFromSystem(UNiagaraSystem& System, TArray<FNiagaraVariable>& ParameterVariables, FGetParameterVariablesFromSystemOptions Options = FGetParameterVariablesFromSystemOptions());
 
-	/** Generates a merged graph for an Emitter properties object that previously had unique graphs per script type.*/
-	bool ConvertToMergedGraph(UNiagaraEmitter* InEmitter);
-
 	/** Helper to clean up copy & pasted graphs.*/
 	void FixUpPastedInputNodes(UEdGraph* Graph, TSet<UEdGraphNode*> PastedNodes);
 
@@ -115,6 +122,16 @@ namespace FNiagaraEditorUtilities
 	TSharedPtr<SWidget> CreateInlineErrorText(TAttribute<FText> ErrorMessage, TAttribute<FText> ErrorTooltip);
 
 	void CompileExistingEmitters(const TArray<UNiagaraEmitter*>& AffectedEmitters);
-	
+
 	bool TryGetEventDisplayName(UNiagaraEmitter* Emitter, FGuid EventUsageId, FText& OutEventDisplayName);
-}
+
+	bool IsCompilableAssetClass(UClass* AssetClass);
+
+	void MarkDependentCompilableAssetsDirty(TArray<UObject*> InObjects);
+
+	void ResolveNumerics(UNiagaraGraph* SourceGraph, bool bForceParametersToResolveNumerics, TArray<FNiagaraVariable>& ChangedNumericParams);
+
+	void FixUpNumericPins(const UEdGraphSchema_Niagara* Schema, UNiagaraNode* Node);
+
+	void PreprocessFunctionGraph(const UEdGraphSchema_Niagara* Schema, UNiagaraGraph* Graph, const TArray<UEdGraphPin*>& CallInputs, const TArray<UEdGraphPin*>& CallOutputs, ENiagaraScriptUsage ScriptUsage);
+};
diff --git a/Engine/Plugins/FX/Niagara/Source/NiagaraEditor/Public/NiagaraGraph.h b/Engine/Plugins/FX/Niagara/Source/NiagaraEditor/Public/NiagaraGraph.h
--- a/Engine/Plugins/FX/Niagara/Source/NiagaraEditor/Public/NiagaraGraph.h
+++ b/Engine/Plugins/FX/Niagara/Source/NiagaraEditor/Public/NiagaraGraph.h
@@ -7,6 +7,8 @@
 #include "NiagaraCommon.h"
 #include "EdGraph/EdGraph.h"
 #include "NiagaraScript.h"
+#include "Misc/SecureHash.h"
+#include "NiagaraParameterStore.h"
 #include "NiagaraGraph.generated.h"
 
 /** This is the type of action that occurred on a given Niagara graph. Note that this should follow from EEdGraphActionType, leaving some slop for growth. */
@@ -15,6 +17,64 @@ enum ENiagaraGraphActionType
 	GRAPHACTION_GenericNeedsRecompile = 0x1 << 16,
 };
 
+class UNiagaraNode;
+
+typedef TPair<FGuid, TWeakObjectPtr<UNiagaraNode>> FNiagaraGraphParameterReference;
+
+USTRUCT()
+struct FNiagaraGraphParameterReferenceCollection
+{
+	GENERATED_USTRUCT_BODY()
+
+public:
+	FNiagaraGraphParameterReferenceCollection(const bool bInCreated = false);
+
+	/** All the references in the graph. */
+	TArray<FNiagaraGraphParameterReference> ParameterReferences;
+
+	UPROPERTY()
+	class UNiagaraGraph* Graph;
+
+	/** Returns true if this parameter was initially created by the user. */
+	bool WasCreated() const;
+
+private:
+	/** Whether this parameter was initially created by the user. */
+	UPROPERTY()
+	bool bCreated;
+
+};
+
+/** Container for UNiagaraGraph cached data for managing CompileIds and Traversals.*/
+USTRUCT()
+struct FNiagaraGraphScriptUsageInfo
+{
+	GENERATED_USTRUCT_BODY()
+
+public:
+	FNiagaraGraphScriptUsageInfo();
+
+	/** The context in which this sub-graph traversal will be used.*/
+	UPROPERTY()
+	ENiagaraScriptUsage UsageType;
+	
+	/** The particular instance of the usage type. Event scripts, for example, have potentially multiple graphs.*/
+	UPROPERTY()
+	FGuid UsageId;
+
+	/** The compile ID last associated with this traversal.*/
+	UPROPERTY()
+	FGuid GeneratedCompileId;
+
+	/** The SHA1 hash that we calculated last traversal. If all zeroes, it is invalid.*/
+	UPROPERTY()
+	TArray<uint8> DataHash;
+
+	/** The traversal of output to input nodes for this graph. This is not a recursive traversal, it just includes nodes from this graph.*/
+	UPROPERTY()
+	TArray<UNiagaraNode*> Traversal;
+};
+
 
 UCLASS(MinimalAPI)
 class UNiagaraGraph : public UEdGraph
@@ -36,10 +96,12 @@ class UNiagaraGraph : public UEdGraph
 			
 	/** Find the first output node bound to the target usage type.*/
 	class UNiagaraNodeOutput* FindOutputNode(ENiagaraScriptUsage TargetUsageType, FGuid TargetUsageId = FGuid()) const;
+	class UNiagaraNodeOutput* FindEquivalentOutputNode(ENiagaraScriptUsage TargetUsageType, FGuid TargetUsageId = FGuid()) const;
 
 	/** Find all output nodes.*/
 	void FindOutputNodes(TArray<UNiagaraNodeOutput*>& OutputNodes) const;
 	void FindOutputNodes(ENiagaraScriptUsage TargetUsageType, TArray<UNiagaraNodeOutput*>& OutputNodes) const;
+	void FindEquivalentOutputNodes(ENiagaraScriptUsage TargetUsageType, TArray<UNiagaraNodeOutput*>& OutputNodes) const;
 
 	/** Options for the FindInputNodes function */
 	struct FFindInputNodeOptions
@@ -49,6 +111,7 @@ class UNiagaraGraph : public UEdGraph
 			, bIncludeParameters(true)
 			, bIncludeAttributes(true)
 			, bIncludeSystemConstants(true)
+			, bIncludeTranslatorConstants(false)
 			, bFilterDuplicates(false)
 			, bFilterByScriptUsage(false)
 			, TargetScriptUsage(ENiagaraScriptUsage::Function)
@@ -63,6 +126,8 @@ class UNiagaraGraph : public UEdGraph
 		bool bIncludeAttributes;
 		/** Whether or not to include system parameters, defaults to true. */
 		bool bIncludeSystemConstants;
+		/** Whether or not to include translator parameters, defaults to false. */
+		bool bIncludeTranslatorConstants;
 		/** Whether of not to filter out duplicate nodes, defaults to false. */
 		bool bFilterDuplicates;
 		/** Whether or not to limit to nodes connected to an output node of the specified script type.*/
@@ -78,14 +143,11 @@ class UNiagaraGraph : public UEdGraph
 
 	/** Get an in-order traversal of a graph by the specified target output script usage.*/
 	void BuildTraversal(TArray<class UNiagaraNode*>& OutNodesTraversed, ENiagaraScriptUsage TargetUsage, FGuid TargetUsageId) const;
-	void BuildTraversal(TArray<class UNiagaraNode*>& OutNodesTraversed, UNiagaraNode* FinalNode) const;
+	static void BuildTraversal(TArray<class UNiagaraNode*>& OutNodesTraversed, UNiagaraNode* FinalNode);
 
 	/** Generates a list of unique input and output parameters for when this script is used as a function. */
 	void GetParameters(TArray<FNiagaraVariable>& Inputs, TArray<FNiagaraVariable>& Outputs) const;
 
-	void FindReadDataSetNodes(TArray<class UNiagaraNodeReadDataSet*>& OutReadNodes) const;
-	void FindWriteDataSetNodes(TArray<class UNiagaraNodeWriteDataSet*>& OutWriteNodes) const;
-
 	/** Returns the index of this variable in the output node of the graph. INDEX_NONE if this is not a valid attribute. */
 	int32 GetOutputNodeVariableIndex(const FNiagaraVariable& Attr)const;
 	void GetOutputNodeVariables(TArray< FNiagaraVariable >& OutAttributes)const;
@@ -97,13 +159,16 @@ class UNiagaraGraph : public UEdGraph
 
 	/** Signal to listeners that the graph has changed */
 	void NotifyGraphNeedsRecompile();
-
+		
 	/** Notifies the graph that a contained data interface has changed. */
 	void NotifyGraphDataInterfaceChanged();
 
 	/** Get all referenced graphs in this specified graph, including this graph. */
 	void GetAllReferencedGraphs(TArray<const UNiagaraGraph*>& Graphs) const;
 
+	/** Gather all the change ids of external references for this specific graph traversal.*/
+	void GatherExternalDependencyIDs(ENiagaraScriptUsage InUsage, const FGuid& InUsageId, TArray<FGuid>& InReferencedIDs, TArray<UObject*>& InReferencedObjs);
+
 	/** Determine if there are any external dependencies wrt to scripts and ensure that those dependencies are sucked into the existing package.*/
 	void SubsumeExternalDependencies(TMap<const UObject*, UObject*>& ExistingConversions);
 
@@ -111,14 +176,21 @@ class UNiagaraGraph : public UEdGraph
 	bool IsOtherSynchronized(const FGuid& InChangeId) const;
 
 	/** Identify that this graph has undergone changes that will require synchronization with a compiled script.*/
-	void MarkGraphRequiresSynchronization();
+	void MarkGraphRequiresSynchronization(FString Reason);
 
+	/** A change was made to the graph that external parties should take note of. The ChangeID will be updated.*/
 	virtual void NotifyGraphChanged() override;
 
+	/** Each graph is given a Change Id that occurs anytime the graph's content is manipulated. This key changing induces several important activities, including being a 
+	value that third parties can poll to see if their cached handling of the graph needs to potentially adjust to changes. Furthermore, for script compilation we cache 
+	the changes that were produced during the traversal of each output node, which are referred to as the CompileID.*/
 	FGuid GetChangeID() { return ChangeId; }
 
+	/** The compile id associated with the output node traversal specified by InUsage and InUsageId. If not found, an invalid Guid is returned.*/
+	FGuid GetCompileID(ENiagaraScriptUsage InUsage, const FGuid& InUsageId);
+
 	/** Walk through the graph for an ParameterMapGet nodes and see if any of them specify a default for VariableName.*/
-	UEdGraphPin* FindParameterMapDefaultValuePin(const FName VariableName);
+	UEdGraphPin* FindParameterMapDefaultValuePin(const FName VariableName, ENiagaraScriptUsage InUsage, ENiagaraScriptUsage InParentUsage) const;
 
 	/** Get the meta-data associated with this variable, if it exists.*/
 	FNiagaraVariableMetaData* GetMetaData(const FNiagaraVariable& InVar);
@@ -130,20 +202,77 @@ class UNiagaraGraph : public UEdGraph
 	/** Remove any meta-data that is no longer being referenced within this graph.*/
 	void PurgeUnreferencedMetaData();
 
+	/** Remove any parameter that is no longer being referenced within this graph.*/
+	void PurgeUnreferencedParameters();
+
+	const TMap<FNiagaraVariable, FNiagaraVariableMetaData>& GetAllMetaData() const { return VariableToMetaData; }
+	TMap<FNiagaraVariable, FNiagaraVariableMetaData>& GetAllMetaData() { return VariableToMetaData; }
+
+	const TMap<FNiagaraVariable, FNiagaraGraphParameterReferenceCollection>& GetParameterMap();
+
+	/** Adds parameter to parameters map setting it as created by the user.*/
+	void AddParameter(const FNiagaraVariable& Parameter);
+
+	/** Remove parameter from map and all the pins associated. */
+	void RemoveParameter(const FNiagaraVariable& Parameter, const bool bNotifyGraphChanged = true);
+
+	/** Rename parameter from map and all the pins associated. */
+	bool RenameParameter(const FNiagaraVariable& Parameter, FName NewName, const bool bInNotifyGraphChanged = true);
+
 	/** Gets a delegate which is called whenever a contained data interfaces changes. */
 	FOnDataInterfaceChanged& OnDataInterfaceChanged();
 
+	void SynchronizeInternalCacheWithGraph(UNiagaraGraph* Other);
+	void InvalidateCachedCompileIds();
+
+	/** Add a listener for OnGraphNeedsRecompile events */
+	FDelegateHandle AddOnGraphNeedsRecompileHandler(const FOnGraphChanged::FDelegate& InHandler);
+
+	/** Remove a listener for OnGraphNeedsRecompile events */
+	void RemoveOnGraphNeedsRecompileHandler(FDelegateHandle Handle);
+
+protected:
+	void RebuildCachedData(bool bForce = false);
+
 private:
 	virtual void NotifyGraphChanged(const FEdGraphEditAction& InAction) override;
+
+	/** Find parameters in the graph. */
+	void FindParameters();
+
+	/** 
+	 Set whether finding parameters is allowed. 
+	 To block finding parameters every time NotifyGraphChanged is called whenever we know it will be called multiple times.
+	*/
+	void SetFindParametersAllowed(const bool bAllowed);
+
+	/** Finding parameters is blocked if false.*/
+	bool bFindParametersAllowed;
+
+	/** A delegate that broadcasts a notification whenever the graph needs recompile due to structural change. */
+	FOnGraphChanged OnGraphNeedsRecompile;
 	
-	/** The current change identifier for this graph. Used to sync status with UNiagaraScripts.*/
+	/** The current change identifier for this graph overall. Used to sync status with UNiagaraScripts.*/
 	UPROPERTY()
 	FGuid ChangeId;
 
+	UPROPERTY()
+	FGuid LastBuiltTraversalDataChangeId;
+
+	UPROPERTY()
+	TArray<FNiagaraGraphScriptUsageInfo> CachedUsageInfo;
+
 	/** Storage of meta-data for variables defined for use explicitly with this graph.*/
 	UPROPERTY()
-	TMap<FNiagaraVariable,FNiagaraVariableMetaData> VariableToMetaData;
+	TMap<FNiagaraVariable, FNiagaraVariableMetaData> VariableToMetaData;
+	
+	/** Parameters found within the graph. */
+	UPROPERTY()
+	TMap<FNiagaraVariable, FNiagaraGraphParameterReferenceCollection> Parameters;
 
 	FOnDataInterfaceChanged OnDataInterfaceChangedDelegate;
+
+	/** Whether currently renaming a parameter to prevent recursion. */
+	bool bIsRenamingParameter;
 };
 
diff --git a/Engine/Plugins/FX/Niagara/Source/NiagaraEditor/Public/NiagaraNode.h b/Engine/Plugins/FX/Niagara/Source/NiagaraEditor/Public/NiagaraNode.h
--- a/Engine/Plugins/FX/Niagara/Source/NiagaraEditor/Public/NiagaraNode.h
+++ b/Engine/Plugins/FX/Niagara/Source/NiagaraEditor/Public/NiagaraNode.h
@@ -8,6 +8,7 @@
 #include "EdGraph/EdGraphNode.h"
 #include "NiagaraEditorCommon.h"
 #include "NiagaraParameterMapHistory.h"
+#include "Misc/Guid.h"
 #include "NiagaraNode.generated.h"
 
 class UEdGraphPin;
@@ -25,15 +26,22 @@ protected:
 
 public:
 
+	DECLARE_MULTICAST_DELEGATE_OneParam(FOnNodeVisualsChanged, UNiagaraNode*);
+
 	//~ Begin UObject interface
 	virtual void PostLoad() override;
 	//~ End UObject interface
 
 	//~ Begin EdGraphNode Interface
+	virtual void PostPlacedNewNode() override;
 	virtual void AutowireNewNode(UEdGraphPin* FromPin)override;	
 	virtual void PinDefaultValueChanged(UEdGraphPin* Pin) override;
 	virtual void PinConnectionListChanged(UEdGraphPin* Pin) override;
 	virtual void PinTypeChanged(UEdGraphPin* Pin) override;
+	virtual void OnRenameNode(const FString& NewName) override;
+	virtual void OnPinRemoved(UEdGraphPin* InRemovedPin) override;
+	virtual void NodeConnectionListChanged() override;	
+	virtual TSharedPtr<SGraphNode> CreateVisualWidget() override; 
 	//~ End EdGraphNode Interface
 
 	/** Get the Niagara graph that owns this node */
@@ -46,7 +54,7 @@ public:
 	/** Gets the asset referenced by this node, or nullptr if there isn't one. */
 	virtual UObject* GetReferencedAsset() const { return nullptr; }
 
-	/** Refreshes the node due to external changes, e.g. the underlying function changed for a function call node. Return true if the graph changed. */
+	/** Refreshes the node due to external changes, e.g. the underlying function changed for a function call node. Return true if the graph changed.*/
 	virtual bool RefreshFromExternalChanges() { return false; }
 
 	virtual void Compile(class FHlslNiagaraTranslator *Translator, TArray<int32>& Outputs);
@@ -55,6 +63,7 @@ public:
 	void GetInputPins(TArray<class UEdGraphPin*>& OutInputPins) const;
 	UEdGraphPin* GetOutputPin(int32 OutputIndex) const;
 	void GetOutputPins(TArray<class UEdGraphPin*>& OutOutputPins) const;
+	UEdGraphPin* GetPinByPersistentGuid(const FGuid& InGuid) const;
 
 	/** Apply any node-specific logic to determine if it is safe to add this node to the graph. This is meant to be called only in the Editor before placing the node.*/
 	virtual bool CanAddToGraph(UNiagaraGraph* TargetGraph, FString& OutErrorMsg) const;
@@ -79,24 +88,42 @@ public:
 
 	/** Verify that the potential rename has produced acceptable results for a pin.*/
 	virtual bool CommitEditablePinName(const FText& InName,  UEdGraphPin* InGraphPinObj) { return false; }
-	
-	virtual void BuildParameterMapHistory(FNiagaraParameterMapHistoryBuilder& OutHistory, bool bRecursive = true);
+	/** Notify the rename was cancelled.*/
+	virtual bool CancelEditablePinName(const FText& InName, UEdGraphPin* InGraphPinObj) { return false; }
 
-	bool ShouldVisit(int32 InVisitID) 
-	{ 
-		bool bShouldVisit = InVisitID != VisitID;
-		VisitID = InVisitID;
-		return bShouldVisit;
-	}
+	virtual void BuildParameterMapHistory(FNiagaraParameterMapHistoryBuilder& OutHistory, bool bRecursive = true);
 	
+	/** Go through all the external dependencies of this node in isolation and add them to the reference id list.*/
+	virtual void GatherExternalDependencyIDs(ENiagaraScriptUsage InMasterUsage, const FGuid& InMasterUsageId, TArray<FGuid>& InReferencedIDs, TArray<UObject*>& InReferencedObjs) const {};
+
+	/** Traces one of this node's output pins to its source output pin if it is a reroute node output pin.*/
+	virtual UEdGraphPin* GetTracedOutputPin(UEdGraphPin* LocallyOwnedOutputPin) const {return LocallyOwnedOutputPin;}
+	static UEdGraphPin* TraceOutputPin(UEdGraphPin* LocallyOwnedOutputPin);
+
+	virtual UEdGraphPin* GetPassThroughPin(const UEdGraphPin* LocallyOwnedOutputPin) const override { return nullptr; }
+	virtual UEdGraphPin* GetPassThroughPin(const UEdGraphPin* LocallyOwnedOutputPin, ENiagaraScriptUsage MasterUsage) const { return nullptr; }
+
+	/** Identify that this node has undergone changes that will require synchronization with a compiled script.*/
+	void MarkNodeRequiresSynchronization(FString Reason, bool bRaiseGraphNeedsRecompile);
+
+	/** Get the change id for this node. This change id is updated whenever the node is manipulated in a way that should force a recompile.*/
+	const FGuid& GetChangeId() const { return ChangeId; }
+
+	/** Set the change id for his node to an explicit value. This should only be called by internal code. */
+	void ForceChangeId(const FGuid& InId, bool bRaiseGraphNeedsRecompile);
+
+	FOnNodeVisualsChanged& OnVisualsChanged();
 protected:
 	virtual int32 CompileInputPin(class FHlslNiagaraTranslator *Translator, UEdGraphPin* Pin);
-	virtual bool IsValidPinToCompile(UEdGraphPin* Pin) const { return true; };
+	virtual bool IsValidPinToCompile(UEdGraphPin* Pin) const { return true; }; 
 
 	/** Route input parameter map to output parameter map if it exists. Note that before calling this function,
 		the input pins should have been visited already.*/
 	virtual void RouteParameterMapAroundMe(FNiagaraParameterMapHistoryBuilder& OutHistory, bool bRecursive);
+	
+	/** The current change identifier for this node. Used to sync status with UNiagaraScripts.*/
+	UPROPERTY()
+	FGuid ChangeId;
 
-	//ID for the most recent visitor. Allows faster tracking of what nodes have been visited.
-	int32 VisitID;
+	FOnNodeVisualsChanged VisualsChangedDelegate;
 };
diff --git a/Engine/Plugins/FX/Niagara/Source/NiagaraEditor/Public/NiagaraNodeAssignment.h b/Engine/Plugins/FX/Niagara/Source/NiagaraEditor/Public/NiagaraNodeAssignment.h
--- a/Engine/Plugins/FX/Niagara/Source/NiagaraEditor/Public/NiagaraNodeAssignment.h
+++ b/Engine/Plugins/FX/Niagara/Source/NiagaraEditor/Public/NiagaraNodeAssignment.h
@@ -11,24 +11,59 @@
 
 class UNiagaraScript;
 
-UCLASS(MinimalAPI)
-class UNiagaraNodeAssignment : public UNiagaraNodeFunctionCall
+UCLASS()
+class NIAGARAEDITOR_API UNiagaraNodeAssignment : public UNiagaraNodeFunctionCall
 {
 	GENERATED_BODY()
 
 public:
-	UPROPERTY(EditAnywhere, Category = Assignment)
-	FNiagaraVariable AssignmentTarget;
+	int32 NumTargets() const { return AssignmentTargets.Num(); }
 
-	UPROPERTY(EditAnywhere, Category = Assignment)
-	FString AssignmentDefaultValue;
+	int32 FindAssignmentTarget(const FName& InName);
+	int32 FindAssignmentTarget(const FName& InName, const FNiagaraTypeDefinition& TypeDef);
+
+	/** Set the assignment target and default value. This does not call RefreshFromExternalChanges, which you will need to do to update the internal graph. Returns true if values were changed.*/
+	bool SetAssignmentTarget(int32 Idx, const FNiagaraVariable& InVar, const FString* InDefaultValue = nullptr);
+	int32 AddAssignmentTarget(const FNiagaraVariable& InVar, const FString* InDefaultValue = nullptr);
+
+	FName GetAssignmentTargetName(int32 Idx) const { return AssignmentTargets[Idx].GetName(); }
+	bool SetAssignmentTargetName(int32 Idx, const FName& InName);
+	const FNiagaraVariable& GetAssignmentTarget(int32 Idx) const { return AssignmentTargets[Idx]; }
+
+	const TArray<FNiagaraVariable>& GetAssignmentTargets() const { return AssignmentTargets; }
+	const TArray<FString>& GetAssignmentDefaults() const { return AssignmentDefaultValues; }
+
+	void MergeUp();
+	void BuildParameterMenu(FMenuBuilder& MenuBuilder, ENiagaraScriptUsage InUsage, UNiagaraNodeOutput* InGraphOutputNode);
 
 	//~ Begin EdGraphNode Interface
+	virtual void PostLoad() override;
 	virtual FText GetNodeTitle(ENodeTitleType::Type TitleType) const override;
 	virtual void AllocateDefaultPins() override;
 
 	//~ UNiagaraNodeFunctionCall interface
 	virtual bool RefreshFromExternalChanges() override;
+	virtual void BuildParameterMapHistory(FNiagaraParameterMapHistoryBuilder& OutHistory, bool bRecursive = true) override;
+
+	void AddParameter(FNiagaraVariable InVar, FString InDefaultValue);
+	void RemoveParameter(const FNiagaraVariable& InVar);
+
+protected:
+
+	UPROPERTY()
+	FNiagaraVariable AssignmentTarget_DEPRECATED;
+
+	UPROPERTY()
+	FString AssignmentDefaultValue_DEPRECATED;
+
+	UPROPERTY(EditAnywhere, Category = Assignment)
+	TArray<FNiagaraVariable> AssignmentTargets;
+
+	UPROPERTY(EditAnywhere, Category = Assignment)
+	TArray<FString> AssignmentDefaultValues;
+
+	UPROPERTY()
+	FString OldFunctionCallName;
 
 private:
 	void GenerateScript();
diff --git a/Engine/Plugins/FX/Niagara/Source/NiagaraEditor/Public/NiagaraNodeCustomHlsl.h b/Engine/Plugins/FX/Niagara/Source/NiagaraEditor/Public/NiagaraNodeCustomHlsl.h
--- a/Engine/Plugins/FX/Niagara/Source/NiagaraEditor/Public/NiagaraNodeCustomHlsl.h
+++ b/Engine/Plugins/FX/Niagara/Source/NiagaraEditor/Public/NiagaraNodeCustomHlsl.h
@@ -27,11 +27,22 @@ public:
 	virtual void OnRenameNode(const FString& NewName) override;
 	virtual FLinearColor GetNodeTitleColor() const override;
 
-	FText GetHlslText() const;
-	void OnCustomHlslTextCommitted(const FText& InText, ETextCommit::Type InType);
+	FText NIAGARAEDITOR_API GetHlslText() const;
+	void NIAGARAEDITOR_API OnCustomHlslTextCommitted(const FText& InText, ETextCommit::Type InType);
+
+	bool GetTokens(TArray<FString>& OutTokens) const;
+
+	virtual void BuildParameterMapHistory(FNiagaraParameterMapHistoryBuilder& OutHistory, bool bRecursive = true) override;
+
+	// Replace items in the tokens array if they start with the src string or optionally src string and a namespace delimiter
+	static void ReplaceExactMatchTokens(TArray<FString>& Tokens, const FString& SrcString, const FString& ReplaceString, bool bAllowNamespaceSeparation);
+	static FNiagaraVariable StripVariableToBaseType(const FNiagaraVariable& InVar);
 
 #if WITH_EDITOR
 	virtual void PostEditChangeProperty(struct FPropertyChangedEvent& PropertyChangedEvent) override;
+
+	void InitAsCustomHlslDynamicInput(const FNiagaraTypeDefinition& OutputType);
+
 #endif
 protected:
 	virtual bool AllowDynamicPins() const override { return true; }
diff --git a/Engine/Plugins/FX/Niagara/Source/NiagaraEditor/Public/NiagaraNodeDataSetBase.h b/Engine/Plugins/FX/Niagara/Source/NiagaraEditor/Public/NiagaraNodeDataSetBase.h
--- a/Engine/Plugins/FX/Niagara/Source/NiagaraEditor/Public/NiagaraNodeDataSetBase.h
+++ b/Engine/Plugins/FX/Niagara/Source/NiagaraEditor/Public/NiagaraNodeDataSetBase.h
@@ -42,8 +42,9 @@ public:
 	//~ End UNiagaraNode Interface
 
 	bool InitializeFromStruct(const UStruct* ReferenceStruct);
-	
+		
 protected:
+	void AddParameterMapPins();
 	bool InitializeFromStructInternal(const UStruct* PayloadStruct);
 
 	bool IsSynchronizedWithStruct(bool bIgnoreConditionVar, FString* Issues, bool bLogIssues = true);
@@ -51,5 +52,7 @@ protected:
 	bool GetSupportedNiagaraTypeDef(const UProperty* Property, FNiagaraTypeDefinition& TypeDef);
 
 	static const FName ConditionVarName;
+	static const FName ParamMapInVarName;
+	static const FName ParamMapOutVarName;
 };
 
diff --git a/Engine/Plugins/FX/Niagara/Source/NiagaraEditor/Public/NiagaraNodeFunctionCall.h b/Engine/Plugins/FX/Niagara/Source/NiagaraEditor/Public/NiagaraNodeFunctionCall.h
--- a/Engine/Plugins/FX/Niagara/Source/NiagaraEditor/Public/NiagaraNodeFunctionCall.h
+++ b/Engine/Plugins/FX/Niagara/Source/NiagaraEditor/Public/NiagaraNodeFunctionCall.h
@@ -16,6 +16,9 @@ class UNiagaraNodeFunctionCall : public UNiagaraNodeWithDynamicPins
 {
 	GENERATED_BODY()
 
+public:
+	DECLARE_MULTICAST_DELEGATE(FOnInputsChanged);
+
 public:
 
 	UPROPERTY(EditAnywhere, Category = "Function")
@@ -47,6 +50,7 @@ public:
 	virtual ENiagaraNumericOutputTypeSelectionMode GetNumericOutputTypeSelectionMode() const override;
 	virtual bool CanAddToGraph(UNiagaraGraph* TargetGraph, FString& OutErrorMsg) const override;
 	virtual void SubsumeExternalDependencies(TMap<const UObject*, UObject*>& ExistingConversions) override;
+	virtual void GatherExternalDependencyIDs(ENiagaraScriptUsage InMasterUsage, const FGuid& InMasterUsageId, TArray<FGuid>& InReferencedIDs, TArray<UObject*>& InReferencedObjs) const override;
 	//End UNiagaraNode interface
 
 	//~ Begin EdGraphNode Interface
@@ -62,11 +66,17 @@ public:
 	virtual void BuildParameterMapHistory(FNiagaraParameterMapHistoryBuilder& OutHistory, bool bRecursive = true) override;
 
 	FString GetFunctionName() const { return FunctionDisplayName; }
+	UNiagaraGraph* GetCalledGraph() const;
+	ENiagaraScriptUsage GetCalledUsage() const;
 
 	/** Walk through the internal script graph for an ParameterMapGet nodes and see if any of them specify a default for VariableName.*/
-	UEdGraphPin* FindParameterMapDefaultValuePin(const FName VariableName);
+	UEdGraphPin* FindParameterMapDefaultValuePin(const FName VariableName, ENiagaraScriptUsage InParentUsage) const;
 
-	virtual bool CanRefreshFromModuleItem() const { return false; };
+	/** Tries to rename this function call to a new name.  The actual name that gets applied might be different due to conflicts with existing
+		nodes with the same name. */
+	void SuggestName(FString SuggestedName);
+
+	FOnInputsChanged& OnInputsChanged();
 
 protected:
 	virtual bool IsValidPinToCompile(UEdGraphPin* Pin) const;
@@ -77,7 +87,7 @@ protected:
 	virtual bool CanMovePin(const UEdGraphPin* Pin) const override { return false; }
 
 	/** Resets the node name based on the referenced script or signature. Guaranteed unique within a given graph instance.*/
-	void ComputeNodeName();
+	void ComputeNodeName(FString SuggestedName = FString());
 
 	void SetPinAutoGeneratedDefaultValue(UEdGraphPin& FunctionInputPin, UNiagaraNodeInput& FunctionScriptInputNode);
 
@@ -87,5 +97,7 @@ protected:
 
 	UPROPERTY()
 	FString FunctionDisplayName;
+
+	FOnInputsChanged OnInputsChangedDelegate;
 };
 
diff --git a/Engine/Plugins/FX/Niagara/Source/NiagaraEditor/Public/NiagaraNodeOp.h b/Engine/Plugins/FX/Niagara/Source/NiagaraEditor/Public/NiagaraNodeOp.h
--- a/Engine/Plugins/FX/Niagara/Source/NiagaraEditor/Public/NiagaraNodeOp.h
+++ b/Engine/Plugins/FX/Niagara/Source/NiagaraEditor/Public/NiagaraNodeOp.h
@@ -26,6 +26,10 @@ public:
 	UPROPERTY()
 	FName OpName;
 
+	//~ Begin UObject interface
+	virtual void PostLoad() override;
+	//~ End UObject interface
+
 	//~ Begin EdGraphNode Interface
 	virtual void AllocateDefaultPins() override;
 	virtual FText GetNodeTitle(ENodeTitleType::Type TitleType) const override;
@@ -42,55 +46,3 @@ public:
 
 
 
-#if 0
-
-class SGraphNodeNiagara : public SGraphNode
-{
-public:
-	SLATE_BEGIN_ARGS(SGraphNodeNiagara){}
-	SLATE_END_ARGS()
-
-		void Construct(const FArguments& InArgs, class UNiagaraNode* InNode);
-
-	//~ Begin SGraphNode Interface
-	virtual void CreatePinWidgets() override;
-	//~ End SGraphNode Interface
-
-	//~ Begin SNodePanel::SNode Interface
-	virtual void MoveTo(const FVector2D& NewPosition, FNodeSet& NodeFilter) override;
-	//~ End SNodePanel::SNode Interface
-
-	UNiagaraNode* GetNiagaraGraphNode() const { return NiagaraNode; }
-
-protected:
-	//~ Begin SGraphNode Interface
-	virtual void AddPin(const TSharedRef<SGraphPin>& PinToAdd) override;
-	virtual void CreateBelowPinControls(TSharedPtr<SVerticalBox> MainBox) override;
-	virtual void SetDefaultTitleAreaWidget(TSharedRef<SOverlay> DefaultTitleAreaWidget) override;
-	virtual TSharedRef<SWidget> CreateNodeContentArea() override;
-	//~ End SGraphNode Interface
-
-	/** Creates a preview viewport if necessary */
-	//TSharedRef<SWidget> CreatePreviewWidget();
-
-	/** Returns visibility of Expression Preview viewport */
-	//EVisibility ExpressionPreviewVisibility() const;
-
-	/** Show/hide Expression Preview */
-	//void OnExpressionPreviewChanged(const ECheckBoxState NewCheckedState);
-
-	/** hidden == unchecked, shown == checked */
-	//ECheckBoxState IsExpressionPreviewChecked() const;
-
-	/** Up when shown, down when hidden */
-	//const FSlateBrush* GetExpressionPreviewArrow() const;
-
-private:
-	/** Slate viewport for rendering preview via custom slate element */
-	//TSharedPtr<FPreviewViewport> PreviewViewport;
-
-	/** Cached material graph node pointer to avoid casting */
-	UNiagaraNode* NiagaraNode;
-};
-
-#endif
diff --git a/Engine/Plugins/FX/Niagara/Source/NiagaraEditor/Public/NiagaraNodeReadDataSet.h b/Engine/Plugins/FX/Niagara/Source/NiagaraEditor/Public/NiagaraNodeReadDataSet.h
--- a/Engine/Plugins/FX/Niagara/Source/NiagaraEditor/Public/NiagaraNodeReadDataSet.h
+++ b/Engine/Plugins/FX/Niagara/Source/NiagaraEditor/Public/NiagaraNodeReadDataSet.h
@@ -27,5 +27,6 @@ public:
 	//~End UNiagaraGraph Interface
 	virtual void Compile(class FHlslNiagaraTranslator* Translator, TArray<int32>& Outputs)override;
 
+	virtual void BuildParameterMapHistory(FNiagaraParameterMapHistoryBuilder& OutHistory, bool bRecursive = true) override;
 };
 
diff --git a/Engine/Plugins/FX/Niagara/Source/NiagaraEditor/Public/NiagaraNodeWithDynamicPins.h b/Engine/Plugins/FX/Niagara/Source/NiagaraEditor/Public/NiagaraNodeWithDynamicPins.h
--- a/Engine/Plugins/FX/Niagara/Source/NiagaraEditor/Public/NiagaraNodeWithDynamicPins.h
+++ b/Engine/Plugins/FX/Niagara/Source/NiagaraEditor/Public/NiagaraNodeWithDynamicPins.h
@@ -4,13 +4,15 @@
 
 #include "NiagaraNode.h"
 #include "Styling/SlateTypes.h"
+#include "NiagaraActions.h"
 #include "EdGraph/EdGraph.h"
-
+#include "EdGraph/EdGraphSchema.h"
 #include "NiagaraNodeWithDynamicPins.generated.h"
 
 class UEdGraphPin;
 class SNiagaraGraphPinAdd;
 
+DECLARE_DELEGATE_OneParam(FOnAddParameter, const FNiagaraVariable&);
 
 /** A base node for niagara nodes with pins which can be dynamically added and removed by the user. */
 UCLASS()
@@ -35,8 +37,11 @@ public:
 	/** Determine whether or not a Niagara type is supported for an Add Pin possibility.*/
 	virtual bool AllowNiagaraTypeForAddPin(const FNiagaraTypeDefinition& InType);
 
-	/** Used in tandem with SNiagaraGraphPinAdd to generate the menu for selecting the pin to add.*/
-	virtual TSharedRef<SWidget> GenerateAddPinMenu(const FString& InWorkingPinName, SNiagaraGraphPinAdd* InPin);
+	/** Used in to gather the actions for selecting the pin to add. */
+	virtual void CollectAddPinActions(FGraphActionListBuilderBase& OutActions, bool& bOutCreateRemainingActions, UEdGraphPin* Pin);
+
+	/** Request a new pin. */
+	void AddParameter(FNiagaraVariable Parameter, UEdGraphPin* AddPin);
 
 protected:
 	virtual bool AllowDynamicPins() const { return true; }
@@ -59,9 +64,6 @@ protected:
 	/** Called to determine if a pin can be moved by the user.*/
 	virtual bool CanMovePin(const UEdGraphPin* Pin) const;
 
-	/** Used GenerateAddPinMenu to build a list of supported types.*/
-	virtual void BuildTypeMenu(FMenuBuilder& InMenuBuilder, const FString& InWorkingName, SNiagaraGraphPinAdd* InPin);
-
 	/** Removes a pin from this node with a transaction. */
 	virtual void RemoveDynamicPin(UEdGraphPin* Pin);
 
@@ -74,7 +76,7 @@ private:
 
 	/** Called when a pin's name text is committed. */
 	void PinNameTextCommitted(const FText& Text, ETextCommit::Type CommitType, UEdGraphPin* Pin);
-
+	
 public:
 	/** The sub category for add pins. */
 	static const FName AddPinSubCategory;
diff --git a/Engine/Plugins/FX/Niagara/Source/NiagaraEditor/Public/NiagaraNodeWriteDataSet.h b/Engine/Plugins/FX/Niagara/Source/NiagaraEditor/Public/NiagaraNodeWriteDataSet.h
--- a/Engine/Plugins/FX/Niagara/Source/NiagaraEditor/Public/NiagaraNodeWriteDataSet.h
+++ b/Engine/Plugins/FX/Niagara/Source/NiagaraEditor/Public/NiagaraNodeWriteDataSet.h
@@ -33,7 +33,8 @@ public:
 	//~ End EdGraphNode Interface
 
 	virtual void Compile(class FHlslNiagaraTranslator* Translator, TArray<int32>& Outputs)override;
-	
+
+	virtual void BuildParameterMapHistory(FNiagaraParameterMapHistoryBuilder& OutHistory, bool bRecursive = true) override;
 	virtual TSharedPtr<SGraphNode> CreateVisualWidget() override;
 protected:
 	void AddConditionPin(int32 PinIndex = -1);
diff --git a/Engine/Plugins/FX/Niagara/Source/NiagaraEditor/Public/NiagaraParameterMapHistory.h b/Engine/Plugins/FX/Niagara/Source/NiagaraEditor/Public/NiagaraParameterMapHistory.h
--- a/Engine/Plugins/FX/Niagara/Source/NiagaraEditor/Public/NiagaraParameterMapHistory.h
+++ b/Engine/Plugins/FX/Niagara/Source/NiagaraEditor/Public/NiagaraParameterMapHistory.h
@@ -25,7 +25,6 @@ public:
 
 	TArray<FNiagaraVariable> VariablesWithOriginalAliasesIntact;
 
-
 	/** Used parameter collections identified during the traversal. TODO: Need to ensure these cannot be GCd if the asset is deleted while it's being used in an in flight compilation. */
 	TArray<UNiagaraParameterCollection*> ParameterCollections;
 	/** Cached off contents of used parameter collections, in case they change during threaded compilation. */
@@ -45,6 +44,12 @@ public:
 	/** List of pins that manipulated the parameter map from input to output. */
 	TArray<const UEdGraphPin*> MapPinHistory;
 
+	/** List of nodes that manipulated the parameter map from input to output.*/
+	TArray<class UNiagaraNode*> MapNodeVisitations;
+
+	/** For each node in MapNodeVisitations, record the start index and end index of variables added within the body of the node, i.e. a Get node will record just the values it pulls out directly. A function call, however, will record all sub-graph traversals.*/
+	TArray<TTuple<uint32, uint32> > MapNodeVariableMetaData;
+
 	/** List of emitter namespaces encountered as this parameter map was built.*/
 	TArray<FString> EmitterNamespacesEncountered;
 	
@@ -52,11 +57,16 @@ public:
 	* Called in a depth-first traversal to identify a given Niagara Parameter Map pin that was touched during traversal.
 	*/
 	int32 RegisterParameterMapPin(const UEdGraphPin* Pin);
-		
+
+	uint32 BeginNodeVisitation(UNiagaraNode* Node);
+	void EndNodeVisitation(uint32 IndexFromBeginNode);
+
 	/**
 	* Find a variable by name with no concern for type.
 	*/
-	int32 FindVariableByName(const FName& VariableName);
+	int32 FindVariableByName(const FName& VariableName, bool bAllowPartialMatch = false);
+
+	static int32 SearchArrayForPartialNameMatch(const TArray<FNiagaraVariable>& Variables, const FName& VariableName);
 	
 	/**
 	* Find a variable by both name and type. 
@@ -80,11 +90,18 @@ public:
 	const UEdGraphPin* GetFinalPin() const;
 
 
+
+	/** Get the first namespace entry for this variable. Optionally includes the trailing period.*/
+	static FString GetNamespace(const FNiagaraVariable& InVar, bool bIncludeDelimiter = true);
+
+
 	/**
 	* Use the input alias map to resolve any aliases in this input variable name.
 	*/
 	static FNiagaraVariable ResolveAliases(const FNiagaraVariable& InVar, const TMap<FString, FString>& InAliases, const TCHAR* InJoinSeparator);
-		
+
+	static FName ResolveEmitterAlias(const FName& InName, const FString& InAlias);
+
 	/**
 	* Remove the Particles namespace if it exists.
 	*/
@@ -111,10 +128,13 @@ public:
 	* for Emitter and System script types and will be specialized to the name of that specific emitter.
 	*/
 	static bool IsAliasedEmitterParameter(const FNiagaraVariable& InVar);
+	static bool IsAliasedEmitterParameter(const FString& InVarName);
 	/** Is this parameter in the special "System" namespace?*/
 	static bool IsSystemParameter(const FNiagaraVariable& InVar);
 	/** Is this parameter in the special "Engine" namespace?*/
 	static bool IsEngineParameter(const FNiagaraVariable& InVar);
+	/*** Is per instance engine parameter. */
+	static bool IsPerInstanceEngineParameter(const FNiagaraVariable& InVar, const FString& EmitterAlias);
 	static bool IsUserParameter(const FNiagaraVariable& InVar);
 	static bool IsRapidIterationParameter(const FNiagaraVariable& InVar);
 	static bool TryGetEmitterAndFunctionCallNamesFromRapidIterationParameter(const FNiagaraVariable& InVar, FString& EmitterName, FString& FunctionCallName);
@@ -126,24 +146,29 @@ public:
 	static bool IsInNamespace(const FNiagaraVariable& InVar, const FString& Namespace);
 
 	/** Given the script type and flags, should we treat the System namespace as read-only?*/
-	static bool IsSystemNamespaceReadOnly(UNiagaraScript* InScript);
+	static bool IsSystemNamespaceReadOnly(const UNiagaraScript* InScript);
 
 	/** Given the script type and flags, should we treat the Emitter aliased namespace as read-only?*/
-	static bool IsEmitterNamespaceReadOnly(UNiagaraScript* InScript);
+	static bool IsEmitterNamespaceReadOnly(const UNiagaraScript* InScript);
+
+	static void GetValidNamespacesForReading(const UNiagaraScript* InScript, TArray<FString>& OutputNamespaces);
+	static void GetValidNamespacesForReading(ENiagaraScriptUsage InScriptUsage, TArray<FString>& OutputNamespaces);
+	static bool IsValidNamespaceForReading(ENiagaraScriptUsage InScriptUsage, FString Namespace);
 
 	/** Called to determine if a given variable should be output from a script. It is not static as it requires the overall context to include emitter namespaces visited for system scripts.*/
-	bool IsPrimaryDataSetOutput(const FNiagaraVariable& InVar, UNiagaraScript* InScript);
-	bool IsPrimaryDataSetOutput(const FNiagaraVariable& InVar, ENiagaraScriptUsage InUsage);
+	bool IsPrimaryDataSetOutput(const FNiagaraVariable& InVar, const UNiagaraScript* InScript, bool bAllowDataInterfaces = false) const;
+	bool IsPrimaryDataSetOutput(const FNiagaraVariable& InVar, ENiagaraScriptUsage InUsage, bool bAllowDataInterfaces = false) const;
 
 	/** Are we required to export this variable as an external constant?*/
-	bool IsExportableExternalConstant(const FNiagaraVariable& InVar, UNiagaraScript* InScript);
+	bool IsExportableExternalConstant(const FNiagaraVariable& InVar, const UNiagaraScript* InScript);
 
 	/** Does this variable belong in a namespace that needs to come in as an external constant to this script?*/
-	static bool IsExternalConstantNamespace(const FNiagaraVariable& InVar, UNiagaraScript* InScript);
+	static bool IsExternalConstantNamespace(const FNiagaraVariable& InVar, const UNiagaraScript* InScript);
+	static bool IsExternalConstantNamespace(const FNiagaraVariable& InVar, ENiagaraScriptUsage InUsage, int32 InUsageBitmask);
 	
 	/** Take a non-namespaced variable and move it to an appropriate external constant namespace for this script type.*/
-	static FNiagaraVariable MoveToExternalConstantNamespaceVariable(const FNiagaraVariable& InVar, UNiagaraScript* InScript);
-	
+	static FNiagaraVariable MoveToExternalConstantNamespaceVariable(const FNiagaraVariable& InVar, const UNiagaraScript* InScript);
+	static FNiagaraVariable MoveToExternalConstantNamespaceVariable(const FNiagaraVariable& InVar, ENiagaraScriptUsage InUsage);	
 	/**
 	* Does this parameter start with the "Particles" namespace?
 	*/
@@ -169,7 +194,7 @@ public:
 	const FNiagaraVariableMetaData* GetMetaData(int32 VarIdx) const;
 	FNiagaraVariableMetaData* GetMetaData(int32 VarIdx);
 
-	static FNiagaraVariable ConvertVariableToRapidIterationConstantName(FNiagaraVariable InVar, const TCHAR* InEmitterName);
+	static FNiagaraVariable ConvertVariableToRapidIterationConstantName(FNiagaraVariable InVar, const TCHAR* InEmitterName, ENiagaraScriptUsage InUsage);
 
 	/**
 	If this is variable is a parameter in one of our tracked collections, return it.
@@ -178,6 +203,7 @@ public:
 	*/
 	UNiagaraParameterCollection* IsParameterCollectionParameter(FNiagaraVariable& InVar, bool& bMissingParameter);
 
+	bool ShouldIgnoreVariableDefault(const FNiagaraVariable& Var)const;
 };
 
 
@@ -204,26 +230,32 @@ public:
 	/** Important. Must be called for each routing of the parameter map. This feeds the list used by TraceParameterMapOutputPin.*/
 	int32 RegisterParameterMapPin(int32 WhichParameterMap, const UEdGraphPin* Pin);
 
+	uint32 BeginNodeVisitation(int32 WhichParameterMap, class UNiagaraNode* Node);
+	void EndNodeVisitation(int32 WhichParameterMap, uint32 IndexFromBeginNode);
+
 	/** Trace back a pin to whom it was connected to to find the current parameter map to use.*/
 	int32 TraceParameterMapOutputPin(const UEdGraphPin* OutputPin);
 
-	void BeginTranslation(UNiagaraScript* Script);
+	void BeginTranslation(const UNiagaraScript* Script);
 
-	void EndTranslation(UNiagaraScript* Script);
+	void EndTranslation(const UNiagaraScript* Script);
 
-	void BeginTranslation(UNiagaraEmitter* Emitter);
+	void BeginTranslation(const FString& EmitterUniqueName);
+	void EndTranslation(const FString& EmitterUniqueName);
 
-	void EndTranslation(UNiagaraEmitter* Emitter);
+	void BeginTranslation(const UNiagaraEmitter* Emitter);
+
+	void EndTranslation(const UNiagaraEmitter* Emitter);
 
 	/**
 	* Record that we have entered a new function scope.
 	*/
-	void EnterFunction(const FString& InNodeName, class UNiagaraScript* InScript, class UNiagaraNode* Node);
+	void EnterFunction(const FString& InNodeName, const class UNiagaraScript* InScript, class UNiagaraNode* Node);
 
 	/**
 	* Record that we have exited a function scope.
 	*/
-	void ExitFunction(const FString& InNodeName, class UNiagaraScript* InScript, class UNiagaraNode* Node);
+	void ExitFunction(const FString& InNodeName, const class UNiagaraScript* InScript, class UNiagaraNode* Node);
 
 	/**
 	* Record that we have entered an emitter scope.
@@ -239,7 +271,7 @@ public:
 	* Use the current alias map to resolve any aliases in this input variable name.
 	*/
 	FNiagaraVariable ResolveAliases(const FNiagaraVariable& InVar);
-
+	
 	/**
 	* Has RegisterNodeVisitation been called yet on the owning node of this pin?
 	*/
@@ -260,11 +292,18 @@ public:
 	* Record that a pin writes to the parameter map. The pin name is expected to be the namespaced parameter map version of the name. If any aliases are in place, they are removed.
 	*/
 	int32 HandleVariableWrite(int32 ParameterMapIndex, const UEdGraphPin* InPin);
+	
+	/**
+	* Record that a variable write to the parameter map. The var name is expected to be the namespaced parameter map version of the name. If any aliases are in place, they are removed.
+	*/
+	int32 HandleVariableWrite(int32 ParameterMapIndex, const FNiagaraVariable& Var);
 
 	/**
 	* Record that a pin reads from the parameter map. The pin name is expected to be the namespaced parameter map version of the name. If any aliases are in place, they are removed.
 	*/
 	int32 HandleVariableRead(int32 ParameterMapIndex, const UEdGraphPin* InPin, bool RegisterReadsAsVariables, const UEdGraphPin* InDefaultPin, bool& OutUsedDefault);
+
+	int32 HandleExternalVariableRead(int32 ParamMapIdx, const FName& InVarName);
 	
 	/**
 	* Get the string that the "Module" namespace maps to currently (if it exists)
@@ -292,6 +331,8 @@ public:
 	void EnableScriptWhitelist(bool bInEnable, ENiagaraScriptUsage InScriptType);
 
 	ENiagaraScriptUsage GetCurrentUsageContext()const;
+	ENiagaraScriptUsage GetBaseUsageContext()const;
+	bool ContextContains(ENiagaraScriptUsage InUsage) const;
 
 	bool GetIgnoreDisabled() const { return bIgnoreDisabled; }
 	void SetIgnoreDisabled(bool bInIgnore) { bIgnoreDisabled = bInIgnore; }
@@ -299,6 +340,10 @@ public:
 	bool IsInEncounteredFunctionNamespace(FNiagaraVariable& InVar) const;
 	bool IsInEncounteredEmitterNamespace(FNiagaraVariable& InVar) const;
 
+	/** Register any user or other external variables that could possibly be encountered but may not be declared explicitly. */
+	void RegisterEncounterableVariables(const TArray<FNiagaraVariable>& Variables);
+	const TArray<FNiagaraVariable>& GetEncounterableVariables() const {	return EncounterableExternalVariables;}
+
 protected:
 	/**
 	* Generate the internal alias map from the current traversal state.
@@ -336,4 +381,6 @@ protected:
 
 	/** Whether or not to ignore disabled nodes.*/
 	bool bIgnoreDisabled;
+
+	TArray<FNiagaraVariable> EncounterableExternalVariables;
 };
\ No newline at end of file
diff --git a/Engine/Plugins/FX/Niagara/Source/NiagaraEditor/Public/NiagaraScriptSource.h b/Engine/Plugins/FX/Niagara/Source/NiagaraEditor/Public/NiagaraScriptSource.h
--- a/Engine/Plugins/FX/Niagara/Source/NiagaraEditor/Public/NiagaraScriptSource.h
+++ b/Engine/Plugins/FX/Niagara/Source/NiagaraEditor/Public/NiagaraScriptSource.h
@@ -24,9 +24,9 @@ class UNiagaraScriptSource : public UNiagaraScriptSourceBase
 	virtual void PostLoad() override;
 
 	// UNiagaraScriptSourceBase interface.
-	virtual ENiagaraScriptCompileStatus Compile(UNiagaraScript* ScriptOwner, FString& OutGraphLevelErrorMessages) override;
+	//virtual ENiagaraScriptCompileStatus Compile(UNiagaraScript* ScriptOwner, FString& OutGraphLevelErrorMessages) override;
 	virtual bool IsSynchronized(const FGuid& InChangeId) override;
-	virtual void MarkNotSynchronized() override;
+	virtual void MarkNotSynchronized(FString Reason) override;
 
 	virtual UNiagaraScriptSourceBase* MakeRecursiveDeepCopy(UObject* DestOuter, TMap<const UObject*, UObject*>& ExistingConversions) const override;
 
@@ -35,27 +35,17 @@ class UNiagaraScriptSource : public UNiagaraScriptSourceBase
 
 	virtual FGuid GetChangeID();
 
-	virtual bool IsPreCompiled() const override;
-	virtual void PreCompile(UNiagaraEmitter* Emitter, bool bClearErrors = true) override;
-	virtual bool GatherPreCompiledVariables(const FString& InNamespaceFilter, TArray<FNiagaraVariable>& OutVars) override;
-	virtual void PostCompile() override;
-	virtual bool PostLoadFromEmitter(UNiagaraEmitter& OwningEmitter) override;
+	virtual void ComputeVMCompilationId(struct FNiagaraVMExecutableDataId& Id, ENiagaraScriptUsage InUsage, const FGuid& InUsageId) const override;
 
-	NIAGARAEDITOR_API virtual bool AddModuleIfMissing(FString ModulePath, ENiagaraScriptUsage Usage, bool& bOutFoundModule)override;
+	virtual void PostLoadFromEmitter(UNiagaraEmitter& OwningEmitter) override;
 
-	TArray<FNiagaraParameterMapHistory>& GetPrecomputedHistories() { return PrecompiledHistories; }
-	class UNiagaraGraph* GetPrecomputedNodeGraph() { return NodeGraphDeepCopy; }
+	NIAGARAEDITOR_API virtual bool AddModuleIfMissing(FString ModulePath, ENiagaraScriptUsage Usage, bool& bOutFoundModule)override;
 
+	virtual void CleanUpOldAndInitializeNewRapidIterationParameters(FString UniqueEmitterName, ENiagaraScriptUsage ScriptUsage, FGuid ScriptUsageId, FNiagaraParameterStore& RapidIterationParameters) const override;
+	virtual void InvalidateCachedCompileIds() override;
 private:
 	void OnGraphChanged(const FEdGraphEditAction &Action);
 	void OnGraphDataInterfaceChanged();
 
-private:
-	bool bIsPrecompiled;
-	TArray<FNiagaraParameterMapHistory> PrecompiledHistories;
 
-	void UpgradeForScriptRapidIterationVariables();
-
-	UPROPERTY()
-	class UNiagaraGraph* NodeGraphDeepCopy;
 };
diff --git a/Engine/Plugins/FX/Niagara/Source/NiagaraEditor/Public/NiagaraStackEditorData.h b/Engine/Plugins/FX/Niagara/Source/NiagaraEditor/Public/NiagaraStackEditorData.h
--- a/Engine/Plugins/FX/Niagara/Source/NiagaraEditor/Public/NiagaraStackEditorData.h
+++ b/Engine/Plugins/FX/Niagara/Source/NiagaraEditor/Public/NiagaraStackEditorData.h
@@ -3,6 +3,7 @@
 #pragma once
 
 #include "NiagaraStackEditorData.generated.h"
+struct FStackIssue;
 
 /** Editor only UI data for emitters. */
 UCLASS()
@@ -11,19 +12,6 @@ class UNiagaraStackEditorData : public UObject
 	GENERATED_BODY()
 
 public:
-	/*
-	 * Gets whether or not a module has been pinned.
-	 * @param ModuleInputKey A unique key for the module. 
-	 */
-	bool GetModuleInputIsPinned(const FString& ModuleInputKey) const;
-
-	/*
-	 * Sets whether or not a module is pinned.
-	 * @param ModuleInputKey A unique key for the module.
-	 * @param bIsPinned Whether or not the module is pinned.
-	 */
-	void SetModuleInputIsPinned(const FString& ModuleInputKey, bool bIsPinned);
-
 	/*
 	* Gets whether or not a module has a rename pending.
 	* @param ModuleInputKey A unique key for the module.
@@ -51,14 +39,68 @@ public:
 	 */
 	void SetStackEntryIsExpanded(const FString& StackEntryKey, bool bIsExpanded);
 
-private:
+	/*
+	* Gets whether or not a stack item is showing advanced items.
+	* @param StackItemKey A unique key for the entry.
+	* @param bIsExpandedDefault The default value to return if the expanded state hasn't been set for the stack entry.
+	*/
+	bool GetStackItemShowAdvanced(const FString& StackEntryKey, bool bShowAdvancedDefault) const;
 
-	UPROPERTY()
-	TMap<FString, bool> ModuleInputKeyToPinnedMap;
+	/*
+	* Sets whether or not a stack entry is Expanded.
+	* @param StackEntryKey A unique key for the entry.
+	* @param bIsExpanded Whether or not the entry is expanded.
+	*/
+	void SetStackItemShowAdvanced(const FString& StackEntryKey, bool bShowAdanced);
+
+	/* Gets whether or not all advanced items should be shown in the stack. */
+	bool GetShowAllAdvanced() const;
+
+	/* Sets whether or not all advanced items should be shown in the stack. */
+	void SetShowAllAdvanced(bool bInShowAllAdvanced);
+
+	/* Gets whether or not item outputs should be shown in the stack. */
+	bool GetShowOutputs() const;
+
+	/* Sets whether or not item outputs should be shown in the stack. */
+	void SetShowOutputs(bool bInShowOutputs);
+
+	/* Gets the last scroll position for the associated stack. */
+	double GetLastScrollPosition() const;
 
+	/* Sets the last scroll position for the associated stack. */
+	void SetLastScrollPosition(double InLastScrollPosition);
+
+	/*
+	* @param Issue the issue to be dismissed (not fixed).
+	*/
+	void DismissStackIssue(FName IssueId);
+
+	/* Restores all the dismissed issues so that the user can see them and choose what to do. */
+	NIAGARAEDITOR_API void UndismissAllIssues();
+
+	/* Gets a reference to the dismissed stack issue array */
+	NIAGARAEDITOR_API const TArray<FName>& GetDismissedStackIssueIds();
+
+private:
 	UPROPERTY()
 	TMap<FString, bool> ModuleInputKeyToRenamePendingMap;
 
 	UPROPERTY()
 	TMap<FString, bool> StackEntryKeyToExpandedMap;
+
+	UPROPERTY()
+	TMap<FString, bool> StackItemKeyToShowAdvancedMap;
+
+	UPROPERTY()
+	bool bShowAllAdvanced;
+
+	UPROPERTY()
+	bool bShowOutputs;
+
+	UPROPERTY()
+	double LastScrollPosition;
+
+	UPROPERTY()
+	TArray<FName> DismissedStackIssueIds;
 };
\ No newline at end of file
diff --git a/Engine/Plugins/FX/Niagara/Source/NiagaraEditor/Public/SNiagaraParameterEditor.h b/Engine/Plugins/FX/Niagara/Source/NiagaraEditor/Public/SNiagaraParameterEditor.h
--- a/Engine/Plugins/FX/Niagara/Source/NiagaraEditor/Public/SNiagaraParameterEditor.h
+++ b/Engine/Plugins/FX/Niagara/Source/NiagaraEditor/Public/SNiagaraParameterEditor.h
@@ -2,6 +2,7 @@
 
 #pragma once
 #include "Widgets/SCompoundWidget.h"
+#include "Widgets/DeclarativeSyntaxSupport.h"
 #include "UObject/StructOnScope.h"
 
 /** Base class for inline parameter editors. These editors are expected to maintain an internal value which
@@ -12,6 +13,18 @@ public:
 	DECLARE_DELEGATE(FOnValueChange);
 
 public:
+	SLATE_BEGIN_ARGS(SNiagaraParameterEditor) 
+		: _HAlign(HAlign_Left)
+		, _VAlign(VAlign_Center)
+	{ }
+		SLATE_ARGUMENT(EHorizontalAlignment, HAlign)
+		SLATE_ARGUMENT(EVerticalAlignment, VAlign)
+		SLATE_ARGUMENT(TOptional<float>, MinimumDesiredWidth)
+		SLATE_ARGUMENT(TOptional<float>, MaximumDesiredWidth)
+	SLATE_END_ARGS();
+
+	void Construct(const FArguments& InArgs);
+
 	/** Updates the internal value of the widget from a struct. */
 	virtual void UpdateInternalValueFromStruct(TSharedRef<FStructOnScope> Struct) = 0;
 
@@ -32,6 +45,21 @@ public:
 	/** Sets the OnValueChanged delegate which is run when the internal value changes. */
 	void SetOnValueChanged(FOnValueChange InOnValueChanged);
 
+	/** Gets an optional minimum desired width for this parameter editor. */
+	const TOptional<float>& GetMinimumDesiredWidth() const;
+
+	/** Gets an optional maximum desired width for this parameter editor. */
+	const TOptional<float>& GetMaximumDesiredWidth() const;
+
+	/** Gets the desired horizontal alignment of this parameter editor for it's parent container. */
+	EHorizontalAlignment GetHorizontalAlignment() const;
+
+	/** Gets the desired horizontal alignment of this parameter editor for it's parent container. */
+	EVerticalAlignment GetVerticalAlignment() const;
+
+	//~ SWidget interface
+	virtual FVector2D ComputeDesiredSize(float LayoutScaleMultiplier) const override;
+
 protected:
 	/** Sets whether this is currently the exclusive editor of this parameter, meaning that the corresponding details view
 		should not be updated.  This hack is necessary because the details view closes all color pickers when
@@ -47,6 +75,9 @@ protected:
 	/** Executes the OnValueChanged delegate. */
 	void ExecuteOnValueChanged();
 
+protected:
+	static const float DefaultInputSize;
+
 private:
 	/** Whether this is currently the exclusive editor of this parameter, meaning that the corresponding details view
 		should not be updated.  This hack is necessary because the details view closes all color pickers when
@@ -61,4 +92,16 @@ private:
 
 	/** A delegate which is executed when the internal value changes. */
 	FOnValueChange OnValueChanged;
-};
+
+	/** The minimum desired width of this parameter editor. */
+	TOptional<float> MinimumDesiredWidth;
+
+	/** The maximum desired width of this parameter editor. */
+	TOptional<float> MaximumDesiredWidth;
+
+	/** The desired horizontal alignment of this parameter editor for it's parent container. */
+	EHorizontalAlignment HorizontalAlignment;
+
+	/** Sets the desired horizontal alignment of this parameter editor for it's parent container. */
+	EVerticalAlignment VerticalAlignment;
+};
\ No newline at end of file
diff --git a/Engine/Plugins/FX/Niagara/Source/NiagaraEditor/Public/ViewModels/NiagaraEmitterHandleViewModel.h b/Engine/Plugins/FX/Niagara/Source/NiagaraEditor/Public/ViewModels/NiagaraEmitterHandleViewModel.h
new file mode 100644
--- /dev/null
+++ b/Engine/Plugins/FX/Niagara/Source/NiagaraEditor/Public/ViewModels/NiagaraEmitterHandleViewModel.h
@@ -0,0 +1,109 @@
+// Copyright 1998-2018 Epic Games, Inc. All Rights Reserved.
+
+#pragma once
+
+#include "Templates/SharedPointer.h"
+#include "Delegates/Delegate.h"
+#include "Types/SlateEnums.h"
+#include "Styling/SlateTypes.h"
+#include "Layout/Visibility.h"
+
+class UNiagaraSystem;
+struct FNiagaraEmitterHandle;
+struct FNiagaraEmitterInstance;
+class FNiagaraEmitterViewModel;
+
+/** The view model for the FNiagaraEmitterEditorWidget. */
+class FNiagaraEmitterHandleViewModel : public TSharedFromThis<FNiagaraEmitterHandleViewModel>
+{
+public:
+	DECLARE_MULTICAST_DELEGATE(FOnPropertyChanged);
+	DECLARE_MULTICAST_DELEGATE(FOnNameChanged);
+public:
+	/** Creates a new emitter editor view model with the supplied emitter handle and simulation. */
+	FNiagaraEmitterHandleViewModel(FNiagaraEmitterHandle* InEmitterHandle, TWeakPtr<FNiagaraEmitterInstance> InSimulation, UNiagaraSystem& InOwningSystem);
+	
+	~FNiagaraEmitterHandleViewModel();
+
+	/** Reuses a the emitter editor view model with the supplied emitter handle and simulation.*/
+	bool Set(FNiagaraEmitterHandle* InEmitterHandle, TWeakPtr<FNiagaraEmitterInstance> InSimulation, UNiagaraSystem& InOwningSystem);
+
+	/** Sets the emitter handle.*/
+	void SetEmitterHandle(FNiagaraEmitterHandle* InEmitterHandle);
+
+	/** Sets the simulation for the emitter this handle references. */
+	void SetSimulation(TWeakPtr<FNiagaraEmitterInstance> InSimulation);
+
+	/** Gets the id of the emitter handle. */
+	FGuid GetId() const;
+	FText GetIdText() const;
+
+	/** Gets the name of the emitter handle. */
+	FName GetName() const;
+
+	/** Sets the name of the emitter handle. */
+	void SetName(FName InName);
+
+	/** Gets the text representation of the emitter handle name. */
+	NIAGARAEDITOR_API FText GetNameText() const;
+
+	/** Called when the contents of the name text control is committed. */
+	NIAGARAEDITOR_API void OnNameTextComitted(const FText& InText, ETextCommit::Type CommitInfo);
+
+	/** Prevent invalid name being set on emitter.*/
+	NIAGARAEDITOR_API bool VerifyNameTextChanged(const FText& NewText, FText& OutErrorMessage);
+
+	/** Called to get the error state of the emitter handle.*/
+	FText GetErrorText() const;
+	EVisibility GetErrorTextVisibility() const;
+	FSlateColor GetErrorTextColor() const;
+
+	/** Called to get the synch state of the emitter handle to its source.*/
+	bool IsSynchronized() const;
+
+	/** Gets whether or not this emitter handle is enabled. */
+	NIAGARAEDITOR_API bool GetIsEnabled() const;
+
+	/** Sets whether or not this emitter handle is enabled. */
+	NIAGARAEDITOR_API void SetIsEnabled(bool bInIsEnabled);
+
+	/** Gets the check state for the is enabled check box. */
+	NIAGARAEDITOR_API ECheckBoxState GetIsEnabledCheckState() const;
+
+	/** Called when the check state of the enabled check box changes. */
+	NIAGARAEDITOR_API void OnIsEnabledCheckStateChanged(ECheckBoxState InCheckState);
+
+	/** Gets the emitter handled being view and edited by this view model. */
+	NIAGARAEDITOR_API FNiagaraEmitterHandle* GetEmitterHandle();
+
+	/** Gets the view model for the emitter this handle references. */
+	NIAGARAEDITOR_API TSharedPtr<FNiagaraEmitterViewModel> GetEmitterViewModel();
+
+
+	/** Opens the source emitter in a stand alone asset editor. */
+	void OpenSourceEmitter();
+
+	/** Gets a multicast delegate which is called any time a property on the handle changes. */
+	FOnPropertyChanged& OnPropertyChanged();
+
+	/** Gets a multicast delegate which is called any time this handle is renamed. */
+	FOnNameChanged& OnNameChanged();
+
+	void Cleanup();
+
+private:
+	/** The emitter handle being displayed and edited by this view model. */
+	FNiagaraEmitterHandle* EmitterHandle;
+
+	/** The System which owns the handled being displayed and edited by this view model. */
+	UNiagaraSystem& OwningSystem;
+
+	/** The view model for emitter this handle references. */
+	TSharedPtr<FNiagaraEmitterViewModel> EmitterViewModel;
+
+	/** A multicast delegate which is called any time a property on the handle changes. */
+	FOnPropertyChanged OnPropertyChangedDelegate;
+
+	/** A multicast delegate which is called any time this emitter handle is renamed. */
+	FOnNameChanged OnNameChangedDelegate;
+};
