Copying //UE4/Dev-Niagara to //UE4/Dev-Main (Source: //UE4/Dev-Niagara @ 4074996)

diff --git a/Engine/Plugins/FX/Niagara/Source/NiagaraEditor/Public/ViewModels/NiagaraSystemViewModel.h b/Engine/Plugins/FX/Niagara/Source/NiagaraEditor/Public/ViewModels/NiagaraSystemViewModel.h
new file mode 100644
--- /dev/null
+++ b/Engine/Plugins/FX/Niagara/Source/NiagaraEditor/Public/ViewModels/NiagaraSystemViewModel.h
@@ -0,0 +1,471 @@
+// Copyright 1998-2018 Epic Games, Inc. All Rights Reserved.
+
+#pragma once
+
+#include "IMovieScenePlayer.h"
+
+#include "EditorUndoClient.h"
+#include "UObject/GCObject.h"
+#include "NiagaraCurveOwner.h"
+#include "ViewModels/TNiagaraViewModelManager.h"
+#include "ISequencer.h"
+
+#include "TickableEditorObject.h"
+#include "ISequencerModule.h"
+#include "UObject/ObjectKey.h"
+
+class UNiagaraSystem;
+class UNiagaraComponent;
+class UNiagaraSequence;
+class UMovieSceneNiagaraEmitterTrack;
+struct FNiagaraVariable;
+struct FNiagaraEmitterHandle;
+class FNiagaraEmitterHandleViewModel;
+class FNiagaraSystemScriptViewModel;
+class FNiagaraSystemInstance;
+class ISequencer;
+struct FAssetData;
+class UNiagaraSystemEditorData;
+struct FRichCurve;
+class UNiagaraEditorSettings;
+struct FNiagaraParameterStore;
+struct FEdGraphEditAction;
+
+/** Defines different editing modes for this system view model. */
+enum class ENiagaraSystemViewModelEditMode
+{
+	/** A system asset is being edited.  This assumes that emitters should be inheriting from a base version and that emitter editing will be restricted. */
+	SystemAsset,
+	/** An emitter asset is being edited.  In this mode the system scripts will not be editable and all emitter editing options are available. */
+	EmitterAsset,
+};
+
+/** Defines options for the niagara System view model */
+struct FNiagaraSystemViewModelOptions
+{
+	FNiagaraSystemViewModelOptions();
+
+	/** Whether or not the user can edit emitters from the timeline. */
+	bool bCanModifyEmittersFromTimeline;
+
+	/** A delegate which is used to generate the content for the add menu in sequencer. */
+	FOnGetAddMenuContent OnGetSequencerAddMenuContent;
+
+	/** Whether or not we use the system's execution state to drive when we reset the timeline*/
+	bool bUseSystemExecStateForTimelineReset;
+
+	/** Whether or not the system represented by this view model can be automatically compiled.  True by default. */
+	bool bCanAutoCompile;
+
+	/** Whether or not the system represented by this view model can be simulated. True by default. */
+	bool bCanSimulate;
+
+	/** Gets the current editing mode for this system. */
+	ENiagaraSystemViewModelEditMode EditMode;
+};
+
+/** A view model for viewing and editing a UNiagaraSystem. */
+class FNiagaraSystemViewModel 
+	: public TSharedFromThis<FNiagaraSystemViewModel>
+	, public FGCObject
+	, public FEditorUndoClient
+	, public FTickableEditorObject
+	, public TNiagaraViewModelManager<UNiagaraSystem, FNiagaraSystemViewModel>
+{
+public:
+	DECLARE_MULTICAST_DELEGATE(FOnEmitterHandleViewModelsChanged);
+
+	DECLARE_MULTICAST_DELEGATE(FOnCurveOwnerChanged);
+
+	DECLARE_MULTICAST_DELEGATE(FOnSelectedEmitterHandlesChanged);
+
+	DECLARE_MULTICAST_DELEGATE(FOnPostSequencerTimeChange)
+
+	DECLARE_MULTICAST_DELEGATE(FOnSystemCompiled);
+
+	DECLARE_MULTICAST_DELEGATE(FOnPinnedEmittersChanged);
+
+	DECLARE_MULTICAST_DELEGATE(FOnPinnedCurvesChanged);
+
+public:
+	/** Creates a new view model with the supplied System and System instance. */
+	FNiagaraSystemViewModel(UNiagaraSystem& InSystem, FNiagaraSystemViewModelOptions InOptions);
+
+	~FNiagaraSystemViewModel();
+
+	/** Gets an array of the view models for the emitter handles owned by this System. */
+	const TArray<TSharedRef<FNiagaraEmitterHandleViewModel>>& GetEmitterHandleViewModels();
+
+	/** Gets an emitter handle view model by id.  Returns an invalid shared ptr if it can't be found. */
+	TSharedPtr<FNiagaraEmitterHandleViewModel> GetEmitterHandleViewModelById(FGuid InEmitterHandleId);
+
+	/** Gets the view model for the System script. */
+	TSharedPtr<FNiagaraSystemScriptViewModel> GetSystemScriptViewModel();
+
+	/** Gets a niagara component for previewing the simulated System. */
+	UNiagaraComponent* GetPreviewComponent();
+
+	/** Gets the sequencer for this System for displaying the timeline. */
+	TSharedPtr<ISequencer> GetSequencer();
+
+	/** Gets the curve owner for the System represented by this view model, for use with the curve editor widget. */
+	FNiagaraCurveOwner& GetCurveOwner();
+
+	/** Get access to the underlying system*/
+	UNiagaraSystem& GetSystem() { return System; }
+
+	/** Gets whether or not emitters can be added from the timeline. */
+	bool GetCanModifyEmittersFromTimeline() const;
+
+	/** Gets the current editing mode for this system view model. */
+	ENiagaraSystemViewModelEditMode GetEditMode() const;
+
+	/** Adds a new emitter to the System from an emitter asset data. */
+	void AddEmitterFromAssetData(const FAssetData& AssetData);
+
+	/** Adds a new emitter to the System. */
+	void AddEmitter(UNiagaraEmitter& Emitter);
+
+	/** Duplicates the selected emitter in this System. */
+	void DuplicateEmitter(TSharedRef<FNiagaraEmitterHandleViewModel> EmitterHandleToDuplicate);
+
+	/** Deletes the selected emitter from the System. */
+	void DeleteEmitter(TSharedRef<FNiagaraEmitterHandleViewModel> EmitterHandleToDelete);
+
+	/** Deletes the emitters with the supplied ids from the system */
+	void DeleteEmitters(TSet<FGuid> EmitterHandleIdsToDelete);
+
+	/** Gets a multicast delegate which is called any time the array of emitter handle view models changes. */
+	FOnEmitterHandleViewModelsChanged& OnEmitterHandleViewModelsChanged();
+
+	/** Gets a delegate which is called any time the data in the curve owner is changed internally by this view model. */
+	FOnCurveOwnerChanged& OnCurveOwnerChanged();
+
+	/** Gets a multicast delegate which is called whenever the selected emitter handles changes. */
+	FOnSelectedEmitterHandlesChanged& OnSelectedEmitterHandlesChanged();
+	
+	/** Gets a multicast delegate which is called whenever we've received and handled a sequencer time update.*/
+	FOnPostSequencerTimeChange& OnPostSequencerTimeChanged();
+
+	/** Gets a multicast delegate which is called whenever the system has been compiled. */
+	FOnSystemCompiled& OnSystemCompiled();
+
+	//~ FGCObject interface
+	virtual void AddReferencedObjects(FReferenceCollector& Collector) override;
+
+	//~ FEditorUndoClient interface
+	virtual void PostUndo(bool bSuccess) override;
+	virtual void PostRedo(bool bSuccess) override { PostUndo(bSuccess); }
+
+	// ~ FTickableEditorObject
+	virtual void Tick(float DeltaTime) override;
+	virtual bool IsTickable() const override;
+	virtual TStatId GetStatId() const override;
+
+	/** Resets the System instance to initial conditions. */
+	void ResetSystem();
+
+	/** Resets the system instance on the next frame. */
+	void RequestResetSystem();
+
+	/** Reinitializes the System instance to initial conditions - rebuilds all data sets and resets data interfaces. */
+	void ReInitializeSystemInstances();
+
+	/** Compiles the spawn and update scripts. */
+	void CompileSystem(bool bForce);
+
+	/* Get the latest status of this view-model's script compilation.*/
+	ENiagaraScriptCompileStatus GetLatestCompileStatus();
+
+	/** Gets the ids for the currently selected emitter handles. */
+	const TArray<FGuid>& GetSelectedEmitterHandleIds();
+
+	/** Sets the currently selected emitter handles by id. */
+	void SetSelectedEmitterHandlesById(TArray<FGuid> InSelectedEmitterHandleIds);
+
+	/** Sets the currently selected emitter handle by id. */
+	void SetSelectedEmitterHandleById(FGuid InSelectedEmitterHandleId);
+
+	/** Gets the currently selected emitter handles. */
+	void GetSelectedEmitterHandles(TArray<TSharedRef<FNiagaraEmitterHandleViewModel>>& OutSelectedEmitterHanldles);
+
+	/** Gets editor specific data which can be stored per system.  If this data hasn't been created the default version will be returned. */
+	const UNiagaraSystemEditorData& GetEditorData() const;
+
+	/** Gets editor specific data which is stored per system.  If this data hasn't been created then it will be created. */
+	UNiagaraSystemEditorData& GetOrCreateEditorData();
+	
+	/** Act as if the system has been fully destroyed although references might persist.*/
+	void Cleanup();
+
+	/** Reinitializes all System instances, and rebuilds emitter handle view models and tracks. */
+	void RefreshAll();
+
+	/** Called to notify the system view model that one of the data objects in the system was modified. */
+	void NotifyDataObjectChanged(UObject* ChangedObject);
+
+	/** Updates all selected emitter's fixed bounds with their current dynamic bounds. */
+	void UpdateEmitterFixedBounds();
+
+	/** Isolates the supplied emitters.  This will remove all other emitters from isolation. */
+	void IsolateEmitters(TArray<TSharedRef<FNiagaraEmitterHandleViewModel>> EmitterHandlesToIsolate);
+
+	/** Toggles the isolation state of a single emitter. */
+	void ToggleEmitterIsolation(TSharedRef<FNiagaraEmitterHandleViewModel> InEmitterHandle);
+
+	/** Returns whether the suplied emitter is isolated. */
+	bool IsEmitterIsolated(TSharedRef<FNiagaraEmitterHandleViewModel> InEmitterHandle);
+
+	/** Export the current state of the system to text.*/
+	void DumpToText(FString& ExportText);
+
+	/** Keeps track of the emitters that are pinned in the stack UI, for persistence*/
+	NIAGARAEDITOR_API TArray<TSharedRef<FNiagaraEmitterHandleViewModel>> GetPinnedEmitterHandles();
+
+	/** Sets an internal state when an emitter was pinned/unpinned in the stack UI */
+	NIAGARAEDITOR_API void SetEmitterPinnedState(TSharedRef<FNiagaraEmitterHandleViewModel> EmitterHandleModel, bool bPinnedState);
+
+	/** Returns a multicast delegate that is broadcast when an emitter pinned state has changed */
+	NIAGARAEDITOR_API FOnPinnedEmittersChanged& GetOnPinnedEmittersChanged();
+
+	/** Checks whether or not an emitter is pinned in the stack UI*/
+	NIAGARAEDITOR_API bool GetIsEmitterPinned(TSharedRef<FNiagaraEmitterHandleViewModel> EmitterHandleModel);
+
+	NIAGARAEDITOR_API void OnPreSave();
+	NIAGARAEDITOR_API void OnPreClose();
+	
+	/** Gets the system toolkit command list. */
+	NIAGARAEDITOR_API TSharedPtr<FUICommandList> GetToolkitCommands();
+	
+	/** Returns a multicast delegate that is broadcast when pinned curves (shown in curve editor) state has changed*/
+	NIAGARAEDITOR_API FOnPinnedCurvesChanged& GetOnPinnedCurvesChanged();
+
+	/** Set the system toolkit command list. */
+	void SetToolkitCommands(const TSharedRef<FUICommandList>& InToolkitCommands);
+
+private:
+
+	/** Sets up the preview component and System instance. */
+	void SetupPreviewComponentAndInstance();
+
+	/** Rebuilds the emitter handle view models. */
+	void RefreshEmitterHandleViewModels();
+
+	/** Rebuilds the sequencer tracks. */
+	void RefreshSequencerTracks();
+
+	/** Gets the sequencer emitter track for the supplied emitter handle view model. */
+	UMovieSceneNiagaraEmitterTrack* GetTrackForHandleViewModel(TSharedRef<FNiagaraEmitterHandleViewModel> EmitterHandleViewModel);
+
+	/** Sets up the sequencer for this emitter. */
+	void SetupSequencer();
+
+	/** Gets the content for the add menu in sequencer. */
+	void GetSequencerAddMenuContent(FMenuBuilder& MenuBuilder, TSharedRef<ISequencer> Sequencer);
+
+	/** Kills all system instances using the system being displayed by this view model. */
+	void KillSystemInstances();
+
+	/** Resets and rebuilds the data in the curve owner. */
+	void ResetCurveData();
+
+	/** Updates the compiled versions of data interfaces when their sources change. */
+	void UpdateCompiledDataInterfaces(UNiagaraDataInterface* ChangedDataInterface);
+
+	/** Called whenever a property on the emitter handle changes. */
+	void EmitterHandlePropertyChanged(TSharedRef<FNiagaraEmitterHandleViewModel> EmitterHandleViewModel);
+
+	/** Called whenever the name on an emitter handle changes. */
+	void EmitterHandleNameChanged(TSharedRef<FNiagaraEmitterHandleViewModel> EmitterHandleViewModel);
+
+	/** Called whenever a property on the emitter changes. */
+	void EmitterPropertyChanged(TSharedRef<FNiagaraEmitterHandleViewModel> EmitterHandleViewModel);
+
+	/** 
+	 * Called whenever a parameter store owned by the system changes.
+	 * @param ChangedParameterStore The parameter store that changed.
+	 * @param OwningScript The script that owns the parameter store, if there is one.
+	 */
+	void SystemParameterStoreChanged(const FNiagaraParameterStore& ChangedParameterStore, const UNiagaraScript* OwningScript);
+
+	/** Called whenever an emitter's script graph changes. */
+	void EmitterScriptGraphChanged(const FEdGraphEditAction& InAction, const UNiagaraScript& OwningScript, const TSharedRef<FNiagaraEmitterHandleViewModel> OwningEmitterHandleViewModel);
+
+	/**
+	* Called whenever an emitter's parameter store owned by the system changes.
+	* @param ChangedParameterStore The parameter store that changed.
+	* @param OwningScript The script that owns the parameter store, if there is one.
+	*/
+	void EmitterParameterStoreChanged(const FNiagaraParameterStore& ChangedParameterStore, const UNiagaraScript& OwningScript, const TSharedRef<FNiagaraEmitterHandleViewModel> OwningEmitterHandleViewModel);
+
+	/** Updates the current simulation for a parameter changing, based on the current simulation options. */
+	void UpdateSimulationFromParameterChange();
+
+	/** Called when a script is compiled */
+	void ScriptCompiled();
+
+	/** Handles when a curve in the curve owner is changed by the curve editor. */
+	void CurveChanged(FRichCurve* ChangedCurve, UObject* CurveOwnerObject);
+
+	/** Called whenever the data in the sequence is changed. */
+	void SequencerDataChanged(EMovieSceneDataChangeType DataChangeType);
+
+	/** Called whenever the global time in the sequencer is changed. */
+	void SequencerTimeChanged();
+
+	/** Called whenever the track selection in sequencer changes. */
+	void SequencerTrackSelectionChanged(TArray<UMovieSceneTrack*> SelectedTracks);
+
+	/** Called whenever the section selection in sequencer changes. */
+	void SequencerSectionSelectionChanged(TArray<UMovieSceneSection*> SelectedSections);
+
+	/** Updates the current emitter handle selection base on the sequencer selection. */
+	void UpdateEmitterHandleSelectionFromSequencer();
+
+	/** Updates the sequencer selection based on the current emitter handle selection. */
+	void UpdateSequencerFromEmitterHandleSelection();
+
+	/** Called when the system instance on the preview component changes. */
+	void PreviewComponentSystemInstanceChanged();
+
+	/** Called whenever the System instance is initialized. */
+	void SystemInstanceInitialized();
+
+	/** Called whenever the System instance is reset.*/
+	void SystemInstanceReset();
+
+	/** Duplicates a set of emitters and refreshes everything.*/
+	void DuplicateEmitters(TSet<FGuid> EmitterHandleIdsToDuplicate);
+
+	/** Adds event handler for the system's scripts. */
+	void AddSystemEventHandlers();
+
+	/** Removes event handlers for the system's scripts. */
+	void RemoveSystemEventHandlers();
+	
+	/** Adds event handler for the system's scripts. */
+	void AddCurveEventHandlers();
+
+	/** Removes event handlers for the system's scripts. */
+	void RemoveCurveEventHandlers();
+
+	/** sends a notification that pinned curves have changed status */
+	void NotifyPinnedCurvesChanged();
+
+private:
+	/** The System being viewed and edited by this view model. */
+	UNiagaraSystem& System;
+
+	/** The component used for previewing the System in a viewport. */
+	UNiagaraComponent* PreviewComponent;
+
+	/** The system instance currently simulating this system if available. */
+	FNiagaraSystemInstance* SystemInstance;
+
+	/** The view models for the emitter handles owned by the System. */
+	TArray<TSharedRef<FNiagaraEmitterHandleViewModel>> EmitterHandleViewModels;
+
+	/** The view model for the System script. */
+	TSharedPtr<FNiagaraSystemScriptViewModel> SystemScriptViewModel;
+
+	/** A niagara sequence for displaying this System in the sequencer timeline. */
+	UNiagaraSequence *NiagaraSequence;
+
+	/** The sequencer instance viewing and editing the niagara sequence. */
+	TSharedPtr<ISequencer> Sequencer;
+
+	/** Flag which indicates we are setting the sequencer time directly in an internal operation. */
+	bool bSettingSequencerTimeDirectly;
+
+	/** The previous play status for sequencer timeline. */
+	EMovieScenePlayerStatus::Type PreviousSequencerStatus;
+
+	/** The previous time for the sequencer timeline. */
+	float PreviousSequencerTime;
+
+	/** Whether or not the user can edit emitters from the timeline. */
+	bool bCanModifyEmittersFromTimeline;
+
+	/** Whether or not we use the system's execution state to drive when we reset the timeline*/
+	bool bUseSystemExecStateForTimelineReset;
+
+	/** Whether or not the system represented by this view model can be automatically compiled. */
+	bool bCanAutoCompile;
+
+	/** Whether or not the system requires a compilation*/
+	bool bForceAutoCompileOnce;
+
+	/** Whether or not the system represented by this view model can be simulated. */
+	bool bCanSimulate;
+
+	/** The current editing mode for this view model. */
+	ENiagaraSystemViewModelEditMode EditMode;
+
+	/** A delegate which is used to generate the content for the add menu in sequencer. */
+	FOnGetAddMenuContent OnGetSequencerAddMenuContent;
+
+	/** A multicast delegate which is called any time the array of emitter handle view models changes. */
+	FOnEmitterHandleViewModelsChanged OnEmitterHandleViewModelsChangedDelegate;
+
+	/** A multicast delegate which is called when the contents of the curve owner is changed internally by this view model. */
+	FOnCurveOwnerChanged OnCurveOwnerChangedDelegate;
+
+	/** A multicast delegate which is called whenever the selected emitter changes. */
+	FOnSelectedEmitterHandlesChanged OnSelectedEmitterHandlesChangedDelegate;
+
+	/** A multicast delegate which is called whenever we've received and handled a sequencer time update.*/
+	FOnPostSequencerTimeChange OnPostSequencerTimeChangeDelegate;
+
+	/** A multicast delegate which is called whenever the system has been compiled. */
+	FOnSystemCompiled OnSystemCompiledDelegate;
+
+	/** A flag for preventing reentrancy when syncrhonizing sequencer data. */
+	bool bUpdatingEmittersFromSequencerDataChange;
+
+	/** A flag for preventing reentrancy when syncrhonizing sequencer data. */
+	bool bUpdatingSequencerFromEmitterDataChange;
+
+	/** A flag for preventing reentrancy when synchronizing system selection with sequencer selection */
+	bool bUpdatingSystemSelectionFromSequencer;
+
+	/** A flag for preventing reentrancy when synchronizing sequencer selection with system selection */
+	bool bUpdatingSequencerSelectionFromSystem;
+
+	/** A curve owner implementation for curves in a niagara System. */
+	FNiagaraCurveOwner CurveOwner;
+
+	/** The ids for the currently selected emitter handles. */
+	TArray<FGuid> SelectedEmitterHandleIds;
+
+	TNiagaraViewModelManager<UNiagaraSystem, FNiagaraSystemViewModel>::Handle RegisteredHandle;
+
+	/** Array of FNiagaraEmitterHandleViewModel representing emitters that are pinned in the stack UI*/
+	TArray<TSharedRef<FNiagaraEmitterHandleViewModel>> PinnedEmitterHandles;
+
+	/** A single param delegate which is called when the pinned status is changed. */
+	FOnPinnedEmittersChanged OnPinnedChangedDelegate;
+
+	/** A cached reference to the niagara editor settings .*/
+	UNiagaraEditorSettings* EditorSettings;
+
+	/** A mapping of script to the delegate handle for it's on parameter map changed delegate. */
+	TMap<FObjectKey, FDelegateHandle> ScriptToOnParameterStoreChangedHandleMap;
+
+	/** A handle to the on parameter changed delegate for the user parameter store. */
+	FDelegateHandle UserParameterStoreChangedHandle;
+
+	/** A flag indicating that a reset has been request on the next frame */
+	bool bResetRequestPending;
+
+	/** The system toolkit commands. */
+	TWeakPtr<class FUICommandList> ToolkitCommands;
+
+	/*A multicast delegate which is called any time the pinned status of the curves changes*/
+	FOnPinnedCurvesChanged OnPinnedCurvesChangedDelegate;
+
+	/*An array of data interfaces used to keep track of how curves are changed by the user*/
+	TArray<UNiagaraDataInterfaceCurveBase*> ShownCurveDataInterfaces;
+
+	/** A flag which indicates that a compile has been requested, but has not completed. */
+	bool bCompilePendingCompletion;
+};
\ No newline at end of file
diff --git a/Engine/Plugins/FX/Niagara/Source/NiagaraEditor/Public/ViewModels/Stack/NiagaraParameterHandle.h b/Engine/Plugins/FX/Niagara/Source/NiagaraEditor/Public/ViewModels/Stack/NiagaraParameterHandle.h
--- a/Engine/Plugins/FX/Niagara/Source/NiagaraEditor/Public/ViewModels/Stack/NiagaraParameterHandle.h
+++ b/Engine/Plugins/FX/Niagara/Source/NiagaraEditor/Public/ViewModels/Stack/NiagaraParameterHandle.h
@@ -19,7 +19,7 @@ public:
 
 	bool operator==(const FNiagaraParameterHandle& Other) const;
 
-	static FNiagaraParameterHandle CreateAliasedModuleParameterHandle(const FNiagaraParameterHandle& ModuleParameterHandle, UNiagaraNodeFunctionCall* ModuleNode);
+	static FNiagaraParameterHandle CreateAliasedModuleParameterHandle(const FNiagaraParameterHandle& ModuleParameterHandle, const UNiagaraNodeFunctionCall* ModuleNode);
 
 	static FNiagaraParameterHandle CreateEngineParameterHandle(const FNiagaraVariable& SystemVariable);
 
@@ -39,6 +39,8 @@ public:
 
 	const FName GetNamespace() const;
 
+	bool IsUserHandle() const;
+
 	bool IsEngineHandle() const;
 
 	bool IsSystemHandle() const;
@@ -49,6 +51,8 @@ public:
 
 	bool IsModuleHandle() const;
 
+	bool IsParameterCollectionHandle() const;
+
 public:
 	static const FName UserNamespace;
 	static const FName EngineNamespace;
@@ -56,6 +60,7 @@ public:
 	static const FName EmitterNamespace;
 	static const FName ParticleAttributeNamespace;
 	static const FName ModuleNamespace;
+	static const FName ParameterCollectionNamespace;
 	static const FString InitialPrefix;
 
 private:
diff --git a/Engine/Plugins/FX/Niagara/Source/NiagaraEditor/Public/ViewModels/Stack/NiagaraStackAdvancedExpander.h b/Engine/Plugins/FX/Niagara/Source/NiagaraEditor/Public/ViewModels/Stack/NiagaraStackAdvancedExpander.h
new file mode 100644
--- /dev/null
+++ b/Engine/Plugins/FX/Niagara/Source/NiagaraEditor/Public/ViewModels/Stack/NiagaraStackAdvancedExpander.h
@@ -0,0 +1,43 @@
+// Copyright 1998-2018 Epic Games, Inc. All Rights Reserved.
+
+#pragma once
+
+#include "ViewModels/Stack/NiagaraStackEntry.h"
+#include "NiagaraStackAdvancedExpander.generated.h"
+
+class UNiagaraNodeFunctionCall;
+class UNiagaraStackFunctionInputCollection;
+class UNiagaraStackModuleItemOutputCollection;
+class UNiagaraNode;
+
+UCLASS()
+class NIAGARAEDITOR_API UNiagaraStackAdvancedExpander : public UNiagaraStackEntry
+{
+	GENERATED_BODY()
+
+public:
+	DECLARE_DELEGATE(FOnToggleShowAdvanced);
+
+public:
+	void Initialize(
+		FRequiredEntryData InRequiredEntryData,
+		FString InOwnerStackItemEditorDataKey,
+		UNiagaraNode* OwningNiagaraNode = nullptr);
+
+	virtual bool GetCanExpand() const override;
+	virtual bool GetIsEnabled() const override;
+	virtual EStackRowStyle GetStackRowStyle() const override;
+
+	void SetOnToggleShowAdvanced(FOnToggleShowAdvanced OnToggleShowAdvanced);
+
+	bool GetShowAdvanced() const;
+
+	void ToggleShowAdvanced();
+
+private:
+	FString OwnerStackItemEditorDataKey;
+
+	UNiagaraNode* OwningNiagaraNode;
+
+	FOnToggleShowAdvanced ToggleShowAdvancedDelegate;
+};
diff --git a/Engine/Plugins/FX/Niagara/Source/NiagaraEditor/Public/ViewModels/Stack/NiagaraStackEmitterSpawnScriptItemGroup.h b/Engine/Plugins/FX/Niagara/Source/NiagaraEditor/Public/ViewModels/Stack/NiagaraStackEmitterSpawnScriptItemGroup.h
--- a/Engine/Plugins/FX/Niagara/Source/NiagaraEditor/Public/ViewModels/Stack/NiagaraStackEmitterSpawnScriptItemGroup.h
+++ b/Engine/Plugins/FX/Niagara/Source/NiagaraEditor/Public/ViewModels/Stack/NiagaraStackEmitterSpawnScriptItemGroup.h
@@ -10,7 +10,6 @@ class FNiagaraEmitterViewModel;
 class FNiagaraScriptViewModel;
 class UNiagaraStackObject;
 class UNiagaraStackSpacer;
-class UNiagaraStackItemExpander;
 
 UCLASS()
 class NIAGARAEDITOR_API UNiagaraStackEmitterPropertiesItem : public UNiagaraStackItem
@@ -18,7 +17,7 @@ class NIAGARAEDITOR_API UNiagaraStackEmitterPropertiesItem : public UNiagaraStac
 	GENERATED_BODY()
 
 public:
-	void Initialize(TSharedRef<FNiagaraSystemViewModel> InSystemViewModel, TSharedRef<FNiagaraEmitterViewModel> InEmitterViewModel, UNiagaraStackEditorData& InStackEditorData);
+	void Initialize(FRequiredEntryData InRequiredEntryData);
 
 	virtual FText GetDisplayName() const override;
 
@@ -29,11 +28,9 @@ public:
 protected:
 	virtual void BeginDestroy() override;
 
-	virtual void RefreshChildrenInternal(const TArray<UNiagaraStackEntry*>& CurrentChildren, TArray<UNiagaraStackEntry*>& NewChildren) override;
+	virtual void RefreshChildrenInternal(const TArray<UNiagaraStackEntry*>& CurrentChildren, TArray<UNiagaraStackEntry*>& NewChildren, TArray<FStackIssue>& NewIssues) override;
 
 private:
-	void EmitterExpandedChanged();
-
 	void EmitterPropertiesChanged();
 
 private:
@@ -43,9 +40,6 @@ private:
 
 	UPROPERTY()
 	UNiagaraStackObject* EmitterObject;
-
-	UPROPERTY()
-	UNiagaraStackItemExpander* EmitterExpander;
 };
 
 UCLASS()
@@ -57,13 +51,9 @@ public:
 	UNiagaraStackEmitterSpawnScriptItemGroup();
 
 protected:
-	void RefreshChildrenInternal(const TArray<UNiagaraStackEntry*>& CurrentChildren, TArray<UNiagaraStackEntry*>& NewChildren) override;
+	void RefreshChildrenInternal(const TArray<UNiagaraStackEntry*>& CurrentChildren, TArray<UNiagaraStackEntry*>& NewChildren, TArray<FStackIssue>& NewIssues) override;
 
 private:
-
 	UPROPERTY()
 	UNiagaraStackEmitterPropertiesItem* PropertiesItem;
-
-	UPROPERTY()
-	UNiagaraStackSpacer* PropertiesSpacer;
 };
diff --git a/Engine/Plugins/FX/Niagara/Source/NiagaraEditor/Public/ViewModels/Stack/NiagaraStackEntry.h b/Engine/Plugins/FX/Niagara/Source/NiagaraEditor/Public/ViewModels/Stack/NiagaraStackEntry.h
--- a/Engine/Plugins/FX/Niagara/Source/NiagaraEditor/Public/ViewModels/Stack/NiagaraStackEntry.h
+++ b/Engine/Plugins/FX/Niagara/Source/NiagaraEditor/Public/ViewModels/Stack/NiagaraStackEntry.h
@@ -6,6 +6,17 @@
 
 class FNiagaraSystemViewModel;
 class FNiagaraEmitterViewModel;
+class FNiagaraScriptViewModel;
+class UNiagaraStackEditorData;
+class UNiagaraStackErrorItem;
+
+UENUM()
+enum class EStackIssueSeverity : uint8
+{
+	Error = 0,
+	Warning, 
+	Info
+};
 
 UCLASS()
 class NIAGARAEDITOR_API UNiagaraStackEntry : public UObject
@@ -13,38 +24,143 @@ class NIAGARAEDITOR_API UNiagaraStackEntry : public UObject
 	GENERATED_BODY()
 
 public:
+	struct FDropResult
+	{
+		explicit FDropResult(bool bInCanDrop, FText InDropMessage = FText())
+			: bCanDrop(bInCanDrop)
+			, DropMessage(InDropMessage)
+		{
+		}
+
+		const bool bCanDrop;
+		const FText DropMessage;
+	};
+
 	DECLARE_MULTICAST_DELEGATE(FOnStructureChanged);
 	DECLARE_MULTICAST_DELEGATE_OneParam(FOnDataObjectModified, UObject*);
+	DECLARE_DELEGATE_RetVal_TwoParams(TOptional<FDropResult>, FOnRequestDrop, const UNiagaraStackEntry&, const TArray<UNiagaraStackEntry*>&);
 	DECLARE_DELEGATE_RetVal_OneParam(bool, FOnFilterChild, const UNiagaraStackEntry&);
+	DECLARE_DELEGATE(FStackIssueFixDelegate);
+
+public:
+	struct NIAGARAEDITOR_API FExecutionCategoryNames
+	{
+		static const FName System;
+		static const FName Emitter;
+		static const FName Particle;
+		static const FName Render;
+	};
+
+	struct NIAGARAEDITOR_API FExecutionSubcategoryNames
+	{
+		static const FName Parameters;
+		static const FName Spawn;
+		static const FName Update;
+		static const FName Event;
+	};
+
+	enum class EStackRowStyle
+	{
+		None,
+		GroupHeader,
+		ItemHeader,
+		ItemContent,
+		ItemContentAdvanced,
+		ItemFooter,
+		ItemCategory,
+		StackIssue
+	};
+
+	struct FRequiredEntryData
+	{
+		FRequiredEntryData(TSharedRef<FNiagaraSystemViewModel> InSystemViewModel, TSharedRef<FNiagaraEmitterViewModel> InEmitterViewModel, FName InExecutionCategoryName, FName InExecutionSubcategoryName, UNiagaraStackEditorData& InStackEditorData)
+			: SystemViewModel(InSystemViewModel)
+			, EmitterViewModel(InEmitterViewModel)
+			, ExecutionCategoryName(InExecutionCategoryName)
+			, ExecutionSubcategoryName(InExecutionSubcategoryName)
+			, StackEditorData(&InStackEditorData)
+		{
+		}
+
+		const TSharedRef<FNiagaraSystemViewModel> SystemViewModel;
+		const TSharedRef<FNiagaraEmitterViewModel> EmitterViewModel;
+		const FName ExecutionCategoryName;
+		const FName ExecutionSubcategoryName;
+		UNiagaraStackEditorData* const StackEditorData;
+	};
+
+	struct FStackSearchItem
+	{
+		FName Key;
+		FText Value;
+		
+		inline bool operator==(FStackSearchItem Item) {
+			return (Item.Key == Key 
+				&& Item.Value.ToString() == Value.ToString());
+		}
+	};
+
+	// stack issue stuff
+	struct FStackIssueFix
+	{
+	public:
+		bool operator == (const FStackIssueFix &Other) const
+		{
+			return Description.CompareTo(Other.Description) == 0 && FixDelegate.GetHandle() == Other.FixDelegate.GetHandle();
+		}
+	public:
+		FText Description;
+		FStackIssueFixDelegate FixDelegate;
+	};
+
+	struct FStackIssue
+	{
+		FStackIssue()
+			: Severity(EStackIssueSeverity::Error)
+			, bCanBeDismissed(false)
+		{}
+		FText ShortDescription;
+		FText LongDescription;
+		FName UniqueIdentifier;
+		EStackIssueSeverity Severity;
+		TArray<FStackIssueFix> Fixes;
+		bool bCanBeDismissed;
+	};
 
 public:
 	UNiagaraStackEntry();
 
-	void Initialize(TSharedRef<FNiagaraSystemViewModel> InSystemViewModel, TSharedRef<FNiagaraEmitterViewModel> InEmitterViewModel);
+	void Initialize(FRequiredEntryData InRequiredEntryData, FString InStackEditorDataKey);
+
+	bool IsValid() const;
 
 	virtual FText GetDisplayName() const;
 
-	virtual FText GetTooltipText() const;
+	UNiagaraStackEditorData& GetStackEditorData() const;
+
+	FString GetStackEditorDataKey() const;
 
-	virtual FName GetTextStyleName() const;
+	virtual FText GetTooltipText() const;
 
 	virtual bool GetCanExpand() const;
 
 	virtual bool IsExpandedByDefault() const;
 
-	virtual bool GetIsExpanded() const;
+	bool GetIsExpanded() const;
 
-	virtual void SetIsExpanded(bool bInExpanded);
+	// Calling this doesn't broadcast structure change automatically due to the expense of synchronizing
+	// expanded state with the tree which is done to prevent items being expanded on tick.
+	void SetIsExpanded(bool bInExpanded);
 
-	virtual FName GetGroupBackgroundName() const;
+	virtual bool GetIsEnabled() const;
 
-	virtual FName GetGroupForegroundName() const;
+	FName GetExecutionCategoryName() const;
 
-	virtual FName GetItemBackgroundName() const;
+	FName GetExecutionSubcategoryName() const;
 
-	virtual FName GetItemForegroundName() const;
+	virtual EStackRowStyle GetStackRowStyle() const;
 
-	virtual int32 GetItemIndentLevel() const;
+	int32 GetIndentLevel() const;
 
 	virtual bool GetShouldShowInStack() const;
 
@@ -52,22 +168,18 @@ public:
 
 	void GetUnfilteredChildren(TArray<UNiagaraStackEntry*>& OutUnfilteredChildren);
 
+	void GetUnfilteredChildren(TArray<UNiagaraStackEntry*>& OutUnfilteredChildren) const;
+
 	FOnStructureChanged& OnStructureChanged();
 
 	FOnDataObjectModified& OnDataObjectModified();
 
 	void RefreshChildren();
-	void RefreshErrors();
+
 
 	FDelegateHandle AddChildFilter(FOnFilterChild ChildFilter);
 	void RemoveChildFilter(FDelegateHandle FilterHandle);
 
-	virtual int32 GetErrorCount() const { return 0; }
-	virtual bool GetErrorFixable(int32 ErrorIdx) const { return false; }
-	virtual bool TryFixError(int32 ErrorIdx) { return false; }
-	virtual FText GetErrorText(int32 ErrorIdx) const { return FText(); }
-	virtual FText GetErrorSummaryText(int32 ErrorIdx) const { return FText(); }
-
 	TSharedRef<FNiagaraSystemViewModel> GetSystemViewModel() const;
 	TSharedRef<FNiagaraEmitterViewModel> GetEmitterViewModel() const;
 
@@ -93,18 +205,62 @@ public:
 		return FindCurrentChildOfTypeByPredicate<ChildType>(CurrentChildren, Predicate);
 	}
 
+	void GetSearchItems(TArray<FStackSearchItem>& SearchItems) const;
+
+	virtual UObject* GetExternalAsset() const;
+
+	virtual bool CanDrag() const;
+
+	TOptional<FDropResult> CanDrop(const TArray<UNiagaraStackEntry*>& DraggedEntries);
+
+	TOptional<FDropResult> Drop(const TArray<UNiagaraStackEntry*>& DraggedEntries);
+
+	void SetOnRequestCanDrop(FOnRequestDrop InOnRequestCanDrop);
+
+	void SetOnRequestDrop(FOnRequestDrop InOnRequestCanDrop);
+
 protected:
-	virtual void RefreshChildrenInternal(const TArray<UNiagaraStackEntry*>& CurrentChildren, TArray<UNiagaraStackEntry*>& NewChildren);
+	virtual void RefreshChildrenInternal(const TArray<UNiagaraStackEntry*>& CurrentChildren, TArray<UNiagaraStackEntry*>& NewChildren, TArray<FStackIssue>& NewIssues);
+
+	virtual void PostRefreshChildrenInternal();
+
+	FRequiredEntryData CreateDefaultChildRequiredData() const;
+
+	virtual int32 GetChildIndentLevel() const;
+
+	virtual void GetAdditionalSearchItemsInternal(TArray<FStackSearchItem>& SearchItems) const;
+
+	virtual TOptional<FDropResult> CanDropInternal(const TArray<UNiagaraStackEntry*>& DraggedEntries);
+
+	virtual TOptional<FDropResult> DropInternal(const TArray<UNiagaraStackEntry*>& DraggedEntries);
+
+	virtual TOptional<FDropResult> ChildRequestCanDropInternal(const UNiagaraStackEntry& TargetChild, const TArray<UNiagaraStackEntry*>& DraggedEntries);
+
+	virtual TOptional<FDropResult> ChildRequestDropInternal(const UNiagaraStackEntry& TargetChild, const TArray<UNiagaraStackEntry*>& DraggedEntries);
+
+	virtual void ChlildStructureChangedInternal();
 
 private:
 	void ChildStructureChanged();
 	
 	void ChildDataObjectModified(UObject* ChangedObject);
 
+	TOptional<FDropResult> ChildRequestCanDrop(const UNiagaraStackEntry& TargetChild, const TArray<UNiagaraStackEntry*>& DraggedEntries);
+
+	TOptional<FDropResult> ChildRequestDrop(const UNiagaraStackEntry& TargetChild, const TArray<UNiagaraStackEntry*>& DraggedEntries);
+
+	void RefreshStackErrorChildren();
+
+	void IssueModified();
+	
 private:
 	TWeakPtr<FNiagaraSystemViewModel> SystemViewModel;
 	TWeakPtr<FNiagaraEmitterViewModel> EmitterViewModel;
 
+	UNiagaraStackEditorData* StackEditorData;
+
+	FString StackEditorDataKey;
+
 	FOnStructureChanged StructureChangedDelegate;
 
 	FOnDataObjectModified DataObjectModifiedDelegate;
@@ -115,7 +271,17 @@ private:
 	TArray<UNiagaraStackEntry*> Children;
 
 	UPROPERTY()
-	TArray<UNiagaraStackEntry*> ErrorChildren;
+	TArray<UNiagaraStackErrorItem*> ErrorChildren;
+
+	mutable TOptional<bool> bIsExpandedCache;
 
-	bool bIsExpanded;
+	int32 IndentLevel;
+
+	FName ExecutionCategoryName;
+	FName ExecutionSubcategoryName;
+
+	FOnRequestDrop OnRequestCanDropDelegate;
+	FOnRequestDrop OnRequestDropDelegate;
+	
+	TArray<FStackIssue> StackIssues;
 };
\ No newline at end of file
diff --git a/Engine/Plugins/FX/Niagara/Source/NiagaraEditor/Public/ViewModels/Stack/NiagaraStackErrorItem.h b/Engine/Plugins/FX/Niagara/Source/NiagaraEditor/Public/ViewModels/Stack/NiagaraStackErrorItem.h
--- a/Engine/Plugins/FX/Niagara/Source/NiagaraEditor/Public/ViewModels/Stack/NiagaraStackErrorItem.h
+++ b/Engine/Plugins/FX/Niagara/Source/NiagaraEditor/Public/ViewModels/Stack/NiagaraStackErrorItem.h
@@ -11,29 +11,76 @@ UCLASS()
 class NIAGARAEDITOR_API UNiagaraStackErrorItem : public UNiagaraStackEntry
 {
 	GENERATED_BODY()
-
 public:
-	UNiagaraStackErrorItem();
+	DECLARE_MULTICAST_DELEGATE(FOnIssueNotify);
+public:
+	void Initialize(FRequiredEntryData InRequiredEntryData, FStackIssue InStackIssue, FString InStackEditorDataKey);
+	FStackIssue GetStackIssue() const { return StackIssue; }
+	virtual FText GetDisplayName() const override;
+	virtual EStackRowStyle GetStackRowStyle() const override;
+	FOnIssueNotify& OnIssueModified();
+	
+protected:
+	virtual void RefreshChildrenInternal(const TArray<UNiagaraStackEntry*>& CurrentChildren, TArray<UNiagaraStackEntry*>& NewChildren, TArray<FStackIssue>& NewIssues);
 
-	void Initialize(TSharedRef<FNiagaraSystemViewModel> InSystemViewModel, TSharedRef<FNiagaraEmitterViewModel> InEmitterViewModel, UNiagaraStackEntry* InErrorSource, int32 InErrorIndex);
+protected:
+	FStackIssue StackIssue;
+	FString EntryStackEditorDataKey;
+	FOnIssueNotify IssueModifiedDelegate;
 
-	//~ UNiagaraStackEntry interface
-	virtual int32 GetItemIndentLevel() const override;
+private:
+	void IssueFixed();
+}; 
 
-	/** Sets the item indent level for this stack entry. */
-	void SetItemIndentLevel(int32 InItemIndentLevel);
+UCLASS()
+class NIAGARAEDITOR_API UNiagaraStackErrorItemLongDescription : public UNiagaraStackEntry
+{
+	GENERATED_BODY()
 
-	int32 GetErrorIndex() const { return ErrorIdx; }
+public:
+	void Initialize(FRequiredEntryData InRequiredEntryData, FStackIssue InStackIssue, FString InStackEditorDataKey);
+	virtual FText GetDisplayName() const override;
+	virtual EStackRowStyle GetStackRowStyle() const override;
 
-	FText ErrorText() const;
-	FText ErrorTextTooltip() const;
-	FReply OnTryFixError();
-	EVisibility CanFixVisibility() const;
+protected:
+	FStackIssue StackIssue;
+};
 
-	virtual FName GetItemBackgroundName() const override;
+UCLASS()
+class NIAGARAEDITOR_API UNiagaraStackErrorItemFix : public UNiagaraStackEntry
+{
+	GENERATED_BODY()
+
+public:
+	void Initialize(FRequiredEntryData InRequiredEntryData, FStackIssue InStackIssue, FStackIssueFix InIssueFix, FString InStackEditorDataKey);
+
+	FStackIssueFix GetStackIssueFix() const { return IssueFix; };
+	FText FixDescription() const;
+	virtual FReply OnTryFixError();
+	virtual EStackRowStyle GetStackRowStyle() const override;
+	virtual FText GetFixButtonText() const;
+	UNiagaraStackErrorItem::FOnIssueNotify& OnIssueFixed();
 
 protected:
-	UNiagaraStackEntry* ErrorSource;
-	int32 ErrorIdx;
-	int32 ItemIndentLevel;
+	FStackIssue StackIssue;
+	FStackIssueFix IssueFix;
+	UNiagaraStackErrorItem::FOnIssueNotify IssueFixedDelegate;
 };
+
+UCLASS()
+class NIAGARAEDITOR_API UNiagaraStackErrorItemDismiss : public UNiagaraStackErrorItemFix
+{
+	GENERATED_BODY()
+
+public:
+	void Initialize(FRequiredEntryData InRequiredEntryData, UNiagaraStackEntry::FStackIssue InStackIssue, FString InStackEditorDataKey);
+
+	virtual FReply OnTryFixError() override;
+	virtual EStackRowStyle GetStackRowStyle() const override;
+	virtual FText GetFixButtonText() const override;
+
+private:
+	void DismissIssue();
+
+};
+
diff --git a/Engine/Plugins/FX/Niagara/Source/NiagaraEditor/Public/ViewModels/Stack/NiagaraStackEventHandlerGroup.h b/Engine/Plugins/FX/Niagara/Source/NiagaraEditor/Public/ViewModels/Stack/NiagaraStackEventHandlerGroup.h
--- a/Engine/Plugins/FX/Niagara/Source/NiagaraEditor/Public/ViewModels/Stack/NiagaraStackEventHandlerGroup.h
+++ b/Engine/Plugins/FX/Niagara/Source/NiagaraEditor/Public/ViewModels/Stack/NiagaraStackEventHandlerGroup.h
@@ -6,7 +6,6 @@
 #include "NiagaraStackEventHandlerGroup.generated.h"
 
 class FNiagaraEmitterViewModel;
-class UNiagaraStackAddEventScriptItem;
 
 UCLASS()
 /** Container for one or more NiagaraStackEventScriptItemGroups, allowing multiple event handlers per script.*/
@@ -18,26 +17,14 @@ public:
 	DECLARE_DELEGATE(FOnItemAdded);
 
 public:
-
-	virtual FText GetDisplayName() const override;
-	void SetDisplayName(FText InDisplayName);
-	
-	virtual bool GetShouldShowInStack() const override;
+	void Initialize(FRequiredEntryData InRequiredEntryData);
 
 	void SetOnItemAdded(FOnItemAdded InOnItemAdded);
-	
-protected:
-	virtual void RefreshChildrenInternal(const TArray<UNiagaraStackEntry*>& CurrentChildren, TArray<UNiagaraStackEntry*>& NewChildren) override;
-
-	void ChildModifiedGroupItems();
 
 private:
-	FText DisplayName;
+	void ItemAddedFromUtilties();
 
-	bool bForceRebuild;
-	
-	UPROPERTY()
-	UNiagaraStackAddEventScriptItem* AddScriptItem;
-	
+private:
+	TSharedPtr<INiagaraStackItemGroupAddUtilities> AddUtilities;
 	FOnItemAdded ItemAddedDelegate;
 };
diff --git a/Engine/Plugins/FX/Niagara/Source/NiagaraEditor/Public/ViewModels/Stack/NiagaraStackEventScriptItemGroup.h b/Engine/Plugins/FX/Niagara/Source/NiagaraEditor/Public/ViewModels/Stack/NiagaraStackEventScriptItemGroup.h
--- a/Engine/Plugins/FX/Niagara/Source/NiagaraEditor/Public/ViewModels/Stack/NiagaraStackEventScriptItemGroup.h
+++ b/Engine/Plugins/FX/Niagara/Source/NiagaraEditor/Public/ViewModels/Stack/NiagaraStackEventScriptItemGroup.h
@@ -6,19 +6,20 @@
 #include "NiagaraEmitter.h"
 #include "ViewModels/Stack/NiagaraStackItem.h"
 #include "ViewModels/Stack/NiagaraStackScriptItemGroup.h"
-#include "ViewModels/Stack/NiagaraStackItemExpander.h"
-#include "ViewModels/Stack/NiagaraStackAddEventScriptItem.h"
 #include "NiagaraStackEventScriptItemGroup.generated.h"
 
 class FNiagaraScriptViewModel;
+class UNiagaraStackObject;
+struct FNiagaraEventScriptProperties;
+class IDetailTreeNode;
 
 UCLASS()
 class NIAGARAEDITOR_API UNiagaraStackEventHandlerPropertiesItem : public UNiagaraStackItem
 {
 	GENERATED_BODY()
-
+		 
 public:
-	void Initialize(TSharedRef<FNiagaraSystemViewModel> InSystemViewModel, TSharedRef<FNiagaraEmitterViewModel> InEmitterViewModel, FGuid InEventScriptUsageId, UNiagaraStackEditorData& InStackEditorData);
+	void Initialize(FRequiredEntryData InRequiredEntryData, FGuid InEventScriptUsageId);
 
 	virtual FText GetDisplayName() const override;
 
@@ -29,26 +30,26 @@ public:
 protected:
 	virtual void BeginDestroy() override;
 
-	virtual void RefreshChildrenInternal(const TArray<UNiagaraStackEntry*>& CurrentChildren, TArray<UNiagaraStackEntry*>& NewChildren) override;
+	virtual void RefreshChildrenInternal(const TArray<UNiagaraStackEntry*>& CurrentChildren, TArray<UNiagaraStackEntry*>& NewChildren, TArray<FStackIssue>& NewIssues) override;
 
 private:
-	void EventHandlerExpandedChanged();
-
 	void EventHandlerPropertiesChanged();
 
+	void SelectEmitterStackObjectRootTreeNodes(TArray<TSharedRef<IDetailTreeNode>> Source, TArray<TSharedRef<IDetailTreeNode>>* Selected);
+
 private:
+	FGuid EventScriptUsageId;
+
 	bool bHasBaseEventHandler;
 
 	mutable TOptional<bool> bCanResetToBase;
 
-	bool bForceRebuildChildStruct;
-
 	TWeakObjectPtr<UNiagaraEmitter> Emitter;
 
-	FGuid EventScriptUsageId;
+	FNiagaraEventScriptProperties* LastUsedEventScriptProperties;
 
 	UPROPERTY()
-	UNiagaraStackItemExpander* EventHandlerExpander;
+	UNiagaraStackObject* EmitterObject;
 };
 
 UCLASS()
@@ -60,26 +61,26 @@ class NIAGARAEDITOR_API UNiagaraStackEventScriptItemGroup : public UNiagaraStack
 public:
 	DECLARE_DELEGATE(FOnModifiedEventHandlers);
 
+public:
+	void Initialize(
+		FRequiredEntryData InRequiredEntryData,
+		TSharedRef<FNiagaraScriptViewModel> InScriptViewModel,
+		ENiagaraScriptUsage InScriptUsage,
+		FGuid InScriptUsageId);
+
+	void SetOnModifiedEventHandlers(FOnModifiedEventHandlers OnModifiedEventHandlers);
+
 protected:
 	FOnModifiedEventHandlers OnModifiedEventHandlersDelegate;
 
-	virtual void RefreshChildrenInternal(const TArray<UNiagaraStackEntry*>& CurrentChildren, TArray<UNiagaraStackEntry*>& NewChildren) override;
+	virtual void RefreshChildrenInternal(const TArray<UNiagaraStackEntry*>& CurrentChildren, TArray<UNiagaraStackEntry*>& NewChildren, TArray<FStackIssue>& NewIssues) override;
 
 	virtual bool CanDelete() const override;
 	virtual bool Delete() override;
 	
-public:
-	virtual FText GetDisplayName() const override;
-	void SetOnModifiedEventHandlers(FOnModifiedEventHandlers OnModifiedEventHandlers);
-
 private:
-	FText DisplayName;
-
 	bool bHasBaseEventHandler;
 
 	UPROPERTY()
 	UNiagaraStackEventHandlerPropertiesItem* EventHandlerProperties;
-
-	UPROPERTY()
-	UNiagaraStackSpacer* EventHandlerPropertiesSpacer;
 };
diff --git a/Engine/Plugins/FX/Niagara/Source/NiagaraEditor/Public/ViewModels/Stack/NiagaraStackFunctionInput.h b/Engine/Plugins/FX/Niagara/Source/NiagaraEditor/Public/ViewModels/Stack/NiagaraStackFunctionInput.h
--- a/Engine/Plugins/FX/Niagara/Source/NiagaraEditor/Public/ViewModels/Stack/NiagaraStackFunctionInput.h
+++ b/Engine/Plugins/FX/Niagara/Source/NiagaraEditor/Public/ViewModels/Stack/NiagaraStackFunctionInput.h
@@ -2,12 +2,14 @@
 
 #pragma once
 
-#include "ViewModels/Stack/NiagaraStackEntry.h"
+#include "ViewModels/Stack/NiagaraStackItem.h"
 #include "NiagaraTypes.h"
 #include "ViewModels/Stack/NiagaraParameterHandle.h"
+#include "ViewModels/Stack/NiagaraStackFunctionInputCondition.h"
 #include "NiagaraStackFunctionInput.generated.h"
 
 class UNiagaraNodeFunctionCall;
+class UNiagaraNodeCustomHlsl;
 class UNiagaraNodeAssignment;
 class UNiagaraNodeParameterMapSet;
 class FStructOnScope;
@@ -16,11 +18,10 @@ class UNiagaraStackObject;
 class UNiagaraScript;
 class UEdGraphPin;
 class UNiagaraDataInterface;
-class UNiagaraStackEditorData;
 
 /** Represents a single module input in the module stack view model. */
 UCLASS()
-class NIAGARAEDITOR_API UNiagaraStackFunctionInput : public UNiagaraStackEntry
+class NIAGARAEDITOR_API UNiagaraStackFunctionInput : public UNiagaraStackItemContent
 {
 	GENERATED_BODY()
 
@@ -36,12 +37,13 @@ public:
 		Dynamic,
 		/** The value is provided by a data interface object. */
 		Data,
+		/** The value is provided by an expression object. */
+		Expression,
 		/** The value source for this input was not set, or couldn't be determined. */
 		Invalid
 	};
 
 	DECLARE_MULTICAST_DELEGATE(FOnValueChanged);
-	DECLARE_MULTICAST_DELEGATE(FOnPinnedChanged);
 
 public:
 	UNiagaraStackFunctionInput();
@@ -53,22 +55,21 @@ public:
 
 	/** 
 	 * Sets the input data for this entry.
-	 * @param InSystemViewModel The view model for the system which owns the stack containing this entry.
-	 * @param InEmitterViewModel The view model for the emitter which owns the stack containing this entry.
+	 * @param InRequiredEntryData The required data for all stack entries.
 	 * @param InStackEditorData The stack editor data for this input.
 	 * @param InModuleNode The module function call which owns this input entry. NOTE: This input might not be an input to the module function call, it may be an input to a dynamic input function call which is owned by the module.
 	 * @param InInputFunctionCallNode The function call which this entry is an input to. NOTE: This node can be a module function call node or a dynamic input node.
 	 * @param InInputParameterHandle The input parameter handle for the function call.
 	 * @param InInputType The type of this input.
+	 * @param InOwnerStackItemEditorDataKey The editor data key of the item that owns this input.
 	 */
 	void Initialize(
-		TSharedRef<FNiagaraSystemViewModel> InSystemViewModel,
-		TSharedRef<FNiagaraEmitterViewModel> InEmitterViewModel,
-		UNiagaraStackEditorData& InStackEditorData,
+		FRequiredEntryData InRequiredEntryData,
 		UNiagaraNodeFunctionCall& InModuleNode,
 		UNiagaraNodeFunctionCall& InInputFunctionCallNode,
 		FName InInputParameterHandle,
-		FNiagaraTypeDefinition InInputType);
+		FNiagaraTypeDefinition InInputType,
+		FString InOwnerStackItemEditorDataKey);
 
 	/** Gets the function call node which owns this input. */
 	const UNiagaraNodeFunctionCall& GetInputFunctionCallNode() const;
@@ -76,24 +77,18 @@ public:
 	/** Gets the current value mode */
 	EValueMode GetValueMode();
 
-	/** Gets whether or not this input can be pinned. */
-	bool GetCanBePinned() const;
-
 	/** Gets the type of this input. */
 	const FNiagaraTypeDefinition& GetInputType() const;
 
 	//~ UNiagaraStackEntry interface
 	virtual FText GetDisplayName() const override;
-	virtual FName GetTextStyleName() const override;
-	virtual bool GetCanExpand() const override;
-	virtual int32 GetItemIndentLevel() const override;
 	virtual FText GetTooltipText() const override;
+	virtual bool GetIsEnabled() const override;
+	
+	UObject* GetExternalAsset() const override;
 
 	FText GetTooltipText(EValueMode InValueMode) const;
 
-	/** Sets the item indent level for this stack entry. */
-	void SetItemIndentLevel(int32 InItemIndentLevel);
-
 	/** Gets the path of parameter handles from the owning module to the function call which owns this input. */
 	const TArray<FNiagaraParameterHandle>& GetInputParameterHandlePath() const;
 
@@ -112,39 +107,45 @@ public:
 	/** Gets the dynamic input node providing the value for this input, if one is available. */
 	UNiagaraNodeFunctionCall* GetDynamicInputNode() const;
 
+	/** Gets the expression input node providing the value for this input, if one is available. */
+	UNiagaraNodeCustomHlsl* GetExpressionNode() const;
+
 	/** Gets the dynamic inputs available for this input. */
 	void GetAvailableDynamicInputs(TArray<UNiagaraScript*>& AvailableDynamicInputs);
 
 	/** Sets the dynamic input script for this input. */
 	void SetDynamicInput(UNiagaraScript* DynamicInput);
 
+	/** Sets the dynamic custom expression script for this input. */
+	void SetCustomExpression(const FString& InputText);
+
 	/** Gets the current struct value of this input is there is one. */
 	TSharedPtr<FStructOnScope> GetLocalValueStruct();
 
 	/** Gets the current data object value of this input is there is one. */
 	UNiagaraDataInterface* GetDataValueObject();
 
-	/** Gets whether or not this input is pinned in the UI. */
-	bool GetIsPinned() const;
-
-	/** Sets whether or not this input is pinned in the UI. */
-	void SetIsPinned(bool bIsPinned);
-
 	/** Called to notify the input that an ongoing change to it's value has begun. */
 	void NotifyBeginLocalValueChange();
 
 	/** Called to notify the input that an ongoing change to it's value has ended. */
 	void NotifyEndLocalValueChange();
 
+	/** Is this pin editable or should it show as disabled?*/
+	bool IsEnabled() const;
+
 	/** Sets this input's local value. */
 	void SetLocalValue(TSharedRef<FStructOnScope> InLocalValue);
-
+	
 	/** Returns whether or not the value or handle of this input has been overridden and can be reset. */
 	bool CanReset() const;
 
 	/** Resets the value and handle of this input to the value and handle defined in the module. */
 	void Reset();
 
+	/** Determine if this field is editable */
+	bool IsEditable() const;
+
 	/** Whether or not this input has a base value.  This is true for emitter instances in systems. */
 	bool EmitterHasBase() const;
 
@@ -166,22 +167,47 @@ public:
 	/** Renames this input to the name specified. */
 	void RenameInput(FName NewName);
 
+	/** Returns whether or not this input can be deleted. */
+	bool CanDeleteInput() const;
+
+	/** Deletes this input */
+	void DeleteInput();
+
 	/** Gets the namespaces which new parameters for this input can be read from. */
 	void GetNamespacesForNewParameters(TArray<FName>& OutNamespacesForNewParameters) const;
 
 	/** Gets a multicast delegate which is called whenever the value on this input changes. */
 	FOnValueChanged& OnValueChanged();
 
-	/** Gets a multicast delegate which is called whenever the pinned state of this input changes. */
-	FOnPinnedChanged& OnPinnedChanged();
+	/** Gets whether or not this input has an associated edit condition input. */
+	bool GetHasEditCondition() const;
+
+	/** Gets whether or not to show a control inline for the edit condition input associated with this input. */
+	bool GetShowEditConditionInline() const;
+
+	/** Gets the enabled value of the edit condition input associated with this input. */
+	bool GetEditConditionEnabled() const;
+
+	/** Sets the enabled value of the edit condition input associated with this input. */
+	void SetEditConditionEnabled(bool bIsEnabled);
+
+	/** Gets whether or not this input has an associated visible condition input. */
+	bool GetHasVisibleCondition() const;
+
+	/** Gets the enabled value of the visible condition input associated with this input. */
+	bool GetVisibleConditionEnabled() const;
+
+	/** Gets whether or not this input is used as an edit condition for another input and should be hidden. */
+	bool GetIsInlineEditConditionToggle() const;
 
 protected:
 	//~ UNiagaraStackEntry interface
-	virtual void RefreshChildrenInternal(const TArray<UNiagaraStackEntry*>& CurrentChildren, TArray<UNiagaraStackEntry*>& NewChildren) override;
+	virtual void RefreshChildrenInternal(const TArray<UNiagaraStackEntry*>& CurrentChildren, TArray<UNiagaraStackEntry*>& NewChildren, TArray<FStackIssue>& NewIssues) override;
 
-	FNiagaraVariable GetDefaultValue() const;
+	FNiagaraVariable GetDefaultVariableForRapidIterationParameter() const;
 	bool UpdateRapidIterationParametersForAffectedScripts(const uint8* Data);
 	bool RemoveRapidIterationParametersForAffectedScripts();
+	FString ResolveDisplayNameArgument(const FString& InArg) const;
 
 private:
 	struct FDataValues
@@ -256,6 +282,7 @@ private:
 		TSharedPtr<FStructOnScope> LocalStruct;
 		FNiagaraParameterHandle LinkedHandle;
 		TWeakObjectPtr<UNiagaraNodeFunctionCall> DynamicNode;
+		TWeakObjectPtr<UNiagaraNodeCustomHlsl> ExpressionNode;
 		FDataValues DataObjects;
 	};
 
@@ -263,6 +290,9 @@ private:
 	/** Refreshes the current values for this input from the state of the graph. */
 	void RefreshValues();
 
+	/** Refreshes additional state for this input which comes from input metadata. */
+	void RefreshFromMetaData();
+
 	/** Called whenever the graph which generated this input changes. */
 	void OnGraphChanged(const struct FEdGraphEditAction& InAction);
 
@@ -290,7 +320,7 @@ private:
 	UEdGraphPin& GetOrCreateOverridePin();
 
 	/** Tries to get a local value for this input if it exists by checking the graph data directly. */
-	bool TryGetCurrentLocalValue(TSharedPtr<FStructOnScope>& LocalValue, UEdGraphPin& ValuePin, TSharedPtr<FStructOnScope> OldValueToReuse);
+	bool TryGetCurrentLocalValue(TSharedPtr<FStructOnScope>& LocalValue, UEdGraphPin& DefaultPin, UEdGraphPin& ValuePin, TSharedPtr<FStructOnScope> OldValueToReuse);
 
 	/** Tries to get a data interface value for this input if it exists by checking the graph data directly .*/
 	bool TryGetCurrentDataValue(FDataValues& DataValues, UEdGraphPin* OverrideValuePin, UEdGraphPin& DefaultValuePin, UNiagaraDataInterface* LocallyOwnedDefaultDataValueObjectToReuse);
@@ -304,15 +334,15 @@ private:
 	/** Removes all nodes connected to the override pin which provide it's value. */
 	void RemoveNodesForOverridePin(UEdGraphPin& OverridePin);
 
+	/** Gets the expression input node providing a value to this input if one exists. */
+	bool TryGetCurrentExpressionValue(TWeakObjectPtr<UNiagaraNodeCustomHlsl>& ExpressionValue, UEdGraphPin* OverridePin);
+
 	/** Determine if the values in this input are possibly under the control of the rapid iteration array on the script.*/
 	bool IsRapidIterationCandidate() const;
 
-	FString CreateRapidIterationVariableName(const FString& InName);
+	FNiagaraVariable CreateRapidIterationVariable(const FName& InName);
 
 private:
-	/** The stack editor data for this function input. */
-	UNiagaraStackEditorData* StackEditorData;
-
 	/** The module function call which owns this input entry. NOTE: This input might not be an input to the module function
 		call, it may be an input to a dynamic input function call which is owned by the module. */
 	TWeakObjectPtr<UNiagaraNodeFunctionCall> OwningModuleNode;
@@ -326,8 +356,8 @@ private:
 	/** The niagara type definition for this input. */
 	FNiagaraTypeDefinition InputType;
 
-	/** Whether or not this input can be pinned. */
-	bool bCanBePinned;
+	/** The meta data for this input, defined in the owning function's script. */
+	FNiagaraVariableMetaData* InputMetaData;
 
 	/** A unique key for this input for looking up editor only UI data. */
 	FString StackEditorDataKey;
@@ -342,12 +372,15 @@ private:
 	  * graph.  This only affects parameter handles which are local module handles. */
 	FNiagaraParameterHandle AliasedInputParameterHandle;
 
-	/** The computed name for the rapid iteration variable that could potentially drive this entry..*/
-	FString RapidIterationParameterName;
+	/** The a rapid iteration variable that could potentially drive this entry..*/
+	FNiagaraVariable RapidIterationParameter;
 
 	/** The name of this input for display in the UI. */
 	FText DisplayName;
 
+	/** Optional override for the display name*/
+	TOptional<FText> DisplayNameOverride;
+
 	/** Pointers and handles to the various values this input can have. */
 	FInputValues InputValues;
 
@@ -359,14 +392,21 @@ private:
 	  * UI reads this value every frame due to attribute updates. */
 	mutable TOptional<UEdGraphPin*> OverridePinCache;
 
-	/** Whether or not this input can be reset to a base value. */
+	/** Whether or not this input can be reset to its default value. */
+	mutable TOptional<bool> bCanReset;
+
+	/** Whether or not this input can be reset to a base value defined by a parent emitter. */
 	mutable TOptional<bool> bCanResetToBase;
 
 	/** A flag to prevent handling graph changes when it's being updated directly by this object. */
 	bool bUpdatingGraphDirectly;
 
+	/** A flag to prevent handling changes to the local value when it's being set directly by this object. */
+	bool bUpdatingLocalValueDirectly;
+
 	/** A handle for removing the graph changed delegate. */
 	FDelegateHandle GraphChangedHandle;
+	FDelegateHandle OnRecompileHandle;
 
 	/** A handle for removing the changed delegate for the rapid iteration parameters for the script
 	   that owns the function that owns this input. */
@@ -375,15 +415,21 @@ private:
 	/** A multicast delegate which is called when the value of this input is changed. */
 	FOnValueChanged ValueChangedDelegate;
 
-	/** A multicast delegate which is called when the pinned state of this input is changed. */
-	FOnPinnedChanged PinnedChangedDelegate;
-
-	/** The item indent level for this stack entry. */
-	int32 ItemIndentLevel;
-
 	/** The script which owns the function which owns this input.  This is also the autoritative version of the rapid iteration parameters. */
 	TWeakObjectPtr<UNiagaraScript> SourceScript;
 
 	/** An array of scripts which this input affects. */
 	TArray<TWeakObjectPtr<UNiagaraScript>> AffectedScripts;
-};
+
+	/** An input condition handler for the edit condition. */
+	FNiagaraStackFunctionInputCondition EditCondition;
+
+	/** An input condition handler for the visible condition. */
+	FNiagaraStackFunctionInputCondition VisibleCondition;
+
+	/** Whether or not to show an inline control for the edit condition input. */
+	bool bShowEditConditionInline;
+
+	/** Whether or not this input is an edit condition toggle. */
+	bool bIsInlineEditConditionToggle;
+};
\ No newline at end of file
diff --git a/Engine/Plugins/FX/Niagara/Source/NiagaraEditor/Public/ViewModels/Stack/NiagaraStackFunctionInputCollection.h b/Engine/Plugins/FX/Niagara/Source/NiagaraEditor/Public/ViewModels/Stack/NiagaraStackFunctionInputCollection.h
--- a/Engine/Plugins/FX/Niagara/Source/NiagaraEditor/Public/ViewModels/Stack/NiagaraStackFunctionInputCollection.h
+++ b/Engine/Plugins/FX/Niagara/Source/NiagaraEditor/Public/ViewModels/Stack/NiagaraStackFunctionInputCollection.h
@@ -2,30 +2,19 @@
 
 #pragma once
 
-#include "ViewModels/Stack/NiagaraStackEntry.h"
+#include "ViewModels/Stack/NiagaraStackItem.h"
+#include "NiagaraTypes.h"
 #include "NiagaraStackFunctionInputCollection.generated.h"
 
 class UNiagaraNodeFunctionCall;
 class UNiagaraStackFunctionInput;
-class UNiagaraStackEditorData;
+class UEdGraphPin;
 
 UCLASS()
-class NIAGARAEDITOR_API UNiagaraStackFunctionInputCollection : public UNiagaraStackEntry
+class NIAGARAEDITOR_API UNiagaraStackFunctionInputCollection : public UNiagaraStackItemContent
 {
 	GENERATED_BODY()
 
-public:
-	DECLARE_MULTICAST_DELEGATE(FOnInputPinnedChanged);
-
-public:
-	struct FDisplayOptions
-	{
-		FText DisplayName;
-		bool bShouldShowInStack;
-		int32 ChildItemIndentLevel;
-		UNiagaraStackEntry::FOnFilterChild ChildFilter;
-	};
-
 public:
 	UNiagaraStackFunctionInputCollection();
 
@@ -34,48 +23,39 @@ public:
 	UNiagaraNodeFunctionCall* GetInputFunctionCallNode() const;
 
 	void Initialize(
-		TSharedRef<FNiagaraSystemViewModel> InSystemViewModel,
-		TSharedRef<FNiagaraEmitterViewModel> InEmitterViewModel,
-		UNiagaraStackEditorData& InStackEditorData,
+		FRequiredEntryData InRequiredEntryData,
 		UNiagaraNodeFunctionCall& InModuleNode,
 		UNiagaraNodeFunctionCall& InInputFunctionCallNode,
-		FDisplayOptions InDisplayOptions);
+		FString InOwnerStackItemEditorDataKey);
 
 	//~ UNiagaraStackEntry interface
 	virtual FText GetDisplayName() const override;
-	virtual FName GetTextStyleName() const override;
-	virtual bool GetCanExpand() const override;
 	virtual bool GetShouldShowInStack() const override;
-	virtual int32 GetErrorCount() const override;
-	virtual bool GetErrorFixable(int32 ErrorIdx) const override;
-	virtual bool TryFixError(int32 ErrorIdx) override;
-	virtual FText GetErrorText(int32 ErrorIdx) const override;
-	virtual FText GetErrorSummaryText(int32 ErrorIdx) const override;
+	virtual bool GetIsEnabled() const;
 
-	FOnInputPinnedChanged& OnInputPinnedChanged();
+	void SetShouldShowInStack(bool bInShouldShowInStack);
 
 protected:
-	virtual void RefreshChildrenInternal(const TArray<UNiagaraStackEntry*>& CurrentChildren, TArray<UNiagaraStackEntry*>& NewChildren) override;
-
+	virtual void RefreshChildrenInternal(const TArray<UNiagaraStackEntry*>& CurrentChildren, TArray<UNiagaraStackEntry*>& NewChildren, TArray<FStackIssue>& NewIssues) override;
+	
 private:
-	void ChildPinnedChanged();
+	void RefreshIssues(TArray<FStackIssue>& NewIssues);
+
+	void OnFunctionInputsChanged();
 
 private:
-	DECLARE_DELEGATE(FFixError);
-	struct FError
+	struct FInputData
 	{
-		FText ErrorText;
-		FText ErrorSummaryText;
-		FFixError Fix;
+		const UEdGraphPin* Pin;
+		FNiagaraTypeDefinition Type;
+		int32 SortKey;
+		FText Category;
 	};
 
-	TArray<FError> Errors;
-
-	UNiagaraStackEditorData* StackEditorData;
 	UNiagaraNodeFunctionCall* ModuleNode;
 	UNiagaraNodeFunctionCall* InputFunctionCallNode;
-
-	FDisplayOptions DisplayOptions;
-
-	FOnInputPinnedChanged InputPinnedChangedDelegate;
+	bool bShouldShowInStack;
+	TArray<FName> ValidAliasedInputNames;
+	TArray<FName> DuplicateInputNames;
+	TArray<const UEdGraphPin*> PinsWithInvalidTypes;
 };
diff --git a/Engine/Plugins/FX/Niagara/Source/NiagaraEditor/Public/ViewModels/Stack/NiagaraStackGraphUtilities.h b/Engine/Plugins/FX/Niagara/Source/NiagaraEditor/Public/ViewModels/Stack/NiagaraStackGraphUtilities.h
--- a/Engine/Plugins/FX/Niagara/Source/NiagaraEditor/Public/ViewModels/Stack/NiagaraStackGraphUtilities.h
+++ b/Engine/Plugins/FX/Niagara/Source/NiagaraEditor/Public/ViewModels/Stack/NiagaraStackGraphUtilities.h
@@ -14,12 +14,14 @@ class UNiagaraNode;
 class UNiagaraNodeInput;
 class UNiagaraNodeOutput;
 class UNiagaraNodeFunctionCall;
+class UNiagaraNodeCustomHlsl;
 class UNiagaraNodeAssignment;
-class UNiagaraMaterialParameterNode;
 class UNiagaraNodeParameterMapSet;
-class UNiagaraStackEditorData;
 class FNiagaraSystemViewModel;
 class FNiagaraEmitterViewModel;
+class UNiagaraStackEditorData;
+class UNiagaraStackEntry;
+class UNiagaraStackErrorItem;
 
 namespace FNiagaraStackGraphUtilities
 {
@@ -56,7 +58,7 @@ namespace FNiagaraStackGraphUtilities
 
 	void ConnectStackNodeGroup(const FStackNodeGroup& ConnectGroup, const FStackNodeGroup& NewPreviousGroup, const FStackNodeGroup& NewNextGroup);
 
-	void InitializeDataInterfaceInputs(TSharedRef<FNiagaraSystemViewModel> SystemViewModel, TSharedRef<FNiagaraEmitterViewModel> EmitterViewModel, UNiagaraStackEditorData& StackEditorData, UNiagaraNodeFunctionCall& ModuleNode, UNiagaraNodeFunctionCall& InputFunctionCallNode);
+	void InitializeStackFunctionInputs(TSharedRef<FNiagaraSystemViewModel> SystemViewModel, TSharedRef<FNiagaraEmitterViewModel> EmitterViewModel, UNiagaraStackEditorData& StackEditorData, UNiagaraNodeFunctionCall& ModuleNode, UNiagaraNodeFunctionCall& InputFunctionCallNode);
 
 	FString GenerateStackFunctionInputEditorDataKey(UNiagaraNodeFunctionCall& FunctionCallNode, FNiagaraParameterHandle InputParameterHandle);
 
@@ -72,43 +74,84 @@ namespace FNiagaraStackGraphUtilities
 
 	UNiagaraNodeParameterMapSet* GetStackFunctionOverrideNode(UNiagaraNodeFunctionCall& FunctionCallNode);
 
-	UNiagaraNodeParameterMapSet& GetOrCreateStackFunctionOverrideNode(UNiagaraNodeFunctionCall& FunctionCallNode);
+	UNiagaraNodeParameterMapSet& GetOrCreateStackFunctionOverrideNode(UNiagaraNodeFunctionCall& FunctionCallNode, const FGuid& PreferredOverrideNodeGuid = FGuid());
 
 	UEdGraphPin* GetStackFunctionInputOverridePin(UNiagaraNodeFunctionCall& StackFunctionCall, FNiagaraParameterHandle AliasedInputParameterHandle);
 
-	UEdGraphPin& GetOrCreateStackFunctionInputOverridePin(UNiagaraNodeFunctionCall& StackFunctionCall, FNiagaraParameterHandle AliasedInputParameterHandle, FNiagaraTypeDefinition InputType);
+	UEdGraphPin& GetOrCreateStackFunctionInputOverridePin(UNiagaraNodeFunctionCall& StackFunctionCall, FNiagaraParameterHandle AliasedInputParameterHandle, FNiagaraTypeDefinition InputType, const FGuid& PreferredOverrideNodeGuid = FGuid());
 
 	void RemoveNodesForStackFunctionInputOverridePin(UEdGraphPin& StackFunctionInputOverridePin);
 
 	void RemoveNodesForStackFunctionInputOverridePin(UEdGraphPin& StackFunctinoInputOverridePin, TArray<TWeakObjectPtr<UNiagaraDataInterface>>& OutRemovedDataObjects);
 
-	void SetLinkedValueHandleForFunctionInput(UEdGraphPin& OverridePin, FNiagaraParameterHandle LinkedParameterHandle);
+	void SetLinkedValueHandleForFunctionInput(UEdGraphPin& OverridePin, FNiagaraParameterHandle LinkedParameterHandle, const FGuid& NewNodePersistentId = FGuid());
+
+	void SetDataValueObjectForFunctionInput(UEdGraphPin& OverridePin, UClass* DataObjectType, FString DataObjectName, UNiagaraDataInterface*& OutDataObject, const FGuid& NewNodePersistentId = FGuid());
 
-	void SetDataValueObjectForFunctionInput(UEdGraphPin& OverridePin, UClass* DataObjectType, FString DataObjectName, UNiagaraDataInterface*& OutDataObject);
+	void SetDynamicInputForFunctionInput(UEdGraphPin& OverridePin, UNiagaraScript* DynamicInput, UNiagaraNodeFunctionCall*& OutDynamicInputFunctionCall, const FGuid& NewNodePersistentId = FGuid(), FString SuggestedName = FString());
 
-	void SetDynamicInputForFunctionInput(UEdGraphPin& OverridePin, UNiagaraScript* DynamicInput, UNiagaraNodeFunctionCall*& OutDynamicInputFunctionCall);
+	void SetCustomExpressionForFunctionInput(UEdGraphPin& OverridePin, UNiagaraNodeCustomHlsl*& OutDynamicInputFunctionCall, const FGuid& NewNodePersistentId = FGuid());
 
 	bool RemoveModuleFromStack(UNiagaraNodeFunctionCall& ModuleNode);
 
+	bool RemoveModuleFromStack(UNiagaraNodeFunctionCall& ModuleNode, TArray<TWeakObjectPtr<UNiagaraNodeInput>>& OutRemovedNodes);
+
 	UNiagaraNodeFunctionCall* AddScriptModuleToStack(FAssetData ModuleScriptAsset, UNiagaraNodeOutput& TargetOutputNode, int32 TargetIndex = INDEX_NONE);
 	
 	bool FindScriptModulesInStack(FAssetData ModuleScriptAsset, UNiagaraNodeOutput& TargetOutputNode, TArray<UNiagaraNodeFunctionCall*> OutFunctionCalls);
 
-	UNiagaraNodeAssignment* AddParameterModuleToStack(FNiagaraVariable ParameterVariable, UNiagaraNodeOutput& TargetOutputNode, int32 TargetIndex = INDEX_NONE, const FString* InDefaultValue = nullptr);
+	UNiagaraNodeAssignment* AddParameterModuleToStack(const TArray<FNiagaraVariable>& ParameterVariables, UNiagaraNodeOutput& TargetOutputNode, int32 TargetIndex, const TArray<FString>& InDefaultValues);
 		
-	UNiagaraMaterialParameterNode* AddMaterialParameterModuleToStack(TSharedRef<FNiagaraEmitterViewModel>& EmitterViewModel, UNiagaraNodeOutput& TargetOutputNode, int32 TargetIndex = INDEX_NONE);
-
 	TOptional<bool> GetModuleIsEnabled(UNiagaraNodeFunctionCall& FunctionCallNode);
 
 	void SetModuleIsEnabled(UNiagaraNodeFunctionCall& FunctionCallNode, bool bIsEnabled);
 
 	bool ValidateGraphForOutput(UNiagaraGraph& NiagaraGraph, ENiagaraScriptUsage ScriptUsage, FGuid ScriptUsageId, FText& ErrorMessage);
 
-	UNiagaraNodeOutput* ResetGraphForOutput(UNiagaraGraph& NiagaraGraph, ENiagaraScriptUsage ScriptUsage, FGuid ScriptUsageId);
+	UNiagaraNodeOutput* ResetGraphForOutput(UNiagaraGraph& NiagaraGraph, ENiagaraScriptUsage ScriptUsage, FGuid ScriptUsageId, const FGuid& PreferredOutputNodeGuid = FGuid(), const FGuid& PreferredInputNodeGuid = FGuid());
 
 	const UNiagaraEmitter* GetBaseEmitter(UNiagaraEmitter& Emitter, UNiagaraSystem& OwningSystem);
 
+	bool IsRapidIterationType(const FNiagaraTypeDefinition& InputType);
+
+	FNiagaraVariable CreateRapidIterationParameter(const FString& UniqueEmitterName, ENiagaraScriptUsage ScriptUsage, const FName& AliasedInputName, const FNiagaraTypeDefinition& InputType);
+
 	void CleanUpStaleRapidIterationParameters(UNiagaraScript& Script, UNiagaraEmitter& OwningEmitter);
 
 	void CleanUpStaleRapidIterationParameters(UNiagaraEmitter& Emitter);
+
+	void GetNewParameterAvailableTypes(TArray<FNiagaraTypeDefinition>& OutAvailableTypes);
+
+	void GetScriptAssetsByUsage(ENiagaraScriptUsage AssetUsage, ENiagaraScriptUsage TargetUsage, TArray<FAssetData>& OutAssets);
+
+	void GetScriptAssetsByDependencyProvided(ENiagaraScriptUsage AssetUsage, FName DependencyName, TArray<FAssetData>& OutAssets);
+
+	void GetAvailableParametersForScript(UNiagaraNodeOutput& ScriptOutputNode, TArray<FNiagaraVariable>& OutAvailableParameters);
+
+	TOptional<FName> GetNamespaceForScriptUsage(ENiagaraScriptUsage ScriptUsage);
+
+	void GetOwningEmitterAndScriptForStackNode(UNiagaraNode& StackNode, UNiagaraSystem& OwningSystem, UNiagaraEmitter*& OutEmitter, UNiagaraScript*& OutScript);
+
+	bool IsValidDefaultDynamicInput(UNiagaraScript& OwningScript, UEdGraphPin& DefaultPin);
+
+	bool DoesDynamicInputMatchDefault(
+		FString EmitterUniqueName,
+		UNiagaraScript& OwningScript,
+		UNiagaraNodeFunctionCall& OwningFunctionCallNode,
+		UEdGraphPin& OverridePin,
+		FName InputName,
+		UEdGraphPin& DefaultPin);
+
+	void ResetToDefaultDynamicInput(
+		TSharedRef<FNiagaraSystemViewModel> SystemViewModel,
+		TSharedRef<FNiagaraEmitterViewModel> EmitterViewModel,
+		UNiagaraStackEditorData& StackEditorData,
+		UNiagaraScript& SourceScript,
+		const TArray<TWeakObjectPtr<UNiagaraScript>> AffectedScripts,
+		UNiagaraNodeFunctionCall& ModuleNode,
+		UNiagaraNodeFunctionCall& InputFunctionCallNode,
+		FName InputName,
+		UEdGraphPin& DefaultPin);
+	
+	bool GetStackIssuesRecursively(const UNiagaraStackEntry* const Entry, TArray<UNiagaraStackErrorItem*>& OutIssues);
 }
diff --git a/Engine/Plugins/FX/Niagara/Source/NiagaraEditor/Public/ViewModels/Stack/NiagaraStackInputCategory.h b/Engine/Plugins/FX/Niagara/Source/NiagaraEditor/Public/ViewModels/Stack/NiagaraStackInputCategory.h
new file mode 100644
--- /dev/null
+++ b/Engine/Plugins/FX/Niagara/Source/NiagaraEditor/Public/ViewModels/Stack/NiagaraStackInputCategory.h
@@ -0,0 +1,59 @@
+// Copyright 1998-2018 Epic Games, Inc. All Rights Reserved.
+
+#pragma once
+
+#include "ViewModels/Stack/NiagaraStackItem.h"
+#include "NiagaraTypes.h"
+#include "NiagaraStackInputCategory.generated.h"
+
+class UNiagaraStackFunctionInput;
+class UNiagaraNodeFunctionCall;
+
+UCLASS()
+class NIAGARAEDITOR_API UNiagaraStackInputCategory : public UNiagaraStackItemContent
+{
+	GENERATED_BODY() 
+
+public:
+	void Initialize(
+		FRequiredEntryData InRequiredEntryData,
+		UNiagaraNodeFunctionCall& InModuleNode,
+		UNiagaraNodeFunctionCall& InInputFunctionCallNode, 
+		FText InCategoryName,
+		FString InOwnerStackItemEditorDataKey);
+	
+	const FText& GetCategoryName() const;
+
+	void ResetInputs();
+
+	void AddInput(FName InInputParameterHandle, FNiagaraTypeDefinition InInputType);
+
+	//~ UNiagaraStackEntry interface
+	virtual FText GetDisplayName() const override;
+	virtual bool GetShouldShowInStack() const override;
+	virtual EStackRowStyle GetStackRowStyle() const override;
+	virtual bool GetIsEnabled() const override;
+
+	void SetShouldShowInStack(bool bInShouldShowInStack);
+
+protected:
+	//~ UNiagaraStackEntry interface
+	virtual void RefreshChildrenInternal(const TArray<UNiagaraStackEntry*>& CurrentChildren, TArray<UNiagaraStackEntry*>& NewChildren, TArray<FStackIssue>& NewIssues) override;
+
+private:
+	bool FilterForVisibleCondition(const UNiagaraStackEntry& Child) const;
+	bool FilterForIsInlineEditConditionToggle(const UNiagaraStackEntry& Child) const;
+
+private:
+	struct FInputParameterHandleAndType
+	{
+		FName ParameterHandle;
+		FNiagaraTypeDefinition Type;
+	};
+
+	UNiagaraNodeFunctionCall* ModuleNode;
+	UNiagaraNodeFunctionCall* InputFunctionCallNode;
+	FText CategoryName;
+	TArray<FInputParameterHandleAndType> Inputs;
+	bool bShouldShowInStack;
+};
\ No newline at end of file
diff --git a/Engine/Plugins/FX/Niagara/Source/NiagaraEditor/Public/ViewModels/Stack/NiagaraStackItem.h b/Engine/Plugins/FX/Niagara/Source/NiagaraEditor/Public/ViewModels/Stack/NiagaraStackItem.h
--- a/Engine/Plugins/FX/Niagara/Source/NiagaraEditor/Public/ViewModels/Stack/NiagaraStackItem.h
+++ b/Engine/Plugins/FX/Niagara/Source/NiagaraEditor/Public/ViewModels/Stack/NiagaraStackItem.h
@@ -3,9 +3,12 @@
 #pragma once
 
 #include "ViewModels/Stack/NiagaraStackEntry.h"
+#include "Layout/Visibility.h"
 #include "NiagaraStackItem.generated.h"
 
-class UNiagaraStackEditorData;
+class UNiagaraStackSpacer;
+class UNiagaraStackAdvancedExpander;
+class UNiagaraNode;
 
 UCLASS()
 class NIAGARAEDITOR_API UNiagaraStackItem : public UNiagaraStackEntry
@@ -16,16 +19,72 @@ public:
 	DECLARE_DELEGATE(FOnModifiedGroupItems);
 
 public:
-	void Initialize(TSharedRef<FNiagaraSystemViewModel> InSystemViewModel, TSharedRef<FNiagaraEmitterViewModel> InEmitterViewModel, UNiagaraStackEditorData& InStackEditorData);
+	void Initialize(FRequiredEntryData InRequiredEntryData, FString InStackEditorDataKey);
+
+	virtual EStackRowStyle GetStackRowStyle() const override;
 
 	void SetOnModifiedGroupItems(FOnModifiedGroupItems OnModifiedGroupItems);
 
+	uint32 GetRecursiveStackIssuesCount() const;
+	EStackIssueSeverity GetHighestStackIssueSeverity() const;
+	
 protected:
-	UNiagaraStackEditorData& GetStackEditorData() const;
+	virtual void RefreshChildrenInternal(const TArray<UNiagaraStackEntry*>& CurrentChildren, TArray<UNiagaraStackEntry*>& NewChildren, TArray<FStackIssue>& NewIssues) override;
+
+	virtual void PostRefreshChildrenInternal() override;
+
+	virtual int32 GetChildIndentLevel() const override;
+
+	virtual UNiagaraNode* GetOwningNiagaraNode() const;
+
+	virtual void ChlildStructureChangedInternal() override;
+
+private:
+	bool FilterAdvancedChildren(const UNiagaraStackEntry& Child) const;
+
+	bool FilterShowAdvancedChild(const UNiagaraStackEntry& Child) const;
+
+	void ToggleShowAdvanced();
 
 protected:
 	FOnModifiedGroupItems ModifiedGroupItemsDelegate;
 
 private:
-	UNiagaraStackEditorData* StackEditorData;
+	bool bHasAdvancedContent;
+
+	UPROPERTY()
+	UNiagaraStackSpacer* FooterSpacer;
+
+	UPROPERTY()
+	UNiagaraStackAdvancedExpander* ShowAdvancedExpander;
+
+	/** How many errors this entry has along its tree. */
+	mutable TOptional<uint32> RecursiveStackIssuesCount;
+	/** The highest severity of issues along this entry's tree. */
+	mutable TOptional<EStackIssueSeverity> HighestIssueSeverity;
+};
+
+UCLASS()
+class NIAGARAEDITOR_API UNiagaraStackItemContent : public UNiagaraStackEntry
+{
+	GENERATED_BODY()
+
+public:
+	void Initialize(FRequiredEntryData InRequiredEntryData, bool bInIsAdvanced, FString InOwningStackItemEditorDataKey, FString InStackEditorDataKey);
+
+	virtual EStackRowStyle GetStackRowStyle() const override;
+
+	bool GetIsAdvanced() const;
+
+protected:
+	FString GetOwnerStackItemEditorDataKey() const;
+
+	void SetIsAdvanced(bool bInIsAdvanced);
+
+private:
+	bool FilterAdvancedChildren(const UNiagaraStackEntry& Child) const;
+
+private:
+	FString OwningStackItemEditorDataKey;
+	bool bIsAdvanced;
 };
diff --git a/Engine/Plugins/FX/Niagara/Source/NiagaraEditor/Public/ViewModels/Stack/NiagaraStackItemGroup.h b/Engine/Plugins/FX/Niagara/Source/NiagaraEditor/Public/ViewModels/Stack/NiagaraStackItemGroup.h
--- a/Engine/Plugins/FX/Niagara/Source/NiagaraEditor/Public/ViewModels/Stack/NiagaraStackItemGroup.h
+++ b/Engine/Plugins/FX/Niagara/Source/NiagaraEditor/Public/ViewModels/Stack/NiagaraStackItemGroup.h
@@ -3,39 +3,48 @@
 #pragma once
 
 #include "ViewModels/Stack/NiagaraStackEntry.h"
+#include "ViewModels/Stack/INiagaraStackItemGroupAddUtilities.h"
+#include "Layout/Visibility.h"
 #include "NiagaraStackItemGroup.generated.h"
 
-class UNiagaraStackEditorData;
-
 UCLASS()
 class NIAGARAEDITOR_API UNiagaraStackItemGroup : public UNiagaraStackEntry
 {
 	GENERATED_BODY()
 
 public:
-	void Initialize(TSharedRef<FNiagaraSystemViewModel> InSystemViewModel, TSharedRef<FNiagaraEmitterViewModel> InEmitterViewModel, UNiagaraStackEditorData& InStackEditorData);
+	void Initialize(FRequiredEntryData InRequiredEntryData, FText InDisplayName, FText InToolTip, INiagaraStackItemGroupAddUtilities* InAddUtilities);
 
 	//~ UNiagaraStackEntry interface
-	virtual bool GetCanExpand() const override;
-	virtual bool GetIsExpanded() const override;
-	virtual void SetIsExpanded(bool bInExpanded) override;
-	virtual FName GetItemBackgroundName() const override;
-	virtual FName GetTextStyleName() const override;
+	virtual FText GetDisplayName() const override;
+	virtual EStackRowStyle GetStackRowStyle() const override;
 	virtual FText GetTooltipText() const override;
 
-	virtual void SetTooltipText(const FText& InText);
-
 	virtual bool CanDelete() const { return false; }
 	virtual bool Delete() { return false; }
 
-	virtual bool CanAdd() const { return false; }
-	virtual bool Add() { return false; }
+	INiagaraStackItemGroupAddUtilities* GetAddUtilities() const;
+
+	uint32 GetRecursiveStackIssuesCount() const;
+	EStackIssueSeverity GetHighestStackIssueSeverity() const;
 	
 protected:
-	UNiagaraStackEditorData& GetStackEditorData() const;
+	void SetDisplayName(FText InDisplayName);
+	
+	virtual void RefreshChildrenInternal(const TArray<UNiagaraStackEntry*>& CurrentChildren, TArray<UNiagaraStackEntry*>& NewChildren, TArray<FStackIssue>& NewIssues) override;
+
+	virtual int32 GetChildIndentLevel() const override;
+
+	virtual void ChlildStructureChangedInternal() override;
 
 private:
-	UNiagaraStackEditorData* StackEditorData;
+	INiagaraStackItemGroupAddUtilities* AddUtilities;
+
+	FText GroupDisplayName;
+	FText GroupToolTip;
 
-	FText TooltipText;
+	/** How many errors this entry has along its tree. */
+	mutable TOptional<uint32> RecursiveStackIssuesCount;
+	/** The highest severity of issues along this entry's tree. */
+	mutable TOptional<EStackIssueSeverity> HighestIssueSeverity;
 };
diff --git a/Engine/Plugins/FX/Niagara/Source/NiagaraEditor/Public/ViewModels/Stack/NiagaraStackModuleItem.h b/Engine/Plugins/FX/Niagara/Source/NiagaraEditor/Public/ViewModels/Stack/NiagaraStackModuleItem.h
--- a/Engine/Plugins/FX/Niagara/Source/NiagaraEditor/Public/ViewModels/Stack/NiagaraStackModuleItem.h
+++ b/Engine/Plugins/FX/Niagara/Source/NiagaraEditor/Public/ViewModels/Stack/NiagaraStackModuleItem.h
@@ -9,8 +9,8 @@ class UNiagaraNodeFunctionCall;
 class UNiagaraNode;
 class UNiagaraStackFunctionInputCollection;
 class UNiagaraStackModuleItemOutputCollection;
-class UNiagaraStackItemExpander;
-class UNiagaraStackEditorData;
+class INiagaraStackItemGroupAddUtilities;
+struct FAssetData;
 
 UCLASS()
 class NIAGARAEDITOR_API UNiagaraStackModuleItem : public UNiagaraStackItem
@@ -21,62 +21,58 @@ public:
 	UNiagaraStackModuleItem();
 
 	const UNiagaraNodeFunctionCall& GetModuleNode() const;
-	void Initialize(TSharedRef<FNiagaraSystemViewModel> InSystemViewModel, TSharedRef<FNiagaraEmitterViewModel> InEmitterViewModel, UNiagaraStackEditorData& InStackEditorData, UNiagaraNodeFunctionCall& InFunctionCallNode);
+
+	UNiagaraNodeFunctionCall& GetModuleNode();
+
+	void Initialize(FRequiredEntryData InRequiredEntryData, INiagaraStackItemGroupAddUtilities* GroupAddUtilities, UNiagaraNodeFunctionCall& InFunctionCallNode);
 
 	virtual FText GetDisplayName() const override;
 	virtual FText GetTooltipText() const override;
 
+	INiagaraStackItemGroupAddUtilities* GetGroupAddUtilities();
+
 	bool CanMoveAndDelete() const;
 	bool CanRefresh() const;
 	void Refresh();
 
-	bool GetIsEnabled() const;
+	virtual bool GetIsEnabled() const override;
 	void SetIsEnabled(bool bInIsEnabled);
 
-	void MoveUp();
-	void MoveDown();
 	void Delete();
 
+	int32 GetModuleIndex();
+	
+	UObject* GetExternalAsset() const override;
+
+	virtual bool CanDrag() const override;
+
+	/** Gets the output node of this module. */
+	class UNiagaraNodeOutput* GetOutputNode() const;
+
+	void NotifyModuleMoved();
+
 protected:
-	virtual void RefreshChildrenInternal(const TArray<UNiagaraStackEntry*>& CurrentChildren, TArray<UNiagaraStackEntry*>& NewChildren) override;
+	virtual void RefreshChildrenInternal(const TArray<UNiagaraStackEntry*>& CurrentChildren, TArray<UNiagaraStackEntry*>& NewChildren, TArray<FStackIssue>& NewIssues) override;
 
 private:
-	void RefreshIsEnabled();
-	void InputPinnedChanged();
-	void ModuleExpandedChanged();
+	bool FilterOutputCollection(const UNiagaraStackEntry& Child) const;
+	void RefreshIssues(TArray<FStackIssue>& NewIssues);
 
 private:
-	virtual int32 GetErrorCount() const override;
-	virtual bool GetErrorFixable(int32 ErrorIdx) const override;
-	virtual bool TryFixError(int32 ErrorIdx) override;
-	virtual FText GetErrorText(int32 ErrorIdx) const override;
-	virtual FText GetErrorSummaryText(int32 ErrorIdx) const override;
+	void RefreshIsEnabled();
+	void InputPinnedChanged();
 
 private:
 	UNiagaraNodeFunctionCall* FunctionCallNode;
-	FString ModuleEditorDataKey;
 	bool bCanMoveAndDelete;
 	bool bIsEnabled;
-
-	DECLARE_DELEGATE_RetVal(bool, FFixError);
-	struct FError
-	{
-		FText ErrorText;
-		FText ErrorSummaryText;
-		FFixError Fix;
-	};
-
-	TArray<FError> Errors;
-
-	UPROPERTY()
-	UNiagaraStackFunctionInputCollection* PinnedInputCollection;
+	bool bCanRefresh;
 
 	UPROPERTY()
-	UNiagaraStackFunctionInputCollection* UnpinnedInputCollection;
+	UNiagaraStackFunctionInputCollection* InputCollection;
 
 	UPROPERTY()
 	UNiagaraStackModuleItemOutputCollection* OutputCollection;
 
-	UPROPERTY()
-	UNiagaraStackItemExpander* ModuleExpander;
+	INiagaraStackItemGroupAddUtilities* GroupAddUtilities;
 };
diff --git a/Engine/Plugins/FX/Niagara/Source/NiagaraEditor/Public/ViewModels/Stack/NiagaraStackModuleItemOutput.h b/Engine/Plugins/FX/Niagara/Source/NiagaraEditor/Public/ViewModels/Stack/NiagaraStackModuleItemOutput.h
--- a/Engine/Plugins/FX/Niagara/Source/NiagaraEditor/Public/ViewModels/Stack/NiagaraStackModuleItemOutput.h
+++ b/Engine/Plugins/FX/Niagara/Source/NiagaraEditor/Public/ViewModels/Stack/NiagaraStackModuleItemOutput.h
@@ -23,15 +23,13 @@ public:
 	 * @param InFunctionCallNode The function call node representing the module in the stack graph which owns this Output.
 	 * @param InOutputParameterHandle The Namespace.Name handle for the Output to the owning module.
 	 */
-	void Initialize(TSharedRef<FNiagaraSystemViewModel> InSystemViewModel, TSharedRef<FNiagaraEmitterViewModel> InEmitterViewModel, UNiagaraNodeFunctionCall& InFunctionCallNode, FName InOutputParameterHandle,
-		FNiagaraTypeDefinition InOutputType);
+	void Initialize(FRequiredEntryData InRequiredEntryData, UNiagaraNodeFunctionCall& InFunctionCallNode, FName InOutputParameterHandle, FNiagaraTypeDefinition InOutputType);
 
 	//~ UNiagaraStackEntry interface
 	virtual FText GetDisplayName() const override;
-	virtual FName GetTextStyleName() const override;
-	virtual bool GetCanExpand() const override;
-	virtual int32 GetItemIndentLevel() const override;
 	virtual FText GetTooltipText() const override;
+	virtual bool GetIsEnabled() const override;
+	virtual EStackRowStyle GetStackRowStyle() const override;
 
 	/** Gets the parameter handle which defined this Output in the module. */
 	const FNiagaraParameterHandle& GetOutputParameterHandle() const;
diff --git a/Engine/Plugins/FX/Niagara/Source/NiagaraEditor/Public/ViewModels/Stack/NiagaraStackModuleItemOutputCollection.h b/Engine/Plugins/FX/Niagara/Source/NiagaraEditor/Public/ViewModels/Stack/NiagaraStackModuleItemOutputCollection.h
--- a/Engine/Plugins/FX/Niagara/Source/NiagaraEditor/Public/ViewModels/Stack/NiagaraStackModuleItemOutputCollection.h
+++ b/Engine/Plugins/FX/Niagara/Source/NiagaraEditor/Public/ViewModels/Stack/NiagaraStackModuleItemOutputCollection.h
@@ -15,20 +15,17 @@ class NIAGARAEDITOR_API UNiagaraStackModuleItemOutputCollection : public UNiagar
 public:
 	UNiagaraStackModuleItemOutputCollection();
 
-	void Initialize(TSharedRef<FNiagaraSystemViewModel> InSystemViewModel, TSharedRef<FNiagaraEmitterViewModel> InEmitterViewModel, UNiagaraNodeFunctionCall& InFunctionCallNode);
-
+	void Initialize(FRequiredEntryData InRequiredEntryData, UNiagaraNodeFunctionCall& InFunctionCallNode);
+	
 	//~ UNiagaraStackEntry interface
 	virtual FText GetDisplayName() const override;
-	virtual FName GetTextStyleName() const override;
-	virtual bool GetCanExpand() const override;
 	virtual bool IsExpandedByDefault() const override;
-
-	void SetDisplayName(FText InDisplayName);
+	virtual bool GetIsEnabled() const override;
+	virtual EStackRowStyle GetStackRowStyle() const override;
 
 protected:
-	virtual void RefreshChildrenInternal(const TArray<UNiagaraStackEntry*>& CurrentChildren, TArray<UNiagaraStackEntry*>& NewChildren) override;
+	virtual void RefreshChildrenInternal(const TArray<UNiagaraStackEntry*>& CurrentChildren, TArray<UNiagaraStackEntry*>& NewChildren, TArray<FStackIssue>& NewIssues) override;
 
 private:
 	UNiagaraNodeFunctionCall* FunctionCallNode;
-	FText DisplayName;
 };
diff --git a/Engine/Plugins/FX/Niagara/Source/NiagaraEditor/Public/ViewModels/Stack/NiagaraStackObject.h b/Engine/Plugins/FX/Niagara/Source/NiagaraEditor/Public/ViewModels/Stack/NiagaraStackObject.h
--- a/Engine/Plugins/FX/Niagara/Source/NiagaraEditor/Public/ViewModels/Stack/NiagaraStackObject.h
+++ b/Engine/Plugins/FX/Niagara/Source/NiagaraEditor/Public/ViewModels/Stack/NiagaraStackObject.h
@@ -2,31 +2,66 @@
 
 #pragma once
 
-#include "ViewModels/Stack/NiagaraStackEntry.h"
+#include "ViewModels/Stack/NiagaraStackItem.h"
+#include "Misc/NotifyHook.h"
+#include "PropertyEditorDelegates.h"
 #include "NiagaraStackObject.generated.h"
 
+class IPropertyRowGenerator;
+class UNiagaraNode;
+class IDetailTreeNode;
+
 UCLASS()
-class NIAGARAEDITOR_API UNiagaraStackObject : public UNiagaraStackEntry
+class NIAGARAEDITOR_API UNiagaraStackObject : public UNiagaraStackItemContent, public FNotifyHook
 {
 	GENERATED_BODY()
+		
+public:
+	DECLARE_DELEGATE_TwoParams(FOnSelectRootNodes, TArray<TSharedRef<IDetailTreeNode>>, TArray<TSharedRef<IDetailTreeNode>>*);
 
 public:
 	UNiagaraStackObject();
 
-	void Initialize(TSharedRef<FNiagaraSystemViewModel> InSystemViewModel, TSharedRef<FNiagaraEmitterViewModel> InEmitterViewModel, UObject* InObject);
+	void Initialize(FRequiredEntryData InRequiredEntryData, UObject* InObject, FString InOwnerStackItemEditorDataKey, UNiagaraNode* InOwningNiagaraNode = nullptr);
+
+	void SetOnSelectRootNodes(FOnSelectRootNodes OnSelectRootNodes);
+
+	void RegisterInstancedCustomPropertyLayout(UStruct* Class, FOnGetDetailCustomizationInstance DetailLayoutDelegate);
+	void RegisterInstancedCustomPropertyTypeLayout(FName PropertyTypeName, FOnGetPropertyTypeCustomizationInstance PropertyTypeLayoutDelegate, TSharedPtr<IPropertyTypeIdentifier> Identifier = nullptr);
 
 	UObject* GetObject();
 
+	//~ FNotifyHook interface
+	virtual void NotifyPostChange(const FPropertyChangedEvent& PropertyChangedEvent, UProperty* PropertyThatChanged) override;
+
 	//~ UNiagaraStackEntry interface
-	virtual int32 GetItemIndentLevel() const override;
+	virtual bool GetIsEnabled() const override;
+	virtual bool GetShouldShowInStack() const override;
 
-	/** Sets the item indent level for this stack entry. */
-	void SetItemIndentLevel(int32 InItemIndentLevel);
+protected:
+	virtual void RefreshChildrenInternal(const TArray<UNiagaraStackEntry*>& CurrentChildren, TArray<UNiagaraStackEntry*>& NewChildren, TArray<FStackIssue>& NewIssues);
 
-	/** Notifies the stack entry that it's object has been modified. */
-	void NotifyObjectChanged();
+private:
+	void PropertyRowsRefreshed();
 
 private:
+	struct FRegisteredClassCustomization
+	{
+		UStruct* Class;
+		FOnGetDetailCustomizationInstance DetailLayoutDelegate;
+	};
+
+	struct FRegisteredPropertyCustomization
+	{
+		FName PropertyTypeName;
+		FOnGetPropertyTypeCustomizationInstance PropertyTypeLayoutDelegate;
+		TSharedPtr<IPropertyTypeIdentifier> Identifier;
+	};
+
 	UObject* Object;
-	int32 ItemIndentLevel;
+	UNiagaraNode* OwningNiagaraNode;
+	FOnSelectRootNodes OnSelectRootNodesDelegate;
+	TArray<FRegisteredClassCustomization> RegisteredClassCustomizations;
+	TArray<FRegisteredPropertyCustomization> RegisteredPropertyCustomizations;
+	TSharedPtr<IPropertyRowGenerator> PropertyRowGenerator;
 };
diff --git a/Engine/Plugins/FX/Niagara/Source/NiagaraEditor/Public/ViewModels/Stack/NiagaraStackParameterStoreEntry.h b/Engine/Plugins/FX/Niagara/Source/NiagaraEditor/Public/ViewModels/Stack/NiagaraStackParameterStoreEntry.h
--- a/Engine/Plugins/FX/Niagara/Source/NiagaraEditor/Public/ViewModels/Stack/NiagaraStackParameterStoreEntry.h
+++ b/Engine/Plugins/FX/Niagara/Source/NiagaraEditor/Public/ViewModels/Stack/NiagaraStackParameterStoreEntry.h
@@ -1,8 +1,7 @@
 // Copyright 1998-2018 Epic Games, Inc. All Rights Reserved.
 
 #pragma once
-
-#include "ViewModels/Stack/NiagaraStackEntry.h"
+#include "ViewModels/Stack/NiagaraStackItem.h"
 #include "NiagaraTypes.h"
 #include "ViewModels/Stack/NiagaraParameterHandle.h"
 #include "NiagaraParameterStore.h"
@@ -12,18 +11,18 @@ class UNiagaraNodeFunctionCall;
 class UNiagaraNodeAssignment;
 class UNiagaraNodeParameterMapSet;
 class FStructOnScope;
-class UNiagaraStackEditorData;
 class UNiagaraStackObject;
+class FNiagaraSystemViewModel;
+class UEdGraphPin;
 
 /** Represents a single module input in the module stack view model. */
 UCLASS()
-class NIAGARAEDITOR_API UNiagaraStackParameterStoreEntry : public UNiagaraStackEntry
+class NIAGARAEDITOR_API UNiagaraStackParameterStoreEntry : public UNiagaraStackItemContent
 {
 	GENERATED_BODY()
 
 public:
 	DECLARE_MULTICAST_DELEGATE(FOnValueChanged);
-
 public:
 	UNiagaraStackParameterStoreEntry();
 
@@ -35,27 +34,21 @@ public:
 	 * @param InStackEditorData The stack editor data for this input.
 	 * @param InInputParameterHandle The input parameter handle for the function call.
 	 * @param InInputType The type of this input.
+	 * @param InOwnerStackItemEditorDataKey The editor data key for the item that owns this entry.
 	 */
 	void Initialize(
-		TSharedRef<FNiagaraSystemViewModel> InSystemViewModel,
-		TSharedRef<FNiagaraEmitterViewModel> InEmitterViewModel,
-		UNiagaraStackEditorData& InStackEditorData,
+		FRequiredEntryData InRequiredEntryData,
 		UObject* InOwner,
 		FNiagaraParameterStore* InParameterStore,
 		FString InInputParameterHandle,
-		FNiagaraTypeDefinition InInputType);
+		FNiagaraTypeDefinition InInputType,
+		FString InOwnerStackItemEditorDataKey);
 
 	/** Gets the type of this input. */
 	const FNiagaraTypeDefinition& GetInputType() const;
 
 	//~ UNiagaraStackEntry interface
 	virtual FText GetDisplayName() const override;
-	virtual FName GetTextStyleName() const override;
-	virtual bool GetCanExpand() const override;
-	virtual int32 GetItemIndentLevel() const override;
-
-	/** Sets the item indent level for this stack entry. */
-	void SetItemIndentLevel(int32 InItemIndentLevel);
 
 	/** Gets the current struct value of this input is there is one. */
 	TSharedPtr<FStructOnScope> GetValueStruct();
@@ -68,7 +61,7 @@ public:
 
 	/** Called to notify the input that an ongoing change to it's value has ended. */
 	void NotifyEndValueChange();
-
+	
 	/** Called to notify the input that it's value has been changed. */
 	void NotifyValueChanged();
 
@@ -90,20 +83,28 @@ public:
 	/** Renames this input to the name specified. */
 	void RenameInput(FString NewName);
 
+	/** Checks if the chosen name is unique (not duplicate) */
+	bool IsUniqueName(FString NewName);
+
 	/** Gets a multicast delegate which is called whenever the value on this input changes. */
 	FOnValueChanged& OnValueChanged();
 
+	/** Delete the parameter from the ParameterStore and notify that the store changed. */
+	void Delete();
+
 protected:
 	//~ UNiagaraStackEntry interface
-	virtual void RefreshChildrenInternal(const TArray<UNiagaraStackEntry*>& CurrentChildren, TArray<UNiagaraStackEntry*>& NewChildren) override;
+	virtual void RefreshChildrenInternal(const TArray<UNiagaraStackEntry*>& CurrentChildren, TArray<UNiagaraStackEntry*>& NewChildren, TArray<FStackIssue>& NewIssues) override;
 
 	void RefreshValueAndHandle();
 	TSharedPtr<FNiagaraVariable> GetCurrentValueVariable();
 	UNiagaraDataInterface* GetCurrentValueObject();
+
 private:
-	/** The stack editor data for this function input. */
-	UNiagaraStackEditorData* StackEditorData;
+	void RemovePins(TArray<UEdGraphPin*> PinsToRemove);
+	TArray<UEdGraphPin*> GetOwningPins(); 
 
+private:
 	/** */
 	FName ParameterName;
 
@@ -123,9 +124,6 @@ private:
 	/** A multicast delegate which is called when the value of this input is changed. */
 	FOnValueChanged ValueChangedDelegate;
 
-	/** The item indent level for this stack entry. */
-	int32 ItemIndentLevel;
-
 	UPROPERTY()
 	UObject* Owner;
 
diff --git a/Engine/Plugins/FX/Niagara/Source/NiagaraEditor/Public/ViewModels/Stack/NiagaraStackParameterStoreGroup.h b/Engine/Plugins/FX/Niagara/Source/NiagaraEditor/Public/ViewModels/Stack/NiagaraStackParameterStoreGroup.h
--- a/Engine/Plugins/FX/Niagara/Source/NiagaraEditor/Public/ViewModels/Stack/NiagaraStackParameterStoreGroup.h
+++ b/Engine/Plugins/FX/Niagara/Source/NiagaraEditor/Public/ViewModels/Stack/NiagaraStackParameterStoreGroup.h
@@ -5,41 +5,51 @@
 #include "ViewModels/Stack/NiagaraStackItemGroup.h"
 #include "NiagaraCommon.h"
 #include "NiagaraParameterStore.h"
+#include "ViewModels/Stack/NiagaraStackParameterStoreEntry.h"
 #include "NiagaraStackParameterStoreGroup.generated.h"
 
 class FNiagaraScriptViewModel;
-class UNiagaraStackEditorData;
 
 UCLASS()
 class NIAGARAEDITOR_API UNiagaraStackParameterStoreGroup : public UNiagaraStackItemGroup
 {
 	GENERATED_BODY()
-
+		
 public:
-	void Initialize(
-		TSharedRef<FNiagaraSystemViewModel> InSystemViewModel,
-		TSharedRef<FNiagaraEmitterViewModel> InEmitterViewModel,
-		UNiagaraStackEditorData& InStackEditorData,
-		TSharedRef<FNiagaraScriptViewModel> InScriptViewModel,
-		UObject* InOwner,
-		FNiagaraParameterStore* InParameterStore);
-	
-	virtual FText GetDisplayName() const override;
-	void SetDisplayName(FText InDisplayName);
+	void Initialize(FRequiredEntryData InRequiredEntryData,	UObject* InOwner, FNiagaraParameterStore* InParameterStore);
+
+	void AddUserParameter(FNiagaraVariable ParameterVariable);
 
 protected:
-	virtual void RefreshChildrenInternal(const TArray<UNiagaraStackEntry*>& CurrentChildren, TArray<UNiagaraStackEntry*>& NewChildren) override;
+	virtual void RefreshChildrenInternal(const TArray<UNiagaraStackEntry*>& CurrentChildren, TArray<UNiagaraStackEntry*>& NewChildren, TArray<FStackIssue>& NewIssues) override;
 
 private:
-	void ItemAdded();
+	TWeakObjectPtr<UObject> Owner;
+	FNiagaraParameterStore* ParameterStore;
+	FDelegateHandle ParameterStoreChangedHandle;
+	TSharedPtr<INiagaraStackItemGroupAddUtilities> AddUtilities;
+};
 
-	void ChildModifiedGroupItems();
+UCLASS()
+class UNiagaraStackParameterStoreItem : public UNiagaraStackItem
+{
+	GENERATED_BODY()
 
-private:
+public:
+	void Initialize(FRequiredEntryData InRequiredEntryData, UObject* InOwner, FNiagaraParameterStore* InParameterStore);
+
+	virtual FText GetDisplayName() const override;
 
-	FText DisplayName;
+	virtual void BeginDestroy() override;
+
+protected:
+	virtual void RefreshChildrenInternal(const TArray<UNiagaraStackEntry*>& CurrentChildren, TArray<UNiagaraStackEntry*>& NewChildren, TArray<FStackIssue>& NewIssues) override;
 
-	UPROPERTY()
-	UObject* Owner;
+private:
+	void ParameterStoreChanged();
+
+private:
+	TWeakObjectPtr<UObject> Owner;
 	FNiagaraParameterStore* ParameterStore;
-};
+	FDelegateHandle ParameterStoreChangedHandle;
+};
\ No newline at end of file
diff --git a/Engine/Plugins/FX/Niagara/Source/NiagaraEditor/Public/ViewModels/Stack/NiagaraStackPropertyRow.h b/Engine/Plugins/FX/Niagara/Source/NiagaraEditor/Public/ViewModels/Stack/NiagaraStackPropertyRow.h
new file mode 100644
--- /dev/null
+++ b/Engine/Plugins/FX/Niagara/Source/NiagaraEditor/Public/ViewModels/Stack/NiagaraStackPropertyRow.h
@@ -0,0 +1,32 @@
+// Copyright 1998-2018 Epic Games, Inc. All Rights Reserved.
+
+#pragma once
+
+#include "ViewModels/Stack/NiagaraStackItem.h"
+#include "NiagaraStackPropertyRow.generated.h"
+
+class IDetailTreeNode;
+class UNiagaraNode;
+
+UCLASS()
+class NIAGARAEDITOR_API UNiagaraStackPropertyRow : public UNiagaraStackItemContent
+{
+	GENERATED_BODY()
+		
+public:
+	void Initialize(FRequiredEntryData InRequiredEntryData, TSharedRef<IDetailTreeNode> InDetailTreeNode, FString InOwnerStackItemEditorDataKey, FString InOwnerStackEditorDataKey, UNiagaraNode* InOwningNiagaraNode);
+
+	TSharedRef<IDetailTreeNode> GetDetailTreeNode() const;
+
+	virtual bool GetIsEnabled() const override;
+
+protected:
+	virtual void RefreshChildrenInternal(const TArray<UNiagaraStackEntry*>& CurrentChildren, TArray<UNiagaraStackEntry*>& NewChildren, TArray<FStackIssue>& NewIssues) override;
+
+	virtual void GetAdditionalSearchItemsInternal(TArray<FStackSearchItem>& SearchItems) const override;
+
+private:
+	TSharedPtr<IDetailTreeNode> DetailTreeNode;
+	UNiagaraNode* OwningNiagaraNode;
+	EStackRowStyle RowStyle;
+};
\ No newline at end of file
diff --git a/Engine/Plugins/FX/Niagara/Source/NiagaraEditor/Public/ViewModels/Stack/NiagaraStackRenderItemGroup.h b/Engine/Plugins/FX/Niagara/Source/NiagaraEditor/Public/ViewModels/Stack/NiagaraStackRenderItemGroup.h
--- a/Engine/Plugins/FX/Niagara/Source/NiagaraEditor/Public/ViewModels/Stack/NiagaraStackRenderItemGroup.h
+++ b/Engine/Plugins/FX/Niagara/Source/NiagaraEditor/Public/ViewModels/Stack/NiagaraStackRenderItemGroup.h
@@ -11,15 +11,15 @@ class NIAGARAEDITOR_API UNiagaraStackRenderItemGroup : public UNiagaraStackItemG
 	GENERATED_BODY()
 
 public:
-	virtual FText GetDisplayName() const override;
-	void SetDisplayName(FText InDisplayName);
+	void Initialize(FRequiredEntryData InRequiredEntryData);
 
 protected:
-	virtual void RefreshChildrenInternal(const TArray<UNiagaraStackEntry*>& CurrentChildren, TArray<UNiagaraStackEntry*>& NewChildren) override;
+	virtual void RefreshChildrenInternal(const TArray<UNiagaraStackEntry*>& CurrentChildren, TArray<UNiagaraStackEntry*>& NewChildren, TArray<FStackIssue>& NewIssues) override;
 
 private:
+	void ItemAdded();
 	void ChildModifiedGroupItems();
 
 private:
-	FText DisplayName;
+	TSharedPtr<INiagaraStackItemGroupAddUtilities> AddUtilities;
 };
