Copying //UE4/Dev-Niagara to //UE4/Dev-Main (Source: //UE4/Dev-Niagara @ 4074996)

diff --git a/Engine/Plugins/FX/Niagara/Source/NiagaraEditor/Public/ViewModels/Stack/NiagaraStackRendererItem.h b/Engine/Plugins/FX/Niagara/Source/NiagaraEditor/Public/ViewModels/Stack/NiagaraStackRendererItem.h
--- a/Engine/Plugins/FX/Niagara/Source/NiagaraEditor/Public/ViewModels/Stack/NiagaraStackRendererItem.h
+++ b/Engine/Plugins/FX/Niagara/Source/NiagaraEditor/Public/ViewModels/Stack/NiagaraStackRendererItem.h
@@ -9,17 +9,16 @@
 class UNiagaraEmitter;
 class UNiagaraRendererProperties;
 class UNiagaraStackObject;
-class UNiagaraStackItemExpander;
 
 UCLASS()
 class NIAGARAEDITOR_API UNiagaraStackRendererItem : public UNiagaraStackItem
 {
 	GENERATED_BODY()
-
+		
 public:
 	UNiagaraStackRendererItem();
 
-	void Initialize(TSharedRef<FNiagaraSystemViewModel> InSystemViewModel, TSharedRef<FNiagaraEmitterViewModel> InEmitterViewModel, UNiagaraStackEditorData& InStackEditorData, UNiagaraRendererProperties* InRendererProperties);
+	void Initialize(FRequiredEntryData InRequiredEntryData, UNiagaraRendererProperties* InRendererProperties);
 
 	UNiagaraRendererProperties* GetRendererProperties();
 
@@ -35,24 +34,19 @@ public:
 
 	void ResetToBase();
 
-	virtual FName GetItemBackgroundName() const override;
-	virtual int32 GetErrorCount() const override;
-	virtual bool GetErrorFixable(int32 ErrorIdx) const override;
-	virtual bool TryFixError(int32 ErrorIdx) override;
-	virtual FText GetErrorText(int32 ErrorIdx) const override;
+	bool GetIsEnabled() const;
+	void SetIsEnabled(bool bInIsEnabled);
 
 	static TArray<FNiagaraVariable> GetMissingVariables(UNiagaraRendererProperties* RendererProperties, UNiagaraEmitter* Emitter);
 	static bool AddMissingVariable(UNiagaraEmitter* Emitter, const FNiagaraVariable& Variable);
 
 protected:
 	virtual void BeginDestroy() override;
-
-	virtual void RefreshChildrenInternal(const TArray<UNiagaraStackEntry*>& CurrentChildren, TArray<UNiagaraStackEntry*>& NewChildren) override;
+	virtual void RefreshChildrenInternal(const TArray<UNiagaraStackEntry*>& CurrentChildren, TArray<UNiagaraStackEntry*>& NewChildren, TArray<FStackIssue>& NewIssues) override;
 
 private:
-	void RendererExpandedChanged();
-
 	void RendererChanged();
+	void RefreshIssues(TArray<FStackIssue>& NewIssues);
 
 private:
 	TWeakObjectPtr<UNiagaraRendererProperties> RendererProperties;
@@ -65,7 +59,4 @@ private:
 
 	UPROPERTY()
 	UNiagaraStackObject* RendererObject;
-
-	UPROPERTY()
-	UNiagaraStackItemExpander* RendererExpander;
 };
diff --git a/Engine/Plugins/FX/Niagara/Source/NiagaraEditor/Public/ViewModels/Stack/NiagaraStackRoot.h b/Engine/Plugins/FX/Niagara/Source/NiagaraEditor/Public/ViewModels/Stack/NiagaraStackRoot.h
--- a/Engine/Plugins/FX/Niagara/Source/NiagaraEditor/Public/ViewModels/Stack/NiagaraStackRoot.h
+++ b/Engine/Plugins/FX/Niagara/Source/NiagaraEditor/Public/ViewModels/Stack/NiagaraStackRoot.h
@@ -21,13 +21,14 @@ class NIAGARAEDITOR_API UNiagaraStackRoot : public UNiagaraStackEntry
 
 public:
 	UNiagaraStackRoot();
+	
+	void Initialize(FRequiredEntryData InRequiredEntryData);
 
-	void Initialize(TSharedRef<FNiagaraSystemViewModel> InSystemViewModel, TSharedRef<FNiagaraEmitterViewModel> InEmitterViewModel);
-
+	virtual bool GetCanExpand() const override;
 	virtual bool GetShouldShowInStack() const override;
 
 protected:
-	virtual void RefreshChildrenInternal(const TArray<UNiagaraStackEntry*>& CurrentChildren, TArray<UNiagaraStackEntry*>& NewChildren) override;
+	virtual void RefreshChildrenInternal(const TArray<UNiagaraStackEntry*>& CurrentChildren, TArray<UNiagaraStackEntry*>& NewChildren, TArray<FStackIssue>& NewIssues) override;
 
 private:
 	void EmitterEventArraysChanged();
diff --git a/Engine/Plugins/FX/Niagara/Source/NiagaraEditor/Public/ViewModels/Stack/NiagaraStackScriptItemGroup.h b/Engine/Plugins/FX/Niagara/Source/NiagaraEditor/Public/ViewModels/Stack/NiagaraStackScriptItemGroup.h
--- a/Engine/Plugins/FX/Niagara/Source/NiagaraEditor/Public/ViewModels/Stack/NiagaraStackScriptItemGroup.h
+++ b/Engine/Plugins/FX/Niagara/Source/NiagaraEditor/Public/ViewModels/Stack/NiagaraStackScriptItemGroup.h
@@ -4,13 +4,14 @@
 
 #include "ViewModels/Stack/NiagaraStackItemGroup.h"
 #include "NiagaraCommon.h"
+#include "UObject/ObjectKey.h"
 #include "NiagaraStackScriptItemGroup.generated.h"
 
 class FNiagaraScriptViewModel;
-class UNiagaraStackEditorData;
-class UNiagaraStackAddScriptModuleItem;
 class UNiagaraStackSpacer;
 class UNiagaraNodeOutput;
+class FScriptItemGroupAddUtilities;
+class UNiagaraStackModuleItem;
 
 UCLASS()
 class NIAGARAEDITOR_API UNiagaraStackScriptItemGroup : public UNiagaraStackItemGroup
@@ -19,9 +20,9 @@ class NIAGARAEDITOR_API UNiagaraStackScriptItemGroup : public UNiagaraStackItemG
 
 public:
 	void Initialize(
-		TSharedRef<FNiagaraSystemViewModel> InSystemViewModel,
-		TSharedRef<FNiagaraEmitterViewModel> InEmitterViewModel,
-		UNiagaraStackEditorData& InStackEditorData,
+		FRequiredEntryData InRequiredEntryData,
+		FText InDisplayName,
+		FText InToolTip,
 		TSharedRef<FNiagaraScriptViewModel> InScriptViewModel,
 		ENiagaraScriptUsage InScriptUsage,
 		FGuid InScriptUsageId = FGuid());
@@ -29,38 +30,29 @@ public:
 	ENiagaraScriptUsage GetScriptUsage() const { return ScriptUsage; }
 	FGuid GetScriptUsageId() const { return ScriptUsageId; }
 
-	virtual FText GetDisplayName() const override;
-	void SetDisplayName(FText InDisplayName);
+protected:
+	virtual void RefreshChildrenInternal(const TArray<UNiagaraStackEntry*>& CurrentChildren, TArray<UNiagaraStackEntry*>& NewChildren, TArray<FStackIssue>& NewIssues) override;
 
-	virtual int32 GetErrorCount() const override;
-	virtual bool GetErrorFixable(int32 ErrorIdx) const override;
-	virtual bool TryFixError(int32 ErrorIdx) override;
-	virtual FText GetErrorText(int32 ErrorIdx) const override;
-	virtual FText GetErrorSummaryText(int32 ErrorIdx) const override;
+	virtual TOptional<FDropResult> ChildRequestCanDropInternal(const UNiagaraStackEntry& TargetChild, const TArray<UNiagaraStackEntry*>& DraggedEntries) override;
 
-protected:
-	virtual void RefreshChildrenInternal(const TArray<UNiagaraStackEntry*>& CurrentChildren, TArray<UNiagaraStackEntry*>& NewChildren) override;
-	TWeakPtr<FNiagaraScriptViewModel> ScriptViewModel;
+	virtual TOptional<FDropResult> ChildRequestDropInternal(const UNiagaraStackEntry& TargetChild, const TArray<UNiagaraStackEntry*>& DraggedEntries) override;
 
 private:
 	void ItemAdded();
 
 	void ChildModifiedGroupItems();
 
+protected:
+	TWeakPtr<FNiagaraScriptViewModel> ScriptViewModel;
+	void RefreshIssues(TArray<FStackIssue>& NewIssues);
+
 private:
+	TSharedPtr<FScriptItemGroupAddUtilities> AddUtilities;
+
 	ENiagaraScriptUsage ScriptUsage;
 
 	FGuid ScriptUsageId;
+	bool bIsValidForOutput;
 
-	FText DisplayName;
-
-	DECLARE_DELEGATE(FFixError);
-	struct FError
-	{
-		FText ErrorText;
-		FText ErrorSummaryText;
-		FFixError Fix;
-	};
-
-	TSharedPtr<FError> Error;
+	TMap<FObjectKey, UNiagaraStackModuleItem*> StackSpacerToModuleItemMap;
 };
diff --git a/Engine/Plugins/FX/Niagara/Source/NiagaraEditor/Public/ViewModels/Stack/NiagaraStackSpacer.h b/Engine/Plugins/FX/Niagara/Source/NiagaraEditor/Public/ViewModels/Stack/NiagaraStackSpacer.h
--- a/Engine/Plugins/FX/Niagara/Source/NiagaraEditor/Public/ViewModels/Stack/NiagaraStackSpacer.h
+++ b/Engine/Plugins/FX/Niagara/Source/NiagaraEditor/Public/ViewModels/Stack/NiagaraStackSpacer.h
@@ -11,14 +11,19 @@ class NIAGARAEDITOR_API UNiagaraStackSpacer : public UNiagaraStackEntry
 	GENERATED_BODY()
 
 public:
-	void Initialize(TSharedRef<FNiagaraSystemViewModel> InSystemViewModel, TSharedRef<FNiagaraEmitterViewModel> InEmitterViewModel, FName SpacerKey = NAME_None);
+	void Initialize(FRequiredEntryData InRequiredEntryData, FName SpacerKey = NAME_None, float InSpacerScale = 1.0f, EStackRowStyle InRowStyle = UNiagaraStackEntry::EStackRowStyle::None);
 
 	//~ UNiagaraStackEntry interface
 	virtual FText GetDisplayName() const override;
-	virtual FName GetItemBackgroundName() const override;
+	virtual bool GetCanExpand() const override;
+	virtual EStackRowStyle GetStackRowStyle() const override;
 
 	FName GetSpacerKey() const;
 
+	float GetSpacerScale() const;
+
 private:
 	FName SpacerKey;
-};
+	float SpacerScale;
+	EStackRowStyle RowStyle;
+};
\ No newline at end of file
diff --git a/Engine/Plugins/FX/Niagara/Source/NiagaraEditor/Public/ViewModels/Stack/NiagaraStackViewModel.h b/Engine/Plugins/FX/Niagara/Source/NiagaraEditor/Public/ViewModels/Stack/NiagaraStackViewModel.h
--- a/Engine/Plugins/FX/Niagara/Source/NiagaraEditor/Public/ViewModels/Stack/NiagaraStackViewModel.h
+++ b/Engine/Plugins/FX/Niagara/Source/NiagaraEditor/Public/ViewModels/Stack/NiagaraStackViewModel.h
@@ -2,11 +2,14 @@
 
 #pragma once
 
+#include "Templates/SharedPointer.h"
 #include "EditorUndoClient.h"
+#include "ViewModels/Stack/NiagaraStackEntry.h"
+#include "UObject/ObjectKey.h"
 #include "NiagaraStackViewModel.generated.h"
 
 class FNiagaraSystemViewModel;
-class FNiagaraEmitterViewModel;
+class FNiagaraEmitterHandleViewModel;
 class UNiagaraStackEntry;
 class UNiagaraStackRoot;
 
@@ -17,35 +20,110 @@ class NIAGARAEDITOR_API UNiagaraStackViewModel : public UObject, public FEditorU
 
 public:
 	DECLARE_MULTICAST_DELEGATE(FOnStructureChanged);
-
+	DECLARE_MULTICAST_DELEGATE(FOnSearchCompleted);
+public:
+	struct FSearchResult
+	{
+		TArray<UNiagaraStackEntry*> EntryPath;
+		UNiagaraStackEntry::FStackSearchItem MatchingItem;
+		UNiagaraStackEntry* GetEntry()
+		{
+			return EntryPath.Num() > 0 ? 
+				EntryPath[EntryPath.Num() - 1] : 
+				nullptr;
+		}
+	};
 public:
-	TSharedPtr<FNiagaraEmitterViewModel> GetEmitterViewModel();
-	void Initialize(TSharedPtr<FNiagaraSystemViewModel> InSystemViewModel, TSharedPtr<FNiagaraEmitterViewModel> InEmitterViewModel);
+	TSharedPtr<FNiagaraEmitterHandleViewModel> GetEmitterHandleViewModel();
+	TSharedPtr<FNiagaraSystemViewModel> GetSystemViewModel();
+	void Initialize(TSharedPtr<FNiagaraSystemViewModel> InSystemViewModel, TSharedPtr<FNiagaraEmitterHandleViewModel> InEmitterHandleViewModel);
 
 	virtual void BeginDestroy() override;
 
 	TArray<UNiagaraStackEntry*>& GetRootEntries();
 
 	FOnStructureChanged& OnStructureChanged();
+	FOnSearchCompleted& OnSearchCompleted();
+
+	bool GetShowAllAdvanced() const;
+	void SetShowAllAdvanced(bool bInShowAllAdvanced);
+
+	bool GetShowOutputs() const;
+	void SetShowOutputs(bool bInShowOutputs);
+
+	double GetLastScrollPosition() const;
+	void SetLastScrollPosition(double InLastScrollPosition);
+
+	void NotifyStructureChanged();
 
 	//~ FEditorUndoClient Interface
 	virtual void PostUndo(bool bSuccess) override;
 	virtual void PostRedo(bool bSuccess) override { PostUndo(bSuccess); }
 
+	virtual void Tick();
+	//~ stack search stuff
+	void OnSearchTextChanged(const FText& SearchText);
+	FText GetCurrentSearchText() const { return CurrentSearchText; };
+	bool IsSearching();
+	const TArray<FSearchResult>& GetCurrentSearchResults();
+	int GetCurrentFocusedMatchIndex() const { return CurrentFocusedSearchMatchIndex; }
+	UNiagaraStackEntry* GetCurrentFocusedEntry();
+	void AddSearchScrollOffset(int NumberOfSteps);
+
+	void GetPathForEntry(UNiagaraStackEntry* Entry, TArray<UNiagaraStackEntry*>& EntryPath);
+
+	/** Starts recursing through all entries to expand all groups and collapse all items. */
+	void CollapseToHeaders();
+	
+	void UndismissAllIssues();
+
+	bool HasDismissedStackIssues();
+
+private:
+
+	/** Recursively Expands all groups and collapses all items in the stack. */
+	void CollapseToHeadersRecursive(TArray<UNiagaraStackEntry*> Entries);
+
+	struct FSearchWorkItem
+	{
+		TArray<UNiagaraStackEntry*> EntryPath;
+		UNiagaraStackEntry* GetEntry()
+		{
+			return EntryPath.Num() > 0 ?
+				EntryPath[EntryPath.Num() - 1] :
+				nullptr;
+		}
+	};
+
 private:
 	void EntryStructureChanged();
 	void EntryDataObjectModified(UObject* ChangedObject);
 	void OnSystemCompiled();
 	void OnEmitterCompiled();
+	/** Called by the tick function to perform partial search */
+	void SearchTick();
+	void GenerateTraversalEntries(UNiagaraStackEntry* Root, TArray<UNiagaraStackEntry*> ParentChain, 
+		TArray<FSearchWorkItem>& TraversedArray);
+	bool ItemMatchesSearchCriteria(UNiagaraStackEntry::FStackSearchItem SearchItem);
+	void GeneratePathForEntry(UNiagaraStackEntry* Root, UNiagaraStackEntry* Entry, TArray<UNiagaraStackEntry*> CurrentPath, TArray<UNiagaraStackEntry*>& EntryPath);
 
 private:
-	TSharedPtr<FNiagaraEmitterViewModel> EmitterViewModel;
+	TSharedPtr<FNiagaraEmitterHandleViewModel> EmitterHandleViewModel;
 	TSharedPtr<FNiagaraSystemViewModel> SystemViewModel;
-
+	
 	TArray<UNiagaraStackEntry*> RootEntries;
 
 	UPROPERTY()
 	UNiagaraStackRoot* RootEntry;
 
 	FOnStructureChanged StructureChangedDelegate;
+
+	// ~Search stuff
+	FText CurrentSearchText;
+	int CurrentFocusedSearchMatchIndex;
+	FOnSearchCompleted SearchCompletedDelegate;
+	TArray<FSearchWorkItem> ItemsToSearch;
+	TArray<FSearchResult> CurrentSearchResults;
+	static const double MaxSearchTime;
+	bool bRestartSearch;
 };
\ No newline at end of file
diff --git a/Engine/Plugins/FX/Niagara/Source/NiagaraEditorWidgets/NiagaraEditorWidgets.Build.cs b/Engine/Plugins/FX/Niagara/Source/NiagaraEditorWidgets/NiagaraEditorWidgets.Build.cs
--- a/Engine/Plugins/FX/Niagara/Source/NiagaraEditorWidgets/NiagaraEditorWidgets.Build.cs
+++ b/Engine/Plugins/FX/Niagara/Source/NiagaraEditorWidgets/NiagaraEditorWidgets.Build.cs
@@ -11,6 +11,7 @@ public class NiagaraEditorWidgets : ModuleRules
 		});
 
 		PrivateDependencyModuleNames.AddRange(new string[] {
+            "NiagaraCore",
 			"Niagara",
 			"NiagaraEditor",
 			"Engine",
@@ -24,7 +25,10 @@ public class NiagaraEditorWidgets : ModuleRules
 			"EditorStyle",
 			"PropertyEditor",
 			"AppFramework",
-		});
+			"Projects",
+			"Sequencer",
+            "EditorWidgets",
+        });
 
 		PrivateIncludePathModuleNames.AddRange(new string[] {
 		});
diff --git a/Engine/Plugins/FX/Niagara/Source/NiagaraEditorWidgets/Private/DetailCustomizations/NiagaraDataInterfaceCurveDetails.cpp b/Engine/Plugins/FX/Niagara/Source/NiagaraEditorWidgets/Private/DetailCustomizations/NiagaraDataInterfaceCurveDetails.cpp
--- a/Engine/Plugins/FX/Niagara/Source/NiagaraEditorWidgets/Private/DetailCustomizations/NiagaraDataInterfaceCurveDetails.cpp
+++ b/Engine/Plugins/FX/Niagara/Source/NiagaraEditorWidgets/Private/DetailCustomizations/NiagaraDataInterfaceCurveDetails.cpp
@@ -9,7 +9,6 @@
 #include "NiagaraDataInterfaceColorCurve.h"
 #include "NiagaraEditorWidgetsModule.h"
 
-#include "SCurveEditor.h"
 #include "SCurveEditor.h"
 #include "PropertyHandle.h"
 #include "DetailLayoutBuilder.h"
@@ -20,6 +19,29 @@
 #include "Brushes/SlateColorBrush.h"
 #include "Modules/ModuleManager.h"
 
+#include "Widgets/Input/SComboButton.h"
+#include "Widgets/Input/SButton.h"
+#include "Widgets/Images/SImage.h"
+#include "NiagaraEditorStyle.h"
+#include "NiagaraEditorWidgetsStyle.h"
+#include "IContentBrowserSingleton.h"
+#include "ContentBrowserModule.h"
+#include "Framework/Application/SlateApplication.h"
+#include "Curves/CurveVector.h"
+#include "Curves/CurveFloat.h"
+#include "Curves/CurveLinearColor.h"
+#include "Curves/RichCurve.h"
+#include "ScopedTransaction.h"
+
+#define LOCTEXT_NAMESPACE "NiagaraDataInterfaceCurveDetails"
+
+FRichCurve* GetCurveFromPropertyHandle(TSharedPtr<IPropertyHandle> Handle)
+{
+	TArray<void*> RawData;
+	Handle->AccessRawData(RawData);
+	return RawData.Num() == 1 ? static_cast<FRichCurve*>(RawData[0]) : nullptr;
+}
+
 class SNiagaraResizeBox : public SCompoundWidget
 {
 public:
@@ -237,15 +259,10 @@ public:
 	}
 
 private:
-	FRichCurve* GetCurveFromPropertyHandle(TSharedPtr<IPropertyHandle> Handle)
-	{
-		TArray<void*> RawData;
-		Handle->AccessRawData(RawData);
-		return RawData.Num() == 1 ? static_cast<FRichCurve*>(RawData[0]) : nullptr;
-	}
-
 	void CurveChanged(FRichCurve* ChangedCurve, UObject* CurveOwnerObject)
 	{
+		UNiagaraDataInterfaceCurveBase* EditedCurve = Cast<UNiagaraDataInterfaceCurveBase>(CurveOwnerObject);
+		EditedCurve->UpdateLUT(); // we need this done before notify change because of the internal copy methods
 		for (TSharedRef<IPropertyHandle> CurveProperty : CurveProperties)
 		{
 			if (GetCurveFromPropertyHandle(CurveProperty) == ChangedCurve)
@@ -269,6 +286,8 @@ private:
 // Curve Base
 void FNiagaraDataInterfaceCurveDetailsBase::CustomizeDetails(IDetailLayoutBuilder& DetailBuilder)
 {
+	CustomDetailBuilder = &DetailBuilder;
+	FNiagaraDataInterfaceDetailsBase::CustomizeDetails(DetailBuilder);
 	// Only support single objects.
 	TArray<TWeakObjectPtr<UObject>> ObjectsBeingCustomized;
 	DetailBuilder.GetObjectsBeingCustomized(ObjectsBeingCustomized);
@@ -276,7 +295,7 @@ void FNiagaraDataInterfaceCurveDetailsBase::CustomizeDetails(IDetailLayoutBuilde
 	{
 		return;
 	}
-
+	CustomizedCurveInterface = Cast<UNiagaraDataInterfaceCurveBase>(ObjectsBeingCustomized[0].Get());
 	FNiagaraEditorWidgetsModule& NiagaraEditorWidgetsModule = FModuleManager::GetModuleChecked<FNiagaraEditorWidgetsModule>("NiagaraEditorWidgets");
 	TSharedRef<FNiagaraStackCurveEditorOptions> StackCurveEditorOptions = NiagaraEditorWidgetsModule.GetOrCreateStackCurveEditorOptionsForObject(
 		ObjectsBeingCustomized[0].Get(), GetDefaultAreCurvesVisible(), GetDefaultHeight());
@@ -299,6 +318,50 @@ void FNiagaraDataInterfaceCurveDetailsBase::CustomizeDetails(IDetailLayoutBuilde
 	}
 
 	IDetailCategoryBuilder& CurveCategory = DetailBuilder.EditCategory("Curve");
+	TSharedRef<IPropertyHandle> ShowInCurveEditorHandle = CustomDetailBuilder->GetProperty(FName("ShowInCurveEditor"), UNiagaraDataInterfaceCurveBase::StaticClass());
+	if (ShowInCurveEditorHandle->IsValidHandle())
+	{
+		ShowInCurveEditorHandle->MarkHiddenByCustomization();
+	}
+	CurveCategory.HeaderContent( 
+		// Checkbox for showing in curve editor
+		SNew(SHorizontalBox)
+		+ SHorizontalBox::Slot()
+		.Padding(5, 0, 5, 0)
+		.HAlign(EHorizontalAlignment::HAlign_Left)
+		.AutoWidth()
+		[
+			SNew(SButton)
+			.ButtonStyle(FEditorStyle::Get(), "HoverHintOnly")
+			.HAlign(HAlign_Center)
+			.VAlign(VAlign_Center)
+			.ContentPadding(1)
+			.ToolTipText(this, &FNiagaraDataInterfaceCurveDetailsBase::GetShowInCurveEditorTooltip)
+			.OnClicked(this, &FNiagaraDataInterfaceCurveDetailsBase::OnToggleShowInCurveEditor)
+			.Content()
+			[
+				SNew(SImage)
+				.Image(FNiagaraEditorWidgetsStyle::Get().GetBrush("NiagaraEditor.ShowInCurveEditorIcon"))
+				.ColorAndOpacity(this, &FNiagaraDataInterfaceCurveDetailsBase::GetShowInCurveEditorImageColor)
+			]
+		]
+		+ SHorizontalBox::Slot()
+		.HAlign(EHorizontalAlignment::HAlign_Right)
+		[
+			SNew(SComboButton)
+			.HasDownArrow(true)
+			.OnGetMenuContent(this, &FNiagaraDataInterfaceCurveDetailsBase::GetCurveToCopyMenu)
+			.ContentPadding(2)
+			.ButtonContent()
+			[
+				SNew(STextBlock)
+				.TextStyle(FNiagaraEditorStyle::Get(), "NiagaraEditor.ParameterText")
+				.ColorAndOpacity(FSlateColor::UseForeground())
+				.Text(NSLOCTEXT("NiagaraDataInterfaceCurveDetails", "Import", "Import"))
+				.ToolTipText(NSLOCTEXT("NiagaraDataInterfaceCurveDetails", "CopyCurveAsset", "Copy data from another Curve asset"))
+			]
+		]
+	);
 	CurveCategory.AddCustomRow(NSLOCTEXT("NiagaraDataInterfaceCurveDetails", "CurveFilterText", "Curve"))
 		.WholeRowContent()
 		[
@@ -312,6 +375,82 @@ void FNiagaraDataInterfaceCurveDetailsBase::CustomizeDetails(IDetailLayoutBuilde
 		];
 }
 
+FText FNiagaraDataInterfaceCurveDetailsBase::GetShowInCurveEditorTooltip() const
+{
+	return LOCTEXT("ShowInCurveEditorToolTip", "Show this curve in the curves tab.");
+}
+
+FSlateColor FNiagaraDataInterfaceCurveDetailsBase::GetShowInCurveEditorImageColor() const
+{
+	return CustomizedCurveInterface->ShowInCurveEditor
+		? FEditorStyle::GetSlateColor("SelectionColor")
+		: FLinearColor::Gray;
+}
+
+FReply FNiagaraDataInterfaceCurveDetailsBase::OnToggleShowInCurveEditor() const
+{
+	TSharedRef<IPropertyHandle> ShowInCurveEditorHandle = CustomDetailBuilder->GetProperty(FName("ShowInCurveEditor"), UNiagaraDataInterfaceCurveBase::StaticClass());
+	if (ShowInCurveEditorHandle->IsValidHandle())
+	{
+		bool bShowInCurveEditor;
+		ShowInCurveEditorHandle->GetValue(bShowInCurveEditor);
+		ShowInCurveEditorHandle->SetValue(!bShowInCurveEditor);
+	}
+	return FReply::Handled();
+}
+
+void FNiagaraDataInterfaceCurveDetailsBase::ImportSelectedAsset(UObject* SelectedAsset)
+{
+	TArray<FRichCurve> FloatCurves;
+	GetFloatCurvesFromAsset(SelectedAsset, FloatCurves);
+	TArray<TSharedRef<IPropertyHandle>> CurveProperties;
+	GetCurveProperties(*CustomDetailBuilder, CurveProperties);
+	if (FloatCurves.Num() == CurveProperties.Num())
+	{
+		FScopedTransaction ImportTransaction(LOCTEXT("ImportCurveTransaction", "Import curve"));
+		CustomizedCurveInterface->Modify();
+		for (int i = 0; i < CurveProperties.Num(); i++)
+		{
+			if (CurveProperties[i]->IsValidHandle())
+			{
+				*GetCurveFromPropertyHandle(CurveProperties[i]) = FloatCurves[i];
+			}
+		}
+		CustomizedCurveInterface->UpdateLUT(); // we need this done before notify change because of the internal copy methods
+		for (auto CurveProperty : CurveProperties)
+		{
+			CurveProperty->NotifyPostChange();
+		}
+	}
+}
+
+TSharedRef<SWidget> FNiagaraDataInterfaceCurveDetailsBase::GetCurveToCopyMenu()
+{
+	FName ClassName = GetSupportedAssetClassName();
+	FAssetPickerConfig AssetPickerConfig;
+	{
+		AssetPickerConfig.OnAssetSelected = FOnAssetSelected::CreateSP(this, &FNiagaraDataInterfaceCurveDetailsBase::CurveToCopySelected);
+		AssetPickerConfig.bAllowNullSelection = false;
+		AssetPickerConfig.InitialAssetViewType = EAssetViewType::List;
+		AssetPickerConfig.Filter.ClassNames.Add(ClassName);
+	}
+
+	FContentBrowserModule& ContentBrowserModule = FModuleManager::Get().LoadModuleChecked<FContentBrowserModule>(TEXT("ContentBrowser"));
+
+	return SNew(SBox)
+		.WidthOverride(300.0f)
+		[
+			ContentBrowserModule.Get().CreateAssetPicker(AssetPickerConfig)
+		];
+	
+	return SNullWidget::NullWidget;
+}
+
+void FNiagaraDataInterfaceCurveDetailsBase::CurveToCopySelected(const FAssetData& AssetData)
+{
+	ImportSelectedAsset(AssetData.GetAsset());
+	FSlateApplication::Get().DismissAllMenus();
+}
 
 // Curve
 TSharedRef<IDetailCustomization> FNiagaraDataInterfaceCurveDetails::MakeInstance()
@@ -321,9 +460,19 @@ TSharedRef<IDetailCustomization> FNiagaraDataInterfaceCurveDetails::MakeInstance
 
 void FNiagaraDataInterfaceCurveDetails::GetCurveProperties(IDetailLayoutBuilder& DetailBuilder, TArray<TSharedRef<IPropertyHandle>>& CurveProperties) const
 {
-	CurveProperties.Add(DetailBuilder.GetProperty(GET_MEMBER_NAME_CHECKED(UNiagaraDataInterfaceCurve, Curve)));
+	CurveProperties.Add(DetailBuilder.GetProperty(FName("Curve"), UNiagaraDataInterfaceCurve::StaticClass()));
 }
 
+FName FNiagaraDataInterfaceCurveDetails::GetSupportedAssetClassName() const
+{
+	return UCurveFloat::StaticClass()->GetFName();
+}
+
+void FNiagaraDataInterfaceCurveDetails::GetFloatCurvesFromAsset(UObject* SelectedAsset, TArray<FRichCurve>& FloatCurves) const
+{
+	UCurveFloat* CurveAsset = Cast<UCurveFloat>(SelectedAsset);
+	FloatCurves.Add(CurveAsset->FloatCurve);
+}
 
 // Vector 2D Curve
 TSharedRef<IDetailCustomization> FNiagaraDataInterfaceVector2DCurveDetails::MakeInstance()
@@ -333,8 +482,22 @@ TSharedRef<IDetailCustomization> FNiagaraDataInterfaceVector2DCurveDetails::Make
 
 void FNiagaraDataInterfaceVector2DCurveDetails::GetCurveProperties(IDetailLayoutBuilder& DetailBuilder, TArray<TSharedRef<IPropertyHandle>>& OutCurveProperties) const
 {
-	OutCurveProperties.Add(DetailBuilder.GetProperty(GET_MEMBER_NAME_CHECKED(UNiagaraDataInterfaceVector2DCurve, XCurve)));
-	OutCurveProperties.Add(DetailBuilder.GetProperty(GET_MEMBER_NAME_CHECKED(UNiagaraDataInterfaceVector2DCurve, YCurve)));
+	OutCurveProperties.Add(DetailBuilder.GetProperty(FName("XCurve"), UNiagaraDataInterfaceVector2DCurve::StaticClass()));
+	OutCurveProperties.Add(DetailBuilder.GetProperty(FName("YCurve"), UNiagaraDataInterfaceVector2DCurve::StaticClass()));
+}
+
+FName FNiagaraDataInterfaceVector2DCurveDetails::GetSupportedAssetClassName() const
+{
+	return UCurveVector::StaticClass()->GetFName();
+}
+
+void FNiagaraDataInterfaceVector2DCurveDetails::GetFloatCurvesFromAsset(UObject* SelectedAsset, TArray<FRichCurve>& FloatCurves) const
+{
+	UCurveVector* CurveAsset = Cast<UCurveVector>(SelectedAsset);
+	for (int i = 0; i < 2; i++)
+	{
+		FloatCurves.Add(CurveAsset->FloatCurves[i]);
+	}
 }
 
 
@@ -346,9 +509,23 @@ TSharedRef<IDetailCustomization> FNiagaraDataInterfaceVectorCurveDetails::MakeIn
 
 void FNiagaraDataInterfaceVectorCurveDetails::GetCurveProperties(IDetailLayoutBuilder& DetailBuilder, TArray<TSharedRef<IPropertyHandle>>& OutCurveProperties) const
 {
-	OutCurveProperties.Add(DetailBuilder.GetProperty(GET_MEMBER_NAME_CHECKED(UNiagaraDataInterfaceVectorCurve, XCurve)));
-	OutCurveProperties.Add(DetailBuilder.GetProperty(GET_MEMBER_NAME_CHECKED(UNiagaraDataInterfaceVectorCurve, YCurve)));
-	OutCurveProperties.Add(DetailBuilder.GetProperty(GET_MEMBER_NAME_CHECKED(UNiagaraDataInterfaceVectorCurve, ZCurve)));
+	OutCurveProperties.Add(DetailBuilder.GetProperty(FName("XCurve"), UNiagaraDataInterfaceVectorCurve::StaticClass()));
+	OutCurveProperties.Add(DetailBuilder.GetProperty(FName("YCurve"), UNiagaraDataInterfaceVectorCurve::StaticClass()));
+	OutCurveProperties.Add(DetailBuilder.GetProperty(FName("ZCurve"), UNiagaraDataInterfaceVectorCurve::StaticClass()));
+}
+
+FName FNiagaraDataInterfaceVectorCurveDetails::GetSupportedAssetClassName() const
+{
+	return UCurveVector::StaticClass()->GetFName();
+}
+
+void FNiagaraDataInterfaceVectorCurveDetails::GetFloatCurvesFromAsset(UObject* SelectedAsset, TArray<FRichCurve>& FloatCurves) const
+{
+	UCurveVector* CurveAsset = Cast<UCurveVector>(SelectedAsset);
+	for (int i = 0; i < 3; i++)
+	{
+		FloatCurves.Add(CurveAsset->FloatCurves[i]);
+	}
 }
 
 
@@ -360,12 +537,25 @@ TSharedRef<IDetailCustomization> FNiagaraDataInterfaceVector4CurveDetails::MakeI
 
 void FNiagaraDataInterfaceVector4CurveDetails::GetCurveProperties(IDetailLayoutBuilder& DetailBuilder, TArray<TSharedRef<IPropertyHandle>>& OutCurveProperties) const
 {
-	OutCurveProperties.Add(DetailBuilder.GetProperty(GET_MEMBER_NAME_CHECKED(UNiagaraDataInterfaceVector4Curve, XCurve)));
-	OutCurveProperties.Add(DetailBuilder.GetProperty(GET_MEMBER_NAME_CHECKED(UNiagaraDataInterfaceVector4Curve, YCurve)));
-	OutCurveProperties.Add(DetailBuilder.GetProperty(GET_MEMBER_NAME_CHECKED(UNiagaraDataInterfaceVector4Curve, ZCurve)));
-	OutCurveProperties.Add(DetailBuilder.GetProperty(GET_MEMBER_NAME_CHECKED(UNiagaraDataInterfaceVector4Curve, WCurve)));
+	OutCurveProperties.Add(DetailBuilder.GetProperty(FName("XCurve"), UNiagaraDataInterfaceVector4Curve::StaticClass()));
+	OutCurveProperties.Add(DetailBuilder.GetProperty(FName("YCurve"), UNiagaraDataInterfaceVector4Curve::StaticClass()));
+	OutCurveProperties.Add(DetailBuilder.GetProperty(FName("ZCurve"), UNiagaraDataInterfaceVector4Curve::StaticClass()));
+	OutCurveProperties.Add(DetailBuilder.GetProperty(FName("WCurve"), UNiagaraDataInterfaceVector4Curve::StaticClass()));
 }
 
+FName FNiagaraDataInterfaceVector4CurveDetails::GetSupportedAssetClassName() const
+{
+	return UCurveLinearColor::StaticClass()->GetFName();
+}
+
+void FNiagaraDataInterfaceVector4CurveDetails::GetFloatCurvesFromAsset(UObject* SelectedAsset, TArray<FRichCurve>& FloatCurves) const
+{
+	UCurveLinearColor* CurveAsset = Cast<UCurveLinearColor>(SelectedAsset);
+	for (int i = 0; i < 4; i++)
+	{
+		FloatCurves.Add(CurveAsset->FloatCurves[i]);
+	}
+}
 
 // Color Curve
 TSharedRef<IDetailCustomization> FNiagaraDataInterfaceColorCurveDetails::MakeInstance()
@@ -375,8 +565,23 @@ TSharedRef<IDetailCustomization> FNiagaraDataInterfaceColorCurveDetails::MakeIns
 
 void FNiagaraDataInterfaceColorCurveDetails::GetCurveProperties(IDetailLayoutBuilder& DetailBuilder, TArray<TSharedRef<IPropertyHandle>>& OutCurveProperties) const
 {
-	OutCurveProperties.Add(DetailBuilder.GetProperty(GET_MEMBER_NAME_CHECKED(UNiagaraDataInterfaceColorCurve, RedCurve)));
-	OutCurveProperties.Add(DetailBuilder.GetProperty(GET_MEMBER_NAME_CHECKED(UNiagaraDataInterfaceColorCurve, GreenCurve)));
-	OutCurveProperties.Add(DetailBuilder.GetProperty(GET_MEMBER_NAME_CHECKED(UNiagaraDataInterfaceColorCurve, BlueCurve)));
-	OutCurveProperties.Add(DetailBuilder.GetProperty(GET_MEMBER_NAME_CHECKED(UNiagaraDataInterfaceColorCurve, AlphaCurve)));
+	OutCurveProperties.Add(DetailBuilder.GetProperty(FName("RedCurve"), UNiagaraDataInterfaceColorCurve::StaticClass()));
+	OutCurveProperties.Add(DetailBuilder.GetProperty(FName("GreenCurve"), UNiagaraDataInterfaceColorCurve::StaticClass()));
+	OutCurveProperties.Add(DetailBuilder.GetProperty(FName("BlueCurve"), UNiagaraDataInterfaceColorCurve::StaticClass()));
+	OutCurveProperties.Add(DetailBuilder.GetProperty(FName("AlphaCurve"), UNiagaraDataInterfaceColorCurve::StaticClass()));
+}
+
+FName FNiagaraDataInterfaceColorCurveDetails::GetSupportedAssetClassName() const
+{
+	return UCurveLinearColor::StaticClass()->GetFName();
+}
+
+void FNiagaraDataInterfaceColorCurveDetails::GetFloatCurvesFromAsset(UObject* SelectedAsset, TArray<FRichCurve>& FloatCurves) const
+{
+	UCurveLinearColor* CurveAsset = Cast<UCurveLinearColor>(SelectedAsset);
+	for (int i = 0; i < 4; i++)
+	{
+		FloatCurves.Add(CurveAsset->FloatCurves[i]);
+	}
 }
+#undef LOCTEXT_NAMESPACE
\ No newline at end of file
diff --git a/Engine/Plugins/FX/Niagara/Source/NiagaraEditorWidgets/Private/DetailCustomizations/NiagaraDataInterfaceCurveDetails.h b/Engine/Plugins/FX/Niagara/Source/NiagaraEditorWidgets/Private/DetailCustomizations/NiagaraDataInterfaceCurveDetails.h
--- a/Engine/Plugins/FX/Niagara/Source/NiagaraEditorWidgets/Private/DetailCustomizations/NiagaraDataInterfaceCurveDetails.h
+++ b/Engine/Plugins/FX/Niagara/Source/NiagaraEditorWidgets/Private/DetailCustomizations/NiagaraDataInterfaceCurveDetails.h
@@ -3,11 +3,18 @@
 #pragma once
 
 #include "IDetailCustomization.h"
+#include "NiagaraDataInterfaceDetails.h"
+#include "Styling/SlateTypes.h"
+#include "Input/Reply.h"
 
 class IPropertyHandle;
+class UNiagaraDataInterfaceCurveBase;
+class SWidget;
+struct FAssetData;
+struct FRichCurve;
 
 /** Base details customization for curve data interfaces. */
-class FNiagaraDataInterfaceCurveDetailsBase : public IDetailCustomization
+class FNiagaraDataInterfaceCurveDetailsBase : public FNiagaraDataInterfaceDetailsBase
 {
 public:
 	virtual void CustomizeDetails(IDetailLayoutBuilder& DetailBuilder) override;
@@ -17,6 +24,21 @@ protected:
 	virtual bool GetIsColorCurve() const { return false; }
 	virtual bool GetDefaultAreCurvesVisible() const { return true; }
 	virtual float GetDefaultHeight() const { return 120; }
+	
+	virtual FName GetSupportedAssetClassName() const = 0;
+	virtual void ImportSelectedAsset(UObject* SelectedAsset);
+	virtual void GetFloatCurvesFromAsset(UObject* SelectedAsset, TArray<FRichCurve>& FloatCurves) const = 0;
+	
+protected:
+	UNiagaraDataInterfaceCurveBase* CustomizedCurveInterface;
+	IDetailLayoutBuilder* CustomDetailBuilder;
+
+private:
+	virtual TSharedRef<SWidget> GetCurveToCopyMenu();
+	FText GetShowInCurveEditorTooltip() const;
+	FSlateColor GetShowInCurveEditorImageColor() const;
+	FReply OnToggleShowInCurveEditor() const;
+	void CurveToCopySelected(const FAssetData& AssetData);
 };
 
 /** Details customization for curve data interfaces. */
@@ -27,8 +49,11 @@ public:
 
 protected:
 	virtual void GetCurveProperties(IDetailLayoutBuilder& DetailBuilder, TArray<TSharedRef<IPropertyHandle>>& OutCurveProperties) const override;
+	virtual FName GetSupportedAssetClassName() const override;
+	virtual void GetFloatCurvesFromAsset(UObject* SelectedAsset, TArray<FRichCurve>& FloatCurves) const override;
 };
 
+
 /** Details customization for vector 2d curve data interfaces. */
 class FNiagaraDataInterfaceVector2DCurveDetails : public FNiagaraDataInterfaceCurveDetailsBase
 {
@@ -37,6 +62,8 @@ public:
 
 protected:
 	virtual void GetCurveProperties(IDetailLayoutBuilder& DetailBuilder, TArray<TSharedRef<IPropertyHandle>>& OutCurveProperties) const override;
+	virtual FName GetSupportedAssetClassName() const override;
+	virtual void GetFloatCurvesFromAsset(UObject* SelectedAsset, TArray<FRichCurve>& FloatCurves) const override;
 };
 
 /** Details customization for vector curve data interfaces. */
@@ -47,6 +74,8 @@ public:
 
 protected:
 	virtual void GetCurveProperties(IDetailLayoutBuilder& DetailBuilder, TArray<TSharedRef<IPropertyHandle>>& OutCurveProperties) const override;
+	virtual FName GetSupportedAssetClassName() const override;
+	virtual void GetFloatCurvesFromAsset(UObject* SelectedAsset, TArray<FRichCurve>& FloatCurves) const override;
 };
 
 /** Details customization for vector 4 curve data interfaces. */
@@ -57,6 +86,8 @@ public:
 
 protected:
 	virtual void GetCurveProperties(IDetailLayoutBuilder& DetailBuilder, TArray<TSharedRef<IPropertyHandle>>& OutCurveProperties) const override;
+	virtual FName GetSupportedAssetClassName() const override;
+	virtual void GetFloatCurvesFromAsset(UObject* SelectedAsset, TArray<FRichCurve>& FloatCurves) const override;
 };
 
 /** Details customization for color curve data interfaces. */
@@ -70,4 +101,6 @@ protected:
 	virtual bool GetIsColorCurve() const override { return true; }
 	virtual bool GetDefaultAreCurvesVisible() const override { return false; }
 	virtual float GetDefaultHeight() const override { return 100; }
+	virtual FName GetSupportedAssetClassName() const override;
+	virtual void GetFloatCurvesFromAsset(UObject* SelectedAsset, TArray<FRichCurve>& FloatCurves) const override;
 };
\ No newline at end of file
diff --git a/Engine/Plugins/FX/Niagara/Source/NiagaraEditorWidgets/Private/DetailCustomizations/NiagaraDataInterfaceSkeletalMeshDetails.cpp b/Engine/Plugins/FX/Niagara/Source/NiagaraEditorWidgets/Private/DetailCustomizations/NiagaraDataInterfaceSkeletalMeshDetails.cpp
new file mode 100644
--- /dev/null
+++ b/Engine/Plugins/FX/Niagara/Source/NiagaraEditorWidgets/Private/DetailCustomizations/NiagaraDataInterfaceSkeletalMeshDetails.cpp
@@ -0,0 +1,105 @@
+ // Copyright 1998-2017 Epic Games, Inc. All Rights Reserved.
+#include "NiagaraDataInterfaceSkeletalMeshDetails.h"
+#include "NiagaraDetailSourcedArrayBuilder.h"
+#include "NiagaraDataInterfaceDetails.h"
+#include "NiagaraDataInterfaceSkeletalMesh.h" 
+#include "DetailLayoutBuilder.h"
+#include "DetailCategoryBuilder.h"
+#include "IDetailGroup.h"
+#include "NiagaraComponent.h"
+
+ #define LOCTEXT_NAMESPACE "FNiagaraDataInterfaceSkeletalMeshDetails"
+
+void FNiagaraDataInterfaceSkeletalMeshDetails::CustomizeDetails(IDetailLayoutBuilder& DetailBuilder)
+ {
+	 
+	 LayoutBuilder = &DetailBuilder;
+	 static const FName MeshCategoryName = TEXT("Mesh");
+
+	 TArray<TWeakObjectPtr<UObject>> SelectedObjects;
+	 DetailBuilder.GetObjectsBeingCustomized(SelectedObjects);
+	 check(SelectedObjects.Num() == 1);
+	 MeshInterface = Cast<UNiagaraDataInterfaceSkeletalMesh>(SelectedObjects[0].Get());
+	 check(MeshInterface);
+	 if (MeshInterface != nullptr)
+	 {
+		 MeshInterface->OnChanged().RemoveAll(this);
+		 MeshInterface->OnChanged().AddRaw(this, &FNiagaraDataInterfaceSkeletalMeshDetails::OnDataChanged);
+	 }
+	 TWeakObjectPtr<USceneComponent> SceneComponent;
+	 USkeletalMeshComponent* FoundSkelComp = nullptr;
+	 MeshObject = UNiagaraDataInterfaceSkeletalMesh::GetSkeletalMeshHelper(MeshInterface, Cast<UNiagaraComponent>(MeshInterface->GetOuter()), SceneComponent, FoundSkelComp);
+	 if (MeshObject != nullptr)
+	 {
+		 MeshObject->GetOnMeshChanged().RemoveAll(this);
+		 MeshObject->GetOnMeshChanged().AddRaw(this, &FNiagaraDataInterfaceSkeletalMeshDetails::OnDataChanged);
+	 }
+
+	 IDetailCategoryBuilder& MeshCategory = DetailBuilder.EditCategory(MeshCategoryName, LOCTEXT("Mesh", "Mesh"));
+	 {
+		 TArray<TSharedRef<IPropertyHandle>> MeshProperties;
+		 MeshCategory.GetDefaultProperties(MeshProperties, true, false);
+
+		 TSharedPtr<IPropertyHandle> RegionsProperty = DetailBuilder.GetProperty(GET_MEMBER_NAME_CHECKED(UNiagaraDataInterfaceSkeletalMesh, SamplingRegions));
+		
+		 for (TSharedPtr<IPropertyHandle> Property : MeshProperties)
+		 {
+			 if (Property->GetProperty() != RegionsProperty->GetProperty())
+			 {
+				 MeshCategory.AddProperty(Property);
+			 }
+			 else
+			 {
+				 TArray<TSharedPtr<FName>> PossibleRegions = GenerateSourceArray();
+				 RegionsBuilder = TSharedPtr<FNiagaraDetailSourcedArrayBuilder>(new FNiagaraDetailSourcedArrayBuilder(Property.ToSharedRef(), PossibleRegions));
+				 MeshCategory.AddCustomBuilder(RegionsBuilder.ToSharedRef());
+			 }
+		 }
+	 }
+ }
+
+ TSharedRef<IDetailCustomization> FNiagaraDataInterfaceSkeletalMeshDetails::MakeInstance()
+ {
+	 return MakeShared<FNiagaraDataInterfaceSkeletalMeshDetails>();
+ }
+
+ void FNiagaraDataInterfaceSkeletalMeshDetails::OnDataChanged()
+ {
+	 TArray<TSharedPtr<FName>> PossibleRegions = GenerateSourceArray();
+	 RegionsBuilder->SetSourceArray(PossibleRegions);
+ }
+
+ TArray<TSharedPtr<FName>> FNiagaraDataInterfaceSkeletalMeshDetails::GenerateSourceArray()
+ {
+	TArray<TSharedPtr<FName>> SourceArray;
+	if (MeshInterface != nullptr)
+	{
+			TWeakObjectPtr<USceneComponent> SceneComponent;
+			USkeletalMeshComponent* FoundSkelComp = nullptr;
+			USkeletalMesh* Mesh = UNiagaraDataInterfaceSkeletalMesh::GetSkeletalMeshHelper(MeshInterface, Cast<UNiagaraComponent>(MeshInterface->GetOuter()), SceneComponent, FoundSkelComp);
+
+			if (Mesh != nullptr)
+			{
+				for (FSkeletalMeshSamplingRegion Region : Mesh->GetSamplingInfo().Regions)
+				{
+					TSharedPtr<FName> RegionName(new FName(Region.Name));
+					SourceArray.Add(RegionName);
+				}
+			}
+	}
+	return SourceArray;
+ }
+
+ FNiagaraDataInterfaceSkeletalMeshDetails::~FNiagaraDataInterfaceSkeletalMeshDetails()
+ {
+	 if (MeshInterface)
+	 {
+		 MeshInterface->OnChanged().RemoveAll(this);
+	 }
+	 if (MeshObject)
+	 {
+		MeshObject->GetOnMeshChanged().RemoveAll(this);
+	 }
+ }
+
+#undef LOCTEXT_NAMESPACE
\ No newline at end of file
diff --git a/Engine/Plugins/FX/Niagara/Source/NiagaraEditorWidgets/Private/DetailCustomizations/NiagaraDetailSourcedArrayBuilder.h b/Engine/Plugins/FX/Niagara/Source/NiagaraEditorWidgets/Private/DetailCustomizations/NiagaraDetailSourcedArrayBuilder.h
new file mode 100644
--- /dev/null
+++ b/Engine/Plugins/FX/Niagara/Source/NiagaraEditorWidgets/Private/DetailCustomizations/NiagaraDetailSourcedArrayBuilder.h
@@ -0,0 +1,36 @@
+// Copyright 1998-2017 Epic Games, Inc. All Rights Reserved.
+
+#pragma once
+
+#include "Types/SlateEnums.h"
+#include "Widgets/Input/SComboBox.h"
+#include "PropertyCustomizationHelpers.h"
+class IDetailLayoutBuilder;
+class IDetailChildrenBuilder;
+class IPropertyHandle;
+class SWidget;
+class SNiagaraNamePropertySelector;
+
+/** Details customization for arrays composed of FName properties (or wrappers). The array contents are selected from a predetermined source list.*/
+class FNiagaraDetailSourcedArrayBuilder : public FDetailArrayBuilder
+										, public TSharedFromThis<FNiagaraDetailSourcedArrayBuilder>
+{
+public:
+	FNiagaraDetailSourcedArrayBuilder(TSharedRef<IPropertyHandle> InBaseProperty, const TArray<TSharedPtr<FName>>& InOptionsSource, const FName InFNameSubproperty = NAME_None, bool InGenerateHeader = true, bool InDisplayResetToDefault = true, bool InDisplayElementNum = true);
+	
+	virtual void GenerateChildContent(IDetailChildrenBuilder& ChildrenBuilder) override;
+	
+	void SetSourceArray(TArray<TSharedPtr<FName>>& InOptionsSource);
+
+private: 
+	void OnGenerateEntry(TSharedRef<IPropertyHandle> ElementProperty, int32 ElementIndex, IDetailChildrenBuilder& ChildrenBuilder);
+	
+private:
+	TArray<TSharedPtr<FName>> OptionsSourceList;
+	/** List of component class names, filtered by the current search string */
+
+	TSharedPtr<IPropertyHandleArray> ArrayProperty;
+	/** Subproperty of type FName that needs to be edited - in case of FName wrappers */
+	FName FNameSubproperty;
+	TSharedPtr<SNiagaraNamePropertySelector> NameSelector;
+};
diff --git a/Engine/Plugins/FX/Niagara/Source/NiagaraEditorWidgets/Private/DetailCustomizations/NiagaraSkeletalMeshSamplingInfoDetails.cpp b/Engine/Plugins/FX/Niagara/Source/NiagaraEditorWidgets/Private/DetailCustomizations/NiagaraSkeletalMeshSamplingInfoDetails.cpp
new file mode 100644
--- /dev/null
+++ b/Engine/Plugins/FX/Niagara/Source/NiagaraEditorWidgets/Private/DetailCustomizations/NiagaraSkeletalMeshSamplingInfoDetails.cpp
@@ -0,0 +1,92 @@
+ // Copyright 1998-2017 Epic Games, Inc. All Rights Reserved.
+#include "NiagaraSkeletalMeshSamplingInfoDetails.h"
+#include "Animation/Skeleton.h"
+#include "ReferenceSkeleton.h"
+#include "Engine/SkeletalMesh.h"
+#include "NiagaraDetailSourcedArrayBuilder.h"
+#include "IDetailCustomization.h"
+#include "DetailLayoutBuilder.h"
+#include "DetailCategoryBuilder.h"
+#include "IDetailGroup.h"
+#include "PropertyCustomizationHelpers.h"
+#include "IDetailChildrenBuilder.h"
+
+
+ #define LOCTEXT_NAMESPACE "FNiagaraSkeletalMeshSamplingInfoDetails"
+
+ void FNiagaraSkeletalMeshSamplingInfoDetails::CustomizeDetails(IDetailLayoutBuilder& DetailBuilder)
+ {
+	 static const FName SamplingCategoryName = TEXT("Sampling");
+	 static const FString RegionsPropertyNameString = "Regions";
+
+	 TArray<TWeakObjectPtr<UObject>> SelectedObjects;
+	 DetailBuilder.GetObjectsBeingCustomized(SelectedObjects);
+	 check(SelectedObjects.Num() == 1);
+	 USkeletalMesh*  MeshObject = Cast<USkeletalMesh>(SelectedObjects[0].Get());
+	 check(MeshObject);
+	 
+	 if (MeshObject->Skeleton != nullptr)
+	 {
+		 //Populate PossibleBonesNames
+		 for (FMeshBoneInfo Bone : MeshObject->Skeleton->GetReferenceSkeleton().GetRefBoneInfo())
+		 {
+			 TSharedPtr<FName> BoneName(new FName(Bone.Name));
+			 PossibleBonesNames.Add(BoneName);
+		 }
+	 }
+
+	 //Populate PossibleMaterialsNames
+	 for (FSkeletalMaterial Material : MeshObject->Materials)
+	 {
+		 TSharedPtr<FName> MaterialName(new FName(Material.MaterialSlotName));
+		 PossibleMaterialsNames.Add(MaterialName);
+	 }
+
+	 IDetailCategoryBuilder& SamplingCategory = DetailBuilder.EditCategory(SamplingCategoryName, LOCTEXT("Sampling", "Sampling"));
+	 {
+		 TArray<TSharedRef<IPropertyHandle>> SamplingProperties;
+		 SamplingCategory.GetDefaultProperties(SamplingProperties, true, false);
+
+		 for (TSharedPtr<IPropertyHandle> Property : SamplingProperties)
+		 {
+			 if (Property->GetProperty()->GetName() == RegionsPropertyNameString)
+			 {
+				 TSharedRef<FDetailArrayBuilder> RegionsBuilder = MakeShareable(new FDetailArrayBuilder(Property.ToSharedRef(), true, true));
+				 RegionsBuilder->OnGenerateArrayElementWidget(FOnGenerateArrayElementWidget::CreateSP(this, &FNiagaraSkeletalMeshSamplingInfoDetails::OnGenerateRegionEntry));
+				 SamplingCategory.AddCustomBuilder(RegionsBuilder);
+			 }
+			 else
+			 {
+				 SamplingCategory.AddProperty(Property);
+			 }
+		 }
+	 }
+ }
+
+ void FNiagaraSkeletalMeshSamplingInfoDetails::OnGenerateRegionEntry(TSharedRef<IPropertyHandle> PropertyHandle, int32 ArrayIndex, IDetailChildrenBuilder& ChildrenBuilder)
+  {
+ 	 IDetailPropertyRow& RegionRow = ChildrenBuilder.AddProperty(PropertyHandle);
+ 
+ 	 FNumberFormattingOptions NoCommas;
+ 	 NoCommas.UseGrouping = false;
+ 	 const FText SlotDesc = FText::Format(LOCTEXT("RegionSlotIndex", "Region #{0}"), FText::AsNumber(ArrayIndex, &NoCommas));
+ 
+ 	 RegionRow.DisplayName(SlotDesc);
+ 
+ 	 RegionRow.ShowPropertyButtons(true);
+	 // Add custom builder for properties
+	 TSharedPtr<IPropertyHandle> BoneFiltersProperty = PropertyHandle->GetChildHandle("BoneFilters");
+	 TSharedRef<FNiagaraDetailSourcedArrayBuilder> BoneBuilder = MakeShareable(new FNiagaraDetailSourcedArrayBuilder(BoneFiltersProperty.ToSharedRef(), PossibleBonesNames, "BoneName"));
+	 ChildrenBuilder.AddCustomBuilder(BoneBuilder);
+
+	 TSharedPtr<IPropertyHandle> MaterialFiltersProperty = PropertyHandle->GetChildHandle("MaterialFilters");
+	 TSharedRef<FNiagaraDetailSourcedArrayBuilder> MaterialBuilder = MakeShareable(new FNiagaraDetailSourcedArrayBuilder(MaterialFiltersProperty.ToSharedRef(), PossibleMaterialsNames, "MaterialName"));
+	 ChildrenBuilder.AddCustomBuilder(MaterialBuilder);
+
+  }
+
+ TSharedRef<IDetailCustomization> FNiagaraSkeletalMeshSamplingInfoDetails::MakeInstance()
+ {
+	 return MakeShared<FNiagaraSkeletalMeshSamplingInfoDetails>();
+ }
+#undef LOCTEXT_NAMESPACE
\ No newline at end of file
diff --git a/Engine/Plugins/FX/Niagara/Source/NiagaraEditorWidgets/Private/DetailCustomizations/NiagaraSkeletalMeshSamplingInfoDetails.h b/Engine/Plugins/FX/Niagara/Source/NiagaraEditorWidgets/Private/DetailCustomizations/NiagaraSkeletalMeshSamplingInfoDetails.h
new file mode 100644
--- /dev/null
+++ b/Engine/Plugins/FX/Niagara/Source/NiagaraEditorWidgets/Private/DetailCustomizations/NiagaraSkeletalMeshSamplingInfoDetails.h
@@ -0,0 +1,21 @@
+// Copyright 1998-2017 Epic Games, Inc. All Rights Reserved.
+
+#pragma once
+
+#include "IDetailCustomization.h"
+class IDetailLayoutBuilder;
+class IPropertyHandle;
+class IDetailChildrenBuilder;
+
+/** Details customization for Niagara skeletal mesh data interface. */
+class FNiagaraSkeletalMeshSamplingInfoDetails : public IDetailCustomization
+{
+public:
+	virtual void CustomizeDetails(IDetailLayoutBuilder& DetailBuilder) override;
+	static TSharedRef<IDetailCustomization> MakeInstance();
+private:
+	void OnGenerateRegionEntry(TSharedRef<IPropertyHandle> ElementProperty, int32 ElementIndex, IDetailChildrenBuilder& ChildrenBuilder);
+private:
+	TArray<TSharedPtr<FName>> PossibleBonesNames;
+	TArray<TSharedPtr<FName>> PossibleMaterialsNames;
+};
diff --git a/Engine/Plugins/FX/Niagara/Source/NiagaraEditorWidgets/Private/NiagaraEditorWidgetsModule.cpp b/Engine/Plugins/FX/Niagara/Source/NiagaraEditorWidgets/Private/NiagaraEditorWidgetsModule.cpp
--- a/Engine/Plugins/FX/Niagara/Source/NiagaraEditorWidgets/Private/NiagaraEditorWidgetsModule.cpp
+++ b/Engine/Plugins/FX/Niagara/Source/NiagaraEditorWidgets/Private/NiagaraEditorWidgetsModule.cpp
@@ -5,6 +5,9 @@
 #include "NiagaraEditorModule.h"
 #include "SNiagaraStack.h"
 #include "DetailCustomizations/NiagaraDataInterfaceCurveDetails.h"
+#include "DetailCustomizations/NiagaraDataInterfaceDetails.h"
+#include "DetailCustomizations/NiagaraDataInterfaceSkeletalMeshDetails.h"
+#include "DetailCustomizations/NiagaraSkeletalMeshSamplingInfoDetails.h"
 
 #include "Modules/ModuleManager.h"
 #include "PropertyEditorModule.h"
@@ -88,11 +91,14 @@ void FNiagaraEditorWidgetsModule::StartupModule()
 	FNiagaraEditorWidgetsStyle::Initialize();
 
 	FPropertyEditorModule& PropertyModule = FModuleManager::LoadModuleChecked<FPropertyEditorModule>("PropertyEditor");
+	PropertyModule.RegisterCustomClassLayout("NiagaraDataInterface", FOnGetDetailCustomizationInstance::CreateStatic(&FNiagaraDataInterfaceDetailsBase::MakeInstance));
 	PropertyModule.RegisterCustomClassLayout("NiagaraDataInterfaceCurve", FOnGetDetailCustomizationInstance::CreateStatic(&FNiagaraDataInterfaceCurveDetails::MakeInstance));
 	PropertyModule.RegisterCustomClassLayout("NiagaraDataInterfaceVector2DCurve", FOnGetDetailCustomizationInstance::CreateStatic(&FNiagaraDataInterfaceVector2DCurveDetails::MakeInstance));
 	PropertyModule.RegisterCustomClassLayout("NiagaraDataInterfaceVectorCurve", FOnGetDetailCustomizationInstance::CreateStatic(&FNiagaraDataInterfaceVectorCurveDetails::MakeInstance));
 	PropertyModule.RegisterCustomClassLayout("NiagaraDataInterfaceVector4Curve", FOnGetDetailCustomizationInstance::CreateStatic(&FNiagaraDataInterfaceVector4CurveDetails::MakeInstance));
 	PropertyModule.RegisterCustomClassLayout("NiagaraDataInterfaceColorCurve", FOnGetDetailCustomizationInstance::CreateStatic(&FNiagaraDataInterfaceColorCurveDetails::MakeInstance));
+	PropertyModule.RegisterCustomClassLayout("NiagaraDataInterfaceSkeletalMesh", FOnGetDetailCustomizationInstance::CreateStatic(&FNiagaraDataInterfaceSkeletalMeshDetails::MakeInstance));
+	PropertyModule.RegisterCustomClassLayout("SkeletalMesh", FOnGetDetailCustomizationInstance::CreateStatic(&FNiagaraSkeletalMeshSamplingInfoDetails::MakeInstance));
 }
 
 void FNiagaraEditorWidgetsModule::ShutdownModule()
@@ -106,11 +112,13 @@ void FNiagaraEditorWidgetsModule::ShutdownModule()
 	FPropertyEditorModule* PropertyModule = FModuleManager::GetModulePtr<FPropertyEditorModule>("PropertyEditor");
 	if (PropertyModule != nullptr)
 	{
+		PropertyModule->UnregisterCustomClassLayout("NiagaraDataInterface");
 		PropertyModule->UnregisterCustomClassLayout("NiagaraDataInterfaceCurve");
 		PropertyModule->UnregisterCustomClassLayout("NiagaraDataInterfaceVector2DCurve");
 		PropertyModule->UnregisterCustomClassLayout("NiagaraDataInterfaceVectorCurve");
 		PropertyModule->UnregisterCustomClassLayout("NiagaraDataInterfaceVector4Curve");
 		PropertyModule->UnregisterCustomClassLayout("NiagaraDataInterfaceColorCurve");
+		PropertyModule->UnregisterCustomClassLayout("NiagaraDataInterfaceSkeletalMesh");
 	}
 
 	FNiagaraEditorWidgetsStyle::Shutdown();
diff --git a/Engine/Plugins/FX/Niagara/Source/NiagaraEditorWidgets/Private/NiagaraEditorWidgetsStyle.cpp b/Engine/Plugins/FX/Niagara/Source/NiagaraEditorWidgets/Private/NiagaraEditorWidgetsStyle.cpp
--- a/Engine/Plugins/FX/Niagara/Source/NiagaraEditorWidgets/Private/NiagaraEditorWidgetsStyle.cpp
+++ b/Engine/Plugins/FX/Niagara/Source/NiagaraEditorWidgets/Private/NiagaraEditorWidgetsStyle.cpp
@@ -8,9 +8,10 @@
 #include "Styling/SlateStyleRegistry.h"
 #include "Styling/SlateTypes.h"
 #include "Styling/CoreStyle.h"
+#include "Interfaces/IPluginManager.h"
 
 TSharedPtr< FSlateStyleSet > FNiagaraEditorWidgetsStyle::NiagaraEditorWidgetsStyleInstance = NULL;
-
+ 
 void FNiagaraEditorWidgetsStyle::Initialize()
 {
 	if (!NiagaraEditorWidgetsStyleInstance.IsValid())
@@ -33,7 +34,15 @@ FName FNiagaraEditorWidgetsStyle::GetStyleSetName()
 	return StyleSetName;
 }
 
+FString RelativePathToPluginPath(const FString& RelativePath, const ANSICHAR* Extension)
+{
+	static FString ContentDir = IPluginManager::Get().FindPlugin(TEXT("Niagara"))->GetContentDir();
+	return (ContentDir / RelativePath) + Extension;
+}
+
+#define IMAGE_PLUGIN_BRUSH( RelativePath, ... ) FSlateImageBrush( RelativePathToPluginPath( RelativePath, ".png" ), __VA_ARGS__ )
 #define IMAGE_BRUSH( RelativePath, ... ) FSlateImageBrush( Style->RootToContentDir( RelativePath, TEXT(".png") ), __VA_ARGS__ )
+#define BOX_PLUGIN_BRUSH( RelativePath, ... ) FSlateBoxBrush( RelativePathToPluginPath( RelativePath, ".png"), __VA_ARGS__ )
 #define BOX_BRUSH( RelativePath, ... ) FSlateBoxBrush( Style->RootToContentDir( RelativePath, TEXT(".png") ), __VA_ARGS__ )
 #define BORDER_BRUSH( RelativePath, ... ) FSlateBorderBrush( Style->RootToContentDir( RelativePath, TEXT(".png") ), __VA_ARGS__ )
 #define DEFAULT_FONT(...) FCoreStyle::GetDefaultFontStyle(__VA_ARGS__)
@@ -41,6 +50,8 @@ FName FNiagaraEditorWidgetsStyle::GetStyleSetName()
 #define BOX_CORE_BRUSH( RelativePath, ... ) FSlateBoxBrush( FPaths::EngineContentDir() / "Editor/Slate" / RelativePath + TEXT(".png"), __VA_ARGS__ )
 
 const FVector2D Icon8x8(8.0f, 8.0f);
+const FVector2D Icon8x16(8.0f, 16.0f);
+const FVector2D Icon12x12(12.0f, 12.0f);
 const FVector2D Icon16x16(16.0f, 16.0f);
 const FVector2D Icon20x20(20.0f, 20.0f);
 const FVector2D Icon40x40(40.0f, 40.0f);
@@ -68,6 +79,12 @@ TSharedRef< FSlateStyleSet > FNiagaraEditorWidgetsStyle::Create()
 		.SetFont(StackDefaultFont);
 	Style->Set("NiagaraEditor.Stack.DefaultText", StackDefaultText);
 
+	FSlateFontInfo StackCategoryFont = DEFAULT_FONT("Bold", 10);
+	FTextBlockStyle StackCategoryText = FTextBlockStyle(NormalText)
+		.SetFont(StackCategoryFont)
+		.SetShadowOffset(FVector2D(1, 1));
+	Style->Set("NiagaraEditor.Stack.CategoryText", StackCategoryText);
+
 	FSlateFontInfo ParameterFont = DEFAULT_FONT("Regular", 8);
 	FTextBlockStyle ParameterText = FTextBlockStyle(NormalText)
 		.SetFont(ParameterFont);
@@ -84,20 +101,58 @@ TSharedRef< FSlateStyleSet > FNiagaraEditorWidgetsStyle::Create()
 	Style->Set("NiagaraEditor.Stack.ItemText", StackItemText);
 
 	Style->Set("NiagaraEditor.Stack.Group.BackgroundColor", FLinearColor(FColor(96, 96, 96)));
+	Style->Set("NiagaraEditor.Stack.Item.HeaderBackgroundColor", FLinearColor(FColor(48, 48, 48)));
+	Style->Set("NiagaraEditor.Stack.Item.ContentBackgroundColor", FLinearColor(FColor(62, 62, 62)));
+	Style->Set("NiagaraEditor.Stack.Item.ContentAdvancedBackgroundColor", FLinearColor(FColor(53, 53, 53)));
+	Style->Set("NiagaraEditor.Stack.Item.FooterBackgroundColor", FLinearColor(FColor(71, 71, 71)));
+	Style->Set("NiagaraEditor.Stack.Item.IssueBackgroundColor", FLinearColor(FColor(120, 120, 62)));
+	Style->Set("NiagaraEditor.Stack.UnknownColor", FLinearColor(1, 0, 1));
 
-	Style->Set("NiagaraEditor.Stack.Item.BackgroundColor", FLinearColor(FColor(35, 35, 35)));
+	Style->Set("NiagaraEditor.Stack.ItemHeaderFooter.BackgroundBrush", new FSlateColorBrush(FLinearColor(FColor(20, 20, 20))));
 
-	Style->Set("NiagaraEditor.Stack.Item.ErrorBackgroundColor", FLinearColor(FColor(35, 0, 0)));
+	Style->Set("NiagaraEditor.Stack.ForegroundColor", FLinearColor(FColor(220, 220, 220)));
+	Style->Set("NiagaraEditor.Stack.GroupForegroundColor", FLinearColor(FColor(220, 220, 220)));
+	Style->Set("NiagaraEditor.Stack.FlatButtonColor", FLinearColor(FColor(191, 191, 191)));
 
-	Style->Set("NiagaraEditor.Stack.ItemHeaderFooter.BackgroundBrush", new FSlateColorBrush(FLinearColor(FColor(20, 20, 20))));
+	Style->Set("NiagaraEditor.Stack.AccentColor.System", FLinearColor(FColor(67, 105, 124)));
+	Style->Set("NiagaraEditor.Stack.AccentColor.Emitter", FLinearColor(FColor(126, 87, 67)));
+	Style->Set("NiagaraEditor.Stack.AccentColor.Particle", FLinearColor(FColor(87, 107, 61)));
+	Style->Set("NiagaraEditor.Stack.AccentColor.Render", FLinearColor(FColor(134, 80, 80)));
+	Style->Set("NiagaraEditor.Stack.AccentColor.None", FLinearColor::Transparent);
 
-	Style->Set("NiagaraEditor.Stack.ForegroundColor", FLinearColor(FColor(191, 191, 191)));
+	Style->Set("NiagaraEditor.Stack.IconColor.System", FLinearColor(FColor(1, 202, 252)));
+	Style->Set("NiagaraEditor.Stack.IconColor.Emitter", FLinearColor(FColor(241, 99, 6)));
+	Style->Set("NiagaraEditor.Stack.IconColor.Particle", FLinearColor(FColor(131, 228, 9)));
+	Style->Set("NiagaraEditor.Stack.IconColor.Render", FLinearColor(FColor(230, 102, 102)));
+
+	Style->Set("NiagaraEditor.Stack.ParametersIcon", new IMAGE_PLUGIN_BRUSH("Icons/SystemParams", Icon12x12, FLinearColor::White));
+	Style->Set("NiagaraEditor.Stack.SpawnIcon", new IMAGE_PLUGIN_BRUSH("Icons/Spawn", Icon12x12, FLinearColor::White));
+	Style->Set("NiagaraEditor.Stack.UpdateIcon", new IMAGE_PLUGIN_BRUSH("Icons/Update", Icon12x12, FLinearColor::White));
+	Style->Set("NiagaraEditor.Stack.EventIcon", new IMAGE_PLUGIN_BRUSH("Icons/Event", Icon12x12, FLinearColor::White));
+
+	Style->Set("NiagaraEditor.Stack.ParametersIconHighlighted", new IMAGE_PLUGIN_BRUSH("Icons/SystemParams", Icon16x16, FLinearColor::White));
+	Style->Set("NiagaraEditor.Stack.SpawnIconHighlighted", new IMAGE_PLUGIN_BRUSH("Icons/Spawn", Icon16x16, FLinearColor::White));
+	Style->Set("NiagaraEditor.Stack.UpdateIconHighlighted", new IMAGE_PLUGIN_BRUSH("Icons/Update", Icon16x16, FLinearColor::White));
+	Style->Set("NiagaraEditor.Stack.EventIconHighlighted", new IMAGE_PLUGIN_BRUSH("Icons/Event", Icon16x16, FLinearColor::White));
+
+	Style->Set("NiagaraEditor.Stack.IconHighlightedSize", 16.0f);
 
 	Style->Set("NiagaraEditor.Stack.Splitter", FSplitterStyle()
 		.SetHandleNormalBrush(IMAGE_CORE_BRUSH("Common/SplitterHandleHighlight", Icon8x8, FLinearColor(.1f, .1f, .1f, 1.0f)))
 		.SetHandleHighlightBrush(IMAGE_CORE_BRUSH("Common/SplitterHandleHighlight", Icon8x8, FLinearColor::White))
 	);
 
+	Style->Set("NiagaraEditor.Stack.SearchHighlightColor", FLinearColor(FColor::Orange));
+	Style->Set("NiagaraEditor.Stack.SearchResult", new BOX_PLUGIN_BRUSH("Icons/SearchResultBorder", FMargin(1.f/8.f)));
+
+	Style->Set("NiagaraEditor.Stack.AddButton", FButtonStyle()
+		.SetNormal(BOX_CORE_BRUSH("Common/FlatButton", 2.0f / 8.0f, FLinearColor(0, 0, 0, .25f)))
+		.SetHovered(BOX_CORE_BRUSH("Common/FlatButton", 2.0f / 8.0f, FEditorStyle::GetSlateColor("SelectionColor")))
+		.SetPressed(BOX_CORE_BRUSH("Common/FlatButton", 2.0f / 8.0f, FEditorStyle::GetSlateColor("SelectionColor_Pressed")))
+	);
+
+	Style->Set("NiagaraEditor.ShowInCurveEditorIcon", new IMAGE_PLUGIN_BRUSH("Icons/ShowInCurveEditor", Icon16x16, FLinearColor::White));
+
 	return Style;
 }
 
