Copying //UE4/Dev-Niagara to //UE4/Dev-Main (Source: //UE4/Dev-Niagara @ 4074996)

diff --git a/Engine/Plugins/FX/Niagara/Source/NiagaraEditorWidgets/Private/SNiagaraStack.cpp b/Engine/Plugins/FX/Niagara/Source/NiagaraEditorWidgets/Private/SNiagaraStack.cpp
--- a/Engine/Plugins/FX/Niagara/Source/NiagaraEditorWidgets/Private/SNiagaraStack.cpp
+++ b/Engine/Plugins/FX/Niagara/Source/NiagaraEditorWidgets/Private/SNiagaraStack.cpp
@@ -3,2360 +3,880 @@
 #include "SNiagaraStack.h"
 
 #include "NiagaraEditorModule.h"
-#include "NiagaraTypes.h"
-#include "INiagaraEditorTypeUtilities.h"
-#include "SNiagaraParameterEditor.h"
+#include "NiagaraEditorCommands.h"
+#include "EditorStyleSet.h"
 #include "NiagaraEditorWidgetsStyle.h"
 #include "NiagaraEditorStyle.h"
-#include "NiagaraScript.h"
+#include "NiagaraEmitterHandle.h"
+#include "ViewModels/NiagaraEmitterViewModel.h"
+#include "ViewModels/NiagaraEmitterHandleViewModel.h"
+#include "ViewModels/NiagaraSystemViewModel.h"
 #include "ViewModels/Stack/NiagaraStackViewModel.h"
-#include "ViewModels/Stack/NiagaraStackEntry.h"
 #include "ViewModels/Stack/NiagaraStackSpacer.h"
-#include "ViewModels/Stack/NiagaraStackItemGroup.h"
-#include "ViewModels/Stack/NiagaraStackItem.h"
-#include "ViewModels/Stack/NiagaraStackItemExpander.h"
-#include "ViewModels/Stack/NiagaraStackFunctionInputCollection.h"
-#include "ViewModels/Stack/NiagaraStackFunctionInput.h"
 #include "ViewModels/Stack/NiagaraStackModuleItemOutputCollection.h"
 #include "ViewModels/Stack/NiagaraStackModuleItemOutput.h"
+#include "ViewModels/Stack/NiagaraStackErrorItem.h"
+#include "ViewModels/Stack/NiagaraStackInputCategory.h" 
+#include "ViewModels/Stack/NiagaraStackPropertyRow.h"
+#include "ViewModels/Stack/NiagaraStackAdvancedExpander.h"
+#include "ViewModels/Stack/NiagaraStackFunctionInputCollection.h"
 #include "ViewModels/Stack/NiagaraStackModuleItem.h"
 #include "ViewModels/Stack/NiagaraStackRendererItem.h"
-#include "ViewModels/Stack/NiagaraStackAddModuleItem.h"
-#include "ViewModels/Stack/NiagaraStackAddRendererItem.h"
-#include "ViewModels/Stack/NiagaraStackObject.h"
-#include "ViewModels/Stack/NiagaraStackErrorItem.h"
-#include "ViewModels/Stack/NiagaraStackStruct.h"
 #include "ViewModels/Stack/NiagaraStackParameterStoreEntry.h"
 #include "ViewModels/Stack/NiagaraStackEmitterSpawnScriptItemGroup.h"
 #include "ViewModels/Stack/NiagaraStackEventScriptItemGroup.h"
-#include "FrontendFilterBase.h"
-#include "Framework/Commands/UIAction.h"
-#include "Framework/MultiBox/MultiBoxBuilder.h"
-#include "Modules/ModuleManager.h"
-#include "AssetRegistryModule.h"
-#include "EdGraph/EdGraphSchema.h"
-#include "Modules/ModuleManager.h"
-#include "IStructureDetailsView.h"
-#include "IDetailsView.h"
-#include "Editor/PropertyEditor/Public/PropertyEditorModule.h"
+#include "Framework/Application/SlateApplication.h"
 #include "IContentBrowserSingleton.h"
 #include "ContentBrowserModule.h"
-#include "SGraphActionMenu.h"
 #include "Widgets/Input/SButton.h"
 #include "Widgets/Images/SImage.h"
 #include "Widgets/Input/SComboButton.h"
-#include "Widgets/Layout/SSpacer.h"
-#include "Framework/Application/SlateApplication.h"
-#include "Widgets/Colors/SSimpleGradient.h"
 #include "Widgets/Text/SInlineEditableTextBlock.h"
 #include "Widgets/Input/SCheckBox.h"
-#include "NiagaraNodeAssignment.h"
-#include "UObject/UObjectGlobals.h"
-#include "NiagaraDataInterface.h"
-#include "ClassViewerModule.h"
-#include "ClassViewerFilter.h"
-#include "NiagaraRendererProperties.h"
-#include "NiagaraConstants.h"
-#include "Misc/NotifyHook.h"
 #include "NiagaraEmitter.h"
-#include "Customizations/NiagaraEventScriptPropertiesCustomization.h"
-#include "EdGraph/EdGraph.h"
-#include "Toolkits/AssetEditorManager.h"
-#include "NiagaraEditorUtilities.h"
-#include "Widgets/Input/SSearchBox.h"
-
-#define LOCTEXT_NAMESPACE "NiagaraStack"
-
-const float IndentSize = 16;
-const float TextIconSize = 16;
-
-class SNiagaraStackFunctionInputName : public SCompoundWidget
+#include "IDetailTreeNode.h"
+#include "Stack/SNiagaraStackFunctionInputName.h"
+#include "Stack/SNiagaraStackFunctionInputValue.h"
+#include "Stack/SNiagaraStackEmitterPropertiesItem.h"
+#include "Stack/SNiagaraStackEventHandlerPropertiesItem.h"
+#include "Stack/SNiagaraStackItemExpander.h"
+#include "Stack/SNiagaraStackItemGroup.h"
+#include "Stack/SNiagaraStackModuleItem.h"
+#include "Stack/SNiagaraStackParameterStoreEntryName.h"
+#include "Stack/SNiagaraStackParameterStoreEntryValue.h"
+#include "Stack/SNiagaraStackRendererItem.h"
+#include "Stack/SNiagaraStackTableRow.h"
+#include "NiagaraEditorWidgetsUtilities.h"
+#include "DragAndDrop/DecoratedDragDropOp.h"
+#include "Stack/SNiagaraStackErrorItem.h"
+#include "NiagaraStackEditorData.h"
+#include "ScopedTransaction.h"
+
+/** Contains data for a socket drag and drop operation in the StackEntry node. */
+class FNiagaraStackEntryDragDropOp : public FDecoratedDragDropOp
 {
 public:
-	DECLARE_DELEGATE_OneParam(FOnColumnWidthChanged, float);
+	DRAG_DROP_OPERATOR_TYPE(FNiagaraStackEntryDragDropOp, FDecoratedDragDropOp)
 
 public:
-	SLATE_BEGIN_ARGS(SNiagaraStackFunctionInputName) { }
-		SLATE_ATTRIBUTE(bool, IsRowActive)
-	SLATE_END_ARGS();
-
-	void Construct(const FArguments& InArgs, UNiagaraStackFunctionInput* InFunctionInput)
-	{
-		FunctionInput = InFunctionInput;
-		IsRowActive = InArgs._IsRowActive;
-
-		PinIsPinnedColor = FNiagaraEditorWidgetsStyle::Get().GetColor(FunctionInput->GetItemForegroundName());
-		PinIsUnpinnedColor = PinIsPinnedColor.Desaturate(.4f);
-
-		ChildSlot
-		[
-			SNew(SHorizontalBox)
-			// Name Label
-			+ SHorizontalBox::Slot()
-			.VAlign(VAlign_Center)
-			[
-				SAssignNew(NameTextBlock, SInlineEditableTextBlock)
-				.Style(FNiagaraEditorStyle::Get(), "NiagaraEditor.ParameterInlineEditableText")
-				.Text_UObject(FunctionInput, &UNiagaraStackEntry::GetDisplayName)
-				.IsReadOnly(this, &SNiagaraStackFunctionInputName::GetIsNameReadOnly)
-				.IsSelected(this, &SNiagaraStackFunctionInputName::GetIsNameWidgetSelected)
-				.OnTextCommitted(this, &SNiagaraStackFunctionInputName::OnNameTextCommitted)
-				.ToolTipText_UObject(FunctionInput, &UNiagaraStackFunctionInput::GetTooltipText, UNiagaraStackFunctionInput::EValueMode::Local)
-			]
-			// Pin
-			+ SHorizontalBox::Slot()
-			.VAlign(VAlign_Center)
-			.AutoWidth()
-			[
-				SNew(SButton)
-				.IsFocusable(false)
-				.Visibility(this, &SNiagaraStackFunctionInputName::GetPinVisibility)
-				.ToolTipText(LOCTEXT("PinToolTip", "Pin this input"))
-				.ButtonStyle(FEditorStyle::Get(), "HoverHintOnly")
-				.ForegroundColor(this, &SNiagaraStackFunctionInputName::GetPinColor)
-				.ContentPadding(2)
-				.OnClicked(this, &SNiagaraStackFunctionInputName::PinButtonPressed)
-				.Content()
-				[
-					SNew(STextBlock)
-					.Font(FEditorStyle::Get().GetFontStyle("FontAwesome.9"))
-					.Text(FText::FromString(FString(TEXT("\xf08d"))))
-					.RenderTransformPivot(FVector2D(.5f, .5f))
-				]
-			]
-		];
-	}
-
-	virtual void Tick(const FGeometry& AllottedGeometry, const double InCurrentTime, const float InDeltaTime) override
-	{
-		if (FunctionInput->GetIsRenamePending() && NameTextBlock.IsValid())
-		{
-			NameTextBlock->EnterEditingMode();
-			FunctionInput->SetIsRenamePending(false);
-		}
-		SCompoundWidget::Tick(AllottedGeometry, InCurrentTime, InDeltaTime);
-	}
-
-private:
-	FReply PinButtonPressed() const
-	{
-		FunctionInput->SetIsPinned(!FunctionInput->GetIsPinned());
-		return FReply::Handled();
-	}
-
-	FSlateColor GetPinColor() const
+	FNiagaraStackEntryDragDropOp(TArray<UNiagaraStackEntry*> InDraggedEntries)
 	{
-		return FunctionInput->GetIsPinned() ? PinIsPinnedColor : PinIsUnpinnedColor;
+		DraggedEntries = InDraggedEntries;
 	}
 
-	EVisibility GetPinVisibility() const
+	const TArray<UNiagaraStackEntry*> GetDraggedEntries() const
 	{
-		if (FunctionInput->GetCanBePinned())
-		{
-			return IsRowActive.Get() || FunctionInput->GetIsPinned() ? EVisibility::Visible : EVisibility::Hidden;
-		}
-		else
-		{
-			return EVisibility::Collapsed;
-		}
-	}
-
-	bool GetIsNameReadOnly() const
-	{
-		return FunctionInput->CanRenameInput() == false;
-	}
-
-	bool GetIsNameWidgetSelected() const
-	{
-		return true;
-	}
-
-	void OnNameTextCommitted(const FText& InText, ETextCommit::Type InCommitType)
-	{
-		FunctionInput->RenameInput(*InText.ToString());
+		return DraggedEntries;
 	}
 
 private:
-	FLinearColor PinIsPinnedColor;
-	FLinearColor PinIsUnpinnedColor;
+	TArray<UNiagaraStackEntry*> DraggedEntries;
+};
 
-	UNiagaraStackFunctionInput* FunctionInput;
+#define LOCTEXT_NAMESPACE "NiagaraStack"
 
-	TSharedPtr<SInlineEditableTextBlock> NameTextBlock;
+const float SpacerHeight = 6;
 
-	TAttribute<bool> IsRowActive;
-};
+const FText SNiagaraStack::OccurencesFormat = NSLOCTEXT("NiagaraStack", "OccurencesFound", "{0} / {1}");
 
-class SNiagaraStackFunctionInputValue : public SCompoundWidget
+void SNiagaraStack::Construct(const FArguments& InArgs, UNiagaraStackViewModel* InStackViewModel)
 {
-public:
-	DECLARE_DELEGATE_OneParam(FOnColumnWidthChanged, float)
-public:
-	SLATE_BEGIN_ARGS(SNiagaraStackFunctionInputValue) { }
-	SLATE_END_ARGS();
-
-	void Construct(const FArguments& InArgs, UNiagaraStackFunctionInput* InFunctionInput)
-	{
-		FunctionInput = InFunctionInput;
-
-		FunctionInput->OnValueChanged().AddSP(this, &SNiagaraStackFunctionInputValue::OnInputValueChanged);
-		DisplayedLocalValueStruct = FunctionInput->GetLocalValueStruct();
+	StackViewModel = InStackViewModel;
+	StackViewModel->OnStructureChanged().AddSP(this, &SNiagaraStack::StackStructureChanged);
+	StackViewModel->OnSearchCompleted().AddSP(this, &SNiagaraStack::OnStackSearchComplete); 
+	NameColumnWidth = .3f;
+	ContentColumnWidth = .7f;
+	PinIsPinnedColor = FNiagaraEditorWidgetsStyle::Get().GetColor("NiagaraEditor.Stack.ForegroundColor");
+	PinIsUnpinnedColor = PinIsPinnedColor.Desaturate(.4f);
+	CurrentPinColor = StackViewModel->GetSystemViewModel()->GetIsEmitterPinned(StackViewModel->GetEmitterHandleViewModel()->AsShared()) ? PinIsPinnedColor: PinIsUnpinnedColor;
 
-		FMargin ItemPadding = FMargin(0);
-		ChildSlot
+	ConstructHeaderWidget();
+	TSharedPtr<SWidget> HeaderBox;
+	ChildSlot
+	[
+		SAssignNew(HeaderBox, SVerticalBox)
+		+ SVerticalBox::Slot()
+		.AutoHeight()
+		.Padding(0, 0, 0, 3)
 		[
-			// Values
-			SNew(SHorizontalBox)
-			+ SHorizontalBox::Slot()
-			.VAlign(VAlign_Center)
-			.AutoWidth()
-			.Padding(0, 0, 3, 0)
-			[
-				// Value Icon
-				SNew(SBox)
-				.WidthOverride(TextIconSize)
-				.VAlign(VAlign_Center)
-				[
-					SNew(STextBlock)
-					.Font(FEditorStyle::Get().GetFontStyle("FontAwesome.10"))
-					.Text(this, &SNiagaraStackFunctionInputValue::GetInputIconText)
-					.ToolTipText(this, &SNiagaraStackFunctionInputValue::GetInputIconToolTip)
-					.ColorAndOpacity(this, &SNiagaraStackFunctionInputValue::GetInputIconColor)
-				]
-			]
-			+ SHorizontalBox::Slot()
-			[
-				// TODO Don't generate all of these widgets for every input, only generate the ones that are used based on the value type.
-
-				// Local struct
-				SNew(SOverlay)
-				+ SOverlay::Slot()
-				[
-					SAssignNew(LocalValueStructContainer, SBox)
-					.Visibility(this, &SNiagaraStackFunctionInputValue::GetValueWidgetVisibility, UNiagaraStackFunctionInput::EValueMode::Local)
-					[
-						ConstructLocalValueStructWidget()
-					]
-				]
-
-				// Linked handle
-				+ SOverlay::Slot()
-				.Padding(0.0f, 0.0f, 0.0f, 2.0f)
-				[
-					SNew(SBox)
-					.Visibility(this, &SNiagaraStackFunctionInputValue::GetValueWidgetVisibility, UNiagaraStackFunctionInput::EValueMode::Linked)
-					.ToolTipText_UObject(InFunctionInput, &UNiagaraStackFunctionInput::GetTooltipText, UNiagaraStackFunctionInput::EValueMode::Linked)
-					.VAlign(VAlign_Center)
-					[
-						SNew(STextBlock)
-						.TextStyle(FNiagaraEditorStyle::Get(), "NiagaraEditor.ParameterText")
-						.Text(this, &SNiagaraStackFunctionInputValue::GetLinkedValueHandleText)
-					]
-				]
-
-				// Data Object
-				+ SOverlay::Slot()
-				.Padding(0.0f, 0.0f, 0.0f, 2.0f)
-				[
-					SNew(SBox)
-					.Visibility(this, &SNiagaraStackFunctionInputValue::GetValueWidgetVisibility, UNiagaraStackFunctionInput::EValueMode::Data)
-					.VAlign(VAlign_Center)
-					[
-						SNew(STextBlock)
-						.TextStyle(FNiagaraEditorStyle::Get(), "NiagaraEditor.ParameterText")
-						.Text(this, &SNiagaraStackFunctionInputValue::GetDataValueText)
-					]
-				]
-
-				// Dynamic input name
-				+ SOverlay::Slot()
-				[
-					SNew(SBox)
-					.Visibility(this, &SNiagaraStackFunctionInputValue::GetValueWidgetVisibility, UNiagaraStackFunctionInput::EValueMode::Dynamic)
-					.VAlign(VAlign_Center)
-					[
-						SNew(STextBlock)
-						.TextStyle(FNiagaraEditorStyle::Get(), "NiagaraEditor.ParameterText")
-						.Text(this, &SNiagaraStackFunctionInputValue::GetDynamicValueText)
-						.OnDoubleClicked(this, &SNiagaraStackFunctionInputValue::DynamicInputTextDoubleClicked)
-					]
-				]
-
-				// Invalid input
-				+ SOverlay::Slot()
-				[
-					SNew(SBox)
-					.Visibility(this, &SNiagaraStackFunctionInputValue::GetValueWidgetVisibility, UNiagaraStackFunctionInput::EValueMode::Invalid)
-					.VAlign(VAlign_Center)
-					[
-						SNew(STextBlock)
-						.TextStyle(FNiagaraEditorStyle::Get(), "NiagaraEditor.ParameterText")
-						.Text(this, &SNiagaraStackFunctionInputValue::GetInvalidValueText)
-						.ToolTipText(this, &SNiagaraStackFunctionInputValue::GetInvalidValueToolTipText)
-					]
-				]
-			]
-
-			// Handle drop-down button
-			+ SHorizontalBox::Slot()
-			.AutoWidth()
-			.VAlign(VAlign_Center)
-			.Padding(3, 0, 0, 0)
-			[
-				SNew(SComboButton)
-				.ButtonStyle(FEditorStyle::Get(), "HoverHintOnly")
-				.ForegroundColor(FSlateColor::UseForeground())
-				.OnGetMenuContent(this, &SNiagaraStackFunctionInputValue::OnGetAvailableHandleMenu)
-				.ContentPadding(FMargin(2))
-				.HAlign(HAlign_Center)
-				.VAlign(VAlign_Center)
-			]
-
-			// Reset Button
-			+ SHorizontalBox::Slot()
-			.AutoWidth()
-			.VAlign(VAlign_Center)
-			.Padding(3, 0, 0, 0)
+			SNew(SBorder)
+			.BorderImage(FEditorStyle::GetBrush("ToolPanel.GroupBorder"))
+			.Padding(3)
 			[
-				SNew(SButton)
-				.IsFocusable(false)
-				.ToolTipText(LOCTEXT("ResetToolTip", "Reset to the default value"))
-				.ButtonStyle(FEditorStyle::Get(), "NoBorder")
-				.ContentPadding(0)
-				.Visibility(this, &SNiagaraStackFunctionInputValue::GetResetButtonVisibility)
-				.OnClicked(this, &SNiagaraStackFunctionInputValue::ResetButtonPressed)
-				.Content()
-				[
-					SNew(SImage)
-					.Image(FEditorStyle::GetBrush("PropertyWindow.DiffersFromDefault"))
-				]
+				HeaderWidget.ToSharedRef()
 			]
-
-			// Reset to base Button
-			+ SHorizontalBox::Slot()
-			.AutoWidth()
-			.VAlign(VAlign_Center)
-			.Padding(3, 0, 0, 0)
+		]
+		+ SVerticalBox::Slot()
+		.Padding(0)
+		[
+			SNew(SBorder)
+			.BorderImage(FEditorStyle::GetBrush("ToolPanel.GroupBorder"))
+			.Padding(5)
 			[
-				SNew(SButton)
-				.IsFocusable(false)
-				.ToolTipText(LOCTEXT("ResetToBaseToolTip", "Reset this input to the value defined by the parent emitter"))
-				.ButtonStyle(FEditorStyle::Get(), "NoBorder")
-				.ContentPadding(0)
-				.Visibility(this, &SNiagaraStackFunctionInputValue::GetResetToBaseButtonVisibility)
-				.OnClicked(this, &SNiagaraStackFunctionInputValue::ResetToBaseButtonPressed)
-				.Content()
-				[
-					SNew(SImage)
-					.Image(FEditorStyle::GetBrush("PropertyWindow.DiffersFromDefault"))
-					.ColorAndOpacity(FSlateColor(FLinearColor::Green))
-				]
+				SAssignNew(StackTree, STreeView<UNiagaraStackEntry*>)
+				.OnGenerateRow(this, &SNiagaraStack::OnGenerateRowForStackItem)
+				.OnGetChildren(this, &SNiagaraStack::OnGetChildren)
+				.TreeItemsSource(&StackViewModel->GetRootEntries())
+				.OnTreeViewScrolled(this, &SNiagaraStack::StackTreeScrolled)
 			]
-		];
-	}
+		]
+	];
 
-private:
-	EVisibility GetValueWidgetVisibility(UNiagaraStackFunctionInput::EValueMode ValidMode) const
-	{
-		return FunctionInput->GetValueMode() == ValidMode ? EVisibility::Visible : EVisibility::Collapsed;
-	}
+	StackTree->SetScrollOffset(StackViewModel->GetLastScrollPosition());
 
-	TSharedRef<SWidget> ConstructLocalValueStructWidget()
+	auto OnHeaderMouseButtonUp = [this](const FGeometry&, const FPointerEvent& MouseEvent) -> FReply
 	{
-		LocalValueStructParameterEditor.Reset();
-		LocalValueStructDetailsView.Reset();
-		if (DisplayedLocalValueStruct.IsValid())
+		if (MouseEvent.GetEffectingButton() == EKeys::RightMouseButton)
 		{
-			FNiagaraEditorModule& NiagaraEditorModule = FModuleManager::GetModuleChecked<FNiagaraEditorModule>("NiagaraEditor");
-			TSharedPtr<INiagaraEditorTypeUtilities> TypeEditorUtilities = NiagaraEditorModule.GetTypeUtilities(FunctionInput->GetInputType());
-			if (TypeEditorUtilities.IsValid() && TypeEditorUtilities->CanCreateParameterEditor())
+			FMenuBuilder MenuBuilder(true, nullptr);
+			MenuBuilder.BeginSection("EmitterInlineMenuActions", LOCTEXT("EmitterActions", "Emitter Actions"));
 			{
-				TSharedPtr<SNiagaraParameterEditor> ParameterEditor = TypeEditorUtilities->CreateParameterEditor(FunctionInput->GetInputType());
-				ParameterEditor->UpdateInternalValueFromStruct(DisplayedLocalValueStruct.ToSharedRef());
-				ParameterEditor->SetOnBeginValueChange(SNiagaraParameterEditor::FOnValueChange::CreateSP(
-					this, &SNiagaraStackFunctionInputValue::ParameterBeginValueChange));
-				ParameterEditor->SetOnEndValueChange(SNiagaraParameterEditor::FOnValueChange::CreateSP(
-					this, &SNiagaraStackFunctionInputValue::ParameterEndValueChange));
-				ParameterEditor->SetOnValueChanged(SNiagaraParameterEditor::FOnValueChange::CreateSP(
-					this, &SNiagaraStackFunctionInputValue::ParameterValueChanged, ParameterEditor.ToSharedRef()));
+				MenuBuilder.AddMenuEntry(
+					LOCTEXT("OpenAndFocusEmitter", "Open and focus source emitter"),
+					LOCTEXT("OpenAndFocusEmitterToolTip", "Open the source emitter in emitter mode"),
+					FSlateIcon(),
+					FUIAction(FExecuteAction::CreateSP(this, &SNiagaraStack::OpenSourceEmitter)));
+				MenuBuilder.AddMenuEntry(
+					LOCTEXT("ShowEmitterInContentBrowser", "Show in content browser"),
+					LOCTEXT("ShowEmitterInContentBrowserToolTip", "Show the emitter in this stack in the Content Browser"),
+					FSlateIcon(),
+					FUIAction(FExecuteAction::CreateSP(this, &SNiagaraStack::ShowEmitterInContentBrowser))); 
 				
-				LocalValueStructParameterEditor = ParameterEditor;
-				return ParameterEditor.ToSharedRef();
-			}
-			else
-			{
-				FPropertyEditorModule& PropertyEditorModule = FModuleManager::GetModuleChecked<FPropertyEditorModule>("PropertyEditor");
-
-				TSharedRef<IStructureDetailsView> StructureDetailsView = PropertyEditorModule.CreateStructureDetailView(
-					FDetailsViewArgs(false, false, false, FDetailsViewArgs::HideNameArea, true),
-					FStructureDetailsViewArgs(),
-					nullptr);
-
-				StructureDetailsView->SetStructureData(DisplayedLocalValueStruct);
-				StructureDetailsView->GetOnFinishedChangingPropertiesDelegate().AddSP(this, &SNiagaraStackFunctionInputValue::ParameterPropertyValueChanged);
+				FUIAction Action(FExecuteAction::CreateSP(this, &SNiagaraStack::SetEmitterEnabled, !StackViewModel->GetEmitterHandleViewModel()->GetIsEnabled()),
+					FCanExecuteAction(),
+					FIsActionChecked::CreateSP(this, &SNiagaraStack::CheckEmitterEnabledStatus, true));
+				MenuBuilder.AddMenuEntry(
+					LOCTEXT("IsEnabled", "Is Enabled"),
+					LOCTEXT("ToggleEmitterEnabledToolTip", "Toggle emitter enabled/disabled state"),
+					FSlateIcon(),
+					Action,
+					NAME_None, 
+					EUserInterfaceActionType::Check);
 
-				LocalValueStructDetailsView = StructureDetailsView;
-				return StructureDetailsView->GetWidget().ToSharedRef();
+				MenuBuilder.AddMenuEntry(
+					LOCTEXT("CollapseStack", "Collapse All"),
+					LOCTEXT("CollapseStackToolTip", "Collapses every row in the stack."),
+					FSlateIcon(),
+					FUIAction(FExecuteAction::CreateSP(this, &SNiagaraStack::CollapseAll)));
+				
+				TSharedPtr<FUICommandInfo> ExpandStackGroups = FNiagaraEditorModule::Get().Commands().CollapseStackToHeaders;
+				MenuBuilder.AddMenuEntry(
+					ExpandStackGroups->GetLabel(),
+					ExpandStackGroups->GetDescription(),
+					FSlateIcon(),
+					FUIAction(FExecuteAction::CreateUObject(StackViewModel, &UNiagaraStackViewModel::CollapseToHeaders)));
 			}
-		}
-		else
-		{
-			return SNullWidget::NullWidget;
-		}
-	}
-
-	void OnInputValueChanged()
-	{
-		TSharedPtr<FStructOnScope> NewLocalValueStruct = FunctionInput->GetLocalValueStruct();
-		if (DisplayedLocalValueStruct == NewLocalValueStruct)
-		{
-			if (LocalValueStructParameterEditor.IsValid())
+			MenuBuilder.EndSection();
+			MenuBuilder.BeginSection("EmitterNavigateTo", LOCTEXT("EmitterNavigateTo", "Navigate to Section:"));
 			{
-				LocalValueStructParameterEditor->UpdateInternalValueFromStruct(DisplayedLocalValueStruct.ToSharedRef());
+				// Parse all children of root entries and if they are UNiagaraStackItemGroup then add a navigate menu entry
+				TArray<UNiagaraStackEntry*> EntriesToProcess(StackViewModel->GetRootEntries());
+				TArray<UNiagaraStackEntry*> RootChildren;
+				while (EntriesToProcess.Num() > 0)
+				{
+					UNiagaraStackEntry* EntryToProcess = EntriesToProcess[0];
+					EntriesToProcess.RemoveAtSwap(0);
+					EntryToProcess->GetUnfilteredChildren(RootChildren);
+				}
+				for (auto RootChild: RootChildren)
+				{
+					UNiagaraStackItemGroup* GroupChild = Cast<UNiagaraStackItemGroup>(RootChild);
+					if (GroupChild != nullptr)
+					{
+						MenuBuilder.AddMenuEntry(
+							RootChild->GetDisplayName(),
+							FText::Format(LOCTEXT("EmitterTooltip", "Navigate to {0}"), RootChild->GetDisplayName()),
+							FSlateIcon(),
+							FUIAction(FExecuteAction::CreateSP(this, &SNiagaraStack::NavigateTo, RootChild)));
+					}
+				}
 			}
-			if (LocalValueStructDetailsView.IsValid())
+			MenuBuilder.EndSection();
+			
+			MenuBuilder.BeginSection("StackActions", LOCTEXT("StackActions", "Stack Actions"));
+			if (StackViewModel->HasDismissedStackIssues())
 			{
-				LocalValueStructDetailsView->SetStructureData(TSharedPtr<FStructOnScope>());
-				LocalValueStructDetailsView->SetStructureData(DisplayedLocalValueStruct);
+				MenuBuilder.AddMenuEntry(
+					LOCTEXT("UndismissIssues", "Undismiss all stack issues"),
+					LOCTEXT("ShowAssetInContentBrowserToolTip", "Undismiss all issues that were previously dismissed for this stack, if any"),
+					FSlateIcon(),
+					FUIAction(FExecuteAction::CreateUObject(StackViewModel, &UNiagaraStackViewModel::UndismissAllIssues)));
 			}
-		}
-		else
-		{
-			DisplayedLocalValueStruct = NewLocalValueStruct;
-			LocalValueStructContainer->SetContent(ConstructLocalValueStructWidget());
-		}
-	}
-
-	void ParameterBeginValueChange()
-	{
-		FunctionInput->NotifyBeginLocalValueChange();
-	}
-
-	void ParameterEndValueChange()
-	{
-		FunctionInput->NotifyEndLocalValueChange();
-	}
-
-	void ParameterValueChanged(TSharedRef<SNiagaraParameterEditor> ParameterEditor)
-	{
-		ParameterEditor->UpdateStructFromInternalValue(FunctionInput->GetLocalValueStruct().ToSharedRef());
-		FunctionInput->SetLocalValue(DisplayedLocalValueStruct.ToSharedRef());
-	}
-
-	void ParameterPropertyValueChanged(const FPropertyChangedEvent& PropertyChangedEvent)
-	{
-		FunctionInput->SetLocalValue(DisplayedLocalValueStruct.ToSharedRef());
-	}
-
-	FText GetLinkedValueHandleText() const
-	{
-		return FText::FromName(FunctionInput->GetLinkedValueHandle().GetParameterHandleString());
-	}
+			MenuBuilder.EndSection();
 
-	FText GetDataValueText() const
-	{
-		if (FunctionInput->GetDataValueObject() != nullptr)
-		{
-			return FunctionInput->GetInputType().GetClass()->GetDisplayNameText();
-		}
-		else
-		{
-			return FText::Format(LOCTEXT("InvalidDataObjectFormat", "{0} (Invalid)"), FunctionInput->GetInputType().GetClass()->GetDisplayNameText());
+			FWidgetPath WidgetPath = MouseEvent.GetEventPath() != nullptr ? *MouseEvent.GetEventPath() : FWidgetPath();
+			FSlateApplication::Get().PushMenu(AsShared(), WidgetPath, MenuBuilder.MakeWidget(), MouseEvent.GetScreenSpacePosition(), FPopupTransitionEffect(FPopupTransitionEffect::ContextMenu));
+			return FReply::Handled();
 		}
-	}
+		return FReply::Unhandled();
+	};
+	HeaderBox->SetOnMouseButtonUp(FPointerEventHandler::CreateLambda(OnHeaderMouseButtonUp));
+	
+	PrimeTreeExpansion();
+}
 
-	FText GetDynamicValueText() const
+void SNiagaraStack::PrimeTreeExpansion()
+{
+	TArray<UNiagaraStackEntry*> EntriesToProcess(StackViewModel->GetRootEntries());
+	while (EntriesToProcess.Num() > 0)
 	{
-		if (FunctionInput->GetDynamicInputNode() != nullptr)
-		{
-			return FText::FromString(FName::NameToDisplayString(FunctionInput->GetDynamicInputNode()->GetFunctionName(), false));
-		}
-		else
-		{
-			return LOCTEXT("InvalidDynamicDisplayName", "(Invalid)");
-		}
-	}
+		UNiagaraStackEntry* EntryToProcess = EntriesToProcess[0];
+		EntriesToProcess.RemoveAtSwap(0);
 
-	FText GetInvalidValueText() const
-	{
-		if (FunctionInput->CanReset())
+		if (EntryToProcess->GetIsExpanded())
 		{
-			return LOCTEXT("InvalidResetLabel", "Unsupported value - Reset to fix.");
+			StackTree->SetItemExpansion(EntryToProcess, true);
+			EntryToProcess->GetFilteredChildren(EntriesToProcess);
 		}
 		else
 		{
-			return LOCTEXT("InvalidLabel", "Unsupported value");
+			StackTree->SetItemExpansion(EntryToProcess, false);
 		}
 	}
+}
 
-	FText GetInvalidValueToolTipText() const
+void SNiagaraStack::ConstructHeaderWidget()
+{
+	if (!StackViewModel->GetEmitterHandleViewModel().IsValid() || !StackViewModel->GetSystemViewModel().IsValid())
 	{
-		if (FunctionInput->CanReset())
-		{
-			return LOCTEXT("InvalidResetToolTip", "This input has an unsupported value assigned in the stack.\nUse the reset button to remove the unsupported value.");
-		}
-		else
-		{
-			return LOCTEXT("InvalidToolTip", "The script that defines the source of this input has\n a default value that can not be displayed in the stack view.");
-		}
+		HeaderWidget = SNullWidget::NullWidget;
 	}
-
-	FReply DynamicInputTextDoubleClicked()
+	else
 	{
-		UNiagaraNodeFunctionCall* DynamicInputNode = FunctionInput->GetDynamicInputNode();
-		if (DynamicInputNode->FunctionScript != nullptr && DynamicInputNode->FunctionScript->IsAsset())
-		{
-			FAssetEditorManager::Get().OpenEditorForAsset(DynamicInputNode->FunctionScript);
-			return FReply::Handled();
-		}
-		return FReply::Unhandled();
-	}
+		SAssignNew(HeaderWidget, SVerticalBox)
+			//~ Enabled check box rename text box, and external header controls.
+			+ SVerticalBox::Slot()
+			.AutoHeight()
+			.HAlign(HAlign_Fill)
+			[
+				//~ Enabled
+				SNew(SHorizontalBox)
+				+ SHorizontalBox::Slot()
+				.HAlign(HAlign_Left)
+				.AutoWidth()
+				.Padding(2)
+				[
+					SNew(SCheckBox)
+					.ToolTipText(LOCTEXT("EnabledToolTip", "Toggles whether this emitter is enabled. Disabled emitters don't simulate or render."))
+					.IsChecked(StackViewModel->GetEmitterHandleViewModel()->AsShared(), &FNiagaraEmitterHandleViewModel::GetIsEnabledCheckState)
+					.OnCheckStateChanged(StackViewModel->GetEmitterHandleViewModel()->AsShared(), &FNiagaraEmitterHandleViewModel::OnIsEnabledCheckStateChanged)
+				]
+				// Pin
+				+ SHorizontalBox::Slot()
+				.VAlign(VAlign_Center)
+				.AutoWidth()
+				[
+					SNew(SButton)
+					.IsFocusable(false)
+					.ToolTipText(LOCTEXT("PinToolTip", "Pin this emitter"))
+					.ButtonStyle(FEditorStyle::Get(), "HoverHintOnly")
+					.ForegroundColor(this, &SNiagaraStack::GetPinColor)
+					.ContentPadding(2)
+					.OnClicked(this, &SNiagaraStack::PinButtonPressed)
+					.Content()
+					[
+						SNew(STextBlock)
+						.Font(FEditorStyle::Get().GetFontStyle("FontAwesome.9"))
+						.Text(FText::FromString(FString(TEXT("\xf08d"))))
+						.RenderTransformPivot(FVector2D(.5f, .5f))
+					]
+				]
+				//~ Name
+				+ SHorizontalBox::Slot()
+				.AutoWidth()
+				.Padding(2)
+				.HAlign(HAlign_Left)
+				[
+					SNew(SInlineEditableTextBlock)
+					.ToolTipText(NSLOCTEXT("NiagaraEmitterEditor", "NameTextToolTip", "Click to edit the emitter name."))
+					.Style(FNiagaraEditorStyle::Get(), "NiagaraEditor.HeadingInlineEditableText")
+					.WrapTextAt(150.0f)
+					.Text(StackViewModel->GetEmitterHandleViewModel()->AsShared(), &FNiagaraEmitterHandleViewModel::GetNameText)
+					.OnTextCommitted(StackViewModel->GetEmitterHandleViewModel()->AsShared(), &FNiagaraEmitterHandleViewModel::OnNameTextComitted)
+					.OnVerifyTextChanged(StackViewModel->GetEmitterHandleViewModel()->AsShared(), &FNiagaraEmitterHandleViewModel::VerifyNameTextChanged)
+				]
+			]
 
-	TSharedRef<SWidget> OnGetAvailableHandleMenu()
-	{
-		FMenuBuilder MenuBuilder(true, nullptr);
-
-		// Set a local value
-		bool bCanSetLocalValue = 
-			FunctionInput->GetValueMode() != UNiagaraStackFunctionInput::EValueMode::Local &&
-			FunctionInput->GetInputType().IsDataInterface() == false;
-		MenuBuilder.AddMenuEntry(
-			LOCTEXT("LocalValue", "Set a local value"),
-			LOCTEXT("LocalValueToolTip", "Set a local editable value for this input."),
-			FSlateIcon(),
-			FUIAction(
-				FExecuteAction::CreateSP(this, &SNiagaraStackFunctionInputValue::SetToLocalValue),
-				FCanExecuteAction::CreateLambda([=]() { return bCanSetLocalValue; })));
-
-		// Add a dynamic input
-		MenuBuilder.BeginSection(NAME_None, LOCTEXT("DynamicInputSection", "Dynamic Inputs"));
-		TArray<UNiagaraScript*> DynamicInputScripts;
-		FunctionInput->GetAvailableDynamicInputs(DynamicInputScripts);
-		for (UNiagaraScript* DynamicInputScript : DynamicInputScripts)
-		{
-			FText DynamicInputText = FText::FromString(FName::NameToDisplayString(DynamicInputScript->GetName(), false));
-			MenuBuilder.AddMenuEntry(
-				DynamicInputText,
-				FText::Format(LOCTEXT("DynamicInputFormat", "Use {0} to provide a value for this input."), DynamicInputText),
-				FSlateIcon(),
-				FUIAction(FExecuteAction::CreateSP(this, &SNiagaraStackFunctionInputValue::DynamicInputScriptSelected, DynamicInputScript)));
-		}
-		MenuBuilder.EndSection();
-
-		// Link existing attribute
-		TArray<FNiagaraParameterHandle> AvailableHandles;
-		FunctionInput->GetAvailableParameterHandles(AvailableHandles);
-
-		TArray<FNiagaraParameterHandle> EngineHandles;
-		TArray<FNiagaraParameterHandle> SystemHandles;
-		TArray<FNiagaraParameterHandle> EmitterHandles;
-		TArray<FNiagaraParameterHandle> ParticleAttributeHandles;
-		TArray<FNiagaraParameterHandle> OtherHandles;
-		for (const FNiagaraParameterHandle AvailableHandle : AvailableHandles)
-		{
-			if (AvailableHandle.IsEngineHandle())
-			{
-				EngineHandles.Add(AvailableHandle);
-			}
-			else if (AvailableHandle.IsSystemHandle())
-			{
-				SystemHandles.Add(AvailableHandle);
-			}
-			else if (AvailableHandle.IsEmitterHandle())
-			{
-				EmitterHandles.Add(AvailableHandle);
-			}
-			else if (AvailableHandle.IsParticleAttributeHandle())
-			{
-				ParticleAttributeHandles.Add(AvailableHandle);
-			}
-			else
-			{
-				OtherHandles.Add(AvailableHandle);
-			}
-		}
+			//~ Stats
+			+ SVerticalBox::Slot()
+			.AutoHeight()
+			.HAlign(HAlign_Fill)
+			.Padding(2)
+			[
+				SNew(STextBlock)
+				.Text(StackViewModel->GetEmitterHandleViewModel()->GetEmitterViewModel()->AsShared(), &FNiagaraEmitterViewModel::GetStatsText)
+			]
 
-		FText MapInputFormat = LOCTEXT("LinkInputFormat", "Link this input to {0}");
-
-		auto AddMenuItemsForHandleList = [&](const TArray<FNiagaraParameterHandle>& Handles, FText SectionDisplayText)
-		{
-			MenuBuilder.BeginSection(NAME_None, SectionDisplayText);
-			for (const FNiagaraParameterHandle& Handle : Handles)
-			{
-				FText HandleDisplayName = FText::FromString(FName::NameToDisplayString(Handle.GetName().ToString(), false));
-				MenuBuilder.AddMenuEntry(
-					HandleDisplayName,
-					FText::Format(MapInputFormat, FText::FromName(Handle.GetParameterHandleString())),
-					FSlateIcon(),
-					FUIAction(FExecuteAction::CreateSP(this, &SNiagaraStackFunctionInputValue::ParameterHandleSelected, Handle)));
-			}
-			MenuBuilder.EndSection();
-		};
-
-		AddMenuItemsForHandleList(EngineHandles, LOCTEXT("EngineSection", "Engine"));
-		AddMenuItemsForHandleList(SystemHandles, LOCTEXT("SystemSection", "System"));
-		AddMenuItemsForHandleList(EmitterHandles, LOCTEXT("EmitterSection", "Emitter"));
-		AddMenuItemsForHandleList(ParticleAttributeHandles, LOCTEXT("ParticleAttributeSection", "Particle Attribute"));
-
-		MenuBuilder.BeginSection(NAME_None, LOCTEXT("OtherSection", "Other"));
-		for (const FNiagaraParameterHandle OtherHandle : OtherHandles)
-		{
-			FText HandleDisplayName = FText::FromString(FName::NameToDisplayString(OtherHandle.GetParameterHandleString().ToString(), false));
-			MenuBuilder.AddMenuEntry(
-				HandleDisplayName,
-				FText::Format(MapInputFormat, FText::FromName(OtherHandle.GetParameterHandleString())),
-				FSlateIcon(),
-				FUIAction(FExecuteAction::CreateSP(this, &SNiagaraStackFunctionInputValue::ParameterHandleSelected, OtherHandle)));
-		}
-		MenuBuilder.EndSection();
-
-		if (AvailableHandles.Num() > 0 || DynamicInputScripts.Num() > 0)
-		{
-			MenuBuilder.AddMenuSeparator();
-		}
-
-		// Read from new attribute
-		TArray<FName> AvailableNamespaces;
-		FunctionInput->GetNamespacesForNewParameters(AvailableNamespaces);
-
-		TArray<FString> InputNames;
-		for (int32 i = FunctionInput->GetInputParameterHandlePath().Num() - 1; i >= 0; i--)
-		{
-			InputNames.Add(FunctionInput->GetInputParameterHandlePath()[i].GetName().ToString());
-		}
-		FName InputName = *FString::Join(InputNames, TEXT("."));
-
-		for (const FName AvailableNamespace : AvailableNamespaces)
-		{
-			FNiagaraParameterHandle HandleToRead(AvailableNamespace, InputName);
-			bool bCanExecute = AvailableHandles.Contains(HandleToRead) == false;
-
-			FFormatNamedArguments Args;
-			Args.Add(TEXT("AvailableNamespace"), FText::FromName(AvailableNamespace));
-
-			MenuBuilder.AddMenuEntry(
-				FText::Format(LOCTEXT("ReadLabelFormat", "Read from new {AvailableNamespace} parameter"), Args),
-				FText::Format(LOCTEXT("ReadToolTipFormat", "Read this input from a new parameter in the {AvailableNamespace} namespace."), Args), 
-				FSlateIcon(),
-				FUIAction(
-					FExecuteAction::CreateSP(this, &SNiagaraStackFunctionInputValue::ParameterHandleSelected, HandleToRead),
-					FCanExecuteAction::CreateLambda([=]() { return bCanExecute; })));
-		}
-
-		return MenuBuilder.MakeWidget();
-	}
-
-	void SetToLocalValue()
-	{
-		const UScriptStruct* LocalValueStruct = FunctionInput->GetInputType().GetScriptStruct();
-		if (LocalValueStruct != nullptr)
-		{
-			TSharedRef<FStructOnScope> LocalValue = MakeShared<FStructOnScope>(LocalValueStruct);
-			TArray<uint8> DefaultValueData;
-			FNiagaraEditorUtilities::GetTypeDefaultValue(FunctionInput->GetInputType(), DefaultValueData);
-			if (DefaultValueData.Num() == LocalValueStruct->GetStructureSize())
-			{
-				FMemory::Memcpy(LocalValue->GetStructMemory(), DefaultValueData.GetData(), DefaultValueData.Num());
-				FunctionInput->SetLocalValue(LocalValue);
-			}
-		}
-	}
-
-	void DynamicInputScriptSelected(UNiagaraScript* DynamicInputScript)
-	{
-		FunctionInput->SetDynamicInput(DynamicInputScript);
-	}
-
-	void ParameterHandleSelected(FNiagaraParameterHandle Handle)
-	{
-		FunctionInput->SetLinkedValueHandle(Handle);
-	}
-
-	EVisibility GetResetButtonVisibility() const
-	{
-		return FunctionInput->CanReset() ? EVisibility::Visible : EVisibility::Hidden;
-	}
-
-	FReply ResetButtonPressed() const
-	{
-		FunctionInput->Reset();
-		return FReply::Handled();
-	}
-
-	EVisibility GetResetToBaseButtonVisibility() const
-	{
-		if (FunctionInput->EmitterHasBase())
-		{
-			return FunctionInput->CanResetToBase() ? EVisibility::Visible : EVisibility::Hidden;
-		}
-		else
-		{
-			return EVisibility::Collapsed;
-		}
-	}
-
-	FReply ResetToBaseButtonPressed() const
-	{
-		FunctionInput->ResetToBase();
-		return FReply::Handled();
-	}
-
-	FText GetInputIconText() const
-	{
-		switch (FunctionInput->GetValueMode())
-		{
-		case UNiagaraStackFunctionInput::EValueMode::Local:
-			return FText::FromString(FString(TEXT("\xf040") /* fa-pencil */));
-		case UNiagaraStackFunctionInput::EValueMode::Linked:
-			return FText::FromString(FString(TEXT("\xf0C1") /* fa-link */));
-		case UNiagaraStackFunctionInput::EValueMode::Data:
-			return FText::FromString(FString(TEXT("\xf1C0") /* fa-database */));
-		case UNiagaraStackFunctionInput::EValueMode::Dynamic:
-			return FText::FromString(FString(TEXT("\xf201") /* fa-line-chart */));
-		case UNiagaraStackFunctionInput::EValueMode::Invalid:
-		default:
-			return FText::FromString(FString(TEXT("\xf128") /* fa-question */));
-		}
-	}
-
-	FText GetInputIconToolTip() const
-	{
-		switch (FunctionInput->GetValueMode())
-		{
-		case UNiagaraStackFunctionInput::EValueMode::Local:
-			return LOCTEXT("StructInputIconToolTip", "Local Value");
-		case UNiagaraStackFunctionInput::EValueMode::Linked:
-			return LOCTEXT("LinkInputIconToolTip", "Linked Value");
-		case UNiagaraStackFunctionInput::EValueMode::Data:
-			return LOCTEXT("DataInterfaceInputIconToolTip", "Data Value");
-		case UNiagaraStackFunctionInput::EValueMode::Dynamic:
-			return LOCTEXT("DynamicInputIconToolTip", "Dynamic Value");
-		case UNiagaraStackFunctionInput::EValueMode::Invalid:
-		default:
-			return LOCTEXT("InvalidInputIconToolTip", "Unsupported value type.  Check the graph for issues.");
-		}
-	}
-
-	FSlateColor GetInputIconColor() const
-	{
-		switch (FunctionInput->GetValueMode())
-		{
-		case UNiagaraStackFunctionInput::EValueMode::Local:
-			return FLinearColor(FColor::Orange);
-		case UNiagaraStackFunctionInput::EValueMode::Linked:
-			return FLinearColor(FColor::Purple);
-		case UNiagaraStackFunctionInput::EValueMode::Data:
-			return FLinearColor(FColor::Yellow);
-		case UNiagaraStackFunctionInput::EValueMode::Dynamic:
-			return FLinearColor(FColor::Cyan);
-		case UNiagaraStackFunctionInput::EValueMode::Invalid:
-		default:
-			return FLinearColor(FColor::White);
-		}
-	}
-
-private:
-	UNiagaraStackFunctionInput* FunctionInput;
-
-	TSharedPtr<FStructOnScope> DisplayedLocalValueStruct;
-
-	TSharedPtr<SBox> LocalValueStructContainer;
-	TSharedPtr<SNiagaraParameterEditor> LocalValueStructParameterEditor;
-	TSharedPtr<IStructureDetailsView> LocalValueStructDetailsView;
-};
-
-class SNiagaraStackParameterStoreEntryValue : public SCompoundWidget
-{
-public:
-	DECLARE_DELEGATE_OneParam(FOnColumnWidthChanged, float)
-public:
-	SLATE_BEGIN_ARGS(SNiagaraStackParameterStoreEntryValue) { }
-	SLATE_END_ARGS();
-
-	void Construct(const FArguments& InArgs, UNiagaraStackParameterStoreEntry* InStackEntry)
-	{
-		StackEntry = InStackEntry;
-
-		StackEntry->OnValueChanged().AddSP(this, &SNiagaraStackParameterStoreEntryValue::OnInputValueChanged);
-		DisplayedValueStruct = StackEntry->GetValueStruct();
-
-		FMargin ItemPadding = FMargin(0);
-		ChildSlot
+			//~ Search, view options
+			+ SVerticalBox::Slot()
+			.AutoHeight()
+			.HAlign(HAlign_Fill)
+			.Padding(2, 4, 2, 4)
 			[
-				// Values
+				// Search box
 				SNew(SHorizontalBox)
 				+ SHorizontalBox::Slot()
-				.VAlign(VAlign_Center)
-				.AutoWidth()
-				.Padding(0, 0, 3, 0)
 				[
-					// Value Icon
-					SNew(SBox)
-					.WidthOverride(TextIconSize)
-					[
-						SNew(STextBlock)
-						.Font(FEditorStyle::Get().GetFontStyle("FontAwesome.10"))
-						.Text(this, &SNiagaraStackParameterStoreEntryValue::GetInputIconText)
-						.ToolTipText(this, &SNiagaraStackParameterStoreEntryValue::GetInputIconToolTip)
-						.ColorAndOpacity(this, &SNiagaraStackParameterStoreEntryValue::GetInputIconColor)
-					]
-				]
-				+ SHorizontalBox::Slot()
-				.VAlign(VAlign_Center)
-				[
-					// Assigned handle
-					SNew(SVerticalBox)
-					// Value struct
-					+ SVerticalBox::Slot()
-					.AutoHeight()
-					[
-						SAssignNew(ValueStructContainer, SBox)
-						[
-							ConstructValueStructWidget()
-						]
-					]
+					SAssignNew(SearchBox, SSearchBox)
+					.HintText(LOCTEXT("StackSearchBoxHint", "Search the stack"))
+					.SearchResultData(this, &SNiagaraStack::GetSearchResultData)
+					.IsSearching(this, &SNiagaraStack::GetIsSearching)
+					.OnTextChanged(this, &SNiagaraStack::OnSearchTextChanged)
+					.DelayChangeNotificationsWhileTyping(true)
+					.OnTextCommitted(this, &SNiagaraStack::OnSearchBoxTextCommitted)
+					.OnSearch(this, &SNiagaraStack::OnSearchBoxSearch)
 				]
-
-				// Handle drop-down button
+				// View options
 				+ SHorizontalBox::Slot()
 				.AutoWidth()
-				.VAlign(VAlign_Center)
-				.Padding(3, 0, 0, 0)
+				.Padding(4, 0, 0, 0)
 				[
 					SNew(SComboButton)
 					.ButtonStyle(FEditorStyle::Get(), "HoverHintOnly")
 					.ForegroundColor(FSlateColor::UseForeground())
-					.OnGetMenuContent(this, &SNiagaraStackParameterStoreEntryValue::OnGetAvailableHandleMenu)
-					.ContentPadding(FMargin(2))
-					.HAlign(HAlign_Center)
-					.VAlign(VAlign_Center)
-				]
-
-				// Reset Button
-				+ SHorizontalBox::Slot()
-				.AutoWidth()
-				.VAlign(VAlign_Center)
-				.Padding(3, 0, 0, 0)
-				[
-					SNew(SButton)
-					.IsFocusable(false)
-					.ToolTipText(LOCTEXT("ResetToolTip", "Reset to the default value"))
-					.ButtonStyle(FEditorStyle::Get(), "NoBorder")
+					.ToolTipText(LOCTEXT("ViewOptionsToolTip", "View Options"))
+					.OnGetMenuContent(this, &SNiagaraStack::GetViewOptionsMenu)
 					.ContentPadding(0)
-					.Visibility(this, &SNiagaraStackParameterStoreEntryValue::GetResetButtonVisibility)
-					.OnClicked(this, &SNiagaraStackParameterStoreEntryValue::ResetButtonPressed)
-					.Content()
-					[
-						SNew(SImage)
-						.Image(FEditorStyle::GetBrush("PropertyWindow.DiffersFromDefault"))
-					]
-				]
-			];
-	}
-
-private:
-	TSharedRef<SWidget> OnGetAvailableHandleMenu()
-	{
-		// TODO: This will need to be adjusted based on the current stack being edited, i.e. system vs emitter vs particle.
-		FMenuBuilder MenuBuilder(true, nullptr);
-		return MenuBuilder.MakeWidget();
-	}
-
-	TSharedRef<SWidget> ConstructValueStructWidget()
-	{
-		ValueStructParameterEditor.Reset();
-		ValueStructDetailsView.Reset();
-		if (DisplayedValueStruct.IsValid())
-		{
-			FNiagaraEditorModule& NiagaraEditorModule = FModuleManager::GetModuleChecked<FNiagaraEditorModule>("NiagaraEditor");
-			FNiagaraTypeDefinition ParameterType = FNiagaraTypeDefinition((UScriptStruct*)DisplayedValueStruct->GetStruct());
-			TSharedPtr<INiagaraEditorTypeUtilities> TypeEditorUtilities = NiagaraEditorModule.GetTypeUtilities(ParameterType);
-			if (TypeEditorUtilities.IsValid() && TypeEditorUtilities->CanCreateParameterEditor())
-			{
-				TSharedPtr<SNiagaraParameterEditor> ParameterEditor = TypeEditorUtilities->CreateParameterEditor(ParameterType);
-				ParameterEditor->UpdateInternalValueFromStruct(DisplayedValueStruct.ToSharedRef());
-				ParameterEditor->SetOnBeginValueChange(SNiagaraParameterEditor::FOnValueChange::CreateSP(
-					this, &SNiagaraStackParameterStoreEntryValue::ParameterBeginValueChange));
-				ParameterEditor->SetOnEndValueChange(SNiagaraParameterEditor::FOnValueChange::CreateSP(
-					this, &SNiagaraStackParameterStoreEntryValue::ParameterEndValueChange));
-				ParameterEditor->SetOnValueChanged(SNiagaraParameterEditor::FOnValueChange::CreateSP(
-					this, &SNiagaraStackParameterStoreEntryValue::ParameterValueChanged, ParameterEditor.ToSharedRef()));
-
-				ValueStructParameterEditor = ParameterEditor;
-				return ParameterEditor.ToSharedRef();
-			}
-			else
-			{
-				FPropertyEditorModule& PropertyEditorModule = FModuleManager::GetModuleChecked<FPropertyEditorModule>("PropertyEditor");
-
-				TSharedRef<IStructureDetailsView> StructureDetailsView = PropertyEditorModule.CreateStructureDetailView(
-					FDetailsViewArgs(false, false, false, FDetailsViewArgs::HideNameArea, true),
-					FStructureDetailsViewArgs(),
-					nullptr);
-
-				StructureDetailsView->SetStructureData(DisplayedValueStruct);
-				StructureDetailsView->GetOnFinishedChangingPropertiesDelegate().AddSP(this, &SNiagaraStackParameterStoreEntryValue::ParameterPropertyValueChanged);
-
-				ValueStructDetailsView = StructureDetailsView;
-				return StructureDetailsView->GetWidget().ToSharedRef();
-			}
-		}
-		else
-		{
-			return SNullWidget::NullWidget;
-		}
-	}
-
-	void OnInputValueChanged()
-	{
-		TSharedPtr<FStructOnScope> NewValueStruct = StackEntry->GetValueStruct();
-		if (DisplayedValueStruct == NewValueStruct)
-		{
-			if (ValueStructParameterEditor.IsValid())
-			{
-				ValueStructParameterEditor->UpdateInternalValueFromStruct(DisplayedValueStruct.ToSharedRef());
-			}
-			if (ValueStructDetailsView.IsValid())
-			{
-				ValueStructDetailsView->SetStructureData(TSharedPtr<FStructOnScope>());
-				ValueStructDetailsView->SetStructureData(DisplayedValueStruct);
-			}
-		}
-		else
-		{
-			DisplayedValueStruct = NewValueStruct;
-			ValueStructContainer->SetContent(ConstructValueStructWidget());
-		}
-	}
-
-	void ParameterBeginValueChange()
-	{
-		StackEntry->NotifyBeginValueChange();
-	}
-
-	void ParameterEndValueChange()
-	{
-		StackEntry->NotifyEndValueChange();
-	}
-
-	void ParameterValueChanged(TSharedRef<SNiagaraParameterEditor> ParameterEditor)
-	{
-		ParameterEditor->UpdateStructFromInternalValue(StackEntry->GetValueStruct().ToSharedRef());
-		StackEntry->NotifyValueChanged();
-	}
-
-	void ParameterPropertyValueChanged(const FPropertyChangedEvent& PropertyChangedEvent)
-	{
-		StackEntry->NotifyValueChanged();
-	}
-
-	EVisibility GetReferenceVisibility() const
-	{
-		return EVisibility::Collapsed;
-	}
-
-
-	EVisibility GetResetButtonVisibility() const
-	{
-		return StackEntry->CanReset() ? EVisibility::Visible : EVisibility::Hidden;
-	}
-
-	FReply ResetButtonPressed() const
-	{
-		StackEntry->Reset();
-		return FReply::Handled();
-	}
-
-	FText GetInputIconText() const
-	{
-		if (DisplayedValueStruct.IsValid())
-		{
-			return FText::FromString(FString(TEXT("\xf040") /* fa-pencil */));
-		}
-		else if (StackEntry->GetValueObject() != nullptr)
-		{
-			return FText::FromString(FString(TEXT("\xf1C0") /* fa-database */));
-		}
-		else
-		{
-			return FText();
-		}
-	}
-
-	FText GetInputIconToolTip() const
-	{
-		if (DisplayedValueStruct.IsValid())
-		{
-			return LOCTEXT("StructInputIconToolTip", "Local Value");
-		}
-		else if (StackEntry->GetValueObject() != nullptr)
-		{
-			return LOCTEXT("DataInterfaceInputIconToolTip", "Data Value");
-		}
-		else
-		{
-			return FText();
-		}
-	}
-
-	FSlateColor GetInputIconColor() const
-	{
-		if (DisplayedValueStruct.IsValid())
-		{
-			return FLinearColor(FColor::Orange);
-		}
-		else if (StackEntry->GetValueObject() != nullptr)
-		{
-			return FLinearColor(FColor::Yellow);
-		}
-		else
-		{
-			return FLinearColor(FColor::White);
-		}
-	}
-
-private:
-	UNiagaraStackParameterStoreEntry* StackEntry;
-
-	TSharedPtr<FStructOnScope> DisplayedValueStruct;
-
-	TSharedPtr<SBox> ValueStructContainer;
-	TSharedPtr<SNiagaraParameterEditor> ValueStructParameterEditor;
-	TSharedPtr<IStructureDetailsView> ValueStructDetailsView;
-};
-
-class SNiagaraStackObject : public SCompoundWidget, public FNotifyHook
-{
-public:
-	DECLARE_DELEGATE_OneParam(FOnColumnWidthChanged, float)
-public:
-	SLATE_BEGIN_ARGS(SNiagaraStackObject) { }
-	SLATE_END_ARGS();
-
-	void Construct(const FArguments& InArgs, UNiagaraStackObject& InObject)
-	{
-		Object = &InObject;
-
-		FPropertyEditorModule& PropertyEditorModule = FModuleManager::GetModuleChecked<FPropertyEditorModule>("PropertyEditor");
-		TSharedRef<IDetailsView> DetailsView = PropertyEditorModule.CreateDetailView(FDetailsViewArgs(false, false, false, FDetailsViewArgs::HideNameArea, true, this));
-		DetailsView->SetObject(Object->GetObject(), true);
-
-		ChildSlot
-		[
-			DetailsView
-		];
-	}
-
-	virtual void NotifyPostChange(const FPropertyChangedEvent& PropertyChangedEvent, UProperty* PropertyThatChanged) override
-	{
-		Object->NotifyObjectChanged();
-	}
-
-private:
-	UNiagaraStackObject* Object;
-};
-
-class SNiagaraStackStruct : public SCompoundWidget, FNotifyHook
-{
-public:
-	DECLARE_DELEGATE_OneParam(FOnColumnWidthChanged, float)
-public:
-	SLATE_BEGIN_ARGS(SNiagaraStackStruct) { }
-	SLATE_END_ARGS();
-	
-
-	void Construct(const FArguments& InArgs, UNiagaraStackStruct& InObject)
-	{
-		Object = &InObject;
-		
-		FStructureDetailsViewArgs StructureDetailViewArgs;
-		FDetailsViewArgs DetailViewArgs(false, false, false, FDetailsViewArgs::HideNameArea, true);
-		DetailViewArgs.NotifyHook = this;
-		FPropertyEditorModule& PropertyEditorModule = FModuleManager::GetModuleChecked<FPropertyEditorModule>("PropertyEditor");
-
-		TSharedRef<IStructureDetailsView> StructureDetailsView = PropertyEditorModule.CreateStructureDetailView(
-			DetailViewArgs, StructureDetailViewArgs, nullptr);
-		if (InObject.HasDetailCustomization())
-		{
-			StructureDetailsView->GetDetailsView()->RegisterInstancedCustomPropertyLayout((UStruct*)InObject.GetStructOnScope()->GetStruct(), InObject.GetDetailsCustomization());
-		}
-
-		StructureDetailsView->SetStructureData(InObject.GetStructOnScope());
-	
-		ChildSlot
-			[
-				StructureDetailsView->GetWidget().ToSharedRef()
-			];
-	}
-
-	virtual void NotifyPreChange(class FEditPropertyChain* PropertyAboutToChange) override
-	{
-		Object->GetOwningObject()->Modify();
-	}
-
-	virtual void NotifyPostChange(const FPropertyChangedEvent& PropertyChangedEvent, class FEditPropertyChain* PropertyThatChanged) override
-	{
-		Object->GetOwningObject()->PostEditChange();
-	}
-
-private:
-	UNiagaraStackStruct* Object;
-};
-
-class SNiagaraStackAddModuleItem : public SCompoundWidget
-{
-public:
-	SLATE_BEGIN_ARGS(SNiagaraStackAddModuleItem) { }
-	SLATE_END_ARGS();
-
-	void Construct(const FArguments& InArgs, UNiagaraStackAddModuleItem& InAddModuleItem)
-	{
-		AddModuleItem = &InAddModuleItem;
-
-		FSlateColor TextColor;
-		FSlateColor BackgroundColor;
-		float Padding;
-		FName TextStyleName;
-		FName FontStyleName;
-		FText ToolTipText;
-		if (AddModuleItem->GetDisplayMode() == UNiagaraStackAddModuleItem::EDisplayMode::Standard)
-		{
-			TextColor = FSlateColor::UseForeground();
-			BackgroundColor = FSlateColor(FLinearColor::Transparent);
-			Padding = 3;
-			TextStyleName = "NormalText.Important";
-			FontStyleName = "FontAwesome.10";
-			ToolTipText = LOCTEXT("AddModuleToolTip", "Add new module");
-		}
-		else // Compact
-		{
-			TextColor = FSlateColor(FLinearColor::Black);
-			BackgroundColor = FSlateColor(FLinearColor::White);
-			Padding = 0;
-			TextStyleName = "NormalText";
-			FontStyleName = "FontAwesome.7";
-			ToolTipText = LOCTEXT("InsertModuleToolTip", "Insert new module");
-		}
-
-		ChildSlot
-		[
-			SAssignNew(AddModuleButton, SComboButton)
-			.ButtonStyle(FEditorStyle::Get(), "HoverHintOnly")
-			.ToolTipText(ToolTipText)
-			.ForegroundColor(TextColor)
-			.HasDownArrow(false)
-			.OnGetMenuContent(this, &SNiagaraStackAddModuleItem::GetAddModuleMenu)
-			.ContentPadding(Padding)
-			.ButtonContent()
-			[
-				SNew(SBorder)
-				.BorderImage(FEditorStyle::GetBrush("WhiteBrush"))
-				.BorderBackgroundColor(BackgroundColor)
-				.Visibility(this, &SNiagaraStackAddModuleItem::GetButtonContentVisibility)
-				.HAlign(HAlign_Center)
-				.VAlign(VAlign_Center)
-				[
-					SNew(STextBlock)
-					.TextStyle(FEditorStyle::Get(), TextStyleName)
-					.Font(FEditorStyle::Get().GetFontStyle(FontStyleName))
-					.Text(FText::FromString(FString(TEXT("\xf067"))) /*fa-plus*/)
-				]
-			]
-		];
-	}
-
-private:
-	EVisibility GetButtonContentVisibility() const 
-	{
-		if (AddModuleItem->GetDisplayMode() == UNiagaraStackAddModuleItem::EDisplayMode::Standard)
-		{
-			return EVisibility::Visible;
-		}
-		else
-		{
-			return IsHovered() ? EVisibility::Visible : EVisibility::Hidden;
-		}
-	}
-
-	TSharedRef<SWidget> GetAddModuleMenu()
-	{
-		FGraphActionMenuBuilder MenuBuilder;
-		TSharedPtr<SGraphActionMenu> AddModuleMenu;
-
-		TSharedRef<SBorder> MenuWidget =  SNew(SBorder)
-			.BorderImage(FEditorStyle::GetBrush("Menu.Background"))
-			.Padding(5)
-			[
-				SNew(SBox)
-				[
-					SAssignNew(AddModuleMenu, SGraphActionMenu)
-					.OnActionSelected(this, &SNiagaraStackAddModuleItem::OnActionSelected)
-					.OnCollectAllActions(this, &SNiagaraStackAddModuleItem::CollectAllActions)
-					.AutoExpandActionMenu(false)
-					.ShowFilterTextBox(true)
-				]
-			];
-		AddModuleButton->SetMenuContentWidgetToFocus(AddModuleMenu->GetFilterTextBox()->AsShared()); 
-		return MenuWidget;
-	}
-
-
-	void CollectAllActions(FGraphActionListBuilderBase& OutAllActions)
-	{
-		if (OutAllActions.OwnerOfTemporaries == nullptr)
-		{
-			OutAllActions.OwnerOfTemporaries = NewObject<UEdGraph>((UObject*)GetTransientPackage());
-		}
-		FAssetRegistryModule& AssetRegistryModule = FModuleManager::LoadModuleChecked<FAssetRegistryModule>(TEXT("AssetRegistry"));
-
-		TArray<FAssetData> ScriptAssets;
-		AssetRegistryModule.Get().GetAssetsByClass(UNiagaraScript::StaticClass()->GetFName(), ScriptAssets);
-		UEnum* NiagaraScriptUsageEnum = FindObjectChecked<UEnum>(ANY_PACKAGE, TEXT("ENiagaraScriptUsage"), true);
-		for (const FAssetData& ScriptAsset : ScriptAssets)
-		{
-			FName UsageName;
-			ScriptAsset.GetTagValue(GET_MEMBER_NAME_CHECKED(UNiagaraScript, Usage), UsageName);
-
-			FText AssetDesc;
-			ScriptAsset.GetTagValue(GET_MEMBER_NAME_CHECKED(UNiagaraScript, Description), AssetDesc);
-
-			FText ModuleCategory;
-			ScriptAsset.GetTagValue(GET_MEMBER_NAME_CHECKED(UNiagaraScript, Category), ModuleCategory);
-
-			FString QualifiedUsageName = "ENiagaraScriptUsage::" + UsageName.ToString();
-			int32 UsageIndex = NiagaraScriptUsageEnum->GetIndexByNameString(QualifiedUsageName);
-
-			const FString BitfieldTagValue = ScriptAsset.GetTagValueRef<FString>(GET_MEMBER_NAME_CHECKED(UNiagaraScript, ModuleUsageBitmask));
-			int32 BitfieldValue = FCString::Atoi(*BitfieldTagValue);
-
-			int32 TargetBit = (BitfieldValue >> (int32)AddModuleItem->GetOutputUsage()) & 1;
-
-			if (UsageIndex != INDEX_NONE && TargetBit == 1)
-			{
-				ENiagaraScriptUsage Usage = static_cast<ENiagaraScriptUsage>(NiagaraScriptUsageEnum->GetValueByIndex(UsageIndex));
-				if (Usage == ENiagaraScriptUsage::Module)
-				{
-					FString DisplayNameString = FName::NameToDisplayString(ScriptAsset.AssetName.ToString(), false);
-					const FText NameText = FText::FromString(DisplayNameString);
-					const FText TooltipDesc = FText::Format(LOCTEXT("FunctionPopupTooltip", "Path: {0}\nDescription: {1}"), FText::FromString(ScriptAsset.ObjectPath.ToString()), AssetDesc);
-					FText CategoryName = LOCTEXT("ModuleNotCategorized", "Uncategorized Modules");
-					if (!ModuleCategory.IsEmptyOrWhitespace())
-						CategoryName = ModuleCategory;
-
-					TSharedPtr<FNiagaraStackGraphSchemaAction> NewNodeAction(new FNiagaraStackGraphSchemaAction(CategoryName, NameText, TooltipDesc, 0, FText(), 
-						FNiagaraStackGraphSchemaAction::FOnPerformStackAction::CreateUObject(AddModuleItem, &UNiagaraStackAddModuleItem::AddScriptModule, ScriptAsset)));
-					OutAllActions.AddAction(NewNodeAction);
-				}
-			}
-		}
-
-		// Generate actions for the available parameters to set.
-		TArray<FNiagaraVariable> AvailableParameters;
-		AddModuleItem->GetAvailableParameters(AvailableParameters);
-		for(const FNiagaraVariable& AvailableParameter : AvailableParameters)
-		{
-			FString DisplayNameString = FName::NameToDisplayString(AvailableParameter.GetName().ToString(), false);
-			const FText NameText = FText::FromString(DisplayNameString);
-			FText VarDesc = FNiagaraConstants::GetAttributeDescription(AvailableParameter);
-			FString VarDefaultValue = FNiagaraConstants::GetAttributeDefaultValue(AvailableParameter);
-			const FText TooltipDesc = FText::Format(LOCTEXT("SetFunctionPopupTooltip", "Description: Set the parameter {0}. {1}"), FText::FromName(AvailableParameter.GetName()), VarDesc);
-			FText CategoryName = LOCTEXT("ModuleSetCategory", "Set Specific Parameters");
-
-			TSharedPtr<FNiagaraStackGraphSchemaAction> NewNodeAction(new FNiagaraStackGraphSchemaAction(CategoryName, NameText, TooltipDesc, 0, FText(),
-				FNiagaraStackGraphSchemaAction::FOnPerformStackAction::CreateUObject(AddModuleItem, &UNiagaraStackAddModuleItem::AddParameterModule, AvailableParameter, false)));
-			OutAllActions.AddAction(NewNodeAction);
-		}
-
-		// Generate actions for setting new typed parameters.
-		TOptional<FName> NewParameterNamespace = AddModuleItem->GetNewParameterNamespace();
-		if (NewParameterNamespace.IsSet())
-		{
-			TArray<FNiagaraTypeDefinition> AvailableTypes;
-			AddModuleItem->GetNewParameterAvailableTypes(AvailableTypes);
-			for (const FNiagaraTypeDefinition& AvailableType : AvailableTypes)
-			{
-				FText NameText = AvailableType.GetNameText();
-				FText Tooltip = FText::Format(LOCTEXT("AddNewParameterTooltipFormat", "Description: Create a new {0} parameter."), NameText);
-				FText CategoryName = LOCTEXT("CreateNewParameterCategory", "Create New Parameter");
-
-				FNiagaraParameterHandle NewParameterHandle(NewParameterNamespace.GetValue(), *(TEXT("New") + AvailableType.GetName()));
-				FNiagaraVariable NewParameter(AvailableType, NewParameterHandle.GetParameterHandleString());
-				TSharedPtr<FNiagaraStackGraphSchemaAction> NewNodeAction(new FNiagaraStackGraphSchemaAction(CategoryName, NameText, Tooltip, 0, FText(),
-					FNiagaraStackGraphSchemaAction::FOnPerformStackAction::CreateUObject(AddModuleItem, &UNiagaraStackAddModuleItem::AddParameterModule, NewParameter, true)));
-				OutAllActions.AddAction(NewNodeAction);
-			}
-		}
-
-		const ENiagaraScriptUsage AddModuleOutputUsage = AddModuleItem->GetOutputUsage();
-		if (AddModuleOutputUsage != ENiagaraScriptUsage::EmitterSpawnScript && AddModuleOutputUsage != ENiagaraScriptUsage::EmitterUpdateScript)
-		{
-			// Generate actions for material
-			FText NameText = LOCTEXT("MaterialParameterCategory", "Create Dynamic Parameters");
-			FText Tooltip = LOCTEXT("MaterialParameterTooltip", "Description: Create dynamic material parameters.");
-			FText CategoryName = LOCTEXT("MaterialParameter", "Material");
-
-			TSharedPtr<FNiagaraStackGraphSchemaAction> NewNodeAction(new FNiagaraStackGraphSchemaAction(CategoryName, NameText, Tooltip, 0, FText(),
-				FNiagaraStackGraphSchemaAction::FOnPerformStackAction::CreateUObject(AddModuleItem, &UNiagaraStackAddModuleItem::AddMaterialParameterModule)));
-			OutAllActions.AddAction(NewNodeAction);
-		}
-	}
-
-	void OnActionSelected(const TArray< TSharedPtr<FEdGraphSchemaAction> >& SelectedActions, ESelectInfo::Type InSelectionType)
-	{
-		if (InSelectionType == ESelectInfo::OnMouseClick || InSelectionType == ESelectInfo::OnKeyPress || SelectedActions.Num() == 0)
-		{
-			for (int32 ActionIndex = 0; ActionIndex < SelectedActions.Num(); ActionIndex++)
-			{
-				TSharedPtr<FEdGraphSchemaAction> CurrentAction = SelectedActions[ActionIndex];
-
-				if (CurrentAction.IsValid())
-				{
-					FSlateApplication::Get().DismissAllMenus();
-
-					TArray<UEdGraphPin*> Pins;
-					CurrentAction->PerformAction(nullptr, Pins, FVector2D(0.0f, 0.0f));
-				}
-			}
-		}
-	}
-
-private:
-	UNiagaraStackAddModuleItem* AddModuleItem;
-	TSharedPtr<SComboButton> AddModuleButton;
-};
-
-class SNiagaraStackAddRendererItem : public SCompoundWidget
-{
-public:
-	SLATE_BEGIN_ARGS(SNiagaraStackAddRendererItem) { }
-	SLATE_END_ARGS();
-
-	void Construct(const FArguments& InArgs, UNiagaraStackAddRendererItem& InAddRendererItem)
-	{
-		AddRendererItem = &InAddRendererItem;
-		ChildSlot
-		[
-			SNew(SComboButton)
-			.ButtonStyle(FEditorStyle::Get(), "HoverHintOnly")
-			.ForegroundColor(FSlateColor::UseForeground())
-			.HasDownArrow(false)
-			.OnGetMenuContent(this, &SNiagaraStackAddRendererItem::GetAddRendererMenu)
-			.ContentPadding(3)
-			.HAlign(HAlign_Center)
-			.VAlign(VAlign_Center)
-			.ButtonContent()
-			[
-				SNew(STextBlock)
-				.TextStyle(FEditorStyle::Get(), "NormalText.Important")
-				.Font(FEditorStyle::Get().GetFontStyle("FontAwesome.10"))
-				.Text(FText::FromString(FString(TEXT("\xf067"))) /*fa-plus*/)
-				.ToolTipText(LOCTEXT("AddRendererToolTip", "Add new renderer"))
-			]
-		];
-	}
-
-private:
-	class FRendererClassFilter : public IClassViewerFilter
-	{
-	public:
-		virtual bool IsClassAllowed(const FClassViewerInitializationOptions& InInitOptions, const UClass* InClass, TSharedRef< class FClassViewerFilterFuncs > InFilterFuncs)
-		{
-			return InClass->HasAnyClassFlags(CLASS_Abstract) == false && InClass->IsChildOf(UNiagaraRendererProperties::StaticClass());
-		}
-
-		virtual bool IsUnloadedClassAllowed(const FClassViewerInitializationOptions& InInitOptions, const TSharedRef< const class IUnloadedBlueprintData > InUnloadedClassData, TSharedRef< class FClassViewerFilterFuncs > InFilterFuncs)
-		{
-			return false;
-		}
-	};
-
-	TSharedRef<SWidget> GetAddRendererMenu()
-	{
-		FClassViewerInitializationOptions Options;
-		Options.bShowDisplayNames = true;
-		Options.ClassFilter = MakeShareable(new FRendererClassFilter());
-
-		FOnClassPicked OnPicked(FOnClassPicked::CreateSP(this, &SNiagaraStackAddRendererItem::RendererClassPicked));
-		return FModuleManager::LoadModuleChecked<FClassViewerModule>("ClassViewer").CreateClassViewer(Options, OnPicked);
-	}
-
-	void RendererClassPicked(UClass* InPickedClass)
-	{
-		AddRendererItem->AddRenderer(InPickedClass);
-	}
-
-private:
-	UNiagaraStackAddRendererItem* AddRendererItem;
-};
-
-class SNiagaraStackModuleItem : public SCompoundWidget
-{
-public:
-	SLATE_BEGIN_ARGS(SNiagaraStackModuleItem) { }
-	SLATE_END_ARGS();
-
-	void Construct(const FArguments& InArgs, UNiagaraStackModuleItem& InModuleItem)
-	{
-		ModuleItem = &InModuleItem;
-
-		ChildSlot
-		[
-			SNew(SHorizontalBox)
-			+ SHorizontalBox::Slot()
-			.Padding(0, 0, 2, 0)
-			.AutoWidth()
-			[
-				SNew(SCheckBox)
-				.IsChecked(this, &SNiagaraStackModuleItem::GetCheckState)
-				.OnCheckStateChanged(this, &SNiagaraStackModuleItem::OnCheckStateChanged)
-			]
-			+ SHorizontalBox::Slot()
-			.Padding(1)
-			[
-				SNew(STextBlock)
-				.TextStyle(FNiagaraEditorWidgetsStyle::Get(), "NiagaraEditor.Stack.ItemText")
-				.ToolTipText_UObject(ModuleItem, &UNiagaraStackEntry::GetTooltipText)
-				.Text_UObject(ModuleItem, &UNiagaraStackEntry::GetDisplayName)
-			]
-			+ SHorizontalBox::Slot()
-			.AutoWidth()
-			[
-				SNew(SButton)
-				.ButtonStyle(FEditorStyle::Get(), "HoverHintOnly")
-				.IsFocusable(false)
-				.ForegroundColor(FNiagaraEditorWidgetsStyle::Get().GetColor(ModuleItem->GetItemForegroundName()))
-				.ToolTipText(LOCTEXT("UpToolTip", "Move this module up the stack"))
-				.Visibility(this, &SNiagaraStackModuleItem::GetEditButtonVisibility)
-				.OnClicked(this, &SNiagaraStackModuleItem::MoveUpClicked)
-				.Content()
-				[
-					SNew(STextBlock)
-					.Font(FEditorStyle::Get().GetFontStyle("FontAwesome.10"))
-					.Text(FText::FromString(FString(TEXT("\xf062"))))
-				]
-			]
-			+ SHorizontalBox::Slot()
-			.AutoWidth()
-			[
-				SNew(SButton)
-				.ButtonStyle(FEditorStyle::Get(), "HoverHintOnly")
-				.IsFocusable(false)
-				.ForegroundColor(FNiagaraEditorWidgetsStyle::Get().GetColor(ModuleItem->GetItemForegroundName()))
-				.ToolTipText(LOCTEXT("DownToolTip", "Move this module down the stack"))
-				.Visibility(this, &SNiagaraStackModuleItem::GetEditButtonVisibility)
-				.OnClicked(this, &SNiagaraStackModuleItem::MoveDownClicked)
-				.Content()
-				[
-					SNew(STextBlock)
-					.Font(FEditorStyle::Get().GetFontStyle("FontAwesome.10"))
-					.Text(FText::FromString(FString(TEXT("\xf063"))))
-				]
-			]
-			+ SHorizontalBox::Slot()
-			.AutoWidth()
-			[
-				SNew(SButton)
-				.ButtonStyle(FEditorStyle::Get(), "HoverHintOnly")
-				.IsFocusable(false)
-				.ForegroundColor(FNiagaraEditorWidgetsStyle::Get().GetColor(ModuleItem->GetItemForegroundName()))
-				.ToolTipText(LOCTEXT("RefreshTooltip", "Refresh this module"))
-				.Visibility(this, &SNiagaraStackModuleItem::GetRefreshVisibility)
-				.OnClicked(this, &SNiagaraStackModuleItem::RefreshClicked)
-				.Content()
-				[
-					SNew(STextBlock)
-					.Font(FEditorStyle::Get().GetFontStyle("FontAwesome.10"))
-					.Text(FText::FromString(FString(TEXT("\xf021"))))
-				]
-			]
-			+ SHorizontalBox::Slot()
-			.AutoWidth()
-			[
-				SNew(SButton)
-				.ButtonStyle(FEditorStyle::Get(), "HoverHintOnly")
-				.IsFocusable(false)
-				.ForegroundColor(FNiagaraEditorWidgetsStyle::Get().GetColor(ModuleItem->GetItemForegroundName()))
-				.ToolTipText(LOCTEXT("DeleteToolTip", "Delete this module"))
-				.Visibility(this, &SNiagaraStackModuleItem::GetEditButtonVisibility)
-				.OnClicked(this, &SNiagaraStackModuleItem::DeleteClicked)
-				.Content()
-				[
-					SNew(STextBlock)
-					.Font(FEditorStyle::Get().GetFontStyle("FontAwesome.10"))
-					.Text(FText::FromString(FString(TEXT("\xf1f8"))))
-				]
-			]
-		];
-	}
-
-public:
-	//~ SWidget interface
-	virtual FReply OnMouseButtonDoubleClick(const FGeometry& InMyGeometry, const FPointerEvent& InMouseEvent) override
-	{
-		const UNiagaraNodeFunctionCall& ModuleFunctionCall = ModuleItem->GetModuleNode();
-		if (ModuleFunctionCall.FunctionScript != nullptr && ModuleFunctionCall.FunctionScript->IsAsset())
-		{
-			FAssetEditorManager::Get().OpenEditorForAsset(const_cast<UNiagaraScript*>(ModuleFunctionCall.FunctionScript));
-			return FReply::Handled();
-		}
-		return FReply::Unhandled();
-	}
-
-private:
-	ECheckBoxState GetCheckState() const
-	{
-		return ModuleItem->GetIsEnabled() ? ECheckBoxState::Checked : ECheckBoxState::Unchecked;
-	}
-
-	void OnCheckStateChanged(ECheckBoxState InCheckState)
-	{
-		ModuleItem->SetIsEnabled(InCheckState == ECheckBoxState::Checked);
-	}
-
-	EVisibility GetEditButtonVisibility() const
-	{
-		return ModuleItem->CanMoveAndDelete() ? EVisibility::Visible : EVisibility::Collapsed;
-	}
-
-	EVisibility GetRefreshVisibility() const
-	{
-		return ModuleItem->CanRefresh() ? EVisibility::Visible : EVisibility::Collapsed;
-	}
-
-	FReply MoveUpClicked()
-	{
-		ModuleItem->MoveUp();
-		return FReply::Handled();
-	}
-
-	FReply MoveDownClicked()
-	{
-		ModuleItem->MoveDown();
-		return FReply::Handled();
-	}
-
-	FReply DeleteClicked()
-	{
-		ModuleItem->Delete();
-		return FReply::Handled();
-	}
-
-	FReply RefreshClicked()
-	{
-		ModuleItem->Refresh();
-		return FReply::Handled();
-	}
-
-private:
-	UNiagaraStackModuleItem* ModuleItem;
-};
-
-class SNiagaraStackItemGroup : public SCompoundWidget
-{
-public:
-	SLATE_BEGIN_ARGS(SNiagaraStackItemGroup) { }
-	SLATE_END_ARGS();
-
-	void Construct(const FArguments& InArgs, UNiagaraStackItemGroup& InItem)
-	{
-		Item = &InItem;
-
-		ChildSlot
-			[
-				SNew(SHorizontalBox)
-				+ SHorizontalBox::Slot()
-				[
-					SNew(STextBlock)
-					.TextStyle(FNiagaraEditorWidgetsStyle::Get(), "NiagaraEditor.Stack.GroupText")
-					.ToolTipText_UObject(Item, &UNiagaraStackEntry::GetTooltipText)
-					.Text_UObject(Item, &UNiagaraStackEntry::GetDisplayName)
-				]
-				+ SHorizontalBox::Slot()
-				.AutoWidth()
-				[
-					SNew(SButton)
-					.Visibility(this, &SNiagaraStackItemGroup::AddVisibility)
-					.ButtonStyle(FEditorStyle::Get(), "HoverHintOnly")
-					.IsFocusable(false)
-					.ForegroundColor(FNiagaraEditorWidgetsStyle::Get().GetColor(Item->GetItemForegroundName()))
-					.ToolTipText(LOCTEXT("AddGroupToolTip", "Add a new group"))
-					.OnClicked(this, &SNiagaraStackItemGroup::AddClicked)
-					.Content()
+					.MenuPlacement(MenuPlacement_BelowRightAnchor)
+					.ButtonContent()
 					[
-						SNew(STextBlock)
-						.Font(FEditorStyle::Get().GetFontStyle("FontAwesome.10"))
-						.Text(FText::FromString(FString(TEXT("\xf067")))) /*fa-plus*/
-					]
-				]
-				+ SHorizontalBox::Slot()
-				.AutoWidth()
-				[
-					SNew(SButton)
-						.Visibility(this, &SNiagaraStackItemGroup::DeleteVisibility)
-						.ButtonStyle(FEditorStyle::Get(), "HoverHintOnly")
-						.IsFocusable(false)
-						.ForegroundColor(FNiagaraEditorWidgetsStyle::Get().GetColor(Item->GetItemForegroundName()))
-						.ToolTipText(LOCTEXT("DeleteGroupToolTip", "Delete this group"))
-						.OnClicked(this, &SNiagaraStackItemGroup::DeleteClicked)
-						.Content()
+						SNew(SBox)
+						.HAlign(HAlign_Center)
+						.VAlign(VAlign_Center)
 						[
-							SNew(STextBlock)
-							.Font(FEditorStyle::Get().GetFontStyle("FontAwesome.10"))
-							.Text(FText::FromString(FString(TEXT("\xf1f8"))))
+							SNew(SImage)
+							.Image(FEditorStyle::GetBrush("GenericViewButton"))
+							.ColorAndOpacity(FNiagaraEditorWidgetsStyle::Get().GetColor("NiagaraEditor.Stack.FlatButtonColor"))
 						]
+					]
 				]
-			];
-	}
-
-private:
-	FReply DeleteClicked()
-	{
-		Item->Delete();
-		return FReply::Handled();
-	}
-
-	EVisibility DeleteVisibility() const
-	{
-		if (Item->CanDelete())
-		{
-			return EVisibility::Visible;
-		}
-		else
-		{
-			return EVisibility::Collapsed;
-		}
-	}
-
-	FReply AddClicked()
-	{
-		Item->Add();
-		return FReply::Handled();
-	}
-
-	EVisibility AddVisibility() const
-	{
-		if (Item->CanAdd())
-		{
-			return EVisibility::Visible;
-		}
-		else
-		{
-			return EVisibility::Collapsed;
-		}
-	}
-
-private:
-	UNiagaraStackItemGroup* Item;
-};
-
-class SNiagaraStackRendererItem : public SCompoundWidget
-{
-public:
-	SLATE_BEGIN_ARGS(SNiagaraStackRendererItem) { }
-	SLATE_END_ARGS();
-
-	void Construct(const FArguments& InArgs, UNiagaraStackRendererItem& InRendererItem)
-	{
-		RendererItem = &InRendererItem;
-
-		ChildSlot
-		[
-			SNew(SHorizontalBox)
-			+ SHorizontalBox::Slot()
-			.Padding(1)
-			[
-				SNew(STextBlock)
-				.TextStyle(FNiagaraEditorWidgetsStyle::Get(), "NiagaraEditor.Stack.ItemText")
-				.ToolTipText_UObject(RendererItem, &UNiagaraStackEntry::GetTooltipText)
-				.Text_UObject(RendererItem, &UNiagaraStackEntry::GetDisplayName)
-			]
-			// Delete button
-			+ SHorizontalBox::Slot()
-			.AutoWidth()
-			[
-				SNew(SButton)
-				.ButtonStyle(FEditorStyle::Get(), "HoverHintOnly")
-				.IsFocusable(false)
-				.ForegroundColor(FNiagaraEditorWidgetsStyle::Get().GetColor(RendererItem->GetItemForegroundName()))
-				.ToolTipText(LOCTEXT("DeleteRendererToolTip", "Delete this Renderer"))
-				.Visibility(this, &SNiagaraStackRendererItem::GetDeleteButtonVisibility)
-				.OnClicked(this, &SNiagaraStackRendererItem::DeleteClicked)
-				.Content()
-				[
-					SNew(STextBlock)
-					.Font(FEditorStyle::Get().GetFontStyle("FontAwesome.10"))
-					.Text(FText::FromString(FString(TEXT("\xf1f8"))))
-				]
-			]
-			// Reset to base Button
-			+ SHorizontalBox::Slot()
-			.AutoWidth()
-			.VAlign(VAlign_Center)
-			.Padding(3, 0, 0, 0)
-			[
-				SNew(SButton)
-				.IsFocusable(false)
-				.ToolTipText(LOCTEXT("ResetRendererToBaseToolTip", "Reset this renderer to the state defined by the parent emitter"))
-				.ButtonStyle(FEditorStyle::Get(), "NoBorder")
-				.ContentPadding(0)
-				.Visibility(this, &SNiagaraStackRendererItem::GetResetToBaseButtonVisibility)
-				.OnClicked(this, &SNiagaraStackRendererItem::ResetToBaseButtonClicked)
-				.Content()
-				[
-					SNew(SImage)
-					.Image(FEditorStyle::GetBrush("PropertyWindow.DiffersFromDefault"))
-					.ColorAndOpacity(FSlateColor(FLinearColor::Green))
-				]
-			]
-		];
-	}
-
-private:
-	EVisibility GetDeleteButtonVisibility() const
-	{
-		return RendererItem->CanDelete() ? EVisibility::Visible : EVisibility::Collapsed;
-	}
-
-	FReply DeleteClicked()
-	{
-		RendererItem->Delete();
-		return FReply::Handled();
-	}
-
-	EVisibility GetResetToBaseButtonVisibility() const
-	{
-		if (RendererItem->CanHaveBase())
-		{
-			return RendererItem->CanResetToBase() ? EVisibility::Visible : EVisibility::Hidden;
-		}
-		else
-		{
-			return EVisibility::Collapsed;
-		}
-	}
-
-	FReply ResetToBaseButtonClicked()
-	{
-		RendererItem->ResetToBase();
-		return FReply::Handled();
+			];
 	}
+}
 
-private:
-	UNiagaraStackRendererItem* RendererItem;
-};
-
-class SNiagaraStackItemExpander : public SCompoundWidget
+FReply SNiagaraStack::PinButtonPressed()
 {
-public:
-	SLATE_BEGIN_ARGS(SNiagaraStackItemExpander) { }
-	SLATE_END_ARGS();
-
-	void Construct(const FArguments& InArgs, UNiagaraStackItemExpander& InItemExpander)
-	{
-		ItemExpander = &InItemExpander;
-		ExpandedToolTipText = LOCTEXT("ExpandedItemToolTip", "Collapse this item");
-		CollapsedToolTipText = LOCTEXT("CollapsedItemToolTip", "Expand this item");
-
-		ChildSlot
-		[
-			SNew(SButton)
-			.ButtonStyle(FEditorStyle::Get(), "NoBorder")
-			.HAlign(HAlign_Center)
-			.ContentPadding(2)
-			.ToolTipText(this, &SNiagaraStackItemExpander::GetToolTipText)
-			.OnClicked(this, &SNiagaraStackItemExpander::ExpandButtonClicked)
-			.Content()
-			[
-				SNew(SImage)
-				.Image(this, &SNiagaraStackItemExpander::GetButtonBrush)
-			]
-		];
-	}
+	bool bNewPinnedState = !StackViewModel->GetSystemViewModel()->GetIsEmitterPinned(StackViewModel->GetEmitterHandleViewModel()->AsShared());
+	StackViewModel->GetSystemViewModel()->SetEmitterPinnedState(StackViewModel->GetEmitterHandleViewModel()->AsShared(), bNewPinnedState);
+	CurrentPinColor = bNewPinnedState ? PinIsPinnedColor : PinIsUnpinnedColor;
+	return FReply::Handled();
+}
 
-private:
-	const FSlateBrush* GetButtonBrush() const
+void SNiagaraStack::OnSearchTextChanged(const FText& SearchText)
+{
+	if (StackViewModel->GetCurrentSearchText().CompareTo(SearchText) != 0)
 	{
-		if (IsHovered())
+		if (SearchExpandTimer.IsValid())
 		{
-			return ItemExpander->GetIsExpanded()
-				? FEditorStyle::GetBrush("DetailsView.PulldownArrow.Up.Hovered") 
-				: FEditorStyle::GetBrush("DetailsView.PulldownArrow.Down.Hovered");
+			UnRegisterActiveTimer(SearchExpandTimer.ToSharedRef());
 		}
-		else
+		// restore expansion state of previous search
+		for (auto SearchResult : StackViewModel->GetCurrentSearchResults())
 		{
-			return ItemExpander->GetIsExpanded() 
-				? FEditorStyle::GetBrush("DetailsView.PulldownArrow.Up") 
-				: FEditorStyle::GetBrush("DetailsView.PulldownArrow.Down");
+			for (UNiagaraStackEntry* ParentalUnit : SearchResult.EntryPath)
+			{
+				StackTree->SetItemExpansion(ParentalUnit, ParentalUnit->GetIsExpanded());
+			}
 		}
-	}
 
-	FText GetToolTipText() const
-	{
-		return ItemExpander->GetIsExpanded() ? ExpandedToolTipText : CollapsedToolTipText;
-	}
-
-	FReply ExpandButtonClicked()
-	{
-		ItemExpander->ToggleExpanded();
-		return FReply::Handled();
+		StackViewModel->OnSearchTextChanged(SearchText);
 	}
+}
 
-private:
-	UNiagaraStackItemExpander* ItemExpander;
-	FText ExpandedToolTipText;
-	FText CollapsedToolTipText;
-};
-
-class SNiagaraStackEmitterPropertiesItem : public SCompoundWidget
+FReply SNiagaraStack::ScrollToNextMatch()
 {
-public:
-	SLATE_BEGIN_ARGS(SNiagaraStackEmitterPropertiesItem) { }
-	SLATE_END_ARGS();
-
-	void Construct(const FArguments& InArgs, UNiagaraStackEmitterPropertiesItem& InEmitterPropertiesItem)
-	{
-		EmitterPropertiesItem = &InEmitterPropertiesItem;
+	AddSearchScrollOffset(1);
+	return FReply::Handled();
+}
 
-		ChildSlot
-		[
-			SNew(SHorizontalBox)
-			+ SHorizontalBox::Slot()
-			.Padding(1)
-			[
-				SNew(STextBlock)
-				.TextStyle(FNiagaraEditorWidgetsStyle::Get(), "NiagaraEditor.Stack.ItemText")
-				.ToolTipText_UObject(EmitterPropertiesItem, &UNiagaraStackEntry::GetTooltipText)
-				.Text_UObject(EmitterPropertiesItem, &UNiagaraStackEntry::GetDisplayName)
-			]
-			// Reset to base Button
-			+ SHorizontalBox::Slot()
-			.AutoWidth()
-			.VAlign(VAlign_Center)
-			.Padding(3, 0, 0, 0)
-			[
-				SNew(SButton)
-				.IsFocusable(false)
-				.ToolTipText(LOCTEXT("ResetEmitterPropertiesToBaseToolTip", "Reset the emitter properties to the state defined by the parent emitter"))
-				.ButtonStyle(FEditorStyle::Get(), "NoBorder")
-				.ContentPadding(0)
-				.Visibility(this, &SNiagaraStackEmitterPropertiesItem::GetResetToBaseButtonVisibility)
-				.OnClicked(this, &SNiagaraStackEmitterPropertiesItem::ResetToBaseButtonClicked)
-				.Content()
-				[
-					SNew(SImage)
-					.Image(FEditorStyle::GetBrush("PropertyWindow.DiffersFromDefault"))
-					.ColorAndOpacity(FSlateColor(FLinearColor::Green))
-				]
-			]
-		];
-	}
+FReply SNiagaraStack::ScrollToPreviousMatch()
+{
+	// move current match to the previous one in the StackTree, wrap around
+	AddSearchScrollOffset(-1);
+	return FReply::Handled();
+}
 
-private:
-	EVisibility GetResetToBaseButtonVisibility() const
+void SNiagaraStack::AddSearchScrollOffset(int NumberOfSteps)
+{
+	if (StackViewModel->IsSearching() || StackViewModel->GetCurrentSearchResults().Num() == 0 || NumberOfSteps == 0)
 	{
-		return EmitterPropertiesItem->CanResetToBase() ? EVisibility::Visible : EVisibility::Collapsed;
+		return;
 	}
 
-	FReply ResetToBaseButtonClicked()
-	{
-		EmitterPropertiesItem->ResetToBase();
-		return FReply::Handled();
-	}
+	StackViewModel->AddSearchScrollOffset(NumberOfSteps);
 
-private:
-	UNiagaraStackEmitterPropertiesItem* EmitterPropertiesItem;
-};
+	StackTree->RequestScrollIntoView(StackViewModel->GetCurrentFocusedEntry());
+}
 
-class SNiagaraStackEventHandlerPropertiesItem : public SCompoundWidget
+TOptional<SSearchBox::FSearchResultData> SNiagaraStack::GetSearchResultData() const
 {
-public:
-	SLATE_BEGIN_ARGS(SNiagaraStackEventHandlerPropertiesItem) { }
-	SLATE_END_ARGS();
-
-	void Construct(const FArguments& InArgs, UNiagaraStackEventHandlerPropertiesItem& InEventHandlerPropertiesItem)
+	if (StackViewModel->GetCurrentSearchText().IsEmpty())
 	{
-		EventHandlerPropertiesItem = &InEventHandlerPropertiesItem;
-
-		ChildSlot
-		[
-			SNew(SHorizontalBox)
-			+ SHorizontalBox::Slot()
-			.Padding(1)
-			[
-				SNew(STextBlock)
-				.TextStyle(FNiagaraEditorWidgetsStyle::Get(), "NiagaraEditor.Stack.ItemText")
-				.ToolTipText_UObject(EventHandlerPropertiesItem, &UNiagaraStackEntry::GetTooltipText)
-				.Text_UObject(EventHandlerPropertiesItem, &UNiagaraStackEntry::GetDisplayName)
-			]
-			// Reset to base Button
-			+ SHorizontalBox::Slot()
-			.AutoWidth()
-			.VAlign(VAlign_Center)
-			.Padding(3, 0, 0, 0)
-			[
-				SNew(SButton)
-				.IsFocusable(false)
-				.ToolTipText(LOCTEXT("ResetEventHandlerPropertiesToBaseToolTip", "Reset the event handler properties to the state defined by the parent emitter"))
-				.ButtonStyle(FEditorStyle::Get(), "NoBorder")
-				.ContentPadding(0)
-				.Visibility(this, &SNiagaraStackEventHandlerPropertiesItem::GetResetToBaseButtonVisibility)
-				.OnClicked(this, &SNiagaraStackEventHandlerPropertiesItem::ResetToBaseButtonClicked)
-				.Content()
-				[
-					SNew(SImage)
-					.Image(FEditorStyle::GetBrush("PropertyWindow.DiffersFromDefault"))
-					.ColorAndOpacity(FSlateColor(FLinearColor::Green))
-				]
-			]
-		];
+		return TOptional<SSearchBox::FSearchResultData>();
 	}
+	return TOptional<SSearchBox::FSearchResultData>({ StackViewModel->GetCurrentSearchResults().Num(), StackViewModel->GetCurrentFocusedMatchIndex() + 1 });
+}
 
-private:
-	EVisibility GetResetToBaseButtonVisibility() const
-	{
-		return EventHandlerPropertiesItem->CanResetToBase() ? EVisibility::Visible : EVisibility::Collapsed;
-	}
+bool SNiagaraStack::GetIsSearching() const
+{
+	return StackViewModel->IsSearching();
+}
 
-	FReply ResetToBaseButtonClicked()
+bool SNiagaraStack::IsEntryFocusedInSearch(UNiagaraStackEntry* Entry) const
+{
+	if (StackViewModel && Entry && StackViewModel->GetCurrentFocusedEntry() == Entry)
 	{
-		EventHandlerPropertiesItem->ResetToBase();
-		return FReply::Handled();
+		return true;
 	}
+	return false;
+}
 
-private:
-	UNiagaraStackEventHandlerPropertiesItem* EventHandlerPropertiesItem;
-};
-
-class SNiagaraStackTableRow : public STableRow<UNiagaraStackEntry*>
+void SNiagaraStack::OpenSourceEmitter()
 {
-public:
-	DECLARE_DELEGATE_OneParam(FOnColumnWidthChanged, float)
-
-public:
-	SLATE_BEGIN_ARGS(SNiagaraStackTableRow)
-		: _GroupPadding(FMargin(5, 0, 5, 0))
-	{}
-		SLATE_ARGUMENT(FMargin, GroupPadding)
-		SLATE_ATTRIBUTE(float, NameColumnWidth)
-		SLATE_ATTRIBUTE(float, ValueColumnWidth)
-		SLATE_EVENT(FOnColumnWidthChanged, OnNameColumnWidthChanged)
-		SLATE_EVENT(FOnColumnWidthChanged, OnValueColumnWidthChanged)
-	SLATE_END_ARGS();
-
-	void Construct(const FArguments& InArgs, UNiagaraStackEntry* InStackEntry, const TSharedRef<STreeView<UNiagaraStackEntry*>>& InOwnerTree)
-	{
-		GroupPadding = InArgs._GroupPadding;
-		NameColumnWidth = InArgs._NameColumnWidth;
-		ValueColumnWidth = InArgs._ValueColumnWidth;
-		NameColumnWidthChanged = InArgs._OnNameColumnWidthChanged;
-		ValueColumnWidthChanged = InArgs._OnValueColumnWidthChanged;
+	FAssetEditorManager::Get().OpenEditorForAsset(const_cast<UNiagaraEmitter*>(StackViewModel->GetEmitterHandleViewModel()->GetEmitterHandle()->GetSource()));
+}
 
-		StackEntry = InStackEntry;
-		OwnerTree = InOwnerTree;
+void SNiagaraStack::SetEmitterEnabled(bool bIsEnabled)
+{
+	StackViewModel->GetEmitterHandleViewModel()->SetIsEnabled(bIsEnabled);
+}
 
-		ExpandedImage = FCoreStyle::Get().GetBrush("TreeArrow_Expanded");
-		CollapsedImage = FCoreStyle::Get().GetBrush("TreeArrow_Collapsed");
+bool SNiagaraStack::CheckEmitterEnabledStatus(bool bIsEnabled)
+{
+	return StackViewModel->GetEmitterHandleViewModel()->GetIsEnabled() == bIsEnabled;
+}
 
-		InactiveItemBackgroundColor = FNiagaraEditorWidgetsStyle::Get().GetColor(StackEntry->GetItemBackgroundName());
-		ActiveItemBackgroundColor = InactiveItemBackgroundColor + FLinearColor(.05f, .05f, .05f, 0.0f);
+void SNiagaraStack::ShowEmitterInContentBrowser()
+{
+	FContentBrowserModule& ContentBrowserModule = FModuleManager::Get().LoadModuleChecked<FContentBrowserModule>(TEXT("ContentBrowser"));
+	TArray<FAssetData> Assets;
+	Assets.Add(FAssetData(StackViewModel->GetEmitterHandleViewModel()->GetEmitterHandle()->GetSource()));
+	ContentBrowserModule.Get().SyncBrowserToAssets(Assets);
+}
 
-		ConstructInternal(STableRow<UNiagaraStackEntry*>::FArguments(), OwnerTree.ToSharedRef());
-	}
+void SNiagaraStack::NavigateTo(UNiagaraStackEntry* Item)
+{
+	StackTree->RequestScrollIntoView(Item);
+}
 
-	virtual void SetNameAndValueContent(TSharedRef<SWidget> InNameWidget, TSharedPtr<SWidget> InValueWidget)
+void CollapseEntriesRecursive(TArray<UNiagaraStackEntry*> Entries)
+{
+	for (UNiagaraStackEntry* Entry : Entries)
 	{
-		TSharedRef<SHorizontalBox> NameContent = SNew(SHorizontalBox)
-			+SHorizontalBox::Slot()
-			.AutoWidth()
-			[
-				SNew(SBox)
-				.WidthOverride(this, &SNiagaraStackTableRow::GetIndentSize)
-			]
-			+SHorizontalBox::Slot()
-			.AutoWidth()
-			.VAlign(VAlign_Center)
-			.Padding(0)
-			[
-				SNew(SButton)
-				.ButtonStyle(FCoreStyle::Get(), "NoBorder")
-				.Visibility(this, &SNiagaraStackTableRow::GetExpanderVisibility)
-				.OnClicked(this, &SNiagaraStackTableRow::ExpandButtonClicked)
-				.ForegroundColor(FSlateColor::UseForeground())
-				.ContentPadding(2)
-				[
-					SNew(SImage)
-					.Image(this, &SNiagaraStackTableRow::GetExpandButtonImage)
-					.ColorAndOpacity(FSlateColor::UseForeground())
-				]
-			]
-			+SHorizontalBox::Slot()
-			[
-				InNameWidget
-			];
-
-		TSharedPtr<SWidget> ChildContent;
-
-		if (InValueWidget.IsValid())
-		{
-			ChildContent = SNew(SSplitter)
-			.Style(FNiagaraEditorWidgetsStyle::Get(), "NiagaraEditor.Stack.Splitter")
-			.PhysicalSplitterHandleSize(1.0f)
-			.HitDetectionSplitterHandleSize(5.0f)
-			+ SSplitter::Slot()
-			.Value(NameColumnWidth)
-			.OnSlotResized(SSplitter::FOnSlotResized::CreateSP(this, &SNiagaraStackTableRow::OnNameColumnWidthChanged))
-			[
-				SNew(SBox)
-				.Padding(FMargin(3, 2, 5, 2))
-				[
-					NameContent
-				]
-			]
-
-			// Value
-			+ SSplitter::Slot()
-			.Value(ValueColumnWidth)
-			.OnSlotResized(SSplitter::FOnSlotResized::CreateSP(this, &SNiagaraStackTableRow::OnValueColumnWidthChanged))
-			[
-				SNew(SBox)
-				.Padding(FMargin(4, 2, 3, 2))
-				[
-					InValueWidget.ToSharedRef()
-				]
-			];
-		}
-		else
+		if (Entry->GetCanExpand())
 		{
-			ChildContent = SNew(SBox)
-			.Padding(FMargin(3, 2, 3, 2))
-			[
-				NameContent
-			];
+			Entry->SetIsExpanded(false);
 		}
 
-		ChildSlot
-		[
-			SNew(SBorder)
-			.BorderImage(FEditorStyle::GetBrush("WhiteBrush"))
-			.BorderBackgroundColor(FNiagaraEditorWidgetsStyle::Get().GetColor(StackEntry->GetGroupBackgroundName()))
-			.Padding(this, &SNiagaraStackTableRow::GetGroupPadding)
-			.Visibility(this, &SNiagaraStackTableRow::GetRowVisibility)
-			[
-				SNew(SBorder)
-				.BorderImage(FEditorStyle::GetBrush("WhiteBrush"))
-				.BorderBackgroundColor(this, &SNiagaraStackTableRow::GetItemBackgroundColor)
-				.ForegroundColor(FNiagaraEditorWidgetsStyle::Get().GetColor(StackEntry->GetItemForegroundName()))
-				.Padding(FMargin(0))
-				[
-					ChildContent.ToSharedRef()
-				]
-			]
-		];
+		TArray<UNiagaraStackEntry*> Children;
+		Entry->GetUnfilteredChildren(Children);
+		CollapseEntriesRecursive(Children);
 	}
+}
 
-	void SetGroupPadding(FMargin InGroupPadding)
-	{
-		GroupPadding = InGroupPadding;
-	}
+void SNiagaraStack::CollapseAll()
+{
+	CollapseEntriesRecursive(StackViewModel->GetRootEntries());
+	StackViewModel->NotifyStructureChanged();
+}
 
-	bool GetIsRowActive() const
-	{
-		return IsHovered();
-	}
+TSharedRef<SWidget> SNiagaraStack::GetViewOptionsMenu() const
+{
+	FMenuBuilder MenuBuilder(false, nullptr);
+
+	MenuBuilder.AddMenuEntry(
+		LOCTEXT("ShowAllAdvancedLabel", "Show all advanced"),
+		LOCTEXT("ShowAllAdvancedToolTip", "Forces all advanced items to be showing in the stack."),
+		FSlateIcon(),
+		FUIAction(
+			FExecuteAction::CreateLambda([=]() { StackViewModel->SetShowAllAdvanced(!StackViewModel->GetShowAllAdvanced()); }),
+			FCanExecuteAction(),
+			FGetActionCheckState::CreateLambda([=]() { return StackViewModel->GetShowAllAdvanced() ? ECheckBoxState::Checked : ECheckBoxState::Unchecked; })),
+		NAME_None, EUserInterfaceActionType::Check);
+
+	MenuBuilder.AddMenuEntry(
+		LOCTEXT("ShowOutputsLabel", "Show outputs"),
+		LOCTEXT("ShowOutputsToolTip", "Whether or now to show module outputs in the stack."),
+		FSlateIcon(),
+		FUIAction(
+			FExecuteAction::CreateLambda([=]() { StackViewModel->SetShowOutputs(!StackViewModel->GetShowOutputs()); }),
+			FCanExecuteAction(),
+			FGetActionCheckState::CreateLambda([=]() { return StackViewModel->GetShowOutputs() ? ECheckBoxState::Checked : ECheckBoxState::Unchecked; })),
+		NAME_None, EUserInterfaceActionType::Check);
+
+	return MenuBuilder.MakeWidget();
+}
 
-	virtual FReply OnMouseButtonDoubleClick(const FGeometry& InMyGeometry, const FPointerEvent& InMouseEvent) override
-	{
-		return FReply::Unhandled();
-	}
+FReply SNiagaraStack::OnRowDragDetected(const FGeometry& InGeometry, const FPointerEvent& InPointerEvent, UNiagaraStackEntry* InStackEntry)
+{
+	if (InStackEntry->CanDrag())
+	{
+		TArray<UNiagaraStackEntry*> DraggedEntries;
+		DraggedEntries.Add(InStackEntry);
+		TSharedRef<FNiagaraStackEntryDragDropOp> DragDropOp = MakeShared<FNiagaraStackEntryDragDropOp>(DraggedEntries);
+		DragDropOp->CurrentHoverText = InStackEntry->GetDisplayName();
+		DragDropOp->CurrentIconBrush = FNiagaraEditorWidgetsStyle::Get().GetBrush(
+			FNiagaraStackEditorWidgetsUtilities::GetIconNameForExecutionSubcategory(InStackEntry->GetExecutionSubcategoryName(), true));
+		DragDropOp->CurrentIconColorAndOpacity = FNiagaraEditorWidgetsStyle::Get().GetColor(
+			FNiagaraStackEditorWidgetsUtilities::GetIconColorNameForExecutionCategory(InStackEntry->GetExecutionCategoryName()));
+		DragDropOp->SetupDefaults();
+		DragDropOp->Construct();
+		return FReply::Handled().BeginDragDrop(DragDropOp);
+	}
+	return FReply::Unhandled();
+}
 
-private:
-	FMargin GetGroupPadding() const
+TOptional<EItemDropZone> SNiagaraStack::OnRowCanAcceptDrop(const FDragDropEvent& InDragDropEvent, EItemDropZone InDropZone, UNiagaraStackEntry* InTargetEntry)
+{
+	TOptional<EItemDropZone> DropZone;
+	TSharedPtr<FNiagaraStackEntryDragDropOp> DragDropOp = InDragDropEvent.GetOperationAs<FNiagaraStackEntryDragDropOp>();
+	if (DragDropOp.IsValid())
 	{
-		return GroupPadding;
-	}
+		DragDropOp->ResetToDefaultToolTip();
+		TOptional<UNiagaraStackEntry::FDropResult> Result = InTargetEntry->CanDrop(DragDropOp->GetDraggedEntries());
+		if (Result.IsSet())
+		{
+			if (Result.GetValue().DropMessage.IsEmptyOrWhitespace() == false)
+			{
+				DragDropOp->CurrentHoverText = FText::Format(LOCTEXT("DropFormat", "{0} - {1}"), DragDropOp->GetDefaultHoverText(), Result.GetValue().DropMessage);
+			}
 
-	EVisibility GetRowVisibility() const
-	{
-		return StackEntry->GetShouldShowInStack() 
-			? EVisibility::Visible 
-			: EVisibility::Collapsed;
+			if (Result.GetValue().bCanDrop)
+			{
+				DropZone = EItemDropZone::OntoItem;
+			}
+			else 
+			{
+				DragDropOp->CurrentIconBrush = FEditorStyle::GetBrush("Icons.Error");
+				DragDropOp->CurrentIconColorAndOpacity = FLinearColor::White;
+			}
+		}
 	}
+	return DropZone;
+}
 
-	FOptionalSize GetIndentSize() const
+FReply SNiagaraStack::OnRowAcceptDrop(const FDragDropEvent& InDragDropEvent, EItemDropZone InDropZone, UNiagaraStackEntry* InTargetEntry)
+{
+	TSharedPtr<FNiagaraStackEntryDragDropOp> DragDropOp = InDragDropEvent.GetOperationAs<FNiagaraStackEntryDragDropOp>();
+	if (DragDropOp.IsValid())
 	{
-		return StackEntry->GetItemIndentLevel() * IndentSize;
-	}
 
-	EVisibility GetExpanderVisibility() const
-	{
-		if (StackEntry->GetCanExpand())
+		if (ensureMsgf(InTargetEntry->Drop(DragDropOp->GetDraggedEntries()).IsSet(), TEXT("Failed to drop stack entry when can drop returned true")))
 		{
-			// TODO Cache this and refresh the cache when the entries structure changes.
-			TArray<UNiagaraStackEntry*> Children;
-			StackEntry->GetFilteredChildren(Children);
-			return Children.Num() > 0
-				? EVisibility::Visible
-				: EVisibility::Hidden;
-		}
-		else
-		{
-			return EVisibility::Collapsed;
+			return FReply::Handled();
 		}
 	}
+	return FReply::Unhandled();
+}
 
-	FReply ExpandButtonClicked()
-	{
-		StackEntry->SetIsExpanded(!StackEntry->GetIsExpanded());
-		OwnerTree->SetItemExpansion(StackEntry, StackEntry->GetIsExpanded());
-		return FReply::Handled();
-	}
+void SNiagaraStack::OnStackSearchComplete()
+{
+	// fire up timer to expand all parentchains!!
+	SearchExpandTimer = RegisterActiveTimer(0.7f, FWidgetActiveTimerDelegate::CreateSP(this, &SNiagaraStack::TriggerExpandSearchResults));
+}
+
+EActiveTimerReturnType SNiagaraStack::TriggerExpandSearchResults(double InCurrentTime, float InDeltaTime)
+{
+	ExpandSearchResults();
+	ScrollToNextMatch();
+	return EActiveTimerReturnType::Stop;
+}
 
-	const FSlateBrush* GetExpandButtonImage() const
+void SNiagaraStack::ExpandSearchResults()
+{
+	for (auto SearchResult : StackViewModel->GetCurrentSearchResults())
 	{
-		return StackEntry->GetIsExpanded() ? ExpandedImage : CollapsedImage;
+		for (UNiagaraStackEntry* ParentalUnit : SearchResult.EntryPath)
+		{
+			StackTree->SetItemExpansion(ParentalUnit, true);
+		}
 	}
+}
 
-	void OnNameColumnWidthChanged(float Width)
+void SNiagaraStack::OnSearchBoxTextCommitted(const FText& NewText, ETextCommit::Type CommitInfo)
+{
+	if (StackViewModel->GetCurrentSearchText().CompareTo(NewText) != 0)
 	{
-		NameColumnWidthChanged.ExecuteIfBound(Width);
+		if (SearchExpandTimer.IsValid())
+		{
+			UnRegisterActiveTimer(SearchExpandTimer.ToSharedRef());
+			ExpandSearchResults();
+			SearchExpandTimer.Reset();
+		}
 	}
+	AddSearchScrollOffset(+1);
+}
 
-	void OnValueColumnWidthChanged(float Width)
+void SNiagaraStack::OnSearchBoxSearch(SSearchBox::SearchDirection Direction)
+{
+	if (Direction == SSearchBox::Next)
 	{
-		ValueColumnWidthChanged.ExecuteIfBound(Width);
+		ScrollToNextMatch();
 	}
-
-	FSlateColor GetItemBackgroundColor() const
+	else if (Direction == SSearchBox::Previous)
 	{
-		return GetIsRowActive() ? ActiveItemBackgroundColor : InactiveItemBackgroundColor;
+		ScrollToPreviousMatch();
 	}
-
-private:
-	UNiagaraStackEntry* StackEntry;
-	TSharedPtr<STreeView<UNiagaraStackEntry*>> OwnerTree;
-
-	TAttribute<float> NameColumnWidth;
-	TAttribute<float> ValueColumnWidth;
-	FOnColumnWidthChanged NameColumnWidthChanged;
-	FOnColumnWidthChanged ValueColumnWidthChanged;
-
-	const FSlateBrush* ExpandedImage;
-	const FSlateBrush* CollapsedImage;
-
-	FLinearColor InactiveItemBackgroundColor;
-	FLinearColor ActiveItemBackgroundColor;
-
-	FMargin GroupPadding;
-};
-
-void SNiagaraStack::Construct(const FArguments& InArgs, UNiagaraStackViewModel* InStackViewModel)
-{
-	StackViewModel = InStackViewModel;
-	StackViewModel->OnStructureChanged().AddSP(this, &SNiagaraStack::StackStructureChanged);
-	NameColumnWidth = .3f;
-	ContentColumnWidth = .7f;
-	ChildSlot
-	[
-		SNew(SBox)
-		.Padding(3)
-		[
-			SAssignNew(StackTree, STreeView<UNiagaraStackEntry*>)
-			.OnGenerateRow(this, &SNiagaraStack::OnGenerateRowForStackItem)
-			.OnGetChildren(this, &SNiagaraStack::OnGetChildren)
-			.TreeItemsSource(&StackViewModel->GetRootEntries())
-		]
-	];
-	PrimeTreeExpansion();
 }
 
-void SNiagaraStack::PrimeTreeExpansion()
+FSlateColor SNiagaraStack::GetTextColorForItem(UNiagaraStackEntry* Item) const
 {
-	TArray<UNiagaraStackEntry*> EntriesToProcess(StackViewModel->GetRootEntries());
-	while (EntriesToProcess.Num() > 0)
+	if (IsEntryFocusedInSearch(Item))
 	{
-		UNiagaraStackEntry* EntryToProcess = EntriesToProcess[0];
-		EntriesToProcess.RemoveAtSwap(0);
-
-		if (EntryToProcess->GetIsExpanded())
-		{
-			StackTree->SetItemExpansion(EntryToProcess, true);
-			EntryToProcess->GetFilteredChildren(EntriesToProcess);
-		}
+		return FSlateColor(FLinearColor(FColor::Orange));
 	}
+	return FSlateColor::UseForeground();
+}
+
+FSlateColor SNiagaraStack::GetPinColor() const
+{
+	return CurrentPinColor;
 }
 
 TSharedRef<ITableRow> SNiagaraStack::OnGenerateRowForStackItem(UNiagaraStackEntry* Item, const TSharedRef<STableViewBase>& OwnerTable)
 {
 	TSharedRef<SNiagaraStackTableRow> Container = ConstructContainerForItem(Item);
-	TSharedRef<SWidget> NameContent = ConstructNameWidgetForItem(Item, Container);
-	TSharedPtr<SWidget> ValueContent = ConstructValueWidgetForItem(Item, Container);
-	Container->SetNameAndValueContent(NameContent, ValueContent);
+	FRowWidgets RowWidgets = ConstructNameAndValueWidgetsForItem(Item, Container);
+	Container->SetNameAndValueContent(RowWidgets.NameWidget, RowWidgets.ValueWidget);
 	return Container;
 }
 
-TSharedRef<SNiagaraStackTableRow> SNiagaraStack::ConstructDefaultRow(UNiagaraStackEntry* Item)
+TSharedRef<SNiagaraStackTableRow> SNiagaraStack::ConstructContainerForItem(UNiagaraStackEntry* Item)
 {
-	return SNew(SNiagaraStackTableRow, Item, StackTree.ToSharedRef())
+	float LeftContentPadding = 4;
+	float RightContentPadding = 6;
+	FMargin ContentPadding(LeftContentPadding, 0, RightContentPadding, 0);
+	FLinearColor ItemBackgroundColor;
+	FLinearColor ItemForegroundColor = FNiagaraEditorWidgetsStyle::Get().GetColor("NiagaraEditor.Stack.ForegroundColor");
+	bool bIsCategoryIconHighlighted;
+	bool bShowExecutionCategoryIcon;
+	switch (Item->GetStackRowStyle())
+	{
+	case UNiagaraStackEntry::EStackRowStyle::None:
+		ItemBackgroundColor = FLinearColor::Transparent;
+		bIsCategoryIconHighlighted = false;
+		bShowExecutionCategoryIcon = false;
+		break;
+	case UNiagaraStackEntry::EStackRowStyle::GroupHeader:
+		ContentPadding = FMargin(LeftContentPadding, 4, 0, 0);
+		ItemBackgroundColor = FLinearColor::Transparent;
+		ItemForegroundColor = FNiagaraEditorWidgetsStyle::Get().GetColor("NiagaraEditor.Stack.GroupForegroundColor");
+		bIsCategoryIconHighlighted = true;
+		bShowExecutionCategoryIcon = true;
+		break;
+	case UNiagaraStackEntry::EStackRowStyle::ItemHeader:
+		ContentPadding = FMargin(LeftContentPadding, 2, 2, 2);
+		ItemBackgroundColor = FNiagaraEditorWidgetsStyle::Get().GetColor("NiagaraEditor.Stack.Item.HeaderBackgroundColor");
+		bIsCategoryIconHighlighted = false;
+		bShowExecutionCategoryIcon = true;
+		break;
+	case UNiagaraStackEntry::EStackRowStyle::ItemContent:
+		ContentPadding = FMargin(LeftContentPadding, 3, RightContentPadding, 3);
+		ItemBackgroundColor = FNiagaraEditorWidgetsStyle::Get().GetColor("NiagaraEditor.Stack.Item.ContentBackgroundColor");
+		bIsCategoryIconHighlighted = false;
+		bShowExecutionCategoryIcon = false;
+		break;
+	case UNiagaraStackEntry::EStackRowStyle::ItemContentAdvanced:
+		ContentPadding = FMargin(LeftContentPadding, 3, RightContentPadding, 3);
+		ItemBackgroundColor = FNiagaraEditorWidgetsStyle::Get().GetColor("NiagaraEditor.Stack.Item.ContentAdvancedBackgroundColor");
+		bIsCategoryIconHighlighted = false;
+		bShowExecutionCategoryIcon = false;
+		break;
+	case UNiagaraStackEntry::EStackRowStyle::ItemFooter:
+		ItemBackgroundColor = FNiagaraEditorWidgetsStyle::Get().GetColor("NiagaraEditor.Stack.Item.FooterBackgroundColor");
+		bIsCategoryIconHighlighted = false;
+		bShowExecutionCategoryIcon = false;
+		break;
+	case UNiagaraStackEntry::EStackRowStyle::ItemCategory:
+		ContentPadding = FMargin(LeftContentPadding, 3, RightContentPadding, 3);
+		ItemBackgroundColor = FNiagaraEditorWidgetsStyle::Get().GetColor("NiagaraEditor.Stack.Item.ContentBackgroundColor");
+		bIsCategoryIconHighlighted = false;
+		bShowExecutionCategoryIcon = false;
+		break;
+	case UNiagaraStackEntry::EStackRowStyle::StackIssue:
+		ContentPadding = FMargin(LeftContentPadding, 3, RightContentPadding, 3);
+		ItemBackgroundColor = FNiagaraEditorWidgetsStyle::Get().GetColor("NiagaraEditor.Stack.Item.IssueBackgroundColor");
+		bIsCategoryIconHighlighted = false;
+		bShowExecutionCategoryIcon = false;
+		break;
+	default:
+		ItemBackgroundColor = FNiagaraEditorWidgetsStyle::Get().GetColor("NiagaraEditor.Stack.UnknownColor");
+		bIsCategoryIconHighlighted = false;
+		bShowExecutionCategoryIcon = false;
+		break;
+	}
+
+	return SNew(SNiagaraStackTableRow, StackViewModel, Item, StackTree.ToSharedRef())
+		.ContentPadding(ContentPadding)
+		.ItemBackgroundColor(ItemBackgroundColor)
+		.ItemForegroundColor(ItemForegroundColor)
+		.IsCategoryIconHighlighted(bIsCategoryIconHighlighted)
+		.ShowExecutionCategoryIcon(bShowExecutionCategoryIcon)
 		.NameColumnWidth(this, &SNiagaraStack::GetNameColumnWidth)
 		.OnNameColumnWidthChanged(this, &SNiagaraStack::OnNameColumnWidthChanged)
 		.ValueColumnWidth(this, &SNiagaraStack::GetContentColumnWidth)
-		.OnValueColumnWidthChanged(this, &SNiagaraStack::OnContentColumnWidthChanged);
+		.OnValueColumnWidthChanged(this, &SNiagaraStack::OnContentColumnWidthChanged)
+		.OnDragDetected(this, &SNiagaraStack::OnRowDragDetected, Item)
+		.OnCanAcceptDrop(this, &SNiagaraStack::OnRowCanAcceptDrop)
+		.OnAcceptDrop(this, &SNiagaraStack::OnRowAcceptDrop);
 }
 
-TSharedRef<SNiagaraStackTableRow> SNiagaraStack::ConstructContainerForItem(UNiagaraStackEntry* Item)
+void SNiagaraStack::Tick(const FGeometry& AllottedGeometry, const double InCurrentTime, const float InDeltaTime)
 {
-	if (Item->GetClass()->IsChildOf(UNiagaraStackItemGroup::StaticClass()))
-	{
-		TSharedRef<SNiagaraStackTableRow> GroupRow = ConstructDefaultRow(Item);
-		GroupRow->SetGroupPadding(FMargin(0));
-		return GroupRow;
-	}
-	else
+	SCompoundWidget::Tick(AllottedGeometry, InCurrentTime, InDeltaTime);
+
+	if (StackViewModel)
 	{
-		return ConstructDefaultRow(Item);
+		StackViewModel->Tick();
 	}
 }
 
-TSharedRef<SWidget> SNiagaraStack::ConstructNameWidgetForItem(UNiagaraStackEntry* Item, TSharedRef<SNiagaraStackTableRow> Container)
+
+SNiagaraStack::FRowWidgets SNiagaraStack::ConstructNameAndValueWidgetsForItem(UNiagaraStackEntry* Item, TSharedRef<SNiagaraStackTableRow> Container)
 {
-	if (Item->GetClass()->IsChildOf(UNiagaraStackSpacer::StaticClass()))
-	{
-		return SNew(SBox)
-			.HeightOverride(6);
-	}
-	else if (Item->GetClass()->IsChildOf(UNiagaraStackAddModuleItem::StaticClass()))
-	{
-		return SNew(SNiagaraStackAddModuleItem, *CastChecked<UNiagaraStackAddModuleItem>(Item));
-	}
-	else if (Item->GetClass()->IsChildOf(UNiagaraStackAddRendererItem::StaticClass()))
+	if (Item->IsA<UNiagaraStackSpacer>())
 	{
-		return SNew(SNiagaraStackAddRendererItem, *CastChecked<UNiagaraStackAddRendererItem>(Item));
+		FMargin ContentPadding = Container->GetContentPadding();
+		Container->SetContentPadding(FMargin(ContentPadding.Left, 0, ContentPadding.Right, 0));
+		return FRowWidgets(SNew(SBox)
+			.HeightOverride(SpacerHeight * CastChecked<UNiagaraStackSpacer>(Item)->GetSpacerScale()));
 	}
-	else if (Item->GetClass()->IsChildOf(UNiagaraStackItemGroup::StaticClass()))
+	else if (Item->IsA<UNiagaraStackItemGroup>())
 	{
-		return SNew(SNiagaraStackItemGroup, *CastChecked<UNiagaraStackItemGroup>(Item));
+		return FRowWidgets(SNew(SNiagaraStackItemGroup, *CastChecked<UNiagaraStackItemGroup>(Item), StackViewModel));
 	}
-	else if (Item->GetClass()->IsChildOf(UNiagaraStackModuleItem::StaticClass()))
+	else if (Item->IsA<UNiagaraStackModuleItem>())
 	{
-		return SNew(SNiagaraStackModuleItem, *CastChecked<UNiagaraStackModuleItem>(Item));
+		TSharedRef<SNiagaraStackModuleItem> ModuleItemWidget = SNew(SNiagaraStackModuleItem, *CastChecked<UNiagaraStackModuleItem>(Item), StackViewModel);
+		Container->AddFillRowContextMenuHandler(SNiagaraStackTableRow::FOnFillRowContextMenu::CreateSP(ModuleItemWidget, &SNiagaraStackModuleItem::FillRowContextMenu));
+		return FRowWidgets(ModuleItemWidget);
 	}
-	else if (Item->GetClass()->IsChildOf(UNiagaraStackRendererItem::StaticClass()))
+	else if (Item->IsA<UNiagaraStackRendererItem>())
 	{
-		return SNew(SNiagaraStackRendererItem, *CastChecked<UNiagaraStackRendererItem>(Item));
+		return FRowWidgets(SNew(SNiagaraStackRendererItem, *CastChecked<UNiagaraStackRendererItem>(Item), StackViewModel));
 	}
-	else if (Item->GetClass()->IsChildOf(UNiagaraStackFunctionInput::StaticClass()))
+	else if (Item->IsA<UNiagaraStackFunctionInput>())
 	{
 		UNiagaraStackFunctionInput* FunctionInput = CastChecked<UNiagaraStackFunctionInput>(Item);
-		return SNew(SNiagaraStackFunctionInputName, FunctionInput)
-			.IsRowActive(Container, &SNiagaraStackTableRow::GetIsRowActive);
+		return FRowWidgets(
+			SNew(SNiagaraStackFunctionInputName, FunctionInput, StackViewModel),
+			SNew(SNiagaraStackFunctionInputValue, FunctionInput));
 	}
-	else if (Item->GetClass()->IsChildOf(UNiagaraStackObject::StaticClass()))
+	else if (Item->IsA<UNiagaraStackErrorItem>())
 	{
-		return SNew(SNiagaraStackObject, *CastChecked<UNiagaraStackObject>(Item));
+		return FRowWidgets(SNew(SNiagaraStackErrorItem, CastChecked<UNiagaraStackErrorItem>(Item), StackViewModel));
 	}
-	else if (Item->GetClass()->IsChildOf(UNiagaraStackStruct::StaticClass()))
+	else if (Item->IsA<UNiagaraStackErrorItemLongDescription>())
 	{
-		return SNew(SNiagaraStackStruct, *CastChecked<UNiagaraStackStruct>(Item));
+		Container->SetOverrideNameAlignment(EHorizontalAlignment::HAlign_Fill, EVerticalAlignment::VAlign_Center);
+		return FRowWidgets(SNew(STextBlock)
+			.TextStyle(FNiagaraEditorWidgetsStyle::Get(), "NiagaraEditor.Stack.ParameterText")
+			.ToolTipText_UObject(Item, &UNiagaraStackEntry::GetTooltipText)
+			.Text_UObject(Item, &UNiagaraStackEntry::GetDisplayName)
+			.ColorAndOpacity(this, &SNiagaraStack::GetTextColorForItem, Item)
+			.HighlightText_UObject(StackViewModel, &UNiagaraStackViewModel::GetCurrentSearchText)
+			.AutoWrapText(true));
 	}
-	else if (Item->GetClass()->IsChildOf(UNiagaraStackErrorItem::StaticClass()))
+	else if (Item->IsA<UNiagaraStackErrorItemFix>())
 	{
-		TSharedPtr<SHorizontalBox> ErrorInternalBox = SNew(SHorizontalBox);
-		UNiagaraStackErrorItem* ErrorItem = Cast<UNiagaraStackErrorItem>(Item);
-		ErrorInternalBox->AddSlot()
-			.HAlign(HAlign_Left)
-			.VAlign(VAlign_Center)
-			.AutoWidth()
-			[
-				SNew(STextBlock)
-				.TextStyle(FNiagaraEditorStyle::Get(), "NiagaraEditor.ParameterText")
-				.Text_UObject(ErrorItem, &UNiagaraStackErrorItem::ErrorText)
-				.ToolTipText_UObject(ErrorItem, &UNiagaraStackErrorItem::ErrorTextTooltip)
-			];
-		ErrorInternalBox->AddSlot()
-			.HAlign(HAlign_Right)
-			.VAlign(VAlign_Center)
-			.Padding(10,0)
-			[
-				SNew(SHorizontalBox)
-				.Visibility_UObject(ErrorItem, &UNiagaraStackErrorItem::CanFixVisibility)
-				+ SHorizontalBox::Slot()
-					.VAlign(VAlign_Center)
-					[
-						SNew(SButton)
-						.TextStyle(FNiagaraEditorStyle::Get(), "NiagaraEditor.ParameterText")
-						.Text(LOCTEXT("FixError", "Fix Error"))
-						.OnClicked_UObject(ErrorItem, &UNiagaraStackErrorItem::OnTryFixError)
-					]
-			];
-
-		return SNew(SHorizontalBox)
-			+ SHorizontalBox::Slot()
-				.AutoWidth()
-				.HAlign(HAlign_Left)
-				.VAlign(VAlign_Center)
-				[
-					SNew(SImage)
-					.Image(FEditorStyle::GetBrush("Icons.Error"))
-				]
-			+ SHorizontalBox::Slot()
-				.AutoWidth()
-				.HAlign(HAlign_Left)
-				.VAlign(VAlign_Center)
-				[
-					ErrorInternalBox.ToSharedRef()
-				];
+		return FRowWidgets(SNew(SNiagaraStackErrorItemFix, CastChecked<UNiagaraStackErrorItemFix>(Item), StackViewModel));
 	}
-	else if (Item->GetClass()->IsChildOf(UNiagaraStackItemExpander::StaticClass()))
+	else if (Item->IsA<UNiagaraStackAdvancedExpander>())
 	{
-		UNiagaraStackItemExpander* ItemExpander = CastChecked<UNiagaraStackItemExpander>(Item);
-		return SNew(SNiagaraStackItemExpander, *ItemExpander);
+		UNiagaraStackAdvancedExpander* ItemExpander = CastChecked<UNiagaraStackAdvancedExpander>(Item);
+		return FRowWidgets(SNew(SNiagaraStackItemExpander, *ItemExpander));
 	}
-	else if (Item->GetClass()->IsChildOf(UNiagaraStackEmitterPropertiesItem::StaticClass()))
+	else if (Item->IsA<UNiagaraStackEmitterPropertiesItem>())
 	{
 		UNiagaraStackEmitterPropertiesItem* PropertiesItem = CastChecked<UNiagaraStackEmitterPropertiesItem>(Item);
-		return SNew(SNiagaraStackEmitterPropertiesItem, *PropertiesItem);
+		return FRowWidgets(SNew(SNiagaraStackEmitterPropertiesItem, *PropertiesItem, StackViewModel));
 	}
-	else if (Item->GetClass()->IsChildOf(UNiagaraStackEventHandlerPropertiesItem::StaticClass()))
+	else if (Item->IsA<UNiagaraStackEventHandlerPropertiesItem>())
 	{
 		UNiagaraStackEventHandlerPropertiesItem* PropertiesItem = CastChecked<UNiagaraStackEventHandlerPropertiesItem>(Item);
-		return SNew(SNiagaraStackEventHandlerPropertiesItem, *PropertiesItem);
+		return FRowWidgets(SNew(SNiagaraStackEventHandlerPropertiesItem, *PropertiesItem, StackViewModel));
 	}
-	else
+	else if (Item->IsA<UNiagaraStackParameterStoreEntry>())
 	{
-		return SNew(STextBlock)
-			.TextStyle(FNiagaraEditorWidgetsStyle::Get(), Item->GetTextStyleName())
-			.ToolTipText_UObject(Item, &UNiagaraStackEntry::GetTooltipText)
-			.Text_UObject(Item, &UNiagaraStackEntry::GetDisplayName);
+		UNiagaraStackParameterStoreEntry* StackEntry = CastChecked<UNiagaraStackParameterStoreEntry>(Item);
+		return FRowWidgets(
+			SNew(SNiagaraStackParameterStoreEntryName, StackEntry, StackViewModel),
+			SNew(SNiagaraStackParameterStoreEntryValue, StackEntry));
 	}
-}
-
-TSharedPtr<SWidget> SNiagaraStack::ConstructValueWidgetForItem(UNiagaraStackEntry* Item, TSharedRef<SNiagaraStackTableRow> Container)
-{
-	if (Item->GetClass()->IsChildOf(UNiagaraStackFunctionInput::StaticClass()))
+	else if (Item->IsA<UNiagaraStackInputCategory>())
 	{
-		UNiagaraStackFunctionInput* FunctionInput = CastChecked<UNiagaraStackFunctionInput>(Item);
-		return SNew(SNiagaraStackFunctionInputValue, FunctionInput);
+		return FRowWidgets(SNew(STextBlock)
+			.TextStyle(FNiagaraEditorWidgetsStyle::Get(), "NiagaraEditor.Stack.CategoryText")
+			.ToolTipText_UObject(Item, &UNiagaraStackEntry::GetTooltipText)
+			.Text_UObject(Item, &UNiagaraStackEntry::GetDisplayName)
+			.ColorAndOpacity(this, &SNiagaraStack::GetTextColorForItem, Item)
+			.HighlightText_UObject(StackViewModel, &UNiagaraStackViewModel::GetCurrentSearchText),
+			SNullWidget::NullWidget);
 	}
-	else if (Item->GetClass()->IsChildOf(UNiagaraStackParameterStoreEntry::StaticClass()))
+	else if (Item->IsA<UNiagaraStackModuleItemOutput>())
 	{
-		UNiagaraStackParameterStoreEntry* FunctionInput = CastChecked<UNiagaraStackParameterStoreEntry>(Item);
-		return SNew(SNiagaraStackParameterStoreEntryValue, FunctionInput);
+		UNiagaraStackModuleItemOutput* ModuleItemOutput = CastChecked<UNiagaraStackModuleItemOutput>(Item);
+		return FRowWidgets(
+			SNew(STextBlock)
+			.TextStyle(FNiagaraEditorWidgetsStyle::Get(), "NiagaraEditor.Stack.DefaultText")
+			.ToolTipText_UObject(Item, &UNiagaraStackEntry::GetTooltipText)
+			.Text_UObject(Item, &UNiagaraStackEntry::GetDisplayName)
+			.ColorAndOpacity(this, &SNiagaraStack::GetTextColorForItem, Item)
+			.HighlightText_UObject(StackViewModel, &UNiagaraStackViewModel::GetCurrentSearchText),
+			SNew(STextBlock)
+			.TextStyle(FNiagaraEditorWidgetsStyle::Get(), "NiagaraEditor.Stack.ParameterText")
+			.Text_UObject(ModuleItemOutput, &UNiagaraStackModuleItemOutput::GetOutputParameterHandleText)
+			.ColorAndOpacity(this, &SNiagaraStack::GetTextColorForItem, Item)
+			.HighlightText_UObject(StackViewModel, &UNiagaraStackViewModel::GetCurrentSearchText));
+	}
+	else if (Item->IsA<UNiagaraStackFunctionInputCollection>() ||
+		Item->IsA<UNiagaraStackModuleItemOutputCollection>())
+	{
+		return FRowWidgets(
+			SNew(STextBlock)
+			.TextStyle(FNiagaraEditorWidgetsStyle::Get(), "NiagaraEditor.Stack.DefaultText")
+			.ToolTipText_UObject(Item, &UNiagaraStackEntry::GetTooltipText)
+			.Text_UObject(Item, &UNiagaraStackEntry::GetDisplayName)
+			.ColorAndOpacity(this, &SNiagaraStack::GetTextColorForItem, Item)
+			.HighlightText_UObject(StackViewModel, &UNiagaraStackViewModel::GetCurrentSearchText),
+			SNullWidget::NullWidget);
 	}
-	else if (Item->GetClass()->IsChildOf(UNiagaraStackModuleItemOutput::StaticClass()))
+	else if (Item->IsA<UNiagaraStackPropertyRow>())
 	{
-		UNiagaraStackModuleItemOutput* ModuleItemOutput = CastChecked<UNiagaraStackModuleItemOutput>(Item);
-		return SNew(STextBlock)
-			.TextStyle(FNiagaraEditorWidgetsStyle::Get(), ModuleItemOutput->GetTextStyleName())
-			.Text_UObject(ModuleItemOutput, &UNiagaraStackModuleItemOutput::GetOutputParameterHandleText);
+		UNiagaraStackPropertyRow* PropertyRow = CastChecked<UNiagaraStackPropertyRow>(Item);
+		FNodeWidgets PropertyRowWidgets = PropertyRow->GetDetailTreeNode()->CreateNodeWidgets();
+		if (PropertyRowWidgets.WholeRowWidget.IsValid())
+		{
+			Container->SetOverrideNameWidth(PropertyRowWidgets.WholeRowWidgetLayoutData.MinWidth, PropertyRowWidgets.WholeRowWidgetLayoutData.MaxWidth);
+			Container->SetOverrideNameAlignment(PropertyRowWidgets.WholeRowWidgetLayoutData.HorizontalAlignment, PropertyRowWidgets.WholeRowWidgetLayoutData.VerticalAlignment);
+			return FRowWidgets(PropertyRowWidgets.WholeRowWidget.ToSharedRef());
+		}
+		else
+		{
+			Container->SetOverrideNameWidth(PropertyRowWidgets.NameWidgetLayoutData.MinWidth, PropertyRowWidgets.NameWidgetLayoutData.MaxWidth);
+			Container->SetOverrideNameAlignment(PropertyRowWidgets.NameWidgetLayoutData.HorizontalAlignment, PropertyRowWidgets.NameWidgetLayoutData.VerticalAlignment);
+			Container->SetOverrideValueWidth(PropertyRowWidgets.ValueWidgetLayoutData.MinWidth, PropertyRowWidgets.ValueWidgetLayoutData.MaxWidth);
+			Container->SetOverrideValueAlignment(PropertyRowWidgets.ValueWidgetLayoutData.HorizontalAlignment, PropertyRowWidgets.ValueWidgetLayoutData.VerticalAlignment);
+			return FRowWidgets(PropertyRowWidgets.NameWidget.ToSharedRef(), PropertyRowWidgets.ValueWidget.ToSharedRef());
+		}
 	}
-	else if (Item->GetClass()->IsChildOf(UNiagaraStackFunctionInputCollection::StaticClass()) ||
-		Item->GetClass()->IsChildOf(UNiagaraStackModuleItemOutputCollection::StaticClass()))
+	else if (Item->IsA<UNiagaraStackItem>())
 	{
-		return SNullWidget::NullWidget;
+		return FRowWidgets(SNew(STextBlock)
+			.TextStyle(FNiagaraEditorWidgetsStyle::Get(), "NiagaraEditor.Stack.ItemText")
+			.ToolTipText_UObject(Item, &UNiagaraStackEntry::GetTooltipText)
+			.Text_UObject(Item, &UNiagaraStackEntry::GetDisplayName)
+			.ColorAndOpacity(this, &SNiagaraStack::GetTextColorForItem, Item)
+			.HighlightText_UObject(StackViewModel, &UNiagaraStackViewModel::GetCurrentSearchText));
 	}
 	else
 	{
-		return TSharedPtr<SWidget>();
+		return FRowWidgets(SNew(STextBlock)
+			.TextStyle(FNiagaraEditorWidgetsStyle::Get(), "NiagaraEditor.Stack.DefaultText")
+			.ToolTipText_UObject(Item, &UNiagaraStackEntry::GetTooltipText)
+			.Text_UObject(Item, &UNiagaraStackEntry::GetDisplayName)
+			.ColorAndOpacity(this, &SNiagaraStack::GetTextColorForItem, Item)
+			.HighlightText_UObject(StackViewModel, &UNiagaraStackViewModel::GetCurrentSearchText));
 	}
 }
 
@@ -2365,6 +885,11 @@ void SNiagaraStack::OnGetChildren(UNiagaraStackEntry* Item, TArray<UNiagaraStack
 	Item->GetFilteredChildren(Children);
 }
 
+void SNiagaraStack::StackTreeScrolled(double ScrollValue)
+{
+	StackViewModel->SetLastScrollPosition(ScrollValue);
+}
+
 float SNiagaraStack::GetNameColumnWidth() const
 {
 	return NameColumnWidth;
@@ -2387,6 +912,7 @@ void SNiagaraStack::OnContentColumnWidthChanged(float Width)
 
 void SNiagaraStack::StackStructureChanged()
 {
+	ConstructHeaderWidget();
 	PrimeTreeExpansion();
 	StackTree->RequestTreeRefresh();
 }
