Copying //UE4/Dev-Niagara to //UE4/Dev-Main (Source: //UE4/Dev-Niagara @ 4074996)

diff --git a/Engine/Plugins/FX/Niagara/Source/NiagaraEditorWidgets/Private/SNiagaraStack.h b/Engine/Plugins/FX/Niagara/Source/NiagaraEditorWidgets/Private/SNiagaraStack.h
--- a/Engine/Plugins/FX/Niagara/Source/NiagaraEditorWidgets/Private/SNiagaraStack.h
+++ b/Engine/Plugins/FX/Niagara/Source/NiagaraEditorWidgets/Private/SNiagaraStack.h
@@ -9,11 +9,13 @@
 #include "Widgets/Views/STableViewBase.h"
 #include "Widgets/Views/STreeView.h"
 #include "ViewModels/Stack/NiagaraStackEntry.h"
-#include "EdGraph/EdGraphSchema.h"
-#include "SNiagaraStack.generated.h"
+#include "Widgets/Input/SSearchBox.h"
+
 
 class UNiagaraStackViewModel;
 class SNiagaraStackTableRow;
+class SSearchBox;
+class FReply;
 
 class SNiagaraStack : public SCompoundWidget, public FEditorUndoClient
 {
@@ -26,20 +28,35 @@ public:
 	void Construct(const FArguments& InArgs, UNiagaraStackViewModel* InStackViewModel);
 
 private:
+	struct FRowWidgets
+	{
+		FRowWidgets(TSharedRef<SWidget> InNameWidget, TSharedRef<SWidget> InValueWidget)
+			: NameWidget(InNameWidget)
+			, ValueWidget(InValueWidget)
+		{
+		}
+
+		FRowWidgets(TSharedRef<SWidget> InWholeRowWidget)
+			: NameWidget(InWholeRowWidget)
+		{
+		}
+
+		TSharedRef<SWidget> NameWidget;
+		TSharedPtr<SWidget> ValueWidget;
+	};
+
 	void PrimeTreeExpansion();
 
 	TSharedRef<ITableRow> OnGenerateRowForStackItem(UNiagaraStackEntry* Item, const TSharedRef<STableViewBase>& OwnerTable);
 
-	TSharedRef<SNiagaraStackTableRow> ConstructDefaultRow(UNiagaraStackEntry* Item);
-
 	TSharedRef<SNiagaraStackTableRow> ConstructContainerForItem(UNiagaraStackEntry* Item);
 
-	TSharedRef<SWidget> ConstructNameWidgetForItem(UNiagaraStackEntry* Item, TSharedRef<SNiagaraStackTableRow> Container);
-
-	TSharedPtr<SWidget> ConstructValueWidgetForItem(UNiagaraStackEntry* Item, TSharedRef<SNiagaraStackTableRow> Container);
+	FRowWidgets ConstructNameAndValueWidgetsForItem(UNiagaraStackEntry* Item, TSharedRef<SNiagaraStackTableRow> Container);
 	
 	void OnGetChildren(UNiagaraStackEntry* Item, TArray<UNiagaraStackEntry*>& Children);
 
+	void StackTreeScrolled(double ScrollValue);
+
 	float GetNameColumnWidth() const;
 	float GetContentColumnWidth() const;
 
@@ -49,6 +66,44 @@ private:
 
 	EVisibility GetVisibilityForItem(UNiagaraStackEntry* Item) const;
 
+	virtual void Tick(const FGeometry& AllottedGeometry, const double InCurrentTime, const float InDeltaTime) override;
+
+	void ConstructHeaderWidget();
+	FSlateColor GetPinColor() const;
+	FReply PinButtonPressed();
+
+	// ~stack search stuff
+	void OnSearchTextChanged(const FText& SearchText);
+	FReply ScrollToNextMatch();
+	FReply ScrollToPreviousMatch();
+	TOptional<SSearchBox::FSearchResultData> GetSearchResultData() const;
+	bool GetIsSearching() const;
+	void OnSearchBoxTextCommitted(const FText& NewText, ETextCommit::Type CommitInfo);	
+	void OnSearchBoxSearch(SSearchBox::SearchDirection Direction);
+	FSlateColor GetTextColorForItem(UNiagaraStackEntry* Item) const;
+	void AddSearchScrollOffset(int NumberOfSteps);
+	void OnStackSearchComplete();
+	EActiveTimerReturnType TriggerExpandSearchResults(double InCurrentTime, float InDeltaTime);
+	void ExpandSearchResults();
+	bool IsEntryFocusedInSearch(UNiagaraStackEntry* Entry) const;
+	
+	// Inline menu commands
+	void OpenSourceEmitter();
+	void SetEmitterEnabled(bool bIsEnabled);
+	bool CheckEmitterEnabledStatus(bool bIsEnabled);
+	void ShowEmitterInContentBrowser();
+	void NavigateTo(UNiagaraStackEntry* Item);
+	void CollapseAll();
+
+	TSharedRef<SWidget> GetViewOptionsMenu() const;
+
+	// Drag/Drop
+	FReply OnRowDragDetected(const FGeometry& InGeometry, const FPointerEvent& InPointerEvent, UNiagaraStackEntry* InStackEntry);
+
+	TOptional<EItemDropZone> OnRowCanAcceptDrop(const FDragDropEvent& InDragDropEvent, EItemDropZone InDropZone, UNiagaraStackEntry* InTargetEntry);
+
+	FReply OnRowAcceptDrop(const FDragDropEvent& InDragDropEvent, EItemDropZone InDropZone, UNiagaraStackEntry* InTargetEntry);
+
 private:
 	UNiagaraStackViewModel* StackViewModel;
 
@@ -57,31 +112,17 @@ private:
 	float NameColumnWidth;
 
 	float ContentColumnWidth;
-};
-
-USTRUCT()
-struct FNiagaraStackGraphSchemaAction : public FEdGraphSchemaAction
-{
-	GENERATED_USTRUCT_BODY();
-
-	DECLARE_DELEGATE(FOnPerformStackAction);
-
-	FNiagaraStackGraphSchemaAction()
-	{
-	}
-
-	FNiagaraStackGraphSchemaAction(FText InNodeCategory, FText InMenuDesc, FText InToolTip, const int32 InGrouping, FText InKeywords, FOnPerformStackAction InAction)
-		: FEdGraphSchemaAction(MoveTemp(InNodeCategory), MoveTemp(InMenuDesc), MoveTemp(InToolTip), InGrouping, MoveTemp(InKeywords))
-		, Action(InAction)
-	{}
+	
+	TSharedPtr<SWidget> HeaderWidget;
 
-	//~ Begin FEdGraphSchemaAction Interface
-	UEdGraphNode* PerformAction(class UEdGraph* ParentGraph, UEdGraphPin* FromPin, const FVector2D Location, bool bSelectNewNode = true) override
-	{
-		Action.ExecuteIfBound();
-		return nullptr;
-	}
-	//~ End FEdGraphSchemaAction Interface
+	FLinearColor PinIsPinnedColor;
+	
+	FLinearColor PinIsUnpinnedColor;
 
-	FOnPerformStackAction Action;
+	FLinearColor CurrentPinColor;
+	
+	// ~ search stuff
+	TSharedPtr<SSearchBox> SearchBox;
+	TSharedPtr<FActiveTimerHandle> SearchExpandTimer;
+	static const FText OccurencesFormat;
 };
diff --git a/Engine/Plugins/FX/Niagara/Source/NiagaraEditorWidgets/Private/SNiagaraStackEntryWidget.cpp b/Engine/Plugins/FX/Niagara/Source/NiagaraEditorWidgets/Private/SNiagaraStackEntryWidget.cpp
new file mode 100644
--- /dev/null
+++ b/Engine/Plugins/FX/Niagara/Source/NiagaraEditorWidgets/Private/SNiagaraStackEntryWidget.cpp
@@ -0,0 +1,32 @@
+// Copyright 1998-2018 Epic Games, Inc. All Rights Reserved.
+
+#include "SNiagaraStackEntryWidget.h"
+#include "ViewModels/Stack/NiagaraStackViewModel.h"
+#include "NiagaraEditorWidgetsStyle.h"
+
+#define	LOCTEXT_NAMESPACE "SNiagaraStackEntryWidget"
+
+FSlateColor SNiagaraStackEntryWidget::GetTextColorForSearch() const
+{
+	if (IsCurrentSearchMatch())
+	{
+		return FNiagaraEditorWidgetsStyle::Get().GetColor("NiagaraEditor.Stack.SearchHighlightColor");
+	} 
+	
+	return FSlateColor::UseForeground();
+}
+
+bool SNiagaraStackEntryWidget::IsCurrentSearchMatch() const
+{
+	auto FocusedEntry = StackViewModel->GetCurrentFocusedEntry();
+	return StackEntryItem != nullptr && FocusedEntry == StackEntryItem;
+}
+
+FReply SNiagaraStackEntryWidget::ExpandEntry()
+{
+	StackEntryItem->SetIsExpanded(true);
+	StackEntryItem->OnStructureChanged().Broadcast();
+	return FReply::Handled();
+}
+
+#undef LOCTEXT_NAMESPACE
diff --git a/Engine/Plugins/FX/Niagara/Source/NiagaraEditorWidgets/Private/SNiagaraStackEntryWidget.h b/Engine/Plugins/FX/Niagara/Source/NiagaraEditorWidgets/Private/SNiagaraStackEntryWidget.h
new file mode 100644
--- /dev/null
+++ b/Engine/Plugins/FX/Niagara/Source/NiagaraEditorWidgets/Private/SNiagaraStackEntryWidget.h
@@ -0,0 +1,24 @@
+// Copyright 1998-2018 Epic Games, Inc. All Rights Reserved.
+
+#pragma once
+
+#include "Widgets/SCompoundWidget.h"
+#include "Widgets/DeclarativeSyntaxSupport.h"
+
+class UNiagaraStackViewModel;
+class UNiagaraStackEntry;
+
+class SNiagaraStackEntryWidget : public SCompoundWidget
+{
+public:
+	FSlateColor GetTextColorForSearch() const;
+	FReply ExpandEntry();
+	
+protected:
+	bool IsCurrentSearchMatch() const;
+	
+protected:
+	UNiagaraStackViewModel * StackViewModel;
+	
+	UNiagaraStackEntry * StackEntryItem;
+};
\ No newline at end of file
diff --git a/Engine/Plugins/FX/Niagara/Source/NiagaraEditorWidgets/Private/Stack/SNiagaraStackEmitterPropertiesItem.cpp b/Engine/Plugins/FX/Niagara/Source/NiagaraEditorWidgets/Private/Stack/SNiagaraStackEmitterPropertiesItem.cpp
new file mode 100644
--- /dev/null
+++ b/Engine/Plugins/FX/Niagara/Source/NiagaraEditorWidgets/Private/Stack/SNiagaraStackEmitterPropertiesItem.cpp
@@ -0,0 +1,69 @@
+// Copyright 1998-2018 Epic Games, Inc. All Rights Reserved.
+
+#include "Stack/SNiagaraStackEmitterPropertiesItem.h"
+#include "NiagaraEditorWidgetsStyle.h"
+#include "NiagaraEditorStyle.h"
+#include "EditorStyleSet.h"
+#include "ViewModels/Stack/NiagaraStackEmitterSpawnScriptItemGroup.h"
+#include "ViewModels/Stack/NiagaraStackViewModel.h"
+#include "Widgets/SBoxPanel.h"
+#include "Widgets/Text/STextBlock.h"
+#include "Widgets/Input/SButton.h"
+#include "Widgets/Images/SImage.h"
+
+#define LOCTEXT_NAMESPACE "NiagaraStackEmitterPropertiesItem"
+
+void SNiagaraStackEmitterPropertiesItem::Construct(const FArguments& InArgs, UNiagaraStackEmitterPropertiesItem& InEmitterPropertiesItem, UNiagaraStackViewModel* InStackViewModel)
+{
+	EmitterPropertiesItem = &InEmitterPropertiesItem;
+	StackEntryItem = EmitterPropertiesItem;
+	StackViewModel = InStackViewModel;
+
+	ChildSlot
+	[
+		SNew(SHorizontalBox)
+		+ SHorizontalBox::Slot()
+		.Padding(0)
+		[
+			SNew(STextBlock)
+			.TextStyle(FNiagaraEditorWidgetsStyle::Get(), "NiagaraEditor.Stack.ItemText")
+			.ToolTipText_UObject(EmitterPropertiesItem, &UNiagaraStackEntry::GetTooltipText)
+			.Text_UObject(EmitterPropertiesItem, &UNiagaraStackEntry::GetDisplayName)
+			.HighlightText_UObject(InStackViewModel, &UNiagaraStackViewModel::GetCurrentSearchText)
+			.ColorAndOpacity(this, &SNiagaraStackEmitterPropertiesItem::GetTextColorForSearch)
+		]
+		// Reset to base Button
+		+ SHorizontalBox::Slot()
+		.AutoWidth()
+		.VAlign(VAlign_Center)
+		.Padding(3, 0, 0, 0)
+		[
+			SNew(SButton)
+			.IsFocusable(false)
+			.ToolTipText(LOCTEXT("ResetEmitterPropertiesToBaseToolTip", "Reset the emitter properties to the state defined by the parent emitter"))
+			.ButtonStyle(FEditorStyle::Get(), "NoBorder")
+			.ContentPadding(0)
+			.Visibility(this, &SNiagaraStackEmitterPropertiesItem::GetResetToBaseButtonVisibility)
+			.OnClicked(this, &SNiagaraStackEmitterPropertiesItem::ResetToBaseButtonClicked)
+			.Content()
+			[
+				SNew(SImage)
+				.Image(FEditorStyle::GetBrush("PropertyWindow.DiffersFromDefault"))
+				.ColorAndOpacity(FSlateColor(FLinearColor::Green))
+			]
+		]
+	];
+}
+
+EVisibility SNiagaraStackEmitterPropertiesItem::GetResetToBaseButtonVisibility() const
+{
+	return EmitterPropertiesItem->CanResetToBase() ? EVisibility::Visible : EVisibility::Collapsed;
+}
+
+FReply SNiagaraStackEmitterPropertiesItem::ResetToBaseButtonClicked()
+{
+	EmitterPropertiesItem->ResetToBase();
+	return FReply::Handled();
+}
+
+#undef LOCTEXT_NAMESPACE
\ No newline at end of file
diff --git a/Engine/Plugins/FX/Niagara/Source/NiagaraEditorWidgets/Private/Stack/SNiagaraStackEmitterPropertiesItem.h b/Engine/Plugins/FX/Niagara/Source/NiagaraEditorWidgets/Private/Stack/SNiagaraStackEmitterPropertiesItem.h
new file mode 100644
--- /dev/null
+++ b/Engine/Plugins/FX/Niagara/Source/NiagaraEditorWidgets/Private/Stack/SNiagaraStackEmitterPropertiesItem.h
@@ -0,0 +1,26 @@
+// Copyright 1998-2018 Epic Games, Inc. All Rights Reserved.
+
+#pragma once
+
+#include "SNiagaraStackEntryWidget.h"
+#include "Layout/Visibility.h"
+
+class UNiagaraStackEmitterPropertiesItem;
+class UNiagaraStackViewModel;
+
+class SNiagaraStackEmitterPropertiesItem: public SNiagaraStackEntryWidget
+{
+public:
+	SLATE_BEGIN_ARGS(SNiagaraStackEmitterPropertiesItem) { }
+	SLATE_END_ARGS();
+
+	void Construct(const FArguments& InArgs, UNiagaraStackEmitterPropertiesItem& InEmitterPropertiesItem, UNiagaraStackViewModel* InStackViewModel);
+
+private:
+	EVisibility GetResetToBaseButtonVisibility() const;
+
+	FReply ResetToBaseButtonClicked();
+
+private:
+	UNiagaraStackEmitterPropertiesItem* EmitterPropertiesItem;
+};
\ No newline at end of file
diff --git a/Engine/Plugins/FX/Niagara/Source/NiagaraEditorWidgets/Private/Stack/SNiagaraStackErrorButton.cpp b/Engine/Plugins/FX/Niagara/Source/NiagaraEditorWidgets/Private/Stack/SNiagaraStackErrorButton.cpp
new file mode 100644
--- /dev/null
+++ b/Engine/Plugins/FX/Niagara/Source/NiagaraEditorWidgets/Private/Stack/SNiagaraStackErrorButton.cpp
@@ -0,0 +1,54 @@
+﻿// Copyright 1998-2018 Epic Games, Inc. All Rights Reserved.
+
+#include "SNiagaraStackErrorButton.h"
+#include "NiagaraEditorWidgetsStyle.h"
+#include "Widgets/Images/SImage.h"
+#include "Widgets/Input/SButton.h"
+#include "Widgets/Layout/SBox.h"
+#include "EditorStyleSet.h"
+
+#define	LOCTEXT_NAMESPACE "SNiagaraStackErrorButton"
+
+void SNiagaraStackErrorButton::Construct(const FArguments& InArgs)
+{
+	IssueSeverity = InArgs._IssueSeverity;
+	ErrorTooltip = InArgs._ErrorTooltip;
+	OnButtonClicked = InArgs._OnButtonClicked;
+	FName IconName("Icons.Error");
+	switch (IssueSeverity.Get())
+	{
+		case EStackIssueSeverity::Error:
+			IconName = "Icons.Error";
+			break;
+		case EStackIssueSeverity::Warning:
+			IconName = "Icons.Warning";
+			break;
+		case EStackIssueSeverity::Info:
+			IconName = "Icons.Info";
+			break;
+		default:
+			IconName = "Icons.Warning";
+			break;
+	}
+	ChildSlot
+	[
+		SNew(SButton)
+		.ButtonStyle(FEditorStyle::Get(), "HoverHintOnly")
+		.ForegroundColor(FSlateColor::UseForeground())
+		.ToolTipText(ErrorTooltip)
+		.ContentPadding(FMargin(1.0f, 0.0f))
+		.OnClicked(OnButtonClicked)
+		.Content()
+		[
+			SNew(SBox)
+			.HAlign(HAlign_Center)
+			.VAlign(VAlign_Center)
+			[
+				SNew(SImage)
+				.Image(FEditorStyle::GetBrush(IconName))
+			]
+		]
+	];
+}
+
+#undef LOCTEXT_NAMESPACE
diff --git a/Engine/Plugins/FX/Niagara/Source/NiagaraEditorWidgets/Private/Stack/SNiagaraStackErrorButton.h b/Engine/Plugins/FX/Niagara/Source/NiagaraEditorWidgets/Private/Stack/SNiagaraStackErrorButton.h
new file mode 100644
--- /dev/null
+++ b/Engine/Plugins/FX/Niagara/Source/NiagaraEditorWidgets/Private/Stack/SNiagaraStackErrorButton.h
@@ -0,0 +1,27 @@
+﻿// Copyright 1998-2018 Epic Games, Inc. All Rights Reserved.
+
+#pragma once
+
+#include "Widgets/SCompoundWidget.h"
+#include "Widgets/DeclarativeSyntaxSupport.h"
+#include "ViewModels/Stack/NiagaraStackEntry.h"
+#include "Widgets/Input/SButton.h"
+
+class UNiagaraStackViewModel;
+class UNiagaraStackEntry;
+
+class SNiagaraStackErrorButton : public SCompoundWidget
+{
+public:
+	SLATE_BEGIN_ARGS(SNiagaraStackErrorButton) {}
+		SLATE_ATTRIBUTE(EStackIssueSeverity, IssueSeverity)
+		SLATE_ATTRIBUTE(FText, ErrorTooltip)
+		SLATE_EVENT(FOnClicked, OnButtonClicked)
+	SLATE_END_ARGS()
+	virtual void Construct(const FArguments& InArgs);
+	
+private:
+	TAttribute<EStackIssueSeverity> IssueSeverity;
+	TAttribute<FText> ErrorTooltip;
+	FOnClicked OnButtonClicked;
+};
\ No newline at end of file
diff --git a/Engine/Plugins/FX/Niagara/Source/NiagaraEditorWidgets/Private/Stack/SNiagaraStackEventHandlerPropertiesItem.cpp b/Engine/Plugins/FX/Niagara/Source/NiagaraEditorWidgets/Private/Stack/SNiagaraStackEventHandlerPropertiesItem.cpp
new file mode 100644
--- /dev/null
+++ b/Engine/Plugins/FX/Niagara/Source/NiagaraEditorWidgets/Private/Stack/SNiagaraStackEventHandlerPropertiesItem.cpp
@@ -0,0 +1,69 @@
+// Copyright 1998-2018 Epic Games, Inc. All Rights Reserved.
+
+#include "Stack/SNiagaraStackEventHandlerPropertiesItem.h"
+#include "NiagaraEditorWidgetsStyle.h"
+#include "NiagaraEditorStyle.h"
+#include "EditorStyleSet.h"
+#include "ViewModels/Stack/NiagaraStackEventScriptItemGroup.h"
+#include "ViewModels/Stack/NiagaraStackViewModel.h"
+#include "Widgets/SBoxPanel.h"
+#include "Widgets/Input/SButton.h"
+#include "Widgets/Text/STextBlock.h"
+#include "Widgets/Images/SImage.h"
+
+#define LOCTEXT_NAMESPACE "NiagaraStackEventHandlerPropertiesItem"
+
+void SNiagaraStackEventHandlerPropertiesItem::Construct(const FArguments& InArgs, UNiagaraStackEventHandlerPropertiesItem& InEventHandlerPropertiesItem, UNiagaraStackViewModel* InStackViewModel)
+{
+	EventHandlerPropertiesItem = &InEventHandlerPropertiesItem;
+	StackEntryItem = EventHandlerPropertiesItem;
+	StackViewModel = InStackViewModel;
+
+	ChildSlot
+	[
+		SNew(SHorizontalBox)
+		+ SHorizontalBox::Slot()
+		.Padding(0)
+		[
+			SNew(STextBlock)
+			.TextStyle(FNiagaraEditorWidgetsStyle::Get(), "NiagaraEditor.Stack.ItemText")
+			.ToolTipText_UObject(EventHandlerPropertiesItem, &UNiagaraStackEntry::GetTooltipText)
+			.Text_UObject(EventHandlerPropertiesItem, &UNiagaraStackEntry::GetDisplayName)
+			.HighlightText_UObject(InStackViewModel, &UNiagaraStackViewModel::GetCurrentSearchText)
+			.ColorAndOpacity(this, &SNiagaraStackEventHandlerPropertiesItem::GetTextColorForSearch)
+		]
+		// Reset to base Button
+		+ SHorizontalBox::Slot()
+		.AutoWidth()
+		.VAlign(VAlign_Center)
+		.Padding(3, 0, 0, 0)
+		[
+			SNew(SButton)
+			.IsFocusable(false)
+			.ToolTipText(LOCTEXT("ResetEventHandlerPropertiesToBaseToolTip", "Reset the event handler properties to the state defined by the parent emitter"))
+			.ButtonStyle(FEditorStyle::Get(), "NoBorder")
+			.ContentPadding(0)
+			.Visibility(this, &SNiagaraStackEventHandlerPropertiesItem::GetResetToBaseButtonVisibility)
+			.OnClicked(this, &SNiagaraStackEventHandlerPropertiesItem::ResetToBaseButtonClicked)
+			.Content()
+			[
+				SNew(SImage)
+				.Image(FEditorStyle::GetBrush("PropertyWindow.DiffersFromDefault"))
+				.ColorAndOpacity(FSlateColor(FLinearColor::Green))
+			]
+		]
+	];
+}
+
+EVisibility SNiagaraStackEventHandlerPropertiesItem::GetResetToBaseButtonVisibility() const
+{
+	return EventHandlerPropertiesItem->CanResetToBase() ? EVisibility::Visible : EVisibility::Collapsed;
+}
+
+FReply SNiagaraStackEventHandlerPropertiesItem::ResetToBaseButtonClicked()
+{
+	EventHandlerPropertiesItem->ResetToBase();
+	return FReply::Handled();
+}
+
+#undef LOCTEXT_NAMESPACE
\ No newline at end of file
diff --git a/Engine/Plugins/FX/Niagara/Source/NiagaraEditorWidgets/Private/Stack/SNiagaraStackEventHandlerPropertiesItem.h b/Engine/Plugins/FX/Niagara/Source/NiagaraEditorWidgets/Private/Stack/SNiagaraStackEventHandlerPropertiesItem.h
new file mode 100644
--- /dev/null
+++ b/Engine/Plugins/FX/Niagara/Source/NiagaraEditorWidgets/Private/Stack/SNiagaraStackEventHandlerPropertiesItem.h
@@ -0,0 +1,26 @@
+// Copyright 1998-2018 Epic Games, Inc. All Rights Reserved.
+
+#pragma once
+
+#include "SNiagaraStackEntryWidget.h"
+#include "Layout/Visibility.h"
+
+class UNiagaraStackEventHandlerPropertiesItem;
+class UNiagaraStackViewModel;
+
+class SNiagaraStackEventHandlerPropertiesItem: public SNiagaraStackEntryWidget
+{
+public:
+	SLATE_BEGIN_ARGS(SNiagaraStackEventHandlerPropertiesItem) { }
+	SLATE_END_ARGS();
+
+	void Construct(const FArguments& InArgs, UNiagaraStackEventHandlerPropertiesItem& InEventHandlerPropertiesItem, UNiagaraStackViewModel* InStackViewModel);
+
+private:
+	EVisibility GetResetToBaseButtonVisibility() const;
+
+	FReply ResetToBaseButtonClicked();
+
+private:
+	UNiagaraStackEventHandlerPropertiesItem* EventHandlerPropertiesItem;
+};
\ No newline at end of file
diff --git a/Engine/Plugins/FX/Niagara/Source/NiagaraEditorWidgets/Private/Stack/SNiagaraStackFunctionInputValue.cpp b/Engine/Plugins/FX/Niagara/Source/NiagaraEditorWidgets/Private/Stack/SNiagaraStackFunctionInputValue.cpp
new file mode 100644
--- /dev/null
+++ b/Engine/Plugins/FX/Niagara/Source/NiagaraEditorWidgets/Private/Stack/SNiagaraStackFunctionInputValue.cpp
@@ -0,0 +1,759 @@
+// Copyright 1998-2018 Epic Games, Inc. All Rights Reserved.
+
+#include "Stack/SNiagaraStackFunctionInputValue.h"
+#include "ViewModels/Stack/NiagaraStackFunctionInput.h"
+#include "NiagaraEditorModule.h"
+#include "INiagaraEditorTypeUtilities.h"
+#include "NiagaraEditorUtilities.h"
+#include "NiagaraEditorWidgetsStyle.h"
+#include "NiagaraEditorStyle.h"
+#include "NiagaraNodeFunctionCall.h"
+#include "NiagaraNodeCustomHlsl.h"
+#include "NiagaraActions.h"
+#include "SNiagaraParameterEditor.h"
+#include "Widgets/SBoxPanel.h"
+#include "Widgets/Layout/SBox.h"
+#include "Widgets/Input/SButton.h"
+#include "Widgets/Input/SEditableTextBox.h"
+#include "Widgets/Input/SComboButton.h"
+#include "Widgets/Images/SImage.h"
+#include "Editor/PropertyEditor/Public/PropertyEditorModule.h"
+#include "Framework/Application/SlateApplication.h"
+#include "IStructureDetailsView.h"
+#include "SDropTarget.h"
+#include "Modules/ModuleManager.h"
+
+#define LOCTEXT_NAMESPACE "NiagaraStackFunctionInputValue"
+
+const float TextIconSize = 16;
+
+void SNiagaraStackFunctionInputValue::Construct(const FArguments& InArgs, UNiagaraStackFunctionInput* InFunctionInput)
+{
+	FunctionInput = InFunctionInput;
+
+	FunctionInput->OnValueChanged().AddSP(this, &SNiagaraStackFunctionInputValue::OnInputValueChanged);
+	DisplayedLocalValueStruct = FunctionInput->GetLocalValueStruct();
+
+	FMargin ItemPadding = FMargin(0);
+	ChildSlot
+	[
+		SNew(SDropTarget)
+		.OnAllowDrop(this, &SNiagaraStackFunctionInputValue::OnFunctionInputAllowDrop)
+		.OnDrop(this, &SNiagaraStackFunctionInputValue::OnFunctionInputDrop)
+		.Content()
+		[
+			// Values
+			SNew(SHorizontalBox)
+			.IsEnabled(this, &SNiagaraStackFunctionInputValue::GetInputEnabled)
+			+ SHorizontalBox::Slot()
+			.VAlign(VAlign_Center)
+			.AutoWidth()
+			.Padding(0, 0, 3, 0)
+			[
+				// Value Icon
+				SNew(SBox)
+				.WidthOverride(TextIconSize)
+				.VAlign(VAlign_Center)
+				.Visibility(this, &SNiagaraStackFunctionInputValue::GetInputIconVisibility)
+				[
+					SNew(STextBlock)
+					.Font(FEditorStyle::Get().GetFontStyle("FontAwesome.10"))
+					.Text(this, &SNiagaraStackFunctionInputValue::GetInputIconText)
+					.ToolTipText(this, &SNiagaraStackFunctionInputValue::GetInputIconToolTip)
+					.ColorAndOpacity(this, &SNiagaraStackFunctionInputValue::GetInputIconColor)
+				]
+			]
+			+ SHorizontalBox::Slot()
+			[
+				// TODO Don't generate all of these widgets for every input, only generate the ones that are used based on the value type.
+
+				// Local struct
+				SNew(SOverlay)
+				+ SOverlay::Slot()
+				[
+					SAssignNew(LocalValueStructContainer, SBox)
+					.Visibility(this, &SNiagaraStackFunctionInputValue::GetValueWidgetVisibility, UNiagaraStackFunctionInput::EValueMode::Local)
+					[
+						ConstructLocalValueStructWidget()
+					]
+				]
+
+				// Linked handle
+				+ SOverlay::Slot()
+				.Padding(0.0f, 0.0f, 0.0f, 2.0f)
+				[
+					SNew(SBox)
+					.Visibility(this, &SNiagaraStackFunctionInputValue::GetValueWidgetVisibility, UNiagaraStackFunctionInput::EValueMode::Linked)
+					.ToolTipText_UObject(InFunctionInput, &UNiagaraStackFunctionInput::GetTooltipText, UNiagaraStackFunctionInput::EValueMode::Linked)
+					.VAlign(VAlign_Center)
+					[
+						SNew(STextBlock)
+						.TextStyle(FNiagaraEditorStyle::Get(), "NiagaraEditor.ParameterText")
+						.Text(this, &SNiagaraStackFunctionInputValue::GetLinkedValueHandleText)
+					]
+				]
+
+				// Data Object
+				+ SOverlay::Slot()
+				.Padding(0.0f, 0.0f, 0.0f, 2.0f)
+				[
+					SNew(SBox)
+					.Visibility(this, &SNiagaraStackFunctionInputValue::GetValueWidgetVisibility, UNiagaraStackFunctionInput::EValueMode::Data)
+					.VAlign(VAlign_Center)
+					[
+						SNew(STextBlock)
+						.TextStyle(FNiagaraEditorStyle::Get(), "NiagaraEditor.ParameterText")
+						.Text(this, &SNiagaraStackFunctionInputValue::GetDataValueText)
+					]
+				]
+
+				// Dynamic input name
+				+ SOverlay::Slot()
+				[
+					SNew(SBox)
+					.Visibility(this, &SNiagaraStackFunctionInputValue::GetValueWidgetVisibility, UNiagaraStackFunctionInput::EValueMode::Dynamic)
+					.VAlign(VAlign_Center)
+					[
+						SNew(STextBlock)
+						.TextStyle(FNiagaraEditorStyle::Get(), "NiagaraEditor.ParameterText")
+						.Text(this, &SNiagaraStackFunctionInputValue::GetDynamicValueText)
+						.OnDoubleClicked(this, &SNiagaraStackFunctionInputValue::DynamicInputTextDoubleClicked)
+					]
+				]
+
+				// Expression input
+				+ SOverlay::Slot()
+				[
+					SNew(SBox)
+					.Visibility(this, &SNiagaraStackFunctionInputValue::GetValueWidgetVisibility, UNiagaraStackFunctionInput::EValueMode::Expression)
+					.VAlign(VAlign_Center)
+					[
+						//SNew(SInlineEditableTextBlock)
+						//.Style(FNiagaraEditorStyle::Get(), "NiagaraEditor.ParameterInlineEditableText")
+						//.Text(this, &SNiagaraStackFunctionInputValue::GetExpressionValueText)
+						//.OnTextCommitted(this, &SNiagaraStackFunctionInputValue::OnExpressionTextCommitted)
+						// SNew(SMultiLineEditableTextBox)
+						//.AutoWrapText(false)
+						SNew(SEditableTextBox)
+						.IsReadOnly(false)
+						.Text(this, &SNiagaraStackFunctionInputValue::GetExpressionValueText)
+						.OnTextCommitted(this, &SNiagaraStackFunctionInputValue::OnExpressionTextCommitted)
+					]
+				]
+
+				// Invalid input
+				+ SOverlay::Slot()
+				[
+					SNew(SBox)
+					.Visibility(this, &SNiagaraStackFunctionInputValue::GetValueWidgetVisibility, UNiagaraStackFunctionInput::EValueMode::Invalid)
+					.VAlign(VAlign_Center)
+					[
+						SNew(STextBlock)
+						.TextStyle(FNiagaraEditorStyle::Get(), "NiagaraEditor.ParameterText")
+						.Text(this, &SNiagaraStackFunctionInputValue::GetInvalidValueText)
+						.ToolTipText(this, &SNiagaraStackFunctionInputValue::GetInvalidValueToolTipText)
+					]
+				]
+			]
+
+			// Handle drop-down button
+			+ SHorizontalBox::Slot()
+			.AutoWidth()
+			.VAlign(VAlign_Center)
+			.Padding(3, 0, 0, 0)
+			[
+				SAssignNew(SetFunctionInputButton, SComboButton)
+				.ButtonStyle(FEditorStyle::Get(), "HoverHintOnly")
+				.ForegroundColor(FSlateColor::UseForeground())
+				.OnGetMenuContent(this, &SNiagaraStackFunctionInputValue::OnGetAvailableHandleMenu)
+				.ContentPadding(FMargin(2))
+				.MenuPlacement(MenuPlacement_BelowRightAnchor)
+				.HAlign(HAlign_Center)
+				.VAlign(VAlign_Center)
+			]
+
+			// Reset Button
+			+ SHorizontalBox::Slot()
+			.AutoWidth()
+			.VAlign(VAlign_Center)
+			.Padding(3, 0, 0, 0)
+			[
+				SNew(SButton)
+				.IsFocusable(false)
+				.ToolTipText(LOCTEXT("ResetToolTip", "Reset to the default value"))
+				.ButtonStyle(FEditorStyle::Get(), "NoBorder")
+				.ContentPadding(0)
+				.Visibility(this, &SNiagaraStackFunctionInputValue::GetResetButtonVisibility)
+				.OnClicked(this, &SNiagaraStackFunctionInputValue::ResetButtonPressed)
+				.Content()
+				[
+					SNew(SImage)
+					.Image(FEditorStyle::GetBrush("PropertyWindow.DiffersFromDefault"))
+				]
+			]
+
+			// Reset to base Button
+			+ SHorizontalBox::Slot()
+			.AutoWidth()
+			.VAlign(VAlign_Center)
+			.Padding(3, 0, 0, 0)
+			[
+				SNew(SButton)
+				.IsFocusable(false)
+				.ToolTipText(LOCTEXT("ResetToBaseToolTip", "Reset this input to the value defined by the parent emitter"))
+				.ButtonStyle(FEditorStyle::Get(), "NoBorder")
+				.ContentPadding(0)
+				.Visibility(this, &SNiagaraStackFunctionInputValue::GetResetToBaseButtonVisibility)
+				.OnClicked(this, &SNiagaraStackFunctionInputValue::ResetToBaseButtonPressed)
+				.Content()
+				[
+					SNew(SImage)
+					.Image(FEditorStyle::GetBrush("PropertyWindow.DiffersFromDefault"))
+					.ColorAndOpacity(FSlateColor(FLinearColor::Green))
+				]
+			]
+		]
+	];
+}
+
+bool SNiagaraStackFunctionInputValue::GetInputEnabled() const
+{
+	return FunctionInput->GetHasEditCondition() == false || FunctionInput->GetEditConditionEnabled();
+}
+
+EVisibility SNiagaraStackFunctionInputValue::GetValueWidgetVisibility(UNiagaraStackFunctionInput::EValueMode ValidMode) const
+{
+	return FunctionInput->GetValueMode() == ValidMode ? EVisibility::Visible : EVisibility::Collapsed;
+}
+
+TSharedRef<SWidget> SNiagaraStackFunctionInputValue::ConstructLocalValueStructWidget()
+{
+	LocalValueStructParameterEditor.Reset();
+	LocalValueStructDetailsView.Reset();
+	if (DisplayedLocalValueStruct.IsValid())
+	{
+		FNiagaraEditorModule& NiagaraEditorModule = FModuleManager::GetModuleChecked<FNiagaraEditorModule>("NiagaraEditor");
+		TSharedPtr<INiagaraEditorTypeUtilities, ESPMode::ThreadSafe> TypeEditorUtilities = NiagaraEditorModule.GetTypeUtilities(FunctionInput->GetInputType());
+		if (TypeEditorUtilities.IsValid() && TypeEditorUtilities->CanCreateParameterEditor())
+		{
+			TSharedPtr<SNiagaraParameterEditor> ParameterEditor = TypeEditorUtilities->CreateParameterEditor(FunctionInput->GetInputType());
+			ParameterEditor->UpdateInternalValueFromStruct(DisplayedLocalValueStruct.ToSharedRef());
+			ParameterEditor->SetOnBeginValueChange(SNiagaraParameterEditor::FOnValueChange::CreateSP(
+				this, &SNiagaraStackFunctionInputValue::ParameterBeginValueChange));
+			ParameterEditor->SetOnEndValueChange(SNiagaraParameterEditor::FOnValueChange::CreateSP(
+				this, &SNiagaraStackFunctionInputValue::ParameterEndValueChange));
+			ParameterEditor->SetOnValueChanged(SNiagaraParameterEditor::FOnValueChange::CreateSP(
+				this, &SNiagaraStackFunctionInputValue::ParameterValueChanged, ParameterEditor.ToSharedRef()));
+
+			LocalValueStructParameterEditor = ParameterEditor;
+
+			return SNew(SBox)
+				.HAlign(ParameterEditor->GetHorizontalAlignment())
+				.VAlign(ParameterEditor->GetVerticalAlignment())
+				[
+					ParameterEditor.ToSharedRef()
+				];
+		}
+		else
+		{
+			FPropertyEditorModule& PropertyEditorModule = FModuleManager::GetModuleChecked<FPropertyEditorModule>("PropertyEditor");
+
+			TSharedRef<IStructureDetailsView> StructureDetailsView = PropertyEditorModule.CreateStructureDetailView(
+				FDetailsViewArgs(false, false, false, FDetailsViewArgs::HideNameArea, true),
+				FStructureDetailsViewArgs(),
+				nullptr);
+
+			StructureDetailsView->SetStructureData(DisplayedLocalValueStruct);
+			StructureDetailsView->GetOnFinishedChangingPropertiesDelegate().AddSP(this, &SNiagaraStackFunctionInputValue::ParameterPropertyValueChanged);
+
+			LocalValueStructDetailsView = StructureDetailsView;
+			return StructureDetailsView->GetWidget().ToSharedRef();
+		}
+	}
+	else
+	{
+		return SNullWidget::NullWidget;
+	}
+}
+
+void SNiagaraStackFunctionInputValue::OnInputValueChanged()
+{
+	TSharedPtr<FStructOnScope> NewLocalValueStruct = FunctionInput->GetLocalValueStruct();
+	if (DisplayedLocalValueStruct == NewLocalValueStruct)
+	{
+		if (LocalValueStructParameterEditor.IsValid())
+		{
+			LocalValueStructParameterEditor->UpdateInternalValueFromStruct(DisplayedLocalValueStruct.ToSharedRef());
+		}
+		if (LocalValueStructDetailsView.IsValid())
+		{
+			LocalValueStructDetailsView->SetStructureData(TSharedPtr<FStructOnScope>());
+			LocalValueStructDetailsView->SetStructureData(DisplayedLocalValueStruct);
+		}
+	}
+	else
+	{
+		DisplayedLocalValueStruct = NewLocalValueStruct;
+		LocalValueStructContainer->SetContent(ConstructLocalValueStructWidget());
+	}
+}
+
+void SNiagaraStackFunctionInputValue::ParameterBeginValueChange()
+{
+	FunctionInput->NotifyBeginLocalValueChange();
+}
+
+void SNiagaraStackFunctionInputValue::ParameterEndValueChange()
+{
+	FunctionInput->NotifyEndLocalValueChange();
+}
+
+void SNiagaraStackFunctionInputValue::ParameterValueChanged(TSharedRef<SNiagaraParameterEditor> ParameterEditor)
+{
+	ParameterEditor->UpdateStructFromInternalValue(FunctionInput->GetLocalValueStruct().ToSharedRef());
+	FunctionInput->SetLocalValue(DisplayedLocalValueStruct.ToSharedRef());
+}
+
+void SNiagaraStackFunctionInputValue::ParameterPropertyValueChanged(const FPropertyChangedEvent& PropertyChangedEvent)
+{
+	FunctionInput->SetLocalValue(DisplayedLocalValueStruct.ToSharedRef());
+}
+
+FText SNiagaraStackFunctionInputValue::GetLinkedValueHandleText() const
+{
+	return FText::FromName(FunctionInput->GetLinkedValueHandle().GetParameterHandleString());
+}
+
+FText SNiagaraStackFunctionInputValue::GetDataValueText() const
+{
+	if (FunctionInput->GetDataValueObject() != nullptr)
+	{
+		return FunctionInput->GetInputType().GetClass()->GetDisplayNameText();
+	}
+	else
+	{
+		return FText::Format(LOCTEXT("InvalidDataObjectFormat", "{0} (Invalid)"), FunctionInput->GetInputType().GetClass()->GetDisplayNameText());
+	}
+}
+
+FText SNiagaraStackFunctionInputValue::GetDynamicValueText() const
+{
+	if (FunctionInput->GetDynamicInputNode() != nullptr)
+	{
+		return FText::FromString(FName::NameToDisplayString(FunctionInput->GetDynamicInputNode()->GetFunctionName(), false));
+	}
+	else
+	{
+		return LOCTEXT("InvalidDynamicDisplayName", "(Invalid)");
+	}
+}
+
+FText SNiagaraStackFunctionInputValue::GetExpressionValueText() const
+{
+	if (FunctionInput->GetExpressionNode() != nullptr)
+	{
+		return FunctionInput->GetExpressionNode()->GetHlslText();
+	}
+	else
+	{
+		return LOCTEXT("InvalidDynamicDisplayName", "(Invalid)");
+	}
+}
+
+void SNiagaraStackFunctionInputValue::OnExpressionTextCommitted(const FText& Name, ETextCommit::Type CommitInfo)
+{
+	if (FunctionInput->GetExpressionNode() != nullptr)
+	{
+		FunctionInput->GetExpressionNode()->OnCustomHlslTextCommitted(Name, CommitInfo);
+	}
+}
+
+FText SNiagaraStackFunctionInputValue::GetInvalidValueText() const
+{
+	if (FunctionInput->CanReset())
+	{
+		return LOCTEXT("InvalidResetLabel", "Unsupported value - Reset to fix.");
+	}
+	else
+	{
+		return LOCTEXT("InvalidLabel", "Custom value");
+	}
+}
+
+FText SNiagaraStackFunctionInputValue::GetInvalidValueToolTipText() const
+{
+	if (FunctionInput->CanReset())
+	{
+		return LOCTEXT("InvalidResetToolTip", "This input has an unsupported value assigned in the stack.\nUse the reset button to remove the unsupported value.");
+	}
+	else
+	{
+		return LOCTEXT("InvalidToolTip", "The script that defines the source of this input has\n a custom default value that can not be displayed in the stack view.\nYou can set a local override value using the drop down menu.");
+	}
+}
+
+FReply SNiagaraStackFunctionInputValue::DynamicInputTextDoubleClicked()
+{
+	UNiagaraNodeFunctionCall* DynamicInputNode = FunctionInput->GetDynamicInputNode();
+	if (DynamicInputNode->FunctionScript != nullptr && DynamicInputNode->FunctionScript->IsAsset())
+	{
+		FAssetEditorManager::Get().OpenEditorForAsset(DynamicInputNode->FunctionScript);
+		return FReply::Handled();
+	}
+	return FReply::Unhandled();
+}
+
+TSharedRef<SExpanderArrow> SNiagaraStackFunctionInputValue::CreateCustomNiagaraFunctionInputActionExpander(const FCustomExpanderData& ActionMenuData)
+{
+	return SNew(SNiagaraFunctionInputActionMenuExpander, ActionMenuData);
+}
+
+TSharedRef<SWidget> SNiagaraStackFunctionInputValue::OnGetAvailableHandleMenu()
+{
+	TSharedPtr<SGraphActionMenu> SelectInputFunctionMenu;
+	TSharedRef<SBorder> MenuWidget = SNew(SBorder)
+	.BorderImage(FEditorStyle::GetBrush("Menu.Background"))
+	.Padding(5)
+	[
+		SNew(SBox)
+		.WidthOverride(300)
+		.HeightOverride(400)
+		[
+			SAssignNew(SelectInputFunctionMenu, SGraphActionMenu)
+			.OnActionSelected(this, &SNiagaraStackFunctionInputValue::OnActionSelected)
+			.OnCollectAllActions(this, &SNiagaraStackFunctionInputValue::CollectAllActions)
+			.AutoExpandActionMenu(false)
+			.ShowFilterTextBox(true)
+			.OnCreateCustomRowExpander_Static(&CreateCustomNiagaraFunctionInputActionExpander)
+		]
+	];
+
+	SetFunctionInputButton->SetMenuContentWidgetToFocus(SelectInputFunctionMenu->GetFilterTextBox()->AsShared());
+	return MenuWidget;
+}
+
+void SNiagaraStackFunctionInputValue::OnActionSelected(const TArray<TSharedPtr<FEdGraphSchemaAction>>& SelectedActions, ESelectInfo::Type InSelectionType)
+{
+	if (InSelectionType == ESelectInfo::OnMouseClick || InSelectionType == ESelectInfo::OnKeyPress || SelectedActions.Num() == 0)
+	{
+		for (int32 ActionIndex = 0; ActionIndex < SelectedActions.Num(); ActionIndex++)
+		{
+			TSharedPtr<FNiagaraMenuAction> CurrentAction = StaticCastSharedPtr<FNiagaraMenuAction>(SelectedActions[ActionIndex]);
+
+			if (CurrentAction.IsValid())
+			{
+				FSlateApplication::Get().DismissAllMenus();
+				CurrentAction->ExecuteAction();
+			}
+		}
+	}
+}
+
+void SNiagaraStackFunctionInputValue::CollectAllActions(FGraphActionListBuilderBase& OutAllActions)
+{
+	// Set a local value
+	{
+		bool bCanSetLocalValue =
+			FunctionInput->GetValueMode() != UNiagaraStackFunctionInput::EValueMode::Local &&
+			FunctionInput->GetInputType().IsDataInterface() == false;
+
+		const FText NameText = LOCTEXT("LocalValue", "Set a local value");
+		const FText Tooltip = FText::Format(LOCTEXT("LocalValueToolTip", "Set a local editable value for this input."), NameText);
+		const FText CategoryName = LOCTEXT("LocalValueCategory", "Local");
+		TSharedPtr<FNiagaraMenuAction> SetLocalValueAction(
+			new FNiagaraMenuAction(CategoryName, NameText, Tooltip, 0, FText(),
+				FNiagaraMenuAction::FOnExecuteStackAction::CreateSP(this, &SNiagaraStackFunctionInputValue::SetToLocalValue),
+				FNiagaraMenuAction::FCanExecuteStackAction::CreateLambda([=]() { return bCanSetLocalValue; })));
+		OutAllActions.AddAction(SetLocalValueAction);
+	}
+
+
+	// Add a dynamic input
+	{
+		const FText CategoryName = LOCTEXT("DynamicInputValueCategory", "Dynamic Inputs");
+		TArray<UNiagaraScript*> DynamicInputScripts;
+		FunctionInput->GetAvailableDynamicInputs(DynamicInputScripts);
+		for (UNiagaraScript* DynamicInputScript : DynamicInputScripts)
+		{
+			const FText DynamicInputText = FText::FromString(FName::NameToDisplayString(DynamicInputScript->GetName(), false));
+			const FText Tooltip = FText::Format(LOCTEXT("DynamicInputFormat", "Use {0} to provide a value for this input. Description: {1}"), DynamicInputText, DynamicInputScript->GetDescription());
+			TSharedPtr<FNiagaraMenuAction> DynamicInputAction(new FNiagaraMenuAction(CategoryName, DynamicInputText, Tooltip, 0, FText(),
+				FNiagaraMenuAction::FOnExecuteStackAction::CreateSP(this, &SNiagaraStackFunctionInputValue::DynamicInputScriptSelected, DynamicInputScript)));
+			OutAllActions.AddAction(DynamicInputAction);
+		}
+	}
+
+	// Link existing attribute
+	TArray<FNiagaraParameterHandle> AvailableHandles;
+	FunctionInput->GetAvailableParameterHandles(AvailableHandles);
+
+	TArray<FNiagaraParameterHandle> UserHandles;
+	TArray<FNiagaraParameterHandle> EngineHandles;
+	TArray<FNiagaraParameterHandle> SystemHandles;
+	TArray<FNiagaraParameterHandle> EmitterHandles;
+	TArray<FNiagaraParameterHandle> ParticleAttributeHandles;
+	TArray<FNiagaraParameterHandle> OtherHandles;
+	for (const FNiagaraParameterHandle AvailableHandle : AvailableHandles)
+	{
+		if (AvailableHandle.IsUserHandle())
+		{
+			UserHandles.Add(AvailableHandle);
+		}
+		else if (AvailableHandle.IsEngineHandle())
+		{
+			EngineHandles.Add(AvailableHandle);
+		}
+		else if (AvailableHandle.IsSystemHandle())
+		{
+			SystemHandles.Add(AvailableHandle);
+		}
+		else if (AvailableHandle.IsEmitterHandle())
+		{
+			EmitterHandles.Add(AvailableHandle);
+		}
+		else if (AvailableHandle.IsParticleAttributeHandle())
+		{
+			ParticleAttributeHandles.Add(AvailableHandle);
+		}
+		else
+		{
+			OtherHandles.Add(AvailableHandle);
+		}
+	}
+
+	{
+		const FString RootCategoryName = FString("Link Inputs");
+		auto AddMenuItemsForHandleList = [&](const TArray<FNiagaraParameterHandle>& Handles, const FText& SectionDisplay)
+		{
+			const FText MapInputFormat = LOCTEXT("LinkInputFormat", "Link this input to {0}");
+			for (const FNiagaraParameterHandle& Handle : Handles)
+			{
+				const FText DisplayName = FText::FromString(FName::NameToDisplayString(Handle.GetName().ToString(), false));
+				const FText Tooltip = FText::Format(MapInputFormat, FText::FromName(Handle.GetParameterHandleString()));
+				TSharedPtr<FNiagaraMenuAction> LinkAction(new FNiagaraMenuAction(SectionDisplay, DisplayName, Tooltip, 0, FText(),
+					FNiagaraMenuAction::FOnExecuteStackAction::CreateSP(this, &SNiagaraStackFunctionInputValue::ParameterHandleSelected, Handle)));
+				OutAllActions.AddAction(LinkAction, RootCategoryName);
+			}
+		};
+
+		AddMenuItemsForHandleList(UserHandles, LOCTEXT("UserSection", "User Exposed"));
+		AddMenuItemsForHandleList(EngineHandles, LOCTEXT("EngineSection", "Engine"));
+		AddMenuItemsForHandleList(SystemHandles, LOCTEXT("SystemSection", "System"));
+		AddMenuItemsForHandleList(EmitterHandles, LOCTEXT("EmitterSection", "Emitter"));
+		AddMenuItemsForHandleList(ParticleAttributeHandles, LOCTEXT("ParticleAttributeSection", "Particle Attribute"));
+		AddMenuItemsForHandleList(OtherHandles, LOCTEXT("OtherSection", "Other"));
+	}
+
+	// Read from new attribute
+	{
+		const FText CategoryName = LOCTEXT("MakeCategory", "Make");
+
+		TArray<FName> AvailableNamespaces;
+		FunctionInput->GetNamespacesForNewParameters(AvailableNamespaces);
+
+		TArray<FString> InputNames;
+		for (int32 i = FunctionInput->GetInputParameterHandlePath().Num() - 1; i >= 0; i--)
+		{
+			InputNames.Add(FunctionInput->GetInputParameterHandlePath()[i].GetName().ToString());
+		}
+		FName InputName = *FString::Join(InputNames, TEXT("."));
+
+		for (const FName AvailableNamespace : AvailableNamespaces)
+		{
+			FNiagaraParameterHandle HandleToRead(AvailableNamespace, InputName);
+			bool bCanExecute = AvailableHandles.Contains(HandleToRead) == false;
+
+			FFormatNamedArguments Args;
+			Args.Add(TEXT("AvailableNamespace"), FText::FromName(AvailableNamespace));
+
+			const FText DisplayName = FText::Format(LOCTEXT("ReadLabelFormat", "Read from new {AvailableNamespace} parameter"), Args);
+			const FText Tooltip = FText::Format(LOCTEXT("ReadToolTipFormat", "Read this input from a new parameter in the {AvailableNamespace} namespace."), Args);
+			TSharedPtr<FNiagaraMenuAction> MakeAction(
+				new FNiagaraMenuAction(CategoryName, DisplayName, Tooltip, 0, FText(),
+					FNiagaraMenuAction::FOnExecuteStackAction::CreateSP(this, &SNiagaraStackFunctionInputValue::ParameterHandleSelected, HandleToRead),
+					FNiagaraMenuAction::FCanExecuteStackAction::CreateLambda([=]() { return bCanExecute; })));
+			OutAllActions.AddAction(MakeAction);
+		}
+	}
+
+	{
+		const FText CategoryName = LOCTEXT("ExpressionCategory", "Expression");
+		const FText DisplayName = LOCTEXT("ExpressionLabel", "Make new expression");
+		const FText Tooltip = LOCTEXT("ExpressionToolTipl", "Resolve this variable with a custom expression.");
+		TSharedPtr<FNiagaraMenuAction> ExpressionAction(new FNiagaraMenuAction(CategoryName, DisplayName, Tooltip, 0, FText(),
+			FNiagaraMenuAction::FOnExecuteStackAction::CreateSP(this, &SNiagaraStackFunctionInputValue::CustomExpressionSelected)));
+		OutAllActions.AddAction(ExpressionAction);
+	}
+
+	if (FunctionInput->CanDeleteInput())
+	{
+		const FText NameText = LOCTEXT("DeleteInput", "Remove");
+		const FText Tooltip = FText::Format(LOCTEXT("DeleteInputTooltip", "Remove input from module."), NameText);
+		TSharedPtr<FNiagaraMenuAction> SetLocalValueAction(
+			new FNiagaraMenuAction(FText::GetEmpty(), NameText, Tooltip, 0, FText::GetEmpty(),
+				FNiagaraMenuAction::FOnExecuteStackAction::CreateUObject(FunctionInput, &UNiagaraStackFunctionInput::DeleteInput),
+				FNiagaraMenuAction::FCanExecuteStackAction::CreateUObject(FunctionInput, &UNiagaraStackFunctionInput::CanDeleteInput)));
+		OutAllActions.AddAction(SetLocalValueAction);
+	}
+}
+
+void SNiagaraStackFunctionInputValue::SetToLocalValue()
+{
+	const UScriptStruct* LocalValueStruct = FunctionInput->GetInputType().GetScriptStruct();
+	if (LocalValueStruct != nullptr)
+	{
+		TSharedRef<FStructOnScope> LocalValue = MakeShared<FStructOnScope>(LocalValueStruct);
+		TArray<uint8> DefaultValueData;
+		FNiagaraEditorUtilities::GetTypeDefaultValue(FunctionInput->GetInputType(), DefaultValueData);
+		if (DefaultValueData.Num() == LocalValueStruct->GetStructureSize())
+		{
+			FMemory::Memcpy(LocalValue->GetStructMemory(), DefaultValueData.GetData(), DefaultValueData.Num());
+			FunctionInput->SetLocalValue(LocalValue);
+		}
+	}
+}
+
+void SNiagaraStackFunctionInputValue::DynamicInputScriptSelected(UNiagaraScript* DynamicInputScript)
+{
+	FunctionInput->SetDynamicInput(DynamicInputScript);
+}
+
+void SNiagaraStackFunctionInputValue::CustomExpressionSelected()
+{
+	FunctionInput->SetCustomExpression(TEXT("// Insert expression here"));
+}
+
+void SNiagaraStackFunctionInputValue::ParameterHandleSelected(FNiagaraParameterHandle Handle)
+{
+	FunctionInput->SetLinkedValueHandle(Handle);
+}
+
+EVisibility SNiagaraStackFunctionInputValue::GetResetButtonVisibility() const
+{
+	return FunctionInput->CanReset() ? EVisibility::Visible : EVisibility::Hidden;
+}
+
+FReply SNiagaraStackFunctionInputValue::ResetButtonPressed() const
+{
+	FunctionInput->Reset();
+	return FReply::Handled();
+}
+
+EVisibility SNiagaraStackFunctionInputValue::GetResetToBaseButtonVisibility() const
+{
+	if (FunctionInput->EmitterHasBase())
+	{
+		return FunctionInput->CanResetToBase() ? EVisibility::Visible : EVisibility::Hidden;
+	}
+	else
+	{
+		return EVisibility::Collapsed;
+	}
+}
+
+FReply SNiagaraStackFunctionInputValue::ResetToBaseButtonPressed() const
+{
+	FunctionInput->ResetToBase();
+	return FReply::Handled();
+}
+
+EVisibility SNiagaraStackFunctionInputValue::GetInputIconVisibility() const
+{
+	return FunctionInput->GetValueMode() == UNiagaraStackFunctionInput::EValueMode::Local
+		? EVisibility::Collapsed
+		: EVisibility::Visible;
+}
+
+FText SNiagaraStackFunctionInputValue::GetInputIconText() const
+{
+	switch (FunctionInput->GetValueMode())
+	{
+	case UNiagaraStackFunctionInput::EValueMode::Linked:
+		return FText::FromString(FString(TEXT("\xf0C1") /* fa-link */));
+	case UNiagaraStackFunctionInput::EValueMode::Data:
+		return FText::FromString(FString(TEXT("\xf1C0") /* fa-database */));
+	case UNiagaraStackFunctionInput::EValueMode::Dynamic:
+		return FText::FromString(FString(TEXT("\xf201") /* fa-line-chart */));
+	case UNiagaraStackFunctionInput::EValueMode::Expression:
+		return FText::FromString(FString(TEXT("\xf120") /* fa-terminal */));
+	case UNiagaraStackFunctionInput::EValueMode::Invalid:
+		return FunctionInput->CanReset()
+			? FText::FromString(FString(TEXT("\xf128") /* fa-question */))
+			: FText::FromString(FString(TEXT("\xf005") /* fa-star */));
+	default:
+		return FText::FromString(FString(TEXT("\xf128") /* fa-question */));
+	}
+}
+
+FText SNiagaraStackFunctionInputValue::GetInputIconToolTip() const
+{
+	static const FText InvalidText = LOCTEXT("InvalidInputIconToolTip", "Unsupported value.  Check the graph for issues.");
+	switch (FunctionInput->GetValueMode())
+	{
+	case UNiagaraStackFunctionInput::EValueMode::Linked:
+		return LOCTEXT("LinkInputIconToolTip", "Linked Value");
+	case UNiagaraStackFunctionInput::EValueMode::Data:
+		return LOCTEXT("DataInterfaceInputIconToolTip", "Data Value");
+	case UNiagaraStackFunctionInput::EValueMode::Dynamic:
+		return LOCTEXT("DynamicInputIconToolTip", "Dynamic Value");
+	case UNiagaraStackFunctionInput::EValueMode::Expression:
+		return LOCTEXT("ExpressionInputIconToolTip", "Custom Expression");
+	case UNiagaraStackFunctionInput::EValueMode::Invalid:
+		return FunctionInput->CanReset()
+			? InvalidText
+			: LOCTEXT("CustomInputIconToolTip", "Custom value");
+	default:
+		return InvalidText;
+	}
+}
+
+FSlateColor SNiagaraStackFunctionInputValue::GetInputIconColor() const
+{
+	switch (FunctionInput->GetValueMode())
+	{
+	case UNiagaraStackFunctionInput::EValueMode::Linked:
+		return FLinearColor(FColor::Purple);
+	case UNiagaraStackFunctionInput::EValueMode::Data:
+		return FLinearColor(FColor::Yellow);
+	case UNiagaraStackFunctionInput::EValueMode::Dynamic:
+		return FLinearColor(FColor::Cyan);
+	case UNiagaraStackFunctionInput::EValueMode::Expression:
+		return FLinearColor(FColor::Green);
+	case UNiagaraStackFunctionInput::EValueMode::Invalid:
+	default:
+		return FLinearColor(FColor::White);
+	}
+}
+
+FReply SNiagaraStackFunctionInputValue::OnFunctionInputDrop(TSharedPtr<FDragDropOperation> DragDropOperation)
+{
+	if (DragDropOperation->IsOfType<FNiagaraStackDragOperation>())
+	{
+		TSharedPtr<FNiagaraStackDragOperation> InputDragDropOperation = StaticCastSharedPtr<FNiagaraStackDragOperation>(DragDropOperation);
+		TSharedPtr<FNiagaraParameterAction> Action = StaticCastSharedPtr<FNiagaraParameterAction>(InputDragDropOperation->GetAction());
+		if (Action.IsValid())
+		{
+			FunctionInput->SetLinkedValueHandle(FNiagaraParameterHandle(Action->GetParameter().GetName()));
+			return FReply::Handled();
+		}
+	}
+
+	return FReply::Unhandled();
+}
+
+bool SNiagaraStackFunctionInputValue::OnFunctionInputAllowDrop(TSharedPtr<FDragDropOperation> DragDropOperation)
+{
+	if (FunctionInput && DragDropOperation->IsOfType<FNiagaraStackDragOperation>())
+	{
+		TSharedPtr<FNiagaraStackDragOperation> InputDragDropOperation = StaticCastSharedPtr<FNiagaraStackDragOperation>(DragDropOperation);
+		TSharedPtr<FNiagaraParameterAction> Action = StaticCastSharedPtr<FNiagaraParameterAction>(InputDragDropOperation->GetAction());
+		if (Action->GetParameter().GetType() == FunctionInput->GetInputType())
+		{
+			return true;
+		}
+	}
+
+	return false;
+}
+
+#undef LOCTEXT_NAMESPACE
diff --git a/Engine/Plugins/FX/Niagara/Source/NiagaraEditorWidgets/Private/Stack/SNiagaraStackFunctionInputValue.h b/Engine/Plugins/FX/Niagara/Source/NiagaraEditorWidgets/Private/Stack/SNiagaraStackFunctionInputValue.h
new file mode 100644
--- /dev/null
+++ b/Engine/Plugins/FX/Niagara/Source/NiagaraEditorWidgets/Private/Stack/SNiagaraStackFunctionInputValue.h
@@ -0,0 +1,142 @@
+// Copyright 1998-2018 Epic Games, Inc. All Rights Reserved.
+
+#pragma once
+
+#include "Widgets/SCompoundWidget.h"
+#include "Widgets/DeclarativeSyntaxSupport.h"
+#include "Widgets/Views/SExpanderArrow.h"
+#include "ViewModels/Stack/NiagaraStackFunctionInput.h"
+#include "SGraphActionMenu.h"
+
+class UNiagaraStackFunctionInput;
+class UNiagaraScript;
+class SNiagaraParameterEditor;
+class SBox;
+class IStructureDetailsView;
+class SComboButton;
+
+class SNiagaraStackFunctionInputValue: public SCompoundWidget
+{
+public:
+	DECLARE_DELEGATE_OneParam(FOnColumnWidthChanged, float)
+public:
+	SLATE_BEGIN_ARGS(SNiagaraStackFunctionInputValue) { }
+	SLATE_END_ARGS();
+
+	void Construct(const FArguments& InArgs, UNiagaraStackFunctionInput* InFunctionInput);
+
+private:
+
+	bool GetInputEnabled() const;
+
+	EVisibility GetValueWidgetVisibility(UNiagaraStackFunctionInput::EValueMode ValidMode) const;
+
+	TSharedRef<SWidget> ConstructLocalValueStructWidget();
+
+	void OnInputValueChanged();
+
+	void ParameterBeginValueChange();
+
+	void ParameterEndValueChange();
+
+	void ParameterValueChanged(TSharedRef<SNiagaraParameterEditor> ParameterEditor);
+
+	void ParameterPropertyValueChanged(const FPropertyChangedEvent& PropertyChangedEvent);
+
+	FText GetLinkedValueHandleText() const;
+
+	FText GetDataValueText() const;
+
+	FText GetDynamicValueText() const;
+
+	FText GetExpressionValueText() const;
+
+	void OnExpressionTextCommitted(const FText& Name, ETextCommit::Type CommitInfo);
+
+	FText GetInvalidValueText() const;
+
+	FText GetInvalidValueToolTipText() const;
+
+	FReply DynamicInputTextDoubleClicked();
+
+	class SNiagaraFunctionInputActionMenuExpander: public SExpanderArrow
+	{
+		SLATE_BEGIN_ARGS(SNiagaraFunctionInputActionMenuExpander) {}
+			SLATE_ATTRIBUTE(float, IndentAmount)
+		SLATE_END_ARGS()
+
+	public:
+		void Construct(const FArguments& InArgs, const FCustomExpanderData& ActionMenuData)
+		{
+			OwnerRowPtr = ActionMenuData.TableRow;
+			IndentAmount = InArgs._IndentAmount;
+			if (!ActionMenuData.RowAction.IsValid())
+			{
+				SExpanderArrow::FArguments SuperArgs;
+				SuperArgs._IndentAmount = InArgs._IndentAmount;
+
+				SExpanderArrow::Construct(SuperArgs, ActionMenuData.TableRow);
+			}
+			else
+			{
+				ChildSlot
+				.Padding(TAttribute<FMargin>(this, &SNiagaraFunctionInputActionMenuExpander::GetCustomIndentPadding))
+				[
+					SNew(SBox)
+				];
+			}
+		}
+
+	private:
+		FMargin GetCustomIndentPadding() const
+		{
+			return SExpanderArrow::GetExpanderPadding();
+		}
+	};
+
+	static TSharedRef<SExpanderArrow> CreateCustomNiagaraFunctionInputActionExpander(const FCustomExpanderData& ActionMenuData);
+
+	TSharedRef<SWidget> OnGetAvailableHandleMenu();
+
+	void OnActionSelected(const TArray<TSharedPtr<FEdGraphSchemaAction>>& SelectedActions, ESelectInfo::Type InSelectionType);
+
+	void CollectAllActions(FGraphActionListBuilderBase& OutAllActions);
+
+	void SetToLocalValue();
+
+	void DynamicInputScriptSelected(UNiagaraScript* DynamicInputScript);
+
+	void CustomExpressionSelected();
+
+	void ParameterHandleSelected(FNiagaraParameterHandle Handle);
+
+	EVisibility GetResetButtonVisibility() const;
+
+	FReply ResetButtonPressed() const;
+
+	EVisibility GetResetToBaseButtonVisibility() const;
+
+	FReply ResetToBaseButtonPressed() const;
+
+	EVisibility GetInputIconVisibility() const;
+
+	FText GetInputIconText() const;
+
+	FText GetInputIconToolTip() const;
+
+	FSlateColor GetInputIconColor() const;
+
+	FReply OnFunctionInputDrop(TSharedPtr<FDragDropOperation> DragDropOperation);
+
+	bool OnFunctionInputAllowDrop(TSharedPtr<FDragDropOperation> DragDropOperation);
+
+private:
+	UNiagaraStackFunctionInput* FunctionInput;
+
+	TSharedPtr<FStructOnScope> DisplayedLocalValueStruct;
+
+	TSharedPtr<SBox> LocalValueStructContainer;
+	TSharedPtr<SNiagaraParameterEditor> LocalValueStructParameterEditor;
+	TSharedPtr<IStructureDetailsView> LocalValueStructDetailsView;
+	TSharedPtr<SComboButton> SetFunctionInputButton;
+};
diff --git a/Engine/Plugins/FX/Niagara/Source/NiagaraEditorWidgets/Private/Stack/SNiagaraStackItemExpander.cpp b/Engine/Plugins/FX/Niagara/Source/NiagaraEditorWidgets/Private/Stack/SNiagaraStackItemExpander.cpp
new file mode 100644
--- /dev/null
+++ b/Engine/Plugins/FX/Niagara/Source/NiagaraEditorWidgets/Private/Stack/SNiagaraStackItemExpander.cpp
@@ -0,0 +1,60 @@
+// Copyright 1998-2018 Epic Games, Inc. All Rights Reserved.
+
+#include "Stack/SNiagaraStackItemExpander.h"
+#include "ViewModels/Stack/NiagaraStackAdvancedExpander.h"
+#include "EditorStyleSet.h"
+#include "Widgets/Input/SButton.h"
+#include "Widgets/Images/SImage.h"
+
+#define LOCTEXT_NAMESPACE "NiagaraStackItemExpander"
+
+void SNiagaraStackItemExpander::Construct(const FArguments& InArgs, UNiagaraStackAdvancedExpander& InItemExpander)
+{
+	ShowAdvancedExpander = &InItemExpander;
+	ExpandedToolTipText = LOCTEXT("HideAdvancedToolTip", "Hide Advanced");
+	CollapsedToolTipText = LOCTEXT("ShowAdvancedToolTip", "Show Advanced");
+
+	ChildSlot
+	[
+		SNew(SButton)
+		.ButtonStyle(FEditorStyle::Get(), "NoBorder")
+		.HAlign(HAlign_Center)
+		.ContentPadding(2)
+		.ToolTipText(this, &SNiagaraStackItemExpander::GetToolTipText)
+		.OnClicked(this, &SNiagaraStackItemExpander::ExpandButtonClicked)
+		.Content()
+		[
+			SNew(SImage)
+			.Image(this, &SNiagaraStackItemExpander::GetButtonBrush)
+		]
+	];
+}
+
+const FSlateBrush* SNiagaraStackItemExpander::GetButtonBrush() const
+{
+	if (IsHovered())
+	{
+		return ShowAdvancedExpander->GetShowAdvanced()
+			? FEditorStyle::GetBrush("DetailsView.PulldownArrow.Up.Hovered")
+			: FEditorStyle::GetBrush("DetailsView.PulldownArrow.Down.Hovered");
+	}
+	else
+	{
+		return ShowAdvancedExpander->GetShowAdvanced()
+			? FEditorStyle::GetBrush("DetailsView.PulldownArrow.Up")
+			: FEditorStyle::GetBrush("DetailsView.PulldownArrow.Down");
+	}
+}
+
+FText SNiagaraStackItemExpander::GetToolTipText() const
+{
+	return ShowAdvancedExpander->GetShowAdvanced() ? ExpandedToolTipText : CollapsedToolTipText;
+}
+
+FReply SNiagaraStackItemExpander::ExpandButtonClicked()
+{
+	ShowAdvancedExpander->ToggleShowAdvanced();
+	return FReply::Handled();
+}
+
+#undef LOCTEXT_NAMESPACE
\ No newline at end of file
diff --git a/Engine/Plugins/FX/Niagara/Source/NiagaraEditorWidgets/Private/Stack/SNiagaraStackItemExpander.h b/Engine/Plugins/FX/Niagara/Source/NiagaraEditorWidgets/Private/Stack/SNiagaraStackItemExpander.h
new file mode 100644
--- /dev/null
+++ b/Engine/Plugins/FX/Niagara/Source/NiagaraEditorWidgets/Private/Stack/SNiagaraStackItemExpander.h
@@ -0,0 +1,29 @@
+// Copyright 1998-2018 Epic Games, Inc. All Rights Reserved.
+
+#pragma once
+
+#include "Widgets/SCompoundWidget.h"
+#include "Widgets/DeclarativeSyntaxSupport.h"
+
+class UNiagaraStackAdvancedExpander;
+
+class SNiagaraStackItemExpander: public SCompoundWidget
+{
+public:
+	SLATE_BEGIN_ARGS(SNiagaraStackItemExpander) { }
+	SLATE_END_ARGS();
+
+	void Construct(const FArguments& InArgs, UNiagaraStackAdvancedExpander& InItemExpander);
+
+private:
+	const struct FSlateBrush* GetButtonBrush() const;
+
+	FText GetToolTipText() const;
+
+	FReply ExpandButtonClicked();
+
+private:
+	UNiagaraStackAdvancedExpander* ShowAdvancedExpander;
+	FText ExpandedToolTipText;
+	FText CollapsedToolTipText;
+};
\ No newline at end of file
diff --git a/Engine/Plugins/FX/Niagara/Source/NiagaraEditorWidgets/Private/Stack/SNiagaraStackItemGroup.cpp b/Engine/Plugins/FX/Niagara/Source/NiagaraEditorWidgets/Private/Stack/SNiagaraStackItemGroup.cpp
new file mode 100644
--- /dev/null
+++ b/Engine/Plugins/FX/Niagara/Source/NiagaraEditorWidgets/Private/Stack/SNiagaraStackItemGroup.cpp
@@ -0,0 +1,207 @@
+// Copyright 1998-2018 Epic Games, Inc. All Rights Reserved.
+
+#include "Stack/SNiagaraStackItemGroup.h"
+#include "NiagaraEditorWidgetsStyle.h"
+#include "NiagaraEditorStyle.h"
+#include "EditorStyleSet.h"
+#include "Stack/SNiagaraStackItemGroupAddMenu.h"
+#include "ViewModels/Stack/NiagaraStackItemGroup.h"
+#include "ViewModels/Stack/NiagaraStackViewModel.h"
+#include "Widgets/Layout/SBox.h"
+#include "Widgets/SBoxPanel.h"
+#include "Widgets/Input/SButton.h"
+#include "Widgets/Input/SComboButton.h"
+#include "Widgets/Input/SEditableTextBox.h"
+#include "Widgets/Text/STextBlock.h"
+#include "Widgets/Images/SImage.h"
+#include "SNiagaraStackErrorButton.h"
+
+#define LOCTEXT_NAMESPACE "NiagaraStackItemGroup"
+
+
+void SNiagaraStackItemGroup::Construct(const FArguments& InArgs, UNiagaraStackItemGroup& InGroup, UNiagaraStackViewModel* InStackViewModel)
+{
+	Group = &InGroup;
+	StackEntryItem = Group;
+	StackViewModel = InStackViewModel;
+
+	ChildSlot
+	[
+		// Name
+		SNew(SHorizontalBox)
+		+ SHorizontalBox::Slot()
+		.VAlign(VAlign_Center)
+		[
+			SNew(STextBlock)
+			.TextStyle(FNiagaraEditorWidgetsStyle::Get(), "NiagaraEditor.Stack.GroupText")
+			.ToolTipText_UObject(Group, &UNiagaraStackEntry::GetTooltipText)
+			.Text_UObject(Group, &UNiagaraStackEntry::GetDisplayName)
+			.HighlightText_UObject(InStackViewModel, &UNiagaraStackViewModel::GetCurrentSearchText)
+			.ColorAndOpacity(this, &SNiagaraStackItemGroup::GetTextColorForSearch)
+		]
+		// Stack issues icon
+		+ SHorizontalBox::Slot()
+		.Padding(2, 0, 0, 0)
+		.AutoWidth()
+		.VAlign(VAlign_Center)
+		.HAlign(HAlign_Right)
+		[
+			SNew(SNiagaraStackErrorButton)
+			.IssueSeverity_UObject(Group, &UNiagaraStackItemGroup::GetHighestStackIssueSeverity)
+			.ErrorTooltip(this, &SNiagaraStackItemGroup::GetErrorButtonTooltipText)
+			.Visibility(this, &SNiagaraStackItemGroup::GetStackIssuesWarningVisibility)
+			.OnButtonClicked(this, &SNiagaraStackItemGroup::ExpandEntry)
+		]
+		// Delete group button
+		+ SHorizontalBox::Slot()
+		.AutoWidth()
+		[
+			SNew(SButton)
+			.Visibility(this, &SNiagaraStackItemGroup::GetDeleteButtonVisibility)
+			.ButtonStyle(FEditorStyle::Get(), "HoverHintOnly")
+			.IsFocusable(false)
+			.ForegroundColor(FNiagaraEditorWidgetsStyle::Get().GetColor("NiagaraEditor.Stack.ForegroundColor"))
+			.ToolTipText(LOCTEXT("DeleteGroupToolTip", "Delete this group"))
+			.OnClicked(this, &SNiagaraStackItemGroup::DeleteClicked)
+			.Content()
+			[
+				SNew(STextBlock)
+				.Font(FEditorStyle::Get().GetFontStyle("FontAwesome.10"))
+				.Text(FText::FromString(FString(TEXT("\xf1f8"))))
+			]
+		]
+		// Add button
+		+ SHorizontalBox::Slot()
+		.AutoWidth()
+		.HAlign(HAlign_Right)
+		[
+			ConstructAddButton()
+		]
+	];
+}
+
+TSharedRef<SWidget> SNiagaraStackItemGroup::ConstructAddButton()
+{
+	AddActionButton.Reset();
+	INiagaraStackItemGroupAddUtilities* AddUtilities = Group->GetAddUtilities();
+	if (AddUtilities != nullptr)
+	{
+		if (AddUtilities->GetAddMode() == INiagaraStackItemGroupAddUtilities::AddFromAction)
+		{
+			return SAssignNew(AddActionButton, SComboButton)
+				.Visibility(this, &SNiagaraStackItemGroup::GetAddButtonVisibility)
+				.ButtonStyle(FNiagaraEditorWidgetsStyle::Get(), "NiagaraEditor.Stack.AddButton")
+				.ToolTipText(this, &SNiagaraStackItemGroup::GetAddButtonToolTipText)
+				.HasDownArrow(false)
+				.OnGetMenuContent(this, &SNiagaraStackItemGroup::GetAddMenu)
+				.ContentPadding(0)
+				.MenuPlacement(MenuPlacement_BelowRightAnchor)
+				.ButtonContent()
+				[
+					SNew(SBox)
+					.WidthOverride(TextIconSize * 2)
+					.HeightOverride(TextIconSize)
+					.HAlign(HAlign_Center)
+					.VAlign(VAlign_Center)
+					[
+						SNew(STextBlock)
+						.TextStyle(FEditorStyle::Get(), "NormalText.Important")
+						.Font(FEditorStyle::Get().GetFontStyle("FontAwesome.10"))
+						.Text(FText::FromString(FString(TEXT("\xf067"))) /*fa-plus*/)
+					]
+				];
+		}
+		else if (AddUtilities->GetAddMode() == INiagaraStackItemGroupAddUtilities::AddDirectly)
+		{
+			return SNew(SButton)
+				.Visibility(this, &SNiagaraStackItemGroup::GetAddButtonVisibility)
+				.ButtonStyle(FNiagaraEditorWidgetsStyle::Get(), "NiagaraEditor.Stack.AddButton")
+				.ToolTipText(this, &SNiagaraStackItemGroup::GetAddButtonToolTipText)
+				.ContentPadding(0)
+				.OnClicked(this, &SNiagaraStackItemGroup::AddDirectlyButtonClicked)
+				.Content()
+				[
+					SNew(SBox)
+					.WidthOverride(TextIconSize * 2)
+					.HeightOverride(TextIconSize)
+					.HAlign(HAlign_Center)
+					.VAlign(VAlign_Center)
+					[
+						SNew(STextBlock)
+						.TextStyle(FEditorStyle::Get(), "NormalText.Important")
+						.Font(FEditorStyle::Get().GetFontStyle("FontAwesome.10"))
+						.Text(FText::FromString(FString(TEXT("\xf067"))) /*fa-plus*/)
+					]
+				];
+		}
+	}
+	return SNullWidget::NullWidget;
+}
+
+EVisibility SNiagaraStackItemGroup::GetDeleteButtonVisibility() const
+{
+	if (Group->CanDelete())
+	{
+		return EVisibility::Visible;
+	}
+	else
+	{
+		return EVisibility::Collapsed;
+	}
+}
+
+EVisibility SNiagaraStackItemGroup::GetAddButtonVisibility() const
+{
+	if (Group->GetAddUtilities() != nullptr)
+	{
+		return EVisibility::Visible;
+	}
+	else
+	{
+		return EVisibility::Collapsed;
+	}
+}
+
+FText SNiagaraStackItemGroup::GetAddButtonToolTipText() const
+{
+	if (Group->GetAddUtilities() != nullptr)
+	{
+		return FText::Format(LOCTEXT("AddToGroupFormat", "Add a new {0} to this group."), Group->GetAddUtilities()->GetAddItemName());
+	}
+	else
+	{
+		return FText();
+	}
+}
+
+FReply SNiagaraStackItemGroup::AddDirectlyButtonClicked()
+{
+	Group->GetAddUtilities()->AddItemDirectly();
+	return FReply::Handled();
+}
+
+FReply SNiagaraStackItemGroup::DeleteClicked()
+{
+	Group->Delete();
+	return FReply::Handled();
+}
+
+TSharedRef<SWidget> SNiagaraStackItemGroup::GetAddMenu()
+{
+	TSharedRef<SNiagaraStackItemGroupAddMenu> AddMenu = SNew(SNiagaraStackItemGroupAddMenu, Group->GetAddUtilities(), INDEX_NONE);
+	AddActionButton->SetMenuContentWidgetToFocus(AddMenu->GetFilterTextBox()->AsShared());
+	return AddMenu;
+}
+
+EVisibility SNiagaraStackItemGroup::GetStackIssuesWarningVisibility() const
+{
+	return  Group->GetRecursiveStackIssuesCount() > 0 ? EVisibility::Visible : EVisibility::Collapsed;
+}
+
+FText SNiagaraStackItemGroup::GetErrorButtonTooltipText() const
+{
+	return FText::Format(LOCTEXT("GroupIssuesTooltip", "This group contains items that have a total of {0} issues, click to expand."), Group->GetRecursiveStackIssuesCount());
+}
+
+
+#undef LOCTEXT_NAMESPACE
\ No newline at end of file
diff --git a/Engine/Plugins/FX/Niagara/Source/NiagaraEditorWidgets/Private/Stack/SNiagaraStackItemGroup.h b/Engine/Plugins/FX/Niagara/Source/NiagaraEditorWidgets/Private/Stack/SNiagaraStackItemGroup.h
new file mode 100644
--- /dev/null
+++ b/Engine/Plugins/FX/Niagara/Source/NiagaraEditorWidgets/Private/Stack/SNiagaraStackItemGroup.h
@@ -0,0 +1,42 @@
+// Copyright 1998-2018 Epic Games, Inc. All Rights Reserved.
+
+#pragma once
+
+#include "SNiagaraStackEntryWidget.h"
+#include "Layout/Visibility.h"
+
+class UNiagaraStackItemGroup;
+class UNiagaraStackViewModel;
+
+class SNiagaraStackItemGroup : public SNiagaraStackEntryWidget
+{
+public:
+	SLATE_BEGIN_ARGS(SNiagaraStackItemGroup) { }
+	SLATE_END_ARGS();
+
+	void Construct(const FArguments& InArgs, UNiagaraStackItemGroup& InGroup, UNiagaraStackViewModel* InStackViewModel);
+
+private:
+	TSharedRef<SWidget> ConstructAddButton();
+
+	EVisibility GetDeleteButtonVisibility() const;
+
+	EVisibility GetAddButtonVisibility() const;
+
+	FText GetAddButtonToolTipText() const;
+
+	FReply AddDirectlyButtonClicked();
+
+	FReply DeleteClicked();
+
+	TSharedRef<SWidget> GetAddMenu();
+
+	EVisibility GetStackIssuesWarningVisibility() const;
+
+	FText GetErrorButtonTooltipText() const;
+
+private:
+	UNiagaraStackItemGroup* Group;
+	TSharedPtr<class SComboButton> AddActionButton; 
+	const float TextIconSize = 16;
+};
diff --git a/Engine/Plugins/FX/Niagara/Source/NiagaraEditorWidgets/Private/Stack/SNiagaraStackItemGroupAddMenu.cpp b/Engine/Plugins/FX/Niagara/Source/NiagaraEditorWidgets/Private/Stack/SNiagaraStackItemGroupAddMenu.cpp
new file mode 100644
--- /dev/null
+++ b/Engine/Plugins/FX/Niagara/Source/NiagaraEditorWidgets/Private/Stack/SNiagaraStackItemGroupAddMenu.cpp
@@ -0,0 +1,74 @@
+// Copyright 1998-2018 Epic Games, Inc. All Rights Reserved.
+#include "Stack/SNiagaraStackItemGroupAddMenu.h"
+#include "EditorStyleSet.h"
+#include "NiagaraActions.h"
+#include "ViewModels/Stack/INiagaraStackItemGroupAddUtilities.h"
+#include "Widgets/Input/SEditableTextBox.h"
+#include "EdGraph/EdGraph.h"
+#include "SGraphActionMenu.h"
+#include "Framework/Application/SlateApplication.h"
+
+void SNiagaraStackItemGroupAddMenu::Construct(const FArguments& InArgs, INiagaraStackItemGroupAddUtilities* InAddUtilities, int32 InInsertIndex)
+{
+	AddUtilities = InAddUtilities;
+	InsertIndex = InInsertIndex;
+	ChildSlot
+	[
+		SNew(SBorder)
+		.BorderImage(FEditorStyle::GetBrush("Menu.Background"))
+		.Padding(5)
+		[
+			SNew(SBox)
+			.WidthOverride(300)
+			.HeightOverride(400)
+			[
+				SAssignNew(AddMenu, SGraphActionMenu)
+				.OnActionSelected(this, &SNiagaraStackItemGroupAddMenu::OnActionSelected)
+				.OnCollectAllActions(this, &SNiagaraStackItemGroupAddMenu::CollectAllAddActions)
+				.AutoExpandActionMenu(AddUtilities->GetAutoExpandAddActions())
+				.ShowFilterTextBox(true)
+			]
+		]
+	];
+}
+
+TSharedPtr<SEditableTextBox> SNiagaraStackItemGroupAddMenu::GetFilterTextBox()
+{
+	return AddMenu->GetFilterTextBox();
+}
+
+void SNiagaraStackItemGroupAddMenu::CollectAllAddActions(FGraphActionListBuilderBase& OutAllActions)
+{
+	if (OutAllActions.OwnerOfTemporaries == nullptr)
+	{
+		OutAllActions.OwnerOfTemporaries = NewObject<UEdGraph>((UObject*)GetTransientPackage());
+	}
+
+	TArray<TSharedRef<INiagaraStackItemGroupAddAction>> AddActions;
+	AddUtilities->GenerateAddActions(AddActions);
+
+	for (TSharedRef<INiagaraStackItemGroupAddAction> AddAction : AddActions)
+	{
+		TSharedPtr<FNiagaraMenuAction> NewNodeAction(
+			new FNiagaraMenuAction(AddAction->GetCategory(), AddAction->GetDisplayName(), AddAction->GetDescription(), 0, FText(),
+				FNiagaraMenuAction::FOnExecuteStackAction::CreateRaw(AddUtilities, &INiagaraStackItemGroupAddUtilities::ExecuteAddAction, AddAction, InsertIndex)));
+		OutAllActions.AddAction(NewNodeAction);
+	}
+}
+
+void SNiagaraStackItemGroupAddMenu::OnActionSelected(const TArray< TSharedPtr<FEdGraphSchemaAction> >& SelectedActions, ESelectInfo::Type InSelectionType)
+{
+	if (InSelectionType == ESelectInfo::OnMouseClick || InSelectionType == ESelectInfo::OnKeyPress || SelectedActions.Num() == 0)
+	{
+		for (int32 ActionIndex = 0; ActionIndex < SelectedActions.Num(); ActionIndex++)
+		{
+			TSharedPtr<FNiagaraMenuAction> CurrentAction = StaticCastSharedPtr<FNiagaraMenuAction>(SelectedActions[ActionIndex]);
+
+			if (CurrentAction.IsValid())
+			{
+				FSlateApplication::Get().DismissAllMenus();
+				CurrentAction->ExecuteAction();
+			}
+		}
+	}
+}
diff --git a/Engine/Plugins/FX/Niagara/Source/NiagaraEditorWidgets/Private/Stack/SNiagaraStackItemGroupAddMenu.h b/Engine/Plugins/FX/Niagara/Source/NiagaraEditorWidgets/Private/Stack/SNiagaraStackItemGroupAddMenu.h
new file mode 100644
--- /dev/null
+++ b/Engine/Plugins/FX/Niagara/Source/NiagaraEditorWidgets/Private/Stack/SNiagaraStackItemGroupAddMenu.h
@@ -0,0 +1,32 @@
+// Copyright 1998-2018 Epic Games, Inc. All Rights Reserved.
+
+#pragma once
+
+#include "Widgets/SCompoundWidget.h"
+#include "Widgets/DeclarativeSyntaxSupport.h"
+#include "EdGraph/EdGraphSchema.h"
+
+class INiagaraStackItemGroupAddUtilities;
+
+class SNiagaraStackItemGroupAddMenu : public SCompoundWidget
+{
+public:
+	SLATE_BEGIN_ARGS(SNiagaraStackItemGroupAddMenu) { }
+	SLATE_END_ARGS();
+
+	void Construct(const FArguments& InArgs, INiagaraStackItemGroupAddUtilities* InAddUtilities, int32 InInsertIndex);
+
+	TSharedPtr<class SEditableTextBox> GetFilterTextBox();
+
+private:
+	void CollectAllAddActions(FGraphActionListBuilderBase& OutAllActions);
+
+	void OnActionSelected(const TArray< TSharedPtr<FEdGraphSchemaAction> >& SelectedActions, ESelectInfo::Type InSelectionType);
+
+private:
+	INiagaraStackItemGroupAddUtilities* AddUtilities;
+
+	int32 InsertIndex;
+
+	TSharedPtr<class SGraphActionMenu> AddMenu;
+};
\ No newline at end of file
diff --git a/Engine/Plugins/FX/Niagara/Source/NiagaraEditorWidgets/Private/Stack/SNiagaraStackModuleItem.cpp b/Engine/Plugins/FX/Niagara/Source/NiagaraEditorWidgets/Private/Stack/SNiagaraStackModuleItem.cpp
new file mode 100644
--- /dev/null
+++ b/Engine/Plugins/FX/Niagara/Source/NiagaraEditorWidgets/Private/Stack/SNiagaraStackModuleItem.cpp
@@ -0,0 +1,303 @@
+// Copyright 1998-2018 Epic Games, Inc. All Rights Reserved.
+
+#include "Stack/SNiagaraStackModuleItem.h"
+#include "NiagaraEditorWidgetsStyle.h"
+#include "NiagaraEditorStyle.h"
+#include "NiagaraActions.h"
+#include "NiagaraConstants.h"
+#include "NiagaraNodeFunctionCall.h"
+#include "NiagaraNodeOutput.h"
+#include "NiagaraNodeAssignment.h"
+#include "Stack/SNiagaraStackItemGroupAddMenu.h"
+#include "ViewModels/Stack/NiagaraStackModuleItem.h"
+#include "ViewModels/Stack/NiagaraStackViewModel.h"
+#include "Widgets/Layout/SBox.h"
+#include "Widgets/SBoxPanel.h"
+#include "Widgets/Input/SButton.h"
+#include "Widgets/Input/SCheckBox.h"
+#include "Widgets/Input/SComboButton.h"
+#include "Widgets/Images/SImage.h"
+#include "Framework/MultiBox/MultiBoxBuilder.h"
+#include "Framework/Application/SlateApplication.h"
+#include "Toolkits/AssetEditorManager.h"
+#include "SDropTarget.h"
+#include "SNiagaraStackErrorButton.h"
+
+#define LOCTEXT_NAMESPACE "NiagaraStackModuleItem"
+
+void SNiagaraStackModuleItem::Construct(const FArguments& InArgs, UNiagaraStackModuleItem& InModuleItem, UNiagaraStackViewModel* InStackViewModel)
+{
+	ModuleItem = &InModuleItem;
+	StackEntryItem = ModuleItem;
+	StackViewModel = InStackViewModel;
+
+	ChildSlot
+	[
+		SNew(SDropTarget)
+		.OnAllowDrop(this, &SNiagaraStackModuleItem::OnModuleItemAllowDrop)
+		.OnDrop(this, &SNiagaraStackModuleItem::OnModuleItemDrop)
+		.Content()
+		[
+			SNew(SHorizontalBox)
+			// Name
+			+ SHorizontalBox::Slot()
+			.VAlign(VAlign_Center)
+			.Padding(0)
+			[
+				SNew(STextBlock)
+				.TextStyle(FNiagaraEditorWidgetsStyle::Get(), "NiagaraEditor.Stack.ItemText")
+				.ToolTipText_UObject(ModuleItem, &UNiagaraStackEntry::GetTooltipText)
+				.Text_UObject(ModuleItem, &UNiagaraStackEntry::GetDisplayName)
+				.HighlightText_UObject(InStackViewModel, &UNiagaraStackViewModel::GetCurrentSearchText)
+				.ColorAndOpacity(this, &SNiagaraStackModuleItem::GetTextColorForSearch)
+				.IsEnabled_UObject(ModuleItem, &UNiagaraStackModuleItem::GetIsEnabled)
+			]
+			// Stack issues icon
+			+ SHorizontalBox::Slot()
+			.Padding(2, 0, 0, 0)
+			.AutoWidth()
+			.VAlign(VAlign_Center)
+			[
+				SNew(SNiagaraStackErrorButton)
+				.IssueSeverity_UObject(ModuleItem, &UNiagaraStackModuleItem::GetHighestStackIssueSeverity)
+				.ErrorTooltip(this, &SNiagaraStackModuleItem::GetErrorButtonTooltipText)
+				.Visibility(this, &SNiagaraStackModuleItem::GetStackIssuesWarningVisibility)
+				.OnButtonClicked(this, &SNiagaraStackModuleItem::ExpandEntry)
+			]
+			// Raise Action Menu button
+			+ SHorizontalBox::Slot()
+			.AutoWidth()
+			[
+				SNew(SComboButton)
+				.HasDownArrow(true)
+				.ButtonStyle(FEditorStyle::Get(), "HoverHintOnly")
+				.ForegroundColor(FSlateColor::UseForeground())
+				.OnGetMenuContent(this, &SNiagaraStackModuleItem::RaiseActionMenuClicked)
+				.ContentPadding(FMargin(2))
+				.HAlign(HAlign_Center)
+				.VAlign(VAlign_Center)
+				.Visibility(this, &SNiagaraStackModuleItem::GetRaiseActionMenuVisibility)
+			]
+			// Refresh button
+			+ SHorizontalBox::Slot()
+			.VAlign(VAlign_Center)
+			.AutoWidth()
+			[
+				SNew(SButton)
+				.ButtonStyle(FEditorStyle::Get(), "HoverHintOnly")
+				.IsFocusable(false)
+				.ForegroundColor(FNiagaraEditorWidgetsStyle::Get().GetColor("NiagaraEditor.Stack.FlatButtonColor"))
+				.ToolTipText(LOCTEXT("RefreshTooltip", "Refresh this module"))
+				.Visibility(this, &SNiagaraStackModuleItem::GetRefreshVisibility)
+				.OnClicked(this, &SNiagaraStackModuleItem::RefreshClicked)
+				.Content()
+				[
+					SNew(STextBlock)
+					.Font(FEditorStyle::Get().GetFontStyle("FontAwesome.10"))
+					.Text(FText::FromString(FString(TEXT("\xf021"))))
+				]
+			]
+			// Delete button
+			+ SHorizontalBox::Slot()
+			.VAlign(VAlign_Center)
+			.AutoWidth()
+			[
+				SNew(SButton)
+				.ButtonStyle(FEditorStyle::Get(), "HoverHintOnly")
+				.IsFocusable(false)
+				.ForegroundColor(FNiagaraEditorWidgetsStyle::Get().GetColor("NiagaraEditor.Stack.FlatButtonColor"))
+				.ToolTipText(LOCTEXT("DeleteToolTip", "Delete this module"))
+				.Visibility(this, &SNiagaraStackModuleItem::GetEditButtonVisibility)
+				.OnClicked(this, &SNiagaraStackModuleItem::DeleteClicked)
+				.Content()
+				[
+					SNew(STextBlock)
+					.Font(FEditorStyle::Get().GetFontStyle("FontAwesome.10"))
+					.Text(FText::FromString(FString(TEXT("\xf1f8"))))
+				]
+			]
+			// Enabled checkbox
+			+ SHorizontalBox::Slot()
+			.Padding(2, 0, 0, 0)
+			.AutoWidth()
+			[
+				SNew(SCheckBox)
+				.IsChecked(this, &SNiagaraStackModuleItem::GetCheckState)
+				.OnCheckStateChanged(this, &SNiagaraStackModuleItem::OnCheckStateChanged)
+			]
+		]
+	];
+}
+
+void SNiagaraStackModuleItem::SetEnabled(bool bInIsEnabled)
+{
+	ModuleItem->SetIsEnabled(bInIsEnabled);
+}
+
+bool SNiagaraStackModuleItem::CheckEnabledStatus(bool bIsEnabled)
+{
+	return ModuleItem->GetIsEnabled() == bIsEnabled;
+}
+
+void SNiagaraStackModuleItem::FillRowContextMenu(FMenuBuilder& MenuBuilder)
+{
+	MenuBuilder.AddMenuEntry(
+		LOCTEXT("InsertModuleAbove", "Insert above"),
+		LOCTEXT("InsertModuleAboveToolTip", "Insert a new module above this module in the stack."),
+		FSlateIcon(),
+		FUIAction(FExecuteAction::CreateSP(this, &SNiagaraStackModuleItem::InsertModuleAbove)));
+	MenuBuilder.AddMenuEntry(
+		LOCTEXT("InsertModuleBelow", "Insert Below"),
+		LOCTEXT("InsertModuleBelowToolTip", "Insert a new module below this module in the stack."),
+		FSlateIcon(),
+		FUIAction(FExecuteAction::CreateSP(this, &SNiagaraStackModuleItem::InsertModuleBelow)));
+
+	FUIAction Action(FExecuteAction::CreateSP(this, &SNiagaraStackModuleItem::SetEnabled, !ModuleItem->GetIsEnabled()),
+		FCanExecuteAction(),
+		FIsActionChecked::CreateSP(this, &SNiagaraStackModuleItem::CheckEnabledStatus, true));
+	MenuBuilder.AddMenuEntry(
+		LOCTEXT("IsEnabled", "Is Enabled"),
+		LOCTEXT("ToggleModuleEnabledToolTip", "Toggle module enabled/disabled state"),
+		FSlateIcon(),
+		Action,
+		NAME_None,
+		EUserInterfaceActionType::Check);
+}
+
+FReply SNiagaraStackModuleItem::OnMouseButtonDoubleClick(const FGeometry& InMyGeometry, const FPointerEvent& InMouseEvent)
+{
+	const UNiagaraNodeFunctionCall& ModuleFunctionCall = ModuleItem->GetModuleNode();
+	if (ModuleFunctionCall.FunctionScript != nullptr && ModuleFunctionCall.FunctionScript->IsAsset())
+	{
+		FAssetEditorManager::Get().OpenEditorForAsset(const_cast<UNiagaraScript*>(ModuleFunctionCall.FunctionScript));
+		return FReply::Handled();
+	}
+	return FReply::Unhandled();
+}
+
+ECheckBoxState SNiagaraStackModuleItem::GetCheckState() const
+{
+	return ModuleItem->GetIsEnabled() ? ECheckBoxState::Checked : ECheckBoxState::Unchecked;
+}
+
+void SNiagaraStackModuleItem::OnCheckStateChanged(ECheckBoxState InCheckState)
+{
+	ModuleItem->SetIsEnabled(InCheckState == ECheckBoxState::Checked);
+}
+
+EVisibility SNiagaraStackModuleItem::GetEditButtonVisibility() const
+{
+	return ModuleItem->CanMoveAndDelete() ? EVisibility::Visible : EVisibility::Collapsed;
+}
+
+EVisibility SNiagaraStackModuleItem::GetRaiseActionMenuVisibility() const
+{
+	return CanRaiseActionMenu() ? EVisibility::Visible : EVisibility::Collapsed;
+}
+
+EVisibility SNiagaraStackModuleItem::GetRefreshVisibility() const
+{
+	return ModuleItem->CanRefresh() ? EVisibility::Visible : EVisibility::Collapsed;
+}
+
+FReply SNiagaraStackModuleItem::DeleteClicked()
+{
+	ModuleItem->Delete();
+	return FReply::Handled();
+}
+
+TSharedRef<SWidget> SNiagaraStackModuleItem::RaiseActionMenuClicked()
+{
+	if (CanRaiseActionMenu())
+	{
+		UNiagaraNodeAssignment* AssignmentNode = Cast<UNiagaraNodeAssignment>(&ModuleItem->GetModuleNode());
+		if (AssignmentNode != nullptr)
+		{
+			UNiagaraNodeOutput* OutputNode = ModuleItem->GetOutputNode();
+			if (OutputNode)
+			{
+				FMenuBuilder MenuBuilder(true, nullptr);
+				/*MenuBuilder.AddMenuEntry(
+					LOCTEXT("MergeLabel", "Merge Up"),
+					LOCTEXT("MergeToolTip", "If a Set Variables node precedes this one in the stack, merge this node (and all variable binding logic) into that stack."),
+					FSlateIcon(),
+					FUIAction(FExecuteAction::CreateUObject(AssignmentNode, &UNiagaraNodeAssignment::MergeUp)));*/
+				MenuBuilder.AddSubMenu(LOCTEXT("AddVariables", "Add Variable"),
+					LOCTEXT("AddVariablesTooltip", "Add another variable to the end of the list"),
+					FNewMenuDelegate::CreateLambda([OutputNode, AssignmentNode](FMenuBuilder& SubMenuBuilder)
+				{
+					AssignmentNode->BuildParameterMenu(SubMenuBuilder, OutputNode->GetUsage(), OutputNode);
+				}));
+
+				return MenuBuilder.MakeWidget();
+			}
+		}
+	}
+
+	return SNullWidget::NullWidget;
+}
+
+bool SNiagaraStackModuleItem::CanRaiseActionMenu() const
+{
+	return Cast<UNiagaraNodeAssignment>(&ModuleItem->GetModuleNode()) != nullptr;
+}
+
+FReply SNiagaraStackModuleItem::RefreshClicked()
+{
+	ModuleItem->Refresh();
+	return FReply::Handled();
+}
+
+void SNiagaraStackModuleItem::InsertModuleAbove()
+{
+	ShowInsertModuleMenu(ModuleItem->GetModuleIndex());
+}
+
+void SNiagaraStackModuleItem::InsertModuleBelow()
+{
+	ShowInsertModuleMenu(ModuleItem->GetModuleIndex() + 1);
+}
+
+void SNiagaraStackModuleItem::ShowInsertModuleMenu(int32 InsertIndex)
+{
+	TSharedRef<SWidget> MenuContent = SNew(SNiagaraStackItemGroupAddMenu, ModuleItem->GetGroupAddUtilities(), InsertIndex);
+	FGeometry ThisGeometry = GetCachedGeometry();
+	FVector2D MenuPosition = FSlateApplication::Get().CalculatePopupWindowPosition(ThisGeometry.GetLayoutBoundingRect(), MenuContent->GetDesiredSize());
+	FSlateApplication::Get().PushMenu(AsShared(), FWidgetPath(), MenuContent, MenuPosition, FPopupTransitionEffect::ContextMenu);
+}
+
+FReply SNiagaraStackModuleItem::OnModuleItemDrop(TSharedPtr<FDragDropOperation> DragDropOperation)
+{
+	if (DragDropOperation->IsOfType<FNiagaraStackDragOperation>())
+	{
+		if (UNiagaraNodeAssignment* NodeAssignment = Cast<UNiagaraNodeAssignment>(&ModuleItem->GetModuleNode()))
+		{
+			TSharedPtr<FNiagaraStackDragOperation> InputDragDropOperation = StaticCastSharedPtr<FNiagaraStackDragOperation>(DragDropOperation);
+			TSharedPtr<FNiagaraParameterAction> Action = StaticCastSharedPtr<FNiagaraParameterAction>(InputDragDropOperation->GetAction());
+			if (Action.IsValid())
+			{
+				NodeAssignment->AddParameter(Action->GetParameter(), FNiagaraConstants::GetAttributeDefaultValue(Action->GetParameter()));
+				return FReply::Handled();
+			}
+		}
+	}
+
+	return FReply::Unhandled();
+}
+
+bool SNiagaraStackModuleItem::OnModuleItemAllowDrop(TSharedPtr<FDragDropOperation> DragDropOperation)
+{
+	return ModuleItem != nullptr && Cast<UNiagaraNodeAssignment>(&ModuleItem->GetModuleNode()) && DragDropOperation->IsOfType<FNiagaraStackDragOperation>();
+}
+
+EVisibility SNiagaraStackModuleItem::GetStackIssuesWarningVisibility() const
+{
+	return ModuleItem->GetRecursiveStackIssuesCount() > 0 ? EVisibility::Visible : EVisibility::Collapsed;
+}
+
+FText SNiagaraStackModuleItem::GetErrorButtonTooltipText() const
+{
+	return FText::Format(LOCTEXT("ModuleIssuesTooltip", "This module has {0} issues, click to expand."), ModuleItem->GetRecursiveStackIssuesCount());
+}
+
+#undef LOCTEXT_NAMESPACE
\ No newline at end of file
diff --git a/Engine/Plugins/FX/Niagara/Source/NiagaraEditorWidgets/Private/Stack/SNiagaraStackModuleItem.h b/Engine/Plugins/FX/Niagara/Source/NiagaraEditorWidgets/Private/Stack/SNiagaraStackModuleItem.h
new file mode 100644
--- /dev/null
+++ b/Engine/Plugins/FX/Niagara/Source/NiagaraEditorWidgets/Private/Stack/SNiagaraStackModuleItem.h
@@ -0,0 +1,68 @@
+// Copyright 1998-2018 Epic Games, Inc. All Rights Reserved.
+
+#pragma once
+
+#include "SNiagaraStackEntryWidget.h"
+#include "Styling/SlateTypes.h"
+#include "Layout/Visibility.h"
+
+class UNiagaraStackModuleItem;
+class UNiagaraStackViewModel;
+
+class SNiagaraStackModuleItem : public SNiagaraStackEntryWidget
+{
+public:
+	SLATE_BEGIN_ARGS(SNiagaraStackModuleItem) { }
+	SLATE_END_ARGS();
+
+	void Construct(const FArguments& InArgs, UNiagaraStackModuleItem& InModuleItem, UNiagaraStackViewModel* InStackViewModel);
+
+	void SetEnabled(bool bInIsEnabled);
+
+	bool CheckEnabledStatus(bool bIsEnabled);
+
+	void FillRowContextMenu(class FMenuBuilder& MenuBuilder);
+
+	//~ SWidget interface
+	virtual FReply OnMouseButtonDoubleClick(const FGeometry& InMyGeometry, const FPointerEvent& InMouseEvent) override;
+
+private:
+	ECheckBoxState GetCheckState() const;
+
+	void OnCheckStateChanged(ECheckBoxState InCheckState);
+
+	EVisibility GetEditButtonVisibility() const;
+
+	EVisibility GetRaiseActionMenuVisibility() const;
+
+	EVisibility GetRefreshVisibility() const;
+
+	FReply MoveUpClicked();
+
+	FReply MoveDownClicked();
+
+	FReply DeleteClicked();
+	
+	TSharedRef<SWidget> RaiseActionMenuClicked();
+
+	bool CanRaiseActionMenu() const;
+
+	FReply RefreshClicked();
+
+	void InsertModuleAbove();
+
+	void InsertModuleBelow();
+
+	void ShowInsertModuleMenu(int32 InsertIndex);
+
+	FReply OnModuleItemDrop(TSharedPtr<class FDragDropOperation> DragDropOperation);
+
+	bool OnModuleItemAllowDrop(TSharedPtr<class FDragDropOperation> DragDropOperation);
+
+	EVisibility GetStackIssuesWarningVisibility() const;
+
+	FText GetErrorButtonTooltipText() const;
+
+private:
+	UNiagaraStackModuleItem* ModuleItem;
+};
\ No newline at end of file
diff --git a/Engine/Plugins/FX/Niagara/Source/NiagaraEditorWidgets/Private/Stack/SNiagaraStackParameterStoreEntryName.cpp b/Engine/Plugins/FX/Niagara/Source/NiagaraEditorWidgets/Private/Stack/SNiagaraStackParameterStoreEntryName.cpp
new file mode 100644
--- /dev/null
+++ b/Engine/Plugins/FX/Niagara/Source/NiagaraEditorWidgets/Private/Stack/SNiagaraStackParameterStoreEntryName.cpp
@@ -0,0 +1,90 @@
+// Copyright 1998-2018 Epic Games, Inc. All Rights Reserved.
+
+#include "Stack/SNiagaraStackParameterStoreEntryName.h"
+#include "NiagaraEditorStyle.h"
+#include "ViewModels/Stack/NiagaraStackViewModel.h"
+#include "ViewModels/Stack/NiagaraStackParameterStoreEntry.h"
+#include "Framework/Notifications/NotificationManager.h"
+#include "Widgets/Notifications/SNotificationList.h"
+#include "Widgets/Text/SInlineEditableTextBlock.h"
+
+#define LOCTEXT_NAMESPACE "NiagaraStackParameterStoreEntryName"
+
+void SNiagaraStackParameterStoreEntryName::Construct(const FArguments& InArgs, UNiagaraStackParameterStoreEntry* InStackEntry, UNiagaraStackViewModel* InStackViewModel)
+{
+	StackEntry = InStackEntry;
+	StackEntryItem = InStackEntry;
+	StackViewModel = InStackViewModel;
+
+	ChildSlot
+	[
+		SNew(SHorizontalBox)
+		// Name Label
+		+ SHorizontalBox::Slot()
+		.VAlign(VAlign_Center)
+		[
+			SAssignNew(NameTextBlock, SInlineEditableTextBlock)
+			.Style(FNiagaraEditorStyle::Get(), "NiagaraEditor.ParameterInlineEditableText")
+			.Text_UObject(StackEntry, &UNiagaraStackEntry::GetDisplayName)
+			.IsReadOnly(this, &SNiagaraStackParameterStoreEntryName::GetIsNameReadOnly)
+			.IsSelected(this, &SNiagaraStackParameterStoreEntryName::GetIsNameWidgetSelected)
+			.OnTextCommitted(this, &SNiagaraStackParameterStoreEntryName::OnNameTextCommitted)
+			.OnVerifyTextChanged(this, &SNiagaraStackParameterStoreEntryName::VerifyNameTextChanged)
+			.HighlightText_UObject(InStackViewModel, &UNiagaraStackViewModel::GetCurrentSearchText)
+			.ColorAndOpacity(this, &SNiagaraStackParameterStoreEntryName::GetTextColorForSearch)
+		]
+	];
+}
+
+void SNiagaraStackParameterStoreEntryName::Tick(const FGeometry& AllottedGeometry, const double InCurrentTime, const float InDeltaTime)
+{
+	// for remote initiating rename when the widget is added
+	if (StackEntry->GetIsRenamePending() && NameTextBlock.IsValid())
+	{
+		NameTextBlock->EnterEditingMode();
+		StackEntry->SetIsRenamePending(false);
+	}
+	SCompoundWidget::Tick(AllottedGeometry, InCurrentTime, InDeltaTime);
+}
+
+bool SNiagaraStackParameterStoreEntryName::GetIsNameReadOnly() const
+{
+	return StackEntry->CanRenameInput() == false;
+}
+
+bool SNiagaraStackParameterStoreEntryName::GetIsNameWidgetSelected() const
+{
+	return true;
+}
+
+bool SNiagaraStackParameterStoreEntryName::VerifyNameTextChanged(const FText& NewText, FText& OutErrorMessage)
+{
+	if (NewText.IsEmpty())
+	{
+		OutErrorMessage = LOCTEXT("NPCNameEmptyWarn", "Cannot have empty name!");
+		return false;
+	}
+	if (!StackEntry->IsUniqueName(NewText.ToString()))
+	{
+		OutErrorMessage = FText::Format(LOCTEXT("NPCNameConflictWarn", "\"{0}\" is already the name of another parameter in this collection."), NewText);
+		return false;
+	}
+
+	return true;
+}
+
+void SNiagaraStackParameterStoreEntryName::OnNameTextCommitted(const FText& InText, ETextCommit::Type InCommitType)
+{
+	if (StackEntry->GetDisplayName().ToString() != *InText.ToString())
+	{
+		StackEntry->RenameInput(*InText.ToString());
+		// toast notification
+		FNotificationInfo Info(FText::Format(LOCTEXT("NiagaraRenamedUserParameter", "System exposed parameter was renamed.\n{0}\n(All links to inner variables were updated in the process.)"), StackEntry->GetDisplayName()));
+		Info.ExpireDuration = 5.0f;
+		Info.bFireAndForget = true;
+		Info.Image = FCoreStyle::Get().GetBrush(TEXT("MessageLog.Info"));
+		FSlateNotificationManager::Get().AddNotification(Info);
+	}
+}
+
+#undef LOCTEXT_NAMESPACE
\ No newline at end of file
diff --git a/Engine/Plugins/FX/Niagara/Source/NiagaraEditorWidgets/Private/Stack/SNiagaraStackParameterStoreEntryName.h b/Engine/Plugins/FX/Niagara/Source/NiagaraEditorWidgets/Private/Stack/SNiagaraStackParameterStoreEntryName.h
new file mode 100644
--- /dev/null
+++ b/Engine/Plugins/FX/Niagara/Source/NiagaraEditorWidgets/Private/Stack/SNiagaraStackParameterStoreEntryName.h
@@ -0,0 +1,37 @@
+// Copyright 1998-2018 Epic Games, Inc. All Rights Reserved.
+
+#pragma once
+
+#include "SNiagaraStackEntryWidget.h"
+
+class UNiagaraStackParameterStoreEntry;
+class UNiagaraStackViewModel;
+class SInlineEditableTextBlock;
+
+class SNiagaraStackParameterStoreEntryName: public SNiagaraStackEntryWidget
+{
+public:
+	DECLARE_DELEGATE_OneParam(FOnColumnWidthChanged, float);
+
+public:
+	SLATE_BEGIN_ARGS(SNiagaraStackParameterStoreEntryName) { }
+	SLATE_END_ARGS();
+
+	void Construct(const FArguments& InArgs, UNiagaraStackParameterStoreEntry* InStackEntry, UNiagaraStackViewModel* InStackViewModel);
+
+	virtual void Tick(const FGeometry& AllottedGeometry, const double InCurrentTime, const float InDeltaTime) override;
+
+private:
+	bool GetIsNameReadOnly() const;
+
+	bool GetIsNameWidgetSelected() const;
+
+	bool VerifyNameTextChanged(const FText& NewText, FText& OutErrorMessage);
+
+	void OnNameTextCommitted(const FText& InText, ETextCommit::Type InCommitType);
+
+private:
+	UNiagaraStackParameterStoreEntry* StackEntry;
+
+	TSharedPtr<SInlineEditableTextBlock> NameTextBlock;
+};
\ No newline at end of file
diff --git a/Engine/Plugins/FX/Niagara/Source/NiagaraEditorWidgets/Private/Stack/SNiagaraStackParameterStoreEntryValue.cpp b/Engine/Plugins/FX/Niagara/Source/NiagaraEditorWidgets/Private/Stack/SNiagaraStackParameterStoreEntryValue.cpp
new file mode 100644
--- /dev/null
+++ b/Engine/Plugins/FX/Niagara/Source/NiagaraEditorWidgets/Private/Stack/SNiagaraStackParameterStoreEntryValue.cpp
@@ -0,0 +1,302 @@
+// Copyright 1998-2018 Epic Games, Inc. All Rights Reserved.
+
+#include "Stack/SNiagaraStackParameterStoreEntryValue.h"
+#include "NiagaraEditorModule.h"
+#include "NiagaraEditorWidgetsStyle.h"
+#include "NiagaraEditorStyle.h"
+#include "INiagaraEditorTypeUtilities.h"
+#include "SNiagaraParameterEditor.h"
+#include "ViewModels/Stack/NiagaraStackParameterStoreEntry.h"
+#include "Widgets/Images/SImage.h"
+#include "Widgets/Layout/SBox.h"
+#include "Widgets/Input/SButton.h"
+#include "Widgets/Input/SComboButton.h"
+#include "Widgets/Notifications/SNotificationList.h"
+#include "Framework/Notifications/NotificationManager.h"
+#include "Framework/MultiBox/MultiBoxBuilder.h"
+#include "PropertyEditorModule.h"
+#include "IStructureDetailsView.h"
+#include "Modules/ModuleManager.h"
+
+#define LOCTEXT_NAMESPACE "NiagaraStackParameterStoreEntryValue"
+
+void SNiagaraStackParameterStoreEntryValue::Construct(const FArguments& InArgs, UNiagaraStackParameterStoreEntry* InStackEntry)
+{
+	StackEntry = InStackEntry;
+
+	StackEntry->OnValueChanged().AddSP(this, &SNiagaraStackParameterStoreEntryValue::OnInputValueChanged);
+	DisplayedValueStruct = StackEntry->GetValueStruct();
+
+	FMargin ItemPadding = FMargin(0);
+	ChildSlot
+	[
+		// Values
+		SNew(SHorizontalBox)
+		+ SHorizontalBox::Slot()
+		.VAlign(VAlign_Center)
+		.AutoWidth()
+		.Padding(0, 0, 3, 0)
+		[
+			// Value Icon
+			SNew(SBox)
+			.WidthOverride(TextIconSize)
+			[
+				SNew(STextBlock)
+				.Font(FEditorStyle::Get().GetFontStyle("FontAwesome.10"))
+				.Text(this, &SNiagaraStackParameterStoreEntryValue::GetInputIconText)
+				.ToolTipText(this, &SNiagaraStackParameterStoreEntryValue::GetInputIconToolTip)
+				.ColorAndOpacity(this, &SNiagaraStackParameterStoreEntryValue::GetInputIconColor)
+			]
+		]
+		+ SHorizontalBox::Slot()
+		.VAlign(VAlign_Center)
+		[
+			// Assigned handle
+			SNew(SVerticalBox)
+			// Value struct
+			+ SVerticalBox::Slot()
+			.AutoHeight()
+			[
+				SAssignNew(ValueStructContainer, SBox)
+				[
+					ConstructValueStructWidget()
+				]
+			]
+		]
+
+		// Handle drop-down button
+		+ SHorizontalBox::Slot()
+		.AutoWidth()
+		.VAlign(VAlign_Center)
+		.Padding(3, 0, 0, 0)
+		[
+			SNew(SComboButton)
+			.ButtonStyle(FEditorStyle::Get(), "HoverHintOnly")
+			.ForegroundColor(FSlateColor::UseForeground())
+			.OnGetMenuContent(this, &SNiagaraStackParameterStoreEntryValue::OnGetAvailableHandleMenu)
+			.ContentPadding(FMargin(2))
+			.HAlign(HAlign_Center)
+			.VAlign(VAlign_Center)
+		]
+
+		// Reset Button
+		+ SHorizontalBox::Slot()
+		.AutoWidth()
+		.VAlign(VAlign_Center)
+		.Padding(3, 0, 0, 0)
+		[
+			SNew(SButton)
+			.IsFocusable(false)
+			.ToolTipText(LOCTEXT("ResetToolTip", "Reset to the default value"))
+			.ButtonStyle(FEditorStyle::Get(), "NoBorder")
+			.ContentPadding(0)
+			.Visibility(this, &SNiagaraStackParameterStoreEntryValue::GetResetButtonVisibility)
+			.OnClicked(this, &SNiagaraStackParameterStoreEntryValue::ResetButtonPressed)
+			.Content()
+			[
+				SNew(SImage)
+				.Image(FEditorStyle::GetBrush("PropertyWindow.DiffersFromDefault"))
+			]
+		]
+		+ SHorizontalBox::Slot()
+		.AutoWidth()
+		[
+			SNew(SButton)
+			.ButtonStyle(FEditorStyle::Get(), "HoverHintOnly")
+			.IsFocusable(false)
+			.ForegroundColor(FSlateColor::UseForeground())
+			.ToolTipText(LOCTEXT("DeleteToolTip", "Delete this parameter"))
+			.Visibility(this, &SNiagaraStackParameterStoreEntryValue::GetDeleteButtonVisibility)
+			.OnClicked(this, &SNiagaraStackParameterStoreEntryValue::DeleteClicked)
+			.Content()
+			[
+				SNew(STextBlock)
+				.Font(FEditorStyle::Get().GetFontStyle("FontAwesome.10"))
+				.Text(FText::FromString(FString(TEXT("\xf1f8"))))
+			]
+		]
+	];
+}
+
+FReply SNiagaraStackParameterStoreEntryValue::DeleteClicked()
+{
+	StackEntry->Delete();
+	// toast notification
+	FNotificationInfo Info(FText::Format(LOCTEXT("NiagaraDeletedUserParameter", "System exposed parameter was deleted.\n{0}\n(All links to inner variables were invalidated in the process.)"), StackEntry->GetDisplayName()));
+	Info.ExpireDuration = 5.0f;
+	Info.bFireAndForget = true;
+	Info.Image = FCoreStyle::Get().GetBrush(TEXT("MessageLog.Info"));
+	FSlateNotificationManager::Get().AddNotification(Info);
+	return FReply::Handled();
+}
+
+TSharedRef<SWidget> SNiagaraStackParameterStoreEntryValue::OnGetAvailableHandleMenu()
+{
+	// TODO: This will need to be adjusted based on the current stack being edited, i.e. system vs emitter vs particle.
+	FMenuBuilder MenuBuilder(true, nullptr);
+	return MenuBuilder.MakeWidget();
+}
+
+TSharedRef<SWidget> SNiagaraStackParameterStoreEntryValue::ConstructValueStructWidget()
+{
+	ValueStructParameterEditor.Reset();
+	ValueStructDetailsView.Reset();
+	if (DisplayedValueStruct.IsValid())
+	{
+		FNiagaraEditorModule& NiagaraEditorModule = FModuleManager::GetModuleChecked<FNiagaraEditorModule>("NiagaraEditor");
+		TSharedPtr<INiagaraEditorTypeUtilities, ESPMode::ThreadSafe> TypeEditorUtilities = NiagaraEditorModule.GetTypeUtilities(StackEntry->GetInputType());
+		if (TypeEditorUtilities.IsValid() && TypeEditorUtilities->CanCreateParameterEditor())
+		{
+			TSharedPtr<SNiagaraParameterEditor> ParameterEditor = TypeEditorUtilities->CreateParameterEditor(StackEntry->GetInputType());
+			ParameterEditor->UpdateInternalValueFromStruct(DisplayedValueStruct.ToSharedRef());
+			ParameterEditor->SetOnBeginValueChange(SNiagaraParameterEditor::FOnValueChange::CreateSP(
+				this, &SNiagaraStackParameterStoreEntryValue::ParameterBeginValueChange));
+			ParameterEditor->SetOnEndValueChange(SNiagaraParameterEditor::FOnValueChange::CreateSP(
+				this, &SNiagaraStackParameterStoreEntryValue::ParameterEndValueChange));
+			ParameterEditor->SetOnValueChanged(SNiagaraParameterEditor::FOnValueChange::CreateSP(
+				this, &SNiagaraStackParameterStoreEntryValue::ParameterValueChanged, ParameterEditor.ToSharedRef()));
+
+			ValueStructParameterEditor = ParameterEditor;
+
+			return SNew(SBox)
+				.HAlign(ParameterEditor->GetHorizontalAlignment())
+				.VAlign(ParameterEditor->GetVerticalAlignment())
+				[
+					ParameterEditor.ToSharedRef()
+				];
+		}
+		else
+		{
+			FPropertyEditorModule& PropertyEditorModule = FModuleManager::GetModuleChecked<FPropertyEditorModule>("PropertyEditor");
+
+			TSharedRef<IStructureDetailsView> StructureDetailsView = PropertyEditorModule.CreateStructureDetailView(
+				FDetailsViewArgs(false, false, false, FDetailsViewArgs::HideNameArea, true),
+				FStructureDetailsViewArgs(),
+				nullptr);
+
+			StructureDetailsView->SetStructureData(DisplayedValueStruct);
+			StructureDetailsView->GetOnFinishedChangingPropertiesDelegate().AddSP(this, &SNiagaraStackParameterStoreEntryValue::ParameterPropertyValueChanged);
+
+			ValueStructDetailsView = StructureDetailsView;
+			return StructureDetailsView->GetWidget().ToSharedRef();
+		}
+	}
+	else
+	{
+		return SNullWidget::NullWidget;
+	}
+}
+
+void SNiagaraStackParameterStoreEntryValue::OnInputValueChanged()
+{
+	TSharedPtr<FStructOnScope> NewValueStruct = StackEntry->GetValueStruct();
+	if (DisplayedValueStruct == NewValueStruct)
+	{
+		if (ValueStructParameterEditor.IsValid())
+		{
+			ValueStructParameterEditor->UpdateInternalValueFromStruct(DisplayedValueStruct.ToSharedRef());
+		}
+		if (ValueStructDetailsView.IsValid())
+		{
+			ValueStructDetailsView->SetStructureData(TSharedPtr<FStructOnScope>());
+			ValueStructDetailsView->SetStructureData(DisplayedValueStruct);
+		}
+	}
+	else
+	{
+		DisplayedValueStruct = NewValueStruct;
+		ValueStructContainer->SetContent(ConstructValueStructWidget());
+	}
+}
+
+void SNiagaraStackParameterStoreEntryValue::ParameterBeginValueChange()
+{
+	StackEntry->NotifyBeginValueChange();
+}
+
+void SNiagaraStackParameterStoreEntryValue::ParameterEndValueChange()
+{
+	StackEntry->NotifyEndValueChange();
+}
+
+void SNiagaraStackParameterStoreEntryValue::ParameterValueChanged(TSharedRef<SNiagaraParameterEditor> ParameterEditor)
+{
+	ParameterEditor->UpdateStructFromInternalValue(StackEntry->GetValueStruct().ToSharedRef());
+	StackEntry->NotifyValueChanged();
+}
+
+void SNiagaraStackParameterStoreEntryValue::ParameterPropertyValueChanged(const FPropertyChangedEvent& PropertyChangedEvent)
+{
+	StackEntry->NotifyValueChanged();
+}
+
+EVisibility SNiagaraStackParameterStoreEntryValue::GetDeleteButtonVisibility() const
+{
+	return StackEntry->CanRenameInput() ? EVisibility::Visible : EVisibility::Collapsed;
+}
+
+EVisibility SNiagaraStackParameterStoreEntryValue::GetReferenceVisibility() const
+{
+	return EVisibility::Collapsed;
+}
+
+EVisibility SNiagaraStackParameterStoreEntryValue::GetResetButtonVisibility() const
+{
+	return StackEntry->CanReset() ? EVisibility::Visible : EVisibility::Hidden;
+}
+
+FReply SNiagaraStackParameterStoreEntryValue::ResetButtonPressed() const
+{
+	StackEntry->Reset();
+	return FReply::Handled();
+}
+
+FText SNiagaraStackParameterStoreEntryValue::GetInputIconText() const
+{
+	if (DisplayedValueStruct.IsValid())
+	{
+		return FText::FromString(FString(TEXT("\xf040") /* fa-pencil */));
+	}
+	else if (StackEntry->GetValueObject() != nullptr)
+	{
+		return FText::FromString(FString(TEXT("\xf1C0") /* fa-database */));
+	}
+	else
+	{
+		return FText();
+	}
+}
+
+FText SNiagaraStackParameterStoreEntryValue::GetInputIconToolTip() const
+{
+	if (DisplayedValueStruct.IsValid())
+	{
+		return LOCTEXT("StructInputIconToolTip", "Local Value");
+	}
+	else if (StackEntry->GetValueObject() != nullptr)
+	{
+		return LOCTEXT("DataInterfaceInputIconToolTip", "Data Value");
+	}
+	else
+	{
+		return FText();
+	}
+}
+
+FSlateColor SNiagaraStackParameterStoreEntryValue::GetInputIconColor() const
+{
+	if (DisplayedValueStruct.IsValid())
+	{
+		return FLinearColor(FColor::Orange);
+	}
+	else if (StackEntry->GetValueObject() != nullptr)
+	{
+		return FLinearColor(FColor::Yellow);
+	}
+	else
+	{
+		return FLinearColor(FColor::White);
+	}
+}
+
+#undef LOCTEXT_NAMESPACE
\ No newline at end of file
