Copying //UE4/Dev-Niagara to //UE4/Dev-Main (Source: //UE4/Dev-Niagara @ 4074996)

diff --git a/Engine/Plugins/FX/Niagara/Source/NiagaraEditorWidgets/Private/Stack/SNiagaraStackParameterStoreEntryValue.h b/Engine/Plugins/FX/Niagara/Source/NiagaraEditorWidgets/Private/Stack/SNiagaraStackParameterStoreEntryValue.h
new file mode 100644
--- /dev/null
+++ b/Engine/Plugins/FX/Niagara/Source/NiagaraEditorWidgets/Private/Stack/SNiagaraStackParameterStoreEntryValue.h
@@ -0,0 +1,64 @@
+// Copyright 1998-2018 Epic Games, Inc. All Rights Reserved.
+
+#pragma once
+
+#include "Widgets/SCompoundWidget.h"
+#include "Widgets/DeclarativeSyntaxSupport.h"
+
+class UNiagaraStackParameterStoreEntry;
+class SNiagaraParameterEditor;
+class FStructOnScope;
+class SBox;
+class IStructureDetailsView;
+
+class SNiagaraStackParameterStoreEntryValue : public SCompoundWidget
+{
+public:
+	DECLARE_DELEGATE_OneParam(FOnColumnWidthChanged, float)
+public:
+	SLATE_BEGIN_ARGS(SNiagaraStackParameterStoreEntryValue) { }
+	SLATE_END_ARGS();
+
+	void Construct(const FArguments& InArgs, UNiagaraStackParameterStoreEntry* InStackEntry);
+
+private:
+	FReply DeleteClicked();
+
+	TSharedRef<SWidget> OnGetAvailableHandleMenu();
+
+	TSharedRef<SWidget> ConstructValueStructWidget();
+
+	void OnInputValueChanged();
+
+	void ParameterBeginValueChange();
+
+	void ParameterEndValueChange();
+
+	void ParameterValueChanged(TSharedRef<SNiagaraParameterEditor> ParameterEditor);
+
+	void ParameterPropertyValueChanged(const FPropertyChangedEvent& PropertyChangedEvent);
+
+	EVisibility GetDeleteButtonVisibility() const;
+
+	EVisibility GetReferenceVisibility() const;
+
+	EVisibility GetResetButtonVisibility() const;
+
+	FReply ResetButtonPressed() const;
+
+	FText GetInputIconText() const;
+
+	FText GetInputIconToolTip() const;
+
+	FSlateColor GetInputIconColor() const;
+
+private:
+	UNiagaraStackParameterStoreEntry* StackEntry;
+
+	TSharedPtr<FStructOnScope> DisplayedValueStruct;
+
+	TSharedPtr<SBox> ValueStructContainer;
+	TSharedPtr<SNiagaraParameterEditor> ValueStructParameterEditor;
+	TSharedPtr<IStructureDetailsView> ValueStructDetailsView;
+	const float TextIconSize = 16;
+};
\ No newline at end of file
diff --git a/Engine/Plugins/FX/Niagara/Source/NiagaraEditorWidgets/Private/Stack/SNiagaraStackRendererItem.cpp b/Engine/Plugins/FX/Niagara/Source/NiagaraEditorWidgets/Private/Stack/SNiagaraStackRendererItem.cpp
new file mode 100644
--- /dev/null
+++ b/Engine/Plugins/FX/Niagara/Source/NiagaraEditorWidgets/Private/Stack/SNiagaraStackRendererItem.cpp
@@ -0,0 +1,162 @@
+// Copyright 1998-2018 Epic Games, Inc. All Rights Reserved.
+
+#include "Stack/SNiagaraStackRendererItem.h"
+#include "NiagaraEditorWidgetsStyle.h"
+#include "NiagaraEditorStyle.h"
+#include "EditorStyleSet.h"
+#include "NiagaraRendererProperties.h"
+#include "ViewModels/Stack/NiagaraStackRendererItem.h"
+#include "ViewModels/Stack/NiagaraStackViewModel.h"
+#include "Widgets/SBoxPanel.h"
+#include "Widgets/Input/SButton.h"
+#include "Widgets/Input/SCheckBox.h"
+#include "Widgets/Images/SImage.h"
+#include "Widgets/Text/STextBlock.h"
+#include "Styling/SlateIconFinder.h"
+#include "Widgets/Images/SImage.h"
+#include "Widgets/Layout/SBox.h"
+#include "SNiagaraStackErrorButton.h"
+
+#define LOCTEXT_NAMESPACE "NiagaraStackRendererItem"
+
+void SNiagaraStackRendererItem::Construct(const FArguments& InArgs, UNiagaraStackRendererItem& InRendererItem, UNiagaraStackViewModel* InStackViewModel)
+{
+	RendererItem = &InRendererItem;
+	StackEntryItem = RendererItem;
+	StackViewModel = InStackViewModel;
+
+	ChildSlot
+	[
+		// Renderer icon
+		SNew(SHorizontalBox)
+		+ SHorizontalBox::Slot()
+		.Padding(5, 0, 0, 0)
+		.AutoWidth()
+		.VAlign(VAlign_Center)
+		[
+			SNew(SImage)
+			.Image(FSlateIconFinder::FindIconBrushForClass(RendererItem->GetRendererProperties()->GetClass()))
+		]
+		// Display name
+		+ SHorizontalBox::Slot()
+		.Padding(5, 0, 0, 0)
+		.VAlign(VAlign_Center)
+		[
+			SNew(STextBlock)
+			.TextStyle(FNiagaraEditorWidgetsStyle::Get(), "NiagaraEditor.Stack.ItemText")
+			.ToolTipText_UObject(RendererItem, &UNiagaraStackEntry::GetTooltipText)
+			.Text_UObject(RendererItem, &UNiagaraStackEntry::GetDisplayName)
+			.HighlightText_UObject(InStackViewModel, &UNiagaraStackViewModel::GetCurrentSearchText)
+			.ColorAndOpacity(this, &SNiagaraStackRendererItem::GetTextColorForSearch)
+		]
+		// Stack issues icon
+		+ SHorizontalBox::Slot()
+		.Padding(2, 0, 0, 0)
+		.AutoWidth()
+		[
+			SNew(SNiagaraStackErrorButton)
+			.IssueSeverity_UObject(RendererItem, &UNiagaraStackRendererItem::GetHighestStackIssueSeverity)
+			.ErrorTooltip(this, &SNiagaraStackRendererItem::GetErrorButtonTooltipText)
+			.Visibility(this, &SNiagaraStackRendererItem::GetStackIssuesWarningVisibility)
+			.OnButtonClicked(this, &SNiagaraStackRendererItem::ExpandEntry)
+		]
+		// Delete button
+		+ SHorizontalBox::Slot()
+		.AutoWidth()
+		[
+			SNew(SButton)
+			.ButtonStyle(FEditorStyle::Get(), "HoverHintOnly")
+			.IsFocusable(false)
+			.ForegroundColor(FNiagaraEditorWidgetsStyle::Get().GetColor("NiagaraEditor.Stack.ForegroundColor"))
+			.ToolTipText(LOCTEXT("DeleteRendererToolTip", "Delete this Renderer"))
+			.Visibility(this, &SNiagaraStackRendererItem::GetDeleteButtonVisibility)
+			.OnClicked(this, &SNiagaraStackRendererItem::DeleteClicked)
+			.Content()
+			[
+				SNew(STextBlock)
+				.Font(FEditorStyle::Get().GetFontStyle("FontAwesome.10"))
+				.Text(FText::FromString(FString(TEXT("\xf1f8"))))
+			]
+		]
+		// Reset to base Button
+		+ SHorizontalBox::Slot()
+		.AutoWidth()
+		.VAlign(VAlign_Center)
+		.Padding(3, 0, 0, 0)
+		[
+			SNew(SButton)
+			.IsFocusable(false)
+			.ToolTipText(LOCTEXT("ResetRendererToBaseToolTip", "Reset this renderer to the state defined by the parent emitter"))
+			.ButtonStyle(FEditorStyle::Get(), "NoBorder")
+			.ContentPadding(0)
+			.Visibility(this, &SNiagaraStackRendererItem::GetResetToBaseButtonVisibility)
+			.OnClicked(this, &SNiagaraStackRendererItem::ResetToBaseButtonClicked)
+			.Content()
+			[
+				SNew(SImage)
+				.Image(FEditorStyle::GetBrush("PropertyWindow.DiffersFromDefault"))
+				.ColorAndOpacity(FSlateColor(FLinearColor::Green))
+			]
+		]
+		// Enabled checkbox
+		+ SHorizontalBox::Slot()
+		.Padding(2, 0, 0, 0)
+		.AutoWidth()
+		[
+			SNew(SCheckBox)
+			.IsChecked(this, &SNiagaraStackRendererItem::CheckEnabledStatus)
+			.OnCheckStateChanged(this, &SNiagaraStackRendererItem::OnCheckStateChanged)
+		]
+	];
+}
+
+EVisibility SNiagaraStackRendererItem::GetDeleteButtonVisibility() const
+{
+	return RendererItem->CanDelete() ? EVisibility::Visible : EVisibility::Collapsed;
+}
+
+FReply SNiagaraStackRendererItem::DeleteClicked()
+{
+	RendererItem->Delete();
+	return FReply::Handled();
+}
+
+EVisibility SNiagaraStackRendererItem::GetResetToBaseButtonVisibility() const
+{
+	if (RendererItem->CanHaveBase())
+	{
+		return RendererItem->CanResetToBase() ? EVisibility::Visible : EVisibility::Hidden;
+	}
+	else
+	{
+		return EVisibility::Collapsed;
+	}
+}
+
+FReply SNiagaraStackRendererItem::ResetToBaseButtonClicked()
+{
+	RendererItem->ResetToBase();
+	return FReply::Handled();
+}
+
+void SNiagaraStackRendererItem::OnCheckStateChanged(ECheckBoxState InCheckState)
+{
+	RendererItem->SetIsEnabled(InCheckState == ECheckBoxState::Checked);
+}
+
+ECheckBoxState SNiagaraStackRendererItem::CheckEnabledStatus() const
+{
+	return RendererItem->GetIsEnabled() ? ECheckBoxState::Checked : ECheckBoxState::Unchecked;
+}
+
+EVisibility SNiagaraStackRendererItem::GetStackIssuesWarningVisibility() const
+{
+	return  RendererItem->GetRecursiveStackIssuesCount() > 0? EVisibility::Visible : EVisibility::Collapsed;
+}
+
+FText SNiagaraStackRendererItem::GetErrorButtonTooltipText() const
+{
+	return FText::Format(LOCTEXT("ModuleIssuesTooltip", "This renderer has {0} issues, click to expand."), RendererItem->GetRecursiveStackIssuesCount());
+}
+
+#undef LOCTEXT_NAMESPACE
\ No newline at end of file
diff --git a/Engine/Plugins/FX/Niagara/Source/NiagaraEditorWidgets/Private/Stack/SNiagaraStackRendererItem.h b/Engine/Plugins/FX/Niagara/Source/NiagaraEditorWidgets/Private/Stack/SNiagaraStackRendererItem.h
new file mode 100644
--- /dev/null
+++ b/Engine/Plugins/FX/Niagara/Source/NiagaraEditorWidgets/Private/Stack/SNiagaraStackRendererItem.h
@@ -0,0 +1,39 @@
+// Copyright 1998-2018 Epic Games, Inc. All Rights Reserved.
+
+#pragma once
+
+#include "SNiagaraStackEntryWidget.h"
+#include "Styling/SlateTypes.h"
+#include "Layout/Visibility.h"
+
+class UNiagaraStackRendererItem;
+class UNiagaraStackViewModel; 
+
+class SNiagaraStackRendererItem: public SNiagaraStackEntryWidget
+{
+public:
+	SLATE_BEGIN_ARGS(SNiagaraStackRendererItem) { }
+	SLATE_END_ARGS();
+
+	void Construct(const FArguments& InArgs, UNiagaraStackRendererItem& InRendererItem, UNiagaraStackViewModel* InStackViewModel);
+	
+private:
+	EVisibility GetDeleteButtonVisibility() const;
+
+	FReply DeleteClicked();
+
+	EVisibility GetResetToBaseButtonVisibility() const;
+
+	FReply ResetToBaseButtonClicked();
+
+	void OnCheckStateChanged(ECheckBoxState InCheckState);
+
+	ECheckBoxState CheckEnabledStatus() const;
+
+	EVisibility GetStackIssuesWarningVisibility() const;
+
+	FText GetErrorButtonTooltipText() const;
+
+private:
+	UNiagaraStackRendererItem* RendererItem;
+};
\ No newline at end of file
diff --git a/Engine/Plugins/FX/Niagara/Source/NiagaraEditorWidgets/Private/Stack/SNiagaraStackTableRow.cpp b/Engine/Plugins/FX/Niagara/Source/NiagaraEditorWidgets/Private/Stack/SNiagaraStackTableRow.cpp
new file mode 100644
--- /dev/null
+++ b/Engine/Plugins/FX/Niagara/Source/NiagaraEditorWidgets/Private/Stack/SNiagaraStackTableRow.cpp
@@ -0,0 +1,455 @@
+// Copyright 1998-2018 Epic Games, Inc. All Rights Reserved.
+
+#include "Stack/SNiagaraStackTableRow.h"
+#include "NiagaraEditorWidgetsStyle.h"
+#include "EditorStyleSet.h"
+#include "ViewModels/Stack/NiagaraStackViewModel.h"
+#include "ViewModels/Stack/NiagaraStackItem.h"
+#include "ViewModels/Stack/NiagaraStackItemGroup.h"
+#include "ViewModels/Stack/NiagaraStackEntry.h"
+#include "NiagaraEditorWidgetsUtilities.h"
+#include "Widgets/Input/SButton.h"
+#include "Widgets/Images/SImage.h"
+#include "Widgets/Layout/SBox.h"
+#include "IContentBrowserSingleton.h"
+#include "ContentBrowserModule.h"
+#include "Toolkits/AssetEditorManager.h"
+#include "Framework/Application/SlateApplication.h"
+
+#define LOCTEXT_NAMESPACE "NiagaraStackTableRow"
+
+const float IndentSize = 12;
+
+void SNiagaraStackTableRow::Construct(const FArguments& InArgs, UNiagaraStackViewModel* InStackViewModel, UNiagaraStackEntry* InStackEntry, const TSharedRef<STreeView<UNiagaraStackEntry*>>& InOwnerTree)
+{
+	ContentPadding = InArgs._ContentPadding;
+	bIsCategoryIconHighlighted = InArgs._IsCategoryIconHighlighted;
+	bShowExecutionCategoryIcon = InArgs._ShowExecutionCategoryIcon;
+	NameColumnWidth = InArgs._NameColumnWidth;
+	ValueColumnWidth = InArgs._ValueColumnWidth;
+	NameColumnWidthChanged = InArgs._OnNameColumnWidthChanged;
+	ValueColumnWidthChanged = InArgs._OnValueColumnWidthChanged;
+	StackViewModel = InStackViewModel;
+	StackEntry = InStackEntry;
+	OwnerTree = InOwnerTree;
+
+	ExpandedImage = FCoreStyle::Get().GetBrush("TreeArrow_Expanded");
+	CollapsedImage = FCoreStyle::Get().GetBrush("TreeArrow_Collapsed");
+
+	InactiveItemBackgroundColor = InArgs._ItemBackgroundColor;
+	ActiveItemBackgroundColor = InactiveItemBackgroundColor + FLinearColor(.05f, .05f, .05f, 0.0f);
+	ForegroundColor = InArgs._ItemForegroundColor;
+
+	ExecutionCategoryToolTipText = InStackEntry->GetExecutionSubcategoryName() != NAME_None
+		? FText::Format(LOCTEXT("ExecutionCategoryToolTipFormat", "{0} - {1}"), FText::FromName(InStackEntry->GetExecutionCategoryName()), FText::FromName(InStackEntry->GetExecutionSubcategoryName()))
+		: FText::FromName(InStackEntry->GetExecutionCategoryName());
+
+	ConstructInternal(
+		STableRow<UNiagaraStackEntry*>::FArguments()
+			.OnDragDetected(InArgs._OnDragDetected)
+			.OnCanAcceptDrop(InArgs._OnCanAcceptDrop)
+			.OnAcceptDrop(InArgs._OnAcceptDrop)
+		, OwnerTree.ToSharedRef());
+}
+
+void SNiagaraStackTableRow::SetOverrideNameWidth(TOptional<float> InMinWidth, TOptional<float> InMaxWidth)
+{
+	NameMinWidth = InMinWidth;
+	NameMaxWidth = InMaxWidth;
+}
+
+void SNiagaraStackTableRow::SetOverrideNameAlignment(EHorizontalAlignment InHAlign, EVerticalAlignment InVAlign)
+{
+	NameHorizontalAlignment = InHAlign;
+	NameVerticalAlignment = InVAlign;
+}
+
+void SNiagaraStackTableRow::SetOverrideValueWidth(TOptional<float> InMinWidth, TOptional<float> InMaxWidth)
+{
+	ValueMinWidth = InMinWidth;
+	ValueMaxWidth = InMaxWidth;
+}
+
+void SNiagaraStackTableRow::SetOverrideValueAlignment(EHorizontalAlignment InHAlign, EVerticalAlignment InVAlign)
+{
+	ValueHorizontalAlignment = InHAlign;
+	ValueVerticalAlignment = InVAlign;
+}
+
+FMargin SNiagaraStackTableRow::GetContentPadding() const
+{
+	return ContentPadding;
+}
+
+void SNiagaraStackTableRow::SetContentPadding(FMargin InContentPadding)
+{
+	ContentPadding = InContentPadding;
+}
+
+void SNiagaraStackTableRow::SetNameAndValueContent(TSharedRef<SWidget> InNameWidget, TSharedPtr<SWidget> InValueWidget)
+{
+	FSlateColor IconColor = FNiagaraEditorWidgetsStyle::Get().GetColor(FNiagaraStackEditorWidgetsUtilities::GetColorNameForExecutionCategory(StackEntry->GetExecutionCategoryName()));
+	if (bIsCategoryIconHighlighted)
+	{
+		IconColor = FNiagaraEditorWidgetsStyle::Get().GetColor(FNiagaraStackEditorWidgetsUtilities::GetIconColorNameForExecutionCategory(StackEntry->GetExecutionCategoryName()));
+	}
+	TSharedRef<SHorizontalBox> NameContent = SNew(SHorizontalBox)
+	.Clipping(EWidgetClipping::OnDemand)
+	// Indent
+	+ SHorizontalBox::Slot()
+	.AutoWidth()
+	[
+		SNew(SBox)
+		.WidthOverride(this, &SNiagaraStackTableRow::GetIndentSize)
+	]
+	// Expand button
+	+ SHorizontalBox::Slot()
+	.AutoWidth()
+	.VAlign(VAlign_Center)
+	.Padding(0, 0, 1, 0)
+	[
+		SNew(SBox)
+		.WidthOverride(14)
+		[
+			SNew(SButton)
+			.ButtonStyle(FCoreStyle::Get(), "NoBorder")
+			.Visibility(this, &SNiagaraStackTableRow::GetExpanderVisibility)
+			.OnClicked(this, &SNiagaraStackTableRow::ExpandButtonClicked)
+			.ForegroundColor(FSlateColor::UseForeground())
+			.ContentPadding(2)
+			.HAlign(HAlign_Center)
+			[
+				SNew(SImage)
+				.Image(this, &SNiagaraStackTableRow::GetExpandButtonImage)
+				.ColorAndOpacity(FSlateColor::UseForeground())
+			]
+		]
+	]
+	// Execution sub-category icon
+	+ SHorizontalBox::Slot()
+	.AutoWidth()
+	.Padding(FMargin(1, 1, 4, 1))
+	.VAlign(EVerticalAlignment::VAlign_Center)
+	.HAlign(EHorizontalAlignment::HAlign_Center)
+	[
+		SNew(SBox)
+		.WidthOverride(FNiagaraEditorWidgetsStyle::Get().GetFloat("NiagaraEditor.Stack.IconHighlightedSize"))
+		.HAlign(EHorizontalAlignment::HAlign_Center)
+		.VAlign(EVerticalAlignment::VAlign_Center)
+		.ToolTipText(ExecutionCategoryToolTipText)
+		.Visibility(this, &SNiagaraStackTableRow::GetExecutionCategoryIconVisibility)
+		[
+			SNew(SImage)
+			.Visibility(this, &SNiagaraStackTableRow::GetExecutionCategoryIconVisibility)
+			.Image(FNiagaraEditorWidgetsStyle::Get().GetBrush(FNiagaraStackEditorWidgetsUtilities::GetIconNameForExecutionSubcategory(StackEntry->GetExecutionSubcategoryName(), bIsCategoryIconHighlighted)))
+			.ColorAndOpacity(IconColor)
+		]
+	]
+	// Name content
+	+ SHorizontalBox::Slot()
+	[
+		InNameWidget
+	];
+
+	TSharedPtr<SWidget> ChildContent;
+
+	if (InValueWidget.IsValid())
+	{
+		ChildContent = SNew(SSplitter)
+		.Style(FNiagaraEditorWidgetsStyle::Get(), "NiagaraEditor.Stack.Splitter")
+		.PhysicalSplitterHandleSize(1.0f)
+		.HitDetectionSplitterHandleSize(5.0f)
+
+		+ SSplitter::Slot()
+		.Value(NameColumnWidth)
+		.OnSlotResized(SSplitter::FOnSlotResized::CreateSP(this, &SNiagaraStackTableRow::OnNameColumnWidthChanged))
+		[
+			SNew(SBox)
+			.Padding(FMargin(ContentPadding.Left, ContentPadding.Top, 5, ContentPadding.Bottom))
+			.HAlign(NameHorizontalAlignment)
+			.VAlign(NameVerticalAlignment)
+			.MinDesiredWidth(NameMinWidth.IsSet() ? NameMinWidth.GetValue() : FOptionalSize())
+			.MaxDesiredWidth(NameMaxWidth.IsSet() ? NameMaxWidth.GetValue() : FOptionalSize())
+			[
+				NameContent
+			]
+		]
+
+		// Value
+		+ SSplitter::Slot()
+		.Value(ValueColumnWidth)
+		.OnSlotResized(SSplitter::FOnSlotResized::CreateSP(this, &SNiagaraStackTableRow::OnValueColumnWidthChanged))
+		[
+			SNew(SBox)
+			.Padding(FMargin(4, ContentPadding.Top, ContentPadding.Right, ContentPadding.Bottom))
+			.HAlign(ValueHorizontalAlignment)
+			.VAlign(ValueVerticalAlignment)
+			.MinDesiredWidth(ValueMinWidth.IsSet() ? ValueMinWidth.GetValue() : FOptionalSize())
+			.MaxDesiredWidth(ValueMaxWidth.IsSet() ? ValueMaxWidth.GetValue() : FOptionalSize())
+			[
+				InValueWidget.ToSharedRef()
+			]
+		];
+	}
+	else
+	{
+		ChildContent = SNew(SBox)
+		.Padding(ContentPadding)
+		.HAlign(NameHorizontalAlignment)
+		.VAlign(NameVerticalAlignment)
+		.MinDesiredWidth(NameMinWidth.IsSet() ? NameMinWidth.GetValue() : FOptionalSize())
+		.MaxDesiredWidth(NameMaxWidth.IsSet() ? NameMaxWidth.GetValue() : FOptionalSize())
+		[
+			NameContent
+		];
+	}
+
+	ChildSlot
+	[
+		SNew(SOverlay)
+		+ SOverlay::Slot()
+		[
+			SNew(SBorder)
+			.BorderImage(FEditorStyle::GetBrush("WhiteBrush"))
+			.BorderBackgroundColor(FNiagaraEditorWidgetsStyle::Get().GetColor(FNiagaraStackEditorWidgetsUtilities::GetColorNameForExecutionCategory(StackEntry->GetExecutionCategoryName())))
+			.Visibility(this, &SNiagaraStackTableRow::GetRowVisibility)
+			.Padding(FMargin(9, 0, 9, 0))
+			[
+				SNew(SBorder)
+				.BorderImage(FEditorStyle::GetBrush("WhiteBrush"))
+				.BorderBackgroundColor(this, &SNiagaraStackTableRow::GetItemBackgroundColor)
+				.ForegroundColor(ForegroundColor)
+				.Padding(0)
+				[
+					ChildContent.ToSharedRef()
+				]
+			]
+		]
+		+ SOverlay::Slot()
+		[
+			SNew(SBorder)
+			.BorderImage(FNiagaraEditorWidgetsStyle::Get().GetBrush("NiagaraEditor.Stack.SearchResult"))
+			.BorderBackgroundColor(FNiagaraEditorWidgetsStyle::Get().GetColor("NiagaraEditor.Stack.SearchHighlightColor"))
+			.Visibility(this, &SNiagaraStackTableRow::GetSearchResultBorderVisibility)
+			.Padding(FMargin(0))
+			[
+				SNullWidget::NullWidget
+			]
+		]
+	];
+}
+
+bool SNiagaraStackTableRow::GetIsRowActive() const
+{
+	return IsHovered();
+}
+
+void SNiagaraStackTableRow::AddFillRowContextMenuHandler(FOnFillRowContextMenu FillRowContextMenuHandler)
+{
+	OnFillRowContextMenuHanders.Add(FillRowContextMenuHandler);
+}
+
+FReply SNiagaraStackTableRow::OnMouseButtonDoubleClick(const FGeometry& InMyGeometry, const FPointerEvent& InMouseEvent)
+{
+	return FReply::Unhandled();
+}
+
+FReply SNiagaraStackTableRow::OnMouseButtonUp(const FGeometry& MyGeometry, const FPointerEvent& MouseEvent)
+{
+	if (MouseEvent.GetEffectingButton() == EKeys::RightMouseButton)
+	{
+		FMenuBuilder MenuBuilder(true, nullptr);
+		MenuBuilder.BeginSection("ModuleActions", LOCTEXT("ModuleActions", "Module Actions"));
+		for (FOnFillRowContextMenu& OnFillRowContextMenuHandler : OnFillRowContextMenuHanders)
+		{
+			OnFillRowContextMenuHandler.ExecuteIfBound(MenuBuilder);
+		}
+		MenuBuilder.EndSection();
+		if (StackEntry->GetExternalAsset() != nullptr)
+		{
+			MenuBuilder.BeginSection("AssetActions", LOCTEXT("AssetActions", "Asset Actions"));
+			MenuBuilder.AddMenuEntry(
+				LOCTEXT("OpenAndFocusAsset", "Open and focus Asset"),
+				FText::Format(LOCTEXT("OpenAndFocusAssetTooltip", "Open {0} in separate editor"), StackEntry->GetDisplayName()),
+				FSlateIcon(),
+				FUIAction(FExecuteAction::CreateSP(this, &SNiagaraStackTableRow::OpenSourceAsset)));
+			MenuBuilder.AddMenuEntry(
+				LOCTEXT("ShowAssetInContentBrowser", "Show in Content Browser"),
+				FText::Format(LOCTEXT("ShowAssetInContentBrowserToolTip", "Navigate to {0} in the Content Browser window"), StackEntry->GetDisplayName()),
+				FSlateIcon(),
+				FUIAction(FExecuteAction::CreateSP(this, &SNiagaraStackTableRow::ShowAssetInContentBrowser)));
+			MenuBuilder.EndSection();
+		}
+		TArray<UNiagaraStackEntry*> EntriesToProcess;
+		TArray<UNiagaraStackEntry*> NavigationEntries;
+		StackViewModel->GetPathForEntry(StackEntry, EntriesToProcess);
+		for (UNiagaraStackEntry* Parent : EntriesToProcess)
+		{
+			UNiagaraStackItemGroup* GroupParent = Cast<UNiagaraStackItemGroup>(Parent);
+			UNiagaraStackItem* ItemParent = Cast<UNiagaraStackItem>(Parent);
+			if (GroupParent != nullptr)
+			{
+				MenuBuilder.BeginSection("StackRowNavigateTo", LOCTEXT("NavigateToSection", "Navigate to:"));
+				MenuBuilder.AddMenuEntry(
+					LOCTEXT("TopOfSection", "Top of section"),
+					FText::Format(LOCTEXT("NavigateToFormatted", "Navigate to {0}"), Parent->GetDisplayName()),
+					FSlateIcon(),
+					FUIAction(FExecuteAction::CreateSP(this, &SNiagaraStackTableRow::NavigateTo, Parent)));
+			}
+			if (ItemParent != nullptr)
+			{
+				MenuBuilder.AddMenuEntry(
+					LOCTEXT("TopOfModule", "Top of module"),
+					FText::Format(LOCTEXT("NavigateToFormatted", "Navigate to {0}"), Parent->GetDisplayName()),
+					FSlateIcon(),
+					FUIAction(FExecuteAction::CreateSP(this, &SNiagaraStackTableRow::NavigateTo, Parent)));
+			}
+			if (GroupParent != nullptr)
+			{
+				MenuBuilder.EndSection();
+			}
+		}
+
+		MenuBuilder.BeginSection("StackActions", LOCTEXT("StackActions", "Stack Actions"));
+		MenuBuilder.AddMenuEntry(
+			LOCTEXT("ExpandAllItems", "Expand all"),
+			LOCTEXT("ExpandAllItemsToolTip", "Expand all items under this header."),
+			FSlateIcon(),
+			FUIAction(FExecuteAction::CreateSP(this, &SNiagaraStackTableRow::ExpandChildren)));
+		MenuBuilder.AddMenuEntry(
+			LOCTEXT("CollapseAllItems", "Collapse all"),
+			LOCTEXT("CollapseAllItemsToolTip", "Collapse all items under this header."),
+			FSlateIcon(),
+			FUIAction(FExecuteAction::CreateSP(this, &SNiagaraStackTableRow::CollapseChildren)));
+		MenuBuilder.EndSection();
+
+		FWidgetPath WidgetPath = MouseEvent.GetEventPath() != nullptr ? *MouseEvent.GetEventPath() : FWidgetPath();
+		FSlateApplication::Get().PushMenu(AsShared(), WidgetPath, MenuBuilder.MakeWidget(), MouseEvent.GetScreenSpacePosition(), FPopupTransitionEffect(FPopupTransitionEffect::ContextMenu));
+		return FReply::Handled();
+	}
+	return STableRow<UNiagaraStackEntry*>::OnMouseButtonUp(MyGeometry, MouseEvent);
+}
+
+void SNiagaraStackTableRow::CollapseChildren()
+{
+	TArray<UNiagaraStackEntry*> Children;
+	StackEntry->GetUnfilteredChildren(Children);
+	for (UNiagaraStackEntry* Child : Children)
+	{
+		if (Child->GetCanExpand())
+		{
+			Child->SetIsExpanded(false);
+		}
+	}
+	// Calling SetIsExpanded doesn't broadcast structure change automatically due to the expense of synchronizing
+	// expanded state with the tree to prevent items being expanded on tick, so we call this manually here.
+	StackEntry->OnStructureChanged().Broadcast();
+}
+
+void SNiagaraStackTableRow::ExpandChildren()
+{
+	TArray<UNiagaraStackEntry*> Children;
+	StackEntry->GetUnfilteredChildren(Children);
+	for (UNiagaraStackEntry* Child : Children)
+	{
+		if (Child->GetCanExpand())
+		{
+			Child->SetIsExpanded(true);
+		}
+	}
+	// Calling SetIsExpanded doesn't broadcast structure change automatically due to the expense of synchronizing
+	// expanded state with the tree to prevent items being expanded on tick, so we call this manually here.
+	StackEntry->OnStructureChanged().Broadcast();
+}
+
+EVisibility SNiagaraStackTableRow::GetRowVisibility() const
+{
+	return StackEntry->GetShouldShowInStack()
+		? EVisibility::Visible
+		: EVisibility::Collapsed;
+}
+
+EVisibility SNiagaraStackTableRow::GetExecutionCategoryIconVisibility() const
+{
+	return bShowExecutionCategoryIcon && StackEntry->GetExecutionSubcategoryName() != NAME_None
+		? EVisibility::Visible
+		: EVisibility::Collapsed;
+}
+
+FOptionalSize SNiagaraStackTableRow::GetIndentSize() const
+{
+	return StackEntry->GetIndentLevel() * IndentSize;
+}
+
+EVisibility SNiagaraStackTableRow::GetExpanderVisibility() const
+{
+	if (StackEntry->GetCanExpand())
+	{
+		// TODO Cache this and refresh the cache when the entries structure changes.
+		TArray<UNiagaraStackEntry*> Children;
+		StackEntry->GetFilteredChildren(Children);
+		return Children.Num() > 0
+			? EVisibility::Visible
+			: EVisibility::Hidden;
+	}
+	else
+	{
+		return EVisibility::Collapsed;
+	}
+}
+
+FReply SNiagaraStackTableRow::ExpandButtonClicked()
+{
+	StackEntry->SetIsExpanded(!StackEntry->GetIsExpanded());
+	// Calling SetIsExpanded doesn't broadcast structure change automatically due to the expense of synchronizing
+	// expanded state with the tree to prevent items being expanded on tick, so we call this manually here.
+	StackEntry->OnStructureChanged().Broadcast();
+	return FReply::Handled();
+}
+
+const FSlateBrush* SNiagaraStackTableRow::GetExpandButtonImage() const
+{
+	return StackEntry->GetIsExpanded() ? ExpandedImage : CollapsedImage;
+}
+
+void SNiagaraStackTableRow::OnNameColumnWidthChanged(float Width)
+{
+	NameColumnWidthChanged.ExecuteIfBound(Width);
+}
+
+void SNiagaraStackTableRow::OnValueColumnWidthChanged(float Width)
+{
+	ValueColumnWidthChanged.ExecuteIfBound(Width);
+}
+
+FSlateColor SNiagaraStackTableRow::GetItemBackgroundColor() const
+{
+	return GetIsRowActive() ? ActiveItemBackgroundColor : InactiveItemBackgroundColor;
+}
+
+EVisibility SNiagaraStackTableRow::GetSearchResultBorderVisibility() const
+{
+	return StackViewModel->GetCurrentFocusedEntry() == StackEntry
+		? EVisibility::HitTestInvisible
+		: EVisibility::Hidden;
+}
+
+void SNiagaraStackTableRow::NavigateTo(UNiagaraStackEntry* Item)
+{
+	OwnerTree->RequestScrollIntoView(Item, 0, true);
+}
+
+void SNiagaraStackTableRow::OpenSourceAsset()
+{
+	FAssetEditorManager::Get().OpenEditorForAsset(StackEntry->GetExternalAsset());
+}
+
+void SNiagaraStackTableRow::ShowAssetInContentBrowser()
+{
+	FContentBrowserModule& ContentBrowserModule = FModuleManager::Get().LoadModuleChecked<FContentBrowserModule>(TEXT("ContentBrowser"));
+	TArray<FAssetData> Assets;
+	Assets.Add(FAssetData(StackEntry->GetExternalAsset()));
+	ContentBrowserModule.Get().SyncBrowserToAssets(Assets);
+}
+
+#undef LOCTEXT_NAMESPACE
\ No newline at end of file
diff --git a/Engine/Plugins/FX/Niagara/Source/NiagaraShader/NiagaraShader.Build.cs b/Engine/Plugins/FX/Niagara/Source/NiagaraShader/NiagaraShader.Build.cs
--- a/Engine/Plugins/FX/Niagara/Source/NiagaraShader/NiagaraShader.Build.cs
+++ b/Engine/Plugins/FX/Niagara/Source/NiagaraShader/NiagaraShader.Build.cs
@@ -13,6 +13,7 @@ public class NiagaraShader : ModuleRules
 				"RenderCore",
                 "ShaderCore",
                 "CoreUObject",
+                "NiagaraCore"
             }
         );
 
@@ -20,8 +21,10 @@ public class NiagaraShader : ModuleRules
         PublicDependencyModuleNames.AddRange(
             new string[] {
                 "Core",
+                "ShaderCore",
                 "VectorVM",
                 "RHI",
+                "NiagaraCore"
             }
         );
 
@@ -38,11 +41,11 @@ public class NiagaraShader : ModuleRules
             });
         }
 
-        PublicIncludePaths.AddRange(
+        PublicIncludePathModuleNames.AddRange(
             new string[] {
             });
 
-        PrivateIncludePaths.AddRange(
+        PrivateIncludePathModuleNames.AddRange(
             new string[] {
                 "Niagara",
             });
diff --git a/Engine/Plugins/FX/Niagara/Source/NiagaraShader/Private/NiagaraShader.cpp b/Engine/Plugins/FX/Niagara/Source/NiagaraShader/Private/NiagaraShader.cpp
--- a/Engine/Plugins/FX/Niagara/Source/NiagaraShader/Private/NiagaraShader.cpp
+++ b/Engine/Plugins/FX/Niagara/Source/NiagaraShader/Private/NiagaraShader.cpp
@@ -17,6 +17,9 @@
 #endif
 #include "ProfilingDebugging/CookStats.h"
 
+#include "NiagaraDataInterfaceBase.h"
+#include "UObject/UObjectGlobals.h"
+
 IMPLEMENT_SHADER_TYPE(,FNiagaraShader, TEXT("/Engine/Private/NiagaraEmitterInstanceShader.usf"),TEXT("SimulateMain"), SF_Compute)
 
 
@@ -55,14 +58,14 @@ uint32 FNiagaraShaderMap::NextCompilingId = 2;
 
 
 /** 
- * Tracks FNiagaraScripts and their shader maps that are being compiled.
+ * Tracks FNiagaraShaderScripts and their shader maps that are being compiled.
  * Uses a TRefCountPtr as this will be the only reference to a shader map while it is being compiled.
  */
-TMap<TRefCountPtr<FNiagaraShaderMap>, TArray<FNiagaraScript*> > FNiagaraShaderMap::NiagaraShaderMapsBeingCompiled;
+TMap<TRefCountPtr<FNiagaraShaderMap>, TArray<FNiagaraShaderScript*> > FNiagaraShaderMap::NiagaraShaderMapsBeingCompiled;
 
 
 
-static inline bool ShouldCacheNiagaraShader(const FNiagaraShaderType* ShaderType, EShaderPlatform Platform, const FNiagaraScript* Script)
+static inline bool ShouldCacheNiagaraShader(const FNiagaraShaderType* ShaderType, EShaderPlatform Platform, const FNiagaraShaderScript* Script)
 {
 	return ShaderType->ShouldCache(Platform, Script) && Script->ShouldCache(Platform, ShaderType);
 }
@@ -70,7 +73,7 @@ static inline bool ShouldCacheNiagaraShader(const FNiagaraShaderType* ShaderType
 
 
 /** Called for every script shader to update the appropriate stats. */
-void UpdateNiagaraShaderCompilingStats(const FNiagaraScript* Script)
+void UpdateNiagaraShaderCompilingStats(const FNiagaraShaderScript* Script)
 {
 	INC_DWORD_STAT_BY(STAT_ShaderCompiling_NumTotalNiagaraShaders,1);
 }
@@ -79,13 +82,22 @@ void UpdateNiagaraShaderCompilingStats(const FNiagaraScript* Script)
 void FNiagaraShaderMapId::Serialize(FArchive& Ar)
 {
 	// You must bump NIAGARASHADERMAP_DERIVEDDATA_VER if changing the serialization of FNiagaraShaderMapId.
-
+	/*const int32 NiagaraCustomVersion = GetLinkerCustomVersion(FNiagaraCustomVersion::GUID);
+	if (NiagaraCustomVersion <= FNiagaraCustomVersion::MovingTranslatorVersionToGuid)
+	{
+		int VersionId = 0;
+		Ar << VersionId;
+		CompilerVersionID.Invalidate();
+	}
+	else*/
+	{
+		Ar << CompilerVersionID;
+	}
 	Ar << BaseScriptID;
-
 	Ar << (int32&)FeatureLevel;
 
 	//ParameterSet.Serialize(Ar);		// NIAGARATODO: at some point we'll need stuff for static switches here
-	//Ar << ReferencedFunctions;		// may want this too, although we've got other mechanisms to determine need to recompile, unsure whether the shader map needs to know
+	Ar << ReferencedDependencyIds;
 }
 
 /** Hashes the script-specific part of this shader map Id. */
@@ -93,17 +105,17 @@ void FNiagaraShaderMapId::GetScriptHash(FSHAHash& OutHash) const
 {
 	FSHA1 HashState;
 
+	HashState.Update((const uint8*)&CompilerVersionID, sizeof(CompilerVersionID));
 	HashState.Update((const uint8*)&BaseScriptID, sizeof(BaseScriptID));
 	HashState.Update((const uint8*)&FeatureLevel, sizeof(FeatureLevel));
 
 	//ParameterSet.UpdateHash(HashState);		// will need for static switches
-
-	/*
-	for (int32 FunctionIndex = 0; FunctionIndex < ReferencedFunctions.Num(); FunctionIndex++)
+		
+	for (int32 Index = 0; Index < ReferencedDependencyIds.Num(); Index++)
 	{
-		HashState.Update((const uint8*)&ReferencedFunctions[FunctionIndex], sizeof(ReferencedFunctions[FunctionIndex]));
+		HashState.Update((const uint8*)&ReferencedDependencyIds[Index], sizeof(ReferencedDependencyIds[Index]));
 	}
-	*/
+	
 	HashState.Final();
 	HashState.GetHash(&OutHash.Hash[0]);
 }
@@ -117,11 +129,24 @@ void FNiagaraShaderMapId::GetScriptHash(FSHAHash& OutHash) const
 bool FNiagaraShaderMapId::operator==(const FNiagaraShaderMapId& ReferenceSet) const
 {
 	if (BaseScriptID != ReferenceSet.BaseScriptID 
-		//|| QualityLevel != ReferenceSet.QualityLevel
-		|| FeatureLevel != ReferenceSet.FeatureLevel)
+		|| FeatureLevel != ReferenceSet.FeatureLevel
+		|| CompilerVersionID != ReferenceSet.CompilerVersionID)
 	{
 		return false;
 	}
+
+	if (ReferencedDependencyIds.Num() != ReferenceSet.ReferencedDependencyIds.Num())
+	{
+		return false;
+	}
+
+	for (int32 i = 0; i < ReferencedDependencyIds.Num(); i++)
+	{
+		if (ReferencedDependencyIds[i] != ReferenceSet.ReferencedDependencyIds[i])
+		{
+			return false;
+		}
+	}
 	/*
 	if (ParameterSet != ReferenceSet.ParameterSet
 		|| ReferencedFunctions.Num() != ReferenceSet.ReferencedFunctions.Num()
@@ -146,9 +171,25 @@ void FNiagaraShaderMapId::AppendKeyString(FString& KeyString) const
 {
 	KeyString += BaseScriptID.ToString();
 	KeyString += TEXT("_");
+
 	FString FeatureLevelString;
 	GetFeatureLevelName(FeatureLevel, FeatureLevelString);
 	KeyString += FeatureLevelString + TEXT("_");
+
+	KeyString += CompilerVersionID.ToString();
+	KeyString += TEXT("_");
+
+	// Add any referenced functions to the key so that we will recompile when they are changed
+	for (int32 FunctionIndex = 0; FunctionIndex < ReferencedDependencyIds.Num(); FunctionIndex++)
+	{
+		KeyString += ReferencedDependencyIds[FunctionIndex].ToString();
+
+		if (FunctionIndex < ReferencedDependencyIds.Num() - 1)
+		{
+			KeyString += TEXT("_");
+		}
+	}
+
 	/*
 	ParameterSet.AppendKeyString(KeyString);
 
@@ -165,23 +206,25 @@ void FNiagaraShaderMapId::AppendKeyString(FString& KeyString) const
 }
 
 
+TMap<const FShaderCompileJob*, TArray<FNiagaraDataInterfaceGPUParamInfo> > FNiagaraShaderType::ExtraParamInfo;
+
 /**
  * Enqueues a compilation for a new shader of this type.
  * @param script - The script to link the shader with.
  */
-FNiagaraShaderCompileJob* FNiagaraShaderType::BeginCompileShader(
+FShaderCompileJob* FNiagaraShaderType::BeginCompileShader(
 	uint32 ShaderMapId,
-	const FNiagaraScript* Script,
+	const FNiagaraShaderScript* Script,
 	FShaderCompilerEnvironment* CompilationEnvironment,
 	EShaderPlatform Platform,
-	TArray<FNiagaraShaderCompileJob*>& NewJobs,
+	TArray<FShaderCommonCompileJob*>& NewJobs,
 	FShaderTarget Target,
-	TArray< FDIBufferDescriptorStore >&InBufferDescriptors
+	TArray<FNiagaraDataInterfaceGPUParamInfo>& InDIParamInfo
 	)
 {
-	FNiagaraShaderCompileJob* NewJob = new FNiagaraShaderCompileJob(ShaderMapId, this, Script->HlslOutput);
+	FShaderCompileJob* NewJob = new FShaderCompileJob(ShaderMapId, nullptr, this, /* PermutationId = */ 0);
 
-	NewJob->DIBufferDescriptors = InBufferDescriptors;	// from hlsl translation and need to be passed to the FNiagaraShader on completion
+	//NewJob->DIParamInfo = InDIParamInfo;	// from hlsl translation and need to be passed to the FNiagaraShader on completion
 
 	NewJob->Input.SharedEnvironment = CompilationEnvironment;
 	NewJob->Input.Target = Target;
@@ -189,7 +232,10 @@ FNiagaraShaderCompileJob* FNiagaraShaderType::BeginCompileShader(
 	NewJob->Input.VirtualSourceFilePath = TEXT("/Engine/Private/NiagaraEmitterInstanceShader.usf");
 	NewJob->Input.EntryPointName = TEXT("SimulateMainComputeCS");
 	NewJob->Input.Environment.SetDefine(TEXT("GPU_SIMULATION"), 1);
-	NewJob->Input.Environment.IncludeVirtualPathToContentsMap.Add(TEXT("/Engine/Generated/NiagaraEmitterInstance.usf"), Script->HlslOutput);
+	NewJob->Input.Environment.IncludeVirtualPathToContentsMap.Add(TEXT("/Engine/Generated/NiagaraEmitterInstance.ush"), Script->HlslOutput);
+	
+	AddReferencedUniformBufferIncludes(NewJob->Input.Environment, NewJob->Input.SourceFilePrefix, (EShaderPlatform)Target.Platform);
+	
 	FShaderCompilerEnvironment& ShaderEnvironment = NewJob->Input.Environment;
 
 	UE_LOG(LogShaders, Verbose, TEXT("			%s"), GetName());
@@ -201,11 +247,81 @@ FNiagaraShaderCompileJob* FNiagaraShaderType::BeginCompileShader(
 	// Allow the shader type to modify the compile environment.
 	SetupCompileEnvironment(Platform, Script, ShaderEnvironment);
 
-	NewJobs.Add(NewJob);
+	FString DebugGroupName = Script->SourceName;
+	::GlobalBeginCompileShader(
+		DebugGroupName,
+		nullptr,
+		this,
+		nullptr,//ShaderPipeline,
+		TEXT("/Engine/Private/NiagaraEmitterInstanceShader.usf"),
+		TEXT("SimulateMainComputeCS"),
+		FShaderTarget(GetFrequency(), Platform),
+		NewJob,
+		NewJobs
+	);
+
+	ExtraParamInfo.Add(NewJob, InDIParamInfo);
+
+	//NewJobs.Add(NewJob);
 
 	return NewJob;
 }
 
+void FNiagaraShaderType::CacheUniformBufferIncludes(TMap<const TCHAR*, FCachedUniformBufferDeclaration>& Cache, EShaderPlatform Platform)
+{
+	for (TMap<const TCHAR*, FCachedUniformBufferDeclaration>::TIterator It(Cache); It; ++It)
+	{
+		FCachedUniformBufferDeclaration& BufferDeclaration = It.Value();
+
+		for (TLinkedList<FUniformBufferStruct*>::TIterator StructIt(FUniformBufferStruct::GetStructList()); StructIt; StructIt.Next())
+		{
+			if (It.Key() == StructIt->GetShaderVariableName())
+			{
+				BufferDeclaration.Declaration[Platform] = MakeShared<FString>();
+				CreateUniformBufferShaderDeclaration(StructIt->GetShaderVariableName(), **StructIt, Platform, *BufferDeclaration.Declaration[Platform].Get());
+				break;
+			}
+		}
+	}
+}
+
+
+
+void FNiagaraShaderType::AddReferencedUniformBufferIncludes(FShaderCompilerEnvironment& OutEnvironment, FString& OutSourceFilePrefix, EShaderPlatform Platform)
+{
+	// Cache uniform buffer struct declarations referenced by this shader type's files
+	if (!bCachedUniformBufferStructDeclarations[Platform])
+	{
+		CacheUniformBufferIncludes(ReferencedUniformBufferStructsCache, Platform);
+		bCachedUniformBufferStructDeclarations[Platform] = true;
+	}
+
+	FString UniformBufferIncludes;
+
+	for (TMap<const TCHAR*, FCachedUniformBufferDeclaration>::TConstIterator It(ReferencedUniformBufferStructsCache); It; ++It)
+	{
+		check(It.Value().Declaration[Platform].IsValid());
+		UniformBufferIncludes += FString::Printf(TEXT("#include \"/Engine/Generated/UniformBuffers/%s.ush\"") LINE_TERMINATOR, It.Key());
+		FString* Declaration = It.Value().Declaration[Platform].Get();
+		check(Declaration);
+		OutEnvironment.IncludeVirtualPathToContentsMap.Add(
+			*FString::Printf(TEXT("/Engine/Generated/UniformBuffers/%s.ush"), It.Key()), *Declaration
+		);
+
+		for (TLinkedList<FUniformBufferStruct*>::TIterator StructIt(FUniformBufferStruct::GetStructList()); StructIt; StructIt.Next())
+		{
+			if (It.Key() == StructIt->GetShaderVariableName())
+			{
+				StructIt->AddResourceTableEntries(OutEnvironment.ResourceTableMap, OutEnvironment.ResourceTableLayoutHashes);
+			}
+		}
+	}
+
+	FString& GeneratedUniformBuffersInclude = OutEnvironment.IncludeVirtualPathToContentsMap.FindOrAdd("/Engine/Generated/GeneratedUniformBuffers.ush");
+	GeneratedUniformBuffersInclude.Append(UniformBufferIncludes);
+}
+
+
 /**
  * Either creates a new instance of this type or returns an equivalent existing shader.
  * @param CurrentJob - Compile job that was enqueued by BeginCompileShader.
@@ -213,7 +329,7 @@ FNiagaraShaderCompileJob* FNiagaraShaderType::BeginCompileShader(
 FShader* FNiagaraShaderType::FinishCompileShader(
 	//const FUniformExpressionSet& UniformExpressionSet,
 	const FSHAHash& ShaderMapHash,
-	const FNiagaraShaderCompileJob& CurrentJob,
+	const FShaderCompileJob& CurrentJob,
 	const FString& InDebugDescription
 	)
 {
@@ -232,8 +348,17 @@ FShader* FNiagaraShaderType::FinishCompileShader(
 	if (!Shader)
 	{
 		const int32 PermutationId = 0;
-		Shader = (*ConstructCompiledRef)(FNiagaraShaderType::CompiledShaderInitializerType(this, PermutationId, CurrentJob.Output, Resource, ShaderMapHash, InDebugDescription, CurrentJob.DIBufferDescriptors));
-		CurrentJob.Output.ParameterMap.VerifyBindingsAreComplete(GetName(), CurrentJob.Output.Target, nullptr);
+		//CurrentJob.Id
+		TArray<FNiagaraDataInterfaceGPUParamInfo> DIParamInfo;
+
+		if (!ExtraParamInfo.RemoveAndCopyValue(&CurrentJob, DIParamInfo))
+		{
+			check(false);
+		}
+
+		Shader = (*ConstructCompiledRef)(FNiagaraShaderType::CompiledShaderInitializerType(this, PermutationId, CurrentJob.Output, Resource, ShaderMapHash, InDebugDescription, DIParamInfo));
+		//CurrentJob.Output.ParameterMap.VerifyBindingsAreComplete(GetName(), CurrentJob.Output.Target, nullptr); // b/c we don't bind data interfaces yet...
+
 	}
 
 	return Shader;
@@ -304,7 +429,7 @@ static FString GetNiagaraShaderMapKeyString(const FNiagaraShaderMapId& ShaderMap
 
 
 
-void FNiagaraShaderMap::LoadFromDerivedDataCache(const FNiagaraScript* Script, const FNiagaraShaderMapId& ShaderMapId, EShaderPlatform Platform, TRefCountPtr<FNiagaraShaderMap>& InOutShaderMap)
+void FNiagaraShaderMap::LoadFromDerivedDataCache(const FNiagaraShaderScript* Script, const FNiagaraShaderMapId& ShaderMapId, EShaderPlatform Platform, TRefCountPtr<FNiagaraShaderMap>& InOutShaderMap)
 {
 	if (InOutShaderMap != NULL)
 	{
@@ -572,7 +697,7 @@ void FNiagaraShaderMap::LoadForRemoteRecompile(FArchive& Ar, EShaderPlatform Sha
 * @param Platform - The platform to compile to
 */
 void FNiagaraShaderMap::Compile(
-	FNiagaraScript* Script,
+	FNiagaraShaderScript* Script,
 	const FNiagaraShaderMapId& InShaderMapId,
 	TRefCountPtr<FShaderCompilerEnvironment> CompilationEnvironment,
 	const FNiagaraComputeShaderCompilationOutput& InNiagaraCompilationOutput,
@@ -591,7 +716,7 @@ void FNiagaraShaderMap::Compile(
 		check(NumRefs > 0);
   
 		// Add this shader map and to NiagaraShaderMapsBeingCompiled
-		TArray<FNiagaraScript*>* CorrespondingScripts = NiagaraShaderMapsBeingCompiled.Find(this);
+		TArray<FNiagaraShaderScript*>* CorrespondingScripts = NiagaraShaderMapsBeingCompiled.Find(this);
   
 		if (CorrespondingScripts)
 		{
@@ -600,12 +725,16 @@ void FNiagaraShaderMap::Compile(
 		}
 		else
 		{
+			Script->RemoveOutstandingCompileId(CompilingId);
 			// Assign a unique identifier so that shaders from this shader map can be associated with it after a deferred compile
 			CompilingId = NextCompilingId;
+			UE_LOG(LogShaders, Log, TEXT("CompilingId = %p %d"), Script, CompilingId);
+			Script->AddCompileId(CompilingId);
+
 			check(NextCompilingId < UINT_MAX);
 			NextCompilingId++;
   
-			TArray<FNiagaraScript*> NewCorrespondingScripts;
+			TArray<FNiagaraShaderScript*> NewCorrespondingScripts;
 			NewCorrespondingScripts.Add(Script);
 			NiagaraShaderMapsBeingCompiled.Add(this, NewCorrespondingScripts);
 #if DEBUG_INFINITESHADERCOMPILE
@@ -621,10 +750,10 @@ void FNiagaraShaderMap::Compile(
 			Platform = InPlatform;
 
 			uint32 NumShaders = 0;
-			TArray<FNiagaraShaderCompileJob*> NewJobs;
+			TArray<FShaderCommonCompileJob*> NewJobs;
 	
 			// Iterate over all shader types.
-			TMap<FShaderType*, FNiagaraShaderCompileJob*> SharedShaderJobs;
+			TMap<FShaderType*, FShaderCompileJob*> SharedShaderJobs;
 			for(TLinkedList<FShaderType*>::TIterator ShaderTypeIt(FShaderType::GetTypeList());ShaderTypeIt;ShaderTypeIt.Next())
 			{
 				FNiagaraShaderType* ShaderType = ShaderTypeIt->GetNiagaraShaderType();
@@ -643,7 +772,7 @@ void FNiagaraShaderMap::Compile(
 							Platform,
 							NewJobs,
 							FShaderTarget(ShaderType->GetFrequency(), Platform),
-							Script->GetDataInterfaceBufferDescriptors()
+							Script->GetDataInterfaceParamInfo()
 							);
 						check(!SharedShaderJobs.Find(ShaderType));
 						SharedShaderJobs.Add(ShaderType, Job);
@@ -654,7 +783,7 @@ void FNiagaraShaderMap::Compile(
   
 			if (!CorrespondingScripts)
 			{
-				UE_LOG(LogShaders, Warning, TEXT("		%u Shaders"), NumShaders);
+				UE_LOG(LogShaders, Log, TEXT("		%u Shaders"), NumShaders);
 			}
 
 			// Register this shader map in the global script->shadermap map
@@ -679,10 +808,10 @@ void FNiagaraShaderMap::Compile(
 	}
 }
 
-FShader* FNiagaraShaderMap::ProcessCompilationResultsForSingleJob(FNiagaraShaderCompileJob* SingleJob, const FSHAHash& ShaderMapHash)
+FShader* FNiagaraShaderMap::ProcessCompilationResultsForSingleJob(FShaderCommonCompileJob* SingleJob, const FSHAHash& ShaderMapHash)
 {
 	check(SingleJob);
-	const FNiagaraShaderCompileJob& CurrentJob = *SingleJob;
+	const FShaderCompileJob& CurrentJob = *((FShaderCompileJob*)SingleJob);
 	check(CurrentJob.Id == CompilingId);
 
 	FShader* Shader = nullptr;
@@ -690,6 +819,8 @@ FShader* FNiagaraShaderMap::ProcessCompilationResultsForSingleJob(FNiagaraShader
 	FNiagaraShaderType* NiagaraShaderType = CurrentJob.ShaderType->GetNiagaraShaderType();
 	check(NiagaraShaderType);
 	Shader = NiagaraShaderType->FinishCompileShader(ShaderMapHash, CurrentJob, FriendlyName);
+	bCompiledSuccessfully = CurrentJob.bSucceeded;
+
 	FNiagaraShader *NiagaraShader = static_cast<FNiagaraShader*>(Shader);
 	check(Shader);
 	check(!HasShader(NiagaraShaderType, /* PermutationId = */ 0));
@@ -698,7 +829,7 @@ FShader* FNiagaraShaderMap::ProcessCompilationResultsForSingleJob(FNiagaraShader
 	return Shader;
 }
 
-bool FNiagaraShaderMap::ProcessCompilationResults(const TArray<FNiagaraShaderCompileJob*>& InCompilationResults, int32& InOutJobIndex, float& TimeBudget)
+bool FNiagaraShaderMap::ProcessCompilationResults(const TArray<FShaderCommonCompileJob*>& InCompilationResults, int32& InOutJobIndex, float& TimeBudget)
 {
 	check(InOutJobIndex < InCompilationResults.Num());
 
@@ -709,7 +840,7 @@ bool FNiagaraShaderMap::ProcessCompilationResults(const TArray<FNiagaraShaderCom
 
 	do
 	{
-		FNiagaraShaderCompileJob* SingleJob = InCompilationResults[InOutJobIndex];
+		FShaderCommonCompileJob* SingleJob = InCompilationResults[InOutJobIndex];
 		ensure(SingleJob);
 
 		{
@@ -735,16 +866,16 @@ bool FNiagaraShaderMap::ProcessCompilationResults(const TArray<FNiagaraShaderCom
 	return false;
 }
 
-bool FNiagaraShaderMap::TryToAddToExistingCompilationTask(FNiagaraScript* Script)
+bool FNiagaraShaderMap::TryToAddToExistingCompilationTask(FNiagaraShaderScript* Script)
 {
 	check(NumRefs > 0);
-	TArray<FNiagaraScript*>* CorrespondingScripts = FNiagaraShaderMap::NiagaraShaderMapsBeingCompiled.Find(this);
+	TArray<FNiagaraShaderScript*>* CorrespondingScripts = FNiagaraShaderMap::NiagaraShaderMapsBeingCompiled.Find(this);
 
 	if (CorrespondingScripts)
 	{
 		CorrespondingScripts->AddUnique(Script);
 
-		//UE_LOG(LogShaders, Error, TEXT("TryToAddToExistingCompilationTask %p %d"), Script, GetCompilingId());
+		UE_LOG(LogShaders, Log, TEXT("TryToAddToExistingCompilationTask %p %d"), Script, GetCompilingId());
 
 #if DEBUG_INFINITESHADERCOMPILE
 		UE_LOG(LogTemp, Display, TEXT("Added shader map 0x%08X%08X from Niagara script 0x%08X%08X"), (int)((int64)(this) >> 32), (int)((int64)(this)), (int)((int64)(Script) >> 32), (int)((int64)(Script)));
@@ -755,7 +886,7 @@ bool FNiagaraShaderMap::TryToAddToExistingCompilationTask(FNiagaraScript* Script
 	return false;
 }
 
-bool FNiagaraShaderMap::IsNiagaraShaderComplete(const FNiagaraScript* Script, const FNiagaraShaderType* ShaderType, bool bSilent)
+bool FNiagaraShaderMap::IsNiagaraShaderComplete(const FNiagaraShaderScript* Script, const FNiagaraShaderType* ShaderType, bool bSilent)
 {
 	// If we should cache this script, it's incomplete if the shader is missing
 	if (ShouldCacheNiagaraShader(ShaderType, Platform, Script) &&	!HasShader((FShaderType*)ShaderType, /* PermutationId = */ 0))
@@ -770,13 +901,13 @@ bool FNiagaraShaderMap::IsNiagaraShaderComplete(const FNiagaraScript* Script, co
 	return true;
 }
 
-bool FNiagaraShaderMap::IsComplete(const FNiagaraScript* Script, bool bSilent)
+bool FNiagaraShaderMap::IsComplete(const FNiagaraShaderScript* Script, bool bSilent)
 {
 	check(!IsInRenderingThread());
 	// Make sure we are operating on a referenced shader map or the below Find will cause this shader map to be deleted,
 	// Since it creates a temporary ref counted pointer.
 	check(NumRefs > 0);
-	const TArray<FNiagaraScript*>* CorrespondingScripts = FNiagaraShaderMap::NiagaraShaderMapsBeingCompiled.Find(this);
+	const TArray<FNiagaraShaderScript*>* CorrespondingScripts = FNiagaraShaderMap::NiagaraShaderMapsBeingCompiled.Find(this);
 
 	if (CorrespondingScripts)
 	{
@@ -798,13 +929,13 @@ bool FNiagaraShaderMap::IsComplete(const FNiagaraScript* Script, bool bSilent)
 	return true;
 }
 
-void FNiagaraShaderMap::LoadMissingShadersFromMemory(const FNiagaraScript* Script)
+void FNiagaraShaderMap::LoadMissingShadersFromMemory(const FNiagaraShaderScript* Script)
 {
 	// Make sure we are operating on a referenced shader map or the below Find will cause this shader map to be deleted,
 	// Since it creates a temporary ref counted pointer.
 	check(NumRefs > 0);
 
-	const TArray<FNiagaraScript*>* CorrespondingScripts = FNiagaraShaderMap::NiagaraShaderMapsBeingCompiled.Find(this);
+	const TArray<FNiagaraShaderScript*>* CorrespondingScripts = FNiagaraShaderMap::NiagaraShaderMapsBeingCompiled.Find(this);
 
 	if (CorrespondingScripts)
 	{
@@ -971,11 +1102,11 @@ void FNiagaraShaderMap::DiscardSerializedShaders()
 	TShaderMap<FNiagaraShaderType>::DiscardSerializedShaders();
 }
 
-void FNiagaraShaderMap::RemovePendingScript(FNiagaraScript* Script)
+void FNiagaraShaderMap::RemovePendingScript(FNiagaraShaderScript* Script)
 {
-	for (TMap<TRefCountPtr<FNiagaraShaderMap>, TArray<FNiagaraScript*> >::TIterator It(NiagaraShaderMapsBeingCompiled); It; ++It)
+	for (TMap<TRefCountPtr<FNiagaraShaderMap>, TArray<FNiagaraShaderScript*> >::TIterator It(NiagaraShaderMapsBeingCompiled); It; ++It)
 	{
-		TArray<FNiagaraScript*>& Scripts = It.Value();
+		TArray<FNiagaraShaderScript*>& Scripts = It.Value();
 		int32 Result = Scripts.Remove(Script);
 		if (Result)
 		{
@@ -994,10 +1125,10 @@ void FNiagaraShaderMap::RemovePendingScript(FNiagaraScript* Script)
 
 void FNiagaraShaderMap::RemovePendingMap(FNiagaraShaderMap* Map)
 {
-	TArray<FNiagaraScript*>* Scripts = NiagaraShaderMapsBeingCompiled.Find(Map);
+	TArray<FNiagaraShaderScript*>* Scripts = NiagaraShaderMapsBeingCompiled.Find(Map);
 	if (Scripts)
 	{
-		for (FNiagaraScript* Script : *Scripts)
+		for (FNiagaraShaderScript* Script : *Scripts)
 		{
 			Script->RemoveOutstandingCompileId(Map->CompilingId);
 			Script->NotifyCompilationFinished();
@@ -1007,12 +1138,12 @@ void FNiagaraShaderMap::RemovePendingMap(FNiagaraShaderMap* Map)
 	NiagaraShaderMapsBeingCompiled.Remove(Map);
 }
 
-const FNiagaraShaderMap* FNiagaraShaderMap::GetShaderMapBeingCompiled(const FNiagaraScript* Script)
+const FNiagaraShaderMap* FNiagaraShaderMap::GetShaderMapBeingCompiled(const FNiagaraShaderScript* Script)
 {
 	// Inefficient search, but only when compiling a lot of shaders
-	for (TMap<TRefCountPtr<FNiagaraShaderMap>, TArray<FNiagaraScript*> >::TIterator It(NiagaraShaderMapsBeingCompiled); It; ++It)
+	for (TMap<TRefCountPtr<FNiagaraShaderMap>, TArray<FNiagaraShaderScript*> >::TIterator It(NiagaraShaderMapsBeingCompiled); It; ++It)
 	{
-		TArray<FNiagaraScript*>& Scripts = It.Value();
+		TArray<FNiagaraShaderScript*>& Scripts = It.Value();
 		
 		for (int32 ScriptIndex = 0; ScriptIndex < Scripts.Num(); ScriptIndex++)
 		{
@@ -1026,7 +1157,7 @@ const FNiagaraShaderMap* FNiagaraShaderMap::GetShaderMapBeingCompiled(const FNia
 	return NULL;
 }
 
-
+//////////////////////////////////////////////////////////////////////////
 
 
 FNiagaraShader::FNiagaraShader(const FNiagaraShaderType::CompiledShaderInitializerType& Initializer)
@@ -1034,32 +1165,96 @@ FNiagaraShader::FNiagaraShader(const FNiagaraShaderType::CompiledShaderInitializ
 	, DebugDescription(Initializer.DebugDescription)
 {
 	check(!DebugDescription.IsEmpty());
-	SetDatainterfaceBufferDescriptors(Initializer.DIBufferDescriptors);
+	SetDataInterfaceParameterInfo(Initializer.DIParamInfo);
 	BindParams(Initializer.ParameterMap);
 }
 
+void FNiagaraShader::SetDataInterfaceParameterInfo(const TArray<FNiagaraDataInterfaceGPUParamInfo>& InDIParamInfo)
+{
+	DataInterfaceParameters.Empty();
+	for (const FNiagaraDataInterfaceGPUParamInfo& DIParamInfo : InDIParamInfo)
+	{
+		DataInterfaceParameters.Emplace(DIParamInfo);
+	}
+}
+
+void FNiagaraShader::BindParams(const FShaderParameterMap &ParameterMap)
+{
+	FloatInputBufferParam.Bind(ParameterMap, TEXT("InputFloat"));
+	IntInputBufferParam.Bind(ParameterMap, TEXT("InputInt"));
+	FloatOutputBufferParam.Bind(ParameterMap, TEXT("OutputFloat"));
+	IntOutputBufferParam.Bind(ParameterMap, TEXT("OutputInt"));
+	OutputIndexBufferParam.Bind(ParameterMap, TEXT("DataSetIndices"));
+	InputIndexBufferParam.Bind(ParameterMap, TEXT("ReadDataSetIndices"));
+	EmitterTickCounterParam.Bind(ParameterMap, TEXT("EmitterTickCounter"));
+	NumEventsPerParticleParam.Bind(ParameterMap, TEXT("NumEventsPerParticle"));
+	NumParticlesPerEventParam.Bind(ParameterMap, TEXT("NumParticlesPerEvent"));
+	CopyInstancesBeforeStartParam.Bind(ParameterMap, TEXT("CopyInstancesBeforeStart"));
+
+	NumSpawnedInstancesParam.Bind(ParameterMap, TEXT("SpawnedInstances"));
+	UpdateStartInstanceParam.Bind(ParameterMap, TEXT("UpdateStartInstance"));
+	NumIndicesPerInstanceParam.Bind(ParameterMap, TEXT("NumIndicesPerInstance"));
+
+	ComponentBufferSizeReadParam.Bind(ParameterMap, TEXT("ComponentBufferSizeRead"));
+	ComponentBufferSizeWriteParam.Bind(ParameterMap, TEXT("ComponentBufferSizeWrite"));
+	EmitterConstantBufferParam.Bind(ParameterMap, TEXT("FEmitterParameters"));
+	ViewUniformBufferParam.Bind(ParameterMap, TEXT("View"));
+
+	// params for event buffers
+	// this is horrendous; need to do this in a uniform buffer instead.
+	//
+	for (uint32 i = 0; i < MAX_CONCURRENT_EVENT_DATASETS; i++)
+	{
+		FString VarName = TEXT("WriteDataSetFloat") + FString::FromInt(i + 1);
+		EventFloatUAVParams[i].Bind(ParameterMap, *VarName);
+		VarName = TEXT("WriteDataSetInt") + FString::FromInt(i + 1);
+		EventIntUAVParams[i].Bind(ParameterMap, *VarName);
+
+		VarName = TEXT("ReadDataSetFloat") + FString::FromInt(i + 1);
+		EventFloatSRVParams[i].Bind(ParameterMap, *VarName);
+		VarName = TEXT("ReadDataSetInt") + FString::FromInt(i + 1);
+		EventIntSRVParams[i].Bind(ParameterMap, *VarName);
+
+		VarName = TEXT("DSComponentBufferSizeReadFloat") + FString::FromInt(i + 1);
+		EventReadFloatStrideParams[i].Bind(ParameterMap, *VarName);
+		VarName = TEXT("DSComponentBufferSizeWriteFloat") + FString::FromInt(i + 1);
+		EventWriteFloatStrideParams[i].Bind(ParameterMap, *VarName);
+		VarName = TEXT("DSComponentBufferSizeReadInt") + FString::FromInt(i + 1);
+		EventReadIntStrideParams[i].Bind(ParameterMap, *VarName);
+		VarName = TEXT("DSComponentBufferSizeWriteInt") + FString::FromInt(i + 1);
+		EventWriteIntStrideParams[i].Bind(ParameterMap, *VarName);
+	}
 
+	for(FNiagaraDataInterfaceParamRef& Param : DataInterfaceParameters)
+	{
+		Param.Bind(ParameterMap);
+	}
+
+	ensure(FloatOutputBufferParam.IsBound() || IntOutputBufferParam.IsBound());	// we should have at least one output buffer we're writing to
+	ensure(ComponentBufferSizeWriteParam.IsBound());
+	ensure(OutputIndexBufferParam.IsBound());
+	ensure(InputIndexBufferParam.IsBound());
+	ensure(UpdateStartInstanceParam.IsBound());
+	ensure(NumSpawnedInstancesParam.IsBound());
+	ensure(NumIndicesPerInstanceParam.IsBound());
+}
 
 bool FNiagaraShader::Serialize(FArchive& Ar)
 {
 	const bool bShaderHasOutdatedParameters = FShader::Serialize(Ar);
 	Ar << NiagaraUniformBuffer;	// do we need to keep about uniform buffer parameters?
 
-//	Ar.UsingCustomVersion(FNiagaraShaderCustomVersion::GUID);		// only changes version if not loading
-//	const int32 NiagaraVer = Ar.CustomVer(FNiagaraShaderCustomVersion::GUID);
-
 	Ar << DebugDescription;
 
 	Ar << FloatInputBufferParam;
 	Ar << IntInputBufferParam;
 	Ar << FloatOutputBufferParam;
 	Ar << IntOutputBufferParam;
-	Ar << OutputIndexBufferParam;
+	Ar << InputIndexBufferParam;
 	Ar << EmitterTickCounterParam;
 
-	Ar << TotalNumInstancesParam;
-	Ar << StartInstancePhase0Param;
-	Ar << StartInstancePhase1Param;
+	Ar << NumSpawnedInstancesParam;
+	Ar << UpdateStartInstanceParam;
 	Ar << ComponentBufferSizeReadParam;
 	Ar << ComponentBufferSizeWriteParam;
 
@@ -1071,46 +1266,20 @@ bool FNiagaraShader::Serialize(FArchive& Ar)
 		Ar << EventIntSRVParams[i];
 	}
 
-	Ar << DIBufferDescriptors;
-
-	// params for data interface buffers
-	if (Ar.IsLoading())
-	{
-		for (FDIBufferDescriptorStore &InterfaceDescs : DIBufferDescriptors)
-		{
-			int32 Idx = NameToDIBufferParamMap.AddDefaulted(1);
-			for (FDIGPUBufferParamDescriptor &Desc : InterfaceDescs.Descriptors)
-			{
-				FShaderResourceParameter &Param = NameToDIBufferParamMap[Idx].Add(*Desc.BufferParamName);
-				Ar << Param;
-			}
-		}
-	}
-	else
-	{
-		uint32 Idx = 0;
-		for (FDIBufferDescriptorStore &InterfaceDescs : DIBufferDescriptors)
-		{
-			for (FDIGPUBufferParamDescriptor &Desc : InterfaceDescs.Descriptors)
-			{
-				FShaderResourceParameter *Param = NameToDIBufferParamMap[Idx].Find(*Desc.BufferParamName);
-				check(Param);
-				Ar << *Param;
-			}
-			Idx++;
-		}
-	}
+	Ar << DataInterfaceParameters;
 
 	Ar << EmitterConstantBufferParam;
 	Ar << DataInterfaceUniformBufferParam;
 	Ar << NumEventsPerParticleParam;
 	Ar << NumParticlesPerEventParam;
 	Ar << CopyInstancesBeforeStartParam;
+	Ar << ViewUniformBufferParam;
 
+	Ar << OutputIndexBufferParam;
+	Ar << NumIndicesPerInstanceParam;
 	return bShaderHasOutdatedParameters;
 }
 
-
 uint32 FNiagaraShader::GetAllocatedSize() const
 {
 	return FShader::GetAllocatedSize()
@@ -1118,17 +1287,75 @@ uint32 FNiagaraShader::GetAllocatedSize() const
 		+ DebugDescription.GetAllocatedSize();
 }
 
-FArchive &operator << (FArchive &Ar, FDIBufferDescriptorStore &Store)
+//////////////////////////////////////////////////////////////////////////
+FNiagaraDataInterfaceParamRef::FNiagaraDataInterfaceParamRef(const FNiagaraDataInterfaceGPUParamInfo& InParameterInfo)
+	: ParameterInfo(InParameterInfo)
+	, Parameters(nullptr)
 {
-	Ar << Store.DataInterfaceName;
-	Ar << Store.Descriptors;
-	return Ar;
+	InitDIClass();
 }
 
+FNiagaraDataInterfaceParamRef::FNiagaraDataInterfaceParamRef()
+	: DIClass(nullptr)
+	, Parameters(nullptr)
+{
+
+}
+
+FNiagaraDataInterfaceParamRef::~FNiagaraDataInterfaceParamRef()
+{
+	if (Parameters)
+	{
+		delete Parameters;
+	}
+}
+
+void FNiagaraDataInterfaceParamRef::ConstructParameters()
+{
+	if (Parameters)
+	{
+		delete Parameters;
+	}
+
+	Parameters = CastChecked<UNiagaraDataInterfaceBase>(DIClass->GetDefaultObject())->ConstructComputeParameters();
+}
 
-FArchive &operator << (FArchive &Ar, FDIGPUBufferParamDescriptor &Desc)
+void FNiagaraDataInterfaceParamRef::InitDIClass()
 {
-	Ar << Desc.BufferParamName;
-	Ar << Desc.Index;
-	return Ar;
+	DIClass = Cast<UClass>(StaticFindObject(UClass::StaticClass(), ANY_PACKAGE, *ParameterInfo.DIClassName, true));
+	ensureMsgf(DIClass, TEXT("Failed to load class for FNiagaraDataInterfaceParamRef. %s"), *ParameterInfo.DIClassName);
+}
+
+void FNiagaraDataInterfaceParamRef::Bind(const class FShaderParameterMap& ParameterMap)
+{
+	ConstructParameters();
+	if (Parameters)
+	{
+		Parameters->Bind(*this, ParameterMap);
+	}
+}
+
+bool operator<<(FArchive &Ar, FNiagaraDataInterfaceParamRef& ParamRef)
+{
+	ParamRef.ParameterInfo.Serialize(Ar);
+
+	if (Ar.IsLoading())
+	{
+		ParamRef.InitDIClass();
+		ParamRef.ConstructParameters();
+	}
+
+	if (ParamRef.Parameters)
+	{
+		ParamRef.Parameters->Serialize(Ar);
+	}
+
+	return true;
+}
+
+bool FNiagaraDataInterfaceGPUParamInfo::Serialize(FArchive& Ar)
+{
+	Ar << DataInterfaceHLSLSymbol;
+	Ar << DIClassName;
+	return true;
 }
diff --git a/Engine/Plugins/FX/Niagara/Source/NiagaraShader/Private/NiagaraShaderCompilationManager.cpp b/Engine/Plugins/FX/Niagara/Source/NiagaraShader/Private/NiagaraShaderCompilationManager.cpp
--- a/Engine/Plugins/FX/Niagara/Source/NiagaraShader/Private/NiagaraShaderCompilationManager.cpp
+++ b/Engine/Plugins/FX/Niagara/Source/NiagaraShader/Private/NiagaraShaderCompilationManager.cpp
@@ -6,6 +6,7 @@
 #include "Interfaces/IShaderFormat.h"
 #include "Interfaces/ITargetPlatformManagerModule.h"
 #endif
+#include "ShaderCompiler.h"
 #include "HAL/FileManager.h"
 #include "Misc/Paths.h"
 
@@ -88,7 +89,7 @@ void FNiagaraShaderCompilationManager::RunCompileJobs()
 		{
 			for (int32 JobIndex = 0; JobIndex < CurrentWorkerInfo.QueuedJobs.Num(); JobIndex++)
 			{
-				FNiagaraShaderCompileJob& CurrentJob = *CurrentWorkerInfo.QueuedJobs[JobIndex];
+				FShaderCompileJob& CurrentJob = *((FShaderCompileJob*)(CurrentWorkerInfo.QueuedJobs[JobIndex]));
 
 				check(!CurrentJob.bFinalized);
 				CurrentJob.bFinalized = true;
@@ -104,9 +105,11 @@ void FNiagaraShaderCompilationManager::RunCompileJobs()
 				}
 				CA_ASSUME(Compiler != NULL);
 
+				UE_LOG(LogNiagaraShaderCompiler, Log, TEXT("Compile Job processing... %s"), *CurrentJob.Input.DebugGroupName);
+
 				FString AbsoluteDebugInfoDirectory = IFileManager::Get().ConvertToAbsolutePathForExternalAppForWrite(*(FPaths::ProjectSavedDir() / TEXT("ShaderDebugInfo")));
 				FPaths::NormalizeDirectoryName(AbsoluteDebugInfoDirectory);
-				CurrentJob.Input.DumpDebugInfoPath = AbsoluteDebugInfoDirectory / Format.ToString() / FString("Niagara");
+				CurrentJob.Input.DumpDebugInfoPath = AbsoluteDebugInfoDirectory / Format.ToString() / CurrentJob.Input.DebugGroupName;
 				if (!IFileManager::Get().DirectoryExists(*CurrentJob.Input.DumpDebugInfoPath))
 				{
 					verifyf(IFileManager::Get().MakeDirectory(*CurrentJob.Input.DumpDebugInfoPath, true), TEXT("Failed to create directory for shader debug info '%s'"), *CurrentJob.Input.DumpDebugInfoPath);
@@ -162,9 +165,9 @@ void FNiagaraShaderCompilationManager::RunCompileJobs()
 
 
 
-NIAGARASHADER_API void FNiagaraShaderCompilationManager::AddJobs(TArray<FNiagaraShaderCompileJob*> InNewJobs)
+NIAGARASHADER_API void FNiagaraShaderCompilationManager::AddJobs(TArray<FShaderCommonCompileJob*> InNewJobs)
 {
-	for (FNiagaraShaderCompileJob *Job : InNewJobs)
+	for (FShaderCommonCompileJob *Job : InNewJobs)
 	{
 		FNiagaraShaderMapCompileResults& ShaderMapInfo = NiagaraShaderMapJobs.FindOrAdd(Job->Id);
 //		ShaderMapInfo.bApplyCompletedShaderMapForRendering = bApplyCompletedShaderMapForRendering;
@@ -214,16 +217,16 @@ void FNiagaraShaderCompilationManager::ProcessCompiledNiagaraShaderMaps(
 {
 	// Keeps shader maps alive as they are passed from the shader compiler and applied to the owning Script
 	TArray<TRefCountPtr<FNiagaraShaderMap> > LocalShaderMapReferences;
-	TMap<FNiagaraScript*, FNiagaraShaderMap*> ScriptsToUpdate;
+	TMap<FNiagaraShaderScript*, FNiagaraShaderMap*> ScriptsToUpdate;
 
 	// Process compiled shader maps in FIFO order, in case a shader map has been enqueued multiple times,
 	// Which can happen if a script is edited while a background compile is going on
 	for (TMap<int32, FNiagaraShaderMapFinalizeResults>::TIterator ProcessIt(CompiledShaderMaps); ProcessIt; ++ProcessIt)
 	{
 		TRefCountPtr<FNiagaraShaderMap> ShaderMap = NULL;
-		TArray<FNiagaraScript*>* Scripts = NULL;
+		TArray<FNiagaraShaderScript*>* Scripts = NULL;
 
-		for (TMap<TRefCountPtr<FNiagaraShaderMap>, TArray<FNiagaraScript*> >::TIterator ShaderMapIt(FNiagaraShaderMap::GetInFlightShaderMaps()); ShaderMapIt; ++ShaderMapIt)
+		for (TMap<TRefCountPtr<FNiagaraShaderMap>, TArray<FNiagaraShaderScript*> >::TIterator ShaderMapIt(FNiagaraShaderMap::GetInFlightShaderMaps()); ShaderMapIt; ++ShaderMapIt)
 		{
 			if (ShaderMapIt.Key()->GetCompilingId() == ProcessIt.Key())
 			{
@@ -237,15 +240,15 @@ void FNiagaraShaderCompilationManager::ProcessCompiledNiagaraShaderMaps(
 		{
 			TArray<FString> Errors;
 			FNiagaraShaderMapFinalizeResults& CompileResults = ProcessIt.Value();
-			const TArray<FNiagaraShaderCompileJob*>& ResultArray = CompileResults.FinishedJobs;
+			const TArray<FShaderCommonCompileJob*>& ResultArray = CompileResults.FinishedJobs;
 
 			// Make a copy of the array as this entry of FNiagaraShaderMap::ShaderMapsBeingCompiled will be removed below
-			TArray<FNiagaraScript*> ScriptArray = *Scripts;
+			TArray<FNiagaraShaderScript*> ScriptArray = *Scripts;
 			bool bSuccess = true;
 
 			for (int32 JobIndex = 0; JobIndex < ResultArray.Num(); JobIndex++)
 			{
-				FNiagaraShaderCompileJob& CurrentJob = *ResultArray[JobIndex];
+				FShaderCompileJob& CurrentJob = *((FShaderCompileJob*)(ResultArray[JobIndex]));
 				bSuccess = bSuccess && CurrentJob.bSucceeded;
 
 				if (bSuccess)
@@ -259,6 +262,19 @@ void FNiagaraShaderCompilationManager::ProcessCompiledNiagaraShaderMaps(
 					{
 						Errors.AddUnique(CurrentJob.Output.Errors[ErrorIndex].GetErrorString());
 					}
+
+					if (CurrentJob.Output.Errors.Num())
+					{
+						UE_LOG(LogShaders, Log, TEXT("There were errors for job \"%s\""), *CurrentJob.Input.DebugGroupName)
+							for (const FShaderCompilerError& Error : CurrentJob.Output.Errors)
+							{
+								UE_LOG(LogShaders, Log, TEXT("Error: %s"), *Error.GetErrorString())
+							}
+					}
+				}
+				else
+				{
+					UE_LOG(LogShaders, Log, TEXT("There were NO errors for job \"%s\""), *CurrentJob.Input.DebugGroupName);
 				}
 			}
 
@@ -278,7 +294,7 @@ void FNiagaraShaderCompilationManager::ProcessCompiledNiagaraShaderMaps(
 				LocalShaderMapReferences.Add(ShaderMap);
 				FNiagaraShaderMap::GetInFlightShaderMaps().Remove(ShaderMap);
 
-				for (FNiagaraScript *Script : ScriptArray)
+				for (FNiagaraShaderScript* Script : ScriptArray)
 				{
 					FNiagaraShaderMap* CompletedShaderMap = ShaderMap;
 
@@ -286,7 +302,7 @@ void FNiagaraShaderCompilationManager::ProcessCompiledNiagaraShaderMaps(
 
 					// Only process results that still match the ID which requested a compile
 					// This avoids applying shadermaps which are out of date and a newer one is in the async compiling pipeline
-					if (Script->GetScriptID() == CompletedShaderMap->GetShaderMapId().BaseScriptID)
+					if (Script->IsSame(CompletedShaderMap->GetShaderMapId()))
 					{
 						if (Errors.Num() != 0)
 						{
@@ -366,9 +382,9 @@ void FNiagaraShaderCompilationManager::ProcessCompiledNiagaraShaderMaps(
 
 	if (ScriptsToUpdate.Num() > 0)
 	{
-		for (TMap<FNiagaraScript*, FNiagaraShaderMap*>::TConstIterator It(ScriptsToUpdate); It; ++It)
+		for (TMap<FNiagaraShaderScript*, FNiagaraShaderMap*>::TConstIterator It(ScriptsToUpdate); It; ++It)
 		{
-			FNiagaraScript* Script = It.Key();
+			FNiagaraShaderScript* Script = It.Key();
 			FNiagaraShaderMap* ShaderMap = It.Value();
 			//check(!ShaderMap || ShaderMap->IsValidForRendering());
 
@@ -376,7 +392,7 @@ void FNiagaraShaderCompilationManager::ProcessCompiledNiagaraShaderMaps(
 
 			ENQUEUE_UNIQUE_RENDER_COMMAND_TWOPARAMETER(
 				FSetShaderMapOnScriptResources,
-				FNiagaraScript*, Script, Script,
+				FNiagaraShaderScript*, Script, Script,
 				FNiagaraShaderMap*, CompiledShaderMap, ShaderMap,
 				{
 					Script->SetRenderingThreadShaderMap(CompiledShaderMap);
@@ -395,6 +411,8 @@ void FNiagaraShaderCompilationManager::FinishCompilation(const TCHAR* ScriptName
 
 	RunCompileJobs();	// since we don't async compile through another process, this will run all oustanding jobs
 	ProcessAsyncResults();	// grab compiled shader maps and assign them to their resources
+
+	check(NiagaraShaderMapJobs.Num() == 0);
 }
 
 
diff --git a/Engine/Plugins/FX/Niagara/Source/NiagaraShader/Private/NiagaraShaderDerivedDataVersion.h b/Engine/Plugins/FX/Niagara/Source/NiagaraShader/Private/NiagaraShaderDerivedDataVersion.h
--- a/Engine/Plugins/FX/Niagara/Source/NiagaraShader/Private/NiagaraShaderDerivedDataVersion.h
+++ b/Engine/Plugins/FX/Niagara/Source/NiagaraShader/Private/NiagaraShaderDerivedDataVersion.h
@@ -11,5 +11,5 @@ NiagaraShaderDerivedDataVersion.h: Shader derived data version for Niagara.
 // In case of merge conflicts with DDC versions, you MUST generate a new GUID and set this new
 // guid as version
 
-#define NIAGARASHADERMAP_DERIVEDDATA_VER		TEXT("1ecd9c26d37a443a9aa832b8feb72476")
+#define NIAGARASHADERMAP_DERIVEDDATA_VER		TEXT("9B98AE21EB7347A99AB0862D8B1EB8A6}")
 													  
diff --git a/Engine/Plugins/FX/Niagara/Source/NiagaraShader/Private/NiagaraShared.cpp b/Engine/Plugins/FX/Niagara/Source/NiagaraShader/Private/NiagaraShared.cpp
--- a/Engine/Plugins/FX/Niagara/Source/NiagaraShader/Private/NiagaraShared.cpp
+++ b/Engine/Plugins/FX/Niagara/Source/NiagaraShader/Private/NiagaraShared.cpp
@@ -17,12 +17,13 @@
 #include "NiagaraShaderCompilationManager.h"
 #include "RendererInterface.h"
 #include "Modules/ModuleManager.h"
+#include "NiagaraCustomVersion.h"
 
 #if WITH_EDITOR
 	NIAGARASHADER_API FNiagaraCompilationQueue* FNiagaraCompilationQueue::Singleton = nullptr;
 #endif
 
-FNiagaraScript::~FNiagaraScript()
+FNiagaraShaderScript::~FNiagaraShaderScript()
 {
 #if WITH_EDITOR
 	if (IsInGameThread())
@@ -33,7 +34,7 @@ FNiagaraScript::~FNiagaraScript()
 }
 
 /** Populates OutEnvironment with defines needed to compile shaders for this script. */
-void FNiagaraScript::SetupShaderCompilationEnvironment(
+void FNiagaraShaderScript::SetupShaderCompilationEnvironment(
 	EShaderPlatform Platform,
 	FShaderCompilerEnvironment& OutEnvironment
 	) const
@@ -42,18 +43,18 @@ void FNiagaraScript::SetupShaderCompilationEnvironment(
 }
 
 
-NIAGARASHADER_API bool FNiagaraScript::ShouldCache(EShaderPlatform Platform, const FShaderType* ShaderType) const
+NIAGARASHADER_API bool FNiagaraShaderScript::ShouldCache(EShaderPlatform Platform, const FShaderType* ShaderType) const
 {
 	check(ShaderType->GetNiagaraShaderType() )
 	return true;
 }
 
-NIAGARASHADER_API void FNiagaraScript::NotifyCompilationFinished()
+NIAGARASHADER_API void FNiagaraShaderScript::NotifyCompilationFinished()
 {
 	OnCompilationCompleteDelegate.Broadcast();
 }
 
-NIAGARASHADER_API void FNiagaraScript::CancelCompilation()
+NIAGARASHADER_API void FNiagaraShaderScript::CancelCompilation()
 {
 #if WITH_EDITOR
 	if (IsInGameThread())
@@ -61,33 +62,38 @@ NIAGARASHADER_API void FNiagaraScript::CancelCompilation()
 		FNiagaraShaderMap::RemovePendingScript(this);
 		FNiagaraCompilationQueue::Get()->RemovePending(this);
 
-		//UE_LOG(LogShaders, Error, TEXT("CancelCompilation %p."), this);
+		UE_LOG(LogShaders, Log, TEXT("CancelCompilation %p."), this);
 		OutstandingCompileShaderMapIds.Empty();
 	}
 #endif
 }
 
-NIAGARASHADER_API void FNiagaraScript::RemoveOutstandingCompileId(const int32 OldOutstandingCompileShaderMapId)
+NIAGARASHADER_API void FNiagaraShaderScript::RemoveOutstandingCompileId(const int32 OldOutstandingCompileShaderMapId)
 {
-	if (0 < OutstandingCompileShaderMapIds.Remove(OldOutstandingCompileShaderMapId))
+	if (0 <= OutstandingCompileShaderMapIds.Remove(OldOutstandingCompileShaderMapId))
 	{
-		//UE_LOG(LogShaders, Error, TEXT("RemoveOutstandingCompileId %p %d"), this, OldOutstandingCompileShaderMapId);
+		UE_LOG(LogShaders, Log, TEXT("RemoveOutstandingCompileId %p %d"), this, OldOutstandingCompileShaderMapId);
 	}
 }
 
-NIAGARASHADER_API void FNiagaraScript::Invalidate()
+NIAGARASHADER_API void FNiagaraShaderScript::Invalidate()
 {
 	CancelCompilation();
 	ReleaseShaderMap();
 }
 
 
-NIAGARASHADER_API void FNiagaraScript::LegacySerialize(FArchive& Ar)
+NIAGARASHADER_API void FNiagaraShaderScript::LegacySerialize(FArchive& Ar)
 {
 }
 
+bool FNiagaraShaderScript::IsSame(const FNiagaraShaderMapId& InId) const
+{
+	return InId.BaseScriptID == BaseScriptId && InId.ReferencedDependencyIds == ReferencedDependencyIds && InId.CompilerVersionID == CompilerVersionId;
+}
+
 
-void FNiagaraScript::GetDependentShaderTypes(EShaderPlatform Platform, TArray<FShaderType*>& OutShaderTypes) const
+void FNiagaraShaderScript::GetDependentShaderTypes(EShaderPlatform Platform, TArray<FShaderType*>& OutShaderTypes) const
 {
 	for (TLinkedList<FShaderType*>::TIterator ShaderTypeIt(FShaderType::GetTypeList()); ShaderTypeIt; ShaderTypeIt.Next())
 	{
@@ -102,7 +108,7 @@ void FNiagaraScript::GetDependentShaderTypes(EShaderPlatform Platform, TArray<FS
 
 
 
-NIAGARASHADER_API void FNiagaraScript::GetShaderMapId(EShaderPlatform Platform, FNiagaraShaderMapId& OutId) const
+NIAGARASHADER_API void FNiagaraShaderScript::GetShaderMapId(EShaderPlatform Platform, FNiagaraShaderMapId& OutId) const
 {
 	if (bLoadedCookedShaderMapId)
 	{
@@ -113,18 +119,20 @@ NIAGARASHADER_API void FNiagaraScript::GetShaderMapId(EShaderPlatform Platform,
 		TArray<FShaderType*> ShaderTypes;
 		GetDependentShaderTypes(Platform, ShaderTypes);
 		OutId.FeatureLevel = GetFeatureLevel();
-		OutId.BaseScriptID = GetScriptID();
+		OutId.BaseScriptID = BaseScriptId;
+		OutId.ReferencedDependencyIds = ReferencedDependencyIds;
+		OutId.CompilerVersionID = FNiagaraCustomVersion::LatestScriptCompileVersion;
 	}
 }
 
 
 
-void FNiagaraScript::AddReferencedObjects(FReferenceCollector& Collector)
+void FNiagaraShaderScript::AddReferencedObjects(FReferenceCollector& Collector)
 {
 }
 
 
-void FNiagaraScript::RegisterShaderMap()
+void FNiagaraShaderScript::RegisterShaderMap()
 {
 	if (GameThreadShaderMap)
 	{
@@ -132,7 +140,7 @@ void FNiagaraScript::RegisterShaderMap()
 	}
 }
 
-void FNiagaraScript::ReleaseShaderMap()
+void FNiagaraShaderScript::ReleaseShaderMap()
 {
 	if (GameThreadShaderMap)
 	{
@@ -140,14 +148,14 @@ void FNiagaraScript::ReleaseShaderMap()
 
 		ENQUEUE_UNIQUE_RENDER_COMMAND_ONEPARAMETER(
 			ReleaseShaderMap,
-			FNiagaraScript*, Script, this,
+			FNiagaraShaderScript*, Script, this,
 			{
 				Script->SetRenderingThreadShaderMap(nullptr);
 			});
 	}
 }
 
-void FNiagaraScript::SerializeShaderMap(FArchive& Ar)
+void FNiagaraShaderScript::SerializeShaderMap(FArchive& Ar)
 {
 	bool bCooked = Ar.IsCooking();
 	Ar << bCooked;
@@ -165,18 +173,17 @@ void FNiagaraScript::SerializeShaderMap(FArchive& Ar)
 			FinishCompilation();
 
 			bool bValid = GameThreadShaderMap != nullptr && GameThreadShaderMap->CompiledSuccessfully();
-			ensure(bValid);
 			Ar << bValid;
 
 			if (bValid)
 			{
 				GameThreadShaderMap->Serialize(Ar);
 			}
-			else
-			{
-				FString Name;
-				UE_LOG(LogShaders, Error, TEXT("ERROR: Failed to compile Niagara shader %s."), *GetFriendlyName());
-			}
+			//else if (GameThreadShaderMap != nullptr && !GameThreadShaderMap->CompiledSuccessfully())
+			//{
+			//	FString Name;
+			//	UE_LOG(LogShaders, Error, TEXT("Failed to compile Niagara shader %s."), *GetFriendlyName());
+			//}
 #endif
 		}
 		else
@@ -204,25 +211,23 @@ void FNiagaraScript::SerializeShaderMap(FArchive& Ar)
 	}
 }
 
-
-
-
-
-void FNiagaraScript::SetScript(UNiagaraScript *InScript, ERHIFeatureLevel::Type InFeatureLevel, FGuid InScriptID, FString InFriendlyName)
+void FNiagaraShaderScript::SetScript(UNiagaraScript *InScript, ERHIFeatureLevel::Type InFeatureLevel, const FGuid& InCompilerVersionID, const FGuid& InBaseScriptID, const TArray<FGuid>& InReferencedDependencyIds, FString InFriendlyName)
 {
-	GTScript = InScript;
-	ScriptId = InScriptID;
+	BaseVMScript = InScript;
+	BaseScriptId = InBaseScriptID;
+	CompilerVersionId = InCompilerVersionID;
+	ReferencedDependencyIds = InReferencedDependencyIds;
 	FriendlyName = InFriendlyName;
 	SetFeatureLevel(InFeatureLevel);
 }
 
-NIAGARASHADER_API  void FNiagaraScript::SetRenderingThreadShaderMap(FNiagaraShaderMap* InShaderMap)
+NIAGARASHADER_API  void FNiagaraShaderScript::SetRenderingThreadShaderMap(FNiagaraShaderMap* InShaderMap)
 {
 	check(IsInRenderingThread());
 	RenderingThreadShaderMap = InShaderMap;
 }
 
-NIAGARASHADER_API  bool FNiagaraScript::IsCompilationFinished() const
+NIAGARASHADER_API  bool FNiagaraShaderScript::IsCompilationFinished() const
 {
 	bool bRet = GameThreadShaderMap && GameThreadShaderMap.IsValid() && GameThreadShaderMap->IsCompilationFinalized();
 	if (OutstandingCompileShaderMapIds.Num() == 0)
@@ -237,7 +242,7 @@ NIAGARASHADER_API  bool FNiagaraScript::IsCompilationFinished() const
 */
 #if WITH_EDITOR
 
-bool FNiagaraScript::CacheShaders(EShaderPlatform Platform, bool bApplyCompletedShaderMapForRendering, bool bForceRecompile, bool bSynchronous)
+bool FNiagaraShaderScript::CacheShaders(EShaderPlatform Platform, bool bApplyCompletedShaderMapForRendering, bool bForceRecompile, bool bSynchronous)
 {
 	FNiagaraShaderMapId NoStaticParametersId;
 	GetShaderMapId(Platform, NoStaticParametersId);
@@ -247,7 +252,7 @@ bool FNiagaraScript::CacheShaders(EShaderPlatform Platform, bool bApplyCompleted
 /**
 * Caches the shaders for this script
 */
-bool FNiagaraScript::CacheShaders(const FNiagaraShaderMapId& ShaderMapId, EShaderPlatform Platform, bool bApplyCompletedShaderMapForRendering, bool bForceRecompile, bool bSynchronous)
+bool FNiagaraShaderScript::CacheShaders(const FNiagaraShaderMapId& ShaderMapId, EShaderPlatform Platform, bool bApplyCompletedShaderMapForRendering, bool bForceRecompile, bool bSynchronous)
 {
 	bool bSucceeded = false;
 
@@ -282,7 +287,7 @@ bool FNiagaraScript::CacheShaders(const FNiagaraShaderMapId& ShaderMapId, EShade
 		UE_LOG(LogTemp, Display, TEXT("Found existing compiling shader for Niagara script %s, linking to other GameThreadShaderMap 0x%08X%08X"), *GetFriendlyName(), (int)((int64)(GameThreadShaderMap.GetReference()) >> 32), (int)((int64)(GameThreadShaderMap.GetReference())));
 #endif
 		OutstandingCompileShaderMapIds.AddUnique(GameThreadShaderMap->GetCompilingId());
-		//UE_LOG(LogShaders, Error, TEXT("AddUniqueExisting %p %d"), this, GameThreadShaderMap->GetCompilingId());
+		UE_LOG(LogShaders, Log, TEXT("CacheShaders AddUniqueExisting %p %d"), this, GameThreadShaderMap->GetCompilingId());
 
 		// Reset the shader map so we fall back to CPU sim until the compile finishes.
 		GameThreadShaderMap = nullptr;
@@ -318,7 +323,7 @@ bool FNiagaraScript::CacheShaders(const FNiagaraShaderMapId& ShaderMapId, EShade
 
 	ENQUEUE_UNIQUE_RENDER_COMMAND_TWOPARAMETER(
 		FSetShaderMapOnScriptResources,
-		FNiagaraScript*, Script, this,
+		FNiagaraShaderScript*, Script, this,
 		FNiagaraShaderMap*, LoadedShaderMap, GameThreadShaderMap,
 		{
 			Script->SetRenderingThreadShaderMap(LoadedShaderMap);
@@ -328,24 +333,47 @@ bool FNiagaraScript::CacheShaders(const FNiagaraShaderMapId& ShaderMapId, EShade
 }
 
 
-void FNiagaraScript::FinishCompilation()
+void FNiagaraShaderScript::FinishCompilation()
 {
 	TArray<int32> ShaderMapIdsToFinish;
 	GetShaderMapIDsWithUnfinishedCompilation(ShaderMapIdsToFinish);
 
 	if (ShaderMapIdsToFinish.Num() > 0)
 	{
+		for (int32 i = 0; i < ShaderMapIdsToFinish.Num(); i++)
+		{
+			UE_LOG(LogShaders, Log, TEXT("FinishCompilation()[%d] %s id %d!"), i, *GetFriendlyName(), ShaderMapIdsToFinish[i]);
+		}
 		// Block until the shader maps that we will save have finished being compiled
 		// NIAGARATODO: implement when async compile works
 		GNiagaraShaderCompilationManager.FinishCompilation(*GetFriendlyName(), ShaderMapIdsToFinish);
+
+		// Shouldn't have anything left to do...
+		TArray<int32> ShaderMapIdsToFinish2;
+		GetShaderMapIDsWithUnfinishedCompilation(ShaderMapIdsToFinish2);
+		ensure(ShaderMapIdsToFinish2.Num() == 0);
 	}
 }
 
 #endif
 
-NIAGARASHADER_API  FNiagaraShader* FNiagaraScript::GetShader() const
+void FNiagaraShaderScript::SetDataInterfaceParamInfo(TArray< FNiagaraDataInterfaceGPUParamInfo >& InDIParamInfo)
 {
-	check(IsInRenderingThread());
+	DIParamInfo = InDIParamInfo;
+}
+
+void FNiagaraShaderScript::SetDataInterfaceParamInfo(TArray<FNiagaraDataInterfaceParamRef>& InDIParamRefs)
+{
+	DIParamInfo.Empty();
+	for (FNiagaraDataInterfaceParamRef& DIParam : InDIParamRefs)
+	{
+		DIParamInfo.Add(DIParam.ParameterInfo);
+	}
+}
+
+NIAGARASHADER_API  FNiagaraShader* FNiagaraShaderScript::GetShader() const
+{
+	check(! IsInGameThread() );
 	if (!GIsEditor || RenderingThreadShaderMap /*&& RenderingThreadShaderMap->IsComplete(this, true)*/)
 	{
 		return RenderingThreadShaderMap->GetShader<FNiagaraShader>();
@@ -353,7 +381,7 @@ NIAGARASHADER_API  FNiagaraShader* FNiagaraScript::GetShader() const
 	return nullptr;
 };
 
-NIAGARASHADER_API  FNiagaraShader* FNiagaraScript::GetShaderGameThread() const
+NIAGARASHADER_API  FNiagaraShader* FNiagaraShaderScript::GetShaderGameThread() const
 {
 	if (GameThreadShaderMap)
 	{
@@ -364,7 +392,7 @@ NIAGARASHADER_API  FNiagaraShader* FNiagaraScript::GetShaderGameThread() const
 };
 
 
-void FNiagaraScript::GetShaderMapIDsWithUnfinishedCompilation(TArray<int32>& ShaderMapIds)
+void FNiagaraShaderScript::GetShaderMapIDsWithUnfinishedCompilation(TArray<int32>& ShaderMapIds)
 {
 	// Build an array of the shader map Id's are not finished compiling.
 	if (GameThreadShaderMap && GameThreadShaderMap.IsValid() && !GameThreadShaderMap->IsCompilationFinalized())
@@ -387,7 +415,7 @@ void FNiagaraScript::GetShaderMapIDsWithUnfinishedCompilation(TArray<int32>& Sha
 * @param OutShaderMap - the shader map to compile
 * @return - true if compile succeeded or was not necessary (shader map for ShaderMapId was found and was complete)
 */
-bool FNiagaraScript::BeginCompileShaderMap(
+bool FNiagaraShaderScript::BeginCompileShaderMap(
 	const FNiagaraShaderMapId& ShaderMapId,
 	EShaderPlatform Platform,
 	TRefCountPtr<FNiagaraShaderMap>& OutShaderMap,
@@ -405,7 +433,7 @@ bool FNiagaraScript::BeginCompileShaderMap(
 	// Queue hlsl generation and shader compilation - Unlike materials, we queue this here, and compilation happens from the editor module
 	TRefCountPtr<FNiagaraShaderMap> NewShaderMap = new FNiagaraShaderMap();
 	OutstandingCompileShaderMapIds.AddUnique(NewShaderMap->GetCompilingId());		
-	//UE_LOG(LogShaders, Error, TEXT("AddUnique %p %d"), this, NewShaderMap->GetCompilingId());
+	UE_LOG(LogShaders, Log, TEXT("BeginCompileShaderMap AddUnique %p %d"), this, NewShaderMap->GetCompilingId());
 
 	FNiagaraCompilationQueue::Get()->Queue(this, NewShaderMap, ShaderMapId, Platform, bApplyCompletedShaderMapForRendering);
 	if (bSynchronous)
diff --git a/Engine/Plugins/FX/Niagara/Source/NiagaraShader/Public/NiagaraShader.h b/Engine/Plugins/FX/Niagara/Source/NiagaraShader/Public/NiagaraShader.h
--- a/Engine/Plugins/FX/Niagara/Source/NiagaraShader/Public/NiagaraShader.h
+++ b/Engine/Plugins/FX/Niagara/Source/NiagaraShader/Public/NiagaraShader.h
@@ -16,7 +16,8 @@
 #include "NiagaraShaderType.h"
 #include "SceneRenderTargetParameters.h"
 
-
+struct FNiagaraDataInterfaceParametersCS;
+class UClass;
 
 template<typename TBufferStruct> class TUniformBufferRef;
 
@@ -43,7 +44,7 @@ public:
 	{
 	}
 
-	static bool ShouldCompilePermutation(EShaderPlatform Platform, const FNiagaraScript* Script)
+	static bool ShouldCompilePermutation(EShaderPlatform Platform, const FNiagaraShaderScript*  Script)
 	{
 		//@todo - lit materials only 
 		return IsFeatureLevelSupported(Platform, ERHIFeatureLevel::SM5);
@@ -52,16 +53,13 @@ public:
 
 	FNiagaraShader(const FNiagaraShaderType::CompiledShaderInitializerType& Initializer);
 
-	typedef void (*ModifyCompilationEnvironmentType)(EShaderPlatform, const FNiagaraScript*, FShaderCompilerEnvironment&);
-
-	static void ModifyCompilationEnvironment(EShaderPlatform Platform, const FNiagaraScript* Script, FShaderCompilerEnvironment& OutEnvironment)
-	{
-	}
+	typedef void (*ModifyCompilationEnvironmentType)(EShaderPlatform, const FNiagaraShaderScript* , FShaderCompilerEnvironment&);
 
-	void SetDatainterfaceBufferDescriptors(const TArray< FDIBufferDescriptorStore > &InBufferDescriptors)
+	static void ModifyCompilationEnvironment(EShaderPlatform Platform, const FNiagaraShaderScript*  Script, FShaderCompilerEnvironment& OutEnvironment)
 	{
-		DIBufferDescriptors = InBufferDescriptors;
 	}
+	
+	void SetDataInterfaceParameterInfo(const TArray<FNiagaraDataInterfaceGPUParamInfo>& InDIParamInfo);
 
 //	FUniformBufferRHIParamRef GetParameterCollectionBuffer(const FGuid& Id, const FSceneInterface* SceneInterface) const;
 	/*
@@ -86,61 +84,7 @@ public:
 	*/
 
 	// Bind parameters
-	void BindParams(const FShaderParameterMap &ParameterMap)
-	{
-		FloatInputBufferParam.Bind(ParameterMap, TEXT("InputFloat"));
-		IntInputBufferParam.Bind(ParameterMap, TEXT("InputInt"));
-		FloatOutputBufferParam.Bind(ParameterMap, TEXT("OutputFloat"));
-		IntOutputBufferParam.Bind(ParameterMap, TEXT("OutputInt"));
-		OutputIndexBufferParam.Bind(ParameterMap, TEXT("DataSetIndices"));
-		EmitterTickCounterParam.Bind(ParameterMap, TEXT("EmitterTickCounter"));
-		NumEventsPerParticleParam.Bind(ParameterMap, TEXT("NumEventsPerParticle"));
-		NumParticlesPerEventParam.Bind(ParameterMap, TEXT("NumParticlesPerEvent"));
-		CopyInstancesBeforeStartParam.Bind(ParameterMap, TEXT("CopyInstancesBeforeStart"));
-
-		TotalNumInstancesParam.Bind(ParameterMap, TEXT("TotalNumInstances"));
-		StartInstancePhase0Param.Bind(ParameterMap, TEXT("StartInstancePhase0"));
-		StartInstancePhase1Param.Bind(ParameterMap, TEXT("StartInstancePhase1"));
-
-		ComponentBufferSizeReadParam.Bind(ParameterMap, TEXT("ComponentBufferSizeRead"));
-		ComponentBufferSizeWriteParam.Bind(ParameterMap, TEXT("ComponentBufferSizeWrite"));
-		EmitterConstantBufferParam.Bind(ParameterMap, TEXT("FEmitterParameters"));
-
-		// params for event buffers
-		// this is horrendous; need to do this in a uniform buffer instead.
-		//
-		for (uint32 i = 0; i < MAX_CONCURRENT_EVENT_DATASETS; i++)
-		{
-			FString VarName = TEXT("WriteDataSetFloat") + FString::FromInt(i+1);
-			EventFloatUAVParams[i].Bind(ParameterMap, *VarName);
-			VarName = TEXT("WriteDataSetInt") + FString::FromInt(i+1);
-			EventIntUAVParams[i].Bind(ParameterMap, *VarName);
-
-			VarName = TEXT("ReadDataSetFloat") + FString::FromInt(i+1);
-			EventFloatSRVParams[i].Bind(ParameterMap, *VarName);
-			VarName = TEXT("ReadDataSetInt") + FString::FromInt(i+1);
-			EventIntSRVParams[i].Bind(ParameterMap, *VarName);
-
-			VarName = TEXT("DSComponentBufferSizeReadFloat") + FString::FromInt(i + 1);
-			EventReadFloatStrideParams[i].Bind(ParameterMap, *VarName);
-			VarName = TEXT("DSComponentBufferSizeWriteFloat") + FString::FromInt(i + 1);
-			EventWriteFloatStrideParams[i].Bind(ParameterMap, *VarName);
-			VarName = TEXT("DSComponentBufferSizeReadInt") + FString::FromInt(i + 1);
-			EventReadIntStrideParams[i].Bind(ParameterMap, *VarName);
-			VarName = TEXT("DSComponentBufferSizeWriteInt") + FString::FromInt(i + 1);
-			EventWriteIntStrideParams[i].Bind(ParameterMap, *VarName);
-		}
-
-		// params for data interface buffers
-		BuildDIBufferParamMap(ParameterMap);
-
-		ensure(FloatOutputBufferParam.IsBound() || IntOutputBufferParam.IsBound());	// we should have at least one output buffer we're writing to
-		ensure(ComponentBufferSizeWriteParam.IsBound());
-		ensure(OutputIndexBufferParam.IsBound());
-		ensure(StartInstancePhase0Param.IsBound());
-		ensure(StartInstancePhase1Param.IsBound());
-		ensure(TotalNumInstancesParam.IsBound());
-	}
+	void BindParams(const FShaderParameterMap &ParameterMap);
 
 	// FShader interface.
 	virtual bool Serialize(FArchive& Ar) override;
@@ -151,15 +95,17 @@ public:
 	FRWShaderParameter FloatOutputBufferParam;
 	FRWShaderParameter IntOutputBufferParam;
 	FRWShaderParameter OutputIndexBufferParam;
+	FShaderResourceParameter InputIndexBufferParam;
 	FShaderUniformBufferParameter EmitterConstantBufferParam;
 	FShaderUniformBufferParameter DataInterfaceUniformBufferParam;
+	FShaderUniformBufferParameter ViewUniformBufferParam;
 	FShaderParameter EmitterTickCounterParam;
 	FShaderParameter NumEventsPerParticleParam;
 	FShaderParameter NumParticlesPerEventParam;
 	FShaderParameter CopyInstancesBeforeStartParam;
-	FShaderParameter TotalNumInstancesParam;
-	FShaderParameter StartInstancePhase0Param;
-	FShaderParameter StartInstancePhase1Param;
+	FShaderParameter NumSpawnedInstancesParam;
+	FShaderParameter UpdateStartInstanceParam;
+	FShaderParameter NumIndicesPerInstanceParam;
 	FShaderParameter ComponentBufferSizeReadParam;
 	FShaderParameter ComponentBufferSizeWriteParam;
 	FRWShaderParameter EventIntUAVParams[MAX_CONCURRENT_EVENT_DATASETS];
@@ -171,46 +117,17 @@ public:
 	FShaderParameter EventReadFloatStrideParams[MAX_CONCURRENT_EVENT_DATASETS];
 	FShaderParameter EventReadIntStrideParams[MAX_CONCURRENT_EVENT_DATASETS];
 
-	FShaderResourceParameter *FindDIBufferParam(int32 DataInterfaceIndex, FName InName)
+	TArray< FNiagaraDataInterfaceParamRef >& GetDIParameters()
 	{
-		if (NameToDIBufferParamMap.Num() > DataInterfaceIndex)
-		{
-			return NameToDIBufferParamMap[DataInterfaceIndex].Find(InName);
-		}
-
-		return nullptr;
+		return DataInterfaceParameters;
 	}
 
-	void BuildDIBufferParamMap(const FShaderParameterMap &ParameterMap)
-	{
-		for (FDIBufferDescriptorStore &InterfaceDescs : DIBufferDescriptors)
-		{
-			int32 Idx = NameToDIBufferParamMap.AddDefaulted(1);
-			for (FDIGPUBufferParamDescriptor &Desc : InterfaceDescs.Descriptors)
-			{
-				FShaderResourceParameter &Param = NameToDIBufferParamMap[Idx].Add(*Desc.BufferParamName);
-				Param.Bind(ParameterMap, *Desc.BufferParamName);
-				ensure(Param.IsBound());
-			}
-		}
-	}
-
-	TArray< FDIBufferDescriptorStore > &GetDIBufferDescriptors()
-	{
-		return DIBufferDescriptors;
-	}
 
 private:
 	FShaderUniformBufferParameter NiagaraUniformBuffer;
 
-	// buffer descriptors for data interfaces holding names and params for binding
-	TArray< FDIBufferDescriptorStore > DIBufferDescriptors;
-
-	// buffer descriptors for event data sets holding names and params for binding
-	TArray< TArray<FDIGPUBufferParamDescriptor> > EventBufferDescriptors;
-
-	// one map per data interface, mapping buffer names to their params
-	TArray< TMap<FName, FShaderResourceParameter> > NameToDIBufferParamMap;
+	// Data about parameters used for each Data Interface.
+	TArray< FNiagaraDataInterfaceParamRef > DataInterfaceParameters;
 
 	/*
 	FDebugUniformExpressionSet	DebugUniformExpressionSet;
diff --git a/Engine/Plugins/FX/Niagara/Source/NiagaraShader/Public/NiagaraShaderCompilationManager.h b/Engine/Plugins/FX/Niagara/Source/NiagaraShader/Public/NiagaraShaderCompilationManager.h
--- a/Engine/Plugins/FX/Niagara/Source/NiagaraShader/Public/NiagaraShaderCompilationManager.h
+++ b/Engine/Plugins/FX/Niagara/Source/NiagaraShader/Public/NiagaraShaderCompilationManager.h
@@ -24,7 +24,7 @@ struct FNiagaraShaderCompileWorkerInfo
 	double StartTime;
 
 	/** Jobs that this worker is responsible for compiling. */
-	TArray<class FNiagaraShaderCompileJob*> QueuedJobs;
+	TArray<class FShaderCommonCompileJob*> QueuedJobs;
 
 	FNiagaraShaderCompileWorkerInfo() :
 		bIssuedTasksToWorker(false),
@@ -45,6 +45,7 @@ struct FNiagaraShaderCompileWorkerInfo
 	}
 };
 
+#if 0
 /** Stores all of the input and output information used to compile a single shader. */
 class FNiagaraShaderCompileJob 
 {
@@ -61,8 +62,11 @@ public:
 	/** Input for the shader compile */
 	FShaderCompilerInput Input;
 	FShaderCompilerOutput Output;
+
+	class FShaderCompileJob* BaseCompileJob;
+
 	FString Hlsl;
-	TArray< FDIBufferDescriptorStore >DIBufferDescriptors;
+	TArray<FNiagaraDataInterfaceGPUParamInfo> DIParamInfo;
 
 	FNiagaraShaderCompileJob(uint32 InId, FShaderType* InShaderType, FString InHlsl) :
 		Id(InId),
@@ -70,10 +74,12 @@ public:
 		bSucceeded(false),
 		bOptimizeForLowLatency(false),
 		ShaderType(InShaderType),
+		BaseCompileJob(nullptr),
 		Hlsl(InHlsl)
 	{
 	}
 };
+#endif
 
 
 
@@ -89,7 +95,7 @@ struct FNiagaraShaderMapCompileResults
 	int32 NumJobsQueued;
 	bool bAllJobsSucceeded;
 	bool bRecreateComponentRenderStateOnCompletion;
-	TArray<FNiagaraShaderCompileJob*> FinishedJobs;
+	TArray<FShaderCommonCompileJob*> FinishedJobs;
 };
 
 
@@ -114,7 +120,7 @@ public:
 	FNiagaraShaderCompilationManager();
 
 	NIAGARASHADER_API void Tick(float DeltaSeconds = 0.0f);
-	NIAGARASHADER_API void AddJobs(TArray<FNiagaraShaderCompileJob*> InNewJobs);
+	NIAGARASHADER_API void AddJobs(TArray<class FShaderCommonCompileJob*> InNewJobs);
 	NIAGARASHADER_API void ProcessAsyncResults();
 
 	void FinishCompilation(const TCHAR* ScriptName, const TArray<int32>& ShaderMapIdsToFinishCompiling);
@@ -123,7 +129,7 @@ private:
 	void ProcessCompiledNiagaraShaderMaps(TMap<int32, FNiagaraShaderMapFinalizeResults>& CompiledShaderMaps, float TimeBudget);
 	void RunCompileJobs();
 
-	TArray<FNiagaraShaderCompileJob*> JobQueue;
+	TArray<class FShaderCommonCompileJob*> JobQueue;
 
 	/** Map from shader map Id to the compile results for that map, used to gather compiled results. */
 	TMap<int32, FNiagaraShaderMapCompileResults> NiagaraShaderMapJobs;
diff --git a/Engine/Plugins/FX/Niagara/Source/NiagaraShader/Public/NiagaraShaderType.h b/Engine/Plugins/FX/Niagara/Source/NiagaraShader/Public/NiagaraShaderType.h
--- a/Engine/Plugins/FX/Niagara/Source/NiagaraShader/Public/NiagaraShaderType.h
+++ b/Engine/Plugins/FX/Niagara/Source/NiagaraShader/Public/NiagaraShaderType.h
@@ -21,14 +21,14 @@
 		Frequency \
 		);
 
-class FNiagaraScript;
+class FNiagaraShaderScript;
 class FShaderCommonCompileJob;
 class FShaderCompileJob;
 class FUniformExpressionSet;
 
 
 /** Called for every Niagara shader to update the appropriate stats. */
-extern void UpdateNiagaraShaderCompilingStats(const FNiagaraScript* Script);
+extern void UpdateNiagaraShaderCompilingStats(const FNiagaraShaderScript* Script);
 
 /**
  * Dump shader stats for a given platform.
@@ -50,7 +50,7 @@ public:
 	{
 		//const FUniformExpressionSet& UniformExpressionSet;
 		const FString DebugDescription;
-		TArray< FDIBufferDescriptorStore > DIBufferDescriptors;
+		TArray< FNiagaraDataInterfaceGPUParamInfo > DIParamInfo;
 
 		CompiledShaderInitializerType(
 			FShaderType* InType,
@@ -59,16 +59,16 @@ public:
 			FShaderResource* InResource,
 			const FSHAHash& InNiagaraShaderMapHash,
 			const FString& InDebugDescription,
-			const TArray< FDIBufferDescriptorStore > &InDataInterfaceBufferDescriptors
+			const TArray< FNiagaraDataInterfaceGPUParamInfo > &InDIParamInfo
 			)
 		: FGlobalShaderType::CompiledShaderInitializerType(InType,InPermutationId,CompilerOutput,InResource, InNiagaraShaderMapHash,nullptr,nullptr)
 		, DebugDescription(InDebugDescription)
-		, DIBufferDescriptors(InDataInterfaceBufferDescriptors)
+		, DIParamInfo(InDIParamInfo)
 		{}
 	};
 	typedef FShader* (*ConstructCompiledType)(const CompiledShaderInitializerType&);
-	typedef bool (*ShouldCompilePermutationType)(EShaderPlatform,const FNiagaraScript*);
-	typedef void (*ModifyCompilationEnvironmentType)(EShaderPlatform,const FNiagaraScript*, FShaderCompilerEnvironment&);
+	typedef bool (*ShouldCompilePermutationType)(EShaderPlatform,const FNiagaraShaderScript*);
+	typedef void (*ModifyCompilationEnvironmentType)(EShaderPlatform,const FNiagaraShaderScript*, FShaderCompilerEnvironment&);
 
 	FNiagaraShaderType(
 		const TCHAR* InName,
@@ -94,14 +94,14 @@ public:
 	 * Enqueues a compilation for a new shader of this type.
 	 * @param Script - The script to link the shader with.
 	 */
-	class FNiagaraShaderCompileJob* BeginCompileShader(
-		uint32 ShaderMapId,
-		const FNiagaraScript* Script,
-		FShaderCompilerEnvironment* ScriptEnvironment,
-		EShaderPlatform Platform,
-		TArray<FNiagaraShaderCompileJob*>& NewJobs,
-		FShaderTarget Target,
-		TArray< FDIBufferDescriptorStore >&InBufferDescriptors
+	class FShaderCompileJob* BeginCompileShader(
+			uint32 ShaderMapId,
+			const FNiagaraShaderScript* Script,
+			FShaderCompilerEnvironment* CompilationEnvironment,
+			EShaderPlatform Platform,
+			TArray<FShaderCommonCompileJob*>& NewJobs,
+			FShaderTarget Target,
+			TArray<FNiagaraDataInterfaceGPUParamInfo>& InDIParamInfo
 		);
 
 	/**
@@ -111,7 +111,7 @@ public:
 	 */
 	FShader* FinishCompileShader(
 		const FSHAHash& NiagaraShaderMapHash,
-		const FNiagaraShaderCompileJob& CurrentJob,
+		const FShaderCompileJob& CurrentJob,
 		const FString& InDebugDescription
 		);
 
@@ -121,11 +121,17 @@ public:
 	 * @param script - The script to check.
 	 * @return True if this shader type should be cached.
 	 */
-	bool ShouldCache(EShaderPlatform Platform,const FNiagaraScript* Script) const
+	bool ShouldCache(EShaderPlatform Platform,const FNiagaraShaderScript* Script) const
 	{
 		return (*ShouldCompilePermutationRef)(Platform, Script);
 	}
 
+	/** Adds include statements for uniform buffers that this shader type references, and builds a prefix for the shader file with the include statements. */
+	void AddReferencedUniformBufferIncludes(FShaderCompilerEnvironment& OutEnvironment, FString& OutSourceFilePrefix, EShaderPlatform Platform);
+
+	void CacheUniformBufferIncludes(TMap<const TCHAR*, FCachedUniformBufferDeclaration>& Cache, EShaderPlatform Platform);
+
+
 protected:
 
 	/**
@@ -133,7 +139,7 @@ protected:
 	 * @param Platform - Platform to compile for.
 	 * @param Environment - The shader compile environment that the function modifies.
 	 */
-	void SetupCompileEnvironment(EShaderPlatform Platform, const FNiagaraScript* Script, FShaderCompilerEnvironment& Environment)
+	void SetupCompileEnvironment(EShaderPlatform Platform, const FNiagaraShaderScript* Script, FShaderCompilerEnvironment& Environment)
 	{
 		// Allow the shader type to modify its compile environment.
 		(*ModifyCompilationEnvironmentRef)(Platform, Script, Environment);
@@ -143,4 +149,6 @@ private:
 	ConstructCompiledType ConstructCompiledRef;
 	ShouldCompilePermutationType ShouldCompilePermutationRef;
 	ModifyCompilationEnvironmentType ModifyCompilationEnvironmentRef;
+
+	static TMap<const FShaderCompileJob*, TArray<FNiagaraDataInterfaceGPUParamInfo> > ExtraParamInfo;
 };
