Copying //UE4/Dev-Niagara to //UE4/Dev-Main (Source: //UE4/Dev-Niagara @ 4074996)

diff --git a/Engine/Plugins/FX/Niagara/Source/NiagaraShader/Public/NiagaraShared.h b/Engine/Plugins/FX/Niagara/Source/NiagaraShader/Public/NiagaraShared.h
--- a/Engine/Plugins/FX/Niagara/Source/NiagaraShader/Public/NiagaraShared.h
+++ b/Engine/Plugins/FX/Niagara/Source/NiagaraShader/Public/NiagaraShared.h
@@ -24,57 +24,65 @@
 #include "Misc/Optional.h"
 #include "NiagaraShared.generated.h"
 
-class FNiagaraScript;
+class FNiagaraShaderScript;
 class FNiagaraShaderMap;
 class FNiagaraShader;
 class FNiagaraShaderMapId;
 class UNiagaraScript;
+struct FNiagaraDataInterfaceParametersCS;
 
 #define MAX_CONCURRENT_EVENT_DATASETS 4
 
-
-
-// holds meta data for GPU buffers from data interface; the DI defines them and returns them, they're
-// passed to the shader compiler, so we can bind parameters to dynamically generated hlsl vars
-//  this is only necessary for buffer params; everything else is in a constant buffer, which we
-//  just allocate memory for and copy the data into
-//
+/**
+* Data coming from that translator that describes parameters needed for each data interface.
+*/
 USTRUCT()
-struct NIAGARASHADER_API FDIGPUBufferParamDescriptor
+struct NIAGARASHADER_API FNiagaraDataInterfaceGPUParamInfo
 {
 	GENERATED_USTRUCT_BODY()
-public:
-	FDIGPUBufferParamDescriptor(FString Name, uint32 Idx)
-		: BufferParamName(Name)
-		, Index(Idx)
-	{}
-	FDIGPUBufferParamDescriptor()
-		: BufferParamName(TEXT("Unknown"))
-		, Index(INDEX_NONE)
-	{}
-
-	friend FArchive &operator << (FArchive &Ar, FDIGPUBufferParamDescriptor &Desc);
+		
+	/** Symbol of this DI in the hlsl. Used for binding parameters. */
 	UPROPERTY()
-	FString BufferParamName;	// the name of the parameter in the hlsl the DI generates
+	FString DataInterfaceHLSLSymbol;
+
+	/** Name of the class for this data interface. Used for constructing the correct parameters struct. */
 	UPROPERTY()
-	uint32 Index;				// the index, reflecting the order in which GetBufferDefinitionHLSL declares them
-};
+	FString DIClassName;
 
+	bool Serialize(FArchive& Ar);
+};
 
+template<> struct TStructOpsTypeTraits<FNiagaraDataInterfaceGPUParamInfo> : public TStructOpsTypeTraitsBase2<FNiagaraDataInterfaceGPUParamInfo>
+{
+	enum
+	{
+		WithSerializer = true,
+	};
+};
 
-/** Data interface GPU buffer descriptors for one data interface
+/** 
+* Shader side data needed for binding data interface parameters.
 */
-USTRUCT()
-struct NIAGARASHADER_API FDIBufferDescriptorStore
+struct FNiagaraDataInterfaceParamRef
 {
-	GENERATED_USTRUCT_BODY()
 public:
-	friend FArchive &operator << (FArchive &Ar, FDIBufferDescriptorStore &Desc);
 
-	UPROPERTY()
-	FName DataInterfaceName;
-	UPROPERTY()
-	TArray< FDIGPUBufferParamDescriptor > Descriptors;
+	FNiagaraDataInterfaceParamRef(const FNiagaraDataInterfaceGPUParamInfo& InParameterInfo);
+	FNiagaraDataInterfaceParamRef();
+	~FNiagaraDataInterfaceParamRef();
+	friend bool operator<<(FArchive& Ar, FNiagaraDataInterfaceParamRef& Desc);
+
+	void Bind(const class FShaderParameterMap& ParameterMap);
+	void ConstructParameters();
+	void InitDIClass();
+
+	FNiagaraDataInterfaceGPUParamInfo ParameterInfo;
+
+	/** The class of this DI which we can use to create the parameter struct. */
+	UClass* DIClass;
+
+	/** Pointer to parameters struct for this data interface. */
+	FNiagaraDataInterfaceParametersCS* Parameters;
 };
 
 
@@ -97,22 +105,24 @@ public:
 class FNiagaraShaderMapId
 {
 public:
+	/** The version of the compiler that this needs to be built against.*/
+	FGuid CompilerVersionID;
+
+	/** Feature level that the shader map is going to be compiled for.  */
+	ERHIFeatureLevel::Type FeatureLevel;
 
 	/**
-	* The GUID of the script this shader represents
-	* Any change to the FNiagaraScript must modify this guid.
+	* The GUID of the subgraph this shader primarily represents.
 	*/
 	FGuid BaseScriptID;
 
-	/** Feature level that the shader map is going to be compiled for.  */
-	ERHIFeatureLevel::Type FeatureLevel;
-
-	/** TODO: may need this - Guids of any functions the script was dependent on. */
-	//TArray<FGuid> ReferencedFunctions;
+	/** Guids of any functions or module scripts the script was dependent on. */
+	TArray<FGuid> ReferencedDependencyIds;
 
 	FNiagaraShaderMapId()
-		: BaseScriptID(0, 0, 0, 0)
+		: CompilerVersionID()
 		, FeatureLevel(ERHIFeatureLevel::SM5)
+		, BaseScriptID(0, 0, 0, 0)
 	{ }
 
 	~FNiagaraShaderMapId()
@@ -169,7 +179,7 @@ class FNiagaraCompilationQueue
 public:
 	struct NiagaraCompilationQueueItem
 	{
-		FNiagaraScript* Script;
+		FNiagaraShaderScript* Script;
 		TRefCountPtr<FNiagaraShaderMap>ShaderMap;
 		FNiagaraShaderMapId ShaderMapId;
 		EShaderPlatform Platform;
@@ -190,7 +200,7 @@ public:
 		return CompilationQueue;
 	}
 
-	void Queue(FNiagaraScript *InScript, TRefCountPtr<FNiagaraShaderMap>InShaderMap, const FNiagaraShaderMapId &MapId, EShaderPlatform InPlatform, bool InApply)
+	void Queue(FNiagaraShaderScript *InScript, TRefCountPtr<FNiagaraShaderMap>InShaderMap, const FNiagaraShaderMapId &MapId, EShaderPlatform InPlatform, bool InApply)
 	{
 		check(IsInGameThread());
 		NiagaraCompilationQueueItem NewQueueItem;
@@ -202,7 +212,7 @@ public:
 		CompilationQueue.Add(NewQueueItem);
 	}
 
-	void RemovePending(FNiagaraScript* InScript)
+	void RemovePending(FNiagaraShaderScript* InScript)
 	{
 		check(IsInGameThread());
 		for (NiagaraCompilationQueueItem& Item : CompilationQueue)
@@ -250,7 +260,7 @@ public:
 	* Attempts to load the shader map for the given script from the Derived Data Cache.
 	* If InOutShaderMap is valid, attempts to load the individual missing shaders instead.
 	*/
-	static void LoadFromDerivedDataCache(const FNiagaraScript* Script, const FNiagaraShaderMapId& ShaderMapId, EShaderPlatform Platform, TRefCountPtr<FNiagaraShaderMap>& InOutShaderMap);
+	static void LoadFromDerivedDataCache(const FNiagaraShaderScript* Script, const FNiagaraShaderMapId& ShaderMapId, EShaderPlatform Platform, TRefCountPtr<FNiagaraShaderMap>& InOutShaderMap);
 
 	FNiagaraShaderMap();
 
@@ -264,7 +274,7 @@ public:
 	* @param Platform - The platform to compile to
 	*/
 	NIAGARASHADER_API void Compile(
-		FNiagaraScript* Script,
+		FNiagaraShaderScript* Script,
 		const FNiagaraShaderMapId& ShaderMapId,
 		TRefCountPtr<FShaderCompilerEnvironment> CompilationEnvironment,
 		const FNiagaraComputeShaderCompilationOutput& InNiagaraCompilationOutput,
@@ -274,17 +284,17 @@ public:
 		);
 
 	/** Sorts the incoming compiled jobs into the appropriate mesh shader maps, and finalizes this shader map so that it can be used for rendering. */
-	bool ProcessCompilationResults(const TArray<class FNiagaraShaderCompileJob*>& InCompilationResults, int32& ResultIndex, float& TimeBudget);
+	bool ProcessCompilationResults(const TArray<class FShaderCommonCompileJob*>& InCompilationResults, int32& ResultIndex, float& TimeBudget);
 
 	/**
 	* Checks whether the shader map is missing any shader types necessary for the given script.
 	* @param Script - The Niagara Script which is checked.
 	* @return True if the shader map has all of the shader types necessary.
 	*/
-	bool IsComplete(const FNiagaraScript* Script, bool bSilent);
+	bool IsComplete(const FNiagaraShaderScript* Script, bool bSilent);
 
 	/** Attempts to load missing shaders from memory. */
-	void LoadMissingShadersFromMemory(const FNiagaraScript* Script);
+	void LoadMissingShadersFromMemory(const FNiagaraShaderScript* Script);
 
 	/**
 	* Checks to see if the shader map is already being compiled for another script, and if so
@@ -292,7 +302,7 @@ public:
 	* @param Script - The Niagara Script we also wish to apply the compiled shader map to.
 	* @return True if the shader map was being compiled and we added Script to the list to be applied.
 	*/
-	bool TryToAddToExistingCompilationTask(FNiagaraScript* Script);
+	bool TryToAddToExistingCompilationTask(FNiagaraShaderScript* Script);
 
 	/** Builds a list of the shaders in a shader map. */
 	NIAGARASHADER_API  void GetShaderList(TMap<FShaderId, FShader*>& OutShaders) const;
@@ -314,11 +324,11 @@ public:
 	void FlushShadersByShaderType(FShaderType* ShaderType);
 
 	/** Removes a Script from NiagaraShaderMapsBeingCompiled. */
-	NIAGARASHADER_API static void RemovePendingScript(FNiagaraScript* Script);
+	NIAGARASHADER_API static void RemovePendingScript(FNiagaraShaderScript* Script);
 	NIAGARASHADER_API static void RemovePendingMap(FNiagaraShaderMap* Map);
 
 	/** Finds a shader map currently being compiled that was enqueued for the given script. */
-	static const FNiagaraShaderMap* GetShaderMapBeingCompiled(const FNiagaraScript* Script);
+	static const FNiagaraShaderMap* GetShaderMapBeingCompiled(const FNiagaraShaderScript* Script);
 
 	/** Serializes the shader map. */
 	void Serialize(FArchive& Ar, bool bInlineShaderResources = true);
@@ -368,7 +378,7 @@ public:
 
 	int32 GetNumRefs() const { return NumRefs; }
 	uint32 GetCompilingId()  { return CompilingId; }
-	static TMap<TRefCountPtr<FNiagaraShaderMap>, TArray<FNiagaraScript*> > &GetInFlightShaderMaps() { return NiagaraShaderMapsBeingCompiled; }
+	static TMap<TRefCountPtr<FNiagaraShaderMap>, TArray<FNiagaraShaderScript*> > &GetInFlightShaderMaps() { return NiagaraShaderMapsBeingCompiled; }
 
 	void SetCompiledSuccessfully(bool bSuccess) { bCompiledSuccessfully = bSuccess; }
 private:
@@ -402,7 +412,7 @@ private:
 	static uint32 NextCompilingId;
 
 	/** Tracks resources and their shader maps that need to be compiled but whose compilation is being deferred. */
-	static TMap<TRefCountPtr<FNiagaraShaderMap>, TArray<FNiagaraScript*> > NiagaraShaderMapsBeingCompiled;
+	static TMap<TRefCountPtr<FNiagaraShaderMap>, TArray<FNiagaraShaderScript*> > NiagaraShaderMapsBeingCompiled;
 
 	/** Uniquely identifies this shader map during compilation, needed for deferred compilation where shaders from multiple shader maps are compiled together. */
 	uint32 CompilingId;
@@ -429,9 +439,9 @@ private:
 	/** Debug information about how the shader map was compiled. */
 	FString DebugDescription;
 
-	FShader* ProcessCompilationResultsForSingleJob(class FNiagaraShaderCompileJob* SingleJob, const FSHAHash& ShaderMapHash);
+	FShader* ProcessCompilationResultsForSingleJob(class FShaderCommonCompileJob* SingleJob, const FSHAHash& ShaderMapHash);
 
-	bool IsNiagaraShaderComplete(const FNiagaraScript* Script, const FNiagaraShaderType* ShaderType, bool bSilent);
+	bool IsNiagaraShaderComplete(const FNiagaraShaderScript* Script, const FNiagaraShaderType* ShaderType, bool bSilent);
 
 	/** Initializes OrderedMeshShaderMaps from the contents of MeshShaderMaps. */
 	void InitOrderedMeshShaderMaps();
@@ -444,16 +454,16 @@ private:
 DECLARE_DYNAMIC_MULTICAST_DELEGATE(FOnNiagaraScriptCompilationComplete);
 
 /**
- * FNiagaraScript represents a Niagara script to the shader compilation process
+ * FNiagaraShaderScript represents a Niagara script to the shader compilation process
  */
-class FNiagaraScript
+class FNiagaraShaderScript
 {
 public:
 
 	/**
 	 * Minimal initialization constructor.
 	 */
-	FNiagaraScript() :
+	FNiagaraShaderScript() :
 		GameThreadShaderMap(NULL),
 		RenderingThreadShaderMap(NULL),
 		FeatureLevel(ERHIFeatureLevel::SM4),
@@ -463,11 +473,11 @@ public:
 	/**
 	 * Destructor
 	 */
-	NIAGARASHADER_API  virtual ~FNiagaraScript();
+	NIAGARASHADER_API  virtual ~FNiagaraShaderScript();
 
 	/**
 	 * Caches the shaders for this script with no static parameters on the given platform.
-	 * This is used by FNiagaraScript
+	 * This is used by FNiagaraShaderScript
 	 */
 	NIAGARASHADER_API  bool CacheShaders(EShaderPlatform Platform, bool bApplyCompletedShaderMapForRendering, bool bForceRecompile, bool bSynchronous = false);
 	bool CacheShaders(const FNiagaraShaderMapId& ShaderMapId, EShaderPlatform Platform, bool bApplyCompletedShaderMapForRendering, bool bForceRecompile, bool bSynchronous = false);
@@ -555,6 +565,10 @@ public:
 	/** Note: SetGameThreadShaderMap must also be called with the same value, but from the game thread. */
 	NIAGARASHADER_API  void SetRenderingThreadShaderMap(FNiagaraShaderMap* InShaderMap);
 
+	void AddCompileId(uint32 Id) 
+	{
+		OutstandingCompileShaderMapIds.Add(Id);
+	}
 
 	void SetShaderMap(FNiagaraShaderMap* InShaderMap)
 	{
@@ -591,34 +605,32 @@ public:
 
 	const FString& GetFriendlyName()	const { return FriendlyName; }
 
-	/** Gets the Guid that represents this script */
-	virtual FGuid GetScriptID() const { return ScriptId; }
 
-	NIAGARASHADER_API void SetScript(UNiagaraScript *InScript, ERHIFeatureLevel::Type InFeatureLevel, FGuid InScriptId, FString InFriendlyName);
+	NIAGARASHADER_API void SetScript(UNiagaraScript *InScript, ERHIFeatureLevel::Type InFeatureLevel, const FGuid& InCompilerVersion, const FGuid& InBaseScriptID, const TArray<FGuid>& InReferencedDependencyIds, FString InFriendlyName);
 
-	UNiagaraScript *GetGTScript()
+	UNiagaraScript *GetBaseVMScript()
 	{
-		return GTScript;
+		return BaseVMScript;
 	}
 
 	void SetCompileErrors(TArray<FString> &InErrors)
 	{
 		CompileErrors = InErrors;
 	}
+
+	FString SourceName;
 	
 	FString HlslOutput;
 
 	NIAGARASHADER_API  FNiagaraShader* GetShader() const;
 	NIAGARASHADER_API  FNiagaraShader* GetShaderGameThread() const;
 	
-	void SetDatainterfaceBufferDescriptors(TArray< FDIBufferDescriptorStore >&InBufferDescriptors)
-	{
-		DIBufferDescriptors = InBufferDescriptors;
-	}
+	NIAGARASHADER_API void SetDataInterfaceParamInfo(TArray< FNiagaraDataInterfaceGPUParamInfo >& InDIParamInfo);
+	NIAGARASHADER_API void SetDataInterfaceParamInfo(TArray< FNiagaraDataInterfaceParamRef >& InDIParamRefs);
 
-	TArray< FDIBufferDescriptorStore >& GetDataInterfaceBufferDescriptors()
+	TArray< FNiagaraDataInterfaceGPUParamInfo >& GetDataInterfaceParamInfo()
 	{
-		return DIBufferDescriptors;
+		return DIParamInfo;
 	}
 
 	NIAGARASHADER_API FOnNiagaraScriptCompilationComplete& OnCompilationComplete()
@@ -626,6 +638,7 @@ public:
 		return OnCompilationCompleteDelegate;
 	}
 
+	bool IsSame(const FNiagaraShaderMapId& InId) const;
 protected:
 
 	// shared code needed for GetUniformScalarParameterExpressions, GetUniformVectorParameterExpressions, GetUniformCubeTextureExpressions..
@@ -644,7 +657,7 @@ protected:
 	}
 
 private:
-	UNiagaraScript* GTScript;
+	UNiagaraScript* BaseVMScript;
 
 	TArray<FString> CompileErrors;
 
@@ -657,18 +670,22 @@ private:
 	TRefCountPtr<FNiagaraShaderMap> GameThreadShaderMap;
 
 	/** 
-	 * Shader map for this FniagaraScript which is accessible by the rendering thread. 
+	 * Shader map for this FNiagaraShaderScript which is accessible by the rendering thread. 
 	 * This must be updated along with GameThreadShaderMap, but on the rendering thread.
 	 */
 	FNiagaraShaderMap* RenderingThreadShaderMap;
 
-	// data interface buffer descriptors per data interface; these come from the HLSL translators and need to be passed down to the shader for binding
-	TArray< FDIBufferDescriptorStore > DIBufferDescriptors;
+	// Information describing data interface parameters. These come from the HLSL translators and need to be passed down to the shader for binding
+	TArray< FNiagaraDataInterfaceGPUParamInfo > DIParamInfo;
 
-	/** 
-	 * unique identifier
-	 */
-	FGuid ScriptId;
+	/** Guid id for base script*/
+	FGuid BaseScriptId;
+
+	/** The compiler version the script was generated with.*/
+	FGuid CompilerVersionId;
+
+	/** Dependencies of the script*/
+	TArray<FGuid> ReferencedDependencyIds;
 
 	/** 
 	 * Contains the compiling id of this shader map when it is being compiled asynchronously. 
diff --git a/Engine/Plugins/FX/Niagara/Source/NiagaraVertexFactories/Private/NiagaraMeshVertexFactory.cpp b/Engine/Plugins/FX/Niagara/Source/NiagaraVertexFactories/Private/NiagaraMeshVertexFactory.cpp
--- a/Engine/Plugins/FX/Niagara/Source/NiagaraVertexFactories/Private/NiagaraMeshVertexFactory.cpp
+++ b/Engine/Plugins/FX/Niagara/Source/NiagaraVertexFactories/Private/NiagaraMeshVertexFactory.cpp
@@ -19,18 +19,34 @@ public:
 	{
 		//PrevTransformBuffer.Bind(ParameterMap, TEXT("PrevTransformBuffer"));
 		NiagaraParticleDataFloat.Bind(ParameterMap, TEXT("NiagaraParticleDataFloat"));
-		NiagaraParticleDataInt.Bind(ParameterMap, TEXT("NiagaraParticleDataInt"));
-		SafeComponentBufferSizeParam.Bind(ParameterMap, TEXT("SafeComponentBufferSize"));
+		FloatDataOffset.Bind(ParameterMap, TEXT("NiagaraFloatDataOffset"));
+		FloatDataStride.Bind(ParameterMap, TEXT("NiagaraFloatDataStride"));
+
+		// 		NiagaraParticleDataInt.Bind(ParameterMap, TEXT("NiagaraParticleDataInt"));
+		// 		FloatDataOffset.Bind(ParameterMap, TEXT("NiagaraInt32DataOffset"));
+		// 		FloatDataStride.Bind(ParameterMap, TEXT("NiagaraInt3DataStride"));
+
 		MeshFacingMode.Bind(ParameterMap, TEXT("MeshFacingMode"));
+		SortedIndices.Bind(ParameterMap, TEXT("SortedIndices"));
+		SortedIndicesOffset.Bind(ParameterMap, TEXT("SortedIndicesOffset"));
+
 	}
 
 	virtual void Serialize(FArchive& Ar) override
 	{
 		//Ar << PrevTransformBuffer;
 		Ar << NiagaraParticleDataFloat;
-		Ar << NiagaraParticleDataInt;
-		Ar << SafeComponentBufferSizeParam;
+		Ar << FloatDataOffset;
+		Ar << FloatDataStride;
+
+		// 		Ar << NiagaraParticleDataInt;
+		// 		Ar << Int32DataOffset;
+		// 		Ar << Int32DataStride;
+
 		Ar << MeshFacingMode;
+
+		Ar << SortedIndices;
+		Ar << SortedIndicesOffset;
 	}
 
 	virtual void SetMesh(FRHICommandList& RHICmdList, FShader* Shader, const FVertexFactory* VertexFactory, const FSceneView& View, const FMeshBatchElement& BatchElement, uint32 DataFlags) const override
@@ -42,10 +58,16 @@ public:
 
 		//SetSRVParameter(RHICmdList, VertexShaderRHI, PrevTransformBuffer, NiagaraMeshVF->GetPreviousTransformBufferSRV());
 
-		SetSRVParameter(RHICmdList, VertexShaderRHI, NiagaraParticleDataFloat, NiagaraMeshVF->GetFloatDataSRV());
-		SetSRVParameter(RHICmdList, VertexShaderRHI, NiagaraParticleDataInt, NiagaraMeshVF->GetIntDataSRV());
-		SetShaderValue(RHICmdList, VertexShaderRHI, SafeComponentBufferSizeParam, NiagaraMeshVF->GetComponentBufferSize());
 		SetShaderValue(RHICmdList, VertexShaderRHI, MeshFacingMode, NiagaraMeshVF->GetMeshFacingMode());
+
+		SetSRVParameter(RHICmdList, VertexShaderRHI, NiagaraParticleDataFloat, NiagaraMeshVF->GetParticleDataFloatSRV());
+		SetShaderValue(RHICmdList, VertexShaderRHI, FloatDataOffset, NiagaraMeshVF->GetFloatDataOffset());
+		SetShaderValue(RHICmdList, VertexShaderRHI, FloatDataStride, NiagaraMeshVF->GetFloatDataStride());
+
+		//SetSRVParameter(RHICmdList, VertexShaderRHI, NiagaraParticleDataInt, NiagaraMeshVF->GetIntDataSRV());
+
+		SetSRVParameter(RHICmdList, VertexShaderRHI, SortedIndices, NiagaraMeshVF->GetSortedIndicesSRV());
+		SetShaderValue(RHICmdList, VertexShaderRHI, SortedIndicesOffset, NiagaraMeshVF->GetSortedIndicesOffset());
 	}
 
 private:
@@ -53,9 +75,16 @@ private:
 	//FShaderResourceParameter PrevTransformBuffer;
 
 	FShaderResourceParameter NiagaraParticleDataFloat;
-	FShaderResourceParameter NiagaraParticleDataInt;
-	FShaderParameter SafeComponentBufferSizeParam;
+	FShaderParameter FloatDataOffset;
+	FShaderParameter FloatDataStride;
+
+	// 	FShaderResourceParameter NiagaraParticleDataInt;
+	// 	FShaderParameter Int32DataOffset;
+	// 	FShaderParameter Int32DataStride;
+
 	FShaderParameter MeshFacingMode;
+	FShaderResourceParameter SortedIndices;
+	FShaderParameter SortedIndicesOffset;
 
 };
 
@@ -66,44 +95,7 @@ void FNiagaraMeshVertexFactory::InitRHI()
 
 	check(GRHISupportsInstancing);
 
-	if (Data.bInitialized)
 	{
-
-		// TODO: get rid of this; need to modify code further up to not skip entire draw when no streams are present here
-			{
-				//checkf(DynamicVertexStride != -1, TEXT("FNiagaraMeshVertexFactory does not have a valid DynamicVertexStride - likely an empty one was made, but SetStrides was not called"));
-				FVertexStream VertexStream;
-				VertexStream.VertexBuffer = NULL;
-				VertexStream.Stride = 0;
-				VertexStream.Offset = 0;
-				Streams.Add(VertexStream);
-
-				// @todo metal: this will need a valid stride when we get to instanced meshes!
-				Elements.Add(FVertexElement(0, Data.TransformComponent[0].Offset, Data.TransformComponent[0].Type, 8, DynamicVertexStride, EnumHasAnyFlags(EVertexStreamUsage::Instancing, Data.TransformComponent[0].VertexStreamUsage)));
-				Elements.Add(FVertexElement(0, Data.TransformComponent[1].Offset, Data.TransformComponent[1].Type, 9, DynamicVertexStride, EnumHasAnyFlags(EVertexStreamUsage::Instancing, Data.TransformComponent[1].VertexStreamUsage)));
-				Elements.Add(FVertexElement(0, Data.TransformComponent[2].Offset, Data.TransformComponent[2].Type, 10, DynamicVertexStride, EnumHasAnyFlags(EVertexStreamUsage::Instancing, Data.TransformComponent[2].VertexStreamUsage)));
-
-				Elements.Add(FVertexElement(0, Data.SubUVs.Offset, Data.SubUVs.Type, 11, DynamicVertexStride, EnumHasAnyFlags(EVertexStreamUsage::Instancing, Data.SubUVs.VertexStreamUsage)));
-				Elements.Add(FVertexElement(0, Data.SubUVLerpAndRelTime.Offset, Data.SubUVLerpAndRelTime.Type, 12, DynamicVertexStride, EnumHasAnyFlags(EVertexStreamUsage::Instancing, Data.SubUVLerpAndRelTime.VertexStreamUsage)));
-
-				Elements.Add(FVertexElement(0, Data.ParticleColorComponent.Offset, Data.ParticleColorComponent.Type, 14, DynamicVertexStride, EnumHasAnyFlags(EVertexStreamUsage::Instancing, Data.ParticleColorComponent.VertexStreamUsage)));
-				Elements.Add(FVertexElement(0, Data.VelocityComponent.Offset, Data.VelocityComponent.Type, 15, DynamicVertexStride, EnumHasAnyFlags(EVertexStreamUsage::Instancing, Data.VelocityComponent.VertexStreamUsage)));
-			}
-
-			// Stream 1 - Dynamic parameter
-			{
-				checkf(DynamicParameterVertexStride != -1, TEXT("FNiagaraMeshVertexFactory does not have a valid DynamicParameterVertexStride - likely an empty one was made, but SetStrides was not called"));
-
-				FVertexStream VertexStream;
-				VertexStream.VertexBuffer = NULL;
-				VertexStream.Stride = 0;
-				VertexStream.Offset = 0;
-				Streams.Add(VertexStream);
-
-				Elements.Add(FVertexElement(1, 0, VET_Float4, 13, DynamicParameterVertexStride, true));
-			}
-
-
 		if (Data.PositionComponent.VertexBuffer != NULL)
 		{
 			Elements.Add(AccessStreamComponent(Data.PositionComponent, 0));
@@ -166,29 +158,6 @@ void FNiagaraMeshVertexFactory::InitRHI()
 	}
 }
 
-void FNiagaraMeshVertexFactory::SetInstanceBuffer(const FVertexBuffer* InstanceBuffer, uint32 StreamOffset, uint32 Stride)
-{
-	Streams[0].VertexBuffer = InstanceBuffer;
-	Streams[0].Offset = StreamOffset;
-	Streams[0].Stride = Stride;
-}
-
-void FNiagaraMeshVertexFactory::SetDynamicParameterBuffer(const FVertexBuffer* InDynamicParameterBuffer, uint32 StreamOffset, uint32 Stride)
-{
-	if (InDynamicParameterBuffer)
-	{
-		Streams[1].VertexBuffer = InDynamicParameterBuffer;
-		Streams[1].Offset = StreamOffset;
-		Streams[1].Stride = Stride;
-	}
-	else
-	{
-		Streams[1].VertexBuffer = &GNullDynamicParameterVertexBuffer;
-		Streams[1].Offset = 0;
-		Streams[1].Stride = 0;
-	}
-}
-
 /*
 uint8* FNiagaraMeshVertexFactory::LockPreviousTransformBuffer(uint32 ParticleCount)
 {
@@ -227,7 +196,7 @@ bool FNiagaraMeshVertexFactory::ShouldCompilePermutation(EShaderPlatform Platfor
 	return (!IsMobilePlatform(Platform) && !IsSwitchPlatform(Platform) && (Material->IsUsedWithNiagaraMeshParticles() || Material->IsSpecialEngineMaterial()));
 }
 
-void FNiagaraMeshVertexFactory::SetData(const FDataType& InData)
+void FNiagaraMeshVertexFactory::SetData(const FStaticMeshDataType& InData)
 {
 	check(IsInRenderingThread());
 	Data = InData;
diff --git a/Engine/Plugins/FX/Niagara/Source/NiagaraVertexFactories/Private/NiagaraRibbonVertexFactory.cpp b/Engine/Plugins/FX/Niagara/Source/NiagaraVertexFactories/Private/NiagaraRibbonVertexFactory.cpp
--- a/Engine/Plugins/FX/Niagara/Source/NiagaraVertexFactories/Private/NiagaraRibbonVertexFactory.cpp
+++ b/Engine/Plugins/FX/Niagara/Source/NiagaraVertexFactories/Private/NiagaraRibbonVertexFactory.cpp
@@ -135,9 +135,11 @@ public:
 		Elements.Add(FVertexElement(0, Offset, VET_Float3, 6, Stride));
 		Offset += sizeof(float) * 3;
 
-		/** Dynamic parameters come from a second stream */
+		/** Dynamic parameters come from additional streams */
 		Elements.Add(FVertexElement(1, 0, VET_Float4, 7, sizeof(FVector4)));
-		Offset += sizeof(float) * 4;
+		Elements.Add(FVertexElement(2, 0, VET_Float4, 8, sizeof(FVector4)));
+		Elements.Add(FVertexElement(3, 0, VET_Float4, 9, sizeof(FVector4)));
+		Elements.Add(FVertexElement(4, 0, VET_Float4, 10, sizeof(FVector4)));
 
 	}
 
@@ -174,7 +176,7 @@ bool FNiagaraRibbonVertexFactory::ShouldCompilePermutation(EShaderPlatform Platf
 */
 void FNiagaraRibbonVertexFactory::ModifyCompilationEnvironment(EShaderPlatform Platform, const class FMaterial* Material, FShaderCompilerEnvironment& OutEnvironment)
 {
-	FParticleVertexFactoryBase::ModifyCompilationEnvironment(Platform, Material, OutEnvironment);
+	FNiagaraVertexFactoryBase::ModifyCompilationEnvironment(Platform, Material, OutEnvironment);
 	OutEnvironment.SetDefine(TEXT("PARTICLE_BEAMTRAIL_FACTORY"), TEXT("1"));
 }
 
@@ -187,6 +189,9 @@ void FNiagaraRibbonVertexFactory::InitRHI()
 
 	FVertexStream* VertexStream = new(Streams) FVertexStream;
 	FVertexStream* DynamicParameterStream = new(Streams) FVertexStream;
+	FVertexStream* DynamicParameter1Stream = new(Streams) FVertexStream;
+	FVertexStream* DynamicParameter2Stream = new(Streams) FVertexStream;
+	FVertexStream* DynamicParameter3Stream = new(Streams) FVertexStream;
 }
 
 FVertexFactoryShaderParameters* FNiagaraRibbonVertexFactory::ConstructShaderParameters(EShaderFrequency ShaderFrequency)
@@ -204,17 +209,17 @@ FVertexFactoryShaderParameters* FNiagaraRibbonVertexFactory::ConstructShaderPara
 
 void FNiagaraRibbonVertexFactory::SetVertexBuffer(const FVertexBuffer* InBuffer, uint32 StreamOffset, uint32 Stride)
 {
-	check(Streams.Num() == 2);
+	check(Streams.Num() == 5);
 	FVertexStream& VertexStream = Streams[0];
 	VertexStream.VertexBuffer = InBuffer;
 	VertexStream.Stride = Stride;
 	VertexStream.Offset = StreamOffset;
 }
 
-void FNiagaraRibbonVertexFactory::SetDynamicParameterBuffer(const FVertexBuffer* InDynamicParameterBuffer, uint32 StreamOffset, uint32 Stride)
+void FNiagaraRibbonVertexFactory::SetDynamicParameterBuffer(const FVertexBuffer* InDynamicParameterBuffer, int32 ParameterIndex, uint32 StreamOffset, uint32 Stride)
 {
-	check(Streams.Num() == 2);
-	FVertexStream& DynamicParameterStream = Streams[1];
+	check(Streams.Num() == 5);
+	FVertexStream& DynamicParameterStream = Streams[1 + ParameterIndex];
 	if (InDynamicParameterBuffer)
 	{
 		DynamicParameterStream.VertexBuffer = InDynamicParameterBuffer;
diff --git a/Engine/Plugins/FX/Niagara/Source/NiagaraVertexFactories/Private/NiagaraSpriteVertexFactory.cpp b/Engine/Plugins/FX/Niagara/Source/NiagaraVertexFactories/Private/NiagaraSpriteVertexFactory.cpp
--- a/Engine/Plugins/FX/Niagara/Source/NiagaraVertexFactories/Private/NiagaraSpriteVertexFactory.cpp
+++ b/Engine/Plugins/FX/Niagara/Source/NiagaraVertexFactories/Private/NiagaraSpriteVertexFactory.cpp
@@ -64,11 +64,20 @@ public:
 	{
 		NumCutoutVerticesPerFrame.Bind(ParameterMap, TEXT("NumCutoutVerticesPerFrame"));
 		CutoutGeometry.Bind(ParameterMap, TEXT("CutoutGeometry"));
+
 		NiagaraParticleDataFloat.Bind(ParameterMap, TEXT("NiagaraParticleDataFloat"));
-		NiagaraParticleDataInt.Bind(ParameterMap, TEXT("NiagaraParticleDataInt"));
-		SafeComponentBufferSizeParam.Bind(ParameterMap, TEXT("SafeComponentBufferSize"));
+		FloatDataOffset.Bind(ParameterMap, TEXT("NiagaraFloatDataOffset"));
+		FloatDataStride.Bind(ParameterMap, TEXT("NiagaraFloatDataStride"));
+
+// 		NiagaraParticleDataInt.Bind(ParameterMap, TEXT("NiagaraParticleDataInt"));
+// 		FloatDataOffset.Bind(ParameterMap, TEXT("NiagaraInt32DataOffset"));
+// 		FloatDataStride.Bind(ParameterMap, TEXT("NiagaraInt3DataStride"));
+
 		ParticleAlignmentMode.Bind(ParameterMap, TEXT("ParticleAlignmentMode"));
 		ParticleFacingMode.Bind(ParameterMap, TEXT("ParticleFacingMode"));
+
+		SortedIndices.Bind(ParameterMap, TEXT("SortedIndices"));
+		SortedIndicesOffset.Bind(ParameterMap, TEXT("SortedIndicesOffset"));
 	}
 
 	virtual void Serialize(FArchive& Ar) override
@@ -77,15 +86,17 @@ public:
 		Ar << CutoutGeometry;
 		Ar << ParticleFacingMode;
 		Ar << ParticleAlignmentMode;
+
 		Ar << NiagaraParticleDataFloat;
-		Ar << NiagaraParticleDataInt;
-		Ar << SafeComponentBufferSizeParam;
-
-		Ar << PositionOffsetParam;
-		Ar << SizeOffsetParam;
-		Ar << RotationOffsetParam;
-		Ar << SubimgOffsetParam;
-		Ar << ColorOffsetParam;
+		Ar << FloatDataOffset;
+		Ar << FloatDataStride;
+
+// 		Ar << NiagaraParticleDataInt;
+// 		Ar << Int32DataOffset;
+// 		Ar << Int32DataStride;
+
+		Ar << SortedIndices;
+		Ar << SortedIndicesOffset;
 	}
 
 	virtual void SetMesh(FRHICommandList& RHICmdList, FShader* Shader,const FVertexFactory* VertexFactory,const FSceneView& View,const FMeshBatchElement& BatchElement,uint32 DataFlags) const override
@@ -100,10 +111,13 @@ public:
 
 		SetShaderValue(RHICmdList, VertexShaderRHI, ParticleAlignmentMode, SpriteVF->GetAlignmentMode());
 		SetShaderValue(RHICmdList, VertexShaderRHI, ParticleFacingMode, SpriteVF->GetFacingMode());
+		
+		SetSRVParameter(RHICmdList, VertexShaderRHI, NiagaraParticleDataFloat, SpriteVF->GetParticleDataFloatSRV());
+		SetShaderValue(RHICmdList, VertexShaderRHI, FloatDataOffset, SpriteVF->GetFloatDataOffset());
+		SetShaderValue(RHICmdList, VertexShaderRHI, FloatDataStride, SpriteVF->GetFloatDataStride());
 
-		SetSRVParameter(RHICmdList, VertexShaderRHI, NiagaraParticleDataFloat, SpriteVF->GetFloatDataSRV());
-		SetSRVParameter(RHICmdList, VertexShaderRHI, NiagaraParticleDataInt, SpriteVF->GetIntDataSRV());
-		SetShaderValue(RHICmdList, VertexShaderRHI, SafeComponentBufferSizeParam, SpriteVF->GetComponentBufferSize());
+		SetSRVParameter(RHICmdList, VertexShaderRHI, SortedIndices, SpriteVF->GetSortedIndicesSRV());
+		SetShaderValue(RHICmdList, VertexShaderRHI, SortedIndicesOffset, SpriteVF->GetSortedIndicesOffset());
 	}
 
 private:
@@ -113,15 +127,17 @@ private:
 	FShaderParameter ParticleFacingMode;
 
 	FShaderResourceParameter CutoutGeometry;
+
 	FShaderResourceParameter NiagaraParticleDataFloat;
-	FShaderResourceParameter NiagaraParticleDataInt;
-	FShaderParameter SafeComponentBufferSizeParam;
-
-	FShaderParameter PositionOffsetParam;
-	FShaderParameter SizeOffsetParam;
-	FShaderParameter RotationOffsetParam;
-	FShaderParameter SubimgOffsetParam;
-	FShaderParameter ColorOffsetParam;
+	FShaderParameter FloatDataOffset;
+	FShaderParameter FloatDataStride;
+
+// 	FShaderResourceParameter NiagaraParticleDataInt;
+// 	FShaderParameter Int32DataOffset;
+// 	FShaderParameter Int32DataStride;
+	
+	FShaderResourceParameter SortedIndices;
+	FShaderParameter SortedIndicesOffset;
 };
 
 class FNiagaraSpriteVertexFactoryShaderParametersPS : public FNiagaraSpriteVertexFactoryShaderParameters
@@ -159,47 +175,9 @@ public:
 	virtual void FillDeclElements(FVertexDeclarationElementList& Elements, int32& Offset)
 	{
 		uint32 InitialStride = sizeof(float) * 2;
-		uint32 PerParticleStride = (sizeof(float) * 4*5 + sizeof(float)*3*2);
-
 		/** The stream to read the texture coordinates from. */
 		check( Offset == 0 );
-		uint32 Stride = bInstanced ? InitialStride : InitialStride + PerParticleStride;
-		Elements.Add(FVertexElement(0, Offset, VET_Float2, 4, Stride, false));
-		Offset += sizeof(float) * 2;
-
-		/** The per-particle streams follow. */
-		if(bInstanced) 
-		{
-			Offset = 0;
-			// update stride
-			Stride = PerParticleStride;
-		}
-
-		/** The stream to read the vertex position from. */
-		Elements.Add(FVertexElement(bInstanced ? 1 : 0, Offset, VET_Float4, 0, Stride, bInstanced));
-		Offset += sizeof(float) * 4;
-		/** The stream to read the vertex old position from. */
-		Elements.Add(FVertexElement(bInstanced ? 1 : 0, Offset, VET_Float4, 1, Stride, bInstanced));
-		Offset += sizeof(float) * 4;
-		/** The stream to read the vertex size/rot/subimage from. */
-		Elements.Add(FVertexElement(bInstanced ? 1 : 0, Offset, VET_Float4, 2, Stride, bInstanced));
-		Offset += sizeof(float) * 4;
-		/** The stream to read the color from.					*/
-		Elements.Add(FVertexElement(bInstanced ? 1 : 0, Offset, VET_Float4, 3, Stride, bInstanced));
-		Offset += sizeof(float) * 4;
-
-		/** The per-particle dynamic parameter stream */
-		Elements.Add(FVertexElement(bInstanced ? 1 : 0, Offset, VET_Float4, 5, Stride, bInstanced));
-		Offset += sizeof(float) * 4;
-
-		/** The stream to read the custom alignment vector from. */
-		Elements.Add(FVertexElement(bInstanced ? 1 : 0, Offset, VET_Float3, 6, Stride, bInstanced));
-		Offset += sizeof(float) * 3;
-		/** The stream to read the custom facing vector from. */
-		Elements.Add(FVertexElement(bInstanced ? 1 : 0, Offset, VET_Float3, 7, Stride, bInstanced));
-		Offset += sizeof(float) * 3;
-
-
+		Elements.Add(FVertexElement(0, Offset, VET_Float2, 0, InitialStride, false));
 	}
 
 	virtual void InitDynamicRHI()
@@ -247,7 +225,7 @@ inline TGlobalResource<FNiagaraSpriteVertexDeclaration>& GetNiagaraSpriteVertexD
 
 bool FNiagaraSpriteVertexFactory::ShouldCompilePermutation(EShaderPlatform Platform, const class FMaterial* Material, const class FShaderType* ShaderType)
 {
-	return (!IsMobilePlatform(Platform) && !IsSwitchPlatform(Platform) && (Material->IsUsedWithNiagaraSprites() || Material->IsSpecialEngineMaterial()));
+	return (!IsSwitchPlatform(Platform) && (Material->IsUsedWithNiagaraSprites() || Material->IsSpecialEngineMaterial()));
 }
 
 /**
@@ -255,7 +233,7 @@ bool FNiagaraSpriteVertexFactory::ShouldCompilePermutation(EShaderPlatform Platf
  */
 void FNiagaraSpriteVertexFactory::ModifyCompilationEnvironment(EShaderPlatform Platform, const class FMaterial* Material, FShaderCompilerEnvironment& OutEnvironment)
 {
-	FParticleVertexFactoryBase::ModifyCompilationEnvironment(Platform, Material, OutEnvironment);
+	FNiagaraVertexFactoryBase::ModifyCompilationEnvironment(Platform, Material, OutEnvironment);
 
 	// Set a define so we can tell in MaterialTemplate.usf when we are compiling a sprite vertex factory
 	OutEnvironment.SetDefine(TEXT("PARTICLE_SPRITE_FACTORY"),TEXT("1"));
@@ -282,16 +260,6 @@ void FNiagaraSpriteVertexFactory::InitStreams()
 		TexCoordStream->Stride = sizeof(FVector2D);
 		TexCoordStream->Offset = 0;
 	}
-	FVertexStream* InstanceStream = new(Streams) FVertexStream;
-}
-
-void FNiagaraSpriteVertexFactory::SetInstanceBuffer(const FVertexBuffer* InInstanceBuffer, uint32 StreamOffset, uint32 Stride, bool bInstanced)
-{
-	check(Streams.Num() == (bInstanced ? 2 : 1));
-	FVertexStream& InstanceStream = Streams[bInstanced ? 1 : 0];
-	InstanceStream.VertexBuffer = InInstanceBuffer;
-	InstanceStream.Stride = Stride;
-	InstanceStream.Offset = StreamOffset;
 }
 
 void FNiagaraSpriteVertexFactory::SetTexCoordBuffer(const FVertexBuffer* InTexCoordBuffer)
diff --git a/Engine/Plugins/FX/Niagara/Source/NiagaraVertexFactories/Private/NiagaraVertexFactoriesModule.cpp b/Engine/Plugins/FX/Niagara/Source/NiagaraVertexFactories/Private/NiagaraVertexFactoriesModule.cpp
--- a/Engine/Plugins/FX/Niagara/Source/NiagaraVertexFactories/Private/NiagaraVertexFactoriesModule.cpp
+++ b/Engine/Plugins/FX/Niagara/Source/NiagaraVertexFactories/Private/NiagaraVertexFactoriesModule.cpp
@@ -5,5 +5,3 @@
 #include "NiagaraVertexFactory.h"
 
 IMPLEMENT_MODULE(INiagaraVertexFactoriesModule, NiagaraVertexFactories);
-
-FRWBuffer FNiagaraVertexFactoryBase::DummyBuffer;
\ No newline at end of file
diff --git a/Engine/Plugins/FX/Niagara/Source/NiagaraVertexFactories/Public/NiagaraMeshVertexFactory.h b/Engine/Plugins/FX/Niagara/Source/NiagaraVertexFactories/Public/NiagaraMeshVertexFactory.h
--- a/Engine/Plugins/FX/Niagara/Source/NiagaraVertexFactories/Public/NiagaraMeshVertexFactory.h
+++ b/Engine/Plugins/FX/Niagara/Source/NiagaraVertexFactories/Public/NiagaraMeshVertexFactory.h
@@ -15,6 +15,8 @@ ParticleVertexFactory.h: Particle vertex factory definitions.
 #include "Components.h"
 #include "SceneManagement.h"
 #include "VertexFactory.h"
+#include "NiagaraGlobalReadBuffer.h"
+
 
 
 class FMaterial;
@@ -22,43 +24,6 @@ class FVertexBuffer;
 struct FDynamicReadBuffer;
 struct FShaderCompilerEnvironment;
 
-// Per-particle data sent to the GPU.
-struct FNiagaraMeshInstanceVertex
-{
-	/** The color of the particle. */
-	FLinearColor Color;
-
-	/** The instance to world transform of the particle. Translation vector is packed into W components. */
-	FVector4 Transform[3];
-
-	/** The velocity of the particle, XYZ: direction, W: speed. */
-	FVector4 Velocity;
-
-	/** The sub-image texture offsets for the particle. */
-	int16 SubUVParams[4];
-
-	/** The sub-image lerp value for the particle. */
-	float SubUVLerp;
-
-	/** The relative time of the particle. */
-	float RelativeTime;
-};
-
-struct FNiagaraMeshInstanceVertexDynamicParameter
-{
-	/** The dynamic parameter of the particle. */
-	float DynamicValue[4];
-};
-
-struct FNiagaraMeshInstanceVertexPrevTransform
-{
-	FVector4 PrevTransform0;
-	FVector4 PrevTransform1;
-	FVector4 PrevTransform2;
-};
-
-
-
 
 /**
 * Uniform buffer for mesh particle vertex factories.
@@ -78,6 +43,10 @@ BEGIN_UNIFORM_BUFFER_STRUCT(FNiagaraMeshUniformParameters, NIAGARAVERTEXFACTORIE
 	UNIFORM_MEMBER(int, ScaleDataOffset)
 	UNIFORM_MEMBER(int, SizeDataOffset)
 	UNIFORM_MEMBER(int, MaterialParamDataOffset)
+	UNIFORM_MEMBER(int, MaterialParam1DataOffset)
+	UNIFORM_MEMBER(int, MaterialParam2DataOffset)
+	UNIFORM_MEMBER(int, MaterialParam3DataOffset)
+	UNIFORM_MEMBER(FVector4, DefaultPos)
 END_UNIFORM_BUFFER_STRUCT(FNiagaraMeshUniformParameters)
 
 typedef TUniformBufferRef<FNiagaraMeshUniformParameters> FNiagaraMeshUniformBufferRef;
@@ -92,57 +61,24 @@ class NIAGARAVERTEXFACTORIES_API FNiagaraMeshVertexFactory : public FNiagaraVert
 {
 	DECLARE_VERTEX_FACTORY_TYPE(FNiagaraMeshVertexFactory);
 public:
-
-	// TODO get rid of the unnecessary components here when streams are no longer necessary
-	struct FDataType : public FStaticMeshDataType
-	{
-		/** The stream to read the vertex  color from. */
-		FVertexStreamComponent ParticleColorComponent;
-
-		/** The stream to read the mesh transform from. */
-		FVertexStreamComponent TransformComponent[3];
-
-		/** The stream to read the particle velocity from */
-		FVertexStreamComponent VelocityComponent;
-
-		/** The stream to read SubUV parameters from.. */
-		FVertexStreamComponent SubUVs;
-
-		/** The stream to read SubUV lerp and the particle relative time from */
-		FVertexStreamComponent SubUVLerpAndRelTime;
-
-		/** Flag to mark as initialized */
-		bool bInitialized;
-
-		FDataType()
-			: bInitialized(false)
-		{
-		}
-	};
-
-	class FBatchParametersCPU : public FOneFrameResource
-	{
-	public:
-		const struct FNiagaraMeshInstanceVertex* InstanceBuffer;
-		const struct FNiagaraMeshInstanceVertexDynamicParameter* DynamicParameterBuffer;
-		const struct FNiagaraMeshInstanceVertexPrevTransform* PrevTransformBuffer;
-	};
-
+	
 	/** Default constructor. */
-	FNiagaraMeshVertexFactory(ENiagaraVertexFactoryType InType, ERHIFeatureLevel::Type InFeatureLevel, int32 InDynamicVertexStride, int32 InDynamicParameterVertexStride)
+	FNiagaraMeshVertexFactory(ENiagaraVertexFactoryType InType, ERHIFeatureLevel::Type InFeatureLevel)
 		: FNiagaraVertexFactoryBase(InType, InFeatureLevel)
-		, DynamicVertexStride(InDynamicVertexStride)
-		, DynamicParameterVertexStride(InDynamicParameterVertexStride)
 		, MeshFacingMode(0)
 		, InstanceVerticesCPU(nullptr)
+		, FloatDataOffset(0)
+		, FloatDataStride(0)
+		, SortedIndicesOffset(0)
 	{}
 
 	FNiagaraMeshVertexFactory()
 		: FNiagaraVertexFactoryBase(NVFT_MAX, ERHIFeatureLevel::Num)
-		, DynamicVertexStride(-1)
-		, DynamicParameterVertexStride(-1)
 		, MeshFacingMode(0)
 		, InstanceVerticesCPU(nullptr)
+		, FloatDataOffset(0)
+		, FloatDataStride(0)
+		, SortedIndicesOffset(0)
 	{}
 
 	/**
@@ -164,41 +100,48 @@ public:
 		OutEnvironment.SetDefine(TEXT("NIAGARA_MESH_INSTANCED"), TEXT("1"));
 	}
 
-	void SetParticleData(const FNiagaraDataSet *InDataSet)
+	void SetParticleData(const FShaderResourceViewRHIRef& InParticleDataFloatSRV, uint32 InFloatDataOffset, uint32 InFloatDataStride)
 	{
-		DataSet = InDataSet;
+		ParticleDataFloatSRV = InParticleDataFloatSRV;
+		FloatDataOffset = InFloatDataOffset;
+		FloatDataStride = InFloatDataStride;
 	}
 
-	inline FShaderResourceViewRHIParamRef GetFloatDataSRV() const
+	void SetSortedIndices(const FShaderResourceViewRHIRef& InSortedIndicesSRV, uint32 InSortedIndicesOffset)
 	{
-		const FShaderResourceViewRHIRef& Ret = DataSet->GetRenderDataFloatSRV();
-		if (Ret.IsValid())
-		{
-			return Ret;
-		}
-		return DummyBuffer.SRV;
+		SortedIndicesSRV = InSortedIndicesSRV;
+		SortedIndicesOffset = InSortedIndicesOffset;
 	}
 
-	inline FShaderResourceViewRHIParamRef GetIntDataSRV() const
+	FORCEINLINE FShaderResourceViewRHIRef GetParticleDataFloatSRV()
 	{
-		const FShaderResourceViewRHIRef& Ret = DataSet->GetRenderDataInt32SRV();
-		if (Ret.IsValid())
-		{
-			return Ret;
-		}
-		return DummyBuffer.SRV;
+		return ParticleDataFloatSRV;
 	}
 
-	uint32 GetComponentBufferSize()
+	FORCEINLINE int32 GetFloatDataOffset()
 	{
-		check(!IsInGameThread());
-		return DataSet->CurrDataRender().GetFloatStride() / sizeof(float);
+		return FloatDataOffset;
+	}
+
+	FORCEINLINE int32 GetFloatDataStride()
+	{
+		return FloatDataStride;
+	}
+
+	FORCEINLINE FShaderResourceViewRHIRef GetSortedIndicesSRV()
+	{
+		return SortedIndicesSRV;
+	}
+
+	FORCEINLINE int32 GetSortedIndicesOffset()
+	{
+		return SortedIndicesOffset;
 	}
 
 	/**
 	* An implementation of the interface used by TSynchronizedResource to update the resource with new data from the game thread.
 	*/
-	void SetData(const FDataType& InData);
+	void SetData(const FStaticMeshDataType& InData);
 
 	/**
 	* Set the uniform buffer for this vertex factory.
@@ -215,26 +158,7 @@ public:
 	{
 		return MeshParticleUniformBuffer;
 	}
-
-	/**
-	* Update the data strides (MUST HAPPEN BEFORE InitRHI is called)
-	*/
-	void SetStrides(int32 InDynamicVertexStride, int32 InDynamicParameterVertexStride)
-	{
-		DynamicVertexStride = InDynamicVertexStride;
-		DynamicParameterVertexStride = InDynamicParameterVertexStride;
-	}
-
-	/**
-	* Set the source vertex buffer that contains particle instance data.
-	*/
-	void SetInstanceBuffer(const FVertexBuffer* InstanceBuffer, uint32 StreamOffset, uint32 Stride);
-
-	/**
-	* Set the source vertex buffer that contains particle dynamic parameter data.
-	*/
-	void SetDynamicParameterBuffer(const FVertexBuffer* InDynamicParameterBuffer, uint32 StreamOffset, uint32 Stride);
-
+	
 	//uint8* LockPreviousTransformBuffer(uint32 ParticleCount);
 	//void UnlockPreviousTransformBuffer();
 	//FShaderResourceViewRHIParamRef GetPreviousTransformBufferSRV() const;
@@ -251,12 +175,7 @@ public:
 	static bool SupportsTessellationShaders() { return true; }
 
 	static FVertexFactoryShaderParameters* ConstructShaderParameters(EShaderFrequency ShaderFrequency);
-
-	FNiagaraMeshInstanceVertices*& GetInstanceVerticesCPU()
-	{
-		return InstanceVerticesCPU;
-	}
-
+	
 	uint32 GetMeshFacingMode() const
 	{
 		return MeshFacingMode;
@@ -268,11 +187,7 @@ public:
 	}
 
 protected:
-	FDataType Data;
-	const FNiagaraDataSet *DataSet;
-	/** Stride information for instanced mesh particles */
-	int32 DynamicVertexStride;
-	int32 DynamicParameterVertexStride;
+	FStaticMeshDataType Data;
 	uint32 MeshFacingMode;
 
 	/** Uniform buffer with mesh particle parameters. */
@@ -280,6 +195,13 @@ protected:
 	
 	/** Used to remember this in the case that we reuse the same vertex factory for multiple renders . */
 	FNiagaraMeshInstanceVertices* InstanceVerticesCPU;
+
+	FShaderResourceViewRHIRef ParticleDataFloatSRV;
+	uint32 FloatDataOffset;
+	uint32 FloatDataStride;
+
+	FShaderResourceViewRHIRef SortedIndicesSRV;
+	uint32 SortedIndicesOffset;
 };
 
 
@@ -288,8 +210,8 @@ class NIAGARAVERTEXFACTORIES_API FNiagaraMeshVertexFactoryEmulatedInstancing : p
 	DECLARE_VERTEX_FACTORY_TYPE(FMeshParticleVertexFactoryEmulatedInstancing);
 
 public:
-	FNiagaraMeshVertexFactoryEmulatedInstancing(ENiagaraVertexFactoryType InType, ERHIFeatureLevel::Type InFeatureLevel, int32 InDynamicVertexStride, int32 InDynamicParameterVertexStride)
-		: FNiagaraMeshVertexFactory(InType, InFeatureLevel, InDynamicVertexStride, InDynamicParameterVertexStride)
+	FNiagaraMeshVertexFactoryEmulatedInstancing(ENiagaraVertexFactoryType InType, ERHIFeatureLevel::Type InFeatureLevel)
+		: FNiagaraMeshVertexFactory(InType, InFeatureLevel)
 	{}
 
 	FNiagaraMeshVertexFactoryEmulatedInstancing()
@@ -322,14 +244,14 @@ inline FNiagaraMeshVertexFactory* ConstructNiagaraMeshVertexFactory()
 	}
 }
 
-inline FNiagaraMeshVertexFactory* ConstructNiagaraMeshVertexFactory(ENiagaraVertexFactoryType InType, ERHIFeatureLevel::Type InFeatureLevel, int32 InDynamicVertexStride, int32 InDynamicParameterVertexStride)
+inline FNiagaraMeshVertexFactory* ConstructNiagaraMeshVertexFactory(ENiagaraVertexFactoryType InType, ERHIFeatureLevel::Type InFeatureLevel)
 {
 	if (GRHISupportsInstancing)
 	{
-		return new FNiagaraMeshVertexFactory(InType, InFeatureLevel, InDynamicVertexStride, InDynamicParameterVertexStride);
+		return new FNiagaraMeshVertexFactory(InType, InFeatureLevel);
 	}
 	else
 	{
-		return new FNiagaraMeshVertexFactoryEmulatedInstancing(InType, InFeatureLevel, InDynamicVertexStride, InDynamicParameterVertexStride);
+		return new FNiagaraMeshVertexFactoryEmulatedInstancing(InType, InFeatureLevel);
 	}
 }
diff --git a/Engine/Plugins/FX/Niagara/Source/NiagaraVertexFactories/Public/NiagaraRibbonVertexFactory.h b/Engine/Plugins/FX/Niagara/Source/NiagaraVertexFactories/Public/NiagaraRibbonVertexFactory.h
--- a/Engine/Plugins/FX/Niagara/Source/NiagaraVertexFactories/Public/NiagaraRibbonVertexFactory.h
+++ b/Engine/Plugins/FX/Niagara/Source/NiagaraVertexFactories/Public/NiagaraRibbonVertexFactory.h
@@ -133,7 +133,7 @@ public:
 	/**
 	* Set the source vertex buffer that contains particle dynamic parameter data.
 	*/
-	void SetDynamicParameterBuffer(const FVertexBuffer* InDynamicParameterBuffer, uint32 StreamOffset, uint32 Stride);
+	void SetDynamicParameterBuffer(const FVertexBuffer* InDynamicParameterBuffer, int32 ParameterIndex, uint32 StreamOffset, uint32 Stride);
 
 
 	void SetParticleData(const FNiagaraDataSet *InDataSet)
@@ -142,26 +142,6 @@ public:
 	}
 
 
-	inline FShaderResourceViewRHIParamRef GetFloatDataSRV() const
-	{
-		check(!IsInGameThread());
-		return DataSet->GetRenderDataFloat().SRV;
-	}
-	inline FShaderResourceViewRHIParamRef GetIntDataSRV() const
-	{
-		check(!IsInGameThread());
-		return DataSet->GetRenderDataInt32().SRV;
-	}
-
-	uint32 GetComponentBufferSize()
-	{
-		check(!IsInGameThread());
-		return DataSet->CurrDataRender().GetFloatStride() / sizeof(float);
-	}
-
-
-
-
 	/**
 	* Construct shader parameters for this type of vertex factory.
 	*/
diff --git a/Engine/Plugins/FX/Niagara/Source/NiagaraVertexFactories/Public/NiagaraSpriteVertexFactory.h b/Engine/Plugins/FX/Niagara/Source/NiagaraVertexFactories/Public/NiagaraSpriteVertexFactory.h
--- a/Engine/Plugins/FX/Niagara/Source/NiagaraVertexFactories/Public/NiagaraSpriteVertexFactory.h
+++ b/Engine/Plugins/FX/Niagara/Source/NiagaraVertexFactories/Public/NiagaraSpriteVertexFactory.h
@@ -12,6 +12,7 @@
 #include "NiagaraVertexFactory.h"
 #include "NiagaraDataSet.h"
 #include "SceneView.h"
+#include "NiagaraGlobalReadBuffer.h"
 
 class FMaterial;
 
@@ -32,7 +33,7 @@ BEGIN_UNIFORM_BUFFER_STRUCT( FNiagaraSpriteUniformParameters, NIAGARAVERTEXFACTO
 	UNIFORM_MEMBER_EX( float, RotationScale, EShaderPrecisionModifier::Half )
 	UNIFORM_MEMBER_EX( float, RotationBias, EShaderPrecisionModifier::Half )
 	UNIFORM_MEMBER_EX( float, NormalsType, EShaderPrecisionModifier::Half )
-	UNIFORM_MEMBER_EX( float, InvDeltaSeconds, EShaderPrecisionModifier::Half )
+	UNIFORM_MEMBER_EX( float, DeltaSeconds, EShaderPrecisionModifier::Half )
 	UNIFORM_MEMBER_EX( FVector2D, PivotOffset, EShaderPrecisionModifier::Half )
 	UNIFORM_MEMBER(int, PositionDataOffset)
 	UNIFORM_MEMBER(int, VelocityDataOffset)
@@ -41,12 +42,16 @@ BEGIN_UNIFORM_BUFFER_STRUCT( FNiagaraSpriteUniformParameters, NIAGARAVERTEXFACTO
 	UNIFORM_MEMBER(int, SubimageDataOffset)
 	UNIFORM_MEMBER(int, ColorDataOffset)
 	UNIFORM_MEMBER(int, MaterialParamDataOffset)
+	UNIFORM_MEMBER(int, MaterialParam1DataOffset)
+	UNIFORM_MEMBER(int, MaterialParam2DataOffset)
+	UNIFORM_MEMBER(int, MaterialParam3DataOffset)
 	UNIFORM_MEMBER(int, FacingOffset)
 	UNIFORM_MEMBER(int, AlignmentOffset)
 	UNIFORM_MEMBER(int, SubImageBlendMode)
 	UNIFORM_MEMBER(int, CameraOffsetOffset)
 	UNIFORM_MEMBER(int, UVScaleOffset)
 	UNIFORM_MEMBER(int, ParticleRandomOffset)
+	UNIFORM_MEMBER(FVector4, DefaultPos)
 	END_UNIFORM_BUFFER_STRUCT(FNiagaraSpriteUniformParameters)
 
 typedef TUniformBufferRef<FNiagaraSpriteUniformParameters> FNiagaraSpriteUniformBufferRef;
@@ -67,7 +72,10 @@ public:
 		NumCutoutVerticesPerFrame(0),
 		CutoutGeometrySRV(nullptr),
 		AlignmentMode(0),
-		FacingMode(0)
+		FacingMode(0),
+		FloatDataOffset(0),
+		FloatDataStride(0),
+		SortedIndicesOffset(0)
 	{}
 
 	FNiagaraSpriteVertexFactory()
@@ -76,7 +84,10 @@ public:
 		NumCutoutVerticesPerFrame(0),
 		CutoutGeometrySRV(nullptr),
 		AlignmentMode(0),
-		FacingMode(0)
+		FacingMode(0),
+		FloatDataOffset(0),
+		FloatDataStride(0),
+		SortedIndicesOffset(0)
 	{}
 
 	// FRenderResource interface.
@@ -93,12 +104,7 @@ public:
 	 * Can be overridden by FVertexFactory subclasses to modify their compile environment just before compilation occurs.
 	 */
 	static void ModifyCompilationEnvironment(EShaderPlatform Platform, const FMaterial* Material, FShaderCompilerEnvironment& OutEnvironment);
-
-	/**
-	 * Set the source vertex buffer that contains particle instance data.
-	 */
-	void SetInstanceBuffer(const FVertexBuffer* InInstanceBuffer, uint32 StreamOffset, uint32 Stride, bool bInstanced);
-
+	
 	void SetTexCoordBuffer(const FVertexBuffer* InTexCoordBuffer);
 
 	inline void SetNumVertsInInstanceBuffer(int32 InNumVertsInInstanceBuffer)
@@ -131,26 +137,42 @@ public:
 	inline int32 GetNumCutoutVerticesPerFrame() const { return NumCutoutVerticesPerFrame; }
 	inline FShaderResourceViewRHIParamRef GetCutoutGeometrySRV() const { return CutoutGeometrySRV; }
 
-	void SetParticleData(const FNiagaraDataSet *InDataSet)
+	void SetParticleData(const FShaderResourceViewRHIRef& InParticleDataFloatSRV, uint32 InFloatDataOffset, uint32 InFloatDataStride)
+	{
+		ParticleDataFloatSRV = InParticleDataFloatSRV;
+		FloatDataOffset = InFloatDataOffset;
+		FloatDataStride = InFloatDataStride;
+	}
+
+	void SetSortedIndices(const FShaderResourceViewRHIRef& InSortedIndicesSRV, uint32 InSortedIndicesOffset)
+	{
+		SortedIndicesSRV = InSortedIndicesSRV;
+		SortedIndicesOffset = InSortedIndicesOffset;
+	}
+
+	FORCEINLINE FShaderResourceViewRHIRef GetParticleDataFloatSRV()
 	{
-		DataSet = InDataSet;
+		return ParticleDataFloatSRV;
 	}
 
-	inline FShaderResourceViewRHIParamRef GetFloatDataSRV() const 
-	{ 
-		check(!IsInGameThread());
-		return DataSet->GetRenderDataFloatSRV();
+	FORCEINLINE int32 GetFloatDataOffset()
+	{
+		return FloatDataOffset;
 	}
-	inline FShaderResourceViewRHIParamRef GetIntDataSRV() const 
-	{ 
-		check(!IsInGameThread());
-		return DataSet->GetRenderDataInt32SRV();
+
+	FORCEINLINE int32 GetFloatDataStride()
+	{
+		return FloatDataStride;
 	}
 
-	uint32 GetComponentBufferSize() 
-	{ 
-		check(!IsInGameThread());
-		return DataSet->CurrDataRender().GetFloatStride() / sizeof(float);
+	FORCEINLINE FShaderResourceViewRHIRef GetSortedIndicesSRV()
+	{
+		return SortedIndicesSRV;
+	}
+
+	FORCEINLINE int32 GetSortedIndicesOffset()
+	{
+		return SortedIndicesOffset;
 	}
 
 	void SetFacingMode(uint32 InMode)
@@ -173,7 +195,6 @@ public:
 		return AlignmentMode;
 	}
 
-
 	/**
 	 * Construct shader parameters for this type of vertex factory.
 	 */
@@ -192,7 +213,14 @@ private:
 
 	int32 NumCutoutVerticesPerFrame;
 	FShaderResourceViewRHIParamRef CutoutGeometrySRV;
-	const FNiagaraDataSet *DataSet;
 	uint32 AlignmentMode;
 	uint32 FacingMode;
+	
+	
+	FShaderResourceViewRHIRef ParticleDataFloatSRV;
+	uint32 FloatDataOffset;
+	uint32 FloatDataStride;
+
+	FShaderResourceViewRHIRef SortedIndicesSRV;
+	uint32 SortedIndicesOffset;
 };
diff --git a/Engine/Plugins/FX/Niagara/Source/NiagaraVertexFactories/Public/NiagaraVertexFactory.h b/Engine/Plugins/FX/Niagara/Source/NiagaraVertexFactories/Public/NiagaraVertexFactory.h
--- a/Engine/Plugins/FX/Niagara/Source/NiagaraVertexFactories/Public/NiagaraVertexFactory.h
+++ b/Engine/Plugins/FX/Niagara/Source/NiagaraVertexFactories/Public/NiagaraVertexFactory.h
@@ -47,15 +47,7 @@ public:
 	{
 		bNeedsDeclaration = false;
 	}
-
-	NIAGARAVERTEXFACTORIES_API static void Init()
-	{
-		if (DummyBuffer.Buffer == nullptr)
-		{
-			DummyBuffer.Initialize(4, 1, PF_R32_UINT);
-		}
-	}
-
+	
 	static void ModifyCompilationEnvironment(EShaderPlatform Platform, const FMaterial* Material, FShaderCompilerEnvironment& OutEnvironment)
 	{
 		FVertexFactory::ModifyCompilationEnvironment(Platform, Material, OutEnvironment);
@@ -100,8 +92,6 @@ public:
 		return true;
 	}
 
-protected:
-	static FRWBuffer DummyBuffer;
 private:
 
 	/** Last state where we set this. We only need to setup these once per frame, so detemine same frame by number, time, and view family. */
diff --git a/Engine/Shaders/Private/Common.ush b/Engine/Shaders/Private/Common.ush
--- a/Engine/Shaders/Private/Common.ush
+++ b/Engine/Shaders/Private/Common.ush
@@ -150,7 +150,7 @@ MaterialFloat4 TextureCubeSampleLevel(TextureCube Tex, SamplerState Sampler, flo
 }
 MaterialFloat TextureCubeSampleDepthLevel(TextureCube TexDepth, SamplerState Sampler, float3 UV, MaterialFloat Mip)
 {
-return TexDepth.SampleLevel(Sampler, UV, Mip);
+return TexDepth.SampleLevel(Sampler, UV, Mip).r;
 }
 MaterialFloat4 TextureCubeSampleBias(TextureCube Tex, SamplerState Sampler, float3 UV, MaterialFloat MipBias)
 {
diff --git a/Engine/Shaders/Private/MaterialTemplate.ush b/Engine/Shaders/Private/MaterialTemplate.ush
--- a/Engine/Shaders/Private/MaterialTemplate.ush
+++ b/Engine/Shaders/Private/MaterialTemplate.ush
@@ -88,8 +88,22 @@ struct FMaterialParticleParameters
 	float4 TranslatedWorldPositionAndSize;
 	/** Macro UV scale and bias. */
 	half4 MacroUV;
-	/** Dynamic parameter used by particle systems. */
+
+	/** Dynamic parameters used by particle systems. */
 	half4 DynamicParameter;
+
+#if NUM_DYNAMIC_PARAMETERS > 1
+	half4 DynamicParameter1;
+#endif
+
+#if NUM_DYNAMIC_PARAMETERS > 2
+	half4 DynamicParameter2;
+#endif
+
+#if NUM_DYNAMIC_PARAMETERS > 3
+	half4 DynamicParameter3;
+#endif
+
 	/** mesh particle orientation */
 	float4x4 LocalToWorld;
 
@@ -104,13 +118,40 @@ struct FMaterialParticleParameters
 	float2 Size;
 };
 
-float4 GetDynamicParameter(FMaterialParticleParameters Parameters, float4 Default)
+float4 GetDynamicParameter(FMaterialParticleParameters Parameters, float4 Default, int ParameterIndex=0)
 {
-	#if (PARTICLE_FACTORY || NIAGARA_PARTICLE_FACTORY)
+#if (NIAGARA_PARTICLE_FACTORY)
+	if(ParameterIndex == 0)
+	{
+		return Parameters.DynamicParameter;
+	}
+#if NUM_DYNAMIC_PARAMETERS > 1
+	else if(ParameterIndex == 1)
+	{
+		return Parameters.DynamicParameter1;
+	}
+#endif
+#if NUM_DYNAMIC_PARAMETERS > 2
+	else if(ParameterIndex == 2)
+	{
+		return Parameters.DynamicParameter2;
+	}
+#endif	
+#if NUM_DYNAMIC_PARAMETERS > 3
+	else if(ParameterIndex == 3)
+	{
+		return Parameters.DynamicParameter3;
+	}
+#endif
+	else
+	{
+		return Default;
+	}
+#elif (PARTICLE_FACTORY)
 	return Parameters.DynamicParameter;
-	#else
+#else
 	return Default;
-	#endif
+#endif
 }
 
 /** 
diff --git a/Engine/Shaders/Private/NiagaraEmitterInstanceShader.usf b/Engine/Shaders/Private/NiagaraEmitterInstanceShader.usf
--- a/Engine/Shaders/Private/NiagaraEmitterInstanceShader.usf
+++ b/Engine/Shaders/Private/NiagaraEmitterInstanceShader.usf
@@ -3,25 +3,39 @@
 /*=============================================================================
 NiagaraSimulationShader.usf: 
 =============================================================================*/
+#pragma warning(disable:4008)
 
 #if GPU_SIMULATION
 	#include "/Engine/Public/Platform.ush"
+	#include "Common.ush"
+#else
+	const static float PI = 3.1415926535897932f;
 #endif
+	const static float TWO_PI = 3.1415926535897932f*2.0f;
 
 #include "Definitions.usf"
 
 
-//TODO: Add common utility code for things such as reading data sources, reading from and writing to events.
-
+/* -----------------------------------------------------------------
+ * GPU simulation utility functions
+ * -----------------------------------------------------------------
+ */
 #if GPU_SIMULATION
 
 	#include "SceneTexturesCommon.ush"
 
 	static uint3 GDispatchThreadId;
+	static uint GCurrentPhase;
 	uint EmitterTickCounter;
 
 	#define THREADGROUP_SIZE 64
 
+	float4 Modulo(float4 x, float4 y){ return fmod(x,y); }
+	float3 Modulo(float3 x, float3 y){ return fmod(x,y); }
+	float2 Modulo(float2 x, float2 y){ return fmod(x,y); }
+	float Modulo(float x, float y){ return fmod(x,y); }
+
+
 	// 4D random number generator inspired by PCGs (permuted congruential generator)
 	// COPY FROM RANDOM.USH, which we can't easily include here due to a sizable dependency chain
 	uint4 NiagaraRand4DPCG32(int4 p)
@@ -113,6 +127,11 @@ NiagaraSimulationShader.usf:
 void EnterStatScope(int ScopeIdx);
 void ExitStatScope();
 
+
+/* -----------------------------------------------------------------
+ * VM simulation function declarations
+ * -----------------------------------------------------------------
+ */
 #if VM_SIMULATION
 	float noise(float x);
 	float noise(float2 x);
@@ -132,6 +151,9 @@ void ExitStatScope();
 	//TODO: replace with proper buffer reading capability and use standard hlsl.
 	int AcquireIndex(int DataSetID, bool DoAcquire);
 
+	void AcquireID(int DataSetID, out int IDIndex, out int IDAcquireTag);
+	void UpdateID(int DataSetID, int IDIndex, int InstanceIndex);
+
 	float InputDataFloat(int DataSetIndex, int RegisterIdx);  //DataSetIndex is 0 for main dataset
 	int InputDataInt(int DataSetIndex, int RegisterIdx);
 	bool InputDataBool(int DataSetIndex, int RegisterIdx);
@@ -145,15 +167,21 @@ void ExitStatScope();
 	void OutputDataBool(int DataSetIndex, int RegisterIndex, int InstanceIndex, bool Value);
 #endif
 
-#if GPU_SIMULATION
-	int ExecIndex()
-	{
-		return GDispatchThreadId.x;
-	}
 
-	uint StartInstancePhase0;
-	uint StartInstancePhase1;
-	uint TotalNumInstances;
+
+
+/* -----------------------------------------------------------------
+ * GPU simulation code
+ * -----------------------------------------------------------------
+ */
+#if GPU_SIMULATION
+	static uint GSpawnStartInstance;
+		
+	/* Parameters for spawning and simulations
+	 */
+	uint UpdateStartInstance;
+	uint SpawnedInstances;
+	uint NumIndicesPerInstance;
 
 	uint ComponentBufferSizeRead;
 	uint ComponentBufferSizeWrite;
@@ -161,15 +189,44 @@ void ExitStatScope();
 	int NumParticlesPerEvent;
 	uint CopyInstancesBeforeStart;
 
+
+	/* Buffers for particle data and DrawIndirect calls
+	 */
 	Buffer<float> InputFloat;
 	Buffer<int> InputInt;
 	RWBuffer<int> RWOutputInt;
 	RWBuffer<float> RWOutputFloat;
-	RWBuffer<int> RWDataSetIndices;
+	RWBuffer<uint> RWDataSetIndices;
+	Buffer<uint> ReadDataSetIndices;
+
+
+	//TODO: IDs for GPU
+//	void AcquireID(int DataSetID, out int IDIndex, out int IDAcquireTag);
+// 	void UpdateID(int DataSetID, int IDIndex, int InstanceIndex);
+
+
 
+	/* Returns the current instance index relative to the operation (spawn/update)
+	 */
+	int ExecIndex()
+	{
+		if (GCurrentPhase == 0)
+		{
+			return UpdateStartInstance + GDispatchThreadId.x - GSpawnStartInstance;
+		}
+		else 
+		{
+			return UpdateStartInstance + GDispatchThreadId.x;
+		}
+	}
+
+
+	/* Acquire an output index - the default index is the scratch instance; one additional instance is allocated 
+	 *	at the end of the buffer, so no branching on -1 is necessary during OutputData operations
+	 */
 	int AcquireIndex(int DataSetID, bool bDoAcquire)
 	{
-		int PrevIdx = -1;
+		int PrevIdx = GSpawnStartInstance + SpawnedInstances;
 		if(bDoAcquire == true)
 		{
 			InterlockedAdd(RWDataSetIndices[DataSetID*5+1], 1, PrevIdx);
@@ -177,6 +234,12 @@ void ExitStatScope();
 		return PrevIdx;
 	}
 
+
+
+	/* ---------------------------------------------------------------------
+	 * InputData operations 
+	 * ---------------------------------------------------------------------
+	 */
 	float InputDataFloat(int DataSetIndex, int RegisterIdx, int InstanceIdx)
 	{
 		return InputFloat[RegisterIdx*ComponentBufferSizeRead + InstanceIdx];
@@ -186,15 +249,30 @@ void ExitStatScope();
 	{
 		return InputInt[RegisterIdx*ComponentBufferSizeRead + InstanceIdx];
 	}
+	
+	bool InputDataBool(int DataSetIndex, int RegisterIdx, int InstanceIdx)
+	{
+		return InputInt[RegisterIdx*ComponentBufferSizeRead + InstanceIdx] == -1;
+	}
+
 
+	/* ---------------------------------------------------------------------
+	 * OutputData operations 
+	 * ---------------------------------------------------------------------
+	 */
 	void OutputDataFloat(int DataSetIndex, int RegisterIndex, int InstanceIndex, float Value)
 	{
-		RWOutputFloat[RegisterIndex*ComponentBufferSizeWrite + InstanceIndex + StartInstancePhase0] = Value;
+		RWOutputFloat[RegisterIndex*ComponentBufferSizeWrite + InstanceIndex + UpdateStartInstance] = Value;
 	}
 
 	void OutputDataInt(int DataSetIndex, int RegisterIndex, int InstanceIndex, int Value)
 	{
-		RWOutputInt[RegisterIndex*ComponentBufferSizeWrite + InstanceIndex + StartInstancePhase0] = Value;
+		RWOutputInt[RegisterIndex*ComponentBufferSizeWrite + InstanceIndex + UpdateStartInstance] = Value;
+	}
+
+	void OutputDataBool(int DataSetIndex, int RegisterIndex, int InstanceIndex, bool Value)
+	{
+		RWOutputInt[RegisterIndex*ComponentBufferSizeWrite + InstanceIndex + UpdateStartInstance] = Value ? -1 : 0;
 	}
 
 	void EnterStatScope(int ID)	{}
@@ -202,12 +280,15 @@ void ExitStatScope();
 #endif
 
 //Include the simulation shader code generated by the node graph.
-#include "/Engine/Generated/NiagaraEmitterInstance.usf"
+#include "/Engine/Generated/NiagaraEmitterInstance.ush"
 
 
-//The VectorVM shader needs only the generated code but a compute shader requires some wrapping.
 #if GPU_SIMULATION
 
+/* 
+ *	CS wrapper for our generated code; calls spawn and update functions on the corresponding instances in the buffer
+ */
+
 [numthreads(THREADGROUP_SIZE, 1, 1)]
 void SimulateMainComputeCS(
 	uint3 GroupId : SV_GroupID,
@@ -219,13 +300,15 @@ void SimulateMainComputeCS(
 	/*
 	if(CopyInstancesBeforeStart == 1)
 	{
-		StartInstancePhase0 = 0;
+		UpdateStartInstance = 0;
 	}
 	*/
 
-	uint InstanceID = StartInstancePhase0 + DispatchThreadId.x;
+	uint InstanceID = UpdateStartInstance + DispatchThreadId.x;
+	GSpawnStartInstance = ReadDataSetIndices[1];				// needed by ExecIndex()
+	int MaxInstances = GSpawnStartInstance + SpawnedInstances;
 
-	if(InstanceID < StartInstancePhase0+TotalNumInstances )
+	if(InstanceID < UpdateStartInstance + MaxInstances )
 	{
 		// calculate event index for the current instance; this is primarily for spawn events, where we do a single read from a secondary DataSet for a subset of particles
 		//  for all-particle events, this will always end up being 0 because NumParticlesPerEvent = total particles in the set
@@ -233,7 +316,7 @@ void SimulateMainComputeCS(
 
 		// before the simulation start - either update or copy
 		//
-		if(InstanceID < StartInstancePhase1)
+		if(InstanceID < GSpawnStartInstance)
 		{
 			// copy instances before simulation start inst (e.g. for event spawns)
 			/*
@@ -244,6 +327,7 @@ void SimulateMainComputeCS(
 			else
 			*/
 			{
+				GCurrentPhase = 1;
 				SimulateMain(InstanceID, EventIndex, 1);
 			}
 
@@ -251,16 +335,17 @@ void SimulateMainComputeCS(
 
 		// after simulation start - spawn phase
 		//
-		else if(InstanceID >= StartInstancePhase1)
+		else if(InstanceID >= GSpawnStartInstance)
 		{
-				SimulateMain(InstanceID, EventIndex, 0);		// SimulateMain spawn phase
+			GCurrentPhase = 0;
+			SimulateMain(InstanceID, EventIndex, 0);		// SimulateMain spawn phase
 		}
 	}
 
 	// fill in the rest of RWDataSetIndices for the DrawIndirect call for the main data set; number of instances is incremented in AcquireIndex
 	if(InstanceID == 0)
 	{
-		RWDataSetIndices[0] = 6;		// IndexCountPerInstance; TODO: need to pass this in from outside based on renderer setup
+		RWDataSetIndices[0] = NumIndicesPerInstance;		// IndexCountPerInstance; passed in from the renderer - TODO: support multiple renderers with multiple indirect draw calls
 		RWDataSetIndices[2] = 0;		// StartIndexLocation
 		RWDataSetIndices[3] = 0;		// BaseVertexLocation
 		RWDataSetIndices[4] = 0;		// StartInstanceLocation
diff --git a/Engine/Shaders/Private/NiagaraMeshVertexFactory.ush b/Engine/Shaders/Private/NiagaraMeshVertexFactory.ush
--- a/Engine/Shaders/Private/NiagaraMeshVertexFactory.ush
+++ b/Engine/Shaders/Private/NiagaraMeshVertexFactory.ush
@@ -27,6 +27,9 @@ uint MeshFacingMode;
 #define MESH_FACING_CAMERA_POSITION 2
 #define MESH_FACING_CAMERA_PLANE 3
 	
+/** Sorted index buffer for particles if we need to sort their draw order (translucency etc). */
+Buffer<int> SortedIndices;
+uint SortedIndicesOffset;
 
 struct FVertexFactoryInput
 {
@@ -72,9 +75,18 @@ struct FVertexFactoryInterpolantsVSToPS
 	float4	ParticleColor : COLOR1;
 #endif
 
-#if USE_DYNAMIC_PARAMETERS
+#if NUM_DYNAMIC_PARAMETERS > 0
 	float4 DynamicParameter : COLOR2;
 #endif
+#if NUM_DYNAMIC_PARAMETERS > 1
+	float4 DynamicParameter1 : COLOR3;
+#endif
+#if NUM_DYNAMIC_PARAMETERS > 2
+	float4 DynamicParameter2 : COLOR4;
+#endif
+#if NUM_DYNAMIC_PARAMETERS > 3
+	float4 DynamicParameter3 : COLOR5;
+#endif
 	
 #if USE_PARTICLE_NODES
 	/** Particle position in camera-centered translated world space */
@@ -111,10 +123,20 @@ struct FVertexFactoryIntermediates
 	float4 SubUVCoords;
 	float SubUVLerp;
 #endif
-#if USE_DYNAMIC_PARAMETERS
-	/** Optional dynamic parameter for the particle. */
+
+	/** Optional dynamic parameters for the particle. */
+#if NUM_DYNAMIC_PARAMETERS > 0
 	float4 DynamicParameter;
-#endif	//USE_DYNAMIC_PARAMETERS
+#endif
+#if NUM_DYNAMIC_PARAMETERS > 1
+	float4 DynamicParameter1;
+#endif
+#if NUM_DYNAMIC_PARAMETERS > 2
+	float4 DynamicParameter2;
+#endif
+#if NUM_DYNAMIC_PARAMETERS > 3
+	float4 DynamicParameter3;
+#endif
 
 	/** The position of the particle, XYZ: direction. */
 	float4 ParticlePosition;
@@ -205,11 +227,18 @@ FMaterialPixelParameters GetMaterialPixelParameters(FVertexFactoryInterpolantsVS
 	Result.Particle.Color = Interpolants.ParticleColor;
 #endif
 
-#if USE_DYNAMIC_PARAMETERS
+#if NUM_DYNAMIC_PARAMETERS > 0
 	Result.Particle.DynamicParameter = Interpolants.DynamicParameter;
-#else
-	Result.Particle.DynamicParameter = float4(1,1,1,1);
-#endif //USE_DYNAMIC_PARAMETERS
+#endif
+#if NUM_DYNAMIC_PARAMETERS > 1
+	Result.Particle.DynamicParameter1 = Interpolants.DynamicParameter1;
+#endif
+#if NUM_DYNAMIC_PARAMETERS > 2
+	Result.Particle.DynamicParameter2 = Interpolants.DynamicParameter2;
+#endif
+#if NUM_DYNAMIC_PARAMETERS > 3
+	Result.Particle.DynamicParameter3 = Interpolants.DynamicParameter3;
+#endif
 
 #if USE_PARTICLE_NODES
 	Result.Particle.TranslatedWorldPositionAndSize.xyz = Interpolants.ParticleTranslatedWorldPosition.xyz;
@@ -250,9 +279,18 @@ FMaterialVertexParameters GetMaterialVertexParameters(FVertexFactoryInput Input,
 	Result.Particle.RelativeTime = Intermediates.RelativeTime;
 	Result.Particle.Color = Intermediates.ParticleColor;
 
-#if USE_DYNAMIC_PARAMETERS
+#if NUM_DYNAMIC_PARAMETERS > 0
 	Result.Particle.DynamicParameter = Intermediates.DynamicParameter;
-#endif	//USE_DYNAMIC_PARAMETERS
+#endif
+#if NUM_DYNAMIC_PARAMETERS > 1
+	Result.Particle.DynamicParameter1 = Intermediates.DynamicParameter1;
+#endif
+#if NUM_DYNAMIC_PARAMETERS > 2
+	Result.Particle.DynamicParameter2 = Intermediates.DynamicParameter2;
+#endif
+#if NUM_DYNAMIC_PARAMETERS > 3
+	Result.Particle.DynamicParameter3 = Intermediates.DynamicParameter3;
+#endif
 
 	float4x4 Transform = GetParticleTransform(Intermediates);
 	Result.InstanceLocalToWorld = Transform;
@@ -308,7 +346,7 @@ float4 CalcPrevWorldPosition(FVertexFactoryInput Input, FVertexFactoryIntermedia
 #endif
 	// if we don't have a full previous transform, we can at least calculate the position to get some motion blur
 	{
-		return CalcWorldPosition(Input, Intermediates)- float4(Intermediates.ParticleVelocity.xyz*Intermediates.ParticleVelocity.w*MeshParticleVF.DeltaSeconds, 0.0);
+		return CalcWorldPosition(Input, Intermediates)- float4(Intermediates.ParticleVelocity.xyz*NiagaraMeshVF.DeltaSeconds, 0.0);
 	}
 }
 
@@ -399,21 +437,50 @@ float4x4 EulerToMat4x4(float Pitch, float Yaw, float Roll)
 FVertexFactoryIntermediates GetVertexFactoryIntermediates(FVertexFactoryInput Input)
 {
 	FVertexFactoryIntermediates Intermediates = (FVertexFactoryIntermediates)0;
+	
+	if(SortedIndicesOffset != 0xFFFFFFFF)
+	{
+		Input.InstanceId = SortedIndices[SortedIndicesOffset + Input.InstanceId];
+	}
 
-	Intermediates.ParticlePosition = float4(GetVec3(NiagaraMeshVF.PositionDataOffset, Input.InstanceId), 1);
-	Intermediates.ParticleColor = GetVec4(NiagaraMeshVF.ColorDataOffset, Input.InstanceId);
-	Intermediates.ParticleVelocity = float4(GetVec3(NiagaraMeshVF.VelocityDataOffset, Input.InstanceId), 1);
+	if (NiagaraMeshVF.PositionDataOffset == -1)
+	{
+		Intermediates.ParticlePosition = NiagaraMeshVF.DefaultPos;
+	}
+	else
+	{
+		Intermediates.ParticlePosition = float4(GetVec3(NiagaraMeshVF.PositionDataOffset, Input.InstanceId), 1);
+	}
 
+	if (NiagaraMeshVF.ColorDataOffset == -1)
+	{
+		Intermediates.ParticleColor = float4(1.0f, 1.0f, 1.0f, 1.0f);
+	}
+	else
+	{
+		Intermediates.ParticleColor = GetVec4(NiagaraMeshVF.ColorDataOffset, Input.InstanceId);
+	}
 
-#if USE_DYNAMIC_PARAMETERS
-	if (NiagaraMeshVF.MaterialParamDataOffset != -1)
+	if (NiagaraMeshVF.VelocityDataOffset == -1)
 	{
-		Intermediates.DynamicParameter = GetVec4(NiagaraMeshVF.MaterialParamDataOffset, Input.InstanceId);
+		Intermediates.ParticleVelocity = float4(0.0f, 0.0f, 0.0f, 1);
 	}
 	else
 	{
-		Intermediates.DynamicParameter = float4(1,1,1,1);
+		Intermediates.ParticleVelocity = float4(GetVec3(NiagaraMeshVF.VelocityDataOffset, Input.InstanceId), 1);
 	}
+
+#if NUM_DYNAMIC_PARAMETERS > 0
+	Intermediates.DynamicParameter  = NiagaraMeshVF.MaterialParamDataOffset != -1 ? GetVec4(NiagaraMeshVF.MaterialParamDataOffset, Input.InstanceId) : float4(1,1,1,1);
+#endif
+#if NUM_DYNAMIC_PARAMETERS > 1
+	Intermediates.DynamicParameter1  = NiagaraMeshVF.MaterialParam1DataOffset != -1 ? GetVec4(NiagaraMeshVF.MaterialParam1DataOffset, Input.InstanceId) : float4(1,1,1,1);
+#endif
+#if NUM_DYNAMIC_PARAMETERS > 2
+	Intermediates.DynamicParameter2  = NiagaraMeshVF.MaterialParam2DataOffset != -1 ? GetVec4(NiagaraMeshVF.MaterialParam2DataOffset, Input.InstanceId) : float4(1,1,1,1);
+#endif
+#if NUM_DYNAMIC_PARAMETERS > 3
+	Intermediates.DynamicParameter3  = NiagaraMeshVF.MaterialParam3DataOffset != -1 ? GetVec4(NiagaraMeshVF.MaterialParam3DataOffset, Input.InstanceId) : float4(1,1,1,1);
 #endif
 
 	float4 Transform1 = float4(1.0f, 0.0f, 0.0f, Intermediates.ParticlePosition.x);
@@ -422,16 +489,13 @@ FVertexFactoryIntermediates GetVertexFactoryIntermediates(FVertexFactoryInput In
 
 	if(NiagaraMeshVF.TransformDataOffset != -1)
 	{
-		float4 AxisAngle = GetVec4(NiagaraMeshVF.TransformDataOffset, Input.InstanceId);
+		float4 Quat = GetVec4(NiagaraMeshVF.TransformDataOffset, Input.InstanceId);
 		
 		//Transform axis to handle local space.
 		//AxisAngle.xyz = mul(float4(AxisAngle.xyz, 0), NiagaraMeshVF.LocalToWorldInverseTransposed);
 
-		float Pitch = 0.0f;
-		float Roll = 0.0f;
-		float Yaw = 0.0f;
-		AxisAngleToPitchYawRoll(AxisAngle, Pitch, Yaw, Roll);
-		float4x4 EulerToMat = EulerToMat4x4(Pitch, Yaw, Roll);
+		
+		float4x4 EulerToMat = QuatTo4x4(normalize(Quat));
 		Transform1.xyz = EulerToMat[0].xyz;
 		Transform2.xyz = EulerToMat[1].xyz;
 		Transform3.xyz = EulerToMat[2].xyz;
@@ -457,7 +521,7 @@ FVertexFactoryIntermediates GetVertexFactoryIntermediates(FVertexFactoryInput In
 	if (MeshFacingMode == MESH_FACING_VELOCITY)
 	{
 		float4 ToAlignmentQuat = FindQuatBetween(float3(1.0f, 0.0f, 0.0f), normalize(Intermediates.ParticleVelocity.xyz));
-		float4x4 ToAlignmentMat = QuatTo4x4(ToAlignmentQuat);
+		float4x4 ToAlignmentMat = QuatTo4x4( normalize(ToAlignmentQuat) );
 		Intermediates.LocalToWorld =  mul(LocalToWorldMat, mul(ToAlignmentMat, ScaleMatrix));
 	}
 	else if (MeshFacingMode == MESH_FACING_CAMERA_POSITION)
@@ -466,14 +530,14 @@ FVertexFactoryIntermediates GetVertexFactoryIntermediates(FVertexFactoryInput In
 		float3 CameraVector = ResolvedView.WorldCameraOrigin.xyz - WorldPosition.xyz;
 		float3 CameraDirection = normalize(CameraVector);
 		float4 ToAlignmentQuat = FindQuatBetween(float3(1.0f, 0.0f, 0.0f), CameraDirection);
-		float4x4 ToAlignmentMat = QuatTo4x4(ToAlignmentQuat);
+		float4x4 ToAlignmentMat = QuatTo4x4( normalize(ToAlignmentQuat) );
 		Intermediates.LocalToWorld = mul(LocalToWorldMat, mul(ToAlignmentMat, ScaleMatrix));
 	}
 	else if (MeshFacingMode == MESH_FACING_CAMERA_PLANE)
 	{
 		float3 CameraFacingVector = -ResolvedView.ViewForward;
 		float4 ToAlignmentQuat = FindQuatBetween(float3(1.0f, 0.0f, 0.0f), CameraFacingVector);
-		float4x4 ToAlignmentMat = QuatTo4x4(ToAlignmentQuat);
+		float4x4 ToAlignmentMat = QuatTo4x4( normalize(ToAlignmentQuat) );
 		Intermediates.LocalToWorld =  mul(LocalToWorldMat, mul(ToAlignmentMat, ScaleMatrix));
 	}
 	else
@@ -616,9 +680,19 @@ FVertexFactoryInterpolantsVSToPS VertexFactoryGetInterpolantsVSToPS(FVertexFacto
 	Interpolants.ParticleColor = Intermediates.ParticleColor;
 #endif
 
-#if USE_DYNAMIC_PARAMETERS
+#if NUM_DYNAMIC_PARAMETERS > 0
 	Interpolants.DynamicParameter = Intermediates.DynamicParameter;
-#endif	//USE_DYNAMIC_PARAMETERS
+#endif
+#if NUM_DYNAMIC_PARAMETERS > 1
+	Interpolants.DynamicParameter1 = Intermediates.DynamicParameter1;
+#endif
+#if NUM_DYNAMIC_PARAMETERS > 2
+	Interpolants.DynamicParameter2 = Intermediates.DynamicParameter2;
+#endif
+#if NUM_DYNAMIC_PARAMETERS > 3
+	Interpolants.DynamicParameter3 = Intermediates.DynamicParameter3;
+#endif
+
 
 #if USE_PARTICLE_NODES
 	Interpolants.ParticleTranslatedWorldPosition = Intermediates.ParticleTranslatedWorldPosition;
@@ -719,9 +793,19 @@ float4 VertexFactoryGetPreviousWorldPosition(FVertexFactoryInput Input, FVertexF
 			TESSELLATION_INTERPOLATE_MEMBER(InterpolantsVSToPS.TexCoords[i]);
 		}
 #endif
-#if USE_DYNAMIC_PARAMETERS
+
+#if NUM_DYNAMIC_PARAMETERS > 0
 		TESSELLATION_INTERPOLATE_MEMBER(InterpolantsVSToPS.DynamicParameter);
-#endif	//USE_DYNAMIC_PARAMETERS
+#endif
+#if NUM_DYNAMIC_PARAMETERS > 1
+		TESSELLATION_INTERPOLATE_MEMBER(InterpolantsVSToPS.DynamicParameter1);
+#endif
+#if NUM_DYNAMIC_PARAMETERS > 2
+		TESSELLATION_INTERPOLATE_MEMBER(InterpolantsVSToPS.DynamicParameter2);
+#endif
+#if NUM_DYNAMIC_PARAMETERS > 3
+		TESSELLATION_INTERPOLATE_MEMBER(InterpolantsVSToPS.DynamicParameter3);
+#endif
 
 		return O;
 	}
diff --git a/Engine/Shaders/Private/NiagaraRibbonVertexFactory.ush b/Engine/Shaders/Private/NiagaraRibbonVertexFactory.ush
--- a/Engine/Shaders/Private/NiagaraRibbonVertexFactory.ush
+++ b/Engine/Shaders/Private/NiagaraRibbonVertexFactory.ush
@@ -26,21 +26,38 @@ float3 TransformVector(float3 InVector)
 /*
 float3 GetNiagaraParticlePosition(uint InstanceID)
 {
+	if (NiagaraRibbonVF.PositionDataOffset == -1)
+	{
+		return TransformPosition(float3(0.0f, 0.0f, 0.0f));
+	}
 	return TransformPosition(GetVec3(NiagaraRibbonVF.PositionDataOffset, InstanceID));
 }
 float GetNiagaraParticleRibbonWidth(uint InstanceID)
 {
+	if (NiagaraRibbonVF.WidthDataOffset == -1)
+	{
+		return 1.0f;
+	}
 	return GetFloat(NiagaraRibbonVF.WidthDataOffset, InstanceID);
 }
 float GetNiagaraParticleRibbonTwist(uint InstanceID)
 {
+	if (NiagaraRibbonVF.TwistDataOffset == -1)
+	{
+		return 0.0f;
+	}
 	return GetFloat(NiagaraRibbonVF.TwistDataOffset, InstanceID);
 }
 float4 GetNiagaraParticleColor(uint InstanceID)
 {
+	if (NiagaraRibbonVF.ColorDataOffset == -1)
+	{
+		return float4(1.0f, 1.0f, 1.0f, 1.0f);
+	}
 	return GetVec4(NiagaraRibbonVF.ColorDataOffset, InstanceID);
 }
-#if USE_DYNAMIC_PARAMETERS
+
+#if NUM_DYNAMIC_PARAMETERS > 0
 float4 GetNiagaraParticleDynamicParameters(uint InstanceID)
 {
 	if (NiagaraRibbonVF.MaterialParamDataOffset == -1)
@@ -50,11 +67,41 @@ float4 GetNiagaraParticleDynamicParameters(uint InstanceID)
 	return GetVec4(NiagaraRibbonVF.MaterialParamDataOffset, InstanceID);
 }
 #endif
+#if NUM_DYNAMIC_PARAMETERS > 1
+float4 GetNiagaraParticleDynamicParameters1(uint InstanceID)
+{
+	if (NiagaraRibbonVF.MaterialParam1DataOffset == -1)
+	{
+		return float4(1.0f, 1.0f, 1.0f, 1.0f);
+	}
+	return GetVec4(NiagaraRibbonVF.MaterialParam1DataOffset, InstanceID);
+}
+#endif
+#if NUM_DYNAMIC_PARAMETERS > 2
+float4 GetNiagaraParticleDynamicParameters2(uint InstanceID)
+{
+	if (NiagaraRibbonVF.MaterialParam2DataOffset == -1)
+	{
+		return float4(1.0f, 1.0f, 1.0f, 1.0f);
+	}
+	return GetVec4(NiagaraRibbonVF.MaterialParam2DataOffset, InstanceID);
+}
+#endif
+#if NUM_DYNAMIC_PARAMETERS > 3
+float4 GetNiagaraParticleDynamicParameters3(uint InstanceID)
+{
+	if (NiagaraRibbonVF.MaterialParam3DataOffset == -1)
+	{
+		return float4(1.0f, 1.0f, 1.0f, 1.0f);
+	}
+	return GetVec4(NiagaraRibbonVF.MaterialParam3DataOffset, InstanceID);
+}
+#endif
 
 
 float3 GetNiagaraParticleFacingVector(uint InstanceID)
 {
-	float3 Val = float3(0.0f, 0.0f, 0.0f);
+	float3 Val = float3(0.0f, 0.0f, 1.0f);
     if(NiagaraRibbonVF.FacingOffset != -1)
 	{
 		Val = GetVec3(NiagaraRibbonVF.FacingOffset, InstanceID);
@@ -77,9 +124,18 @@ struct FVertexFactoryInput
 	float Twist						: ATTRIBUTE5;
 	float3 CustomFacingVector		: ATTRIBUTE6;
 
-#if USE_DYNAMIC_PARAMETERS
+#if NUM_DYNAMIC_PARAMETERS > 0
 	float4 DynamicParameter0		: ATTRIBUTE7;
-#endif	//USE_DYNAMIC_PARAMETERS
+#endif
+#if NUM_DYNAMIC_PARAMETERS > 1
+	float4 DynamicParameter1		: ATTRIBUTE8;
+#endif
+#if NUM_DYNAMIC_PARAMETERS > 2
+	float4 DynamicParameter2		: ATTRIBUTE9;
+#endif
+#if NUM_DYNAMIC_PARAMETERS > 3
+	float4 DynamicParameter3		: ATTRIBUTE10;
+#endif
 
 /** Optional instance ID for vertex layered rendering */
 #if FEATURE_LEVEL >= FEATURE_LEVEL_SM4
@@ -90,6 +146,9 @@ struct FVertexFactoryInput
 #if FEATURE_LEVEL >= FEATURE_LEVEL_SM4 && ONEPASS_POINTLIGHT_SHADOW && USING_VERTEX_SHADER_LAYER
 	uint InstanceId	: SV_InstanceID;
 #endif
+
+	// soon
+	//float NormalizedAge			: ATTRIBUTE8;
 };
 
 struct FVertexFactoryInterpolantsVSToPS
@@ -105,9 +164,18 @@ struct FVertexFactoryInterpolantsVSToPS
 	float4  Color				: TEXCOORD3;
 #endif
 
-#if USE_DYNAMIC_PARAMETERS
+#if NUM_DYNAMIC_PARAMETERS > 0
 	float4 DynamicParameter0	: TEXCOORD4;
-#endif	//USE_DYNAMIC_PARAMETERS
+#endif
+#if NUM_DYNAMIC_PARAMETERS > 1
+	float4 DynamicParameter1	: TEXCOORD5;
+#endif
+#if NUM_DYNAMIC_PARAMETERS > 2
+	float4 DynamicParameter2	: TEXCOORD6;
+#endif
+#if NUM_DYNAMIC_PARAMETERS > 3
+	float4 DynamicParameter3	: TEXCOORD7;
+#endif
 
 #if INSTANCED_STEREO
 	nointerpolation uint PackedEyeIndex : PACKED_EYE_INDEX;
@@ -119,7 +187,7 @@ struct FVertexFactoryIntermediates
 	/** The color of the sprite. */
 	float4 Color;
 	float RibbonWidth;
-	float4 Right, Up;
+	float3 Right, Up;
 	float3 CameraDirection;
 };
 
@@ -160,10 +228,19 @@ FMaterialPixelParameters GetMaterialPixelParameters(FVertexFactoryInterpolantsVS
 #endif
 
 	Result.TwoSidedSign = 1;
-
-#if USE_DYNAMIC_PARAMETERS
+	
+#if NUM_DYNAMIC_PARAMETERS > 0
 	Result.Particle.DynamicParameter = Interpolants.DynamicParameter0;
-#endif	//USE_DYNAMIC_PARAMETERS
+#endif
+#if NUM_DYNAMIC_PARAMETERS > 1
+	Result.Particle.DynamicParameter1 = Interpolants.DynamicParameter1;
+#endif
+#if NUM_DYNAMIC_PARAMETERS > 2
+	Result.Particle.DynamicParameter2 = Interpolants.DynamicParameter2;
+#endif
+#if NUM_DYNAMIC_PARAMETERS > 3
+	Result.Particle.DynamicParameter3 = Interpolants.DynamicParameter3;
+#endif
 
 	return Result;
 }
@@ -203,16 +280,15 @@ float3 SafeNormalize(float3 V)
 	return V / sqrt(max(dot(V,V),0.01));
 }
 
-void GetTangents(FVertexFactoryInput Input,out float4 Right,out float4 Up, out float3 CameraDirection)
+void GetTangents(FVertexFactoryInput Input,out float3 Right,out float3 Up, out float3 CameraDirection)
 {
 	float4	Position			= TransformLocalToTranslatedWorld(Input.Position.xyz);
-			//OldPosition			= TransformLocalToTranslatedWorld(Input.OldPosition.xyz);
 
-	CameraDirection		= SafeNormalize(ResolvedView.WorldCameraOrigin - Input.Position.xyz );		// TODO: handle local space
-	float3 ParticleDirection	= Input.Direction; //SafeNormalize(Position.xyz - OldPosition.xyz);
+	CameraDirection				= SafeNormalize(ResolvedView.WorldCameraOrigin - Input.Position.xyz );		// TODO: handle local space
+	float3 ParticleDirection	= Input.Direction;
 
 	float4	Right_Square		= BeamTrailVF.CameraRight,
-			Up_Square			= BeamTrailVF.CameraUp;			
+			Up_Square			= BeamTrailVF.CameraUp;		
 			
 	float4	Right_Rotated		= (-1.0 * cos(Input.Twist) * Up_Square) + (sin(Input.Twist) * Right_Square),
 			Up_Rotated			= (       sin(Input.Twist) * Up_Square) + (cos(Input.Twist) * Right_Square);
@@ -220,20 +296,36 @@ void GetTangents(FVertexFactoryInput Input,out float4 Right,out float4 Up, out f
 	float4	Right_Velocity		= float4( SafeNormalize( cross( CameraDirection, ParticleDirection	) ), 0.0 ),
 			Up_Velocity			= float4( ParticleDirection, 0.0 );
 
-	float3 NormDir = SafeNormalize(Input.Direction.xyz);
-	float3 RightDir;
-	if(NiagaraRibbonVF.UseCustomFacing)
+	float3 NormDir = SafeNormalize(ParticleDirection);
+	float3 RightDir, Facing;
+
+
+	/*if(NiagaraRibbonVF.UseCustomFacing)
+	{
+		Facing = Input.CustomFacingVector;
+	}
+	else*/
+	{
+		Facing = CameraDirection.xyz;
+	}
+
+	if(Input.Twist != 0.0f)
 	{
-		RightDir = SafeNormalize(cross(NormDir, Input.CustomFacingVector));
+		float cosa = cos(Input.Twist);
+		Facing = Facing * cosa + cross(Facing, Input.Direction) * sin(Input.Twist) + Input.Direction * dot(Input.Direction, Facing) * (1-cosa);
+	}
+
+	if( abs(dot(NormDir, Facing)) < 0.999f)
+	{
+		RightDir = SafeNormalize(cross(NormDir, Facing));
 	}
 	else
 	{
-		//float Blend = abs(dot(NormDir, SafeNormalize(CameraDirection.xyz)));
-		RightDir = SafeNormalize(cross(NormDir, CameraDirection.xyz));
+		RightDir = SafeNormalize(cross(NormDir, float3(0.0f, 0.0f, -1.0f)));
 	}
-	
-	Right = float4(RightDir, 0.0f);
-	Up = float4(NormDir, 0.0f);
+
+	Right = RightDir;
+	Up = NormDir;
 	return;
 	/*
 	//	enum EParticleScreenAlignment
@@ -252,36 +344,23 @@ float4 CalcWorldPosition(FVertexFactoryInput Input, FVertexFactoryIntermediates
 	// expand and align
 	int Sgn = (((int)Input.VertexID & 0x1) * 2) - 1;
 	float3 NormDir = SafeNormalize(Input.Direction.xyz);
-	float3 RightDir;
-	if(NiagaraRibbonVF.UseCustomFacing)
-	{
-		RightDir = SafeNormalize(cross(NormDir, Input.CustomFacingVector));
-	}
-	else
-	{
-		//float Blend = abs(dot(NormDir, SafeNormalize(Intermediates.CameraDirection.xyz)));
-		RightDir = SafeNormalize(cross(NormDir, Intermediates.CameraDirection.xyz));
-	}
-	float3 UpDir = cross(RightDir, NormDir);
-
-	//float4	Right_Rotated		= (-1.0 * cos(Input.Twist) * UpDir) + (sin(Input.Twist) * RightDir),
-	//		Up_Rotated			= (       sin(Input.Twist) * UpDir) + (cos(Input.Twist) * RightDir);
-
+	float3 RightDir = Intermediates.Right;
+	float3 UpDir = Intermediates.Up;
 	float3 Pos = Input.Position.xyz + RightDir * Sgn * Input.Size.xxx;
 	return TransformLocalToTranslatedWorld(Pos);
 }
 
 /** derive basis vectors */
-float3x3 CalcTangentBasis(FVertexFactoryInput Input)
+float3x3 CalcTangentBasis(FVertexFactoryInput Input, FVertexFactoryIntermediates Intermediates)
 {
 	float4	Right,
 			Up;
 	float3	CameraDir;
-	GetTangents(Input,Right,Up, CameraDir);
+	//GetTangents(Input,Right,Up, CameraDir);
 	return float3x3(
-			Right.xyz,
-			Up.xyz,
-			-normalize(cross(Right.xyz,Up.xyz))
+			Intermediates.Right.xyz,
+			Intermediates.Up.xyz,
+			-normalize(cross(Intermediates.Right.xyz, Intermediates.Up.xyz))
 			);
 }
 
@@ -320,11 +399,20 @@ FVertexFactoryInterpolantsVSToPS VertexFactoryGetInterpolantsVSToPS(FVertexFacto
 	Interpolants.Color		= Intermediates.Color;
 #endif
 
-#if USE_DYNAMIC_PARAMETERS
+#if NUM_DYNAMIC_PARAMETERS > 0
 	Interpolants.DynamicParameter0 = Input.DynamicParameter0;
-#endif	//USE_DYNAMIC_PARAMETERS
+#endif
+#if NUM_DYNAMIC_PARAMETERS > 1
+	Interpolants.DynamicParameter1 = Input.DynamicParameter1;
+#endif
+#if NUM_DYNAMIC_PARAMETERS > 2
+	Interpolants.DynamicParameter2 = Input.DynamicParameter2;
+#endif
+#if NUM_DYNAMIC_PARAMETERS > 3
+	Interpolants.DynamicParameter3 = Input.DynamicParameter3;
+#endif
 
-	float3x3 TangentToLocal = CalcTangentBasis(Input);
+	float3x3 TangentToLocal = CalcTangentBasis(Input, Intermediates);
 	float3x3 TangentToWorld = mul(TangentToLocal, GetLocalToWorld3x3());
 
 	Interpolants.TangentToWorld0 = TangentToWorld[0];
@@ -350,7 +438,7 @@ float4 VertexFactoryGetPreviousWorldPosition(FVertexFactoryInput Input, FVertexF
 */
 float3x3 VertexFactoryGetTangentToLocal( FVertexFactoryInput Input, FVertexFactoryIntermediates Intermediates )
 {
-	return CalcTangentBasis(Input);
+	return CalcTangentBasis(Input, Intermediates);
 }
 
 #if INSTANCED_STEREO
@@ -363,4 +451,4 @@ uint VertexFactoryGetEyeIndex(uint InstanceId)
 float4 VertexFactoryGetTranslatedPrimitiveVolumeBounds(FVertexFactoryInterpolantsVSToPS Interpolants)
 {
 	return 0;
-}
+}
\ No newline at end of file
