Rewrote the Niagara error/warning logging system to not use the Blueprint FCompilerResultsLog, which potentially had some thread safety issues as used in Niagara.
#tests PC Autotests pass, successful PC cook of EngineTests

diff --git a/Engine/Plugins/FX/Niagara/Source/NiagaraEditor/Private/NiagaraCompiler.cpp b/Engine/Plugins/FX/Niagara/Source/NiagaraEditor/Private/NiagaraCompiler.cpp
--- a/Engine/Plugins/FX/Niagara/Source/NiagaraEditor/Private/NiagaraCompiler.cpp
+++ b/Engine/Plugins/FX/Niagara/Source/NiagaraEditor/Private/NiagaraCompiler.cpp
@@ -511,42 +511,42 @@ TSharedPtr<FNiagaraVMExecutableData> FNiagaraEditorModule::CompileScript(const F
 		VMCompilationTime = Results.CompileTime;
 	}
 	
-	TArray<TSharedRef<FTokenizedMessage>> Messages;
-	if (TranslateResults.MessageLog)
+	TArray<FNiagaraCompileEvent> Messages;
+	if (TranslateResults.CompileEvents.Num() > 0)
 	{
-		Messages.Append(TranslateResults.MessageLog->Messages);
+		Messages.Append(TranslateResults.CompileEvents);
 	}
-	if (Results.MessageLog)
+	if (Results.CompileEvents.Num() > 0)
 	{
-		Messages.Append(Results.MessageLog->Messages);
+		Messages.Append(Results.CompileEvents);
 	}
 
 	FString OutGraphLevelErrorMessages;
-	for (TSharedRef<FTokenizedMessage> Message : Messages)
+	for (const FNiagaraCompileEvent& Message : Messages)
 	{
-		if (Message->GetSeverity() == EMessageSeverity::Info)
+		if (Message.Type == FNiagaraCompileEventType::Log)
 		{
 		#if defined(NIAGARA_SCRIPT_COMPILE_LOGGING_MEDIUM)
-			UE_LOG(LogNiagaraCompiler, Log, TEXT("%s"), *Message->ToText().ToString());
+			UE_LOG(LogNiagaraCompiler, Log, TEXT("%s"), *Message.Message);
 		#endif
 		}
-		else if (Message->GetSeverity() == EMessageSeverity::Warning || Message->GetSeverity() == EMessageSeverity::PerformanceWarning)
+		else if (Message.Type == FNiagaraCompileEventType::Warning )
 		{
 		#if defined(NIAGARA_SCRIPT_COMPILE_LOGGING_MEDIUM)
-			UE_LOG(LogNiagaraCompiler, Warning, TEXT("%s"), *Message->ToText().ToString());
+			UE_LOG(LogNiagaraCompiler, Warning, TEXT("%s"), *Message.Message);
 		#endif
 		}
-		else if (Message->GetSeverity() == EMessageSeverity::Error || Message->GetSeverity() == EMessageSeverity::CriticalError)
+		else if (Message.Type == FNiagaraCompileEventType::Error)
 		{
 		#if defined(NIAGARA_SCRIPT_COMPILE_LOGGING_MEDIUM)
-			UE_LOG(LogNiagaraCompiler, Error, TEXT("%s"), *Message->ToText().ToString());
+			UE_LOG(LogNiagaraCompiler, Error, TEXT("%s"), *Message.Message);
 		#endif
 			// Write the error messages to the string as well so that they can be echoed up the chain.
 			if (OutGraphLevelErrorMessages.Len() > 0)
 			{
 				OutGraphLevelErrorMessages += "\n";
 			}
-			OutGraphLevelErrorMessages += Message->ToText().ToString();
+			OutGraphLevelErrorMessages += Message.Message;
 		}
 	}
 	
@@ -614,7 +614,7 @@ ENiagaraScriptCompileStatus FNiagaraCompileResults::CompileResultsToSummary(cons
 	ENiagaraScriptCompileStatus SummaryStatus = ENiagaraScriptCompileStatus::NCS_Unknown;
 	if (CompileResults != nullptr)
 	{
-		if (CompileResults->MessageLog->NumErrors > 0)
+		if (CompileResults->NumErrors > 0)
 		{
 			SummaryStatus = ENiagaraScriptCompileStatus::NCS_Error;
 		}
@@ -622,7 +622,7 @@ ENiagaraScriptCompileStatus FNiagaraCompileResults::CompileResultsToSummary(cons
 		{
 			if (CompileResults->bVMSucceeded)
 			{
-				if (CompileResults->MessageLog->NumWarnings)
+				if (CompileResults->NumWarnings)
 				{
 					SummaryStatus = ENiagaraScriptCompileStatus::NCS_UpToDateWithWarnings;
 				}
@@ -634,7 +634,7 @@ ENiagaraScriptCompileStatus FNiagaraCompileResults::CompileResultsToSummary(cons
 
 			if (CompileResults->bComputeSucceeded)
 			{
-				if (CompileResults->MessageLog->NumWarnings)
+				if (CompileResults->NumWarnings)
 				{
 					SummaryStatus = ENiagaraScriptCompileStatus::NCS_ComputeUpToDateWithWarnings;
 				}
@@ -816,10 +816,8 @@ FNiagaraCompileResults FHlslNiagaraCompiler::CompileScript(const FNiagaraCompile
 
 
 FHlslNiagaraCompiler::FHlslNiagaraCompiler()
-	: CompileResults(&MessageLog)
+	: CompileResults()
 {
-	// Make the message log silent so we're not spamming the blueprint log.
-	MessageLog.bSilentMode = true;
 }
 
 
@@ -827,13 +825,15 @@ FHlslNiagaraCompiler::FHlslNiagaraCompiler()
 void FHlslNiagaraCompiler::Error(FText ErrorText)
 {
 	FString ErrorString = FString::Printf(TEXT("%s"), *ErrorText.ToString());
-	MessageLog.Error(*ErrorString);
+	CompileResults.CompileEvents.Add(FNiagaraCompileEvent(FNiagaraCompileEventType::Error, ErrorString));
+	CompileResults.NumErrors++;
 }
 
 void FHlslNiagaraCompiler::Warning(FText WarningText)
 {
 	FString WarnString = FString::Printf(TEXT("%s"), *WarningText.ToString());
-	MessageLog.Warning(*WarnString);
+	CompileResults.CompileEvents.Add(FNiagaraCompileEvent(FNiagaraCompileEventType::Warning, WarnString));
+	CompileResults.NumWarnings++;
 }
 
 //////////////////////////////////////////////////////////////////////////
diff --git a/Engine/Plugins/FX/Niagara/Source/NiagaraEditor/Private/NiagaraCompiler.h b/Engine/Plugins/FX/Niagara/Source/NiagaraEditor/Private/NiagaraCompiler.h
--- a/Engine/Plugins/FX/Niagara/Source/NiagaraEditor/Private/NiagaraCompiler.h
+++ b/Engine/Plugins/FX/Niagara/Source/NiagaraEditor/Private/NiagaraCompiler.h
@@ -19,9 +19,6 @@ struct FNiagaraTranslatorOutput;
 class NIAGARAEDITOR_API FHlslNiagaraCompiler : public INiagaraCompiler
 {
 protected:	
-	/** Message log. Automatically handles marking the NodeGraph with errors. */
-	FCompilerResultsLog MessageLog;
-
 	/** Captures information about a script compile. */
 	FNiagaraCompileResults CompileResults;
 
diff --git a/Engine/Plugins/FX/Niagara/Source/NiagaraEditor/Private/NiagaraHlslTranslator.cpp b/Engine/Plugins/FX/Niagara/Source/NiagaraEditor/Private/NiagaraHlslTranslator.cpp
--- a/Engine/Plugins/FX/Niagara/Source/NiagaraEditor/Private/NiagaraHlslTranslator.cpp
+++ b/Engine/Plugins/FX/Niagara/Source/NiagaraEditor/Private/NiagaraHlslTranslator.cpp
@@ -146,7 +146,7 @@ ENiagaraScriptCompileStatus FNiagaraTranslateResults::TranslateResultsToSummary(
 	ENiagaraScriptCompileStatus SummaryStatus = ENiagaraScriptCompileStatus::NCS_Unknown;
 	if (TranslateResults != nullptr)
 	{
-		if (TranslateResults->MessageLog->NumErrors > 0)
+		if (TranslateResults->NumErrors > 0)
 		{
 			SummaryStatus = ENiagaraScriptCompileStatus::NCS_Error;
 		}
@@ -154,7 +154,7 @@ ENiagaraScriptCompileStatus FNiagaraTranslateResults::TranslateResultsToSummary(
 		{
 			if (TranslateResults->bHLSLGenSucceeded)
 			{
-				if (TranslateResults->MessageLog->NumWarnings)
+				if (TranslateResults->NumWarnings)
 				{
 					SummaryStatus = ENiagaraScriptCompileStatus::NCS_UpToDateWithWarnings;
 				}
@@ -355,13 +355,11 @@ void FHlslNiagaraTranslator::GenerateFunctionSignature(ENiagaraScriptUsage Scrip
 
 FHlslNiagaraTranslator::FHlslNiagaraTranslator()
 	: Schema(nullptr)
-	, TranslateResults(&MessageLog)
+	, TranslateResults()
 	, CurrentBodyChunkMode(ENiagaraCodeChunkMode::Body)
 	, ActiveStageIdx(-1)
 	, bInitializedDefaults(false)
 {
-	// Make the message log silent so we're not spamming the blueprint log.
-	MessageLog.bSilentMode = true;
 }
 
 
@@ -659,9 +657,6 @@ const FNiagaraTranslateResults &FHlslNiagaraTranslator::Translate(const FNiagara
 	TranslateResults.bHLSLGenSucceeded = false;
 	TranslateResults.OutputHLSL = "";
 
-	//Should we roll our own message/error log and put it in a window somewhere?
-	MessageLog.SetSourcePath(CompileOptions.GetPathName());
-
 	UNiagaraGraph* SourceGraph = CompileData->NodeGraphDeepCopy;
 
 	if (!SourceGraph)
@@ -901,7 +896,7 @@ const FNiagaraTranslateResults &FHlslNiagaraTranslator::Translate(const FNiagara
 
 	ActiveHistoryForFunctionCalls.EndTranslation(GetUniqueEmitterName());
 
-	TranslateResults.bHLSLGenSucceeded = MessageLog.NumErrors == 0;
+	TranslateResults.bHLSLGenSucceeded = TranslateResults.NumErrors == 0;
 
 	//If we're compiling a function then we have all we need already, we don't want to actually generate shader/vm code.
 	if (FunctionCtx())
@@ -2097,9 +2092,9 @@ FString FHlslNiagaraTranslator::GetUniqueSymbolName(FName BaseName)
 	return RetString;
 }
 
-void FHlslNiagaraTranslator::EnterFunction(const FString& Name, FNiagaraFunctionSignature& Signature, TArray<int32>& Inputs)
+void FHlslNiagaraTranslator::EnterFunction(const FString& Name, FNiagaraFunctionSignature& Signature, TArray<int32>& Inputs, const FGuid& InGuid)
 {
-	FunctionContextStack.Emplace(Name, Signature, Inputs);
+	FunctionContextStack.Emplace(Name, Signature, Inputs, InGuid);
 	//May need some more heavy and scoped symbol tracking?
 
 	//Add new scope for pin reuse.
@@ -2178,6 +2173,17 @@ FString FHlslNiagaraTranslator::GetCallstack()
 	return Callstack;
 }
 
+TArray<FGuid> FHlslNiagaraTranslator::GetCallstackGuids()
+{
+	TArray<FGuid> Callstack;
+	for (FFunctionContext& Ctx : FunctionContextStack)
+	{
+		Callstack.Add(Ctx.Id);
+	}
+
+	return Callstack;
+}
+
 FString FHlslNiagaraTranslator::GeneratedConstantString(FVector4 Constant)
 {
 	TArray<FStringFormatArg> Args;
@@ -3146,7 +3152,7 @@ void FHlslNiagaraTranslator::Emitter(class UNiagaraNodeEmitter* EmitterNode, TAr
 	}
 
 	// We act like a function call here as the semantics are identical.
-	RegisterFunctionCall(ScriptUsage, Name, FullName, Source, Signature, false, FString(), Inputs, CallInputs, CallOutputs, Signature);
+	RegisterFunctionCall(ScriptUsage, Name, FullName, EmitterNode->NodeGuid, Source, Signature, false, FString(), Inputs, CallInputs, CallOutputs, Signature);
 	GenerateFunctionCall(Signature, Inputs, Outputs);
 
 	// Clear out the parameter map writes to emitter module parameters as they should not be shared across emitters.
@@ -3814,7 +3820,7 @@ void FHlslNiagaraTranslator::FunctionCall(UNiagaraNodeFunctionCall* FunctionNode
 		HandleCustomHlslNode(CustomFunctionHlsl, ScriptUsage, Name, FullName, bCustomHlsl, CustomHlsl, Signature, Inputs);
 	}
 	
-	RegisterFunctionCall(ScriptUsage, Name, FullName,  Source, Signature, bCustomHlsl, CustomHlsl, Inputs, CallInputs, CallOutputs, OutputSignature);
+	RegisterFunctionCall(ScriptUsage, Name, FullName,  FunctionNode->NodeGuid, Source, Signature, bCustomHlsl, CustomHlsl, Inputs, CallInputs, CallOutputs, OutputSignature);
 	GenerateFunctionCall(OutputSignature, Inputs, Outputs);
 
 	if (bCustomHlsl)
@@ -4034,7 +4040,7 @@ void FHlslNiagaraTranslator::HandleCustomHlslNode(UNiagaraNodeCustomHlsl* Custom
 	OutCustomHlsl = TEXT("\n") + OutCustomHlsl + TEXT("\n");
 }
 
-void FHlslNiagaraTranslator::RegisterFunctionCall(ENiagaraScriptUsage ScriptUsage, const FString& InName, const FString& InFullName, UNiagaraScriptSource* Source,  FNiagaraFunctionSignature& InSignature, bool bIsCustomHlsl, const FString& InCustomHlsl, TArray<int32>& Inputs, const TArray<UEdGraphPin*>& CallInputs, const TArray<UEdGraphPin*>& CallOutputs,
+void FHlslNiagaraTranslator::RegisterFunctionCall(ENiagaraScriptUsage ScriptUsage, const FString& InName, const FString& InFullName, const FGuid& CallNodeId, UNiagaraScriptSource* Source,  FNiagaraFunctionSignature& InSignature, bool bIsCustomHlsl, const FString& InCustomHlsl, TArray<int32>& Inputs, const TArray<UEdGraphPin*>& CallInputs, const TArray<UEdGraphPin*>& CallOutputs,
 	FNiagaraFunctionSignature& OutSignature)
 {
 	SCOPE_CYCLE_COUNTER(STAT_NiagaraEditor_Module_NiagaraHLSLTranslator_RegisterFunctionCall);
@@ -4086,7 +4092,7 @@ void FHlslNiagaraTranslator::RegisterFunctionCall(ENiagaraScriptUsage ScriptUsag
 			SCOPE_CYCLE_COUNTER(STAT_NiagaraEditor_Module_NiagaraHLSLTranslator_FuncBody);
 
 			//We've not compiled this function yet so compile it now.
-			EnterFunction(InName, OutSignature, Inputs);
+			EnterFunction(InName, OutSignature, Inputs, CallNodeId);
 
 			UNiagaraNodeOutput* FuncOutput = SourceGraph->FindOutputNode(ScriptUsage);
 			check(FuncOutput);
@@ -4229,7 +4235,7 @@ void FHlslNiagaraTranslator::RegisterFunctionCall(ENiagaraScriptUsage ScriptUsag
 			if (!FuncBody)
 			{
 				//We've not compiled this function yet so compile it now.
-				EnterFunction(InName, OutSignature, Inputs);
+				EnterFunction(InName, OutSignature, Inputs, CallNodeId);
 
 				FString FunctionDefStr = InCustomHlsl;
 				// We don't support an empty function definition when calling a real function.
@@ -4837,13 +4843,15 @@ int32 FHlslNiagaraTranslator::CompileOutputPin(const UEdGraphPin* InPin)
 void FHlslNiagaraTranslator::Error(FText ErrorText, UNiagaraNode* Node, UEdGraphPin* Pin)
 {
 	FString ErrorString = FString::Printf(TEXT("Node: @@ - Pin: @@ - %s - Callstack: %s"), *ErrorText.ToString(), *GetCallstack());
-	MessageLog.Error(*ErrorString, Node, Pin);
+	TranslateResults.CompileEvents.Add(FNiagaraCompileEvent(FNiagaraCompileEventType::Error, ErrorString, Node ? Node->NodeGuid : FGuid(), Pin ? Pin->PersistentGuid : FGuid(), GetCallstackGuids()));
+	TranslateResults.NumErrors++;
 }
 
 void FHlslNiagaraTranslator::Warning(FText WarningText, UNiagaraNode* Node, UEdGraphPin* Pin)
 {
 	FString WarnString = FString::Printf(TEXT("Node: @@ - Pin: @@ - %s - Callstack: %s"), *WarningText.ToString(), *GetCallstack());
-	MessageLog.Warning(*WarnString, Node, Pin);
+	TranslateResults.CompileEvents.Add(FNiagaraCompileEvent(FNiagaraCompileEventType::Warning, WarnString, Node ? Node->NodeGuid : FGuid(), Pin ? Pin->PersistentGuid : FGuid(), GetCallstackGuids()));
+	TranslateResults.NumWarnings++;
 }
 
 bool FHlslNiagaraTranslator::GetFunctionParameter(const FNiagaraVariable& Parameter, int32& OutParam)const
diff --git a/Engine/Plugins/FX/Niagara/Source/NiagaraEditor/Private/NiagaraHlslTranslator.h b/Engine/Plugins/FX/Niagara/Source/NiagaraEditor/Private/NiagaraHlslTranslator.h
--- a/Engine/Plugins/FX/Niagara/Source/NiagaraEditor/Private/NiagaraHlslTranslator.h
+++ b/Engine/Plugins/FX/Niagara/Source/NiagaraEditor/Private/NiagaraHlslTranslator.h
@@ -65,13 +65,14 @@ struct FNiagaraTranslateResults
 	bool bHLSLGenSucceeded;
 
 	/** A results log with messages, warnings, and errors which occurred during the compile. */
-	FCompilerResultsLog* MessageLog;
+	TArray<FNiagaraCompileEvent> CompileEvents;
+	uint32 NumErrors;
+	uint32 NumWarnings;
 
 	/** A string representation of the compilation output. */
 	FString OutputHLSL;
 
-	FNiagaraTranslateResults(FCompilerResultsLog* InMessageLog = nullptr)
-		: MessageLog(InMessageLog)
+	FNiagaraTranslateResults() : NumErrors(0), NumWarnings(0)
 	{
 	}
 
@@ -331,7 +332,7 @@ protected:
 	/** Map of function graphs we've seen before and already pre-processed. */
 	TMap<const UNiagaraGraph*, UNiagaraGraph*> PreprocessedFunctions;
 
-	void RegisterFunctionCall(ENiagaraScriptUsage ScriptUsage, const FString& InName, const FString& InFullName, UNiagaraScriptSource* Source, FNiagaraFunctionSignature& InSignature, bool bIsCustomHlsl, const FString& InCustomHlsl, TArray<int32>& Inputs, const TArray<UEdGraphPin*>& CallInputs, const TArray<UEdGraphPin*>& CallOutputs,
+	void RegisterFunctionCall(ENiagaraScriptUsage ScriptUsage, const FString& InName, const FString& InFullName, const FGuid& CallNodeId, UNiagaraScriptSource* Source, FNiagaraFunctionSignature& InSignature, bool bIsCustomHlsl, const FString& InCustomHlsl, TArray<int32>& Inputs, const TArray<UEdGraphPin*>& CallInputs, const TArray<UEdGraphPin*>& CallOutputs,
 		FNiagaraFunctionSignature& OutSignature);
 	void GenerateFunctionCall(FNiagaraFunctionSignature& FunctionSignature, TArray<int32>& Inputs, TArray<int32>& Outputs);
 	FString GetFunctionSignature(const FNiagaraFunctionSignature& Sig);
@@ -351,17 +352,20 @@ protected:
 		FString Name;
 		FNiagaraFunctionSignature& Signature;
 		TArray<int32>& Inputs;
-		FFunctionContext(const FString& InName, FNiagaraFunctionSignature& InSig, TArray<int32>& InInputs)
+		FGuid Id;
+		FFunctionContext(const FString& InName, FNiagaraFunctionSignature& InSig, TArray<int32>& InInputs, const FGuid& InId)
 			: Name(InName)
 			, Signature(InSig)
 			, Inputs(InInputs)
+			, Id(InId)
 		{}
 	};
 	TArray<FFunctionContext> FunctionContextStack;
 	const FFunctionContext* FunctionCtx()const { return FunctionContextStack.Num() > 0 ? &FunctionContextStack.Last() : nullptr; }
-	void EnterFunction(const FString& Name, FNiagaraFunctionSignature& Signature, TArray<int32>& Inputs);
+	void EnterFunction(const FString& Name, FNiagaraFunctionSignature& Signature, TArray<int32>& Inputs, const FGuid& InGuid);
 	void ExitFunction();
 	FString GetCallstack();
+	TArray<FGuid> GetCallstackGuids();
 
 	void EnterStatsScope(FNiagaraStatScope StatScope);
 	void ExitStatsScope();
diff --git a/Engine/Plugins/FX/Niagara/Source/NiagaraEditor/Public/INiagaraCompiler.h b/Engine/Plugins/FX/Niagara/Source/NiagaraEditor/Public/INiagaraCompiler.h
--- a/Engine/Plugins/FX/Niagara/Source/NiagaraEditor/Public/INiagaraCompiler.h
+++ b/Engine/Plugins/FX/Niagara/Source/NiagaraEditor/Public/INiagaraCompiler.h
@@ -18,6 +18,33 @@ struct FNiagaraVMExecutableData;
 class FNiagaraCompileRequestData;
 class FNiagaraCompileOptions;
 
+/** Defines the compile event types for translation/compilation.*/
+enum FNiagaraCompileEventType
+{
+	Log = 0,
+	Warning = 1,
+	Error = 2
+};
+
+/** Records necessary information to give UI cues for errors/logs/warnings during compile.*/
+struct FNiagaraCompileEvent
+{
+public:
+	FNiagaraCompileEvent(FNiagaraCompileEventType InType, const FString& InMessage, FGuid InNodeGuid = FGuid(), FGuid InPinGuid = FGuid(), const TArray<FGuid>& InCallstackGuids = TArray<FGuid>())
+	: Type(InType), Message(InMessage), NodeGuid(InNodeGuid), PinGuid(InPinGuid), StackGuids(InCallstackGuids){}
+
+	/** Whether or not this is an error, warning, or info*/
+	FNiagaraCompileEventType Type;
+	/* The message itself*/
+	FString Message;
+	/** The node guid that generated the compile event*/
+	FGuid NodeGuid;
+	/** The pin persistent id that generated the compile event*/
+	FGuid PinGuid;
+	/** The compile stack frame of node id's*/
+	TArray<FGuid> StackGuids;
+};
+
 
 /** Defines information about the results of a Niagara script compile. */
 struct FNiagaraCompileResults
@@ -27,17 +54,19 @@ struct FNiagaraCompileResults
 
 	/** Whether or not the script compiled successfully for GPU compute */
 	bool bComputeSucceeded;
-
-	/** A results log with messages, warnings, and errors which occurred during the compile. */
-	FCompilerResultsLog* MessageLog;
-
+	
 	/** The actual final compiled data.*/
 	TSharedPtr<FNiagaraVMExecutableData> Data;
 
 	float CompileTime;
 
-	FNiagaraCompileResults(FCompilerResultsLog* InMessageLog = nullptr)
-		: MessageLog(InMessageLog), CompileTime(0.0f)
+	/** Tracking any compilation warnings or errors that occur.*/
+	TArray<FNiagaraCompileEvent> CompileEvents;
+	uint32 NumErrors;
+	uint32 NumWarnings;
+
+	FNiagaraCompileResults()
+		: CompileTime(0.0f), NumErrors(0), NumWarnings(0)
 	{
 	}
 
