RL - The uniformbuffer description is cached in multiple levels for shader compilation. Notable Caches are the ShaderDirectoryFileMapping the MMCP File cache and multiple other smaller file caches along the way. Instead of invalidating those caches I made the Uniform buffer layout completely immutable for all Platform varieties. Differences in lauout are triggered via defines as a progammer would have done it writing the code by hand. THe padding has been unified between GLSL and HLSL padding rules and theoretically work for both of them uniformly.

diff --git a/Engine/Plugins/FX/Niagara/Source/NiagaraShader/Private/NiagaraShader.cpp b/Engine/Plugins/FX/Niagara/Source/NiagaraShader/Private/NiagaraShader.cpp
--- a/Engine/Plugins/FX/Niagara/Source/NiagaraShader/Private/NiagaraShader.cpp
+++ b/Engine/Plugins/FX/Niagara/Source/NiagaraShader/Private/NiagaraShader.cpp
@@ -290,10 +290,10 @@ void FNiagaraShaderType::CacheUniformBufferIncludes(TMap<const TCHAR*, FCachedUn
 void FNiagaraShaderType::AddReferencedUniformBufferIncludes(FShaderCompilerEnvironment& OutEnvironment, FString& OutSourceFilePrefix, EShaderPlatform Platform)
 {
 	// Cache uniform buffer struct declarations referenced by this shader type's files
-	if (!bCachedUniformBufferStructDeclarations[Platform])
+	if (!bCachedUniformBufferStructDeclarations)
 	{
 		CacheUniformBufferIncludes(ReferencedUniformBufferStructsCache, Platform);
-		bCachedUniformBufferStructDeclarations[Platform] = true;
+		bCachedUniformBufferStructDeclarations = true;
 	}
 
 	FString UniformBufferIncludes;
@@ -319,6 +319,12 @@ void FNiagaraShaderType::AddReferencedUniformBufferIncludes(FShaderCompilerEnvir
 
 	FString& GeneratedUniformBuffersInclude = OutEnvironment.IncludeVirtualPathToContentsMap.FindOrAdd("/Engine/Generated/GeneratedUniformBuffers.ush");
 	GeneratedUniformBuffersInclude.Append(UniformBufferIncludes);
+
+	ERHIFeatureLevel::Type MaxFeatureLevel = GetMaxSupportedFeatureLevel(Platform);
+	if (MaxFeatureLevel >= ERHIFeatureLevel::SM4)
+	{
+		OutEnvironment.SetDefine(TEXT("PLATFORM_SUPPORTS_SRV_UB"), TEXT("1"));
+	}
 }
 
 
