RL - The uniformbuffer description is cached in multiple levels for shader compilation. Notable Caches are the ShaderDirectoryFileMapping the MMCP File cache and multiple other smaller file caches along the way. Instead of invalidating those caches I made the Uniform buffer layout completely immutable for all Platform varieties. Differences in lauout are triggered via defines as a progammer would have done it writing the code by hand. THe padding has been unified between GLSL and HLSL padding rules and theoretically work for both of them uniformly.

diff --git a/Engine/Plugins/FX/Niagara/Source/NiagaraShader/Private/NiagaraShader.cpp b/Engine/Plugins/FX/Niagara/Source/NiagaraShader/Private/NiagaraShader.cpp
--- a/Engine/Plugins/FX/Niagara/Source/NiagaraShader/Private/NiagaraShader.cpp
+++ b/Engine/Plugins/FX/Niagara/Source/NiagaraShader/Private/NiagaraShader.cpp
@@ -277,8 +277,8 @@ void FNiagaraShaderType::CacheUniformBufferIncludes(TMap<const TCHAR*, FCachedUn
 		{
 			if (It.Key() == StructIt->GetShaderVariableName())
 			{
-				BufferDeclaration.Declaration[Platform] = MakeShared<FString>();
-				CreateUniformBufferShaderDeclaration(StructIt->GetShaderVariableName(), **StructIt, Platform, *BufferDeclaration.Declaration[Platform].Get());
+				BufferDeclaration.Declaration = MakeShared<FString>();
+				CreateUniformBufferShaderDeclaration(StructIt->GetShaderVariableName(), **StructIt, *BufferDeclaration.Declaration.Get());
 				break;
 			}
 		}
@@ -300,9 +300,9 @@ void FNiagaraShaderType::AddReferencedUniformBufferIncludes(FShaderCompilerEnvir
 
 	for (TMap<const TCHAR*, FCachedUniformBufferDeclaration>::TConstIterator It(ReferencedUniformBufferStructsCache); It; ++It)
 	{
-		check(It.Value().Declaration[Platform].IsValid());
+		check(It.Value().Declaration.IsValid());
 		UniformBufferIncludes += FString::Printf(TEXT("#include \"/Engine/Generated/UniformBuffers/%s.ush\"") LINE_TERMINATOR, It.Key());
-		FString* Declaration = It.Value().Declaration[Platform].Get();
+		FString* Declaration = It.Value().Declaration.Get();
 		check(Declaration);
 		OutEnvironment.IncludeVirtualPathToContentsMap.Add(
 			*FString::Printf(TEXT("/Engine/Generated/UniformBuffers/%s.ush"), It.Key()), *Declaration
