Niagara - Fix a crash caused by delegates not being cleaned up properly in the emitter view model.

diff --git a/Engine/Plugins/FX/Niagara/Source/NiagaraEditor/Private/ViewModels/NiagaraEmitterViewModel.cpp b/Engine/Plugins/FX/Niagara/Source/NiagaraEditor/Private/ViewModels/NiagaraEmitterViewModel.cpp
--- a/Engine/Plugins/FX/Niagara/Source/NiagaraEditor/Private/ViewModels/NiagaraEmitterViewModel.cpp
+++ b/Engine/Plugins/FX/Niagara/Source/NiagaraEditor/Private/ViewModels/NiagaraEmitterViewModel.cpp
@@ -46,6 +46,7 @@ void FNiagaraEmitterViewModel::Cleanup()
 	if (Emitter.IsValid())
 	{
 		Emitter->OnEmitterVMCompiled().RemoveAll(this);
+		Emitter->OnPropertiesChanged().RemoveAll(this);
 	}
 
 	if (SharedScriptViewModel.IsValid())
