Niagara - Prevent an assert when undoing dependency fixes in the stack.

diff --git a/Engine/Plugins/FX/Niagara/Source/NiagaraEditor/Private/ViewModels/Stack/NiagaraStackFunctionInput.cpp b/Engine/Plugins/FX/Niagara/Source/NiagaraEditor/Private/ViewModels/Stack/NiagaraStackFunctionInput.cpp
--- a/Engine/Plugins/FX/Niagara/Source/NiagaraEditor/Private/ViewModels/Stack/NiagaraStackFunctionInput.cpp
+++ b/Engine/Plugins/FX/Niagara/Source/NiagaraEditor/Private/ViewModels/Stack/NiagaraStackFunctionInput.cpp
@@ -207,13 +207,14 @@ void UNiagaraStackFunctionInput::FinalizeInternal()
 	{
 		OwningFunctionCallNode->GetGraph()->RemoveOnGraphChangedHandler(GraphChangedHandle);
 		OwningFunctionCallNode->GetNiagaraGraph()->RemoveOnGraphNeedsRecompileHandler(OnRecompileHandle);
+	}
 
-		if (SourceScript.IsValid())
-		{
-			SourceScript->RapidIterationParameters.RemoveOnChangedHandler(RapidIterationParametersChangedHandle);
-			SourceScript->GetSource()->OnChanged().RemoveAll(this);
-		}
+	if (SourceScript.IsValid())
+	{
+		SourceScript->RapidIterationParameters.RemoveOnChangedHandler(RapidIterationParametersChangedHandle);
+		SourceScript->GetSource()->OnChanged().RemoveAll(this);
 	}
+
 	Super::FinalizeInternal();
 }
 
