Niagara - Prevent missing dependencies from loading all niagara script assets.

diff --git a/Engine/Plugins/FX/Niagara/Source/NiagaraEditor/Private/ViewModels/Stack/NiagaraStackGraphUtilities.cpp b/Engine/Plugins/FX/Niagara/Source/NiagaraEditor/Private/ViewModels/Stack/NiagaraStackGraphUtilities.cpp
--- a/Engine/Plugins/FX/Niagara/Source/NiagaraEditor/Private/ViewModels/Stack/NiagaraStackGraphUtilities.cpp
+++ b/Engine/Plugins/FX/Niagara/Source/NiagaraEditor/Private/ViewModels/Stack/NiagaraStackGraphUtilities.cpp
@@ -1240,17 +1240,11 @@ void FNiagaraStackGraphUtilities::GetScriptAssetsByDependencyProvided(ENiagaraSc
 	{
 		auto TagName = GET_MEMBER_NAME_CHECKED(UNiagaraScript, ProvidedDependencies);
 
-		TArray<UObject::FAssetRegistryTag> Tags;
-		ScriptAsset.GetAsset()->GetAssetRegistryTags(Tags);
-		UObject::FAssetRegistryTag* FoundTag = Tags.FindByPredicate([=](const UObject::FAssetRegistryTag& InTag)
+		FString ProvidedDependenciesString;
+		if(ScriptAsset.GetTagValue(GET_MEMBER_NAME_CHECKED(UNiagaraScript, ProvidedDependencies), ProvidedDependenciesString) && ProvidedDependenciesString.IsEmpty() == false)
 		{
-			return InTag.Name == TagName;
-		});
-		if (FoundTag != nullptr)
-		{
-			// parse it and search through it%
 			TArray<FString> DependencyStrings;
-			FoundTag->Value.ParseIntoArray(DependencyStrings, TEXT(","));
+			ProvidedDependenciesString.ParseIntoArray(DependencyStrings, TEXT(","));
 			for (FString DependencyString: DependencyStrings)
 			{
 				if (FName(*DependencyString) == DependencyName)
