Niagara - Prevent multiple save prompt dialogs on shut down when you have script or emitter editors open with unsaved changes and you decide not to save them.

diff --git a/Engine/Plugins/FX/Niagara/Source/NiagaraEditor/Private/Toolkits/NiagaraScriptToolkit.cpp b/Engine/Plugins/FX/Niagara/Source/NiagaraEditor/Private/Toolkits/NiagaraScriptToolkit.cpp
--- a/Engine/Plugins/FX/Niagara/Source/NiagaraEditor/Private/Toolkits/NiagaraScriptToolkit.cpp
+++ b/Engine/Plugins/FX/Niagara/Source/NiagaraEditor/Private/Toolkits/NiagaraScriptToolkit.cpp
@@ -198,6 +198,8 @@ void FNiagaraScriptToolkit::Initialize( const EToolkitMode::Type Mode, const TSh
 		const FString TabInitializationPayload(TEXT(""));		// NOTE: Payload not currently used for table properties
 		SpawnToolkitTab( NodeGraphTabId, TabInitializationPayload, EToolkitTabSpot::Details );
 	}*/
+
+	bChangesDiscarded = false;
 }
 
 FName FNiagaraScriptToolkit::GetToolkitFName() const
@@ -643,7 +645,7 @@ void FNiagaraScriptToolkit::UpdateOriginalNiagaraScript()
 
 bool FNiagaraScriptToolkit::OnRequestClose()
 {
-	if (IsEditScriptDifferentFromOriginalScript())
+	if (bChangesDiscarded == false && IsEditScriptDifferentFromOriginalScript())
 	{
 		// find out the user wants to do with this dirty NiagaraScript
 		EAppReturnType::Type YesNoCancelReply = FMessageDialog::Open(EAppMsgType::YesNoCancel,
@@ -660,8 +662,8 @@ bool FNiagaraScriptToolkit::OnRequestClose()
 			break;
 
 		case EAppReturnType::No:
-			// exit
-			// do nothing.. bNiagaraScriptDirty = false;
+			// Set the changes discarded to avoid showing the dialog multiple times when request close is called multiple times on shut down.
+			bChangesDiscarded = true;
 			break;
 
 		case EAppReturnType::Cancel:
diff --git a/Engine/Plugins/FX/Niagara/Source/NiagaraEditor/Private/Toolkits/NiagaraScriptToolkit.h b/Engine/Plugins/FX/Niagara/Source/NiagaraEditor/Private/Toolkits/NiagaraScriptToolkit.h
--- a/Engine/Plugins/FX/Niagara/Source/NiagaraEditor/Private/Toolkits/NiagaraScriptToolkit.h
+++ b/Engine/Plugins/FX/Niagara/Source/NiagaraEditor/Private/Toolkits/NiagaraScriptToolkit.h
@@ -125,4 +125,5 @@ private:
 	FDelegateHandle OnEditedScriptGraphChangedHandle;
 
 	bool bEditedScriptHasPendingChanges;
+	bool bChangesDiscarded;
 };
diff --git a/Engine/Plugins/FX/Niagara/Source/NiagaraEditor/Private/Toolkits/NiagaraSystemToolkit.cpp b/Engine/Plugins/FX/Niagara/Source/NiagaraEditor/Private/Toolkits/NiagaraSystemToolkit.cpp
--- a/Engine/Plugins/FX/Niagara/Source/NiagaraEditor/Private/Toolkits/NiagaraSystemToolkit.cpp
+++ b/Engine/Plugins/FX/Niagara/Source/NiagaraEditor/Private/Toolkits/NiagaraSystemToolkit.cpp
@@ -347,6 +347,8 @@ void FNiagaraSystemToolkit::InitializeInternal(const EToolkitMode::Type Mode, co
 	SetupCommands();
 	ExtendToolbar();
 	RegenerateMenusAndToolbars();
+
+	bChangesDiscarded = false;
 }
 
 FName FNiagaraSystemToolkit::GetToolkitFName() const
@@ -1080,7 +1082,7 @@ bool FNiagaraSystemToolkit::OnRequestClose()
 	if (SystemToolkitMode == ESystemToolkitMode::Emitter)
 	{
 		TSharedPtr<FNiagaraEmitterViewModel> EmitterViewModel = SystemViewModel->GetEmitterHandleViewModels()[0]->GetEmitterViewModel();
-		if (EmitterViewModel->GetEmitter()->GetChangeId() != LastSyncedEmitterChangeId)
+		if (bChangesDiscarded == false && EmitterViewModel->GetEmitter()->GetChangeId() != LastSyncedEmitterChangeId)
 		{
 			// find out the user wants to do with this dirty NiagaraScript
 			EAppReturnType::Type YesNoCancelReply = FMessageDialog::Open(EAppMsgType::YesNoCancel,
@@ -1095,12 +1097,10 @@ bool FNiagaraSystemToolkit::OnRequestClose()
 				// update NiagaraScript and exit
 				UpdateOriginalEmitter();
 				break;
-
 			case EAppReturnType::No:
-				// exit
-				// do nothing.. bNiagaraScriptDirty = false;
+				// Set the changes discarded to avoid showing the dialog multiple times when request close is called multiple times on shut down.
+				bChangesDiscarded = true;
 				break;
-
 			case EAppReturnType::Cancel:
 				// don't exit
 				return false;
diff --git a/Engine/Plugins/FX/Niagara/Source/NiagaraEditor/Private/Toolkits/NiagaraSystemToolkit.h b/Engine/Plugins/FX/Niagara/Source/NiagaraEditor/Private/Toolkits/NiagaraSystemToolkit.h
--- a/Engine/Plugins/FX/Niagara/Source/NiagaraEditor/Private/Toolkits/NiagaraSystemToolkit.h
+++ b/Engine/Plugins/FX/Niagara/Source/NiagaraEditor/Private/Toolkits/NiagaraSystemToolkit.h
@@ -145,6 +145,8 @@ private:
 
 	TSharedPtr<class SNiagaraParameterMapView> ParameterMapView;
 
+	bool bChangesDiscarded;
+
 	static const FName ViewportTabID;
 	static const FName CurveEditorTabID;
 	static const FName SequencerTabID;
