Niagara - Fix an issue where deleting a module and undoing would not restore input values for rapid iteration parameters.

diff --git a/Engine/Plugins/FX/Niagara/Source/NiagaraEditor/Private/ViewModels/Stack/NiagaraStackModuleItem.cpp b/Engine/Plugins/FX/Niagara/Source/NiagaraEditor/Private/ViewModels/Stack/NiagaraStackModuleItem.cpp
--- a/Engine/Plugins/FX/Niagara/Source/NiagaraEditor/Private/ViewModels/Stack/NiagaraStackModuleItem.cpp
+++ b/Engine/Plugins/FX/Niagara/Source/NiagaraEditor/Private/ViewModels/Stack/NiagaraStackModuleItem.cpp
@@ -564,7 +564,19 @@ void UNiagaraStackModuleItem::Delete()
 {
 	checkf(CanMoveAndDelete(), TEXT("This module can't be deleted"));
 
+	UNiagaraNodeOutput* OutputNode = FNiagaraStackGraphUtilities::GetEmitterOutputNodeForStackNode(*FunctionCallNode);
+	const FNiagaraEmitterHandle* EmitterHandle = FNiagaraEditorUtilities::GetEmitterHandleForEmitter(GetSystemViewModel()->GetSystem(), *GetEmitterViewModel()->GetEmitter());
+	checkf(OutputNode != nullptr && EmitterHandle != nullptr, TEXT("Invalid Stack - Output node or emitter handle could not be found for module"));
+
+	UNiagaraScript* OwningScript = FNiagaraEditorUtilities::GetScriptFromSystem(
+		GetSystemViewModel()->GetSystem(), EmitterHandle->GetId(), OutputNode->GetUsage(), OutputNode->GetUsageId());
+	checkf(OwningScript != nullptr, TEXT("Invalid Stack - Owning script could not be found for module"));
+
 	FScopedTransaction ScopedTransaction(LOCTEXT("RemoveAModuleFromTheStack", "Remove a module from the stack"));
+
+	// Modify the script so that any rapid iteration parameters in use are recorded.
+	OwningScript->Modify();
+
 	TArray<TWeakObjectPtr<UNiagaraNodeInput>> RemovedNodes;
 	bool bRemoved = FNiagaraStackGraphUtilities::RemoveModuleFromStack(*FunctionCallNode, RemovedNodes);
 	if (bRemoved)
