Niagara - Prevent a crash when the user sets the loop duration of their emitter to be negative.

diff --git a/Engine/Plugins/FX/Niagara/Source/NiagaraEditor/Private/NiagaraStackFunctionInputBinder.cpp b/Engine/Plugins/FX/Niagara/Source/NiagaraEditor/Private/NiagaraStackFunctionInputBinder.cpp
--- a/Engine/Plugins/FX/Niagara/Source/NiagaraEditor/Private/NiagaraStackFunctionInputBinder.cpp
+++ b/Engine/Plugins/FX/Niagara/Source/NiagaraEditor/Private/NiagaraStackFunctionInputBinder.cpp
@@ -187,6 +187,11 @@ FNiagaraTypeDefinition FNiagaraStackFunctionInputBinder::GetInputType() const
 	return InputType;
 }
 
+UNiagaraNodeFunctionCall* FNiagaraStackFunctionInputBinder::GetFunctionCallNode() const
+{
+	return FunctionCallNode.Get();
+}
+
 TArray<uint8> FNiagaraStackFunctionInputBinder::GetData() const
 {
 	checkf(Script.IsValid(), TEXT("Bound script is not valid"));
diff --git a/Engine/Plugins/FX/Niagara/Source/NiagaraEditor/Private/Sequencer/NiagaraSequence/Sections/MovieSceneNiagaraEmitterSection.cpp b/Engine/Plugins/FX/Niagara/Source/NiagaraEditor/Private/Sequencer/NiagaraSequence/Sections/MovieSceneNiagaraEmitterSection.cpp
--- a/Engine/Plugins/FX/Niagara/Source/NiagaraEditor/Private/Sequencer/NiagaraSequence/Sections/MovieSceneNiagaraEmitterSection.cpp
+++ b/Engine/Plugins/FX/Niagara/Source/NiagaraEditor/Private/Sequencer/NiagaraSequence/Sections/MovieSceneNiagaraEmitterSection.cpp
@@ -257,6 +257,15 @@ void UMovieSceneNiagaraEmitterSection::UpdateSectionFromTimeRangeModule(const FF
 	{
 		float ModuleStartTime = StartTimeBinder.GetValue<float>();
 		float ModuleLength = LengthBinder.GetValue<float>();
+
+		if (ModuleLength < 0)
+		{
+			// TODO: Add ui support for this issue rather than a log error.
+			UE_LOG(LogNiagaraEditor, Error, TEXT("Invalid length in niagara editor timeline.  Bound Module: %s Bound Input: %s"),
+				LengthBinder.GetFunctionCallNode() != nullptr ? *LengthBinder.GetFunctionCallNode()->GetFunctionName() : TEXT("Unknown"),
+				*LengthBinder.GetInputName().ToString());
+			ModuleLength = 0;
+		}
 		
 		SetRange(TRange<FFrameNumber>(
 			(ModuleStartTime * InTickResolution).RoundToFrame(),
diff --git a/Engine/Plugins/FX/Niagara/Source/NiagaraEditor/Public/NiagaraStackFunctionInputBinder.h b/Engine/Plugins/FX/Niagara/Source/NiagaraEditor/Public/NiagaraStackFunctionInputBinder.h
--- a/Engine/Plugins/FX/Niagara/Source/NiagaraEditor/Public/NiagaraStackFunctionInputBinder.h
+++ b/Engine/Plugins/FX/Niagara/Source/NiagaraEditor/Public/NiagaraStackFunctionInputBinder.h
@@ -48,6 +48,8 @@ public:
 
 	FNiagaraTypeDefinition GetInputType() const;
 
+	UNiagaraNodeFunctionCall* GetFunctionCallNode() const;
+
 	template<typename ValueType>
 	ValueType GetValue() const
 	{
