Niagara - Prevent parameter store change events with uninitialized data.

diff --git a/Engine/Plugins/FX/Niagara/Source/NiagaraEditor/Private/NiagaraEmitterFactoryNew.cpp b/Engine/Plugins/FX/Niagara/Source/NiagaraEditor/Private/NiagaraEmitterFactoryNew.cpp
--- a/Engine/Plugins/FX/Niagara/Source/NiagaraEditor/Private/NiagaraEmitterFactoryNew.cpp
+++ b/Engine/Plugins/FX/Niagara/Source/NiagaraEditor/Private/NiagaraEmitterFactoryNew.cpp
@@ -123,7 +123,8 @@ UObject* UNiagaraEmitterFactoryNew::FactoryCreateNew(UClass* Class, UObject* InP
 							{
 								FNiagaraVariable VelocityVar(FNiagaraTypeDefinition::GetVec3Def(), *(TEXT("Constants.Emitter.") + CallNode->GetFunctionName() + TEXT(".Velocity")));
 								VelocityVar.SetValue(FVector(0.0f, 0.0f, 100.0f));
-								NewEmitter->SpawnScriptProps.Script->RapidIterationParameters.AddParameter(VelocityVar);
+								bool bAddParameterIfMissing = true;
+								NewEmitter->SpawnScriptProps.Script->RapidIterationParameters.SetParameterData(VelocityVar.GetData(), VelocityVar, bAddParameterIfMissing);
 							}
 						}
 						else
diff --git a/Engine/Plugins/FX/Niagara/Source/NiagaraEditor/Private/NiagaraScriptMergeManager.cpp b/Engine/Plugins/FX/Niagara/Source/NiagaraEditor/Private/NiagaraScriptMergeManager.cpp
--- a/Engine/Plugins/FX/Niagara/Source/NiagaraEditor/Private/NiagaraScriptMergeManager.cpp
+++ b/Engine/Plugins/FX/Niagara/Source/NiagaraEditor/Private/NiagaraScriptMergeManager.cpp
@@ -1721,8 +1721,8 @@ FNiagaraScriptMergeManager::FApplyDiffResults FNiagaraScriptMergeManager::AddInp
 			FNiagaraVariable RapidIterationParameter(OverrideToAdd->GetLocalValueRapidIterationParameter().GetValue().GetType(), *RapidIterationParameterName);
 			const uint8* SourceData = OverrideToAdd->GetOwningScript()->RapidIterationParameters.GetParameterData(OverrideToAdd->GetLocalValueRapidIterationParameter().GetValue());
 			OwningScript.Modify();
-			OwningScript.RapidIterationParameters.AddParameter(RapidIterationParameter);
-			OwningScript.RapidIterationParameters.SetParameterData(SourceData, RapidIterationParameter);
+			bool bAddParameterIfMissing = true;
+			OwningScript.RapidIterationParameters.SetParameterData(SourceData, RapidIterationParameter, bAddParameterIfMissing);
 			Results.bSucceeded = true;
 		}
 		else
diff --git a/Engine/Plugins/FX/Niagara/Source/NiagaraEditor/Private/NiagaraScriptSource.cpp b/Engine/Plugins/FX/Niagara/Source/NiagaraEditor/Private/NiagaraScriptSource.cpp
--- a/Engine/Plugins/FX/Niagara/Source/NiagaraEditor/Private/NiagaraScriptSource.cpp
+++ b/Engine/Plugins/FX/Niagara/Source/NiagaraEditor/Private/NiagaraScriptSource.cpp
@@ -204,8 +204,8 @@ void InitializeNewRapidIterationParametersForNode(const UEdGraphSchema_Niagara*
 						{
 							FNiagaraVariable DefaultVariable = Schema->PinToNiagaraVariable(DefaultPin, true);
 							check(DefaultVariable.GetData() != nullptr);
-							RapidIterationParameters.AddParameter(RapidIterationParameter, false);
-							RapidIterationParameters.SetParameterData(DefaultVariable.GetData(), RapidIterationParameter);
+							bool bAddParameterIfMissing = true;
+							RapidIterationParameters.SetParameterData(DefaultVariable.GetData(), RapidIterationParameter, bAddParameterIfMissing);
 						}
 					}
 				}
diff --git a/Engine/Plugins/FX/Niagara/Source/NiagaraEditor/Private/ViewModels/Stack/NiagaraStackFunctionInput.cpp b/Engine/Plugins/FX/Niagara/Source/NiagaraEditor/Private/ViewModels/Stack/NiagaraStackFunctionInput.cpp
--- a/Engine/Plugins/FX/Niagara/Source/NiagaraEditor/Private/ViewModels/Stack/NiagaraStackFunctionInput.cpp
+++ b/Engine/Plugins/FX/Niagara/Source/NiagaraEditor/Private/ViewModels/Stack/NiagaraStackFunctionInput.cpp
@@ -926,11 +926,8 @@ void UNiagaraStackFunctionInput::SetLocalValue(TSharedRef<FStructOnScope> InLoca
 
 		for (TWeakObjectPtr<UNiagaraScript> Script : AffectedScripts)
 		{
-			if(Script->RapidIterationParameters.AddParameter(RapidIterationParameter))
-			{
-				UE_LOG(LogNiagaraEditor, Log, TEXT("Adding Parameter %s to Script %s"), *RapidIterationParameter.GetName().ToString(), *Script->GetFullName());
-			}
-			Script->RapidIterationParameters.SetParameterData(InLocalValue->GetStructMemory(), RapidIterationParameter);
+			bool bAddParameterIfMissing = true;
+			Script->RapidIterationParameters.SetParameterData(InLocalValue->GetStructMemory(), RapidIterationParameter, bAddParameterIfMissing);
 		}
 	}
 	else 
@@ -1039,11 +1036,8 @@ bool UNiagaraStackFunctionInput::UpdateRapidIterationParametersForAffectedScript
 
 	for (TWeakObjectPtr<UNiagaraScript> Script : AffectedScripts)
 	{
-		if (Script->RapidIterationParameters.AddParameter(RapidIterationParameter))
-		{
-			UE_LOG(LogNiagaraEditor, Log, TEXT("Adding Parameter %s to Script %s"), *RapidIterationParameter.GetName().ToString(), *Script->GetFullName());
-		}
-		Script->RapidIterationParameters.SetParameterData(Data, RapidIterationParameter);
+		bool bAddParameterIfMissing = true;
+		Script->RapidIterationParameters.SetParameterData(Data, RapidIterationParameter, bAddParameterIfMissing);
 	}
 	GetSystemViewModel()->ResetSystem();
 	return true;
@@ -1245,9 +1239,9 @@ void UNiagaraStackFunctionInput::ResetToBase()
 						AffectedScript->Modify();
 						for (FNiagaraVariable& OwningScriptRapidIterationParameter : OwningScriptRapidIterationParameters)
 						{
-							AffectedScript->RapidIterationParameters.AddParameter(OwningScriptRapidIterationParameter);
+							bool bAddParameterIfMissing = true;
 							AffectedScript->RapidIterationParameters.SetParameterData(
-								OwningScript->RapidIterationParameters.GetParameterData(OwningScriptRapidIterationParameter), OwningScriptRapidIterationParameter);
+								OwningScript->RapidIterationParameters.GetParameterData(OwningScriptRapidIterationParameter), OwningScriptRapidIterationParameter, bAddParameterIfMissing);
 						}
 					}
 				}
