Niagara - Fix issues with editing burst keys through the timeline context menu.

diff --git a/Engine/Plugins/FX/Niagara/Source/NiagaraEditor/Private/Sequencer/NiagaraSequence/Sections/MovieSceneNiagaraEmitterSection.cpp b/Engine/Plugins/FX/Niagara/Source/NiagaraEditor/Private/Sequencer/NiagaraSequence/Sections/MovieSceneNiagaraEmitterSection.cpp
--- a/Engine/Plugins/FX/Niagara/Source/NiagaraEditor/Private/Sequencer/NiagaraSequence/Sections/MovieSceneNiagaraEmitterSection.cpp
+++ b/Engine/Plugins/FX/Niagara/Source/NiagaraEditor/Private/Sequencer/NiagaraSequence/Sections/MovieSceneNiagaraEmitterSection.cpp
@@ -99,7 +99,7 @@ TSharedPtr<FStructOnScope> GetKeyStruct(TMovieSceneChannelHandle<FMovieSceneNiag
 	int32 KeyValueIndex = Channel.Get()->GetData().GetIndex(InHandle);
 	if (KeyValueIndex != INDEX_NONE)
 	{
-		FNiagaraEmitterSectionKey KeyValue = Channel.Get()->GetData().GetValues()[KeyValueIndex];
+		FNiagaraEmitterSectionKey& KeyValue = Channel.Get()->GetData().GetValues()[KeyValueIndex];
 		return MakeShared<FStructOnScope>(KeyValue.Value.GetType().GetStruct(), KeyValue.Value.GetData());
 	}
 	return TSharedPtr<FStructOnScope>();
