Niagara - Reset cached dependency data when the emitter view models are refreshed in the system.  Fixes incorrect dependencies when applying changes from parent emitters in systems.

diff --git a/Engine/Plugins/FX/Niagara/Source/NiagaraEditor/Private/ViewModels/NiagaraSystemViewModel.cpp b/Engine/Plugins/FX/Niagara/Source/NiagaraEditor/Private/ViewModels/NiagaraSystemViewModel.cpp
--- a/Engine/Plugins/FX/Niagara/Source/NiagaraEditor/Private/ViewModels/NiagaraSystemViewModel.cpp
+++ b/Engine/Plugins/FX/Niagara/Source/NiagaraEditor/Private/ViewModels/NiagaraSystemViewModel.cpp
@@ -713,6 +713,7 @@ void FNiagaraSystemViewModel::RefreshEmitterHandleViewModels()
 {
 	TArray<TSharedRef<FNiagaraEmitterHandleViewModel>> OldViewModels = EmitterHandleViewModels;
 	EmitterHandleViewModels.Empty();
+	EmitterToCachedStackModuleData.Empty();
 
 	// Map existing view models to the real instances that now exist. Reuse if we can. Create a new one if we cannot.
 	TArray<FGuid> ValidEmitterHandleIds;
