Fix monolithic compile errors in Niagara modules.

- HoudiniNiagaraEditor calls IMPLEMENT_MODULE with incorrect module name.
- NiagaraEditorWidgetsStyle has RelativePathToPluginPath function copy/pasted from NiagaraEditorStyle.
- NiagaraNodeReroute defines duplicate global symbol called "PC_Wildcard"
- SNiagaraParameterMapView defines duplicate instance of LogNiagara.

diff --git a/Engine/Plugins/FX/HoudiniNiagara/Source/HoudiniNiagaraEditor/Private/HoudiniNiagaraEditor.cpp b/Engine/Plugins/FX/HoudiniNiagara/Source/HoudiniNiagaraEditor/Private/HoudiniNiagaraEditor.cpp
--- a/Engine/Plugins/FX/HoudiniNiagara/Source/HoudiniNiagaraEditor/Private/HoudiniNiagaraEditor.cpp
+++ b/Engine/Plugins/FX/HoudiniNiagara/Source/HoudiniNiagaraEditor/Private/HoudiniNiagaraEditor.cpp
@@ -93,4 +93,4 @@ void FHoudiniNiagaraEditorModule::ShutdownModule()
 
 #undef LOCTEXT_NAMESPACE
 	
-IMPLEMENT_MODULE(FHoudiniNiagaraEditorModule, HoudiniNiagara)
+IMPLEMENT_MODULE(FHoudiniNiagaraEditorModule, HoudiniNiagaraEditor)
diff --git a/Engine/Plugins/FX/Niagara/Source/NiagaraEditor/Private/NiagaraEditorStyle.cpp b/Engine/Plugins/FX/Niagara/Source/NiagaraEditor/Private/NiagaraEditorStyle.cpp
--- a/Engine/Plugins/FX/Niagara/Source/NiagaraEditor/Private/NiagaraEditorStyle.cpp
+++ b/Engine/Plugins/FX/Niagara/Source/NiagaraEditor/Private/NiagaraEditorStyle.cpp
@@ -34,7 +34,7 @@ FName FNiagaraEditorStyle::GetStyleSetName()
 	return StyleSetName;
 }
 
-FString RelativePathToPluginPath(const FString& RelativePath, const ANSICHAR* Extension)
+NIAGARAEDITOR_API FString RelativePathToPluginPath(const FString& RelativePath, const ANSICHAR* Extension)
 {
 	static FString ContentDir = IPluginManager::Get().FindPlugin(TEXT("Niagara"))->GetContentDir();
 	return (ContentDir / RelativePath) + Extension;
diff --git a/Engine/Plugins/FX/Niagara/Source/NiagaraEditor/Private/NiagaraNodeReroute.cpp b/Engine/Plugins/FX/Niagara/Source/NiagaraEditor/Private/NiagaraNodeReroute.cpp
--- a/Engine/Plugins/FX/Niagara/Source/NiagaraEditor/Private/NiagaraNodeReroute.cpp
+++ b/Engine/Plugins/FX/Niagara/Source/NiagaraEditor/Private/NiagaraNodeReroute.cpp
@@ -6,7 +6,7 @@
 
 #define LOCTEXT_NAMESPACE "NiagaraNodeReroute"
 
-const char* PC_Wildcard = "wildcard";
+const char* PC_Wildcard_Niagara = "wildcard";
 
 UNiagaraNodeReroute::UNiagaraNodeReroute(const FObjectInitializer& ObjectInitializer)
 	: Super(ObjectInitializer)
@@ -32,11 +32,11 @@ void UNiagaraNodeReroute::PostLoad()
 void UNiagaraNodeReroute::AllocateDefaultPins()
 {
 	const FName InputPinName(TEXT("InputPin"));
-	UEdGraphPin* MyInputPin = CreatePin(EGPD_Input, PC_Wildcard, InputPinName);
+	UEdGraphPin* MyInputPin = CreatePin(EGPD_Input, PC_Wildcard_Niagara, InputPinName);
 	MyInputPin->bDefaultValueIsIgnored = true;
 
 	const FName OutputPinName(TEXT("OutputPin"));
-	CreatePin(EGPD_Output, PC_Wildcard, OutputPinName);
+	CreatePin(EGPD_Output, PC_Wildcard_Niagara, OutputPinName);
 }
 
 FText UNiagaraNodeReroute::GetTooltipText() const
@@ -138,7 +138,7 @@ void UNiagaraNodeReroute::PropagatePinType()
 
 	for (UEdGraphPin* Inputs : MyInputPin->LinkedTo)
 	{
-		if (Inputs->PinType.PinCategory != PC_Wildcard)
+		if (Inputs->PinType.PinCategory != PC_Wildcard_Niagara)
 		{
 			PropagatePinTypeFromDirection(true);
 			return;
@@ -147,7 +147,7 @@ void UNiagaraNodeReroute::PropagatePinType()
 
 	for (UEdGraphPin* Outputs : MyOutputPin->LinkedTo)
 	{
-		if (Outputs->PinType.PinCategory != PC_Wildcard)
+		if (Outputs->PinType.PinCategory != PC_Wildcard_Niagara)
 		{
 			PropagatePinTypeFromDirection(false);
 			return;
@@ -170,11 +170,11 @@ void UNiagaraNodeReroute::PropagatePinType()
 		// Revert to wildcard
 		MyInputPin->BreakAllPinLinks();
 		MyInputPin->PinType.ResetToDefaults();
-		MyInputPin->PinType.PinCategory = PC_Wildcard;
+		MyInputPin->PinType.PinCategory = PC_Wildcard_Niagara;
 
 		MyOutputPin->BreakAllPinLinks();
 		MyOutputPin->PinType.ResetToDefaults();
-		MyOutputPin->PinType.PinCategory = PC_Wildcard;
+		MyOutputPin->PinType.PinCategory = PC_Wildcard_Niagara;
 	}
 }
 
diff --git a/Engine/Plugins/FX/Niagara/Source/NiagaraEditor/Private/Widgets/SNiagaraParameterMapView.cpp b/Engine/Plugins/FX/Niagara/Source/NiagaraEditor/Private/Widgets/SNiagaraParameterMapView.cpp
--- a/Engine/Plugins/FX/Niagara/Source/NiagaraEditor/Private/Widgets/SNiagaraParameterMapView.cpp
+++ b/Engine/Plugins/FX/Niagara/Source/NiagaraEditor/Private/Widgets/SNiagaraParameterMapView.cpp
@@ -35,8 +35,6 @@
 
 #define LOCTEXT_NAMESPACE "NiagaraParameterMapView"
 
-DEFINE_LOG_CATEGORY(LogNiagara);
-
 FText NiagaraParameterMapSectionID::OnGetSectionTitle(const NiagaraParameterMapSectionID::Type InSection)
 {
 	/* Setup an appropriate name for the section for this node */
diff --git a/Engine/Plugins/FX/Niagara/Source/NiagaraEditorWidgets/Private/NiagaraEditorWidgetsStyle.cpp b/Engine/Plugins/FX/Niagara/Source/NiagaraEditorWidgets/Private/NiagaraEditorWidgetsStyle.cpp
--- a/Engine/Plugins/FX/Niagara/Source/NiagaraEditorWidgets/Private/NiagaraEditorWidgetsStyle.cpp
+++ b/Engine/Plugins/FX/Niagara/Source/NiagaraEditorWidgets/Private/NiagaraEditorWidgetsStyle.cpp
@@ -34,11 +34,7 @@ FName FNiagaraEditorWidgetsStyle::GetStyleSetName()
 	return StyleSetName;
 }
 
-FString RelativePathToPluginPath(const FString& RelativePath, const ANSICHAR* Extension)
-{
-	static FString ContentDir = IPluginManager::Get().FindPlugin(TEXT("Niagara"))->GetContentDir();
-	return (ContentDir / RelativePath) + Extension;
-}
+NIAGARAEDITOR_API FString RelativePathToPluginPath(const FString& RelativePath, const ANSICHAR* Extension);
 
 #define IMAGE_PLUGIN_BRUSH( RelativePath, ... ) FSlateImageBrush( RelativePathToPluginPath( RelativePath, ".png" ), __VA_ARGS__ )
 #define IMAGE_BRUSH( RelativePath, ... ) FSlateImageBrush( Style->RootToContentDir( RelativePath, TEXT(".png") ), __VA_ARGS__ )
