Merging CL 4127755 to fix static analysis warning.

diff --git a/Engine/Plugins/FX/Niagara/Source/NiagaraEditor/Private/Sequencer/NiagaraSequence/Sections/MovieSceneNiagaraEmitterSection.cpp b/Engine/Plugins/FX/Niagara/Source/NiagaraEditor/Private/Sequencer/NiagaraSequence/Sections/MovieSceneNiagaraEmitterSection.cpp
--- a/Engine/Plugins/FX/Niagara/Source/NiagaraEditor/Private/Sequencer/NiagaraSequence/Sections/MovieSceneNiagaraEmitterSection.cpp
+++ b/Engine/Plugins/FX/Niagara/Source/NiagaraEditor/Private/Sequencer/NiagaraSequence/Sections/MovieSceneNiagaraEmitterSection.cpp
@@ -99,8 +99,9 @@ TSharedPtr<FStructOnScope> GetKeyStruct(TMovieSceneChannelHandle<FMovieSceneNiag
 	int32 KeyValueIndex = Channel.Get()->GetData().GetIndex(InHandle);
 	if (KeyValueIndex != INDEX_NONE)
 	{
-		FNiagaraEmitterSectionKey& KeyValue = Channel.Get()->GetData().GetValues()[KeyValueIndex];
-		return MakeShared<FStructOnScope>(KeyValue.Value.GetType().GetStruct(), KeyValue.Value.GetData());
+		FNiagaraTypeDefinition KeyType = Channel.Get()->GetData().GetValues()[KeyValueIndex].Value.GetType();
+		uint8* KeyData = Channel.Get()->GetData().GetValues()[KeyValueIndex].Value.GetData();
+		return MakeShared<FStructOnScope>(KeyType.GetStruct(), KeyData);
 	}
 	return TSharedPtr<FStructOnScope>();
 }
