Hide dynamic load balanced sim option for now as it doesn't work and could cause problems

diff --git a/Engine/Plugins/FX/Niagara/Source/Niagara/Public/NiagaraCommon.h b/Engine/Plugins/FX/Niagara/Source/Niagara/Public/NiagaraCommon.h
--- a/Engine/Plugins/FX/Niagara/Source/Niagara/Public/NiagaraCommon.h
+++ b/Engine/Plugins/FX/Niagara/Source/Niagara/Public/NiagaraCommon.h
@@ -34,7 +34,7 @@ enum class ENiagaraSimTarget : uint8
 {
 	CPUSim,
 	GPUComputeSim,
-	DynamicLoadBalancedSim
+	DynamicLoadBalancedSim UMETA(Hidden)
 };
 
 
