Don't call Alloc() for GPU sim

diff --git a/Engine/Plugins/FX/Niagara/Source/Niagara/Private/NiagaraEmitterInstance.cpp b/Engine/Plugins/FX/Niagara/Source/Niagara/Private/NiagaraEmitterInstance.cpp
--- a/Engine/Plugins/FX/Niagara/Source/Niagara/Private/NiagaraEmitterInstance.cpp
+++ b/Engine/Plugins/FX/Niagara/Source/Niagara/Private/NiagaraEmitterInstance.cpp
@@ -819,7 +819,10 @@ void FNiagaraEmitterInstance::Tick(float DeltaSeconds)
 	if (OrigNumParticles == 0 && ExecutionState != ENiagaraExecutionState::Active)
 	{
 		//Clear out curr buffer in case it had some data in previously.
-		Data.Allocate(0);
+		if (CachedEmitter->SimTarget == ENiagaraSimTarget::CPUSim)
+		{
+			Data.Allocate(0);
+		}
 		CPUTimeMS = TickTime.GetElapsedMilliseconds();
 		return;
 	}
