Fixed crash on spawn in Niagara content examples.

diff --git a/Engine/Plugins/FX/Niagara/Source/Niagara/Private/NiagaraSystemSimulation.cpp b/Engine/Plugins/FX/Niagara/Source/Niagara/Private/NiagaraSystemSimulation.cpp
--- a/Engine/Plugins/FX/Niagara/Source/Niagara/Private/NiagaraSystemSimulation.cpp
+++ b/Engine/Plugins/FX/Niagara/Source/Niagara/Private/NiagaraSystemSimulation.cpp
@@ -324,6 +324,7 @@ bool FNiagaraSystemSimulation::Tick(float DeltaSeconds)
 		}
 
 		TArray<FNiagaraDataSetExecutionInfo, TInlineAllocator<8>> DataSetExecInfos;
+		DataSetExecInfos.SetNum(2);
 		{
 			SCOPE_CYCLE_COUNTER(STAT_NiagaraSystemSim_PrepareForSimulate);
 			auto PreSimulateAndTransferParams = [&](int32 SystemIndex)
@@ -386,8 +387,6 @@ bool FNiagaraSystemSimulation::Tick(float DeltaSeconds)
 			UpdateGlobalSpawnCountScaleParam.SetValue(GlobalSpawnCountScale);
 			SpawnGlobalSystemCountScaleParam.SetValue(GlobalSystemCountScale);
 			UpdateGlobalSystemCountScaleParam.SetValue(GlobalSystemCountScale);
-
-			DataSetExecInfos.Emplace(&DataSet, 0, false, true);
 		}
 
 		//TODO: JIRA - UE-60096 - Remove.
@@ -403,9 +402,8 @@ bool FNiagaraSystemSimulation::Tick(float DeltaSeconds)
 
 			//Run Spawn
 			SpawnExecContext.Tick(SoloSystemInstance);//We can't require a specific instance here as these are for all instances.
-			DataSetExecInfos.SetNum(1, false);
-			DataSetExecInfos[0].StartInstance = OrigNum;
-			DataSetExecInfos.Emplace(&SpawnInstanceParameterDataSet, OrigNum, false, false);
+			DataSetExecInfos[0] = FNiagaraDataSetExecutionInfo(&DataSet, OrigNum, false, true);
+			DataSetExecInfos[1] = FNiagaraDataSetExecutionInfo(&SpawnInstanceParameterDataSet, OrigNum, false, false);
 			SpawnExecContext.Execute(SpawnNum, DataSetExecInfos);
 
 			if (GbDumpSystemData || System->bDumpDebugSystemInfo)
@@ -439,8 +437,8 @@ bool FNiagaraSystemSimulation::Tick(float DeltaSeconds)
 
 			//Run update.
 			UpdateExecContext.Tick(SystemInstances[0]);
-			DataSetExecInfos[0].StartInstance = 0;
-			DataSetExecInfos.Emplace(&UpdateInstanceParameterDataSet, 0, false, false);
+			DataSetExecInfos[0] = FNiagaraDataSetExecutionInfo(&DataSet, 0, false, true);
+			DataSetExecInfos[1] = FNiagaraDataSetExecutionInfo(&UpdateInstanceParameterDataSet, 0, false, false);
 
 // 			if (GbDumpSystemData || System->bDumpDebugSystemInfo)
 // 			{
@@ -455,6 +453,7 @@ bool FNiagaraSystemSimulation::Tick(float DeltaSeconds)
 			{
 				UE_LOG(LogNiagara, Log, TEXT("=== Updated %d Systems ==="), OrigNum);
 				DataSet.Dump(true, 0, OrigNum);
+				UpdateInstanceParameterDataSet.Dump(false, 0, OrigNum);
 			}
 
  			//Also run the update script on the newly spawned systems too.
@@ -466,8 +465,8 @@ bool FNiagaraSystemSimulation::Tick(float DeltaSeconds)
 
 				//Run update.
 				UpdateExecContext.Tick(SystemInstances[0]);
-				DataSetExecInfos[0].StartInstance = OrigNum;
-				DataSetExecInfos[1].StartInstance = OrigNum;
+				DataSetExecInfos[0] = FNiagaraDataSetExecutionInfo(&DataSet, OrigNum, false, true);
+				DataSetExecInfos[1] = FNiagaraDataSetExecutionInfo(&UpdateInstanceParameterDataSet, OrigNum, false, false);
 
 				UpdateExecContext.Parameters.SetParameterValue(0.0001f, SYS_PARAM_ENGINE_DELTA_TIME);
 				UpdateExecContext.Parameters.SetParameterValue(10000.0f, SYS_PARAM_ENGINE_INV_DELTA_TIME);
@@ -478,6 +477,7 @@ bool FNiagaraSystemSimulation::Tick(float DeltaSeconds)
 				{
 					UE_LOG(LogNiagara, Log, TEXT("=== Spawn Updated %d Systems ==="), SpawnNum);
 					DataSet.Dump(true, OrigNum, SpawnNum);
+					UpdateInstanceParameterDataSet.Dump(false, OrigNum, SpawnNum);
 				}
 			}
 
