Niagara - Prevent a crash when deleting burst keys from an emitter that's in a system.

diff --git a/Engine/Plugins/FX/Niagara/Source/NiagaraEditor/Private/ViewModels/Stack/NiagaraStackEntry.cpp b/Engine/Plugins/FX/Niagara/Source/NiagaraEditor/Private/ViewModels/Stack/NiagaraStackEntry.cpp
--- a/Engine/Plugins/FX/Niagara/Source/NiagaraEditor/Private/ViewModels/Stack/NiagaraStackEntry.cpp
+++ b/Engine/Plugins/FX/Niagara/Source/NiagaraEditor/Private/ViewModels/Stack/NiagaraStackEntry.cpp
@@ -306,6 +306,11 @@ UNiagaraStackEntry::FOnRequestFullRefresh& UNiagaraStackEntry::OnRequestFullRefr
 	return RequestFullRefreshDelegate;
 }
 
+UNiagaraStackEntry::FOnRequestFullRefresh& UNiagaraStackEntry::OnRequestFullRefreshDeferred()
+{
+	return RequestFullRefreshDeferredDelegate;
+}
+
 int32 UNiagaraStackEntry::GetIndentLevel() const
 {
 	return IndentLevel;
@@ -410,6 +415,7 @@ void UNiagaraStackEntry::RefreshChildren()
 		Child->OnStructureChanged().RemoveAll(this);
 		Child->OnDataObjectModified().RemoveAll(this);
 		Child->OnRequestFullRefresh().RemoveAll(this);
+		Child->OnRequestFullRefreshDeferred().RemoveAll(this);
 		Child->SetOnRequestCanDrop(FOnRequestDrop());
 		Child->SetOnRequestDrop(FOnRequestDrop());
 	}
@@ -418,6 +424,7 @@ void UNiagaraStackEntry::RefreshChildren()
 		ErrorChild->OnStructureChanged().RemoveAll(this);
 		ErrorChild->OnDataObjectModified().RemoveAll(this);
 		ErrorChild->OnRequestFullRefresh().RemoveAll(this);
+		ErrorChild->OnRequestFullRefreshDeferred().RemoveAll(this);
 		ErrorChild->OnIssueModified().RemoveAll(this);
 	}
 
@@ -445,6 +452,7 @@ void UNiagaraStackEntry::RefreshChildren()
 		Child->OnStructureChanged().AddUObject(this, &UNiagaraStackEntry::ChildStructureChanged);
 		Child->OnDataObjectModified().AddUObject(this, &UNiagaraStackEntry::ChildDataObjectModified);
 		Child->OnRequestFullRefresh().AddUObject(this, &UNiagaraStackEntry::ChildRequestFullRefresh);
+		Child->OnRequestFullRefreshDeferred().AddUObject(this, &UNiagaraStackEntry::ChildRequestFullRefreshDeferred);
 		Child->SetOnRequestCanDrop(FOnRequestDrop::CreateUObject(this, &UNiagaraStackEntry::ChildRequestCanDrop));
 		Child->SetOnRequestDrop(FOnRequestDrop::CreateUObject(this, &UNiagaraStackEntry::ChildRequestDrop));
 	}
@@ -462,6 +470,7 @@ void UNiagaraStackEntry::RefreshChildren()
 		ErrorChild->OnStructureChanged().AddUObject(this, &UNiagaraStackEntry::ChildStructureChanged);
 		ErrorChild->OnDataObjectModified().AddUObject(this, &UNiagaraStackEntry::ChildDataObjectModified);
 		ErrorChild->OnRequestFullRefresh().AddUObject(this, &UNiagaraStackEntry::ChildRequestFullRefresh);
+		ErrorChild->OnRequestFullRefreshDeferred().AddUObject(this, &UNiagaraStackEntry::ChildRequestFullRefreshDeferred);
 		ErrorChild->OnIssueModified().AddUObject(this, &UNiagaraStackEntry::IssueModified);
 	}
 
@@ -551,6 +560,11 @@ void UNiagaraStackEntry::ChildRequestFullRefresh()
 	RequestFullRefreshDelegate.Broadcast();
 }
 
+void UNiagaraStackEntry::ChildRequestFullRefreshDeferred()
+{
+	RequestFullRefreshDeferredDelegate.Broadcast();
+}
+
 TOptional<UNiagaraStackEntry::FDropResult> UNiagaraStackEntry::ChildRequestCanDrop(const UNiagaraStackEntry& TargetChild, const TArray<UNiagaraStackEntry*>& DraggedEntries)
 {
 	TOptional<FDropResult> Result = ChildRequestCanDropInternal(TargetChild, DraggedEntries);
diff --git a/Engine/Plugins/FX/Niagara/Source/NiagaraEditor/Private/ViewModels/Stack/NiagaraStackScriptItemGroup.cpp b/Engine/Plugins/FX/Niagara/Source/NiagaraEditor/Private/ViewModels/Stack/NiagaraStackScriptItemGroup.cpp
--- a/Engine/Plugins/FX/Niagara/Source/NiagaraEditor/Private/ViewModels/Stack/NiagaraStackScriptItemGroup.cpp
+++ b/Engine/Plugins/FX/Niagara/Source/NiagaraEditor/Private/ViewModels/Stack/NiagaraStackScriptItemGroup.cpp
@@ -264,6 +264,17 @@ void UNiagaraStackScriptItemGroup::Initialize(
 	ScriptViewModel = InScriptViewModel;
 	ScriptUsage = InScriptUsage;
 	ScriptUsageId = InScriptUsageId;
+	OnGraphChangedHandle = InScriptViewModel->GetGraphViewModel()->GetGraph()->AddOnGraphChangedHandler(
+		FOnGraphChanged::FDelegate::CreateUObject(this, &UNiagaraStackScriptItemGroup::OnScriptGraphChanged));
+}
+
+void UNiagaraStackScriptItemGroup::FinalizeInternal()
+{
+	if (ScriptViewModel.IsValid())
+	{
+		ScriptViewModel.Pin()->GetGraphViewModel()->GetGraph()->RemoveOnGraphChangedHandler(OnGraphChangedHandle);
+	}
+	Super::FinalizeInternal();
 }
 
 void UNiagaraStackScriptItemGroup::RefreshChildrenInternal(const TArray<UNiagaraStackEntry*>& CurrentChildren, TArray<UNiagaraStackEntry*>& NewChildren, TArray<FStackIssue>& NewIssues)
@@ -609,5 +620,13 @@ void UNiagaraStackScriptItemGroup::ChildModifiedGroupItems()
 	RefreshChildren();
 }
 
+void UNiagaraStackScriptItemGroup::OnScriptGraphChanged(const struct FEdGraphEditAction& InAction)
+{
+	if (InAction.Action == GRAPHACTION_RemoveNode)
+	{
+		OnRequestFullRefreshDeferred().Broadcast();
+	}
+}
+
 
 #undef LOCTEXT_NAMESPACE
diff --git a/Engine/Plugins/FX/Niagara/Source/NiagaraEditor/Private/ViewModels/Stack/NiagaraStackViewModel.cpp b/Engine/Plugins/FX/Niagara/Source/NiagaraEditor/Private/ViewModels/Stack/NiagaraStackViewModel.cpp
--- a/Engine/Plugins/FX/Niagara/Source/NiagaraEditor/Private/ViewModels/Stack/NiagaraStackViewModel.cpp
+++ b/Engine/Plugins/FX/Niagara/Source/NiagaraEditor/Private/ViewModels/Stack/NiagaraStackViewModel.cpp
@@ -36,6 +36,7 @@ void UNiagaraStackViewModel::Initialize(TSharedPtr<FNiagaraSystemViewModel> InSy
 		RootEntry->OnStructureChanged().RemoveAll(this);
 		RootEntry->OnDataObjectModified().RemoveAll(this);
 		RootEntry->OnRequestFullRefresh().RemoveAll(this);
+		RootEntry->OnRequestFullRefreshDeferred().RemoveAll(this);
 		RootEntries.Empty();
 		RootEntry->Finalize();
 		RootEntry = nullptr;
@@ -72,12 +73,14 @@ void UNiagaraStackViewModel::Initialize(TSharedPtr<FNiagaraSystemViewModel> InSy
 		RootEntry->OnStructureChanged().AddUObject(this, &UNiagaraStackViewModel::EntryStructureChanged);
 		RootEntry->OnDataObjectModified().AddUObject(this, &UNiagaraStackViewModel::EntryDataObjectModified);
 		RootEntry->OnRequestFullRefresh().AddUObject(this, &UNiagaraStackViewModel::EntryRequestFullRefresh);
+		RootEntry->OnRequestFullRefreshDeferred().AddUObject(this, &UNiagaraStackViewModel::EntryRequestFullRefresh);
 		RootEntries.Add(RootEntry);
 	}
 
 	CurrentFocusedSearchMatchIndex = -1;
 	StructureChangedDelegate.Broadcast();
 	bRestartSearch = false;
+	bRefreshPending = false;
 }
 
 void UNiagaraStackViewModel::Finalize()
@@ -95,6 +98,13 @@ void UNiagaraStackViewModel::Tick()
 {
 	if (RootEntry)
 	{
+		if (bRefreshPending)
+		{
+			RootEntry->RefreshChildren();
+			OnSearchTextChanged(CurrentSearchText);
+			bRefreshPending = false;
+		}
+
 		SearchTick();
 	}
 }
@@ -404,4 +414,9 @@ void UNiagaraStackViewModel::EntryRequestFullRefresh()
 	RootEntry->RefreshChildren();
 }
 
+void UNiagaraStackViewModel::EntryRequestFullRefreshDeferred()
+{
+	bRefreshPending = true;
+}
+
 #undef LOCTEXT_NAMESPACE
diff --git a/Engine/Plugins/FX/Niagara/Source/NiagaraEditor/Public/ViewModels/Stack/NiagaraStackEntry.h b/Engine/Plugins/FX/Niagara/Source/NiagaraEditor/Public/ViewModels/Stack/NiagaraStackEntry.h
--- a/Engine/Plugins/FX/Niagara/Source/NiagaraEditor/Public/ViewModels/Stack/NiagaraStackEntry.h
+++ b/Engine/Plugins/FX/Niagara/Source/NiagaraEditor/Public/ViewModels/Stack/NiagaraStackEntry.h
@@ -39,6 +39,7 @@ public:
 	DECLARE_MULTICAST_DELEGATE(FOnStructureChanged);
 	DECLARE_MULTICAST_DELEGATE_OneParam(FOnDataObjectModified, UObject*);
 	DECLARE_MULTICAST_DELEGATE(FOnRequestFullRefresh);
+	DECLARE_MULTICAST_DELEGATE(FOnRequestFullRefreshDeferred);
 	DECLARE_DELEGATE_RetVal_TwoParams(TOptional<FDropResult>, FOnRequestDrop, const UNiagaraStackEntry&, const TArray<UNiagaraStackEntry*>&);
 	DECLARE_DELEGATE_RetVal_OneParam(bool, FOnFilterChild, const UNiagaraStackEntry&);
 	DECLARE_DELEGATE(FStackIssueFixDelegate);
@@ -206,6 +207,8 @@ public:
 
 	FOnRequestFullRefresh& OnRequestFullRefresh();
 
+	FOnRequestFullRefresh& OnRequestFullRefreshDeferred();
+
 	void RefreshChildren();
 
 
@@ -275,6 +278,8 @@ private:
 
 	void ChildRequestFullRefresh();
 
+	void ChildRequestFullRefreshDeferred();
+
 	TOptional<FDropResult> ChildRequestCanDrop(const UNiagaraStackEntry& TargetChild, const TArray<UNiagaraStackEntry*>& DraggedEntries);
 
 	TOptional<FDropResult> ChildRequestDrop(const UNiagaraStackEntry& TargetChild, const TArray<UNiagaraStackEntry*>& DraggedEntries);
@@ -297,6 +302,8 @@ private:
 
 	FOnRequestFullRefresh RequestFullRefreshDelegate;
 
+	FOnRequestFullRefresh RequestFullRefreshDeferredDelegate;
+
 	TArray<FOnFilterChild> ChildFilters;
 
 	UPROPERTY()
diff --git a/Engine/Plugins/FX/Niagara/Source/NiagaraEditor/Public/ViewModels/Stack/NiagaraStackScriptItemGroup.h b/Engine/Plugins/FX/Niagara/Source/NiagaraEditor/Public/ViewModels/Stack/NiagaraStackScriptItemGroup.h
--- a/Engine/Plugins/FX/Niagara/Source/NiagaraEditor/Public/ViewModels/Stack/NiagaraStackScriptItemGroup.h
+++ b/Engine/Plugins/FX/Niagara/Source/NiagaraEditor/Public/ViewModels/Stack/NiagaraStackScriptItemGroup.h
@@ -34,6 +34,8 @@ public:
 protected:
 	virtual void RefreshChildrenInternal(const TArray<UNiagaraStackEntry*>& CurrentChildren, TArray<UNiagaraStackEntry*>& NewChildren, TArray<FStackIssue>& NewIssues) override;
 
+	virtual void FinalizeInternal() override;
+
 	virtual TOptional<FDropResult> ChildRequestCanDropInternal(const UNiagaraStackEntry& TargetChild, const TArray<UNiagaraStackEntry*>& DraggedEntries) override;
 
 	virtual TOptional<FDropResult> ChildRequestDropInternal(const UNiagaraStackEntry& TargetChild, const TArray<UNiagaraStackEntry*>& DraggedEntries) override;
@@ -43,6 +45,8 @@ private:
 
 	void ChildModifiedGroupItems();
 
+	void OnScriptGraphChanged(const struct FEdGraphEditAction& InAction);
+
 protected:
 	TWeakPtr<FNiagaraScriptViewModel> ScriptViewModel;
 	void RefreshIssues(TArray<FStackIssue>& NewIssues);
@@ -55,5 +59,7 @@ private:
 	FGuid ScriptUsageId;
 	bool bIsValidForOutput;
 
+	FDelegateHandle OnGraphChangedHandle;
+
 	TMap<FObjectKey, UNiagaraStackModuleItem*> StackSpacerToModuleItemMap;
 };
diff --git a/Engine/Plugins/FX/Niagara/Source/NiagaraEditor/Public/ViewModels/Stack/NiagaraStackViewModel.h b/Engine/Plugins/FX/Niagara/Source/NiagaraEditor/Public/ViewModels/Stack/NiagaraStackViewModel.h
--- a/Engine/Plugins/FX/Niagara/Source/NiagaraEditor/Public/ViewModels/Stack/NiagaraStackViewModel.h
+++ b/Engine/Plugins/FX/Niagara/Source/NiagaraEditor/Public/ViewModels/Stack/NiagaraStackViewModel.h
@@ -101,6 +101,7 @@ private:
 	void EntryStructureChanged();
 	void EntryDataObjectModified(UObject* ChangedObject);
 	void EntryRequestFullRefresh();
+	void EntryRequestFullRefreshDeferred();
 	void OnSystemCompiled();
 	void OnEmitterCompiled();
 	/** Called by the tick function to perform partial search */
@@ -129,4 +130,5 @@ private:
 	TArray<FSearchResult> CurrentSearchResults;
 	static const double MaxSearchTime;
 	bool bRestartSearch;
+	bool bRefreshPending;
 };
\ No newline at end of file
