Niagara - Prevent data loss when deleting burst keys through the system timeline and undoing that operation.

diff --git a/Engine/Plugins/FX/Niagara/Source/NiagaraEditor/Private/NiagaraScriptMergeManager.cpp b/Engine/Plugins/FX/Niagara/Source/NiagaraEditor/Private/NiagaraScriptMergeManager.cpp
--- a/Engine/Plugins/FX/Niagara/Source/NiagaraEditor/Private/NiagaraScriptMergeManager.cpp
+++ b/Engine/Plugins/FX/Niagara/Source/NiagaraEditor/Private/NiagaraScriptMergeManager.cpp
@@ -1854,7 +1854,7 @@ FNiagaraScriptMergeManager::FApplyDiffResults FNiagaraScriptMergeManager::ApplyS
 	// Apply the graph actions.
 	for (TSharedRef<FNiagaraStackFunctionMergeAdapter> RemoveModule : RemoveModules)
 	{
-		bool bRemoveResults = FNiagaraStackGraphUtilities::RemoveModuleFromStack(*RemoveModule->GetFunctionCallNode());
+		bool bRemoveResults = FNiagaraStackGraphUtilities::RemoveModuleFromStack(*BaseScriptStackAdapter->GetScript(), *RemoveModule->GetFunctionCallNode());
 		if (bRemoveResults == false)
 		{
 			Results.bSucceeded = false;
diff --git a/Engine/Plugins/FX/Niagara/Source/NiagaraEditor/Private/Sequencer/NiagaraSequence/Sections/MovieSceneNiagaraEmitterSection.cpp b/Engine/Plugins/FX/Niagara/Source/NiagaraEditor/Private/Sequencer/NiagaraSequence/Sections/MovieSceneNiagaraEmitterSection.cpp
--- a/Engine/Plugins/FX/Niagara/Source/NiagaraEditor/Private/Sequencer/NiagaraSequence/Sections/MovieSceneNiagaraEmitterSection.cpp
+++ b/Engine/Plugins/FX/Niagara/Source/NiagaraEditor/Private/Sequencer/NiagaraSequence/Sections/MovieSceneNiagaraEmitterSection.cpp
@@ -384,7 +384,7 @@ void UMovieSceneNiagaraEmitterSection::UpdateKeyModulesFromSection(FChannelAndMo
 	{
 		ChannelAndModules.ModulesAndBinders.RemoveAll(
 			[=](FModuleAndBinders& ModuleAndBinders) { return ModuleAndBinders.Module == ModuleWithMissingKey; });
-		FNiagaraStackGraphUtilities::RemoveModuleFromStack(*ModuleWithMissingKey);
+		FNiagaraStackGraphUtilities::RemoveModuleFromStack(GetSystemViewModel().GetSystem(), GetEmitterHandleViewModel()->GetId(), *ModuleWithMissingKey);
 	}
 
 	// Create new modules for new keys.
diff --git a/Engine/Plugins/FX/Niagara/Source/NiagaraEditor/Private/ViewModels/Stack/NiagaraStackGraphUtilities.cpp b/Engine/Plugins/FX/Niagara/Source/NiagaraEditor/Private/ViewModels/Stack/NiagaraStackGraphUtilities.cpp
--- a/Engine/Plugins/FX/Niagara/Source/NiagaraEditor/Private/ViewModels/Stack/NiagaraStackGraphUtilities.cpp
+++ b/Engine/Plugins/FX/Niagara/Source/NiagaraEditor/Private/ViewModels/Stack/NiagaraStackGraphUtilities.cpp
@@ -820,13 +820,32 @@ void FNiagaraStackGraphUtilities::SetCustomExpressionForFunctionInput(UEdGraphPi
 	}
 }
 
-bool FNiagaraStackGraphUtilities::RemoveModuleFromStack(UNiagaraNodeFunctionCall& ModuleNode)
+bool FNiagaraStackGraphUtilities::RemoveModuleFromStack(UNiagaraSystem& OwningSystem, FGuid OwningEmitterId, UNiagaraNodeFunctionCall& ModuleNode)
 {
 	TArray<TWeakObjectPtr<UNiagaraNodeInput>> RemovedNodes;
-	return RemoveModuleFromStack(ModuleNode, RemovedNodes);
+	return RemoveModuleFromStack(OwningSystem, OwningEmitterId, ModuleNode, RemovedNodes);
 }
 
-bool FNiagaraStackGraphUtilities::RemoveModuleFromStack(UNiagaraNodeFunctionCall& ModuleNode, TArray<TWeakObjectPtr<UNiagaraNodeInput>>& OutRemovedNodes)
+bool FNiagaraStackGraphUtilities::RemoveModuleFromStack(UNiagaraSystem& OwningSystem, FGuid OwningEmitterId, UNiagaraNodeFunctionCall& ModuleNode, TArray<TWeakObjectPtr<UNiagaraNodeInput>>& OutRemovedInputNodes)
+{
+	// Find the owning script so it can be modified as part of the transaction so that rapid iteration parameters values are retained upon undo.
+	UNiagaraNodeOutput* OutputNode = GetEmitterOutputNodeForStackNode(ModuleNode);
+	checkf(OutputNode != nullptr, TEXT("Invalid Stack - Output node could not be found for module"));
+
+	UNiagaraScript* OwningScript = FNiagaraEditorUtilities::GetScriptFromSystem(
+		OwningSystem, OwningEmitterId, OutputNode->GetUsage(), OutputNode->GetUsageId());
+	checkf(OwningScript != nullptr, TEXT("Invalid Stack - Owning script could not be found for module"));
+
+	return RemoveModuleFromStack(*OwningScript, ModuleNode, OutRemovedInputNodes);
+}
+
+bool FNiagaraStackGraphUtilities::RemoveModuleFromStack(UNiagaraScript& OwningScript, UNiagaraNodeFunctionCall& ModuleNode)
+{
+	TArray<TWeakObjectPtr<UNiagaraNodeInput>> RemovedNodes;
+	return RemoveModuleFromStack(OwningScript, ModuleNode, RemovedNodes);
+}
+
+bool FNiagaraStackGraphUtilities::RemoveModuleFromStack(UNiagaraScript& OwningScript, UNiagaraNodeFunctionCall& ModuleNode, TArray<TWeakObjectPtr<UNiagaraNodeInput>>& OutRemovedInputNodes)
 {
 	TArray<FNiagaraStackGraphUtilities::FStackNodeGroup> StackNodeGroups;
 	FNiagaraStackGraphUtilities::GetStackNodeGroups(ModuleNode, StackNodeGroups);
@@ -837,6 +856,8 @@ bool FNiagaraStackGraphUtilities::RemoveModuleFromStack(UNiagaraNodeFunctionCall
 		return false;
 	}
 
+	OwningScript.Modify();
+
 	// Disconnect the group from the stack first to make collecting the nodes to remove easier.
 	FNiagaraStackGraphUtilities::DisconnectStackNodeGroup(StackNodeGroups[ModuleStackIndex], StackNodeGroups[ModuleStackIndex - 1], StackNodeGroups[ModuleStackIndex + 1]);
 
@@ -874,7 +895,7 @@ bool FNiagaraStackGraphUtilities::RemoveModuleFromStack(UNiagaraNodeFunctionCall
 		UNiagaraNodeInput* InputNode = Cast<UNiagaraNodeInput>(NodeToRemove);
 		if (InputNode != nullptr)
 		{
-			OutRemovedNodes.Add(InputNode);
+			OutRemovedInputNodes.Add(InputNode);
 		}
 	}
 
diff --git a/Engine/Plugins/FX/Niagara/Source/NiagaraEditor/Private/ViewModels/Stack/NiagaraStackModuleItem.cpp b/Engine/Plugins/FX/Niagara/Source/NiagaraEditor/Private/ViewModels/Stack/NiagaraStackModuleItem.cpp
--- a/Engine/Plugins/FX/Niagara/Source/NiagaraEditor/Private/ViewModels/Stack/NiagaraStackModuleItem.cpp
+++ b/Engine/Plugins/FX/Niagara/Source/NiagaraEditor/Private/ViewModels/Stack/NiagaraStackModuleItem.cpp
@@ -564,21 +564,13 @@ void UNiagaraStackModuleItem::Delete()
 {
 	checkf(CanMoveAndDelete(), TEXT("This module can't be deleted"));
 
-	UNiagaraNodeOutput* OutputNode = FNiagaraStackGraphUtilities::GetEmitterOutputNodeForStackNode(*FunctionCallNode);
-	const FNiagaraEmitterHandle* EmitterHandle = FNiagaraEditorUtilities::GetEmitterHandleForEmitter(GetSystemViewModel()->GetSystem(), *GetEmitterViewModel()->GetEmitter());
-	checkf(OutputNode != nullptr && EmitterHandle != nullptr, TEXT("Invalid Stack - Output node or emitter handle could not be found for module"));
-
-	UNiagaraScript* OwningScript = FNiagaraEditorUtilities::GetScriptFromSystem(
-		GetSystemViewModel()->GetSystem(), EmitterHandle->GetId(), OutputNode->GetUsage(), OutputNode->GetUsageId());
-	checkf(OwningScript != nullptr, TEXT("Invalid Stack - Owning script could not be found for module"));
-
 	FScopedTransaction ScopedTransaction(LOCTEXT("RemoveAModuleFromTheStack", "Remove a module from the stack"));
 
-	// Modify the script so that any rapid iteration parameters in use are recorded.
-	OwningScript->Modify();
+	const FNiagaraEmitterHandle* EmitterHandle = FNiagaraEditorUtilities::GetEmitterHandleForEmitter(GetSystemViewModel()->GetSystem(), *GetEmitterViewModel()->GetEmitter());
+	checkf(EmitterHandle != nullptr, TEXT("Invalid Stack - Emitter handle could not be found for module"));
 
 	TArray<TWeakObjectPtr<UNiagaraNodeInput>> RemovedNodes;
-	bool bRemoved = FNiagaraStackGraphUtilities::RemoveModuleFromStack(*FunctionCallNode, RemovedNodes);
+	bool bRemoved = FNiagaraStackGraphUtilities::RemoveModuleFromStack(GetSystemViewModel()->GetSystem(), EmitterHandle->GetId(), *FunctionCallNode, RemovedNodes);
 	if (bRemoved)
 	{
 		UNiagaraGraph* Graph = FunctionCallNode->GetNiagaraGraph();
diff --git a/Engine/Plugins/FX/Niagara/Source/NiagaraEditor/Public/ViewModels/Stack/NiagaraStackGraphUtilities.h b/Engine/Plugins/FX/Niagara/Source/NiagaraEditor/Public/ViewModels/Stack/NiagaraStackGraphUtilities.h
--- a/Engine/Plugins/FX/Niagara/Source/NiagaraEditor/Public/ViewModels/Stack/NiagaraStackGraphUtilities.h
+++ b/Engine/Plugins/FX/Niagara/Source/NiagaraEditor/Public/ViewModels/Stack/NiagaraStackGraphUtilities.h
@@ -94,9 +94,13 @@ namespace FNiagaraStackGraphUtilities
 
 	void SetCustomExpressionForFunctionInput(UEdGraphPin& OverridePin, UNiagaraNodeCustomHlsl*& OutDynamicInputFunctionCall, const FGuid& NewNodePersistentId = FGuid());
 
-	bool RemoveModuleFromStack(UNiagaraNodeFunctionCall& ModuleNode);
+	bool RemoveModuleFromStack(UNiagaraSystem& OwningSystem, FGuid OwningEmitterId, UNiagaraNodeFunctionCall& ModuleNode);
 
-	bool RemoveModuleFromStack(UNiagaraNodeFunctionCall& ModuleNode, TArray<TWeakObjectPtr<UNiagaraNodeInput>>& OutRemovedNodes);
+	bool RemoveModuleFromStack(UNiagaraSystem& OwningSystem, FGuid OwningEmitterId, UNiagaraNodeFunctionCall& ModuleNode, TArray<TWeakObjectPtr<UNiagaraNodeInput>>& OutRemovedInputNodes);
+
+	bool RemoveModuleFromStack(UNiagaraScript& OwningScript, UNiagaraNodeFunctionCall& ModuleNode);
+
+	bool RemoveModuleFromStack(UNiagaraScript& OwningScript, UNiagaraNodeFunctionCall& ModuleNode, TArray<TWeakObjectPtr<UNiagaraNodeInput>>& OutRemovedInputNodes);
 
 	UNiagaraNodeFunctionCall* AddScriptModuleToStack(FAssetData ModuleScriptAsset, UNiagaraNodeOutput& TargetOutputNode, int32 TargetIndex = INDEX_NONE);
 	
