Fix for system sims + a new auto test for system simulations.

diff --git a/Engine/Plugins/FX/Niagara/Source/Niagara/Private/NiagaraSystemSimulation.cpp b/Engine/Plugins/FX/Niagara/Source/Niagara/Private/NiagaraSystemSimulation.cpp
--- a/Engine/Plugins/FX/Niagara/Source/Niagara/Private/NiagaraSystemSimulation.cpp
+++ b/Engine/Plugins/FX/Niagara/Source/Niagara/Private/NiagaraSystemSimulation.cpp
@@ -298,9 +298,9 @@ bool FNiagaraSystemSimulation::Tick(float DeltaSeconds)
 
 				Inst->TickDataInterfaces(DeltaSeconds, false);
 				Inst->SetPendingSpawn(false);
-				Inst->SystemInstanceIndex = SystemInstances.Add(Inst);
 				if (!Inst->IsComplete())
 				{
+					Inst->SystemInstanceIndex = SystemInstances.Add(Inst);
 					++SpawnNum;
 				}
 				else
@@ -402,7 +402,7 @@ bool FNiagaraSystemSimulation::Tick(float DeltaSeconds)
 
 			//Run Spawn
 			SpawnExecContext.Tick(SoloSystemInstance);//We can't require a specific instance here as these are for all instances.
-			DataSetExecInfos[0] = FNiagaraDataSetExecutionInfo(&DataSet, OrigNum, false, true);
+			DataSetExecInfos[0] = FNiagaraDataSetExecutionInfo(&DataSet, OrigNum, false, false);
 			DataSetExecInfos[1] = FNiagaraDataSetExecutionInfo(&SpawnInstanceParameterDataSet, OrigNum, false, false);
 			SpawnExecContext.Execute(SpawnNum, DataSetExecInfos);
 
@@ -437,7 +437,7 @@ bool FNiagaraSystemSimulation::Tick(float DeltaSeconds)
 
 			//Run update.
 			UpdateExecContext.Tick(SystemInstances[0]);
-			DataSetExecInfos[0] = FNiagaraDataSetExecutionInfo(&DataSet, 0, false, true);
+			DataSetExecInfos[0] = FNiagaraDataSetExecutionInfo(&DataSet, 0, false, false);
 			DataSetExecInfos[1] = FNiagaraDataSetExecutionInfo(&UpdateInstanceParameterDataSet, 0, false, false);
 
 // 			if (GbDumpSystemData || System->bDumpDebugSystemInfo)
@@ -465,7 +465,7 @@ bool FNiagaraSystemSimulation::Tick(float DeltaSeconds)
 
 				//Run update.
 				UpdateExecContext.Tick(SystemInstances[0]);
-				DataSetExecInfos[0] = FNiagaraDataSetExecutionInfo(&DataSet, OrigNum, false, true);
+				DataSetExecInfos[0] = FNiagaraDataSetExecutionInfo(&DataSet, OrigNum, false, false);
 				DataSetExecInfos[1] = FNiagaraDataSetExecutionInfo(&UpdateInstanceParameterDataSet, OrigNum, false, false);
 
 				UpdateExecContext.Parameters.SetParameterValue(0.0001f, SYS_PARAM_ENGINE_DELTA_TIME);
