Fixed issue causing cross talk between system simulations.

diff --git a/Engine/Plugins/FX/Niagara/Source/Niagara/Private/NiagaraDataSet.cpp b/Engine/Plugins/FX/Niagara/Source/Niagara/Private/NiagaraDataSet.cpp
--- a/Engine/Plugins/FX/Niagara/Source/Niagara/Private/NiagaraDataSet.cpp
+++ b/Engine/Plugins/FX/Niagara/Source/Niagara/Private/NiagaraDataSet.cpp
@@ -251,10 +251,12 @@ void FNiagaraDataBuffer::Allocate(uint32 InNumInstances, bool bMaintainExisting)
 		DEC_MEMORY_STAT_BY(STAT_NiagaraParticleMemory, FloatData.Max() + Int32Data.Max());
 
 		int32 OldFloatStride = FloatStride;
+		TArray<uint8> OldFloatData = FloatData;
 		FloatStride = GetSafeComponentBufferSize(NumInstancesAllocated * sizeof(float));
 		FloatData.SetNum(FloatStride * Owner->GetNumFloatComponents(), false);
 
 		int32 OldInt32Stride = Int32Stride;
+		TArray<uint8> OldIntData = Int32Data;
 		Int32Stride = GetSafeComponentBufferSize(NumInstancesAllocated * sizeof(int32));
 		Int32Data.SetNum(Int32Stride * Owner->GetNumInt32Components(), false);
 
@@ -265,20 +267,20 @@ void FNiagaraDataBuffer::Allocate(uint32 InNumInstances, bool bMaintainExisting)
 		{
 			if (FloatStride != OldFloatStride && FloatStride > 0 && OldFloatStride > 0)
 			{
-				for (uint32 CompIdx = Owner->TotalFloatComponents-1; CompIdx > 0; --CompIdx)
+				for (int32 CompIdx = (int32)Owner->TotalFloatComponents-1; CompIdx >= 0; --CompIdx)
 				{
-					uint8* Src = FloatData.GetData() + OldFloatStride * CompIdx;
+					uint8* Src = OldFloatData.GetData() + OldFloatStride * CompIdx;
 					uint8* Dst = FloatData.GetData() + FloatStride * CompIdx;
 					FMemory::Memcpy(Dst, Src, OldFloatStride);
 				}
 			}
 			if (Int32Stride != OldInt32Stride && Int32Stride > 0 && OldInt32Stride > 0)
 			{
-				for (uint32 CompIdx = Owner->TotalInt32Components - 1; CompIdx > 0; --CompIdx)
+				for (int32 CompIdx = (int32)Owner->TotalInt32Components - 1; CompIdx >= 0; --CompIdx)
 				{
-					uint8* Src = Int32Data.GetData() + OldInt32Stride * CompIdx;
+					uint8* Src = OldIntData.GetData() + OldInt32Stride * CompIdx;
 					uint8* Dst = Int32Data.GetData() + Int32Stride * CompIdx;
-					FMemory::Memcpy(Dst, Src, OldFloatStride);
+					FMemory::Memcpy(Dst, Src, OldInt32Stride);
 				}
 			}
 		}
