New auto for setting user parameters on many system instances.
New debugging and nan checking code to help with UE-60317

diff --git a/Engine/Plugins/FX/Niagara/Source/Niagara/Classes/NiagaraDataSet.h b/Engine/Plugins/FX/Niagara/Source/Niagara/Classes/NiagaraDataSet.h
--- a/Engine/Plugins/FX/Niagara/Source/Niagara/Classes/NiagaraDataSet.h
+++ b/Engine/Plugins/FX/Niagara/Source/Niagara/Classes/NiagaraDataSet.h
@@ -40,8 +40,8 @@ public:
 	void AllocateGPU(uint32 InNumInstances, FRHICommandList &RHICmdList);
 	void SwapInstances(uint32 OldIndex, uint32 NewIndex);
 	void KillInstance(uint32 InstanceIdx);
-	void CopyTo(FNiagaraDataBuffer& DestBuffer, int32 StartIdx, int32 NumInstances);
-	void CopyTo(FNiagaraDataBuffer& DestBuffer);
+	void CopyTo(FNiagaraDataBuffer& DestBuffer, int32 StartIdx, int32 NumInstances)const;
+	void CopyTo(FNiagaraDataBuffer& DestBuffer)const;
 
 	const TArray<uint8>& GetFloatBuffer()const { return FloatData; }
 	const TArray<uint8>& GetInt32Buffer()const { return Int32Data; }
@@ -76,6 +76,16 @@ public:
 		return (int32*)GetComponentPtrInt32(ComponentIdx) + InstanceIdx;
 	}
 
+	FORCEINLINE float* GetInstancePtrFloat(uint32 ComponentIdx, uint32 InstanceIdx)const
+	{
+		return (float*)GetComponentPtrFloat(ComponentIdx) + InstanceIdx;
+	}
+
+	FORCEINLINE int32* GetInstancePtrInt32(uint32 ComponentIdx, uint32 InstanceIdx)const
+	{
+		return (int32*)GetComponentPtrInt32(ComponentIdx) + InstanceIdx;
+	}
+
 	uint32 GetNumInstances()const { return NumInstances; }
 	uint32 GetNumInstancesAllocated()const { return NumInstancesAllocated; }
 
@@ -112,6 +122,8 @@ public:
 	FORCEINLINE const FNiagaraDataSet* GetOwner()const { return Owner; }
 
 	int32 TransferInstance(FNiagaraDataBuffer& SourceBuffer, int32 InstanceIndex);
+
+	bool CheckForNaNs()const;
 private:
 
 	FORCEINLINE int32 GetSafeComponentBufferSize(int32 RequiredSize) const
@@ -284,6 +296,10 @@ public:
 		CheckCorrectThread();
 		CurrData().Allocate(NumInstances, bMaintainExisting);
 
+#if NIAGARA_NAN_CHECKING
+		CheckForNaNs();
+#endif
+
 		if (bNeedsPersistentIDs)
 		{
 			int32 NumUsedIDs = MaxUsedID + 1;
@@ -347,6 +363,9 @@ public:
 	
 	FORCEINLINE void Tick()
 	{
+#if NIAGARA_NAN_CHECKING
+		CheckForNaNs();
+#endif
 		SwapBuffers();
 	}
 	
@@ -509,10 +528,24 @@ public:
 		return false;
 	}
 
-	void Dump(bool bCurr, int32 StartIdx = 0, int32 NumInstances = INDEX_NONE);
-	void Dump(FNiagaraDataSet& Other, bool bCurr, int32 StartIdx = 0, int32 NumInstances = INDEX_NONE);
+	void Dump(bool bCurr, int32 StartIdx = 0, int32 NumInstances = INDEX_NONE)const;
+	void Dump(FNiagaraDataSet& Other, bool bCurr, int32 StartIdx = 0, int32 NumInstances = INDEX_NONE)const;
 	FORCEINLINE const TArray<FNiagaraVariable> &GetVariables() { return Variables; }
 
+	void CheckForNaNs()const
+	{
+		if (CurrData().CheckForNaNs())
+		{
+			Dump(true);
+			ensureAlwaysMsgf(false, TEXT("NiagaraDataSet contains NaNs!"));
+		}
+
+		if (PrevData().CheckForNaNs())
+		{
+			Dump(false);
+			ensureAlwaysMsgf(false, TEXT("NiagaraDataSet contains NaNs!"));
+		}
+	}
 
 	// Data set index buffer management
 	// these buffers hold the number of instances for the buffers; the first five uint32s are the DrawIndirect parameters for rendering of the main particle data set
@@ -1625,5 +1658,79 @@ private:
 	TArray<FNiagaraVariable*> Variables;
 	TArray<const FNiagaraVariableLayoutInfo*> VarLayouts;
 
+	uint32 CurrIdx;
+};
+
+/**
+Iterator that will pull or push data between a DataSet and some FNiagaraVariables it contains.
+Super slow. Don't use at runtime.
+*/
+struct FNiagaraDataSetVariableIteratorConst
+{
+	FNiagaraDataSetVariableIteratorConst(const FNiagaraDataSet& InDataSet, uint32 StartIdx = 0, bool bCurrBuffer = true)
+		: DataSet(InDataSet)
+		, DataBuffer(bCurrBuffer ? DataSet.CurrData() : DataSet.PrevData())
+		, CurrIdx(StartIdx)
+	{
+	}
+
+	void Get()
+	{
+		for (int32 VarIdx = 0; VarIdx < Variables.Num(); ++VarIdx)
+		{
+			FNiagaraVariable* Var = Variables[VarIdx];
+			const FNiagaraVariableLayoutInfo* Layout = VarLayouts[VarIdx];
+			check(Var && Layout);
+			uint8* ValuePtr = Var->GetData();
+
+			for (uint32 CompIdx = 0; CompIdx < Layout->GetNumFloatComponents(); ++CompIdx)
+			{
+				uint32 CompBufferOffset = Layout->FloatComponentStart + CompIdx;
+				float* Src = DataBuffer.GetInstancePtrFloat(CompBufferOffset, CurrIdx);
+				float* Dst = (float*)(ValuePtr + Layout->LayoutInfo.FloatComponentByteOffsets[CompIdx]);
+				*Dst = *Src;
+			}
+
+			for (uint32 CompIdx = 0; CompIdx < Layout->GetNumInt32Components(); ++CompIdx)
+			{
+				uint32 CompBufferOffset = Layout->Int32ComponentStart + CompIdx;
+				int32* Src = DataBuffer.GetInstancePtrInt32(CompBufferOffset, CurrIdx);
+				int32* Dst = (int32*)(ValuePtr + Layout->LayoutInfo.Int32ComponentByteOffsets[CompIdx]);
+				*Dst = *Src;
+			}
+		}
+	}
+	
+	void Advance() { ++CurrIdx; }
+
+	bool IsValid()const
+	{
+		return CurrIdx < DataBuffer.GetNumInstances();
+	}
+
+	uint32 GetCurrIndex()const { return CurrIdx; }
+
+	void AddVariable(FNiagaraVariable* InVar)
+	{
+		check(InVar);
+		Variables.AddUnique(InVar);
+		VarLayouts.AddUnique(DataSet.GetVariableLayout(*InVar));
+		InVar->AllocateData();
+	}
+
+	void AddVariables(TArray<FNiagaraVariable>& Vars)
+	{
+		for (FNiagaraVariable& Var : Vars)
+		{
+			AddVariable(&Var);
+		}
+	}
+private:
+
+	const FNiagaraDataSet & DataSet;
+	const FNiagaraDataBuffer& DataBuffer;
+	TArray<FNiagaraVariable*> Variables;
+	TArray<const FNiagaraVariableLayoutInfo*> VarLayouts;
+
 	uint32 CurrIdx;
 };
\ No newline at end of file
diff --git a/Engine/Plugins/FX/Niagara/Source/Niagara/Classes/NiagaraEmitterInstance.h b/Engine/Plugins/FX/Niagara/Source/Niagara/Classes/NiagaraEmitterInstance.h
--- a/Engine/Plugins/FX/Niagara/Source/Niagara/Classes/NiagaraEmitterInstance.h
+++ b/Engine/Plugins/FX/Niagara/Source/Niagara/Classes/NiagaraEmitterInstance.h
@@ -102,6 +102,7 @@ public:
 
 	NIAGARA_API bool IsReadyToRun() const;
 
+	void Dump()const;
 private:
 
 	void CheckForErrors();
diff --git a/Engine/Plugins/FX/Niagara/Source/Niagara/Private/NiagaraDataSet.cpp b/Engine/Plugins/FX/Niagara/Source/Niagara/Private/NiagaraDataSet.cpp
--- a/Engine/Plugins/FX/Niagara/Source/Niagara/Private/NiagaraDataSet.cpp
+++ b/Engine/Plugins/FX/Niagara/Source/Niagara/Private/NiagaraDataSet.cpp
@@ -112,13 +112,13 @@ FNiagaraDataSet::~FNiagaraDataSet()
 
 }
 
-void FNiagaraDataSet::Dump(FNiagaraDataSet& Other, bool bCurr, int32 StartIdx , int32 NumInstances)
+void FNiagaraDataSet::Dump(FNiagaraDataSet& Other, bool bCurr, int32 StartIdx , int32 NumInstances)const
 {
 	Other.Reset();
 	Other.Variables = Variables;
 	Other.VariableLayouts = VariableLayouts;
 
-	FNiagaraDataBuffer& DataBuffer = bCurr ? CurrData() : PrevData();
+	const FNiagaraDataBuffer& DataBuffer = bCurr ? CurrData() : PrevData();
 	FNiagaraDataBuffer& OtherDataBuffer = Other.CurrData();
 
 	
@@ -132,11 +132,11 @@ void FNiagaraDataSet::Dump(FNiagaraDataSet& Other, bool bCurr, int32 StartIdx ,
 }
 
 
-void FNiagaraDataSet::Dump(bool bCurr, int32 StartIdx, int32 NumInstances)
+void FNiagaraDataSet::Dump(bool bCurr, int32 StartIdx, int32 NumInstances)const
 {
 	TArray<FNiagaraVariable> Vars(Variables);
 
-	FNiagaraDataSetVariableIterator Itr(*this, StartIdx, bCurr);
+	FNiagaraDataSetVariableIteratorConst Itr(*this, StartIdx, bCurr);
 	Itr.AddVariables(Vars);
 
 	if (NumInstances == INDEX_NONE)
@@ -240,6 +240,22 @@ int32 FNiagaraDataBuffer::TransferInstance(FNiagaraDataBuffer& SourceBuffer, int
 	return INDEX_NONE;
 }
 
+bool FNiagaraDataBuffer::CheckForNaNs()const
+{
+	bool bContainsNaNs = false;
+	int32 NumFloatComponents = Owner->GetNumFloatComponents();
+	for (int32 CompIdx = 0; CompIdx < NumFloatComponents && !bContainsNaNs; ++CompIdx)
+	{
+		for (int32 InstIdx = 0; InstIdx < (int32)NumInstances && !bContainsNaNs; ++InstIdx)
+		{
+			float Val = *GetInstancePtrFloat(CompIdx, InstIdx);
+			bContainsNaNs = FMath::IsNaN(Val) || !FMath::IsFinite(Val);
+		}
+	}
+
+	return bContainsNaNs;
+}
+
 void FNiagaraDataBuffer::Allocate(uint32 InNumInstances, bool bMaintainExisting)
 {
 	check(Owner);
@@ -370,9 +386,13 @@ void FNiagaraDataBuffer::KillInstance(uint32 InstanceIdx)
 		int32* Dst = GetInstancePtrInt32(CompIdx, InstanceIdx);
 		*Dst = *Src;
 	}
+
+#if NIAGARA_NAN_CHECKING
+	CheckForNaNs();
+#endif
 }
 
-void FNiagaraDataBuffer::CopyTo(FNiagaraDataBuffer& DestBuffer, int32 InStartIdx, int32 InNumInstances)
+void FNiagaraDataBuffer::CopyTo(FNiagaraDataBuffer& DestBuffer, int32 InStartIdx, int32 InNumInstances)const
 {
 	if (InStartIdx < 0 || (uint32)InStartIdx > NumInstances)
 	{
@@ -392,8 +412,8 @@ void FNiagaraDataBuffer::CopyTo(FNiagaraDataBuffer& DestBuffer, int32 InStartIdx
 
 		for (uint32 CompIdx = 0; CompIdx < Owner->TotalFloatComponents; ++CompIdx)
 		{
-			float* SrcStart = GetInstancePtrFloat(CompIdx, InStartIdx);
-			float* SrcEnd = GetInstancePtrFloat(CompIdx, InStartIdx + InNumInstances);
+			const float* SrcStart = GetInstancePtrFloat(CompIdx, InStartIdx);
+			const float* SrcEnd = GetInstancePtrFloat(CompIdx, InStartIdx + InNumInstances);
 			float* Dst = DestBuffer.GetInstancePtrFloat(CompIdx, 0);
 			size_t Count = SrcEnd - SrcStart;
 			FMemory::Memcpy(Dst, SrcStart, Count*sizeof(float));
@@ -408,8 +428,8 @@ void FNiagaraDataBuffer::CopyTo(FNiagaraDataBuffer& DestBuffer, int32 InStartIdx
 		}
 		for (uint32 CompIdx = 0; CompIdx < Owner->TotalInt32Components; ++CompIdx)
 		{
-			int32* SrcStart = GetInstancePtrInt32(CompIdx, InStartIdx);
-			int32* SrcEnd = GetInstancePtrInt32(CompIdx, InStartIdx + InNumInstances);
+			const int32* SrcStart = GetInstancePtrInt32(CompIdx, InStartIdx);
+			const int32* SrcEnd = GetInstancePtrInt32(CompIdx, InStartIdx + InNumInstances);
 			int32* Dst = DestBuffer.GetInstancePtrInt32(CompIdx, 0);
 			size_t Count = SrcEnd - SrcStart;
 			FMemory::Memcpy(Dst, SrcStart, Count * sizeof(int32));
@@ -426,7 +446,7 @@ void FNiagaraDataBuffer::CopyTo(FNiagaraDataBuffer& DestBuffer, int32 InStartIdx
 	}
 }
 
-void FNiagaraDataBuffer::CopyTo(FNiagaraDataBuffer& DestBuffer)
+void FNiagaraDataBuffer::CopyTo(FNiagaraDataBuffer& DestBuffer)const
 {
 	DestBuffer.FloatStride = FloatStride;
 	DestBuffer.FloatData = FloatData;
diff --git a/Engine/Plugins/FX/Niagara/Source/Niagara/Private/NiagaraEmitterInstance.cpp b/Engine/Plugins/FX/Niagara/Source/Niagara/Private/NiagaraEmitterInstance.cpp
--- a/Engine/Plugins/FX/Niagara/Source/Niagara/Private/NiagaraEmitterInstance.cpp
+++ b/Engine/Plugins/FX/Niagara/Source/Niagara/Private/NiagaraEmitterInstance.cpp
@@ -103,6 +103,20 @@ bool FNiagaraEmitterInstance::IsReadyToRun() const
 	return true;
 }
 
+void FNiagaraEmitterInstance::Dump()const
+{
+	UE_LOG(LogNiagara, Log, TEXT("==  %s ========"), *CachedEmitter->GetUniqueEmitterName());
+	UE_LOG(LogNiagara, Log, TEXT(".................Spawn................."));
+	SpawnExecContext.Parameters.DumpParameters(true);
+	UE_LOG(LogNiagara, Log, TEXT(".................Update................."));
+	UpdateExecContext.Parameters.DumpParameters(true);
+	UE_LOG(LogNiagara, Log, TEXT("................. %s Combined Parameters ................."), TEXT("GPU Script"));
+	GPUExecContext.CombinedParamStore.DumpParameters();
+	UE_LOG(LogNiagara, Log, TEXT("................. Particles ................."));
+	ParticleDataSet->Dump(false);
+	ParticleDataSet->Dump(true);
+}
+
 void FNiagaraEmitterInstance::Init(int32 InEmitterIdx, FName InSystemInstanceName)
 {
 	check(ParticleDataSet);
@@ -617,6 +631,7 @@ FBox FNiagaraEmitterInstance::CalculateDynamicBounds()
 				UE_LOG(LogNiagara, Warning, TEXT("Particle position data contains NaNs. Likely a divide by zero somewhere in your modules. Emitter \"%s\" in System \"%s\""),
 					*CachedEmitter->GetName(), *ParentSystemInstance->GetSystem()->GetName());
 				bEncounteredNaNs = true;
+				ParentSystemInstance->Dump();
 			}
 #endif
 		}
diff --git a/Engine/Plugins/FX/Niagara/Source/Niagara/Private/NiagaraParameterStore.cpp b/Engine/Plugins/FX/Niagara/Source/Niagara/Private/NiagaraParameterStore.cpp
--- a/Engine/Plugins/FX/Niagara/Source/Niagara/Private/NiagaraParameterStore.cpp
+++ b/Engine/Plugins/FX/Niagara/Source/Niagara/Private/NiagaraParameterStore.cpp
@@ -138,9 +138,54 @@ bool FNiagaraParameterStore::VerifyBinding(const FNiagaraParameterStore* DestSto
 #endif
 }
 
+void FNiagaraParameterStore::CheckForNaNs()const
+{
+	for (const TPair<FNiagaraVariable, int32>& ParamOffset : ParameterOffsets)
+	{
+		const FNiagaraVariable& Var = ParamOffset.Key;
+		int32 Offset = ParamOffset.Value;
+		bool bContainsNans = false;
+		if (Var.GetType() == FNiagaraTypeDefinition::GetFloatDef())
+		{
+			float Val = *(float*)GetParameterData(Offset);
+			bContainsNans = FMath::IsNaN(Val) || !FMath::IsFinite(Val);
+		}
+		else if (Var.GetType() == FNiagaraTypeDefinition::GetVec2Def())
+		{
+			FVector2D Val = *(FVector2D*)GetParameterData(Offset);
+			bContainsNans = Val.ContainsNaN();
+		}
+		else if (Var.GetType() == FNiagaraTypeDefinition::GetVec3Def())
+		{
+			FVector Val = *(FVector*)GetParameterData(Offset);
+			bContainsNans = Val.ContainsNaN();
+		}
+		else if (Var.GetType() == FNiagaraTypeDefinition::GetVec4Def())
+		{
+			FVector4 Val = *(FVector4*)GetParameterData(Offset);
+			bContainsNans = Val.ContainsNaN();
+		}
+		else if (Var.GetType() == FNiagaraTypeDefinition::GetMatrix4Def())
+		{
+			FMatrix Val;
+			FMemory::Memcpy(&Val, GetParameterData(Offset), sizeof(FMatrix));
+			bContainsNans = Val.ContainsNaN();
+		}
+
+		if (bContainsNans)
+		{
+			ensureAlwaysMsgf(false, TEXT("Niagara Parameter Store containts Nans!\n%s"), *DebugName);
+			DumpParameters(false);
+		}
+	}
+}
+
 void FNiagaraParameterStore::Tick()
 {
 	SCOPE_CYCLE_COUNTER(STAT_NiagaraParameterStoreTick);
+#if NIAGARA_NAN_CHECKING
+	CheckForNaNs();
+#endif
 	if (bParametersDirty || bInterfacesDirty)
 	{
 		for (TPair<FNiagaraParameterStore*, FNiagaraParameterStoreBinding>& Binding : Bindings)
@@ -167,7 +212,7 @@ void FNiagaraParameterStore::UnbindFromSourceStores()
 	ensureMsgf(SourceStores.Num() == 0, TEXT("Parameter store source array was not empty after unbinding all sources. Something seriously wrong."));
 }
 
-void FNiagaraParameterStore::DumpParameters(bool bDumpBindings)
+void FNiagaraParameterStore::DumpParameters(bool bDumpBindings)const
 {
 	TArray<FNiagaraVariable> Vars;
 	GetParameters(Vars);
@@ -179,7 +224,7 @@ void FNiagaraParameterStore::DumpParameters(bool bDumpBindings)
 
 	if (bDumpBindings)
 	{
-		for (TPair<FNiagaraParameterStore*, FNiagaraParameterStoreBinding>& Binding : Bindings)
+		for (const TPair<FNiagaraParameterStore*, FNiagaraParameterStoreBinding>& Binding : Bindings)
 		{
 			Binding.Value.Dump(Binding.Key, this);
 		}
diff --git a/Engine/Plugins/FX/Niagara/Source/Niagara/Private/NiagaraScriptExecutionContext.cpp b/Engine/Plugins/FX/Niagara/Source/Niagara/Private/NiagaraScriptExecutionContext.cpp
--- a/Engine/Plugins/FX/Niagara/Source/Niagara/Private/NiagaraScriptExecutionContext.cpp
+++ b/Engine/Plugins/FX/Niagara/Source/Niagara/Private/NiagaraScriptExecutionContext.cpp
@@ -148,6 +148,9 @@ bool FNiagaraScriptExecutionContext::Execute(uint32 NumInstances, TArray<FNiagar
 		SCOPE_CYCLE_COUNTER(STAT_NiagaraSimRegisterSetup);
 		for (FNiagaraDataSetExecutionInfo& DataSetInfo : DataSetInfos)
 		{
+#if NIAGARA_NAN_CHECKING
+			DataSetInfo.DataSet->CheckForNaNs();
+#endif
 			check(DataSetInfo.DataSet);
 			FDataSetMeta SetMeta(DataSetInfo.DataSet->GetSizeBytes(), &InputRegisters[NumInputRegisters], NumInputRegisters, DataSetInfo.StartInstance, 
 				&DataSetInfo.DataSet->CurrIDTable(), &DataSetInfo.DataSet->GetFreeIDTable(), &DataSetInfo.DataSet->GetNumFreeIDs(), &DataSetInfo.DataSet->GetMaxUsedID(), DataSetInfo.DataSet->GetIDAcquireTag());
@@ -185,6 +188,11 @@ bool FNiagaraScriptExecutionContext::Execute(uint32 NumInstances, TArray<FNiagar
 	for (int Idx = 0; Idx < DataSetInfos.Num(); Idx++)
 	{
 		FNiagaraDataSetExecutionInfo& Info = DataSetInfos[Idx];
+
+#if NIAGARA_NAN_CHECKING
+		Info.DataSet->CheckForNaNs();
+#endif
+
 		if (Info.bUpdateInstanceCount)
 		{
 			Info.DataSet->SetNumInstances(Info.StartInstance + DataSetMetaTable[Idx].DataSetAccessIndex + 1);
diff --git a/Engine/Plugins/FX/Niagara/Source/Niagara/Private/NiagaraSystemInstance.cpp b/Engine/Plugins/FX/Niagara/Source/Niagara/Private/NiagaraSystemInstance.cpp
--- a/Engine/Plugins/FX/Niagara/Source/Niagara/Private/NiagaraSystemInstance.cpp
+++ b/Engine/Plugins/FX/Niagara/Source/Niagara/Private/NiagaraSystemInstance.cpp
@@ -120,6 +120,15 @@ void FNiagaraSystemInstance::SetActualExecutionState(ENiagaraExecutionState InSt
 	}
 }
 
+void FNiagaraSystemInstance::Dump()const
+{
+	GetSystemSimulation()->DumpInstance(this);
+	for (auto& Emitter : Emitters)
+	{
+		Emitter->Dump();
+	}
+}
+
 #if WITH_EDITORONLY_DATA
 bool FNiagaraSystemInstance::RequestCapture(const FGuid& RequestId)
 {
diff --git a/Engine/Plugins/FX/Niagara/Source/Niagara/Private/NiagaraSystemSimulation.cpp b/Engine/Plugins/FX/Niagara/Source/Niagara/Private/NiagaraSystemSimulation.cpp
--- a/Engine/Plugins/FX/Niagara/Source/Niagara/Private/NiagaraSystemSimulation.cpp
+++ b/Engine/Plugins/FX/Niagara/Source/Niagara/Private/NiagaraSystemSimulation.cpp
@@ -246,6 +246,20 @@ void FNiagaraSystemSimulation::TransferInstance(FNiagaraSystemSimulation* Source
 	}
 }
 
+void FNiagaraSystemSimulation::DumpInstance(const FNiagaraSystemInstance* Inst)const
+{
+	UE_LOG(LogNiagara, Log, TEXT("==  %s (%d) ========"), *Inst->GetSystem()->GetFullName(), Inst->SystemInstanceIndex);
+	UE_LOG(LogNiagara, Log, TEXT(".................Spawn................."));
+	SpawnExecContext.Parameters.DumpParameters(false);
+	SpawnInstanceParameterDataSet.Dump(false, Inst->SystemInstanceIndex, 1);
+	UE_LOG(LogNiagara, Log, TEXT(".................Update................."));
+	UpdateExecContext.Parameters.DumpParameters(false);
+	UpdateInstanceParameterDataSet.Dump(false, Inst->SystemInstanceIndex, 1);
+	UE_LOG(LogNiagara, Log, TEXT("................. System Instance ................."));
+	DataSet.Dump(false, Inst->SystemInstanceIndex, 1);
+	DataSet.Dump(true, Inst->SystemInstanceIndex, 1);
+}
+
 bool FNiagaraSystemSimulation::Tick(float DeltaSeconds)
 {
 	UNiagaraSystem* System = WeakSystem.Get();
@@ -673,6 +687,10 @@ void FNiagaraSystemSimulation::RemoveInstance(FNiagaraSystemInstance* Instance)
 			SystemInstances[SystemIndex]->SystemInstanceIndex = SystemIndex;
 		}
 	}
+
+#if NIAGARA_NAN_CHECKING
+	DataSet.CheckForNaNs();
+#endif
 }
 
 void FNiagaraSystemSimulation::AddInstance(FNiagaraSystemInstance* Instance)
diff --git a/Engine/Plugins/FX/Niagara/Source/Niagara/Public/NiagaraCommon.h b/Engine/Plugins/FX/Niagara/Source/Niagara/Public/NiagaraCommon.h
--- a/Engine/Plugins/FX/Niagara/Source/Niagara/Public/NiagaraCommon.h
+++ b/Engine/Plugins/FX/Niagara/Source/Niagara/Public/NiagaraCommon.h
@@ -16,6 +16,9 @@ class FNiagaraSystemInstance;
 class UNiagaraParameterCollection;
 struct FNiagaraParameterStore;
 
+//#define NIAGARA_NAN_CHECKING 1
+#define NIAGARA_NAN_CHECKING 0
+
 const uint32 NIAGARA_COMPUTE_THREADGROUP_SIZE = 64;
 const uint32 NIAGARA_MAX_COMPUTE_THREADGROUPS = 65536;
 
diff --git a/Engine/Plugins/FX/Niagara/Source/Niagara/Public/NiagaraParameterStore.h b/Engine/Plugins/FX/Niagara/Source/Niagara/Public/NiagaraParameterStore.h
--- a/Engine/Plugins/FX/Niagara/Source/Niagara/Public/NiagaraParameterStore.h
+++ b/Engine/Plugins/FX/Niagara/Source/Niagara/Public/NiagaraParameterStore.h
@@ -60,7 +60,7 @@ struct FNiagaraParameterStoreBinding
 	FORCEINLINE_DEBUGGABLE void Initialize(FNiagaraParameterStore* DestStore, FNiagaraParameterStore* SrcStore);
 	FORCEINLINE_DEBUGGABLE bool VerifyBinding(const FNiagaraParameterStore* DestStore, const FNiagaraParameterStore* SrcStore)const;
 	FORCEINLINE_DEBUGGABLE void Tick(FNiagaraParameterStore* DestStore, FNiagaraParameterStore* SrcStore, bool bForce = false);
-	FORCEINLINE_DEBUGGABLE void Dump(FNiagaraParameterStore* DestStore, FNiagaraParameterStore* SrcStore)const;
+	FORCEINLINE_DEBUGGABLE void Dump(const FNiagaraParameterStore* DestStore, const FNiagaraParameterStore* SrcStore)const;
 	/** TODO: Merge contiguous ranges into a single binding? */
 	//FORCEINLINE_DEBUGGABLE void Optimize();
 
@@ -129,7 +129,7 @@ public:
 	UObject* GetOwner()const { return Owner; }
 
 	void Dump();
-	void DumpParameters(bool bDumpBindings = false);
+	void DumpParameters(bool bDumpBindings = false)const;
 
 	FORCEINLINE uint32 GetParametersDirty() const { return bParametersDirty; }
 	FORCEINLINE uint32 GetInterfacesDirty() const { return bInterfacesDirty; }
@@ -154,6 +154,8 @@ public:
 	
 	bool VerifyBinding(const FNiagaraParameterStore* InDestStore)const;
 
+	void CheckForNaNs()const;
+
 	/**
 	Adds the passed parameter to this store.
 	Does nothing if this parameter is already present.
@@ -559,9 +561,13 @@ FORCEINLINE_DEBUGGABLE void FNiagaraParameterStoreBinding::Tick(FNiagaraParamete
 			DestStore->SetDataInterface(SrcStore->GetDataInterface(Binding.SrcOffset), Binding.DestOffset);
 		}
 	}
+
+#if NIAGARA_NAN_CHECKING
+	DestStore->CheckForNaNs();
+#endif
 }
 
-FORCEINLINE_DEBUGGABLE void FNiagaraParameterStoreBinding::Dump(FNiagaraParameterStore* DestStore, FNiagaraParameterStore* SrcStore)const
+FORCEINLINE_DEBUGGABLE void FNiagaraParameterStoreBinding::Dump(const FNiagaraParameterStore* DestStore, const FNiagaraParameterStore* SrcStore)const
 {
 #if WITH_EDITORONLY_DATA
 	UE_LOG(LogNiagara, Log, TEXT("\n\nDest Store: %s\n"), *DestStore->DebugName);
diff --git a/Engine/Plugins/FX/Niagara/Source/Niagara/Public/NiagaraSystemInstance.h b/Engine/Plugins/FX/Niagara/Source/Niagara/Public/NiagaraSystemInstance.h
--- a/Engine/Plugins/FX/Niagara/Source/Niagara/Public/NiagaraSystemInstance.h
+++ b/Engine/Plugins/FX/Niagara/Source/Niagara/Public/NiagaraSystemInstance.h
@@ -192,6 +192,9 @@ public:
 
 #endif
 
+	/** Dumps all of this systems info to the log. */
+	void Dump()const;
+
 private:
 	/** Builds the emitter simulations. */
 	void InitEmitters();
diff --git a/Engine/Plugins/FX/Niagara/Source/Niagara/Public/NiagaraSystemSimulation.h b/Engine/Plugins/FX/Niagara/Source/Niagara/Public/NiagaraSystemSimulation.h
--- a/Engine/Plugins/FX/Niagara/Source/Niagara/Public/NiagaraSystemSimulation.h
+++ b/Engine/Plugins/FX/Niagara/Source/Niagara/Public/NiagaraSystemSimulation.h
@@ -64,6 +64,10 @@ struct FNiagaraParameterStoreToDataSetBinding
 
 	FORCEINLINE_DEBUGGABLE void DataSetToParameterStore(FNiagaraParameterStore& ParameterStore, FNiagaraDataSet& DataSet, int32 DataSetInstanceIndex)
 	{
+#if NIAGARA_NAN_CHECKING
+		DataSet.CheckForNaNs();
+#endif
+
 		FNiagaraDataBuffer& CurrBuffer = DataSet.CurrData();
 
 		for (const FDataOffsets& DataOffsets : FloatOffsets)
@@ -77,6 +81,10 @@ struct FNiagaraParameterStoreToDataSetBinding
 			ParameterStore.SetParameterByOffset(DataOffsets.ParameterOffset, *DataSetPtr);
 		}
 
+#if NIAGARA_NAN_CHECKING
+		ParameterStore.CheckForNaNs();
+#endif
+
 		ParameterStore.OnParameterChange();
 	}
 
@@ -85,6 +93,10 @@ struct FNiagaraParameterStoreToDataSetBinding
 		FNiagaraDataBuffer& CurrBuffer = DataSet.CurrData();
 		const uint8* ParameterData = ParameterStore.GetParameterDataArray().GetData();
 
+#if NIAGARA_NAN_CHECKING
+		ParameterStore.CheckForNaNs();
+#endif
+
 		for (const FDataOffsets& DataOffsets : FloatOffsets)
 		{
 			float* ParamPtr = (float*)(ParameterData + DataOffsets.ParameterOffset);
@@ -97,6 +109,10 @@ struct FNiagaraParameterStoreToDataSetBinding
 			int32* DataSetPtr = CurrBuffer.GetInstancePtrInt32(DataOffsets.DataSetComponentOffset, DataSetInstanceIndex);
 			*DataSetPtr = *ParamPtr;
 		}
+
+#if NIAGARA_NAN_CHECKING
+		DataSet.CheckForNaNs();
+#endif
 	}
 };
 
@@ -120,6 +136,8 @@ public:
 
 	/** Transfers a system instance from SourceSimulation. */
 	void TransferInstance(FNiagaraSystemSimulation* SourceSimulation, FNiagaraSystemInstance* SystemInst);
+
+	void DumpInstance(const FNiagaraSystemInstance* Inst)const;
 protected:
 
 	/** System of instances being simulated.  We use a weak object ptr here because once the last referencing object goes away this system may be come invalid at runtime. */
