Niagara - Add a CVar which prevents a few panels in the niagara editor from showing up in the Window menu.  These tabs are being hidden by default since they're for tool development and debugging, and can easily crash the editor when used incorrectly.

diff --git a/Engine/Plugins/FX/Niagara/Source/NiagaraEditor/Private/Toolkits/NiagaraSystemToolkit.cpp b/Engine/Plugins/FX/Niagara/Source/NiagaraEditor/Private/Toolkits/NiagaraSystemToolkit.cpp
--- a/Engine/Plugins/FX/Niagara/Source/NiagaraEditor/Private/Toolkits/NiagaraSystemToolkit.cpp
+++ b/Engine/Plugins/FX/Niagara/Source/NiagaraEditor/Private/Toolkits/NiagaraSystemToolkit.cpp
@@ -73,6 +73,14 @@ static FAutoConsoleVariableRef CVarSuppressNiagaraSystems(
 	ECVF_Default
 );
 
+static int32 GbShowNiagaraDeveloperWindows = 0;
+static FAutoConsoleVariableRef CVarShowNiagaraDeveloperWindows(
+	TEXT("fx.ShowNiagaraDeveloperWindows"),
+	GbShowNiagaraDeveloperWindows,
+	TEXT("If > 0 the niagara system and emitter editors will show additional developer windows.\nThese windows are for niagara tool development and debugging and editing the data\n directly in these windows can cause instability.\n"),
+	ECVF_Default
+);
+
 void FNiagaraSystemToolkit::RegisterTabSpawners(const TSharedRef<class FTabManager>& InTabManager)
 {
 	WorkspaceMenuCategory = InTabManager->AddLocalWorkspaceMenuCategory(LOCTEXT("WorkspaceMenu_NiagaraSystemEditor", "Niagara System"));
@@ -94,7 +102,8 @@ void FNiagaraSystemToolkit::RegisterTabSpawners(const TSharedRef<class FTabManag
 
 	InTabManager->RegisterTabSpawner(SystemScriptTabID, FOnSpawnTab::CreateSP(this, &FNiagaraSystemToolkit::SpawnTab_SystemScript))
 		.SetDisplayName(LOCTEXT("SystemScript", "System Script"))
-		.SetGroup(WorkspaceMenuCategory.ToSharedRef());
+		.SetGroup(WorkspaceMenuCategory.ToSharedRef())
+		.SetAutoGenerateMenuEntry(GbShowNiagaraDeveloperWindows);
 
 	InTabManager->RegisterTabSpawner(SystemDetailsTabID, FOnSpawnTab::CreateSP(this, &FNiagaraSystemToolkit::SpawnTab_SystemDetails))
 		.SetDisplayName(LOCTEXT("SystemDetails", "System Details"))
@@ -110,7 +119,8 @@ void FNiagaraSystemToolkit::RegisterTabSpawners(const TSharedRef<class FTabManag
 
 	InTabManager->RegisterTabSpawner(SelectedEmitterGraphTabID, FOnSpawnTab::CreateSP(this, &FNiagaraSystemToolkit::SpawnTab_SelectedEmitterGraph))
 		.SetDisplayName(LOCTEXT("SelectedEmitterGraph", "Selected Emitter Graph"))
-		.SetGroup(WorkspaceMenuCategory.ToSharedRef());
+		.SetGroup(WorkspaceMenuCategory.ToSharedRef())
+		.SetAutoGenerateMenuEntry(GbShowNiagaraDeveloperWindows);
 
 	InTabManager->RegisterTabSpawner(DebugSpreadsheetTabID, FOnSpawnTab::CreateSP(this, &FNiagaraSystemToolkit::SpawnTab_DebugSpreadsheet))
 		.SetDisplayName(LOCTEXT("DebugSpreadsheet", "Attribute Spreadsheet"))
