Niagara - Fix warning from last submissing about using an int32 as a bool.

diff --git a/Engine/Plugins/FX/Niagara/Source/NiagaraEditor/Private/Toolkits/NiagaraSystemToolkit.cpp b/Engine/Plugins/FX/Niagara/Source/NiagaraEditor/Private/Toolkits/NiagaraSystemToolkit.cpp
--- a/Engine/Plugins/FX/Niagara/Source/NiagaraEditor/Private/Toolkits/NiagaraSystemToolkit.cpp
+++ b/Engine/Plugins/FX/Niagara/Source/NiagaraEditor/Private/Toolkits/NiagaraSystemToolkit.cpp
@@ -103,7 +103,7 @@ void FNiagaraSystemToolkit::RegisterTabSpawners(const TSharedRef<class FTabManag
 	InTabManager->RegisterTabSpawner(SystemScriptTabID, FOnSpawnTab::CreateSP(this, &FNiagaraSystemToolkit::SpawnTab_SystemScript))
 		.SetDisplayName(LOCTEXT("SystemScript", "System Script"))
 		.SetGroup(WorkspaceMenuCategory.ToSharedRef())
-		.SetAutoGenerateMenuEntry(GbShowNiagaraDeveloperWindows);
+		.SetAutoGenerateMenuEntry(GbShowNiagaraDeveloperWindows != 0);
 
 	InTabManager->RegisterTabSpawner(SystemDetailsTabID, FOnSpawnTab::CreateSP(this, &FNiagaraSystemToolkit::SpawnTab_SystemDetails))
 		.SetDisplayName(LOCTEXT("SystemDetails", "System Details"))
@@ -120,7 +120,7 @@ void FNiagaraSystemToolkit::RegisterTabSpawners(const TSharedRef<class FTabManag
 	InTabManager->RegisterTabSpawner(SelectedEmitterGraphTabID, FOnSpawnTab::CreateSP(this, &FNiagaraSystemToolkit::SpawnTab_SelectedEmitterGraph))
 		.SetDisplayName(LOCTEXT("SelectedEmitterGraph", "Selected Emitter Graph"))
 		.SetGroup(WorkspaceMenuCategory.ToSharedRef())
-		.SetAutoGenerateMenuEntry(GbShowNiagaraDeveloperWindows);
+		.SetAutoGenerateMenuEntry(GbShowNiagaraDeveloperWindows != 0);
 
 	InTabManager->RegisterTabSpawner(DebugSpreadsheetTabID, FOnSpawnTab::CreateSP(this, &FNiagaraSystemToolkit::SpawnTab_DebugSpreadsheet))
 		.SetDisplayName(LOCTEXT("DebugSpreadsheet", "Attribute Spreadsheet"))
