Compile fix for none editor

diff --git a/Engine/Plugins/FX/Niagara/Source/Niagara/Private/NiagaraParameterStore.cpp b/Engine/Plugins/FX/Niagara/Source/Niagara/Private/NiagaraParameterStore.cpp
--- a/Engine/Plugins/FX/Niagara/Source/Niagara/Private/NiagaraParameterStore.cpp
+++ b/Engine/Plugins/FX/Niagara/Source/Niagara/Private/NiagaraParameterStore.cpp
@@ -174,7 +174,7 @@ void FNiagaraParameterStore::CheckForNaNs()const
 
 		if (bContainsNans)
 		{
-			ensureAlwaysMsgf(false, TEXT("Niagara Parameter Store containts Nans!\n%s"), *DebugName);
+			ensureAlwaysMsgf(false, TEXT("Niagara Parameter Store containts Nans!\n"));
 			DumpParameters(false);
 		}
 	}
