Niagara - Prevent a crash when undoing enabling/disabling a module.  Also fix undo not working correctly for enable/disable in some cases.

diff --git a/Engine/Plugins/FX/Niagara/Source/NiagaraEditor/Private/ViewModels/Stack/NiagaraStackGraphUtilities.cpp b/Engine/Plugins/FX/Niagara/Source/NiagaraEditor/Private/ViewModels/Stack/NiagaraStackGraphUtilities.cpp
--- a/Engine/Plugins/FX/Niagara/Source/NiagaraEditor/Private/ViewModels/Stack/NiagaraStackGraphUtilities.cpp
+++ b/Engine/Plugins/FX/Niagara/Source/NiagaraEditor/Private/ViewModels/Stack/NiagaraStackGraphUtilities.cpp
@@ -1021,6 +1021,7 @@ void FNiagaraStackGraphUtilities::SetModuleIsEnabled(UNiagaraNodeFunctionCall& F
 	GetAllNodesForModule(FunctionCallNode, ModuleNodes);
 	for (UNiagaraNode* ModuleNode : ModuleNodes)
 	{
+		ModuleNode->Modify();
 		ModuleNode->SetEnabledState(bIsEnabled ? ENodeEnabledState::Enabled : ENodeEnabledState::Disabled, true);
 		ModuleNode->MarkNodeRequiresSynchronization(__FUNCTION__, false);
 	}
diff --git a/Engine/Plugins/FX/Niagara/Source/NiagaraEditor/Private/ViewModels/Stack/NiagaraStackScriptItemGroup.cpp b/Engine/Plugins/FX/Niagara/Source/NiagaraEditor/Private/ViewModels/Stack/NiagaraStackScriptItemGroup.cpp
--- a/Engine/Plugins/FX/Niagara/Source/NiagaraEditor/Private/ViewModels/Stack/NiagaraStackScriptItemGroup.cpp
+++ b/Engine/Plugins/FX/Niagara/Source/NiagaraEditor/Private/ViewModels/Stack/NiagaraStackScriptItemGroup.cpp
@@ -264,15 +264,16 @@ void UNiagaraStackScriptItemGroup::Initialize(
 	ScriptViewModel = InScriptViewModel;
 	ScriptUsage = InScriptUsage;
 	ScriptUsageId = InScriptUsageId;
-	OnGraphChangedHandle = InScriptViewModel->GetGraphViewModel()->GetGraph()->AddOnGraphChangedHandler(
+	ScriptGraph = InScriptViewModel->GetGraphViewModel()->GetGraph();
+	ScriptGraph->AddOnGraphChangedHandler(
 		FOnGraphChanged::FDelegate::CreateUObject(this, &UNiagaraStackScriptItemGroup::OnScriptGraphChanged));
 }
 
 void UNiagaraStackScriptItemGroup::FinalizeInternal()
 {
-	if (ScriptViewModel.IsValid())
+	if (ScriptGraph.IsValid())
 	{
-		ScriptViewModel.Pin()->GetGraphViewModel()->GetGraph()->RemoveOnGraphChangedHandler(OnGraphChangedHandle);
+		ScriptGraph->RemoveOnGraphChangedHandler(OnGraphChangedHandle);
 	}
 	Super::FinalizeInternal();
 }
diff --git a/Engine/Plugins/FX/Niagara/Source/NiagaraEditor/Public/ViewModels/Stack/NiagaraStackScriptItemGroup.h b/Engine/Plugins/FX/Niagara/Source/NiagaraEditor/Public/ViewModels/Stack/NiagaraStackScriptItemGroup.h
--- a/Engine/Plugins/FX/Niagara/Source/NiagaraEditor/Public/ViewModels/Stack/NiagaraStackScriptItemGroup.h
+++ b/Engine/Plugins/FX/Niagara/Source/NiagaraEditor/Public/ViewModels/Stack/NiagaraStackScriptItemGroup.h
@@ -13,6 +13,7 @@ class UNiagaraNodeOutput;
 class UNiagaraNodeFunctionCall;
 class FScriptItemGroupAddUtilities;
 class UNiagaraStackModuleItem;
+class UEdGraph;
 
 UCLASS()
 class NIAGARAEDITOR_API UNiagaraStackScriptItemGroup : public UNiagaraStackItemGroup
@@ -59,6 +60,8 @@ private:
 	FGuid ScriptUsageId;
 	bool bIsValidForOutput;
 
+	TWeakObjectPtr<UEdGraph> ScriptGraph;
+
 	FDelegateHandle OnGraphChangedHandle;
 
 	TMap<FObjectKey, UNiagaraStackModuleItem*> StackSpacerToModuleItemMap;
