Fixing Vulkan RHI so that we clean up on app pre exit before the RHI is fully gone. Module cleanup is too late.

#test ran PC with -Vulkan on test asset mentioned in jira. Successfully exited after code change. Crashed pre code change.

diff --git a/Engine/Plugins/FX/Niagara/Source/Niagara/Private/NiagaraModule.cpp b/Engine/Plugins/FX/Niagara/Source/Niagara/Private/NiagaraModule.cpp
--- a/Engine/Plugins/FX/Niagara/Source/Niagara/Private/NiagaraModule.cpp
+++ b/Engine/Plugins/FX/Niagara/Source/Niagara/Private/NiagaraModule.cpp
@@ -16,6 +16,7 @@
 #include "NiagaraMeshRendererProperties.h"
 #include "NiagaraRibbonRendererProperties.h"
 #include "NiagaraRenderer.h"
+#include "Misc/CoreDelegates.h"
 
 IMPLEMENT_MODULE(INiagaraModule, Niagara);
 
@@ -225,6 +226,13 @@ void INiagaraModule::StartupModule()
 	UNiagaraSpriteRendererProperties::InitCDOPropertiesAfterModuleStartup();
 	UNiagaraRibbonRendererProperties::InitCDOPropertiesAfterModuleStartup();
 	UNiagaraMeshRendererProperties::InitCDOPropertiesAfterModuleStartup();
+
+	FCoreDelegates::OnPreExit.AddRaw(this, &INiagaraModule::ShutdownRenderingResources);
+}
+
+void INiagaraModule::ShutdownRenderingResources()
+{
+	FNiagaraViewDataMgr::Shutdown();
 }
 
 void INiagaraModule::ShutdownModule()
@@ -238,7 +246,7 @@ void INiagaraModule::ShutdownModule()
 	}
 
 	CVarDetailLevel.AsVariable()->SetOnChangedCallback(FConsoleVariableDelegate());
-	FNiagaraViewDataMgr::Shutdown();
+	ShutdownRenderingResources();
 }
 
 FNiagaraWorldManager* INiagaraModule::GetWorldManager(UWorld* World)
diff --git a/Engine/Plugins/FX/Niagara/Source/Niagara/Private/NiagaraWorldManager.cpp b/Engine/Plugins/FX/Niagara/Source/Niagara/Private/NiagaraWorldManager.cpp
--- a/Engine/Plugins/FX/Niagara/Source/Niagara/Private/NiagaraWorldManager.cpp
+++ b/Engine/Plugins/FX/Niagara/Source/Niagara/Private/NiagaraWorldManager.cpp
@@ -32,10 +32,12 @@ void FNiagaraViewDataMgr::Init()
 
 void FNiagaraViewDataMgr::Shutdown()
 {
-
+	GNiagaraViewDataManager.SceneDepthTexture = nullptr;
+	GNiagaraViewDataManager.SceneNormalTexture = nullptr;
+	GNiagaraViewDataManager.ViewUniformBuffer = nullptr;
+	GNiagaraViewDataManager.SceneTexturesUniformParams.SafeRelease();
 }
 
-
 FNiagaraWorldManager::FNiagaraWorldManager(UWorld* InWorld)
 	: World(InWorld)
 	, CachedEffectsQuality(INDEX_NONE)
diff --git a/Engine/Plugins/FX/Niagara/Source/Niagara/Public/NiagaraModule.h b/Engine/Plugins/FX/Niagara/Source/Niagara/Public/NiagaraModule.h
--- a/Engine/Plugins/FX/Niagara/Source/Niagara/Public/NiagaraModule.h
+++ b/Engine/Plugins/FX/Niagara/Source/Niagara/Public/NiagaraModule.h
@@ -62,6 +62,7 @@ public:
 public:
 	virtual void StartupModule()override;
 	virtual void ShutdownModule()override;
+	void ShutdownRenderingResources();
 
 	FDelegateHandle SetOnProcessShaderCompilationQueue(FOnProcessQueue InOnProcessQueue);
 	void ResetOnProcessShaderCompilationQueue(FDelegateHandle DelegateHandle);
