Added protections to try and prevent UE4Editor_NiagaraEditor!UNiagaraStackFunctionInput::CanReset()  crash
#tests n/a

diff --git a/Engine/Plugins/FX/Niagara/Source/NiagaraEditor/Private/ViewModels/Stack/NiagaraStackFunctionInput.cpp b/Engine/Plugins/FX/Niagara/Source/NiagaraEditor/Private/ViewModels/Stack/NiagaraStackFunctionInput.cpp
--- a/Engine/Plugins/FX/Niagara/Source/NiagaraEditor/Private/ViewModels/Stack/NiagaraStackFunctionInput.cpp
+++ b/Engine/Plugins/FX/Niagara/Source/NiagaraEditor/Private/ViewModels/Stack/NiagaraStackFunctionInput.cpp
@@ -977,35 +977,38 @@ bool UNiagaraStackFunctionInput::CanReset() const
 		else
 		{
 			UEdGraphPin* DefaultPin = GetDefaultPin();
-			if(DefaultPin->LinkedTo.Num() == 0)
-			{
-				if (GetOverridePin() != nullptr)
-				{
-					bNewCanReset = true;
-				}
-				else if (IsRapidIterationCandidate())
-				{
-					FNiagaraVariable DefaultVar = GetDefaultVariableForRapidIterationParameter();
-					bool bHasValidLocalValue = InputValues.LocalStruct.IsValid();
-					bool bHasValidDefaultValue = DefaultVar.IsValid();
-					bNewCanReset = bHasValidLocalValue && bHasValidDefaultValue && FNiagaraEditorUtilities::DataMatches(DefaultVar, *InputValues.LocalStruct.Get()) == false;
-				}
-				else
-				{
-					bNewCanReset = false;
-				}
-			}
-			else
-			{
-				if (FNiagaraStackGraphUtilities::IsValidDefaultDynamicInput(*SourceScript, *DefaultPin))
+			if (ensure(DefaultPin != nullptr))
+			{			
+				if(DefaultPin->LinkedTo.Num() == 0)
 				{
-					UEdGraphPin* OverridePin = GetOverridePin();
-					bNewCanReset = OverridePin == nullptr || FNiagaraStackGraphUtilities::DoesDynamicInputMatchDefault(GetEmitterViewModel()->GetEmitter()->GetUniqueEmitterName(), *SourceScript,
-						*OwningFunctionCallNode, *OverridePin, InputParameterHandle.GetName(), *DefaultPin) == false;
+					if (GetOverridePin() != nullptr)
+					{
+						bNewCanReset = true;
+					}
+					else if (IsRapidIterationCandidate())
+					{
+						FNiagaraVariable DefaultVar = GetDefaultVariableForRapidIterationParameter();
+						bool bHasValidLocalValue = InputValues.LocalStruct.IsValid();
+						bool bHasValidDefaultValue = DefaultVar.IsValid();
+						bNewCanReset = bHasValidLocalValue && bHasValidDefaultValue && FNiagaraEditorUtilities::DataMatches(DefaultVar, *InputValues.LocalStruct.Get()) == false;
+					}
+					else
+					{
+						bNewCanReset = false;
+					}
 				}
 				else
 				{
-					bNewCanReset = GetOverridePin() != nullptr;
+					if (FNiagaraStackGraphUtilities::IsValidDefaultDynamicInput(*SourceScript, *DefaultPin))
+					{
+						UEdGraphPin* OverridePin = GetOverridePin();
+						bNewCanReset = OverridePin == nullptr || FNiagaraStackGraphUtilities::DoesDynamicInputMatchDefault(GetEmitterViewModel()->GetEmitter()->GetUniqueEmitterName(), *SourceScript,
+							*OwningFunctionCallNode, *OverridePin, InputParameterHandle.GetName(), *DefaultPin) == false;
+					}
+					else
+					{
+						bNewCanReset = GetOverridePin() != nullptr;
+					}
 				}
 			}
 		}
