Niagara - Fix a crash when editing inputs on a module with an edit condition that's been added in a system.

diff --git a/Engine/Plugins/FX/Niagara/Source/Niagara/Classes/NiagaraSystem.h b/Engine/Plugins/FX/Niagara/Source/Niagara/Classes/NiagaraSystem.h
--- a/Engine/Plugins/FX/Niagara/Source/Niagara/Classes/NiagaraSystem.h
+++ b/Engine/Plugins/FX/Niagara/Source/Niagara/Classes/NiagaraSystem.h
@@ -149,6 +149,9 @@ public:
 	/** Called to query whether or not this emitter is referenced as the source to any emitter handles for this System.*/
 	bool ReferencesSourceEmitter(UNiagaraEmitter& Emitter);
 
+	/** Determines if this system has the supplied emitter as an editable and simulating emitter instance. */
+	bool ReferencesInstanceEmitter(UNiagaraEmitter& Emitter);
+
 	/** Updates all handles which use this emitter as their source. */
 	void UpdateFromEmitterChanges(UNiagaraEmitter& ChangedSourceEmitter);
 
diff --git a/Engine/Plugins/FX/Niagara/Source/Niagara/Private/NiagaraSystem.cpp b/Engine/Plugins/FX/Niagara/Source/Niagara/Private/NiagaraSystem.cpp
--- a/Engine/Plugins/FX/Niagara/Source/Niagara/Private/NiagaraSystem.cpp
+++ b/Engine/Plugins/FX/Niagara/Source/Niagara/Private/NiagaraSystem.cpp
@@ -445,6 +445,18 @@ bool UNiagaraSystem::ReferencesSourceEmitter(UNiagaraEmitter& Emitter)
 	return false;
 }
 
+bool UNiagaraSystem::ReferencesInstanceEmitter(UNiagaraEmitter& Emitter)
+{
+	for (FNiagaraEmitterHandle& Handle : EmitterHandles)
+	{
+		if (&Emitter == Handle.GetInstance())
+		{
+			return true;
+		}
+	}
+	return false;
+}
+
 void UNiagaraSystem::UpdateFromEmitterChanges(UNiagaraEmitter& ChangedSourceEmitter)
 {
 	bool bNeedsCompile = false;
diff --git a/Engine/Plugins/FX/Niagara/Source/NiagaraEditor/Private/NiagaraEditorUtilities.cpp b/Engine/Plugins/FX/Niagara/Source/NiagaraEditor/Private/NiagaraEditorUtilities.cpp
--- a/Engine/Plugins/FX/Niagara/Source/NiagaraEditor/Private/NiagaraEditorUtilities.cpp
+++ b/Engine/Plugins/FX/Niagara/Source/NiagaraEditor/Private/NiagaraEditorUtilities.cpp
@@ -498,11 +498,11 @@ void FNiagaraEditorUtilities::CompileExistingEmitters(const TArray<UNiagaraEmitt
 			continue;
 		}
 
-		// We only need to compile emitters referenced directly by systems since emitters can now only be used in the context 
+		// We only need to compile emitters referenced directly as instances by systems since emitters can now only be used in the context 
 		// of a system.
 		for (TObjectIterator<UNiagaraSystem> SystemIterator; SystemIterator; ++SystemIterator)
 		{
-			if (SystemIterator->ReferencesSourceEmitter(*Emitter))
+			if (SystemIterator->ReferencesInstanceEmitter(*Emitter))
 			{
 				SystemIterator->RequestCompile(false);
 
