Niagara - Fix capturing while paused in the attribute spreadsheet.

diff --git a/Engine/Plugins/FX/Niagara/Source/NiagaraEditor/Private/Widgets/SNiagaraSpreadsheetView.cpp b/Engine/Plugins/FX/Niagara/Source/NiagaraEditor/Private/Widgets/SNiagaraSpreadsheetView.cpp
--- a/Engine/Plugins/FX/Niagara/Source/NiagaraEditor/Private/Widgets/SNiagaraSpreadsheetView.cpp
+++ b/Engine/Plugins/FX/Niagara/Source/NiagaraEditor/Private/Widgets/SNiagaraSpreadsheetView.cpp
@@ -1375,8 +1375,13 @@ FReply SNiagaraSpreadsheetView::OnCaptureRequestPressed()
 {
 	FFrameRate TickResolution = SystemViewModel->GetSequencer()->GetFocusedTickResolution();
 	float LocalTime = SystemViewModel->GetSequencer()->GetLocalTime().AsSeconds();
-	float SnapInterval = SystemViewModel->GetSequencer()->GetFocusedDisplayRate().AsInterval();
-	float TargetCaptureTime = LocalTime + SnapInterval;
+	
+	// The preview component in the editor is using the 'DesiredAge' update mode so each frame it determines if the difference
+	// between the current age and the desired age is greater then the seek delta and if so it advanced the simulation the correct
+	// number of times.  We want to ensure that we simulate a single step so we get the seek delta from the component and add that 
+	// to the current time.
+	float SimulationStep = SystemViewModel->GetPreviewComponent()->GetSeekDelta();
+	float TargetCaptureTime = LocalTime + SimulationStep;
 
 	TArray<TSharedRef<FNiagaraEmitterHandleViewModel>> SelectedEmitterHandles;
 	SystemViewModel->GetSelectedEmitterHandles(SelectedEmitterHandles);
