Avoiding copy of sim data not needed in most situations.

diff --git a/Engine/Plugins/FX/Niagara/Source/Niagara/Private/NiagaraDataSet.cpp b/Engine/Plugins/FX/Niagara/Source/Niagara/Private/NiagaraDataSet.cpp
--- a/Engine/Plugins/FX/Niagara/Source/Niagara/Private/NiagaraDataSet.cpp
+++ b/Engine/Plugins/FX/Niagara/Source/Niagara/Private/NiagaraDataSet.cpp
@@ -267,12 +267,20 @@ void FNiagaraDataBuffer::Allocate(uint32 InNumInstances, bool bMaintainExisting)
 		DEC_MEMORY_STAT_BY(STAT_NiagaraParticleMemory, FloatData.Max() + Int32Data.Max());
 
 		int32 OldFloatStride = FloatStride;
-		TArray<uint8> OldFloatData = FloatData;
+		TArray<uint8> OldFloatData;
+		int32 OldInt32Stride = Int32Stride;
+		TArray<uint8> OldIntData;
+
+		if (bMaintainExisting)
+		{
+			//Need to copy off old data so we can copy it back into the newly laid out buffers. TODO: Avoid this needless copying.
+			OldFloatData = FloatData;
+			OldIntData = Int32Data;
+		}
+
 		FloatStride = GetSafeComponentBufferSize(NumInstancesAllocated * sizeof(float));
 		FloatData.SetNum(FloatStride * Owner->GetNumFloatComponents(), false);
 
-		int32 OldInt32Stride = Int32Stride;
-		TArray<uint8> OldIntData = Int32Data;
 		Int32Stride = GetSafeComponentBufferSize(NumInstancesAllocated * sizeof(int32));
 		Int32Data.SetNum(Int32Stride * Owner->GetNumInt32Components(), false);
 
