Placing a max limit of 1000000 particles per emitter on CPU as a temp work around for some other issues with very high particle counts we're seeing.

diff --git a/Engine/Plugins/FX/Niagara/Source/Niagara/Private/NiagaraEmitterInstance.cpp b/Engine/Plugins/FX/Niagara/Source/Niagara/Private/NiagaraEmitterInstance.cpp
--- a/Engine/Plugins/FX/Niagara/Source/Niagara/Private/NiagaraEmitterInstance.cpp
+++ b/Engine/Plugins/FX/Niagara/Source/Niagara/Private/NiagaraEmitterInstance.cpp
@@ -32,6 +32,17 @@ static FAutoConsoleVariableRef CVarNiagaraDumpParticleData(
 	ECVF_Default
 	);
 
+/**
+TODO: This is mainly to avoid hard limits in our storage/alloc code etc rather than for perf reasons.
+We should improve our hard limit/safety code and possibly add a max for perf reasons.
+*/
+static int32 GMaxCPUParticlesPerEmitter = 1000000;
+static FAutoConsoleVariableRef CVarMaxCPUParticlesPerEmitter(
+	TEXT("fx.MaxCPUParticlesPerEmitter"),
+	GMaxCPUParticlesPerEmitter,
+	TEXT("The max number of supported CPU particles per emitter. \n"),
+	ECVF_Default
+);
 //////////////////////////////////////////////////////////////////////////
 
 const FName FNiagaraEmitterInstance::PositionName(TEXT("Position"));
@@ -1002,6 +1013,18 @@ void FNiagaraEmitterInstance::Tick(float DeltaSeconds)
 	}
 
 	int32 AllocationSize = OrigNumParticles + SpawnTotal + EventSpawnTotal;
+
+	//Ensure we don't blow our current hard limits on cpu particle count.
+	//TODO: These current limits can be improved relatively easily. Though perf in at these counts will obviously be an issue anyway.
+	if (CachedEmitter->SimTarget == ENiagaraSimTarget::CPUSim && AllocationSize > GMaxCPUParticlesPerEmitter)
+	{
+		UE_LOG(LogNiagara, Warning, TEXT("Emitter %s has attemted to exceed the max CPU particle count! | Max: %d | Requested: %u"), *CachedEmitter->GetUniqueEmitterName(), GMaxCPUParticlesPerEmitter, AllocationSize);
+		//For now we completely bail out of spawning new particles. Possibly should improve this in future.
+		AllocationSize = OrigNumParticles;
+		SpawnTotal = 0;
+		EventSpawnTotal = 0;
+	}
+
 	//Allocate space for prev frames particles and any new one's we're going to spawn.
 	Data.Allocate(AllocationSize);
 	for (FNiagaraDataSet* SpawnEventDataSet : SpawnScriptEventDataSets)
