Niagara - Fix a crash when deleting any module with a data interface.

diff --git a/Engine/Plugins/FX/Niagara/Source/NiagaraEditor/Private/ViewModels/Stack/NiagaraStackViewModel.cpp b/Engine/Plugins/FX/Niagara/Source/NiagaraEditor/Private/ViewModels/Stack/NiagaraStackViewModel.cpp
--- a/Engine/Plugins/FX/Niagara/Source/NiagaraEditor/Private/ViewModels/Stack/NiagaraStackViewModel.cpp
+++ b/Engine/Plugins/FX/Niagara/Source/NiagaraEditor/Private/ViewModels/Stack/NiagaraStackViewModel.cpp
@@ -73,7 +73,7 @@ void UNiagaraStackViewModel::Initialize(TSharedPtr<FNiagaraSystemViewModel> InSy
 		RootEntry->OnStructureChanged().AddUObject(this, &UNiagaraStackViewModel::EntryStructureChanged);
 		RootEntry->OnDataObjectModified().AddUObject(this, &UNiagaraStackViewModel::EntryDataObjectModified);
 		RootEntry->OnRequestFullRefresh().AddUObject(this, &UNiagaraStackViewModel::EntryRequestFullRefresh);
-		RootEntry->OnRequestFullRefreshDeferred().AddUObject(this, &UNiagaraStackViewModel::EntryRequestFullRefresh);
+		RootEntry->OnRequestFullRefreshDeferred().AddUObject(this, &UNiagaraStackViewModel::EntryRequestFullRefreshDeferred);
 		RootEntries.Add(RootEntry);
 	}
 
