Niagara - Fix a crash when setting EmitterLifeClycle->NextLoopDuration to really large numbers.

diff --git a/Engine/Plugins/FX/Niagara/Source/NiagaraEditor/Private/Sequencer/NiagaraSequence/Sections/MovieSceneNiagaraEmitterSection.cpp b/Engine/Plugins/FX/Niagara/Source/NiagaraEditor/Private/Sequencer/NiagaraSequence/Sections/MovieSceneNiagaraEmitterSection.cpp
--- a/Engine/Plugins/FX/Niagara/Source/NiagaraEditor/Private/Sequencer/NiagaraSequence/Sections/MovieSceneNiagaraEmitterSection.cpp
+++ b/Engine/Plugins/FX/Niagara/Source/NiagaraEditor/Private/Sequencer/NiagaraSequence/Sections/MovieSceneNiagaraEmitterSection.cpp
@@ -267,10 +267,20 @@ void UMovieSceneNiagaraEmitterSection::UpdateSectionFromTimeRangeModule(const FF
 				*LengthBinder.GetInputName().ToString());
 			ModuleLength = 0;
 		}
+
+		FFrameNumber StartFrame = (ModuleStartTime * InTickResolution).RoundToFrame();
+		FFrameNumber EndFrame = ((ModuleStartTime + ModuleLength) * InTickResolution).RoundToFrame();
+		if (EndFrame < StartFrame)
+		{
+			// The frame value has overflowed and is negative so clamp to the max frame.
+			// TODO: Add ui support for this issue rather than a log error.
+			UE_LOG(LogNiagaraEditor, Error, TEXT("Invalid length in niagara editor timeline.  Bound Module: %s Bound Input: %s"),
+				LengthBinder.GetFunctionCallNode() != nullptr ? *LengthBinder.GetFunctionCallNode()->GetFunctionName() : TEXT("Unknown"),
+				*LengthBinder.GetInputName().ToString());
+			EndFrame.Value = TNumericLimits<int32>::Max();
+		}
 		
-		SetRange(TRange<FFrameNumber>(
-			(ModuleStartTime * InTickResolution).RoundToFrame(),
-			((ModuleStartTime + ModuleLength) * InTickResolution).RoundToFrame()));
+		SetRange(TRange<FFrameNumber>(StartFrame, EndFrame));
 	}
 	else
 	{
