Niagara - Prevent a crash when an emitter or system is referencing a data interface from a parameter collection and that parameter is deleted from the collection.

diff --git a/Engine/Plugins/FX/Niagara/Source/NiagaraEditor/Private/NiagaraHlslTranslator.cpp b/Engine/Plugins/FX/Niagara/Source/NiagaraEditor/Private/NiagaraHlslTranslator.cpp
--- a/Engine/Plugins/FX/Niagara/Source/NiagaraEditor/Private/NiagaraHlslTranslator.cpp
+++ b/Engine/Plugins/FX/Niagara/Source/NiagaraEditor/Private/NiagaraHlslTranslator.cpp
@@ -3500,7 +3500,9 @@ bool FHlslNiagaraTranslator::ParameterMapRegisterExternalConstantNamespaceVariab
 			Collection = ParamMapHistories[InParamMapHistoryIdx].IsParameterCollectionParameter(InVariable, bMissingParameter);
 			if (Collection && bMissingParameter)
 			{
-				Error(FText::Format(LOCTEXT("MissingNPCParameterError", "Parameter {0} was not found in Parameter Collection {1}"), FText::FromName(InVariable.GetName()), FText::FromString(Collection->GetFullName())), InNode, nullptr);
+				Error(FText::Format(LOCTEXT("MissingNPCParameterError", "Parameter named {0} of type {1} was not found in Parameter Collection {2}"),
+					FText::FromName(InVariable.GetName()), InVariable.GetType().GetNameText(), FText::FromString(Collection->GetFullName())), InNode, InDefaultPin);
+				return false;
 			}
 		}
 
@@ -3523,6 +3525,12 @@ bool FHlslNiagaraTranslator::ParameterMapRegisterExternalConstantNamespaceVariab
 					if (Collection)
 					{
 						DataInterface = Collection->GetDefaultInstance()->GetParameterStore().GetDataInterface(InVariable);
+						if (DataInterface == nullptr)
+						{
+							Error(FText::Format(LOCTEXT("ParameterCollectionDataInterfaceNotFoundErrorFormat", "Data interface named {0} of type {1} was not found in Parameter Collection {2}"),
+								FText::FromName(InVariable.GetName()), InVariable.GetType().GetNameText(), FText::FromString(Collection->GetFullName())), InNode, InDefaultPin);
+							return false;
+						}
 					}
 					else
 					{
@@ -5415,7 +5423,7 @@ int32 FHlslNiagaraTranslator::CompileOutputPin(const UEdGraphPin* InPin)
 	return Ret;
 }
 
-void FHlslNiagaraTranslator::Error(FText ErrorText, UNiagaraNode* Node, UEdGraphPin* Pin)
+void FHlslNiagaraTranslator::Error(FText ErrorText, const UNiagaraNode* Node, const UEdGraphPin* Pin)
 {
 	FString NodePinStr = TEXT("");
 	FString NodePinPrefix = TEXT(" - ");
@@ -5436,7 +5444,7 @@ void FHlslNiagaraTranslator::Error(FText ErrorText, UNiagaraNode* Node, UEdGraph
 	TranslateResults.NumErrors++;
 }
 
-void FHlslNiagaraTranslator::Warning(FText WarningText, UNiagaraNode* Node, UEdGraphPin* Pin)
+void FHlslNiagaraTranslator::Warning(FText WarningText, const UNiagaraNode* Node, const UEdGraphPin* Pin)
 {
 	FString NodePinStr = TEXT("");
 	FString NodePinPrefix = TEXT(" - ");
diff --git a/Engine/Plugins/FX/Niagara/Source/NiagaraEditor/Private/NiagaraHlslTranslator.h b/Engine/Plugins/FX/Niagara/Source/NiagaraEditor/Private/NiagaraHlslTranslator.h
--- a/Engine/Plugins/FX/Niagara/Source/NiagaraEditor/Private/NiagaraHlslTranslator.h
+++ b/Engine/Plugins/FX/Niagara/Source/NiagaraEditor/Private/NiagaraHlslTranslator.h
@@ -448,8 +448,8 @@ public:
 	virtual void Convert(class UNiagaraNodeConvert* Convert, TArray <int32>& Inputs, TArray<int32>& Outputs);
 	virtual void If(TArray<FNiagaraVariable>& Vars, int32 Condition, TArray<int32>& PathA, TArray<int32>& PathB, TArray<int32>& Outputs);
 
-	virtual void Error(FText ErrorText, UNiagaraNode* Node, UEdGraphPin* Pin);
-	virtual void Warning(FText WarningText, UNiagaraNode* Node, UEdGraphPin* Pin);
+	virtual void Error(FText ErrorText, const UNiagaraNode* Node, const UEdGraphPin* Pin);
+	virtual void Warning(FText WarningText, const UNiagaraNode* Node, const UEdGraphPin* Pin);
 
 	virtual bool GetFunctionParameter(const FNiagaraVariable& Parameter, int32& OutParam)const;
 
