Resolved CVar name conflict.

diff --git a/Engine/Plugins/FX/Niagara/Source/Niagara/Private/NiagaraEmitterInstance.cpp b/Engine/Plugins/FX/Niagara/Source/Niagara/Private/NiagaraEmitterInstance.cpp
--- a/Engine/Plugins/FX/Niagara/Source/Niagara/Private/NiagaraEmitterInstance.cpp
+++ b/Engine/Plugins/FX/Niagara/Source/Niagara/Private/NiagaraEmitterInstance.cpp
@@ -36,11 +36,11 @@ static FAutoConsoleVariableRef CVarNiagaraDumpParticleData(
 TODO: This is mainly to avoid hard limits in our storage/alloc code etc rather than for perf reasons.
 We should improve our hard limit/safety code and possibly add a max for perf reasons.
 */
-static int32 GMaxCPUParticlesPerEmitter = 1000000;
-static FAutoConsoleVariableRef CVarMaxCPUParticlesPerEmitter(
-	TEXT("fx.MaxCPUParticlesPerEmitter"),
-	GMaxCPUParticlesPerEmitter,
-	TEXT("The max number of supported CPU particles per emitter. \n"),
+static int32 GMaxNiagaraCPUParticlesPerEmitter = 1000000;
+static FAutoConsoleVariableRef CVarMaxNiagaraCPUParticlesPerEmitter(
+	TEXT("fx.MaxNiagaraCPUParticlesPerEmitter"),
+	GMaxNiagaraCPUParticlesPerEmitter,
+	TEXT("The max number of supported CPU particles per emitter in Niagara. \n"),
 	ECVF_Default
 );
 //////////////////////////////////////////////////////////////////////////
@@ -1016,9 +1016,9 @@ void FNiagaraEmitterInstance::Tick(float DeltaSeconds)
 
 	//Ensure we don't blow our current hard limits on cpu particle count.
 	//TODO: These current limits can be improved relatively easily. Though perf in at these counts will obviously be an issue anyway.
-	if (CachedEmitter->SimTarget == ENiagaraSimTarget::CPUSim && AllocationSize > GMaxCPUParticlesPerEmitter)
+	if (CachedEmitter->SimTarget == ENiagaraSimTarget::CPUSim && AllocationSize > GMaxNiagaraCPUParticlesPerEmitter)
 	{
-		UE_LOG(LogNiagara, Warning, TEXT("Emitter %s has attemted to exceed the max CPU particle count! | Max: %d | Requested: %u"), *CachedEmitter->GetUniqueEmitterName(), GMaxCPUParticlesPerEmitter, AllocationSize);
+		UE_LOG(LogNiagara, Warning, TEXT("Emitter %s has attemted to exceed the max CPU particle count! | Max: %d | Requested: %u"), *CachedEmitter->GetUniqueEmitterName(), GMaxNiagaraCPUParticlesPerEmitter, AllocationSize);
 		//For now we completely bail out of spawning new particles. Possibly should improve this in future.
 		AllocationSize = OrigNumParticles;
 		SpawnTotal = 0;
