Fix for param store crash.
Not unbinding parameter stores correctly in some cases.

diff --git a/Engine/Plugins/FX/Niagara/Source/Niagara/Private/NiagaraParameterStore.cpp b/Engine/Plugins/FX/Niagara/Source/Niagara/Private/NiagaraParameterStore.cpp
--- a/Engine/Plugins/FX/Niagara/Source/Niagara/Private/NiagaraParameterStore.cpp
+++ b/Engine/Plugins/FX/Niagara/Source/Niagara/Private/NiagaraParameterStore.cpp
@@ -80,6 +80,7 @@ FNiagaraParameterStore::~FNiagaraParameterStore()
 	{
 		Binding.Value.Empty(Binding.Key, this);
 	}
+	Bindings.Empty();
 }
 
 void FNiagaraParameterStore::Bind(FNiagaraParameterStore* DestStore)
@@ -491,6 +492,7 @@ void FNiagaraParameterStore::Empty(bool bClearBindings)
 	DataInterfaces.Empty();
 	if (bClearBindings)
 	{
+		UnbindFromSourceStores();
 		Bindings.Empty();
 	}
 }
@@ -505,6 +507,7 @@ void FNiagaraParameterStore::Reset(bool bClearBindings)
 	DataInterfaces.Reset();
 	if (bClearBindings)
 	{
+		UnbindFromSourceStores();
 		Bindings.Reset();
 	}
 }
diff --git a/Engine/Plugins/FX/Niagara/Source/Niagara/Private/NiagaraScript.cpp b/Engine/Plugins/FX/Niagara/Source/Niagara/Private/NiagaraScript.cpp
--- a/Engine/Plugins/FX/Niagara/Source/Niagara/Private/NiagaraScript.cpp
+++ b/Engine/Plugins/FX/Niagara/Source/Niagara/Private/NiagaraScript.cpp
@@ -768,8 +768,8 @@ void UNiagaraScript::SetVMCompilationResults(const FNiagaraVMExecutableDataId& I
 void UNiagaraScript::InvalidateExecutionReadyParameterStores()
 {
 	// Make sure that we regenerate any parameter stores, since they must be kept in sync with the layout from script compilation.
-	ScriptExecutionParamStoreCPU.Clear();
-	ScriptExecutionParamStoreGPU.Clear();
+	ScriptExecutionParamStoreCPU.Empty();
+	ScriptExecutionParamStoreGPU.Empty();
 }
 
 void UNiagaraScript::InvalidateCachedCompileIds()
diff --git a/Engine/Plugins/FX/Niagara/Source/Niagara/Private/NiagaraScriptExecutionParameterStore.cpp b/Engine/Plugins/FX/Niagara/Source/Niagara/Private/NiagaraScriptExecutionParameterStore.cpp
--- a/Engine/Plugins/FX/Niagara/Source/Niagara/Private/NiagaraScriptExecutionParameterStore.cpp
+++ b/Engine/Plugins/FX/Niagara/Source/Niagara/Private/NiagaraScriptExecutionParameterStore.cpp
@@ -174,7 +174,7 @@ void FNiagaraScriptExecutionParameterStore::InitFromOwningScript(UNiagaraScript*
 
 void FNiagaraScriptExecutionParameterStore::InitFromOwningContext(UNiagaraScript* Script, ENiagaraSimTarget SimTarget, bool bNotifyAsDirty)
 {
-	Clear();
+	Empty();
 	PaddingInfo.Empty();
 
 #if WITH_EDITORONLY_DATA
diff --git a/Engine/Plugins/FX/Niagara/Source/Niagara/Public/NiagaraScriptExecutionParameterStore.h b/Engine/Plugins/FX/Niagara/Source/Niagara/Public/NiagaraScriptExecutionParameterStore.h
--- a/Engine/Plugins/FX/Niagara/Source/Niagara/Public/NiagaraScriptExecutionParameterStore.h
+++ b/Engine/Plugins/FX/Niagara/Source/Niagara/Public/NiagaraScriptExecutionParameterStore.h
@@ -63,6 +63,14 @@ public:
 		check(0);//Can't rename parameters for an execution store.
 	}
 
+	void Empty(bool bClearBindings=true)
+	{
+		FNiagaraParameterStore::Empty(bClearBindings);
+		PaddingInfo.Empty();
+		PaddedParameterSize = 0;
+		bInitialized = false;
+	}
+
 	// Just the external parameters, not previous or internal...
 	uint32 GetExternalParameterSize() { return ParameterSize; }
 
@@ -72,13 +80,6 @@ public:
 	// Helper that converts the data from the base type array internally into the padded out renderer-ready format.
 	void CopyParameterDataToPaddedBuffer(uint8* InTargetBuffer, uint32 InTargetBufferSizeInBytes);
 
-	void Clear()
-	{
-		Empty();
-		PaddingInfo.Empty();
-		PaddedParameterSize = 0;
-		bInitialized = false;
-	}
 
 	bool IsInitialized() const {return bInitialized;}
 
