Niagara - Prevent a crash when deleting modules and undoing that delete.

diff --git a/Engine/Plugins/FX/Niagara/Source/NiagaraEditor/Private/ViewModels/NiagaraSystemViewModel.cpp b/Engine/Plugins/FX/Niagara/Source/NiagaraEditor/Private/ViewModels/NiagaraSystemViewModel.cpp
--- a/Engine/Plugins/FX/Niagara/Source/NiagaraEditor/Private/ViewModels/NiagaraSystemViewModel.cpp
+++ b/Engine/Plugins/FX/Niagara/Source/NiagaraEditor/Private/ViewModels/NiagaraSystemViewModel.cpp
@@ -508,45 +508,54 @@ void FNiagaraSystemViewModel::PostUndo(bool bSuccess)
 
 void FNiagaraSystemViewModel::Tick(float DeltaTime)
 {
-	bool bRecompile = false;
-
-	check(SystemScriptViewModel.IsValid());
-	if (SystemScriptViewModel->GetLatestCompileStatus() == ENiagaraScriptCompileStatus::NCS_Dirty)
+	if (bForceAutoCompileOnce || (GetDefault<UNiagaraEditorSettings>()->GetAutoCompile() && bCanAutoCompile))
 	{
-		//SystemScriptViewModel->CompileSystem();
-		//UE_LOG(LogNiagaraEditor, Log, TEXT("Compiling %s due to dirty scripts."), *System.GetName());
-		bRecompile |= true;
-	}
+		bool bRecompile = false;
 
-	for (TSharedRef<FNiagaraEmitterHandleViewModel> EmitterHandleViewModel : EmitterHandleViewModels)
-	{
-		if (EmitterHandleViewModel->GetEmitterViewModel()->GetLatestCompileStatus() == ENiagaraScriptCompileStatus::NCS_Dirty)
+		check(SystemScriptViewModel.IsValid());
+		if (SystemScriptViewModel->GetLatestCompileStatus() == ENiagaraScriptCompileStatus::NCS_Dirty)
 		{
+			//SystemScriptViewModel->CompileSystem();
+			//UE_LOG(LogNiagaraEditor, Log, TEXT("Compiling %s due to dirty scripts."), *System.GetName());
 			bRecompile |= true;
-			//EmitterHandleViewModel->GetEmitterViewModel()->CompileScripts();
-			//UE_LOG(LogNiagaraEditor, Log, TEXT("Compiling %s - %s due to dirty scripts."), *System.GetName(), *EmitterHandleViewModel->GetName().ToString());
 		}
-	}
 
-	if (System.HasOutstandingCompilationRequests() == false)
-	{
-		if (bCompilePendingCompletion)
+		for (TSharedRef<FNiagaraEmitterHandleViewModel> EmitterHandleViewModel : EmitterHandleViewModels)
 		{
-			bCompilePendingCompletion = false;
-			OnSystemCompiled().Broadcast();
+			if (EmitterHandleViewModel->GetEmitterViewModel()->GetLatestCompileStatus() == ENiagaraScriptCompileStatus::NCS_Dirty)
+			{
+				bRecompile |= true;
+				//EmitterHandleViewModel->GetEmitterViewModel()->CompileScripts();
+				//UE_LOG(LogNiagaraEditor, Log, TEXT("Compiling %s - %s due to dirty scripts."), *System.GetName(), *EmitterHandleViewModel->GetName().ToString());
+			}
 		}
 
-		if (bRecompile || bForceAutoCompileOnce)
+		if (System.HasOutstandingCompilationRequests() == false)
 		{
-			CompileSystem(false);
-			bForceAutoCompileOnce = false;
-		}
+			if (bCompilePendingCompletion)
+			{
+				bCompilePendingCompletion = false;
+				OnSystemCompiled().Broadcast();
+			}
 
-		if (bResetRequestPending)
-		{
-			ResetSystem();
+			if (bRecompile || bForceAutoCompileOnce)
+			{
+				CompileSystem(false);
+				bForceAutoCompileOnce = false;
+			}
+
+			if (bResetRequestPending)
+			{
+				ResetSystem();
+			}
 		}
 	}
+
+	if (EmitterIdsRequiringSequencerTrackUpdate.Num() > 0)
+	{
+		UpdateSequencerTracksForEmitters(EmitterIdsRequiringSequencerTrackUpdate);
+		EmitterIdsRequiringSequencerTrackUpdate.Empty();
+	}
 }
 
 void FNiagaraSystemViewModel::OnPreSave()
@@ -608,11 +617,6 @@ const TArray<FNiagaraStackModuleData>& FNiagaraSystemViewModel::GetStackModuleDa
 	return EmitterToCachedStackModuleData[EmitterHandleId];
 }
 
-bool FNiagaraSystemViewModel::IsTickable() const
-{
-	return bForceAutoCompileOnce || (GetDefault<UNiagaraEditorSettings>()->GetAutoCompile() && bCanAutoCompile);
-}
-
 TStatId FNiagaraSystemViewModel::GetStatId() const
 {
 	RETURN_QUICK_DECLARE_CYCLE_STAT(FNiagaraSystemViewModel, STATGROUP_Tickables);
@@ -839,6 +843,23 @@ void FNiagaraSystemViewModel::RefreshSequencerTracks()
 	}
 
 	Sequencer->NotifyMovieSceneDataChanged(EMovieSceneDataChangeType::MovieSceneStructureItemsChanged);
+
+	// Since we just rebuilt all of the sequencer tracks, these updates don't need to be done.
+	EmitterIdsRequiringSequencerTrackUpdate.Empty();
+}
+
+void FNiagaraSystemViewModel::UpdateSequencerTracksForEmitters(const TArray<FGuid>& EmitterIdsRequiringUpdate)
+{
+	TGuardValue<bool> UpdateGuard(bUpdatingSequencerFromEmitterDataChange, true);
+	for (UMovieSceneTrack* Track : NiagaraSequence->GetMovieScene()->GetMasterTracks())
+	{
+		UMovieSceneNiagaraEmitterTrack* EmitterTrack = CastChecked<UMovieSceneNiagaraEmitterTrack>(Track);
+		if (EmitterIdsRequiringUpdate.Contains(EmitterTrack->GetEmitterHandleViewModel()->GetId()))
+		{
+			EmitterTrack->UpdateTrackFromEmitterGraphChange(NiagaraSequence->GetMovieScene()->GetTickResolution());
+		}
+	}
+	Sequencer->NotifyMovieSceneDataChanged(EMovieSceneDataChangeType::MovieSceneStructureItemsChanged);
 }
 
 UMovieSceneNiagaraEmitterTrack* FNiagaraSystemViewModel::GetTrackForHandleViewModel(TSharedRef<FNiagaraEmitterHandleViewModel> EmitterHandleViewModel)
@@ -1186,16 +1207,7 @@ void FNiagaraSystemViewModel::EmitterScriptGraphChanged(const FEdGraphEditAction
 {
 	if (bUpdatingEmittersFromSequencerDataChange == false)
 	{
-		TGuardValue<bool> UpdateGuard(bUpdatingSequencerFromEmitterDataChange, true);
-		for (UMovieSceneTrack* Track : NiagaraSequence->GetMovieScene()->GetMasterTracks())
-		{
-			UMovieSceneNiagaraEmitterTrack* EmitterTrack = CastChecked<UMovieSceneNiagaraEmitterTrack>(Track);
-			if (EmitterTrack->GetEmitterHandleViewModel() == OwningEmitterHandleViewModel)
-			{
-				EmitterTrack->UpdateTrackFromEmitterGraphChange(NiagaraSequence->GetMovieScene()->GetTickResolution());
-			}
-		}
-		Sequencer->NotifyMovieSceneDataChanged(EMovieSceneDataChangeType::MovieSceneStructureItemsChanged);
+		EmitterIdsRequiringSequencerTrackUpdate.AddUnique(OwningEmitterHandleViewModel->GetId());
 	}
 	// Remove from cache on graph change 
 	EmitterToCachedStackModuleData.Remove(OwningEmitterHandleViewModel->GetId());
diff --git a/Engine/Plugins/FX/Niagara/Source/NiagaraEditor/Private/ViewModels/Stack/NiagaraStackGraphUtilities.cpp b/Engine/Plugins/FX/Niagara/Source/NiagaraEditor/Private/ViewModels/Stack/NiagaraStackGraphUtilities.cpp
--- a/Engine/Plugins/FX/Niagara/Source/NiagaraEditor/Private/ViewModels/Stack/NiagaraStackGraphUtilities.cpp
+++ b/Engine/Plugins/FX/Niagara/Source/NiagaraEditor/Private/ViewModels/Stack/NiagaraStackGraphUtilities.cpp
@@ -870,7 +870,7 @@ bool FNiagaraStackGraphUtilities::RemoveModuleFromStack(UNiagaraScript& OwningSc
 	{
 		UNiagaraNode* NodeToRemove = NodesToCheck[0];
 		NodesToCheck.RemoveAt(0);
-		NodesToRemove.Add(NodeToRemove);
+		NodesToRemove.AddUnique(NodeToRemove);
 
 		TArray<UEdGraphPin*> InputPins;
 		NodeToRemove->GetInputPins(InputPins);
@@ -891,6 +891,7 @@ bool FNiagaraStackGraphUtilities::RemoveModuleFromStack(UNiagaraScript& OwningSc
 	UNiagaraGraph* Graph = ModuleNode.GetNiagaraGraph();
 	for (UNiagaraNode* NodeToRemove : NodesToRemove)
 	{
+		NodeToRemove->Modify();
 		Graph->RemoveNode(NodeToRemove);
 		UNiagaraNodeInput* InputNode = Cast<UNiagaraNodeInput>(NodeToRemove);
 		if (InputNode != nullptr)
diff --git a/Engine/Plugins/FX/Niagara/Source/NiagaraEditor/Public/ViewModels/NiagaraSystemViewModel.h b/Engine/Plugins/FX/Niagara/Source/NiagaraEditor/Public/ViewModels/NiagaraSystemViewModel.h
--- a/Engine/Plugins/FX/Niagara/Source/NiagaraEditor/Public/ViewModels/NiagaraSystemViewModel.h
+++ b/Engine/Plugins/FX/Niagara/Source/NiagaraEditor/Public/ViewModels/NiagaraSystemViewModel.h
@@ -167,7 +167,7 @@ public:
 
 	// ~ FTickableEditorObject
 	virtual void Tick(float DeltaTime) override;
-	virtual bool IsTickable() const override;
+	virtual bool IsTickable() const override { return true; }
 	virtual TStatId GetStatId() const override;
 
 	/** Resets the System instance to initial conditions. */
@@ -267,6 +267,9 @@ private:
 	/** Rebuilds the sequencer tracks. */
 	void RefreshSequencerTracks();
 
+	/** Updates the data in the sequencer tracks for the specified emitter ids. */
+	void UpdateSequencerTracksForEmitters(const TArray<FGuid>& EmitterIdsRequiringUpdate);
+
 	/** Gets the sequencer emitter track for the supplied emitter handle view model. */
 	UMovieSceneNiagaraEmitterTrack* GetTrackForHandleViewModel(TSharedRef<FNiagaraEmitterHandleViewModel> EmitterHandleViewModel);
 
@@ -490,4 +493,7 @@ private:
 	
 	/** A handle to the on graph changed delegate for the system script. */
 	FDelegateHandle SystemScriptGraphChangedHandler;
+
+	/** An array of emitter handle ids which need their sequencer tracks refreshed next frame. */
+	TArray<FGuid> EmitterIdsRequiringSequencerTrackUpdate;
 };
\ No newline at end of file
