Niagara - Move some transactions from UNiagaraGraph to the UI where the action is initiated to prevent undoable actions resulting from non-user actions such as merging and postload.

diff --git a/Engine/Plugins/FX/Niagara/Source/NiagaraEditor/Private/NiagaraGraph.cpp b/Engine/Plugins/FX/Niagara/Source/NiagaraEditor/Private/NiagaraGraph.cpp
--- a/Engine/Plugins/FX/Niagara/Source/NiagaraEditor/Private/NiagaraGraph.cpp
+++ b/Engine/Plugins/FX/Niagara/Source/NiagaraEditor/Private/NiagaraGraph.cpp
@@ -29,7 +29,6 @@
 #include "NiagaraNode.h"
 #include "EdGraphSchema_Niagara.h"
 #include "ViewModels/Stack/NiagaraParameterHandle.h"
-#include "ScopedTransaction.h"
 
 DECLARE_CYCLE_STAT(TEXT("NiagaraEditor - Graph - FindInputNodes"), STAT_NiagaraEditor_Graph_FindInputNodes, STATGROUP_NiagaraEditor);
 DECLARE_CYCLE_STAT(TEXT("NiagaraEditor - Graph - FindInputNodes_NotFilterUsage"), STAT_NiagaraEditor_Graph_FindInputNodes_NotFilterUsage, STATGROUP_NiagaraEditor);
@@ -538,7 +537,6 @@ void UNiagaraGraph::AddParameter(const FNiagaraVariable& Parameter)
 
 void UNiagaraGraph::RemoveParameter(const FNiagaraVariable& Parameter, const bool bNotifyGraphChanged /*= true*/)
 {
-	FScopedTransaction RemoveParametersWithPins(LOCTEXT("RemoveParametersWithPins", "Remove parameter and referenced pins"));
 	FNiagaraGraphParameterReferenceCollection* ReferenceCollection = Parameters.Find(Parameter);
 	if (ReferenceCollection)
 	{
@@ -579,8 +577,6 @@ bool UNiagaraGraph::RenameParameter(const FNiagaraVariable& Parameter, FName New
 	}
 	bIsRenamingParameter = true;
 
-	FScopedTransaction RemoveParametersWithPins(LOCTEXT("RenameParameter", "Rename parameter, referenced pins and metadata"));
-
 	// Prevent finding all parameters and metadata when renaming each pin.
 	SetFindParametersAllowed(false);
 	
diff --git a/Engine/Plugins/FX/Niagara/Source/NiagaraEditor/Private/Widgets/SNiagaraParameterMapPaletteItem.cpp b/Engine/Plugins/FX/Niagara/Source/NiagaraEditor/Private/Widgets/SNiagaraParameterMapPaletteItem.cpp
--- a/Engine/Plugins/FX/Niagara/Source/NiagaraEditor/Private/Widgets/SNiagaraParameterMapPaletteItem.cpp
+++ b/Engine/Plugins/FX/Niagara/Source/NiagaraEditor/Private/Widgets/SNiagaraParameterMapPaletteItem.cpp
@@ -8,6 +8,7 @@
 #include "NiagaraEditorStyle.h"
 #include "Widgets/Input/SComboButton.h"
 #include "Widgets/Images/SImage.h"
+#include "ScopedTransaction.h"
 
 #define LOCTEXT_NAMESPACE "NiagaraParameterMapPalleteItem"
 
@@ -87,6 +88,7 @@ void SNiagaraParameterMapPalleteItem::Construct(const FArguments& InArgs, FCreat
 void SNiagaraParameterMapPalleteItem::OnNameTextCommitted(const FText& NewText, ETextCommit::Type InTextCommit)
 {
 	TSharedPtr<FNiagaraParameterAction> ParameterAction = StaticCastSharedPtr<FNiagaraParameterAction>(ActionPtr.Pin());
+	FScopedTransaction RenameParametersWithPins(LOCTEXT("RenameParameter", "Rename parameter, referenced pins and metadata"));
 	OnItemRenamed.ExecuteIfBound(NewText, *ParameterAction.Get());
 }
 
diff --git a/Engine/Plugins/FX/Niagara/Source/NiagaraEditor/Private/Widgets/SNiagaraParameterMapView.cpp b/Engine/Plugins/FX/Niagara/Source/NiagaraEditor/Private/Widgets/SNiagaraParameterMapView.cpp
--- a/Engine/Plugins/FX/Niagara/Source/NiagaraEditor/Private/Widgets/SNiagaraParameterMapView.cpp
+++ b/Engine/Plugins/FX/Niagara/Source/NiagaraEditor/Private/Widgets/SNiagaraParameterMapView.cpp
@@ -32,6 +32,7 @@
 #include "EdGraph/EdGraphPin.h"
 #include "EdGraphSchema_Niagara.h"
 #include "ViewModels/NiagaraEmitterHandleViewModel.h"
+#include "ScopedTransaction.h"
 
 #define LOCTEXT_NAMESPACE "NiagaraParameterMapView"
 
@@ -603,6 +604,7 @@ void SNiagaraParameterMapView::OnDeleteEntry()
 				}
 			}
 
+			FScopedTransaction RemoveParametersWithPins(LOCTEXT("RemoveParametersWithPins", "Remove parameter and referenced pins"));
 			for (const TWeakObjectPtr<UNiagaraGraph>& GraphWeakPtr : Graphs)
 			{
 				if (GraphWeakPtr.IsValid())
