Niagara - Prevent a crash when adding and emitter to a system, and then undoing that add.

diff --git a/Engine/Plugins/FX/Niagara/Source/NiagaraEditorWidgets/Private/DetailCustomizations/NiagaraDataInterfaceDetails.cpp b/Engine/Plugins/FX/Niagara/Source/NiagaraEditorWidgets/Private/DetailCustomizations/NiagaraDataInterfaceDetails.cpp
--- a/Engine/Plugins/FX/Niagara/Source/NiagaraEditorWidgets/Private/DetailCustomizations/NiagaraDataInterfaceDetails.cpp
+++ b/Engine/Plugins/FX/Niagara/Source/NiagaraEditorWidgets/Private/DetailCustomizations/NiagaraDataInterfaceDetails.cpp
@@ -118,11 +118,15 @@ class FNiagaraDataInterfaceCustomNodeBuilder : public IDetailCustomNodeBuilder
 											 , public TSharedFromThis<FNiagaraDataInterfaceCustomNodeBuilder>
 {
 public:
-	FNiagaraDataInterfaceCustomNodeBuilder(UNiagaraDataInterface* InDataInterface, IDetailLayoutBuilder* InDetailBuilder)
-		: DataInterface(InDataInterface)
-		, DetailBuilder(InDetailBuilder)
+	FNiagaraDataInterfaceCustomNodeBuilder(IDetailLayoutBuilder* InDetailBuilder)
+		: DetailBuilder(InDetailBuilder)
 	{
-		DataInterface->OnChanged().AddRaw(this, &FNiagaraDataInterfaceCustomNodeBuilder::OnDataInterfaceChanged);
+	}
+
+	void Initialize(UNiagaraDataInterface& InDataInterface)
+	{
+		DataInterface = &InDataInterface;
+		DataInterface->OnChanged().AddSP(this, &FNiagaraDataInterfaceCustomNodeBuilder::OnDataInterfaceChanged);
 	}
 
 	~FNiagaraDataInterfaceCustomNodeBuilder()
@@ -201,7 +205,8 @@ void FNiagaraDataInterfaceDetailsBase::CustomizeDetails(IDetailLayoutBuilder& De
 	DataInterface->OnChanged().AddRaw(this, &FNiagaraDataInterfaceDetailsBase::OnDataChanged);
 	IDetailCategoryBuilder& ErrorsBuilderRef = DetailBuilder.EditCategory(ErrorsCategoryName, LOCTEXT("Errors", "Errors"), ECategoryPriority::Important);
 	ErrorsCategoryBuilder = &ErrorsBuilderRef;
-	CustomBuilder = TSharedPtr<FNiagaraDataInterfaceCustomNodeBuilder>(new FNiagaraDataInterfaceCustomNodeBuilder(DataInterface.Get(), &DetailBuilder));
+	CustomBuilder = MakeShared<FNiagaraDataInterfaceCustomNodeBuilder>(&DetailBuilder);
+	CustomBuilder->Initialize(*DataInterface);
 	ErrorsCategoryBuilder->AddCustomBuilder(CustomBuilder.ToSharedRef());
 	OnDataChanged();
 }
