Niagara - Fix crash on redo when adding an emitter an doing and undo and redo.

diff --git a/Engine/Plugins/FX/Niagara/Source/NiagaraEditorWidgets/Private/DetailCustomizations/NiagaraDataInterfaceDetails.cpp b/Engine/Plugins/FX/Niagara/Source/NiagaraEditorWidgets/Private/DetailCustomizations/NiagaraDataInterfaceDetails.cpp
--- a/Engine/Plugins/FX/Niagara/Source/NiagaraEditorWidgets/Private/DetailCustomizations/NiagaraDataInterfaceDetails.cpp
+++ b/Engine/Plugins/FX/Niagara/Source/NiagaraEditorWidgets/Private/DetailCustomizations/NiagaraDataInterfaceDetails.cpp
@@ -202,7 +202,7 @@ void FNiagaraDataInterfaceDetailsBase::CustomizeDetails(IDetailLayoutBuilder& De
 	check(SelectedObjects.Num() == 1);
 	DataInterface = Cast<UNiagaraDataInterface>(SelectedObjects[0].Get());
 	check(DataInterface.IsValid());
-	DataInterface->OnChanged().AddRaw(this, &FNiagaraDataInterfaceDetailsBase::OnDataChanged);
+	DataInterface->OnChanged().AddSP(this, &FNiagaraDataInterfaceDetailsBase::OnDataChanged);
 	IDetailCategoryBuilder& ErrorsBuilderRef = DetailBuilder.EditCategory(ErrorsCategoryName, LOCTEXT("Errors", "Errors"), ECategoryPriority::Important);
 	ErrorsCategoryBuilder = &ErrorsBuilderRef;
 	CustomBuilder = MakeShared<FNiagaraDataInterfaceCustomNodeBuilder>(&DetailBuilder);
