Niagara - Fix a crash when dragging a function input value, and then immiediately assigning that value to be linked.

diff --git a/Engine/Plugins/FX/Niagara/Source/NiagaraEditorWidgets/Private/Stack/SNiagaraStackFunctionInputValue.cpp b/Engine/Plugins/FX/Niagara/Source/NiagaraEditorWidgets/Private/Stack/SNiagaraStackFunctionInputValue.cpp
--- a/Engine/Plugins/FX/Niagara/Source/NiagaraEditorWidgets/Private/Stack/SNiagaraStackFunctionInputValue.cpp
+++ b/Engine/Plugins/FX/Niagara/Source/NiagaraEditorWidgets/Private/Stack/SNiagaraStackFunctionInputValue.cpp
@@ -244,7 +244,7 @@ TSharedRef<SWidget> SNiagaraStackFunctionInputValue::ConstructLocalValueStructWi
 			ParameterEditor->SetOnEndValueChange(SNiagaraParameterEditor::FOnValueChange::CreateSP(
 				this, &SNiagaraStackFunctionInputValue::ParameterEndValueChange));
 			ParameterEditor->SetOnValueChanged(SNiagaraParameterEditor::FOnValueChange::CreateSP(
-				this, &SNiagaraStackFunctionInputValue::ParameterValueChanged, ParameterEditor.ToSharedRef()));
+				this, &SNiagaraStackFunctionInputValue::ParameterValueChanged, TWeakPtr<SNiagaraParameterEditor>(ParameterEditor)));
 
 			LocalValueStructParameterEditor = ParameterEditor;
 
@@ -309,10 +309,14 @@ void SNiagaraStackFunctionInputValue::ParameterEndValueChange()
 	FunctionInput->NotifyEndLocalValueChange();
 }
 
-void SNiagaraStackFunctionInputValue::ParameterValueChanged(TSharedRef<SNiagaraParameterEditor> ParameterEditor)
+void SNiagaraStackFunctionInputValue::ParameterValueChanged(TWeakPtr<SNiagaraParameterEditor> ParameterEditor)
 {
-	ParameterEditor->UpdateStructFromInternalValue(FunctionInput->GetLocalValueStruct().ToSharedRef());
-	FunctionInput->SetLocalValue(DisplayedLocalValueStruct.ToSharedRef());
+	TSharedPtr<SNiagaraParameterEditor> ParameterEditorPinned = ParameterEditor.Pin();
+	if (ParameterEditorPinned.IsValid())
+	{
+		ParameterEditorPinned->UpdateStructFromInternalValue(DisplayedLocalValueStruct.ToSharedRef());
+		FunctionInput->SetLocalValue(DisplayedLocalValueStruct.ToSharedRef());
+	}
 }
 
 void SNiagaraStackFunctionInputValue::ParameterPropertyValueChanged(const FPropertyChangedEvent& PropertyChangedEvent)
diff --git a/Engine/Plugins/FX/Niagara/Source/NiagaraEditorWidgets/Private/Stack/SNiagaraStackFunctionInputValue.h b/Engine/Plugins/FX/Niagara/Source/NiagaraEditorWidgets/Private/Stack/SNiagaraStackFunctionInputValue.h
--- a/Engine/Plugins/FX/Niagara/Source/NiagaraEditorWidgets/Private/Stack/SNiagaraStackFunctionInputValue.h
+++ b/Engine/Plugins/FX/Niagara/Source/NiagaraEditorWidgets/Private/Stack/SNiagaraStackFunctionInputValue.h
@@ -39,7 +39,7 @@ private:
 
 	void ParameterEndValueChange();
 
-	void ParameterValueChanged(TSharedRef<SNiagaraParameterEditor> ParameterEditor);
+	void ParameterValueChanged(TWeakPtr<SNiagaraParameterEditor> ParameterEditor);
 
 	void ParameterPropertyValueChanged(const FPropertyChangedEvent& PropertyChangedEvent);
 
