Niagara - Prevent renaming and deleting parameters from the emitter and system level from renaming and deleting parameters in the module and dynamic input script assets.

diff --git a/Engine/Plugins/FX/Niagara/Source/NiagaraEditor/Private/NiagaraGraph.cpp b/Engine/Plugins/FX/Niagara/Source/NiagaraEditor/Private/NiagaraGraph.cpp
--- a/Engine/Plugins/FX/Niagara/Source/NiagaraEditor/Private/NiagaraGraph.cpp
+++ b/Engine/Plugins/FX/Niagara/Source/NiagaraEditor/Private/NiagaraGraph.cpp
@@ -547,7 +547,7 @@ void UNiagaraGraph::RemoveParameter(const FNiagaraVariable& Parameter, const boo
 		{
 			const FNiagaraGraphParameterReference& Reference = ReferenceCollection->ParameterReferences[Index];
 			UNiagaraNode* Node = Reference.Value.Get();
-			if (Node)
+			if (Node && Node->GetGraph() == this)
 			{
 				UEdGraphPin* Pin = Node->GetPinByPersistentGuid(Reference.Key);
 				if (Pin)
@@ -592,7 +592,7 @@ bool UNiagaraGraph::RenameParameter(const FNiagaraVariable& Parameter, FName New
 		for (FNiagaraGraphParameterReference& Reference : NewReferences.ParameterReferences)
 		{
 			UNiagaraNode* Node = Reference.Value.Get();
-			if (Node)
+			if (Node && Node->GetGraph() == this)
 			{
 				UEdGraphPin* Pin = Node->GetPinByPersistentGuid(Reference.Key);
 				if (Pin)
