Fixes two simultaneous issues.
1) Dynamic inputs and modules should only allow parameter maps and data interfaces as inputs. Functions allow any parameter as input. Dynamic inputs should allow any non-data interface type as the output. Modules should only allow a parameter map as output. Functions can output any non-data interface type. This should now match what we can create via right-click menu in the graph.
2) We weren't modifying the name of the existing pin when editing an output parameter name in the script details. This causes us to lose any connections to the pin.

#tests did halo testing around connecting output pins, created dynamic input/module/function graphs and made sure that the options matched expectations

diff --git a/Engine/Plugins/FX/Niagara/Source/NiagaraEditor/Private/ViewModels/NiagaraScriptInputCollectionViewModel.cpp b/Engine/Plugins/FX/Niagara/Source/NiagaraEditor/Private/ViewModels/NiagaraScriptInputCollectionViewModel.cpp
--- a/Engine/Plugins/FX/Niagara/Source/NiagaraEditor/Private/ViewModels/NiagaraScriptInputCollectionViewModel.cpp
+++ b/Engine/Plugins/FX/Niagara/Source/NiagaraEditor/Private/ViewModels/NiagaraScriptInputCollectionViewModel.cpp
@@ -374,6 +374,18 @@ void FNiagaraScriptInputCollectionViewModel::RefreshParameterViewModels()
 
 bool FNiagaraScriptInputCollectionViewModel::SupportsType(const FNiagaraTypeDefinition& Type) const
 {
+	if (Scripts.Num() == 1 && Scripts[0] != nullptr)
+	{
+		// We only support parameter map inputs for dynamic inputs and modules, with the ability to create data interfaces as needed for defaults.
+		if (Scripts[0]->GetUsage() == ENiagaraScriptUsage::DynamicInput || Scripts[0]->GetUsage() == ENiagaraScriptUsage::Module)
+		{
+			if (Type != FNiagaraTypeDefinition::GetParameterMapDef() && !Type.IsDataInterface())
+			{
+				return false;
+			}
+		}
+	}
+
 	return bCanHaveNumericParameters || Type != FNiagaraTypeDefinition::GetGenericNumericDef();
 }
 
diff --git a/Engine/Plugins/FX/Niagara/Source/NiagaraEditor/Private/ViewModels/NiagaraScriptOutputCollectionViewModel.cpp b/Engine/Plugins/FX/Niagara/Source/NiagaraEditor/Private/ViewModels/NiagaraScriptOutputCollectionViewModel.cpp
--- a/Engine/Plugins/FX/Niagara/Source/NiagaraEditor/Private/ViewModels/NiagaraScriptOutputCollectionViewModel.cpp
+++ b/Engine/Plugins/FX/Niagara/Source/NiagaraEditor/Private/ViewModels/NiagaraScriptOutputCollectionViewModel.cpp
@@ -272,6 +272,17 @@ void FNiagaraScriptOutputCollectionViewModel::RefreshParameterViewModels()
 
 bool FNiagaraScriptOutputCollectionViewModel::SupportsType(const FNiagaraTypeDefinition& Type) const
 {
+	if (Scripts.Num() == 1 && Scripts[0] != nullptr)
+	{
+		// We only support parameter map outputs for modules.
+		if (Scripts[0]->GetUsage() == ENiagaraScriptUsage::Module)
+		{
+			if (Type != FNiagaraTypeDefinition::GetParameterMapDef())
+			{
+				return false;
+			}
+		}
+	}
 	return Type.GetScriptStruct() != nullptr && (bCanHaveNumericParameters || Type != FNiagaraTypeDefinition::GetGenericNumericDef());
 }
 
@@ -307,6 +318,16 @@ void FNiagaraScriptOutputCollectionViewModel::OnParameterNameChanged(FName OldNa
 	{
 		ParameterVariable->SetName(FNiagaraUtilities::GetUniqueName(ParameterVariable->GetName(), CurrentNames));
 	}
+
+	// Make sure to edit the pin name so that when we regenerate the node the 
+	// links are preserved.
+	UEdGraphPin* FoundPin = OutputNode->FindPin(OldName, EEdGraphPinDirection::EGPD_Input);
+	if (FoundPin)
+	{
+		FoundPin->Modify();
+		FoundPin->PinName = NewName;
+	}
+
 	
 	OutputNode->NotifyOutputVariablesChanged();
 
