Niagara - Fix an issue where the interpolated spawn script was using a stale version of a data interface defined in the update script.

diff --git a/Engine/Plugins/FX/Niagara/Source/Niagara/Private/NiagaraSystem.cpp b/Engine/Plugins/FX/Niagara/Source/Niagara/Private/NiagaraSystem.cpp
--- a/Engine/Plugins/FX/Niagara/Source/Niagara/Private/NiagaraSystem.cpp
+++ b/Engine/Plugins/FX/Niagara/Source/Niagara/Private/NiagaraSystem.cpp
@@ -813,6 +813,7 @@ bool UNiagaraSystem::QueryCompileComplete(bool bWait, bool bDoPost, bool bDoNotA
 
 		InitEmitterSpawnAttributes();
 
+		// Prepare rapid iteration parameters for execution.
 		TArray<UNiagaraScript*> Scripts;
 		TMap<UNiagaraScript*, UNiagaraScript*> ScriptDependencyMap;
 		TMap<UNiagaraScript*, FString> ScriptToEmitterNameMap;
@@ -824,21 +825,21 @@ bool UNiagaraSystem::QueryCompileComplete(bool bWait, bool bDoPost, bool bDoNotA
 			Scripts.AddUnique(CompiledScript);
 			ScriptToEmitterNameMap.Add(CompiledScript, Emitter != nullptr ? Emitter->GetUniqueEmitterName() : FString());
 
-			if (CompiledScript->GetUsage() == ENiagaraScriptUsage::EmitterSpawnScript)
+			if (UNiagaraScript::IsEquivalentUsage(CompiledScript->GetUsage(), ENiagaraScriptUsage::EmitterSpawnScript))
 			{
 				Scripts.AddUnique(SystemSpawnScript);
 				ScriptDependencyMap.Add(CompiledScript, SystemSpawnScript);
 				ScriptToEmitterNameMap.Add(SystemSpawnScript, FString());
 			}
 
-			if (CompiledScript->GetUsage() == ENiagaraScriptUsage::EmitterUpdateScript)
+			if (UNiagaraScript::IsEquivalentUsage(CompiledScript->GetUsage(), ENiagaraScriptUsage::EmitterUpdateScript))
 			{
 				Scripts.AddUnique(SystemUpdateScript);
 				ScriptDependencyMap.Add(CompiledScript, SystemUpdateScript);
 				ScriptToEmitterNameMap.Add(SystemSpawnScript, FString());
 			}
 
-			if (CompiledScript->GetUsage() == ENiagaraScriptUsage::ParticleSpawnScript)
+			if (UNiagaraScript::IsEquivalentUsage(CompiledScript->GetUsage(), ENiagaraScriptUsage::ParticleSpawnScript))
 			{
 				if (Emitter && Emitter->SimTarget == ENiagaraSimTarget::GPUComputeSim)
 				{
@@ -848,7 +849,7 @@ bool UNiagaraSystem::QueryCompileComplete(bool bWait, bool bDoPost, bool bDoNotA
 				}
 			}
 
-			if (CompiledScript->GetUsage() == ENiagaraScriptUsage::ParticleUpdateScript)
+			if (UNiagaraScript::IsEquivalentUsage(CompiledScript->GetUsage(), ENiagaraScriptUsage::ParticleUpdateScript))
 			{
 				if (Emitter && Emitter->SimTarget == ENiagaraSimTarget::GPUComputeSim)
 				{
@@ -867,6 +868,37 @@ bool UNiagaraSystem::QueryCompileComplete(bool bWait, bool bDoPost, bool bDoNotA
 
 		FNiagaraUtilities::PrepareRapidIterationParameters(Scripts, ScriptDependencyMap, ScriptToEmitterNameMap);
 
+		// HACK: This is a temporary hack to fix an issue where data interfaces used by modules and dynamic inputs in the
+		// particle update script aren't being shared by the interpolated spawn script when accessed directly.  This works
+		// properly if the data interface is assigned to a named particle parameter and then linked to an input.
+		// TODO: Bind these data interfaces the same way parameter data interfaces are bound.
+		for (FEmitterCompiledScriptPair& EmitterCompiledScriptPair : ActiveCompilations[ActiveCompileIdx].EmitterCompiledScriptPairs)
+		{
+			UNiagaraEmitter* Emitter = EmitterCompiledScriptPair.Emitter;
+			UNiagaraScript* CompiledScript = EmitterCompiledScriptPair.CompiledScript;
+
+			if (UNiagaraScript::IsEquivalentUsage(CompiledScript->GetUsage(), ENiagaraScriptUsage::ParticleUpdateScript))
+			{
+				UNiagaraScript* SpawnScript = Emitter->SpawnScriptProps.Script;
+				for (const FNiagaraScriptDataInterfaceInfo& UpdateDataInterfaceInfo : CompiledScript->GetCachedDefaultDataInterfaces())
+				{
+					if (UpdateDataInterfaceInfo.RegisteredParameterMapRead == NAME_None && UpdateDataInterfaceInfo.RegisteredParameterMapWrite == NAME_None)
+					{
+						// If the data interface isn't being read or written to a parameter map then it won't be bound properly so we
+						// assign the update scripts copy of the data interface to the spawn scripts copy by pointer so that they will share
+						// the data interface at runtime and will both be updated in the editor.
+						for (FNiagaraScriptDataInterfaceInfo& SpawnDataInterfaceInfo : SpawnScript->GetCachedDefaultDataInterfaces())
+						{
+							if (UpdateDataInterfaceInfo.Name == SpawnDataInterfaceInfo.Name)
+							{
+								SpawnDataInterfaceInfo.DataInterface = UpdateDataInterfaceInfo.DataInterface;
+							}
+						}
+					}
+				}
+			}
+		}
+
 		ActiveCompilations[ActiveCompileIdx].RootObjects.Empty();
 
 		DetermineIfSolo();
