Niagara - Submitting for Brad - Prevent a default sprite renderer from getting added at the system level when adding an emitter with no renderers.

diff --git a/Engine/Plugins/FX/Niagara/Source/Niagara/Private/NiagaraEmitter.cpp b/Engine/Plugins/FX/Niagara/Source/Niagara/Private/NiagaraEmitter.cpp
--- a/Engine/Plugins/FX/Niagara/Source/Niagara/Private/NiagaraEmitter.cpp
+++ b/Engine/Plugins/FX/Niagara/Source/Niagara/Private/NiagaraEmitter.cpp
@@ -97,8 +97,6 @@ void UNiagaraEmitter::PostInitProperties()
 	Super::PostInitProperties();
 	if (HasAnyFlags(RF_ClassDefaultObject | RF_NeedLoad) == false)
 	{
-		RendererProperties.Add(NewObject<UNiagaraSpriteRendererProperties>(this, "Renderer"));
-
 		SpawnScriptProps.Script = NewObject<UNiagaraScript>(this, "SpawnScript", EObjectFlags::RF_Transactional);
 		SpawnScriptProps.Script->SetUsage(ENiagaraScriptUsage::ParticleSpawnScript);
 
diff --git a/Engine/Plugins/FX/Niagara/Source/Niagara/Public/NiagaraSpriteRendererProperties.h b/Engine/Plugins/FX/Niagara/Source/Niagara/Public/NiagaraSpriteRendererProperties.h
--- a/Engine/Plugins/FX/Niagara/Source/Niagara/Public/NiagaraSpriteRendererProperties.h
+++ b/Engine/Plugins/FX/Niagara/Source/Niagara/Public/NiagaraSpriteRendererProperties.h
@@ -40,7 +40,7 @@ enum class ENiagaraSpriteFacingMode : uint8
 
 
 UCLASS(editinlinenew)
-class UNiagaraSpriteRendererProperties : public UNiagaraRendererProperties
+class NIAGARA_API UNiagaraSpriteRendererProperties : public UNiagaraRendererProperties
 {
 public:
 	GENERATED_BODY()
diff --git a/Engine/Plugins/FX/Niagara/Source/NiagaraEditor/Private/NiagaraEmitterFactoryNew.cpp b/Engine/Plugins/FX/Niagara/Source/NiagaraEditor/Private/NiagaraEmitterFactoryNew.cpp
--- a/Engine/Plugins/FX/Niagara/Source/NiagaraEditor/Private/NiagaraEmitterFactoryNew.cpp
+++ b/Engine/Plugins/FX/Niagara/Source/NiagaraEditor/Private/NiagaraEmitterFactoryNew.cpp
@@ -6,6 +6,7 @@
 #include "NiagaraScriptFactoryNew.h"
 #include "NiagaraGraph.h"
 #include "NiagaraScriptSource.h"
+#include "NiagaraSpriteRendererProperties.h"
 #include "NiagaraNodeOutput.h"
 #include "NiagaraEditorSettings.h"
 #include "AssetData.h"
@@ -45,6 +46,8 @@ UObject* UNiagaraEmitterFactoryNew::FactoryCreateNew(UClass* Class, UObject* InP
 
 		NewEmitter = NewObject<UNiagaraEmitter>(InParent, Class, Name, Flags | RF_Transactional);
 
+		NewEmitter->AddRenderer(NewObject<UNiagaraSpriteRendererProperties>(NewEmitter, "Renderer"));
+
 		UNiagaraScriptSource* Source = NewObject<UNiagaraScriptSource>(NewEmitter, NAME_None, RF_Transactional);
 		if (Source)
 		{
