Niagara - Fix a crash when sampling a data interface direction in a script.

diff --git a/Engine/Plugins/FX/Niagara/Source/NiagaraEditor/Private/NiagaraHlslTranslator.cpp b/Engine/Plugins/FX/Niagara/Source/NiagaraEditor/Private/NiagaraHlslTranslator.cpp
--- a/Engine/Plugins/FX/Niagara/Source/NiagaraEditor/Private/NiagaraHlslTranslator.cpp
+++ b/Engine/Plugins/FX/Niagara/Source/NiagaraEditor/Private/NiagaraHlslTranslator.cpp
@@ -4407,6 +4407,13 @@ void FHlslNiagaraTranslator::FunctionCall(UNiagaraNodeFunctionCall* FunctionNode
 	}
 	
 	RegisterFunctionCall(ScriptUsage, Name, FullName,  FunctionNode->NodeGuid, Source, Signature, bCustomHlsl, CustomHlsl, Inputs, CallInputs, CallOutputs, OutputSignature);
+
+	if (OutputSignature.IsValid() == false)
+	{
+		Error(LOCTEXT("FunctionCallInvalidSignature", "Could not generate a valid function signature."), FunctionNode, nullptr);
+		return;
+	}
+
 	GenerateFunctionCall(OutputSignature, Inputs, Outputs);
 
 	if (bCustomHlsl)
@@ -4846,7 +4853,15 @@ void FHlslNiagaraTranslator::RegisterFunctionCall(ENiagaraScriptUsage ScriptUsag
 			// interface. It could be that the existing node has been removed and the graph
 			// needs to be refactored. If that's the case, emit an error.
 			UObject* const* FoundCDO = CompileData->CDOs.Find(Info.Type.GetClass());
-			check(FoundCDO != nullptr);
+			if (FoundCDO == nullptr)
+			{
+				// If the cdo wasn't found, the data interface was not passed through a parameter map and so it won't be bound correctly, so add a compile error
+				// and invalidate the signature.
+				Error(LOCTEXT("DataInterfaceNotFoundInParameterMap", "Data interfaces can not be sampled directly, they must be passed through a parameter map to be bound correctly."), nullptr, nullptr);
+				OutSignature.Name = NAME_None;
+				return;
+			}
+
 			UNiagaraDataInterface* CDO = Cast<UNiagaraDataInterface>(*FoundCDO);
 			if (CDO != nullptr && OutSignature.bMemberFunction)
 			{
@@ -5428,10 +5443,23 @@ void FHlslNiagaraTranslator::Error(FText ErrorText, const UNiagaraNode* Node, co
 	FString NodePinStr = TEXT("");
 	FString NodePinPrefix = TEXT(" - ");
 	FString NodePinSuffix = TEXT("");
-	if (Node && Node->GetName().Len() > 0)
+	if (Node)
 	{
-		NodePinStr += TEXT("Node: ") + Node->GetName();
-		NodePinSuffix = TEXT(" - ");
+		FString NodeTitle = Node->GetNodeTitle(ENodeTitleType::FullTitle).ToString();
+		if (NodeTitle.Len() > 0)
+		{
+			NodePinStr += TEXT("Node: ") + NodeTitle;
+			NodePinSuffix = TEXT(" - ");
+		}
+		else
+		{
+			FString NodeName = Node->GetName();
+			if (NodeName.Len() > 0)
+			{
+				NodePinStr += TEXT("Node: ") + NodeName;
+				NodePinSuffix = TEXT(" - ");
+			}
+		}
 	}
 	if (Pin && Pin->PinFriendlyName.ToString().Len() > 0)
 	{
