Niagara - Prevent a crash when deleting user paramters with data interfaces.

diff --git a/Engine/Plugins/FX/Niagara/Source/NiagaraEditor/Private/ViewModels/Stack/NiagaraStackParameterStoreEntry.cpp b/Engine/Plugins/FX/Niagara/Source/NiagaraEditor/Private/ViewModels/Stack/NiagaraStackParameterStoreEntry.cpp
--- a/Engine/Plugins/FX/Niagara/Source/NiagaraEditor/Private/ViewModels/Stack/NiagaraStackParameterStoreEntry.cpp
+++ b/Engine/Plugins/FX/Niagara/Source/NiagaraEditor/Private/ViewModels/Stack/NiagaraStackParameterStoreEntry.cpp
@@ -253,8 +253,8 @@ void UNiagaraStackParameterStoreEntry::RenameInput(FString NewName)
 			}
 		}
 
-		GEditor->BeginTransaction(LOCTEXT("RenameParameter", "Rename Parameter"));
-		ParameterStore->GetOwner()->Modify();
+		FScopedTransaction ScopedTransaction(LOCTEXT("RenameUserParameter", "Rename user parameter"));
+		Owner->Modify();
 		// remove old one (a bit overkill but it beats duplicating code)
 		RemovePins(OwningPins);
 		// TODO Would it be better to actually keep variable name being ActualName and ParameterName being ParameterHandle.GetParameterHandleString(), and rewrite the way the Entry is built?
@@ -268,17 +268,22 @@ void UNiagaraStackParameterStoreEntry::RenameInput(FString NewName)
 			// generate current link
 			FNiagaraStackGraphUtilities::SetLinkedValueHandleForFunctionInput(*LinkedPin, ParameterHandle);
 		}
-		GEditor->EndTransaction();
+
+		ParameterName = VariableName;
+		DisplayName = FText::FromName(ParameterName);
 	}
 }
 
 void UNiagaraStackParameterStoreEntry::Delete()
 {
+	FScopedTransaction ScopedTransaction(LOCTEXT("RemoveUserParameter", "Remove user parameter"));
+
 	// for delete, do a parameter map traversal to deduce all the usages and then  remove them 
 	TArray<UEdGraphPin*> OwningPins = GetOwningPins();
-	NotifyBeginValueChange();
 	RemovePins(OwningPins);
+
 	//remove from store
+	Owner->Modify();
 	ParameterStore->RemoveParameter(FNiagaraVariable(InputType, ParameterName));
 	if (InputType.GetClass() != nullptr)
 	{
@@ -288,7 +293,8 @@ void UNiagaraStackParameterStoreEntry::Delete()
 			GetSystemViewModel()->NotifyDataObjectChanged(DataInterface);
 		}
 	}
-	NotifyEndValueChange();
+
+	ParameterDeletedDelegate.Broadcast();
 }
 
 void UNiagaraStackParameterStoreEntry::RemovePins(TArray<UEdGraphPin*> OwningPins /*, bool bSetPreviousValue*/)
@@ -334,6 +340,11 @@ UNiagaraStackParameterStoreEntry::FOnValueChanged& UNiagaraStackParameterStoreEn
 	return ValueChangedDelegate;
 }
 
+UNiagaraStackParameterStoreEntry::FOnParameterDeleted& UNiagaraStackParameterStoreEntry::OnParameterDeleted()
+{
+	return ParameterDeletedDelegate;
+}
+
 TSharedPtr<FNiagaraVariable> UNiagaraStackParameterStoreEntry::GetCurrentValueVariable()
 {
 	if (InputType.GetClass() == nullptr)
diff --git a/Engine/Plugins/FX/Niagara/Source/NiagaraEditor/Private/ViewModels/Stack/NiagaraStackParameterStoreGroup.cpp b/Engine/Plugins/FX/Niagara/Source/NiagaraEditor/Private/ViewModels/Stack/NiagaraStackParameterStoreGroup.cpp
--- a/Engine/Plugins/FX/Niagara/Source/NiagaraEditor/Private/ViewModels/Stack/NiagaraStackParameterStoreGroup.cpp
+++ b/Engine/Plugins/FX/Niagara/Source/NiagaraEditor/Private/ViewModels/Stack/NiagaraStackParameterStoreGroup.cpp
@@ -56,11 +56,11 @@ private:
 	FNiagaraVariable NewParameterVariable;
 };
 
-class FParameterStoreGroupAddUtiliites : public FNiagaraStackItemGroupAddUtilities
+class FParameterStoreGroupAddUtiliites : public TNiagaraStackItemGroupAddUtilities<FNiagaraVariable>
 {
 public:
-	FParameterStoreGroupAddUtiliites(UObject& InParameterStoreOwner, FNiagaraParameterStore& InParameterStore, UNiagaraStackEditorData& InStackEditorData)
-		: FNiagaraStackItemGroupAddUtilities(LOCTEXT("ScriptGroupAddItemName", "Parameter"), EAddMode::AddFromAction, true)
+	FParameterStoreGroupAddUtiliites(UObject& InParameterStoreOwner, FNiagaraParameterStore& InParameterStore, UNiagaraStackEditorData& InStackEditorData, FOnItemAdded InOnItemAdded)
+		: TNiagaraStackItemGroupAddUtilities(LOCTEXT("ScriptGroupAddItemName", "Parameter"), EAddMode::AddFromAction, true, InOnItemAdded)
 		, ParameterStoreOwner(InParameterStoreOwner)
 		, ParameterStore(InParameterStore)
 		, StackEditorData(InStackEditorData)
@@ -93,6 +93,8 @@ public:
 
 		ParameterStore.AddParameter(ParameterVariable);
 		StackEditorData.SetModuleInputIsRenamePending(ParameterVariable.GetName().ToString(), true);
+
+		OnItemAdded.ExecuteIfBound(ParameterVariable);
 	}
 
 private:
@@ -106,7 +108,8 @@ void UNiagaraStackParameterStoreGroup::Initialize(
 	UObject* InOwner,
 	FNiagaraParameterStore* InParameterStore)
 {
-	AddUtilities = MakeShared<FParameterStoreGroupAddUtiliites>(*InOwner, *InParameterStore, *InRequiredEntryData.StackEditorData);
+	AddUtilities = MakeShared<FParameterStoreGroupAddUtiliites>(*InOwner, *InParameterStore, *InRequiredEntryData.StackEditorData, 
+		FParameterStoreGroupAddUtiliites::FOnItemAdded::CreateUObject(this, &UNiagaraStackParameterStoreGroup::ParameterAdded));
 	FText DisplayName = LOCTEXT("SystemExposedVariablesGroup", "System Exposed Parameters");
 	FText ToolTip = LOCTEXT("SystemExposedVariablesGroupToolTip", "Displays the variables created in the User namespace. These variables are exposed to owning UComponents, blueprints, etc.");
 	Super::Initialize(InRequiredEntryData, DisplayName, ToolTip, AddUtilities.Get());
@@ -145,6 +148,11 @@ void UNiagaraStackParameterStoreGroup::RefreshChildrenInternal(const TArray<UNia
 	Super::RefreshChildrenInternal(CurrentChildren, NewChildren, NewIssues);
 }
 
+void UNiagaraStackParameterStoreGroup::ParameterAdded(FNiagaraVariable AddedParameter)
+{
+	RefreshChildren();
+}
+
 void UNiagaraStackParameterStoreItem::Initialize(
 	FRequiredEntryData InRequiredEntryData,
 	UObject* InOwner,
@@ -154,17 +162,6 @@ void UNiagaraStackParameterStoreItem::Initialize(
 
 	Owner = InOwner;
 	ParameterStore = InParameterStore;
-	ParameterStoreChangedHandle = ParameterStore->AddOnChangedHandler(
-		FNiagaraParameterStore::FOnChanged::FDelegate::CreateUObject(this, &UNiagaraStackParameterStoreItem::ParameterStoreChanged));
-}
-
-void UNiagaraStackParameterStoreItem::FinalizeInternal()
-{
-	if (Owner != nullptr)
-	{
-		ParameterStore->RemoveOnChangedHandler(ParameterStoreChangedHandle);
-	}
-	Super::FinalizeInternal();
 }
 
 FText UNiagaraStackParameterStoreItem::GetDisplayName() const
@@ -188,6 +185,7 @@ void UNiagaraStackParameterStoreItem::RefreshChildrenInternal(const TArray<UNiag
 			{
 				ValueObjectEntry = NewObject<UNiagaraStackParameterStoreEntry>(this);
 				ValueObjectEntry->Initialize(CreateDefaultChildRequiredData(), Owner.Get(), ParameterStore, Var.GetName().ToString(), Var.GetType(), GetStackEditorDataKey());
+				ValueObjectEntry->OnParameterDeleted().AddUObject(this, &UNiagaraStackParameterStoreItem::ParameterDeleted);
 			}
 
 			NewChildren.Add(ValueObjectEntry);
@@ -196,7 +194,7 @@ void UNiagaraStackParameterStoreItem::RefreshChildrenInternal(const TArray<UNiag
 	Super::RefreshChildrenInternal(CurrentChildren, NewChildren, NewIssues);
 }
 
-void UNiagaraStackParameterStoreItem::ParameterStoreChanged()
+void UNiagaraStackParameterStoreItem::ParameterDeleted()
 {
 	RefreshChildren();
 }
diff --git a/Engine/Plugins/FX/Niagara/Source/NiagaraEditor/Public/ViewModels/Stack/NiagaraStackParameterStoreEntry.h b/Engine/Plugins/FX/Niagara/Source/NiagaraEditor/Public/ViewModels/Stack/NiagaraStackParameterStoreEntry.h
--- a/Engine/Plugins/FX/Niagara/Source/NiagaraEditor/Public/ViewModels/Stack/NiagaraStackParameterStoreEntry.h
+++ b/Engine/Plugins/FX/Niagara/Source/NiagaraEditor/Public/ViewModels/Stack/NiagaraStackParameterStoreEntry.h
@@ -23,6 +23,8 @@ class NIAGARAEDITOR_API UNiagaraStackParameterStoreEntry : public UNiagaraStackI
 
 public:
 	DECLARE_MULTICAST_DELEGATE(FOnValueChanged);
+	DECLARE_MULTICAST_DELEGATE(FOnParameterDeleted);
+
 public:
 	UNiagaraStackParameterStoreEntry();
 
@@ -88,6 +90,9 @@ public:
 	/** Gets a multicast delegate which is called whenever the value on this input changes. */
 	FOnValueChanged& OnValueChanged();
 
+	/** Gets a multicast delegate which is called when this parameter is deleted. */
+	FOnParameterDeleted& OnParameterDeleted();
+
 	/** Delete the parameter from the ParameterStore and notify that the store changed. */
 	void Delete();
 
@@ -116,13 +121,16 @@ private:
 	/** A local copy of the value of this input if one is available. */
 	TSharedPtr<FStructOnScope> LocalValueStruct;
 
+	/** A multicast delegate which is called when the value of this input is changed. */
+	FOnValueChanged ValueChangedDelegate;
+
+	/** A multicast delegate which is called when this parameter is deleted. */
+	FOnParameterDeleted ParameterDeletedDelegate;
+
 	/** A pointer to the data interface object for this input if one is available. */
 	UPROPERTY()
 	UNiagaraDataInterface* ValueObject;
 
-	/** A multicast delegate which is called when the value of this input is changed. */
-	FOnValueChanged ValueChangedDelegate;
-
 	UPROPERTY()
 	UObject* Owner;
 
diff --git a/Engine/Plugins/FX/Niagara/Source/NiagaraEditor/Public/ViewModels/Stack/NiagaraStackParameterStoreGroup.h b/Engine/Plugins/FX/Niagara/Source/NiagaraEditor/Public/ViewModels/Stack/NiagaraStackParameterStoreGroup.h
--- a/Engine/Plugins/FX/Niagara/Source/NiagaraEditor/Public/ViewModels/Stack/NiagaraStackParameterStoreGroup.h
+++ b/Engine/Plugins/FX/Niagara/Source/NiagaraEditor/Public/ViewModels/Stack/NiagaraStackParameterStoreGroup.h
@@ -18,11 +18,12 @@ class NIAGARAEDITOR_API UNiagaraStackParameterStoreGroup : public UNiagaraStackI
 public:
 	void Initialize(FRequiredEntryData InRequiredEntryData,	UObject* InOwner, FNiagaraParameterStore* InParameterStore);
 
-	void AddUserParameter(FNiagaraVariable ParameterVariable);
-
 protected:
 	virtual void RefreshChildrenInternal(const TArray<UNiagaraStackEntry*>& CurrentChildren, TArray<UNiagaraStackEntry*>& NewChildren, TArray<FStackIssue>& NewIssues) override;
 
+private:
+	void ParameterAdded(FNiagaraVariable AddedParameter);
+
 private:
 	TWeakObjectPtr<UObject> Owner;
 	FNiagaraParameterStore* ParameterStore;
@@ -41,15 +42,12 @@ public:
 	virtual FText GetDisplayName() const override;
 
 protected:
-	virtual void FinalizeInternal() override;
-
 	virtual void RefreshChildrenInternal(const TArray<UNiagaraStackEntry*>& CurrentChildren, TArray<UNiagaraStackEntry*>& NewChildren, TArray<FStackIssue>& NewIssues) override;
 
 private:
-	void ParameterStoreChanged();
+	void ParameterDeleted();
 
 private:
 	TWeakObjectPtr<UObject> Owner;
 	FNiagaraParameterStore* ParameterStore;
-	FDelegateHandle ParameterStoreChangedHandle;
 };
\ No newline at end of file
diff --git a/Engine/Plugins/FX/Niagara/Source/NiagaraEditorWidgets/Private/Stack/SNiagaraStackParameterStoreEntryValue.cpp b/Engine/Plugins/FX/Niagara/Source/NiagaraEditorWidgets/Private/Stack/SNiagaraStackParameterStoreEntryValue.cpp
--- a/Engine/Plugins/FX/Niagara/Source/NiagaraEditorWidgets/Private/Stack/SNiagaraStackParameterStoreEntryValue.cpp
+++ b/Engine/Plugins/FX/Niagara/Source/NiagaraEditorWidgets/Private/Stack/SNiagaraStackParameterStoreEntryValue.cpp
@@ -120,13 +120,16 @@ void SNiagaraStackParameterStoreEntryValue::Construct(const FArguments& InArgs,
 
 FReply SNiagaraStackParameterStoreEntryValue::DeleteClicked()
 {
-	StackEntry->Delete();
 	// toast notification
 	FNotificationInfo Info(FText::Format(LOCTEXT("NiagaraDeletedUserParameter", "System exposed parameter was deleted.\n{0}\n(All links to inner variables were invalidated in the process.)"), StackEntry->GetDisplayName()));
 	Info.ExpireDuration = 5.0f;
 	Info.bFireAndForget = true;
 	Info.Image = FCoreStyle::Get().GetBrush(TEXT("MessageLog.Info"));
 	FSlateNotificationManager::Get().AddNotification(Info);
+
+	// Delete after the notification is posted to prevent the entry from becoming invalidated before generating the message.
+	StackEntry->Delete();
+
 	return FReply::Handled();
 }
 
