Niagara - Fix a crash when dragging a module from one emitter editor to another emitter editor.

diff --git a/Engine/Plugins/FX/Niagara/Source/NiagaraEditor/Private/ViewModels/Stack/NiagaraStackScriptItemGroup.cpp b/Engine/Plugins/FX/Niagara/Source/NiagaraEditor/Private/ViewModels/Stack/NiagaraStackScriptItemGroup.cpp
--- a/Engine/Plugins/FX/Niagara/Source/NiagaraEditor/Private/ViewModels/Stack/NiagaraStackScriptItemGroup.cpp
+++ b/Engine/Plugins/FX/Niagara/Source/NiagaraEditor/Private/ViewModels/Stack/NiagaraStackScriptItemGroup.cpp
@@ -583,7 +583,7 @@ TOptional<UNiagaraStackEntry::FDropResult> UNiagaraStackScriptItemGroup::ChildRe
 					if (SourceEmitterHandle != nullptr)
 					{
 						UNiagaraNodeOutput* SourceModuleOutputNode = FNiagaraStackGraphUtilities::GetEmitterOutputNodeForStackNode(SourceModuleItem->GetModuleNode());
-						UNiagaraScript* SourceModuleScript = FNiagaraEditorUtilities::GetScriptFromSystem(GetSystemViewModel()->GetSystem(), SourceEmitterHandle->GetId(),
+						UNiagaraScript* SourceModuleScript = FNiagaraEditorUtilities::GetScriptFromSystem(SourceModuleItem->GetSystemViewModel()->GetSystem(), SourceEmitterHandle->GetId(),
 							SourceModuleOutputNode->GetUsage(), SourceModuleOutputNode->GetUsageId());
 
 						const FNiagaraEmitterHandle* TargetEmitterHandle = FNiagaraEditorUtilities::GetEmitterHandleForEmitter(GetSystemViewModel()->GetSystem(), *GetEmitterViewModel()->GetEmitter());
