Niagara - Prevent a crash when fixing dependencies, where the fix is going to insert a module that's has 'module' as a valid usage.

diff --git a/Engine/Plugins/FX/Niagara/Source/NiagaraEditor/Private/ViewModels/Stack/NiagaraStackModuleItem.cpp b/Engine/Plugins/FX/Niagara/Source/NiagaraEditor/Private/ViewModels/Stack/NiagaraStackModuleItem.cpp
--- a/Engine/Plugins/FX/Niagara/Source/NiagaraEditor/Private/ViewModels/Stack/NiagaraStackModuleItem.cpp
+++ b/Engine/Plugins/FX/Niagara/Source/NiagaraEditor/Private/ViewModels/Stack/NiagaraStackModuleItem.cpp
@@ -32,6 +32,10 @@
 
 #include "ScopedTransaction.h"
 
+// TODO: Remove these
+#include "Widgets/Notifications/SNotificationList.h"
+#include "Framework/Notifications/NotificationManager.h"
+
 #define LOCTEXT_NAMESPACE "NiagaraStackModuleItem"
 
 TArray<ENiagaraScriptUsage> UsagePriority = { // Ordered such as the highest priority has the largest index
@@ -51,6 +55,13 @@ UNiagaraNodeOutput* GetOutputNodeForModuleDependency(ENiagaraScriptUsage Dependa
 		TArray<ENiagaraScriptUsage> SupportedUsages = UNiagaraScript::GetSupportedUsageContextsForBitmask(DependencyScript->ModuleUsageBitmask);
 		ENiagaraScriptUsage ScriptUsage = SupportedUsages[0];
 		int32 ClosestDistance = MAX_int32;
+
+		int32 DependantIndex = UsagePriority.IndexOfByPredicate(
+			[&](const ENiagaraScriptUsage CurrentUsage)
+		{
+			return UNiagaraScript::IsEquivalentUsage(DependantUsage, CurrentUsage);
+		});
+
 		for (ENiagaraScriptUsage PossibleUsage : SupportedUsages)
 		{
 			int32 PossibleIndex = UsagePriority.IndexOfByPredicate(
@@ -58,11 +69,12 @@ UNiagaraNodeOutput* GetOutputNodeForModuleDependency(ENiagaraScriptUsage Dependa
 				{
 					return UNiagaraScript::IsEquivalentUsage(PossibleUsage, CurrentUsage);
 				});
-			int32 DependantIndex = UsagePriority.IndexOfByPredicate(
-				[&](const ENiagaraScriptUsage CurrentUsage)
-				{
-					return UNiagaraScript::IsEquivalentUsage(DependantUsage, CurrentUsage);
-				});
+
+			if (PossibleIndex == INDEX_NONE)
+			{
+				// This usage isn't in the execution flow so check the next one.
+				continue;
+			}
 
 			int32 Distance = PossibleIndex - DependantIndex;
 			bool bCorrectOrder = (Dependency.Type == ENiagaraModuleDependencyType::PreDependency && Distance >= 0) || (Dependency.Type == ENiagaraModuleDependencyType::PostDependency && Distance <= 0);
@@ -96,6 +108,7 @@ UNiagaraNodeOutput* GetOutputNodeForModuleDependency(ENiagaraScriptUsage Dependa
 				}
 			}
 		}
+
 		if (OutputScript != nullptr)
 		{
 			TargetOutputNode = FNiagaraEditorUtilities::GetScriptOutputNode(*OutputScript);
@@ -438,13 +451,23 @@ void UNiagaraStackModuleItem::RefreshIssues(TArray<FStackIssue>& NewIssues)
 									}
 								}
 							}
+
 							if (TargetOutputNode == nullptr)
 							{
-								checkf(ModuleIndex != INDEX_NONE, TEXT("Module data wasn't found in system for current module!"));
-								// I am looking for a script in the same emitter as this module, so get it from there
-								UNiagaraScript& TargetScript = *FNiagaraEditorUtilities::GetScriptFromSystem(GetSystemViewModel()->GetSystem(), SystemModuleData[ModuleIndex].EmitterHandleId, DependencyScript->GetUsage(), DependencyScript->GetUsageId());
-								TargetOutputNode = FNiagaraEditorUtilities::GetScriptOutputNode(TargetScript);
+								// If no output node was found than the dependency can't be resolved and it most likely misconfigured in data.
+								// TODO: Don't show this toast here, change the fix delegate to return a fix result with whether or not the fix succeeded and any error message for the user.
+								FNotificationInfo Error(LOCTEXT("FixFailedToast", "Failed to fix the dependency since\nwe could not find a compatible place to insert the module.\nPlease check the configuration of the dependency.\nSee the log for more details."));
+								Error.ExpireDuration = 5.0f;
+								Error.bFireAndForget = true;
+								Error.Image = FCoreStyle::Get().GetBrush(TEXT("MessageLog.Error"));
+								FSlateNotificationManager::Get().AddNotification(Error);
+								FString ModuleAssetFullName;
+								ModuleAsset.GetFullName(ModuleAssetFullName);
+								UE_LOG(LogNiagaraEditor, Error, TEXT("Dependency fix failed, could not find a compatible place to insert the module.\nModule requiring dependency: %s\nModule providing dependency: %s\nDependency name: %s\nDependency type: %s"),
+									*FunctionCallNode->FunctionScript->GetFullName(), *ModuleAssetFullName, *Dependency.Id.ToString(), Dependency.Type == ENiagaraModuleDependencyType::PreDependency ? TEXT("Pre-dependency") : TEXT("Post-dependency"));
+								return;
 							}
+
 							TArray<FNiagaraStackModuleData> ScriptModuleData = SystemModuleData.FilterByPredicate([&](FNiagaraStackModuleData CurrentData) {return CurrentData.Usage == DependencyScript->GetUsage(); });
 							int32 PreIndex = INDEX_NONE; // index of last pre dependency
 							int32 PostIndex = INDEX_NONE; // index of fist post dependency, the module will have to be placed between these indexes
