Niagara - Fix a crash when deleting an emitter from a multi-emitter system and then undoing/redoing.

diff --git a/Engine/Plugins/FX/Niagara/Source/NiagaraEditor/Private/ViewModels/NiagaraSystemScriptViewModel.cpp b/Engine/Plugins/FX/Niagara/Source/NiagaraEditor/Private/ViewModels/NiagaraSystemScriptViewModel.cpp
--- a/Engine/Plugins/FX/Niagara/Source/NiagaraEditor/Private/ViewModels/NiagaraSystemScriptViewModel.cpp
+++ b/Engine/Plugins/FX/Niagara/Source/NiagaraEditor/Private/ViewModels/NiagaraSystemScriptViewModel.cpp
@@ -168,14 +168,14 @@ void FNiagaraSystemScriptViewModel::RebuildEmitterNodes()
 		CurrentEmitterNode->Modify();
 		UEdGraphPin* InPin = CurrentEmitterNode->GetInputPin(0);
 		UEdGraphPin* OutPin = CurrentEmitterNode->GetOutputPin(0);
-		if (InPin != nullptr && OutPin != nullptr)
+		UEdGraphPin* InPinLinkedPin = InPin != nullptr && InPin->LinkedTo.Num() == 1 ? InPin->LinkedTo[0] : nullptr;
+		UEdGraphPin* OutPinLinkedPin = OutPin != nullptr && OutPin->LinkedTo.Num() == 1 ? OutPin->LinkedTo[0] : nullptr;
+		CurrentEmitterNode->DestroyNode();
+
+		if (InPinLinkedPin != nullptr &&& OutPinLinkedPin != nullptr)
 		{
-			if (0 != InPin->LinkedTo.Num() && 0 != OutPin->LinkedTo.Num())
-			{
-				InPin->LinkedTo[0]->MakeLinkTo(OutPin->LinkedTo[0]);
-			}
+			InPinLinkedPin->MakeLinkTo(OutPinLinkedPin);
 		}
-		CurrentEmitterNode->DestroyNode();
 	}
 
 	// Add output nodes if they don't exist.
