Disabling single particle event type.

diff --git a/Engine/Plugins/FX/Niagara/Source/Niagara/Classes/NiagaraEmitter.h b/Engine/Plugins/FX/Niagara/Source/Niagara/Classes/NiagaraEmitter.h
--- a/Engine/Plugins/FX/Niagara/Source/Niagara/Classes/NiagaraEmitter.h
+++ b/Engine/Plugins/FX/Niagara/Source/Niagara/Classes/NiagaraEmitter.h
@@ -94,7 +94,7 @@ enum class EScriptExecutionMode : uint8
 	/** The event script is run only on the particles that were spawned in response to the current event in the emitter.*/
 	SpawnedParticles,
 	/** The event script is run only on the particle whose int32 ParticleIndex is specified in the event payload.*/
-	SingleParticle
+	SingleParticle UMETA(Hidden)
 };
 
 USTRUCT()
diff --git a/Engine/Plugins/FX/Niagara/Source/Niagara/Private/NiagaraEmitterInstance.cpp b/Engine/Plugins/FX/Niagara/Source/Niagara/Private/NiagaraEmitterInstance.cpp
--- a/Engine/Plugins/FX/Niagara/Source/Niagara/Private/NiagaraEmitterInstance.cpp
+++ b/Engine/Plugins/FX/Niagara/Source/Niagara/Private/NiagaraEmitterInstance.cpp
@@ -1298,62 +1298,63 @@ void FNiagaraEmitterInstance::Tick(float DeltaSeconds)
 			}
 		}
 
-
-		// handle single-particle events
-		// TODO: we'll need a way to either skip execution of the VM if an index comes back as invalid, or we'll have to pre-process
-		// event/particle arrays; this is currently a very naive (and comparatively slow) implementation, until full indexed reads work
-		if (EventHandlerProps.Script && EventHandlerProps.ExecutionMode == EScriptExecutionMode::SingleParticle && EventSet[EventScriptIdx])
-		{
-
-			SCOPE_CYCLE_COUNTER(STAT_NiagaraEventHandle);
-			FNiagaraVariable IndexVar(FNiagaraTypeDefinition::GetIntDef(), "ParticleIndex");
-			FNiagaraDataSetIterator<int32> IndexItr(*EventSet[EventScriptIdx], IndexVar, 0, false);
-			if (IndexItr.IsValid() && EventSet[EventScriptIdx]->GetPrevNumInstances() > 0)
-			{
-				EventExecCountBindings[EventScriptIdx].SetValue(1);
-
-				Data.CopyCurToPrev();
-				uint32 NumParticles = Data.GetNumInstances();
-
-				for (uint32 i = 0; i < EventSet[EventScriptIdx]->GetPrevNumInstances(); i++)
-				{
-					int32 Index = *IndexItr;
-					IndexItr.Advance();
-					DataSetExecInfos.SetNum(1, false);
-					DataSetExecInfos[0].StartInstance = Index;
-					DataSetExecInfos[0].bUpdateInstanceCount = false;
-					DataSetExecInfos.Emplace(EventSet[EventScriptIdx], i, false, false);
-					EventExecContexts[EventScriptIdx].Execute(1, DataSetExecInfos);
-
-					if (GbDumpParticleData || System->bDumpDebugEmitterInfo)
-					{
-						ensure(EventHandlerProps.Script->RapidIterationParameters.VerifyBinding(&EventExecContexts[EventScriptIdx].Parameters));
-						UE_LOG(LogNiagara, Log, TEXT("=== Event %d Src Parameters ==="), EventScriptIdx);
-						EventHandlerProps.Script->RapidIterationParameters.Dump();
-						UE_LOG(LogNiagara, Log, TEXT("=== Event %d Context Parameters ==="), EventScriptIdx);
-						EventExecContexts[EventScriptIdx].Parameters.Dump();
-						UE_LOG(LogNiagara, Log, TEXT("=== Event %d Particles (%d index written, %d total) ==="), EventScriptIdx, Index, Data.GetNumInstances());
-						Data.Dump(true, Index, 1);
-					}
-
-
-#if WITH_EDITORONLY_DATA
-					if (ParentSystemInstance->ShouldCaptureThisFrame())
-					{
-						FGuid EventGuid = EventExecContexts[EventScriptIdx].Script->GetUsageId();
-						FNiagaraScriptDebuggerInfo* DebugInfo = ParentSystemInstance->GetActiveCaptureWrite(CachedIDName, ENiagaraScriptUsage::ParticleEventScript, EventGuid);
-						if (DebugInfo)
-						{
-							Data.Dump(DebugInfo->Frame, true, Index, 1);
-							//DebugInfo->Frame.Dump(true, 0, 1);
-							DebugInfo->Parameters = EventExecContexts[EventScriptIdx].Parameters;
-						}
-					}
-#endif
-					ensure(NumParticles == Data.GetNumInstances());
-				}
-			}
-		}
+		//TODO: Disabling this event mode for now until it can be reworked. Currently it uses index directly with can easily be invalid and cause undefined behavior.
+		//
+//		// handle single-particle events
+//		// TODO: we'll need a way to either skip execution of the VM if an index comes back as invalid, or we'll have to pre-process
+//		// event/particle arrays; this is currently a very naive (and comparatively slow) implementation, until full indexed reads work
+// 		if (EventHandlerProps.Script && EventHandlerProps.ExecutionMode == EScriptExecutionMode::SingleParticle && EventSet[EventScriptIdx])
+// 		{
+// 
+// 			SCOPE_CYCLE_COUNTER(STAT_NiagaraEventHandle);
+// 			FNiagaraVariable IndexVar(FNiagaraTypeDefinition::GetIntDef(), "ParticleIndex");
+// 			FNiagaraDataSetIterator<int32> IndexItr(*EventSet[EventScriptIdx], IndexVar, 0, false);
+// 			if (IndexItr.IsValid() && EventSet[EventScriptIdx]->GetPrevNumInstances() > 0)
+// 			{
+// 				EventExecCountBindings[EventScriptIdx].SetValue(1);
+// 
+// 				Data.CopyCurToPrev();
+// 				uint32 NumParticles = Data.GetNumInstances();
+// 
+// 				for (uint32 i = 0; i < EventSet[EventScriptIdx]->GetPrevNumInstances(); i++)
+// 				{
+// 					int32 Index = *IndexItr;
+// 					IndexItr.Advance();
+// 					DataSetExecInfos.SetNum(1, false);
+// 					DataSetExecInfos[0].StartInstance = Index;
+// 					DataSetExecInfos[0].bUpdateInstanceCount = false;
+// 					DataSetExecInfos.Emplace(EventSet[EventScriptIdx], i, false, false);
+// 					EventExecContexts[EventScriptIdx].Execute(1, DataSetExecInfos);
+// 
+// 					if (GbDumpParticleData || System->bDumpDebugEmitterInfo)
+// 					{
+// 						ensure(EventHandlerProps.Script->RapidIterationParameters.VerifyBinding(&EventExecContexts[EventScriptIdx].Parameters));
+// 						UE_LOG(LogNiagara, Log, TEXT("=== Event %d Src Parameters ==="), EventScriptIdx);
+// 						EventHandlerProps.Script->RapidIterationParameters.Dump();
+// 						UE_LOG(LogNiagara, Log, TEXT("=== Event %d Context Parameters ==="), EventScriptIdx);
+// 						EventExecContexts[EventScriptIdx].Parameters.Dump();
+// 						UE_LOG(LogNiagara, Log, TEXT("=== Event %d Particles (%d index written, %d total) ==="), EventScriptIdx, Index, Data.GetNumInstances());
+// 						Data.Dump(true, Index, 1);
+// 					}
+// 
+// 
+// #if WITH_EDITORONLY_DATA
+// 					if (ParentSystemInstance->ShouldCaptureThisFrame())
+// 					{
+// 						FGuid EventGuid = EventExecContexts[EventScriptIdx].Script->GetUsageId();
+// 						FNiagaraScriptDebuggerInfo* DebugInfo = ParentSystemInstance->GetActiveCaptureWrite(CachedIDName, ENiagaraScriptUsage::ParticleEventScript, EventGuid);
+// 						if (DebugInfo)
+// 						{
+// 							Data.Dump(DebugInfo->Frame, true, Index, 1);
+// 							//DebugInfo->Frame.Dump(true, 0, 1);
+// 							DebugInfo->Parameters = EventExecContexts[EventScriptIdx].Parameters;
+// 						}
+// 					}
+// #endif
+// 					ensure(NumParticles == Data.GetNumInstances());
+// 				}
+// 			}
+// 		}
 	}
 
 	PostProcessParticles();
