Disallowing multiple GPU dispatches per frame for each emitter.

diff --git a/Engine/Plugins/FX/Niagara/Source/Niagara/Classes/NiagaraScriptExecutionContext.h b/Engine/Plugins/FX/Niagara/Source/Niagara/Classes/NiagaraScriptExecutionContext.h
--- a/Engine/Plugins/FX/Niagara/Source/Niagara/Classes/NiagaraScriptExecutionContext.h
+++ b/Engine/Plugins/FX/Niagara/Source/Niagara/Classes/NiagaraScriptExecutionContext.h
@@ -90,17 +90,19 @@ struct FNiagaraComputeExecutionContext
 		, UpdateScript(nullptr)
 		, GPUScript(nullptr)
 		, RTUpdateScript(0)
-		, RTSpawnScript(0)	
+		, RTSpawnScript(0)
 		, RTGPUScript(0)
 		, GPUDataReadback(nullptr)
 		, AccumulatedSpawnRate(0)
 		, NumIndicesPerInstance(0)
+		, bPendingExecution(0)
 	{
 	}
 
 	void Reset()
 	{
 		AccumulatedSpawnRate = 0;
+		bPendingExecution = 0;
 		if (GPUDataReadback)
 		{
 			delete GPUDataReadback;
@@ -194,4 +196,7 @@ struct FNiagaraComputeExecutionContext
 	FRHIGPUMemoryReadback *GPUDataReadback;
 	uint32 AccumulatedSpawnRate;
 	uint32 NumIndicesPerInstance;	// how many vtx indices per instance the renderer is going to have for its draw call
+
+	/** Ensures we only enqueue each context once before they're dispatched. */
+	uint32 bPendingExecution : 1;
 };
diff --git a/Engine/Plugins/FX/Niagara/Source/Niagara/Private/NiagaraEmitterInstanceBatcher.cpp b/Engine/Plugins/FX/Niagara/Source/Niagara/Private/NiagaraEmitterInstanceBatcher.cpp
--- a/Engine/Plugins/FX/Niagara/Source/Niagara/Private/NiagaraEmitterInstanceBatcher.cpp
+++ b/Engine/Plugins/FX/Niagara/Source/Niagara/Private/NiagaraEmitterInstanceBatcher.cpp
@@ -27,14 +27,18 @@ uint32 FNiagaraComputeExecutionContext::TickCounter = 0;
 void NiagaraEmitterInstanceBatcher::Queue(FNiagaraComputeExecutionContext *InContext)
 {
 	//SimulationQueue[CurQueueIndex]->Add(InContext);
-
-	ENQUEUE_UNIQUE_RENDER_COMMAND_THREEPARAMETER(QueueNiagaraDispatch,
-		TArray<FNiagaraComputeExecutionContext*>*, Queue, &SimulationQueue[0],
-		uint32, QueueIndex, CurQueueIndex,
-		FNiagaraComputeExecutionContext*, ExecContext, InContext,
-		{
-			Queue[QueueIndex].Add(ExecContext);
-		});
+		ENQUEUE_UNIQUE_RENDER_COMMAND_THREEPARAMETER(QueueNiagaraDispatch,
+			TArray<FNiagaraComputeExecutionContext*>*, Queue, &SimulationQueue[0],
+			uint32, QueueIndex, CurQueueIndex,
+			FNiagaraComputeExecutionContext*, ExecContext, InContext,
+			{
+				//Don't queue the same context for execution multiple times. TODO: possibly try to combine/accumulate the tick info if we happen to have > 1 before it's executed.
+				if (!ExecContext->bPendingExecution)
+				{
+					Queue[QueueIndex].Add(ExecContext);
+					ExecContext->bPendingExecution = true;
+				}
+			});
 }
 
 
@@ -56,6 +60,7 @@ void NiagaraEmitterInstanceBatcher::TickSingle(FNiagaraComputeExecutionContext *
 
 	check(IsInRenderingThread());
 	Context->MainDataSet->Tick();
+	Context->bPendingExecution = false;
 
 	FNiagaraComputeExecutionContext::TickCounter++;
 
