Niagara - Fix an issue with data interfaces defined in emitters called "Emitter"

diff --git a/Engine/Plugins/FX/Niagara/Source/Niagara/Private/NiagaraFunctionLibrary.cpp b/Engine/Plugins/FX/Niagara/Source/Niagara/Private/NiagaraFunctionLibrary.cpp
--- a/Engine/Plugins/FX/Niagara/Source/Niagara/Private/NiagaraFunctionLibrary.cpp
+++ b/Engine/Plugins/FX/Niagara/Source/Niagara/Private/NiagaraFunctionLibrary.cpp
@@ -29,6 +29,9 @@ UNiagaraComponent* UNiagaraFunctionLibrary::SpawnSystemAtLocation(UObject* World
 		{
 			AActor* Actor = World->GetWorldSettings();
 			PSC = NewObject<UNiagaraComponent>((Actor ? Actor : (UObject*)World));
+#if WITH_EDITORONLY_DATA
+			PSC->bWaitForCompilationOnActivate = true;
+#endif
 			PSC->SetAsset(SystemTemplate);
 			PSC->SetAutoDestroy(bAutoDestroy);
 			PSC->bAutoActivate = false;
diff --git a/Engine/Plugins/FX/Niagara/Source/NiagaraEditor/Private/NiagaraHlslTranslator.cpp b/Engine/Plugins/FX/Niagara/Source/NiagaraEditor/Private/NiagaraHlslTranslator.cpp
--- a/Engine/Plugins/FX/Niagara/Source/NiagaraEditor/Private/NiagaraHlslTranslator.cpp
+++ b/Engine/Plugins/FX/Niagara/Source/NiagaraEditor/Private/NiagaraHlslTranslator.cpp
@@ -4228,23 +4228,24 @@ int32 FHlslNiagaraTranslator::RegisterDataInterface(FNiagaraVariable& Var, UNiag
 	}
 
 	//If we get here then this is a new data interface.
+	FName DataInterfaceName;
+	if (FNiagaraParameterMapHistory::IsAliasedEmitterParameter(Var.GetName().ToString()))
+	{
+		FNiagaraVariable AliasedVar = ActiveHistoryForFunctionCalls.ResolveAliases(Var);
+		DataInterfaceName = AliasedVar.GetName();
+	}
+	else
+	{
+		DataInterfaceName = Var.GetName();
+	}
+
 	int32 Idx = CompilationOutput.ScriptData.DataInterfaceInfo.IndexOfByPredicate([&](const FNiagaraScriptDataInterfaceCompileInfo& OtherInfo)
 	{
-		return OtherInfo.Name == Var.GetName();
+		return OtherInfo.Name == DataInterfaceName;
 	});
+
 	if (Idx == INDEX_NONE)
 	{
-		FName DataInterfaceName;
-		if (FNiagaraParameterMapHistory::IsAliasedEmitterParameter(Var.GetName().ToString()))
-		{
-			FNiagaraVariable AliasedVar = ActiveHistoryForFunctionCalls.ResolveAliases(Var);
-			DataInterfaceName = AliasedVar.GetName();
-		}
-		else
-		{
-			DataInterfaceName = Var.GetName();
-		}
-
 		Idx = CompilationOutput.ScriptData.DataInterfaceInfo.AddDefaulted();
 		CompilationOutput.ScriptData.DataInterfaceInfo[Idx].Name = DataInterfaceName;
 		CompilationOutput.ScriptData.DataInterfaceInfo[Idx].Type = Var.GetType();
diff --git a/Engine/Plugins/FX/Niagara/Source/NiagaraEditor/Private/NiagaraParameterMapHistory.cpp b/Engine/Plugins/FX/Niagara/Source/NiagaraEditor/Private/NiagaraParameterMapHistory.cpp
--- a/Engine/Plugins/FX/Niagara/Source/NiagaraEditor/Private/NiagaraParameterMapHistory.cpp
+++ b/Engine/Plugins/FX/Niagara/Source/NiagaraEditor/Private/NiagaraParameterMapHistory.cpp
@@ -235,6 +235,12 @@ FNiagaraVariable FNiagaraParameterMapHistory::ResolveAliases(const FNiagaraVaria
 
 FName FNiagaraParameterMapHistory::ResolveEmitterAlias(const FName& InName, const FString& InAlias)
 {
+	// If the alias is empty than the name can't be resolved.
+	if (InAlias.IsEmpty())
+	{
+		return InName;
+	}
+
 	FNiagaraVariable Var(FNiagaraTypeDefinition::GetFloatDef(), InName);
 	TMap<FString, FString> ResolveMap;
 	ResolveMap.Add(TEXT("Emitter"), InAlias);
