Crash cooking Niagara content for dedicated servers.

diff --git a/Engine/Plugins/FX/Niagara/Source/Niagara/Private/NiagaraScript.cpp b/Engine/Plugins/FX/Niagara/Source/Niagara/Private/NiagaraScript.cpp
--- a/Engine/Plugins/FX/Niagara/Source/Niagara/Private/NiagaraScript.cpp
+++ b/Engine/Plugins/FX/Niagara/Source/Niagara/Private/NiagaraScript.cpp
@@ -913,25 +913,23 @@ void UNiagaraScript::BeginCacheForCookedPlatformData(const ITargetPlatform *Targ
 
 		FNiagaraShaderScript** CachedScriptResourceForPlatformPtr = CachedScriptResourcesForCooking.Find(TargetPlatform);
 
-		if (CachedScriptResourceForPlatformPtr == nullptr)
+		if (DesiredShaderFormats.Num() > 0 && CachedScriptResourceForPlatformPtr == nullptr)
 		{
 			CachedScriptResourcesForCooking.Add(TargetPlatform);
 			FNiagaraShaderScript* &CachedScriptResourceForPlatform = *CachedScriptResourcesForCooking.Find(TargetPlatform);
 
-			if (DesiredShaderFormats.Num())
+			// Cache for all the shader formats that the cooking target requires
+			for (int32 FormatIndex = 0; FormatIndex < DesiredShaderFormats.Num(); FormatIndex++)
 			{
-				// Cache for all the shader formats that the cooking target requires
-				for (int32 FormatIndex = 0; FormatIndex < DesiredShaderFormats.Num(); FormatIndex++)
+				if (DesiredShaderFormats[FormatIndex] != TEXT("PCD3D_SM4"))	// TODO: remove once we get rid of SM4 globally
 				{
-					if (DesiredShaderFormats[FormatIndex] != TEXT("PCD3D_SM4"))	// TODO: remove once we get rid of SM4 globally
-					{
-						const EShaderPlatform LegacyShaderPlatform = ShaderFormatToLegacyShaderPlatform(DesiredShaderFormats[FormatIndex]);
+					const EShaderPlatform LegacyShaderPlatform = ShaderFormatToLegacyShaderPlatform(DesiredShaderFormats[FormatIndex]);
 
-						// Begin caching shaders for the target platform and store the FNiagaraShaderScript being compiled into CachedScriptResourcesForCooking
-						CacheResourceShadersForCooking(LegacyShaderPlatform, CachedScriptResourceForPlatform);
-					}
+					// Begin caching shaders for the target platform and store the FNiagaraShaderScript being compiled into CachedScriptResourcesForCooking
+					CacheResourceShadersForCooking(LegacyShaderPlatform, CachedScriptResourceForPlatform);
 				}
 			}
+			check(CachedScriptResourceForPlatform);
 		}
 	}
 }
@@ -1195,7 +1193,7 @@ void SerializeNiagaraShaderMaps(const TMap<const ITargetPlatform*, FNiagaraShade
 			auto& PlatformScriptResourcesToSave = *PlatformScriptResourcesToSavePtr;
 
 			ScriptResourceToSavePtr = PlatformScriptResourcesToSave.Find(Ar.CookingTarget());
-			check(ScriptResourceToSavePtr != NULL || (Ar.GetLinker() == NULL));
+			//check(ScriptResourceToSavePtr != NULL || (Ar.GetLinker() == NULL));
 			if (ScriptResourceToSavePtr != NULL)
 			{
 				NumResourcesToSave++;
@@ -1207,6 +1205,7 @@ void SerializeNiagaraShaderMaps(const TMap<const ITargetPlatform*, FNiagaraShade
 		if (ScriptResourceToSavePtr)
 		{
 			ScriptResourceToSave = *ScriptResourceToSavePtr;
+			check(ScriptResourceToSave);
 			ScriptResourceToSave->SerializeShaderMap(Ar);
 		}
 
