Cleaner fix for UE-61214. Using TGlobalResource to register with the global FRenderResource pool and clear out any texture references upon app shutdown.

#tests now hits the RHI shutdown without need for PreExit

diff --git a/Engine/Plugins/FX/Niagara/Source/Niagara/Private/NiagaraModule.cpp b/Engine/Plugins/FX/Niagara/Source/Niagara/Private/NiagaraModule.cpp
--- a/Engine/Plugins/FX/Niagara/Source/Niagara/Private/NiagaraModule.cpp
+++ b/Engine/Plugins/FX/Niagara/Source/Niagara/Private/NiagaraModule.cpp
@@ -227,7 +227,6 @@ void INiagaraModule::StartupModule()
 	UNiagaraRibbonRendererProperties::InitCDOPropertiesAfterModuleStartup();
 	UNiagaraMeshRendererProperties::InitCDOPropertiesAfterModuleStartup();
 
-	FCoreDelegates::OnPreExit.AddRaw(this, &INiagaraModule::ShutdownRenderingResources);
 }
 
 void INiagaraModule::ShutdownRenderingResources()
diff --git a/Engine/Plugins/FX/Niagara/Source/Niagara/Private/NiagaraWorldManager.cpp b/Engine/Plugins/FX/Niagara/Source/Niagara/Private/NiagaraWorldManager.cpp
--- a/Engine/Plugins/FX/Niagara/Source/Niagara/Private/NiagaraWorldManager.cpp
+++ b/Engine/Plugins/FX/Niagara/Source/Niagara/Private/NiagaraWorldManager.cpp
@@ -13,10 +13,11 @@
 #include "NiagaraDataInterfaceSkeletalMesh.h"
 #include "EngineModule.h"
 
-FNiagaraViewDataMgr GNiagaraViewDataManager;
+TGlobalResource<FNiagaraViewDataMgr> GNiagaraViewDataManager;
 
-FNiagaraViewDataMgr::FNiagaraViewDataMgr()
-	: SceneDepthTexture(nullptr)
+FNiagaraViewDataMgr::FNiagaraViewDataMgr() 
+	: FRenderResource()
+	, SceneDepthTexture(nullptr)
 	, SceneNormalTexture(nullptr)
 	, ViewUniformBuffer(nullptr)
 {
@@ -32,10 +33,20 @@ void FNiagaraViewDataMgr::Init()
 
 void FNiagaraViewDataMgr::Shutdown()
 {
-	GNiagaraViewDataManager.SceneDepthTexture = nullptr;
-	GNiagaraViewDataManager.SceneNormalTexture = nullptr;
-	GNiagaraViewDataManager.ViewUniformBuffer = nullptr;
-	GNiagaraViewDataManager.SceneTexturesUniformParams.SafeRelease();
+	GNiagaraViewDataManager.ReleaseDynamicRHI();
+}
+
+void FNiagaraViewDataMgr::InitDynamicRHI()
+{
+
+}
+
+void FNiagaraViewDataMgr::ReleaseDynamicRHI()
+{
+	SceneDepthTexture = nullptr;
+	SceneNormalTexture = nullptr;
+	ViewUniformBuffer = nullptr;
+	SceneTexturesUniformParams.SafeRelease();
 }
 
 FNiagaraWorldManager::FNiagaraWorldManager(UWorld* InWorld)
diff --git a/Engine/Plugins/FX/Niagara/Source/Niagara/Public/NiagaraWorldManager.h b/Engine/Plugins/FX/Niagara/Source/Niagara/Public/NiagaraWorldManager.h
--- a/Engine/Plugins/FX/Niagara/Source/Niagara/Public/NiagaraWorldManager.h
+++ b/Engine/Plugins/FX/Niagara/Source/Niagara/Public/NiagaraWorldManager.h
@@ -18,7 +18,7 @@ class UNiagaraParameterCollectionInstance;
 
 
 
-class FNiagaraViewDataMgr
+class FNiagaraViewDataMgr : public FRenderResource
 {
 public:
 	FNiagaraViewDataMgr();
@@ -38,6 +38,10 @@ public:
 	FTexture2DRHIParamRef GetSceneNormalTexture() { return SceneNormalTexture; }
 	FUniformBufferRHIParamRef GetViewUniformBuffer() { return ViewUniformBuffer; }
 	TUniformBufferRef<FSceneTexturesUniformParameters> GetSceneTextureUniformParameters() { return SceneTexturesUniformParams; }
+
+	virtual void InitDynamicRHI() override;
+
+	virtual void ReleaseDynamicRHI() override;
 private:
 	FTexture2DRHIParamRef SceneDepthTexture;
 	FTexture2DRHIParamRef SceneNormalTexture;
@@ -47,7 +51,7 @@ private:
 	FPostOpaqueRenderDelegate PostOpaqueDelegate;
 };
 
-extern FNiagaraViewDataMgr GNiagaraViewDataManager;
+extern TGlobalResource<FNiagaraViewDataMgr> GNiagaraViewDataManager;
 
 
 /**
