Disabling opengl 4.0 compilation of materials for sprites/ribbons/meshes

#tests n/a

diff --git a/Engine/Plugins/FX/Niagara/Source/NiagaraVertexFactories/Private/NiagaraMeshVertexFactory.cpp b/Engine/Plugins/FX/Niagara/Source/NiagaraVertexFactories/Private/NiagaraMeshVertexFactory.cpp
--- a/Engine/Plugins/FX/Niagara/Source/NiagaraVertexFactories/Private/NiagaraMeshVertexFactory.cpp
+++ b/Engine/Plugins/FX/Niagara/Source/NiagaraVertexFactories/Private/NiagaraMeshVertexFactory.cpp
@@ -193,7 +193,7 @@ FShaderResourceViewRHIParamRef FNiagaraMeshVertexFactory::GetPreviousTransformBu
 
 bool FNiagaraMeshVertexFactory::ShouldCompilePermutation(EShaderPlatform Platform, const class FMaterial* Material, const class FShaderType* ShaderType)
 {
-	return (!IsMobilePlatform(Platform) && !IsSwitchPlatform(Platform) && (Material->IsUsedWithNiagaraMeshParticles() || Material->IsSpecialEngineMaterial()));
+	return (!IsMobilePlatform(Platform) && !IsSwitchPlatform(Platform) && Platform != SP_OPENGL_SM4 && (Material->IsUsedWithNiagaraMeshParticles() || Material->IsSpecialEngineMaterial()));
 }
 
 void FNiagaraMeshVertexFactory::SetData(const FStaticMeshDataType& InData)
diff --git a/Engine/Plugins/FX/Niagara/Source/NiagaraVertexFactories/Private/NiagaraRibbonVertexFactory.cpp b/Engine/Plugins/FX/Niagara/Source/NiagaraVertexFactories/Private/NiagaraRibbonVertexFactory.cpp
--- a/Engine/Plugins/FX/Niagara/Source/NiagaraVertexFactories/Private/NiagaraRibbonVertexFactory.cpp
+++ b/Engine/Plugins/FX/Niagara/Source/NiagaraVertexFactories/Private/NiagaraRibbonVertexFactory.cpp
@@ -168,7 +168,7 @@ static TGlobalResource<FNiagaraRibbonVertexDeclaration> GNiagaraRibbonVertexDecl
 
 bool FNiagaraRibbonVertexFactory::ShouldCompilePermutation(EShaderPlatform Platform, const class FMaterial* Material, const class FShaderType* ShaderType)
 {
-	return (!IsMobilePlatform(Platform) && !IsSwitchPlatform(Platform) && (Material->IsUsedWithNiagaraRibbons() || Material->IsSpecialEngineMaterial()));
+	return (!IsMobilePlatform(Platform) && !IsSwitchPlatform(Platform) && Platform != SP_OPENGL_SM4 && (Material->IsUsedWithNiagaraRibbons() || Material->IsSpecialEngineMaterial()));
 }
 
 /**
diff --git a/Engine/Plugins/FX/Niagara/Source/NiagaraVertexFactories/Private/NiagaraSpriteVertexFactory.cpp b/Engine/Plugins/FX/Niagara/Source/NiagaraVertexFactories/Private/NiagaraSpriteVertexFactory.cpp
--- a/Engine/Plugins/FX/Niagara/Source/NiagaraVertexFactories/Private/NiagaraSpriteVertexFactory.cpp
+++ b/Engine/Plugins/FX/Niagara/Source/NiagaraVertexFactories/Private/NiagaraSpriteVertexFactory.cpp
@@ -225,7 +225,7 @@ inline TGlobalResource<FNiagaraSpriteVertexDeclaration>& GetNiagaraSpriteVertexD
 
 bool FNiagaraSpriteVertexFactory::ShouldCompilePermutation(EShaderPlatform Platform, const class FMaterial* Material, const class FShaderType* ShaderType)
 {
-	return (!IsSwitchPlatform(Platform) && (Material->IsUsedWithNiagaraSprites() || Material->IsSpecialEngineMaterial()));
+	return (!IsSwitchPlatform(Platform) && Platform != SP_OPENGL_SM4 && (Material->IsUsedWithNiagaraSprites() || Material->IsSpecialEngineMaterial()));
 }
 
 /**
