Niagara - Fix inheritance for inputs on "Set Variables" modules in the emitter section of the stack.

diff --git a/Engine/Plugins/FX/Niagara/Source/NiagaraEditor/Private/NiagaraScriptMergeManager.cpp b/Engine/Plugins/FX/Niagara/Source/NiagaraEditor/Private/NiagaraScriptMergeManager.cpp
--- a/Engine/Plugins/FX/Niagara/Source/NiagaraEditor/Private/NiagaraScriptMergeManager.cpp
+++ b/Engine/Plugins/FX/Niagara/Source/NiagaraEditor/Private/NiagaraScriptMergeManager.cpp
@@ -194,9 +194,25 @@ FNiagaraStackFunctionMergeAdapter::FNiagaraStackFunctionMergeAdapter(FString InE
 	for (const FNiagaraVariable& RapidIterationParameter : RapidIterationParameters)
 	{
 		FNiagaraParameterHandle AliasedInputHandle(*RapidIterationParameter.GetName().ToString().RightChop(RapidIterationParameterNamePrefix.Len()));
-		if (AliasedInputHandle.GetNamespace().ToString() == FunctionCallNode->GetFunctionName() && AliasedInputsAdded.Contains(AliasedInputHandle.GetParameterHandleString().ToString()) == false)
+		if (AliasedInputHandle.GetNamespace().ToString() == FunctionCallNode->GetFunctionName())
 		{
-			InputOverrides.Add(MakeShared<FNiagaraStackFunctionInputOverrideMergeAdapter>(UniqueEmitterName, *OwningScript.Get(), *FunctionCallNode.Get(), AliasedInputHandle.GetName().ToString(), RapidIterationParameter));
+			// Currently rapid iteration parameters for assignment nodes in emitter scripts get double aliased which prevents their inputs from
+			// being diffed correctly, so we need to un-mangle the names here so that the diffs are correct.
+			if (FunctionCallNode->IsA<UNiagaraNodeAssignment>() &&
+				(OwningScript->GetUsage() == ENiagaraScriptUsage::EmitterSpawnScript || OwningScript->GetUsage() == ENiagaraScriptUsage::EmitterUpdateScript))
+			{
+				FString InputName = AliasedInputHandle.GetName().ToString();
+				if (InputName.StartsWith(UniqueEmitterName + TEXT(".")))
+				{
+					FString UnaliasedInputName = TEXT("Emitter") + InputName.RightChop(UniqueEmitterName.Len());
+					AliasedInputHandle = FNiagaraParameterHandle(AliasedInputHandle.GetNamespace(), *UnaliasedInputName);
+				}
+			}
+
+			if (AliasedInputsAdded.Contains(AliasedInputHandle.GetParameterHandleString().ToString()) == false)
+			{
+				InputOverrides.Add(MakeShared<FNiagaraStackFunctionInputOverrideMergeAdapter>(UniqueEmitterName, *OwningScript.Get(), *FunctionCallNode.Get(), AliasedInputHandle.GetName().ToString(), RapidIterationParameter));
+			}
 		}
 	}
 }
@@ -1632,11 +1648,11 @@ FNiagaraScriptMergeManager::FApplyDiffResults FNiagaraScriptMergeManager::AddInp
 		}
 	}
 
+	FNiagaraParameterHandle FunctionInputHandle(FNiagaraParameterHandle::ModuleNamespace, *OverrideToAdd->GetInputName());
+	FNiagaraParameterHandle AliasedFunctionInputHandle = FNiagaraParameterHandle::CreateAliasedModuleParameterHandle(FunctionInputHandle, &TargetFunctionCall);
+
 	if (OverrideToAdd->GetOverridePin() != nullptr)
 	{
-		FNiagaraParameterHandle FunctionInputHandle(FNiagaraParameterHandle::ModuleNamespace, *OverrideToAdd->GetInputName());
-		FNiagaraParameterHandle AliasedFunctionInputHandle = FNiagaraParameterHandle::CreateAliasedModuleParameterHandle(FunctionInputHandle, &TargetFunctionCall);
-
 		const UEdGraphSchema_Niagara* NiagaraSchema = GetDefault<UEdGraphSchema_Niagara>();
 		FNiagaraTypeDefinition InputType = NiagaraSchema->PinToTypeDefinition(OverrideToAdd->GetOverridePin());
 
@@ -1717,8 +1733,8 @@ FNiagaraScriptMergeManager::FApplyDiffResults FNiagaraScriptMergeManager::AddInp
 	{
 		if (OverrideToAdd->GetLocalValueRapidIterationParameter().IsSet())
 		{
-			FString RapidIterationParameterName = FString::Printf(TEXT("Constants.%s.%s.%s"), *UniqueEmitterName, *TargetFunctionCall.GetFunctionName(), *OverrideToAdd->GetInputName());
-			FNiagaraVariable RapidIterationParameter(OverrideToAdd->GetLocalValueRapidIterationParameter().GetValue().GetType(), *RapidIterationParameterName);
+			FNiagaraVariable RapidIterationParameter = FNiagaraStackGraphUtilities::CreateRapidIterationParameter(
+				UniqueEmitterName, OwningScript.GetUsage(), AliasedFunctionInputHandle.GetParameterHandleString(), OverrideToAdd->GetLocalValueRapidIterationParameter().GetValue().GetType());
 			const uint8* SourceData = OverrideToAdd->GetOwningScript()->RapidIterationParameters.GetParameterData(OverrideToAdd->GetLocalValueRapidIterationParameter().GetValue());
 			OwningScript.Modify();
 			bool bAddParameterIfMissing = true;
