Niagara - Don't allow emitters to have a '.' in their name, it breaks rapid iteration parameters when it's renamed.

diff --git a/Engine/Plugins/FX/Niagara/Source/Niagara/Private/NiagaraEmitterHandle.cpp b/Engine/Plugins/FX/Niagara/Source/Niagara/Private/NiagaraEmitterHandle.cpp
--- a/Engine/Plugins/FX/Niagara/Source/Niagara/Private/NiagaraEmitterHandle.cpp
+++ b/Engine/Plugins/FX/Niagara/Source/Niagara/Private/NiagaraEmitterHandle.cpp
@@ -103,6 +103,7 @@ void FNiagaraEmitterHandle::SetName(FName InName, UNiagaraSystem& InOwnerSystem)
 	FString TempNameString = InName.ToString();
 	TempNameString.ReplaceInline(TEXT(" "), TEXT("_"));
 	TempNameString.ReplaceInline(TEXT("\t"), TEXT("_"));
+	TempNameString.ReplaceInline(TEXT("."), TEXT("_"));
 	InName = *TempNameString;
 
 	TSet<FName> OtherEmitterNames;
