Niagara - Fix UV scaling, and add additional scaling and offset modes for ribbons using normalized age for ordering.

#TESTS Existing automated tests, plus several new tests to verify the functionality.

diff --git a/Engine/Plugins/FX/Niagara/Source/Niagara/Private/NiagaraRendererRibbons.cpp b/Engine/Plugins/FX/Niagara/Source/Niagara/Private/NiagaraRendererRibbons.cpp
--- a/Engine/Plugins/FX/Niagara/Source/Niagara/Private/NiagaraRendererRibbons.cpp
+++ b/Engine/Plugins/FX/Niagara/Source/Niagara/Private/NiagaraRendererRibbons.cpp
@@ -270,6 +270,57 @@ bool NiagaraRendererRibbons::SetMaterialUsage()
 	return Material && Material->CheckMaterialUsage_Concurrent(MATUSAGE_NiagaraRibbons);
 }
 
+void CalculateUVScaleAndOffsets(const FNiagaraDataSetAccessor<float>& SortKeyData, const TArray<int32>& RibbonIndices, bool bSortKeyIsAge, int32 NumSegments,
+	float InUVTilingDistance, const FVector2D& InUVScale, const FVector2D& InUVOffset, ENiagaraRibbonAgeOffsetMode InAgeOffsetMode, FVector2D& OutUVScale, FVector2D& OutUVOffset)
+{
+	if (bSortKeyIsAge && InUVTilingDistance == 0)
+	{
+		float AgeUScale;
+		float AgeUOffset;
+		if (InAgeOffsetMode == ENiagaraRibbonAgeOffsetMode::Scale)
+		{
+			// In scale mode we scale and offset the UVs so that no part of the texture is clipped. In order to prevent
+			// clipping at the ends we'll have to move the UVs in up to the size of a single segment of the ribbon since
+			// that's the distance we'll need to to smoothly interpolate when a new segment is added, or when an old segment
+			// is removed.  We calculate the end offset when the end of the ribbon is within a single time step of 0 or 1
+			// which is then normalized to the range of a single segment.  We can then calculate how many segments we actually
+			// have to draw the scaled ribbon, and can offset the start by the correctly scaled offset.
+			float FirstAge = SortKeyData[RibbonIndices[0]];
+			float SecondAge = SortKeyData[RibbonIndices[1]];
+			float SecondToLastAge = SortKeyData[RibbonIndices[RibbonIndices.Num() - 2]];
+			float LastAge = SortKeyData[RibbonIndices.Last()];
+
+			float StartTimeStep = SecondAge - FirstAge;
+			float StartTimeOffset = FirstAge < StartTimeStep ? StartTimeStep - FirstAge : 0;
+			float StartSegmentOffset = StartTimeOffset / StartTimeStep;
+
+			float EndTimeStep = LastAge - SecondToLastAge;
+			float EndTimeOffset = 1 - LastAge < EndTimeStep ? EndTimeStep - (1 - LastAge) : 0;
+			float EndSegmentOffset = EndTimeOffset / EndTimeStep;
+
+			float AvailableSegments = NumSegments - (StartSegmentOffset + EndSegmentOffset);
+			AgeUScale = NumSegments / AvailableSegments;
+			AgeUOffset = -((StartSegmentOffset / NumSegments) * AgeUScale);
+		}
+		else
+		{
+			float FirstAge = SortKeyData[RibbonIndices[0]];
+			float LastAge = SortKeyData[RibbonIndices.Last()];
+
+			AgeUScale = LastAge - FirstAge;
+			AgeUOffset = FirstAge;
+		}
+
+		OutUVScale = FVector2D(AgeUScale * InUVScale.X, InUVScale.Y);
+		OutUVOffset = FVector2D((AgeUOffset * InUVScale.X) + InUVOffset.X, InUVOffset.Y);
+	}
+	else
+	{
+		OutUVScale = InUVScale;
+		OutUVOffset = InUVOffset;
+	}
+}
+
 FNiagaraDynamicDataBase *NiagaraRendererRibbons::GenerateVertexData(const FNiagaraSceneProxy* Proxy, FNiagaraDataSet &Data, const ENiagaraSimTarget Target)
 {
 	SCOPE_CYCLE_COUNTER(STAT_NiagaraGenRibbonVertexData);
@@ -288,7 +339,7 @@ FNiagaraDynamicDataBase *NiagaraRendererRibbons::GenerateVertexData(const FNiaga
 
 	// TODO : deal with the dynamic vertex material parameter should the user have specified it as an output...
 
-	FVector2D UVs[4] = { FVector2D(1.0f, 0.0f), FVector2D(1.0f, 1.0f) };
+	FVector2D UVs[2] = { FVector2D(1.0f, 1.0f), FVector2D(1.0f, 0.0f) };
 	int32 NumTotalVerts = 0;
 
 	FNiagaraDataSetAccessor<FVector> PosData(Data, Properties->PositionBinding.DataSetVariable);
@@ -364,15 +415,17 @@ FNiagaraDynamicDataBase *NiagaraRendererRibbons::GenerateVertexData(const FNiaga
 			FVector PrevPos, PrevPos2, PrevDir(0.0f, 0.0f, 0.1f);
 			float TotalDistance = 0.0f;
 
-			FVector2D AgeOffset;
-			if (bSortKeyIsAge)
-			{
-				AgeOffset = FVector2D(SortKeyData[RibbonIndices.Last()], 0.0f);
-			}
-			else
-			{
-				AgeOffset = FVector2D::ZeroVector;
-			}
+			int32 NumSegments = NumIndices - 1;
+
+			FVector2D UV0Offset;
+			FVector2D UV0Scale;
+			FVector2D UV1Offset;
+			FVector2D UV1Scale;
+
+			CalculateUVScaleAndOffsets(SortKeyData, RibbonIndices, bSortKeyIsAge, NumSegments, Properties->UV0TilingDistance,
+				Properties->UV0Scale, Properties->UV0Offset, Properties->UV0AgeOffsetMode, UV0Scale, UV0Offset);
+			CalculateUVScaleAndOffsets(SortKeyData, RibbonIndices, bSortKeyIsAge, NumSegments, Properties->UV1TilingDistance,
+				Properties->UV1Scale, Properties->UV1Offset, Properties->UV1AgeOffsetMode, UV1Scale, UV1Offset);
 
 			for (int32 i = 0; i < NumIndices; i++)
 			{
@@ -397,18 +450,18 @@ FNiagaraDynamicDataBase *NiagaraRendererRibbons::GenerateVertexData(const FNiaga
 				{
 					FVector NormDir = ParticleDir.GetSafeNormal();
 					PrevDir = NormDir;
-					FVector2D UV0Mult((float)i / NumIndices, 1.0f);
-					FVector2D UV1Mult((float)i / NumIndices, 1.0f);
-
-					if (Properties->UV0TilingDistance)
-					{
-						UV0Mult = FVector2D(TotalDistance / Properties->UV0TilingDistance, 1.0f);
-					}
-					if (Properties->UV1TilingDistance)
-					{
-						UV1Mult = FVector2D(TotalDistance / Properties->UV1TilingDistance, 1.0f);
-					}
 
+					// Calculate sub-uv for this segment based on the either the distance or index, depending on whether the distance tile factor was set.
+					float U0ForSegment = Properties->UV0TilingDistance
+						? TotalDistance / Properties->UV0TilingDistance
+						: (float)i / NumSegments;
+					float U1ForSegment = Properties->UV1TilingDistance
+						? TotalDistance / Properties->UV1TilingDistance
+						: (float)i / NumSegments;
+
+					FVector2D UV0ForSegment(U0ForSegment, 1.0f);
+					FVector2D UV1ForSegment(U1ForSegment, 1.0f);
+					
 					//Todo. Possibly pull out to templated func to avoid safety code in loop here.
 					FLinearColor Color = ColData.GetSafe(Index1, FLinearColor::White);
 					float Twist = TwistData.GetSafe(Index1, 0.0f);
@@ -416,8 +469,8 @@ FNiagaraDynamicDataBase *NiagaraRendererRibbons::GenerateVertexData(const FNiaga
 					FVector Align = AlignData.GetSafe(Index1, FVector::UpVector);
 
 					TotalDistance += ParticleDir.Size();
-					AddRibbonVert(RenderData, ParticlePos, UVs[0] * UV0Mult + AgeOffset, UVs[0] * UV1Mult + AgeOffset, Color, Twist, Size, NormDir, Align);
-					AddRibbonVert(RenderData, ParticlePos, UVs[1] * UV0Mult + AgeOffset, UVs[1] * UV1Mult + AgeOffset, Color, Twist, Size, NormDir, Align);
+					AddRibbonVert(RenderData, ParticlePos, UVs[0] * UV0ForSegment * UV0Scale + UV0Offset, UVs[0] * UV1ForSegment * UV1Scale + UV1Offset, Color, Twist, Size, NormDir, Align);
+					AddRibbonVert(RenderData, ParticlePos, UVs[1] * UV0ForSegment * UV0Scale + UV0Offset, UVs[1] * UV1ForSegment * UV1Scale + UV1Offset, Color, Twist, Size, NormDir, Align);
 					if (MaterialParamData.IsValid())
 					{
 						AddDynamicParam(DynamicData->MaterialParameterVertexData[0], MaterialParamData[Index1]);
diff --git a/Engine/Plugins/FX/Niagara/Source/Niagara/Public/NiagaraRendererRibbons.h b/Engine/Plugins/FX/Niagara/Source/Niagara/Public/NiagaraRendererRibbons.h
--- a/Engine/Plugins/FX/Niagara/Source/Niagara/Public/NiagaraRendererRibbons.h
+++ b/Engine/Plugins/FX/Niagara/Source/Niagara/Public/NiagaraRendererRibbons.h
@@ -43,10 +43,10 @@ public:
 		NewVertex.Color = Color;
 		NewVertex.Size = Size;
 		NewVertex.Rotation = Rotation;
-		NewVertex.Tex_U = UV1.X * Properties->UV0Scale.X;
-		NewVertex.Tex_V = UV1.Y * Properties->UV0Scale.Y;
-		NewVertex.Tex_U2 = UV2.X * Properties->UV1Scale.X;
-		NewVertex.Tex_V2 = UV2.Y * Properties->UV1Scale.Y;
+		NewVertex.Tex_U = UV1.X;
+		NewVertex.Tex_V = UV1.Y;
+		NewVertex.Tex_U2 = UV2.X;
+		NewVertex.Tex_V2 = UV2.Y;
 		NewVertex.CustomFacingVector = CustomFacing;
 		RenderData.Add(NewVertex);
 	}
diff --git a/Engine/Plugins/FX/Niagara/Source/Niagara/Public/NiagaraRibbonRendererProperties.h b/Engine/Plugins/FX/Niagara/Source/Niagara/Public/NiagaraRibbonRendererProperties.h
--- a/Engine/Plugins/FX/Niagara/Source/Niagara/Public/NiagaraRibbonRendererProperties.h
+++ b/Engine/Plugins/FX/Niagara/Source/Niagara/Public/NiagaraRibbonRendererProperties.h
@@ -15,6 +15,16 @@ enum class ENiagaraRibbonFacingMode : uint8
 	Custom
 };
 
+/** Defines different modes for offsetting UVs by age when ordering ribbon particles using normalized age. */
+UENUM()
+enum class ENiagaraRibbonAgeOffsetMode : uint8
+{
+	/** Offset the UVs by age for smooth texture movement, but scale the 0-1 UV range to the current normalized age range of the particles. */
+	Scale,
+	/** Offset the UVs by age for smooth texture movement, but use the normalized age range directly as the UV range which will clip the texture for normalized age ranges that don't go from 0-1. */
+	Clip
+};
+
 UCLASS(editinlinenew)
 class UNiagaraRibbonRendererProperties : public UNiagaraRendererProperties
 {
@@ -26,8 +36,10 @@ public:
 		, FacingMode(ENiagaraRibbonFacingMode::Screen)
 		, UV0TilingDistance(0.0f)
 		, UV0Scale(FVector2D(1.0f, 1.0f))
+		, UV0AgeOffsetMode(ENiagaraRibbonAgeOffsetMode::Scale)
 		, UV1TilingDistance(0.0f)
 		, UV1Scale(FVector2D(1.0f, 1.0f))
+		, UV1AgeOffsetMode(ENiagaraRibbonAgeOffsetMode::Scale)
 	{
 	}
 
@@ -37,15 +49,29 @@ public:
 	UPROPERTY(EditAnywhere, Category = "Ribbon Rendering")
 	ENiagaraRibbonFacingMode FacingMode;
 
+	/** Tiles UV0 based on the distance traversed by the ribbon. Disables offsetting UVs by age. */
 	UPROPERTY(EditAnywhere, Category = "Ribbon Rendering")
 	float UV0TilingDistance;
 	UPROPERTY(EditAnywhere, Category = "Ribbon Rendering")
 	FVector2D UV0Scale;
+	UPROPERTY(EditAnywhere, Category = "Ribbon Rendering")
+	FVector2D UV0Offset;
 
+	/** Defines the mode to use when offsetting UV channel 0 by age which enables smooth texture movement when particles are added and removed at the end of the ribbon.  Not used when the RibbonLinkOrder binding is in use or when tiling distance is in use. */
+	UPROPERTY(EditAnywhere, Category = "Ribbon Rendering")
+	ENiagaraRibbonAgeOffsetMode UV0AgeOffsetMode;
+
+	/** Tiles UV1 based on the distance traversed by the ribbon. Disables offsetting UVs by age. */
 	UPROPERTY(EditAnywhere, Category = "Ribbon Rendering")
 	float UV1TilingDistance;
 	UPROPERTY(EditAnywhere, Category = "Ribbon Rendering")
 	FVector2D UV1Scale;
+	UPROPERTY(EditAnywhere, Category = "Ribbon Rendering")
+	FVector2D UV1Offset;
+
+	/** Defines the mode to use when offsetting UV channel 1 by age which enables smooth texture movement when particles are added and removed at the end of the ribbon.  Not used when the RibbonLinkOrder binding is in use or when tiling distance is in use. */
+	UPROPERTY(EditAnywhere, Category = "Ribbon Rendering")
+	ENiagaraRibbonAgeOffsetMode UV1AgeOffsetMode;
 
 	//~ UNiagaraRendererProperties interface
 	virtual NiagaraRenderer* CreateEmitterRenderer(ERHIFeatureLevel::Type FeatureLevel) override;
