#tests none

diff --git a/Engine/Plugins/FX/Niagara/Source/NiagaraEditor/Private/Widgets/SNiagaraParameterMapView.cpp b/Engine/Plugins/FX/Niagara/Source/NiagaraEditor/Private/Widgets/SNiagaraParameterMapView.cpp
--- a/Engine/Plugins/FX/Niagara/Source/NiagaraEditor/Private/Widgets/SNiagaraParameterMapView.cpp
+++ b/Engine/Plugins/FX/Niagara/Source/NiagaraEditor/Private/Widgets/SNiagaraParameterMapView.cpp
@@ -116,6 +116,10 @@ void FNiagaraParameterMapViewCommands::RegisterCommands()
 
 SNiagaraParameterMapView::~SNiagaraParameterMapView()
 {
+	// Unregister all commands for right click on action node
+	ToolkitCommands->UnmapAction(FNiagaraParameterMapViewCommands::Get().DeleteEntry);
+	ToolkitCommands->UnmapAction(FGenericCommands::Get().Rename);
+
 	TSet<UObject*> Objects = SelectedObjects->GetSelectedObjects();
 	for (UObject* Object : Objects)
 	{
