Niagara - Prevent a crash when sampling color from a static mesh has no color data.

#tests Mesh without color doesn't crash, mesh with color is sampled correctly.

diff --git a/Engine/Plugins/FX/Niagara/Source/Niagara/Private/NiagaraDataInterfaceStaticMesh.cpp b/Engine/Plugins/FX/Niagara/Source/Niagara/Private/NiagaraDataInterfaceStaticMesh.cpp
--- a/Engine/Plugins/FX/Niagara/Source/Niagara/Private/NiagaraDataInterfaceStaticMesh.cpp
+++ b/Engine/Plugins/FX/Niagara/Source/Niagara/Private/NiagaraDataInterfaceStaticMesh.cpp
@@ -1294,28 +1294,51 @@ void UNiagaraDataInterfaceStaticMesh::GetTriCoordColor(FVectorVMContext& Context
 	const FIndexArrayView& Indices = Res.IndexBuffer.GetArrayView();
 	const FColorVertexBuffer& Colors = Res.VertexBuffers.ColorVertexBuffer;
 
-	for (int32 i = 0; i < Context.NumInstances; ++i)
+	if (Colors.GetNumVertices() > 0)
 	{
-		int32 Tri = TriParam.Get();
-		int32 Idx0 = Indices[Tri];
-		int32 Idx1 = Indices[Tri + 1];
-		int32 Idx2 = Indices[Tri + 2];
-
-		FLinearColor Color = BarycentricInterpolate(BaryXParam.Get(), BaryYParam.Get(), BaryZParam.Get(),
-			Colors.VertexColor(Idx0).ReinterpretAsLinear(), Colors.VertexColor(Idx1).ReinterpretAsLinear(), Colors.VertexColor(Idx2).ReinterpretAsLinear());
-
-		*OutColorR.GetDest() = Color.R;
-		*OutColorG.GetDest() = Color.G;
-		*OutColorB.GetDest() = Color.B;
-		*OutColorA.GetDest() = Color.A;
-		TriParam.Advance();
-		BaryXParam.Advance();
-		BaryYParam.Advance();
-		BaryZParam.Advance();
-		OutColorR.Advance();
-		OutColorG.Advance();
-		OutColorB.Advance();
-		OutColorA.Advance();
+		for (int32 i = 0; i < Context.NumInstances; ++i)
+		{
+			int32 Tri = TriParam.Get();
+			int32 Idx0 = Indices[Tri];
+			int32 Idx1 = Indices[Tri + 1];
+			int32 Idx2 = Indices[Tri + 2];
+
+			FLinearColor Color = BarycentricInterpolate(BaryXParam.Get(), BaryYParam.Get(), BaryZParam.Get(),
+				Colors.VertexColor(Idx0).ReinterpretAsLinear(), Colors.VertexColor(Idx1).ReinterpretAsLinear(), Colors.VertexColor(Idx2).ReinterpretAsLinear());
+
+			*OutColorR.GetDest() = Color.R;
+			*OutColorG.GetDest() = Color.G;
+			*OutColorB.GetDest() = Color.B;
+			*OutColorA.GetDest() = Color.A;
+			TriParam.Advance();
+			BaryXParam.Advance();
+			BaryYParam.Advance();
+			BaryZParam.Advance();
+			OutColorR.Advance();
+			OutColorG.Advance();
+			OutColorB.Advance();
+			OutColorA.Advance();
+		}
+	}
+	else
+	{
+		// This mesh doesn't have color information so set the color to white.
+		FLinearColor Color = FLinearColor::White;
+		for (int32 i = 0; i < Context.NumInstances; ++i)
+		{
+			*OutColorR.GetDest() = Color.R;
+			*OutColorG.GetDest() = Color.G;
+			*OutColorB.GetDest() = Color.B;
+			*OutColorA.GetDest() = Color.A;
+			TriParam.Advance();
+			BaryXParam.Advance();
+			BaryYParam.Advance();
+			BaryZParam.Advance();
+			OutColorR.Advance();
+			OutColorG.Advance();
+			OutColorB.Advance();
+			OutColorA.Advance();
+		}
 	}
 }
 
