Niagara - Don't allow the same parameter to be added to an assignment node multiple times, and validate the namespace correctly.  Fixes a crash when adding the same parameter multiple times and then renaming it.

diff --git a/Engine/Plugins/FX/Niagara/Source/NiagaraEditor/Private/NiagaraNodeAssignment.cpp b/Engine/Plugins/FX/Niagara/Source/NiagaraEditor/Private/NiagaraNodeAssignment.cpp
--- a/Engine/Plugins/FX/Niagara/Source/NiagaraEditor/Private/NiagaraNodeAssignment.cpp
+++ b/Engine/Plugins/FX/Niagara/Source/NiagaraEditor/Private/NiagaraNodeAssignment.cpp
@@ -226,12 +226,15 @@ void UNiagaraNodeAssignment::BuildParameterMenu(FMenuBuilder& MenuBuilder, ENiag
 		FString VarDefaultValue = FNiagaraConstants::GetAttributeDefaultValue(AvailableParameter);
 		const FText TooltipDesc = FText::Format(LOCTEXT("SetFunctionPopupTooltip", "Description: Set the parameter {0}. {1}"), FText::FromName(AvailableParameter.GetName()), VarDesc);
 		FText CategoryName = LOCTEXT("ModuleSetCategory", "Set Specific Parameters");
+		bool bCanExecute = AssignmentTargets.Contains(AvailableParameter) == false;
 
 		MenuBuilder.AddMenuEntry(
 			NameText,
 			TooltipDesc,
 			FSlateIcon(),
-			FUIAction(FExecuteAction::CreateUObject(this, &UNiagaraNodeAssignment::AddParameter, AvailableParameter, VarDefaultValue)));
+			FUIAction(
+				FExecuteAction::CreateUObject(this, &UNiagaraNodeAssignment::AddParameter, AvailableParameter, VarDefaultValue),
+				FCanExecuteAction::CreateLambda([bCanExecute] { return bCanExecute; })));
 	}
 
 }
diff --git a/Engine/Plugins/FX/Niagara/Source/NiagaraEditor/Private/ViewModels/Stack/NiagaraStackFunctionInput.cpp b/Engine/Plugins/FX/Niagara/Source/NiagaraEditor/Private/ViewModels/Stack/NiagaraStackFunctionInput.cpp
--- a/Engine/Plugins/FX/Niagara/Source/NiagaraEditor/Private/ViewModels/Stack/NiagaraStackFunctionInput.cpp
+++ b/Engine/Plugins/FX/Niagara/Source/NiagaraEditor/Private/ViewModels/Stack/NiagaraStackFunctionInput.cpp
@@ -1264,7 +1264,8 @@ FNiagaraVariable UNiagaraStackFunctionInput::CreateRapidIterationVariable(const
 bool UNiagaraStackFunctionInput::CanRenameInput() const
 {
 	// Only module level assignment node inputs can be renamed.
-	return OwningAssignmentNode.IsValid() && InputParameterHandlePath.Num() == 1;
+	return OwningAssignmentNode.IsValid() && InputParameterHandlePath.Num() == 1 &&
+		OwningAssignmentNode->FindAssignmentTarget(InputParameterHandle.GetName()) != INDEX_NONE;
 }
 
 bool UNiagaraStackFunctionInput::GetIsRenamePending() const
diff --git a/Engine/Plugins/FX/Niagara/Source/NiagaraEditor/Private/ViewModels/Stack/NiagaraStackModuleItem.cpp b/Engine/Plugins/FX/Niagara/Source/NiagaraEditor/Private/ViewModels/Stack/NiagaraStackModuleItem.cpp
--- a/Engine/Plugins/FX/Niagara/Source/NiagaraEditor/Private/ViewModels/Stack/NiagaraStackModuleItem.cpp
+++ b/Engine/Plugins/FX/Niagara/Source/NiagaraEditor/Private/ViewModels/Stack/NiagaraStackModuleItem.cpp
@@ -141,6 +141,7 @@ void UNiagaraStackModuleItem::Initialize(FRequiredEntryData InRequiredEntryData,
 	Super::Initialize(InRequiredEntryData, ModuleStackEditorDataKey);
 	GroupAddUtilities = InGroupAddUtilities;
 	FunctionCallNode = &InFunctionCallNode;
+	OutputNode = FNiagaraStackGraphUtilities::GetEmitterOutputNodeForStackNode(*FunctionCallNode);
 	AddChildFilter(FOnFilterChild::CreateUObject(this, &UNiagaraStackModuleItem::FilterOutputCollection));
 
 	// Update bCanMoveAndDelete
@@ -155,7 +156,6 @@ void UNiagaraStackModuleItem::Initialize(FRequiredEntryData InRequiredEntryData,
 		TSharedRef<FNiagaraScriptMergeManager> MergeManager = FNiagaraScriptMergeManager::Get();
 
 		const UNiagaraEmitter* BaseEmitter = FNiagaraStackGraphUtilities::GetBaseEmitter(*GetEmitterViewModel()->GetEmitter(), GetSystemViewModel()->GetSystem());
-		UNiagaraNodeOutput* OutputNode = FNiagaraStackGraphUtilities::GetEmitterOutputNodeForStackNode(*FunctionCallNode);
 
 		bool bIsMergeable = MergeManager->IsMergeableScriptUsage(OutputNode->GetUsage());
 		bool bHasBaseModule = bIsMergeable && BaseEmitter != nullptr && MergeManager->HasBaseModule(*BaseEmitter, OutputNode->GetUsage(), OutputNode->GetUsageId(), FunctionCallNode->NodeGuid);
@@ -289,11 +289,37 @@ void UNiagaraStackModuleItem::RefreshIssues(TArray<FStackIssue>& NewIssues)
 
 			NewIssues.Add(InconsistentEnabledError);
 		}
+
+		UNiagaraNodeAssignment* AssignmentFunctionCall = Cast<UNiagaraNodeAssignment>(FunctionCallNode);
+		if (AssignmentFunctionCall != nullptr)
+		{
+			TSet<FNiagaraVariable> FoundAssignmentTargets;
+			for (const FNiagaraVariable& AssignmentTarget : AssignmentFunctionCall->GetAssignmentTargets())
+			{
+				if (FoundAssignmentTargets.Contains(AssignmentTarget))
+				{
+					FText FixDescription = LOCTEXT("RemoveDuplicate", "Remove Duplicate");
+					FStackIssueFix RemoveDuplicateFix(FixDescription, FStackIssueFixDelegate::CreateLambda([AssignmentFunctionCall, AssignmentTarget]()
+					{
+						AssignmentFunctionCall->RemoveParameter(AssignmentTarget);
+					}));
+					FStackIssue DuplicateAssignmentTargetError(
+						EStackIssueSeverity::Error,
+						LOCTEXT("DuplicateAssignmentTargetErrorSummary", "Duplicate variables detected."),
+						LOCTEXT("InconsistentEnabledError", "This 'Set Variables' module is attempting to set the same variable more than once, which is unsupported."),
+						GetStackEditorDataKey(),
+						false,
+						RemoveDuplicateFix);
+
+					NewIssues.Add(DuplicateAssignmentTargetError);
+				}
+				FoundAssignmentTargets.Add(AssignmentTarget);
+			}
+		}
 	}
 	// Generate dependency errors with their fixes
 	TArray<UNiagaraNodeFunctionCall*> FoundCalls;
 	TArray<FNiagaraModuleDependency> DependenciesNeeded;
-	UNiagaraNodeOutput* OutputNode = GetOutputNode();
 
 	TArray<FNiagaraStackModuleData> SystemModuleData = GetSystemViewModel()->GetStackModuleDataForEmitter(GetEmitterViewModel());
 	int32 ModuleIndex = INDEX_NONE;
@@ -387,7 +413,7 @@ void UNiagaraStackModuleItem::RefreshIssues(TArray<FStackIssue>& NewIssues)
 					FText AndEnableModule = bNeedsEnable ? FText::Format(LOCTEXT("AndEnableDependency", "And enable dependency module {0}"), FText::FromString(DisorderedNode.ModuleNode->GetFunctionName())) : FText();
 					UNiagaraStackEntry::FStackIssueFix Fix(
 						FText::Format(LOCTEXT("ReorderDependency", "Reposition this module in the correct order related to {0} {1}"), FText::FromString(DisorderedNode.ModuleNode->GetFunctionName()), AndEnableModule),
-						FStackIssueFixDelegate::CreateLambda([this, bNeedsEnable, DisorderedNode, SystemModuleData, Dependency, ModuleIndex, OutputNode]()
+						FStackIssueFixDelegate::CreateLambda([this, bNeedsEnable, DisorderedNode, SystemModuleData, Dependency, ModuleIndex]()
 					{
 						FScopedTransaction ScopedTransaction(LOCTEXT("ReorderDependencyModule", "Reorder dependency module"));
 
@@ -631,7 +657,7 @@ int32 UNiagaraStackModuleItem::GetModuleIndex()
 
 UNiagaraNodeOutput* UNiagaraStackModuleItem::GetOutputNode() const
 {
-	return FNiagaraStackGraphUtilities::GetEmitterOutputNodeForStackNode(*FunctionCallNode);
+	return OutputNode;
 }
 
 void UNiagaraStackModuleItem::NotifyModuleMoved()
@@ -639,14 +665,38 @@ void UNiagaraStackModuleItem::NotifyModuleMoved()
 	ModifiedGroupItemsDelegate.ExecuteIfBound();
 }
 
-bool UNiagaraStackModuleItem::CanAddInput() const
+bool ParameterIsCompatibleWithScriptUsage(FNiagaraVariable Parameter, ENiagaraScriptUsage Usage)
+{
+	FNiagaraParameterHandle ParameterHandle(Parameter.GetName());
+	switch (Usage)
+	{
+	case ENiagaraScriptUsage::SystemSpawnScript:
+	case ENiagaraScriptUsage::SystemUpdateScript:
+		return ParameterHandle.IsSystemHandle();
+	case ENiagaraScriptUsage::EmitterSpawnScript:
+	case ENiagaraScriptUsage::EmitterUpdateScript:
+		return ParameterHandle.IsEmitterHandle();
+	case ENiagaraScriptUsage::ParticleSpawnScript:
+	case ENiagaraScriptUsage::ParticleSpawnScriptInterpolated:
+	case ENiagaraScriptUsage::ParticleUpdateScript:
+	case ENiagaraScriptUsage::ParticleEventScript:
+		return ParameterHandle.IsParticleAttributeHandle();
+	default:
+		return false;
+	}
+}
+
+bool UNiagaraStackModuleItem::CanAddInput(FNiagaraVariable InputParameter) const
 {
-	return FunctionCallNode->IsA<UNiagaraNodeAssignment>();
+	UNiagaraNodeAssignment* AssignmentModule = Cast<UNiagaraNodeAssignment>(FunctionCallNode);
+	return AssignmentModule != nullptr &&
+		AssignmentModule->GetAssignmentTargets().Contains(InputParameter) == false &&
+		ParameterIsCompatibleWithScriptUsage(InputParameter, OutputNode->GetUsage());
 }
 
 void UNiagaraStackModuleItem::AddInput(FNiagaraVariable InputParameter)
 {
-	if(ensureMsgf(CanAddInput(), TEXT("This module doesn't support adding inputs.")))
+	if(ensureMsgf(CanAddInput(InputParameter), TEXT("This module doesn't support adding this input.")))
 	{
 		UNiagaraNodeAssignment* AssignmentNode = CastChecked<UNiagaraNodeAssignment>(FunctionCallNode);
 		AssignmentNode->AddParameter(InputParameter, FNiagaraConstants::GetAttributeDefaultValue(InputParameter));
diff --git a/Engine/Plugins/FX/Niagara/Source/NiagaraEditor/Public/ViewModels/Stack/NiagaraStackModuleItem.h b/Engine/Plugins/FX/Niagara/Source/NiagaraEditor/Public/ViewModels/Stack/NiagaraStackModuleItem.h
--- a/Engine/Plugins/FX/Niagara/Source/NiagaraEditor/Public/ViewModels/Stack/NiagaraStackModuleItem.h
+++ b/Engine/Plugins/FX/Niagara/Source/NiagaraEditor/Public/ViewModels/Stack/NiagaraStackModuleItem.h
@@ -7,7 +7,6 @@
 #include "NiagaraStackModuleItem.generated.h"
 
 class UNiagaraNodeFunctionCall;
-class UNiagaraNode;
 class UNiagaraStackFunctionInputCollection;
 class UNiagaraStackModuleItemOutputCollection;
 class INiagaraStackItemGroupAddUtilities;
@@ -52,7 +51,7 @@ public:
 
 	void NotifyModuleMoved();
 
-	bool CanAddInput() const;
+	bool CanAddInput(FNiagaraVariable InputParameter) const;
 
 	void AddInput(FNiagaraVariable InputParameter);
 
@@ -67,6 +66,7 @@ private:
 	void RefreshIsEnabled();
 
 private:
+	UNiagaraNodeOutput* OutputNode;
 	UNiagaraNodeFunctionCall* FunctionCallNode;
 	bool bCanMoveAndDelete;
 	bool bIsEnabled;
diff --git a/Engine/Plugins/FX/Niagara/Source/NiagaraEditorWidgets/Private/Stack/SNiagaraStackModuleItem.cpp b/Engine/Plugins/FX/Niagara/Source/NiagaraEditorWidgets/Private/Stack/SNiagaraStackModuleItem.cpp
--- a/Engine/Plugins/FX/Niagara/Source/NiagaraEditorWidgets/Private/Stack/SNiagaraStackModuleItem.cpp
+++ b/Engine/Plugins/FX/Niagara/Source/NiagaraEditorWidgets/Private/Stack/SNiagaraStackModuleItem.cpp
@@ -269,11 +269,11 @@ void SNiagaraStackModuleItem::ShowInsertModuleMenu(int32 InsertIndex)
 
 FReply SNiagaraStackModuleItem::OnModuleItemDrop(TSharedPtr<FDragDropOperation> DragDropOperation)
 {
-	if (ModuleItem->CanAddInput() && DragDropOperation->IsOfType<FNiagaraStackDragOperation>())
+	if (DragDropOperation->IsOfType<FNiagaraStackDragOperation>())
 	{
 		TSharedPtr<FNiagaraStackDragOperation> InputDragDropOperation = StaticCastSharedPtr<FNiagaraStackDragOperation>(DragDropOperation);
 		TSharedPtr<FNiagaraParameterAction> Action = StaticCastSharedPtr<FNiagaraParameterAction>(InputDragDropOperation->GetAction());
-		if (Action.IsValid())
+		if (Action.IsValid() && ModuleItem->CanAddInput(Action->GetParameter()))
 		{
 			ModuleItem->AddInput(Action->GetParameter());
 			return FReply::Handled();
@@ -285,14 +285,11 @@ FReply SNiagaraStackModuleItem::OnModuleItemDrop(TSharedPtr<FDragDropOperation>
 
 bool SNiagaraStackModuleItem::OnModuleItemAllowDrop(TSharedPtr<FDragDropOperation> DragDropOperation)
 {
-	if (ModuleItem->CanAddInput() && DragDropOperation->IsOfType<FNiagaraStackDragOperation>())
+	if (DragDropOperation->IsOfType<FNiagaraStackDragOperation>())
 	{
 		TSharedPtr<FNiagaraStackDragOperation> InputDragDropOperation = StaticCastSharedPtr<FNiagaraStackDragOperation>(DragDropOperation);
 		TSharedPtr<FNiagaraParameterAction> Action = StaticCastSharedPtr<FNiagaraParameterAction>(InputDragDropOperation->GetAction());
-		if (FNiagaraStackGraphUtilities::ParameterAllowedInExecutionCategory(Action->GetParameter().GetName(), ModuleItem->GetExecutionCategoryName()))
-		{
-			return true;
-		}
+		return Action.IsValid() && ModuleItem->CanAddInput(Action->GetParameter());
 	}
 
 	return false;
