Niagara - Speculative fix to prevent a crash on shutdown due to delegates.

diff --git a/Engine/Plugins/FX/Niagara/Source/NiagaraEditor/Private/Widgets/SNiagaraGraphNode.cpp b/Engine/Plugins/FX/Niagara/Source/NiagaraEditor/Private/Widgets/SNiagaraGraphNode.cpp
--- a/Engine/Plugins/FX/Niagara/Source/NiagaraEditor/Private/Widgets/SNiagaraGraphNode.cpp
+++ b/Engine/Plugins/FX/Niagara/Source/NiagaraEditor/Private/Widgets/SNiagaraGraphNode.cpp
@@ -14,7 +14,7 @@ SNiagaraGraphNode::SNiagaraGraphNode() : SGraphNode()
 
 SNiagaraGraphNode::~SNiagaraGraphNode()
 {
-	if (NiagaraNode)
+	if (NiagaraNode.IsValid())
 	{
 		NiagaraNode->OnVisualsChanged().RemoveAll(this);
 	}
@@ -37,12 +37,12 @@ void SNiagaraGraphNode::HandleNiagaraNodeChanged(UNiagaraNode* InNode)
 void SNiagaraGraphNode::RegisterNiagaraGraphNode(UEdGraphNode* InNode)
 {
 	NiagaraNode = Cast<UNiagaraNode>(InNode);
-	NiagaraNode->OnVisualsChanged().AddRaw(this, &SNiagaraGraphNode::HandleNiagaraNodeChanged);
+	NiagaraNode->OnVisualsChanged().AddSP(this, &SNiagaraGraphNode::HandleNiagaraNodeChanged);
 }
 
 void SNiagaraGraphNode::UpdateGraphNode()
 {
-	check(NiagaraNode != nullptr);
+	check(NiagaraNode.IsValid());
 	SGraphNode::UpdateGraphNode();
 	LastSyncedNodeChangeId = NiagaraNode->GetChangeId();
 }
diff --git a/Engine/Plugins/FX/Niagara/Source/NiagaraEditor/Private/Widgets/SNiagaraGraphNode.h b/Engine/Plugins/FX/Niagara/Source/NiagaraEditor/Private/Widgets/SNiagaraGraphNode.h
--- a/Engine/Plugins/FX/Niagara/Source/NiagaraEditor/Private/Widgets/SNiagaraGraphNode.h
+++ b/Engine/Plugins/FX/Niagara/Source/NiagaraEditor/Private/Widgets/SNiagaraGraphNode.h
@@ -23,6 +23,6 @@ public:
 protected:
 	void RegisterNiagaraGraphNode(UEdGraphNode* InNode);
 	void HandleNiagaraNodeChanged(UNiagaraNode* InNode);
-	UNiagaraNode* NiagaraNode;
+	TWeakObjectPtr<UNiagaraNode> NiagaraNode;
 	FGuid LastSyncedNodeChangeId;
 };
\ No newline at end of file
