FNiagaraCompileRequestData::GetReferencedObjects
Making sure that the cloned graph makes it onto the list of saved objects.

UNiagaraScript::BeginCacheForCookedPlatformData
Commandlets like DerivedDataCacheCommandlet call BeginCacheForCookedPlatformData directly on objects. This may mean that we have not properly gotten the HLSL script generated by the time that we get here. This does the awkward work of waiting on the parent system to finish generating the HLSL before we can begin compiling it for the GPU.

UNiagaraEditorModule::OnPreGarbageCollection
For commandlets like GenerateDistillFileSetsCommandlet, they just load the package and do some hierarchy navigation within it tracking sub-assets, then they garbage collect. Since nothing is holding onto the system at the root level, it will be summarily killed and any of references will also be killed. To thwart this for now, we are forcing the compilations to complete BEFORE garbage collection kicks in. To do otherwise for now has too many loose ends (a system may be left around after the level has been unloaded, leaving behind weird external references, etc). This should be revisited when more time is available (i.e. not days before a release is due to go out).

#tests auto tests pass

diff --git a/Engine/Plugins/FX/Niagara/Source/Niagara/Classes/NiagaraScript.h b/Engine/Plugins/FX/Niagara/Source/Niagara/Classes/NiagaraScript.h
--- a/Engine/Plugins/FX/Niagara/Source/Niagara/Classes/NiagaraScript.h
+++ b/Engine/Plugins/FX/Niagara/Source/Niagara/Classes/NiagaraScript.h
@@ -515,6 +515,7 @@ private:
 	FNiagaraScriptExecutionParameterStore ScriptExecutionParamStoreGPU;
 
 #if WITH_EDITORONLY_DATA
+	class UNiagaraSystem* FindRootSystem();
 
 	/** 'Source' data/graphs for this script */
 	UPROPERTY()
diff --git a/Engine/Plugins/FX/Niagara/Source/Niagara/Private/NiagaraScript.cpp b/Engine/Plugins/FX/Niagara/Source/Niagara/Private/NiagaraScript.cpp
--- a/Engine/Plugins/FX/Niagara/Source/Niagara/Private/NiagaraScript.cpp
+++ b/Engine/Plugins/FX/Niagara/Source/Niagara/Private/NiagaraScript.cpp
@@ -191,6 +191,22 @@ UNiagaraScript::~UNiagaraScript()
 }
 
 #if WITH_EDITORONLY_DATA
+class UNiagaraSystem* UNiagaraScript::FindRootSystem()
+{
+	UObject* Obj = GetOuter();
+	if (UNiagaraEmitter* Emitter = Cast<UNiagaraEmitter>(Obj))
+	{
+		Obj = Emitter->GetOuter();
+	}
+
+	if (UNiagaraSystem* Sys = Cast<UNiagaraSystem>(Obj))
+	{
+		return Sys;
+	}
+
+	return nullptr;
+}
+
 void UNiagaraScript::ComputeVMCompilationId(FNiagaraVMExecutableDataId& Id) const
 {
 	Id = FNiagaraVMExecutableDataId();
@@ -908,6 +924,15 @@ void UNiagaraScript::BeginCacheForCookedPlatformData(const ITargetPlatform *Targ
 {
 	if (CanBeRunOnGpu())
 	{
+		// Commandlets like DerivedDataCacheCommandlet call BeginCacheForCookedPlatformData directly on objects. This may mean that
+		// we have not properly gotten the HLSL script generated by the time that we get here. This does the awkward work of 
+		// waiting on the parent system to finish generating the HLSL before we can begin compiling it for the GPU.
+		UNiagaraSystem* SystemOwner = FindRootSystem();
+		if (SystemOwner)
+		{
+			SystemOwner->WaitForCompilationComplete();
+		}
+
 		TArray<FName> DesiredShaderFormats;
 		TargetPlatform->GetAllTargetedShaderFormats(DesiredShaderFormats);
 
diff --git a/Engine/Plugins/FX/Niagara/Source/NiagaraEditor/Private/NiagaraCompiler.cpp b/Engine/Plugins/FX/Niagara/Source/NiagaraEditor/Private/NiagaraCompiler.cpp
--- a/Engine/Plugins/FX/Niagara/Source/NiagaraEditor/Private/NiagaraCompiler.cpp
+++ b/Engine/Plugins/FX/Niagara/Source/NiagaraEditor/Private/NiagaraCompiler.cpp
@@ -298,6 +298,7 @@ void FNiagaraCompileRequestData::GetReferencedObjects(TArray<UObject*>& Objects)
 			for (int32 i = 0; i < Iter.Value().Num(); i++)
 			{
 				Objects.Add(Iter.Value()[i].ClonedScript);
+				Objects.Add(Iter.Value()[i].ClonedGraph);
 			}
 			++Iter;
 		}
diff --git a/Engine/Plugins/FX/Niagara/Source/NiagaraEditor/Private/NiagaraEditorModule.cpp b/Engine/Plugins/FX/Niagara/Source/NiagaraEditor/Private/NiagaraEditorModule.cpp
--- a/Engine/Plugins/FX/Niagara/Source/NiagaraEditor/Private/NiagaraEditorModule.cpp
+++ b/Engine/Plugins/FX/Niagara/Source/NiagaraEditor/Private/NiagaraEditorModule.cpp
@@ -297,6 +297,7 @@ void FNiagaraEditorModule::StartupModule()
 	RegisterAssetTypeAction(AssetTools, MakeShareable(new FAssetTypeActions_NiagaraParameterCollectionInstance()));
 
 	UNiagaraSettings::OnSettingsChanged().AddRaw(this, &FNiagaraEditorModule::OnNiagaraSettingsChangedEvent);
+	FCoreUObjectDelegates::GetPreGarbageCollectDelegate().AddRaw(this, &FNiagaraEditorModule::OnPreGarbageCollection);
 
 	// register details customization
 	FPropertyEditorModule& PropertyModule = FModuleManager::LoadModuleChecked<FPropertyEditorModule>("PropertyEditor");
@@ -491,7 +492,7 @@ void FNiagaraEditorModule::ShutdownModule()
 
 	MenuExtensibilityManager.Reset();
 	ToolBarExtensibilityManager.Reset();
-
+	
 	if (FModuleManager::Get().IsModuleLoaded("AssetTools"))
 	{
 		IAssetTools& AssetTools = FModuleManager::GetModuleChecked<FAssetToolsModule>("AssetTools").Get();
@@ -504,6 +505,8 @@ void FNiagaraEditorModule::ShutdownModule()
 
 	UNiagaraSettings::OnSettingsChanged().RemoveAll(this);
 
+	FCoreUObjectDelegates::GetPreGarbageCollectDelegate().RemoveAll(this);
+
 	
 	if (FModuleManager::Get().IsModuleLoaded("PropertyEditor"))
 	{
@@ -677,7 +680,7 @@ void FNiagaraEditorModule::UnregisterSettings()
 	}
 }
 
-void FNiagaraEditorModule::AddReferencedObjects( FReferenceCollector& Collector )
+void FNiagaraEditorModule::AddReferencedObjects(FReferenceCollector& Collector)
 {
 	if (SequencerSettings)
 	{
@@ -685,4 +688,22 @@ void FNiagaraEditorModule::AddReferencedObjects( FReferenceCollector& Collector
 	}
 }
 
+void FNiagaraEditorModule::OnPreGarbageCollection()
+{
+	// For commandlets like GenerateDistillFileSetsCommandlet, they just load the package and do some hierarchy navigation within it 
+	// tracking sub-assets, then they garbage collect. Since nothing is holding onto the system at the root level, it will be summarily
+	// killed and any of references will also be killed. To thwart this for now, we are forcing the compilations to complete BEFORE
+	// garbage collection kicks in. To do otherwise for now has too many loose ends (a system may be left around after the level has been
+	// unloaded, leaving behind weird external references, etc). This should be revisited when more time is available (i.e. not days before a 
+	// release is due to go out).
+	for (TObjectIterator<UNiagaraSystem> It; It; ++It)
+	{
+		UNiagaraSystem* System = *It;
+		if (System && System->HasOutstandingCompilationRequests())
+		{
+			System->WaitForCompilationComplete();
+		}
+	}
+}
+
 #undef LOCTEXT_NAMESPACE
diff --git a/Engine/Plugins/FX/Niagara/Source/NiagaraEditor/Private/NiagaraHlslTranslator.cpp b/Engine/Plugins/FX/Niagara/Source/NiagaraEditor/Private/NiagaraHlslTranslator.cpp
--- a/Engine/Plugins/FX/Niagara/Source/NiagaraEditor/Private/NiagaraHlslTranslator.cpp
+++ b/Engine/Plugins/FX/Niagara/Source/NiagaraEditor/Private/NiagaraHlslTranslator.cpp
@@ -4685,6 +4685,13 @@ void FHlslNiagaraTranslator::RegisterFunctionCall(ENiagaraScriptUsage ScriptUsag
 		{
 			SCOPE_CYCLE_COUNTER(STAT_NiagaraEditor_Module_NiagaraHLSLTranslator_FuncBody);
 
+			if (OutSignature.Name == NAME_None)
+			{
+				const FString* ModuleAlias = ActiveHistoryForFunctionCalls.GetModuleAlias();
+				Error(FText::Format(LOCTEXT("FunctionCallMissingFunction", "Function call signature does not reference a function. Top-level module: {0} Source: {1}"), ModuleAlias ? FText::FromString(*ModuleAlias) : FText::FromString(TEXT("Unknown module")), FText::FromString(CompileOptions.FullName)), nullptr, nullptr);
+				return;
+			}
+
 			//We've not compiled this function yet so compile it now.
 			EnterFunction(InName, OutSignature, Inputs, CallNodeId);
 
diff --git a/Engine/Plugins/FX/Niagara/Source/NiagaraEditor/Public/NiagaraEditorModule.h b/Engine/Plugins/FX/Niagara/Source/NiagaraEditor/Public/NiagaraEditorModule.h
--- a/Engine/Plugins/FX/Niagara/Source/NiagaraEditor/Public/NiagaraEditorModule.h
+++ b/Engine/Plugins/FX/Niagara/Source/NiagaraEditor/Public/NiagaraEditorModule.h
@@ -95,6 +95,7 @@ public:
 private:
 	void RegisterAssetTypeAction(IAssetTools& AssetTools, TSharedRef<IAssetTypeActions> Action);
 	void OnNiagaraSettingsChangedEvent(const FString& PropertyName, const UNiagaraSettings* Settings);
+	void OnPreGarbageCollection();
 
 
 	/** FGCObject interface */
