Niagara - Fix an issue where the list of relevant scripts in the shared script view model could get out of sync when changing properties on an emitter such as interpolated spawning and gpu simulation.  This could result in an emitter recompiling forever if it started as GPU and was then switched to CPU and then a force compile was requested since it would include the GPU script when determing the compile status, but it would never compile it.

diff --git a/Engine/Plugins/FX/Niagara/Source/NiagaraEditor/Private/ViewModels/NiagaraEmitterViewModel.cpp b/Engine/Plugins/FX/Niagara/Source/NiagaraEditor/Private/ViewModels/NiagaraEmitterViewModel.cpp
--- a/Engine/Plugins/FX/Niagara/Source/NiagaraEditor/Private/ViewModels/NiagaraEmitterViewModel.cpp
+++ b/Engine/Plugins/FX/Niagara/Source/NiagaraEditor/Private/ViewModels/NiagaraEmitterViewModel.cpp
@@ -36,26 +36,9 @@ FNiagaraEmitterViewModel::FNiagaraEmitterViewModel(UNiagaraEmitter* InEmitter, T
 	, Simulation(InSimulation)
 	, SharedScriptViewModel(MakeShareable(new FNiagaraScriptViewModel(InEmitter, LOCTEXT("SharedDisplayName", "Graph"), ENiagaraParameterEditMode::EditAll)))
 	, bUpdatingSelectionInternally(false)
-	, LastEventScriptStatus(ENiagaraScriptCompileStatus::NCS_Unknown)
-{
-	//UE_LOG(LogNiagaraEditor, Warning, TEXT("FNiagaraEmitterViewModel %p"), this);
-
-	ExecutionStateEnum = FindObject<UEnum>(ANY_PACKAGE, TEXT("ENiagaraExecutionState"));
-
-	if (Emitter.IsValid() && Emitter->GetEventHandlers().Num() != 0 && Emitter->GetEventHandlers()[0].Script && Emitter->GetEventHandlers()[0].Script->GetVMExecutableData().IsValid()
-		&& Emitter->GetEventHandlers()[0].Script->GetVMExecutableData().ByteCode.Num() != 0)
-	{
-		LastEventScriptStatus = ENiagaraScriptCompileStatus::NCS_UpToDate;
-	}
-	
-	if (Emitter.IsValid())
-	{
-		Emitter->OnEmitterVMCompiled().AddRaw(this, &FNiagaraEmitterViewModel::OnVMCompiled);
-	}
-
-	RegisteredHandle = RegisterViewModelWithMap(Emitter.Get(), this);
-
-	AddScriptEventHandlers();
+	, ExecutionStateEnum(FindObject<UEnum>(ANY_PACKAGE, TEXT("ENiagaraExecutionState")))
+{	
+	SetEmitter(InEmitter);
 }
 
 void FNiagaraEmitterViewModel::Cleanup()
@@ -95,6 +78,7 @@ void FNiagaraEmitterViewModel::SetEmitter(UNiagaraEmitter* InEmitter)
 	if (Emitter.IsValid())
 	{
 		Emitter->OnEmitterVMCompiled().RemoveAll(this);
+		Emitter->OnPropertiesChanged().RemoveAll(this);
 	}
 
 	UnregisterViewModelWithMap(RegisteredHandle);
@@ -106,6 +90,7 @@ void FNiagaraEmitterViewModel::SetEmitter(UNiagaraEmitter* InEmitter)
 	if (Emitter.IsValid())
 	{
 		Emitter->OnEmitterVMCompiled().AddRaw(this, &FNiagaraEmitterViewModel::OnVMCompiled);
+		Emitter->OnPropertiesChanged().AddRaw(this, &FNiagaraEmitterViewModel::OnEmitterPropertiesChanged);
 	}
 
 	AddScriptEventHandlers();
@@ -395,4 +380,11 @@ void FNiagaraEmitterViewModel::ScriptParameterStoreChanged(const FNiagaraParamet
 	OnScriptParameterStoreChangedDelegate.Broadcast(ChangedParameterStore, OwningScript);
 }
 
+void FNiagaraEmitterViewModel::OnEmitterPropertiesChanged()
+{
+	// When the properties change we reset the scripts on the script view model because gpu/cpu or interpolation may have changed.
+	SharedScriptViewModel->SetScripts(Emitter.Get());
+	OnPropertyChangedDelegate.Broadcast();
+}
+
 #undef LOCTEXT_NAMESPACE
diff --git a/Engine/Plugins/FX/Niagara/Source/NiagaraEditor/Public/ViewModels/NiagaraEmitterViewModel.h b/Engine/Plugins/FX/Niagara/Source/NiagaraEditor/Public/ViewModels/NiagaraEmitterViewModel.h
--- a/Engine/Plugins/FX/Niagara/Source/NiagaraEditor/Public/ViewModels/NiagaraEmitterViewModel.h
+++ b/Engine/Plugins/FX/Niagara/Source/NiagaraEditor/Public/ViewModels/NiagaraEmitterViewModel.h
@@ -87,6 +87,8 @@ private:
 
 	void ScriptParameterStoreChanged(const FNiagaraParameterStore& ChangedParameterStore, const UNiagaraScript& OwningScript);
 
+	void OnEmitterPropertiesChanged();
+
 private:
 
 	/** The text format stats display .*/
@@ -119,8 +121,6 @@ private:
 
 	FOnScriptParameterStoreChanged OnScriptParameterStoreChangedDelegate;
 
-	ENiagaraScriptCompileStatus LastEventScriptStatus;
-
 	TNiagaraViewModelManager<UNiagaraEmitter, FNiagaraEmitterViewModel>::Handle RegisteredHandle;
 
 	UEnum* ExecutionStateEnum;
