Niagara - Adding and removing pins from an assignment node wasn't correctly invalidating the graph which I broke with my crash fix 4058428 since I thought the refresh call would do that.

diff --git a/Engine/Plugins/FX/Niagara/Source/NiagaraEditor/Private/NiagaraNodeAssignment.cpp b/Engine/Plugins/FX/Niagara/Source/NiagaraEditor/Private/NiagaraNodeAssignment.cpp
--- a/Engine/Plugins/FX/Niagara/Source/NiagaraEditor/Private/NiagaraNodeAssignment.cpp
+++ b/Engine/Plugins/FX/Niagara/Source/NiagaraEditor/Private/NiagaraNodeAssignment.cpp
@@ -238,6 +238,7 @@ void UNiagaraNodeAssignment::AddParameter(FNiagaraVariable InVar, FString InDefa
 	AddAssignmentTarget(InVar, &InDefaultValue);
 
 	RefreshFromExternalChanges();
+	MarkNodeRequiresSynchronization(__FUNCTION__, true);
 	OnInputsChangedDelegate.Broadcast();
 }
 
@@ -265,6 +266,7 @@ void UNiagaraNodeAssignment::RemoveParameter(const FNiagaraVariable& InVar)
 	}
 
 	RefreshFromExternalChanges();
+	MarkNodeRequiresSynchronization(__FUNCTION__, true);
 	OnInputsChangedDelegate.Broadcast();
 }
 
