Niagara - Made few changes to stack issue handling while fixing an issue where the stack error wouldn't change when the compile error changed.
+ Changed the unique identifier for stack issues to be automatically generated from a hash of the combined stack editor data key and the long description of the error.
+ Changed the stack issue unique identifier from an FName to an FString to avoid poluting the name table with lots of generated hash strings.
+ Encapsulated all of the stack issue data to validate the required inputs.

diff --git a/Engine/Plugins/FX/Niagara/Source/NiagaraEditor/Private/NiagaraStackEditorData.cpp b/Engine/Plugins/FX/Niagara/Source/NiagaraEditor/Private/NiagaraStackEditorData.cpp
--- a/Engine/Plugins/FX/Niagara/Source/NiagaraEditor/Private/NiagaraStackEditorData.cpp
+++ b/Engine/Plugins/FX/Niagara/Source/NiagaraEditor/Private/NiagaraStackEditorData.cpp
@@ -71,7 +71,7 @@ void UNiagaraStackEditorData::SetLastScrollPosition(double InLastScrollPosition)
 	LastScrollPosition = InLastScrollPosition;
 }
 
-void UNiagaraStackEditorData::DismissStackIssue(FName IssueId)
+void UNiagaraStackEditorData::DismissStackIssue(FString IssueId)
 {
 	DismissedStackIssueIds.AddUnique(IssueId);
 }
@@ -81,7 +81,7 @@ void UNiagaraStackEditorData::UndismissAllIssues()
 	DismissedStackIssueIds.Empty();
 }
 
-const TArray<FName>& UNiagaraStackEditorData::GetDismissedStackIssueIds()
+const TArray<FString>& UNiagaraStackEditorData::GetDismissedStackIssueIds()
 {
 	return DismissedStackIssueIds;
 }
diff --git a/Engine/Plugins/FX/Niagara/Source/NiagaraEditor/Private/ViewModels/Stack/NiagaraStackEntry.cpp b/Engine/Plugins/FX/Niagara/Source/NiagaraEditor/Private/ViewModels/Stack/NiagaraStackEntry.cpp
--- a/Engine/Plugins/FX/Niagara/Source/NiagaraEditor/Private/ViewModels/Stack/NiagaraStackEntry.cpp
+++ b/Engine/Plugins/FX/Niagara/Source/NiagaraEditor/Private/ViewModels/Stack/NiagaraStackEntry.cpp
@@ -3,6 +3,7 @@
 #include "ViewModels/Stack/NiagaraStackEntry.h"
 #include "ViewModels/Stack/NiagaraStackErrorItem.h"
 #include "NiagaraStackEditorData.h"
+#include "Misc/SecureHash.h"
 
 const FName UNiagaraStackEntry::FExecutionCategoryNames::System = TEXT("System");
 const FName UNiagaraStackEntry::FExecutionCategoryNames::Emitter = TEXT("Emitter");
@@ -14,6 +15,91 @@ const FName UNiagaraStackEntry::FExecutionSubcategoryNames::Spawn = TEXT("Spawn"
 const FName UNiagaraStackEntry::FExecutionSubcategoryNames::Update = TEXT("Update");
 const FName UNiagaraStackEntry::FExecutionSubcategoryNames::Event = TEXT("Event");
 
+UNiagaraStackEntry::FStackIssueFix::FStackIssueFix()
+{
+}
+
+UNiagaraStackEntry::FStackIssueFix::FStackIssueFix(FText InDescription, FStackIssueFixDelegate InFixDelegate)
+	: Description(InDescription)
+	, FixDelegate(InFixDelegate)
+{
+	checkf(Description.IsEmptyOrWhitespace() == false, TEXT("Description can not be empty."));
+	checkf(InFixDelegate.IsBound(), TEXT("Fix delegate must be bound."));
+}
+
+bool UNiagaraStackEntry::FStackIssueFix::IsValid() const
+{
+	return FixDelegate.IsBound();
+}
+
+const FText& UNiagaraStackEntry::FStackIssueFix::GetDescription() const
+{
+	return Description;
+}
+
+const UNiagaraStackEntry::FStackIssueFixDelegate& UNiagaraStackEntry::FStackIssueFix::GetFixDelegate() const
+{
+	return FixDelegate;
+}
+
+UNiagaraStackEntry::FStackIssue::FStackIssue()
+{
+}
+
+UNiagaraStackEntry::FStackIssue::FStackIssue(EStackIssueSeverity InSeverity, FText InShortDescription, FText InLongDescription, FString InStackEditorDataKey, bool bInCanBeDismissed)
+	: Severity(InSeverity)
+	, ShortDescription(InShortDescription)
+	, LongDescription(InLongDescription)
+	, UniqueIdentifier(FMD5::HashAnsiString(*FString::Printf(TEXT("%s-%s"), *InStackEditorDataKey, *LongDescription.ToString())))
+	, bCanBeDismissed(bInCanBeDismissed)
+{
+	checkf(ShortDescription.IsEmptyOrWhitespace() == false, TEXT("Short description can not be empty."));
+	checkf(LongDescription.IsEmptyOrWhitespace() == false, TEXT("Long description can not be empty."));
+	checkf(InStackEditorDataKey.IsEmpty() == false, TEXT("Stack editor data key can not be empty."));
+}
+
+bool UNiagaraStackEntry::FStackIssue::IsValid()
+{
+	return UniqueIdentifier.IsEmpty() == false;
+}
+
+EStackIssueSeverity UNiagaraStackEntry::FStackIssue::GetSeverity() const
+{
+	return Severity;
+}
+
+const FText& UNiagaraStackEntry::FStackIssue::GetShortDescription() const
+{
+	return ShortDescription;
+}
+
+const FText& UNiagaraStackEntry::FStackIssue::GetLongDescription() const
+{
+	return LongDescription;
+}
+
+bool UNiagaraStackEntry::FStackIssue::GetCanBeDismissed() const
+{
+	return bCanBeDismissed;
+}
+
+const FString& UNiagaraStackEntry::FStackIssue::GetUniqueIdentifier() const
+{
+	return UniqueIdentifier;
+}
+
+void UNiagaraStackEntry::FStackIssue::AddFix(FStackIssueFix InFix)
+{
+	checkf(IsValid(), TEXT("Issue is not valid."));
+	checkf(InFix.IsValid(), TEXT("Fix must be valid"));
+	Fixes.Add(InFix);
+}
+
+const TArray<UNiagaraStackEntry::FStackIssueFix>& UNiagaraStackEntry::FStackIssue::GetFixes() const
+{
+	return Fixes;
+}
+
 UNiagaraStackEntry::UNiagaraStackEntry()
 	: IndentLevel(0)
 {
@@ -305,7 +391,7 @@ void UNiagaraStackEntry::RefreshChildren()
 	}
 	
 	// Stack issues refresh
-	NewStackIssues.RemoveAll([=](const FStackIssue& Issue) { return GetStackEditorData().GetDismissedStackIssueIds().Contains(Issue.UniqueIdentifier); }); 
+	NewStackIssues.RemoveAll([=](const FStackIssue& Issue) { return Issue.GetCanBeDismissed() && GetStackEditorData().GetDismissedStackIssueIds().Contains(Issue.GetUniqueIdentifier()); }); 
 
 	StackIssues.Empty();
 	StackIssues.Append(NewStackIssues);
@@ -332,7 +418,7 @@ void UNiagaraStackEntry::RefreshStackErrorChildren()
 	{
 		UNiagaraStackErrorItem* ErrorEntry = nullptr;
 		UNiagaraStackErrorItem** Found = ErrorChildren.FindByPredicate(
-			[&](UNiagaraStackErrorItem* CurrentChild) { return CurrentChild->GetStackIssue().UniqueIdentifier == Issue.UniqueIdentifier; });
+			[&](UNiagaraStackErrorItem* CurrentChild) { return CurrentChild->GetStackIssue().GetUniqueIdentifier() == Issue.GetUniqueIdentifier(); });
 		if (Found == nullptr)
 		{
 			ErrorEntry = NewObject<UNiagaraStackErrorItem>(this);
diff --git a/Engine/Plugins/FX/Niagara/Source/NiagaraEditor/Private/ViewModels/Stack/NiagaraStackErrorItem.cpp b/Engine/Plugins/FX/Niagara/Source/NiagaraEditor/Private/ViewModels/Stack/NiagaraStackErrorItem.cpp
--- a/Engine/Plugins/FX/Niagara/Source/NiagaraEditor/Private/ViewModels/Stack/NiagaraStackErrorItem.cpp
+++ b/Engine/Plugins/FX/Niagara/Source/NiagaraEditor/Private/ViewModels/Stack/NiagaraStackErrorItem.cpp
@@ -9,7 +9,7 @@
 
 void UNiagaraStackErrorItem::Initialize(FRequiredEntryData InRequiredEntryData, FStackIssue InStackIssue, FString InStackEditorDataKey)
 {
-	FString ErrorStackEditorDataKey = FString::Printf(TEXT("%s-Error-%s"), *InStackEditorDataKey, InStackIssue.UniqueIdentifier.GetPlainANSIString());
+	FString ErrorStackEditorDataKey = FString::Printf(TEXT("%s-Error-%s"), *InStackEditorDataKey, *InStackIssue.GetUniqueIdentifier());
 	Super::Initialize(InRequiredEntryData, ErrorStackEditorDataKey);
 	StackIssue = InStackIssue;
 	EntryStackEditorDataKey = InStackEditorDataKey;
@@ -17,7 +17,7 @@ void UNiagaraStackErrorItem::Initialize(FRequiredEntryData InRequiredEntryData,
 
 FText UNiagaraStackErrorItem::GetDisplayName() const
 {
-	return StackIssue.ShortDescription;
+	return StackIssue.GetShortDescription();
 }
 
 UNiagaraStackEntry::EStackRowStyle UNiagaraStackErrorItem::GetStackRowStyle() const
@@ -51,20 +51,20 @@ void UNiagaraStackErrorItem::RefreshChildrenInternal(const TArray<UNiagaraStackE
 	NewChildren.Add(ErrorEntryLongDescription);
 
 	// fixes
-	for (int i = 0; i < StackIssue.Fixes.Num(); i++)
+	for (int i = 0; i < StackIssue.GetFixes().Num(); i++)
 	{
 		UNiagaraStackErrorItemFix* ErrorEntryFix = FindCurrentChildOfTypeByPredicate<UNiagaraStackErrorItemFix>(CurrentChildren,
-			[&](UNiagaraStackErrorItemFix* CurrentChild) { return CurrentChild->GetStackIssueFix() == StackIssue.Fixes[i]; });
+			[&](UNiagaraStackErrorItemFix* CurrentChild) { return CurrentChild->GetStackIssueFix() == StackIssue.GetFixes()[i]; });
 		if (ErrorEntryFix == nullptr)
 		{
 			ErrorEntryFix = NewObject<UNiagaraStackErrorItemFix>(this);
-			ErrorEntryFix->Initialize(CreateDefaultChildRequiredData(), StackIssue, StackIssue.Fixes[i], EntryStackEditorDataKey);
+			ErrorEntryFix->Initialize(CreateDefaultChildRequiredData(), StackIssue, StackIssue.GetFixes()[i], EntryStackEditorDataKey);
 		}
 		ErrorEntryFix->OnIssueFixed().AddUObject(this, &UNiagaraStackErrorItem::IssueFixed);
 		NewChildren.Add(ErrorEntryFix);
 	}
 	// dismiss button
-	if (StackIssue.bCanBeDismissed)
+	if (StackIssue.GetCanBeDismissed())
 	{
 		UNiagaraStackErrorItemDismiss* ErrorEntryDismiss = FindCurrentChildOfTypeByPredicate<UNiagaraStackErrorItemDismiss>(CurrentChildren,
 			[&](UNiagaraStackErrorItemDismiss* CurrentChild) { return true; });
@@ -93,7 +93,7 @@ void UNiagaraStackErrorItemLongDescription::Initialize(FRequiredEntryData InRequ
 
 FText UNiagaraStackErrorItemLongDescription::GetDisplayName() const
 {
-	return StackIssue.LongDescription;
+	return StackIssue.GetLongDescription();
 }
 
 UNiagaraStackEntry::EStackRowStyle UNiagaraStackErrorItemLongDescription::GetStackRowStyle() const
@@ -113,18 +113,12 @@ void UNiagaraStackErrorItemFix::Initialize(FRequiredEntryData InRequiredEntryDat
 
 FText UNiagaraStackErrorItemFix::FixDescription() const
 {
-	if (IssueFix.Description.IsEmptyOrWhitespace())
-	{
-		ensureMsgf(false, TEXT("Fix description is empty or whitespace."));
-		return LOCTEXT("FixIssue", "Fix the issue");
-	}
-	return IssueFix.Description;
+	return IssueFix.GetDescription();
 }
 
 FReply UNiagaraStackErrorItemFix::OnTryFixError()
 {
-	ensureMsgf(IssueFix.FixDelegate.IsBound(), TEXT("Fix delegate is not bound!"));
-	IssueFix.FixDelegate.ExecuteIfBound();
+	IssueFix.GetFixDelegate().ExecuteIfBound();
 	OnIssueFixed().Broadcast();
 	return FReply::Handled();
 }
@@ -151,22 +145,15 @@ void UNiagaraStackErrorItemDismiss::Initialize(FRequiredEntryData InRequiredEntr
 	FString ErrorStackEditorDataKey = FString::Printf(TEXT("Dismiss-%s"), *InStackEditorDataKey);
 	UNiagaraStackEntry::Initialize(InRequiredEntryData, ErrorStackEditorDataKey);
 	StackIssue = InStackIssue;
-	IssueFix.Description = LOCTEXT("DismissError", "Dismiss the issue without fixing (I know what I'm doing).");
+	IssueFix = FStackIssueFix(
+		LOCTEXT("DismissError", "Dismiss the issue without fixing (I know what I'm doing)."),
+		FStackIssueFixDelegate::CreateUObject(this, &UNiagaraStackErrorItemDismiss::DismissIssue));
 }
 
 void UNiagaraStackErrorItemDismiss::DismissIssue()
 {
 	GetStackEditorData().Modify();
-	GetStackEditorData().DismissStackIssue(StackIssue.UniqueIdentifier);
-}
-
-FReply UNiagaraStackErrorItemDismiss::OnTryFixError()
-{
-	ensureMsgf(IssueFixedDelegate.IsBound(), TEXT("Issue fixed delegate is not bound on the dismiss entry!"));
-	FScopedTransaction ScopedTransaction(LOCTEXT("DismissIssue", "Dismiss issue"));
-	DismissIssue();
-	OnIssueFixed().Broadcast();
-	return FReply::Handled();
+	GetStackEditorData().DismissStackIssue(StackIssue.GetUniqueIdentifier());
 }
 
 UNiagaraStackEntry::EStackRowStyle UNiagaraStackErrorItemDismiss::GetStackRowStyle() const
diff --git a/Engine/Plugins/FX/Niagara/Source/NiagaraEditor/Private/ViewModels/Stack/NiagaraStackFunctionInputCollection.cpp b/Engine/Plugins/FX/Niagara/Source/NiagaraEditor/Private/ViewModels/Stack/NiagaraStackFunctionInputCollection.cpp
--- a/Engine/Plugins/FX/Niagara/Source/NiagaraEditor/Private/ViewModels/Stack/NiagaraStackFunctionInputCollection.cpp
+++ b/Engine/Plugins/FX/Niagara/Source/NiagaraEditor/Private/ViewModels/Stack/NiagaraStackFunctionInputCollection.cpp
@@ -192,52 +192,59 @@ void UNiagaraStackFunctionInputCollection::RefreshIssues(TArray<FStackIssue>& Ne
 				FNiagaraParameterHandle(OverridePin->PinName).GetNamespace().ToString() == InputFunctionCallNode->GetFunctionName() &&
 				ValidAliasedInputNames.Contains(OverridePin->PinName) == false)
 			{
-				UNiagaraStackEntry::FStackIssue InvalidInputError;
-				InvalidInputError.ShortDescription = FText::Format(LOCTEXT("InvalidInputSummaryFormat", "Invalid Input: {0}"), FText::FromString(OverridePin->PinName.ToString()));
-				InvalidInputError.LongDescription = FText::Format(LOCTEXT("InvalidInputFormat", "The input {0} was previously overriden but is no longer exposed by the function {1}.\nPress the fix button to remove this unused override data,\nor check the function definition to see why this input is no longer exposed."),
-					FText::FromString(OverridePin->PinName.ToString()), FText::FromString(InputFunctionCallNode->GetFunctionName()));
-				InvalidInputError.UniqueIdentifier = FName(*FString::Printf(TEXT("%s-InvalidInput-%s"), *GetStackEditorDataKey(), *OverridePin->PinName.ToString()));
-				UNiagaraStackEntry::FStackIssueFix Fix;
-				Fix.Description = LOCTEXT("RemoveInvalidInputTransaction", "Remove invalid input override.");
-				Fix.FixDelegate.BindLambda([=]()
-				{
-					FScopedTransaction ScopedTransaction(Fix.Description);
-					TArray<TWeakObjectPtr<UNiagaraDataInterface>> RemovedDataObjects;
-					FNiagaraStackGraphUtilities::RemoveNodesForStackFunctionInputOverridePin(*OverridePin, RemovedDataObjects);
-					for (TWeakObjectPtr<UNiagaraDataInterface> RemovedDataObject : RemovedDataObjects)
+				FStackIssue InvalidInputOverrideError(
+					EStackIssueSeverity::Error,
+					FText::Format(LOCTEXT("InvalidInputSummaryFormat", "Invalid Input Override: {0}"), FText::FromString(OverridePin->PinName.ToString())),
+					FText::Format(LOCTEXT("InvalidInputFormat", "The input {0} was previously overriden but is no longer exposed by the function {1}.\nPress the fix button to remove this unused override data,\nor check the function definition to see why this input is no longer exposed."),
+						FText::FromString(OverridePin->PinName.ToString()), FText::FromString(InputFunctionCallNode->GetFunctionName())),
+					GetStackEditorDataKey(),
+					false);
+
+				FText FixDescription = LOCTEXT("RemoveInvalidInputTransaction", "Remove invalid input override.");
+				FStackIssueFix RemoveInputOverrideFix(
+					FixDescription,
+					UNiagaraStackEntry::FStackIssueFixDelegate::CreateLambda([=]()
 					{
-						if (RemovedDataObject.IsValid())
+						FScopedTransaction ScopedTransaction(FixDescription);
+						TArray<TWeakObjectPtr<UNiagaraDataInterface>> RemovedDataObjects;
+						FNiagaraStackGraphUtilities::RemoveNodesForStackFunctionInputOverridePin(*OverridePin, RemovedDataObjects);
+						for (TWeakObjectPtr<UNiagaraDataInterface> RemovedDataObject : RemovedDataObjects)
 						{
-							OnDataObjectModified().Broadcast(RemovedDataObject.Get());
+							if (RemovedDataObject.IsValid())
+							{
+								OnDataObjectModified().Broadcast(RemovedDataObject.Get());
+							}
 						}
-					}
-					OverridePin->GetOwningNode()->RemovePin(OverridePin);
-				});
-				InvalidInputError.Fixes.Add(Fix);
-				NewIssues.Add(InvalidInputError);
+						OverridePin->GetOwningNode()->RemovePin(OverridePin);
+					}));
+
+				InvalidInputOverrideError.AddFix(RemoveInputOverrideFix);
+				NewIssues.Add(InvalidInputOverrideError);
 			}
 		}
 	}
 
 	for (const FName& DuplicateInputName : DuplicateInputNames)
 	{
-		UNiagaraStackEntry::FStackIssue DuplicateInputError;
-
-		FString DuplicateInputNameString = DuplicateInputName.ToString();
-
-		DuplicateInputError.ShortDescription = FText::Format(LOCTEXT("DuplicateInputSummaryFormat", "Duplicate Input: {0}"), FText::FromString(DuplicateInputNameString));
-		DuplicateInputError.LongDescription = FText::Format(LOCTEXT("DuplicateInputFormat", "There are multiple inputs with the same name {0}, but different types exposed by the function {1}.\nThis is not suppored and must be fixed in the script that defines this function."),
-			FText::FromString(DuplicateInputNameString), FText::FromString(InputFunctionCallNode->GetFunctionName()));
-		DuplicateInputError.UniqueIdentifier = FName(*FString::Printf(TEXT("%s-DuplicateInput-%s"), *GetStackEditorDataKey(), *DuplicateInputNameString));
+		FStackIssue DuplicateInputError(
+			EStackIssueSeverity::Error,
+			FText::Format(LOCTEXT("DuplicateInputSummaryFormat", "Duplicate Input: {0}"), FText::FromName(DuplicateInputName)),
+			FText::Format(LOCTEXT("DuplicateInputFormat", "There are multiple inputs with the same name {0}, but different types exposed by the function {1}.\nThis is not suppored and must be fixed in the script that defines this function."),
+				FText::FromName(DuplicateInputName), FText::FromString(InputFunctionCallNode->GetFunctionName())),
+			GetStackEditorDataKey(),
+			false);
 		NewIssues.Add(DuplicateInputError);
 	}
 
 	for (const UEdGraphPin* PinWithInvalidType : PinsWithInvalidTypes)
 	{
-		UNiagaraStackEntry::FStackIssue InputWithInvalidTypeError;
-		InputWithInvalidTypeError.ShortDescription = FText::Format(LOCTEXT("InputWithInvalidTypeSummaryFormat", "Input has an invalid type: {0}"), FText::FromString(PinWithInvalidType->PinName.ToString()));
-		InputWithInvalidTypeError.LongDescription = FText::Format(LOCTEXT("InputWithInvalidTypeFormat", "The input {0} on function {1} has a type which is invalid.\nThe type of this input likely doesn't exist anymore.\nThis input must be fixed in the script before this module can be used."),
-			FText::FromString(PinWithInvalidType->PinName.ToString()), FText::FromString(InputFunctionCallNode->GetFunctionName()));
+		FStackIssue InputWithInvalidTypeError(
+			EStackIssueSeverity::Error,
+			FText::Format(LOCTEXT("InputWithInvalidTypeSummaryFormat", "Input has an invalid type: {0}"), FText::FromName(PinWithInvalidType->PinName)),
+			FText::Format(LOCTEXT("InputWithInvalidTypeFormat", "The input {0} on function {1} has a type which is invalid.\nThe type of this input likely doesn't exist anymore.\nThis input must be fixed in the script before this module can be used."),
+				FText::FromName(PinWithInvalidType->PinName), FText::FromString(InputFunctionCallNode->GetFunctionName())),
+			GetStackEditorDataKey(),
+			false);
 		NewIssues.Add(InputWithInvalidTypeError);
 	}
 }
diff --git a/Engine/Plugins/FX/Niagara/Source/NiagaraEditor/Private/ViewModels/Stack/NiagaraStackItem.cpp b/Engine/Plugins/FX/Niagara/Source/NiagaraEditor/Private/ViewModels/Stack/NiagaraStackItem.cpp
--- a/Engine/Plugins/FX/Niagara/Source/NiagaraEditor/Private/ViewModels/Stack/NiagaraStackItem.cpp
+++ b/Engine/Plugins/FX/Niagara/Source/NiagaraEditor/Private/ViewModels/Stack/NiagaraStackItem.cpp
@@ -35,9 +35,9 @@ uint32 UNiagaraStackItem::GetRecursiveStackIssuesCount() const
 		EStackIssueSeverity MinSeverity = EStackIssueSeverity::Info;
 		for (auto Issue : RecursiveIssues)
 		{
-			if (Issue->GetStackIssue().Severity < MinSeverity)
+			if (Issue->GetStackIssue().GetSeverity() < MinSeverity)
 			{
-				MinSeverity = Issue->GetStackIssue().Severity;
+				MinSeverity = Issue->GetStackIssue().GetSeverity();
 			}
 		}
 		HighestIssueSeverity = MinSeverity;
diff --git a/Engine/Plugins/FX/Niagara/Source/NiagaraEditor/Private/ViewModels/Stack/NiagaraStackItemGroup.cpp b/Engine/Plugins/FX/Niagara/Source/NiagaraEditor/Private/ViewModels/Stack/NiagaraStackItemGroup.cpp
--- a/Engine/Plugins/FX/Niagara/Source/NiagaraEditor/Private/ViewModels/Stack/NiagaraStackItemGroup.cpp
+++ b/Engine/Plugins/FX/Niagara/Source/NiagaraEditor/Private/ViewModels/Stack/NiagaraStackItemGroup.cpp
@@ -47,9 +47,9 @@ uint32 UNiagaraStackItemGroup::GetRecursiveStackIssuesCount() const
 		EStackIssueSeverity MinSeverity = EStackIssueSeverity::Info;
 		for (auto Issue : RecursiveIssues)
 		{
-			if (Issue->GetStackIssue().Severity < MinSeverity)
+			if (Issue->GetStackIssue().GetSeverity() < MinSeverity)
 			{
-				MinSeverity = Issue->GetStackIssue().Severity;
+				MinSeverity = Issue->GetStackIssue().GetSeverity();
 			}
 		}
 		HighestIssueSeverity = MinSeverity;
diff --git a/Engine/Plugins/FX/Niagara/Source/NiagaraEditor/Private/ViewModels/Stack/NiagaraStackModuleItem.cpp b/Engine/Plugins/FX/Niagara/Source/NiagaraEditor/Private/ViewModels/Stack/NiagaraStackModuleItem.cpp
--- a/Engine/Plugins/FX/Niagara/Source/NiagaraEditor/Private/ViewModels/Stack/NiagaraStackModuleItem.cpp
+++ b/Engine/Plugins/FX/Niagara/Source/NiagaraEditor/Private/ViewModels/Stack/NiagaraStackModuleItem.cpp
@@ -166,10 +166,13 @@ void UNiagaraStackModuleItem::RefreshIssues(TArray<FStackIssue>& NewIssues)
 	{
 		if (!FunctionCallNode->ScriptIsValid())
 		{
-			UNiagaraStackEntry::FStackIssue InvalidScriptError;
-			InvalidScriptError.ShortDescription = LOCTEXT("MissingModule", "The referenced module's script asset is missing.");
-			InvalidScriptError.LongDescription = InvalidScriptError.ShortDescription;
-			InvalidScriptError.UniqueIdentifier = FName(*FString::Printf(TEXT("MissingModule-%s"), *GetStackEditorDataKey()));
+			FStackIssue InvalidScriptError(
+				EStackIssueSeverity::Error,
+				LOCTEXT("MissingModuleShortDescription", "Invalid module script."),
+				LOCTEXT("MissingModuleShortDescription", "The script this module is supposed to execute is missing or invalid for other reasons.  If it depends on an external script that no longer exists there will be load errors in the log."),
+				GetStackEditorDataKey(),
+				false);
+
 			NewIssues.Add(InvalidScriptError);
 		}
 
@@ -177,18 +180,23 @@ void UNiagaraStackModuleItem::RefreshIssues(TArray<FStackIssue>& NewIssues)
 		if (!IsEnabled.IsSet())
 		{
 			bIsEnabled = false;
-			UNiagaraStackEntry::FStackIssue InconsistentEnabledError;
-			InconsistentEnabledError.ShortDescription = LOCTEXT("InconsistentEnabledErrorSummary", "The enabled state for module is inconsistent.");
-			InconsistentEnabledError.LongDescription = LOCTEXT("InconsistentEnabledError", "This module is using multiple functions and their enabled state is inconsistent.\nClick fix to make all of the functions for this module enabled.");
-			InconsistentEnabledError.UniqueIdentifier = FName(*FString::Printf(TEXT("InconsistentEnabled-%s"), *GetStackEditorDataKey()));
-			UNiagaraStackEntry::FStackIssueFix Fix;
-			Fix.Description = LOCTEXT("EnableModule", "Enable module");
-			Fix.FixDelegate.BindLambda([=]()
-			{
-				FScopedTransaction ScopedTransaction(Fix.Description);
-				SetIsEnabled(true);
-			});
-			InconsistentEnabledError.Fixes.Add(Fix);
+			FStackIssue InconsistentEnabledError(
+				EStackIssueSeverity::Error,
+				LOCTEXT("InconsistentEnabledErrorSummary", "The enabled state for module is inconsistent."),
+				LOCTEXT("InconsistentEnabledError", "This module is using multiple functions and their enabled state is inconsistent.\nClick fix to make all of the functions for this module enabled."),
+				GetStackEditorDataKey(),
+				false);
+
+			FText FixDescription = LOCTEXT("EnableModule", "Enable module");
+			FStackIssueFix EnableFix(
+				FixDescription,
+				FStackIssueFixDelegate::CreateLambda([=]()
+				{
+					FScopedTransaction ScopedTransaction(FixDescription);
+					SetIsEnabled(true);
+				}));
+
+			InconsistentEnabledError.AddFix(EnableFix);
 			NewIssues.Add(InconsistentEnabledError);
 		}
 	}
diff --git a/Engine/Plugins/FX/Niagara/Source/NiagaraEditor/Private/ViewModels/Stack/NiagaraStackRendererItem.cpp b/Engine/Plugins/FX/Niagara/Source/NiagaraEditor/Private/ViewModels/Stack/NiagaraStackRendererItem.cpp
--- a/Engine/Plugins/FX/Niagara/Source/NiagaraEditor/Private/ViewModels/Stack/NiagaraStackRendererItem.cpp
+++ b/Engine/Plugins/FX/Niagara/Source/NiagaraEditor/Private/ViewModels/Stack/NiagaraStackRendererItem.cpp
@@ -234,33 +234,42 @@ void UNiagaraStackRendererItem::RefreshIssues(TArray<FStackIssue>& NewIssues)
 {
 	for (FNiagaraVariable Attribute : MissingAttributes)
 	{
-		UNiagaraStackEntry::FStackIssue MissingAttributeError;
-		MissingAttributeError.ShortDescription = LOCTEXT("FailedRendererBindShort", "An attribute is missing.");
-		MissingAttributeError.LongDescription = FText::Format(LOCTEXT("FailedRendererBind", "Missing attribute \"{0}\" of Type \"{1}\"."), FText::FromName(Attribute.GetName()), Attribute.GetType().GetNameText());
-		MissingAttributeError.UniqueIdentifier = FName(*FString::Printf(TEXT("%s-AttributeMissing-%s"), *GetStackEditorDataKey(), *Attribute.GetName().ToString()));
-		UNiagaraStackEntry::FStackIssueFix Fix;
-		Fix.Description = LOCTEXT("AddMissingVariable", "Add missing variable");
-		Fix.FixDelegate.BindLambda([=]()
-		{
-			FScopedTransaction ScopedTransaction(Fix.Description);
-			if (AddMissingVariable(GetEmitterViewModel()->GetEmitter(), Attribute))
+		FStackIssue MissingAttributeError(
+			EStackIssueSeverity::Error,
+			LOCTEXT("FailedRendererBindShort", "An attribute is missing."),
+			FText::Format(LOCTEXT("FailedRendererBind", "Missing attribute \"{0}\" of Type \"{1}\"."), FText::FromName(Attribute.GetName()), Attribute.GetType().GetNameText()),
+			GetStackEditorDataKey(),
+			false);
+
+		FText FixDescription = LOCTEXT("AddMissingVariable", "Add missing variable");
+		FStackIssueFix AddAttributeFix(
+			FixDescription,
+			FStackIssueFixDelegate::CreateLambda([=]()
 			{
-				FNotificationInfo Info(FText::Format(LOCTEXT("AddedVariableForFix", "Added {0} to the Spawn script to support the renderer."), FText::FromName(Attribute.GetName())));
-				Info.ExpireDuration = 5.0f;
-				Info.bFireAndForget = true;
-				Info.Image = FCoreStyle::Get().GetBrush(TEXT("MessageLog.Info"));
-				FSlateNotificationManager::Get().AddNotification(Info);
-			}
-		});
-		MissingAttributeError.Fixes.Add(Fix);
+				FScopedTransaction ScopedTransaction(FixDescription);
+				if (AddMissingVariable(GetEmitterViewModel()->GetEmitter(), Attribute))
+				{
+					FNotificationInfo Info(FText::Format(LOCTEXT("AddedVariableForFix", "Added {0} to the Spawn script to support the renderer."), FText::FromName(Attribute.GetName())));
+					Info.ExpireDuration = 5.0f;
+					Info.bFireAndForget = true;
+					Info.Image = FCoreStyle::Get().GetBrush(TEXT("MessageLog.Info"));
+					FSlateNotificationManager::Get().AddNotification(Info);
+				}
+			}));
+
+		MissingAttributeError.AddFix(AddAttributeFix);
 		NewIssues.Add(MissingAttributeError);
 	}
+
 	if (RendererProperties->GetIsEnabled() && !RendererProperties->IsSimTargetSupported(GetEmitterViewModel()->GetEmitter()->SimTarget))
 	{
-		UNiagaraStackEntry::FStackIssue TargetSupportError;
-		TargetSupportError.ShortDescription = LOCTEXT("FailedRendererDueToSimTarget", "Renderer incompatible with SimTarget mode.");
-		TargetSupportError.LongDescription = FText::Format(LOCTEXT("FailedRendererDueToSimTargetLong", "Renderer incompatible with SimTarget mode \"{0}\"."), (int32)GetEmitterViewModel()->GetEmitter()->SimTarget);
-		TargetSupportError.UniqueIdentifier = FName(*FString::Printf(TEXT("%s-SimTargetModeError-%d"), *GetStackEditorDataKey(), (int32)GetEmitterViewModel()->GetEmitter()->SimTarget));
+		FStackIssue TargetSupportError(
+			EStackIssueSeverity::Error,
+			LOCTEXT("FailedRendererDueToSimTarget", "Renderer incompatible with SimTarget mode."),
+			FText::Format(LOCTEXT("FailedRendererDueToSimTargetLong", "Renderer incompatible with SimTarget mode \"{0}\"."), (int32)GetEmitterViewModel()->GetEmitter()->SimTarget),
+			GetStackEditorDataKey(),
+			false);
+
 		NewIssues.Add(TargetSupportError);
 	}
 }
diff --git a/Engine/Plugins/FX/Niagara/Source/NiagaraEditor/Private/ViewModels/Stack/NiagaraStackScriptItemGroup.cpp b/Engine/Plugins/FX/Niagara/Source/NiagaraEditor/Private/ViewModels/Stack/NiagaraStackScriptItemGroup.cpp
--- a/Engine/Plugins/FX/Niagara/Source/NiagaraEditor/Private/ViewModels/Stack/NiagaraStackScriptItemGroup.cpp
+++ b/Engine/Plugins/FX/Niagara/Source/NiagaraEditor/Private/ViewModels/Stack/NiagaraStackScriptItemGroup.cpp
@@ -330,21 +330,25 @@ void UNiagaraStackScriptItemGroup::RefreshIssues(TArray<FStackIssue>& NewIssues)
 	
 	if (FNiagaraStackGraphUtilities::ValidateGraphForOutput(*Graph, ScriptUsage, ScriptUsageId, ErrorMessage) == false)
 	{
-		UE_LOG(LogNiagaraEditor, Error, TEXT("Failed to Create Stack.  Message: %s"), *ErrorMessage.ToString());
-		UNiagaraStackEntry::FStackIssue Error;
-		Error.LongDescription = LOCTEXT("InvalidErrorText", "The data used to generate the stack has been corrupted and can not be used.\nUsing the fix option will reset this part of the stack to its default empty state.");
-		Error.ShortDescription = LOCTEXT("InvalidErrorSummaryText", "The stack data is invalid");
-		Error.UniqueIdentifier = FName(*FString::Printf(TEXT("StackDataInvalid-%s"), *GetStackEditorDataKey()));
-		UNiagaraStackEntry::FStackIssueFix Fix;
-		Fix.Description = LOCTEXT("FixStackGraph", "Fix invalid stack graph");
-		Fix.FixDelegate.BindLambda([=]()
-		{
-			FScopedTransaction ScopedTransaction(Fix.Description);
-			FNiagaraStackGraphUtilities::ResetGraphForOutput(*Graph, ScriptUsage, ScriptUsageId);
-			FNiagaraStackGraphUtilities::RelayoutGraph(*Graph);
-		});
-		Error.Fixes.Add(Fix);
-		NewIssues.Add(Error); 
+		FStackIssue InvalidStackError(
+			EStackIssueSeverity::Error,
+			LOCTEXT("InvalidErrorSummaryText", "The stack data is invalid"),
+			LOCTEXT("InvalidErrorText", "The data used to generate the stack has been corrupted and can not be used.\nUsing the fix option will reset this part of the stack to its default empty state."),
+			GetStackEditorDataKey(),
+			false);
+
+		FText FixDescription = LOCTEXT("FixStackGraph", "Fix invalid stack graph");
+		FStackIssueFix ResetStackFix(
+			FixDescription,
+			FStackIssueFixDelegate::CreateLambda([=]()
+			{
+				FScopedTransaction ScopedTransaction(FixDescription);
+				FNiagaraStackGraphUtilities::ResetGraphForOutput(*Graph, ScriptUsage, ScriptUsageId);
+				FNiagaraStackGraphUtilities::RelayoutGraph(*Graph);
+			}));
+
+		InvalidStackError.AddFix(ResetStackFix);
+		NewIssues.Add(InvalidStackError); 
 	}
 	else
 	{
@@ -364,27 +368,33 @@ void UNiagaraStackScriptItemGroup::RefreshIssues(TArray<FStackIssue>& NewIssues)
 			if (!FNiagaraStackGraphUtilities::FindScriptModulesInStack(ModuleScriptAsset, *MatchingOutputNode, FoundCalls))
 			{
 				bForcedError = true;
-				UNiagaraStackEntry::FStackIssue Error;
-				Error.LongDescription = LOCTEXT("SystemLifeCycleWarning", "The stack needs a SystemLifeCycle module.");;
-				Error.ShortDescription = LOCTEXT("MissingRequiredMode", "Missing required module.");
-				Error.UniqueIdentifier = FName(*FString::Printf(TEXT("MissingLifecycleModule-%s"), *GetStackEditorDataKey()));
-				UNiagaraStackEntry::FStackIssueFix Fix;
-				Fix.Description = LOCTEXT("AddingSystemLifecycleModule", "Adding System Lifecycle Module.");
-				Fix.FixDelegate.BindLambda([=]()
-				{
-					FScopedTransaction ScopedTransaction(Fix.Description);
-					UNiagaraNodeFunctionCall* AddedModuleNode = FNiagaraStackGraphUtilities::AddScriptModuleToStack(ModuleScriptAsset, *MatchingOutputNode);
-					if (AddedModuleNode == nullptr)
+
+				FStackIssue MissingLifeCycleError(
+					EStackIssueSeverity::Error,
+					LOCTEXT("SystemLifeCycleWarning", "The stack needs a SystemLifeCycle module."),
+					LOCTEXT("MissingRequiredMode", "Missing required module."),
+					GetStackEditorDataKey(),
+					false);
+
+				FText FixDescription = LOCTEXT("AddingSystemLifecycleModule", "Adding System Lifecycle Module.");
+				FStackIssueFix AddLifeCycleFix(
+					FixDescription,
+					FStackIssueFixDelegate::CreateLambda([=]()
 					{
-						FNotificationInfo Info(LOCTEXT("FailedToAddSystemLifecycle", "Failed to add system life cycle module.\nCheck the log for errors."));
-						Info.ExpireDuration = 5.0f;
-						Info.bFireAndForget = true;
-						Info.Image = FCoreStyle::Get().GetBrush(TEXT("MessageLog.Error"));
-						FSlateNotificationManager::Get().AddNotification(Info);
-					}
-				});
-				Error.Fixes.Add(Fix);
-				NewIssues.Add(Error);
+						FScopedTransaction ScopedTransaction(FixDescription);
+						UNiagaraNodeFunctionCall* AddedModuleNode = FNiagaraStackGraphUtilities::AddScriptModuleToStack(ModuleScriptAsset, *MatchingOutputNode);
+						if (AddedModuleNode == nullptr)
+						{
+							FNotificationInfo Info(LOCTEXT("FailedToAddSystemLifecycle", "Failed to add system life cycle module.\nCheck the log for errors."));
+							Info.ExpireDuration = 5.0f;
+							Info.bFireAndForget = true;
+							Info.Image = FCoreStyle::Get().GetBrush(TEXT("MessageLog.Error"));
+							FSlateNotificationManager::Get().AddNotification(Info);
+						}
+					}));
+
+				MissingLifeCycleError.AddFix(AddLifeCycleFix);
+				NewIssues.Add(MissingLifeCycleError);
 			}
 		}
 
@@ -393,11 +403,14 @@ void UNiagaraStackScriptItemGroup::RefreshIssues(TArray<FStackIssue>& NewIssues)
 		{
 			if (Status == ENiagaraScriptCompileStatus::NCS_Error)
 			{
-				UNiagaraStackEntry::FStackIssue Error;
-				Error.LongDescription = ScriptViewModelPinned->GetScriptErrors(GetScriptUsage(), GetScriptUsageId());
-				Error.ShortDescription = LOCTEXT("ConpileErrorSummary", "The stack has compile errors.");
-				Error.UniqueIdentifier = FName(*FString::Printf(TEXT("CompileErrors-%s"), *GetStackEditorDataKey()));
-				NewIssues.Add(Error);
+				FStackIssue CompileError(
+					EStackIssueSeverity::Error,
+					LOCTEXT("ConpileErrorSummary", "The stack has compile errors."),
+					ScriptViewModelPinned->GetScriptErrors(GetScriptUsage(), GetScriptUsageId()),
+					GetStackEditorDataKey(),
+					true);
+
+				NewIssues.Add(CompileError);
 			}
 		}
 	}
diff --git a/Engine/Plugins/FX/Niagara/Source/NiagaraEditor/Public/NiagaraStackEditorData.h b/Engine/Plugins/FX/Niagara/Source/NiagaraEditor/Public/NiagaraStackEditorData.h
--- a/Engine/Plugins/FX/Niagara/Source/NiagaraEditor/Public/NiagaraStackEditorData.h
+++ b/Engine/Plugins/FX/Niagara/Source/NiagaraEditor/Public/NiagaraStackEditorData.h
@@ -74,13 +74,13 @@ public:
 	/*
 	* @param Issue the issue to be dismissed (not fixed).
 	*/
-	void DismissStackIssue(FName IssueId);
+	void DismissStackIssue(FString IssueId);
 
 	/* Restores all the dismissed issues so that the user can see them and choose what to do. */
 	NIAGARAEDITOR_API void UndismissAllIssues();
 
 	/* Gets a reference to the dismissed stack issue array */
-	NIAGARAEDITOR_API const TArray<FName>& GetDismissedStackIssueIds();
+	NIAGARAEDITOR_API const TArray<FString>& GetDismissedStackIssueIds();
 
 private:
 	UPROPERTY()
@@ -102,5 +102,5 @@ private:
 	double LastScrollPosition;
 
 	UPROPERTY()
-	TArray<FName> DismissedStackIssueIds;
+	TArray<FString> DismissedStackIssueIds;
 };
\ No newline at end of file
diff --git a/Engine/Plugins/FX/Niagara/Source/NiagaraEditor/Public/ViewModels/Stack/NiagaraStackEntry.h b/Engine/Plugins/FX/Niagara/Source/NiagaraEditor/Public/ViewModels/Stack/NiagaraStackEntry.h
--- a/Engine/Plugins/FX/Niagara/Source/NiagaraEditor/Public/ViewModels/Stack/NiagaraStackEntry.h
+++ b/Engine/Plugins/FX/Niagara/Source/NiagaraEditor/Public/ViewModels/Stack/NiagaraStackEntry.h
@@ -103,28 +103,55 @@ public:
 	// stack issue stuff
 	struct FStackIssueFix
 	{
-	public:
+		FStackIssueFix();
+
+		FStackIssueFix(FText InDescription, FStackIssueFixDelegate InFixDelegate);
+
+		bool IsValid() const;
+
+		const FText& GetDescription() const;
+
+		const FStackIssueFixDelegate& GetFixDelegate() const;
+
 		bool operator == (const FStackIssueFix &Other) const
 		{
 			return Description.CompareTo(Other.Description) == 0 && FixDelegate.GetHandle() == Other.FixDelegate.GetHandle();
 		}
-	public:
+
+	private:
 		FText Description;
 		FStackIssueFixDelegate FixDelegate;
 	};
 
 	struct FStackIssue
 	{
-		FStackIssue()
-			: Severity(EStackIssueSeverity::Error)
-			, bCanBeDismissed(false)
-		{}
+		FStackIssue();
+
+		FStackIssue(EStackIssueSeverity InSeverity, FText InShortDescription, FText InLongDescription, FString InStackEditorDataKey, bool bInCanBeDismissed);
+
+		bool IsValid();
+
+		EStackIssueSeverity GetSeverity() const;
+
+		const FText& GetShortDescription() const;
+
+		const FText& GetLongDescription() const;
+
+		const FString& GetUniqueIdentifier() const;
+
+		bool GetCanBeDismissed() const;
+
+		void AddFix(FStackIssueFix InFix);
+
+		const TArray<FStackIssueFix>& GetFixes() const;
+
+	private:
+		EStackIssueSeverity Severity;
 		FText ShortDescription;
 		FText LongDescription;
-		FName UniqueIdentifier;
-		EStackIssueSeverity Severity;
-		TArray<FStackIssueFix> Fixes;
+		FString UniqueIdentifier;
 		bool bCanBeDismissed;
+		TArray<FStackIssueFix> Fixes;
 	};
 
 public:
diff --git a/Engine/Plugins/FX/Niagara/Source/NiagaraEditor/Public/ViewModels/Stack/NiagaraStackErrorItem.h b/Engine/Plugins/FX/Niagara/Source/NiagaraEditor/Public/ViewModels/Stack/NiagaraStackErrorItem.h
--- a/Engine/Plugins/FX/Niagara/Source/NiagaraEditor/Public/ViewModels/Stack/NiagaraStackErrorItem.h
+++ b/Engine/Plugins/FX/Niagara/Source/NiagaraEditor/Public/ViewModels/Stack/NiagaraStackErrorItem.h
@@ -56,7 +56,7 @@ public:
 
 	FStackIssueFix GetStackIssueFix() const { return IssueFix; };
 	FText FixDescription() const;
-	virtual FReply OnTryFixError();
+	FReply OnTryFixError();
 	virtual EStackRowStyle GetStackRowStyle() const override;
 	virtual FText GetFixButtonText() const;
 	UNiagaraStackErrorItem::FOnIssueNotify& OnIssueFixed();
@@ -75,7 +75,6 @@ class NIAGARAEDITOR_API UNiagaraStackErrorItemDismiss : public UNiagaraStackErro
 public:
 	void Initialize(FRequiredEntryData InRequiredEntryData, UNiagaraStackEntry::FStackIssue InStackIssue, FString InStackEditorDataKey);
 
-	virtual FReply OnTryFixError() override;
 	virtual EStackRowStyle GetStackRowStyle() const override;
 	virtual FText GetFixButtonText() const override;
 
