Niagara - Minor change missed in last checkin.

diff --git a/Engine/Plugins/FX/Niagara/Source/NiagaraEditor/Private/ViewModels/Stack/NiagaraStackScriptItemGroup.cpp b/Engine/Plugins/FX/Niagara/Source/NiagaraEditor/Private/ViewModels/Stack/NiagaraStackScriptItemGroup.cpp
--- a/Engine/Plugins/FX/Niagara/Source/NiagaraEditor/Private/ViewModels/Stack/NiagaraStackScriptItemGroup.cpp
+++ b/Engine/Plugins/FX/Niagara/Source/NiagaraEditor/Private/ViewModels/Stack/NiagaraStackScriptItemGroup.cpp
@@ -408,7 +408,7 @@ void UNiagaraStackScriptItemGroup::RefreshIssues(TArray<FStackIssue>& NewIssues)
 					LOCTEXT("ConpileErrorSummary", "The stack has compile errors."),
 					ScriptViewModelPinned->GetScriptErrors(GetScriptUsage(), GetScriptUsageId()),
 					GetStackEditorDataKey(),
-					true);
+					false);
 
 				NewIssues.Add(CompileError);
 			}
