Niagara - Fix the assignment node so that it uses a "Begin Defaults" node instead of a regular input node when it's hooking up linked defaults.

diff --git a/Engine/Plugins/FX/Niagara/Source/NiagaraCore/Public/NiagaraCustomVersion.h b/Engine/Plugins/FX/Niagara/Source/NiagaraCore/Public/NiagaraCustomVersion.h
--- a/Engine/Plugins/FX/Niagara/Source/NiagaraCore/Public/NiagaraCustomVersion.h
+++ b/Engine/Plugins/FX/Niagara/Source/NiagaraCore/Public/NiagaraCustomVersion.h
@@ -99,7 +99,9 @@ struct FNiagaraCustomVersion
 		DataInterfaceComputeShaderParamRefactor, // refactor of CS parameters allowing regular params as well as buffers.
 
 		CurveLUTRegen, // bumping version and forcing curves to regen their LUT on version change.
-				
+
+		AssignmentNodeUsesBeginDefaults, // Changing the graph generation for assignment nodes so that it uses a "Begin Defaults" node where appropriate.
+
 		// -----<new versions can be added above this line>-------------------------------------------------
 		VersionPlusOne,
 		LatestVersion = VersionPlusOne - 1,
diff --git a/Engine/Plugins/FX/Niagara/Source/NiagaraEditor/Private/NiagaraNodeAssignment.cpp b/Engine/Plugins/FX/Niagara/Source/NiagaraEditor/Private/NiagaraNodeAssignment.cpp
--- a/Engine/Plugins/FX/Niagara/Source/NiagaraEditor/Private/NiagaraNodeAssignment.cpp
+++ b/Engine/Plugins/FX/Niagara/Source/NiagaraEditor/Private/NiagaraNodeAssignment.cpp
@@ -153,6 +153,15 @@ void UNiagaraNodeAssignment::PostLoad()
 			UE_LOG(LogNiagaraEditor, Log, TEXT("Found old Assignment Node, nothing was attached???? variable \"%s\" in \"%s\""), *AssignmentTarget_DEPRECATED.GetName().ToString(), *GetFullName());
 		}
 	}
+	else
+	{
+		const int32 NiagaraVer = GetLinkerCustomVersion(FNiagaraCustomVersion::GUID);
+		if (NiagaraVer < FNiagaraCustomVersion::AssignmentNodeUsesBeginDefaults)
+		{
+			FunctionScript = nullptr;
+			GenerateScript();
+		}
+	}
 }
 
 void UNiagaraNodeAssignment::BuildParameterMapHistory(FNiagaraParameterMapHistoryBuilder& OutHistory, bool bRecursive)
@@ -291,14 +300,45 @@ void UNiagaraNodeAssignment::InitializeScript(UNiagaraScript* NewScript)
 		
 		TArray<UNiagaraNodeInput*> InputNodes;
 		CreatedGraph->FindInputNodes(InputNodes);
-		if (InputNodes.Num() == 0)
+
+		UNiagaraNodeInput* InputMapInputNode;
+		UNiagaraNodeInput** InputMapInputNodePtr = InputNodes.FindByPredicate([](UNiagaraNodeInput* InputNode)
+		{ 
+			return InputNode->Usage == ENiagaraInputNodeUsage::Parameter && InputNode->Input.GetType() == FNiagaraTypeDefinition::GetParameterMapDef() && InputNode->Input.GetName() == "InputMap"; 
+		});
+		if (InputMapInputNodePtr == nullptr)
 		{
 			FGraphNodeCreator<UNiagaraNodeInput> InputNodeCreator(*CreatedGraph);
-			UNiagaraNodeInput* InputNode = InputNodeCreator.CreateNode();
-			InputNode->Input = FNiagaraVariable(FNiagaraTypeDefinition::GetParameterMapDef(), TEXT("InputMap"));
-			InputNode->Usage = ENiagaraInputNodeUsage::Parameter;
+			InputMapInputNode = InputNodeCreator.CreateNode();
+			InputMapInputNode->Input = FNiagaraVariable(FNiagaraTypeDefinition::GetParameterMapDef(), TEXT("InputMap"));
+			InputMapInputNode->Usage = ENiagaraInputNodeUsage::Parameter;
 			InputNodeCreator.Finalize();
-			InputNodes.Add(InputNode);
+		}
+		else
+		{
+			InputMapInputNode = *InputMapInputNodePtr;
+		}
+
+		UNiagaraNodeInput* BeginDefaultsInputNode;
+		UNiagaraNodeInput** BeginDefaultsInputNodePtr = InputNodes.FindByPredicate([](UNiagaraNodeInput* InputNode)
+		{
+			return InputNode->Usage == ENiagaraInputNodeUsage::TranslatorConstant && InputNode->Input == TRANSLATOR_PARAM_BEGIN_DEFAULTS;
+		});
+		if (BeginDefaultsInputNodePtr == nullptr)
+		{
+			FGraphNodeCreator<UNiagaraNodeInput> InputNodeCreator(*CreatedGraph);
+			BeginDefaultsInputNode = InputNodeCreator.CreateNode();
+			BeginDefaultsInputNode->Input = TRANSLATOR_PARAM_BEGIN_DEFAULTS;
+			BeginDefaultsInputNode->Usage = ENiagaraInputNodeUsage::TranslatorConstant;
+			BeginDefaultsInputNode->ExposureOptions.bCanAutoBind = true;
+			BeginDefaultsInputNode->ExposureOptions.bHidden = true;
+			BeginDefaultsInputNode->ExposureOptions.bRequired = false;
+			BeginDefaultsInputNode->ExposureOptions.bExposed = false;
+			InputNodeCreator.Finalize();
+		}
+		else
+		{
+			BeginDefaultsInputNode = *BeginDefaultsInputNodePtr;
 		}
 
 		UNiagaraNodeOutput* OutputNode = CreatedGraph->FindOutputNode(ENiagaraScriptUsage::Module);
@@ -325,7 +365,7 @@ void UNiagaraNodeAssignment::InitializeScript(UNiagaraScript* NewScript)
 			InputNodeCreator.Finalize();
 			SetNodes.Add(InputNode);
 
-			InputNodes[0]->GetOutputPin(0)->MakeLinkTo(SetNodes[0]->GetInputPin(0));
+			InputMapInputNode->GetOutputPin(0)->MakeLinkTo(SetNodes[0]->GetInputPin(0));
 			SetNodes[0]->GetOutputPin(0)->MakeLinkTo(OutputNode->GetInputPin(0));
 		}
 
@@ -338,7 +378,7 @@ void UNiagaraNodeAssignment::InitializeScript(UNiagaraScript* NewScript)
 			InputNodeCreator.Finalize();
 			GetNodes.Add(InputNode);
 
-			InputNodes[0]->GetOutputPin(0)->MakeLinkTo(GetNodes[0]->GetInputPin(0));
+			InputMapInputNode->GetOutputPin(0)->MakeLinkTo(GetNodes[0]->GetInputPin(0));
 		}
 		if (GetNodes.Num() == 1)
 		{
@@ -347,7 +387,7 @@ void UNiagaraNodeAssignment::InitializeScript(UNiagaraScript* NewScript)
 			InputNodeCreator.Finalize();
 			GetNodes.Add(InputNode);
 
-			InputNodes[0]->GetOutputPin(0)->MakeLinkTo(GetNodes[1]->GetInputPin(0));
+			BeginDefaultsInputNode->GetOutputPin(0)->MakeLinkTo(GetNodes[1]->GetInputPin(0));
 		}
 
 		// Clean out existing pins
