Niagara - Fix an issue where the stack function input collection wasn't generating errors correctly due to data being cached between refreshes which became stale.

diff --git a/Engine/Plugins/FX/Niagara/Source/NiagaraEditor/Private/ViewModels/Stack/NiagaraStackFunctionInputCollection.cpp b/Engine/Plugins/FX/Niagara/Source/NiagaraEditor/Private/ViewModels/Stack/NiagaraStackFunctionInputCollection.cpp
--- a/Engine/Plugins/FX/Niagara/Source/NiagaraEditor/Private/ViewModels/Stack/NiagaraStackFunctionInputCollection.cpp
+++ b/Engine/Plugins/FX/Niagara/Source/NiagaraEditor/Private/ViewModels/Stack/NiagaraStackFunctionInputCollection.cpp
@@ -77,6 +77,9 @@ void UNiagaraStackFunctionInputCollection::RefreshChildrenInternal(const TArray<
 	const UEdGraphSchema_Niagara* NiagaraSchema = GetDefault<UEdGraphSchema_Niagara>();
 
 	TArray<FName> ProcessedInputNames;
+	TArray<FName> DuplicateInputNames;
+	TArray<FName> ValidAliasedInputNames;
+	TArray<const UEdGraphPin*> PinsWithInvalidTypes;
 
 	FText UncategorizedName = LOCTEXT("Uncategorized", "Uncategorized");
 
@@ -172,10 +175,10 @@ void UNiagaraStackFunctionInputCollection::RefreshChildrenInternal(const TArray<
 		}
 		InputCategory->AddInput(InputData.Pin->PinName, InputData.Type);
 	}
-	RefreshIssues(NewIssues);
+	RefreshIssues(DuplicateInputNames, ValidAliasedInputNames, PinsWithInvalidTypes, NewIssues);
 }
 
-void UNiagaraStackFunctionInputCollection::RefreshIssues(TArray<FStackIssue>& NewIssues)
+void UNiagaraStackFunctionInputCollection::RefreshIssues(TArray<FName> DuplicateInputNames, TArray<FName> ValidAliasedInputNames, TArray<const UEdGraphPin*> PinsWithInvalidTypes, TArray<FStackIssue>& NewIssues)
 {
 	// Try to find function input overrides which are no longer valid so we can generate errors for them.
 	UNiagaraNodeParameterMapSet* OverrideNode = FNiagaraStackGraphUtilities::GetStackFunctionOverrideNode(*InputFunctionCallNode);
@@ -224,6 +227,7 @@ void UNiagaraStackFunctionInputCollection::RefreshIssues(TArray<FStackIssue>& Ne
 		}
 	}
 
+	// Generate issues for duplicate input names.
 	for (const FName& DuplicateInputName : DuplicateInputNames)
 	{
 		FStackIssue DuplicateInputError(
@@ -236,6 +240,7 @@ void UNiagaraStackFunctionInputCollection::RefreshIssues(TArray<FStackIssue>& Ne
 		NewIssues.Add(DuplicateInputError);
 	}
 
+	// Generate issues for invalid types.
 	for (const UEdGraphPin* PinWithInvalidType : PinsWithInvalidTypes)
 	{
 		FStackIssue InputWithInvalidTypeError(
diff --git a/Engine/Plugins/FX/Niagara/Source/NiagaraEditor/Public/ViewModels/Stack/NiagaraStackFunctionInputCollection.h b/Engine/Plugins/FX/Niagara/Source/NiagaraEditor/Public/ViewModels/Stack/NiagaraStackFunctionInputCollection.h
--- a/Engine/Plugins/FX/Niagara/Source/NiagaraEditor/Public/ViewModels/Stack/NiagaraStackFunctionInputCollection.h
+++ b/Engine/Plugins/FX/Niagara/Source/NiagaraEditor/Public/ViewModels/Stack/NiagaraStackFunctionInputCollection.h
@@ -39,7 +39,7 @@ protected:
 	virtual void RefreshChildrenInternal(const TArray<UNiagaraStackEntry*>& CurrentChildren, TArray<UNiagaraStackEntry*>& NewChildren, TArray<FStackIssue>& NewIssues) override;
 	
 private:
-	void RefreshIssues(TArray<FStackIssue>& NewIssues);
+	void RefreshIssues(TArray<FName> DuplicateInputNames, TArray<FName> ValidAliasedInputNames, TArray<const UEdGraphPin*> PinsWithInvalidTypes, TArray<FStackIssue>& NewIssues);
 
 	void OnFunctionInputsChanged();
 
@@ -55,7 +55,4 @@ private:
 	UNiagaraNodeFunctionCall* ModuleNode;
 	UNiagaraNodeFunctionCall* InputFunctionCallNode;
 	bool bShouldShowInStack;
-	TArray<FName> ValidAliasedInputNames;
-	TArray<FName> DuplicateInputNames;
-	TArray<const UEdGraphPin*> PinsWithInvalidTypes;
 };
