Niagara - Fix the stack function input so that it cleans up properly when removing pins from assignment nodes, also fix undo for the remove operation.

diff --git a/Engine/Plugins/FX/Niagara/Source/NiagaraEditor/Private/ViewModels/Stack/NiagaraStackFunctionInput.cpp b/Engine/Plugins/FX/Niagara/Source/NiagaraEditor/Private/ViewModels/Stack/NiagaraStackFunctionInput.cpp
--- a/Engine/Plugins/FX/Niagara/Source/NiagaraEditor/Private/ViewModels/Stack/NiagaraStackFunctionInput.cpp
+++ b/Engine/Plugins/FX/Niagara/Source/NiagaraEditor/Private/ViewModels/Stack/NiagaraStackFunctionInput.cpp
@@ -1343,18 +1343,21 @@ void UNiagaraStackFunctionInput::DeleteInput()
 {
 	if (UNiagaraNodeAssignment* NodeAssignment = Cast<UNiagaraNodeAssignment>(OwningFunctionCallNode.Get()))
 	{
-		FNiagaraVariable Var = FNiagaraVariable(GetInputType(), GetInputParameterHandle().GetName());
-		NodeAssignment->RemoveParameter(Var);
+		FScopedTransaction ScopedTransaction(LOCTEXT("RemoveInputTransaction", "Remove Input"));
 
-		FNiagaraParameterHandle ParameterHandle(Var.GetName());
-		UEdGraphPin* OverridePin = FNiagaraStackGraphUtilities::GetStackFunctionInputOverridePin(*OwningFunctionCallNode.Get(), ParameterHandle);
-		if (OverridePin && OverridePin->GetOwningNode())
+		UEdGraphPin* OverridePin = GetOverridePin();
+		if (OverridePin != nullptr)
 		{
-			UEdGraphNode* Node = OverridePin->GetOwningNode();
-			FNiagaraStackGraphUtilities::RemoveNodesForStackFunctionInputOverridePin(*OverridePin);
-			OverridePin->GetOwningNode()->RemovePin(OverridePin);
-			OverridePin = nullptr;
+			// If there is an override pin and connected nodes, remove them before removing the input since removing
+			// the input will prevent us from finding the override pin.
+			RemoveNodesForOverridePin(*OverridePin);
+			UNiagaraNodeParameterMapSet* OverrideNode = GetOverrideNode();
+			OverrideNode->RemovePin(OverridePin);
 		}
+
+		FNiagaraVariable Var = FNiagaraVariable(GetInputType(), GetInputParameterHandle().GetName());
+		NodeAssignment->Modify();
+		NodeAssignment->RemoveParameter(Var);
 	}
 }
 
