Set persistend guid for if node input pins on creating a new output pin.

diff --git a/Engine/Plugins/FX/Niagara/Source/NiagaraEditor/Private/NiagaraNodeIf.cpp b/Engine/Plugins/FX/Niagara/Source/NiagaraEditor/Private/NiagaraNodeIf.cpp
--- a/Engine/Plugins/FX/Niagara/Source/NiagaraEditor/Private/NiagaraNodeIf.cpp
+++ b/Engine/Plugins/FX/Niagara/Source/NiagaraEditor/Private/NiagaraNodeIf.cpp
@@ -153,12 +153,17 @@ FGuid UNiagaraNodeIf::AddOutput(FNiagaraTypeDefinition Type, const FName& Name)
 	OutputVars.Add(NewOutput);
 	FGuid Guid = FGuid::NewGuid();
 	OutputVarGuids.Add(Guid);
-	InputAVarGuids.Add(FGuid::NewGuid());
-	InputBVarGuids.Add(FGuid::NewGuid());
 
 	const UEdGraphSchema_Niagara* Schema = CastChecked<UEdGraphSchema_Niagara>(GetSchema());
-	CreatePin(EGPD_Input, Schema->TypeDefinitionToPinType(Type), *(Name.ToString() + InputAPinSuffix), InputAVarGuids.Num());
-	CreatePin(EGPD_Input, Schema->TypeDefinitionToPinType(Type), *(Name.ToString() + InputBPinSuffix), InputAVarGuids.Num() + InputBVarGuids.Num());
+	FGuid PinAGuid = FGuid::NewGuid();
+	InputAVarGuids.Add(PinAGuid);
+	UEdGraphPin* PinA = CreatePin(EGPD_Input, Schema->TypeDefinitionToPinType(Type), *(Name.ToString() + InputAPinSuffix), InputAVarGuids.Num());
+	PinA->PersistentGuid = PinAGuid;
+
+	FGuid PinBGuid = FGuid::NewGuid();
+	InputBVarGuids.Add(PinBGuid);
+	UEdGraphPin* PinB = CreatePin(EGPD_Input, Schema->TypeDefinitionToPinType(Type), *(Name.ToString() + InputBPinSuffix), InputAVarGuids.Num() + InputBVarGuids.Num());
+	PinB->PersistentGuid = PinBGuid;
 
 	return Guid;
 }
