Fixing underlying GPU collision system after Rendering refactored to use the FSceneTexturesUniformParameters instead of individual textures. Note that GPU collision only works with the primary back buffer. We will need more work to support split-screen or PIP.

#tests Collsion test GPU now is functional, but we still get a few nondeterministic strays in different directions keeping me from turning it on at the moment

diff --git a/Engine/Plugins/FX/Niagara/Source/Niagara/Niagara.Build.cs b/Engine/Plugins/FX/Niagara/Source/Niagara/Niagara.Build.cs
--- a/Engine/Plugins/FX/Niagara/Source/Niagara/Niagara.Build.cs
+++ b/Engine/Plugins/FX/Niagara/Source/Niagara/Niagara.Build.cs
@@ -12,10 +12,11 @@ public class Niagara : ModuleRules
                 "NiagaraShader",
                 "Core",
                 "Engine",
-				"RenderCore",
+                "RenderCore",
                 "UtilityShaders",
                 "ShaderCore",
-                "TimeManagement"
+                "TimeManagement",
+                "Renderer",
             }
         );
 
diff --git a/Engine/Plugins/FX/Niagara/Source/Niagara/Private/NiagaraDataInterfaceCollisionQuery.cpp b/Engine/Plugins/FX/Niagara/Source/Niagara/Private/NiagaraDataInterfaceCollisionQuery.cpp
--- a/Engine/Plugins/FX/Niagara/Source/Niagara/Private/NiagaraDataInterfaceCollisionQuery.cpp
+++ b/Engine/Plugins/FX/Niagara/Source/Niagara/Private/NiagaraDataInterfaceCollisionQuery.cpp
@@ -126,6 +126,7 @@ void UNiagaraDataInterfaceCollisionQuery::GetFunctions(TArray<FNiagaraFunctionSi
 // build the shader function HLSL; function name is passed in, as it's defined per-DI; that way, configuration could change
 // the HLSL in the spirit of a static switch
 // TODO: need a way to identify each specific function here
+
 // 
 bool UNiagaraDataInterfaceCollisionQuery::GetFunctionHLSL(const FName& DefinitionFunctionName, FString InstanceFunctionName, FNiagaraDataInterfaceGPUParamInfo& ParamInfo, FString& OutHLSL)
 {
@@ -135,12 +136,6 @@ bool UNiagaraDataInterfaceCollisionQuery::GetFunctionHLSL(const FName& Definitio
 	{
 		OutHLSL += TEXT("float3 WorldPositionFromSceneDepth(float2 ScreenPosition, float SceneDepth)\n{\n\tfloat4 HomogeneousWorldPosition = mul(float4(ScreenPosition * SceneDepth, SceneDepth, 1), View.ScreenToWorld);\n\treturn HomogeneousWorldPosition.xyz / HomogeneousWorldPosition.w;\n}\n");
 		OutHLSL += TEXT("\n");
-		OutHLSL += TEXT("// For retrieving the world-space normal.\n\
-							//#if FORWARD_SHADING\n\
-							Texture2D GBufferATexture;\n\
-							SamplerState GBufferATextureSampler;\n\
-							//#endif\n\n\
-							");
 		OutHLSL += TEXT("void ") + InstanceFunctionName + TEXT("(in int InQueryID, in float3 In_ParticlePos, in float3 In_ParticleVel, in float In_DeltaSeconds, float CollisionRadius, in float CollisionDepthBounds, int In_InstanceData, \
 			out int Out_QueryID, out bool OutCollisionValid, out float3 Out_CollisionPos, out float3 Out_CollisionNormal, out float Out_Friction, out float Out_Restitution) \n{\n");
 		// get the screen position
@@ -166,34 +161,10 @@ bool UNiagaraDataInterfaceCollisionQuery::GetFunctionHLSL(const FName& Definitio
 			float SceneDepth = CalcSceneDepth(ScreenUV);\n\
 			if (abs(ClipPosition.w - SceneDepth) < CollisionDepthBounds)\n\
 			{\n\
-			/*#if FORWARD_SHADING_ACCURATE_NORMAL && FORWARD_SHADING\n\
-				float SceneDepth0 = CalcSceneDepth(ScreenUV + float2(View.BufferSizeAndInvSize.z, 0.0));\n\
-				float SceneDepth1 = CalcSceneDepth(ScreenUV + float2(0.0, View.BufferSizeAndInvSize.w));\n\
-				// When using the forward shading, the normal of the pixel is approximated by the derivative of the world position\n\
-				// of the pixel. But in on the visible edge this derivative can become very high, making CollisionPlane almost\n\
-				// perpendicular to the view plane. In these case the particle may collide the visible edges of the diferent meshes\n\
-				// in the view frustum. To avoid this, we disable the collision test if one of the derivate is above a threshold.\n\
-				if (max(abs(SceneDepth - SceneDepth0), abs(SceneDepth - SceneDepth1)) > CollisionDepthBounds)\n\
-				{\n\
-					OutCollisionValid = false;\n\
-					return;\n\
-				}\n\
-			#endif*/\n\
 				// Reconstruct world position.\n\
 				float3 WorldPosition = WorldPositionFromSceneDepth(ScreenPosition.xy, SceneDepth);\n\
-			/*#if FORWARD_SHADING && FORWARD_SHADING_ACCURATE_NORMAL\n\
-				// Compute normal using the two additional neighbooring pixel depth buffer fetches.\n\
-				float3 WorldPosition0 = WorldPositionFromSceneDepth(ScreenPosition.xy + float2(2 * View.ViewSizeAndInvSize.z, 0.0), SceneDepth0);\n\
-				float3 WorldPosition1 = WorldPositionFromSceneDepth(ScreenPosition.xy - float2(0.0, 2 * View.ViewSizeAndInvSize.w), SceneDepth1);\n\
-				float3 WorldNormal = normalize(cross(WorldPosition0 - WorldPosition, WorldPosition1 - WorldPosition));\n\
-			#elif FORWARD_SHADING // && !FORWARD_SHADING_ACCURATE_NORMAL\n\
-				// Compute normal using a ddx/ddy hack, hopefully the neighbooring particules will be close enough...\n\
-				float3 WorldNormal = normalize(cross(ddx(WorldPosition), ddy(WorldPosition)));\n\
-				WorldNormal *= sign(dot(View.WorldCameraOrigin - NewPosition, WorldNormal));\n\
-			#else*/ //!FORWARD_SHADING\n\
 				// Sample the normal buffer to create a plane to collide against.\n\
-				float3 WorldNormal = Texture2DSampleLevel(GBufferATexture, GBufferATextureSampler, ScreenUV, 0).xyz * 2.0 - 1.0;\n\
-			//#endif\n\
+				float3 WorldNormal = Texture2DSampleLevel(SceneTexturesStruct.GBufferATexture, SceneTexturesStruct.GBufferATextureSampler, ScreenUV, 0).xyz * 2.0 - 1.0;\n\
 				float4 CollisionPlane = float4(WorldNormal, dot(WorldPosition.xyz, WorldNormal));\n\
 				// Compute the portion of velocity normal to the collision plane.\n\
 				float VelocityDot = dot(CollisionPlane.xyz, DeltaPosition.xyz);\n\
@@ -228,7 +199,6 @@ bool UNiagaraDataInterfaceCollisionQuery::GetFunctionHLSL(const FName& Definitio
 void UNiagaraDataInterfaceCollisionQuery::GetParameterDefinitionHLSL(FNiagaraDataInterfaceGPUParamInfo& ParamInfo, FString& OutHLSL)
 {
 	// we don't need to add these to hlsl, as they're already in common.ush
-	//OutHLSL = "Texture2D GBufferATexture;\nSamplerState GBufferATextureSampler;";	
 }
 
 DEFINE_NDI_RAW_FUNC_BINDER(UNiagaraDataInterfaceCollisionQuery, SubmitQuery);
@@ -473,18 +443,13 @@ struct FNiagaraDataInterfaceParametersCS_CollisionQuery : public FNiagaraDataInt
 {
 	virtual void Bind(const FNiagaraDataInterfaceParamRef& ParamRef, const class FShaderParameterMap& ParameterMap) override
 	{
-		SceneDepthTextureParameter.Bind(ParameterMap, TEXT("SceneDepthTexture"));
-		SceneDepthTextureParameterSampler.Bind(ParameterMap, TEXT("SceneDepthTextureSampler"));
-		GBufferATextureParameter.Bind(ParameterMap, TEXT("GBufferATexture"));
-		GBufferATextureParameterSampler.Bind(ParameterMap, TEXT("GBufferATextureSampler"));
+		PassUniformBuffer.Bind(ParameterMap, FSceneTexturesUniformParameters::StaticStruct.GetShaderVariableName());
+		check(PassUniformBuffer.IsBound());
 	}
 
 	virtual void Serialize(FArchive& Ar)override
 	{
-		Ar << SceneDepthTextureParameter;
-		Ar << SceneDepthTextureParameterSampler;
-		Ar << GBufferATextureParameter;
-		Ar << GBufferATextureParameterSampler;
+		Ar << PassUniformBuffer;
 	}
 
 	virtual void Set(FRHICommandList& RHICmdList, FNiagaraShader* Shader, class UNiagaraDataInterface* DataInterface) const override
@@ -492,40 +457,15 @@ struct FNiagaraDataInterfaceParametersCS_CollisionQuery : public FNiagaraDataInt
 		check(IsInRenderingThread());
 
 		const FComputeShaderRHIParamRef ComputeShaderRHI = Shader->GetComputeShader();
-		//UNiagaraDataInterfaceCollisionQuery* CollisionDI = CastChecked<UNiagaraDataInterfaceCollisionQuery>(DataInterface);
-
-		FTexture2DRHIParamRef SceneDepthTexture = GNiagaraViewDataManager.GetSceneDepthTexture();
-		FTexture2DRHIParamRef GBufferATexture = GNiagaraViewDataManager.GetSceneNormalTexture() ? GNiagaraViewDataManager.GetSceneNormalTexture() : nullptr;
-		RHICmdList.TransitionResource(EResourceTransitionAccess::EReadable, SceneDepthTexture);
-		RHICmdList.TransitionResource(EResourceTransitionAccess::EReadable, GBufferATexture);
-
-		SetTextureParameter(
-			RHICmdList,
-			ComputeShaderRHI,
-			SceneDepthTextureParameter,
-			SceneDepthTextureParameterSampler,
-			TStaticSamplerState<SF_Bilinear, AM_Clamp, AM_Clamp, AM_Clamp>::GetRHI(),
-			SceneDepthTexture//TODO: Does this work for multiple views?
-		);
-
-		SetTextureParameter(
-			RHICmdList,
-			ComputeShaderRHI,
-			GBufferATextureParameter,
-			GBufferATextureParameterSampler,
-			TStaticSamplerState<SF_Bilinear, AM_Clamp, AM_Clamp, AM_Clamp>::GetRHI(),
-			GBufferATexture//TODO: Does this work for multiple views?
-		);
+
+		TUniformBufferRef<FSceneTexturesUniformParameters> SceneTextureUniformParams = GNiagaraViewDataManager.GetSceneTextureUniformParameters();
+		SetUniformBufferParameter(RHICmdList, ComputeShaderRHI, PassUniformBuffer/*Shader->GetUniformBufferParameter(SceneTexturesUniformBufferStruct)*/, SceneTextureUniformParams);
 	}
 
 private:
 
 	/** The SceneDepthTexture parameter for depth buffer collision. */
-	FShaderResourceParameter SceneDepthTextureParameter;
-	FShaderResourceParameter SceneDepthTextureParameterSampler;
-	/** The GBufferATexture parameter for depth buffer collision. */
-	FShaderResourceParameter GBufferATextureParameter;
-	FShaderResourceParameter GBufferATextureParameterSampler;
+	FShaderUniformBufferParameter PassUniformBuffer;
 };
 
 FNiagaraDataInterfaceParametersCS* UNiagaraDataInterfaceCollisionQuery::ConstructComputeParameters()const
diff --git a/Engine/Plugins/FX/Niagara/Source/Niagara/Private/NiagaraDataSet.cpp b/Engine/Plugins/FX/Niagara/Source/Niagara/Private/NiagaraDataSet.cpp
--- a/Engine/Plugins/FX/Niagara/Source/Niagara/Private/NiagaraDataSet.cpp
+++ b/Engine/Plugins/FX/Niagara/Source/Niagara/Private/NiagaraDataSet.cpp
@@ -34,11 +34,11 @@ void FNiagaraDataSet::SetShaderParams(FNiagaraShader *Shader, FRHICommandList &C
 		CommandList.TransitionResource(EResourceTransitionAccess::EReadable, EResourceTransitionPipeline::EComputeToCompute, PrevData().GetGPUBufferInt()->UAV);
 		if (PrevData().GetNumInstancesAllocated() > 0)
 		{
-			CommandList.SetShaderResourceViewParameter(Shader->GetComputeShader(), Shader->FloatInputBufferParam.GetBaseIndex(), PrevData().GetGPUBufferInt()->SRV);
+			CommandList.SetShaderResourceViewParameter(Shader->GetComputeShader(), Shader->IntInputBufferParam.GetBaseIndex(), PrevData().GetGPUBufferInt()->SRV);
 		}
 		else
 		{
-			CommandList.SetShaderResourceViewParameter(Shader->GetComputeShader(), Shader->FloatInputBufferParam.GetBaseIndex(), NiagaraRenderer::GetDummyIntBuffer().SRV);
+			CommandList.SetShaderResourceViewParameter(Shader->GetComputeShader(), Shader->IntInputBufferParam.GetBaseIndex(), NiagaraRenderer::GetDummyIntBuffer().SRV);
 		}
 	}
 	if (Shader->FloatOutputBufferParam.IsUAVBound())
diff --git a/Engine/Plugins/FX/Niagara/Source/Niagara/Public/NiagaraWorldManager.h b/Engine/Plugins/FX/Niagara/Source/Niagara/Public/NiagaraWorldManager.h
--- a/Engine/Plugins/FX/Niagara/Source/Niagara/Public/NiagaraWorldManager.h
+++ b/Engine/Plugins/FX/Niagara/Source/Niagara/Public/NiagaraWorldManager.h
@@ -31,17 +31,19 @@ public:
 		SceneDepthTexture = Params.DepthTexture;
 		ViewUniformBuffer = Params.ViewUniformBuffer;
 		SceneNormalTexture = Params.NormalTexture;
+		SceneTexturesUniformParams = Params.SceneTexturesUniformParams;
 	}
 
 	FTexture2DRHIParamRef GetSceneDepthTexture() { return SceneDepthTexture; }
 	FTexture2DRHIParamRef GetSceneNormalTexture() { return SceneNormalTexture; }
 	FUniformBufferRHIParamRef GetViewUniformBuffer() { return ViewUniformBuffer; }
-
+	TUniformBufferRef<FSceneTexturesUniformParameters> GetSceneTextureUniformParameters() { return SceneTexturesUniformParams; }
 private:
 	FTexture2DRHIParamRef SceneDepthTexture;
 	FTexture2DRHIParamRef SceneNormalTexture;
 	FUniformBufferRHIParamRef ViewUniformBuffer;
 
+	TUniformBufferRef<FSceneTexturesUniformParameters> SceneTexturesUniformParams;
 	FPostOpaqueRenderDelegate PostOpaqueDelegate;
 };
 
diff --git a/Engine/Source/Runtime/RenderCore/Public/RendererInterface.h b/Engine/Source/Runtime/RenderCore/Public/RendererInterface.h
--- a/Engine/Source/Runtime/RenderCore/Public/RendererInterface.h
+++ b/Engine/Source/Runtime/RenderCore/Public/RendererInterface.h
@@ -14,6 +14,7 @@
 #include "RenderResource.h"
 #include "RenderUtils.h"
 #include "Misc/EnumClassFlags.h"
+#include "UniformBuffer.h"
 
 class FCanvas;
 class FMaterial;
@@ -21,6 +22,7 @@ class FSceneInterface;
 class FSceneRenderTargets;
 class FSceneView;
 class FSceneViewFamily;
+class FSceneTexturesUniformParameters;
 struct FMeshBatch;
 struct FSynthBenchmarkResults;
 
@@ -546,6 +548,7 @@ class FPostOpaqueRenderParameters
 		FRHITexture2D* SmallDepthTexture;
 		FRHICommandListImmediate* RHICmdList;
 		FUniformBufferRHIParamRef ViewUniformBuffer;
+		TUniformBufferRef<FSceneTexturesUniformParameters> SceneTexturesUniformParams;
 		void* Uid; // A unique identifier for the view.
 };
 DECLARE_DELEGATE_OneParam(FPostOpaqueRenderDelegate, class FPostOpaqueRenderParameters&);
@@ -734,7 +737,7 @@ public:
 
 	virtual void RegisterPostOpaqueRenderDelegate(const FPostOpaqueRenderDelegate& PostOpaqueRenderDelegate) = 0;
 	virtual void RegisterOverlayRenderDelegate(const FPostOpaqueRenderDelegate& OverlayRenderDelegate) = 0;
-	virtual void RenderPostOpaqueExtensions(const class FViewInfo& View, FRHICommandListImmediate& RHICmdList, class FSceneRenderTargets& SceneContext) = 0;
+	virtual void RenderPostOpaqueExtensions(const class FViewInfo& View, FRHICommandListImmediate& RHICmdList, class FSceneRenderTargets& SceneContext, TUniformBufferRef<FSceneTexturesUniformParameters>& SceneTextureUniformParams) = 0;
 	virtual void RenderOverlayExtensions(const class FViewInfo& View, FRHICommandListImmediate& RHICmdList, FSceneRenderTargets& SceneContext) = 0;
 	virtual bool HasPostOpaqueExtentions() const = 0;
 
diff --git a/Engine/Source/Runtime/Renderer/Private/DeferredShadingRenderer.cpp b/Engine/Source/Runtime/Renderer/Private/DeferredShadingRenderer.cpp
--- a/Engine/Source/Runtime/Renderer/Private/DeferredShadingRenderer.cpp
+++ b/Engine/Source/Runtime/Renderer/Private/DeferredShadingRenderer.cpp
@@ -1080,12 +1080,16 @@ void FDeferredShadingSceneRenderer::Render(FRHICommandListImmediate& RHICmdList)
 	IRendererModule& RendererModule = GetRendererModule();
 	if (RendererModule.HasPostOpaqueExtentions())
 	{
+		FSceneTexturesUniformParameters SceneTextureParameters;
+		SetupSceneTextureUniformParameters(SceneContext, FeatureLevel, ESceneTextureSetupMode::SceneDepth | ESceneTextureSetupMode::GBuffers, SceneTextureParameters);
+		TUniformBufferRef<FSceneTexturesUniformParameters> SceneTextureUniformBuffer = TUniformBufferRef<FSceneTexturesUniformParameters>::CreateUniformBufferImmediate(SceneTextureParameters, UniformBuffer_SingleFrame);
+
 		SceneContext.BeginRenderingSceneColor(RHICmdList);
 		for (int32 ViewIndex = 0; ViewIndex < Views.Num(); ++ViewIndex)
 		{
 			const FViewInfo& View = Views[ViewIndex];
 			RHICmdList.SetViewport(View.ViewRect.Min.X, View.ViewRect.Min.Y, 0.0f, View.ViewRect.Max.X, View.ViewRect.Max.Y, 1.0f);
-			RendererModule.RenderPostOpaqueExtensions(View, RHICmdList, SceneContext);
+			RendererModule.RenderPostOpaqueExtensions(View, RHICmdList, SceneContext, SceneTextureUniformBuffer);
 		}
 	}
 	SetRenderTarget(RHICmdList, nullptr, 0, 0, nullptr);
diff --git a/Engine/Source/Runtime/Renderer/Private/RendererModule.h b/Engine/Source/Runtime/Renderer/Private/RendererModule.h
--- a/Engine/Source/Runtime/Renderer/Private/RendererModule.h
+++ b/Engine/Source/Runtime/Renderer/Private/RendererModule.h
@@ -82,7 +82,7 @@ public:
 
 	virtual void RegisterPostOpaqueRenderDelegate(const FPostOpaqueRenderDelegate& PostOpaqueRenderDelegate) override;
 	virtual void RegisterOverlayRenderDelegate(const FPostOpaqueRenderDelegate& OverlayRenderDelegate) override;
-	virtual void RenderPostOpaqueExtensions(const FViewInfo& View, FRHICommandListImmediate& RHICmdList, FSceneRenderTargets& SceneContext) override;
+	virtual void RenderPostOpaqueExtensions(const FViewInfo& View, FRHICommandListImmediate& RHICmdList, FSceneRenderTargets& SceneContext, TUniformBufferRef<FSceneTexturesUniformParameters>& SceneTextureUniformParams) override;
 	virtual void RenderOverlayExtensions(const FViewInfo& View, FRHICommandListImmediate& RHICmdList, FSceneRenderTargets& SceneContext) override;
 
 	virtual bool HasPostOpaqueExtentions() const override
diff --git a/Engine/Source/Runtime/Renderer/Private/SceneRendering.cpp b/Engine/Source/Runtime/Renderer/Private/SceneRendering.cpp
--- a/Engine/Source/Runtime/Renderer/Private/SceneRendering.cpp
+++ b/Engine/Source/Runtime/Renderer/Private/SceneRendering.cpp
@@ -3171,7 +3171,7 @@ void FRendererModule::RegisterOverlayRenderDelegate(const FPostOpaqueRenderDeleg
 	this->OverlayRenderDelegate = InOverlayRenderDelegate;
 }
 
-void FRendererModule::RenderPostOpaqueExtensions(const FViewInfo& View, FRHICommandListImmediate& RHICmdList, FSceneRenderTargets& SceneContext)
+void FRendererModule::RenderPostOpaqueExtensions(const FViewInfo& View, FRHICommandListImmediate& RHICmdList, FSceneRenderTargets& SceneContext, TUniformBufferRef<FSceneTexturesUniformParameters>& SceneTextureUniformParams)
 {
 	check(IsInRenderingThread());
 	FPostOpaqueRenderParameters RenderParameters;
@@ -3181,6 +3181,7 @@ void FRendererModule::RenderPostOpaqueExtensions(const FViewInfo& View, FRHIComm
 	RenderParameters.NormalTexture = SceneContext.GBufferA.IsValid() ? SceneContext.GetGBufferATexture() : nullptr;
 	RenderParameters.SmallDepthTexture = SceneContext.GetSmallDepthSurface()->GetTexture2D();
 	RenderParameters.ViewUniformBuffer = View.ViewUniformBuffer;
+	RenderParameters.SceneTexturesUniformParams = SceneTextureUniformParams;
 
 	RenderParameters.ViewportRect = View.ViewRect;
 	RenderParameters.RHICmdList = &RHICmdList;
diff --git a/Engine/Source/Runtime/ShaderCore/Public/ShaderParameterUtils.h b/Engine/Source/Runtime/ShaderCore/Public/ShaderParameterUtils.h
--- a/Engine/Source/Runtime/ShaderCore/Public/ShaderParameterUtils.h
+++ b/Engine/Source/Runtime/ShaderCore/Public/ShaderParameterUtils.h
@@ -188,7 +188,7 @@ inline void SetPixelShaderBool(
  * Sets the value of a shader texture parameter.  Template'd on shader type
  */
 template<typename ShaderTypeRHIParamRef, typename TRHICmdList>
-FORCEINLINE void SetTextureParameter(
+ void SetTextureParameter(
 	TRHICmdList& RHICmdList,
 	ShaderTypeRHIParamRef Shader,
 	const FShaderResourceParameter& TextureParameter,
