Allowing several descriptions to be multiline, accessible by Shift + Enter.

#tests created descriptions for both module fields and modules themselves that were multi-line. confirmed UI was correct.

diff --git a/Engine/Plugins/FX/Niagara/Source/Niagara/Classes/NiagaraScript.h b/Engine/Plugins/FX/Niagara/Source/Niagara/Classes/NiagaraScript.h
--- a/Engine/Plugins/FX/Niagara/Source/Niagara/Classes/NiagaraScript.h
+++ b/Engine/Plugins/FX/Niagara/Source/Niagara/Classes/NiagaraScript.h
@@ -63,7 +63,7 @@ public:
 	UPROPERTY(AssetRegistrySearchable, EditAnywhere, Category = Script)
 		ENiagaraModuleDependencyType Type; // e.g. PreDependency,
 										   /** Detailed description of the dependency */
-	UPROPERTY(AssetRegistrySearchable, EditAnywhere, Category = Script)
+	UPROPERTY(AssetRegistrySearchable, EditAnywhere, Category = Script, meta = (MultiLine = true))
 		FText Description;
 };
 
@@ -311,7 +311,7 @@ public:
 	ENiagaraNumericOutputTypeSelectionMode NumericOutputTypeSelectionMode;
 
 #if WITH_EDITORONLY_DATA
-	UPROPERTY(AssetRegistrySearchable, EditAnywhere, Category = Script)
+	UPROPERTY(AssetRegistrySearchable, EditAnywhere, Category = Script, meta = (MultiLine = true))
 	FText Description;
 
 	UPROPERTY(EditAnywhere, Category = Script)
diff --git a/Engine/Plugins/FX/Niagara/Source/Niagara/Public/NiagaraTypes.h b/Engine/Plugins/FX/Niagara/Source/Niagara/Public/NiagaraTypes.h
--- a/Engine/Plugins/FX/Niagara/Source/Niagara/Public/NiagaraTypes.h
+++ b/Engine/Plugins/FX/Niagara/Source/Niagara/Public/NiagaraTypes.h
@@ -300,7 +300,7 @@ public:
 	UPROPERTY(EditAnywhere, Category = "Variable")
 	TMap<FName, FString> PropertyMetaData;
 
-	UPROPERTY(EditAnywhere, Category = "Variable")
+	UPROPERTY(EditAnywhere, Category = "Variable", meta = (MultiLine = true))
 	FText Description;
 
 	UPROPERTY(EditAnywhere, Category = "Variable")
