Added tooltips to the parameter map get and set nodes that should grealy improve understanding.

#tests n/a

diff --git a/Engine/Plugins/FX/Niagara/Source/NiagaraEditor/Private/NiagaraNodeParameterMapBase.cpp b/Engine/Plugins/FX/Niagara/Source/NiagaraEditor/Private/NiagaraNodeParameterMapBase.cpp
--- a/Engine/Plugins/FX/Niagara/Source/NiagaraEditor/Private/NiagaraNodeParameterMapBase.cpp
+++ b/Engine/Plugins/FX/Niagara/Source/NiagaraEditor/Private/NiagaraNodeParameterMapBase.cpp
@@ -154,11 +154,45 @@ void UNiagaraNodeParameterMapBase::GetPinHoverText(const UEdGraphPin& Pin, FStri
 		const UEdGraphSchema_Niagara* Schema = Cast<UEdGraphSchema_Niagara>(NiagaraGraph->GetSchema());
 		if (Schema)
 		{
-			FNiagaraVariable Var = FNiagaraVariable(Schema->PinToTypeDefinition(&Pin), Pin.PinName);
+			if (IsAddPin(&Pin))
+			{
+				HoverTextOut = LOCTEXT("ParameterMapAddString", "Request a new variable from the parameter map.").ToString();
+				return;
+			}
+
+			FNiagaraTypeDefinition TypeDef = Schema->PinToTypeDefinition(&Pin);
+
+			if (Pin.Direction == EEdGraphPinDirection::EGPD_Input)
+			{
+				if (&Pin == GetInputPin(0) && TypeDef == FNiagaraTypeDefinition::GetParameterMapDef())
+				{
+					HoverTextOut = LOCTEXT("ParameterMapInString", "The source parameter map where we pull the values from.").ToString();
+					return;
+				}
+			}
+
+			if (Pin.Direction == EEdGraphPinDirection::EGPD_Output)
+			{
+				if (&Pin == GetOutputPin(0) && TypeDef == FNiagaraTypeDefinition::GetParameterMapDef())
+				{
+					HoverTextOut = LOCTEXT("ParameterMapOutString", "The destination parameter map where we write the values to.").ToString();
+					return;
+				}
+			}
+
+			FNiagaraVariable Var = FNiagaraVariable(TypeDef, Pin.PinName);
 			const FNiagaraVariableMetaData* Metadata = NiagaraGraph->GetMetaData(Var);
 			if (Metadata)
 			{
-				HoverTextOut = Metadata->Description.ToString();
+				FText Desc = FText::Format(LOCTEXT("GetVarTooltip", "Name: \"{0}\"\nType: {1}\nDesc: {2}"), FText::FromName(Pin.PinName),
+					TypeDef.GetNameText(), Metadata->Description);
+				HoverTextOut = Desc.ToString();
+			}
+			else
+			{
+				FText Desc = FText::Format(LOCTEXT("GetVarTooltip", "Name: \"{0}\"\nType: {1}\nDesc: None"), FText::FromName(Pin.PinName),
+					TypeDef.GetNameText());
+				HoverTextOut = Desc.ToString();
 			}
 		}
 	}
diff --git a/Engine/Plugins/FX/Niagara/Source/NiagaraEditor/Private/NiagaraNodeParameterMapGet.cpp b/Engine/Plugins/FX/Niagara/Source/NiagaraEditor/Private/NiagaraNodeParameterMapGet.cpp
--- a/Engine/Plugins/FX/Niagara/Source/NiagaraEditor/Private/NiagaraNodeParameterMapGet.cpp
+++ b/Engine/Plugins/FX/Niagara/Source/NiagaraEditor/Private/NiagaraNodeParameterMapGet.cpp
@@ -207,18 +207,26 @@ UEdGraphPin* UNiagaraNodeParameterMapGet::GetDefaultPin(UEdGraphPin* OutputPin)
 }
 
 
-UEdGraphPin* UNiagaraNodeParameterMapGet::GetOutputPinForDefault(UEdGraphPin* DefaultPin)
+UEdGraphPin* UNiagaraNodeParameterMapGet::GetOutputPinForDefault(const UEdGraphPin* DefaultPin) const
 {
 	TArray<UEdGraphPin*> OutputPins;
 	GetOutputPins(OutputPins);
 
-	const FGuid* OutputGuid = PinOutputToPinDefaultPersistentId.Find(DefaultPin->PersistentGuid);
+	FGuid OutputGuid;
+	for (auto It : PinOutputToPinDefaultPersistentId)
+	{
+		if (It.Value == DefaultPin->PersistentGuid)
+		{
+			OutputGuid = It.Key;
+			break;
+		}
+	}
 
-	if (OutputGuid != nullptr)
+	if (OutputGuid.IsValid())
 	{
 		for (UEdGraphPin* OutputPin : OutputPins)
 		{
-			if ((*OutputGuid) == OutputPin->PersistentGuid)
+			if (OutputGuid == OutputPin->PersistentGuid)
 			{
 				return OutputPin;
 			}
@@ -431,6 +439,66 @@ void UNiagaraNodeParameterMapGet::GatherExternalDependencyIDs(ENiagaraScriptUsag
 	}
 }
 
+void UNiagaraNodeParameterMapGet::GetPinHoverText(const UEdGraphPin& Pin, FString& HoverTextOut) const
+{
+	// Get hover text from metadata description.
+	const UNiagaraGraph* NiagaraGraph = GetNiagaraGraph();
+	if (NiagaraGraph)
+	{
+		const UEdGraphSchema_Niagara* Schema = Cast<UEdGraphSchema_Niagara>(NiagaraGraph->GetSchema());
+
+		if (Schema)
+		{
+			FNiagaraTypeDefinition TypeDef = Schema->PinToTypeDefinition(&Pin);
+			if (IsAddPin(&Pin))
+			{
+				HoverTextOut = LOCTEXT("ParameterMapAddString", "Request a new variable from the parameter map.").ToString();
+				return;
+			}
+
+			if (Pin.Direction == EEdGraphPinDirection::EGPD_Input)
+			{
+				if (&Pin == GetInputPin(0) && TypeDef == FNiagaraTypeDefinition::GetParameterMapDef())
+				{
+					HoverTextOut = LOCTEXT("ParameterMapInString", "The source parameter map where we pull the values from.").ToString();
+					return;
+				}
+				const UEdGraphPin* OutputPin = GetOutputPinForDefault(&Pin);
+				if (OutputPin)
+				{
+					TypeDef = Schema->PinToTypeDefinition(OutputPin);
+					FNiagaraVariable Var = FNiagaraVariable(TypeDef, OutputPin->PinName);
+					const FNiagaraVariableMetaData* Metadata = NiagaraGraph->GetMetaData(Var);
+					if (Metadata)
+					{
+						FText Desc = FText::Format(LOCTEXT("DefaultValueTooltip", "Default value for \"{0}\" if no other module has set it previously in the stack."), FText::FromName(OutputPin->PinName));
+						HoverTextOut = Desc.ToString();
+						return;
+					}
+				}
+				
+			}
+			else
+			{
+				FNiagaraVariable Var = FNiagaraVariable(TypeDef, Pin.PinName);
+				const FNiagaraVariableMetaData* Metadata = NiagaraGraph->GetMetaData(Var);
+				if (Metadata)
+				{
+					FText Desc = FText::Format(LOCTEXT("GetVarTooltip", "Name: \"{0}\"\nType: {1}\nDesc: {2}"), FText::FromName(Pin.PinName),
+						TypeDef.GetNameText(), Metadata->Description);
+					HoverTextOut = Desc.ToString();
+				}
+				else
+				{
+					FText Desc = FText::Format(LOCTEXT("GetVarTooltip", "Name: \"{0}\"\nType: {1}\nDesc: None"), FText::FromName(Pin.PinName),
+						TypeDef.GetNameText());
+					HoverTextOut = Desc.ToString();
+				}
+			}
+		}
+	}
+}
+
 void UNiagaraNodeParameterMapGet::GetContextMenuActions(const FGraphNodeContextMenuBuilder& Context) const
 {
 	UNiagaraNodeParameterMapBase::GetContextMenuActions(Context);
diff --git a/Engine/Plugins/FX/Niagara/Source/NiagaraEditor/Private/NiagaraNodeParameterMapGet.h b/Engine/Plugins/FX/Niagara/Source/NiagaraEditor/Private/NiagaraNodeParameterMapGet.h
--- a/Engine/Plugins/FX/Niagara/Source/NiagaraEditor/Private/NiagaraNodeParameterMapGet.h
+++ b/Engine/Plugins/FX/Niagara/Source/NiagaraEditor/Private/NiagaraNodeParameterMapGet.h
@@ -40,6 +40,7 @@ public:
 	virtual void PostLoad() override;
 
 	void GatherExternalDependencyIDs(ENiagaraScriptUsage InMasterUsage, const FGuid& InMasterUsageId, TArray<FGuid>& InReferencedIDs, TArray<UObject*>& InReferencedObjs) const override;
+	virtual void GetPinHoverText(const UEdGraphPin& Pin, FString& HoverTextOut) const override;
 
 protected:
 	virtual void OnNewTypedPinAdded(UEdGraphPin* NewPin) override;
@@ -52,7 +53,7 @@ protected:
 	void SynchronizeDefaultInputPin(UEdGraphPin* DefaultPin, UEdGraphPin* OutputPin);
 	
 	/** Reverse the lookup from GetDefaultPin.*/
-	UEdGraphPin* GetOutputPinForDefault(UEdGraphPin* DefaultPin);
+	UEdGraphPin* GetOutputPinForDefault(const UEdGraphPin* DefaultPin) const;
 	
 	/** Properly set up the default input pin for an output pin.*/
 	UEdGraphPin* CreateDefaultPin(UEdGraphPin* OutputPin);
