Pins and parameters unique name on creation

diff --git a/Engine/Plugins/FX/Niagara/Source/NiagaraEditor/Private/NiagaraNodeParameterMapBase.cpp b/Engine/Plugins/FX/Niagara/Source/NiagaraEditor/Private/NiagaraNodeParameterMapBase.cpp
--- a/Engine/Plugins/FX/Niagara/Source/NiagaraEditor/Private/NiagaraNodeParameterMapBase.cpp
+++ b/Engine/Plugins/FX/Niagara/Source/NiagaraEditor/Private/NiagaraNodeParameterMapBase.cpp
@@ -202,11 +202,31 @@ void UNiagaraNodeParameterMapBase::OnPinRenamed(UEdGraphPin* RenamedPin, const F
 {
 	RenamedPin->PinFriendlyName = FText::FromName(RenamedPin->PinName);
 
+	TArray<UEdGraphPin*> InOrOutPins;
+	if (RenamedPin->Direction == EEdGraphPinDirection::EGPD_Input)
+	{
+		GetInputPins(InOrOutPins);
+	}
+	else
+	{
+		GetOutputPins(InOrOutPins);
+	}
+
+	TSet<FName> Names;
+	for (const UEdGraphPin* Pin : InOrOutPins)
+	{
+		if (Pin != RenamedPin)
+		{
+			Names.Add(Pin->GetFName());
+		}
+	}
+	const FName NewUniqueName = FNiagaraUtilities::GetUniqueName(*RenamedPin->GetName(), Names);
+
 	FNiagaraTypeDefinition VarType = CastChecked<UEdGraphSchema_Niagara>(GetSchema())->PinToTypeDefinition(RenamedPin);
 	FNiagaraVariable Var(VarType, *OldName);
 
 	UNiagaraGraph* Graph = GetNiagaraGraph();
-	Graph->RenameParameter(Var, RenamedPin->GetFName(), false /*bInNotifyGraphChanged*/); // Notify graph changed in child.
+	Graph->RenameParameter(Var, NewUniqueName, false /*bInNotifyGraphChanged*/); // Notify graph changed in child.
 
 	if (RenamedPin == PinPendingRename)
 	{
diff --git a/Engine/Plugins/FX/Niagara/Source/NiagaraEditor/Private/NiagaraNodeParameterMapGet.cpp b/Engine/Plugins/FX/Niagara/Source/NiagaraEditor/Private/NiagaraNodeParameterMapGet.cpp
--- a/Engine/Plugins/FX/Niagara/Source/NiagaraEditor/Private/NiagaraNodeParameterMapGet.cpp
+++ b/Engine/Plugins/FX/Niagara/Source/NiagaraEditor/Private/NiagaraNodeParameterMapGet.cpp
@@ -131,6 +131,20 @@ void UNiagaraNodeParameterMapGet::OnNewTypedPinAdded(UEdGraphPin* NewPin)
 {
 	if (NewPin->Direction == EEdGraphPinDirection::EGPD_Output)
 	{
+		TArray<UEdGraphPin*> OutputPins;
+		GetOutputPins(OutputPins);
+
+		TSet<FName> Names;
+		for (const UEdGraphPin* Pin : OutputPins)
+		{
+			if (Pin != NewPin)
+			{
+				Names.Add(Pin->GetFName());
+			}
+		}
+		const FName NewUniqueName = FNiagaraUtilities::GetUniqueName(*NewPin->GetName(), Names);
+		NewPin->PinName = NewUniqueName;
+
 		const UEdGraphSchema_Niagara* Schema = GetDefault<UEdGraphSchema_Niagara>();
 		FNiagaraTypeDefinition TypeDef = Schema->PinToTypeDefinition(NewPin);
 
diff --git a/Engine/Plugins/FX/Niagara/Source/NiagaraEditor/Private/NiagaraNodeParameterMapSet.cpp b/Engine/Plugins/FX/Niagara/Source/NiagaraEditor/Private/NiagaraNodeParameterMapSet.cpp
--- a/Engine/Plugins/FX/Niagara/Source/NiagaraEditor/Private/NiagaraNodeParameterMapSet.cpp
+++ b/Engine/Plugins/FX/Niagara/Source/NiagaraEditor/Private/NiagaraNodeParameterMapSet.cpp
@@ -75,6 +75,23 @@ void UNiagaraNodeParameterMapSet::OnNewTypedPinAdded(UEdGraphPin* NewPin)
 		return;
 	}
 
+	if (NewPin->Direction == EEdGraphPinDirection::EGPD_Input)
+	{
+		TArray<UEdGraphPin*> InputPins;
+		GetInputPins(InputPins);
+
+		TSet<FName> Names;
+		for (const UEdGraphPin* Pin : InputPins)
+		{
+			if (Pin != NewPin)
+			{
+				Names.Add(Pin->GetFName());
+			}
+		}
+		const FName NewUniqueName = FNiagaraUtilities::GetUniqueName(*NewPin->GetName(), Names);
+		NewPin->PinName = NewUniqueName;
+	}
+
 	if (!NewPin->PersistentGuid.IsValid())
 	{
 		NewPin->PersistentGuid = FGuid::NewGuid();
diff --git a/Engine/Plugins/FX/Niagara/Source/NiagaraEditor/Private/Widgets/SNiagaraParameterMapView.cpp b/Engine/Plugins/FX/Niagara/Source/NiagaraEditor/Private/Widgets/SNiagaraParameterMapView.cpp
--- a/Engine/Plugins/FX/Niagara/Source/NiagaraEditor/Private/Widgets/SNiagaraParameterMapView.cpp
+++ b/Engine/Plugins/FX/Niagara/Source/NiagaraEditor/Private/Widgets/SNiagaraParameterMapView.cpp
@@ -8,6 +8,7 @@
 #include "Widgets/Input/SComboButton.h"
 #include "Widgets/Input/SSearchBox.h"
 #include "Widgets/SNiagaraGraphPinAdd.h"
+#include "NiagaraCommon.h"
 #include "NiagaraScript.h"
 #include "NiagaraScriptSource.h"
 #include "NiagaraEmitter.h"
@@ -232,6 +233,18 @@ bool SNiagaraParameterMapView::IsEditingMode() const
 
 void SNiagaraParameterMapView::AddParameter(FNiagaraVariable NewVariable)
 {
+	TSet<FName> Names;
+	for (auto& GraphWeakPtr : Graphs)
+	{
+		UNiagaraGraph* Graph = GraphWeakPtr.Get();
+		for (const auto& ParameterElement : Graph->GetParameterMap())
+		{
+			Names.Add(ParameterElement.Key.GetName());
+		}
+	}
+	const FName NewUniqueName = FNiagaraUtilities::GetUniqueName(NewVariable.GetName(), Names);
+	NewVariable.SetName(NewUniqueName);
+
 	bool bAddedParameter = false;
 	// Check whether we have to add this parameter to the user exposed system parameters.
 	FNiagaraParameterHandle ParameterHandle;
@@ -958,7 +971,7 @@ void SNiagaraAddParameterMenu::AddParameterSelected(FNiagaraVariable NewVariable
 	if (bCreateCustomName)
 	{
 		const FString ResultName = (InSection != NiagaraParameterMapSectionID::NONE ? NiagaraParameterMapSectionID::OnGetSectionTitle(InSection).ToString() + TEXT(".") : FString())
-			+ FString::Printf(TEXT("%s%d"), *NewVariable.GetName().ToString(), Graph->GetParameterMap().Num());
+			+ NewVariable.GetName().ToString();
 		NewVariable.SetName(FName(*ResultName));
 	}
 
