Fix: Adding a property pin to a Niagara Module Script map node creates a duplicate of that property in the Properties menu

diff --git a/Engine/Plugins/FX/Niagara/Source/NiagaraEditor/Private/Widgets/SNiagaraParameterMapView.cpp b/Engine/Plugins/FX/Niagara/Source/NiagaraEditor/Private/Widgets/SNiagaraParameterMapView.cpp
--- a/Engine/Plugins/FX/Niagara/Source/NiagaraEditor/Private/Widgets/SNiagaraParameterMapView.cpp
+++ b/Engine/Plugins/FX/Niagara/Source/NiagaraEditor/Private/Widgets/SNiagaraParameterMapView.cpp
@@ -854,50 +854,66 @@ void SNiagaraAddParameterMenu::CollectAllActions(FGraphActionListBuilderBase& Ou
 		CollectMakeNew(OutAllActions, NiagaraParameterMapSectionID::OTHER);
 	}
 
-	// Collect locals
-	if (Graph.Get())
+	// Engine
+	if (CanCollectSection(NiagaraParameterMapSectionID::ENGINE))
 	{
-		TArray<FNiagaraVariable> Parameters;
-
-		if (ShowGraphParameters.Get())
-		{
-			for (const auto& Parameter : Graph.Get()->GetParameterMap())
-			{
-				Parameters.AddUnique(Parameter.Key);
-			}
-		}
+		const FText Category = NiagaraParameterMapSectionID::OnGetSectionTitle(NiagaraParameterMapSectionID::ENGINE);
+		TArray<FNiagaraVariable> Variables = FNiagaraConstants::GetEngineConstants();
+		AddParameterGroup(OutAllActions,
+			Variables,
+			NiagaraParameterMapSectionID::ENGINE,
+			ShowNamespaceCategory.Get() ? NiagaraParameterMapSectionID::OnGetSectionTitle(NiagaraParameterMapSectionID::ENGINE) : FText::GetEmpty(),
+			FString(), true, false);
+	}
 
-		Parameters.Sort([](const FNiagaraVariable& A, const FNiagaraVariable& B) { return (A.GetName() < B.GetName()); });
+	// Collect parameter actions
+	if (ShowGraphParameters.Get() && Graph.Get())
+	{
+		TMap<FNiagaraVariable, FNiagaraGraphParameterReferenceCollection> ParameterEntries = Graph.Get()->GetParameterMap();
+		ParameterEntries.KeySort([](const FNiagaraVariable& A, const FNiagaraVariable& B) { return (A.GetName() < B.GetName()); });
 
-		for (const FNiagaraVariable& Parameter : Parameters)
+		for (const auto& ParameterEntry : ParameterEntries)
 		{
+			const FNiagaraVariable& Parameter = ParameterEntry.Key;
 			FNiagaraParameterHandle Handle;
 			const NiagaraParameterMapSectionID::Type ParameterSectionID = NiagaraParameterMapSectionID::OnGetSectionFromVariable(Parameter, Handle, /*Default*/ NiagaraParameterMapSectionID::OTHER);
 			if (CanCollectSection(ParameterSectionID))
 			{
 				const FText Category = ShowNamespaceCategory.Get() ? NiagaraParameterMapSectionID::OnGetSectionTitle(ParameterSectionID) : FText::GetEmpty();
 				const FText DisplayName = FText::FromName(Parameter.GetName());
-				const FText Tooltip = FText::GetEmpty();
-				TSharedPtr<FNiagaraMenuAction> Action(new FNiagaraMenuAction(
-					Category, DisplayName, Tooltip, 0, FText::GetEmpty(),
-					FNiagaraMenuAction::FOnExecuteStackAction::CreateSP(this, &SNiagaraAddParameterMenu::AddParameterSelected, Parameter, false, ParameterSectionID)));
-				OutAllActions.AddAction(Action);
+
+				// Only add this action if it isn't already in the list.
+				bool bUnique = true;
+				for (int32 Index = 0; Index < OutAllActions.GetNumActions(); Index++)
+				{
+					const FGraphActionListBuilderBase::ActionGroup& ActionGroup = OutAllActions.GetAction(Index);
+					for (const TSharedPtr<FEdGraphSchemaAction>& SchemaAction : ActionGroup.Actions)
+					{
+						if (SchemaAction->GetMenuDescription().EqualTo(DisplayName))
+						{
+							bUnique = false;
+							break;
+						}
+					}
+
+					if (!bUnique)
+					{
+						break;
+					}
+				}
+
+				if (bUnique)
+				{
+					const FText Tooltip = FText::GetEmpty();
+					TSharedPtr<FNiagaraMenuAction> Action(new FNiagaraMenuAction(
+						Category, DisplayName, Tooltip, 0, FText::GetEmpty(),
+						FNiagaraMenuAction::FOnExecuteStackAction::CreateSP(this, &SNiagaraAddParameterMenu::AddParameterSelected, Parameter, false, ParameterSectionID)));
+
+					OutAllActions.AddAction(Action);
+				}
 			}
 		}
 	}
-
-	// Engine
-	if (CanCollectSection(NiagaraParameterMapSectionID::ENGINE))
-	{
-		const FText Category = NiagaraParameterMapSectionID::OnGetSectionTitle(NiagaraParameterMapSectionID::ENGINE);
-		TArray<FNiagaraVariable> Variables = FNiagaraConstants::GetEngineConstants();
-		AddParameterGroup(OutAllActions,
-			Variables,
-			NiagaraParameterMapSectionID::ENGINE,
-			ShowNamespaceCategory.Get() ? NiagaraParameterMapSectionID::OnGetSectionTitle(NiagaraParameterMapSectionID::ENGINE) : FText::GetEmpty(),
-			FString(), true, false);
-	}
-
 }
 
 void SNiagaraAddParameterMenu::AddParameterGroup(FGraphActionListBuilderBase& OutActions, TArray<FNiagaraVariable>& Variables, const NiagaraParameterMapSectionID::Type InSection, const FText& Category, const FString& RootCategory, const bool bSort, const bool bCustomName)
