Parameter, source and dest pins of a parameter map node have a subcategories. Only pins with the parameter subcategory will be found by the graph.

diff --git a/Engine/Plugins/FX/Niagara/Source/NiagaraEditor/Private/NiagaraGraph.cpp b/Engine/Plugins/FX/Niagara/Source/NiagaraEditor/Private/NiagaraGraph.cpp
--- a/Engine/Plugins/FX/Niagara/Source/NiagaraEditor/Private/NiagaraGraph.cpp
+++ b/Engine/Plugins/FX/Niagara/Source/NiagaraEditor/Private/NiagaraGraph.cpp
@@ -587,7 +587,7 @@ bool UNiagaraGraph::RenameParameter(const FNiagaraVariable& Parameter, FName New
 	// Create the new parameter
 	FNiagaraVariable NewParameter = Parameter;
 	NewParameter.SetName(NewName);
-	
+
 	FNiagaraGraphParameterReferenceCollection* ReferenceCollection = Parameters.Find(Parameter);
 	if (ReferenceCollection)
 	{
@@ -1107,30 +1107,33 @@ void UNiagaraGraph::FindParameters()
 
 	auto AddParameterReference = [&](const FNiagaraVariable& Parameter, const UEdGraphPin* Pin, FNiagaraGraphParameterReferenceCollection*& ReferenceCollection)
 	{
-		const FNiagaraGraphParameterReference Reference(Pin->PersistentGuid, Cast<UNiagaraNode>(Pin->GetOwningNode()));
-		bool bNewReference = true;
-		if (ReferenceCollection)
+		if (Pin->PinType.PinSubCategory == UNiagaraNodeParameterMapBase::ParameterPinSubCategory)
 		{
-			ReferenceCollection->ParameterReferences.AddUnique(Reference);
-			bNewReference = false;
-		}
-		else
-		{
-			FNiagaraGraphParameterReferenceCollection* FoundReferenceCollection = Parameters.Find(Parameter);
-			if (FoundReferenceCollection)
+			const FNiagaraGraphParameterReference Reference(Pin->PersistentGuid, Cast<UNiagaraNode>(Pin->GetOwningNode()));
+			bool bNewReference = true;
+			if (ReferenceCollection)
 			{
-				ReferenceCollection = FoundReferenceCollection;
-				FoundReferenceCollection->ParameterReferences.AddUnique(Reference);
+				ReferenceCollection->ParameterReferences.AddUnique(Reference);
 				bNewReference = false;
 			}
-		}
+			else
+			{
+				FNiagaraGraphParameterReferenceCollection* FoundReferenceCollection = Parameters.Find(Parameter);
+				if (FoundReferenceCollection)
+				{
+					ReferenceCollection = FoundReferenceCollection;
+					FoundReferenceCollection->ParameterReferences.AddUnique(Reference);
+					bNewReference = false;
+				}
+			}
 
-		if (bNewReference)
-		{
-			FNiagaraGraphParameterReferenceCollection NewReferenceCollection;
-			NewReferenceCollection.ParameterReferences.Add(Reference);
-			NewReferenceCollection.Graph = this;
-			Parameters.Add(Parameter, NewReferenceCollection);
+			if (bNewReference)
+			{
+				FNiagaraGraphParameterReferenceCollection NewReferenceCollection;
+				NewReferenceCollection.ParameterReferences.Add(Reference);
+				NewReferenceCollection.Graph = this;
+				Parameters.Add(Parameter, NewReferenceCollection);
+			}
 		}
 	};
 
@@ -1160,26 +1163,26 @@ void UNiagaraGraph::FindParameters()
 	{
 		for (UEdGraphPin* Pin : Node->Pins)
 		{
-			const FNiagaraVariable Parameter = NiagaraSchema->PinToNiagaraVariable(Pin, false);
-			const FNiagaraParameterHandle Handle = FNiagaraParameterHandle(Parameter.GetName());
-
-			if (Handle.IsModuleHandle() && !FNiagaraConstants::IsNiagaraConstant(Parameter))
+			if (Pin->PinType.PinSubCategory == UNiagaraNodeParameterMapBase::ParameterPinSubCategory)
 			{
-				FNiagaraVariableMetaData* MetaData = VariableToMetaData.Find(Parameter);
-				if (MetaData)
-				{
-					MetaData->ReferencerNodes.AddUnique(Node);
-				}
-				else
+				const FNiagaraVariable Parameter = NiagaraSchema->PinToNiagaraVariable(Pin, false);
+				const FNiagaraParameterHandle Handle = FNiagaraParameterHandle(Parameter.GetName());
+
+				if (Handle.IsModuleHandle() && !FNiagaraConstants::IsNiagaraConstant(Parameter))
 				{
-					FNiagaraVariableMetaData NewVariableMetadata;
-					NewVariableMetadata.ReferencerNodes.Add(Node);
-					VariableToMetaData.Add(Parameter, NewVariableMetadata);
+					FNiagaraVariableMetaData* MetaData = VariableToMetaData.Find(Parameter);
+					if (MetaData)
+					{
+						MetaData->ReferencerNodes.AddUnique(Node);
+					}
+					else
+					{
+						FNiagaraVariableMetaData NewVariableMetadata;
+						NewVariableMetadata.ReferencerNodes.Add(Node);
+						VariableToMetaData.Add(Parameter, NewVariableMetadata);
+					}
 				}
-			}
 			
-			if (!Handle.GetNamespace().IsNone())
-			{
 				const FNiagaraGraphParameterReference Reference(Pin->PersistentGuid, Cast<UNiagaraNode>(Pin->GetOwningNode()));
 				FNiagaraGraphParameterReferenceCollection* FoundParameterReferenceCollection = Parameters.Find(Parameter);
 				if (FoundParameterReferenceCollection)
diff --git a/Engine/Plugins/FX/Niagara/Source/NiagaraEditor/Private/NiagaraNodeParameterMapBase.cpp b/Engine/Plugins/FX/Niagara/Source/NiagaraEditor/Private/NiagaraNodeParameterMapBase.cpp
--- a/Engine/Plugins/FX/Niagara/Source/NiagaraEditor/Private/NiagaraNodeParameterMapBase.cpp
+++ b/Engine/Plugins/FX/Niagara/Source/NiagaraEditor/Private/NiagaraNodeParameterMapBase.cpp
@@ -25,6 +25,10 @@
 
 #define LOCTEXT_NAMESPACE "NiagaraNodeParameterMapBase"
 
+const FName UNiagaraNodeParameterMapBase::ParameterPinSubCategory("ParameterPin");
+const FName UNiagaraNodeParameterMapBase::SourcePinName("Source");
+const FName UNiagaraNodeParameterMapBase::DestPinName("Dest");
+
 UNiagaraNodeParameterMapBase::UNiagaraNodeParameterMapBase() 
 	: UNiagaraNodeWithDynamicPins()
 	, PinPendingRename(nullptr)
diff --git a/Engine/Plugins/FX/Niagara/Source/NiagaraEditor/Private/NiagaraNodeParameterMapBase.h b/Engine/Plugins/FX/Niagara/Source/NiagaraEditor/Private/NiagaraNodeParameterMapBase.h
--- a/Engine/Plugins/FX/Niagara/Source/NiagaraEditor/Private/NiagaraNodeParameterMapBase.h
+++ b/Engine/Plugins/FX/Niagara/Source/NiagaraEditor/Private/NiagaraNodeParameterMapBase.h
@@ -39,4 +39,11 @@ protected:
 	virtual void OnPinRenamed(UEdGraphPin* RenamedPin, const FString& OldName) override;
 
 	UEdGraphPin* PinPendingRename;
+
+public:
+	/** The sub category for parameter pins. */
+	static const FName ParameterPinSubCategory;
+
+	static const FName SourcePinName;
+	static const FName DestPinName;
 };
diff --git a/Engine/Plugins/FX/Niagara/Source/NiagaraEditor/Private/NiagaraNodeParameterMapGet.cpp b/Engine/Plugins/FX/Niagara/Source/NiagaraEditor/Private/NiagaraNodeParameterMapGet.cpp
--- a/Engine/Plugins/FX/Niagara/Source/NiagaraEditor/Private/NiagaraNodeParameterMapGet.cpp
+++ b/Engine/Plugins/FX/Niagara/Source/NiagaraEditor/Private/NiagaraNodeParameterMapGet.cpp
@@ -31,7 +31,7 @@ void UNiagaraNodeParameterMapGet::AllocateDefaultPins()
 {
 	PinPendingRename = nullptr;
 	const UEdGraphSchema_Niagara* Schema = GetDefault<UEdGraphSchema_Niagara>();
-	CreatePin(EGPD_Input, Schema->TypeDefinitionToPinType(FNiagaraTypeDefinition::GetParameterMapDef()), TEXT("Source"));
+	CreatePin(EGPD_Input, Schema->TypeDefinitionToPinType(FNiagaraTypeDefinition::GetParameterMapDef()), *UNiagaraNodeParameterMapBase::SourcePinName.ToString());
 	CreateAddPin(EGPD_Output);
 }
 
@@ -153,6 +153,8 @@ void UNiagaraNodeParameterMapGet::OnNewTypedPinAdded(UEdGraphPin* NewPin)
 		{
 			UEdGraphPin* DefaultPin = CreateDefaultPin(NewPin);
 		}
+
+		NewPin->PinType.PinSubCategory = UNiagaraNodeParameterMapBase::ParameterPinSubCategory;
 	}
 
 	if (HasAnyFlags(RF_NeedLoad | RF_NeedPostLoad | RF_NeedInitialization))
@@ -257,19 +259,23 @@ void UNiagaraNodeParameterMapGet::PostLoad()
 	GetOutputPins(OutputPins);
 	for (int32 i = 0; i < OutputPins.Num(); i++)
 	{
-		if (IsAddPin(OutputPins[i]))
+		UEdGraphPin* OutputPin = OutputPins[i];
+		if (IsAddPin(OutputPin))
 		{
 			continue;
 		}
-		UEdGraphPin* InputPin = GetDefaultPin(OutputPins[i]);
+
+		UEdGraphPin* InputPin = GetDefaultPin(OutputPin);
 		if (InputPin == nullptr)
 		{
-			CreateDefaultPin(OutputPins[i]);
+			CreateDefaultPin(OutputPin);
 		}
 		else
 		{
-			SynchronizeDefaultInputPin(InputPin, OutputPins[i]);
+			SynchronizeDefaultInputPin(InputPin, OutputPin);
 		}
+
+		OutputPin->PinType.PinSubCategory = UNiagaraNodeParameterMapBase::ParameterPinSubCategory;
 	}
 }
 
diff --git a/Engine/Plugins/FX/Niagara/Source/NiagaraEditor/Private/NiagaraNodeParameterMapSet.cpp b/Engine/Plugins/FX/Niagara/Source/NiagaraEditor/Private/NiagaraNodeParameterMapSet.cpp
--- a/Engine/Plugins/FX/Niagara/Source/NiagaraEditor/Private/NiagaraNodeParameterMapSet.cpp
+++ b/Engine/Plugins/FX/Niagara/Source/NiagaraEditor/Private/NiagaraNodeParameterMapSet.cpp
@@ -26,8 +26,8 @@ void UNiagaraNodeParameterMapSet::AllocateDefaultPins()
 {
 	PinPendingRename = nullptr;
 	const UEdGraphSchema_Niagara* Schema = GetDefault<UEdGraphSchema_Niagara>();
-	CreatePin(EGPD_Input, Schema->TypeDefinitionToPinType(FNiagaraTypeDefinition::GetParameterMapDef()), TEXT("Source"));
-	CreatePin(EGPD_Output, Schema->TypeDefinitionToPinType(FNiagaraTypeDefinition::GetParameterMapDef()), TEXT("Dest"));
+	CreatePin(EGPD_Input, Schema->TypeDefinitionToPinType(FNiagaraTypeDefinition::GetParameterMapDef()), *UNiagaraNodeParameterMapBase::SourcePinName.ToString());
+	CreatePin(EGPD_Output, Schema->TypeDefinitionToPinType(FNiagaraTypeDefinition::GetParameterMapDef()), *UNiagaraNodeParameterMapBase::DestPinName.ToString());
 	CreateAddPin(EGPD_Input);
 }
 
@@ -90,6 +90,7 @@ void UNiagaraNodeParameterMapSet::OnNewTypedPinAdded(UEdGraphPin* NewPin)
 		}
 		const FName NewUniqueName = FNiagaraUtilities::GetUniqueName(*NewPin->GetName(), Names);
 		NewPin->PinName = NewUniqueName;
+		NewPin->PinType.PinSubCategory = UNiagaraNodeParameterMapBase::ParameterPinSubCategory;
 	}
 
 	if (!NewPin->PersistentGuid.IsValid())
@@ -287,6 +288,11 @@ void UNiagaraNodeParameterMapSet::PostLoad()
 		{
 			Pin->PersistentGuid = FGuid::NewGuid();
 		}
+
+		if (Pin->Direction == EEdGraphPinDirection::EGPD_Input && Pin->GetFName() != UNiagaraNodeParameterMapBase::SourcePinName && !IsAddPin(Pin))
+		{
+			Pin->PinType.PinSubCategory = UNiagaraNodeParameterMapBase::ParameterPinSubCategory;
+		}
 	}
 	Super::PostLoad();
 }
