Niagara - Fix issues with drag/drop
+ Don't allow the user to drop a module if the usage flags of the target script aren't supported.
+ Allow dragging to different scripts event if they are in different graphs, or different emitters.
+ Transfer rapid iteration paramters correctly when moving modules between scripts.
+ Fix undo for rapid iteration paramters when undoing a move.

diff --git a/Engine/Plugins/FX/Niagara/Source/NiagaraCore/Public/NiagaraCustomVersion.h b/Engine/Plugins/FX/Niagara/Source/NiagaraCore/Public/NiagaraCustomVersion.h
--- a/Engine/Plugins/FX/Niagara/Source/NiagaraCore/Public/NiagaraCustomVersion.h
+++ b/Engine/Plugins/FX/Niagara/Source/NiagaraCore/Public/NiagaraCustomVersion.h
@@ -102,6 +102,8 @@ struct FNiagaraCustomVersion
 
 		AssignmentNodeUsesBeginDefaults, // Changing the graph generation for assignment nodes so that it uses a "Begin Defaults" node where appropriate.
 
+		AssignmentNodeHasCorrectUsageBitmask, // Updating the usage flage bitmask for assignment nodes to match the part of the stack it's used in.
+
 		// -----<new versions can be added above this line>-------------------------------------------------
 		VersionPlusOne,
 		LatestVersion = VersionPlusOne - 1,
diff --git a/Engine/Plugins/FX/Niagara/Source/NiagaraEditor/Private/NiagaraEditorUtilities.cpp b/Engine/Plugins/FX/Niagara/Source/NiagaraEditor/Private/NiagaraEditorUtilities.cpp
--- a/Engine/Plugins/FX/Niagara/Source/NiagaraEditor/Private/NiagaraEditorUtilities.cpp
+++ b/Engine/Plugins/FX/Niagara/Source/NiagaraEditor/Private/NiagaraEditorUtilities.cpp
@@ -818,4 +818,63 @@ void FNiagaraEditorUtilities::PreprocessFunctionGraph(const UEdGraphSchema_Niaga
 
 }
 
+UNiagaraNodeOutput* FNiagaraEditorUtilities::GetScriptOutputNode(UNiagaraScript& Script)
+{
+	UNiagaraScriptSource* Source = CastChecked<UNiagaraScriptSource>(Script.GetSource());
+	return Source->NodeGraph->FindEquivalentOutputNode(Script.GetUsage(), Script.GetUsageId());
+}
+
+UNiagaraScript* FNiagaraEditorUtilities::GetScriptFromSystem(UNiagaraSystem& System, FGuid EmitterHandleId, ENiagaraScriptUsage Usage, FGuid UsageId)
+{
+	if (UNiagaraScript::IsEquivalentUsage(Usage, ENiagaraScriptUsage::SystemSpawnScript))
+	{
+		return System.GetSystemSpawnScript();
+	}
+	else if (UNiagaraScript::IsEquivalentUsage(Usage, ENiagaraScriptUsage::SystemUpdateScript))
+	{
+		return System.GetSystemUpdateScript();
+	}
+	else if (EmitterHandleId.IsValid())
+	{
+		const FNiagaraEmitterHandle* ScriptEmitterHandle = System.GetEmitterHandles().FindByPredicate(
+			[EmitterHandleId](const FNiagaraEmitterHandle& EmitterHandle) { return EmitterHandle.GetId() == EmitterHandleId; });
+		if (ScriptEmitterHandle != nullptr)
+		{
+			if (UNiagaraScript::IsEquivalentUsage(Usage, ENiagaraScriptUsage::EmitterSpawnScript))
+			{
+				return ScriptEmitterHandle->GetInstance()->EmitterSpawnScriptProps.Script;
+			}
+			else if (UNiagaraScript::IsEquivalentUsage(Usage, ENiagaraScriptUsage::EmitterUpdateScript))
+			{
+				return ScriptEmitterHandle->GetInstance()->EmitterUpdateScriptProps.Script;
+			}
+			else if (UNiagaraScript::IsEquivalentUsage(Usage, ENiagaraScriptUsage::ParticleSpawnScript))
+			{
+				return ScriptEmitterHandle->GetInstance()->SpawnScriptProps.Script;
+			}
+			else if (UNiagaraScript::IsEquivalentUsage(Usage, ENiagaraScriptUsage::ParticleUpdateScript))
+			{
+				return ScriptEmitterHandle->GetInstance()->UpdateScriptProps.Script;
+			}
+			else if (UNiagaraScript::IsEquivalentUsage(Usage, ENiagaraScriptUsage::ParticleEventScript))
+			{
+				for (const FNiagaraEventScriptProperties& EventScriptProperties : ScriptEmitterHandle->GetInstance()->GetEventHandlers())
+				{
+					if (EventScriptProperties.Script->GetUsageId() == UsageId)
+					{
+						return EventScriptProperties.Script;
+					}
+				}
+			}
+		}
+	}
+	return nullptr;
+}
+
+const FNiagaraEmitterHandle* FNiagaraEditorUtilities::GetEmitterHandleForEmitter(UNiagaraSystem& System, UNiagaraEmitter& Emitter)
+{
+	return System.GetEmitterHandles().FindByPredicate(
+		[&Emitter](const FNiagaraEmitterHandle& EmitterHandle) { return EmitterHandle.GetInstance() == &Emitter; });
+}
+
 #undef LOCTEXT_NAMESPACE
diff --git a/Engine/Plugins/FX/Niagara/Source/NiagaraEditor/Private/NiagaraNodeAssignment.cpp b/Engine/Plugins/FX/Niagara/Source/NiagaraEditor/Private/NiagaraNodeAssignment.cpp
--- a/Engine/Plugins/FX/Niagara/Source/NiagaraEditor/Private/NiagaraNodeAssignment.cpp
+++ b/Engine/Plugins/FX/Niagara/Source/NiagaraEditor/Private/NiagaraNodeAssignment.cpp
@@ -161,6 +161,13 @@ void UNiagaraNodeAssignment::PostLoad()
 			FunctionScript = nullptr;
 			GenerateScript();
 		}
+		if (NiagaraVer < FNiagaraCustomVersion::AssignmentNodeHasCorrectUsageBitmask)
+		{
+			if (FunctionScript != nullptr)
+			{
+				UpdateUsageBitmaskFromOwningScript();
+			}
+		}
 	}
 }
 
@@ -177,6 +184,7 @@ void UNiagaraNodeAssignment::GenerateScript()
 		FunctionScript->SetUsage(ENiagaraScriptUsage::Module);
 		FunctionScript->Description = LOCTEXT("AssignmentNodeDesc", "Sets one or more variables in the stack.");
 		InitializeScript(FunctionScript);
+		UpdateUsageBitmaskFromOwningScript();
 		ComputeNodeName();
 	}
 }
@@ -279,6 +287,11 @@ void UNiagaraNodeAssignment::RemoveParameter(const FNiagaraVariable& InVar)
 	OnInputsChangedDelegate.Broadcast();
 }
 
+void UNiagaraNodeAssignment::UpdateUsageBitmaskFromOwningScript()
+{
+	FunctionScript->ModuleUsageBitmask = CalculateScriptUsageBitmask();
+}
+
 void UNiagaraNodeAssignment::InitializeScript(UNiagaraScript* NewScript)
 {
 	if (NewScript != NULL)
@@ -461,6 +474,44 @@ void UNiagaraNodeAssignment::InitializeScript(UNiagaraScript* NewScript)
 	}
 }
 
+int32 UsageToBitmask(ENiagaraScriptUsage Usage)
+{
+	return 1 << (int32)Usage;
+}
+
+int32 UNiagaraNodeAssignment::CalculateScriptUsageBitmask()
+{
+	int32 UsageBitmask = 0;
+	UNiagaraNodeOutput* OutputNode = FNiagaraStackGraphUtilities::GetEmitterOutputNodeForStackNode(*this);
+	if (OutputNode != nullptr)
+	{
+		if (UNiagaraScript::IsEquivalentUsage(OutputNode->GetUsage(), ENiagaraScriptUsage::SystemSpawnScript) ||
+			UNiagaraScript::IsEquivalentUsage(OutputNode->GetUsage(), ENiagaraScriptUsage::SystemUpdateScript))
+		{
+			UsageBitmask =
+				UsageToBitmask(ENiagaraScriptUsage::SystemSpawnScript) |
+				UsageToBitmask(ENiagaraScriptUsage::SystemUpdateScript);
+		}
+		if (UNiagaraScript::IsEquivalentUsage(OutputNode->GetUsage(), ENiagaraScriptUsage::EmitterSpawnScript) ||
+			UNiagaraScript::IsEquivalentUsage(OutputNode->GetUsage(), ENiagaraScriptUsage::EmitterUpdateScript))
+		{
+			UsageBitmask =
+				UsageToBitmask(ENiagaraScriptUsage::EmitterSpawnScript) |
+				UsageToBitmask(ENiagaraScriptUsage::EmitterUpdateScript);
+		}
+		if (UNiagaraScript::IsEquivalentUsage(OutputNode->GetUsage(), ENiagaraScriptUsage::ParticleSpawnScript) ||
+			UNiagaraScript::IsEquivalentUsage(OutputNode->GetUsage(), ENiagaraScriptUsage::ParticleUpdateScript) ||
+			UNiagaraScript::IsEquivalentUsage(OutputNode->GetUsage(), ENiagaraScriptUsage::ParticleEventScript))
+		{
+			UsageBitmask =
+				UsageToBitmask(ENiagaraScriptUsage::ParticleSpawnScript) |
+				UsageToBitmask(ENiagaraScriptUsage::ParticleUpdateScript) |
+				UsageToBitmask(ENiagaraScriptUsage::ParticleEventScript);
+		}
+	}
+	return UsageBitmask;
+}
+
 int32 UNiagaraNodeAssignment::FindAssignmentTarget(const FName& InName, const FNiagaraTypeDefinition& InType)
 {
 	for (int32 i = 0; i < AssignmentTargets.Num(); i++)
diff --git a/Engine/Plugins/FX/Niagara/Source/NiagaraEditor/Private/NiagaraParameterMapHistory.cpp b/Engine/Plugins/FX/Niagara/Source/NiagaraEditor/Private/NiagaraParameterMapHistory.cpp
--- a/Engine/Plugins/FX/Niagara/Source/NiagaraEditor/Private/NiagaraParameterMapHistory.cpp
+++ b/Engine/Plugins/FX/Niagara/Source/NiagaraEditor/Private/NiagaraParameterMapHistory.cpp
@@ -337,7 +337,7 @@ bool FNiagaraParameterMapHistory::IsRapidIterationParameter(const FNiagaraVariab
 	return IsInNamespace(InVar, PARAM_MAP_RAPID_ITERATION_STR);
 }
 
-bool FNiagaraParameterMapHistory::TryGetEmitterAndFunctionCallNamesFromRapidIterationParameter(const FNiagaraVariable& InVar, FString& EmitterName, FString& FunctionCallName)
+bool FNiagaraParameterMapHistory::SplitRapidIterationParameterName(const FNiagaraVariable& InVar, FString& EmitterName, FString& FunctionCallName, FString& InputName)
 {
 	TArray<FString> SplitName;
 	InVar.GetName().ToString().ParseIntoArray(SplitName, TEXT("."));
@@ -345,6 +345,11 @@ bool FNiagaraParameterMapHistory::TryGetEmitterAndFunctionCallNamesFromRapidIter
 	{
 		EmitterName = SplitName[1];
 		FunctionCallName = SplitName[2];
+		InputName = SplitName[3];
+		for (int i = 4; i < SplitName.Num(); i++)
+		{
+			InputName += TEXT(".") + SplitName[i];
+		}
 		return true;
 	}
 	return false;
diff --git a/Engine/Plugins/FX/Niagara/Source/NiagaraEditor/Private/ViewModels/Stack/NiagaraStackGraphUtilities.cpp b/Engine/Plugins/FX/Niagara/Source/NiagaraEditor/Private/ViewModels/Stack/NiagaraStackGraphUtilities.cpp
--- a/Engine/Plugins/FX/Niagara/Source/NiagaraEditor/Private/ViewModels/Stack/NiagaraStackGraphUtilities.cpp
+++ b/Engine/Plugins/FX/Niagara/Source/NiagaraEditor/Private/ViewModels/Stack/NiagaraStackGraphUtilities.cpp
@@ -931,6 +931,8 @@ UNiagaraNodeAssignment* FNiagaraStackGraphUtilities::AddParameterModuleToStack(c
 	AssignmentNodeCreator.Finalize();
 
 	ConnectModuleNode(*NewAssignmentNode, TargetOutputNode, TargetIndex);
+	NewAssignmentNode->UpdateUsageBitmaskFromOwningScript();
+
 	return NewAssignmentNode;
 }
 
@@ -1145,7 +1147,8 @@ void FNiagaraStackGraphUtilities::CleanUpStaleRapidIterationParameters(UNiagaraS
 		{
 			FString EmitterName;
 			FString FunctionCallName;
-			if (FNiagaraParameterMapHistory::TryGetEmitterAndFunctionCallNamesFromRapidIterationParameter(RapidIterationParameter, EmitterName, FunctionCallName))
+			FString InputName;
+			if (FNiagaraParameterMapHistory::SplitRapidIterationParameterName(RapidIterationParameter, EmitterName, FunctionCallName, InputName))
 			{
 				if (EmitterName != OwningEmitter.GetUniqueEmitterName() || ValidFunctionCallNames.Contains(FunctionCallName) == false)
 				{
@@ -1630,4 +1633,199 @@ bool FNiagaraStackGraphUtilities::GetStackIssuesRecursively(const UNiagaraStackE
 	return OutIssues.Num() > 0;
 }
 
+void FNiagaraStackGraphUtilities::MoveModule(UNiagaraScript& SourceScript, UNiagaraNodeFunctionCall& ModuleToMove, UNiagaraSystem& TargetSystem, FGuid TargetEmitterHandleId, ENiagaraScriptUsage TargetUsage, FGuid TargetUsageId, int32 TargetModuleIndex)
+{
+	UNiagaraScript* TargetScript = FNiagaraEditorUtilities::GetScriptFromSystem(TargetSystem, TargetEmitterHandleId, TargetUsage, TargetUsageId);
+	checkf(TargetScript != nullptr, TEXT("Target script not found"));
+
+	UNiagaraNodeOutput* TargetOutputNode = FNiagaraEditorUtilities::GetScriptOutputNode(*TargetScript);
+	checkf(TargetOutputNode != nullptr, TEXT("Target stack is invalid"));
+
+	TArray<FStackNodeGroup> SourceGroups;
+	GetStackNodeGroups(ModuleToMove, SourceGroups);
+	int32 SourceGroupIndex = SourceGroups.IndexOfByPredicate([&ModuleToMove](const FStackNodeGroup& SourceGroup) { return SourceGroup.EndNode == &ModuleToMove; });
+	TArray<UNiagaraNode*> SourceGroupNodes;
+	SourceGroups[SourceGroupIndex].GetAllNodesInGroup(SourceGroupNodes);
+
+	UNiagaraGraph* SourceGraph = ModuleToMove.GetNiagaraGraph();
+	UNiagaraGraph* TargetGraph = TargetOutputNode->GetNiagaraGraph();
+
+	// If the source and target scripts don't match we need to collect the rapid iteration parameter values for each function in the source group
+	// so we can restore them after moving.
+	TMap<FGuid, TArray<FNiagaraVariable>> SourceFunctionIdToRapidIterationParametersMap;
+	if (&SourceScript != TargetScript)
+	{
+		TMap<FString, FGuid> FunctionCallNameToNodeIdMap;
+		for (UNiagaraNode* SourceGroupNode : SourceGroupNodes)
+		{
+			UNiagaraNodeFunctionCall* FunctionNode = Cast<UNiagaraNodeFunctionCall>(SourceGroupNode);
+			if (FunctionNode != nullptr)
+			{
+				FunctionCallNameToNodeIdMap.Add(FunctionNode->GetFunctionName(), FunctionNode->NodeGuid);
+			}
+		}
+
+		TArray<FNiagaraVariable> ScriptRapidIterationParameters;
+		SourceScript.RapidIterationParameters.GetParameters(ScriptRapidIterationParameters);
+		for (FNiagaraVariable& ScriptRapidIterationParameter : ScriptRapidIterationParameters)
+		{
+			FString EmitterName;
+			FString FunctionCallName;
+			FString InputName;
+			if (FNiagaraParameterMapHistory::SplitRapidIterationParameterName(
+				ScriptRapidIterationParameter, EmitterName, FunctionCallName, InputName))
+			{
+				FGuid* NodeIdPtr = FunctionCallNameToNodeIdMap.Find(FunctionCallName);
+				if (NodeIdPtr != nullptr)
+				{
+					TArray<FNiagaraVariable>& RapidIterationParameters = SourceFunctionIdToRapidIterationParametersMap.FindOrAdd(*NodeIdPtr);
+					RapidIterationParameters.Add(ScriptRapidIterationParameter);
+					RapidIterationParameters.Last().SetData(SourceScript.RapidIterationParameters.GetParameterData(ScriptRapidIterationParameter));
+				}
+			}
+		}
+	}
+
+	FStackNodeGroup TargetGroup;
+	TArray<UNiagaraNode*> TargetGroupNodes;
+	TMap<FGuid, FGuid> OldNodeIdToNewIdMap;
+	if (SourceGraph == TargetGraph)
+	{
+		TargetGroup = SourceGroups[SourceGroupIndex];
+		TargetGroupNodes = SourceGroupNodes;
+	}
+	else 
+	{
+		// If the module is being inserted into a different graph all of the nodes need to be copied into the target graph.
+		FStackNodeGroup SourceGroup = SourceGroups[SourceGroupIndex];
+		
+		TSet<UObject*> NodesToCopy;
+		for (UNiagaraNode* SourceGroupNode : SourceGroupNodes)
+		{
+			SourceGroupNode->PrepareForCopying();
+			NodesToCopy.Add(SourceGroupNode);
+		}
+
+		FString ExportedText;
+		FEdGraphUtilities::ExportNodesToText(NodesToCopy, ExportedText);
+
+		TSet<UEdGraphNode*> CopiedNodesSet;
+		FEdGraphUtilities::ImportNodesFromText(TargetGraph, ExportedText, CopiedNodesSet);
+		TArray<UEdGraphNode*> CopiedNodes = CopiedNodesSet.Array();
+
+		// Collect the start and end nodes for the group by ID before assigning the copied nodes new ids.
+		UEdGraphNode** CopiedEndNode = CopiedNodes.FindByPredicate([SourceGroup](UEdGraphNode* CopiedNode)
+			{ return CopiedNode->NodeGuid == SourceGroup.EndNode->NodeGuid; });
+		checkf(CopiedEndNode != nullptr, TEXT("Group copy failed"));
+		TargetGroup.EndNode = CastChecked<UNiagaraNode>(*CopiedEndNode);
+
+		for (UNiagaraNode* StartNode : SourceGroup.StartNodes)
+		{
+			UEdGraphNode** CopiedStartNode = CopiedNodes.FindByPredicate([StartNode](UEdGraphNode* CopiedNode)
+				{ return CopiedNode->NodeGuid == StartNode->NodeGuid; });
+			checkf(CopiedStartNode != nullptr, TEXT("Group copy failed"));
+			TargetGroup.StartNodes.Add(CastChecked<UNiagaraNode>(*CopiedStartNode));
+		}
+
+		TargetGroup.GetAllNodesInGroup(TargetGroupNodes);
+
+		// Assign all of the new nodes new ids and mark them as requiring synchronization.
+		for (UEdGraphNode* CopiedNode : CopiedNodes)
+		{
+			FGuid OldId = CopiedNode->NodeGuid;
+			CopiedNode->CreateNewGuid();
+			OldNodeIdToNewIdMap.Add(OldId, CopiedNode->NodeGuid);
+			UNiagaraNode* CopiedNiagaraNode = Cast<UNiagaraNode>(CopiedNode);
+			if (CopiedNiagaraNode != nullptr)
+			{
+				CopiedNiagaraNode->MarkNodeRequiresSynchronization(__FUNCTION__, false);
+			}
+		}
+	}
+
+	TArray<FStackNodeGroup> TargetGroups;
+	GetStackNodeGroups(*TargetOutputNode, TargetGroups);
+
+	// The first group is the output node, so to get the group index from module index we need to add 1, but 
+	// if a valid index wasn't supplied, than we insert at the end.
+	int32 TargetGroupIndex = TargetModuleIndex != INDEX_NONE ? TargetModuleIndex + 1 : TargetGroups.Num() - 1;
+
+	// Remove the source module group from it's stack, and insert the source or copied nodes into the target stack.
+	DisconnectStackNodeGroup(SourceGroups[SourceGroupIndex], SourceGroups[SourceGroupIndex - 1], SourceGroups[SourceGroupIndex + 1]);
+	if (SourceGraph != TargetGraph)
+	{
+		// If the graphs were different also remove the nodes from the source graph.
+		for (UNiagaraNode* SourceGroupNode : SourceGroupNodes)
+		{
+			SourceGraph->RemoveNode(SourceGroupNode);
+		}
+	}
+
+	ConnectStackNodeGroup(TargetGroup, TargetGroups[TargetGroupIndex - 1], TargetGroups[TargetGroupIndex]);
+
+	// Copy any rapid iteration parameters cached earlier into the target script.
+	if (SourceFunctionIdToRapidIterationParametersMap.Num() != 0)
+	{
+		SourceScript.Modify();
+		TargetScript->Modify();
+		if (SourceGraph == TargetGraph)
+		{
+			// If the module was dropped in the same graph than neither the emitter or function call name could have changed
+			// so we can just add them directly to the target script.
+			for (auto It = SourceFunctionIdToRapidIterationParametersMap.CreateIterator(); It; ++It)
+			{
+				FGuid& FunctionId = It.Key();
+				TArray<FNiagaraVariable>& RapidIterationParameters = It.Value();
+				for (FNiagaraVariable& RapidIterationParameter : RapidIterationParameters)
+				{
+					TargetScript->RapidIterationParameters.SetParameterData(RapidIterationParameter.GetData(), RapidIterationParameter, true);
+				}
+			}
+		}
+		else
+		{
+			// If the module was moved to a different graph it's possible that the emitter name or function call name could have
+			// changed so we need to construct new rapid iteration parameters.
+			FString EmitterName;
+			if (TargetEmitterHandleId.IsValid())
+			{
+				const FNiagaraEmitterHandle* TargetEmitterHandle = TargetSystem.GetEmitterHandles().FindByPredicate(
+					[TargetEmitterHandleId](const FNiagaraEmitterHandle& EmitterHandle) { return EmitterHandle.GetId() == TargetEmitterHandleId; });
+				EmitterName = TargetEmitterHandle->GetUniqueInstanceName();
+			}
+
+			TArray<FNiagaraVariable> TargetRapidIterationParameters;
+			for (auto It = SourceFunctionIdToRapidIterationParametersMap.CreateIterator(); It; ++It)
+			{
+				FGuid& FunctionId = It.Key();
+				TArray<FNiagaraVariable>& RapidIterationParameters = It.Value();
+
+				if (SourceGraph == TargetGraph)
+				{
+					TargetRapidIterationParameters.Append(RapidIterationParameters);
+				}
+				else
+				{
+					FGuid TargetNodeId = OldNodeIdToNewIdMap[FunctionId];
+					UNiagaraNode** TargetFunctionNodePtr = TargetGroupNodes.FindByPredicate(
+						[TargetNodeId](UNiagaraNode* TargetGroupNode) { return TargetGroupNode->NodeGuid == TargetNodeId; });
+					checkf(TargetFunctionNodePtr != nullptr, TEXT("Target nodes not copied correctly"));
+					UNiagaraNodeFunctionCall* TargetFunctionNode = CastChecked<UNiagaraNodeFunctionCall>(*TargetFunctionNodePtr);
+					for (FNiagaraVariable& RapidIterationParameter : RapidIterationParameters)
+					{
+						FString OldEmitterName;
+						FString OldFunctionCallName;
+						FString InputName;
+						FNiagaraParameterMapHistory::SplitRapidIterationParameterName(RapidIterationParameter, OldEmitterName, OldFunctionCallName, InputName);
+						FNiagaraParameterHandle ModuleHandle = FNiagaraParameterHandle::CreateModuleParameterHandle(*InputName);
+						FNiagaraParameterHandle AliasedModuleHandle = FNiagaraParameterHandle::CreateAliasedModuleParameterHandle(ModuleHandle, TargetFunctionNode);
+						FNiagaraVariable TargetRapidIterationParameter = CreateRapidIterationParameter(EmitterName, TargetUsage, AliasedModuleHandle.GetParameterHandleString(), RapidIterationParameter.GetType());
+						TargetScript->RapidIterationParameters.SetParameterData(RapidIterationParameter.GetData(), TargetRapidIterationParameter, true);
+					}
+				}
+			}
+		}
+	}
+}
+
 #undef LOCTEXT_NAMESPACE
diff --git a/Engine/Plugins/FX/Niagara/Source/NiagaraEditor/Private/ViewModels/Stack/NiagaraStackScriptItemGroup.cpp b/Engine/Plugins/FX/Niagara/Source/NiagaraEditor/Private/ViewModels/Stack/NiagaraStackScriptItemGroup.cpp
--- a/Engine/Plugins/FX/Niagara/Source/NiagaraEditor/Private/ViewModels/Stack/NiagaraStackScriptItemGroup.cpp
+++ b/Engine/Plugins/FX/Niagara/Source/NiagaraEditor/Private/ViewModels/Stack/NiagaraStackScriptItemGroup.cpp
@@ -3,6 +3,7 @@
 #include "ViewModels/Stack/NiagaraStackScriptItemGroup.h"
 #include "ViewModels/Stack/NiagaraStackModuleItem.h"
 #include "ViewModels/Stack/NiagaraStackSpacer.h"
+#include "ViewModels/NiagaraSystemViewModel.h"
 #include "ViewModels/NiagaraEmitterViewModel.h"
 #include "NiagaraScriptViewModel.h"
 #include "NiagaraScriptGraphViewModel.h"
@@ -15,6 +16,8 @@
 #include "ViewModels/Stack/NiagaraStackGraphUtilities.h"
 #include "NiagaraConstants.h"
 #include "NiagaraStackEditorData.h"
+#include "NiagaraSystem.h"
+#include "NiagaraEmitterHandle.h"
 
 #include "Internationalization/Internationalization.h"
 #include "ScopedTransaction.h"
@@ -486,9 +489,17 @@ TOptional<UNiagaraStackEntry::FDropResult> UNiagaraStackScriptItemGroup::ChildRe
 				return FDropResult(false, LOCTEXT("CantMoveModuleError", "This inherited module can't be moved."));
 			}
 
-			if (SourceModuleItem->GetModuleNode().GetGraph() != ScriptViewModel.Pin()->GetGraphViewModel()->GetGraph())
+			TArray<ENiagaraScriptUsage> SourceUsages = SourceModuleItem->GetModuleNode().FunctionScript->GetSupportedUsageContexts();
+			if (SourceUsages.ContainsByPredicate([this](ENiagaraScriptUsage SourceUsage) { return UNiagaraScript::IsEquivalentUsage(ScriptUsage, SourceUsage); }) == false)
 			{
-				return FDropResult(false, LOCTEXT("CantMoveModuleBetweenGraphsError", "This module can not be moved to this section of the stack"));
+				return FDropResult(false, LOCTEXT("CantMoveByUsage", "This module can't be moved to this section of the stack because its it's not supported in this context."));
+			}
+
+			const FNiagaraEmitterHandle* SourceEmitterHandle = FNiagaraEditorUtilities::GetEmitterHandleForEmitter(
+				SourceModuleItem->GetSystemViewModel()->GetSystem(), *SourceModuleItem->GetEmitterViewModel()->GetEmitter());
+			if (SourceEmitterHandle == nullptr)
+			{
+				return FDropResult(false, LOCTEXT("CantMoveFromAnotherSystem", "This module can't be moved into this system from a different system."));
 			}
 
 			const UNiagaraStackSpacer* TargetSpacer = Cast<UNiagaraStackSpacer>(&TargetChild);
@@ -528,10 +539,10 @@ TOptional<UNiagaraStackEntry::FDropResult> UNiagaraStackScriptItemGroup::ChildRe
 	if (bIsValidForOutput && DraggedEntries.Num() == 1)
 	{
 		UNiagaraStackModuleItem* SourceModuleItem = Cast<UNiagaraStackModuleItem>(DraggedEntries[0]);
-		UNiagaraGraph* TargetGraph = ScriptViewModel.Pin()->GetGraphViewModel()->GetGraph();
+		TArray<ENiagaraScriptUsage> SourceUsages = SourceModuleItem->GetModuleNode().FunctionScript->GetSupportedUsageContexts();
 		if (SourceModuleItem != nullptr &&
 			SourceModuleItem->CanMoveAndDelete() &&
-			SourceModuleItem->GetModuleNode().GetGraph() == TargetGraph)
+			SourceUsages.ContainsByPredicate([this](ENiagaraScriptUsage SourceUsage) { return UNiagaraScript::IsEquivalentUsage(ScriptUsage, SourceUsage); }))
 		{
 			const UNiagaraStackSpacer* TargetSpacer = Cast<UNiagaraStackSpacer>(&TargetChild);
 			if (TargetSpacer != nullptr)
@@ -539,33 +550,30 @@ TOptional<UNiagaraStackEntry::FDropResult> UNiagaraStackScriptItemGroup::ChildRe
 				UNiagaraStackModuleItem** TargetModuleItemPtr = StackSpacerToModuleItemMap.Find(FObjectKey(TargetSpacer));
 				if (TargetModuleItemPtr != nullptr)
 				{
-					UNiagaraStackModuleItem* TargetModuleItem = *TargetModuleItemPtr;
-					TArray<FNiagaraStackGraphUtilities::FStackNodeGroup> SourceStackGroups;
-					TArray<FNiagaraStackGraphUtilities::FStackNodeGroup> TargetStackGroups;
-					int32 SourceGroupIndex;
-					int32 TargetGroupIndex;
-					GenerateDragDropData(
-						&SourceModuleItem->GetModuleNode(), TargetModuleItem != nullptr ? &TargetModuleItem->GetModuleNode() : nullptr,
-						TargetGraph, ScriptUsage, ScriptUsageId,
-						SourceStackGroups, SourceGroupIndex,
-						TargetStackGroups, TargetGroupIndex);
-
-					// Make sure the source and target indices are within safe ranges, and make sure that the insert target isn't the source target or the spot directly
-					// after the source target since that won't actually move the module.
-					if (SourceGroupIndex > 0 && SourceGroupIndex < SourceStackGroups.Num() - 1 && TargetGroupIndex > 0 && TargetGroupIndex < TargetStackGroups.Num() &&
-						SourceStackGroups[SourceGroupIndex].EndNode != TargetStackGroups[TargetGroupIndex].EndNode &&
-						SourceStackGroups[SourceGroupIndex].EndNode != TargetStackGroups[TargetGroupIndex - 1].EndNode)
+					const FNiagaraEmitterHandle* SourceEmitterHandle = FNiagaraEditorUtilities::GetEmitterHandleForEmitter(
+						SourceModuleItem->GetSystemViewModel()->GetSystem(), *SourceModuleItem->GetEmitterViewModel()->GetEmitter());
+					if (SourceEmitterHandle != nullptr)
 					{
-						FScopedTransaction ScopedTransaction(LOCTEXT("DragAndDropModule", "Drag and drop module"));
-						FNiagaraStackGraphUtilities::DisconnectStackNodeGroup(SourceStackGroups[SourceGroupIndex], SourceStackGroups[SourceGroupIndex - 1], SourceStackGroups[SourceGroupIndex + 1]);
-						FNiagaraStackGraphUtilities::ConnectStackNodeGroup(SourceStackGroups[SourceGroupIndex], TargetStackGroups[TargetGroupIndex - 1], TargetStackGroups[TargetGroupIndex]);
+						UNiagaraNodeOutput* SourceModuleOutputNode = FNiagaraStackGraphUtilities::GetEmitterOutputNodeForStackNode(SourceModuleItem->GetModuleNode());
+						UNiagaraScript* SourceModuleScript = FNiagaraEditorUtilities::GetScriptFromSystem(GetSystemViewModel()->GetSystem(), SourceEmitterHandle->GetId(),
+							SourceModuleOutputNode->GetUsage(), SourceModuleOutputNode->GetUsageId());
+
+						const FNiagaraEmitterHandle* TargetEmitterHandle = FNiagaraEditorUtilities::GetEmitterHandleForEmitter(GetSystemViewModel()->GetSystem(), *GetEmitterViewModel()->GetEmitter());
+
+						int32 TargetIndex = *TargetModuleItemPtr != nullptr ? (*TargetModuleItemPtr)->GetModuleIndex() : INDEX_NONE;
 
-						FNiagaraStackGraphUtilities::RelayoutGraph(*TargetGraph);
-						TargetGraph->NotifyGraphNeedsRecompile();
+						FScopedTransaction ScopedTransaction(LOCTEXT("DragAndDropModule", "Drag and drop module"));
+						{
+							FNiagaraStackGraphUtilities::MoveModule(*SourceModuleScript, SourceModuleItem->GetModuleNode(), GetSystemViewModel()->GetSystem(), TargetEmitterHandle->GetId(),
+								ScriptUsage, ScriptUsageId, TargetIndex);
 
-						SourceModuleItem->NotifyModuleMoved();
-						RefreshChildren();
+							UNiagaraGraph* TargetGraph = ScriptViewModel.Pin()->GetGraphViewModel()->GetGraph();
+							FNiagaraStackGraphUtilities::RelayoutGraph(*TargetGraph);
+							TargetGraph->NotifyGraphNeedsRecompile();
 
+							SourceModuleItem->NotifyModuleMoved();
+							RefreshChildren();
+						}
 						return FDropResult(true);
 					}
 				}
diff --git a/Engine/Plugins/FX/Niagara/Source/NiagaraEditor/Public/NiagaraEditorUtilities.h b/Engine/Plugins/FX/Niagara/Source/NiagaraEditor/Public/NiagaraEditorUtilities.h
--- a/Engine/Plugins/FX/Niagara/Source/NiagaraEditor/Public/NiagaraEditorUtilities.h
+++ b/Engine/Plugins/FX/Niagara/Source/NiagaraEditor/Public/NiagaraEditorUtilities.h
@@ -8,10 +8,12 @@
 #include "Misc/Attribute.h"
 
 class UNiagaraNodeInput;
+class UNiagaraNodeOutput;
 struct FNiagaraVariable;
 struct FNiagaraTypeDefinition;
 class UNiagaraGraph;
 class UNiagaraSystem;
+struct FNiagaraEmitterHandle;
 class UNiagaraEmitter;
 class UNiagaraScript;
 class FStructOnScope;
@@ -134,4 +136,17 @@ namespace FNiagaraEditorUtilities
 	void FixUpNumericPins(const UEdGraphSchema_Niagara* Schema, UNiagaraNode* Node);
 
 	void PreprocessFunctionGraph(const UEdGraphSchema_Niagara* Schema, UNiagaraGraph* Graph, const TArray<UEdGraphPin*>& CallInputs, const TArray<UEdGraphPin*>& CallOutputs, ENiagaraScriptUsage ScriptUsage);
+
+	UNiagaraNodeOutput* GetScriptOutputNode(UNiagaraScript& Script);
+
+	UNiagaraScript* GetScriptFromSystem(UNiagaraSystem& System, FGuid EmitterHandleId, ENiagaraScriptUsage Usage, FGuid UsageId);
+
+	/**
+	 * Gets an emitter handle from a system and an owned emitter.  This handle will become invalid if emitters are added or
+	 * removed from the system, so in general this value should not be cached across frames.
+	 * @param System The source system which owns the emitter handles.
+	 * @param The emitter to search for in the system.
+	 * @returns The emitter handle for the supplied emitter, or nullptr if the emitter isn't owned by this system.
+	 */
+	const FNiagaraEmitterHandle* GetEmitterHandleForEmitter(UNiagaraSystem& System, UNiagaraEmitter& Emitter);
 };
diff --git a/Engine/Plugins/FX/Niagara/Source/NiagaraEditor/Public/NiagaraNodeAssignment.h b/Engine/Plugins/FX/Niagara/Source/NiagaraEditor/Public/NiagaraNodeAssignment.h
--- a/Engine/Plugins/FX/Niagara/Source/NiagaraEditor/Public/NiagaraNodeAssignment.h
+++ b/Engine/Plugins/FX/Niagara/Source/NiagaraEditor/Public/NiagaraNodeAssignment.h
@@ -48,6 +48,8 @@ public:
 	void AddParameter(FNiagaraVariable InVar, FString InDefaultValue);
 	void RemoveParameter(const FNiagaraVariable& InVar);
 
+	void UpdateUsageBitmaskFromOwningScript();
+
 protected:
 
 	UPROPERTY()
@@ -69,5 +71,7 @@ private:
 	void GenerateScript();
 
 	void InitializeScript(UNiagaraScript* NewScript);
+
+	int32 CalculateScriptUsageBitmask();
 };
 
diff --git a/Engine/Plugins/FX/Niagara/Source/NiagaraEditor/Public/NiagaraParameterMapHistory.h b/Engine/Plugins/FX/Niagara/Source/NiagaraEditor/Public/NiagaraParameterMapHistory.h
--- a/Engine/Plugins/FX/Niagara/Source/NiagaraEditor/Public/NiagaraParameterMapHistory.h
+++ b/Engine/Plugins/FX/Niagara/Source/NiagaraEditor/Public/NiagaraParameterMapHistory.h
@@ -137,7 +137,7 @@ public:
 	static bool IsPerInstanceEngineParameter(const FNiagaraVariable& InVar, const FString& EmitterAlias);
 	static bool IsUserParameter(const FNiagaraVariable& InVar);
 	static bool IsRapidIterationParameter(const FNiagaraVariable& InVar);
-	static bool TryGetEmitterAndFunctionCallNamesFromRapidIterationParameter(const FNiagaraVariable& InVar, FString& EmitterName, FString& FunctionCallName);
+	static bool SplitRapidIterationParameterName(const FNiagaraVariable& InVar, FString& EmitterName, FString& FunctionCallName, FString& InputName);
 	
 	/** Take an input string and make it hlsl safe.*/
 	static FString MakeSafeNamespaceString(const FString& InStr);
diff --git a/Engine/Plugins/FX/Niagara/Source/NiagaraEditor/Public/ViewModels/Stack/NiagaraStackGraphUtilities.h b/Engine/Plugins/FX/Niagara/Source/NiagaraEditor/Public/ViewModels/Stack/NiagaraStackGraphUtilities.h
--- a/Engine/Plugins/FX/Niagara/Source/NiagaraEditor/Public/ViewModels/Stack/NiagaraStackGraphUtilities.h
+++ b/Engine/Plugins/FX/Niagara/Source/NiagaraEditor/Public/ViewModels/Stack/NiagaraStackGraphUtilities.h
@@ -154,4 +154,6 @@ namespace FNiagaraStackGraphUtilities
 		UEdGraphPin& DefaultPin);
 	
 	bool GetStackIssuesRecursively(const UNiagaraStackEntry* const Entry, TArray<UNiagaraStackErrorItem*>& OutIssues);
+
+	void MoveModule(UNiagaraScript& SourceScript, UNiagaraNodeFunctionCall& ModuleToMove, UNiagaraSystem& TargetSystem, FGuid TargetEmitterHandleId, ENiagaraScriptUsage TargetUsage, FGuid TargetUsageId, int32 TargetModuleIndex);
 }
