Improved functionality of module dependencies: intercategory module dependencies now work, module order is fixed.
#tests none

diff --git a/Engine/Plugins/FX/Niagara/Source/NiagaraEditor/Private/ViewModels/NiagaraSystemViewModel.cpp b/Engine/Plugins/FX/Niagara/Source/NiagaraEditor/Private/ViewModels/NiagaraSystemViewModel.cpp
--- a/Engine/Plugins/FX/Niagara/Source/NiagaraEditor/Private/ViewModels/NiagaraSystemViewModel.cpp
+++ b/Engine/Plugins/FX/Niagara/Source/NiagaraEditor/Private/ViewModels/NiagaraSystemViewModel.cpp
@@ -20,6 +20,7 @@
 #include "NiagaraGraph.h"
 #include "NiagaraNodeInput.h"
 #include "NiagaraEditorModule.h"
+#include "NiagaraNodeFunctionCall.h"
 
 #include "Editor.h"
 
@@ -125,9 +126,8 @@ void FNiagaraSystemViewModel::Cleanup()
 	}
 
 	PreviewComponent = nullptr;
-	SystemScriptViewModel.Reset();
-	
 	RemoveSystemEventHandlers();
+	SystemScriptViewModel.Reset();
 }
 
 FNiagaraSystemViewModel::~FNiagaraSystemViewModel()
@@ -573,6 +573,32 @@ void FNiagaraSystemViewModel::SetToolkitCommands(const TSharedRef<FUICommandList
 	ToolkitCommands = InToolkitCommands;
 }
 
+const TArray<FNiagaraStackModuleData>& FNiagaraSystemViewModel::GetStackModuleDataForEmitter(TSharedRef<FNiagaraEmitterViewModel> EmitterViewModel)
+{
+	FGuid EmitterHandleId;
+	TSharedRef<FNiagaraEmitterHandleViewModel>* FoundModel = EmitterHandleViewModels.FindByPredicate([EmitterViewModel](TSharedRef<FNiagaraEmitterHandleViewModel> CurrentViewModel)
+	{ return CurrentViewModel->GetEmitterViewModel() == EmitterViewModel; });
+	checkf(FoundModel != nullptr, TEXT("Couldn't get stack module data for emitter"));
+	if (FoundModel)
+	{
+		EmitterHandleId = (*FoundModel)->GetEmitterHandle()->GetId();
+		if (!EmitterToCachedStackModuleData.Contains(EmitterHandleId))
+		{
+			// If not cached, rebuild
+			UNiagaraEmitter* Emitter = (*FoundModel)->GetEmitterViewModel()->GetEmitter();
+			TArray<FNiagaraStackModuleData> StackModuleData;
+			BuildStackModuleData(GetSystem().GetSystemSpawnScript(), EmitterHandleId, StackModuleData);
+			BuildStackModuleData(GetSystem().GetSystemUpdateScript(), EmitterHandleId, StackModuleData);
+			BuildStackModuleData(Emitter->EmitterSpawnScriptProps.Script, EmitterHandleId, StackModuleData);
+			BuildStackModuleData(Emitter->EmitterUpdateScriptProps.Script, EmitterHandleId, StackModuleData);
+			BuildStackModuleData(Emitter->SpawnScriptProps.Script, EmitterHandleId, StackModuleData);
+			BuildStackModuleData(Emitter->UpdateScriptProps.Script, EmitterHandleId, StackModuleData);
+			EmitterToCachedStackModuleData.Add(EmitterHandleId) = StackModuleData;
+		}
+	}
+	return EmitterToCachedStackModuleData[EmitterHandleId];
+}
+
 bool FNiagaraSystemViewModel::IsTickable() const
 {
 	return bForceAutoCompileOnce || (GetDefault<UNiagaraEditorSettings>()->GetAutoCompile() && bCanAutoCompile);
@@ -1142,6 +1168,13 @@ void FNiagaraSystemViewModel::EmitterScriptGraphChanged(const FEdGraphEditAction
 		}
 		Sequencer->NotifyMovieSceneDataChanged(EMovieSceneDataChangeType::MovieSceneStructureItemsChanged);
 	}
+	// Remove from cache on graph change 
+	EmitterToCachedStackModuleData.Remove(OwningEmitterHandleViewModel->GetId());
+}
+
+void FNiagaraSystemViewModel::SystemScriptGraphChanged(const FEdGraphEditAction& InAction)
+{
+	EmitterToCachedStackModuleData.Empty();
 }
 
 void FNiagaraSystemViewModel::EmitterParameterStoreChanged(const FNiagaraParameterStore& ChangedParameterStore, const UNiagaraScript& OwningScript, const TSharedRef<FNiagaraEmitterHandleViewModel> OwningEmitterHandleViewModel)
@@ -1581,6 +1614,8 @@ void FNiagaraSystemViewModel::AddSystemEventHandlers()
 		UserParameterStoreChangedHandle = System.GetExposedParameters().AddOnChangedHandler(
 			FNiagaraParameterStore::FOnChanged::FDelegate::CreateRaw<FNiagaraSystemViewModel, const FNiagaraParameterStore&, const UNiagaraScript*>(
 				this, &FNiagaraSystemViewModel::SystemParameterStoreChanged, System.GetExposedParameters(), nullptr));
+
+		SystemScriptGraphChangedHandler = SystemScriptViewModel->GetGraphViewModel()->GetGraph()->AddOnGraphChangedHandler(FOnGraphChanged::FDelegate::CreateRaw(this, &FNiagaraSystemViewModel::SystemScriptGraphChanged));
 	}
 }
 
@@ -1602,6 +1637,10 @@ void FNiagaraSystemViewModel::RemoveSystemEventHandlers()
 		}
 
 		System.GetExposedParameters().RemoveOnChangedHandler(UserParameterStoreChangedHandle);
+		if (SystemScriptViewModel.IsValid())
+		{
+			SystemScriptViewModel->GetGraphViewModel()->GetGraph()->RemoveOnGraphChangedHandler(SystemScriptGraphChangedHandler);
+		}
 	}
 
 	ScriptToOnParameterStoreChangedHandleMap.Empty();
@@ -1613,4 +1652,33 @@ void FNiagaraSystemViewModel::NotifyPinnedCurvesChanged()
 	OnPinnedCurvesChangedDelegate.Broadcast();
 }
 
+void FNiagaraSystemViewModel::BuildStackModuleData(UNiagaraScript* Script, FGuid InEmitterHandleId, TArray<FNiagaraStackModuleData>& OutStackModuleData)
+{
+	UNiagaraNodeOutput* OutputNode = FNiagaraEditorUtilities::GetScriptOutputNode(*Script);
+	TArray<FNiagaraStackGraphUtilities::FStackNodeGroup> StackGroups;
+	FNiagaraStackGraphUtilities::GetStackNodeGroups((UNiagaraNode&)*OutputNode, StackGroups);
+
+	int StackIndex = 0;
+	if (StackGroups.Num() > 2)
+	{
+		for (int i = 1; i < StackGroups.Num() - 1; i++)
+		{
+			FNiagaraStackGraphUtilities::FStackNodeGroup& StackGroup = StackGroups[i];
+			StackIndex = i - 1;
+			TArray<UNiagaraNode*> GroupNodes;
+			StackGroup.GetAllNodesInGroup(GroupNodes);
+			UNiagaraNodeFunctionCall * ModuleNode = Cast<UNiagaraNodeFunctionCall>(StackGroup.EndNode);
+			if (ModuleNode != nullptr)
+			{
+				ENiagaraScriptUsage Usage = Script->GetUsage();
+				FGuid UsageId = Script->GetUsageId();
+				int32 Index = StackIndex;
+				FNiagaraStackModuleData ModuleData = { ModuleNode, Usage, UsageId, Index, InEmitterHandleId };
+				OutStackModuleData.Add(ModuleData);
+			}
+		}
+	}
+
+}
+
 #undef LOCTEXT_NAMESPACE // NiagaraSystemViewModel
\ No newline at end of file
diff --git a/Engine/Plugins/FX/Niagara/Source/NiagaraEditor/Private/ViewModels/Stack/NiagaraStackEntry.cpp b/Engine/Plugins/FX/Niagara/Source/NiagaraEditor/Private/ViewModels/Stack/NiagaraStackEntry.cpp
--- a/Engine/Plugins/FX/Niagara/Source/NiagaraEditor/Private/ViewModels/Stack/NiagaraStackEntry.cpp
+++ b/Engine/Plugins/FX/Niagara/Source/NiagaraEditor/Private/ViewModels/Stack/NiagaraStackEntry.cpp
@@ -267,6 +267,11 @@ UNiagaraStackEntry::FOnDataObjectModified& UNiagaraStackEntry::OnDataObjectModif
 	return DataObjectModifiedDelegate;
 }
 
+UNiagaraStackEntry::FOnRequestFullRefresh& UNiagaraStackEntry::OnRequestFullRefresh()
+{
+	return RequestFullRefreshDelegate;
+}
+
 int32 UNiagaraStackEntry::GetIndentLevel() const
 {
 	return IndentLevel;
@@ -364,6 +369,7 @@ void UNiagaraStackEntry::RefreshChildren()
 	{
 		Child->OnStructureChanged().RemoveAll(this);
 		Child->OnDataObjectModified().RemoveAll(this);
+		Child->OnRequestFullRefresh().RemoveAll(this);
 		Child->SetOnRequestCanDrop(FOnRequestDrop());
 		Child->SetOnRequestDrop(FOnRequestDrop());
 	}
@@ -371,6 +377,7 @@ void UNiagaraStackEntry::RefreshChildren()
 	{
 		ErrorChild->OnStructureChanged().RemoveAll(this);
 		ErrorChild->OnDataObjectModified().RemoveAll(this);
+		ErrorChild->OnRequestFullRefresh().RemoveAll(this);
 		ErrorChild->OnIssueModified().RemoveAll(this);
 	}
 
@@ -386,6 +393,7 @@ void UNiagaraStackEntry::RefreshChildren()
 		Child->RefreshChildren();
 		Child->OnStructureChanged().AddUObject(this, &UNiagaraStackEntry::ChildStructureChanged);
 		Child->OnDataObjectModified().AddUObject(this, &UNiagaraStackEntry::ChildDataObjectModified);
+		Child->OnRequestFullRefresh().AddUObject(this, &UNiagaraStackEntry::ChildRequestFullRefresh);
 		Child->SetOnRequestCanDrop(FOnRequestDrop::CreateUObject(this, &UNiagaraStackEntry::ChildRequestCanDrop));
 		Child->SetOnRequestDrop(FOnRequestDrop::CreateUObject(this, &UNiagaraStackEntry::ChildRequestDrop));
 	}
@@ -402,6 +410,7 @@ void UNiagaraStackEntry::RefreshChildren()
 		ErrorChild->RefreshChildren();
 		ErrorChild->OnStructureChanged().AddUObject(this, &UNiagaraStackEntry::ChildStructureChanged);
 		ErrorChild->OnDataObjectModified().AddUObject(this, &UNiagaraStackEntry::ChildDataObjectModified);
+		ErrorChild->OnRequestFullRefresh().AddUObject(this, &UNiagaraStackEntry::ChildRequestFullRefresh);
 		ErrorChild->OnIssueModified().AddUObject(this, &UNiagaraStackEntry::IssueModified);
 	}
 
@@ -468,6 +477,11 @@ void UNiagaraStackEntry::ChildDataObjectModified(UObject* ChangedObject)
 	DataObjectModifiedDelegate.Broadcast(ChangedObject);
 }
 
+void UNiagaraStackEntry::ChildRequestFullRefresh()
+{
+	RequestFullRefreshDelegate.Broadcast();
+}
+
 TOptional<UNiagaraStackEntry::FDropResult> UNiagaraStackEntry::ChildRequestCanDrop(const UNiagaraStackEntry& TargetChild, const TArray<UNiagaraStackEntry*>& DraggedEntries)
 {
 	TOptional<FDropResult> Result = ChildRequestCanDropInternal(TargetChild, DraggedEntries);
diff --git a/Engine/Plugins/FX/Niagara/Source/NiagaraEditor/Private/ViewModels/Stack/NiagaraStackModuleItem.cpp b/Engine/Plugins/FX/Niagara/Source/NiagaraEditor/Private/ViewModels/Stack/NiagaraStackModuleItem.cpp
--- a/Engine/Plugins/FX/Niagara/Source/NiagaraEditor/Private/ViewModels/Stack/NiagaraStackModuleItem.cpp
+++ b/Engine/Plugins/FX/Niagara/Source/NiagaraEditor/Private/ViewModels/Stack/NiagaraStackModuleItem.cpp
@@ -15,10 +15,13 @@
 #include "NiagaraEmitterEditorData.h"
 #include "NiagaraStackEditorData.h"
 #include "NiagaraEditorStyle.h"
+#include "NiagaraEditorUtilities.h"
 #include "NiagaraGraph.h"
 #include "NiagaraScript.h"
 #include "ViewModels/Stack/NiagaraStackGraphUtilities.h"
 #include "NiagaraEmitterHandle.h"
+#include "NiagaraEmitter.h"
+#include "NiagaraSystem.h"
 #include "NiagaraScriptMergeManager.h"
 #include "NiagaraStackEditorData.h"
 #include "NiagaraScriptSource.h"
@@ -30,6 +33,76 @@
 
 #define LOCTEXT_NAMESPACE "NiagaraStackModuleItem"
 
+TArray<ENiagaraScriptUsage> UsagePriority = { // Ordered such as the highest priority has the largest index
+	ENiagaraScriptUsage::ParticleUpdateScript,
+	ENiagaraScriptUsage::ParticleSpawnScript,
+	ENiagaraScriptUsage::EmitterUpdateScript,
+	ENiagaraScriptUsage::EmitterSpawnScript,
+	ENiagaraScriptUsage::SystemUpdateScript,
+	ENiagaraScriptUsage::SystemSpawnScript };
+
+UNiagaraNodeOutput* GetOutputNodeForModuleDependency(ENiagaraScriptUsage DependantUsage, UNiagaraScript* DependencyScript, UNiagaraSystem& System, UNiagaraEmitter* Emitter, FNiagaraModuleDependency Dependency)
+{
+	UNiagaraNodeOutput* TargetOutputNode = nullptr;
+	if (DependencyScript)
+	{
+		UNiagaraScript* OutputScript = nullptr;
+		TArray<ENiagaraScriptUsage> SupportedUsages = UNiagaraScript::GetSupportedUsageContextsForBitmask(DependencyScript->ModuleUsageBitmask);
+		ENiagaraScriptUsage ScriptUsage = SupportedUsages[0];
+		int32 ClosestDistance = MAX_int32;
+		for (ENiagaraScriptUsage PossibleUsage : SupportedUsages)
+		{
+			int32 PossibleIndex = UsagePriority.IndexOfByPredicate(
+				[&](const ENiagaraScriptUsage CurrentUsage)
+				{
+					return UNiagaraScript::IsEquivalentUsage(PossibleUsage, CurrentUsage);
+				});
+			int32 DependantIndex = UsagePriority.IndexOfByPredicate(
+				[&](const ENiagaraScriptUsage CurrentUsage)
+				{
+					return UNiagaraScript::IsEquivalentUsage(DependantUsage, CurrentUsage);
+				});
+
+			int32 Distance = PossibleIndex - DependantIndex;
+			bool bCorrectOrder = (Dependency.Type == ENiagaraModuleDependencyType::PreDependency && Distance >= 0) || (Dependency.Type == ENiagaraModuleDependencyType::PostDependency && Distance <= 0);
+			if ((FMath::Abs(Distance) < ClosestDistance) && bCorrectOrder)
+			{
+				ClosestDistance = Distance;
+				ScriptUsage = PossibleUsage;
+				if (UNiagaraScript::IsEquivalentUsage(ScriptUsage, ENiagaraScriptUsage::SystemSpawnScript))
+				{
+					OutputScript = System.GetSystemSpawnScript();
+				}
+				else if (UNiagaraScript::IsEquivalentUsage(ScriptUsage, ENiagaraScriptUsage::SystemUpdateScript))
+				{
+					OutputScript = System.GetSystemUpdateScript();
+				}
+				else if (UNiagaraScript::IsEquivalentUsage(ScriptUsage, ENiagaraScriptUsage::EmitterSpawnScript))
+				{
+					OutputScript = Emitter->EmitterSpawnScriptProps.Script;
+				}
+				else if (UNiagaraScript::IsEquivalentUsage(ScriptUsage, ENiagaraScriptUsage::EmitterUpdateScript))
+				{
+					OutputScript = Emitter->EmitterUpdateScriptProps.Script;
+				}
+				else if (UNiagaraScript::IsEquivalentUsage(ScriptUsage, ENiagaraScriptUsage::ParticleSpawnScript))
+				{
+					OutputScript = Emitter->SpawnScriptProps.Script;
+				}
+				else if (UNiagaraScript::IsEquivalentUsage(ScriptUsage, ENiagaraScriptUsage::ParticleUpdateScript))
+				{
+					OutputScript = Emitter->UpdateScriptProps.Script;
+				}
+			}
+		}
+		if (OutputScript != nullptr)
+		{
+			TargetOutputNode = FNiagaraEditorUtilities::GetScriptOutputNode(*OutputScript);
+		}
+	}
+	return TargetOutputNode;
+}
+
 UNiagaraStackModuleItem::UNiagaraStackModuleItem()
 	: FunctionCallNode(nullptr)
 	, bCanRefresh(false)
@@ -200,12 +273,23 @@ void UNiagaraStackModuleItem::RefreshIssues(TArray<FStackIssue>& NewIssues)
 			NewIssues.Add(InconsistentEnabledError);
 		}
 	}
-	/*// Generate dependency errors with their fixes
+	// Generate dependency errors with their fixes
 	TArray<UNiagaraNodeFunctionCall*> FoundCalls;
-	TArray<FNiagaraStackGraphUtilities::FStackNodeGroup> StackNodeGroups;
-	FNiagaraStackGraphUtilities::GetStackNodeGroups(*FunctionCallNode, StackNodeGroups);
 	TArray<FNiagaraModuleDependency> DependenciesNeeded;
 	UNiagaraNodeOutput* OutputNode = GetOutputNode();
+
+	TArray<FNiagaraStackModuleData> SystemModuleData = GetSystemViewModel()->GetStackModuleDataForEmitter(GetEmitterViewModel());
+	int32 ModuleIndex = INDEX_NONE;
+	for (int i = 0; i < SystemModuleData.Num(); i++)
+	{
+		auto ModuleData = SystemModuleData[i];
+		if (ModuleData.ModuleNode == FunctionCallNode)
+		{
+			ModuleIndex = i;
+			break;
+		}
+	}
+
 	for (FNiagaraModuleDependency Dependency: FunctionCallNode->FunctionScript->RequiredDependencies)
 	{
 		if (Dependency.Id == NAME_None)
@@ -215,14 +299,38 @@ void UNiagaraStackModuleItem::RefreshIssues(TArray<FStackIssue>& NewIssues)
 		bool bDependencyMet = false;
 		UNiagaraNodeFunctionCall* FunctionNode = nullptr;
 		TArray <UNiagaraNodeFunctionCall*> DisabledDependencies;
-		for (int i = 1; i < StackNodeGroups.Num() - 1; i++)
+		TArray <FNiagaraStackModuleData> DisorderedDependencies;
+		
+		int32 DependencyModuleIndex = INDEX_NONE;
+		for (FNiagaraStackModuleData ModuleData : SystemModuleData)
 		{
-			FNiagaraStackGraphUtilities::FStackNodeGroup& StackNodeGroup = StackNodeGroups[i];
-			FunctionNode = Cast<UNiagaraNodeFunctionCall> (StackNodeGroup.EndNode);
-			if (FunctionNode != nullptr && FunctionNode->FunctionScript->ProvidedDependencies.Contains(Dependency.Id))
-			// treat disabled dependencies separately because they are easier to fix
+			FunctionNode = ModuleData.ModuleNode;
+			DependencyModuleIndex++;
+			if (FunctionNode != nullptr && FunctionNode->FunctionScript->ProvidedDependencies.Contains(Dependency.Id)) 
 			{
-				if (FunctionNode->IsNodeEnabled() == false)
+				auto DependencyOutputUsage = ModuleData.Usage;
+				int32 PossibleIndex = UsagePriority.IndexOfByPredicate(
+					[&](const ENiagaraScriptUsage CurrentUsage)
+					{
+						return UNiagaraScript::IsEquivalentUsage(DependencyOutputUsage, CurrentUsage);
+					});
+				int32 DependantIndex = UsagePriority.IndexOfByPredicate(
+					[&](const ENiagaraScriptUsage CurrentUsage)
+					{
+						return UNiagaraScript::IsEquivalentUsage(OutputNode->GetUsage(), CurrentUsage);
+					});
+				int32 Distance = PossibleIndex - DependantIndex;
+				
+				bool bIncorrectOrder = Distance == 0 ?  ((Dependency.Type == ENiagaraModuleDependencyType::PreDependency && ModuleIndex < DependencyModuleIndex)
+															|| (Dependency.Type == ENiagaraModuleDependencyType::PostDependency && ModuleIndex > DependencyModuleIndex))
+													 : ((Dependency.Type == ENiagaraModuleDependencyType::PreDependency && Distance < 0)
+															|| (Dependency.Type == ENiagaraModuleDependencyType::PostDependency && Distance > 0));
+										 
+				if (bIncorrectOrder)
+				{
+					DisorderedDependencies.Add(ModuleData);
+				}
+				else if (FunctionNode->IsNodeEnabled() == false) 
 				{
 					DisabledDependencies.Add(FunctionNode);
 				}
@@ -236,54 +344,156 @@ void UNiagaraStackModuleItem::RefreshIssues(TArray<FStackIssue>& NewIssues)
 		if (bDependencyMet == false)
 		{
 			DependenciesNeeded.Add(Dependency);
-			UNiagaraStackEntry::FStackIssue Error;
-			Error.ShortDescription = LOCTEXT("DependencyWarning", "The module has unmet dependencies.");;
-			Error.LongDescription = FText::Format(LOCTEXT("DependencyWarningLong", "The following dependency is not met: {0}"), Dependency.Description);
-			Error.bCanBeDismissed = true;
-			Error.UniqueIdentifier = FName(*FString::Printf(TEXT("%s-dependency-%s"), *GetStackEditorDataKey(), *Dependency.Id.ToString()));
+			
+			FText DependencyTypeString = Dependency.Type == ENiagaraModuleDependencyType::PreDependency ? LOCTEXT("PreDependency", "pre-dependency") : LOCTEXT("PostDependency", "post-dependency");
+			UNiagaraStackEntry::FStackIssue Error(
+				EStackIssueSeverity::Error,
+				LOCTEXT("DependencyWarning", "The module has unmet dependencies."),
+				FText::Format(LOCTEXT("DependencyWarningLong", "The following {0} is not met: {1}"), DependencyTypeString, Dependency.Description),
+				FString::Printf(TEXT("%s-dependency-%s"), *GetStackEditorDataKey(), *Dependency.Id.ToString()),
+				true);
+			
 
 			for ( UNiagaraNodeFunctionCall* DisabledNode : DisabledDependencies) // module exists but disabled
 			{
-				UNiagaraStackEntry::FStackIssueFix Fix;
-				Fix.Description = FText::Format(LOCTEXT("EnableDependency", "Enable dependency module {0}"), FText::FromString(DisabledNode->GetFunctionName()));
-				Fix.FixDelegate.BindLambda([=]()
-				{
-					FScopedTransaction ScopedTransaction(LOCTEXT("EnableDependencyModule", "Enable dependency module"));
-					DisabledNode->Modify();
-					FNiagaraStackGraphUtilities::SetModuleIsEnabled(*DisabledNode, true);
-				});
-				Error.Fixes.Add(Fix);
+				UNiagaraStackEntry::FStackIssueFix Fix(
+					FText::Format(LOCTEXT("EnableDependency", "Enable dependency module {0}"), FText::FromString(DisabledNode->GetFunctionName())),
+					FStackIssueFixDelegate::CreateLambda([=]()
+					{
+						FScopedTransaction ScopedTransaction(LOCTEXT("EnableDependencyModule", "Enable dependency module"));
+						DisabledNode->Modify();
+						FNiagaraStackGraphUtilities::SetModuleIsEnabled(*DisabledNode, true);
+					}));
+				Error.AddFix(Fix);
 			}
 
-			TArray<FAssetData> ModuleAssets;
-			FNiagaraStackGraphUtilities::GetScriptAssetsByDependencyProvided(ENiagaraScriptUsage::Module, Dependency.Id, ModuleAssets);
-			for (FAssetData ModuleAsset : ModuleAssets) 
+			for (FNiagaraStackModuleData DisorderedNode : DisorderedDependencies) // module exists but is not in the correct order (and possibly also disabled)
 			{
-				auto DisabledAsset = DisabledDependencies.FindByPredicate([=](UNiagaraNodeFunctionCall* DisabledNode) {return DisabledNode->GetReferencedAsset() == ModuleAsset.GetAsset(); });
-				if (DisabledAsset != nullptr)
+				bool bNeedsEnable = !DisorderedNode.ModuleNode->IsNodeEnabled();
+				FText AndEnableModule = bNeedsEnable ? LOCTEXT("AndEnableModule", "and enable it") : FText();
+				UNiagaraStackEntry::FStackIssueFix Fix(
+					FText::Format(LOCTEXT("ReorderDependency", "Reposition dependency module {0} in the correct order {1}"), FText::FromString(DisorderedNode.ModuleNode->GetFunctionName()), AndEnableModule),
+					FStackIssueFixDelegate::CreateLambda([=]()
+					{
+						FScopedTransaction ScopedTransaction(LOCTEXT("ReorderDependencyModule", "Reorder dependency module"));
+						DisorderedNode.ModuleNode->Modify();
+						// reorder node
+						int32 CorrectIndex = Dependency.Type == ENiagaraModuleDependencyType::PostDependency ? DisorderedNode.Index : DisorderedNode.Index + 1;
+						checkf(ModuleIndex != INDEX_NONE, TEXT("Module data wasn't found in system for current module!"));
+						UNiagaraScript& OwningScript = *FNiagaraEditorUtilities::GetScriptFromSystem(GetSystemViewModel()->GetSystem(), SystemModuleData[ModuleIndex].EmitterHandleId, SystemModuleData[ModuleIndex].Usage, SystemModuleData[ModuleIndex].UsageId);
+						FNiagaraStackGraphUtilities::MoveModule(OwningScript, *FunctionCallNode, GetSystemViewModel()->GetSystem(), DisorderedNode.EmitterHandleId, DisorderedNode.Usage, DisorderedNode.UsageId, CorrectIndex);
+						// enable if needed
+						if (bNeedsEnable)
+						{
+							FNiagaraStackGraphUtilities::SetModuleIsEnabled(*DisorderedNode.ModuleNode, true);
+						}
+						FNiagaraStackGraphUtilities::RelayoutGraph(*OutputNode->GetGraph());
+						OnRequestFullRefresh().Broadcast();
+					}));
+				Error.AddFix(Fix);
+			}
+			if (DisorderedDependencies.Num() == 0 && DisabledDependencies.Num() == 0)
+			{
+				TArray<FAssetData> ModuleAssets;
+				FNiagaraStackGraphUtilities::GetScriptAssetsByDependencyProvided(ENiagaraScriptUsage::Module, Dependency.Id, ModuleAssets);
+				for (FAssetData ModuleAsset : ModuleAssets)
 				{
-					continue;
+					FText FixDescription = FText::Format(LOCTEXT("AddDependency", "Add new dependency module {0}"), FText::FromName(ModuleAsset.AssetName));
+					UNiagaraStackEntry::FStackIssueFix Fix(
+						FixDescription,
+						FStackIssueFixDelegate::CreateLambda([=]()
+					{
+						FScopedTransaction ScopedTransaction(FixDescription);
+						UNiagaraNodeFunctionCall* NewModuleNode = nullptr;
+						int32 TargetIndex = 0; 
+						UNiagaraScript* DependencyScript = Cast<UNiagaraScript>(ModuleAsset.GetAsset());
+						checkf(DependencyScript != nullptr, TEXT("Add module action failed"));
+						// Determine the output node for the group where the added dependency module belongs
+						UNiagaraNodeOutput* TargetOutputNode = nullptr;
+						for (int i = ModuleIndex; i < SystemModuleData.Num() && i >= 0; i = Dependency.Type == ENiagaraModuleDependencyType::PostDependency ? i + 1 : i - 1) // moving up or down depending on type
+						// starting at current module, which is a dependant
+						{
+							auto FoundRequirement = SystemModuleData[i].ModuleNode->FunctionScript->RequiredDependencies.FindByPredicate(
+									[&](FNiagaraModuleDependency CurrentDependency)
+									{
+										return CurrentDependency.Id == Dependency.Id;
+									});
+							if (FoundRequirement) // check for multiple dependendants along the way, and stop adjacent to the last one
+							{
+								ENiagaraScriptUsage DependencyUsage = SystemModuleData[i].Usage;
+								TargetOutputNode = GetOutputNodeForModuleDependency(DependencyUsage, DependencyScript, GetSystemViewModel()->GetSystem(), GetEmitterViewModel()->GetEmitter(), Dependency);
+								if (TargetOutputNode != nullptr)
+								{
+									auto CurrentOutputNode = FNiagaraStackGraphUtilities::GetEmitterOutputNodeForStackNode(*SystemModuleData[i].ModuleNode);
+									if (TargetOutputNode == CurrentOutputNode)
+									{
+										TargetIndex = Dependency.Type == ENiagaraModuleDependencyType::PostDependency ? SystemModuleData[i].Index + 1 : SystemModuleData[i].Index;
+									}
+									else
+									{
+										TargetIndex = Dependency.Type == ENiagaraModuleDependencyType::PostDependency ? 0 : INDEX_NONE;
+									}
+								}
+							}
+						}
+						if (TargetOutputNode == nullptr)
+						{
+							checkf(ModuleIndex != INDEX_NONE, TEXT("Module data wasn't found in system for current module!"));
+							// I am looking for a script in the same emitter as this module, so get it from there
+							UNiagaraScript& TargetScript = *FNiagaraEditorUtilities::GetScriptFromSystem(GetSystemViewModel()->GetSystem(), SystemModuleData[ModuleIndex].EmitterHandleId, DependencyScript->GetUsage(), DependencyScript->GetUsageId());
+							TargetOutputNode = FNiagaraEditorUtilities::GetScriptOutputNode(TargetScript);
+						}
+						TArray<FNiagaraStackModuleData> ScriptModuleData = SystemModuleData.FilterByPredicate([&](FNiagaraStackModuleData CurrentData) {return CurrentData.Usage == DependencyScript->GetUsage(); });
+						int32 PreIndex = INDEX_NONE; // index of last pre dependency
+						int32 PostIndex = INDEX_NONE; // index of fist post dependency, the module will have to be placed between these indexes
+						// for now, we skip the case where the dependencies are fulfilled in other script groups as well as here, because that's extremely unlikely
+						if (TargetIndex == INDEX_NONE)
+						{
+							TargetIndex = 0; //start at the beginning to look for potential dependencies of this dependency
+						}
+						for (int32 i = TargetIndex ; i < ScriptModuleData.Num() && i >= 0; i = Dependency.Type == ENiagaraModuleDependencyType::PostDependency ? i + 1 : i - 1)
+						{
+							UNiagaraNodeFunctionCall * CurrentNode =  ScriptModuleData[i].ModuleNode;
+							for (FNiagaraModuleDependency Requirement : DependencyScript->RequiredDependencies)
+							{
+								if (Requirement.Id == NAME_None)
+								{
+									continue;
+								}
+								bool bDependencyMet = false;
+								if (CurrentNode->FunctionScript->ProvidedDependencies.Contains(Requirement.Id))
+								{
+									if (Requirement.Type == ENiagaraModuleDependencyType::PreDependency)
+									{
+										PostIndex = i;
+									}
+									else if (PreIndex == INDEX_NONE) // only record the first post-dependency
+									{
+										PreIndex = i;
+									}
+								}
+							}
+						}
+						if (PostIndex != INDEX_NONE) 
+						{
+							TargetIndex = 0; // if it has post dependencies place it at the top
+							if (PreIndex != INDEX_NONE)
+							{
+								TargetIndex = PostIndex; // if it also has post dependencies just add it before its first post dependency
+							}
+						}
+						NewModuleNode = FNiagaraStackGraphUtilities::AddScriptModuleToStack(ModuleAsset, *TargetOutputNode, TargetIndex); // no target index for now, until fix
+						checkf(NewModuleNode != nullptr, TEXT("Add module action failed"));
+						FNiagaraStackGraphUtilities::InitializeStackFunctionInputs(GetSystemViewModel(), GetEmitterViewModel(), GetStackEditorData(), *NewModuleNode, *NewModuleNode);
+						FNiagaraStackGraphUtilities::RelayoutGraph(*TargetOutputNode->GetGraph());
+						OnRequestFullRefresh().Broadcast();
+					}));
+					Error.AddFix(Fix);
 				}
-				UNiagaraStackEntry::FStackIssueFix Fix;
-				Fix.Description = LOCTEXT("InsertNewModule", "Insert new module");
-				int32 InsertIndex = GetModuleIndex();
-				Fix.FixDelegate.BindLambda([=]()
-				{
-					FScopedTransaction ScopedTransaction(Fix.Description);
-					UNiagaraNodeFunctionCall* NewModuleNode = nullptr;
-					int32 TargetIndex = Dependency.Type == ENiagaraModuleDependencyType::PostDependency ? GetModuleIndex() + 1 : GetModuleIndex();
-					NewModuleNode = FNiagaraStackGraphUtilities::AddScriptModuleToStack(ModuleAsset, *OutputNode, TargetIndex);
-					checkf(NewModuleNode != nullptr, TEXT("Add module action failed"));
-					FNiagaraStackGraphUtilities::InitializeStackFunctionInputs(GetSystemViewModel(), GetEmitterViewModel(), GetStackEditorData(), *NewModuleNode, *NewModuleNode);
-					FNiagaraStackGraphUtilities::RelayoutGraph(*OutputNode->GetGraph());
-				});
-				Fix.Description = FText::Format(LOCTEXT("AddDependency", "Add dependency module {0}"), FText::FromName(ModuleAsset.AssetName));
-				Error.Fixes.Add(Fix);
 			}
-			
 			NewIssues.Add(Error);
 		}
-	}*/
+	}
 }
 
 bool UNiagaraStackModuleItem::FilterOutputCollection(const UNiagaraStackEntry& Child) const
diff --git a/Engine/Plugins/FX/Niagara/Source/NiagaraEditor/Private/ViewModels/Stack/NiagaraStackViewModel.cpp b/Engine/Plugins/FX/Niagara/Source/NiagaraEditor/Private/ViewModels/Stack/NiagaraStackViewModel.cpp
--- a/Engine/Plugins/FX/Niagara/Source/NiagaraEditor/Private/ViewModels/Stack/NiagaraStackViewModel.cpp
+++ b/Engine/Plugins/FX/Niagara/Source/NiagaraEditor/Private/ViewModels/Stack/NiagaraStackViewModel.cpp
@@ -35,6 +35,7 @@ void UNiagaraStackViewModel::Initialize(TSharedPtr<FNiagaraSystemViewModel> InSy
 	{
 		RootEntry->OnStructureChanged().RemoveAll(this);
 		RootEntry->OnDataObjectModified().RemoveAll(this);
+		RootEntry->OnRequestFullRefresh().RemoveAll(this);
 		RootEntries.Empty();
 		RootEntry = nullptr;
 		GEditor->UnregisterForUndo(this);
@@ -70,6 +71,7 @@ void UNiagaraStackViewModel::Initialize(TSharedPtr<FNiagaraSystemViewModel> InSy
 		RootEntry->RefreshChildren();
 		RootEntry->OnStructureChanged().AddUObject(this, &UNiagaraStackViewModel::EntryStructureChanged);
 		RootEntry->OnDataObjectModified().AddUObject(this, &UNiagaraStackViewModel::EntryDataObjectModified);
+		RootEntry->OnRequestFullRefresh().AddUObject(this, &UNiagaraStackViewModel::EntryRequestFullRefresh);
 		RootEntries.Add(RootEntry);
 	}
 	CurrentFocusedSearchMatchIndex = -1;
@@ -400,4 +402,10 @@ void UNiagaraStackViewModel::EntryDataObjectModified(UObject* ChangedObject)
 	OnSearchTextChanged(CurrentSearchText);
 }
 
+void UNiagaraStackViewModel::EntryRequestFullRefresh()
+{
+	checkf(RootEntry != nullptr, TEXT("Can not process full refresh when the root entry doesn't exist"));
+	RootEntry->RefreshChildren();
+}
+
 #undef LOCTEXT_NAMESPACE
diff --git a/Engine/Plugins/FX/Niagara/Source/NiagaraEditor/Public/ViewModels/NiagaraSystemViewModel.h b/Engine/Plugins/FX/Niagara/Source/NiagaraEditor/Public/ViewModels/NiagaraSystemViewModel.h
--- a/Engine/Plugins/FX/Niagara/Source/NiagaraEditor/Public/ViewModels/NiagaraSystemViewModel.h
+++ b/Engine/Plugins/FX/Niagara/Source/NiagaraEditor/Public/ViewModels/NiagaraSystemViewModel.h
@@ -30,6 +30,8 @@ struct FRichCurve;
 class UNiagaraEditorSettings;
 struct FNiagaraParameterStore;
 struct FEdGraphEditAction;
+class UNiagaraNodeFunctionCall;
+class FNiagaraEmitterViewModel;
 
 /** Defines different editing modes for this system view model. */
 enum class ENiagaraSystemViewModelEditMode
@@ -64,6 +66,15 @@ struct FNiagaraSystemViewModelOptions
 	ENiagaraSystemViewModelEditMode EditMode;
 };
 
+struct FNiagaraStackModuleData
+{
+	UNiagaraNodeFunctionCall* ModuleNode;
+	ENiagaraScriptUsage Usage;
+	FGuid UsageId;
+	int32 Index;
+	FGuid EmitterHandleId;
+};
+
 /** A view model for viewing and editing a UNiagaraSystem. */
 class FNiagaraSystemViewModel 
 	: public TSharedFromThis<FNiagaraSystemViewModel>
@@ -243,6 +254,9 @@ public:
 	/** Set the system toolkit command list. */
 	void SetToolkitCommands(const TSharedRef<FUICommandList>& InToolkitCommands);
 
+	/** Gets the stack module data for the provided emitter, for use in module dependencies. */
+	const TArray<FNiagaraStackModuleData>& GetStackModuleDataForEmitter(TSharedRef<FNiagaraEmitterViewModel> EmitterViewModel);
+
 private:
 
 	/** Sets up the preview component and System instance. */
@@ -291,6 +305,9 @@ private:
 	/** Called whenever an emitter's script graph changes. */
 	void EmitterScriptGraphChanged(const FEdGraphEditAction& InAction, const UNiagaraScript& OwningScript, const TSharedRef<FNiagaraEmitterHandleViewModel> OwningEmitterHandleViewModel);
 
+	/** Called whenever the system script graph changes. */
+	void SystemScriptGraphChanged(const FEdGraphEditAction& InAction);
+
 	/**
 	* Called whenever an emitter's parameter store owned by the system changes.
 	* @param ChangedParameterStore The parameter store that changed.
@@ -352,6 +369,9 @@ private:
 	/** sends a notification that pinned curves have changed status */
 	void NotifyPinnedCurvesChanged();
 
+	/** builds stack module data for use in module dependencies */
+	void BuildStackModuleData(UNiagaraScript* Script, FGuid InEmitterHandleId, TArray<FNiagaraStackModuleData>& OutStackModuleData);
+
 private:
 	/** The System being viewed and edited by this view model. */
 	UNiagaraSystem& System;
@@ -468,4 +488,10 @@ private:
 
 	/** A flag which indicates that a compile has been requested, but has not completed. */
 	bool bCompilePendingCompletion;
+
+	/** The cache of stack module data for each emitter */
+	TMap<FGuid, TArray<FNiagaraStackModuleData>> EmitterToCachedStackModuleData;
+	
+	/** A handle to the on graph changed delegate for the system script. */
+	FDelegateHandle SystemScriptGraphChangedHandler;
 };
\ No newline at end of file
diff --git a/Engine/Plugins/FX/Niagara/Source/NiagaraEditor/Public/ViewModels/Stack/NiagaraStackEntry.h b/Engine/Plugins/FX/Niagara/Source/NiagaraEditor/Public/ViewModels/Stack/NiagaraStackEntry.h
--- a/Engine/Plugins/FX/Niagara/Source/NiagaraEditor/Public/ViewModels/Stack/NiagaraStackEntry.h
+++ b/Engine/Plugins/FX/Niagara/Source/NiagaraEditor/Public/ViewModels/Stack/NiagaraStackEntry.h
@@ -38,6 +38,7 @@ public:
 
 	DECLARE_MULTICAST_DELEGATE(FOnStructureChanged);
 	DECLARE_MULTICAST_DELEGATE_OneParam(FOnDataObjectModified, UObject*);
+	DECLARE_MULTICAST_DELEGATE(FOnRequestFullRefresh);
 	DECLARE_DELEGATE_RetVal_TwoParams(TOptional<FDropResult>, FOnRequestDrop, const UNiagaraStackEntry&, const TArray<UNiagaraStackEntry*>&);
 	DECLARE_DELEGATE_RetVal_OneParam(bool, FOnFilterChild, const UNiagaraStackEntry&);
 	DECLARE_DELEGATE(FStackIssueFixDelegate);
@@ -201,6 +202,8 @@ public:
 
 	FOnDataObjectModified& OnDataObjectModified();
 
+	FOnRequestFullRefresh& OnRequestFullRefresh();
+
 	void RefreshChildren();
 
 
@@ -272,6 +275,8 @@ private:
 	
 	void ChildDataObjectModified(UObject* ChangedObject);
 
+	void ChildRequestFullRefresh();
+
 	TOptional<FDropResult> ChildRequestCanDrop(const UNiagaraStackEntry& TargetChild, const TArray<UNiagaraStackEntry*>& DraggedEntries);
 
 	TOptional<FDropResult> ChildRequestDrop(const UNiagaraStackEntry& TargetChild, const TArray<UNiagaraStackEntry*>& DraggedEntries);
@@ -292,6 +297,8 @@ private:
 
 	FOnDataObjectModified DataObjectModifiedDelegate;
 
+	FOnRequestFullRefresh RequestFullRefreshDelegate;
+
 	TArray<FOnFilterChild> ChildFilters;
 
 	UPROPERTY()
diff --git a/Engine/Plugins/FX/Niagara/Source/NiagaraEditor/Public/ViewModels/Stack/NiagaraStackModuleItem.h b/Engine/Plugins/FX/Niagara/Source/NiagaraEditor/Public/ViewModels/Stack/NiagaraStackModuleItem.h
--- a/Engine/Plugins/FX/Niagara/Source/NiagaraEditor/Public/ViewModels/Stack/NiagaraStackModuleItem.h
+++ b/Engine/Plugins/FX/Niagara/Source/NiagaraEditor/Public/ViewModels/Stack/NiagaraStackModuleItem.h
@@ -60,7 +60,6 @@ private:
 
 private:
 	void RefreshIsEnabled();
-	void InputPinnedChanged();
 
 private:
 	UNiagaraNodeFunctionCall* FunctionCallNode;
diff --git a/Engine/Plugins/FX/Niagara/Source/NiagaraEditor/Public/ViewModels/Stack/NiagaraStackViewModel.h b/Engine/Plugins/FX/Niagara/Source/NiagaraEditor/Public/ViewModels/Stack/NiagaraStackViewModel.h
--- a/Engine/Plugins/FX/Niagara/Source/NiagaraEditor/Public/ViewModels/Stack/NiagaraStackViewModel.h
+++ b/Engine/Plugins/FX/Niagara/Source/NiagaraEditor/Public/ViewModels/Stack/NiagaraStackViewModel.h
@@ -98,6 +98,7 @@ private:
 private:
 	void EntryStructureChanged();
 	void EntryDataObjectModified(UObject* ChangedObject);
+	void EntryRequestFullRefresh();
 	void OnSystemCompiled();
 	void OnEmitterCompiled();
 	/** Called by the tick function to perform partial search */
diff --git a/Engine/Plugins/FX/Niagara/Source/NiagaraEditorWidgets/Private/Stack/SNiagaraStackModuleItem.h b/Engine/Plugins/FX/Niagara/Source/NiagaraEditorWidgets/Private/Stack/SNiagaraStackModuleItem.h
--- a/Engine/Plugins/FX/Niagara/Source/NiagaraEditorWidgets/Private/Stack/SNiagaraStackModuleItem.h
+++ b/Engine/Plugins/FX/Niagara/Source/NiagaraEditorWidgets/Private/Stack/SNiagaraStackModuleItem.h
@@ -37,10 +37,6 @@ private:
 
 	EVisibility GetRefreshVisibility() const;
 
-	FReply MoveUpClicked();
-
-	FReply MoveDownClicked();
-
 	FReply DeleteClicked();
 	
 	TSharedRef<SWidget> RaiseActionMenuClicked();
