Validating modules reads and writes.
You cannot read/write from particles namespace in system and emitter scripts
You cannot write to user or NPC namespaces ever
You cannot write to system/emitter namespaces in particle scripts

#tests auto-tests pass

diff --git a/Engine/Plugins/FX/Niagara/Source/NiagaraCore/Private/NiagaraCustomVersion.cpp b/Engine/Plugins/FX/Niagara/Source/NiagaraCore/Private/NiagaraCustomVersion.cpp
--- a/Engine/Plugins/FX/Niagara/Source/NiagaraCore/Private/NiagaraCustomVersion.cpp
+++ b/Engine/Plugins/FX/Niagara/Source/NiagaraCore/Private/NiagaraCustomVersion.cpp
@@ -8,4 +8,4 @@ const FGuid FNiagaraCustomVersion::GUID(0xFCF57AFA, 0x50764283, 0xB9A9E658, 0xFF
 // Register the custom version with core
 FCustomVersionRegistration GRegisterNiagaraCustomVersion(FNiagaraCustomVersion::GUID, FNiagaraCustomVersion::LatestVersion, TEXT("NiagaraVer"));
 
-const FGuid FNiagaraCustomVersion::LatestScriptCompileVersion(0x0000B90F, 0xD8A4445C, 0x9E951397, 0xC5499F26);
\ No newline at end of file
+const FGuid FNiagaraCustomVersion::LatestScriptCompileVersion(0x338DE03B, 0x967544E2, 0xA906FC75, 0x60EC8661 );
\ No newline at end of file
diff --git a/Engine/Plugins/FX/Niagara/Source/NiagaraEditor/Private/NiagaraHlslTranslator.cpp b/Engine/Plugins/FX/Niagara/Source/NiagaraEditor/Private/NiagaraHlslTranslator.cpp
--- a/Engine/Plugins/FX/Niagara/Source/NiagaraEditor/Private/NiagaraHlslTranslator.cpp
+++ b/Engine/Plugins/FX/Niagara/Source/NiagaraEditor/Private/NiagaraHlslTranslator.cpp
@@ -3255,6 +3255,13 @@ void FHlslNiagaraTranslator::HandleParameterRead(int32 ParamMapHistoryIdx, const
 	bool bIsCandidateForRapidIteration = false;
 	const UEdGraphPin* InputPin = DefaultPin;
 
+	FString Namespace = FNiagaraParameterMapHistory::GetNamespace(Var);
+	if (!ParamMapHistories[ParamMapHistoryIdx].IsValidNamespaceForReading(CompileOptions.TargetUsage, CompileOptions.TargetUsageBitmask, Namespace))
+	{
+		Error(FText::Format(LOCTEXT("InvalidReadingNamespace", "Variable {0} is in a namespace that isn't valid for reading"), FText::FromName(Var.GetName())), ErrorNode, nullptr);
+		return;
+	}
+
 	if (FNiagaraParameterMapHistory::IsExternalConstantNamespace(Var, CompileOptions.TargetUsage, CompileOptions.GetTargetUsageBitmask()))
 	{
 		if (ParameterMapRegisterExternalConstantNamespaceVariable(Var, ErrorNode, ParamMapHistoryIdx, OutputChunkId, DefaultPin))
@@ -3978,7 +3985,7 @@ void FHlslNiagaraTranslator::HandleCustomHlslNode(UNiagaraNodeCustomHlsl* Custom
 		// Clean up any namespaced variables in the token list if they are aliased or promote any tokens that are namespaced to the parent 
 		// parameter map.
 		TArray<FString> PossibleNamespaces;
-		FNiagaraParameterMapHistory::GetValidNamespacesForReading(CompileOptions.TargetUsage, PossibleNamespaces);
+		FNiagaraParameterMapHistory::GetValidNamespacesForReading(CompileOptions.TargetUsage, 0, PossibleNamespaces);
 
 		for (FNiagaraParameterMapHistory& History : ParamMapHistories)
 		{
diff --git a/Engine/Plugins/FX/Niagara/Source/NiagaraEditor/Private/NiagaraNodeCustomHlsl.cpp b/Engine/Plugins/FX/Niagara/Source/NiagaraEditor/Private/NiagaraNodeCustomHlsl.cpp
--- a/Engine/Plugins/FX/Niagara/Source/NiagaraEditor/Private/NiagaraNodeCustomHlsl.cpp
+++ b/Engine/Plugins/FX/Niagara/Source/NiagaraEditor/Private/NiagaraNodeCustomHlsl.cpp
@@ -261,7 +261,7 @@ void UNiagaraNodeCustomHlsl::BuildParameterMapHistory(FNiagaraParameterMapHistor
 		}
 
 		TArray<FString> PossibleNamespaces;
-		FNiagaraParameterMapHistory::GetValidNamespacesForReading(OutHistory.GetBaseUsageContext(), PossibleNamespaces);
+		FNiagaraParameterMapHistory::GetValidNamespacesForReading(OutHistory.GetBaseUsageContext(), 0, PossibleNamespaces);
 
 		if ((bHasParamMapOutput || bHasParamMapInput) && ParamMapIdx != INDEX_NONE)
 		{
diff --git a/Engine/Plugins/FX/Niagara/Source/NiagaraEditor/Private/NiagaraParameterMapHistory.cpp b/Engine/Plugins/FX/Niagara/Source/NiagaraEditor/Private/NiagaraParameterMapHistory.cpp
--- a/Engine/Plugins/FX/Niagara/Source/NiagaraEditor/Private/NiagaraParameterMapHistory.cpp
+++ b/Engine/Plugins/FX/Niagara/Source/NiagaraEditor/Private/NiagaraParameterMapHistory.cpp
@@ -45,20 +45,32 @@ FNiagaraParameterMapHistory::FNiagaraParameterMapHistory()
 
 void FNiagaraParameterMapHistory::GetValidNamespacesForReading(const UNiagaraScript* InScript, TArray<FString>& OutputNamespaces)
 {
-	GetValidNamespacesForReading(InScript->GetUsage(), OutputNamespaces);
+	GetValidNamespacesForReading(InScript->GetUsage(), 0, OutputNamespaces);
 }
 
-void FNiagaraParameterMapHistory::GetValidNamespacesForReading(ENiagaraScriptUsage InScriptUsage, TArray<FString>& OutputNamespaces)
+void FNiagaraParameterMapHistory::GetValidNamespacesForReading(ENiagaraScriptUsage InScriptUsage, int32 InUsageBitmask, TArray<FString>& OutputNamespaces)
 {
+	TArray<ENiagaraScriptUsage> SupportedContexts;
+	SupportedContexts.Add(InScriptUsage);
+	if (UNiagaraScript::IsStandaloneScript(InScriptUsage))
+	{
+		SupportedContexts.Append(UNiagaraScript::GetSupportedUsageContextsForBitmask(InUsageBitmask));
+	}
+
 	OutputNamespaces.Add(PARAM_MAP_MODULE_STR);
 	OutputNamespaces.Add(PARAM_MAP_ENGINE_STR);
 	OutputNamespaces.Add(PARAM_MAP_NPC_STR);
 	OutputNamespaces.Add(PARAM_MAP_USER_STR);
 	OutputNamespaces.Add(PARAM_MAP_SYSTEM_STR);
 	OutputNamespaces.Add(PARAM_MAP_EMITTER_STR);
-	if (UNiagaraScript::IsParticleScript(InScriptUsage))
+
+	for (ENiagaraScriptUsage Usage : SupportedContexts)
 	{
-		OutputNamespaces.Add(PARAM_MAP_ATTRIBUTE_STR);
+		if (UNiagaraScript::IsParticleScript(Usage))
+		{
+			OutputNamespaces.Add(PARAM_MAP_ATTRIBUTE_STR);
+			break;
+		}
 	}
 }
 
@@ -79,16 +91,26 @@ FString FNiagaraParameterMapHistory::GetNamespace(const FNiagaraVariable& InVar,
 	}
 }
 
-bool FNiagaraParameterMapHistory::IsValidNamespaceForReading(ENiagaraScriptUsage InScriptUsage, FString Namespace)
+bool FNiagaraParameterMapHistory::IsValidNamespaceForReading(ENiagaraScriptUsage InScriptUsage, int32 InUsageBitmask, FString Namespace)
 {
 	TArray<FString> OutputNamespaces;
-	GetValidNamespacesForReading(InScriptUsage, OutputNamespaces);
+	GetValidNamespacesForReading(InScriptUsage, InUsageBitmask, OutputNamespaces);
+
+	TArray<FString> ConcernedNamespaces;
+	ConcernedNamespaces.Add(PARAM_MAP_MODULE_STR);
+	ConcernedNamespaces.Add(PARAM_MAP_ENGINE_STR);
+	ConcernedNamespaces.Add(PARAM_MAP_NPC_STR);
+	ConcernedNamespaces.Add(PARAM_MAP_USER_STR);
+	ConcernedNamespaces.Add(PARAM_MAP_SYSTEM_STR);
+	ConcernedNamespaces.Add(PARAM_MAP_EMITTER_STR);
+	ConcernedNamespaces.Add(PARAM_MAP_ATTRIBUTE_STR);
 	
 	if (!Namespace.EndsWith(TEXT(".")))
 	{
 		Namespace.Append(TEXT("."));
 	}
 
+	// Pass if we are in the allowed list
 	for (const FString& ValidNamespace : OutputNamespaces)
 	{
 		if (Namespace.StartsWith(ValidNamespace))
@@ -97,7 +119,17 @@ bool FNiagaraParameterMapHistory::IsValidNamespaceForReading(ENiagaraScriptUsage
 		}
 	}
 
-	return false;
+	// Only fail if we're using a namespace that we know is one of the reserved ones.
+	for (const FString& ConcernedNamespace : ConcernedNamespaces)
+	{
+		if (Namespace.StartsWith(ConcernedNamespace))
+		{
+			return false;
+		}
+	}
+
+	// This means that we are using a namespace that isn't one of the primary engine namespaces, so we don't care and let it go.
+	return true;
 }
 
 
diff --git a/Engine/Plugins/FX/Niagara/Source/NiagaraEditor/Public/NiagaraParameterMapHistory.h b/Engine/Plugins/FX/Niagara/Source/NiagaraEditor/Public/NiagaraParameterMapHistory.h
--- a/Engine/Plugins/FX/Niagara/Source/NiagaraEditor/Public/NiagaraParameterMapHistory.h
+++ b/Engine/Plugins/FX/Niagara/Source/NiagaraEditor/Public/NiagaraParameterMapHistory.h
@@ -152,8 +152,8 @@ public:
 	static bool IsEmitterNamespaceReadOnly(const UNiagaraScript* InScript);
 
 	static void GetValidNamespacesForReading(const UNiagaraScript* InScript, TArray<FString>& OutputNamespaces);
-	static void GetValidNamespacesForReading(ENiagaraScriptUsage InScriptUsage, TArray<FString>& OutputNamespaces);
-	static bool IsValidNamespaceForReading(ENiagaraScriptUsage InScriptUsage, FString Namespace);
+	static void GetValidNamespacesForReading(ENiagaraScriptUsage InScriptUsage, int32 InUsageBitmask, TArray<FString>& OutputNamespaces);
+	static bool IsValidNamespaceForReading(ENiagaraScriptUsage InScriptUsage, int32 InUsageBitmask, FString Namespace);
 
 	/** Called to determine if a given variable should be output from a script. It is not static as it requires the overall context to include emitter namespaces visited for system scripts.*/
 	bool IsPrimaryDataSetOutput(const FNiagaraVariable& InVar, const UNiagaraScript* InScript, bool bAllowDataInterfaces = false) const;
