Niagara - When creating the render state in the niagara component, also send the dynamic data the same frame since the emitter might not actually tick the next frame.

#tests engine tests.

diff --git a/Engine/Plugins/FX/Niagara/Source/Niagara/Private/NiagaraComponent.cpp b/Engine/Plugins/FX/Niagara/Source/Niagara/Private/NiagaraComponent.cpp
--- a/Engine/Plugins/FX/Niagara/Source/Niagara/Private/NiagaraComponent.cpp
+++ b/Engine/Plugins/FX/Niagara/Source/Niagara/Private/NiagaraComponent.cpp
@@ -747,6 +747,14 @@ TSharedPtr<FNiagaraSystemSimulation, ESPMode::ThreadSafe> UNiagaraComponent::Get
 	return nullptr;
 }
 
+void UNiagaraComponent::CreateRenderState_Concurrent()
+{
+	Super::CreateRenderState_Concurrent();
+	// The emitter instance may not tick again next frame so we send the dynamic data here so that the current state
+	// renders.  This can happen when while editing, or any time the age update mode is set to desired age.
+	SendRenderDynamicData_Concurrent();
+}
+
 void UNiagaraComponent::SendRenderDynamicData_Concurrent()
 {
 	SCOPE_CYCLE_COUNTER(STAT_NiagaraComponentSendRenderData);
diff --git a/Engine/Plugins/FX/Niagara/Source/Niagara/Public/NiagaraComponent.h b/Engine/Plugins/FX/Niagara/Source/Niagara/Public/NiagaraComponent.h
--- a/Engine/Plugins/FX/Niagara/Source/Niagara/Public/NiagaraComponent.h
+++ b/Engine/Plugins/FX/Niagara/Source/Niagara/Public/NiagaraComponent.h
@@ -94,6 +94,7 @@ private:
 protected:
 	virtual void OnRegister() override;
 	virtual void OnUnregister() override;
+	virtual void CreateRenderState_Concurrent() override;
 	virtual void SendRenderDynamicData_Concurrent() override;
 	virtual void BeginDestroy() override;
 	//virtual void OnAttachmentChanged() override;
