Fix assert attempting to add a Niagara emitter parameter to a system before tracking an emitter.
Also passes all graphs to the add button, to avoid any use of Graphs[0] in SNiagaraParameterMapView.

diff --git a/Engine/Plugins/FX/Niagara/Source/NiagaraEditor/Private/Widgets/SNiagaraGraphPinAdd.cpp b/Engine/Plugins/FX/Niagara/Source/NiagaraEditor/Private/Widgets/SNiagaraGraphPinAdd.cpp
--- a/Engine/Plugins/FX/Niagara/Source/NiagaraEditor/Private/Widgets/SNiagaraGraphPinAdd.cpp
+++ b/Engine/Plugins/FX/Niagara/Source/NiagaraEditor/Private/Widgets/SNiagaraGraphPinAdd.cpp
@@ -64,7 +64,9 @@ TSharedRef<SWidget> SNiagaraGraphPinAdd::OnGetAddButtonMenuContent()
 {
 	if (OwningNode != nullptr)
 	{
-		TSharedRef<SNiagaraAddParameterMenu> MenuWidget = SNew(SNiagaraAddParameterMenu, OwningNode->GetNiagaraGraph())
+		TArray<TWeakObjectPtr<UNiagaraGraph>> Graphs;
+		Graphs.Add(OwningNode->GetNiagaraGraph());
+		TSharedRef<SNiagaraAddParameterMenu> MenuWidget = SNew(SNiagaraAddParameterMenu, Graphs)
 			.OnAddParameter_UObject(OwningNode, &UNiagaraNodeWithDynamicPins::AddParameter, GetPinObj()) // For non custom actions
 			.OnCollectCustomActions_UObject(OwningNode, &UNiagaraNodeWithDynamicPins::CollectAddPinActions, GetPinObj())
 			.OnAllowMakeType_UObject(OwningNode, &UNiagaraNodeWithDynamicPins::AllowNiagaraTypeForAddPin);
diff --git a/Engine/Plugins/FX/Niagara/Source/NiagaraEditor/Private/Widgets/SNiagaraParameterMapView.cpp b/Engine/Plugins/FX/Niagara/Source/NiagaraEditor/Private/Widgets/SNiagaraParameterMapView.cpp
--- a/Engine/Plugins/FX/Niagara/Source/NiagaraEditor/Private/Widgets/SNiagaraParameterMapView.cpp
+++ b/Engine/Plugins/FX/Niagara/Source/NiagaraEditor/Private/Widgets/SNiagaraParameterMapView.cpp
@@ -226,9 +226,9 @@ void SNiagaraParameterMapView::Tick(const FGeometry& AllottedGeometry, const dou
 	}
 }
 
-bool SNiagaraParameterMapView::IsEditingMode() const
+bool SNiagaraParameterMapView::ParameterAddEnabled() const
 {
-	return true;
+	return Graphs.Num() > 0;
 }
 
 void SNiagaraParameterMapView::AddParameter(FNiagaraVariable NewVariable)
@@ -447,7 +447,7 @@ TSharedRef<SWidget> SNiagaraParameterMapView::CreateAddToSectionButton(const Nia
 	.ForegroundColor(FEditorStyle::GetSlateColor("DefaultForeground"))
 	.ContentPadding(FMargin(2, 0))
 	.OnGetMenuContent(this, &SNiagaraParameterMapView::OnGetParameterMenu, InSection)
-	.IsEnabled(this, &SNiagaraParameterMapView::IsEditingMode)
+	.IsEnabled(this, &SNiagaraParameterMapView::ParameterAddEnabled)
 	.HAlign(HAlign_Center)
 	.VAlign(VAlign_Center)
 	.HasDownArrow(false)
@@ -489,7 +489,7 @@ bool SNiagaraParameterMapView::SelectionHasContextMenu() const
 
 TSharedRef<SWidget> SNiagaraParameterMapView::OnGetParameterMenu(const NiagaraParameterMapSectionID::Type InSection)
 {
-	TSharedRef<SNiagaraAddParameterMenu> MenuWidget = SNew(SNiagaraAddParameterMenu, Graphs[0].Get())
+	TSharedRef<SNiagaraAddParameterMenu> MenuWidget = SNew(SNiagaraAddParameterMenu, Graphs)
 		.OnAddParameter(this, &SNiagaraParameterMapView::AddParameter)
 		.Section(InSection)
 		.ShowNamespaceCategory(false)
@@ -709,7 +709,7 @@ void SNiagaraParameterMapView::RefreshActions()
 /************************************************************************/
 /* SNiagaraAddParameterMenu                                             */
 /************************************************************************/
-void SNiagaraAddParameterMenu::Construct(const FArguments& InArgs, UNiagaraGraph* InGraph)
+void SNiagaraAddParameterMenu::Construct(const FArguments& InArgs, TArray<TWeakObjectPtr<UNiagaraGraph>> InGraphs)
 {
 	this->OnAddParameter = InArgs._OnAddParameter;
 	this->OnCollectCustomActions = InArgs._OnCollectCustomActions;
@@ -720,7 +720,7 @@ void SNiagaraAddParameterMenu::Construct(const FArguments& InArgs, UNiagaraGraph
 	this->ShowGraphParameters = InArgs._ShowGraphParameters;
 	this->AutoExpandMenu = InArgs._AutoExpandMenu;
 
-	Graph = InGraph;
+	Graphs = InGraphs;
 
 	ChildSlot
 	[
@@ -786,8 +786,9 @@ void SNiagaraAddParameterMenu::CollectAllActions(FGraphActionListBuilderBase& Ou
 	if (CanCollectSection(NiagaraParameterMapSectionID::PARTICLE))
 	{
 		bool bSupportsAttributes = true;
-		if (Graph.Get())
+		for (TWeakObjectPtr<UNiagaraGraph>& GraphWeakPtr : Graphs)
 		{
+			UNiagaraGraph* Graph = GraphWeakPtr.Get();
 			bool IsModule = Graph->FindOutputNode(ENiagaraScriptUsage::Module) != nullptr || Graph->FindOutputNode(ENiagaraScriptUsage::DynamicInput) != nullptr
 				|| Graph->FindOutputNode(ENiagaraScriptUsage::Function) != nullptr;
 
@@ -867,49 +868,52 @@ void SNiagaraAddParameterMenu::CollectAllActions(FGraphActionListBuilderBase& Ou
 	}
 
 	// Collect parameter actions
-	if (ShowGraphParameters.Get() && Graph.Get())
+	if (ShowGraphParameters.Get())
 	{
-		TMap<FNiagaraVariable, FNiagaraGraphParameterReferenceCollection> ParameterEntries = Graph.Get()->GetParameterMap();
-		ParameterEntries.KeySort([](const FNiagaraVariable& A, const FNiagaraVariable& B) { return (A.GetName() < B.GetName()); });
-
-		for (const auto& ParameterEntry : ParameterEntries)
+		for (TWeakObjectPtr<UNiagaraGraph>& Graph : Graphs)
 		{
-			const FNiagaraVariable& Parameter = ParameterEntry.Key;
-			FNiagaraParameterHandle Handle;
-			const NiagaraParameterMapSectionID::Type ParameterSectionID = NiagaraParameterMapSectionID::OnGetSectionFromVariable(Parameter, Handle, /*Default*/ NiagaraParameterMapSectionID::OTHER);
-			if (CanCollectSection(ParameterSectionID))
-			{
-				const FText Category = ShowNamespaceCategory.Get() ? NiagaraParameterMapSectionID::OnGetSectionTitle(ParameterSectionID) : FText::GetEmpty();
-				const FText DisplayName = FText::FromName(Parameter.GetName());
+			TMap<FNiagaraVariable, FNiagaraGraphParameterReferenceCollection> ParameterEntries = Graph.Get()->GetParameterMap();
+			ParameterEntries.KeySort([](const FNiagaraVariable& A, const FNiagaraVariable& B) { return (A.GetName() < B.GetName()); });
 
-				// Only add this action if it isn't already in the list.
-				bool bUnique = true;
-				for (int32 Index = 0; Index < OutAllActions.GetNumActions(); Index++)
+			for (const auto& ParameterEntry : ParameterEntries)
+			{
+				const FNiagaraVariable& Parameter = ParameterEntry.Key;
+				FNiagaraParameterHandle Handle;
+				const NiagaraParameterMapSectionID::Type ParameterSectionID = NiagaraParameterMapSectionID::OnGetSectionFromVariable(Parameter, Handle, /*Default*/ NiagaraParameterMapSectionID::OTHER);
+				if (CanCollectSection(ParameterSectionID))
 				{
-					const FGraphActionListBuilderBase::ActionGroup& ActionGroup = OutAllActions.GetAction(Index);
-					for (const TSharedPtr<FEdGraphSchemaAction>& SchemaAction : ActionGroup.Actions)
+					const FText Category = ShowNamespaceCategory.Get() ? NiagaraParameterMapSectionID::OnGetSectionTitle(ParameterSectionID) : FText::GetEmpty();
+					const FText DisplayName = FText::FromName(Parameter.GetName());
+
+					// Only add this action if it isn't already in the list.
+					bool bUnique = true;
+					for (int32 Index = 0; Index < OutAllActions.GetNumActions(); Index++)
 					{
-						if (SchemaAction->GetMenuDescription().EqualTo(DisplayName))
+						const FGraphActionListBuilderBase::ActionGroup& ActionGroup = OutAllActions.GetAction(Index);
+						for (const TSharedPtr<FEdGraphSchemaAction>& SchemaAction : ActionGroup.Actions)
+						{
+							if (SchemaAction->GetMenuDescription().EqualTo(DisplayName))
+							{
+								bUnique = false;
+								break;
+							}
+						}
+
+						if (!bUnique)
 						{
-							bUnique = false;
 							break;
 						}
 					}
 
-					if (!bUnique)
+					if (bUnique)
 					{
-						break;
-					}
-				}
+						const FText Tooltip = FText::GetEmpty();
+						TSharedPtr<FNiagaraMenuAction> Action(new FNiagaraMenuAction(
+							Category, DisplayName, Tooltip, 0, FText::GetEmpty(),
+							FNiagaraMenuAction::FOnExecuteStackAction::CreateSP(this, &SNiagaraAddParameterMenu::AddParameterSelected, Parameter, false, ParameterSectionID)));
 
-				if (bUnique)
-				{
-					const FText Tooltip = FText::GetEmpty();
-					TSharedPtr<FNiagaraMenuAction> Action(new FNiagaraMenuAction(
-						Category, DisplayName, Tooltip, 0, FText::GetEmpty(),
-						FNiagaraMenuAction::FOnExecuteStackAction::CreateSP(this, &SNiagaraAddParameterMenu::AddParameterSelected, Parameter, false, ParameterSectionID)));
-
-					OutAllActions.AddAction(Action);
+						OutAllActions.AddAction(Action);
+					}
 				}
 			}
 		}
diff --git a/Engine/Plugins/FX/Niagara/Source/NiagaraEditor/Private/Widgets/SNiagaraParameterMapView.h b/Engine/Plugins/FX/Niagara/Source/NiagaraEditor/Private/Widgets/SNiagaraParameterMapView.h
--- a/Engine/Plugins/FX/Niagara/Source/NiagaraEditor/Private/Widgets/SNiagaraParameterMapView.h
+++ b/Engine/Plugins/FX/Niagara/Source/NiagaraEditor/Private/Widgets/SNiagaraParameterMapView.h
@@ -79,8 +79,8 @@ public:
 
 	virtual void Tick(const FGeometry& AllottedGeometry, const double InCurrentTime, const float InDeltaTime) override;
 
-	/** Helper function indicating whether we're in editing mode, and can modify the target blueprint */
-	bool IsEditingMode() const;
+	/** Wheter the add parameter button should be enabled. */
+	bool ParameterAddEnabled() const;
 
 	/** Adds parameter to the graph parameter store and refreshes the menu. */
 	void AddParameter(FNiagaraVariable NewVariable);
@@ -183,7 +183,7 @@ public:
 		SLATE_ATTRIBUTE(bool, AutoExpandMenu)
 	SLATE_END_ARGS();
 
-	void Construct(const FArguments& InArgs, UNiagaraGraph* InGraph);
+	void Construct(const FArguments& InArgs, TArray<TWeakObjectPtr<UNiagaraGraph>> InGraphs);
 
 	TSharedRef<SEditableTextBox> GetSearchBox();
 
@@ -203,7 +203,7 @@ private:
 	FOnCollectCustomActions OnCollectCustomActions;
 	FOnAllowMakeType OnAllowMakeType;
 
-	TWeakObjectPtr<UNiagaraGraph> Graph;
+	TArray<TWeakObjectPtr<UNiagaraGraph>> Graphs;
 
 	TAttribute<NiagaraParameterMapSectionID::Type> Section;
 	TAttribute<bool> AllowCreatingNew;
