Prevent circular connections when trying to connect pins

diff --git a/Engine/Plugins/FX/Niagara/Source/NiagaraEditor/Private/EdGraphSchema_Niagara.cpp b/Engine/Plugins/FX/Niagara/Source/NiagaraEditor/Private/EdGraphSchema_Niagara.cpp
--- a/Engine/Plugins/FX/Niagara/Source/NiagaraEditor/Private/EdGraphSchema_Niagara.cpp
+++ b/Engine/Plugins/FX/Niagara/Source/NiagaraEditor/Private/EdGraphSchema_Niagara.cpp
@@ -926,6 +926,12 @@ const FPinConnectionResponse UEdGraphSchema_Niagara::CanCreateConnection(const U
 		}
 	}
 
+	int32 Depth = 0;
+	if (UEdGraphSchema_Niagara::CheckCircularConnection(PinB->GetOwningNode(), PinB->Direction, PinA, Depth))
+	{
+		return FPinConnectionResponse(CONNECT_RESPONSE_DISALLOW, TEXT("Circular connection found"));
+	}
+
 	// See if we want to break existing connections (if its an input with an existing connection)
 	const bool bBreakExistingDueToDataInput = (InputPin->LinkedTo.Num() > 0);
 	if (bBreakExistingDueToDataInput)
@@ -1408,6 +1414,44 @@ void UEdGraphSchema_Niagara::ConvertNumericPinToType(UEdGraphPin* InGraphPin, FN
 	}
 }
 
+bool UEdGraphSchema_Niagara::CheckCircularConnection(const UEdGraphNode* InRootNode, const EEdGraphPinDirection InRootPinDirection, const UEdGraphPin* InPin, int32& OutDepth)
+{
+	if (InPin->GetOwningNode() == InRootNode)
+	{
+		return true;
+	}
+
+	static const int32 MaxDepth = 3;
+	OutDepth++;
+	if (OutDepth > MaxDepth)
+	{
+		return false;
+	}
+
+	for (const UEdGraphPin* Pin : InPin->GetOwningNode()->GetAllPins())
+	{
+		if (Pin->Direction == InRootPinDirection && Pin != InPin)
+		{
+			for (const UEdGraphPin* LinkedPin : Pin->LinkedTo)
+			{
+				if (UEdGraphSchema_Niagara::CheckCircularConnection(InRootNode, InRootPinDirection, LinkedPin, OutDepth))
+				{
+					return true;
+				}
+
+				// If the CheckCircularConnection call above returned without finding the root node and was too deep.
+				if (OutDepth > MaxDepth)
+				{
+					return false;
+				}
+			}
+		}
+	}
+
+	OutDepth--;
+	return false;
+}
+
 void UEdGraphSchema_Niagara::GetNumericConversionToSubMenuActions(class FMenuBuilder& MenuBuilder, UEdGraphPin* InGraphPin)
 {
 	// Add all the types we could convert to
diff --git a/Engine/Plugins/FX/Niagara/Source/NiagaraEditor/Public/EdGraphSchema_Niagara.h b/Engine/Plugins/FX/Niagara/Source/NiagaraEditor/Public/EdGraphSchema_Niagara.h
--- a/Engine/Plugins/FX/Niagara/Source/NiagaraEditor/Public/EdGraphSchema_Niagara.h
+++ b/Engine/Plugins/FX/Niagara/Source/NiagaraEditor/Public/EdGraphSchema_Niagara.h
@@ -124,7 +124,7 @@ private:
 	void GetBreakLinkToSubMenuActions(class FMenuBuilder& MenuBuilder, UEdGraphPin* InGraphPin);
 	void GetNumericConversionToSubMenuActions(class FMenuBuilder& MenuBuilder, UEdGraphPin* InGraphPin);
 	void ConvertNumericPinToType(UEdGraphPin* InPin, FNiagaraTypeDefinition TypeDef);
-
+	static bool CheckCircularConnection(const UEdGraphNode* InRootNode, const EEdGraphPinDirection InRootPinDirection, const UEdGraphPin* InPin, int32& OutDepth);
 };
 
 class FNiagaraConnectionDrawingPolicy : public FConnectionDrawingPolicy
