Fixing static code analysis issues by moving the RefreshIssues call inside the FunctionCallNode nullcheck
#tests none

diff --git a/Engine/Plugins/FX/Niagara/Source/NiagaraEditor/Private/ViewModels/Stack/NiagaraStackModuleItem.cpp b/Engine/Plugins/FX/Niagara/Source/NiagaraEditor/Private/ViewModels/Stack/NiagaraStackModuleItem.cpp
--- a/Engine/Plugins/FX/Niagara/Source/NiagaraEditor/Private/ViewModels/Stack/NiagaraStackModuleItem.cpp
+++ b/Engine/Plugins/FX/Niagara/Source/NiagaraEditor/Private/ViewModels/Stack/NiagaraStackModuleItem.cpp
@@ -224,13 +224,13 @@ void UNiagaraStackModuleItem::RefreshChildrenInternal(const TArray<UNiagaraStack
 
 		NewChildren.Add(InputCollection);
 		NewChildren.Add(OutputCollection);
-	}
 
-	RefreshIsEnabled();
+		RefreshIsEnabled();
+
+		Super::RefreshChildrenInternal(CurrentChildren, NewChildren, NewIssues);
 
-	Super::RefreshChildrenInternal(CurrentChildren, NewChildren, NewIssues);
-	
-	RefreshIssues(NewIssues);
+		RefreshIssues(NewIssues);
+	}
 }
 
 void UNiagaraStackModuleItem::RefreshIssues(TArray<FStackIssue>& NewIssues)
