Fixing uniform buffer alignment and padding to 16 bytes for all vector types; bumping vec2 and vec3 uniforms to vec4, and adding component mask to code chunk for accesses to uniform chunks according to their initial type
OpenGL requires this since because adherence to the std140 memory layout is shaky at best when it comes to sub-16-byte vector types

diff --git a/Engine/Plugins/FX/Niagara/Source/Niagara/Private/NiagaraScriptExecutionParameterStore.cpp b/Engine/Plugins/FX/Niagara/Source/Niagara/Private/NiagaraScriptExecutionParameterStore.cpp
--- a/Engine/Plugins/FX/Niagara/Source/Niagara/Private/NiagaraScriptExecutionParameterStore.cpp
+++ b/Engine/Plugins/FX/Niagara/Source/Niagara/Private/NiagaraScriptExecutionParameterStore.cpp
@@ -38,50 +38,50 @@ uint32 OffsetAlign(uint32 SrcOffset, uint32 Size)
 
 void FNiagaraScriptExecutionParameterStore::GenerateLayoutInfoInternal(TArray<FNiagaraScriptExecutionPaddingInfo>& Members, uint32& NextMemberOffset, const UStruct* InSrcStruct, uint32 InSrcOffset)
 {
+	uint32 VectorPaddedSize = (TUniformBufferTypeInfo<FVector4, false>::NumRows * TUniformBufferTypeInfo<FVector4, false>::NumColumns) * sizeof(float);
+
 	// Now insert an appropriate data member into the mix...
 	if (InSrcStruct == FNiagaraTypeDefinition::GetBoolStruct() || InSrcStruct == FNiagaraTypeDefinition::GetIntStruct())
 	{
-		//Members.Emplace(*Property->GetName(), TEXT(""), InSrcOffset, (EUniformBufferBaseType)TUniformBufferTypeInfo<uint32>::BaseType, EShaderPrecisionModifier::Float, TUniformBufferTypeInfo<uint32>::NumRows, TUniformBufferTypeInfo<uint32>::NumColumns, TUniformBufferTypeInfo<uint32>::NumElements, TUniformBufferTypeInfo<uint32>::GetStruct());
-		Members.Emplace(InSrcOffset, Align(NextMemberOffset, TUniformBufferTypeInfo<uint32, false>::Alignment), (TUniformBufferTypeInfo<uint32, false>::NumRows * TUniformBufferTypeInfo<uint32, false>::NumColumns) * sizeof(uint32));
+		uint32 IntSize = (TUniformBufferTypeInfo<uint32, false>::NumRows * TUniformBufferTypeInfo<uint32, false>::NumColumns) * sizeof(uint32);
+		Members.Emplace(InSrcOffset, Align(NextMemberOffset, TUniformBufferTypeInfo<uint32, false>::Alignment), IntSize, IntSize);
 		InSrcOffset += sizeof(uint32);
-		NextMemberOffset = Members[Members.Num() - 1].DestOffset + Members[Members.Num() - 1].Size;
+		NextMemberOffset = Members[Members.Num() - 1].DestOffset + Members[Members.Num() - 1].DestSize;
 	}
 	else if (InSrcStruct == FNiagaraTypeDefinition::GetFloatStruct())
 	{
-		//Members.Emplace(*Property->GetName(), TEXT(""), InSrcOffset, (EUniformBufferBaseType)TUniformBufferTypeInfo < float > ::BaseType, EShaderPrecisionModifier::Float, TUniformBufferTypeInfo<float>::NumRows, TUniformBufferTypeInfo<float>::NumColumns, TUniformBufferTypeInfo<float>::NumElements, TUniformBufferTypeInfo<float>::GetStruct());
-		Members.Emplace(InSrcOffset, Align(NextMemberOffset, TUniformBufferTypeInfo<float, false>::Alignment), (TUniformBufferTypeInfo<float, false>::NumRows * TUniformBufferTypeInfo<float, false>::NumColumns) * sizeof(float));
+		uint32 FloatSize = (TUniformBufferTypeInfo<float, false>::NumRows * TUniformBufferTypeInfo<float, false>::NumColumns) * sizeof(float);
+		Members.Emplace(InSrcOffset, Align(NextMemberOffset, TUniformBufferTypeInfo<float, false>::Alignment), FloatSize, FloatSize);
 		InSrcOffset += sizeof(float);
-		NextMemberOffset = Members[Members.Num() - 1].DestOffset + Members[Members.Num() - 1].Size;
+		NextMemberOffset = Members[Members.Num() - 1].DestOffset + Members[Members.Num() - 1].DestSize;
 	}
 	else if (InSrcStruct == FNiagaraTypeDefinition::GetVec2Struct())
 	{
-		//Members.Emplace(*Property->GetName(), TEXT(""), InSrcOffset, (EUniformBufferBaseType)TUniformBufferTypeInfo < FVector2D > ::BaseType, EShaderPrecisionModifier::Float, TUniformBufferTypeInfo<FVector2D>::NumRows, TUniformBufferTypeInfo<FVector2D>::NumColumns, TUniformBufferTypeInfo<FVector2D>::NumElements, TUniformBufferTypeInfo<FVector2D>::GetStruct());
 		uint32 StructFinalSize = (TUniformBufferTypeInfo<FVector2D, false>::NumRows * TUniformBufferTypeInfo<FVector2D, false>::NumColumns) * sizeof(float);
-		Members.Emplace(InSrcOffset, OffsetAlign(NextMemberOffset, StructFinalSize),StructFinalSize);
+		Members.Emplace(InSrcOffset, OffsetAlign(NextMemberOffset, VectorPaddedSize), StructFinalSize, VectorPaddedSize);
 		InSrcOffset += sizeof(FVector2D);
-		NextMemberOffset = Members[Members.Num() - 1].DestOffset + Members[Members.Num() - 1].Size;
+		NextMemberOffset = Members[Members.Num() - 1].DestOffset + Members[Members.Num() - 1].DestSize;
 	}
 	else if (InSrcStruct == FNiagaraTypeDefinition::GetVec3Struct())
 	{
-		//Members.Emplace(*Property->GetName(), TEXT(""), InSrcOffset, (EUniformBufferBaseType)TUniformBufferTypeInfo < FVector > ::BaseType, EShaderPrecisionModifier::Float, TUniformBufferTypeInfo<FVector>::NumRows, TUniformBufferTypeInfo<FVector>::NumColumns, TUniformBufferTypeInfo<FVector>::NumElements, TUniformBufferTypeInfo<FVector>::GetStruct());
 		uint32 StructFinalSize = (TUniformBufferTypeInfo<FVector, false>::NumRows * TUniformBufferTypeInfo<FVector, false>::NumColumns) * sizeof(float);
-		Members.Emplace(InSrcOffset, OffsetAlign(NextMemberOffset, StructFinalSize), StructFinalSize);
+		Members.Emplace(InSrcOffset, OffsetAlign(NextMemberOffset, VectorPaddedSize), StructFinalSize, VectorPaddedSize);
 		InSrcOffset += sizeof(FVector);
-		NextMemberOffset = Members[Members.Num() - 1].DestOffset + Members[Members.Num() - 1].Size;
+		NextMemberOffset = Members[Members.Num() - 1].DestOffset + Members[Members.Num() - 1].DestSize;
 	}
 	else if (InSrcStruct == FNiagaraTypeDefinition::GetVec4Struct() || InSrcStruct == FNiagaraTypeDefinition::GetColorStruct() || InSrcStruct == FNiagaraTypeDefinition::GetQuatStruct())
 	{
-		//Members.Emplace(*Property->GetName(), TEXT(""), InSrcOffset, (EUniformBufferBaseType)TUniformBufferTypeInfo < FVector4 > ::BaseType, EShaderPrecisionModifier::Float, TUniformBufferTypeInfo<FVector4>::NumRows, TUniformBufferTypeInfo<FVector4>::NumColumns, TUniformBufferTypeInfo<FVector4>::NumElements, TUniformBufferTypeInfo<FVector4>::GetStruct());
-		Members.Emplace(InSrcOffset, Align(NextMemberOffset, TUniformBufferTypeInfo<FVector4, false>::Alignment), (TUniformBufferTypeInfo<FVector4, false>::NumRows * TUniformBufferTypeInfo<FVector4, false>::NumColumns) * sizeof(float));
+		uint32 StructFinalSize = (TUniformBufferTypeInfo<FVector4, false>::NumRows * TUniformBufferTypeInfo<FVector4, false>::NumColumns) * sizeof(float);
+		Members.Emplace(InSrcOffset, Align(NextMemberOffset, TUniformBufferTypeInfo<FVector4, false>::Alignment), StructFinalSize, StructFinalSize);
 		InSrcOffset += sizeof(FVector4);
-		NextMemberOffset = Members[Members.Num() - 1].DestOffset + Members[Members.Num() - 1].Size;
+		NextMemberOffset = Members[Members.Num() - 1].DestOffset + Members[Members.Num() - 1].DestSize;
 	}
 	else if (InSrcStruct == FNiagaraTypeDefinition::GetMatrix4Struct())
 	{
-		//Members.Emplace(*Property->GetName(), TEXT(""), InSrcOffset, (EUniformBufferBaseType)TUniformBufferTypeInfo < FMatrix > ::BaseType, EShaderPrecisionModifier::Float, TUniformBufferTypeInfo<FMatrix>::NumRows, TUniformBufferTypeInfo<FMatrix>::NumColumns, TUniformBufferTypeInfo<FMatrix>::NumElements, TUniformBufferTypeInfo<FMatrix>::GetStruct());
-		Members.Emplace(InSrcOffset, Align(NextMemberOffset, TUniformBufferTypeInfo<FMatrix, false>::Alignment), (TUniformBufferTypeInfo<FMatrix, false>::NumRows * TUniformBufferTypeInfo<FMatrix, false>::NumColumns) * sizeof(float));
+		uint32 StructFinalSize = (TUniformBufferTypeInfo<FMatrix, false>::NumRows * TUniformBufferTypeInfo<FMatrix, false>::NumColumns) * sizeof(float);
+		Members.Emplace(InSrcOffset, Align(NextMemberOffset, TUniformBufferTypeInfo<FMatrix, false>::Alignment), StructFinalSize, StructFinalSize);
 		InSrcOffset += sizeof(FMatrix);
-		NextMemberOffset = Members[Members.Num() - 1].DestOffset + Members[Members.Num() - 1].Size;
+		NextMemberOffset = Members[Members.Num() - 1].DestOffset + Members[Members.Num() - 1].DestSize;
 	}
 	else
 	{
@@ -131,13 +131,13 @@ void FNiagaraScriptExecutionParameterStore::AddPaddedParamSize(const FNiagaraTyp
 	uint32 NextMemberOffset = 0;
 	if (PaddingInfo.Num() != 0)
 	{
-		NextMemberOffset = PaddingInfo[PaddingInfo.Num() - 1].DestOffset + PaddingInfo[PaddingInfo.Num() - 1].Size;
+		NextMemberOffset = PaddingInfo[PaddingInfo.Num() - 1].DestOffset + PaddingInfo[PaddingInfo.Num() - 1].DestSize;
 	}
 	GenerateLayoutInfoInternal(PaddingInfo, NextMemberOffset, InParamType.GetScriptStruct(), InOffset);
 
 	if (PaddingInfo.Num() != 0)
 	{
-		NextMemberOffset = PaddingInfo[PaddingInfo.Num() - 1].DestOffset + PaddingInfo[PaddingInfo.Num() - 1].Size;
+		NextMemberOffset = PaddingInfo[PaddingInfo.Num() - 1].DestOffset + PaddingInfo[PaddingInfo.Num() - 1].DestSize;
 		PaddedParameterSize = Align(NextMemberOffset, UNIFORM_BUFFER_STRUCT_ALIGNMENT);
 	}
 	else
@@ -294,8 +294,8 @@ void FNiagaraScriptExecutionParameterStore::CopyParameterDataToPaddedBuffer(uint
 	const uint8* SrcData = GetParameterDataArray().GetData();
 	for (int32 i = 0; i < PaddingInfo.Num(); i++)
 	{
-		check((PaddingInfo[i].DestOffset + PaddingInfo[i].Size) <= InTargetBufferSizeInBytes);
-		check((PaddingInfo[i].SrcOffset + PaddingInfo[i].Size) <= (uint32)GetParameterDataArray().Num());
-		FMemory::Memcpy(InTargetBuffer + PaddingInfo[i].DestOffset, SrcData + PaddingInfo[i].SrcOffset, PaddingInfo[i].Size);
+		check((PaddingInfo[i].DestOffset + PaddingInfo[i].DestSize) <= InTargetBufferSizeInBytes);
+		check((PaddingInfo[i].SrcOffset + PaddingInfo[i].SrcSize) <= (uint32)GetParameterDataArray().Num());
+		FMemory::Memcpy(InTargetBuffer + PaddingInfo[i].DestOffset, SrcData + PaddingInfo[i].SrcOffset, PaddingInfo[i].SrcSize);
 	}
 }
diff --git a/Engine/Plugins/FX/Niagara/Source/Niagara/Public/NiagaraScriptExecutionParameterStore.h b/Engine/Plugins/FX/Niagara/Source/Niagara/Public/NiagaraScriptExecutionParameterStore.h
--- a/Engine/Plugins/FX/Niagara/Source/Niagara/Public/NiagaraScriptExecutionParameterStore.h
+++ b/Engine/Plugins/FX/Niagara/Source/Niagara/Public/NiagaraScriptExecutionParameterStore.h
@@ -10,11 +10,12 @@ struct FNiagaraScriptExecutionPaddingInfo
 {
 	GENERATED_USTRUCT_BODY()
 public:
-	FNiagaraScriptExecutionPaddingInfo() : SrcOffset(0), DestOffset(0), Size(0) {}
-	FNiagaraScriptExecutionPaddingInfo(uint32 InSrcOffset, uint32 InDestOffset, uint32 InSize) : SrcOffset(InSrcOffset), DestOffset(InDestOffset), Size(InSize) {}
+	FNiagaraScriptExecutionPaddingInfo() : SrcOffset(0), DestOffset(0), SrcSize(0), DestSize(0) {}
+	FNiagaraScriptExecutionPaddingInfo(uint32 InSrcOffset, uint32 InDestOffset, uint32 InSrcSize, uint32 InDestSize) : SrcOffset(InSrcOffset), DestOffset(InDestOffset), SrcSize(InSrcSize), DestSize(InDestSize) {}
 	uint32 SrcOffset;
 	uint32 DestOffset;
-	uint32 Size;
+	uint32 SrcSize;
+	uint32 DestSize;
 };
 
 /**
diff --git a/Engine/Plugins/FX/Niagara/Source/NiagaraEditor/Private/NiagaraHlslTranslator.cpp b/Engine/Plugins/FX/Niagara/Source/NiagaraEditor/Private/NiagaraHlslTranslator.cpp
--- a/Engine/Plugins/FX/Niagara/Source/NiagaraEditor/Private/NiagaraHlslTranslator.cpp
+++ b/Engine/Plugins/FX/Niagara/Source/NiagaraEditor/Private/NiagaraHlslTranslator.cpp
@@ -223,7 +223,7 @@ FString FHlslNiagaraTranslator::GetCodeAsSource(int32 ChunkIdx)
 	if(ChunkIdx>=0 && ChunkIdx<CodeChunks.Num())
 	{ 
 		FNiagaraCodeChunk& Chunk = CodeChunks[ChunkIdx];
-		return Chunk.SymbolName;
+		return Chunk.SymbolName + Chunk.ComponentMask;
 	}
 	return "Undefined";
 }
@@ -1383,7 +1383,7 @@ void FHlslNiagaraTranslator::DefineInterpolatedParametersFunction(FString &HlslO
 						FNiagaraCodeChunk& Chunk = CodeChunks[ChunkIdx];
 						if (ShouldInterpolateParameter(*FoundNamespacedVar))
 						{
-							HlslOutputString += TEXT("\tContext.") + PrevMap + TEXT(".") + FoundName + TEXT(" = lerp(") + INTERPOLATED_PARAMETER_PREFIX + Chunk.SymbolName + TEXT(", ") + Chunk.SymbolName + TEXT(", ") + TEXT("SpawnInterp);\n");
+							HlslOutputString += TEXT("\tContext.") + PrevMap + TEXT(".") + FoundName + TEXT(" = lerp(") + INTERPOLATED_PARAMETER_PREFIX + Chunk.SymbolName + Chunk.ComponentMask + TEXT(", ") + Chunk.SymbolName + Chunk.ComponentMask + TEXT(", ") + TEXT("SpawnInterp);\n");
 						}
 						else
 						{
@@ -2209,6 +2209,21 @@ int32 FHlslNiagaraTranslator::AddUniformChunk(FString SymbolName, const FNiagara
 		FNiagaraCodeChunk& Chunk = CodeChunks[Ret];
 		Chunk.SymbolName = GetSanitizedSymbolName(SymbolName);
 		Chunk.Type = Type;
+
+		if (CompileOptions.TargetUsage == ENiagaraScriptUsage::ParticleGPUComputeScript)
+		{
+			if (Type == FNiagaraTypeDefinition::GetVec2Def())
+			{
+				Chunk.Type = FNiagaraTypeDefinition::GetVec4Def();
+				Chunk.ComponentMask = TEXT(".xy");
+			}
+			else if (Type == FNiagaraTypeDefinition::GetVec3Def())
+			{
+				Chunk.Type = FNiagaraTypeDefinition::GetVec4Def();
+				Chunk.ComponentMask = TEXT(".xyz");
+			}
+		}
+
 		Chunk.Mode = ENiagaraCodeChunkMode::Uniform;
 
 		ChunksByMode[(int32)ENiagaraCodeChunkMode::Uniform].Add(Ret);
@@ -3006,6 +3021,7 @@ bool FHlslNiagaraTranslator::ParameterMapRegisterExternalConstantNamespaceVariab
 
 				CompilationOutput.ScriptData.Parameters.SetOrAdd(InVariable);
 				UniformIdx = ChunksByMode[(int32)ENiagaraCodeChunkMode::Uniform].Num();
+
 				UniformChunk = AddUniformChunk(SymbolNameDefined, InVariable.GetType());
 				ParamMapDefinedSystemVarsToUniformChunks.Add(InVariable.GetName(), UniformIdx);
 				ParamMapDefinedSystemToNamespaceVars.Add(InVariable.GetName(), InVariable);
diff --git a/Engine/Plugins/FX/Niagara/Source/NiagaraEditor/Private/NiagaraHlslTranslator.h b/Engine/Plugins/FX/Niagara/Source/NiagaraEditor/Private/NiagaraHlslTranslator.h
--- a/Engine/Plugins/FX/Niagara/Source/NiagaraEditor/Private/NiagaraHlslTranslator.h
+++ b/Engine/Plugins/FX/Niagara/Source/NiagaraEditor/Private/NiagaraHlslTranslator.h
@@ -198,12 +198,15 @@ struct FNiagaraCodeChunk
 	bool bIsTerminated;
 	/** Chunks used as input for this chunk. */
 	TArray<int32> SourceChunks;
+	/** Component mask for access to padded uniforms; will be empty except for float2 and float3 uniforms */
+	FString ComponentMask;
 
 	ENiagaraCodeChunkMode Mode;
 
 	FNiagaraCodeChunk()
 		: bDecl(true)
 		, bIsTerminated(true)
+		, ComponentMask("")
 		, Mode(ENiagaraCodeChunkMode::Num)
 	{
 		Type = FNiagaraTypeDefinition::GetFloatDef();
