Niagara - Renderer update fixes.
+ Trigger data object changed when adding, removing, and changing the enabled state of renderers so that the simulation updates.
+ Fix undo for changing the enabled state on renderers.

diff --git a/Engine/Plugins/FX/Niagara/Source/NiagaraEditor/Private/ViewModels/Stack/NiagaraStackEventHandlerGroup.cpp b/Engine/Plugins/FX/Niagara/Source/NiagaraEditor/Private/ViewModels/Stack/NiagaraStackEventHandlerGroup.cpp
--- a/Engine/Plugins/FX/Niagara/Source/NiagaraEditor/Private/ViewModels/Stack/NiagaraStackEventHandlerGroup.cpp
+++ b/Engine/Plugins/FX/Niagara/Source/NiagaraEditor/Private/ViewModels/Stack/NiagaraStackEventHandlerGroup.cpp
@@ -17,11 +17,11 @@
 
 #define LOCTEXT_NAMESPACE "NiagaraStackEventHandlerGroup"
 
-class FEventHandlerGroupAddUtilities : public FNiagaraStackItemGroupAddUtilities
+class FEventHandlerGroupAddUtilities : public TNiagaraStackItemGroupAddUtilities<FNiagaraEventScriptProperties>
 {
 public:
 	FEventHandlerGroupAddUtilities(TSharedRef<FNiagaraEmitterViewModel> InEmitterViewModel, FOnItemAdded InOnItemAdded)
-		: FNiagaraStackItemGroupAddUtilities(LOCTEXT("ScriptGroupAddItemName", "Event Handler"), EAddMode::AddDirectly, false, InOnItemAdded)
+		: TNiagaraStackItemGroupAddUtilities(LOCTEXT("ScriptGroupAddItemName", "Event Handler"), EAddMode::AddDirectly, false, InOnItemAdded)
 		, EmitterViewModel(InEmitterViewModel)
 	{
 	}
@@ -56,7 +56,7 @@ public:
 		// TODO: Move the logic for managing event handlers into the emitter view model or script view model.
 		EmitterViewModelPinned->SetEmitter(Emitter);
 
-		OnItemAdded.ExecuteIfBound();
+		OnItemAdded.ExecuteIfBound(EventScriptProperties);
 	}
 
 	virtual void GenerateAddActions(TArray<TSharedRef<INiagaraStackItemGroupAddAction>>& OutAddActions) const override { unimplemented(); }
@@ -71,7 +71,7 @@ void UNiagaraStackEventHandlerGroup::Initialize(FRequiredEntryData InRequiredEnt
 	FText DisplayName = LOCTEXT("EventGroupName", "Add Event Handler");
 	FText ToolTip = LOCTEXT("EventGroupTooltip", "Determines how this Emitter responds to incoming events. There can be more than one event handler script stack per Emitter.");
 	AddUtilities = MakeShared<FEventHandlerGroupAddUtilities>(InRequiredEntryData.EmitterViewModel, 
-		FNiagaraStackItemGroupAddUtilities::FOnItemAdded::CreateUObject(this, &UNiagaraStackEventHandlerGroup::ItemAddedFromUtilties));
+		TNiagaraStackItemGroupAddUtilities<FNiagaraEventScriptProperties>::FOnItemAdded::CreateUObject(this, &UNiagaraStackEventHandlerGroup::ItemAddedFromUtilties));
 	Super::Initialize(InRequiredEntryData, DisplayName, ToolTip, AddUtilities.Get());
 }
 
@@ -80,7 +80,7 @@ void UNiagaraStackEventHandlerGroup::SetOnItemAdded(FOnItemAdded InOnItemAdded)
 	ItemAddedDelegate = InOnItemAdded;
 }
 
-void UNiagaraStackEventHandlerGroup::ItemAddedFromUtilties()
+void UNiagaraStackEventHandlerGroup::ItemAddedFromUtilties(FNiagaraEventScriptProperties AddedEventHandler)
 {
 	ItemAddedDelegate.ExecuteIfBound();
 }
diff --git a/Engine/Plugins/FX/Niagara/Source/NiagaraEditor/Private/ViewModels/Stack/NiagaraStackParameterStoreGroup.cpp b/Engine/Plugins/FX/Niagara/Source/NiagaraEditor/Private/ViewModels/Stack/NiagaraStackParameterStoreGroup.cpp
--- a/Engine/Plugins/FX/Niagara/Source/NiagaraEditor/Private/ViewModels/Stack/NiagaraStackParameterStoreGroup.cpp
+++ b/Engine/Plugins/FX/Niagara/Source/NiagaraEditor/Private/ViewModels/Stack/NiagaraStackParameterStoreGroup.cpp
@@ -55,7 +55,7 @@ class FParameterStoreGroupAddUtiliites : public FNiagaraStackItemGroupAddUtiliti
 {
 public:
 	FParameterStoreGroupAddUtiliites(UObject& InParameterStoreOwner, FNiagaraParameterStore& InParameterStore, UNiagaraStackEditorData& InStackEditorData)
-		: FNiagaraStackItemGroupAddUtilities(LOCTEXT("ScriptGroupAddItemName", "Parameter"), EAddMode::AddFromAction, true, FOnItemAdded())
+		: FNiagaraStackItemGroupAddUtilities(LOCTEXT("ScriptGroupAddItemName", "Parameter"), EAddMode::AddFromAction, true)
 		, ParameterStoreOwner(InParameterStoreOwner)
 		, ParameterStore(InParameterStore)
 		, StackEditorData(InStackEditorData)
diff --git a/Engine/Plugins/FX/Niagara/Source/NiagaraEditor/Private/ViewModels/Stack/NiagaraStackRenderItemGroup.cpp b/Engine/Plugins/FX/Niagara/Source/NiagaraEditor/Private/ViewModels/Stack/NiagaraStackRenderItemGroup.cpp
--- a/Engine/Plugins/FX/Niagara/Source/NiagaraEditor/Private/ViewModels/Stack/NiagaraStackRenderItemGroup.cpp
+++ b/Engine/Plugins/FX/Niagara/Source/NiagaraEditor/Private/ViewModels/Stack/NiagaraStackRenderItemGroup.cpp
@@ -48,11 +48,11 @@ private:
 	UClass* RendererClass;
 };
 
-class FRenderItemGroupAddUtilities : public FNiagaraStackItemGroupAddUtilities
+class FRenderItemGroupAddUtilities : public TNiagaraStackItemGroupAddUtilities<UNiagaraRendererProperties*>
 {
 public:
 	FRenderItemGroupAddUtilities(TSharedRef<FNiagaraEmitterViewModel> InEmitterViewModel, FOnItemAdded InOnItemAdded)
-		: FNiagaraStackItemGroupAddUtilities(LOCTEXT("RenderGroupAddItemName", "Renderer"), EAddMode::AddFromAction, true, InOnItemAdded)
+		: TNiagaraStackItemGroupAddUtilities(LOCTEXT("RenderGroupAddItemName", "Renderer"), EAddMode::AddFromAction, true, InOnItemAdded)
 		, EmitterViewModel(InEmitterViewModel)
 	{
 	}
@@ -105,7 +105,7 @@ public:
 			FSlateNotificationManager::Get().AddNotification(Info);
 		}
 
-		OnItemAdded.ExecuteIfBound();
+		OnItemAdded.ExecuteIfBound(RendererProperties);
 	}
 
 private:
@@ -117,7 +117,7 @@ void UNiagaraStackRenderItemGroup::Initialize(FRequiredEntryData InRequiredEntry
 	FText DisplayName = LOCTEXT("RenderGroupName", "Render");
 	FText ToolTip = LOCTEXT("RendererGroupTooltip", "Describes how we should display/present each particle. Note that this doesn't have to be visual. Multiple renderers are supported. Order in this stack is not necessarily relevant to draw order.");
 	AddUtilities = MakeShared<FRenderItemGroupAddUtilities>(InRequiredEntryData.EmitterViewModel,
-		FNiagaraStackItemGroupAddUtilities::FOnItemAdded::CreateUObject(this, &UNiagaraStackRenderItemGroup::ItemAdded));
+		TNiagaraStackItemGroupAddUtilities<UNiagaraRendererProperties*>::FOnItemAdded::CreateUObject(this, &UNiagaraStackRenderItemGroup::ItemAdded));
 	Super::Initialize(InRequiredEntryData, DisplayName, ToolTip, AddUtilities.Get());
 }
 
@@ -156,9 +156,10 @@ void UNiagaraStackRenderItemGroup::RefreshChildrenInternal(const TArray<UNiagara
 	Super::RefreshChildrenInternal(CurrentChildren, NewChildren, NewIssues);
 }
 
-void UNiagaraStackRenderItemGroup::ItemAdded()
+void UNiagaraStackRenderItemGroup::ItemAdded(UNiagaraRendererProperties* AddedRenderer)
 {
 	RefreshChildren();
+	OnDataObjectModified().Broadcast(AddedRenderer);
 }
 
 void UNiagaraStackRenderItemGroup::ChildModifiedGroupItems()
diff --git a/Engine/Plugins/FX/Niagara/Source/NiagaraEditor/Private/ViewModels/Stack/NiagaraStackRendererItem.cpp b/Engine/Plugins/FX/Niagara/Source/NiagaraEditor/Private/ViewModels/Stack/NiagaraStackRendererItem.cpp
--- a/Engine/Plugins/FX/Niagara/Source/NiagaraEditor/Private/ViewModels/Stack/NiagaraStackRendererItem.cpp
+++ b/Engine/Plugins/FX/Niagara/Source/NiagaraEditor/Private/ViewModels/Stack/NiagaraStackRendererItem.cpp
@@ -159,6 +159,7 @@ void UNiagaraStackRendererItem::Delete()
 	Emitter->Modify();
 	Emitter->RemoveRenderer(RendererProperties.Get());
 
+	OnDataObjectModified().Broadcast(RendererProperties.Get());
 	ModifiedGroupItemsDelegate.ExecuteIfBound();
 }
 
@@ -201,8 +202,9 @@ bool UNiagaraStackRendererItem::GetIsEnabled() const
 void UNiagaraStackRendererItem::SetIsEnabled(bool bInIsEnabled)
 {
 	FScopedTransaction ScopedTransaction(LOCTEXT("SetRendererEnabledState", "Set renderer enabled/disabled state."));
-	RendererObject->Modify();
+	RendererProperties->Modify();
 	RendererProperties->SetIsEnabled(bInIsEnabled);
+	OnDataObjectModified().Broadcast(RendererProperties.Get());
 }
 
 void UNiagaraStackRendererItem::BeginDestroy()
diff --git a/Engine/Plugins/FX/Niagara/Source/NiagaraEditor/Private/ViewModels/Stack/NiagaraStackScriptItemGroup.cpp b/Engine/Plugins/FX/Niagara/Source/NiagaraEditor/Private/ViewModels/Stack/NiagaraStackScriptItemGroup.cpp
--- a/Engine/Plugins/FX/Niagara/Source/NiagaraEditor/Private/ViewModels/Stack/NiagaraStackScriptItemGroup.cpp
+++ b/Engine/Plugins/FX/Niagara/Source/NiagaraEditor/Private/ViewModels/Stack/NiagaraStackScriptItemGroup.cpp
@@ -131,11 +131,11 @@ private:
 	bool bIsMaterialParameterModuleAction;
 };
 
-class FScriptItemGroupAddUtilities : public FNiagaraStackItemGroupAddUtilities
+class FScriptItemGroupAddUtilities : public TNiagaraStackItemGroupAddUtilities<UNiagaraNodeFunctionCall*>
 {
 public:
 	FScriptItemGroupAddUtilities(TSharedRef<FNiagaraSystemViewModel> InSystemViewModel, TSharedRef<FNiagaraEmitterViewModel> InEmitterViewModel, UNiagaraStackEditorData& InStackEditorData, FOnItemAdded InOnItemAdded)
-		: FNiagaraStackItemGroupAddUtilities(LOCTEXT("ScriptGroupAddItemName", "Module"), EAddMode::AddFromAction, false, InOnItemAdded)
+		: TNiagaraStackItemGroupAddUtilities(LOCTEXT("ScriptGroupAddItemName", "Module"), EAddMode::AddFromAction, false, InOnItemAdded)
 		, SystemViewModel(InSystemViewModel)
 		, EmitterViewModel(InEmitterViewModel)
 		, StackEditorData(InStackEditorData)
@@ -202,7 +202,7 @@ public:
 		checkf(NewModuleNode != nullptr, TEXT("Add module action failed"));
 		FNiagaraStackGraphUtilities::InitializeStackFunctionInputs(SystemViewModel.Pin().ToSharedRef(), EmitterViewModel.Pin().ToSharedRef(), StackEditorData, *NewModuleNode, *NewModuleNode);
 		FNiagaraStackGraphUtilities::RelayoutGraph(*OutputNode->GetGraph());
-		OnItemAdded.ExecuteIfBound();
+		OnItemAdded.ExecuteIfBound(NewModuleNode);
 	}
 private:
 	UNiagaraNodeFunctionCall* AddScriptAssetModule(const FAssetData& AssetData, int32 TargetIndex)
@@ -249,7 +249,7 @@ void UNiagaraStackScriptItemGroup::Initialize(
 {
 	checkf(ScriptViewModel.IsValid() == false, TEXT("Can not set the script view model more than once."));
 	AddUtilities = MakeShared<FScriptItemGroupAddUtilities>(InRequiredEntryData.SystemViewModel, InRequiredEntryData.EmitterViewModel,
-		*InRequiredEntryData.StackEditorData, FNiagaraStackItemGroupAddUtilities::FOnItemAdded::CreateUObject(this, &UNiagaraStackScriptItemGroup::ItemAdded));
+		*InRequiredEntryData.StackEditorData, TNiagaraStackItemGroupAddUtilities<UNiagaraNodeFunctionCall*>::FOnItemAdded::CreateUObject(this, &UNiagaraStackScriptItemGroup::ItemAdded));
 	Super::Initialize(InRequiredEntryData, InDisplayName, InToolTip, AddUtilities.Get());
 	ScriptViewModel = InScriptViewModel;
 	ScriptUsage = InScriptUsage;
@@ -583,7 +583,7 @@ TOptional<UNiagaraStackEntry::FDropResult> UNiagaraStackScriptItemGroup::ChildRe
 	return FDropResult(false);
 }
 
-void UNiagaraStackScriptItemGroup::ItemAdded()
+void UNiagaraStackScriptItemGroup::ItemAdded(UNiagaraNodeFunctionCall* AddedModule)
 {
 	RefreshChildren();
 }
diff --git a/Engine/Plugins/FX/Niagara/Source/NiagaraEditor/Public/ViewModels/Stack/INiagaraStackItemGroupAddUtilities.h b/Engine/Plugins/FX/Niagara/Source/NiagaraEditor/Public/ViewModels/Stack/INiagaraStackItemGroupAddUtilities.h
--- a/Engine/Plugins/FX/Niagara/Source/NiagaraEditor/Public/ViewModels/Stack/INiagaraStackItemGroupAddUtilities.h
+++ b/Engine/Plugins/FX/Niagara/Source/NiagaraEditor/Public/ViewModels/Stack/INiagaraStackItemGroupAddUtilities.h
@@ -56,14 +56,10 @@ public:
 class FNiagaraStackItemGroupAddUtilities : public INiagaraStackItemGroupAddUtilities
 {
 public:
-	DECLARE_DELEGATE(FOnItemAdded);
-
-public:
-	FNiagaraStackItemGroupAddUtilities(FText InAddItemName, EAddMode InAddMode, bool bInAutoExpandAddActions, FOnItemAdded InOnItemAdded)
+	FNiagaraStackItemGroupAddUtilities(FText InAddItemName, EAddMode InAddMode, bool bInAutoExpandAddActions)
 		: AddItemName(InAddItemName)
 		, bAutoExpandAddActions(bInAutoExpandAddActions)
 		, AddMode(InAddMode)
-		, OnItemAdded(InOnItemAdded)
 	{
 	}
 
@@ -75,5 +71,21 @@ protected:
 	FText AddItemName;
 	bool bAutoExpandAddActions;
 	EAddMode AddMode;
+};
+
+template<typename AddedItemType>
+class TNiagaraStackItemGroupAddUtilities : public FNiagaraStackItemGroupAddUtilities
+{
+public:
+	DECLARE_DELEGATE_OneParam(FOnItemAdded, AddedItemType);
+
+public:
+	TNiagaraStackItemGroupAddUtilities(FText InAddItemName, EAddMode InAddMode, bool bInAutoExpandAddActions, FOnItemAdded InOnItemAdded)
+		: FNiagaraStackItemGroupAddUtilities(InAddItemName, InAddMode, bInAutoExpandAddActions)
+		, OnItemAdded(InOnItemAdded)
+	{
+	}
+
+protected:
 	FOnItemAdded OnItemAdded;
 };
\ No newline at end of file
diff --git a/Engine/Plugins/FX/Niagara/Source/NiagaraEditor/Public/ViewModels/Stack/NiagaraStackEventHandlerGroup.h b/Engine/Plugins/FX/Niagara/Source/NiagaraEditor/Public/ViewModels/Stack/NiagaraStackEventHandlerGroup.h
--- a/Engine/Plugins/FX/Niagara/Source/NiagaraEditor/Public/ViewModels/Stack/NiagaraStackEventHandlerGroup.h
+++ b/Engine/Plugins/FX/Niagara/Source/NiagaraEditor/Public/ViewModels/Stack/NiagaraStackEventHandlerGroup.h
@@ -3,6 +3,7 @@
 #pragma once
 
 #include "ViewModels/Stack/NiagaraStackItemGroup.h"
+#include "NiagaraEmitter.h"
 #include "NiagaraStackEventHandlerGroup.generated.h"
 
 class FNiagaraEmitterViewModel;
@@ -22,7 +23,7 @@ public:
 	void SetOnItemAdded(FOnItemAdded InOnItemAdded);
 
 private:
-	void ItemAddedFromUtilties();
+	void ItemAddedFromUtilties(FNiagaraEventScriptProperties AddedEventHandler);
 
 private:
 	TSharedPtr<INiagaraStackItemGroupAddUtilities> AddUtilities;
diff --git a/Engine/Plugins/FX/Niagara/Source/NiagaraEditor/Public/ViewModels/Stack/NiagaraStackRenderItemGroup.h b/Engine/Plugins/FX/Niagara/Source/NiagaraEditor/Public/ViewModels/Stack/NiagaraStackRenderItemGroup.h
--- a/Engine/Plugins/FX/Niagara/Source/NiagaraEditor/Public/ViewModels/Stack/NiagaraStackRenderItemGroup.h
+++ b/Engine/Plugins/FX/Niagara/Source/NiagaraEditor/Public/ViewModels/Stack/NiagaraStackRenderItemGroup.h
@@ -5,6 +5,8 @@
 #include "ViewModels/Stack/NiagaraStackItemGroup.h"
 #include "NiagaraStackRenderItemGroup.generated.h"
 
+class UNiagaraRendererProperties;
+
 UCLASS()
 class NIAGARAEDITOR_API UNiagaraStackRenderItemGroup : public UNiagaraStackItemGroup
 {
@@ -17,7 +19,7 @@ protected:
 	virtual void RefreshChildrenInternal(const TArray<UNiagaraStackEntry*>& CurrentChildren, TArray<UNiagaraStackEntry*>& NewChildren, TArray<FStackIssue>& NewIssues) override;
 
 private:
-	void ItemAdded();
+	void ItemAdded(UNiagaraRendererProperties* AddedRenderer);
 	void ChildModifiedGroupItems();
 
 private:
diff --git a/Engine/Plugins/FX/Niagara/Source/NiagaraEditor/Public/ViewModels/Stack/NiagaraStackScriptItemGroup.h b/Engine/Plugins/FX/Niagara/Source/NiagaraEditor/Public/ViewModels/Stack/NiagaraStackScriptItemGroup.h
--- a/Engine/Plugins/FX/Niagara/Source/NiagaraEditor/Public/ViewModels/Stack/NiagaraStackScriptItemGroup.h
+++ b/Engine/Plugins/FX/Niagara/Source/NiagaraEditor/Public/ViewModels/Stack/NiagaraStackScriptItemGroup.h
@@ -10,6 +10,7 @@
 class FNiagaraScriptViewModel;
 class UNiagaraStackSpacer;
 class UNiagaraNodeOutput;
+class UNiagaraNodeFunctionCall;
 class FScriptItemGroupAddUtilities;
 class UNiagaraStackModuleItem;
 
@@ -38,7 +39,7 @@ protected:
 	virtual TOptional<FDropResult> ChildRequestDropInternal(const UNiagaraStackEntry& TargetChild, const TArray<UNiagaraStackEntry*>& DraggedEntries) override;
 
 private:
-	void ItemAdded();
+	void ItemAdded(UNiagaraNodeFunctionCall* AddedModule);
 
 	void ChildModifiedGroupItems();
 
