Fixing when you can create certain pin types to prevent invalid types from appearing in the list.

#tests autotests pass on PC

diff --git a/Engine/Plugins/FX/Niagara/Source/NiagaraEditor/Private/Widgets/SNiagaraGraphPinAdd.cpp b/Engine/Plugins/FX/Niagara/Source/NiagaraEditor/Private/Widgets/SNiagaraGraphPinAdd.cpp
--- a/Engine/Plugins/FX/Niagara/Source/NiagaraEditor/Private/Widgets/SNiagaraGraphPinAdd.cpp
+++ b/Engine/Plugins/FX/Niagara/Source/NiagaraEditor/Private/Widgets/SNiagaraGraphPinAdd.cpp
@@ -69,7 +69,8 @@ TSharedRef<SWidget> SNiagaraGraphPinAdd::OnGetAddButtonMenuContent()
 		TSharedRef<SNiagaraAddParameterMenu> MenuWidget = SNew(SNiagaraAddParameterMenu, Graphs)
 			.OnAddParameter_UObject(OwningNode, &UNiagaraNodeWithDynamicPins::AddParameter, GetPinObj()) // For non custom actions
 			.OnCollectCustomActions_UObject(OwningNode, &UNiagaraNodeWithDynamicPins::CollectAddPinActions, GetPinObj())
-			.OnAllowMakeType_UObject(OwningNode, &UNiagaraNodeWithDynamicPins::AllowNiagaraTypeForAddPin);
+			.OnAllowMakeType_UObject(OwningNode, &UNiagaraNodeWithDynamicPins::AllowNiagaraTypeForAddPin)
+			.IsParameterRead(GraphPinObj ? GraphPinObj->Direction == EGPD_Output : true);
 		
 		FSlateApplication::Get().SetKeyboardFocus(MenuWidget->GetSearchBox());
 		return MenuWidget;
diff --git a/Engine/Plugins/FX/Niagara/Source/NiagaraEditor/Private/Widgets/SNiagaraParameterMapView.cpp b/Engine/Plugins/FX/Niagara/Source/NiagaraEditor/Private/Widgets/SNiagaraParameterMapView.cpp
--- a/Engine/Plugins/FX/Niagara/Source/NiagaraEditor/Private/Widgets/SNiagaraParameterMapView.cpp
+++ b/Engine/Plugins/FX/Niagara/Source/NiagaraEditor/Private/Widgets/SNiagaraParameterMapView.cpp
@@ -719,6 +719,7 @@ void SNiagaraAddParameterMenu::Construct(const FArguments& InArgs, TArray<TWeakO
 	this->ShowNamespaceCategory = InArgs._ShowNamespaceCategory;
 	this->ShowGraphParameters = InArgs._ShowGraphParameters;
 	this->AutoExpandMenu = InArgs._AutoExpandMenu;
+	this->IsParameterRead = InArgs._IsParameterRead;
 
 	Graphs = InGraphs;
 
@@ -782,70 +783,89 @@ void SNiagaraAddParameterMenu::CollectAllActions(FGraphActionListBuilderBase& Ou
 		return ID == NiagaraParameterMapSectionID::NONE || (ID != NiagaraParameterMapSectionID::NONE && ID == GivenSectionID);
 	};
 
-	// Particle
-	if (CanCollectSection(NiagaraParameterMapSectionID::PARTICLE))
+	TArray<NiagaraParameterMapSectionID::Type> IDsExcluded;
+	// If this is a write node, exclude any read-only vars.
+	if (!IsParameterRead.Get())
 	{
-		bool bSupportsAttributes = true;
-		for (TWeakObjectPtr<UNiagaraGraph>& GraphWeakPtr : Graphs)
-		{
-			UNiagaraGraph* Graph = GraphWeakPtr.Get();
-			bool IsModule = Graph->FindOutputNode(ENiagaraScriptUsage::Module) != nullptr || Graph->FindOutputNode(ENiagaraScriptUsage::DynamicInput) != nullptr
-				|| Graph->FindOutputNode(ENiagaraScriptUsage::Function) != nullptr;
+		IDsExcluded.Add(NiagaraParameterMapSectionID::USER);
+		IDsExcluded.Add(NiagaraParameterMapSectionID::ENGINE);
+		IDsExcluded.Add(NiagaraParameterMapSectionID::PARAMETERCOLLECTION);
+	}
+
+	// If this doesn't have particles in the script, exclude reading or writing them.
+	for (TWeakObjectPtr<UNiagaraGraph>& GraphWeakPtr : Graphs)
+	{
+		UNiagaraGraph* Graph = GraphWeakPtr.Get();
+		bool IsModule = Graph->FindOutputNode(ENiagaraScriptUsage::Module) != nullptr || Graph->FindOutputNode(ENiagaraScriptUsage::DynamicInput) != nullptr
+			|| Graph->FindOutputNode(ENiagaraScriptUsage::Function) != nullptr;
 
-			UNiagaraScriptSource* Source = Cast<UNiagaraScriptSource>(Graph->GetOuter());
-			if (Source && IsModule)
+		UNiagaraScriptSource* Source = Cast<UNiagaraScriptSource>(Graph->GetOuter());
+		if (Source && IsModule)
+		{
+			UNiagaraScript* Script = Cast<UNiagaraScript>(Source->GetOuter());
+			if (Script)
 			{
-				UNiagaraScript* Script = Cast<UNiagaraScript>(Source->GetOuter());
-				if (Script)
+				TArray<ENiagaraScriptUsage> Usages = Script->GetSupportedUsageContexts();
+				if (!Usages.Contains(ENiagaraScriptUsage::ParticleEventScript) && 
+					!Usages.Contains(ENiagaraScriptUsage::ParticleSpawnScript) && 
+					!Usages.Contains(ENiagaraScriptUsage::ParticleUpdateScript))
+				{
+					IDsExcluded.Add(NiagaraParameterMapSectionID::PARTICLE);
+				}
+
+				if (!IsParameterRead.Get())
 				{
-					TArray<ENiagaraScriptUsage> Usages = Script->GetSupportedUsageContexts();
-					if (Usages.Contains(ENiagaraScriptUsage::ParticleEventScript) || Usages.Contains(ENiagaraScriptUsage::ParticleSpawnScript) || Usages.Contains(ENiagaraScriptUsage::ParticleUpdateScript))
+					if (!Usages.Contains(ENiagaraScriptUsage::SystemSpawnScript) &&
+						!Usages.Contains(ENiagaraScriptUsage::SystemUpdateScript))
 					{
-						bSupportsAttributes = true;
+						IDsExcluded.Add(NiagaraParameterMapSectionID::SYSTEM);
 					}
-					else
+
+					if (!Usages.Contains(ENiagaraScriptUsage::EmitterSpawnScript) &&
+						!Usages.Contains(ENiagaraScriptUsage::EmitterUpdateScript))
 					{
-						bSupportsAttributes = false;
+						IDsExcluded.Add(NiagaraParameterMapSectionID::EMITTER);
 					}
 				}
 			}
 		}
-
-		if (bSupportsAttributes)
-		{
-			const FText Category = ShowNamespaceCategory.Get() ? NiagaraParameterMapSectionID::OnGetSectionTitle(NiagaraParameterMapSectionID::PARTICLE) : FText::GetEmpty();
-			TArray<FNiagaraVariable> Variables = FNiagaraConstants::GetCommonParticleAttributes();
-			AddParameterGroup(OutAllActions, Variables, NiagaraParameterMapSectionID::PARTICLE, Category, FString(), true, false);
-			CollectMakeNew(OutAllActions, NiagaraParameterMapSectionID::PARTICLE);
-		}
+	}
+	// Particle
+	if (CanCollectSection(NiagaraParameterMapSectionID::PARTICLE) && !IDsExcluded.Contains(NiagaraParameterMapSectionID::PARTICLE))
+	{
+		const FText Category = ShowNamespaceCategory.Get() ? NiagaraParameterMapSectionID::OnGetSectionTitle(NiagaraParameterMapSectionID::PARTICLE) : FText::GetEmpty();
+		TArray<FNiagaraVariable> Variables;
+		Variables = FNiagaraConstants::GetCommonParticleAttributes();
+		AddParameterGroup(OutAllActions, Variables, NiagaraParameterMapSectionID::PARTICLE, Category, FString(), true, false);
+		CollectMakeNew(OutAllActions, NiagaraParameterMapSectionID::PARTICLE);
 	}
 
 	// Emitter
-	if (CanCollectSection(NiagaraParameterMapSectionID::EMITTER))
+	if (CanCollectSection(NiagaraParameterMapSectionID::EMITTER) && !IDsExcluded.Contains(NiagaraParameterMapSectionID::EMITTER))
 	{
 		CollectMakeNew(OutAllActions, NiagaraParameterMapSectionID::EMITTER);
 	}
 
 	// Module
-	if (CanCollectSection(NiagaraParameterMapSectionID::MODULE))
+	if (CanCollectSection(NiagaraParameterMapSectionID::MODULE) && !IDsExcluded.Contains(NiagaraParameterMapSectionID::MODULE))
 	{
 		CollectMakeNew(OutAllActions, NiagaraParameterMapSectionID::MODULE);
 	}
 
 	// System
-	if (CanCollectSection(NiagaraParameterMapSectionID::SYSTEM))
+	if (CanCollectSection(NiagaraParameterMapSectionID::SYSTEM) && !IDsExcluded.Contains(NiagaraParameterMapSectionID::SYSTEM))
 	{
 		CollectMakeNew(OutAllActions, NiagaraParameterMapSectionID::SYSTEM);
 	}
 
 	// User
-	if (CanCollectSection(NiagaraParameterMapSectionID::USER))
+	if (CanCollectSection(NiagaraParameterMapSectionID::USER) && !IDsExcluded.Contains(NiagaraParameterMapSectionID::USER))
 	{
 		CollectMakeNew(OutAllActions, NiagaraParameterMapSectionID::USER);
 	}
 
 	// Parameter collections
-	if (CanCollectSection(NiagaraParameterMapSectionID::PARAMETERCOLLECTION))
+	if (CanCollectSection(NiagaraParameterMapSectionID::PARAMETERCOLLECTION) && !IDsExcluded.Contains(NiagaraParameterMapSectionID::PARAMETERCOLLECTION))
 	{
 		CollectParameterCollectionsActions(OutAllActions);
 	}
@@ -856,7 +876,7 @@ void SNiagaraAddParameterMenu::CollectAllActions(FGraphActionListBuilderBase& Ou
 	}
 
 	// Engine
-	if (CanCollectSection(NiagaraParameterMapSectionID::ENGINE))
+	if (CanCollectSection(NiagaraParameterMapSectionID::ENGINE) && !IDsExcluded.Contains(NiagaraParameterMapSectionID::ENGINE))
 	{
 		const FText Category = NiagaraParameterMapSectionID::OnGetSectionTitle(NiagaraParameterMapSectionID::ENGINE);
 		TArray<FNiagaraVariable> Variables = FNiagaraConstants::GetEngineConstants();
@@ -882,6 +902,11 @@ void SNiagaraAddParameterMenu::CollectAllActions(FGraphActionListBuilderBase& Ou
 				const NiagaraParameterMapSectionID::Type ParameterSectionID = NiagaraParameterMapSectionID::OnGetSectionFromVariable(Parameter, Handle, /*Default*/ NiagaraParameterMapSectionID::OTHER);
 				if (CanCollectSection(ParameterSectionID))
 				{
+					if (IDsExcluded.Contains(ParameterSectionID))
+					{
+						continue;
+					}
+
 					const FText Category = ShowNamespaceCategory.Get() ? NiagaraParameterMapSectionID::OnGetSectionTitle(ParameterSectionID) : FText::GetEmpty();
 					const FText DisplayName = FText::FromName(Parameter.GetName());
 
diff --git a/Engine/Plugins/FX/Niagara/Source/NiagaraEditor/Private/Widgets/SNiagaraParameterMapView.h b/Engine/Plugins/FX/Niagara/Source/NiagaraEditor/Private/Widgets/SNiagaraParameterMapView.h
--- a/Engine/Plugins/FX/Niagara/Source/NiagaraEditor/Private/Widgets/SNiagaraParameterMapView.h
+++ b/Engine/Plugins/FX/Niagara/Source/NiagaraEditor/Private/Widgets/SNiagaraParameterMapView.h
@@ -172,7 +172,8 @@ public:
 		, _AllowCreatingNew(true)
 		, _ShowNamespaceCategory(true)
 		, _ShowGraphParameters(true)
-		, _AutoExpandMenu(false) {}
+		, _AutoExpandMenu(false)
+		, _IsParameterRead(true) {}
 		SLATE_EVENT(FOnAddParameter, OnAddParameter)
 		SLATE_EVENT(FOnCollectCustomActions, OnCollectCustomActions)
 		SLATE_EVENT(FOnAllowMakeType, OnAllowMakeType)
@@ -181,6 +182,7 @@ public:
 		SLATE_ATTRIBUTE(bool, ShowNamespaceCategory)
 		SLATE_ATTRIBUTE(bool, ShowGraphParameters)
 		SLATE_ATTRIBUTE(bool, AutoExpandMenu)
+		SLATE_ATTRIBUTE(bool, IsParameterRead)
 	SLATE_END_ARGS();
 
 	void Construct(const FArguments& InArgs, TArray<TWeakObjectPtr<UNiagaraGraph>> InGraphs);
@@ -210,4 +212,5 @@ private:
 	TAttribute<bool> ShowNamespaceCategory;
 	TAttribute<bool> ShowGraphParameters;
 	TAttribute<bool> AutoExpandMenu;
+	TAttribute<bool> IsParameterRead;
 };
