Getting translation set up for Frank's rapid iteration parameter rework in support for default dynamic inputs

#tests n/a

diff --git a/Engine/Plugins/FX/Niagara/Source/Niagara/Public/NiagaraCommon.h b/Engine/Plugins/FX/Niagara/Source/Niagara/Public/NiagaraCommon.h
--- a/Engine/Plugins/FX/Niagara/Source/Niagara/Public/NiagaraCommon.h
+++ b/Engine/Plugins/FX/Niagara/Source/Niagara/Public/NiagaraCommon.h
@@ -66,6 +66,7 @@ enum class ENiagaraInputNodeUsage : uint8
 	Attribute,
 	SystemConstant,
 	TranslatorConstant,
+	RapidIterationParameter
 };
 
 /**
diff --git a/Engine/Plugins/FX/Niagara/Source/NiagaraEditor/Private/EdGraphSchema_Niagara.cpp b/Engine/Plugins/FX/Niagara/Source/NiagaraEditor/Private/EdGraphSchema_Niagara.cpp
--- a/Engine/Plugins/FX/Niagara/Source/NiagaraEditor/Private/EdGraphSchema_Niagara.cpp
+++ b/Engine/Plugins/FX/Niagara/Source/NiagaraEditor/Private/EdGraphSchema_Niagara.cpp
@@ -54,6 +54,7 @@ const FLinearColor UEdGraphSchema_Niagara::NodeTitleColor_FunctionCall = FLinear
 const FLinearColor UEdGraphSchema_Niagara::NodeTitleColor_CustomHlsl = FLinearColor::Yellow;
 const FLinearColor UEdGraphSchema_Niagara::NodeTitleColor_Event = FLinearColor::Red;
 const FLinearColor UEdGraphSchema_Niagara::NodeTitleColor_TranslatorConstant = FLinearColor::Gray;
+const FLinearColor UEdGraphSchema_Niagara::NodeTitleColor_RapidIteration = FLinearColor::Black;
 
 const FName UEdGraphSchema_Niagara::PinCategoryType("Type");
 const FName UEdGraphSchema_Niagara::PinCategoryMisc("Misc");
@@ -798,6 +799,27 @@ TArray<TSharedPtr<FNiagaraSchemaAction_NewNode> > UEdGraphSchema_Niagara::GetGra
 					InputAction->NodeTemplate = InputNode;
 				}
 
+				// TODO sckime please remove this..
+				for (FNiagaraTypeDefinition Type : RegisteredTypes)
+				{
+					FText MenuCat;
+					if (const UClass* Class = Type.GetClass())
+					{
+						continue;
+					}
+					else
+					{
+						MenuCat = LOCTEXT("AddRIParameterCat", "Add Rapid Iteration Param");
+					}
+
+					const FText MenuDesc = FText::Format(MenuDescFmt, Type.GetStruct()->GetDisplayNameText());
+					TSharedPtr<FNiagaraSchemaAction_NewNode> InputAction = AddNewNodeAction(NewActions, MenuCat, MenuDesc, *MenuDesc.ToString(), FText::GetEmpty());
+					UNiagaraNodeInput* InputNode = NewObject<UNiagaraNodeInput>(OwnerOfTemporaries);
+					FNiagaraEditorUtilities::InitializeParameterInputNode(*InputNode, Type, NiagaraGraph);
+					InputNode->Usage = ENiagaraInputNodeUsage::RapidIterationParameter;
+					InputAction->NodeTemplate = InputNode;
+				}
+				
 				if (PinType != FNiagaraTypeDefinition::GetGenericNumericDef())
 				{
 					//For correctly typed pins, offer the correct type at the top level.				
diff --git a/Engine/Plugins/FX/Niagara/Source/NiagaraEditor/Private/NiagaraHlslTranslator.cpp b/Engine/Plugins/FX/Niagara/Source/NiagaraEditor/Private/NiagaraHlslTranslator.cpp
--- a/Engine/Plugins/FX/Niagara/Source/NiagaraEditor/Private/NiagaraHlslTranslator.cpp
+++ b/Engine/Plugins/FX/Niagara/Source/NiagaraEditor/Private/NiagaraHlslTranslator.cpp
@@ -2345,6 +2345,67 @@ bool FHlslNiagaraTranslator::ShouldInterpolateParameter(const FNiagaraVariable&
 	return true;
 }
 
+int32 FHlslNiagaraTranslator::GetRapidIterationParameter(const FNiagaraVariable& Parameter)
+{
+	if (!AddStructToDefinitionSet(Parameter.GetType()))
+	{
+		Error(FText::Format(LOCTEXT("GetRapidIterationParameterTypeFail", "Cannot handle type {0}! Variable: {1}"), Parameter.GetType().GetNameText(), FText::FromName(Parameter.GetName())), nullptr, nullptr);
+		return INDEX_NONE;
+	}
+
+	int32 FuncParam = INDEX_NONE;
+	if (GetFunctionParameter(Parameter, FuncParam))
+	{
+		Error(FText::Format(LOCTEXT("GetRapidIterationParameterFuncParamFail", "Variable: {0} cannot be a function parameter because it is a RapidIterationParameter type."), FText::FromName(Parameter.GetName())), nullptr, nullptr);
+		return INDEX_NONE;
+	}
+
+	bool bIsCandidateForRapidIteration = false;
+	if (ActiveHistoryForFunctionCalls.InTopLevelFunctionCall(CompileOptions.TargetUsage))
+	{
+		if (Parameter.GetType() != FNiagaraTypeDefinition::GetBoolDef() && !Parameter.GetType().IsEnum() && !Parameter.GetType().IsDataInterface())
+		{
+			bIsCandidateForRapidIteration = true;
+		}
+		else
+		{
+			Error(FText::Format(LOCTEXT("GetRapidIterationParameterTypeFail", "Variable: {0} cannot be a RapidIterationParameter input node because it isn't a supported type {1}"), FText::FromName(Parameter.GetName()), Parameter.GetType().GetNameText()), nullptr, nullptr);
+			return INDEX_NONE;
+		}
+	}
+	else
+	{
+		Error(FText::Format(LOCTEXT("GetRapidIterationParameterInTopLevelFail", "Variable: {0} cannot be a RapidIterationParameter input node because it isn't in the top level of an emitter/system/particle graph."), FText::FromName(Parameter.GetName())), nullptr, nullptr);
+		return INDEX_NONE;
+	}
+
+	FNiagaraVariable RapidIterationConstantVar = Parameter;
+
+	int32 LastSetChunkIdx = INDEX_NONE;
+	// Check to see if this is the first time we've encountered this node and it is a viable candidate for rapid iteration
+	if (bIsCandidateForRapidIteration && TranslationOptions.bParameterRapidIteration)
+	{
+		
+		// go ahead and make it into a constant variable..
+		int32 OutputChunkId = INDEX_NONE;
+		if (ParameterMapRegisterExternalConstantNamespaceVariable(Parameter, nullptr, INDEX_NONE, OutputChunkId, nullptr))
+		{
+			return OutputChunkId;
+		}
+	}
+	else
+	{
+		int32 FoundIdx = TranslationOptions.OverrideModuleConstants.Find(RapidIterationConstantVar);
+		if (FoundIdx != INDEX_NONE)
+		{
+			int32 OutputChunkId = GetConstant(TranslationOptions.OverrideModuleConstants[FoundIdx]);
+			return OutputChunkId;
+		}
+	}
+
+	return INDEX_NONE;
+}
+
 int32 FHlslNiagaraTranslator::GetParameter(const FNiagaraVariable& Parameter)
 {
 
@@ -2962,10 +3023,14 @@ bool FHlslNiagaraTranslator::ParameterMapRegisterExternalConstantNamespaceVariab
 		}
 
 		bool bMissingParameter = false;
-		UNiagaraParameterCollection* Collection = ParamMapHistories[InParamMapHistoryIdx].IsParameterCollectionParameter(InVariable, bMissingParameter);
-		if (Collection && bMissingParameter)
+		UNiagaraParameterCollection* Collection = nullptr;
+		if (InParamMapHistoryIdx >= 0)
 		{
-			Error(FText::Format(LOCTEXT("MissingNPCParameterError", "Parameter {0} was not found in Parameter Collection {1}"), FText::FromName(InVariable.GetName()), FText::FromString(Collection->GetFullName())), InNode, nullptr);
+			Collection = ParamMapHistories[InParamMapHistoryIdx].IsParameterCollectionParameter(InVariable, bMissingParameter);
+			if (Collection && bMissingParameter)
+			{
+				Error(FText::Format(LOCTEXT("MissingNPCParameterError", "Parameter {0} was not found in Parameter Collection {1}"), FText::FromName(InVariable.GetName()), FText::FromString(Collection->GetFullName())), InNode, nullptr);
+			}
 		}
 
 		bool bIsDataInterface = InVariable.GetType().IsDataInterface();
diff --git a/Engine/Plugins/FX/Niagara/Source/NiagaraEditor/Private/NiagaraHlslTranslator.h b/Engine/Plugins/FX/Niagara/Source/NiagaraEditor/Private/NiagaraHlslTranslator.h
--- a/Engine/Plugins/FX/Niagara/Source/NiagaraEditor/Private/NiagaraHlslTranslator.h
+++ b/Engine/Plugins/FX/Niagara/Source/NiagaraEditor/Private/NiagaraHlslTranslator.h
@@ -408,6 +408,8 @@ public:
 	virtual void Output(UNiagaraNodeOutput* OutputNode, const TArray<int32>& ComputedInputs);
 
 	virtual int32 GetParameter(const FNiagaraVariable& Parameter);
+	virtual int32 GetRapidIterationParameter(const FNiagaraVariable& Parameter);
+
 
 	virtual int32 GetAttribute(const FNiagaraVariable& Attribute);
 
diff --git a/Engine/Plugins/FX/Niagara/Source/NiagaraEditor/Private/NiagaraNodeInput.cpp b/Engine/Plugins/FX/Niagara/Source/NiagaraEditor/Private/NiagaraNodeInput.cpp
--- a/Engine/Plugins/FX/Niagara/Source/NiagaraEditor/Private/NiagaraNodeInput.cpp
+++ b/Engine/Plugins/FX/Niagara/Source/NiagaraEditor/Private/NiagaraNodeInput.cpp
@@ -319,6 +319,8 @@ FLinearColor UNiagaraNodeInput::GetNodeTitleColor() const
 		return Schema->NodeTitleColor_Attribute;
 	case ENiagaraInputNodeUsage::TranslatorConstant:
 		return Schema->NodeTitleColor_TranslatorConstant;
+	case ENiagaraInputNodeUsage::RapidIterationParameter:
+		return Schema->NodeTitleColor_RapidIteration;
 	default:
 		// TODO: Do something better here.
 		return FLinearColor::Black;
@@ -503,6 +505,8 @@ void UNiagaraNodeInput::Compile(class FHlslNiagaraTranslator* Translator, TArray
 		Outputs.Add(Translator->GetAttribute(Input)); break;
 	case ENiagaraInputNodeUsage::TranslatorConstant:
 		Outputs.Add(Translator->GetParameter(Input)); break;
+	case ENiagaraInputNodeUsage::RapidIterationParameter:
+		Outputs.Add(Translator->GetRapidIterationParameter(Input)); break;
 	default:
 		check(false);
 	}
diff --git a/Engine/Plugins/FX/Niagara/Source/NiagaraEditor/Private/Widgets/SNiagaraGraphNodeInput.cpp b/Engine/Plugins/FX/Niagara/Source/NiagaraEditor/Private/Widgets/SNiagaraGraphNodeInput.cpp
--- a/Engine/Plugins/FX/Niagara/Source/NiagaraEditor/Private/Widgets/SNiagaraGraphNodeInput.cpp
+++ b/Engine/Plugins/FX/Niagara/Source/NiagaraEditor/Private/Widgets/SNiagaraGraphNodeInput.cpp
@@ -64,7 +64,12 @@ void SNiagaraGraphNodeInput::Construct(const FArguments& InArgs, UEdGraphNode* I
 bool SNiagaraGraphNodeInput::IsNameReadOnly() const
 {
 	const UNiagaraNodeInput* InputNode = CastChecked<UNiagaraNodeInput>(GraphNode);
-	return InputNode->Usage != ENiagaraInputNodeUsage::Parameter;
+	if (InputNode->Usage == ENiagaraInputNodeUsage::Parameter ||
+		InputNode->Usage == ENiagaraInputNodeUsage::RapidIterationParameter)
+	{
+		return false;
+	}
+	return true;
 }
 
 void SNiagaraGraphNodeInput::RequestRenameOnSpawn()
diff --git a/Engine/Plugins/FX/Niagara/Source/NiagaraEditor/Public/EdGraphSchema_Niagara.h b/Engine/Plugins/FX/Niagara/Source/NiagaraEditor/Public/EdGraphSchema_Niagara.h
--- a/Engine/Plugins/FX/Niagara/Source/NiagaraEditor/Public/EdGraphSchema_Niagara.h
+++ b/Engine/Plugins/FX/Niagara/Source/NiagaraEditor/Public/EdGraphSchema_Niagara.h
@@ -119,6 +119,7 @@ class NIAGARAEDITOR_API UEdGraphSchema_Niagara : public UEdGraphSchema
 	static const FLinearColor NodeTitleColor_CustomHlsl;
 	static const FLinearColor NodeTitleColor_Event; 
 	static const FLinearColor NodeTitleColor_TranslatorConstant;
+	static const FLinearColor NodeTitleColor_RapidIteration;
 	
 private:
 	void GetBreakLinkToSubMenuActions(class FMenuBuilder& MenuBuilder, UEdGraphPin* InGraphPin);
