Fix renaming variables will not work if the user is only changing capitalization.

diff --git a/Engine/Plugins/FX/Niagara/Source/NiagaraEditor/Private/Widgets/SNiagaraParameterMapView.cpp b/Engine/Plugins/FX/Niagara/Source/NiagaraEditor/Private/Widgets/SNiagaraParameterMapView.cpp
--- a/Engine/Plugins/FX/Niagara/Source/NiagaraEditor/Private/Widgets/SNiagaraParameterMapView.cpp
+++ b/Engine/Plugins/FX/Niagara/Source/NiagaraEditor/Private/Widgets/SNiagaraParameterMapView.cpp
@@ -654,7 +654,7 @@ bool SNiagaraParameterMapView::CanRequestRenameOnActionNode() const
 void SNiagaraParameterMapView::OnPostRenameActionNode(const FText& InText, FNiagaraParameterAction& InAction)
 {
 	const FName NewName = FName(*InText.ToString());
-	if (InAction.Parameter.GetName() != NewName)
+	if (!InAction.Parameter.GetName().IsEqual(NewName, ENameCase::CaseSensitive))
 	{
 		FNiagaraParameterHandle ParameterHandle;
 		if (NiagaraParameterMapSectionID::OnGetSectionFromVariable(InAction.Parameter, ParameterHandle) == NiagaraParameterMapSectionID::USER)
