Fix adding a new attribute makes it hidden in the attribute spreadsheet. Now shows the added attribute when doing a new capture.

diff --git a/Engine/Plugins/FX/Niagara/Source/NiagaraEditor/Private/Widgets/SNiagaraSpreadsheetView.cpp b/Engine/Plugins/FX/Niagara/Source/NiagaraEditor/Private/Widgets/SNiagaraSpreadsheetView.cpp
--- a/Engine/Plugins/FX/Niagara/Source/NiagaraEditor/Private/Widgets/SNiagaraSpreadsheetView.cpp
+++ b/Engine/Plugins/FX/Niagara/Source/NiagaraEditor/Private/Widgets/SNiagaraSpreadsheetView.cpp
@@ -1175,6 +1175,7 @@ void SNiagaraSpreadsheetView::ResetColumns(EUITab Tab)
 		{
 			CaptureData[(int32)i].OutputHeaderRow->ClearColumns();
 
+			const TArray<FName>& PreviousSupportedFields = CaptureData[(int32)i].SupportedOutputFields.IsValid() ? *CaptureData[(int32)i].SupportedOutputFields.Get() : TArray<FName>();
 			CaptureData[(int32)i].SupportedOutputFields = MakeShared<TArray<FName> >();
 			CaptureData[(int32)i].OutputFieldInfoMap = MakeShared<TMap<FName, FieldInfo> >();
 			uint32 TotalFloatComponents = 0;
@@ -1213,6 +1214,7 @@ void SNiagaraSpreadsheetView::ResetColumns(EUITab Tab)
 
 				for (int32 VarIdx = 0; VarIdx < PropertyNames.Num(); VarIdx++)
 				{
+					const FName PropertyName = PropertyNames[VarIdx];
 					if (FieldInfos[VarIdx].bFloat)
 					{
 						FieldInfos[VarIdx].FloatStartOffset += TotalFloatComponentsBeforeStruct;
@@ -1224,13 +1226,19 @@ void SNiagaraSpreadsheetView::ResetColumns(EUITab Tab)
 						TotalInt32Components++;
 					}
 
-					CaptureData[i].SupportedOutputFields->Add(PropertyNames[VarIdx]);
-					CaptureData[i].OutputFieldInfoMap->Add(PropertyNames[VarIdx], FieldInfos[VarIdx]);
+					CaptureData[i].SupportedOutputFields->Add(PropertyName);
+					CaptureData[i].OutputFieldInfoMap->Add(PropertyName, FieldInfos[VarIdx]);
 					
+					// Show new attributes
+					if (!bInitialColumns && PreviousSupportedFields.Find(PropertyName) == INDEX_NONE)
+					{
+						CaptureData[(int32)Tab].FilteredOutputFields.AddUnique(PropertyName);
+					}
+
 					if (CaptureData[i].bOutputColumnsAreAttributes 
-						&& (bInitialColumns || IsOutputAttributeEnabled(Tab, PropertyNames[VarIdx])))
+						&& (bInitialColumns || IsOutputAttributeEnabled(Tab, PropertyName)))
 					{
-						ColumnNames.Add(PropertyNames[VarIdx]);
+						ColumnNames.Add(PropertyName);
 					}
 				}
 			}
