Reorder parameter list categories

diff --git a/Engine/Plugins/FX/Niagara/Source/NiagaraEditor/Private/Widgets/SNiagaraParameterMapView.cpp b/Engine/Plugins/FX/Niagara/Source/NiagaraEditor/Private/Widgets/SNiagaraParameterMapView.cpp
--- a/Engine/Plugins/FX/Niagara/Source/NiagaraEditor/Private/Widgets/SNiagaraParameterMapView.cpp
+++ b/Engine/Plugins/FX/Niagara/Source/NiagaraEditor/Private/Widgets/SNiagaraParameterMapView.cpp
@@ -363,11 +363,11 @@ void SNiagaraParameterMapView::CollectStaticSections(TArray<int32>& StaticSectio
 		StaticSectionIDs.Add(NiagaraParameterMapSectionID::MODULE);
 	}
 	StaticSectionIDs.Add(NiagaraParameterMapSectionID::ENGINE);
-	StaticSectionIDs.Add(NiagaraParameterMapSectionID::EMITTER);
+	StaticSectionIDs.Add(NiagaraParameterMapSectionID::PARAMETERCOLLECTION);
+	StaticSectionIDs.Add(NiagaraParameterMapSectionID::USER);
 	StaticSectionIDs.Add(NiagaraParameterMapSectionID::SYSTEM);
+	StaticSectionIDs.Add(NiagaraParameterMapSectionID::EMITTER);
 	StaticSectionIDs.Add(NiagaraParameterMapSectionID::PARTICLE);
-	StaticSectionIDs.Add(NiagaraParameterMapSectionID::USER);
-	StaticSectionIDs.Add(NiagaraParameterMapSectionID::PARAMETERCOLLECTION);
 	StaticSectionIDs.Add(NiagaraParameterMapSectionID::OTHER);
 }
 
diff --git a/Engine/Plugins/FX/Niagara/Source/NiagaraEditor/Private/Widgets/SNiagaraParameterMapView.h b/Engine/Plugins/FX/Niagara/Source/NiagaraEditor/Private/Widgets/SNiagaraParameterMapView.h
--- a/Engine/Plugins/FX/Niagara/Source/NiagaraEditor/Private/Widgets/SNiagaraParameterMapView.h
+++ b/Engine/Plugins/FX/Niagara/Source/NiagaraEditor/Private/Widgets/SNiagaraParameterMapView.h
@@ -30,12 +30,12 @@ namespace NiagaraParameterMapSectionID
 	{
 		NONE = 0,
 		MODULE,
-		EMITTER,
-		PARTICLE,
-		SYSTEM,
 		ENGINE,
-		USER,
 		PARAMETERCOLLECTION,
+		USER,
+		SYSTEM,
+		EMITTER,
+		PARTICLE,
 		OTHER,
 	};
 
