Disabled modules no longer display errors.
Also, enabling/disabling modules is now registered with the Undo system
Also fixed the GUID generation for all issues, now issues are properly differentiated from each other on refresh.
#tests none

diff --git a/Engine/Plugins/FX/Niagara/Source/NiagaraEditor/Private/ViewModels/Stack/NiagaraStackEntry.cpp b/Engine/Plugins/FX/Niagara/Source/NiagaraEditor/Private/ViewModels/Stack/NiagaraStackEntry.cpp
--- a/Engine/Plugins/FX/Niagara/Source/NiagaraEditor/Private/ViewModels/Stack/NiagaraStackEntry.cpp
+++ b/Engine/Plugins/FX/Niagara/Source/NiagaraEditor/Private/ViewModels/Stack/NiagaraStackEntry.cpp
@@ -46,16 +46,35 @@ UNiagaraStackEntry::FStackIssue::FStackIssue()
 {
 }
 
-UNiagaraStackEntry::FStackIssue::FStackIssue(EStackIssueSeverity InSeverity, FText InShortDescription, FText InLongDescription, FString InStackEditorDataKey, bool bInCanBeDismissed)
+UNiagaraStackEntry::FStackIssue::FStackIssue(EStackIssueSeverity InSeverity, FText InShortDescription, FText InLongDescription, FString InStackEditorDataKey, bool bInCanBeDismissed, TArray<FStackIssueFix>& InFixes)
 	: Severity(InSeverity)
 	, ShortDescription(InShortDescription)
 	, LongDescription(InLongDescription)
-	, UniqueIdentifier(FMD5::HashAnsiString(*FString::Printf(TEXT("%s-%s"), *InStackEditorDataKey, *LongDescription.ToString())))
 	, bCanBeDismissed(bInCanBeDismissed)
+	, Fixes(InFixes)
 {
 	checkf(ShortDescription.IsEmptyOrWhitespace() == false, TEXT("Short description can not be empty."));
 	checkf(LongDescription.IsEmptyOrWhitespace() == false, TEXT("Long description can not be empty."));
 	checkf(InStackEditorDataKey.IsEmpty() == false, TEXT("Stack editor data key can not be empty."));
+	FString Descriptions = LongDescription.ToString();
+	for (FStackIssueFix Fix : Fixes)
+	{
+		Descriptions += Fix.GetDescription().ToString();
+	}
+	UniqueIdentifier = FMD5::HashAnsiString(*FString::Printf(TEXT("%s-%s"), *InStackEditorDataKey, *Descriptions));
+}
+
+UNiagaraStackEntry::FStackIssue::FStackIssue(EStackIssueSeverity InSeverity, FText InShortDescription, FText InLongDescription, FString InStackEditorDataKey, bool bInCanBeDismissed, FStackIssueFix InFix)
+{
+	TArray<FStackIssueFix> FixArray;
+	FixArray.Add(InFix);
+	FStackIssue(InSeverity, InShortDescription, InLongDescription, InStackEditorDataKey, bInCanBeDismissed, FixArray);
+}
+
+UNiagaraStackEntry::FStackIssue::FStackIssue(EStackIssueSeverity InSeverity, FText InShortDescription, FText InLongDescription, FString InStackEditorDataKey, bool bInCanBeDismissed)
+{
+	TArray<FStackIssueFix> FixArray;
+	FStackIssue(InSeverity, InShortDescription, InLongDescription, InStackEditorDataKey, bInCanBeDismissed, FixArray);
 }
 
 bool UNiagaraStackEntry::FStackIssue::IsValid()
@@ -88,13 +107,6 @@ const FString& UNiagaraStackEntry::FStackIssue::GetUniqueIdentifier() const
 	return UniqueIdentifier;
 }
 
-void UNiagaraStackEntry::FStackIssue::AddFix(FStackIssueFix InFix)
-{
-	checkf(IsValid(), TEXT("Issue is not valid."));
-	checkf(InFix.IsValid(), TEXT("Fix must be valid"));
-	Fixes.Add(InFix);
-}
-
 const TArray<UNiagaraStackEntry::FStackIssueFix>& UNiagaraStackEntry::FStackIssue::GetFixes() const
 {
 	return Fixes;
diff --git a/Engine/Plugins/FX/Niagara/Source/NiagaraEditor/Private/ViewModels/Stack/NiagaraStackFunctionInputCollection.cpp b/Engine/Plugins/FX/Niagara/Source/NiagaraEditor/Private/ViewModels/Stack/NiagaraStackFunctionInputCollection.cpp
--- a/Engine/Plugins/FX/Niagara/Source/NiagaraEditor/Private/ViewModels/Stack/NiagaraStackFunctionInputCollection.cpp
+++ b/Engine/Plugins/FX/Niagara/Source/NiagaraEditor/Private/ViewModels/Stack/NiagaraStackFunctionInputCollection.cpp
@@ -195,33 +195,33 @@ void UNiagaraStackFunctionInputCollection::RefreshIssues(TArray<FName> Duplicate
 				FNiagaraParameterHandle(OverridePin->PinName).GetNamespace().ToString() == InputFunctionCallNode->GetFunctionName() &&
 				ValidAliasedInputNames.Contains(OverridePin->PinName) == false)
 			{
-				FStackIssue InvalidInputOverrideError(
-					EStackIssueSeverity::Error,
-					FText::Format(LOCTEXT("InvalidInputSummaryFormat", "Invalid Input Override: {0}"), FText::FromString(OverridePin->PinName.ToString())),
-					FText::Format(LOCTEXT("InvalidInputFormat", "The input {0} was previously overriden but is no longer exposed by the function {1}.\nPress the fix button to remove this unused override data,\nor check the function definition to see why this input is no longer exposed."),
-						FText::FromString(OverridePin->PinName.ToString()), FText::FromString(InputFunctionCallNode->GetFunctionName())),
-					GetStackEditorDataKey(),
-					false);
-
 				FText FixDescription = LOCTEXT("RemoveInvalidInputTransaction", "Remove invalid input override.");
 				FStackIssueFix RemoveInputOverrideFix(
 					FixDescription,
 					UNiagaraStackEntry::FStackIssueFixDelegate::CreateLambda([=]()
+				{
+					FScopedTransaction ScopedTransaction(FixDescription);
+					TArray<TWeakObjectPtr<UNiagaraDataInterface>> RemovedDataObjects;
+					FNiagaraStackGraphUtilities::RemoveNodesForStackFunctionInputOverridePin(*OverridePin, RemovedDataObjects);
+					for (TWeakObjectPtr<UNiagaraDataInterface> RemovedDataObject : RemovedDataObjects)
 					{
-						FScopedTransaction ScopedTransaction(FixDescription);
-						TArray<TWeakObjectPtr<UNiagaraDataInterface>> RemovedDataObjects;
-						FNiagaraStackGraphUtilities::RemoveNodesForStackFunctionInputOverridePin(*OverridePin, RemovedDataObjects);
-						for (TWeakObjectPtr<UNiagaraDataInterface> RemovedDataObject : RemovedDataObjects)
+						if (RemovedDataObject.IsValid())
 						{
-							if (RemovedDataObject.IsValid())
-							{
-								OnDataObjectModified().Broadcast(RemovedDataObject.Get());
-							}
+							OnDataObjectModified().Broadcast(RemovedDataObject.Get());
 						}
-						OverridePin->GetOwningNode()->RemovePin(OverridePin);
-					}));
+					}
+					OverridePin->GetOwningNode()->RemovePin(OverridePin);
+				}));
+
+				FStackIssue InvalidInputOverrideError(
+					EStackIssueSeverity::Error,
+					FText::Format(LOCTEXT("InvalidInputSummaryFormat", "Invalid Input Override: {0}"), FText::FromString(OverridePin->PinName.ToString())),
+					FText::Format(LOCTEXT("InvalidInputFormat", "The input {0} was previously overriden but is no longer exposed by the function {1}.\nPress the fix button to remove this unused override data,\nor check the function definition to see why this input is no longer exposed."),
+						FText::FromString(OverridePin->PinName.ToString()), FText::FromString(InputFunctionCallNode->GetFunctionName())),
+					GetStackEditorDataKey(),
+					false,
+					RemoveInputOverrideFix);
 
-				InvalidInputOverrideError.AddFix(RemoveInputOverrideFix);
 				NewIssues.Add(InvalidInputOverrideError);
 			}
 		}
diff --git a/Engine/Plugins/FX/Niagara/Source/NiagaraEditor/Private/ViewModels/Stack/NiagaraStackGraphUtilities.cpp b/Engine/Plugins/FX/Niagara/Source/NiagaraEditor/Private/ViewModels/Stack/NiagaraStackGraphUtilities.cpp
--- a/Engine/Plugins/FX/Niagara/Source/NiagaraEditor/Private/ViewModels/Stack/NiagaraStackGraphUtilities.cpp
+++ b/Engine/Plugins/FX/Niagara/Source/NiagaraEditor/Private/ViewModels/Stack/NiagaraStackGraphUtilities.cpp
@@ -965,6 +965,7 @@ TOptional<bool> FNiagaraStackGraphUtilities::GetModuleIsEnabled(UNiagaraNodeFunc
 
 void FNiagaraStackGraphUtilities::SetModuleIsEnabled(UNiagaraNodeFunctionCall& FunctionCallNode, bool bIsEnabled)
 {
+	FunctionCallNode.Modify();
 	TArray<UNiagaraNode*> ModuleNodes;
 	GetAllNodesForModule(FunctionCallNode, ModuleNodes);
 	for (UNiagaraNode* ModuleNode : ModuleNodes)
diff --git a/Engine/Plugins/FX/Niagara/Source/NiagaraEditor/Private/ViewModels/Stack/NiagaraStackModuleItem.cpp b/Engine/Plugins/FX/Niagara/Source/NiagaraEditor/Private/ViewModels/Stack/NiagaraStackModuleItem.cpp
--- a/Engine/Plugins/FX/Niagara/Source/NiagaraEditor/Private/ViewModels/Stack/NiagaraStackModuleItem.cpp
+++ b/Engine/Plugins/FX/Niagara/Source/NiagaraEditor/Private/ViewModels/Stack/NiagaraStackModuleItem.cpp
@@ -235,6 +235,11 @@ void UNiagaraStackModuleItem::RefreshChildrenInternal(const TArray<UNiagaraStack
 
 void UNiagaraStackModuleItem::RefreshIssues(TArray<FStackIssue>& NewIssues)
 {
+	if (!GetIsEnabled())
+	{
+		NewIssues.Empty();
+		return;
+	}
 	if (FunctionCallNode != nullptr)
 	{
 		if (!FunctionCallNode->ScriptIsValid())
@@ -253,23 +258,21 @@ void UNiagaraStackModuleItem::RefreshIssues(TArray<FStackIssue>& NewIssues)
 		if (!IsEnabled.IsSet())
 		{
 			bIsEnabled = false;
+			FText FixDescription = LOCTEXT("EnableModule", "Enable module");
+			FStackIssueFix EnableFix(
+				FixDescription,
+				FStackIssueFixDelegate::CreateLambda([this, FixDescription]()
+			{
+				SetIsEnabled(true);;
+			}));
 			FStackIssue InconsistentEnabledError(
 				EStackIssueSeverity::Error,
 				LOCTEXT("InconsistentEnabledErrorSummary", "The enabled state for module is inconsistent."),
 				LOCTEXT("InconsistentEnabledError", "This module is using multiple functions and their enabled state is inconsistent.\nClick fix to make all of the functions for this module enabled."),
 				GetStackEditorDataKey(),
-				false);
+				false,
+				EnableFix);
 
-			FText FixDescription = LOCTEXT("EnableModule", "Enable module");
-			FStackIssueFix EnableFix(
-				FixDescription,
-				FStackIssueFixDelegate::CreateLambda([=]()
-				{
-					FScopedTransaction ScopedTransaction(FixDescription);
-					SetIsEnabled(true);
-				}));
-
-			InconsistentEnabledError.AddFix(EnableFix);
 			NewIssues.Add(InconsistentEnabledError);
 		}
 	}
@@ -343,40 +346,36 @@ void UNiagaraStackModuleItem::RefreshIssues(TArray<FStackIssue>& NewIssues)
 		}
 		if (bDependencyMet == false)
 		{
+			TArray<FStackIssueFix> Fixes;
 			DependenciesNeeded.Add(Dependency);
 			
 			FText DependencyTypeString = Dependency.Type == ENiagaraModuleDependencyType::PreDependency ? LOCTEXT("PreDependency", "pre-dependency") : LOCTEXT("PostDependency", "post-dependency");
-			UNiagaraStackEntry::FStackIssue Error(
-				EStackIssueSeverity::Error,
-				LOCTEXT("DependencyWarning", "The module has unmet dependencies."),
-				FText::Format(LOCTEXT("DependencyWarningLong", "The following {0} is not met: {1}"), DependencyTypeString, Dependency.Description),
-				FString::Printf(TEXT("%s-dependency-%s"), *GetStackEditorDataKey(), *Dependency.Id.ToString()),
-				true);
-			
 
 			for ( UNiagaraNodeFunctionCall* DisabledNode : DisabledDependencies) // module exists but disabled
 			{
 				UNiagaraStackEntry::FStackIssueFix Fix(
 					FText::Format(LOCTEXT("EnableDependency", "Enable dependency module {0}"), FText::FromString(DisabledNode->GetFunctionName())),
-					FStackIssueFixDelegate::CreateLambda([=]()
+					FStackIssueFixDelegate::CreateLambda([this, DisabledNode]()
 					{
 						FScopedTransaction ScopedTransaction(LOCTEXT("EnableDependencyModule", "Enable dependency module"));
-						DisabledNode->Modify();
 						FNiagaraStackGraphUtilities::SetModuleIsEnabled(*DisabledNode, true);
+						OnRequestFullRefresh().Broadcast();
+
 					}));
-				Error.AddFix(Fix);
+				Fixes.Add(Fix);
 			}
 
 			for (FNiagaraStackModuleData DisorderedNode : DisorderedDependencies) // module exists but is not in the correct order (and possibly also disabled)
 			{
 				bool bNeedsEnable = !DisorderedNode.ModuleNode->IsNodeEnabled();
-				FText AndEnableModule = bNeedsEnable ? LOCTEXT("AndEnableModule", "and enable it") : FText();
+				FText AndEnableModule = bNeedsEnable ? FText::Format(LOCTEXT("AndEnableDependency", "And enable dependency module {0}"), FText::FromString(DisorderedNode.ModuleNode->GetFunctionName())) : FText();
 				UNiagaraStackEntry::FStackIssueFix Fix(
-					FText::Format(LOCTEXT("ReorderDependency", "Reposition dependency module {0} in the correct order {1}"), FText::FromString(DisorderedNode.ModuleNode->GetFunctionName()), AndEnableModule),
-					FStackIssueFixDelegate::CreateLambda([=]()
+					FText::Format(LOCTEXT("ReorderDependency", "Reposition this module in the correct order related to {0} {1}"), FText::FromString(DisorderedNode.ModuleNode->GetFunctionName()), AndEnableModule),
+					FStackIssueFixDelegate::CreateLambda([this, bNeedsEnable, DisorderedNode, SystemModuleData, Dependency, ModuleIndex, OutputNode]()
 					{
 						FScopedTransaction ScopedTransaction(LOCTEXT("ReorderDependencyModule", "Reorder dependency module"));
-						DisorderedNode.ModuleNode->Modify();
+					
+						FunctionCallNode->Modify();
 						// reorder node
 						int32 CorrectIndex = Dependency.Type == ENiagaraModuleDependencyType::PostDependency ? DisorderedNode.Index : DisorderedNode.Index + 1;
 						checkf(ModuleIndex != INDEX_NONE, TEXT("Module data wasn't found in system for current module!"));
@@ -390,7 +389,7 @@ void UNiagaraStackModuleItem::RefreshIssues(TArray<FStackIssue>& NewIssues)
 						FNiagaraStackGraphUtilities::RelayoutGraph(*OutputNode->GetGraph());
 						OnRequestFullRefresh().Broadcast();
 					}));
-				Error.AddFix(Fix);
+				Fixes.Add(Fix);
 			}
 			if (DisorderedDependencies.Num() == 0 && DisabledDependencies.Num() == 0)
 			{
@@ -488,9 +487,16 @@ void UNiagaraStackModuleItem::RefreshIssues(TArray<FStackIssue>& NewIssues)
 						FNiagaraStackGraphUtilities::RelayoutGraph(*TargetOutputNode->GetGraph());
 						OnRequestFullRefresh().Broadcast();
 					}));
-					Error.AddFix(Fix);
+					Fixes.Add(Fix);
 				}
 			}
+			UNiagaraStackEntry::FStackIssue Error(
+				EStackIssueSeverity::Error,
+				LOCTEXT("DependencyWarning", "The module has unmet dependencies."),
+				FText::Format(LOCTEXT("DependencyWarningLong", "The following {0} is not met: {1}"), DependencyTypeString, Dependency.Description),
+				FString::Printf(TEXT("%s-dependency-%s"), *GetStackEditorDataKey(), *Dependency.Id.ToString()),
+				true,
+				Fixes);
 			NewIssues.Add(Error);
 		}
 	}
@@ -544,8 +550,10 @@ bool UNiagaraStackModuleItem::GetIsEnabled() const
 
 void UNiagaraStackModuleItem::SetIsEnabled(bool bInIsEnabled)
 {
+	FScopedTransaction ScopedTransaction(LOCTEXT("EnableDisableModule", "Enable/Disable Module"));
 	FNiagaraStackGraphUtilities::SetModuleIsEnabled(*FunctionCallNode, bInIsEnabled);
 	bIsEnabled = bInIsEnabled;
+	OnRequestFullRefresh().Broadcast();
 }
 
 void UNiagaraStackModuleItem::Delete()
diff --git a/Engine/Plugins/FX/Niagara/Source/NiagaraEditor/Private/ViewModels/Stack/NiagaraStackRendererItem.cpp b/Engine/Plugins/FX/Niagara/Source/NiagaraEditor/Private/ViewModels/Stack/NiagaraStackRendererItem.cpp
--- a/Engine/Plugins/FX/Niagara/Source/NiagaraEditor/Private/ViewModels/Stack/NiagaraStackRendererItem.cpp
+++ b/Engine/Plugins/FX/Niagara/Source/NiagaraEditor/Private/ViewModels/Stack/NiagaraStackRendererItem.cpp
@@ -236,30 +236,30 @@ void UNiagaraStackRendererItem::RefreshIssues(TArray<FStackIssue>& NewIssues)
 {
 	for (FNiagaraVariable Attribute : MissingAttributes)
 	{
+		FText FixDescription = LOCTEXT("AddMissingVariable", "Add missing variable");
+		FStackIssueFix AddAttributeFix(
+			FixDescription,
+			FStackIssueFixDelegate::CreateLambda([=]()
+		{
+			FScopedTransaction ScopedTransaction(FixDescription);
+			if (AddMissingVariable(GetEmitterViewModel()->GetEmitter(), Attribute))
+			{
+				FNotificationInfo Info(FText::Format(LOCTEXT("AddedVariableForFix", "Added {0} to the Spawn script to support the renderer."), FText::FromName(Attribute.GetName())));
+				Info.ExpireDuration = 5.0f;
+				Info.bFireAndForget = true;
+				Info.Image = FCoreStyle::Get().GetBrush(TEXT("MessageLog.Info"));
+				FSlateNotificationManager::Get().AddNotification(Info);
+			}
+		}));
+
 		FStackIssue MissingAttributeError(
 			EStackIssueSeverity::Error,
 			LOCTEXT("FailedRendererBindShort", "An attribute is missing."),
 			FText::Format(LOCTEXT("FailedRendererBind", "Missing attribute \"{0}\" of Type \"{1}\"."), FText::FromName(Attribute.GetName()), Attribute.GetType().GetNameText()),
 			GetStackEditorDataKey(),
-			false);
+			false,
+			AddAttributeFix);
 
-		FText FixDescription = LOCTEXT("AddMissingVariable", "Add missing variable");
-		FStackIssueFix AddAttributeFix(
-			FixDescription,
-			FStackIssueFixDelegate::CreateLambda([=]()
-			{
-				FScopedTransaction ScopedTransaction(FixDescription);
-				if (AddMissingVariable(GetEmitterViewModel()->GetEmitter(), Attribute))
-				{
-					FNotificationInfo Info(FText::Format(LOCTEXT("AddedVariableForFix", "Added {0} to the Spawn script to support the renderer."), FText::FromName(Attribute.GetName())));
-					Info.ExpireDuration = 5.0f;
-					Info.bFireAndForget = true;
-					Info.Image = FCoreStyle::Get().GetBrush(TEXT("MessageLog.Info"));
-					FSlateNotificationManager::Get().AddNotification(Info);
-				}
-			}));
-
-		MissingAttributeError.AddFix(AddAttributeFix);
 		NewIssues.Add(MissingAttributeError);
 	}
 
diff --git a/Engine/Plugins/FX/Niagara/Source/NiagaraEditor/Private/ViewModels/Stack/NiagaraStackScriptItemGroup.cpp b/Engine/Plugins/FX/Niagara/Source/NiagaraEditor/Private/ViewModels/Stack/NiagaraStackScriptItemGroup.cpp
--- a/Engine/Plugins/FX/Niagara/Source/NiagaraEditor/Private/ViewModels/Stack/NiagaraStackScriptItemGroup.cpp
+++ b/Engine/Plugins/FX/Niagara/Source/NiagaraEditor/Private/ViewModels/Stack/NiagaraStackScriptItemGroup.cpp
@@ -333,24 +333,25 @@ void UNiagaraStackScriptItemGroup::RefreshIssues(TArray<FStackIssue>& NewIssues)
 	
 	if (FNiagaraStackGraphUtilities::ValidateGraphForOutput(*Graph, ScriptUsage, ScriptUsageId, ErrorMessage) == false)
 	{
-		FStackIssue InvalidStackError(
-			EStackIssueSeverity::Error,
-			LOCTEXT("InvalidErrorSummaryText", "The stack data is invalid"),
-			LOCTEXT("InvalidErrorText", "The data used to generate the stack has been corrupted and can not be used.\nUsing the fix option will reset this part of the stack to its default empty state."),
-			GetStackEditorDataKey(),
-			false);
 
 		FText FixDescription = LOCTEXT("FixStackGraph", "Fix invalid stack graph");
 		FStackIssueFix ResetStackFix(
 			FixDescription,
 			FStackIssueFixDelegate::CreateLambda([=]()
-			{
-				FScopedTransaction ScopedTransaction(FixDescription);
-				FNiagaraStackGraphUtilities::ResetGraphForOutput(*Graph, ScriptUsage, ScriptUsageId);
-				FNiagaraStackGraphUtilities::RelayoutGraph(*Graph);
-			}));
+		{
+			FScopedTransaction ScopedTransaction(FixDescription);
+			FNiagaraStackGraphUtilities::ResetGraphForOutput(*Graph, ScriptUsage, ScriptUsageId);
+			FNiagaraStackGraphUtilities::RelayoutGraph(*Graph);
+		}));
+
+		FStackIssue InvalidStackError(
+			EStackIssueSeverity::Error,
+			LOCTEXT("InvalidErrorSummaryText", "The stack data is invalid"),
+			LOCTEXT("InvalidErrorText", "The data used to generate the stack has been corrupted and can not be used.\nUsing the fix option will reset this part of the stack to its default empty state."),
+			GetStackEditorDataKey(),
+			false,
+			ResetStackFix);
 
-		InvalidStackError.AddFix(ResetStackFix);
 		NewIssues.Add(InvalidStackError); 
 	}
 	else
@@ -372,31 +373,31 @@ void UNiagaraStackScriptItemGroup::RefreshIssues(TArray<FStackIssue>& NewIssues)
 			{
 				bForcedError = true;
 
-				FStackIssue MissingLifeCycleError(
-					EStackIssueSeverity::Error,
-					LOCTEXT("SystemLifeCycleWarning", "The stack needs a SystemLifeCycle module."),
-					LOCTEXT("MissingRequiredMode", "Missing required module."),
-					GetStackEditorDataKey(),
-					false);
-
 				FText FixDescription = LOCTEXT("AddingSystemLifecycleModule", "Adding System Lifecycle Module.");
 				FStackIssueFix AddLifeCycleFix(
 					FixDescription,
 					FStackIssueFixDelegate::CreateLambda([=]()
+				{
+					FScopedTransaction ScopedTransaction(FixDescription);
+					UNiagaraNodeFunctionCall* AddedModuleNode = FNiagaraStackGraphUtilities::AddScriptModuleToStack(ModuleScriptAsset, *MatchingOutputNode);
+					if (AddedModuleNode == nullptr)
 					{
-						FScopedTransaction ScopedTransaction(FixDescription);
-						UNiagaraNodeFunctionCall* AddedModuleNode = FNiagaraStackGraphUtilities::AddScriptModuleToStack(ModuleScriptAsset, *MatchingOutputNode);
-						if (AddedModuleNode == nullptr)
-						{
-							FNotificationInfo Info(LOCTEXT("FailedToAddSystemLifecycle", "Failed to add system life cycle module.\nCheck the log for errors."));
-							Info.ExpireDuration = 5.0f;
-							Info.bFireAndForget = true;
-							Info.Image = FCoreStyle::Get().GetBrush(TEXT("MessageLog.Error"));
-							FSlateNotificationManager::Get().AddNotification(Info);
-						}
-					}));
+						FNotificationInfo Info(LOCTEXT("FailedToAddSystemLifecycle", "Failed to add system life cycle module.\nCheck the log for errors."));
+						Info.ExpireDuration = 5.0f;
+						Info.bFireAndForget = true;
+						Info.Image = FCoreStyle::Get().GetBrush(TEXT("MessageLog.Error"));
+						FSlateNotificationManager::Get().AddNotification(Info);
+					}
+				}));
+
+				FStackIssue MissingLifeCycleError(
+					EStackIssueSeverity::Error,
+					LOCTEXT("SystemLifeCycleWarning", "The stack needs a SystemLifeCycle module."),
+					LOCTEXT("MissingRequiredMode", "Missing required module."),
+					GetStackEditorDataKey(),
+					false,
+					AddLifeCycleFix);
 
-				MissingLifeCycleError.AddFix(AddLifeCycleFix);
 				NewIssues.Add(MissingLifeCycleError);
 			}
 		}
diff --git a/Engine/Plugins/FX/Niagara/Source/NiagaraEditor/Public/ViewModels/Stack/NiagaraStackEntry.h b/Engine/Plugins/FX/Niagara/Source/NiagaraEditor/Public/ViewModels/Stack/NiagaraStackEntry.h
--- a/Engine/Plugins/FX/Niagara/Source/NiagaraEditor/Public/ViewModels/Stack/NiagaraStackEntry.h
+++ b/Engine/Plugins/FX/Niagara/Source/NiagaraEditor/Public/ViewModels/Stack/NiagaraStackEntry.h
@@ -128,6 +128,10 @@ public:
 	{
 		FStackIssue();
 
+		FStackIssue(EStackIssueSeverity InSeverity, FText InShortDescription, FText InLongDescription, FString InStackEditorDataKey, bool bInCanBeDismissed, TArray<FStackIssueFix>& InFixes);
+
+		FStackIssue(EStackIssueSeverity InSeverity, FText InShortDescription, FText InLongDescription, FString InStackEditorDataKey, bool bInCanBeDismissed, FStackIssueFix InFix);
+
 		FStackIssue(EStackIssueSeverity InSeverity, FText InShortDescription, FText InLongDescription, FString InStackEditorDataKey, bool bInCanBeDismissed);
 
 		bool IsValid();
@@ -142,8 +146,6 @@ public:
 
 		bool GetCanBeDismissed() const;
 
-		void AddFix(FStackIssueFix InFix);
-
 		const TArray<FStackIssueFix>& GetFixes() const;
 
 	private:
