Cleaning up UI for code view

#tests n/a

diff --git a/Engine/Plugins/FX/Niagara/Source/NiagaraEditor/Private/Widgets/SNiagaraGeneratedCodeView.cpp b/Engine/Plugins/FX/Niagara/Source/NiagaraEditor/Private/Widgets/SNiagaraGeneratedCodeView.cpp
--- a/Engine/Plugins/FX/Niagara/Source/NiagaraEditor/Private/Widgets/SNiagaraGeneratedCodeView.cpp
+++ b/Engine/Plugins/FX/Niagara/Source/NiagaraEditor/Private/Widgets/SNiagaraGeneratedCodeView.cpp
@@ -131,7 +131,7 @@ void SNiagaraGeneratedCodeView::Construct(const FArguments& InArgs, TSharedRef<F
 					[
 						SAssignNew(ScriptNameCombo, SComboButton)
 						.OnGetMenuContent(this, &SNiagaraGeneratedCodeView::MakeScriptMenu)
-						.ComboButtonStyle(FEditorStyle::Get(), "ContentBrowser.Filters.Style")
+						.ComboButtonStyle(FEditorStyle::Get(), "GenericFilters.ComboButtonStyle")
 						.ForegroundColor(FLinearColor::White)
 						.ContentPadding(0)
 						.ToolTipText(LOCTEXT("ScriptsToolTip", "Select a script to view below."))
@@ -145,7 +145,7 @@ void SNiagaraGeneratedCodeView::Construct(const FArguments& InArgs, TSharedRef<F
 								.Padding(2, 0, 0, 0)
 								[
 									SNew(STextBlock)
-									.TextStyle(FEditorStyle::Get(), "ContentBrowser.Filters.Text")
+									.TextStyle(FEditorStyle::Get(), "GenericFilters.TextStyle")
 									.Text(LOCTEXT("Scripts", "Scripts"))
 								]
 						]
@@ -489,10 +489,14 @@ void SNiagaraGeneratedCodeView::UpdateUI()
 							.ExternalScrollbar(GeneratedCode[i].VerticalScrollBar)
 							+ SScrollBox::Slot()
 							[
-								SAssignNew(GeneratedCode[i].Text, SMultiLineEditableText)
+								SAssignNew(GeneratedCode[i].Text, SMultiLineEditableTextBox)
 								.ClearTextSelectionOnFocusLoss(false)
 								.IsReadOnly(true)
-								.TextStyle(FNiagaraEditorStyle::Get(), "NiagaraEditor.CodeView.Hlsl.Normal")
+								.Style(FEditorStyle::Get(), "Log.TextBox")
+								.TextStyle(FEditorStyle::Get(), "Log.Normal")
+								.ForegroundColor(FLinearColor::Yellow)
+								.ReadOnlyForegroundColor(FLinearColor::Yellow)
+								//.TextStyle(FNiagaraEditorStyle::Get(), "NiagaraEditor.CodeView.Hlsl.Normal")
 							]
 						]
 					]
diff --git a/Engine/Plugins/FX/Niagara/Source/NiagaraEditor/Private/Widgets/SNiagaraGeneratedCodeView.h b/Engine/Plugins/FX/Niagara/Source/NiagaraEditor/Private/Widgets/SNiagaraGeneratedCodeView.h
--- a/Engine/Plugins/FX/Niagara/Source/NiagaraEditor/Private/Widgets/SNiagaraGeneratedCodeView.h
+++ b/Engine/Plugins/FX/Niagara/Source/NiagaraEditor/Private/Widgets/SNiagaraGeneratedCodeView.h
@@ -44,7 +44,7 @@ protected:
 		FGuid UsageId;
 
 		TArray<FString> HlslByLines;
-		TSharedPtr<SMultiLineEditableText> Text;
+		TSharedPtr<SMultiLineEditableTextBox> Text;
 		TSharedPtr<SScrollBar> HorizontalScrollBar;
 		TSharedPtr<SScrollBar> VerticalScrollBar;
 		TSharedPtr<SVerticalBox> Container;
