Fix for crash during Jonathan's deletion of the Set node in SolveForcesInVelocity.

#tests n/a

diff --git a/Engine/Plugins/FX/Niagara/Source/NiagaraEditor/Private/NiagaraParameterMapHistory.cpp b/Engine/Plugins/FX/Niagara/Source/NiagaraEditor/Private/NiagaraParameterMapHistory.cpp
--- a/Engine/Plugins/FX/Niagara/Source/NiagaraEditor/Private/NiagaraParameterMapHistory.cpp
+++ b/Engine/Plugins/FX/Niagara/Source/NiagaraEditor/Private/NiagaraParameterMapHistory.cpp
@@ -777,7 +777,7 @@ int32 FNiagaraParameterMapHistoryBuilder::RegisterParameterMapPin(int32 WhichPar
 
 int32 FNiagaraParameterMapHistoryBuilder::TraceParameterMapOutputPin(const UEdGraphPin* OutputPin)
 {
-	if (OutputPin->Direction == EEdGraphPinDirection::EGPD_Output)
+	if (OutputPin && OutputPin->Direction == EEdGraphPinDirection::EGPD_Output)
 	{
 		OutputPin = UNiagaraNode::TraceOutputPin(const_cast<UEdGraphPin*>(OutputPin));
 		if (OutputPin)
