Added missing abs for GPU sim

diff --git a/Engine/Plugins/FX/Niagara/Source/Niagara/Private/NiagaraDataInterfaceCurlNoise.cpp b/Engine/Plugins/FX/Niagara/Source/Niagara/Private/NiagaraDataInterfaceCurlNoise.cpp
--- a/Engine/Plugins/FX/Niagara/Source/Niagara/Private/NiagaraDataInterfaceCurlNoise.cpp
+++ b/Engine/Plugins/FX/Niagara/Source/Niagara/Private/NiagaraDataInterfaceCurlNoise.cpp
@@ -160,7 +160,7 @@ bool UNiagaraDataInterfaceCurlNoise::GetFunctionHLSL(const FName& DefinitionFunc
 	OutHLSL += TEXT("void ") + InstanceFunctionName + TEXT("(in float3 In_XYZ, out float3 Out_Value) \n{\n");
 	OutHLSL += TEXT("\t float3 a = trunc((In_XYZ*0.2) / 15.0);\n");
 	OutHLSL += TEXT("\t float3 ModXYZ = (In_XYZ*0.2) - a*15.0;\n");
-	OutHLSL += TEXT("\t int3 IntCoord = int3(ModXYZ.x, ModXYZ.y, ModXYZ.z);\n");
+	OutHLSL += TEXT("\t int3 IntCoord = abs(int3(ModXYZ.x, ModXYZ.y, ModXYZ.z));\n");
 	OutHLSL += TEXT("\t float3 frc = frac(ModXYZ);\n");
 	// Trilinear interpolation, as defined by https://en.wikipedia.org/wiki/Trilinear_interpolation
 	// HLSL lerp is defined as lerp(x, y, s) = x*(1-s) + y*s
