Also removing issues from disabled input collections and renderer items
#tests none

diff --git a/Engine/Plugins/FX/Niagara/Source/NiagaraEditor/Private/ViewModels/Stack/NiagaraStackFunctionInputCollection.cpp b/Engine/Plugins/FX/Niagara/Source/NiagaraEditor/Private/ViewModels/Stack/NiagaraStackFunctionInputCollection.cpp
--- a/Engine/Plugins/FX/Niagara/Source/NiagaraEditor/Private/ViewModels/Stack/NiagaraStackFunctionInputCollection.cpp
+++ b/Engine/Plugins/FX/Niagara/Source/NiagaraEditor/Private/ViewModels/Stack/NiagaraStackFunctionInputCollection.cpp
@@ -180,6 +180,11 @@ void UNiagaraStackFunctionInputCollection::RefreshChildrenInternal(const TArray<
 
 void UNiagaraStackFunctionInputCollection::RefreshIssues(TArray<FName> DuplicateInputNames, TArray<FName> ValidAliasedInputNames, TArray<const UEdGraphPin*> PinsWithInvalidTypes, TArray<FStackIssue>& NewIssues)
 {
+	if (!GetIsEnabled())
+	{
+		NewIssues.Empty();
+		return;
+	}
 	// Try to find function input overrides which are no longer valid so we can generate errors for them.
 	UNiagaraNodeParameterMapSet* OverrideNode = FNiagaraStackGraphUtilities::GetStackFunctionOverrideNode(*InputFunctionCallNode);
 	if (OverrideNode != nullptr)
diff --git a/Engine/Plugins/FX/Niagara/Source/NiagaraEditor/Private/ViewModels/Stack/NiagaraStackRendererItem.cpp b/Engine/Plugins/FX/Niagara/Source/NiagaraEditor/Private/ViewModels/Stack/NiagaraStackRendererItem.cpp
--- a/Engine/Plugins/FX/Niagara/Source/NiagaraEditor/Private/ViewModels/Stack/NiagaraStackRendererItem.cpp
+++ b/Engine/Plugins/FX/Niagara/Source/NiagaraEditor/Private/ViewModels/Stack/NiagaraStackRendererItem.cpp
@@ -234,6 +234,11 @@ void UNiagaraStackRendererItem::RefreshChildrenInternal(const TArray<UNiagaraSta
 
 void UNiagaraStackRendererItem::RefreshIssues(TArray<FStackIssue>& NewIssues)
 {
+	if (!GetIsEnabled())
+	{
+		NewIssues.Empty();
+		return;
+	}
 	for (FNiagaraVariable Attribute : MissingAttributes)
 	{
 		FText FixDescription = LOCTEXT("AddMissingVariable", "Add missing variable");
