Making emitter local space a constant embeded directly into emitter and particle scripts.
Allows a lot of optimization and exposes the value to emitter scripts properly.

diff --git a/Engine/Plugins/FX/Niagara/Source/Niagara/Classes/NiagaraConstants.h b/Engine/Plugins/FX/Niagara/Source/Niagara/Classes/NiagaraConstants.h
--- a/Engine/Plugins/FX/Niagara/Source/Niagara/Classes/NiagaraConstants.h
+++ b/Engine/Plugins/FX/Niagara/Source/Niagara/Classes/NiagaraConstants.h
@@ -4,6 +4,22 @@
 #include "NiagaraTypes.h"
 #include "NiagaraCommon.h"
 
+#define PARAM_MAP_NPC_STR TEXT("NPC.")
+#define PARAM_MAP_ENGINE_STR TEXT("Engine.")
+#define PARAM_MAP_ENGINE_OWNER_STR TEXT("Engine.Owner.")
+#define PARAM_MAP_ENGINE_SYSTEM_STR TEXT("Engine.System.")
+#define PARAM_MAP_ENGINE_EMITTER_STR TEXT("Engine.Emitter.")
+#define PARAM_MAP_USER_STR TEXT("User.")
+#define PARAM_MAP_SYSTEM_STR TEXT("System.")
+#define PARAM_MAP_EMITTER_STR TEXT("Emitter.")
+#define PARAM_MAP_MODULE_STR TEXT("Module.")
+#define PARAM_MAP_ATTRIBUTE_STR TEXT("Particles.")
+#define PARAM_MAP_INITIAL_STR TEXT("Initial.")
+#define PARAM_MAP_INITIAL_BASE_STR TEXT("Initial")
+#define PARAM_MAP_RAPID_ITERATION_STR TEXT("Constants.")
+#define PARAM_MAP_RAPID_ITERATION_BASE_STR TEXT("Constants")
+
+
 #define SYS_PARAM_ENGINE_DELTA_TIME					INiagaraModule::GetVar_Engine_DeltaTime()
 #define SYS_PARAM_ENGINE_INV_DELTA_TIME				INiagaraModule::GetVar_Engine_InvDeltaTime()
 #define SYS_PARAM_ENGINE_TIME						INiagaraModule::GetVar_Engine_Time()
diff --git a/Engine/Plugins/FX/Niagara/Source/Niagara/Classes/NiagaraEmitterInstance.h b/Engine/Plugins/FX/Niagara/Source/Niagara/Classes/NiagaraEmitterInstance.h
--- a/Engine/Plugins/FX/Niagara/Source/Niagara/Classes/NiagaraEmitterInstance.h
+++ b/Engine/Plugins/FX/Niagara/Source/Niagara/Classes/NiagaraEmitterInstance.h
@@ -140,16 +140,11 @@ private:
 	FNiagaraParameterDirectBinding<float> SpawnIntervalBindingGPU;
 	FNiagaraParameterDirectBinding<float> InterpSpawnStartBindingGPU;
 	FNiagaraParameterDirectBinding<float> EmitterAgeBindingGPU;
-	FNiagaraParameterDirectBinding<FNiagaraBool> EmitterLocalSpaceBindingGPU;
 
 	FNiagaraParameterDirectBinding<float> SpawnEmitterAgeBinding;
 	FNiagaraParameterDirectBinding<float> UpdateEmitterAgeBinding;
 	TArray<FNiagaraParameterDirectBinding<float>> EventEmitterAgeBindings;
 
-	FNiagaraParameterDirectBinding<FNiagaraBool> SpawnEmitterLocalSpaceBinding;
-	FNiagaraParameterDirectBinding<FNiagaraBool> UpdateEmitterLocalSpaceBinding;
-	TArray<FNiagaraParameterDirectBinding<FNiagaraBool>> EventEmitterLocalSpaceBindings;
-
 	FNiagaraParameterDirectBinding<int32> SpawnExecCountBinding;
 	FNiagaraParameterDirectBinding<int32> UpdateExecCountBinding;
 	TArray<FNiagaraParameterDirectBinding<int32>> EventExecCountBindings;
diff --git a/Engine/Plugins/FX/Niagara/Source/Niagara/Private/NiagaraEmitter.cpp b/Engine/Plugins/FX/Niagara/Source/Niagara/Private/NiagaraEmitter.cpp
--- a/Engine/Plugins/FX/Niagara/Source/Niagara/Private/NiagaraEmitter.cpp
+++ b/Engine/Plugins/FX/Niagara/Source/Niagara/Private/NiagaraEmitter.cpp
@@ -316,6 +316,17 @@ void UNiagaraEmitter::PostEditChangeProperty(struct FPropertyChangedEvent& Prope
 			GraphSource->MarkNotSynchronized(TEXT("Emitter Requires Persistent IDs changed."));
 		}
 
+#if WITH_EDITORONLY_DATA
+		UNiagaraSystem::RequestCompileForEmitter(this);
+#endif
+	}
+	if (PropertyName == GET_MEMBER_NAME_CHECKED(UNiagaraEmitter, bLocalSpace))
+	{
+		if (GraphSource != nullptr)
+		{
+			GraphSource->MarkNotSynchronized(TEXT("Emitter LocalSpace changed."));
+		}
+
 #if WITH_EDITORONLY_DATA
 		UNiagaraSystem::RequestCompileForEmitter(this);
 #endif
diff --git a/Engine/Plugins/FX/Niagara/Source/Niagara/Private/NiagaraEmitterInstance.cpp b/Engine/Plugins/FX/Niagara/Source/Niagara/Private/NiagaraEmitterInstance.cpp
--- a/Engine/Plugins/FX/Niagara/Source/Niagara/Private/NiagaraEmitterInstance.cpp
+++ b/Engine/Plugins/FX/Niagara/Source/Niagara/Private/NiagaraEmitterInstance.cpp
@@ -253,19 +253,6 @@ void FNiagaraEmitterInstance::Init(int32 InEmitterIdx, FName InSystemInstanceNam
 		EmitterAgeBindingGPU.Init(GPUExecContext.CombinedParamStore, EmitterAgeParam);
 	}
 
-	FNiagaraVariable EmitterLocalSpaceParam = CachedEmitter->ToEmitterParameter(SYS_PARAM_EMITTER_LOCALSPACE);
-	SpawnEmitterLocalSpaceBinding.Init(SpawnExecContext.Parameters, EmitterLocalSpaceParam);
-	UpdateEmitterLocalSpaceBinding.Init(UpdateExecContext.Parameters, EmitterLocalSpaceParam);
-	EventEmitterLocalSpaceBindings.SetNum(NumEvents);
-	for (int32 i = 0; i < NumEvents; i++)
-	{
-		EventEmitterLocalSpaceBindings[i].Init(EventExecContexts[i].Parameters, EmitterLocalSpaceParam);
-	}
-	if (CachedEmitter->SimTarget == ENiagaraSimTarget::GPUComputeSim)
-	{
-		EmitterLocalSpaceBindingGPU.Init(GPUExecContext.CombinedParamStore, EmitterLocalSpaceParam);
-	}
-
 	SpawnExecCountBinding.Init(SpawnExecContext.Parameters, SYS_PARAM_ENGINE_EXEC_COUNT);
 	UpdateExecCountBinding.Init(UpdateExecContext.Parameters, SYS_PARAM_ENGINE_EXEC_COUNT);
 	EventExecCountBindings.SetNum(NumEvents);
@@ -847,17 +834,9 @@ void FNiagaraEmitterInstance::Tick(float DeltaSeconds)
 			Binding.SetValue(Age);
 		}
 		
-		SpawnEmitterLocalSpaceBinding.SetValue(CachedEmitter->bLocalSpace);
-		UpdateEmitterLocalSpaceBinding.SetValue(CachedEmitter->bLocalSpace);
-		for (FNiagaraParameterDirectBinding<FNiagaraBool>& Binding : EventEmitterLocalSpaceBindings)
-		{
-			Binding.SetValue(CachedEmitter->bLocalSpace);
-		}
-
 		if (CachedEmitter->SimTarget == ENiagaraSimTarget::GPUComputeSim)
 		{
 			EmitterAgeBindingGPU.SetValue(Age);
-			EmitterLocalSpaceBindingGPU.SetValue(CachedEmitter->bLocalSpace);
 		}
 	}
 	
diff --git a/Engine/Plugins/FX/Niagara/Source/Niagara/Private/NiagaraScript.cpp b/Engine/Plugins/FX/Niagara/Source/Niagara/Private/NiagaraScript.cpp
--- a/Engine/Plugins/FX/Niagara/Source/Niagara/Private/NiagaraScript.cpp
+++ b/Engine/Plugins/FX/Niagara/Source/Niagara/Private/NiagaraScript.cpp
@@ -207,8 +207,7 @@ void UNiagaraScript::ComputeVMCompilationId(FNiagaraVMExecutableDataId& Id) cons
 	// Ideally we wouldn't want to do this but rather than push the data down
 	// from the emitter.
 	UObject* Obj = GetOuter();
-	UNiagaraEmitter* Emitter = Cast<UNiagaraEmitter>(Obj);
-	if (Emitter)
+	if (UNiagaraEmitter* Emitter = Cast<UNiagaraEmitter>(Obj))
 	{
 		if ((Emitter->bInterpolatedSpawning && Usage == ENiagaraScriptUsage::ParticleGPUComputeScript) || 
 			(Emitter->bInterpolatedSpawning && Usage == ENiagaraScriptUsage::ParticleSpawnScript) ||
@@ -220,6 +219,25 @@ void UNiagaraScript::ComputeVMCompilationId(FNiagaraVMExecutableDataId& Id) cons
 		{
 			Id.AdditionalDefines.Add(TEXT("RequiresPersistentIDs"));
 		}
+		if (Emitter->bLocalSpace)
+		{
+			Id.AdditionalDefines.Add(TEXT("Emitter.Localspace"));
+		}
+	}
+
+	if (UNiagaraSystem* System = Cast<UNiagaraSystem>(Obj))
+	{
+		for (const FNiagaraEmitterHandle& EmitterHandle: System->GetEmitterHandles())
+		{
+			UNiagaraEmitter* Emitter = Cast<UNiagaraEmitter>(EmitterHandle.GetInstance());
+			if (Emitter)
+			{
+				if (Emitter->bLocalSpace)
+				{
+					Id.AdditionalDefines.Add(Emitter->GetUniqueEmitterName() + TEXT(".Localspace"));
+				}
+			}
+		}
 	}
 
 	Source->ComputeVMCompilationId(Id, Usage, UsageId);
diff --git a/Engine/Plugins/FX/Niagara/Source/NiagaraCore/Public/NiagaraCustomVersion.h b/Engine/Plugins/FX/Niagara/Source/NiagaraCore/Public/NiagaraCustomVersion.h
--- a/Engine/Plugins/FX/Niagara/Source/NiagaraCore/Public/NiagaraCustomVersion.h
+++ b/Engine/Plugins/FX/Niagara/Source/NiagaraCore/Public/NiagaraCustomVersion.h
@@ -104,6 +104,8 @@ struct FNiagaraCustomVersion
 
 		AssignmentNodeHasCorrectUsageBitmask, // Updating the usage flage bitmask for assignment nodes to match the part of the stack it's used in.
 
+		EmitterLocalSpaceLiteralConstant, //Emitter local space is compiled into the hlsl as a literal constant to expose it to emitter scripts and allow for some better optimization of particle transforms.
+
 		// -----<new versions can be added above this line>-------------------------------------------------
 		VersionPlusOne,
 		LatestVersion = VersionPlusOne - 1,
diff --git a/Engine/Plugins/FX/Niagara/Source/NiagaraEditor/Private/NiagaraHlslTranslator.cpp b/Engine/Plugins/FX/Niagara/Source/NiagaraEditor/Private/NiagaraHlslTranslator.cpp
--- a/Engine/Plugins/FX/Niagara/Source/NiagaraEditor/Private/NiagaraHlslTranslator.cpp
+++ b/Engine/Plugins/FX/Niagara/Source/NiagaraEditor/Private/NiagaraHlslTranslator.cpp
@@ -2530,7 +2530,7 @@ int32 FHlslNiagaraTranslator::GetConstant(const FNiagaraVariable& Constant)
 	return AddBodyChunk(GetUniqueSymbolName(TEXT("Constant")), ConstantStr, Constant.GetType());
 }
 
-int32 FHlslNiagaraTranslator::GetConstantFloat(float InConstantValue)
+int32 FHlslNiagaraTranslator::GetConstantDirect(float InConstantValue)
 {
 	FNiagaraVariable Constant(FNiagaraTypeDefinition::GetFloatDef(), TEXT("Constant"));
 	Constant.SetValue(InConstantValue);
@@ -2538,6 +2538,14 @@ int32 FHlslNiagaraTranslator::GetConstantFloat(float InConstantValue)
 	return GetConstant(Constant);
 }
 
+int32 FHlslNiagaraTranslator::GetConstantDirect(bool InConstantValue)
+{
+	FNiagaraVariable Constant(FNiagaraTypeDefinition::GetBoolDef(), TEXT("Constant"));
+	Constant.SetValue(InConstantValue);
+
+	return GetConstant(Constant);
+}
+
 FString FHlslNiagaraTranslator::GenerateConstantString(const FNiagaraVariable& Constant)
 {
 	FNiagaraTypeDefinition Type = Constant.GetType();
@@ -2999,6 +3007,20 @@ bool FHlslNiagaraTranslator::RequiresInterpolation() const
 	return false;
 }
 
+bool FHlslNiagaraTranslator::GetLiteralConstantVariable(FNiagaraVariable& OutVar)
+{
+	if (FNiagaraParameterMapHistory::IsInNamespace(OutVar, PARAM_MAP_EMITTER_STR) || FNiagaraParameterMapHistory::IsInNamespace(OutVar, PARAM_MAP_SYSTEM_STR))
+	{
+		FNiagaraVariable ResolvedVar = ActiveHistoryForFunctionCalls.ResolveAliases(OutVar);
+		if (OutVar == FNiagaraVariable(FNiagaraTypeDefinition::GetBoolDef(), TEXT("Emitter.Localspace")))
+		{
+			bool bEmitterLocalSpace = CompileOptions.AdditionalDefines.Contains(ResolvedVar.GetName().ToString());
+			OutVar.SetValue(bEmitterLocalSpace ? FNiagaraBool(true) : FNiagaraBool(false));
+			return true;
+		}
+	}
+	return false;
+}
 
 bool FHlslNiagaraTranslator::ParameterMapRegisterExternalConstantNamespaceVariable(FNiagaraVariable InVariable, UNiagaraNode* InNode, int32 InParamMapHistoryIdx, int32& Output, const UEdGraphPin* InDefaultPin)
 {
@@ -3017,7 +3039,7 @@ bool FHlslNiagaraTranslator::ParameterMapRegisterExternalConstantNamespaceVariab
 			if (InVariable == SYS_PARAM_ENGINE_DELTA_TIME || InVariable == SYS_PARAM_ENGINE_INV_DELTA_TIME)
 			{
 				Warning(FText::Format(LOCTEXT("GetParameterInvalidParam", "Cannot call system variable {0} in a spawn script! It is invalid."), FText::FromName(InVariable.GetName())), nullptr, nullptr);
-				Output = GetConstantFloat(0.0f);
+				Output = GetConstantDirect(0.0f);
 				return true;
 			}
 		}
@@ -3343,6 +3365,13 @@ void FHlslNiagaraTranslator::HandleParameterRead(int32 ParamMapHistoryIdx, const
 		return;
 	}
 
+	//Some special variables can be replaced directly with constants which allows for extra optimization in the compiler.
+	if (GetLiteralConstantVariable(Var))
+	{
+		OutputChunkId = GetConstant(Var);
+		return;
+	}
+
 	if (FNiagaraParameterMapHistory::IsExternalConstantNamespace(Var, CompileOptions.TargetUsage, CompileOptions.GetTargetUsageBitmask()))
 	{
 		if (ParameterMapRegisterExternalConstantNamespaceVariable(Var, ErrorNode, ParamMapHistoryIdx, OutputChunkId, DefaultPin))
diff --git a/Engine/Plugins/FX/Niagara/Source/NiagaraEditor/Private/NiagaraHlslTranslator.h b/Engine/Plugins/FX/Niagara/Source/NiagaraEditor/Private/NiagaraHlslTranslator.h
--- a/Engine/Plugins/FX/Niagara/Source/NiagaraEditor/Private/NiagaraHlslTranslator.h
+++ b/Engine/Plugins/FX/Niagara/Source/NiagaraEditor/Private/NiagaraHlslTranslator.h
@@ -489,7 +489,8 @@ private:
 		FNiagaraFunctionSignature& OutSignature, TArray<int32>& Inputs);
 	
 	// Add a raw float constant chunk
-	int32 GetConstantFloat(float InValue);
+	int32 GetConstantDirect(float InValue);
+	int32 GetConstantDirect(bool InValue);
 
 	FNiagaraTypeDefinition GetChildType(const FNiagaraTypeDefinition& BaseType, const FName& PropertyName);
 	FString NamePathToString(const FString& Prefix, const FNiagaraTypeDefinition& RootType, const TArray<FName>& NamePath);
@@ -527,6 +528,9 @@ private:
 	bool IsSpawnScript() const;
 	bool RequiresInterpolation() const;
 
+	/** If OutVar can be replaced by a literal constant, it's data is initialized with the correct value and we return true. Returns false otherwise. */
+	bool GetLiteralConstantVariable(FNiagaraVariable& OutVar);
+
 	/** Map of symbol names to count of times it's been used. Used for generating unique symbol names. */
 	TMap<FName, uint32> SymbolCounts;
 
diff --git a/Engine/Plugins/FX/Niagara/Source/NiagaraEditor/Private/NiagaraParameterMapHistory.cpp b/Engine/Plugins/FX/Niagara/Source/NiagaraEditor/Private/NiagaraParameterMapHistory.cpp
--- a/Engine/Plugins/FX/Niagara/Source/NiagaraEditor/Private/NiagaraParameterMapHistory.cpp
+++ b/Engine/Plugins/FX/Niagara/Source/NiagaraEditor/Private/NiagaraParameterMapHistory.cpp
@@ -22,21 +22,6 @@
 
 #define LOCTEXT_NAMESPACE "NiagaraEditor"
 
-#define PARAM_MAP_NPC_STR TEXT("NPC.")
-#define PARAM_MAP_ENGINE_STR TEXT("Engine.")
-#define PARAM_MAP_ENGINE_OWNER_STR TEXT("Engine.Owner.")
-#define PARAM_MAP_ENGINE_SYSTEM_STR TEXT("Engine.System.")
-#define PARAM_MAP_ENGINE_EMITTER_STR TEXT("Engine.Emitter.")
-#define PARAM_MAP_USER_STR TEXT("User.")
-#define PARAM_MAP_SYSTEM_STR TEXT("System.")
-#define PARAM_MAP_EMITTER_STR TEXT("Emitter.")
-#define PARAM_MAP_MODULE_STR TEXT("Module.")
-#define PARAM_MAP_ATTRIBUTE_STR TEXT("Particles.")
-#define PARAM_MAP_INITIAL_STR TEXT("Initial.")
-#define PARAM_MAP_INITIAL_BASE_STR TEXT("Initial")
-#define PARAM_MAP_RAPID_ITERATION_STR TEXT("Constants.")
-#define PARAM_MAP_RAPID_ITERATION_BASE_STR TEXT("Constants")
-
 
 FNiagaraParameterMapHistory::FNiagaraParameterMapHistory() 
 {
