#tests none

diff --git a/Engine/Plugins/FX/Niagara/Source/NiagaraEditorWidgets/Private/DetailCustomizations/NiagaraDataInterfaceSkeletalMeshDetails.cpp b/Engine/Plugins/FX/Niagara/Source/NiagaraEditorWidgets/Private/DetailCustomizations/NiagaraDataInterfaceSkeletalMeshDetails.cpp
--- a/Engine/Plugins/FX/Niagara/Source/NiagaraEditorWidgets/Private/DetailCustomizations/NiagaraDataInterfaceSkeletalMeshDetails.cpp
+++ b/Engine/Plugins/FX/Niagara/Source/NiagaraEditorWidgets/Private/DetailCustomizations/NiagaraDataInterfaceSkeletalMeshDetails.cpp
@@ -24,7 +24,7 @@ void FNiagaraDataInterfaceSkeletalMeshDetails::CustomizeDetails(IDetailLayoutBui
 	 if (MeshInterface != nullptr)
 	 {
 		 MeshInterface->OnChanged().RemoveAll(this);
-		 MeshInterface->OnChanged().AddRaw(this, &FNiagaraDataInterfaceSkeletalMeshDetails::OnDataChanged);
+		 MeshInterface->OnChanged().AddRaw(this, &FNiagaraDataInterfaceSkeletalMeshDetails::OnInterfaceChanged);
 	 }
 	 TWeakObjectPtr<USceneComponent> SceneComponent;
 	 USkeletalMeshComponent* FoundSkelComp = nullptr;
@@ -63,6 +63,23 @@ void FNiagaraDataInterfaceSkeletalMeshDetails::CustomizeDetails(IDetailLayoutBui
 	 return MakeShared<FNiagaraDataInterfaceSkeletalMeshDetails>();
  }
 
+ void FNiagaraDataInterfaceSkeletalMeshDetails::OnInterfaceChanged()
+ {
+     // Rebuild the data changed listener
+	 TWeakObjectPtr<USceneComponent> SceneComponent;
+	 USkeletalMeshComponent* FoundSkelComp = nullptr;
+	 if (MeshObject != nullptr)
+	 {
+		 MeshObject->GetOnMeshChanged().RemoveAll(this);
+	 }
+	 MeshObject = UNiagaraDataInterfaceSkeletalMesh::GetSkeletalMeshHelper(MeshInterface, Cast<UNiagaraComponent>(MeshInterface->GetOuter()), SceneComponent, FoundSkelComp);
+	 if (MeshObject != nullptr)
+	 {
+		 MeshObject->GetOnMeshChanged().AddRaw(this, &FNiagaraDataInterfaceSkeletalMeshDetails::OnDataChanged);
+	 }
+	 OnDataChanged();
+ }
+
  void FNiagaraDataInterfaceSkeletalMeshDetails::OnDataChanged()
  {
 	 TArray<TSharedPtr<FName>> PossibleRegions = GenerateSourceArray();
diff --git a/Engine/Plugins/FX/Niagara/Source/NiagaraEditorWidgets/Private/DetailCustomizations/NiagaraDataInterfaceSkeletalMeshDetails.h b/Engine/Plugins/FX/Niagara/Source/NiagaraEditorWidgets/Private/DetailCustomizations/NiagaraDataInterfaceSkeletalMeshDetails.h
--- a/Engine/Plugins/FX/Niagara/Source/NiagaraEditorWidgets/Private/DetailCustomizations/NiagaraDataInterfaceSkeletalMeshDetails.h
+++ b/Engine/Plugins/FX/Niagara/Source/NiagaraEditorWidgets/Private/DetailCustomizations/NiagaraDataInterfaceSkeletalMeshDetails.h
@@ -24,6 +24,7 @@ public:
 	static TSharedRef<IDetailCustomization> MakeInstance();
 
 private:
+	void OnInterfaceChanged();
 	void OnDataChanged();
 	TArray<TSharedPtr<FName>> GenerateSourceArray();
 
diff --git a/Engine/Plugins/FX/Niagara/Source/NiagaraEditorWidgets/Private/DetailCustomizations/NiagaraDetailSourcedArrayBuilder.cpp b/Engine/Plugins/FX/Niagara/Source/NiagaraEditorWidgets/Private/DetailCustomizations/NiagaraDetailSourcedArrayBuilder.cpp
--- a/Engine/Plugins/FX/Niagara/Source/NiagaraEditorWidgets/Private/DetailCustomizations/NiagaraDetailSourcedArrayBuilder.cpp
+++ b/Engine/Plugins/FX/Niagara/Source/NiagaraEditorWidgets/Private/DetailCustomizations/NiagaraDetailSourcedArrayBuilder.cpp
@@ -17,15 +17,17 @@ FNiagaraDetailSourcedArrayBuilder::FNiagaraDetailSourcedArrayBuilder(TSharedRef<
 	, ArrayProperty(InBaseProperty->AsArray())
 	, FNameSubproperty(InFNameSubproperty)
 {
+	// bradut TODO set new Delegate for when the number of children in the array changes, where you empty the map of selectors and double check if it is regenerated and also call 
+	// 
+	FSimpleDelegate OnArrayNumChildrenChanged = FSimpleDelegate::CreateRaw(this, &FNiagaraDetailSourcedArrayBuilder::OnArrayNumChildrenChanged);
+	ArrayProperty->SetOnNumElementsChanged(OnArrayNumChildrenChanged);
+	uint32 EntryCount;
+	ArrayProperty->GetNumElements(EntryCount);
+	NameSelectors.AddZeroed(EntryCount);
 };
 
 void FNiagaraDetailSourcedArrayBuilder::OnGenerateEntry(TSharedRef<IPropertyHandle> PropertyHandle, int32 ArrayIndex, IDetailChildrenBuilder& ChildrenBuilder)
 {
-	TSharedRef<IPropertyHandle> SubPropertyHandle = PropertyHandle;
-	if (FNameSubproperty != NAME_None)
-	{
-		SubPropertyHandle = PropertyHandle->GetChildHandle(FNameSubproperty).ToSharedRef();
-	}
 	IDetailPropertyRow& RegionRow = ChildrenBuilder.AddProperty(PropertyHandle);
 
 	FNumberFormattingOptions NoCommas;
@@ -35,9 +37,13 @@ void FNiagaraDetailSourcedArrayBuilder::OnGenerateEntry(TSharedRef<IPropertyHand
 	RegionRow.DisplayName(SlotDesc);
 
 	RegionRow.ShowPropertyButtons(true);
-	
-	NameSelector = SNew(SNiagaraNamePropertySelector, SubPropertyHandle, OptionsSourceList);
-	
+
+	TSharedPtr<IPropertyHandle> SubPropertyHandle = PropertyHandle;
+	if (FNameSubproperty != NAME_None)
+	{
+		SubPropertyHandle = PropertyHandle->GetChildHandle(FNameSubproperty).ToSharedRef();
+	}
+	NameSelectors[ArrayIndex] = SNew(SNiagaraNamePropertySelector, SubPropertyHandle.ToSharedRef(), OptionsSourceList);
 	RegionRow.CustomWidget(false)
 		.NameContent()
 		[
@@ -47,16 +53,28 @@ void FNiagaraDetailSourcedArrayBuilder::OnGenerateEntry(TSharedRef<IPropertyHand
 		.MaxDesiredWidth(TOptional<float>())
 		[
 			// add combo button
-			NameSelector.ToSharedRef()
+			NameSelectors[ArrayIndex].ToSharedRef()
 		];
 }
 
+void FNiagaraDetailSourcedArrayBuilder::OnArrayNumChildrenChanged()
+{
+	NameSelectors.Empty();
+	uint32 EntryCount;
+	ArrayProperty->GetNumElements(EntryCount);
+	NameSelectors.AddZeroed(EntryCount); // allocate the memory for ease of use when refreshing children
+	RefreshChildren();
+}
+
 void FNiagaraDetailSourcedArrayBuilder::SetSourceArray(TArray<TSharedPtr<FName>>& InOptionsSource)
 {
 	OptionsSourceList = InOptionsSource;
-	if (NameSelector.IsValid())
+	for (auto NameSelector : NameSelectors)
 	{
-		NameSelector->SetSourceArray(InOptionsSource);
+		if (NameSelector.IsValid())
+		{
+			NameSelector->SetSourceArray(InOptionsSource);
+		}
 	}
 }
 
diff --git a/Engine/Plugins/FX/Niagara/Source/NiagaraEditorWidgets/Private/DetailCustomizations/NiagaraDetailSourcedArrayBuilder.h b/Engine/Plugins/FX/Niagara/Source/NiagaraEditorWidgets/Private/DetailCustomizations/NiagaraDetailSourcedArrayBuilder.h
--- a/Engine/Plugins/FX/Niagara/Source/NiagaraEditorWidgets/Private/DetailCustomizations/NiagaraDetailSourcedArrayBuilder.h
+++ b/Engine/Plugins/FX/Niagara/Source/NiagaraEditorWidgets/Private/DetailCustomizations/NiagaraDetailSourcedArrayBuilder.h
@@ -24,7 +24,7 @@ public:
 
 private: 
 	void OnGenerateEntry(TSharedRef<IPropertyHandle> ElementProperty, int32 ElementIndex, IDetailChildrenBuilder& ChildrenBuilder);
-	
+	void OnArrayNumChildrenChanged();
 private:
 	TArray<TSharedPtr<FName>> OptionsSourceList;
 	/** List of component class names, filtered by the current search string */
@@ -32,5 +32,5 @@ private:
 	TSharedPtr<IPropertyHandleArray> ArrayProperty;
 	/** Subproperty of type FName that needs to be edited - in case of FName wrappers */
 	FName FNameSubproperty;
-	TSharedPtr<SNiagaraNamePropertySelector> NameSelector;
+	TArray<TSharedPtr<SNiagaraNamePropertySelector>> NameSelectors;
 };
diff --git a/Engine/Plugins/FX/Niagara/Source/NiagaraEditorWidgets/Private/DetailCustomizations/SNiagaraNamePropertySelector.cpp b/Engine/Plugins/FX/Niagara/Source/NiagaraEditorWidgets/Private/DetailCustomizations/SNiagaraNamePropertySelector.cpp
--- a/Engine/Plugins/FX/Niagara/Source/NiagaraEditorWidgets/Private/DetailCustomizations/SNiagaraNamePropertySelector.cpp
+++ b/Engine/Plugins/FX/Niagara/Source/NiagaraEditorWidgets/Private/DetailCustomizations/SNiagaraNamePropertySelector.cpp
@@ -14,8 +14,8 @@ void SNiagaraNamePropertySelector::Construct(const FArguments& InArgs, TSharedRe
 {
 	PropertyHandle = InBaseProperty;
 	OptionsSourceList = InOptionsSource;
-	
 	GenerateFilteredElementList(CurrentSearchString.ToString());
+
 	SAssignNew(ElementsListView, SListView<TSharedPtr<FName>>)
 		.ListItemsSource(&FilteredSourceList)
 		.OnSelectionChanged(this, &SNiagaraNamePropertySelector::OnSelectionChanged)
@@ -108,7 +108,8 @@ void SNiagaraNamePropertySelector::GenerateFilteredElementList(const FString& In
 {
 	if (InSearchText.IsEmpty())
 	{
-		FilteredSourceList = OptionsSourceList;
+		FilteredSourceList.Empty();
+		FilteredSourceList.Append(OptionsSourceList);
 	}
 	else
 	{
