Niagara - Fix some issues with modify, transactions, and change ids which was causing assets to be dirty or modified on load, or were allowing internal operation to be undone.
+ Move some transactions from public utility functions into private functions called by menu items in the UNiagaraNodeWithDynamicPins.
+ Prevent a few modify calls in UNiagaraEmitter from marking the package dirty since they're sometimes called as a result of compiling and in the other cases earlier modifies would have already marked the package dirty.
+ In the system view model, don't create transactions when adding an emitter if the system view model is in emitter asset mode since the user should be able to undo it.
+ In the system toolkit when opening an emitter asset initialize, clean up, and propagate the rapid iteration parameters before copying the emitter to prevent the change ids from getting out of sync after the compile completes.
+ In the system toolkit when trying to see if an emitter has changed using the change ids, use the last synced id from the copied emitter instead of the original emitter since duplicating the emitter can change the id, and there's not way to set it externally.

diff --git a/Engine/Plugins/FX/Niagara/Source/Niagara/Classes/NiagaraScript.h b/Engine/Plugins/FX/Niagara/Source/Niagara/Classes/NiagaraScript.h
--- a/Engine/Plugins/FX/Niagara/Source/Niagara/Classes/NiagaraScript.h
+++ b/Engine/Plugins/FX/Niagara/Source/Niagara/Classes/NiagaraScript.h
@@ -390,7 +390,7 @@ public:
 	class UNiagaraScriptSourceBase *GetSource() { return Source; }
 	const class UNiagaraScriptSourceBase *GetSource() const  { return Source; }
 	void SetSource(class UNiagaraScriptSourceBase *InSource) { Source = InSource; }
-	void CleanUpOldAndInitializeNewRapidIterationParameters(FString UniqueEmitterName);
+	NIAGARA_API void CleanUpOldAndInitializeNewRapidIterationParameters(FString UniqueEmitterName);
 
 	NIAGARA_API FGuid GetBaseChangeID() const;
 	NIAGARA_API ENiagaraScriptCompileStatus GetLastCompileStatus() const;
diff --git a/Engine/Plugins/FX/Niagara/Source/Niagara/Private/NiagaraEmitter.cpp b/Engine/Plugins/FX/Niagara/Source/Niagara/Private/NiagaraEmitter.cpp
--- a/Engine/Plugins/FX/Niagara/Source/Niagara/Private/NiagaraEmitter.cpp
+++ b/Engine/Plugins/FX/Niagara/Source/Niagara/Private/NiagaraEmitter.cpp
@@ -235,7 +235,7 @@ void UNiagaraEmitter::PostLoad()
 
 	if (bGenerateNewChangeId)
 	{
-		ChangeId = FGuid::NewGuid();
+		UpdateChangeId();
 	}
 
 	GraphSource->OnChanged().AddUObject(this, &UNiagaraEmitter::GraphSourceChanged);
@@ -672,7 +672,9 @@ void UNiagaraEmitter::SyncEmitterAlias(const FString& InOldName, const FString&
 
 	for (UNiagaraScript* Script : Scripts)
 	{
-		Script->Modify();
+		// We don't mark the package dirty here because this can happen as a result of a compile and we don't want to dirty files
+		// due to compilation, in cases where the package should be marked dirty an previous modify would have already done this.
+		Script->Modify(false);
 		Script->SyncAliases(RenameMap);
 	}
 }
@@ -736,6 +738,7 @@ FNiagaraEventScriptProperties* UNiagaraEmitter::GetEventHandlerByIdUnsafe(FGuid
 
 void UNiagaraEmitter::AddEventHandler(FNiagaraEventScriptProperties EventHandler)
 {
+	Modify();
 	EventHandlerScriptProps.Add(EventHandler);
 #if WITH_EDITOR
 	EventHandler.Script->RapidIterationParameters.AddOnChangedHandler(
@@ -746,6 +749,7 @@ void UNiagaraEmitter::AddEventHandler(FNiagaraEventScriptProperties EventHandler
 
 void UNiagaraEmitter::RemoveEventHandlerByUsageId(FGuid EventHandlerUsageId)
 {
+	Modify();
 	auto FindEventHandlerById = [=](const FNiagaraEventScriptProperties& EventHandler) { return EventHandler.Script->GetUsageId() == EventHandlerUsageId; };
 #if WITH_EDITOR
 	FNiagaraEventScriptProperties* EventHandler = EventHandlerScriptProps.FindByPredicate(FindEventHandlerById);
@@ -775,7 +779,9 @@ void UNiagaraEmitter::BeginDestroy()
 
 void UNiagaraEmitter::UpdateChangeId()
 {
-	Modify();
+	// We don't mark the package dirty here because this can happen as a result of a compile and we don't want to dirty files
+	// due to compilation, in cases where the package should be marked dirty an previous modify would have already done this.
+	Modify(false);
 	ChangeId = FGuid::NewGuid();
 }
 
diff --git a/Engine/Plugins/FX/Niagara/Source/NiagaraEditor/Private/NiagaraNodeCustomHlsl.cpp b/Engine/Plugins/FX/Niagara/Source/NiagaraEditor/Private/NiagaraNodeCustomHlsl.cpp
--- a/Engine/Plugins/FX/Niagara/Source/NiagaraEditor/Private/NiagaraNodeCustomHlsl.cpp
+++ b/Engine/Plugins/FX/Niagara/Source/NiagaraEditor/Private/NiagaraNodeCustomHlsl.cpp
@@ -167,7 +167,6 @@ void UNiagaraNodeCustomHlsl::InitAsCustomHlslDynamicInput(const FNiagaraTypeDefi
 /** Called when a new typed pin is added by the user. */
 void UNiagaraNodeCustomHlsl::OnNewTypedPinAdded(UEdGraphPin* NewPin)
 {
-	FScopedTransaction Transaction(LOCTEXT("AddedTypedPin", "New Pin Added"));
 	UNiagaraNodeWithDynamicPins::OnNewTypedPinAdded(NewPin);
 	RebuildSignatureFromPins();
 }
@@ -175,7 +174,6 @@ void UNiagaraNodeCustomHlsl::OnNewTypedPinAdded(UEdGraphPin* NewPin)
 /** Called when a pin is renamed. */
 void UNiagaraNodeCustomHlsl::OnPinRenamed(UEdGraphPin* RenamedPin, const FString& OldPinName)
 {
-	FScopedTransaction Transaction(LOCTEXT("PinRenamed", "Pin Rename"));
 	UNiagaraNodeWithDynamicPins::OnPinRenamed(RenamedPin, OldPinName);
 	RebuildSignatureFromPins();
 }
@@ -183,14 +181,12 @@ void UNiagaraNodeCustomHlsl::OnPinRenamed(UEdGraphPin* RenamedPin, const FString
 /** Removes a pin from this node with a transaction. */
 void UNiagaraNodeCustomHlsl::RemoveDynamicPin(UEdGraphPin* Pin)
 {
-	FScopedTransaction Transaction(LOCTEXT("PinRemoved", "Remove Pin"));
 	UNiagaraNodeWithDynamicPins::RemoveDynamicPin(Pin);
 	RebuildSignatureFromPins();
 }
 
 void UNiagaraNodeCustomHlsl::MoveDynamicPin(UEdGraphPin* Pin, int32 DirectionToMove)
 {
-	FScopedTransaction Transaction(LOCTEXT("PinMoved", "Moved Pin"));
 	UNiagaraNodeWithDynamicPins::MoveDynamicPin(Pin, DirectionToMove);
 	RebuildSignatureFromPins();
 }
diff --git a/Engine/Plugins/FX/Niagara/Source/NiagaraEditor/Private/NiagaraNodeParameterMapGet.cpp b/Engine/Plugins/FX/Niagara/Source/NiagaraEditor/Private/NiagaraNodeParameterMapGet.cpp
--- a/Engine/Plugins/FX/Niagara/Source/NiagaraEditor/Private/NiagaraNodeParameterMapGet.cpp
+++ b/Engine/Plugins/FX/Niagara/Source/NiagaraEditor/Private/NiagaraNodeParameterMapGet.cpp
@@ -171,8 +171,6 @@ void UNiagaraNodeParameterMapGet::OnNewTypedPinAdded(UEdGraphPin* NewPin)
 
 void UNiagaraNodeParameterMapGet::RemoveDynamicPin(UEdGraphPin* Pin)
 {
-	FScopedTransaction RemovePinTransaction(LOCTEXT("RemovePinTransaction", "Remove pin"));
-
 	UEdGraphPin* DefaultPin = nullptr;
 	if (Pin->Direction == EEdGraphPinDirection::EGPD_Output)
 	{
diff --git a/Engine/Plugins/FX/Niagara/Source/NiagaraEditor/Private/NiagaraNodeWithDynamicPins.cpp b/Engine/Plugins/FX/Niagara/Source/NiagaraEditor/Private/NiagaraNodeWithDynamicPins.cpp
--- a/Engine/Plugins/FX/Niagara/Source/NiagaraEditor/Private/NiagaraNodeWithDynamicPins.cpp
+++ b/Engine/Plugins/FX/Niagara/Source/NiagaraEditor/Private/NiagaraNodeWithDynamicPins.cpp
@@ -139,7 +139,6 @@ void UNiagaraNodeWithDynamicPins::MoveDynamicPin(UEdGraphPin* Pin, int32 Directi
 		{
 			if (i + DirectionToMove >= 0 && i + DirectionToMove < SameDirectionPins.Num())
 			{
-				FScopedTransaction MovePinTransaction(LOCTEXT("MovePinTransaction", "Moved pin"));
 				Modify();
 				UEdGraphPin* PinOld = SameDirectionPins[i + DirectionToMove];
 				if (PinOld)
@@ -188,7 +187,7 @@ void UNiagaraNodeWithDynamicPins::GetContextMenuActions(const FGraphNodeContextM
 				LOCTEXT("RemoveDynamicPin", "Remove pin"),
 				LOCTEXT("RemoveDynamicPinToolTip", "Remove this pin and any connections."),
 				FSlateIcon(),
-				FUIAction(FExecuteAction::CreateUObject(const_cast<UNiagaraNodeWithDynamicPins*>(this), &UNiagaraNodeWithDynamicPins::RemoveDynamicPin, const_cast<UEdGraphPin*>(Context.Pin))));
+				FUIAction(FExecuteAction::CreateUObject(const_cast<UNiagaraNodeWithDynamicPins*>(this), &UNiagaraNodeWithDynamicPins::RemoveDynamicPinFromMenu, const_cast<UEdGraphPin*>(Context.Pin))));
 		}
 		if (CanMovePin(Context.Pin))
 		{
@@ -210,7 +209,7 @@ void UNiagaraNodeWithDynamicPins::GetContextMenuActions(const FGraphNodeContextM
 					LOCTEXT("MoveDynamicPinUp", "Move pin up"),
 					LOCTEXT("MoveDynamicPinToolTipUp", "Move this pin and any connections one slot up."),
 					FSlateIcon(),
-					FUIAction(FExecuteAction::CreateUObject(const_cast<UNiagaraNodeWithDynamicPins*>(this), &UNiagaraNodeWithDynamicPins::MoveDynamicPin, const_cast<UEdGraphPin*>(Context.Pin), -1)));
+					FUIAction(FExecuteAction::CreateUObject(const_cast<UNiagaraNodeWithDynamicPins*>(this), &UNiagaraNodeWithDynamicPins::MoveDynamicPinFromMenu, const_cast<UEdGraphPin*>(Context.Pin), -1)));
 			}
 			if (PinIdx >= 0 && PinIdx < SameDirectionPins.Num() - 1)
 			{
@@ -261,10 +260,7 @@ void UNiagaraNodeWithDynamicPins::AddParameter(FNiagaraVariable Parameter, UEdGr
 
 void UNiagaraNodeWithDynamicPins::RemoveDynamicPin(UEdGraphPin* Pin)
 {
-	FScopedTransaction RemovePinTransaction(LOCTEXT("RemovePinTransaction", "Remove pin"));
 	RemovePin(Pin);
-	//GetGraph()->NotifyGraphChanged();
-
 	MarkNodeRequiresSynchronization(__FUNCTION__, true);
 }
 
@@ -286,4 +282,16 @@ void UNiagaraNodeWithDynamicPins::PinNameTextCommitted(const FText& Text, ETextC
 	}
 }
 
+void UNiagaraNodeWithDynamicPins::RemoveDynamicPinFromMenu(UEdGraphPin* Pin)
+{
+	FScopedTransaction RemovePinTransaction(LOCTEXT("RemovePinTransaction", "Remove pin"));
+	RemoveDynamicPin(Pin);
+}
+
+void UNiagaraNodeWithDynamicPins::MoveDynamicPinFromMenu(UEdGraphPin* Pin, int32 DirectionToMove)
+{
+	FScopedTransaction MovePinTransaction(LOCTEXT("MovePinTransaction", "Moved pin"));
+	MoveDynamicPin(Pin, DirectionToMove);
+}
+
 #undef LOCTEXT_NAMESPACE
diff --git a/Engine/Plugins/FX/Niagara/Source/NiagaraEditor/Private/Toolkits/NiagaraSystemToolkit.cpp b/Engine/Plugins/FX/Niagara/Source/NiagaraEditor/Private/Toolkits/NiagaraSystemToolkit.cpp
--- a/Engine/Plugins/FX/Niagara/Source/NiagaraEditor/Private/Toolkits/NiagaraSystemToolkit.cpp
+++ b/Engine/Plugins/FX/Niagara/Source/NiagaraEditor/Private/Toolkits/NiagaraSystemToolkit.cpp
@@ -197,10 +197,28 @@ void FNiagaraSystemToolkit::InitializeWithEmitter(const EToolkitMode::Type Mode,
 
 	Emitter = &InEmitter;
 
+	// Before copying the emitter clean up and propagate the rapid iteration parameters so that the post compile cleanup 
+	// and propagation doesn't cause the change ids to become out of sync.
+	FString EmitterName = "Emitter";
+	Emitter->EmitterSpawnScriptProps.Script->CleanUpOldAndInitializeNewRapidIterationParameters(EmitterName);
+	Emitter->EmitterUpdateScriptProps.Script->CleanUpOldAndInitializeNewRapidIterationParameters(EmitterName);
+	Emitter->SpawnScriptProps.Script->CleanUpOldAndInitializeNewRapidIterationParameters(EmitterName);
+	Emitter->UpdateScriptProps.Script->CleanUpOldAndInitializeNewRapidIterationParameters(EmitterName);
+	for (const FNiagaraEventScriptProperties& EventHandler : Emitter->GetEventHandlers())
+	{
+		EventHandler.Script->CleanUpOldAndInitializeNewRapidIterationParameters(EmitterName);
+	}
+	Emitter->UpdateScriptProps.Script->RapidIterationParameters.CopyParametersTo(
+		Emitter->SpawnScriptProps.Script->RapidIterationParameters, false, FNiagaraParameterStore::EDataInterfaceCopyMethod::None);
+
 	ResetLoaders(GetTransientPackage()); // Make sure that we're not going to get invalid version number linkers into the package we are going into. 
 	GetTransientPackage()->LinkerCustomVersion.Empty();
 
 	UNiagaraEmitter* EditableEmitter = (UNiagaraEmitter*)StaticDuplicateObject(Emitter, GetTransientPackage(), NAME_None, ~RF_Standalone, UNiagaraEmitter::StaticClass());
+	
+	// We set this to the copy's change id here instead of the original emitter's change id because the copy's change id may have been
+	// updated from the original as part of post load and we use this id to detect if the editable emitter has been changed.
+	LastSyncedEmitterChangeId = EditableEmitter->GetChangeId();
 
 	FNiagaraSystemViewModelOptions SystemOptions;
 	SystemOptions.bCanModifyEmittersFromTimeline = false;
@@ -933,6 +951,9 @@ void FNiagaraSystemToolkit::UpdateOriginalEmitter()
 		RF_AllFlags,
 		Emitter->GetClass());
 
+	// Record the last synced change id to detect future changes.
+	LastSyncedEmitterChangeId = EditableEmitter->GetChangeId();
+
 	checkSlow (UNiagaraEmitter::GetForceCompileOnLoad() || Emitter->GetChangeId() == EditableEmitter->GetChangeId());
 
 	// Restore RF_Standalone on the original emitter, as it had been removed from the preview emitter so that it could be GC'd.
@@ -1042,7 +1063,7 @@ bool FNiagaraSystemToolkit::OnRequestClose()
 	if (SystemToolkitMode == ESystemToolkitMode::Emitter)
 	{
 		TSharedPtr<FNiagaraEmitterViewModel> EmitterViewModel = SystemViewModel->GetEmitterHandleViewModels()[0]->GetEmitterViewModel();
-		if (EmitterViewModel->GetEmitter()->GetChangeId() != Emitter->GetChangeId())
+		if (EmitterViewModel->GetEmitter()->GetChangeId() != LastSyncedEmitterChangeId)
 		{
 			// find out the user wants to do with this dirty NiagaraScript
 			EAppReturnType::Type YesNoCancelReply = FMessageDialog::Open(EAppMsgType::YesNoCancel,
@@ -1102,7 +1123,7 @@ bool FNiagaraSystemToolkit::OnApplyEnabled() const
 	if (Emitter != nullptr)
 	{
 		TSharedPtr<FNiagaraEmitterViewModel> EmitterViewModel = SystemViewModel->GetEmitterHandleViewModels()[0]->GetEmitterViewModel();
-		return EmitterViewModel->GetEmitter()->GetChangeId() != Emitter->GetChangeId();
+		return EmitterViewModel->GetEmitter()->GetChangeId() != LastSyncedEmitterChangeId;
 	}
 	return false;
 }
diff --git a/Engine/Plugins/FX/Niagara/Source/NiagaraEditor/Private/Toolkits/NiagaraSystemToolkit.h b/Engine/Plugins/FX/Niagara/Source/NiagaraEditor/Private/Toolkits/NiagaraSystemToolkit.h
--- a/Engine/Plugins/FX/Niagara/Source/NiagaraEditor/Private/Toolkits/NiagaraSystemToolkit.h
+++ b/Engine/Plugins/FX/Niagara/Source/NiagaraEditor/Private/Toolkits/NiagaraSystemToolkit.h
@@ -130,6 +130,9 @@ private:
 	/** The emitter being edited in emitter mode, or null when editing in system mode. */
 	UNiagaraEmitter* Emitter;
 
+	/** The value of the emitter change id from the last time it was in sync with the original emitter. */
+	FGuid LastSyncedEmitterChangeId;
+
 	ESystemToolkitMode SystemToolkitMode;
 
 	TSharedPtr<SNiagaraSystemViewport> Viewport;
diff --git a/Engine/Plugins/FX/Niagara/Source/NiagaraEditor/Private/ViewModels/NiagaraSystemViewModel.cpp b/Engine/Plugins/FX/Niagara/Source/NiagaraEditor/Private/ViewModels/NiagaraSystemViewModel.cpp
--- a/Engine/Plugins/FX/Niagara/Source/NiagaraEditor/Private/ViewModels/NiagaraSystemViewModel.cpp
+++ b/Engine/Plugins/FX/Niagara/Source/NiagaraEditor/Private/ViewModels/NiagaraSystemViewModel.cpp
@@ -290,7 +290,12 @@ void FNiagaraSystemViewModel::AddEmitter(UNiagaraEmitter& Emitter)
 	// Kill all system instances before modifying the emitter handle list to prevent accessing deleted data.
 	KillSystemInstances();
 
-	const FScopedTransaction Transaction(LOCTEXT("AddEmitter", "Add emitter"));
+	// When editing an emitter asset the system is a placeholder and we don't want to make adding an emitter to it
+	// undoable.
+	if (EditMode != ENiagaraSystemViewModelEditMode::EmitterAsset)
+	{
+		GEditor->BeginTransaction(LOCTEXT("AddEmitter", "Add emitter"));
+	}
 
 	TSet<FName> EmitterHandleNames;
 	for (const FNiagaraEmitterHandle& EmitterHandle : System.GetEmitterHandles())
@@ -312,6 +317,11 @@ void FNiagaraSystemViewModel::AddEmitter(UNiagaraEmitter& Emitter)
 	check(SystemScriptViewModel.IsValid());
 	SystemScriptViewModel->RebuildEmitterNodes();
 
+	if (EditMode == ENiagaraSystemViewModelEditMode::SystemAsset)
+	{
+		GEditor->EndTransaction();
+	}
+
 	if (System.GetNumEmitters() == 1 && EditorSettings->GetAutoPlay())
 	{
 		// When adding a new emitter to an empty system start playing.
diff --git a/Engine/Plugins/FX/Niagara/Source/NiagaraEditor/Public/NiagaraNodeWithDynamicPins.h b/Engine/Plugins/FX/Niagara/Source/NiagaraEditor/Public/NiagaraNodeWithDynamicPins.h
--- a/Engine/Plugins/FX/Niagara/Source/NiagaraEditor/Public/NiagaraNodeWithDynamicPins.h
+++ b/Engine/Plugins/FX/Niagara/Source/NiagaraEditor/Public/NiagaraNodeWithDynamicPins.h
@@ -76,6 +76,10 @@ private:
 
 	/** Called when a pin's name text is committed. */
 	void PinNameTextCommitted(const FText& Text, ETextCommit::Type CommitType, UEdGraphPin* Pin);
+
+	void RemoveDynamicPinFromMenu(UEdGraphPin* Pin);
+
+	void MoveDynamicPinFromMenu(UEdGraphPin* Pin, int32 DirectionToMove);
 	
 public:
 	/** The sub category for add pins. */
