Fixing display of errors

#tests now errors in compilation properly display

diff --git a/Engine/Plugins/FX/Niagara/Source/NiagaraEditor/Private/ViewModels/Stack/NiagaraStackEntry.cpp b/Engine/Plugins/FX/Niagara/Source/NiagaraEditor/Private/ViewModels/Stack/NiagaraStackEntry.cpp
--- a/Engine/Plugins/FX/Niagara/Source/NiagaraEditor/Private/ViewModels/Stack/NiagaraStackEntry.cpp
+++ b/Engine/Plugins/FX/Niagara/Source/NiagaraEditor/Private/ViewModels/Stack/NiagaraStackEntry.cpp
@@ -46,7 +46,7 @@ UNiagaraStackEntry::FStackIssue::FStackIssue()
 {
 }
 
-UNiagaraStackEntry::FStackIssue::FStackIssue(EStackIssueSeverity InSeverity, FText InShortDescription, FText InLongDescription, FString InStackEditorDataKey, bool bInCanBeDismissed, TArray<FStackIssueFix>& InFixes)
+UNiagaraStackEntry::FStackIssue::FStackIssue(EStackIssueSeverity InSeverity, FText InShortDescription, FText InLongDescription, FString InStackEditorDataKey, bool bInCanBeDismissed, const TArray<FStackIssueFix>& InFixes)
 	: Severity(InSeverity)
 	, ShortDescription(InShortDescription)
 	, LongDescription(InLongDescription)
@@ -65,16 +65,13 @@ UNiagaraStackEntry::FStackIssue::FStackIssue(EStackIssueSeverity InSeverity, FTe
 }
 
 UNiagaraStackEntry::FStackIssue::FStackIssue(EStackIssueSeverity InSeverity, FText InShortDescription, FText InLongDescription, FString InStackEditorDataKey, bool bInCanBeDismissed, FStackIssueFix InFix)
+	: FStackIssue(InSeverity, InShortDescription, InLongDescription, InStackEditorDataKey, bInCanBeDismissed, TArray<FStackIssueFix>({InFix }))
 {
-	TArray<FStackIssueFix> FixArray;
-	FixArray.Add(InFix);
-	FStackIssue(InSeverity, InShortDescription, InLongDescription, InStackEditorDataKey, bInCanBeDismissed, FixArray);
 }
 
 UNiagaraStackEntry::FStackIssue::FStackIssue(EStackIssueSeverity InSeverity, FText InShortDescription, FText InLongDescription, FString InStackEditorDataKey, bool bInCanBeDismissed)
+	: FStackIssue(InSeverity, InShortDescription, InLongDescription, InStackEditorDataKey, bInCanBeDismissed, TArray<FStackIssueFix>())
 {
-	TArray<FStackIssueFix> FixArray;
-	FStackIssue(InSeverity, InShortDescription, InLongDescription, InStackEditorDataKey, bInCanBeDismissed, FixArray);
 }
 
 bool UNiagaraStackEntry::FStackIssue::IsValid()
diff --git a/Engine/Plugins/FX/Niagara/Source/NiagaraEditor/Public/ViewModels/Stack/NiagaraStackEntry.h b/Engine/Plugins/FX/Niagara/Source/NiagaraEditor/Public/ViewModels/Stack/NiagaraStackEntry.h
--- a/Engine/Plugins/FX/Niagara/Source/NiagaraEditor/Public/ViewModels/Stack/NiagaraStackEntry.h
+++ b/Engine/Plugins/FX/Niagara/Source/NiagaraEditor/Public/ViewModels/Stack/NiagaraStackEntry.h
@@ -128,7 +128,7 @@ public:
 	{
 		FStackIssue();
 
-		FStackIssue(EStackIssueSeverity InSeverity, FText InShortDescription, FText InLongDescription, FString InStackEditorDataKey, bool bInCanBeDismissed, TArray<FStackIssueFix>& InFixes);
+		FStackIssue(EStackIssueSeverity InSeverity, FText InShortDescription, FText InLongDescription, FString InStackEditorDataKey, bool bInCanBeDismissed, const TArray<FStackIssueFix>& InFixes);
 
 		FStackIssue(EStackIssueSeverity InSeverity, FText InShortDescription, FText InLongDescription, FString InStackEditorDataKey, bool bInCanBeDismissed, FStackIssueFix InFix);
 
