Fixed crash in the editor when opening a Niagara Emitter/Script containing outdated script functions.
Fixed incorrect type cast, FNiagaraFloat instead of FNiagaraInt32.

diff --git a/Engine/Plugins/FX/Niagara/Source/NiagaraEditor/Private/TypeEditorUtilities/NiagaraIntegerTypeEditorUtilities.cpp b/Engine/Plugins/FX/Niagara/Source/NiagaraEditor/Private/TypeEditorUtilities/NiagaraIntegerTypeEditorUtilities.cpp
--- a/Engine/Plugins/FX/Niagara/Source/NiagaraEditor/Private/TypeEditorUtilities/NiagaraIntegerTypeEditorUtilities.cpp
+++ b/Engine/Plugins/FX/Niagara/Source/NiagaraEditor/Private/TypeEditorUtilities/NiagaraIntegerTypeEditorUtilities.cpp
@@ -97,7 +97,7 @@ bool FNiagaraEditorIntegerTypeUtilities::CanHandlePinDefaults() const
 FString FNiagaraEditorIntegerTypeUtilities::GetPinDefaultStringFromValue(const FNiagaraVariable& AllocatedVariable) const
 {
 	checkf(AllocatedVariable.IsDataAllocated(), TEXT("Can not generate a default value string for an unallocated variable."));
-	return Lex::ToString(AllocatedVariable.GetValue<FNiagaraFloat>().Value);
+	return Lex::ToString(AllocatedVariable.GetValue<FNiagaraInt32>().Value);
 }
 
 bool FNiagaraEditorIntegerTypeUtilities::SetValueFromPinDefaultString(const FString& StringValue, FNiagaraVariable& Variable) const
diff --git a/Engine/Plugins/FX/Niagara/Source/NiagaraEditor/Private/ViewModels/Stack/NiagaraStackModuleItem.cpp b/Engine/Plugins/FX/Niagara/Source/NiagaraEditor/Private/ViewModels/Stack/NiagaraStackModuleItem.cpp
--- a/Engine/Plugins/FX/Niagara/Source/NiagaraEditor/Private/ViewModels/Stack/NiagaraStackModuleItem.cpp
+++ b/Engine/Plugins/FX/Niagara/Source/NiagaraEditor/Private/ViewModels/Stack/NiagaraStackModuleItem.cpp
@@ -309,7 +309,7 @@ void UNiagaraStackModuleItem::RefreshIssues(TArray<FStackIssue>& NewIssues)
 		{
 			FunctionNode = ModuleData.ModuleNode;
 			DependencyModuleIndex++;
-			if (FunctionNode != nullptr && FunctionNode->FunctionScript->ProvidedDependencies.Contains(Dependency.Id)) 
+			if (FunctionNode != nullptr && FunctionNode->FunctionScript != nullptr && FunctionNode->FunctionScript->ProvidedDependencies.Contains(Dependency.Id))
 			{
 				auto DependencyOutputUsage = ModuleData.Usage;
 				int32 PossibleIndex = UsagePriority.IndexOfByPredicate(
