Niagara - Fix a few more issues where asset editors would open with their assets modified.
+ Fix rapid iteration parameter preparation so that it doesn't modify the parameter store if it doesn't change after syncing with the graphs and propagating from dependencies.  This fixes the emitter editor allowing changes to be applied on open.
+ Refactor the change notification for the script tool kit so that it uses the graph change and property change messages to determine if any changes have been made and can be applied.

#tests auto tests

diff --git a/Engine/Plugins/FX/Niagara/Source/Niagara/Classes/NiagaraScript.h b/Engine/Plugins/FX/Niagara/Source/Niagara/Classes/NiagaraScript.h
--- a/Engine/Plugins/FX/Niagara/Source/Niagara/Classes/NiagaraScript.h
+++ b/Engine/Plugins/FX/Niagara/Source/Niagara/Classes/NiagaraScript.h
@@ -390,7 +390,6 @@ public:
 	class UNiagaraScriptSourceBase *GetSource() { return Source; }
 	const class UNiagaraScriptSourceBase *GetSource() const  { return Source; }
 	void SetSource(class UNiagaraScriptSourceBase *InSource) { Source = InSource; }
-	NIAGARA_API void CleanUpOldAndInitializeNewRapidIterationParameters(FString UniqueEmitterName);
 
 	NIAGARA_API FGuid GetBaseChangeID() const;
 	NIAGARA_API ENiagaraScriptCompileStatus GetLastCompileStatus() const;
diff --git a/Engine/Plugins/FX/Niagara/Source/Niagara/Private/NiagaraCommon.cpp b/Engine/Plugins/FX/Niagara/Source/Niagara/Private/NiagaraCommon.cpp
--- a/Engine/Plugins/FX/Niagara/Source/Niagara/Private/NiagaraCommon.cpp
+++ b/Engine/Plugins/FX/Niagara/Source/Niagara/Private/NiagaraCommon.cpp
@@ -6,11 +6,15 @@
 #include "NiagaraSystemInstance.h"
 #include "NiagaraParameterCollection.h"
 #include "NiagaraComponent.h"
+#include "NiagaraScriptSourceBase.h"
+#include "NiagaraStats.h"
 #include "UObject/Linker.h"
 #include "UObject/Class.h"
 #include "UObject/Package.h"
 #include "Modules/ModuleManager.h"
 
+DECLARE_CYCLE_STAT(TEXT("Niagara - Utilities - PrepareRapidIterationParameters"), STAT_Niagara_Utilities_PrepareRapidIterationParameters, STATGROUP_Niagara);
+
 //////////////////////////////////////////////////////////////////////////
 
 FString FNiagaraTypeHelper::ToString(const uint8* ValueData, const UScriptStruct* Struct)
@@ -300,4 +304,81 @@ UNiagaraDataInterface* FNiagaraScriptDataInterfaceCompileInfo::GetDefaultDataInt
 	check(IsInGameThread());
 	UNiagaraDataInterface* Obj = CastChecked<UNiagaraDataInterface>(const_cast<UClass*>(Type.GetClass())->GetDefaultObject(true));
 	return Obj;
+}
+
+void FNiagaraUtilities::PrepareRapidIterationParameters(const TArray<UNiagaraScript*>& Scripts, const TMap<UNiagaraScript*, UNiagaraScript*>& ScriptDependencyMap, const TMap<UNiagaraScript*, FString>& ScriptToEmitterNameMap)
+{
+	SCOPE_CYCLE_COUNTER(STAT_Niagara_Utilities_PrepareRapidIterationParameters);
+
+	TMap<UNiagaraScript*, FNiagaraParameterStore> ScriptToPreparedParameterStoreMap;
+
+	// Remove old and initialize new parameters.
+	for (UNiagaraScript* Script : Scripts)
+	{
+		FNiagaraParameterStore& ParameterStoreToPrepare = ScriptToPreparedParameterStoreMap.FindOrAdd(Script);
+		Script->RapidIterationParameters.CopyParametersTo(ParameterStoreToPrepare, false, FNiagaraParameterStore::EDataInterfaceCopyMethod::None);
+		const FString* EmitterName = ScriptToEmitterNameMap.Find(Script);
+		checkf(EmitterName != nullptr, TEXT("Script to emitter name map must have an entry for each script to be processed."));
+		Script->GetSource()->CleanUpOldAndInitializeNewRapidIterationParameters(*EmitterName, Script->GetUsage(), Script->GetUsageId(), ParameterStoreToPrepare);
+	}
+
+	// Copy parameters for dependencies.
+	for (auto It = ScriptToPreparedParameterStoreMap.CreateIterator(); It; ++It)
+	{
+		UNiagaraScript* Script = It.Key();
+		FNiagaraParameterStore& PreparedParameterStore = It.Value();
+		UNiagaraScript*const* DependentScriptPtr = ScriptDependencyMap.Find(Script);
+		if (DependentScriptPtr != nullptr)
+		{
+			UNiagaraScript* DependentScript = *DependentScriptPtr;
+			FNiagaraParameterStore* DependentPreparedParameterStore = ScriptToPreparedParameterStoreMap.Find(DependentScript);
+			checkf(DependentPreparedParameterStore != nullptr, TEXT("Dependent scripts must be one of the scripts being processed."));
+			PreparedParameterStore.CopyParametersTo(*DependentPreparedParameterStore, false, FNiagaraParameterStore::EDataInterfaceCopyMethod::None);
+		}
+	}
+
+	// Resolve prepared parameters with the source parameters.
+	for (auto It = ScriptToPreparedParameterStoreMap.CreateIterator(); It; ++It)
+	{
+		UNiagaraScript* Script = It.Key();
+		FNiagaraParameterStore& PreparedParameterStore = It.Value();
+
+		bool bOverwriteParameters = false;
+		if (Script->RapidIterationParameters.GetNumParameters() != PreparedParameterStore.GetNumParameters())
+		{
+			bOverwriteParameters = true;
+		}
+		else
+		{
+			const TMap<FNiagaraVariable, int32>& SourceParameterOffsets = Script->RapidIterationParameters.GetParameterOffests();
+			for (auto ParameterOffsetIt = SourceParameterOffsets.CreateConstIterator(); ParameterOffsetIt; ++ParameterOffsetIt)
+			{
+				const FNiagaraVariable& SourceParameter = ParameterOffsetIt.Key();
+				int32 SourceOffset = ParameterOffsetIt.Value();
+
+				int32 PreparedOffset = PreparedParameterStore.IndexOf(SourceParameter);
+				if (PreparedOffset == INDEX_NONE)
+				{
+					bOverwriteParameters = true;
+					break;
+				}
+				else
+				{
+					if (FMemory::Memcmp(
+						Script->RapidIterationParameters.GetParameterData(SourceOffset),
+						PreparedParameterStore.GetParameterData(PreparedOffset),
+						SourceParameter.GetSizeInBytes()) != 0)
+					{
+						bOverwriteParameters = true;
+						break;
+					}
+				}
+			}
+		}
+
+		if (bOverwriteParameters)
+		{
+			Script->RapidIterationParameters = PreparedParameterStore;
+		}
+	}
 }
\ No newline at end of file
diff --git a/Engine/Plugins/FX/Niagara/Source/Niagara/Private/NiagaraParameterStore.cpp b/Engine/Plugins/FX/Niagara/Source/Niagara/Private/NiagaraParameterStore.cpp
--- a/Engine/Plugins/FX/Niagara/Source/Niagara/Private/NiagaraParameterStore.cpp
+++ b/Engine/Plugins/FX/Niagara/Source/Niagara/Private/NiagaraParameterStore.cpp
@@ -62,6 +62,9 @@ FNiagaraParameterStore& FNiagaraParameterStore::operator=(const FNiagaraParamete
 	INC_MEMORY_STAT_BY(STAT_NiagaraParamStoreMemory, ParameterData.Num());
 	DataInterfaces = Other.DataInterfaces;
 	++LayoutVersion;
+#if WITH_EDITOR
+	OnChangedDelegate.Broadcast();
+#endif
 	//Don't copy bindings. We just want the data.
 	return *this;
 }
diff --git a/Engine/Plugins/FX/Niagara/Source/Niagara/Private/NiagaraScript.cpp b/Engine/Plugins/FX/Niagara/Source/Niagara/Private/NiagaraScript.cpp
--- a/Engine/Plugins/FX/Niagara/Source/Niagara/Private/NiagaraScript.cpp
+++ b/Engine/Plugins/FX/Niagara/Source/Niagara/Private/NiagaraScript.cpp
@@ -1348,11 +1348,6 @@ bool UNiagaraScript::LegacyCanBeRunOnGpu() const
 
 
 #if WITH_EDITORONLY_DATA
-void UNiagaraScript::CleanUpOldAndInitializeNewRapidIterationParameters(FString UniqueEmitterName)
-{
-	Source->CleanUpOldAndInitializeNewRapidIterationParameters(UniqueEmitterName, Usage, UsageId, RapidIterationParameters);
-}
-
 FGuid UNiagaraScript::GetBaseChangeID() const
 {
 	return Source->GetChangeID(); 
diff --git a/Engine/Plugins/FX/Niagara/Source/Niagara/Private/NiagaraSystem.cpp b/Engine/Plugins/FX/Niagara/Source/Niagara/Private/NiagaraSystem.cpp
--- a/Engine/Plugins/FX/Niagara/Source/Niagara/Private/NiagaraSystem.cpp
+++ b/Engine/Plugins/FX/Niagara/Source/Niagara/Private/NiagaraSystem.cpp
@@ -801,50 +801,60 @@ bool UNiagaraSystem::QueryCompileComplete(bool bWait, bool bDoPost, bool bDoNotA
 
 		InitEmitterSpawnAttributes();
 
-		// Clean up old, and initialize new rapid iterations parameters due to graph changes.
+		TArray<UNiagaraScript*> Scripts;
+		TMap<UNiagaraScript*, UNiagaraScript*> ScriptDependencyMap;
+		TMap<UNiagaraScript*, FString> ScriptToEmitterNameMap;
 		for (FEmitterCompiledScriptPair& EmitterCompiledScriptPair : ActiveCompilations[ActiveCompileIdx].EmitterCompiledScriptPairs)
 		{
-			if (EmitterCompiledScriptPair.Emitter == nullptr)
-			{
-				continue;
-			}
+			UNiagaraEmitter* Emitter = EmitterCompiledScriptPair.Emitter;
+			UNiagaraScript* CompiledScript = EmitterCompiledScriptPair.CompiledScript;
 
-			EmitterCompiledScriptPair.CompiledScript->CleanUpOldAndInitializeNewRapidIterationParameters(EmitterCompiledScriptPair.Emitter->GetUniqueEmitterName());
-		}
+			Scripts.AddUnique(CompiledScript);
+			ScriptToEmitterNameMap.Add(CompiledScript, Emitter != nullptr ? Emitter->GetUniqueEmitterName() : FString());
 
-		// Push rapid iteration parameters for dependent scripts.
-		for (FEmitterCompiledScriptPair& EmitterCompiledScriptPair : ActiveCompilations[ActiveCompileIdx].EmitterCompiledScriptPairs)
-		{
-			if (EmitterCompiledScriptPair.Emitter == nullptr)
+			if (CompiledScript->GetUsage() == ENiagaraScriptUsage::EmitterSpawnScript)
 			{
-				continue;
+				Scripts.AddUnique(SystemSpawnScript);
+				ScriptDependencyMap.Add(CompiledScript, SystemSpawnScript);
+				ScriptToEmitterNameMap.Add(SystemSpawnScript, FString());
 			}
 
-			const bool bOnlyAddParametersOnCopy = false;
-			if (EmitterCompiledScriptPair.CompiledScript->GetUsage() == ENiagaraScriptUsage::ParticleUpdateScript && EmitterCompiledScriptPair.Emitter->bInterpolatedSpawning)
-			{
-				EmitterCompiledScriptPair.CompiledScript->RapidIterationParameters.CopyParametersTo(
-					EmitterCompiledScriptPair.Emitter->SpawnScriptProps.Script->RapidIterationParameters, bOnlyAddParametersOnCopy, FNiagaraParameterStore::EDataInterfaceCopyMethod::None);
-			}
-			else if (EmitterCompiledScriptPair.CompiledScript->GetUsage() == ENiagaraScriptUsage::EmitterSpawnScript)
+			if (CompiledScript->GetUsage() == ENiagaraScriptUsage::EmitterUpdateScript)
 			{
-				EmitterCompiledScriptPair.CompiledScript->RapidIterationParameters.CopyParametersTo(
-					SystemSpawnScript->RapidIterationParameters, bOnlyAddParametersOnCopy, FNiagaraParameterStore::EDataInterfaceCopyMethod::None);
+				Scripts.AddUnique(SystemUpdateScript);
+				ScriptDependencyMap.Add(CompiledScript, SystemUpdateScript);
+				ScriptToEmitterNameMap.Add(SystemSpawnScript, FString());
 			}
-			else if (EmitterCompiledScriptPair.CompiledScript->GetUsage() == ENiagaraScriptUsage::EmitterUpdateScript)
+
+			if (CompiledScript->GetUsage() == ENiagaraScriptUsage::ParticleSpawnScript)
 			{
-				EmitterCompiledScriptPair.CompiledScript->RapidIterationParameters.CopyParametersTo(
-					SystemUpdateScript->RapidIterationParameters, bOnlyAddParametersOnCopy, FNiagaraParameterStore::EDataInterfaceCopyMethod::None);
+				if (Emitter->SimTarget == ENiagaraSimTarget::GPUComputeSim)
+				{
+					Scripts.AddUnique(Emitter->GetGPUComputeScript());
+					ScriptDependencyMap.Add(CompiledScript, Emitter->GetGPUComputeScript());
+					ScriptToEmitterNameMap.Add(Emitter->GetGPUComputeScript(), Emitter->GetUniqueEmitterName());
+				}
 			}
-			else if (EmitterCompiledScriptPair.CompiledScript->GetUsage() == ENiagaraScriptUsage::ParticleGPUComputeScript && EmitterCompiledScriptPair.Emitter->SimTarget != ENiagaraSimTarget::CPUSim)
+
+			if (CompiledScript->GetUsage() == ENiagaraScriptUsage::ParticleUpdateScript)
 			{
-				EmitterCompiledScriptPair.Emitter->SpawnScriptProps.Script->RapidIterationParameters.CopyParametersTo(
-					EmitterCompiledScriptPair.CompiledScript->RapidIterationParameters, bOnlyAddParametersOnCopy, FNiagaraParameterStore::EDataInterfaceCopyMethod::None);
-				EmitterCompiledScriptPair.Emitter->UpdateScriptProps.Script->RapidIterationParameters.CopyParametersTo(
-					EmitterCompiledScriptPair.CompiledScript->RapidIterationParameters, bOnlyAddParametersOnCopy, FNiagaraParameterStore::EDataInterfaceCopyMethod::None);
+				if (Emitter->SimTarget == ENiagaraSimTarget::GPUComputeSim)
+				{
+					Scripts.AddUnique(Emitter->GetGPUComputeScript());
+					ScriptDependencyMap.Add(CompiledScript, Emitter->GetGPUComputeScript());
+					ScriptToEmitterNameMap.Add(Emitter->GetGPUComputeScript(), Emitter->GetUniqueEmitterName());
+				}
+				else if (Emitter->bInterpolatedSpawning)
+				{
+					Scripts.AddUnique(Emitter->SpawnScriptProps.Script);
+					ScriptDependencyMap.Add(CompiledScript, Emitter->SpawnScriptProps.Script);
+					ScriptToEmitterNameMap.Add(Emitter->SpawnScriptProps.Script, Emitter->GetUniqueEmitterName());
+				}
 			}
 		}
 
+		FNiagaraUtilities::PrepareRapidIterationParameters(Scripts, ScriptDependencyMap, ScriptToEmitterNameMap);
+
 		ActiveCompilations[ActiveCompileIdx].RootObjects.Empty();
 
 		DetermineIfSolo();
diff --git a/Engine/Plugins/FX/Niagara/Source/Niagara/Public/NiagaraCommon.h b/Engine/Plugins/FX/Niagara/Source/Niagara/Public/NiagaraCommon.h
--- a/Engine/Plugins/FX/Niagara/Source/Niagara/Public/NiagaraCommon.h
+++ b/Engine/Plugins/FX/Niagara/Source/Niagara/Public/NiagaraCommon.h
@@ -513,4 +513,15 @@ namespace FNiagaraUtilities
 	FNiagaraVariable NIAGARA_API ConvertVariableToRapidIterationConstantName(FNiagaraVariable InVar, const TCHAR* InEmitterName, ENiagaraScriptUsage InUsage);
 
 	void CollectScriptDataInterfaceParameters(const UObject& Owner, const TArray<UNiagaraScript*>& Scripts, FNiagaraParameterStore& OutDataInterfaceParameters);
+
+	/**
+	 * Prepares rapid iteration parameter stores for simulation by removing old parameters no longer used by functions, by initializing new parameters
+	 * added to functions, and by copying parameters across parameter stores for interscript dependencies.
+	 * @param Scripts The scripts who's rapid iteration parameter stores will be processed.
+	 * @param ScriptDependencyMap A map of script dependencies where the key is the source script and the value is the script which depends on the source.  All scripts in this
+	 * map must be contained in the Scripts array, both keys and values.
+	 * @param ScriptToEmitterNameMap An array of scripts to the name of the emitter than owns them.  If this is a system script the name can be empty.  All scripts in the
+	 * scripts array must have an entry in this map.
+	 */
+	void NIAGARA_API PrepareRapidIterationParameters(const TArray<UNiagaraScript*>& Scripts, const TMap<UNiagaraScript*, UNiagaraScript*>& ScriptDependencyMap, const TMap<UNiagaraScript*, FString>& ScriptToEmitterNameMap);
 };
\ No newline at end of file
diff --git a/Engine/Plugins/FX/Niagara/Source/NiagaraEditor/Private/Toolkits/NiagaraScriptToolkit.cpp b/Engine/Plugins/FX/Niagara/Source/NiagaraEditor/Private/Toolkits/NiagaraScriptToolkit.cpp
--- a/Engine/Plugins/FX/Niagara/Source/NiagaraEditor/Private/Toolkits/NiagaraScriptToolkit.cpp
+++ b/Engine/Plugins/FX/Niagara/Source/NiagaraEditor/Private/Toolkits/NiagaraScriptToolkit.cpp
@@ -65,6 +65,7 @@ FNiagaraScriptToolkit::FNiagaraScriptToolkit()
 FNiagaraScriptToolkit::~FNiagaraScriptToolkit()
 {
 	EditedNiagaraScript->OnVMScriptCompiled().RemoveAll(this);
+	ScriptViewModel->GetGraphViewModel()->GetGraph()->RemoveOnGraphNeedsRecompileHandler(OnEditedScriptGraphChangedHandle);
 }
 
 void FNiagaraScriptToolkit::RegisterTabSpawners(const TSharedRef<class FTabManager>& InTabManager)
@@ -112,8 +113,12 @@ void FNiagaraScriptToolkit::Initialize( const EToolkitMode::Type Mode, const TSh
 	GetTransientPackage()->LinkerCustomVersion.Empty();
 	EditedNiagaraScript = (UNiagaraScript*)StaticDuplicateObject(OriginalNiagaraScript, GetTransientPackage(), NAME_None, ~RF_Standalone, UNiagaraScript::StaticClass());
 	EditedNiagaraScript->OnVMScriptCompiled().AddSP(this, &FNiagaraScriptToolkit::OnVMScriptCompiled);
+	bEditedScriptHasPendingChanges = false;
 	
 	ScriptViewModel = MakeShareable(new FNiagaraScriptViewModel(EditedNiagaraScript, LOCTEXT("NiagaraScriptDisplayName", "Niagara Script"), ENiagaraParameterEditMode::EditAll));
+	OnEditedScriptGraphChangedHandle = ScriptViewModel->GetGraphViewModel()->GetGraph()->AddOnGraphNeedsRecompileHandler(
+		FOnGraphChanged::FDelegate::CreateRaw(this, &FNiagaraScriptToolkit::OnEditedScriptGraphChanged));
+
 	DetailsSelection = MakeShareable(new FNiagaraObjectSelection());
 	DetailsSelection->SetSelectedObject(EditedNiagaraScript);
 	
@@ -229,7 +234,7 @@ TSharedRef<SDockTab> FNiagaraScriptToolkit::SpawnTabNodeGraph( const FSpawnTabAr
 		];
 }
 
-void FNiagaraScriptToolkit::OnDetailsSelectionPropertyFinishedChanging(const FPropertyChangedEvent& InEvent)
+void FNiagaraScriptToolkit::OnEditedScriptPropertyFinishedChanging(const FPropertyChangedEvent& InEvent)
 {
 	// We need to synchronize the Usage field in the property editor with the actual node
 	// in the graph.
@@ -261,6 +266,7 @@ void FNiagaraScriptToolkit::OnDetailsSelectionPropertyFinishedChanging(const FPr
 			}
 		}
 	}
+	bEditedScriptHasPendingChanges = true;
 }
 
 void FNiagaraScriptToolkit::OnVMScriptCompiled(UNiagaraScript* InScript)
@@ -281,7 +287,7 @@ TSharedRef<SDockTab> FNiagaraScriptToolkit::SpawnTabNodeDetails(const FSpawnTabA
 	FDetailsViewArgs DetailsViewArgs(false, false, true, FDetailsViewArgs::HideNameArea, true);
 	TSharedRef<IDetailsView> DetailsView = PropertyEditorModule.CreateDetailView(DetailsViewArgs);
 
-	DetailsView->OnFinishedChangingProperties().AddRaw(this, &FNiagaraScriptToolkit::OnDetailsSelectionPropertyFinishedChanging);
+	DetailsView->OnFinishedChangingProperties().AddRaw(this, &FNiagaraScriptToolkit::OnEditedScriptPropertyFinishedChanging);
 	DetailsView->RegisterInstancedCustomPropertyLayout(
 		UNiagaraScript::StaticClass(),
 		FOnGetDetailCustomizationInstance::CreateStatic(&FNiagaraScriptDetails::MakeInstance, ScriptViewModelWeakPtr)
@@ -455,9 +461,7 @@ void FNiagaraScriptToolkit::RefreshNodes()
 
 bool FNiagaraScriptToolkit::IsEditScriptDifferentFromOriginalScript() const
 {
-	UNiagaraScript* EditScript = ScriptViewModel->GetScript(OriginalNiagaraScript->GetUsage(), OriginalNiagaraScript->GetUsageId());
-	return EditScript->AreScriptAndSourceSynchronized() == false ||
-		EditScript->GetBaseChangeID() != OriginalNiagaraScript->GetBaseChangeID();
+	return bEditedScriptHasPendingChanges;
 }
 
 void FNiagaraScriptToolkit::OnApply()
@@ -631,6 +635,7 @@ void FNiagaraScriptToolkit::UpdateOriginalNiagaraScript()
 	FNiagaraEditorUtilities::CompileExistingEmitters(AffectedEmitters);
 
 	GWarn->EndSlowTask();
+	bEditedScriptHasPendingChanges = false;
 }
 
 
@@ -666,4 +671,9 @@ bool FNiagaraScriptToolkit::OnRequestClose()
 	return true;
 }
 
+void FNiagaraScriptToolkit::OnEditedScriptGraphChanged(const FEdGraphEditAction& InAction)
+{
+	bEditedScriptHasPendingChanges = true;
+}
+
 #undef LOCTEXT_NAMESPACE
diff --git a/Engine/Plugins/FX/Niagara/Source/NiagaraEditor/Private/Toolkits/NiagaraScriptToolkit.h b/Engine/Plugins/FX/Niagara/Source/NiagaraEditor/Private/Toolkits/NiagaraScriptToolkit.h
--- a/Engine/Plugins/FX/Niagara/Source/NiagaraEditor/Private/Toolkits/NiagaraScriptToolkit.h
+++ b/Engine/Plugins/FX/Niagara/Source/NiagaraEditor/Private/Toolkits/NiagaraScriptToolkit.h
@@ -17,6 +17,7 @@ class UNiagaraScript;
 class UNiagaraScriptSource;
 class FNiagaraScriptViewModel;
 class FNiagaraObjectSelection;
+struct FEdGraphEditAction;
 
 /** Viewer/editor for a DataTable */
 class FNiagaraScriptToolkit : public FAssetEditorToolkit, public FGCObject
@@ -61,10 +62,10 @@ protected:
 	virtual void SaveAssetAs_Execute() override;
 	virtual bool OnRequestClose() override;
 
-	void OnDetailsSelectionPropertyFinishedChanging(const FPropertyChangedEvent& InEvent);
-	void OnVMScriptCompiled(UNiagaraScript* InScript);
-
 private:
+	void OnEditedScriptPropertyFinishedChanging(const FPropertyChangedEvent& InEvent);
+
+	void OnVMScriptCompiled(UNiagaraScript* InScript);
 
 	/** Spawns the tab with the update graph inside */
 	TSharedRef<SDockTab> SpawnTabNodeGraph(const FSpawnTabArgs& Args);
@@ -94,7 +95,6 @@ private:
 	FSlateIcon GetRefreshStatusImage() const;
 	FText GetRefreshStatusTooltip() const;
 	
-private:
 	bool IsEditScriptDifferentFromOriginalScript() const;
 
 	/** Command for the apply button */
@@ -103,6 +103,10 @@ private:
 
 	void UpdateOriginalNiagaraScript();
 
+	void OnEditedScriptGraphChanged(const FEdGraphEditAction& InAction);
+
+private:
+
 	/** The Script being edited */
 	TSharedPtr<FNiagaraScriptViewModel> ScriptViewModel;
 
@@ -117,4 +121,8 @@ private:
 	/** Stats log, with the log listing that it reflects */
 	TSharedPtr<class SWidget> Stats;
 	TSharedPtr<class IMessageLogListing> StatsListing;
+
+	FDelegateHandle OnEditedScriptGraphChangedHandle;
+
+	bool bEditedScriptHasPendingChanges;
 };
diff --git a/Engine/Plugins/FX/Niagara/Source/NiagaraEditor/Private/Toolkits/NiagaraSystemToolkit.cpp b/Engine/Plugins/FX/Niagara/Source/NiagaraEditor/Private/Toolkits/NiagaraSystemToolkit.cpp
--- a/Engine/Plugins/FX/Niagara/Source/NiagaraEditor/Private/Toolkits/NiagaraSystemToolkit.cpp
+++ b/Engine/Plugins/FX/Niagara/Source/NiagaraEditor/Private/Toolkits/NiagaraSystemToolkit.cpp
@@ -197,19 +197,38 @@ void FNiagaraSystemToolkit::InitializeWithEmitter(const EToolkitMode::Type Mode,
 
 	Emitter = &InEmitter;
 
-	// Before copying the emitter clean up and propagate the rapid iteration parameters so that the post compile cleanup 
-	// and propagation doesn't cause the change ids to become out of sync.
+	// Before copying the emitter prepare the rapid iteration parameters so that the post compile prepare doesn't
+	// cause the change ids to become out of sync.
 	FString EmitterName = "Emitter";
-	Emitter->EmitterSpawnScriptProps.Script->CleanUpOldAndInitializeNewRapidIterationParameters(EmitterName);
-	Emitter->EmitterUpdateScriptProps.Script->CleanUpOldAndInitializeNewRapidIterationParameters(EmitterName);
-	Emitter->SpawnScriptProps.Script->CleanUpOldAndInitializeNewRapidIterationParameters(EmitterName);
-	Emitter->UpdateScriptProps.Script->CleanUpOldAndInitializeNewRapidIterationParameters(EmitterName);
-	for (const FNiagaraEventScriptProperties& EventHandler : Emitter->GetEventHandlers())
+	TArray<UNiagaraScript*> Scripts;
+	TMap<UNiagaraScript*, UNiagaraScript*> ScriptDependencyMap;
+	TMap<UNiagaraScript*, FString> ScriptToEmitterNameMap;
+
+	Scripts.Add(Emitter->EmitterSpawnScriptProps.Script);
+	ScriptToEmitterNameMap.Add(Emitter->EmitterSpawnScriptProps.Script, EmitterName);
+
+	Scripts.Add(Emitter->EmitterUpdateScriptProps.Script);
+	ScriptToEmitterNameMap.Add(Emitter->EmitterUpdateScriptProps.Script, EmitterName);
+
+	Scripts.Add(Emitter->SpawnScriptProps.Script);
+	ScriptToEmitterNameMap.Add(Emitter->SpawnScriptProps.Script, EmitterName);
+
+	Scripts.Add(Emitter->UpdateScriptProps.Script);
+	ScriptToEmitterNameMap.Add(Emitter->UpdateScriptProps.Script, EmitterName);
+
+	if (Emitter->SimTarget == ENiagaraSimTarget::GPUComputeSim)
+	{
+		Scripts.Add(Emitter->GetGPUComputeScript());
+		ScriptToEmitterNameMap.Add(Emitter->GetGPUComputeScript(), EmitterName);
+		ScriptDependencyMap.Add(Emitter->SpawnScriptProps.Script, Emitter->GetGPUComputeScript());
+		ScriptDependencyMap.Add(Emitter->UpdateScriptProps.Script, Emitter->GetGPUComputeScript());
+	}
+	else if (Emitter->bInterpolatedSpawning)
 	{
-		EventHandler.Script->CleanUpOldAndInitializeNewRapidIterationParameters(EmitterName);
+		ScriptDependencyMap.Add(Emitter->UpdateScriptProps.Script, Emitter->SpawnScriptProps.Script);
 	}
-	Emitter->UpdateScriptProps.Script->RapidIterationParameters.CopyParametersTo(
-		Emitter->SpawnScriptProps.Script->RapidIterationParameters, false, FNiagaraParameterStore::EDataInterfaceCopyMethod::None);
+
+	FNiagaraUtilities::PrepareRapidIterationParameters(Scripts, ScriptDependencyMap, ScriptToEmitterNameMap);
 
 	ResetLoaders(GetTransientPackage()); // Make sure that we're not going to get invalid version number linkers into the package we are going into. 
 	GetTransientPackage()->LinkerCustomVersion.Empty();
