Fix focus search box on add parameter menu

diff --git a/Engine/Plugins/FX/Niagara/Source/NiagaraEditor/Private/Widgets/SNiagaraParameterMapView.cpp b/Engine/Plugins/FX/Niagara/Source/NiagaraEditor/Private/Widgets/SNiagaraParameterMapView.cpp
--- a/Engine/Plugins/FX/Niagara/Source/NiagaraEditor/Private/Widgets/SNiagaraParameterMapView.cpp
+++ b/Engine/Plugins/FX/Niagara/Source/NiagaraEditor/Private/Widgets/SNiagaraParameterMapView.cpp
@@ -136,6 +136,7 @@ void SNiagaraParameterMapView::Construct(const FArguments& InArgs, const TShared
 	bNeedsRefresh = false;
 	ToolkitType = InToolkitType;
 	ToolkitCommands = InToolkitCommands;
+	AddParameterButtons.SetNum(NiagaraParameterMapSectionID::OTHER + 1);
 
 	SelectedObjects = InSelectedObjects;
 	SelectedObjects->OnSelectedObjectsChanged().AddSP(this, &SNiagaraParameterMapView::SelectedObjectsChanged);
@@ -476,6 +477,7 @@ TSharedRef<SWidget> SNiagaraParameterMapView::CreateAddToSectionButton(const Nia
 			.ShadowOffset(FVector2D(1,1))
 		]
 	];
+	AddParameterButtons[InSection] = Button;
 
 	return Button.ToSharedRef();
 }
@@ -496,7 +498,7 @@ TSharedRef<SWidget> SNiagaraParameterMapView::OnGetParameterMenu(const NiagaraPa
 		.ShowGraphParameters(false)
 		.AutoExpandMenu(true);
 
-	FSlateApplication::Get().SetKeyboardFocus(MenuWidget->GetSearchBox());
+	AddParameterButtons[InSection]->SetMenuContentWidgetToFocus(MenuWidget->GetSearchBox()->AsShared());
 	return MenuWidget;
 }
 
diff --git a/Engine/Plugins/FX/Niagara/Source/NiagaraEditor/Private/Widgets/SNiagaraParameterMapView.h b/Engine/Plugins/FX/Niagara/Source/NiagaraEditor/Private/Widgets/SNiagaraParameterMapView.h
--- a/Engine/Plugins/FX/Niagara/Source/NiagaraEditor/Private/Widgets/SNiagaraParameterMapView.h
+++ b/Engine/Plugins/FX/Niagara/Source/NiagaraEditor/Private/Widgets/SNiagaraParameterMapView.h
@@ -18,6 +18,7 @@ class SGraphActionMenu;
 class SEditableTextBox;
 class SExpanderArrow;
 class SSearchBox;
+class SComboButton;
 class SNiagaraGraphPinAdd;
 class FNiagaraObjectSelection;
 class UNiagaraGraph;
@@ -144,6 +145,9 @@ private:
 	/** The filter box that handles filtering for both graph action menus. */
 	TSharedPtr<SSearchBox> FilterBox;
 
+	/** Add parameter buttons for all sections. */
+	TArray<TSharedPtr<SComboButton>> AddParameterButtons;
+
 	/** The selected objects being viewed and edited by this widget. */
 	TSharedPtr<FNiagaraObjectSelection> SelectedObjects;
 
