Fixes for metadata details in script toolkit (now the apply and compile buttons refresh and sort the metadata collection). The metadata functionality is fixed. Delete, add and modify work just as before, but the sorting isn't applied because refreshing the whole collection is skipped for internal changes.
#tests none

diff --git a/Engine/Plugins/FX/Niagara/Source/NiagaraEditor/Private/Customizations/NiagaraScriptDetails.cpp b/Engine/Plugins/FX/Niagara/Source/NiagaraEditor/Private/Customizations/NiagaraScriptDetails.cpp
--- a/Engine/Plugins/FX/Niagara/Source/NiagaraEditor/Private/Customizations/NiagaraScriptDetails.cpp
+++ b/Engine/Plugins/FX/Niagara/Source/NiagaraEditor/Private/Customizations/NiagaraScriptDetails.cpp
@@ -179,7 +179,7 @@ void FNiagaraScriptDetails::CustomizeDetails(IDetailLayoutBuilder& DetailBuilder
 				.HAlign(HAlign_Right)
 				.VAlign(VAlign_Center)
 				.ContentPadding(1)
-				.ToolTipText(LOCTEXT("RefreshMetadataToolTip", "Save metadata changes to graph and refresh the view"))
+				.ToolTipText(LOCTEXT("RefreshMetadataToolTip", "Refresh the view according to the latest Editor Sort Priority values"))
 				.OnClicked(this, &FNiagaraScriptDetails::OnRefreshMetadata)
 				.Content()
 				[
@@ -194,7 +194,7 @@ void FNiagaraScriptDetails::CustomizeDetails(IDetailLayoutBuilder& DetailBuilder
 
 FReply FNiagaraScriptDetails::OnRefreshMetadata()
 {
-	ScriptViewModel->GetMetadataCollectionViewModel()->RequestRefresh(true);
+	ScriptViewModel->GetMetadataCollectionViewModel()->RequestRefresh();
     return FReply::Handled();
 }
 
diff --git a/Engine/Plugins/FX/Niagara/Source/NiagaraEditor/Private/Toolkits/NiagaraScriptToolkit.cpp b/Engine/Plugins/FX/Niagara/Source/NiagaraEditor/Private/Toolkits/NiagaraScriptToolkit.cpp
--- a/Engine/Plugins/FX/Niagara/Source/NiagaraEditor/Private/Toolkits/NiagaraScriptToolkit.cpp
+++ b/Engine/Plugins/FX/Niagara/Source/NiagaraEditor/Private/Toolkits/NiagaraScriptToolkit.cpp
@@ -452,6 +452,7 @@ FText FNiagaraScriptToolkit::GetRefreshStatusTooltip() const
 void FNiagaraScriptToolkit::CompileScript(bool bForce)
 {
 	ScriptViewModel->CompileStandaloneScript();
+	ScriptViewModel->RefreshMetadataCollection();
 }
 
 void FNiagaraScriptToolkit::RefreshNodes()
@@ -539,6 +540,7 @@ void FNiagaraScriptToolkit::UpdateOriginalNiagaraScript()
 
 	// Compile and then overwrite the original script in place by constructing a new one with the same name
 	ScriptViewModel->CompileStandaloneScript();
+	ScriptViewModel->RefreshMetadataCollection();
 	OriginalNiagaraScript = (UNiagaraScript*)StaticDuplicateObject(EditedNiagaraScript, OriginalNiagaraScript->GetOuter(), OriginalNiagaraScript->GetFName(),
 		RF_AllFlags,
 		OriginalNiagaraScript->GetClass());
diff --git a/Engine/Plugins/FX/Niagara/Source/NiagaraEditor/Private/ViewModels/NiagaraMetaDataCollectionViewModel.cpp b/Engine/Plugins/FX/Niagara/Source/NiagaraEditor/Private/ViewModels/NiagaraMetaDataCollectionViewModel.cpp
--- a/Engine/Plugins/FX/Niagara/Source/NiagaraEditor/Private/ViewModels/NiagaraMetaDataCollectionViewModel.cpp
+++ b/Engine/Plugins/FX/Niagara/Source/NiagaraEditor/Private/ViewModels/NiagaraMetaDataCollectionViewModel.cpp
@@ -1,4 +1,4 @@
-﻿// Copyright 1998-2018 Epic Games, Inc. All Rights Reserved.
+// Copyright 1998-2018 Epic Games, Inc. All Rights Reserved.
 
 #include "NiagaraMetaDataCollectionViewModel.h"
 #include "NiagaraTypes.h"
@@ -18,7 +18,7 @@
 FNiagaraMetaDataCollectionViewModel::FNiagaraMetaDataCollectionViewModel()
 	: ModuleGraph(nullptr)
 	, bNeedsRefresh(false)
-	, bNeedsNotify(false)
+	, bInternalGraphChange(false)
 {
 }
 
@@ -26,9 +26,8 @@ void FNiagaraMetaDataCollectionViewModel::Tick(float DeltaTime)
 {
 	if (bNeedsRefresh)
 	{
-		Refresh(bNeedsNotify);
+		Refresh();
 		bNeedsRefresh = false;
-		bNeedsNotify = false;
 	}
 }
 
@@ -56,7 +55,7 @@ void FNiagaraMetaDataCollectionViewModel::SetGraph(UNiagaraGraph* InGraph)
 	// now build variables
 	ModuleGraph = InGraph;
 	
-	Refresh(false);
+	Refresh();
 
 	OnGraphChangedHandle = ModuleGraph->AddOnGraphChangedHandler(FOnGraphChanged::FDelegate::CreateRaw(this, &FNiagaraMetaDataCollectionViewModel::OnGraphChanged));
 	OnRecompileHandle = ModuleGraph->AddOnGraphNeedsRecompileHandler(FOnGraphChanged::FDelegate::CreateRaw(this, &FNiagaraMetaDataCollectionViewModel::OnGraphChanged));
@@ -77,23 +76,27 @@ TArray<TSharedRef<FNiagaraMetaDataViewModel>>& FNiagaraMetaDataCollectionViewMod
 	return MetaDataViewModels;
 }
 
-void FNiagaraMetaDataCollectionViewModel::RequestRefresh(bool bNotifyGraph)
+void FNiagaraMetaDataCollectionViewModel::RequestRefresh()
 {
 	bNeedsRefresh = true;
-	bNeedsNotify = bNotifyGraph;
 }
 
 void FNiagaraMetaDataCollectionViewModel::OnGraphChanged(const struct FEdGraphEditAction& InAction)
 {
-	RequestRefresh(false);
+	if (!bInternalGraphChange)
+	{
+		RequestRefresh();
+	}
 } 
 
-void FNiagaraMetaDataCollectionViewModel::BroadcastChanged()
+void FNiagaraMetaDataCollectionViewModel::ChildMetadataChanged()
 {
+	TGuardValue<bool> UpdateGuard(bInternalGraphChange, true);
 	OnCollectionChangedDelegate.Broadcast();
+	ModuleGraph->NotifyGraphNeedsRecompile();
 }
 
-void FNiagaraMetaDataCollectionViewModel::Refresh(bool bNotifyGraph)
+void FNiagaraMetaDataCollectionViewModel::Refresh()
 {
 	if (!ModuleGraph)
 	{
@@ -101,13 +104,14 @@ void FNiagaraMetaDataCollectionViewModel::Refresh(bool bNotifyGraph)
 	}
 
 	// Get the variable metadata from the graph
-	MetaDataViewModels.Empty();
+	CleanupMetadata();
 	for (auto& MetadataElement : ModuleGraph->GetAllMetaData())
 	{
 		TSharedPtr<FNiagaraMetaDataViewModel> MetadataViewModel = GetMetadataViewModelForVariable(MetadataElement.Key);
 		if (!MetadataViewModel.IsValid())
 		{
 			MetadataViewModel = MakeShared<FNiagaraMetaDataViewModel>(MetadataElement.Key, *ModuleGraph);
+			MetadataViewModel->OnMetadataChanged().AddRaw(this, &FNiagaraMetaDataCollectionViewModel::ChildMetadataChanged);
 			MetaDataViewModels.Add(MetadataViewModel.ToSharedRef());
 		}
 
@@ -121,12 +125,10 @@ void FNiagaraMetaDataCollectionViewModel::Refresh(bool bNotifyGraph)
 			}
 		}
 	}
+
 	SortViewModels();
-	if (bNotifyGraph)
-	{
-		ModuleGraph->NotifyGraphNeedsRecompile();
-	}
-	BroadcastChanged();
+
+	OnCollectionChangedDelegate.Broadcast();
 }
 
 void FNiagaraMetaDataCollectionViewModel::SortViewModels()
@@ -170,11 +172,20 @@ void FNiagaraMetaDataCollectionViewModel::SortViewModels()
 
 void FNiagaraMetaDataCollectionViewModel::Cleanup()
 {
-	MetaDataViewModels.Empty();
+	CleanupMetadata();
 	ModuleGraph->RemoveOnGraphChangedHandler(OnGraphChangedHandle);
 	ModuleGraph->RemoveOnGraphNeedsRecompileHandler(OnRecompileHandle);
 }
 
+void FNiagaraMetaDataCollectionViewModel::CleanupMetadata()
+{
+	for (int32 i = 0; i < MetaDataViewModels.Num(); i++)
+	{
+		MetaDataViewModels[i]->OnMetadataChanged().RemoveAll(this);
+	}
+	MetaDataViewModels.Empty();
+}
+
 FNiagaraMetaDataCollectionViewModel::~FNiagaraMetaDataCollectionViewModel()
 {
 	Cleanup();
diff --git a/Engine/Plugins/FX/Niagara/Source/NiagaraEditor/Private/ViewModels/NiagaraMetaDataCollectionViewModel.h b/Engine/Plugins/FX/Niagara/Source/NiagaraEditor/Private/ViewModels/NiagaraMetaDataCollectionViewModel.h
--- a/Engine/Plugins/FX/Niagara/Source/NiagaraEditor/Private/ViewModels/NiagaraMetaDataCollectionViewModel.h
+++ b/Engine/Plugins/FX/Niagara/Source/NiagaraEditor/Private/ViewModels/NiagaraMetaDataCollectionViewModel.h
@@ -1,4 +1,4 @@
-﻿// Copyright 1998-2018 Epic Games, Inc. All Rights Reserved.
+// Copyright 1998-2018 Epic Games, Inc. All Rights Reserved.
 
 #pragma once
 
@@ -31,19 +31,23 @@ public:
 	TArray<TSharedRef<FNiagaraMetaDataViewModel>>& GetVariableModels();
 	
 	/** Request refresh the data from the graph and builds the viewmodels next frame. */
-	void RequestRefresh(bool bNotifyGraph);
+	void RequestRefresh();
 
 	/** returns the delegate to be called when the collection changes*/
 	FOnCollectionChanged& OnCollectionChanged() { return OnCollectionChangedDelegate; }
 
 private:
-	void Refresh(bool bNotifyGraph);
+	void Refresh();
 	void SortViewModels();
 	/** Callback for when the graph changes and the collection viewmodel needs to react */
 	void OnGraphChanged(const struct FEdGraphEditAction& InAction);
 	void Cleanup();
-	void BroadcastChanged();
+	/** Removes metadata listeners and empties MetaDataViewModels array */
+	void CleanupMetadata();
+	/** Called when one of the metadata in the viewmodel array has changed */
+	void ChildMetadataChanged();
 	TSharedPtr<FNiagaraMetaDataViewModel> GetMetadataViewModelForVariable(FNiagaraVariable InVariable);
+
 private:
 	/** The variables */
 	TArray <TSharedRef<FNiagaraMetaDataViewModel>> MetaDataViewModels;
@@ -61,6 +65,6 @@ private:
 	/** Refresh the UI next frames. */
 	bool bNeedsRefresh;
 
-	/** Needs to notify the graph on the next refesh. */
-	bool bNeedsNotify;
+	/** Guard flag that goes up when the graph change was done internally within this object, so we need to avoid refreshing. */
+	bool bInternalGraphChange;
 };
\ No newline at end of file
diff --git a/Engine/Plugins/FX/Niagara/Source/NiagaraEditor/Private/ViewModels/NiagaraScriptViewModel.cpp b/Engine/Plugins/FX/Niagara/Source/NiagaraEditor/Private/ViewModels/NiagaraScriptViewModel.cpp
--- a/Engine/Plugins/FX/Niagara/Source/NiagaraEditor/Private/ViewModels/NiagaraScriptViewModel.cpp
+++ b/Engine/Plugins/FX/Niagara/Source/NiagaraEditor/Private/ViewModels/NiagaraScriptViewModel.cpp
@@ -391,6 +391,11 @@ TSharedRef<FNiagaraMetaDataCollectionViewModel> FNiagaraScriptViewModel::GetMeta
 	return MetaDataCollectionViewModel;
 }
 
+void FNiagaraScriptViewModel::RefreshMetadataCollection()
+{
+	MetaDataCollectionViewModel->RequestRefresh();
+}
+
 TSharedRef<FNiagaraScriptGraphViewModel> FNiagaraScriptViewModel::GetGraphViewModel()
 {
 	return GraphViewModel;
diff --git a/Engine/Plugins/FX/Niagara/Source/NiagaraEditor/Private/ViewModels/NiagaraScriptViewModel.h b/Engine/Plugins/FX/Niagara/Source/NiagaraEditor/Private/ViewModels/NiagaraScriptViewModel.h
--- a/Engine/Plugins/FX/Niagara/Source/NiagaraEditor/Private/ViewModels/NiagaraScriptViewModel.h
+++ b/Engine/Plugins/FX/Niagara/Source/NiagaraEditor/Private/ViewModels/NiagaraScriptViewModel.h
@@ -41,6 +41,9 @@ public:
 	/** Gets the view model for the metadata collection. */
 	TSharedRef<FNiagaraMetaDataCollectionViewModel> GetMetadataCollectionViewModel();
 
+	/** Refreshes the metadata collection */
+	void RefreshMetadataCollection();
+
 	TSharedRef<INiagaraParameterCollectionViewModel> GetInputParameterMapViewModel();
 	TSharedRef<INiagaraParameterCollectionViewModel> GetOutputParameterMapViewModel();
 
