Now generating compiler debug info for VM shaders just like the rest of Materials using the r.DumpShaderDebugInfo

#tests now properly generate data in a VM folder sibling to other generated debug shader data

diff --git a/Engine/Plugins/FX/Niagara/Source/NiagaraEditor/Private/NiagaraCompiler.cpp b/Engine/Plugins/FX/Niagara/Source/NiagaraEditor/Private/NiagaraCompiler.cpp
--- a/Engine/Plugins/FX/Niagara/Source/NiagaraEditor/Private/NiagaraCompiler.cpp
+++ b/Engine/Plugins/FX/Niagara/Source/NiagaraEditor/Private/NiagaraCompiler.cpp
@@ -27,6 +27,7 @@
 #include "ShaderCore.h"
 #include "EdGraphSchema_Niagara.h"
 #include "Misc/FileHelper.h"
+#include "ShaderCompiler.h"
 
 #define LOCTEXT_NAMESPACE "NiagaraCompiler"
 
@@ -39,7 +40,7 @@ DECLARE_CYCLE_STAT(TEXT("Niagara - Module - CompileShader_VectorVMSucceeded"), S
 DECLARE_CYCLE_STAT(TEXT("Niagara - ScriptSource - PreCompile"), STAT_NiagaraEditor_ScriptSource_PreCompile, STATGROUP_NiagaraEditor);
 DECLARE_CYCLE_STAT(TEXT("Niagara - HlslCompiler - TestCompileShader_VectorVM"), STAT_NiagaraEditor_HlslCompiler_TestCompileShader_VectorVM, STATGROUP_NiagaraEditor);
 
-static int32 GbForceNiagaraTranslatorSingleThreaded = 0;
+static int32 GbForceNiagaraTranslatorSingleThreaded = 1;
 static FAutoConsoleVariableRef CVarForceNiagaraTranslatorSingleThreaded(
 	TEXT("fx.ForceNiagaraTranslatorSingleThreaded"),
 	GbForceNiagaraTranslatorSingleThreaded,
@@ -343,6 +344,7 @@ void FNiagaraCompileRequestData::FinishPrecompile(UNiagaraScriptSource* ScriptSo
 {
 	{
 		ENiagaraScriptCompileStatusEnum = FindObject<UEnum>(ANY_PACKAGE, TEXT("ENiagaraScriptCompileStatus"), true);
+		ENiagaraScriptUsageEnum = FindObject<UEnum>(ANY_PACKAGE, TEXT("ENiagaraScriptUsage"), true);
 
 		PrecompiledHistories.Empty();
 
@@ -409,6 +411,8 @@ TSharedPtr<FNiagaraCompileRequestDataBase, ESPMode::ThreadSafe> FNiagaraEditorMo
 	TSharedPtr<FNiagaraCompileRequestData, ESPMode::ThreadSafe> BasePtr = MakeShared<FNiagaraCompileRequestData, ESPMode::ThreadSafe>();
 	TArray<TSharedPtr<FNiagaraCompileRequestData, ESPMode::ThreadSafe>> DependentRequests;
 
+	BasePtr->SourceName = InObj->GetName();
+
 	if (Script)
 	{
 		UNiagaraScriptSource* Source = Cast<UNiagaraScriptSource>(Script->GetSource());
@@ -437,6 +441,7 @@ TSharedPtr<FNiagaraCompileRequestDataBase, ESPMode::ThreadSafe> FNiagaraEditorMo
 			TSharedPtr<FNiagaraCompileRequestData, ESPMode::ThreadSafe> EmitterPtr = MakeShared<FNiagaraCompileRequestData, ESPMode::ThreadSafe>();
 			EmitterPtr->DeepCopyGraphs(Cast<UNiagaraScriptSource>(Handle.GetInstance()->GraphSource), ENiagaraScriptUsage::EmitterSpawnScript);
 			EmitterPtr->EmitterUniqueName = Handle.GetInstance()->GetUniqueEmitterName();
+			EmitterPtr->SourceName = BasePtr->SourceName;
 			BasePtr->EmitterData.Add(EmitterPtr);
 		}
 
@@ -688,6 +693,30 @@ FNiagaraCompileResults FHlslNiagaraCompiler::CompileScript(const FNiagaraCompile
 		Input.EntryPointName = TEXT("SimulateMain");
 		Input.Environment.SetDefine(TEXT("VM_SIMULATION"), 1);
 		Input.Environment.IncludeVirtualPathToContentsMap.Add(TEXT("/Engine/Generated/NiagaraEmitterInstance.ush"), TranslatedHLSL);
+		Input.bGenerateDirectCompileFile = false;
+		Input.DumpDebugInfoRootPath = GShaderCompilingManager->GetAbsoluteShaderDebugInfoDirectory() / TEXT("VM");
+		FString UsageIdStr = !InOptions.TargetUsageId.IsValid() ? TEXT("") : (TEXT("_") + InOptions.TargetUsageId.ToString());
+		Input.DebugGroupName = InCompileRequest->SourceName / InCompileRequest->EmitterUniqueName / InCompileRequest->ENiagaraScriptUsageEnum->GetNameStringByValue((int64)InOptions.TargetUsage) + UsageIdStr;
+		Input.DumpDebugInfoPath = Input.DumpDebugInfoRootPath / Input.DebugGroupName;
+
+		if (GShaderCompilingManager->GetDumpShaderDebugInfo())
+		{
+			// Sanitize the name to be used as a path
+			// List mostly comes from set of characters not allowed by windows in a path.  Just try to rename a file and type one of these for the list.
+			Input.DumpDebugInfoPath.ReplaceInline(TEXT("<"), TEXT("("));
+			Input.DumpDebugInfoPath.ReplaceInline(TEXT(">"), TEXT(")"));
+			Input.DumpDebugInfoPath.ReplaceInline(TEXT("::"), TEXT("=="));
+			Input.DumpDebugInfoPath.ReplaceInline(TEXT("|"), TEXT("_"));
+			Input.DumpDebugInfoPath.ReplaceInline(TEXT("*"), TEXT("-"));
+			Input.DumpDebugInfoPath.ReplaceInline(TEXT("?"), TEXT("!"));
+			Input.DumpDebugInfoPath.ReplaceInline(TEXT("\""), TEXT("\'"));
+
+			if (!IFileManager::Get().DirectoryExists(*Input.DumpDebugInfoPath))
+			{
+				verifyf(IFileManager::Get().MakeDirectory(*Input.DumpDebugInfoPath, true), TEXT("Failed to create directory for shader debug info '%s'"), *Input.DumpDebugInfoPath);
+			}
+		}
+
 		FShaderCompilerOutput Output;
 
 		FVectorVMCompilationOutput CompilationOutput;
@@ -707,7 +736,7 @@ FNiagaraCompileResults FHlslNiagaraCompiler::CompileScript(const FNiagaraCompile
 		{
 			SCOPE_CYCLE_COUNTER(STAT_NiagaraEditor_HlslCompiler_CompileShader_VectorVM);
 			static FCriticalSection CritSec;
-
+			
 			CritSec.Lock();
 			double StartTime = FPlatformTime::Seconds();
 			CompileResults.bVMSucceeded = CompileShader_VectorVM(Input, Output, FString(FPlatformProcess::ShaderDir()), 0, CompilationOutput);
diff --git a/Engine/Plugins/FX/Niagara/Source/NiagaraEditor/Private/NiagaraHlslTranslator.h b/Engine/Plugins/FX/Niagara/Source/NiagaraEditor/Private/NiagaraHlslTranslator.h
--- a/Engine/Plugins/FX/Niagara/Source/NiagaraEditor/Private/NiagaraHlslTranslator.h
+++ b/Engine/Plugins/FX/Niagara/Source/NiagaraEditor/Private/NiagaraHlslTranslator.h
@@ -115,8 +115,10 @@ public:
 	FString EmitterUniqueName;
 	TArray<TSharedPtr<FNiagaraCompileRequestData, ESPMode::ThreadSafe>> EmitterData;
 	UNiagaraScriptSource* Source;
+	FString SourceName;
 
 	UEnum* ENiagaraScriptCompileStatusEnum;
+	UEnum* ENiagaraScriptUsageEnum;
 
 	struct FunctionData
 	{
diff --git a/Engine/Source/Developer/ShaderFormatVectorVM/Private/VectorVMShaderCompiler.cpp b/Engine/Source/Developer/ShaderFormatVectorVM/Private/VectorVMShaderCompiler.cpp
--- a/Engine/Source/Developer/ShaderFormatVectorVM/Private/VectorVMShaderCompiler.cpp
+++ b/Engine/Source/Developer/ShaderFormatVectorVM/Private/VectorVMShaderCompiler.cpp
@@ -75,28 +75,27 @@ bool CompileShader_VectorVM(const FShaderCompilerInput& Input, FShaderCompilerOu
 
 	const EHlslShaderFrequency Frequency = HSF_VertexShader;
 
-	//TODO: Do this later when we implement the rest of the shader plumbing stuff.
-// 		// Write out the preprocessed file and a batch file to compile it if requested (DumpDebugInfoPath is valid)
-// 		if (bDumpDebugInfo)
-// 		{
-// 			FArchive* FileWriter = IFileManager::Get().CreateFileWriter(*(Input.DumpDebugInfoPath / Input.GetSourceFilename()));
-// 			if (FileWriter)
-// 			{
-// 				auto AnsiSourceFile = StringCast<ANSICHAR>(*PreprocessedShader);
-// 				FileWriter->Serialize((ANSICHAR*)AnsiSourceFile.Get(), AnsiSourceFile.Length());
-//				{
-//					FString Line = CrossCompiler::CreateResourceTableFromEnvironment(Input.Environment);
-//					FileWriter->Serialize(TCHAR_TO_ANSI(*Line), Line.Len());
-//				}
-// 				FileWriter->Close();
-// 				delete FileWriter;
-// 			}
-// 
-// 			if (Input.bGenerateDirectCompileFile)
-// 			{
-// 				FFileHelper::SaveStringToFile(CreateShaderCompilerWorkerDirectCommandLine(Input), *(Input.DumpDebugInfoPath / TEXT("DirectCompile.txt")));
-// 			}
-// 		}
+ 		// Write out the preprocessed file and a batch file to compile it if requested (DumpDebugInfoPath is valid)
+	if (bDumpDebugInfo)
+	{
+		FArchive* FileWriter = IFileManager::Get().CreateFileWriter(*(Input.DumpDebugInfoPath / Input.GetSourceFilename()));
+		if (FileWriter)
+		{
+			auto AnsiSourceFile = StringCast<ANSICHAR>(*PreprocessedShader);
+			FileWriter->Serialize((ANSICHAR*)AnsiSourceFile.Get(), AnsiSourceFile.Length());
+			{
+				FString Line = CrossCompiler::CreateResourceTableFromEnvironment(Input.Environment);
+				FileWriter->Serialize(TCHAR_TO_ANSI(*Line), Line.Len());
+			}
+			FileWriter->Close();
+			delete FileWriter;
+		}
+
+		if (Input.bGenerateDirectCompileFile)
+		{
+			FFileHelper::SaveStringToFile(CreateShaderCompilerWorkerDirectCommandLine(Input), *(Input.DumpDebugInfoPath / TEXT("DirectCompile.txt")));
+		}
+	}
 
 	//Is stuff like this needed? What others?
 	uint32 CCFlags = HLSLCC_NoPreprocess;
diff --git a/Engine/Source/Runtime/Engine/Private/ShaderCompiler/ShaderCompiler.cpp b/Engine/Source/Runtime/Engine/Private/ShaderCompiler/ShaderCompiler.cpp
--- a/Engine/Source/Runtime/Engine/Private/ShaderCompiler/ShaderCompiler.cpp
+++ b/Engine/Source/Runtime/Engine/Private/ShaderCompiler/ShaderCompiler.cpp
@@ -1557,6 +1557,11 @@ FShaderCompilingManager::FShaderCompilingManager() :
 	Thread->StartThread();
 }
 
+bool FShaderCompilingManager::GetDumpShaderDebugInfo() const
+{
+	return GDumpShaderDebugInfo != 0;
+}
+
 void FShaderCompilingManager::AddJobs(TArray<FShaderCommonCompileJob*>& NewJobs, bool bApplyCompletedShaderMapForRendering, bool bOptimizeForLowLatency, bool bRecreateComponentRenderStateOnCompletion)
 {
 	check(!FPlatformProperties::RequiresCookedData());
diff --git a/Engine/Source/Runtime/Engine/Public/ShaderCompiler.h b/Engine/Source/Runtime/Engine/Public/ShaderCompiler.h
--- a/Engine/Source/Runtime/Engine/Public/ShaderCompiler.h
+++ b/Engine/Source/Runtime/Engine/Public/ShaderCompiler.h
@@ -532,6 +532,8 @@ public:
 		NumExternalJobs = NumJobs;
 	}
 
+	ENGINE_API bool GetDumpShaderDebugInfo() const;
+
 	const FString& GetAbsoluteShaderDebugInfoDirectory() const
 	{
 		return AbsoluteShaderDebugInfoDirectory;
