Niagara - Make the lifetime of stack entries well defined so that we can safely remove delegate bindings and clear out pointers.

diff --git a/Engine/Plugins/FX/Niagara/Source/NiagaraEditor/Private/ViewModels/Stack/NiagaraStackEmitterSpawnScriptItemGroup.cpp b/Engine/Plugins/FX/Niagara/Source/NiagaraEditor/Private/ViewModels/Stack/NiagaraStackEmitterSpawnScriptItemGroup.cpp
--- a/Engine/Plugins/FX/Niagara/Source/NiagaraEditor/Private/ViewModels/Stack/NiagaraStackEmitterSpawnScriptItemGroup.cpp
+++ b/Engine/Plugins/FX/Niagara/Source/NiagaraEditor/Private/ViewModels/Stack/NiagaraStackEmitterSpawnScriptItemGroup.cpp
@@ -20,6 +20,15 @@ void UNiagaraStackEmitterPropertiesItem::Initialize(FRequiredEntryData InRequire
 	Emitter->OnPropertiesChanged().AddUObject(this, &UNiagaraStackEmitterPropertiesItem::EmitterPropertiesChanged);
 }
 
+void UNiagaraStackEmitterPropertiesItem::FinalizeInternal()
+{
+	if (Emitter.IsValid())
+	{
+		Emitter->OnPropertiesChanged().RemoveAll(this);
+	}
+	Super::FinalizeInternal();
+}
+
 FText UNiagaraStackEmitterPropertiesItem::GetDisplayName() const
 {
 	return LOCTEXT("EmitterPropertiesDisplayName", "Emitter Properties");
@@ -53,15 +62,6 @@ void UNiagaraStackEmitterPropertiesItem::ResetToBase()
 	}
 }
 
-void UNiagaraStackEmitterPropertiesItem::BeginDestroy()
-{
-	if (Emitter.IsValid())
-	{
-		Emitter->OnPropertiesChanged().RemoveAll(this);
-	}
-	Super::BeginDestroy();
-}
-
 void UNiagaraStackEmitterPropertiesItem::RefreshChildrenInternal(const TArray<UNiagaraStackEntry*>& CurrentChildren, TArray<UNiagaraStackEntry*>& NewChildren, TArray<FStackIssue>& NewIssues)
 {
 	if (EmitterObject == nullptr)
diff --git a/Engine/Plugins/FX/Niagara/Source/NiagaraEditor/Private/ViewModels/Stack/NiagaraStackEntry.cpp b/Engine/Plugins/FX/Niagara/Source/NiagaraEditor/Private/ViewModels/Stack/NiagaraStackEntry.cpp
--- a/Engine/Plugins/FX/Niagara/Source/NiagaraEditor/Private/ViewModels/Stack/NiagaraStackEntry.cpp
+++ b/Engine/Plugins/FX/Niagara/Source/NiagaraEditor/Private/ViewModels/Stack/NiagaraStackEntry.cpp
@@ -111,11 +111,13 @@ const TArray<UNiagaraStackEntry::FStackIssueFix>& UNiagaraStackEntry::FStackIssu
 
 UNiagaraStackEntry::UNiagaraStackEntry()
 	: IndentLevel(0)
+	, bIsFinalized(false)
 {
 }
 
 void UNiagaraStackEntry::Initialize(FRequiredEntryData InRequiredEntryData, FString InStackEditorDataKey)
 {
+	checkf(bIsFinalized == false, TEXT("Can not initialize an entry after it has been finalized."));
 	SystemViewModel = InRequiredEntryData.SystemViewModel;
 	EmitterViewModel = InRequiredEntryData.EmitterViewModel;
 	ExecutionCategoryName = InRequiredEntryData.ExecutionCategoryName;
@@ -124,9 +126,26 @@ void UNiagaraStackEntry::Initialize(FRequiredEntryData InRequiredEntryData, FStr
 	StackEditorDataKey = InStackEditorDataKey;
 }
 
-bool UNiagaraStackEntry::IsValid() const
+void UNiagaraStackEntry::Finalize()
 {
-	return SystemViewModel.IsValid();
+	FinalizeInternal();
+	checkf(bIsFinalized, TEXT("Parent FinalizeInternal not called from overriden FinalizeInternal"));
+
+	SystemViewModel.Reset();
+	EmitterViewModel.Reset();
+	StackEditorData = nullptr;
+
+	for (UNiagaraStackEntry* Child : Children)
+	{
+		Child->Finalize();
+	}
+	Children.Empty();
+
+	for (UNiagaraStackEntry* ErrorChild : ErrorChildren)
+	{
+		ErrorChild->Finalize();
+	}
+	ErrorChildren.Empty();
 }
 
 FText UNiagaraStackEntry::GetDisplayName() const
@@ -370,8 +389,14 @@ void UNiagaraStackEntry::ChlildStructureChangedInternal()
 {
 }
 
+void UNiagaraStackEntry::FinalizeInternal()
+{
+	bIsFinalized = true;
+}
+
 void UNiagaraStackEntry::RefreshChildren()
 {
+	checkf(bIsFinalized == false, TEXT("Can not refresh children on an entry after it has been finalized."));
 	checkf(SystemViewModel.IsValid() && EmitterViewModel.IsValid(), TEXT("Base stack entry not initialized."));
 
 	for (UNiagaraStackEntry* Child : Children)
@@ -393,6 +418,17 @@ void UNiagaraStackEntry::RefreshChildren()
 	TArray<UNiagaraStackEntry*> NewChildren;
 	TArray<FStackIssue> NewStackIssues;
 	RefreshChildrenInternal(Children, NewChildren, NewStackIssues);
+
+	// If any of the current children were not moved to the new children collection than finalize them since
+	// they weren't reused.
+	for (UNiagaraStackEntry* Child : Children)
+	{
+		if (NewChildren.Contains(Child) == false)
+		{
+			Child->Finalize();
+		}
+	}
+
 	Children.Empty();
 	Children.Append(NewChildren);
 
@@ -448,6 +484,17 @@ void UNiagaraStackEntry::RefreshStackErrorChildren()
 		}
 		NewErrorChildren.Add(ErrorEntry);
 	}
+
+	// If any of the current error children were not moved to the new error children collection than finalize them since
+	// they weren't reused.
+	for (UNiagaraStackEntry* ErrorChild : ErrorChildren)
+	{
+		if(NewErrorChildren.Contains(ErrorChild) == false)
+		{
+			ErrorChild->Finalize();
+		}
+	}
+
 	ErrorChildren.Empty();
 	ErrorChildren.Append(NewErrorChildren);
 }
@@ -457,6 +504,12 @@ void UNiagaraStackEntry::IssueModified()
 	RefreshChildren();
 }
 
+void UNiagaraStackEntry::BeginDestroy()
+{
+	ensureMsgf(HasAnyFlags(RF_ClassDefaultObject) || bIsFinalized, TEXT("Stack entry being destroyed but it was not finalized."));
+	Super::BeginDestroy();
+}
+
 void UNiagaraStackEntry::RefreshChildrenInternal(const TArray<UNiagaraStackEntry*>& CurrentChildren, TArray<UNiagaraStackEntry*>& NewChildren, TArray<FStackIssue>& NewIssues)
 {
 }
diff --git a/Engine/Plugins/FX/Niagara/Source/NiagaraEditor/Private/ViewModels/Stack/NiagaraStackEventScriptItemGroup.cpp b/Engine/Plugins/FX/Niagara/Source/NiagaraEditor/Private/ViewModels/Stack/NiagaraStackEventScriptItemGroup.cpp
--- a/Engine/Plugins/FX/Niagara/Source/NiagaraEditor/Private/ViewModels/Stack/NiagaraStackEventScriptItemGroup.cpp
+++ b/Engine/Plugins/FX/Niagara/Source/NiagaraEditor/Private/ViewModels/Stack/NiagaraStackEventScriptItemGroup.cpp
@@ -44,6 +44,15 @@ void UNiagaraStackEventHandlerPropertiesItem::Initialize(FRequiredEntryData InRe
 	}
 }
 
+void UNiagaraStackEventHandlerPropertiesItem::FinalizeInternal()
+{
+	if (Emitter.IsValid())
+	{
+		Emitter->OnPropertiesChanged().RemoveAll(this);
+	}
+	Super::FinalizeInternal();
+}
+
 FText UNiagaraStackEventHandlerPropertiesItem::GetDisplayName() const
 {
 	return LOCTEXT("EventHandlerPropertiesDisplayName", "Event Handler Properties");
@@ -85,15 +94,6 @@ void UNiagaraStackEventHandlerPropertiesItem::ResetToBase()
 	}
 }
 
-void UNiagaraStackEventHandlerPropertiesItem::BeginDestroy()
-{
-	if (Emitter.IsValid())
-	{
-		Emitter->OnPropertiesChanged().RemoveAll(this);
-	}
-	Super::BeginDestroy();
-}
-
 void UNiagaraStackEventHandlerPropertiesItem::RefreshChildrenInternal(const TArray<UNiagaraStackEntry*>& CurrentChildren, TArray<UNiagaraStackEntry*>& NewChildren, TArray<FStackIssue>& NewIssues)
 {
 	if (EmitterObject == nullptr)
diff --git a/Engine/Plugins/FX/Niagara/Source/NiagaraEditor/Private/ViewModels/Stack/NiagaraStackFunctionInput.cpp b/Engine/Plugins/FX/Niagara/Source/NiagaraEditor/Private/ViewModels/Stack/NiagaraStackFunctionInput.cpp
--- a/Engine/Plugins/FX/Niagara/Source/NiagaraEditor/Private/ViewModels/Stack/NiagaraStackFunctionInput.cpp
+++ b/Engine/Plugins/FX/Niagara/Source/NiagaraEditor/Private/ViewModels/Stack/NiagaraStackFunctionInput.cpp
@@ -53,22 +53,6 @@ UNiagaraStackFunctionInput::UNiagaraStackFunctionInput()
 {
 }
 
-void UNiagaraStackFunctionInput::BeginDestroy()
-{
-	Super::BeginDestroy();
-	if (OwningFunctionCallNode.IsValid())
-	{
-		OwningFunctionCallNode->GetGraph()->RemoveOnGraphChangedHandler(GraphChangedHandle);
-		OwningFunctionCallNode->GetNiagaraGraph()->RemoveOnGraphNeedsRecompileHandler(OnRecompileHandle);
-
-		if (SourceScript.IsValid())
-		{
-			SourceScript->RapidIterationParameters.RemoveOnChangedHandler(RapidIterationParametersChangedHandle);
-			SourceScript->GetSource()->OnChanged().RemoveAll(this);
-		}
-	}
-}
-
 void UNiagaraStackFunctionInput::AddReferencedObjects(UObject* InThis, FReferenceCollector& Collector)
 {
 	UNiagaraStackFunctionInput* This = Cast<UNiagaraStackFunctionInput>(InThis);
@@ -217,6 +201,22 @@ void UNiagaraStackFunctionInput::Initialize(
 	VisibleCondition.Initialize(SourceScript.Get(), AffectedScriptsNotWeak, GetEmitterViewModel()->GetEmitter()->GetUniqueEmitterName(), OwningFunctionCallNode.Get());
 }
 
+void UNiagaraStackFunctionInput::FinalizeInternal()
+{
+	if (OwningFunctionCallNode.IsValid())
+	{
+		OwningFunctionCallNode->GetGraph()->RemoveOnGraphChangedHandler(GraphChangedHandle);
+		OwningFunctionCallNode->GetNiagaraGraph()->RemoveOnGraphNeedsRecompileHandler(OnRecompileHandle);
+
+		if (SourceScript.IsValid())
+		{
+			SourceScript->RapidIterationParameters.RemoveOnChangedHandler(RapidIterationParametersChangedHandle);
+			SourceScript->GetSource()->OnChanged().RemoveAll(this);
+		}
+	}
+	Super::FinalizeInternal();
+}
+
 const UNiagaraNodeFunctionCall& UNiagaraStackFunctionInput::GetInputFunctionCallNode() const
 {
 	return *OwningFunctionCallNode.Get();
diff --git a/Engine/Plugins/FX/Niagara/Source/NiagaraEditor/Private/ViewModels/Stack/NiagaraStackFunctionInputCollection.cpp b/Engine/Plugins/FX/Niagara/Source/NiagaraEditor/Private/ViewModels/Stack/NiagaraStackFunctionInputCollection.cpp
--- a/Engine/Plugins/FX/Niagara/Source/NiagaraEditor/Private/ViewModels/Stack/NiagaraStackFunctionInputCollection.cpp
+++ b/Engine/Plugins/FX/Niagara/Source/NiagaraEditor/Private/ViewModels/Stack/NiagaraStackFunctionInputCollection.cpp
@@ -49,6 +49,12 @@ void UNiagaraStackFunctionInputCollection::Initialize(
 	InputFunctionCallNode->OnInputsChanged().AddUObject(this, &UNiagaraStackFunctionInputCollection::OnFunctionInputsChanged);
 }
 
+void UNiagaraStackFunctionInputCollection::FinalizeInternal()
+{
+	InputFunctionCallNode->OnInputsChanged().RemoveAll(this);
+	Super::FinalizeInternal();
+}
+
 FText UNiagaraStackFunctionInputCollection::GetDisplayName() const
 {
 	return LOCTEXT("InputCollectionDisplayName", "Inputs");
@@ -151,22 +157,27 @@ void UNiagaraStackFunctionInputCollection::RefreshChildrenInternal(const TArray<
 	// Populate the category children
 	for (const FInputData InputData : InputDataCollection)
 	{
+		// Try to find an existing category in the already processed children.
 		UNiagaraStackInputCategory* InputCategory = FindCurrentChildOfTypeByPredicate<UNiagaraStackInputCategory>(NewChildren,
 			[&](UNiagaraStackInputCategory* CurrentCategory) { return CurrentCategory->GetCategoryName().CompareTo(InputData.Category) == 0; });
 
 		if (InputCategory == nullptr)
 		{
+			// If we haven't added any children to this category yet see if there is one that can be reused from the current children.
 			InputCategory = FindCurrentChildOfTypeByPredicate<UNiagaraStackInputCategory>(CurrentChildren,
 				[&](UNiagaraStackInputCategory* CurrentCategory) { return CurrentCategory->GetCategoryName().CompareTo(InputData.Category) == 0; });
 			if (InputCategory == nullptr)
 			{
+				// If we don't have a current child for this category make a new one.
 				InputCategory = NewObject<UNiagaraStackInputCategory>(this);
 				InputCategory->Initialize(CreateDefaultChildRequiredData(), *ModuleNode, *InputFunctionCallNode, InputData.Category, GetOwnerStackItemEditorDataKey());
 			}
 			else
 			{
+				// We found a category to reuse, but we need to reset the inputs before we can start adding the current set of inputs.
 				InputCategory->ResetInputs();
 			}
+
 			if (InputData.Category.CompareTo(UncategorizedName) == 0)
 			{
 				InputCategory->SetShouldShowInStack(false);
@@ -261,10 +272,7 @@ void UNiagaraStackFunctionInputCollection::RefreshIssues(TArray<FName> Duplicate
 
 void UNiagaraStackFunctionInputCollection::OnFunctionInputsChanged()
 {
-	if (IsValid())
-	{
-		RefreshChildren();
-	}
+	RefreshChildren();
 }
 
 #undef LOCTEXT_NAMESPACE
diff --git a/Engine/Plugins/FX/Niagara/Source/NiagaraEditor/Private/ViewModels/Stack/NiagaraStackGraphUtilities.cpp b/Engine/Plugins/FX/Niagara/Source/NiagaraEditor/Private/ViewModels/Stack/NiagaraStackGraphUtilities.cpp
--- a/Engine/Plugins/FX/Niagara/Source/NiagaraEditor/Private/ViewModels/Stack/NiagaraStackGraphUtilities.cpp
+++ b/Engine/Plugins/FX/Niagara/Source/NiagaraEditor/Private/ViewModels/Stack/NiagaraStackGraphUtilities.cpp
@@ -454,6 +454,7 @@ void FNiagaraStackGraphUtilities::InitializeStackFunctionInputs(TSharedRef<FNiag
 		}
 	}
 
+	FunctionInputCollection->Finalize();
 	SystemViewModel->NotifyDataObjectChanged(nullptr);
 }
 
diff --git a/Engine/Plugins/FX/Niagara/Source/NiagaraEditor/Private/ViewModels/Stack/NiagaraStackObject.cpp b/Engine/Plugins/FX/Niagara/Source/NiagaraEditor/Private/ViewModels/Stack/NiagaraStackObject.cpp
--- a/Engine/Plugins/FX/Niagara/Source/NiagaraEditor/Private/ViewModels/Stack/NiagaraStackObject.cpp
+++ b/Engine/Plugins/FX/Niagara/Source/NiagaraEditor/Private/ViewModels/Stack/NiagaraStackObject.cpp
@@ -60,6 +60,15 @@ bool UNiagaraStackObject::GetShouldShowInStack() const
 	return false;
 }
 
+void UNiagaraStackObject::FinalizeInternal()
+{
+	if (PropertyRowGenerator.IsValid())
+	{
+		PropertyRowGenerator->OnRowsRefreshed().RemoveAll(this);
+	}
+	Super::FinalizeInternal();
+}
+
 void UNiagaraStackObject::RefreshChildrenInternal(const TArray<UNiagaraStackEntry*>& CurrentChildren, TArray<UNiagaraStackEntry*>& NewChildren, TArray<FStackIssue>& NewIssues)
 {
 	if (PropertyRowGenerator.IsValid() == false)
@@ -121,8 +130,5 @@ void UNiagaraStackObject::RefreshChildrenInternal(const TArray<UNiagaraStackEntr
 
 void UNiagaraStackObject::PropertyRowsRefreshed()
 {
-	if (IsValid())
-	{
-		RefreshChildren();
-	}
+	RefreshChildren();
 }
diff --git a/Engine/Plugins/FX/Niagara/Source/NiagaraEditor/Private/ViewModels/Stack/NiagaraStackParameterStoreEntry.cpp b/Engine/Plugins/FX/Niagara/Source/NiagaraEditor/Private/ViewModels/Stack/NiagaraStackParameterStoreEntry.cpp
--- a/Engine/Plugins/FX/Niagara/Source/NiagaraEditor/Private/ViewModels/Stack/NiagaraStackParameterStoreEntry.cpp
+++ b/Engine/Plugins/FX/Niagara/Source/NiagaraEditor/Private/ViewModels/Stack/NiagaraStackParameterStoreEntry.cpp
@@ -31,11 +31,6 @@ UNiagaraStackParameterStoreEntry::UNiagaraStackParameterStoreEntry()
 {
 }
 
-void UNiagaraStackParameterStoreEntry::BeginDestroy()
-{
-	Super::BeginDestroy();
-}
-
 void UNiagaraStackParameterStoreEntry::Initialize(
 	FRequiredEntryData InRequiredEntryData,
 	UObject* InOwner,
diff --git a/Engine/Plugins/FX/Niagara/Source/NiagaraEditor/Private/ViewModels/Stack/NiagaraStackParameterStoreGroup.cpp b/Engine/Plugins/FX/Niagara/Source/NiagaraEditor/Private/ViewModels/Stack/NiagaraStackParameterStoreGroup.cpp
--- a/Engine/Plugins/FX/Niagara/Source/NiagaraEditor/Private/ViewModels/Stack/NiagaraStackParameterStoreGroup.cpp
+++ b/Engine/Plugins/FX/Niagara/Source/NiagaraEditor/Private/ViewModels/Stack/NiagaraStackParameterStoreGroup.cpp
@@ -153,18 +153,18 @@ void UNiagaraStackParameterStoreItem::Initialize(
 		FNiagaraParameterStore::FOnChanged::FDelegate::CreateUObject(this, &UNiagaraStackParameterStoreItem::ParameterStoreChanged));
 }
 
-FText UNiagaraStackParameterStoreItem::GetDisplayName() const
-{
-	return LOCTEXT("ParameterItemDisplayName", "Parameters");
-}
-
-void UNiagaraStackParameterStoreItem::BeginDestroy()
+void UNiagaraStackParameterStoreItem::FinalizeInternal()
 {
 	if (Owner != nullptr)
 	{
 		ParameterStore->RemoveOnChangedHandler(ParameterStoreChangedHandle);
 	}
-	Super::BeginDestroy();
+	Super::FinalizeInternal();
+}
+
+FText UNiagaraStackParameterStoreItem::GetDisplayName() const
+{
+	return LOCTEXT("ParameterItemDisplayName", "Parameters");
 }
 
 void UNiagaraStackParameterStoreItem::RefreshChildrenInternal(const TArray<UNiagaraStackEntry*>& CurrentChildren, TArray<UNiagaraStackEntry*>& NewChildren, TArray<FStackIssue>& NewIssues)
@@ -193,10 +193,7 @@ void UNiagaraStackParameterStoreItem::RefreshChildrenInternal(const TArray<UNiag
 
 void UNiagaraStackParameterStoreItem::ParameterStoreChanged()
 {
-	if (IsValid())
-	{
-		RefreshChildren();
-	}
+	RefreshChildren();
 }
 
 #undef LOCTEXT_NAMESPACE
diff --git a/Engine/Plugins/FX/Niagara/Source/NiagaraEditor/Private/ViewModels/Stack/NiagaraStackRendererItem.cpp b/Engine/Plugins/FX/Niagara/Source/NiagaraEditor/Private/ViewModels/Stack/NiagaraStackRendererItem.cpp
--- a/Engine/Plugins/FX/Niagara/Source/NiagaraEditor/Private/ViewModels/Stack/NiagaraStackRendererItem.cpp
+++ b/Engine/Plugins/FX/Niagara/Source/NiagaraEditor/Private/ViewModels/Stack/NiagaraStackRendererItem.cpp
@@ -52,6 +52,15 @@ void UNiagaraStackRendererItem::Initialize(FRequiredEntryData InRequiredEntryDat
 	}
 }
 
+void UNiagaraStackRendererItem::FinalizeInternal()
+{
+	if (RendererProperties.IsValid())
+	{
+		RendererProperties->OnChanged().RemoveAll(this);
+	}
+	Super::FinalizeInternal();
+}
+
 TArray<FNiagaraVariable> UNiagaraStackRendererItem::GetMissingVariables(UNiagaraRendererProperties* RendererProperties, UNiagaraEmitter* Emitter)
 {
 	TArray<FNiagaraVariable> MissingAttributes;
@@ -207,15 +216,6 @@ void UNiagaraStackRendererItem::SetIsEnabled(bool bInIsEnabled)
 	OnDataObjectModified().Broadcast(RendererProperties.Get());
 }
 
-void UNiagaraStackRendererItem::BeginDestroy()
-{
-	if (RendererProperties.IsValid())
-	{
-		RendererProperties->OnChanged().RemoveAll(this);
-	}
-	Super::BeginDestroy();
-}
-
 void UNiagaraStackRendererItem::RefreshChildrenInternal(const TArray<UNiagaraStackEntry*>& CurrentChildren, TArray<UNiagaraStackEntry*>& NewChildren, TArray<FStackIssue>& NewIssues)
 {
 	if (RendererObject == nullptr)
diff --git a/Engine/Plugins/FX/Niagara/Source/NiagaraEditor/Private/ViewModels/Stack/NiagaraStackViewModel.cpp b/Engine/Plugins/FX/Niagara/Source/NiagaraEditor/Private/ViewModels/Stack/NiagaraStackViewModel.cpp
--- a/Engine/Plugins/FX/Niagara/Source/NiagaraEditor/Private/ViewModels/Stack/NiagaraStackViewModel.cpp
+++ b/Engine/Plugins/FX/Niagara/Source/NiagaraEditor/Private/ViewModels/Stack/NiagaraStackViewModel.cpp
@@ -37,6 +37,7 @@ void UNiagaraStackViewModel::Initialize(TSharedPtr<FNiagaraSystemViewModel> InSy
 		RootEntry->OnDataObjectModified().RemoveAll(this);
 		RootEntry->OnRequestFullRefresh().RemoveAll(this);
 		RootEntries.Empty();
+		RootEntry->Finalize();
 		RootEntry = nullptr;
 		GEditor->UnregisterForUndo(this);
 	}
@@ -52,7 +53,6 @@ void UNiagaraStackViewModel::Initialize(TSharedPtr<FNiagaraSystemViewModel> InSy
 	}
 
 	SystemViewModel = InSystemViewModel;
-	//UE_LOG(LogNiagaraEditor, Warning, TEXT("UNiagaraStackViewModel %p refs FNiagaraEmitterViewModel %p"), this, InEmitterViewModel.Get());
 	EmitterHandleViewModel = InEmitterHandleViewModel;
 	TSharedPtr<FNiagaraEmitterViewModel> EmitterViewModel = InEmitterHandleViewModel.IsValid() ? InEmitterHandleViewModel->GetEmitterViewModel() : TSharedPtr<FNiagaraEmitterViewModel>();
 
@@ -74,25 +74,21 @@ void UNiagaraStackViewModel::Initialize(TSharedPtr<FNiagaraSystemViewModel> InSy
 		RootEntry->OnRequestFullRefresh().AddUObject(this, &UNiagaraStackViewModel::EntryRequestFullRefresh);
 		RootEntries.Add(RootEntry);
 	}
+
 	CurrentFocusedSearchMatchIndex = -1;
 	StructureChangedDelegate.Broadcast();
 	bRestartSearch = false;
 }
 
+void UNiagaraStackViewModel::Finalize()
+{
+	Initialize(nullptr, nullptr);
+}
+
 void UNiagaraStackViewModel::BeginDestroy()
 {
+	checkf(HasAnyFlags(RF_ClassDefaultObject) || (SystemViewModel.IsValid() == false && EmitterHandleViewModel.IsValid() == false), TEXT("Stack view model not finalized."));
 	Super::BeginDestroy();
-	if (EmitterHandleViewModel.IsValid() && EmitterHandleViewModel->GetEmitterViewModel().IsValid())
-	{
-		GEditor->UnregisterForUndo(this);
-		EmitterHandleViewModel->GetEmitterViewModel()->OnScriptCompiled().RemoveAll(this);
-
-		//UE_LOG(LogNiagaraEditor, Log, TEXT("UNiagaraStackViewModel::BeginDestroy"));
-		EmitterHandleViewModel->GetEmitterViewModel().Reset();
-		EmitterHandleViewModel.Reset();
-		SystemViewModel.Reset();
-	} 
-	RootEntry = nullptr;
 }
 
 void UNiagaraStackViewModel::Tick()
diff --git a/Engine/Plugins/FX/Niagara/Source/NiagaraEditor/Private/Widgets/SNiagaraSelectedEmitterHandles.cpp b/Engine/Plugins/FX/Niagara/Source/NiagaraEditor/Private/Widgets/SNiagaraSelectedEmitterHandles.cpp
--- a/Engine/Plugins/FX/Niagara/Source/NiagaraEditor/Private/Widgets/SNiagaraSelectedEmitterHandles.cpp
+++ b/Engine/Plugins/FX/Niagara/Source/NiagaraEditor/Private/Widgets/SNiagaraSelectedEmitterHandles.cpp
@@ -59,10 +59,7 @@ SNiagaraSelectedEmitterHandles::~SNiagaraSelectedEmitterHandles()
 	SystemViewModel->OnSelectedEmitterHandlesChanged().RemoveAll(this);
 	SystemViewModel->GetOnPinnedEmittersChanged().RemoveAll(this);
 
-	for (auto Model : StackViewModels)
-	{
-		Model->Initialize(nullptr, nullptr);
-	}
+	ResetViewModels();
 }
 
 void SNiagaraSelectedEmitterHandles::AddReferencedObjects(FReferenceCollector& Collector)
@@ -72,7 +69,9 @@ void SNiagaraSelectedEmitterHandles::AddReferencedObjects(FReferenceCollector& C
 
 void SNiagaraSelectedEmitterHandles::RefreshEmitterWidgets()
 {
-	ClearViewModels();
+	ResetWidgets();
+	ResetViewModels();
+
 	TArray<TSharedPtr<FNiagaraEmitterHandleViewModel>> EmitterHandlesToDisplay;
 	EmitterHandlesToDisplay.Append(SystemViewModel->GetPinnedEmitterHandles());
 	TArray<TSharedRef<FNiagaraEmitterHandleViewModel>> SelectedEmitterHandles;
@@ -97,13 +96,20 @@ void SNiagaraSelectedEmitterHandles::RefreshEmitterWidgets()
 	EmitterHandlesToDisplay.Empty();
 }
 
-void SNiagaraSelectedEmitterHandles::ClearViewModels()
+void SNiagaraSelectedEmitterHandles::ResetWidgets()
 {
 	while (EmitterSplitter->GetChildren()->Num() > 0)
 	{
 		int i = EmitterSplitter->GetChildren()->Num() - 1;
 		EmitterSplitter->RemoveAt(i);
-		StackViewModels[i]->Initialize(nullptr, nullptr);
+	}
+}
+
+void SNiagaraSelectedEmitterHandles::ResetViewModels()
+{
+	for (UNiagaraStackViewModel* StackViewModel : StackViewModels)
+	{
+		StackViewModel->Finalize();
 	}
 	StackViewModels.Empty();
 }
diff --git a/Engine/Plugins/FX/Niagara/Source/NiagaraEditor/Private/Widgets/SNiagaraSelectedEmitterHandles.h b/Engine/Plugins/FX/Niagara/Source/NiagaraEditor/Private/Widgets/SNiagaraSelectedEmitterHandles.h
--- a/Engine/Plugins/FX/Niagara/Source/NiagaraEditor/Private/Widgets/SNiagaraSelectedEmitterHandles.h
+++ b/Engine/Plugins/FX/Niagara/Source/NiagaraEditor/Private/Widgets/SNiagaraSelectedEmitterHandles.h
@@ -41,9 +41,12 @@ private:
 
 	void OnEmitterPinnedChanged();
 
-	void ClearViewModels();
+	void ResetWidgets();
+
+	void ResetViewModels();
 
 	void CollapseToHeaders();
+
 private:
 	TSharedPtr<FNiagaraSystemViewModel> SystemViewModel;
 
diff --git a/Engine/Plugins/FX/Niagara/Source/NiagaraEditor/Public/ViewModels/Stack/NiagaraStackEmitterSpawnScriptItemGroup.h b/Engine/Plugins/FX/Niagara/Source/NiagaraEditor/Public/ViewModels/Stack/NiagaraStackEmitterSpawnScriptItemGroup.h
--- a/Engine/Plugins/FX/Niagara/Source/NiagaraEditor/Public/ViewModels/Stack/NiagaraStackEmitterSpawnScriptItemGroup.h
+++ b/Engine/Plugins/FX/Niagara/Source/NiagaraEditor/Public/ViewModels/Stack/NiagaraStackEmitterSpawnScriptItemGroup.h
@@ -26,7 +26,7 @@ public:
 	void ResetToBase();
 
 protected:
-	virtual void BeginDestroy() override;
+	virtual void FinalizeInternal() override;
 
 	virtual void RefreshChildrenInternal(const TArray<UNiagaraStackEntry*>& CurrentChildren, TArray<UNiagaraStackEntry*>& NewChildren, TArray<FStackIssue>& NewIssues) override;
 
diff --git a/Engine/Plugins/FX/Niagara/Source/NiagaraEditor/Public/ViewModels/Stack/NiagaraStackEntry.h b/Engine/Plugins/FX/Niagara/Source/NiagaraEditor/Public/ViewModels/Stack/NiagaraStackEntry.h
--- a/Engine/Plugins/FX/Niagara/Source/NiagaraEditor/Public/ViewModels/Stack/NiagaraStackEntry.h
+++ b/Engine/Plugins/FX/Niagara/Source/NiagaraEditor/Public/ViewModels/Stack/NiagaraStackEntry.h
@@ -162,7 +162,7 @@ public:
 
 	void Initialize(FRequiredEntryData InRequiredEntryData, FString InStackEditorDataKey);
 
-	bool IsValid() const;
+	void Finalize();
 
 	virtual FText GetDisplayName() const;
 
@@ -229,14 +229,6 @@ public:
 		return nullptr;
 	}
 
-	template<typename ChildType, typename PredicateType>
-	ChildType* FindChildOfTypeByPredicate(PredicateType Predicate)
-	{
-		TArray<UNiagaraStackEntry*> CurrentChildren;
-		GetUnfilteredChildren(CurrentChildren);
-		return FindCurrentChildOfTypeByPredicate<ChildType>(CurrentChildren, Predicate);
-	}
-
 	void GetSearchItems(TArray<FStackSearchItem>& SearchItems) const;
 
 	virtual UObject* GetExternalAsset() const;
@@ -252,6 +244,8 @@ public:
 	void SetOnRequestDrop(FOnRequestDrop InOnRequestCanDrop);
 
 protected:
+	virtual void BeginDestroy() override;
+
 	virtual void RefreshChildrenInternal(const TArray<UNiagaraStackEntry*>& CurrentChildren, TArray<UNiagaraStackEntry*>& NewChildren, TArray<FStackIssue>& NewIssues);
 
 	virtual void PostRefreshChildrenInternal();
@@ -272,6 +266,8 @@ protected:
 
 	virtual void ChlildStructureChangedInternal();
 
+	virtual void FinalizeInternal();
+
 private:
 	void ChildStructureChanged();
 	
@@ -320,4 +316,6 @@ private:
 	FOnRequestDrop OnRequestDropDelegate;
 	
 	TArray<FStackIssue> StackIssues;
+
+	bool bIsFinalized;
 };
\ No newline at end of file
diff --git a/Engine/Plugins/FX/Niagara/Source/NiagaraEditor/Public/ViewModels/Stack/NiagaraStackEventScriptItemGroup.h b/Engine/Plugins/FX/Niagara/Source/NiagaraEditor/Public/ViewModels/Stack/NiagaraStackEventScriptItemGroup.h
--- a/Engine/Plugins/FX/Niagara/Source/NiagaraEditor/Public/ViewModels/Stack/NiagaraStackEventScriptItemGroup.h
+++ b/Engine/Plugins/FX/Niagara/Source/NiagaraEditor/Public/ViewModels/Stack/NiagaraStackEventScriptItemGroup.h
@@ -28,7 +28,7 @@ public:
 	void ResetToBase();
 
 protected:
-	virtual void BeginDestroy() override;
+	virtual void FinalizeInternal() override;
 
 	virtual void RefreshChildrenInternal(const TArray<UNiagaraStackEntry*>& CurrentChildren, TArray<UNiagaraStackEntry*>& NewChildren, TArray<FStackIssue>& NewIssues) override;
 
diff --git a/Engine/Plugins/FX/Niagara/Source/NiagaraEditor/Public/ViewModels/Stack/NiagaraStackFunctionInput.h b/Engine/Plugins/FX/Niagara/Source/NiagaraEditor/Public/ViewModels/Stack/NiagaraStackFunctionInput.h
--- a/Engine/Plugins/FX/Niagara/Source/NiagaraEditor/Public/ViewModels/Stack/NiagaraStackFunctionInput.h
+++ b/Engine/Plugins/FX/Niagara/Source/NiagaraEditor/Public/ViewModels/Stack/NiagaraStackFunctionInput.h
@@ -48,8 +48,6 @@ public:
 public:
 	UNiagaraStackFunctionInput();
 
-	virtual void BeginDestroy() override;
-
 	static void AddReferencedObjects(UObject* InThis, FReferenceCollector& Collector);
 	void AddReferencedObjects(FReferenceCollector& Collector);
 
@@ -202,6 +200,7 @@ public:
 
 protected:
 	//~ UNiagaraStackEntry interface
+	virtual void FinalizeInternal() override;
 	virtual void RefreshChildrenInternal(const TArray<UNiagaraStackEntry*>& CurrentChildren, TArray<UNiagaraStackEntry*>& NewChildren, TArray<FStackIssue>& NewIssues) override;
 
 	FNiagaraVariable GetDefaultVariableForRapidIterationParameter() const;
diff --git a/Engine/Plugins/FX/Niagara/Source/NiagaraEditor/Public/ViewModels/Stack/NiagaraStackFunctionInputCollection.h b/Engine/Plugins/FX/Niagara/Source/NiagaraEditor/Public/ViewModels/Stack/NiagaraStackFunctionInputCollection.h
--- a/Engine/Plugins/FX/Niagara/Source/NiagaraEditor/Public/ViewModels/Stack/NiagaraStackFunctionInputCollection.h
+++ b/Engine/Plugins/FX/Niagara/Source/NiagaraEditor/Public/ViewModels/Stack/NiagaraStackFunctionInputCollection.h
@@ -36,6 +36,8 @@ public:
 	void SetShouldShowInStack(bool bInShouldShowInStack);
 
 protected:
+	virtual void FinalizeInternal() override;
+
 	virtual void RefreshChildrenInternal(const TArray<UNiagaraStackEntry*>& CurrentChildren, TArray<UNiagaraStackEntry*>& NewChildren, TArray<FStackIssue>& NewIssues) override;
 	
 private:
diff --git a/Engine/Plugins/FX/Niagara/Source/NiagaraEditor/Public/ViewModels/Stack/NiagaraStackObject.h b/Engine/Plugins/FX/Niagara/Source/NiagaraEditor/Public/ViewModels/Stack/NiagaraStackObject.h
--- a/Engine/Plugins/FX/Niagara/Source/NiagaraEditor/Public/ViewModels/Stack/NiagaraStackObject.h
+++ b/Engine/Plugins/FX/Niagara/Source/NiagaraEditor/Public/ViewModels/Stack/NiagaraStackObject.h
@@ -39,6 +39,8 @@ public:
 	virtual bool GetShouldShowInStack() const override;
 
 protected:
+	virtual void FinalizeInternal() override;
+
 	virtual void RefreshChildrenInternal(const TArray<UNiagaraStackEntry*>& CurrentChildren, TArray<UNiagaraStackEntry*>& NewChildren, TArray<FStackIssue>& NewIssues);
 
 private:
diff --git a/Engine/Plugins/FX/Niagara/Source/NiagaraEditor/Public/ViewModels/Stack/NiagaraStackParameterStoreEntry.h b/Engine/Plugins/FX/Niagara/Source/NiagaraEditor/Public/ViewModels/Stack/NiagaraStackParameterStoreEntry.h
--- a/Engine/Plugins/FX/Niagara/Source/NiagaraEditor/Public/ViewModels/Stack/NiagaraStackParameterStoreEntry.h
+++ b/Engine/Plugins/FX/Niagara/Source/NiagaraEditor/Public/ViewModels/Stack/NiagaraStackParameterStoreEntry.h
@@ -26,7 +26,6 @@ public:
 public:
 	UNiagaraStackParameterStoreEntry();
 
-	virtual void BeginDestroy() override;
 	/** 
 	 * Sets the input data for this entry.
 	 * @param InSystemViewModel The view model for the system which owns the stack containing this entry.
diff --git a/Engine/Plugins/FX/Niagara/Source/NiagaraEditor/Public/ViewModels/Stack/NiagaraStackParameterStoreGroup.h b/Engine/Plugins/FX/Niagara/Source/NiagaraEditor/Public/ViewModels/Stack/NiagaraStackParameterStoreGroup.h
--- a/Engine/Plugins/FX/Niagara/Source/NiagaraEditor/Public/ViewModels/Stack/NiagaraStackParameterStoreGroup.h
+++ b/Engine/Plugins/FX/Niagara/Source/NiagaraEditor/Public/ViewModels/Stack/NiagaraStackParameterStoreGroup.h
@@ -40,9 +40,9 @@ public:
 
 	virtual FText GetDisplayName() const override;
 
-	virtual void BeginDestroy() override;
-
 protected:
+	virtual void FinalizeInternal() override;
+
 	virtual void RefreshChildrenInternal(const TArray<UNiagaraStackEntry*>& CurrentChildren, TArray<UNiagaraStackEntry*>& NewChildren, TArray<FStackIssue>& NewIssues) override;
 
 private:
diff --git a/Engine/Plugins/FX/Niagara/Source/NiagaraEditor/Public/ViewModels/Stack/NiagaraStackRendererItem.h b/Engine/Plugins/FX/Niagara/Source/NiagaraEditor/Public/ViewModels/Stack/NiagaraStackRendererItem.h
--- a/Engine/Plugins/FX/Niagara/Source/NiagaraEditor/Public/ViewModels/Stack/NiagaraStackRendererItem.h
+++ b/Engine/Plugins/FX/Niagara/Source/NiagaraEditor/Public/ViewModels/Stack/NiagaraStackRendererItem.h
@@ -41,7 +41,8 @@ public:
 	static bool AddMissingVariable(UNiagaraEmitter* Emitter, const FNiagaraVariable& Variable);
 
 protected:
-	virtual void BeginDestroy() override;
+	virtual void FinalizeInternal() override;
+
 	virtual void RefreshChildrenInternal(const TArray<UNiagaraStackEntry*>& CurrentChildren, TArray<UNiagaraStackEntry*>& NewChildren, TArray<FStackIssue>& NewIssues) override;
 
 private:
diff --git a/Engine/Plugins/FX/Niagara/Source/NiagaraEditor/Public/ViewModels/Stack/NiagaraStackViewModel.h b/Engine/Plugins/FX/Niagara/Source/NiagaraEditor/Public/ViewModels/Stack/NiagaraStackViewModel.h
--- a/Engine/Plugins/FX/Niagara/Source/NiagaraEditor/Public/ViewModels/Stack/NiagaraStackViewModel.h
+++ b/Engine/Plugins/FX/Niagara/Source/NiagaraEditor/Public/ViewModels/Stack/NiagaraStackViewModel.h
@@ -38,6 +38,8 @@ public:
 	TSharedPtr<FNiagaraSystemViewModel> GetSystemViewModel();
 	void Initialize(TSharedPtr<FNiagaraSystemViewModel> InSystemViewModel, TSharedPtr<FNiagaraEmitterHandleViewModel> InEmitterHandleViewModel);
 
+	void Finalize();
+
 	virtual void BeginDestroy() override;
 
 	TArray<UNiagaraStackEntry*>& GetRootEntries();
