Stack issues now update fix delegates on each refresh, even if the fix GUIDs don't change, to account for other possible changes in the graph. Had to introduce unique identifiers for fixes too, now the issue entry is using the same recycle mechanism for fixes that the base stack entry uses for issues.
#tests none

diff --git a/Engine/Plugins/FX/Niagara/Source/NiagaraEditor/Private/ViewModels/Stack/NiagaraStackEntry.cpp b/Engine/Plugins/FX/Niagara/Source/NiagaraEditor/Private/ViewModels/Stack/NiagaraStackEntry.cpp
--- a/Engine/Plugins/FX/Niagara/Source/NiagaraEditor/Private/ViewModels/Stack/NiagaraStackEntry.cpp
+++ b/Engine/Plugins/FX/Niagara/Source/NiagaraEditor/Private/ViewModels/Stack/NiagaraStackEntry.cpp
@@ -22,6 +22,7 @@ UNiagaraStackEntry::FStackIssueFix::FStackIssueFix()
 UNiagaraStackEntry::FStackIssueFix::FStackIssueFix(FText InDescription, FStackIssueFixDelegate InFixDelegate)
 	: Description(InDescription)
 	, FixDelegate(InFixDelegate)
+	, UniqueIdentifier(FMD5::HashAnsiString(*FString::Printf(TEXT("%s"), *InDescription.ToString())))
 {
 	checkf(Description.IsEmptyOrWhitespace() == false, TEXT("Description can not be empty."));
 	checkf(InFixDelegate.IsBound(), TEXT("Fix delegate must be bound."));
@@ -37,6 +38,16 @@ const FText& UNiagaraStackEntry::FStackIssueFix::GetDescription() const
 	return Description;
 }
 
+void UNiagaraStackEntry::FStackIssueFix::SetFixDelegate(const FStackIssueFixDelegate& InFixDelegate)
+{
+	FixDelegate = InFixDelegate;
+}
+
+const FString& UNiagaraStackEntry::FStackIssueFix::GetUniqueIdentifier() const
+{
+	return UniqueIdentifier;
+}
+
 const UNiagaraStackEntry::FStackIssueFixDelegate& UNiagaraStackEntry::FStackIssueFix::GetFixDelegate() const
 {
 	return FixDelegate;
@@ -50,18 +61,13 @@ UNiagaraStackEntry::FStackIssue::FStackIssue(EStackIssueSeverity InSeverity, FTe
 	: Severity(InSeverity)
 	, ShortDescription(InShortDescription)
 	, LongDescription(InLongDescription)
+	, UniqueIdentifier(FMD5::HashAnsiString(*FString::Printf(TEXT("%s-%s"), *InStackEditorDataKey, *InLongDescription.ToString())))
 	, bCanBeDismissed(bInCanBeDismissed)
 	, Fixes(InFixes)
 {
 	checkf(ShortDescription.IsEmptyOrWhitespace() == false, TEXT("Short description can not be empty."));
 	checkf(LongDescription.IsEmptyOrWhitespace() == false, TEXT("Long description can not be empty."));
 	checkf(InStackEditorDataKey.IsEmpty() == false, TEXT("Stack editor data key can not be empty."));
-	FString Descriptions = LongDescription.ToString();
-	for (FStackIssueFix Fix : Fixes)
-	{
-		Descriptions += Fix.GetDescription().ToString();
-	}
-	UniqueIdentifier = FMD5::HashAnsiString(*FString::Printf(TEXT("%s-%s"), *InStackEditorDataKey, *Descriptions));
 }
 
 UNiagaraStackEntry::FStackIssue::FStackIssue(EStackIssueSeverity InSeverity, FText InShortDescription, FText InLongDescription, FString InStackEditorDataKey, bool bInCanBeDismissed, FStackIssueFix InFix)
@@ -481,6 +487,7 @@ void UNiagaraStackEntry::RefreshStackErrorChildren()
 		else
 		{
 			ErrorEntry = *Found;
+			ErrorEntry->SetStackIssue(Issue); // we found the entry by id but we want to properly refresh the subentries of the issue (specifically its fixes), too
 		}
 		NewErrorChildren.Add(ErrorEntry);
 	}
diff --git a/Engine/Plugins/FX/Niagara/Source/NiagaraEditor/Private/ViewModels/Stack/NiagaraStackErrorItem.cpp b/Engine/Plugins/FX/Niagara/Source/NiagaraEditor/Private/ViewModels/Stack/NiagaraStackErrorItem.cpp
--- a/Engine/Plugins/FX/Niagara/Source/NiagaraEditor/Private/ViewModels/Stack/NiagaraStackErrorItem.cpp
+++ b/Engine/Plugins/FX/Niagara/Source/NiagaraEditor/Private/ViewModels/Stack/NiagaraStackErrorItem.cpp
@@ -15,6 +15,11 @@ void UNiagaraStackErrorItem::Initialize(FRequiredEntryData InRequiredEntryData,
 	EntryStackEditorDataKey = InStackEditorDataKey;
 }
 
+void UNiagaraStackErrorItem::SetStackIssue(const FStackIssue& InStackIssue)
+{
+	StackIssue = InStackIssue;
+}
+
 FText UNiagaraStackErrorItem::GetDisplayName() const
 {
 	return StackIssue.GetShortDescription();
@@ -53,12 +58,17 @@ void UNiagaraStackErrorItem::RefreshChildrenInternal(const TArray<UNiagaraStackE
 	// fixes
 	for (int i = 0; i < StackIssue.GetFixes().Num(); i++)
 	{
+		UNiagaraStackEntry::FStackIssueFix CurrentFix = StackIssue.GetFixes()[i];
 		UNiagaraStackErrorItemFix* ErrorEntryFix = FindCurrentChildOfTypeByPredicate<UNiagaraStackErrorItemFix>(CurrentChildren,
-			[&](UNiagaraStackErrorItemFix* CurrentChild) { return CurrentChild->GetStackIssueFix() == StackIssue.GetFixes()[i]; });
+			[&](UNiagaraStackErrorItemFix* CurrentChild) { return CurrentChild->GetStackIssueFix().GetUniqueIdentifier() == CurrentFix.GetUniqueIdentifier(); });
 		if (ErrorEntryFix == nullptr)
 		{
 			ErrorEntryFix = NewObject<UNiagaraStackErrorItemFix>(this);
-			ErrorEntryFix->Initialize(CreateDefaultChildRequiredData(), StackIssue, StackIssue.GetFixes()[i], EntryStackEditorDataKey);
+			ErrorEntryFix->Initialize(CreateDefaultChildRequiredData(), StackIssue, CurrentFix, EntryStackEditorDataKey);
+		}
+		else
+		{
+			ErrorEntryFix->SetFixDelegate(CurrentFix.GetFixDelegate());
 		}
 		ErrorEntryFix->OnIssueFixed().AddUObject(this, &UNiagaraStackErrorItem::IssueFixed);
 		NewChildren.Add(ErrorEntryFix);
@@ -138,6 +148,11 @@ UNiagaraStackErrorItem::FOnIssueNotify& UNiagaraStackErrorItemFix::OnIssueFixed(
 	return IssueFixedDelegate;
 }
 
+void UNiagaraStackErrorItemFix::SetFixDelegate(const FStackIssueFixDelegate& InFixDelegate)
+{
+	IssueFix.SetFixDelegate(InFixDelegate);
+}
+
 //UNiagaraStackErrorItemDismiss
 
 void UNiagaraStackErrorItemDismiss::Initialize(FRequiredEntryData InRequiredEntryData, UNiagaraStackEntry::FStackIssue InStackIssue, FString InStackEditorDataKey)
diff --git a/Engine/Plugins/FX/Niagara/Source/NiagaraEditor/Private/ViewModels/Stack/NiagaraStackModuleItem.cpp b/Engine/Plugins/FX/Niagara/Source/NiagaraEditor/Private/ViewModels/Stack/NiagaraStackModuleItem.cpp
--- a/Engine/Plugins/FX/Niagara/Source/NiagaraEditor/Private/ViewModels/Stack/NiagaraStackModuleItem.cpp
+++ b/Engine/Plugins/FX/Niagara/Source/NiagaraEditor/Private/ViewModels/Stack/NiagaraStackModuleItem.cpp
@@ -481,7 +481,7 @@ void UNiagaraStackModuleItem::RefreshIssues(TArray<FStackIssue>& NewIssues)
 								TargetIndex = PostIndex; // if it also has post dependencies just add it before its first post dependency
 							}
 						}
-						NewModuleNode = FNiagaraStackGraphUtilities::AddScriptModuleToStack(ModuleAsset, *TargetOutputNode, TargetIndex); // no target index for now, until fix
+						NewModuleNode = FNiagaraStackGraphUtilities::AddScriptModuleToStack(ModuleAsset, *TargetOutputNode, TargetIndex);
 						checkf(NewModuleNode != nullptr, TEXT("Add module action failed"));
 						FNiagaraStackGraphUtilities::InitializeStackFunctionInputs(GetSystemViewModel(), GetEmitterViewModel(), GetStackEditorData(), *NewModuleNode, *NewModuleNode);
 						FNiagaraStackGraphUtilities::RelayoutGraph(*TargetOutputNode->GetGraph());
@@ -493,7 +493,7 @@ void UNiagaraStackModuleItem::RefreshIssues(TArray<FStackIssue>& NewIssues)
 			UNiagaraStackEntry::FStackIssue Error(
 				EStackIssueSeverity::Error,
 				LOCTEXT("DependencyWarning", "The module has unmet dependencies."),
-				FText::Format(LOCTEXT("DependencyWarningLong", "The following {0} is not met: {1}"), DependencyTypeString, Dependency.Description),
+				FText::Format(LOCTEXT("DependencyWarningLong", "The following {0} is not met: {1}; {2}"), DependencyTypeString, FText::FromName(Dependency.Id), Dependency.Description),
 				FString::Printf(TEXT("%s-dependency-%s"), *GetStackEditorDataKey(), *Dependency.Id.ToString()),
 				true,
 				Fixes);
diff --git a/Engine/Plugins/FX/Niagara/Source/NiagaraEditor/Public/ViewModels/Stack/NiagaraStackEntry.h b/Engine/Plugins/FX/Niagara/Source/NiagaraEditor/Public/ViewModels/Stack/NiagaraStackEntry.h
--- a/Engine/Plugins/FX/Niagara/Source/NiagaraEditor/Public/ViewModels/Stack/NiagaraStackEntry.h
+++ b/Engine/Plugins/FX/Niagara/Source/NiagaraEditor/Public/ViewModels/Stack/NiagaraStackEntry.h
@@ -112,16 +112,16 @@ public:
 
 		const FText& GetDescription() const;
 
+		void SetFixDelegate(const FStackIssueFixDelegate& InFixDelegate);
+
 		const FStackIssueFixDelegate& GetFixDelegate() const;
 
-		bool operator == (const FStackIssueFix &Other) const
-		{
-			return Description.CompareTo(Other.Description) == 0 && FixDelegate.GetHandle() == Other.FixDelegate.GetHandle();
-		}
+		const FString& GetUniqueIdentifier() const;
 
 	private:
 		FText Description;
 		FStackIssueFixDelegate FixDelegate;
+		FString UniqueIdentifier;
 	};
 
 	struct FStackIssue
diff --git a/Engine/Plugins/FX/Niagara/Source/NiagaraEditor/Public/ViewModels/Stack/NiagaraStackErrorItem.h b/Engine/Plugins/FX/Niagara/Source/NiagaraEditor/Public/ViewModels/Stack/NiagaraStackErrorItem.h
--- a/Engine/Plugins/FX/Niagara/Source/NiagaraEditor/Public/ViewModels/Stack/NiagaraStackErrorItem.h
+++ b/Engine/Plugins/FX/Niagara/Source/NiagaraEditor/Public/ViewModels/Stack/NiagaraStackErrorItem.h
@@ -16,6 +16,7 @@ public:
 public:
 	void Initialize(FRequiredEntryData InRequiredEntryData, FStackIssue InStackIssue, FString InStackEditorDataKey);
 	FStackIssue GetStackIssue() const { return StackIssue; }
+	void SetStackIssue(const FStackIssue& InStackIssue);
 	virtual FText GetDisplayName() const override;
 	virtual EStackRowStyle GetStackRowStyle() const override;
 	FOnIssueNotify& OnIssueModified();
@@ -60,6 +61,7 @@ public:
 	virtual EStackRowStyle GetStackRowStyle() const override;
 	virtual FText GetFixButtonText() const;
 	UNiagaraStackErrorItem::FOnIssueNotify& OnIssueFixed();
+	void SetFixDelegate(const FStackIssueFixDelegate& InFixDelegate);
 
 protected:
 	FStackIssue StackIssue;
