Niagara - Fix input initialization for drag/drop with a "Set Variables" node.

diff --git a/Engine/Plugins/FX/Niagara/Source/NiagaraEditor/Private/ViewModels/Stack/NiagaraStackGraphUtilities.cpp b/Engine/Plugins/FX/Niagara/Source/NiagaraEditor/Private/ViewModels/Stack/NiagaraStackGraphUtilities.cpp
--- a/Engine/Plugins/FX/Niagara/Source/NiagaraEditor/Private/ViewModels/Stack/NiagaraStackGraphUtilities.cpp
+++ b/Engine/Plugins/FX/Niagara/Source/NiagaraEditor/Private/ViewModels/Stack/NiagaraStackGraphUtilities.cpp
@@ -426,7 +426,9 @@ void FNiagaraStackGraphUtilities::ConnectStackNodeGroup(const FStackNodeGroup& C
 	}
 }
 
-void FNiagaraStackGraphUtilities::InitializeStackFunctionInputs(TSharedRef<FNiagaraSystemViewModel> SystemViewModel, TSharedRef<FNiagaraEmitterViewModel> EmitterViewModel, UNiagaraStackEditorData& StackEditorData, UNiagaraNodeFunctionCall& ModuleNode, UNiagaraNodeFunctionCall& InputFunctionCallNode)
+DECLARE_DELEGATE_RetVal_OneParam(bool, FInputSelector, UNiagaraStackFunctionInput*);
+
+void InitializeStackFunctionInputsInternal(TSharedRef<FNiagaraSystemViewModel> SystemViewModel, TSharedRef<FNiagaraEmitterViewModel> EmitterViewModel, UNiagaraStackEditorData& StackEditorData, UNiagaraNodeFunctionCall& ModuleNode, UNiagaraNodeFunctionCall& InputFunctionCallNode, FInputSelector InputSelector)
 {
 	UNiagaraStackFunctionInputCollection* FunctionInputCollection = NewObject<UNiagaraStackFunctionInputCollection>(GetTransientPackage()); 
 	UNiagaraStackEntry::FRequiredEntryData RequiredEntryData(SystemViewModel, EmitterViewModel, NAME_None, NAME_None, StackEditorData);
@@ -446,7 +448,7 @@ void FNiagaraStackGraphUtilities::InitializeStackFunctionInputs(TSharedRef<FNiag
 			for (UNiagaraStackEntry* CategoryChild : CategoryChildren)
 			{
 				UNiagaraStackFunctionInput* FunctionInput = Cast<UNiagaraStackFunctionInput>(CategoryChild);
-				if (FunctionInput != nullptr && FunctionInput->CanReset())
+				if (FunctionInput != nullptr && (InputSelector.IsBound() == false || InputSelector.Execute(FunctionInput)) && FunctionInput->CanReset())
 				{
 					FunctionInput->Reset();
 				}
@@ -458,6 +460,21 @@ void FNiagaraStackGraphUtilities::InitializeStackFunctionInputs(TSharedRef<FNiag
 	SystemViewModel->NotifyDataObjectChanged(nullptr);
 }
 
+void FNiagaraStackGraphUtilities::InitializeStackFunctionInputs(TSharedRef<FNiagaraSystemViewModel> SystemViewModel, TSharedRef<FNiagaraEmitterViewModel> EmitterViewModel, UNiagaraStackEditorData& StackEditorData, UNiagaraNodeFunctionCall& ModuleNode, UNiagaraNodeFunctionCall& InputFunctionCallNode)
+{
+	InitializeStackFunctionInputsInternal(SystemViewModel, EmitterViewModel, StackEditorData, ModuleNode, InputFunctionCallNode, FInputSelector());
+}
+
+void FNiagaraStackGraphUtilities::InitializeStackFunctionInput(TSharedRef<FNiagaraSystemViewModel> SystemViewModel, TSharedRef<FNiagaraEmitterViewModel> EmitterViewModel, UNiagaraStackEditorData& StackEditorData, UNiagaraNodeFunctionCall& ModuleNode, UNiagaraNodeFunctionCall& InputFunctionCallNode, FName InputName)
+{
+	FInputSelector InputSelector;
+	InputSelector.BindLambda([&InputName](UNiagaraStackFunctionInput* Input)
+	{
+		return Input->GetInputParameterHandle().GetName() == InputName;
+	});
+	InitializeStackFunctionInputsInternal(SystemViewModel, EmitterViewModel, StackEditorData, ModuleNode, InputFunctionCallNode, InputSelector);
+}
+
 FString FNiagaraStackGraphUtilities::GenerateStackFunctionInputEditorDataKey(UNiagaraNodeFunctionCall& FunctionCallNode, FNiagaraParameterHandle InputParameterHandle)
 {
 	return FunctionCallNode.GetFunctionName() + InputParameterHandle.GetParameterHandleString().ToString();
diff --git a/Engine/Plugins/FX/Niagara/Source/NiagaraEditor/Private/ViewModels/Stack/NiagaraStackModuleItem.cpp b/Engine/Plugins/FX/Niagara/Source/NiagaraEditor/Private/ViewModels/Stack/NiagaraStackModuleItem.cpp
--- a/Engine/Plugins/FX/Niagara/Source/NiagaraEditor/Private/ViewModels/Stack/NiagaraStackModuleItem.cpp
+++ b/Engine/Plugins/FX/Niagara/Source/NiagaraEditor/Private/ViewModels/Stack/NiagaraStackModuleItem.cpp
@@ -27,6 +27,7 @@
 #include "NiagaraScriptSource.h"
 #include "NiagaraNodeAssignment.h"
 #include "NiagaraCommon.h"
+#include "NiagaraConstants.h"
 #include "Widgets/SWidget.h"
 
 #include "ScopedTransaction.h"
@@ -608,6 +609,21 @@ void UNiagaraStackModuleItem::NotifyModuleMoved()
 	ModifiedGroupItemsDelegate.ExecuteIfBound();
 }
 
+bool UNiagaraStackModuleItem::CanAddInput() const
+{
+	return FunctionCallNode->IsA<UNiagaraNodeAssignment>();
+}
+
+void UNiagaraStackModuleItem::AddInput(FNiagaraVariable InputParameter)
+{
+	if(ensureMsgf(CanAddInput(), TEXT("This module doesn't support adding inputs.")))
+	{
+		UNiagaraNodeAssignment* AssignmentNode = CastChecked<UNiagaraNodeAssignment>(FunctionCallNode);
+		AssignmentNode->AddParameter(InputParameter, FNiagaraConstants::GetAttributeDefaultValue(InputParameter));
+		FNiagaraStackGraphUtilities::InitializeStackFunctionInput(GetSystemViewModel(), GetEmitterViewModel(), GetStackEditorData(), *FunctionCallNode, *FunctionCallNode, InputParameter.GetName());
+	}
+}
+
 UObject* UNiagaraStackModuleItem::GetExternalAsset() const
 {
 	if (GetModuleNode().FunctionScript != nullptr && GetModuleNode().FunctionScript->IsAsset())
diff --git a/Engine/Plugins/FX/Niagara/Source/NiagaraEditor/Public/ViewModels/Stack/NiagaraStackGraphUtilities.h b/Engine/Plugins/FX/Niagara/Source/NiagaraEditor/Public/ViewModels/Stack/NiagaraStackGraphUtilities.h
--- a/Engine/Plugins/FX/Niagara/Source/NiagaraEditor/Public/ViewModels/Stack/NiagaraStackGraphUtilities.h
+++ b/Engine/Plugins/FX/Niagara/Source/NiagaraEditor/Public/ViewModels/Stack/NiagaraStackGraphUtilities.h
@@ -60,6 +60,8 @@ namespace FNiagaraStackGraphUtilities
 
 	void InitializeStackFunctionInputs(TSharedRef<FNiagaraSystemViewModel> SystemViewModel, TSharedRef<FNiagaraEmitterViewModel> EmitterViewModel, UNiagaraStackEditorData& StackEditorData, UNiagaraNodeFunctionCall& ModuleNode, UNiagaraNodeFunctionCall& InputFunctionCallNode);
 
+	void InitializeStackFunctionInput(TSharedRef<FNiagaraSystemViewModel> SystemViewModel, TSharedRef<FNiagaraEmitterViewModel> EmitterViewModel, UNiagaraStackEditorData& StackEditorData, UNiagaraNodeFunctionCall& ModuleNode, UNiagaraNodeFunctionCall& InputFunctionCallNode, FName InputName);
+
 	FString GenerateStackFunctionInputEditorDataKey(UNiagaraNodeFunctionCall& FunctionCallNode, FNiagaraParameterHandle InputParameterHandle);
 
 	FString GenerateStackModuleEditorDataKey(UNiagaraNodeFunctionCall& ModuleNode);
diff --git a/Engine/Plugins/FX/Niagara/Source/NiagaraEditor/Public/ViewModels/Stack/NiagaraStackModuleItem.h b/Engine/Plugins/FX/Niagara/Source/NiagaraEditor/Public/ViewModels/Stack/NiagaraStackModuleItem.h
--- a/Engine/Plugins/FX/Niagara/Source/NiagaraEditor/Public/ViewModels/Stack/NiagaraStackModuleItem.h
+++ b/Engine/Plugins/FX/Niagara/Source/NiagaraEditor/Public/ViewModels/Stack/NiagaraStackModuleItem.h
@@ -3,6 +3,7 @@
 #pragma once
 
 #include "ViewModels/Stack/NiagaraStackItem.h"
+#include "NiagaraTypes.h"
 #include "NiagaraStackModuleItem.generated.h"
 
 class UNiagaraNodeFunctionCall;
@@ -51,6 +52,10 @@ public:
 
 	void NotifyModuleMoved();
 
+	bool CanAddInput() const;
+
+	void AddInput(FNiagaraVariable InputParameter);
+
 protected:
 	virtual void RefreshChildrenInternal(const TArray<UNiagaraStackEntry*>& CurrentChildren, TArray<UNiagaraStackEntry*>& NewChildren, TArray<FStackIssue>& NewIssues) override;
 
diff --git a/Engine/Plugins/FX/Niagara/Source/NiagaraEditorWidgets/Private/Stack/SNiagaraStackModuleItem.cpp b/Engine/Plugins/FX/Niagara/Source/NiagaraEditorWidgets/Private/Stack/SNiagaraStackModuleItem.cpp
--- a/Engine/Plugins/FX/Niagara/Source/NiagaraEditorWidgets/Private/Stack/SNiagaraStackModuleItem.cpp
+++ b/Engine/Plugins/FX/Niagara/Source/NiagaraEditorWidgets/Private/Stack/SNiagaraStackModuleItem.cpp
@@ -268,17 +268,14 @@ void SNiagaraStackModuleItem::ShowInsertModuleMenu(int32 InsertIndex)
 
 FReply SNiagaraStackModuleItem::OnModuleItemDrop(TSharedPtr<FDragDropOperation> DragDropOperation)
 {
-	if (DragDropOperation->IsOfType<FNiagaraStackDragOperation>())
+	if (ModuleItem->CanAddInput() && DragDropOperation->IsOfType<FNiagaraStackDragOperation>())
 	{
-		if (UNiagaraNodeAssignment* NodeAssignment = Cast<UNiagaraNodeAssignment>(&ModuleItem->GetModuleNode()))
+		TSharedPtr<FNiagaraStackDragOperation> InputDragDropOperation = StaticCastSharedPtr<FNiagaraStackDragOperation>(DragDropOperation);
+		TSharedPtr<FNiagaraParameterAction> Action = StaticCastSharedPtr<FNiagaraParameterAction>(InputDragDropOperation->GetAction());
+		if (Action.IsValid())
 		{
-			TSharedPtr<FNiagaraStackDragOperation> InputDragDropOperation = StaticCastSharedPtr<FNiagaraStackDragOperation>(DragDropOperation);
-			TSharedPtr<FNiagaraParameterAction> Action = StaticCastSharedPtr<FNiagaraParameterAction>(InputDragDropOperation->GetAction());
-			if (Action.IsValid())
-			{
-				NodeAssignment->AddParameter(Action->GetParameter(), FNiagaraConstants::GetAttributeDefaultValue(Action->GetParameter()));
-				return FReply::Handled();
-			}
+			ModuleItem->AddInput(Action->GetParameter());
+			return FReply::Handled();
 		}
 	}
 
@@ -287,7 +284,7 @@ FReply SNiagaraStackModuleItem::OnModuleItemDrop(TSharedPtr<FDragDropOperation>
 
 bool SNiagaraStackModuleItem::OnModuleItemAllowDrop(TSharedPtr<FDragDropOperation> DragDropOperation)
 {
-	return ModuleItem != nullptr && Cast<UNiagaraNodeAssignment>(&ModuleItem->GetModuleNode()) && DragDropOperation->IsOfType<FNiagaraStackDragOperation>();
+	return ModuleItem->CanAddInput() && DragDropOperation->IsOfType<FNiagaraStackDragOperation>();
 }
 
 EVisibility SNiagaraStackModuleItem::GetStackIssuesWarningVisibility() const
