Niagara - Fix op description tool tip and keyword searches in the graph add menu, fix and standardize tool tip handling for script objects in menus, and add support for keyword searches for user defined scripts to match the built in ops.

diff --git a/Engine/Plugins/FX/Niagara/Source/Niagara/Classes/NiagaraScript.h b/Engine/Plugins/FX/Niagara/Source/Niagara/Classes/NiagaraScript.h
--- a/Engine/Plugins/FX/Niagara/Source/Niagara/Classes/NiagaraScript.h
+++ b/Engine/Plugins/FX/Niagara/Source/Niagara/Classes/NiagaraScript.h
@@ -314,6 +314,10 @@ public:
 	UPROPERTY(AssetRegistrySearchable, EditAnywhere, Category = Script, meta = (MultiLine = true))
 	FText Description;
 
+	/** A list of space separated keywords which can be used to find this script in editor menus. */
+	UPROPERTY(AssetRegistrySearchable, EditAnywhere, Category = Script)
+	FText Keywords;
+
 	UPROPERTY(EditAnywhere, Category = Script)
 	TMap<FName, FString> ScriptMetaData;
 #endif
diff --git a/Engine/Plugins/FX/Niagara/Source/NiagaraEditor/Private/EdGraphSchema_Niagara.cpp b/Engine/Plugins/FX/Niagara/Source/NiagaraEditor/Private/EdGraphSchema_Niagara.cpp
--- a/Engine/Plugins/FX/Niagara/Source/NiagaraEditor/Private/EdGraphSchema_Niagara.cpp
+++ b/Engine/Plugins/FX/Niagara/Source/NiagaraEditor/Private/EdGraphSchema_Niagara.cpp
@@ -314,7 +314,7 @@ TArray<TSharedPtr<FNiagaraSchemaAction_NewNode> > UEdGraphSchema_Niagara::GetGra
 		const TArray<FNiagaraOpInfo>& OpInfos = FNiagaraOpInfo::GetOpInfoArray();
 		for (const FNiagaraOpInfo& OpInfo : OpInfos)
 		{
-			TSharedPtr<FNiagaraSchemaAction_NewNode> AddOpAction = AddNewNodeAction(NewActions, OpInfo.Category, OpInfo.FriendlyName, OpInfo.Name, FText::GetEmpty());
+			TSharedPtr<FNiagaraSchemaAction_NewNode> AddOpAction = AddNewNodeAction(NewActions, OpInfo.Category, OpInfo.FriendlyName, OpInfo.Name, OpInfo.Description, OpInfo.Keywords);
 			UNiagaraNodeOp* OpNode = NewObject<UNiagaraNodeOp>(OwnerOfTemporaries);
 			OpNode->OpName = OpInfo.Name;
 			AddOpAction->NodeTemplate = OpNode;
@@ -331,6 +331,26 @@ TArray<TSharedPtr<FNiagaraSchemaAction_NewNode> > UEdGraphSchema_Niagara::GetGra
 		FunctionCallAction->NodeTemplate = CustomHlslNode;
 	}
 
+	auto AddScriptFunctionAction = [&NewActions, OwnerOfTemporaries](const FText& Category, const FAssetData& ScriptAsset)
+	{
+		FText AssetDesc;
+		ScriptAsset.GetTagValue(GET_MEMBER_NAME_CHECKED(UNiagaraScript, Description), AssetDesc);
+
+		FText Keywords;
+		ScriptAsset.GetTagValue(GET_MEMBER_NAME_CHECKED(UNiagaraScript, Keywords), Keywords);
+
+		FString DisplayNameString = FName::NameToDisplayString(ScriptAsset.AssetName.ToString(), false);
+
+		const FText MenuDesc = FText::FromString(DisplayNameString);
+		const FText TooltipDesc = FNiagaraEditorUtilities::FormatScriptAssetDescription(AssetDesc, ScriptAsset.ObjectPath);
+
+		TSharedPtr<FNiagaraSchemaAction_NewNode> FunctionCallAction = AddNewNodeAction(NewActions, Category, MenuDesc, *DisplayNameString, TooltipDesc, Keywords);
+
+		UNiagaraNodeFunctionCall* FunctionCallNode = NewObject<UNiagaraNodeFunctionCall>(OwnerOfTemporaries);
+		FunctionCallNode->FunctionScriptAssetObjectPath = ScriptAsset.ObjectPath;
+		FunctionCallAction->NodeTemplate = FunctionCallNode;
+	};
+
 	//Add functions
 	FAssetRegistryModule& AssetRegistryModule = FModuleManager::LoadModuleChecked<FAssetRegistryModule>(TEXT("AssetRegistry"));
 	TArray<FAssetData> ScriptAssets;
@@ -343,9 +363,6 @@ TArray<TSharedPtr<FNiagaraSchemaAction_NewNode> > UEdGraphSchema_Niagara::GetGra
 			FName UsageName;
 			ScriptAsset.GetTagValue(GET_MEMBER_NAME_CHECKED(UNiagaraScript, Usage), UsageName);
 
-			FText AssetDesc;
-			ScriptAsset.GetTagValue(GET_MEMBER_NAME_CHECKED(UNiagaraScript, Description), AssetDesc);
-
 			FString QualifiedUsageName = "ENiagaraScriptUsage::" + UsageName.ToString();
 			int32 UsageIndex = NiagaraScriptUsageEnum->GetIndexByNameString(QualifiedUsageName);
 			if (UsageIndex != INDEX_NONE)
@@ -353,23 +370,13 @@ TArray<TSharedPtr<FNiagaraSchemaAction_NewNode> > UEdGraphSchema_Niagara::GetGra
 				ENiagaraScriptUsage Usage = static_cast<ENiagaraScriptUsage>(NiagaraScriptUsageEnum->GetValueByIndex(UsageIndex));
 				if (Usage == ENiagaraScriptUsage::Function)
 				{
-					FString DisplayNameString = FName::NameToDisplayString(ScriptAsset.AssetName.ToString(), false);
-					const FText MenuDesc = FText::FromString(DisplayNameString);
-					const FText TooltipDesc = FText::Format(LOCTEXT("FunctionPopupTooltip", "Path: {0}\nDescription: {1}"), FText::FromString(ScriptAsset.ObjectPath.ToString()), AssetDesc);
-
-
-					TSharedPtr<FNiagaraSchemaAction_NewNode> FunctionCallAction = AddNewNodeAction(NewActions, LOCTEXT("Function Menu Title", "Functions"), MenuDesc, *DisplayNameString, FText::FromName(ScriptAsset.ObjectPath));
-
-					UNiagaraNodeFunctionCall* FunctionCallNode = NewObject<UNiagaraNodeFunctionCall>(OwnerOfTemporaries);
-					FunctionCallNode->FunctionScriptAssetObjectPath = ScriptAsset.ObjectPath;
-					FunctionCallAction->NodeTemplate = FunctionCallNode;
+					AddScriptFunctionAction(LOCTEXT("Function Menu Title", "Functions"), ScriptAsset);
 				}
 			}
 		}
 	}
 
 	//Add modules
-	//if (GraphUsage == ENiagaraScriptUsage::ParticleSpawnScript || GraphUsage == ENiagaraScriptUsage::ParticleUpdateScript) 
 	if (!bFunctionGraph)
 	{
 		for (const FAssetData& ScriptAsset : ScriptAssets)
@@ -378,8 +385,6 @@ TArray<TSharedPtr<FNiagaraSchemaAction_NewNode> > UEdGraphSchema_Niagara::GetGra
 			ScriptAsset.GetTagValue(GET_MEMBER_NAME_CHECKED(UNiagaraScript, Usage), UsageName);
 			const FString BitfieldTagValue = ScriptAsset.GetTagValueRef<FString>(GET_MEMBER_NAME_CHECKED(UNiagaraScript, ModuleUsageBitmask));
 			int32 BitfieldValue = FCString::Atoi(*BitfieldTagValue);
-			FText AssetDesc;
-			ScriptAsset.GetTagValue(GET_MEMBER_NAME_CHECKED(UNiagaraScript, Description), AssetDesc);
 
 			ENiagaraScriptUsage TargetUsage = ENiagaraScriptUsage::Module;
 			if (bUpdateGraph)
@@ -400,15 +405,7 @@ TArray<TSharedPtr<FNiagaraSchemaAction_NewNode> > UEdGraphSchema_Niagara::GetGra
 				ENiagaraScriptUsage Usage = static_cast<ENiagaraScriptUsage>(NiagaraScriptUsageEnum->GetValueByIndex(UsageIndex));
 				if (Usage == ENiagaraScriptUsage::Module)
 				{
-					FString DisplayNameString = FName::NameToDisplayString(ScriptAsset.AssetName.ToString(), false);
-					const FText MenuDesc = FText::FromString(DisplayNameString);
-					const FText TooltipDesc = FText::Format(LOCTEXT("ModulePopupTooltip", "Path: {0}\nDescription: {1}"), FText::FromString(ScriptAsset.ObjectPath.ToString()), AssetDesc);
-
-					TSharedPtr<FNiagaraSchemaAction_NewNode> FunctionCallAction = AddNewNodeAction(NewActions, LOCTEXT("Module Menu Title", "Modules"), MenuDesc, *DisplayNameString, FText::FromName(ScriptAsset.ObjectPath));
-
-					UNiagaraNodeFunctionCall* FunctionCallNode = NewObject<UNiagaraNodeFunctionCall>(OwnerOfTemporaries);
-					FunctionCallNode->FunctionScriptAssetObjectPath = ScriptAsset.ObjectPath;
-					FunctionCallAction->NodeTemplate = FunctionCallNode;
+					AddScriptFunctionAction(LOCTEXT("Module Menu Title", "Modules"), ScriptAsset);
 				}
 			}
 		}
@@ -428,17 +425,7 @@ TArray<TSharedPtr<FNiagaraSchemaAction_NewNode> > UEdGraphSchema_Niagara::GetGra
 				ENiagaraScriptUsage Usage = static_cast<ENiagaraScriptUsage>(NiagaraScriptUsageEnum->GetValueByIndex(UsageIndex));
 				if (Usage == ENiagaraScriptUsage::DynamicInput)
 				{
-					FString DisplayNameString = FName::NameToDisplayString(ScriptAsset.AssetName.ToString(), false);
-					const FText MenuDesc = FText::FromString(DisplayNameString);
-					FText AssetDesc;
-					ScriptAsset.GetTagValue(GET_MEMBER_NAME_CHECKED(UNiagaraScript, Description), AssetDesc);
-					const FText TooltipDesc = FText::Format(LOCTEXT("DynamicInputPopupTooltip", "Path: {0}\nDescription: {1}"), FText::FromString(ScriptAsset.ObjectPath.ToString()), AssetDesc);
-
-					TSharedPtr<FNiagaraSchemaAction_NewNode> FunctionCallAction = AddNewNodeAction(NewActions, LOCTEXT("DynamicInputMenu", "Dynamic Inputs"), MenuDesc, *DisplayNameString, FText::FromName(ScriptAsset.ObjectPath));
-
-					UNiagaraNodeFunctionCall* FunctionCallNode = NewObject<UNiagaraNodeFunctionCall>(OwnerOfTemporaries);
-					FunctionCallNode->FunctionScriptAssetObjectPath = ScriptAsset.ObjectPath;
-					FunctionCallAction->NodeTemplate = FunctionCallNode;
+					AddScriptFunctionAction(LOCTEXT("Dynamic Input Menu Title", "Dynamic Inputs"), ScriptAsset);
 				}
 			}
 		}
diff --git a/Engine/Plugins/FX/Niagara/Source/NiagaraEditor/Private/NiagaraEditorUtilities.cpp b/Engine/Plugins/FX/Niagara/Source/NiagaraEditor/Private/NiagaraEditorUtilities.cpp
--- a/Engine/Plugins/FX/Niagara/Source/NiagaraEditor/Private/NiagaraEditorUtilities.cpp
+++ b/Engine/Plugins/FX/Niagara/Source/NiagaraEditor/Private/NiagaraEditorUtilities.cpp
@@ -877,4 +877,11 @@ const FNiagaraEmitterHandle* FNiagaraEditorUtilities::GetEmitterHandleForEmitter
 		[&Emitter](const FNiagaraEmitterHandle& EmitterHandle) { return EmitterHandle.GetInstance() == &Emitter; });
 }
 
+FText FNiagaraEditorUtilities::FormatScriptAssetDescription(FText Description, FName Path)
+{
+	return Description.IsEmptyOrWhitespace()
+		? FText::Format(LOCTEXT("ScriptAssetDescriptionFormatPathOnly", "Path: {0}"), FText::FromName(Path))
+		: FText::Format(LOCTEXT("ScriptAssetDescriptionFormat", "Description: {1}\nPath: {0}"), FText::FromName(Path), Description);
+}
+
 #undef LOCTEXT_NAMESPACE
diff --git a/Engine/Plugins/FX/Niagara/Source/NiagaraEditor/Private/ViewModels/Stack/NiagaraStackParameterStoreGroup.cpp b/Engine/Plugins/FX/Niagara/Source/NiagaraEditor/Private/ViewModels/Stack/NiagaraStackParameterStoreGroup.cpp
--- a/Engine/Plugins/FX/Niagara/Source/NiagaraEditor/Private/ViewModels/Stack/NiagaraStackParameterStoreGroup.cpp
+++ b/Engine/Plugins/FX/Niagara/Source/NiagaraEditor/Private/ViewModels/Stack/NiagaraStackParameterStoreGroup.cpp
@@ -47,6 +47,11 @@ public:
 		return FText::Format(LOCTEXT("AddParameterActionDescriptionFormat", "Create a new {0} parameter."), GetDisplayName());
 	}
 
+	virtual FText GetKeywords() const override
+	{
+		return FText();
+	}
+
 private:
 	FNiagaraVariable NewParameterVariable;
 };
diff --git a/Engine/Plugins/FX/Niagara/Source/NiagaraEditor/Private/ViewModels/Stack/NiagaraStackRenderItemGroup.cpp b/Engine/Plugins/FX/Niagara/Source/NiagaraEditor/Private/ViewModels/Stack/NiagaraStackRenderItemGroup.cpp
--- a/Engine/Plugins/FX/Niagara/Source/NiagaraEditor/Private/ViewModels/Stack/NiagaraStackRenderItemGroup.cpp
+++ b/Engine/Plugins/FX/Niagara/Source/NiagaraEditor/Private/ViewModels/Stack/NiagaraStackRenderItemGroup.cpp
@@ -39,6 +39,11 @@ public:
 		return FText::FromString(RendererClass->GetDescription());
 	}
 
+	virtual FText GetKeywords() const override
+	{
+		return FText();
+	}
+
 	UClass* GetRendererClass() const
 	{
 		return RendererClass;
diff --git a/Engine/Plugins/FX/Niagara/Source/NiagaraEditor/Private/ViewModels/Stack/NiagaraStackScriptItemGroup.cpp b/Engine/Plugins/FX/Niagara/Source/NiagaraEditor/Private/ViewModels/Stack/NiagaraStackScriptItemGroup.cpp
--- a/Engine/Plugins/FX/Niagara/Source/NiagaraEditor/Private/ViewModels/Stack/NiagaraStackScriptItemGroup.cpp
+++ b/Engine/Plugins/FX/Niagara/Source/NiagaraEditor/Private/ViewModels/Stack/NiagaraStackScriptItemGroup.cpp
@@ -44,9 +44,12 @@ public:
 
 		FText AssetDescription;
 		AssetData.GetTagValue(GET_MEMBER_NAME_CHECKED(UNiagaraScript, Description), AssetDescription);
-		FText Description = FText::Format(LOCTEXT("AssetModuleDescriptionFormat", "Path: {0}\nDescription: {1}"), FText::FromString(AssetData.ObjectPath.ToString()), AssetDescription);
+		FText Description = FNiagaraEditorUtilities::FormatScriptAssetDescription(AssetDescription, AssetData.ObjectPath);
 
-		return MakeShareable(new FScriptGroupAddAction(Category, DisplayName, Description, FNiagaraVariable(), false, AssetData, false));
+		FText Keywords;
+		AssetData.GetTagValue(GET_MEMBER_NAME_CHECKED(UNiagaraScript, Keywords), Keywords);
+
+		return MakeShareable(new FScriptGroupAddAction(Category, DisplayName, Description, Keywords, FNiagaraVariable(), false, AssetData, false));
 	}
 
 	static TSharedRef<FScriptGroupAddAction> CreateExistingParameterModuleAction(FNiagaraVariable ParameterVariable)
@@ -59,7 +62,7 @@ public:
 		FText AttributeDescription = FNiagaraConstants::GetAttributeDescription(ParameterVariable);
 		FText Description = FText::Format(LOCTEXT("ExistingParameterModuleDescriptoinFormat", "Description: Set the parameter {0}. {1}"), FText::FromName(ParameterVariable.GetName()), AttributeDescription);
 
-		return MakeShareable(new FScriptGroupAddAction(Category, DisplayName, Description, ParameterVariable, false, FAssetData(), false));
+		return MakeShareable(new FScriptGroupAddAction(Category, DisplayName, Description, FText(), ParameterVariable, false, FAssetData(), false));
 	}
 
 	static TSharedRef<FScriptGroupAddAction> CreateNewParameterModuleAction(FName NewParameterNamespace, FNiagaraTypeDefinition NewParameterType)
@@ -71,7 +74,7 @@ public:
 		FNiagaraParameterHandle NewParameterHandle(NewParameterNamespace, *(TEXT("New") + NewParameterType.GetName()));
 		FNiagaraVariable NewParameter(NewParameterType, NewParameterHandle.GetParameterHandleString());
 
-		return MakeShareable(new FScriptGroupAddAction(Category, DisplayName, Description, NewParameter, true, FAssetData(), false));
+		return MakeShareable(new FScriptGroupAddAction(Category, DisplayName, Description, FText(), NewParameter, true, FAssetData(), false));
 	}
 
 	virtual FText GetCategory() const override
@@ -89,6 +92,11 @@ public:
 		return Description;
 	}
 
+	virtual FText GetKeywords() const override
+	{
+		return Keywords;
+	}
+
 	const FNiagaraVariable& GetModuleParameterVariable() const
 	{
 		return ModuleParameterVariable;
@@ -110,10 +118,11 @@ public:
 	}
 
 private:
-	FScriptGroupAddAction(FText InCategory, FText InDisplayName, FText InDescription, FNiagaraVariable InModuleParameterVariable, bool bInRenameParameterOnAdd, FAssetData InModuleAssetData, bool bInIsMaterialParameterModuleAction)
+	FScriptGroupAddAction(FText InCategory, FText InDisplayName, FText InDescription, FText InKeywords, FNiagaraVariable InModuleParameterVariable, bool bInRenameParameterOnAdd, FAssetData InModuleAssetData, bool bInIsMaterialParameterModuleAction)
 		: Category(InCategory)
 		, DisplayName(InDisplayName)
 		, Description(InDescription)
+		, Keywords(InKeywords)
 		, ModuleParameterVariable(InModuleParameterVariable)
 		, bRenameParameterOnAdd(bInRenameParameterOnAdd)
 		, ModuleAssetData(InModuleAssetData)
@@ -125,6 +134,7 @@ private:
 	FText Category;
 	FText DisplayName;
 	FText Description;
+	FText Keywords;
 	FNiagaraVariable ModuleParameterVariable;
 	bool bRenameParameterOnAdd;
 	FAssetData ModuleAssetData;
diff --git a/Engine/Plugins/FX/Niagara/Source/NiagaraEditor/Public/NiagaraEditorUtilities.h b/Engine/Plugins/FX/Niagara/Source/NiagaraEditor/Public/NiagaraEditorUtilities.h
--- a/Engine/Plugins/FX/Niagara/Source/NiagaraEditor/Public/NiagaraEditorUtilities.h
+++ b/Engine/Plugins/FX/Niagara/Source/NiagaraEditor/Public/NiagaraEditorUtilities.h
@@ -149,4 +149,6 @@ namespace FNiagaraEditorUtilities
 	 * @returns The emitter handle for the supplied emitter, or nullptr if the emitter isn't owned by this system.
 	 */
 	const FNiagaraEmitterHandle* GetEmitterHandleForEmitter(UNiagaraSystem& System, UNiagaraEmitter& Emitter);
+
+	NIAGARAEDITOR_API FText FormatScriptAssetDescription(FText Description, FName Path);
 };
diff --git a/Engine/Plugins/FX/Niagara/Source/NiagaraEditor/Public/ViewModels/Stack/INiagaraStackItemGroupAddUtilities.h b/Engine/Plugins/FX/Niagara/Source/NiagaraEditor/Public/ViewModels/Stack/INiagaraStackItemGroupAddUtilities.h
--- a/Engine/Plugins/FX/Niagara/Source/NiagaraEditor/Public/ViewModels/Stack/INiagaraStackItemGroupAddUtilities.h
+++ b/Engine/Plugins/FX/Niagara/Source/NiagaraEditor/Public/ViewModels/Stack/INiagaraStackItemGroupAddUtilities.h
@@ -15,6 +15,9 @@ public:
 	/** Gets a long description of what will happen if this add action is executed. */
 	virtual FText GetDescription() const = 0;
 
+	/** Gets a space separated string of keywords which expose additional search terms for this action. */
+	virtual FText GetKeywords() const = 0;
+
 	virtual ~INiagaraStackItemGroupAddAction() { }
 };
 
diff --git a/Engine/Plugins/FX/Niagara/Source/NiagaraEditorWidgets/Private/Stack/SNiagaraStackFunctionInputValue.cpp b/Engine/Plugins/FX/Niagara/Source/NiagaraEditorWidgets/Private/Stack/SNiagaraStackFunctionInputValue.cpp
--- a/Engine/Plugins/FX/Niagara/Source/NiagaraEditorWidgets/Private/Stack/SNiagaraStackFunctionInputValue.cpp
+++ b/Engine/Plugins/FX/Niagara/Source/NiagaraEditorWidgets/Private/Stack/SNiagaraStackFunctionInputValue.cpp
@@ -476,8 +476,8 @@ void SNiagaraStackFunctionInputValue::CollectAllActions(FGraphActionListBuilderB
 		for (UNiagaraScript* DynamicInputScript : DynamicInputScripts)
 		{
 			const FText DynamicInputText = FText::FromString(FName::NameToDisplayString(DynamicInputScript->GetName(), false));
-			const FText Tooltip = FText::Format(LOCTEXT("DynamicInputFormat", "Use {0} to provide a value for this input. Description: {1}"), DynamicInputText, DynamicInputScript->GetDescription());
-			TSharedPtr<FNiagaraMenuAction> DynamicInputAction(new FNiagaraMenuAction(CategoryName, DynamicInputText, Tooltip, 0, FText(),
+			const FText Tooltip = FNiagaraEditorUtilities::FormatScriptAssetDescription(DynamicInputScript->Description, *DynamicInputScript->GetPathName());
+			TSharedPtr<FNiagaraMenuAction> DynamicInputAction(new FNiagaraMenuAction(CategoryName, DynamicInputText, Tooltip, 0, DynamicInputScript->Keywords,
 				FNiagaraMenuAction::FOnExecuteStackAction::CreateSP(this, &SNiagaraStackFunctionInputValue::DynamicInputScriptSelected, DynamicInputScript)));
 			OutAllActions.AddAction(DynamicInputAction);
 		}
diff --git a/Engine/Plugins/FX/Niagara/Source/NiagaraEditorWidgets/Private/Stack/SNiagaraStackItemGroupAddMenu.cpp b/Engine/Plugins/FX/Niagara/Source/NiagaraEditorWidgets/Private/Stack/SNiagaraStackItemGroupAddMenu.cpp
--- a/Engine/Plugins/FX/Niagara/Source/NiagaraEditorWidgets/Private/Stack/SNiagaraStackItemGroupAddMenu.cpp
+++ b/Engine/Plugins/FX/Niagara/Source/NiagaraEditorWidgets/Private/Stack/SNiagaraStackItemGroupAddMenu.cpp
@@ -50,7 +50,7 @@ void SNiagaraStackItemGroupAddMenu::CollectAllAddActions(FGraphActionListBuilder
 	for (TSharedRef<INiagaraStackItemGroupAddAction> AddAction : AddActions)
 	{
 		TSharedPtr<FNiagaraMenuAction> NewNodeAction(
-			new FNiagaraMenuAction(AddAction->GetCategory(), AddAction->GetDisplayName(), AddAction->GetDescription(), 0, FText(),
+			new FNiagaraMenuAction(AddAction->GetCategory(), AddAction->GetDisplayName(), AddAction->GetDescription(), 0, AddAction->GetKeywords(),
 				FNiagaraMenuAction::FOnExecuteStackAction::CreateRaw(AddUtilities, &INiagaraStackItemGroupAddUtilities::ExecuteAddAction, AddAction, InsertIndex)));
 		OutAllActions.AddAction(NewNodeAction);
 	}
