Fixing wyeth's torus error, which was caused by us not properly initializing defaults. We now initialize defaults in three waves in spawn scripts. Wave 1 are any straight up constants at the top of the spawn function. Wave 2 is inlined in spawn just before the function that needs them is called. Wave 3 is at the bottom of spawn in a section called HandleMissingDefaultValues.
Also updated the error and warning messages to be much clearer text.

#tests auto-tests pass

diff --git a/Engine/Plugins/FX/Niagara/Source/NiagaraCore/Private/NiagaraCustomVersion.cpp b/Engine/Plugins/FX/Niagara/Source/NiagaraCore/Private/NiagaraCustomVersion.cpp
--- a/Engine/Plugins/FX/Niagara/Source/NiagaraCore/Private/NiagaraCustomVersion.cpp
+++ b/Engine/Plugins/FX/Niagara/Source/NiagaraCore/Private/NiagaraCustomVersion.cpp
@@ -8,4 +8,4 @@ const FGuid FNiagaraCustomVersion::GUID(0xFCF57AFA, 0x50764283, 0xB9A9E658, 0xFF
 // Register the custom version with core
 FCustomVersionRegistration GRegisterNiagaraCustomVersion(FNiagaraCustomVersion::GUID, FNiagaraCustomVersion::LatestVersion, TEXT("NiagaraVer"));
 
-const FGuid FNiagaraCustomVersion::LatestScriptCompileVersion(0x338DE03B, 0x967544E2, 0xA906FC75, 0x60EC8661 );
\ No newline at end of file
+const FGuid FNiagaraCustomVersion::LatestScriptCompileVersion(0x73061024, 0xFC2646FF, 0x82F694A8, 0xF619C0CE);
\ No newline at end of file
diff --git a/Engine/Plugins/FX/Niagara/Source/NiagaraEditor/Private/NiagaraHlslTranslator.cpp b/Engine/Plugins/FX/Niagara/Source/NiagaraEditor/Private/NiagaraHlslTranslator.cpp
--- a/Engine/Plugins/FX/Niagara/Source/NiagaraEditor/Private/NiagaraHlslTranslator.cpp
+++ b/Engine/Plugins/FX/Niagara/Source/NiagaraEditor/Private/NiagaraHlslTranslator.cpp
@@ -385,27 +385,42 @@ FString FHlslNiagaraTranslator::GetFunctionDefinitions()
 
 	
 
-FString FHlslNiagaraTranslator::BuildMissingDefaults()
+void FHlslNiagaraTranslator::BuildMissingDefaults()
 {
-	FString HlslOutputString = TEXT("void HandleMissingDefaultValues(inout FSimulationContext Context)\n{\n");
+	AddBodyComment(TEXT("// Begin HandleMissingDefaultValues"));
+
 	if (UNiagaraScript::IsSpawnScript(TranslationStages[ActiveStageIdx].ScriptUsage))
 	{
+		// First go through all the variables that we did not write the defaults for yet. For spawn scripts, this usually
+		// means variables that reference other variables but are not themselves used within spawn.
+		for (FNiagaraVariable& Var : DeferredVariablesMissingDefault)
+		{
+			const UEdGraphPin* DefaultPin = UniqueVarToDefaultPin.FindChecked(Var);
+			bool bWriteToParamMapEntries = UniqueVarToWriteToParamMap.FindChecked(Var);
+			int32 OutputChunkId = INDEX_NONE;
+			HandleParameterRead(ActiveStageIdx, Var, DefaultPin, DefaultPin != nullptr ? Cast<UNiagaraNode>(DefaultPin->GetOwningNode()) : nullptr, OutputChunkId, !bWriteToParamMapEntries);
+		}
+
+		DeferredVariablesMissingDefault.Empty();
+
+		// Now go through and initialize any "Particles.Initial." variables
 		for (FNiagaraVariable& Var : InitialNamespaceVariablesMissingDefault)
 		{
 			if (FNiagaraParameterMapHistory::IsInitialValue(Var))
 			{
 				FNiagaraVariable SourceForInitialValue = FNiagaraParameterMapHistory::GetSourceForInitialValue(Var);
 				FString ParameterMapInstanceName = GetParameterMapInstanceName(0);
-				HlslOutputString += TEXT("\t") + ParameterMapInstanceName + TEXT(".") + GetSanitizedSymbolName(Var.GetName().ToString() + TEXT(" = ") +
-					ParameterMapInstanceName + TEXT(".") + GetSanitizedSymbolName(SourceForInitialValue.GetName().ToString()) + TEXT(";"));
+				FString Value = FString::Printf(TEXT("%s.%s = %s.%s;\n"), *ParameterMapInstanceName, *GetSanitizedSymbolName(Var.GetName().ToString()),
+					*ParameterMapInstanceName, *GetSanitizedSymbolName(SourceForInitialValue.GetName().ToString()));
+				AddBodyChunk(Value);
 				continue;
 			}
 		}
-		HlslOutputString += TEXT("\n");
+
+		InitialNamespaceVariablesMissingDefault.Empty();
 	}
-	
-	HlslOutputString += TEXT("\n}\n");
-	return HlslOutputString;
+
+	AddBodyComment(TEXT("// End HandleMissingDefaultValues\n\n"));
 }
 
 FString FHlslNiagaraTranslator::BuildParameterMapHlslDefinitions(TArray<FNiagaraVariable>& PrimaryDataSetOutputEntries)
@@ -844,13 +859,13 @@ const FNiagaraTranslateResults &FHlslNiagaraTranslator::Translate(const FNiagara
 			CurrentBodyChunkMode = TranslationStages[i].ChunkModeIndex;
 			if (UNiagaraScript::IsParticleSpawnScript(TranslationStages[i].ScriptUsage))
 			{
-				AddBodyChunk(TEXT(""), bHasInterpolatedSpawn ? TEXT("//Begin Interpolated Spawn Script!") : TEXT("//Begin Spawn Script!"), FNiagaraTypeDefinition::GetIntDef(), false, false);
+				AddBodyComment(bHasInterpolatedSpawn ? TEXT("//Begin Interpolated Spawn Script!") : TEXT("//Begin Spawn Script!"));
 				CurrentParamMapIndices.Empty();
 				CurrentParamMapIndices.Add(0);
 				TranslationStages[i].OutputNode->Compile(ThisTranslator, OutputChunks);
 				InstanceWrite = FDataSetAccessInfo(); // Reset after building the output..
-				AddBodyChunk(TEXT(""), TEXT("//End Spawn Script!\n\n"), FNiagaraTypeDefinition::GetIntDef(), false, false);
-				AddBodyChunk(TEXT(""), TEXT("HandleMissingDefaultValues(Context);"), FNiagaraTypeDefinition::GetIntDef(), false, false);
+				AddBodyComment(TEXT("//End Spawn Script!\n\n"));
+				BuildMissingDefaults();
 			}
 			
 			if (TranslationStages[i].bInterpolatePreviousParams)
@@ -860,12 +875,12 @@ const FNiagaraTranslateResults &FHlslNiagaraTranslator::Translate(const FNiagara
 
 			if (UNiagaraScript::IsParticleUpdateScript(TranslationStages[i].ScriptUsage))
 			{
-				AddBodyChunk(TEXT(""), TEXT("//Begin Update Script!"), FNiagaraTypeDefinition::GetIntDef(), false, false);
+				AddBodyComment(TEXT("//Begin Update Script!"));
 				//Now we compile the update script (with partial dt) and read from the temp values written above.
 				CurrentParamMapIndices.Empty();
 				CurrentParamMapIndices.Add(1);
 				TranslationStages[i].OutputNode->Compile(ThisTranslator, OutputChunks);
-				AddBodyChunk(TEXT(""), TEXT("//End Update Script!\n\n"), FNiagaraTypeDefinition::GetIntDef(), false, false);
+				AddBodyComment(TEXT("//End Update Script!\n\n"));
 			}
 		}
 		CurrentBodyChunkMode = ENiagaraCodeChunkMode::Body;
@@ -882,7 +897,7 @@ const FNiagaraTranslateResults &FHlslNiagaraTranslator::Translate(const FNiagara
 
 		if (IsSpawnScript())
 		{
-			AddBodyChunk(TEXT(""), TEXT("HandleMissingDefaultValues(Context);"), FNiagaraTypeDefinition::GetIntDef(), false, false);
+			BuildMissingDefaults();
 		}
 	}
 	else
@@ -1110,16 +1125,7 @@ const FNiagaraTranslateResults &FHlslNiagaraTranslator::Translate(const FNiagara
 		}
 
 		DefineInterpolatedParametersFunction(HlslOutput);
-
-		if (IsSpawnScript())
-		{
-			// The missing defaults function assumes the active stage index is the 0th value in the array.
-			int32 CurStageIdx = ActiveStageIdx;
-			ActiveStageIdx = 0;
-			HlslOutput += BuildMissingDefaults();
-			ActiveStageIdx = CurStageIdx;
-		}
-
+		
 		// define functions for reading and writing all secondary data sets
 		DefineDataSetReadFunction(HlslOutput, CompilationOutput.ScriptData.ReadDataSets);
 		DefineDataSetWriteFunction(HlslOutput, CompilationOutput.ScriptData.WriteDataSets, WriteConditionVars);
@@ -2066,6 +2072,7 @@ FString FHlslNiagaraTranslator::GetSanitizedSymbolName(FString SymbolName, bool
 	Ret.ReplaceInline(TEXT(","), TEXT("_"));
 	Ret.ReplaceInline(TEXT("-"), TEXT("_"));
 	Ret.ReplaceInline(TEXT(":"), TEXT("_"));
+
 	if (bCollapsNamespaces)
 	{
 		Ret.ReplaceInline(TEXT("."), TEXT("_"));
@@ -2254,6 +2261,16 @@ int32 FHlslNiagaraTranslator::AddSourceChunk(FString SymbolName, const FNiagaraT
 }
 
 
+int32 FHlslNiagaraTranslator::AddBodyComment(const FString& Comment)
+{
+	return AddBodyChunk(TEXT(""), Comment, FNiagaraTypeDefinition::GetIntDef(), false, false);
+}
+
+int32 FHlslNiagaraTranslator::AddBodyChunk(const FString& Value)
+{
+	return AddBodyChunk(TEXT(""), Value, FNiagaraTypeDefinition::GetIntDef(), INDEX_NONE, false, false);
+}
+
 int32 FHlslNiagaraTranslator::AddBodyChunk(FString SymbolName, FString Definition, const FNiagaraTypeDefinition& Type, TArray<int32>& SourceChunks, bool bDecl, bool bIsTerminated)
 {
 	check(CurrentBodyChunkMode == ENiagaraCodeChunkMode::Body || CurrentBodyChunkMode == ENiagaraCodeChunkMode::SpawnBody || CurrentBodyChunkMode == ENiagaraCodeChunkMode::UpdateBody);
@@ -2621,12 +2638,13 @@ FString FHlslNiagaraTranslator::GenerateConstantString(const FNiagaraVariable& C
 void FHlslNiagaraTranslator::InitializeParameterMapDefaults(int32 ParamMapHistoryIdx)
 {
 	bInitializedDefaults = true;
-	AddBodyChunk(TEXT(""), TEXT("//Begin Initialize Parameter Map Defaults"), FNiagaraTypeDefinition::GetIntDef(), false, false);
-	TArray<FNiagaraVariable> UniqueVars;
+	AddBodyComment(TEXT("//Begin Initialize Parameter Map Defaults"));
 	check(ParamMapHistories.Num() == TranslationStages.Num());
 
-	TMap<FNiagaraVariable, const UEdGraphPin*> UniqueVarToDefaultPin;
-	TMap<FNiagaraVariable, bool> UniqueVarToWriteToParamMap;
+	UniqueVars.Empty();
+	UniqueVarToDefaultPin.Empty();
+	UniqueVarToWriteToParamMap.Empty();
+	UniqueVarToChunk.Empty();
 
 	// First pass just use the current parameter map.
 	{
@@ -2654,6 +2672,8 @@ void FHlslNiagaraTranslator::InitializeParameterMapDefaults(int32 ParamMapHistor
 	// Above we added all the known from the spawn script, now let's add for all the others.
 	if (UNiagaraScript::IsSpawnScript(TranslationStages[ActiveStageIdx].ScriptUsage))
 	{
+		// Go through all referenced parameter maps and pull in any variables that are 
+		// in the primary data set output namespaces.
 		for (int32 ParamMapIdx = 0; ParamMapIdx < OtherOutputParamMapHistories.Num(); ParamMapIdx++)
 		{
 			const FNiagaraParameterMapHistory& History = OtherOutputParamMapHistories[ParamMapIdx];
@@ -2672,14 +2692,21 @@ void FHlslNiagaraTranslator::InitializeParameterMapDefaults(int32 ParamMapHistor
 			}
 		}
 
+		// Now sort them into buckets: Defined by constants (write immediately), Defined as initial values (delay to end),
+		// or defined by linkage or other script (defer to end if not originating from spawn, otherwise insert before first use)
 		for (FNiagaraVariable& Var : UniqueVars)
 		{
 			const UEdGraphPin* DefaultPin = UniqueVarToDefaultPin.FindChecked(Var);
 			bool bWriteToParamMapEntries = UniqueVarToWriteToParamMap.FindChecked(Var);
 			int32 OutputChunkId = INDEX_NONE;
-			HandleParameterRead(ParamMapHistoryIdx, Var, DefaultPin, DefaultPin != nullptr ? Cast<UNiagaraNode>(DefaultPin->GetOwningNode()) : nullptr, OutputChunkId, !bWriteToParamMapEntries);
 
-			if (FNiagaraParameterMapHistory::IsInitialValue(Var))
+			// During the initial pass, only support constants for the default pin.
+			if (!FNiagaraParameterMapHistory::IsInitialValue(Var) && (DefaultPin == nullptr || DefaultPin->LinkedTo.Num() == 0))
+			{
+				HandleParameterRead(ParamMapHistoryIdx, Var, DefaultPin, DefaultPin != nullptr ? Cast<UNiagaraNode>(DefaultPin->GetOwningNode()) : nullptr, OutputChunkId, !bWriteToParamMapEntries);
+				UniqueVarToChunk.Add(Var, OutputChunkId);
+			}
+			else if (FNiagaraParameterMapHistory::IsInitialValue(Var))
 			{
 				FNiagaraVariable SourceForInitialValue = FNiagaraParameterMapHistory::GetSourceForInitialValue(Var);
 				if (!UniqueVars.Contains(SourceForInitialValue))
@@ -2688,10 +2715,14 @@ void FHlslNiagaraTranslator::InitializeParameterMapDefaults(int32 ParamMapHistor
 				}
 				InitialNamespaceVariablesMissingDefault.Add(Var);
 			}
+			else
+			{
+				DeferredVariablesMissingDefault.Add(Var);
+			}
 		}
 	}
 
-	AddBodyChunk(TEXT(""), TEXT("//End Initialize Parameter Map Defaults"), FNiagaraTypeDefinition::GetIntDef(), false, false);
+	AddBodyComment(TEXT("//End Initialize Parameter Map Defaults"));
 }
 
 void FHlslNiagaraTranslator::Output(UNiagaraNodeOutput* OutputNode, const TArray<int32>& ComputedInputs)
@@ -3497,6 +3528,14 @@ void FHlslNiagaraTranslator::HandleParameterRead(int32 ParamMapHistoryIdx, const
 					Error(FText::Format(LOCTEXT("MissingInitialValueSource", "Variable {0} is used, but its source variable {1} is not set!"), FText::FromName(Var.GetName()), FText::FromName(SourceForInitialValue.GetName())), nullptr, nullptr);
 				}
 			}
+			else if (UniqueVars.Contains(Var) && UniqueVarToChunk.Contains(Var))
+			{
+				int32* FoundIdx = UniqueVarToChunk.Find(Var);
+				if (FoundIdx != nullptr)
+				{
+					LastSetChunkIdx = *FoundIdx;
+				}
+			}
 		}
 
 		if (LastSetChunkIdx == INDEX_NONE && !bIgnoreDefaultValue)
@@ -3567,6 +3606,10 @@ void FHlslNiagaraTranslator::HandleParameterRead(int32 ParamMapHistoryIdx, const
 						ParamMapSetVariablesToChunks[ParamMapHistoryIdx][VarIdx] = LastSetChunkIdx;
 						ParamMapDefinedAttributesToNamespaceVars.FindOrAdd(Var.GetName()) = Var;
 					}
+					else if (VarIdx == INDEX_NONE && UniqueVars.Contains(Var))
+					{
+						ParamMapDefinedAttributesToNamespaceVars.FindOrAdd(Var.GetName()) = Var;
+					}
 					else
 					{
 						Error(FText::Format(LOCTEXT("NoVarDefaultFound", "Default found for {0}, but not found in ParameterMap traversal"), FText::FromName(Var.GetName())), ErrorNode, nullptr);
@@ -4250,17 +4293,14 @@ void FHlslNiagaraTranslator::RegisterFunctionCall(ENiagaraScriptUsage ScriptUsag
 								{
 									const UEdGraphPin* DefaultPin = History.GetDefaultValuePin(VarIdx);
 									HandleParameterRead(ActiveStageIdx, AliasedVar, DefaultPin, ParamNode, LastSetChunkIdx);
-/*
-									if (History.IsPrimaryDataSetOutput(AliasedVar, CompileOptions.TargetUsage))
+									
+									// If this variable was in the pending defaults list, go ahead and remove it
+									// as we added it before first use...
+									if (DeferredVariablesMissingDefault.Contains(Var))
 									{
-										// insert defaults here..
+										DeferredVariablesMissingDefault.Remove(Var);
+										UniqueVarToChunk.Add(Var, LastSetChunkIdx);
 									}
-									else if (History.IsAliasedModuleParameter(AliasedVar))
-									{
-										// insert defaults here..
-										HandleParameterRead(ActiveStageIdx, AliasedVar, DefaultPin, ParamNode, LastSetChunkIdx);
-									}
-									*/
 								}
 							}
 						}
@@ -4959,14 +4999,42 @@ int32 FHlslNiagaraTranslator::CompileOutputPin(const UEdGraphPin* InPin)
 
 void FHlslNiagaraTranslator::Error(FText ErrorText, UNiagaraNode* Node, UEdGraphPin* Pin)
 {
-	FString ErrorString = FString::Printf(TEXT("Node: @@ - Pin: @@ - %s - Callstack: %s"), *ErrorText.ToString(), *GetCallstack());
+	FString NodePinStr = TEXT("");
+	FString NodePinPrefix = TEXT(" - ");
+	FString NodePinSuffix = TEXT("");
+	if (Node && Node->GetName().Len() > 0)
+	{
+		NodePinStr += TEXT("Node: ") + Node->GetName();
+		NodePinSuffix = TEXT(" - ");
+	}
+	if (Pin && Pin->PinFriendlyName.ToString().Len() > 0)
+	{
+		NodePinStr += TEXT(" Pin: ") + Pin->PinFriendlyName.ToString();
+		NodePinSuffix = TEXT(" - ");
+	}
+	
+	FString ErrorString = FString::Printf(TEXT("%s%s%s%sCallstack: %s"), *ErrorText.ToString(), *NodePinPrefix, *NodePinStr, *NodePinSuffix, *GetCallstack());
 	TranslateResults.CompileEvents.Add(FNiagaraCompileEvent(FNiagaraCompileEventType::Error, ErrorString, Node ? Node->NodeGuid : FGuid(), Pin ? Pin->PersistentGuid : FGuid(), GetCallstackGuids()));
 	TranslateResults.NumErrors++;
 }
 
 void FHlslNiagaraTranslator::Warning(FText WarningText, UNiagaraNode* Node, UEdGraphPin* Pin)
 {
-	FString WarnString = FString::Printf(TEXT("Node: @@ - Pin: @@ - %s - Callstack: %s"), *WarningText.ToString(), *GetCallstack());
+	FString NodePinStr = TEXT("");
+	FString NodePinPrefix = TEXT(" - ");
+	FString NodePinSuffix = TEXT("");
+	if (Node && Node->GetName().Len() > 0)
+	{
+		NodePinStr += TEXT("Node: ") + Node->GetName();
+		NodePinSuffix = TEXT(" - ");
+	}
+	if (Pin && Pin->PinFriendlyName.ToString().Len() > 0)
+	{
+		NodePinStr += TEXT(" Pin: ") + Pin->PinFriendlyName.ToString();
+		NodePinSuffix = TEXT(" - ");
+	}
+
+	FString WarnString = FString::Printf(TEXT("%s%s%s%sCallstack: %s"), *WarningText.ToString(), *NodePinPrefix, *NodePinStr, *NodePinSuffix, *GetCallstack());
 	TranslateResults.CompileEvents.Add(FNiagaraCompileEvent(FNiagaraCompileEventType::Warning, WarnString, Node ? Node->NodeGuid : FGuid(), Pin ? Pin->PersistentGuid : FGuid(), GetCallstackGuids()));
 	TranslateResults.NumWarnings++;
 }
diff --git a/Engine/Plugins/FX/Niagara/Source/NiagaraEditor/Private/NiagaraHlslTranslator.h b/Engine/Plugins/FX/Niagara/Source/NiagaraEditor/Private/NiagaraHlslTranslator.h
--- a/Engine/Plugins/FX/Niagara/Source/NiagaraEditor/Private/NiagaraHlslTranslator.h
+++ b/Engine/Plugins/FX/Niagara/Source/NiagaraEditor/Private/NiagaraHlslTranslator.h
@@ -390,6 +390,9 @@ protected:
 	int32 AddBodyChunk(FString SymbolName, FString Definition, const FNiagaraTypeDefinition& Type, TArray<int32>& SourceChunks, bool bDecl = true, bool bIsTerminated = true);
 	int32 AddBodyChunk(FString SymbolName, FString Definition, const FNiagaraTypeDefinition& Type, int32 SourceChunk, bool bDecl = true, bool bIsTerminated = true);
 	int32 AddBodyChunk(FString SymbolName, FString Definition, const FNiagaraTypeDefinition& Type, bool bDecl = true, bool bIsTerminated = true);
+	int32 AddBodyComment(const FString& Comment);
+	int32 AddBodyChunk(const FString& Definition);
+
 
 	FString GetFunctionDefinitions();
 
@@ -476,10 +479,10 @@ public:
 
 private:
 	void InitializeParameterMapDefaults(int32 ParamMapHistoryIdx);
-	void HandleParameterRead(int32 ParamMapHistoryIdx, const FNiagaraVariable& Var, const UEdGraphPin* DefaultPin, UNiagaraNode* ErrorNode, int32& OutputChunkId, bool bTreatAsUnknownParameterMap = false);
+	void HandleParameterRead(int32 ParamMapHistoryIdx, const FNiagaraVariable& Var, const UEdGraphPin* DefaultPin, UNiagaraNode* ErrorNode, int32& OutputChunkId,  bool bTreatAsUnknownParameterMap = false);
 	bool ShouldConsiderTargetParameterMap(ENiagaraScriptUsage InUsage) const;
 	FString BuildParameterMapHlslDefinitions(TArray<FNiagaraVariable>& PrimaryDataSetOutputEntries);
-	FString BuildMissingDefaults();
+	void BuildMissingDefaults();
 	void FinalResolveNamespacedTokens(const FString& ParameterMapInstanceNamespace, TArray<FString>& Tokens, TArray<FString>& ValidChildNamespaces, FNiagaraParameterMapHistoryBuilder& Builder, TArray<FNiagaraVariable>& UniqueParameterMapEntriesAliasesIntact, TArray<FNiagaraVariable>& UniqueParameterMapEntries, int32 ParamMapHistoryIdx);
 
 	void HandleNamespacedExternalVariablesToDataSetRead(TArray<FNiagaraVariable>& InDataSetVars, FString InNamespaceStr);
@@ -567,6 +570,18 @@ private:
 	// Synced to the external variables used when bulk compiling system scripts.
 	TArray<FNiagaraVariable> ExternalVariablesForBulkUsage;
 
+	// List of primary output variables encountered that need to be properly handled in spawn scripts.
+	TArray<FNiagaraVariable> UniqueVars;
+	
+	// Map of primary ouput variable description to its default value pin
+	TMap<FNiagaraVariable, const UEdGraphPin*> UniqueVarToDefaultPin;
+	
+	// Map of primary output variable description to whether or not it came from this script's parameter map
+	TMap<FNiagaraVariable, bool> UniqueVarToWriteToParamMap;
+	
+	// Map ofthe primary output variable description to the actual chunk id that wrote to it.
+	TMap<FNiagaraVariable, int32> UniqueVarToChunk;
+
 	// Strings to be inserted within the main function
 	TArray<FString> MainPreSimulateChunks;
 
@@ -595,5 +610,8 @@ private:
 	bool bInitializedDefaults;
 
 	TArray<const UEdGraphPin*> CurrentDefaultPinTraversal;
+	// Variables that need to be initialized based on some other variable's value at the end of spawn.
 	TArray<FNiagaraVariable> InitialNamespaceVariablesMissingDefault;
+	// Variables that need to be initialized in the body or at the end of spawn.
+	TArray<FNiagaraVariable> DeferredVariablesMissingDefault;
 };
