Removing the old compile debug file generation and now unified with the existing shader compiler workflow for the future.
If r.DumpShaderDebugInfo=1, make sure that we generate the assembly, ush, and params files in the Saved\ShaderDebugInfo\VM\<SYSTEM_NAME>\<EMITTER_NAME>\<SCRIPT_NAME_AND_USAGE_ID_IF_NONZERO>

#tests auto-tests pass

diff --git a/Engine/Plugins/FX/Niagara/Source/Niagara/Private/NiagaraScript.cpp b/Engine/Plugins/FX/Niagara/Source/Niagara/Private/NiagaraScript.cpp
--- a/Engine/Plugins/FX/Niagara/Source/Niagara/Private/NiagaraScript.cpp
+++ b/Engine/Plugins/FX/Niagara/Source/Niagara/Private/NiagaraScript.cpp
@@ -24,15 +24,6 @@
 #endif
 
 
-
-static int32 GbLogNiagaraSystemCompileResults = 1;
-static FAutoConsoleVariableRef CVarLogNiagaraSystemCompileResults(
-	TEXT("fx.LogNiagaraSystemCompileResults"),
-	GbLogNiagaraSystemCompileResults,
-	TEXT("If > 0 Niagara System scripts hlsl will be written to a text format when opened and closed in the editor. \n"),
-	ECVF_Default
-);
-
 DECLARE_STATS_GROUP(TEXT("Niagara Detailed"), STATGROUP_NiagaraDetailed, STATCAT_Advanced);
 
 FNiagaraScriptDebuggerInfo::FNiagaraScriptDebuggerInfo() : FrameLastWriteId(-1)
@@ -770,36 +761,7 @@ void UNiagaraScript::SetVMCompilationResults(const FNiagaraVMExecutableDataId& I
 	}
 
 	InvalidateExecutionReadyParameterStores();
-
-	if (GbLogNiagaraSystemCompileResults > 0)
-	{
-		FString ExportText = CachedScriptVM.LastHlslTranslation;
-		FString ExportTextAsm = CachedScriptVM.LastAssemblyTranslation;
-		if (IsGPUScript(GetUsage()))
-		{
-			ExportText = CachedScriptVM.LastHlslTranslationGPU;
-			ExportTextAsm = "";
-		}
-		FString ExportTextParams;
-		for (const FNiagaraVariable& Var : CachedScriptVM.Parameters.Parameters)
-		{
-			ExportTextParams += Var.ToString();
-			ExportTextParams += "\n";
-		}
-
-		FString ExportRIParams = RapidIterationParameters.ToString();
-		
-		UEnum* UsageEnum = Cast<UEnum>(FindObject<UEnum>(ANY_PACKAGE, TEXT("ENiagaraScriptUsage"), true));
-		FString TypeName = UsageEnum->GetNameStringByValue((int64)GetUsage());
-		FString FilePath = GetOutermost()->FileName.ToString();
-		FString PathPart, FilenamePart, ExtensionPart;
-		FPaths::Split(FilePath, PathPart, FilenamePart, ExtensionPart);
-		WriteTextFileToDisk(FPaths::ProjectLogDir(), FilenamePart + TEXT(".")+ TypeName + TEXT(".hlsl"), ExportText, true);
-		WriteTextFileToDisk(FPaths::ProjectLogDir(), FilenamePart + TEXT(".")+ TypeName + TEXT(".asm"), ExportTextAsm, true);
-		WriteTextFileToDisk(FPaths::ProjectLogDir(), FilenamePart + TEXT(".") + TypeName + TEXT(".params"), ExportTextParams, true);
-		WriteTextFileToDisk(FPaths::ProjectLogDir(), FilenamePart + TEXT(".") + TypeName + TEXT(".riparams"), ExportRIParams, true);
-	}
-
+	
 	OnVMScriptCompiled().Broadcast(this);
 }
 
diff --git a/Engine/Plugins/FX/Niagara/Source/NiagaraEditor/Private/NiagaraCompiler.cpp b/Engine/Plugins/FX/Niagara/Source/NiagaraEditor/Private/NiagaraCompiler.cpp
--- a/Engine/Plugins/FX/Niagara/Source/NiagaraEditor/Private/NiagaraCompiler.cpp
+++ b/Engine/Plugins/FX/Niagara/Source/NiagaraEditor/Private/NiagaraCompiler.cpp
@@ -665,8 +665,39 @@ FNiagaraCompileResults FHlslNiagaraCompiler::CompileScript(const FNiagaraCompile
 	
 	CompileResults.Data->LastHlslTranslation = TEXT("");
 
+	FShaderCompilerInput Input;
+	Input.VirtualSourceFilePath = TEXT("/Engine/Private/NiagaraEmitterInstanceShader.usf");
+	Input.EntryPointName = TEXT("SimulateMain");
+	Input.Environment.SetDefine(TEXT("VM_SIMULATION"), 1);
+	Input.Environment.IncludeVirtualPathToContentsMap.Add(TEXT("/Engine/Generated/NiagaraEmitterInstance.ush"), TranslatedHLSL);
+	Input.bGenerateDirectCompileFile = false;
+	Input.DumpDebugInfoRootPath = GShaderCompilingManager->GetAbsoluteShaderDebugInfoDirectory() / TEXT("VM");
+	FString UsageIdStr = !InOptions.TargetUsageId.IsValid() ? TEXT("") : (TEXT("_") + InOptions.TargetUsageId.ToString());
+	Input.DebugGroupName = InCompileRequest->SourceName / InCompileRequest->EmitterUniqueName / InCompileRequest->ENiagaraScriptUsageEnum->GetNameStringByValue((int64)InOptions.TargetUsage) + UsageIdStr;
+	Input.DumpDebugInfoPath = Input.DumpDebugInfoRootPath / Input.DebugGroupName;
+
+	if (GShaderCompilingManager->GetDumpShaderDebugInfo())
+	{
+		// Sanitize the name to be used as a path
+		// List mostly comes from set of characters not allowed by windows in a path.  Just try to rename a file and type one of these for the list.
+		Input.DumpDebugInfoPath.ReplaceInline(TEXT("<"), TEXT("("));
+		Input.DumpDebugInfoPath.ReplaceInline(TEXT(">"), TEXT(")"));
+		Input.DumpDebugInfoPath.ReplaceInline(TEXT("::"), TEXT("=="));
+		Input.DumpDebugInfoPath.ReplaceInline(TEXT("|"), TEXT("_"));
+		Input.DumpDebugInfoPath.ReplaceInline(TEXT("*"), TEXT("-"));
+		Input.DumpDebugInfoPath.ReplaceInline(TEXT("?"), TEXT("!"));
+		Input.DumpDebugInfoPath.ReplaceInline(TEXT("\""), TEXT("\'"));
+
+		if (!IFileManager::Get().DirectoryExists(*Input.DumpDebugInfoPath))
+		{
+			verifyf(IFileManager::Get().MakeDirectory(*Input.DumpDebugInfoPath, true), TEXT("Failed to create directory for shader debug info '%s'"), *Input.DumpDebugInfoPath);
+		}
+	}
+
+	bool bGPUScript = false;
 	if (InOptions.TargetUsage == ENiagaraScriptUsage::ParticleGPUComputeScript)
 	{
+		bGPUScript = true;
 		CompileResults.bComputeSucceeded = false;
 		if (TranslatorOutput != nullptr && TranslatorOutput->Errors.Len() > 0)
 		{
@@ -688,35 +719,6 @@ FNiagaraCompileResults FHlslNiagaraCompiler::CompileScript(const FNiagaraCompile
 	}
 	else
 	{
-		FShaderCompilerInput Input;
-		Input.VirtualSourceFilePath = TEXT("/Engine/Private/NiagaraEmitterInstanceShader.usf");
-		Input.EntryPointName = TEXT("SimulateMain");
-		Input.Environment.SetDefine(TEXT("VM_SIMULATION"), 1);
-		Input.Environment.IncludeVirtualPathToContentsMap.Add(TEXT("/Engine/Generated/NiagaraEmitterInstance.ush"), TranslatedHLSL);
-		Input.bGenerateDirectCompileFile = false;
-		Input.DumpDebugInfoRootPath = GShaderCompilingManager->GetAbsoluteShaderDebugInfoDirectory() / TEXT("VM");
-		FString UsageIdStr = !InOptions.TargetUsageId.IsValid() ? TEXT("") : (TEXT("_") + InOptions.TargetUsageId.ToString());
-		Input.DebugGroupName = InCompileRequest->SourceName / InCompileRequest->EmitterUniqueName / InCompileRequest->ENiagaraScriptUsageEnum->GetNameStringByValue((int64)InOptions.TargetUsage) + UsageIdStr;
-		Input.DumpDebugInfoPath = Input.DumpDebugInfoRootPath / Input.DebugGroupName;
-
-		if (GShaderCompilingManager->GetDumpShaderDebugInfo())
-		{
-			// Sanitize the name to be used as a path
-			// List mostly comes from set of characters not allowed by windows in a path.  Just try to rename a file and type one of these for the list.
-			Input.DumpDebugInfoPath.ReplaceInline(TEXT("<"), TEXT("("));
-			Input.DumpDebugInfoPath.ReplaceInline(TEXT(">"), TEXT(")"));
-			Input.DumpDebugInfoPath.ReplaceInline(TEXT("::"), TEXT("=="));
-			Input.DumpDebugInfoPath.ReplaceInline(TEXT("|"), TEXT("_"));
-			Input.DumpDebugInfoPath.ReplaceInline(TEXT("*"), TEXT("-"));
-			Input.DumpDebugInfoPath.ReplaceInline(TEXT("?"), TEXT("!"));
-			Input.DumpDebugInfoPath.ReplaceInline(TEXT("\""), TEXT("\'"));
-
-			if (!IFileManager::Get().DirectoryExists(*Input.DumpDebugInfoPath))
-			{
-				verifyf(IFileManager::Get().MakeDirectory(*Input.DumpDebugInfoPath, true), TEXT("Failed to create directory for shader debug info '%s'"), *Input.DumpDebugInfoPath);
-			}
-		}
-
 		FShaderCompilerOutput Output;
 
 		FVectorVMCompilationOutput CompilationOutput;
@@ -840,6 +842,26 @@ FNiagaraCompileResults FHlslNiagaraCompiler::CompileScript(const FNiagaraCompile
 		}
 	}
 
+	if (GShaderCompilingManager->GetDumpShaderDebugInfo() && CompileResults.Data.IsValid())
+	{
+		FString ExportText = CompileResults.Data->LastHlslTranslation;
+		FString ExportTextAsm = CompileResults.Data->LastAssemblyTranslation;
+		if (bGPUScript)
+		{
+			ExportText = CompileResults.Data->LastHlslTranslationGPU;
+			ExportTextAsm = "";
+		}
+		FString ExportTextParams;
+		for (const FNiagaraVariable& Var : CompileResults.Data->Parameters.Parameters)
+		{
+			ExportTextParams += Var.ToString();
+			ExportTextParams += "\n";
+		}
+		
+		FNiagaraEditorUtilities::WriteTextFileToDisk(Input.DumpDebugInfoPath, TEXT("NiagaraEmitterInstance.ush"), ExportText, true);
+		FNiagaraEditorUtilities::WriteTextFileToDisk(Input.DumpDebugInfoPath, TEXT("NiagaraEmitterInstance.asm"), ExportTextAsm, true);
+		FNiagaraEditorUtilities::WriteTextFileToDisk(Input.DumpDebugInfoPath, TEXT("NiagaraEmitterInstance.params"), ExportTextParams, true);
+	}
 	return CompileResults;
 }
 
