Fixing static analysis errors

#tests n/a

diff --git a/Engine/Plugins/FX/Niagara/Source/Niagara/Private/NiagaraSystem.cpp b/Engine/Plugins/FX/Niagara/Source/Niagara/Private/NiagaraSystem.cpp
--- a/Engine/Plugins/FX/Niagara/Source/Niagara/Private/NiagaraSystem.cpp
+++ b/Engine/Plugins/FX/Niagara/Source/Niagara/Private/NiagaraSystem.cpp
@@ -828,7 +828,7 @@ bool UNiagaraSystem::QueryCompileComplete(bool bWait, bool bDoPost, bool bDoNotA
 
 			if (CompiledScript->GetUsage() == ENiagaraScriptUsage::ParticleSpawnScript)
 			{
-				if (Emitter->SimTarget == ENiagaraSimTarget::GPUComputeSim)
+				if (Emitter && Emitter->SimTarget == ENiagaraSimTarget::GPUComputeSim)
 				{
 					Scripts.AddUnique(Emitter->GetGPUComputeScript());
 					ScriptDependencyMap.Add(CompiledScript, Emitter->GetGPUComputeScript());
@@ -838,13 +838,13 @@ bool UNiagaraSystem::QueryCompileComplete(bool bWait, bool bDoPost, bool bDoNotA
 
 			if (CompiledScript->GetUsage() == ENiagaraScriptUsage::ParticleUpdateScript)
 			{
-				if (Emitter->SimTarget == ENiagaraSimTarget::GPUComputeSim)
+				if (Emitter && Emitter->SimTarget == ENiagaraSimTarget::GPUComputeSim)
 				{
 					Scripts.AddUnique(Emitter->GetGPUComputeScript());
 					ScriptDependencyMap.Add(CompiledScript, Emitter->GetGPUComputeScript());
 					ScriptToEmitterNameMap.Add(Emitter->GetGPUComputeScript(), Emitter->GetUniqueEmitterName());
 				}
-				else if (Emitter->bInterpolatedSpawning)
+				else if (Emitter && Emitter->bInterpolatedSpawning)
 				{
 					Scripts.AddUnique(Emitter->SpawnScriptProps.Script);
 					ScriptDependencyMap.Add(CompiledScript, Emitter->SpawnScriptProps.Script);
diff --git a/Engine/Plugins/FX/Niagara/Source/NiagaraEditor/Private/ViewModels/Stack/NiagaraStackModuleItem.cpp b/Engine/Plugins/FX/Niagara/Source/NiagaraEditor/Private/ViewModels/Stack/NiagaraStackModuleItem.cpp
--- a/Engine/Plugins/FX/Niagara/Source/NiagaraEditor/Private/ViewModels/Stack/NiagaraStackModuleItem.cpp
+++ b/Engine/Plugins/FX/Niagara/Source/NiagaraEditor/Private/ViewModels/Stack/NiagaraStackModuleItem.cpp
@@ -294,88 +294,90 @@ void UNiagaraStackModuleItem::RefreshIssues(TArray<FStackIssue>& NewIssues)
 		}
 	}
 
-	for (FNiagaraModuleDependency Dependency: FunctionCallNode->FunctionScript->RequiredDependencies)
+	if (FunctionCallNode && FunctionCallNode->FunctionScript)
 	{
-		if (Dependency.Id == NAME_None)
+		for (FNiagaraModuleDependency Dependency : FunctionCallNode->FunctionScript->RequiredDependencies)
 		{
-			continue;
-		}
-		bool bDependencyMet = false;
-		UNiagaraNodeFunctionCall* FunctionNode = nullptr;
-		TArray <UNiagaraNodeFunctionCall*> DisabledDependencies;
-		TArray <FNiagaraStackModuleData> DisorderedDependencies;
-		
-		int32 DependencyModuleIndex = INDEX_NONE;
-		for (FNiagaraStackModuleData ModuleData : SystemModuleData)
-		{
-			FunctionNode = ModuleData.ModuleNode;
-			DependencyModuleIndex++;
-			if (FunctionNode != nullptr && FunctionNode->FunctionScript != nullptr && FunctionNode->FunctionScript->ProvidedDependencies.Contains(Dependency.Id))
+			if (Dependency.Id == NAME_None)
+			{
+				continue;
+			}
+			bool bDependencyMet = false;
+			UNiagaraNodeFunctionCall* FunctionNode = nullptr;
+			TArray <UNiagaraNodeFunctionCall*> DisabledDependencies;
+			TArray <FNiagaraStackModuleData> DisorderedDependencies;
+
+			int32 DependencyModuleIndex = INDEX_NONE;
+			for (FNiagaraStackModuleData ModuleData : SystemModuleData)
 			{
-				auto DependencyOutputUsage = ModuleData.Usage;
-				int32 PossibleIndex = UsagePriority.IndexOfByPredicate(
-					[&](const ENiagaraScriptUsage CurrentUsage)
+				FunctionNode = ModuleData.ModuleNode;
+				DependencyModuleIndex++;
+				if (FunctionNode != nullptr && FunctionNode->FunctionScript != nullptr && FunctionNode->FunctionScript->ProvidedDependencies.Contains(Dependency.Id))
+				{
+					auto DependencyOutputUsage = ModuleData.Usage;
+					int32 PossibleIndex = UsagePriority.IndexOfByPredicate(
+						[&](const ENiagaraScriptUsage CurrentUsage)
 					{
 						return UNiagaraScript::IsEquivalentUsage(DependencyOutputUsage, CurrentUsage);
 					});
-				int32 DependantIndex = UsagePriority.IndexOfByPredicate(
-					[&](const ENiagaraScriptUsage CurrentUsage)
+					int32 DependantIndex = UsagePriority.IndexOfByPredicate(
+						[&](const ENiagaraScriptUsage CurrentUsage)
 					{
 						return UNiagaraScript::IsEquivalentUsage(OutputNode->GetUsage(), CurrentUsage);
 					});
-				int32 Distance = PossibleIndex - DependantIndex;
-				
-				bool bIncorrectOrder = Distance == 0 ?  ((Dependency.Type == ENiagaraModuleDependencyType::PreDependency && ModuleIndex < DependencyModuleIndex)
-															|| (Dependency.Type == ENiagaraModuleDependencyType::PostDependency && ModuleIndex > DependencyModuleIndex))
-													 : ((Dependency.Type == ENiagaraModuleDependencyType::PreDependency && Distance < 0)
-															|| (Dependency.Type == ENiagaraModuleDependencyType::PostDependency && Distance > 0));
-										 
-				if (bIncorrectOrder)
-				{
-					DisorderedDependencies.Add(ModuleData);
-				}
-				else if (FunctionNode->IsNodeEnabled() == false) 
-				{
-					DisabledDependencies.Add(FunctionNode);
-				}
-				else
-				{
-					bDependencyMet = true;
-					break;
+					int32 Distance = PossibleIndex - DependantIndex;
+
+					bool bIncorrectOrder = Distance == 0 ? ((Dependency.Type == ENiagaraModuleDependencyType::PreDependency && ModuleIndex < DependencyModuleIndex)
+						|| (Dependency.Type == ENiagaraModuleDependencyType::PostDependency && ModuleIndex > DependencyModuleIndex))
+						: ((Dependency.Type == ENiagaraModuleDependencyType::PreDependency && Distance < 0)
+							|| (Dependency.Type == ENiagaraModuleDependencyType::PostDependency && Distance > 0));
+
+					if (bIncorrectOrder)
+					{
+						DisorderedDependencies.Add(ModuleData);
+					}
+					else if (FunctionNode->IsNodeEnabled() == false)
+					{
+						DisabledDependencies.Add(FunctionNode);
+					}
+					else
+					{
+						bDependencyMet = true;
+						break;
+					}
 				}
 			}
-		}
-		if (bDependencyMet == false)
-		{
-			TArray<FStackIssueFix> Fixes;
-			DependenciesNeeded.Add(Dependency);
-			
-			FText DependencyTypeString = Dependency.Type == ENiagaraModuleDependencyType::PreDependency ? LOCTEXT("PreDependency", "pre-dependency") : LOCTEXT("PostDependency", "post-dependency");
-
-			for ( UNiagaraNodeFunctionCall* DisabledNode : DisabledDependencies) // module exists but disabled
+			if (bDependencyMet == false)
 			{
-				UNiagaraStackEntry::FStackIssueFix Fix(
-					FText::Format(LOCTEXT("EnableDependency", "Enable dependency module {0}"), FText::FromString(DisabledNode->GetFunctionName())),
-					FStackIssueFixDelegate::CreateLambda([this, DisabledNode]()
+				TArray<FStackIssueFix> Fixes;
+				DependenciesNeeded.Add(Dependency);
+
+				FText DependencyTypeString = Dependency.Type == ENiagaraModuleDependencyType::PreDependency ? LOCTEXT("PreDependency", "pre-dependency") : LOCTEXT("PostDependency", "post-dependency");
+
+				for (UNiagaraNodeFunctionCall* DisabledNode : DisabledDependencies) // module exists but disabled
+				{
+					UNiagaraStackEntry::FStackIssueFix Fix(
+						FText::Format(LOCTEXT("EnableDependency", "Enable dependency module {0}"), FText::FromString(DisabledNode->GetFunctionName())),
+						FStackIssueFixDelegate::CreateLambda([this, DisabledNode]()
 					{
 						FScopedTransaction ScopedTransaction(LOCTEXT("EnableDependencyModule", "Enable dependency module"));
 						FNiagaraStackGraphUtilities::SetModuleIsEnabled(*DisabledNode, true);
 						OnRequestFullRefresh().Broadcast();
 
 					}));
-				Fixes.Add(Fix);
-			}
+					Fixes.Add(Fix);
+				}
 
-			for (FNiagaraStackModuleData DisorderedNode : DisorderedDependencies) // module exists but is not in the correct order (and possibly also disabled)
-			{
-				bool bNeedsEnable = !DisorderedNode.ModuleNode->IsNodeEnabled();
-				FText AndEnableModule = bNeedsEnable ? FText::Format(LOCTEXT("AndEnableDependency", "And enable dependency module {0}"), FText::FromString(DisorderedNode.ModuleNode->GetFunctionName())) : FText();
-				UNiagaraStackEntry::FStackIssueFix Fix(
-					FText::Format(LOCTEXT("ReorderDependency", "Reposition this module in the correct order related to {0} {1}"), FText::FromString(DisorderedNode.ModuleNode->GetFunctionName()), AndEnableModule),
-					FStackIssueFixDelegate::CreateLambda([this, bNeedsEnable, DisorderedNode, SystemModuleData, Dependency, ModuleIndex, OutputNode]()
+				for (FNiagaraStackModuleData DisorderedNode : DisorderedDependencies) // module exists but is not in the correct order (and possibly also disabled)
+				{
+					bool bNeedsEnable = !DisorderedNode.ModuleNode->IsNodeEnabled();
+					FText AndEnableModule = bNeedsEnable ? FText::Format(LOCTEXT("AndEnableDependency", "And enable dependency module {0}"), FText::FromString(DisorderedNode.ModuleNode->GetFunctionName())) : FText();
+					UNiagaraStackEntry::FStackIssueFix Fix(
+						FText::Format(LOCTEXT("ReorderDependency", "Reposition this module in the correct order related to {0} {1}"), FText::FromString(DisorderedNode.ModuleNode->GetFunctionName()), AndEnableModule),
+						FStackIssueFixDelegate::CreateLambda([this, bNeedsEnable, DisorderedNode, SystemModuleData, Dependency, ModuleIndex, OutputNode]()
 					{
 						FScopedTransaction ScopedTransaction(LOCTEXT("ReorderDependencyModule", "Reorder dependency module"));
-					
+
 						FunctionCallNode->Modify();
 						// reorder node
 						int32 CorrectIndex = Dependency.Type == ENiagaraModuleDependencyType::PostDependency ? DisorderedNode.Index : DisorderedNode.Index + 1;
@@ -390,115 +392,116 @@ void UNiagaraStackModuleItem::RefreshIssues(TArray<FStackIssue>& NewIssues)
 						FNiagaraStackGraphUtilities::RelayoutGraph(*OutputNode->GetGraph());
 						OnRequestFullRefresh().Broadcast();
 					}));
-				Fixes.Add(Fix);
-			}
-			if (DisorderedDependencies.Num() == 0 && DisabledDependencies.Num() == 0)
-			{
-				TArray<FAssetData> ModuleAssets;
-				FNiagaraStackGraphUtilities::GetScriptAssetsByDependencyProvided(ENiagaraScriptUsage::Module, Dependency.Id, ModuleAssets);
-				for (FAssetData ModuleAsset : ModuleAssets)
+					Fixes.Add(Fix);
+				}
+				if (DisorderedDependencies.Num() == 0 && DisabledDependencies.Num() == 0)
 				{
-					FText FixDescription = FText::Format(LOCTEXT("AddDependency", "Add new dependency module {0}"), FText::FromName(ModuleAsset.AssetName));
-					UNiagaraStackEntry::FStackIssueFix Fix(
-						FixDescription,
-						FStackIssueFixDelegate::CreateLambda([=]()
+					TArray<FAssetData> ModuleAssets;
+					FNiagaraStackGraphUtilities::GetScriptAssetsByDependencyProvided(ENiagaraScriptUsage::Module, Dependency.Id, ModuleAssets);
+					for (FAssetData ModuleAsset : ModuleAssets)
 					{
-						FScopedTransaction ScopedTransaction(FixDescription);
-						UNiagaraNodeFunctionCall* NewModuleNode = nullptr;
-						int32 TargetIndex = 0; 
-						UNiagaraScript* DependencyScript = Cast<UNiagaraScript>(ModuleAsset.GetAsset());
-						checkf(DependencyScript != nullptr, TEXT("Add module action failed"));
-						// Determine the output node for the group where the added dependency module belongs
-						UNiagaraNodeOutput* TargetOutputNode = nullptr;
-						for (int i = ModuleIndex; i < SystemModuleData.Num() && i >= 0; i = Dependency.Type == ENiagaraModuleDependencyType::PostDependency ? i + 1 : i - 1) // moving up or down depending on type
-						// starting at current module, which is a dependant
+						FText FixDescription = FText::Format(LOCTEXT("AddDependency", "Add new dependency module {0}"), FText::FromName(ModuleAsset.AssetName));
+						UNiagaraStackEntry::FStackIssueFix Fix(
+							FixDescription,
+							FStackIssueFixDelegate::CreateLambda([=]()
 						{
-							auto FoundRequirement = SystemModuleData[i].ModuleNode->FunctionScript->RequiredDependencies.FindByPredicate(
-									[&](FNiagaraModuleDependency CurrentDependency)
-									{
-										return CurrentDependency.Id == Dependency.Id;
-									});
-							if (FoundRequirement) // check for multiple dependendants along the way, and stop adjacent to the last one
+							FScopedTransaction ScopedTransaction(FixDescription);
+							UNiagaraNodeFunctionCall* NewModuleNode = nullptr;
+							int32 TargetIndex = 0;
+							UNiagaraScript* DependencyScript = Cast<UNiagaraScript>(ModuleAsset.GetAsset());
+							checkf(DependencyScript != nullptr, TEXT("Add module action failed"));
+							// Determine the output node for the group where the added dependency module belongs
+							UNiagaraNodeOutput* TargetOutputNode = nullptr;
+							for (int i = ModuleIndex; i < SystemModuleData.Num() && i >= 0; i = Dependency.Type == ENiagaraModuleDependencyType::PostDependency ? i + 1 : i - 1) // moving up or down depending on type
+							// starting at current module, which is a dependant
 							{
-								ENiagaraScriptUsage DependencyUsage = SystemModuleData[i].Usage;
-								TargetOutputNode = GetOutputNodeForModuleDependency(DependencyUsage, DependencyScript, GetSystemViewModel()->GetSystem(), GetEmitterViewModel()->GetEmitter(), Dependency);
-								if (TargetOutputNode != nullptr)
+								auto FoundRequirement = SystemModuleData[i].ModuleNode->FunctionScript->RequiredDependencies.FindByPredicate(
+									[&](FNiagaraModuleDependency CurrentDependency)
 								{
-									auto CurrentOutputNode = FNiagaraStackGraphUtilities::GetEmitterOutputNodeForStackNode(*SystemModuleData[i].ModuleNode);
-									if (TargetOutputNode == CurrentOutputNode)
-									{
-										TargetIndex = Dependency.Type == ENiagaraModuleDependencyType::PostDependency ? SystemModuleData[i].Index + 1 : SystemModuleData[i].Index;
-									}
-									else
+									return CurrentDependency.Id == Dependency.Id;
+								});
+								if (FoundRequirement) // check for multiple dependendants along the way, and stop adjacent to the last one
+								{
+									ENiagaraScriptUsage DependencyUsage = SystemModuleData[i].Usage;
+									TargetOutputNode = GetOutputNodeForModuleDependency(DependencyUsage, DependencyScript, GetSystemViewModel()->GetSystem(), GetEmitterViewModel()->GetEmitter(), Dependency);
+									if (TargetOutputNode != nullptr)
 									{
-										TargetIndex = Dependency.Type == ENiagaraModuleDependencyType::PostDependency ? 0 : INDEX_NONE;
+										auto CurrentOutputNode = FNiagaraStackGraphUtilities::GetEmitterOutputNodeForStackNode(*SystemModuleData[i].ModuleNode);
+										if (TargetOutputNode == CurrentOutputNode)
+										{
+											TargetIndex = Dependency.Type == ENiagaraModuleDependencyType::PostDependency ? SystemModuleData[i].Index + 1 : SystemModuleData[i].Index;
+										}
+										else
+										{
+											TargetIndex = Dependency.Type == ENiagaraModuleDependencyType::PostDependency ? 0 : INDEX_NONE;
+										}
 									}
 								}
 							}
-						}
-						if (TargetOutputNode == nullptr)
-						{
-							checkf(ModuleIndex != INDEX_NONE, TEXT("Module data wasn't found in system for current module!"));
-							// I am looking for a script in the same emitter as this module, so get it from there
-							UNiagaraScript& TargetScript = *FNiagaraEditorUtilities::GetScriptFromSystem(GetSystemViewModel()->GetSystem(), SystemModuleData[ModuleIndex].EmitterHandleId, DependencyScript->GetUsage(), DependencyScript->GetUsageId());
-							TargetOutputNode = FNiagaraEditorUtilities::GetScriptOutputNode(TargetScript);
-						}
-						TArray<FNiagaraStackModuleData> ScriptModuleData = SystemModuleData.FilterByPredicate([&](FNiagaraStackModuleData CurrentData) {return CurrentData.Usage == DependencyScript->GetUsage(); });
-						int32 PreIndex = INDEX_NONE; // index of last pre dependency
-						int32 PostIndex = INDEX_NONE; // index of fist post dependency, the module will have to be placed between these indexes
-						// for now, we skip the case where the dependencies are fulfilled in other script groups as well as here, because that's extremely unlikely
-						if (TargetIndex == INDEX_NONE)
-						{
-							TargetIndex = 0; //start at the beginning to look for potential dependencies of this dependency
-						}
-						for (int32 i = TargetIndex ; i < ScriptModuleData.Num() && i >= 0; i = Dependency.Type == ENiagaraModuleDependencyType::PostDependency ? i + 1 : i - 1)
-						{
-							UNiagaraNodeFunctionCall * CurrentNode =  ScriptModuleData[i].ModuleNode;
-							for (FNiagaraModuleDependency Requirement : DependencyScript->RequiredDependencies)
+							if (TargetOutputNode == nullptr)
 							{
-								if (Requirement.Id == NAME_None)
-								{
-									continue;
-								}
-								bool bDependencyMet = false;
-								if (CurrentNode->FunctionScript->ProvidedDependencies.Contains(Requirement.Id))
+								checkf(ModuleIndex != INDEX_NONE, TEXT("Module data wasn't found in system for current module!"));
+								// I am looking for a script in the same emitter as this module, so get it from there
+								UNiagaraScript& TargetScript = *FNiagaraEditorUtilities::GetScriptFromSystem(GetSystemViewModel()->GetSystem(), SystemModuleData[ModuleIndex].EmitterHandleId, DependencyScript->GetUsage(), DependencyScript->GetUsageId());
+								TargetOutputNode = FNiagaraEditorUtilities::GetScriptOutputNode(TargetScript);
+							}
+							TArray<FNiagaraStackModuleData> ScriptModuleData = SystemModuleData.FilterByPredicate([&](FNiagaraStackModuleData CurrentData) {return CurrentData.Usage == DependencyScript->GetUsage(); });
+							int32 PreIndex = INDEX_NONE; // index of last pre dependency
+							int32 PostIndex = INDEX_NONE; // index of fist post dependency, the module will have to be placed between these indexes
+							// for now, we skip the case where the dependencies are fulfilled in other script groups as well as here, because that's extremely unlikely
+							if (TargetIndex == INDEX_NONE)
+							{
+								TargetIndex = 0; //start at the beginning to look for potential dependencies of this dependency
+							}
+							for (int32 i = TargetIndex; i < ScriptModuleData.Num() && i >= 0; i = Dependency.Type == ENiagaraModuleDependencyType::PostDependency ? i + 1 : i - 1)
+							{
+								UNiagaraNodeFunctionCall * CurrentNode = ScriptModuleData[i].ModuleNode;
+								for (FNiagaraModuleDependency Requirement : DependencyScript->RequiredDependencies)
 								{
-									if (Requirement.Type == ENiagaraModuleDependencyType::PreDependency)
+									if (Requirement.Id == NAME_None)
 									{
-										PostIndex = i;
+										continue;
 									}
-									else if (PreIndex == INDEX_NONE) // only record the first post-dependency
+									bool bDependencyMet = false;
+									if (CurrentNode->FunctionScript->ProvidedDependencies.Contains(Requirement.Id))
 									{
-										PreIndex = i;
+										if (Requirement.Type == ENiagaraModuleDependencyType::PreDependency)
+										{
+											PostIndex = i;
+										}
+										else if (PreIndex == INDEX_NONE) // only record the first post-dependency
+										{
+											PreIndex = i;
+										}
 									}
 								}
 							}
-						}
-						if (PostIndex != INDEX_NONE) 
-						{
-							TargetIndex = 0; // if it has post dependencies place it at the top
-							if (PreIndex != INDEX_NONE)
+							if (PostIndex != INDEX_NONE)
 							{
-								TargetIndex = PostIndex; // if it also has post dependencies just add it before its first post dependency
+								TargetIndex = 0; // if it has post dependencies place it at the top
+								if (PreIndex != INDEX_NONE)
+								{
+									TargetIndex = PostIndex; // if it also has post dependencies just add it before its first post dependency
+								}
 							}
-						}
-						NewModuleNode = FNiagaraStackGraphUtilities::AddScriptModuleToStack(ModuleAsset, *TargetOutputNode, TargetIndex);
-						checkf(NewModuleNode != nullptr, TEXT("Add module action failed"));
-						FNiagaraStackGraphUtilities::InitializeStackFunctionInputs(GetSystemViewModel(), GetEmitterViewModel(), GetStackEditorData(), *NewModuleNode, *NewModuleNode);
-						FNiagaraStackGraphUtilities::RelayoutGraph(*TargetOutputNode->GetGraph());
-						OnRequestFullRefresh().Broadcast();
-					}));
-					Fixes.Add(Fix);
+							NewModuleNode = FNiagaraStackGraphUtilities::AddScriptModuleToStack(ModuleAsset, *TargetOutputNode, TargetIndex);
+							checkf(NewModuleNode != nullptr, TEXT("Add module action failed"));
+							FNiagaraStackGraphUtilities::InitializeStackFunctionInputs(GetSystemViewModel(), GetEmitterViewModel(), GetStackEditorData(), *NewModuleNode, *NewModuleNode);
+							FNiagaraStackGraphUtilities::RelayoutGraph(*TargetOutputNode->GetGraph());
+							OnRequestFullRefresh().Broadcast();
+						}));
+						Fixes.Add(Fix);
+					}
 				}
+				UNiagaraStackEntry::FStackIssue Error(
+					EStackIssueSeverity::Error,
+					LOCTEXT("DependencyWarning", "The module has unmet dependencies."),
+					FText::Format(LOCTEXT("DependencyWarningLong", "The following {0} is not met: {1}; {2}"), DependencyTypeString, FText::FromName(Dependency.Id), Dependency.Description),
+					FString::Printf(TEXT("%s-dependency-%s"), *GetStackEditorDataKey(), *Dependency.Id.ToString()),
+					true,
+					Fixes);
+				NewIssues.Add(Error);
 			}
-			UNiagaraStackEntry::FStackIssue Error(
-				EStackIssueSeverity::Error,
-				LOCTEXT("DependencyWarning", "The module has unmet dependencies."),
-				FText::Format(LOCTEXT("DependencyWarningLong", "The following {0} is not met: {1}; {2}"), DependencyTypeString, FText::FromName(Dependency.Id), Dependency.Description),
-				FString::Printf(TEXT("%s-dependency-%s"), *GetStackEditorDataKey(), *Dependency.Id.ToString()),
-				true,
-				Fixes);
-			NewIssues.Add(Error);
 		}
 	}
 }
