Tooltip variable types

diff --git a/Engine/Plugins/FX/Niagara/Source/NiagaraEditor/Private/NiagaraNode.cpp b/Engine/Plugins/FX/Niagara/Source/NiagaraEditor/Private/NiagaraNode.cpp
--- a/Engine/Plugins/FX/Niagara/Source/NiagaraEditor/Private/NiagaraNode.cpp
+++ b/Engine/Plugins/FX/Niagara/Source/NiagaraEditor/Private/NiagaraNode.cpp
@@ -223,6 +223,23 @@ TSharedPtr<SGraphNode> UNiagaraNode::CreateVisualWidget()
 }
 
 
+void UNiagaraNode::GetPinHoverText(const UEdGraphPin& Pin, FString& HoverTextOut) const
+{
+	const UNiagaraGraph* NiagaraGraph = GetNiagaraGraph();
+	if (NiagaraGraph)
+	{
+		const UEdGraphSchema_Niagara* Schema = Cast<UEdGraphSchema_Niagara>(NiagaraGraph->GetSchema());
+		if (Schema)
+		{
+			FNiagaraTypeDefinition TypeDef = Schema->PinToTypeDefinition(&Pin);
+			const FText Text = FText::Format(LOCTEXT("PinHoverTooltip", "Name: \"{0}\"\nType: {1}"),
+				FText::FromName(Pin.PinName),
+				TypeDef.GetNameText());
+			HoverTextOut = Text.ToString();
+		}
+	}
+}
+
 void UNiagaraNode::MarkNodeRequiresSynchronization(FString Reason, bool bRaiseGraphNeedsRecompile)
 {
 	Modify();
diff --git a/Engine/Plugins/FX/Niagara/Source/NiagaraEditor/Private/ViewModels/Stack/NiagaraStackFunctionInput.cpp b/Engine/Plugins/FX/Niagara/Source/NiagaraEditor/Private/ViewModels/Stack/NiagaraStackFunctionInput.cpp
--- a/Engine/Plugins/FX/Niagara/Source/NiagaraEditor/Private/ViewModels/Stack/NiagaraStackFunctionInput.cpp
+++ b/Engine/Plugins/FX/Niagara/Source/NiagaraEditor/Private/ViewModels/Stack/NiagaraStackFunctionInput.cpp
@@ -287,12 +287,16 @@ FText UNiagaraStackFunctionInput::GetTooltipText(EValueMode InValueMode) const
 		MetaData = NodeGraph->GetMetaData(ValueVariable);
 	}
 
+	FText Description = FText::GetEmpty();
 	if (MetaData != nullptr)
 	{
-		return MetaData->Description;
+		Description = MetaData->Description;
 	}
 
-	return FText::FromName(ValueVariable.GetName());
+	return FText::Format(LOCTEXT("FunctionInputTooltip", "Name: {0} \nType: {1} \nDesc: {2}"),
+		FText::FromName(ValueVariable.GetName()),
+		ValueVariable.GetType().GetNameText(),
+		Description);
 }
 
 void UNiagaraStackFunctionInput::RefreshChildrenInternal(const TArray<UNiagaraStackEntry*>& CurrentChildren, TArray<UNiagaraStackEntry*>& NewChildren, TArray<FStackIssue>& NewIssues)
diff --git a/Engine/Plugins/FX/Niagara/Source/NiagaraEditor/Private/Widgets/SNiagaraParameterMapView.cpp b/Engine/Plugins/FX/Niagara/Source/NiagaraEditor/Private/Widgets/SNiagaraParameterMapView.cpp
--- a/Engine/Plugins/FX/Niagara/Source/NiagaraEditor/Private/Widgets/SNiagaraParameterMapView.cpp
+++ b/Engine/Plugins/FX/Niagara/Source/NiagaraEditor/Private/Widgets/SNiagaraParameterMapView.cpp
@@ -341,7 +341,7 @@ void SNiagaraParameterMapView::CollectAllActions(FGraphActionListBuilderBase& Ou
 
 	ParameterEntries.KeySort([](const FNiagaraVariable& A, const FNiagaraVariable& B) { return (A.GetName() < B.GetName()); });
 
-	const FText TooltipFormat = LOCTEXT("Parameters", "Name: {0} Type: {1}");
+	const FText TooltipFormat = LOCTEXT("Parameters", "Name: {0} \nType: {1}");
 	for (const auto& ParameterEntry : ParameterEntries)
 	{
 		const FNiagaraVariable& Parameter = ParameterEntry.Key;
diff --git a/Engine/Plugins/FX/Niagara/Source/NiagaraEditor/Public/NiagaraNode.h b/Engine/Plugins/FX/Niagara/Source/NiagaraEditor/Public/NiagaraNode.h
--- a/Engine/Plugins/FX/Niagara/Source/NiagaraEditor/Public/NiagaraNode.h
+++ b/Engine/Plugins/FX/Niagara/Source/NiagaraEditor/Public/NiagaraNode.h
@@ -42,6 +42,7 @@ public:
 	virtual void OnPinRemoved(UEdGraphPin* InRemovedPin) override;
 	virtual void NodeConnectionListChanged() override;	
 	virtual TSharedPtr<SGraphNode> CreateVisualWidget() override; 
+	virtual void GetPinHoverText(const UEdGraphPin& Pin, FString& HoverTextOut) const override;
 	//~ End EdGraphNode Interface
 
 	/** Get the Niagara graph that owns this node */
