New variables in maps or parameter view will get the name Namespace.NewVariable. This is not an actual fix for UE-59633, but gives the user the incentive to rename variables.
#codereveiw Shaun.Kime

diff --git a/Engine/Plugins/FX/Niagara/Source/NiagaraEditor/Private/Widgets/SNiagaraParameterMapView.cpp b/Engine/Plugins/FX/Niagara/Source/NiagaraEditor/Private/Widgets/SNiagaraParameterMapView.cpp
--- a/Engine/Plugins/FX/Niagara/Source/NiagaraEditor/Private/Widgets/SNiagaraParameterMapView.cpp
+++ b/Engine/Plugins/FX/Niagara/Source/NiagaraEditor/Private/Widgets/SNiagaraParameterMapView.cpp
@@ -1017,8 +1017,9 @@ void SNiagaraAddParameterMenu::AddParameterSelected(FNiagaraVariable NewVariable
 {
 	if (bCreateCustomName)
 	{
+		const static FString NewVariableDefaultName = FString("NewVariable");
 		const FString ResultName = (InSection != NiagaraParameterMapSectionID::NONE ? NiagaraParameterMapSectionID::OnGetSectionTitle(InSection).ToString() + TEXT(".") : FString())
-			+ NewVariable.GetName().ToString();
+			+ NewVariableDefaultName;
 		NewVariable.SetName(FName(*ResultName));
 	}
 
