Fixing build issues

#tests n/a

diff --git a/Engine/Plugins/FX/Niagara/Source/Niagara/Private/NiagaraCommon.cpp b/Engine/Plugins/FX/Niagara/Source/Niagara/Private/NiagaraCommon.cpp
--- a/Engine/Plugins/FX/Niagara/Source/Niagara/Private/NiagaraCommon.cpp
+++ b/Engine/Plugins/FX/Niagara/Source/Niagara/Private/NiagaraCommon.cpp
@@ -305,7 +305,7 @@ UNiagaraDataInterface* FNiagaraScriptDataInterfaceCompileInfo::GetDefaultDataInt
 	UNiagaraDataInterface* Obj = CastChecked<UNiagaraDataInterface>(const_cast<UClass*>(Type.GetClass())->GetDefaultObject(true));
 	return Obj;
 }
-
+#if WITH_EDITORONLY_DATA
 void FNiagaraUtilities::PrepareRapidIterationParameters(const TArray<UNiagaraScript*>& Scripts, const TMap<UNiagaraScript*, UNiagaraScript*>& ScriptDependencyMap, const TMap<UNiagaraScript*, FString>& ScriptToEmitterNameMap)
 {
 	SCOPE_CYCLE_COUNTER(STAT_Niagara_Utilities_PrepareRapidIterationParameters);
@@ -381,4 +381,5 @@ void FNiagaraUtilities::PrepareRapidIterationParameters(const TArray<UNiagaraScr
 			Script->RapidIterationParameters = PreparedParameterStore;
 		}
 	}
-}
\ No newline at end of file
+}
+#endif
\ No newline at end of file
diff --git a/Engine/Plugins/FX/Niagara/Source/Niagara/Public/NiagaraCommon.h b/Engine/Plugins/FX/Niagara/Source/Niagara/Public/NiagaraCommon.h
--- a/Engine/Plugins/FX/Niagara/Source/Niagara/Public/NiagaraCommon.h
+++ b/Engine/Plugins/FX/Niagara/Source/Niagara/Public/NiagaraCommon.h
@@ -514,6 +514,7 @@ namespace FNiagaraUtilities
 
 	void CollectScriptDataInterfaceParameters(const UObject& Owner, const TArray<UNiagaraScript*>& Scripts, FNiagaraParameterStore& OutDataInterfaceParameters);
 
+#if WITH_EDITORONLY_DATA
 	/**
 	 * Prepares rapid iteration parameter stores for simulation by removing old parameters no longer used by functions, by initializing new parameters
 	 * added to functions, and by copying parameters across parameter stores for interscript dependencies.
@@ -524,4 +525,5 @@ namespace FNiagaraUtilities
 	 * scripts array must have an entry in this map.
 	 */
 	void NIAGARA_API PrepareRapidIterationParameters(const TArray<UNiagaraScript*>& Scripts, const TMap<UNiagaraScript*, UNiagaraScript*>& ScriptDependencyMap, const TMap<UNiagaraScript*, FString>& ScriptToEmitterNameMap);
+#endif
 };
\ No newline at end of file
