Fix creating an input parameter node from an input pin.

diff --git a/Engine/Plugins/FX/Niagara/Source/NiagaraEditor/Private/NiagaraNodeInput.cpp b/Engine/Plugins/FX/Niagara/Source/NiagaraEditor/Private/NiagaraNodeInput.cpp
--- a/Engine/Plugins/FX/Niagara/Source/NiagaraEditor/Private/NiagaraNodeInput.cpp
+++ b/Engine/Plugins/FX/Niagara/Source/NiagaraEditor/Private/NiagaraNodeInput.cpp
@@ -352,15 +352,16 @@ bool UNiagaraNodeInput::ReferencesSameInput(UNiagaraNodeInput* Other) const
 
 void UNiagaraNodeInput::AutowireNewNode(UEdGraphPin* FromPin)
 {
-	if (FromPin != nullptr)
+	UEdGraphPin* OutputPin = GetOutputPin(0);
+	if (FromPin != nullptr && OutputPin && OutputPin->PinType == FromPin->PinType)
 	{
 		const UEdGraphSchema_Niagara* Schema = CastChecked<UEdGraphSchema_Niagara>(GetSchema());
 		check(Schema);
+
 		if (Usage == ENiagaraInputNodeUsage::Parameter)
 		{
 			TArray<UNiagaraNodeInput*> InputNodes;
 			GetGraph()->GetNodesOfClass(InputNodes);
-			TSet<FName> InputNames;
 			int32 NumMatches = 0;
 			int32 HighestSortPriority = -1; // Set to -1 initially, so that in the event of no nodes, we still get zero.
 			for (UNiagaraNodeInput* InputNode : InputNodes)
@@ -372,18 +373,18 @@ void UNiagaraNodeInput::AutowireNewNode(UEdGraphPin* FromPin)
 
 				if (InputNode != this && InputNode->Usage == ENiagaraInputNodeUsage::Parameter)
 				{
-					if (this->ReferencesSameInput(InputNode))
+					if (ReferencesSameInput(InputNode))
 					{
 						NumMatches++;
+
 						check(InputNode->Input.GetName() == this->Input.GetName());
 						check(InputNode->ExposureOptions.bCanAutoBind == this->ExposureOptions.bCanAutoBind);
 						check(InputNode->ExposureOptions.bExposed == this->ExposureOptions.bExposed);
 						check(InputNode->ExposureOptions.bHidden == this->ExposureOptions.bHidden);
 						check(InputNode->ExposureOptions.bRequired == this->ExposureOptions.bRequired);
 						check(InputNode->DataInterface == this->DataInterface);
-						check(InputNode->CallSortPriority == this->CallSortPriority);
 					}
-					InputNames.Add(InputNode->Input.GetName());
+
 					if (InputNode->CallSortPriority > HighestSortPriority)
 					{
 						HighestSortPriority = InputNode->CallSortPriority;
@@ -399,37 +400,14 @@ void UNiagaraNodeInput::AutowireNewNode(UEdGraphPin* FromPin)
 				Type = Schema->PinToTypeDefinition(FromPin);
 			}
 
-			if (NumMatches == 0)
-			{
-				const FString PinName = FromPin->PinFriendlyName.IsEmpty() ? FromPin->PinName.ToString() : FromPin->PinFriendlyName.ToString();
-				if (UNiagaraNodeOp* OpNode = Cast<UNiagaraNodeOp>(FromPin->GetOwningNode()))
-				{
-					const FNiagaraOpInfo* OpInfo = FNiagaraOpInfo::GetOpInfo(OpNode->OpName);
-					check(OpInfo);
-					Input.SetType(Type);
-					CandidateName = (*(OpInfo->FriendlyName.ToString() + TEXT(" ") + PinName));
-				}
-				else if (UNiagaraNodeFunctionCall* FuncNode = Cast<UNiagaraNodeFunctionCall>(FromPin->GetOwningNode()))
-				{
-					Input.SetType(Type);
-					CandidateName = (*(FuncNode->GetNodeTitle(ENodeTitleType::FullTitle).ToString() + TEXT(" ") + PinName));
-				}
-				else
-				{
-					Input.SetType(Type);
-					CandidateName = (*PinName);
-				}
-
-				Input.SetName(FNiagaraUtilities::GetUniqueName(CandidateName, FNiagaraEditorUtilities::GetSystemConstantNames().Union(InputNames)));
-				CallSortPriority = HighestSortPriority + 1;
-			}
+			CallSortPriority = HighestSortPriority + 1;
 			ReallocatePins();
 		}
 
 		TArray<UEdGraphPin*> OutPins;
 		GetOutputPins(OutPins);
 		check(OutPins.Num() == 1 && OutPins[0] != NULL);
-		
+
 		if (GetSchema()->TryCreateConnection(FromPin, OutPins[0]))
 		{
 			FromPin->GetOwningNode()->NodeConnectionListChanged();
