Fix being able to drop parameters in the system view on incorrect execution categories.

diff --git a/Engine/Plugins/FX/Niagara/Source/NiagaraEditor/Private/ViewModels/Stack/NiagaraStackGraphUtilities.cpp b/Engine/Plugins/FX/Niagara/Source/NiagaraEditor/Private/ViewModels/Stack/NiagaraStackGraphUtilities.cpp
--- a/Engine/Plugins/FX/Niagara/Source/NiagaraEditor/Private/ViewModels/Stack/NiagaraStackGraphUtilities.cpp
+++ b/Engine/Plugins/FX/Niagara/Source/NiagaraEditor/Private/ViewModels/Stack/NiagaraStackGraphUtilities.cpp
@@ -1847,4 +1847,26 @@ void FNiagaraStackGraphUtilities::MoveModule(UNiagaraScript& SourceScript, UNiag
 	}
 }
 
+bool FNiagaraStackGraphUtilities::ParameterAllowedInExecutionCategory(const FName InParameterName, const FName ExecutionCategory)
+{
+	FNiagaraParameterHandle Handle = FNiagaraParameterHandle(InParameterName);
+	if (Handle.IsSystemHandle())
+	{
+		return ExecutionCategory == UNiagaraStackEntry::FExecutionCategoryNames::System
+			|| ExecutionCategory == UNiagaraStackEntry::FExecutionCategoryNames::Emitter
+			|| ExecutionCategory == UNiagaraStackEntry::FExecutionCategoryNames::Particle;
+	}
+	else if (Handle.IsEmitterHandle())
+	{
+		return ExecutionCategory == UNiagaraStackEntry::FExecutionCategoryNames::Emitter
+			|| ExecutionCategory == UNiagaraStackEntry::FExecutionCategoryNames::Particle;
+	}
+	else if (Handle.IsParticleAttributeHandle())
+	{
+		return ExecutionCategory == UNiagaraStackEntry::FExecutionCategoryNames::Particle;
+	}
+
+	return true;
+}
+
 #undef LOCTEXT_NAMESPACE
diff --git a/Engine/Plugins/FX/Niagara/Source/NiagaraEditor/Public/ViewModels/Stack/NiagaraStackGraphUtilities.h b/Engine/Plugins/FX/Niagara/Source/NiagaraEditor/Public/ViewModels/Stack/NiagaraStackGraphUtilities.h
--- a/Engine/Plugins/FX/Niagara/Source/NiagaraEditor/Public/ViewModels/Stack/NiagaraStackGraphUtilities.h
+++ b/Engine/Plugins/FX/Niagara/Source/NiagaraEditor/Public/ViewModels/Stack/NiagaraStackGraphUtilities.h
@@ -158,4 +158,8 @@ namespace FNiagaraStackGraphUtilities
 	bool GetStackIssuesRecursively(const UNiagaraStackEntry* const Entry, TArray<UNiagaraStackErrorItem*>& OutIssues);
 
 	void MoveModule(UNiagaraScript& SourceScript, UNiagaraNodeFunctionCall& ModuleToMove, UNiagaraSystem& TargetSystem, FGuid TargetEmitterHandleId, ENiagaraScriptUsage TargetUsage, FGuid TargetUsageId, int32 TargetModuleIndex);
+
+	/** Whether a parameter is allowed to be used in a certain execution category. 
+		Used to check if parameter can be dropped on a module or funciton stack entry. */
+	NIAGARAEDITOR_API bool ParameterAllowedInExecutionCategory(const FName InParameterName, const FName ExecutionCategory);
 }
diff --git a/Engine/Plugins/FX/Niagara/Source/NiagaraEditorWidgets/Private/Stack/SNiagaraStackFunctionInputValue.cpp b/Engine/Plugins/FX/Niagara/Source/NiagaraEditorWidgets/Private/Stack/SNiagaraStackFunctionInputValue.cpp
--- a/Engine/Plugins/FX/Niagara/Source/NiagaraEditorWidgets/Private/Stack/SNiagaraStackFunctionInputValue.cpp
+++ b/Engine/Plugins/FX/Niagara/Source/NiagaraEditorWidgets/Private/Stack/SNiagaraStackFunctionInputValue.cpp
@@ -11,6 +11,7 @@
 #include "NiagaraNodeCustomHlsl.h"
 #include "NiagaraActions.h"
 #include "SNiagaraParameterEditor.h"
+#include "ViewModels/Stack/NiagaraStackGraphUtilities.h"
 #include "Widgets/SBoxPanel.h"
 #include "Widgets/Layout/SBox.h"
 #include "Widgets/Input/SButton.h"
@@ -747,7 +748,7 @@ bool SNiagaraStackFunctionInputValue::OnFunctionInputAllowDrop(TSharedPtr<FDragD
 	{
 		TSharedPtr<FNiagaraStackDragOperation> InputDragDropOperation = StaticCastSharedPtr<FNiagaraStackDragOperation>(DragDropOperation);
 		TSharedPtr<FNiagaraParameterAction> Action = StaticCastSharedPtr<FNiagaraParameterAction>(InputDragDropOperation->GetAction());
-		if (Action->GetParameter().GetType() == FunctionInput->GetInputType())
+		if (Action->GetParameter().GetType() == FunctionInput->GetInputType() && FNiagaraStackGraphUtilities::ParameterAllowedInExecutionCategory(Action->GetParameter().GetName(), FunctionInput->GetExecutionCategoryName()))
 		{
 			return true;
 		}
diff --git a/Engine/Plugins/FX/Niagara/Source/NiagaraEditorWidgets/Private/Stack/SNiagaraStackModuleItem.cpp b/Engine/Plugins/FX/Niagara/Source/NiagaraEditorWidgets/Private/Stack/SNiagaraStackModuleItem.cpp
--- a/Engine/Plugins/FX/Niagara/Source/NiagaraEditorWidgets/Private/Stack/SNiagaraStackModuleItem.cpp
+++ b/Engine/Plugins/FX/Niagara/Source/NiagaraEditorWidgets/Private/Stack/SNiagaraStackModuleItem.cpp
@@ -11,6 +11,7 @@
 #include "Stack/SNiagaraStackItemGroupAddMenu.h"
 #include "ViewModels/Stack/NiagaraStackModuleItem.h"
 #include "ViewModels/Stack/NiagaraStackViewModel.h"
+#include "ViewModels/Stack/NiagaraStackGraphUtilities.h"
 #include "Widgets/Layout/SBox.h"
 #include "Widgets/SBoxPanel.h"
 #include "Widgets/Input/SButton.h"
@@ -284,7 +285,17 @@ FReply SNiagaraStackModuleItem::OnModuleItemDrop(TSharedPtr<FDragDropOperation>
 
 bool SNiagaraStackModuleItem::OnModuleItemAllowDrop(TSharedPtr<FDragDropOperation> DragDropOperation)
 {
-	return ModuleItem->CanAddInput() && DragDropOperation->IsOfType<FNiagaraStackDragOperation>();
+	if (ModuleItem->CanAddInput() && DragDropOperation->IsOfType<FNiagaraStackDragOperation>())
+	{
+		TSharedPtr<FNiagaraStackDragOperation> InputDragDropOperation = StaticCastSharedPtr<FNiagaraStackDragOperation>(DragDropOperation);
+		TSharedPtr<FNiagaraParameterAction> Action = StaticCastSharedPtr<FNiagaraParameterAction>(InputDragDropOperation->GetAction());
+		if (FNiagaraStackGraphUtilities::ParameterAllowedInExecutionCategory(Action->GetParameter().GetName(), ModuleItem->GetExecutionCategoryName()))
+		{
+			return true;
+		}
+	}
+
+	return false;
 }
 
 EVisibility SNiagaraStackModuleItem::GetStackIssuesWarningVisibility() const
