Niagara - Fix playback issues where completed systems wouldn't simulate again until you pressed play.

diff --git a/Engine/Plugins/FX/Niagara/Source/NiagaraEditor/Private/AssetTypeActions/AssetTypeActions_NiagaraEmitter.cpp b/Engine/Plugins/FX/Niagara/Source/NiagaraEditor/Private/AssetTypeActions/AssetTypeActions_NiagaraEmitter.cpp
--- a/Engine/Plugins/FX/Niagara/Source/NiagaraEditor/Private/AssetTypeActions/AssetTypeActions_NiagaraEmitter.cpp
+++ b/Engine/Plugins/FX/Niagara/Source/NiagaraEditor/Private/AssetTypeActions/AssetTypeActions_NiagaraEmitter.cpp
@@ -86,7 +86,6 @@ void FAssetTypeActions_NiagaraEmitter::ExecuteNewNiagaraSystem(TArray<TWeakObjec
 
 	FNiagaraSystemViewModelOptions SystemOptions;
 	SystemOptions.bCanModifyEmittersFromTimeline = true;
-	SystemOptions.bUseSystemExecStateForTimelineReset = true;
 	SystemOptions.EditMode = ENiagaraSystemViewModelEditMode::SystemAsset;
 
 	TArray<UObject*> ObjectsToSync;
diff --git a/Engine/Plugins/FX/Niagara/Source/NiagaraEditor/Private/Toolkits/NiagaraSystemToolkit.cpp b/Engine/Plugins/FX/Niagara/Source/NiagaraEditor/Private/Toolkits/NiagaraSystemToolkit.cpp
--- a/Engine/Plugins/FX/Niagara/Source/NiagaraEditor/Private/Toolkits/NiagaraSystemToolkit.cpp
+++ b/Engine/Plugins/FX/Niagara/Source/NiagaraEditor/Private/Toolkits/NiagaraSystemToolkit.cpp
@@ -168,7 +168,6 @@ void FNiagaraSystemToolkit::InitializeWithSystem(const EToolkitMode::Type Mode,
 
 	FNiagaraSystemViewModelOptions SystemOptions;
 	SystemOptions.bCanModifyEmittersFromTimeline = true;
-	SystemOptions.bUseSystemExecStateForTimelineReset = true;
 	SystemOptions.EditMode = ENiagaraSystemViewModelEditMode::SystemAsset;
 	SystemOptions.OnGetSequencerAddMenuContent.BindSP(this, &FNiagaraSystemToolkit::GetSequencerAddMenuContent);
 
@@ -241,7 +240,6 @@ void FNiagaraSystemToolkit::InitializeWithEmitter(const EToolkitMode::Type Mode,
 
 	FNiagaraSystemViewModelOptions SystemOptions;
 	SystemOptions.bCanModifyEmittersFromTimeline = false;
-	SystemOptions.bUseSystemExecStateForTimelineReset = false;
 	SystemOptions.EditMode = ENiagaraSystemViewModelEditMode::EmitterAsset;
 
 	SystemViewModel = MakeShareable(new FNiagaraSystemViewModel(*System, SystemOptions));
diff --git a/Engine/Plugins/FX/Niagara/Source/NiagaraEditor/Private/ViewModels/NiagaraSystemViewModel.cpp b/Engine/Plugins/FX/Niagara/Source/NiagaraEditor/Private/ViewModels/NiagaraSystemViewModel.cpp
--- a/Engine/Plugins/FX/Niagara/Source/NiagaraEditor/Private/ViewModels/NiagaraSystemViewModel.cpp
+++ b/Engine/Plugins/FX/Niagara/Source/NiagaraEditor/Private/ViewModels/NiagaraSystemViewModel.cpp
@@ -56,7 +56,6 @@ FNiagaraSystemViewModel::FNiagaraSystemViewModel(UNiagaraSystem& InSystem, FNiag
 	, NiagaraSequence(nullptr)
 	, bSettingSequencerTimeDirectly(false)
 	, bCanModifyEmittersFromTimeline(InOptions.bCanModifyEmittersFromTimeline)
-	, bUseSystemExecStateForTimelineReset(InOptions.bUseSystemExecStateForTimelineReset)
 	, bCanAutoCompile(InOptions.bCanAutoCompile)
 	, bForceAutoCompileOnce(false)
 	, bCanSimulate(InOptions.bCanSimulate)
@@ -902,9 +901,25 @@ void FNiagaraSystemViewModel::SetupSequencer()
 
 void FNiagaraSystemViewModel::ResetSystem()
 {
-	if (Sequencer->GetPlaybackStatus() == EMovieScenePlayerStatus::Playing || EditorSettings->GetResimulateOnChangeWhilePaused() == false)
+	ResetSystemInternal(true);
+}
+
+void FNiagaraSystemViewModel::ResetSystemInternal(bool bCanResetTime)
+{
+	bool bResetAge = bCanResetTime && (Sequencer->GetPlaybackStatus() == EMovieScenePlayerStatus::Playing || EditorSettings->GetResimulateOnChangeWhilePaused() == false);
+	if (bResetAge)
 	{
-		Sequencer->SetGlobalTime(0);
+		TGuardValue<bool> Guard(bSettingSequencerTimeDirectly, true);
+		if (Sequencer->GetPlaybackStatus() == EMovieScenePlayerStatus::Playing)
+		{
+			Sequencer->SetPlaybackStatus(EMovieScenePlayerStatus::Stopped);
+			Sequencer->SetGlobalTime(0);
+			Sequencer->SetPlaybackStatus(EMovieScenePlayerStatus::Playing);
+		}
+		else
+		{
+			Sequencer->SetGlobalTime(0);
+		}
 	}
 	
 	for (TObjectIterator<UNiagaraComponent> ComponentIt; ComponentIt; ++ComponentIt)
@@ -913,6 +928,10 @@ void FNiagaraSystemViewModel::ResetSystem()
 		if (Component->GetAsset() == &System)
 		{
 			Component->ResetSystem();
+			if (bResetAge && Component->GetAgeUpdateMode() == ENiagaraAgeUpdateMode::DesiredAge)
+			{
+				Component->SetDesiredAge(0);
+			}
 		}
 	}
 
@@ -1398,59 +1417,13 @@ void FNiagaraSystemViewModel::SequencerTimeChanged()
 		// Avoid reentrancy if we're setting the time directly.
 		if (bSettingSequencerTimeDirectly == false && CurrentSequencerTime != PreviousSequencerTime)
 		{
-			bool bUpdateDesiredAge = false;
-			bool bResetSystemInstance = false;
+			// Skip the first update after going from stopped to playing or from playing to stopped because snapping in sequencer may have made
+			// the time reverse by a small amount, and sending that update to the System will reset it unnecessarily.
+			bool bStartedPlaying = CurrentStatus == EMovieScenePlayerStatus::Playing && PreviousSequencerStatus != EMovieScenePlayerStatus::Playing;
+			bool bEndedPlaying = CurrentStatus != EMovieScenePlayerStatus::Playing && PreviousSequencerStatus == EMovieScenePlayerStatus::Playing;
 
-			if (CurrentStatus == EMovieScenePlayerStatus::Playing)
-			{
-				bool bSystemIsAlive = false;
-				//float MaxEmitterTime = 0.0f;
-
-				if (bUseSystemExecStateForTimelineReset)
-				{
-					if (SystemInstance->IsComplete())
-					{
-						bSystemIsAlive = false;
-					}
-					else
-					{
-						bSystemIsAlive = true;
-					}
-				}
-				else 
-				{
-					for (const TSharedRef<FNiagaraEmitterInstance> Simulation : SystemInstance->GetEmitters())
-					{
-						if (!Simulation->IsComplete())
-						{
-							bSystemIsAlive |= true;
-						}
-					}
-				}
-
-				float PlaybackRangeEnd = NiagaraSequence->GetMovieScene()->GetTickResolution().AsSeconds(
-					NiagaraSequence->GetMovieScene()->GetPlaybackRange().GetUpperBoundValue());
-				bool bIsUpdatingOrCanSpawn = bSystemIsAlive && CurrentSequencerTime < PlaybackRangeEnd;
-				if (bIsUpdatingOrCanSpawn)
-				{
-					// Skip the first update after going from stopped to playing because snapping in sequencer may have made the time
-					// reverse by a small amount, and sending that update to the System will reset it unnecessarily.
-					bUpdateDesiredAge = PreviousSequencerStatus != EMovieScenePlayerStatus::Stopped;
-				}
-				else
-				{
-					// If there are no active particles and no more particles will be spawned reset the System so it loops.
-					bResetSystemInstance = true;
-				}
-			}
-			else
-			{
-				// If the time changed and we're not playing, or stopping playing, than the user is scrubbing, or jumping to a different time using some
-				// other means so just update the System time.  Skip the first update after going from playing to stopped because snapping in sequencer 
-				// may have made the time reverse by a small amount, and sending that update to the System will reset it unnecessarily.
-				bool bStoppedPlaying = PreviousSequencerStatus == EMovieScenePlayerStatus::Playing && CurrentStatus != EMovieScenePlayerStatus::Playing;
-				bUpdateDesiredAge = !bStoppedPlaying;
-			}
+			bool bUpdateDesiredAge = bStartedPlaying == false && bStartedPlaying == false;
+			bool bResetSystemInstance = SystemInstance->IsComplete();
 
 			if (bUpdateDesiredAge)
 			{
@@ -1466,7 +1439,9 @@ void FNiagaraSystemViewModel::SequencerTimeChanged()
 
 			if (bResetSystemInstance)
 			{
-				ResetSystem();
+				// We don't want to reset the current time if we're scrubbing.
+				bool bCanResetTime = CurrentStatus == EMovieScenePlayerStatus::Playing;
+				ResetSystemInternal(bCanResetTime);
 			}
 		}
 	}
diff --git a/Engine/Plugins/FX/Niagara/Source/NiagaraEditor/Public/ViewModels/NiagaraSystemViewModel.h b/Engine/Plugins/FX/Niagara/Source/NiagaraEditor/Public/ViewModels/NiagaraSystemViewModel.h
--- a/Engine/Plugins/FX/Niagara/Source/NiagaraEditor/Public/ViewModels/NiagaraSystemViewModel.h
+++ b/Engine/Plugins/FX/Niagara/Source/NiagaraEditor/Public/ViewModels/NiagaraSystemViewModel.h
@@ -53,9 +53,6 @@ struct FNiagaraSystemViewModelOptions
 	/** A delegate which is used to generate the content for the add menu in sequencer. */
 	FOnGetAddMenuContent OnGetSequencerAddMenuContent;
 
-	/** Whether or not we use the system's execution state to drive when we reset the timeline*/
-	bool bUseSystemExecStateForTimelineReset;
-
 	/** Whether or not the system represented by this view model can be automatically compiled.  True by default. */
 	bool bCanAutoCompile;
 
@@ -258,6 +255,8 @@ public:
 	const TArray<FNiagaraStackModuleData>& GetStackModuleDataForEmitter(TSharedRef<FNiagaraEmitterViewModel> EmitterViewModel);
 
 private:
+	/** Reset the current simulation for the system */
+	void ResetSystemInternal(bool bCanResetTime);
 
 	/** Sets up the preview component and System instance. */
 	void SetupPreviewComponentAndInstance();
@@ -406,9 +405,6 @@ private:
 	/** Whether or not the user can edit emitters from the timeline. */
 	bool bCanModifyEmittersFromTimeline;
 
-	/** Whether or not we use the system's execution state to drive when we reset the timeline*/
-	bool bUseSystemExecStateForTimelineReset;
-
 	/** Whether or not the system represented by this view model can be automatically compiled. */
 	bool bCanAutoCompile;
 
