Niagara - Prevent crash on shutdown.

diff --git a/Engine/Plugins/FX/Niagara/Source/NiagaraEditor/Private/ViewModels/Stack/NiagaraStackObject.cpp b/Engine/Plugins/FX/Niagara/Source/NiagaraEditor/Private/ViewModels/Stack/NiagaraStackObject.cpp
--- a/Engine/Plugins/FX/Niagara/Source/NiagaraEditor/Private/ViewModels/Stack/NiagaraStackObject.cpp
+++ b/Engine/Plugins/FX/Niagara/Source/NiagaraEditor/Private/ViewModels/Stack/NiagaraStackObject.cpp
@@ -65,6 +65,7 @@ void UNiagaraStackObject::FinalizeInternal()
 	if (PropertyRowGenerator.IsValid())
 	{
 		PropertyRowGenerator->OnRowsRefreshed().RemoveAll(this);
+		PropertyRowGenerator.Reset();
 	}
 	Super::FinalizeInternal();
 }
