Niagara - Fix static analysis warning for posible null dereference in UNiagaraScriptItemGroup

diff --git a/Engine/Plugins/FX/Niagara/Source/NiagaraEditor/Private/ViewModels/Stack/NiagaraStackScriptItemGroup.cpp b/Engine/Plugins/FX/Niagara/Source/NiagaraEditor/Private/ViewModels/Stack/NiagaraStackScriptItemGroup.cpp
--- a/Engine/Plugins/FX/Niagara/Source/NiagaraEditor/Private/ViewModels/Stack/NiagaraStackScriptItemGroup.cpp
+++ b/Engine/Plugins/FX/Niagara/Source/NiagaraEditor/Private/ViewModels/Stack/NiagaraStackScriptItemGroup.cpp
@@ -547,9 +547,9 @@ TOptional<UNiagaraStackEntry::FDropResult> UNiagaraStackScriptItemGroup::ChildRe
 
 TOptional<UNiagaraStackEntry::FDropResult> UNiagaraStackScriptItemGroup::ChildRequestDropInternal(const UNiagaraStackEntry& TargetChild, const TArray<UNiagaraStackEntry*>& DraggedEntries)
 {
-	if (bIsValidForOutput && DraggedEntries.Num() == 1)
+	if (bIsValidForOutput && DraggedEntries.Num() == 1 && DraggedEntries[0]->IsA<UNiagaraStackModuleItem>())
 	{
-		UNiagaraStackModuleItem* SourceModuleItem = Cast<UNiagaraStackModuleItem>(DraggedEntries[0]);
+		UNiagaraStackModuleItem* SourceModuleItem = CastChecked<UNiagaraStackModuleItem>(DraggedEntries[0]);
 		TArray<ENiagaraScriptUsage> SourceUsages = SourceModuleItem->GetModuleNode().FunctionScript->GetSupportedUsageContexts();
 		if (SourceModuleItem != nullptr &&
 			SourceModuleItem->CanMoveAndDelete() &&
