Now when errors appear in a cook for Niagara GPU shaders, we see them in the same location as the cook log

#tests n/a

diff --git a/Engine/Plugins/FX/Niagara/Source/NiagaraShader/Private/NiagaraShaderCompilationManager.cpp b/Engine/Plugins/FX/Niagara/Source/NiagaraShader/Private/NiagaraShaderCompilationManager.cpp
--- a/Engine/Plugins/FX/Niagara/Source/NiagaraShader/Private/NiagaraShaderCompilationManager.cpp
+++ b/Engine/Plugins/FX/Niagara/Source/NiagaraShader/Private/NiagaraShaderCompilationManager.cpp
@@ -245,7 +245,7 @@ void FNiagaraShaderCompilationManager::ProcessCompiledNiagaraShaderMaps(
 			// Make a copy of the array as this entry of FNiagaraShaderMap::ShaderMapsBeingCompiled will be removed below
 			TArray<FNiagaraShaderScript*> ScriptArray = *Scripts;
 			bool bSuccess = true;
-
+			
 			for (int32 JobIndex = 0; JobIndex < ResultArray.Num(); JobIndex++)
 			{
 				FShaderCompileJob& CurrentJob = *((FShaderCompileJob*)(ResultArray[JobIndex]));
@@ -265,7 +265,7 @@ void FNiagaraShaderCompilationManager::ProcessCompiledNiagaraShaderMaps(
 
 					if (CurrentJob.Output.Errors.Num())
 					{
-						UE_LOG(LogShaders, Log, TEXT("There were errors for job \"%s\""), *CurrentJob.Input.DebugGroupName)
+						UE_LOG(LogNiagaraShaderCompiler, Warning, TEXT("There were errors for job \"%s\""), *CurrentJob.Input.DebugGroupName)
 							for (const FShaderCompilerError& Error : CurrentJob.Output.Errors)
 							{
 								UE_LOG(LogShaders, Log, TEXT("Error: %s"), *Error.GetErrorString())
@@ -274,7 +274,7 @@ void FNiagaraShaderCompilationManager::ProcessCompiledNiagaraShaderMaps(
 				}
 				else
 				{
-					UE_LOG(LogShaders, Log, TEXT("There were NO errors for job \"%s\""), *CurrentJob.Input.DebugGroupName);
+					UE_LOG(LogNiagaraShaderCompiler, Log, TEXT("There were NO errors for job \"%s\""), *CurrentJob.Input.DebugGroupName);
 				}
 			}
 
