Fix for transforms in emitter scripts.
Param->Dataset bindings weren't handling structs correctly.

diff --git a/Engine/Plugins/FX/Niagara/Source/Niagara/Public/NiagaraTypes.h b/Engine/Plugins/FX/Niagara/Source/Niagara/Public/NiagaraTypes.h
--- a/Engine/Plugins/FX/Niagara/Source/Niagara/Public/NiagaraTypes.h
+++ b/Engine/Plugins/FX/Niagara/Source/Niagara/Public/NiagaraTypes.h
@@ -203,25 +203,26 @@ struct FNiagaraTypeLayoutInfo
 	}
 
 private:
-	static void GenerateLayoutInfoInternal(FNiagaraTypeLayoutInfo& Layout, const UScriptStruct* Struct)
+	static void GenerateLayoutInfoInternal(FNiagaraTypeLayoutInfo& Layout, const UScriptStruct* Struct, int32 BaseOffest = 0)
 	{
 		for (TFieldIterator<UProperty> PropertyIt(Struct, EFieldIteratorFlags::IncludeSuper); PropertyIt; ++PropertyIt)
 		{
 			UProperty* Property = *PropertyIt;
+			int32 PropOffset = BaseOffest + Property->GetOffset_ForInternal();
 			if (Property->IsA(UFloatProperty::StaticClass()))
 			{
 				Layout.FloatComponentRegisterOffsets.Add(Layout.GetNumComponents());
-				Layout.FloatComponentByteOffsets.Add(Property->GetOffset_ForInternal());
+				Layout.FloatComponentByteOffsets.Add(PropOffset);
 			}
 			else if (Property->IsA(UIntProperty::StaticClass()) || Property->IsA(UBoolProperty::StaticClass()))
 			{
 				Layout.Int32ComponentRegisterOffsets.Add(Layout.GetNumComponents());
-				Layout.Int32ComponentByteOffsets.Add(Property->GetOffset_ForInternal());
+				Layout.Int32ComponentByteOffsets.Add(PropOffset);
 			}
 			//Should be able to support double easily enough
 			else if (UStructProperty* StructProp = CastChecked<UStructProperty>(Property))
 			{
-				GenerateLayoutInfoInternal(Layout, StructProp->Struct);
+				GenerateLayoutInfoInternal(Layout, StructProp->Struct, PropOffset);
 			}
 			else
 			{
