-fresh build of hlslcc for Linux
-fixed a warning in NiagaraStackModuleItem.cpp

diff --git a/Engine/Plugins/FX/Niagara/Source/NiagaraEditor/Private/ViewModels/Stack/NiagaraStackModuleItem.cpp b/Engine/Plugins/FX/Niagara/Source/NiagaraEditor/Private/ViewModels/Stack/NiagaraStackModuleItem.cpp
--- a/Engine/Plugins/FX/Niagara/Source/NiagaraEditor/Private/ViewModels/Stack/NiagaraStackModuleItem.cpp
+++ b/Engine/Plugins/FX/Niagara/Source/NiagaraEditor/Private/ViewModels/Stack/NiagaraStackModuleItem.cpp
@@ -462,7 +462,7 @@ void UNiagaraStackModuleItem::RefreshIssues(TArray<FStackIssue>& NewIssues)
 									{
 										continue;
 									}
-									bool bDependencyMet = false;
+
 									if (CurrentNode->FunctionScript->ProvidedDependencies.Contains(Requirement.Id))
 									{
 										if (Requirement.Type == ENiagaraModuleDependencyType::PreDependency)
