Fixing nightly build issues with redundant left and right side of &&
CI Issue:
d:\build\++ue4+dev-niagara+compile\sync\engine\plugins\fx\niagara\source\niagaraeditor\private\viewmodels\niagarasystemviewmodel.cpp(1425): warning V501: There are identical sub-expressions 'bStartedPlaying == false' to the left and to the right of the '&&' operator.

#tests auto-tests pass

diff --git a/Engine/Plugins/FX/Niagara/Source/NiagaraEditor/Private/ViewModels/NiagaraSystemViewModel.cpp b/Engine/Plugins/FX/Niagara/Source/NiagaraEditor/Private/ViewModels/NiagaraSystemViewModel.cpp
--- a/Engine/Plugins/FX/Niagara/Source/NiagaraEditor/Private/ViewModels/NiagaraSystemViewModel.cpp
+++ b/Engine/Plugins/FX/Niagara/Source/NiagaraEditor/Private/ViewModels/NiagaraSystemViewModel.cpp
@@ -1422,7 +1422,7 @@ void FNiagaraSystemViewModel::SequencerTimeChanged()
 			bool bStartedPlaying = CurrentStatus == EMovieScenePlayerStatus::Playing && PreviousSequencerStatus != EMovieScenePlayerStatus::Playing;
 			bool bEndedPlaying = CurrentStatus != EMovieScenePlayerStatus::Playing && PreviousSequencerStatus == EMovieScenePlayerStatus::Playing;
 
-			bool bUpdateDesiredAge = bStartedPlaying == false && bStartedPlaying == false;
+			bool bUpdateDesiredAge = bStartedPlaying == false;
 			bool bResetSystemInstance = SystemInstance->IsComplete();
 
 			if (bUpdateDesiredAge)
