Fix for CI issue:
d:\build\++ue4+dev-niagara+compile\sync\engine\plugins\fx\niagara\source\niagaraeditor\private\viewmodels\stack\niagarastackscriptitemgroup.cpp(553): warning V595: The 'SourceModuleItem' pointer was utilized before it was verified against nullptr. Check lines: 553, 554.

#tests auto-tests pass

diff --git a/Engine/Plugins/FX/Niagara/Source/NiagaraEditor/Private/ViewModels/Stack/NiagaraStackScriptItemGroup.cpp b/Engine/Plugins/FX/Niagara/Source/NiagaraEditor/Private/ViewModels/Stack/NiagaraStackScriptItemGroup.cpp
--- a/Engine/Plugins/FX/Niagara/Source/NiagaraEditor/Private/ViewModels/Stack/NiagaraStackScriptItemGroup.cpp
+++ b/Engine/Plugins/FX/Niagara/Source/NiagaraEditor/Private/ViewModels/Stack/NiagaraStackScriptItemGroup.cpp
@@ -550,7 +550,12 @@ TOptional<UNiagaraStackEntry::FDropResult> UNiagaraStackScriptItemGroup::ChildRe
 	if (bIsValidForOutput && DraggedEntries.Num() == 1 && DraggedEntries[0]->IsA<UNiagaraStackModuleItem>())
 	{
 		UNiagaraStackModuleItem* SourceModuleItem = CastChecked<UNiagaraStackModuleItem>(DraggedEntries[0]);
-		TArray<ENiagaraScriptUsage> SourceUsages = SourceModuleItem->GetModuleNode().FunctionScript->GetSupportedUsageContexts();
+		TArray<ENiagaraScriptUsage> SourceUsages;
+		if (SourceModuleItem != nullptr)
+		{
+			SourceUsages = SourceModuleItem->GetModuleNode().FunctionScript->GetSupportedUsageContexts();
+		}
+
 		if (SourceModuleItem != nullptr &&
 			SourceModuleItem->CanMoveAndDelete() &&
 			SourceUsages.ContainsByPredicate([this](ENiagaraScriptUsage SourceUsage) { return UNiagaraScript::IsEquivalentUsage(ScriptUsage, SourceUsage); }))
