Fixing capitalization bug with custom hlsl expressions
#tests was able to rename Emitter.age to Emitter.Age

diff --git a/Engine/Plugins/FX/Niagara/Source/NiagaraEditor/Private/NiagaraNodeCustomHlsl.cpp b/Engine/Plugins/FX/Niagara/Source/NiagaraEditor/Private/NiagaraNodeCustomHlsl.cpp
--- a/Engine/Plugins/FX/Niagara/Source/NiagaraEditor/Private/NiagaraNodeCustomHlsl.cpp
+++ b/Engine/Plugins/FX/Niagara/Source/NiagaraEditor/Private/NiagaraNodeCustomHlsl.cpp
@@ -39,7 +39,7 @@ FText UNiagaraNodeCustomHlsl::GetHlslText() const
 void UNiagaraNodeCustomHlsl::OnCustomHlslTextCommitted(const FText& InText, ETextCommit::Type InType)
 {
 	FString NewValue = InText.ToString();
-	if (NewValue != CustomHlsl)
+	if (!NewValue.Equals(CustomHlsl, ESearchCase::CaseSensitive))
 	{
 		FScopedTransaction Transaction(LOCTEXT("CustomHlslCommit", "Edited Custom Hlsl"));
 		Modify();
