Fixed code from crashing if the source emitter wasn't saved.
#tests passes bug description test

diff --git a/Engine/Plugins/FX/Niagara/Source/NiagaraEditor/Private/ViewModels/NiagaraEmitterHandleViewModel.cpp b/Engine/Plugins/FX/Niagara/Source/NiagaraEditor/Private/ViewModels/NiagaraEmitterHandleViewModel.cpp
--- a/Engine/Plugins/FX/Niagara/Source/NiagaraEditor/Private/ViewModels/NiagaraEmitterHandleViewModel.cpp
+++ b/Engine/Plugins/FX/Niagara/Source/NiagaraEditor/Private/ViewModels/NiagaraEmitterHandleViewModel.cpp
@@ -227,7 +227,11 @@ void FNiagaraEmitterHandleViewModel::OpenSourceEmitter()
 {
 	if (EmitterHandle)
 	{
-		FAssetEditorManager::Get().OpenEditorForAsset(const_cast<UNiagaraEmitter*>(EmitterHandle->GetSource()));
+		UNiagaraEmitter* EmitterSource = const_cast<UNiagaraEmitter*>(EmitterHandle->GetSource());
+		if (EmitterSource)
+		{
+			FAssetEditorManager::Get().OpenEditorForAsset(EmitterSource);
+		}
 	}
 }
 
diff --git a/Engine/Plugins/FX/Niagara/Source/NiagaraEditorWidgets/Private/SNiagaraStack.cpp b/Engine/Plugins/FX/Niagara/Source/NiagaraEditorWidgets/Private/SNiagaraStack.cpp
--- a/Engine/Plugins/FX/Niagara/Source/NiagaraEditorWidgets/Private/SNiagaraStack.cpp
+++ b/Engine/Plugins/FX/Niagara/Source/NiagaraEditorWidgets/Private/SNiagaraStack.cpp
@@ -135,7 +135,7 @@ void SNiagaraStack::Construct(const FArguments& InArgs, UNiagaraStackViewModel*
 					LOCTEXT("OpenAndFocusEmitter", "Open and focus source emitter"),
 					LOCTEXT("OpenAndFocusEmitterToolTip", "Open the source emitter in emitter mode"),
 					FSlateIcon(),
-					FUIAction(FExecuteAction::CreateSP(this, &SNiagaraStack::OpenSourceEmitter)));
+					FUIAction(FExecuteAction::CreateSP(this, &SNiagaraStack::OpenSourceEmitter), FCanExecuteAction::CreateSP(this, &SNiagaraStack::CanOpenSourceEmitter)));
 				MenuBuilder.AddMenuEntry(
 					LOCTEXT("ShowEmitterInContentBrowser", "Show in content browser"),
 					LOCTEXT("ShowEmitterInContentBrowserToolTip", "Show the emitter in this stack in the Content Browser"),
@@ -433,7 +433,28 @@ bool SNiagaraStack::IsEntryFocusedInSearch(UNiagaraStackEntry* Entry) const
 
 void SNiagaraStack::OpenSourceEmitter()
 {
-	FAssetEditorManager::Get().OpenEditorForAsset(const_cast<UNiagaraEmitter*>(StackViewModel->GetEmitterHandleViewModel()->GetEmitterHandle()->GetSource()));
+	if (StackViewModel && StackViewModel->GetEmitterHandleViewModel().IsValid() && StackViewModel->GetEmitterHandleViewModel()->GetEmitterHandle())
+	{
+		UNiagaraEmitter* Emitter = const_cast<UNiagaraEmitter*>(StackViewModel->GetEmitterHandleViewModel()->GetEmitterHandle()->GetSource());
+		if (Emitter != nullptr)
+		{
+			FAssetEditorManager::Get().OpenEditorForAsset(Emitter);
+		}
+	}
+}
+
+bool SNiagaraStack::CanOpenSourceEmitter() const
+{
+	if (StackViewModel && StackViewModel->GetEmitterHandleViewModel().IsValid() && StackViewModel->GetEmitterHandleViewModel()->GetEmitterHandle())
+	{
+		UNiagaraEmitter* Emitter = const_cast<UNiagaraEmitter*>(StackViewModel->GetEmitterHandleViewModel()->GetEmitterHandle()->GetSource());
+		if (Emitter != nullptr)
+		{
+			return true;
+		}
+	}
+
+	return false;
 }
 
 void SNiagaraStack::SetEmitterEnabled(bool bIsEnabled)
diff --git a/Engine/Plugins/FX/Niagara/Source/NiagaraEditorWidgets/Private/SNiagaraStack.h b/Engine/Plugins/FX/Niagara/Source/NiagaraEditorWidgets/Private/SNiagaraStack.h
--- a/Engine/Plugins/FX/Niagara/Source/NiagaraEditorWidgets/Private/SNiagaraStack.h
+++ b/Engine/Plugins/FX/Niagara/Source/NiagaraEditorWidgets/Private/SNiagaraStack.h
@@ -89,6 +89,7 @@ private:
 	
 	// Inline menu commands
 	void OpenSourceEmitter();
+	bool CanOpenSourceEmitter() const;
 	void SetEmitterEnabled(bool bIsEnabled);
 	bool CheckEmitterEnabledStatus(bool bIsEnabled);
 	void ShowEmitterInContentBrowser();
