Making Niagara variables safe when used in an hlsl context
#tests all auto-tests pass

diff --git a/Engine/Plugins/FX/Niagara/Source/NiagaraCore/Private/NiagaraCustomVersion.cpp b/Engine/Plugins/FX/Niagara/Source/NiagaraCore/Private/NiagaraCustomVersion.cpp
--- a/Engine/Plugins/FX/Niagara/Source/NiagaraCore/Private/NiagaraCustomVersion.cpp
+++ b/Engine/Plugins/FX/Niagara/Source/NiagaraCore/Private/NiagaraCustomVersion.cpp
@@ -8,4 +8,4 @@ const FGuid FNiagaraCustomVersion::GUID(0xFCF57AFA, 0x50764283, 0xB9A9E658, 0xFF
 // Register the custom version with core
 FCustomVersionRegistration GRegisterNiagaraCustomVersion(FNiagaraCustomVersion::GUID, FNiagaraCustomVersion::LatestVersion, TEXT("NiagaraVer"));
 
-const FGuid FNiagaraCustomVersion::LatestScriptCompileVersion(0xDF96299F, 0x970043FE, 0x9CE8ADF0, 0xD2D4E1AA);
\ No newline at end of file
+const FGuid FNiagaraCustomVersion::LatestScriptCompileVersion(0x67DE59E9, 0x4ED94B3C, 0x845B4A87, 0x1C44F338);
\ No newline at end of file
diff --git a/Engine/Plugins/FX/Niagara/Source/NiagaraEditor/Private/NiagaraEditorModule.cpp b/Engine/Plugins/FX/Niagara/Source/NiagaraEditor/Private/NiagaraEditorModule.cpp
--- a/Engine/Plugins/FX/Niagara/Source/NiagaraEditor/Private/NiagaraEditorModule.cpp
+++ b/Engine/Plugins/FX/Niagara/Source/NiagaraEditor/Private/NiagaraEditorModule.cpp
@@ -87,6 +87,7 @@
 #include "Customizations/NiagaraTypeCustomizations.h"
 #include "Customizations/NiagaraEventScriptPropertiesCustomization.h"
 #include "HAL/IConsoleManager.h"
+#include "NiagaraHlslTranslator.h"
 
 
 IMPLEMENT_MODULE( FNiagaraEditorModule, NiagaraEditor );
@@ -281,6 +282,7 @@ class FNiagaraSystemColorParameterTrackEditor : public FNiagaraSystemParameterTr
 
 void FNiagaraEditorModule::StartupModule()
 {
+	FHlslNiagaraTranslator::Init();
 	MenuExtensibilityManager = MakeShareable(new FExtensibilityManager);
 	ToolBarExtensibilityManager = MakeShareable(new FExtensibilityManager);
 
diff --git a/Engine/Plugins/FX/Niagara/Source/NiagaraEditor/Private/NiagaraHlslTranslator.cpp b/Engine/Plugins/FX/Niagara/Source/NiagaraEditor/Private/NiagaraHlslTranslator.cpp
--- a/Engine/Plugins/FX/Niagara/Source/NiagaraEditor/Private/NiagaraHlslTranslator.cpp
+++ b/Engine/Plugins/FX/Niagara/Source/NiagaraEditor/Private/NiagaraHlslTranslator.cpp
@@ -564,10 +564,11 @@ FString FHlslNiagaraTranslator::BuildParameterMapHlslDefinitions(TArray<FNiagara
 		}
 
 		TArray<FString> StructNameArray;
-		int32 NumFound = Variable.GetName().ToString().ParseIntoArray(StructNameArray, TEXT("."));
+		FString SanitizedVarName = GetSanitizedSymbolName(Variable.GetName().ToString());
+		int32 NumFound = SanitizedVarName.ParseIntoArray(StructNameArray, TEXT("."));
 		if (NumFound == 1) // Meaning no split above
 		{
-			Error(FText::Format(LOCTEXT("OnlyOneNamespaceEntry", "Only one namespace entry found for: {0}"), FText::FromName(Variable.GetName())), nullptr, nullptr);
+			Error(FText::Format(LOCTEXT("OnlyOneNamespaceEntry", "Only one namespace entry found for: {0}"), FText::FromString(SanitizedVarName)), nullptr, nullptr);
 		}
 		else if (NumFound > 1)
 		{
@@ -1173,7 +1174,7 @@ const FNiagaraTranslateResults &FHlslNiagaraTranslator::Translate(const FNiagara
 		{
 			if (FNiagaraParameterMapHistory::IsAttribute(Var))
 			{
-				FNiagaraVariable BasicAttribVar = FNiagaraParameterMapHistory::ResolveAsBasicAttribute(Var);
+				FNiagaraVariable BasicAttribVar = FNiagaraParameterMapHistory::ResolveAsBasicAttribute(Var, false);
 				CompilationOutput.ScriptData.Attributes.AddUnique(BasicAttribVar);
 			}
 			else
@@ -2063,15 +2064,432 @@ void FHlslNiagaraTranslator::DecomposeVariableAccess(UStruct* Struct, bool bRead
 	}
 };
 
+TMap<FString, FString> FHlslNiagaraTranslator::ReplacementsForInvalid;
+
+void FHlslNiagaraTranslator::Init()
+{
+	if (ReplacementsForInvalid.Num() == 0)
+	{
+		// hlsl keywords
+		ReplacementsForInvalid.Add(TEXT("AppendStructuredBuffer"), TEXT("APPENDSTRUCTUREDBUFFER_VAR"));
+		ReplacementsForInvalid.Add(TEXT("asm"), TEXT("ASSEMBLY_VAR"));
+		ReplacementsForInvalid.Add(TEXT("asm_fragment"), TEXT("ASSEMBLY_FRAGMENT_VAR"));
+		ReplacementsForInvalid.Add(TEXT("BlendState"), TEXT("BLEND_STATE_VAR"));
+		ReplacementsForInvalid.Add(TEXT("bool"), TEXT("BOOL_VAR"));
+		ReplacementsForInvalid.Add(TEXT("bool1"), TEXT("BOOL1_VAR"));
+		ReplacementsForInvalid.Add(TEXT("bool2"), TEXT("BOOL2_VAR"));
+		ReplacementsForInvalid.Add(TEXT("bool3"), TEXT("BOOL3_VAR"));
+		ReplacementsForInvalid.Add(TEXT("bool4"), TEXT("BOOL4_VAR"));
+		ReplacementsForInvalid.Add(TEXT("break"), TEXT("BREAK_VAR"));
+		ReplacementsForInvalid.Add(TEXT("Buffer"), TEXT("BUFFER_VAR"));
+		ReplacementsForInvalid.Add(TEXT("ByteAddressBuffer"), TEXT("BYTE_ADDRESS_BUFFER_VAR"));
+		ReplacementsForInvalid.Add(TEXT("case"), TEXT("CASE_VAR"));
+		ReplacementsForInvalid.Add(TEXT("cbuffer"), TEXT("CBUFFER_VAR"));
+		ReplacementsForInvalid.Add(TEXT("centroid"), TEXT("CENTROID_VAR"));
+		ReplacementsForInvalid.Add(TEXT("class"), TEXT("CLASS_VAR"));
+		ReplacementsForInvalid.Add(TEXT("column_major"), TEXT("COLUMN_MAJOR_VAR"));
+		ReplacementsForInvalid.Add(TEXT("compile"), TEXT("COMPILE_VAR"));
+		ReplacementsForInvalid.Add(TEXT("compile_fragment"), TEXT("COMPILE_FRAGMENT_VAR"));
+		ReplacementsForInvalid.Add(TEXT("CompileShader"), TEXT("COMPILE_SHADER_VAR"));
+		ReplacementsForInvalid.Add(TEXT("const"), TEXT("CONST_VAR"));
+		ReplacementsForInvalid.Add(TEXT("continue"), TEXT("CONTINUE_VAR"));
+		ReplacementsForInvalid.Add(TEXT("ComputeShader"), TEXT("COMPUTE_SHADER_VAR"));
+		ReplacementsForInvalid.Add(TEXT("ConsumeStructuredBuffer"), TEXT("CONSUME_STRUCTURED_BUFFER"));
+		ReplacementsForInvalid.Add(TEXT("default"), TEXT("DEFAULT_VAR"));
+		ReplacementsForInvalid.Add(TEXT("DepthStencilState"), TEXT("DEPTHSTENCILSTATE_VAR"));
+		ReplacementsForInvalid.Add(TEXT("DepthStencilView"), TEXT("DEPTHSTENCILVIEW_VAR"));
+		ReplacementsForInvalid.Add(TEXT("discard"), TEXT("DISCARD_VAR"));
+		ReplacementsForInvalid.Add(TEXT("do"), TEXT("DO_VAR"));
+		ReplacementsForInvalid.Add(TEXT("double"), TEXT("DOUBLE_VAR"));
+		ReplacementsForInvalid.Add(TEXT("double1"), TEXT("DOUBLE1_VAR"));
+		ReplacementsForInvalid.Add(TEXT("double2"), TEXT("DOUBLE2_VAR"));
+		ReplacementsForInvalid.Add(TEXT("double3"), TEXT("DOUBLE3_VAR"));
+		ReplacementsForInvalid.Add(TEXT("double4"), TEXT("DOUBLE4_VAR"));
+		ReplacementsForInvalid.Add(TEXT("DomainShader"), TEXT("DOMAINSHADER_VAR"));
+		ReplacementsForInvalid.Add(TEXT("dword"), TEXT("DWORD_VAR"));
+		ReplacementsForInvalid.Add(TEXT("else"), TEXT("ELSE_VAR"));
+		ReplacementsForInvalid.Add(TEXT("export"), TEXT("EXPORT_VAR"));
+		ReplacementsForInvalid.Add(TEXT("extern"), TEXT("EXTERN_VAR"));
+		ReplacementsForInvalid.Add(TEXT("false"), TEXT("FALSE_VAR"));
+		ReplacementsForInvalid.Add(TEXT("float"), TEXT("FLOAT_VAR"));
+		ReplacementsForInvalid.Add(TEXT("float1"), TEXT("FLOAT1_VAR"));
+		ReplacementsForInvalid.Add(TEXT("float2"), TEXT("FLOAT2_VAR"));
+		ReplacementsForInvalid.Add(TEXT("float3"), TEXT("FLOAT3_VAR"));
+		ReplacementsForInvalid.Add(TEXT("float4"), TEXT("FLOAT4_VAR"));
+		ReplacementsForInvalid.Add(TEXT("for"), TEXT("FOR_VAR"));
+		ReplacementsForInvalid.Add(TEXT("fxgroup"), TEXT("FXGROUP_VAR"));
+		ReplacementsForInvalid.Add(TEXT("GeometryShader"), TEXT("GEOMETRYSHADER_VAR"));
+		ReplacementsForInvalid.Add(TEXT("groupshared"), TEXT("GROUPSHARED_VAR"));
+		ReplacementsForInvalid.Add(TEXT("half"), TEXT("HALF_VAR"));
+		ReplacementsForInvalid.Add(TEXT("half1"), TEXT("HALF1_VAR"));
+		ReplacementsForInvalid.Add(TEXT("half2"), TEXT("HALF2_VAR"));
+		ReplacementsForInvalid.Add(TEXT("half3"), TEXT("HALF3_VAR"));
+		ReplacementsForInvalid.Add(TEXT("half4"), TEXT("HALF4_VAR"));
+		ReplacementsForInvalid.Add(TEXT("Hullshader"), TEXT("HULLSHADER_VAR"));
+		ReplacementsForInvalid.Add(TEXT("if"), TEXT("IF_VAR"));
+		ReplacementsForInvalid.Add(TEXT("in"), TEXT("IN_VAR"));
+		ReplacementsForInvalid.Add(TEXT("inline"), TEXT("INLINE_VAR"));
+		ReplacementsForInvalid.Add(TEXT("inout"), TEXT("INOUT_VAR"));
+		ReplacementsForInvalid.Add(TEXT("InputPatch"), TEXT("INPUTPATCH_VAR"));
+		ReplacementsForInvalid.Add(TEXT("int"), TEXT("INT_VAR"));
+		ReplacementsForInvalid.Add(TEXT("int1"), TEXT("INT1_VAR"));
+		ReplacementsForInvalid.Add(TEXT("int2"), TEXT("INT2_VAR"));
+		ReplacementsForInvalid.Add(TEXT("int3"), TEXT("INT3_VAR"));
+		ReplacementsForInvalid.Add(TEXT("int4"), TEXT("INT4_VAR"));
+		ReplacementsForInvalid.Add(TEXT("interface"), TEXT("INTERFACE_VAR"));
+		ReplacementsForInvalid.Add(TEXT("line"), TEXT("LINE_VAR"));
+		ReplacementsForInvalid.Add(TEXT("lineadj"), TEXT("LINEADJ_VAR"));
+		ReplacementsForInvalid.Add(TEXT("linear"), TEXT("LINEAR_VAR"));
+		ReplacementsForInvalid.Add(TEXT("LineStream"), TEXT("LINESTREAM_VAR"));
+		ReplacementsForInvalid.Add(TEXT("matrix"), TEXT("MATRIX_VAR"));
+		ReplacementsForInvalid.Add(TEXT("min16float"), TEXT("MIN16FLOAT_VAR"));
+		ReplacementsForInvalid.Add(TEXT("min10float"), TEXT("MIN10FLOAT_VAR"));
+		ReplacementsForInvalid.Add(TEXT("min16int"), TEXT("MIN16INT_VAR"));
+		ReplacementsForInvalid.Add(TEXT("min12int"), TEXT("MIN12INT_VAR"));
+		ReplacementsForInvalid.Add(TEXT("min16uint"), TEXT("MIN16UINT_VAR"));
+		ReplacementsForInvalid.Add(TEXT("namespace"), TEXT("NAMESPACE_VAR"));
+		ReplacementsForInvalid.Add(TEXT("nointerpolation"), TEXT("NOINTERPOLATION_VAR"));
+		ReplacementsForInvalid.Add(TEXT("noperspective"), TEXT("NOPERSPECTIVE_VAR"));
+		ReplacementsForInvalid.Add(TEXT("NULL"), TEXT("NULL_VAR"));
+		ReplacementsForInvalid.Add(TEXT("out"), TEXT("OUT_VAR"));
+		ReplacementsForInvalid.Add(TEXT("OutputPatch"), TEXT("OUTPUTPATCH_VAR"));
+		ReplacementsForInvalid.Add(TEXT("packoffset"), TEXT("PACKOFFSET,_VAR"));
+		ReplacementsForInvalid.Add(TEXT("pass"), TEXT("PASS,_VAR"));
+		ReplacementsForInvalid.Add(TEXT("pixelfragment"), TEXT("PIXELFRAGMENT_VAR"));
+		ReplacementsForInvalid.Add(TEXT("PixelShader"), TEXT("PIXELSHADER_VAR"));
+		ReplacementsForInvalid.Add(TEXT("point"), TEXT("POINT_VAR"));
+		ReplacementsForInvalid.Add(TEXT("PointStream"), TEXT("POINTSTREAM_VAR"));
+		ReplacementsForInvalid.Add(TEXT("precise"), TEXT("PRECISE_VAR"));
+		ReplacementsForInvalid.Add(TEXT("RasterizerState"), TEXT("RASTERIZERSTATE_VAR"));
+		ReplacementsForInvalid.Add(TEXT("RenderTargetView"), TEXT("RENDERTARGETVIEW_VAR"));
+		ReplacementsForInvalid.Add(TEXT("return"), TEXT("RETURN_VAR"));
+		ReplacementsForInvalid.Add(TEXT("register"), TEXT("REGISTER_VAR"));
+		ReplacementsForInvalid.Add(TEXT("row_major"), TEXT("ROW_MAJOR_VAR"));
+		ReplacementsForInvalid.Add(TEXT("RWBuffer"), TEXT("RWBUFFER_VAR"));
+		ReplacementsForInvalid.Add(TEXT("RWByteAddressBuffer"), TEXT("RWBYTEADDRESSBUFFER_VAR"));
+		ReplacementsForInvalid.Add(TEXT("RWStructuredBuffer"), TEXT("RWSTRUCTUREDBUFFER_VAR"));
+		ReplacementsForInvalid.Add(TEXT("RWTexture1D"), TEXT("RWTEXTURE1D_VAR"));
+		ReplacementsForInvalid.Add(TEXT("RWTexture1DArray"), TEXT("RWTEXTURE1DARRAY_VAR"));
+		ReplacementsForInvalid.Add(TEXT("RWTexture2D"), TEXT("RWTEXTURE2D_VAR"));
+		ReplacementsForInvalid.Add(TEXT("RWTexture2DArray"), TEXT("RWTEXTURE2DARRAY_VAR"));
+		ReplacementsForInvalid.Add(TEXT("RWTexture3D"), TEXT("RWTEXTURE3D_VAR"));
+		ReplacementsForInvalid.Add(TEXT("sample"), TEXT("SAMPLE_VAR"));
+		ReplacementsForInvalid.Add(TEXT("sampler"), TEXT("SAMPLER_VAR"));
+		ReplacementsForInvalid.Add(TEXT("SamplerState"), TEXT("SAMPLERSTATE_VAR"));
+		ReplacementsForInvalid.Add(TEXT("SamplerComparisonState"), TEXT("SAMPLERCOMPARISONSTATE_VAR"));
+		ReplacementsForInvalid.Add(TEXT("shared"), TEXT("SHARED_VAR"));
+		ReplacementsForInvalid.Add(TEXT("snorm"), TEXT("SNORM_VAR"));
+		ReplacementsForInvalid.Add(TEXT("stateblock"), TEXT("STATEBLOCK_VAR"));
+		ReplacementsForInvalid.Add(TEXT("stateblock_state"), TEXT("STATEBLOCK_STATE_VAR"));
+		ReplacementsForInvalid.Add(TEXT("static"), TEXT("STATIC_VAR"));
+		ReplacementsForInvalid.Add(TEXT("string"), TEXT("STRING_VAR"));
+		ReplacementsForInvalid.Add(TEXT("struct"), TEXT("STRUCT_VAR"));
+		ReplacementsForInvalid.Add(TEXT("switch"), TEXT("SWITCH_VAR"));
+		ReplacementsForInvalid.Add(TEXT("StructuredBuffer"), TEXT("STRUCTUREDBUFFER_VAR"));
+		ReplacementsForInvalid.Add(TEXT("tbuffer"), TEXT("TBUFFER_VAR"));
+		ReplacementsForInvalid.Add(TEXT("technique"), TEXT("TECHNIQUE_VAR"));
+		ReplacementsForInvalid.Add(TEXT("technique10"), TEXT("TECHNIQUE10_VAR"));
+		ReplacementsForInvalid.Add(TEXT("technique11"), TEXT("TECHNIQUE11_VAR"));
+		ReplacementsForInvalid.Add(TEXT("texture"), TEXT("TEXTURE_VAR"));
+		ReplacementsForInvalid.Add(TEXT("Texture1D"), TEXT("TEXTURE1D_VAR"));
+		ReplacementsForInvalid.Add(TEXT("Texture1DArray"), TEXT("TEXTURE1DARRAY_VAR"));
+		ReplacementsForInvalid.Add(TEXT("Texture2D"), TEXT("TEXTURE2D_VAR"));
+		ReplacementsForInvalid.Add(TEXT("Texture2DArray"), TEXT("TEXTURE2DARRAY_VAR"));
+		ReplacementsForInvalid.Add(TEXT("Texture2DMS"), TEXT("TEXTURE2DMS_VAR"));
+		ReplacementsForInvalid.Add(TEXT("Texture2DMSArray"), TEXT("TEXTURE2DMSARRAY_VAR"));
+		ReplacementsForInvalid.Add(TEXT("Texture3D"), TEXT("TEXTURE3D_VAR"));
+		ReplacementsForInvalid.Add(TEXT("TextureCube"), TEXT("TEXTURECUBE_VAR"));
+		ReplacementsForInvalid.Add(TEXT("TextureCubeArray"), TEXT("TEXTURECUBEARRAY_VAR"));
+		ReplacementsForInvalid.Add(TEXT("true"), TEXT("TRUE_VAR"));
+		ReplacementsForInvalid.Add(TEXT("typedef"), TEXT("TYPEDEF_VAR"));
+		ReplacementsForInvalid.Add(TEXT("triangle"), TEXT("TRIANGLE_VAR"));
+		ReplacementsForInvalid.Add(TEXT("triangleadj"), TEXT("TRIANGLEADJ_VAR"));
+		ReplacementsForInvalid.Add(TEXT("TriangleStream"), TEXT("TRIANGLESTREAM_VAR"));
+		ReplacementsForInvalid.Add(TEXT("uint"), TEXT("UINT_VAR"));
+		ReplacementsForInvalid.Add(TEXT("uint1"), TEXT("UINT1_VAR"));
+		ReplacementsForInvalid.Add(TEXT("uint2"), TEXT("UINT2_VAR"));
+		ReplacementsForInvalid.Add(TEXT("uint3"), TEXT("UINT3_VAR"));
+		ReplacementsForInvalid.Add(TEXT("uint4"), TEXT("UINT4_VAR"));
+		ReplacementsForInvalid.Add(TEXT("uniform"), TEXT("UNIFORM_VAR"));
+		ReplacementsForInvalid.Add(TEXT("unorm"), TEXT("UNORM_VAR"));
+		ReplacementsForInvalid.Add(TEXT("unsigned"), TEXT("UNSIGNED_VAR"));
+		ReplacementsForInvalid.Add(TEXT("vector"), TEXT("VECTOR_VAR"));
+		ReplacementsForInvalid.Add(TEXT("vertexfragment"), TEXT("VERTEXFRAGMENT_VAR"));
+		ReplacementsForInvalid.Add(TEXT("VertexShader"), TEXT("VERTEXSHADER_VAR"));
+		ReplacementsForInvalid.Add(TEXT("void"), TEXT("VOID_VAR"));
+		ReplacementsForInvalid.Add(TEXT("volatile"), TEXT("VOLATILE_VAR"));
+		ReplacementsForInvalid.Add(TEXT("while"), TEXT("WHILE_VAR"));
+
+		// additional glsl keywords
+		ReplacementsForInvalid.Add(TEXT("attribute"), TEXT("ATTRIBUTE_VAR"));
+		ReplacementsForInvalid.Add(TEXT("varying"), TEXT("VARYING_VAR"));
+		ReplacementsForInvalid.Add(TEXT("layout"), TEXT("LAYOUT_VAR"));
+		ReplacementsForInvalid.Add(TEXT("mat2"), TEXT("MAT2_VAR"));
+		ReplacementsForInvalid.Add(TEXT("mat3"), TEXT("MAT3_VAR"));
+		ReplacementsForInvalid.Add(TEXT("mat4"), TEXT("MAT4_VAR"));
+		ReplacementsForInvalid.Add(TEXT("mat2x2"), TEXT("MAT2X2_VAR"));
+		ReplacementsForInvalid.Add(TEXT("mat2x3"), TEXT("MAT2X3_VAR"));
+		ReplacementsForInvalid.Add(TEXT("mat2x4"), TEXT("MAT2X4_VAR"));
+		ReplacementsForInvalid.Add(TEXT("mat3x2"), TEXT("MAT3X2_VAR"));
+		ReplacementsForInvalid.Add(TEXT("mat3x3"), TEXT("MAT3X3_VAR"));
+		ReplacementsForInvalid.Add(TEXT("mat3x4"), TEXT("MAT3X4_VAR"));
+		ReplacementsForInvalid.Add(TEXT("mat4x2"), TEXT("MAT4X2_VAR"));
+		ReplacementsForInvalid.Add(TEXT("mat4x3"), TEXT("MAT4X3_VAR"));
+		ReplacementsForInvalid.Add(TEXT("mat4x4"), TEXT("MAT4X4_VAR"));
+		ReplacementsForInvalid.Add(TEXT("vec2"), TEXT("VEC2_VAR"));
+		ReplacementsForInvalid.Add(TEXT("vec3"), TEXT("VEC3_VAR"));
+		ReplacementsForInvalid.Add(TEXT("vec4"), TEXT("VEC4_VAR"));
+		ReplacementsForInvalid.Add(TEXT("ivec2"), TEXT("IVEC2_VAR"));
+		ReplacementsForInvalid.Add(TEXT("ivec3"), TEXT("IVEC3_VAR"));
+		ReplacementsForInvalid.Add(TEXT("ivec4"), TEXT("IVEC4_VAR"));
+		ReplacementsForInvalid.Add(TEXT("bvec2"), TEXT("BVEC2_VAR"));
+		ReplacementsForInvalid.Add(TEXT("bvec3"), TEXT("BVEC3_VAR"));
+		ReplacementsForInvalid.Add(TEXT("bvec4"), TEXT("BVEC4_VAR"));
+		ReplacementsForInvalid.Add(TEXT("uvec2"), TEXT("UVEC2_VAR"));
+		ReplacementsForInvalid.Add(TEXT("uvec3"), TEXT("UVEC3_VAR"));
+		ReplacementsForInvalid.Add(TEXT("uvec4"), TEXT("UVEC4_VAR"));
+		ReplacementsForInvalid.Add(TEXT("lowp"), TEXT("LOWP_VAR"));
+		ReplacementsForInvalid.Add(TEXT("mediump"), TEXT("MEDIUMP_VAR"));
+		ReplacementsForInvalid.Add(TEXT("highp"), TEXT("HIGHP_VAR"));
+		ReplacementsForInvalid.Add(TEXT("precision"), TEXT("PRECISION_VAR"));
+		ReplacementsForInvalid.Add(TEXT("sampler1D"), TEXT("SAMPLER1D_VAR"));
+		ReplacementsForInvalid.Add(TEXT("sampler2D"), TEXT("SAMPLER2D_VAR"));
+		ReplacementsForInvalid.Add(TEXT("sampler3D"), TEXT("SAMPLER3D_VAR"));
+		ReplacementsForInvalid.Add(TEXT("samplerCube"), TEXT("SAMPLERCUBE_VAR"));
+		ReplacementsForInvalid.Add(TEXT("sampler1DShadow"), TEXT("SAMPLER1DSHADOW_VAR"));
+		ReplacementsForInvalid.Add(TEXT("sampler2DShadow"), TEXT("SAMPLER2DSHADOW_VAR"));
+		ReplacementsForInvalid.Add(TEXT("samplerCubeShadow"), TEXT("SAMPLERCUBESHADOW_VAR"));
+		ReplacementsForInvalid.Add(TEXT("sampler1DArray"), TEXT("SAMPLER1DARRAY_VAR"));
+		ReplacementsForInvalid.Add(TEXT("sampler2DArray"), TEXT("SAMPLER2DARRAY_VAR"));
+		ReplacementsForInvalid.Add(TEXT("sampler1DArrayShadow"), TEXT("SAMPLER1DARRAYSHADOW_VAR"));
+		ReplacementsForInvalid.Add(TEXT("sampler2DArrayShadow"), TEXT("SAMPLER2DARRAYSHADOW_VAR"));
+		ReplacementsForInvalid.Add(TEXT("isampler1D"), TEXT("ISAMPLER1D_VAR"));
+		ReplacementsForInvalid.Add(TEXT("isampler2D"), TEXT("ISAMPLER2D_VAR"));
+		ReplacementsForInvalid.Add(TEXT("isampler3D"), TEXT("ISAMPLER3D_VAR"));
+		ReplacementsForInvalid.Add(TEXT("isamplerCube"), TEXT("ISAMPLERCUBE_VAR"));
+		ReplacementsForInvalid.Add(TEXT("isampler1DArray"), TEXT("ISAMPLER1DARRAY_VAR"));
+		ReplacementsForInvalid.Add(TEXT("isampler2DArray"), TEXT("ISAMPLER2DARRAY_VAR"));
+		ReplacementsForInvalid.Add(TEXT("usampler1D"), TEXT("USAMPLER1D_VAR"));
+		ReplacementsForInvalid.Add(TEXT("usampler2D"), TEXT("USAMPLER2D_VAR"));
+		ReplacementsForInvalid.Add(TEXT("usampler3D"), TEXT("USAMPLER3D_VAR"));
+		ReplacementsForInvalid.Add(TEXT("usamplerCube"), TEXT("USAMPLERCUBE_VAR"));
+		ReplacementsForInvalid.Add(TEXT("usampler1DArray"), TEXT("USAMPLER1DARRAY_VAR"));
+		ReplacementsForInvalid.Add(TEXT("usampler2DArray"), TEXT("USAMPLER2DARRAY_VAR"));
+		ReplacementsForInvalid.Add(TEXT("sampler2DRect"), TEXT("SAMPLER2DRECT_VAR"));
+		ReplacementsForInvalid.Add(TEXT("sampler2DRectShadow"), TEXT("SAMPLER2DRECTSHADOW_VAR"));
+		ReplacementsForInvalid.Add(TEXT("isampler2DRect"), TEXT("ISAMPLER2DRECT_VAR"));
+		ReplacementsForInvalid.Add(TEXT("usampler2DRect"), TEXT("USAMPLER2DRECT_VAR"));
+		ReplacementsForInvalid.Add(TEXT("samplerBuffer"), TEXT("SAMPLERBUFFER_VAR"));
+		ReplacementsForInvalid.Add(TEXT("isamplerBuffer"), TEXT("ISAMPLERBUFFER_VAR"));
+		ReplacementsForInvalid.Add(TEXT("usamplerBuffer"), TEXT("USAMPLERBUFFER_VAR"));
+		ReplacementsForInvalid.Add(TEXT("struct"), TEXT("STRUCT_VAR"));
+		ReplacementsForInvalid.Add(TEXT("flat"), TEXT("FLAT_VAR"));
+		ReplacementsForInvalid.Add(TEXT("smooth"), TEXT("SMOOTH_VAR"));
+		ReplacementsForInvalid.Add(TEXT("invariant"), TEXT("INVARIANT_VAR"));
+		ReplacementsForInvalid.Add(TEXT("union"), TEXT("UNION_VAR"));
+		ReplacementsForInvalid.Add(TEXT("enum"), TEXT("ENUM_VAR"));
+		ReplacementsForInvalid.Add(TEXT("typedef"), TEXT("TYPEDEF_VAR"));
+		ReplacementsForInvalid.Add(TEXT("template"), TEXT("TEMPLATE_VAR"));
+		ReplacementsForInvalid.Add(TEXT("this"), TEXT("THIS_VAR"));
+		ReplacementsForInvalid.Add(TEXT("packed"), TEXT("PACKED_VAR"));
+		ReplacementsForInvalid.Add(TEXT("goto"), TEXT("GOTO_VAR"));
+		ReplacementsForInvalid.Add(TEXT("noinline"), TEXT("NOINLINE_VAR"));
+		ReplacementsForInvalid.Add(TEXT("public"), TEXT("PUBLIC_VAR"));
+		ReplacementsForInvalid.Add(TEXT("external"), TEXT("EXTERNAL_VAR"));
+		ReplacementsForInvalid.Add(TEXT("long short"), TEXT("LONG SHORT_VAR"));
+		ReplacementsForInvalid.Add(TEXT("superp"), TEXT("SUPERP_VAR"));
+		ReplacementsForInvalid.Add(TEXT("input"), TEXT("INPUT_VAR"));
+		ReplacementsForInvalid.Add(TEXT("output"), TEXT("OUTPUT_VAR"));
+		ReplacementsForInvalid.Add(TEXT("hvec2"), TEXT("HVEC2_VAR"));
+		ReplacementsForInvalid.Add(TEXT("hvec3"), TEXT("HVEC3_VAR"));
+		ReplacementsForInvalid.Add(TEXT("hvec4"), TEXT("HVEC4_VAR"));
+		ReplacementsForInvalid.Add(TEXT("dvec2"), TEXT("DVEC2_VAR"));
+		ReplacementsForInvalid.Add(TEXT("dvec3"), TEXT("DVEC3_VAR"));
+		ReplacementsForInvalid.Add(TEXT("dvec4"), TEXT("DVEC4_VAR"));
+		ReplacementsForInvalid.Add(TEXT("fvec2"), TEXT("FVEC2_VAR"));
+		ReplacementsForInvalid.Add(TEXT("fvec3"), TEXT("FVEC3_VAR"));
+		ReplacementsForInvalid.Add(TEXT("fvec4"), TEXT("FVEC4_VAR"));
+		ReplacementsForInvalid.Add(TEXT("sampler3DRect"), TEXT("SAMPLER3DRECT_VAR"));
+		ReplacementsForInvalid.Add(TEXT("filter"), TEXT("FILTER_VAR"));
+		ReplacementsForInvalid.Add(TEXT("image1D"), TEXT("IMAGE1D_VAR"));
+		ReplacementsForInvalid.Add(TEXT("image2D"), TEXT("IMAGE2D_VAR"));
+		ReplacementsForInvalid.Add(TEXT("image3D"), TEXT("IMAGE3D_VAR"));
+		ReplacementsForInvalid.Add(TEXT("imageCube"), TEXT("IMAGECUBE_VAR"));
+		ReplacementsForInvalid.Add(TEXT("iimage1D"), TEXT("IIMAGE1D_VAR"));
+		ReplacementsForInvalid.Add(TEXT("iimage2D"), TEXT("IIMAGE2D_VAR"));
+		ReplacementsForInvalid.Add(TEXT("iimage3D"), TEXT("IIMAGE3D_VAR"));
+		ReplacementsForInvalid.Add(TEXT("iimageCub"), TEXT("IIMAGECUBE_VAR"));
+		ReplacementsForInvalid.Add(TEXT("uimage1D"), TEXT("UIMAGE1D_VAR"));
+		ReplacementsForInvalid.Add(TEXT("uimage2D"), TEXT("UIMAGE2D_VAR"));
+		ReplacementsForInvalid.Add(TEXT("uimage3D"), TEXT("UIMAGE3D_VAR"));
+		ReplacementsForInvalid.Add(TEXT("uimageCube"), TEXT("UIMAGECUBE_VAR"));
+		ReplacementsForInvalid.Add(TEXT("image1DArray"), TEXT("IMAGE1DARRAY_VAR"));
+		ReplacementsForInvalid.Add(TEXT("image2DArray"), TEXT("IMAGE2DARRAY_VAR"));
+		ReplacementsForInvalid.Add(TEXT("iimage1DArray"), TEXT("IIMAGE1DARRAY_VAR"));
+		ReplacementsForInvalid.Add(TEXT("iimage2DArray"), TEXT("IIMAGE2DARRAY_VAR"));
+		ReplacementsForInvalid.Add(TEXT("uimage1DArray"), TEXT("UIMAGE1DARRAY_VAR"));
+		ReplacementsForInvalid.Add(TEXT("uimage2DArray"), TEXT("UIMAGE2DARRAY_VAR"));
+		ReplacementsForInvalid.Add(TEXT("image1DShadow"), TEXT("IMAGE1DSHADOW_VAR"));
+		ReplacementsForInvalid.Add(TEXT("image2DShadow"), TEXT("IMAGE2DSHADOW_VAR"));
+		ReplacementsForInvalid.Add(TEXT("image1DArrayShadow"), TEXT("IMAGE1DARRAYSHADOW_VAR"));
+		ReplacementsForInvalid.Add(TEXT("image2DArrayShadow"), TEXT("IMAGE2DARRAYSHADOW_VAR"));
+		ReplacementsForInvalid.Add(TEXT("imageBuffer"), TEXT("IMAGEBUFFER_VAR"));
+		ReplacementsForInvalid.Add(TEXT("iimageBuffer"), TEXT("IIMAGEBUFFER_VAR"));
+		ReplacementsForInvalid.Add(TEXT("uimageBuffer"), TEXT("UIMAGEBUFFER_VAR"));
+		ReplacementsForInvalid.Add(TEXT("sizeof"), TEXT("SIZEOF_VAR"));
+		ReplacementsForInvalid.Add(TEXT("cast"), TEXT("CAST_VAR"));
+		ReplacementsForInvalid.Add(TEXT("using"), TEXT("USING_VAR"));
+		ReplacementsForInvalid.Add(TEXT("row_major"), TEXT("ROW_MAJOR_VAR"));
+		ReplacementsForInvalid.Add(TEXT("int8_t"), TEXT("INT8_T_VAR"));
+		ReplacementsForInvalid.Add(TEXT("unsigned char"), TEXT("UNSIGNED_CHAR_VAR"));
+		ReplacementsForInvalid.Add(TEXT("uchar"), TEXT("UCHAR_VAR"));
+		ReplacementsForInvalid.Add(TEXT("uint8_t"), TEXT("UINT8_T_VAR"));
+		ReplacementsForInvalid.Add(TEXT("short"), TEXT("SHORT_VAR"));
+		ReplacementsForInvalid.Add(TEXT("int16_t"), TEXT("INT16_T_VAR"));
+		ReplacementsForInvalid.Add(TEXT("unsigned shorT"), TEXT("UNSIGNED_SHORT_VAR"));
+		ReplacementsForInvalid.Add(TEXT("ushort"), TEXT("USHORT_VAR"));
+		ReplacementsForInvalid.Add(TEXT("unit16_t"), TEXT("UNIT16_T_VAR"));
+		ReplacementsForInvalid.Add(TEXT("int32_t"), TEXT("INT32_T_VAR"));
+		ReplacementsForInvalid.Add(TEXT("unsigned int"), TEXT("UNSIGNED_INT_VAR"));
+		ReplacementsForInvalid.Add(TEXT("uint32_t"), TEXT("UINT32_T_VAR"));
+		ReplacementsForInvalid.Add(TEXT("smooth"), TEXT("SMOOTH_VAR"));
+		ReplacementsForInvalid.Add(TEXT("invariant"), TEXT("INVARIANT_VAR"));
+		ReplacementsForInvalid.Add(TEXT("union"), TEXT("UNION_VAR"));
+		ReplacementsForInvalid.Add(TEXT("enum"), TEXT("ENUM_VAR"));
+		ReplacementsForInvalid.Add(TEXT("typedef"), TEXT("TYPEDEF_VAR"));
+		ReplacementsForInvalid.Add(TEXT("template"), TEXT("TEMPLATE_VAR"));
+		ReplacementsForInvalid.Add(TEXT("this"), TEXT("THIS_VAR"));
+		ReplacementsForInvalid.Add(TEXT("packed"), TEXT("PACKED_VAR"));
+		ReplacementsForInvalid.Add(TEXT("goto"), TEXT("GOTO_VAR"));
+		ReplacementsForInvalid.Add(TEXT("noinline"), TEXT("NOINLINE_VAR"));
+		ReplacementsForInvalid.Add(TEXT("public"), TEXT("PUBLIC_VAR"));
+		ReplacementsForInvalid.Add(TEXT("external"), TEXT("EXTERNAL_VAR"));
+		ReplacementsForInvalid.Add(TEXT(" long short"), TEXT("LONG_SHORT_VAR"));
+		ReplacementsForInvalid.Add(TEXT("superp"), TEXT("SUPERP_VAR"));
+		ReplacementsForInvalid.Add(TEXT("input"), TEXT("INPUT_VAR"));
+		ReplacementsForInvalid.Add(TEXT("output"), TEXT("OUTPUT_VAR"));
+		ReplacementsForInvalid.Add(TEXT("hvec2"), TEXT("HVEC2_VAR"));
+		ReplacementsForInvalid.Add(TEXT("hvec3"), TEXT("HVEC3_VAR"));
+		ReplacementsForInvalid.Add(TEXT("hvec4"), TEXT("HVEC4_VAR"));
+		ReplacementsForInvalid.Add(TEXT("dvec2"), TEXT("DVEC2_VAR"));
+		ReplacementsForInvalid.Add(TEXT("dvec3"), TEXT("DVEC3_VAR"));
+		ReplacementsForInvalid.Add(TEXT("dvec4"), TEXT("DVEC4_VAR"));
+		ReplacementsForInvalid.Add(TEXT("fvec2"), TEXT("FVEC2_VAR"));
+		ReplacementsForInvalid.Add(TEXT("fvec3"), TEXT("FVEC3_VAR"));
+		ReplacementsForInvalid.Add(TEXT("fvec4"), TEXT("FVEC4_VAR"));
+		ReplacementsForInvalid.Add(TEXT("sampler3DRect"), TEXT("SAMPLER3DRECT_VAR"));
+		ReplacementsForInvalid.Add(TEXT("filter"), TEXT("FILTER_VAR"));
+		ReplacementsForInvalid.Add(TEXT("image2D"), TEXT("IMAGE2D_VAR"));
+		ReplacementsForInvalid.Add(TEXT("image3D"), TEXT("IMAGE3D_VAR"));
+		ReplacementsForInvalid.Add(TEXT("imageCube"), TEXT("IMAGECUBE_VAR"));
+		ReplacementsForInvalid.Add(TEXT("iimage1D"), TEXT("IIMAGE1D_VAR"));
+		ReplacementsForInvalid.Add(TEXT("iimage2D"), TEXT("IIMAGE2D_VAR"));
+		ReplacementsForInvalid.Add(TEXT("iimage3D"), TEXT("IIMAGE3D_VAR"));
+		ReplacementsForInvalid.Add(TEXT("iimageCube"), TEXT("IIMAGECUBE_VAR"));
+		ReplacementsForInvalid.Add(TEXT("uimage1D"), TEXT("UIMAGE1D_VAR"));
+		ReplacementsForInvalid.Add(TEXT("uimage2D"), TEXT("UIMAGE2D_VAR"));
+		ReplacementsForInvalid.Add(TEXT("uimage3D"), TEXT("UIMAGE3D_VAR"));
+		ReplacementsForInvalid.Add(TEXT("uimageCube"), TEXT("UIMAGECUBE_VAR"));
+		ReplacementsForInvalid.Add(TEXT("image1DArray"), TEXT("IMAGE1DARRAY_VAR"));
+		ReplacementsForInvalid.Add(TEXT("image2DArray"), TEXT("IMAGE2DARRAY_VAR"));
+		ReplacementsForInvalid.Add(TEXT("iimage1DArray"), TEXT("IIMAGE1DARRAY_VAR"));
+		ReplacementsForInvalid.Add(TEXT("iimage2DArray"), TEXT("IIMAGE2DARRAY_VAR"));
+		ReplacementsForInvalid.Add(TEXT("uimage1DArray"), TEXT("UIMAGE1DARRAY_VAR"));
+		ReplacementsForInvalid.Add(TEXT("uimage2DArray"), TEXT("UIMAGE2DARRAY_VAR"));
+		ReplacementsForInvalid.Add(TEXT("image1DShadow"), TEXT("IMAGE1DSHADOW_VAR"));
+		ReplacementsForInvalid.Add(TEXT("image2DShadow"), TEXT("IMAGE2DSHADOW_VAR"));
+		ReplacementsForInvalid.Add(TEXT("image1DArrayShadow"), TEXT("IMAGE1DARRAYSHADOW_VAR"));
+		ReplacementsForInvalid.Add(TEXT("image2DArrayShadow"), TEXT("IMAGE2DARRAYSHADOW_VAR"));
+		ReplacementsForInvalid.Add(TEXT("imageBuffer"), TEXT("IMAGEBUFFER_VAR"));
+		ReplacementsForInvalid.Add(TEXT("iimageBuffer"), TEXT("IIMAGEBUFFER_VAR"));
+		ReplacementsForInvalid.Add(TEXT("uimageBuffer"), TEXT("UIMAGEBUFFER_VAR"));
+		ReplacementsForInvalid.Add(TEXT("sizeof"), TEXT("SIZEOF_VAR"));
+		ReplacementsForInvalid.Add(TEXT("cast"), TEXT("CAST_VAR"));
+		ReplacementsForInvalid.Add(TEXT("using"), TEXT("USING_VAR"));
+		ReplacementsForInvalid.Add(TEXT("row_major"), TEXT("ROW_MAJOR_VAR"));
+		ReplacementsForInvalid.Add(TEXT("int8_t"), TEXT("INT8_T_VAR"));
+		ReplacementsForInvalid.Add(TEXT("unsigned char"), TEXT("UNSIGNED_CHAR_VAR"));
+		ReplacementsForInvalid.Add(TEXT("uchar"), TEXT("UCHAR_VAR"));
+		ReplacementsForInvalid.Add(TEXT("uint8_t"), TEXT("UINT8_T_VAR"));
+		ReplacementsForInvalid.Add(TEXT("short"), TEXT("SHORT_VAR"));
+		ReplacementsForInvalid.Add(TEXT("int16_t"), TEXT("INT16_T_VAR"));
+		ReplacementsForInvalid.Add(TEXT("unsigned short"), TEXT("UNSIGNED_SHORT_VAR"));
+		ReplacementsForInvalid.Add(TEXT("ushort"), TEXT("USHORT_VAR"));
+		ReplacementsForInvalid.Add(TEXT("unit16_t"), TEXT("UNIT16_T_VAR"));
+		ReplacementsForInvalid.Add(TEXT("int32_t"), TEXT("INT32_T_VAR"));
+		ReplacementsForInvalid.Add(TEXT("unsigned int"), TEXT("UNSIGNED_INT_VAR"));
+		ReplacementsForInvalid.Add(TEXT("uint32_t"), TEXT("UINT32_T_VAR"));
+	}
+}
+
 FString FHlslNiagaraTranslator::GetSanitizedSymbolName(FString SymbolName, bool bCollapsNamespaces)
 {
+	if (SymbolName.Len() == 0)
+	{
+		return SymbolName;
+	}
+
+	check(ReplacementsForInvalid.Num() != 0);
+
 	FString Ret = SymbolName;
-	Ret.ReplaceInline(TEXT(" "), TEXT(""));
+
+	// Split up into individual namespaces...
+	TArray<FString> SplitName;
+	Ret.ParseIntoArray(SplitName, TEXT("."));
+
+	// Rules for variable namespaces..
+	for (int32 i = 0; i < SplitName.Num(); i++)
+	{
+		SplitName[i].ReplaceInline(TEXT("__"), TEXT("ASC95ASC95")); // Opengl reserves "__" within a name
+		if (SplitName[i][0] >= TCHAR('0') && SplitName[i][0] <= TCHAR('9')) // Cannot start with a numeric digit
+		{
+			SplitName[i] = TEXT("INTEGER_") + SplitName[i];
+		}
+		
+		FString* FoundReplacementStr = ReplacementsForInvalid.Find(SplitName[i]); // Look for the string in the keyword protections array.
+		if (FoundReplacementStr)
+		{
+			SplitName[i] = *FoundReplacementStr;
+		}		
+
+		SplitName[i].ReplaceInline(TEXT("\t"), TEXT(""));
+		SplitName[i].ReplaceInline(TEXT(" "), TEXT(""));
+
+		// Handle internationalization of characters..
+		FString ChangedSplitName;
+		ChangedSplitName.Reserve(SplitName[i].Len()*6); // Assign room for every current char to be 'ASCXXX'
+		for (int32 j = 0; j < SplitName[i].Len(); j++)
+		{
+			if ((SplitName[i][j] >= TCHAR('0') && SplitName[i][j] <= TCHAR('9')) ||
+				(SplitName[i][j] >= TCHAR('A') && SplitName[i][j] <= TCHAR('Z')) ||
+				(SplitName[i][j] >= TCHAR('a') && SplitName[i][j] <= TCHAR('z')) ||
+				SplitName[i][j] == TCHAR('_') ||
+				SplitName[i][j] == TCHAR(' '))
+			{
+				// Do nothing.. these are valid chars..
+				ChangedSplitName.AppendChar(SplitName[i][j]);
+			}
+			else
+			{
+				// Need to replace the bad characters..
+				ChangedSplitName.Append(TEXT("ASC"));
+				ChangedSplitName.AppendInt((int32)SplitName[i][j]);
+			}
+		}
+		SplitName[i] = ChangedSplitName;
+	}
+
+	// Gather back into single string..
+	Ret = FString::Join<FString>(SplitName, TEXT("."));
+
+	/*
 	Ret.ReplaceInline(TEXT("\\"), TEXT("_"));
 	Ret.ReplaceInline(TEXT("/"), TEXT("_"));
 	Ret.ReplaceInline(TEXT(","), TEXT("_"));
 	Ret.ReplaceInline(TEXT("-"), TEXT("_"));
 	Ret.ReplaceInline(TEXT(":"), TEXT("_"));
+	Ret = Ret.ConvertTabsToSpaces(0);
+	*/
 
 	if (bCollapsNamespaces)
 	{
@@ -2238,7 +2656,7 @@ int32 FHlslNiagaraTranslator::AddUniformChunk(FString SymbolName, const FNiagara
 	return Ret;
 }
 
-int32 FHlslNiagaraTranslator::AddSourceChunk(FString SymbolName, const FNiagaraTypeDefinition& Type)
+int32 FHlslNiagaraTranslator::AddSourceChunk(FString SymbolName, const FNiagaraTypeDefinition& Type, bool bSanitize)
 {
 	int32 Ret = CodeChunks.IndexOfByPredicate(
 		[&](const FNiagaraCodeChunk& Chunk)
@@ -2251,7 +2669,7 @@ int32 FHlslNiagaraTranslator::AddSourceChunk(FString SymbolName, const FNiagaraT
 	{
 		Ret = CodeChunks.AddDefaulted();
 		FNiagaraCodeChunk& Chunk = CodeChunks[Ret];
-		Chunk.SymbolName = GetSanitizedSymbolName(SymbolName);
+		Chunk.SymbolName = bSanitize ? GetSanitizedSymbolName(SymbolName) : SymbolName;
 		Chunk.Type = Type;
 		Chunk.Mode = ENiagaraCodeChunkMode::Source;
 
@@ -2752,7 +3170,7 @@ void FHlslNiagaraTranslator::Output(UNiagaraNodeOutput* OutputNode, const TArray
 
 			if (Attributes[i].GetType() != FNiagaraTypeDefinition::GetParameterMapDef())
 			{
-				FString SymbolName = *(TEXT("Out_") + GetSanitizedSymbolName(Attributes[i].GetName().ToString()));
+				FString SymbolName = *GetSanitizedSymbolName((TEXT("Out_") + Attributes[i].GetName().ToString()));
 				ENiagaraCodeChunkMode OldMode = CurrentBodyChunkMode;
 				CurrentBodyChunkMode = ENiagaraCodeChunkMode::Body;
 				AddBodyChunk(SymbolName, TEXT("{0}"), Attributes[i].GetType(), Inputs[i], false);
@@ -4050,7 +4468,7 @@ void FHlslNiagaraTranslator::FinalResolveNamespacedTokens(const FString& Paramet
 					}
 
 					
-					Tokens[i] = ParameterMapInstanceNamespace + Var.GetName().ToString();
+					Tokens[i] = ParameterMapInstanceNamespace + GetSanitizedSymbolName(Var.GetName().ToString());
 					break;
 				}
 			}
@@ -4444,7 +4862,7 @@ void FHlslNiagaraTranslator::RegisterFunctionCall(ENiagaraScriptUsage ScriptUsag
 				if (Info.UserPtrIdx != INDEX_NONE)
 				{
 					//This interface requires per instance data via a user ptr so place the index to it at the end of the inputs.
-					Inputs.Add(AddSourceChunk(LexToString(Info.UserPtrIdx), FNiagaraTypeDefinition::GetIntDef()));
+					Inputs.Add(AddSourceChunk(LexToString(Info.UserPtrIdx), FNiagaraTypeDefinition::GetIntDef(), false));
 					OutSignature.Inputs.Add(FNiagaraVariable(FNiagaraTypeDefinition::GetIntDef(), TEXT("InstanceData")));
 				}
 			}
@@ -4652,7 +5070,7 @@ FString FHlslNiagaraTranslator::GetFunctionSignature(const FNiagaraFunctionSigna
 					SigStr += TEXT(", ");
 				}
 
-				SigStr += TEXT("out ") + FHlslNiagaraTranslator::GetStructHlslTypeName(Output.GetType()) + TEXT(" Out_") + FHlslNiagaraTranslator::GetSanitizedSymbolName(Output.GetName().ToString());
+				SigStr += TEXT("out ") + FHlslNiagaraTranslator::GetStructHlslTypeName(Output.GetType()) + TEXT(" ") + FHlslNiagaraTranslator::GetSanitizedSymbolName(TEXT("Out_") + Output.GetName().ToString());
 				++ParamIdx;
 			}
 		}
diff --git a/Engine/Plugins/FX/Niagara/Source/NiagaraEditor/Private/NiagaraHlslTranslator.h b/Engine/Plugins/FX/Niagara/Source/NiagaraEditor/Private/NiagaraHlslTranslator.h
--- a/Engine/Plugins/FX/Niagara/Source/NiagaraEditor/Private/NiagaraHlslTranslator.h
+++ b/Engine/Plugins/FX/Niagara/Source/NiagaraEditor/Private/NiagaraHlslTranslator.h
@@ -381,7 +381,7 @@ protected:
 	FString GeneratedConstantString(FVector4 Constant);
 
 	/* Add a chunk that is not written to the source, only used as a source chunk for others. */
-	int32 AddSourceChunk(FString SymbolName, const FNiagaraTypeDefinition& Type);
+	int32 AddSourceChunk(FString SymbolName, const FNiagaraTypeDefinition& Type, bool bSanitize = true);
 
 	/** Add a chunk defining a uniform value. */
 	int32 AddUniformChunk(FString SymbolName, const FNiagaraTypeDefinition& Type);
@@ -401,6 +401,8 @@ public:
 
 	FHlslNiagaraTranslator();
 	virtual ~FHlslNiagaraTranslator() {}
+
+	static void Init();
 	
 	virtual const FNiagaraTranslateResults &Translate(const FNiagaraCompileRequestData* InCompileData, const FNiagaraCompileOptions& InCompileOptions, FHlslNiagaraTranslatorOptions Options);
 	FNiagaraTranslatorOutput &GetTranslateOutput() { return CompilationOutput; }
@@ -614,4 +616,7 @@ private:
 	TArray<FNiagaraVariable> InitialNamespaceVariablesMissingDefault;
 	// Variables that need to be initialized in the body or at the end of spawn.
 	TArray<FNiagaraVariable> DeferredVariablesMissingDefault;
+
+	static TMap<FString, FString> ReplacementsForInvalid;
+
 };
diff --git a/Engine/Plugins/FX/Niagara/Source/NiagaraEditor/Private/NiagaraParameterMapHistory.cpp b/Engine/Plugins/FX/Niagara/Source/NiagaraEditor/Private/NiagaraParameterMapHistory.cpp
--- a/Engine/Plugins/FX/Niagara/Source/NiagaraEditor/Private/NiagaraParameterMapHistory.cpp
+++ b/Engine/Plugins/FX/Niagara/Source/NiagaraEditor/Private/NiagaraParameterMapHistory.cpp
@@ -250,15 +250,17 @@ FString FNiagaraParameterMapHistory::MakeSafeNamespaceString(const FString& InSt
 }
 
 
-FNiagaraVariable FNiagaraParameterMapHistory::ResolveAsBasicAttribute(const FNiagaraVariable& InVar)
+FNiagaraVariable FNiagaraParameterMapHistory::ResolveAsBasicAttribute(const FNiagaraVariable& InVar, bool bSanitizeString)
 {
 	if (IsAttribute(InVar))
 	{
 		FString ParamName = InVar.GetName().ToString();
 		ParamName.RemoveAt(0, FString(PARAM_MAP_ATTRIBUTE_STR).Len());
 
-		ParamName = MakeSafeNamespaceString(ParamName);
-
+		if (bSanitizeString)
+		{
+			ParamName = MakeSafeNamespaceString(ParamName);
+		}
 		FNiagaraVariable RetVar = InVar;
 		RetVar.SetName(*ParamName);
 		return RetVar;
@@ -269,12 +271,15 @@ FNiagaraVariable FNiagaraParameterMapHistory::ResolveAsBasicAttribute(const FNia
 	}
 }
 
-FNiagaraVariable FNiagaraParameterMapHistory::BasicAttributeToNamespacedAttribute(const FNiagaraVariable& InVar)
+FNiagaraVariable FNiagaraParameterMapHistory::BasicAttributeToNamespacedAttribute(const FNiagaraVariable& InVar, bool bSanitizeString)
 {
 	FString ParamName = InVar.GetName().ToString();
 	ParamName.InsertAt(0, FString(PARAM_MAP_ATTRIBUTE_STR));
 
-	ParamName = MakeSafeNamespaceString(ParamName);
+	if (bSanitizeString)
+	{
+		ParamName = MakeSafeNamespaceString(ParamName);
+	}
 
 	FNiagaraVariable RetVar = InVar;
 	RetVar.SetName(*ParamName);
diff --git a/Engine/Plugins/FX/Niagara/Source/NiagaraEditor/Public/NiagaraParameterMapHistory.h b/Engine/Plugins/FX/Niagara/Source/NiagaraEditor/Public/NiagaraParameterMapHistory.h
--- a/Engine/Plugins/FX/Niagara/Source/NiagaraEditor/Public/NiagaraParameterMapHistory.h
+++ b/Engine/Plugins/FX/Niagara/Source/NiagaraEditor/Public/NiagaraParameterMapHistory.h
@@ -105,13 +105,13 @@ public:
 	/**
 	* Remove the Particles namespace if it exists.
 	*/
-	static FNiagaraVariable ResolveAsBasicAttribute(const FNiagaraVariable& InVar);
+	static FNiagaraVariable ResolveAsBasicAttribute(const FNiagaraVariable& InVar, bool bSanitizeInput = true);
 
 
 	/**
 	* Reverses ResolveAsBasicAttribute.
 	*/
-	static FNiagaraVariable BasicAttributeToNamespacedAttribute(const FNiagaraVariable& InVar);
+	static FNiagaraVariable BasicAttributeToNamespacedAttribute(const FNiagaraVariable& InVar, bool bSanitizeInput = true);
 
 	
 	/** Prepends the namespace string to the variable name.*/
