Making it so that you can only AddParameters in graphs where it makes sense to do so. Promoting variables also follows the same logic.

#tests made dynamic input, function, and module graphs and made sure that each one was only capable of creating the types that were proper

diff --git a/Engine/Plugins/FX/Niagara/Source/NiagaraEditor/Private/EdGraphSchema_Niagara.cpp b/Engine/Plugins/FX/Niagara/Source/NiagaraEditor/Private/EdGraphSchema_Niagara.cpp
--- a/Engine/Plugins/FX/Niagara/Source/NiagaraEditor/Private/EdGraphSchema_Niagara.cpp
+++ b/Engine/Plugins/FX/Niagara/Source/NiagaraEditor/Private/EdGraphSchema_Niagara.cpp
@@ -182,6 +182,49 @@ TSharedPtr<FNiagaraSchemaAction_NewNode> AddNewNodeAction(TArray<TSharedPtr<FNia
 	return NewAction;
 }
 
+
+
+bool IsSystemGraph(const UNiagaraGraph* NiagaraGraph)
+{
+	TArray<UNiagaraNodeEmitter*> Emitters;
+	NiagaraGraph->GetNodesOfClass<UNiagaraNodeEmitter>(Emitters);
+	bool bSystemGraph = Emitters.Num() != 0 || NiagaraGraph->FindOutputNode(ENiagaraScriptUsage::SystemSpawnScript) != nullptr || NiagaraGraph->FindOutputNode(ENiagaraScriptUsage::SystemUpdateScript) != nullptr;
+	return bSystemGraph;
+}
+
+bool IsParticleGraph(const UNiagaraGraph* NiagaraGraph)
+{
+	bool bParticleGraph = NiagaraGraph->FindOutputNode(ENiagaraScriptUsage::ParticleSpawnScriptInterpolated) != nullptr || NiagaraGraph->FindOutputNode(ENiagaraScriptUsage::ParticleSpawnScript) != nullptr || NiagaraGraph->FindOutputNode(ENiagaraScriptUsage::ParticleUpdateScript) != nullptr;
+	return bParticleGraph;
+}
+
+bool IsModuleGraph(const UNiagaraGraph* NiagaraGraph)
+{
+	bool bModuleGraph = NiagaraGraph->FindOutputNode(ENiagaraScriptUsage::Module) != nullptr;
+	return bModuleGraph;
+}
+
+bool IsDynamicInputGraph(const UNiagaraGraph* NiagaraGraph)
+{
+	bool bDynamicInputGraph = NiagaraGraph->FindOutputNode(ENiagaraScriptUsage::DynamicInput) != nullptr;
+	return bDynamicInputGraph;
+}
+
+
+bool IsFunctionGraph(const UNiagaraGraph* NiagaraGraph)
+{
+	bool bFunctionGraph = NiagaraGraph->FindOutputNode(ENiagaraScriptUsage::Function) != nullptr;
+	return bFunctionGraph;
+}
+
+
+bool IsParticleUpdateGraph(const UNiagaraGraph* NiagaraGraph)
+{
+	bool bUpdateGraph = NiagaraGraph->FindOutputNode(ENiagaraScriptUsage::ParticleUpdateScript) != nullptr;
+	return bUpdateGraph;
+}
+
+
 const UNiagaraGraph* GetAlternateGraph(const UNiagaraGraph* NiagaraGraph)
 {
 	UNiagaraScriptSource* ScriptSource = Cast<UNiagaraScriptSource>(NiagaraGraph->GetOuter());
@@ -299,16 +342,13 @@ TArray<TSharedPtr<FNiagaraSchemaAction_NewNode> > UEdGraphSchema_Niagara::GetGra
 	TArray<TSharedPtr<FNiagaraSchemaAction_NewNode> > NewActions;
 
 	const UNiagaraGraph* NiagaraGraph = CastChecked<UNiagaraGraph>(CurrentGraph);
-	TArray<UNiagaraNodeEmitter*> Emitters;
-	NiagaraGraph->GetNodesOfClass<UNiagaraNodeEmitter>(Emitters);
-	bool bSystemGraph = Emitters.Num() != 0 || NiagaraGraph->FindOutputNode(ENiagaraScriptUsage::SystemSpawnScript) != nullptr || NiagaraGraph->FindOutputNode(ENiagaraScriptUsage::SystemUpdateScript) != nullptr;
-	bool bModuleGraph = NiagaraGraph->FindOutputNode(ENiagaraScriptUsage::Module) != nullptr;
-	bool bDynamicInputGraph = NiagaraGraph->FindOutputNode(ENiagaraScriptUsage::DynamicInput) != nullptr;
-	bool bFunctionGraph = NiagaraGraph->FindOutputNode(ENiagaraScriptUsage::Function) != nullptr;
-	bool bUpdateGraph = NiagaraGraph->FindOutputNode(ENiagaraScriptUsage::ParticleUpdateScript) != nullptr;
-
 	
-
+	bool bSystemGraph = IsSystemGraph(NiagaraGraph);
+	bool bModuleGraph = IsModuleGraph(NiagaraGraph);
+	bool bDynamicInputGraph = IsDynamicInputGraph(NiagaraGraph);
+	bool bFunctionGraph = IsFunctionGraph(NiagaraGraph);
+	bool bParticleUpdateGraph = IsParticleUpdateGraph(NiagaraGraph);
+	
 	if (GbAllowAllNiagaraNodesInEmitterGraphs || bModuleGraph || bFunctionGraph || bSystemGraph)
 	{
 		const TArray<FNiagaraOpInfo>& OpInfos = FNiagaraOpInfo::GetOpInfoArray();
@@ -356,7 +396,7 @@ TArray<TSharedPtr<FNiagaraSchemaAction_NewNode> > UEdGraphSchema_Niagara::GetGra
 	TArray<FAssetData> ScriptAssets;
 	AssetRegistryModule.Get().GetAssetsByClass(UNiagaraScript::StaticClass()->GetFName(), ScriptAssets);
 	UEnum* NiagaraScriptUsageEnum = FindObjectChecked<UEnum>(ANY_PACKAGE, TEXT("ENiagaraScriptUsage"), true);
-	if (GbAllowAllNiagaraNodesInEmitterGraphs || bModuleGraph || bFunctionGraph || bSystemGraph)
+	if (GbAllowAllNiagaraNodesInEmitterGraphs || bModuleGraph || bFunctionGraph || bDynamicInputGraph)
 	{
 		for (const FAssetData& ScriptAsset : ScriptAssets)
 		{
@@ -387,7 +427,7 @@ TArray<TSharedPtr<FNiagaraSchemaAction_NewNode> > UEdGraphSchema_Niagara::GetGra
 			int32 BitfieldValue = FCString::Atoi(*BitfieldTagValue);
 
 			ENiagaraScriptUsage TargetUsage = ENiagaraScriptUsage::Module;
-			if (bUpdateGraph)
+			if (bParticleUpdateGraph)
 			{
 				TargetUsage = ENiagaraScriptUsage::ParticleUpdateScript;
 			}
@@ -465,7 +505,7 @@ TArray<TSharedPtr<FNiagaraSchemaAction_NewNode> > UEdGraphSchema_Niagara::GetGra
 	}
 	
 	TArray<ENiagaraScriptUsage> UsageTypesToAdd;
-	if (bUpdateGraph)
+	if (bParticleUpdateGraph)
 	{
 		UsageTypesToAdd.Add(ENiagaraScriptUsage::ParticleEventScript);
 		UsageTypesToAdd.Add(ENiagaraScriptUsage::EmitterSpawnScript);
@@ -717,20 +757,23 @@ TArray<TSharedPtr<FNiagaraSchemaAction_NewNode> > UEdGraphSchema_Niagara::GetGra
 		TArray<UNiagaraNodeInput*> InputNodes;
 		NiagaraGraph->GetNodesOfClass(InputNodes);
 
-		//Emitter constants managed by the system.
-		const TArray<FNiagaraVariable>& SystemConstants = FNiagaraConstants::GetEngineConstants();
-		for (const FNiagaraVariable& SysConst : SystemConstants)
+		if (bFunctionGraph)
 		{
-			FFormatNamedArguments Args;
-			Args.Add(TEXT("Constant"), FText::FromName(SysConst.GetName()));
-			const FText MenuDesc = FText::Format(LOCTEXT("GetSystemConstant", "Get {Constant}"), Args);
+			//Emitter constants managed by the system.
+			const TArray<FNiagaraVariable>& SystemConstants = FNiagaraConstants::GetEngineConstants();
+			for (const FNiagaraVariable& SysConst : SystemConstants)
+			{
+				FFormatNamedArguments Args;
+				Args.Add(TEXT("Constant"), FText::FromName(SysConst.GetName()));
+				const FText MenuDesc = FText::Format(LOCTEXT("GetSystemConstant", "Get {Constant}"), Args);
 
-			TSharedPtr<FNiagaraSchemaAction_NewNode> GetConstAction = AddNewNodeAction(NewActions, LOCTEXT("System Parameters Menu Title", "System Parameters"), MenuDesc, SysConst.GetName(), FText::GetEmpty());
+				TSharedPtr<FNiagaraSchemaAction_NewNode> GetConstAction = AddNewNodeAction(NewActions, LOCTEXT("System Parameters Menu Title", "System Parameters"), MenuDesc, SysConst.GetName(), FText::GetEmpty());
 
-			UNiagaraNodeInput* InputNode = NewObject<UNiagaraNodeInput>(OwnerOfTemporaries);
-			InputNode->Usage = ENiagaraInputNodeUsage::SystemConstant;
-			InputNode->Input = SysConst;
-			GetConstAction->NodeTemplate = InputNode;
+				UNiagaraNodeInput* InputNode = NewObject<UNiagaraNodeInput>(OwnerOfTemporaries);
+				InputNode->Usage = ENiagaraInputNodeUsage::SystemConstant;
+				InputNode->Input = SysConst;
+				GetConstAction->NodeTemplate = InputNode;
+			}
 		}
 
 		//Emitter constants managed by the Translator.
@@ -777,6 +820,16 @@ TArray<TSharedPtr<FNiagaraSchemaAction_NewNode> > UEdGraphSchema_Niagara::GetGra
 					else
 					{
 						MenuCat = LOCTEXT("AddParameterCat", "Add Parameter");
+
+						// If you are in dynamic inputs or modules, we only allow free-range variables for 
+						// data interfaces and parameter maps.
+						if (bDynamicInputGraph || bModuleGraph)
+						{
+							if (Type != FNiagaraTypeDefinition::GetParameterMapDef())
+							{
+								continue;
+							}
+						}
 					}
 						
 					const FText MenuDesc = FText::Format(MenuDescFmt, Type.GetStruct()->GetDisplayNameText());
@@ -787,24 +840,27 @@ TArray<TSharedPtr<FNiagaraSchemaAction_NewNode> > UEdGraphSchema_Niagara::GetGra
 				}
 
 				// TODO sckime please remove this..
-				for (FNiagaraTypeDefinition Type : RegisteredTypes)
+				if (bSystemGraph || IsParticleGraph(NiagaraGraph))
 				{
-					FText MenuCat;
-					if (const UClass* Class = Type.GetClass())
-					{
-						continue;
-					}
-					else
+					for (FNiagaraTypeDefinition Type : RegisteredTypes)
 					{
-						MenuCat = LOCTEXT("AddRIParameterCat", "Add Rapid Iteration Param");
-					}
+						FText MenuCat;
+						if (const UClass* Class = Type.GetClass())
+						{
+							continue;
+						}
+						else
+						{
+							MenuCat = LOCTEXT("AddRIParameterCat", "Add Rapid Iteration Param");
+						}
 
-					const FText MenuDesc = FText::Format(MenuDescFmt, Type.GetStruct()->GetDisplayNameText());
-					TSharedPtr<FNiagaraSchemaAction_NewNode> InputAction = AddNewNodeAction(NewActions, MenuCat, MenuDesc, *MenuDesc.ToString(), FText::GetEmpty());
-					UNiagaraNodeInput* InputNode = NewObject<UNiagaraNodeInput>(OwnerOfTemporaries);
-					FNiagaraEditorUtilities::InitializeParameterInputNode(*InputNode, Type, NiagaraGraph);
-					InputNode->Usage = ENiagaraInputNodeUsage::RapidIterationParameter;
-					InputAction->NodeTemplate = InputNode;
+						const FText MenuDesc = FText::Format(MenuDescFmt, Type.GetStruct()->GetDisplayNameText());
+						TSharedPtr<FNiagaraSchemaAction_NewNode> InputAction = AddNewNodeAction(NewActions, MenuCat, MenuDesc, *MenuDesc.ToString(), FText::GetEmpty());
+						UNiagaraNodeInput* InputNode = NewObject<UNiagaraNodeInput>(OwnerOfTemporaries);
+						FNiagaraEditorUtilities::InitializeParameterInputNode(*InputNode, Type, NiagaraGraph);
+						InputNode->Usage = ENiagaraInputNodeUsage::RapidIterationParameter;
+						InputAction->NodeTemplate = InputNode;
+					}
 				}
 				
 				if (PinType != FNiagaraTypeDefinition::GetGenericNumericDef())
@@ -1498,6 +1554,16 @@ void UEdGraphSchema_Niagara::ToggleNodeEnabledState(UNiagaraNode* InNode) const
 	}
 }
 
+bool UEdGraphSchema_Niagara::CanPromoteSinglePinToParameter(const UEdGraphPin* SourcePin) 
+{
+	const UNiagaraGraph* NiagaraGraph = Cast<UNiagaraGraph>(SourcePin->GetOwningNode()->GetGraph());
+	if (IsFunctionGraph(NiagaraGraph))
+	{
+		return true;
+	}
+	return false;
+}
+
 void UEdGraphSchema_Niagara::PromoteSinglePinToParameter(UEdGraphPin* SourcePin)
 {
 	if (SourcePin)
@@ -1560,7 +1626,8 @@ void UEdGraphSchema_Niagara::GetContextMenuActions(const UEdGraph* CurrentGraph,
 			if (InGraphPin->Direction == EEdGraphPinDirection::EGPD_Input)
 			{
 				MenuBuilder->AddMenuEntry(LOCTEXT("PromoteToParameter", "Promote to Parameter"), LOCTEXT("PromoteToParameterTooltip", "Create a parameter argument and connect this pin to that parameter."), FSlateIcon(),
-					FUIAction(FExecuteAction::CreateUObject((UEdGraphSchema_Niagara*const)this, &UEdGraphSchema_Niagara::PromoteSinglePinToParameter, const_cast<UEdGraphPin*>(InGraphPin))));
+					FUIAction(FExecuteAction::CreateUObject((UEdGraphSchema_Niagara*const)this, &UEdGraphSchema_Niagara::PromoteSinglePinToParameter, const_cast<UEdGraphPin*>(InGraphPin)),
+						FCanExecuteAction::CreateStatic(&UEdGraphSchema_Niagara::CanPromoteSinglePinToParameter, InGraphPin)));
 				if (InGraphPin->LinkedTo.Num() == 0 && InGraphPin->bDefaultValueIsIgnored == false)
 				{
 					MenuBuilder->AddMenuEntry(
diff --git a/Engine/Plugins/FX/Niagara/Source/NiagaraEditor/Public/EdGraphSchema_Niagara.h b/Engine/Plugins/FX/Niagara/Source/NiagaraEditor/Public/EdGraphSchema_Niagara.h
--- a/Engine/Plugins/FX/Niagara/Source/NiagaraEditor/Public/EdGraphSchema_Niagara.h
+++ b/Engine/Plugins/FX/Niagara/Source/NiagaraEditor/Public/EdGraphSchema_Niagara.h
@@ -83,6 +83,7 @@ class NIAGARAEDITOR_API UEdGraphSchema_Niagara : public UEdGraphSchema
 
 	TArray<TSharedPtr<FNiagaraSchemaAction_NewNode> > GetGraphContextActions(const UEdGraph* CurrentGraph, TArray<UObject*>& SelectedObjects, const UEdGraphPin* FromPin, UEdGraph* OwnerOfTemporaries) const;
 	void PromoteSinglePinToParameter(UEdGraphPin* SourcePin);
+	static bool CanPromoteSinglePinToParameter(const UEdGraphPin* SourcePin);
 	void ToggleNodeEnabledState(class UNiagaraNode* InNode) const;
 
 	/** 
