Potential fix for check encountered in the wild.

#tests auto-tests pass

diff --git a/Engine/Plugins/FX/Niagara/Source/NiagaraEditor/Private/NiagaraNodeInput.cpp b/Engine/Plugins/FX/Niagara/Source/NiagaraEditor/Private/NiagaraNodeInput.cpp
--- a/Engine/Plugins/FX/Niagara/Source/NiagaraEditor/Private/NiagaraNodeInput.cpp
+++ b/Engine/Plugins/FX/Niagara/Source/NiagaraEditor/Private/NiagaraNodeInput.cpp
@@ -435,7 +435,12 @@ void UNiagaraNodeInput::Compile(class FHlslNiagaraTranslator* Translator, TArray
 	}
 
 	UNiagaraGraph* Graph = GetNiagaraGraph();
-	check(Input.GetType() != FNiagaraTypeDefinition::GetGenericNumericDef());
+	if (Input.GetType() == FNiagaraTypeDefinition::GetGenericNumericDef())
+	{
+		Outputs.Add(INDEX_NONE);
+		Translator->Error(LOCTEXT("InvalidPinType", "Numeric types should be able to be inferred from use by this phase of compilation."), this, nullptr);
+		return;
+	}
 
 	int32 FunctionParam = INDEX_NONE;
 	if (IsExposed() && Translator->GetFunctionParameter(Input, FunctionParam))
