Multiple fixes to the spline component, which hasn't received attention in a long time. It is only supported on CPU VM.

#tests new auto-test now passes

diff --git a/Engine/Plugins/FX/Niagara/Source/Niagara/Classes/NiagaraDataInterfaceSpline.h b/Engine/Plugins/FX/Niagara/Source/Niagara/Classes/NiagaraDataInterfaceSpline.h
--- a/Engine/Plugins/FX/Niagara/Source/Niagara/Classes/NiagaraDataInterfaceSpline.h
+++ b/Engine/Plugins/FX/Niagara/Source/Niagara/Classes/NiagaraDataInterfaceSpline.h
@@ -19,7 +19,7 @@ struct FNDISpline_InstanceData
 	FMatrix TransformInverseTransposed;
 };
 
-/** Data Interface allowing sampling of static meshes. */
+/** Data Interface allowing sampling of in-world spline components. Note that this data interface is very experimental. */
 UCLASS(EditInlineNew, Category = "Splines", meta = (DisplayName = "Spline"))
 class NIAGARA_API UNiagaraDataInterfaceSpline : public UNiagaraDataInterface
 {
diff --git a/Engine/Plugins/FX/Niagara/Source/Niagara/Private/NiagaraDataInterfaceSpline.cpp b/Engine/Plugins/FX/Niagara/Source/Niagara/Private/NiagaraDataInterfaceSpline.cpp
--- a/Engine/Plugins/FX/Niagara/Source/Niagara/Private/NiagaraDataInterfaceSpline.cpp
+++ b/Engine/Plugins/FX/Niagara/Source/Niagara/Private/NiagaraDataInterfaceSpline.cpp
@@ -398,22 +398,42 @@ void UNiagaraDataInterfaceSpline::SampleSplinePositionByUnitDistance(FVectorVMCo
 	FRegisterHandler<float> OutPosY(Context);
 	FRegisterHandler<float> OutPosZ(Context);
 
-	checkSlow(InstData->Component);
-
-	for (int32 i = 0; i < Context.NumInstances; ++i)
+	if (InstData->Component)
 	{
-		float DistanceUnitDistance = SplineSampleParam.Get();
 
-		FVector Pos = InstData->Component->GetLocationAtDistanceAlongSpline(DistanceUnitDistance * InstData->Component->GetSplineLength(), ESplineCoordinateSpace::Local);
-		TransformHandler.TransformPosition(Pos, InstData->Transform);
-
-		*OutPosX.GetDest() = Pos.X;
-		*OutPosY.GetDest() = Pos.Y;
-		*OutPosZ.GetDest() = Pos.Z;
-		SplineSampleParam.Advance();
-		OutPosX.Advance();
-		OutPosY.Advance();
-		OutPosZ.Advance();
+		for (int32 i = 0; i < Context.NumInstances; ++i)
+		{
+			float DistanceUnitDistance = SplineSampleParam.Get();
+
+			FVector Pos = InstData->Component->GetLocationAtDistanceAlongSpline(DistanceUnitDistance * InstData->Component->GetSplineLength(), ESplineCoordinateSpace::Local);
+			TransformHandler.TransformPosition(Pos, InstData->Transform);
+
+			*OutPosX.GetDest() = Pos.X;
+			*OutPosY.GetDest() = Pos.Y;
+			*OutPosZ.GetDest() = Pos.Z;
+			SplineSampleParam.Advance();
+			OutPosX.Advance();
+			OutPosY.Advance();
+			OutPosZ.Advance();
+		}
+	}
+	else
+	{
+		for (int32 i = 0; i < Context.NumInstances; ++i)
+		{
+			float DistanceUnitDistance = SplineSampleParam.Get();
+
+			FVector Pos = FVector(EForceInit::ForceInitToZero);
+			TransformHandler.TransformPosition(Pos, InstData->Transform);
+
+			*OutPosX.GetDest() = Pos.X;
+			*OutPosY.GetDest() = Pos.Y;
+			*OutPosZ.GetDest() = Pos.Z;
+			SplineSampleParam.Advance();
+			OutPosX.Advance();
+			OutPosY.Advance();
+			OutPosZ.Advance();
+		}
 	}
 }
 
@@ -427,22 +447,42 @@ void UNiagaraDataInterfaceSpline::SampleSplineUpVectorByUnitDistance(FVectorVMCo
 	FRegisterHandler<float> OutPosY(Context);
 	FRegisterHandler<float> OutPosZ(Context);
 
-	checkSlow(InstData->Component);
-
-	for (int32 i = 0; i < Context.NumInstances; ++i)
+	if (InstData->Component)
 	{
-		float DistanceUnitDistance = SplineSampleParam.Get();
 
-		FVector Pos = InstData->Component->GetUpVectorAtDistanceAlongSpline(DistanceUnitDistance * InstData->Component->GetSplineLength(), ESplineCoordinateSpace::Local);
-		TransformHandler.TransformPosition(Pos, InstData->Transform);
-
-		*OutPosX.GetDest() = Pos.X;
-		*OutPosY.GetDest() = Pos.Y;
-		*OutPosZ.GetDest() = Pos.Z;
-		SplineSampleParam.Advance();
-		OutPosX.Advance();
-		OutPosY.Advance();
-		OutPosZ.Advance();
+		for (int32 i = 0; i < Context.NumInstances; ++i)
+		{
+			float DistanceUnitDistance = SplineSampleParam.Get();
+
+			FVector Pos = InstData->Component->GetUpVectorAtDistanceAlongSpline(DistanceUnitDistance * InstData->Component->GetSplineLength(), ESplineCoordinateSpace::Local);
+			TransformHandler.TransformVector(Pos, InstData->Transform);
+
+			*OutPosX.GetDest() = Pos.X;
+			*OutPosY.GetDest() = Pos.Y;
+			*OutPosZ.GetDest() = Pos.Z;
+			SplineSampleParam.Advance();
+			OutPosX.Advance();
+			OutPosY.Advance();
+			OutPosZ.Advance();
+		}
+	}
+	else
+	{
+		for (int32 i = 0; i < Context.NumInstances; ++i)
+		{
+			float DistanceUnitDistance = SplineSampleParam.Get();
+
+			FVector Pos = FVector(0.0f, 0.0f, 1.0f); 
+			TransformHandler.TransformVector(Pos, InstData->Transform);
+
+			*OutPosX.GetDest() = Pos.X;
+			*OutPosY.GetDest() = Pos.Y;
+			*OutPosZ.GetDest() = Pos.Z;
+			SplineSampleParam.Advance();
+			OutPosX.Advance();
+			OutPosY.Advance();
+			OutPosZ.Advance();
+		}
 	}
 }
 
@@ -455,22 +495,41 @@ void UNiagaraDataInterfaceSpline::SampleSplineRightVectorByUnitDistance(FVectorV
 	FRegisterHandler<float> OutPosX(Context);
 	FRegisterHandler<float> OutPosY(Context);
 	FRegisterHandler<float> OutPosZ(Context);
-	checkSlow(InstData->Component);
-
-	for (int32 i = 0; i < Context.NumInstances; ++i)
+	
+	
+	if (InstData->Component)
 	{
-		float DistanceUnitDistance = SplineSampleParam.Get();
-
-		FVector Pos = InstData->Component->GetRightVectorAtDistanceAlongSpline(DistanceUnitDistance * InstData->Component->GetSplineLength(), ESplineCoordinateSpace::Local);
-		TransformHandler.TransformPosition(Pos, InstData->Transform);
-
-		*OutPosX.GetDest() = Pos.X;
-		*OutPosY.GetDest() = Pos.Y;
-		*OutPosZ.GetDest() = Pos.Z;
-		SplineSampleParam.Advance();
-		OutPosX.Advance();
-		OutPosY.Advance();
-		OutPosZ.Advance();
+		for (int32 i = 0; i < Context.NumInstances; ++i)
+		{
+			float DistanceUnitDistance = SplineSampleParam.Get();
+
+			FVector Pos = InstData->Component->GetRightVectorAtDistanceAlongSpline(DistanceUnitDistance * InstData->Component->GetSplineLength(), ESplineCoordinateSpace::Local);
+			TransformHandler.TransformVector(Pos, InstData->Transform);
+
+			*OutPosX.GetDest() = Pos.X;
+			*OutPosY.GetDest() = Pos.Y;
+			*OutPosZ.GetDest() = Pos.Z;
+			SplineSampleParam.Advance();
+			OutPosX.Advance();
+			OutPosY.Advance();
+			OutPosZ.Advance();
+		}
+	}
+	else
+	{
+		for (int32 i = 0; i < Context.NumInstances; ++i)
+		{
+			FVector Pos = FVector(-1.0f, 0.0f, 0.0f); 
+			TransformHandler.TransformVector(Pos, InstData->Transform);
+
+			*OutPosX.GetDest() = Pos.X;
+			*OutPosY.GetDest() = Pos.Y;
+			*OutPosZ.GetDest() = Pos.Z;
+			SplineSampleParam.Advance();
+			OutPosX.Advance();
+			OutPosY.Advance();
+			OutPosZ.Advance();
+		}
 	}
 }
 
@@ -484,21 +543,41 @@ void UNiagaraDataInterfaceSpline::SampleSplineTangentByUnitDistance(FVectorVMCon
 	FRegisterHandler<float> OutPosY(Context);
 	FRegisterHandler<float> OutPosZ(Context);
 
-	checkSlow(InstData->Component);
-	for (int32 i = 0; i < Context.NumInstances; ++i)
+	if (InstData->Component)
 	{
-		float DistanceUnitDistance = SplineSampleParam.Get();
-
-		FVector Pos = InstData->Component->GetTangentAtDistanceAlongSpline(DistanceUnitDistance * InstData->Component->GetSplineLength(), ESplineCoordinateSpace::Local);
-		TransformHandler.TransformPosition(Pos, InstData->Transform);
-
-		*OutPosX.GetDest() = Pos.X;
-		*OutPosY.GetDest() = Pos.Y;
-		*OutPosZ.GetDest() = Pos.Z;
-		SplineSampleParam.Advance();
-		OutPosX.Advance();
-		OutPosY.Advance();
-		OutPosZ.Advance();
+		for (int32 i = 0; i < Context.NumInstances; ++i)
+		{
+			float DistanceUnitDistance = SplineSampleParam.Get();
+
+			FVector Pos = InstData->Component->GetTangentAtDistanceAlongSpline(DistanceUnitDistance * InstData->Component->GetSplineLength(), ESplineCoordinateSpace::Local);
+			TransformHandler.TransformVector(Pos, InstData->Transform);
+
+			*OutPosX.GetDest() = Pos.X;
+			*OutPosY.GetDest() = Pos.Y;
+			*OutPosZ.GetDest() = Pos.Z;
+			SplineSampleParam.Advance();
+			OutPosX.Advance();
+			OutPosY.Advance();
+			OutPosZ.Advance();
+		}
+	}
+	else
+	{
+		for (int32 i = 0; i < Context.NumInstances; ++i)
+		{
+			float DistanceUnitDistance = SplineSampleParam.Get();
+
+			FVector Pos = FVector(EForceInit::ForceInitToZero); 
+			TransformHandler.TransformVector(Pos, InstData->Transform);
+
+			*OutPosX.GetDest() = Pos.X;
+			*OutPosY.GetDest() = Pos.Y;
+			*OutPosZ.GetDest() = Pos.Z;
+			SplineSampleParam.Advance();
+			OutPosX.Advance();
+			OutPosY.Advance();
+			OutPosZ.Advance();
+		}
 	}
 }
 
@@ -512,22 +591,41 @@ void UNiagaraDataInterfaceSpline::SampleSplineDirectionByUnitDistance(FVectorVMC
 	FRegisterHandler<float> OutPosY(Context);
 	FRegisterHandler<float> OutPosZ(Context);
 
-	checkSlow(InstData->Component);
-
-	for (int32 i = 0; i < Context.NumInstances; ++i)
+	if (InstData->Component)
 	{
-		float DistanceUnitDistance = SplineSampleParam.Get();
-
-		FVector Pos = InstData->Component->GetDirectionAtDistanceAlongSpline(DistanceUnitDistance * InstData->Component->GetSplineLength(), ESplineCoordinateSpace::Local);
-		TransformHandler.TransformPosition(Pos, InstData->Transform);
-
-		*OutPosX.GetDest() = Pos.X;
-		*OutPosY.GetDest() = Pos.Y;
-		*OutPosZ.GetDest() = Pos.Z;
-		SplineSampleParam.Advance();
-		OutPosX.Advance();
-		OutPosY.Advance();
-		OutPosZ.Advance();
+		for (int32 i = 0; i < Context.NumInstances; ++i)
+		{
+			float DistanceUnitDistance = SplineSampleParam.Get();
+
+			FVector Pos = InstData->Component->GetDirectionAtDistanceAlongSpline(DistanceUnitDistance * InstData->Component->GetSplineLength(), ESplineCoordinateSpace::Local);
+			TransformHandler.TransformVector(Pos, InstData->Transform);
+
+			*OutPosX.GetDest() = Pos.X;
+			*OutPosY.GetDest() = Pos.Y;
+			*OutPosZ.GetDest() = Pos.Z;
+			SplineSampleParam.Advance();
+			OutPosX.Advance();
+			OutPosY.Advance();
+			OutPosZ.Advance();
+		}
+	}
+	else
+	{
+		for (int32 i = 0; i < Context.NumInstances; ++i)
+		{
+			float DistanceUnitDistance = SplineSampleParam.Get();
+
+			FVector Pos = FVector(0.0f, 1.0f, 0.0f); 
+			TransformHandler.TransformVector(Pos, InstData->Transform);
+
+			*OutPosX.GetDest() = Pos.X;
+			*OutPosY.GetDest() = Pos.Y;
+			*OutPosZ.GetDest() = Pos.Z;
+			SplineSampleParam.Advance();
+			OutPosX.Advance();
+			OutPosY.Advance();
+			OutPosZ.Advance();
+		}
 	}
 }
 
@@ -580,30 +678,48 @@ void UNiagaraDataInterfaceSpline::FindClosestUnitDistanceFromPositionWS(FVectorV
 	FUserPtrHandler<FNDISpline_InstanceData> InstData(Context);
 	FRegisterHandler<float> OutUnitDistance(Context);
 
-	checkSlow(InstData->Component);
+	if (InstData->Component)
+	{
 
-	const int32 NumPoints = InstData->Component->GetSplinePointsPosition().Points.Num();
-	const float FinalKeyTime = InstData->Component->GetSplinePointsPosition().Points[NumPoints - 1].InVal;
+		const int32 NumPoints = InstData->Component->GetSplinePointsPosition().Points.Num();
+		const float FinalKeyTime = InstData->Component->GetSplinePointsPosition().Points[NumPoints - 1].InVal;
 
-	for (int32 i = 0; i < Context.NumInstances; ++i)
-	{
-		float PosX = PosXParam.Get();
-		float PosY = PosYParam.Get();
-		float PosZ = PosZParam.Get();
+		for (int32 i = 0; i < Context.NumInstances; ++i)
+		{
+			float PosX = PosXParam.Get();
+			float PosY = PosYParam.Get();
+			float PosZ = PosZParam.Get();
 
-		FVector Pos(PosX, PosY, PosZ);
+			FVector Pos(PosX, PosY, PosZ);
 
-		// This first call finds the key time, but this is not in 0..1 range for the spline. 
-		float KeyTime = InstData->Component->FindInputKeyClosestToWorldLocation(Pos);
-		// We need to convert into the range by dividing through by the overall duration of the spline according to the keys.
-		float UnitDistance = KeyTime / FinalKeyTime;
+			// This first call finds the key time, but this is not in 0..1 range for the spline. 
+			float KeyTime = InstData->Component->FindInputKeyClosestToWorldLocation(Pos);
+			// We need to convert into the range by dividing through by the overall duration of the spline according to the keys.
+			float UnitDistance = KeyTime / FinalKeyTime;
 
-		*OutUnitDistance.GetDest() = UnitDistance;
+			*OutUnitDistance.GetDest() = UnitDistance;
 
-		PosXParam.Advance();
-		PosYParam.Advance();
-		PosZParam.Advance();
-		OutUnitDistance.Advance();
+			PosXParam.Advance();
+			PosYParam.Advance();
+			PosZParam.Advance();
+			OutUnitDistance.Advance();
+		}
+	}
+	else
+	{
+		for (int32 i = 0; i < Context.NumInstances; ++i)
+		{
+			float PosX = PosXParam.Get();
+			float PosY = PosYParam.Get();
+			float PosZ = PosZParam.Get();
+
+			*OutUnitDistance.GetDest() = 0.0f;
+
+			PosXParam.Advance();
+			PosYParam.Advance();
+			PosZParam.Advance();
+			OutUnitDistance.Advance();
+		}
 	}
 }
 
