System Sim Hack for Update On Spawn

diff --git a/Engine/Plugins/FX/Niagara/Source/Niagara/Private/NiagaraEmitterInstance.cpp b/Engine/Plugins/FX/Niagara/Source/Niagara/Private/NiagaraEmitterInstance.cpp
--- a/Engine/Plugins/FX/Niagara/Source/Niagara/Private/NiagaraEmitterInstance.cpp
+++ b/Engine/Plugins/FX/Niagara/Source/Niagara/Private/NiagaraEmitterInstance.cpp
@@ -739,6 +739,17 @@ void FNiagaraEmitterInstance::PreTick()
 		return;
 	}
 
+	if (TickCount == 0)
+	{
+		//On our very first frame we prime any previous params (for interpolation).
+		SpawnExecContext.PostTick();
+		UpdateExecContext.PostTick();
+		for (FNiagaraScriptExecutionContext& EventContext : EventExecContexts)
+		{
+			EventContext.PostTick();
+		}
+	}
+
 	checkSlow(Data.GetNumVariables() > 0);
 	checkSlow(CachedEmitter->SpawnScriptProps.Script);
 	checkSlow(CachedEmitter->UpdateScriptProps.Script);
diff --git a/Engine/Plugins/FX/Niagara/Source/Niagara/Private/NiagaraSystemSimulation.cpp b/Engine/Plugins/FX/Niagara/Source/Niagara/Private/NiagaraSystemSimulation.cpp
--- a/Engine/Plugins/FX/Niagara/Source/Niagara/Private/NiagaraSystemSimulation.cpp
+++ b/Engine/Plugins/FX/Niagara/Source/Niagara/Private/NiagaraSystemSimulation.cpp
@@ -32,6 +32,14 @@ static FAutoConsoleVariableRef CVarNiagaraDumpSystemData(
 	ECVF_Default
 );
 
+static int32 GbSystemUpdateOnSpawn = 1;
+static FAutoConsoleVariableRef CVarSystemUpdateOnSpawn(
+	TEXT("fx.SystemUpdateOnSpawn"),
+	GbSystemUpdateOnSpawn,
+	TEXT("If > 0, system simulations are given a small update after spawn. \n"),
+	ECVF_Default
+);
+
 //Pretick can no longer be run in parallel. Will likely remain this way.
 // static int32 GbParallelSystemPreTick = 0;
 // static FAutoConsoleVariableRef CVarNiagaraParallelSystemPreTick(
@@ -364,9 +372,6 @@ bool FNiagaraSystemSimulation::Tick(float DeltaSeconds)
 			SpawnInstanceParameterDataSet.Tick();
 			UpdateInstanceParameterDataSet.Tick();
 
-			DataSet.Tick();
-			DataSet.Allocate(NewNum);
-
 			//Setup the few real constants like delta time.
 			float InvDt = 1.0f / DeltaSeconds;
 			float GlobalSpawnCountScale = INiagaraModule::GetGlobalSpawnCountScale();
@@ -385,10 +390,52 @@ bool FNiagaraSystemSimulation::Tick(float DeltaSeconds)
 			DataSetExecInfos.Emplace(&DataSet, 0, false, true);
 		}
 
+		//TODO: JIRA - UE-60096 - Remove.
+		//We're having to allocate and spawn before update here so we have to do needless copies.			
+		//Ideally this should be compiled directly into the script similarly to interpolated particle spawning.
+		if (SpawnNum)
+		{
+			SCOPE_CYCLE_COUNTER(STAT_NiagaraSystemSim_Spawn);
+			DataSet.Allocate(NewNum, true);
+
+			DataSet.SetNumInstances(NewNum);
+			//SpawnInstanceParameterDataSet.SetNumInstances(NewNum);
+
+			//Run Spawn
+			SpawnExecContext.Tick(SoloSystemInstance);//We can't require a specific instance here as these are for all instances.
+			DataSetExecInfos.SetNum(1, false);
+			DataSetExecInfos[0].StartInstance = OrigNum;
+			DataSetExecInfos.Emplace(&SpawnInstanceParameterDataSet, OrigNum, false, false);
+			SpawnExecContext.Execute(SpawnNum, DataSetExecInfos);
+
+			if (GbDumpSystemData || System->bDumpDebugSystemInfo)
+			{
+				UE_LOG(LogNiagara, Log, TEXT("=== Spwaned %d Systems ==="), SpawnNum);
+				DataSet.Dump(true, OrigNum, SpawnNum);
+				SpawnInstanceParameterDataSet.Dump(false, OrigNum, SpawnNum);
+			}
+
+#if WITH_EDITORONLY_DATA
+			if (SoloSystemInstance && SoloSystemInstance->ShouldCaptureThisFrame())
+			{
+				FNiagaraScriptDebuggerInfo* DebugInfo = SoloSystemInstance->GetActiveCaptureWrite(NAME_None, ENiagaraScriptUsage::SystemSpawnScript, FGuid());
+				if (DebugInfo)
+				{
+					DataSet.Dump(DebugInfo->Frame, true, OrigNum, SpawnNum);
+					//DebugInfo->Frame.Dump(true, 0, SpawnNum);
+					DebugInfo->Parameters = SpawnExecContext.Parameters;
+				}
+			}
+#endif
+		}
+
+		DataSet.Tick();
+		DataSet.Allocate(NewNum);
+
 		{
 			SCOPE_CYCLE_COUNTER(STAT_NiagaraSystemSim_Update);
-			DataSet.SetNumInstances(OrigNum);
-			UpdateInstanceParameterDataSet.SetNumInstances(OrigNum);
+			DataSet.SetNumInstances(NewNum);
+			//UpdateInstanceParameterDataSet.SetNumInstances(OrigNum);
 
 			//Run update.
 			UpdateExecContext.Tick(SystemInstances[0]);
@@ -410,13 +457,37 @@ bool FNiagaraSystemSimulation::Tick(float DeltaSeconds)
 				DataSet.Dump(true, 0, OrigNum);
 			}
 
+ 			//Also run the update script on the newly spawned systems too.
+			//TODO: JIRA - UE-60096 - Remove.
+			//Ideally this should be compiled directly into the script similarly to interpolated particle spawning.
+			if (SpawnNum && GbSystemUpdateOnSpawn)
+			{
+				DataSet.SetNumInstances(NewNum);
+
+				//Run update.
+				UpdateExecContext.Tick(SystemInstances[0]);
+				DataSetExecInfos[0].StartInstance = OrigNum;
+				DataSetExecInfos[1].StartInstance = OrigNum;
+
+				UpdateExecContext.Parameters.SetParameterValue(0.0001f, SYS_PARAM_ENGINE_DELTA_TIME);
+				UpdateExecContext.Parameters.SetParameterValue(10000.0f, SYS_PARAM_ENGINE_INV_DELTA_TIME);
+
+				UpdateExecContext.Execute(SpawnNum, DataSetExecInfos);
+
+				if (GbDumpSystemData || System->bDumpDebugSystemInfo)
+				{
+					UE_LOG(LogNiagara, Log, TEXT("=== Spawn Updated %d Systems ==="), SpawnNum);
+					DataSet.Dump(true, OrigNum, SpawnNum);
+				}
+			}
+
 #if WITH_EDITORONLY_DATA
 			if (SoloSystemInstance && SoloSystemInstance->ShouldCaptureThisFrame())
 			{
 				FNiagaraScriptDebuggerInfo* DebugInfo = SoloSystemInstance->GetActiveCaptureWrite(NAME_None, ENiagaraScriptUsage::SystemUpdateScript, FGuid());
 				if (DebugInfo)
 				{
-					DataSet.Dump(DebugInfo->Frame, true, 0, OrigNum);
+					DataSet.Dump(DebugInfo->Frame, true, 0, NewNum);
 					//DebugInfo->Frame.Dump(true, 0, OrigNum);
 					DebugInfo->Parameters = UpdateExecContext.Parameters;
 				}
@@ -434,40 +505,6 @@ bool FNiagaraSystemSimulation::Tick(float DeltaSeconds)
 			}
 		}
 
-		if(SpawnNum)
-		{
-			SCOPE_CYCLE_COUNTER(STAT_NiagaraSystemSim_Spawn);
-			DataSet.SetNumInstances(NewNum);
-			SpawnInstanceParameterDataSet.SetNumInstances(NewNum);
-
-			//Run Spawn
-			SpawnExecContext.Tick(SoloSystemInstance);//We can't require a specific instance here as these are for all instances.
-			DataSetExecInfos.SetNum(1, false);
-			DataSetExecInfos[0].StartInstance = OrigNum;
-			DataSetExecInfos.Emplace(&SpawnInstanceParameterDataSet, OrigNum, false, false);
-			SpawnExecContext.Execute(SpawnNum, DataSetExecInfos);
-
-			if (GbDumpSystemData || System->bDumpDebugSystemInfo)
-			{
-				UE_LOG(LogNiagara, Log, TEXT("=== Spwaned %d Systems ==="), SpawnNum);
-				DataSet.Dump(true, OrigNum, SpawnNum);
-				SpawnInstanceParameterDataSet.Dump(false, OrigNum, SpawnNum);
-			}
-
-#if WITH_EDITORONLY_DATA
-			if (SoloSystemInstance && SoloSystemInstance->ShouldCaptureThisFrame())
-			{
-				FNiagaraScriptDebuggerInfo* DebugInfo = SoloSystemInstance->GetActiveCaptureWrite(NAME_None, ENiagaraScriptUsage::SystemSpawnScript, FGuid());
-				if (DebugInfo)
-				{
-					DataSet.Dump(DebugInfo->Frame, true, OrigNum, SpawnNum);
-					//DebugInfo->Frame.Dump(true, 0, SpawnNum);
-					DebugInfo->Parameters = SpawnExecContext.Parameters;
-				}
-			}
-#endif
-		}
-
 		{
 			SCOPE_CYCLE_COUNTER(STAT_NiagaraSystemSim_TransferParameters);
 			int32 SystemIndex = 0;
