Fix for crash where dynamic inputs didn't have a proper parameter map input value.

#tests new auto-test added

diff --git a/Engine/Plugins/FX/Niagara/Source/NiagaraEditor/Private/ViewModels/Stack/NiagaraStackGraphUtilities.cpp b/Engine/Plugins/FX/Niagara/Source/NiagaraEditor/Private/ViewModels/Stack/NiagaraStackGraphUtilities.cpp
--- a/Engine/Plugins/FX/Niagara/Source/NiagaraEditor/Private/ViewModels/Stack/NiagaraStackGraphUtilities.cpp
+++ b/Engine/Plugins/FX/Niagara/Source/NiagaraEditor/Private/ViewModels/Stack/NiagaraStackGraphUtilities.cpp
@@ -616,54 +616,57 @@ void FNiagaraStackGraphUtilities::RemoveNodesForStackFunctionInputOverridePin(UE
 		{
 			UNiagaraNodeFunctionCall* DynamicInputNode = CastChecked<UNiagaraNodeFunctionCall>(OverrideValueNode);
 			UEdGraphPin* DynamicInputNodeInputPin = FNiagaraStackGraphUtilities::GetParameterMapInputPin(*DynamicInputNode);
-			UNiagaraNodeParameterMapSet* DynamicInputNodeOverrideNode = Cast<UNiagaraNodeParameterMapSet>(DynamicInputNodeInputPin->LinkedTo[0]->GetOwningNode());
-			if (DynamicInputNodeOverrideNode != nullptr)
+			if (DynamicInputNodeInputPin && DynamicInputNodeInputPin->LinkedTo.Num() > 0 && DynamicInputNodeInputPin->LinkedTo[0] != nullptr)
 			{
-				TArray<UEdGraphPin*> InputPins;
-				DynamicInputNodeOverrideNode->GetInputPins(InputPins);
-				for (UEdGraphPin* InputPin : InputPins)
+				UNiagaraNodeParameterMapSet* DynamicInputNodeOverrideNode = Cast<UNiagaraNodeParameterMapSet>(DynamicInputNodeInputPin->LinkedTo[0]->GetOwningNode());
+				if (DynamicInputNodeOverrideNode != nullptr)
 				{
-					if (InputPin->PinName.ToString().StartsWith(DynamicInputNode->GetFunctionName()))
+					TArray<UEdGraphPin*> InputPins;
+					DynamicInputNodeOverrideNode->GetInputPins(InputPins);
+					for (UEdGraphPin* InputPin : InputPins)
 					{
-						RemoveNodesForStackFunctionInputOverridePin(*InputPin, OutRemovedDataObjects);
-						DynamicInputNodeOverrideNode->RemovePin(InputPin);
+						if (InputPin->PinName.ToString().StartsWith(DynamicInputNode->GetFunctionName()))
+						{
+							RemoveNodesForStackFunctionInputOverridePin(*InputPin, OutRemovedDataObjects);
+							DynamicInputNodeOverrideNode->RemovePin(InputPin);
+						}
 					}
-				}
 
-				TArray<UEdGraphPin*> NewInputPins;
-				DynamicInputNodeOverrideNode->GetInputPins(NewInputPins);
-				if (NewInputPins.Num() == 2)
-				{
-					// If there are only 2 input pins left, they are the parameter map input and the add pin, so the dynamic input's override node 
-					// can be removed.  This not always be the case when removing dynamic input nodes because they share the same override node.
-					UEdGraphPin* InputPin = FNiagaraStackGraphUtilities::GetParameterMapInputPin(*DynamicInputNodeOverrideNode);
-					UEdGraphPin* OutputPin = FNiagaraStackGraphUtilities::GetParameterMapOutputPin(*DynamicInputNodeOverrideNode);
-
-					if (ensureMsgf(InputPin != nullptr && InputPin->LinkedTo.Num() == 1 && OutputPin != nullptr && 
-						OutputPin->LinkedTo.Num() >= 2,	TEXT("Invalid Stack - Dynamic input node override node not connected correctly.")))
+					TArray<UEdGraphPin*> NewInputPins;
+					DynamicInputNodeOverrideNode->GetInputPins(NewInputPins);
+					if (NewInputPins.Num() == 2)
 					{
-						// The DynamicInputOverrideNode will have a single input which is the previous module or override map set, and
-						// two or more output links, one to the dynamic input node, one to the next override map set, and 0 or more links
-						// to other dynamic inputs on sibling inputs.  Collect these linked pins to reconnect after removing the override node.
-						UEdGraphPin* LinkedInputPin = InputPin->LinkedTo[0];
-						TArray<UEdGraphPin*> LinkedOutputPins;
-						for (UEdGraphPin* LinkedOutputPin : OutputPin->LinkedTo)
+						// If there are only 2 input pins left, they are the parameter map input and the add pin, so the dynamic input's override node 
+						// can be removed.  This not always be the case when removing dynamic input nodes because they share the same override node.
+						UEdGraphPin* InputPin = FNiagaraStackGraphUtilities::GetParameterMapInputPin(*DynamicInputNodeOverrideNode);
+						UEdGraphPin* OutputPin = FNiagaraStackGraphUtilities::GetParameterMapOutputPin(*DynamicInputNodeOverrideNode);
+
+						if (ensureMsgf(InputPin != nullptr && InputPin->LinkedTo.Num() == 1 && OutputPin != nullptr &&
+							OutputPin->LinkedTo.Num() >= 2, TEXT("Invalid Stack - Dynamic input node override node not connected correctly.")))
 						{
-							if (LinkedOutputPin->GetOwningNode() != DynamicInputNode)
+							// The DynamicInputOverrideNode will have a single input which is the previous module or override map set, and
+							// two or more output links, one to the dynamic input node, one to the next override map set, and 0 or more links
+							// to other dynamic inputs on sibling inputs.  Collect these linked pins to reconnect after removing the override node.
+							UEdGraphPin* LinkedInputPin = InputPin->LinkedTo[0];
+							TArray<UEdGraphPin*> LinkedOutputPins;
+							for (UEdGraphPin* LinkedOutputPin : OutputPin->LinkedTo)
 							{
-								LinkedOutputPins.Add(LinkedOutputPin);
+								if (LinkedOutputPin->GetOwningNode() != DynamicInputNode)
+								{
+									LinkedOutputPins.Add(LinkedOutputPin);
+								}
 							}
-						}
 
-						// Disconnect the override node and remove it.
-						InputPin->BreakAllPinLinks();
-						OutputPin->BreakAllPinLinks();
-						Graph->RemoveNode(DynamicInputNodeOverrideNode);
+							// Disconnect the override node and remove it.
+							InputPin->BreakAllPinLinks();
+							OutputPin->BreakAllPinLinks();
+							Graph->RemoveNode(DynamicInputNodeOverrideNode);
 
-						// Reconnect the pins which were connected to the removed override node.
-						for (UEdGraphPin* LinkedOutputPin : LinkedOutputPins)
-						{
-							LinkedInputPin->MakeLinkTo(LinkedOutputPin);
+							// Reconnect the pins which were connected to the removed override node.
+							for (UEdGraphPin* LinkedOutputPin : LinkedOutputPins)
+							{
+								LinkedInputPin->MakeLinkTo(LinkedOutputPin);
+							}
 						}
 					}
 				}
@@ -746,15 +749,24 @@ void FNiagaraStackGraphUtilities::SetDynamicInputForFunctionInput(UEdGraphPin& O
 	const UEdGraphSchema_Niagara* NiagaraSchema = GetDefault<UEdGraphSchema_Niagara>();
 
 	FNiagaraTypeDefinition InputType = NiagaraSchema->PinToTypeDefinition(&OverridePin);
-	checkf(FunctionCallInputPin != nullptr, TEXT("Dynamic Input function call did not have a parameter map input pin."));
-	checkf(FunctionCallOutputPins.Num() == 1 && NiagaraSchema->PinToTypeDefinition(FunctionCallOutputPins[0]) == InputType, TEXT("Invalid Stack Graph - Dynamic Input function did not have the correct typed output pin"));
 
-	UEdGraphPin* OverrideNodeInputPin = FNiagaraStackGraphUtilities::GetParameterMapInputPin(*OverrideNode);
-	UEdGraphPin* PreviousStackNodeOutputPin = OverrideNodeInputPin->LinkedTo[0];
-	checkf(PreviousStackNodeOutputPin != nullptr, TEXT("Invalid Stack Graph - No previous stack node."));
 
-	FunctionCallInputPin->MakeLinkTo(PreviousStackNodeOutputPin);
-	FunctionCallOutputPins[0]->MakeLinkTo(&OverridePin);
+	UEdGraphPin* OverrideNodeInputPin = FNiagaraStackGraphUtilities::GetParameterMapInputPin(*OverrideNode);
+	UEdGraphPin* PreviousStackNodeOutputPin = nullptr;
+	if (OverrideNodeInputPin != nullptr)
+	{
+		PreviousStackNodeOutputPin = OverrideNodeInputPin->LinkedTo[0];
+	}
+	
+	if (FunctionCallInputPin != nullptr && PreviousStackNodeOutputPin != nullptr)
+	{
+		FunctionCallInputPin->MakeLinkTo(PreviousStackNodeOutputPin);
+	}
+	
+	if (FunctionCallOutputPins.Num() >= 1 && FunctionCallOutputPins[0] != nullptr)
+	{
+		FunctionCallOutputPins[0]->MakeLinkTo(&OverridePin);
+	}
 
 	OutDynamicInputFunctionCall = FunctionCallNode;
 
