Fix niagara compile button for small toolbar icons.

diff --git a/Engine/Plugins/FX/Niagara/Source/NiagaraEditor/Private/NiagaraEditorStyle.cpp b/Engine/Plugins/FX/Niagara/Source/NiagaraEditor/Private/NiagaraEditorStyle.cpp
--- a/Engine/Plugins/FX/Niagara/Source/NiagaraEditor/Private/NiagaraEditorStyle.cpp
+++ b/Engine/Plugins/FX/Niagara/Source/NiagaraEditor/Private/NiagaraEditorStyle.cpp
@@ -156,9 +156,13 @@ TSharedRef< FSlateStyleSet > FNiagaraEditorStyle::Create()
 	Style->Set("NiagaraEditor.SimulationOptions.Small", new IMAGE_PLUGIN_BRUSH("Icons/Commands/icon_simulationOptions_40x", Icon20x20));
 
 	Style->Set("Niagara.CompileStatus.Unknown", new IMAGE_BRUSH("Icons/CompileStatus_Working", Icon40x40));
+	Style->Set("Niagara.CompileStatus.Unknown.Small", new IMAGE_BRUSH("Icons/CompileStatus_Working", Icon20x20));
 	Style->Set("Niagara.CompileStatus.Error",   new IMAGE_BRUSH("Icons/CompileStatus_Fail", Icon40x40));
+	Style->Set("Niagara.CompileStatus.Error.Small", new IMAGE_BRUSH("Icons/CompileStatus_Fail", Icon20x20));
 	Style->Set("Niagara.CompileStatus.Good",    new IMAGE_BRUSH("Icons/CompileStatus_Good", Icon40x40));
+	Style->Set("Niagara.CompileStatus.Good.Small", new IMAGE_BRUSH("Icons/CompileStatus_Good", Icon20x20));
 	Style->Set("Niagara.CompileStatus.Warning", new IMAGE_BRUSH("Icons/CompileStatus_Warning", Icon40x40));
+	Style->Set("Niagara.CompileStatus.Warning.Small", new IMAGE_BRUSH("Icons/CompileStatus_Warning", Icon20x20));
 	Style->Set("Niagara.Asset.ReimportAsset.Needed", new IMAGE_BRUSH("Icons/icon_Reimport_Needed_40x", Icon40x40));
 	Style->Set("Niagara.Asset.ReimportAsset.Default", new IMAGE_BRUSH("Icons/icon_Reimport_40x", Icon40x40));
 
