Fix add parameter enum menu entries for parameter collections.

diff --git a/Engine/Plugins/FX/Niagara/Source/NiagaraEditor/Private/Widgets/SNiagaraParameterCollection.cpp b/Engine/Plugins/FX/Niagara/Source/NiagaraEditor/Private/Widgets/SNiagaraParameterCollection.cpp
--- a/Engine/Plugins/FX/Niagara/Source/NiagaraEditor/Private/Widgets/SNiagaraParameterCollection.cpp
+++ b/Engine/Plugins/FX/Niagara/Source/NiagaraEditor/Private/Widgets/SNiagaraParameterCollection.cpp
@@ -378,13 +378,18 @@ TSharedRef<SWidget> SNiagaraParameterCollection::GetAddMenuContent()
 	FMenuBuilder AddMenuBuilder(true, nullptr);
 	for (TSharedPtr<FNiagaraTypeDefinition> AvailableType : Collection->GetAvailableTypes())
 	{
-		AddMenuBuilder.AddMenuEntry
-		(
-			AvailableType->GetStruct()->GetDisplayNameText(),
-			FText(),
-			FSlateIcon(),
-			FUIAction(FExecuteAction::CreateSP(Collection.ToSharedRef(), &INiagaraParameterCollectionViewModel::AddParameter, AvailableType))
-		);
+		if (AvailableType->GetStruct() != nullptr || AvailableType->GetEnum() != nullptr)
+		{
+			const FText DisplayName = AvailableType->GetEnum() != nullptr ? 
+				FText::FromName(AvailableType->GetEnum()->GetFName()) : AvailableType->GetStruct()->GetDisplayNameText();
+			AddMenuBuilder.AddMenuEntry
+			(
+				DisplayName,
+				FText(),
+				FSlateIcon(),
+				FUIAction(FExecuteAction::CreateSP(Collection.ToSharedRef(), &INiagaraParameterCollectionViewModel::AddParameter, AvailableType))
+			);
+		}
 	}
 	return AddMenuBuilder.MakeWidget();
 }
