Set keyboard focus on add pin searchbox

diff --git a/Engine/Plugins/FX/Niagara/Source/NiagaraEditor/Private/Widgets/SNiagaraGraphPinAdd.cpp b/Engine/Plugins/FX/Niagara/Source/NiagaraEditor/Private/Widgets/SNiagaraGraphPinAdd.cpp
--- a/Engine/Plugins/FX/Niagara/Source/NiagaraEditor/Private/Widgets/SNiagaraGraphPinAdd.cpp
+++ b/Engine/Plugins/FX/Niagara/Source/NiagaraEditor/Private/Widgets/SNiagaraGraphPinAdd.cpp
@@ -43,7 +43,7 @@ void SNiagaraGraphPinAdd::Construct(const FArguments& InArgs, UEdGraphPin* InGra
 
 TSharedRef<SWidget>	SNiagaraGraphPinAdd::ConstructAddButton()
 {
-	return SNew(SComboButton)
+	AddButton = SNew(SComboButton)
 		.HasDownArrow(false)
 		.ButtonStyle(FEditorStyle::Get(), "HoverHintOnly")
 		.ForegroundColor(FSlateColor::UseForeground())
@@ -58,6 +58,8 @@ TSharedRef<SWidget>	SNiagaraGraphPinAdd::ConstructAddButton()
 			.ColorAndOpacity(FSlateColor::UseForeground())
 			.Image(FEditorStyle::GetBrush("Plus"))
 		];
+
+	return AddButton.ToSharedRef();
 }
 
 TSharedRef<SWidget> SNiagaraGraphPinAdd::OnGetAddButtonMenuContent()
@@ -72,7 +74,7 @@ TSharedRef<SWidget> SNiagaraGraphPinAdd::OnGetAddButtonMenuContent()
 			.OnAllowMakeType_UObject(OwningNode, &UNiagaraNodeWithDynamicPins::AllowNiagaraTypeForAddPin)
 			.IsParameterRead(GraphPinObj ? GraphPinObj->Direction == EGPD_Output : true);
 		
-		FSlateApplication::Get().SetKeyboardFocus(MenuWidget->GetSearchBox());
+		AddButton->SetMenuContentWidgetToFocus(MenuWidget->GetSearchBox());
 		return MenuWidget;
 	}
 	else
diff --git a/Engine/Plugins/FX/Niagara/Source/NiagaraEditor/Private/Widgets/SNiagaraGraphPinAdd.h b/Engine/Plugins/FX/Niagara/Source/NiagaraEditor/Private/Widgets/SNiagaraGraphPinAdd.h
--- a/Engine/Plugins/FX/Niagara/Source/NiagaraEditor/Private/Widgets/SNiagaraGraphPinAdd.h
+++ b/Engine/Plugins/FX/Niagara/Source/NiagaraEditor/Private/Widgets/SNiagaraGraphPinAdd.h
@@ -24,6 +24,7 @@ private:
 
 	TSharedRef<SWidget> OnGetAddButtonMenuContent();
 	
+	TSharedPtr<class SComboButton> AddButton;
 
 private:
 	UNiagaraNodeWithDynamicPins* OwningNode;
