Fix renaming variables will not work if the user is only changing capitalization.

diff --git a/Engine/Plugins/FX/Niagara/Source/NiagaraEditor/Private/TNiagaraGraphPinEditableName.h b/Engine/Plugins/FX/Niagara/Source/NiagaraEditor/Private/TNiagaraGraphPinEditableName.h
--- a/Engine/Plugins/FX/Niagara/Source/NiagaraEditor/Private/TNiagaraGraphPinEditableName.h
+++ b/Engine/Plugins/FX/Niagara/Source/NiagaraEditor/Private/TNiagaraGraphPinEditableName.h
@@ -48,7 +48,7 @@ protected:
 
 	void OnTextCommitted(const FText& InText, ETextCommit::Type InCommitType)
 	{
-		if (this->GraphPinObj->PinName.ToString() != InText.ToString())
+		if (!GraphPinObj->PinName.ToString().Equals(InText.ToString(), ESearchCase::CaseSensitive))
 		{
 			UNiagaraNode* ParentNode = Cast<UNiagaraNode>(this->GraphPinObj->GetOwningNode());
 			if (ParentNode != nullptr)
