#tests none

diff --git a/Engine/Plugins/FX/Niagara/Source/NiagaraEditor/Private/Toolkits/NiagaraScriptToolkit.cpp b/Engine/Plugins/FX/Niagara/Source/NiagaraEditor/Private/Toolkits/NiagaraScriptToolkit.cpp
--- a/Engine/Plugins/FX/Niagara/Source/NiagaraEditor/Private/Toolkits/NiagaraScriptToolkit.cpp
+++ b/Engine/Plugins/FX/Niagara/Source/NiagaraEditor/Private/Toolkits/NiagaraScriptToolkit.cpp
@@ -43,6 +43,8 @@
 #include "Developer/MessageLog/Public/MessageLogInitializationOptions.h"
 #include "Developer/MessageLog/Public/MessageLogModule.h"
 
+#include "AssetTypeActions/AssetTypeActions_NiagaraScript.h"
+
 #define LOCTEXT_NAMESPACE "NiagaraScriptToolkit"
 
 DECLARE_CYCLE_STAT(TEXT("Niagara - ScriptToolkit - OnApply"), STAT_NiagaraEditor_ScriptToolkit_OnApply, STATGROUP_NiagaraEditor);
@@ -114,8 +116,23 @@ void FNiagaraScriptToolkit::Initialize( const EToolkitMode::Type Mode, const TSh
 	EditedNiagaraScript = (UNiagaraScript*)StaticDuplicateObject(OriginalNiagaraScript, GetTransientPackage(), NAME_None, ~RF_Standalone, UNiagaraScript::StaticClass());
 	EditedNiagaraScript->OnVMScriptCompiled().AddSP(this, &FNiagaraScriptToolkit::OnVMScriptCompiled);
 	bEditedScriptHasPendingChanges = false;
-	
-	ScriptViewModel = MakeShareable(new FNiagaraScriptViewModel(EditedNiagaraScript, LOCTEXT("NiagaraScriptDisplayName", "Niagara Script"), ENiagaraParameterEditMode::EditAll));
+
+	// Determine display name for panel heading based on asset usage type
+	FText DisplayName = LOCTEXT("NiagaraScriptDisplayName", "Niagara Script");
+	if (EditedNiagaraScript->GetUsage() == ENiagaraScriptUsage::Function)
+	{
+		DisplayName = FAssetTypeActions_NiagaraScriptFunctions::GetFormattedName();
+	}
+	else if (EditedNiagaraScript->GetUsage() == ENiagaraScriptUsage::Module)
+	{
+		DisplayName = FAssetTypeActions_NiagaraScriptModules::GetFormattedName();
+	}
+	else if (EditedNiagaraScript->GetUsage() == ENiagaraScriptUsage::DynamicInput)
+	{
+		DisplayName = FAssetTypeActions_NiagaraScriptDynamicInputs::GetFormattedName();
+	}
+	ScriptViewModel = MakeShareable(new FNiagaraScriptViewModel(EditedNiagaraScript, DisplayName, ENiagaraParameterEditMode::EditAll));
+
 	OnEditedScriptGraphChangedHandle = ScriptViewModel->GetGraphViewModel()->GetGraph()->AddOnGraphNeedsRecompileHandler(
 		FOnGraphChanged::FDelegate::CreateRaw(this, &FNiagaraScriptToolkit::OnEditedScriptGraphChanged));
 
