Fix for #jira UE-61428 (stack search jumps between occurences without user input)
#tests none

diff --git a/Engine/Plugins/FX/Niagara/Source/NiagaraEditorWidgets/Private/SNiagaraStack.cpp b/Engine/Plugins/FX/Niagara/Source/NiagaraEditorWidgets/Private/SNiagaraStack.cpp
--- a/Engine/Plugins/FX/Niagara/Source/NiagaraEditorWidgets/Private/SNiagaraStack.cpp
+++ b/Engine/Plugins/FX/Niagara/Source/NiagaraEditorWidgets/Private/SNiagaraStack.cpp
@@ -89,6 +89,7 @@ void SNiagaraStack::Construct(const FArguments& InArgs, UNiagaraStackViewModel*
 	PinIsPinnedColor = FNiagaraEditorWidgetsStyle::Get().GetColor("NiagaraEditor.Stack.ForegroundColor");
 	PinIsUnpinnedColor = PinIsPinnedColor.Desaturate(.4f);
 	CurrentPinColor = StackViewModel->GetSystemViewModel()->GetIsEmitterPinned(StackViewModel->GetEmitterHandleViewModel()->AsShared()) ? PinIsPinnedColor: PinIsUnpinnedColor;
+	bNeedsJumpToNextOccurence = false;
 
 	ConstructHeaderWidget();
 	TSharedPtr<SWidget> HeaderBox;
@@ -378,7 +379,7 @@ void SNiagaraStack::OnSearchTextChanged(const FText& SearchText)
 				StackTree->SetItemExpansion(ParentalUnit, ParentalUnit->GetIsExpanded());
 			}
 		}
-
+		bNeedsJumpToNextOccurence = true;
 		StackViewModel->OnSearchTextChanged(SearchText);
 	}
 }
@@ -599,7 +600,11 @@ void SNiagaraStack::OnStackSearchComplete()
 EActiveTimerReturnType SNiagaraStack::TriggerExpandSearchResults(double InCurrentTime, float InDeltaTime)
 {
 	ExpandSearchResults();
-	ScrollToNextMatch();
+	if (bNeedsJumpToNextOccurence)
+	{
+		ScrollToNextMatch();
+		bNeedsJumpToNextOccurence = false;
+	}
 	return EActiveTimerReturnType::Stop;
 }
 
@@ -625,7 +630,11 @@ void SNiagaraStack::OnSearchBoxTextCommitted(const FText& NewText, ETextCommit::
 			SearchExpandTimer.Reset();
 		}
 	}
-	AddSearchScrollOffset(+1);
+	if (bNeedsJumpToNextOccurence || CommitInfo == ETextCommit::OnEnter) // hasn't been autojumped yet or we hit enter
+	{
+		AddSearchScrollOffset(+1);
+		bNeedsJumpToNextOccurence = false;
+	}
 }
 
 void SNiagaraStack::OnSearchBoxSearch(SSearchBox::SearchDirection Direction)
diff --git a/Engine/Plugins/FX/Niagara/Source/NiagaraEditorWidgets/Private/SNiagaraStack.h b/Engine/Plugins/FX/Niagara/Source/NiagaraEditorWidgets/Private/SNiagaraStack.h
--- a/Engine/Plugins/FX/Niagara/Source/NiagaraEditorWidgets/Private/SNiagaraStack.h
+++ b/Engine/Plugins/FX/Niagara/Source/NiagaraEditorWidgets/Private/SNiagaraStack.h
@@ -126,4 +126,5 @@ private:
 	TSharedPtr<SSearchBox> SearchBox;
 	TSharedPtr<FActiveTimerHandle> SearchExpandTimer;
 	static const FText OccurencesFormat;
+	bool bNeedsJumpToNextOccurence;
 };
