Bye guys. It was awesome.

#tests none
The bug is fixed but I couldn't change the highlight color due to a bug with TextStyle that the Editor team knew about.

diff --git a/Engine/Plugins/FX/Niagara/Source/NiagaraEditor/Private/Widgets/SNiagaraGeneratedCodeView.cpp b/Engine/Plugins/FX/Niagara/Source/NiagaraEditor/Private/Widgets/SNiagaraGeneratedCodeView.cpp
--- a/Engine/Plugins/FX/Niagara/Source/NiagaraEditor/Private/Widgets/SNiagaraGeneratedCodeView.cpp
+++ b/Engine/Plugins/FX/Niagara/Source/NiagaraEditor/Private/Widgets/SNiagaraGeneratedCodeView.cpp
@@ -210,7 +210,7 @@ TSharedRef<SWidget> SNiagaraGeneratedCodeView::MakeScriptMenu()
 	return MenuBuilder.MakeWidget();
 }
 
-FReply SNiagaraGeneratedCodeView::SearchUpClicked()
+FReply SNiagaraGeneratedCodeView::SearchDownClicked()
 {
 	if (ActiveFoundTextEntries.Num() > 0)
 	{
@@ -228,7 +228,7 @@ FReply SNiagaraGeneratedCodeView::SearchUpClicked()
 	return FReply::Handled();
 }
 
-FReply SNiagaraGeneratedCodeView::SearchDownClicked()
+FReply SNiagaraGeneratedCodeView::SearchUpClicked()
 {
 	if (ActiveFoundTextEntries.Num() > 0)
 	{
@@ -261,6 +261,8 @@ void SNiagaraGeneratedCodeView::OnSearchTextChanged(const FText& InFilterText)
 
 void SNiagaraGeneratedCodeView::DoSearch(const FText& InFilterText)
 {
+	const FText OldText = GeneratedCode[TabState].Text->GetSearchText();
+	GeneratedCode[TabState].Text->SetSearchText(InFilterText);
 	GeneratedCode[TabState].Text->BeginSearch(InFilterText, ESearchCase::IgnoreCase, false);
 	InFilterText.ToString();
 
@@ -273,7 +275,6 @@ void SNiagaraGeneratedCodeView::DoSearch(const FText& InFilterText)
 	}
 
 	ActiveFoundTextEntries.Empty();
-	CurrentFoundTextEntry = INDEX_NONE;
 	for (int32 i = 0; i < GeneratedCode[TabState].HlslByLines.Num(); i++)
 	{
 		int32 LastPos = INDEX_NONE;
@@ -286,24 +287,40 @@ void SNiagaraGeneratedCodeView::DoSearch(const FText& InFilterText)
 		}
 	}
 
-	if (ActiveFoundTextEntries.Num() > 0)
+	if (ActiveFoundTextEntries.Num() > 0 && OldText.CompareTo(InFilterText) != 0)
 	{
 		CurrentFoundTextEntry = 0;
 		//GeneratedCode[TabState].Text->ScrollTo(ActiveFoundTextEntries[CurrentFoundTextEntry]);
 	}
+	else if (ActiveFoundTextEntries.Num() == 0)
+	{
+		CurrentFoundTextEntry = INDEX_NONE;
+	}
 
 	SetSearchMofN();
 }
 
 void SNiagaraGeneratedCodeView::SetSearchMofN()
 {
-	SearchFoundMOfNText->SetText(FText::Format(LOCTEXT("MOfN", "{0} of {1}"), FText::AsNumber(CurrentFoundTextEntry), FText::AsNumber(ActiveFoundTextEntries.Num())));
+	SearchFoundMOfNText->SetText(FText::Format(LOCTEXT("MOfN", "{0} of {1}"), FText::AsNumber(CurrentFoundTextEntry + 1), FText::AsNumber(ActiveFoundTextEntries.Num())));
 	//SearchFoundMOfNText->SetText(FText::Format(LOCTEXT("MOfN", "{1} found"), FText::AsNumber(CurrentFoundTextEntry), FText::AsNumber(ActiveFoundTextEntries.Num())));
 }
 
 void SNiagaraGeneratedCodeView::OnSearchTextCommitted(const FText& InFilterText, ETextCommit::Type InCommitType)
 {
 	OnSearchTextChanged(InFilterText);
+	if (ActiveFoundTextEntries.Num() > 0)
+	{
+		CurrentFoundTextEntry++;
+		if (CurrentFoundTextEntry == ActiveFoundTextEntries.Num())
+		{
+			CurrentFoundTextEntry = 0;
+		}
+	}
+
+	GeneratedCode[TabState].Text->AdvanceSearch(true);
+
+	SetSearchMofN();
 }
 
 void SNiagaraGeneratedCodeView::OnCodeCompiled()
@@ -479,26 +496,14 @@ void SNiagaraGeneratedCodeView::UpdateUI()
 					+ SHorizontalBox::Slot()
 					.FillWidth(1.0f)
 					[
-						SNew(SScrollBox)
-						.Orientation(Orient_Horizontal)
-						.ExternalScrollbar(GeneratedCode[i].HorizontalScrollBar)
-						+ SScrollBox::Slot()
-						[
-							SNew(SScrollBox)
-							.Orientation(Orient_Vertical)
-							.ExternalScrollbar(GeneratedCode[i].VerticalScrollBar)
-							+ SScrollBox::Slot()
-							[
-								SAssignNew(GeneratedCode[i].Text, SMultiLineEditableTextBox)
-								.ClearTextSelectionOnFocusLoss(false)
-								.IsReadOnly(true)
-								.Style(FEditorStyle::Get(), "Log.TextBox")
-								.TextStyle(FEditorStyle::Get(), "Log.Normal")
-								.ForegroundColor(FLinearColor::Yellow)
-								.ReadOnlyForegroundColor(FLinearColor::Yellow)
-								//.TextStyle(FNiagaraEditorStyle::Get(), "NiagaraEditor.CodeView.Hlsl.Normal")
-							]
-						]
+						SAssignNew(GeneratedCode[i].Text, SMultiLineEditableTextBox)
+						.ClearTextSelectionOnFocusLoss(false)
+						.IsReadOnly(true)
+						.TextStyle(FNiagaraEditorStyle::Get(), "NiagaraEditor.CodeView.Hlsl.Normal")
+						.BackgroundColor(FLinearColor::Black)
+						.SearchText(this, &SNiagaraGeneratedCodeView::GetSearchText)
+						.HScrollBar(GeneratedCode[i].HorizontalScrollBar)
+						.VScrollBar(GeneratedCode[i].VerticalScrollBar)
 					]
 					+ SHorizontalBox::Slot()
 					.AutoWidth()
