Integrating work from 4.20 to Dev-Niagara

#tests auto-tests will be updated in subsequent checkin

diff --git a/Engine/Plugins/FX/HoudiniNiagara/Source/HoudiniNiagaraEditor/Private/HoudiniCSVFactory.cpp b/Engine/Plugins/FX/HoudiniNiagara/Source/HoudiniNiagaraEditor/Private/HoudiniCSVFactory.cpp
--- a/Engine/Plugins/FX/HoudiniNiagara/Source/HoudiniNiagaraEditor/Private/HoudiniCSVFactory.cpp
+++ b/Engine/Plugins/FX/HoudiniNiagara/Source/HoudiniNiagaraEditor/Private/HoudiniCSVFactory.cpp
@@ -37,92 +37,91 @@ DEFINE_LOG_CATEGORY(LogHoudiniNiagaraEditor);
 // UHoudiniCSVFactory 
 UHoudiniCSVFactory::UHoudiniCSVFactory(const FObjectInitializer& ObjectInitializer) : Super( ObjectInitializer )
 {
-    // This factory is responsible for manufacturing HoudiniEngine assets.
-    SupportedClass = UHoudiniCSV::StaticClass();
+	// This factory is responsible for manufacturing HoudiniEngine assets.
+	SupportedClass = UHoudiniCSV::StaticClass();
 
-    // This factory does not manufacture new objects from scratch.
-    bCreateNew = false;
+	// This factory does not manufacture new objects from scratch.
+	bCreateNew = false;
 
-    // This factory will open the editor for each new object.
-    bEditAfterNew = false;
+	// This factory will open the editor for each new object.
+	bEditAfterNew = false;
 
-    // This factory will import objects from files.
-    bEditorImport = true;
+	// This factory will import objects from files.
+	bEditorImport = true;
 
-    // Factory does not import objects from text.
-    bText = true;
+	// Factory does not import objects from text.
+	bText = true;
 
-    // Add supported formats.
-    Formats.Add(FString(TEXT("csv;")) + NSLOCTEXT("HoudiniCSVFactory", "FormatCSV", "CSV File").ToString());
-    Formats.Add(FString(TEXT("hcsv;")) + NSLOCTEXT("HoudiniCSVFactory", "FormatHCSV", "HCSV File").ToString());
+	// Add supported formats.
+	Formats.Add(FString(TEXT("hcsv;")) + NSLOCTEXT("HoudiniCSVFactory", "FormatHCSV", "HCSV File").ToString());
 }
 
 
 UObject* 
 UHoudiniCSVFactory::FactoryCreateNew(UClass* Class, UObject* InParent, FName Name, EObjectFlags Flags, UObject* Context, FFeedbackContext* Warn)
 {
-    UHoudiniCSV* NewHoudiniCSVObject = NewObject<UHoudiniCSV>(InParent, Class, Name, Flags | RF_Transactional);
-    return NewHoudiniCSVObject;
+	UHoudiniCSV* NewHoudiniCSVObject = NewObject<UHoudiniCSV>(InParent, Class, Name, Flags | RF_Transactional);
+	return NewHoudiniCSVObject;
 }
 
 UObject* 
 UHoudiniCSVFactory::FactoryCreateFile(UClass* InClass, UObject* InParent, FName InName, EObjectFlags Flags, const FString& Filename, const TCHAR* Parms, FFeedbackContext* Warn, bool& bOutOperationCanceled)
 {
-    bOutOperationCanceled = false;
+	bOutOperationCanceled = false;
 
-    UHoudiniCSV* NewHoudiniCSVObject = NewObject<UHoudiniCSV>(InParent, InClass, InName, Flags);
-    if ( !NewHoudiniCSVObject )
+	UHoudiniCSV* NewHoudiniCSVObject = NewObject<UHoudiniCSV>(InParent, InClass, InName, Flags);
+	if ( !NewHoudiniCSVObject )
 		return nullptr;
 
-    if ( !NewHoudiniCSVObject->UpdateFromFile( Filename ) )
+	if ( !NewHoudiniCSVObject->UpdateFromFile( Filename ) )
 		return nullptr;   
 
-    // Create reimport information.
-    UAssetImportData * AssetImportData = NewHoudiniCSVObject->AssetImportData;
-    if (!AssetImportData)
-    {
+	// Create reimport information.
+	UAssetImportData * AssetImportData = NewHoudiniCSVObject->AssetImportData;
+	if (!AssetImportData)
+	{
 		AssetImportData = NewObject< UAssetImportData >(NewHoudiniCSVObject, UAssetImportData::StaticClass());
 		NewHoudiniCSVObject->AssetImportData = AssetImportData;
-    }
+	}
 
-    //NewHoudiniCSVObject->AssetImportData->Update(Filename);
-    AssetImportData->Update(UFactory::GetCurrentFilename());
+	//NewHoudiniCSVObject->AssetImportData->Update(Filename);
+	AssetImportData->Update(UFactory::GetCurrentFilename());
 
-    // Broadcast notification that the new asset has been imported.
-    FEditorDelegates::OnAssetPostImport.Broadcast(this, NewHoudiniCSVObject);
+	// Broadcast notification that the new asset has been imported.
+	FEditorDelegates::OnAssetPostImport.Broadcast(this, NewHoudiniCSVObject);
 
-    return NewHoudiniCSVObject;
+	return NewHoudiniCSVObject;
 }
 
 FText
 UHoudiniCSVFactory::GetDisplayName() const
 {
-    return LOCTEXT("HoudiniCSVFactoryDescription", "Houdini CSV File");
+	return LOCTEXT("HoudiniCSVFactoryDescription", "Houdini CSV File");
 }
 
 bool
 UHoudiniCSVFactory::DoesSupportClass(UClass * Class)
 {
-    return Class == SupportedClass;
+	return Class == SupportedClass;
 }
 bool 
 UHoudiniCSVFactory::FactoryCanImport(const FString& Filename)
 {
-    const FString Extension = FPaths::GetExtension(Filename);
+	const FString Extension = FPaths::GetExtension(Filename);
 
-    if ( Extension == TEXT("csv") || Extension == TEXT("hcsv") )
-    {
+	if (Extension == TEXT("hcsv") )
+	{
 		return true;
-    }
-    return false;
+	}
+	return false;
 }
 
 bool 
 UHoudiniCSVFactory::CanReimport(UObject* Obj, TArray<FString>& OutFilenames)
 {
-    UHoudiniCSV* HoudiniCSV = Cast<UHoudiniCSV>(Obj);
-    if (HoudiniCSV)
-    {
+	UHoudiniCSV* HoudiniCSV = Cast<UHoudiniCSV>(Obj);
+	if (HoudiniCSV)
+	{
 		if (HoudiniCSV->AssetImportData)
 		{
 			HoudiniCSV->AssetImportData->ExtractFilenames(OutFilenames);
@@ -132,68 +131,68 @@ UHoudiniCSVFactory::CanReimport(UObject* Obj, TArray<FString>& OutFilenames)
 			OutFilenames.Add(TEXT(""));
 		}
 		return true;
-    }
-    return false;
+	}
+	return false;
 }
 
 void 
 UHoudiniCSVFactory::SetReimportPaths(UObject* Obj, const TArray<FString>& NewReimportPaths)
 {
-    UHoudiniCSV* HoudiniCSV = Cast<UHoudiniCSV>( Obj );
-    if ( HoudiniCSV 
+	UHoudiniCSV* HoudiniCSV = Cast<UHoudiniCSV>( Obj );
+	if ( HoudiniCSV 
 		&& HoudiniCSV->AssetImportData
 		&&  ensure( NewReimportPaths.Num() == 1 ) )
-    {
+	{
 		HoudiniCSV->AssetImportData->UpdateFilenameOnly(NewReimportPaths[0]);
-    }
+	}
 }
 
 EReimportResult::Type 
 UHoudiniCSVFactory::Reimport(UObject* Obj)
 {
-    UHoudiniCSV* HoudiniCSV = Cast<UHoudiniCSV>(Obj);
-    if (!HoudiniCSV || !HoudiniCSV->AssetImportData )
+	UHoudiniCSV* HoudiniCSV = Cast<UHoudiniCSV>(Obj);
+	if (!HoudiniCSV || !HoudiniCSV->AssetImportData )
 		return EReimportResult::Failed;
 
-    const FString FilePath = HoudiniCSV->AssetImportData->GetFirstFilename();
-    if (!FilePath.Len())
+	const FString FilePath = HoudiniCSV->AssetImportData->GetFirstFilename();
+	if (!FilePath.Len())
 		return EReimportResult::Failed;
 
-    UE_LOG(LogHoudiniNiagaraEditor, Log, TEXT("Performing atomic reimport of [%s]"), *FilePath);
+	UE_LOG(LogHoudiniNiagaraEditor, Log, TEXT("Performing atomic reimport of [%s]"), *FilePath);
 
-    // Ensure that the file provided by the path exists
-    if (IFileManager::Get().FileSize(*FilePath) == INDEX_NONE)
-    {
+	// Ensure that the file provided by the path exists
+	if (IFileManager::Get().FileSize(*FilePath) == INDEX_NONE)
+	{
 		UE_LOG(LogHoudiniNiagaraEditor, Warning, TEXT("Cannot reimport: source file cannot be found."));
 		return EReimportResult::Failed;
-    }
+	}
 
-    bool OutCanceled = false;
-    if (ImportObject(HoudiniCSV->GetClass(), HoudiniCSV->GetOuter(), *HoudiniCSV->GetName(), RF_Public | RF_Standalone, FilePath, nullptr, OutCanceled) != nullptr)
-    {
+	bool OutCanceled = false;
+	if (ImportObject(HoudiniCSV->GetClass(), HoudiniCSV->GetOuter(), *HoudiniCSV->GetName(), RF_Public | RF_Standalone, FilePath, nullptr, OutCanceled) != nullptr)
+	{
 		UE_LOG(LogHoudiniNiagaraEditor, Log, TEXT("Imported successfully"));
 		// Try to find the outer package so we can dirty it up
 		if (HoudiniCSV->GetOuter())
 			HoudiniCSV->GetOuter()->MarkPackageDirty();
 		else
 			HoudiniCSV->MarkPackageDirty();
-    }
-    else if (OutCanceled)
-    {
+	}
+	else if (OutCanceled)
+	{
 		UE_LOG(LogHoudiniNiagaraEditor, Warning, TEXT("-- import canceled"));
-    }
-    else
-    {
+	}
+	else
+	{
 		UE_LOG(LogHoudiniNiagaraEditor, Warning, TEXT("-- import failed"));
-    }
+	}
 
-    return EReimportResult::Succeeded;
+	return EReimportResult::Succeeded;
 }
 
 int32 
 UHoudiniCSVFactory::GetPriority() const
 {
-    return ImportPriority;
+	return ImportPriority;
 }
 
 #undef LOCTEXT_NAMESPACE
\ No newline at end of file
diff --git a/Engine/Plugins/FX/HoudiniNiagara/Source/HoudiniNiagaraEditor/Private/HoudiniNiagaraEditor.cpp b/Engine/Plugins/FX/HoudiniNiagara/Source/HoudiniNiagaraEditor/Private/HoudiniNiagaraEditor.cpp
--- a/Engine/Plugins/FX/HoudiniNiagara/Source/HoudiniNiagaraEditor/Private/HoudiniNiagaraEditor.cpp
+++ b/Engine/Plugins/FX/HoudiniNiagara/Source/HoudiniNiagaraEditor/Private/HoudiniNiagaraEditor.cpp
@@ -31,12 +31,12 @@
 
 void FHoudiniNiagaraEditorModule::StartupModule()
 {
-    // Register the Houdini CSV Type Actions
-    IAssetTools& AssetTools = FModuleManager::LoadModuleChecked< FAssetToolsModule >("AssetTools").Get();
+	// Register the Houdini CSV Type Actions
+	IAssetTools& AssetTools = FModuleManager::LoadModuleChecked< FAssetToolsModule >("AssetTools").Get();
 
-    TSharedRef< IAssetTypeActions > HCSVAction = MakeShareable( new FHoudiniCSVAssetActions() );
-    AssetTools.RegisterAssetTypeActions( HCSVAction );
-    AssetTypeActions.Add( HCSVAction );
+	TSharedRef< IAssetTypeActions > HCSVAction = MakeShareable( new FHoudiniCSVAssetActions() );
+	AssetTools.RegisterAssetTypeActions( HCSVAction );
+	AssetTypeActions.Add( HCSVAction );
 
 	// Create Slate style set.
 	if (!StyleSet.IsValid())
@@ -69,16 +69,16 @@ void FHoudiniNiagaraEditorModule::StartupModule()
 
 void FHoudiniNiagaraEditorModule::ShutdownModule()
 {
-    // Unregister asset type actions we have previously registered.
-    if ( FModuleManager::Get().IsModuleLoaded("AssetTools") )
-    {
+	// Unregister asset type actions we have previously registered.
+	if ( FModuleManager::Get().IsModuleLoaded("AssetTools") )
+	{
 		IAssetTools & AssetTools = FModuleManager::GetModuleChecked< FAssetToolsModule >("AssetTools").Get();
 
 		for ( int32 Index = 0; Index < AssetTypeActions.Num(); ++Index )
 			AssetTools.UnregisterAssetTypeActions( AssetTypeActions[Index].ToSharedRef() );
 
 		AssetTypeActions.Empty();
-    }
+	}
 
 	// Unregister Slate style set.
 	if (StyleSet.IsValid())
@@ -93,4 +93,4 @@ void FHoudiniNiagaraEditorModule::ShutdownModule()
 
 #undef LOCTEXT_NAMESPACE
 	
-IMPLEMENT_MODULE(FHoudiniNiagaraEditorModule, HoudiniNiagara)
+IMPLEMENT_MODULE(FHoudiniNiagaraEditorModule, HoudiniNiagaraEditor)
diff --git a/Engine/Plugins/FX/Niagara/Source/Niagara/Classes/NiagaraDataInterfaceSimpleCounter.h b/Engine/Plugins/FX/Niagara/Source/Niagara/Classes/NiagaraDataInterfaceSimpleCounter.h
--- a/Engine/Plugins/FX/Niagara/Source/Niagara/Classes/NiagaraDataInterfaceSimpleCounter.h
+++ b/Engine/Plugins/FX/Niagara/Source/Niagara/Classes/NiagaraDataInterfaceSimpleCounter.h
@@ -1,4 +1,4 @@
-// Copyright 1998-2016 Epic Games, Inc. All Rights Reserved.
+// Copyright 1998-2018 Epic Games, Inc. All Rights Reserved.
 #pragma once
 
 #include "NiagaraCommon.h"
diff --git a/Engine/Plugins/FX/Niagara/Source/Niagara/Classes/NiagaraDataSet.h b/Engine/Plugins/FX/Niagara/Source/Niagara/Classes/NiagaraDataSet.h
--- a/Engine/Plugins/FX/Niagara/Source/Niagara/Classes/NiagaraDataSet.h
+++ b/Engine/Plugins/FX/Niagara/Source/Niagara/Classes/NiagaraDataSet.h
@@ -40,8 +40,8 @@ public:
 	void AllocateGPU(uint32 InNumInstances, FRHICommandList &RHICmdList);
 	void SwapInstances(uint32 OldIndex, uint32 NewIndex);
 	void KillInstance(uint32 InstanceIdx);
-	void CopyTo(FNiagaraDataBuffer& DestBuffer, int32 StartIdx, int32 NumInstances);
-	void CopyTo(FNiagaraDataBuffer& DestBuffer);
+	void CopyTo(FNiagaraDataBuffer& DestBuffer, int32 StartIdx, int32 NumInstances)const;
+	void CopyTo(FNiagaraDataBuffer& DestBuffer)const;
 
 	const TArray<uint8>& GetFloatBuffer()const { return FloatData; }
 	const TArray<uint8>& GetInt32Buffer()const { return Int32Data; }
@@ -76,6 +76,16 @@ public:
 		return (int32*)GetComponentPtrInt32(ComponentIdx) + InstanceIdx;
 	}
 
+	FORCEINLINE float* GetInstancePtrFloat(uint32 ComponentIdx, uint32 InstanceIdx)const
+	{
+		return (float*)GetComponentPtrFloat(ComponentIdx) + InstanceIdx;
+	}
+
+	FORCEINLINE int32* GetInstancePtrInt32(uint32 ComponentIdx, uint32 InstanceIdx)const
+	{
+		return (int32*)GetComponentPtrInt32(ComponentIdx) + InstanceIdx;
+	}
+
 	uint32 GetNumInstances()const { return NumInstances; }
 	uint32 GetNumInstancesAllocated()const { return NumInstancesAllocated; }
 
@@ -112,6 +122,8 @@ public:
 	FORCEINLINE const FNiagaraDataSet* GetOwner()const { return Owner; }
 
 	int32 TransferInstance(FNiagaraDataBuffer& SourceBuffer, int32 InstanceIndex);
+
+	bool CheckForNaNs()const;
 private:
 
 	FORCEINLINE int32 GetSafeComponentBufferSize(int32 RequiredSize) const
@@ -284,6 +296,10 @@ public:
 		CheckCorrectThread();
 		CurrData().Allocate(NumInstances, bMaintainExisting);
 
+#if NIAGARA_NAN_CHECKING
+		CheckForNaNs();
+#endif
+
 		if (bNeedsPersistentIDs)
 		{
 			int32 NumUsedIDs = MaxUsedID + 1;
@@ -347,6 +363,9 @@ public:
 	
 	FORCEINLINE void Tick()
 	{
+#if NIAGARA_NAN_CHECKING
+		CheckForNaNs();
+#endif
 		SwapBuffers();
 	}
 	
@@ -509,10 +528,24 @@ public:
 		return false;
 	}
 
-	void Dump(bool bCurr, int32 StartIdx = 0, int32 NumInstances = INDEX_NONE);
-	void Dump(FNiagaraDataSet& Other, bool bCurr, int32 StartIdx = 0, int32 NumInstances = INDEX_NONE);
+	void Dump(bool bCurr, int32 StartIdx = 0, int32 NumInstances = INDEX_NONE)const;
+	void Dump(FNiagaraDataSet& Other, bool bCurr, int32 StartIdx = 0, int32 NumInstances = INDEX_NONE)const;
 	FORCEINLINE const TArray<FNiagaraVariable> &GetVariables() { return Variables; }
 
+	void CheckForNaNs()const
+	{
+		if (CurrData().CheckForNaNs())
+		{
+			Dump(true);
+			ensureAlwaysMsgf(false, TEXT("NiagaraDataSet contains NaNs!"));
+		}
+
+		if (PrevData().CheckForNaNs())
+		{
+			Dump(false);
+			ensureAlwaysMsgf(false, TEXT("NiagaraDataSet contains NaNs!"));
+		}
+	}
 
 	// Data set index buffer management
 	// these buffers hold the number of instances for the buffers; the first five uint32s are the DrawIndirect parameters for rendering of the main particle data set
@@ -1625,5 +1658,79 @@ private:
 	TArray<FNiagaraVariable*> Variables;
 	TArray<const FNiagaraVariableLayoutInfo*> VarLayouts;
 
+	uint32 CurrIdx;
+};
+
+/**
+Iterator that will pull or push data between a DataSet and some FNiagaraVariables it contains.
+Super slow. Don't use at runtime.
+*/
+struct FNiagaraDataSetVariableIteratorConst
+{
+	FNiagaraDataSetVariableIteratorConst(const FNiagaraDataSet& InDataSet, uint32 StartIdx = 0, bool bCurrBuffer = true)
+		: DataSet(InDataSet)
+		, DataBuffer(bCurrBuffer ? DataSet.CurrData() : DataSet.PrevData())
+		, CurrIdx(StartIdx)
+	{
+	}
+
+	void Get()
+	{
+		for (int32 VarIdx = 0; VarIdx < Variables.Num(); ++VarIdx)
+		{
+			FNiagaraVariable* Var = Variables[VarIdx];
+			const FNiagaraVariableLayoutInfo* Layout = VarLayouts[VarIdx];
+			check(Var && Layout);
+			uint8* ValuePtr = Var->GetData();
+
+			for (uint32 CompIdx = 0; CompIdx < Layout->GetNumFloatComponents(); ++CompIdx)
+			{
+				uint32 CompBufferOffset = Layout->FloatComponentStart + CompIdx;
+				float* Src = DataBuffer.GetInstancePtrFloat(CompBufferOffset, CurrIdx);
+				float* Dst = (float*)(ValuePtr + Layout->LayoutInfo.FloatComponentByteOffsets[CompIdx]);
+				*Dst = *Src;
+			}
+
+			for (uint32 CompIdx = 0; CompIdx < Layout->GetNumInt32Components(); ++CompIdx)
+			{
+				uint32 CompBufferOffset = Layout->Int32ComponentStart + CompIdx;
+				int32* Src = DataBuffer.GetInstancePtrInt32(CompBufferOffset, CurrIdx);
+				int32* Dst = (int32*)(ValuePtr + Layout->LayoutInfo.Int32ComponentByteOffsets[CompIdx]);
+				*Dst = *Src;
+			}
+		}
+	}
+	
+	void Advance() { ++CurrIdx; }
+
+	bool IsValid()const
+	{
+		return CurrIdx < DataBuffer.GetNumInstances();
+	}
+
+	uint32 GetCurrIndex()const { return CurrIdx; }
+
+	void AddVariable(FNiagaraVariable* InVar)
+	{
+		check(InVar);
+		Variables.AddUnique(InVar);
+		VarLayouts.AddUnique(DataSet.GetVariableLayout(*InVar));
+		InVar->AllocateData();
+	}
+
+	void AddVariables(TArray<FNiagaraVariable>& Vars)
+	{
+		for (FNiagaraVariable& Var : Vars)
+		{
+			AddVariable(&Var);
+		}
+	}
+private:
+
+	const FNiagaraDataSet & DataSet;
+	const FNiagaraDataBuffer& DataBuffer;
+	TArray<FNiagaraVariable*> Variables;
+	TArray<const FNiagaraVariableLayoutInfo*> VarLayouts;
+
 	uint32 CurrIdx;
 };
\ No newline at end of file
diff --git a/Engine/Plugins/FX/Niagara/Source/Niagara/Classes/NiagaraEmitter.h b/Engine/Plugins/FX/Niagara/Source/Niagara/Classes/NiagaraEmitter.h
--- a/Engine/Plugins/FX/Niagara/Source/Niagara/Classes/NiagaraEmitter.h
+++ b/Engine/Plugins/FX/Niagara/Source/Niagara/Classes/NiagaraEmitter.h
@@ -94,7 +94,7 @@ enum class EScriptExecutionMode : uint8
 	/** The event script is run only on the particles that were spawned in response to the current event in the emitter.*/
 	SpawnedParticles,
 	/** The event script is run only on the particle whose int32 ParticleIndex is specified in the event payload.*/
-	SingleParticle
+	SingleParticle UMETA(Hidden)
 };
 
 USTRUCT()
diff --git a/Engine/Plugins/FX/Niagara/Source/Niagara/Classes/NiagaraEmitterInstance.h b/Engine/Plugins/FX/Niagara/Source/Niagara/Classes/NiagaraEmitterInstance.h
--- a/Engine/Plugins/FX/Niagara/Source/Niagara/Classes/NiagaraEmitterInstance.h
+++ b/Engine/Plugins/FX/Niagara/Source/Niagara/Classes/NiagaraEmitterInstance.h
@@ -102,6 +102,7 @@ public:
 
 	NIAGARA_API bool IsReadyToRun() const;
 
+	void Dump()const;
 private:
 
 	void CheckForErrors();
diff --git a/Engine/Plugins/FX/Niagara/Source/Niagara/Classes/NiagaraScript.h b/Engine/Plugins/FX/Niagara/Source/Niagara/Classes/NiagaraScript.h
--- a/Engine/Plugins/FX/Niagara/Source/Niagara/Classes/NiagaraScript.h
+++ b/Engine/Plugins/FX/Niagara/Source/Niagara/Classes/NiagaraScript.h
@@ -515,6 +515,7 @@ private:
 	FNiagaraScriptExecutionParameterStore ScriptExecutionParamStoreGPU;
 
 #if WITH_EDITORONLY_DATA
+	class UNiagaraSystem* FindRootSystem();
 
 	/** 'Source' data/graphs for this script */
 	UPROPERTY()
diff --git a/Engine/Plugins/FX/Niagara/Source/Niagara/Classes/NiagaraScriptExecutionContext.h b/Engine/Plugins/FX/Niagara/Source/Niagara/Classes/NiagaraScriptExecutionContext.h
--- a/Engine/Plugins/FX/Niagara/Source/Niagara/Classes/NiagaraScriptExecutionContext.h
+++ b/Engine/Plugins/FX/Niagara/Source/Niagara/Classes/NiagaraScriptExecutionContext.h
@@ -90,17 +90,19 @@ struct FNiagaraComputeExecutionContext
 		, UpdateScript(nullptr)
 		, GPUScript(nullptr)
 		, RTUpdateScript(0)
-		, RTSpawnScript(0)	
+		, RTSpawnScript(0)
 		, RTGPUScript(0)
 		, GPUDataReadback(nullptr)
 		, AccumulatedSpawnRate(0)
 		, NumIndicesPerInstance(0)
+		, bPendingExecution(0)
 	{
 	}
 
 	void Reset()
 	{
 		AccumulatedSpawnRate = 0;
+		bPendingExecution = 0;
 		if (GPUDataReadback)
 		{
 			delete GPUDataReadback;
@@ -194,4 +196,7 @@ struct FNiagaraComputeExecutionContext
 	FRHIGPUMemoryReadback *GPUDataReadback;
 	uint32 AccumulatedSpawnRate;
 	uint32 NumIndicesPerInstance;	// how many vtx indices per instance the renderer is going to have for its draw call
+
+	/** Ensures we only enqueue each context once before they're dispatched. */
+	uint32 bPendingExecution : 1;
 };
diff --git a/Engine/Plugins/FX/Niagara/Source/Niagara/Classes/NiagaraSystem.h b/Engine/Plugins/FX/Niagara/Source/Niagara/Classes/NiagaraSystem.h
--- a/Engine/Plugins/FX/Niagara/Source/Niagara/Classes/NiagaraSystem.h
+++ b/Engine/Plugins/FX/Niagara/Source/Niagara/Classes/NiagaraSystem.h
@@ -149,6 +149,9 @@ public:
 	/** Called to query whether or not this emitter is referenced as the source to any emitter handles for this System.*/
 	bool ReferencesSourceEmitter(UNiagaraEmitter& Emitter);
 
+	/** Determines if this system has the supplied emitter as an editable and simulating emitter instance. */
+	bool ReferencesInstanceEmitter(UNiagaraEmitter& Emitter);
+
 	/** Updates all handles which use this emitter as their source. */
 	void UpdateFromEmitterChanges(UNiagaraEmitter& ChangedSourceEmitter);
 
diff --git a/Engine/Plugins/FX/Niagara/Source/Niagara/Private/MovieScene/MovieSceneNiagaraSystemSpawnSection.cpp b/Engine/Plugins/FX/Niagara/Source/Niagara/Private/MovieScene/MovieSceneNiagaraSystemSpawnSection.cpp
--- a/Engine/Plugins/FX/Niagara/Source/Niagara/Private/MovieScene/MovieSceneNiagaraSystemSpawnSection.cpp
+++ b/Engine/Plugins/FX/Niagara/Source/Niagara/Private/MovieScene/MovieSceneNiagaraSystemSpawnSection.cpp
@@ -1,4 +1,4 @@
-// Copyright 1998-2017 Epic Games, Inc. All Rights Reserved.
+// Copyright 1998-2018 Epic Games, Inc. All Rights Reserved.
 
 #include "MovieScene/MovieSceneNiagaraSystemSpawnSection.h"
 
diff --git a/Engine/Plugins/FX/Niagara/Source/Niagara/Private/MovieScene/MovieSceneNiagaraSystemTrack.cpp b/Engine/Plugins/FX/Niagara/Source/Niagara/Private/MovieScene/MovieSceneNiagaraSystemTrack.cpp
--- a/Engine/Plugins/FX/Niagara/Source/Niagara/Private/MovieScene/MovieSceneNiagaraSystemTrack.cpp
+++ b/Engine/Plugins/FX/Niagara/Source/Niagara/Private/MovieScene/MovieSceneNiagaraSystemTrack.cpp
@@ -1,4 +1,4 @@
-// Copyright 1998-2017 Epic Games, Inc. All Rights Reserved.
+// Copyright 1998-2018 Epic Games, Inc. All Rights Reserved.
 
 #include "MovieScene/MovieSceneNiagaraSystemTrack.h"
 #include "CoreMinimal.h"
diff --git a/Engine/Plugins/FX/Niagara/Source/Niagara/Private/MovieScene/MovieSceneNiagaraSystemTrackTemplate.cpp b/Engine/Plugins/FX/Niagara/Source/Niagara/Private/MovieScene/MovieSceneNiagaraSystemTrackTemplate.cpp
--- a/Engine/Plugins/FX/Niagara/Source/Niagara/Private/MovieScene/MovieSceneNiagaraSystemTrackTemplate.cpp
+++ b/Engine/Plugins/FX/Niagara/Source/Niagara/Private/MovieScene/MovieSceneNiagaraSystemTrackTemplate.cpp
@@ -1,4 +1,4 @@
-// Copyright 1998-2017 Epic Games, Inc. All Rights Reserved.
+// Copyright 1998-2018 Epic Games, Inc. All Rights Reserved.
 
 #include "MovieSceneNiagaraSystemTrackTemplate.h"
 #include "MovieSceneExecutionToken.h"
diff --git a/Engine/Plugins/FX/Niagara/Source/Niagara/Private/MovieScene/MovieSceneNiagaraSystemTrackTemplate.h b/Engine/Plugins/FX/Niagara/Source/Niagara/Private/MovieScene/MovieSceneNiagaraSystemTrackTemplate.h
--- a/Engine/Plugins/FX/Niagara/Source/Niagara/Private/MovieScene/MovieSceneNiagaraSystemTrackTemplate.h
+++ b/Engine/Plugins/FX/Niagara/Source/Niagara/Private/MovieScene/MovieSceneNiagaraSystemTrackTemplate.h
@@ -1,4 +1,4 @@
-// Copyright 1998-2017 Epic Games, Inc. All Rights Reserved.
+// Copyright 1998-2018 Epic Games, Inc. All Rights Reserved.
 
 #pragma once
 
diff --git a/Engine/Plugins/FX/Niagara/Source/Niagara/Private/MovieScene/MovieSceneNiagaraTrack.cpp b/Engine/Plugins/FX/Niagara/Source/Niagara/Private/MovieScene/MovieSceneNiagaraTrack.cpp
--- a/Engine/Plugins/FX/Niagara/Source/Niagara/Private/MovieScene/MovieSceneNiagaraTrack.cpp
+++ b/Engine/Plugins/FX/Niagara/Source/Niagara/Private/MovieScene/MovieSceneNiagaraTrack.cpp
@@ -1,4 +1,4 @@
-// Copyright 1998-2017 Epic Games, Inc. All Rights Reserved.
+// Copyright 1998-2018 Epic Games, Inc. All Rights Reserved.
 
 #include "MovieScene/MovieSceneNiagaraTrack.h"
 
diff --git a/Engine/Plugins/FX/Niagara/Source/Niagara/Private/NiagaraDataInterfaceSimpleCounter.cpp b/Engine/Plugins/FX/Niagara/Source/Niagara/Private/NiagaraDataInterfaceSimpleCounter.cpp
--- a/Engine/Plugins/FX/Niagara/Source/Niagara/Private/NiagaraDataInterfaceSimpleCounter.cpp
+++ b/Engine/Plugins/FX/Niagara/Source/Niagara/Private/NiagaraDataInterfaceSimpleCounter.cpp
@@ -1,4 +1,4 @@
-// Copyright 1998-2016 Epic Games, Inc. All Rights Reserved.
+// Copyright 1998-2018 Epic Games, Inc. All Rights Reserved.
 
 #include "NiagaraDataInterfaceSimpleCounter.h"
 #include "NiagaraTypes.h"
diff --git a/Engine/Plugins/FX/Niagara/Source/Niagara/Private/NiagaraDataInterfaceStaticMesh.cpp b/Engine/Plugins/FX/Niagara/Source/Niagara/Private/NiagaraDataInterfaceStaticMesh.cpp
--- a/Engine/Plugins/FX/Niagara/Source/Niagara/Private/NiagaraDataInterfaceStaticMesh.cpp
+++ b/Engine/Plugins/FX/Niagara/Source/Niagara/Private/NiagaraDataInterfaceStaticMesh.cpp
@@ -1294,28 +1294,51 @@ void UNiagaraDataInterfaceStaticMesh::GetTriCoordColor(FVectorVMContext& Context
 	const FIndexArrayView& Indices = Res.IndexBuffer.GetArrayView();
 	const FColorVertexBuffer& Colors = Res.VertexBuffers.ColorVertexBuffer;
 
-	for (int32 i = 0; i < Context.NumInstances; ++i)
+	if (Colors.GetNumVertices() > 0)
 	{
-		int32 Tri = TriParam.Get();
-		int32 Idx0 = Indices[Tri];
-		int32 Idx1 = Indices[Tri + 1];
-		int32 Idx2 = Indices[Tri + 2];
-
-		FLinearColor Color = BarycentricInterpolate(BaryXParam.Get(), BaryYParam.Get(), BaryZParam.Get(),
-			Colors.VertexColor(Idx0).ReinterpretAsLinear(), Colors.VertexColor(Idx1).ReinterpretAsLinear(), Colors.VertexColor(Idx2).ReinterpretAsLinear());
-
-		*OutColorR.GetDest() = Color.R;
-		*OutColorG.GetDest() = Color.G;
-		*OutColorB.GetDest() = Color.B;
-		*OutColorA.GetDest() = Color.A;
-		TriParam.Advance();
-		BaryXParam.Advance();
-		BaryYParam.Advance();
-		BaryZParam.Advance();
-		OutColorR.Advance();
-		OutColorG.Advance();
-		OutColorB.Advance();
-		OutColorA.Advance();
+		for (int32 i = 0; i < Context.NumInstances; ++i)
+		{
+			int32 Tri = TriParam.Get();
+			int32 Idx0 = Indices[Tri];
+			int32 Idx1 = Indices[Tri + 1];
+			int32 Idx2 = Indices[Tri + 2];
+
+			FLinearColor Color = BarycentricInterpolate(BaryXParam.Get(), BaryYParam.Get(), BaryZParam.Get(),
+				Colors.VertexColor(Idx0).ReinterpretAsLinear(), Colors.VertexColor(Idx1).ReinterpretAsLinear(), Colors.VertexColor(Idx2).ReinterpretAsLinear());
+
+			*OutColorR.GetDest() = Color.R;
+			*OutColorG.GetDest() = Color.G;
+			*OutColorB.GetDest() = Color.B;
+			*OutColorA.GetDest() = Color.A;
+			TriParam.Advance();
+			BaryXParam.Advance();
+			BaryYParam.Advance();
+			BaryZParam.Advance();
+			OutColorR.Advance();
+			OutColorG.Advance();
+			OutColorB.Advance();
+			OutColorA.Advance();
+		}
+	}
+	else
+	{
+		// This mesh doesn't have color information so set the color to white.
+		FLinearColor Color = FLinearColor::White;
+		for (int32 i = 0; i < Context.NumInstances; ++i)
+		{
+			*OutColorR.GetDest() = Color.R;
+			*OutColorG.GetDest() = Color.G;
+			*OutColorB.GetDest() = Color.B;
+			*OutColorA.GetDest() = Color.A;
+			TriParam.Advance();
+			BaryXParam.Advance();
+			BaryYParam.Advance();
+			BaryZParam.Advance();
+			OutColorR.Advance();
+			OutColorG.Advance();
+			OutColorB.Advance();
+			OutColorA.Advance();
+		}
 	}
 }
 
diff --git a/Engine/Plugins/FX/Niagara/Source/Niagara/Private/NiagaraDataSet.cpp b/Engine/Plugins/FX/Niagara/Source/Niagara/Private/NiagaraDataSet.cpp
--- a/Engine/Plugins/FX/Niagara/Source/Niagara/Private/NiagaraDataSet.cpp
+++ b/Engine/Plugins/FX/Niagara/Source/Niagara/Private/NiagaraDataSet.cpp
@@ -112,13 +112,13 @@ FNiagaraDataSet::~FNiagaraDataSet()
 
 }
 
-void FNiagaraDataSet::Dump(FNiagaraDataSet& Other, bool bCurr, int32 StartIdx , int32 NumInstances)
+void FNiagaraDataSet::Dump(FNiagaraDataSet& Other, bool bCurr, int32 StartIdx , int32 NumInstances)const
 {
 	Other.Reset();
 	Other.Variables = Variables;
 	Other.VariableLayouts = VariableLayouts;
 
-	FNiagaraDataBuffer& DataBuffer = bCurr ? CurrData() : PrevData();
+	const FNiagaraDataBuffer& DataBuffer = bCurr ? CurrData() : PrevData();
 	FNiagaraDataBuffer& OtherDataBuffer = Other.CurrData();
 
 	
@@ -132,11 +132,11 @@ void FNiagaraDataSet::Dump(FNiagaraDataSet& Other, bool bCurr, int32 StartIdx ,
 }
 
 
-void FNiagaraDataSet::Dump(bool bCurr, int32 StartIdx, int32 NumInstances)
+void FNiagaraDataSet::Dump(bool bCurr, int32 StartIdx, int32 NumInstances)const
 {
 	TArray<FNiagaraVariable> Vars(Variables);
 
-	FNiagaraDataSetVariableIterator Itr(*this, StartIdx, bCurr);
+	FNiagaraDataSetVariableIteratorConst Itr(*this, StartIdx, bCurr);
 	Itr.AddVariables(Vars);
 
 	if (NumInstances == INDEX_NONE)
@@ -240,6 +240,22 @@ int32 FNiagaraDataBuffer::TransferInstance(FNiagaraDataBuffer& SourceBuffer, int
 	return INDEX_NONE;
 }
 
+bool FNiagaraDataBuffer::CheckForNaNs()const
+{
+	bool bContainsNaNs = false;
+	int32 NumFloatComponents = Owner->GetNumFloatComponents();
+	for (int32 CompIdx = 0; CompIdx < NumFloatComponents && !bContainsNaNs; ++CompIdx)
+	{
+		for (int32 InstIdx = 0; InstIdx < (int32)NumInstances && !bContainsNaNs; ++InstIdx)
+		{
+			float Val = *GetInstancePtrFloat(CompIdx, InstIdx);
+			bContainsNaNs = FMath::IsNaN(Val) || !FMath::IsFinite(Val);
+		}
+	}
+
+	return bContainsNaNs;
+}
+
 void FNiagaraDataBuffer::Allocate(uint32 InNumInstances, bool bMaintainExisting)
 {
 	check(Owner);
@@ -251,10 +267,20 @@ void FNiagaraDataBuffer::Allocate(uint32 InNumInstances, bool bMaintainExisting)
 		DEC_MEMORY_STAT_BY(STAT_NiagaraParticleMemory, FloatData.Max() + Int32Data.Max());
 
 		int32 OldFloatStride = FloatStride;
+		TArray<uint8> OldFloatData;
+		int32 OldInt32Stride = Int32Stride;
+		TArray<uint8> OldIntData;
+
+		if (bMaintainExisting)
+		{
+			//Need to copy off old data so we can copy it back into the newly laid out buffers. TODO: Avoid this needless copying.
+			OldFloatData = FloatData;
+			OldIntData = Int32Data;
+		}
+
 		FloatStride = GetSafeComponentBufferSize(NumInstancesAllocated * sizeof(float));
 		FloatData.SetNum(FloatStride * Owner->GetNumFloatComponents(), false);
 
-		int32 OldInt32Stride = Int32Stride;
 		Int32Stride = GetSafeComponentBufferSize(NumInstancesAllocated * sizeof(int32));
 		Int32Data.SetNum(Int32Stride * Owner->GetNumInt32Components(), false);
 
@@ -265,20 +291,20 @@ void FNiagaraDataBuffer::Allocate(uint32 InNumInstances, bool bMaintainExisting)
 		{
 			if (FloatStride != OldFloatStride && FloatStride > 0 && OldFloatStride > 0)
 			{
-				for (uint32 CompIdx = Owner->TotalFloatComponents-1; CompIdx > 0; --CompIdx)
+				for (int32 CompIdx = (int32)Owner->TotalFloatComponents-1; CompIdx >= 0; --CompIdx)
 				{
-					uint8* Src = FloatData.GetData() + OldFloatStride * CompIdx;
+					uint8* Src = OldFloatData.GetData() + OldFloatStride * CompIdx;
 					uint8* Dst = FloatData.GetData() + FloatStride * CompIdx;
 					FMemory::Memcpy(Dst, Src, OldFloatStride);
 				}
 			}
 			if (Int32Stride != OldInt32Stride && Int32Stride > 0 && OldInt32Stride > 0)
 			{
-				for (uint32 CompIdx = Owner->TotalInt32Components - 1; CompIdx > 0; --CompIdx)
+				for (int32 CompIdx = (int32)Owner->TotalInt32Components - 1; CompIdx >= 0; --CompIdx)
 				{
-					uint8* Src = Int32Data.GetData() + OldInt32Stride * CompIdx;
+					uint8* Src = OldIntData.GetData() + OldInt32Stride * CompIdx;
 					uint8* Dst = Int32Data.GetData() + Int32Stride * CompIdx;
-					FMemory::Memcpy(Dst, Src, OldFloatStride);
+					FMemory::Memcpy(Dst, Src, OldInt32Stride);
 				}
 			}
 		}
@@ -368,9 +394,13 @@ void FNiagaraDataBuffer::KillInstance(uint32 InstanceIdx)
 		int32* Dst = GetInstancePtrInt32(CompIdx, InstanceIdx);
 		*Dst = *Src;
 	}
+
+#if NIAGARA_NAN_CHECKING
+	CheckForNaNs();
+#endif
 }
 
-void FNiagaraDataBuffer::CopyTo(FNiagaraDataBuffer& DestBuffer, int32 InStartIdx, int32 InNumInstances)
+void FNiagaraDataBuffer::CopyTo(FNiagaraDataBuffer& DestBuffer, int32 InStartIdx, int32 InNumInstances)const
 {
 	if (InStartIdx < 0 || (uint32)InStartIdx > NumInstances)
 	{
@@ -390,8 +420,8 @@ void FNiagaraDataBuffer::CopyTo(FNiagaraDataBuffer& DestBuffer, int32 InStartIdx
 
 		for (uint32 CompIdx = 0; CompIdx < Owner->TotalFloatComponents; ++CompIdx)
 		{
-			float* SrcStart = GetInstancePtrFloat(CompIdx, InStartIdx);
-			float* SrcEnd = GetInstancePtrFloat(CompIdx, InStartIdx + InNumInstances);
+			const float* SrcStart = GetInstancePtrFloat(CompIdx, InStartIdx);
+			const float* SrcEnd = GetInstancePtrFloat(CompIdx, InStartIdx + InNumInstances);
 			float* Dst = DestBuffer.GetInstancePtrFloat(CompIdx, 0);
 			size_t Count = SrcEnd - SrcStart;
 			FMemory::Memcpy(Dst, SrcStart, Count*sizeof(float));
@@ -406,8 +436,8 @@ void FNiagaraDataBuffer::CopyTo(FNiagaraDataBuffer& DestBuffer, int32 InStartIdx
 		}
 		for (uint32 CompIdx = 0; CompIdx < Owner->TotalInt32Components; ++CompIdx)
 		{
-			int32* SrcStart = GetInstancePtrInt32(CompIdx, InStartIdx);
-			int32* SrcEnd = GetInstancePtrInt32(CompIdx, InStartIdx + InNumInstances);
+			const int32* SrcStart = GetInstancePtrInt32(CompIdx, InStartIdx);
+			const int32* SrcEnd = GetInstancePtrInt32(CompIdx, InStartIdx + InNumInstances);
 			int32* Dst = DestBuffer.GetInstancePtrInt32(CompIdx, 0);
 			size_t Count = SrcEnd - SrcStart;
 			FMemory::Memcpy(Dst, SrcStart, Count * sizeof(int32));
@@ -424,7 +454,7 @@ void FNiagaraDataBuffer::CopyTo(FNiagaraDataBuffer& DestBuffer, int32 InStartIdx
 	}
 }
 
-void FNiagaraDataBuffer::CopyTo(FNiagaraDataBuffer& DestBuffer)
+void FNiagaraDataBuffer::CopyTo(FNiagaraDataBuffer& DestBuffer)const
 {
 	DestBuffer.FloatStride = FloatStride;
 	DestBuffer.FloatData = FloatData;
diff --git a/Engine/Plugins/FX/Niagara/Source/Niagara/Private/NiagaraEditorDataBase.cpp b/Engine/Plugins/FX/Niagara/Source/Niagara/Private/NiagaraEditorDataBase.cpp
--- a/Engine/Plugins/FX/Niagara/Source/Niagara/Private/NiagaraEditorDataBase.cpp
+++ b/Engine/Plugins/FX/Niagara/Source/Niagara/Private/NiagaraEditorDataBase.cpp
@@ -0,0 +1 @@
+// Copyright 1998-2018 Epic Games, Inc. All Rights Reserved.
diff --git a/Engine/Plugins/FX/Niagara/Source/Niagara/Private/NiagaraEmitter.cpp b/Engine/Plugins/FX/Niagara/Source/Niagara/Private/NiagaraEmitter.cpp
--- a/Engine/Plugins/FX/Niagara/Source/Niagara/Private/NiagaraEmitter.cpp
+++ b/Engine/Plugins/FX/Niagara/Source/Niagara/Private/NiagaraEmitter.cpp
@@ -97,8 +97,6 @@ void UNiagaraEmitter::PostInitProperties()
 	Super::PostInitProperties();
 	if (HasAnyFlags(RF_ClassDefaultObject | RF_NeedLoad) == false)
 	{
-		RendererProperties.Add(NewObject<UNiagaraSpriteRendererProperties>(this, "Renderer"));
-
 		SpawnScriptProps.Script = NewObject<UNiagaraScript>(this, "SpawnScript", EObjectFlags::RF_Transactional);
 		SpawnScriptProps.Script->SetUsage(ENiagaraScriptUsage::ParticleSpawnScript);
 
diff --git a/Engine/Plugins/FX/Niagara/Source/Niagara/Private/NiagaraEmitterHandle.cpp b/Engine/Plugins/FX/Niagara/Source/Niagara/Private/NiagaraEmitterHandle.cpp
--- a/Engine/Plugins/FX/Niagara/Source/Niagara/Private/NiagaraEmitterHandle.cpp
+++ b/Engine/Plugins/FX/Niagara/Source/Niagara/Private/NiagaraEmitterHandle.cpp
@@ -103,6 +103,7 @@ void FNiagaraEmitterHandle::SetName(FName InName, UNiagaraSystem& InOwnerSystem)
 	FString TempNameString = InName.ToString();
 	TempNameString.ReplaceInline(TEXT(" "), TEXT("_"));
 	TempNameString.ReplaceInline(TEXT("\t"), TEXT("_"));
+	TempNameString.ReplaceInline(TEXT("."), TEXT("_"));
 	InName = *TempNameString;
 
 	TSet<FName> OtherEmitterNames;
diff --git a/Engine/Plugins/FX/Niagara/Source/Niagara/Private/NiagaraEmitterInstance.cpp b/Engine/Plugins/FX/Niagara/Source/Niagara/Private/NiagaraEmitterInstance.cpp
--- a/Engine/Plugins/FX/Niagara/Source/Niagara/Private/NiagaraEmitterInstance.cpp
+++ b/Engine/Plugins/FX/Niagara/Source/Niagara/Private/NiagaraEmitterInstance.cpp
@@ -32,6 +32,17 @@ static FAutoConsoleVariableRef CVarNiagaraDumpParticleData(
 	ECVF_Default
 	);
 
+/**
+TODO: This is mainly to avoid hard limits in our storage/alloc code etc rather than for perf reasons.
+We should improve our hard limit/safety code and possibly add a max for perf reasons.
+*/
+static int32 GMaxNiagaraCPUParticlesPerEmitter = 1000000;
+static FAutoConsoleVariableRef CVarMaxNiagaraCPUParticlesPerEmitter(
+	TEXT("fx.MaxNiagaraCPUParticlesPerEmitter"),
+	GMaxNiagaraCPUParticlesPerEmitter,
+	TEXT("The max number of supported CPU particles per emitter in Niagara. \n"),
+	ECVF_Default
+);
 //////////////////////////////////////////////////////////////////////////
 
 const FName FNiagaraEmitterInstance::PositionName(TEXT("Position"));
@@ -103,6 +114,20 @@ bool FNiagaraEmitterInstance::IsReadyToRun() const
 	return true;
 }
 
+void FNiagaraEmitterInstance::Dump()const
+{
+	UE_LOG(LogNiagara, Log, TEXT("==  %s ========"), *CachedEmitter->GetUniqueEmitterName());
+	UE_LOG(LogNiagara, Log, TEXT(".................Spawn................."));
+	SpawnExecContext.Parameters.DumpParameters(true);
+	UE_LOG(LogNiagara, Log, TEXT(".................Update................."));
+	UpdateExecContext.Parameters.DumpParameters(true);
+	UE_LOG(LogNiagara, Log, TEXT("................. %s Combined Parameters ................."), TEXT("GPU Script"));
+	GPUExecContext.CombinedParamStore.DumpParameters();
+	UE_LOG(LogNiagara, Log, TEXT("................. Particles ................."));
+	ParticleDataSet->Dump(false);
+	ParticleDataSet->Dump(true);
+}
+
 void FNiagaraEmitterInstance::Init(int32 InEmitterIdx, FName InSystemInstanceName)
 {
 	check(ParticleDataSet);
@@ -617,6 +642,7 @@ FBox FNiagaraEmitterInstance::CalculateDynamicBounds()
 				UE_LOG(LogNiagara, Warning, TEXT("Particle position data contains NaNs. Likely a divide by zero somewhere in your modules. Emitter \"%s\" in System \"%s\""),
 					*CachedEmitter->GetName(), *ParentSystemInstance->GetSystem()->GetName());
 				bEncounteredNaNs = true;
+				ParentSystemInstance->Dump();
 			}
 #endif
 		}
@@ -987,6 +1013,18 @@ void FNiagaraEmitterInstance::Tick(float DeltaSeconds)
 	}
 
 	int32 AllocationSize = OrigNumParticles + SpawnTotal + EventSpawnTotal;
+
+	//Ensure we don't blow our current hard limits on cpu particle count.
+	//TODO: These current limits can be improved relatively easily. Though perf in at these counts will obviously be an issue anyway.
+	if (CachedEmitter->SimTarget == ENiagaraSimTarget::CPUSim && AllocationSize > GMaxNiagaraCPUParticlesPerEmitter)
+	{
+		UE_LOG(LogNiagara, Warning, TEXT("Emitter %s has attemted to exceed the max CPU particle count! | Max: %d | Requested: %u"), *CachedEmitter->GetUniqueEmitterName(), GMaxNiagaraCPUParticlesPerEmitter, AllocationSize);
+		//For now we completely bail out of spawning new particles. Possibly should improve this in future.
+		AllocationSize = OrigNumParticles;
+		SpawnTotal = 0;
+		EventSpawnTotal = 0;
+	}
+
 	//Allocate space for prev frames particles and any new one's we're going to spawn.
 	Data.Allocate(AllocationSize);
 	for (FNiagaraDataSet* SpawnEventDataSet : SpawnScriptEventDataSets)
@@ -1260,62 +1298,63 @@ void FNiagaraEmitterInstance::Tick(float DeltaSeconds)
 			}
 		}
 
-
-		// handle single-particle events
-		// TODO: we'll need a way to either skip execution of the VM if an index comes back as invalid, or we'll have to pre-process
-		// event/particle arrays; this is currently a very naive (and comparatively slow) implementation, until full indexed reads work
-		if (EventHandlerProps.Script && EventHandlerProps.ExecutionMode == EScriptExecutionMode::SingleParticle && EventSet[EventScriptIdx])
-		{
-
-			SCOPE_CYCLE_COUNTER(STAT_NiagaraEventHandle);
-			FNiagaraVariable IndexVar(FNiagaraTypeDefinition::GetIntDef(), "ParticleIndex");
-			FNiagaraDataSetIterator<int32> IndexItr(*EventSet[EventScriptIdx], IndexVar, 0, false);
-			if (IndexItr.IsValid() && EventSet[EventScriptIdx]->GetPrevNumInstances() > 0)
-			{
-				EventExecCountBindings[EventScriptIdx].SetValue(1);
-
-				Data.CopyCurToPrev();
-				uint32 NumParticles = Data.GetNumInstances();
-
-				for (uint32 i = 0; i < EventSet[EventScriptIdx]->GetPrevNumInstances(); i++)
-				{
-					int32 Index = *IndexItr;
-					IndexItr.Advance();
-					DataSetExecInfos.SetNum(1, false);
-					DataSetExecInfos[0].StartInstance = Index;
-					DataSetExecInfos[0].bUpdateInstanceCount = false;
-					DataSetExecInfos.Emplace(EventSet[EventScriptIdx], i, false, false);
-					EventExecContexts[EventScriptIdx].Execute(1, DataSetExecInfos);
-
-					if (GbDumpParticleData || System->bDumpDebugEmitterInfo)
-					{
-						ensure(EventHandlerProps.Script->RapidIterationParameters.VerifyBinding(&EventExecContexts[EventScriptIdx].Parameters));
-						UE_LOG(LogNiagara, Log, TEXT("=== Event %d Src Parameters ==="), EventScriptIdx);
-						EventHandlerProps.Script->RapidIterationParameters.Dump();
-						UE_LOG(LogNiagara, Log, TEXT("=== Event %d Context Parameters ==="), EventScriptIdx);
-						EventExecContexts[EventScriptIdx].Parameters.Dump();
-						UE_LOG(LogNiagara, Log, TEXT("=== Event %d Particles (%d index written, %d total) ==="), EventScriptIdx, Index, Data.GetNumInstances());
-						Data.Dump(true, Index, 1);
-					}
-
-
-#if WITH_EDITORONLY_DATA
-					if (ParentSystemInstance->ShouldCaptureThisFrame())
-					{
-						FGuid EventGuid = EventExecContexts[EventScriptIdx].Script->GetUsageId();
-						FNiagaraScriptDebuggerInfo* DebugInfo = ParentSystemInstance->GetActiveCaptureWrite(CachedIDName, ENiagaraScriptUsage::ParticleEventScript, EventGuid);
-						if (DebugInfo)
-						{
-							Data.Dump(DebugInfo->Frame, true, Index, 1);
-							//DebugInfo->Frame.Dump(true, 0, 1);
-							DebugInfo->Parameters = EventExecContexts[EventScriptIdx].Parameters;
-						}
-					}
-#endif
-					ensure(NumParticles == Data.GetNumInstances());
-				}
-			}
-		}
+		//TODO: Disabling this event mode for now until it can be reworked. Currently it uses index directly with can easily be invalid and cause undefined behavior.
+		//
+//		// handle single-particle events
+//		// TODO: we'll need a way to either skip execution of the VM if an index comes back as invalid, or we'll have to pre-process
+//		// event/particle arrays; this is currently a very naive (and comparatively slow) implementation, until full indexed reads work
+// 		if (EventHandlerProps.Script && EventHandlerProps.ExecutionMode == EScriptExecutionMode::SingleParticle && EventSet[EventScriptIdx])
+// 		{
+// 
+// 			SCOPE_CYCLE_COUNTER(STAT_NiagaraEventHandle);
+// 			FNiagaraVariable IndexVar(FNiagaraTypeDefinition::GetIntDef(), "ParticleIndex");
+// 			FNiagaraDataSetIterator<int32> IndexItr(*EventSet[EventScriptIdx], IndexVar, 0, false);
+// 			if (IndexItr.IsValid() && EventSet[EventScriptIdx]->GetPrevNumInstances() > 0)
+// 			{
+// 				EventExecCountBindings[EventScriptIdx].SetValue(1);
+// 
+// 				Data.CopyCurToPrev();
+// 				uint32 NumParticles = Data.GetNumInstances();
+// 
+// 				for (uint32 i = 0; i < EventSet[EventScriptIdx]->GetPrevNumInstances(); i++)
+// 				{
+// 					int32 Index = *IndexItr;
+// 					IndexItr.Advance();
+// 					DataSetExecInfos.SetNum(1, false);
+// 					DataSetExecInfos[0].StartInstance = Index;
+// 					DataSetExecInfos[0].bUpdateInstanceCount = false;
+// 					DataSetExecInfos.Emplace(EventSet[EventScriptIdx], i, false, false);
+// 					EventExecContexts[EventScriptIdx].Execute(1, DataSetExecInfos);
+// 
+// 					if (GbDumpParticleData || System->bDumpDebugEmitterInfo)
+// 					{
+// 						ensure(EventHandlerProps.Script->RapidIterationParameters.VerifyBinding(&EventExecContexts[EventScriptIdx].Parameters));
+// 						UE_LOG(LogNiagara, Log, TEXT("=== Event %d Src Parameters ==="), EventScriptIdx);
+// 						EventHandlerProps.Script->RapidIterationParameters.Dump();
+// 						UE_LOG(LogNiagara, Log, TEXT("=== Event %d Context Parameters ==="), EventScriptIdx);
+// 						EventExecContexts[EventScriptIdx].Parameters.Dump();
+// 						UE_LOG(LogNiagara, Log, TEXT("=== Event %d Particles (%d index written, %d total) ==="), EventScriptIdx, Index, Data.GetNumInstances());
+// 						Data.Dump(true, Index, 1);
+// 					}
+// 
+// 
+// #if WITH_EDITORONLY_DATA
+// 					if (ParentSystemInstance->ShouldCaptureThisFrame())
+// 					{
+// 						FGuid EventGuid = EventExecContexts[EventScriptIdx].Script->GetUsageId();
+// 						FNiagaraScriptDebuggerInfo* DebugInfo = ParentSystemInstance->GetActiveCaptureWrite(CachedIDName, ENiagaraScriptUsage::ParticleEventScript, EventGuid);
+// 						if (DebugInfo)
+// 						{
+// 							Data.Dump(DebugInfo->Frame, true, Index, 1);
+// 							//DebugInfo->Frame.Dump(true, 0, 1);
+// 							DebugInfo->Parameters = EventExecContexts[EventScriptIdx].Parameters;
+// 						}
+// 					}
+// #endif
+// 					ensure(NumParticles == Data.GetNumInstances());
+// 				}
+// 			}
+// 		}
 	}
 
 	PostProcessParticles();
diff --git a/Engine/Plugins/FX/Niagara/Source/Niagara/Private/NiagaraEmitterInstanceBatcher.cpp b/Engine/Plugins/FX/Niagara/Source/Niagara/Private/NiagaraEmitterInstanceBatcher.cpp
--- a/Engine/Plugins/FX/Niagara/Source/Niagara/Private/NiagaraEmitterInstanceBatcher.cpp
+++ b/Engine/Plugins/FX/Niagara/Source/Niagara/Private/NiagaraEmitterInstanceBatcher.cpp
@@ -27,14 +27,18 @@ uint32 FNiagaraComputeExecutionContext::TickCounter = 0;
 void NiagaraEmitterInstanceBatcher::Queue(FNiagaraComputeExecutionContext *InContext)
 {
 	//SimulationQueue[CurQueueIndex]->Add(InContext);
-
-	ENQUEUE_UNIQUE_RENDER_COMMAND_THREEPARAMETER(QueueNiagaraDispatch,
-		TArray<FNiagaraComputeExecutionContext*>*, Queue, &SimulationQueue[0],
-		uint32, QueueIndex, CurQueueIndex,
-		FNiagaraComputeExecutionContext*, ExecContext, InContext,
-		{
-			Queue[QueueIndex].Add(ExecContext);
-		});
+		ENQUEUE_UNIQUE_RENDER_COMMAND_THREEPARAMETER(QueueNiagaraDispatch,
+			TArray<FNiagaraComputeExecutionContext*>*, Queue, &SimulationQueue[0],
+			uint32, QueueIndex, CurQueueIndex,
+			FNiagaraComputeExecutionContext*, ExecContext, InContext,
+			{
+				//Don't queue the same context for execution multiple times. TODO: possibly try to combine/accumulate the tick info if we happen to have > 1 before it's executed.
+				if (!ExecContext->bPendingExecution)
+				{
+					Queue[QueueIndex].Add(ExecContext);
+					ExecContext->bPendingExecution = true;
+				}
+			});
 }
 
 
@@ -56,6 +60,7 @@ void NiagaraEmitterInstanceBatcher::TickSingle(FNiagaraComputeExecutionContext *
 
 	check(IsInRenderingThread());
 	Context->MainDataSet->Tick();
+	Context->bPendingExecution = false;
 
 	FNiagaraComputeExecutionContext::TickCounter++;
 
diff --git a/Engine/Plugins/FX/Niagara/Source/Niagara/Private/NiagaraFunctionLibrary.cpp b/Engine/Plugins/FX/Niagara/Source/Niagara/Private/NiagaraFunctionLibrary.cpp
--- a/Engine/Plugins/FX/Niagara/Source/Niagara/Private/NiagaraFunctionLibrary.cpp
+++ b/Engine/Plugins/FX/Niagara/Source/Niagara/Private/NiagaraFunctionLibrary.cpp
@@ -29,6 +29,9 @@ UNiagaraComponent* UNiagaraFunctionLibrary::SpawnSystemAtLocation(UObject* World
 		{
 			AActor* Actor = World->GetWorldSettings();
 			PSC = NewObject<UNiagaraComponent>((Actor ? Actor : (UObject*)World));
+#if WITH_EDITORONLY_DATA
+			PSC->bWaitForCompilationOnActivate = true;
+#endif
 			PSC->SetAsset(SystemTemplate);
 			PSC->SetAutoDestroy(bAutoDestroy);
 			PSC->bAutoActivate = false;
diff --git a/Engine/Plugins/FX/Niagara/Source/Niagara/Private/NiagaraModule.cpp b/Engine/Plugins/FX/Niagara/Source/Niagara/Private/NiagaraModule.cpp
--- a/Engine/Plugins/FX/Niagara/Source/Niagara/Private/NiagaraModule.cpp
+++ b/Engine/Plugins/FX/Niagara/Source/Niagara/Private/NiagaraModule.cpp
@@ -16,6 +16,7 @@
 #include "NiagaraMeshRendererProperties.h"
 #include "NiagaraRibbonRendererProperties.h"
 #include "NiagaraRenderer.h"
+#include "Misc/CoreDelegates.h"
 
 IMPLEMENT_MODULE(INiagaraModule, Niagara);
 
@@ -225,6 +226,12 @@ void INiagaraModule::StartupModule()
 	UNiagaraSpriteRendererProperties::InitCDOPropertiesAfterModuleStartup();
 	UNiagaraRibbonRendererProperties::InitCDOPropertiesAfterModuleStartup();
 	UNiagaraMeshRendererProperties::InitCDOPropertiesAfterModuleStartup();
+
+}
+
+void INiagaraModule::ShutdownRenderingResources()
+{
+	FNiagaraViewDataMgr::Shutdown();
 }
 
 void INiagaraModule::ShutdownModule()
@@ -238,7 +245,7 @@ void INiagaraModule::ShutdownModule()
 	}
 
 	CVarDetailLevel.AsVariable()->SetOnChangedCallback(FConsoleVariableDelegate());
-	FNiagaraViewDataMgr::Shutdown();
+	ShutdownRenderingResources();
 }
 
 FNiagaraWorldManager* INiagaraModule::GetWorldManager(UWorld* World)
diff --git a/Engine/Plugins/FX/Niagara/Source/Niagara/Private/NiagaraParameterStore.cpp b/Engine/Plugins/FX/Niagara/Source/Niagara/Private/NiagaraParameterStore.cpp
--- a/Engine/Plugins/FX/Niagara/Source/Niagara/Private/NiagaraParameterStore.cpp
+++ b/Engine/Plugins/FX/Niagara/Source/Niagara/Private/NiagaraParameterStore.cpp
@@ -80,6 +80,7 @@ FNiagaraParameterStore::~FNiagaraParameterStore()
 	{
 		Binding.Value.Empty(Binding.Key, this);
 	}
+	Bindings.Empty();
 }
 
 void FNiagaraParameterStore::Bind(FNiagaraParameterStore* DestStore)
@@ -138,9 +139,54 @@ bool FNiagaraParameterStore::VerifyBinding(const FNiagaraParameterStore* DestSto
 #endif
 }
 
+void FNiagaraParameterStore::CheckForNaNs()const
+{
+	for (const TPair<FNiagaraVariable, int32>& ParamOffset : ParameterOffsets)
+	{
+		const FNiagaraVariable& Var = ParamOffset.Key;
+		int32 Offset = ParamOffset.Value;
+		bool bContainsNans = false;
+		if (Var.GetType() == FNiagaraTypeDefinition::GetFloatDef())
+		{
+			float Val = *(float*)GetParameterData(Offset);
+			bContainsNans = FMath::IsNaN(Val) || !FMath::IsFinite(Val);
+		}
+		else if (Var.GetType() == FNiagaraTypeDefinition::GetVec2Def())
+		{
+			FVector2D Val = *(FVector2D*)GetParameterData(Offset);
+			bContainsNans = Val.ContainsNaN();
+		}
+		else if (Var.GetType() == FNiagaraTypeDefinition::GetVec3Def())
+		{
+			FVector Val = *(FVector*)GetParameterData(Offset);
+			bContainsNans = Val.ContainsNaN();
+		}
+		else if (Var.GetType() == FNiagaraTypeDefinition::GetVec4Def())
+		{
+			FVector4 Val = *(FVector4*)GetParameterData(Offset);
+			bContainsNans = Val.ContainsNaN();
+		}
+		else if (Var.GetType() == FNiagaraTypeDefinition::GetMatrix4Def())
+		{
+			FMatrix Val;
+			FMemory::Memcpy(&Val, GetParameterData(Offset), sizeof(FMatrix));
+			bContainsNans = Val.ContainsNaN();
+		}
+
+		if (bContainsNans)
+		{
+			ensureAlwaysMsgf(false, TEXT("Niagara Parameter Store containts Nans!\n"));
+			DumpParameters(false);
+		}
+	}
+}
+
 void FNiagaraParameterStore::Tick()
 {
 	SCOPE_CYCLE_COUNTER(STAT_NiagaraParameterStoreTick);
+#if NIAGARA_NAN_CHECKING
+	CheckForNaNs();
+#endif
 	if (bParametersDirty || bInterfacesDirty)
 	{
 		for (TPair<FNiagaraParameterStore*, FNiagaraParameterStoreBinding>& Binding : Bindings)
@@ -167,7 +213,7 @@ void FNiagaraParameterStore::UnbindFromSourceStores()
 	ensureMsgf(SourceStores.Num() == 0, TEXT("Parameter store source array was not empty after unbinding all sources. Something seriously wrong."));
 }
 
-void FNiagaraParameterStore::DumpParameters(bool bDumpBindings)
+void FNiagaraParameterStore::DumpParameters(bool bDumpBindings)const
 {
 	TArray<FNiagaraVariable> Vars;
 	GetParameters(Vars);
@@ -179,7 +225,7 @@ void FNiagaraParameterStore::DumpParameters(bool bDumpBindings)
 
 	if (bDumpBindings)
 	{
-		for (TPair<FNiagaraParameterStore*, FNiagaraParameterStoreBinding>& Binding : Bindings)
+		for (const TPair<FNiagaraParameterStore*, FNiagaraParameterStoreBinding>& Binding : Bindings)
 		{
 			Binding.Value.Dump(Binding.Key, this);
 		}
@@ -446,6 +492,7 @@ void FNiagaraParameterStore::Empty(bool bClearBindings)
 	DataInterfaces.Empty();
 	if (bClearBindings)
 	{
+		UnbindFromSourceStores();
 		Bindings.Empty();
 	}
 }
@@ -460,6 +507,7 @@ void FNiagaraParameterStore::Reset(bool bClearBindings)
 	DataInterfaces.Reset();
 	if (bClearBindings)
 	{
+		UnbindFromSourceStores();
 		Bindings.Reset();
 	}
 }
diff --git a/Engine/Plugins/FX/Niagara/Source/Niagara/Private/NiagaraRendererRibbons.cpp b/Engine/Plugins/FX/Niagara/Source/Niagara/Private/NiagaraRendererRibbons.cpp
--- a/Engine/Plugins/FX/Niagara/Source/Niagara/Private/NiagaraRendererRibbons.cpp
+++ b/Engine/Plugins/FX/Niagara/Source/Niagara/Private/NiagaraRendererRibbons.cpp
@@ -270,6 +270,57 @@ bool NiagaraRendererRibbons::SetMaterialUsage()
 	return Material && Material->CheckMaterialUsage_Concurrent(MATUSAGE_NiagaraRibbons);
 }
 
+void CalculateUVScaleAndOffsets(const FNiagaraDataSetAccessor<float>& SortKeyData, const TArray<int32>& RibbonIndices, bool bSortKeyIsAge, int32 NumSegments,
+	float InUVTilingDistance, const FVector2D& InUVScale, const FVector2D& InUVOffset, ENiagaraRibbonAgeOffsetMode InAgeOffsetMode, FVector2D& OutUVScale, FVector2D& OutUVOffset)
+{
+	if (bSortKeyIsAge && InUVTilingDistance == 0)
+	{
+		float AgeUScale;
+		float AgeUOffset;
+		if (InAgeOffsetMode == ENiagaraRibbonAgeOffsetMode::Scale)
+		{
+			// In scale mode we scale and offset the UVs so that no part of the texture is clipped. In order to prevent
+			// clipping at the ends we'll have to move the UVs in up to the size of a single segment of the ribbon since
+			// that's the distance we'll need to to smoothly interpolate when a new segment is added, or when an old segment
+			// is removed.  We calculate the end offset when the end of the ribbon is within a single time step of 0 or 1
+			// which is then normalized to the range of a single segment.  We can then calculate how many segments we actually
+			// have to draw the scaled ribbon, and can offset the start by the correctly scaled offset.
+			float FirstAge = SortKeyData[RibbonIndices[0]];
+			float SecondAge = SortKeyData[RibbonIndices[1]];
+			float SecondToLastAge = SortKeyData[RibbonIndices[RibbonIndices.Num() - 2]];
+			float LastAge = SortKeyData[RibbonIndices.Last()];
+
+			float StartTimeStep = SecondAge - FirstAge;
+			float StartTimeOffset = FirstAge < StartTimeStep ? StartTimeStep - FirstAge : 0;
+			float StartSegmentOffset = StartTimeOffset / StartTimeStep;
+
+			float EndTimeStep = LastAge - SecondToLastAge;
+			float EndTimeOffset = 1 - LastAge < EndTimeStep ? EndTimeStep - (1 - LastAge) : 0;
+			float EndSegmentOffset = EndTimeOffset / EndTimeStep;
+
+			float AvailableSegments = NumSegments - (StartSegmentOffset + EndSegmentOffset);
+			AgeUScale = NumSegments / AvailableSegments;
+			AgeUOffset = -((StartSegmentOffset / NumSegments) * AgeUScale);
+		}
+		else
+		{
+			float FirstAge = SortKeyData[RibbonIndices[0]];
+			float LastAge = SortKeyData[RibbonIndices.Last()];
+
+			AgeUScale = LastAge - FirstAge;
+			AgeUOffset = FirstAge;
+		}
+
+		OutUVScale = FVector2D(AgeUScale * InUVScale.X, InUVScale.Y);
+		OutUVOffset = FVector2D((AgeUOffset * InUVScale.X) + InUVOffset.X, InUVOffset.Y);
+	}
+	else
+	{
+		OutUVScale = InUVScale;
+		OutUVOffset = InUVOffset;
+	}
+}
+
 FNiagaraDynamicDataBase *NiagaraRendererRibbons::GenerateVertexData(const FNiagaraSceneProxy* Proxy, FNiagaraDataSet &Data, const ENiagaraSimTarget Target)
 {
 	SCOPE_CYCLE_COUNTER(STAT_NiagaraGenRibbonVertexData);
@@ -288,7 +339,7 @@ FNiagaraDynamicDataBase *NiagaraRendererRibbons::GenerateVertexData(const FNiaga
 
 	// TODO : deal with the dynamic vertex material parameter should the user have specified it as an output...
 
-	FVector2D UVs[4] = { FVector2D(1.0f, 0.0f), FVector2D(1.0f, 1.0f) };
+	FVector2D UVs[2] = { FVector2D(1.0f, 1.0f), FVector2D(1.0f, 0.0f) };
 	int32 NumTotalVerts = 0;
 
 	FNiagaraDataSetAccessor<FVector> PosData(Data, Properties->PositionBinding.DataSetVariable);
@@ -364,15 +415,17 @@ FNiagaraDynamicDataBase *NiagaraRendererRibbons::GenerateVertexData(const FNiaga
 			FVector PrevPos, PrevPos2, PrevDir(0.0f, 0.0f, 0.1f);
 			float TotalDistance = 0.0f;
 
-			FVector2D AgeOffset;
-			if (bSortKeyIsAge)
-			{
-				AgeOffset = FVector2D(SortKeyData[RibbonIndices.Last()], 0.0f);
-			}
-			else
-			{
-				AgeOffset = FVector2D::ZeroVector;
-			}
+			int32 NumSegments = NumIndices - 1;
+
+			FVector2D UV0Offset;
+			FVector2D UV0Scale;
+			FVector2D UV1Offset;
+			FVector2D UV1Scale;
+
+			CalculateUVScaleAndOffsets(SortKeyData, RibbonIndices, bSortKeyIsAge, NumSegments, Properties->UV0TilingDistance,
+				Properties->UV0Scale, Properties->UV0Offset, Properties->UV0AgeOffsetMode, UV0Scale, UV0Offset);
+			CalculateUVScaleAndOffsets(SortKeyData, RibbonIndices, bSortKeyIsAge, NumSegments, Properties->UV1TilingDistance,
+				Properties->UV1Scale, Properties->UV1Offset, Properties->UV1AgeOffsetMode, UV1Scale, UV1Offset);
 
 			for (int32 i = 0; i < NumIndices; i++)
 			{
@@ -397,18 +450,18 @@ FNiagaraDynamicDataBase *NiagaraRendererRibbons::GenerateVertexData(const FNiaga
 				{
 					FVector NormDir = ParticleDir.GetSafeNormal();
 					PrevDir = NormDir;
-					FVector2D UV0Mult((float)i / NumIndices, 1.0f);
-					FVector2D UV1Mult((float)i / NumIndices, 1.0f);
-
-					if (Properties->UV0TilingDistance)
-					{
-						UV0Mult = FVector2D(TotalDistance / Properties->UV0TilingDistance, 1.0f);
-					}
-					if (Properties->UV1TilingDistance)
-					{
-						UV1Mult = FVector2D(TotalDistance / Properties->UV1TilingDistance, 1.0f);
-					}
 
+					// Calculate sub-uv for this segment based on the either the distance or index, depending on whether the distance tile factor was set.
+					float U0ForSegment = Properties->UV0TilingDistance
+						? TotalDistance / Properties->UV0TilingDistance
+						: (float)i / NumSegments;
+					float U1ForSegment = Properties->UV1TilingDistance
+						? TotalDistance / Properties->UV1TilingDistance
+						: (float)i / NumSegments;
+
+					FVector2D UV0ForSegment(U0ForSegment, 1.0f);
+					FVector2D UV1ForSegment(U1ForSegment, 1.0f);
+					
 					//Todo. Possibly pull out to templated func to avoid safety code in loop here.
 					FLinearColor Color = ColData.GetSafe(Index1, FLinearColor::White);
 					float Twist = TwistData.GetSafe(Index1, 0.0f);
@@ -416,8 +469,8 @@ FNiagaraDynamicDataBase *NiagaraRendererRibbons::GenerateVertexData(const FNiaga
 					FVector Align = AlignData.GetSafe(Index1, FVector::UpVector);
 
 					TotalDistance += ParticleDir.Size();
-					AddRibbonVert(RenderData, ParticlePos, UVs[0] * UV0Mult + AgeOffset, UVs[0] * UV1Mult + AgeOffset, Color, Twist, Size, NormDir, Align);
-					AddRibbonVert(RenderData, ParticlePos, UVs[1] * UV0Mult + AgeOffset, UVs[1] * UV1Mult + AgeOffset, Color, Twist, Size, NormDir, Align);
+					AddRibbonVert(RenderData, ParticlePos, UVs[0] * UV0ForSegment * UV0Scale + UV0Offset, UVs[0] * UV1ForSegment * UV1Scale + UV1Offset, Color, Twist, Size, NormDir, Align);
+					AddRibbonVert(RenderData, ParticlePos, UVs[1] * UV0ForSegment * UV0Scale + UV0Offset, UVs[1] * UV1ForSegment * UV1Scale + UV1Offset, Color, Twist, Size, NormDir, Align);
 					if (MaterialParamData.IsValid())
 					{
 						AddDynamicParam(DynamicData->MaterialParameterVertexData[0], MaterialParamData[Index1]);
diff --git a/Engine/Plugins/FX/Niagara/Source/Niagara/Private/NiagaraScript.cpp b/Engine/Plugins/FX/Niagara/Source/Niagara/Private/NiagaraScript.cpp
--- a/Engine/Plugins/FX/Niagara/Source/Niagara/Private/NiagaraScript.cpp
+++ b/Engine/Plugins/FX/Niagara/Source/Niagara/Private/NiagaraScript.cpp
@@ -191,6 +191,22 @@ UNiagaraScript::~UNiagaraScript()
 }
 
 #if WITH_EDITORONLY_DATA
+class UNiagaraSystem* UNiagaraScript::FindRootSystem()
+{
+	UObject* Obj = GetOuter();
+	if (UNiagaraEmitter* Emitter = Cast<UNiagaraEmitter>(Obj))
+	{
+		Obj = Emitter->GetOuter();
+	}
+
+	if (UNiagaraSystem* Sys = Cast<UNiagaraSystem>(Obj))
+	{
+		return Sys;
+	}
+
+	return nullptr;
+}
+
 void UNiagaraScript::ComputeVMCompilationId(FNiagaraVMExecutableDataId& Id) const
 {
 	Id = FNiagaraVMExecutableDataId();
@@ -768,8 +784,8 @@ void UNiagaraScript::SetVMCompilationResults(const FNiagaraVMExecutableDataId& I
 void UNiagaraScript::InvalidateExecutionReadyParameterStores()
 {
 	// Make sure that we regenerate any parameter stores, since they must be kept in sync with the layout from script compilation.
-	ScriptExecutionParamStoreCPU.Clear();
-	ScriptExecutionParamStoreGPU.Clear();
+	ScriptExecutionParamStoreCPU.Empty();
+	ScriptExecutionParamStoreGPU.Empty();
 }
 
 void UNiagaraScript::InvalidateCachedCompileIds()
@@ -908,30 +924,37 @@ void UNiagaraScript::BeginCacheForCookedPlatformData(const ITargetPlatform *Targ
 {
 	if (CanBeRunOnGpu())
 	{
+		// Commandlets like DerivedDataCacheCommandlet call BeginCacheForCookedPlatformData directly on objects. This may mean that
+		// we have not properly gotten the HLSL script generated by the time that we get here. This does the awkward work of 
+		// waiting on the parent system to finish generating the HLSL before we can begin compiling it for the GPU.
+		UNiagaraSystem* SystemOwner = FindRootSystem();
+		if (SystemOwner)
+		{
+			SystemOwner->WaitForCompilationComplete();
+		}
+
 		TArray<FName> DesiredShaderFormats;
 		TargetPlatform->GetAllTargetedShaderFormats(DesiredShaderFormats);
 
 		FNiagaraShaderScript** CachedScriptResourceForPlatformPtr = CachedScriptResourcesForCooking.Find(TargetPlatform);
 
-		if (CachedScriptResourceForPlatformPtr == nullptr)
+		if (DesiredShaderFormats.Num() > 0 && CachedScriptResourceForPlatformPtr == nullptr)
 		{
 			CachedScriptResourcesForCooking.Add(TargetPlatform);
 			FNiagaraShaderScript* &CachedScriptResourceForPlatform = *CachedScriptResourcesForCooking.Find(TargetPlatform);
 
-			if (DesiredShaderFormats.Num())
+			// Cache for all the shader formats that the cooking target requires
+			for (int32 FormatIndex = 0; FormatIndex < DesiredShaderFormats.Num(); FormatIndex++)
 			{
-				// Cache for all the shader formats that the cooking target requires
-				for (int32 FormatIndex = 0; FormatIndex < DesiredShaderFormats.Num(); FormatIndex++)
+				if (DesiredShaderFormats[FormatIndex] != TEXT("PCD3D_SM4"))	// TODO: remove once we get rid of SM4 globally
 				{
-					if (DesiredShaderFormats[FormatIndex] != TEXT("PCD3D_SM4"))	// TODO: remove once we get rid of SM4 globally
-					{
-						const EShaderPlatform LegacyShaderPlatform = ShaderFormatToLegacyShaderPlatform(DesiredShaderFormats[FormatIndex]);
+					const EShaderPlatform LegacyShaderPlatform = ShaderFormatToLegacyShaderPlatform(DesiredShaderFormats[FormatIndex]);
 
-						// Begin caching shaders for the target platform and store the FNiagaraShaderScript being compiled into CachedScriptResourcesForCooking
-						CacheResourceShadersForCooking(LegacyShaderPlatform, CachedScriptResourceForPlatform);
-					}
+					// Begin caching shaders for the target platform and store the FNiagaraShaderScript being compiled into CachedScriptResourcesForCooking
+					CacheResourceShadersForCooking(LegacyShaderPlatform, CachedScriptResourceForPlatform);
 				}
 			}
+			check(CachedScriptResourceForPlatform);
 		}
 	}
 }
@@ -1195,7 +1218,7 @@ void SerializeNiagaraShaderMaps(const TMap<const ITargetPlatform*, FNiagaraShade
 			auto& PlatformScriptResourcesToSave = *PlatformScriptResourcesToSavePtr;
 
 			ScriptResourceToSavePtr = PlatformScriptResourcesToSave.Find(Ar.CookingTarget());
-			check(ScriptResourceToSavePtr != NULL || (Ar.GetLinker() == NULL));
+			//check(ScriptResourceToSavePtr != NULL || (Ar.GetLinker() == NULL));
 			if (ScriptResourceToSavePtr != NULL)
 			{
 				NumResourcesToSave++;
@@ -1207,6 +1230,7 @@ void SerializeNiagaraShaderMaps(const TMap<const ITargetPlatform*, FNiagaraShade
 		if (ScriptResourceToSavePtr)
 		{
 			ScriptResourceToSave = *ScriptResourceToSavePtr;
+			check(ScriptResourceToSave);
 			ScriptResourceToSave->SerializeShaderMap(Ar);
 		}
 
diff --git a/Engine/Plugins/FX/Niagara/Source/Niagara/Private/NiagaraScriptExecutionContext.cpp b/Engine/Plugins/FX/Niagara/Source/Niagara/Private/NiagaraScriptExecutionContext.cpp
--- a/Engine/Plugins/FX/Niagara/Source/Niagara/Private/NiagaraScriptExecutionContext.cpp
+++ b/Engine/Plugins/FX/Niagara/Source/Niagara/Private/NiagaraScriptExecutionContext.cpp
@@ -148,6 +148,9 @@ bool FNiagaraScriptExecutionContext::Execute(uint32 NumInstances, TArray<FNiagar
 		SCOPE_CYCLE_COUNTER(STAT_NiagaraSimRegisterSetup);
 		for (FNiagaraDataSetExecutionInfo& DataSetInfo : DataSetInfos)
 		{
+#if NIAGARA_NAN_CHECKING
+			DataSetInfo.DataSet->CheckForNaNs();
+#endif
 			check(DataSetInfo.DataSet);
 			FDataSetMeta SetMeta(DataSetInfo.DataSet->GetSizeBytes(), &InputRegisters[NumInputRegisters], NumInputRegisters, DataSetInfo.StartInstance, 
 				&DataSetInfo.DataSet->CurrIDTable(), &DataSetInfo.DataSet->GetFreeIDTable(), &DataSetInfo.DataSet->GetNumFreeIDs(), &DataSetInfo.DataSet->GetMaxUsedID(), DataSetInfo.DataSet->GetIDAcquireTag());
@@ -185,6 +188,11 @@ bool FNiagaraScriptExecutionContext::Execute(uint32 NumInstances, TArray<FNiagar
 	for (int Idx = 0; Idx < DataSetInfos.Num(); Idx++)
 	{
 		FNiagaraDataSetExecutionInfo& Info = DataSetInfos[Idx];
+
+#if NIAGARA_NAN_CHECKING
+		Info.DataSet->CheckForNaNs();
+#endif
+
 		if (Info.bUpdateInstanceCount)
 		{
 			Info.DataSet->SetNumInstances(Info.StartInstance + DataSetMetaTable[Idx].DataSetAccessIndex + 1);
diff --git a/Engine/Plugins/FX/Niagara/Source/Niagara/Private/NiagaraScriptExecutionParameterStore.cpp b/Engine/Plugins/FX/Niagara/Source/Niagara/Private/NiagaraScriptExecutionParameterStore.cpp
--- a/Engine/Plugins/FX/Niagara/Source/Niagara/Private/NiagaraScriptExecutionParameterStore.cpp
+++ b/Engine/Plugins/FX/Niagara/Source/Niagara/Private/NiagaraScriptExecutionParameterStore.cpp
@@ -1,4 +1,4 @@
-// Copyright 1998-2016 Epic Games, Inc. All Rights Reserved.
+// Copyright 1998-2018 Epic Games, Inc. All Rights Reserved.
 
 #include "NiagaraScriptExecutionParameterStore.h"
 #include "NiagaraStats.h"
@@ -174,7 +174,7 @@ void FNiagaraScriptExecutionParameterStore::InitFromOwningScript(UNiagaraScript*
 
 void FNiagaraScriptExecutionParameterStore::InitFromOwningContext(UNiagaraScript* Script, ENiagaraSimTarget SimTarget, bool bNotifyAsDirty)
 {
-	Clear();
+	Empty();
 	PaddingInfo.Empty();
 
 #if WITH_EDITORONLY_DATA
diff --git a/Engine/Plugins/FX/Niagara/Source/Niagara/Private/NiagaraSystem.cpp b/Engine/Plugins/FX/Niagara/Source/Niagara/Private/NiagaraSystem.cpp
--- a/Engine/Plugins/FX/Niagara/Source/Niagara/Private/NiagaraSystem.cpp
+++ b/Engine/Plugins/FX/Niagara/Source/Niagara/Private/NiagaraSystem.cpp
@@ -445,6 +445,18 @@ bool UNiagaraSystem::ReferencesSourceEmitter(UNiagaraEmitter& Emitter)
 	return false;
 }
 
+bool UNiagaraSystem::ReferencesInstanceEmitter(UNiagaraEmitter& Emitter)
+{
+	for (FNiagaraEmitterHandle& Handle : EmitterHandles)
+	{
+		if (&Emitter == Handle.GetInstance())
+		{
+			return true;
+		}
+	}
+	return false;
+}
+
 void UNiagaraSystem::UpdateFromEmitterChanges(UNiagaraEmitter& ChangedSourceEmitter)
 {
 	bool bNeedsCompile = false;
@@ -801,6 +813,7 @@ bool UNiagaraSystem::QueryCompileComplete(bool bWait, bool bDoPost, bool bDoNotA
 
 		InitEmitterSpawnAttributes();
 
+		// Prepare rapid iteration parameters for execution.
 		TArray<UNiagaraScript*> Scripts;
 		TMap<UNiagaraScript*, UNiagaraScript*> ScriptDependencyMap;
 		TMap<UNiagaraScript*, FString> ScriptToEmitterNameMap;
@@ -812,21 +825,21 @@ bool UNiagaraSystem::QueryCompileComplete(bool bWait, bool bDoPost, bool bDoNotA
 			Scripts.AddUnique(CompiledScript);
 			ScriptToEmitterNameMap.Add(CompiledScript, Emitter != nullptr ? Emitter->GetUniqueEmitterName() : FString());
 
-			if (CompiledScript->GetUsage() == ENiagaraScriptUsage::EmitterSpawnScript)
+			if (UNiagaraScript::IsEquivalentUsage(CompiledScript->GetUsage(), ENiagaraScriptUsage::EmitterSpawnScript))
 			{
 				Scripts.AddUnique(SystemSpawnScript);
 				ScriptDependencyMap.Add(CompiledScript, SystemSpawnScript);
 				ScriptToEmitterNameMap.Add(SystemSpawnScript, FString());
 			}
 
-			if (CompiledScript->GetUsage() == ENiagaraScriptUsage::EmitterUpdateScript)
+			if (UNiagaraScript::IsEquivalentUsage(CompiledScript->GetUsage(), ENiagaraScriptUsage::EmitterUpdateScript))
 			{
 				Scripts.AddUnique(SystemUpdateScript);
 				ScriptDependencyMap.Add(CompiledScript, SystemUpdateScript);
 				ScriptToEmitterNameMap.Add(SystemSpawnScript, FString());
 			}
 
-			if (CompiledScript->GetUsage() == ENiagaraScriptUsage::ParticleSpawnScript)
+			if (UNiagaraScript::IsEquivalentUsage(CompiledScript->GetUsage(), ENiagaraScriptUsage::ParticleSpawnScript))
 			{
 				if (Emitter && Emitter->SimTarget == ENiagaraSimTarget::GPUComputeSim)
 				{
@@ -836,7 +849,7 @@ bool UNiagaraSystem::QueryCompileComplete(bool bWait, bool bDoPost, bool bDoNotA
 				}
 			}
 
-			if (CompiledScript->GetUsage() == ENiagaraScriptUsage::ParticleUpdateScript)
+			if (UNiagaraScript::IsEquivalentUsage(CompiledScript->GetUsage(), ENiagaraScriptUsage::ParticleUpdateScript))
 			{
 				if (Emitter && Emitter->SimTarget == ENiagaraSimTarget::GPUComputeSim)
 				{
@@ -855,6 +868,37 @@ bool UNiagaraSystem::QueryCompileComplete(bool bWait, bool bDoPost, bool bDoNotA
 
 		FNiagaraUtilities::PrepareRapidIterationParameters(Scripts, ScriptDependencyMap, ScriptToEmitterNameMap);
 
+		// HACK: This is a temporary hack to fix an issue where data interfaces used by modules and dynamic inputs in the
+		// particle update script aren't being shared by the interpolated spawn script when accessed directly.  This works
+		// properly if the data interface is assigned to a named particle parameter and then linked to an input.
+		// TODO: Bind these data interfaces the same way parameter data interfaces are bound.
+		for (FEmitterCompiledScriptPair& EmitterCompiledScriptPair : ActiveCompilations[ActiveCompileIdx].EmitterCompiledScriptPairs)
+		{
+			UNiagaraEmitter* Emitter = EmitterCompiledScriptPair.Emitter;
+			UNiagaraScript* CompiledScript = EmitterCompiledScriptPair.CompiledScript;
+
+			if (UNiagaraScript::IsEquivalentUsage(CompiledScript->GetUsage(), ENiagaraScriptUsage::ParticleUpdateScript))
+			{
+				UNiagaraScript* SpawnScript = Emitter->SpawnScriptProps.Script;
+				for (const FNiagaraScriptDataInterfaceInfo& UpdateDataInterfaceInfo : CompiledScript->GetCachedDefaultDataInterfaces())
+				{
+					if (UpdateDataInterfaceInfo.RegisteredParameterMapRead == NAME_None && UpdateDataInterfaceInfo.RegisteredParameterMapWrite == NAME_None)
+					{
+						// If the data interface isn't being read or written to a parameter map then it won't be bound properly so we
+						// assign the update scripts copy of the data interface to the spawn scripts copy by pointer so that they will share
+						// the data interface at runtime and will both be updated in the editor.
+						for (FNiagaraScriptDataInterfaceInfo& SpawnDataInterfaceInfo : SpawnScript->GetCachedDefaultDataInterfaces())
+						{
+							if (UpdateDataInterfaceInfo.Name == SpawnDataInterfaceInfo.Name)
+							{
+								SpawnDataInterfaceInfo.DataInterface = UpdateDataInterfaceInfo.DataInterface;
+							}
+						}
+					}
+				}
+			}
+		}
+
 		ActiveCompilations[ActiveCompileIdx].RootObjects.Empty();
 
 		DetermineIfSolo();
diff --git a/Engine/Plugins/FX/Niagara/Source/Niagara/Private/NiagaraSystemInstance.cpp b/Engine/Plugins/FX/Niagara/Source/Niagara/Private/NiagaraSystemInstance.cpp
--- a/Engine/Plugins/FX/Niagara/Source/Niagara/Private/NiagaraSystemInstance.cpp
+++ b/Engine/Plugins/FX/Niagara/Source/Niagara/Private/NiagaraSystemInstance.cpp
@@ -120,6 +120,15 @@ void FNiagaraSystemInstance::SetActualExecutionState(ENiagaraExecutionState InSt
 	}
 }
 
+void FNiagaraSystemInstance::Dump()const
+{
+	GetSystemSimulation()->DumpInstance(this);
+	for (auto& Emitter : Emitters)
+	{
+		Emitter->Dump();
+	}
+}
+
 #if WITH_EDITORONLY_DATA
 bool FNiagaraSystemInstance::RequestCapture(const FGuid& RequestId)
 {
@@ -905,7 +914,7 @@ void FNiagaraSystemInstance::TickInstanceParameters(float DeltaSeconds)
 			for (FConstPlayerControllerIterator Iterator = World->GetPlayerControllerIterator(); Iterator; ++Iterator)
 			{
 				APlayerController* PlayerController = Iterator->Get();
-				if (PlayerController->IsLocalPlayerController())
+				if (PlayerController && PlayerController->IsLocalPlayerController())
 				{
 					FVector* POVLoc = new(PlayerViewLocations) FVector;
 					FRotator POVRotation;
diff --git a/Engine/Plugins/FX/Niagara/Source/Niagara/Private/NiagaraSystemSimulation.cpp b/Engine/Plugins/FX/Niagara/Source/Niagara/Private/NiagaraSystemSimulation.cpp
--- a/Engine/Plugins/FX/Niagara/Source/Niagara/Private/NiagaraSystemSimulation.cpp
+++ b/Engine/Plugins/FX/Niagara/Source/Niagara/Private/NiagaraSystemSimulation.cpp
@@ -246,6 +246,20 @@ void FNiagaraSystemSimulation::TransferInstance(FNiagaraSystemSimulation* Source
 	}
 }
 
+void FNiagaraSystemSimulation::DumpInstance(const FNiagaraSystemInstance* Inst)const
+{
+	UE_LOG(LogNiagara, Log, TEXT("==  %s (%d) ========"), *Inst->GetSystem()->GetFullName(), Inst->SystemInstanceIndex);
+	UE_LOG(LogNiagara, Log, TEXT(".................Spawn................."));
+	SpawnExecContext.Parameters.DumpParameters(false);
+	SpawnInstanceParameterDataSet.Dump(false, Inst->SystemInstanceIndex, 1);
+	UE_LOG(LogNiagara, Log, TEXT(".................Update................."));
+	UpdateExecContext.Parameters.DumpParameters(false);
+	UpdateInstanceParameterDataSet.Dump(false, Inst->SystemInstanceIndex, 1);
+	UE_LOG(LogNiagara, Log, TEXT("................. System Instance ................."));
+	DataSet.Dump(false, Inst->SystemInstanceIndex, 1);
+	DataSet.Dump(true, Inst->SystemInstanceIndex, 1);
+}
+
 bool FNiagaraSystemSimulation::Tick(float DeltaSeconds)
 {
 	UNiagaraSystem* System = WeakSystem.Get();
@@ -298,9 +312,9 @@ bool FNiagaraSystemSimulation::Tick(float DeltaSeconds)
 
 				Inst->TickDataInterfaces(DeltaSeconds, false);
 				Inst->SetPendingSpawn(false);
-				Inst->SystemInstanceIndex = SystemInstances.Add(Inst);
 				if (!Inst->IsComplete())
 				{
+					Inst->SystemInstanceIndex = SystemInstances.Add(Inst);
 					++SpawnNum;
 				}
 				else
@@ -324,6 +338,7 @@ bool FNiagaraSystemSimulation::Tick(float DeltaSeconds)
 		}
 
 		TArray<FNiagaraDataSetExecutionInfo, TInlineAllocator<8>> DataSetExecInfos;
+		DataSetExecInfos.SetNum(2);
 		{
 			SCOPE_CYCLE_COUNTER(STAT_NiagaraSystemSim_PrepareForSimulate);
 			auto PreSimulateAndTransferParams = [&](int32 SystemIndex)
@@ -386,8 +401,6 @@ bool FNiagaraSystemSimulation::Tick(float DeltaSeconds)
 			UpdateGlobalSpawnCountScaleParam.SetValue(GlobalSpawnCountScale);
 			SpawnGlobalSystemCountScaleParam.SetValue(GlobalSystemCountScale);
 			UpdateGlobalSystemCountScaleParam.SetValue(GlobalSystemCountScale);
-
-			DataSetExecInfos.Emplace(&DataSet, 0, false, true);
 		}
 
 		//TODO: JIRA - UE-60096 - Remove.
@@ -403,9 +416,8 @@ bool FNiagaraSystemSimulation::Tick(float DeltaSeconds)
 
 			//Run Spawn
 			SpawnExecContext.Tick(SoloSystemInstance);//We can't require a specific instance here as these are for all instances.
-			DataSetExecInfos.SetNum(1, false);
-			DataSetExecInfos[0].StartInstance = OrigNum;
-			DataSetExecInfos.Emplace(&SpawnInstanceParameterDataSet, OrigNum, false, false);
+			DataSetExecInfos[0] = FNiagaraDataSetExecutionInfo(&DataSet, OrigNum, false, false);
+			DataSetExecInfos[1] = FNiagaraDataSetExecutionInfo(&SpawnInstanceParameterDataSet, OrigNum, false, false);
 			SpawnExecContext.Execute(SpawnNum, DataSetExecInfos);
 
 			if (GbDumpSystemData || System->bDumpDebugSystemInfo)
@@ -439,8 +451,8 @@ bool FNiagaraSystemSimulation::Tick(float DeltaSeconds)
 
 			//Run update.
 			UpdateExecContext.Tick(SystemInstances[0]);
-			DataSetExecInfos[0].StartInstance = 0;
-			DataSetExecInfos.Emplace(&UpdateInstanceParameterDataSet, 0, false, false);
+			DataSetExecInfos[0] = FNiagaraDataSetExecutionInfo(&DataSet, 0, false, false);
+			DataSetExecInfos[1] = FNiagaraDataSetExecutionInfo(&UpdateInstanceParameterDataSet, 0, false, false);
 
 // 			if (GbDumpSystemData || System->bDumpDebugSystemInfo)
 // 			{
@@ -455,6 +467,7 @@ bool FNiagaraSystemSimulation::Tick(float DeltaSeconds)
 			{
 				UE_LOG(LogNiagara, Log, TEXT("=== Updated %d Systems ==="), OrigNum);
 				DataSet.Dump(true, 0, OrigNum);
+				UpdateInstanceParameterDataSet.Dump(false, 0, OrigNum);
 			}
 
  			//Also run the update script on the newly spawned systems too.
@@ -466,8 +479,8 @@ bool FNiagaraSystemSimulation::Tick(float DeltaSeconds)
 
 				//Run update.
 				UpdateExecContext.Tick(SystemInstances[0]);
-				DataSetExecInfos[0].StartInstance = OrigNum;
-				DataSetExecInfos[1].StartInstance = OrigNum;
+				DataSetExecInfos[0] = FNiagaraDataSetExecutionInfo(&DataSet, OrigNum, false, false);
+				DataSetExecInfos[1] = FNiagaraDataSetExecutionInfo(&UpdateInstanceParameterDataSet, OrigNum, false, false);
 
 				UpdateExecContext.Parameters.SetParameterValue(0.0001f, SYS_PARAM_ENGINE_DELTA_TIME);
 				UpdateExecContext.Parameters.SetParameterValue(10000.0f, SYS_PARAM_ENGINE_INV_DELTA_TIME);
@@ -478,6 +491,7 @@ bool FNiagaraSystemSimulation::Tick(float DeltaSeconds)
 				{
 					UE_LOG(LogNiagara, Log, TEXT("=== Spawn Updated %d Systems ==="), SpawnNum);
 					DataSet.Dump(true, OrigNum, SpawnNum);
+					UpdateInstanceParameterDataSet.Dump(false, OrigNum, SpawnNum);
 				}
 			}
 
@@ -673,6 +687,10 @@ void FNiagaraSystemSimulation::RemoveInstance(FNiagaraSystemInstance* Instance)
 			SystemInstances[SystemIndex]->SystemInstanceIndex = SystemIndex;
 		}
 	}
+
+#if NIAGARA_NAN_CHECKING
+	DataSet.CheckForNaNs();
+#endif
 }
 
 void FNiagaraSystemSimulation::AddInstance(FNiagaraSystemInstance* Instance)
diff --git a/Engine/Plugins/FX/Niagara/Source/Niagara/Private/NiagaraWorldManager.cpp b/Engine/Plugins/FX/Niagara/Source/Niagara/Private/NiagaraWorldManager.cpp
--- a/Engine/Plugins/FX/Niagara/Source/Niagara/Private/NiagaraWorldManager.cpp
+++ b/Engine/Plugins/FX/Niagara/Source/Niagara/Private/NiagaraWorldManager.cpp
@@ -13,10 +13,11 @@
 #include "NiagaraDataInterfaceSkeletalMesh.h"
 #include "EngineModule.h"
 
-FNiagaraViewDataMgr GNiagaraViewDataManager;
+TGlobalResource<FNiagaraViewDataMgr> GNiagaraViewDataManager;
 
-FNiagaraViewDataMgr::FNiagaraViewDataMgr()
-	: SceneDepthTexture(nullptr)
+FNiagaraViewDataMgr::FNiagaraViewDataMgr() 
+	: FRenderResource()
+	, SceneDepthTexture(nullptr)
 	, SceneNormalTexture(nullptr)
 	, ViewUniformBuffer(nullptr)
 {
@@ -31,10 +32,22 @@ void FNiagaraViewDataMgr::Init()
 }
 
 void FNiagaraViewDataMgr::Shutdown()
+{
+	GNiagaraViewDataManager.ReleaseDynamicRHI();
+}
+
+void FNiagaraViewDataMgr::InitDynamicRHI()
 {
 
 }
 
+void FNiagaraViewDataMgr::ReleaseDynamicRHI()
+{
+	SceneDepthTexture = nullptr;
+	SceneNormalTexture = nullptr;
+	ViewUniformBuffer = nullptr;
+	SceneTexturesUniformParams.SafeRelease();
+}
 
 FNiagaraWorldManager::FNiagaraWorldManager(UWorld* InWorld)
 	: World(InWorld)
diff --git a/Engine/Plugins/FX/Niagara/Source/Niagara/Public/NiagaraCommon.h b/Engine/Plugins/FX/Niagara/Source/Niagara/Public/NiagaraCommon.h
--- a/Engine/Plugins/FX/Niagara/Source/Niagara/Public/NiagaraCommon.h
+++ b/Engine/Plugins/FX/Niagara/Source/Niagara/Public/NiagaraCommon.h
@@ -16,6 +16,9 @@ class FNiagaraSystemInstance;
 class UNiagaraParameterCollection;
 struct FNiagaraParameterStore;
 
+//#define NIAGARA_NAN_CHECKING 1
+#define NIAGARA_NAN_CHECKING 0
+
 const uint32 NIAGARA_COMPUTE_THREADGROUP_SIZE = 64;
 const uint32 NIAGARA_MAX_COMPUTE_THREADGROUPS = 65536;
 
diff --git a/Engine/Plugins/FX/Niagara/Source/Niagara/Public/NiagaraEditorDataBase.h b/Engine/Plugins/FX/Niagara/Source/Niagara/Public/NiagaraEditorDataBase.h
--- a/Engine/Plugins/FX/Niagara/Source/Niagara/Public/NiagaraEditorDataBase.h
+++ b/Engine/Plugins/FX/Niagara/Source/Niagara/Public/NiagaraEditorDataBase.h
@@ -1,4 +1,4 @@
-// Copyright 1998-2017 Epic Games, Inc. All Rights Reserved.
+// Copyright 1998-2018 Epic Games, Inc. All Rights Reserved.
 
 #pragma once
 
diff --git a/Engine/Plugins/FX/Niagara/Source/Niagara/Public/NiagaraModule.h b/Engine/Plugins/FX/Niagara/Source/Niagara/Public/NiagaraModule.h
--- a/Engine/Plugins/FX/Niagara/Source/Niagara/Public/NiagaraModule.h
+++ b/Engine/Plugins/FX/Niagara/Source/Niagara/Public/NiagaraModule.h
@@ -62,6 +62,7 @@ public:
 public:
 	virtual void StartupModule()override;
 	virtual void ShutdownModule()override;
+	void ShutdownRenderingResources();
 
 	FDelegateHandle SetOnProcessShaderCompilationQueue(FOnProcessQueue InOnProcessQueue);
 	void ResetOnProcessShaderCompilationQueue(FDelegateHandle DelegateHandle);
diff --git a/Engine/Plugins/FX/Niagara/Source/Niagara/Public/NiagaraParameterStore.h b/Engine/Plugins/FX/Niagara/Source/Niagara/Public/NiagaraParameterStore.h
--- a/Engine/Plugins/FX/Niagara/Source/Niagara/Public/NiagaraParameterStore.h
+++ b/Engine/Plugins/FX/Niagara/Source/Niagara/Public/NiagaraParameterStore.h
@@ -60,7 +60,7 @@ struct FNiagaraParameterStoreBinding
 	FORCEINLINE_DEBUGGABLE void Initialize(FNiagaraParameterStore* DestStore, FNiagaraParameterStore* SrcStore);
 	FORCEINLINE_DEBUGGABLE bool VerifyBinding(const FNiagaraParameterStore* DestStore, const FNiagaraParameterStore* SrcStore)const;
 	FORCEINLINE_DEBUGGABLE void Tick(FNiagaraParameterStore* DestStore, FNiagaraParameterStore* SrcStore, bool bForce = false);
-	FORCEINLINE_DEBUGGABLE void Dump(FNiagaraParameterStore* DestStore, FNiagaraParameterStore* SrcStore)const;
+	FORCEINLINE_DEBUGGABLE void Dump(const FNiagaraParameterStore* DestStore, const FNiagaraParameterStore* SrcStore)const;
 	/** TODO: Merge contiguous ranges into a single binding? */
 	//FORCEINLINE_DEBUGGABLE void Optimize();
 
@@ -129,7 +129,7 @@ public:
 	UObject* GetOwner()const { return Owner; }
 
 	void Dump();
-	void DumpParameters(bool bDumpBindings = false);
+	void DumpParameters(bool bDumpBindings = false)const;
 
 	FORCEINLINE uint32 GetParametersDirty() const { return bParametersDirty; }
 	FORCEINLINE uint32 GetInterfacesDirty() const { return bInterfacesDirty; }
@@ -154,6 +154,8 @@ public:
 	
 	bool VerifyBinding(const FNiagaraParameterStore* InDestStore)const;
 
+	void CheckForNaNs()const;
+
 	/**
 	Adds the passed parameter to this store.
 	Does nothing if this parameter is already present.
@@ -559,9 +561,13 @@ FORCEINLINE_DEBUGGABLE void FNiagaraParameterStoreBinding::Tick(FNiagaraParamete
 			DestStore->SetDataInterface(SrcStore->GetDataInterface(Binding.SrcOffset), Binding.DestOffset);
 		}
 	}
+
+#if NIAGARA_NAN_CHECKING
+	DestStore->CheckForNaNs();
+#endif
 }
 
-FORCEINLINE_DEBUGGABLE void FNiagaraParameterStoreBinding::Dump(FNiagaraParameterStore* DestStore, FNiagaraParameterStore* SrcStore)const
+FORCEINLINE_DEBUGGABLE void FNiagaraParameterStoreBinding::Dump(const FNiagaraParameterStore* DestStore, const FNiagaraParameterStore* SrcStore)const
 {
 #if WITH_EDITORONLY_DATA
 	UE_LOG(LogNiagara, Log, TEXT("\n\nDest Store: %s\n"), *DestStore->DebugName);
diff --git a/Engine/Plugins/FX/Niagara/Source/Niagara/Public/NiagaraRendererRibbons.h b/Engine/Plugins/FX/Niagara/Source/Niagara/Public/NiagaraRendererRibbons.h
--- a/Engine/Plugins/FX/Niagara/Source/Niagara/Public/NiagaraRendererRibbons.h
+++ b/Engine/Plugins/FX/Niagara/Source/Niagara/Public/NiagaraRendererRibbons.h
@@ -43,10 +43,10 @@ public:
 		NewVertex.Color = Color;
 		NewVertex.Size = Size;
 		NewVertex.Rotation = Rotation;
-		NewVertex.Tex_U = UV1.X * Properties->UV0Scale.X;
-		NewVertex.Tex_V = UV1.Y * Properties->UV0Scale.Y;
-		NewVertex.Tex_U2 = UV2.X * Properties->UV1Scale.X;
-		NewVertex.Tex_V2 = UV2.Y * Properties->UV1Scale.Y;
+		NewVertex.Tex_U = UV1.X;
+		NewVertex.Tex_V = UV1.Y;
+		NewVertex.Tex_U2 = UV2.X;
+		NewVertex.Tex_V2 = UV2.Y;
 		NewVertex.CustomFacingVector = CustomFacing;
 		RenderData.Add(NewVertex);
 	}
diff --git a/Engine/Plugins/FX/Niagara/Source/Niagara/Public/NiagaraRibbonRendererProperties.h b/Engine/Plugins/FX/Niagara/Source/Niagara/Public/NiagaraRibbonRendererProperties.h
--- a/Engine/Plugins/FX/Niagara/Source/Niagara/Public/NiagaraRibbonRendererProperties.h
+++ b/Engine/Plugins/FX/Niagara/Source/Niagara/Public/NiagaraRibbonRendererProperties.h
@@ -15,6 +15,16 @@ enum class ENiagaraRibbonFacingMode : uint8
 	Custom
 };
 
+/** Defines different modes for offsetting UVs by age when ordering ribbon particles using normalized age. */
+UENUM()
+enum class ENiagaraRibbonAgeOffsetMode : uint8
+{
+	/** Offset the UVs by age for smooth texture movement, but scale the 0-1 UV range to the current normalized age range of the particles. */
+	Scale,
+	/** Offset the UVs by age for smooth texture movement, but use the normalized age range directly as the UV range which will clip the texture for normalized age ranges that don't go from 0-1. */
+	Clip
+};
+
 UCLASS(editinlinenew)
 class UNiagaraRibbonRendererProperties : public UNiagaraRendererProperties
 {
@@ -26,8 +36,10 @@ public:
 		, FacingMode(ENiagaraRibbonFacingMode::Screen)
 		, UV0TilingDistance(0.0f)
 		, UV0Scale(FVector2D(1.0f, 1.0f))
+		, UV0AgeOffsetMode(ENiagaraRibbonAgeOffsetMode::Scale)
 		, UV1TilingDistance(0.0f)
 		, UV1Scale(FVector2D(1.0f, 1.0f))
+		, UV1AgeOffsetMode(ENiagaraRibbonAgeOffsetMode::Scale)
 	{
 	}
 
@@ -37,15 +49,29 @@ public:
 	UPROPERTY(EditAnywhere, Category = "Ribbon Rendering")
 	ENiagaraRibbonFacingMode FacingMode;
 
+	/** Tiles UV0 based on the distance traversed by the ribbon. Disables offsetting UVs by age. */
 	UPROPERTY(EditAnywhere, Category = "Ribbon Rendering")
 	float UV0TilingDistance;
 	UPROPERTY(EditAnywhere, Category = "Ribbon Rendering")
 	FVector2D UV0Scale;
+	UPROPERTY(EditAnywhere, Category = "Ribbon Rendering")
+	FVector2D UV0Offset;
 
+	/** Defines the mode to use when offsetting UV channel 0 by age which enables smooth texture movement when particles are added and removed at the end of the ribbon.  Not used when the RibbonLinkOrder binding is in use or when tiling distance is in use. */
+	UPROPERTY(EditAnywhere, Category = "Ribbon Rendering")
+	ENiagaraRibbonAgeOffsetMode UV0AgeOffsetMode;
+
+	/** Tiles UV1 based on the distance traversed by the ribbon. Disables offsetting UVs by age. */
 	UPROPERTY(EditAnywhere, Category = "Ribbon Rendering")
 	float UV1TilingDistance;
 	UPROPERTY(EditAnywhere, Category = "Ribbon Rendering")
 	FVector2D UV1Scale;
+	UPROPERTY(EditAnywhere, Category = "Ribbon Rendering")
+	FVector2D UV1Offset;
+
+	/** Defines the mode to use when offsetting UV channel 1 by age which enables smooth texture movement when particles are added and removed at the end of the ribbon.  Not used when the RibbonLinkOrder binding is in use or when tiling distance is in use. */
+	UPROPERTY(EditAnywhere, Category = "Ribbon Rendering")
+	ENiagaraRibbonAgeOffsetMode UV1AgeOffsetMode;
 
 	//~ UNiagaraRendererProperties interface
 	virtual NiagaraRenderer* CreateEmitterRenderer(ERHIFeatureLevel::Type FeatureLevel) override;
@@ -55,7 +81,7 @@ public:
 	virtual void PostEditChangeProperty(struct FPropertyChangedEvent& PropertyChangedEvent) override;
 	virtual const TArray<FNiagaraVariable>& GetRequiredAttributes() override;
 	virtual const TArray<FNiagaraVariable>& GetOptionalAttributes() override;
-    virtual bool IsMaterialValidForRenderer(UMaterial* Material, FText& InvalidMessage) override;
+	virtual bool IsMaterialValidForRenderer(UMaterial* Material, FText& InvalidMessage) override;
 	virtual void FixMaterial(UMaterial* Material);
 #endif
 
diff --git a/Engine/Plugins/FX/Niagara/Source/Niagara/Public/NiagaraScriptExecutionParameterStore.h b/Engine/Plugins/FX/Niagara/Source/Niagara/Public/NiagaraScriptExecutionParameterStore.h
--- a/Engine/Plugins/FX/Niagara/Source/Niagara/Public/NiagaraScriptExecutionParameterStore.h
+++ b/Engine/Plugins/FX/Niagara/Source/Niagara/Public/NiagaraScriptExecutionParameterStore.h
@@ -1,3 +1,5 @@
+// Copyright 1998-2018 Epic Games, Inc. All Rights Reserved.
+
 #pragma once
 
 #include "NiagaraCommon.h"
@@ -61,6 +63,14 @@ public:
 		check(0);//Can't rename parameters for an execution store.
 	}
 
+	void Empty(bool bClearBindings=true)
+	{
+		FNiagaraParameterStore::Empty(bClearBindings);
+		PaddingInfo.Empty();
+		PaddedParameterSize = 0;
+		bInitialized = false;
+	}
+
 	// Just the external parameters, not previous or internal...
 	uint32 GetExternalParameterSize() { return ParameterSize; }
 
@@ -70,13 +80,6 @@ public:
 	// Helper that converts the data from the base type array internally into the padded out renderer-ready format.
 	void CopyParameterDataToPaddedBuffer(uint8* InTargetBuffer, uint32 InTargetBufferSizeInBytes);
 
-	void Clear()
-	{
-		Empty();
-		PaddingInfo.Empty();
-		PaddedParameterSize = 0;
-		bInitialized = false;
-	}
 
 	bool IsInitialized() const {return bInitialized;}
 
diff --git a/Engine/Plugins/FX/Niagara/Source/Niagara/Public/NiagaraSpriteRendererProperties.h b/Engine/Plugins/FX/Niagara/Source/Niagara/Public/NiagaraSpriteRendererProperties.h
--- a/Engine/Plugins/FX/Niagara/Source/Niagara/Public/NiagaraSpriteRendererProperties.h
+++ b/Engine/Plugins/FX/Niagara/Source/Niagara/Public/NiagaraSpriteRendererProperties.h
@@ -40,7 +40,7 @@ enum class ENiagaraSpriteFacingMode : uint8
 
 
 UCLASS(editinlinenew)
-class UNiagaraSpriteRendererProperties : public UNiagaraRendererProperties
+class NIAGARA_API UNiagaraSpriteRendererProperties : public UNiagaraRendererProperties
 {
 public:
 	GENERATED_BODY()
diff --git a/Engine/Plugins/FX/Niagara/Source/Niagara/Public/NiagaraSystemInstance.h b/Engine/Plugins/FX/Niagara/Source/Niagara/Public/NiagaraSystemInstance.h
--- a/Engine/Plugins/FX/Niagara/Source/Niagara/Public/NiagaraSystemInstance.h
+++ b/Engine/Plugins/FX/Niagara/Source/Niagara/Public/NiagaraSystemInstance.h
@@ -192,6 +192,9 @@ public:
 
 #endif
 
+	/** Dumps all of this systems info to the log. */
+	void Dump()const;
+
 private:
 	/** Builds the emitter simulations. */
 	void InitEmitters();
diff --git a/Engine/Plugins/FX/Niagara/Source/Niagara/Public/NiagaraSystemSimulation.h b/Engine/Plugins/FX/Niagara/Source/Niagara/Public/NiagaraSystemSimulation.h
--- a/Engine/Plugins/FX/Niagara/Source/Niagara/Public/NiagaraSystemSimulation.h
+++ b/Engine/Plugins/FX/Niagara/Source/Niagara/Public/NiagaraSystemSimulation.h
@@ -64,6 +64,10 @@ struct FNiagaraParameterStoreToDataSetBinding
 
 	FORCEINLINE_DEBUGGABLE void DataSetToParameterStore(FNiagaraParameterStore& ParameterStore, FNiagaraDataSet& DataSet, int32 DataSetInstanceIndex)
 	{
+#if NIAGARA_NAN_CHECKING
+		DataSet.CheckForNaNs();
+#endif
+
 		FNiagaraDataBuffer& CurrBuffer = DataSet.CurrData();
 
 		for (const FDataOffsets& DataOffsets : FloatOffsets)
@@ -77,6 +81,10 @@ struct FNiagaraParameterStoreToDataSetBinding
 			ParameterStore.SetParameterByOffset(DataOffsets.ParameterOffset, *DataSetPtr);
 		}
 
+#if NIAGARA_NAN_CHECKING
+		ParameterStore.CheckForNaNs();
+#endif
+
 		ParameterStore.OnParameterChange();
 	}
 
@@ -85,6 +93,10 @@ struct FNiagaraParameterStoreToDataSetBinding
 		FNiagaraDataBuffer& CurrBuffer = DataSet.CurrData();
 		const uint8* ParameterData = ParameterStore.GetParameterDataArray().GetData();
 
+#if NIAGARA_NAN_CHECKING
+		ParameterStore.CheckForNaNs();
+#endif
+
 		for (const FDataOffsets& DataOffsets : FloatOffsets)
 		{
 			float* ParamPtr = (float*)(ParameterData + DataOffsets.ParameterOffset);
@@ -97,6 +109,10 @@ struct FNiagaraParameterStoreToDataSetBinding
 			int32* DataSetPtr = CurrBuffer.GetInstancePtrInt32(DataOffsets.DataSetComponentOffset, DataSetInstanceIndex);
 			*DataSetPtr = *ParamPtr;
 		}
+
+#if NIAGARA_NAN_CHECKING
+		DataSet.CheckForNaNs();
+#endif
 	}
 };
 
@@ -120,6 +136,8 @@ public:
 
 	/** Transfers a system instance from SourceSimulation. */
 	void TransferInstance(FNiagaraSystemSimulation* SourceSimulation, FNiagaraSystemInstance* SystemInst);
+
+	void DumpInstance(const FNiagaraSystemInstance* Inst)const;
 protected:
 
 	/** System of instances being simulated.  We use a weak object ptr here because once the last referencing object goes away this system may be come invalid at runtime. */
@@ -140,7 +158,7 @@ protected:
 	Though all these parameters can be unique per instance so for now lets just do the simple thing and improve later.
 	*/
 	FNiagaraDataSet SpawnInstanceParameterDataSet;
- 	FNiagaraDataSet UpdateInstanceParameterDataSet;
+	FNiagaraDataSet UpdateInstanceParameterDataSet;
 
 	FNiagaraScriptExecutionContext SpawnExecContext;
 	FNiagaraScriptExecutionContext UpdateExecContext;
diff --git a/Engine/Plugins/FX/Niagara/Source/Niagara/Public/NiagaraWorldManager.h b/Engine/Plugins/FX/Niagara/Source/Niagara/Public/NiagaraWorldManager.h
--- a/Engine/Plugins/FX/Niagara/Source/Niagara/Public/NiagaraWorldManager.h
+++ b/Engine/Plugins/FX/Niagara/Source/Niagara/Public/NiagaraWorldManager.h
@@ -18,7 +18,7 @@ class UNiagaraParameterCollectionInstance;
 
 
 
-class FNiagaraViewDataMgr
+class FNiagaraViewDataMgr : public FRenderResource
 {
 public:
 	FNiagaraViewDataMgr();
@@ -38,6 +38,10 @@ public:
 	FTexture2DRHIParamRef GetSceneNormalTexture() { return SceneNormalTexture; }
 	FUniformBufferRHIParamRef GetViewUniformBuffer() { return ViewUniformBuffer; }
 	TUniformBufferRef<FSceneTexturesUniformParameters> GetSceneTextureUniformParameters() { return SceneTexturesUniformParams; }
+
+	virtual void InitDynamicRHI() override;
+
+	virtual void ReleaseDynamicRHI() override;
 private:
 	FTexture2DRHIParamRef SceneDepthTexture;
 	FTexture2DRHIParamRef SceneNormalTexture;
@@ -47,7 +51,7 @@ private:
 	FPostOpaqueRenderDelegate PostOpaqueDelegate;
 };
 
-extern FNiagaraViewDataMgr GNiagaraViewDataManager;
+extern TGlobalResource<FNiagaraViewDataMgr> GNiagaraViewDataManager;
 
 
 /**
diff --git a/Engine/Plugins/FX/Niagara/Source/NiagaraCore/Private/NiagaraCustomVersion.cpp b/Engine/Plugins/FX/Niagara/Source/NiagaraCore/Private/NiagaraCustomVersion.cpp
--- a/Engine/Plugins/FX/Niagara/Source/NiagaraCore/Private/NiagaraCustomVersion.cpp
+++ b/Engine/Plugins/FX/Niagara/Source/NiagaraCore/Private/NiagaraCustomVersion.cpp
@@ -8,4 +8,4 @@ const FGuid FNiagaraCustomVersion::GUID(0xFCF57AFA, 0x50764283, 0xB9A9E658, 0xFF
 // Register the custom version with core
 FCustomVersionRegistration GRegisterNiagaraCustomVersion(FNiagaraCustomVersion::GUID, FNiagaraCustomVersion::LatestVersion, TEXT("NiagaraVer"));
 
-const FGuid FNiagaraCustomVersion::LatestScriptCompileVersion(0x67DE59E9, 0x4ED94B3C, 0x845B4A87, 0x1C44F338);
\ No newline at end of file
+const FGuid FNiagaraCustomVersion::LatestScriptCompileVersion(0xEAA7D8F4, 0x8BE04550, 0xAF0EC099, 0x5A99A2F4);
\ No newline at end of file
diff --git a/Engine/Plugins/FX/Niagara/Source/NiagaraEditor/Private/NiagaraCompiler.cpp b/Engine/Plugins/FX/Niagara/Source/NiagaraEditor/Private/NiagaraCompiler.cpp
--- a/Engine/Plugins/FX/Niagara/Source/NiagaraEditor/Private/NiagaraCompiler.cpp
+++ b/Engine/Plugins/FX/Niagara/Source/NiagaraEditor/Private/NiagaraCompiler.cpp
@@ -298,6 +298,7 @@ void FNiagaraCompileRequestData::GetReferencedObjects(TArray<UObject*>& Objects)
 			for (int32 i = 0; i < Iter.Value().Num(); i++)
 			{
 				Objects.Add(Iter.Value()[i].ClonedScript);
+				Objects.Add(Iter.Value()[i].ClonedGraph);
 			}
 			++Iter;
 		}
diff --git a/Engine/Plugins/FX/Niagara/Source/NiagaraEditor/Private/NiagaraEditorModule.cpp b/Engine/Plugins/FX/Niagara/Source/NiagaraEditor/Private/NiagaraEditorModule.cpp
--- a/Engine/Plugins/FX/Niagara/Source/NiagaraEditor/Private/NiagaraEditorModule.cpp
+++ b/Engine/Plugins/FX/Niagara/Source/NiagaraEditor/Private/NiagaraEditorModule.cpp
@@ -297,6 +297,7 @@ void FNiagaraEditorModule::StartupModule()
 	RegisterAssetTypeAction(AssetTools, MakeShareable(new FAssetTypeActions_NiagaraParameterCollectionInstance()));
 
 	UNiagaraSettings::OnSettingsChanged().AddRaw(this, &FNiagaraEditorModule::OnNiagaraSettingsChangedEvent);
+	FCoreUObjectDelegates::GetPreGarbageCollectDelegate().AddRaw(this, &FNiagaraEditorModule::OnPreGarbageCollection);
 
 	// register details customization
 	FPropertyEditorModule& PropertyModule = FModuleManager::LoadModuleChecked<FPropertyEditorModule>("PropertyEditor");
@@ -491,7 +492,7 @@ void FNiagaraEditorModule::ShutdownModule()
 
 	MenuExtensibilityManager.Reset();
 	ToolBarExtensibilityManager.Reset();
-
+	
 	if (FModuleManager::Get().IsModuleLoaded("AssetTools"))
 	{
 		IAssetTools& AssetTools = FModuleManager::GetModuleChecked<FAssetToolsModule>("AssetTools").Get();
@@ -504,6 +505,8 @@ void FNiagaraEditorModule::ShutdownModule()
 
 	UNiagaraSettings::OnSettingsChanged().RemoveAll(this);
 
+	FCoreUObjectDelegates::GetPreGarbageCollectDelegate().RemoveAll(this);
+
 	
 	if (FModuleManager::Get().IsModuleLoaded("PropertyEditor"))
 	{
@@ -677,7 +680,7 @@ void FNiagaraEditorModule::UnregisterSettings()
 	}
 }
 
-void FNiagaraEditorModule::AddReferencedObjects( FReferenceCollector& Collector )
+void FNiagaraEditorModule::AddReferencedObjects(FReferenceCollector& Collector)
 {
 	if (SequencerSettings)
 	{
@@ -685,4 +688,22 @@ void FNiagaraEditorModule::AddReferencedObjects( FReferenceCollector& Collector
 	}
 }
 
+void FNiagaraEditorModule::OnPreGarbageCollection()
+{
+	// For commandlets like GenerateDistillFileSetsCommandlet, they just load the package and do some hierarchy navigation within it 
+	// tracking sub-assets, then they garbage collect. Since nothing is holding onto the system at the root level, it will be summarily
+	// killed and any of references will also be killed. To thwart this for now, we are forcing the compilations to complete BEFORE
+	// garbage collection kicks in. To do otherwise for now has too many loose ends (a system may be left around after the level has been
+	// unloaded, leaving behind weird external references, etc). This should be revisited when more time is available (i.e. not days before a 
+	// release is due to go out).
+	for (TObjectIterator<UNiagaraSystem> It; It; ++It)
+	{
+		UNiagaraSystem* System = *It;
+		if (System && System->HasOutstandingCompilationRequests())
+		{
+			System->WaitForCompilationComplete();
+		}
+	}
+}
+
 #undef LOCTEXT_NAMESPACE
diff --git a/Engine/Plugins/FX/Niagara/Source/NiagaraEditor/Private/NiagaraEditorStyle.cpp b/Engine/Plugins/FX/Niagara/Source/NiagaraEditor/Private/NiagaraEditorStyle.cpp
--- a/Engine/Plugins/FX/Niagara/Source/NiagaraEditor/Private/NiagaraEditorStyle.cpp
+++ b/Engine/Plugins/FX/Niagara/Source/NiagaraEditor/Private/NiagaraEditorStyle.cpp
@@ -34,7 +34,7 @@ FName FNiagaraEditorStyle::GetStyleSetName()
 	return StyleSetName;
 }
 
-FString RelativePathToPluginPath(const FString& RelativePath, const ANSICHAR* Extension)
+NIAGARAEDITOR_API FString RelativePathToPluginPath(const FString& RelativePath, const ANSICHAR* Extension)
 {
 	static FString ContentDir = IPluginManager::Get().FindPlugin(TEXT("Niagara"))->GetContentDir();
 	return (ContentDir / RelativePath) + Extension;
diff --git a/Engine/Plugins/FX/Niagara/Source/NiagaraEditor/Private/NiagaraEditorUtilities.cpp b/Engine/Plugins/FX/Niagara/Source/NiagaraEditor/Private/NiagaraEditorUtilities.cpp
--- a/Engine/Plugins/FX/Niagara/Source/NiagaraEditor/Private/NiagaraEditorUtilities.cpp
+++ b/Engine/Plugins/FX/Niagara/Source/NiagaraEditor/Private/NiagaraEditorUtilities.cpp
@@ -498,11 +498,11 @@ void FNiagaraEditorUtilities::CompileExistingEmitters(const TArray<UNiagaraEmitt
 			continue;
 		}
 
-		// We only need to compile emitters referenced directly by systems since emitters can now only be used in the context 
+		// We only need to compile emitters referenced directly as instances by systems since emitters can now only be used in the context 
 		// of a system.
 		for (TObjectIterator<UNiagaraSystem> SystemIterator; SystemIterator; ++SystemIterator)
 		{
-			if (SystemIterator->ReferencesSourceEmitter(*Emitter))
+			if (SystemIterator->ReferencesInstanceEmitter(*Emitter))
 			{
 				SystemIterator->RequestCompile(false);
 
diff --git a/Engine/Plugins/FX/Niagara/Source/NiagaraEditor/Private/NiagaraEmitterFactoryNew.cpp b/Engine/Plugins/FX/Niagara/Source/NiagaraEditor/Private/NiagaraEmitterFactoryNew.cpp
--- a/Engine/Plugins/FX/Niagara/Source/NiagaraEditor/Private/NiagaraEmitterFactoryNew.cpp
+++ b/Engine/Plugins/FX/Niagara/Source/NiagaraEditor/Private/NiagaraEmitterFactoryNew.cpp
@@ -6,6 +6,7 @@
 #include "NiagaraScriptFactoryNew.h"
 #include "NiagaraGraph.h"
 #include "NiagaraScriptSource.h"
+#include "NiagaraSpriteRendererProperties.h"
 #include "NiagaraNodeOutput.h"
 #include "NiagaraEditorSettings.h"
 #include "AssetData.h"
@@ -45,6 +46,8 @@ UObject* UNiagaraEmitterFactoryNew::FactoryCreateNew(UClass* Class, UObject* InP
 
 		NewEmitter = NewObject<UNiagaraEmitter>(InParent, Class, Name, Flags | RF_Transactional);
 
+		NewEmitter->AddRenderer(NewObject<UNiagaraSpriteRendererProperties>(NewEmitter, "Renderer"));
+
 		UNiagaraScriptSource* Source = NewObject<UNiagaraScriptSource>(NewEmitter, NAME_None, RF_Transactional);
 		if (Source)
 		{
@@ -123,7 +126,8 @@ UObject* UNiagaraEmitterFactoryNew::FactoryCreateNew(UClass* Class, UObject* InP
 							{
 								FNiagaraVariable VelocityVar(FNiagaraTypeDefinition::GetVec3Def(), *(TEXT("Constants.Emitter.") + CallNode->GetFunctionName() + TEXT(".Velocity")));
 								VelocityVar.SetValue(FVector(0.0f, 0.0f, 100.0f));
-								NewEmitter->SpawnScriptProps.Script->RapidIterationParameters.AddParameter(VelocityVar);
+								bool bAddParameterIfMissing = true;
+								NewEmitter->SpawnScriptProps.Script->RapidIterationParameters.SetParameterData(VelocityVar.GetData(), VelocityVar, bAddParameterIfMissing);
 							}
 						}
 						else
diff --git a/Engine/Plugins/FX/Niagara/Source/NiagaraEditor/Private/NiagaraGraph.cpp b/Engine/Plugins/FX/Niagara/Source/NiagaraEditor/Private/NiagaraGraph.cpp
--- a/Engine/Plugins/FX/Niagara/Source/NiagaraEditor/Private/NiagaraGraph.cpp
+++ b/Engine/Plugins/FX/Niagara/Source/NiagaraEditor/Private/NiagaraGraph.cpp
@@ -29,7 +29,6 @@
 #include "NiagaraNode.h"
 #include "EdGraphSchema_Niagara.h"
 #include "ViewModels/Stack/NiagaraParameterHandle.h"
-#include "ScopedTransaction.h"
 
 DECLARE_CYCLE_STAT(TEXT("NiagaraEditor - Graph - FindInputNodes"), STAT_NiagaraEditor_Graph_FindInputNodes, STATGROUP_NiagaraEditor);
 DECLARE_CYCLE_STAT(TEXT("NiagaraEditor - Graph - FindInputNodes_NotFilterUsage"), STAT_NiagaraEditor_Graph_FindInputNodes_NotFilterUsage, STATGROUP_NiagaraEditor);
@@ -538,7 +537,6 @@ void UNiagaraGraph::AddParameter(const FNiagaraVariable& Parameter)
 
 void UNiagaraGraph::RemoveParameter(const FNiagaraVariable& Parameter, const bool bNotifyGraphChanged /*= true*/)
 {
-	FScopedTransaction RemoveParametersWithPins(LOCTEXT("RemoveParametersWithPins", "Remove parameter and referenced pins"));
 	FNiagaraGraphParameterReferenceCollection* ReferenceCollection = Parameters.Find(Parameter);
 	if (ReferenceCollection)
 	{
@@ -549,7 +547,7 @@ void UNiagaraGraph::RemoveParameter(const FNiagaraVariable& Parameter, const boo
 		{
 			const FNiagaraGraphParameterReference& Reference = ReferenceCollection->ParameterReferences[Index];
 			UNiagaraNode* Node = Reference.Value.Get();
-			if (Node)
+			if (Node && Node->GetGraph() == this)
 			{
 				UEdGraphPin* Pin = Node->GetPinByPersistentGuid(Reference.Key);
 				if (Pin)
@@ -579,8 +577,6 @@ bool UNiagaraGraph::RenameParameter(const FNiagaraVariable& Parameter, FName New
 	}
 	bIsRenamingParameter = true;
 
-	FScopedTransaction RemoveParametersWithPins(LOCTEXT("RenameParameter", "Rename parameter, referenced pins and metadata"));
-
 	// Prevent finding all parameters and metadata when renaming each pin.
 	SetFindParametersAllowed(false);
 	
@@ -596,7 +592,7 @@ bool UNiagaraGraph::RenameParameter(const FNiagaraVariable& Parameter, FName New
 		for (FNiagaraGraphParameterReference& Reference : NewReferences.ParameterReferences)
 		{
 			UNiagaraNode* Node = Reference.Value.Get();
-			if (Node)
+			if (Node && Node->GetGraph() == this)
 			{
 				UEdGraphPin* Pin = Node->GetPinByPersistentGuid(Reference.Key);
 				if (Pin)
diff --git a/Engine/Plugins/FX/Niagara/Source/NiagaraEditor/Private/NiagaraHlslTranslator.cpp b/Engine/Plugins/FX/Niagara/Source/NiagaraEditor/Private/NiagaraHlslTranslator.cpp
--- a/Engine/Plugins/FX/Niagara/Source/NiagaraEditor/Private/NiagaraHlslTranslator.cpp
+++ b/Engine/Plugins/FX/Niagara/Source/NiagaraEditor/Private/NiagaraHlslTranslator.cpp
@@ -3500,7 +3500,9 @@ bool FHlslNiagaraTranslator::ParameterMapRegisterExternalConstantNamespaceVariab
 			Collection = ParamMapHistories[InParamMapHistoryIdx].IsParameterCollectionParameter(InVariable, bMissingParameter);
 			if (Collection && bMissingParameter)
 			{
-				Error(FText::Format(LOCTEXT("MissingNPCParameterError", "Parameter {0} was not found in Parameter Collection {1}"), FText::FromName(InVariable.GetName()), FText::FromString(Collection->GetFullName())), InNode, nullptr);
+				Error(FText::Format(LOCTEXT("MissingNPCParameterError", "Parameter named {0} of type {1} was not found in Parameter Collection {2}"),
+					FText::FromName(InVariable.GetName()), InVariable.GetType().GetNameText(), FText::FromString(Collection->GetFullName())), InNode, InDefaultPin);
+				return false;
 			}
 		}
 
@@ -3523,6 +3525,12 @@ bool FHlslNiagaraTranslator::ParameterMapRegisterExternalConstantNamespaceVariab
 					if (Collection)
 					{
 						DataInterface = Collection->GetDefaultInstance()->GetParameterStore().GetDataInterface(InVariable);
+						if (DataInterface == nullptr)
+						{
+							Error(FText::Format(LOCTEXT("ParameterCollectionDataInterfaceNotFoundErrorFormat", "Data interface named {0} of type {1} was not found in Parameter Collection {2}"),
+								FText::FromName(InVariable.GetName()), InVariable.GetType().GetNameText(), FText::FromString(Collection->GetFullName())), InNode, InDefaultPin);
+							return false;
+						}
 					}
 					else
 					{
@@ -4220,23 +4228,24 @@ int32 FHlslNiagaraTranslator::RegisterDataInterface(FNiagaraVariable& Var, UNiag
 	}
 
 	//If we get here then this is a new data interface.
+	FName DataInterfaceName;
+	if (FNiagaraParameterMapHistory::IsAliasedEmitterParameter(Var.GetName().ToString()))
+	{
+		FNiagaraVariable AliasedVar = ActiveHistoryForFunctionCalls.ResolveAliases(Var);
+		DataInterfaceName = AliasedVar.GetName();
+	}
+	else
+	{
+		DataInterfaceName = Var.GetName();
+	}
+
 	int32 Idx = CompilationOutput.ScriptData.DataInterfaceInfo.IndexOfByPredicate([&](const FNiagaraScriptDataInterfaceCompileInfo& OtherInfo)
 	{
-		return OtherInfo.Name == Var.GetName();
+		return OtherInfo.Name == DataInterfaceName;
 	});
+
 	if (Idx == INDEX_NONE)
 	{
-		FName DataInterfaceName;
-		if (FNiagaraParameterMapHistory::IsAliasedEmitterParameter(Var.GetName().ToString()))
-		{
-			FNiagaraVariable AliasedVar = ActiveHistoryForFunctionCalls.ResolveAliases(Var);
-			DataInterfaceName = AliasedVar.GetName();
-		}
-		else
-		{
-			DataInterfaceName = Var.GetName();
-		}
-
 		Idx = CompilationOutput.ScriptData.DataInterfaceInfo.AddDefaulted();
 		CompilationOutput.ScriptData.DataInterfaceInfo[Idx].Name = DataInterfaceName;
 		CompilationOutput.ScriptData.DataInterfaceInfo[Idx].Type = Var.GetType();
@@ -4399,6 +4408,13 @@ void FHlslNiagaraTranslator::FunctionCall(UNiagaraNodeFunctionCall* FunctionNode
 	}
 	
 	RegisterFunctionCall(ScriptUsage, Name, FullName,  FunctionNode->NodeGuid, Source, Signature, bCustomHlsl, CustomHlsl, Inputs, CallInputs, CallOutputs, OutputSignature);
+
+	if (OutputSignature.IsValid() == false)
+	{
+		Error(LOCTEXT("FunctionCallInvalidSignature", "Could not generate a valid function signature."), FunctionNode, nullptr);
+		return;
+	}
+
 	GenerateFunctionCall(OutputSignature, Inputs, Outputs);
 
 	if (bCustomHlsl)
@@ -4669,6 +4685,13 @@ void FHlslNiagaraTranslator::RegisterFunctionCall(ENiagaraScriptUsage ScriptUsag
 		{
 			SCOPE_CYCLE_COUNTER(STAT_NiagaraEditor_Module_NiagaraHLSLTranslator_FuncBody);
 
+			if (OutSignature.Name == NAME_None)
+			{
+				const FString* ModuleAlias = ActiveHistoryForFunctionCalls.GetModuleAlias();
+				Error(FText::Format(LOCTEXT("FunctionCallMissingFunction", "Function call signature does not reference a function. Top-level module: {0} Source: {1}"), ModuleAlias ? FText::FromString(*ModuleAlias) : FText::FromString(TEXT("Unknown module")), FText::FromString(CompileOptions.FullName)), nullptr, nullptr);
+				return;
+			}
+
 			//We've not compiled this function yet so compile it now.
 			EnterFunction(InName, OutSignature, Inputs, CallNodeId);
 
@@ -4838,7 +4861,15 @@ void FHlslNiagaraTranslator::RegisterFunctionCall(ENiagaraScriptUsage ScriptUsag
 			// interface. It could be that the existing node has been removed and the graph
 			// needs to be refactored. If that's the case, emit an error.
 			UObject* const* FoundCDO = CompileData->CDOs.Find(Info.Type.GetClass());
-			check(FoundCDO != nullptr);
+			if (FoundCDO == nullptr)
+			{
+				// If the cdo wasn't found, the data interface was not passed through a parameter map and so it won't be bound correctly, so add a compile error
+				// and invalidate the signature.
+				Error(LOCTEXT("DataInterfaceNotFoundInParameterMap", "Data interfaces can not be sampled directly, they must be passed through a parameter map to be bound correctly."), nullptr, nullptr);
+				OutSignature.Name = NAME_None;
+				return;
+			}
+
 			UNiagaraDataInterface* CDO = Cast<UNiagaraDataInterface>(*FoundCDO);
 			if (CDO != nullptr && OutSignature.bMemberFunction)
 			{
@@ -5415,15 +5446,28 @@ int32 FHlslNiagaraTranslator::CompileOutputPin(const UEdGraphPin* InPin)
 	return Ret;
 }
 
-void FHlslNiagaraTranslator::Error(FText ErrorText, UNiagaraNode* Node, UEdGraphPin* Pin)
+void FHlslNiagaraTranslator::Error(FText ErrorText, const UNiagaraNode* Node, const UEdGraphPin* Pin)
 {
 	FString NodePinStr = TEXT("");
 	FString NodePinPrefix = TEXT(" - ");
 	FString NodePinSuffix = TEXT("");
-	if (Node && Node->GetName().Len() > 0)
+	if (Node)
 	{
-		NodePinStr += TEXT("Node: ") + Node->GetName();
-		NodePinSuffix = TEXT(" - ");
+		FString NodeTitle = Node->GetNodeTitle(ENodeTitleType::FullTitle).ToString();
+		if (NodeTitle.Len() > 0)
+		{
+			NodePinStr += TEXT("Node: ") + NodeTitle;
+			NodePinSuffix = TEXT(" - ");
+		}
+		else
+		{
+			FString NodeName = Node->GetName();
+			if (NodeName.Len() > 0)
+			{
+				NodePinStr += TEXT("Node: ") + NodeName;
+				NodePinSuffix = TEXT(" - ");
+			}
+		}
 	}
 	if (Pin && Pin->PinFriendlyName.ToString().Len() > 0)
 	{
@@ -5436,7 +5480,7 @@ void FHlslNiagaraTranslator::Error(FText ErrorText, UNiagaraNode* Node, UEdGraph
 	TranslateResults.NumErrors++;
 }
 
-void FHlslNiagaraTranslator::Warning(FText WarningText, UNiagaraNode* Node, UEdGraphPin* Pin)
+void FHlslNiagaraTranslator::Warning(FText WarningText, const UNiagaraNode* Node, const UEdGraphPin* Pin)
 {
 	FString NodePinStr = TEXT("");
 	FString NodePinPrefix = TEXT(" - ");
diff --git a/Engine/Plugins/FX/Niagara/Source/NiagaraEditor/Private/NiagaraHlslTranslator.h b/Engine/Plugins/FX/Niagara/Source/NiagaraEditor/Private/NiagaraHlslTranslator.h
--- a/Engine/Plugins/FX/Niagara/Source/NiagaraEditor/Private/NiagaraHlslTranslator.h
+++ b/Engine/Plugins/FX/Niagara/Source/NiagaraEditor/Private/NiagaraHlslTranslator.h
@@ -448,8 +448,8 @@ public:
 	virtual void Convert(class UNiagaraNodeConvert* Convert, TArray <int32>& Inputs, TArray<int32>& Outputs);
 	virtual void If(TArray<FNiagaraVariable>& Vars, int32 Condition, TArray<int32>& PathA, TArray<int32>& PathB, TArray<int32>& Outputs);
 
-	virtual void Error(FText ErrorText, UNiagaraNode* Node, UEdGraphPin* Pin);
-	virtual void Warning(FText WarningText, UNiagaraNode* Node, UEdGraphPin* Pin);
+	virtual void Error(FText ErrorText, const UNiagaraNode* Node, const UEdGraphPin* Pin);
+	virtual void Warning(FText WarningText, const UNiagaraNode* Node, const UEdGraphPin* Pin);
 
 	virtual bool GetFunctionParameter(const FNiagaraVariable& Parameter, int32& OutParam)const;
 
diff --git a/Engine/Plugins/FX/Niagara/Source/NiagaraEditor/Private/NiagaraNodeAssignment.cpp b/Engine/Plugins/FX/Niagara/Source/NiagaraEditor/Private/NiagaraNodeAssignment.cpp
--- a/Engine/Plugins/FX/Niagara/Source/NiagaraEditor/Private/NiagaraNodeAssignment.cpp
+++ b/Engine/Plugins/FX/Niagara/Source/NiagaraEditor/Private/NiagaraNodeAssignment.cpp
@@ -226,12 +226,15 @@ void UNiagaraNodeAssignment::BuildParameterMenu(FMenuBuilder& MenuBuilder, ENiag
 		FString VarDefaultValue = FNiagaraConstants::GetAttributeDefaultValue(AvailableParameter);
 		const FText TooltipDesc = FText::Format(LOCTEXT("SetFunctionPopupTooltip", "Description: Set the parameter {0}. {1}"), FText::FromName(AvailableParameter.GetName()), VarDesc);
 		FText CategoryName = LOCTEXT("ModuleSetCategory", "Set Specific Parameters");
+		bool bCanExecute = AssignmentTargets.Contains(AvailableParameter) == false;
 
 		MenuBuilder.AddMenuEntry(
 			NameText,
 			TooltipDesc,
 			FSlateIcon(),
-			FUIAction(FExecuteAction::CreateUObject(this, &UNiagaraNodeAssignment::AddParameter, AvailableParameter, VarDefaultValue)));
+			FUIAction(
+				FExecuteAction::CreateUObject(this, &UNiagaraNodeAssignment::AddParameter, AvailableParameter, VarDefaultValue),
+				FCanExecuteAction::CreateLambda([bCanExecute] { return bCanExecute; })));
 	}
 
 }
diff --git a/Engine/Plugins/FX/Niagara/Source/NiagaraEditor/Private/NiagaraNodeReroute.cpp b/Engine/Plugins/FX/Niagara/Source/NiagaraEditor/Private/NiagaraNodeReroute.cpp
--- a/Engine/Plugins/FX/Niagara/Source/NiagaraEditor/Private/NiagaraNodeReroute.cpp
+++ b/Engine/Plugins/FX/Niagara/Source/NiagaraEditor/Private/NiagaraNodeReroute.cpp
@@ -6,7 +6,7 @@
 
 #define LOCTEXT_NAMESPACE "NiagaraNodeReroute"
 
-const char* PC_Wildcard = "wildcard";
+const char* PC_Wildcard_Niagara = "wildcard";
 
 UNiagaraNodeReroute::UNiagaraNodeReroute(const FObjectInitializer& ObjectInitializer)
 	: Super(ObjectInitializer)
@@ -32,11 +32,11 @@ void UNiagaraNodeReroute::PostLoad()
 void UNiagaraNodeReroute::AllocateDefaultPins()
 {
 	const FName InputPinName(TEXT("InputPin"));
-	UEdGraphPin* MyInputPin = CreatePin(EGPD_Input, PC_Wildcard, InputPinName);
+	UEdGraphPin* MyInputPin = CreatePin(EGPD_Input, PC_Wildcard_Niagara, InputPinName);
 	MyInputPin->bDefaultValueIsIgnored = true;
 
 	const FName OutputPinName(TEXT("OutputPin"));
-	CreatePin(EGPD_Output, PC_Wildcard, OutputPinName);
+	CreatePin(EGPD_Output, PC_Wildcard_Niagara, OutputPinName);
 }
 
 FText UNiagaraNodeReroute::GetTooltipText() const
@@ -138,7 +138,7 @@ void UNiagaraNodeReroute::PropagatePinType()
 
 	for (UEdGraphPin* Inputs : MyInputPin->LinkedTo)
 	{
-		if (Inputs->PinType.PinCategory != PC_Wildcard)
+		if (Inputs->PinType.PinCategory != PC_Wildcard_Niagara)
 		{
 			PropagatePinTypeFromDirection(true);
 			return;
@@ -147,7 +147,7 @@ void UNiagaraNodeReroute::PropagatePinType()
 
 	for (UEdGraphPin* Outputs : MyOutputPin->LinkedTo)
 	{
-		if (Outputs->PinType.PinCategory != PC_Wildcard)
+		if (Outputs->PinType.PinCategory != PC_Wildcard_Niagara)
 		{
 			PropagatePinTypeFromDirection(false);
 			return;
@@ -170,11 +170,11 @@ void UNiagaraNodeReroute::PropagatePinType()
 		// Revert to wildcard
 		MyInputPin->BreakAllPinLinks();
 		MyInputPin->PinType.ResetToDefaults();
-		MyInputPin->PinType.PinCategory = PC_Wildcard;
+		MyInputPin->PinType.PinCategory = PC_Wildcard_Niagara;
 
 		MyOutputPin->BreakAllPinLinks();
 		MyOutputPin->PinType.ResetToDefaults();
-		MyOutputPin->PinType.PinCategory = PC_Wildcard;
+		MyOutputPin->PinType.PinCategory = PC_Wildcard_Niagara;
 	}
 }
 
diff --git a/Engine/Plugins/FX/Niagara/Source/NiagaraEditor/Private/NiagaraParameterMapHistory.cpp b/Engine/Plugins/FX/Niagara/Source/NiagaraEditor/Private/NiagaraParameterMapHistory.cpp
--- a/Engine/Plugins/FX/Niagara/Source/NiagaraEditor/Private/NiagaraParameterMapHistory.cpp
+++ b/Engine/Plugins/FX/Niagara/Source/NiagaraEditor/Private/NiagaraParameterMapHistory.cpp
@@ -235,6 +235,12 @@ FNiagaraVariable FNiagaraParameterMapHistory::ResolveAliases(const FNiagaraVaria
 
 FName FNiagaraParameterMapHistory::ResolveEmitterAlias(const FName& InName, const FString& InAlias)
 {
+	// If the alias is empty than the name can't be resolved.
+	if (InAlias.IsEmpty())
+	{
+		return InName;
+	}
+
 	FNiagaraVariable Var(FNiagaraTypeDefinition::GetFloatDef(), InName);
 	TMap<FString, FString> ResolveMap;
 	ResolveMap.Add(TEXT("Emitter"), InAlias);
diff --git a/Engine/Plugins/FX/Niagara/Source/NiagaraEditor/Private/NiagaraScriptMergeManager.cpp b/Engine/Plugins/FX/Niagara/Source/NiagaraEditor/Private/NiagaraScriptMergeManager.cpp
--- a/Engine/Plugins/FX/Niagara/Source/NiagaraEditor/Private/NiagaraScriptMergeManager.cpp
+++ b/Engine/Plugins/FX/Niagara/Source/NiagaraEditor/Private/NiagaraScriptMergeManager.cpp
@@ -194,9 +194,25 @@ FNiagaraStackFunctionMergeAdapter::FNiagaraStackFunctionMergeAdapter(FString InE
 	for (const FNiagaraVariable& RapidIterationParameter : RapidIterationParameters)
 	{
 		FNiagaraParameterHandle AliasedInputHandle(*RapidIterationParameter.GetName().ToString().RightChop(RapidIterationParameterNamePrefix.Len()));
-		if (AliasedInputHandle.GetNamespace().ToString() == FunctionCallNode->GetFunctionName() && AliasedInputsAdded.Contains(AliasedInputHandle.GetParameterHandleString().ToString()) == false)
+		if (AliasedInputHandle.GetNamespace().ToString() == FunctionCallNode->GetFunctionName())
 		{
-			InputOverrides.Add(MakeShared<FNiagaraStackFunctionInputOverrideMergeAdapter>(UniqueEmitterName, *OwningScript.Get(), *FunctionCallNode.Get(), AliasedInputHandle.GetName().ToString(), RapidIterationParameter));
+			// Currently rapid iteration parameters for assignment nodes in emitter scripts get double aliased which prevents their inputs from
+			// being diffed correctly, so we need to un-mangle the names here so that the diffs are correct.
+			if (FunctionCallNode->IsA<UNiagaraNodeAssignment>() &&
+				(OwningScript->GetUsage() == ENiagaraScriptUsage::EmitterSpawnScript || OwningScript->GetUsage() == ENiagaraScriptUsage::EmitterUpdateScript))
+			{
+				FString InputName = AliasedInputHandle.GetName().ToString();
+				if (InputName.StartsWith(UniqueEmitterName + TEXT(".")))
+				{
+					FString UnaliasedInputName = TEXT("Emitter") + InputName.RightChop(UniqueEmitterName.Len());
+					AliasedInputHandle = FNiagaraParameterHandle(AliasedInputHandle.GetNamespace(), *UnaliasedInputName);
+				}
+			}
+
+			if (AliasedInputsAdded.Contains(AliasedInputHandle.GetParameterHandleString().ToString()) == false)
+			{
+				InputOverrides.Add(MakeShared<FNiagaraStackFunctionInputOverrideMergeAdapter>(UniqueEmitterName, *OwningScript.Get(), *FunctionCallNode.Get(), AliasedInputHandle.GetName().ToString(), RapidIterationParameter));
+			}
 		}
 	}
 }
@@ -1632,11 +1648,11 @@ FNiagaraScriptMergeManager::FApplyDiffResults FNiagaraScriptMergeManager::AddInp
 		}
 	}
 
+	FNiagaraParameterHandle FunctionInputHandle(FNiagaraParameterHandle::ModuleNamespace, *OverrideToAdd->GetInputName());
+	FNiagaraParameterHandle AliasedFunctionInputHandle = FNiagaraParameterHandle::CreateAliasedModuleParameterHandle(FunctionInputHandle, &TargetFunctionCall);
+
 	if (OverrideToAdd->GetOverridePin() != nullptr)
 	{
-		FNiagaraParameterHandle FunctionInputHandle(FNiagaraParameterHandle::ModuleNamespace, *OverrideToAdd->GetInputName());
-		FNiagaraParameterHandle AliasedFunctionInputHandle = FNiagaraParameterHandle::CreateAliasedModuleParameterHandle(FunctionInputHandle, &TargetFunctionCall);
-
 		const UEdGraphSchema_Niagara* NiagaraSchema = GetDefault<UEdGraphSchema_Niagara>();
 		FNiagaraTypeDefinition InputType = NiagaraSchema->PinToTypeDefinition(OverrideToAdd->GetOverridePin());
 
@@ -1717,12 +1733,12 @@ FNiagaraScriptMergeManager::FApplyDiffResults FNiagaraScriptMergeManager::AddInp
 	{
 		if (OverrideToAdd->GetLocalValueRapidIterationParameter().IsSet())
 		{
-			FString RapidIterationParameterName = FString::Printf(TEXT("Constants.%s.%s.%s"), *UniqueEmitterName, *TargetFunctionCall.GetFunctionName(), *OverrideToAdd->GetInputName());
-			FNiagaraVariable RapidIterationParameter(OverrideToAdd->GetLocalValueRapidIterationParameter().GetValue().GetType(), *RapidIterationParameterName);
+			FNiagaraVariable RapidIterationParameter = FNiagaraStackGraphUtilities::CreateRapidIterationParameter(
+				UniqueEmitterName, OwningScript.GetUsage(), AliasedFunctionInputHandle.GetParameterHandleString(), OverrideToAdd->GetLocalValueRapidIterationParameter().GetValue().GetType());
 			const uint8* SourceData = OverrideToAdd->GetOwningScript()->RapidIterationParameters.GetParameterData(OverrideToAdd->GetLocalValueRapidIterationParameter().GetValue());
 			OwningScript.Modify();
-			OwningScript.RapidIterationParameters.AddParameter(RapidIterationParameter);
-			OwningScript.RapidIterationParameters.SetParameterData(SourceData, RapidIterationParameter);
+			bool bAddParameterIfMissing = true;
+			OwningScript.RapidIterationParameters.SetParameterData(SourceData, RapidIterationParameter, bAddParameterIfMissing);
 			Results.bSucceeded = true;
 		}
 		else
@@ -1854,7 +1870,7 @@ FNiagaraScriptMergeManager::FApplyDiffResults FNiagaraScriptMergeManager::ApplyS
 	// Apply the graph actions.
 	for (TSharedRef<FNiagaraStackFunctionMergeAdapter> RemoveModule : RemoveModules)
 	{
-		bool bRemoveResults = FNiagaraStackGraphUtilities::RemoveModuleFromStack(*RemoveModule->GetFunctionCallNode());
+		bool bRemoveResults = FNiagaraStackGraphUtilities::RemoveModuleFromStack(*BaseScriptStackAdapter->GetScript(), *RemoveModule->GetFunctionCallNode());
 		if (bRemoveResults == false)
 		{
 			Results.bSucceeded = false;
diff --git a/Engine/Plugins/FX/Niagara/Source/NiagaraEditor/Private/NiagaraScriptSource.cpp b/Engine/Plugins/FX/Niagara/Source/NiagaraEditor/Private/NiagaraScriptSource.cpp
--- a/Engine/Plugins/FX/Niagara/Source/NiagaraEditor/Private/NiagaraScriptSource.cpp
+++ b/Engine/Plugins/FX/Niagara/Source/NiagaraEditor/Private/NiagaraScriptSource.cpp
@@ -204,8 +204,8 @@ void InitializeNewRapidIterationParametersForNode(const UEdGraphSchema_Niagara*
 						{
 							FNiagaraVariable DefaultVariable = Schema->PinToNiagaraVariable(DefaultPin, true);
 							check(DefaultVariable.GetData() != nullptr);
-							RapidIterationParameters.AddParameter(RapidIterationParameter, false);
-							RapidIterationParameters.SetParameterData(DefaultVariable.GetData(), RapidIterationParameter);
+							bool bAddParameterIfMissing = true;
+							RapidIterationParameters.SetParameterData(DefaultVariable.GetData(), RapidIterationParameter, bAddParameterIfMissing);
 						}
 					}
 				}
diff --git a/Engine/Plugins/FX/Niagara/Source/NiagaraEditor/Private/NiagaraStackFunctionInputBinder.cpp b/Engine/Plugins/FX/Niagara/Source/NiagaraEditor/Private/NiagaraStackFunctionInputBinder.cpp
--- a/Engine/Plugins/FX/Niagara/Source/NiagaraEditor/Private/NiagaraStackFunctionInputBinder.cpp
+++ b/Engine/Plugins/FX/Niagara/Source/NiagaraEditor/Private/NiagaraStackFunctionInputBinder.cpp
@@ -1,4 +1,4 @@
-// Copyright 1998-2016 Epic Games, Inc. All Rights Reserved.
+// Copyright 1998-2018 Epic Games, Inc. All Rights Reserved.
 
 #include "NiagaraStackFunctionInputBinder.h"
 #include "NiagaraScript.h"
@@ -187,6 +187,11 @@ FNiagaraTypeDefinition FNiagaraStackFunctionInputBinder::GetInputType() const
 	return InputType;
 }
 
+UNiagaraNodeFunctionCall* FNiagaraStackFunctionInputBinder::GetFunctionCallNode() const
+{
+	return FunctionCallNode.Get();
+}
+
 TArray<uint8> FNiagaraStackFunctionInputBinder::GetData() const
 {
 	checkf(Script.IsValid(), TEXT("Bound script is not valid"));
diff --git a/Engine/Plugins/FX/Niagara/Source/NiagaraEditor/Private/Sequencer/LevelSequence/NiagaraSystemSpawnSection.cpp b/Engine/Plugins/FX/Niagara/Source/NiagaraEditor/Private/Sequencer/LevelSequence/NiagaraSystemSpawnSection.cpp
--- a/Engine/Plugins/FX/Niagara/Source/NiagaraEditor/Private/Sequencer/LevelSequence/NiagaraSystemSpawnSection.cpp
+++ b/Engine/Plugins/FX/Niagara/Source/NiagaraEditor/Private/Sequencer/LevelSequence/NiagaraSystemSpawnSection.cpp
@@ -1,4 +1,4 @@
-// Copyright 1998-2017 Epic Games, Inc. All Rights Reserved.
+// Copyright 1998-2018 Epic Games, Inc. All Rights Reserved.
 
 #include "Sequencer/LevelSequence/NiagaraSystemSpawnSection.h"
 #include "MovieScene/MovieSceneNiagaraSystemSpawnSection.h"
diff --git a/Engine/Plugins/FX/Niagara/Source/NiagaraEditor/Private/Sequencer/LevelSequence/NiagaraSystemSpawnSection.h b/Engine/Plugins/FX/Niagara/Source/NiagaraEditor/Private/Sequencer/LevelSequence/NiagaraSystemSpawnSection.h
--- a/Engine/Plugins/FX/Niagara/Source/NiagaraEditor/Private/Sequencer/LevelSequence/NiagaraSystemSpawnSection.h
+++ b/Engine/Plugins/FX/Niagara/Source/NiagaraEditor/Private/Sequencer/LevelSequence/NiagaraSystemSpawnSection.h
@@ -1,4 +1,4 @@
-// Copyright 1998-2017 Epic Games, Inc. All Rights Reserved.
+// Copyright 1998-2018 Epic Games, Inc. All Rights Reserved.
 
 #pragma once
 
diff --git a/Engine/Plugins/FX/Niagara/Source/NiagaraEditor/Private/Sequencer/LevelSequence/NiagaraSystemTrackEditor.cpp b/Engine/Plugins/FX/Niagara/Source/NiagaraEditor/Private/Sequencer/LevelSequence/NiagaraSystemTrackEditor.cpp
--- a/Engine/Plugins/FX/Niagara/Source/NiagaraEditor/Private/Sequencer/LevelSequence/NiagaraSystemTrackEditor.cpp
+++ b/Engine/Plugins/FX/Niagara/Source/NiagaraEditor/Private/Sequencer/LevelSequence/NiagaraSystemTrackEditor.cpp
@@ -1,4 +1,4 @@
-// Copyright 1998-2017 Epic Games, Inc. All Rights Reserved.
+// Copyright 1998-2018 Epic Games, Inc. All Rights Reserved.
 
 #include "NiagaraSystemTrackEditor.h"
 #include "MovieScene/MovieSceneNiagaraSystemTrack.h"
diff --git a/Engine/Plugins/FX/Niagara/Source/NiagaraEditor/Private/Sequencer/LevelSequence/NiagaraSystemTrackEditor.h b/Engine/Plugins/FX/Niagara/Source/NiagaraEditor/Private/Sequencer/LevelSequence/NiagaraSystemTrackEditor.h
--- a/Engine/Plugins/FX/Niagara/Source/NiagaraEditor/Private/Sequencer/LevelSequence/NiagaraSystemTrackEditor.h
+++ b/Engine/Plugins/FX/Niagara/Source/NiagaraEditor/Private/Sequencer/LevelSequence/NiagaraSystemTrackEditor.h
@@ -1,4 +1,4 @@
-// Copyright 1998-2017 Epic Games, Inc. All Rights Reserved.
+// Copyright 1998-2018 Epic Games, Inc. All Rights Reserved.
 
 #pragma once
 
diff --git a/Engine/Plugins/FX/Niagara/Source/NiagaraEditor/Private/Sequencer/NiagaraSequence/Sections/MovieSceneNiagaraEmitterSection.cpp b/Engine/Plugins/FX/Niagara/Source/NiagaraEditor/Private/Sequencer/NiagaraSequence/Sections/MovieSceneNiagaraEmitterSection.cpp
--- a/Engine/Plugins/FX/Niagara/Source/NiagaraEditor/Private/Sequencer/NiagaraSequence/Sections/MovieSceneNiagaraEmitterSection.cpp
+++ b/Engine/Plugins/FX/Niagara/Source/NiagaraEditor/Private/Sequencer/NiagaraSequence/Sections/MovieSceneNiagaraEmitterSection.cpp
@@ -1,4 +1,4 @@
-// Copyright 1998-2016 Epic Games, Inc. All Rights Reserved.
+// Copyright 1998-2018 Epic Games, Inc. All Rights Reserved.
 
 #include "MovieSceneNiagaraEmitterSection.h"
 #include "NiagaraEmitterSection.h"
@@ -99,8 +99,9 @@ TSharedPtr<FStructOnScope> GetKeyStruct(TMovieSceneChannelHandle<FMovieSceneNiag
 	int32 KeyValueIndex = Channel.Get()->GetData().GetIndex(InHandle);
 	if (KeyValueIndex != INDEX_NONE)
 	{
-		FNiagaraEmitterSectionKey KeyValue = Channel.Get()->GetData().GetValues()[KeyValueIndex];
-		return MakeShared<FStructOnScope>(KeyValue.Value.GetType().GetStruct(), KeyValue.Value.GetData());
+		FNiagaraTypeDefinition KeyType = Channel.Get()->GetData().GetValues()[KeyValueIndex].Value.GetType();
+		uint8* KeyData = Channel.Get()->GetData().GetValues()[KeyValueIndex].Value.GetData();
+		return MakeShared<FStructOnScope>(KeyType.GetStruct(), KeyData);
 	}
 	return TSharedPtr<FStructOnScope>();
 }
@@ -257,10 +258,29 @@ void UMovieSceneNiagaraEmitterSection::UpdateSectionFromTimeRangeModule(const FF
 	{
 		float ModuleStartTime = StartTimeBinder.GetValue<float>();
 		float ModuleLength = LengthBinder.GetValue<float>();
+
+		if (ModuleLength < 0)
+		{
+			// TODO: Add ui support for this issue rather than a log error.
+			UE_LOG(LogNiagaraEditor, Error, TEXT("Invalid length in niagara editor timeline.  Bound Module: %s Bound Input: %s"),
+				LengthBinder.GetFunctionCallNode() != nullptr ? *LengthBinder.GetFunctionCallNode()->GetFunctionName() : TEXT("Unknown"),
+				*LengthBinder.GetInputName().ToString());
+			ModuleLength = 0;
+		}
+
+		FFrameNumber StartFrame = (ModuleStartTime * InTickResolution).RoundToFrame();
+		FFrameNumber EndFrame = ((ModuleStartTime + ModuleLength) * InTickResolution).RoundToFrame();
+		if (EndFrame < StartFrame)
+		{
+			// The frame value has overflowed and is negative so clamp to the max frame.
+			// TODO: Add ui support for this issue rather than a log error.
+			UE_LOG(LogNiagaraEditor, Error, TEXT("Invalid length in niagara editor timeline.  Bound Module: %s Bound Input: %s"),
+				LengthBinder.GetFunctionCallNode() != nullptr ? *LengthBinder.GetFunctionCallNode()->GetFunctionName() : TEXT("Unknown"),
+				*LengthBinder.GetInputName().ToString());
+			EndFrame.Value = TNumericLimits<int32>::Max();
+		}
 		
-		SetRange(TRange<FFrameNumber>(
-			(ModuleStartTime * InTickResolution).RoundToFrame(),
-			((ModuleStartTime + ModuleLength) * InTickResolution).RoundToFrame()));
+		SetRange(TRange<FFrameNumber>(StartFrame, EndFrame));
 	}
 	else
 	{
@@ -375,7 +395,7 @@ void UMovieSceneNiagaraEmitterSection::UpdateKeyModulesFromSection(FChannelAndMo
 	{
 		ChannelAndModules.ModulesAndBinders.RemoveAll(
 			[=](FModuleAndBinders& ModuleAndBinders) { return ModuleAndBinders.Module == ModuleWithMissingKey; });
-		FNiagaraStackGraphUtilities::RemoveModuleFromStack(*ModuleWithMissingKey);
+		FNiagaraStackGraphUtilities::RemoveModuleFromStack(GetSystemViewModel().GetSystem(), GetEmitterHandleViewModel()->GetId(), *ModuleWithMissingKey);
 	}
 
 	// Create new modules for new keys.
diff --git a/Engine/Plugins/FX/Niagara/Source/NiagaraEditor/Private/Sequencer/NiagaraSequence/Sections/MovieSceneNiagaraEmitterSection.h b/Engine/Plugins/FX/Niagara/Source/NiagaraEditor/Private/Sequencer/NiagaraSequence/Sections/MovieSceneNiagaraEmitterSection.h
--- a/Engine/Plugins/FX/Niagara/Source/NiagaraEditor/Private/Sequencer/NiagaraSequence/Sections/MovieSceneNiagaraEmitterSection.h
+++ b/Engine/Plugins/FX/Niagara/Source/NiagaraEditor/Private/Sequencer/NiagaraSequence/Sections/MovieSceneNiagaraEmitterSection.h
@@ -1,4 +1,4 @@
-// Copyright 1998-2016 Epic Games, Inc. All Rights Reserved.
+// Copyright 1998-2018 Epic Games, Inc. All Rights Reserved.
 
 #pragma once
 
diff --git a/Engine/Plugins/FX/Niagara/Source/NiagaraEditor/Private/Sequencer/NiagaraSequence/Sections/NiagaraEmitterSection.cpp b/Engine/Plugins/FX/Niagara/Source/NiagaraEditor/Private/Sequencer/NiagaraSequence/Sections/NiagaraEmitterSection.cpp
--- a/Engine/Plugins/FX/Niagara/Source/NiagaraEditor/Private/Sequencer/NiagaraSequence/Sections/NiagaraEmitterSection.cpp
+++ b/Engine/Plugins/FX/Niagara/Source/NiagaraEditor/Private/Sequencer/NiagaraSequence/Sections/NiagaraEmitterSection.cpp
@@ -1,4 +1,4 @@
-// Copyright 1998-2016 Epic Games, Inc. All Rights Reserved.
+// Copyright 1998-2018 Epic Games, Inc. All Rights Reserved.
 
 #include "NiagaraEmitterSection.h"
 #include "MovieSceneNiagaraEmitterSection.h"
diff --git a/Engine/Plugins/FX/Niagara/Source/NiagaraEditor/Private/Sequencer/NiagaraSequence/Sections/NiagaraEmitterSection.h b/Engine/Plugins/FX/Niagara/Source/NiagaraEditor/Private/Sequencer/NiagaraSequence/Sections/NiagaraEmitterSection.h
--- a/Engine/Plugins/FX/Niagara/Source/NiagaraEditor/Private/Sequencer/NiagaraSequence/Sections/NiagaraEmitterSection.h
+++ b/Engine/Plugins/FX/Niagara/Source/NiagaraEditor/Private/Sequencer/NiagaraSequence/Sections/NiagaraEmitterSection.h
@@ -1,4 +1,4 @@
-// Copyright 1998-2016 Epic Games, Inc. All Rights Reserved.
+// Copyright 1998-2018 Epic Games, Inc. All Rights Reserved.
 
 #pragma once
 
diff --git a/Engine/Plugins/FX/Niagara/Source/NiagaraEditor/Private/Toolkits/NiagaraScriptToolkit.cpp b/Engine/Plugins/FX/Niagara/Source/NiagaraEditor/Private/Toolkits/NiagaraScriptToolkit.cpp
--- a/Engine/Plugins/FX/Niagara/Source/NiagaraEditor/Private/Toolkits/NiagaraScriptToolkit.cpp
+++ b/Engine/Plugins/FX/Niagara/Source/NiagaraEditor/Private/Toolkits/NiagaraScriptToolkit.cpp
@@ -215,6 +215,8 @@ void FNiagaraScriptToolkit::Initialize( const EToolkitMode::Type Mode, const TSh
 		const FString TabInitializationPayload(TEXT(""));		// NOTE: Payload not currently used for table properties
 		SpawnToolkitTab( NodeGraphTabId, TabInitializationPayload, EToolkitTabSpot::Details );
 	}*/
+
+	bChangesDiscarded = false;
 }
 
 FName FNiagaraScriptToolkit::GetToolkitFName() const
@@ -660,7 +662,7 @@ void FNiagaraScriptToolkit::UpdateOriginalNiagaraScript()
 
 bool FNiagaraScriptToolkit::OnRequestClose()
 {
-	if (IsEditScriptDifferentFromOriginalScript())
+	if (bChangesDiscarded == false && IsEditScriptDifferentFromOriginalScript())
 	{
 		// find out the user wants to do with this dirty NiagaraScript
 		EAppReturnType::Type YesNoCancelReply = FMessageDialog::Open(EAppMsgType::YesNoCancel,
@@ -677,8 +679,8 @@ bool FNiagaraScriptToolkit::OnRequestClose()
 			break;
 
 		case EAppReturnType::No:
-			// exit
-			// do nothing.. bNiagaraScriptDirty = false;
+			// Set the changes discarded to avoid showing the dialog multiple times when request close is called multiple times on shut down.
+			bChangesDiscarded = true;
 			break;
 
 		case EAppReturnType::Cancel:
diff --git a/Engine/Plugins/FX/Niagara/Source/NiagaraEditor/Private/Toolkits/NiagaraScriptToolkit.h b/Engine/Plugins/FX/Niagara/Source/NiagaraEditor/Private/Toolkits/NiagaraScriptToolkit.h
--- a/Engine/Plugins/FX/Niagara/Source/NiagaraEditor/Private/Toolkits/NiagaraScriptToolkit.h
+++ b/Engine/Plugins/FX/Niagara/Source/NiagaraEditor/Private/Toolkits/NiagaraScriptToolkit.h
@@ -125,4 +125,5 @@ private:
 	FDelegateHandle OnEditedScriptGraphChangedHandle;
 
 	bool bEditedScriptHasPendingChanges;
+	bool bChangesDiscarded;
 };
diff --git a/Engine/Plugins/FX/Niagara/Source/NiagaraEditor/Private/Toolkits/NiagaraSystemToolkit.cpp b/Engine/Plugins/FX/Niagara/Source/NiagaraEditor/Private/Toolkits/NiagaraSystemToolkit.cpp
--- a/Engine/Plugins/FX/Niagara/Source/NiagaraEditor/Private/Toolkits/NiagaraSystemToolkit.cpp
+++ b/Engine/Plugins/FX/Niagara/Source/NiagaraEditor/Private/Toolkits/NiagaraSystemToolkit.cpp
@@ -73,6 +73,14 @@ static FAutoConsoleVariableRef CVarSuppressNiagaraSystems(
 	ECVF_Default
 );
 
+static int32 GbShowNiagaraDeveloperWindows = 0;
+static FAutoConsoleVariableRef CVarShowNiagaraDeveloperWindows(
+	TEXT("fx.ShowNiagaraDeveloperWindows"),
+	GbShowNiagaraDeveloperWindows,
+	TEXT("If > 0 the niagara system and emitter editors will show additional developer windows.\nThese windows are for niagara tool development and debugging and editing the data\n directly in these windows can cause instability.\n"),
+	ECVF_Default
+);
+
 void FNiagaraSystemToolkit::RegisterTabSpawners(const TSharedRef<class FTabManager>& InTabManager)
 {
 	WorkspaceMenuCategory = InTabManager->AddLocalWorkspaceMenuCategory(LOCTEXT("WorkspaceMenu_NiagaraSystemEditor", "Niagara System"));
@@ -94,7 +102,8 @@ void FNiagaraSystemToolkit::RegisterTabSpawners(const TSharedRef<class FTabManag
 
 	InTabManager->RegisterTabSpawner(SystemScriptTabID, FOnSpawnTab::CreateSP(this, &FNiagaraSystemToolkit::SpawnTab_SystemScript))
 		.SetDisplayName(LOCTEXT("SystemScript", "System Script"))
-		.SetGroup(WorkspaceMenuCategory.ToSharedRef());
+		.SetGroup(WorkspaceMenuCategory.ToSharedRef())
+		.SetAutoGenerateMenuEntry(GbShowNiagaraDeveloperWindows != 0);
 
 	InTabManager->RegisterTabSpawner(SystemDetailsTabID, FOnSpawnTab::CreateSP(this, &FNiagaraSystemToolkit::SpawnTab_SystemDetails))
 		.SetDisplayName(LOCTEXT("SystemDetails", "System Details"))
@@ -110,7 +119,8 @@ void FNiagaraSystemToolkit::RegisterTabSpawners(const TSharedRef<class FTabManag
 
 	InTabManager->RegisterTabSpawner(SelectedEmitterGraphTabID, FOnSpawnTab::CreateSP(this, &FNiagaraSystemToolkit::SpawnTab_SelectedEmitterGraph))
 		.SetDisplayName(LOCTEXT("SelectedEmitterGraph", "Selected Emitter Graph"))
-		.SetGroup(WorkspaceMenuCategory.ToSharedRef());
+		.SetGroup(WorkspaceMenuCategory.ToSharedRef())
+		.SetAutoGenerateMenuEntry(GbShowNiagaraDeveloperWindows != 0);
 
 	InTabManager->RegisterTabSpawner(DebugSpreadsheetTabID, FOnSpawnTab::CreateSP(this, &FNiagaraSystemToolkit::SpawnTab_DebugSpreadsheet))
 		.SetDisplayName(LOCTEXT("DebugSpreadsheet", "Attribute Spreadsheet"))
@@ -347,6 +357,8 @@ void FNiagaraSystemToolkit::InitializeInternal(const EToolkitMode::Type Mode, co
 	SetupCommands();
 	ExtendToolbar();
 	RegenerateMenusAndToolbars();
+
+	bChangesDiscarded = false;
 }
 
 FName FNiagaraSystemToolkit::GetToolkitFName() const
@@ -1080,7 +1092,7 @@ bool FNiagaraSystemToolkit::OnRequestClose()
 	if (SystemToolkitMode == ESystemToolkitMode::Emitter)
 	{
 		TSharedPtr<FNiagaraEmitterViewModel> EmitterViewModel = SystemViewModel->GetEmitterHandleViewModels()[0]->GetEmitterViewModel();
-		if (EmitterViewModel->GetEmitter()->GetChangeId() != LastSyncedEmitterChangeId)
+		if (bChangesDiscarded == false && EmitterViewModel->GetEmitter()->GetChangeId() != LastSyncedEmitterChangeId)
 		{
 			// find out the user wants to do with this dirty NiagaraScript
 			EAppReturnType::Type YesNoCancelReply = FMessageDialog::Open(EAppMsgType::YesNoCancel,
@@ -1095,12 +1107,10 @@ bool FNiagaraSystemToolkit::OnRequestClose()
 				// update NiagaraScript and exit
 				UpdateOriginalEmitter();
 				break;
-
 			case EAppReturnType::No:
-				// exit
-				// do nothing.. bNiagaraScriptDirty = false;
+				// Set the changes discarded to avoid showing the dialog multiple times when request close is called multiple times on shut down.
+				bChangesDiscarded = true;
 				break;
-
 			case EAppReturnType::Cancel:
 				// don't exit
 				return false;
diff --git a/Engine/Plugins/FX/Niagara/Source/NiagaraEditor/Private/Toolkits/NiagaraSystemToolkit.h b/Engine/Plugins/FX/Niagara/Source/NiagaraEditor/Private/Toolkits/NiagaraSystemToolkit.h
--- a/Engine/Plugins/FX/Niagara/Source/NiagaraEditor/Private/Toolkits/NiagaraSystemToolkit.h
+++ b/Engine/Plugins/FX/Niagara/Source/NiagaraEditor/Private/Toolkits/NiagaraSystemToolkit.h
@@ -145,6 +145,8 @@ private:
 
 	TSharedPtr<class SNiagaraParameterMapView> ParameterMapView;
 
+	bool bChangesDiscarded;
+
 	static const FName ViewportTabID;
 	static const FName CurveEditorTabID;
 	static const FName SequencerTabID;
diff --git a/Engine/Plugins/FX/Niagara/Source/NiagaraEditor/Private/ViewModels/NiagaraEmitterViewModel.cpp b/Engine/Plugins/FX/Niagara/Source/NiagaraEditor/Private/ViewModels/NiagaraEmitterViewModel.cpp
--- a/Engine/Plugins/FX/Niagara/Source/NiagaraEditor/Private/ViewModels/NiagaraEmitterViewModel.cpp
+++ b/Engine/Plugins/FX/Niagara/Source/NiagaraEditor/Private/ViewModels/NiagaraEmitterViewModel.cpp
@@ -46,6 +46,7 @@ void FNiagaraEmitterViewModel::Cleanup()
 	if (Emitter.IsValid())
 	{
 		Emitter->OnEmitterVMCompiled().RemoveAll(this);
+		Emitter->OnPropertiesChanged().RemoveAll(this);
 	}
 
 	if (SharedScriptViewModel.IsValid())
diff --git a/Engine/Plugins/FX/Niagara/Source/NiagaraEditor/Private/ViewModels/NiagaraScriptInputCollectionViewModel.cpp b/Engine/Plugins/FX/Niagara/Source/NiagaraEditor/Private/ViewModels/NiagaraScriptInputCollectionViewModel.cpp
--- a/Engine/Plugins/FX/Niagara/Source/NiagaraEditor/Private/ViewModels/NiagaraScriptInputCollectionViewModel.cpp
+++ b/Engine/Plugins/FX/Niagara/Source/NiagaraEditor/Private/ViewModels/NiagaraScriptInputCollectionViewModel.cpp
@@ -374,6 +374,18 @@ void FNiagaraScriptInputCollectionViewModel::RefreshParameterViewModels()
 
 bool FNiagaraScriptInputCollectionViewModel::SupportsType(const FNiagaraTypeDefinition& Type) const
 {
+	if (Scripts.Num() == 1 && Scripts[0] != nullptr)
+	{
+		// We only support parameter map inputs for dynamic inputs and modules, with the ability to create data interfaces as needed for defaults.
+		if (Scripts[0]->GetUsage() == ENiagaraScriptUsage::DynamicInput || Scripts[0]->GetUsage() == ENiagaraScriptUsage::Module)
+		{
+			if (Type != FNiagaraTypeDefinition::GetParameterMapDef() && !Type.IsDataInterface())
+			{
+				return false;
+			}
+		}
+	}
+
 	return bCanHaveNumericParameters || Type != FNiagaraTypeDefinition::GetGenericNumericDef();
 }
 
diff --git a/Engine/Plugins/FX/Niagara/Source/NiagaraEditor/Private/ViewModels/NiagaraScriptOutputCollectionViewModel.cpp b/Engine/Plugins/FX/Niagara/Source/NiagaraEditor/Private/ViewModels/NiagaraScriptOutputCollectionViewModel.cpp
--- a/Engine/Plugins/FX/Niagara/Source/NiagaraEditor/Private/ViewModels/NiagaraScriptOutputCollectionViewModel.cpp
+++ b/Engine/Plugins/FX/Niagara/Source/NiagaraEditor/Private/ViewModels/NiagaraScriptOutputCollectionViewModel.cpp
@@ -272,6 +272,17 @@ void FNiagaraScriptOutputCollectionViewModel::RefreshParameterViewModels()
 
 bool FNiagaraScriptOutputCollectionViewModel::SupportsType(const FNiagaraTypeDefinition& Type) const
 {
+	if (Scripts.Num() == 1 && Scripts[0] != nullptr)
+	{
+		// We only support parameter map outputs for modules.
+		if (Scripts[0]->GetUsage() == ENiagaraScriptUsage::Module)
+		{
+			if (Type != FNiagaraTypeDefinition::GetParameterMapDef())
+			{
+				return false;
+			}
+		}
+	}
 	return Type.GetScriptStruct() != nullptr && (bCanHaveNumericParameters || Type != FNiagaraTypeDefinition::GetGenericNumericDef());
 }
 
@@ -307,6 +318,16 @@ void FNiagaraScriptOutputCollectionViewModel::OnParameterNameChanged(FName OldNa
 	{
 		ParameterVariable->SetName(FNiagaraUtilities::GetUniqueName(ParameterVariable->GetName(), CurrentNames));
 	}
+
+	// Make sure to edit the pin name so that when we regenerate the node the 
+	// links are preserved.
+	UEdGraphPin* FoundPin = OutputNode->FindPin(OldName, EEdGraphPinDirection::EGPD_Input);
+	if (FoundPin)
+	{
+		FoundPin->Modify();
+		FoundPin->PinName = NewName;
+	}
+
 	
 	OutputNode->NotifyOutputVariablesChanged();
 
diff --git a/Engine/Plugins/FX/Niagara/Source/NiagaraEditor/Private/ViewModels/NiagaraSystemScriptViewModel.cpp b/Engine/Plugins/FX/Niagara/Source/NiagaraEditor/Private/ViewModels/NiagaraSystemScriptViewModel.cpp
--- a/Engine/Plugins/FX/Niagara/Source/NiagaraEditor/Private/ViewModels/NiagaraSystemScriptViewModel.cpp
+++ b/Engine/Plugins/FX/Niagara/Source/NiagaraEditor/Private/ViewModels/NiagaraSystemScriptViewModel.cpp
@@ -168,14 +168,14 @@ void FNiagaraSystemScriptViewModel::RebuildEmitterNodes()
 		CurrentEmitterNode->Modify();
 		UEdGraphPin* InPin = CurrentEmitterNode->GetInputPin(0);
 		UEdGraphPin* OutPin = CurrentEmitterNode->GetOutputPin(0);
-		if (InPin != nullptr && OutPin != nullptr)
+		UEdGraphPin* InPinLinkedPin = InPin != nullptr && InPin->LinkedTo.Num() == 1 ? InPin->LinkedTo[0] : nullptr;
+		UEdGraphPin* OutPinLinkedPin = OutPin != nullptr && OutPin->LinkedTo.Num() == 1 ? OutPin->LinkedTo[0] : nullptr;
+		CurrentEmitterNode->DestroyNode();
+
+		if (InPinLinkedPin != nullptr &&& OutPinLinkedPin != nullptr)
 		{
-			if (0 != InPin->LinkedTo.Num() && 0 != OutPin->LinkedTo.Num())
-			{
-				InPin->LinkedTo[0]->MakeLinkTo(OutPin->LinkedTo[0]);
-			}
+			InPinLinkedPin->MakeLinkTo(OutPinLinkedPin);
 		}
-		CurrentEmitterNode->DestroyNode();
 	}
 
 	// Add output nodes if they don't exist.
diff --git a/Engine/Plugins/FX/Niagara/Source/NiagaraEditor/Private/ViewModels/NiagaraSystemViewModel.cpp b/Engine/Plugins/FX/Niagara/Source/NiagaraEditor/Private/ViewModels/NiagaraSystemViewModel.cpp
--- a/Engine/Plugins/FX/Niagara/Source/NiagaraEditor/Private/ViewModels/NiagaraSystemViewModel.cpp
+++ b/Engine/Plugins/FX/Niagara/Source/NiagaraEditor/Private/ViewModels/NiagaraSystemViewModel.cpp
@@ -508,45 +508,54 @@ void FNiagaraSystemViewModel::PostUndo(bool bSuccess)
 
 void FNiagaraSystemViewModel::Tick(float DeltaTime)
 {
-	bool bRecompile = false;
-
-	check(SystemScriptViewModel.IsValid());
-	if (SystemScriptViewModel->GetLatestCompileStatus() == ENiagaraScriptCompileStatus::NCS_Dirty)
+	if (bForceAutoCompileOnce || (GetDefault<UNiagaraEditorSettings>()->GetAutoCompile() && bCanAutoCompile))
 	{
-		//SystemScriptViewModel->CompileSystem();
-		//UE_LOG(LogNiagaraEditor, Log, TEXT("Compiling %s due to dirty scripts."), *System.GetName());
-		bRecompile |= true;
-	}
+		bool bRecompile = false;
 
-	for (TSharedRef<FNiagaraEmitterHandleViewModel> EmitterHandleViewModel : EmitterHandleViewModels)
-	{
-		if (EmitterHandleViewModel->GetEmitterViewModel()->GetLatestCompileStatus() == ENiagaraScriptCompileStatus::NCS_Dirty)
+		check(SystemScriptViewModel.IsValid());
+		if (SystemScriptViewModel->GetLatestCompileStatus() == ENiagaraScriptCompileStatus::NCS_Dirty)
 		{
+			//SystemScriptViewModel->CompileSystem();
+			//UE_LOG(LogNiagaraEditor, Log, TEXT("Compiling %s due to dirty scripts."), *System.GetName());
 			bRecompile |= true;
-			//EmitterHandleViewModel->GetEmitterViewModel()->CompileScripts();
-			//UE_LOG(LogNiagaraEditor, Log, TEXT("Compiling %s - %s due to dirty scripts."), *System.GetName(), *EmitterHandleViewModel->GetName().ToString());
 		}
-	}
 
-	if (System.HasOutstandingCompilationRequests() == false)
-	{
-		if (bCompilePendingCompletion)
+		for (TSharedRef<FNiagaraEmitterHandleViewModel> EmitterHandleViewModel : EmitterHandleViewModels)
 		{
-			bCompilePendingCompletion = false;
-			OnSystemCompiled().Broadcast();
+			if (EmitterHandleViewModel->GetEmitterViewModel()->GetLatestCompileStatus() == ENiagaraScriptCompileStatus::NCS_Dirty)
+			{
+				bRecompile |= true;
+				//EmitterHandleViewModel->GetEmitterViewModel()->CompileScripts();
+				//UE_LOG(LogNiagaraEditor, Log, TEXT("Compiling %s - %s due to dirty scripts."), *System.GetName(), *EmitterHandleViewModel->GetName().ToString());
+			}
 		}
 
-		if (bRecompile || bForceAutoCompileOnce)
+		if (System.HasOutstandingCompilationRequests() == false)
 		{
-			CompileSystem(false);
-			bForceAutoCompileOnce = false;
-		}
+			if (bCompilePendingCompletion)
+			{
+				bCompilePendingCompletion = false;
+				OnSystemCompiled().Broadcast();
+			}
 
-		if (bResetRequestPending)
-		{
-			ResetSystem();
+			if (bRecompile || bForceAutoCompileOnce)
+			{
+				CompileSystem(false);
+				bForceAutoCompileOnce = false;
+			}
+
+			if (bResetRequestPending)
+			{
+				ResetSystem();
+			}
 		}
 	}
+
+	if (EmitterIdsRequiringSequencerTrackUpdate.Num() > 0)
+	{
+		UpdateSequencerTracksForEmitters(EmitterIdsRequiringSequencerTrackUpdate);
+		EmitterIdsRequiringSequencerTrackUpdate.Empty();
+	}
 }
 
 void FNiagaraSystemViewModel::OnPreSave()
@@ -608,11 +617,6 @@ const TArray<FNiagaraStackModuleData>& FNiagaraSystemViewModel::GetStackModuleDa
 	return EmitterToCachedStackModuleData[EmitterHandleId];
 }
 
-bool FNiagaraSystemViewModel::IsTickable() const
-{
-	return bForceAutoCompileOnce || (GetDefault<UNiagaraEditorSettings>()->GetAutoCompile() && bCanAutoCompile);
-}
-
 TStatId FNiagaraSystemViewModel::GetStatId() const
 {
 	RETURN_QUICK_DECLARE_CYCLE_STAT(FNiagaraSystemViewModel, STATGROUP_Tickables);
@@ -713,6 +717,7 @@ void FNiagaraSystemViewModel::RefreshEmitterHandleViewModels()
 {
 	TArray<TSharedRef<FNiagaraEmitterHandleViewModel>> OldViewModels = EmitterHandleViewModels;
 	EmitterHandleViewModels.Empty();
+	EmitterToCachedStackModuleData.Empty();
 
 	// Map existing view models to the real instances that now exist. Reuse if we can. Create a new one if we cannot.
 	TArray<FGuid> ValidEmitterHandleIds;
@@ -838,6 +843,23 @@ void FNiagaraSystemViewModel::RefreshSequencerTracks()
 	}
 
 	Sequencer->NotifyMovieSceneDataChanged(EMovieSceneDataChangeType::MovieSceneStructureItemsChanged);
+
+	// Since we just rebuilt all of the sequencer tracks, these updates don't need to be done.
+	EmitterIdsRequiringSequencerTrackUpdate.Empty();
+}
+
+void FNiagaraSystemViewModel::UpdateSequencerTracksForEmitters(const TArray<FGuid>& EmitterIdsRequiringUpdate)
+{
+	TGuardValue<bool> UpdateGuard(bUpdatingSequencerFromEmitterDataChange, true);
+	for (UMovieSceneTrack* Track : NiagaraSequence->GetMovieScene()->GetMasterTracks())
+	{
+		UMovieSceneNiagaraEmitterTrack* EmitterTrack = CastChecked<UMovieSceneNiagaraEmitterTrack>(Track);
+		if (EmitterIdsRequiringUpdate.Contains(EmitterTrack->GetEmitterHandleViewModel()->GetId()))
+		{
+			EmitterTrack->UpdateTrackFromEmitterGraphChange(NiagaraSequence->GetMovieScene()->GetTickResolution());
+		}
+	}
+	Sequencer->NotifyMovieSceneDataChanged(EMovieSceneDataChangeType::MovieSceneStructureItemsChanged);
 }
 
 UMovieSceneNiagaraEmitterTrack* FNiagaraSystemViewModel::GetTrackForHandleViewModel(TSharedRef<FNiagaraEmitterHandleViewModel> EmitterHandleViewModel)
@@ -1185,16 +1207,7 @@ void FNiagaraSystemViewModel::EmitterScriptGraphChanged(const FEdGraphEditAction
 {
 	if (bUpdatingEmittersFromSequencerDataChange == false)
 	{
-		TGuardValue<bool> UpdateGuard(bUpdatingSequencerFromEmitterDataChange, true);
-		for (UMovieSceneTrack* Track : NiagaraSequence->GetMovieScene()->GetMasterTracks())
-		{
-			UMovieSceneNiagaraEmitterTrack* EmitterTrack = CastChecked<UMovieSceneNiagaraEmitterTrack>(Track);
-			if (EmitterTrack->GetEmitterHandleViewModel() == OwningEmitterHandleViewModel)
-			{
-				EmitterTrack->UpdateTrackFromEmitterGraphChange(NiagaraSequence->GetMovieScene()->GetTickResolution());
-			}
-		}
-		Sequencer->NotifyMovieSceneDataChanged(EMovieSceneDataChangeType::MovieSceneStructureItemsChanged);
+		EmitterIdsRequiringSequencerTrackUpdate.AddUnique(OwningEmitterHandleViewModel->GetId());
 	}
 	// Remove from cache on graph change 
 	EmitterToCachedStackModuleData.Remove(OwningEmitterHandleViewModel->GetId());
diff --git a/Engine/Plugins/FX/Niagara/Source/NiagaraEditor/Private/ViewModels/Stack/NiagaraStackEntry.cpp b/Engine/Plugins/FX/Niagara/Source/NiagaraEditor/Private/ViewModels/Stack/NiagaraStackEntry.cpp
--- a/Engine/Plugins/FX/Niagara/Source/NiagaraEditor/Private/ViewModels/Stack/NiagaraStackEntry.cpp
+++ b/Engine/Plugins/FX/Niagara/Source/NiagaraEditor/Private/ViewModels/Stack/NiagaraStackEntry.cpp
@@ -306,6 +306,11 @@ UNiagaraStackEntry::FOnRequestFullRefresh& UNiagaraStackEntry::OnRequestFullRefr
 	return RequestFullRefreshDelegate;
 }
 
+UNiagaraStackEntry::FOnRequestFullRefresh& UNiagaraStackEntry::OnRequestFullRefreshDeferred()
+{
+	return RequestFullRefreshDeferredDelegate;
+}
+
 int32 UNiagaraStackEntry::GetIndentLevel() const
 {
 	return IndentLevel;
@@ -410,6 +415,7 @@ void UNiagaraStackEntry::RefreshChildren()
 		Child->OnStructureChanged().RemoveAll(this);
 		Child->OnDataObjectModified().RemoveAll(this);
 		Child->OnRequestFullRefresh().RemoveAll(this);
+		Child->OnRequestFullRefreshDeferred().RemoveAll(this);
 		Child->SetOnRequestCanDrop(FOnRequestDrop());
 		Child->SetOnRequestDrop(FOnRequestDrop());
 	}
@@ -418,6 +424,7 @@ void UNiagaraStackEntry::RefreshChildren()
 		ErrorChild->OnStructureChanged().RemoveAll(this);
 		ErrorChild->OnDataObjectModified().RemoveAll(this);
 		ErrorChild->OnRequestFullRefresh().RemoveAll(this);
+		ErrorChild->OnRequestFullRefreshDeferred().RemoveAll(this);
 		ErrorChild->OnIssueModified().RemoveAll(this);
 	}
 
@@ -445,6 +452,7 @@ void UNiagaraStackEntry::RefreshChildren()
 		Child->OnStructureChanged().AddUObject(this, &UNiagaraStackEntry::ChildStructureChanged);
 		Child->OnDataObjectModified().AddUObject(this, &UNiagaraStackEntry::ChildDataObjectModified);
 		Child->OnRequestFullRefresh().AddUObject(this, &UNiagaraStackEntry::ChildRequestFullRefresh);
+		Child->OnRequestFullRefreshDeferred().AddUObject(this, &UNiagaraStackEntry::ChildRequestFullRefreshDeferred);
 		Child->SetOnRequestCanDrop(FOnRequestDrop::CreateUObject(this, &UNiagaraStackEntry::ChildRequestCanDrop));
 		Child->SetOnRequestDrop(FOnRequestDrop::CreateUObject(this, &UNiagaraStackEntry::ChildRequestDrop));
 	}
@@ -462,6 +470,7 @@ void UNiagaraStackEntry::RefreshChildren()
 		ErrorChild->OnStructureChanged().AddUObject(this, &UNiagaraStackEntry::ChildStructureChanged);
 		ErrorChild->OnDataObjectModified().AddUObject(this, &UNiagaraStackEntry::ChildDataObjectModified);
 		ErrorChild->OnRequestFullRefresh().AddUObject(this, &UNiagaraStackEntry::ChildRequestFullRefresh);
+		ErrorChild->OnRequestFullRefreshDeferred().AddUObject(this, &UNiagaraStackEntry::ChildRequestFullRefreshDeferred);
 		ErrorChild->OnIssueModified().AddUObject(this, &UNiagaraStackEntry::IssueModified);
 	}
 
@@ -472,7 +481,7 @@ void UNiagaraStackEntry::RefreshChildren()
 
 void UNiagaraStackEntry::RefreshStackErrorChildren()
 {
-    // keep the error entries that are already built
+	// keep the error entries that are already built
 	TArray<UNiagaraStackErrorItem*> NewErrorChildren;
 	for (FStackIssue Issue : StackIssues)
 	{
@@ -551,6 +560,11 @@ void UNiagaraStackEntry::ChildRequestFullRefresh()
 	RequestFullRefreshDelegate.Broadcast();
 }
 
+void UNiagaraStackEntry::ChildRequestFullRefreshDeferred()
+{
+	RequestFullRefreshDeferredDelegate.Broadcast();
+}
+
 TOptional<UNiagaraStackEntry::FDropResult> UNiagaraStackEntry::ChildRequestCanDrop(const UNiagaraStackEntry& TargetChild, const TArray<UNiagaraStackEntry*>& DraggedEntries)
 {
 	TOptional<FDropResult> Result = ChildRequestCanDropInternal(TargetChild, DraggedEntries);
diff --git a/Engine/Plugins/FX/Niagara/Source/NiagaraEditor/Private/ViewModels/Stack/NiagaraStackFunctionInput.cpp b/Engine/Plugins/FX/Niagara/Source/NiagaraEditor/Private/ViewModels/Stack/NiagaraStackFunctionInput.cpp
--- a/Engine/Plugins/FX/Niagara/Source/NiagaraEditor/Private/ViewModels/Stack/NiagaraStackFunctionInput.cpp
+++ b/Engine/Plugins/FX/Niagara/Source/NiagaraEditor/Private/ViewModels/Stack/NiagaraStackFunctionInput.cpp
@@ -207,13 +207,14 @@ void UNiagaraStackFunctionInput::FinalizeInternal()
 	{
 		OwningFunctionCallNode->GetGraph()->RemoveOnGraphChangedHandler(GraphChangedHandle);
 		OwningFunctionCallNode->GetNiagaraGraph()->RemoveOnGraphNeedsRecompileHandler(OnRecompileHandle);
+	}
 
-		if (SourceScript.IsValid())
-		{
-			SourceScript->RapidIterationParameters.RemoveOnChangedHandler(RapidIterationParametersChangedHandle);
-			SourceScript->GetSource()->OnChanged().RemoveAll(this);
-		}
+	if (SourceScript.IsValid())
+	{
+		SourceScript->RapidIterationParameters.RemoveOnChangedHandler(RapidIterationParametersChangedHandle);
+		SourceScript->GetSource()->OnChanged().RemoveAll(this);
 	}
+
 	Super::FinalizeInternal();
 }
 
@@ -925,11 +926,8 @@ void UNiagaraStackFunctionInput::SetLocalValue(TSharedRef<FStructOnScope> InLoca
 
 		for (TWeakObjectPtr<UNiagaraScript> Script : AffectedScripts)
 		{
-			if(Script->RapidIterationParameters.AddParameter(RapidIterationParameter))
-			{
-				UE_LOG(LogNiagaraEditor, Log, TEXT("Adding Parameter %s to Script %s"), *RapidIterationParameter.GetName().ToString(), *Script->GetFullName());
-			}
-			Script->RapidIterationParameters.SetParameterData(InLocalValue->GetStructMemory(), RapidIterationParameter);
+			bool bAddParameterIfMissing = true;
+			Script->RapidIterationParameters.SetParameterData(InLocalValue->GetStructMemory(), RapidIterationParameter, bAddParameterIfMissing);
 		}
 	}
 	else 
@@ -979,35 +977,38 @@ bool UNiagaraStackFunctionInput::CanReset() const
 		else
 		{
 			UEdGraphPin* DefaultPin = GetDefaultPin();
-			if(DefaultPin->LinkedTo.Num() == 0)
-			{
-				if (GetOverridePin() != nullptr)
-				{
-					bNewCanReset = true;
-				}
-				else if (IsRapidIterationCandidate())
-				{
-					FNiagaraVariable DefaultVar = GetDefaultVariableForRapidIterationParameter();
-					bool bHasValidLocalValue = InputValues.LocalStruct.IsValid();
-					bool bHasValidDefaultValue = DefaultVar.IsValid();
-					bNewCanReset = bHasValidLocalValue && bHasValidDefaultValue && FNiagaraEditorUtilities::DataMatches(DefaultVar, *InputValues.LocalStruct.Get()) == false;
-				}
-				else
-				{
-					bNewCanReset = false;
-				}
-			}
-			else
-			{
-				if (FNiagaraStackGraphUtilities::IsValidDefaultDynamicInput(*SourceScript, *DefaultPin))
+			if (ensure(DefaultPin != nullptr))
+			{			
+				if(DefaultPin->LinkedTo.Num() == 0)
 				{
-					UEdGraphPin* OverridePin = GetOverridePin();
-					bNewCanReset = OverridePin == nullptr || FNiagaraStackGraphUtilities::DoesDynamicInputMatchDefault(GetEmitterViewModel()->GetEmitter()->GetUniqueEmitterName(), *SourceScript,
-						*OwningFunctionCallNode, *OverridePin, InputParameterHandle.GetName(), *DefaultPin) == false;
+					if (GetOverridePin() != nullptr)
+					{
+						bNewCanReset = true;
+					}
+					else if (IsRapidIterationCandidate())
+					{
+						FNiagaraVariable DefaultVar = GetDefaultVariableForRapidIterationParameter();
+						bool bHasValidLocalValue = InputValues.LocalStruct.IsValid();
+						bool bHasValidDefaultValue = DefaultVar.IsValid();
+						bNewCanReset = bHasValidLocalValue && bHasValidDefaultValue && FNiagaraEditorUtilities::DataMatches(DefaultVar, *InputValues.LocalStruct.Get()) == false;
+					}
+					else
+					{
+						bNewCanReset = false;
+					}
 				}
 				else
 				{
-					bNewCanReset = GetOverridePin() != nullptr;
+					if (FNiagaraStackGraphUtilities::IsValidDefaultDynamicInput(*SourceScript, *DefaultPin))
+					{
+						UEdGraphPin* OverridePin = GetOverridePin();
+						bNewCanReset = OverridePin == nullptr || FNiagaraStackGraphUtilities::DoesDynamicInputMatchDefault(GetEmitterViewModel()->GetEmitter()->GetUniqueEmitterName(), *SourceScript,
+							*OwningFunctionCallNode, *OverridePin, InputParameterHandle.GetName(), *DefaultPin) == false;
+					}
+					else
+					{
+						bNewCanReset = GetOverridePin() != nullptr;
+					}
 				}
 			}
 		}
@@ -1038,11 +1039,8 @@ bool UNiagaraStackFunctionInput::UpdateRapidIterationParametersForAffectedScript
 
 	for (TWeakObjectPtr<UNiagaraScript> Script : AffectedScripts)
 	{
-		if (Script->RapidIterationParameters.AddParameter(RapidIterationParameter))
-		{
-			UE_LOG(LogNiagaraEditor, Log, TEXT("Adding Parameter %s to Script %s"), *RapidIterationParameter.GetName().ToString(), *Script->GetFullName());
-		}
-		Script->RapidIterationParameters.SetParameterData(Data, RapidIterationParameter);
+		bool bAddParameterIfMissing = true;
+		Script->RapidIterationParameters.SetParameterData(Data, RapidIterationParameter, bAddParameterIfMissing);
 	}
 	GetSystemViewModel()->ResetSystem();
 	return true;
@@ -1244,9 +1242,9 @@ void UNiagaraStackFunctionInput::ResetToBase()
 						AffectedScript->Modify();
 						for (FNiagaraVariable& OwningScriptRapidIterationParameter : OwningScriptRapidIterationParameters)
 						{
-							AffectedScript->RapidIterationParameters.AddParameter(OwningScriptRapidIterationParameter);
+							bool bAddParameterIfMissing = true;
 							AffectedScript->RapidIterationParameters.SetParameterData(
-								OwningScript->RapidIterationParameters.GetParameterData(OwningScriptRapidIterationParameter), OwningScriptRapidIterationParameter);
+								OwningScript->RapidIterationParameters.GetParameterData(OwningScriptRapidIterationParameter), OwningScriptRapidIterationParameter, bAddParameterIfMissing);
 						}
 					}
 				}
@@ -1266,7 +1264,8 @@ FNiagaraVariable UNiagaraStackFunctionInput::CreateRapidIterationVariable(const
 bool UNiagaraStackFunctionInput::CanRenameInput() const
 {
 	// Only module level assignment node inputs can be renamed.
-	return OwningAssignmentNode.IsValid() && InputParameterHandlePath.Num() == 1;
+	return OwningAssignmentNode.IsValid() && InputParameterHandlePath.Num() == 1 &&
+		OwningAssignmentNode->FindAssignmentTarget(InputParameterHandle.GetName()) != INDEX_NONE;
 }
 
 bool UNiagaraStackFunctionInput::GetIsRenamePending() const
diff --git a/Engine/Plugins/FX/Niagara/Source/NiagaraEditor/Private/ViewModels/Stack/NiagaraStackGraphUtilities.cpp b/Engine/Plugins/FX/Niagara/Source/NiagaraEditor/Private/ViewModels/Stack/NiagaraStackGraphUtilities.cpp
--- a/Engine/Plugins/FX/Niagara/Source/NiagaraEditor/Private/ViewModels/Stack/NiagaraStackGraphUtilities.cpp
+++ b/Engine/Plugins/FX/Niagara/Source/NiagaraEditor/Private/ViewModels/Stack/NiagaraStackGraphUtilities.cpp
@@ -820,13 +820,32 @@ void FNiagaraStackGraphUtilities::SetCustomExpressionForFunctionInput(UEdGraphPi
 	}
 }
 
-bool FNiagaraStackGraphUtilities::RemoveModuleFromStack(UNiagaraNodeFunctionCall& ModuleNode)
+bool FNiagaraStackGraphUtilities::RemoveModuleFromStack(UNiagaraSystem& OwningSystem, FGuid OwningEmitterId, UNiagaraNodeFunctionCall& ModuleNode)
 {
 	TArray<TWeakObjectPtr<UNiagaraNodeInput>> RemovedNodes;
-	return RemoveModuleFromStack(ModuleNode, RemovedNodes);
+	return RemoveModuleFromStack(OwningSystem, OwningEmitterId, ModuleNode, RemovedNodes);
 }
 
-bool FNiagaraStackGraphUtilities::RemoveModuleFromStack(UNiagaraNodeFunctionCall& ModuleNode, TArray<TWeakObjectPtr<UNiagaraNodeInput>>& OutRemovedNodes)
+bool FNiagaraStackGraphUtilities::RemoveModuleFromStack(UNiagaraSystem& OwningSystem, FGuid OwningEmitterId, UNiagaraNodeFunctionCall& ModuleNode, TArray<TWeakObjectPtr<UNiagaraNodeInput>>& OutRemovedInputNodes)
+{
+	// Find the owning script so it can be modified as part of the transaction so that rapid iteration parameters values are retained upon undo.
+	UNiagaraNodeOutput* OutputNode = GetEmitterOutputNodeForStackNode(ModuleNode);
+	checkf(OutputNode != nullptr, TEXT("Invalid Stack - Output node could not be found for module"));
+
+	UNiagaraScript* OwningScript = FNiagaraEditorUtilities::GetScriptFromSystem(
+		OwningSystem, OwningEmitterId, OutputNode->GetUsage(), OutputNode->GetUsageId());
+	checkf(OwningScript != nullptr, TEXT("Invalid Stack - Owning script could not be found for module"));
+
+	return RemoveModuleFromStack(*OwningScript, ModuleNode, OutRemovedInputNodes);
+}
+
+bool FNiagaraStackGraphUtilities::RemoveModuleFromStack(UNiagaraScript& OwningScript, UNiagaraNodeFunctionCall& ModuleNode)
+{
+	TArray<TWeakObjectPtr<UNiagaraNodeInput>> RemovedNodes;
+	return RemoveModuleFromStack(OwningScript, ModuleNode, RemovedNodes);
+}
+
+bool FNiagaraStackGraphUtilities::RemoveModuleFromStack(UNiagaraScript& OwningScript, UNiagaraNodeFunctionCall& ModuleNode, TArray<TWeakObjectPtr<UNiagaraNodeInput>>& OutRemovedInputNodes)
 {
 	TArray<FNiagaraStackGraphUtilities::FStackNodeGroup> StackNodeGroups;
 	FNiagaraStackGraphUtilities::GetStackNodeGroups(ModuleNode, StackNodeGroups);
@@ -837,6 +856,8 @@ bool FNiagaraStackGraphUtilities::RemoveModuleFromStack(UNiagaraNodeFunctionCall
 		return false;
 	}
 
+	OwningScript.Modify();
+
 	// Disconnect the group from the stack first to make collecting the nodes to remove easier.
 	FNiagaraStackGraphUtilities::DisconnectStackNodeGroup(StackNodeGroups[ModuleStackIndex], StackNodeGroups[ModuleStackIndex - 1], StackNodeGroups[ModuleStackIndex + 1]);
 
@@ -849,7 +870,7 @@ bool FNiagaraStackGraphUtilities::RemoveModuleFromStack(UNiagaraNodeFunctionCall
 	{
 		UNiagaraNode* NodeToRemove = NodesToCheck[0];
 		NodesToCheck.RemoveAt(0);
-		NodesToRemove.Add(NodeToRemove);
+		NodesToRemove.AddUnique(NodeToRemove);
 
 		TArray<UEdGraphPin*> InputPins;
 		NodeToRemove->GetInputPins(InputPins);
@@ -870,11 +891,12 @@ bool FNiagaraStackGraphUtilities::RemoveModuleFromStack(UNiagaraNodeFunctionCall
 	UNiagaraGraph* Graph = ModuleNode.GetNiagaraGraph();
 	for (UNiagaraNode* NodeToRemove : NodesToRemove)
 	{
+		NodeToRemove->Modify();
 		Graph->RemoveNode(NodeToRemove);
 		UNiagaraNodeInput* InputNode = Cast<UNiagaraNodeInput>(NodeToRemove);
 		if (InputNode != nullptr)
 		{
-			OutRemovedNodes.Add(InputNode);
+			OutRemovedInputNodes.Add(InputNode);
 		}
 	}
 
@@ -1000,6 +1022,7 @@ void FNiagaraStackGraphUtilities::SetModuleIsEnabled(UNiagaraNodeFunctionCall& F
 	GetAllNodesForModule(FunctionCallNode, ModuleNodes);
 	for (UNiagaraNode* ModuleNode : ModuleNodes)
 	{
+		ModuleNode->Modify();
 		ModuleNode->SetEnabledState(bIsEnabled ? ENodeEnabledState::Enabled : ENodeEnabledState::Disabled, true);
 		ModuleNode->MarkNodeRequiresSynchronization(__FUNCTION__, false);
 	}
@@ -1252,17 +1275,11 @@ void FNiagaraStackGraphUtilities::GetScriptAssetsByDependencyProvided(ENiagaraSc
 	{
 		auto TagName = GET_MEMBER_NAME_CHECKED(UNiagaraScript, ProvidedDependencies);
 
-		TArray<UObject::FAssetRegistryTag> Tags;
-		ScriptAsset.GetAsset()->GetAssetRegistryTags(Tags);
-		UObject::FAssetRegistryTag* FoundTag = Tags.FindByPredicate([=](const UObject::FAssetRegistryTag& InTag)
-		{
-			return InTag.Name == TagName;
-		});
-		if (FoundTag != nullptr)
+		FString ProvidedDependenciesString;
+		if(ScriptAsset.GetTagValue(GET_MEMBER_NAME_CHECKED(UNiagaraScript, ProvidedDependencies), ProvidedDependenciesString) && ProvidedDependenciesString.IsEmpty() == false)
 		{
-			// parse it and search through it%
 			TArray<FString> DependencyStrings;
-			FoundTag->Value.ParseIntoArray(DependencyStrings, TEXT(","));
+			ProvidedDependenciesString.ParseIntoArray(DependencyStrings, TEXT(","));
 			for (FString DependencyString: DependencyStrings)
 			{
 				if (FName(*DependencyString) == DependencyName)
diff --git a/Engine/Plugins/FX/Niagara/Source/NiagaraEditor/Private/ViewModels/Stack/NiagaraStackModuleItem.cpp b/Engine/Plugins/FX/Niagara/Source/NiagaraEditor/Private/ViewModels/Stack/NiagaraStackModuleItem.cpp
--- a/Engine/Plugins/FX/Niagara/Source/NiagaraEditor/Private/ViewModels/Stack/NiagaraStackModuleItem.cpp
+++ b/Engine/Plugins/FX/Niagara/Source/NiagaraEditor/Private/ViewModels/Stack/NiagaraStackModuleItem.cpp
@@ -32,6 +32,10 @@
 
 #include "ScopedTransaction.h"
 
+// TODO: Remove these
+#include "Widgets/Notifications/SNotificationList.h"
+#include "Framework/Notifications/NotificationManager.h"
+
 #define LOCTEXT_NAMESPACE "NiagaraStackModuleItem"
 
 TArray<ENiagaraScriptUsage> UsagePriority = { // Ordered such as the highest priority has the largest index
@@ -51,6 +55,13 @@ UNiagaraNodeOutput* GetOutputNodeForModuleDependency(ENiagaraScriptUsage Dependa
 		TArray<ENiagaraScriptUsage> SupportedUsages = UNiagaraScript::GetSupportedUsageContextsForBitmask(DependencyScript->ModuleUsageBitmask);
 		ENiagaraScriptUsage ScriptUsage = SupportedUsages[0];
 		int32 ClosestDistance = MAX_int32;
+
+		int32 DependantIndex = UsagePriority.IndexOfByPredicate(
+			[&](const ENiagaraScriptUsage CurrentUsage)
+		{
+			return UNiagaraScript::IsEquivalentUsage(DependantUsage, CurrentUsage);
+		});
+
 		for (ENiagaraScriptUsage PossibleUsage : SupportedUsages)
 		{
 			int32 PossibleIndex = UsagePriority.IndexOfByPredicate(
@@ -58,11 +69,12 @@ UNiagaraNodeOutput* GetOutputNodeForModuleDependency(ENiagaraScriptUsage Dependa
 				{
 					return UNiagaraScript::IsEquivalentUsage(PossibleUsage, CurrentUsage);
 				});
-			int32 DependantIndex = UsagePriority.IndexOfByPredicate(
-				[&](const ENiagaraScriptUsage CurrentUsage)
-				{
-					return UNiagaraScript::IsEquivalentUsage(DependantUsage, CurrentUsage);
-				});
+
+			if (PossibleIndex == INDEX_NONE)
+			{
+				// This usage isn't in the execution flow so check the next one.
+				continue;
+			}
 
 			int32 Distance = PossibleIndex - DependantIndex;
 			bool bCorrectOrder = (Dependency.Type == ENiagaraModuleDependencyType::PreDependency && Distance >= 0) || (Dependency.Type == ENiagaraModuleDependencyType::PostDependency && Distance <= 0);
@@ -96,6 +108,7 @@ UNiagaraNodeOutput* GetOutputNodeForModuleDependency(ENiagaraScriptUsage Dependa
 				}
 			}
 		}
+
 		if (OutputScript != nullptr)
 		{
 			TargetOutputNode = FNiagaraEditorUtilities::GetScriptOutputNode(*OutputScript);
@@ -128,6 +141,7 @@ void UNiagaraStackModuleItem::Initialize(FRequiredEntryData InRequiredEntryData,
 	Super::Initialize(InRequiredEntryData, ModuleStackEditorDataKey);
 	GroupAddUtilities = InGroupAddUtilities;
 	FunctionCallNode = &InFunctionCallNode;
+	OutputNode = FNiagaraStackGraphUtilities::GetEmitterOutputNodeForStackNode(*FunctionCallNode);
 	AddChildFilter(FOnFilterChild::CreateUObject(this, &UNiagaraStackModuleItem::FilterOutputCollection));
 
 	// Update bCanMoveAndDelete
@@ -142,7 +156,6 @@ void UNiagaraStackModuleItem::Initialize(FRequiredEntryData InRequiredEntryData,
 		TSharedRef<FNiagaraScriptMergeManager> MergeManager = FNiagaraScriptMergeManager::Get();
 
 		const UNiagaraEmitter* BaseEmitter = FNiagaraStackGraphUtilities::GetBaseEmitter(*GetEmitterViewModel()->GetEmitter(), GetSystemViewModel()->GetSystem());
-		UNiagaraNodeOutput* OutputNode = FNiagaraStackGraphUtilities::GetEmitterOutputNodeForStackNode(*FunctionCallNode);
 
 		bool bIsMergeable = MergeManager->IsMergeableScriptUsage(OutputNode->GetUsage());
 		bool bHasBaseModule = bIsMergeable && BaseEmitter != nullptr && MergeManager->HasBaseModule(*BaseEmitter, OutputNode->GetUsage(), OutputNode->GetUsageId(), FunctionCallNode->NodeGuid);
@@ -276,11 +289,37 @@ void UNiagaraStackModuleItem::RefreshIssues(TArray<FStackIssue>& NewIssues)
 
 			NewIssues.Add(InconsistentEnabledError);
 		}
+
+		UNiagaraNodeAssignment* AssignmentFunctionCall = Cast<UNiagaraNodeAssignment>(FunctionCallNode);
+		if (AssignmentFunctionCall != nullptr)
+		{
+			TSet<FNiagaraVariable> FoundAssignmentTargets;
+			for (const FNiagaraVariable& AssignmentTarget : AssignmentFunctionCall->GetAssignmentTargets())
+			{
+				if (FoundAssignmentTargets.Contains(AssignmentTarget))
+				{
+					FText FixDescription = LOCTEXT("RemoveDuplicate", "Remove Duplicate");
+					FStackIssueFix RemoveDuplicateFix(FixDescription, FStackIssueFixDelegate::CreateLambda([AssignmentFunctionCall, AssignmentTarget]()
+					{
+						AssignmentFunctionCall->RemoveParameter(AssignmentTarget);
+					}));
+					FStackIssue DuplicateAssignmentTargetError(
+						EStackIssueSeverity::Error,
+						LOCTEXT("DuplicateAssignmentTargetErrorSummary", "Duplicate variables detected."),
+						LOCTEXT("InconsistentEnabledError", "This 'Set Variables' module is attempting to set the same variable more than once, which is unsupported."),
+						GetStackEditorDataKey(),
+						false,
+						RemoveDuplicateFix);
+
+					NewIssues.Add(DuplicateAssignmentTargetError);
+				}
+				FoundAssignmentTargets.Add(AssignmentTarget);
+			}
+		}
 	}
 	// Generate dependency errors with their fixes
 	TArray<UNiagaraNodeFunctionCall*> FoundCalls;
 	TArray<FNiagaraModuleDependency> DependenciesNeeded;
-	UNiagaraNodeOutput* OutputNode = GetOutputNode();
 
 	TArray<FNiagaraStackModuleData> SystemModuleData = GetSystemViewModel()->GetStackModuleDataForEmitter(GetEmitterViewModel());
 	int32 ModuleIndex = INDEX_NONE;
@@ -374,7 +413,7 @@ void UNiagaraStackModuleItem::RefreshIssues(TArray<FStackIssue>& NewIssues)
 					FText AndEnableModule = bNeedsEnable ? FText::Format(LOCTEXT("AndEnableDependency", "And enable dependency module {0}"), FText::FromString(DisorderedNode.ModuleNode->GetFunctionName())) : FText();
 					UNiagaraStackEntry::FStackIssueFix Fix(
 						FText::Format(LOCTEXT("ReorderDependency", "Reposition this module in the correct order related to {0} {1}"), FText::FromString(DisorderedNode.ModuleNode->GetFunctionName()), AndEnableModule),
-						FStackIssueFixDelegate::CreateLambda([this, bNeedsEnable, DisorderedNode, SystemModuleData, Dependency, ModuleIndex, OutputNode]()
+						FStackIssueFixDelegate::CreateLambda([this, bNeedsEnable, DisorderedNode, SystemModuleData, Dependency, ModuleIndex]()
 					{
 						FScopedTransaction ScopedTransaction(LOCTEXT("ReorderDependencyModule", "Reorder dependency module"));
 
@@ -438,13 +477,23 @@ void UNiagaraStackModuleItem::RefreshIssues(TArray<FStackIssue>& NewIssues)
 									}
 								}
 							}
+
 							if (TargetOutputNode == nullptr)
 							{
-								checkf(ModuleIndex != INDEX_NONE, TEXT("Module data wasn't found in system for current module!"));
-								// I am looking for a script in the same emitter as this module, so get it from there
-								UNiagaraScript& TargetScript = *FNiagaraEditorUtilities::GetScriptFromSystem(GetSystemViewModel()->GetSystem(), SystemModuleData[ModuleIndex].EmitterHandleId, DependencyScript->GetUsage(), DependencyScript->GetUsageId());
-								TargetOutputNode = FNiagaraEditorUtilities::GetScriptOutputNode(TargetScript);
+								// If no output node was found than the dependency can't be resolved and it most likely misconfigured in data.
+								// TODO: Don't show this toast here, change the fix delegate to return a fix result with whether or not the fix succeeded and any error message for the user.
+								FNotificationInfo Error(LOCTEXT("FixFailedToast", "Failed to fix the dependency since\nwe could not find a compatible place to insert the module.\nPlease check the configuration of the dependency.\nSee the log for more details."));
+								Error.ExpireDuration = 5.0f;
+								Error.bFireAndForget = true;
+								Error.Image = FCoreStyle::Get().GetBrush(TEXT("MessageLog.Error"));
+								FSlateNotificationManager::Get().AddNotification(Error);
+								FString ModuleAssetFullName;
+								ModuleAsset.GetFullName(ModuleAssetFullName);
+								UE_LOG(LogNiagaraEditor, Error, TEXT("Dependency fix failed, could not find a compatible place to insert the module.\nModule requiring dependency: %s\nModule providing dependency: %s\nDependency name: %s\nDependency type: %s"),
+									*FunctionCallNode->FunctionScript->GetFullName(), *ModuleAssetFullName, *Dependency.Id.ToString(), Dependency.Type == ENiagaraModuleDependencyType::PreDependency ? TEXT("Pre-dependency") : TEXT("Post-dependency"));
+								return;
 							}
+
 							TArray<FNiagaraStackModuleData> ScriptModuleData = SystemModuleData.FilterByPredicate([&](FNiagaraStackModuleData CurrentData) {return CurrentData.Usage == DependencyScript->GetUsage(); });
 							int32 PreIndex = INDEX_NONE; // index of last pre dependency
 							int32 PostIndex = INDEX_NONE; // index of fist post dependency, the module will have to be placed between these indexes
@@ -565,8 +614,12 @@ void UNiagaraStackModuleItem::Delete()
 	checkf(CanMoveAndDelete(), TEXT("This module can't be deleted"));
 
 	FScopedTransaction ScopedTransaction(LOCTEXT("RemoveAModuleFromTheStack", "Remove a module from the stack"));
+
+	const FNiagaraEmitterHandle* EmitterHandle = FNiagaraEditorUtilities::GetEmitterHandleForEmitter(GetSystemViewModel()->GetSystem(), *GetEmitterViewModel()->GetEmitter());
+	checkf(EmitterHandle != nullptr, TEXT("Invalid Stack - Emitter handle could not be found for module"));
+
 	TArray<TWeakObjectPtr<UNiagaraNodeInput>> RemovedNodes;
-	bool bRemoved = FNiagaraStackGraphUtilities::RemoveModuleFromStack(*FunctionCallNode, RemovedNodes);
+	bool bRemoved = FNiagaraStackGraphUtilities::RemoveModuleFromStack(GetSystemViewModel()->GetSystem(), EmitterHandle->GetId(), *FunctionCallNode, RemovedNodes);
 	if (bRemoved)
 	{
 		UNiagaraGraph* Graph = FunctionCallNode->GetNiagaraGraph();
@@ -604,7 +657,7 @@ int32 UNiagaraStackModuleItem::GetModuleIndex()
 
 UNiagaraNodeOutput* UNiagaraStackModuleItem::GetOutputNode() const
 {
-	return FNiagaraStackGraphUtilities::GetEmitterOutputNodeForStackNode(*FunctionCallNode);
+	return OutputNode;
 }
 
 void UNiagaraStackModuleItem::NotifyModuleMoved()
@@ -612,14 +665,38 @@ void UNiagaraStackModuleItem::NotifyModuleMoved()
 	ModifiedGroupItemsDelegate.ExecuteIfBound();
 }
 
-bool UNiagaraStackModuleItem::CanAddInput() const
+bool ParameterIsCompatibleWithScriptUsage(FNiagaraVariable Parameter, ENiagaraScriptUsage Usage)
+{
+	FNiagaraParameterHandle ParameterHandle(Parameter.GetName());
+	switch (Usage)
+	{
+	case ENiagaraScriptUsage::SystemSpawnScript:
+	case ENiagaraScriptUsage::SystemUpdateScript:
+		return ParameterHandle.IsSystemHandle();
+	case ENiagaraScriptUsage::EmitterSpawnScript:
+	case ENiagaraScriptUsage::EmitterUpdateScript:
+		return ParameterHandle.IsEmitterHandle();
+	case ENiagaraScriptUsage::ParticleSpawnScript:
+	case ENiagaraScriptUsage::ParticleSpawnScriptInterpolated:
+	case ENiagaraScriptUsage::ParticleUpdateScript:
+	case ENiagaraScriptUsage::ParticleEventScript:
+		return ParameterHandle.IsParticleAttributeHandle();
+	default:
+		return false;
+	}
+}
+
+bool UNiagaraStackModuleItem::CanAddInput(FNiagaraVariable InputParameter) const
 {
-	return FunctionCallNode->IsA<UNiagaraNodeAssignment>();
+	UNiagaraNodeAssignment* AssignmentModule = Cast<UNiagaraNodeAssignment>(FunctionCallNode);
+	return AssignmentModule != nullptr &&
+		AssignmentModule->GetAssignmentTargets().Contains(InputParameter) == false &&
+		ParameterIsCompatibleWithScriptUsage(InputParameter, OutputNode->GetUsage());
 }
 
 void UNiagaraStackModuleItem::AddInput(FNiagaraVariable InputParameter)
 {
-	if(ensureMsgf(CanAddInput(), TEXT("This module doesn't support adding inputs.")))
+	if(ensureMsgf(CanAddInput(InputParameter), TEXT("This module doesn't support adding this input.")))
 	{
 		UNiagaraNodeAssignment* AssignmentNode = CastChecked<UNiagaraNodeAssignment>(FunctionCallNode);
 		AssignmentNode->AddParameter(InputParameter, FNiagaraConstants::GetAttributeDefaultValue(InputParameter));
diff --git a/Engine/Plugins/FX/Niagara/Source/NiagaraEditor/Private/ViewModels/Stack/NiagaraStackParameterStoreEntry.cpp b/Engine/Plugins/FX/Niagara/Source/NiagaraEditor/Private/ViewModels/Stack/NiagaraStackParameterStoreEntry.cpp
--- a/Engine/Plugins/FX/Niagara/Source/NiagaraEditor/Private/ViewModels/Stack/NiagaraStackParameterStoreEntry.cpp
+++ b/Engine/Plugins/FX/Niagara/Source/NiagaraEditor/Private/ViewModels/Stack/NiagaraStackParameterStoreEntry.cpp
@@ -253,8 +253,8 @@ void UNiagaraStackParameterStoreEntry::RenameInput(FString NewName)
 			}
 		}
 
-		GEditor->BeginTransaction(LOCTEXT("RenameParameter", "Rename Parameter"));
-		ParameterStore->GetOwner()->Modify();
+		FScopedTransaction ScopedTransaction(LOCTEXT("RenameUserParameter", "Rename user parameter"));
+		Owner->Modify();
 		// remove old one (a bit overkill but it beats duplicating code)
 		RemovePins(OwningPins);
 		// TODO Would it be better to actually keep variable name being ActualName and ParameterName being ParameterHandle.GetParameterHandleString(), and rewrite the way the Entry is built?
@@ -268,17 +268,22 @@ void UNiagaraStackParameterStoreEntry::RenameInput(FString NewName)
 			// generate current link
 			FNiagaraStackGraphUtilities::SetLinkedValueHandleForFunctionInput(*LinkedPin, ParameterHandle);
 		}
-		GEditor->EndTransaction();
+
+		ParameterName = VariableName;
+		DisplayName = FText::FromName(ParameterName);
 	}
 }
 
 void UNiagaraStackParameterStoreEntry::Delete()
 {
+	FScopedTransaction ScopedTransaction(LOCTEXT("RemoveUserParameter", "Remove user parameter"));
+
 	// for delete, do a parameter map traversal to deduce all the usages and then  remove them 
 	TArray<UEdGraphPin*> OwningPins = GetOwningPins();
-	NotifyBeginValueChange();
 	RemovePins(OwningPins);
+
 	//remove from store
+	Owner->Modify();
 	ParameterStore->RemoveParameter(FNiagaraVariable(InputType, ParameterName));
 	if (InputType.GetClass() != nullptr)
 	{
@@ -288,7 +293,8 @@ void UNiagaraStackParameterStoreEntry::Delete()
 			GetSystemViewModel()->NotifyDataObjectChanged(DataInterface);
 		}
 	}
-	NotifyEndValueChange();
+
+	ParameterDeletedDelegate.Broadcast();
 }
 
 void UNiagaraStackParameterStoreEntry::RemovePins(TArray<UEdGraphPin*> OwningPins /*, bool bSetPreviousValue*/)
@@ -334,6 +340,11 @@ UNiagaraStackParameterStoreEntry::FOnValueChanged& UNiagaraStackParameterStoreEn
 	return ValueChangedDelegate;
 }
 
+UNiagaraStackParameterStoreEntry::FOnParameterDeleted& UNiagaraStackParameterStoreEntry::OnParameterDeleted()
+{
+	return ParameterDeletedDelegate;
+}
+
 TSharedPtr<FNiagaraVariable> UNiagaraStackParameterStoreEntry::GetCurrentValueVariable()
 {
 	if (InputType.GetClass() == nullptr)
diff --git a/Engine/Plugins/FX/Niagara/Source/NiagaraEditor/Private/ViewModels/Stack/NiagaraStackParameterStoreGroup.cpp b/Engine/Plugins/FX/Niagara/Source/NiagaraEditor/Private/ViewModels/Stack/NiagaraStackParameterStoreGroup.cpp
--- a/Engine/Plugins/FX/Niagara/Source/NiagaraEditor/Private/ViewModels/Stack/NiagaraStackParameterStoreGroup.cpp
+++ b/Engine/Plugins/FX/Niagara/Source/NiagaraEditor/Private/ViewModels/Stack/NiagaraStackParameterStoreGroup.cpp
@@ -56,11 +56,11 @@ private:
 	FNiagaraVariable NewParameterVariable;
 };
 
-class FParameterStoreGroupAddUtiliites : public FNiagaraStackItemGroupAddUtilities
+class FParameterStoreGroupAddUtiliites : public TNiagaraStackItemGroupAddUtilities<FNiagaraVariable>
 {
 public:
-	FParameterStoreGroupAddUtiliites(UObject& InParameterStoreOwner, FNiagaraParameterStore& InParameterStore, UNiagaraStackEditorData& InStackEditorData)
-		: FNiagaraStackItemGroupAddUtilities(LOCTEXT("ScriptGroupAddItemName", "Parameter"), EAddMode::AddFromAction, true)
+	FParameterStoreGroupAddUtiliites(UObject& InParameterStoreOwner, FNiagaraParameterStore& InParameterStore, UNiagaraStackEditorData& InStackEditorData, FOnItemAdded InOnItemAdded)
+		: TNiagaraStackItemGroupAddUtilities(LOCTEXT("ScriptGroupAddItemName", "Parameter"), EAddMode::AddFromAction, true, InOnItemAdded)
 		, ParameterStoreOwner(InParameterStoreOwner)
 		, ParameterStore(InParameterStore)
 		, StackEditorData(InStackEditorData)
@@ -93,6 +93,8 @@ public:
 
 		ParameterStore.AddParameter(ParameterVariable);
 		StackEditorData.SetModuleInputIsRenamePending(ParameterVariable.GetName().ToString(), true);
+
+		OnItemAdded.ExecuteIfBound(ParameterVariable);
 	}
 
 private:
@@ -106,7 +108,8 @@ void UNiagaraStackParameterStoreGroup::Initialize(
 	UObject* InOwner,
 	FNiagaraParameterStore* InParameterStore)
 {
-	AddUtilities = MakeShared<FParameterStoreGroupAddUtiliites>(*InOwner, *InParameterStore, *InRequiredEntryData.StackEditorData);
+	AddUtilities = MakeShared<FParameterStoreGroupAddUtiliites>(*InOwner, *InParameterStore, *InRequiredEntryData.StackEditorData, 
+		FParameterStoreGroupAddUtiliites::FOnItemAdded::CreateUObject(this, &UNiagaraStackParameterStoreGroup::ParameterAdded));
 	FText DisplayName = LOCTEXT("SystemExposedVariablesGroup", "System Exposed Parameters");
 	FText ToolTip = LOCTEXT("SystemExposedVariablesGroupToolTip", "Displays the variables created in the User namespace. These variables are exposed to owning UComponents, blueprints, etc.");
 	Super::Initialize(InRequiredEntryData, DisplayName, ToolTip, AddUtilities.Get());
@@ -145,6 +148,11 @@ void UNiagaraStackParameterStoreGroup::RefreshChildrenInternal(const TArray<UNia
 	Super::RefreshChildrenInternal(CurrentChildren, NewChildren, NewIssues);
 }
 
+void UNiagaraStackParameterStoreGroup::ParameterAdded(FNiagaraVariable AddedParameter)
+{
+	RefreshChildren();
+}
+
 void UNiagaraStackParameterStoreItem::Initialize(
 	FRequiredEntryData InRequiredEntryData,
 	UObject* InOwner,
@@ -154,17 +162,6 @@ void UNiagaraStackParameterStoreItem::Initialize(
 
 	Owner = InOwner;
 	ParameterStore = InParameterStore;
-	ParameterStoreChangedHandle = ParameterStore->AddOnChangedHandler(
-		FNiagaraParameterStore::FOnChanged::FDelegate::CreateUObject(this, &UNiagaraStackParameterStoreItem::ParameterStoreChanged));
-}
-
-void UNiagaraStackParameterStoreItem::FinalizeInternal()
-{
-	if (Owner != nullptr)
-	{
-		ParameterStore->RemoveOnChangedHandler(ParameterStoreChangedHandle);
-	}
-	Super::FinalizeInternal();
 }
 
 FText UNiagaraStackParameterStoreItem::GetDisplayName() const
@@ -188,6 +185,7 @@ void UNiagaraStackParameterStoreItem::RefreshChildrenInternal(const TArray<UNiag
 			{
 				ValueObjectEntry = NewObject<UNiagaraStackParameterStoreEntry>(this);
 				ValueObjectEntry->Initialize(CreateDefaultChildRequiredData(), Owner.Get(), ParameterStore, Var.GetName().ToString(), Var.GetType(), GetStackEditorDataKey());
+				ValueObjectEntry->OnParameterDeleted().AddUObject(this, &UNiagaraStackParameterStoreItem::ParameterDeleted);
 			}
 
 			NewChildren.Add(ValueObjectEntry);
@@ -196,7 +194,7 @@ void UNiagaraStackParameterStoreItem::RefreshChildrenInternal(const TArray<UNiag
 	Super::RefreshChildrenInternal(CurrentChildren, NewChildren, NewIssues);
 }
 
-void UNiagaraStackParameterStoreItem::ParameterStoreChanged()
+void UNiagaraStackParameterStoreItem::ParameterDeleted()
 {
 	RefreshChildren();
 }
diff --git a/Engine/Plugins/FX/Niagara/Source/NiagaraEditor/Private/ViewModels/Stack/NiagaraStackScriptItemGroup.cpp b/Engine/Plugins/FX/Niagara/Source/NiagaraEditor/Private/ViewModels/Stack/NiagaraStackScriptItemGroup.cpp
--- a/Engine/Plugins/FX/Niagara/Source/NiagaraEditor/Private/ViewModels/Stack/NiagaraStackScriptItemGroup.cpp
+++ b/Engine/Plugins/FX/Niagara/Source/NiagaraEditor/Private/ViewModels/Stack/NiagaraStackScriptItemGroup.cpp
@@ -264,6 +264,18 @@ void UNiagaraStackScriptItemGroup::Initialize(
 	ScriptViewModel = InScriptViewModel;
 	ScriptUsage = InScriptUsage;
 	ScriptUsageId = InScriptUsageId;
+	ScriptGraph = InScriptViewModel->GetGraphViewModel()->GetGraph();
+	ScriptGraph->AddOnGraphChangedHandler(
+		FOnGraphChanged::FDelegate::CreateUObject(this, &UNiagaraStackScriptItemGroup::OnScriptGraphChanged));
+}
+
+void UNiagaraStackScriptItemGroup::FinalizeInternal()
+{
+	if (ScriptGraph.IsValid())
+	{
+		ScriptGraph->RemoveOnGraphChangedHandler(OnGraphChangedHandle);
+	}
+	Super::FinalizeInternal();
 }
 
 void UNiagaraStackScriptItemGroup::RefreshChildrenInternal(const TArray<UNiagaraStackEntry*>& CurrentChildren, TArray<UNiagaraStackEntry*>& NewChildren, TArray<FStackIssue>& NewIssues)
@@ -571,7 +583,7 @@ TOptional<UNiagaraStackEntry::FDropResult> UNiagaraStackScriptItemGroup::ChildRe
 					if (SourceEmitterHandle != nullptr)
 					{
 						UNiagaraNodeOutput* SourceModuleOutputNode = FNiagaraStackGraphUtilities::GetEmitterOutputNodeForStackNode(SourceModuleItem->GetModuleNode());
-						UNiagaraScript* SourceModuleScript = FNiagaraEditorUtilities::GetScriptFromSystem(GetSystemViewModel()->GetSystem(), SourceEmitterHandle->GetId(),
+						UNiagaraScript* SourceModuleScript = FNiagaraEditorUtilities::GetScriptFromSystem(SourceModuleItem->GetSystemViewModel()->GetSystem(), SourceEmitterHandle->GetId(),
 							SourceModuleOutputNode->GetUsage(), SourceModuleOutputNode->GetUsageId());
 
 						const FNiagaraEmitterHandle* TargetEmitterHandle = FNiagaraEditorUtilities::GetEmitterHandleForEmitter(GetSystemViewModel()->GetSystem(), *GetEmitterViewModel()->GetEmitter());
@@ -609,5 +621,13 @@ void UNiagaraStackScriptItemGroup::ChildModifiedGroupItems()
 	RefreshChildren();
 }
 
+void UNiagaraStackScriptItemGroup::OnScriptGraphChanged(const struct FEdGraphEditAction& InAction)
+{
+	if (InAction.Action == GRAPHACTION_RemoveNode)
+	{
+		OnRequestFullRefreshDeferred().Broadcast();
+	}
+}
+
 
 #undef LOCTEXT_NAMESPACE
diff --git a/Engine/Plugins/FX/Niagara/Source/NiagaraEditor/Private/ViewModels/Stack/NiagaraStackViewModel.cpp b/Engine/Plugins/FX/Niagara/Source/NiagaraEditor/Private/ViewModels/Stack/NiagaraStackViewModel.cpp
--- a/Engine/Plugins/FX/Niagara/Source/NiagaraEditor/Private/ViewModels/Stack/NiagaraStackViewModel.cpp
+++ b/Engine/Plugins/FX/Niagara/Source/NiagaraEditor/Private/ViewModels/Stack/NiagaraStackViewModel.cpp
@@ -36,6 +36,7 @@ void UNiagaraStackViewModel::Initialize(TSharedPtr<FNiagaraSystemViewModel> InSy
 		RootEntry->OnStructureChanged().RemoveAll(this);
 		RootEntry->OnDataObjectModified().RemoveAll(this);
 		RootEntry->OnRequestFullRefresh().RemoveAll(this);
+		RootEntry->OnRequestFullRefreshDeferred().RemoveAll(this);
 		RootEntries.Empty();
 		RootEntry->Finalize();
 		RootEntry = nullptr;
@@ -72,12 +73,14 @@ void UNiagaraStackViewModel::Initialize(TSharedPtr<FNiagaraSystemViewModel> InSy
 		RootEntry->OnStructureChanged().AddUObject(this, &UNiagaraStackViewModel::EntryStructureChanged);
 		RootEntry->OnDataObjectModified().AddUObject(this, &UNiagaraStackViewModel::EntryDataObjectModified);
 		RootEntry->OnRequestFullRefresh().AddUObject(this, &UNiagaraStackViewModel::EntryRequestFullRefresh);
+		RootEntry->OnRequestFullRefreshDeferred().AddUObject(this, &UNiagaraStackViewModel::EntryRequestFullRefreshDeferred);
 		RootEntries.Add(RootEntry);
 	}
 
 	CurrentFocusedSearchMatchIndex = -1;
 	StructureChangedDelegate.Broadcast();
 	bRestartSearch = false;
+	bRefreshPending = false;
 }
 
 void UNiagaraStackViewModel::Finalize()
@@ -95,6 +98,13 @@ void UNiagaraStackViewModel::Tick()
 {
 	if (RootEntry)
 	{
+		if (bRefreshPending)
+		{
+			RootEntry->RefreshChildren();
+			OnSearchTextChanged(CurrentSearchText);
+			bRefreshPending = false;
+		}
+
 		SearchTick();
 	}
 }
@@ -404,4 +414,9 @@ void UNiagaraStackViewModel::EntryRequestFullRefresh()
 	RootEntry->RefreshChildren();
 }
 
+void UNiagaraStackViewModel::EntryRequestFullRefreshDeferred()
+{
+	bRefreshPending = true;
+}
+
 #undef LOCTEXT_NAMESPACE
diff --git a/Engine/Plugins/FX/Niagara/Source/NiagaraEditor/Private/Widgets/SNiagaraGraphNode.cpp b/Engine/Plugins/FX/Niagara/Source/NiagaraEditor/Private/Widgets/SNiagaraGraphNode.cpp
--- a/Engine/Plugins/FX/Niagara/Source/NiagaraEditor/Private/Widgets/SNiagaraGraphNode.cpp
+++ b/Engine/Plugins/FX/Niagara/Source/NiagaraEditor/Private/Widgets/SNiagaraGraphNode.cpp
@@ -14,7 +14,7 @@ SNiagaraGraphNode::SNiagaraGraphNode() : SGraphNode()
 
 SNiagaraGraphNode::~SNiagaraGraphNode()
 {
-	if (NiagaraNode)
+	if (NiagaraNode.IsValid())
 	{
 		NiagaraNode->OnVisualsChanged().RemoveAll(this);
 	}
@@ -37,12 +37,12 @@ void SNiagaraGraphNode::HandleNiagaraNodeChanged(UNiagaraNode* InNode)
 void SNiagaraGraphNode::RegisterNiagaraGraphNode(UEdGraphNode* InNode)
 {
 	NiagaraNode = Cast<UNiagaraNode>(InNode);
-	NiagaraNode->OnVisualsChanged().AddRaw(this, &SNiagaraGraphNode::HandleNiagaraNodeChanged);
+	NiagaraNode->OnVisualsChanged().AddSP(this, &SNiagaraGraphNode::HandleNiagaraNodeChanged);
 }
 
 void SNiagaraGraphNode::UpdateGraphNode()
 {
-	check(NiagaraNode != nullptr);
+	check(NiagaraNode.IsValid());
 	SGraphNode::UpdateGraphNode();
 	LastSyncedNodeChangeId = NiagaraNode->GetChangeId();
 }
diff --git a/Engine/Plugins/FX/Niagara/Source/NiagaraEditor/Private/Widgets/SNiagaraGraphNode.h b/Engine/Plugins/FX/Niagara/Source/NiagaraEditor/Private/Widgets/SNiagaraGraphNode.h
--- a/Engine/Plugins/FX/Niagara/Source/NiagaraEditor/Private/Widgets/SNiagaraGraphNode.h
+++ b/Engine/Plugins/FX/Niagara/Source/NiagaraEditor/Private/Widgets/SNiagaraGraphNode.h
@@ -23,6 +23,6 @@ public:
 protected:
 	void RegisterNiagaraGraphNode(UEdGraphNode* InNode);
 	void HandleNiagaraNodeChanged(UNiagaraNode* InNode);
-	UNiagaraNode* NiagaraNode;
+	TWeakObjectPtr<UNiagaraNode> NiagaraNode;
 	FGuid LastSyncedNodeChangeId;
 };
\ No newline at end of file
diff --git a/Engine/Plugins/FX/Niagara/Source/NiagaraEditor/Private/Widgets/SNiagaraParameterMapPaletteItem.cpp b/Engine/Plugins/FX/Niagara/Source/NiagaraEditor/Private/Widgets/SNiagaraParameterMapPaletteItem.cpp
--- a/Engine/Plugins/FX/Niagara/Source/NiagaraEditor/Private/Widgets/SNiagaraParameterMapPaletteItem.cpp
+++ b/Engine/Plugins/FX/Niagara/Source/NiagaraEditor/Private/Widgets/SNiagaraParameterMapPaletteItem.cpp
@@ -8,6 +8,7 @@
 #include "NiagaraEditorStyle.h"
 #include "Widgets/Input/SComboButton.h"
 #include "Widgets/Images/SImage.h"
+#include "ScopedTransaction.h"
 
 #define LOCTEXT_NAMESPACE "NiagaraParameterMapPalleteItem"
 
@@ -87,6 +88,7 @@ void SNiagaraParameterMapPalleteItem::Construct(const FArguments& InArgs, FCreat
 void SNiagaraParameterMapPalleteItem::OnNameTextCommitted(const FText& NewText, ETextCommit::Type InTextCommit)
 {
 	TSharedPtr<FNiagaraParameterAction> ParameterAction = StaticCastSharedPtr<FNiagaraParameterAction>(ActionPtr.Pin());
+	FScopedTransaction RenameParametersWithPins(LOCTEXT("RenameParameter", "Rename parameter, referenced pins and metadata"));
 	OnItemRenamed.ExecuteIfBound(NewText, *ParameterAction.Get());
 }
 
diff --git a/Engine/Plugins/FX/Niagara/Source/NiagaraEditor/Private/Widgets/SNiagaraParameterMapView.cpp b/Engine/Plugins/FX/Niagara/Source/NiagaraEditor/Private/Widgets/SNiagaraParameterMapView.cpp
--- a/Engine/Plugins/FX/Niagara/Source/NiagaraEditor/Private/Widgets/SNiagaraParameterMapView.cpp
+++ b/Engine/Plugins/FX/Niagara/Source/NiagaraEditor/Private/Widgets/SNiagaraParameterMapView.cpp
@@ -32,11 +32,10 @@
 #include "EdGraph/EdGraphPin.h"
 #include "EdGraphSchema_Niagara.h"
 #include "ViewModels/NiagaraEmitterHandleViewModel.h"
+#include "ScopedTransaction.h"
 
 #define LOCTEXT_NAMESPACE "NiagaraParameterMapView"
 
-DEFINE_LOG_CATEGORY(LogNiagara);
-
 FText NiagaraParameterMapSectionID::OnGetSectionTitle(const NiagaraParameterMapSectionID::Type InSection)
 {
 	/* Setup an appropriate name for the section for this node */
@@ -117,6 +116,10 @@ void FNiagaraParameterMapViewCommands::RegisterCommands()
 
 SNiagaraParameterMapView::~SNiagaraParameterMapView()
 {
+	// Unregister all commands for right click on action node
+	ToolkitCommands->UnmapAction(FNiagaraParameterMapViewCommands::Get().DeleteEntry);
+	ToolkitCommands->UnmapAction(FGenericCommands::Get().Rename);
+
 	TSet<UObject*> Objects = SelectedObjects->GetSelectedObjects();
 	for (UObject* Object : Objects)
 	{
@@ -605,6 +608,7 @@ void SNiagaraParameterMapView::OnDeleteEntry()
 				}
 			}
 
+			FScopedTransaction RemoveParametersWithPins(LOCTEXT("RemoveParametersWithPins", "Remove parameter and referenced pins"));
 			for (const TWeakObjectPtr<UNiagaraGraph>& GraphWeakPtr : Graphs)
 			{
 				if (GraphWeakPtr.IsValid())
diff --git a/Engine/Plugins/FX/Niagara/Source/NiagaraEditor/Private/Widgets/SNiagaraSelectedEmitterHandles.cpp b/Engine/Plugins/FX/Niagara/Source/NiagaraEditor/Private/Widgets/SNiagaraSelectedEmitterHandles.cpp
--- a/Engine/Plugins/FX/Niagara/Source/NiagaraEditor/Private/Widgets/SNiagaraSelectedEmitterHandles.cpp
+++ b/Engine/Plugins/FX/Niagara/Source/NiagaraEditor/Private/Widgets/SNiagaraSelectedEmitterHandles.cpp
@@ -1,4 +1,4 @@
-// Copyright 1998-2016 Epic Games, Inc. All Rights Reserved.
+// Copyright 1998-2018 Epic Games, Inc. All Rights Reserved.
 
 #include "SNiagaraSelectedEmitterHandles.h"
 #include "NiagaraEditorStyle.h"
diff --git a/Engine/Plugins/FX/Niagara/Source/NiagaraEditor/Private/Widgets/SNiagaraSelectedEmitterHandles.h b/Engine/Plugins/FX/Niagara/Source/NiagaraEditor/Private/Widgets/SNiagaraSelectedEmitterHandles.h
--- a/Engine/Plugins/FX/Niagara/Source/NiagaraEditor/Private/Widgets/SNiagaraSelectedEmitterHandles.h
+++ b/Engine/Plugins/FX/Niagara/Source/NiagaraEditor/Private/Widgets/SNiagaraSelectedEmitterHandles.h
@@ -1,4 +1,4 @@
-// Copyright 1998-2016 Epic Games, Inc. All Rights Reserved.
+// Copyright 1998-2018 Epic Games, Inc. All Rights Reserved.
 
 #pragma once
 
diff --git a/Engine/Plugins/FX/Niagara/Source/NiagaraEditor/Private/Widgets/SNiagaraSpreadsheetView.cpp b/Engine/Plugins/FX/Niagara/Source/NiagaraEditor/Private/Widgets/SNiagaraSpreadsheetView.cpp
--- a/Engine/Plugins/FX/Niagara/Source/NiagaraEditor/Private/Widgets/SNiagaraSpreadsheetView.cpp
+++ b/Engine/Plugins/FX/Niagara/Source/NiagaraEditor/Private/Widgets/SNiagaraSpreadsheetView.cpp
@@ -1375,8 +1375,13 @@ FReply SNiagaraSpreadsheetView::OnCaptureRequestPressed()
 {
 	FFrameRate TickResolution = SystemViewModel->GetSequencer()->GetFocusedTickResolution();
 	float LocalTime = SystemViewModel->GetSequencer()->GetLocalTime().AsSeconds();
-	float SnapInterval = SystemViewModel->GetSequencer()->GetFocusedDisplayRate().AsInterval();
-	float TargetCaptureTime = LocalTime + SnapInterval;
+	
+	// The preview component in the editor is using the 'DesiredAge' update mode so each frame it determines if the difference
+	// between the current age and the desired age is greater then the seek delta and if so it advanced the simulation the correct
+	// number of times.  We want to ensure that we simulate a single step so we get the seek delta from the component and add that 
+	// to the current time.
+	float SimulationStep = SystemViewModel->GetPreviewComponent()->GetSeekDelta();
+	float TargetCaptureTime = LocalTime + SimulationStep;
 
 	TArray<TSharedRef<FNiagaraEmitterHandleViewModel>> SelectedEmitterHandles;
 	SystemViewModel->GetSelectedEmitterHandles(SelectedEmitterHandles);
diff --git a/Engine/Plugins/FX/Niagara/Source/NiagaraEditor/Public/NiagaraEditorModule.h b/Engine/Plugins/FX/Niagara/Source/NiagaraEditor/Public/NiagaraEditorModule.h
--- a/Engine/Plugins/FX/Niagara/Source/NiagaraEditor/Public/NiagaraEditorModule.h
+++ b/Engine/Plugins/FX/Niagara/Source/NiagaraEditor/Public/NiagaraEditorModule.h
@@ -95,6 +95,7 @@ public:
 private:
 	void RegisterAssetTypeAction(IAssetTools& AssetTools, TSharedRef<IAssetTypeActions> Action);
 	void OnNiagaraSettingsChangedEvent(const FString& PropertyName, const UNiagaraSettings* Settings);
+	void OnPreGarbageCollection();
 
 
 	/** FGCObject interface */
diff --git a/Engine/Plugins/FX/Niagara/Source/NiagaraEditor/Public/NiagaraStackFunctionInputBinder.h b/Engine/Plugins/FX/Niagara/Source/NiagaraEditor/Public/NiagaraStackFunctionInputBinder.h
--- a/Engine/Plugins/FX/Niagara/Source/NiagaraEditor/Public/NiagaraStackFunctionInputBinder.h
+++ b/Engine/Plugins/FX/Niagara/Source/NiagaraEditor/Public/NiagaraStackFunctionInputBinder.h
@@ -1,4 +1,4 @@
-// Copyright 1998-2016 Epic Games, Inc. All Rights Reserved.
+// Copyright 1998-2018 Epic Games, Inc. All Rights Reserved.
 
 #pragma once
 
@@ -48,6 +48,8 @@ public:
 
 	FNiagaraTypeDefinition GetInputType() const;
 
+	UNiagaraNodeFunctionCall* GetFunctionCallNode() const;
+
 	template<typename ValueType>
 	ValueType GetValue() const
 	{
diff --git a/Engine/Plugins/FX/Niagara/Source/NiagaraEditor/Public/ViewModels/NiagaraSystemViewModel.h b/Engine/Plugins/FX/Niagara/Source/NiagaraEditor/Public/ViewModels/NiagaraSystemViewModel.h
--- a/Engine/Plugins/FX/Niagara/Source/NiagaraEditor/Public/ViewModels/NiagaraSystemViewModel.h
+++ b/Engine/Plugins/FX/Niagara/Source/NiagaraEditor/Public/ViewModels/NiagaraSystemViewModel.h
@@ -167,7 +167,7 @@ public:
 
 	// ~ FTickableEditorObject
 	virtual void Tick(float DeltaTime) override;
-	virtual bool IsTickable() const override;
+	virtual bool IsTickable() const override { return true; }
 	virtual TStatId GetStatId() const override;
 
 	/** Resets the System instance to initial conditions. */
@@ -267,6 +267,9 @@ private:
 	/** Rebuilds the sequencer tracks. */
 	void RefreshSequencerTracks();
 
+	/** Updates the data in the sequencer tracks for the specified emitter ids. */
+	void UpdateSequencerTracksForEmitters(const TArray<FGuid>& EmitterIdsRequiringUpdate);
+
 	/** Gets the sequencer emitter track for the supplied emitter handle view model. */
 	UMovieSceneNiagaraEmitterTrack* GetTrackForHandleViewModel(TSharedRef<FNiagaraEmitterHandleViewModel> EmitterHandleViewModel);
 
@@ -490,4 +493,7 @@ private:
 	
 	/** A handle to the on graph changed delegate for the system script. */
 	FDelegateHandle SystemScriptGraphChangedHandler;
+
+	/** An array of emitter handle ids which need their sequencer tracks refreshed next frame. */
+	TArray<FGuid> EmitterIdsRequiringSequencerTrackUpdate;
 };
\ No newline at end of file
diff --git a/Engine/Plugins/FX/Niagara/Source/NiagaraEditor/Public/ViewModels/Stack/NiagaraStackEntry.h b/Engine/Plugins/FX/Niagara/Source/NiagaraEditor/Public/ViewModels/Stack/NiagaraStackEntry.h
--- a/Engine/Plugins/FX/Niagara/Source/NiagaraEditor/Public/ViewModels/Stack/NiagaraStackEntry.h
+++ b/Engine/Plugins/FX/Niagara/Source/NiagaraEditor/Public/ViewModels/Stack/NiagaraStackEntry.h
@@ -39,6 +39,7 @@ public:
 	DECLARE_MULTICAST_DELEGATE(FOnStructureChanged);
 	DECLARE_MULTICAST_DELEGATE_OneParam(FOnDataObjectModified, UObject*);
 	DECLARE_MULTICAST_DELEGATE(FOnRequestFullRefresh);
+	DECLARE_MULTICAST_DELEGATE(FOnRequestFullRefreshDeferred);
 	DECLARE_DELEGATE_RetVal_TwoParams(TOptional<FDropResult>, FOnRequestDrop, const UNiagaraStackEntry&, const TArray<UNiagaraStackEntry*>&);
 	DECLARE_DELEGATE_RetVal_OneParam(bool, FOnFilterChild, const UNiagaraStackEntry&);
 	DECLARE_DELEGATE(FStackIssueFixDelegate);
@@ -206,6 +207,8 @@ public:
 
 	FOnRequestFullRefresh& OnRequestFullRefresh();
 
+	FOnRequestFullRefresh& OnRequestFullRefreshDeferred();
+
 	void RefreshChildren();
 
 
@@ -275,6 +278,8 @@ private:
 
 	void ChildRequestFullRefresh();
 
+	void ChildRequestFullRefreshDeferred();
+
 	TOptional<FDropResult> ChildRequestCanDrop(const UNiagaraStackEntry& TargetChild, const TArray<UNiagaraStackEntry*>& DraggedEntries);
 
 	TOptional<FDropResult> ChildRequestDrop(const UNiagaraStackEntry& TargetChild, const TArray<UNiagaraStackEntry*>& DraggedEntries);
@@ -297,6 +302,8 @@ private:
 
 	FOnRequestFullRefresh RequestFullRefreshDelegate;
 
+	FOnRequestFullRefresh RequestFullRefreshDeferredDelegate;
+
 	TArray<FOnFilterChild> ChildFilters;
 
 	UPROPERTY()
diff --git a/Engine/Plugins/FX/Niagara/Source/NiagaraEditor/Public/ViewModels/Stack/NiagaraStackGraphUtilities.h b/Engine/Plugins/FX/Niagara/Source/NiagaraEditor/Public/ViewModels/Stack/NiagaraStackGraphUtilities.h
--- a/Engine/Plugins/FX/Niagara/Source/NiagaraEditor/Public/ViewModels/Stack/NiagaraStackGraphUtilities.h
+++ b/Engine/Plugins/FX/Niagara/Source/NiagaraEditor/Public/ViewModels/Stack/NiagaraStackGraphUtilities.h
@@ -94,9 +94,13 @@ namespace FNiagaraStackGraphUtilities
 
 	void SetCustomExpressionForFunctionInput(UEdGraphPin& OverridePin, UNiagaraNodeCustomHlsl*& OutDynamicInputFunctionCall, const FGuid& NewNodePersistentId = FGuid());
 
-	bool RemoveModuleFromStack(UNiagaraNodeFunctionCall& ModuleNode);
+	bool RemoveModuleFromStack(UNiagaraSystem& OwningSystem, FGuid OwningEmitterId, UNiagaraNodeFunctionCall& ModuleNode);
 
-	bool RemoveModuleFromStack(UNiagaraNodeFunctionCall& ModuleNode, TArray<TWeakObjectPtr<UNiagaraNodeInput>>& OutRemovedNodes);
+	bool RemoveModuleFromStack(UNiagaraSystem& OwningSystem, FGuid OwningEmitterId, UNiagaraNodeFunctionCall& ModuleNode, TArray<TWeakObjectPtr<UNiagaraNodeInput>>& OutRemovedInputNodes);
+
+	bool RemoveModuleFromStack(UNiagaraScript& OwningScript, UNiagaraNodeFunctionCall& ModuleNode);
+
+	bool RemoveModuleFromStack(UNiagaraScript& OwningScript, UNiagaraNodeFunctionCall& ModuleNode, TArray<TWeakObjectPtr<UNiagaraNodeInput>>& OutRemovedInputNodes);
 
 	UNiagaraNodeFunctionCall* AddScriptModuleToStack(FAssetData ModuleScriptAsset, UNiagaraNodeOutput& TargetOutputNode, int32 TargetIndex = INDEX_NONE);
 	
diff --git a/Engine/Plugins/FX/Niagara/Source/NiagaraEditor/Public/ViewModels/Stack/NiagaraStackModuleItem.h b/Engine/Plugins/FX/Niagara/Source/NiagaraEditor/Public/ViewModels/Stack/NiagaraStackModuleItem.h
--- a/Engine/Plugins/FX/Niagara/Source/NiagaraEditor/Public/ViewModels/Stack/NiagaraStackModuleItem.h
+++ b/Engine/Plugins/FX/Niagara/Source/NiagaraEditor/Public/ViewModels/Stack/NiagaraStackModuleItem.h
@@ -7,7 +7,6 @@
 #include "NiagaraStackModuleItem.generated.h"
 
 class UNiagaraNodeFunctionCall;
-class UNiagaraNode;
 class UNiagaraStackFunctionInputCollection;
 class UNiagaraStackModuleItemOutputCollection;
 class INiagaraStackItemGroupAddUtilities;
@@ -52,7 +51,7 @@ public:
 
 	void NotifyModuleMoved();
 
-	bool CanAddInput() const;
+	bool CanAddInput(FNiagaraVariable InputParameter) const;
 
 	void AddInput(FNiagaraVariable InputParameter);
 
@@ -67,6 +66,7 @@ private:
 	void RefreshIsEnabled();
 
 private:
+	UNiagaraNodeOutput* OutputNode;
 	UNiagaraNodeFunctionCall* FunctionCallNode;
 	bool bCanMoveAndDelete;
 	bool bIsEnabled;
diff --git a/Engine/Plugins/FX/Niagara/Source/NiagaraEditor/Public/ViewModels/Stack/NiagaraStackParameterStoreEntry.h b/Engine/Plugins/FX/Niagara/Source/NiagaraEditor/Public/ViewModels/Stack/NiagaraStackParameterStoreEntry.h
--- a/Engine/Plugins/FX/Niagara/Source/NiagaraEditor/Public/ViewModels/Stack/NiagaraStackParameterStoreEntry.h
+++ b/Engine/Plugins/FX/Niagara/Source/NiagaraEditor/Public/ViewModels/Stack/NiagaraStackParameterStoreEntry.h
@@ -23,6 +23,8 @@ class NIAGARAEDITOR_API UNiagaraStackParameterStoreEntry : public UNiagaraStackI
 
 public:
 	DECLARE_MULTICAST_DELEGATE(FOnValueChanged);
+	DECLARE_MULTICAST_DELEGATE(FOnParameterDeleted);
+
 public:
 	UNiagaraStackParameterStoreEntry();
 
@@ -88,6 +90,9 @@ public:
 	/** Gets a multicast delegate which is called whenever the value on this input changes. */
 	FOnValueChanged& OnValueChanged();
 
+	/** Gets a multicast delegate which is called when this parameter is deleted. */
+	FOnParameterDeleted& OnParameterDeleted();
+
 	/** Delete the parameter from the ParameterStore and notify that the store changed. */
 	void Delete();
 
@@ -116,13 +121,16 @@ private:
 	/** A local copy of the value of this input if one is available. */
 	TSharedPtr<FStructOnScope> LocalValueStruct;
 
+	/** A multicast delegate which is called when the value of this input is changed. */
+	FOnValueChanged ValueChangedDelegate;
+
+	/** A multicast delegate which is called when this parameter is deleted. */
+	FOnParameterDeleted ParameterDeletedDelegate;
+
 	/** A pointer to the data interface object for this input if one is available. */
 	UPROPERTY()
 	UNiagaraDataInterface* ValueObject;
 
-	/** A multicast delegate which is called when the value of this input is changed. */
-	FOnValueChanged ValueChangedDelegate;
-
 	UPROPERTY()
 	UObject* Owner;
 
diff --git a/Engine/Plugins/FX/Niagara/Source/NiagaraEditor/Public/ViewModels/Stack/NiagaraStackParameterStoreGroup.h b/Engine/Plugins/FX/Niagara/Source/NiagaraEditor/Public/ViewModels/Stack/NiagaraStackParameterStoreGroup.h
--- a/Engine/Plugins/FX/Niagara/Source/NiagaraEditor/Public/ViewModels/Stack/NiagaraStackParameterStoreGroup.h
+++ b/Engine/Plugins/FX/Niagara/Source/NiagaraEditor/Public/ViewModels/Stack/NiagaraStackParameterStoreGroup.h
@@ -18,11 +18,12 @@ class NIAGARAEDITOR_API UNiagaraStackParameterStoreGroup : public UNiagaraStackI
 public:
 	void Initialize(FRequiredEntryData InRequiredEntryData,	UObject* InOwner, FNiagaraParameterStore* InParameterStore);
 
-	void AddUserParameter(FNiagaraVariable ParameterVariable);
-
 protected:
 	virtual void RefreshChildrenInternal(const TArray<UNiagaraStackEntry*>& CurrentChildren, TArray<UNiagaraStackEntry*>& NewChildren, TArray<FStackIssue>& NewIssues) override;
 
+private:
+	void ParameterAdded(FNiagaraVariable AddedParameter);
+
 private:
 	TWeakObjectPtr<UObject> Owner;
 	FNiagaraParameterStore* ParameterStore;
@@ -41,15 +42,12 @@ public:
 	virtual FText GetDisplayName() const override;
 
 protected:
-	virtual void FinalizeInternal() override;
-
 	virtual void RefreshChildrenInternal(const TArray<UNiagaraStackEntry*>& CurrentChildren, TArray<UNiagaraStackEntry*>& NewChildren, TArray<FStackIssue>& NewIssues) override;
 
 private:
-	void ParameterStoreChanged();
+	void ParameterDeleted();
 
 private:
 	TWeakObjectPtr<UObject> Owner;
 	FNiagaraParameterStore* ParameterStore;
-	FDelegateHandle ParameterStoreChangedHandle;
 };
\ No newline at end of file
diff --git a/Engine/Plugins/FX/Niagara/Source/NiagaraEditor/Public/ViewModels/Stack/NiagaraStackScriptItemGroup.h b/Engine/Plugins/FX/Niagara/Source/NiagaraEditor/Public/ViewModels/Stack/NiagaraStackScriptItemGroup.h
--- a/Engine/Plugins/FX/Niagara/Source/NiagaraEditor/Public/ViewModels/Stack/NiagaraStackScriptItemGroup.h
+++ b/Engine/Plugins/FX/Niagara/Source/NiagaraEditor/Public/ViewModels/Stack/NiagaraStackScriptItemGroup.h
@@ -13,6 +13,7 @@ class UNiagaraNodeOutput;
 class UNiagaraNodeFunctionCall;
 class FScriptItemGroupAddUtilities;
 class UNiagaraStackModuleItem;
+class UEdGraph;
 
 UCLASS()
 class NIAGARAEDITOR_API UNiagaraStackScriptItemGroup : public UNiagaraStackItemGroup
@@ -34,6 +35,8 @@ public:
 protected:
 	virtual void RefreshChildrenInternal(const TArray<UNiagaraStackEntry*>& CurrentChildren, TArray<UNiagaraStackEntry*>& NewChildren, TArray<FStackIssue>& NewIssues) override;
 
+	virtual void FinalizeInternal() override;
+
 	virtual TOptional<FDropResult> ChildRequestCanDropInternal(const UNiagaraStackEntry& TargetChild, const TArray<UNiagaraStackEntry*>& DraggedEntries) override;
 
 	virtual TOptional<FDropResult> ChildRequestDropInternal(const UNiagaraStackEntry& TargetChild, const TArray<UNiagaraStackEntry*>& DraggedEntries) override;
@@ -43,6 +46,8 @@ private:
 
 	void ChildModifiedGroupItems();
 
+	void OnScriptGraphChanged(const struct FEdGraphEditAction& InAction);
+
 protected:
 	TWeakPtr<FNiagaraScriptViewModel> ScriptViewModel;
 	void RefreshIssues(TArray<FStackIssue>& NewIssues);
@@ -55,5 +60,9 @@ private:
 	FGuid ScriptUsageId;
 	bool bIsValidForOutput;
 
+	TWeakObjectPtr<UEdGraph> ScriptGraph;
+
+	FDelegateHandle OnGraphChangedHandle;
+
 	TMap<FObjectKey, UNiagaraStackModuleItem*> StackSpacerToModuleItemMap;
 };
diff --git a/Engine/Plugins/FX/Niagara/Source/NiagaraEditor/Public/ViewModels/Stack/NiagaraStackViewModel.h b/Engine/Plugins/FX/Niagara/Source/NiagaraEditor/Public/ViewModels/Stack/NiagaraStackViewModel.h
--- a/Engine/Plugins/FX/Niagara/Source/NiagaraEditor/Public/ViewModels/Stack/NiagaraStackViewModel.h
+++ b/Engine/Plugins/FX/Niagara/Source/NiagaraEditor/Public/ViewModels/Stack/NiagaraStackViewModel.h
@@ -101,6 +101,7 @@ private:
 	void EntryStructureChanged();
 	void EntryDataObjectModified(UObject* ChangedObject);
 	void EntryRequestFullRefresh();
+	void EntryRequestFullRefreshDeferred();
 	void OnSystemCompiled();
 	void OnEmitterCompiled();
 	/** Called by the tick function to perform partial search */
@@ -129,4 +130,5 @@ private:
 	TArray<FSearchResult> CurrentSearchResults;
 	static const double MaxSearchTime;
 	bool bRestartSearch;
+	bool bRefreshPending;
 };
\ No newline at end of file
diff --git a/Engine/Plugins/FX/Niagara/Source/NiagaraEditorWidgets/Private/DetailCustomizations/NiagaraDataInterfaceDetails.cpp b/Engine/Plugins/FX/Niagara/Source/NiagaraEditorWidgets/Private/DetailCustomizations/NiagaraDataInterfaceDetails.cpp
--- a/Engine/Plugins/FX/Niagara/Source/NiagaraEditorWidgets/Private/DetailCustomizations/NiagaraDataInterfaceDetails.cpp
+++ b/Engine/Plugins/FX/Niagara/Source/NiagaraEditorWidgets/Private/DetailCustomizations/NiagaraDataInterfaceDetails.cpp
@@ -1,4 +1,5 @@
-// Copyright 1998-2017 Epic Games, Inc. All Rights Reserved.
+// Copyright 1998-2018 Epic Games, Inc. All Rights Reserved.
+
 #include "NiagaraDataInterfaceDetails.h"
 #include "IDetailCustomization.h"
 #include "NiagaraDataInterface.h"
@@ -117,11 +118,15 @@ class FNiagaraDataInterfaceCustomNodeBuilder : public IDetailCustomNodeBuilder
 											 , public TSharedFromThis<FNiagaraDataInterfaceCustomNodeBuilder>
 {
 public:
-	FNiagaraDataInterfaceCustomNodeBuilder(UNiagaraDataInterface* InDataInterface, IDetailLayoutBuilder* InDetailBuilder)
-		: DataInterface(InDataInterface)
-		, DetailBuilder(InDetailBuilder)
+	FNiagaraDataInterfaceCustomNodeBuilder(IDetailLayoutBuilder* InDetailBuilder)
+		: DetailBuilder(InDetailBuilder)
+	{
+	}
+
+	void Initialize(UNiagaraDataInterface& InDataInterface)
 	{
-		DataInterface->OnChanged().AddRaw(this, &FNiagaraDataInterfaceCustomNodeBuilder::OnDataInterfaceChanged);
+		DataInterface = &InDataInterface;
+		DataInterface->OnChanged().AddSP(this, &FNiagaraDataInterfaceCustomNodeBuilder::OnDataInterfaceChanged);
 	}
 
 	~FNiagaraDataInterfaceCustomNodeBuilder()
@@ -197,10 +202,11 @@ void FNiagaraDataInterfaceDetailsBase::CustomizeDetails(IDetailLayoutBuilder& De
 	check(SelectedObjects.Num() == 1);
 	DataInterface = Cast<UNiagaraDataInterface>(SelectedObjects[0].Get());
 	check(DataInterface.IsValid());
-	DataInterface->OnChanged().AddRaw(this, &FNiagaraDataInterfaceDetailsBase::OnDataChanged);
+	DataInterface->OnChanged().AddSP(this, &FNiagaraDataInterfaceDetailsBase::OnDataChanged);
 	IDetailCategoryBuilder& ErrorsBuilderRef = DetailBuilder.EditCategory(ErrorsCategoryName, LOCTEXT("Errors", "Errors"), ECategoryPriority::Important);
 	ErrorsCategoryBuilder = &ErrorsBuilderRef;
-	CustomBuilder = TSharedPtr<FNiagaraDataInterfaceCustomNodeBuilder>(new FNiagaraDataInterfaceCustomNodeBuilder(DataInterface.Get(), &DetailBuilder));
+	CustomBuilder = MakeShared<FNiagaraDataInterfaceCustomNodeBuilder>(&DetailBuilder);
+	CustomBuilder->Initialize(*DataInterface);
 	ErrorsCategoryBuilder->AddCustomBuilder(CustomBuilder.ToSharedRef());
 	OnDataChanged();
 }
diff --git a/Engine/Plugins/FX/Niagara/Source/NiagaraEditorWidgets/Private/DetailCustomizations/NiagaraDataInterfaceDetails.h b/Engine/Plugins/FX/Niagara/Source/NiagaraEditorWidgets/Private/DetailCustomizations/NiagaraDataInterfaceDetails.h
--- a/Engine/Plugins/FX/Niagara/Source/NiagaraEditorWidgets/Private/DetailCustomizations/NiagaraDataInterfaceDetails.h
+++ b/Engine/Plugins/FX/Niagara/Source/NiagaraEditorWidgets/Private/DetailCustomizations/NiagaraDataInterfaceDetails.h
@@ -1,4 +1,4 @@
-// Copyright 1998-2017 Epic Games, Inc. All Rights Reserved.
+// Copyright 1998-2018 Epic Games, Inc. All Rights Reserved.
 
 #pragma once
 
diff --git a/Engine/Plugins/FX/Niagara/Source/NiagaraEditorWidgets/Private/DetailCustomizations/NiagaraDataInterfaceSkeletalMeshDetails.cpp b/Engine/Plugins/FX/Niagara/Source/NiagaraEditorWidgets/Private/DetailCustomizations/NiagaraDataInterfaceSkeletalMeshDetails.cpp
--- a/Engine/Plugins/FX/Niagara/Source/NiagaraEditorWidgets/Private/DetailCustomizations/NiagaraDataInterfaceSkeletalMeshDetails.cpp
+++ b/Engine/Plugins/FX/Niagara/Source/NiagaraEditorWidgets/Private/DetailCustomizations/NiagaraDataInterfaceSkeletalMeshDetails.cpp
@@ -1,4 +1,5 @@
- // Copyright 1998-2017 Epic Games, Inc. All Rights Reserved.
+// Copyright 1998-2018 Epic Games, Inc. All Rights Reserved.
+
 #include "NiagaraDataInterfaceSkeletalMeshDetails.h"
 #include "NiagaraDetailSourcedArrayBuilder.h"
 #include "NiagaraDataInterfaceDetails.h"
@@ -65,7 +66,7 @@ void FNiagaraDataInterfaceSkeletalMeshDetails::CustomizeDetails(IDetailLayoutBui
 
  void FNiagaraDataInterfaceSkeletalMeshDetails::OnInterfaceChanged()
  {
-     // Rebuild the data changed listener
+	 // Rebuild the data changed listener
 	 TWeakObjectPtr<USceneComponent> SceneComponent;
 	 USkeletalMeshComponent* FoundSkelComp = nullptr;
 	 if (MeshObject != nullptr)
diff --git a/Engine/Plugins/FX/Niagara/Source/NiagaraEditorWidgets/Private/DetailCustomizations/NiagaraDataInterfaceSkeletalMeshDetails.h b/Engine/Plugins/FX/Niagara/Source/NiagaraEditorWidgets/Private/DetailCustomizations/NiagaraDataInterfaceSkeletalMeshDetails.h
--- a/Engine/Plugins/FX/Niagara/Source/NiagaraEditorWidgets/Private/DetailCustomizations/NiagaraDataInterfaceSkeletalMeshDetails.h
+++ b/Engine/Plugins/FX/Niagara/Source/NiagaraEditorWidgets/Private/DetailCustomizations/NiagaraDataInterfaceSkeletalMeshDetails.h
@@ -1,4 +1,4 @@
-// Copyright 1998-2017 Epic Games, Inc. All Rights Reserved.
+// Copyright 1998-2018 Epic Games, Inc. All Rights Reserved.
 
 #pragma once
 
diff --git a/Engine/Plugins/FX/Niagara/Source/NiagaraEditorWidgets/Private/DetailCustomizations/NiagaraDetailSourcedArrayBuilder.cpp b/Engine/Plugins/FX/Niagara/Source/NiagaraEditorWidgets/Private/DetailCustomizations/NiagaraDetailSourcedArrayBuilder.cpp
--- a/Engine/Plugins/FX/Niagara/Source/NiagaraEditorWidgets/Private/DetailCustomizations/NiagaraDetailSourcedArrayBuilder.cpp
+++ b/Engine/Plugins/FX/Niagara/Source/NiagaraEditorWidgets/Private/DetailCustomizations/NiagaraDetailSourcedArrayBuilder.cpp
@@ -1,4 +1,5 @@
- // Copyright 1998-2017 Epic Games, Inc. All Rights Reserved.
+// Copyright 1998-2018 Epic Games, Inc. All Rights Reserved.
+
 #include "NiagaraDetailSourcedArrayBuilder.h"
 #include "NiagaraDataInterfaceDetails.h"
 #include "NiagaraDataInterfaceSkeletalMesh.h" 
diff --git a/Engine/Plugins/FX/Niagara/Source/NiagaraEditorWidgets/Private/DetailCustomizations/NiagaraDetailSourcedArrayBuilder.h b/Engine/Plugins/FX/Niagara/Source/NiagaraEditorWidgets/Private/DetailCustomizations/NiagaraDetailSourcedArrayBuilder.h
--- a/Engine/Plugins/FX/Niagara/Source/NiagaraEditorWidgets/Private/DetailCustomizations/NiagaraDetailSourcedArrayBuilder.h
+++ b/Engine/Plugins/FX/Niagara/Source/NiagaraEditorWidgets/Private/DetailCustomizations/NiagaraDetailSourcedArrayBuilder.h
@@ -1,4 +1,4 @@
-// Copyright 1998-2017 Epic Games, Inc. All Rights Reserved.
+// Copyright 1998-2018 Epic Games, Inc. All Rights Reserved.
 
 #pragma once
 
diff --git a/Engine/Plugins/FX/Niagara/Source/NiagaraEditorWidgets/Private/DetailCustomizations/NiagaraSkeletalMeshSamplingInfoDetails.cpp b/Engine/Plugins/FX/Niagara/Source/NiagaraEditorWidgets/Private/DetailCustomizations/NiagaraSkeletalMeshSamplingInfoDetails.cpp
--- a/Engine/Plugins/FX/Niagara/Source/NiagaraEditorWidgets/Private/DetailCustomizations/NiagaraSkeletalMeshSamplingInfoDetails.cpp
+++ b/Engine/Plugins/FX/Niagara/Source/NiagaraEditorWidgets/Private/DetailCustomizations/NiagaraSkeletalMeshSamplingInfoDetails.cpp
@@ -1,4 +1,5 @@
- // Copyright 1998-2017 Epic Games, Inc. All Rights Reserved.
+// Copyright 1998-2018 Epic Games, Inc. All Rights Reserved.
+
 #include "NiagaraSkeletalMeshSamplingInfoDetails.h"
 #include "Animation/Skeleton.h"
 #include "ReferenceSkeleton.h"
@@ -65,15 +66,15 @@
 
  void FNiagaraSkeletalMeshSamplingInfoDetails::OnGenerateRegionEntry(TSharedRef<IPropertyHandle> PropertyHandle, int32 ArrayIndex, IDetailChildrenBuilder& ChildrenBuilder)
   {
- 	 IDetailPropertyRow& RegionRow = ChildrenBuilder.AddProperty(PropertyHandle);
+	 IDetailPropertyRow& RegionRow = ChildrenBuilder.AddProperty(PropertyHandle);
  
- 	 FNumberFormattingOptions NoCommas;
- 	 NoCommas.UseGrouping = false;
- 	 const FText SlotDesc = FText::Format(LOCTEXT("RegionSlotIndex", "Region #{0}"), FText::AsNumber(ArrayIndex, &NoCommas));
+	 FNumberFormattingOptions NoCommas;
+	 NoCommas.UseGrouping = false;
+	 const FText SlotDesc = FText::Format(LOCTEXT("RegionSlotIndex", "Region #{0}"), FText::AsNumber(ArrayIndex, &NoCommas));
  
- 	 RegionRow.DisplayName(SlotDesc);
+	 RegionRow.DisplayName(SlotDesc);
  
- 	 RegionRow.ShowPropertyButtons(true);
+	 RegionRow.ShowPropertyButtons(true);
 	 // Add custom builder for properties
 	 TSharedPtr<IPropertyHandle> BoneFiltersProperty = PropertyHandle->GetChildHandle("BoneFilters");
 	 TSharedRef<FNiagaraDetailSourcedArrayBuilder> BoneBuilder = MakeShareable(new FNiagaraDetailSourcedArrayBuilder(BoneFiltersProperty.ToSharedRef(), PossibleBonesNames, "BoneName"));
diff --git a/Engine/Plugins/FX/Niagara/Source/NiagaraEditorWidgets/Private/DetailCustomizations/NiagaraSkeletalMeshSamplingInfoDetails.h b/Engine/Plugins/FX/Niagara/Source/NiagaraEditorWidgets/Private/DetailCustomizations/NiagaraSkeletalMeshSamplingInfoDetails.h
--- a/Engine/Plugins/FX/Niagara/Source/NiagaraEditorWidgets/Private/DetailCustomizations/NiagaraSkeletalMeshSamplingInfoDetails.h
+++ b/Engine/Plugins/FX/Niagara/Source/NiagaraEditorWidgets/Private/DetailCustomizations/NiagaraSkeletalMeshSamplingInfoDetails.h
@@ -1,4 +1,4 @@
-// Copyright 1998-2017 Epic Games, Inc. All Rights Reserved.
+// Copyright 1998-2018 Epic Games, Inc. All Rights Reserved.
 
 #pragma once
 
diff --git a/Engine/Plugins/FX/Niagara/Source/NiagaraEditorWidgets/Private/DetailCustomizations/SNiagaraNamePropertySelector.cpp b/Engine/Plugins/FX/Niagara/Source/NiagaraEditorWidgets/Private/DetailCustomizations/SNiagaraNamePropertySelector.cpp
--- a/Engine/Plugins/FX/Niagara/Source/NiagaraEditorWidgets/Private/DetailCustomizations/SNiagaraNamePropertySelector.cpp
+++ b/Engine/Plugins/FX/Niagara/Source/NiagaraEditorWidgets/Private/DetailCustomizations/SNiagaraNamePropertySelector.cpp
@@ -1,4 +1,5 @@
- // Copyright 1998-2017 Epic Games, Inc. All Rights Reserved.
+// Copyright 1998-2018 Epic Games, Inc. All Rights Reserved.
+
 #include "SNiagaraNamePropertySelector.h"
 #include "NiagaraDataInterfaceDetails.h"
 #include "PropertyCustomizationHelpers.h"
diff --git a/Engine/Plugins/FX/Niagara/Source/NiagaraEditorWidgets/Private/DetailCustomizations/SNiagaraNamePropertySelector.h b/Engine/Plugins/FX/Niagara/Source/NiagaraEditorWidgets/Private/DetailCustomizations/SNiagaraNamePropertySelector.h
--- a/Engine/Plugins/FX/Niagara/Source/NiagaraEditorWidgets/Private/DetailCustomizations/SNiagaraNamePropertySelector.h
+++ b/Engine/Plugins/FX/Niagara/Source/NiagaraEditorWidgets/Private/DetailCustomizations/SNiagaraNamePropertySelector.h
@@ -1,4 +1,4 @@
-// Copyright 1998-2017 Epic Games, Inc. All Rights Reserved.
+// Copyright 1998-2018 Epic Games, Inc. All Rights Reserved.
 
 #pragma once
 
diff --git a/Engine/Plugins/FX/Niagara/Source/NiagaraEditorWidgets/Private/NiagaraEditorWidgetsStyle.cpp b/Engine/Plugins/FX/Niagara/Source/NiagaraEditorWidgets/Private/NiagaraEditorWidgetsStyle.cpp
--- a/Engine/Plugins/FX/Niagara/Source/NiagaraEditorWidgets/Private/NiagaraEditorWidgetsStyle.cpp
+++ b/Engine/Plugins/FX/Niagara/Source/NiagaraEditorWidgets/Private/NiagaraEditorWidgetsStyle.cpp
@@ -34,11 +34,7 @@ FName FNiagaraEditorWidgetsStyle::GetStyleSetName()
 	return StyleSetName;
 }
 
-FString RelativePathToPluginPath(const FString& RelativePath, const ANSICHAR* Extension)
-{
-	static FString ContentDir = IPluginManager::Get().FindPlugin(TEXT("Niagara"))->GetContentDir();
-	return (ContentDir / RelativePath) + Extension;
-}
+NIAGARAEDITOR_API FString RelativePathToPluginPath(const FString& RelativePath, const ANSICHAR* Extension);
 
 #define IMAGE_PLUGIN_BRUSH( RelativePath, ... ) FSlateImageBrush( RelativePathToPluginPath( RelativePath, ".png" ), __VA_ARGS__ )
 #define IMAGE_BRUSH( RelativePath, ... ) FSlateImageBrush( Style->RootToContentDir( RelativePath, TEXT(".png") ), __VA_ARGS__ )
diff --git a/Engine/Plugins/FX/Niagara/Source/NiagaraEditorWidgets/Private/Stack/SNiagaraStackFunctionInputValue.cpp b/Engine/Plugins/FX/Niagara/Source/NiagaraEditorWidgets/Private/Stack/SNiagaraStackFunctionInputValue.cpp
--- a/Engine/Plugins/FX/Niagara/Source/NiagaraEditorWidgets/Private/Stack/SNiagaraStackFunctionInputValue.cpp
+++ b/Engine/Plugins/FX/Niagara/Source/NiagaraEditorWidgets/Private/Stack/SNiagaraStackFunctionInputValue.cpp
@@ -244,7 +244,7 @@ TSharedRef<SWidget> SNiagaraStackFunctionInputValue::ConstructLocalValueStructWi
 			ParameterEditor->SetOnEndValueChange(SNiagaraParameterEditor::FOnValueChange::CreateSP(
 				this, &SNiagaraStackFunctionInputValue::ParameterEndValueChange));
 			ParameterEditor->SetOnValueChanged(SNiagaraParameterEditor::FOnValueChange::CreateSP(
-				this, &SNiagaraStackFunctionInputValue::ParameterValueChanged, ParameterEditor.ToSharedRef()));
+				this, &SNiagaraStackFunctionInputValue::ParameterValueChanged, TWeakPtr<SNiagaraParameterEditor>(ParameterEditor)));
 
 			LocalValueStructParameterEditor = ParameterEditor;
 
@@ -309,10 +309,14 @@ void SNiagaraStackFunctionInputValue::ParameterEndValueChange()
 	FunctionInput->NotifyEndLocalValueChange();
 }
 
-void SNiagaraStackFunctionInputValue::ParameterValueChanged(TSharedRef<SNiagaraParameterEditor> ParameterEditor)
+void SNiagaraStackFunctionInputValue::ParameterValueChanged(TWeakPtr<SNiagaraParameterEditor> ParameterEditor)
 {
-	ParameterEditor->UpdateStructFromInternalValue(FunctionInput->GetLocalValueStruct().ToSharedRef());
-	FunctionInput->SetLocalValue(DisplayedLocalValueStruct.ToSharedRef());
+	TSharedPtr<SNiagaraParameterEditor> ParameterEditorPinned = ParameterEditor.Pin();
+	if (ParameterEditorPinned.IsValid())
+	{
+		ParameterEditorPinned->UpdateStructFromInternalValue(DisplayedLocalValueStruct.ToSharedRef());
+		FunctionInput->SetLocalValue(DisplayedLocalValueStruct.ToSharedRef());
+	}
 }
 
 void SNiagaraStackFunctionInputValue::ParameterPropertyValueChanged(const FPropertyChangedEvent& PropertyChangedEvent)
diff --git a/Engine/Plugins/FX/Niagara/Source/NiagaraEditorWidgets/Private/Stack/SNiagaraStackFunctionInputValue.h b/Engine/Plugins/FX/Niagara/Source/NiagaraEditorWidgets/Private/Stack/SNiagaraStackFunctionInputValue.h
--- a/Engine/Plugins/FX/Niagara/Source/NiagaraEditorWidgets/Private/Stack/SNiagaraStackFunctionInputValue.h
+++ b/Engine/Plugins/FX/Niagara/Source/NiagaraEditorWidgets/Private/Stack/SNiagaraStackFunctionInputValue.h
@@ -39,7 +39,7 @@ private:
 
 	void ParameterEndValueChange();
 
-	void ParameterValueChanged(TSharedRef<SNiagaraParameterEditor> ParameterEditor);
+	void ParameterValueChanged(TWeakPtr<SNiagaraParameterEditor> ParameterEditor);
 
 	void ParameterPropertyValueChanged(const FPropertyChangedEvent& PropertyChangedEvent);
 
diff --git a/Engine/Plugins/FX/Niagara/Source/NiagaraEditorWidgets/Private/Stack/SNiagaraStackModuleItem.cpp b/Engine/Plugins/FX/Niagara/Source/NiagaraEditorWidgets/Private/Stack/SNiagaraStackModuleItem.cpp
--- a/Engine/Plugins/FX/Niagara/Source/NiagaraEditorWidgets/Private/Stack/SNiagaraStackModuleItem.cpp
+++ b/Engine/Plugins/FX/Niagara/Source/NiagaraEditorWidgets/Private/Stack/SNiagaraStackModuleItem.cpp
@@ -269,11 +269,11 @@ void SNiagaraStackModuleItem::ShowInsertModuleMenu(int32 InsertIndex)
 
 FReply SNiagaraStackModuleItem::OnModuleItemDrop(TSharedPtr<FDragDropOperation> DragDropOperation)
 {
-	if (ModuleItem->CanAddInput() && DragDropOperation->IsOfType<FNiagaraStackDragOperation>())
+	if (DragDropOperation->IsOfType<FNiagaraStackDragOperation>())
 	{
 		TSharedPtr<FNiagaraStackDragOperation> InputDragDropOperation = StaticCastSharedPtr<FNiagaraStackDragOperation>(DragDropOperation);
 		TSharedPtr<FNiagaraParameterAction> Action = StaticCastSharedPtr<FNiagaraParameterAction>(InputDragDropOperation->GetAction());
-		if (Action.IsValid())
+		if (Action.IsValid() && ModuleItem->CanAddInput(Action->GetParameter()))
 		{
 			ModuleItem->AddInput(Action->GetParameter());
 			return FReply::Handled();
@@ -285,14 +285,11 @@ FReply SNiagaraStackModuleItem::OnModuleItemDrop(TSharedPtr<FDragDropOperation>
 
 bool SNiagaraStackModuleItem::OnModuleItemAllowDrop(TSharedPtr<FDragDropOperation> DragDropOperation)
 {
-	if (ModuleItem->CanAddInput() && DragDropOperation->IsOfType<FNiagaraStackDragOperation>())
+	if (DragDropOperation->IsOfType<FNiagaraStackDragOperation>())
 	{
 		TSharedPtr<FNiagaraStackDragOperation> InputDragDropOperation = StaticCastSharedPtr<FNiagaraStackDragOperation>(DragDropOperation);
 		TSharedPtr<FNiagaraParameterAction> Action = StaticCastSharedPtr<FNiagaraParameterAction>(InputDragDropOperation->GetAction());
-		if (FNiagaraStackGraphUtilities::ParameterAllowedInExecutionCategory(Action->GetParameter().GetName(), ModuleItem->GetExecutionCategoryName()))
-		{
-			return true;
-		}
+		return Action.IsValid() && ModuleItem->CanAddInput(Action->GetParameter());
 	}
 
 	return false;
diff --git a/Engine/Plugins/FX/Niagara/Source/NiagaraEditorWidgets/Private/Stack/SNiagaraStackParameterStoreEntryValue.cpp b/Engine/Plugins/FX/Niagara/Source/NiagaraEditorWidgets/Private/Stack/SNiagaraStackParameterStoreEntryValue.cpp
--- a/Engine/Plugins/FX/Niagara/Source/NiagaraEditorWidgets/Private/Stack/SNiagaraStackParameterStoreEntryValue.cpp
+++ b/Engine/Plugins/FX/Niagara/Source/NiagaraEditorWidgets/Private/Stack/SNiagaraStackParameterStoreEntryValue.cpp
@@ -120,13 +120,16 @@ void SNiagaraStackParameterStoreEntryValue::Construct(const FArguments& InArgs,
 
 FReply SNiagaraStackParameterStoreEntryValue::DeleteClicked()
 {
-	StackEntry->Delete();
 	// toast notification
 	FNotificationInfo Info(FText::Format(LOCTEXT("NiagaraDeletedUserParameter", "System exposed parameter was deleted.\n{0}\n(All links to inner variables were invalidated in the process.)"), StackEntry->GetDisplayName()));
 	Info.ExpireDuration = 5.0f;
 	Info.bFireAndForget = true;
 	Info.Image = FCoreStyle::Get().GetBrush(TEXT("MessageLog.Info"));
 	FSlateNotificationManager::Get().AddNotification(Info);
+
+	// Delete after the notification is posted to prevent the entry from becoming invalidated before generating the message.
+	StackEntry->Delete();
+
 	return FReply::Handled();
 }
 
diff --git a/Engine/Plugins/FX/Niagara/Source/NiagaraShader/Private/NiagaraShader.cpp b/Engine/Plugins/FX/Niagara/Source/NiagaraShader/Private/NiagaraShader.cpp
--- a/Engine/Plugins/FX/Niagara/Source/NiagaraShader/Private/NiagaraShader.cpp
+++ b/Engine/Plugins/FX/Niagara/Source/NiagaraShader/Private/NiagaraShader.cpp
@@ -319,6 +319,12 @@ void FNiagaraShaderType::AddReferencedUniformBufferIncludes(FShaderCompilerEnvir
 
 	FString& GeneratedUniformBuffersInclude = OutEnvironment.IncludeVirtualPathToContentsMap.FindOrAdd("/Engine/Generated/GeneratedUniformBuffers.ush");
 	GeneratedUniformBuffersInclude.Append(UniformBufferIncludes);
+
+	ERHIFeatureLevel::Type MaxFeatureLevel = GetMaxSupportedFeatureLevel(Platform);
+	if (MaxFeatureLevel >= ERHIFeatureLevel::SM4)
+	{
+		OutEnvironment.SetDefine(TEXT("PLATFORM_SUPPORTS_SRV_UB"), TEXT("1"));
+	}
 }
 
 
diff --git a/Engine/Plugins/FX/Niagara/Source/NiagaraShader/Public/NiagaraShader.h b/Engine/Plugins/FX/Niagara/Source/NiagaraShader/Public/NiagaraShader.h
--- a/Engine/Plugins/FX/Niagara/Source/NiagaraShader/Public/NiagaraShader.h
+++ b/Engine/Plugins/FX/Niagara/Source/NiagaraShader/Public/NiagaraShader.h
@@ -58,6 +58,11 @@ public:
 	static void ModifyCompilationEnvironment(EShaderPlatform Platform, const FNiagaraShaderScript*  Script, FShaderCompilerEnvironment& OutEnvironment)
 	{
 	}
+
+	static bool ValidateCompiledResult(EShaderPlatform Platform, const FShaderParameterMap& ParameterMap, TArray<FString>& OutError)
+	{
+		return true;
+	}
 	
 	void SetDataInterfaceParameterInfo(const TArray<FNiagaraDataInterfaceGPUParamInfo>& InDIParamInfo);
 
diff --git a/Engine/Plugins/FX/Niagara/Source/NiagaraVertexFactories/Private/NiagaraMeshVertexFactory.cpp b/Engine/Plugins/FX/Niagara/Source/NiagaraVertexFactories/Private/NiagaraMeshVertexFactory.cpp
--- a/Engine/Plugins/FX/Niagara/Source/NiagaraVertexFactories/Private/NiagaraMeshVertexFactory.cpp
+++ b/Engine/Plugins/FX/Niagara/Source/NiagaraVertexFactories/Private/NiagaraMeshVertexFactory.cpp
@@ -193,7 +193,7 @@ FShaderResourceViewRHIParamRef FNiagaraMeshVertexFactory::GetPreviousTransformBu
 
 bool FNiagaraMeshVertexFactory::ShouldCompilePermutation(EShaderPlatform Platform, const class FMaterial* Material, const class FShaderType* ShaderType)
 {
-	return (!IsMobilePlatform(Platform) && !IsSwitchPlatform(Platform) && (Material->IsUsedWithNiagaraMeshParticles() || Material->IsSpecialEngineMaterial()));
+	return (!IsMobilePlatform(Platform) && !IsSwitchPlatform(Platform) && Platform != SP_OPENGL_SM4 && (Material->IsUsedWithNiagaraMeshParticles() || Material->IsSpecialEngineMaterial()));
 }
 
 void FNiagaraMeshVertexFactory::SetData(const FStaticMeshDataType& InData)
diff --git a/Engine/Plugins/FX/Niagara/Source/NiagaraVertexFactories/Private/NiagaraRibbonVertexFactory.cpp b/Engine/Plugins/FX/Niagara/Source/NiagaraVertexFactories/Private/NiagaraRibbonVertexFactory.cpp
--- a/Engine/Plugins/FX/Niagara/Source/NiagaraVertexFactories/Private/NiagaraRibbonVertexFactory.cpp
+++ b/Engine/Plugins/FX/Niagara/Source/NiagaraVertexFactories/Private/NiagaraRibbonVertexFactory.cpp
@@ -168,7 +168,7 @@ static TGlobalResource<FNiagaraRibbonVertexDeclaration> GNiagaraRibbonVertexDecl
 
 bool FNiagaraRibbonVertexFactory::ShouldCompilePermutation(EShaderPlatform Platform, const class FMaterial* Material, const class FShaderType* ShaderType)
 {
-	return (!IsMobilePlatform(Platform) && !IsSwitchPlatform(Platform) && (Material->IsUsedWithNiagaraRibbons() || Material->IsSpecialEngineMaterial()));
+	return (!IsMobilePlatform(Platform) && !IsSwitchPlatform(Platform) && Platform != SP_OPENGL_SM4 && (Material->IsUsedWithNiagaraRibbons() || Material->IsSpecialEngineMaterial()));
 }
 
 /**
diff --git a/Engine/Plugins/FX/Niagara/Source/NiagaraVertexFactories/Private/NiagaraSpriteVertexFactory.cpp b/Engine/Plugins/FX/Niagara/Source/NiagaraVertexFactories/Private/NiagaraSpriteVertexFactory.cpp
--- a/Engine/Plugins/FX/Niagara/Source/NiagaraVertexFactories/Private/NiagaraSpriteVertexFactory.cpp
+++ b/Engine/Plugins/FX/Niagara/Source/NiagaraVertexFactories/Private/NiagaraSpriteVertexFactory.cpp
@@ -225,7 +225,7 @@ inline TGlobalResource<FNiagaraSpriteVertexDeclaration>& GetNiagaraSpriteVertexD
 
 bool FNiagaraSpriteVertexFactory::ShouldCompilePermutation(EShaderPlatform Platform, const class FMaterial* Material, const class FShaderType* ShaderType)
 {
-	return (!IsSwitchPlatform(Platform) && (Material->IsUsedWithNiagaraSprites() || Material->IsSpecialEngineMaterial()));
+	return (!IsSwitchPlatform(Platform) && Platform != SP_OPENGL_SM4 && (Material->IsUsedWithNiagaraSprites() || Material->IsSpecialEngineMaterial()));
 }
 
 /**
@@ -250,7 +250,7 @@ void FNiagaraSpriteVertexFactory::InitRHI()
 
 void FNiagaraSpriteVertexFactory::InitStreams()
 {
-    const bool bInstanced = GRHISupportsInstancing;
+	const bool bInstanced = GRHISupportsInstancing;
 
 	check(Streams.Num() == 0);
 	if(bInstanced) 
