Trying to fix linux build
#tests n/a

diff --git a/Engine/Plugins/FX/Niagara/Source/NiagaraEditor/Private/TNiagaraGraphPinEditableName.h b/Engine/Plugins/FX/Niagara/Source/NiagaraEditor/Private/TNiagaraGraphPinEditableName.h
--- a/Engine/Plugins/FX/Niagara/Source/NiagaraEditor/Private/TNiagaraGraphPinEditableName.h
+++ b/Engine/Plugins/FX/Niagara/Source/NiagaraEditor/Private/TNiagaraGraphPinEditableName.h
@@ -48,7 +48,7 @@ protected:
 
 	void OnTextCommitted(const FText& InText, ETextCommit::Type InCommitType)
 	{
-		if (!GraphPinObj->PinName.ToString().Equals(InText.ToString(), ESearchCase::CaseSensitive))
+		if (!this->GraphPinObj->PinName.ToString().Equals(InText.ToString(), ESearchCase::CaseSensitive))
 		{
 			UNiagaraNode* ParentNode = Cast<UNiagaraNode>(this->GraphPinObj->GetOwningNode());
 			if (ParentNode != nullptr)
