Niagara - Make the skeletal mesh data interface detail customization handle object lifetime and delegates more safely.

diff --git a/Engine/Plugins/FX/Niagara/Source/NiagaraEditorWidgets/Private/DetailCustomizations/NiagaraDataInterfaceSkeletalMeshDetails.cpp b/Engine/Plugins/FX/Niagara/Source/NiagaraEditorWidgets/Private/DetailCustomizations/NiagaraDataInterfaceSkeletalMeshDetails.cpp
--- a/Engine/Plugins/FX/Niagara/Source/NiagaraEditorWidgets/Private/DetailCustomizations/NiagaraDataInterfaceSkeletalMeshDetails.cpp
+++ b/Engine/Plugins/FX/Niagara/Source/NiagaraEditorWidgets/Private/DetailCustomizations/NiagaraDataInterfaceSkeletalMeshDetails.cpp
@@ -12,28 +12,28 @@
  #define LOCTEXT_NAMESPACE "FNiagaraDataInterfaceSkeletalMeshDetails"
 
 void FNiagaraDataInterfaceSkeletalMeshDetails::CustomizeDetails(IDetailLayoutBuilder& DetailBuilder)
- {
-	 
+{
 	 LayoutBuilder = &DetailBuilder;
 	 static const FName MeshCategoryName = TEXT("Mesh");
 
 	 TArray<TWeakObjectPtr<UObject>> SelectedObjects;
 	 DetailBuilder.GetObjectsBeingCustomized(SelectedObjects);
-	 check(SelectedObjects.Num() == 1);
-	 MeshInterface = Cast<UNiagaraDataInterfaceSkeletalMesh>(SelectedObjects[0].Get());
-	 check(MeshInterface);
-	 if (MeshInterface != nullptr)
+	 if(SelectedObjects.Num() != 1 || SelectedObjects[0]->IsA<UNiagaraDataInterfaceSkeletalMesh>() == false)
 	 {
-		 MeshInterface->OnChanged().RemoveAll(this);
-		 MeshInterface->OnChanged().AddRaw(this, &FNiagaraDataInterfaceSkeletalMeshDetails::OnInterfaceChanged);
+	 	return;
 	 }
+
+	 MeshInterface = CastChecked<UNiagaraDataInterfaceSkeletalMesh>(SelectedObjects[0].Get());
+	 MeshInterface->OnChanged().RemoveAll(this);
+	 MeshInterface->OnChanged().AddSP(this, &FNiagaraDataInterfaceSkeletalMeshDetails::OnInterfaceChanged);
+
 	 TWeakObjectPtr<USceneComponent> SceneComponent;
 	 USkeletalMeshComponent* FoundSkelComp = nullptr;
-	 MeshObject = UNiagaraDataInterfaceSkeletalMesh::GetSkeletalMeshHelper(MeshInterface, Cast<UNiagaraComponent>(MeshInterface->GetOuter()), SceneComponent, FoundSkelComp);
-	 if (MeshObject != nullptr)
+	 MeshObject = UNiagaraDataInterfaceSkeletalMesh::GetSkeletalMeshHelper(MeshInterface.Get(), Cast<UNiagaraComponent>(MeshInterface->GetOuter()), SceneComponent, FoundSkelComp);
+	 if (MeshObject.IsValid())
 	 {
 		 MeshObject->GetOnMeshChanged().RemoveAll(this);
-		 MeshObject->GetOnMeshChanged().AddRaw(this, &FNiagaraDataInterfaceSkeletalMeshDetails::OnDataChanged);
+		 MeshObject->GetOnMeshChanged().AddSP(this, &FNiagaraDataInterfaceSkeletalMeshDetails::OnDataChanged);
 	 }
 
 	 IDetailCategoryBuilder& MeshCategory = DetailBuilder.EditCategory(MeshCategoryName, LOCTEXT("Mesh", "Mesh"));
@@ -69,14 +69,14 @@ void FNiagaraDataInterfaceSkeletalMeshDetails::CustomizeDetails(IDetailLayoutBui
 	 // Rebuild the data changed listener
 	 TWeakObjectPtr<USceneComponent> SceneComponent;
 	 USkeletalMeshComponent* FoundSkelComp = nullptr;
-	 if (MeshObject != nullptr)
+	 if (MeshObject.IsValid())
 	 {
 		 MeshObject->GetOnMeshChanged().RemoveAll(this);
 	 }
-	 MeshObject = UNiagaraDataInterfaceSkeletalMesh::GetSkeletalMeshHelper(MeshInterface, Cast<UNiagaraComponent>(MeshInterface->GetOuter()), SceneComponent, FoundSkelComp);
-	 if (MeshObject != nullptr)
+	 MeshObject = UNiagaraDataInterfaceSkeletalMesh::GetSkeletalMeshHelper(MeshInterface.Get(), Cast<UNiagaraComponent>(MeshInterface->GetOuter()), SceneComponent, FoundSkelComp);
+	 if (MeshObject.IsValid())
 	 {
-		 MeshObject->GetOnMeshChanged().AddRaw(this, &FNiagaraDataInterfaceSkeletalMeshDetails::OnDataChanged);
+		 MeshObject->GetOnMeshChanged().AddSP(this, &FNiagaraDataInterfaceSkeletalMeshDetails::OnDataChanged);
 	 }
 	 OnDataChanged();
  }
@@ -90,11 +90,11 @@ void FNiagaraDataInterfaceSkeletalMeshDetails::CustomizeDetails(IDetailLayoutBui
  TArray<TSharedPtr<FName>> FNiagaraDataInterfaceSkeletalMeshDetails::GenerateSourceArray()
  {
 	TArray<TSharedPtr<FName>> SourceArray;
-	if (MeshInterface != nullptr)
+	if (MeshInterface.IsValid())
 	{
 			TWeakObjectPtr<USceneComponent> SceneComponent;
 			USkeletalMeshComponent* FoundSkelComp = nullptr;
-			USkeletalMesh* Mesh = UNiagaraDataInterfaceSkeletalMesh::GetSkeletalMeshHelper(MeshInterface, Cast<UNiagaraComponent>(MeshInterface->GetOuter()), SceneComponent, FoundSkelComp);
+			USkeletalMesh* Mesh = UNiagaraDataInterfaceSkeletalMesh::GetSkeletalMeshHelper(MeshInterface.Get(), Cast<UNiagaraComponent>(MeshInterface->GetOuter()), SceneComponent, FoundSkelComp);
 
 			if (Mesh != nullptr)
 			{
@@ -110,11 +110,11 @@ void FNiagaraDataInterfaceSkeletalMeshDetails::CustomizeDetails(IDetailLayoutBui
 
  FNiagaraDataInterfaceSkeletalMeshDetails::~FNiagaraDataInterfaceSkeletalMeshDetails()
  {
-	 if (MeshInterface)
+	 if (MeshInterface.IsValid())
 	 {
 		 MeshInterface->OnChanged().RemoveAll(this);
 	 }
-	 if (MeshObject)
+	 if (MeshObject.IsValid())
 	 {
 		MeshObject->GetOnMeshChanged().RemoveAll(this);
 	 }
diff --git a/Engine/Plugins/FX/Niagara/Source/NiagaraEditorWidgets/Private/DetailCustomizations/NiagaraDataInterfaceSkeletalMeshDetails.h b/Engine/Plugins/FX/Niagara/Source/NiagaraEditorWidgets/Private/DetailCustomizations/NiagaraDataInterfaceSkeletalMeshDetails.h
--- a/Engine/Plugins/FX/Niagara/Source/NiagaraEditorWidgets/Private/DetailCustomizations/NiagaraDataInterfaceSkeletalMeshDetails.h
+++ b/Engine/Plugins/FX/Niagara/Source/NiagaraEditorWidgets/Private/DetailCustomizations/NiagaraDataInterfaceSkeletalMeshDetails.h
@@ -31,6 +31,6 @@ private:
 private:
 	TSharedPtr<FNiagaraDetailSourcedArrayBuilder> RegionsBuilder;
 	IDetailLayoutBuilder* LayoutBuilder;
-	UNiagaraDataInterfaceSkeletalMesh*  MeshInterface;
-	USkeletalMesh* MeshObject;
+	TWeakObjectPtr<UNiagaraDataInterfaceSkeletalMesh>  MeshInterface;
+	TWeakObjectPtr<USkeletalMesh> MeshObject;
 };
