First draft of a textuer 2d data interface. Currently only works on GPU as we don't expose the pixels to the CPU.

#tests new functional test added

diff --git a/Engine/Plugins/FX/Niagara/Source/Niagara/Classes/NiagaraDataInterfaceTexture2D.h b/Engine/Plugins/FX/Niagara/Source/Niagara/Classes/NiagaraDataInterfaceTexture2D.h
new file mode 100644
--- /dev/null
+++ b/Engine/Plugins/FX/Niagara/Source/Niagara/Classes/NiagaraDataInterfaceTexture2D.h
@@ -0,0 +1,52 @@
+// Copyright 1998-2018 Epic Games, Inc. All Rights Reserved.
+
+#pragma once
+
+#include "NiagaraDataInterface.h"
+#include "NiagaraCommon.h"
+#include "VectorVM.h"
+#include "NiagaraDataInterfaceTexture2D.generated.h"
+
+/** Data Interface allowing sampling of a texture */
+UCLASS(EditInlineNew, Category = "Texture", meta = (DisplayName = "Texture Sample 2D"))
+class NIAGARA_API UNiagaraDataInterfaceTexture2D : public UNiagaraDataInterface
+{
+	GENERATED_UCLASS_BODY()
+		bool bGPUBufferDirty;
+public:
+
+	UPROPERTY(EditAnywhere, Category = "Texture")
+	UTexture2D *Texture;
+
+	//UObject Interface
+	virtual void PostInitProperties()override;
+	virtual void PostLoad() override;
+#if WITH_EDITOR
+	virtual void PostEditChangeProperty(struct FPropertyChangedEvent& PropertyChangedEvent) override;
+#endif
+	//UObject Interface End
+
+	//UNiagaraDataInterface Interface
+	virtual void GetFunctions(TArray<FNiagaraFunctionSignature>& OutFunctions)override;
+	virtual void GetVMExternalFunction(const FVMExternalFunctionBindingInfo& BindingInfo, void* InstanceData, FVMExternalFunction &OutFunc) override;
+	virtual bool CanExecuteOnTarget(ENiagaraSimTarget Target)const override { return Target==ENiagaraSimTarget::GPUComputeSim; }
+	//UNiagaraDataInterface Interface End
+
+	template<typename XType, typename YType>
+	void SampleTexture(FVectorVMContext& Context);
+
+	virtual bool Equals(const UNiagaraDataInterface* Other) const override;
+
+	// GPU sim functionality
+	virtual bool GetFunctionHLSL(const FName&  DefinitionFunctionName, FString InstanceFunctionName, FNiagaraDataInterfaceGPUParamInfo& ParamInfo, FString& OutHLSL) override;
+	virtual void GetParameterDefinitionHLSL(FNiagaraDataInterfaceGPUParamInfo& ParamInfo, FString& OutHLSL) override;
+	virtual FNiagaraDataInterfaceParametersCS* ConstructComputeParameters()const override;
+
+	//FRWBuffer& GetGPUBuffer();
+	static const FString TextureName;
+	static const FString SamplerName;
+protected:
+	virtual bool CopyToInternal(UNiagaraDataInterface* Destination) const override;
+
+	static const FName SampleTextureName;
+};
diff --git a/Engine/Plugins/FX/Niagara/Source/Niagara/Private/NiagaraDataInterfaceTexture2D.cpp b/Engine/Plugins/FX/Niagara/Source/Niagara/Private/NiagaraDataInterfaceTexture2D.cpp
new file mode 100644
--- /dev/null
+++ b/Engine/Plugins/FX/Niagara/Source/Niagara/Private/NiagaraDataInterfaceTexture2D.cpp
@@ -0,0 +1,188 @@
+// Copyright 1998-2018 Epic Games, Inc. All Rights Reserved.
+
+#include "NiagaraDataInterfaceTexture2D.h"
+#include "NiagaraShader.h"
+#include "ShaderParameterUtils.h"
+
+const FName UNiagaraDataInterfaceTexture2D::SampleTextureName(TEXT("SampleTexture2D"));
+const FString UNiagaraDataInterfaceTexture2D::TextureName(TEXT("Texture_"));
+const FString UNiagaraDataInterfaceTexture2D::SamplerName(TEXT("Sampler_"));
+
+UNiagaraDataInterfaceTexture2D::UNiagaraDataInterfaceTexture2D(FObjectInitializer const& ObjectInitializer)
+	: Super(ObjectInitializer)
+	, Texture(nullptr)
+{
+
+}
+
+void UNiagaraDataInterfaceTexture2D::PostInitProperties()
+{
+	Super::PostInitProperties();
+
+	if (HasAnyFlags(RF_ClassDefaultObject))
+	{
+		FNiagaraTypeRegistry::Register(FNiagaraTypeDefinition(GetClass()), true, false, false);
+	}
+}
+
+void UNiagaraDataInterfaceTexture2D::PostLoad()
+{
+	Super::PostLoad();
+}
+
+#if WITH_EDITOR
+
+void UNiagaraDataInterfaceTexture2D::PostEditChangeProperty(struct FPropertyChangedEvent& PropertyChangedEvent)
+{
+	Super::PostEditChangeProperty(PropertyChangedEvent);
+
+	/*if (PropertyChangedEvent.Property && PropertyChangedEvent.Property->GetFName() == GET_MEMBER_NAME_CHECKED(UNiagaraDataInterfaceTexture2D, Seed))
+	{
+		InitNoiseLUT();
+	}*/
+}
+
+#endif
+
+bool UNiagaraDataInterfaceTexture2D::CopyToInternal(UNiagaraDataInterface* Destination) const
+{
+	if (!Super::CopyToInternal(Destination))
+	{
+		return false;
+	}
+	UNiagaraDataInterfaceTexture2D* DestinationTexture = CastChecked<UNiagaraDataInterfaceTexture2D>(Destination);
+	DestinationTexture->Texture = Texture;
+
+	return true;
+}
+
+bool UNiagaraDataInterfaceTexture2D::Equals(const UNiagaraDataInterface* Other) const
+{
+	if (!Super::Equals(Other))
+	{
+		return false;
+	}
+	const UNiagaraDataInterfaceTexture2D* OtherTexture = CastChecked<const UNiagaraDataInterfaceTexture2D>(Other);
+	return OtherTexture->Texture == Texture;
+}
+
+void UNiagaraDataInterfaceTexture2D::GetFunctions(TArray<FNiagaraFunctionSignature>& OutFunctions)
+{
+	FNiagaraFunctionSignature Sig;
+	Sig.Name = SampleTextureName;
+	Sig.bMemberFunction = true;
+	Sig.bRequiresContext = false;
+	Sig.Inputs.Add(FNiagaraVariable(FNiagaraTypeDefinition(GetClass()), TEXT("Texture")));
+	Sig.Inputs.Add(FNiagaraVariable(FNiagaraTypeDefinition::GetVec2Def(), TEXT("UV")));
+	Sig.Outputs.Add(FNiagaraVariable(FNiagaraTypeDefinition::GetVec4Def(), TEXT("Value")));
+	//Sig.Owner = *GetName();
+
+	OutFunctions.Add(Sig);
+}
+
+DEFINE_NDI_RAW_FUNC_BINDER(UNiagaraDataInterfaceTexture2D, SampleTexture);
+void UNiagaraDataInterfaceTexture2D::GetVMExternalFunction(const FVMExternalFunctionBindingInfo& BindingInfo, void* InstanceData, FVMExternalFunction &OutFunc)
+{
+	check(BindingInfo.Name == SampleTextureName);
+	check(BindingInfo.GetNumInputs() == 2 && BindingInfo.GetNumOutputs() == 4);
+	TNDIParamBinder<0, float, TNDIParamBinder<1, float, NDI_RAW_FUNC_BINDER(UNiagaraDataInterfaceTexture2D, SampleTexture)>>::Bind(this, BindingInfo, InstanceData, OutFunc);
+}
+
+
+template<typename ComponentType, typename XType, typename YType>
+void SampleData(void *RawData, XType XParam, YType YParam, uint32 NumComponents, FVector2D &TexSize, ComponentType *Out)
+{
+	for (uint32 C = 0; C < NumComponents; C++)
+	{
+		Out[C] = static_cast<ComponentType*>(RawData)[Y*(int32)TexSize.X + X];
+	}
+}
+
+template<typename XType, typename YType>
+void UNiagaraDataInterfaceTexture2D::SampleTexture(FVectorVMContext& Context)
+{
+	XType XParam(Context);
+	YType YParam(Context);
+	FRegisterHandler<float> OutSampleR(Context);
+	FRegisterHandler<float> OutSampleG(Context);
+	FRegisterHandler<float> OutSampleB(Context);
+	FRegisterHandler<float> OutSampleA(Context);
+
+	int32 X = XParam.GetAndAdvance();
+	int32 Y = YParam.GetAndAdvance();
+
+	*OutSampleR.GetDestAndAdvance() = 1.0;
+	*OutSampleG.GetDestAndAdvance() = 1.0;
+	*OutSampleB.GetDestAndAdvance() = 1.0;
+	*OutSampleA.GetDestAndAdvance() = 1.0;
+}
+
+bool UNiagaraDataInterfaceTexture2D::GetFunctionHLSL(const FName& DefinitionFunctionName, FString InstanceFunctionName, FNiagaraDataInterfaceGPUParamInfo& ParamInfo, FString& OutHLSL)
+{
+	FString HLSLTextureName = TextureName + ParamInfo.DataInterfaceHLSLSymbol;
+	FString HLSLSamplerName = SamplerName + ParamInfo.DataInterfaceHLSLSymbol;
+	OutHLSL += TEXT("void ") + InstanceFunctionName + TEXT("(in float2 In_UV, out float4 Out_Value) \n{\n");
+	OutHLSL += TEXT("\t Out_Value = ") + HLSLTextureName + TEXT(".SampleLevel(") + HLSLSamplerName + TEXT(", In_UV, 0);\n");
+	OutHLSL += TEXT("\n}\n");
+	return true;
+}
+
+void UNiagaraDataInterfaceTexture2D::GetParameterDefinitionHLSL(FNiagaraDataInterfaceGPUParamInfo& ParamInfo, FString& OutHLSL)
+{
+	FString HLSLTextureName = TextureName + ParamInfo.DataInterfaceHLSLSymbol;
+	FString HLSLSamplerName = SamplerName + ParamInfo.DataInterfaceHLSLSymbol;
+	OutHLSL += TEXT("Texture2D ") + HLSLTextureName + TEXT(";\n");
+	OutHLSL += TEXT("SamplerState ") + HLSLSamplerName + TEXT(";\n");
+}
+
+
+
+struct FNiagaraDataInterfaceParametersCS_Texture : public FNiagaraDataInterfaceParametersCS
+{
+	virtual void Bind(const FNiagaraDataInterfaceParamRef& ParamRef, const class FShaderParameterMap& ParameterMap) override
+	{
+		TextureParam.Bind(ParameterMap, *(UNiagaraDataInterfaceTexture2D::TextureName + ParamRef.ParameterInfo.DataInterfaceHLSLSymbol));
+		SamplerParam.Bind(ParameterMap, *(UNiagaraDataInterfaceTexture2D::SamplerName + ParamRef.ParameterInfo.DataInterfaceHLSLSymbol));
+		ensure(TextureParam.IsBound());
+		ensure(SamplerParam.IsBound());
+	}
+
+	virtual void Serialize(FArchive& Ar)override
+	{
+		Ar << TextureParam;
+		Ar << SamplerParam;
+	}
+
+	virtual void Set(FRHICommandList& RHICmdList, FNiagaraShader* Shader, class UNiagaraDataInterface* DataInterface) const override
+	{
+		check(IsInRenderingThread());
+
+		FComputeShaderRHIParamRef ComputeShaderRHI = Shader->GetComputeShader();
+		UNiagaraDataInterfaceTexture2D* TextureDI = CastChecked<UNiagaraDataInterfaceTexture2D>(DataInterface);
+		UTexture *Texture = TextureDI->Texture;
+		if (!Texture)
+		{
+			return;
+		}
+		FTextureRHIParamRef TextureRHI = Texture->Resource->TextureRHI;
+		SetTextureParameter(
+			RHICmdList,
+			ComputeShaderRHI,
+			TextureParam,
+			SamplerParam,
+			TStaticSamplerState<SF_Trilinear, AM_Clamp, AM_Clamp, AM_Clamp>::GetRHI(),
+			TextureRHI
+		);
+	}
+
+
+private:
+
+	FShaderResourceParameter TextureParam;
+	FShaderResourceParameter SamplerParam;
+};
+
+FNiagaraDataInterfaceParametersCS* UNiagaraDataInterfaceTexture2D::ConstructComputeParameters()const
+{
+	return new FNiagaraDataInterfaceParametersCS_Texture();
+}
\ No newline at end of file
