Adding a basic CPU texture path as well for the data interface texture. Note that we're not doing any interp on the final pixel for CPU.

#tests n/a

diff --git a/Engine/Plugins/FX/Niagara/Source/Niagara/Classes/NiagaraDataInterfaceTexture2D.h b/Engine/Plugins/FX/Niagara/Source/Niagara/Classes/NiagaraDataInterfaceTexture2D.h
--- a/Engine/Plugins/FX/Niagara/Source/Niagara/Classes/NiagaraDataInterfaceTexture2D.h
+++ b/Engine/Plugins/FX/Niagara/Source/Niagara/Classes/NiagaraDataInterfaceTexture2D.h
@@ -16,7 +16,7 @@ class NIAGARA_API UNiagaraDataInterfaceTexture2D : public UNiagaraDataInterface
 public:
 
 	UPROPERTY(EditAnywhere, Category = "Texture")
-	UTexture2D *Texture;
+	UTexture2D* Texture;
 
 	//UObject Interface
 	virtual void PostInitProperties()override;
@@ -46,7 +46,15 @@ public:
 	static const FString TextureName;
 	static const FString SamplerName;
 protected:
+#if WITH_EDITOR
+	bool CopyTextureToCPUBackup(UTexture* SourceTexture, TArray<uint8>& TargetBuffer);
+	void CopyTextureData(const uint8* Source, uint8* Dest, uint32 SizeX, uint32 SizeY, uint32 BytesPerPixel, uint32 SourceStride, uint32 DestStride);
+
+#endif
 	virtual bool CopyToInternal(UNiagaraDataInterface* Destination) const override;
 
+	UPROPERTY()
+	TArray<uint8> CPUTextureData;
+
 	static const FName SampleTextureName;
 };
diff --git a/Engine/Plugins/FX/Niagara/Source/Niagara/Private/NiagaraDataInterfaceTexture2D.cpp b/Engine/Plugins/FX/Niagara/Source/Niagara/Private/NiagaraDataInterfaceTexture2D.cpp
--- a/Engine/Plugins/FX/Niagara/Source/Niagara/Private/NiagaraDataInterfaceTexture2D.cpp
+++ b/Engine/Plugins/FX/Niagara/Source/Niagara/Private/NiagaraDataInterfaceTexture2D.cpp
@@ -4,6 +4,9 @@
 #include "NiagaraShader.h"
 #include "ShaderParameterUtils.h"
 
+
+#define LOCTEXT_NAMESPACE "UNiagaraDataInterfaceTexture2D"
+
 const FName UNiagaraDataInterfaceTexture2D::SampleTextureName(TEXT("SampleTexture2D"));
 const FString UNiagaraDataInterfaceTexture2D::TextureName(TEXT("Texture_"));
 const FString UNiagaraDataInterfaceTexture2D::SamplerName(TEXT("Sampler_"));
@@ -36,10 +39,58 @@ void UNiagaraDataInterfaceTexture2D::PostEditChangeProperty(struct FPropertyChan
 {
 	Super::PostEditChangeProperty(PropertyChangedEvent);
 
-	/*if (PropertyChangedEvent.Property && PropertyChangedEvent.Property->GetFName() == GET_MEMBER_NAME_CHECKED(UNiagaraDataInterfaceTexture2D, Seed))
+	if (PropertyChangedEvent.Property && PropertyChangedEvent.Property->GetFName() == GET_MEMBER_NAME_CHECKED(UNiagaraDataInterfaceTexture2D, Texture))
+	{
+		CopyTextureToCPUBackup(Texture, CPUTextureData);
+	}
+}
+
+bool UNiagaraDataInterfaceTexture2D::CopyTextureToCPUBackup(UTexture* SourceTexture, TArray<uint8>& TargetBuffer)
+{
+	if (SourceTexture == nullptr)
+	{
+		TargetBuffer.Empty();
+		return true;
+	}
+
+	FTextureSource& SourceData = SourceTexture->Source;
+	FIntPoint SourceSize = FIntPoint(SourceData.GetSizeX(), SourceData.GetSizeY());
+
+	{
+		const int32 BytesPerPixel = 4;
+		TargetBuffer.Empty();
+		TargetBuffer.AddZeroed(SourceSize.X * SourceSize.Y * BytesPerPixel);
+	}
+
+	
+	if (SourceData.GetFormat() == TSF_BGRA8)
+	{
+		uint32 BytesPerPixel = SourceData.GetBytesPerPixel();
+		uint8* OffsetSource = SourceData.LockMip(0);
+		uint8* OffsetDest = TargetBuffer.GetData();
+		CopyTextureData(OffsetSource, OffsetDest, SourceSize.X, SourceSize.Y, BytesPerPixel, SourceData.GetSizeX() * BytesPerPixel, SourceSize.X * BytesPerPixel);
+		SourceData.UnlockMip(0);
+	}
+	else 
+	{
+		UE_LOG(LogNiagara, Error, TEXT("Texture %s is not BGRA8, which isn't supported in data interfaces yet"), *SourceTexture->GetName());
+	}
+
+	return true;
+}
+
+void UNiagaraDataInterfaceTexture2D::CopyTextureData(const uint8* Source, uint8* Dest, uint32 SizeX, uint32 SizeY, uint32 BytesPerPixel, uint32 SourceStride, uint32 DestStride)
+{
+	const uint32 NumBytesPerRow = SizeX * BytesPerPixel;
+
+	for (uint32 Y = 0; Y < SizeY; ++Y)
 	{
-		InitNoiseLUT();
-	}*/
+		FMemory::Memcpy(
+			Dest + (DestStride * Y),
+			Source + (SourceStride * Y),
+			NumBytesPerRow
+		);
+	}
 }
 
 #endif
@@ -52,6 +103,7 @@ bool UNiagaraDataInterfaceTexture2D::CopyToInternal(UNiagaraDataInterface* Desti
 	}
 	UNiagaraDataInterfaceTexture2D* DestinationTexture = CastChecked<UNiagaraDataInterfaceTexture2D>(Destination);
 	DestinationTexture->Texture = Texture;
+	DestinationTexture->CPUTextureData = CPUTextureData;
 
 	return true;
 }
@@ -74,6 +126,7 @@ void UNiagaraDataInterfaceTexture2D::GetFunctions(TArray<FNiagaraFunctionSignatu
 	Sig.bRequiresContext = false;
 	Sig.Inputs.Add(FNiagaraVariable(FNiagaraTypeDefinition(GetClass()), TEXT("Texture")));
 	Sig.Inputs.Add(FNiagaraVariable(FNiagaraTypeDefinition::GetVec2Def(), TEXT("UV")));
+	Sig.SetDescription(LOCTEXT("TextureSampleDesc", "Sample mip level 0 of the input 2d texture at the specified UV coordinates. The UV origin (0,0) is in the upper left hand corner of the image."));
 	Sig.Outputs.Add(FNiagaraVariable(FNiagaraTypeDefinition::GetVec4Def(), TEXT("Value")));
 	//Sig.Owner = *GetName();
 
@@ -88,16 +141,6 @@ void UNiagaraDataInterfaceTexture2D::GetVMExternalFunction(const FVMExternalFunc
 	TNDIParamBinder<0, float, TNDIParamBinder<1, float, NDI_RAW_FUNC_BINDER(UNiagaraDataInterfaceTexture2D, SampleTexture)>>::Bind(this, BindingInfo, InstanceData, OutFunc);
 }
 
-
-template<typename ComponentType, typename XType, typename YType>
-void SampleData(void *RawData, XType XParam, YType YParam, uint32 NumComponents, FVector2D &TexSize, ComponentType *Out)
-{
-	for (uint32 C = 0; C < NumComponents; C++)
-	{
-		Out[C] = static_cast<ComponentType*>(RawData)[Y*(int32)TexSize.X + X];
-	}
-}
-
 template<typename XType, typename YType>
 void UNiagaraDataInterfaceTexture2D::SampleTexture(FVectorVMContext& Context)
 {
@@ -108,13 +151,46 @@ void UNiagaraDataInterfaceTexture2D::SampleTexture(FVectorVMContext& Context)
 	FRegisterHandler<float> OutSampleB(Context);
 	FRegisterHandler<float> OutSampleA(Context);
 
-	int32 X = XParam.GetAndAdvance();
-	int32 Y = YParam.GetAndAdvance();
+	if (CPUTextureData.GetAllocatedSize() == 0 || Texture == nullptr)
+	{
+		for (int32 i = 0; i < Context.NumInstances; ++i)
+		{
+			float X = XParam.GetAndAdvance();
+			float Y = YParam.GetAndAdvance();
+			*OutSampleR.GetDestAndAdvance() = 1.0;
+			*OutSampleG.GetDestAndAdvance() = 0.0;
+			*OutSampleB.GetDestAndAdvance() = 1.0;
+			*OutSampleA.GetDestAndAdvance() = 1.0;
+		}
+	}
+	else
+	{
+		const int32 BytesPerPixel = 4;
+		int32 IntSizeX = Texture->GetSizeX();
+		int32 IntSizeY = Texture->GetSizeY();
+		float SizeX = (float)Texture->GetSizeX();
+		float SizeY = (float)Texture->GetSizeY();
+
+		for (int32 i = 0; i < Context.NumInstances; ++i)
+		{
+			float X = fmodf(XParam.GetAndAdvance() * SizeX, SizeX);
+			float Y = fmodf(YParam.GetAndAdvance() * SizeY, SizeY);
+			int32 XInt = trunc(X);
+			int32 YInt = trunc(Y);
+			int32 SampleIdx = YInt * IntSizeX * BytesPerPixel + XInt * BytesPerPixel;
+			ensure(CPUTextureData.Num() > SampleIdx);
+			uint8 B0 = CPUTextureData[SampleIdx + 0];
+			uint8 G0 = CPUTextureData[SampleIdx + 1];
+			uint8 R0 = CPUTextureData[SampleIdx + 2];
+			uint8 A0 = CPUTextureData[SampleIdx + 3];
+
+			*OutSampleR.GetDestAndAdvance() = ((float)R0) / 255.0f;
+			*OutSampleG.GetDestAndAdvance() = ((float)G0) / 255.0f;
+			*OutSampleB.GetDestAndAdvance() = ((float)B0) / 255.0f;
+			*OutSampleA.GetDestAndAdvance() = ((float)A0) / 255.0f;
+		}
+	}
 
-	*OutSampleR.GetDestAndAdvance() = 1.0;
-	*OutSampleG.GetDestAndAdvance() = 1.0;
-	*OutSampleB.GetDestAndAdvance() = 1.0;
-	*OutSampleA.GetDestAndAdvance() = 1.0;
 }
 
 bool UNiagaraDataInterfaceTexture2D::GetFunctionHLSL(const FName& DefinitionFunctionName, FString InstanceFunctionName, FNiagaraDataInterfaceGPUParamInfo& ParamInfo, FString& OutHLSL)
@@ -141,10 +217,20 @@ struct FNiagaraDataInterfaceParametersCS_Texture : public FNiagaraDataInterfaceP
 {
 	virtual void Bind(const FNiagaraDataInterfaceParamRef& ParamRef, const class FShaderParameterMap& ParameterMap) override
 	{
-		TextureParam.Bind(ParameterMap, *(UNiagaraDataInterfaceTexture2D::TextureName + ParamRef.ParameterInfo.DataInterfaceHLSLSymbol));
-		SamplerParam.Bind(ParameterMap, *(UNiagaraDataInterfaceTexture2D::SamplerName + ParamRef.ParameterInfo.DataInterfaceHLSLSymbol));
-		ensure(TextureParam.IsBound());
-		ensure(SamplerParam.IsBound());
+		FString TexName = UNiagaraDataInterfaceTexture2D::TextureName + ParamRef.ParameterInfo.DataInterfaceHLSLSymbol;
+		FString SampleName = (UNiagaraDataInterfaceTexture2D::SamplerName + ParamRef.ParameterInfo.DataInterfaceHLSLSymbol);
+		TextureParam.Bind(ParameterMap, *TexName);
+		SamplerParam.Bind(ParameterMap, *SampleName);
+		
+		if (!TextureParam.IsBound())
+		{
+			UE_LOG(LogNiagara, Warning, TEXT("Binding failed for FNiagaraDataInterfaceParametersCS_Texture Texture %s. Was it optimized out?"), *TexName)
+		}
+
+		if (!SamplerParam.IsBound())
+		{
+			UE_LOG(LogNiagara, Warning, TEXT("Binding failed for FNiagaraDataInterfaceParametersCS_Texture Sampler %s. Was it optimized out?"), *SampleName)
+		}
 	}
 
 	virtual void Serialize(FArchive& Ar)override
@@ -170,7 +256,7 @@ struct FNiagaraDataInterfaceParametersCS_Texture : public FNiagaraDataInterfaceP
 			ComputeShaderRHI,
 			TextureParam,
 			SamplerParam,
-			TStaticSamplerState<SF_Trilinear, AM_Clamp, AM_Clamp, AM_Clamp>::GetRHI(),
+			TStaticSamplerState<SF_Trilinear, AM_Wrap, AM_Wrap, AM_Wrap>::GetRHI(),
 			TextureRHI
 		);
 	}
