Trying to mimic GPU tiling when using negative uvs

#tests n/a

diff --git a/Engine/Plugins/FX/Niagara/Source/Niagara/Private/NiagaraDataInterfaceTexture2D.cpp b/Engine/Plugins/FX/Niagara/Source/Niagara/Private/NiagaraDataInterfaceTexture2D.cpp
--- a/Engine/Plugins/FX/Niagara/Source/Niagara/Private/NiagaraDataInterfaceTexture2D.cpp
+++ b/Engine/Plugins/FX/Niagara/Source/Niagara/Private/NiagaraDataInterfaceTexture2D.cpp
@@ -168,15 +168,21 @@ void UNiagaraDataInterfaceTexture2D::SampleTexture(FVectorVMContext& Context)
 		const int32 BytesPerPixel = 4;
 		int32 IntSizeX = Texture->GetSizeX();
 		int32 IntSizeY = Texture->GetSizeY();
-		float SizeX = (float)Texture->GetSizeX();
-		float SizeY = (float)Texture->GetSizeY();
+		float SizeX = (float)Texture->GetSizeX() - 1.0f;
+		float SizeY = (float)Texture->GetSizeY() - 1.0f;
 
 		for (int32 i = 0; i < Context.NumInstances; ++i)
 		{
-			float X = fmodf(XParam.GetAndAdvance() * SizeX, SizeX);
-			float Y = fmodf(YParam.GetAndAdvance() * SizeY, SizeY);
-			int32 XInt = trunc(X);
-			int32 YInt = trunc(Y);
+			float ParamXVal = XParam.GetAndAdvance();
+			float ParamYVal = YParam.GetAndAdvance();
+			float X = (fmodf(ParamXVal * SizeX, SizeX));
+			float Y = (fmodf(ParamYVal * SizeY, SizeY));
+
+			float XNorm = X < 0.0f ? (SizeX - fabsf(X)) : X;
+			float YNorm = Y < 0.0f ? (SizeY - fabsf(Y)) : Y;
+
+			int32 XInt = floorf(XNorm);
+			int32 YInt = floorf(YNorm);
 			int32 SampleIdx = YInt * IntSizeX * BytesPerPixel + XInt * BytesPerPixel;
 			ensure(CPUTextureData.Num() > SampleIdx);
 			uint8 B0 = CPUTextureData[SampleIdx + 0];
