Niagara - Change the serialization for the texture data interface so that we're serializing the cpu copy of the texture directly instead of letting the property system serialize it.

diff --git a/Engine/Plugins/FX/Niagara/Source/Niagara/Classes/NiagaraDataInterfaceTexture2D.h b/Engine/Plugins/FX/Niagara/Source/Niagara/Classes/NiagaraDataInterfaceTexture2D.h
--- a/Engine/Plugins/FX/Niagara/Source/Niagara/Classes/NiagaraDataInterfaceTexture2D.h
+++ b/Engine/Plugins/FX/Niagara/Source/Niagara/Classes/NiagaraDataInterfaceTexture2D.h
@@ -24,6 +24,7 @@ public:
 #if WITH_EDITOR
 	virtual void PostEditChangeProperty(struct FPropertyChangedEvent& PropertyChangedEvent) override;
 #endif
+	virtual void Serialize(FArchive& Ar) override;
 	//UObject Interface End
 
 	//UNiagaraDataInterface Interface
@@ -53,7 +54,6 @@ protected:
 #endif
 	virtual bool CopyToInternal(UNiagaraDataInterface* Destination) const override;
 
-	UPROPERTY()
 	TArray<uint8> CPUTextureData;
 
 	static const FName SampleTextureName;
diff --git a/Engine/Plugins/FX/Niagara/Source/Niagara/Private/NiagaraDataInterfaceTexture2D.cpp b/Engine/Plugins/FX/Niagara/Source/Niagara/Private/NiagaraDataInterfaceTexture2D.cpp
--- a/Engine/Plugins/FX/Niagara/Source/Niagara/Private/NiagaraDataInterfaceTexture2D.cpp
+++ b/Engine/Plugins/FX/Niagara/Source/Niagara/Private/NiagaraDataInterfaceTexture2D.cpp
@@ -3,6 +3,7 @@
 #include "NiagaraDataInterfaceTexture2D.h"
 #include "NiagaraShader.h"
 #include "ShaderParameterUtils.h"
+#include "NiagaraCustomVersion.h"
 
 
 #define LOCTEXT_NAMESPACE "UNiagaraDataInterfaceTexture2D"
@@ -31,6 +32,14 @@ void UNiagaraDataInterfaceTexture2D::PostInitProperties()
 void UNiagaraDataInterfaceTexture2D::PostLoad()
 {
 	Super::PostLoad();
+	const int32 NiagaraVer = GetLinkerCustomVersion(FNiagaraCustomVersion::GUID);
+	if (NiagaraVer < FNiagaraCustomVersion::TextureDataInterfaceUsesCustomSerialize)
+	{
+		if (Texture != nullptr)
+		{
+			CopyTextureToCPUBackup(Texture, CPUTextureData);
+		}
+	}
 }
 
 #if WITH_EDITOR
@@ -95,6 +104,16 @@ void UNiagaraDataInterfaceTexture2D::CopyTextureData(const uint8* Source, uint8*
 
 #endif
 
+void UNiagaraDataInterfaceTexture2D::Serialize(FArchive& Ar)
+{
+	Super::Serialize(Ar);
+	if (Ar.IsLoading() == false || Ar.CustomVer(FNiagaraCustomVersion::GUID) >= FNiagaraCustomVersion::TextureDataInterfaceUsesCustomSerialize)
+	{
+		Ar << CPUTextureData;
+	}
+	Ar.UsingCustomVersion(FNiagaraCustomVersion::GUID);
+}
+
 bool UNiagaraDataInterfaceTexture2D::CopyToInternal(UNiagaraDataInterface* Destination) const
 {
 	if (!Super::CopyToInternal(Destination))
diff --git a/Engine/Plugins/FX/Niagara/Source/NiagaraCore/Public/NiagaraCustomVersion.h b/Engine/Plugins/FX/Niagara/Source/NiagaraCore/Public/NiagaraCustomVersion.h
--- a/Engine/Plugins/FX/Niagara/Source/NiagaraCore/Public/NiagaraCustomVersion.h
+++ b/Engine/Plugins/FX/Niagara/Source/NiagaraCore/Public/NiagaraCustomVersion.h
@@ -106,6 +106,8 @@ struct FNiagaraCustomVersion
 
 		EmitterLocalSpaceLiteralConstant, //Emitter local space is compiled into the hlsl as a literal constant to expose it to emitter scripts and allow for some better optimization of particle transforms.
 
+		TextureDataInterfaceUsesCustomSerialize, // The cpu cache of the texture is now directly serialized instead of using array property serialization.
+
 		// -----<new versions can be added above this line>-------------------------------------------------
 		VersionPlusOne,
 		LatestVersion = VersionPlusOne - 1,
