Niagara - Fix CIS compile issues from texture data interface changes, also fix some crashes on existing texture data interface assets.

diff --git a/Engine/Plugins/FX/Niagara/Source/Niagara/Private/NiagaraDataInterfaceTexture2D.cpp b/Engine/Plugins/FX/Niagara/Source/Niagara/Private/NiagaraDataInterfaceTexture2D.cpp
--- a/Engine/Plugins/FX/Niagara/Source/Niagara/Private/NiagaraDataInterfaceTexture2D.cpp
+++ b/Engine/Plugins/FX/Niagara/Source/Niagara/Private/NiagaraDataInterfaceTexture2D.cpp
@@ -32,14 +32,17 @@ void UNiagaraDataInterfaceTexture2D::PostInitProperties()
 void UNiagaraDataInterfaceTexture2D::PostLoad()
 {
 	Super::PostLoad();
+#if WITH_EDITOR
 	const int32 NiagaraVer = GetLinkerCustomVersion(FNiagaraCustomVersion::GUID);
 	if (NiagaraVer < FNiagaraCustomVersion::TextureDataInterfaceUsesCustomSerialize)
 	{
 		if (Texture != nullptr)
 		{
+			Texture->ConditionalPostLoad();
 			CopyTextureToCPUBackup(Texture, CPUTextureData);
 		}
 	}
+#endif
 }
 
 #if WITH_EDITOR
