Fix for compile issue when not using UpdateAge

diff --git a/Engine/Plugins/FX/Niagara/Source/NiagaraEditor/Private/NiagaraHlslTranslator.cpp b/Engine/Plugins/FX/Niagara/Source/NiagaraEditor/Private/NiagaraHlslTranslator.cpp
--- a/Engine/Plugins/FX/Niagara/Source/NiagaraEditor/Private/NiagaraHlslTranslator.cpp
+++ b/Engine/Plugins/FX/Niagara/Source/NiagaraEditor/Private/NiagaraHlslTranslator.cpp
@@ -1666,7 +1666,6 @@ void FHlslNiagaraTranslator::DefineMain(FString &OutHlslOutput,
 				
 				if (StageIdx + 1 < TranslationStages.Num() && TranslationStages[StageIdx + 1].bCopyPreviousParams)
 				{
-					OutHlslOutput += FString::Printf(TEXT("\t\tContext.%s.DataInstance.Alive = Context.%s.DataInstance.Alive;\n"), *TranslationStages[StageIdx + 1].PassNamespace, *TranslationStages[StageIdx].PassNamespace);
 					OutHlslOutput += TEXT("\t\t//Begin Transfer of Attributes!\n");
 					if (ParamMapDefinedAttributesToNamespaceVars.Num() != 0)
 					{
@@ -1708,7 +1707,6 @@ void FHlslNiagaraTranslator::DefineMain(FString &OutHlslOutput,
 
 				if (StageIdx + 1 < TranslationStages.Num() && TranslationStages[StageIdx + 1].bCopyPreviousParams)
 				{
-					OutHlslOutput += FString::Printf(TEXT("\tContext.%s.DataInstance.Alive = Context.%s.DataInstance.Alive;\n"), *TranslationStages[StageIdx + 1].PassNamespace, *TranslationStages[StageIdx].PassNamespace);
 					OutHlslOutput += TEXT("\t//Begin Transfer of Attributes!\n");
 					if (ParamMapDefinedAttributesToNamespaceVars.Num() != 0)
 					{
