Niagara - Remove GPU contexts from the emitter instance batcher when the emitter instance is destroyed.  Fixes a crash when opening a GPU emitter that's running in the world.

Note: Copy from 4.20

diff --git a/Engine/Plugins/FX/Niagara/Source/Niagara/Classes/NiagaraEmitterInstanceBatcher.h b/Engine/Plugins/FX/Niagara/Source/Niagara/Classes/NiagaraEmitterInstanceBatcher.h
--- a/Engine/Plugins/FX/Niagara/Source/Niagara/Classes/NiagaraEmitterInstanceBatcher.h
+++ b/Engine/Plugins/FX/Niagara/Source/Niagara/Classes/NiagaraEmitterInstanceBatcher.h
@@ -19,6 +19,8 @@ the same VectorVM byte code / compute shader code
 struct FNiagaraScriptExecutionContext;
 struct FNiagaraComputeExecutionContext;
 
+#define SIMULATION_QUEUE_COUNT 2
+
 class NiagaraEmitterInstanceBatcher : public FTickableGameObject, public FComputeDispatcher
 {
 public:
@@ -52,6 +54,8 @@ public:
 
 	void Queue(FNiagaraComputeExecutionContext *InContext);
 
+	void Remove(FNiagaraComputeExecutionContext* InContext);
+
 	virtual ETickableTickType GetTickableTickType() const override
 	{
 		return ETickableTickType::Always;
@@ -113,5 +117,5 @@ private:
 	static NiagaraEmitterInstanceBatcher* BatcherSingleton;
 
 	uint32 CurQueueIndex;
-	TArray<FNiagaraComputeExecutionContext*> SimulationQueue[2];
+	TArray<FNiagaraComputeExecutionContext*> SimulationQueue[SIMULATION_QUEUE_COUNT];
 };
diff --git a/Engine/Plugins/FX/Niagara/Source/Niagara/Private/NiagaraEmitterInstance.cpp b/Engine/Plugins/FX/Niagara/Source/Niagara/Private/NiagaraEmitterInstance.cpp
--- a/Engine/Plugins/FX/Niagara/Source/Niagara/Private/NiagaraEmitterInstance.cpp
+++ b/Engine/Plugins/FX/Niagara/Source/Niagara/Private/NiagaraEmitterInstance.cpp
@@ -70,6 +70,11 @@ FNiagaraEmitterInstance::~FNiagaraEmitterInstance()
 	CachedBounds.Init();
 	UnbindParameters();
 
+	if (CachedEmitter != nullptr && CachedEmitter->SimTarget == ENiagaraSimTarget::GPUComputeSim)
+	{
+		NiagaraEmitterInstanceBatcher::Get()->Remove(&GPUExecContext);
+	}
+
 	/** We defer the deletion of the particle dataset to the RT to be sure all in-flight RT commands have finished using it.*/
 	ENQUEUE_UNIQUE_RENDER_COMMAND_ONEPARAMETER(FDeleteParticleDataSetCommand,
 		FNiagaraDataSet*, DataSet, ParticleDataSet,
diff --git a/Engine/Plugins/FX/Niagara/Source/Niagara/Private/NiagaraEmitterInstanceBatcher.cpp b/Engine/Plugins/FX/Niagara/Source/Niagara/Private/NiagaraEmitterInstanceBatcher.cpp
--- a/Engine/Plugins/FX/Niagara/Source/Niagara/Private/NiagaraEmitterInstanceBatcher.cpp
+++ b/Engine/Plugins/FX/Niagara/Source/Niagara/Private/NiagaraEmitterInstanceBatcher.cpp
@@ -41,6 +41,18 @@ void NiagaraEmitterInstanceBatcher::Queue(FNiagaraComputeExecutionContext *InCon
 			});
 }
 
+void NiagaraEmitterInstanceBatcher::Remove(FNiagaraComputeExecutionContext *InContext)
+{
+	ENQUEUE_UNIQUE_RENDER_COMMAND_TWOPARAMETER(RemoveNiagaraDispatch,
+		TArray<FNiagaraComputeExecutionContext*>*, Queue, &SimulationQueue[0],
+		FNiagaraComputeExecutionContext*, ExecContext, InContext,
+		{
+			for (int32 i = 0; i < SIMULATION_QUEUE_COUNT; i++)
+			{
+				Queue[i].Remove(ExecContext);
+			}
+		});
+}
 
 void NiagaraEmitterInstanceBatcher::ExecuteAll(FRHICommandList &RHICmdList, FUniformBufferRHIParamRef ViewUniformBuffer)
 {
diff --git a/Engine/Plugins/FX/Niagara/Source/NiagaraEditor/Private/ViewModels/NiagaraSystemViewModel.cpp b/Engine/Plugins/FX/Niagara/Source/NiagaraEditor/Private/ViewModels/NiagaraSystemViewModel.cpp
--- a/Engine/Plugins/FX/Niagara/Source/NiagaraEditor/Private/ViewModels/NiagaraSystemViewModel.cpp
+++ b/Engine/Plugins/FX/Niagara/Source/NiagaraEditor/Private/ViewModels/NiagaraSystemViewModel.cpp
@@ -88,17 +88,18 @@ void FNiagaraSystemViewModel::Cleanup()
 {
 	UE_LOG(LogNiagaraEditor, Warning, TEXT("Cleanup System view model %p"), this);
 
-	if (PreviewComponent)
-	{
-		PreviewComponent->OnSystemInstanceChanged().RemoveAll(this);
-	}
-
 	if (SystemInstance)
 	{
 		SystemInstance->OnInitialized().RemoveAll(this);
 		SystemInstance->OnReset().RemoveAll(this);
 	}
 
+	if (PreviewComponent)
+	{
+		PreviewComponent->OnSystemInstanceChanged().RemoveAll(this);
+		PreviewComponent->DeactivateImmediate();
+	}
+
 	CurveOwner.EmptyCurves();
 
 	GEditor->UnregisterForUndo(this);
