Niagara - Fix bounds calculation to match other engine components, and fix uniform parameter updating.
+ In the niagara scene proxy, notify the renderers when the proxy transform changes.
+ In the niagara component, don't reset the system bounds when sending the render data, and transform the system bounds by the local to world transform in CalcBounds.
+ In each of the renderers, release the WorldSpacePrimitiveUniformBuffer when the proxy transform changes, so that the uniforms are updated for the next render.

diff --git a/Engine/Plugins/FX/Niagara/Source/Niagara/Classes/NiagaraEmitterInstance.h b/Engine/Plugins/FX/Niagara/Source/Niagara/Classes/NiagaraEmitterInstance.h
--- a/Engine/Plugins/FX/Niagara/Source/Niagara/Classes/NiagaraEmitterInstance.h
+++ b/Engine/Plugins/FX/Niagara/Source/Niagara/Classes/NiagaraEmitterInstance.h
@@ -56,7 +56,7 @@ public:
 
 	uint32 CalculateEventSpawnCount(const FNiagaraEventScriptProperties &EventHandlerProps, TArray<int32, TInlineAllocator<16>>& EventSpawnCounts, FNiagaraDataSet *EventSet);
 
-	NIAGARA_API FBox CalculateDynamicBounds();
+	NIAGARA_API TOptional<FBox> CalculateDynamicBounds();
 
 	FNiagaraDataSet &GetData()	{ return *ParticleDataSet; }
 
diff --git a/Engine/Plugins/FX/Niagara/Source/Niagara/Private/NiagaraComponent.cpp b/Engine/Plugins/FX/Niagara/Source/Niagara/Private/NiagaraComponent.cpp
--- a/Engine/Plugins/FX/Niagara/Source/Niagara/Private/NiagaraComponent.cpp
+++ b/Engine/Plugins/FX/Niagara/Source/Niagara/Private/NiagaraComponent.cpp
@@ -190,7 +190,13 @@ void FNiagaraSceneProxy::CreateRenderThreadResources()
 
 void FNiagaraSceneProxy::OnTransformChanged()
 {
-	//WorldSpacePrimitiveUniformBuffer.ReleaseResource();
+	for (NiagaraRenderer* Renderer : EmitterRenderers)
+	{
+		if (Renderer)
+		{
+			Renderer->TransformChanged();
+		}
+	}
 }
 
 FPrimitiveViewRelevance FNiagaraSceneProxy::GetViewRelevance(const FSceneView* View) const
@@ -760,8 +766,6 @@ void UNiagaraComponent::SendRenderDynamicData_Concurrent()
 	SCOPE_CYCLE_COUNTER(STAT_NiagaraComponentSendRenderData);
 	if (SystemInstance.IsValid() && SceneProxy)
 	{
-		SystemInstance->GetSystemBounds().Init();
-		
 		FNiagaraSceneProxy* NiagaraProxy = static_cast<FNiagaraSceneProxy*>(SceneProxy);
 
 		for (int32 i = 0; i < SystemInstance->GetEmitters().Num(); i++)
@@ -813,7 +817,7 @@ int32 UNiagaraComponent::GetNumMaterials() const
 
 FBoxSphereBounds UNiagaraComponent::CalcBounds(const FTransform& LocalToWorld) const
 {
-	FBox SimBounds(ForceInit);
+	FBoxSphereBounds SystemBounds;
 	if (SystemInstance.IsValid())
 	{
 		SystemInstance->GetSystemBounds().Init();
@@ -823,13 +827,16 @@ FBoxSphereBounds UNiagaraComponent::CalcBounds(const FTransform& LocalToWorld) c
 			SystemInstance->GetSystemBounds() += Sim.GetBounds();
 		}
 		FBox BoundingBox = SystemInstance->GetSystemBounds();
-		const FVector ExpandAmount = FVector(0.0f, 0.0f, 0.0f);// BoundingBox.GetExtent() * 0.1f;
-		BoundingBox = FBox(BoundingBox.Min - ExpandAmount, BoundingBox.Max + ExpandAmount);
 
-		FBoxSphereBounds BSBounds(BoundingBox);
-		return BSBounds;
+		SystemBounds = FBoxSphereBounds(BoundingBox);
+	}
+	else
+	{
+		FBox SimBounds(ForceInit);
+		SystemBounds = FBoxSphereBounds(SimBounds);
 	}
-	return FBoxSphereBounds(SimBounds);
+
+	return SystemBounds.TransformBy(LocalToWorld);
 }
 
 FPrimitiveSceneProxy* UNiagaraComponent::CreateSceneProxy()
diff --git a/Engine/Plugins/FX/Niagara/Source/Niagara/Private/NiagaraEmitterInstance.cpp b/Engine/Plugins/FX/Niagara/Source/Niagara/Private/NiagaraEmitterInstance.cpp
--- a/Engine/Plugins/FX/Niagara/Source/Niagara/Private/NiagaraEmitterInstance.cpp
+++ b/Engine/Plugins/FX/Niagara/Source/Niagara/Private/NiagaraEmitterInstance.cpp
@@ -99,13 +99,6 @@ void FNiagaraEmitterInstance::ClearRenderer()
 
 FBox FNiagaraEmitterInstance::GetBounds()
 {
-	checkSlow(CachedEmitter);
-	{
-		if (CachedEmitter->bFixedBounds || CachedEmitter->SimTarget == ENiagaraSimTarget::GPUComputeSim)//We must use fixed bounds on GPU sim.
-		{
-			return CachedEmitter->FixedBounds.TransformBy(ParentSystemInstance->GetComponent()->GetComponentToWorld());
-		}
-	}
 	return CachedBounds;
 }
 
@@ -595,19 +588,26 @@ int FNiagaraEmitterInstance::GetTotalBytesUsed()
 	return BytesUsed;
 }
 
-FBox FNiagaraEmitterInstance::CalculateDynamicBounds()
+TOptional<FBox> FNiagaraEmitterInstance::CalculateDynamicBounds()
 {
 	checkSlow(ParticleDataSet);
 	FNiagaraDataSet& Data = *ParticleDataSet;
 	int32 NumInstances = Data.GetNumInstances();
 	FBox Ret;
 	Ret.Init();
+
 	if (IsComplete() || NumInstances == 0 || CachedEmitter->SimTarget == ENiagaraSimTarget::GPUComputeSim)//TODO: Pull data back from gpu buffers to get bounds for GPU sims.
 	{
-		return Ret;
+		return TOptional<FBox>();
 	}
 
 	PositionAccessor.InitForAccess(true);
+
+	if (PositionAccessor.IsValid() == false)
+	{
+		return TOptional<FBox>();
+	}
+
 	SizeAccessor.InitForAccess(true);
 	MeshScaleAccessor.InitForAccess(true);
 
@@ -671,12 +671,6 @@ FBox FNiagaraEmitterInstance::CalculateDynamicBounds()
 
 	Ret = Ret.ExpandBy(MaxSize*MaxBaseSize);
 
-	if (CachedEmitter->bLocalSpace || PositionAccessor.IsValid() == false)
-	{
-		FMatrix SystemLocalToWorld = GetParentSystemInstance()->GetComponent()->GetComponentTransform().ToMatrixNoScale();
-		Ret = Ret.TransformBy(SystemLocalToWorld);
-	}
-
 	return Ret;
 }
 
@@ -688,16 +682,29 @@ void FNiagaraEmitterInstance::PostProcessParticles()
 	SCOPE_CYCLE_COUNTER(STAT_NiagaraKill);
 
 	checkSlow(CachedEmitter);
-	if (CachedEmitter->bFixedBounds)
+	CachedBounds.Init();
+	if (CachedEmitter->bFixedBounds || CachedEmitter->SimTarget == ENiagaraSimTarget::GPUComputeSim)
 	{
 		CachedBounds = CachedEmitter->FixedBounds;
 	}
 	else
 	{
-		checkSlow(ParticleDataSet);
-		FNiagaraDataSet& Data = *ParticleDataSet;
-		int32 NumParticles = Data.GetNumInstances();
-		CachedBounds = CalculateDynamicBounds();
+		TOptional<FBox> DynamicBounds = CalculateDynamicBounds();
+		if (DynamicBounds.IsSet())
+		{
+			if (CachedEmitter->bLocalSpace)
+			{
+				CachedBounds = DynamicBounds.GetValue();
+			}
+			else
+			{
+				CachedBounds = DynamicBounds.GetValue().TransformBy(ParentSystemInstance->GetComponent()->GetComponentToWorld().Inverse());
+			}
+		}
+		else
+		{
+			CachedBounds = CachedEmitter->FixedBounds;
+		}
 	}
 }
 
diff --git a/Engine/Plugins/FX/Niagara/Source/Niagara/Private/NiagaraRenderer.cpp b/Engine/Plugins/FX/Niagara/Source/Niagara/Private/NiagaraRenderer.cpp
--- a/Engine/Plugins/FX/Niagara/Source/Niagara/Private/NiagaraRenderer.cpp
+++ b/Engine/Plugins/FX/Niagara/Source/Niagara/Private/NiagaraRenderer.cpp
@@ -61,8 +61,6 @@ FRWBuffer& NiagaraRenderer::GetDummyIntBuffer()
 	return DummyIntBuffer.Buffer;
 }
 
-
-
 NiagaraRenderer::NiagaraRenderer()
 	: CPUTimeMS(0.0f)
 	, bLocalSpace(false)
@@ -301,6 +299,10 @@ bool NiagaraRendererLights::SetMaterialUsage()
 	return false;
 }
 
+void NiagaraRendererLights::TransformChanged()
+{
+}
+
 #if WITH_EDITORONLY_DATA
 
 const TArray<FNiagaraVariable>& NiagaraRendererLights::GetRequiredAttributes()
diff --git a/Engine/Plugins/FX/Niagara/Source/Niagara/Private/NiagaraRendererMeshes.cpp b/Engine/Plugins/FX/Niagara/Source/Niagara/Private/NiagaraRendererMeshes.cpp
--- a/Engine/Plugins/FX/Niagara/Source/Niagara/Private/NiagaraRendererMeshes.cpp
+++ b/Engine/Plugins/FX/Niagara/Source/Niagara/Private/NiagaraRendererMeshes.cpp
@@ -343,7 +343,10 @@ bool NiagaraRendererMeshes::SetMaterialUsage()
 	return Material && Material->CheckMaterialUsage_Concurrent(MATUSAGE_NiagaraMeshParticles);
 }
 
-
+void NiagaraRendererMeshes::TransformChanged()
+{
+	WorldSpacePrimitiveUniformBuffer.ReleaseResource();
+}
 
 /** Update render data buffer from attributes */
 FNiagaraDynamicDataBase *NiagaraRendererMeshes::GenerateVertexData(const FNiagaraSceneProxy* Proxy, FNiagaraDataSet &Data, const ENiagaraSimTarget Target)
diff --git a/Engine/Plugins/FX/Niagara/Source/Niagara/Private/NiagaraRendererRibbons.cpp b/Engine/Plugins/FX/Niagara/Source/Niagara/Private/NiagaraRendererRibbons.cpp
--- a/Engine/Plugins/FX/Niagara/Source/Niagara/Private/NiagaraRendererRibbons.cpp
+++ b/Engine/Plugins/FX/Niagara/Source/Niagara/Private/NiagaraRendererRibbons.cpp
@@ -270,6 +270,11 @@ bool NiagaraRendererRibbons::SetMaterialUsage()
 	return Material && Material->CheckMaterialUsage_Concurrent(MATUSAGE_NiagaraRibbons);
 }
 
+void NiagaraRendererRibbons::TransformChanged()
+{
+	WorldSpacePrimitiveUniformBuffer.ReleaseResource();
+}
+
 void CalculateUVScaleAndOffsets(const FNiagaraDataSetAccessor<float>& SortKeyData, const TArray<int32>& RibbonIndices, bool bSortKeyIsAge, int32 NumSegments,
 	float InUVTilingDistance, const FVector2D& InUVScale, const FVector2D& InUVOffset, ENiagaraRibbonAgeOffsetMode InAgeOffsetMode, FVector2D& OutUVScale, FVector2D& OutUVOffset)
 {
diff --git a/Engine/Plugins/FX/Niagara/Source/Niagara/Private/NiagaraRendererSprites.cpp b/Engine/Plugins/FX/Niagara/Source/Niagara/Private/NiagaraRendererSprites.cpp
--- a/Engine/Plugins/FX/Niagara/Source/Niagara/Private/NiagaraRendererSprites.cpp
+++ b/Engine/Plugins/FX/Niagara/Source/Niagara/Private/NiagaraRendererSprites.cpp
@@ -296,6 +296,11 @@ bool NiagaraRendererSprites::SetMaterialUsage()
 	return Material && Material->CheckMaterialUsage_Concurrent(MATUSAGE_NiagaraSprites);
 }
 
+void NiagaraRendererSprites::TransformChanged()
+{
+	WorldSpacePrimitiveUniformBuffer.ReleaseResource();
+}
+
 /** Update render data buffer from attributes */
 FNiagaraDynamicDataBase *NiagaraRendererSprites::GenerateVertexData(const FNiagaraSceneProxy* Proxy, FNiagaraDataSet &Data, const ENiagaraSimTarget Target)
 {
diff --git a/Engine/Plugins/FX/Niagara/Source/Niagara/Public/NiagaraRenderer.h b/Engine/Plugins/FX/Niagara/Source/Niagara/Public/NiagaraRenderer.h
--- a/Engine/Plugins/FX/Niagara/Source/Niagara/Public/NiagaraRenderer.h
+++ b/Engine/Plugins/FX/Niagara/Source/Niagara/Public/NiagaraRenderer.h
@@ -93,6 +93,8 @@ public:
 
 	virtual bool SetMaterialUsage() = 0;
 
+	virtual void TransformChanged() = 0;
+
 #if WITH_EDITORONLY_DATA
 	virtual const TArray<FNiagaraVariable>& GetRequiredAttributes() = 0;
 	virtual const TArray<FNiagaraVariable>& GetOptionalAttributes() = 0;
@@ -196,6 +198,7 @@ public:
 
 	virtual void GetDynamicMeshElements(const TArray<const FSceneView*>& Views, const FSceneViewFamily& ViewFamily, uint32 VisibilityMap, FMeshElementCollector& Collector, const FNiagaraSceneProxy *SceneProxy) const override;
 	virtual bool SetMaterialUsage() override;
+	virtual void TransformChanged() override;
 	/** Update render data buffer from attributes */
 	FNiagaraDynamicDataBase *GenerateVertexData(const FNiagaraSceneProxy* Proxy, FNiagaraDataSet &Data, const ENiagaraSimTarget Target) override;
 
@@ -280,6 +283,7 @@ public:
 
 	virtual void GetDynamicMeshElements(const TArray<const FSceneView*>& Views, const FSceneViewFamily& ViewFamily, uint32 VisibilityMap, FMeshElementCollector& Collector, const FNiagaraSceneProxy *SceneProxy) const override;
 	virtual bool SetMaterialUsage() override;
+	virtual void TransformChanged() override;
 	/** Update render data buffer from attributes */
 	FNiagaraDynamicDataBase *GenerateVertexData(const FNiagaraSceneProxy* Proxy, FNiagaraDataSet &Data, const ENiagaraSimTarget Target) override;
 
diff --git a/Engine/Plugins/FX/Niagara/Source/Niagara/Public/NiagaraRendererMeshes.h b/Engine/Plugins/FX/Niagara/Source/Niagara/Public/NiagaraRendererMeshes.h
--- a/Engine/Plugins/FX/Niagara/Source/Niagara/Public/NiagaraRendererMeshes.h
+++ b/Engine/Plugins/FX/Niagara/Source/Niagara/Public/NiagaraRendererMeshes.h
@@ -42,6 +42,7 @@ public:
 
 	virtual void GetDynamicMeshElements(const TArray<const FSceneView*>& Views, const FSceneViewFamily& ViewFamily, uint32 VisibilityMap, FMeshElementCollector& Collector, const FNiagaraSceneProxy *SceneProxy) const override;
 	virtual bool SetMaterialUsage() override;
+	virtual void TransformChanged() override;
 	/** Update render data buffer from attributes */
 	FNiagaraDynamicDataBase *GenerateVertexData(const FNiagaraSceneProxy* Proxy, FNiagaraDataSet &Data, const ENiagaraSimTarget Target) override;
 
diff --git a/Engine/Plugins/FX/Niagara/Source/Niagara/Public/NiagaraRendererRibbons.h b/Engine/Plugins/FX/Niagara/Source/Niagara/Public/NiagaraRendererRibbons.h
--- a/Engine/Plugins/FX/Niagara/Source/Niagara/Public/NiagaraRendererRibbons.h
+++ b/Engine/Plugins/FX/Niagara/Source/Niagara/Public/NiagaraRendererRibbons.h
@@ -31,6 +31,8 @@ public:
 
 	virtual bool SetMaterialUsage() override;
 
+	virtual void TransformChanged() override;
+
 	/** Update render data buffer from attributes */
 	FNiagaraDynamicDataBase *GenerateVertexData(const FNiagaraSceneProxy* Proxy, FNiagaraDataSet &Data, const ENiagaraSimTarget Target) override;
 
diff --git a/Engine/Plugins/FX/Niagara/Source/NiagaraEditor/Private/ViewModels/NiagaraSystemViewModel.cpp b/Engine/Plugins/FX/Niagara/Source/NiagaraEditor/Private/ViewModels/NiagaraSystemViewModel.cpp
--- a/Engine/Plugins/FX/Niagara/Source/NiagaraEditor/Private/ViewModels/NiagaraSystemViewModel.cpp
+++ b/Engine/Plugins/FX/Niagara/Source/NiagaraEditor/Private/ViewModels/NiagaraSystemViewModel.cpp
@@ -1583,11 +1583,21 @@ void FNiagaraSystemViewModel::UpdateEmitterFixedBounds()
 		{
 			if (&EmitterInst->GetEmitterHandle() == SelectedEmitterHandle && !EmitterInst->IsComplete())
 			{
-				FBox EmitterBounds = EmitterInst->CalculateDynamicBounds();
-				Emitter->Modify();
-				Emitter->bFixedBounds = true;
-				//Dynamic bounds are in world space. Transform back to local.
-				Emitter->FixedBounds = EmitterBounds.TransformBy(PreviewComponent->GetComponentToWorld().Inverse());
+				TOptional<FBox> EmitterBounds = EmitterInst->CalculateDynamicBounds();
+				if (EmitterBounds.IsSet())
+				{
+					Emitter->Modify();
+					Emitter->bFixedBounds = true;
+					if (Emitter->bLocalSpace)
+					{
+						Emitter->FixedBounds = EmitterBounds.GetValue();
+					}
+					else
+					{
+						//Dynamic bounds are in world space. Transform back to local.
+						Emitter->FixedBounds = EmitterBounds.GetValue().TransformBy(PreviewComponent->GetComponentToWorld().Inverse());
+					}
+				}
 			}
 		}
 	}
