Now allowing a generic UTexture instead of a UTexture2D to eventually support render targets and media textures. Renamed NiagaraDataInterfaceTexture2D to NiagaraDataInterfaceTexture as it now holds a UTexture* instead of a UTexture2D*. This is to facilitate the aforementioned other class types, which derive from UTexture.
Had to make a class redirector for this and make sure that it worked for data interfaces.

#tests old texture sample test still works

diff --git a/Engine/Plugins/FX/Niagara/Config/BaseNiagara.ini b/Engine/Plugins/FX/Niagara/Config/BaseNiagara.ini
--- a/Engine/Plugins/FX/Niagara/Config/BaseNiagara.ini
+++ b/Engine/Plugins/FX/Niagara/Config/BaseNiagara.ini
@@ -9,6 +9,7 @@
 +ClassRedirects=(OldName="NiagaraEffectRendererSprites", NewName="/Script/Niagara.NiagaraRendererSprites")
 +ClassRedirects=(OldName="NiagaraEffectRendererMeshes", NewName="/Script/Niagara.NiagaraRendererMeshes")
 +ClassRedirects=(OldName="NiagaraEffectRendererProperties", NewName="/Script/Niagara.NiagaraRendererProperties")
++ClassRedirects=(OldName="NiagaraDataInterfaceTexture2D", NewName="/Script/Niagara.NiagaraDataInterfaceTexture")
 
 +EnumRedirects=(OldName="ENiagaraScriptUsage",ValueChanges=(("ENiagaraScriptUsage::SpawnScript","ENiagaraScriptUsage::ParticleSpawnScript"), ("ENiagaraScriptUsage::SpawnScriptInterpolated","ENiagaraScriptUsage::ParticleSpawnScriptInterpolated"),("ENiagaraScriptUsage::UpdateScript","ENiagaraScriptUsage::ParticleUpdateScript"), ("ENiagaraScriptUsage::EventScript","ENiagaraScriptUsage::ParticleEventScript"), ("ENiagaraScriptUsage::SystemScript","ENiagaraScriptUsage::SystemSpawnScript")) )
 
diff --git a/Engine/Plugins/FX/Niagara/Source/Niagara/Classes/NiagaraDataInterfaceTexture.h b/Engine/Plugins/FX/Niagara/Source/Niagara/Classes/NiagaraDataInterfaceTexture.h
new file mode 100644
--- /dev/null
+++ b/Engine/Plugins/FX/Niagara/Source/Niagara/Classes/NiagaraDataInterfaceTexture.h
@@ -0,0 +1,60 @@
+// Copyright 1998-2018 Epic Games, Inc. All Rights Reserved.
+
+#pragma once
+
+#include "NiagaraDataInterface.h"
+#include "NiagaraCommon.h"
+#include "VectorVM.h"
+#include "NiagaraDataInterfaceTexture.generated.h"
+
+/** Data Interface allowing sampling of a texture */
+UCLASS(EditInlineNew, Category = "Texture", meta = (DisplayName = "Texture Sample"))
+class NIAGARA_API UNiagaraDataInterfaceTexture : public UNiagaraDataInterface
+{
+	GENERATED_UCLASS_BODY()
+		bool bGPUBufferDirty;
+public:
+
+	UPROPERTY(EditAnywhere, Category = "Texture")
+	UTexture* Texture;
+
+	//UObject Interface
+	virtual void PostInitProperties()override;
+	virtual void PostLoad() override;
+#if WITH_EDITOR
+	virtual void PostEditChangeProperty(struct FPropertyChangedEvent& PropertyChangedEvent) override;
+#endif
+	virtual void Serialize(FArchive& Ar) override;
+	//UObject Interface End
+
+	//UNiagaraDataInterface Interface
+	virtual void GetFunctions(TArray<FNiagaraFunctionSignature>& OutFunctions)override;
+	virtual void GetVMExternalFunction(const FVMExternalFunctionBindingInfo& BindingInfo, void* InstanceData, FVMExternalFunction &OutFunc) override;
+	virtual bool CanExecuteOnTarget(ENiagaraSimTarget Target)const override { return Target==ENiagaraSimTarget::GPUComputeSim; }
+	//UNiagaraDataInterface Interface End
+
+	template<typename XType, typename YType>
+	void SampleTexture(FVectorVMContext& Context);
+
+	virtual bool Equals(const UNiagaraDataInterface* Other) const override;
+
+	// GPU sim functionality
+	virtual bool GetFunctionHLSL(const FName&  DefinitionFunctionName, FString InstanceFunctionName, FNiagaraDataInterfaceGPUParamInfo& ParamInfo, FString& OutHLSL) override;
+	virtual void GetParameterDefinitionHLSL(FNiagaraDataInterfaceGPUParamInfo& ParamInfo, FString& OutHLSL) override;
+	virtual FNiagaraDataInterfaceParametersCS* ConstructComputeParameters()const override;
+
+	//FRWBuffer& GetGPUBuffer();
+	static const FString TextureName;
+	static const FString SamplerName;
+protected:
+#if WITH_EDITOR
+	bool CopyTextureToCPUBackup(UTexture* SourceTexture, TArray<uint8>& TargetBuffer);
+	void CopyTextureData(const uint8* Source, uint8* Dest, uint32 SizeX, uint32 SizeY, uint32 BytesPerPixel, uint32 SourceStride, uint32 DestStride);
+
+#endif
+	virtual bool CopyToInternal(UNiagaraDataInterface* Destination) const override;
+
+	TArray<uint8> CPUTextureData;
+
+	static const FName SampleTextureName;
+};
diff --git a/Engine/Plugins/FX/Niagara/Source/Niagara/Classes/NiagaraDataInterfaceTexture2D.h b/Engine/Plugins/FX/Niagara/Source/Niagara/Classes/NiagaraDataInterfaceTexture2D.h
deleted file mode 100644
diff --git a/Engine/Plugins/FX/Niagara/Source/Niagara/Private/NiagaraDataInterfaceTexture.cpp b/Engine/Plugins/FX/Niagara/Source/Niagara/Private/NiagaraDataInterfaceTexture.cpp
new file mode 100644
--- /dev/null
+++ b/Engine/Plugins/FX/Niagara/Source/Niagara/Private/NiagaraDataInterfaceTexture.cpp
@@ -0,0 +1,302 @@
+// Copyright 1998-2018 Epic Games, Inc. All Rights Reserved.
+
+#include "NiagaraDataInterfaceTexture.h"
+#include "NiagaraShader.h"
+#include "ShaderParameterUtils.h"
+#include "NiagaraCustomVersion.h"
+
+
+#define LOCTEXT_NAMESPACE "UNiagaraDataInterfaceTexture"
+
+const FName UNiagaraDataInterfaceTexture::SampleTextureName(TEXT("SampleTexture2D"));
+const FString UNiagaraDataInterfaceTexture::TextureName(TEXT("Texture_"));
+const FString UNiagaraDataInterfaceTexture::SamplerName(TEXT("Sampler_"));
+
+UNiagaraDataInterfaceTexture::UNiagaraDataInterfaceTexture(FObjectInitializer const& ObjectInitializer)
+	: Super(ObjectInitializer)
+	, Texture(nullptr)
+{
+
+}
+
+void UNiagaraDataInterfaceTexture::PostInitProperties()
+{
+	Super::PostInitProperties();
+
+	if (HasAnyFlags(RF_ClassDefaultObject))
+	{
+		FNiagaraTypeRegistry::Register(FNiagaraTypeDefinition(GetClass()), true, false, false);
+	}
+}
+
+void UNiagaraDataInterfaceTexture::PostLoad()
+{
+	Super::PostLoad();
+#if WITH_EDITOR
+	const int32 NiagaraVer = GetLinkerCustomVersion(FNiagaraCustomVersion::GUID);
+	if (NiagaraVer < FNiagaraCustomVersion::TextureDataInterfaceUsesCustomSerialize)
+	{
+		if (Texture != nullptr)
+		{
+			Texture->ConditionalPostLoad();
+			CopyTextureToCPUBackup(Texture, CPUTextureData);
+		}
+	}
+#endif
+}
+
+#if WITH_EDITOR
+
+void UNiagaraDataInterfaceTexture::PostEditChangeProperty(struct FPropertyChangedEvent& PropertyChangedEvent)
+{
+	Super::PostEditChangeProperty(PropertyChangedEvent);
+
+	if (PropertyChangedEvent.Property && PropertyChangedEvent.Property->GetFName() == GET_MEMBER_NAME_CHECKED(UNiagaraDataInterfaceTexture, Texture))
+	{
+		CopyTextureToCPUBackup(Texture, CPUTextureData);
+	}
+}
+
+bool UNiagaraDataInterfaceTexture::CopyTextureToCPUBackup(UTexture* SourceTexture, TArray<uint8>& TargetBuffer)
+{
+	if (SourceTexture == nullptr)
+	{
+		TargetBuffer.Empty();
+		return true;
+	}
+
+	FTextureSource& SourceData = SourceTexture->Source;
+	FIntPoint SourceSize = FIntPoint(SourceData.GetSizeX(), SourceData.GetSizeY());
+
+	{
+		const int32 BytesPerPixel = 4;
+		TargetBuffer.Empty();
+		TargetBuffer.AddZeroed(SourceSize.X * SourceSize.Y * BytesPerPixel);
+	}
+
+	
+	if (SourceData.GetFormat() == TSF_BGRA8)
+	{
+		uint32 BytesPerPixel = SourceData.GetBytesPerPixel();
+		uint8* OffsetSource = SourceData.LockMip(0);
+		uint8* OffsetDest = TargetBuffer.GetData();
+		CopyTextureData(OffsetSource, OffsetDest, SourceSize.X, SourceSize.Y, BytesPerPixel, SourceData.GetSizeX() * BytesPerPixel, SourceSize.X * BytesPerPixel);
+		SourceData.UnlockMip(0);
+	}
+	else 
+	{
+		UE_LOG(LogNiagara, Error, TEXT("Texture %s is not BGRA8, which isn't supported in data interfaces yet"), *SourceTexture->GetName());
+	}
+
+	return true;
+}
+
+void UNiagaraDataInterfaceTexture::CopyTextureData(const uint8* Source, uint8* Dest, uint32 SizeX, uint32 SizeY, uint32 BytesPerPixel, uint32 SourceStride, uint32 DestStride)
+{
+	const uint32 NumBytesPerRow = SizeX * BytesPerPixel;
+
+	for (uint32 Y = 0; Y < SizeY; ++Y)
+	{
+		FMemory::Memcpy(
+			Dest + (DestStride * Y),
+			Source + (SourceStride * Y),
+			NumBytesPerRow
+		);
+	}
+}
+
+#endif
+
+void UNiagaraDataInterfaceTexture::Serialize(FArchive& Ar)
+{
+	Super::Serialize(Ar);
+	if (Ar.IsLoading() == false || Ar.CustomVer(FNiagaraCustomVersion::GUID) >= FNiagaraCustomVersion::TextureDataInterfaceUsesCustomSerialize)
+	{
+		Ar << CPUTextureData;
+	}
+	Ar.UsingCustomVersion(FNiagaraCustomVersion::GUID);
+}
+
+bool UNiagaraDataInterfaceTexture::CopyToInternal(UNiagaraDataInterface* Destination) const
+{
+	if (!Super::CopyToInternal(Destination))
+	{
+		return false;
+	}
+	UNiagaraDataInterfaceTexture* DestinationTexture = CastChecked<UNiagaraDataInterfaceTexture>(Destination);
+	DestinationTexture->Texture = Texture;
+	DestinationTexture->CPUTextureData = CPUTextureData;
+
+	return true;
+}
+
+bool UNiagaraDataInterfaceTexture::Equals(const UNiagaraDataInterface* Other) const
+{
+	if (!Super::Equals(Other))
+	{
+		return false;
+	}
+	const UNiagaraDataInterfaceTexture* OtherTexture = CastChecked<const UNiagaraDataInterfaceTexture>(Other);
+	return OtherTexture->Texture == Texture;
+}
+
+void UNiagaraDataInterfaceTexture::GetFunctions(TArray<FNiagaraFunctionSignature>& OutFunctions)
+{
+	FNiagaraFunctionSignature Sig;
+	Sig.Name = SampleTextureName;
+	Sig.bMemberFunction = true;
+	Sig.bRequiresContext = false;
+	Sig.Inputs.Add(FNiagaraVariable(FNiagaraTypeDefinition(GetClass()), TEXT("Texture")));
+	Sig.Inputs.Add(FNiagaraVariable(FNiagaraTypeDefinition::GetVec2Def(), TEXT("UV")));
+	Sig.SetDescription(LOCTEXT("TextureSampleDesc", "Sample mip level 0 of the input 2d texture at the specified UV coordinates. The UV origin (0,0) is in the upper left hand corner of the image."));
+	Sig.Outputs.Add(FNiagaraVariable(FNiagaraTypeDefinition::GetVec4Def(), TEXT("Value")));
+	//Sig.Owner = *GetName();
+
+	OutFunctions.Add(Sig);
+}
+
+DEFINE_NDI_RAW_FUNC_BINDER(UNiagaraDataInterfaceTexture, SampleTexture);
+void UNiagaraDataInterfaceTexture::GetVMExternalFunction(const FVMExternalFunctionBindingInfo& BindingInfo, void* InstanceData, FVMExternalFunction &OutFunc)
+{
+	check(BindingInfo.Name == SampleTextureName);
+	check(BindingInfo.GetNumInputs() == 2 && BindingInfo.GetNumOutputs() == 4);
+	TNDIParamBinder<0, float, TNDIParamBinder<1, float, NDI_RAW_FUNC_BINDER(UNiagaraDataInterfaceTexture, SampleTexture)>>::Bind(this, BindingInfo, InstanceData, OutFunc);
+}
+
+template<typename XType, typename YType>
+void UNiagaraDataInterfaceTexture::SampleTexture(FVectorVMContext& Context)
+{
+	XType XParam(Context);
+	YType YParam(Context);
+	FRegisterHandler<float> OutSampleR(Context);
+	FRegisterHandler<float> OutSampleG(Context);
+	FRegisterHandler<float> OutSampleB(Context);
+	FRegisterHandler<float> OutSampleA(Context);
+
+	if (CPUTextureData.GetAllocatedSize() == 0 || Texture == nullptr)
+	{
+		for (int32 i = 0; i < Context.NumInstances; ++i)
+		{
+			float X = XParam.GetAndAdvance();
+			float Y = YParam.GetAndAdvance();
+			*OutSampleR.GetDestAndAdvance() = 1.0;
+			*OutSampleG.GetDestAndAdvance() = 0.0;
+			*OutSampleB.GetDestAndAdvance() = 1.0;
+			*OutSampleA.GetDestAndAdvance() = 1.0;
+		}
+	}
+	else
+	{
+		const int32 BytesPerPixel = 4;
+		int32 IntSizeX = Texture->GetSurfaceWidth();
+		int32 IntSizeY = Texture->GetSurfaceHeight();
+		float SizeX = (float)Texture->GetSurfaceWidth() - 1.0f;
+		float SizeY = (float)Texture->GetSurfaceHeight() - 1.0f;
+
+		for (int32 i = 0; i < Context.NumInstances; ++i)
+		{
+			float ParamXVal = XParam.GetAndAdvance();
+			float ParamYVal = YParam.GetAndAdvance();
+			float X = (fmodf(ParamXVal * SizeX, SizeX));
+			float Y = (fmodf(ParamYVal * SizeY, SizeY));
+
+			float XNorm = X < 0.0f ? (SizeX - fabsf(X)) : X;
+			float YNorm = Y < 0.0f ? (SizeY - fabsf(Y)) : Y;
+
+			int32 XInt = floorf(XNorm);
+			int32 YInt = floorf(YNorm);
+			int32 SampleIdx = YInt * IntSizeX * BytesPerPixel + XInt * BytesPerPixel;
+			ensure(CPUTextureData.Num() > SampleIdx);
+			uint8 B0 = CPUTextureData[SampleIdx + 0];
+			uint8 G0 = CPUTextureData[SampleIdx + 1];
+			uint8 R0 = CPUTextureData[SampleIdx + 2];
+			uint8 A0 = CPUTextureData[SampleIdx + 3];
+
+			*OutSampleR.GetDestAndAdvance() = ((float)R0) / 255.0f;
+			*OutSampleG.GetDestAndAdvance() = ((float)G0) / 255.0f;
+			*OutSampleB.GetDestAndAdvance() = ((float)B0) / 255.0f;
+			*OutSampleA.GetDestAndAdvance() = ((float)A0) / 255.0f;
+		}
+	}
+
+}
+
+bool UNiagaraDataInterfaceTexture::GetFunctionHLSL(const FName& DefinitionFunctionName, FString InstanceFunctionName, FNiagaraDataInterfaceGPUParamInfo& ParamInfo, FString& OutHLSL)
+{
+	FString HLSLTextureName = TextureName + ParamInfo.DataInterfaceHLSLSymbol;
+	FString HLSLSamplerName = SamplerName + ParamInfo.DataInterfaceHLSLSymbol;
+	OutHLSL += TEXT("void ") + InstanceFunctionName + TEXT("(in float2 In_UV, out float4 Out_Value) \n{\n");
+	OutHLSL += TEXT("\t Out_Value = ") + HLSLTextureName + TEXT(".SampleLevel(") + HLSLSamplerName + TEXT(", In_UV, 0);\n");
+	OutHLSL += TEXT("\n}\n");
+	return true;
+}
+
+void UNiagaraDataInterfaceTexture::GetParameterDefinitionHLSL(FNiagaraDataInterfaceGPUParamInfo& ParamInfo, FString& OutHLSL)
+{
+	FString HLSLTextureName = TextureName + ParamInfo.DataInterfaceHLSLSymbol;
+	FString HLSLSamplerName = SamplerName + ParamInfo.DataInterfaceHLSLSymbol;
+	OutHLSL += TEXT("Texture2D ") + HLSLTextureName + TEXT(";\n");
+	OutHLSL += TEXT("SamplerState ") + HLSLSamplerName + TEXT(";\n");
+}
+
+
+
+struct FNiagaraDataInterfaceParametersCS_Texture : public FNiagaraDataInterfaceParametersCS
+{
+	virtual void Bind(const FNiagaraDataInterfaceParamRef& ParamRef, const class FShaderParameterMap& ParameterMap) override
+	{
+		FString TexName = UNiagaraDataInterfaceTexture::TextureName + ParamRef.ParameterInfo.DataInterfaceHLSLSymbol;
+		FString SampleName = (UNiagaraDataInterfaceTexture::SamplerName + ParamRef.ParameterInfo.DataInterfaceHLSLSymbol);
+		TextureParam.Bind(ParameterMap, *TexName);
+		SamplerParam.Bind(ParameterMap, *SampleName);
+		
+		if (!TextureParam.IsBound())
+		{
+			UE_LOG(LogNiagara, Warning, TEXT("Binding failed for FNiagaraDataInterfaceParametersCS_Texture Texture %s. Was it optimized out?"), *TexName)
+		}
+
+		if (!SamplerParam.IsBound())
+		{
+			UE_LOG(LogNiagara, Warning, TEXT("Binding failed for FNiagaraDataInterfaceParametersCS_Texture Sampler %s. Was it optimized out?"), *SampleName)
+		}
+	}
+
+	virtual void Serialize(FArchive& Ar)override
+	{
+		Ar << TextureParam;
+		Ar << SamplerParam;
+	}
+
+	virtual void Set(FRHICommandList& RHICmdList, FNiagaraShader* Shader, class UNiagaraDataInterface* DataInterface) const override
+	{
+		check(IsInRenderingThread());
+
+		FComputeShaderRHIParamRef ComputeShaderRHI = Shader->GetComputeShader();
+		UNiagaraDataInterfaceTexture* TextureDI = CastChecked<UNiagaraDataInterfaceTexture>(DataInterface);
+		UTexture *Texture = TextureDI->Texture;
+		if (!Texture)
+		{
+			return;
+		}
+		FTextureRHIParamRef TextureRHI = Texture->TextureReference.TextureReferenceRHI;
+		SetTextureParameter(
+			RHICmdList,
+			ComputeShaderRHI,
+			TextureParam,
+			SamplerParam,
+			TStaticSamplerState<SF_Trilinear, AM_Wrap, AM_Wrap, AM_Wrap>::GetRHI(),
+			TextureRHI
+		);
+	}
+
+
+private:
+
+	FShaderResourceParameter TextureParam;
+	FShaderResourceParameter SamplerParam;
+};
+
+FNiagaraDataInterfaceParametersCS* UNiagaraDataInterfaceTexture::ConstructComputeParameters()const
+{
+	return new FNiagaraDataInterfaceParametersCS_Texture();
+}
\ No newline at end of file
diff --git a/Engine/Plugins/FX/Niagara/Source/Niagara/Private/NiagaraDataInterfaceTexture2D.cpp b/Engine/Plugins/FX/Niagara/Source/Niagara/Private/NiagaraDataInterfaceTexture2D.cpp
deleted file mode 100644
diff --git a/Engine/Plugins/FX/Niagara/Source/NiagaraShader/Private/NiagaraShader.cpp b/Engine/Plugins/FX/Niagara/Source/NiagaraShader/Private/NiagaraShader.cpp
--- a/Engine/Plugins/FX/Niagara/Source/NiagaraShader/Private/NiagaraShader.cpp
+++ b/Engine/Plugins/FX/Niagara/Source/NiagaraShader/Private/NiagaraShader.cpp
@@ -10,6 +10,7 @@
 #include "ShaderCompiler.h"
 #include "NiagaraShaderDerivedDataVersion.h"
 #include "NiagaraShaderCompilationManager.h"
+#include "UObject/CoreRedirects.h"
 #if WITH_EDITOR
 	#include "Interfaces/ITargetPlatformManagerModule.h"
 	#include "TickableEditorObject.h"
@@ -1329,6 +1330,17 @@ void FNiagaraDataInterfaceParamRef::ConstructParameters()
 void FNiagaraDataInterfaceParamRef::InitDIClass()
 {
 	DIClass = Cast<UClass>(StaticFindObject(UClass::StaticClass(), ANY_PACKAGE, *ParameterInfo.DIClassName, true));
+	if (DIClass == nullptr)
+	{
+		FCoreRedirectObjectName OldObjName;
+		OldObjName.ObjectName = *ParameterInfo.DIClassName;
+		FCoreRedirectObjectName NewObjName = FCoreRedirects::GetRedirectedName(ECoreRedirectFlags::Type_Class, OldObjName);
+		if (NewObjName.IsValid())
+		{
+			DIClass = Cast<UClass>(StaticFindObject(UClass::StaticClass(), ANY_PACKAGE, *NewObjName.ObjectName.ToString(), true));
+		}
+
+	}
 	ensureMsgf(DIClass, TEXT("Failed to load class for FNiagaraDataInterfaceParamRef. %s"), *ParameterInfo.DIClassName);
 }
 
