Niagara - Finish Olaf's work on gpu ribbon vetex generation.

diff --git a/Engine/Plugins/FX/Niagara/Source/Niagara/Private/NiagaraRendererRibbons.cpp b/Engine/Plugins/FX/Niagara/Source/Niagara/Private/NiagaraRendererRibbons.cpp
--- a/Engine/Plugins/FX/Niagara/Source/Niagara/Private/NiagaraRendererRibbons.cpp
+++ b/Engine/Plugins/FX/Niagara/Source/Niagara/Private/NiagaraRendererRibbons.cpp
@@ -16,11 +16,27 @@ struct FNiagaraDynamicDataRibbon : public FNiagaraDynamicDataBase
 {
 	TArray<FNiagaraRibbonVertex> VertexData;
 	TArray<int16> IndexData;
-	TArray<FNiagaraRibbonVertexDynamicParameter> MaterialParameterVertexData[4];
+	TArray<int32> SortedIndices;
+	TArray<float> TotalDistances;
+	TArray<float> PackedPerRibbonDataByIndex;
 
 	//Direct ptr to the dataset. ONLY FOR USE BE GPU EMITTERS.
 	//TODO: Even this needs to go soon.
 	const FNiagaraDataSet *DataSet;
+
+	// start and end world space position of the ribbon, to figure out draw direction
+	FVector StartPos;
+	FVector EndPos;
+
+	void PackPerRibbonData(float U0Scale, float U0Offset, float U1Scale, float U1Offset, uint32 NumSegments, uint32 StartVertexIndex)
+	{
+		PackedPerRibbonDataByIndex.Add(U0Scale);
+		PackedPerRibbonDataByIndex.Add(U0Offset);
+		PackedPerRibbonDataByIndex.Add(U1Scale);
+		PackedPerRibbonDataByIndex.Add(U1Offset);
+		PackedPerRibbonDataByIndex.Add(*(float*)&NumSegments);
+		PackedPerRibbonDataByIndex.Add(*(float*)&StartVertexIndex);
+	}
 };
 
 class FNiagaraMeshCollectorResourcesRibbon : public FOneFrameResource
@@ -72,7 +88,6 @@ void NiagaraRendererRibbons::GetDynamicMeshElements(const TArray<const FSceneVie
 	SimpleTimer MeshElementsTimer;
 	FNiagaraDynamicDataRibbon *DynamicDataRibbon = static_cast<FNiagaraDynamicDataRibbon*>(DynamicDataRender);
 	if (!DynamicDataRibbon 
-		|| DynamicDataRibbon->VertexData.Num() < 4
 		|| DynamicDataRibbon->IndexData.Num() == 0
 		|| nullptr == Properties
 		)
@@ -83,32 +98,21 @@ void NiagaraRendererRibbons::GetDynamicMeshElements(const TArray<const FSceneVie
 	const bool bIsWireframe = ViewFamily.EngineShowFlags.Wireframe;
 	FMaterialRenderProxy* MaterialRenderProxy = Material->GetRenderProxy(SceneProxy->IsSelected(), SceneProxy->IsHovered());
 
-	int32 SizeInBytes = DynamicDataRibbon->VertexData.GetTypeSize() * DynamicDataRibbon->VertexData.Num();
-	FGlobalDynamicVertexBuffer::FAllocation LocalDynamicVertexAllocation = FGlobalDynamicVertexBuffer::Get().Allocate(SizeInBytes);
-	FGlobalDynamicVertexBuffer::FAllocation LocalDynamicVertexMaterialParamsAllocation[4];
 	FGlobalDynamicIndexBuffer::FAllocation DynamicIndexAllocation = FGlobalDynamicIndexBuffer::Get().Allocate(DynamicDataRibbon->IndexData.Num(), sizeof(int16));
 
-	auto AllocDynamicParameterBuffer = [&](int32 DynamicParamIndex)
-	{
-		if (DynamicDataRibbon->MaterialParameterVertexData[DynamicParamIndex].Num() > 0)
-		{
-			int32 MatParamSizeInBytes = DynamicDataRibbon->MaterialParameterVertexData[DynamicParamIndex].GetTypeSize() * DynamicDataRibbon->MaterialParameterVertexData[DynamicParamIndex].Num();
-			LocalDynamicVertexMaterialParamsAllocation[DynamicParamIndex] = FGlobalDynamicVertexBuffer::Get().Allocate(MatParamSizeInBytes);
+	int32 TotalFloatSize = DynamicDataRibbon->RTParticleData.GetFloatBuffer().Num() / sizeof(float);
+	FNiagaraGlobalReadBuffer::FAllocation ParticleData;
 
-			if (LocalDynamicVertexMaterialParamsAllocation[DynamicParamIndex].IsValid())
-			{
-				// Copy the extra material vertex data over.
-				FMemory::Memcpy(LocalDynamicVertexMaterialParamsAllocation[DynamicParamIndex].Buffer, DynamicDataRibbon->MaterialParameterVertexData[DynamicParamIndex].GetData(), MatParamSizeInBytes);
-			}
-		}
-	};
-	AllocDynamicParameterBuffer(0);
-	AllocDynamicParameterBuffer(1);
-	AllocDynamicParameterBuffer(2);
-	AllocDynamicParameterBuffer(3);
+	int32 SizeInBytes = DynamicDataRibbon->VertexData.GetTypeSize() * DynamicDataRibbon->VertexData.Num();
+	FGlobalDynamicVertexBuffer::FAllocation LocalDynamicVertexAllocation = FGlobalDynamicVertexBuffer::Get().Allocate(SizeInBytes);
 
+	if (DynamicDataRibbon->DataSet->GetSimTarget() == ENiagaraSimTarget::CPUSim)
+	{
+		ParticleData = FNiagaraGlobalReadBuffer::Get().AllocateFloat(TotalFloatSize);
+		FMemory::Memcpy(ParticleData.Buffer, DynamicDataRibbon->RTParticleData.GetFloatBuffer().GetData(), DynamicDataRibbon->RTParticleData.GetFloatBuffer().Num());
+	}
 
-	if (LocalDynamicVertexAllocation.IsValid())
+	if (DynamicIndexAllocation.IsValid())
 	{
 		// Update the primitive uniform buffer if needed.
 		if (!WorldSpacePrimitiveUniformBuffer.IsInitialized())
@@ -141,18 +145,34 @@ void NiagaraRendererRibbons::GetDynamicMeshElements(const TArray<const FSceneVie
 			{
 				const FSceneView* View = Views[ViewIndex];
 
+				// figure out whether start is closer to the view plane than end
+				float StartDist = FVector::DotProduct( View->GetViewDirection(), DynamicDataRibbon->StartPos - View->ViewLocation);
+				float EndDist = FVector::DotProduct( View->GetViewDirection(), DynamicDataRibbon->EndPos - View->ViewLocation);
+				bool bInvertOrder = ((StartDist > EndDist) && Properties->DrawDirection == ENiagaraRibbonDrawDirection::BackToFront)
+					|| ((StartDist < EndDist) && Properties->DrawDirection == ENiagaraRibbonDrawDirection::FrontToBack);
+
 				FNiagaraMeshCollectorResourcesRibbon& CollectorResources = Collector.AllocateOneFrameResource<FNiagaraMeshCollectorResourcesRibbon>();
 				FNiagaraRibbonUniformParameters PerViewUniformParameters;// = UniformParameters;
+				PerViewUniformParameters.LocalToWorld = bLocalSpace ? SceneProxy->GetLocalToWorld() : FMatrix::Identity;//For now just handle local space like this but maybe in future have a VF variant to avoid the transform entirely?
+				PerViewUniformParameters.LocalToWorldInverseTransposed = bLocalSpace ? SceneProxy->GetLocalToWorld().Inverse().GetTransposed() : FMatrix::Identity;
+				PerViewUniformParameters.DeltaSeconds = ViewFamily.DeltaWorldTime;
 				PerViewUniformParameters.CameraUp = View->GetViewUp(); // FVector4(0.0f, 0.0f, 1.0f, 0.0f);
 				PerViewUniformParameters.CameraRight = View->GetViewRight();//	FVector4(1.0f, 0.0f, 0.0f, 0.0f);
 				PerViewUniformParameters.ScreenAlignment = FVector4(0.0f, 0.0f, 0.0f, 0.0f);
-				PerViewUniformParameters.UseCustomFacing = Properties->FacingMode == ENiagaraRibbonFacingMode::Custom;
+				PerViewUniformParameters.UseCustomFacing = Properties->FacingMode == ENiagaraRibbonFacingMode::Custom ? 1 : 0;
+				PerViewUniformParameters.TotalNumInstances = DynamicDataRibbon->RTParticleData.GetNumInstances();
 
 				PerViewUniformParameters.PositionDataOffset = PositionDataOffset;
+				PerViewUniformParameters.VelocityDataOffset = VelocityDataOffset;
 				PerViewUniformParameters.ColorDataOffset = ColorDataOffset;
 				PerViewUniformParameters.WidthDataOffset = WidthDataOffset;
 				PerViewUniformParameters.TwistDataOffset = TwistDataOffset;
-				CollectorResources.VertexFactory.SetParticleData(DynamicDataRibbon->DataSet);
+				PerViewUniformParameters.FacingDataOffset = FacingDataOffset;
+				PerViewUniformParameters.InvertDrawOrder = bInvertOrder ? 1 : 0;
+				PerViewUniformParameters.UV0TilingDistance = Properties->UV0TilingDistance;
+				PerViewUniformParameters.UV1TilingDistance = Properties->UV1TilingDistance;
+				PerViewUniformParameters.PackedVData = FVector4(Properties->UV0Scale.Y, Properties->UV0Offset.Y, Properties->UV1Scale.Y, Properties->UV1Offset.Y);
+				CollectorResources.VertexFactory.SetParticleData(ParticleData.ReadBuffer->SRV, ParticleData.FirstIndex / sizeof(float), DynamicDataRibbon->RTParticleData.GetFloatStride() / sizeof(float));;
 
 				// Collector.AllocateOneFrameResource uses default ctor, initialize the vertex factory
 				CollectorResources.VertexFactory.SetParticleFactoryType(NVFT_Ribbon);
@@ -160,24 +180,39 @@ void NiagaraRendererRibbons::GetDynamicMeshElements(const TArray<const FSceneVie
 				CollectorResources.UniformBuffer = FNiagaraRibbonUniformBufferRef::CreateUniformBufferImmediate(PerViewUniformParameters, UniformBuffer_SingleFrame);
 
 				CollectorResources.VertexFactory.InitResource();
-				CollectorResources.VertexFactory.SetBeamTrailUniformBuffer(CollectorResources.UniformBuffer);
+				CollectorResources.VertexFactory.SetRibbonUniformBuffer(CollectorResources.UniformBuffer);
 				CollectorResources.VertexFactory.SetVertexBuffer(LocalDynamicVertexAllocation.VertexBuffer, LocalDynamicVertexAllocation.VertexOffset, sizeof(FNiagaraRibbonVertex));
 
-				auto SetDynamicParameterBuffer = [&](int32 DynamicParameterIndex)
+				if (!DynamicDataRibbon->SortedIndices.Num())
 				{
-					if (DynamicDataRibbon->MaterialParameterVertexData[DynamicParameterIndex].Num() > 0 && LocalDynamicVertexMaterialParamsAllocation[DynamicParameterIndex].IsValid())
-					{
-						CollectorResources.VertexFactory.SetDynamicParameterBuffer(LocalDynamicVertexMaterialParamsAllocation[DynamicParameterIndex].VertexBuffer, DynamicParameterIndex, LocalDynamicVertexMaterialParamsAllocation[DynamicParameterIndex].VertexOffset, sizeof(FNiagaraRibbonVertexDynamicParameter));
-					}
-					else
-					{
-						CollectorResources.VertexFactory.SetDynamicParameterBuffer(NULL, DynamicParameterIndex, 0, 0);
-					}
-				};
-				SetDynamicParameterBuffer(0);
-				SetDynamicParameterBuffer(1);
-				SetDynamicParameterBuffer(2);
-				SetDynamicParameterBuffer(3);
+					return;
+				}
+
+				// TODO: need to make these two a global alloc buffer as well, not recreate
+				// pass in the sorted indices so the VS can fetch the particle data in order
+				FReadBuffer SortedIndicesBuffer;
+				SortedIndicesBuffer.Initialize(sizeof(int32), DynamicDataRibbon->SortedIndices.Num(), EPixelFormat::PF_R32_FLOAT, BUF_Volatile);
+				void *IndexPtr = RHILockVertexBuffer(SortedIndicesBuffer.Buffer, 0, DynamicDataRibbon->SortedIndices.Num() * sizeof(int32), RLM_WriteOnly);
+				FMemory::Memcpy(IndexPtr, DynamicDataRibbon->SortedIndices.GetData(), DynamicDataRibbon->SortedIndices.Num() * sizeof(int32));
+				RHIUnlockVertexBuffer(SortedIndicesBuffer.Buffer);
+				CollectorResources.VertexFactory.SetSortedIndices(SortedIndicesBuffer.SRV, 0);
+
+				// pass in the CPU generated total segment distance (for tiling distance modes); needs to be a buffer so we can fetch them in the correct order based on Draw Direction (front->back or back->front)
+				//	otherwise UVs will pop when draw direction changes based on camera view point
+				FReadBuffer TotalDistancesBuffer;
+				TotalDistancesBuffer.Initialize(sizeof(float), DynamicDataRibbon->TotalDistances.Num(), EPixelFormat::PF_R32_FLOAT, BUF_Volatile);
+				void *TotalDistancesPtr = RHILockVertexBuffer(TotalDistancesBuffer.Buffer, 0, DynamicDataRibbon->TotalDistances.Num() * sizeof(float), RLM_WriteOnly);
+				FMemory::Memcpy(TotalDistancesPtr, DynamicDataRibbon->TotalDistances.GetData(), DynamicDataRibbon->TotalDistances.Num() * sizeof(float));
+				RHIUnlockVertexBuffer(TotalDistancesBuffer.Buffer);
+				CollectorResources.VertexFactory.SetSegmentDistances(TotalDistancesBuffer.SRV);
+
+				// Copy the packed u data for stable age based uv generation.
+				FReadBuffer PackedPerRibbonDataByIndexBuffer;
+				PackedPerRibbonDataByIndexBuffer.Initialize(sizeof(float), DynamicDataRibbon->PackedPerRibbonDataByIndex.Num(), EPixelFormat::PF_R32_FLOAT, BUF_Volatile);
+				void *PackedPerRibbonDataByIndexPtr = RHILockVertexBuffer(PackedPerRibbonDataByIndexBuffer.Buffer, 0, DynamicDataRibbon->PackedPerRibbonDataByIndex.Num() * sizeof(float), RLM_WriteOnly);
+				FMemory::Memcpy(PackedPerRibbonDataByIndexPtr, DynamicDataRibbon->PackedPerRibbonDataByIndex.GetData(), DynamicDataRibbon->PackedPerRibbonDataByIndex.Num() * sizeof(float));
+				RHIUnlockVertexBuffer(PackedPerRibbonDataByIndexBuffer.Buffer);
+				CollectorResources.VertexFactory.SetPackedPerRibbonDataByIndexSRV(PackedPerRibbonDataByIndexBuffer.SRV);
 
 				FMeshBatch& MeshBatch = Collector.AllocateMesh();
 				MeshBatch.VertexFactory = &CollectorResources.VertexFactory;
@@ -236,10 +271,6 @@ int NiagaraRendererRibbons::GetDynamicDataSize()
 	{
 		Size += (static_cast<FNiagaraDynamicDataRibbon*>(DynamicDataRender))->VertexData.GetAllocatedSize();
 		Size += (static_cast<FNiagaraDynamicDataRibbon*>(DynamicDataRender))->IndexData.GetAllocatedSize();
-		Size += (static_cast<FNiagaraDynamicDataRibbon*>(DynamicDataRender))->MaterialParameterVertexData[0].GetAllocatedSize();
-		Size += (static_cast<FNiagaraDynamicDataRibbon*>(DynamicDataRender))->MaterialParameterVertexData[1].GetAllocatedSize();
-		Size += (static_cast<FNiagaraDynamicDataRibbon*>(DynamicDataRender))->MaterialParameterVertexData[2].GetAllocatedSize();
-		Size += (static_cast<FNiagaraDynamicDataRibbon*>(DynamicDataRender))->MaterialParameterVertexData[3].GetAllocatedSize();
 	}
 
 	return Size;
@@ -247,7 +278,7 @@ int NiagaraRendererRibbons::GetDynamicDataSize()
 
 bool NiagaraRendererRibbons::HasDynamicData()
 {
-	return DynamicDataRender && (static_cast<FNiagaraDynamicDataRibbon*>(DynamicDataRender))->VertexData.Num() > 0;
+	return DynamicDataRender && (static_cast<FNiagaraDynamicDataRibbon*>(DynamicDataRender))->IndexData.Num() > 0;
 }
 
 #if WITH_EDITORONLY_DATA
@@ -275,10 +306,13 @@ void NiagaraRendererRibbons::TransformChanged()
 	WorldSpacePrimitiveUniformBuffer.ReleaseResource();
 }
 
-void CalculateUVScaleAndOffsets(const FNiagaraDataSetAccessor<float>& SortKeyData, const TArray<int32>& RibbonIndices, bool bSortKeyIsAge, int32 NumSegments,
-	float InUVTilingDistance, const FVector2D& InUVScale, const FVector2D& InUVOffset, ENiagaraRibbonAgeOffsetMode InAgeOffsetMode, FVector2D& OutUVScale, FVector2D& OutUVOffset)
+void CalculateUVScaleAndOffsets(
+	const FNiagaraDataSetAccessor<float>& SortKeyData, const TArray<int32>& RibbonIndices,
+	bool bSortKeyIsAge, int32 StartIndex, int32 EndIndex, int32 NumSegments,
+	float InUTilingDistance, float InUScale, float InUOffset, ENiagaraRibbonAgeOffsetMode InAgeOffsetMode,
+	float& OutUScale, float& OutUOffset)
 {
-	if (bSortKeyIsAge && InUVTilingDistance == 0)
+	if (EndIndex - StartIndex > 0 && bSortKeyIsAge && InUTilingDistance == 0)
 	{
 		float AgeUScale;
 		float AgeUOffset;
@@ -290,10 +324,10 @@ void CalculateUVScaleAndOffsets(const FNiagaraDataSetAccessor<float>& SortKeyDat
 			// is removed.  We calculate the end offset when the end of the ribbon is within a single time step of 0 or 1
 			// which is then normalized to the range of a single segment.  We can then calculate how many segments we actually
 			// have to draw the scaled ribbon, and can offset the start by the correctly scaled offset.
-			float FirstAge = SortKeyData[RibbonIndices[0]];
-			float SecondAge = SortKeyData[RibbonIndices[1]];
-			float SecondToLastAge = SortKeyData[RibbonIndices[RibbonIndices.Num() - 2]];
-			float LastAge = SortKeyData[RibbonIndices.Last()];
+			float FirstAge = SortKeyData[RibbonIndices[StartIndex]];
+			float SecondAge = SortKeyData[RibbonIndices[StartIndex + 1]];
+			float SecondToLastAge = SortKeyData[RibbonIndices[EndIndex - 1]];
+			float LastAge = SortKeyData[RibbonIndices[EndIndex]];
 
 			float StartTimeStep = SecondAge - FirstAge;
 			float StartTimeOffset = FirstAge < StartTimeStep ? StartTimeStep - FirstAge : 0;
@@ -309,20 +343,20 @@ void CalculateUVScaleAndOffsets(const FNiagaraDataSetAccessor<float>& SortKeyDat
 		}
 		else
 		{
-			float FirstAge = SortKeyData[RibbonIndices[0]];
-			float LastAge = SortKeyData[RibbonIndices.Last()];
+			float FirstAge = SortKeyData[RibbonIndices[StartIndex]];
+			float LastAge = SortKeyData[RibbonIndices[EndIndex]];
 
 			AgeUScale = LastAge - FirstAge;
 			AgeUOffset = FirstAge;
 		}
 
-		OutUVScale = FVector2D(AgeUScale * InUVScale.X, InUVScale.Y);
-		OutUVOffset = FVector2D((AgeUOffset * InUVScale.X) + InUVOffset.X, InUVOffset.Y);
+		OutUScale = AgeUScale * InUScale;
+		OutUOffset = (AgeUOffset * InUScale) + InUOffset;
 	}
 	else
 	{
-		OutUVScale = InUVScale;
-		OutUVOffset = InUVOffset;
+		OutUScale = InUScale;
+		OutUOffset = InUOffset;
 	}
 }
 
@@ -335,21 +369,17 @@ FNiagaraDynamicDataBase *NiagaraRendererRibbons::GenerateVertexData(const FNiaga
 	{
 		return nullptr;
 	}
-	FNiagaraDynamicDataRibbon *DynamicData = new FNiagaraDynamicDataRibbon;
-	TArray<FNiagaraRibbonVertex>& RenderData = DynamicData->VertexData;
+	FNiagaraDynamicDataRibbon* DynamicData = new FNiagaraDynamicDataRibbon;
+	TArray<FNiagaraRibbonVertex>& VertexData = DynamicData->VertexData;
 	TArray<int16>& IndexData = DynamicData->IndexData;
 
-	RenderData.Empty();
+	VertexData.Empty();
 	IndexData.Empty();
 
 	// TODO : deal with the dynamic vertex material parameter should the user have specified it as an output...
-
-	FVector2D UVs[2] = { FVector2D(1.0f, 1.0f), FVector2D(1.0f, 0.0f) };
 	int32 NumTotalVerts = 0;
 
 	FNiagaraDataSetAccessor<FVector> PosData(Data, Properties->PositionBinding.DataSetVariable);
-	FNiagaraDataSetAccessor<FVector> VelData(Data, Properties->VelocityBinding.DataSetVariable);
-	FNiagaraDataSetAccessor<FLinearColor> ColData(Data, Properties->ColorBinding.DataSetVariable);
 
 	bool bSortKeyIsAge = false;
 	FNiagaraDataSetAccessor<float> SortKeyData(Data, Properties->RibbonLinkOrderBinding.DataSetVariable);
@@ -367,20 +397,16 @@ FNiagaraDynamicDataBase *NiagaraRendererRibbons::GenerateVertexData(const FNiaga
 
 	if (PositionDataOffset == INDEX_NONE || LastSyncedId != Properties->SyncId)
 	{
-
-		const FNiagaraVariableLayoutInfo* PositionLayout = Data.GetVariableLayout(Properties->PositionBinding.DataSetVariable);
-		const FNiagaraVariableLayoutInfo* TwistLayout = Data.GetVariableLayout(Properties->RibbonTwistBinding.DataSetVariable);
-		const FNiagaraVariableLayoutInfo* WidthLayout = Data.GetVariableLayout(Properties->RibbonWidthBinding.DataSetVariable);
-		const FNiagaraVariableLayoutInfo* ColorLayout = Data.GetVariableLayout(Properties->ColorBinding.DataSetVariable);
-
 		// required attributes
-		PositionDataOffset = PositionLayout ? PositionLayout->FloatComponentStart : INDEX_NONE;
-		ColorDataOffset = ColorLayout ? ColorLayout->FloatComponentStart : INDEX_NONE;
+		int32 IntDummy;
+		Data.GetVariableComponentOffsets(Properties->PositionBinding.DataSetVariable, PositionDataOffset, IntDummy);
+		Data.GetVariableComponentOffsets(Properties->VelocityBinding.DataSetVariable, VelocityDataOffset, IntDummy);
+		Data.GetVariableComponentOffsets(Properties->ColorBinding.DataSetVariable, ColorDataOffset, IntDummy);
 
 		// optional attributes
-		int32 IntDummy;
 		Data.GetVariableComponentOffsets(Properties->RibbonWidthBinding.DataSetVariable, WidthDataOffset, IntDummy);
 		Data.GetVariableComponentOffsets(Properties->RibbonTwistBinding.DataSetVariable, TwistDataOffset, IntDummy);
+		Data.GetVariableComponentOffsets(Properties->RibbonFacingBinding.DataSetVariable, FacingDataOffset, IntDummy);
 		LastSyncedId = Properties->SyncId;
 	}
 
@@ -412,26 +438,16 @@ FNiagaraDynamicDataBase *NiagaraRendererRibbons::GenerateVertexData(const FNiaga
 	bool bSimpleIDs = !bFullIDs && RibbonIdData.IsValid();
 	bool bMultiRibbons = bFullIDs || bSimpleIDs;
 
-	auto AddRibbonVerts = [&](TArray<int32>& RibbonIndices)
+	auto AddRibbonVerts = [&](TArray<int32>& RibbonIndices, uint32 RibbonIndex)
 	{
+		int32 StartIndex = DynamicData->SortedIndices.Num();
 		int32 NumIndices = RibbonIndices.Num();
 		if (NumIndices > 1)
 		{
 			FVector PrevPos, PrevPos2, PrevDir(0.0f, 0.0f, 0.1f);
 			float TotalDistance = 0.0f;
 
-			int32 NumSegments = NumIndices - 1;
-
-			FVector2D UV0Offset;
-			FVector2D UV0Scale;
-			FVector2D UV1Offset;
-			FVector2D UV1Scale;
-
-			CalculateUVScaleAndOffsets(SortKeyData, RibbonIndices, bSortKeyIsAge, NumSegments, Properties->UV0TilingDistance,
-				Properties->UV0Scale, Properties->UV0Offset, Properties->UV0AgeOffsetMode, UV0Scale, UV0Offset);
-			CalculateUVScaleAndOffsets(SortKeyData, RibbonIndices, bSortKeyIsAge, NumSegments, Properties->UV1TilingDistance,
-				Properties->UV1Scale, Properties->UV1Offset, Properties->UV1AgeOffsetMode, UV1Scale, UV1Offset);
-
+			uint32 LastParticipatingParticle = RibbonIndices[0];
 			for (int32 i = 0; i < NumIndices; i++)
 			{
 				uint32 Index1 = RibbonIndices[i];
@@ -441,63 +457,32 @@ FNiagaraDynamicDataBase *NiagaraRendererRibbons::GenerateVertexData(const FNiaga
 				if (i < NumIndices - 1)
 				{
 					Index2 = RibbonIndices[i + 1];
-					ParticleDir = PosData[Index2] - ParticlePos;
+					ParticleDir = PosData[Index2] - PosData[LastParticipatingParticle];
 				}
 				else
 				{
 					Index2 = RibbonIndices[i - 1];
-					ParticleDir = ParticlePos - PosData[Index2];
+					ParticleDir = PosData[LastParticipatingParticle] - PosData[Index2];
 				}
 
 				// if two ribbon particles were spawned too close together, we skip one
 				// but never skip the last, because that will result in invalid indices from the prev loop
 				if (ParticleDir.SizeSquared() > 0.002 || i == NumIndices - 1)
 				{
+					DynamicData->SortedIndices.Add(Index1);
+
+					LastParticipatingParticle = Index2;
 					FVector NormDir = ParticleDir.GetSafeNormal();
 					PrevDir = NormDir;
 
-					// Calculate sub-uv for this segment based on the either the distance or index, depending on whether the distance tile factor was set.
-					float U0ForSegment = Properties->UV0TilingDistance
-						? TotalDistance / Properties->UV0TilingDistance
-						: (float)i / NumSegments;
-					float U1ForSegment = Properties->UV1TilingDistance
-						? TotalDistance / Properties->UV1TilingDistance
-						: (float)i / NumSegments;
-
-					FVector2D UV0ForSegment(U0ForSegment, 1.0f);
-					FVector2D UV1ForSegment(U1ForSegment, 1.0f);
-					
-					//Todo. Possibly pull out to templated func to avoid safety code in loop here.
-					FLinearColor Color = ColData.GetSafe(Index1, FLinearColor::White);
-					float Twist = TwistData.GetSafe(Index1, 0.0f);
-					float Size = SizeData.GetSafe(Index1, 1.0f);
-					FVector Align = AlignData.GetSafe(Index1, FVector::UpVector);
-
-					TotalDistance += ParticleDir.Size();
-					AddRibbonVert(RenderData, ParticlePos, UVs[0] * UV0ForSegment * UV0Scale + UV0Offset, UVs[0] * UV1ForSegment * UV1Scale + UV1Offset, Color, Twist, Size, NormDir, Align);
-					AddRibbonVert(RenderData, ParticlePos, UVs[1] * UV0ForSegment * UV0Scale + UV0Offset, UVs[1] * UV1ForSegment * UV1Scale + UV1Offset, Color, Twist, Size, NormDir, Align);
-					if (MaterialParamData.IsValid())
-					{
-						AddDynamicParam(DynamicData->MaterialParameterVertexData[0], MaterialParamData[Index1]);
-						AddDynamicParam(DynamicData->MaterialParameterVertexData[0], MaterialParamData[Index1]);
-					}
-					if (MaterialParam1Data.IsValid())
-					{
-						AddDynamicParam(DynamicData->MaterialParameterVertexData[1], MaterialParam1Data[Index1]);
-						AddDynamicParam(DynamicData->MaterialParameterVertexData[1], MaterialParam1Data[Index1]);
-					}
-					if (MaterialParam2Data.IsValid())
-					{
-						AddDynamicParam(DynamicData->MaterialParameterVertexData[2], MaterialParam2Data[Index1]);
-						AddDynamicParam(DynamicData->MaterialParameterVertexData[2], MaterialParam2Data[Index1]);
-					}
-					if (MaterialParam3Data.IsValid())
+					DynamicData->TotalDistances.Add(TotalDistance);
+					FNiagaraRibbonVertex NewVertex;
+					NewVertex.RibbonIndex = RibbonIndex;
+					VertexData.Add(NewVertex);
+					VertexData.Add(NewVertex);
+
+					if (i < NumIndices - 1)
 					{
-						AddDynamicParam(DynamicData->MaterialParameterVertexData[3], MaterialParam3Data[Index1]);
-						AddDynamicParam(DynamicData->MaterialParameterVertexData[3], MaterialParam3Data[Index1]);
-					}
- 					if (i < NumIndices - 1)
- 					{
 						IndexData.Add(NumTotalVerts);
 						IndexData.Add(NumTotalVerts + 1);
 						IndexData.Add(NumTotalVerts + 2);
@@ -507,16 +492,37 @@ FNiagaraDynamicDataBase *NiagaraRendererRibbons::GenerateVertexData(const FNiaga
 					}
 					NumTotalVerts += 2;
 				}
+
+				TotalDistance += ParticleDir.Size();
 			}
 		}
+
+		float U0Offset;
+		float U0Scale;
+		float U1Offset;
+		float U1Scale;
+
+		int32 EndIndex = DynamicData->SortedIndices.Num() - 1;
+		int32 NumSegments = EndIndex - StartIndex;
+		int32 StartVertexIndex = StartIndex * 2; // We add two vertices for each particle.
+
+		CalculateUVScaleAndOffsets(SortKeyData, DynamicData->SortedIndices, bSortKeyIsAge, StartIndex, DynamicData->SortedIndices.Num() - 1, NumSegments,
+			Properties->UV0TilingDistance,	Properties->UV0Scale.X, Properties->UV0Offset.X, Properties->UV0AgeOffsetMode, U0Scale, U0Offset);
+		CalculateUVScaleAndOffsets(SortKeyData, DynamicData->SortedIndices, bSortKeyIsAge, StartIndex, DynamicData->SortedIndices.Num() - 1, NumSegments,
+			Properties->UV1TilingDistance, Properties->UV1Scale.X, Properties->UV1Offset.X, Properties->UV1AgeOffsetMode, U1Scale, U1Offset);
+
+		DynamicData->PackPerRibbonData(U0Scale, U0Offset, U1Scale, U1Offset, NumSegments, StartVertexIndex);
 	};
+
+	// store the start and end positions for the ribbon for draw distance flipping 
+	DynamicData->StartPos = PosData[0];
+	DynamicData->EndPos = PosData[Data.GetNumInstances() - 1];
 	
 	//TODO: Move sorting to share code with sprite and mesh sorting and support the custom sorting key.
 	int32 TotalIndices = Data.GetNumInstances();
 
 	if (!bMultiRibbons)
 	{
-		int32 CurIdx = 0;
 		TArray<int32> SortedIndices;
 		for (int32 i = 0; i < TotalIndices; ++i)
 		{
@@ -525,11 +531,10 @@ FNiagaraDynamicDataBase *NiagaraRendererRibbons::GenerateVertexData(const FNiaga
 
 		SortedIndices.Sort([&SortKeyData](const int32& A, const int32& B) {	return (SortKeyData[A] < SortKeyData[B]); });
 
-		AddRibbonVerts(SortedIndices);
+		AddRibbonVerts(SortedIndices, 0);
 	}
 	else
 	{
-		int32 CurIdx = 0;
 		if (bFullIDs)
 		{
 			TMap<FNiagaraID, TArray<int32>> MultiRibbonSortedIndices;
@@ -540,13 +545,14 @@ FNiagaraDynamicDataBase *NiagaraRendererRibbons::GenerateVertexData(const FNiaga
 				Indices.Add(i);
 			}
 
+			uint32 RibbonIndex = 0;
 			for (TPair<FNiagaraID, TArray<int32>>& Pair : MultiRibbonSortedIndices)
 			{
-				FNiagaraID& ID = Pair.Key;
 				TArray<int32>& SortedIndices = Pair.Value;
-
 				SortedIndices.Sort([&SortKeyData](const int32& A, const int32& B) {	return (SortKeyData[A] < SortKeyData[B]); });
-				AddRibbonVerts(SortedIndices);			
+				AddRibbonVerts(SortedIndices, RibbonIndex);
+				
+				RibbonIndex++;
 			};
 		}
 		else
@@ -562,17 +568,25 @@ FNiagaraDynamicDataBase *NiagaraRendererRibbons::GenerateVertexData(const FNiaga
 				Indices.Add(i);
 			}
 
+			uint32 RibbonIndex = 0;
 			for (TPair<int32, TArray<int32>>& Pair : MultiRibbonSortedIndices)
 			{
-				int32 ID = Pair.Key;
 				TArray<int32>& SortedIndices = Pair.Value;
-
 				SortedIndices.Sort([&SortKeyData](const int32& A, const int32& B) {	return (SortKeyData[A] < SortKeyData[B]); });
-				AddRibbonVerts(SortedIndices);
+				AddRibbonVerts(SortedIndices, RibbonIndex);
+				RibbonIndex++;
 			};
 		}
 	}
 
+	if (Data.CurrData().GetNumInstances() > 0)
+	{
+		//TODO: This buffer is far fatter than needed. Just pull out the data needed for rendering.
+		Data.CurrData().CopyTo(DynamicData->RTParticleData);
+
+		DynamicData->DataSet = &Data;
+	}
+
 	CPUTimeMS = VertexDataTimer.GetElapsedMilliseconds();
 
 	return DynamicData;
diff --git a/Engine/Plugins/FX/Niagara/Source/Niagara/Public/NiagaraRendererRibbons.h b/Engine/Plugins/FX/Niagara/Source/Niagara/Public/NiagaraRendererRibbons.h
--- a/Engine/Plugins/FX/Niagara/Source/Niagara/Public/NiagaraRendererRibbons.h
+++ b/Engine/Plugins/FX/Niagara/Source/Niagara/Public/NiagaraRendererRibbons.h
@@ -36,23 +36,6 @@ public:
 	/** Update render data buffer from attributes */
 	FNiagaraDynamicDataBase *GenerateVertexData(const FNiagaraSceneProxy* Proxy, FNiagaraDataSet &Data, const ENiagaraSimTarget Target) override;
 
-	void AddRibbonVert(TArray<FNiagaraRibbonVertex>& RenderData, FVector ParticlePos, FVector2D UV1, FVector2D UV2,
-		const FLinearColor &Color, const float Rotation, const float Size, const FVector NormDir, const FVector CustomFacing)
-	{
-		FNiagaraRibbonVertex NewVertex;
-		NewVertex.Position = ParticlePos;
-		NewVertex.Direction = NormDir;
-		NewVertex.Color = Color;
-		NewVertex.Size = Size;
-		NewVertex.Rotation = Rotation;
-		NewVertex.Tex_U = UV1.X;
-		NewVertex.Tex_V = UV1.Y;
-		NewVertex.Tex_U2 = UV2.X;
-		NewVertex.Tex_V2 = UV2.Y;
-		NewVertex.CustomFacingVector = CustomFacing;
-		RenderData.Add(NewVertex);
-	}
-
 	void AddDynamicParam(TArray<FNiagaraRibbonVertexDynamicParameter>& ParamData, const FVector4& DynamicParam)
 	{
 		FNiagaraRibbonVertexDynamicParameter Param;
@@ -76,7 +59,8 @@ public:
 
 	UClass *GetPropertiesClass() override { return UNiagaraRibbonRendererProperties::StaticClass(); }
 	void SetRendererProperties(UNiagaraRendererProperties *Props) override { Properties = Cast<UNiagaraRibbonRendererProperties>(Props); }
-	virtual UNiagaraRendererProperties* GetRendererProperties() const override {
+	virtual UNiagaraRendererProperties* GetRendererProperties() const override 
+	{
 		return Properties;
 	}
 
@@ -85,8 +69,10 @@ private:
 	UNiagaraRibbonRendererProperties *Properties;
 	mutable TUniformBuffer<FPrimitiveUniformShaderParameters> WorldSpacePrimitiveUniformBuffer;
 	int32 PositionDataOffset;
+	int32 VelocityDataOffset;
 	int32 WidthDataOffset;
 	int32 TwistDataOffset;
+	int32 FacingDataOffset;
 	int32 ColorDataOffset;
 	int32 LastSyncedId;
 };
diff --git a/Engine/Plugins/FX/Niagara/Source/Niagara/Public/NiagaraRibbonRendererProperties.h b/Engine/Plugins/FX/Niagara/Source/Niagara/Public/NiagaraRibbonRendererProperties.h
--- a/Engine/Plugins/FX/Niagara/Source/Niagara/Public/NiagaraRibbonRendererProperties.h
+++ b/Engine/Plugins/FX/Niagara/Source/Niagara/Public/NiagaraRibbonRendererProperties.h
@@ -25,6 +25,14 @@ enum class ENiagaraRibbonAgeOffsetMode : uint8
 	Clip
 };
 
+/** This enum decides in which order the ribbon segments will be rendered */
+UENUM()
+enum class ENiagaraRibbonDrawDirection : uint8
+{
+	FrontToBack,
+	BackToFront
+};
+
 UCLASS(editinlinenew)
 class UNiagaraRibbonRendererProperties : public UNiagaraRendererProperties
 {
@@ -73,6 +81,10 @@ public:
 	UPROPERTY(EditAnywhere, Category = "Ribbon Rendering")
 	ENiagaraRibbonAgeOffsetMode UV1AgeOffsetMode;
 
+	/** If true, the particles are only sorted when using a translucent material. */
+	UPROPERTY(EditAnywhere, Category = "Ribbon Rendering")
+	ENiagaraRibbonDrawDirection DrawDirection;
+
 	//~ UNiagaraRendererProperties interface
 	virtual NiagaraRenderer* CreateEmitterRenderer(ERHIFeatureLevel::Type FeatureLevel) override;
 	virtual void GetUsedMaterials(TArray<UMaterialInterface*>& OutMaterials) const override;
diff --git a/Engine/Plugins/FX/Niagara/Source/NiagaraVertexFactories/Private/NiagaraRibbonVertexFactory.cpp b/Engine/Plugins/FX/Niagara/Source/NiagaraVertexFactories/Private/NiagaraRibbonVertexFactory.cpp
--- a/Engine/Plugins/FX/Niagara/Source/NiagaraVertexFactories/Private/NiagaraRibbonVertexFactory.cpp
+++ b/Engine/Plugins/FX/Niagara/Source/NiagaraVertexFactories/Private/NiagaraRibbonVertexFactory.cpp
@@ -32,44 +32,57 @@ class FNiagaraRibbonVertexFactoryShaderParametersVS : public FNiagaraRibbonVerte
 public:
 	virtual void Bind(const FShaderParameterMap& ParameterMap) override
 	{
-		/*NiagaraParticleDataFloat.Bind(ParameterMap, TEXT("NiagaraParticleDataFloat"));
-		NiagaraParticleDataInt.Bind(ParameterMap, TEXT("NiagaraParticleDataInt"));
-		SafeComponentBufferSizeParam.Bind(ParameterMap, TEXT("SafeComponentBufferSize"));*/
+		NiagaraParticleDataFloat.Bind(ParameterMap, TEXT("NiagaraParticleDataFloat"));
+		FloatDataOffset.Bind(ParameterMap, TEXT("NiagaraFloatDataOffset"));
+		FloatDataStride.Bind(ParameterMap, TEXT("NiagaraFloatDataStride"));
+		SortedIndices.Bind(ParameterMap, TEXT("SortedIndices"));
+		SortedIndicesOffset.Bind(ParameterMap, TEXT("SortedIndicesOffset"));
+		SegmentDistances.Bind(ParameterMap, TEXT("SegmentDistances"));
+		PackedPerRibbonDataByIndex.Bind(ParameterMap, TEXT("PackedPerRibbonDataByIndex"));
+
+		ensure(NiagaraParticleDataFloat.IsBound());
+		ensure(FloatDataOffset.IsBound());
+		ensure(FloatDataStride.IsBound());
+		ensure(SortedIndices.IsBound());
+		ensure(SortedIndicesOffset.IsBound());
 	}
 
 	virtual void Serialize(FArchive& Ar) override
 	{
-		/*Ar << PositionOffsetParam;
-		Ar << WidthOffsetParam;
-		Ar << TwistOffsetParam;
-		Ar << ColorOffsetParam;
-
 		Ar << NiagaraParticleDataFloat;
-		Ar << NiagaraParticleDataInt;
-		Ar << SafeComponentBufferSizeParam;*/
+		Ar << FloatDataOffset;
+		Ar << FloatDataStride;
+		Ar << SortedIndices;
+		Ar << SortedIndicesOffset;
+		Ar << SegmentDistances;
+		Ar << PackedPerRibbonDataByIndex;
 	}
 
 	virtual void SetMesh(FRHICommandList& RHICmdList, FShader* Shader, const FVertexFactory* VertexFactory, const FSceneView& View, const FMeshBatchElement& BatchElement, uint32 DataFlags) const override
 	{
 		FNiagaraRibbonVertexFactory* RibbonVF = (FNiagaraRibbonVertexFactory*)VertexFactory;
 		FVertexShaderRHIParamRef VertexShaderRHI = Shader->GetVertexShader();
-		SetUniformBufferParameter(RHICmdList, Shader->GetVertexShader(), Shader->GetUniformBufferParameter<FNiagaraRibbonUniformParameters>(), RibbonVF->GetBeamTrailUniformBuffer());
+		SetUniformBufferParameter(RHICmdList, Shader->GetVertexShader(), Shader->GetUniformBufferParameter<FNiagaraRibbonUniformParameters>(), RibbonVF->GetRibbonUniformBuffer());
+
+		SetSRVParameter(RHICmdList, VertexShaderRHI, NiagaraParticleDataFloat, RibbonVF->GetParticleDataFloatSRV());
+		SetShaderValue(RHICmdList, VertexShaderRHI, FloatDataOffset, RibbonVF->GetFloatDataOffset());
+		SetShaderValue(RHICmdList, VertexShaderRHI, FloatDataStride, RibbonVF->GetFloatDataStride());
 
-		/*
-		SetSRVParameter(RHICmdList, VertexShaderRHI, NiagaraParticleDataFloat, RibbonVF->GetFloatDataSRV());
-		SetSRVParameter(RHICmdList, VertexShaderRHI, NiagaraParticleDataInt, RibbonVF->GetIntDataSRV());
-		SetShaderValue(RHICmdList, VertexShaderRHI, SafeComponentBufferSizeParam, RibbonVF->GetComponentBufferSize());
-		*/
+		SetSRVParameter(RHICmdList, VertexShaderRHI, SortedIndices, RibbonVF->GetSortedIndicesSRV());
+		SetSRVParameter(RHICmdList, VertexShaderRHI, SegmentDistances, RibbonVF->GetSegmentDistancesSRV());
+		SetSRVParameter(RHICmdList, VertexShaderRHI, PackedPerRibbonDataByIndex, RibbonVF->GetPackedPerRibbonDataByIndexSRV());
+		SetShaderValue(RHICmdList, VertexShaderRHI, SortedIndicesOffset, RibbonVF->GetSortedIndicesOffset());
 	}
 
 private:
-	FShaderParameter PositionOffsetParam;
-	FShaderParameter WidthOffsetParam;
-	FShaderParameter TwistOffsetParam;
-	FShaderParameter ColorOffsetParam;
 	FShaderResourceParameter NiagaraParticleDataFloat;
-	FShaderResourceParameter NiagaraParticleDataInt;
-	FShaderParameter SafeComponentBufferSizeParam;
+	FShaderParameter FloatDataOffset;
+	FShaderParameter FloatDataStride;
+
+	FShaderResourceParameter SortedIndices;
+	FShaderResourceParameter SegmentDistances;
+	FShaderResourceParameter PackedPerRibbonDataByIndex;
+	FShaderParameter SortedIndicesOffset;
 };
 
 
@@ -91,14 +104,14 @@ public:
 	virtual void SetMesh(FRHICommandList& RHICmdList, FShader* Shader, const FVertexFactory* VertexFactory, const FSceneView& View, const FMeshBatchElement& BatchElement, uint32 DataFlags) const override
 	{
 		FNiagaraRibbonVertexFactory* RibbonVF = (FNiagaraRibbonVertexFactory*)VertexFactory;
-		SetUniformBufferParameter(RHICmdList, Shader->GetPixelShader(), Shader->GetUniformBufferParameter<FNiagaraRibbonUniformParameters>(), RibbonVF->GetBeamTrailUniformBuffer());
+		SetUniformBufferParameter(RHICmdList, Shader->GetPixelShader(), Shader->GetUniformBufferParameter<FNiagaraRibbonUniformParameters>(), RibbonVF->GetRibbonUniformBuffer());
 	}
 };
 
 
 ///////////////////////////////////////////////////////////////////////////////
 /**
-* The particle system beam trail vertex declaration resource type.
+* The Niagara ribbon vertex declaration resource type.
 */
 class FNiagaraRibbonVertexDeclaration : public FRenderResource
 {
@@ -110,37 +123,10 @@ public:
 
 	virtual void FillDeclElements(FVertexDeclarationElementList& Elements, int32& Offset)
 	{
-		uint16 Stride = sizeof(FNiagaraRibbonVertex);
-		/** The stream to read the vertex position from. */
-		Elements.Add(FVertexElement(0, Offset, VET_Float3, 0, Stride));
-		Offset += sizeof(float) * 3;
-		/** The stream to read the direction from. */
-		Elements.Add(FVertexElement(0, Offset, VET_Float3, 1, Stride));
-		Offset += sizeof(float) * 3;
-		/** The stream to read the vertex size from. */
-		Elements.Add(FVertexElement(0, Offset, VET_Float1, 2, Stride));
-		Offset += sizeof(float) * 1;
-		/** The stream to read the color from.					*/
-		Elements.Add(FVertexElement(0, Offset, VET_Float4, 3, Stride));
-		Offset += sizeof(float) * 4;
-		/** The stream to read the texture coordinates from.	*/
-		Elements.Add(FVertexElement(0, Offset, VET_Float4, 4, Stride));
-		Offset += sizeof(float) * 4;
-
-		/** The stream to read the twist from. */
-		Elements.Add(FVertexElement(0, Offset, VET_Float1, 5, Stride));
-		Offset += sizeof(float);
-
-		/** The stream to read the custom alignment vector from. */
-		Elements.Add(FVertexElement(0, Offset, VET_Float3, 6, Stride));
-		Offset += sizeof(float) * 3;
-
-		/** Dynamic parameters come from additional streams */
-		Elements.Add(FVertexElement(1, 0, VET_Float4, 7, sizeof(FVector4)));
-		Elements.Add(FVertexElement(2, 0, VET_Float4, 8, sizeof(FVector4)));
-		Elements.Add(FVertexElement(3, 0, VET_Float4, 9, sizeof(FVector4)));
-		Elements.Add(FVertexElement(4, 0, VET_Float4, 10, sizeof(FVector4)));
-
+		uint32 Stride = sizeof(FNiagaraRibbonVertex);
+		/** The stream to read the ribbon index from. */
+		Elements.Add(FVertexElement(0, Offset, VET_UByte4, 0, Stride));
+		Offset += sizeof(int32);
 	}
 
 	virtual void InitDynamicRHI()
@@ -150,7 +136,7 @@ public:
 		FillDeclElements(Elements, Offset);
 
 		// Create the vertex declaration for rendering the factory normally.
-		// This is done in InitDynamicRHI instead of InitRHI to allow FParticleBeamTrailVertexFactory::InitRHI
+		// This is done in InitDynamicRHI instead of InitRHI to allow FNiagaraRibbonVertexFactory::InitRHI
 		// to rely on it being initialized, since InitDynamicRHI is called before InitRHI.
 		VertexDeclarationRHI = RHICreateVertexDeclaration(Elements);
 	}
@@ -177,7 +163,7 @@ bool FNiagaraRibbonVertexFactory::ShouldCompilePermutation(EShaderPlatform Platf
 void FNiagaraRibbonVertexFactory::ModifyCompilationEnvironment(EShaderPlatform Platform, const class FMaterial* Material, FShaderCompilerEnvironment& OutEnvironment)
 {
 	FNiagaraVertexFactoryBase::ModifyCompilationEnvironment(Platform, Material, OutEnvironment);
-	OutEnvironment.SetDefine(TEXT("PARTICLE_BEAMTRAIL_FACTORY"), TEXT("1"));
+	OutEnvironment.SetDefine(TEXT("NIAGARA_RIBBON_FACTORY"), TEXT("1"));
 }
 
 /**
diff --git a/Engine/Plugins/FX/Niagara/Source/NiagaraVertexFactories/Public/NiagaraRibbonVertexFactory.h b/Engine/Plugins/FX/Niagara/Source/NiagaraVertexFactories/Public/NiagaraRibbonVertexFactory.h
--- a/Engine/Plugins/FX/Niagara/Source/NiagaraVertexFactories/Public/NiagaraRibbonVertexFactory.h
+++ b/Engine/Plugins/FX/Niagara/Source/NiagaraVertexFactories/Public/NiagaraRibbonVertexFactory.h
@@ -17,41 +17,19 @@ class FMaterial;
 
 
 
-//	FParticleBeamTrailVertex
+//	FNiagaraRibbonVertex
 struct FNiagaraRibbonVertex
 {
-	/** The position of the particle. */
-	FVector Position;
-	/** The direction of the particle. */
-	FVector Direction;
-	/** The size of the particle. */
-	float Size;
-	/** The color of the particle. */
-	FLinearColor Color;
-
-	/** The second UV set for the particle				*/
-	float			Tex_U;
-	float			Tex_V;
-	float			Tex_U2;
-	float			Tex_V2;
-
-
-	/** The rotation of the particle. */
-	float Rotation;
-
-	/** The relative time of the particle. */
-	FVector CustomFacingVector;
-
-	/** The relative time of the particle. */
-	//float NormalizedAge;
+	/** The index of the ribbon for multi-ribbon emitters */
+	uint32 RibbonIndex;
 };
 
 
-//	FParticleBeamTrailVertexDynamicParameter
+//	FNiagaraRibbonVertexDynamicParameter
 struct FNiagaraRibbonVertexDynamicParameter
 {
 	/** The dynamic parameter of the particle			*/
-	float			DynamicValue[4];
+	float DynamicValue[4];
 };
 
 /**
@@ -62,11 +40,21 @@ BEGIN_UNIFORM_BUFFER_STRUCT(FNiagaraRibbonUniformParameters, NIAGARAVERTEXFACTOR
 	UNIFORM_MEMBER(FVector4, CameraUp)
 	UNIFORM_MEMBER(FVector4, ScreenAlignment)
 	UNIFORM_MEMBER(int, PositionDataOffset)
+	UNIFORM_MEMBER(int, VelocityDataOffset)
 	UNIFORM_MEMBER(int, WidthDataOffset)
 	UNIFORM_MEMBER(int, TwistDataOffset)
 	UNIFORM_MEMBER(int, ColorDataOffset)
+	UNIFORM_MEMBER(int, FacingDataOffset)
 	UNIFORM_MEMBER(int, MaterialParamDataOffset)
-	UNIFORM_MEMBER(bool, UseCustomFacing)
+	UNIFORM_MEMBER(int, TotalNumInstances)
+	UNIFORM_MEMBER(uint32, UseCustomFacing)
+	UNIFORM_MEMBER(uint32, InvertDrawOrder)
+	UNIFORM_MEMBER(float, UV0TilingDistance)
+	UNIFORM_MEMBER(float, UV1TilingDistance)
+	UNIFORM_MEMBER(FVector4, PackedVData)
+	UNIFORM_MEMBER_EX(FMatrix, LocalToWorld, EShaderPrecisionModifier::Half)
+	UNIFORM_MEMBER_EX(FMatrix, LocalToWorldInverseTransposed, EShaderPrecisionModifier::Half)
+	UNIFORM_MEMBER_EX(float, DeltaSeconds, EShaderPrecisionModifier::Half)
 	END_UNIFORM_BUFFER_STRUCT(FNiagaraRibbonUniformParameters)
 typedef TUniformBufferRef<FNiagaraRibbonUniformParameters> FNiagaraRibbonUniformBufferRef;
 
@@ -75,7 +63,7 @@ typedef TUniformBufferRef<FNiagaraRibbonUniformParameters> FNiagaraRibbonUniform
 */
 class NIAGARAVERTEXFACTORIES_API FNiagaraRibbonVertexFactory : public FNiagaraVertexFactoryBase
 {
-	DECLARE_VERTEX_FACTORY_TYPE(FParticleBeamTrailVertexFactory);
+	DECLARE_VERTEX_FACTORY_TYPE(FNiagaraRibbonVertexFactory);
 
 public:
 
@@ -112,17 +100,17 @@ public:
 	/**
 	* Set the uniform buffer for this vertex factory.
 	*/
-	FORCEINLINE void SetBeamTrailUniformBuffer(FNiagaraRibbonUniformBufferRef InSpriteUniformBuffer)
+	FORCEINLINE void SetRibbonUniformBuffer(FNiagaraRibbonUniformBufferRef InSpriteUniformBuffer)
 	{
-		BeamTrailUniformBuffer = InSpriteUniformBuffer;
+		NiagaraRibbonUniformBuffer = InSpriteUniformBuffer;
 	}
 
 	/**
 	* Retrieve the uniform buffer for this vertex factory.
 	*/
-	FORCEINLINE FNiagaraRibbonUniformBufferRef GetBeamTrailUniformBuffer()
+	FORCEINLINE FNiagaraRibbonUniformBufferRef GetRibbonUniformBuffer()
 	{
-		return BeamTrailUniformBuffer;
+		return NiagaraRibbonUniformBuffer;
 	}
 
 	/**
@@ -136,11 +124,63 @@ public:
 	void SetDynamicParameterBuffer(const FVertexBuffer* InDynamicParameterBuffer, int32 ParameterIndex, uint32 StreamOffset, uint32 Stride);
 
 
-	void SetParticleData(const FNiagaraDataSet *InDataSet)
+	void SetParticleData(const FShaderResourceViewRHIRef& InParticleDataFloatSRV, uint32 InFloatDataOffset, uint32 InFloatDataStride)
 	{
-		DataSet = InDataSet;
+		ParticleDataFloatSRV = InParticleDataFloatSRV;
+		FloatDataOffset = InFloatDataOffset;
+		FloatDataStride = InFloatDataStride;
 	}
 
+	void SetSortedIndices(const FShaderResourceViewRHIRef& InSortedIndicesSRV, uint32 InSortedIndicesOffset)
+	{
+		SortedIndicesSRV = InSortedIndicesSRV;
+		SortedIndicesOffset = InSortedIndicesOffset;
+	}
+
+	void SetSegmentDistances(const FShaderResourceViewRHIRef& InSegmentDistancesSRV)
+	{
+		SegmentDistancesSRV = InSegmentDistancesSRV;
+	}
+
+	void SetPackedPerRibbonDataByIndexSRV(const FShaderResourceViewRHIRef& InPackedPerRibbonDataByIndexSRV)
+	{
+		PackedPerRibbonDataByIndexSRV = InPackedPerRibbonDataByIndexSRV;
+	}
+
+	FORCEINLINE FShaderResourceViewRHIRef GetParticleDataFloatSRV()
+	{
+		return ParticleDataFloatSRV;
+	}
+
+	FORCEINLINE int32 GetFloatDataOffset()
+	{
+		return FloatDataOffset;
+	}
+
+	FORCEINLINE int32 GetFloatDataStride()
+	{
+		return FloatDataStride;
+	}
+
+	FORCEINLINE FShaderResourceViewRHIRef GetSortedIndicesSRV()
+	{
+		return SortedIndicesSRV;
+	}
+
+	FORCEINLINE int32 GetSortedIndicesOffset()
+	{
+		return SortedIndicesOffset;
+	}
+
+	FORCEINLINE FShaderResourceViewRHIRef GetSegmentDistancesSRV()
+	{
+		return SegmentDistancesSRV;
+	}
+
+	FORCEINLINE FShaderResourceViewRHIRef GetPackedPerRibbonDataByIndexSRV()
+	{
+		return PackedPerRibbonDataByIndexSRV;
+	}
 
 	/**
 	* Construct shader parameters for this type of vertex factory.
@@ -165,7 +205,7 @@ public:
 private:
 
 	/** Uniform buffer with beam/trail parameters. */
-	FNiagaraRibbonUniformBufferRef BeamTrailUniformBuffer;
+	FNiagaraRibbonUniformBufferRef NiagaraRibbonUniformBuffer;
 
 	/** Used to hold the index buffer allocation information when we call GDME more than once per frame. */
 	FIndexBuffer* IndexBuffer;
@@ -173,4 +213,15 @@ private:
 	int32 OutTriangleCount;
 
 	const FNiagaraDataSet *DataSet;
+
+	FShaderResourceViewRHIRef ParticleDataFloatSRV;
+	uint32 FloatDataOffset;
+	uint32 FloatDataStride;
+
+	FShaderResourceViewRHIRef SortedIndicesSRV;
+	FShaderResourceViewRHIRef SegmentDistancesSRV;
+	FShaderResourceViewRHIRef PackedPerRibbonDataByIndexSRV;
+
+	uint32 SortedIndicesOffset;
+
 };
\ No newline at end of file
diff --git a/Engine/Shaders/Private/NiagaraRibbonVertexFactory.ush b/Engine/Shaders/Private/NiagaraRibbonVertexFactory.ush
--- a/Engine/Shaders/Private/NiagaraRibbonVertexFactory.ush
+++ b/Engine/Shaders/Private/NiagaraRibbonVertexFactory.ush
@@ -5,25 +5,26 @@
 =============================================================================*/
 
 #include "VertexFactoryCommon.ush"
+#include "NiagaraVFParticleAccess.usf"
 
 
 
 
 float3 TransformPosition(float3 InPosition)
 {
-	return mul(float4(InPosition,1.0), NiagaraSpriteVF.LocalToWorld).xyz;
+	return mul(float4(InPosition,1.0), NiagaraRibbonVF.LocalToWorld).xyz;
 }
 
 float3 TransformVector(float3 InVector)
 {
-	return mul(float4(InVector,0.0), NiagaraSpriteVF.LocalToWorldInverseTransposed).xyz;
+	return mul(float4(InVector,0.0), NiagaraRibbonVF.LocalToWorldInverseTransposed).xyz;
 }
 
 
 /*********************************************************************************
  *							Particle specific
  *********************************************************************************/
-/*
+
 float3 GetNiagaraParticlePosition(uint InstanceID)
 {
 	if (NiagaraRibbonVF.PositionDataOffset == -1)
@@ -32,6 +33,15 @@ float3 GetNiagaraParticlePosition(uint InstanceID)
 	}
 	return TransformPosition(GetVec3(NiagaraRibbonVF.PositionDataOffset, InstanceID));
 }
+float3 GetNiagaraParticleVelocity(uint InstanceID)
+{
+	if(NiagaraRibbonVF.VelocityDataOffset  == -1)
+	{
+		// See FNiagaraConstants for a synchronized value...
+		return TransformVector(float3(0.0f,0.0f,0.0f));
+	}
+	return TransformVector(GetVec3(NiagaraRibbonVF.VelocityDataOffset, InstanceID));
+}
 float GetNiagaraParticleRibbonWidth(uint InstanceID)
 {
 	if (NiagaraRibbonVF.WidthDataOffset == -1)
@@ -102,53 +112,30 @@ float4 GetNiagaraParticleDynamicParameters3(uint InstanceID)
 float3 GetNiagaraParticleFacingVector(uint InstanceID)
 {
 	float3 Val = float3(0.0f, 0.0f, 1.0f);
-    if(NiagaraRibbonVF.FacingOffset != -1)
+    if(NiagaraRibbonVF.FacingDataOffset != -1)
 	{
-		Val = GetVec3(NiagaraRibbonVF.FacingOffset, InstanceID);
+		Val = GetVec3(NiagaraRibbonVF.FacingDataOffset, InstanceID);
 	}
 	return TransformVector(Val);
 }
-*/
 
 
+Buffer<int> SortedIndices;
+uint SortedIndicesOffset;
 
+Buffer<float> SegmentDistances;
+Buffer<float> PackedPerRibbonDataByIndex;
 
 struct FVertexFactoryInput
 {
-	float3	Position				: ATTRIBUTE0;
-	float3 Direction				: ATTRIBUTE1;
-	float	Size					: ATTRIBUTE2;
-	float4	Color					: ATTRIBUTE3;
-	float4	TexCoord				: ATTRIBUTE4;
+	uint VertexId					: SV_VertexID;
 
-	float Twist						: ATTRIBUTE5;
-	float3 CustomFacingVector		: ATTRIBUTE6;
+	uint RibbonIndex				: ATTRIBUTE0;
 
-#if NUM_DYNAMIC_PARAMETERS > 0
-	float4 DynamicParameter0		: ATTRIBUTE7;
-#endif
-#if NUM_DYNAMIC_PARAMETERS > 1
-	float4 DynamicParameter1		: ATTRIBUTE8;
-#endif
-#if NUM_DYNAMIC_PARAMETERS > 2
-	float4 DynamicParameter2		: ATTRIBUTE9;
-#endif
-#if NUM_DYNAMIC_PARAMETERS > 3
-	float4 DynamicParameter3		: ATTRIBUTE10;
-#endif
-
-/** Optional instance ID for vertex layered rendering */
-#if FEATURE_LEVEL >= FEATURE_LEVEL_SM4
-	uint VertexID	: SV_VertexID;
-#endif
-
-/** Optional instance ID for vertex layered rendering */
+	// Optional instance ID for vertex layered rendering
 #if FEATURE_LEVEL >= FEATURE_LEVEL_SM4 && ONEPASS_POINTLIGHT_SHADOW && USING_VERTEX_SHADER_LAYER
 	uint InstanceId	: SV_InstanceID;
 #endif
-
-	// soon
-	//float NormalizedAge			: ATTRIBUTE8;
 };
 
 struct FVertexFactoryInterpolantsVSToPS
@@ -158,23 +145,25 @@ struct FVertexFactoryInterpolantsVSToPS
 	// Last row of the tangent to world matrix in xyz
 	float4	TangentToWorld2	: TANGENTTOWORLD2;
 
-	float4	TexCoord			: TEXCOORD2;
-
-#if NEEDS_PARTICLE_COLOR
-	float4  Color				: TEXCOORD3;
-#endif
-
 #if NUM_DYNAMIC_PARAMETERS > 0
-	float4 DynamicParameter0	: TEXCOORD4;
+	float4 DynamicParameter0	: PARTICLE_DYNAMIC_PARAM0;
 #endif
 #if NUM_DYNAMIC_PARAMETERS > 1
-	float4 DynamicParameter1	: TEXCOORD5;
+	float4 DynamicParameter1	: PARTICLE_DYNAMIC_PARAM1;
 #endif
 #if NUM_DYNAMIC_PARAMETERS > 2
-	float4 DynamicParameter2	: TEXCOORD6;
+	float4 DynamicParameter2	: PARTICLE_DYNAMIC_PARAM2;
 #endif
 #if NUM_DYNAMIC_PARAMETERS > 3
-	float4 DynamicParameter3	: TEXCOORD7;
+	float4 DynamicParameter3	: PARTICLE_DYNAMIC_PARAM3;
+#endif
+
+#if NUM_MATERIAL_TEXCOORDS
+	float4	TexCoords[(NUM_MATERIAL_TEXCOORDS + 1) / 2]	: TEXCOORD3;
+#endif
+
+#if NEEDS_PARTICLE_COLOR
+	float4  Color				: TEXCOORD2;
 #endif
 
 #if INSTANCED_STEREO
@@ -184,13 +173,110 @@ struct FVertexFactoryInterpolantsVSToPS
 
 struct FVertexFactoryIntermediates
 {
-	/** The color of the sprite. */
+	float3 Position;
+	float3 OldPosition;
+	float3 Direction;
+	float3 CustomFacingVector;
 	float4 Color;
 	float RibbonWidth;
-	float3 Right, Up;
+	float Twist;
+	float3 Right;
+	float3 Up;
 	float3 CameraDirection;
+	float4 TexCoord;
+
+#if NUM_DYNAMIC_PARAMETERS > 0
+	float4 DynamicParameter;
+#endif
+#if NUM_DYNAMIC_PARAMETERS > 1
+	float4 DynamicParameter1;
+#endif
+#if NUM_DYNAMIC_PARAMETERS > 2
+	float4 DynamicParameter2;
+#endif
+#if NUM_DYNAMIC_PARAMETERS > 3
+	float4 DynamicParameter3;
+#endif
 };
 
+
+#if NUM_MATERIAL_TEXCOORDS
+bool UVIndexUseZW(int UVIndex)
+{
+	#if COMPILER_GLSL_ES2
+	return (frac((float)UVIndex / 2.0f) == 0.5f);
+	#else
+	return (UVIndex % 2) != 0;
+	#endif
+}
+
+
+float2 GetUV(FVertexFactoryInterpolantsVSToPS Interpolants, int UVIndex)
+{
+	float4 UVVector = Interpolants.TexCoords[UVIndex / 2];
+	return UVIndexUseZW(UVIndex) ? UVVector.zw : UVVector.xy;
+}
+
+void SetUV(inout FVertexFactoryInterpolantsVSToPS Interpolants, int UVIndex, float2 InValue)
+{
+	FLATTEN
+	if (UVIndexUseZW(UVIndex))
+	{
+		Interpolants.TexCoords[UVIndex / 2].zw = InValue;
+	}
+	else
+	{
+		Interpolants.TexCoords[UVIndex / 2].xy = InValue;
+	}
+}
+
+#endif
+
+int GetParticleIdByVertexId(uint VertexId)
+{
+	uint ParticleId = VertexId / 2;
+
+	/*
+	if(NiagaraRibbonVF.InvertDrawOrder == 1)
+	{
+		ParticleId = NiagaraRibbonVF.TotalNumInstances - ParticleId;
+	}
+	*/
+
+	if(SortedIndicesOffset != 0xFFFFFFFF)
+	{
+		ParticleId = SortedIndices[SortedIndicesOffset + ParticleId];
+	}
+
+	return ParticleId;
+}
+
+int GetParticleId(FVertexFactoryInput Input)
+{
+	return GetParticleIdByVertexId(Input.VertexId);
+}
+
+int GetPreviousParticleId(FVertexFactoryInput Input)
+{
+	return GetParticleIdByVertexId(Input.VertexId - 2);
+}
+
+int GetNextParticleId(FVertexFactoryInput Input)
+{
+	return GetParticleIdByVertexId(Input.VertexId + 2);
+}
+
+void UnpackPerRibbonDataByIndex(int RibbonIndex, out float2 UV0Scale, out float2 UV0Offset, out float2 UV1Scale, out float2 UV1Offset, out uint NumSegments, out uint StartVertexId)
+{
+	int Index = RibbonIndex * 6;
+	UV0Scale = float2(PackedPerRibbonDataByIndex[Index], NiagaraRibbonVF.PackedVData.x);
+	UV0Offset = float2(PackedPerRibbonDataByIndex[Index + 1], NiagaraRibbonVF.PackedVData.y);
+	UV1Scale = float2(PackedPerRibbonDataByIndex[Index + 2], NiagaraRibbonVF.PackedVData.z);
+	UV1Offset = float2(PackedPerRibbonDataByIndex[Index + 3], NiagaraRibbonVF.PackedVData.w);
+	NumSegments = asuint(PackedPerRibbonDataByIndex[Index + 4]);
+	StartVertexId = asuint(PackedPerRibbonDataByIndex[Index + 5]);
+}
+
 /** Converts from vertex factory specific interpolants to a FMaterialPixelParameters, which is used by material inputs. */
 FMaterialPixelParameters GetMaterialPixelParameters(FVertexFactoryInterpolantsVSToPS Interpolants, float4 SvPosition)
 {
@@ -198,20 +284,13 @@ FMaterialPixelParameters GetMaterialPixelParameters(FVertexFactoryInterpolantsVS
 	FMaterialPixelParameters Result = MakeInitializedMaterialPixelParameters();
 
 #if NUM_MATERIAL_TEXCOORDS
-	#if NUM_MATERIAL_TEXCOORDS >= 1
-		Result.TexCoords[0] = Interpolants.TexCoord.xy;
-	#if NUM_MATERIAL_TEXCOORDS >= 2
-		Result.TexCoords[1] = Interpolants.TexCoord.zw;
-	#endif	// >= 2
-	#endif	// >= 1
-	#if NUM_MATERIAL_TEXCOORDS > 2
-		UNROLL
-		for(int CoordinateIndex = 2;CoordinateIndex < NUM_MATERIAL_TEXCOORDS;CoordinateIndex++)
-		{
-			Result.TexCoords[CoordinateIndex] = Interpolants.TexCoord.xy;
-		}
-	#endif	// > 2
-#endif
+	UNROLL
+	for (int CoordinateIndex = 0; CoordinateIndex < NUM_MATERIAL_TEXCOORDS; CoordinateIndex++)
+	{
+		Result.TexCoords[CoordinateIndex] = GetUV(Interpolants, CoordinateIndex);
+	}
+#endif	// NUM_MATERIAL_TEXCOORDS
+
 	Result.VertexColor = 1;
 
 #if NEEDS_PARTICLE_COLOR
@@ -223,14 +302,10 @@ FMaterialPixelParameters GetMaterialPixelParameters(FVertexFactoryInterpolantsVS
 	Result.TangentToWorld = AssembleTangentToWorld( TangentToWorld0, TangentToWorld2 );
 	Result.UnMirrored = 1;
 
-#if LIGHTMAP_UV_ACCESS
-	Result.LightmapUVs = Interpolants.TexCoord.xy;
-#endif
-
 	Result.TwoSidedSign = 1;
-	
+
 #if NUM_DYNAMIC_PARAMETERS > 0
-	Result.Particle.DynamicParameter = Interpolants.DynamicParameter0;
+	Result.Particle.DynamicParameter0 = Interpolants.DynamicParameter0;
 #endif
 #if NUM_DYNAMIC_PARAMETERS > 1
 	Result.Particle.DynamicParameter1 = Interpolants.DynamicParameter1;
@@ -250,17 +325,34 @@ FMaterialVertexParameters GetMaterialVertexParameters(FVertexFactoryInput Input,
 {
 	FMaterialVertexParameters Result = (FMaterialVertexParameters)0;
 	Result.WorldPosition = WorldPosition;
-	Result.VertexColor = Input.Color;
+	Result.VertexColor = Intermediates.Color;
 	Result.TangentToWorld = mul(TangentToLocal, GetLocalToWorld3x3()); 
 	Result.Particle.Color = Intermediates.Color;
-	Result.PreSkinnedPosition = Input.Position.xyz;
+	Result.PreSkinnedPosition = Intermediates.Position.xyz;
 	Result.PreSkinnedNormal = TangentToLocal[2].xyz;
 
+#if NUM_DYNAMIC_PARAMETERS > 0
+	Result.Particle.DynamicParameter = Intermediates.DynamicParameter;
+#endif
+#if NUM_DYNAMIC_PARAMETERS > 1
+	Result.Particle.DynamicParameter1 = Intermediates.DynamicParameter1;
+#endif
+#if NUM_DYNAMIC_PARAMETERS > 2
+	Result.Particle.DynamicParameter2 = Intermediates.DynamicParameter2;
+#endif
+#if NUM_DYNAMIC_PARAMETERS > 3
+	Result.Particle.DynamicParameter3 = Intermediates.DynamicParameter3;
+#endif
+
+#if USE_PARTICLE_VELOCITY
+	Result.Particle.Velocity = Intermediates.ParticleVelocity;
+#endif
+
 #if NUM_MATERIAL_TEXCOORDS_VERTEX
-	for(int CoordinateIndex = 0; CoordinateIndex < NUM_MATERIAL_TEXCOORDS_VERTEX; CoordinateIndex++)
-	{
-		Result.TexCoords[CoordinateIndex] = Input.TexCoord.xy;
-	}
+	Result.TexCoords[0] = Intermediates.TexCoord.xy;
+	#if NUM_MATERIAL_TEXCOORDS_VERTEX > 1
+		Result.TexCoords[1] = Intermediates.TexCoord.zw;
+	#endif
 #endif
 	return Result;
 }
@@ -280,40 +372,28 @@ float3 SafeNormalize(float3 V)
 	return V / sqrt(max(dot(V,V),0.01));
 }
 
-void GetTangents(FVertexFactoryInput Input,out float3 Right,out float3 Up, out float3 CameraDirection)
+void GetTangents(FVertexFactoryInput Input, FVertexFactoryIntermediates Intermediates, out float3 Right,out float3 Up, out float3 CameraDirection)
 {
-	float4	Position			= TransformLocalToTranslatedWorld(Input.Position.xyz);
-
-	CameraDirection				= SafeNormalize(ResolvedView.WorldCameraOrigin - Input.Position.xyz );		// TODO: handle local space
-	float3 ParticleDirection	= Input.Direction;
-
-	float4	Right_Square		= BeamTrailVF.CameraRight,
-			Up_Square			= BeamTrailVF.CameraUp;		
-			
-	float4	Right_Rotated		= (-1.0 * cos(Input.Twist) * Up_Square) + (sin(Input.Twist) * Right_Square),
-			Up_Rotated			= (       sin(Input.Twist) * Up_Square) + (cos(Input.Twist) * Right_Square);
+	float4	Position			= TransformLocalToTranslatedWorld(Intermediates.Position.xyz);
 
-	float4	Right_Velocity		= float4( SafeNormalize( cross( CameraDirection, ParticleDirection	) ), 0.0 ),
-			Up_Velocity			= float4( ParticleDirection, 0.0 );
+	CameraDirection				= SafeNormalize(ResolvedView.WorldCameraOrigin - Intermediates.Position.xyz );		// TODO: handle local space
+	float3 ParticleDirection	= Intermediates.Direction;
 
 	float3 NormDir = SafeNormalize(ParticleDirection);
 	float3 RightDir, Facing;
 
 
-	/*if(NiagaraRibbonVF.UseCustomFacing)
+	if(NiagaraRibbonVF.UseCustomFacing)
 	{
-		Facing = Input.CustomFacingVector;
+		Facing = Intermediates.CustomFacingVector;
 	}
-	else*/
+	else
 	{
 		Facing = CameraDirection.xyz;
 	}
 
-	if(Input.Twist != 0.0f)
-	{
-		float cosa = cos(Input.Twist);
-		Facing = Facing * cosa + cross(Facing, Input.Direction) * sin(Input.Twist) + Input.Direction * dot(Input.Direction, Facing) * (1-cosa);
-	}
+	float cosa = cos(Intermediates.Twist);
+	Facing = Facing * cosa + cross(Facing, NormDir) * sin(Intermediates.Twist) + NormDir * dot(NormDir, Facing) * (1-cosa);
 
 	if( abs(dot(NormDir, Facing)) < 0.999f)
 	{
@@ -327,50 +407,117 @@ void GetTangents(FVertexFactoryInput Input,out float3 Right,out float3 Up, out f
 	Right = RightDir;
 	Up = NormDir;
 	return;
-	/*
-	//	enum EParticleScreenAlignment
-	//	{
-	//		PSA_Square,
-	//		PSA_Rectangle,
-	//		PSA_Velocity
-	//	};
-	Right				= BeamTrailVF.ScreenAlignment.x > 1.5f ? Right_Velocity : Right_Rotated;
-	Up					= BeamTrailVF.ScreenAlignment.x > 1.5f ? Up_Velocity	: Up_Rotated;
-	*/
 }
 
 float4 CalcWorldPosition(FVertexFactoryInput Input, FVertexFactoryIntermediates Intermediates)
 {
 	// expand and align
-	int Sgn = (((int)Input.VertexID & 0x1) * 2) - 1;
-	float3 NormDir = SafeNormalize(Input.Direction.xyz);
+	int Sgn = (((int)Input.VertexId & 0x1) * 2) - 1;
+	float3 NormDir = SafeNormalize(Intermediates.Direction.xyz);
 	float3 RightDir = Intermediates.Right;
 	float3 UpDir = Intermediates.Up;
-	float3 Pos = Input.Position.xyz + RightDir * Sgn * Input.Size.xxx;
+	float3 Pos = Intermediates.Position.xyz + RightDir * Sgn * Intermediates.RibbonWidth.xxx;
 	return TransformLocalToTranslatedWorld(Pos);
 }
 
 /** derive basis vectors */
 float3x3 CalcTangentBasis(FVertexFactoryInput Input, FVertexFactoryIntermediates Intermediates)
 {
-	float4	Right,
+	float3	Right,
 			Up;
 	float3	CameraDir;
-	//GetTangents(Input,Right,Up, CameraDir);
+	GetTangents(Input, Intermediates, Right, Up, CameraDir);
 	return float3x3(
 			Intermediates.Right.xyz,
 			Intermediates.Up.xyz,
-			-normalize(cross(Intermediates.Right.xyz, Intermediates.Up.xyz))
+			normalize(cross(Intermediates.Right.xyz, Intermediates.Up.xyz))
 			);
 }
 
 FVertexFactoryIntermediates GetVertexFactoryIntermediates(FVertexFactoryInput Input)
 {
 	FVertexFactoryIntermediates Intermediates;
+	Intermediates = (FVertexFactoryIntermediates)0;
+	uint ParticleId = GetParticleId(Input);
+	
+	float2 UV0Scale;
+	float2 UV0Offset;
+	float2 UV1Scale;
+	float2 UV1Offset;
+	uint NumSegments;
+	uint StartVertexId;
+	UnpackPerRibbonDataByIndex(Input.RibbonIndex, UV0Scale, UV0Offset, UV1Scale, UV1Offset, NumSegments,  StartVertexId);
+	
+	uint StartParticleId = GetParticleIdByVertexId(StartVertexId);
+
+	Intermediates.Position = GetNiagaraParticlePosition(ParticleId);
+	if(ParticleId == StartParticleId)
+	{
+		uint NextParticleId = GetNextParticleId(Input);
+		float3 NextParticlePos = GetNiagaraParticlePosition(NextParticleId);
+		Intermediates.Direction = NextParticlePos - Intermediates.Position;
+	}
+	else
+	{
+		uint PreviousParticleId = GetPreviousParticleId(Input);
+		float3 PreviousParticlePos = GetNiagaraParticlePosition(PreviousParticleId);
+		Intermediates.Direction = Intermediates.Position - PreviousParticlePos;
+	}
+	
+	float SegmentLength  = length(Intermediates.Direction);
+	Intermediates.OldPosition = Intermediates.Position.xyz - GetNiagaraParticleVelocity(ParticleId) * NiagaraRibbonVF.DeltaSeconds;
+	Intermediates.Color	= GetNiagaraParticleColor(ParticleId);
+	
+	if(NiagaraRibbonVF.UseCustomFacing)
+	{
+		Intermediates.CustomFacingVector = GetNiagaraParticleFacingVector(ParticleId);
+	}
 
-	Intermediates.Color = Input.Color;
-    Intermediates.RibbonWidth = Input.Size;
-	GetTangents(Input, Intermediates.Right, Intermediates.Up, Intermediates.CameraDirection);
+#if NUM_DYNAMIC_PARAMETERS > 0
+	Intermediates.DynamicParameter = GetNiagaraParticleDynamicParameters(ParticleId);
+#endif
+#if NUM_DYNAMIC_PARAMETERS > 1
+	Intermediates.DynamicParameter1 = GetNiagaraParticleDynamicParameters1(ParticleId);
+#endif
+#if NUM_DYNAMIC_PARAMETERS > 2
+	Intermediates.DynamicParameter2 = GetNiagaraParticleDynamicParameters2(ParticleId);
+#endif
+#if NUM_DYNAMIC_PARAMETERS > 3
+	Intermediates.DynamicParameter3 = GetNiagaraParticleDynamicParameters3(ParticleId);
+#endif
+
+    Intermediates.RibbonWidth = GetNiagaraParticleRibbonWidth(ParticleId);
+    Intermediates.Twist = GetNiagaraParticleRibbonTwist(ParticleId);
+
+	GetTangents(Input, Intermediates, Intermediates.Right, Intermediates.Up, Intermediates.CameraDirection);
+
+	float U0ForSegment;
+	if(NiagaraRibbonVF.UV0TilingDistance != 0.0f)
+	{
+		U0ForSegment = SegmentDistances[Input.VertexId / 2] / NiagaraRibbonVF.UV0TilingDistance;
+	}
+	else
+	{
+		U0ForSegment = (float)((Input.VertexId - StartVertexId) / 2) / NumSegments;
+	}
+
+	float U1ForSegment;
+	if(NiagaraRibbonVF.UV1TilingDistance != 0.0f)
+	{
+		U1ForSegment = SegmentDistances[Input.VertexId / 2] / NiagaraRibbonVF.UV0TilingDistance;
+	}
+	else
+	{
+		U1ForSegment = (float)((Input.VertexId - StartVertexId) / 2) / NumSegments;
+	}
+
+	float VForSegment = (Input.VertexId + 1) & 0x1;
+	
+	float2 UV0ForSegment = float2(U0ForSegment, VForSegment);
+	float2 UV1ForSegment = float2(U1ForSegment, VForSegment);
+	
+	Intermediates.TexCoord.xy = UV0ForSegment * UV0Scale + UV0Offset;
+	Intermediates.TexCoord.zw = UV1ForSegment * UV1Scale + UV1Offset;
 
 	return Intermediates;
 }
@@ -393,7 +540,14 @@ float3 VertexFactoryGetPositionForVertexLighting(FVertexFactoryInput Input, FVer
 FVertexFactoryInterpolantsVSToPS VertexFactoryGetInterpolantsVSToPS(FVertexFactoryInput Input, FVertexFactoryIntermediates Intermediates, FMaterialVertexParameters VertexParameters)
 {
 	FVertexFactoryInterpolantsVSToPS Interpolants;
-	Interpolants.TexCoord	= Input.TexCoord;
+	Interpolants = (FVertexFactoryInterpolantsVSToPS)0;
+
+#if NUM_MATERIAL_TEXCOORDS
+	float2 UV1 = Intermediates.TexCoord.xy;
+	float2 UV2 = Intermediates.TexCoord.zw;
+	SetUV(Interpolants, 0, UV1);
+	SetUV(Interpolants, 1, UV2);
+#endif
 
 #if NEEDS_PARTICLE_COLOR
 	Interpolants.Color		= Intermediates.Color;
@@ -425,9 +579,11 @@ FVertexFactoryInterpolantsVSToPS VertexFactoryGetInterpolantsVSToPS(FVertexFacto
 	return Interpolants;
 }
 
+
 float4 VertexFactoryGetPreviousWorldPosition(FVertexFactoryInput Input, FVertexFactoryIntermediates Intermediates)
 {
-	return VertexFactoryGetWorldPosition(Input, Intermediates);
+	uint ParticleId = GetParticleId(Input);
+	return float4(CalcWorldPosition(Input, Intermediates).xyz - GetNiagaraParticleVelocity(ParticleId) * NiagaraRibbonVF.DeltaSeconds, 1);
 }
 
 /**
