Niagara - Fix dynamic material parameters for ribbons and add an automated test for it.

diff --git a/Engine/Plugins/FX/Niagara/Source/Niagara/Private/NiagaraRendererRibbons.cpp b/Engine/Plugins/FX/Niagara/Source/Niagara/Private/NiagaraRendererRibbons.cpp
--- a/Engine/Plugins/FX/Niagara/Source/Niagara/Private/NiagaraRendererRibbons.cpp
+++ b/Engine/Plugins/FX/Niagara/Source/Niagara/Private/NiagaraRendererRibbons.cpp
@@ -168,6 +168,10 @@ void NiagaraRendererRibbons::GetDynamicMeshElements(const TArray<const FSceneVie
 				PerViewUniformParameters.WidthDataOffset = WidthDataOffset;
 				PerViewUniformParameters.TwistDataOffset = TwistDataOffset;
 				PerViewUniformParameters.FacingDataOffset = FacingDataOffset;
+				PerViewUniformParameters.MaterialParamDataOffset = MaterialParamOffset;
+				PerViewUniformParameters.MaterialParam1DataOffset = MaterialParamOffset1;
+				PerViewUniformParameters.MaterialParam2DataOffset = MaterialParamOffset2;
+				PerViewUniformParameters.MaterialParam3DataOffset = MaterialParamOffset3;
 				PerViewUniformParameters.InvertDrawOrder = bInvertOrder ? 1 : 0;
 				PerViewUniformParameters.UV0TilingDistance = Properties->UV0TilingDistance;
 				PerViewUniformParameters.UV1TilingDistance = Properties->UV1TilingDistance;
@@ -407,6 +411,12 @@ FNiagaraDynamicDataBase *NiagaraRendererRibbons::GenerateVertexData(const FNiaga
 		Data.GetVariableComponentOffsets(Properties->RibbonWidthBinding.DataSetVariable, WidthDataOffset, IntDummy);
 		Data.GetVariableComponentOffsets(Properties->RibbonTwistBinding.DataSetVariable, TwistDataOffset, IntDummy);
 		Data.GetVariableComponentOffsets(Properties->RibbonFacingBinding.DataSetVariable, FacingDataOffset, IntDummy);
+
+		Data.GetVariableComponentOffsets(Properties->DynamicMaterialBinding.DataSetVariable, MaterialParamOffset, IntDummy);
+		Data.GetVariableComponentOffsets(Properties->DynamicMaterial1Binding.DataSetVariable, MaterialParamOffset1, IntDummy);
+		Data.GetVariableComponentOffsets(Properties->DynamicMaterial2Binding.DataSetVariable, MaterialParamOffset2, IntDummy);
+		Data.GetVariableComponentOffsets(Properties->DynamicMaterial3Binding.DataSetVariable, MaterialParamOffset3, IntDummy);
+
 		LastSyncedId = Properties->SyncId;
 	}
 
diff --git a/Engine/Plugins/FX/Niagara/Source/Niagara/Public/NiagaraRendererRibbons.h b/Engine/Plugins/FX/Niagara/Source/Niagara/Public/NiagaraRendererRibbons.h
--- a/Engine/Plugins/FX/Niagara/Source/Niagara/Public/NiagaraRendererRibbons.h
+++ b/Engine/Plugins/FX/Niagara/Source/Niagara/Public/NiagaraRendererRibbons.h
@@ -75,6 +75,10 @@ private:
 	int32 FacingDataOffset;
 	int32 ColorDataOffset;
 	int32 LastSyncedId;
+	int32 MaterialParamOffset;
+	int32 MaterialParamOffset1;
+	int32 MaterialParamOffset2;
+	int32 MaterialParamOffset3;
 };
 
 
diff --git a/Engine/Plugins/FX/Niagara/Source/NiagaraVertexFactories/Public/NiagaraRibbonVertexFactory.h b/Engine/Plugins/FX/Niagara/Source/NiagaraVertexFactories/Public/NiagaraRibbonVertexFactory.h
--- a/Engine/Plugins/FX/Niagara/Source/NiagaraVertexFactories/Public/NiagaraRibbonVertexFactory.h
+++ b/Engine/Plugins/FX/Niagara/Source/NiagaraVertexFactories/Public/NiagaraRibbonVertexFactory.h
@@ -46,6 +46,9 @@ BEGIN_UNIFORM_BUFFER_STRUCT(FNiagaraRibbonUniformParameters, NIAGARAVERTEXFACTOR
 	UNIFORM_MEMBER(int, ColorDataOffset)
 	UNIFORM_MEMBER(int, FacingDataOffset)
 	UNIFORM_MEMBER(int, MaterialParamDataOffset)
+	UNIFORM_MEMBER(int, MaterialParam1DataOffset)
+	UNIFORM_MEMBER(int, MaterialParam2DataOffset)
+	UNIFORM_MEMBER(int, MaterialParam3DataOffset)
 	UNIFORM_MEMBER(int, TotalNumInstances)
 	UNIFORM_MEMBER(uint32, UseCustomFacing)
 	UNIFORM_MEMBER(uint32, InvertDrawOrder)
diff --git a/Engine/Shaders/Private/NiagaraRibbonVertexFactory.ush b/Engine/Shaders/Private/NiagaraRibbonVertexFactory.ush
--- a/Engine/Shaders/Private/NiagaraRibbonVertexFactory.ush
+++ b/Engine/Shaders/Private/NiagaraRibbonVertexFactory.ush
@@ -146,7 +146,7 @@ struct FVertexFactoryInterpolantsVSToPS
 	float4	TangentToWorld2	: TANGENTTOWORLD2;
 
 #if NUM_DYNAMIC_PARAMETERS > 0
-	float4 DynamicParameter0	: PARTICLE_DYNAMIC_PARAM0;
+	float4 DynamicParameter	: PARTICLE_DYNAMIC_PARAM0;
 #endif
 #if NUM_DYNAMIC_PARAMETERS > 1
 	float4 DynamicParameter1	: PARTICLE_DYNAMIC_PARAM1;
@@ -305,7 +305,7 @@ FMaterialPixelParameters GetMaterialPixelParameters(FVertexFactoryInterpolantsVS
 	Result.TwoSidedSign = 1;
 
 #if NUM_DYNAMIC_PARAMETERS > 0
-	Result.Particle.DynamicParameter0 = Interpolants.DynamicParameter0;
+	Result.Particle.DynamicParameter = Interpolants.DynamicParameter;
 #endif
 #if NUM_DYNAMIC_PARAMETERS > 1
 	Result.Particle.DynamicParameter1 = Interpolants.DynamicParameter1;
@@ -554,16 +554,16 @@ FVertexFactoryInterpolantsVSToPS VertexFactoryGetInterpolantsVSToPS(FVertexFacto
 #endif
 
 #if NUM_DYNAMIC_PARAMETERS > 0
-	Interpolants.DynamicParameter0 = Input.DynamicParameter0;
+	Interpolants.DynamicParameter = Intermediates.DynamicParameter;
 #endif
 #if NUM_DYNAMIC_PARAMETERS > 1
-	Interpolants.DynamicParameter1 = Input.DynamicParameter1;
+	Interpolants.DynamicParameter1 = Intermediates.DynamicParameter1;
 #endif
 #if NUM_DYNAMIC_PARAMETERS > 2
-	Interpolants.DynamicParameter2 = Input.DynamicParameter2;
+	Interpolants.DynamicParameter2 = Intermediates.DynamicParameter2;
 #endif
 #if NUM_DYNAMIC_PARAMETERS > 3
-	Interpolants.DynamicParameter3 = Input.DynamicParameter3;
+	Interpolants.DynamicParameter3 = Intermediates.DynamicParameter3;
 #endif
 
 	float3x3 TangentToLocal = CalcTangentBasis(Input, Intermediates);
