Niagara - Fix a few places which were not using the "Equivilent" versions when dealing with script usages.

diff --git a/Engine/Plugins/FX/Niagara/Source/NiagaraEditor/Private/NiagaraHlslTranslator.cpp b/Engine/Plugins/FX/Niagara/Source/NiagaraEditor/Private/NiagaraHlslTranslator.cpp
--- a/Engine/Plugins/FX/Niagara/Source/NiagaraEditor/Private/NiagaraHlslTranslator.cpp
+++ b/Engine/Plugins/FX/Niagara/Source/NiagaraEditor/Private/NiagaraHlslTranslator.cpp
@@ -710,8 +710,8 @@ const FNiagaraTranslateResults &FHlslNiagaraTranslator::Translate(const FNiagara
 		TranslationStages[1].PassNamespace = TEXT("MapUpdate");
 		TranslationStages[0].ChunkModeIndex = ENiagaraCodeChunkMode::SpawnBody;
 		TranslationStages[1].ChunkModeIndex = ENiagaraCodeChunkMode::UpdateBody;
-		TranslationStages[0].OutputNode = SourceGraph->FindOutputNode(ENiagaraScriptUsage::ParticleSpawnScript, TranslationStages[0].UsageId);
-		TranslationStages[1].OutputNode = SourceGraph->FindOutputNode(TranslationStages[1].ScriptUsage, TranslationStages[1].UsageId);
+		TranslationStages[0].OutputNode = SourceGraph->FindEquivalentOutputNode(ENiagaraScriptUsage::ParticleSpawnScript, TranslationStages[0].UsageId);
+		TranslationStages[1].OutputNode = SourceGraph->FindEquivalentOutputNode(TranslationStages[1].ScriptUsage, TranslationStages[1].UsageId);
 		TranslationStages[1].bInterpolatePreviousParams = true;
 		ParamMapHistories.AddDefaulted(2);
 		ParamMapSetVariablesToChunks.AddDefaulted(2);
@@ -723,8 +723,8 @@ const FNiagaraTranslateResults &FHlslNiagaraTranslator::Translate(const FNiagara
 		TranslationStages[1].PassNamespace = TEXT("MapUpdate");
 		TranslationStages[0].ChunkModeIndex = ENiagaraCodeChunkMode::SpawnBody;
 		TranslationStages[1].ChunkModeIndex = ENiagaraCodeChunkMode::UpdateBody;
-		TranslationStages[0].OutputNode = SourceGraph->FindOutputNode(ENiagaraScriptUsage::ParticleSpawnScript, TranslationStages[0].UsageId);
-		TranslationStages[1].OutputNode = SourceGraph->FindOutputNode(TranslationStages[1].ScriptUsage, TranslationStages[1].UsageId);
+		TranslationStages[0].OutputNode = SourceGraph->FindEquivalentOutputNode(ENiagaraScriptUsage::ParticleSpawnScript, TranslationStages[0].UsageId);
+		TranslationStages[1].OutputNode = SourceGraph->FindEquivalentOutputNode(TranslationStages[1].ScriptUsage, TranslationStages[1].UsageId);
 		TranslationStages[1].bInterpolatePreviousParams = bHasInterpolatedSpawn;
 		ParamMapHistories.AddDefaulted(2);
 		ParamMapSetVariablesToChunks.AddDefaulted(2);
@@ -732,7 +732,7 @@ const FNiagaraTranslateResults &FHlslNiagaraTranslator::Translate(const FNiagara
 	default:
 		TranslationStages.Add(FHlslNiagaraTranslationStage(CompileOptions.TargetUsage, CompileOptions.TargetUsageId));
 		TranslationStages[0].PassNamespace = TEXT("Map");		
-		TranslationStages[0].OutputNode = SourceGraph->FindOutputNode(TranslationStages[0].ScriptUsage, TranslationStages[0].UsageId);
+		TranslationStages[0].OutputNode = SourceGraph->FindEquivalentOutputNode(TranslationStages[0].ScriptUsage, TranslationStages[0].UsageId);
 		TranslationStages[0].ChunkModeIndex = ENiagaraCodeChunkMode::Body;
 		ParamMapHistories.AddDefaulted(1);
 		ParamMapSetVariablesToChunks.AddDefaulted(1);
diff --git a/Engine/Plugins/FX/Niagara/Source/NiagaraEditor/Private/NiagaraScriptMergeManager.cpp b/Engine/Plugins/FX/Niagara/Source/NiagaraEditor/Private/NiagaraScriptMergeManager.cpp
--- a/Engine/Plugins/FX/Niagara/Source/NiagaraEditor/Private/NiagaraScriptMergeManager.cpp
+++ b/Engine/Plugins/FX/Niagara/Source/NiagaraEditor/Private/NiagaraScriptMergeManager.cpp
@@ -416,23 +416,25 @@ void FNiagaraEmitterMergeAdapter::Initialize(const UNiagaraEmitter& InEmitter, U
 	TArray<UNiagaraNodeOutput*> EventOutputNodes;
 	for (UNiagaraNodeOutput* OutputNode : OutputNodes)
 	{
-		switch (OutputNode->GetUsage())
+		if (UNiagaraScript::IsEquivalentUsage(OutputNode->GetUsage(), ENiagaraScriptUsage::EmitterSpawnScript))
 		{
-		case ENiagaraScriptUsage::EmitterSpawnScript:
 			EmitterSpawnStack = MakeShared<FNiagaraScriptStackMergeAdapter>(*OutputNode, *Emitter->EmitterSpawnScriptProps.Script, Emitter->GetUniqueEmitterName());
-			break;
-		case ENiagaraScriptUsage::EmitterUpdateScript:
+		}
+		else if (UNiagaraScript::IsEquivalentUsage(OutputNode->GetUsage(), ENiagaraScriptUsage::EmitterUpdateScript))
+		{
 			EmitterUpdateStack = MakeShared<FNiagaraScriptStackMergeAdapter>(*OutputNode, *Emitter->EmitterUpdateScriptProps.Script, Emitter->GetUniqueEmitterName());
-			break;
-		case ENiagaraScriptUsage::ParticleSpawnScript:
+		}
+		else if (UNiagaraScript::IsEquivalentUsage(OutputNode->GetUsage(), ENiagaraScriptUsage::ParticleSpawnScript))
+		{
 			ParticleSpawnStack = MakeShared<FNiagaraScriptStackMergeAdapter>(*OutputNode, *Emitter->SpawnScriptProps.Script, Emitter->GetUniqueEmitterName());
-			break;
-		case ENiagaraScriptUsage::ParticleUpdateScript:
+		}
+		else if (UNiagaraScript::IsEquivalentUsage(OutputNode->GetUsage(), ENiagaraScriptUsage::ParticleUpdateScript))
+		{
 			ParticleUpdateStack = MakeShared<FNiagaraScriptStackMergeAdapter>(*OutputNode, *Emitter->UpdateScriptProps.Script, Emitter->GetUniqueEmitterName());
-			break;
-		case ENiagaraScriptUsage::ParticleEventScript:
+		}
+		else if(UNiagaraScript::IsEquivalentUsage(OutputNode->GetUsage(), ENiagaraScriptUsage::ParticleEventScript))
+		{
 			EventOutputNodes.Add(OutputNode);
-			break;
 		}
 	}
 
diff --git a/Engine/Plugins/FX/Niagara/Source/NiagaraEditor/Private/ViewModels/Stack/NiagaraStackGraphUtilities.cpp b/Engine/Plugins/FX/Niagara/Source/NiagaraEditor/Private/ViewModels/Stack/NiagaraStackGraphUtilities.cpp
--- a/Engine/Plugins/FX/Niagara/Source/NiagaraEditor/Private/ViewModels/Stack/NiagaraStackGraphUtilities.cpp
+++ b/Engine/Plugins/FX/Niagara/Source/NiagaraEditor/Private/ViewModels/Stack/NiagaraStackGraphUtilities.cpp
@@ -1031,7 +1031,7 @@ void FNiagaraStackGraphUtilities::SetModuleIsEnabled(UNiagaraNodeFunctionCall& F
 
 bool FNiagaraStackGraphUtilities::ValidateGraphForOutput(UNiagaraGraph& NiagaraGraph, ENiagaraScriptUsage ScriptUsage, FGuid ScriptUsageId, FText& ErrorMessage)
 {
-	UNiagaraNodeOutput* OutputNode = NiagaraGraph.FindOutputNode(ScriptUsage, ScriptUsageId);
+	UNiagaraNodeOutput* OutputNode = NiagaraGraph.FindEquivalentOutputNode(ScriptUsage, ScriptUsageId);
 	if (OutputNode == nullptr)
 	{
 		ErrorMessage = LOCTEXT("ValidateNoOutputMessage", "Output node doesn't exist for script.");
diff --git a/Engine/Plugins/FX/Niagara/Source/NiagaraEditor/Private/ViewModels/Stack/NiagaraStackScriptItemGroup.cpp b/Engine/Plugins/FX/Niagara/Source/NiagaraEditor/Private/ViewModels/Stack/NiagaraStackScriptItemGroup.cpp
--- a/Engine/Plugins/FX/Niagara/Source/NiagaraEditor/Private/ViewModels/Stack/NiagaraStackScriptItemGroup.cpp
+++ b/Engine/Plugins/FX/Niagara/Source/NiagaraEditor/Private/ViewModels/Stack/NiagaraStackScriptItemGroup.cpp
@@ -292,7 +292,7 @@ void UNiagaraStackScriptItemGroup::RefreshChildrenInternal(const TArray<UNiagara
 	{
 		bIsValidForOutput = true;
 
-		UNiagaraNodeOutput* MatchingOutputNode = Graph->FindOutputNode(ScriptUsage, ScriptUsageId);
+		UNiagaraNodeOutput* MatchingOutputNode = Graph->FindEquivalentOutputNode(ScriptUsage, ScriptUsageId);
 		AddUtilities->SetOutputNode(MatchingOutputNode);
 
 		TArray<UNiagaraNodeFunctionCall*> ModuleNodes;
