Niagara - Fix the generated code window so that it doesn't show an extra gpu hlsl view which is invalid, and doesn't show an invalid assembly entry for gpu.

diff --git a/Engine/Plugins/FX/Niagara/Source/NiagaraEditor/Private/Widgets/SNiagaraGeneratedCodeView.cpp b/Engine/Plugins/FX/Niagara/Source/NiagaraEditor/Private/Widgets/SNiagaraGeneratedCodeView.cpp
--- a/Engine/Plugins/FX/Niagara/Source/NiagaraEditor/Private/Widgets/SNiagaraGeneratedCodeView.cpp
+++ b/Engine/Plugins/FX/Niagara/Source/NiagaraEditor/Private/Widgets/SNiagaraGeneratedCodeView.cpp
@@ -341,9 +341,7 @@ void SNiagaraGeneratedCodeView::UpdateUI()
 	TArray<uint32> ScriptDisplayTypes;
 	UNiagaraSystem& System = SystemViewModel->GetSystem();
 	Scripts.Add(System.GetSystemSpawnScript());
-	ScriptDisplayTypes.Add(0);
 	Scripts.Add(System.GetSystemUpdateScript());
-	ScriptDisplayTypes.Add(0);
 
 	TArray<TSharedRef<FNiagaraEmitterHandleViewModel>> SelectedEmitterHandles;
 	SystemViewModel->GetSelectedEmitterHandles(SelectedEmitterHandles);
@@ -355,25 +353,27 @@ void SNiagaraGeneratedCodeView::UpdateUI()
 			TArray<UNiagaraScript*> EmitterScripts;
 			Handle->GetInstance()->GetScripts(EmitterScripts);
 			Scripts.Append(EmitterScripts);
-			ScriptDisplayTypes.AddZeroed(EmitterScripts.Num());
 		}
 	}
 
-	// find the particle spawn script
-	TArray<UNiagaraScript*> DupeScriptsForAssembly;
-	UNiagaraScript* ParticleGPUScript = nullptr;
-	for (UNiagaraScript *Script : Scripts)
+	// Mark the scripts with the correct display type and copy references for the non-gpu scripts for the assembly view.
+	int32 OriginalScriptCount = Scripts.Num();
+	ScriptDisplayTypes.AddUninitialized(OriginalScriptCount);
+	for (int32 i = 0; i < OriginalScriptCount; i++)
 	{
-		DupeScriptsForAssembly.Add(Script);
-		ScriptDisplayTypes.Add(2);
-		if (Script->Usage == ENiagaraScriptUsage::ParticleGPUComputeScript)
+		UNiagaraScript* Script = Scripts[i];
+		if (Script->GetUsage() == ENiagaraScriptUsage::ParticleGPUComputeScript)
 		{
-			ParticleGPUScript = Script;
+			ScriptDisplayTypes[i] = 1;
+		}
+		else
+		{
+			ScriptDisplayTypes[i] = 0;
+
+			Scripts.Add(Script);
+			ScriptDisplayTypes.Add(2);
 		}
 	}
-	Scripts.Append(DupeScriptsForAssembly);
-	Scripts.Add(ParticleGPUScript); // add for the GPU update/spawn script
-	ScriptDisplayTypes.Add(1);
 		
 	GeneratedCode.SetNum(Scripts.Num());
 
@@ -398,9 +398,9 @@ void SNiagaraGeneratedCodeView::UpdateUI()
 			if (bIsGPU)
 			{
 				GeneratedCode[i].Usage = Scripts[i]->Usage;
-				if (ParticleGPUScript->GetVMExecutableData().IsValid())
+				if (Scripts[i]->GetVMExecutableData().IsValid())
 				{
-					ParticleGPUScript->GetVMExecutableData().LastHlslTranslationGPU.ParseIntoArrayLines(OutputByLines, false);
+					Scripts[i]->GetVMExecutableData().LastHlslTranslationGPU.ParseIntoArrayLines(OutputByLines, false);
 				}
 			}
 			else if (bIsAssembly)
