Niagara - Fix GPU compile issues.
+ Fix an issue in the translator where DataInstance.Alive was still referenced in a GPU script even if it wasn't in the parameter map.
+ Fix an issue where data DataInstance.Alive wan't being transfered to the update script correctly when interpolated spawn was turned on on GPU scripts.
+ Change BaseScriptID generation for gpu scripts to use it's own unique ID instead of the spawn scripts ID and update it when the hash of both the spawn and update traversals change.

diff --git a/Engine/Plugins/FX/Niagara/Source/NiagaraCore/Private/NiagaraCustomVersion.cpp b/Engine/Plugins/FX/Niagara/Source/NiagaraCore/Private/NiagaraCustomVersion.cpp
--- a/Engine/Plugins/FX/Niagara/Source/NiagaraCore/Private/NiagaraCustomVersion.cpp
+++ b/Engine/Plugins/FX/Niagara/Source/NiagaraCore/Private/NiagaraCustomVersion.cpp
@@ -8,4 +8,4 @@ const FGuid FNiagaraCustomVersion::GUID(0xFCF57AFA, 0x50764283, 0xB9A9E658, 0xFF
 // Register the custom version with core
 FCustomVersionRegistration GRegisterNiagaraCustomVersion(FNiagaraCustomVersion::GUID, FNiagaraCustomVersion::LatestVersion, TEXT("NiagaraVer"));
 
-const FGuid FNiagaraCustomVersion::LatestScriptCompileVersion(0xEAA7D8F4, 0x8BE04550, 0xAF0EC099, 0x5A99A2F4);
\ No newline at end of file
+const FGuid FNiagaraCustomVersion::LatestScriptCompileVersion(0x8B5FB323, 0xB0604B82, 0xA0F758A2, 0x39848EC2);
\ No newline at end of file
diff --git a/Engine/Plugins/FX/Niagara/Source/NiagaraEditor/Private/NiagaraGraph.cpp b/Engine/Plugins/FX/Niagara/Source/NiagaraEditor/Private/NiagaraGraph.cpp
--- a/Engine/Plugins/FX/Niagara/Source/NiagaraEditor/Private/NiagaraGraph.cpp
+++ b/Engine/Plugins/FX/Niagara/Source/NiagaraEditor/Private/NiagaraGraph.cpp
@@ -198,13 +198,6 @@ FGuid UNiagaraGraph::GetCompileID(ENiagaraScriptUsage InUsage, const FGuid& InUs
 {
 	RebuildCachedData();
 
-	// Since there gpu compute script contains spawn, update, and emitter logic, and we can only return one,
-	// just return the particle spawn script here.
-	if (InUsage == ENiagaraScriptUsage::ParticleGPUComputeScript)
-	{
-		InUsage = ENiagaraScriptUsage::ParticleSpawnScript;
-	}
-
 	for (int32 j = 0; j < CachedUsageInfo.Num(); j++)
 	{
 		if (UNiagaraScript::IsEquivalentUsage(CachedUsageInfo[j].UsageType, InUsage) && CachedUsageInfo[j].UsageId == InUsageId)
@@ -750,6 +743,9 @@ void UNiagaraGraph::RebuildCachedData(bool bForce)
 	UEnum* FoundEnum = nullptr;
 	bool bNeedsAnyNewCompileIds = false;
 
+	FNiagaraGraphScriptUsageInfo* ParticleSpawnUsageInfo = nullptr;
+	FNiagaraGraphScriptUsageInfo* ParticleUpdateUsageInfo = nullptr;
+
 	for (int32 i = 0; i < NiagaraOutputNodes.Num(); i++)
 	{
 		UNiagaraNodeOutput* OutputNode = NiagaraOutputNodes[i];
@@ -823,6 +819,49 @@ void UNiagaraGraph::RebuildCachedData(bool bForce)
 				//UE_LOG(LogNiagaraEditor, Log, TEXT("'%s' changes NOT detected in %s .. keeping guid: %s"), *GetFullName(), *ResultsEnum, *NewUsageCache[i].GeneratedCompileId.ToString());
 			}
 		}
+
+		if (UNiagaraScript::IsEquivalentUsage(NewUsageCache[i].UsageType, ENiagaraScriptUsage::ParticleSpawnScript) && NewUsageCache[i].UsageId == FGuid())
+		{
+			ParticleSpawnUsageInfo = &NewUsageCache[i];
+		}
+
+		if (UNiagaraScript::IsEquivalentUsage(NewUsageCache[i].UsageType, ENiagaraScriptUsage::ParticleUpdateScript) && NewUsageCache[i].UsageId == FGuid())
+		{
+			ParticleUpdateUsageInfo = &NewUsageCache[i];
+		}
+	}
+
+	if (ParticleSpawnUsageInfo != nullptr && ParticleUpdateUsageInfo != nullptr)
+	{
+		// If we have info for both spawn and update generate the gpu version too.
+		FNiagaraGraphScriptUsageInfo GpuUsageInfo;
+		GpuUsageInfo.UsageType = ENiagaraScriptUsage::ParticleGPUComputeScript;
+		GpuUsageInfo.UsageId = FGuid();
+		GpuUsageInfo.Traversal.Append(ParticleSpawnUsageInfo->Traversal);
+		GpuUsageInfo.Traversal.Append(ParticleUpdateUsageInfo->Traversal);
+
+		FSHA1 HashState;
+		for (UNiagaraNode* Node : GpuUsageInfo.Traversal)
+		{
+			FGuid Guid = Node->GetChangeId();
+			HashState.Update((const uint8*)&Guid, sizeof(FGuid));
+		}
+		HashState.Final();
+
+		check(sizeof(uint8) * GpuUsageInfo.DataHash.Num() == sizeof(FSHAHash));
+		HashState.GetHash(GpuUsageInfo.DataHash.GetData());
+
+		FNiagaraGraphScriptUsageInfo* OldGpuUsageInfo = CachedUsageInfo.FindByPredicate([](const FNiagaraGraphScriptUsageInfo& UsageInfo) { return UsageInfo.UsageType == ENiagaraScriptUsage::ParticleGPUComputeScript && UsageInfo.UsageId == FGuid(); });
+		if (OldGpuUsageInfo != nullptr && OldGpuUsageInfo->DataHash == GpuUsageInfo.DataHash)
+		{
+			GpuUsageInfo.GeneratedCompileId = OldGpuUsageInfo->GeneratedCompileId;
+		}
+		else
+		{
+			GpuUsageInfo.GeneratedCompileId = FGuid::NewGuid();
+		}
+
+		NewUsageCache.Add(GpuUsageInfo);
 	}
 
 	// Debug logic, usually disabled at top of file.
@@ -910,12 +949,8 @@ void UNiagaraGraph::GatherExternalDependencyIDs(ENiagaraScriptUsage InUsage, con
 			// Add all chains that we depend on.
 			if (UNiagaraScript::IsUsageDependentOn(InUsage, CachedUsageInfo[i].UsageType)) 
 			{
-				// Skip adding to list because we already did it in GetCompileId above if spawn script.
-				if (CachedUsageInfo[i].UsageType != ENiagaraScriptUsage::ParticleSpawnScript)
-				{
-					InReferencedIDs.Add(CachedUsageInfo[i].GeneratedCompileId);
-					InReferencedObjs.Add(CachedUsageInfo[i].Traversal.Last());
-				}
+				InReferencedIDs.Add(CachedUsageInfo[i].GeneratedCompileId);
+				InReferencedObjs.Add(CachedUsageInfo[i].Traversal.Last());
 
 				for (UNiagaraNode* Node : CachedUsageInfo[i].Traversal)
 				{
diff --git a/Engine/Plugins/FX/Niagara/Source/NiagaraEditor/Private/NiagaraHlslTranslator.cpp b/Engine/Plugins/FX/Niagara/Source/NiagaraEditor/Private/NiagaraHlslTranslator.cpp
--- a/Engine/Plugins/FX/Niagara/Source/NiagaraEditor/Private/NiagaraHlslTranslator.cpp
+++ b/Engine/Plugins/FX/Niagara/Source/NiagaraEditor/Private/NiagaraHlslTranslator.cpp
@@ -1636,7 +1636,8 @@ void FHlslNiagaraTranslator::DefineMain(FString &OutHlslOutput,
 	{
 		OutHlslOutput += TEXT("\t") + MainPreSimulateChunks[i] + TEXT("\n");
 	}
-	
+
+	bool bInitializeAliveForGPU = false;
 	{
 		// call the read data set function
 		OutHlslOutput += TEXT("\tReadDataSets(Context);\n");
@@ -1644,14 +1645,43 @@ void FHlslNiagaraTranslator::DefineMain(FString &OutHlslOutput,
 		// branch between spawn and update
 		if (CompilationTarget == ENiagaraSimTarget::GPUComputeSim)
 		{
+			// Determine 
+			TArray<FName> DataSetNames;
+			for (const FNiagaraDataSetID& ReadId : ReadIds)
+			{
+				DataSetNames.AddUnique(ReadId.Name);
+			}
+			for (const FNiagaraDataSetID& WriteId : WriteIds)
+			{
+				DataSetNames.AddUnique(WriteId.Name);
+			}
+
+			for (int32 i = 0; i < ParamMapHistories.Num(); i++)
+			{
+				for (FName DataSetName : DataSetNames)
+				{
+					if (ParamMapHistories[i].FindVariable(*(DataSetName.ToString() + TEXT(".Alive")), FNiagaraTypeDefinition::GetBoolDef()) != INDEX_NONE)
+					{
+						bInitializeAliveForGPU = true;
+						break;
+					}
+				}
+				if (bInitializeAliveForGPU)
+				{
+					break;
+				}
+			}
+
 			OutHlslOutput += TEXT("\tint StartingPhase = Phase;\n");
 			OutHlslOutput += TEXT("\tGCurrentPhase = Phase;\n");
 
 			for (int32 StageIdx = 0; StageIdx < TranslationStages.Num(); StageIdx++)
 			{
 				OutHlslOutput += FString::Printf(TEXT("\tif(Phase==%d)\n\t{\n"), StageIdx);
-				
-				OutHlslOutput += FString::Printf(TEXT("\t\tif (StartingPhase == %d)\n\t\t{\n\t\t\tContext.%s.DataInstance.Alive=true;\n\t\t}\n"), StageIdx, *TranslationStages[StageIdx].PassNamespace);
+				if (bInitializeAliveForGPU)
+				{
+					OutHlslOutput += FString::Printf(TEXT("\t\tif (StartingPhase == %d)\n\t\t{\n\t\t\tContext.%s.DataInstance.Alive=true;\n\t\t}\n"), StageIdx, *TranslationStages[StageIdx].PassNamespace);
+				}
 
 				if (StageIdx == 0)
 				{
@@ -1669,7 +1699,9 @@ void FHlslNiagaraTranslator::DefineMain(FString &OutHlslOutput,
 					OutHlslOutput += TEXT("\t\t//Begin Transfer of Attributes!\n");
 					if (ParamMapDefinedAttributesToNamespaceVars.Num() != 0)
 					{
-						FString CopyStr = TEXT("\t\tContext.") + TranslationStages[StageIdx + 1].PassNamespace + TEXT(".Particles = Context.") + TranslationStages[StageIdx].PassNamespace + TEXT(".Particles;\n");
+						FString CopyStr =
+							TEXT("\t\tContext.") + TranslationStages[StageIdx + 1].PassNamespace + TEXT(".Particles = Context.") + TranslationStages[StageIdx].PassNamespace + TEXT(".Particles;\n") +
+							TEXT("\t\tContext.") + TranslationStages[StageIdx + 1].PassNamespace + TEXT(".DataInstance = Context.") + TranslationStages[StageIdx].PassNamespace + TEXT(".DataInstance;\n");
 						OutHlslOutput += CopyStr;
 					}
 					OutHlslOutput += TEXT("\t\t//End Transfer of Attributes!\n\n");
@@ -1681,11 +1713,6 @@ void FHlslNiagaraTranslator::DefineMain(FString &OutHlslOutput,
 					OutHlslOutput += FString::Printf(TEXT("\t\tPhase = %d;\n"), StageIdx + 1);
 					OutHlslOutput += FString::Printf(TEXT("\t\tGCurrentPhase = %d;\n"), StageIdx + 1);
 				}
-
-				else if (StageIdx + 1 < TranslationStages.Num())
-				{
-					OutHlslOutput += FString::Printf(TEXT("\t\tContext.%s.Particles = Context.%s.Particles;\n"), *TranslationStages[TranslationStages.Num() - 1].PassNamespace, *TranslationStages[StageIdx].PassNamespace);
-				}
 				
 				OutHlslOutput += TEXT("\t}\n");
 			}
@@ -1750,7 +1777,10 @@ void FHlslNiagaraTranslator::DefineMain(FString &OutHlslOutput,
 				DefineDataSetVariableReads(HlslOutput, ReadIds[VarArrayIdx], VarArrayIdx, *ArrayRef);
 			}
 
-			OutHlslOutput += TEXT("\tContext.Map.DataInstance.Alive = true;\n");
+			if (bInitializeAliveForGPU)
+			{
+				OutHlslOutput += TEXT("\tContext.Map.DataInstance.Alive = true;\n");
+			}
 
 			for (int32 VarArrayIdx = 0; VarArrayIdx < InstanceWriteVars.Num(); VarArrayIdx++)
 			{
