Adding a get dims functions to the texture data interface
Fixed issue where render targets and media textures weren't being sampled.
#tests modified existing tests based off module changes

diff --git a/Engine/Plugins/FX/Niagara/Source/Niagara/Classes/NiagaraDataInterfaceTexture.h b/Engine/Plugins/FX/Niagara/Source/Niagara/Classes/NiagaraDataInterfaceTexture.h
--- a/Engine/Plugins/FX/Niagara/Source/Niagara/Classes/NiagaraDataInterfaceTexture.h
+++ b/Engine/Plugins/FX/Niagara/Source/Niagara/Classes/NiagaraDataInterfaceTexture.h
@@ -35,6 +35,7 @@ public:
 
 	template<typename XType, typename YType>
 	void SampleTexture(FVectorVMContext& Context);
+	void GetTextureDimensions(FVectorVMContext& Context);
 
 	virtual bool Equals(const UNiagaraDataInterface* Other) const override;
 
@@ -46,6 +47,7 @@ public:
 	//FRWBuffer& GetGPUBuffer();
 	static const FString TextureName;
 	static const FString SamplerName;
+	static const FString DimensionsBaseName;
 protected:
 #if WITH_EDITOR
 	bool CopyTextureToCPUBackup(UTexture* SourceTexture, TArray<uint8>& TargetBuffer);
@@ -57,4 +59,5 @@ protected:
 	TArray<uint8> CPUTextureData;
 
 	static const FName SampleTextureName;
+	static const FName TextureDimsName;
 };
diff --git a/Engine/Plugins/FX/Niagara/Source/Niagara/Private/NiagaraDataInterfaceTexture.cpp b/Engine/Plugins/FX/Niagara/Source/Niagara/Private/NiagaraDataInterfaceTexture.cpp
--- a/Engine/Plugins/FX/Niagara/Source/Niagara/Private/NiagaraDataInterfaceTexture.cpp
+++ b/Engine/Plugins/FX/Niagara/Source/Niagara/Private/NiagaraDataInterfaceTexture.cpp
@@ -9,8 +9,10 @@
 #define LOCTEXT_NAMESPACE "UNiagaraDataInterfaceTexture"
 
 const FName UNiagaraDataInterfaceTexture::SampleTextureName(TEXT("SampleTexture2D"));
+const FName UNiagaraDataInterfaceTexture::TextureDimsName(TEXT("TextureDimensions2D"));
 const FString UNiagaraDataInterfaceTexture::TextureName(TEXT("Texture_"));
 const FString UNiagaraDataInterfaceTexture::SamplerName(TEXT("Sampler_"));
+const FString UNiagaraDataInterfaceTexture::DimensionsBaseName(TEXT("Dimensions_"));
 
 UNiagaraDataInterfaceTexture::UNiagaraDataInterfaceTexture(FObjectInitializer const& ObjectInitializer)
 	: Super(ObjectInitializer)
@@ -85,7 +87,7 @@ bool UNiagaraDataInterfaceTexture::CopyTextureToCPUBackup(UTexture* SourceTextur
 	}
 	else 
 	{
-		UE_LOG(LogNiagara, Error, TEXT("Texture %s is not BGRA8, which isn't supported in data interfaces yet"), *SourceTexture->GetName());
+		UE_LOG(LogNiagara, Warning, TEXT("Texture %s is not BGRA8, which means this texture cannot be used with the CPU VM."), *SourceTexture->GetName());
 	}
 
 	return true;
@@ -142,25 +144,71 @@ bool UNiagaraDataInterfaceTexture::Equals(const UNiagaraDataInterface* Other) co
 
 void UNiagaraDataInterfaceTexture::GetFunctions(TArray<FNiagaraFunctionSignature>& OutFunctions)
 {
-	FNiagaraFunctionSignature Sig;
-	Sig.Name = SampleTextureName;
-	Sig.bMemberFunction = true;
-	Sig.bRequiresContext = false;
-	Sig.Inputs.Add(FNiagaraVariable(FNiagaraTypeDefinition(GetClass()), TEXT("Texture")));
-	Sig.Inputs.Add(FNiagaraVariable(FNiagaraTypeDefinition::GetVec2Def(), TEXT("UV")));
-	Sig.SetDescription(LOCTEXT("TextureSampleDesc", "Sample mip level 0 of the input 2d texture at the specified UV coordinates. The UV origin (0,0) is in the upper left hand corner of the image."));
-	Sig.Outputs.Add(FNiagaraVariable(FNiagaraTypeDefinition::GetVec4Def(), TEXT("Value")));
-	//Sig.Owner = *GetName();
-
-	OutFunctions.Add(Sig);
-}
+	{
+		FNiagaraFunctionSignature Sig;
+		Sig.Name = SampleTextureName;
+		Sig.bMemberFunction = true;
+		Sig.bRequiresContext = false;
+		Sig.Inputs.Add(FNiagaraVariable(FNiagaraTypeDefinition(GetClass()), TEXT("Texture")));
+		Sig.Inputs.Add(FNiagaraVariable(FNiagaraTypeDefinition::GetVec2Def(), TEXT("UV")));
+		Sig.SetDescription(LOCTEXT("TextureSampleDesc", "Sample mip level 0 of the input 2d texture at the specified UV coordinates. The UV origin (0,0) is in the upper left hand corner of the image."));
+		Sig.Outputs.Add(FNiagaraVariable(FNiagaraTypeDefinition::GetVec4Def(), TEXT("Value")));
+		//Sig.Owner = *GetName();
+
+		OutFunctions.Add(Sig);
+	}
 
+	{
+		FNiagaraFunctionSignature Sig;
+		Sig.Name = TextureDimsName;
+		Sig.bMemberFunction = true;
+		Sig.bRequiresContext = false;
+		Sig.Inputs.Add(FNiagaraVariable(FNiagaraTypeDefinition(GetClass()), TEXT("Texture")));
+		Sig.SetDescription(LOCTEXT("TextureDimsDesc", "Get the dimensions of mip 0 of the texture."));
+		Sig.Outputs.Add(FNiagaraVariable(FNiagaraTypeDefinition::GetVec2Def(), TEXT("Dimensions2D")));
+		//Sig.Owner = *GetName();
+
+		OutFunctions.Add(Sig);
+	}
+}
 DEFINE_NDI_RAW_FUNC_BINDER(UNiagaraDataInterfaceTexture, SampleTexture);
 void UNiagaraDataInterfaceTexture::GetVMExternalFunction(const FVMExternalFunctionBindingInfo& BindingInfo, void* InstanceData, FVMExternalFunction &OutFunc)
 {
-	check(BindingInfo.Name == SampleTextureName);
-	check(BindingInfo.GetNumInputs() == 2 && BindingInfo.GetNumOutputs() == 4);
-	TNDIParamBinder<0, float, TNDIParamBinder<1, float, NDI_RAW_FUNC_BINDER(UNiagaraDataInterfaceTexture, SampleTexture)>>::Bind(this, BindingInfo, InstanceData, OutFunc);
+	if (BindingInfo.Name == SampleTextureName)
+	{
+		check(BindingInfo.GetNumInputs() == 2 && BindingInfo.GetNumOutputs() == 4);
+		TNDIParamBinder<0, float, TNDIParamBinder<1, float, NDI_RAW_FUNC_BINDER(UNiagaraDataInterfaceTexture, SampleTexture)>>::Bind(this, BindingInfo, InstanceData, OutFunc);
+	}
+	else if (BindingInfo.Name == TextureDimsName)
+	{
+		check(BindingInfo.GetNumInputs() == 0 && BindingInfo.GetNumOutputs() == 2);
+		OutFunc = FVMExternalFunction::CreateUObject(this, &UNiagaraDataInterfaceTexture::GetTextureDimensions);
+	}
+}
+
+void UNiagaraDataInterfaceTexture::GetTextureDimensions(FVectorVMContext& Context)
+{
+	FRegisterHandler<float> OutWidth(Context);
+	FRegisterHandler<float> OutHeight(Context);
+
+	if (Texture == nullptr)
+	{
+		for (int32 i = 0; i < Context.NumInstances; ++i)
+		{
+			*OutWidth.GetDestAndAdvance() = 0.0f;
+			*OutHeight.GetDestAndAdvance() = 0.0f;
+		}
+	}
+	else
+	{
+		float Width = Texture->GetSurfaceWidth();
+		float Height = Texture->GetSurfaceHeight();
+		for (int32 i = 0; i < Context.NumInstances; ++i)
+		{
+			*OutWidth.GetDestAndAdvance() = Width;
+			*OutHeight.GetDestAndAdvance() = Height;
+		}
+	}
 }
 
 template<typename XType, typename YType>
@@ -223,12 +271,24 @@ void UNiagaraDataInterfaceTexture::SampleTexture(FVectorVMContext& Context)
 
 bool UNiagaraDataInterfaceTexture::GetFunctionHLSL(const FName& DefinitionFunctionName, FString InstanceFunctionName, FNiagaraDataInterfaceGPUParamInfo& ParamInfo, FString& OutHLSL)
 {
-	FString HLSLTextureName = TextureName + ParamInfo.DataInterfaceHLSLSymbol;
-	FString HLSLSamplerName = SamplerName + ParamInfo.DataInterfaceHLSLSymbol;
-	OutHLSL += TEXT("void ") + InstanceFunctionName + TEXT("(in float2 In_UV, out float4 Out_Value) \n{\n");
-	OutHLSL += TEXT("\t Out_Value = ") + HLSLTextureName + TEXT(".SampleLevel(") + HLSLSamplerName + TEXT(", In_UV, 0);\n");
-	OutHLSL += TEXT("\n}\n");
-	return true;
+	if (DefinitionFunctionName == SampleTextureName)
+	{
+		FString HLSLTextureName = TextureName + ParamInfo.DataInterfaceHLSLSymbol;
+		FString HLSLSamplerName = SamplerName + ParamInfo.DataInterfaceHLSLSymbol;
+		OutHLSL += TEXT("void ") + InstanceFunctionName + TEXT("(in float2 In_UV, out float4 Out_Value) \n{\n");
+		OutHLSL += TEXT("\t Out_Value = ") + HLSLTextureName + TEXT(".SampleLevel(") + HLSLSamplerName + TEXT(", In_UV, 0);\n");
+		OutHLSL += TEXT("\n}\n");
+		return true;
+	}
+	else if (DefinitionFunctionName == TextureDimsName)
+	{
+		FString DimsVar = DimensionsBaseName + ParamInfo.DataInterfaceHLSLSymbol;
+		OutHLSL += TEXT("void ") + InstanceFunctionName + TEXT("(out float2 Out_Value) \n{\n");
+		OutHLSL += TEXT("\t Out_Value = ") + DimsVar + TEXT(";\n");
+		OutHLSL += TEXT("\n}\n");
+		return true;
+	}
+	return false;
 }
 
 void UNiagaraDataInterfaceTexture::GetParameterDefinitionHLSL(FNiagaraDataInterfaceGPUParamInfo& ParamInfo, FString& OutHLSL)
@@ -237,6 +297,8 @@ void UNiagaraDataInterfaceTexture::GetParameterDefinitionHLSL(FNiagaraDataInterf
 	FString HLSLSamplerName = SamplerName + ParamInfo.DataInterfaceHLSLSymbol;
 	OutHLSL += TEXT("Texture2D ") + HLSLTextureName + TEXT(";\n");
 	OutHLSL += TEXT("SamplerState ") + HLSLSamplerName + TEXT(";\n");
+	OutHLSL += TEXT("float2 ") + DimensionsBaseName + ParamInfo.DataInterfaceHLSLSymbol + TEXT(";\n");
+
 }
 
 
@@ -259,12 +321,19 @@ struct FNiagaraDataInterfaceParametersCS_Texture : public FNiagaraDataInterfaceP
 		{
 			UE_LOG(LogNiagara, Warning, TEXT("Binding failed for FNiagaraDataInterfaceParametersCS_Texture Sampler %s. Was it optimized out?"), *SampleName)
 		}
+
+		Dimensions.Bind(ParameterMap, *(UNiagaraDataInterfaceTexture::DimensionsBaseName + ParamRef.ParameterInfo.DataInterfaceHLSLSymbol));
+
 	}
 
 	virtual void Serialize(FArchive& Ar)override
 	{
 		Ar << TextureParam;
 		Ar << SamplerParam;
+		if (Ar.IsLoading() == false || Ar.CustomVer(FNiagaraCustomVersion::GUID) >= FNiagaraCustomVersion::TextureDataInterfaceSizeSerialize)
+		{
+			Ar << Dimensions;
+		}
 	}
 
 	virtual void Set(FRHICommandList& RHICmdList, FNiagaraShader* Shader, class UNiagaraDataInterface* DataInterface) const override
@@ -274,8 +343,12 @@ struct FNiagaraDataInterfaceParametersCS_Texture : public FNiagaraDataInterfaceP
 		FComputeShaderRHIParamRef ComputeShaderRHI = Shader->GetComputeShader();
 		UNiagaraDataInterfaceTexture* TextureDI = CastChecked<UNiagaraDataInterfaceTexture>(DataInterface);
 		UTexture *Texture = TextureDI->Texture;
+		float TexDims[2];
 		if (!Texture)
 		{
+			TexDims[0] = 0.0f;
+			TexDims[1] = 0.0f;
+			SetShaderValue(RHICmdList, ComputeShaderRHI, Dimensions, TexDims);
 			return;
 		}
 		FTextureRHIParamRef TextureRHI = Texture->TextureReference.TextureReferenceRHI;
@@ -287,6 +360,9 @@ struct FNiagaraDataInterfaceParametersCS_Texture : public FNiagaraDataInterfaceP
 			TStaticSamplerState<SF_Trilinear, AM_Wrap, AM_Wrap, AM_Wrap>::GetRHI(),
 			TextureRHI
 		);
+		TexDims[0] = TextureDI->Texture->GetSurfaceWidth();
+		TexDims[1] = TextureDI->Texture->GetSurfaceHeight();
+		SetShaderValue(RHICmdList, ComputeShaderRHI, Dimensions, TexDims);
 	}
 
 
@@ -294,6 +370,7 @@ private:
 
 	FShaderResourceParameter TextureParam;
 	FShaderResourceParameter SamplerParam;
+	FShaderParameter Dimensions;
 };
 
 FNiagaraDataInterfaceParametersCS* UNiagaraDataInterfaceTexture::ConstructComputeParameters()const
diff --git a/Engine/Plugins/FX/Niagara/Source/NiagaraCore/Private/NiagaraCustomVersion.cpp b/Engine/Plugins/FX/Niagara/Source/NiagaraCore/Private/NiagaraCustomVersion.cpp
--- a/Engine/Plugins/FX/Niagara/Source/NiagaraCore/Private/NiagaraCustomVersion.cpp
+++ b/Engine/Plugins/FX/Niagara/Source/NiagaraCore/Private/NiagaraCustomVersion.cpp
@@ -8,4 +8,4 @@ const FGuid FNiagaraCustomVersion::GUID(0xFCF57AFA, 0x50764283, 0xB9A9E658, 0xFF
 // Register the custom version with core
 FCustomVersionRegistration GRegisterNiagaraCustomVersion(FNiagaraCustomVersion::GUID, FNiagaraCustomVersion::LatestVersion, TEXT("NiagaraVer"));
 
-const FGuid FNiagaraCustomVersion::LatestScriptCompileVersion(0x8B5FB323, 0xB0604B82, 0xA0F758A2, 0x39848EC2);
\ No newline at end of file
+const FGuid FNiagaraCustomVersion::LatestScriptCompileVersion(0x6EB4D29C, 0x0CA547DB, 0x96910F18, 0x9CF291BC);
\ No newline at end of file
diff --git a/Engine/Plugins/FX/Niagara/Source/NiagaraCore/Public/NiagaraCustomVersion.h b/Engine/Plugins/FX/Niagara/Source/NiagaraCore/Public/NiagaraCustomVersion.h
--- a/Engine/Plugins/FX/Niagara/Source/NiagaraCore/Public/NiagaraCustomVersion.h
+++ b/Engine/Plugins/FX/Niagara/Source/NiagaraCore/Public/NiagaraCustomVersion.h
@@ -108,6 +108,7 @@ struct FNiagaraCustomVersion
 
 		TextureDataInterfaceUsesCustomSerialize, // The cpu cache of the texture is now directly serialized instead of using array property serialization.
 
+		TextureDataInterfaceSizeSerialize, // The texture data interface now streams size info
 		// -----<new versions can be added above this line>-------------------------------------------------
 		VersionPlusOne,
 		LatestVersion = VersionPlusOne - 1,
