Niagara - Pull request to add support for setting integer user parameters from blueprint.  Updated the automated test too.

PR #4941: Added Int32 Set function to niagara (Contributed by [REDACTED_PERSON])

diff --git a/Engine/Plugins/FX/Niagara/Source/Niagara/Private/NiagaraComponent.cpp b/Engine/Plugins/FX/Niagara/Source/Niagara/Private/NiagaraComponent.cpp
--- a/Engine/Plugins/FX/Niagara/Source/Niagara/Private/NiagaraComponent.cpp
+++ b/Engine/Plugins/FX/Niagara/Source/Niagara/Private/NiagaraComponent.cpp
@@ -916,6 +916,13 @@ void UNiagaraComponent::SetNiagaraVariableFloat(const FString& InVariableName, f
 	OverrideParameters.SetParameterValue(InValue, FNiagaraVariable(FNiagaraTypeDefinition::GetFloatDef(), VarName), true);
 }
 
+void UNiagaraComponent::SetNiagaraVariableInt(const FString& InVariableName, int32 InValue)
+{
+	FName VarName = FName(*InVariableName);
+
+	OverrideParameters.SetParameterValue(InValue, FNiagaraVariable(FNiagaraTypeDefinition::GetIntDef(), VarName), true);
+}
+
 void UNiagaraComponent::SetNiagaraVariableBool(const FString& InVariableName, bool InValue)
 {
 	FName VarName = FName(*InVariableName);
diff --git a/Engine/Plugins/FX/Niagara/Source/Niagara/Public/NiagaraComponent.h b/Engine/Plugins/FX/Niagara/Source/Niagara/Public/NiagaraComponent.h
--- a/Engine/Plugins/FX/Niagara/Source/Niagara/Public/NiagaraComponent.h
+++ b/Engine/Plugins/FX/Niagara/Source/Niagara/Public/NiagaraComponent.h
@@ -220,6 +220,10 @@ public:
 	UFUNCTION(BlueprintCallable, Category = Niagara, meta = (DisplayName = "Set Niagara Variable (Float)"))
 	void SetNiagaraVariableFloat(const FString& InVariableName, float InValue);
 
+	/** Sets a Niagara int parameter by name, overriding locally if necessary.*/
+	UFUNCTION(BlueprintCallable, Category = Niagara, meta = (DisplayName = "Set Niagara Variable (Int32)"))
+	void SetNiagaraVariableInt(const FString& InVariableName, int32 InValue);
+
 	/** Sets a Niagara float parameter by name, overriding locally if necessary.*/
 	UFUNCTION(BlueprintCallable, Category = Niagara, meta = (DisplayName = "Set Niagara Variable (Bool)"))
 	void SetNiagaraVariableBool(const FString& InVariableName, bool InValue);
