Niagara - Prevent a crash when fixing dependencies in event scripts, also update the modules dealing with force so they're not available in event scripts.

diff --git a/Engine/Plugins/FX/Niagara/Source/NiagaraEditor/Private/ViewModels/Stack/NiagaraStackModuleItem.cpp b/Engine/Plugins/FX/Niagara/Source/NiagaraEditor/Private/ViewModels/Stack/NiagaraStackModuleItem.cpp
--- a/Engine/Plugins/FX/Niagara/Source/NiagaraEditor/Private/ViewModels/Stack/NiagaraStackModuleItem.cpp
+++ b/Engine/Plugins/FX/Niagara/Source/NiagaraEditor/Private/ViewModels/Stack/NiagaraStackModuleItem.cpp
@@ -333,7 +333,7 @@ void UNiagaraStackModuleItem::RefreshIssues(TArray<FStackIssue>& NewIssues)
 		}
 	}
 
-	if (FunctionCallNode && FunctionCallNode->FunctionScript)
+	if (ModuleIndex != INDEX_NONE && FunctionCallNode && FunctionCallNode->FunctionScript)
 	{
 		for (FNiagaraModuleDependency Dependency : FunctionCallNode->FunctionScript->RequiredDependencies)
 		{
