Niagara - Fix bounds for GPU emitters which I broke in 4248783.

diff --git a/Engine/Plugins/FX/Niagara/Source/Niagara/Private/NiagaraEmitterInstance.cpp b/Engine/Plugins/FX/Niagara/Source/Niagara/Private/NiagaraEmitterInstance.cpp
--- a/Engine/Plugins/FX/Niagara/Source/Niagara/Private/NiagaraEmitterInstance.cpp
+++ b/Engine/Plugins/FX/Niagara/Source/Niagara/Private/NiagaraEmitterInstance.cpp
@@ -1011,8 +1011,7 @@ void FNiagaraEmitterInstance::Tick(float DeltaSeconds)
 			EventContext.PostTick();
 		}
 
-		CachedBounds.Init();
-		CachedBounds = CachedBounds.ExpandBy(FVector(20.0, 20.0, 20.0));	// temp until GPU sims update bounds
+		CachedBounds = CachedEmitter->FixedBounds;
 
 		CPUTimeMS = TickTime.GetElapsedMilliseconds();
 
