Niagara - Fix another issue where GPU systems could fail to compile due to DataInstance.Alive usage.

diff --git a/Engine/Plugins/FX/Niagara/Source/NiagaraEditor/Private/NiagaraHlslTranslator.cpp b/Engine/Plugins/FX/Niagara/Source/NiagaraEditor/Private/NiagaraHlslTranslator.cpp
--- a/Engine/Plugins/FX/Niagara/Source/NiagaraEditor/Private/NiagaraHlslTranslator.cpp
+++ b/Engine/Plugins/FX/Niagara/Source/NiagaraEditor/Private/NiagaraHlslTranslator.cpp
@@ -1637,7 +1637,7 @@ void FHlslNiagaraTranslator::DefineMain(FString &OutHlslOutput,
 		OutHlslOutput += TEXT("\t") + MainPreSimulateChunks[i] + TEXT("\n");
 	}
 
-	bool bInitializeAliveForGPU = false;
+	bool bGpuUsesAlive = false;
 	{
 		// call the read data set function
 		OutHlslOutput += TEXT("\tReadDataSets(Context);\n");
@@ -1662,11 +1662,11 @@ void FHlslNiagaraTranslator::DefineMain(FString &OutHlslOutput,
 				{
 					if (ParamMapHistories[i].FindVariable(*(DataSetName.ToString() + TEXT(".Alive")), FNiagaraTypeDefinition::GetBoolDef()) != INDEX_NONE)
 					{
-						bInitializeAliveForGPU = true;
+						bGpuUsesAlive = true;
 						break;
 					}
 				}
-				if (bInitializeAliveForGPU)
+				if (bGpuUsesAlive)
 				{
 					break;
 				}
@@ -1678,7 +1678,7 @@ void FHlslNiagaraTranslator::DefineMain(FString &OutHlslOutput,
 			for (int32 StageIdx = 0; StageIdx < TranslationStages.Num(); StageIdx++)
 			{
 				OutHlslOutput += FString::Printf(TEXT("\tif(Phase==%d)\n\t{\n"), StageIdx);
-				if (bInitializeAliveForGPU)
+				if (bGpuUsesAlive)
 				{
 					OutHlslOutput += FString::Printf(TEXT("\t\tif (StartingPhase == %d)\n\t\t{\n\t\t\tContext.%s.DataInstance.Alive=true;\n\t\t}\n"), StageIdx, *TranslationStages[StageIdx].PassNamespace);
 				}
@@ -1699,10 +1699,15 @@ void FHlslNiagaraTranslator::DefineMain(FString &OutHlslOutput,
 					OutHlslOutput += TEXT("\t\t//Begin Transfer of Attributes!\n");
 					if (ParamMapDefinedAttributesToNamespaceVars.Num() != 0)
 					{
-						FString CopyStr =
-							TEXT("\t\tContext.") + TranslationStages[StageIdx + 1].PassNamespace + TEXT(".Particles = Context.") + TranslationStages[StageIdx].PassNamespace + TEXT(".Particles;\n") +
-							TEXT("\t\tContext.") + TranslationStages[StageIdx + 1].PassNamespace + TEXT(".DataInstance = Context.") + TranslationStages[StageIdx].PassNamespace + TEXT(".DataInstance;\n");
-						OutHlslOutput += CopyStr;
+						FString CopyParticlesStr =
+							TEXT("\t\tContext.") + TranslationStages[StageIdx + 1].PassNamespace + TEXT(".Particles = Context.") + TranslationStages[StageIdx].PassNamespace + TEXT(".Particles;\n");
+						OutHlslOutput += CopyParticlesStr;
+
+						if (bGpuUsesAlive)
+						{
+							FString CopyDataInstanceStr = TEXT("\t\tContext.") + TranslationStages[StageIdx + 1].PassNamespace + TEXT(".DataInstance = Context.") + TranslationStages[StageIdx].PassNamespace + TEXT(".DataInstance;\n");
+							OutHlslOutput += CopyDataInstanceStr;
+						}
 					}
 					OutHlslOutput += TEXT("\t\t//End Transfer of Attributes!\n\n");
 				}
@@ -1777,7 +1782,7 @@ void FHlslNiagaraTranslator::DefineMain(FString &OutHlslOutput,
 				DefineDataSetVariableReads(HlslOutput, ReadIds[VarArrayIdx], VarArrayIdx, *ArrayRef);
 			}
 
-			if (bInitializeAliveForGPU)
+			if (bGpuUsesAlive)
 			{
 				OutHlslOutput += TEXT("\tContext.Map.DataInstance.Alive = true;\n");
 			}
