Fixed missing undef LOCTEXT_NAMESPACE.

diff --git a/Engine/Plugins/FX/Niagara/Source/Niagara/Private/NiagaraDataInterfaceTexture.cpp b/Engine/Plugins/FX/Niagara/Source/Niagara/Private/NiagaraDataInterfaceTexture.cpp
--- a/Engine/Plugins/FX/Niagara/Source/Niagara/Private/NiagaraDataInterfaceTexture.cpp
+++ b/Engine/Plugins/FX/Niagara/Source/Niagara/Private/NiagaraDataInterfaceTexture.cpp
@@ -376,4 +376,6 @@ private:
 FNiagaraDataInterfaceParametersCS* UNiagaraDataInterfaceTexture::ConstructComputeParameters()const
 {
 	return new FNiagaraDataInterfaceParametersCS_Texture();
-}
\ No newline at end of file
+}
+
+#undef LOCTEXT_NAMESPACE
\ No newline at end of file
