Adding access to Ryan's PseudoVolume textures
Does nothing on CPU VM at the moment. Should be a passthrough to common.ush's PseudoVolumeTexture.

#tests created an emitter that compiled, am leaving up to Jon Lindquist to test fully

diff --git a/Engine/Plugins/FX/Niagara/Source/Niagara/Classes/NiagaraDataInterfaceTexture.h b/Engine/Plugins/FX/Niagara/Source/Niagara/Classes/NiagaraDataInterfaceTexture.h
--- a/Engine/Plugins/FX/Niagara/Source/Niagara/Classes/NiagaraDataInterfaceTexture.h
+++ b/Engine/Plugins/FX/Niagara/Source/Niagara/Classes/NiagaraDataInterfaceTexture.h
@@ -36,6 +36,7 @@ public:
 	template<typename XType, typename YType>
 	void SampleTexture(FVectorVMContext& Context);
 	void GetTextureDimensions(FVectorVMContext& Context);
+	void NoOp(FVectorVMContext& Context);
 
 	virtual bool Equals(const UNiagaraDataInterface* Other) const override;
 
@@ -58,6 +59,8 @@ protected:
 
 	TArray<uint8> CPUTextureData;
 
-	static const FName SampleTextureName;
+	static const FName SampleTexture2DName;
+	static const FName SampleVolumeTextureName;
+	static const FName SamplePseudoVolumeTextureName;
 	static const FName TextureDimsName;
 };
diff --git a/Engine/Plugins/FX/Niagara/Source/Niagara/Private/NiagaraDataInterfaceTexture.cpp b/Engine/Plugins/FX/Niagara/Source/Niagara/Private/NiagaraDataInterfaceTexture.cpp
--- a/Engine/Plugins/FX/Niagara/Source/Niagara/Private/NiagaraDataInterfaceTexture.cpp
+++ b/Engine/Plugins/FX/Niagara/Source/Niagara/Private/NiagaraDataInterfaceTexture.cpp
@@ -8,7 +8,9 @@
 
 #define LOCTEXT_NAMESPACE "UNiagaraDataInterfaceTexture"
 
-const FName UNiagaraDataInterfaceTexture::SampleTextureName(TEXT("SampleTexture2D"));
+const FName UNiagaraDataInterfaceTexture::SampleTexture2DName(TEXT("SampleTexture2D"));
+const FName UNiagaraDataInterfaceTexture::SampleVolumeTextureName(TEXT("SampleVolumeTexture"));
+const FName UNiagaraDataInterfaceTexture::SamplePseudoVolumeTextureName(TEXT("SamplePseudoVolumeTexture"));
 const FName UNiagaraDataInterfaceTexture::TextureDimsName(TEXT("TextureDimensions2D"));
 const FString UNiagaraDataInterfaceTexture::TextureName(TEXT("Texture_"));
 const FString UNiagaraDataInterfaceTexture::SamplerName(TEXT("Sampler_"));
@@ -146,7 +148,7 @@ void UNiagaraDataInterfaceTexture::GetFunctions(TArray<FNiagaraFunctionSignature
 {
 	{
 		FNiagaraFunctionSignature Sig;
-		Sig.Name = SampleTextureName;
+		Sig.Name = SampleTexture2DName;
 		Sig.bMemberFunction = true;
 		Sig.bRequiresContext = false;
 		Sig.Inputs.Add(FNiagaraVariable(FNiagaraTypeDefinition(GetClass()), TEXT("Texture")));
@@ -158,6 +160,50 @@ void UNiagaraDataInterfaceTexture::GetFunctions(TArray<FNiagaraFunctionSignature
 		OutFunctions.Add(Sig);
 	}
 
+	//{
+	//	FNiagaraFunctionSignature Sig;
+	//	Sig.Name = SampleVolumeTextureName;
+	//	Sig.bMemberFunction = true;
+	//	Sig.bRequiresContext = false;
+	//	Sig.Inputs.Add(FNiagaraVariable(FNiagaraTypeDefinition(GetClass()), TEXT("Texture")));
+	//	Sig.Inputs.Add(FNiagaraVariable(FNiagaraTypeDefinition::GetVec3Def(), TEXT("UVW")));
+	//	Sig.SetDescription(LOCTEXT("TextureSampleDesc", "Sample mip level 0 of the input 3d texture at the specified UVW coordinates. The UVW origin (0,0) is in the bottom left hand corner of the volume."));
+	//	Sig.Outputs.Add(FNiagaraVariable(FNiagaraTypeDefinition::GetVec4Def(), TEXT("Value")));
+	//	//Sig.Owner = *GetName();
+
+	//	OutFunctions.Add(Sig);
+	//}
+
+	{
+		FNiagaraFunctionSignature Sig;
+		Sig.Name = SamplePseudoVolumeTextureName;
+		Sig.bMemberFunction = true;
+		Sig.bRequiresContext = false;
+		Sig.Inputs.Add(FNiagaraVariable(FNiagaraTypeDefinition(GetClass()), TEXT("Texture")));
+		Sig.Inputs.Add(FNiagaraVariable(FNiagaraTypeDefinition::GetVec3Def(), TEXT("UVW")));
+		Sig.Inputs.Add(FNiagaraVariable(FNiagaraTypeDefinition::GetVec2Def(), TEXT("XYNumFrames")));
+		Sig.Inputs.Add(FNiagaraVariable(FNiagaraTypeDefinition::GetFloatDef(), TEXT("TotalNumFrames")));
+		Sig.Inputs.Add(FNiagaraVariable(FNiagaraTypeDefinition::GetIntDef(), TEXT("MipMode")));
+		Sig.Inputs.Add(FNiagaraVariable(FNiagaraTypeDefinition::GetFloatDef(), TEXT("MipLevel")));
+		Sig.Inputs.Add(FNiagaraVariable(FNiagaraTypeDefinition::GetVec2Def(), TEXT("DDX")));
+		Sig.Inputs.Add(FNiagaraVariable(FNiagaraTypeDefinition::GetVec2Def(), TEXT("DDY")));
+		
+		Sig.SetDescription(LOCTEXT("TextureSampleDesc", "Return a pseudovolume texture sample.\n"
+			"Useful for simulating 3D texturing with a 2D texture or as a texture flipbook with lerped transitions.\n"
+			"Treats 2d layout of frames as a 3d texture and performs bilinear filtering by blending with an offset Z frame.\n"
+			"Texture = Input Texture Object storing Volume Data\n"
+			"UVW = Input float3 for Position, 0 - 1\n"
+			"XYNumFrames = Input float for num frames in x, y directions\n"
+			"TotalNumFrames = Input float for num total frames\n"
+			"MipMode = Sampling mode : 0 = use miplevel, 1 = use UV computed gradients, 2 = Use gradients(default = 0)\n"
+			"MipLevel = MIP level to use in mipmode = 0 (default 0)\n"
+			"DDX, DDY = Texture gradients in mipmode = 2\n"));
+		Sig.Outputs.Add(FNiagaraVariable(FNiagaraTypeDefinition::GetVec4Def(), TEXT("Value")));
+		//Sig.Owner = *GetName();
+
+		OutFunctions.Add(Sig);
+	}
+
 	{
 		FNiagaraFunctionSignature Sig;
 		Sig.Name = TextureDimsName;
@@ -174,11 +220,17 @@ void UNiagaraDataInterfaceTexture::GetFunctions(TArray<FNiagaraFunctionSignature
 DEFINE_NDI_RAW_FUNC_BINDER(UNiagaraDataInterfaceTexture, SampleTexture);
 void UNiagaraDataInterfaceTexture::GetVMExternalFunction(const FVMExternalFunctionBindingInfo& BindingInfo, void* InstanceData, FVMExternalFunction &OutFunc)
 {
-	if (BindingInfo.Name == SampleTextureName)
+	if (BindingInfo.Name == SampleTexture2DName)
 	{
 		check(BindingInfo.GetNumInputs() == 2 && BindingInfo.GetNumOutputs() == 4);
 		TNDIParamBinder<0, float, TNDIParamBinder<1, float, NDI_RAW_FUNC_BINDER(UNiagaraDataInterfaceTexture, SampleTexture)>>::Bind(this, BindingInfo, InstanceData, OutFunc);
 	}
+	else if (BindingInfo.Name == SamplePseudoVolumeTextureName)
+	{
+		//check(BindingInfo.GetNumInputs() == 2 && BindingInfo.GetNumOutputs() == 4);
+		OutFunc = FVMExternalFunction::CreateUObject(this, &UNiagaraDataInterfaceTexture::NoOp);
+
+	}
 	else if (BindingInfo.Name == TextureDimsName)
 	{
 		check(BindingInfo.GetNumInputs() == 0 && BindingInfo.GetNumOutputs() == 2);
@@ -186,6 +238,12 @@ void UNiagaraDataInterfaceTexture::GetVMExternalFunction(const FVMExternalFuncti
 	}
 }
 
+
+void UNiagaraDataInterfaceTexture::NoOp(FVectorVMContext& Context)
+{
+
+}
+
 void UNiagaraDataInterfaceTexture::GetTextureDimensions(FVectorVMContext& Context)
 {
 	FRegisterHandler<float> OutWidth(Context);
@@ -271,7 +329,7 @@ void UNiagaraDataInterfaceTexture::SampleTexture(FVectorVMContext& Context)
 
 bool UNiagaraDataInterfaceTexture::GetFunctionHLSL(const FName& DefinitionFunctionName, FString InstanceFunctionName, FNiagaraDataInterfaceGPUParamInfo& ParamInfo, FString& OutHLSL)
 {
-	if (DefinitionFunctionName == SampleTextureName)
+	if (DefinitionFunctionName == SampleTexture2DName)
 	{
 		FString HLSLTextureName = TextureName + ParamInfo.DataInterfaceHLSLSymbol;
 		FString HLSLSamplerName = SamplerName + ParamInfo.DataInterfaceHLSLSymbol;
@@ -280,6 +338,24 @@ bool UNiagaraDataInterfaceTexture::GetFunctionHLSL(const FName& DefinitionFuncti
 		OutHLSL += TEXT("\n}\n");
 		return true;
 	}
+	/*else if (DefinitionFunctionName == SampleVolumeTextureName)
+	{
+		FString HLSLTextureName = TextureName + ParamInfo.DataInterfaceHLSLSymbol;
+		FString HLSLSamplerName = SamplerName + ParamInfo.DataInterfaceHLSLSymbol;
+		OutHLSL += TEXT("void ") + InstanceFunctionName + TEXT("(in float3 In_UV, out float4 Out_Value) \n{\n");
+		OutHLSL += TEXT("\t Out_Value = ") + HLSLTextureName + TEXT(".SampleLevel(") + HLSLSamplerName + TEXT(", In_UV, 0);\n");
+		OutHLSL += TEXT("\n}\n");
+		return true;
+	}*/
+	else if (DefinitionFunctionName == SamplePseudoVolumeTextureName)
+	{
+		FString HLSLTextureName = TextureName + ParamInfo.DataInterfaceHLSLSymbol;
+		FString HLSLSamplerName = SamplerName + ParamInfo.DataInterfaceHLSLSymbol;
+		OutHLSL += TEXT("void ") + InstanceFunctionName + TEXT("(in float3 In_UVW, in float2 In_XYNumFrames, in float In_TotalNumFrames, in int In_MipMode, in float In_MipLevel, in float2 In_DDX, in float2 In_DDY, out float4 Out_Value) \n{\n");
+		OutHLSL += TEXT("\t Out_Value = PseudoVolumeTexture(") + HLSLTextureName + TEXT(", ") + HLSLSamplerName + TEXT(", In_UVW, In_XYNumFrames, In_TotalNumFrames, (uint) In_MipMode, In_MipLevel, In_DDX, In_DDY); \n");
+		OutHLSL += TEXT("\n}\n");
+		return true;
+	}
 	else if (DefinitionFunctionName == TextureDimsName)
 	{
 		FString DimsVar = DimensionsBaseName + ParamInfo.DataInterfaceHLSLSymbol;
@@ -357,7 +433,7 @@ struct FNiagaraDataInterfaceParametersCS_Texture : public FNiagaraDataInterfaceP
 			ComputeShaderRHI,
 			TextureParam,
 			SamplerParam,
-			TStaticSamplerState<SF_Trilinear, AM_Wrap, AM_Wrap, AM_Wrap>::GetRHI(),
+			Texture->Resource->SamplerStateRHI,
 			TextureRHI
 		);
 		TexDims[0] = TextureDI->Texture->GetSurfaceWidth();
