Missing file

diff --git a/Engine/Plugins/FX/Niagara/Source/Niagara/Classes/NDISkeletalMeshCommon.h b/Engine/Plugins/FX/Niagara/Source/Niagara/Classes/NDISkeletalMeshCommon.h
new file mode 100644
--- /dev/null
+++ b/Engine/Plugins/FX/Niagara/Source/Niagara/Classes/NDISkeletalMeshCommon.h
@@ -0,0 +1,392 @@
+// Copyright 1998-2018 Epic Games, Inc. All Rights Reserved.
+#pragma once
+
+//////////////////////////////////////////////////////////////////////////
+// Helper classes for reducing duplicate code when accessing vertex positions. 
+
+struct FSkeletalMeshAccessorHelper
+{
+	FSkeletalMeshAccessorHelper()
+		: Comp(nullptr)
+		, Mesh(nullptr)
+		, LODData(nullptr)
+		, SkinWeightBuffer(nullptr)
+		, IndexBuffer(nullptr)
+		, SamplingRegion(nullptr)
+		, SamplingRegionBuiltData(nullptr)
+		, SkinningData(nullptr)
+	{
+
+	}
+
+	template<typename FilterMode, typename AreaWeightingMode>
+	FORCEINLINE void Init(FNDISkeletalMesh_InstanceData* InstData)
+	{
+		Comp = Cast<USkeletalMeshComponent>(InstData->Component.Get());
+		Mesh = InstData->Mesh;
+		LODData = &InstData->GetLODRenderDataAndSkinWeights(SkinWeightBuffer);
+		IndexBuffer = LODData->MultiSizeIndexContainer.GetIndexBuffer();
+		SkinningData = InstData->SkinningData.SkinningData.Get();
+		Usage = InstData->SkinningData.Usage;
+
+		if (Comp)
+		{
+			const USkinnedMeshComponent* BaseComp = Comp->GetBaseComponent();
+			BoneComponentSpaceTransforms = &BaseComp->GetComponentSpaceTransforms();
+			PrevBoneComponentSpaceTransforms = &BaseComp->GetPreviousComponentTransformsArray();
+		}
+	}
+
+	USkeletalMeshComponent* Comp;
+	USkeletalMesh* Mesh;
+	TWeakObjectPtr<USkeletalMesh> MeshSafe;
+	FSkeletalMeshLODRenderData* LODData;
+	FSkinWeightVertexBuffer* SkinWeightBuffer;
+	FRawStaticIndexBuffer16or32Interface* IndexBuffer;
+	const FSkeletalMeshSamplingRegion* SamplingRegion;
+	const FSkeletalMeshSamplingRegionBuiltData* SamplingRegionBuiltData;
+	FSkeletalMeshSkinningData* SkinningData;
+	FSkeletalMeshSkinningDataUsage Usage;
+	const TArray<FTransform>* BoneComponentSpaceTransforms;
+	const TArray<FTransform>* PrevBoneComponentSpaceTransforms;
+};
+
+template<>
+void FSkeletalMeshAccessorHelper::Init<
+	TIntegralConstant<ENDISkeletalMesh_FilterMode, ENDISkeletalMesh_FilterMode::SingleRegion>,
+	TIntegralConstant<ENDISkelMesh_AreaWeightingMode, ENDISkelMesh_AreaWeightingMode::None>>
+	(FNDISkeletalMesh_InstanceData* InstData);
+
+template<>
+void FSkeletalMeshAccessorHelper::Init<
+	TIntegralConstant<ENDISkeletalMesh_FilterMode, ENDISkeletalMesh_FilterMode::SingleRegion>,
+	TIntegralConstant<ENDISkelMesh_AreaWeightingMode, ENDISkelMesh_AreaWeightingMode::AreaWeighted>>
+	(FNDISkeletalMesh_InstanceData* InstData);
+
+//////////////////////////////////////////////////////////////////////////
+
+template<typename SkinningMode>
+struct FSkinnedPositionAccessorHelper
+{
+	FORCEINLINE void GetTrianlgeIndices(int32 Tri, int32& Idx0, int32& Idx1, int32& Idx2)
+	{
+		checkf(false, TEXT("Must provide a specialization for this template type"));
+	}
+
+	FORCEINLINE void GetSkinnedTrianglePositions(FSkeletalMeshAccessorHelper& Accessor,
+		int32 Tri, FVector& OutPos0, FVector& OutPos1, FVector& OutPos2, FVector& OutPrev0, FVector& OutPrev1, FVector& OutPrev2, int32& Idx0, int32& Idx1, int32& Idx2)
+	{
+		checkf(false, TEXT("Must provide a specialization for this template type"));
+	}
+
+	FORCEINLINE void GetSkinnedTrianglePositions(FSkeletalMeshAccessorHelper& Accessor,
+		int32 Tri, FVector& OutPos0, FVector& OutPos1, FVector& OutPos2, int32& Idx0, int32& Idx1, int32& Idx2)
+	{
+		checkf(false, TEXT("Must provide a specialization for this template type"));
+	}
+
+	FORCEINLINE FVector GetSkinnedVertexPosition(FSkeletalMeshAccessorHelper& Accessor, int32 VertexIndex)
+	{
+		checkf(false, TEXT("Must provide a specialization for this template type"));
+	}
+
+	FORCEINLINE FVector GetSkinnedVertexPreviousPosition(FSkeletalMeshAccessorHelper& Accessor, int32 VertexIndex)
+	{
+		checkf(false, TEXT("Must provide a specialization for this template type"));
+	}
+
+	FORCEINLINE FVector GetSkinnedBonePosition(FSkeletalMeshAccessorHelper& Accessor, int32 BoneIndex)
+	{
+		checkf(false, TEXT("Must provide a specialization for this template type"));
+	}
+
+	FORCEINLINE FVector GetSkinnedBonePreviousPosition(FSkeletalMeshAccessorHelper& Accessor, int32 BoneIndex)
+	{
+		checkf(false, TEXT("Must provide a specialization for this template type"));
+	}
+};
+
+template<>
+struct FSkinnedPositionAccessorHelper<TIntegralConstant<ENDISkeletalMesh_SkinningMode, ENDISkeletalMesh_SkinningMode::None>>
+{
+	FORCEINLINE void GetTrianlgeIndices(FSkeletalMeshAccessorHelper& Accessor, int32 Tri, int32& Idx0, int32& Idx1, int32& Idx2)
+	{
+		checkSlow(Tri + 2 < Accessor.IndexBuffer->Num());
+		Idx0 = Accessor.IndexBuffer->Get(Tri);
+		Idx1 = Accessor.IndexBuffer->Get(Tri + 1);
+		Idx2 = Accessor.IndexBuffer->Get(Tri + 2);
+	}
+
+	FORCEINLINE void GetSkinnedTrianglePositions(FSkeletalMeshAccessorHelper& Accessor, int32 Idx0, int32 Idx1, int32 Idx2, FVector& OutPos0, FVector& OutPos1, FVector& OutPos2)
+	{
+		OutPos0 = GetSkeletalMeshRefVertLocation(Accessor.Mesh, *Accessor.LODData, *Accessor.SkinWeightBuffer, Idx0);
+		OutPos1 = GetSkeletalMeshRefVertLocation(Accessor.Mesh, *Accessor.LODData, *Accessor.SkinWeightBuffer, Idx1);
+		OutPos2 = GetSkeletalMeshRefVertLocation(Accessor.Mesh, *Accessor.LODData, *Accessor.SkinWeightBuffer, Idx2);
+	}
+
+	FORCEINLINE void GetSkinnedTrianglePreviousPositions(FSkeletalMeshAccessorHelper& Accessor, int32 Idx0, int32 Idx1, int32 Idx2, FVector& OutPos0, FVector& OutPos1, FVector& OutPos2)
+	{
+		OutPos0 = GetSkeletalMeshRefVertLocation(Accessor.Mesh, *Accessor.LODData, *Accessor.SkinWeightBuffer, Idx0);
+		OutPos1 = GetSkeletalMeshRefVertLocation(Accessor.Mesh, *Accessor.LODData, *Accessor.SkinWeightBuffer, Idx1);
+		OutPos2 = GetSkeletalMeshRefVertLocation(Accessor.Mesh, *Accessor.LODData, *Accessor.SkinWeightBuffer, Idx2);
+	}
+
+	FORCEINLINE FVector GetSkinnedVertexPosition(FSkeletalMeshAccessorHelper& Accessor, int32 VertexIndex)
+	{
+		return GetSkeletalMeshRefVertLocation(Accessor.Mesh, *Accessor.LODData, *Accessor.SkinWeightBuffer, VertexIndex);
+	}
+
+	FORCEINLINE FVector GetSkinnedVertexPreviousPosition(FSkeletalMeshAccessorHelper& Accessor, int32 VertexIndex)
+	{
+		return GetSkeletalMeshRefVertLocation(Accessor.Mesh, *Accessor.LODData, *Accessor.SkinWeightBuffer, VertexIndex);
+	}
+
+	FORCEINLINE_DEBUGGABLE FVector GetSkinnedBonePosition(FSkeletalMeshAccessorHelper& Accessor, int32 BoneIndex)
+	{
+		return Accessor.Mesh->GetComposedRefPoseMatrix(BoneIndex).GetOrigin();
+	}
+
+	FORCEINLINE_DEBUGGABLE FVector GetSkinnedBonePreviousPosition(FSkeletalMeshAccessorHelper& Accessor, int32 BoneIndex)
+	{
+		return Accessor.Mesh->GetComposedRefPoseMatrix(BoneIndex).GetOrigin();
+	}
+};
+
+template<>
+struct FSkinnedPositionAccessorHelper<TIntegralConstant<ENDISkeletalMesh_SkinningMode, ENDISkeletalMesh_SkinningMode::SkinOnTheFly>>
+{
+	FORCEINLINE void GetTrianlgeIndices(FSkeletalMeshAccessorHelper& Accessor, int32 Tri, int32& Idx0, int32& Idx1, int32& Idx2)
+	{
+		checkSlow(Tri + 2 < Accessor.IndexBuffer->Num());
+		Idx0 = Accessor.IndexBuffer->Get(Tri);
+		Idx1 = Accessor.IndexBuffer->Get(Tri + 1);
+		Idx2 = Accessor.IndexBuffer->Get(Tri + 2);
+	}
+
+	FORCEINLINE void GetSkinnedTrianglePositions(FSkeletalMeshAccessorHelper& Accessor, int32 Idx0, int32 Idx1, int32 Idx2, FVector& OutPos0, FVector& OutPos1, FVector& OutPos2)
+	{
+		OutPos0 = USkeletalMeshComponent::GetSkinnedVertexPosition(Accessor.Comp, Idx0, *Accessor.LODData, *Accessor.SkinWeightBuffer, Accessor.SkinningData->CurrBoneRefToLocals());
+		OutPos1 = USkeletalMeshComponent::GetSkinnedVertexPosition(Accessor.Comp, Idx1, *Accessor.LODData, *Accessor.SkinWeightBuffer, Accessor.SkinningData->CurrBoneRefToLocals());
+		OutPos2 = USkeletalMeshComponent::GetSkinnedVertexPosition(Accessor.Comp, Idx2, *Accessor.LODData, *Accessor.SkinWeightBuffer, Accessor.SkinningData->CurrBoneRefToLocals());
+	}
+
+	FORCEINLINE void GetSkinnedTrianglePreviousPositions(FSkeletalMeshAccessorHelper& Accessor, int32 Idx0, int32 Idx1, int32 Idx2, FVector& OutPos0, FVector& OutPos1, FVector& OutPos2)
+	{
+		OutPos0 = USkeletalMeshComponent::GetSkinnedVertexPosition(Accessor.Comp, Idx0, *Accessor.LODData, *Accessor.SkinWeightBuffer, Accessor.SkinningData->PrevBoneRefToLocals());
+		OutPos1 = USkeletalMeshComponent::GetSkinnedVertexPosition(Accessor.Comp, Idx1, *Accessor.LODData, *Accessor.SkinWeightBuffer, Accessor.SkinningData->PrevBoneRefToLocals());
+		OutPos2 = USkeletalMeshComponent::GetSkinnedVertexPosition(Accessor.Comp, Idx2, *Accessor.LODData, *Accessor.SkinWeightBuffer, Accessor.SkinningData->PrevBoneRefToLocals());
+	}
+
+	FORCEINLINE FVector GetSkinnedVertexPosition(FSkeletalMeshAccessorHelper& Accessor, int32 VertexIndex)
+	{
+		return USkeletalMeshComponent::GetSkinnedVertexPosition(Accessor.Comp, VertexIndex, *Accessor.LODData, *Accessor.SkinWeightBuffer, Accessor.SkinningData->CurrBoneRefToLocals());
+	}
+
+	FORCEINLINE FVector GetSkinnedVertexPreviousPosition(FSkeletalMeshAccessorHelper& Accessor, int32 VertexIndex)
+	{
+		return USkeletalMeshComponent::GetSkinnedVertexPosition(Accessor.Comp, VertexIndex, *Accessor.LODData, *Accessor.SkinWeightBuffer, Accessor.SkinningData->PrevBoneRefToLocals());
+	}
+
+	FORCEINLINE_DEBUGGABLE FVector GetSkinnedBonePosition(FSkeletalMeshAccessorHelper& Accessor, int32 BoneIndex)
+	{
+		return (*Accessor.BoneComponentSpaceTransforms)[BoneIndex].GetLocation();
+	}
+
+	FORCEINLINE_DEBUGGABLE FVector GetSkinnedBonePreviousPosition(FSkeletalMeshAccessorHelper& Accessor, int32 BoneIndex)
+	{
+		return (*Accessor.PrevBoneComponentSpaceTransforms)[BoneIndex].GetLocation();
+	}
+};
+
+template<>
+struct FSkinnedPositionAccessorHelper<TIntegralConstant<ENDISkeletalMesh_SkinningMode, ENDISkeletalMesh_SkinningMode::PreSkin>>
+{
+	FORCEINLINE void GetTrianlgeIndices(FSkeletalMeshAccessorHelper& Accessor, int32 Tri, int32& Idx0, int32& Idx1, int32& Idx2)
+	{
+		checkSlow(Tri + 2 < Accessor.IndexBuffer->Num());
+		Idx0 = Accessor.IndexBuffer->Get(Tri);
+		Idx1 = Accessor.IndexBuffer->Get(Tri + 1);
+		Idx2 = Accessor.IndexBuffer->Get(Tri + 2);
+	}
+
+	FORCEINLINE void GetSkinnedTrianglePositions(FSkeletalMeshAccessorHelper& Accessor, int32 Idx0, int32 Idx1, int32 Idx2, FVector& OutPos0, FVector& OutPos1, FVector& OutPos2)
+	{
+		OutPos0 = Accessor.SkinningData->GetPosition(Accessor.Usage.GetLODIndex(), Idx0);
+		OutPos1 = Accessor.SkinningData->GetPosition(Accessor.Usage.GetLODIndex(), Idx1);
+		OutPos2 = Accessor.SkinningData->GetPosition(Accessor.Usage.GetLODIndex(), Idx2);
+	}
+
+	FORCEINLINE void GetSkinnedTrianglePreviousPositions(FSkeletalMeshAccessorHelper& Accessor, int32 Idx0, int32 Idx1, int32 Idx2, FVector& OutPos0, FVector& OutPos1, FVector& OutPos2)
+	{
+		OutPos0 = Accessor.SkinningData->GetPreviousPosition(Accessor.Usage.GetLODIndex(), Idx0);
+		OutPos1 = Accessor.SkinningData->GetPreviousPosition(Accessor.Usage.GetLODIndex(), Idx1);
+		OutPos2 = Accessor.SkinningData->GetPreviousPosition(Accessor.Usage.GetLODIndex(), Idx2);
+	}
+
+	FORCEINLINE FVector GetSkinnedVertexPosition(FSkeletalMeshAccessorHelper& Accessor, int32 VertexIndex)
+	{
+		return Accessor.SkinningData->GetPosition(Accessor.Usage.GetLODIndex(), VertexIndex);
+	}
+
+	FORCEINLINE FVector GetSkinnedVertexPreviousPosition(FSkeletalMeshAccessorHelper& Accessor, int32 VertexIndex)
+	{
+		return Accessor.SkinningData->GetPreviousPosition(Accessor.Usage.GetLODIndex(), VertexIndex);
+	}
+
+	FORCEINLINE_DEBUGGABLE FVector GetSkinnedBonePosition(FSkeletalMeshAccessorHelper& Accessor, int32 BoneIndex)
+	{
+		return (*Accessor.BoneComponentSpaceTransforms)[BoneIndex].GetLocation();
+	}
+
+	FORCEINLINE_DEBUGGABLE FVector GetSkinnedBonePreviousPosition(FSkeletalMeshAccessorHelper& Accessor, int32 BoneIndex)
+	{
+		return (*Accessor.PrevBoneComponentSpaceTransforms)[BoneIndex].GetLocation();
+	}
+};
+
+//////////////////////////////////////////////////////////////////////////
+// Helper for accessing misc vertex data
+template<bool bUseFullPrecisionUVs>
+struct FSkelMeshVertexAccessor
+{
+	FORCEINLINE FVector2D GetVertexUV(FSkeletalMeshLODRenderData& LODData, int32 VertexIdx, int32 UVChannel)const
+	{
+		if (bUseFullPrecisionUVs)
+		{
+			return LODData.StaticVertexBuffers.StaticMeshVertexBuffer.GetVertexUV_Typed<EStaticMeshVertexUVType::HighPrecision>(VertexIdx, UVChannel);
+		}
+		else
+		{
+			return LODData.StaticVertexBuffers.StaticMeshVertexBuffer.GetVertexUV_Typed<EStaticMeshVertexUVType::Default>(VertexIdx, UVChannel);
+		}
+	}
+
+	FORCEINLINE FLinearColor GetVertexColor(FSkeletalMeshLODRenderData& LODData, int32 VertexIdx)const
+	{
+		return LODData.StaticVertexBuffers.ColorVertexBuffer.VertexColor(VertexIdx);
+	}
+};
+
+//////////////////////////////////////////////////////////////////////////
+//Function Binders.
+
+//External function binder choosing between template specializations based on if we're area weighting or not.
+template<typename NextBinder>
+struct TAreaWeightingModeBinder
+{
+	template<typename... ParamTypes>
+	static void Bind(UNiagaraDataInterface* Interface, const FVMExternalFunctionBindingInfo& BindingInfo, void* InstanceData, FVMExternalFunction &OutFunc)
+	{
+		FNDISkeletalMesh_InstanceData* InstData = (FNDISkeletalMesh_InstanceData*)InstanceData;
+		check(InstData);
+		UNiagaraDataInterfaceSkeletalMesh* MeshInterface = CastChecked<UNiagaraDataInterfaceSkeletalMesh>(Interface);
+		const FSkeletalMeshSamplingInfo& SamplingInfo = InstData->Mesh->GetSamplingInfo();
+
+		bool bAreaWeighting = false;
+		if (InstData->SamplingRegionIndices.Num() > 1)
+		{
+			bAreaWeighting = InstData->SamplingRegionAreaWeightedSampler.IsValid();
+		}
+		else if (InstData->SamplingRegionIndices.Num() == 1)
+		{
+			const FSkeletalMeshSamplingRegion& Region = SamplingInfo.GetRegion(InstData->SamplingRegionIndices[0]);
+			bAreaWeighting = Region.bSupportUniformlyDistributedSampling;
+		}
+		else
+		{
+			int32 LODIndex = InstData->GetLODIndex();
+			check(InstData->Mesh->GetLODInfo(LODIndex)->bAllowCPUAccess);
+			bAreaWeighting = InstData->Mesh->GetLODInfo(LODIndex)->bSupportUniformlyDistributedSampling;
+		}
+
+		if (bAreaWeighting)
+		{
+			NextBinder::template Bind<ParamTypes..., TIntegralConstant<ENDISkelMesh_AreaWeightingMode, ENDISkelMesh_AreaWeightingMode::AreaWeighted>>(Interface, BindingInfo, InstanceData, OutFunc);
+		}
+		else
+		{
+			NextBinder::template Bind<ParamTypes..., TIntegralConstant<ENDISkelMesh_AreaWeightingMode, ENDISkelMesh_AreaWeightingMode::None>>(Interface, BindingInfo, InstanceData, OutFunc);
+		}
+	}
+};
+
+//External function binder choosing between template specializations based on filtering methods
+template<typename NextBinder>
+struct TFilterModeBinder
+{
+	template<typename... ParamTypes>
+	static void Bind(UNiagaraDataInterface* Interface, const FVMExternalFunctionBindingInfo& BindingInfo, void* InstanceData, FVMExternalFunction &OutFunc)
+	{
+		FNDISkeletalMesh_InstanceData* InstData = (FNDISkeletalMesh_InstanceData*)InstanceData;
+		check(InstData);
+		UNiagaraDataInterfaceSkeletalMesh* MeshInterface = CastChecked<UNiagaraDataInterfaceSkeletalMesh>(Interface);
+
+		if (InstData->SamplingRegionIndices.Num() == 1)
+		{
+			NextBinder::template Bind<ParamTypes..., TIntegralConstant<ENDISkeletalMesh_FilterMode, ENDISkeletalMesh_FilterMode::SingleRegion>>(Interface, BindingInfo, InstanceData, OutFunc);
+		}
+		else if (InstData->SamplingRegionIndices.Num() > 1)
+		{
+			NextBinder::template Bind<ParamTypes..., TIntegralConstant<ENDISkeletalMesh_FilterMode, ENDISkeletalMesh_FilterMode::MultiRegion>>(Interface, BindingInfo, InstanceData, OutFunc);
+		}
+		else
+		{
+			NextBinder::template Bind<ParamTypes..., TIntegralConstant<ENDISkeletalMesh_FilterMode, ENDISkeletalMesh_FilterMode::None>>(Interface, BindingInfo, InstanceData, OutFunc);
+		}
+	}
+};
+
+//External function binder choosing between template specializations based vetrex data format
+template<typename NextBinder>
+struct TVertexAccessorBinder
+{
+	template<typename... ParamTypes>
+	static void Bind(UNiagaraDataInterface* Interface, const FVMExternalFunctionBindingInfo& BindingInfo, void* InstanceData, FVMExternalFunction &OutFunc)
+	{
+		FNDISkeletalMesh_InstanceData* InstData = (FNDISkeletalMesh_InstanceData*)InstanceData;
+		UNiagaraDataInterfaceSkeletalMesh* MeshInterface = CastChecked<UNiagaraDataInterfaceSkeletalMesh>(Interface);
+		USkeletalMeshComponent* Component = Cast<USkeletalMeshComponent>(InstData->Component.Get());
+		FSkinWeightVertexBuffer* SkinWeightBuffer = nullptr;
+		FSkeletalMeshLODRenderData& LODData = InstData->GetLODRenderDataAndSkinWeights(SkinWeightBuffer);
+
+		if (LODData.StaticVertexBuffers.StaticMeshVertexBuffer.GetUseFullPrecisionUVs())
+		{
+			NextBinder::template Bind<ParamTypes..., FSkelMeshVertexAccessor<true>>(Interface, BindingInfo, InstanceData, OutFunc);
+		}
+		else
+		{
+			NextBinder::template Bind<ParamTypes..., FSkelMeshVertexAccessor<false>>(Interface, BindingInfo, InstanceData, OutFunc);
+		}
+	}
+};
+
+//External function binder choosing between template specializations based on skinning mode
+template<typename NextBinder>
+struct TSkinningModeBinder
+{
+	template<typename... ParamTypes>
+	static void Bind(UNiagaraDataInterface* Interface, const FVMExternalFunctionBindingInfo& BindingInfo, void* InstanceData, FVMExternalFunction &OutFunc)
+	{
+		FNDISkeletalMesh_InstanceData* InstData = (FNDISkeletalMesh_InstanceData*)InstanceData;
+		UNiagaraDataInterfaceSkeletalMesh* MeshInterface = CastChecked<UNiagaraDataInterfaceSkeletalMesh>(Interface);
+		USkeletalMeshComponent* Component = Cast<USkeletalMeshComponent>(InstData->Component.Get());
+		if (MeshInterface->SkinningMode == ENDISkeletalMesh_SkinningMode::None || !Component)//Can't skin if we have no component.
+		{
+			NextBinder::template Bind<ParamTypes..., FSkinnedPositionAccessorHelper<TIntegralConstant<ENDISkeletalMesh_SkinningMode, ENDISkeletalMesh_SkinningMode::None>>>(Interface, BindingInfo, InstanceData, OutFunc);
+		}
+		else if (MeshInterface->SkinningMode == ENDISkeletalMesh_SkinningMode::SkinOnTheFly)
+		{
+			check(Component);
+			NextBinder::template Bind<ParamTypes..., FSkinnedPositionAccessorHelper<TIntegralConstant<ENDISkeletalMesh_SkinningMode, ENDISkeletalMesh_SkinningMode::SkinOnTheFly>>>(Interface, BindingInfo, InstanceData, OutFunc);
+		}
+		else if (MeshInterface->SkinningMode == ENDISkeletalMesh_SkinningMode::PreSkin)
+		{
+			check(Component);
+			NextBinder::template Bind<ParamTypes..., FSkinnedPositionAccessorHelper<TIntegralConstant<ENDISkeletalMesh_SkinningMode, ENDISkeletalMesh_SkinningMode::PreSkin>>>(Interface, BindingInfo, InstanceData, OutFunc);
+		}
+		else
+		{
+			checkf(false, TEXT("Invalid skinning mode in %s"), *Interface->GetPathName());
+		}
+	}
+};
\ No newline at end of file
