Fix for Skeletal Mesh Issues

diff --git a/Engine/Plugins/FX/Niagara/Source/Niagara/Private/NDISkeletalMesh_TriangleSampling.cpp b/Engine/Plugins/FX/Niagara/Source/Niagara/Private/NDISkeletalMesh_TriangleSampling.cpp
--- a/Engine/Plugins/FX/Niagara/Source/Niagara/Private/NDISkeletalMesh_TriangleSampling.cpp
+++ b/Engine/Plugins/FX/Niagara/Source/Niagara/Private/NDISkeletalMesh_TriangleSampling.cpp
@@ -810,7 +810,7 @@ void UNiagaraDataInterfaceSkeletalMesh::GetTriCoordSkinnedData(FVectorVMContext&
 
 	for (int32 i = 0; i < Context.NumInstances; ++i)
 	{
-		FMeshTriCoordinate MeshTriCoord(TriParam.GetAndAdvance(), FVector(BaryXParam.GetAndAdvance(), BaryYParam.GetAndAdvance(), BaryYParam.GetAndAdvance()));
+		FMeshTriCoordinate MeshTriCoord(TriParam.GetAndAdvance(), FVector(BaryXParam.GetAndAdvance(), BaryYParam.GetAndAdvance(), BaryZParam.GetAndAdvance()));
 
 		SkinningHandler.GetTrianlgeIndices(Accessor, MeshTriCoord.Tri, Idx0, Idx1, Idx2);
 
