Some fixes to skeletal mesh customization.

Still some aspects that don't work as expected.
Specificallly, the RefreshChildren call in SetSourceArray()

diff --git a/Engine/Plugins/FX/Niagara/Source/NiagaraEditorWidgets/Private/DetailCustomizations/NiagaraDataInterfaceSkeletalMeshDetails.cpp b/Engine/Plugins/FX/Niagara/Source/NiagaraEditorWidgets/Private/DetailCustomizations/NiagaraDataInterfaceSkeletalMeshDetails.cpp
--- a/Engine/Plugins/FX/Niagara/Source/NiagaraEditorWidgets/Private/DetailCustomizations/NiagaraDataInterfaceSkeletalMeshDetails.cpp
+++ b/Engine/Plugins/FX/Niagara/Source/NiagaraEditorWidgets/Private/DetailCustomizations/NiagaraDataInterfaceSkeletalMeshDetails.cpp
@@ -8,6 +8,7 @@
 #include "DetailCategoryBuilder.h"
 #include "IDetailGroup.h"
 #include "NiagaraComponent.h"
+#include "Engine/SkeletalMeshSocket.h"
 
  #define LOCTEXT_NAMESPACE "FNiagaraDataInterfaceSkeletalMeshDetails"
 
@@ -15,6 +16,7 @@ void FNiagaraDataInterfaceSkeletalMeshDetails::CustomizeDetails(IDetailLayoutBui
 {
 	 LayoutBuilder = &DetailBuilder;
 	 static const FName MeshCategoryName = TEXT("Mesh");
+	 static const FName SkelCategoryName = TEXT("Skeleton");
 
 	 TArray<TWeakObjectPtr<UObject>> SelectedObjects;
 	 DetailBuilder.GetObjectsBeingCustomized(SelectedObjects);
@@ -36,27 +38,61 @@ void FNiagaraDataInterfaceSkeletalMeshDetails::CustomizeDetails(IDetailLayoutBui
 		 MeshObject->GetOnMeshChanged().AddSP(this, &FNiagaraDataInterfaceSkeletalMeshDetails::OnDataChanged);
 	 }
 
-	 IDetailCategoryBuilder& MeshCategory = DetailBuilder.EditCategory(MeshCategoryName, LOCTEXT("Mesh", "Mesh"));
+	 MeshCategory = &DetailBuilder.EditCategory(MeshCategoryName, LOCTEXT("Mesh", "Mesh"));
 	 {
 		 TArray<TSharedRef<IPropertyHandle>> MeshProperties;
-		 MeshCategory.GetDefaultProperties(MeshProperties, true, false);
+		 MeshCategory->GetDefaultProperties(MeshProperties, true, true);
 
 		 TSharedPtr<IPropertyHandle> RegionsProperty = DetailBuilder.GetProperty(GET_MEMBER_NAME_CHECKED(UNiagaraDataInterfaceSkeletalMesh, SamplingRegions));
 		
 		 for (TSharedPtr<IPropertyHandle> Property : MeshProperties)
 		 {
-			 if (Property->GetProperty() != RegionsProperty->GetProperty())
+			 UProperty* PropertyPtr = Property->GetProperty();
+			 TArray<TSharedPtr<FName>> PossibleNames;
+			 if (PropertyPtr == RegionsProperty->GetProperty())
 			 {
-				 MeshCategory.AddProperty(Property);
+				 GenerateRegionsArray(PossibleNames);
+				 RegionsBuilder = TSharedPtr<FNiagaraDetailSourcedArrayBuilder>(new FNiagaraDetailSourcedArrayBuilder(Property.ToSharedRef(), PossibleNames));
+				 MeshCategory->AddCustomBuilder(RegionsBuilder.ToSharedRef());
 			 }
 			 else
 			 {
-				 TArray<TSharedPtr<FName>> PossibleRegions = GenerateSourceArray();
-				 RegionsBuilder = TSharedPtr<FNiagaraDetailSourcedArrayBuilder>(new FNiagaraDetailSourcedArrayBuilder(Property.ToSharedRef(), PossibleRegions));
-				 MeshCategory.AddCustomBuilder(RegionsBuilder.ToSharedRef());
+				 MeshCategory->AddProperty(Property);
 			 }
 		 }
 	 }
+
+	 SkelCategory = &DetailBuilder.EditCategory(SkelCategoryName, LOCTEXT("SkeletonCat", "Skeleton"));
+	 {
+		 TArray<TSharedRef<IPropertyHandle>> SkelProperties;
+		 SkelCategory->GetDefaultProperties(SkelProperties, true, true);
+
+		 TSharedPtr<IPropertyHandle> BonesProperty = DetailBuilder.GetProperty(GET_MEMBER_NAME_CHECKED(UNiagaraDataInterfaceSkeletalMesh, SpecificBones));
+		 TSharedPtr<IPropertyHandle> SocketsProperty = DetailBuilder.GetProperty(GET_MEMBER_NAME_CHECKED(UNiagaraDataInterfaceSkeletalMesh, SpecificSockets));
+
+		 for (TSharedPtr<IPropertyHandle> Property : SkelProperties)
+		 {
+			 UProperty* PropertyPtr = Property->GetProperty();
+			 TArray<TSharedPtr<FName>> PossibleNames;
+
+			if (PropertyPtr == BonesProperty->GetProperty())
+			 {
+				GenerateBonesArray(PossibleNames);
+				BonesBuilder = TSharedPtr<FNiagaraDetailSourcedArrayBuilder>(new FNiagaraDetailSourcedArrayBuilder(Property.ToSharedRef(), PossibleNames));
+				SkelCategory->AddCustomBuilder(BonesBuilder.ToSharedRef());
+			 }
+			 else if (PropertyPtr == SocketsProperty->GetProperty())
+			 {
+				 GenerateSocketsArray(PossibleNames);
+				 SocketsBuilder = TSharedPtr<FNiagaraDetailSourcedArrayBuilder>(new FNiagaraDetailSourcedArrayBuilder(Property.ToSharedRef(), PossibleNames));
+				 SkelCategory->AddCustomBuilder(SocketsBuilder.ToSharedRef());
+			 }
+			 else
+			 {
+				 SkelCategory->AddProperty(Property);
+			 }			 
+		 }
+	 }
  }
 
  TSharedRef<IDetailCustomization> FNiagaraDataInterfaceSkeletalMeshDetails::MakeInstance()
@@ -83,29 +119,73 @@ void FNiagaraDataInterfaceSkeletalMeshDetails::CustomizeDetails(IDetailLayoutBui
 
  void FNiagaraDataInterfaceSkeletalMeshDetails::OnDataChanged()
  {
-	 TArray<TSharedPtr<FName>> PossibleRegions = GenerateSourceArray();
-	 RegionsBuilder->SetSourceArray(PossibleRegions);
+	TArray<TSharedPtr<FName>> PossibleNames;
+	GenerateRegionsArray(PossibleNames);
+	RegionsBuilder->SetSourceArray(PossibleNames);
+
+	GenerateBonesArray(PossibleNames);
+	BonesBuilder->SetSourceArray(PossibleNames);
+
+	GenerateSocketsArray(PossibleNames);
+	SocketsBuilder->SetSourceArray(PossibleNames);
+ }
+
+void FNiagaraDataInterfaceSkeletalMeshDetails::GenerateRegionsArray(TArray<TSharedPtr<FName>>& SourceArray)
+ {
+	SourceArray.Reset();
+	 if (MeshInterface.IsValid())
+	 {
+		 TWeakObjectPtr<USceneComponent> SceneComponent;
+		 USkeletalMeshComponent* FoundSkelComp = nullptr;
+		 USkeletalMesh* Mesh = UNiagaraDataInterfaceSkeletalMesh::GetSkeletalMeshHelper(MeshInterface.Get(), Cast<UNiagaraComponent>(MeshInterface->GetOuter()), SceneComponent, FoundSkelComp);
+
+		 if (Mesh != nullptr)
+		 {
+			 for (FSkeletalMeshSamplingRegion Region : Mesh->GetSamplingInfo().Regions)
+			 {
+				 SourceArray.Add(MakeShared<FName>(Region.Name));
+			 }
+		 }
+	 }
+ }
+
+ void FNiagaraDataInterfaceSkeletalMeshDetails::GenerateBonesArray(TArray<TSharedPtr<FName>>& SourceArray)
+ {
+	 SourceArray.Reset();
+	 if (MeshInterface.IsValid())
+	 {
+		 TWeakObjectPtr<USceneComponent> SceneComponent;
+		 USkeletalMeshComponent* FoundSkelComp = nullptr;
+		 USkeletalMesh* Mesh = UNiagaraDataInterfaceSkeletalMesh::GetSkeletalMeshHelper(MeshInterface.Get(), Cast<UNiagaraComponent>(MeshInterface->GetOuter()), SceneComponent, FoundSkelComp);
+
+		 if (Mesh != nullptr)
+		 {
+			 for (const FMeshBoneInfo& Bone : Mesh->RefSkeleton.GetRefBoneInfo())
+			 {
+				 SourceArray.Add(MakeShared<FName>(Bone.Name));
+			 }
+		 }
+	 }
  }
 
- TArray<TSharedPtr<FName>> FNiagaraDataInterfaceSkeletalMeshDetails::GenerateSourceArray()
+void FNiagaraDataInterfaceSkeletalMeshDetails::GenerateSocketsArray(TArray<TSharedPtr<FName>>& SourceArray)
  {
-	TArray<TSharedPtr<FName>> SourceArray;
-	if (MeshInterface.IsValid())
-	{
-			TWeakObjectPtr<USceneComponent> SceneComponent;
-			USkeletalMeshComponent* FoundSkelComp = nullptr;
-			USkeletalMesh* Mesh = UNiagaraDataInterfaceSkeletalMesh::GetSkeletalMeshHelper(MeshInterface.Get(), Cast<UNiagaraComponent>(MeshInterface->GetOuter()), SceneComponent, FoundSkelComp);
-
-			if (Mesh != nullptr)
-			{
-				for (FSkeletalMeshSamplingRegion Region : Mesh->GetSamplingInfo().Regions)
-				{
-					TSharedPtr<FName> RegionName(new FName(Region.Name));
-					SourceArray.Add(RegionName);
-				}
-			}
-	}
-	return SourceArray;
+	SourceArray.Reset();
+	 if (MeshInterface.IsValid())
+	 {
+		 TWeakObjectPtr<USceneComponent> SceneComponent;
+		 USkeletalMeshComponent* FoundSkelComp = nullptr;
+		 USkeletalMesh* Mesh = UNiagaraDataInterfaceSkeletalMesh::GetSkeletalMeshHelper(MeshInterface.Get(), Cast<UNiagaraComponent>(MeshInterface->GetOuter()), SceneComponent, FoundSkelComp);
+
+		 if (Mesh != nullptr)
+		 {
+			 for (int32 SocketIdx = 0; SocketIdx < Mesh->NumSockets(); ++SocketIdx)
+			 {
+				 const USkeletalMeshSocket* SocketInfo = Mesh->GetSocketByIndex(SocketIdx);
+				 SourceArray.Add(MakeShared<FName>(SocketInfo->SocketName));
+			 }
+		 }
+	 }
  }
 
  FNiagaraDataInterfaceSkeletalMeshDetails::~FNiagaraDataInterfaceSkeletalMeshDetails()
diff --git a/Engine/Plugins/FX/Niagara/Source/NiagaraEditorWidgets/Private/DetailCustomizations/NiagaraDataInterfaceSkeletalMeshDetails.h b/Engine/Plugins/FX/Niagara/Source/NiagaraEditorWidgets/Private/DetailCustomizations/NiagaraDataInterfaceSkeletalMeshDetails.h
--- a/Engine/Plugins/FX/Niagara/Source/NiagaraEditorWidgets/Private/DetailCustomizations/NiagaraDataInterfaceSkeletalMeshDetails.h
+++ b/Engine/Plugins/FX/Niagara/Source/NiagaraEditorWidgets/Private/DetailCustomizations/NiagaraDataInterfaceSkeletalMeshDetails.h
@@ -26,11 +26,18 @@ public:
 private:
 	void OnInterfaceChanged();
 	void OnDataChanged();
-	TArray<TSharedPtr<FName>> GenerateSourceArray();
+	void GenerateRegionsArray(TArray<TSharedPtr<FName>>& SourceArray);
+	void GenerateBonesArray(TArray<TSharedPtr<FName>>& SourceArray);
+	void GenerateSocketsArray(TArray<TSharedPtr<FName>>& SourceArray);
 
 private:
 	TSharedPtr<FNiagaraDetailSourcedArrayBuilder> RegionsBuilder;
+	TSharedPtr<FNiagaraDetailSourcedArrayBuilder> BonesBuilder;
+	TSharedPtr<FNiagaraDetailSourcedArrayBuilder> SocketsBuilder;
 	IDetailLayoutBuilder* LayoutBuilder;
 	TWeakObjectPtr<UNiagaraDataInterfaceSkeletalMesh>  MeshInterface;
 	TWeakObjectPtr<USkeletalMesh> MeshObject;
+
+	IDetailCategoryBuilder* MeshCategory;
+	IDetailCategoryBuilder* SkelCategory;
 };
diff --git a/Engine/Plugins/FX/Niagara/Source/NiagaraEditorWidgets/Private/DetailCustomizations/NiagaraDetailSourcedArrayBuilder.cpp b/Engine/Plugins/FX/Niagara/Source/NiagaraEditorWidgets/Private/DetailCustomizations/NiagaraDetailSourcedArrayBuilder.cpp
--- a/Engine/Plugins/FX/Niagara/Source/NiagaraEditorWidgets/Private/DetailCustomizations/NiagaraDetailSourcedArrayBuilder.cpp
+++ b/Engine/Plugins/FX/Niagara/Source/NiagaraEditorWidgets/Private/DetailCustomizations/NiagaraDetailSourcedArrayBuilder.cpp
@@ -18,13 +18,6 @@ FNiagaraDetailSourcedArrayBuilder::FNiagaraDetailSourcedArrayBuilder(TSharedRef<
 	, ArrayProperty(InBaseProperty->AsArray())
 	, FNameSubproperty(InFNameSubproperty)
 {
-	// bradut TODO set new Delegate for when the number of children in the array changes, where you empty the map of selectors and double check if it is regenerated and also call 
-	// 
-	FSimpleDelegate OnArrayNumChildrenChanged = FSimpleDelegate::CreateRaw(this, &FNiagaraDetailSourcedArrayBuilder::OnArrayNumChildrenChanged);
-	ArrayProperty->SetOnNumElementsChanged(OnArrayNumChildrenChanged);
-	uint32 EntryCount;
-	ArrayProperty->GetNumElements(EntryCount);
-	NameSelectors.AddZeroed(EntryCount);
 };
 
 void FNiagaraDetailSourcedArrayBuilder::OnGenerateEntry(TSharedRef<IPropertyHandle> PropertyHandle, int32 ArrayIndex, IDetailChildrenBuilder& ChildrenBuilder)
@@ -44,39 +37,40 @@ void FNiagaraDetailSourcedArrayBuilder::OnGenerateEntry(TSharedRef<IPropertyHand
 	{
 		SubPropertyHandle = PropertyHandle->GetChildHandle(FNameSubproperty).ToSharedRef();
 	}
-	NameSelectors[ArrayIndex] = SNew(SNiagaraNamePropertySelector, SubPropertyHandle.ToSharedRef(), OptionsSourceList);
-	RegionRow.CustomWidget(false)
-		.NameContent()
-		[
-			SubPropertyHandle->CreatePropertyNameWidget()
-		]
+
+	if (OptionsSourceList.Num() > 0)
+	{
+		RegionRow.CustomWidget(false)
+			.NameContent()
+			[
+				SubPropertyHandle->CreatePropertyNameWidget()
+			]
+			.ValueContent()
+			.MaxDesiredWidth(TOptional<float>())
+			[
+				SNew(SNiagaraNamePropertySelector, SubPropertyHandle.ToSharedRef(), OptionsSourceList)
+			];
+	}
+	else
+	{
+		RegionRow.CustomWidget(false)
+			.NameContent()
+			[
+				SubPropertyHandle->CreatePropertyNameWidget()
+			]
 		.ValueContent()
-		.MaxDesiredWidth(TOptional<float>())
-		[
-			// add combo button
-			NameSelectors[ArrayIndex].ToSharedRef()
-		];
-}
+			.MaxDesiredWidth(TOptional<float>())
+			[
 
-void FNiagaraDetailSourcedArrayBuilder::OnArrayNumChildrenChanged()
-{
-	NameSelectors.Empty();
-	uint32 EntryCount;
-	ArrayProperty->GetNumElements(EntryCount);
-	NameSelectors.AddZeroed(EntryCount); // allocate the memory for ease of use when refreshing children
-	RefreshChildren();
+				SubPropertyHandle->CreatePropertyValueWidget(false)
+			];
+	}
 }
 
 void FNiagaraDetailSourcedArrayBuilder::SetSourceArray(TArray<TSharedPtr<FName>>& InOptionsSource)
 {
 	OptionsSourceList = InOptionsSource;
-	for (auto NameSelector : NameSelectors)
-	{
-		if (NameSelector.IsValid())
-		{
-			NameSelector->SetSourceArray(InOptionsSource);
-		}
-	}
+	RefreshChildren();
 }
 
 void FNiagaraDetailSourcedArrayBuilder::GenerateChildContent(IDetailChildrenBuilder& ChildrenBuilder)
diff --git a/Engine/Plugins/FX/Niagara/Source/NiagaraEditorWidgets/Private/DetailCustomizations/NiagaraDetailSourcedArrayBuilder.h b/Engine/Plugins/FX/Niagara/Source/NiagaraEditorWidgets/Private/DetailCustomizations/NiagaraDetailSourcedArrayBuilder.h
--- a/Engine/Plugins/FX/Niagara/Source/NiagaraEditorWidgets/Private/DetailCustomizations/NiagaraDetailSourcedArrayBuilder.h
+++ b/Engine/Plugins/FX/Niagara/Source/NiagaraEditorWidgets/Private/DetailCustomizations/NiagaraDetailSourcedArrayBuilder.h
@@ -24,7 +24,7 @@ public:
 
 private: 
 	void OnGenerateEntry(TSharedRef<IPropertyHandle> ElementProperty, int32 ElementIndex, IDetailChildrenBuilder& ChildrenBuilder);
-	void OnArrayNumChildrenChanged();
+
 private:
 	TArray<TSharedPtr<FName>> OptionsSourceList;
 	/** List of component class names, filtered by the current search string */
@@ -32,5 +32,4 @@ private:
 	TSharedPtr<IPropertyHandleArray> ArrayProperty;
 	/** Subproperty of type FName that needs to be edited - in case of FName wrappers */
 	FName FNameSubproperty;
-	TArray<TSharedPtr<SNiagaraNamePropertySelector>> NameSelectors;
 };
diff --git a/Engine/Plugins/FX/Niagara/Source/NiagaraEditorWidgets/Private/DetailCustomizations/NiagaraSkeletalMeshSamplingInfoDetails.cpp b/Engine/Plugins/FX/Niagara/Source/NiagaraEditorWidgets/Private/DetailCustomizations/NiagaraSkeletalMeshSamplingInfoDetails.cpp
deleted file mode 100644
diff --git a/Engine/Plugins/FX/Niagara/Source/NiagaraEditorWidgets/Private/DetailCustomizations/NiagaraSkeletalMeshSamplingInfoDetails.h b/Engine/Plugins/FX/Niagara/Source/NiagaraEditorWidgets/Private/DetailCustomizations/NiagaraSkeletalMeshSamplingInfoDetails.h
deleted file mode 100644
diff --git a/Engine/Plugins/FX/Niagara/Source/NiagaraEditorWidgets/Private/NiagaraEditorWidgetsModule.cpp b/Engine/Plugins/FX/Niagara/Source/NiagaraEditorWidgets/Private/NiagaraEditorWidgetsModule.cpp
--- a/Engine/Plugins/FX/Niagara/Source/NiagaraEditorWidgets/Private/NiagaraEditorWidgetsModule.cpp
+++ b/Engine/Plugins/FX/Niagara/Source/NiagaraEditorWidgets/Private/NiagaraEditorWidgetsModule.cpp
@@ -7,7 +7,6 @@
 #include "DetailCustomizations/NiagaraDataInterfaceCurveDetails.h"
 #include "DetailCustomizations/NiagaraDataInterfaceDetails.h"
 #include "DetailCustomizations/NiagaraDataInterfaceSkeletalMeshDetails.h"
-#include "DetailCustomizations/NiagaraSkeletalMeshSamplingInfoDetails.h"
 
 #include "Modules/ModuleManager.h"
 #include "PropertyEditorModule.h"
@@ -98,7 +97,6 @@ void FNiagaraEditorWidgetsModule::StartupModule()
 	PropertyModule.RegisterCustomClassLayout("NiagaraDataInterfaceVector4Curve", FOnGetDetailCustomizationInstance::CreateStatic(&FNiagaraDataInterfaceVector4CurveDetails::MakeInstance));
 	PropertyModule.RegisterCustomClassLayout("NiagaraDataInterfaceColorCurve", FOnGetDetailCustomizationInstance::CreateStatic(&FNiagaraDataInterfaceColorCurveDetails::MakeInstance));
 	PropertyModule.RegisterCustomClassLayout("NiagaraDataInterfaceSkeletalMesh", FOnGetDetailCustomizationInstance::CreateStatic(&FNiagaraDataInterfaceSkeletalMeshDetails::MakeInstance));
-	PropertyModule.RegisterCustomClassLayout("SkeletalMesh", FOnGetDetailCustomizationInstance::CreateStatic(&FNiagaraSkeletalMeshSamplingInfoDetails::MakeInstance));
 }
 
 void FNiagaraEditorWidgetsModule::ShutdownModule()
diff --git a/Engine/Source/Editor/Persona/Private/Customization/SkeletalMeshRegionCustomization.cpp b/Engine/Source/Editor/Persona/Private/Customization/SkeletalMeshRegionCustomization.cpp
new file mode 100644
--- /dev/null
+++ b/Engine/Source/Editor/Persona/Private/Customization/SkeletalMeshRegionCustomization.cpp
@@ -0,0 +1,209 @@
+// Copyright 1998-2018 Epic Games, Inc. All Rights Reserved.
+
+#include "Customization/SkeletalMeshRegionCustomization.h"
+
+#include "UObject/ObjectMacros.h"
+#include "UObject/Class.h"
+#include "IDetailsView.h"
+#include "EditorStyleSet.h"
+#include "Widgets/DeclarativeSyntaxSupport.h"
+#include "PropertyHandle.h"
+#include "DetailLayoutBuilder.h"
+#include "DetailCategoryBuilder.h"
+
+#include "Widgets/Text/STextBlock.h"
+#include "Widgets/Input/SNumericEntryBox.h"
+
+#include "Animation/AnimSequence.h"
+#include "Animation/BlendSpaceBase.h"
+#include "Animation/BlendSpace1D.h"
+#include "SAnimationBlendSpaceGridWidget.h"
+#include "PropertyCustomizationHelpers.h"
+
+#include "PackageTools.h"
+#include "IDetailGroup.h"
+#include "Editor.h"
+#include "IDetailChildrenBuilder.h"
+#include "Widgets/Input/SComboBox.h"
+
+#define LOCTEXT_NAMESPACE "SkeletalMeshRegionCustomization"
+
+void FNiagaraSkeletalMeshRegionBoneFilterDetails::CustomizeHeader(TSharedRef<IPropertyHandle> StructPropertyHandle, class FDetailWidgetRow& HeaderRow, IPropertyTypeCustomizationUtils& StructCustomizationUtils)
+{
+	HeaderRow
+		.NameContent()
+		[
+			StructPropertyHandle->CreatePropertyNameWidget()
+		];
+}
+
+void FNiagaraSkeletalMeshRegionBoneFilterDetails::CustomizeChildren(TSharedRef<IPropertyHandle> StructPropertyHandle, class IDetailChildrenBuilder& ChildBuilder, IPropertyTypeCustomizationUtils& StructCustomizationUtils)
+{
+	TArray<TWeakObjectPtr<UObject>> SelectedObjects;
+	ChildBuilder.GetParentCategory().GetParentLayout().GetObjectsBeingCustomized(SelectedObjects);
+	check(SelectedObjects.Num() == 1);
+	MeshObject = Cast<USkeletalMesh>(SelectedObjects[0].Get());
+	check(MeshObject);
+	
+	uint32 NumChildren;
+	StructPropertyHandle->GetNumChildren(NumChildren);
+
+	for (uint32 ChildIndex = 0; ChildIndex < NumChildren; ++ChildIndex)
+	{
+		TSharedRef<IPropertyHandle> ChildHandle = StructPropertyHandle->GetChildHandle(ChildIndex).ToSharedRef();
+		const FName ChildName = ChildHandle->GetProperty()->GetFName();
+
+		// Pick out Min and Max range value properties
+		if (ChildName == GET_MEMBER_NAME_CHECKED(FSkeletalMeshSamplingRegionBoneFilter, BoneName))
+		{
+			BoneNameHandle = ChildHandle;
+			ChildBuilder.AddCustomRow(BoneNameHandle->GetPropertyDisplayName())
+				.NameWidget
+				[
+					BoneNameHandle->CreatePropertyNameWidget()
+				]
+				.ValueWidget
+				[
+					SNew(SComboBox<TSharedPtr<FName>>)
+					.OptionsSource(&PossibleBoneNames)
+					.ContentPadding(2.0f)
+					.OnGenerateWidget(this, &FNiagaraSkeletalMeshRegionBoneFilterDetails::HandleBoneNameComboBoxGenerateWidget)
+					.OnSelectionChanged(this, &FNiagaraSkeletalMeshRegionBoneFilterDetails::HandleBoneNameComboBoxSelectionChanged)
+					.OnComboBoxOpening(this, &FNiagaraSkeletalMeshRegionBoneFilterDetails::OnComboOpening)
+					[
+						SNew(STextBlock)
+						.Text(this, &FNiagaraSkeletalMeshRegionBoneFilterDetails::HandleBoneNameComboBoxContentText)
+					]
+				];
+		}
+		else
+		{
+			ChildBuilder.AddProperty(ChildHandle);
+		}
+	}
+}
+
+FText FNiagaraSkeletalMeshRegionBoneFilterDetails::HandleBoneNameComboBoxContentText() const
+{
+	FName OutName;
+	BoneNameHandle->GetValue(OutName);
+	return FText::FromName(OutName);
+}
+
+TSharedRef<SWidget> FNiagaraSkeletalMeshRegionBoneFilterDetails::HandleBoneNameComboBoxGenerateWidget(TSharedPtr<FName> StringItem)
+{
+	FName DefaultName = StringItem.IsValid() ? *StringItem : NAME_None;
+	return SNew(STextBlock).Text(FText::FromName(DefaultName));
+}
+
+void FNiagaraSkeletalMeshRegionBoneFilterDetails::HandleBoneNameComboBoxSelectionChanged(TSharedPtr<FName> StringItem, ESelectInfo::Type SelectInfo)
+{
+	BoneNameHandle->SetValue(*StringItem);
+}
+
+
+void FNiagaraSkeletalMeshRegionBoneFilterDetails::OnComboOpening()
+{
+	PossibleBoneNames.Reset();
+	if (MeshObject->Skeleton != nullptr)
+	{
+		//Populate PossibleBonesNames
+		for (FMeshBoneInfo Bone : MeshObject->Skeleton->GetReferenceSkeleton().GetRefBoneInfo())
+		{
+			TSharedPtr<FName> BoneName(new FName(Bone.Name));
+			PossibleBoneNames.Add(BoneName);
+		}
+	}
+}
+
+//////////////////////////////////////////////////////////////////////////
+
+
+void FNiagaraSkeletalMeshRegionMaterialFilterDetails::CustomizeHeader(TSharedRef<IPropertyHandle> StructPropertyHandle, class FDetailWidgetRow& HeaderRow, IPropertyTypeCustomizationUtils& StructCustomizationUtils)
+{
+	HeaderRow
+		.NameContent()
+		[
+			StructPropertyHandle->CreatePropertyNameWidget()
+		];
+}
+
+void FNiagaraSkeletalMeshRegionMaterialFilterDetails::CustomizeChildren(TSharedRef<IPropertyHandle> StructPropertyHandle, class IDetailChildrenBuilder& ChildBuilder, IPropertyTypeCustomizationUtils& StructCustomizationUtils)
+{
+	TArray<TWeakObjectPtr<UObject>> SelectedObjects;
+	ChildBuilder.GetParentCategory().GetParentLayout().GetObjectsBeingCustomized(SelectedObjects);
+	check(SelectedObjects.Num() == 1);
+	MeshObject = Cast<USkeletalMesh>(SelectedObjects[0].Get());
+	check(MeshObject);
+	
+	uint32 NumChildren;
+	StructPropertyHandle->GetNumChildren(NumChildren);
+
+	for (uint32 ChildIndex = 0; ChildIndex < NumChildren; ++ChildIndex)
+	{
+		TSharedRef<IPropertyHandle> ChildHandle = StructPropertyHandle->GetChildHandle(ChildIndex).ToSharedRef();
+		const FName ChildName = ChildHandle->GetProperty()->GetFName();
+
+		// Pick out Min and Max range value properties
+		if (ChildName == GET_MEMBER_NAME_CHECKED(FSkeletalMeshSamplingRegionMaterialFilter, MaterialName))
+		{
+			MaterialNameHandle = ChildHandle;
+			ChildBuilder.AddCustomRow(MaterialNameHandle->GetPropertyDisplayName())
+				.NameWidget
+				[
+					MaterialNameHandle->CreatePropertyNameWidget()
+				]
+			.ValueWidget
+				[
+					SNew(SComboBox<TSharedPtr<FName>>)
+					.OptionsSource(&PossibleMaterialNames)
+					.ContentPadding(2.0f)
+					.OnGenerateWidget(this, &FNiagaraSkeletalMeshRegionMaterialFilterDetails::HandleMaterialNameComboBoxGenerateWidget)
+					.OnSelectionChanged(this, &FNiagaraSkeletalMeshRegionMaterialFilterDetails::HandleMaterialNameComboBoxSelectionChanged)
+					.OnComboBoxOpening(this, &FNiagaraSkeletalMeshRegionMaterialFilterDetails::OnComboOpening)
+					[
+						SNew(STextBlock)
+						.Text(this, &FNiagaraSkeletalMeshRegionMaterialFilterDetails::HandleMaterialNameComboBoxContentText)
+					]
+				];
+		}
+		else
+		{
+			ChildBuilder.AddProperty(ChildHandle);
+		}
+	}
+}
+
+FText FNiagaraSkeletalMeshRegionMaterialFilterDetails::HandleMaterialNameComboBoxContentText() const
+{
+	FName OutName;
+	MaterialNameHandle->GetValue(OutName);
+	return FText::FromName(OutName);
+}
+
+TSharedRef<SWidget> FNiagaraSkeletalMeshRegionMaterialFilterDetails::HandleMaterialNameComboBoxGenerateWidget(TSharedPtr<FName> StringItem)
+{
+	FName DefaultName = StringItem.IsValid() ? *StringItem : NAME_None;
+	return SNew(STextBlock).Text(FText::FromName(DefaultName));
+}
+
+void FNiagaraSkeletalMeshRegionMaterialFilterDetails::HandleMaterialNameComboBoxSelectionChanged(TSharedPtr<FName> StringItem, ESelectInfo::Type SelectInfo)
+{
+	MaterialNameHandle->SetValue(*StringItem);
+}
+
+void FNiagaraSkeletalMeshRegionMaterialFilterDetails::OnComboOpening()
+{
+	PossibleMaterialNames.Reset();
+	if (MeshObject->Skeleton != nullptr)
+	{
+		//Populate PossibleMaterialsNames
+		for (FSkeletalMaterial Material : MeshObject->Materials)
+		{
+			TSharedPtr<FName> MaterialName(new FName(Material.MaterialSlotName));
+			PossibleMaterialNames.Add(MaterialName);
+		}
+	}
+}
+
+#undef LOCTEXT_NAMESPACE
\ No newline at end of file
diff --git a/Engine/Source/Editor/Persona/Private/Customization/SkeletalMeshRegionCustomization.h b/Engine/Source/Editor/Persona/Private/Customization/SkeletalMeshRegionCustomization.h
new file mode 100644
--- /dev/null
+++ b/Engine/Source/Editor/Persona/Private/Customization/SkeletalMeshRegionCustomization.h
@@ -0,0 +1,70 @@
+// Copyright 1998-2018 Epic Games, Inc. All Rights Reserved.
+
+#pragma once
+
+#include "CoreMinimal.h"
+#include "Layout/Visibility.h"
+#include "PropertyHandle.h"
+#include "IDetailCustomization.h"
+
+class IDetailLayoutBuilder;
+class IPropertyHandle;
+class IDetailChildrenBuilder;
+class USkeletalMesh;
+
+
+class FNiagaraSkeletalMeshRegionBoneFilterDetails : public IPropertyTypeCustomization
+{
+public:
+	static TSharedRef<IPropertyTypeCustomization> MakeInstance()
+	{
+		return MakeShareable(new FNiagaraSkeletalMeshRegionBoneFilterDetails());
+	}
+
+	// IDetailCustomization interface
+	virtual void CustomizeHeader(TSharedRef<IPropertyHandle> StructPropertyHandle, class FDetailWidgetRow& HeaderRow, IPropertyTypeCustomizationUtils& StructCustomizationUtils) override;
+	virtual void CustomizeChildren(TSharedRef<IPropertyHandle> StructPropertyHandle, class IDetailChildrenBuilder& ChildBuilder, IPropertyTypeCustomizationUtils& StructCustomizationUtils) override;
+
+private:
+	FText HandleBoneNameComboBoxContentText() const;
+	TSharedRef<SWidget> HandleBoneNameComboBoxGenerateWidget(TSharedPtr<FName> StringItem);
+	void HandleBoneNameComboBoxSelectionChanged(TSharedPtr<FName> StringItem, ESelectInfo::Type SelectInfo);
+
+	void OnComboOpening();
+
+	/** The mesh object whose details we're customizing. */
+	USkeletalMesh * MeshObject;
+
+	TArray<TSharedPtr<FName>> PossibleBoneNames;
+
+	TSharedPtr<IPropertyHandle> BoneNameHandle;
+};
+
+class FNiagaraSkeletalMeshRegionMaterialFilterDetails : public IPropertyTypeCustomization
+{
+public:
+	static TSharedRef<IPropertyTypeCustomization> MakeInstance()
+	{
+		return MakeShareable(new FNiagaraSkeletalMeshRegionMaterialFilterDetails());
+	}
+
+	// IDetailCustomization interface
+	virtual void CustomizeHeader(TSharedRef<IPropertyHandle> StructPropertyHandle, class FDetailWidgetRow& HeaderRow, IPropertyTypeCustomizationUtils& StructCustomizationUtils) override;
+	virtual void CustomizeChildren(TSharedRef<IPropertyHandle> StructPropertyHandle, class IDetailChildrenBuilder& ChildBuilder, IPropertyTypeCustomizationUtils& StructCustomizationUtils) override;
+
+private:
+	FText HandleMaterialNameComboBoxContentText() const;
+	TSharedRef<SWidget> HandleMaterialNameComboBoxGenerateWidget(TSharedPtr<FName> StringItem);
+	void HandleMaterialNameComboBoxSelectionChanged(TSharedPtr<FName> StringItem, ESelectInfo::Type SelectInfo);
+
+	void OnComboOpening();
+
+	/** The mesh object whose details we're customizing. */
+	USkeletalMesh * MeshObject;
+
+	TArray<TSharedPtr<FName>> PossibleMaterialNames;
+
+	TSharedPtr<IPropertyHandle> MaterialNameHandle;
+};
+
+
diff --git a/Engine/Source/Editor/Persona/Private/PersonaModule.cpp b/Engine/Source/Editor/Persona/Private/PersonaModule.cpp
--- a/Engine/Source/Editor/Persona/Private/PersonaModule.cpp
+++ b/Engine/Source/Editor/Persona/Private/PersonaModule.cpp
@@ -76,6 +76,7 @@
 #include "PersonaPreviewSceneAnimationController.h"
 #include "PersonaPreviewSceneRefPoseController.h"
 #include "AssetViewerSettings.h"
+#include "Customization/SkeletalMeshRegionCustomization.h"
 
 IMPLEMENT_MODULE( FPersonaModule, Persona );
 
@@ -131,6 +132,9 @@ void FPersonaModule::StartupModule()
 
 		PropertyModule.RegisterCustomPropertyTypeLayout("BlendParameter", FOnGetPropertyTypeCustomizationInstance::CreateStatic(&FBlendParameterDetails::MakeInstance));
 		PropertyModule.RegisterCustomPropertyTypeLayout("InterpolationParameter", FOnGetPropertyTypeCustomizationInstance::CreateStatic(&FInterpolationParameterDetails::MakeInstance));
+
+		PropertyModule.RegisterCustomPropertyTypeLayout("SkeletalMeshSamplingRegionBoneFilter", FOnGetPropertyTypeCustomizationInstance::CreateStatic(&FNiagaraSkeletalMeshRegionBoneFilterDetails::MakeInstance));
+		PropertyModule.RegisterCustomPropertyTypeLayout("SkeletalMeshSamplingRegionMaterialFilter", FOnGetPropertyTypeCustomizationInstance::CreateStatic(&FNiagaraSkeletalMeshRegionMaterialFilterDetails::MakeInstance));
 	}
 
 	// Register the editor modes
@@ -166,6 +170,9 @@ void FPersonaModule::ShutdownModule()
 
 		PropertyModule.UnregisterCustomPropertyTypeLayout("BlendParameter");
 		PropertyModule.UnregisterCustomPropertyTypeLayout("InterpolationParameter");
+
+		PropertyModule.UnregisterCustomClassLayout("SkeletalMeshSamplingRegionBoneFilter");
+		PropertyModule.UnregisterCustomClassLayout("SkeletalMeshSamplingRegionMaterialFilter");
 	}
 }
 
