Some additional saftey in some skel mesh sampling functions.
Avoids at least one crash with UVSets in GetTriCoordUV()

diff --git a/Engine/Plugins/FX/Niagara/Source/Niagara/Private/NDISkeletalMesh_BoneSampling.cpp b/Engine/Plugins/FX/Niagara/Source/Niagara/Private/NDISkeletalMesh_BoneSampling.cpp
--- a/Engine/Plugins/FX/Niagara/Source/Niagara/Private/NDISkeletalMesh_BoneSampling.cpp
+++ b/Engine/Plugins/FX/Niagara/Source/Niagara/Private/NDISkeletalMesh_BoneSampling.cpp
@@ -246,10 +246,17 @@ void UNiagaraDataInterfaceSkeletalMesh::GetSpecificBoneAt(FVectorVMContext& Cont
 	const TArray<int32>& SpecificBonesArray = InstData->SpecificBones;
 
 	int32 Max = SpecificBones.Num() - 1;
-	for (int32 i = 0; i < Context.NumInstances; ++i)
+	if (Max != INDEX_NONE)
+	{
+		for (int32 i = 0; i < Context.NumInstances; ++i)
+		{
+			int32 BoneIndex = FMath::Clamp(BoneParam.GetAndAdvance(), 0, Max);
+			*OutBone.GetDestAndAdvance() = SpecificBonesArray[BoneIndex];
+		}
+	}
+	else
 	{
-		int32 BoneIndex = FMath::Clamp(BoneParam.GetAndAdvance(), 0, Max);
-		*OutBone.GetDestAndAdvance() = SpecificBonesArray[BoneIndex];
+		FMemory::Memset(OutBone.GetDest(), 0xFF, Context.NumInstances * sizeof(int32));
 	}
 }
 
@@ -263,10 +270,17 @@ void UNiagaraDataInterfaceSkeletalMesh::RandomSpecificBone(FVectorVMContext& Con
 	const TArray<int32>& SpecificBonesArray = InstData->SpecificBones;
 
 	int32 Max = SpecificBones.Num() - 1;
-	for (int32 i = 0; i < Context.NumInstances; ++i)
+	if (Max != INDEX_NONE)
+	{
+		for (int32 i = 0; i < Context.NumInstances; ++i)
+		{
+			int32 BoneIndex = Context.RandStream.RandRange(0, Max);
+			*OutBone.GetDestAndAdvance() = SpecificBonesArray[BoneIndex];
+		}
+	}
+	else
 	{
-		int32 BoneIndex = Context.RandStream.RandRange(0, Max);
-		*OutBone.GetDestAndAdvance() = SpecificBonesArray[BoneIndex];
+		FMemory::Memset(OutBone.GetDest(), 0xFF, Context.NumInstances * sizeof(int32));
 	}
 }
 
@@ -419,7 +433,7 @@ void UNiagaraDataInterfaceSkeletalMesh::GetSpecificSocketBoneAt(FVectorVMContext
 	const TArray<int32>& SpecificSocketsArray = InstData->SpecificSocketBones;
 
 	int32 Max = SpecificSockets.Num() - 1;
-	if (Max >= 0)
+	if (Max != INDEX_NONE)
 	{
 		for (int32 i = 0; i < Context.NumInstances; ++i)
 		{
@@ -429,10 +443,7 @@ void UNiagaraDataInterfaceSkeletalMesh::GetSpecificSocketBoneAt(FVectorVMContext
 	}
 	else
 	{
-		for (int32 i = 0; i < Context.NumInstances; ++i)
-		{
-			*OutSocketBone.GetDestAndAdvance() = INDEX_NONE;
-		}
+		FMemory::Memset(OutSocketBone.GetDest(), 0xFF, Context.NumInstances * sizeof(int32));
 	}
 }
 
@@ -446,7 +457,7 @@ void UNiagaraDataInterfaceSkeletalMesh::RandomSpecificSocketBone(FVectorVMContex
 	const TArray<int32>& SpecificSocketsArray = InstData->SpecificSocketBones;
 
 	int32 Max = SpecificSockets.Num() - 1;
-	if (Max >= 0)
+	if (Max != INDEX_NONE)
 	{
 		for (int32 i = 0; i < Context.NumInstances; ++i)
 		{
@@ -456,10 +467,7 @@ void UNiagaraDataInterfaceSkeletalMesh::RandomSpecificSocketBone(FVectorVMContex
 	}
 	else
 	{
-		for (int32 i = 0; i < Context.NumInstances; ++i)
-		{
-			*OutSocketBone.GetDestAndAdvance() = INDEX_NONE;
-		}
+		FMemory::Memset(OutSocketBone.GetDest(), 0xFF, Context.NumInstances * sizeof(int32));
 	}
 }
 
diff --git a/Engine/Plugins/FX/Niagara/Source/Niagara/Private/NDISkeletalMesh_TriangleSampling.cpp b/Engine/Plugins/FX/Niagara/Source/Niagara/Private/NDISkeletalMesh_TriangleSampling.cpp
--- a/Engine/Plugins/FX/Niagara/Source/Niagara/Private/NDISkeletalMesh_TriangleSampling.cpp
+++ b/Engine/Plugins/FX/Niagara/Source/Niagara/Private/NDISkeletalMesh_TriangleSampling.cpp
@@ -673,6 +673,7 @@ void UNiagaraDataInterfaceSkeletalMesh::GetTriCoordUV(FVectorVMContext& Context)
 	FMultiSizeIndexContainer& Indices = LODData.MultiSizeIndexContainer;
 	FRawStaticIndexBuffer16or32Interface* IndexBuffer = Indices.GetIndexBuffer();
 	int32 TriMax = IndexBuffer->Num() - 3;
+	int32 UVSetMax = LODData.StaticVertexBuffers.StaticMeshVertexBuffer.GetNumTexCoords() - 1;
 	float InvDt = 1.0f / InstData->DeltaSeconds;
 	for (int32 i = 0; i < Context.NumInstances; ++i)
 	{
@@ -684,6 +685,7 @@ void UNiagaraDataInterfaceSkeletalMesh::GetTriCoordUV(FVectorVMContext& Context)
 		int32 Idx2 = IndexBuffer->Get(Tri + 2);
 		FVector2D UV0;		FVector2D UV1;		FVector2D UV2;
 		int32 UVSet = UVSetParam.Get();
+		UVSet = FMath::Clamp(UVSet, 0, UVSetMax);
 		UV0 = VertAccessor.GetVertexUV(LODData, Idx0, UVSet);
 		UV1 = VertAccessor.GetVertexUV(LODData, Idx1, UVSet);
 		UV2 = VertAccessor.GetVertexUV(LODData, Idx2, UVSet);
