Some proposed bone socket mesh modules for wyeth with a couple of fixes in there too.

Also duped off Skeletal Mesh Location and tweaked it's interaction with other modules to reduce particle bloat.
Thoughts?

Smooth spawning in single emitter to come....

diff --git a/Engine/Plugins/FX/Niagara/Source/Niagara/Private/NiagaraDataInterfaceSkeletalMesh.cpp b/Engine/Plugins/FX/Niagara/Source/Niagara/Private/NiagaraDataInterfaceSkeletalMesh.cpp
--- a/Engine/Plugins/FX/Niagara/Source/Niagara/Private/NiagaraDataInterfaceSkeletalMesh.cpp
+++ b/Engine/Plugins/FX/Niagara/Source/Niagara/Private/NiagaraDataInterfaceSkeletalMesh.cpp
@@ -14,6 +14,7 @@
 #include "NiagaraStats.h"
 #include "Templates/AlignmentTemplates.h"
 #include "NDISkeletalMeshCommon.h"
+#include "Engine/SkeletalMeshSocket.h"
 
 #define LOCTEXT_NAMESPACE "NiagaraDataInterfaceSkeletalMesh"
 
@@ -610,9 +611,9 @@ bool FNDISkeletalMesh_InstanceData::Init(UNiagaraDataInterfaceSkeletalMesh* Inte
 	for (int32 SocketIdx = 0; SocketIdx < SpecificSockets.Num(); ++SocketIdx)
 	{
 		FName SocketName = Interface->SpecificSockets[SocketIdx];
-		int32 Socket = INDEX_NONE;
-		Mesh->FindSocketAndIndex(SocketName, Socket);
-		if (Socket == INDEX_NONE)
+		int32 SocketIndex = INDEX_NONE;
+		USkeletalMeshSocket* Socket = Mesh->FindSocketAndIndex(SocketName, SocketIndex);
+		if (SocketIndex == INDEX_NONE)
 		{
 			MissingSockets.Add(SocketName);
 			SpecificSockets[SocketIdx] = 0;
@@ -620,8 +621,9 @@ bool FNDISkeletalMesh_InstanceData::Init(UNiagaraDataInterfaceSkeletalMesh* Inte
 		}
 		else
 		{
-			SpecificSockets[SocketIdx] = Socket;
-			SpecificSocketBones[SocketIdx] = RefSkel.FindBoneIndex(SocketName);
+			check(Socket);
+			SpecificSockets[SocketIdx] = SocketIndex;
+			SpecificSocketBones[SocketIdx] = RefSkel.FindBoneIndex(Socket->BoneName);
 		}
 	}
 
