Smooth interpolated spawning now possible from discreet points in a single emitter.
Spawn Rate module exposes a spawn group that is available in particle spawn for use in placing initial position of particles etc.
Other spawn modules will need to be updated to write spawn group if this method proves workable.

diff --git a/Engine/Plugins/FX/Niagara/Source/Niagara/Classes/NiagaraConstants.h b/Engine/Plugins/FX/Niagara/Source/Niagara/Classes/NiagaraConstants.h
--- a/Engine/Plugins/FX/Niagara/Source/Niagara/Classes/NiagaraConstants.h
+++ b/Engine/Plugins/FX/Niagara/Source/Niagara/Classes/NiagaraConstants.h
@@ -59,6 +59,7 @@
 #define SYS_PARAM_EMITTER_SPAWNRATE					INiagaraModule::GetVar_Emitter_SpawnRate()
 #define SYS_PARAM_EMITTER_SPAWN_INTERVAL			INiagaraModule::GetVar_Emitter_SpawnInterval()
 #define SYS_PARAM_EMITTER_INTERP_SPAWN_START_DT		INiagaraModule::GetVar_Emitter_InterpSpawnStartDt()
+#define SYS_PARAM_EMITTER_SPAWN_GROUP		INiagaraModule::GetVar_Emitter_SpawnGroup()
 
 #define SYS_PARAM_PARTICLES_ID						INiagaraModule::GetVar_Particles_ID()
 #define SYS_PARAM_PARTICLES_POSITION				INiagaraModule::GetVar_Particles_Position()
diff --git a/Engine/Plugins/FX/Niagara/Source/Niagara/Classes/NiagaraDataSet.h b/Engine/Plugins/FX/Niagara/Source/Niagara/Classes/NiagaraDataSet.h
--- a/Engine/Plugins/FX/Niagara/Source/Niagara/Classes/NiagaraDataSet.h
+++ b/Engine/Plugins/FX/Niagara/Source/Niagara/Classes/NiagaraDataSet.h
@@ -1379,12 +1379,14 @@ struct FNiagaraDataSetAccessor<FNiagaraSpawnInfo> : public FNiagaraDataSetAccess
 			CountBase = (int32*)DataBuffer->GetComponentPtrInt32(VarLayout->Int32ComponentStart);
 			InterpStartDtBase = (float*)DataBuffer->GetComponentPtrFloat(VarLayout->FloatComponentStart);
 			IntervalDtBase = (float*)DataBuffer->GetComponentPtrFloat(VarLayout->FloatComponentStart + 1);
+			GroupBase = (int32*)DataBuffer->GetComponentPtrInt32(VarLayout->Int32ComponentStart + 1);
 		}
 		else
 		{
 			CountBase = nullptr;
 			InterpStartDtBase = nullptr;
 			IntervalDtBase = nullptr;
+			GroupBase = nullptr;
 		}
 	}
 
@@ -1415,6 +1417,7 @@ struct FNiagaraDataSetAccessor<FNiagaraSpawnInfo> : public FNiagaraDataSetAccess
 		OutValue.Count = CountBase[Index];
 		OutValue.InterpStartDt = InterpStartDtBase[Index];
 		OutValue.IntervalDt = IntervalDtBase[Index];
+		OutValue.SpawnGroup = GroupBase[Index];
 	}
 
 	FORCEINLINE void Set(int32 Index, const FNiagaraSpawnInfo& InValue)
@@ -1422,16 +1425,18 @@ struct FNiagaraDataSetAccessor<FNiagaraSpawnInfo> : public FNiagaraDataSetAccess
 		CountBase[Index] = InValue.Count;
 		InterpStartDtBase[Index] = InValue.InterpStartDt;
 		IntervalDtBase[Index] = InValue.IntervalDt;
+		GroupBase[Index] = InValue.SpawnGroup;
 	}
 
 
-	FORCEINLINE bool BaseIsValid() const { return CountBase != nullptr && InterpStartDtBase != nullptr && IntervalDtBase != nullptr; }
+	FORCEINLINE bool BaseIsValid() const { return CountBase != nullptr && InterpStartDtBase != nullptr && IntervalDtBase != nullptr && GroupBase != nullptr; }
 
 private:
 
 	int32* CountBase;
 	float* InterpStartDtBase;
 	float* IntervalDtBase;
+	int32* GroupBase;
 };
 
 template<>
diff --git a/Engine/Plugins/FX/Niagara/Source/Niagara/Classes/NiagaraEmitterInstance.h b/Engine/Plugins/FX/Niagara/Source/Niagara/Classes/NiagaraEmitterInstance.h
--- a/Engine/Plugins/FX/Niagara/Source/Niagara/Classes/NiagaraEmitterInstance.h
+++ b/Engine/Plugins/FX/Niagara/Source/Niagara/Classes/NiagaraEmitterInstance.h
@@ -137,9 +137,12 @@ private:
 
 	FNiagaraParameterDirectBinding<float> SpawnIntervalBinding;
 	FNiagaraParameterDirectBinding<float> InterpSpawnStartBinding;
+	FNiagaraParameterDirectBinding<int32> SpawnGroupBinding;
 
 	FNiagaraParameterDirectBinding<float> SpawnIntervalBindingGPU;
 	FNiagaraParameterDirectBinding<float> InterpSpawnStartBindingGPU;
+	FNiagaraParameterDirectBinding<int32> SpawnGroupBindingGPU;
+
 	FNiagaraParameterDirectBinding<float> EmitterAgeBindingGPU;
 
 	FNiagaraParameterDirectBinding<float> SpawnEmitterAgeBinding;
diff --git a/Engine/Plugins/FX/Niagara/Source/Niagara/Private/NiagaraConstants.cpp b/Engine/Plugins/FX/Niagara/Source/Niagara/Private/NiagaraConstants.cpp
--- a/Engine/Plugins/FX/Niagara/Source/Niagara/Private/NiagaraConstants.cpp
+++ b/Engine/Plugins/FX/Niagara/Source/Niagara/Private/NiagaraConstants.cpp
@@ -57,6 +57,7 @@ void FNiagaraConstants::Init()
 		SystemParameters.Add(SYS_PARAM_ENGINE_SYSTEM_AGE);
 		SystemParameters.Add(SYS_PARAM_EMITTER_AGE);
 		SystemParameters.Add(SYS_PARAM_EMITTER_LOCALSPACE);
+		SystemParameters.Add(SYS_PARAM_EMITTER_SPAWN_GROUP);
 	}
 
 	if (TranslatorParameters.Num() == 0)
@@ -87,6 +88,7 @@ void FNiagaraConstants::Init()
 		UpdatedSystemParameters.Add(FName(TEXT("Emitter Spawn Rate")), SYS_PARAM_EMITTER_SPAWNRATE);
 		UpdatedSystemParameters.Add(FName(TEXT("Emitter Spawn Interval")), SYS_PARAM_EMITTER_SPAWN_INTERVAL);
 		UpdatedSystemParameters.Add(FName(TEXT("Emitter Interp Spawn Start Dt")), SYS_PARAM_EMITTER_INTERP_SPAWN_START_DT);
+		UpdatedSystemParameters.Add(FName(TEXT("Emitter Spawn Group")), SYS_PARAM_EMITTER_SPAWN_GROUP);
 
 		UpdatedSystemParameters.Add(FName(TEXT("Delta Time")), SYS_PARAM_ENGINE_DELTA_TIME);
 		UpdatedSystemParameters.Add(FName(TEXT("Emitter Age")), SYS_PARAM_EMITTER_AGE);
@@ -105,6 +107,7 @@ void FNiagaraConstants::Init()
 		UpdatedSystemParameters.Add(FName(TEXT("Spawn Rate")), SYS_PARAM_EMITTER_SPAWNRATE);
 		UpdatedSystemParameters.Add(FName(TEXT("Spawn Interval")), SYS_PARAM_EMITTER_SPAWN_INTERVAL);
 		UpdatedSystemParameters.Add(FName(TEXT("Interp Spawn Start Dt")), SYS_PARAM_EMITTER_INTERP_SPAWN_START_DT);
+		UpdatedSystemParameters.Add(FName(TEXT("Spawn Group")), SYS_PARAM_EMITTER_SPAWN_GROUP);
 		UpdatedSystemParameters.Add(FName(TEXT("Inv Delta Time")), SYS_PARAM_ENGINE_INV_DELTA_TIME);
 	}
 
diff --git a/Engine/Plugins/FX/Niagara/Source/Niagara/Private/NiagaraEmitterInstance.cpp b/Engine/Plugins/FX/Niagara/Source/Niagara/Private/NiagaraEmitterInstance.cpp
--- a/Engine/Plugins/FX/Niagara/Source/Niagara/Private/NiagaraEmitterInstance.cpp
+++ b/Engine/Plugins/FX/Niagara/Source/Niagara/Private/NiagaraEmitterInstance.cpp
@@ -255,11 +255,13 @@ void FNiagaraEmitterInstance::Init(int32 InEmitterIdx, FName InSystemInstanceNam
 	//Setup direct bindings for setting parameter values.
 	SpawnIntervalBinding.Init(SpawnExecContext.Parameters, CachedEmitter->ToEmitterParameter(SYS_PARAM_EMITTER_SPAWN_INTERVAL));
 	InterpSpawnStartBinding.Init(SpawnExecContext.Parameters, CachedEmitter->ToEmitterParameter(SYS_PARAM_EMITTER_INTERP_SPAWN_START_DT));
+	SpawnGroupBinding.Init(SpawnExecContext.Parameters, CachedEmitter->ToEmitterParameter(SYS_PARAM_EMITTER_SPAWN_GROUP));
 
 	if (CachedEmitter->SimTarget == ENiagaraSimTarget::GPUComputeSim)
 	{
 		SpawnIntervalBindingGPU.Init(GPUExecContext.CombinedParamStore, CachedEmitter->ToEmitterParameter(SYS_PARAM_EMITTER_SPAWN_INTERVAL));
 		InterpSpawnStartBindingGPU.Init(GPUExecContext.CombinedParamStore, CachedEmitter->ToEmitterParameter(SYS_PARAM_EMITTER_INTERP_SPAWN_START_DT));
+		SpawnGroupBindingGPU.Init(GPUExecContext.CombinedParamStore, CachedEmitter->ToEmitterParameter(SYS_PARAM_EMITTER_SPAWN_GROUP));
 	}
 
 	FNiagaraVariable EmitterAgeParam = CachedEmitter->ToEmitterParameter(SYS_PARAM_EMITTER_AGE);
@@ -731,7 +733,7 @@ bool FNiagaraEmitterInstance::HandleCompletion(bool bForce)
 bool FNiagaraEmitterInstance::RequiredPersistentID()const
 {
 	//TODO: can we have this be enabled at runtime from outside the system?
-	return GetEmitterHandle().GetInstance()->RequiresPersistantIDs();
+	return GetEmitterHandle().GetInstance()->RequiresPersistantIDs() || ParticleDataSet->HasVariable(SYS_PARAM_PARTICLES_ID);
 }
 
 /** 
@@ -956,6 +958,7 @@ void FNiagaraEmitterInstance::Tick(float DeltaSeconds)
 			{
 				SpawnIntervalBindingGPU.SetValue(Info.IntervalDt);
 				InterpSpawnStartBindingGPU.SetValue(Info.InterpStartDt);
+				SpawnGroupBindingGPU.SetValue(Info.SpawnGroup);
 				bOnlySetOnce = true;
 			}
 			else if (Info.Count > 0)
@@ -1145,6 +1148,7 @@ void FNiagaraEmitterInstance::Tick(float DeltaSeconds)
 		{
 			SpawnIntervalBinding.SetValue(Info.IntervalDt);
 			InterpSpawnStartBinding.SetValue(Info.InterpStartDt);
+			SpawnGroupBinding.SetValue(Info.SpawnGroup);
 
 			SpawnParticles(Info.Count, TEXT("Regular Spawn"));
 		};
@@ -1161,6 +1165,7 @@ void FNiagaraEmitterInstance::Tick(float DeltaSeconds)
 				//Event spawns are instantaneous at the middle of the frame?
 				SpawnIntervalBinding.SetValue(0.0f);
 				InterpSpawnStartBinding.SetValue(DeltaSeconds * 0.5f);
+				SpawnGroupBinding.SetValue(0);
 
 				SpawnParticles(EventNumToSpawn, TEXT("Event Spawn"));
 			}
diff --git a/Engine/Plugins/FX/Niagara/Source/Niagara/Private/NiagaraModule.cpp b/Engine/Plugins/FX/Niagara/Source/Niagara/Private/NiagaraModule.cpp
--- a/Engine/Plugins/FX/Niagara/Source/Niagara/Private/NiagaraModule.cpp
+++ b/Engine/Plugins/FX/Niagara/Source/Niagara/Private/NiagaraModule.cpp
@@ -96,6 +96,7 @@ FNiagaraVariable INiagaraModule::Emitter_LocalSpace;
 FNiagaraVariable INiagaraModule::Emitter_SpawnRate;
 FNiagaraVariable INiagaraModule::Emitter_SpawnInterval;
 FNiagaraVariable INiagaraModule::Emitter_InterpSpawnStartDt;
+FNiagaraVariable INiagaraModule::Emitter_SpawnGroup;
 
 FNiagaraVariable INiagaraModule::Particles_ID;
 FNiagaraVariable INiagaraModule::Particles_Position;
@@ -189,6 +190,7 @@ void INiagaraModule::StartupModule()
 	Emitter_SpawnRate = FNiagaraVariable(FNiagaraTypeDefinition::GetFloatDef(), TEXT("Emitter.SpawnRate"));
 	Emitter_SpawnInterval = FNiagaraVariable(FNiagaraTypeDefinition::GetFloatDef(), TEXT("Emitter.SpawnInterval"));
 	Emitter_InterpSpawnStartDt = FNiagaraVariable(FNiagaraTypeDefinition::GetFloatDef(), TEXT("Emitter.InterpSpawnStartDt"));
+	Emitter_SpawnGroup = FNiagaraVariable(FNiagaraTypeDefinition::GetIntDef(), TEXT("Emitter.SpawnGroup"));
 
 	Particles_ID = FNiagaraVariable(FNiagaraTypeDefinition::GetIDDef(), TEXT("Particles.ID"));
 	Particles_Position = FNiagaraVariable(FNiagaraTypeDefinition::GetVec3Def(), TEXT("Particles.Position"));
diff --git a/Engine/Plugins/FX/Niagara/Source/Niagara/Public/NiagaraModule.h b/Engine/Plugins/FX/Niagara/Source/Niagara/Public/NiagaraModule.h
--- a/Engine/Plugins/FX/Niagara/Source/Niagara/Public/NiagaraModule.h
+++ b/Engine/Plugins/FX/Niagara/Source/Niagara/Public/NiagaraModule.h
@@ -153,6 +153,7 @@ public:
 	FORCEINLINE static const FNiagaraVariable&  GetVar_Emitter_SpawnRate() { return Emitter_SpawnRate; }
 	FORCEINLINE static const FNiagaraVariable&  GetVar_Emitter_SpawnInterval() { return Emitter_SpawnInterval; }
 	FORCEINLINE static const FNiagaraVariable&  GetVar_Emitter_InterpSpawnStartDt() { return Emitter_InterpSpawnStartDt; }
+	FORCEINLINE static const FNiagaraVariable&  GetVar_Emitter_SpawnGroup() { return Emitter_SpawnGroup; }
 
 	FORCEINLINE static const FNiagaraVariable&  GetVar_Particles_ID() { return Particles_ID; }
 	FORCEINLINE static const FNiagaraVariable&  GetVar_Particles_Position() { return Particles_Position; }
@@ -239,6 +240,7 @@ private:
 	static FNiagaraVariable Emitter_SpawnRate;
 	static FNiagaraVariable Emitter_SpawnInterval;
 	static FNiagaraVariable Emitter_InterpSpawnStartDt;
+	static FNiagaraVariable Emitter_SpawnGroup;
 
 	static FNiagaraVariable Particles_ID;
 	static FNiagaraVariable Particles_Position;
diff --git a/Engine/Plugins/FX/Niagara/Source/Niagara/Public/NiagaraTypes.h b/Engine/Plugins/FX/Niagara/Source/Niagara/Public/NiagaraTypes.h
--- a/Engine/Plugins/FX/Niagara/Source/Niagara/Public/NiagaraTypes.h
+++ b/Engine/Plugins/FX/Niagara/Source/Niagara/Public/NiagaraTypes.h
@@ -145,6 +145,13 @@ struct FNiagaraSpawnInfo
 	/** The sub frame delta time between each particle. */
 	UPROPERTY(EditAnywhere, Category = SpawnInfo)
 	float IntervalDt;
+	/**
+	 * An integer used to identify this spawn info.
+	 * Typically this is unused.
+	 * An example usage is when using multiple spawn modules to spawn from multiple discreet locations.
+	 */
+	UPROPERTY(EditAnywhere, Category = SpawnInfo)
+	int32 SpawnGroup;
 };
 
 USTRUCT(Blueprintable, meta = (DisplayName = "Niagara ID"))
diff --git a/Engine/Plugins/FX/Niagara/Source/NiagaraCore/Private/NiagaraCustomVersion.cpp b/Engine/Plugins/FX/Niagara/Source/NiagaraCore/Private/NiagaraCustomVersion.cpp
--- a/Engine/Plugins/FX/Niagara/Source/NiagaraCore/Private/NiagaraCustomVersion.cpp
+++ b/Engine/Plugins/FX/Niagara/Source/NiagaraCore/Private/NiagaraCustomVersion.cpp
@@ -8,4 +8,4 @@ const FGuid FNiagaraCustomVersion::GUID(0xFCF57AFA, 0x50764283, 0xB9A9E658, 0xFF
 // Register the custom version with core
 FCustomVersionRegistration GRegisterNiagaraCustomVersion(FNiagaraCustomVersion::GUID, FNiagaraCustomVersion::LatestVersion, TEXT("NiagaraVer"));
 
-const FGuid FNiagaraCustomVersion::LatestScriptCompileVersion(0x6EB4D29C, 0x0CA547DB, 0x96910A18, 0x9CF291BC);
\ No newline at end of file
+const FGuid FNiagaraCustomVersion::LatestScriptCompileVersion(0x6DB2DD7A, 0xBE944FC7, 0xB2409AB0, 0xC13AF894);
\ No newline at end of file
diff --git a/Engine/Plugins/FX/Niagara/Source/NiagaraEditor/Private/NiagaraHlslTranslator.cpp b/Engine/Plugins/FX/Niagara/Source/NiagaraEditor/Private/NiagaraHlslTranslator.cpp
--- a/Engine/Plugins/FX/Niagara/Source/NiagaraEditor/Private/NiagaraHlslTranslator.cpp
+++ b/Engine/Plugins/FX/Niagara/Source/NiagaraEditor/Private/NiagaraHlslTranslator.cpp
@@ -815,6 +815,9 @@ const FNiagaraTranslateResults &FHlslNiagaraTranslator::Translate(const FNiagara
 
 					Var = FNiagaraVariable(FNiagaraTypeDefinition::GetFloatDef(), TEXT("Interpolation.Emitter_InterpSpawnStartDt"));
 					FoundHistory.AddVariable(Var, Var, nullptr);
+
+					Var = FNiagaraVariable(FNiagaraTypeDefinition::GetIntDef(), TEXT("Interpolation.Emitter_SpawnGroup"));
+					FoundHistory.AddVariable(Var, Var, nullptr);
 				}
 
 				ParamMapHistories[ParamMapIdx] = (FoundHistory);
@@ -940,6 +943,7 @@ const FNiagaraTranslateResults &FHlslNiagaraTranslator::Translate(const FNiagara
 			ParameterMapRegisterExternalConstantNamespaceVariable(SYS_PARAM_EMITTER_SPAWNRATE, nullptr, 0, OutputIdx, nullptr);
 			ParameterMapRegisterExternalConstantNamespaceVariable(SYS_PARAM_EMITTER_SPAWN_INTERVAL, nullptr, 0, OutputIdx, nullptr);
 			ParameterMapRegisterExternalConstantNamespaceVariable(SYS_PARAM_EMITTER_INTERP_SPAWN_START_DT, nullptr, 0, OutputIdx, nullptr);
+			ParameterMapRegisterExternalConstantNamespaceVariable(SYS_PARAM_EMITTER_SPAWN_GROUP, nullptr, 0, OutputIdx, nullptr);
 		}
 
 		// Generate the Parameter Map HLSL definitions. We don't add to the final HLSL output here. We just build up the strings and tables
@@ -1360,6 +1364,7 @@ void FHlslNiagaraTranslator::DefineInterpolatedParametersFunction(FString &HlslO
 			HlslOutputString += TEXT("\tContext.") + PrevMap + TEXT(".Interpolation.SpawnInterp = SpawnInterp;\n");
 			HlslOutputString += TEXT("\tContext.") + PrevMap + TEXT(".Interpolation.Emitter_SpawnInterval = Emitter_SpawnInterval;\n");
 			HlslOutputString += TEXT("\tContext.") + PrevMap + TEXT(".Interpolation.Emitter_InterpSpawnStartDt = Emitter_InterpSpawnStartDt;\n");
+			HlslOutputString += TEXT("\tContext.") + PrevMap + TEXT(".Interpolation.Emitter_SpawnGroup = Emitter_SpawnGroup;\n");
 
 			for (int32 UniformIdx = 0; UniformIdx < ChunksByMode[(int32)ENiagaraCodeChunkMode::Uniform].Num(); ++UniformIdx)
 			{
@@ -2805,7 +2810,8 @@ bool FHlslNiagaraTranslator::ShouldInterpolateParameter(const FNiagaraVariable&
 		Parameter == SYS_PARAM_ENGINE_EXEC_COUNT || 
 		Parameter == SYS_PARAM_EMITTER_SPAWNRATE ||
 		Parameter == SYS_PARAM_EMITTER_SPAWN_INTERVAL ||
-		Parameter == SYS_PARAM_EMITTER_INTERP_SPAWN_START_DT )
+		Parameter == SYS_PARAM_EMITTER_INTERP_SPAWN_START_DT ||
+		Parameter == SYS_PARAM_EMITTER_SPAWN_GROUP )
 	{
 		return false;
 	}
