Fixes to aviod some rendering validation checks.

diff --git a/Engine/Plugins/FX/Niagara/Source/NiagaraVertexFactories/Private/NiagaraMeshVertexFactory.cpp b/Engine/Plugins/FX/Niagara/Source/NiagaraVertexFactories/Private/NiagaraMeshVertexFactory.cpp
--- a/Engine/Plugins/FX/Niagara/Source/NiagaraVertexFactories/Private/NiagaraMeshVertexFactory.cpp
+++ b/Engine/Plugins/FX/Niagara/Source/NiagaraVertexFactories/Private/NiagaraMeshVertexFactory.cpp
@@ -66,7 +66,7 @@ public:
 
 		//SetSRVParameter(RHICmdList, VertexShaderRHI, NiagaraParticleDataInt, NiagaraMeshVF->GetIntDataSRV());
 
-		SetSRVParameter(RHICmdList, VertexShaderRHI, SortedIndices, NiagaraMeshVF->GetSortedIndicesSRV());
+		SetSRVParameter(RHICmdList, VertexShaderRHI, SortedIndices, NiagaraMeshVF->GetSortedIndicesSRV() ? NiagaraMeshVF->GetSortedIndicesSRV() : GFNiagaraNullSortedIndicesVertexBuffer.VertexBufferSRV);
 		SetShaderValue(RHICmdList, VertexShaderRHI, SortedIndicesOffset, NiagaraMeshVF->GetSortedIndicesOffset());
 	}
 
diff --git a/Engine/Plugins/FX/Niagara/Source/NiagaraVertexFactories/Private/NiagaraSpriteVertexFactory.cpp b/Engine/Plugins/FX/Niagara/Source/NiagaraVertexFactories/Private/NiagaraSpriteVertexFactory.cpp
--- a/Engine/Plugins/FX/Niagara/Source/NiagaraVertexFactories/Private/NiagaraSpriteVertexFactory.cpp
+++ b/Engine/Plugins/FX/Niagara/Source/NiagaraVertexFactories/Private/NiagaraSpriteVertexFactory.cpp
@@ -69,9 +69,9 @@ public:
 		FloatDataOffset.Bind(ParameterMap, TEXT("NiagaraFloatDataOffset"));
 		FloatDataStride.Bind(ParameterMap, TEXT("NiagaraFloatDataStride"));
 
-// 		NiagaraParticleDataInt.Bind(ParameterMap, TEXT("NiagaraParticleDataInt"));
-// 		FloatDataOffset.Bind(ParameterMap, TEXT("NiagaraInt32DataOffset"));
-// 		FloatDataStride.Bind(ParameterMap, TEXT("NiagaraInt3DataStride"));
+//  		NiagaraParticleDataInt.Bind(ParameterMap, TEXT("NiagaraParticleDataInt"));
+//  		Int32DataOffset.Bind(ParameterMap, TEXT("NiagaraInt32DataOffset"));
+//  		Int32DataStride.Bind(ParameterMap, TEXT("NiagaraInt3DataStride"));
 
 		ParticleAlignmentMode.Bind(ParameterMap, TEXT("ParticleAlignmentMode"));
 		ParticleFacingMode.Bind(ParameterMap, TEXT("ParticleFacingMode"));
@@ -91,9 +91,9 @@ public:
 		Ar << FloatDataOffset;
 		Ar << FloatDataStride;
 
-// 		Ar << NiagaraParticleDataInt;
-// 		Ar << Int32DataOffset;
-// 		Ar << Int32DataStride;
+//  		Ar << NiagaraParticleDataInt;
+//  		Ar << Int32DataOffset;
+//  		Ar << Int32DataStride;
 
 		Ar << SortedIndices;
 		Ar << SortedIndicesOffset;
@@ -116,7 +116,7 @@ public:
 		SetShaderValue(RHICmdList, VertexShaderRHI, FloatDataOffset, SpriteVF->GetFloatDataOffset());
 		SetShaderValue(RHICmdList, VertexShaderRHI, FloatDataStride, SpriteVF->GetFloatDataStride());
 
-		SetSRVParameter(RHICmdList, VertexShaderRHI, SortedIndices, SpriteVF->GetSortedIndicesSRV());
+		SetSRVParameter(RHICmdList, VertexShaderRHI, SortedIndices, SpriteVF->GetSortedIndicesSRV() ? SpriteVF->GetSortedIndicesSRV() : GFNiagaraNullSortedIndicesVertexBuffer.VertexBufferSRV);
 		SetShaderValue(RHICmdList, VertexShaderRHI, SortedIndicesOffset, SpriteVF->GetSortedIndicesOffset());
 	}
 
@@ -132,9 +132,9 @@ private:
 	FShaderParameter FloatDataOffset;
 	FShaderParameter FloatDataStride;
 
-// 	FShaderResourceParameter NiagaraParticleDataInt;
-// 	FShaderParameter Int32DataOffset;
-// 	FShaderParameter Int32DataStride;
+//  	FShaderResourceParameter NiagaraParticleDataInt;
+//  	FShaderParameter Int32DataOffset;
+//  	FShaderParameter Int32DataStride;
 	
 	FShaderResourceParameter SortedIndices;
 	FShaderParameter SortedIndicesOffset;
diff --git a/Engine/Plugins/FX/Niagara/Source/NiagaraVertexFactories/Private/NiagaraVertexFactoriesModule.cpp b/Engine/Plugins/FX/Niagara/Source/NiagaraVertexFactories/Private/NiagaraVertexFactoriesModule.cpp
--- a/Engine/Plugins/FX/Niagara/Source/NiagaraVertexFactories/Private/NiagaraVertexFactoriesModule.cpp
+++ b/Engine/Plugins/FX/Niagara/Source/NiagaraVertexFactories/Private/NiagaraVertexFactoriesModule.cpp
@@ -5,3 +5,6 @@
 #include "NiagaraVertexFactory.h"
 
 IMPLEMENT_MODULE(INiagaraVertexFactoriesModule, NiagaraVertexFactories);
+
+
+TGlobalResource<FNiagaraNullSortedIndicesVertexBuffer> GFNiagaraNullSortedIndicesVertexBuffer;
\ No newline at end of file
diff --git a/Engine/Plugins/FX/Niagara/Source/NiagaraVertexFactories/Public/NiagaraVertexFactory.h b/Engine/Plugins/FX/Niagara/Source/NiagaraVertexFactories/Public/NiagaraVertexFactory.h
--- a/Engine/Plugins/FX/Niagara/Source/NiagaraVertexFactories/Public/NiagaraVertexFactory.h
+++ b/Engine/Plugins/FX/Niagara/Source/NiagaraVertexFactories/Public/NiagaraVertexFactory.h
@@ -15,6 +15,34 @@
 class FMaterial;
 
 
+class FNiagaraNullSortedIndicesVertexBuffer : public FVertexBuffer
+{
+public:
+	/**
+	* Initialize the RHI for this rendering resource
+	*/
+	virtual void InitRHI() override
+	{
+		// create a static vertex buffer
+		FRHIResourceCreateInfo CreateInfo;
+		void* BufferData = nullptr;
+		VertexBufferRHI = RHICreateAndLockVertexBuffer(sizeof(int32), BUF_Static | BUF_ShaderResource, CreateInfo, BufferData);
+		FMemory::Memzero(BufferData, sizeof(int32));
+		RHIUnlockVertexBuffer(VertexBufferRHI);
+
+		VertexBufferSRV = RHICreateShaderResourceView(VertexBufferRHI, sizeof(int32), PF_R32_SINT);
+	}
+
+	virtual void ReleaseRHI() override
+	{
+		VertexBufferSRV.SafeRelease();
+		FVertexBuffer::ReleaseRHI();
+	}
+
+	FShaderResourceViewRHIRef VertexBufferSRV;
+};
+extern TGlobalResource<FNiagaraNullSortedIndicesVertexBuffer> GFNiagaraNullSortedIndicesVertexBuffer;
+
 /**
 * Enum identifying the type of a particle vertex factory.
 */
diff --git a/Engine/Shaders/Private/NiagaraVFParticleAccess.usf b/Engine/Shaders/Private/NiagaraVFParticleAccess.usf
--- a/Engine/Shaders/Private/NiagaraVFParticleAccess.usf
+++ b/Engine/Shaders/Private/NiagaraVFParticleAccess.usf
@@ -2,12 +2,17 @@
 /*********************************************************************************
  *							Niagara Vertex Factory data access
  *********************************************************************************/
+
+//Currently we have no integer data being fed into the renderers.
+//Buffer<int> NiagaraParticleDataInt;
+//uint NiagaraIntDataOffset;
+//uint NiagaraIntDataStride;
+
+
 Buffer<float> NiagaraParticleDataFloat;
-Buffer<int> NiagaraParticleDataInt;
 uint NiagaraFloatDataOffset;
 uint NiagaraFloatDataStride;
 
-
 float GetFloat(int RegisterIdx, uint InstanceID)
 {
 	return NiagaraParticleDataFloat[NiagaraFloatDataOffset + (RegisterIdx * NiagaraFloatDataStride + InstanceID)];
