Niagara - Fix an issue where nested dynamic inputs could be reset when modifying dynamic inputs for the same module.

diff --git a/Engine/Plugins/FX/Niagara/Source/NiagaraEditor/Private/ViewModels/Stack/NiagaraStackGraphUtilities.cpp b/Engine/Plugins/FX/Niagara/Source/NiagaraEditor/Private/ViewModels/Stack/NiagaraStackGraphUtilities.cpp
--- a/Engine/Plugins/FX/Niagara/Source/NiagaraEditor/Private/ViewModels/Stack/NiagaraStackGraphUtilities.cpp
+++ b/Engine/Plugins/FX/Niagara/Source/NiagaraEditor/Private/ViewModels/Stack/NiagaraStackGraphUtilities.cpp
@@ -625,7 +625,8 @@ void FNiagaraStackGraphUtilities::RemoveNodesForStackFunctionInputOverridePin(UE
 					DynamicInputNodeOverrideNode->GetInputPins(InputPins);
 					for (UEdGraphPin* InputPin : InputPins)
 					{
-						if (InputPin->PinName.ToString().StartsWith(DynamicInputNode->GetFunctionName()))
+						FNiagaraParameterHandle InputHandle(InputPin->PinName);
+						if (InputHandle.GetNamespace().ToString() == DynamicInputNode->GetFunctionName())
 						{
 							RemoveNodesForStackFunctionInputOverridePin(*InputPin, OutRemovedDataObjects);
 							DynamicInputNodeOverrideNode->RemovePin(InputPin);
