When compiling on shader types that we don't fully support on the GPU, we end up leaking compile id's. This handles those compile id's as errors.

diff --git a/Engine/Plugins/FX/Niagara/Source/NiagaraShader/Private/NiagaraShader.cpp b/Engine/Plugins/FX/Niagara/Source/NiagaraShader/Private/NiagaraShader.cpp
--- a/Engine/Plugins/FX/Niagara/Source/NiagaraShader/Private/NiagaraShader.cpp
+++ b/Engine/Plugins/FX/Niagara/Source/NiagaraShader/Private/NiagaraShader.cpp
@@ -735,7 +735,6 @@ void FNiagaraShaderMap::Compile(
 			Script->RemoveOutstandingCompileId(CompilingId);
 			// Assign a unique identifier so that shaders from this shader map can be associated with it after a deferred compile
 			CompilingId = NextCompilingId;
-			UE_LOG(LogShaders, Log, TEXT("CompilingId = %p %d"), Script, CompilingId);
 			Script->AddCompileId(CompilingId);
 
 			check(NextCompilingId < UINT_MAX);
@@ -786,6 +785,12 @@ void FNiagaraShaderMap::Compile(
 					}
 					NumShaders++;
 				}
+				else if (ShaderType)
+				{
+					UE_LOG(LogWindows, Display, TEXT("Skipping compilation of %s as it isn't supported on this target type."), *Script->SourceName);
+					Script->RemoveOutstandingCompileId(CompilingId);
+					Script->NotifyCompilationFinished();
+				}
 			}
   
 			if (!CorrespondingScripts)
