Niagara - Finish implementing custom thumbnails for systems and emitters.

diff --git a/Engine/Plugins/FX/Niagara/Source/NiagaraEditor/Private/NiagaraEditorModule.cpp b/Engine/Plugins/FX/Niagara/Source/NiagaraEditor/Private/NiagaraEditorModule.cpp
--- a/Engine/Plugins/FX/Niagara/Source/NiagaraEditor/Private/NiagaraEditorModule.cpp
+++ b/Engine/Plugins/FX/Niagara/Source/NiagaraEditor/Private/NiagaraEditorModule.cpp
@@ -12,6 +12,7 @@
 #include "SequencerChannelInterface.h"
 #include "SequencerSettings.h"
 #include "AssetRegistryModule.h"
+#include "ThumbnailRendering/ThumbnailManager.h"
 
 #include "AssetTypeActions/AssetTypeActions_NiagaraSystem.h"
 #include "AssetTypeActions/AssetTypeActions_NiagaraEmitter.h"
@@ -88,6 +89,7 @@
 #include "Customizations/NiagaraEventScriptPropertiesCustomization.h"
 #include "HAL/IConsoleManager.h"
 #include "NiagaraHlslTranslator.h"
+#include "NiagaraThumbnailRenderer.h"
 
 
 IMPLEMENT_MODULE( FNiagaraEditorModule, NiagaraEditor );
@@ -475,6 +477,9 @@ void FNiagaraEditorModule::StartupModule()
 		TEXT("fx.DumpRapidIterationParametersForAsset"),
 		TEXT("Dumps the values of the rapid iteration parameters for the specified asset by path."),
 		FConsoleCommandWithArgsDelegate::CreateStatic(&DumpRapidIterationParamersForAsset));
+
+	UThumbnailManager::Get().RegisterCustomRenderer(UNiagaraEmitter::StaticClass(), UNiagaraEmitterThumbnailRenderer::StaticClass());
+	UThumbnailManager::Get().RegisterCustomRenderer(UNiagaraSystem::StaticClass(), UNiagaraSystemThumbnailRenderer::StaticClass());
 }
 
 
@@ -553,6 +558,12 @@ void FNiagaraEditorModule::ShutdownModule()
 	{
 		IConsoleManager::Get().UnregisterConsoleObject(DumpRapidIterationParametersForAsset);
 	}
+
+	if (UObjectInitialized())
+	{
+		UThumbnailManager::Get().UnregisterCustomRenderer(UNiagaraEmitter::StaticClass());
+		UThumbnailManager::Get().UnregisterCustomRenderer(UNiagaraSystem::StaticClass());
+	}
 }
 
 FNiagaraEditorModule& FNiagaraEditorModule::Get()
diff --git a/Engine/Plugins/FX/Niagara/Source/NiagaraEditor/Private/NiagaraThumbnailRenderer.cpp b/Engine/Plugins/FX/Niagara/Source/NiagaraEditor/Private/NiagaraThumbnailRenderer.cpp
new file mode 100644
--- /dev/null
+++ b/Engine/Plugins/FX/Niagara/Source/NiagaraEditor/Private/NiagaraThumbnailRenderer.cpp
@@ -0,0 +1,55 @@
+// Copyright 1998-2018 Epic Games, Inc. All Rights Reserved.
+
+#include "NiagaraThumbnailRenderer.h"
+#include "CanvasTypes.h"
+#include "NiagaraEmitter.h"
+#include "NiagaraSystem.h"
+
+bool UNiagaraThumbnailRendererBase::CanVisualizeAsset(UObject* Object)
+{
+	return GetThumbnailTextureFromObject(Object) != nullptr;
+}
+
+void UNiagaraThumbnailRendererBase::GetThumbnailSize(UObject* Object, float Zoom, uint32& OutWidth, uint32& OutHeight) const
+{
+	UTexture2D* ObjectTexture = GetThumbnailTextureFromObject(Object);
+	if (ObjectTexture != nullptr)
+	{
+		OutWidth = Zoom * ObjectTexture->GetSizeX();
+		OutHeight = Zoom * ObjectTexture->GetSizeY();
+	}
+	else
+	{
+		OutWidth = 0;
+		OutHeight = 0;
+	}
+}
+
+void UNiagaraThumbnailRendererBase::Draw(UObject* Object, int32 X, int32 Y, uint32 Width, uint32 Height, FRenderTarget*, FCanvas* Canvas)
+{
+	UTexture2D* ObjectTexture = GetThumbnailTextureFromObject(Object);
+	if (ObjectTexture != nullptr)
+	{
+		Canvas->DrawTile(X, Y, Width, Height, 0.0f, 0.0f, 1.0f, 1.0f, FLinearColor::White, ObjectTexture->Resource, false);
+	}
+}
+
+UTexture2D* UNiagaraEmitterThumbnailRenderer::GetThumbnailTextureFromObject(UObject* Object) const
+{
+	UNiagaraEmitter* Emitter = Cast<UNiagaraEmitter>(Object);
+	if (Emitter != nullptr)
+	{
+		return Emitter->ThumbnailImage;
+	}
+	return nullptr;
+}
+
+UTexture2D* UNiagaraSystemThumbnailRenderer::GetThumbnailTextureFromObject(UObject* Object) const
+{
+	UNiagaraSystem* System = Cast<UNiagaraSystem>(Object);
+	if (System != nullptr)
+	{
+		return System->ThumbnailImage;
+	}
+	return nullptr;
+}
\ No newline at end of file
diff --git a/Engine/Plugins/FX/Niagara/Source/NiagaraEditor/Private/NiagaraThumbnailRenderer.h b/Engine/Plugins/FX/Niagara/Source/NiagaraEditor/Private/NiagaraThumbnailRenderer.h
new file mode 100644
--- /dev/null
+++ b/Engine/Plugins/FX/Niagara/Source/NiagaraEditor/Private/NiagaraThumbnailRenderer.h
@@ -0,0 +1,41 @@
+// Copyright 1998-2018 Epic Games, Inc. All Rights Reserved.
+
+#pragma once
+
+#include "ThumbnailRendering/TextureThumbnailRenderer.h"
+#include "NiagaraThumbnailRenderer.generated.h"
+
+class UTexture2D;
+
+UCLASS()
+class UNiagaraThumbnailRendererBase : public UTextureThumbnailRenderer
+{
+	GENERATED_BODY()
+
+public:
+	// UThumbnailRenderer interface.
+	virtual bool CanVisualizeAsset(UObject* Object) override;
+	virtual void GetThumbnailSize(UObject* Object, float Zoom, uint32& OutWidth, uint32& OutHeight) const override;
+	virtual void Draw(UObject* Object, int32 X, int32 Y, uint32 Width, uint32 Height, FRenderTarget*, FCanvas* Canvas) override;
+
+protected:
+	virtual UTexture2D* GetThumbnailTextureFromObject(UObject* Object) const { return nullptr; }
+};
+
+UCLASS()
+class UNiagaraEmitterThumbnailRenderer : public UNiagaraThumbnailRendererBase
+{
+	GENERATED_BODY()
+
+protected:
+	virtual UTexture2D* GetThumbnailTextureFromObject(UObject* Object) const override;
+};
+
+UCLASS()
+class UNiagaraSystemThumbnailRenderer : public UNiagaraThumbnailRendererBase
+{
+	GENERATED_BODY()
+
+protected:
+	virtual UTexture2D* GetThumbnailTextureFromObject(UObject* Object) const override;
+};
\ No newline at end of file
diff --git a/Engine/Plugins/FX/Niagara/Source/NiagaraEditor/Private/Toolkits/NiagaraSystemToolkit.cpp b/Engine/Plugins/FX/Niagara/Source/NiagaraEditor/Private/Toolkits/NiagaraSystemToolkit.cpp
--- a/Engine/Plugins/FX/Niagara/Source/NiagaraEditor/Private/Toolkits/NiagaraSystemToolkit.cpp
+++ b/Engine/Plugins/FX/Niagara/Source/NiagaraEditor/Private/Toolkits/NiagaraSystemToolkit.cpp
@@ -247,6 +247,7 @@ void FNiagaraSystemToolkit::InitializeWithEmitter(const EToolkitMode::Type Mode,
 	// We set this to the copy's change id here instead of the original emitter's change id because the copy's change id may have been
 	// updated from the original as part of post load and we use this id to detect if the editable emitter has been changed.
 	LastSyncedEmitterChangeId = EditableEmitter->GetChangeId();
+	bEmitterThumbnailUpdated = false;
 
 	FNiagaraSystemViewModelOptions SystemOptions;
 	SystemOptions.bCanModifyEmittersFromTimeline = false;
@@ -387,7 +388,8 @@ TSharedRef<SDockTab> FNiagaraSystemToolkit::SpawnTab_Viewport(const FSpawnTabArg
 {
 	check(Args.GetTabId().TabType == ViewportTabID);
 
-	Viewport = SNew(SNiagaraSystemViewport);
+	Viewport = SNew(SNiagaraSystemViewport)
+		.OnThumbnailCaptured(this, &FNiagaraSystemToolkit::OnThumbnailCaptured);
 
 	TSharedRef<SDockTab> SpawnedTab =
 		SNew(SDockTab)
@@ -674,9 +676,25 @@ void FNiagaraSystemToolkit::SetupCommands()
 
 void FNiagaraSystemToolkit::OnSaveThumbnailImage()
 {
-	if (Viewport.IsValid() && Viewport->GetViewportClient().IsValid())
+	if (Viewport.IsValid())
 	{
-		Viewport->CreateThumbnail();
+		Viewport->CreateThumbnail(SystemToolkitMode == ESystemToolkitMode::System ? (UObject*)System : Emitter);
+	}
+}
+
+void FNiagaraSystemToolkit::OnThumbnailCaptured(UTexture2D* Thumbnail)
+{
+	if (SystemToolkitMode == ESystemToolkitMode::System)
+	{
+		System->MarkPackageDirty();
+		System->ThumbnailImage = Thumbnail;
+	}
+	else if (SystemToolkitMode == ESystemToolkitMode::Emitter) 
+	{
+		TSharedPtr<FNiagaraEmitterViewModel> EditableEmitterViewModel = SystemViewModel->GetEmitterHandleViewModels()[0]->GetEmitterViewModel();
+		UNiagaraEmitter* EditableEmitter = EditableEmitterViewModel->GetEmitter();
+		EditableEmitter->ThumbnailImage = Thumbnail;
+		bEmitterThumbnailUpdated = true;
 	}
 }
 
@@ -951,50 +969,59 @@ void FNiagaraSystemToolkit::UpdateOriginalEmitter()
 {
 	checkf(SystemToolkitMode == ESystemToolkitMode::Emitter, TEXT("There is no original emitter to update in system mode."));
 
-	const FScopedBusyCursor BusyCursor;
-	const FText LocalizedScriptEditorApply = NSLOCTEXT("UnrealEd", "ToolTip_NiagaraEmitterEditorApply", "Apply changes to original emitter and its use in the world.");
-	GWarn->BeginSlowTask(LocalizedScriptEditorApply, true);
-	GWarn->StatusUpdate(1, 1, LocalizedScriptEditorApply);
+	TSharedPtr<FNiagaraEmitterViewModel> EditableEmitterViewModel = SystemViewModel->GetEmitterHandleViewModels()[0]->GetEmitterViewModel();
+	UNiagaraEmitter* EditableEmitter = EditableEmitterViewModel->GetEmitter();
 
-	if (Emitter->IsSelected())
+	if (EditableEmitter->GetChangeId() != LastSyncedEmitterChangeId)
 	{
-		GEditor->GetSelectedObjects()->Deselect(Emitter);
-	}
+		const FScopedBusyCursor BusyCursor;
+		const FText LocalizedScriptEditorApply = NSLOCTEXT("UnrealEd", "ToolTip_NiagaraEmitterEditorApply", "Apply changes to original emitter and its use in the world.");
+		GWarn->BeginSlowTask(LocalizedScriptEditorApply, true);
+		GWarn->StatusUpdate(1, 1, LocalizedScriptEditorApply);
 
-	ResetLoaders(Emitter->GetOutermost()); // Make sure that we're not going to get invalid version number linkers into the package we are going into. 
-	Emitter->GetOutermost()->LinkerCustomVersion.Empty();
+		if (Emitter->IsSelected())
+		{
+			GEditor->GetSelectedObjects()->Deselect(Emitter);
+		}
 
-	TSharedPtr<FNiagaraEmitterViewModel> EditableEmitterViewModel = SystemViewModel->GetEmitterHandleViewModels()[0]->GetEmitterViewModel();
-	UNiagaraEmitter* EditableEmitter = EditableEmitterViewModel->GetEmitter();
+		ResetLoaders(Emitter->GetOutermost()); // Make sure that we're not going to get invalid version number linkers into the package we are going into. 
+		Emitter->GetOutermost()->LinkerCustomVersion.Empty();
 
-	TArray<UNiagaraScript*> AllScripts;
-	EditableEmitter->GetScripts(AllScripts, true);
-	for (UNiagaraScript* Script : AllScripts)
-	{
-		checkSlow(Script->AreScriptAndSourceSynchronized());
-	}
-	checkSlow(EditableEmitter->AreAllScriptAndSourcesSynchronized());
+		TArray<UNiagaraScript*> AllScripts;
+		EditableEmitter->GetScripts(AllScripts, true);
+		for (UNiagaraScript* Script : AllScripts)
+		{
+			checkSlow(Script->AreScriptAndSourceSynchronized());
+		}
+		checkSlow(EditableEmitter->AreAllScriptAndSourcesSynchronized());
 
-	// overwrite the original script in place by constructing a new one with the same name
-	Emitter = (UNiagaraEmitter*)StaticDuplicateObject(EditableEmitter, Emitter->GetOuter(), Emitter->GetFName(),
-		RF_AllFlags,
-		Emitter->GetClass());
+		// overwrite the original script in place by constructing a new one with the same name
+		Emitter = (UNiagaraEmitter*)StaticDuplicateObject(EditableEmitter, Emitter->GetOuter(),
+			Emitter->GetFName(), RF_AllFlags, Emitter->GetClass());
 
-	// Record the last synced change id to detect future changes.
-	LastSyncedEmitterChangeId = EditableEmitter->GetChangeId();
+		// Record the last synced change id to detect future changes.
+		LastSyncedEmitterChangeId = EditableEmitter->GetChangeId();
+		bEmitterThumbnailUpdated = false;
 
-	checkSlow (UNiagaraEmitter::GetForceCompileOnLoad() || Emitter->GetChangeId() == EditableEmitter->GetChangeId());
+		checkSlow(UNiagaraEmitter::GetForceCompileOnLoad() || Emitter->GetChangeId() == EditableEmitter->GetChangeId());
 
-	// Restore RF_Standalone on the original emitter, as it had been removed from the preview emitter so that it could be GC'd.
-	Emitter->SetFlags(RF_Standalone);
+		// Restore RF_Standalone on the original emitter, as it had been removed from the preview emitter so that it could be GC'd.
+		Emitter->SetFlags(RF_Standalone);
 
-	TArray<UNiagaraEmitter*> AffectedEmitters;
-	AffectedEmitters.Add(Emitter);
-	UpdateExistingEmitters();
+		TArray<UNiagaraEmitter*> AffectedEmitters;
+		AffectedEmitters.Add(Emitter);
+		UpdateExistingEmitters();
 
-	checkSlow(UNiagaraEmitter::GetForceCompileOnLoad() || Emitter->GetChangeId() == EditableEmitter->GetChangeId());
+		checkSlow(UNiagaraEmitter::GetForceCompileOnLoad() || Emitter->GetChangeId() == EditableEmitter->GetChangeId());
 
-	GWarn->EndSlowTask();
+		GWarn->EndSlowTask();
+	}
+	else if(bEmitterThumbnailUpdated)
+	{
+		Emitter->MarkPackageDirty();
+		Emitter->ThumbnailImage = (UTexture2D*)StaticDuplicateObject(EditableEmitter->ThumbnailImage, Emitter);
+		bEmitterThumbnailUpdated = false;
+	}
 }
 
 void FNiagaraSystemToolkit::UpdateExistingEmitters()
@@ -1092,7 +1119,7 @@ bool FNiagaraSystemToolkit::OnRequestClose()
 	if (SystemToolkitMode == ESystemToolkitMode::Emitter)
 	{
 		TSharedPtr<FNiagaraEmitterViewModel> EmitterViewModel = SystemViewModel->GetEmitterHandleViewModels()[0]->GetEmitterViewModel();
-		if (bChangesDiscarded == false && EmitterViewModel->GetEmitter()->GetChangeId() != LastSyncedEmitterChangeId)
+		if (bChangesDiscarded == false && (EmitterViewModel->GetEmitter()->GetChangeId() != LastSyncedEmitterChangeId || bEmitterThumbnailUpdated))
 		{
 			// find out the user wants to do with this dirty NiagaraScript
 			EAppReturnType::Type YesNoCancelReply = FMessageDialog::Open(EAppMsgType::YesNoCancel,
@@ -1150,7 +1177,7 @@ bool FNiagaraSystemToolkit::OnApplyEnabled() const
 	if (Emitter != nullptr)
 	{
 		TSharedPtr<FNiagaraEmitterViewModel> EmitterViewModel = SystemViewModel->GetEmitterHandleViewModels()[0]->GetEmitterViewModel();
-		return EmitterViewModel->GetEmitter()->GetChangeId() != LastSyncedEmitterChangeId;
+		return EmitterViewModel->GetEmitter()->GetChangeId() != LastSyncedEmitterChangeId || bEmitterThumbnailUpdated;
 	}
 	return false;
 }
diff --git a/Engine/Plugins/FX/Niagara/Source/NiagaraEditor/Private/Toolkits/NiagaraSystemToolkit.h b/Engine/Plugins/FX/Niagara/Source/NiagaraEditor/Private/Toolkits/NiagaraSystemToolkit.h
--- a/Engine/Plugins/FX/Niagara/Source/NiagaraEditor/Private/Toolkits/NiagaraSystemToolkit.h
+++ b/Engine/Plugins/FX/Niagara/Source/NiagaraEditor/Private/Toolkits/NiagaraSystemToolkit.h
@@ -123,6 +123,9 @@ private:
 private:
 	TSharedRef<SWidget> GenerateBoundsMenuContent(TSharedRef<FUICommandList> InCommandList);
 	void OnSaveThumbnailImage();
+	void OnThumbnailCaptured(UTexture2D* Thumbnail);
+
+private:
 
 	/** The System being edited in system mode, or the placeholder system being edited in emitter mode. */
 	UNiagaraSystem* System;
@@ -133,6 +136,10 @@ private:
 	/** The value of the emitter change id from the last time it was in sync with the original emitter. */
 	FGuid LastSyncedEmitterChangeId;
 
+	/** Whether or not the emitter thumbnail has been updated.  The is needed because after the first update the
+		screenshot uobject is reused, so a pointer comparison doesn't work to checking if the images has been updated. */
+	bool bEmitterThumbnailUpdated;
+
 	ESystemToolkitMode SystemToolkitMode;
 
 	TSharedPtr<SNiagaraSystemViewport> Viewport;
diff --git a/Engine/Plugins/FX/Niagara/Source/NiagaraEditor/Private/Widgets/SNiagaraSystemViewport.cpp b/Engine/Plugins/FX/Niagara/Source/NiagaraEditor/Private/Widgets/SNiagaraSystemViewport.cpp
--- a/Engine/Plugins/FX/Niagara/Source/NiagaraEditor/Private/Widgets/SNiagaraSystemViewport.cpp
+++ b/Engine/Plugins/FX/Niagara/Source/NiagaraEditor/Private/Widgets/SNiagaraSystemViewport.cpp
@@ -24,7 +24,10 @@
 class FNiagaraSystemViewportClient : public FEditorViewportClient
 {
 public:
-	FNiagaraSystemViewportClient(FAdvancedPreviewScene& InPreviewScene, const TSharedRef<SNiagaraSystemViewport>& InNiagaraEditorViewport);
+	DECLARE_DELEGATE_OneParam(FOnScreenShotCaptured, UTexture2D*);
+
+public:
+	FNiagaraSystemViewportClient(FAdvancedPreviewScene& InPreviewScene, const TSharedRef<SNiagaraSystemViewport>& InNiagaraEditorViewport, FOnScreenShotCaptured InOnScreenShotCaptured);
 	
 	// FEditorViewportClient interface
 	virtual FLinearColor GetBackgroundColor() const override;
@@ -41,10 +44,14 @@ public:
 
 	TSharedPtr<SNiagaraSystemViewport> NiagaraViewport;
 	bool bCaptureScreenShot;
+	TWeakObjectPtr<UObject> ScreenShotOwner;
+
+	FOnScreenShotCaptured OnScreenShotCaptured;
 };
 
-FNiagaraSystemViewportClient::FNiagaraSystemViewportClient(FAdvancedPreviewScene& InPreviewScene, const TSharedRef<SNiagaraSystemViewport>& InNiagaraEditorViewport)
+FNiagaraSystemViewportClient::FNiagaraSystemViewportClient(FAdvancedPreviewScene& InPreviewScene, const TSharedRef<SNiagaraSystemViewport>& InNiagaraEditorViewport, FOnScreenShotCaptured InOnScreenShotCaptured)
 	: FEditorViewportClient(nullptr, &InPreviewScene, StaticCastSharedRef<SEditorViewport>(InNiagaraEditorViewport))
+	, OnScreenShotCaptured(InOnScreenShotCaptured)
 {
 	NiagaraViewport = InNiagaraEditorViewport;
 
@@ -101,7 +108,7 @@ void FNiagaraSystemViewportClient::Draw(FViewport* InViewport,FCanvas* Canvas)
 
 	FEditorViewportClient::Draw(InViewport, Canvas);
 
-	if (bCaptureScreenShot && ParticleSystem != nullptr)
+	if (bCaptureScreenShot && ScreenShotOwner.IsValid() && OnScreenShotCaptured.IsBound())
 	{
 		
 		int32 SrcWidth = InViewport->GetSizeXY().X;
@@ -121,13 +128,14 @@ void FNiagaraSystemViewportClient::Draw(FViewport* InViewport,FCanvas* Canvas)
 			// Compress.
 			FCreateTexture2DParameters Params;
 			Params.bDeferCompression = true;
-			ParticleSystem->ThumbnailImage = FImageUtils::CreateTexture2D(ScaledWidth, ScaledHeight, ScaledBitmap, ParticleSystem, TEXT("ThumbnailTexture"), RF_NoFlags, Params);
+			
+			UTexture2D* ThumbnailImage = FImageUtils::CreateTexture2D(ScaledWidth, ScaledHeight, ScaledBitmap, ScreenShotOwner.Get(), TEXT("ThumbnailTexture"), RF_NoFlags, Params);
 
-			ParticleSystem->ThumbnailImageOutOfDate = false;
-			ParticleSystem->MarkPackageDirty();
+			OnScreenShotCaptured.Execute(ThumbnailImage);
 		}
 
 		bCaptureScreenShot = false;
+		ScreenShotOwner.Reset();
 	}
 }
 
@@ -182,6 +190,7 @@ void SNiagaraSystemViewport::Construct(const FArguments& InArgs)
 	PreviewComponent = nullptr;
 	AdvancedPreviewScene = MakeShareable(new FAdvancedPreviewScene(FPreviewScene::ConstructionValues()));
 	AdvancedPreviewScene->SetFloorVisibility(false);
+	OnThumbnailCaptured = InArgs._OnThumbnailCaptured;
 
 	SEditorViewport::Construct( SEditorViewport::FArguments() );
 }
@@ -194,11 +203,12 @@ SNiagaraSystemViewport::~SNiagaraSystemViewport()
 	}
 }
 
-void SNiagaraSystemViewport::CreateThumbnail()
+void SNiagaraSystemViewport::CreateThumbnail(UObject* InScreenShotOwner)
 {
 	if (SystemViewportClient.IsValid() && PreviewComponent != nullptr)
 	{
 		SystemViewportClient->bCaptureScreenShot = true;
+		SystemViewportClient->ScreenShotOwner = InScreenShotOwner;
 	}
 }
 
@@ -282,6 +292,11 @@ bool SNiagaraSystemViewport::IsVisible() const
 	return ViewportWidget.IsValid() && (!ParentTab.IsValid() || ParentTab.Pin()->IsForeground()) && SEditorViewport::IsVisible() ;
 }
 
+void SNiagaraSystemViewport::OnScreenShotCaptured(UTexture2D* ScreenShot)
+{
+	OnThumbnailCaptured.ExecuteIfBound(ScreenShot);
+}
+
 void SNiagaraSystemViewport::BindCommands()
 {
 	SEditorViewport::BindCommands();
@@ -348,7 +363,8 @@ bool SNiagaraSystemViewport::IsTogglePreviewBackgroundChecked() const
 
 TSharedRef<FEditorViewportClient> SNiagaraSystemViewport::MakeEditorViewportClient() 
 {
-	SystemViewportClient = MakeShareable( new FNiagaraSystemViewportClient(*AdvancedPreviewScene.Get(), SharedThis(this)) );
+	SystemViewportClient = MakeShareable( new FNiagaraSystemViewportClient(*AdvancedPreviewScene.Get(), SharedThis(this), 
+		FNiagaraSystemViewportClient::FOnScreenShotCaptured::CreateSP(this, &SNiagaraSystemViewport::OnScreenShotCaptured) ) );
 	
 	SystemViewportClient->SetViewLocation( FVector::ZeroVector );
 	SystemViewportClient->SetViewRotation( FRotator::ZeroRotator );
diff --git a/Engine/Plugins/FX/Niagara/Source/NiagaraEditor/Private/Widgets/SNiagaraSystemViewport.h b/Engine/Plugins/FX/Niagara/Source/NiagaraEditor/Private/Widgets/SNiagaraSystemViewport.h
--- a/Engine/Plugins/FX/Niagara/Source/NiagaraEditor/Private/Widgets/SNiagaraSystemViewport.h
+++ b/Engine/Plugins/FX/Niagara/Source/NiagaraEditor/Private/Widgets/SNiagaraSystemViewport.h
@@ -22,9 +22,12 @@ class FNiagaraSystemInstance;
  */
 class SNiagaraSystemViewport : public SEditorViewport, public FGCObject, public ICommonEditorViewportToolbarInfoProvider
 {
+public:
+	DECLARE_DELEGATE_OneParam(FOnThumbnailCaptured, UTexture2D*);
+
 public:
 	SLATE_BEGIN_ARGS( SNiagaraSystemViewport ){}
-	//SLATE_ARGUMENT(TWeakPtr<INiagaraSystemEditor>, SystemEditor)
+		SLATE_EVENT(FOnThumbnailCaptured, OnThumbnailCaptured)
 	SLATE_END_ARGS()
 	
 	void Construct(const FArguments& InArgs);
@@ -75,7 +78,7 @@ public:
 
 	bool GetDrawElement(EDrawElements Element) const;
 	void ToggleDrawElement(EDrawElements Element);
-	void CreateThumbnail();
+	void CreateThumbnail(UObject* InScreenShotOwner);
 
 
 	bool IsToggleOrbitChecked() const;
@@ -91,6 +94,11 @@ protected:
 	virtual void PopulateViewportOverlays(TSharedRef<class SOverlay> Overlay) override;
 	EVisibility OnGetViewportCompileTextVisibility() const;
 
+private:
+	bool IsVisible() const override;
+
+	void OnScreenShotCaptured(UTexture2D* ScreenShot);
+
 private:
 	/** The parent tab where this viewport resides */
 	TWeakPtr<SDockTab> ParentTab;
@@ -100,8 +108,6 @@ private:
 
 	TSharedPtr<STextBlock> CompileText;
 	
-	bool IsVisible() const override;
-	
 	/** Pointer back to the material editor tool that owns us */
 	//TWeakPtr<INiagaraSystemEditor> SystemEditorPtr;
 	
@@ -111,4 +117,6 @@ private:
 	TSharedPtr<class FNiagaraSystemViewportClient> SystemViewportClient;
 
 	uint32 DrawFlags;
+
+	FOnThumbnailCaptured OnThumbnailCaptured;
 };
