Some fixes to enable Niagara more platforms.
There is still instability from the implementation of FRHIGPUMemoryReadback.
Niagara param buffer now use the UniformBuffer_SingleDraw flag.

diff --git a/Engine/Plugins/FX/Niagara/Source/Niagara/Classes/NiagaraScriptExecutionContext.h b/Engine/Plugins/FX/Niagara/Source/Niagara/Classes/NiagaraScriptExecutionContext.h
--- a/Engine/Plugins/FX/Niagara/Source/Niagara/Classes/NiagaraScriptExecutionContext.h
+++ b/Engine/Plugins/FX/Niagara/Source/Niagara/Classes/NiagaraScriptExecutionContext.h
@@ -92,6 +92,7 @@ struct FNiagaraComputeExecutionContext
 		, RTUpdateScript(0)
 		, RTSpawnScript(0)
 		, RTGPUScript(0)
+		, CBufferLayout(TEXT("Niagara Compute Sim CBuffer"))
 		, GPUDataReadback(nullptr)
 		, AccumulatedSpawnRate(0)
 		, NumIndicesPerInstance(0)
@@ -185,7 +186,8 @@ struct FNiagaraComputeExecutionContext
 	UNiagaraScript* GPUScript;
 	class FNiagaraShaderScript*  RTUpdateScript;
 	class FNiagaraShaderScript*  RTSpawnScript;
-	class FNiagaraShaderScript*  RTGPUScript;
+	class FNiagaraShaderScript*  RTGPUScript;	
+	FRHIUniformBufferLayout CBufferLayout; // Persistent layouts used to create Compute Sim CBuffer
 	TArray<uint8, TAlignedHeapAllocator<16>> ParamData_RT;		// RT side copy of the parameter data
 	FNiagaraScriptExecutionParameterStore CombinedParamStore;
 	static uint32 TickCounter;
diff --git a/Engine/Plugins/FX/Niagara/Source/Niagara/Private/NiagaraComponent.cpp b/Engine/Plugins/FX/Niagara/Source/Niagara/Private/NiagaraComponent.cpp
--- a/Engine/Plugins/FX/Niagara/Source/Niagara/Private/NiagaraComponent.cpp
+++ b/Engine/Plugins/FX/Niagara/Source/Niagara/Private/NiagaraComponent.cpp
@@ -492,13 +492,6 @@ void UNiagaraComponent::Activate(bool bReset /* = false */)
 		return;
 	}
 
-	if (IsSwitchPlatform(GMaxRHIShaderPlatform))
-	{
-		UE_LOG(LogNiagara, Warning, TEXT("Failed to activate Niagara component as Niagara is not yet supported on this platform: %s"), *LegacyShaderPlatformToShaderFormat(GMaxRHIShaderPlatform).ToString());
-		OnSystemComplete();
-		return;
-	}
-
 	SCOPE_CYCLE_COUNTER(STAT_NiagaraComponentActivate);
 	if (Asset == nullptr)
 	{
diff --git a/Engine/Plugins/FX/Niagara/Source/Niagara/Private/NiagaraEmitterInstance.cpp b/Engine/Plugins/FX/Niagara/Source/Niagara/Private/NiagaraEmitterInstance.cpp
--- a/Engine/Plugins/FX/Niagara/Source/Niagara/Private/NiagaraEmitterInstance.cpp
+++ b/Engine/Plugins/FX/Niagara/Source/Niagara/Private/NiagaraEmitterInstance.cpp
@@ -991,9 +991,10 @@ void FNiagaraEmitterInstance::Tick(float DeltaSeconds)
 
 		int32 ParmSize = GPUExecContext.CombinedParamStore.GetPaddedParameterSizeInBytes();
 		
-
 		GPUExecContext.ParamData_RT.SetNumZeroed(ParmSize);
 		GPUExecContext.CombinedParamStore.CopyParameterDataToPaddedBuffer(GPUExecContext.ParamData_RT.GetData(), ParmSize);
+		// Because each context is only ran once each frame, the CBuffer layout stays constant for the lifetime duration of the CBuffer (one frame).
+		GPUExecContext.CBufferLayout.ConstantBufferSize = ParmSize;
 
 		// push event data sets to the context
 		for (FNiagaraDataSet *Set : UpdateScriptEventDataSets)
diff --git a/Engine/Plugins/FX/Niagara/Source/Niagara/Private/NiagaraEmitterInstanceBatcher.cpp b/Engine/Plugins/FX/Niagara/Source/Niagara/Private/NiagaraEmitterInstanceBatcher.cpp
--- a/Engine/Plugins/FX/Niagara/Source/Niagara/Private/NiagaraEmitterInstanceBatcher.cpp
+++ b/Engine/Plugins/FX/Niagara/Source/Niagara/Private/NiagaraEmitterInstanceBatcher.cpp
@@ -249,6 +249,7 @@ void NiagaraEmitterInstanceBatcher::Run(const FNiagaraComputeExecutionContext *C
 
 	FNiagaraDataSet *DataSet = Context->MainDataSet;
 	const FNiagaraParameterStore& ParameterStore = Context->CombinedParamStore;
+	const FRHIUniformBufferLayout& CBufferLayout = Context->CBufferLayout;
 	const TArray<uint8, TAlignedHeapAllocator<16>> &Params = Context->ParamData_RT;
 	const FRWBuffer &WriteIndexBuffer = DataSet->GetCurDataSetIndices();
 	FRWBuffer &ReadIndexBuffer = DataSet->GetPrevDataSetIndices();
@@ -309,13 +310,11 @@ void NiagaraEmitterInstanceBatcher::Run(const FNiagaraComputeExecutionContext *C
 	}
 
 	// setup script parameters
-	FRHIUniformBufferLayout CBufferLayout(TEXT("Niagara Compute Sim CBuffer"));
-	CBufferLayout.ConstantBufferSize = Params.Num();
 	if (CBufferLayout.ConstantBufferSize)
 	{
 		check(CBufferLayout.Resources.Num() == 0);
 		const uint8* ParamData = Params.GetData();
-		FUniformBufferRHIRef CBuffer = RHICreateUniformBuffer(ParamData, CBufferLayout, EUniformBufferUsage::UniformBuffer_MultiFrame);
+		FUniformBufferRHIRef CBuffer = RHICreateUniformBuffer(ParamData, CBufferLayout, EUniformBufferUsage::UniformBuffer_SingleDraw);
 		RHICmdList.SetShaderUniformBuffer(Shader->GetComputeShader(), Shader->EmitterConstantBufferParam.GetBaseIndex(), CBuffer);
 	}
 
diff --git a/Engine/Plugins/FX/Niagara/Source/NiagaraShader/Private/NiagaraShared.cpp b/Engine/Plugins/FX/Niagara/Source/NiagaraShader/Private/NiagaraShared.cpp
--- a/Engine/Plugins/FX/Niagara/Source/NiagaraShader/Private/NiagaraShared.cpp
+++ b/Engine/Plugins/FX/Niagara/Source/NiagaraShader/Private/NiagaraShared.cpp
@@ -373,7 +373,7 @@ void FNiagaraShaderScript::SetDataInterfaceParamInfo(TArray<FNiagaraDataInterfac
 
 NIAGARASHADER_API  FNiagaraShader* FNiagaraShaderScript::GetShader() const
 {
-	check(! IsInGameThread() );
+	check(!GIsThreadedRendering || !IsInGameThread());
 	if (!GIsEditor || RenderingThreadShaderMap /*&& RenderingThreadShaderMap->IsComplete(this, true)*/)
 	{
 		return RenderingThreadShaderMap->GetShader<FNiagaraShader>();
diff --git a/Engine/Plugins/FX/Niagara/Source/NiagaraVertexFactories/Private/NiagaraMeshVertexFactory.cpp b/Engine/Plugins/FX/Niagara/Source/NiagaraVertexFactories/Private/NiagaraMeshVertexFactory.cpp
--- a/Engine/Plugins/FX/Niagara/Source/NiagaraVertexFactories/Private/NiagaraMeshVertexFactory.cpp
+++ b/Engine/Plugins/FX/Niagara/Source/NiagaraVertexFactories/Private/NiagaraMeshVertexFactory.cpp
@@ -193,7 +193,7 @@ FShaderResourceViewRHIParamRef FNiagaraMeshVertexFactory::GetPreviousTransformBu
 
 bool FNiagaraMeshVertexFactory::ShouldCompilePermutation(EShaderPlatform Platform, const class FMaterial* Material, const class FShaderType* ShaderType)
 {
-	return (!IsMobilePlatform(Platform) && !IsSwitchPlatform(Platform) && Platform != SP_OPENGL_SM4 && (Material->IsUsedWithNiagaraMeshParticles() || Material->IsSpecialEngineMaterial()));
+	return (!IsMobilePlatform(Platform) && Platform != SP_OPENGL_SM4 && (Material->IsUsedWithNiagaraMeshParticles() || Material->IsSpecialEngineMaterial()));
 }
 
 void FNiagaraMeshVertexFactory::SetData(const FStaticMeshDataType& InData)
diff --git a/Engine/Plugins/FX/Niagara/Source/NiagaraVertexFactories/Private/NiagaraRibbonVertexFactory.cpp b/Engine/Plugins/FX/Niagara/Source/NiagaraVertexFactories/Private/NiagaraRibbonVertexFactory.cpp
--- a/Engine/Plugins/FX/Niagara/Source/NiagaraVertexFactories/Private/NiagaraRibbonVertexFactory.cpp
+++ b/Engine/Plugins/FX/Niagara/Source/NiagaraVertexFactories/Private/NiagaraRibbonVertexFactory.cpp
@@ -154,7 +154,7 @@ static TGlobalResource<FNiagaraRibbonVertexDeclaration> GNiagaraRibbonVertexDecl
 
 bool FNiagaraRibbonVertexFactory::ShouldCompilePermutation(EShaderPlatform Platform, const class FMaterial* Material, const class FShaderType* ShaderType)
 {
-	return (!IsMobilePlatform(Platform) && !IsSwitchPlatform(Platform) && Platform != SP_OPENGL_SM4 && (Material->IsUsedWithNiagaraRibbons() || Material->IsSpecialEngineMaterial()));
+	return (!IsMobilePlatform(Platform) && Platform != SP_OPENGL_SM4 && (Material->IsUsedWithNiagaraRibbons() || Material->IsSpecialEngineMaterial()));
 }
 
 /**
diff --git a/Engine/Plugins/FX/Niagara/Source/NiagaraVertexFactories/Private/NiagaraSpriteVertexFactory.cpp b/Engine/Plugins/FX/Niagara/Source/NiagaraVertexFactories/Private/NiagaraSpriteVertexFactory.cpp
--- a/Engine/Plugins/FX/Niagara/Source/NiagaraVertexFactories/Private/NiagaraSpriteVertexFactory.cpp
+++ b/Engine/Plugins/FX/Niagara/Source/NiagaraVertexFactories/Private/NiagaraSpriteVertexFactory.cpp
@@ -225,7 +225,7 @@ inline TGlobalResource<FNiagaraSpriteVertexDeclaration>& GetNiagaraSpriteVertexD
 
 bool FNiagaraSpriteVertexFactory::ShouldCompilePermutation(EShaderPlatform Platform, const class FMaterial* Material, const class FShaderType* ShaderType)
 {
-	return (!IsSwitchPlatform(Platform) && Platform != SP_OPENGL_SM4 && (Material->IsUsedWithNiagaraSprites() || Material->IsSpecialEngineMaterial()));
+	return (Platform != SP_OPENGL_SM4 && (Material->IsUsedWithNiagaraSprites() || Material->IsSpecialEngineMaterial()));
 }
 
 /**
