Small issue with ribbon rendering code that triggesrs some debug gfx validation checks.

diff --git a/Engine/Plugins/FX/Niagara/Source/Niagara/Private/NiagaraRendererRibbons.cpp b/Engine/Plugins/FX/Niagara/Source/Niagara/Private/NiagaraRendererRibbons.cpp
--- a/Engine/Plugins/FX/Niagara/Source/Niagara/Private/NiagaraRendererRibbons.cpp
+++ b/Engine/Plugins/FX/Niagara/Source/Niagara/Private/NiagaraRendererRibbons.cpp
@@ -201,7 +201,7 @@ void NiagaraRendererRibbons::GetDynamicMeshElements(const TArray<const FSceneVie
 				// TODO: need to make these two a global alloc buffer as well, not recreate
 				// pass in the sorted indices so the VS can fetch the particle data in order
 				FReadBuffer SortedIndicesBuffer;
-				SortedIndicesBuffer.Initialize(sizeof(int32), DynamicDataRibbon->SortedIndices.Num(), EPixelFormat::PF_R32_FLOAT, BUF_Volatile);
+				SortedIndicesBuffer.Initialize(sizeof(int32), DynamicDataRibbon->SortedIndices.Num(), EPixelFormat::PF_R32_SINT, BUF_Volatile);
 				void *IndexPtr = RHILockVertexBuffer(SortedIndicesBuffer.Buffer, 0, DynamicDataRibbon->SortedIndices.Num() * sizeof(int32), RLM_WriteOnly);
 				FMemory::Memcpy(IndexPtr, DynamicDataRibbon->SortedIndices.GetData(), DynamicDataRibbon->SortedIndices.Num() * sizeof(int32));
 				RHIUnlockVertexBuffer(SortedIndicesBuffer.Buffer);
