Pause feature for simulations.
Functions also exposed to BP.

diff --git a/Engine/Plugins/FX/Niagara/Source/Niagara/Private/NiagaraComponent.cpp b/Engine/Plugins/FX/Niagara/Source/Niagara/Private/NiagaraComponent.cpp
--- a/Engine/Plugins/FX/Niagara/Source/Niagara/Private/NiagaraComponent.cpp
+++ b/Engine/Plugins/FX/Niagara/Source/Niagara/Private/NiagaraComponent.cpp
@@ -467,6 +467,23 @@ void UNiagaraComponent::AdvanceSimulationByTime(float SimulateTime, float TickDe
 	}
 }
 
+void UNiagaraComponent::SetPaused(bool bInPaused)
+{
+	if (SystemInstance.IsValid())
+	{
+		SystemInstance->SetPaused(bInPaused);
+	}
+}
+
+bool UNiagaraComponent::IsPaused()const
+{
+	if (SystemInstance.IsValid())
+	{
+		return SystemInstance->IsPaused();
+	}
+	return false;
+}
+
 bool UNiagaraComponent::InitializeSystem()
 {
 	if (SystemInstance.IsValid() == false)
diff --git a/Engine/Plugins/FX/Niagara/Source/Niagara/Private/NiagaraSystemInstance.cpp b/Engine/Plugins/FX/Niagara/Source/Niagara/Private/NiagaraSystemInstance.cpp
--- a/Engine/Plugins/FX/Niagara/Source/Niagara/Private/NiagaraSystemInstance.cpp
+++ b/Engine/Plugins/FX/Niagara/Source/Niagara/Private/NiagaraSystemInstance.cpp
@@ -46,6 +46,7 @@ FNiagaraSystemInstance::FNiagaraSystemInstance(UNiagaraComponent* InComponent)
 	, bForceSolo(false)
 	, bPendingSpawn(false)
 	, bHasTickingEmitters(true)
+	, bPaused(false)
 	, RequestedExecutionState(ENiagaraExecutionState::Complete)
 	, ActualExecutionState(ENiagaraExecutionState::Complete)
 {
@@ -274,7 +275,6 @@ void FNiagaraSystemInstance::Activate(EResetMode InResetMode)
 
 void FNiagaraSystemInstance::Deactivate(bool bImmediate)
 {
-
 	SCOPE_CYCLE_COUNTER(STAT_NiagaraSystemDeactivate);
 	if (IsComplete())
 	{
@@ -335,6 +335,27 @@ void FNiagaraSystemInstance::Complete()
 	}
 }
 
+void FNiagaraSystemInstance::SetPaused(bool bInPaused)
+{
+	if (bInPaused == bPaused)
+	{
+		return;
+	}
+	
+	FNiagaraSystemSimulation* SystemSim = GetSystemSimulation().Get();
+	check(SystemSim);
+	if (bInPaused)
+	{
+		SystemSim->PauseInstance(this);
+	}
+	else
+	{
+		SystemSim->UnpauseInstance(this);
+	}
+
+	bPaused = bInPaused;
+}
+
 void FNiagaraSystemInstance::Reset(FNiagaraSystemInstance::EResetMode Mode, bool bBindParams)
 {
 	SCOPE_CYCLE_COUNTER(STAT_NiagaraSystemReset);
@@ -352,6 +373,8 @@ void FNiagaraSystemInstance::Reset(FNiagaraSystemInstance::EResetMode Mode, bool
 
 	Component->LastRenderTime = Component->GetWorld()->GetTimeSeconds();
 
+	SetPaused(false);
+
 	if (SystemSim)
 	{
 		SystemSim->RemoveInstance(this);
diff --git a/Engine/Plugins/FX/Niagara/Source/Niagara/Private/NiagaraSystemSimulation.cpp b/Engine/Plugins/FX/Niagara/Source/Niagara/Private/NiagaraSystemSimulation.cpp
--- a/Engine/Plugins/FX/Niagara/Source/Niagara/Private/NiagaraSystemSimulation.cpp
+++ b/Engine/Plugins/FX/Niagara/Source/Niagara/Private/NiagaraSystemSimulation.cpp
@@ -94,6 +94,11 @@ bool FNiagaraSystemSimulation::Init(UNiagaraSystem* InSystem, UWorld* InWorld, F
 		DataSet.AddVariables(System->GetSystemUpdateScript()->GetVMExecutableData().Attributes);
 		DataSet.Finalize();
 
+		PausedInstanceData.Init(FNiagaraDataSetID(), ENiagaraSimTarget::CPUSim);
+		PausedInstanceData.AddVariables(System->GetSystemSpawnScript()->GetVMExecutableData().Attributes);
+		PausedInstanceData.AddVariables(System->GetSystemUpdateScript()->GetVMExecutableData().Attributes);
+		PausedInstanceData.Finalize();
+
 		{
 			SpawnInstanceParameterDataSet.Init(FNiagaraDataSetID(), ENiagaraSimTarget::CPUSim);
 			FNiagaraParameters* EngineParamsSpawn = System->GetSystemSpawnScript()->GetVMExecutableData().DataSetToParameters.Find(TEXT("Engine"));
@@ -279,8 +284,8 @@ bool FNiagaraSystemSimulation::Tick(float DeltaSeconds)
 	SCOPE_CYCLE_COUNTER(STAT_NiagaraSystemSim);
 
 	int32 OrigNum = SystemInstances.Num();
-	int32 SpawnNum = PendingSystemInstances.Num();
-	int32 NewNum = OrigNum + SpawnNum;
+	int32 SpawnNum = 0;
+	int32 NewNum = 0;
 
 	{
 		SCOPE_CYCLE_COUNTER(STAT_NiagaraSystemSim_PreSimulate);
@@ -306,9 +311,30 @@ bool FNiagaraSystemSimulation::Tick(float DeltaSeconds)
 			//Pre tick and gather any still valid pending instances for spawn.
 			SystemInstances.Reserve(NewNum);
 			SpawnNum = 0;
-			while (PendingSystemInstances.Num())
+			int32 NumPending = PendingSystemInstances.Num();
+			for (int32 i = NumPending - 1; i >= 0; --i)
 			{
-				FNiagaraSystemInstance* Inst = PendingSystemInstances.Pop();
+				FNiagaraSystemInstance* Inst = PendingSystemInstances[i];
+
+				//Don't spawn systems that are paused. Keep them in pending list so they are spawned when unpaused.
+				if (Inst->IsPaused())
+				{
+					if (GbDumpSystemData || (System && System->bDumpDebugSystemInfo))
+					{
+						UE_LOG(LogNiagara, Log, TEXT("=== Skipping Paused Pending Spawn %d ==="), Inst->SystemInstanceIndex);
+						//DataSet.Dump(true, Instance->SystemInstanceIndex, 1);
+					}
+					continue;
+				}
+
+				check(Inst->SystemInstanceIndex == i);
+				PendingSystemInstances.RemoveAt(i);
+
+				if (GbDumpSystemData || (System && System->bDumpDebugSystemInfo))
+				{
+					UE_LOG(LogNiagara, Log, TEXT("=== Spawning %d -> %d ==="), Inst->SystemInstanceIndex, SystemInstances.Num());
+					//DataSet.Dump(true, Instance->SystemInstanceIndex, 1);
+				}
 
 				Inst->TickDataInterfaces(DeltaSeconds, false);
 				Inst->SetPendingSpawn(false);
@@ -654,6 +680,11 @@ void FNiagaraSystemSimulation::RemoveInstance(FNiagaraSystemInstance* Instance)
 	UNiagaraSystem* System = WeakSystem.Get();
 	if (Instance->IsPendingSpawn())
 	{
+		if (GbDumpSystemData || (System && System->bDumpDebugSystemInfo))
+		{
+			UE_LOG(LogNiagara, Log, TEXT("=== Removing Pending Spawn %d ==="), Instance->SystemInstanceIndex);
+			DataSet.Dump(true, Instance->SystemInstanceIndex, 1);
+		}
 		int32 SystemIndex = Instance->SystemInstanceIndex;
 		check(Instance == PendingSystemInstances[SystemIndex]);
 		PendingSystemInstances.RemoveAtSwap(SystemIndex);
@@ -664,6 +695,23 @@ void FNiagaraSystemSimulation::RemoveInstance(FNiagaraSystemInstance* Instance)
 			PendingSystemInstances[SystemIndex]->SystemInstanceIndex = SystemIndex;
 		}
 	}
+	else if (Instance->IsPaused())
+	{
+		if (GbDumpSystemData || (System && System->bDumpDebugSystemInfo))
+		{
+			UE_LOG(LogNiagara, Log, TEXT("=== Removing Paused %d ==="), Instance->SystemInstanceIndex);
+			DataSet.Dump(true, Instance->SystemInstanceIndex, 1);
+		}
+
+		int32 SystemIndex = Instance->SystemInstanceIndex;
+		check(Instance == PausedSystemInstances[SystemIndex]);
+		PausedSystemInstances.RemoveAtSwap(SystemIndex);
+		Instance->SystemInstanceIndex = INDEX_NONE;
+		if (PausedSystemInstances.IsValidIndex(SystemIndex))
+		{
+			PausedSystemInstances[SystemIndex]->SystemInstanceIndex = SystemIndex;
+		}
+	}
 	else if (SystemInstances.IsValidIndex(Instance->SystemInstanceIndex))
 	{
 		if (GbDumpSystemData || (System && System->bDumpDebugSystemInfo))
@@ -695,8 +743,98 @@ void FNiagaraSystemSimulation::RemoveInstance(FNiagaraSystemInstance* Instance)
 
 void FNiagaraSystemSimulation::AddInstance(FNiagaraSystemInstance* Instance)
 {
+	check(Instance->SystemInstanceIndex == INDEX_NONE);
 	Instance->SetPendingSpawn(true);
 	Instance->SystemInstanceIndex = PendingSystemInstances.Add(Instance);
+
+	UNiagaraSystem* System = WeakSystem.Get();
+	if (GbDumpSystemData || (System && System->bDumpDebugSystemInfo))
+	{
+		UE_LOG(LogNiagara, Log, TEXT("=== Adding To Pending Spawn %d ==="), Instance->SystemInstanceIndex);
+		//DataSet.Dump(true, Instance->SystemInstanceIndex, 1);
+	}
+}
+
+void FNiagaraSystemSimulation::PauseInstance(FNiagaraSystemInstance* Instance)
+{
+	check(!Instance->IsPaused());
+
+	UNiagaraSystem* System = WeakSystem.Get();
+	if (Instance->IsPendingSpawn())
+	{
+		if (GbDumpSystemData || (System && System->bDumpDebugSystemInfo))
+		{
+			UE_LOG(LogNiagara, Log, TEXT("=== Pausing Pending Spawn %d ==="), Instance->SystemInstanceIndex);
+			//DataSet.Dump(true, Instance->SystemInstanceIndex, 1);
+		}
+		//Nothing to do for pending spawn systems.
+		check(PendingSystemInstances[Instance->SystemInstanceIndex] == Instance);
+		return;
+	}
+
+	if (GbDumpSystemData || (System && System->bDumpDebugSystemInfo))
+	{
+		UE_LOG(LogNiagara, Log, TEXT("=== Pausing System %d ==="), Instance->SystemInstanceIndex);
+		DataSet.Dump(true, Instance->SystemInstanceIndex, 1);
+	}
+
+	int32 SystemIndex = Instance->SystemInstanceIndex;
+	check(SystemInstances.IsValidIndex(SystemIndex));
+	check(Instance == SystemInstances[SystemIndex]);
+
+	int32 NewSystemIndex = PausedInstanceData.TransferInstance(DataSet, SystemIndex);
+	DataSet.KillInstance(SystemIndex);
+
+	check(PausedSystemInstances.Num() == NewSystemIndex);
+	Instance->SystemInstanceIndex = NewSystemIndex;
+	PausedSystemInstances.Add(Instance);
+
+	SystemInstances.RemoveAtSwap(SystemIndex);
+	if (SystemInstances.IsValidIndex(SystemIndex))
+	{
+		SystemInstances[SystemIndex]->SystemInstanceIndex = SystemIndex;
+	}
+}
+
+void FNiagaraSystemSimulation::UnpauseInstance(FNiagaraSystemInstance* Instance)
+{
+	check(Instance->IsPaused());
+
+	UNiagaraSystem* System = WeakSystem.Get();
+	if (Instance->IsPendingSpawn())
+	{
+		if (GbDumpSystemData || (System && System->bDumpDebugSystemInfo))
+		{
+			UE_LOG(LogNiagara, Log, TEXT("=== Unpausing Pending Spawn %d ==="), Instance->SystemInstanceIndex);
+			//DataSet.Dump(true, Instance->SystemInstanceIndex, 1);
+		}
+		//Nothing to do for pending spawn systems.
+		check(PendingSystemInstances[Instance->SystemInstanceIndex] == Instance);
+		return;
+	}
+
+	if (GbDumpSystemData || (System && System->bDumpDebugSystemInfo))
+	{
+		UE_LOG(LogNiagara, Log, TEXT("=== Unpausing System %d ==="), Instance->SystemInstanceIndex);
+		DataSet.Dump(true, Instance->SystemInstanceIndex, 1);
+	}
+
+	int32 SystemIndex = Instance->SystemInstanceIndex;
+	check(PausedSystemInstances.IsValidIndex(SystemIndex));
+	check(Instance == PausedSystemInstances[SystemIndex]);
+
+	int32 NewSystemIndex = DataSet.TransferInstance(PausedInstanceData, SystemIndex);
+	PausedInstanceData.KillInstance(SystemIndex);
+
+	check(SystemInstances.Num() == NewSystemIndex);
+	Instance->SystemInstanceIndex = NewSystemIndex;
+	SystemInstances.Add(Instance);
+
+	PausedSystemInstances.RemoveAtSwap(SystemIndex);
+	if (PausedSystemInstances.IsValidIndex(SystemIndex))
+	{
+		PausedSystemInstances[SystemIndex]->SystemInstanceIndex = SystemIndex;
+	}
 }
 
 void FNiagaraSystemSimulation::InitBindings(FNiagaraSystemInstance* SystemInst)
diff --git a/Engine/Plugins/FX/Niagara/Source/Niagara/Public/NiagaraComponent.h b/Engine/Plugins/FX/Niagara/Source/Niagara/Public/NiagaraComponent.h
--- a/Engine/Plugins/FX/Niagara/Source/Niagara/Public/NiagaraComponent.h
+++ b/Engine/Plugins/FX/Niagara/Source/Niagara/Public/NiagaraComponent.h
@@ -266,6 +266,12 @@ public:
 	UFUNCTION(BlueprintCallable, Category = Niagara)
 	void AdvanceSimulationByTime(float SimulateTime, float TickDeltaSeconds);
 
+	UFUNCTION(BlueprintCallable, Category = Niagara)
+	void SetPaused(bool bInPaused);
+
+	UFUNCTION(BlueprintCallable, Category = Niagara)
+	bool IsPaused()const;
+
 	//~ Begin UObject Interface.
 	virtual void PostLoad();
 #if WITH_EDITOR
diff --git a/Engine/Plugins/FX/Niagara/Source/Niagara/Public/NiagaraSystemInstance.h b/Engine/Plugins/FX/Niagara/Source/Niagara/Public/NiagaraSystemInstance.h
--- a/Engine/Plugins/FX/Niagara/Source/Niagara/Public/NiagaraSystemInstance.h
+++ b/Engine/Plugins/FX/Niagara/Source/Niagara/Public/NiagaraSystemInstance.h
@@ -59,6 +59,9 @@ public:
 	void Deactivate(bool bImmediate = false);
 	void Complete();
 
+	void SetPaused(bool bInPaused);
+	FORCEINLINE bool IsPaused()const { return bPaused; }
+
 	void SetSolo(bool bInSolo);
 
 	//void RebindParameterCollection(UNiagaraParameterCollectionInstance* OldInstance, UNiagaraParameterCollectionInstance* NewInstance);
@@ -295,6 +298,9 @@ private:
 	/** If this instance has any currently ticking emitters. If false, allows us to skip some work. */
 	uint32 bHasTickingEmitters : 1;
 
+	/** If this system is paused. When paused it will not tick and never complete etc. */
+	uint32 bPaused : 1;
+
 	/* Execution state requested by external code/BPs calling Activate/Deactivate. */
 	ENiagaraExecutionState RequestedExecutionState;
 
diff --git a/Engine/Plugins/FX/Niagara/Source/Niagara/Public/NiagaraSystemSimulation.h b/Engine/Plugins/FX/Niagara/Source/Niagara/Public/NiagaraSystemSimulation.h
--- a/Engine/Plugins/FX/Niagara/Source/Niagara/Public/NiagaraSystemSimulation.h
+++ b/Engine/Plugins/FX/Niagara/Source/Niagara/Public/NiagaraSystemSimulation.h
@@ -128,6 +128,9 @@ public:
 	void RemoveInstance(FNiagaraSystemInstance* Instance);
 	void AddInstance(FNiagaraSystemInstance* Instance);
 
+	void PauseInstance(FNiagaraSystemInstance* Instance);
+	void UnpauseInstance(FNiagaraSystemInstance* Instance);
+
 	FORCEINLINE UNiagaraSystem* GetSystem()const { return WeakSystem.Get(); }
 
 	UNiagaraParameterCollectionInstance* GetParameterCollectionInstance(UNiagaraParameterCollection* Collection);
@@ -197,6 +200,10 @@ protected:
 	/** System instances that are pending to be spawned. */
 	TArray<FNiagaraSystemInstance*> PendingSystemInstances;
 
+	/** System instances that are paused. */
+	TArray<FNiagaraSystemInstance*> PausedSystemInstances;
+	FNiagaraDataSet PausedInstanceData;
+
 	TArray<TArray<FNiagaraDataSetAccessor<FNiagaraSpawnInfo>>> EmitterSpawnInfoAccessors;
 
 	void InitBindings(FNiagaraSystemInstance* SystemInst);
