Texture sampling added for GPU only. CPU just reports a pink texture.

#tests autotests for texture pass

diff --git a/Engine/Plugins/FX/Niagara/Source/Niagara/Classes/NiagaraDataInterfaceTexture.h b/Engine/Plugins/FX/Niagara/Source/Niagara/Classes/NiagaraDataInterfaceTexture.h
--- a/Engine/Plugins/FX/Niagara/Source/Niagara/Classes/NiagaraDataInterfaceTexture.h
+++ b/Engine/Plugins/FX/Niagara/Source/Niagara/Classes/NiagaraDataInterfaceTexture.h
@@ -50,14 +50,8 @@ public:
 	static const FString SamplerName;
 	static const FString DimensionsBaseName;
 protected:
-#if WITH_EDITOR
-	bool CopyTextureToCPUBackup(UTexture* SourceTexture, TArray<uint8>& TargetBuffer);
-	void CopyTextureData(const uint8* Source, uint8* Dest, uint32 SizeX, uint32 SizeY, uint32 BytesPerPixel, uint32 SourceStride, uint32 DestStride);
-
-#endif
 	virtual bool CopyToInternal(UNiagaraDataInterface* Destination) const override;
 
-	TArray<uint8> CPUTextureData;
 
 	static const FName SampleTexture2DName;
 	static const FName SampleVolumeTextureName;
diff --git a/Engine/Plugins/FX/Niagara/Source/Niagara/Private/NiagaraDataInterfaceTexture.cpp b/Engine/Plugins/FX/Niagara/Source/Niagara/Private/NiagaraDataInterfaceTexture.cpp
--- a/Engine/Plugins/FX/Niagara/Source/Niagara/Private/NiagaraDataInterfaceTexture.cpp
+++ b/Engine/Plugins/FX/Niagara/Source/Niagara/Private/NiagaraDataInterfaceTexture.cpp
@@ -43,7 +43,6 @@ void UNiagaraDataInterfaceTexture::PostLoad()
 		if (Texture != nullptr)
 		{
 			Texture->ConditionalPostLoad();
-			CopyTextureToCPUBackup(Texture, CPUTextureData);
 		}
 	}
 #endif
@@ -54,60 +53,9 @@ void UNiagaraDataInterfaceTexture::PostLoad()
 void UNiagaraDataInterfaceTexture::PostEditChangeProperty(struct FPropertyChangedEvent& PropertyChangedEvent)
 {
 	Super::PostEditChangeProperty(PropertyChangedEvent);
-
-	if (PropertyChangedEvent.Property && PropertyChangedEvent.Property->GetFName() == GET_MEMBER_NAME_CHECKED(UNiagaraDataInterfaceTexture, Texture))
-	{
-		CopyTextureToCPUBackup(Texture, CPUTextureData);
-	}
 }
 
-bool UNiagaraDataInterfaceTexture::CopyTextureToCPUBackup(UTexture* SourceTexture, TArray<uint8>& TargetBuffer)
-{
-	if (SourceTexture == nullptr)
-	{
-		TargetBuffer.Empty();
-		return true;
-	}
-
-	FTextureSource& SourceData = SourceTexture->Source;
-	FIntPoint SourceSize = FIntPoint(SourceData.GetSizeX(), SourceData.GetSizeY());
-
-	{
-		const int32 BytesPerPixel = 4;
-		TargetBuffer.Empty();
-		TargetBuffer.AddZeroed(SourceSize.X * SourceSize.Y * BytesPerPixel);
-	}
 
-	
-	if (SourceData.GetFormat() == TSF_BGRA8)
-	{
-		uint32 BytesPerPixel = SourceData.GetBytesPerPixel();
-		uint8* OffsetSource = SourceData.LockMip(0);
-		uint8* OffsetDest = TargetBuffer.GetData();
-		CopyTextureData(OffsetSource, OffsetDest, SourceSize.X, SourceSize.Y, BytesPerPixel, SourceData.GetSizeX() * BytesPerPixel, SourceSize.X * BytesPerPixel);
-		SourceData.UnlockMip(0);
-	}
-	else 
-	{
-		UE_LOG(LogNiagara, Warning, TEXT("Texture %s is not BGRA8, which means this texture cannot be used with the CPU VM."), *SourceTexture->GetName());
-	}
-
-	return true;
-}
-
-void UNiagaraDataInterfaceTexture::CopyTextureData(const uint8* Source, uint8* Dest, uint32 SizeX, uint32 SizeY, uint32 BytesPerPixel, uint32 SourceStride, uint32 DestStride)
-{
-	const uint32 NumBytesPerRow = SizeX * BytesPerPixel;
-
-	for (uint32 Y = 0; Y < SizeY; ++Y)
-	{
-		FMemory::Memcpy(
-			Dest + (DestStride * Y),
-			Source + (SourceStride * Y),
-			NumBytesPerRow
-		);
-	}
-}
 
 #endif
 
@@ -116,7 +64,8 @@ void UNiagaraDataInterfaceTexture::Serialize(FArchive& Ar)
 	Super::Serialize(Ar);
 	if (Ar.IsLoading() == false || Ar.CustomVer(FNiagaraCustomVersion::GUID) >= FNiagaraCustomVersion::TextureDataInterfaceUsesCustomSerialize)
 	{
-		Ar << CPUTextureData;
+		TArray<uint8> StreamData;
+		Ar << StreamData;
 	}
 	Ar.UsingCustomVersion(FNiagaraCustomVersion::GUID);
 }
@@ -129,7 +78,6 @@ bool UNiagaraDataInterfaceTexture::CopyToInternal(UNiagaraDataInterface* Destina
 	}
 	UNiagaraDataInterfaceTexture* DestinationTexture = CastChecked<UNiagaraDataInterfaceTexture>(Destination);
 	DestinationTexture->Texture = Texture;
-	DestinationTexture->CPUTextureData = CPUTextureData;
 
 	return true;
 }
@@ -279,50 +227,14 @@ void UNiagaraDataInterfaceTexture::SampleTexture(FVectorVMContext& Context)
 	FRegisterHandler<float> OutSampleB(Context);
 	FRegisterHandler<float> OutSampleA(Context);
 
-	if (CPUTextureData.GetAllocatedSize() == 0 || Texture == nullptr)
+	for (int32 i = 0; i < Context.NumInstances; ++i)
 	{
-		for (int32 i = 0; i < Context.NumInstances; ++i)
-		{
-			float X = XParam.GetAndAdvance();
-			float Y = YParam.GetAndAdvance();
-			*OutSampleR.GetDestAndAdvance() = 1.0;
-			*OutSampleG.GetDestAndAdvance() = 0.0;
-			*OutSampleB.GetDestAndAdvance() = 1.0;
-			*OutSampleA.GetDestAndAdvance() = 1.0;
-		}
-	}
-	else
-	{
-		const int32 BytesPerPixel = 4;
-		int32 IntSizeX = Texture->GetSurfaceWidth();
-		int32 IntSizeY = Texture->GetSurfaceHeight();
-		float SizeX = (float)Texture->GetSurfaceWidth() - 1.0f;
-		float SizeY = (float)Texture->GetSurfaceHeight() - 1.0f;
-
-		for (int32 i = 0; i < Context.NumInstances; ++i)
-		{
-			float ParamXVal = XParam.GetAndAdvance();
-			float ParamYVal = YParam.GetAndAdvance();
-			float X = (fmodf(ParamXVal * SizeX, SizeX));
-			float Y = (fmodf(ParamYVal * SizeY, SizeY));
-
-			float XNorm = X < 0.0f ? (SizeX - fabsf(X)) : X;
-			float YNorm = Y < 0.0f ? (SizeY - fabsf(Y)) : Y;
-
-			int32 XInt = floorf(XNorm);
-			int32 YInt = floorf(YNorm);
-			int32 SampleIdx = YInt * IntSizeX * BytesPerPixel + XInt * BytesPerPixel;
-			ensure(CPUTextureData.Num() > SampleIdx);
-			uint8 B0 = CPUTextureData[SampleIdx + 0];
-			uint8 G0 = CPUTextureData[SampleIdx + 1];
-			uint8 R0 = CPUTextureData[SampleIdx + 2];
-			uint8 A0 = CPUTextureData[SampleIdx + 3];
-
-			*OutSampleR.GetDestAndAdvance() = ((float)R0) / 255.0f;
-			*OutSampleG.GetDestAndAdvance() = ((float)G0) / 255.0f;
-			*OutSampleB.GetDestAndAdvance() = ((float)B0) / 255.0f;
-			*OutSampleA.GetDestAndAdvance() = ((float)A0) / 255.0f;
-		}
+		float X = XParam.GetAndAdvance();
+		float Y = YParam.GetAndAdvance();
+		*OutSampleR.GetDestAndAdvance() = 1.0;
+		*OutSampleG.GetDestAndAdvance() = 0.0;
+		*OutSampleB.GetDestAndAdvance() = 1.0;
+		*OutSampleA.GetDestAndAdvance() = 1.0;
 	}
 
 }
