Prevent drawing of particles if the RHI layer doesn't support SRVs

#tests auto-tests still pass

diff --git a/Engine/Plugins/FX/Niagara/Source/Niagara/Private/NiagaraRendererMeshes.cpp b/Engine/Plugins/FX/Niagara/Source/Niagara/Private/NiagaraRendererMeshes.cpp
--- a/Engine/Plugins/FX/Niagara/Source/Niagara/Private/NiagaraRendererMeshes.cpp
+++ b/Engine/Plugins/FX/Niagara/Source/Niagara/Private/NiagaraRendererMeshes.cpp
@@ -123,7 +123,9 @@ void NiagaraRendererMeshes::GetDynamicMeshElements(const TArray<const FSceneView
 	if (!DynamicDataMesh 
 		|| DynamicDataMesh->RTParticleData.GetNumInstancesAllocated() == 0
 		|| DynamicDataMesh->RTParticleData.GetNumInstances() == 0
-		|| nullptr == Properties)
+		|| nullptr == Properties
+		|| !GSupportsResourceView  // Current shader requires SRV to draw properly in all cases.
+		)
 	{
 		return;
 	}
diff --git a/Engine/Plugins/FX/Niagara/Source/Niagara/Private/NiagaraRendererRibbons.cpp b/Engine/Plugins/FX/Niagara/Source/Niagara/Private/NiagaraRendererRibbons.cpp
--- a/Engine/Plugins/FX/Niagara/Source/Niagara/Private/NiagaraRendererRibbons.cpp
+++ b/Engine/Plugins/FX/Niagara/Source/Niagara/Private/NiagaraRendererRibbons.cpp
@@ -94,6 +94,7 @@ void NiagaraRendererRibbons::GetDynamicMeshElements(const TArray<const FSceneVie
 	if (!DynamicDataRibbon 
 		|| DynamicDataRibbon->IndexData.Num() == 0
 		|| nullptr == Properties
+		|| !GSupportsResourceView // Current shader requires SRV to draw properly in all cases.
 		)
 	{
 		return;
diff --git a/Engine/Plugins/FX/Niagara/Source/Niagara/Private/NiagaraRendererSprites.cpp b/Engine/Plugins/FX/Niagara/Source/Niagara/Private/NiagaraRendererSprites.cpp
--- a/Engine/Plugins/FX/Niagara/Source/Niagara/Private/NiagaraRendererSprites.cpp
+++ b/Engine/Plugins/FX/Niagara/Source/Niagara/Private/NiagaraRendererSprites.cpp
@@ -90,6 +90,7 @@ void NiagaraRendererSprites::GetDynamicMeshElements(const TArray<const FSceneVie
 		|| DynamicDataSprites->RTParticleData.GetNumInstancesAllocated() == 0
 		|| DynamicDataSprites->RTParticleData.GetNumInstances() == 0
 		|| nullptr == Properties
+		|| !GSupportsResourceView // Current shader requires SRV to draw properly in all cases.
 		)
 	{
 		return;
