Merging //UE4/Dev-Main to Dev-Build (//UE4/Dev-Build)

#rnx

diff --git a/Engine/Plugins/FX/HoudiniNiagara/Source/HoudiniNiagara/Private/NiagaraDataInterfaceHoudiniCSV.cpp b/Engine/Plugins/FX/HoudiniNiagara/Source/HoudiniNiagara/Private/NiagaraDataInterfaceHoudiniCSV.cpp
--- a/Engine/Plugins/FX/HoudiniNiagara/Source/HoudiniNiagara/Private/NiagaraDataInterfaceHoudiniCSV.cpp
+++ b/Engine/Plugins/FX/HoudiniNiagara/Source/HoudiniNiagara/Private/NiagaraDataInterfaceHoudiniCSV.cpp
@@ -2065,7 +2065,7 @@ struct FNiagaraDataInterfaceParametersCS_HoudiniCSV : public FNiagaraDataInterfa
 		Ar << LastSpawnTime;
 	}
 
-	virtual void Set( FRHICommandList& RHICmdList, FNiagaraShader* Shader, class UNiagaraDataInterface* DataInterface ) const override
+	virtual void Set( FRHICommandList& RHICmdList, FNiagaraShader* Shader, class UNiagaraDataInterface* DataInterface, void* PerInstanceData) const override
 	{
 		check( IsInRenderingThread() );
 
diff --git a/Engine/Plugins/FX/Niagara/Config/BaseNiagara.ini b/Engine/Plugins/FX/Niagara/Config/BaseNiagara.ini
--- a/Engine/Plugins/FX/Niagara/Config/BaseNiagara.ini
+++ b/Engine/Plugins/FX/Niagara/Config/BaseNiagara.ini
@@ -16,9 +16,11 @@
 ;+EnumRedirects=(OldName="ENiagaraSortMode",ValueChanges=(("ENiagaraSortMode::SortNone","ENiagaraSortMode::None"), ("ENiagaraSortMode::SortViewDepth","ENiagaraSortMode::ViewDepth"), ("ENiagaraSortMode::SortViewDistance","ENiagaraSortMode::ViewDistance")) )
 
 [/Script/Niagara.NiagaraSettings]
-+AdditionalParameterEnums=/Niagara/Functions/Localspace/ENiagaraCoordinateSpace.ENiagaraCoordinateSpace
-+AdditionalParameterEnums=/Niagara/Functions/Localspace/ENiagaraOrientationAxis.ENiagaraOrientationAxis
-+AdditionalParameterEnums=/Niagara/DynamicInputs/Bool/ENiagaraBooleanLogicOps.ENiagaraBooleanLogicOps
++AdditionalParameterEnums=/Niagara/Enums/ENiagaraCoordinateSpace.ENiagaraCoordinateSpace
++AdditionalParameterEnums=/Niagara/Enums/ENiagaraOrientationAxis.ENiagaraOrientationAxis
++AdditionalParameterEnums=/Niagara/Enums/ENiagaraBooleanLogicOps.ENiagaraBooleanLogicOps
++AdditionalParameterEnums=/Niagara/Enums/ENiagaraRandomnessMode.ENiagaraRandomnessMode
++AdditionalParameterEnums=/Niagara/Enums/ENiagaraExpansionMode.ENiagaraExpansionMode
 
 [/Script/NiagaraEditor.NiagaraEditorSettings] 
 DefaultSystem=None
@@ -30,11 +32,14 @@ DefaultModuleScript=/Niagara/DefaultAssets/DefaultModule.DefaultModule
 +GraphCreationShortcuts=(Name="Numeric::Add",Input=(Key=A,bShift=False,bCtrl=False,bAlt=False,bCmd=False))
 +GraphCreationShortcuts=(Name="Numeric::Div",Input=(Key=D,bShift=False,bCtrl=False,bAlt=False,bCmd=False))
 +GraphCreationShortcuts=(Name="Numeric::Pow",Input=(Key=E,bShift=False,bCtrl=False,bAlt=False,bCmd=False))
++GraphCreationShortcuts=(Name="Numeric::Subtract",Input=(Key=S,bShift=False,bCtrl=False,bAlt=False,bCmd=False))
 +GraphCreationShortcuts=(Name="If",Input=(Key=I,bShift=False,bCtrl=False,bAlt=False,bCmd=False))
 +GraphCreationShortcuts=(Name="Numeric::Mul",Input=(Key=M,bShift=False,bCtrl=False,bAlt=False,bCmd=False))
 +GraphCreationShortcuts=(Name="Numeric::Normalize",Input=(Key=N,bShift=False,bCtrl=False,bAlt=False,bCmd=False))
 +GraphCreationShortcuts=(Name="Numeric::OneMinus",Input=(Key=O,bShift=False,bCtrl=False,bAlt=False,bCmd=False))
++GraphCreationShortcuts=(Name="Lerp",Input=(Key=L,bShift=False,bCtrl=False,bAlt=False,bCmd=False))
 +GraphCreationShortcuts=(Name="float",Input=(Key=One,bShift=False,bCtrl=False,bAlt=False,bCmd=False))
++GraphCreationShortcuts=(Name="bool",Input=(Key=B,bShift=False,bCtrl=False,bAlt=False,bCmd=False))
 +GraphCreationShortcuts=(Name="Vector2D",Input=(Key=Two,bShift=False,bCtrl=False,bAlt=False,bCmd=False))
 +GraphCreationShortcuts=(Name="Vector",Input=(Key=Three,bShift=False,bCtrl=False,bAlt=False,bCmd=False))
 +GraphCreationShortcuts=(Name="Vector4",Input=(Key=Four,bShift=False,bCtrl=False,bAlt=False,bCmd=False))
diff --git a/Engine/Plugins/FX/Niagara/Source/Niagara/Classes/NiagaraCollision.h b/Engine/Plugins/FX/Niagara/Source/Niagara/Classes/NiagaraCollision.h
--- a/Engine/Plugins/FX/Niagara/Source/Niagara/Classes/NiagaraCollision.h
+++ b/Engine/Plugins/FX/Niagara/Source/Niagara/Classes/NiagaraCollision.h
@@ -105,7 +105,7 @@ struct FNiagaraDICollsionQueryResult
 	int32 PhysicalMaterialIdx;
 	float Friction;
 	float Restitution;
-
+	bool IsInsideMesh;
 };
 
 class FNiagaraDICollisionQueryBatch
@@ -157,6 +157,8 @@ public:
 	}
 
 	int32 SubmitQuery(FVector Position, FVector Direction, float CollisionSize, float DeltaSeconds);
+	int32 SubmitQuery(FVector StartPos, FVector EndPos, ECollisionChannel TraceChannel);
+	bool PerformQuery(FVector StartPos, FVector EndPos, FNiagaraDICollsionQueryResult &Result, ECollisionChannel TraceChannel);
 	bool GetQueryResult(uint32 TraceID, FNiagaraDICollsionQueryResult &Result);
 
 private:
diff --git a/Engine/Plugins/FX/Niagara/Source/Niagara/Classes/NiagaraConstants.h b/Engine/Plugins/FX/Niagara/Source/Niagara/Classes/NiagaraConstants.h
--- a/Engine/Plugins/FX/Niagara/Source/Niagara/Classes/NiagaraConstants.h
+++ b/Engine/Plugins/FX/Niagara/Source/Niagara/Classes/NiagaraConstants.h
@@ -20,78 +20,87 @@
 #define PARAM_MAP_RAPID_ITERATION_BASE_STR TEXT("Constants")
 
 
-#define SYS_PARAM_ENGINE_DELTA_TIME					INiagaraModule::GetVar_Engine_DeltaTime()
-#define SYS_PARAM_ENGINE_INV_DELTA_TIME				INiagaraModule::GetVar_Engine_InvDeltaTime()
-#define SYS_PARAM_ENGINE_TIME						INiagaraModule::GetVar_Engine_Time()
-#define SYS_PARAM_ENGINE_REAL_TIME					INiagaraModule::GetVar_Engine_RealTime()
-#define SYS_PARAM_ENGINE_POSITION					INiagaraModule::GetVar_Engine_Owner_Position()
-#define SYS_PARAM_ENGINE_VELOCITY					INiagaraModule::GetVar_Engine_Owner_Velocity()
-#define SYS_PARAM_ENGINE_X_AXIS						INiagaraModule::GetVar_Engine_Owner_XAxis()
-#define SYS_PARAM_ENGINE_Y_AXIS						INiagaraModule::GetVar_Engine_Owner_YAxis()
-#define SYS_PARAM_ENGINE_Z_AXIS						INiagaraModule::GetVar_Engine_Owner_ZAxis()
-#define SYS_PARAM_ENGINE_SCALE						INiagaraModule::GetVar_Engine_Owner_Scale()
-
-#define SYS_PARAM_ENGINE_LOCAL_TO_WORLD				INiagaraModule::GetVar_Engine_Owner_SystemLocalToWorld()
-#define SYS_PARAM_ENGINE_WORLD_TO_LOCAL				INiagaraModule::GetVar_Engine_Owner_SystemWorldToLocal()
-#define SYS_PARAM_ENGINE_LOCAL_TO_WORLD_TRANSPOSED	INiagaraModule::GetVar_Engine_Owner_SystemLocalToWorldTransposed()
-#define SYS_PARAM_ENGINE_WORLD_TO_LOCAL_TRANSPOSED	INiagaraModule::GetVar_Engine_Owner_SystemWorldToLocalTransposed()
-#define SYS_PARAM_ENGINE_LOCAL_TO_WORLD_NO_SCALE	INiagaraModule::GetVar_Engine_Owner_SystemLocalToWorldNoScale()
-#define SYS_PARAM_ENGINE_WORLD_TO_LOCAL_NO_SCALE	INiagaraModule::GetVar_Engine_Owner_SystemWorldToLocalNoScale()
-
-#define SYS_PARAM_ENGINE_TIME_SINCE_RENDERED		INiagaraModule::GetVar_Engine_Owner_TimeSinceRendered()
-#define SYS_PARAM_ENGINE_MIN_DIST_TO_CAMERA			INiagaraModule::GetVar_Engine_Owner_MinDistanceToCamera()
-
-#define SYS_PARAM_ENGINE_EXECUTION_STATE			INiagaraModule::GetVar_Engine_Owner_ExecutionState()
-
-#define SYS_PARAM_ENGINE_EXEC_COUNT					INiagaraModule::GetVar_Engine_ExecutionCount()
-#define SYS_PARAM_ENGINE_EMITTER_NUM_PARTICLES		INiagaraModule::GetVar_Engine_Emitter_NumParticles()
-#define SYS_PARAM_ENGINE_SYSTEM_NUM_EMITTERS_ALIVE	INiagaraModule::GetVar_Engine_System_NumEmittersAlive()
-#define SYS_PARAM_ENGINE_SYSTEM_NUM_EMITTERS		INiagaraModule::GetVar_Engine_System_NumEmitters()
-#define SYS_PARAM_ENGINE_NUM_SYSTEM_INSTANCES		INiagaraModule::GetVar_Engine_NumSystemInstances()
-
-#define SYS_PARAM_ENGINE_GLOBAL_SPAWN_COUNT_SCALE	INiagaraModule::GetVar_Engine_GlobalSpawnCountScale()
-#define SYS_PARAM_ENGINE_GLOBAL_SYSTEM_COUNT_SCALE	INiagaraModule::GetVar_Engine_GlobalSystemScale()
-
-#define SYS_PARAM_ENGINE_SYSTEM_AGE					INiagaraModule::GetVar_Engine_System_Age()
-
-#define SYS_PARAM_EMITTER_AGE						INiagaraModule::GetVar_Emitter_Age()
-#define SYS_PARAM_EMITTER_LOCALSPACE				INiagaraModule::GetVar_Emitter_LocalSpace()
-#define SYS_PARAM_EMITTER_SPAWNRATE					INiagaraModule::GetVar_Emitter_SpawnRate()
-#define SYS_PARAM_EMITTER_SPAWN_INTERVAL			INiagaraModule::GetVar_Emitter_SpawnInterval()
-#define SYS_PARAM_EMITTER_INTERP_SPAWN_START_DT		INiagaraModule::GetVar_Emitter_InterpSpawnStartDt()
-#define SYS_PARAM_EMITTER_SPAWN_GROUP		INiagaraModule::GetVar_Emitter_SpawnGroup()
-
-#define SYS_PARAM_PARTICLES_ID						INiagaraModule::GetVar_Particles_ID()
-#define SYS_PARAM_PARTICLES_POSITION				INiagaraModule::GetVar_Particles_Position()
-#define SYS_PARAM_PARTICLES_VELOCITY				INiagaraModule::GetVar_Particles_Velocity()
-#define SYS_PARAM_PARTICLES_COLOR					INiagaraModule::GetVar_Particles_Color()
-#define SYS_PARAM_PARTICLES_SPRITE_ROTATION			INiagaraModule::GetVar_Particles_SpriteRotation()
-#define SYS_PARAM_PARTICLES_NORMALIZED_AGE			INiagaraModule::GetVar_Particles_NormalizedAge()
-#define SYS_PARAM_PARTICLES_SPRITE_SIZE				INiagaraModule::GetVar_Particles_SpriteSize()
-#define SYS_PARAM_PARTICLES_SPRITE_FACING			INiagaraModule::GetVar_Particles_SpriteFacing()
-#define SYS_PARAM_PARTICLES_SPRITE_ALIGNMENT		INiagaraModule::GetVar_Particles_SpriteAlignment()
-#define SYS_PARAM_PARTICLES_SUB_IMAGE_INDEX			INiagaraModule::GetVar_Particles_SubImageIndex()
-#define SYS_PARAM_PARTICLES_DYNAMIC_MATERIAL_PARAM	INiagaraModule::GetVar_Particles_DynamicMaterialParameter()
-#define SYS_PARAM_PARTICLES_DYNAMIC_MATERIAL_PARAM_1	INiagaraModule::GetVar_Particles_DynamicMaterialParameter1()
-#define SYS_PARAM_PARTICLES_DYNAMIC_MATERIAL_PARAM_2	INiagaraModule::GetVar_Particles_DynamicMaterialParameter2()
-#define SYS_PARAM_PARTICLES_DYNAMIC_MATERIAL_PARAM_3	INiagaraModule::GetVar_Particles_DynamicMaterialParameter3()
-#define SYS_PARAM_PARTICLES_SCALE					INiagaraModule::GetVar_Particles_Scale()
-#define SYS_PARAM_PARTICLES_LIFETIME				INiagaraModule::GetVar_Particles_Lifetime()
-#define SYS_PARAM_PARTICLES_MESH_ORIENTATION		INiagaraModule::GetVar_Particles_MeshOrientation()
-#define SYS_PARAM_PARTICLES_UV_SCALE				INiagaraModule::GetVar_Particles_UVScale()
-#define SYS_PARAM_PARTICLES_CAMERA_OFFSET			INiagaraModule::GetVar_Particles_CameraOffset()
-#define SYS_PARAM_PARTICLES_MATERIAL_RANDOM			INiagaraModule::GetVar_Particles_MaterialRandom()
-#define SYS_PARAM_PARTICLES_LIGHT_RADIUS			INiagaraModule::GetVar_Particles_LightRadius()
-
-#define SYS_PARAM_PARTICLES_RIBBONID				INiagaraModule::GetVar_Particles_RibbonID()
-#define SYS_PARAM_PARTICLES_RIBBONWIDTH				INiagaraModule::GetVar_Particles_RibbonWidth()
-#define SYS_PARAM_PARTICLES_RIBBONTWIST				INiagaraModule::GetVar_Particles_RibbonTwist()
-#define SYS_PARAM_PARTICLES_RIBBONFACING			INiagaraModule::GetVar_Particles_RibbonFacing()
-#define SYS_PARAM_PARTICLES_RIBBONLINKORDER			INiagaraModule::GetVar_Particles_RibbonLinkOrder()
-
-#define SYS_PARAM_INSTANCE_ALIVE					INiagaraModule::GetVar_DataInstance_Alive()
-
-#define TRANSLATOR_PARAM_BEGIN_DEFAULTS				INiagaraModule::GetVar_BeginDefaults()
+#define SYS_PARAM_ENGINE_DELTA_TIME                      INiagaraModule::GetVar_Engine_DeltaTime()
+#define SYS_PARAM_ENGINE_INV_DELTA_TIME                  INiagaraModule::GetVar_Engine_InvDeltaTime()
+#define SYS_PARAM_ENGINE_TIME                            INiagaraModule::GetVar_Engine_Time()
+#define SYS_PARAM_ENGINE_REAL_TIME                       INiagaraModule::GetVar_Engine_RealTime()
+#define SYS_PARAM_ENGINE_POSITION                        INiagaraModule::GetVar_Engine_Owner_Position()
+#define SYS_PARAM_ENGINE_VELOCITY                        INiagaraModule::GetVar_Engine_Owner_Velocity()
+#define SYS_PARAM_ENGINE_X_AXIS                          INiagaraModule::GetVar_Engine_Owner_XAxis()
+#define SYS_PARAM_ENGINE_Y_AXIS                          INiagaraModule::GetVar_Engine_Owner_YAxis()
+#define SYS_PARAM_ENGINE_Z_AXIS                          INiagaraModule::GetVar_Engine_Owner_ZAxis()
+#define SYS_PARAM_ENGINE_ROTATION                        INiagaraModule::GetVar_Engine_Owner_Rotation()
+#define SYS_PARAM_ENGINE_SCALE                           INiagaraModule::GetVar_Engine_Owner_Scale()
+
+#define SYS_PARAM_ENGINE_LOCAL_TO_WORLD                  INiagaraModule::GetVar_Engine_Owner_SystemLocalToWorld()
+#define SYS_PARAM_ENGINE_WORLD_TO_LOCAL                  INiagaraModule::GetVar_Engine_Owner_SystemWorldToLocal()
+#define SYS_PARAM_ENGINE_LOCAL_TO_WORLD_TRANSPOSED       INiagaraModule::GetVar_Engine_Owner_SystemLocalToWorldTransposed()
+#define SYS_PARAM_ENGINE_WORLD_TO_LOCAL_TRANSPOSED       INiagaraModule::GetVar_Engine_Owner_SystemWorldToLocalTransposed()
+#define SYS_PARAM_ENGINE_LOCAL_TO_WORLD_NO_SCALE         INiagaraModule::GetVar_Engine_Owner_SystemLocalToWorldNoScale()
+#define SYS_PARAM_ENGINE_WORLD_TO_LOCAL_NO_SCALE         INiagaraModule::GetVar_Engine_Owner_SystemWorldToLocalNoScale()
+
+#define SYS_PARAM_ENGINE_TIME_SINCE_RENDERED             INiagaraModule::GetVar_Engine_Owner_TimeSinceRendered()
+#define SYS_PARAM_ENGINE_MIN_DIST_TO_CAMERA              INiagaraModule::GetVar_Engine_Owner_MinDistanceToCamera()
+
+#define SYS_PARAM_ENGINE_EXECUTION_STATE                 INiagaraModule::GetVar_Engine_Owner_ExecutionState()
+
+#define SYS_PARAM_ENGINE_EXEC_COUNT                      INiagaraModule::GetVar_Engine_ExecutionCount()
+#define SYS_PARAM_ENGINE_EMITTER_NUM_PARTICLES           INiagaraModule::GetVar_Engine_Emitter_NumParticles()
+#define SYS_PARAM_ENGINE_EMITTER_TOTAL_SPAWNED_PARTICLES INiagaraModule::GetVar_Engine_Emitter_TotalSpawnedParticles()
+#define SYS_PARAM_ENGINE_SYSTEM_NUM_EMITTERS_ALIVE       INiagaraModule::GetVar_Engine_System_NumEmittersAlive()
+#define SYS_PARAM_ENGINE_SYSTEM_NUM_EMITTERS             INiagaraModule::GetVar_Engine_System_NumEmitters()
+#define SYS_PARAM_ENGINE_NUM_SYSTEM_INSTANCES            INiagaraModule::GetVar_Engine_NumSystemInstances()
+
+#define SYS_PARAM_ENGINE_GLOBAL_SPAWN_COUNT_SCALE        INiagaraModule::GetVar_Engine_GlobalSpawnCountScale()
+#define SYS_PARAM_ENGINE_GLOBAL_SYSTEM_COUNT_SCALE       INiagaraModule::GetVar_Engine_GlobalSystemScale()
+
+#define SYS_PARAM_ENGINE_SYSTEM_AGE                      INiagaraModule::GetVar_Engine_System_Age()
+#define SYS_PARAM_ENGINE_SYSTEM_TICK_COUNT               INiagaraModule::GetVar_Engine_System_TickCount()
+
+#define SYS_PARAM_EMITTER_AGE                            INiagaraModule::GetVar_Emitter_Age()
+#define SYS_PARAM_EMITTER_LOCALSPACE                     INiagaraModule::GetVar_Emitter_LocalSpace()
+#define SYS_PARAM_EMITTER_DETERMINISM                    INiagaraModule::GetVar_Emitter_Determinism()
+#define SYS_PARAM_EMITTER_RANDOM_SEED                    INiagaraModule::GetVar_Emitter_RandomSeed()
+#define SYS_PARAM_EMITTER_SPAWNRATE                      INiagaraModule::GetVar_Emitter_SpawnRate()
+#define SYS_PARAM_EMITTER_SPAWN_INTERVAL                 INiagaraModule::GetVar_Emitter_SpawnInterval()
+#define SYS_PARAM_EMITTER_INTERP_SPAWN_START_DT          INiagaraModule::GetVar_Emitter_InterpSpawnStartDt()
+#define SYS_PARAM_EMITTER_SPAWN_GROUP                    INiagaraModule::GetVar_Emitter_SpawnGroup()
+
+#define SYS_PARAM_PARTICLES_UNIQUE_ID                    INiagaraModule::GetVar_Particles_UniqueID()
+#define SYS_PARAM_PARTICLES_ID                           INiagaraModule::GetVar_Particles_ID()
+#define SYS_PARAM_PARTICLES_POSITION                     INiagaraModule::GetVar_Particles_Position()
+#define SYS_PARAM_PARTICLES_VELOCITY                     INiagaraModule::GetVar_Particles_Velocity()
+#define SYS_PARAM_PARTICLES_COLOR                        INiagaraModule::GetVar_Particles_Color()
+#define SYS_PARAM_PARTICLES_SPRITE_ROTATION              INiagaraModule::GetVar_Particles_SpriteRotation()
+#define SYS_PARAM_PARTICLES_NORMALIZED_AGE               INiagaraModule::GetVar_Particles_NormalizedAge()
+#define SYS_PARAM_PARTICLES_SPRITE_SIZE                  INiagaraModule::GetVar_Particles_SpriteSize()
+#define SYS_PARAM_PARTICLES_SPRITE_FACING                INiagaraModule::GetVar_Particles_SpriteFacing()
+#define SYS_PARAM_PARTICLES_SPRITE_ALIGNMENT             INiagaraModule::GetVar_Particles_SpriteAlignment()
+#define SYS_PARAM_PARTICLES_SUB_IMAGE_INDEX              INiagaraModule::GetVar_Particles_SubImageIndex()
+#define SYS_PARAM_PARTICLES_DYNAMIC_MATERIAL_PARAM       INiagaraModule::GetVar_Particles_DynamicMaterialParameter()
+#define SYS_PARAM_PARTICLES_DYNAMIC_MATERIAL_PARAM_1     INiagaraModule::GetVar_Particles_DynamicMaterialParameter1()
+#define SYS_PARAM_PARTICLES_DYNAMIC_MATERIAL_PARAM_2     INiagaraModule::GetVar_Particles_DynamicMaterialParameter2()
+#define SYS_PARAM_PARTICLES_DYNAMIC_MATERIAL_PARAM_3     INiagaraModule::GetVar_Particles_DynamicMaterialParameter3()
+#define SYS_PARAM_PARTICLES_SCALE                        INiagaraModule::GetVar_Particles_Scale()
+#define SYS_PARAM_PARTICLES_LIFETIME                     INiagaraModule::GetVar_Particles_Lifetime()
+#define SYS_PARAM_PARTICLES_MESH_ORIENTATION             INiagaraModule::GetVar_Particles_MeshOrientation()
+#define SYS_PARAM_PARTICLES_UV_SCALE                     INiagaraModule::GetVar_Particles_UVScale()
+#define SYS_PARAM_PARTICLES_CAMERA_OFFSET                INiagaraModule::GetVar_Particles_CameraOffset()
+#define SYS_PARAM_PARTICLES_MATERIAL_RANDOM              INiagaraModule::GetVar_Particles_MaterialRandom()
+#define SYS_PARAM_PARTICLES_LIGHT_RADIUS                 INiagaraModule::GetVar_Particles_LightRadius()
+#define SYS_PARAM_PARTICLES_LIGHT_EXPONENT               INiagaraModule::GetVar_Particles_LightExponent()
+#define SYS_PARAM_PARTICLES_LIGHT_ENABLED                INiagaraModule::GetVar_Particles_LightEnabled()
+#define SYS_PARAM_PARTICLES_LIGHT_VOLUMETRIC_SCATTERING  INiagaraModule::GetVar_Particles_LightVolumetricScattering()
+
+#define SYS_PARAM_PARTICLES_RIBBONID                     INiagaraModule::GetVar_Particles_RibbonID()
+#define SYS_PARAM_PARTICLES_RIBBONWIDTH                  INiagaraModule::GetVar_Particles_RibbonWidth()
+#define SYS_PARAM_PARTICLES_RIBBONTWIST                  INiagaraModule::GetVar_Particles_RibbonTwist()
+#define SYS_PARAM_PARTICLES_RIBBONFACING                 INiagaraModule::GetVar_Particles_RibbonFacing()
+#define SYS_PARAM_PARTICLES_RIBBONLINKORDER              INiagaraModule::GetVar_Particles_RibbonLinkOrder()
+
+#define SYS_PARAM_INSTANCE_ALIVE                         INiagaraModule::GetVar_DataInstance_Alive()
+
+#define TRANSLATOR_PARAM_BEGIN_DEFAULTS                  INiagaraModule::GetVar_BeginDefaults()
 
 struct NIAGARA_API FNiagaraConstants
 {
diff --git a/Engine/Plugins/FX/Niagara/Source/Niagara/Classes/NiagaraDataInterfaceCollisionQuery.h b/Engine/Plugins/FX/Niagara/Source/Niagara/Classes/NiagaraDataInterfaceCollisionQuery.h
--- a/Engine/Plugins/FX/Niagara/Source/Niagara/Classes/NiagaraDataInterfaceCollisionQuery.h
+++ b/Engine/Plugins/FX/Niagara/Source/Niagara/Classes/NiagaraDataInterfaceCollisionQuery.h
@@ -35,9 +35,6 @@ public:
 	//UObject Interface
 	virtual void PostInitProperties() override;
 	virtual void PostLoad() override;
-#if WITH_EDITOR
-	virtual void PostEditChangeProperty(struct FPropertyChangedEvent& PropertyChangedEvent) override;
-#endif
 	//UObject Interface End
 
 	/** Initializes the per instance data for this interface. Returns false if there was some error and the simulation should be disabled. */
@@ -48,10 +45,8 @@ public:
 	/** Ticks the per instance data for this interface, if it has any. */
 	virtual bool PerInstanceTick(void* PerInstanceData, FNiagaraSystemInstance* SystemInstance, float DeltaSeconds);
 	virtual bool PerInstanceTickPostSimulate(void* PerInstanceData, FNiagaraSystemInstance* SystemInstance, float DeltaSeconds);
-
 	virtual int32 PerInstanceDataSize()const { return sizeof(CQDIPerInstanceData); }
-
-
+	
 	virtual void GetFunctions(TArray<FNiagaraFunctionSignature>& OutFunctions)override;
 	virtual void GetVMExternalFunction(const FVMExternalFunctionBindingInfo& BindingInfo, void* InstanceData, FVMExternalFunction &OutFunc) override;
 
@@ -59,19 +54,20 @@ public:
 	void SubmitQuery(FVectorVMContext& Context);
 	void ReadQuery(FVectorVMContext& Context);
 	void PerformQuery(FVectorVMContext& Context);
+	void PerformQuerySyncCPU(FVectorVMContext& Context);
+	void PerformQueryAsyncCPU(FVectorVMContext& Context);
+	void PerformQueryGPU(FVectorVMContext& Context);
+	void QuerySceneDepth(FVectorVMContext& Context);
+	void QueryMeshDistanceField(FVectorVMContext& Context);
 
-	virtual bool Equals(const UNiagaraDataInterface* Other) const override;
 	virtual bool CanExecuteOnTarget(ENiagaraSimTarget Target) const override { return true; }
 
 	virtual bool GetFunctionHLSL(const FName&  DefinitionFunctionName, FString InstanceFunctionName, FNiagaraDataInterfaceGPUParamInfo& ParamInfo, FString& OutHLSL) override;
 	virtual void GetParameterDefinitionHLSL(FNiagaraDataInterfaceGPUParamInfo& ParamInfo, FString& OutHLSL) override;
-	virtual FNiagaraDataInterfaceParametersCS* ConstructComputeParameters()const override;
+	virtual FNiagaraDataInterfaceParametersCS* ConstructComputeParameters() const override;
 	
-
-protected:
-	virtual bool CopyToInternal(UNiagaraDataInterface* Destination) const override;
-
 private:
 
 	static FCriticalSection CriticalSection;
+	UEnum* TraceChannelEnum;
 };
diff --git a/Engine/Plugins/FX/Niagara/Source/Niagara/Classes/NiagaraDataInterfaceCurlNoise.h b/Engine/Plugins/FX/Niagara/Source/Niagara/Classes/NiagaraDataInterfaceCurlNoise.h
--- a/Engine/Plugins/FX/Niagara/Source/Niagara/Classes/NiagaraDataInterfaceCurlNoise.h
+++ b/Engine/Plugins/FX/Niagara/Source/Niagara/Classes/NiagaraDataInterfaceCurlNoise.h
@@ -12,52 +12,37 @@ UCLASS(EditInlineNew, Category = "Curl Noise LUT", meta = (DisplayName = "Curl N
 class NIAGARA_API UNiagaraDataInterfaceCurlNoise : public UNiagaraDataInterface
 {
 	GENERATED_UCLASS_BODY()
-	bool bGPUBufferDirty;
 public:
 
 	UPROPERTY(EditAnywhere, Category = "Curl Noise")
 	uint32 Seed;
-	
+
+	// Precalculated when Seed changes. 
+	FVector OffsetFromSeed;
+
 	//UObject Interface
 	virtual void PostInitProperties()override;
 	virtual void PostLoad() override;
-	virtual void BeginDestroy() override;
-	virtual bool IsReadyForFinishDestroy() override;
 #if WITH_EDITOR
+	virtual void PreEditChange(UProperty* PropertyAboutToChange) override;
 	virtual void PostEditChangeProperty(struct FPropertyChangedEvent& PropertyChangedEvent) override;
 #endif
 	//UObject Interface End
 
 	//UNiagaraDataInterface Interface
-	virtual void GetFunctions(TArray<FNiagaraFunctionSignature>& OutFunctions)override;
+	virtual void GetFunctions(TArray<FNiagaraFunctionSignature>& OutFunctions) override;
 	virtual void GetVMExternalFunction(const FVMExternalFunctionBindingInfo& BindingInfo, void* InstanceData, FVMExternalFunction &OutFunc) override;
 	virtual bool CanExecuteOnTarget(ENiagaraSimTarget Target)const override { return true; }
+	virtual bool Equals(const UNiagaraDataInterface* Other) const override;
 	//UNiagaraDataInterface Interface End
 
 	void SampleNoiseField(FVectorVMContext& Context);
 
-	virtual bool Equals(const UNiagaraDataInterface* Other) const override;
-
 	// GPU sim functionality
 	virtual bool GetFunctionHLSL(const FName&  DefinitionFunctionName, FString InstanceFunctionName, FNiagaraDataInterfaceGPUParamInfo& ParamInfo, FString& OutHLSL) override;
 	virtual void GetParameterDefinitionHLSL(FNiagaraDataInterfaceGPUParamInfo& ParamInfo, FString& OutHLSL) override;
-	virtual FNiagaraDataInterfaceParametersCS* ConstructComputeParameters()const override;
+	virtual FNiagaraDataInterfaceParametersCS* ConstructComputeParameters() const override;
 
-	void ReleaseResource();
-	FRWBuffer& GetGPUBuffer();
-	static const FString CurlNoiseBufferName;
 protected:
 	virtual bool CopyToInternal(UNiagaraDataInterface* Destination) const override;
-	void InitNoiseLUT();
-
-	VectorRegister NoiseTable[17][17][17];
-
-	template<typename T>
-	void ReplicateBorder(T* DestBuffer);
-
-	TUniquePtr<FRWBuffer> GPUBuffer;
-	/** A fence which is used to keep track of the rendering thread releasing RHI resources. */
-	FRenderCommandFence ReleaseResourcesFence;
-
-	static const FName SampleNoiseFieldName;
 };
diff --git a/Engine/Plugins/FX/Niagara/Source/Niagara/Classes/NiagaraDataInterfaceVectorField.h b/Engine/Plugins/FX/Niagara/Source/Niagara/Classes/NiagaraDataInterfaceVectorField.h
--- a/Engine/Plugins/FX/Niagara/Source/Niagara/Classes/NiagaraDataInterfaceVectorField.h
+++ b/Engine/Plugins/FX/Niagara/Source/Niagara/Classes/NiagaraDataInterfaceVectorField.h
@@ -7,35 +7,13 @@
 
 class FNiagaraSystemInstance;
 
-/*struct FNDIVectorField_InstanceData
-{
-	UVectorField* Field;
-	uint32 SizeX;
-	uint32 SizeY;
-	uint32 SizeZ;
-	float DeltaSeconds;
-
-	FBox LocalBounds;
-
-	FVector TilingAxes;
-
-	FORCEINLINE_DEBUGGABLE bool ResetRequired(class UNiagaraDataInterfaceVectorField* Interface)const;
-
-	FORCEINLINE_DEBUGGABLE bool Init(class UNiagaraDataInterfaceVectorField* Interface, FNiagaraSystemInstance* SystemInstance);
-	FORCEINLINE_DEBUGGABLE bool Cleanup(class UNiagaraDataInterfaceVectorField* Interface, FNiagaraSystemInstance* SystemInstance);
-	FORCEINLINE_DEBUGGABLE bool Tick(class UNiagaraDataInterfaceVectorField* Interface, FNiagaraSystemInstance* SystemInstance, float InDeltaSeconds);
-	void UpdateTransforms(const FMatrix& LocalToWorld, FMatrix& OutVolumeToWorld, FMatrix& OutWorldToVolume);
-	const void* Lock();
-	void Unlock();
-};*/
-
 UCLASS(EditInlineNew, Category = "Vector Field", meta = (DisplayName = "Vector Field"))
 class NIAGARA_API UNiagaraDataInterfaceVectorField : public UNiagaraDataInterface
 {
 	GENERATED_UCLASS_BODY()
 
 public:
-	/** Vector field used to sample from. */
+	/** Vector field to sample from. */
 	UPROPERTY(EditAnywhere, Category = VectorField)
 	UVectorField* Field;
 
@@ -48,57 +26,47 @@ public:
 
 public:
 	//~ UObject interface
+
+	virtual void PostInitProperties() override;
+	virtual void PostLoad() override; 
 #if WITH_EDITOR
-	virtual void PostInitProperties()override;
 	virtual void PostEditChangeProperty(struct FPropertyChangedEvent& PropertyChangedEvent) override;
-	virtual void PostLoad() override;
+	virtual void PreEditChange(UProperty* PropertyAboutToChange) override;
 #endif
 	//~ UObject interface END
 
 	//~ UNiagaraDataInterface interface
-	virtual bool InitPerInstanceData(void* PerInstanceData, FNiagaraSystemInstance* SystemInstance)override;
-	virtual void DestroyPerInstanceData(void* PerInstanceData, FNiagaraSystemInstance* SystemInstance)override;
-	virtual bool PerInstanceTick(void* PerInstanceData, FNiagaraSystemInstance* SystemInstance, float DeltaSeconds) override;
-	virtual int32 PerInstanceDataSize()const override;
-
-	virtual void GetFunctions(TArray<FNiagaraFunctionSignature>& OutFunctions)override;
-	virtual void GetVMExternalFunction(const FVMExternalFunctionBindingInfo& BindingInfo, void* InstanceData, FVMExternalFunction &OutFunc)override;
+	// VM functionality
+	virtual void GetFunctions(TArray<FNiagaraFunctionSignature>& OutFunctions) override;
+	virtual void GetVMExternalFunction(const FVMExternalFunctionBindingInfo& BindingInfo, void* InstanceData, FVMExternalFunction &OutFunc) override;
 	virtual bool Equals(const UNiagaraDataInterface* Other) const override;
-	virtual bool CanExecuteOnTarget(ENiagaraSimTarget Target)const override;
+	virtual bool CanExecuteOnTarget(ENiagaraSimTarget Target) const override;
+
+#if WITH_EDITOR	
+	// Editor functionality
+	virtual TArray<FNiagaraDataInterfaceError> GetErrors() override;
+#endif
 
 	// GPU sim functionality
+	virtual void GetParameterDefinitionHLSL(FNiagaraDataInterfaceGPUParamInfo& ParamInfo, FString& OutHLSL) override; 
 	virtual bool GetFunctionHLSL(const FName&  DefinitionFunctionName, FString InstanceFunctionName, FNiagaraDataInterfaceGPUParamInfo& ParamInfo, FString& OutHLSL) override;
-	virtual void GetParameterDefinitionHLSL(FNiagaraDataInterfaceGPUParamInfo& ParamInfo, FString& OutHLSL) override;
-	virtual FNiagaraDataInterfaceParametersCS* ConstructComputeParameters()const override;
+	virtual FNiagaraDataInterfaceParametersCS* ConstructComputeParameters() const override;
 	//~ UNiagaraDataInterface interface END
 
-	void SampleVectorField(FVectorVMContext& Context);	
-
+	// VM functions
 	void GetFieldDimensions(FVectorVMContext& Context);
+	void GetFieldBounds(FVectorVMContext& Context); 
+	void GetFieldTilingAxes(FVectorVMContext& Context);
+	void SampleVectorField(FVectorVMContext& Context);
+	
+	//	
+	FVector GetTilingAxes() const;
+	FVector GetDimensions() const;
+	FVector GetMinBounds() const;
+	FVector GetMaxBounds() const;
 
-	void GetFieldBounds(FVectorVMContext& Context);
-
-	bool bGPUBufferDirty;
-
-	static const FString BufferBaseName;
-	static const FString DimentionsBaseName;
-	static const FString BoundsMinBaseName;
-	static const FString BoundsMaxBaseName;
-
-	FRWBuffer& GetGPUBuffer();
-	FORCEINLINE FVector GetDimentions()const { return FVector(SizeX, SizeY, SizeZ); }
-	FORCEINLINE FVector GetBoundsMin()const { return LocalBounds.Min; }
-	FORCEINLINE FVector GetBoundsMax()const { return LocalBounds.Max; }
 protected:
-
-	const void* Lock();
-	void Unlock();
+	//~ UNiagaraDataInterface interface
 	virtual bool CopyToInternal(UNiagaraDataInterface* Destination) const override;
-	void InitField();
-
-	uint32 SizeX, SizeY, SizeZ;
-	FVector TilingAxes;
-	FBox LocalBounds;
-
-	FRWBuffer GPUBuffer;
+	//~ UNiagaraDataInterface interface END
 };
diff --git a/Engine/Plugins/FX/Niagara/Source/Niagara/Classes/NiagaraDataSet.h b/Engine/Plugins/FX/Niagara/Source/Niagara/Classes/NiagaraDataSet.h
--- a/Engine/Plugins/FX/Niagara/Source/Niagara/Classes/NiagaraDataSet.h
+++ b/Engine/Plugins/FX/Niagara/Source/Niagara/Classes/NiagaraDataSet.h
@@ -42,6 +42,7 @@ public:
 	void KillInstance(uint32 InstanceIdx);
 	void CopyTo(FNiagaraDataBuffer& DestBuffer, int32 StartIdx, int32 NumInstances)const;
 	void CopyTo(FNiagaraDataBuffer& DestBuffer)const;
+	void GPUCopyTo(FNiagaraDataBuffer& DestBuffer, float* GPUReadBackFloat, int* GPUReadBackInt, int32 StartIdx, int32 NumInstances)const;
 
 	const TArray<uint8>& GetFloatBuffer()const { return FloatData; }
 	const TArray<uint8>& GetInt32Buffer()const { return Int32Data; }
@@ -86,6 +87,26 @@ public:
 		return (int32*)GetComponentPtrInt32(ComponentIdx) + InstanceIdx;
 	}
 
+	FORCEINLINE uint8* GetComponentPtrFloat(float* BasePtr, uint32 ComponentIdx) const
+	{
+		return (uint8*)BasePtr + FloatStride * ComponentIdx;
+	}
+
+	FORCEINLINE uint8* GetComponentPtrInt32(int* BasePtr, uint32 ComponentIdx) const
+	{
+		return (uint8*)BasePtr + Int32Stride * ComponentIdx;
+	}
+
+	FORCEINLINE float* GetInstancePtrFloat(float* BasePtr, uint32 ComponentIdx, uint32 InstanceIdx)const
+	{
+		return (float*)GetComponentPtrFloat(BasePtr, ComponentIdx) + InstanceIdx;
+	}
+
+	FORCEINLINE int32* GetInstancePtrInt32(int* BasePtr, uint32 ComponentIdx, uint32 InstanceIdx)const
+	{
+		return (int32*)GetComponentPtrInt32(BasePtr, ComponentIdx) + InstanceIdx;
+	}
+
 	uint32 GetNumInstances()const { return NumInstances; }
 	uint32 GetNumInstancesAllocated()const { return NumInstancesAllocated; }
 
@@ -288,7 +309,7 @@ public:
 		}
 	}
 
-	void SetShaderParams(class FNiagaraShader *Shader, FRHICommandList &CommandList);
+	void SetShaderParams(class FNiagaraShader *Shader, FRHICommandList &CommandList, uint32& WriteBufferIdx, uint32& ReadBufferIdx);
 	void UnsetShaderParams(class FNiagaraShader *Shader, FRHICommandList &CommandList);
 	void Allocate(int32 NumInstances, bool bMaintainExisting=false)
 	{
@@ -326,6 +347,7 @@ public:
 				}
 				//UE_LOG(LogNiagara, Warning, TEXT("DataSetAllocate: Adding New Free IDs: %d - "), RequiredIDs - 1, ExistingNumIDs);
 			}
+#if 0
 			else if (RequiredIDs < ExistingNumIDs >> 1)//Configurable?
 			{
 				//If the max id we use has reduced significantly then we can shrink the tables.
@@ -350,6 +372,7 @@ public:
 				check(NumFreeIDs <= RequiredIDs);
 				FreeIDsTable.SetNumUninitialized(NumFreeIDs);
 			}
+#endif
 			else
 			{
 				//Drop in required size not great enough so just allocate same size.
@@ -384,6 +407,18 @@ public:
 		return CurrBuffer > 0 ? CurrBuffer - 1 : MaxBufferIdx;
 	}
 
+	FORCEINLINE int32 GetCurrBufferIdx() const
+	{
+		return CurrBuffer;
+	}
+
+	FORCEINLINE FNiagaraDataBuffer& GetDataByIndex(int32 InIdx)
+	{
+		check(InIdx < 3);
+		//CheckCorrectThread();//TODO: We should be able to enable these checks and have well defined GT/RT ownership but GPU sims fire these a lot currently.
+		return Data[InIdx];
+	}
+
 	FORCEINLINE FNiagaraDataBuffer& CurrData() 
 	{
 		//CheckCorrectThread();//TODO: We should be able to enable these checks and have well defined GT/RT ownership but GPU sims fire these a lot currently.
@@ -532,6 +567,7 @@ public:
 
 	void Dump(bool bCurr, int32 StartIdx = 0, int32 NumInstances = INDEX_NONE)const;
 	void Dump(FNiagaraDataSet& Other, bool bCurr, int32 StartIdx = 0, int32 NumInstances = INDEX_NONE)const;
+	void DumpGPU(FNiagaraDataSet& Other, float* GPUReadBackFloat, int* GPUReadBackInt, int32 StartIdx = 0, int32 NumInstances = INDEX_NONE)const;
 	FORCEINLINE const TArray<FNiagaraVariable> &GetVariables() { return Variables; }
 
 	void CheckForNaNs()const
diff --git a/Engine/Plugins/FX/Niagara/Source/Niagara/Classes/NiagaraEmitter.h b/Engine/Plugins/FX/Niagara/Source/Niagara/Classes/NiagaraEmitter.h
--- a/Engine/Plugins/FX/Niagara/Source/Niagara/Classes/NiagaraEmitter.h
+++ b/Engine/Plugins/FX/Niagara/Source/Niagara/Classes/NiagaraEmitter.h
@@ -66,9 +66,8 @@ struct FNiagaraEventGeneratorProperties
 
 	}
 
-	FNiagaraEventGeneratorProperties(FNiagaraDataSetProperties &Props, FName InEventGenerator, FName InSourceEmitter)
+	FNiagaraEventGeneratorProperties(FNiagaraDataSetProperties &Props, FName InEventGenerator)
 		: ID(Props.ID.Name)
-		, SourceEmitter(InSourceEmitter)
 		, SetProps(Props)		
 	{
 
@@ -78,8 +77,8 @@ struct FNiagaraEventGeneratorProperties
 	UPROPERTY(EditAnywhere, Category = "Event Receiver")
 	int32 MaxEventsPerFrame; //TODO - More complex allocation so that we can grow dynamically if more space is needed ?
 
+	UPROPERTY()
 	FName ID;
-	FName SourceEmitter;
 
 	UPROPERTY()
 	FNiagaraDataSetProperties SetProps;
@@ -196,9 +195,18 @@ public:
 	virtual void PostLoad() override;
 	//End UObject Interface
 
+	/** Toggles whether or not the particles within this emitter are relative to the emitter origin or in global space. */ 
 	UPROPERTY(EditAnywhere, Category = "Emitter")
 	bool bLocalSpace;
 
+	/** Toggles whether to globally make the random number generator be deterministic or non-deterministic. Any random calculation that is set to the emitter defaults will inherit this value. It is still possible to tweak individual random to be deterministic or not. In this case deterministic means that it will return the same results for the same configuration of the emitter as long as delta time is not variable. Any changes to the emitter's individual scripts will adjust the results. */
+	UPROPERTY(EditAnywhere, Category = "Emitter")
+	bool bDeterminism;
+
+	/** An emitter-based seed for the deterministic random number generator. */
+	UPROPERTY(EditAnywhere, Category = "Emitter", meta = (EditCondition = "bDeterminism"))
+	int32 RandomSeed;
+
 	//UPROPERTY(EditAnywhere, Category = "Emitter")
 	//float StartTime;
 	//UPROPERTY(EditAnywhere, Category = "Emitter")
@@ -256,6 +264,14 @@ public:
 	UPROPERTY(EditAnywhere, Category = "Emitter")
 	uint32 bRequiresPersistentIDs : 1;
 
+	/** Limits the delta time per tick to prevent simulation spikes due to frame lags. */
+	UPROPERTY(EditAnywhere, AdvancedDisplay, Category = "Emitter", meta = (EditCondition = "bLimitDeltaTime"))
+	float MaxDeltaTimePerTick;
+
+	/** Whether to limit the max tick delta time or not. */
+	UPROPERTY(EditAnywhere, AdvancedDisplay, Category = "Emitter", meta = (InlineEditConditionToggle))
+	uint32 bLimitDeltaTime : 1;
+
 	void NIAGARA_API GetScripts(TArray<UNiagaraScript*>& OutScripts, bool bCompilableOnly = true);
 
 	NIAGARA_API UNiagaraScript* GetScript(ENiagaraScriptUsage Usage, FGuid UsageId);
@@ -332,6 +348,9 @@ public:
 
 	void NIAGARA_API RemoveEventHandlerByUsageId(FGuid EventHandlerUsageId);
 
+	/* Gets whether or not the supplied event generator id matches an event generator which is shared between the particle spawn and update scrips. */
+	bool IsEventGeneratorShared(FName EventGeneratorId) const;
+
 protected:
 	virtual void BeginDestroy() override;
 
@@ -368,6 +387,9 @@ private:
 	UPROPERTY()
 	UNiagaraScript* GPUComputeScript;
 
+	UPROPERTY()
+	TArray<FName> SharedEventGeneratorIds;
+
 #if WITH_EDITOR
 	FOnPropertiesChanged OnPropertiesChangedDelegate;
 #endif
diff --git a/Engine/Plugins/FX/Niagara/Source/Niagara/Classes/NiagaraEmitterInstance.h b/Engine/Plugins/FX/Niagara/Source/Niagara/Classes/NiagaraEmitterInstance.h
--- a/Engine/Plugins/FX/Niagara/Source/Niagara/Classes/NiagaraEmitterInstance.h
+++ b/Engine/Plugins/FX/Niagara/Source/Niagara/Classes/NiagaraEmitterInstance.h
@@ -47,7 +47,7 @@ public:
 
 	void PreTick();
 	void Tick(float DeltaSeconds);
-	void PostProcessParticles();
+	void PostTick();
 	bool HandleCompletion(bool bForce=false);
 
 	bool RequiredPersistentID()const;
@@ -73,6 +73,7 @@ public:
 	void NIAGARA_API UpdateEmitterRenderer(ERHIFeatureLevel::Type FeatureLevel, TArray<NiagaraRenderer*>& ToBeAddedList, TArray<NiagaraRenderer*>& ToBeRemovedList);
 
 	FORCEINLINE int32 GetNumParticles()const	{ return ParticleDataSet->GetNumInstances(); }
+	FORCEINLINE int32 GetTotalSpawnedParticles()const	{ return TotalSpawnedParticles; }
 
 	NIAGARA_API const FNiagaraEmitterHandle& GetEmitterHandle() const;
 
@@ -103,6 +104,9 @@ public:
 	NIAGARA_API bool IsReadyToRun() const;
 
 	void Dump()const;
+
+	bool WaitForDebugInfo();
+
 private:
 
 	void CheckForErrors();
@@ -116,6 +120,8 @@ private:
 	uint32 Loops;
 
 	int32 TickCount;
+
+	int32 TotalSpawnedParticles;
 	
 	/** Typical resets must be deferred until the tick as the RT could still be using the current buffer. */
 	uint32 bResetPending : 1;
@@ -153,6 +159,22 @@ private:
 	FNiagaraParameterDirectBinding<int32> UpdateExecCountBinding;
 	TArray<FNiagaraParameterDirectBinding<int32>> EventExecCountBindings;
 
+	/*
+	FNiagaraParameterDirectBinding<int32> SpawnTotalSpawnedParticlesBinding;
+	FNiagaraParameterDirectBinding<int32> UpdateTotalSpawnedParticlesBinding;
+	FNiagaraParameterDirectBinding<int32> GPUTotalSpawnedParticlesBinding;
+	*/
+
+	FNiagaraParameterDirectBinding<int32> SpawnRandomSeedBinding;
+	FNiagaraParameterDirectBinding<int32> UpdateRandomSeedBinding;
+	FNiagaraParameterDirectBinding<int32> GPURandomSeedBinding;
+
+	/*
+	FNiagaraParameterDirectBinding<int32> SpawnDeterminismBinding;
+	FNiagaraParameterDirectBinding<int32> UpdateDeterminismBinding;
+	FNiagaraParameterDirectBinding<int32> GPUDeterminismBinding;
+	*/
+	
 	/** particle simulation data. Must be a shared ref as various things on the RT can have direct ref to it. */
 	FNiagaraDataSet* ParticleDataSet;
 
@@ -166,7 +188,9 @@ private:
 	TArray<FNiagaraDataSet*> UpdateScriptEventDataSets;
 	TArray<FNiagaraDataSet*> SpawnScriptEventDataSets;
 	TMap<FNiagaraDataSetID, FNiagaraDataSet*> DataSetMap;
-	
+
+	TArray<bool> UpdateEventGeneratorIsSharedByIndex;
+	TArray<bool> SpawnEventGeneratorIsSharedByIndex;
 
 	FName OwnerSystemInstanceName;
 
diff --git a/Engine/Plugins/FX/Niagara/Source/Niagara/Classes/NiagaraEmitterInstanceBatcher.h b/Engine/Plugins/FX/Niagara/Source/Niagara/Classes/NiagaraEmitterInstanceBatcher.h
--- a/Engine/Plugins/FX/Niagara/Source/Niagara/Classes/NiagaraEmitterInstanceBatcher.h
+++ b/Engine/Plugins/FX/Niagara/Source/Niagara/Classes/NiagaraEmitterInstanceBatcher.h
@@ -94,11 +94,14 @@ public:
 	void ExecuteAll(FRHICommandList &RHICmdList, FUniformBufferRHIParamRef ViewUniformBuffer);
 	void TickSingle(FNiagaraComputeExecutionContext *Context, FRHICommandList &RHICmdList, FUniformBufferRHIParamRef ViewUniformBuffer) const;
 
+	void ProcessDebugInfo(FRHICommandList &RHICmdList, const FNiagaraComputeExecutionContext *Context) const;
+
 	void SetPrevDataStrideParams(const FNiagaraDataSet *Set, FNiagaraShader *Shader, FRHICommandList &RHICmdList) const;
 
 	void SetupEventUAVs(const FNiagaraComputeExecutionContext *Context, uint32 NumInstances, FRHICommandList &RHICmdList) const;
 	void UnsetEventUAVs(const FNiagaraComputeExecutionContext *Context, FRHICommandList &RHICmdList) const;
-	void SetDataInterfaceParameters(const TArray<UNiagaraDataInterface*> &DataInterfaces, FNiagaraShader *Shader, FRHICommandList &RHICmdList) const;
+	void SetDataInterfaceParameters(const TArray<UNiagaraDataInterface*> &DataInterfaces, FNiagaraShader *Shader, FRHICommandList &RHICmdList, const FNiagaraComputeExecutionContext *Context) const;
+	void UnsetDataInterfaceParameters(const TArray<UNiagaraDataInterface*> &DataInterfaces, FNiagaraShader* Shader, FRHICommandList &RHICmdList, const FNiagaraComputeExecutionContext *Context) const;
 
 	void Run(	const FNiagaraComputeExecutionContext *Context, 
 				uint32 UpdateStartInstance, 
diff --git a/Engine/Plugins/FX/Niagara/Source/Niagara/Classes/NiagaraParameterCollection.h b/Engine/Plugins/FX/Niagara/Source/Niagara/Classes/NiagaraParameterCollection.h
--- a/Engine/Plugins/FX/Niagara/Source/Niagara/Classes/NiagaraParameterCollection.h
+++ b/Engine/Plugins/FX/Niagara/Source/Niagara/Classes/NiagaraParameterCollection.h
@@ -162,6 +162,8 @@ public:
 	//~UObject interface
 	virtual void PostLoad()override;
 	//~UObject interface
+
+	FName GetNamespace()const { return Namespace; }
 protected:
 	
 	void MakeNamespaceNameUnique();
diff --git a/Engine/Plugins/FX/Niagara/Source/Niagara/Classes/NiagaraScript.h b/Engine/Plugins/FX/Niagara/Source/Niagara/Classes/NiagaraScript.h
--- a/Engine/Plugins/FX/Niagara/Source/Niagara/Classes/NiagaraScript.h
+++ b/Engine/Plugins/FX/Niagara/Source/Niagara/Classes/NiagaraScript.h
@@ -51,6 +51,15 @@ enum class ENiagaraModuleDependencyType : uint8
 	PostDependency
 };
 
+UENUM()
+enum class ENiagaraModuleDependencyScriptConstraint : uint8
+{
+	/** The module providing the dependency must be in the same script e.g. if the module requiring the dependency is in "Particle Spawn" the module providing the dependency must also be in "Particle Spawn". */
+	SameScript,
+	/** The module providing the dependency can be in any script as long as it satisfies the dependency type, e.g. if the module requiring the dependency is in "Particle Spawn" the module providing the dependency could be in "Emitter Spawn". */
+	AllScripts
+};
+
 USTRUCT()
 struct FNiagaraModuleDependency
 {
@@ -59,12 +68,24 @@ public:
 	/** Specifies the provided id of the required dependent module (e.g. 'ProvidesNormalizedAge') */
 	UPROPERTY(AssetRegistrySearchable, EditAnywhere, Category = Script)
 	FName Id;
+
 	/** Whether the dependency belongs before or after this module */
 	UPROPERTY(AssetRegistrySearchable, EditAnywhere, Category = Script)
-	ENiagaraModuleDependencyType Type; // e.g. PreDependency,
-										   /** Detailed description of the dependency */
+	ENiagaraModuleDependencyType Type; // e.g. PreDependency
+
+	/** Specifies constraints related to the source script a modules provising a depency. */
+	UPROPERTY(AssetRegistrySearchable, EditAnywhere, Category = Script)
+	ENiagaraModuleDependencyScriptConstraint ScriptConstraint;
+	
+	/** Detailed description of the dependency */
 	UPROPERTY(AssetRegistrySearchable, EditAnywhere, Category = Script, meta = (MultiLine = true))
 	FText Description;
+
+	FNiagaraModuleDependency()
+	{
+		Type = ENiagaraModuleDependencyType::PreDependency;
+		ScriptConstraint = ENiagaraModuleDependencyScriptConstraint::SameScript;
+	}
 };
 
 struct FNiagaraScriptDebuggerInfo
@@ -72,6 +93,8 @@ struct FNiagaraScriptDebuggerInfo
 	FNiagaraScriptDebuggerInfo();
 	FNiagaraScriptDebuggerInfo(FName InName, ENiagaraScriptUsage InUsage, const FGuid& InUsageId);
 
+	bool bWaitForGPU;
+
 	FName HandleName;
 
 	ENiagaraScriptUsage Usage;
@@ -83,6 +106,8 @@ struct FNiagaraScriptDebuggerInfo
 	FNiagaraDataSet Frame;
 
 	FNiagaraParameterStore Parameters;
+
+	TAtomic<bool> bWritten;
 };
 
 
@@ -300,6 +325,15 @@ public:
 	/** Dependencies required by this module from other modules on the stack */
 	UPROPERTY(EditAnywhere, Category = Script)
 	TArray<FNiagaraModuleDependency> RequiredDependencies;
+
+	/* If this script is no longer meant to be used, this option should be set.*/
+	UPROPERTY(EditAnywhere, Category = "Script")
+	uint32 bDeprecated : 1;
+
+	/* Which script to use if this is deprecated.*/
+	UPROPERTY(EditAnywhere, Category = "Script", meta = (EditCondition = "bDeprecated"))
+	UNiagaraScript* DeprecationRecommendation;
+
 #endif
 
 
@@ -319,7 +353,7 @@ public:
 	UPROPERTY(AssetRegistrySearchable, EditAnywhere, Category = Script)
 	FText Keywords;
 
-	UPROPERTY(EditAnywhere, Category = Script)
+	UPROPERTY(EditAnywhere, Category = Script, DisplayName = "Script Metadata", meta = (ToolTip = "Script Metadata"))
 	TMap<FName, FString> ScriptMetaData;
 #endif
 
diff --git a/Engine/Plugins/FX/Niagara/Source/Niagara/Classes/NiagaraScriptExecutionContext.h b/Engine/Plugins/FX/Niagara/Source/Niagara/Classes/NiagaraScriptExecutionContext.h
--- a/Engine/Plugins/FX/Niagara/Source/Niagara/Classes/NiagaraScriptExecutionContext.h
+++ b/Engine/Plugins/FX/Niagara/Source/Niagara/Classes/NiagaraScriptExecutionContext.h
@@ -96,7 +96,18 @@ struct FNiagaraComputeExecutionContext
 		, GPUDataReadback(nullptr)
 		, AccumulatedSpawnRate(0)
 		, NumIndicesPerInstance(0)
+		, PerInstanceData(nullptr)
+		, PerInstanceDataSize(0)
+		, PerInstanceDataInterfaceOffsets(nullptr)
 		, PendingExecutionQueueMask(0)
+#if WITH_EDITORONLY_DATA
+		, GPUDebugDataReadbackFloat(nullptr)
+		, GPUDebugDataReadbackInt(nullptr)
+		, GPUDebugDataReadbackCounts(nullptr)
+		, GPUDebugDataCurrBufferIdx(0xFFFFFFFF)
+		, GPUDebugDataFloatSize(0)
+		, GPUDebugDataIntSize(0)
+#endif	  
 	{
 	}
 
@@ -106,7 +117,26 @@ struct FNiagaraComputeExecutionContext
 		if (GPUDataReadback)
 		{
 			delete GPUDataReadback;
+			GPUDataReadback = nullptr;
 		}
+
+#if WITH_EDITORONLY_DATA
+		if (GPUDebugDataReadbackFloat)
+		{
+			delete GPUDebugDataReadbackFloat;
+			GPUDebugDataReadbackFloat = nullptr;
+		}
+		if (GPUDebugDataReadbackInt)
+		{
+			delete GPUDebugDataReadbackInt;
+			GPUDebugDataReadbackInt = nullptr;
+		}
+		if (GPUDebugDataReadbackCounts)
+		{
+			delete GPUDebugDataReadbackCounts;
+			GPUDebugDataReadbackCounts = nullptr;
+		}
+#endif
 	}
 
 	void Reset()
@@ -120,9 +150,10 @@ struct FNiagaraComputeExecutionContext
 		);
 	}
 
-	void InitParams(UNiagaraScript* InGPUComputeScript, UNiagaraScript* InSpawnScript, UNiagaraScript *InUpdateScript, ENiagaraSimTarget SimTarget)
+	void InitParams(UNiagaraScript* InGPUComputeScript, UNiagaraScript* InSpawnScript, UNiagaraScript *InUpdateScript, ENiagaraSimTarget InSimTarget, const FString& InDebugSimName)
 	{
-		CombinedParamStore.InitFromOwningContext(InGPUComputeScript, SimTarget, true);
+		DebugSimName = InDebugSimName;
+		CombinedParamStore.InitFromOwningContext(InGPUComputeScript, InSimTarget, true);
 
 		GPUScript = InGPUComputeScript;
 		SpawnScript = InSpawnScript;
@@ -189,13 +220,31 @@ private:
 		if (GPUDataReadback)
 		{
 			delete GPUDataReadback;
+			GPUDataReadback = nullptr;
+		}
+
+#if WITH_EDITORONLY_DATA
+		if (GPUDebugDataReadbackFloat)
+		{
+			delete GPUDebugDataReadbackFloat;
+			GPUDebugDataReadbackFloat = nullptr;
 		}
-		GPUDataReadback = nullptr;
+		if (GPUDebugDataReadbackInt)
+		{
+			delete GPUDebugDataReadbackInt;
+			GPUDebugDataReadbackInt = nullptr;
+		}
+		if (GPUDebugDataReadbackCounts)
+		{
+			delete GPUDebugDataReadbackCounts;
+			GPUDebugDataReadbackCounts = nullptr;
+		}
+#endif
 	}
 
 public:
 	const TArray<FNiagaraEventScriptProperties> &GetEventHandlers() const { return EventHandlerScriptProps; }
-
+	FString DebugSimName;
 	class FNiagaraDataSet *MainDataSet;
 	TArray<FNiagaraDataSet*>UpdateEventWriteDataSets;
 	TArray<FNiagaraEventScriptProperties> EventHandlerScriptProps;
@@ -222,6 +271,21 @@ public:
 	uint32 AccumulatedSpawnRate;
 	uint32 NumIndicesPerInstance;	// how many vtx indices per instance the renderer is going to have for its draw call
 
+	void* PerInstanceData; // Data stored on parent system instance
+	uint32 PerInstanceDataSize; // Size of data stored on parent system instance in bytes
+	TMap<TWeakObjectPtr<UNiagaraDataInterface>, int32>* PerInstanceDataInterfaceOffsets;
+
 	/** Ensures we only enqueue each context once per queue before they're dispatched. See SIMULATION_QUEUE_COUNT */
 	uint32 PendingExecutionQueueMask;
+
+
+#if WITH_EDITORONLY_DATA
+	mutable FRHIGPUMemoryReadback *GPUDebugDataReadbackFloat;
+	mutable FRHIGPUMemoryReadback *GPUDebugDataReadbackInt;
+	mutable FRHIGPUMemoryReadback *GPUDebugDataReadbackCounts;
+	mutable int32 GPUDebugDataCurrBufferIdx;
+	mutable uint32 GPUDebugDataFloatSize;
+	mutable uint32 GPUDebugDataIntSize;
+	mutable TSharedPtr<struct FNiagaraScriptDebuggerInfo, ESPMode::ThreadSafe> DebugInfo;
+#endif
 };
diff --git a/Engine/Plugins/FX/Niagara/Source/Niagara/Classes/NiagaraSystem.h b/Engine/Plugins/FX/Niagara/Source/Niagara/Classes/NiagaraSystem.h
--- a/Engine/Plugins/FX/Niagara/Source/Niagara/Classes/NiagaraSystem.h
+++ b/Engine/Plugins/FX/Niagara/Source/Niagara/Classes/NiagaraSystem.h
@@ -10,6 +10,7 @@
 #include "NiagaraEmitterInstance.h"
 #include "NiagaraEmitterHandle.h"
 #include "NiagaraParameterCollection.h"
+#include "NiagaraUserRedirectionParameterStore.h"
 #include "NiagaraSystem.generated.h"
 
 #if WITH_EDITORONLY_DATA
@@ -124,9 +125,8 @@ public:
 	}
 
 	/** From the last compile, what are the variables that were exported out of the system for external use?*/
-	const FNiagaraParameterStore& GetExposedParameters() const {	return ExposedParameters; }
-	FNiagaraParameterStore& GetExposedParameters()  { return ExposedParameters; }
-
+	const FNiagaraUserRedirectionParameterStore& GetExposedParameters() const {	return ExposedParameters; }
+	FNiagaraUserRedirectionParameterStore& GetExposedParameters()  { return ExposedParameters; }
 
 	/** Gets the System script which is used to populate the System parameters and parameter bindings. */
 	UNiagaraScript* GetSystemSpawnScript();
@@ -137,9 +137,6 @@ public:
 	/** Are there any pending compile requests?*/
 	bool HasOutstandingCompilationRequests() const;
 
-	/** Returns whether this system has to be run in solo or not. */
-	bool IsSolo()const;
-
 	FORCEINLINE bool NeedsWarmup()const { return WarmupTickCount > 0 && WarmupTickDelta > SMALL_NUMBER; }
 	FORCEINLINE float GetWarmupTime()const { return WarmupTime; }
 	FORCEINLINE int32 GetWarmupTickCount()const { return WarmupTickCount; }
@@ -231,13 +228,17 @@ public:
 	UPROPERTY(EditAnywhere, Category = "Debug")
 	bool bDumpDebugEmitterInfo;
 
+	bool HasSystemScriptDIsWithPerInstanceData() const;
+
+	const TArray<FName>& GetUserDINamesReadInSystemScripts() const;
+
 private:
 #if WITH_EDITORONLY_DATA
 	INiagaraModule::FMergeEmitterResults MergeChangesForEmitterHandle(FNiagaraEmitterHandle& EmitterHandle);
 	bool QueryCompileComplete(bool bWait, bool bDoPost, bool bDoNotApply = false);
 #endif
 
-	void DetermineIfSolo();
+	void UpdatePostCompileDIInfo();
 protected:
 
 	/** Handles to the emitter this System will simulate. */
@@ -270,7 +271,7 @@ protected:
 
 	/** Variables exposed to the outside work for tweaking*/
 	UPROPERTY()
-	FNiagaraParameterStore ExposedParameters;
+	FNiagaraUserRedirectionParameterStore ExposedParameters;
 
 #if WITH_EDITORONLY_DATA	
 
@@ -299,9 +300,11 @@ protected:
 	UPROPERTY(EditAnywhere, Category = Warmup)
 	float WarmupTickDelta;
 
+	void InitEmitterSpawnAttributes();
 
 	UPROPERTY()
-	uint32 bSolo : 1;
+	bool bHasSystemScriptDIsWithPerInstanceData;
 
-	void InitEmitterSpawnAttributes();
+	UPROPERTY()
+	TArray<FName> UserDINamesReadInSystemScripts;
 };
diff --git a/Engine/Plugins/FX/Niagara/Source/Niagara/Private/NiagaraCollision.cpp b/Engine/Plugins/FX/Niagara/Source/Niagara/Private/NiagaraCollision.cpp
--- a/Engine/Plugins/FX/Niagara/Source/Niagara/Private/NiagaraCollision.cpp
+++ b/Engine/Plugins/FX/Niagara/Source/Niagara/Private/NiagaraCollision.cpp
@@ -181,13 +181,6 @@ void FNiagaraCollisionBatch::GenerateEventsFromResults(FNiagaraEmitterInstance *
 	}
 }
 
-
-
-
-
-
-
-
 int32 FNiagaraDICollisionQueryBatch::SubmitQuery(FVector Position, FVector Direction, float CollisionSize, float DeltaSeconds)
 {
 	SCOPE_CYCLE_COUNTER(STAT_NiagaraCollision);
@@ -229,6 +222,72 @@ int32 FNiagaraDICollisionQueryBatch::SubmitQuery(FVector Position, FVector Direc
 	return Ret;
 }
 
+int32 FNiagaraDICollisionQueryBatch::SubmitQuery(FVector StartPos, FVector EndPos, ECollisionChannel TraceChannel)
+{
+	SCOPE_CYCLE_COUNTER(STAT_NiagaraCollision);
+	if (!CollisionWorld)
+	{
+		return INDEX_NONE;
+	}
+
+	FCollisionQueryParams QueryParams(SCENE_QUERY_STAT(NiagaraAsync));
+	QueryParams.OwnerTag = "Niagara";
+	QueryParams.bFindInitialOverlaps = false;
+	QueryParams.bReturnFaceIndex = false;
+	QueryParams.bReturnPhysicalMaterial = true;
+	QueryParams.bTraceComplex = false;
+	QueryParams.bIgnoreTouches = true;
+	FTraceHandle Handle = CollisionWorld->AsyncLineTraceByChannel(EAsyncTraceType::Single, StartPos, EndPos, TraceChannel, QueryParams, FCollisionResponseParams::DefaultResponseParam, nullptr, TraceID);
+	FNiagaraCollisionTrace Trace;
+	Trace.CollisionTraceHandle = Handle;
+	Trace.SourceParticleIndex = TraceID;
+
+	int32 TraceIdx = CollisionTraces[GetWriteBufferIdx()].Add(Trace);
+	IdToTraceIdx[GetWriteBufferIdx()].Add(TraceID) = TraceIdx;
+
+	int32 Ret = TraceID;
+	TraceID++;
+	return Ret;
+}
+
+bool FNiagaraDICollisionQueryBatch::PerformQuery(FVector StartPos, FVector EndPos, FNiagaraDICollsionQueryResult &Result, ECollisionChannel TraceChannel)
+{
+	SCOPE_CYCLE_COUNTER(STAT_NiagaraCollision);
+	if (!CollisionWorld)
+	{
+		return false;
+	}
+
+	FCollisionQueryParams QueryParams(SCENE_QUERY_STAT(NiagaraSync));
+	QueryParams.OwnerTag = "Niagara";
+	QueryParams.bFindInitialOverlaps = false;
+	QueryParams.bReturnFaceIndex = false;
+	QueryParams.bReturnPhysicalMaterial = true;
+	QueryParams.bTraceComplex = false;
+	QueryParams.bIgnoreTouches = true;
+	FHitResult TraceResult;
+	bool ValidHit = CollisionWorld->LineTraceSingleByChannel(TraceResult, StartPos, EndPos, TraceChannel, QueryParams);
+	if (ValidHit)
+	{
+		Result.IsInsideMesh = TraceResult.bStartPenetrating;
+		Result.CollisionPos = TraceResult.ImpactPoint;
+		Result.CollisionNormal = TraceResult.ImpactNormal;
+		if (TraceResult.PhysMaterial.IsValid())
+		{
+			Result.PhysicalMaterialIdx = TraceResult.PhysMaterial->GetUniqueID();
+			Result.Friction = TraceResult.PhysMaterial->Friction;
+			Result.Restitution = TraceResult.PhysMaterial->Restitution;
+		}
+		else
+		{
+			Result.PhysicalMaterialIdx = -1;
+			Result.Friction = 0.0f;
+			Result.Restitution = 0.0f;
+		}
+	}
+
+	return ValidHit;
+}
 
 bool FNiagaraDICollisionQueryBatch::GetQueryResult(uint32 InTraceID, FNiagaraDICollsionQueryResult &Result)
 {
@@ -262,7 +321,7 @@ bool FNiagaraDICollisionQueryBatch::GetQueryResult(uint32 InTraceID, FNiagaraDIC
 		{
 			// grab the first hit that blocks
 			FHitResult *Hit = FHitResult::GetFirstBlockingHit(CurData.OutHits);
-			if (Hit && Hit->IsValidBlockingHit())
+			if (Hit && Hit->bBlockingHit)
 			{
 				/*
 				FNiagaraCollisionEventPayload Event;
@@ -280,6 +339,7 @@ bool FNiagaraDICollisionQueryBatch::GetQueryResult(uint32 InTraceID, FNiagaraDIC
 
 				Payloads.Add(Event);
 				*/
+				Result.IsInsideMesh = Hit->bStartPenetrating;
 				Result.CollisionPos = Hit->ImpactPoint;// -NormVel*(CurTrace.CollisionSize / 2);
 				Result.CollisionNormal = Hit->ImpactNormal;
 				Result.TraceID = InTraceID;
diff --git a/Engine/Plugins/FX/Niagara/Source/Niagara/Private/NiagaraCommon.cpp b/Engine/Plugins/FX/Niagara/Source/Niagara/Private/NiagaraCommon.cpp
--- a/Engine/Plugins/FX/Niagara/Source/Niagara/Private/NiagaraCommon.cpp
+++ b/Engine/Plugins/FX/Niagara/Source/Niagara/Private/NiagaraCommon.cpp
@@ -273,7 +273,7 @@ void FNiagaraUtilities::CollectScriptDataInterfaceParameters(const UObject& Owne
 
 bool FNiagaraScriptDataInterfaceCompileInfo::CanExecuteOnTarget(ENiagaraSimTarget SimTarget) const
 {
-	check(IsInGameThread());
+	// Note that this can be called on non-game threads. We ensure that the data interface CDO object is already in existence at application init time.
 	UNiagaraDataInterface* Obj = GetDefaultDataInterface();
 	if (Obj)
 	{
@@ -283,28 +283,13 @@ bool FNiagaraScriptDataInterfaceCompileInfo::CanExecuteOnTarget(ENiagaraSimTarge
 	return false;
 }
 
-bool FNiagaraScriptDataInterfaceCompileInfo::IsSystemSolo() const
-{
-	check(IsInGameThread());
-	if (Name.ToString().StartsWith("User."))
-	{
-		return true;
-	}
-
-	UNiagaraDataInterface* Obj = GetDefaultDataInterface();
-	if (Obj && Obj->PerInstanceDataSize() > 0)
-	{
-		return true;
-	}
-	return false;
-}
-
 UNiagaraDataInterface* FNiagaraScriptDataInterfaceCompileInfo::GetDefaultDataInterface() const
 {
-	check(IsInGameThread());
-	UNiagaraDataInterface* Obj = CastChecked<UNiagaraDataInterface>(const_cast<UClass*>(Type.GetClass())->GetDefaultObject(true));
+	// Note that this can be called on non-game threads. We ensure that the data interface CDO object is already in existence at application init time, so we don't allow this to be auto-created.
+	UNiagaraDataInterface* Obj = CastChecked<UNiagaraDataInterface>(const_cast<UClass*>(Type.GetClass())->GetDefaultObject(false));
 	return Obj;
 }
+
 #if WITH_EDITORONLY_DATA
 void FNiagaraUtilities::PrepareRapidIterationParameters(const TArray<UNiagaraScript*>& Scripts, const TMap<UNiagaraScript*, UNiagaraScript*>& ScriptDependencyMap, const TMap<UNiagaraScript*, FString>& ScriptToEmitterNameMap)
 {
@@ -350,7 +335,7 @@ void FNiagaraUtilities::PrepareRapidIterationParameters(const TArray<UNiagaraScr
 		}
 		else
 		{
-			const TMap<FNiagaraVariable, int32>& SourceParameterOffsets = Script->RapidIterationParameters.GetParameterOffests();
+			const TMap<FNiagaraVariable, int32>& SourceParameterOffsets = Script->RapidIterationParameters.GetParameterOffsets();
 			for (auto ParameterOffsetIt = SourceParameterOffsets.CreateConstIterator(); ParameterOffsetIt; ++ParameterOffsetIt)
 			{
 				const FNiagaraVariable& SourceParameter = ParameterOffsetIt.Key();
diff --git a/Engine/Plugins/FX/Niagara/Source/Niagara/Private/NiagaraComponent.cpp b/Engine/Plugins/FX/Niagara/Source/Niagara/Private/NiagaraComponent.cpp
--- a/Engine/Plugins/FX/Niagara/Source/Niagara/Private/NiagaraComponent.cpp
+++ b/Engine/Plugins/FX/Niagara/Source/Niagara/Private/NiagaraComponent.cpp
@@ -495,7 +495,7 @@ bool UNiagaraComponent::InitializeSystem()
 #if WITH_EDITORONLY_DATA
 		OnSystemInstanceChangedDelegate.Broadcast();
 #endif
-		SystemInstance->Init(GetAsset(), bForceSolo);
+		SystemInstance->Init(bForceSolo);
 		return true;
 	}
 	return false;
@@ -1018,6 +1018,7 @@ void UNiagaraComponent::PostLoad()
 	{
 		Asset->ConditionalPostLoad();
 #if WITH_EDITOR
+		PostLoadNormalizeOverrideNames();
 		SynchronizeWithSourceSystem();
 		AssetExposedParametersChangedHandle = Asset->GetExposedParameters().AddOnChangedHandler(
 			FNiagaraParameterStore::FOnChanged::FDelegate::CreateUObject(this, &UNiagaraComponent::AssetExposedParametersChanged));
@@ -1075,19 +1076,22 @@ void UNiagaraComponent::SynchronizeWithSourceSystem()
 	{
 		OverrideParameters.Empty();
 		EditorOverridesValue.Empty();
+#if WITH_EDITORONLY_DATA
+		OnSynchronizedWithAssetParametersDelegate.Broadcast();
+#endif
 		return;
 	}
 
 	TArray<FNiagaraVariable> SourceVars;
 	Asset->GetExposedParameters().GetParameters(SourceVars);
-
 	for (FNiagaraVariable& Param : SourceVars)
 	{
 		OverrideParameters.AddParameter(Param, true);
 	}
 
 	TArray<FNiagaraVariable> ExistingVars;
-	OverrideParameters.GetParameters(ExistingVars);
+	OverrideParameters.GetUserParameters(ExistingVars);
+	Asset->GetExposedParameters().GetUserParameters(SourceVars);
 
 	for (FNiagaraVariable ExistingVar : ExistingVars)
 	{
@@ -1118,6 +1122,7 @@ void UNiagaraComponent::SynchronizeWithSourceSystem()
 void UNiagaraComponent::AssetExposedParametersChanged()
 {
 	SynchronizeWithSourceSystem();
+	ReinitializeSystem();
 }
 #endif
 
@@ -1174,6 +1179,19 @@ void UNiagaraComponent::SetMaxSimTime(float InMaxTime)
 }
 
 #if WITH_EDITOR
+
+void UNiagaraComponent::PostLoadNormalizeOverrideNames()
+{
+	TMap<FName, bool> ValueMap;
+	for (TPair<FName, bool> Pair : EditorOverridesValue)
+	{
+		bool IsOldUserParam = Pair.Key.ToString().StartsWith(TEXT("User."));
+		FName ValueName = IsOldUserParam ? (*Pair.Key.ToString().RightChop(5)) : Pair.Key;
+		ValueMap.Add(ValueName, Pair.Value);
+	}
+	EditorOverridesValue = ValueMap;
+}
+
 bool UNiagaraComponent::IsParameterValueOverriddenLocally(const FName& InParamName)
 {
 	bool* FoundVar = EditorOverridesValue.Find(InParamName);
@@ -1203,9 +1221,9 @@ void UNiagaraComponent::SetParameterValueOverriddenLocally(const FNiagaraVariabl
 		Asset->GetExposedParameters().CopyParameterData(OverrideParameters, InParam);
 	}
 	
-	if (bRequiresSystemInstanceReset)
+	if (bRequiresSystemInstanceReset && SystemInstance)
 	{
-		SystemInstance->Reset(FNiagaraSystemInstance::EResetMode::ResetSystem, true);
+		SystemInstance->Reset(FNiagaraSystemInstance::EResetMode::ResetAll, true);
 	}
 	
 }
@@ -1229,8 +1247,15 @@ void UNiagaraComponent::SetAsset(UNiagaraSystem* InAsset)
 
 #if WITH_EDITOR
 		SynchronizeWithSourceSystem();
-		AssetExposedParametersChangedHandle = Asset->GetExposedParameters().AddOnChangedHandler(
-			FNiagaraParameterStore::FOnChanged::FDelegate::CreateUObject(this, &UNiagaraComponent::AssetExposedParametersChanged));
+		if (Asset != nullptr)
+		{
+			AssetExposedParametersChangedHandle = Asset->GetExposedParameters().AddOnChangedHandler(
+				FNiagaraParameterStore::FOnChanged::FDelegate::CreateUObject(this, &UNiagaraComponent::AssetExposedParametersChanged));
+		}
+		else
+		{
+			AssetExposedParametersChangedHandle.Reset();
+		}
 #endif
 
 		//Force a reinit.
diff --git a/Engine/Plugins/FX/Niagara/Source/Niagara/Private/NiagaraConstants.cpp b/Engine/Plugins/FX/Niagara/Source/Niagara/Private/NiagaraConstants.cpp
--- a/Engine/Plugins/FX/Niagara/Source/Niagara/Private/NiagaraConstants.cpp
+++ b/Engine/Plugins/FX/Niagara/Source/Niagara/Private/NiagaraConstants.cpp
@@ -35,6 +35,8 @@ void FNiagaraConstants::Init()
 		SystemParameters.Add(SYS_PARAM_ENGINE_Y_AXIS);
 		SystemParameters.Add(SYS_PARAM_ENGINE_Z_AXIS);
 
+		SystemParameters.Add(SYS_PARAM_ENGINE_ROTATION);
+
 		SystemParameters.Add(SYS_PARAM_ENGINE_LOCAL_TO_WORLD);
 		SystemParameters.Add(SYS_PARAM_ENGINE_WORLD_TO_LOCAL);
 		SystemParameters.Add(SYS_PARAM_ENGINE_LOCAL_TO_WORLD_TRANSPOSED);
@@ -49,14 +51,18 @@ void FNiagaraConstants::Init()
 
 		SystemParameters.Add(SYS_PARAM_ENGINE_EXEC_COUNT);
 		SystemParameters.Add(SYS_PARAM_ENGINE_EMITTER_NUM_PARTICLES);
+		SystemParameters.Add(SYS_PARAM_ENGINE_EMITTER_TOTAL_SPAWNED_PARTICLES);
 		SystemParameters.Add(SYS_PARAM_ENGINE_SYSTEM_NUM_EMITTERS_ALIVE);
 		SystemParameters.Add(SYS_PARAM_ENGINE_SYSTEM_NUM_EMITTERS);
 		SystemParameters.Add(SYS_PARAM_ENGINE_NUM_SYSTEM_INSTANCES);
 		SystemParameters.Add(SYS_PARAM_ENGINE_GLOBAL_SPAWN_COUNT_SCALE);
 		SystemParameters.Add(SYS_PARAM_ENGINE_GLOBAL_SYSTEM_COUNT_SCALE);
 		SystemParameters.Add(SYS_PARAM_ENGINE_SYSTEM_AGE);
+		SystemParameters.Add(SYS_PARAM_ENGINE_SYSTEM_TICK_COUNT);
 		SystemParameters.Add(SYS_PARAM_EMITTER_AGE);
 		SystemParameters.Add(SYS_PARAM_EMITTER_LOCALSPACE);
+		SystemParameters.Add(SYS_PARAM_EMITTER_DETERMINISM);
+		SystemParameters.Add(SYS_PARAM_EMITTER_RANDOM_SEED);
 		SystemParameters.Add(SYS_PARAM_EMITTER_SPAWN_GROUP);
 	}
 
@@ -74,6 +80,7 @@ void FNiagaraConstants::Init()
 		UpdatedSystemParameters.Add(FName(TEXT("System X Axis")), SYS_PARAM_ENGINE_X_AXIS);
 		UpdatedSystemParameters.Add(FName(TEXT("System Y Axis")), SYS_PARAM_ENGINE_Y_AXIS);
 		UpdatedSystemParameters.Add(FName(TEXT("System Z Axis")), SYS_PARAM_ENGINE_Z_AXIS);
+		UpdatedSystemParameters.Add(FName(TEXT("System Rotation")), SYS_PARAM_ENGINE_ROTATION);
 
 		UpdatedSystemParameters.Add(FName(TEXT("System Local To World")), SYS_PARAM_ENGINE_LOCAL_TO_WORLD);
 		UpdatedSystemParameters.Add(FName(TEXT("System World To Local")), SYS_PARAM_ENGINE_WORLD_TO_LOCAL);
@@ -93,11 +100,14 @@ void FNiagaraConstants::Init()
 		UpdatedSystemParameters.Add(FName(TEXT("Delta Time")), SYS_PARAM_ENGINE_DELTA_TIME);
 		UpdatedSystemParameters.Add(FName(TEXT("Emitter Age")), SYS_PARAM_EMITTER_AGE);
 		UpdatedSystemParameters.Add(FName(TEXT("Emitter Local Space")), SYS_PARAM_EMITTER_LOCALSPACE);
+		UpdatedSystemParameters.Add(FName(TEXT("Emitter Random Seed")), SYS_PARAM_EMITTER_RANDOM_SEED);
+		UpdatedSystemParameters.Add(FName(TEXT("Emitter Determinism")), SYS_PARAM_EMITTER_DETERMINISM);
 		UpdatedSystemParameters.Add(FName(TEXT("Effect Position")), SYS_PARAM_ENGINE_POSITION);
 		UpdatedSystemParameters.Add(FName(TEXT("Effect Velocity")), SYS_PARAM_ENGINE_VELOCITY);
 		UpdatedSystemParameters.Add(FName(TEXT("Effect X Axis")), SYS_PARAM_ENGINE_X_AXIS);
 		UpdatedSystemParameters.Add(FName(TEXT("Effect Y Axis")), SYS_PARAM_ENGINE_Y_AXIS);
 		UpdatedSystemParameters.Add(FName(TEXT("Effect Z Axis")), SYS_PARAM_ENGINE_Z_AXIS);
+		UpdatedSystemParameters.Add(FName(TEXT("Effect Rotation")), SYS_PARAM_ENGINE_ROTATION);
 
 		UpdatedSystemParameters.Add(FName(TEXT("Effect Local To World")), SYS_PARAM_ENGINE_LOCAL_TO_WORLD);
 		UpdatedSystemParameters.Add(FName(TEXT("Effect World To Local")), SYS_PARAM_ENGINE_WORLD_TO_LOCAL);
@@ -118,6 +128,7 @@ void FNiagaraConstants::Init()
 		SystemStrMap.Add(SYS_PARAM_ENGINE_TIME, LOCTEXT("EngineTimeDesc", "Time in seconds since level began play, but IS paused when the game is paused, and IS dilated/clamped."));
 		SystemStrMap.Add(SYS_PARAM_ENGINE_REAL_TIME, LOCTEXT("EngineRealTimeDesc", "Time in seconds since level began play, but IS NOT paused when the game is paused, and IS NOT dilated/clamped."));
 		SystemStrMap.Add(SYS_PARAM_ENGINE_SYSTEM_AGE, LOCTEXT("EngineSystemTimeDesc", "Time in seconds since the system was first created. Managed by the NiagaraSystemInstance in code."));
+		SystemStrMap.Add(SYS_PARAM_ENGINE_SYSTEM_TICK_COUNT, LOCTEXT("EngineSystemTickCount", "The current tick of this system simulation."));
 
 		SystemStrMap.Add(SYS_PARAM_ENGINE_POSITION, LOCTEXT("EnginePositionDesc", "The owning component's position in world space."));
 		SystemStrMap.Add(SYS_PARAM_ENGINE_SCALE, LOCTEXT("EngineScaleDesc", "The owning component's scale in world space."));
@@ -125,6 +136,7 @@ void FNiagaraConstants::Init()
 		SystemStrMap.Add(SYS_PARAM_ENGINE_X_AXIS, LOCTEXT("XAxisDesc", "The X-axis of the owning component."));
 		SystemStrMap.Add(SYS_PARAM_ENGINE_Y_AXIS, LOCTEXT("YAxisDesc", "The Y-axis of the owning component."));
 		SystemStrMap.Add(SYS_PARAM_ENGINE_Z_AXIS, LOCTEXT("ZAxisDesc", "The Z-axis of the owning component."));
+		SystemStrMap.Add(SYS_PARAM_ENGINE_ROTATION, LOCTEXT("EngineRotationDesc", "The owning component's rotation in world space."));
 
 		SystemStrMap.Add(SYS_PARAM_ENGINE_LOCAL_TO_WORLD, LOCTEXT("LocalToWorldDesc", "Owning component's local space to world space transform matrix."));
 		SystemStrMap.Add(SYS_PARAM_ENGINE_WORLD_TO_LOCAL, LOCTEXT("WorldToLocalDesc", "Owning component's world space to local space transform matrix."));
@@ -139,6 +151,7 @@ void FNiagaraConstants::Init()
 
 		SystemStrMap.Add(SYS_PARAM_ENGINE_EXEC_COUNT, LOCTEXT("ExecCountDesc", "The index of this particle in the read buffer."));
 		SystemStrMap.Add(SYS_PARAM_ENGINE_EMITTER_NUM_PARTICLES, LOCTEXT("EmitterNumParticles", "The number of particles for this emitter at the beginning of simulation. Should only be used in Emitter scripts."));
+		SystemStrMap.Add(SYS_PARAM_ENGINE_EMITTER_TOTAL_SPAWNED_PARTICLES, LOCTEXT("EmitterTotalSpawnedParticles", "The total number of particles spawned for this emitter at the beginning of this simulation. Should only be used by the particle spawn script the assign unique IDs."));
 		SystemStrMap.Add(SYS_PARAM_ENGINE_SYSTEM_NUM_EMITTERS_ALIVE, LOCTEXT("SystemNumEmittersAlive", "The number of emitters still alive attached to this system. Should only be used in System scripts."));
 		SystemStrMap.Add(SYS_PARAM_ENGINE_SYSTEM_NUM_EMITTERS, LOCTEXT("SystemNumEmitters", "The number of emitters attached to this system. Should only be used in System scripts."));
 		SystemStrMap.Add(SYS_PARAM_ENGINE_NUM_SYSTEM_INSTANCES, LOCTEXT("SystemNumInstances", "The number of instances of the this system currently ticking. Should only be used in System scripts."));
@@ -148,6 +161,7 @@ void FNiagaraConstants::Init()
 
 	if (Attributes.Num() == 0)
 	{
+		Attributes.Add(SYS_PARAM_PARTICLES_UNIQUE_ID);
 		Attributes.Add(SYS_PARAM_PARTICLES_ID);
 		Attributes.Add(SYS_PARAM_PARTICLES_POSITION);
 		Attributes.Add(SYS_PARAM_PARTICLES_VELOCITY);
@@ -201,6 +215,8 @@ void FNiagaraConstants::Init()
 		AttrDataSetKeyMap.Add(SYS_PARAM_PARTICLES_UV_SCALE, GetAttributeAsDataSetKey(SYS_PARAM_PARTICLES_UV_SCALE));
 		AttrDataSetKeyMap.Add(SYS_PARAM_PARTICLES_MATERIAL_RANDOM, GetAttributeAsDataSetKey(SYS_PARAM_PARTICLES_MATERIAL_RANDOM));
 		AttrDataSetKeyMap.Add(SYS_PARAM_PARTICLES_LIGHT_RADIUS, GetAttributeAsDataSetKey(SYS_PARAM_PARTICLES_LIGHT_RADIUS));
+		AttrDataSetKeyMap.Add(SYS_PARAM_PARTICLES_LIGHT_EXPONENT, GetAttributeAsDataSetKey(SYS_PARAM_PARTICLES_LIGHT_EXPONENT));
+		AttrDataSetKeyMap.Add(SYS_PARAM_PARTICLES_LIGHT_VOLUMETRIC_SCATTERING, GetAttributeAsDataSetKey(SYS_PARAM_PARTICLES_LIGHT_VOLUMETRIC_SCATTERING));
 		AttrDataSetKeyMap.Add(SYS_PARAM_PARTICLES_RIBBONID, GetAttributeAsDataSetKey(SYS_PARAM_PARTICLES_RIBBONID));
 		AttrDataSetKeyMap.Add(SYS_PARAM_PARTICLES_RIBBONWIDTH, GetAttributeAsDataSetKey(SYS_PARAM_PARTICLES_RIBBONWIDTH));
 		AttrDataSetKeyMap.Add(SYS_PARAM_PARTICLES_RIBBONTWIST, GetAttributeAsDataSetKey(SYS_PARAM_PARTICLES_RIBBONTWIST));
@@ -309,6 +325,16 @@ void FNiagaraConstants::Init()
 		Var.SetValue<float>(100.0f);
 		AttrDefaultsValueMap.Add(SYS_PARAM_PARTICLES_LIGHT_RADIUS, Var);
 
+		AttrDefaultsStrMap.Add(SYS_PARAM_PARTICLES_LIGHT_EXPONENT, TEXT("10.0"));
+		Var = SYS_PARAM_PARTICLES_LIGHT_EXPONENT;
+		Var.SetValue<float>(10.0f);
+		AttrDefaultsValueMap.Add(SYS_PARAM_PARTICLES_LIGHT_EXPONENT, Var);
+
+		AttrDefaultsStrMap.Add(SYS_PARAM_PARTICLES_LIGHT_VOLUMETRIC_SCATTERING, TEXT("1.0"));
+		Var = SYS_PARAM_PARTICLES_LIGHT_VOLUMETRIC_SCATTERING;
+		Var.SetValue<float>(1.0f);
+		AttrDefaultsValueMap.Add(SYS_PARAM_PARTICLES_LIGHT_VOLUMETRIC_SCATTERING, Var);
+
 		AttrDefaultsStrMap.Add(SYS_PARAM_PARTICLES_RIBBONID, TEXT("0"));
 		Var = SYS_PARAM_PARTICLES_RIBBONID;
 		Var.SetValue<FNiagaraID>(FNiagaraID());
@@ -358,6 +384,8 @@ void FNiagaraConstants::Init()
 		AttrDescStrMap.Add(SYS_PARAM_PARTICLES_UV_SCALE, LOCTEXT("UVScalerParamDesc", "Used to multiply the generated UVs for Sprite renderers."));
 		AttrDescStrMap.Add(SYS_PARAM_PARTICLES_MATERIAL_RANDOM, LOCTEXT("MaterialRandomParamDesc", "Used to drive the Particle Random node in the Material Editor. Without this set, any Particle Randoms will get 0.0."));
 		AttrDescStrMap.Add(SYS_PARAM_PARTICLES_LIGHT_RADIUS, LOCTEXT("LightRadiusParamDesc", "Used to drive the radius of the light when using a Light renderer."));
+		AttrDescStrMap.Add(SYS_PARAM_PARTICLES_LIGHT_EXPONENT, LOCTEXT("LightExponentParamDesc", "Used to drive the attenuation of the light when using a Light renderer without inverse squared falloff enabled."));
+		AttrDescStrMap.Add(SYS_PARAM_PARTICLES_LIGHT_VOLUMETRIC_SCATTERING, LOCTEXT("LightVolumetricScatteringParamDesc", "Used to drive the volumetric scattering intensity of the light when using a Light renderer."));
 		AttrDescStrMap.Add(SYS_PARAM_INSTANCE_ALIVE, LOCTEXT("AliveParamDesc", "Used to determine whether or not this particle instance is still valid or if it can be deleted."));
 		AttrDescStrMap.Add(SYS_PARAM_PARTICLES_RIBBONID, LOCTEXT("RibbonIDDesc", "Sets the ribbon id for a particle. Particles with the same ribbon id will be connected into a ribbon."));
 		AttrDescStrMap.Add(SYS_PARAM_PARTICLES_RIBBONWIDTH, LOCTEXT("RibbonWidthDesc", "Sets the ribbon width for a particle, in UE4 units."));
@@ -365,6 +393,7 @@ void FNiagaraConstants::Init()
 		AttrDescStrMap.Add(SYS_PARAM_PARTICLES_RIBBONFACING, LOCTEXT("RibbonFacingDesc", "Sets the facing vector of the ribbon at the particle position."));
 		AttrDescStrMap.Add(SYS_PARAM_PARTICLES_RIBBONLINKORDER, LOCTEXT("RibbonLinkOrderDesc", "Explicit order for linking particles within a ribbon. Particles of the same ribbon id will be connected into a ribbon in incrementing order of this attribute value."));
 		AttrDescStrMap.Add(SYS_PARAM_PARTICLES_ID, LOCTEXT("IDDesc", "Engine managed particle attribute that is a persistent ID for each particle."));
+		AttrDescStrMap.Add(SYS_PARAM_PARTICLES_UNIQUE_ID, LOCTEXT("UniqueIDDesc", "Engine managed particle attribute that is a unique ID for each particle. The ID is incremented for each new particle spawned.")); 
 	}
 
 	if (AttrMetaData.Num() == 0)
@@ -394,8 +423,11 @@ void FNiagaraConstants::Init()
 	if (EngineManagedAttributes.Num() == 0)
 	{
 		EngineManagedAttributes.Add(SYS_PARAM_PARTICLES_ID);
+		// NOTE(mv): UniqueID needs to be distinct from ID, as the latter is not guaranteed to be contiguous and will reuse labels
+		//           It is unique and sequential, never resetting until the simulation is reset. 
+		//           It needs to be engine managed, otherwise the scripts cannot write to it when it isn't referenced in any scripts. 
+		EngineManagedAttributes.Add(SYS_PARAM_PARTICLES_UNIQUE_ID); 
 	}
-
 }
 
 const TArray<FNiagaraVariable>& FNiagaraConstants::GetEngineConstants()
diff --git a/Engine/Plugins/FX/Niagara/Source/Niagara/Private/NiagaraDataInterface.cpp b/Engine/Plugins/FX/Niagara/Source/Niagara/Private/NiagaraDataInterface.cpp
--- a/Engine/Plugins/FX/Niagara/Source/Niagara/Private/NiagaraDataInterface.cpp
+++ b/Engine/Plugins/FX/Niagara/Source/Niagara/Private/NiagaraDataInterface.cpp
@@ -187,7 +187,7 @@ struct FNiagaraDataInterfaceParametersCS_Curve : public FNiagaraDataInterfacePar
 		Ar << CurveLUT;
 	}
 
-	virtual void Set(FRHICommandList& RHICmdList, FNiagaraShader* Shader, class UNiagaraDataInterface* DataInterface) const override
+	virtual void Set(FRHICommandList& RHICmdList, FNiagaraShader* Shader, class UNiagaraDataInterface* DataInterface, void* PerInstanceData) const override
 	{
 		check(IsInRenderingThread());
 
