Merging //UE4/Dev-Main to Dev-Build (//UE4/Dev-Build)

diff --git a/Engine/Plugins/FX/Niagara/Source/Niagara/Private/NiagaraDataInterfaceCollisionQuery.cpp b/Engine/Plugins/FX/Niagara/Source/Niagara/Private/NiagaraDataInterfaceCollisionQuery.cpp
--- a/Engine/Plugins/FX/Niagara/Source/Niagara/Private/NiagaraDataInterfaceCollisionQuery.cpp
+++ b/Engine/Plugins/FX/Niagara/Source/Niagara/Private/NiagaraDataInterfaceCollisionQuery.cpp
@@ -5,6 +5,7 @@
 #include "NiagaraCustomVersion.h"
 #include "NiagaraWorldManager.h"
 #include "ShaderParameterUtils.h"
+#include "GlobalDistanceFieldParameters.h"
 #include "Shader.h"
 
 //////////////////////////////////////////////////////////////////////////
@@ -15,6 +16,7 @@ FCriticalSection UNiagaraDataInterfaceCollisionQuery::CriticalSection;
 UNiagaraDataInterfaceCollisionQuery::UNiagaraDataInterfaceCollisionQuery(FObjectInitializer const& ObjectInitializer)
 	: Super(ObjectInitializer)
 {
+	TraceChannelEnum = FindObject<UEnum>(ANY_PACKAGE, TEXT("ECollisionChannel"), true);
 }
 
 bool UNiagaraDataInterfaceCollisionQuery::InitPerInstanceData(void* PerInstanceData, FNiagaraSystemInstance* InSystemInstance)
@@ -47,35 +49,6 @@ void UNiagaraDataInterfaceCollisionQuery::PostLoad()
 	const int32 NiagaraVer = GetLinkerCustomVersion(FNiagaraCustomVersion::GUID);
 }
 
-#if WITH_EDITOR
-
-void UNiagaraDataInterfaceCollisionQuery::PostEditChangeProperty(struct FPropertyChangedEvent& PropertyChangedEvent)
-{
-	Super::PostEditChangeProperty(PropertyChangedEvent);
-}
-
-#endif
-
-
-bool UNiagaraDataInterfaceCollisionQuery::CopyToInternal(UNiagaraDataInterface* Destination) const
-{
-	if (!Super::CopyToInternal(Destination))
-	{
-		return false;
-	}
-	return true;
-}
-
-bool UNiagaraDataInterfaceCollisionQuery::Equals(const UNiagaraDataInterface* Other) const
-{
-	if (!Super::Equals(Other))
-	{
-		return false;
-	}
-	return true;
-}
-
-
 void UNiagaraDataInterfaceCollisionQuery::GetFunctions(TArray<FNiagaraFunctionSignature>& OutFunctions)
 {
 	FNiagaraFunctionSignature Sig3;
@@ -121,6 +94,78 @@ void UNiagaraDataInterfaceCollisionQuery::GetFunctions(TArray<FNiagaraFunctionSi
 	Sig2.Outputs.Add(FNiagaraVariable(FNiagaraTypeDefinition::GetVec3Def(), TEXT("CollisionVelocity")));
 	Sig2.Outputs.Add(FNiagaraVariable(FNiagaraTypeDefinition::GetVec3Def(), TEXT("CollisionNormal")));
 	OutFunctions.Add(Sig2);
+
+	FNiagaraFunctionSignature SigGpu;
+	SigGpu.Name = TEXT("PerformCollisionQueryGPUShader");
+	SigGpu.bMemberFunction = true;
+	SigGpu.bRequiresContext = false;
+	SigGpu.Inputs.Add(FNiagaraVariable(FNiagaraTypeDefinition(GetClass()), TEXT("CollisionQuery")));
+	SigGpu.Inputs.Add(FNiagaraVariable(FNiagaraTypeDefinition::GetVec3Def(), TEXT("DepthSamplePosWorld")));
+	SigGpu.Inputs.Add(FNiagaraVariable(FNiagaraTypeDefinition::GetVec3Def(), TEXT("TraceEndWorld")));
+	SigGpu.Inputs.Add(FNiagaraVariable(FNiagaraTypeDefinition::GetFloatDef(), TEXT("SceneDepthBounds")));
+	SigGpu.Inputs.Add(FNiagaraVariable(FNiagaraTypeDefinition::GetFloatDef(), TEXT("ParticleRadius")));
+	SigGpu.Inputs.Add(FNiagaraVariable(FNiagaraTypeDefinition::GetBoolDef(), TEXT("UseMeshDistanceField")));
+	SigGpu.Outputs.Add(FNiagaraVariable(FNiagaraTypeDefinition::GetBoolDef(), TEXT("CollisionValid")));
+	SigGpu.Outputs.Add(FNiagaraVariable(FNiagaraTypeDefinition::GetVec3Def(), TEXT("CollisionPosWorld")));
+	SigGpu.Outputs.Add(FNiagaraVariable(FNiagaraTypeDefinition::GetVec3Def(), TEXT("CollisionNormal")));
+	OutFunctions.Add(SigGpu);
+
+	FNiagaraFunctionSignature SigDepth;
+	SigDepth.Name = TEXT("QuerySceneDepthGPU");
+	SigDepth.bMemberFunction = true;
+	SigDepth.bRequiresContext = false;
+	SigDepth.Inputs.Add(FNiagaraVariable(FNiagaraTypeDefinition(GetClass()), TEXT("CollisionQuery")));
+	SigDepth.Inputs.Add(FNiagaraVariable(FNiagaraTypeDefinition::GetVec3Def(), TEXT("DepthSamplePosWorld")));
+	SigDepth.Outputs.Add(FNiagaraVariable(FNiagaraTypeDefinition::GetFloatDef(), TEXT("SceneDepth")));
+	SigDepth.Outputs.Add(FNiagaraVariable(FNiagaraTypeDefinition::GetVec3Def(), TEXT("CameraPosWorld")));
+	SigDepth.Outputs.Add(FNiagaraVariable(FNiagaraTypeDefinition::GetBoolDef(), TEXT("IsInsideView")));
+	SigDepth.Outputs.Add(FNiagaraVariable(FNiagaraTypeDefinition::GetVec3Def(), TEXT("SamplePosWorld")));
+	SigDepth.Outputs.Add(FNiagaraVariable(FNiagaraTypeDefinition::GetVec3Def(), TEXT("SampleWorldNormal")));
+	OutFunctions.Add(SigDepth);
+
+	FNiagaraFunctionSignature SigMeshField;
+	SigMeshField.Name = TEXT("QueryMeshDistanceFieldGPU");
+	SigMeshField.bMemberFunction = true;
+	SigMeshField.bRequiresContext = false;
+	SigMeshField.Inputs.Add(FNiagaraVariable(FNiagaraTypeDefinition(GetClass()), TEXT("CollisionQuery")));
+	SigMeshField.Inputs.Add(FNiagaraVariable(FNiagaraTypeDefinition::GetVec3Def(), TEXT("FieldSamplePosWorld")));
+	SigMeshField.Outputs.Add(FNiagaraVariable(FNiagaraTypeDefinition::GetFloatDef(), TEXT("DistanceToNearestSurface")));
+	SigMeshField.Outputs.Add(FNiagaraVariable(FNiagaraTypeDefinition::GetVec3Def(), TEXT("FieldGradient")));
+	OutFunctions.Add(SigMeshField);
+
+	FNiagaraFunctionSignature SigCpuSync;
+	SigCpuSync.Name = TEXT("PerformCollisionQuerySyncCPU");
+	SigCpuSync.bMemberFunction = true;
+	SigCpuSync.bRequiresContext = false;
+	SigCpuSync.Inputs.Add(FNiagaraVariable(FNiagaraTypeDefinition(GetClass()), TEXT("CollisionQuery")));
+	SigCpuSync.Inputs.Add(FNiagaraVariable(FNiagaraTypeDefinition::GetVec3Def(), TEXT("TraceStartWorld")));
+	SigCpuSync.Inputs.Add(FNiagaraVariable(FNiagaraTypeDefinition::GetVec3Def(), TEXT("TraceEndWorld")));
+	SigCpuSync.Inputs.Add(FNiagaraVariable(FNiagaraTypeDefinition(TraceChannelEnum), TEXT("TraceChannel")));
+	SigCpuSync.Outputs.Add(FNiagaraVariable(FNiagaraTypeDefinition::GetBoolDef(), TEXT("CollisionValid")));
+	SigCpuSync.Outputs.Add(FNiagaraVariable(FNiagaraTypeDefinition::GetBoolDef(), TEXT("IsTraceInsideMesh")));
+	SigCpuSync.Outputs.Add(FNiagaraVariable(FNiagaraTypeDefinition::GetVec3Def(), TEXT("CollisionPosWorld")));
+	SigCpuSync.Outputs.Add(FNiagaraVariable(FNiagaraTypeDefinition::GetVec3Def(), TEXT("CollisionNormal")));
+	SigCpuSync.Outputs.Add(FNiagaraVariable(FNiagaraTypeDefinition::GetFloatDef(), TEXT("CollisionMaterialFriction")));
+	SigCpuSync.Outputs.Add(FNiagaraVariable(FNiagaraTypeDefinition::GetFloatDef(), TEXT("CollisionMaterialRestitution")));
+	OutFunctions.Add(SigCpuSync);
+
+	FNiagaraFunctionSignature SigCpuAsync;
+	SigCpuAsync.Name = TEXT("PerformCollisionQueryAsyncCPU");
+	SigCpuAsync.bMemberFunction = true;
+	SigCpuAsync.bRequiresContext = false;
+	SigCpuAsync.Inputs.Add(FNiagaraVariable(FNiagaraTypeDefinition(GetClass()), TEXT("CollisionQuery")));
+	SigCpuAsync.Inputs.Add(FNiagaraVariable(FNiagaraTypeDefinition::GetIntDef(), TEXT("PreviousFrameQueryID")));
+	SigCpuAsync.Inputs.Add(FNiagaraVariable(FNiagaraTypeDefinition::GetVec3Def(), TEXT("TraceStartWorld")));
+	SigCpuAsync.Inputs.Add(FNiagaraVariable(FNiagaraTypeDefinition::GetVec3Def(), TEXT("TraceEndWorld")));
+	SigCpuAsync.Inputs.Add(FNiagaraVariable(FNiagaraTypeDefinition(TraceChannelEnum), TEXT("TraceChannel")));
+	SigCpuAsync.Outputs.Add(FNiagaraVariable(FNiagaraTypeDefinition::GetIntDef(), TEXT("NextFrameQueryID")));
+	SigCpuAsync.Outputs.Add(FNiagaraVariable(FNiagaraTypeDefinition::GetBoolDef(), TEXT("CollisionValid")));
+	SigCpuAsync.Outputs.Add(FNiagaraVariable(FNiagaraTypeDefinition::GetBoolDef(), TEXT("IsTraceInsideMesh")));
+	SigCpuAsync.Outputs.Add(FNiagaraVariable(FNiagaraTypeDefinition::GetVec3Def(), TEXT("CollisionPosWorld")));
+	SigCpuAsync.Outputs.Add(FNiagaraVariable(FNiagaraTypeDefinition::GetVec3Def(), TEXT("CollisionNormal")));
+	SigCpuAsync.Outputs.Add(FNiagaraVariable(FNiagaraTypeDefinition::GetFloatDef(), TEXT("CollisionMaterialFriction")));
+	SigCpuAsync.Outputs.Add(FNiagaraVariable(FNiagaraTypeDefinition::GetFloatDef(), TEXT("CollisionMaterialRestitution")));
+	OutFunctions.Add(SigCpuAsync);
 }
 
 // build the shader function HLSL; function name is passed in, as it's defined per-DI; that way, configuration could change
@@ -134,9 +179,7 @@ bool UNiagaraDataInterfaceCollisionQuery::GetFunctionHLSL(const FName& Definitio
 	// assuming it'll usually be the same thread trying to call ReadQuery for a particular QueryID, that submitted it in the first place.
 	if (DefinitionFunctionName == TEXT("PerformCollisionQuery"))
 	{
-		OutHLSL += TEXT("float3 WorldPositionFromSceneDepth(float2 ScreenPosition, float SceneDepth)\n{\n\tfloat4 HomogeneousWorldPosition = mul(float4(ScreenPosition * SceneDepth, SceneDepth, 1), View.ScreenToWorld);\n\treturn HomogeneousWorldPosition.xyz / HomogeneousWorldPosition.w;\n}\n");
-		OutHLSL += TEXT("\n");
-		OutHLSL += TEXT("void ") + InstanceFunctionName + TEXT("(in int InQueryID, in float3 In_ParticlePos, in float3 In_ParticleVel, in float In_DeltaSeconds, float CollisionRadius, in float CollisionDepthBounds, int In_InstanceData, \
+		OutHLSL += TEXT("void ") + InstanceFunctionName + TEXT("(in int InQueryID, in float3 In_ParticlePos, in float3 In_ParticleVel, in float In_DeltaSeconds, float CollisionRadius, in float CollisionDepthBounds, \
 			out int Out_QueryID, out bool OutCollisionValid, out float3 Out_CollisionPos, out float3 Out_CollisionNormal, out float Out_Friction, out float Out_Restitution) \n{\n");
 		// get the screen position
 		OutHLSL += TEXT("\
@@ -185,13 +228,110 @@ bool UNiagaraDataInterfaceCollisionQuery::GetFunctionHLSL(const FName& Definitio
 					OutCollisionValid = false; \n\
 				}\n\
 			}\n\
-		}\n\
-		\n\
-		\n");
-
-		OutHLSL += TEXT("}\n");
+		}\n}\n\n");
 	}
+	else if (DefinitionFunctionName == TEXT("PerformCollisionQueryGPUShader"))
+	{
+		FString SceneDepthFunction = InstanceFunctionName + TEXT("_SceneDepthCollision");
+		FString DistanceFieldFunction = InstanceFunctionName + TEXT("_DistanceFieldCollision");
 
+		OutHLSL += TEXT("void ") + SceneDepthFunction + TEXT("(in float3 In_SamplePos, in float3 In_TraceEndPos, in float CollisionDepthBounds, in float ParticleRadius, out bool OutCollisionValid, out float3 Out_CollisionPos, out float3 Out_CollisionNormal) \n{\n\
+		OutCollisionValid = false;\n\
+		Out_CollisionPos = In_SamplePos;\n\
+		Out_CollisionNormal = float3(0.0, 0.0, 1.0);\n\
+		float4 SamplePosition = float4(In_SamplePos + View.PreViewTranslation, 1); \n\
+		float4 ClipPosition = mul(SamplePosition, View.TranslatedWorldToClip);\n\
+		float2 ScreenPosition = ClipPosition.xy / ClipPosition.w;\n\
+		// Don't try to collide if the particle falls outside the view.\n\
+		if (all(abs(ScreenPosition.xy) <= float2(1, 1)))\n\
+		{\n\
+			// Sample the depth buffer to get a world position near the particle.\n\
+			float2 ScreenUV = ScreenPosition * View.ScreenPositionScaleBias.xy + View.ScreenPositionScaleBias.wz;\n\
+			float SceneDepth = CalcSceneDepth(ScreenUV);\n\
+			if (abs(ClipPosition.w - SceneDepth) < CollisionDepthBounds)\n\
+			{\n\
+				// Reconstruct world position.\n\
+				float3 WorldPosition = WorldPositionFromSceneDepth(ScreenPosition.xy, SceneDepth);\n\
+				// Sample the normal buffer to create a plane to collide against.\n\
+				float3 WorldNormal = Texture2DSampleLevel(SceneTexturesStruct.GBufferATexture, SceneTexturesStruct.GBufferATextureSampler, ScreenUV, 0).xyz * 2.0 - 1.0;\n\
+				float4 CollisionPlane = float4(WorldNormal, dot(WorldPosition.xyz, WorldNormal));\n\
+				// Compute the portion of velocity normal to the collision plane.\n\
+				float VelocityDot = dot(CollisionPlane.xyz, (In_TraceEndPos - In_SamplePos).xyz);\n\
+				float d_back = (dot(CollisionPlane.xyz, In_SamplePos.xyz) + ParticleRadius - CollisionPlane.w);\n\
+				float d_front = (dot(CollisionPlane.xyz, In_TraceEndPos.xyz) - ParticleRadius - CollisionPlane.w);\n\
+				// distance to the plane from current and predicted position\n\
+				if (d_back >= 0.0f && d_front <= 0.0f && VelocityDot < 0.0f)\n\
+				{\n\
+					OutCollisionValid = true;\n\
+					Out_CollisionPos = In_SamplePos + (WorldNormal*d_back);\n\
+					Out_CollisionNormal = WorldNormal;\n\
+				}\n\
+			}\n\
+		}\
+		\n}\n\n");
+		OutHLSL += TEXT("void ") + DistanceFieldFunction + TEXT("(in float3 InPosition, in float3 In_TraceEndPos, out bool OutCollisionValid, out float3 Out_CollisionPos, out float3 Out_CollisionNormal)\n{\n\
+		float DistanceToNearestSurface = GetDistanceToNearestSurfaceGlobal(InPosition);\n\
+		if (DistanceToNearestSurface < length(In_TraceEndPos - InPosition))\n\
+		{\n\
+			OutCollisionValid = true;\n\
+			Out_CollisionNormal = normalize(GetDistanceFieldGradientGlobal(InPosition));\n\
+			Out_CollisionPos = InPosition - Out_CollisionNormal * DistanceToNearestSurface;\n\
+		}\n\
+		else\n\
+		{\n\
+			OutCollisionValid = false;\n\
+			Out_CollisionNormal = float3(0.0, 0.0, 1.0);\n\
+			Out_CollisionPos = InPosition;\n\
+		}\n}\n\n");
+		OutHLSL += TEXT("void ") + InstanceFunctionName + TEXT("(in float3 In_SamplePos, in float3 In_TraceEndPos, in float CollisionDepthBounds, ") +
+			TEXT("in float ParticleRadius, in bool UseMeshDistanceField, out bool OutCollisionValid, out float3 Out_CollisionPos, out float3 Out_CollisionNormal) \n{\n");
+		OutHLSL += TEXT("\
+			if (UseMeshDistanceField)\n\
+			{\n\
+				") + DistanceFieldFunction + TEXT("(In_SamplePos, In_TraceEndPos, OutCollisionValid, Out_CollisionPos, Out_CollisionNormal);\n\
+			}\n\
+			else\n\
+			{\n\
+				") + SceneDepthFunction + TEXT("(In_SamplePos, In_TraceEndPos, CollisionDepthBounds, ParticleRadius, OutCollisionValid, Out_CollisionPos, Out_CollisionNormal);\n\
+			}\n}\n\n");
+	}
+	else if (DefinitionFunctionName == TEXT("QuerySceneDepthGPU"))
+	{
+		OutHLSL += TEXT("void ") + InstanceFunctionName + TEXT("(in float3 In_SamplePos, out float Out_SceneDepth, out float3 Out_CameraPosWorld, out bool Out_IsInsideView, out float3 Out_WorldPos, out float3 Out_WorldNormal) \n{\n");
+		OutHLSL += TEXT("\
+			Out_SceneDepth = -1;\n\
+			Out_WorldPos = float3(0.0, 0.0, 0.0);\n\
+			Out_WorldNormal = float3(0.0, 0.0, 1.0);\n\
+			Out_IsInsideView = true;\n\
+			Out_CameraPosWorld.xyz = View.WorldCameraOrigin.xyz;\n\
+			float4 SamplePosition = float4(In_SamplePos + View.PreViewTranslation, 1);\n\
+			float4 ClipPosition = mul(SamplePosition, View.TranslatedWorldToClip);\n\
+			float2 ScreenPosition = ClipPosition.xy / ClipPosition.w;\n\
+			// Check if the sample is inside the view.\n\
+			if (all(abs(ScreenPosition.xy) <= float2(1, 1)))\n\
+			{\n\
+				// Sample the depth buffer to get a world position near the sample position.\n\
+				float2 ScreenUV = ScreenPosition * View.ScreenPositionScaleBias.xy + View.ScreenPositionScaleBias.wz;\n\
+				float SceneDepth = CalcSceneDepth(ScreenUV);\n\
+				Out_SceneDepth = SceneDepth;\n\
+				// Reconstruct world position.\n\
+				Out_WorldPos = WorldPositionFromSceneDepth(ScreenPosition.xy, SceneDepth);\n\
+				// Sample the normal buffer\n\
+				Out_WorldNormal = Texture2DSampleLevel(SceneTexturesStruct.GBufferATexture, SceneTexturesStruct.GBufferATextureSampler, ScreenUV, 0).xyz * 2.0 - 1.0;\n\
+			}\n\
+			else\n\
+			{\n\
+				Out_IsInsideView = false;\n\
+			}\n}\n\n");
+	}
+	else if (DefinitionFunctionName == TEXT("QueryMeshDistanceFieldGPU"))
+	{
+		OutHLSL += TEXT("void ") + InstanceFunctionName + TEXT("(in float3 In_SamplePos, out float Out_DistanceToNearestSurface, out float3 Out_FieldGradient) \n{\n");
+		OutHLSL += TEXT("\
+			Out_DistanceToNearestSurface = GetDistanceToNearestSurfaceGlobal(In_SamplePos);\n\
+			Out_FieldGradient = GetDistanceFieldGradientGlobal(In_SamplePos);\
+			\n}\n\n");
+	}
 
 	return true;
 }
@@ -204,6 +344,11 @@ void UNiagaraDataInterfaceCollisionQuery::GetParameterDefinitionHLSL(FNiagaraDat
 DEFINE_NDI_DIRECT_FUNC_BINDER(UNiagaraDataInterfaceCollisionQuery, SubmitQuery);
 DEFINE_NDI_DIRECT_FUNC_BINDER(UNiagaraDataInterfaceCollisionQuery, ReadQuery);
 DEFINE_NDI_DIRECT_FUNC_BINDER(UNiagaraDataInterfaceCollisionQuery, PerformQuery);
+DEFINE_NDI_DIRECT_FUNC_BINDER(UNiagaraDataInterfaceCollisionQuery, PerformQuerySyncCPU);
+DEFINE_NDI_DIRECT_FUNC_BINDER(UNiagaraDataInterfaceCollisionQuery, PerformQueryAsyncCPU);
+DEFINE_NDI_DIRECT_FUNC_BINDER(UNiagaraDataInterfaceCollisionQuery, PerformQueryGPU);
+DEFINE_NDI_DIRECT_FUNC_BINDER(UNiagaraDataInterfaceCollisionQuery, QuerySceneDepth);
+DEFINE_NDI_DIRECT_FUNC_BINDER(UNiagaraDataInterfaceCollisionQuery, QueryMeshDistanceField);
 
 void UNiagaraDataInterfaceCollisionQuery::GetVMExternalFunction(const FVMExternalFunctionBindingInfo& BindingInfo, void* InstanceData, FVMExternalFunction &OutFunc)
 {
@@ -220,6 +365,26 @@ void UNiagaraDataInterfaceCollisionQuery::GetVMExternalFunction(const FVMExterna
 	{
 		NDI_FUNC_BINDER(UNiagaraDataInterfaceCollisionQuery, PerformQuery)::Bind(this, OutFunc);
 	}
+	else if (BindingInfo.Name == TEXT("PerformCollisionQuerySyncCPU"))
+	{
+		NDI_FUNC_BINDER(UNiagaraDataInterfaceCollisionQuery, PerformQuerySyncCPU)::Bind(this, OutFunc);
+	}
+	else if (BindingInfo.Name == TEXT("PerformCollisionQueryAsyncCPU"))
+	{
+		NDI_FUNC_BINDER(UNiagaraDataInterfaceCollisionQuery, PerformQueryAsyncCPU)::Bind(this, OutFunc);
+	}
+	else if (BindingInfo.Name == TEXT("PerformCollisionQueryGPUShader"))
+	{
+		NDI_FUNC_BINDER(UNiagaraDataInterfaceCollisionQuery, PerformQueryGPU)::Bind(this, OutFunc);
+	}
+	else if (BindingInfo.Name == TEXT("QuerySceneDepthGPU"))
+	{
+		NDI_FUNC_BINDER(UNiagaraDataInterfaceCollisionQuery, QuerySceneDepth)::Bind(this, OutFunc);
+	}
+	else if (BindingInfo.Name == TEXT("QueryMeshDistanceFieldGPU"))
+	{
+		NDI_FUNC_BINDER(UNiagaraDataInterfaceCollisionQuery, QueryMeshDistanceField)::Bind(this, OutFunc);
+	}
 	else
 	{
 		UE_LOG(LogNiagara, Error, TEXT("Could not find data interface external function. %s\n"),
@@ -326,6 +491,240 @@ void UNiagaraDataInterfaceCollisionQuery::PerformQuery(FVectorVMContext& Context
 
 }
 
+void UNiagaraDataInterfaceCollisionQuery::PerformQuerySyncCPU(FVectorVMContext & Context)
+{
+	VectorVM::FExternalFuncInputHandler<float> StartPosParamX(Context);
+	VectorVM::FExternalFuncInputHandler<float> StartPosParamY(Context);
+	VectorVM::FExternalFuncInputHandler<float> StartPosParamZ(Context);
+
+	VectorVM::FExternalFuncInputHandler<float> EndPosParamX(Context);
+	VectorVM::FExternalFuncInputHandler<float> EndPosParamY(Context);
+	VectorVM::FExternalFuncInputHandler<float> EndPosParamZ(Context);
+
+	VectorVM::FExternalFuncInputHandler<ECollisionChannel> TraceChannelParam(Context);
+
+	VectorVM::FUserPtrHandler<CQDIPerInstanceData> InstanceData(Context);
+
+	VectorVM::FExternalFuncRegisterHandler<int32> OutQueryValid(Context);
+	VectorVM::FExternalFuncRegisterHandler<int32> OutInsideMesh(Context);
+	VectorVM::FExternalFuncRegisterHandler<float> OutCollisionPosX(Context);
+	VectorVM::FExternalFuncRegisterHandler<float> OutCollisionPosY(Context);
+	VectorVM::FExternalFuncRegisterHandler<float> OutCollisionPosZ(Context);
+	VectorVM::FExternalFuncRegisterHandler<float> OutCollisionNormX(Context);
+	VectorVM::FExternalFuncRegisterHandler<float> OutCollisionNormY(Context);
+	VectorVM::FExternalFuncRegisterHandler<float> OutCollisionNormZ(Context);
+	VectorVM::FExternalFuncRegisterHandler<float> OutFriction(Context);
+	VectorVM::FExternalFuncRegisterHandler<float> OutRestitution(Context);
+
+	FScopeLock ScopeLock(&CriticalSection);
+	for (int32 i = 0; i < Context.NumInstances; ++i)
+	{
+		FVector Pos(StartPosParamX.GetAndAdvance(), StartPosParamY.GetAndAdvance(), StartPosParamZ.GetAndAdvance());
+		FVector Dir(EndPosParamX.GetAndAdvance(), EndPosParamY.GetAndAdvance(), EndPosParamZ.GetAndAdvance());
+		ECollisionChannel TraceChannel = TraceChannelParam.GetAndAdvance();
+		ensure(!Pos.ContainsNaN());
+		FNiagaraDICollsionQueryResult Res;
+		bool Valid = InstanceData->CollisionBatch.PerformQuery(Pos, Dir, Res, TraceChannel);
+		if (Valid)
+		{
+			*OutQueryValid.GetDestAndAdvance() = 0xFFFFFFFF; //->SetValue(true);
+			*OutInsideMesh.GetDestAndAdvance() = Res.IsInsideMesh ? 0xFFFFFFFF : 0;
+			*OutCollisionPosX.GetDestAndAdvance() = Res.CollisionPos.X;
+			*OutCollisionPosY.GetDestAndAdvance() = Res.CollisionPos.Y;
+			*OutCollisionPosZ.GetDestAndAdvance() = Res.CollisionPos.Z;
+			*OutCollisionNormX.GetDestAndAdvance() = Res.CollisionNormal.X;
+			*OutCollisionNormY.GetDestAndAdvance() = Res.CollisionNormal.Y;
+			*OutCollisionNormZ.GetDestAndAdvance() = Res.CollisionNormal.Z;
+			*OutFriction.GetDestAndAdvance() = Res.Friction;
+			*OutRestitution.GetDestAndAdvance() = Res.Restitution;
+		}
+		else
+		{
+			*OutQueryValid.GetDestAndAdvance() = 0; //->SetValue(false);
+			*OutInsideMesh.GetDestAndAdvance() = 0;
+			*OutCollisionPosX.GetDestAndAdvance() = 0.0f;
+			*OutCollisionPosY.GetDestAndAdvance() = 0.0f;
+			*OutCollisionPosZ.GetDestAndAdvance() = 0.0f;
+			*OutCollisionNormX.GetDestAndAdvance() = 0.0f;
+			*OutCollisionNormY.GetDestAndAdvance() = 0.0f;
+			*OutCollisionNormZ.GetDestAndAdvance() = 0.0f;
+			*OutFriction.GetDestAndAdvance() = 0.0f;
+			*OutRestitution.GetDestAndAdvance() = 0.0f;
+		}
+	}
+}
+
+void UNiagaraDataInterfaceCollisionQuery::PerformQueryAsyncCPU(FVectorVMContext & Context)
+{
+	VectorVM::FExternalFuncInputHandler<int32> InIDParam(Context);
+	VectorVM::FExternalFuncInputHandler<float> StartPosParamX(Context);
+	VectorVM::FExternalFuncInputHandler<float> StartPosParamY(Context);
+	VectorVM::FExternalFuncInputHandler<float> StartPosParamZ(Context);
+
+	VectorVM::FExternalFuncInputHandler<float> EndPosParamX(Context);
+	VectorVM::FExternalFuncInputHandler<float> EndPosParamY(Context);
+	VectorVM::FExternalFuncInputHandler<float> EndPosParamZ(Context);
+
+	VectorVM::FExternalFuncInputHandler<ECollisionChannel> TraceChannelParam(Context);
+
+	VectorVM::FUserPtrHandler<CQDIPerInstanceData> InstanceData(Context);
+
+	VectorVM::FExternalFuncRegisterHandler<int32> OutQueryID(Context);
+
+	VectorVM::FExternalFuncRegisterHandler<int32> OutQueryValid(Context);
+	VectorVM::FExternalFuncRegisterHandler<int32> OutInsideMesh(Context);
+	VectorVM::FExternalFuncRegisterHandler<float> OutCollisionPosX(Context);
+	VectorVM::FExternalFuncRegisterHandler<float> OutCollisionPosY(Context);
+	VectorVM::FExternalFuncRegisterHandler<float> OutCollisionPosZ(Context);
+	VectorVM::FExternalFuncRegisterHandler<float> OutCollisionNormX(Context);
+	VectorVM::FExternalFuncRegisterHandler<float> OutCollisionNormY(Context);
+	VectorVM::FExternalFuncRegisterHandler<float> OutCollisionNormZ(Context);
+	VectorVM::FExternalFuncRegisterHandler<float> OutFriction(Context);
+	VectorVM::FExternalFuncRegisterHandler<float> OutRestitution(Context);
+
+	FScopeLock ScopeLock(&CriticalSection);
+	for (int32 i = 0; i < Context.NumInstances; ++i)
+	{
+		FVector Pos(StartPosParamX.GetAndAdvance(), StartPosParamY.GetAndAdvance(), StartPosParamZ.GetAndAdvance());
+		FVector Dir(EndPosParamX.GetAndAdvance(), EndPosParamY.GetAndAdvance(), EndPosParamZ.GetAndAdvance());
+		ECollisionChannel TraceChannel = TraceChannelParam.GetAndAdvance();
+		ensure(!Pos.ContainsNaN());
+		*OutQueryID.GetDestAndAdvance() = InstanceData->CollisionBatch.SubmitQuery(Pos, Dir, TraceChannel);
+
+		// try to retrieve a query with the supplied query ID
+		FNiagaraDICollsionQueryResult Res;
+		int32 ID = InIDParam.GetAndAdvance();
+		bool Valid = InstanceData->CollisionBatch.GetQueryResult(ID, Res);
+		if (Valid)
+		{
+			*OutQueryValid.GetDestAndAdvance() = 0xFFFFFFFF; //->SetValue(true);
+			*OutInsideMesh.GetDestAndAdvance() = Res.IsInsideMesh ? 0xFFFFFFFF : 0;
+			*OutCollisionPosX.GetDestAndAdvance() = Res.CollisionPos.X;
+			*OutCollisionPosY.GetDestAndAdvance() = Res.CollisionPos.Y;
+			*OutCollisionPosZ.GetDestAndAdvance() = Res.CollisionPos.Z;
+			*OutCollisionNormX.GetDestAndAdvance() = Res.CollisionNormal.X;
+			*OutCollisionNormY.GetDestAndAdvance() = Res.CollisionNormal.Y;
+			*OutCollisionNormZ.GetDestAndAdvance() = Res.CollisionNormal.Z;
+			*OutFriction.GetDestAndAdvance() = Res.Friction;
+			*OutRestitution.GetDestAndAdvance() = Res.Restitution;
+		}
+		else
+		{
+			*OutQueryValid.GetDestAndAdvance() = 0; //->SetValue(false);
+			*OutInsideMesh.GetDestAndAdvance() = 0;
+			*OutCollisionPosX.GetDestAndAdvance() = 0.0f;
+			*OutCollisionPosY.GetDestAndAdvance() = 0.0f;
+			*OutCollisionPosZ.GetDestAndAdvance() = 0.0f;
+			*OutCollisionNormX.GetDestAndAdvance() = 0.0f;
+			*OutCollisionNormY.GetDestAndAdvance() = 0.0f;
+			*OutCollisionNormZ.GetDestAndAdvance() = 0.0f;
+			*OutFriction.GetDestAndAdvance() = 0.0f;
+			*OutRestitution.GetDestAndAdvance() = 0.0f;
+		}
+	}
+}
+
+
+void UNiagaraDataInterfaceCollisionQuery::PerformQueryGPU(FVectorVMContext& Context)
+{
+	UE_LOG(LogNiagara, Error, TEXT("GPU only function 'PerformQueryGPU' called on CPU VM, check your module code to fix."));
+
+	VectorVM::FExternalFuncInputHandler<float> StartPosParamX(Context);
+	VectorVM::FExternalFuncInputHandler<float> StartPosParamY(Context);
+	VectorVM::FExternalFuncInputHandler<float> StartPosParamZ(Context);
+	VectorVM::FExternalFuncInputHandler<float> EndPosParamX(Context);
+	VectorVM::FExternalFuncInputHandler<float> EndPosParamY(Context);
+	VectorVM::FExternalFuncInputHandler<float> EndPosParamZ(Context);
+	VectorVM::FExternalFuncInputHandler<float> SceneDepthBoundsParam(Context);
+	VectorVM::FExternalFuncInputHandler<float> ParticleRadiusParam(Context);
+	VectorVM::FExternalFuncInputHandler<int32> UseDistanceFieldParam(Context);
+
+	VectorVM::FUserPtrHandler<CQDIPerInstanceData> InstanceData(Context);
+
+	VectorVM::FExternalFuncRegisterHandler<int32> OutQueryValid(Context);
+	VectorVM::FExternalFuncRegisterHandler<float> OutCollisionPosX(Context);
+	VectorVM::FExternalFuncRegisterHandler<float> OutCollisionPosY(Context);
+	VectorVM::FExternalFuncRegisterHandler<float> OutCollisionPosZ(Context);
+	VectorVM::FExternalFuncRegisterHandler<float> OutCollisionNormX(Context);
+	VectorVM::FExternalFuncRegisterHandler<float> OutCollisionNormY(Context);
+	VectorVM::FExternalFuncRegisterHandler<float> OutCollisionNormZ(Context);
+
+	FScopeLock ScopeLock(&CriticalSection);
+	for (int32 i = 0; i < Context.NumInstances; ++i)
+	{
+		*OutQueryValid.GetDestAndAdvance() = 0;// ->SetValue(false);
+		*OutCollisionPosX.GetDestAndAdvance() = 0.0f;
+		*OutCollisionPosY.GetDestAndAdvance() = 0.0f;
+		*OutCollisionPosZ.GetDestAndAdvance() = 0.0f;
+		*OutCollisionNormX.GetDestAndAdvance() = 0.0f;
+		*OutCollisionNormY.GetDestAndAdvance() = 0.0f;
+		*OutCollisionNormZ.GetDestAndAdvance() = 1.0f;
+	}
+}
+
+void UNiagaraDataInterfaceCollisionQuery::QuerySceneDepth(FVectorVMContext & Context)
+{
+	UE_LOG(LogNiagara, Error, TEXT("GPU only function 'QuerySceneDepthGPU' called on CPU VM, check your module code to fix."));
+
+	VectorVM::FExternalFuncInputHandler<float> SamplePosParamX(Context);
+	VectorVM::FExternalFuncInputHandler<float> SamplePosParamY(Context);
+	VectorVM::FExternalFuncInputHandler<float> SamplePosParamZ(Context);
+	
+	VectorVM::FUserPtrHandler<CQDIPerInstanceData> InstanceData(Context);
+
+	VectorVM::FExternalFuncRegisterHandler<float> OutSceneDepth(Context);
+	VectorVM::FExternalFuncRegisterHandler<float> OutCameraPosX(Context);
+	VectorVM::FExternalFuncRegisterHandler<float> OutCameraPosY(Context);
+	VectorVM::FExternalFuncRegisterHandler<float> OutCameraPosZ(Context);
+	VectorVM::FExternalFuncRegisterHandler<int32> OutIsInsideView(Context);
+	VectorVM::FExternalFuncRegisterHandler<float> OutWorldPosX(Context);
+	VectorVM::FExternalFuncRegisterHandler<float> OutWorldPosY(Context);
+	VectorVM::FExternalFuncRegisterHandler<float> OutWorldPosZ(Context);
+	VectorVM::FExternalFuncRegisterHandler<float> OutWorldNormX(Context);
+	VectorVM::FExternalFuncRegisterHandler<float> OutWorldNormY(Context);
+	VectorVM::FExternalFuncRegisterHandler<float> OutWorldNormZ(Context);
+
+	FScopeLock ScopeLock(&CriticalSection);
+	for (int32 i = 0; i < Context.NumInstances; ++i)
+	{
+		*OutSceneDepth.GetDestAndAdvance() = -1;
+		*OutIsInsideView.GetDestAndAdvance() = 0;
+		*OutWorldPosX.GetDestAndAdvance() = 0.0f;
+		*OutWorldPosY.GetDestAndAdvance() = 0.0f;
+		*OutWorldPosZ.GetDestAndAdvance() = 0.0f;
+		*OutWorldNormX.GetDestAndAdvance() = 0.0f;
+		*OutWorldNormY.GetDestAndAdvance() = 0.0f;
+		*OutWorldNormZ.GetDestAndAdvance() = 1.0f;
+		*OutCameraPosX.GetDestAndAdvance() = 0.0f;
+		*OutCameraPosY.GetDestAndAdvance() = 0.0f;
+		*OutCameraPosZ.GetDestAndAdvance() = 0.0f;
+	}
+}
+
+void UNiagaraDataInterfaceCollisionQuery::QueryMeshDistanceField(FVectorVMContext& Context)
+{
+	UE_LOG(LogNiagara, Error, TEXT("GPU only function 'QueryMeshDistanceFieldGPU' called on CPU VM, check your module code to fix."));
+
+	VectorVM::FExternalFuncInputHandler<float> SamplePosParamX(Context);
+	VectorVM::FExternalFuncInputHandler<float> SamplePosParamY(Context);
+	VectorVM::FExternalFuncInputHandler<float> SamplePosParamZ(Context);
+
+	VectorVM::FUserPtrHandler<CQDIPerInstanceData> InstanceData(Context);
+
+	VectorVM::FExternalFuncRegisterHandler<float> OutSurfaceDistance(Context);
+	VectorVM::FExternalFuncRegisterHandler<float> OutFieldGradientX(Context);
+	VectorVM::FExternalFuncRegisterHandler<float> OutFieldGradientY(Context);
+	VectorVM::FExternalFuncRegisterHandler<float> OutFieldGradientZ(Context);
+
+	FScopeLock ScopeLock(&CriticalSection);
+	for (int32 i = 0; i < Context.NumInstances; ++i)
+	{
+		*OutSurfaceDistance.GetDestAndAdvance() = -1;
+		*OutFieldGradientX.GetDestAndAdvance() = 0.0f;
+		*OutFieldGradientY.GetDestAndAdvance() = 0.0f;
+		*OutFieldGradientZ.GetDestAndAdvance() = 1.0f;
+	}
+}
 
 void UNiagaraDataInterfaceCollisionQuery::SubmitQuery(FVectorVMContext& Context)
 {
@@ -438,31 +837,44 @@ struct FNiagaraDataInterfaceParametersCS_CollisionQuery : public FNiagaraDataInt
 	virtual void Bind(const FNiagaraDataInterfaceParamRef& ParamRef, const class FShaderParameterMap& ParameterMap) override
 	{
 		PassUniformBuffer.Bind(ParameterMap, FSceneTexturesUniformParameters::StaticStructMetadata.GetShaderVariableName());
-		check(PassUniformBuffer.IsBound());
+		
+		GlobalDistanceFieldParameters.Bind(ParameterMap);
+		if (GlobalDistanceFieldParameters.IsBound())
+		{
+			GNiagaraViewDataManager.SetGlobalDistanceFieldUsage();
+		}
 	}
 
-	virtual void Serialize(FArchive& Ar)override
+	virtual void Serialize(FArchive& Ar) override
 	{
 		Ar << PassUniformBuffer;
+		Ar << GlobalDistanceFieldParameters;
 	}
 
-	virtual void Set(FRHICommandList& RHICmdList, FNiagaraShader* Shader, class UNiagaraDataInterface* DataInterface) const override
+	virtual void Set(FRHICommandList& RHICmdList, FNiagaraShader* Shader, class UNiagaraDataInterface* DataInterface, void* PerInstanceData) const override
 	{
 		check(IsInRenderingThread());
 
 		const FComputeShaderRHIParamRef ComputeShaderRHI = Shader->GetComputeShader();
-
+		
 		TUniformBufferRef<FSceneTexturesUniformParameters> SceneTextureUniformParams = GNiagaraViewDataManager.GetSceneTextureUniformParameters();
 		SetUniformBufferParameter(RHICmdList, ComputeShaderRHI, PassUniformBuffer/*Shader->GetUniformBufferParameter(SceneTexturesUniformBufferStruct)*/, SceneTextureUniformParams);
+		if (GlobalDistanceFieldParameters.IsBound())
+		{
+			GNiagaraViewDataManager.SetGlobalDistanceFieldUsage();
+			GlobalDistanceFieldParameters.Set(RHICmdList, ComputeShaderRHI, *GNiagaraViewDataManager.GetGlobalDistanceFieldParameters());
+		}		
 	}
 
 private:
 
 	/** The SceneDepthTexture parameter for depth buffer collision. */
 	FShaderUniformBufferParameter PassUniformBuffer;
+
+	FGlobalDistanceFieldParameters GlobalDistanceFieldParameters;
 };
 
-FNiagaraDataInterfaceParametersCS* UNiagaraDataInterfaceCollisionQuery::ConstructComputeParameters()const
+FNiagaraDataInterfaceParametersCS* UNiagaraDataInterfaceCollisionQuery::ConstructComputeParameters() const
 {
 	return new FNiagaraDataInterfaceParametersCS_CollisionQuery();
 }
\ No newline at end of file
diff --git a/Engine/Plugins/FX/Niagara/Source/Niagara/Private/NiagaraDataInterfaceCurlNoise.cpp b/Engine/Plugins/FX/Niagara/Source/Niagara/Private/NiagaraDataInterfaceCurlNoise.cpp
--- a/Engine/Plugins/FX/Niagara/Source/Niagara/Private/NiagaraDataInterfaceCurlNoise.cpp
+++ b/Engine/Plugins/FX/Niagara/Source/Niagara/Private/NiagaraDataInterfaceCurlNoise.cpp
@@ -3,59 +3,348 @@
 #include "NiagaraDataInterfaceCurlNoise.h"
 #include "NiagaraShader.h"
 #include "ShaderParameterUtils.h"
+#include "Math/IntVector.h"
 
-const FName UNiagaraDataInterfaceCurlNoise::SampleNoiseFieldName(TEXT("SampleNoiseField"));
-const FString UNiagaraDataInterfaceCurlNoise::CurlNoiseBufferName(TEXT("CurlNoiseBuffer_"));
+static const FName SampleNoiseFieldName(TEXT("SampleNoiseField"));
+static const FString OffsetFromSeedBaseName(TEXT("OffsetFromSeed_"));
 
-UNiagaraDataInterfaceCurlNoise::UNiagaraDataInterfaceCurlNoise(FObjectInitializer const& ObjectInitializer)
-	: Super(ObjectInitializer)
-	, bGPUBufferDirty(true)
-	, Seed(0)
-	, GPUBuffer(new FRWBuffer)
+// ----------------------------------------------------- start noise helpers -----------------------------------------------------
+// Fairly straightforward CPU implementation of the equivalent HLSL code found in Engine/Shaders/Private/Random.ush
+// Also contains a partial implementation of some math types and functions found in HLSL but not in UE
+
+// TODO(mv): Move this to a separate NiagaraMath unit.  
+struct FNiagaraUIntVector
+{
+	uint32 X;
+	uint32 Y;
+	uint32 Z;
+
+	FNiagaraUIntVector() : X(0), Y(0), Z(0) {}
+	FNiagaraUIntVector(uint32 X, uint32 Y, uint32 Z) : X(X), Y(Y), Z(Z) {}
+	FNiagaraUIntVector(uint32 Value) : X(Value), Y(Value), Z(Value) {}
+	FNiagaraUIntVector(FIntVector Values) : X(Values.X), Y(Values.Y), Z(Values.Z) {}
+
+	FNiagaraUIntVector operator+(FNiagaraUIntVector Rhs)
+	{
+		return FNiagaraUIntVector{ X + Rhs.X, Y + Rhs.Y, Z + Rhs.Z };
+	}
+
+	FNiagaraUIntVector operator*(FNiagaraUIntVector Rhs)
+	{
+		return FNiagaraUIntVector{ X * Rhs.X, Y * Rhs.Y, Z * Rhs.Z };
+	}
+
+	FNiagaraUIntVector operator>>(uint32 Shift)
+	{
+		return FNiagaraUIntVector(X >> Shift, Y >> Shift, Z >> Shift);
+	}
+
+	FNiagaraUIntVector operator&(FNiagaraUIntVector Rhs)
+	{
+		return FNiagaraUIntVector{ X & Rhs.X, Y & Rhs.Y, Z & Rhs.Z };
+	}
+
+	uint32& operator[](int Index)
+	{
+		if (Index == 0) return X;
+		if (Index == 1) return Y;
+		if (Index == 2) return Z;
+		check(false); 
+		return X;
+	}
+
+	const uint32& operator[](int Index) const
+	{
+		if (Index == 0) return X;
+		if (Index == 1) return Y;
+		if (Index == 2) return Z;
+		check(false);
+		return X;
+	}
+};
+
+FNiagaraUIntVector Rand3DPCG16(FIntVector p)
 {
+	FNiagaraUIntVector v = FNiagaraUIntVector(p);
+	
+	v = v * 1664525u + 1013904223u;
+
+	v.X += v.Y*v.Z;
+	v.Y += v.Z*v.X;
+	v.Z += v.X*v.Y;
+	v.X += v.Y*v.Z;
+	v.Y += v.Z*v.X;
+	v.Z += v.X*v.Y;
+
+	return v >> 16u;
+}
 
+FVector NiagaraVectorFrac(FVector v)
+{
+	return FVector(FMath::Frac(v.X), FMath::Frac(v.Y), FMath::Frac(v.Z));
 }
 
-void UNiagaraDataInterfaceCurlNoise::PostInitProperties()
+FVector NoiseTileWrap(FVector v, bool bTiling, float RepeatSize)
 {
-	Super::PostInitProperties();
+	FVector vv = bTiling ? (NiagaraVectorFrac(v / RepeatSize) * RepeatSize) : v;
+	return vv;
+}
 
-	if (HasAnyFlags(RF_ClassDefaultObject))
+struct FNiagaraMatrix4x3
+{
+	FVector Row0;
+	FVector Row1;
+	FVector Row2;
+	FVector Row3;
+
+	FNiagaraMatrix4x3() : Row0(FVector()), Row1(FVector()), Row2(FVector()), Row3(FVector()) {}
+	FNiagaraMatrix4x3(FVector Row0, FVector Row1, FVector Row2, FVector Row3) : Row0(Row0), Row1(Row1), Row2(Row2), Row3(Row3) {}
+
+	FVector& operator[](int Row)
 	{
-		FNiagaraTypeRegistry::Register(FNiagaraTypeDefinition(GetClass()), true, false, false);
+		if (Row == 0) return Row0;
+		if (Row == 1) return Row1;
+		if (Row == 2) return Row2;
+		if (Row == 3) return Row3;
+		check(false);
+		return Row0;
 	}
+	const FVector& operator[](int Row) const
+	{
+		if (Row == 0) return Row0;
+		if (Row == 1) return Row1;
+		if (Row == 2) return Row2;
+		if (Row == 3) return Row3;
+		check(false);
+		return Row0;
+	}
+};
+
+FVector NiagaraVectorFloor(FVector v) {
+	return FVector(FGenericPlatformMath::FloorToFloat(v.X),
+		           FGenericPlatformMath::FloorToFloat(v.Y),
+		           FGenericPlatformMath::FloorToFloat(v.Z));
+};
+
+FVector NiagaraVectorStep(FVector v, FVector u)
+{
+	return FVector(u.X >= v.X ? 1 : 0, 
+		           u.Y >= v.Y ? 1 : 0, 
+		           u.Z >= v.Z ? 1 : 0);
 }
 
-void UNiagaraDataInterfaceCurlNoise::PostLoad()
+FVector NiagaraVectorSwizzle(FVector v, uint32 x, uint32 y, uint32 z)
 {
-	Super::PostLoad();
-	InitNoiseLUT();
+	return FVector(v[x], v[y], v[z]);
+}
+
+FVector NiagaraVectorMin(FVector u, FVector v)
+{
+	return FVector(u.X < v.X ? u.X : v.X, 
+		           u.Y < v.Y ? u.Y : v.Y, 
+		           u.Z < v.Z ? u.Z : v.Z);
+}
+
+FVector NiagaraVectorMax(FVector u, FVector v)
+{
+	return FVector(u.X > v.X ? u.X : v.X,
+		           u.Y > v.Y ? u.Y : v.Y,
+	           	   u.Z > v.Z ? u.Z : v.Z);
+} 
+
+FNiagaraMatrix4x3 SimplexCorners(FVector v)
+{
+	FVector tet = NiagaraVectorFloor(v + v.X / 3 + v.Y / 3 + v.Z / 3);
+	FVector base = tet - tet.X / 6 - tet.Y / 6 - tet.Z / 6;
+	FVector f = v - base;
+
+	FVector g = NiagaraVectorStep(NiagaraVectorSwizzle(f, 1, 2, 0), f), h = FVector(1.0f) - NiagaraVectorSwizzle(g, 2, 0, 1);
+	FVector a1 = NiagaraVectorMin(g, h) - 1. / 6., a2 = NiagaraVectorMax(g, h) - 1. / 3.;
+
+	return FNiagaraMatrix4x3(base, base + a1, base + a2, base + 0.5);
+}
+
+FVector4 NiagaraVector4Saturate(FVector4 v) {
+	return FVector4(FMath::Clamp(v.X, 0.0f, 1.0f), FMath::Clamp(v.Y, 0.0f, 1.0f), FMath::Clamp(v.Z, 0.0f, 1.0f), FMath::Clamp(v.W, 0.0f, 1.0f));
+}
+
+FVector4 SimplexSmooth(FNiagaraMatrix4x3 f)
+{
+	const float scale = 1024. / 375.;
+	FVector4 d = FVector4(FVector::DotProduct(f[0], f[0]), FVector::DotProduct(f[1], f[1]), FVector::DotProduct(f[2], f[2]), FVector::DotProduct(f[3], f[3]));
+	FVector4 s = NiagaraVector4Saturate(2.0f * d);
+	return scale * (FVector4(1.0f, 1.0f, 1.0f, 1.0f) + s * (FVector4(-3.0f, -3.0f, -3.0f, -3.0f) + s * (FVector4(3.0f, 3.0f, 3.0f, 3.0f) - s)));
+}
+
+struct FNiagaraMatrix3x4
+{
+	FVector4 row0;
+	FVector4 row1;
+	FVector4 row2;
+
+	FNiagaraMatrix3x4() : row0(FVector4(0.0, 0.0, 0.0, 0.0)), row1(FVector4(0.0, 0.0, 0.0, 0.0)), row2(FVector4(0.0, 0.0, 0.0, 0.0)) {}
+	FNiagaraMatrix3x4(FVector4 row0, FVector4 row1, FVector4 row2) : row0(row0), row1(row1), row2(row2) {}
+
+	FVector4& operator[](int row)
+	{
+		if (row == 0) return row0;
+		if (row == 1) return row1;
+		if (row == 2) return row2;
+		check(false);
+		return row0;
+	}
+	const FVector4& operator[](int row) const
+	{
+		if (row == 0) return row0;
+		if (row == 1) return row1;
+		if (row == 2) return row2;
+		check(false); 
+		return row0; 
+	}
+};
+
+FNiagaraMatrix3x4 SimplexDSmooth(FNiagaraMatrix4x3 f)
+{
+	const float scale = 1024. / 375.;
+	FVector4 d = FVector4(FVector::DotProduct(f[0], f[0]), FVector::DotProduct(f[1], f[1]), FVector::DotProduct(f[2], f[2]), FVector::DotProduct(f[3], f[3]));
+	FVector4 s = NiagaraVector4Saturate(2 * d);
+	s = -12 * FVector4(scale, scale, scale, scale) + s * (24 * FVector4(scale, scale, scale, scale) - s * 12 * scale);
+
+	return FNiagaraMatrix3x4(
+		s * FVector4(f[0][0], f[1][0], f[2][0], f[3][0]),
+		s * FVector4(f[0][1], f[1][1], f[2][1], f[3][1]),
+		s * FVector4(f[0][2], f[1][2], f[2][2], f[3][2]));
+}
+
+FNiagaraUIntVector FNiagaraUIntVectorSwizzle(FNiagaraUIntVector v, int x, int y, int z)
+{
+	return FNiagaraUIntVector(v[x], v[y], v[z]);
+}
+
+FVector FNiagaraUIntVectorToFVector(FNiagaraUIntVector v)
+{
+	return FVector(v.X, v.Y, v.Z);
+}
+
+FVector MulFVector4AndFNiagaraMatrix4x3(FVector4 lhs, FNiagaraMatrix4x3 rhs)
+{
+	return FVector(lhs[0] * rhs[0][0] + lhs[1] * rhs[1][0] + lhs[2] * rhs[2][0] + lhs[3] * rhs[3][0],
+		           lhs[0] * rhs[0][1] + lhs[1] * rhs[1][1] + lhs[2] * rhs[2][1] + lhs[3] * rhs[3][1],
+	           	   lhs[0] * rhs[0][2] + lhs[1] * rhs[1][2] + lhs[2] * rhs[2][2] + lhs[3] * rhs[3][2]);
+}
+
+FVector MulFNiagaraMatrix3x4FAndVector4(FNiagaraMatrix3x4 lhs, FVector4 rhs)
+{
+	return FVector(lhs[0][0] * rhs[0] + lhs[0][1] * rhs[1] + lhs[0][2] * rhs[2] + lhs[0][3] * rhs[3],
+		           lhs[1][0] * rhs[0] + lhs[1][1] * rhs[1] + lhs[1][2] * rhs[2] + lhs[1][3] * rhs[3],
+		           lhs[2][0] * rhs[0] + lhs[2][1] * rhs[1] + lhs[2][2] * rhs[2] + lhs[2][3] * rhs[3]);
+}
+
+#define MGradientMask FNiagaraUIntVector(0x8000, 0x4000, 0x2000)
+#define MGradientScale FVector(1. / 0x4000, 1. / 0x2000, 1. / 0x1000)
+
+FNiagaraMatrix3x4 JacobianSimplex_ALU(FVector v)
+{
+	FNiagaraMatrix4x3 T = SimplexCorners(v);
+	FNiagaraUIntVector rand;
+	FNiagaraMatrix4x3 gvec[3], fv;
+	FNiagaraMatrix3x4 grad;
+
+	fv[0] = v - T[0];
+	rand = Rand3DPCG16(FIntVector(NiagaraVectorFloor(6 * T[0] + 0.5)));
+	gvec[0][0] = FVector(FNiagaraUIntVectorToFVector(FNiagaraUIntVectorSwizzle(rand, 0, 0, 0) & MGradientMask)) * MGradientScale - 1;
+	gvec[1][0] = FVector(FNiagaraUIntVectorToFVector(FNiagaraUIntVectorSwizzle(rand, 1, 1, 1) & MGradientMask)) * MGradientScale - 1;
+	gvec[2][0] = FVector(FNiagaraUIntVectorToFVector(FNiagaraUIntVectorSwizzle(rand, 2, 2, 2) & MGradientMask)) * MGradientScale - 1;
+	grad[0][0] = FVector::DotProduct(gvec[0][0], fv[0]);
+	grad[1][0] = FVector::DotProduct(gvec[1][0], fv[0]);
+	grad[2][0] = FVector::DotProduct(gvec[2][0], fv[0]);
+
+	fv[1] = v - T[1];
+	rand = Rand3DPCG16(FIntVector(NiagaraVectorFloor(6 * T[1] + 0.5)));
+	gvec[0][1] = FVector(FNiagaraUIntVectorToFVector(FNiagaraUIntVectorSwizzle(rand, 0, 0, 0) & MGradientMask)) * MGradientScale - 1;
+	gvec[1][1] = FVector(FNiagaraUIntVectorToFVector(FNiagaraUIntVectorSwizzle(rand, 1, 1, 1) & MGradientMask)) * MGradientScale - 1;
+	gvec[2][1] = FVector(FNiagaraUIntVectorToFVector(FNiagaraUIntVectorSwizzle(rand, 2, 2, 2) & MGradientMask)) * MGradientScale - 1;
+	grad[0][1] = FVector::DotProduct(gvec[0][1], fv[1]);
+	grad[1][1] = FVector::DotProduct(gvec[1][1], fv[1]);
+	grad[2][1] = FVector::DotProduct(gvec[2][1], fv[1]);
+
+	fv[2] = v - T[2];
+	rand = Rand3DPCG16(FIntVector(NiagaraVectorFloor(6 * T[2] + 0.5)));
+	gvec[0][2] = FVector(FNiagaraUIntVectorToFVector(FNiagaraUIntVectorSwizzle(rand, 0, 0, 0) & MGradientMask)) * MGradientScale - 1;
+	gvec[1][2] = FVector(FNiagaraUIntVectorToFVector(FNiagaraUIntVectorSwizzle(rand, 1, 1, 1) & MGradientMask)) * MGradientScale - 1;
+	gvec[2][2] = FVector(FNiagaraUIntVectorToFVector(FNiagaraUIntVectorSwizzle(rand, 2, 2, 2) & MGradientMask)) * MGradientScale - 1;
+	grad[0][2] = FVector::DotProduct(gvec[0][2], fv[2]);
+	grad[1][2] = FVector::DotProduct(gvec[1][2], fv[2]);
+	grad[2][2] = FVector::DotProduct(gvec[2][2], fv[2]);
+
+	fv[3] = v - T[3];
+	rand = Rand3DPCG16(FIntVector(NiagaraVectorFloor(6 * T[3] + 0.5)));
+	gvec[0][3] = FVector(FNiagaraUIntVectorToFVector(FNiagaraUIntVectorSwizzle(rand, 0, 0, 0) & MGradientMask)) * MGradientScale - 1;
+	gvec[1][3] = FVector(FNiagaraUIntVectorToFVector(FNiagaraUIntVectorSwizzle(rand, 1, 1, 1) & MGradientMask)) * MGradientScale - 1;
+	gvec[2][3] = FVector(FNiagaraUIntVectorToFVector(FNiagaraUIntVectorSwizzle(rand, 2, 2, 2) & MGradientMask)) * MGradientScale - 1;
+	grad[0][3] = FVector::DotProduct(gvec[0][3], fv[3]);
+	grad[1][3] = FVector::DotProduct(gvec[1][3], fv[3]);
+	grad[2][3] = FVector::DotProduct(gvec[2][3], fv[3]);
+
+	FVector4 sv = SimplexSmooth(fv);
+	FNiagaraMatrix3x4 ds = SimplexDSmooth(fv);
+
+	FNiagaraMatrix3x4 jacobian;
+	jacobian[0] = FVector4(MulFVector4AndFNiagaraMatrix4x3(sv, gvec[0]) + MulFNiagaraMatrix3x4FAndVector4(ds, grad[0]), Dot4(sv, grad[0]));
+	jacobian[1] = FVector4(MulFVector4AndFNiagaraMatrix4x3(sv, gvec[1]) + MulFNiagaraMatrix3x4FAndVector4(ds, grad[1]), Dot4(sv, grad[1]));
+	jacobian[2] = FVector4(MulFVector4AndFNiagaraMatrix4x3(sv, gvec[2]) + MulFNiagaraMatrix3x4FAndVector4(ds, grad[2]), Dot4(sv, grad[2]));
+
+	return jacobian;
 }
 
-void UNiagaraDataInterfaceCurlNoise::BeginDestroy()
+#undef MGradientMask
+#undef MGradientScale 
+
+// ----------------------------------------------------- end noise helpers -----------------------------------------------------
+
+
+UNiagaraDataInterfaceCurlNoise::UNiagaraDataInterfaceCurlNoise(FObjectInitializer const& ObjectInitializer)
+	: Super(ObjectInitializer)
+	, Seed(0)
 {
-	Super::BeginDestroy();
+	OffsetFromSeed = FNiagaraUIntVectorToFVector(Rand3DPCG16(FIntVector(Seed, Seed, Seed))) / 100.0;
+}
+
+void UNiagaraDataInterfaceCurlNoise::PostInitProperties()
+{
+	Super::PostInitProperties();
 
-	if (FApp::CanEverRender() && !HasAnyFlags(RF_ClassDefaultObject))
+	if (HasAnyFlags(RF_ClassDefaultObject))
 	{
-		ReleaseResource();
+		FNiagaraTypeRegistry::Register(FNiagaraTypeDefinition(GetClass()), true, false, false);
 	}
 }
 
-bool UNiagaraDataInterfaceCurlNoise::IsReadyForFinishDestroy()
+void UNiagaraDataInterfaceCurlNoise::PostLoad()
 {
-	return ReleaseResourcesFence.IsFenceComplete();
+	Super::PostLoad();
+	OffsetFromSeed = FNiagaraUIntVectorToFVector(Rand3DPCG16(FIntVector(Seed, Seed, Seed))) / 100.0;
 }
 
 #if WITH_EDITOR
 
+void UNiagaraDataInterfaceCurlNoise::PreEditChange(UProperty* PropertyAboutToChange)
+{
+	Super::PreEditChange(PropertyAboutToChange);
+	
+	// Flush the rendering thread before making any changes to make sure the 
+	// data read by the compute shader isn't subject to a race condition.
+	// TODO(mv): Solve properly using something like a RT Proxy.
+	FlushRenderingCommands();
+}
 void UNiagaraDataInterfaceCurlNoise::PostEditChangeProperty(struct FPropertyChangedEvent& PropertyChangedEvent)
 {
 	Super::PostEditChangeProperty(PropertyChangedEvent);
 
 	if (PropertyChangedEvent.Property && PropertyChangedEvent.Property->GetFName() == GET_MEMBER_NAME_CHECKED(UNiagaraDataInterfaceCurlNoise, Seed))
 	{
-		InitNoiseLUT();
+		// NOTE: Calculate the offset based on the seed on-change instead of on every invocation for every particle...
+		OffsetFromSeed = FNiagaraUIntVectorToFVector(Rand3DPCG16(FIntVector(Seed, Seed, Seed))) / 100.0;
 	}
 }
 
@@ -69,7 +358,7 @@ bool UNiagaraDataInterfaceCurlNoise::CopyToInternal(UNiagaraDataInterface* Desti
 	}
 	UNiagaraDataInterfaceCurlNoise* DestinationCurlNoise = CastChecked<UNiagaraDataInterfaceCurlNoise>(Destination);
 	DestinationCurlNoise->Seed = Seed;
-	DestinationCurlNoise->InitNoiseLUT();
+	DestinationCurlNoise->OffsetFromSeed = OffsetFromSeed;
 
 	return true;
 }
@@ -115,298 +404,76 @@ void UNiagaraDataInterfaceCurlNoise::SampleNoiseField(FVectorVMContext& Context)
 	VectorVM::FExternalFuncRegisterHandler<float> OutSampleY(Context);
 	VectorVM::FExternalFuncRegisterHandler<float> OutSampleZ(Context);
 
-	const VectorRegister One = MakeVectorRegister(1.0f, 1.0f, 1.0f, 1.0f);
-	const VectorRegister Zero = MakeVectorRegister(0.0f, 0.0f, 0.0f, 0.0f);
-	// We use 15,15,15 here because the value will be between 0 and 1. 1 should indicate the maximal index of the 
-	// real table and we pad that out to handle border work. So our values for the Cxyz below should not
-	// ever be 17 in any dimension as that would overflow the array.
-	const VectorRegister VecSize = MakeVectorRegister(15.0f, 15.0f, 15.0f, 15.0f);
-
-	float Di = 0.2f; // Hard-coded scale TODO remove!
-	const VectorRegister Div = MakeVectorRegister(Di, Di, Di, Di);
-
 	for (int32 InstanceIdx = 0; InstanceIdx < Context.NumInstances; ++InstanceIdx)
 	{
-		VectorRegister InCoords = MakeVectorRegister(XParam.GetAndAdvance(), YParam.GetAndAdvance(), ZParam.GetAndAdvance(), 0.0f);
-		VectorRegister Dst = MakeVectorRegister(0.0f, 0.0f, 0.0f, 0.0f);
-
-		VectorRegister Coords = VectorMod(VectorAbs(VectorMultiply(InCoords, Div)), VecSize);
-		Coords = VectorMin(Coords, VecSize);
-		Coords = VectorMax(Coords, Zero);
-		const float *CoordPtr = reinterpret_cast<float const*>(&Coords);
-		const int32 Cx = CoordPtr[0];
-		const int32 Cy = CoordPtr[1];
-		const int32 Cz = CoordPtr[2];
-
-		VectorRegister Frac = VectorFractional(Coords);
-		VectorRegister Alpha = VectorReplicate(Frac, 0);
-		VectorRegister OneMinusAlpha = VectorSubtract(One, Alpha);
-
-		// Trilinear interpolation, as defined by https://en.wikipedia.org/wiki/Trilinear_interpolation
-		ensure(Cx + 1 < 17);
-		ensure(Cy + 1 < 17);
-		ensure(Cz + 1 < 17);
-		VectorRegister C00 = VectorMultiplyAdd(NoiseTable[Cx][Cy][Cz], OneMinusAlpha, VectorMultiply(NoiseTable[Cx + 1][Cy][Cz], Alpha)); // (x0, y0, z0)(1-a) + (x1, y0, z0)a
-		VectorRegister C01 = VectorMultiplyAdd(NoiseTable[Cx][Cy][Cz + 1], OneMinusAlpha, VectorMultiply(NoiseTable[Cx + 1][Cy][Cz + 1], Alpha)); // (x0, y0, z1)(1-a) + (x1, y0, z1)a
-		VectorRegister C10 = VectorMultiplyAdd(NoiseTable[Cx][Cy + 1][Cz], OneMinusAlpha, VectorMultiply(NoiseTable[Cx + 1][Cy + 1][Cz], Alpha)); // (x0, y1, z0)(1-a) + (x1, y1, z0)a
-		VectorRegister C11 = VectorMultiplyAdd(NoiseTable[Cx][Cy + 1][Cz + 1], OneMinusAlpha, VectorMultiply(NoiseTable[Cx + 1][Cy + 1][Cz + 1], Alpha)); // (x0, y1, z1)(1-a) + (x1, y1, z1)a
-
-		Alpha = VectorReplicate(Frac, 1);
-		OneMinusAlpha = VectorSubtract(One, Alpha);
-		VectorRegister C0 = VectorMultiplyAdd(C00, OneMinusAlpha, VectorMultiply(C10, Alpha));
-		VectorRegister C1 = VectorMultiplyAdd(C01, OneMinusAlpha, VectorMultiply(C11, Alpha));
-
-		Alpha = VectorReplicate(Frac, 2);
-		OneMinusAlpha = VectorSubtract(One, Alpha);
-		VectorRegister ZV = VectorMultiplyAdd(C0, OneMinusAlpha, VectorMultiply(C1, Alpha));
-
-		Dst = VectorAdd(Dst, ZV);
-
-		float *RegPtr = reinterpret_cast<float*>(&Dst);
-		*OutSampleX.GetDestAndAdvance() = RegPtr[0];
-		*OutSampleY.GetDestAndAdvance() = RegPtr[1];
-		*OutSampleZ.GetDestAndAdvance() = RegPtr[2];
-	}
-}
-
-bool UNiagaraDataInterfaceCurlNoise::GetFunctionHLSL(const FName& DefinitionFunctionName, FString InstanceFunctionName, FNiagaraDataInterfaceGPUParamInfo& ParamInfo, FString& OutHLSL)
-{
-	FString BufferName = CurlNoiseBufferName + ParamInfo.DataInterfaceHLSLSymbol;
-	OutHLSL += TEXT("void ") + InstanceFunctionName + TEXT("(in float3 In_XYZ, out float3 Out_Value) \n{\n");
-	OutHLSL += TEXT("\t float3 a = trunc((In_XYZ*0.2) / 15.0);\n");
-	OutHLSL += TEXT("\t float3 ModXYZ = (In_XYZ*0.2) - a*15.0;\n");
-	OutHLSL += TEXT("\t int3 IntCoord = abs(int3(ModXYZ.x, ModXYZ.y, ModXYZ.z));\n");
-	OutHLSL += TEXT("\t float3 frc = frac(ModXYZ);\n");
-	// Trilinear interpolation, as defined by https://en.wikipedia.org/wiki/Trilinear_interpolation
-	// HLSL lerp is defined as lerp(x, y, s) = x*(1-s) + y*s
-	OutHLSL += TEXT("\t float3 V1 = ") + BufferName + TEXT("[IntCoord.x + IntCoord.y*17 + IntCoord.z*17*17].xyz;\n"); // x0,y0,z0
-	OutHLSL += TEXT("\t float3 V2 = ") + BufferName + TEXT("[IntCoord.x+1 + IntCoord.y*17 + IntCoord.z*17*17].xyz;\n"); //  x1,y0,z0
-	OutHLSL += TEXT("\t float3 C00 = lerp(V1, V2, frc.xxx);\n");
-	OutHLSL += TEXT("\t V1 = ") + BufferName + TEXT("[IntCoord.x + IntCoord.y*17 + (IntCoord.z+1)*17*17].xyz;\n"); // x0, y0, z1
-	OutHLSL += TEXT("\t V2 = ") + BufferName + TEXT("[IntCoord.x+1 + IntCoord.y*17 + (IntCoord.z+1)*17*17].xyz;\n"); //x1, y0, z1
-	OutHLSL += TEXT("\t float3 C01 = lerp(V1, V2, frc.xxx);\n");
-	OutHLSL += TEXT("\t V1 = ") + BufferName + TEXT("[IntCoord.x + (IntCoord.y+1)*17 + IntCoord.z*17*17].xyz;\n"); //x0, y1, z0
-	OutHLSL += TEXT("\t V2 = ") + BufferName + TEXT("[IntCoord.x+1 + (IntCoord.y+1)*17 + IntCoord.z*17*17].xyz;\n"); // x1, y1, z0
-	OutHLSL += TEXT("\t float3 C10 = lerp(V1, V2, frc.xxx);\n");
-	OutHLSL += TEXT("\t V1 = ") + BufferName + TEXT("[IntCoord.x + (IntCoord.y+1)*17 + (IntCoord.z+1)*17*17].xyz;\n"); // x0, y1, z1
-	OutHLSL += TEXT("\t V2 = ") + BufferName + TEXT("[IntCoord.x+1 + (IntCoord.y+1)*17 + (IntCoord.z+1)*17*17].xyz;\n"); //x1, y1, z1
-	OutHLSL += TEXT("\t float3 C11 = lerp(V1, V2, frc.xxx);\n");
-	OutHLSL += TEXT("\t float3 C0 = lerp(C00, C10, frc.yyy);\n");
-	OutHLSL += TEXT("\t float3 C1 = lerp(C01, C11, frc.yyy);\n");
-	OutHLSL += TEXT("\t Out_Value = lerp(C0, C1, frc.zzz);\n");
-	OutHLSL += TEXT("\n}\n");
-	return true;
-}
-
-void UNiagaraDataInterfaceCurlNoise::GetParameterDefinitionHLSL(FNiagaraDataInterfaceGPUParamInfo& ParamInfo, FString& OutHLSL)
-{
-	FString BufferName = CurlNoiseBufferName + ParamInfo.DataInterfaceHLSLSymbol;
-	OutHLSL += TEXT("Buffer<float4> ") + BufferName + TEXT(";\n");
-}
-
-void UNiagaraDataInterfaceCurlNoise::ReleaseResource()
-{
-	ENQUEUE_UNIQUE_RENDER_COMMAND_ONEPARAMETER(NiagaraCurlNoiseReleaseCommand,
-		FRWBuffer*,GPUBuffer,GPUBuffer.Get(),
-	{
-		GPUBuffer->Release();
-	});
-	// Insert a fence to signal when these commands completed
-	ReleaseResourcesFence.BeginFence();
-}
-
-FRWBuffer& UNiagaraDataInterfaceCurlNoise::GetGPUBuffer()
-{
-	check(IsInRenderingThread());
+		FVector InCoords = FVector(XParam.GetAndAdvance(), YParam.GetAndAdvance(), ZParam.GetAndAdvance());
 
-	if (bGPUBufferDirty)
-	{
-		GPUBuffer->Release();
-		uint32 BufferSize = 17 * 17 * 17 * sizeof(float) * 4;
-		//int32 *BufferData = static_cast<int32*>(RHILockVertexBuffer(GPUBuffer.Buffer.Buffer, 0, BufferSize, EResourceLockMode::RLM_WriteOnly));
-		TResourceArray<FVector4> TempTable;
-		TempTable.AddUninitialized(17 * 17 * 17);
-		for (int z = 0; z < 17; z++)
-		{
-			for (int y = 0; y < 17; y++)
-			{
-				for (int x = 0; x < 17; x++)
-				{
-					float *R = (float*)( &(NoiseTable[x][y][z]) );
-					TempTable[x + y * 17 + z * 17 * 17] = FVector4(R[0], R[1], R[2], 0.0f);
-				}
-			}
-		}
-		GPUBuffer->Initialize(sizeof(float) * 4, 17 * 17 * 17, EPixelFormat::PF_A32B32G32R32F, BUF_Static, TEXT("CurlnoiseTable"), &TempTable);
-		//FPlatformMemory::Memcpy(BufferData, TempTable, BufferSize);
-		//RHIUnlockVertexBuffer(GPUBuffer.Buffer.Buffer);
-		bGPUBufferDirty = false;
+		// See comments to JacobianSimplex_ALU in Random.ush
+		FNiagaraMatrix3x4 J = JacobianSimplex_ALU(InCoords + OffsetFromSeed);
+		*OutSampleX.GetDestAndAdvance() = J[1][2] - J[2][1];
+		*OutSampleY.GetDestAndAdvance() = J[2][0] - J[0][2];
+		*OutSampleZ.GetDestAndAdvance() = J[0][1] - J[1][0];
 	}
-
-	return *GPUBuffer;
 }
 
-
-// replicate a border for filtering
-template<typename T>
-void UNiagaraDataInterfaceCurlNoise::ReplicateBorder(T* DestBuffer)
+bool UNiagaraDataInterfaceCurlNoise::GetFunctionHLSL(const FName& DefinitionFunctionName, FString InstanceFunctionName, FNiagaraDataInterfaceGPUParamInfo& ParamInfo, FString& OutHLSL)
 {
-	uint32 Extent = 17;
-	uint32 Square = Extent*Extent;
-	uint32 Last = Extent - 1;
-
-	for (uint32 z = 0; z < Extent; ++z)
-	{
-		for (uint32 y = 0; y < Extent; ++y)
+	static const TCHAR *FormatSample = TEXT(R"(
+		void {FunctionName}(float3 In_XYZ, out float3 Out_Value)
 		{
-			DestBuffer[Last + y * Extent + z * Square] = DestBuffer[0 + y * Extent + z * Square];
+			// NOTE(mv): The comments in random.ush claims that the unused part is optimized away, so it only uses 6 out of 12 values in our case.
+			float3x4 J = JacobianSimplex_ALU(In_XYZ + {OffsetFromSeedName}, false, 1.0);
+			Out_Value = float3(J[1][2]-J[2][1], J[2][0]-J[0][2], J[0][1]-J[1][0]); // See comments to JacobianSimplex_ALU in Random.ush
 		}
-	}
-	for (uint32 z = 0; z < Extent; ++z)
-	{
-		for (uint32 x = 0; x < Extent; ++x)
-		{
-			DestBuffer[x + Last * Extent + z * Square] = DestBuffer[x + 0 * Extent + z * Square];
-		}
-	}
-	for (uint32 y = 0; y < Extent; ++y)
-	{
-		for (uint32 x = 0; x < Extent; ++x)
-		{
-			DestBuffer[x + y * Extent + Last * Square] = DestBuffer[x + y * Extent + 0 * Square];
-		}
-	}
+	)");
+	TMap<FString, FStringFormatArg> ArgsSample = {
+		{TEXT("FunctionName"), InstanceFunctionName},
+		{TEXT("OffsetFromSeedName"), OffsetFromSeedBaseName + ParamInfo.DataInterfaceHLSLSymbol},
+	};
+	OutHLSL += FString::Format(FormatSample, ArgsSample);
+	return true;
 }
 
-void UNiagaraDataInterfaceCurlNoise::InitNoiseLUT()
+void UNiagaraDataInterfaceCurlNoise::GetParameterDefinitionHLSL(FNiagaraDataInterfaceGPUParamInfo& ParamInfo, FString& OutHLSL)
 {
-	// seed random stream
-	FRandomStream RandStream(Seed);
-
-	// random noise
-	float TempTable[17][17][17];
-	for (int z = 0; z < 17; z++)
-	{
-		for (int y = 0; y < 17; y++)
-		{
-			for (int x = 0; x < 17; x++)
-			{
-				float f1 = RandStream.FRandRange(-1.0f, 1.0f);
-				TempTable[x][y][z] = f1;
-			}
-		}
-	}
-
-	// pad
-	/*
-	for (int i = 0; i < 17; i++)
-	{
-		for (int j = 0; j < 17; j++)
-		{
-			TempTable[i][j][16] = TempTable[i][j][0];
-			TempTable[i][16][j] = TempTable[i][0][j];
-			TempTable[16][j][i] = TempTable[0][j][i];
-		}
-	}
-	*/
-	ReplicateBorder(&TempTable[0][0][0]);
-
-	// compute gradients
-	FVector TempTable2[17][17][17];
-	for (int z = 0; z < 16; z++)
-	{
-		for (int y = 0; y < 16; y++)
-		{
-			for (int x = 0; x < 16; x++)
-			{
-				FVector XGrad = FVector(1.0f, 0.0f, TempTable[x][y][z] - TempTable[x + 1][y][z]);
-				FVector YGrad = FVector(0.0f, 1.0f, TempTable[x][y][z] - TempTable[x][y + 1][z]);
-				FVector ZGrad = FVector(0.0f, 1.0f, TempTable[x][y][z] - TempTable[x][y][z + 1]);
-
-				FVector Grad = FVector(XGrad.Z, YGrad.Z, ZGrad.Z);
-				TempTable2[x][y][z] = Grad;
-			}
-		}
-	}
-
-	/*
-	// pad
-	for (int i = 0; i < 17; i++)
-	{
-		for (int j = 0; j < 17; j++)
-		{
-			TempTable2[i][j][16] = TempTable2[i][j][0];
-			TempTable2[i][16][j] = TempTable2[i][0][j];
-			TempTable2[16][j][i] = TempTable2[0][j][i];
-		}
-	}
-	*/
-	ReplicateBorder(&TempTable2[0][0][0]);
-
-	// http://prideout.net/blog/?p=63           ???
-	// compute curl of gradient field
-	for (int z = 0; z < 16; z++)
-	{
-		for (int y = 0; y < 16; y++)
-		{
-			for (int x = 0; x < 16; x++)
-			{
-				FVector Dy = TempTable2[x][y][z] - TempTable2[x][y + 1][z];
-				FVector Sy = TempTable2[x][y][z] + TempTable2[x][y + 1][z];
-				FVector Dx = TempTable2[x][y][z] - TempTable2[x + 1][y][z];
-				FVector Sx = TempTable2[x][y][z] + TempTable2[x + 1][y][z];
-				FVector Dz = TempTable2[x][y][z] - TempTable2[x][y][z + 1];
-				FVector Sz = TempTable2[x][y][z] + TempTable2[x][y][z + 1];
-				FVector Dir = FVector(Dy.Z - Sz.Y, Dz.X - Sx.Z, Dx.Y - Sy.X);
-
-				NoiseTable[x][y][z] = MakeVectorRegister(Dir.X, Dir.Y, Dir.Z, 0.0f);
-			}
-		}
-	}
-
-	// pad
-	/*
-	for (int i = 0; i < 17; i++)
-	{
-		for (int j = 0; j < 17; j++)
-		{
-			NoiseTable[i][j][16] = NoiseTable[i][j][0];
-			NoiseTable[i][16][j] = NoiseTable[i][0][j];
-			NoiseTable[16][j][i] = NoiseTable[0][j][i];
-		}
-	}*/
-	ReplicateBorder(&NoiseTable[0][0][0]);
-
-	bGPUBufferDirty = true;
+	static const TCHAR *FormatDeclarations = TEXT(R"(
+		float3 {OffsetFromSeedName};
+	)");
+
+	TMap<FString, FStringFormatArg> ArgsDeclarations = {
+		{TEXT("OffsetFromSeedName"), OffsetFromSeedBaseName + ParamInfo.DataInterfaceHLSLSymbol},
+	};
+	OutHLSL += FString::Format(FormatDeclarations, ArgsDeclarations);
 }
 
-
 struct FNiagaraDataInterfaceParametersCS_CurlNoise : public FNiagaraDataInterfaceParametersCS
 {
 	virtual void Bind(const FNiagaraDataInterfaceParamRef& ParamRef, const class FShaderParameterMap& ParameterMap) override
 	{
-		CurlNoiseBuffer.Bind(ParameterMap, *(UNiagaraDataInterfaceCurlNoise::CurlNoiseBufferName + ParamRef.ParameterInfo.DataInterfaceHLSLSymbol));
+		OffsetFromSeed.Bind(ParameterMap, *(OffsetFromSeedBaseName + ParamRef.ParameterInfo.DataInterfaceHLSLSymbol));
 	}
 
 	virtual void Serialize(FArchive& Ar)override
 	{
-		Ar << CurlNoiseBuffer;
+		Ar << OffsetFromSeed;
 	}
 
-	virtual void Set(FRHICommandList& RHICmdList, FNiagaraShader* Shader, class UNiagaraDataInterface* DataInterface) const override
+	virtual void Set(FRHICommandList& RHICmdList, FNiagaraShader* Shader, class UNiagaraDataInterface* DataInterface, void* PerInstanceData) const override
 	{
 		check(IsInRenderingThread());
 
+		// Get shader and DI
 		const FComputeShaderRHIParamRef ComputeShaderRHI = Shader->GetComputeShader();
-		UNiagaraDataInterfaceCurlNoise* CurlNoiseDI = CastChecked<UNiagaraDataInterfaceCurlNoise>(DataInterface);
-		FRWBuffer& GPUBuffer = CurlNoiseDI->GetGPUBuffer();
+		UNiagaraDataInterfaceCurlNoise* CNDI = CastChecked<UNiagaraDataInterfaceCurlNoise>(DataInterface);
 
-		RHICmdList.SetShaderResourceViewParameter(ComputeShaderRHI, CurlNoiseBuffer.GetBaseIndex(), GPUBuffer.SRV);
+		// Note: There is a flush in PreEditChange to make sure everything is synced up at this point 
+
+		// Set parameters
+		SetShaderValue(RHICmdList, ComputeShaderRHI, OffsetFromSeed, CNDI->OffsetFromSeed);
 	}
 
 private:
-
-	FShaderResourceParameter CurlNoiseBuffer;
+	FShaderParameter OffsetFromSeed;
 };
 
 FNiagaraDataInterfaceParametersCS* UNiagaraDataInterfaceCurlNoise::ConstructComputeParameters()const
diff --git a/Engine/Plugins/FX/Niagara/Source/Niagara/Private/NiagaraDataInterfaceTexture.cpp b/Engine/Plugins/FX/Niagara/Source/Niagara/Private/NiagaraDataInterfaceTexture.cpp
--- a/Engine/Plugins/FX/Niagara/Source/Niagara/Private/NiagaraDataInterfaceTexture.cpp
+++ b/Engine/Plugins/FX/Niagara/Source/Niagara/Private/NiagaraDataInterfaceTexture.cpp
@@ -373,7 +373,7 @@ struct FNiagaraDataInterfaceParametersCS_Texture : public FNiagaraDataInterfaceP
 		}
 	}
 
-	virtual void Set(FRHICommandList& RHICmdList, FNiagaraShader* Shader, class UNiagaraDataInterface* DataInterface) const override
+	virtual void Set(FRHICommandList& RHICmdList, FNiagaraShader* Shader, class UNiagaraDataInterface* DataInterface, void* PerInstanceData) const override
 	{
 		check(IsInRenderingThread());
 
@@ -388,20 +388,33 @@ struct FNiagaraDataInterfaceParametersCS_Texture : public FNiagaraDataInterfaceP
 			SetShaderValue(RHICmdList, ComputeShaderRHI, Dimensions, TexDims);
 			return;
 		}
+
 		FTextureRHIParamRef TextureRHI = Texture->TextureReference.TextureReferenceRHI->GetReferencedTexture();
-		SetTextureParameter(
-			RHICmdList,
-			ComputeShaderRHI,
-			TextureParam,
-			SamplerParam,
-			Texture->Resource->SamplerStateRHI,
-			TextureRHI
-		);
-		TexDims[0] = TextureDI->Texture->GetSurfaceWidth();
-		TexDims[1] = TextureDI->Texture->GetSurfaceHeight();
+		if (TextureRHI != nullptr)
+		{
+			SetTextureParameter(
+				RHICmdList,
+				ComputeShaderRHI,
+				TextureParam,
+				SamplerParam,
+				Texture->Resource->SamplerStateRHI,
+				TextureRHI
+			);
+			TexDims[0] = TextureDI->Texture->GetSurfaceWidth();
+			TexDims[1] = TextureDI->Texture->GetSurfaceHeight();
+
+		}
+		else
+		{
+			TexDims[0] = 0.0f;
+			TexDims[1] = 0.0f;
+			SetShaderValue(RHICmdList, ComputeShaderRHI, Dimensions, TexDims);
+			return;
+		}
+
 		SetShaderValue(RHICmdList, ComputeShaderRHI, Dimensions, TexDims);
 	}
-
+	
 
 private:
 
