Merging //UE4/Dev-Main to Dev-Build (//UE4/Dev-Build)

diff --git a/Engine/Plugins/FX/Niagara/Source/Niagara/Private/NiagaraDataInterfaceVectorField.cpp b/Engine/Plugins/FX/Niagara/Source/Niagara/Private/NiagaraDataInterfaceVectorField.cpp
--- a/Engine/Plugins/FX/Niagara/Source/Niagara/Private/NiagaraDataInterfaceVectorField.cpp
+++ b/Engine/Plugins/FX/Niagara/Source/Niagara/Private/NiagaraDataInterfaceVectorField.cpp
@@ -1,317 +1,59 @@
 // Copyright 1998-2019 Epic Games, Inc. All Rights Reserved.
 
 #include "NiagaraDataInterfaceVectorField.h"
-#include "Internationalization/Internationalization.h"
 #include "VectorField/VectorFieldStatic.h"
 #include "VectorField/VectorFieldAnimated.h"
-#include "Math/Float16Color.h"
 #include "NiagaraShader.h"
 #include "ShaderParameterUtils.h"
+//#include "Internationalization/Internationalization.h"
 
 #define LOCTEXT_NAMESPACE "NiagaraDataInterfaceVectorField"
 
+// Global HLSL variable base names, used by HLSL.
+static const FString SamplerBaseName(TEXT("VectorFieldSampler_"));
+static const FString TextureBaseName(TEXT("VectorFieldTexture_"));
+static const FString TilingAxesBaseName(TEXT("TilingAxes_"));
+static const FString DimensionsBaseName(TEXT("Dimensions_"));
+static const FString MinBoundsBaseName(TEXT("MinBounds_"));
+static const FString MaxBoundsBaseName(TEXT("MaxBounds_"));
 
-const FString UNiagaraDataInterfaceVectorField::BufferBaseName(TEXT("VectorFieldBuffer_"));
-const FString UNiagaraDataInterfaceVectorField::DimentionsBaseName(TEXT("Dimentions_"));
-const FString UNiagaraDataInterfaceVectorField::BoundsMinBaseName(TEXT("BoundsMin_"));
-const FString UNiagaraDataInterfaceVectorField::BoundsMaxBaseName(TEXT("BoundsMax_"));
-
-#if 0
-//////////////////////////////////////////////////////////////////////////
-//FNDIVectorField_InstanceData
-
-bool FNDIVectorField_InstanceData::Init(UNiagaraDataInterfaceVectorField* Interface, FNiagaraSystemInstance* SystemInstance)
-{
-	UE_LOG(LogNiagara, Log, TEXT("FNDIVectorField_InstanceData %p"), this);
-	check(SystemInstance);
-	/*
-	USkeletalMesh* PrevMesh = Mesh;
-	Component = nullptr;
-	Mesh = nullptr;
-	Transform = FMatrix::Identity;
-	TransformInverseTransposed = FMatrix::Identity;
-	PrevTransform = FMatrix::Identity;
-	PrevTransformInverseTransposed = FMatrix::Identity;
-	*/
-	DeltaSeconds = 0.0f;
-	Field = Interface->Field;
-
-	UVectorFieldStatic* StaticVectorField = Cast<UVectorFieldStatic>(Field);
-	UVectorFieldAnimated* AnimatedVectorField = Cast<UVectorFieldAnimated>(Field);
-	SizeX = 0;
-	SizeY = 0;
-	SizeZ = 0;
-
-	/*
-	OutParameters.WorldToVolume[Index] = VectorFieldInstance->WorldToVolume;
-	OutParameters.VolumeToWorld[Index] = VectorFieldInstance->VolumeToWorldNoScale;
-	OutParameters.VolumeSize[Index] = FVector4(Resource->SizeX, Resource->SizeY, Resource->SizeZ, 0);
-	OutParameters.IntensityAndTightness[Index] = FVector4(Intensity, Tightness, 0, 0);
-	OutParameters.TilingAxes[Index].X = VectorFieldInstance->bTileX ? 1.0f : 0.0f;
-	OutParameters.TilingAxes[Index].Y = VectorFieldInstance->bTileY ? 1.0f : 0.0f;
-	OutParameters.TilingAxes[Index].Z = VectorFieldInstance->bTileZ ? 1.0f : 0.0f;
-	*/
-	TilingAxes.X = Interface->bTileX ? 1.0f : 0.0f;
-	TilingAxes.Y = Interface->bTileY ? 1.0f : 0.0f;
-	TilingAxes.Z = Interface->bTileZ ? 1.0f : 0.0f;
-
-
-	if (StaticVectorField)
-	{
-		SizeX = (uint32)StaticVectorField->SizeX;
-		SizeY = (uint32)StaticVectorField->SizeY;
-		SizeZ = (uint32)StaticVectorField->SizeZ;
-		LocalBounds = StaticVectorField->Bounds;
-	}
-	else if (AnimatedVectorField)
-	{
-		SizeX = (uint32)AnimatedVectorField->VolumeSizeX;
-		SizeY = (uint32)AnimatedVectorField->VolumeSizeY;
-		SizeZ = (uint32)AnimatedVectorField->VolumeSizeZ;
-		LocalBounds = AnimatedVectorField->Bounds;
-	}
-#if 0
-	USkeletalMeshComponent* NewSkelComp = nullptr;
-	if (Interface->Source)
-	{
-		ASkeletalMeshActor* MeshActor = Cast<ASkeletalMeshActor>(Interface->Source);
-		USkeletalMeshComponent* SourceComp = nullptr;
-		if (MeshActor != nullptr)
-		{
-			SourceComp = MeshActor->GetSkeletalMeshComponent();
-		}
-		else
-		{
-			SourceComp = Interface->Source->FindComponentByClass<USkeletalMeshComponent>();
-		}
-
-		if (SourceComp)
-		{
-			Mesh = SourceComp->SkeletalMesh;
-			NewSkelComp = SourceComp;
-		}
-		else
-		{
-			Component = Interface->Source->GetRootComponent();
-		}
-	}
-	else
-	{
-		if (UNiagaraComponent* SimComp = SystemInstance->GetComponent())
-		{
-			if (USkeletalMeshComponent* ParentComp = Cast<USkeletalMeshComponent>(SimComp->GetAttachParent()))
-			{
-				NewSkelComp = ParentComp;
-				Mesh = ParentComp->SkeletalMesh;
-			}
-			else if (USkeletalMeshComponent* OuterComp = SimComp->GetTypedOuter<USkeletalMeshComponent>())
-			{
-				NewSkelComp = OuterComp;
-				Mesh = OuterComp->SkeletalMesh;
-			}
-			else if (AActor* Owner = SimComp->GetAttachmentRootActor())
-			{
-				TArray<UActorComponent*> SourceComps = Owner->GetComponentsByClass(USkeletalMeshComponent::StaticClass());
-				for (UActorComponent* ActorComp : SourceComps)
-				{
-					USkeletalMeshComponent* SourceComp = Cast<USkeletalMeshComponent>(ActorComp);
-					if (SourceComp)
-					{
-						USkeletalMesh* PossibleMesh = SourceComp->SkeletalMesh;
-						if (PossibleMesh != nullptr/* && PossibleMesh->bAllowCPUAccess*/)
-						{
-							Mesh = PossibleMesh;
-							NewSkelComp = SourceComp;
-							break;
-						}
-					}
-				}
-			}
-
-			if (!Component.IsValid())
-			{
-				Component = SimComp;
-			}
-		}
-	}
-
-	if (NewSkelComp)
-	{
-		Component = NewSkelComp;
-	}
-
-	check(Component.IsValid());
-
-	if (!Mesh && Interface->DefaultMesh)
-	{
-		Mesh = Interface->DefaultMesh;
-	}
-
-	if (Component.IsValid() && Mesh)
-	{
-		PrevTransform = Transform;
-		PrevTransformInverseTransposed = TransformInverseTransposed;
-		Transform = Component->GetComponentToWorld().ToMatrixWithScale();
-		TransformInverseTransposed = Transform.InverseFast().GetTransposed();
-	}
-
-	if (!Mesh)
-	{
-		UE_LOG(LogNiagara, Log, TEXT("SkeletalMesh data interface has no valid mesh. Failed InitPerInstanceData - %s"), *Interface->GetFullName());
-		return false;
-	}
-
-	// 	if (!Mesh->bAllowCPUAccess)
-	// 	{
-	// 		UE_LOG(LogNiagara, Log, TEXT("SkeletalMesh data interface using a mesh that does not allow CPU access. Failed InitPerInstanceData - Mesh: %s"), *Mesh->GetFullName());
-	// 		return false;
-	// 	}
-
-	if (!Component.IsValid())
-	{
-		UE_LOG(LogNiagara, Log, TEXT("SkeletalMesh data interface has no valid component. Failed InitPerInstanceData - %s"), *Interface->GetFullName());
-		return false;
-	}
-
-	//	bIsAreaWeightedSampling = Mesh->bSupportUniformlyDistributedSampling;
-
-	// 	//Init the instance filter
-	// 	ValidSections.Empty();
-	// 	FStaticMeshLODResources& Res = Mesh->RenderData->LODResources[0];
-	// 	for (int32 i = 0; i < Res.Sections.Num(); ++i)
-	// 	{
-	// 		if (Interface->SectionFilter.AllowedMaterialSlots.Num() == 0 || Interface->SectionFilter.AllowedMaterialSlots.Contains(Res.Sections[i].MaterialIndex))
-	// 		{
-	// 			ValidSections.Add(i);
-	// 		}
-	// 	}
-	// 
-	// 	if (GetValidSections().Num() == 0)
-	// 	{
-	// 		UE_LOG(LogNiagara, Log, TEXT("StaticMesh data interface has a section filter preventing any spawning. Failed InitPerInstanceData - %s"), *Interface->GetFullName());
-	// 		return false;
-	// 	}
-	// 
-	// 	Sampler.Init(&Res, this);
-
-	if (NewSkelComp)
-	{
-		TWeakObjectPtr<USkeletalMeshComponent> SkelWeakCompPtr = NewSkelComp;
-		SkinningData = SystemInstance->GetWorldManager()->GetSkeletalMeshGeneratedData().GetCachedSkinningData(SkelWeakCompPtr);
-	}
-#endif
-
-	return true;
-}
-
-void FNDIVectorField_InstanceData::UpdateTransforms(const FMatrix& LocalToWorld, FMatrix& OutVolumeToWorld, FMatrix& OutWorldToVolume)
-{
-	const FVector VolumeOffset = LocalBounds.Min;
-	const FVector VolumeScale = LocalBounds.Max - LocalBounds.Min;
-	//VolumeToWorldNoScale = LocalToWorld.GetMatrixWithoutScale().RemoveTranslation();
-	OutVolumeToWorld = FScaleMatrix(VolumeScale) * FTranslationMatrix(VolumeOffset)
-		* LocalToWorld;
-	OutWorldToVolume = OutVolumeToWorld.InverseFast();
-}
-
-const void* FNDIVectorField_InstanceData::Lock()
-{
-	UVectorFieldStatic* StaticVectorField = Cast<UVectorFieldStatic>(Field);
-	UVectorFieldAnimated* AnimatedVectorField = Cast<UVectorFieldAnimated>(Field);
-
-	if (StaticVectorField)
-	{
-		return StaticVectorField->SourceData.LockReadOnly();
-	}
-	return nullptr;
-}
-
-void FNDIVectorField_InstanceData::Unlock()
-{
-	UVectorFieldStatic* StaticVectorField = Cast<UVectorFieldStatic>(Field);
-	UVectorFieldAnimated* AnimatedVectorField = Cast<UVectorFieldAnimated>(Field);
-
-	if (StaticVectorField)
-	{
-		StaticVectorField->SourceData.Unlock();
-	}
-}
-
-
-bool FNDIVectorField_InstanceData::ResetRequired(UNiagaraDataInterfaceVectorField* Interface)const
-{
-	//USceneComponent* Comp = Component.Get();
-	//if (!Comp)
-	//{
-	//	//The component we were bound to is no longer valid so we have to trigger a reset.
-	//	return true;
-	//}
-
-	//if (USkeletalMeshComponent* SkelComp = Cast<USkeletalMeshComponent>(Comp))
-	//{
-	//	if (!SkelComp->SkeletalMesh)
-	//	{
-	//		return true;
-	//	}
-	//}
-	//else
-	//{
-		if (!Interface->Field || Interface->Field != Field)
-		{
-			return true;
-		}
-	//}
-
-	return false;
-}
+// Global VM function names, also used by the shaders code generation methods.
+static const FName SampleVectorFieldName("SampleField");
+static const FName GetVectorFieldTilingAxesName("FieldTilingAxes");
+static const FName GetVectorFieldDimensionsName("FieldDimensions");
+static const FName GetVectorFieldBoundsName("FieldBounds");
 
-bool FNDIVectorField_InstanceData::Tick(UNiagaraDataInterfaceVectorField* Interface, FNiagaraSystemInstance* SystemInstance, float InDeltaSeconds)
-{
-	if (ResetRequired(Interface))
-	{
-		return true;
-	}
-	else
-	{
-		DeltaSeconds = InDeltaSeconds;
-		/*if (Component.IsValid() && Mesh)
-		{
-			PrevTransform = Transform;
-			PrevTransformInverseTransposed = TransformInverseTransposed;
-			Transform = Component->GetComponentToWorld().ToMatrixWithScale();
-			TransformInverseTransposed = Transform.InverseFast().GetTransposed();
-		}
-		else
-		{
-			PrevTransform = FMatrix::Identity;
-			PrevTransformInverseTransposed = FMatrix::Identity;
-			Transform = FMatrix::Identity;
-			TransformInverseTransposed = FMatrix::Identity;
-		}*/
-		return false;
-	}
-}
-#endif
+/*--------------------------------------------------------------------------------------------------------------------------*/
 
-//Instance Data END
-//////////////////////////////////////////////////////////////////////////
 UNiagaraDataInterfaceVectorField::UNiagaraDataInterfaceVectorField(FObjectInitializer const& ObjectInitializer)
 	: Super(ObjectInitializer)
 	, Field(nullptr)
 	, bTileX(false)
 	, bTileY(false)
 	, bTileZ(false)
-	, bGPUBufferDirty(true)
 {
 }
 
+/*--------------------------------------------------------------------------------------------------------------------------*/
+
+
 #if WITH_EDITOR
+void UNiagaraDataInterfaceVectorField::PreEditChange(UProperty* PropertyAboutToChange)
+{
+	Super::PreEditChange(PropertyAboutToChange);
+	
+	// Flush the rendering thread before making any changes to make sure the 
+	// data read by the compute shader isn't subject to a race condition.
+	// TODO(mv): Solve properly using something like a RT Proxy.
+	FlushRenderingCommands();
+}
 
 void UNiagaraDataInterfaceVectorField::PostEditChangeProperty(struct FPropertyChangedEvent& PropertyChangedEvent)
 {
 	Super::PostEditChangeProperty(PropertyChangedEvent);
-	InitField();
 }
+#endif //WITH_EDITOR
+
 
 void UNiagaraDataInterfaceVectorField::PostLoad()
 {
@@ -320,7 +62,6 @@ void UNiagaraDataInterfaceVectorField::PostLoad()
 	{
 		Field->ConditionalPostLoad();
 	}
-	InitField();
 }
 
 void UNiagaraDataInterfaceVectorField::PostInitProperties()
@@ -330,15 +71,11 @@ void UNiagaraDataInterfaceVectorField::PostInitProperties()
 	//Can we register data interfaces as regular types and fold them into the FNiagaraVariable framework for UI and function calls etc?
 	if (HasAnyFlags(RF_ClassDefaultObject))
 	{
-		FNiagaraTypeRegistry::Register(FNiagaraTypeDefinition(GetClass()), true, false, false);
+		FNiagaraTypeRegistry::Register(FNiagaraTypeDefinition(GetClass()), /*bCanBeParameter*/ true, /*bCanBePayload*/ false, /*bIsUserDefined*/ false);
 	}
 }
 
-#endif //WITH_EDITOR
-
-static const FName SampleVectorFieldName("SampleField");
-static const FName GetVectorDimsName("FieldDimensions");
-static const FName GetVectorFieldBoundsName("FieldBounds");
+/*--------------------------------------------------------------------------------------------------------------------------*/
 
 void UNiagaraDataInterfaceVectorField::GetFunctions(TArray<FNiagaraFunctionSignature>& OutFunctions)
 {
@@ -356,7 +93,7 @@ void UNiagaraDataInterfaceVectorField::GetFunctions(TArray<FNiagaraFunctionSigna
 
 	{
 		FNiagaraFunctionSignature Sig;
-		Sig.Name = GetVectorDimsName;
+		Sig.Name = GetVectorFieldDimensionsName;
 		Sig.Inputs.Add(FNiagaraVariable(FNiagaraTypeDefinition(GetClass()), TEXT("Vector Field")));
 		Sig.Outputs.Add(FNiagaraVariable(FNiagaraTypeDefinition::GetVec3Def(), TEXT("Dimensions")));
 
@@ -365,12 +102,23 @@ void UNiagaraDataInterfaceVectorField::GetFunctions(TArray<FNiagaraFunctionSigna
 		OutFunctions.Add(Sig);
 	}
 
+	{
+		FNiagaraFunctionSignature Sig;
+		Sig.Name = GetVectorFieldTilingAxesName;
+		Sig.Inputs.Add(FNiagaraVariable(FNiagaraTypeDefinition(GetClass()), TEXT("Vector Field")));
+		Sig.Outputs.Add(FNiagaraVariable(FNiagaraTypeDefinition::GetVec3Def(), TEXT("TilingAxes")));
+
+		Sig.bMemberFunction = true;
+		Sig.bRequiresContext = false;
+		OutFunctions.Add(Sig);
+	}
+
 	{
 		FNiagaraFunctionSignature Sig;
 		Sig.Name = GetVectorFieldBoundsName;
 		Sig.Inputs.Add(FNiagaraVariable(FNiagaraTypeDefinition(GetClass()), TEXT("Vector Field")));
-		Sig.Outputs.Add(FNiagaraVariable(FNiagaraTypeDefinition::GetVec3Def(), TEXT("BoundMin")));
-		Sig.Outputs.Add(FNiagaraVariable(FNiagaraTypeDefinition::GetVec3Def(), TEXT("BoundMax")));
+		Sig.Outputs.Add(FNiagaraVariable(FNiagaraTypeDefinition::GetVec3Def(), TEXT("MinBounds")));
+		Sig.Outputs.Add(FNiagaraVariable(FNiagaraTypeDefinition::GetVec3Def(), TEXT("MaxBounds")));
 
 		Sig.bMemberFunction = true;
 		Sig.bRequiresContext = false;
@@ -385,7 +133,7 @@ void UNiagaraDataInterfaceVectorField::GetVMExternalFunction(const FVMExternalFu
 	{
 		NDI_FUNC_BINDER(UNiagaraDataInterfaceVectorField, SampleVectorField)::Bind(this, OutFunc);
 	}
-	else if (BindingInfo.Name == GetVectorDimsName && BindingInfo.GetNumInputs() == 0 && BindingInfo.GetNumOutputs() == 3)
+	else if (BindingInfo.Name == GetVectorFieldDimensionsName && BindingInfo.GetNumInputs() == 0 && BindingInfo.GetNumOutputs() == 3)
 	{
 		OutFunc = FVMExternalFunction::CreateUObject(this, &UNiagaraDataInterfaceVectorField::GetFieldDimensions);
 	}
@@ -393,170 +141,321 @@ void UNiagaraDataInterfaceVectorField::GetVMExternalFunction(const FVMExternalFu
 	{
 		OutFunc = FVMExternalFunction::CreateUObject(this, &UNiagaraDataInterfaceVectorField::GetFieldBounds);
 	}
+	else if (BindingInfo.Name == GetVectorFieldTilingAxesName && BindingInfo.GetNumInputs() == 0 && BindingInfo.GetNumOutputs() == 3)
+	{
+		OutFunc = FVMExternalFunction::CreateUObject(this, &UNiagaraDataInterfaceVectorField::GetFieldTilingAxes);
+	}
 }
 
-void UNiagaraDataInterfaceVectorField::GetFieldDimensions(FVectorVMContext& Context)
+bool UNiagaraDataInterfaceVectorField::Equals(const UNiagaraDataInterface* Other) const
 {
-	/*VectorVM::FUserPtrHandler<FNDIVectorField_InstanceData> InstData(Context);
-	if (!InstData.Get())
-	{
-		UE_LOG(LogNiagara, Warning, TEXT("FNDIVectorField_InstanceData has invalid instance data. %s"), *GetPathName());
-	}*/
-
-	VectorVM::FExternalFuncRegisterHandler<float> OutSizeX(Context);	
-	VectorVM::FExternalFuncRegisterHandler<float> OutSizeY(Context);
-	VectorVM::FExternalFuncRegisterHandler<float> OutSizeZ(Context);
-	
-	for (int32 i = 0; i < Context.NumInstances; ++i)
+	if (!Super::Equals(Other))
 	{
-		*OutSizeX.GetDest() = (float)SizeX;		
-		*OutSizeY.GetDest() = (float)SizeY;
-		*OutSizeZ.GetDest() = (float)SizeZ;
-
-		OutSizeX.Advance();
-		OutSizeY.Advance();
-		OutSizeZ.Advance();
+		return false;
 	}
+	const UNiagaraDataInterfaceVectorField* OtherTyped = CastChecked<const UNiagaraDataInterfaceVectorField>(Other);
+	return OtherTyped->Field == Field 
+		&& OtherTyped->bTileX == bTileX
+		&& OtherTyped->bTileY == bTileY
+		&& OtherTyped->bTileZ == bTileZ;
 }
 
-void UNiagaraDataInterfaceVectorField::GetFieldBounds(FVectorVMContext& Context)
+bool UNiagaraDataInterfaceVectorField::CanExecuteOnTarget(ENiagaraSimTarget Target) const
 {
-	/*VectorVM::FUserPtrHandler<FNDIVectorField_InstanceData> InstData(Context);
-	if (!InstData.Get())
-	{
-		UE_LOG(LogNiagara, Warning, TEXT("FNDIVectorField_InstanceData has invalid instance data. %s"), *GetPathName());
-	}*/
+	return true;
+}
 
-	VectorVM::FExternalFuncRegisterHandler<float> OutMinX(Context);
-	VectorVM::FExternalFuncRegisterHandler<float> OutMinY(Context);
-	VectorVM::FExternalFuncRegisterHandler<float> OutMinZ(Context);
-	VectorVM::FExternalFuncRegisterHandler<float> OutMaxX(Context);
-	VectorVM::FExternalFuncRegisterHandler<float> OutMaxY(Context);
-	VectorVM::FExternalFuncRegisterHandler<float> OutMaxZ(Context);
+/*--------------------------------------------------------------------------------------------------------------------------*/
+
+#if WITH_EDITOR	
+TArray<FNiagaraDataInterfaceError> UNiagaraDataInterfaceVectorField::GetErrors()
+{
+	UVectorFieldStatic* StaticVectorField = Cast<UVectorFieldStatic>(Field);
+	UVectorFieldAnimated* AnimatedVectorField = Cast<UVectorFieldAnimated>(Field);
 	
-	for (int32 i = 0; i < Context.NumInstances; ++i)
+	// TODO(mv): Improve error messages?
+	TArray<FNiagaraDataInterfaceError> Errors;
+	if (StaticVectorField != nullptr && !StaticVectorField->bAllowCPUAccess)
+	{
+		FNiagaraDataInterfaceError CPUAccessNotAllowedError(
+			FText::Format(
+				LOCTEXT("CPUAccessNotAllowedError", "This Vector Field needs CPU access in order to be used properly.({0})"), 
+				FText::FromString(StaticVectorField->GetName())
+			),
+			LOCTEXT("CPUAccessNotAllowedErrorSummary", "CPU access error"),
+			FNiagaraDataInterfaceFix::CreateLambda(
+				[=]()
+				{
+					StaticVectorField->SetCPUAccessEnabled();
+					return true;
+				}
+			)
+		);
+		Errors.Add(CPUAccessNotAllowedError);
+	}
+	else if (AnimatedVectorField != nullptr)
 	{
-		*OutMinX.GetDest() = LocalBounds.Min.X;
-		*OutMinY.GetDest() = LocalBounds.Min.Y;
-		*OutMinZ.GetDest() = LocalBounds.Min.Z;
-		*OutMaxX.GetDest() = LocalBounds.Max.X;
-		*OutMaxY.GetDest() = LocalBounds.Max.Y;
-		*OutMaxZ.GetDest() = LocalBounds.Max.Z;
-
-		OutMinX.Advance();
-		OutMinY.Advance();
-		OutMinZ.Advance();
-		OutMaxX.Advance();
-		OutMaxY.Advance();
-		OutMaxZ.Advance();
+		FNiagaraDataInterfaceError AnimatedVectorFieldsNotSupportedError(
+			LOCTEXT("AnimatedVectorFieldsNotSupportedErrorSummary", "Invalid vector field type."),
+			LOCTEXT("AnimatedVectorFieldsNotSupportedError", "Animated vector fields are not supported."),
+			nullptr
+		);
+		Errors.Add(AnimatedVectorFieldsNotSupportedError);
+	}
+	else if (Field == nullptr)
+	{
+		FNiagaraDataInterfaceError VectorFieldNotLoadedError(
+			LOCTEXT("VectorFieldNotLoadedErrorSummary", "No Vector Field is loaded."),
+			LOCTEXT("VectorFieldNotLoadedError", "No Vector Field is loaded."),
+			nullptr
+		);
+		Errors.Add(VectorFieldNotLoadedError);
 	}
+	return Errors;
 }
+#endif // WITH_EDITOR	
 
-/*
-* Compute the influence of vector fields on a particle at the given position.
-* @param OutForce - Force to apply to the particle.
-* @param OutVelocity - Direct velocity influence on the particle.
-* @param Position - Position of the particle.
-* @param PerParticleScale - Amount by which to scale the influence on this particle.
-void EvaluateVectorFields(out float3 OutForce, out float4 OutVelocity, float3 Position, float PerParticleScale)
+/*--------------------------------------------------------------------------------------------------------------------------*/
+
+void UNiagaraDataInterfaceVectorField::GetParameterDefinitionHLSL(FNiagaraDataInterfaceGPUParamInfo& ParamInfo, FString& OutHLSL)
 {
-	float3 TotalForce = 0;
-	float3 WeightedVelocity = 0;
-	float TotalWeight = 0;
-	float FinalWeight = 0;
+	static const TCHAR *FormatDeclarations = TEXT(R"(
+		float3 {TilingAxesName};
+		float3 {DimensionsName};
+		float3 {MinBoundsName};
+		float3 {MaxBoundsName};
+		Texture3D {TextureName};
+		SamplerState {SamplerName};
+	)");
+	TMap<FString, FStringFormatArg> ArgsDeclarations = {
+		{ TEXT("TilingAxesName"), TilingAxesBaseName + ParamInfo.DataInterfaceHLSLSymbol },
+		{ TEXT("DimensionsName"), DimensionsBaseName + ParamInfo.DataInterfaceHLSLSymbol },
+		{ TEXT("MinBoundsName"),  MinBoundsBaseName + ParamInfo.DataInterfaceHLSLSymbol },
+		{ TEXT("MaxBoundsName"),  MaxBoundsBaseName + ParamInfo.DataInterfaceHLSLSymbol },
+		{ TEXT("TextureName"),    TextureBaseName + ParamInfo.DataInterfaceHLSLSymbol },
+		{ TEXT("SamplerName"),    SamplerBaseName + ParamInfo.DataInterfaceHLSLSymbol }
+	};
+	OutHLSL += FString::Format(FormatDeclarations, ArgsDeclarations);
+}
 
-	for (int VectorFieldIndex = 0; VectorFieldIndex < VectorFields.Count; ++VectorFieldIndex)
+bool UNiagaraDataInterfaceVectorField::GetFunctionHLSL(const FName& DefinitionFunctionName, FString InstanceFunctionName, FNiagaraDataInterfaceGPUParamInfo& ParamInfo, FString& OutHLSL)
+{
+	if (DefinitionFunctionName == SampleVectorFieldName)
 	{
-		float2 IntensityAndTightness = VectorFields.IntensityAndTightness[VectorFieldIndex].xy;
-		float Intensity = IntensityAndTightness.x * PerParticleScale;
-		float Tightness = IntensityAndTightness.y;
-		float3 VolumeSize = VectorFields.VolumeSize[VectorFieldIndex].xyz;
-		float3 VolumeUV = mul(float4(Position.xyz, 1), VectorFields.WorldToVolume[VectorFieldIndex]).xyz;
-		//Tile the UVs if needed. TilingAxes will be 1.0 or 0.0 in each channel depending on which axes are being tiled, if any.
-		VolumeUV -= floor(VolumeUV * VectorFields.TilingAxes[VectorFieldIndex].xyz);
-
-		float3 AxisWeights =
-			saturate(VolumeUV * VolumeSize.xyz) *
-			saturate((1.0f - VolumeUV) * VolumeSize.xyz);
-		float DistanceWeight = min(AxisWeights.x, min(AxisWeights.y, AxisWeights.z));
-
-		// @todo compat hack: Some compilers only allow constant indexing into a texture array
-		//		float3 VectorSample = Texture3DSample(VectorFieldTextures[VectorFieldIndex], VectorFieldTexturesSampler, saturate(VolumeUV)).xyz;
-		float3 VectorSample = SampleVectorFieldTexture(VectorFieldIndex, saturate(VolumeUV));
-
-		float3 Vec = mul(float4(VectorSample, 0), VectorFields.VolumeToWorld[VectorFieldIndex]).xyz;
-		TotalForce += (Vec * DistanceWeight * Intensity);
-		WeightedVelocity += (Vec * Intensity * DistanceWeight * Tightness);
-		TotalWeight += (DistanceWeight * Tightness);
-		FinalWeight = max(FinalWeight, DistanceWeight * Tightness);
+		static const TCHAR *FormatSample = TEXT(R"(
+			void {FunctionName}(float3 In_SamplePoint, out float3 Out_Sample)
+			{
+				float3 SamplePoint = (In_SamplePoint - {MinBoundsName}) / ({MaxBoundsName} - {MinBoundsName});
+				Out_Sample = Texture3DSample({TextureName}, {SamplerName}, SamplePoint).xyz;
+			}
+		)");
+		TMap<FString, FStringFormatArg> ArgsSample = {
+			{TEXT("FunctionName"), InstanceFunctionName},
+			{TEXT("TextureName"), TextureBaseName + ParamInfo.DataInterfaceHLSLSymbol},
+			{TEXT("MinBoundsName"), MinBoundsBaseName + ParamInfo.DataInterfaceHLSLSymbol},
+			{TEXT("MaxBoundsName"), MaxBoundsBaseName + ParamInfo.DataInterfaceHLSLSymbol},
+			{TEXT("SamplerName"), SamplerBaseName + ParamInfo.DataInterfaceHLSLSymbol}
+		};
+		OutHLSL += FString::Format(FormatSample, ArgsSample);
+		return true;
 	}
-
-	// Forces are additive.
-	OutForce = TotalForce;
-	// Velocities use a weighted average.
-	OutVelocity.xyz = WeightedVelocity / (TotalWeight + 0.001f);
-	OutVelocity.w = FinalWeight;
+	else if (DefinitionFunctionName == GetVectorFieldTilingAxesName)
+	{
+		static const TCHAR *FormatTilingAxes = TEXT(R"(
+			void {FunctionName}(out float3 Out_TilingAxes)
+			{
+				Out_TilingAxes = {TilingAxesName};
+			}
+		)");
+		TMap<FString, FStringFormatArg> ArgsTilingAxes = {
+			{TEXT("FunctionName"), InstanceFunctionName},
+			{TEXT("TilingAxesName"), TilingAxesBaseName + ParamInfo.DataInterfaceHLSLSymbol}
+		};
+		OutHLSL += FString::Format(FormatTilingAxes, ArgsTilingAxes);
+		return true;
+	}
+	else if (DefinitionFunctionName == GetVectorFieldDimensionsName)
+	{
+		static const TCHAR *FormatDimensions = TEXT(R"(
+			void {FunctionName}(out float3 Out_Dimensions)
+			{
+				Out_Dimensions = {DimensionsName};
+			}
+		)");
+		TMap<FString, FStringFormatArg> ArgsDimensions = {
+			{TEXT("FunctionName"), InstanceFunctionName},
+			{TEXT("DimensionsName"), DimensionsBaseName + ParamInfo.DataInterfaceHLSLSymbol}
+		};
+		OutHLSL += FString::Format(FormatDimensions, ArgsDimensions);
+		return true;
+	}
+	else if (DefinitionFunctionName == GetVectorFieldBoundsName)
+	{
+		static const TCHAR *FormatBounds = TEXT(R"(
+			void {FunctionName}(out float3 Out_MinBounds, out float3 Out_MaxBounds)
+			{
+				Out_MinBounds = {MinBoundsName};
+				Out_MaxBounds = {MaxBoundsName};
+			}
+		)");
+		TMap<FString, FStringFormatArg> ArgsBounds = {
+			{TEXT("FunctionName"), InstanceFunctionName},
+			{TEXT("MinBoundsName"), MinBoundsBaseName + ParamInfo.DataInterfaceHLSLSymbol},
+			{TEXT("MaxBoundsName"), MaxBoundsBaseName + ParamInfo.DataInterfaceHLSLSymbol}
+		};
+		OutHLSL += FString::Format(FormatBounds, ArgsBounds);
+		return true;
+	}
+	return false;
 }
-*/
 
-void UNiagaraDataInterfaceVectorField::InitField()
+struct FNiagaraDataInterfaceParametersCS_VectorField : public FNiagaraDataInterfaceParametersCS
 {
-	UVectorFieldStatic* StaticVectorField = Cast<UVectorFieldStatic>(Field);
-	UVectorFieldAnimated* AnimatedVectorField = Cast<UVectorFieldAnimated>(Field);
-	SizeX = 0;
-	SizeY = 0;
-	SizeZ = 0;
-
-	/*
-	OutParameters.WorldToVolume[Index] = VectorFieldInstance->WorldToVolume;
-	OutParameters.VolumeToWorld[Index] = VectorFieldInstance->VolumeToWorldNoScale;
-	OutParameters.VolumeSize[Index] = FVector4(Resource->SizeX, Resource->SizeY, Resource->SizeZ, 0);
-	OutParameters.IntensityAndTightness[Index] = FVector4(Intensity, Tightness, 0, 0);
-	OutParameters.TilingAxes[Index].X = VectorFieldInstance->bTileX ? 1.0f : 0.0f;
-	OutParameters.TilingAxes[Index].Y = VectorFieldInstance->bTileY ? 1.0f : 0.0f;
-	OutParameters.TilingAxes[Index].Z = VectorFieldInstance->bTileZ ? 1.0f : 0.0f;
-	*/
-	TilingAxes.X = bTileX ? 1.0f : 0.0f;
-	TilingAxes.Y = bTileY ? 1.0f : 0.0f;
-	TilingAxes.Z = bTileZ ? 1.0f : 0.0f;
+	virtual void Bind(const FNiagaraDataInterfaceParamRef& ParamRef, const class FShaderParameterMap& ParameterMap) override
+	{
+		VectorFieldSampler.Bind(ParameterMap, *(SamplerBaseName + ParamRef.ParameterInfo.DataInterfaceHLSLSymbol));
+		VectorFieldTexture.Bind(ParameterMap, *(TextureBaseName + ParamRef.ParameterInfo.DataInterfaceHLSLSymbol));
+		TilingAxes.Bind(ParameterMap, *(TilingAxesBaseName + ParamRef.ParameterInfo.DataInterfaceHLSLSymbol));
+		Dimensions.Bind(ParameterMap, *(DimensionsBaseName + ParamRef.ParameterInfo.DataInterfaceHLSLSymbol));
+		MinBounds.Bind(ParameterMap, *(MinBoundsBaseName + ParamRef.ParameterInfo.DataInterfaceHLSLSymbol));
+		MaxBounds.Bind(ParameterMap, *(MaxBoundsBaseName + ParamRef.ParameterInfo.DataInterfaceHLSLSymbol));
+	}
 
+	virtual void Serialize(FArchive& Ar)override
+	{
+		Ar << VectorFieldSampler;
+		Ar << VectorFieldTexture;
+		Ar << TilingAxes;
+		Ar << Dimensions;
+		Ar << MinBounds;
+		Ar << MaxBounds;
+	}
 
-	if (StaticVectorField)
+	virtual void Set(FRHICommandList& RHICmdList, FNiagaraShader* Shader, class UNiagaraDataInterface* DataInterface, void* PerInstanceData) const override
 	{
-		SizeX = (uint32)StaticVectorField->SizeX;
-		SizeY = (uint32)StaticVectorField->SizeY;
-		SizeZ = (uint32)StaticVectorField->SizeZ;
-		LocalBounds = StaticVectorField->Bounds;
+		check(IsInRenderingThread());
+
+		// Different sampler states used by the computer shader to sample 3D vector field. 
+		// Encoded as bitflags. To sample: 
+		//     1st bit: X-axis tiling flag
+		//     2nd bit: Y-axis tiling flag
+		//     3rd bit: Z-axis tiling flag
+		static FSamplerStateRHIParamRef SamplerStates[8] = { nullptr };
+		if (SamplerStates[0] == nullptr)
+		{
+			SamplerStates[0] = TStaticSamplerState<SF_Bilinear, AM_Clamp, AM_Clamp, AM_Clamp>::GetRHI();
+			SamplerStates[1] = TStaticSamplerState<SF_Bilinear, AM_Wrap,  AM_Clamp, AM_Clamp>::GetRHI();
+			SamplerStates[2] = TStaticSamplerState<SF_Bilinear, AM_Clamp, AM_Wrap,  AM_Clamp>::GetRHI();
+			SamplerStates[3] = TStaticSamplerState<SF_Bilinear, AM_Wrap,  AM_Wrap,  AM_Clamp>::GetRHI();
+			SamplerStates[4] = TStaticSamplerState<SF_Bilinear, AM_Clamp, AM_Clamp, AM_Wrap>::GetRHI();
+			SamplerStates[5] = TStaticSamplerState<SF_Bilinear, AM_Wrap,  AM_Clamp, AM_Wrap>::GetRHI();
+			SamplerStates[6] = TStaticSamplerState<SF_Bilinear, AM_Clamp, AM_Wrap,  AM_Wrap>::GetRHI();
+			SamplerStates[7] = TStaticSamplerState<SF_Bilinear, AM_Wrap,  AM_Wrap,  AM_Wrap>::GetRHI();
+		}
+
+		// Get shader and DI
+		const FComputeShaderRHIParamRef ComputeShaderRHI = Shader->GetComputeShader();
+		UNiagaraDataInterfaceVectorField* VFDI = CastChecked<UNiagaraDataInterfaceVectorField>(DataInterface);
+		
+		// Note: There is a flush in PreEditChange to make sure everything is synced up at this point 
+
+		// Get and set 3D texture handle from the currently bound vector field.
+		UVectorFieldStatic* StaticVectorField = Cast<UVectorFieldStatic>(VFDI->Field); 
+		FRHITexture* VolumeTextureRHI = StaticVectorField ? (FRHITexture*)StaticVectorField->GetVolumeTextureRef() : (FRHITexture*)GBlackVolumeTexture->TextureRHI;
+		SetTextureParameter(RHICmdList, ComputeShaderRHI, VectorFieldTexture, VolumeTextureRHI); 
+		
+		// Get and set sampler state
+		FSamplerStateRHIParamRef SamplerState = SamplerStates[int(VFDI->bTileX) + 2 * int(VFDI->bTileY) + 4 * int(VFDI->bTileZ)];
+		SetSamplerParameter(RHICmdList, ComputeShaderRHI, VectorFieldSampler, SamplerState);
+
+		//
+		SetShaderValue(RHICmdList, ComputeShaderRHI, TilingAxes, VFDI->GetTilingAxes());
+		SetShaderValue(RHICmdList, ComputeShaderRHI, Dimensions, VFDI->GetDimensions());
+		SetShaderValue(RHICmdList, ComputeShaderRHI, MinBounds, VFDI->GetMinBounds());
+		SetShaderValue(RHICmdList, ComputeShaderRHI, MaxBounds, VFDI->GetMaxBounds());
 	}
-	else if (AnimatedVectorField)
+
+private:
+
+	FShaderResourceParameter VectorFieldSampler;
+	FShaderResourceParameter VectorFieldTexture;
+	FShaderParameter TilingAxes;
+	FShaderParameter Dimensions;
+	FShaderParameter MinBounds;
+	FShaderParameter MaxBounds;
+};
+
+FNiagaraDataInterfaceParametersCS* UNiagaraDataInterfaceVectorField::ConstructComputeParameters()const
+{
+	return new FNiagaraDataInterfaceParametersCS_VectorField();
+}
+
+/*--------------------------------------------------------------------------------------------------------------------------*/
+
+void UNiagaraDataInterfaceVectorField::GetFieldTilingAxes(FVectorVMContext& Context)
+{
+	VectorVM::FExternalFuncRegisterHandler<float> OutSizeX(Context);
+	VectorVM::FExternalFuncRegisterHandler<float> OutSizeY(Context);
+	VectorVM::FExternalFuncRegisterHandler<float> OutSizeZ(Context);
+
+	FVector Tilings = GetTilingAxes();
+	for (int32 i = 0; i < Context.NumInstances; ++i)
 	{
-		SizeX = (uint32)AnimatedVectorField->VolumeSizeX;
-		SizeY = (uint32)AnimatedVectorField->VolumeSizeY;
-		SizeZ = (uint32)AnimatedVectorField->VolumeSizeZ;
-		LocalBounds = AnimatedVectorField->Bounds;
+		*OutSizeX.GetDest() = Tilings.X;
+		*OutSizeY.GetDest() = Tilings.Y;
+		*OutSizeZ.GetDest() = Tilings.Z;
+
+		OutSizeX.Advance();
+		OutSizeY.Advance();
+		OutSizeZ.Advance();
 	}
 }
 
-const void* UNiagaraDataInterfaceVectorField::Lock()
+void UNiagaraDataInterfaceVectorField::GetFieldDimensions(FVectorVMContext& Context)
 {
-	UVectorFieldStatic* StaticVectorField = Cast<UVectorFieldStatic>(Field);
-	UVectorFieldAnimated* AnimatedVectorField = Cast<UVectorFieldAnimated>(Field);
+	VectorVM::FExternalFuncRegisterHandler<float> OutSizeX(Context);
+	VectorVM::FExternalFuncRegisterHandler<float> OutSizeY(Context);
+	VectorVM::FExternalFuncRegisterHandler<float> OutSizeZ(Context);
 
-	if (StaticVectorField)
+	FVector Dim = GetDimensions();
+	for (int32 i = 0; i < Context.NumInstances; ++i)
 	{
-		return StaticVectorField->SourceData.LockReadOnly();
+		*OutSizeX.GetDest() = Dim.X;
+		*OutSizeY.GetDest() = Dim.Y;
+		*OutSizeZ.GetDest() = Dim.Z;
+
+		OutSizeX.Advance();
+		OutSizeY.Advance();
+		OutSizeZ.Advance();
 	}
-	return nullptr;
 }
 
-void UNiagaraDataInterfaceVectorField::Unlock()
+void UNiagaraDataInterfaceVectorField::GetFieldBounds(FVectorVMContext& Context)
 {
-	UVectorFieldStatic* StaticVectorField = Cast<UVectorFieldStatic>(Field);
-	UVectorFieldAnimated* AnimatedVectorField = Cast<UVectorFieldAnimated>(Field);
+	VectorVM::FExternalFuncRegisterHandler<float> OutMinX(Context);
+	VectorVM::FExternalFuncRegisterHandler<float> OutMinY(Context);
+	VectorVM::FExternalFuncRegisterHandler<float> OutMinZ(Context);
+	VectorVM::FExternalFuncRegisterHandler<float> OutMaxX(Context);
+	VectorVM::FExternalFuncRegisterHandler<float> OutMaxY(Context);
+	VectorVM::FExternalFuncRegisterHandler<float> OutMaxZ(Context);
 
-	if (StaticVectorField)
+	FVector MinBounds = GetMinBounds();
+	FVector MaxBounds = GetMaxBounds();
+	for (int32 i = 0; i < Context.NumInstances; ++i)
 	{
-		StaticVectorField->SourceData.Unlock();
+		*OutMinX.GetDest() = MinBounds.X;
+		*OutMinY.GetDest() = MinBounds.Y;
+		*OutMinZ.GetDest() = MinBounds.Z;
+		*OutMaxX.GetDest() = MaxBounds.X;
+		*OutMaxY.GetDest() = MaxBounds.Y;
+		*OutMaxZ.GetDest() = MaxBounds.Z;
+
+		OutMinX.Advance();
+		OutMinY.Advance();
+		OutMinZ.Advance();
+		OutMaxX.Advance();
+		OutMaxY.Advance();
+		OutMaxZ.Advance();
 	}
 }
 
@@ -567,39 +466,103 @@ void UNiagaraDataInterfaceVectorField::SampleVectorField(FVectorVMContext& Conte
 	VectorVM::FExternalFuncInputHandler<float> YParam(Context);
 	VectorVM::FExternalFuncInputHandler<float> ZParam(Context);
 
-	// User pointer buffer (as requested)
-	//VectorVM::FUserPtrHandler<FNDIVectorField_InstanceData> InstData(Context);
-
 	// Outputs...
 	VectorVM::FExternalFuncRegisterHandler<float> OutSampleX(Context);
 	VectorVM::FExternalFuncRegisterHandler<float> OutSampleY(Context);
 	VectorVM::FExternalFuncRegisterHandler<float> OutSampleZ(Context);
 
-	/*UE_LOG(LogNiagara, Log, TEXT("Registers: In: %d %d %d %d Out: %d %d %d"),
-		XParam.RegisterIndex, YParam.RegisterIndex, ZParam.RegisterIndex, InstData.RegisterIndex,
-		OutSampleX.RegisterIndex, OutSampleY.RegisterIndex, OutSampleZ.RegisterIndex);*/
-
 	UVectorFieldStatic* StaticVectorField = Cast<UVectorFieldStatic>(Field);
 	UVectorFieldAnimated* AnimatedVectorField = Cast<UVectorFieldAnimated>(Field);
 
-	//FMatrix VolumeToWorld;
-	//FMatrix WorldToVolume;
-	//FMatrix LocalToWorld = FMatrix::Identity;
-	//InstData->UpdateTransforms(LocalToWorld, VolumeToWorld, WorldToVolume);
+	if (StaticVectorField != nullptr && StaticVectorField->bAllowCPUAccess)
+	{
+		const FVector4 TilingAxes = FVector4(bTileX ? 1.0f : 0.0f, bTileY ? 1.0f : 0.0f, bTileZ ? 1.0f : 0.0f, 0.0);
 
-	const FFloat16Color* Data = (const FFloat16Color*)Lock();
+		const uint32 SizeX = (uint32)StaticVectorField->SizeX;
+		const uint32 SizeY = (uint32)StaticVectorField->SizeY;
+		const uint32 SizeZ = (uint32)StaticVectorField->SizeZ;
+		const FVector4 Size(SizeX, SizeY, SizeZ, 1.0f);
+
+		const FVector4 MinBounds(StaticVectorField->Bounds.Min.X, StaticVectorField->Bounds.Min.Y, StaticVectorField->Bounds.Min.Z, 1.0f);
+		const FVector4 MaxBounds(StaticVectorField->Bounds.Max.X, StaticVectorField->Bounds.Max.Y, StaticVectorField->Bounds.Max.Z, 1.0f);
+
+		const FVector4 *Data = StaticVectorField->CPUData.GetData();
+
+		check(Data != nullptr); // TODO(mv): Guard better? This shouldn't happen 
+
+		// Math helper
+		static auto FVector4Clamp = [](FVector4 v, FVector4 a, FVector4 b) {
+			return FVector4(FMath::Clamp(v.X, a.X, b.X),
+				FMath::Clamp(v.Y, a.Y, b.Y),
+				FMath::Clamp(v.Z, a.Z, b.Z),
+				FMath::Clamp(v.W, a.W, b.W));
+		};
+
+		static auto FVector4Floor = [](FVector4 v) {
+			return FVector4(FGenericPlatformMath::FloorToFloat(v.X),
+				FGenericPlatformMath::FloorToFloat(v.Y),
+				FGenericPlatformMath::FloorToFloat(v.Z),
+				FGenericPlatformMath::FloorToFloat(v.W));
+		};
 
-	if (Data == nullptr || Field == nullptr || (StaticVectorField == nullptr && AnimatedVectorField == nullptr) || AnimatedVectorField != nullptr)
-	{
-		// Set the default vector to positive X axis...
 		for (int32 InstanceIdx = 0; InstanceIdx < Context.NumInstances; ++InstanceIdx)
 		{
-			VectorRegister Dst = MakeVectorRegister(1.0f, 0.0f, 0.0f, 0.0f);
+			// Position in Volume Space
+			FVector4 Pos(XParam.Get(), YParam.Get(), ZParam.Get(), 0.0f);
+
+			// Normalize position
+			Pos = (Pos - MinBounds) / (MaxBounds - MinBounds);
+
+			// Scaled position
+			Pos = Pos * Size;
+
+			// Offset by half a cell size due to sample being in the center of its cell
+			Pos = Pos - FVector4(0.5f, 0.5f, 0.5f, 0.0f);
+
+			// 
+			FVector4 Index0 = FVector4Floor(Pos);
+			FVector4 Index1 = Index0 + FVector4(1.0f, 1.0f, 1.0f, 0.0f);
 
-			float *RegPtr = reinterpret_cast<float*>(&Dst);
-			*OutSampleX.GetDest() = RegPtr[0];
-			*OutSampleY.GetDest() = RegPtr[1];
-			*OutSampleZ.GetDest() = RegPtr[2];
+			// 
+			FVector4 Fraction = Pos - Index0;
+
+			Index0 = Index0 - TilingAxes*FVector4Floor(Index0 / Size)*Size;
+			Index1 = Index1 - TilingAxes*FVector4Floor(Index1 / Size)*Size;
+
+			Index0 = FVector4Clamp(Index0, FVector4(0.0f), Size - FVector4(1.0f, 1.0f, 1.0f, 0.0f));
+			Index1 = FVector4Clamp(Index1, FVector4(0.0f), Size - FVector4(1.0f, 1.0f, 1.0f, 0.0f));
+
+			// Sample by regular trilinear interpolation:
+
+			// TODO(mv): Optimize indexing for cache? Periodicity is problematic...
+			// TODO(mv): Vectorize?
+			// Fetch corners
+			FVector4 V000 = Data[int(Index0.X + SizeX * Index0.Y + SizeX * SizeY * Index0.Z)];
+			FVector4 V100 = Data[int(Index1.X + SizeX * Index0.Y + SizeX * SizeY * Index0.Z)];
+			FVector4 V010 = Data[int(Index0.X + SizeX * Index1.Y + SizeX * SizeY * Index0.Z)];
+			FVector4 V110 = Data[int(Index1.X + SizeX * Index1.Y + SizeX * SizeY * Index0.Z)];
+			FVector4 V001 = Data[int(Index0.X + SizeX * Index0.Y + SizeX * SizeY * Index1.Z)];
+			FVector4 V101 = Data[int(Index1.X + SizeX * Index0.Y + SizeX * SizeY * Index1.Z)];
+			FVector4 V011 = Data[int(Index0.X + SizeX * Index1.Y + SizeX * SizeY * Index1.Z)];
+			FVector4 V111 = Data[int(Index1.X + SizeX * Index1.Y + SizeX * SizeY * Index1.Z)];
+
+			// Blend x-axis
+			FVector4 V00 = FMath::Lerp(V000, V100, Fraction.X);
+			FVector4 V01 = FMath::Lerp(V001, V101, Fraction.X);
+			FVector4 V10 = FMath::Lerp(V010, V110, Fraction.X);
+			FVector4 V11 = FMath::Lerp(V011, V111, Fraction.X);
+
+			// Blend y-axis
+			FVector4 V0 = FMath::Lerp(V00, V10, Fraction.Y);
+			FVector4 V1 = FMath::Lerp(V01, V11, Fraction.Y);
+
+			// Blend z-axis
+			FVector4 V = FMath::Lerp(V0, V1, Fraction.Z);
+
+			// Write final output...
+			*OutSampleX.GetDest() = V.X;
+			*OutSampleY.GetDest() = V.Y;
+			*OutSampleZ.GetDest() = V.Z;
 
 			XParam.Advance();
 			YParam.Advance();
@@ -609,338 +572,93 @@ void UNiagaraDataInterfaceVectorField::SampleVectorField(FVectorVMContext& Conte
 			OutSampleZ.Advance();
 		}
 	}
-	else if (StaticVectorField != nullptr)
+	else
 	{
-		FVector VolumeSize((float)SizeX, (float)SizeY, (float)SizeZ);
-		uint32 VolumeSizeX = (uint32)SizeX;
-		uint32 VolumeSizeY = (uint32)SizeY;
-		uint32 VolumeSizeZ = (uint32)SizeZ;
+		// TODO(mv): Add warnings?
+		if (StaticVectorField != nullptr && !StaticVectorField->bAllowCPUAccess)
+		{
+			// No access to static vector data
+		}
+		else if (AnimatedVectorField != nullptr)
+		{
+			// Animated vector field not supported
+		}
+		else if (Field == nullptr)
+		{
+			// Vector field not loaded
+		}
 
+		// Set the default vector to positive X axis corresponding to a velocity of 100 cm/s
+		// Rationale: Setting to the zero vector can be visually confusing and likely to cause problems elsewhere
 		for (int32 InstanceIdx = 0; InstanceIdx < Context.NumInstances; ++InstanceIdx)
 		{
-			FVector Pos(XParam.Get(), YParam.Get(), ZParam.Get());
-			
-			// float3 VolumeUV = mul(float4(Position.xyz, 1), VectorFields.WorldToVolume[VectorFieldIndex]).xyz;
-			//FNDITransformHandler Handler;
-			FVector VolumeUV = Pos;
-			//Handler.TransformPosition(VolumeUV, WorldToVolume);
-
-			// VolumeUV -= floor(VolumeUV * VectorFields.TilingAxes[VectorFieldIndex].xyz);
-			FVector TiledVolume = VolumeUV * TilingAxes;
-			VolumeUV.X -= FMath::FloorToFloat(TiledVolume.X);
-			VolumeUV.Y -= FMath::FloorToFloat(TiledVolume.Y);
-			VolumeUV.Z -= FMath::FloorToFloat(TiledVolume.Z);
-
-			// Saturate(VolumeUV) 
-			VolumeUV.X = FMath::Clamp(VolumeUV.X, 0.0f, 1.0f);
-			VolumeUV.Y = FMath::Clamp(VolumeUV.Y, 0.0f, 1.0f);
-			VolumeUV.Z = FMath::Clamp(VolumeUV.Z, 0.0f, 1.0f);
-
-			// Sample texture...
-			VolumeUV = VolumeUV * VolumeSize;
-
-			/**
-			int3 IntCoord = int3(ModXYZ.x, ModXYZ.y, ModXYZ.z);\n");
-			float3 frc = frac(ModXYZ);\n");
-			*/
-			FVector VecIndex;
-			FVector VecFrac;
-			VecFrac.X = FGenericPlatformMath::Modf(VolumeUV.X, &VecIndex.X);
-			VecFrac.Y = FGenericPlatformMath::Modf(VolumeUV.Y, &VecIndex.Y);
-			VecFrac.Z = FGenericPlatformMath::Modf(VolumeUV.Z, &VecIndex.Z);
-
-			uint32 IndexX = (uint32)VecIndex.X;
-			uint32 IndexY = (uint32)VecIndex.Y;
-			uint32 IndexZ = (uint32)VecIndex.Z;
-
-#define VF_MakeIndex(X, Y, Z) ((X)  + ((Y) * VolumeSizeX) + ((Z) * VolumeSizeX * VolumeSizeY))
-#define VF_Tile(I, Range)  (((I) >= Range) ? 0 : (I))
-
-			// Trilinear interpolation, as defined by https://en.wikipedia.org/wiki/Trilinear_interpolation
-			/** 
-			float3 frc = frac(ModXYZ);
-			float3 V1 = Buffer[IntCoord.x + IntCoord.y*VolumeSizeX + IntCoord.z*VolumeSizeX*VolumeSizeY].xyz;  // V(x0, y0, z0)
-			float3 V2 = Buffer[IntCoord.x + 1 + IntCoord.y*VolumeSizeX + IntCoord.z*VolumeSizeX*VolumeSizeY].xyz; // V(x1, y0, z0)
-			float3 C00 = lerp(V1, V2, frc.xxx);
-			V1 = Buffer[IntCoord.x + IntCoord.y*VolumeSizeX + (IntCoord.z+1)*VolumeSizeX*VolumeSizeY].xyz;  // V(x0, y0, z1)
-			V2 = Buffer[IntCoord.x + 1 + IntCoord.y*VolumeSizeX + (IntCoord.z+1)*VolumeSizeX*VolumeSizeY].xyz; // V(x1, y0, z1)
-			float3 C01 = lerp(V1, V2, frc.xxx);
-			V1 = Buffer[IntCoord.x + (IntCoord.y+1)*VolumeSizeX + IntCoord.z*VolumeSizeX*VolumeSizeY].xyz;  // V(x0, y1, z0)
-			V2 = Buffer[IntCoord.x + 1 + (IntCoord.y+1)*VolumeSizeX + IntCoord.z*VolumeSizeX*VolumeSizeY].xyz; // V(x1, y1, z0)
-			float3 C10 = lerp(V1, V2, frc.xxx);
-			V1 = Buffer[IntCoord.x + (IntCoord.y+1)*VolumeSizeX + (IntCoord.z+1)*VolumeSizeX*VolumeSizeY].xyz;  // V(x0, y1, z1)
-			V2 = Buffer[IntCoord.x + 1 + (IntCoord.y+1)*VolumeSizeX + (IntCoord.z+1)*VolumeSizeX*VolumeSizeY].xyz; // V(x1, y1, z1)
-			float3 C11 = lerp(V1, V2, frc.xxx);
-			float3 C0 = lerp(C00, C10, frc.yyy);
-			float3 C1 = lerp(c01, C11, frc.yyy);
-			Out_Value = lerp(C0, C1, frc.zzz);
-			*/
-			uint32 Index = VF_MakeIndex(VF_Tile(IndexX, VolumeSizeX), VF_Tile(IndexY, VolumeSizeY), VF_Tile(IndexZ, VolumeSizeZ));
-			check(Index < (uint32)(VolumeSizeX*VolumeSizeY*VolumeSizeZ));
-			FVector V1 ((float)Data[Index].R, (float)Data[Index].G, (float)Data[Index].B);  // V(x0, y0, z0)
-			Index = VF_MakeIndex(VF_Tile(IndexX + 1, VolumeSizeX), VF_Tile(IndexY, VolumeSizeY), VF_Tile(IndexZ, VolumeSizeZ));
-			check(Index < (uint32)(VolumeSizeX*VolumeSizeY*VolumeSizeZ));
-			FVector V2((float)Data[Index].R, (float)Data[Index].G, (float)Data[Index].B); // V(x1, y0, z0)
-			FVector C00 = FMath::Lerp(V1, V2, VecFrac.X);
-
-			Index = VF_MakeIndex(IndexX, VF_Tile(IndexY, VolumeSizeY), VF_Tile(IndexZ + 1, VolumeSizeZ));
-			check(Index < (uint32)(VolumeSizeX*VolumeSizeY*VolumeSizeZ));
-			V1 = FVector((float)Data[Index].R, (float)Data[Index].G, (float)Data[Index].B);// V(x0, y0, z1)
-			Index = VF_MakeIndex(VF_Tile(IndexX + 1, VolumeSizeX), VF_Tile(IndexY, VolumeSizeY), VF_Tile(IndexZ + 1, VolumeSizeZ));
-			check(Index < (uint32)(VolumeSizeX*VolumeSizeY*VolumeSizeZ));
-			V1 = FVector((float)Data[Index].R, (float)Data[Index].G, (float)Data[Index].B);// V(x1, y0, z1)
-			FVector C01 = FMath::Lerp(V1, V2, VecFrac.X);
-
-			Index = VF_MakeIndex(VF_Tile(IndexX, VolumeSizeX), VF_Tile(IndexY + 1, VolumeSizeY), VF_Tile(IndexZ, VolumeSizeZ));
-			check(Index < (uint32)(VolumeSizeX*VolumeSizeY*VolumeSizeZ));
-			V1 = FVector((float)Data[Index].R, (float)Data[Index].G, (float)Data[Index].B);// V(x0, y1, z0)
-			Index = VF_MakeIndex(VF_Tile(IndexX + 1, VolumeSizeX), VF_Tile(IndexY + 1, VolumeSizeY), VF_Tile(IndexZ, VolumeSizeZ));
-			check(Index < (uint32)(VolumeSizeX*VolumeSizeY*VolumeSizeZ));
-			V1 = FVector((float)Data[Index].R, (float)Data[Index].G, (float)Data[Index].B); // V(x1, y1, z0)
-			FVector C10 = FMath::Lerp(V1, V2, VecFrac.X);
-
-			Index = VF_MakeIndex(VF_Tile(IndexX, VolumeSizeX), VF_Tile(IndexY + 1, VolumeSizeY), VF_Tile(IndexZ + 1, VolumeSizeZ));
-			check(Index < (uint32)(VolumeSizeX*VolumeSizeY*VolumeSizeZ));
-			V1 = FVector((float)Data[Index].R, (float)Data[Index].G, (float)Data[Index].B);// V(x0, y1, z1)
-			Index = VF_MakeIndex(VF_Tile(IndexX + 1, VolumeSizeX), VF_Tile(IndexY + 1, VolumeSizeY), VF_Tile(IndexZ + 1, VolumeSizeZ));
-			check(Index < (uint32)(VolumeSizeX*VolumeSizeY*VolumeSizeZ));
-			V1 = FVector((float)Data[Index].R, (float)Data[Index].G, (float)Data[Index].B);// V(x1, y1, z1)
-			FVector C11 = FMath::Lerp(V1, V2, VecFrac.X);
-
-			FVector C0 = FMath::Lerp(C00, C10, VecFrac.Y);
-			FVector C1 = FMath::Lerp(C01, C11, VecFrac.Y);
-			FVector C = FMath::Lerp(C0, C1, VecFrac.Z);
-
-			// Write final output...
-			*OutSampleX.GetDest() = C.X;
-			*OutSampleY.GetDest() = C.Y;
-			*OutSampleZ.GetDest() = C.Z;
+			*OutSampleX.GetDest() = 0.0f;
+			*OutSampleY.GetDest() = 0.0f;
+			*OutSampleZ.GetDest() = 0.0f;
 
 			XParam.Advance();
 			YParam.Advance();
-			ZParam.Advance();
+			ZParam.Advance();	
 			OutSampleX.Advance();
 			OutSampleY.Advance();
 			OutSampleZ.Advance();
 		}
 	}
 
-	Unlock();
 }
 
+/*--------------------------------------------------------------------------------------------------------------------------*/
 
-bool UNiagaraDataInterfaceVectorField::CanExecuteOnTarget(ENiagaraSimTarget Target)const 
+FVector UNiagaraDataInterfaceVectorField::GetTilingAxes() const
 {
-	return true; 
+	return FVector(float(bTileX), float(bTileY), float(bTileZ));
 }
 
-bool UNiagaraDataInterfaceVectorField::CopyToInternal(UNiagaraDataInterface* Destination) const
+FVector UNiagaraDataInterfaceVectorField::GetDimensions() const
 {
-	if (!Super::CopyToInternal(Destination))
+	UVectorFieldStatic* StaticVectorField = Cast<UVectorFieldStatic>(Field);
+	if (StaticVectorField)
 	{
-		return false;
+		return FVector(StaticVectorField->SizeX, StaticVectorField->SizeY, StaticVectorField->SizeZ);
 	}
-
-	UNiagaraDataInterfaceVectorField* OtherTyped = CastChecked<UNiagaraDataInterfaceVectorField>(Destination);
-	OtherTyped->Field = Field;
-	OtherTyped->bTileX = bTileX;
-	OtherTyped->bTileY = bTileY;
-	OtherTyped->bTileZ = bTileZ;
-	OtherTyped->bGPUBufferDirty = true;
-	return true;
+	return FVector{ 1.0f, 1.0f, 1.0f }; // Matches GBlackVolumeTexture
 }
 
-bool UNiagaraDataInterfaceVectorField::Equals(const UNiagaraDataInterface* Other) const
+FVector UNiagaraDataInterfaceVectorField::GetMinBounds() const
 {
-	if (!Super::Equals(Other))
+	UVectorFieldStatic* StaticVectorField = Cast<UVectorFieldStatic>(Field);
+	if (StaticVectorField)
 	{
-		return false;
+		return StaticVectorField->Bounds.Min;
 	}
-	const UNiagaraDataInterfaceVectorField* OtherTyped = CastChecked<const UNiagaraDataInterfaceVectorField>(Other);
-	return OtherTyped->Field == Field 
-		&& OtherTyped->bTileX == bTileX
-		&& OtherTyped->bTileY == bTileY
-		&& OtherTyped->bTileZ == bTileZ;
-}
-
-
-int32 UNiagaraDataInterfaceVectorField::PerInstanceDataSize()const  
-{
-	return 0;
-}
-
-bool UNiagaraDataInterfaceVectorField::InitPerInstanceData(void* PerInstanceData, FNiagaraSystemInstance* SystemInstance)
-{
-	//FNDIVectorField_InstanceData* Inst = new (PerInstanceData) FNDIVectorField_InstanceData();
-	//return Inst->Init(this, SystemInstance);
-	return false;
-}
-
-void UNiagaraDataInterfaceVectorField::DestroyPerInstanceData(void* PerInstanceData, FNiagaraSystemInstance* SystemInstance)
-{
-	/*FNDIVectorField_InstanceData* Inst = (FNDIVectorField_InstanceData*)PerInstanceData;
-	Inst->~FNDIVectorField_InstanceData();*/
+	return FVector{-1.0f, -1.0f, -1.0f};
 }
 
-bool UNiagaraDataInterfaceVectorField::PerInstanceTick(void* PerInstanceData, FNiagaraSystemInstance* SystemInstance, float InDeltaSeconds)
+FVector UNiagaraDataInterfaceVectorField::GetMaxBounds() const
 {
-	/*FNDIVectorField_InstanceData* Inst = (FNDIVectorField_InstanceData*)PerInstanceData;
-	return Inst->Tick(this, SystemInstance, InDeltaSeconds);*/
-	return false;
-}
-
-bool UNiagaraDataInterfaceVectorField::GetFunctionHLSL(const FName& DefinitionFunctionName, FString InstanceFunctionName, FNiagaraDataInterfaceGPUParamInfo& ParamInfo, FString& OutHLSL)
-{
-	if (DefinitionFunctionName == SampleVectorFieldName)
-	{
-		FString BufferName = "VectorFieldLUT_" + ParamInfo.DataInterfaceHLSLSymbol;
-		OutHLSL += TEXT("void ") + InstanceFunctionName + TEXT("(float3 In_SamplePoint,  out float3 Out_Value) \n{\n");
-		OutHLSL += TEXT("\t int3 IntCoord = int3(0,0,0);\n");
-		OutHLSL += TEXT("\t int3 Dimensions = int3(1,1,1);\n");
-		OutHLSL += TEXT("\t Out_Value = ") + BufferName + TEXT("[IntCoord.x + IntCoord.y*Dimensions.x + IntCoord.z*Dimensions.x*Dimensions.y].xyz;\n");
-		OutHLSL += TEXT("\n}\n");
-		return true;
-	}
-	else if (DefinitionFunctionName == GetVectorDimsName)
-	{
-		FString DimsVar = DimentionsBaseName + ParamInfo.DataInterfaceHLSLSymbol;
-		OutHLSL += TEXT("void ") + InstanceFunctionName + TEXT("(out float3 Out_Value) \n{\n");
-		OutHLSL += TEXT("\t Out_Value = ") + DimsVar + TEXT(";\n");
-		OutHLSL += TEXT("\n}\n");
-		return true;
-	}
-	else if (DefinitionFunctionName == GetVectorFieldBoundsName)
+	UVectorFieldStatic* StaticVectorField = Cast<UVectorFieldStatic>(Field);
+	if (StaticVectorField)
 	{
-		FString BoundsMinVar = BoundsMinBaseName + ParamInfo.DataInterfaceHLSLSymbol;
-		FString BoundsMaxVar = BoundsMaxBaseName + ParamInfo.DataInterfaceHLSLSymbol;
-		OutHLSL += TEXT("void ") + InstanceFunctionName + TEXT("(out float3 Out_MinBounds, out float3 Out_MaxBounds) \n{\n");
-		OutHLSL += TEXT("\t Out_MinBounds = ") + BoundsMinVar + TEXT(";\n");
-		OutHLSL += TEXT("\t Out_MaxBounds = ") + BoundsMaxVar + TEXT(";\n");
-		OutHLSL += TEXT("\n}\n");
-		return true;
+		return StaticVectorField->Bounds.Max;
 	}
-	return false;
+	return FVector{1.0f, 1.0f, 1.0f};
 }
 
-void UNiagaraDataInterfaceVectorField::GetParameterDefinitionHLSL(FNiagaraDataInterfaceGPUParamInfo& ParamInfo, FString& OutHLSL)
-{
-	OutHLSL += TEXT("Buffer<float4> ") + BufferBaseName + ParamInfo.DataInterfaceHLSLSymbol + TEXT(";\n");
-	OutHLSL += TEXT("float3 ") + DimentionsBaseName + ParamInfo.DataInterfaceHLSLSymbol + TEXT(";\n");
-}
+/*--------------------------------------------------------------------------------------------------------------------------*/
 
-FRWBuffer& UNiagaraDataInterfaceVectorField::GetGPUBuffer()
-{
-	check(IsInRenderingThread());
-
-	//TODO: Doing this here is not really threadsafe. Need to move to an RT proxy style model?
-	if (bGPUBufferDirty)
-	{
-		UVectorFieldStatic* StaticVectorField = Cast<UVectorFieldStatic>(Field);
-		UVectorFieldAnimated* AnimatedVectorField = Cast<UVectorFieldAnimated>(Field);
-
-		uint32 Width = 10;
-		uint32 Height = 10;
-		uint32 Depth = 10;
-		const FFloat16Color* Data = nullptr;
-
-		if (StaticVectorField)
-		{
-			Width = StaticVectorField->SizeX;
-			Height = StaticVectorField->SizeY;
-			Depth = StaticVectorField->SizeZ;
-		}
-		
-		GPUBuffer.Release();
-		uint32 BufferSize = Width * Height * Depth * sizeof(float) * 4;
-		TResourceArray<FVector4> TempTable;
-		TempTable.AddUninitialized(Width*Height*Depth);
-		
-		if (StaticVectorField)
-		{
-			Data = (const FFloat16Color*)StaticVectorField->SourceData.LockReadOnly();
-
-			// TODO sckime - can we just send this down as half4's ???
-			for (uint32 z = 0; z < Depth; z++)
-			{
-				for (uint32 y = 0; y < Height; y++)
-				{
-					for (uint32 x = 0; x < Width; x++)
-					{
-						uint32 Index = x + (y* Height) + (z * Width * Height);
-						TempTable[Index] = FVector4(Data[Index].R, Data[Index].G, Data[Index].B, 0.0f);
-					}
-				}
-			}
-			StaticVectorField->SourceData.Unlock();
-		}
-		else
-		{
-			// We don't support the other types for now.
-			for (uint32 z = 0; z < Depth; z++)
-			{
-				for (uint32 y = 0; y < Height; y++)
-				{
-					for (uint32 x = 0; x < Width; x++)
-					{
-						uint32 Index = x + (y* Height) + (z * Width * Height);
-						TempTable[Index] = FVector4(1.0f, 0.0f, 0.0f, 0.0f);
-					}
-				}
-			}
-		}
-
-		GPUBuffer.Initialize(sizeof(float) * 4, Width*Height*Depth, EPixelFormat::PF_A32B32G32R32F, BUF_Static, TEXT("CurlnoiseTable"), &TempTable);
-		//FPlatformMemory::Memcpy(BufferData, TempTable, BufferSize);
-		//RHIUnlockVertexBuffer(GPUBuffer.Buffer.Buffer);
-		bGPUBufferDirty = false;
-	}
-
-	return GPUBuffer;
-}
-
-struct FNiagaraDataInterfaceParametersCS_VectorField : public FNiagaraDataInterfaceParametersCS
+bool UNiagaraDataInterfaceVectorField::CopyToInternal(UNiagaraDataInterface* Destination) const
 {
-	virtual void Bind(const FNiagaraDataInterfaceParamRef& ParamRef, const class FShaderParameterMap& ParameterMap) override
-	{
-		VectorFieldBuffer.Bind(ParameterMap, *(UNiagaraDataInterfaceVectorField::BufferBaseName + ParamRef.ParameterInfo.DataInterfaceHLSLSymbol));
-		Dimentions.Bind(ParameterMap, *(UNiagaraDataInterfaceVectorField::DimentionsBaseName + ParamRef.ParameterInfo.DataInterfaceHLSLSymbol));
-		BoundsMin.Bind(ParameterMap, *(UNiagaraDataInterfaceVectorField::BoundsMinBaseName + ParamRef.ParameterInfo.DataInterfaceHLSLSymbol));
-		BoundsMax.Bind(ParameterMap, *(UNiagaraDataInterfaceVectorField::BoundsMaxBaseName + ParamRef.ParameterInfo.DataInterfaceHLSLSymbol));
-	}
-
-	virtual void Serialize(FArchive& Ar)override
-	{
-		Ar << VectorFieldBuffer;
-		Ar << Dimentions;
-		Ar << BoundsMin;
-		Ar << BoundsMax;
-	}
-
-	virtual void Set(FRHICommandList& RHICmdList, FNiagaraShader* Shader, class UNiagaraDataInterface* DataInterface) const override 
+	if (!Super::CopyToInternal(Destination))
 	{
-		check(IsInRenderingThread());
-
-		const FComputeShaderRHIParamRef ComputeShaderRHI = Shader->GetComputeShader();
-		UNiagaraDataInterfaceVectorField* VFDI = CastChecked<UNiagaraDataInterfaceVectorField>(DataInterface);
-		FRWBuffer& VFBuffer = VFDI->GetGPUBuffer();
-
-		RHICmdList.SetShaderResourceViewParameter(ComputeShaderRHI, VectorFieldBuffer.GetBaseIndex(), VFBuffer.SRV);
-		SetShaderValue(RHICmdList, ComputeShaderRHI, Dimentions, VFDI->GetDimentions());
-		SetShaderValue(RHICmdList, ComputeShaderRHI, BoundsMin, VFDI->GetBoundsMin());
-		SetShaderValue(RHICmdList, ComputeShaderRHI, BoundsMax, VFDI->GetBoundsMax());
+		return false;
 	}
 
-private:
-
-	FShaderResourceParameter VectorFieldBuffer;
-	FShaderParameter Dimentions;
-	FShaderParameter BoundsMin;
-	FShaderParameter BoundsMax;
-};
-
-FNiagaraDataInterfaceParametersCS* UNiagaraDataInterfaceVectorField::ConstructComputeParameters()const
-{
-	return new FNiagaraDataInterfaceParametersCS_VectorField();
+	UNiagaraDataInterfaceVectorField* OtherTyped = CastChecked<UNiagaraDataInterfaceVectorField>(Destination);
+	OtherTyped->Field = Field;
+	OtherTyped->bTileX = bTileX;
+	OtherTyped->bTileY = bTileY;
+	OtherTyped->bTileZ = bTileZ;
+	return true;
 }
 
 #undef LOCTEXT_NAMESPACE
diff --git a/Engine/Plugins/FX/Niagara/Source/Niagara/Private/NiagaraDataSet.cpp b/Engine/Plugins/FX/Niagara/Source/Niagara/Private/NiagaraDataSet.cpp
--- a/Engine/Plugins/FX/Niagara/Source/Niagara/Private/NiagaraDataSet.cpp
+++ b/Engine/Plugins/FX/Niagara/Source/Niagara/Private/NiagaraDataSet.cpp
@@ -13,16 +13,21 @@ DECLARE_CYCLE_STAT(TEXT("InitRenderData"), STAT_InitRenderData, STATGROUP_Niagar
 
 
 //////////////////////////////////////////////////////////////////////////
-void FNiagaraDataSet::SetShaderParams(FNiagaraShader *Shader, FRHICommandList &CommandList)
+void FNiagaraDataSet::SetShaderParams(FNiagaraShader *Shader, FRHICommandList &CommandList, uint32& WriteBufferIdx, uint32& ReadBufferIdx)
 {
 	check(IsInRenderingThread());
-
+	WriteBufferIdx = INDEX_NONE;
+	ReadBufferIdx = INDEX_NONE;
 	if (Shader->FloatInputBufferParam.IsBound())
 	{
 		CommandList.TransitionResource(EResourceTransitionAccess::EReadable, EResourceTransitionPipeline::EComputeToCompute, PrevData().GetGPUBufferFloat()->UAV);
 		if (PrevData().GetNumInstancesAllocated() > 0)
 		{
 			CommandList.SetShaderResourceViewParameter(Shader->GetComputeShader(), Shader->FloatInputBufferParam.GetBaseIndex(), PrevData().GetGPUBufferFloat()->SRV);
+			if (ReadBufferIdx == INDEX_NONE)
+			{
+				ReadBufferIdx = GetPrevBufferIdx();
+			}
 		}
 		else
 		{
@@ -35,6 +40,10 @@ void FNiagaraDataSet::SetShaderParams(FNiagaraShader *Shader, FRHICommandList &C
 		if (PrevData().GetNumInstancesAllocated() > 0)
 		{
 			CommandList.SetShaderResourceViewParameter(Shader->GetComputeShader(), Shader->IntInputBufferParam.GetBaseIndex(), PrevData().GetGPUBufferInt()->SRV);
+			if (ReadBufferIdx == INDEX_NONE)
+			{
+				ReadBufferIdx = GetPrevBufferIdx();
+			}
 		}
 		else
 		{
@@ -45,11 +54,20 @@ void FNiagaraDataSet::SetShaderParams(FNiagaraShader *Shader, FRHICommandList &C
 	{
 		CommandList.TransitionResource(EResourceTransitionAccess::EWritable, EResourceTransitionPipeline::EGfxToCompute, CurrData().GetGPUBufferFloat()->UAV);
 		CommandList.SetUAVParameter(Shader->GetComputeShader(), Shader->FloatOutputBufferParam.GetUAVIndex(), CurrData().GetGPUBufferFloat()->UAV);
+
+		if (WriteBufferIdx == INDEX_NONE)
+		{
+			WriteBufferIdx = GetCurrBufferIdx();
+		}
 	}
 	if (Shader->IntOutputBufferParam.IsUAVBound())
 	{
 		CommandList.TransitionResource(EResourceTransitionAccess::EWritable, EResourceTransitionPipeline::EGfxToCompute, CurrData().GetGPUBufferInt()->UAV);
 		CommandList.SetUAVParameter(Shader->GetComputeShader(), Shader->IntOutputBufferParam.GetUAVIndex(), CurrData().GetGPUBufferInt()->UAV);
+		if (WriteBufferIdx == INDEX_NONE)
+		{
+			WriteBufferIdx = GetCurrBufferIdx();
+		}
 	}
 
 	if (Shader->ComponentBufferSizeWriteParam.IsBound())
@@ -131,6 +149,29 @@ void FNiagaraDataSet::Dump(FNiagaraDataSet& Other, bool bCurr, int32 StartIdx ,
 	DataBuffer.CopyTo(OtherDataBuffer, StartIdx, NumInstances);
 }
 
+void FNiagaraDataSet::DumpGPU(FNiagaraDataSet& Other, float* GPUReadBackFloat, int* GPUReadBackInt, int32 StartIdx, int32 NumInstances)const
+{
+	check(IsInRenderingThread());
+	Other.Reset();
+	Other.Variables = Variables;
+	Other.VariableLayouts = VariableLayouts;
+
+	Other.TotalFloatComponents = TotalFloatComponents;
+	Other.TotalInt32Components = TotalInt32Components;
+
+	const FNiagaraDataBuffer& DataBuffer = CurrData();
+	FNiagaraDataBuffer& OtherDataBuffer = Other.CurrData();
+
+	if (OtherDataBuffer.GetNumInstancesAllocated() != DataBuffer.GetNumInstancesAllocated())
+	{
+		Other.Finalize();
+		OtherDataBuffer.Allocate(DataBuffer.GetNumInstancesAllocated());
+	}
+
+	DataBuffer.GPUCopyTo(OtherDataBuffer, GPUReadBackFloat, GPUReadBackInt, StartIdx, NumInstances);
+}
+
+
 
 void FNiagaraDataSet::Dump(bool bCurr, int32 StartIdx, int32 NumInstances)const
 {
@@ -454,6 +495,66 @@ void FNiagaraDataBuffer::CopyTo(FNiagaraDataBuffer& DestBuffer, int32 InStartIdx
 	}
 }
 
+void FNiagaraDataBuffer::GPUCopyTo(FNiagaraDataBuffer& DestBuffer, float* GPUReadBackFloat, int* GPUReadBackInt, int32 InStartIdx, int32 InNumInstances)const
+{
+	if (InStartIdx < 0 || (uint32)InStartIdx > NumInstances)
+	{
+		InStartIdx = NumInstances;
+	}
+	if (InNumInstances < 0 || ((uint32)InNumInstances + (uint32)InStartIdx) > NumInstances)
+	{
+		InNumInstances = NumInstances - InStartIdx;
+	}
+
+	if (InNumInstances != 0)
+	{
+		if (DestBuffer.NumInstancesAllocated != NumInstancesAllocated)
+		{
+			DestBuffer.Allocate(NumInstancesAllocated);
+		}
+
+		if (GPUReadBackFloat)
+		{
+			for (uint32 CompIdx = 0; CompIdx < Owner->TotalFloatComponents; ++CompIdx)
+			{
+				const float* SrcStart = GetInstancePtrFloat(GPUReadBackFloat, CompIdx, InStartIdx);
+				const float* SrcEnd = GetInstancePtrFloat(GPUReadBackFloat, CompIdx, InStartIdx + InNumInstances);
+				float* Dst = DestBuffer.GetInstancePtrFloat(CompIdx, 0);
+				size_t Count = SrcEnd - SrcStart;
+				FMemory::Memcpy(Dst, SrcStart, Count * sizeof(float));
+
+				if (Count > 0)
+				{
+					for (size_t i = 0; i < Count; i++)
+					{
+						check(SrcStart[i] == Dst[i]);
+					}
+				}
+			}
+		}
+		if (GPUReadBackInt)
+		{
+			for (uint32 CompIdx = 0; CompIdx < Owner->TotalInt32Components; ++CompIdx)
+			{
+				const int32* SrcStart = GetInstancePtrInt32(GPUReadBackInt, CompIdx, InStartIdx);
+				const int32* SrcEnd = GetInstancePtrInt32(GPUReadBackInt, CompIdx, InStartIdx + InNumInstances);
+				int32* Dst = DestBuffer.GetInstancePtrInt32(CompIdx, 0);
+				size_t Count = SrcEnd - SrcStart;
+				FMemory::Memcpy(Dst, SrcStart, Count * sizeof(int32));
+
+				if (Count > 0)
+				{
+					for (size_t i = 0; i < Count; i++)
+					{
+						check(SrcStart[i] == Dst[i]);
+					}
+				}
+			}
+		}
+		DestBuffer.SetNumInstances(InNumInstances);
+	}
+}
+
 void FNiagaraDataBuffer::CopyTo(FNiagaraDataBuffer& DestBuffer)const
 {
 	DestBuffer.FloatStride = FloatStride;
@@ -463,3 +564,4 @@ void FNiagaraDataBuffer::CopyTo(FNiagaraDataBuffer& DestBuffer)const
 	DestBuffer.NumInstancesAllocated = NumInstancesAllocated;
 	DestBuffer.NumInstances = NumInstances;
 }
+
diff --git a/Engine/Plugins/FX/Niagara/Source/Niagara/Private/NiagaraEmitter.cpp b/Engine/Plugins/FX/Niagara/Source/Niagara/Private/NiagaraEmitter.cpp
--- a/Engine/Plugins/FX/Niagara/Source/Niagara/Private/NiagaraEmitter.cpp
+++ b/Engine/Plugins/FX/Niagara/Source/Niagara/Private/NiagaraEmitter.cpp
@@ -48,7 +48,7 @@ void FNiagaraEmitterScriptProperties::InitDataSetAccess()
 
 		for (FNiagaraDataSetProperties &WriteID : Script->GetVMExecutableData().WriteDataSets)
 		{
-			FNiagaraEventGeneratorProperties Props(WriteID, "", "");
+			FNiagaraEventGeneratorProperties Props(WriteID, "");
 			EventGenerators.Add(Props);
 		}
 	}
@@ -86,6 +86,8 @@ UNiagaraEmitter::UNiagaraEmitter(const FObjectInitializer& Initializer)
 , bUseMinDetailLevel(false)
 , bUseMaxDetailLevel(false)
 , bRequiresPersistentIDs(false)
+, MaxDeltaTimePerTick(0.125)
+, bLimitDeltaTime(true)
 #if WITH_EDITORONLY_DATA
 , ThumbnailImageOutOfDate(true)
 #endif
@@ -329,6 +331,18 @@ void UNiagaraEmitter::PostEditChangeProperty(struct FPropertyChangedEvent& Prope
 		UNiagaraSystem::RequestCompileForEmitter(this);
 #endif
 	}
+	if (PropertyName == GET_MEMBER_NAME_CHECKED(UNiagaraEmitter, bDeterminism))
+	{
+		if (GraphSource != nullptr)
+		{
+			GraphSource->MarkNotSynchronized(TEXT("Emitter Determinism changed."));
+		}
+
+#if WITH_EDITORONLY_DATA
+		UNiagaraSystem::RequestCompileForEmitter(this);
+#endif
+	}
+
 	ThumbnailImageOutOfDate = true;
 	ChangeId = FGuid::NewGuid();
 	OnPropertiesChangedDelegate.Broadcast();
@@ -531,6 +545,20 @@ void UNiagaraEmitter::OnPostCompile()
 	SpawnScriptProps.InitDataSetAccess();
 	UpdateScriptProps.InitDataSetAccess();
 
+	TSet<FName> SpawnIds;
+	TSet<FName> UpdateIds;
+	for (const FNiagaraEventGeneratorProperties& SpawnGeneratorProps : SpawnScriptProps.EventGenerators)
+	{
+		SpawnIds.Add(SpawnGeneratorProps.ID);
+	}
+	for (const FNiagaraEventGeneratorProperties& UpdateGeneratorProps : UpdateScriptProps.EventGenerators)
+	{
+		UpdateIds.Add(UpdateGeneratorProps.ID);
+	}
+
+	SharedEventGeneratorIds.Empty();
+	SharedEventGeneratorIds.Append(SpawnIds.Intersect(UpdateIds).Array());
+
 	for (int32 i = 0; i < EventHandlerScriptProps.Num(); i++)
 	{
 		if (EventHandlerScriptProps[i].Script)
@@ -762,6 +790,11 @@ void UNiagaraEmitter::RemoveEventHandlerByUsageId(FGuid EventHandlerUsageId)
 #endif
 }
 
+bool UNiagaraEmitter::IsEventGeneratorShared(FName EventGeneratorId) const
+{
+	return SharedEventGeneratorIds.Contains(EventGeneratorId);
+}
+
 void UNiagaraEmitter::BeginDestroy()
 {
 #if WITH_EDITOR
diff --git a/Engine/Plugins/FX/Niagara/Source/Niagara/Private/NiagaraEmitterInstance.cpp b/Engine/Plugins/FX/Niagara/Source/Niagara/Private/NiagaraEmitterInstance.cpp
--- a/Engine/Plugins/FX/Niagara/Source/Niagara/Private/NiagaraEmitterInstance.cpp
+++ b/Engine/Plugins/FX/Niagara/Source/Niagara/Private/NiagaraEmitterInstance.cpp
@@ -17,12 +17,11 @@
 DECLARE_DWORD_COUNTER_STAT(TEXT("Num Custom Events"), STAT_NiagaraNumCustomEvents, STATGROUP_Niagara);
 
 //DECLARE_CYCLE_STAT(TEXT("Tick"), STAT_NiagaraTick, STATGROUP_Niagara);
-DECLARE_CYCLE_STAT(TEXT("Simulate"), STAT_NiagaraSimulate, STATGROUP_Niagara);
-DECLARE_CYCLE_STAT(TEXT("Spawn"), STAT_NiagaraSpawn, STATGROUP_Niagara);
-DECLARE_CYCLE_STAT(TEXT("Spawn"), STAT_NiagaraEvents, STATGROUP_Niagara);
-DECLARE_CYCLE_STAT(TEXT("Kill"), STAT_NiagaraKill, STATGROUP_Niagara);
-DECLARE_CYCLE_STAT(TEXT("Event Handling"), STAT_NiagaraEventHandle, STATGROUP_Niagara);
-DECLARE_CYCLE_STAT(TEXT("Error Check"), STAT_NiagaraEmitterErrorCheck, STATGROUP_Niagara);
+DECLARE_CYCLE_STAT(TEXT("Emitter Simulate [CNC]"), STAT_NiagaraSimulate, STATGROUP_Niagara);
+DECLARE_CYCLE_STAT(TEXT("Emitter Spawn [CNC]"), STAT_NiagaraSpawn, STATGROUP_Niagara);
+DECLARE_CYCLE_STAT(TEXT("Emitter Post Tick [CNC]"), STAT_NiagaraEmitterPostTick, STATGROUP_Niagara);
+DECLARE_CYCLE_STAT(TEXT("Emitter Event Handling [CNC]"), STAT_NiagaraEventHandle, STATGROUP_Niagara);
+DECLARE_CYCLE_STAT(TEXT("Emitter Error Check [CNC]"), STAT_NiagaraEmitterErrorCheck, STATGROUP_Niagara);
 
 static int32 GbDumpParticleData = 0;
 static FAutoConsoleVariableRef CVarNiagaraDumpParticleData(
@@ -229,6 +228,8 @@ void FNiagaraEmitterInstance::Init(int32 InEmitterIdx, FName InSystemInstanceNam
 
 	ensure(CachedEmitter->UpdateScriptProps.DataSetAccessSynchronized());
 	UpdateScriptEventDataSets.Empty();
+	UpdateEventGeneratorIsSharedByIndex.SetNumZeroed(CachedEmitter->UpdateScriptProps.EventGenerators.Num());
+	int32 UpdateEventGeneratorIndex = 0;
 	for (const FNiagaraEventGeneratorProperties &GeneratorProps : CachedEmitter->UpdateScriptProps.EventGenerators)
 	{
 		FNiagaraDataSet *Set = FNiagaraEventDataSetMgr::CreateEventDataSet(ParentSystemInstance->GetIDName(), EmitterHandle.GetIdName(), GeneratorProps.SetProps.ID.Name);
@@ -236,10 +237,14 @@ void FNiagaraEmitterInstance::Init(int32 InEmitterIdx, FName InSystemInstanceNam
 		Set->AddVariables(GeneratorProps.SetProps.Variables);
 		Set->Finalize();
 		UpdateScriptEventDataSets.Add(Set);
+		UpdateEventGeneratorIsSharedByIndex[UpdateEventGeneratorIndex] = CachedEmitter->IsEventGeneratorShared(GeneratorProps.ID);
+		UpdateEventGeneratorIndex++;
 	}
 
 	ensure(CachedEmitter->SpawnScriptProps.DataSetAccessSynchronized());
 	SpawnScriptEventDataSets.Empty();
+	SpawnEventGeneratorIsSharedByIndex.SetNumZeroed(CachedEmitter->SpawnScriptProps.EventGenerators.Num());
+	int32 SpawnEventGeneratorIndex = 0;
 	for (const FNiagaraEventGeneratorProperties &GeneratorProps : CachedEmitter->SpawnScriptProps.EventGenerators)
 	{
 		FNiagaraDataSet *Set = FNiagaraEventDataSetMgr::CreateEventDataSet(ParentSystemInstance->GetIDName(), EmitterHandle.GetIdName(), GeneratorProps.SetProps.ID.Name);
@@ -247,6 +252,8 @@ void FNiagaraEmitterInstance::Init(int32 InEmitterIdx, FName InSystemInstanceNam
 		Set->AddVariables(GeneratorProps.SetProps.Variables);
 		Set->Finalize();
 		SpawnScriptEventDataSets.Add(Set);
+		SpawnEventGeneratorIsSharedByIndex[SpawnEventGeneratorIndex] = CachedEmitter->IsEventGeneratorShared(GeneratorProps.ID);
+		SpawnEventGeneratorIndex++;
 	}
 
 	SpawnExecContext.Init(CachedEmitter->SpawnScriptProps.Script, CachedEmitter->SimTarget);
@@ -256,7 +263,7 @@ void FNiagaraEmitterInstance::Init(int32 InEmitterIdx, FName InSystemInstanceNam
 	if (CachedEmitter->SimTarget == ENiagaraSimTarget::GPUComputeSim)
 	{
 		GPUExecContext = new FNiagaraComputeExecutionContext();
-		GPUExecContext->InitParams(CachedEmitter->GetGPUComputeScript(), CachedEmitter->SpawnScriptProps.Script, CachedEmitter->UpdateScriptProps.Script, CachedEmitter->SimTarget);
+		GPUExecContext->InitParams(CachedEmitter->GetGPUComputeScript(), CachedEmitter->SpawnScriptProps.Script, CachedEmitter->UpdateScriptProps.Script, CachedEmitter->SimTarget, CachedEmitter->GetUniqueEmitterName());
 		GPUExecContext->MainDataSet = &Data;
 		GPUExecContext->RTGPUScript = CachedEmitter->GetGPUComputeScript()->GetRenderThreadScript();
 		GPUExecContext->RTSpawnScript = CachedEmitter->SpawnScriptProps.Script->GetRenderThreadScript();
@@ -304,6 +311,16 @@ void FNiagaraEmitterInstance::Init(int32 InEmitterIdx, FName InSystemInstanceNam
 		EmitterAgeBindingGPU.Init(GPUExecContext->CombinedParamStore, EmitterAgeParam);
 	}
 
+	// Initialize the random seed
+	FNiagaraVariable EmitterRandomSeedParam = CachedEmitter->ToEmitterParameter(SYS_PARAM_EMITTER_RANDOM_SEED);
+	SpawnRandomSeedBinding.Init(SpawnExecContext.Parameters, EmitterRandomSeedParam);
+	UpdateRandomSeedBinding.Init(UpdateExecContext.Parameters, EmitterRandomSeedParam);
+	if (CachedEmitter->SimTarget == ENiagaraSimTarget::GPUComputeSim && GPUExecContext != nullptr)
+	{
+		GPURandomSeedBinding.Init(GPUExecContext->CombinedParamStore, CachedEmitter->ToEmitterParameter(EmitterRandomSeedParam));
+	}
+
+	// Initialize the exec count
 	SpawnExecCountBinding.Init(SpawnExecContext.Parameters, SYS_PARAM_ENGINE_EXEC_COUNT);
 	UpdateExecCountBinding.Init(UpdateExecContext.Parameters, SYS_PARAM_ENGINE_EXEC_COUNT);
 	EventExecCountBindings.SetNum(NumEvents);
@@ -319,7 +336,7 @@ void FNiagaraEmitterInstance::Init(int32 InEmitterIdx, FName InSystemInstanceNam
 	}
 	else
 	{
-		//Init accessors for PostProcessParticles
+		//Init accessors for PostTick
 		PositionAccessor = FNiagaraDataSetAccessor<FVector>(Data, FNiagaraVariable(FNiagaraTypeDefinition::GetVec3Def(), PositionName));
 		SizeAccessor = FNiagaraDataSetAccessor<FVector2D>(Data, FNiagaraVariable(FNiagaraTypeDefinition::GetVec2Def(), SizeName));
 		MeshScaleAccessor = FNiagaraDataSetAccessor<FVector>(Data, FNiagaraVariable(FNiagaraTypeDefinition::GetVec3Def(), MeshScaleName));
@@ -342,6 +359,7 @@ void FNiagaraEmitterInstance::ResetSimulation()
 	Age = 0;
 	Loops = 0;
 	TickCount = 0;
+	TotalSpawnedParticles = 0;
 	CachedBounds.Init();
 
 	ParticleDataSet->ResetBuffers();
@@ -713,12 +731,12 @@ TOptional<FBox> FNiagaraEmitterInstance::CalculateDynamicBounds()
 	return Ret;
 }
 
-/** Look for dead particles and move from the end of the list to the dead location, compacting in the process
-  * Also calculates bounds; Kill will be removed from this once we do conditional write
+/** 
+  * Do any post work such as calculating dynamic bounds.
   */
-void FNiagaraEmitterInstance::PostProcessParticles()
+void FNiagaraEmitterInstance::PostTick()
 {
-	SCOPE_CYCLE_COUNTER(STAT_NiagaraKill);
+	SCOPE_CYCLE_COUNTER(STAT_NiagaraEmitterPostTick);
 
 	checkSlow(CachedEmitter);
 	CachedBounds.Init();
@@ -868,12 +886,30 @@ void FNiagaraEmitterInstance::PreTick()
 	ParticleDataSet->SetIDAcquireTag(TickCount);
 }
 
+bool FNiagaraEmitterInstance::WaitForDebugInfo()
+{
+	if (CachedEmitter->SimTarget == ENiagaraSimTarget::GPUComputeSim)
+	{
+		
+		FNiagaraComputeExecutionContext* DebugContext = GPUExecContext;
+
+		ENQUEUE_RENDER_COMMAND(CaptureCommand)([DebugContext](FRHICommandListImmediate& RHICmdList)
+			{
+				NiagaraEmitterInstanceBatcher::Get()->ProcessDebugInfo(RHICmdList, DebugContext);
+			}
+		);
+		return true;
+	}
+	return false;
+}
+
 
 void FNiagaraEmitterInstance::Tick(float DeltaSeconds)
 {
 	SCOPE_CYCLE_COUNTER(STAT_NiagaraTick);
 	SimpleTimer TickTime;
 
+
 	if (HandleCompletion())
 	{
 		CPUTimeMS = TickTime.GetElapsedMilliseconds();
@@ -884,6 +920,9 @@ void FNiagaraEmitterInstance::Tick(float DeltaSeconds)
 	FNiagaraDataSet& Data = *ParticleDataSet;
 	Age += DeltaSeconds;
 
+
+	//UE_LOG(LogNiagara, Warning, TEXT("Emitter Tick %f"), Age);
+
 	if (ExecutionState == ENiagaraExecutionState::InactiveClear)
 	{
 		Data.ResetBuffers();
@@ -925,9 +964,13 @@ void FNiagaraEmitterInstance::Tick(float DeltaSeconds)
 			Binding.SetValue(Age);
 		}
 		
+		SpawnRandomSeedBinding.SetValue(CachedEmitter->RandomSeed);
+		UpdateRandomSeedBinding.SetValue(CachedEmitter->RandomSeed);
+
 		if (CachedEmitter->SimTarget == ENiagaraSimTarget::GPUComputeSim && GPUExecContext != nullptr)
 		{
 			EmitterAgeBindingGPU.SetValue(Age);
+			GPURandomSeedBinding.SetValue(CachedEmitter->RandomSeed);
 		}
 	}
 	
@@ -971,6 +1014,11 @@ void FNiagaraEmitterInstance::Tick(float DeltaSeconds)
 			uint32 EventSpawnNum = CalculateEventSpawnCount(EventHandlerProps, EventSpawnCounts[i], EventSet[i]);
 			EventSpawnTotal += EventSpawnNum;
 			EventHandlerSpawnCounts[i] = EventSpawnNum;
+			if (CachedEmitter->SimTarget == ENiagaraSimTarget::GPUComputeSim && GPUExecContext != nullptr) {
+				// NOTE(mv): Separate particle count path for GPU emitters, as they early out..
+				// not actually reached yet due to GPU events not working. CPU events are counted further down...
+				TotalSpawnedParticles += EventSpawnNum;
+			}
 		}
 	}
 
@@ -979,9 +1027,29 @@ void FNiagaraEmitterInstance::Tick(float DeltaSeconds)
 	 */
 	if (CachedEmitter->SimTarget == ENiagaraSimTarget::GPUComputeSim && GPUExecContext != nullptr)
 	{
+		//FNiagaraComputeExecutionContext *ComputeContext = new FNiagaraComputeExecutionContext();
+		GPUExecContext->MainDataSet = &Data;
+		GPUExecContext->RTGPUScript = CachedEmitter->GetGPUComputeScript()->GetRenderThreadScript();
+		GPUExecContext->RTSpawnScript = CachedEmitter->SpawnScriptProps.Script->GetRenderThreadScript();
+		GPUExecContext->RTUpdateScript = CachedEmitter->UpdateScriptProps.Script->GetRenderThreadScript();
 		GPUExecContext->SpawnRateInstances = SpawnTotal;
 		GPUExecContext->EventSpawnTotal = EventSpawnTotal;
 		GPUExecContext->NumIndicesPerInstance = CachedEmitter->GetRenderers()[0]->GetNumIndicesPerInstance();
+		
+#if WITH_EDITORONLY_DATA
+		if (ParentSystemInstance->ShouldCaptureThisFrame())
+		{
+			TSharedPtr<struct FNiagaraScriptDebuggerInfo, ESPMode::ThreadSafe> DebugInfo = ParentSystemInstance->GetActiveCaptureWrite(CachedIDName, ENiagaraScriptUsage::ParticleGPUComputeScript, FGuid());
+			if (DebugInfo)
+			{
+				//Data.Dump(DebugInfo->Frame, true, 0, OrigNumParticles);
+				//DebugInfo->Frame.Dump(true, 0, OrigNumParticles);
+				DebugInfo->Parameters = GPUExecContext->CombinedParamStore;
+				GPUExecContext->DebugInfo = DebugInfo;
+			}
+		}
+#endif
+		ParentSystemInstance->GetPerInstanceDataAndOffsets(GPUExecContext->PerInstanceData, GPUExecContext->PerInstanceDataSize, GPUExecContext->PerInstanceDataInterfaceOffsets);
 
 		bool bOnlySetOnce = false;
 		for (FNiagaraSpawnInfo& Info : SpawnInfos)
@@ -998,6 +1066,9 @@ void FNiagaraEmitterInstance::Tick(float DeltaSeconds)
 				UE_LOG(LogNiagara, Log, TEXT("Multiple spawns are happening this frame. Only doing the first!"));
 				break;
 			}
+
+			// NOTE(mv): Separate particle count path for GPU emitters, as they early out..
+			TotalSpawnedParticles += Info.Count;
 		}
 
 		//GPUExecContext.UpdateInterfaces = CachedEmitter->UpdateScriptProps.Script->GetCachedDefaultDataInterfaces();
@@ -1066,7 +1137,7 @@ void FNiagaraEmitterInstance::Tick(float DeltaSeconds)
 	//TODO: These current limits can be improved relatively easily. Though perf in at these counts will obviously be an issue anyway.
 	if (CachedEmitter->SimTarget == ENiagaraSimTarget::CPUSim && AllocationSize > GMaxNiagaraCPUParticlesPerEmitter)
 	{
-		UE_LOG(LogNiagara, Warning, TEXT("Emitter %s has attemted to exceed the max CPU particle count! | Max: %d | Requested: %u"), *CachedEmitter->GetUniqueEmitterName(), GMaxNiagaraCPUParticlesPerEmitter, AllocationSize);
+		UE_LOG(LogNiagara, Warning, TEXT("Emitter %s has attempted to exceed the max CPU particle count! | Max: %d | Requested: %u"), *CachedEmitter->GetUniqueEmitterName(), GMaxNiagaraCPUParticlesPerEmitter, AllocationSize);
 		//For now we completely bail out of spawning new particles. Possibly should improve this in future.
 		AllocationSize = OrigNumParticles;
 		SpawnTotal = 0;
@@ -1075,14 +1146,33 @@ void FNiagaraEmitterInstance::Tick(float DeltaSeconds)
 
 	//Allocate space for prev frames particles and any new one's we're going to spawn.
 	Data.Allocate(AllocationSize);
+
+	int32 SpawnEventGeneratorIndex = 0;
 	for (FNiagaraDataSet* SpawnEventDataSet : SpawnScriptEventDataSets)
 	{
-		SpawnEventDataSet->Allocate(SpawnTotal + EventSpawnTotal);
+		int32 NumToAllocate = SpawnTotal + EventSpawnTotal;
+		if (SpawnEventGeneratorIsSharedByIndex[SpawnEventGeneratorIndex])
+		{
+			// For shared event data sets we need to allocate storage for the current particles since
+			// the same data set will be used in the update execution.
+			NumToAllocate += OrigNumParticles;
+		}
+		SpawnEventDataSet->Allocate(NumToAllocate);
+		SpawnEventGeneratorIndex++;
 	}
+
+	int32 UpdateEventGeneratorIndex = 0;
 	for (FNiagaraDataSet* UpdateEventDataSet : UpdateScriptEventDataSets)
 	{
-		UpdateEventDataSet->Allocate(OrigNumParticles);
+		if (UpdateEventGeneratorIsSharedByIndex[UpdateEventGeneratorIndex] == false)
+		{
+			// We only allocate update event data sets if they're not shared, because shared event datasets will have already
+			// been allocated as part of the spawn event data set handling.
+			UpdateEventDataSet->Allocate(OrigNumParticles);
+		}
+		UpdateEventGeneratorIndex++;
 	}
+
 	TArray<FNiagaraDataSetExecutionInfo, TInlineAllocator<8>> DataSetExecInfos;
 	DataSetExecInfos.Emplace(&Data, 0, false, true);
 
@@ -1105,8 +1195,8 @@ void FNiagaraEmitterInstance::Tick(float DeltaSeconds)
 		DataSetExecInfos[0].StartInstance = 0;
 		for (FNiagaraDataSet* EventDataSet : UpdateScriptEventDataSets)
 		{
-			DataSetExecInfos.Emplace(EventDataSet, 0, false, true);
 			EventDataSet->SetNumInstances(OrigNumParticles);
+			DataSetExecInfos.Emplace(EventDataSet, 0, false, true);
 		}
 		UpdateExecContext.Execute(OrigNumParticles, DataSetExecInfos);
 		int32 DeltaParticles = Data.GetNumInstances() - OrigNumParticles;
@@ -1126,12 +1216,13 @@ void FNiagaraEmitterInstance::Tick(float DeltaSeconds)
 #if WITH_EDITORONLY_DATA
 	if (ParentSystemInstance->ShouldCaptureThisFrame())
 	{
-		FNiagaraScriptDebuggerInfo* DebugInfo = ParentSystemInstance->GetActiveCaptureWrite(CachedIDName, ENiagaraScriptUsage::ParticleUpdateScript, FGuid());
+		TSharedPtr<struct FNiagaraScriptDebuggerInfo, ESPMode::ThreadSafe> DebugInfo = ParentSystemInstance->GetActiveCaptureWrite(CachedIDName, ENiagaraScriptUsage::ParticleUpdateScript, FGuid());
 		if (DebugInfo)
 		{
 			Data.Dump(DebugInfo->Frame, true, 0, OrigNumParticles);
 			//DebugInfo->Frame.Dump(true, 0, OrigNumParticles);
 			DebugInfo->Parameters = UpdateExecContext.Parameters;
+			DebugInfo->bWritten = true;
 		}
 	}
 #endif
@@ -1152,6 +1243,12 @@ void FNiagaraEmitterInstance::Tick(float DeltaSeconds)
 				int32 OrigNum = Data.GetNumInstances();
 				Data.SetNumInstances(OrigNum + Num);
 
+				// NOTE(mv): Updates the count after setting the variable, such that the TotalSpawnedParticles value read 
+				//           in the script has the count at the start of the frame. 
+				//           This way UniqueID = TotalSpawnedParticles + ExecIndex provide unique and sequential identifiers. 
+				// NOTE(mv): Only for CPU particles, as GPU particles early outs further up and has a separate increment. 
+				TotalSpawnedParticles += Num;
+
 				SpawnExecCountBinding.SetValue(Num);
 				DataSetExecInfos.SetNum(1, false);
 				DataSetExecInfos[0].StartInstance = OrigNum;
@@ -1212,12 +1309,13 @@ void FNiagaraEmitterInstance::Tick(float DeltaSeconds)
 	int32 TotalNumSpawned = NumAfterSpawn - NumBeforeSpawn;
 	if (ParentSystemInstance->ShouldCaptureThisFrame())
 	{
-		FNiagaraScriptDebuggerInfo* DebugInfo = ParentSystemInstance->GetActiveCaptureWrite(CachedIDName, ENiagaraScriptUsage::ParticleSpawnScript, FGuid());
+		TSharedPtr<struct FNiagaraScriptDebuggerInfo, ESPMode::ThreadSafe> DebugInfo = ParentSystemInstance->GetActiveCaptureWrite(CachedIDName, ENiagaraScriptUsage::ParticleSpawnScript, FGuid());
 		if (DebugInfo)
 		{
 			Data.Dump(DebugInfo->Frame, true, NumBeforeSpawn, TotalNumSpawned);
 			//DebugInfo->Frame.Dump(true, 0, Num);
 			DebugInfo->Parameters = SpawnExecContext.Parameters;
+			DebugInfo->bWritten = true;
 		}
 	}
 #endif
@@ -1270,12 +1368,13 @@ void FNiagaraEmitterInstance::Tick(float DeltaSeconds)
 					if (ParentSystemInstance->ShouldCaptureThisFrame())
 					{
 						FGuid EventGuid = EventExecContexts[EventScriptIdx].Script->GetUsageId();
-						FNiagaraScriptDebuggerInfo* DebugInfo = ParentSystemInstance->GetActiveCaptureWrite(CachedIDName, ENiagaraScriptUsage::ParticleEventScript, EventGuid);
+						TSharedPtr<struct FNiagaraScriptDebuggerInfo, ESPMode::ThreadSafe> DebugInfo = ParentSystemInstance->GetActiveCaptureWrite(CachedIDName, ENiagaraScriptUsage::ParticleEventScript, EventGuid);
 						if (DebugInfo)
 						{
 							Data.Dump(DebugInfo->Frame, true, EventSpawnStart, ActualEventNumToSpawn);		
 							//DebugInfo->Frame.Dump(true, 0, ActualEventNumToSpawn);
 							DebugInfo->Parameters = EventExecContexts[EventScriptIdx].Parameters;
+							DebugInfo->bWritten = true;
 						}
 					}
 #endif
@@ -1333,12 +1432,13 @@ void FNiagaraEmitterInstance::Tick(float DeltaSeconds)
 					if (ParentSystemInstance->ShouldCaptureThisFrame())
 					{
 						FGuid EventGuid = EventExecContexts[EventScriptIdx].Script->GetUsageId();
-						FNiagaraScriptDebuggerInfo* DebugInfo = ParentSystemInstance->GetActiveCaptureWrite(CachedIDName, ENiagaraScriptUsage::ParticleEventScript, EventGuid);
+						TSharedPtr<struct FNiagaraScriptDebuggerInfo, ESPMode::ThreadSafe> DebugInfo = ParentSystemInstance->GetActiveCaptureWrite(CachedIDName, ENiagaraScriptUsage::ParticleEventScript, EventGuid);
 						if (DebugInfo)
 						{
 							Data.Dump(DebugInfo->Frame, true, 0, NumInstancesPrev);
 							//DebugInfo->Frame.Dump(true, 0, NumInstancesPrev);
 							DebugInfo->Parameters = EventExecContexts[EventScriptIdx].Parameters;
+							DebugInfo->bWritten = true;
 						}
 					}
 #endif
@@ -1392,7 +1492,7 @@ void FNiagaraEmitterInstance::Tick(float DeltaSeconds)
 // 					if (ParentSystemInstance->ShouldCaptureThisFrame())
 // 					{
 // 						FGuid EventGuid = EventExecContexts[EventScriptIdx].Script->GetUsageId();
-// 						FNiagaraScriptDebuggerInfo* DebugInfo = ParentSystemInstance->GetActiveCaptureWrite(CachedIDName, ENiagaraScriptUsage::ParticleEventScript, EventGuid);
+// 						TSharedPtr<struct FNiagaraScriptDebuggerInfo, ESPMode::ThreadSafe> DebugInfo = ParentSystemInstance->GetActiveCaptureWrite(CachedIDName, ENiagaraScriptUsage::ParticleEventScript, EventGuid);
 // 						if (DebugInfo)
 // 						{
 // 							Data.Dump(DebugInfo->Frame, true, Index, 1);
@@ -1407,7 +1507,7 @@ void FNiagaraEmitterInstance::Tick(float DeltaSeconds)
 // 		}
 	}
 
-	PostProcessParticles();
+	PostTick();
 
 	SpawnExecContext.PostTick();
 	UpdateExecContext.PostTick();
diff --git a/Engine/Plugins/FX/Niagara/Source/Niagara/Private/NiagaraEmitterInstanceBatcher.cpp b/Engine/Plugins/FX/Niagara/Source/Niagara/Private/NiagaraEmitterInstanceBatcher.cpp
--- a/Engine/Plugins/FX/Niagara/Source/Niagara/Private/NiagaraEmitterInstanceBatcher.cpp
+++ b/Engine/Plugins/FX/Niagara/Source/Niagara/Private/NiagaraEmitterInstanceBatcher.cpp
@@ -26,6 +26,7 @@ uint32 FNiagaraComputeExecutionContext::TickCounter = 0;
 
 void NiagaraEmitterInstanceBatcher::Queue(FNiagaraComputeExecutionContext *InContext)
 {
+	//UE_LOG(LogNiagara, Warning, TEXT("Submitted!"));
 	//SimulationQueue[CurQueueIndex]->Add(InContext);
 		ENQUEUE_UNIQUE_RENDER_COMMAND_THREEPARAMETER(QueueNiagaraDispatch,
 			TArray<FNiagaraComputeExecutionContext*>*, Queue, &SimulationQueue[0],
@@ -58,6 +59,8 @@ void NiagaraEmitterInstanceBatcher::Remove(FNiagaraComputeExecutionContext *InCo
 
 void NiagaraEmitterInstanceBatcher::ExecuteAll(FRHICommandList &RHICmdList, FUniformBufferRHIParamRef ViewUniformBuffer)
 {
+	SCOPED_DRAW_EVENT(RHICmdList, NiagaraEmitterInstanceBatcher_ExecuteAll);
+
 	const uint32 TickedQueueIndex = (CurQueueIndex ^ 0x1);
 	const uint32 TickedQueueIndexMask = (1 << TickedQueueIndex);
 
@@ -91,6 +94,13 @@ void NiagaraEmitterInstanceBatcher::TickSingle(FNiagaraComputeExecutionContext *
 		return;
 	}
 
+#if WITH_EDITORONLY_DATA
+	if (Context->DebugInfo.IsValid())
+	{
+		ProcessDebugInfo(RHICmdList, Context);
+	}
+#endif // WITH_EDITORONLY_DATA
+
 	uint32 PrevNumInstances = Context->MainDataSet->PrevData().GetNumInstances();
 	uint32 NewNumInstances = Context->SpawnRateInstances + Context->EventSpawnTotal + PrevNumInstances;
 	Context->EventSpawnTotal = GetEventSpawnTotal(Context);
@@ -193,21 +203,31 @@ void NiagaraEmitterInstanceBatcher::ResolveDatasetWrites(FRHICommandList &RHICmd
 	}
 	else if (Context->GPUDataReadback->IsReady())
 	{
+		bool bSuccessfullyRead = false;
 		{
 		    SCOPE_CYCLE_COUNTER(STAT_NiagaraGPUReadback_RT);
 		    int32 *NumInstancesAfterSim = static_cast<int32*>(Context->GPUDataReadback->RetrieveData(64 * sizeof(int32)));
-		    int32 ExistingDataCount = Context->MainDataSet->CurrData().GetNumInstances();// index 1 is always the count
-		    int32 NewExistingDataCount = NumInstancesAfterSim[1] + Context->AccumulatedSpawnRate;
-		    Context->MainDataSet->CurrData().SetNumInstances(NewExistingDataCount);	
-		    FString PathName = Context->GPUScript->GetOutermost()->GetPathName();
-		    // UE_LOG(LogNiagara, Log, TEXT("GPU Syncup %s : Was(%d) Now(%d)"), *PathName, ExistingDataCount, NewExistingDataCount );
-		    SET_DWORD_STAT(STAT_NiagaraGPUParticles, NewExistingDataCount);
-		    SET_DWORD_STAT(STAT_NiagaraReadbackLatency, 0);
-    
-		    Context->GPUDataReadback->Finish();
-    
-		    Context->AccumulatedSpawnRate = 0;
+			if (NumInstancesAfterSim)
+			{
+				int32 ExistingDataCount = Context->MainDataSet->CurrData().GetNumInstances();// index 1 is always the count
+				int32 NewExistingDataCount = NumInstancesAfterSim[1] + Context->AccumulatedSpawnRate;
+				Context->MainDataSet->CurrData().SetNumInstances(NewExistingDataCount);
+				FString PathName = Context->GPUScript->GetOutermost()->GetPathName();
+				// UE_LOG(LogNiagara, Log, TEXT("GPU Syncup %s : Was(%d) Now(%d)"), *PathName, ExistingDataCount, NewExistingDataCount );
+				SET_DWORD_STAT(STAT_NiagaraGPUParticles, NewExistingDataCount);
+				SET_DWORD_STAT(STAT_NiagaraReadbackLatency, 0);
+
+				Context->GPUDataReadback->Finish();
+
+				Context->AccumulatedSpawnRate = 0;
+				bSuccessfullyRead = true;
+			}
+			else
+			{
+				UE_LOG(LogNiagara, Warning, TEXT("GPUDataReadback said it was ready, but returned an invalid buffer. Skipping this time.."));
+			}
 		}
+		if (bSuccessfullyRead)
 		{
 			SCOPE_CYCLE_COUNTER(STAT_NiagaraAllocateGPUReadback_RT);
 			// The following code seems to take significant time on d3d12
@@ -216,9 +236,84 @@ void NiagaraEmitterInstanceBatcher::ResolveDatasetWrites(FRHICommandList &RHICmd
 			Context->GPUDataReadback->Insert(RHICmdList);
 		}
 	}
+
+}
+
+void NiagaraEmitterInstanceBatcher::ProcessDebugInfo(FRHICommandList &RHICmdList, const FNiagaraComputeExecutionContext *Context) const
+{
+#if WITH_EDITORONLY_DATA
+	// This method may be called from one of two places: in the tick or as part of a paused frame looking for the debug info that was submitted previously...
+	// Note that PrevData is where we expect the data to be for rendering, as per NiagaraEmitterInstanceBatcher::TickSingle
+	if (Context->DebugInfo.IsValid())
+	{
+		
+		// Fire off the readback if not already doing so
+		if (!Context->GPUDebugDataReadbackFloat && !Context->GPUDebugDataReadbackInt && !Context->GPUDebugDataReadbackCounts)
+		{
+			// Do nothing.., handled in Run
+		}
+		// We may not have floats or ints, but we should have at least one of the two
+		else if ((Context->GPUDebugDataReadbackFloat == nullptr || Context->GPUDebugDataReadbackFloat->IsReady()) 
+				&& (Context->GPUDebugDataReadbackInt == nullptr || Context->GPUDebugDataReadbackInt->IsReady())
+				&& Context->GPUDebugDataReadbackCounts->IsReady()
+			)
+		{
+			//UE_LOG(LogNiagara, Warning, TEXT("Read back!"));
+
+			int32 *NumInstancesAfterSim = static_cast<int32*>(Context->GPUDebugDataReadbackCounts->RetrieveData(64 * sizeof(int32)));
+			int32 NewExistingDataCount = NumInstancesAfterSim[1];
+			{
+				float* FloatDataBuffer = nullptr;
+				if (Context->GPUDebugDataReadbackFloat)
+				{
+					FloatDataBuffer = static_cast<float*>(Context->GPUDebugDataReadbackFloat->RetrieveData(Context->GPUDebugDataFloatSize));
+				}
+				int* IntDataBuffer = nullptr;
+				if (Context->GPUDebugDataReadbackInt)
+				{
+					IntDataBuffer = static_cast<int*>(Context->GPUDebugDataReadbackInt->RetrieveData(Context->GPUDebugDataIntSize));
+				}
+				Context->MainDataSet->DumpGPU(Context->DebugInfo->Frame, FloatDataBuffer, IntDataBuffer, 0, NewExistingDataCount);
+				Context->DebugInfo->bWritten = true;
+
+				if (Context->GPUDebugDataReadbackFloat)
+				{
+					Context->GPUDebugDataReadbackFloat->Finish();
+				}
+				if (Context->GPUDebugDataReadbackInt)
+				{
+					Context->GPUDebugDataReadbackInt->Finish();
+				}
+				Context->GPUDebugDataReadbackCounts->Finish();
+			}
+			{
+				// The following code seems to take significant time on d3d12
+				// Clear out the readback buffers...
+				if (Context->GPUDebugDataReadbackFloat)
+				{
+					delete Context->GPUDebugDataReadbackFloat;
+					Context->GPUDebugDataReadbackFloat = nullptr;
+				}
+				if (Context->GPUDebugDataReadbackInt)
+				{
+					delete Context->GPUDebugDataReadbackInt;
+					Context->GPUDebugDataReadbackInt = nullptr;
+				}
+				delete Context->GPUDebugDataReadbackCounts;
+				Context->GPUDebugDataReadbackCounts = nullptr;	
+				Context->GPUDebugDataFloatSize = 0;
+				Context->GPUDebugDataIntSize = 0;
+			}
+
+			// We've updated the debug info directly, now we need to no longer keep asking and querying because this frame is done!
+			Context->DebugInfo.Reset();
+		}
+	}
+#endif // WITH_EDITORONLY_DATA
 }
 
 
+
 /* Resize data set buffers and set number of instances 
  *	Allocates one additional instance at the end, which is a scratch instance; by setting the default index from AcquireIndex in the shader
  *	to that scratch index, we can avoid branching in every single OutputData function
@@ -231,6 +326,7 @@ void NiagaraEmitterInstanceBatcher::ResizeCurrentBuffer(FRHICommandList &RHICmdL
 	//
 	if (NewNumInstances > PrevNumInstances)
 	{
+		//UE_LOG(LogNiagara, Warning, TEXT("Resize up!"));
 		Context->MainDataSet->CurrData().AllocateGPU(NewNumInstances + 1, RHICmdList);
 		Context->MainDataSet->CurrData().SetNumInstances(NewNumInstances);
 	}
@@ -239,6 +335,7 @@ void NiagaraEmitterInstanceBatcher::ResizeCurrentBuffer(FRHICommandList &RHICmdL
 	//
 	else if (Context->MainDataSet->CurrData().GetNumInstances() < Context->MainDataSet->PrevData().GetNumInstances())
 	{
+		//UE_LOG(LogNiagara, Warning, TEXT("Resize down!"));
 		Context->MainDataSet->CurrData().AllocateGPU(PrevNumInstances + 1, RHICmdList);
 		Context->MainDataSet->CurrData().SetNumInstances(PrevNumInstances);
 	}
@@ -247,7 +344,7 @@ void NiagaraEmitterInstanceBatcher::ResizeCurrentBuffer(FRHICommandList &RHICmdL
 
 /* Set shader parameters for data interfaces
  */
-void NiagaraEmitterInstanceBatcher::SetDataInterfaceParameters(const TArray<UNiagaraDataInterface*> &DataInterfaces, FNiagaraShader* Shader, FRHICommandList &RHICmdList) const
+void NiagaraEmitterInstanceBatcher::SetDataInterfaceParameters(const TArray<UNiagaraDataInterface*> &DataInterfaces, FNiagaraShader* Shader, FRHICommandList &RHICmdList, const FNiagaraComputeExecutionContext *Context) const
 {
 	// set up data interface buffers, as defined by the DIs during compilation
 	//
@@ -257,7 +354,44 @@ void NiagaraEmitterInstanceBatcher::SetDataInterfaceParameters(const TArray<UNia
 		FNiagaraDataInterfaceParamRef& DIParam = Shader->GetDIParameters()[InterfaceIndex];
 		if (DIParam.Parameters)
 		{
-			DIParam.Parameters->Set(RHICmdList, Shader, Interface);
+			void* PerInstanceData = nullptr;
+			int32* OffsetFound = nullptr;
+			if (Context->PerInstanceDataSize != 0 && Context->PerInstanceDataInterfaceOffsets != nullptr)
+			{
+				OffsetFound = Context->PerInstanceDataInterfaceOffsets->Find(Interface);
+				if (OffsetFound != nullptr)
+				{
+					PerInstanceData = (*OffsetFound) + (uint8*)Context->PerInstanceData;
+				}
+			}
+			DIParam.Parameters->Set(RHICmdList, Shader, Interface, PerInstanceData);
+		}
+
+		InterfaceIndex++;
+	}
+}
+
+void NiagaraEmitterInstanceBatcher::UnsetDataInterfaceParameters(const TArray<UNiagaraDataInterface*> &DataInterfaces, FNiagaraShader* Shader, FRHICommandList &RHICmdList, const FNiagaraComputeExecutionContext *Context) const
+{
+	// set up data interface buffers, as defined by the DIs during compilation
+	//
+	uint32 InterfaceIndex = 0;
+	for (UNiagaraDataInterface* Interface : DataInterfaces)
+	{
+		FNiagaraDataInterfaceParamRef& DIParam = Shader->GetDIParameters()[InterfaceIndex];
+		if (DIParam.Parameters)
+		{
+			void* PerInstanceData = nullptr;
+			int32* OffsetFound = nullptr;
+			if (Context->PerInstanceDataSize != 0 && Context->PerInstanceDataInterfaceOffsets != nullptr)
+			{
+				OffsetFound = Context->PerInstanceDataInterfaceOffsets->Find(Interface);
+				if (OffsetFound != nullptr)
+				{
+					PerInstanceData = (*OffsetFound) + (uint8*)Context->PerInstanceData;
+				}
+			}
+			DIParam.Parameters->Unset(RHICmdList, Shader, Interface, PerInstanceData);
 		}
 
 		InterfaceIndex++;
@@ -272,9 +406,14 @@ void NiagaraEmitterInstanceBatcher::Run(const FNiagaraComputeExecutionContext *C
 {
 	if (TotalNumInstances == 0)
 	{
+		SCOPED_DRAW_EVENTF(RHICmdList, NiagaraGPUSimulationCS, TEXT("Niagara Gpu Sim - %s - NumInstances: %u"),
+			*Context->DebugSimName,
+			TotalNumInstances);
 		return;
 	}
 
+	//UE_LOG(LogNiagara, Warning, TEXT("Run"));
+
 	FNiagaraDataSet *DataSet = Context->MainDataSet;
 	const FNiagaraParameterStore& ParameterStore = Context->CombinedParamStore;
 	const FRHIUniformBufferLayout& CBufferLayout = Context->CBufferLayout;
@@ -304,11 +443,14 @@ void NiagaraEmitterInstanceBatcher::Run(const FNiagaraComputeExecutionContext *C
 		RHICmdList.SetShaderUniformBuffer(Shader->GetComputeShader(), Shader->ViewUniformBufferParam.GetBaseIndex(), ViewUniformBuffer);
 	}
 
-	SetDataInterfaceParameters(ParameterStore.GetDataInterfaces(), Shader, RHICmdList);
+	SetDataInterfaceParameters(ParameterStore.GetDataInterfaces(), Shader, RHICmdList, Context);
 
 	// set the shader and data set params 
 	//
-	DataSet->SetShaderParams(Shader, RHICmdList);
+	uint32 WriteBufferIdx = 0;
+	uint32 ReadBufferIdx = 0;
+	DataSet->SetShaderParams(Shader, RHICmdList, WriteBufferIdx, ReadBufferIdx);
+	
 
 	// set the index buffer uav
 	//
@@ -350,13 +492,50 @@ void NiagaraEmitterInstanceBatcher::Run(const FNiagaraComputeExecutionContext *C
 	//
 	if (TotalNumInstances)
 	{
-		SCOPED_DRAW_EVENTF(RHICmdList, NiagaraGPUSimulationCS, TEXT("Niagara GPU Simulation"));
+		SCOPED_DRAW_EVENTF(RHICmdList, NiagaraGPUSimulationCS, TEXT("Niagara Gpu Sim - %s - NumInstances: %u"),
+			*Context->DebugSimName,
+			TotalNumInstances);
 		SCOPED_GPU_STAT(RHICmdList, NiagaraGPUSimulation);
 		DispatchComputeShader(RHICmdList, Shader, NumThreadGroups, 1, 1);
 	}
 
+#if WITH_EDITORONLY_DATA
+	// Check to see if we need to queue up a debug dump..
+	if (Context->DebugInfo.IsValid())
+	{
+		//UE_LOG(LogNiagara, Warning, TEXT("Queued up!"));
+
+		if (!Context->GPUDebugDataReadbackFloat && !Context->GPUDebugDataReadbackInt && !Context->GPUDebugDataReadbackCounts && Context->MainDataSet != nullptr)
+		{
+			FRWBuffer &DatasetIndexBufferWrite = Context->MainDataSet->GetCurDataSetIndices();
+
+			Context->GPUDebugDataCurrBufferIdx = Context->MainDataSet->GetCurrBufferIdx();
+			Context->GPUDebugDataFloatSize = 0;
+			Context->GPUDebugDataIntSize = 0;
+
+			if (Context->MainDataSet->GetNumFloatComponents() > 0)
+			{
+				Context->GPUDebugDataReadbackFloat = new FRHIGPUMemoryReadback(Context->MainDataSet->GetDataByIndex(WriteBufferIdx).GetGPUBufferFloat()->Buffer, TEXT("Niagara GPU Debug Info Float Emitter Readback"));
+				Context->GPUDebugDataReadbackFloat->Insert(RHICmdList);
+				Context->GPUDebugDataFloatSize = Context->MainDataSet->GetDataByIndex(WriteBufferIdx).GetGPUBufferFloat()->NumBytes;
+			}
+
+			if (Context->MainDataSet->GetNumInt32Components() > 0)
+			{
+				Context->GPUDebugDataReadbackInt = new FRHIGPUMemoryReadback(Context->MainDataSet->GetDataByIndex(WriteBufferIdx).GetGPUBufferInt()->Buffer, TEXT("Niagara GPU Debug Info Int Emitter Readback"));
+				Context->GPUDebugDataReadbackInt->Insert(RHICmdList);
+				Context->GPUDebugDataIntSize = Context->MainDataSet->GetDataByIndex(WriteBufferIdx).GetGPUBufferInt()->NumBytes;
+			}
+
+			Context->GPUDebugDataReadbackCounts = new FRHIGPUMemoryReadback(DatasetIndexBufferWrite.Buffer, TEXT("Niagara GPU Emitter Readback"));
+			Context->GPUDebugDataReadbackCounts->Insert(RHICmdList);
+		}
+	}
+#endif // WITH_EDITORONLY_DATA
+
 	// Unset UAV parameters and transition resources (TODO: resource transition should be moved to the renderer)
 	// 
+	UnsetDataInterfaceParameters(ParameterStore.GetDataInterfaces(), Shader, RHICmdList, Context);
 	DataSet->UnsetShaderParams(Shader, RHICmdList);
 	Shader->OutputIndexBufferParam.UnsetUAV(RHICmdList, Shader->GetComputeShader());
 }
