Merging //UE4/Dev-Main to Dev-Build (//UE4/Dev-Build)

diff --git a/Engine/Plugins/FX/Niagara/Source/Niagara/Private/NiagaraLightRendererProperties.cpp b/Engine/Plugins/FX/Niagara/Source/Niagara/Private/NiagaraLightRendererProperties.cpp
--- a/Engine/Plugins/FX/Niagara/Source/Niagara/Private/NiagaraLightRendererProperties.cpp
+++ b/Engine/Plugins/FX/Niagara/Source/Niagara/Private/NiagaraLightRendererProperties.cpp
@@ -4,7 +4,7 @@
 #include "NiagaraRenderer.h"
 #include "NiagaraConstants.h"
 UNiagaraLightRendererProperties::UNiagaraLightRendererProperties()
-	: RadiusScale(1.0f), ColorAdd(FVector(0.0f, 0.0f, 0.0f))
+	: bUseInverseSquaredFalloff(1), bAffectsTranslucency(0), bOverrideRenderingEnabled(0), RadiusScale(1.0f), ColorAdd(FVector(0.0f, 0.0f, 0.0f))
 {
 }
 
@@ -16,6 +16,9 @@ void UNiagaraLightRendererProperties::PostInitProperties()
 		PositionBinding = FNiagaraConstants::GetAttributeDefaultBinding(SYS_PARAM_PARTICLES_POSITION);
 		ColorBinding = FNiagaraConstants::GetAttributeDefaultBinding(SYS_PARAM_PARTICLES_COLOR);
 		RadiusBinding = FNiagaraConstants::GetAttributeDefaultBinding(SYS_PARAM_PARTICLES_LIGHT_RADIUS);
+		LightExponentBinding = FNiagaraConstants::GetAttributeDefaultBinding(SYS_PARAM_PARTICLES_LIGHT_EXPONENT);
+		LightRenderingEnabledBinding = FNiagaraConstants::GetAttributeDefaultBinding(SYS_PARAM_PARTICLES_LIGHT_ENABLED);
+		VolumetricScatteringBinding = FNiagaraConstants::GetAttributeDefaultBinding(SYS_PARAM_PARTICLES_LIGHT_VOLUMETRIC_SCATTERING);
 	}
 }
 
@@ -27,6 +30,9 @@ void UNiagaraLightRendererProperties::InitCDOPropertiesAfterModuleStartup()
 	CDO->PositionBinding = FNiagaraConstants::GetAttributeDefaultBinding(SYS_PARAM_PARTICLES_POSITION);
 	CDO->ColorBinding = FNiagaraConstants::GetAttributeDefaultBinding(SYS_PARAM_PARTICLES_COLOR);
 	CDO->RadiusBinding = FNiagaraConstants::GetAttributeDefaultBinding(SYS_PARAM_PARTICLES_LIGHT_RADIUS);
+	CDO->LightExponentBinding = FNiagaraConstants::GetAttributeDefaultBinding(SYS_PARAM_PARTICLES_LIGHT_EXPONENT);
+	CDO->LightRenderingEnabledBinding = FNiagaraConstants::GetAttributeDefaultBinding(SYS_PARAM_PARTICLES_LIGHT_ENABLED);
+	CDO->VolumetricScatteringBinding = FNiagaraConstants::GetAttributeDefaultBinding(SYS_PARAM_PARTICLES_LIGHT_VOLUMETRIC_SCATTERING);
 }
 
 NiagaraRenderer* UNiagaraLightRendererProperties::CreateEmitterRenderer(ERHIFeatureLevel::Type FeatureLevel)
@@ -56,6 +62,9 @@ const TArray<FNiagaraVariable>& UNiagaraLightRendererProperties::GetOptionalAttr
 		Attrs.Add(SYS_PARAM_PARTICLES_POSITION);
 		Attrs.Add(SYS_PARAM_PARTICLES_COLOR);
 		Attrs.Add(SYS_PARAM_PARTICLES_LIGHT_RADIUS);
+		Attrs.Add(SYS_PARAM_PARTICLES_LIGHT_EXPONENT);
+		Attrs.Add(SYS_PARAM_PARTICLES_LIGHT_ENABLED);
+		Attrs.Add(SYS_PARAM_PARTICLES_LIGHT_VOLUMETRIC_SCATTERING);
 	}
 	return Attrs;
 }
diff --git a/Engine/Plugins/FX/Niagara/Source/Niagara/Private/NiagaraModule.cpp b/Engine/Plugins/FX/Niagara/Source/Niagara/Private/NiagaraModule.cpp
--- a/Engine/Plugins/FX/Niagara/Source/Niagara/Private/NiagaraModule.cpp
+++ b/Engine/Plugins/FX/Niagara/Source/Niagara/Private/NiagaraModule.cpp
@@ -69,6 +69,7 @@ FNiagaraVariable INiagaraModule::Engine_Owner_XAxis;
 FNiagaraVariable INiagaraModule::Engine_Owner_YAxis;
 FNiagaraVariable INiagaraModule::Engine_Owner_ZAxis;
 FNiagaraVariable INiagaraModule::Engine_Owner_Scale;
+FNiagaraVariable INiagaraModule::Engine_Owner_Rotation;
 
 FNiagaraVariable INiagaraModule::Engine_Owner_SystemLocalToWorld;
 FNiagaraVariable INiagaraModule::Engine_Owner_SystemWorldToLocal;
@@ -84,6 +85,8 @@ FNiagaraVariable INiagaraModule::Engine_Owner_ExecutionState;
 
 FNiagaraVariable INiagaraModule::Engine_ExecutionCount;
 FNiagaraVariable INiagaraModule::Engine_Emitter_NumParticles;
+FNiagaraVariable INiagaraModule::Engine_Emitter_TotalSpawnedParticles;
+FNiagaraVariable INiagaraModule::Engine_System_TickCount;
 FNiagaraVariable INiagaraModule::Engine_System_NumEmittersAlive;
 FNiagaraVariable INiagaraModule::Engine_System_NumEmitters;
 FNiagaraVariable INiagaraModule::Engine_NumSystemInstances;
@@ -95,11 +98,14 @@ FNiagaraVariable INiagaraModule::Engine_System_Age;
 
 FNiagaraVariable INiagaraModule::Emitter_Age;
 FNiagaraVariable INiagaraModule::Emitter_LocalSpace;
+FNiagaraVariable INiagaraModule::Emitter_Determinism;
+FNiagaraVariable INiagaraModule::Emitter_RandomSeed;
 FNiagaraVariable INiagaraModule::Emitter_SpawnRate;
 FNiagaraVariable INiagaraModule::Emitter_SpawnInterval;
 FNiagaraVariable INiagaraModule::Emitter_InterpSpawnStartDt;
 FNiagaraVariable INiagaraModule::Emitter_SpawnGroup;
 
+FNiagaraVariable INiagaraModule::Particles_UniqueID;
 FNiagaraVariable INiagaraModule::Particles_ID;
 FNiagaraVariable INiagaraModule::Particles_Position;
 FNiagaraVariable INiagaraModule::Particles_Velocity;
@@ -121,6 +127,9 @@ FNiagaraVariable INiagaraModule::Particles_UVScale;
 FNiagaraVariable INiagaraModule::Particles_CameraOffset;
 FNiagaraVariable INiagaraModule::Particles_MaterialRandom;
 FNiagaraVariable INiagaraModule::Particles_LightRadius;
+FNiagaraVariable INiagaraModule::Particles_LightExponent;
+FNiagaraVariable INiagaraModule::Particles_LightEnabled;
+FNiagaraVariable INiagaraModule::Particles_LightVolumetricScattering;
 FNiagaraVariable INiagaraModule::Particles_RibbonID;
 FNiagaraVariable INiagaraModule::Particles_RibbonWidth;
 FNiagaraVariable INiagaraModule::Particles_RibbonTwist;
@@ -167,6 +176,7 @@ void INiagaraModule::StartupModule()
 	Engine_Owner_YAxis = FNiagaraVariable(FNiagaraTypeDefinition::GetVec3Def(), TEXT("Engine.Owner.SystemYAxis"));
 	Engine_Owner_ZAxis = FNiagaraVariable(FNiagaraTypeDefinition::GetVec3Def(), TEXT("Engine.Owner.SystemZAxis"));
 	Engine_Owner_Scale = FNiagaraVariable(FNiagaraTypeDefinition::GetVec3Def(), TEXT("Engine.Owner.Scale"));
+	Engine_Owner_Rotation = FNiagaraVariable(FNiagaraTypeDefinition::GetQuatDef(), TEXT("Engine.Owner.Rotation"));
 
 	Engine_Owner_SystemLocalToWorld = FNiagaraVariable(FNiagaraTypeDefinition::GetMatrix4Def(), TEXT("Engine.Owner.SystemLocalToWorld"));
 	Engine_Owner_SystemWorldToLocal = FNiagaraVariable(FNiagaraTypeDefinition::GetMatrix4Def(), TEXT("Engine.Owner.SystemWorldToLocal"));
@@ -182,6 +192,8 @@ void INiagaraModule::StartupModule()
 	
 	Engine_ExecutionCount = FNiagaraVariable(FNiagaraTypeDefinition::GetIntDef(), TEXT("Engine.ExecutionCount"));
 	Engine_Emitter_NumParticles = FNiagaraVariable(FNiagaraTypeDefinition::GetIntDef(), TEXT("Engine.Emitter.NumParticles"));
+	Engine_Emitter_TotalSpawnedParticles = FNiagaraVariable(FNiagaraTypeDefinition::GetIntDef(), TEXT("Engine.Emitter.TotalSpawnedParticles"));
+	Engine_System_TickCount = FNiagaraVariable(FNiagaraTypeDefinition::GetIntDef(), TEXT("Engine.System.TickCount"));
 	Engine_System_NumEmittersAlive = FNiagaraVariable(FNiagaraTypeDefinition::GetIntDef(), TEXT("Engine.System.NumEmittersAlive"));
 	Engine_System_NumEmitters = FNiagaraVariable(FNiagaraTypeDefinition::GetIntDef(), TEXT("Engine.System.NumEmitters"));
 	Engine_NumSystemInstances = FNiagaraVariable(FNiagaraTypeDefinition::GetIntDef(), TEXT("Engine.NumSystemInstances"));
@@ -192,11 +204,14 @@ void INiagaraModule::StartupModule()
 	Engine_System_Age = FNiagaraVariable(FNiagaraTypeDefinition::GetFloatDef(), TEXT("Engine.System.Age"));
 	Emitter_Age = FNiagaraVariable(FNiagaraTypeDefinition::GetFloatDef(), TEXT("Emitter.Age"));
 	Emitter_LocalSpace = FNiagaraVariable(FNiagaraTypeDefinition::GetBoolDef(), TEXT("Emitter.LocalSpace"));
+	Emitter_RandomSeed = FNiagaraVariable(FNiagaraTypeDefinition::GetIntDef(), TEXT("Emitter.RandomSeed"));
+	Emitter_Determinism = FNiagaraVariable(FNiagaraTypeDefinition::GetBoolDef(), TEXT("Emitter.Determinism"));
 	Emitter_SpawnRate = FNiagaraVariable(FNiagaraTypeDefinition::GetFloatDef(), TEXT("Emitter.SpawnRate"));
 	Emitter_SpawnInterval = FNiagaraVariable(FNiagaraTypeDefinition::GetFloatDef(), TEXT("Emitter.SpawnInterval"));
 	Emitter_InterpSpawnStartDt = FNiagaraVariable(FNiagaraTypeDefinition::GetFloatDef(), TEXT("Emitter.InterpSpawnStartDt"));
 	Emitter_SpawnGroup = FNiagaraVariable(FNiagaraTypeDefinition::GetIntDef(), TEXT("Emitter.SpawnGroup"));
 
+	Particles_UniqueID = FNiagaraVariable(FNiagaraTypeDefinition::GetIntDef(), TEXT("Particles.UniqueID"));
 	Particles_ID = FNiagaraVariable(FNiagaraTypeDefinition::GetIDDef(), TEXT("Particles.ID"));
 	Particles_Position = FNiagaraVariable(FNiagaraTypeDefinition::GetVec3Def(), TEXT("Particles.Position"));
 	Particles_Velocity = FNiagaraVariable(FNiagaraTypeDefinition::GetVec3Def(), TEXT("Particles.Velocity"));
@@ -218,6 +233,9 @@ void INiagaraModule::StartupModule()
 	Particles_CameraOffset = FNiagaraVariable(FNiagaraTypeDefinition::GetFloatDef(), TEXT("Particles.CameraOffset"));
 	Particles_MaterialRandom = FNiagaraVariable(FNiagaraTypeDefinition::GetFloatDef(), TEXT("Particles.MaterialRandom"));
 	Particles_LightRadius = FNiagaraVariable(FNiagaraTypeDefinition::GetFloatDef(), TEXT("Particles.LightRadius"));
+	Particles_LightExponent = FNiagaraVariable(FNiagaraTypeDefinition::GetFloatDef(), TEXT("Particles.LightExponent"));
+	Particles_LightEnabled = FNiagaraVariable(FNiagaraTypeDefinition::GetBoolDef(), TEXT("Particles.LightEnabled"));
+	Particles_LightVolumetricScattering = FNiagaraVariable(FNiagaraTypeDefinition::GetFloatDef(), TEXT("Particles.LightVolumetricScattering"));
 	Particles_RibbonID = FNiagaraVariable(FNiagaraTypeDefinition::GetIDDef(), TEXT("Particles.RibbonID"));
 	Particles_RibbonWidth = FNiagaraVariable(FNiagaraTypeDefinition::GetFloatDef(), TEXT("Particles.RibbonWidth"));
 	Particles_RibbonTwist = FNiagaraVariable(FNiagaraTypeDefinition::GetFloatDef(), TEXT("Particles.RibbonTwist"));
@@ -598,6 +616,8 @@ void FNiagaraTypeDefinition::RecreateUserDefinedTypeRegistry()
 
 	for (FSoftObjectPath AssetRef : TotalStructAssets)
 	{
+		FName AssetRefPathNamePreResolve = AssetRef.GetAssetPathName();
+
 		UObject* Obj = AssetRef.ResolveObject();
 		if (Obj == nullptr)
 		{
@@ -613,6 +633,11 @@ void FNiagaraTypeDefinition::RecreateUserDefinedTypeRegistry()
 			{
 				FNiagaraTypeRegistry::Register(ScriptStruct, ParamRefFound != nullptr, PayloadRefFound != nullptr, true);
 			}
+			if (Obj->GetPathName() != AssetRefPathNamePreResolve.ToString())
+			{
+				UE_LOG(LogNiagara, Warning, TEXT("Additional parameter/payload enum has moved from where it was in settings (this may cause errors at runtime): Was: \"%s\" Now: \"%s\""), *AssetRefPathNamePreResolve.ToString(), *Obj->GetPathName());
+			}
+
 		}
 		else
 		{
@@ -623,6 +648,7 @@ void FNiagaraTypeDefinition::RecreateUserDefinedTypeRegistry()
 
 	for (FSoftObjectPath AssetRef : Settings->AdditionalParameterEnums)
 	{
+		FName AssetRefPathNamePreResolve = AssetRef.GetAssetPathName();
 		UObject* Obj = AssetRef.ResolveObject();
 		if (Obj == nullptr)
 		{
@@ -631,13 +657,18 @@ void FNiagaraTypeDefinition::RecreateUserDefinedTypeRegistry()
 
 		if (Obj != nullptr)
 		{
-			const FSoftObjectPath* ParamRefFound = Settings->AdditionalParameterEnums.FindByPredicate([&](const FStringAssetReference& Ref) { return Ref.ToString() == AssetRef.ToString(); });
+			const FSoftObjectPath* ParamRefFound = Settings->AdditionalParameterEnums.FindByPredicate([&](const FSoftObjectPath& Ref) { return Ref.ToString() == AssetRef.ToString(); });
 			const FSoftObjectPath* PayloadRefFound = nullptr;
 			UEnum* Enum = Cast<UEnum>(Obj);
 			if (Enum != nullptr)
 			{
 				FNiagaraTypeRegistry::Register(Enum, ParamRefFound != nullptr, PayloadRefFound != nullptr, true);
 			}
+
+			if (Obj->GetPathName() != AssetRefPathNamePreResolve.ToString())
+			{
+				UE_LOG(LogNiagara, Warning, TEXT("Additional parameter/payload enum has moved from where it was in settings (this may cause errors at runtime): Was: \"%s\" Now: \"%s\""), *AssetRefPathNamePreResolve.ToString(), *Obj->GetPathName());
+			}
 		}
 		else
 		{
@@ -757,27 +788,10 @@ FNiagaraTypeDefinition FNiagaraTypeDefinition::GetNumericOutputType(const TArray
 	{
 		return SortedTypeDefinitions[0];
 	}
-}
-
-//////////////////////////////////////////////////////////////////////////
 
-template<>
-void FNiagaraVariable::SetValue(const bool& Data)
-{
-	check(TypeDef == FNiagaraTypeDefinition::GetBoolDef());
-	AllocateData();
-	FNiagaraBool* BoolStruct = (FNiagaraBool*)GetData();
-	BoolStruct->SetValue(Data);
+	return FNiagaraTypeDefinition::GetGenericNumericDef();
 }
 
-template<>
-bool FNiagaraVariable::GetValue() const
-{
-	check(TypeDef == FNiagaraTypeDefinition::GetBoolDef());
-	check(IsDataAllocated());
-	FNiagaraBool* BoolStruct = (FNiagaraBool*)GetData();
-	return BoolStruct->GetValue();
-}
 
 //////////////////////////////////////////////////////////////////////////
 
diff --git a/Engine/Plugins/FX/Niagara/Source/Niagara/Private/NiagaraParameterStore.cpp b/Engine/Plugins/FX/Niagara/Source/Niagara/Private/NiagaraParameterStore.cpp
--- a/Engine/Plugins/FX/Niagara/Source/Niagara/Private/NiagaraParameterStore.cpp
+++ b/Engine/Plugins/FX/Niagara/Source/Niagara/Private/NiagaraParameterStore.cpp
@@ -389,8 +389,33 @@ void FNiagaraParameterStore::CopyParametersTo(FNiagaraParameterStore& DestStore,
 	while (It)
 	{
 		FNiagaraVariable Parameter = It->Key;
-		int32 DestIndex = DestStore.IndexOf(Parameter);
 		int32 SrcIndex = It->Value;
+		++It;
+
+		if (Parameter.IsValid() == false)
+		{
+			FString StoreDebugName;
+#if WITH_EDITORONLY_DATA
+			StoreDebugName = DebugName.IsEmpty() == false ? DebugName : TEXT("Unknown");
+#else
+			StoreDebugName = TEXT("Unknown");
+#endif
+			FString StoreName;
+			if (Owner != nullptr)
+			{
+				StoreName = Owner->GetPathName() + TEXT(".") + StoreDebugName;
+			}
+			else
+			{
+				StoreName = StoreDebugName;
+			}
+
+			UE_LOG(LogNiagara, Error, TEXT("Invalid parameter found while attempting to copy parameters from one parameter store to another.  Parameter Store: %s Parameter Name: %s Parameter Type: %s"), 
+				*StoreName, *Parameter.GetName().ToString(), Parameter.GetType().IsValid() ? *Parameter.GetType().GetName() : TEXT("Unknown"));
+			continue;
+		}
+
+		int32 DestIndex = DestStore.IndexOf(Parameter);
 		bool bWrite = false;
 		if (DestIndex == INDEX_NONE)
 		{
@@ -431,7 +456,6 @@ void FNiagaraParameterStore::CopyParametersTo(FNiagaraParameterStore& DestStore,
 				}
 			}
 		}
-		++It;
 	}
 	DestStore.OnLayoutChange();
 }
diff --git a/Engine/Plugins/FX/Niagara/Source/Niagara/Private/NiagaraRenderer.cpp b/Engine/Plugins/FX/Niagara/Source/Niagara/Private/NiagaraRenderer.cpp
--- a/Engine/Plugins/FX/Niagara/Source/Niagara/Private/NiagaraRenderer.cpp
+++ b/Engine/Plugins/FX/Niagara/Source/Niagara/Private/NiagaraRenderer.cpp
@@ -231,41 +231,49 @@ FNiagaraDynamicDataBase *NiagaraRendererLights::GenerateVertexData(const FNiagar
 	FNiagaraDataSetIterator<FVector> PosItr(Data, Properties->PositionBinding.DataSetVariable);
 	FNiagaraDataSetIterator<FLinearColor> ColItr(Data, Properties->ColorBinding.DataSetVariable);
 	FNiagaraDataSetIterator<float> RadiusItr(Data, Properties->RadiusBinding.DataSetVariable);
+	FNiagaraDataSetIterator<float> ExponentItr(Data, Properties->LightExponentBinding.DataSetVariable);
+	FNiagaraDataSetIterator<float> ScatteringItr(Data, Properties->VolumetricScatteringBinding.DataSetVariable);
+	FNiagaraDataSetIterator<int32> EnabledItr(Data, Properties->LightRenderingEnabledBinding.DataSetVariable);
 
 	FNiagaraDynamicDataLights *DynamicData = new FNiagaraDynamicDataLights;
+	const FMatrix& LocalToWorldMatrix = Proxy->GetLocalToWorld();
 	FVector DefaultColor = FVector(Properties->ColorBinding.DefaultValueIfNonExistent.GetValue<FLinearColor>());
-	FVector DefaultPos = FVector4(Proxy->GetLocalToWorld().GetOrigin());
+	FVector DefaultPos = FVector4(LocalToWorldMatrix.GetOrigin());
 	float DefaultRadius = Properties->RadiusBinding.DefaultValueIfNonExistent.GetValue<float>();
+	float DefaultScattering = Properties->VolumetricScatteringBinding.DefaultValueIfNonExistent.GetValue<float>();
 
 	DynamicData->LightArray.Empty();
 
 	for (uint32 ParticleIndex = 0; ParticleIndex < Data.GetNumInstances(); ParticleIndex++)
 	{
-		SimpleLightData LightData;
-		LightData.LightEntry.Radius = (RadiusItr.IsValid() ? (*RadiusItr) :  DefaultRadius) * Properties->RadiusScale;	//LightPayload->RadiusScale * (Size.X + Size.Y) / 2.0f;
-		LightData.LightEntry.Color = (ColItr.IsValid() ? FVector((*ColItr)) : DefaultColor) + Properties->ColorAdd;				//FVector(Particle.Color) * Particle.Color.A * LightPayload->ColorScale;
-		LightData.LightEntry.Exponent = 1.0;
-		LightData.LightEntry.bAffectTranslucency = true;
-		LightData.PerViewEntry.Position = PosItr.IsValid() ? (*PosItr) : DefaultPos;
+		bool ShouldRenderParticleLight = !Properties->bOverrideRenderingEnabled || (EnabledItr.IsValid() ? (*EnabledItr) : true);
+		float LightRadius = (RadiusItr.IsValid() ? (*RadiusItr) : DefaultRadius) * Properties->RadiusScale;
+		if (ShouldRenderParticleLight && LightRadius > 0)
+		{
+			SimpleLightData LightData;
+			LightData.LightEntry.Radius = LightRadius;
+			LightData.LightEntry.Color = (ColItr.IsValid() ? FVector((*ColItr)) : DefaultColor) + Properties->ColorAdd;
+			LightData.LightEntry.Exponent = Properties->bUseInverseSquaredFalloff ? 0 : (ExponentItr.IsValid() ? (*ExponentItr) : 1);
+			LightData.LightEntry.bAffectTranslucency = Properties->bAffectsTranslucency;
+			LightData.LightEntry.VolumetricScatteringIntensity = ScatteringItr.IsValid() ? (*ScatteringItr) : DefaultScattering;
+			LightData.PerViewEntry.Position = PosItr.IsValid() ? (*PosItr) : DefaultPos;
+			if (bLocalSpace)
+			{
+				LightData.PerViewEntry.Position = LocalToWorldMatrix.TransformPosition(LightData.PerViewEntry.Position);
+			}
 
-		DynamicData->LightArray.Add(LightData);
+			DynamicData->LightArray.Add(LightData);
+		}
 
 		PosItr.Advance();
 		ColItr.Advance();
 		RadiusItr.Advance();
-	}
-
-	if (bLocalSpace)
-	{
-		FMatrix Mat = Proxy->GetLocalToWorld();
-		for (uint32 ParticleIndex = 0; ParticleIndex < Data.GetNumInstances(); ParticleIndex++)
-		{
-			DynamicData->LightArray[ParticleIndex].PerViewEntry.Position = Mat.TransformPosition(DynamicData->LightArray[ParticleIndex].PerViewEntry.Position);
-		}
+		ExponentItr.Advance();
+		ScatteringItr.Advance();
+		EnabledItr.Advance();
 	}
 
 	CPUTimeMS = VertexDataTimer.GetElapsedMilliseconds();
-
 	return DynamicData;
 }
 
diff --git a/Engine/Plugins/FX/Niagara/Source/Niagara/Private/NiagaraRendererMeshes.cpp b/Engine/Plugins/FX/Niagara/Source/Niagara/Private/NiagaraRendererMeshes.cpp
--- a/Engine/Plugins/FX/Niagara/Source/Niagara/Private/NiagaraRendererMeshes.cpp
+++ b/Engine/Plugins/FX/Niagara/Source/Niagara/Private/NiagaraRendererMeshes.cpp
@@ -8,8 +8,8 @@
 #include "Async/ParallelFor.h"
 #include "Engine/StaticMesh.h"
 
-DECLARE_CYCLE_STAT(TEXT("Generate Mesh Vertex Data"), STAT_NiagaraGenMeshVertexData, STATGROUP_Niagara);
-DECLARE_CYCLE_STAT(TEXT("Render Meshes"), STAT_NiagaraRenderMeshes, STATGROUP_Niagara);
+DECLARE_CYCLE_STAT(TEXT("Generate Mesh Vertex Data [GT]"), STAT_NiagaraGenMeshVertexData, STATGROUP_Niagara);
+DECLARE_CYCLE_STAT(TEXT("Render Meshes [RT]"), STAT_NiagaraRenderMeshes, STATGROUP_Niagara);
 
 
 
diff --git a/Engine/Plugins/FX/Niagara/Source/Niagara/Private/NiagaraRendererRibbons.cpp b/Engine/Plugins/FX/Niagara/Source/Niagara/Private/NiagaraRendererRibbons.cpp
--- a/Engine/Plugins/FX/Niagara/Source/Niagara/Private/NiagaraRendererRibbons.cpp
+++ b/Engine/Plugins/FX/Niagara/Source/Niagara/Private/NiagaraRendererRibbons.cpp
@@ -6,18 +6,18 @@
 #include "NiagaraDataSet.h"
 #include "NiagaraStats.h"
 
-DECLARE_CYCLE_STAT(TEXT("Generate Ribbon Vertex Data"), STAT_NiagaraGenRibbonVertexData, STATGROUP_Niagara);
-DECLARE_CYCLE_STAT(TEXT("Render Ribbons"), STAT_NiagaraRenderRibbons, STATGROUP_Niagara);
+DECLARE_CYCLE_STAT(TEXT("Generate Ribbon Vertex Data [GT]"), STAT_NiagaraGenRibbonVertexData, STATGROUP_Niagara);
+DECLARE_CYCLE_STAT(TEXT("Render Ribbons [RT]"), STAT_NiagaraRenderRibbons, STATGROUP_Niagara);
 
 DECLARE_CYCLE_STAT(TEXT("Genereate GPU Buffers"), STAT_NiagaraGenRibbonGpuBuffers, STATGROUP_Niagara);
 
 
 struct FNiagaraDynamicDataRibbon : public FNiagaraDynamicDataBase
 {
-	TArray<FNiagaraRibbonVertex> VertexData;
 	TArray<int16> IndexData;
 	TArray<int32> SortedIndices;
 	TArray<float> TotalDistances;
+	TArray<uint32> MultiRibbonIndices;
 	TArray<float> PackedPerRibbonDataByIndex;
 
 	//Direct ptr to the dataset. ONLY FOR USE BE GPU EMITTERS.
@@ -108,9 +108,6 @@ void NiagaraRendererRibbons::GetDynamicMeshElements(const TArray<const FSceneVie
 	int32 TotalFloatSize = DynamicDataRibbon->RTParticleData.GetFloatBuffer().Num() / sizeof(float);
 	FNiagaraGlobalReadBuffer::FAllocation ParticleData;
 
-	int32 SizeInBytes = DynamicDataRibbon->VertexData.GetTypeSize() * DynamicDataRibbon->VertexData.Num();
-	FGlobalDynamicVertexBuffer::FAllocation LocalDynamicVertexAllocation = FGlobalDynamicVertexBuffer::Get().Allocate(SizeInBytes);
-
 	if (DynamicDataRibbon->DataSet->GetSimTarget() == ENiagaraSimTarget::CPUSim)
 	{
 		ParticleData = FNiagaraGlobalReadBuffer::Get().AllocateFloat(TotalFloatSize);
@@ -140,7 +137,6 @@ void NiagaraRendererRibbons::GetDynamicMeshElements(const TArray<const FSceneVie
 		}
 
 		// Copy the vertex data over.
-		FMemory::Memcpy(LocalDynamicVertexAllocation.Buffer, DynamicDataRibbon->VertexData.GetData(), SizeInBytes);
 		FMemory::Memcpy(DynamicIndexAllocation.Buffer, DynamicDataRibbon->IndexData.GetData(), DynamicDataRibbon->IndexData.Num() * sizeof(int16));
 
 		// Compute the per-view uniform buffers.
@@ -192,7 +188,6 @@ void NiagaraRendererRibbons::GetDynamicMeshElements(const TArray<const FSceneVie
 
 				CollectorResources.VertexFactory.InitResource();
 				CollectorResources.VertexFactory.SetRibbonUniformBuffer(CollectorResources.UniformBuffer);
-				CollectorResources.VertexFactory.SetVertexBuffer(LocalDynamicVertexAllocation.VertexBuffer, LocalDynamicVertexAllocation.VertexOffset, sizeof(FNiagaraRibbonVertex));
 
 				if (!DynamicDataRibbon->SortedIndices.Num())
 				{
@@ -217,6 +212,14 @@ void NiagaraRendererRibbons::GetDynamicMeshElements(const TArray<const FSceneVie
 				RHIUnlockVertexBuffer(TotalDistancesBuffer.Buffer);
 				CollectorResources.VertexFactory.SetSegmentDistances(TotalDistancesBuffer.SRV);
 
+				// Copy a buffer which has the per particle multi ribbon index.
+				FReadBuffer MultiRibbonIndicesBuffer;
+				MultiRibbonIndicesBuffer.Initialize(sizeof(uint32), DynamicDataRibbon->MultiRibbonIndices.Num(), EPixelFormat::PF_R32_UINT, BUF_Volatile);
+				void* MultiRibbonIndexPtr = RHILockVertexBuffer(MultiRibbonIndicesBuffer.Buffer, 0, DynamicDataRibbon->MultiRibbonIndices.Num() * sizeof(uint32), RLM_WriteOnly);
+				FMemory::Memcpy(MultiRibbonIndexPtr, DynamicDataRibbon->MultiRibbonIndices.GetData(), DynamicDataRibbon->MultiRibbonIndices.Num() * sizeof(uint32));
+				RHIUnlockVertexBuffer(MultiRibbonIndicesBuffer.Buffer);
+				CollectorResources.VertexFactory.SetMultiRibbonIndicesSRV(MultiRibbonIndicesBuffer.SRV);
+
 				// Copy the packed u data for stable age based uv generation.
 				FReadBuffer PackedPerRibbonDataByIndexBuffer;
 				PackedPerRibbonDataByIndexBuffer.Initialize(sizeof(float), DynamicDataRibbon->PackedPerRibbonDataByIndex.Num(), EPixelFormat::PF_R32_FLOAT, BUF_Volatile);
@@ -252,7 +255,7 @@ void NiagaraRendererRibbons::GetDynamicMeshElements(const TArray<const FSceneVie
 				check(MeshElement.NumPrimitives > 0);
 				MeshElement.NumInstances = 1;
 				MeshElement.MinVertexIndex = 0;
-				MeshElement.MaxVertexIndex = DynamicDataRibbon->VertexData.Num() - 1;
+				MeshElement.MaxVertexIndex = 0;
 				MeshElement.PrimitiveUniformBufferResource = &WorldSpacePrimitiveUniformBuffer;
 
 				Collector.AddMesh(ViewIndex, MeshBatch);
@@ -280,8 +283,12 @@ int NiagaraRendererRibbons::GetDynamicDataSize()
 	uint32 Size = sizeof(FNiagaraDynamicDataRibbon);
 	if (DynamicDataRender)
 	{
-		Size += (static_cast<FNiagaraDynamicDataRibbon*>(DynamicDataRender))->VertexData.GetAllocatedSize();
-		Size += (static_cast<FNiagaraDynamicDataRibbon*>(DynamicDataRender))->IndexData.GetAllocatedSize();
+		FNiagaraDynamicDataRibbon* RibbonDynamicData = static_cast<FNiagaraDynamicDataRibbon*>(DynamicDataRender);
+		Size += RibbonDynamicData->IndexData.GetAllocatedSize();
+		Size += RibbonDynamicData->SortedIndices.GetAllocatedSize();
+		Size += RibbonDynamicData->TotalDistances.GetAllocatedSize();
+		Size += RibbonDynamicData->MultiRibbonIndices.GetAllocatedSize();
+		Size += RibbonDynamicData->PackedPerRibbonDataByIndex.GetAllocatedSize();
 	}
 
 	return Size;
@@ -381,12 +388,8 @@ FNiagaraDynamicDataBase *NiagaraRendererRibbons::GenerateVertexData(const FNiaga
 		return nullptr;
 	}
 	FNiagaraDynamicDataRibbon* DynamicData = new FNiagaraDynamicDataRibbon;
-	TArray<FNiagaraRibbonVertex>& VertexData = DynamicData->VertexData;
 	TArray<int16>& IndexData = DynamicData->IndexData;
 
-	VertexData.Empty();
-	IndexData.Empty();
-
 	// TODO : deal with the dynamic vertex material parameter should the user have specified it as an output...
 	int32 NumTotalVerts = 0;
 
@@ -495,10 +498,7 @@ FNiagaraDynamicDataBase *NiagaraRendererRibbons::GenerateVertexData(const FNiaga
 					PrevDir = NormDir;
 
 					DynamicData->TotalDistances.Add(TotalDistance);
-					FNiagaraRibbonVertex NewVertex;
-					NewVertex.RibbonIndex = RibbonIndex;
-					VertexData.Add(NewVertex);
-					VertexData.Add(NewVertex);
+					DynamicData->MultiRibbonIndices.Add(RibbonIndex);
 
 					if (i < NumIndices - 1)
 					{
diff --git a/Engine/Plugins/FX/Niagara/Source/Niagara/Private/NiagaraRendererSprites.cpp b/Engine/Plugins/FX/Niagara/Source/Niagara/Private/NiagaraRendererSprites.cpp
--- a/Engine/Plugins/FX/Niagara/Source/Niagara/Private/NiagaraRendererSprites.cpp
+++ b/Engine/Plugins/FX/Niagara/Source/Niagara/Private/NiagaraRendererSprites.cpp
@@ -6,8 +6,8 @@
 #include "NiagaraDataSet.h"
 #include "NiagaraStats.h"
 
-DECLARE_CYCLE_STAT(TEXT("Generate Sprite Vertex Data"), STAT_NiagaraGenSpriteVertexData, STATGROUP_Niagara);
-DECLARE_CYCLE_STAT(TEXT("Render Sprites"), STAT_NiagaraRenderSprites, STATGROUP_Niagara);
+DECLARE_CYCLE_STAT(TEXT("Generate Sprite Vertex Data [GT]"), STAT_NiagaraGenSpriteVertexData, STATGROUP_Niagara);
+DECLARE_CYCLE_STAT(TEXT("Render Sprites [RT]"), STAT_NiagaraRenderSprites, STATGROUP_Niagara);
 
 DECLARE_CYCLE_STAT(TEXT("Genereate GPU Buffers"), STAT_NiagaraGenSpriteGpuBuffers, STATGROUP_Niagara);
 
diff --git a/Engine/Plugins/FX/Niagara/Source/Niagara/Private/NiagaraScript.cpp b/Engine/Plugins/FX/Niagara/Source/Niagara/Private/NiagaraScript.cpp
--- a/Engine/Plugins/FX/Niagara/Source/Niagara/Private/NiagaraScript.cpp
+++ b/Engine/Plugins/FX/Niagara/Source/Niagara/Private/NiagaraScript.cpp
@@ -16,6 +16,9 @@
 #include "UObject/Linker.h"
 #include "HAL/PlatformFilemanager.h"
 #include "Misc/FileHelper.h"
+#include "UObject/EditorObjectVersion.h"
+#include "UObject/ReleaseObjectVersion.h"
+
 
 #if WITH_EDITOR
 	#include "NiagaraScriptDerivedData.h"
@@ -26,13 +29,21 @@
 
 DECLARE_STATS_GROUP(TEXT("Niagara Detailed"), STATGROUP_NiagaraDetailed, STATCAT_Advanced);
 
-FNiagaraScriptDebuggerInfo::FNiagaraScriptDebuggerInfo() : FrameLastWriteId(-1)
+FNiagaraScriptDebuggerInfo::FNiagaraScriptDebuggerInfo() : bWaitForGPU(false), FrameLastWriteId(-1), bWritten(false)
 {
 }
 
 
-FNiagaraScriptDebuggerInfo::FNiagaraScriptDebuggerInfo(FName InName, ENiagaraScriptUsage InUsage, const FGuid& InUsageId) : HandleName(InName), Usage(InUsage), UsageId(InUsageId), FrameLastWriteId(-1)
+FNiagaraScriptDebuggerInfo::FNiagaraScriptDebuggerInfo(FName InName, ENiagaraScriptUsage InUsage, const FGuid& InUsageId) : HandleName(InName), Usage(InUsage), UsageId(InUsageId), FrameLastWriteId(-1), bWritten(false)
 {
+	if (InUsage == ENiagaraScriptUsage::ParticleGPUComputeScript)
+	{
+		bWaitForGPU = true;
+	}
+	else
+	{
+		bWaitForGPU = false;
+	}
 }
 
 
@@ -230,6 +241,10 @@ void UNiagaraScript::ComputeVMCompilationId(FNiagaraVMExecutableDataId& Id) cons
 		{
 			Id.AdditionalDefines.Add(TEXT("Emitter.Localspace"));
 		}
+		if (Emitter->bDeterminism)
+		{
+			Id.AdditionalDefines.Add(TEXT("Emitter.Determinism"));
+		}
 	}
 
 	if (UNiagaraSystem* System = Cast<UNiagaraSystem>(Obj))
@@ -243,6 +258,10 @@ void UNiagaraScript::ComputeVMCompilationId(FNiagaraVMExecutableDataId& Id) cons
 				{
 					Id.AdditionalDefines.Add(Emitter->GetUniqueEmitterName() + TEXT(".Localspace"));
 				}
+				if (Emitter->bDeterminism)
+				{
+					Id.AdditionalDefines.Add(Emitter->GetUniqueEmitterName() + TEXT(".Determinism"));
+				}
 			}
 		}
 	}
@@ -476,7 +495,23 @@ void UNiagaraScript::GenerateStatScopeIDs()
 
 void UNiagaraScript::PostEditChangeProperty(FPropertyChangedEvent& PropertyChangedEvent)
 {
+	Super::PostEditChangeProperty(PropertyChangedEvent);
+
+	FName PropertyName;
+	if (PropertyChangedEvent.Property)
+	{
+		PropertyName = PropertyChangedEvent.Property->GetFName();
+	}
+
 	CacheResourceShadersForRendering(true);
+
+	if (PropertyName == GET_MEMBER_NAME_CHECKED(UNiagaraScript, bDeprecated) || PropertyName == GET_MEMBER_NAME_CHECKED(UNiagaraScript, DeprecationRecommendation))
+	{
+		if (Source)
+		{
+			Source->MarkNotSynchronized(TEXT("Deprecation changed."));
+		}
+	}
 }
 
 #endif
@@ -1034,9 +1069,12 @@ void UNiagaraScript::CacheResourceShadersForRendering(bool bRegenerateId, bool b
 			//if (ScriptResourcesByFeatureLevel[FeatureLevel])
 			{
 				EShaderPlatform ShaderPlatform = GShaderPlatformForFeatureLevel[CacheFeatureLevel];
-				ResourceToCache = ScriptResourcesByFeatureLevel[CacheFeatureLevel];
-				CacheShadersForResources(ShaderPlatform, &ScriptResource, true);
-				ScriptResourcesByFeatureLevel[CacheFeatureLevel] = &ScriptResource;
+				if (IsFeatureLevelSupported(ShaderPlatform, ERHIFeatureLevel::SM5) || IsFeatureLevelSupported(ShaderPlatform, ERHIFeatureLevel::ES3_1))
+				{
+					ResourceToCache = ScriptResourcesByFeatureLevel[CacheFeatureLevel];
+					CacheShadersForResources(ShaderPlatform, &ScriptResource, true);
+					ScriptResourcesByFeatureLevel[CacheFeatureLevel] = &ScriptResource;
+				}
 			}
 		}
 	}
@@ -1063,6 +1101,14 @@ void UNiagaraScript::SyncAliases(const TMap<FString, FString>& RenameMap)
 	// Now handle any Parameters overall..
 	for (int32 i = 0; i < GetVMExecutableData().Parameters.Parameters.Num(); i++)
 	{
+		if (GetVMExecutableData().Parameters.Parameters[i].IsValid() == false)
+		{
+			const FNiagaraVariable& InvalidParameter = GetVMExecutableData().Parameters.Parameters[i];
+			UE_LOG(LogNiagara, Error, TEXT("Invalid parameter found while syncing script aliases.  Script: %s Parameter Name: %s Parameter Type: %s"),
+				*GetPathName(), *InvalidParameter.GetName().ToString(), InvalidParameter.GetType().IsValid() ? *InvalidParameter.GetType().GetName() : TEXT("Unknown"));
+			continue;
+		}
+
 		FNiagaraVariable Var = GetVMExecutableData().Parameters.Parameters[i];
 		FNiagaraVariable NewVar = FNiagaraVariable::ResolveAliases(Var, RenameMap);
 		if (NewVar.GetName() != Var.GetName())
@@ -1199,6 +1245,11 @@ NIAGARA_API bool UNiagaraScript::DidScriptCompilationSucceed(bool bGPUScript) co
 
 void SerializeNiagaraShaderMaps(const TMap<const ITargetPlatform*, TArray<FNiagaraShaderScript*>>* PlatformScriptResourcesToSave, FArchive& Ar, TArray<FNiagaraShaderScript>& OutLoadedResources)
 {
+	Ar.UsingCustomVersion(FFortniteMainBranchObjectVersion::GUID);
+	Ar.UsingCustomVersion(FRenderingObjectVersion::GUID);
+	Ar.UsingCustomVersion(FEditorObjectVersion::GUID);
+	Ar.UsingCustomVersion(FReleaseObjectVersion::GUID);
+
 //	SCOPED_LOADTIMER(SerializeInlineShaderMaps);
 	if (Ar.IsSaving())
 	{
diff --git a/Engine/Plugins/FX/Niagara/Source/Niagara/Private/NiagaraScriptExecutionContext.cpp b/Engine/Plugins/FX/Niagara/Source/Niagara/Private/NiagaraScriptExecutionContext.cpp
--- a/Engine/Plugins/FX/Niagara/Source/Niagara/Private/NiagaraScriptExecutionContext.cpp
+++ b/Engine/Plugins/FX/Niagara/Source/Niagara/Private/NiagaraScriptExecutionContext.cpp
@@ -16,6 +16,14 @@ DECLARE_CYCLE_STAT(TEXT("Rebind DInterface Func Table"), STAT_NiagaraRebindDataI
 
 uint32 FNiagaraScriptExecutionContext::TickCounter = 0;
 
+static int32 GbExecVMScripts = 1;
+static FAutoConsoleVariableRef CVarNiagaraExecVMScripts(
+	TEXT("fx.ExecVMScripts"),
+	GbExecVMScripts,
+	TEXT("If > 0 VM scripts will be executed, otherwise they won't, useful for looking at the bytecode for a crashing compiled script. \n"),
+	ECVF_Default
+);
+
 FNiagaraScriptExecutionContext::FNiagaraScriptExecutionContext()
 	: Script(nullptr)
 {
@@ -164,7 +172,7 @@ bool FNiagaraScriptExecutionContext::Execute(uint32 NumInstances, TArray<FNiagar
 		}
 	}
 
-	if (bRegisterSetupCompleted)
+	if (bRegisterSetupCompleted && GbExecVMScripts != 0)
 	{
 		VectorVM::Exec(
 			Script->GetVMExecutableData().ByteCode.GetData(),
diff --git a/Engine/Plugins/FX/Niagara/Source/Niagara/Private/NiagaraStats.h b/Engine/Plugins/FX/Niagara/Source/Niagara/Private/NiagaraStats.h
--- a/Engine/Plugins/FX/Niagara/Source/Niagara/Private/NiagaraStats.h
+++ b/Engine/Plugins/FX/Niagara/Source/Niagara/Private/NiagaraStats.h
@@ -8,11 +8,18 @@
 #include "NiagaraModule.h"
 
 DECLARE_STATS_GROUP(TEXT("Niagara"), STATGROUP_Niagara, STATCAT_Advanced);
-DECLARE_CYCLE_STAT(TEXT("Render Total"), STAT_NiagaraRender, STATGROUP_Niagara);
-DECLARE_CYCLE_STAT(TEXT("Render Total (GT)"), STAT_NiagaraRenderGT, STATGROUP_Niagara);
+DECLARE_CYCLE_STAT(TEXT("Render Total [RT]"), STAT_NiagaraRender, STATGROUP_Niagara);
+DECLARE_CYCLE_STAT(TEXT("Render Total [GT]"), STAT_NiagaraRenderGT, STATGROUP_Niagara);
 DECLARE_DWORD_COUNTER_STAT(TEXT("NumParticles"), STAT_NiagaraNumParticles, STATGROUP_Niagara);
-DECLARE_CYCLE_STAT(TEXT("Constant Setup"), STAT_NiagaraConstants, STATGROUP_Niagara);
-DECLARE_CYCLE_STAT(TEXT("Tick"), STAT_NiagaraTick, STATGROUP_Niagara);
+DECLARE_CYCLE_STAT(TEXT("Constant Setup [CNC]"), STAT_NiagaraConstants, STATGROUP_Niagara);
+DECLARE_CYCLE_STAT(TEXT("Emitter Tick [CNC]"), STAT_NiagaraTick, STATGROUP_Niagara);
 DECLARE_DWORD_COUNTER_STAT(TEXT("NumSystems"), STAT_NiagaraNumSystems, STATGROUP_Niagara);
 DECLARE_MEMORY_STAT(TEXT("Niagara particle data memory"), STAT_NiagaraParticleMemory, STATGROUP_Niagara);
 DECLARE_MEMORY_STAT(TEXT("Niagara vertex buffer memory"), STAT_NiagaraVBMemory, STATGROUP_Niagara);
+
+
+DECLARE_STATS_GROUP(TEXT("NiagaraOverview"), STATGROUP_NiagaraOverview, STATCAT_Advanced);
+DECLARE_CYCLE_STAT(TEXT("GT Total"), STAT_NiagaraOverview_GT, STATGROUP_NiagaraOverview);
+DECLARE_CYCLE_STAT(TEXT("GT Concurrent Total"), STAT_NiagaraOverview_GT_CNC, STATGROUP_NiagaraOverview);
+DECLARE_CYCLE_STAT(TEXT("RT Total"), STAT_NiagaraOverview_RT, STATGROUP_NiagaraOverview);
+DECLARE_CYCLE_STAT(TEXT("RT Concurrent Total"), STAT_NiagaraOverview_RT_CNC, STATGROUP_NiagaraOverview);
\ No newline at end of file
diff --git a/Engine/Plugins/FX/Niagara/Source/Niagara/Private/NiagaraSystem.cpp b/Engine/Plugins/FX/Niagara/Source/Niagara/Private/NiagaraSystem.cpp
--- a/Engine/Plugins/FX/Niagara/Source/Niagara/Private/NiagaraSystem.cpp
+++ b/Engine/Plugins/FX/Niagara/Source/Niagara/Private/NiagaraSystem.cpp
@@ -42,6 +42,7 @@ UNiagaraSystem::UNiagaraSystem(const FObjectInitializer& ObjectInitializer)
 , WarmupTime(0.0f)
 , WarmupTickCount(0)
 , WarmupTickDelta(1.0f / 15.0f)
+, bHasSystemScriptDIsWithPerInstanceData(false)
 {
 }
 
@@ -174,8 +175,6 @@ void UNiagaraSystem::PostEditChangeProperty(struct FPropertyChangedEvent& Proper
 	FNiagaraSystemUpdateContext(this, true);
 
 	ThumbnailImageOutOfDate = true;
-
-	DetermineIfSolo();
 	
 	if (PropertyChangedEvent.Property != nullptr)
 	{
@@ -229,6 +228,10 @@ void UNiagaraSystem::PostLoad()
 		}
 	}
 
+#if UE_EDITOR
+	ExposedParameters.RecreateRedirections();
+#endif
+
 #if WITH_EDITORONLY_DATA
 	TArray<UNiagaraScript*> AllSystemScripts;
 	
@@ -391,8 +394,6 @@ void UNiagaraSystem::PostLoad()
 	}
 #endif
 #endif
-
-	DetermineIfSolo();
 }
 
 #if WITH_EDITORONLY_DATA
@@ -546,46 +547,42 @@ bool UNiagaraSystem::HasOutstandingCompilationRequests() const
 	return ActiveCompilations.Num() > 0;
 }
 
+bool UNiagaraSystem::HasSystemScriptDIsWithPerInstanceData() const
+{
+	return bHasSystemScriptDIsWithPerInstanceData;
+}
 
-bool UNiagaraSystem::IsSolo()const
+const TArray<FName>& UNiagaraSystem::GetUserDINamesReadInSystemScripts() const
 {
-	return bSolo;
+	return UserDINamesReadInSystemScripts;
 }
 
-void UNiagaraSystem::DetermineIfSolo()
+void CheckDICompileInfo(const TArray<FNiagaraScriptDataInterfaceCompileInfo>& ScriptDICompileInfos, bool& bOutbHasSystemDIsWithPerInstanceData, TArray<FName>& OutUserDINamesReadInSystemScripts)
 {
-	//Determine if we can update normally or have to update solo.
-	bSolo = false;
-	//If our scripts have any interfaces that require instance data.
-	UNiagaraScript* SystemSpawn = GetSystemSpawnScript();
-	if (SystemSpawn->GetVMExecutableData().IsValid())
+	for (const FNiagaraScriptDataInterfaceCompileInfo& ScriptDICompileInfo : ScriptDICompileInfos)
 	{
-		for (int32 i = 0; !bSolo && i < SystemSpawn->GetVMExecutableData().DataInterfaceInfo.Num(); ++i)
+		UNiagaraDataInterface* DefaultDataInterface = ScriptDICompileInfo.GetDefaultDataInterface();
+		if (DefaultDataInterface != nullptr && DefaultDataInterface->PerInstanceDataSize() > 0)
 		{
-			FNiagaraScriptDataInterfaceCompileInfo& Info = SystemSpawn->GetVMExecutableData().DataInterfaceInfo[i];
-			if (Info.IsSystemSolo())//Temp hack to force solo on any systems with system scrips needing user (aka per instance) interfaces.
-			{
-				bSolo = true;
-				break;
-			}
+			bOutbHasSystemDIsWithPerInstanceData = true;
 		}
-	}
 
-	UNiagaraScript* SystemUpdate = GetSystemUpdateScript();
-	if (SystemUpdate->GetVMExecutableData().IsValid())
-	{
-		for (int32 i = 0; !bSolo && i < SystemUpdate->GetVMExecutableData().DataInterfaceInfo.Num(); ++i)
+		if (ScriptDICompileInfo.RegisteredParameterMapRead.ToString().StartsWith(TEXT("User.")))
 		{
-			FNiagaraScriptDataInterfaceCompileInfo& Info = SystemUpdate->GetVMExecutableData().DataInterfaceInfo[i];
-			if (Info.IsSystemSolo())//Temp hack to force solo on any systems with system scrips needing user (aka per instance) interfaces.
-			{
-				bSolo = true;
-				break;
-			}
+			OutUserDINamesReadInSystemScripts.AddUnique(ScriptDICompileInfo.RegisteredParameterMapRead);
 		}
 	}
 }
 
+void UNiagaraSystem::UpdatePostCompileDIInfo()
+{
+	bHasSystemScriptDIsWithPerInstanceData = false;
+	UserDINamesReadInSystemScripts.Empty();
+
+	CheckDICompileInfo(SystemSpawnScript->GetVMExecutableData().DataInterfaceInfo, bHasSystemScriptDIsWithPerInstanceData, UserDINamesReadInSystemScripts);
+	CheckDICompileInfo(SystemUpdateScript->GetVMExecutableData().DataInterfaceInfo, bHasSystemScriptDIsWithPerInstanceData, UserDINamesReadInSystemScripts);
+}
+
 bool UNiagaraSystem::IsValid()const
 {
 	if (!SystemSpawnScript || !SystemUpdateScript)
@@ -901,7 +898,7 @@ bool UNiagaraSystem::QueryCompileComplete(bool bWait, bool bDoPost, bool bDoNotA
 
 		ActiveCompilations[ActiveCompileIdx].RootObjects.Empty();
 
-		DetermineIfSolo();
+		UpdatePostCompileDIInfo();
 
 		UE_LOG(LogNiagara, Log, TEXT("Compiling System %s took %f sec (wall time), %f sec (combined time)."), *GetFullName(), (float)(FPlatformTime::Seconds() - ActiveCompilations[ActiveCompileIdx].StartTime),
 			CombinedCompileTime);
diff --git a/Engine/Plugins/FX/Niagara/Source/Niagara/Private/NiagaraSystemInstance.cpp b/Engine/Plugins/FX/Niagara/Source/Niagara/Private/NiagaraSystemInstance.cpp
--- a/Engine/Plugins/FX/Niagara/Source/Niagara/Private/NiagaraSystemInstance.cpp
+++ b/Engine/Plugins/FX/Niagara/Source/Niagara/Private/NiagaraSystemInstance.cpp
@@ -14,16 +14,17 @@
 #include "Templates/AlignmentTemplates.h"
 
 
-DECLARE_CYCLE_STAT(TEXT("System Activate (GT)"), STAT_NiagaraSystemActivate, STATGROUP_Niagara);
-DECLARE_CYCLE_STAT(TEXT("System Deactivate (GT)"), STAT_NiagaraSystemDeactivate, STATGROUP_Niagara);
-DECLARE_CYCLE_STAT(TEXT("System Complete (GT)"), STAT_NiagaraSystemComplete, STATGROUP_Niagara);
-DECLARE_CYCLE_STAT(TEXT("Parallel Tick"), STAT_NiagaraParallelTick, STATGROUP_Niagara);
-DECLARE_CYCLE_STAT(TEXT("System Reset (GT)"), STAT_NiagaraSystemReset, STATGROUP_Niagara);
-DECLARE_CYCLE_STAT(TEXT("System Reinit (GT)"), STAT_NiagaraSystemReinit, STATGROUP_Niagara);
-DECLARE_CYCLE_STAT(TEXT("System Init Emitters (GT)"), STAT_NiagaraSystemInitEmitters, STATGROUP_Niagara);
-DECLARE_CYCLE_STAT(TEXT("System Advance Simulation "), STAT_NiagaraSystemAdvanceSim, STATGROUP_Niagara);
-DECLARE_CYCLE_STAT(TEXT("System SetSolo "), STAT_NiagaraSystemSetSolo, STATGROUP_Niagara); 
-DECLARE_CYCLE_STAT(TEXT("System PreSimulateTick "), STAT_NiagaraSystemPreSimulateTick, STATGROUP_Niagara); 
+DECLARE_CYCLE_STAT(TEXT("System Activate [GT]"), STAT_NiagaraSystemActivate, STATGROUP_Niagara);
+DECLARE_CYCLE_STAT(TEXT("System Deactivate [GT]"), STAT_NiagaraSystemDeactivate, STATGROUP_Niagara);
+DECLARE_CYCLE_STAT(TEXT("System Complete [GT]"), STAT_NiagaraSystemComplete, STATGROUP_Niagara);
+DECLARE_CYCLE_STAT(TEXT("System Reset [GT]"), STAT_NiagaraSystemReset, STATGROUP_Niagara);
+DECLARE_CYCLE_STAT(TEXT("System Reinit [GT]"), STAT_NiagaraSystemReinit, STATGROUP_Niagara);
+DECLARE_CYCLE_STAT(TEXT("System Init Emitters [GT]"), STAT_NiagaraSystemInitEmitters, STATGROUP_Niagara);
+DECLARE_CYCLE_STAT(TEXT("System Advance Simulation [GT] "), STAT_NiagaraSystemAdvanceSim, STATGROUP_Niagara);
+DECLARE_CYCLE_STAT(TEXT("System SetSolo[GT] "), STAT_NiagaraSystemSetSolo, STATGROUP_Niagara);
+
+DECLARE_CYCLE_STAT(TEXT("System PreSimulateTick [CNC]"), STAT_NiagaraSystemPreSimulateTick, STATGROUP_Niagara);
+DECLARE_CYCLE_STAT(TEXT("System Instance Tick [CNC]"), STAT_NiagaraSystemInstanceTick, STATGROUP_Niagara);
 
  
 /** Safety time to allow for the LastRenderTime coming back from the RT. */
@@ -39,6 +40,7 @@ FNiagaraSystemInstance::FNiagaraSystemInstance(UNiagaraComponent* InComponent)
 	: SystemInstanceIndex(INDEX_NONE)
 	, Component(InComponent)
 	, Age(0.0f)
+	, TickCount(0)
 	, ID(FGuid::NewGuid())
 	, IDName(*ID.ToString())
 	, InstanceParameters(Component)
@@ -54,7 +56,7 @@ FNiagaraSystemInstance::FNiagaraSystemInstance(UNiagaraComponent* InComponent)
 	SystemBounds.Init();
 }
 
-void FNiagaraSystemInstance::Init(UNiagaraSystem* InSystem, bool bInForceSolo)
+void FNiagaraSystemInstance::Init(bool bInForceSolo)
 {
 	bForceSolo = bInForceSolo;
 	ActualExecutionState = ENiagaraExecutionState::Inactive;
@@ -139,6 +141,8 @@ bool FNiagaraSystemInstance::RequestCapture(const FGuid& RequestId)
 		return false;
 	}
 
+	UE_LOG(LogNiagara, Warning, TEXT("Capture requested!"));
+
 	bWasSoloPriorToCaptureRequest = bSolo;
 	SetSolo(true);
 
@@ -156,7 +160,10 @@ bool FNiagaraSystemInstance::RequestCapture(const FGuid& RequestId)
 
 		for (UNiagaraScript* Script : Scripts)
 		{
-			TempCaptureHolder->Add(MakeShared<FNiagaraScriptDebuggerInfo, ESPMode::ThreadSafe>(Handle.GetIdName(), Script->GetUsage(), Script->GetUsageId()));
+			TSharedPtr<struct FNiagaraScriptDebuggerInfo, ESPMode::ThreadSafe> DebugInfoPtr = MakeShared<FNiagaraScriptDebuggerInfo, ESPMode::ThreadSafe>(Handle.GetIdName(), Script->GetUsage(), Script->GetUsageId());
+			DebugInfoPtr->bWritten = false;
+
+			TempCaptureHolder->Add(DebugInfoPtr);
 		}
 	}
 	CapturedFrames.Add(RequestId, TempCaptureHolder);
@@ -190,6 +197,27 @@ bool FNiagaraSystemInstance::QueryCaptureResults(const FGuid& RequestId, TArray<
 		TArray<TSharedPtr<struct FNiagaraScriptDebuggerInfo, ESPMode::ThreadSafe>>* Array = FoundEntry->Get();
 		OutCaptureResults.SetNum(Array->Num());
 
+		bool bWaitForGPU = false;
+		{
+			for (int32 i = 0; i < FoundEntry->Get()->Num(); i++)
+			{
+				if ((*Array)[i]->bWaitForGPU && (*Array)[i]->bWritten == false)
+				{
+					bWaitForGPU = true;
+				}
+			}
+			
+			if (bWaitForGPU)
+			{
+				for (TSharedRef<FNiagaraEmitterInstance> CachedEmitter : Emitters)
+				{
+					CachedEmitter->WaitForDebugInfo();
+				}
+				return false;
+			}
+		}
+
+
 		for (int32 i = 0; i < FoundEntry->Get()->Num(); i++)
 		{
 			OutCaptureResults[i] = (*Array)[i];
@@ -206,7 +234,7 @@ TArray<TSharedPtr<struct FNiagaraScriptDebuggerInfo, ESPMode::ThreadSafe>>* FNia
 	return CurrentCapture.Get();
 }
 
-FNiagaraScriptDebuggerInfo* FNiagaraSystemInstance::GetActiveCaptureWrite(const FName& InHandleName, ENiagaraScriptUsage InUsage, const FGuid& InUsageId)
+TSharedPtr<struct FNiagaraScriptDebuggerInfo, ESPMode::ThreadSafe> FNiagaraSystemInstance::GetActiveCaptureWrite(const FName& InHandleName, ENiagaraScriptUsage InUsage, const FGuid& InUsageId)
 {
 	if (CurrentCapture.IsValid())
 	{
@@ -217,7 +245,7 @@ FNiagaraScriptDebuggerInfo* FNiagaraSystemInstance::GetActiveCaptureWrite(const
 
 		if (FoundEntry != nullptr)
 		{
-			return FoundEntry->Get();
+			return *FoundEntry;
 		}
 	}
 	return nullptr;
@@ -241,7 +269,7 @@ void FNiagaraSystemInstance::SetSolo(bool bInSolo)
 	if (bInSolo)
 	{
 		TSharedPtr<FNiagaraSystemSimulation, ESPMode::ThreadSafe> NewSoloSim = MakeShared<FNiagaraSystemSimulation, ESPMode::ThreadSafe>();
-		NewSoloSim->Init(System, Component->GetWorld(), this);
+		NewSoloSim->Init(System, Component->GetWorld(), true);
 
 		NewSoloSim->TransferInstance(SystemSimulation.Get(), this);	
 
@@ -344,14 +372,16 @@ void FNiagaraSystemInstance::SetPaused(bool bInPaused)
 	}
 	
 	FNiagaraSystemSimulation* SystemSim = GetSystemSimulation().Get();
-	check(SystemSim);
-	if (bInPaused)
-	{
-		SystemSim->PauseInstance(this);
-	}
-	else
+	if (SystemSim)
 	{
-		SystemSim->UnpauseInstance(this);
+		if (bInPaused)
+		{
+			SystemSim->PauseInstance(this);
+		}
+		else
+		{
+			SystemSim->UnpauseInstance(this);
+		}
 	}
 
 	bPaused = bInPaused;
@@ -361,7 +391,6 @@ void FNiagaraSystemInstance::Reset(FNiagaraSystemInstance::EResetMode Mode, bool
 {
 	SCOPE_CYCLE_COUNTER(STAT_NiagaraSystemReset);
 
-	FNiagaraSystemSimulation* SystemSim = GetSystemSimulation().Get();
 	if (Mode == EResetMode::None)
 	{
 		// Right now we don't support binding with reset mode none.
@@ -376,9 +405,9 @@ void FNiagaraSystemInstance::Reset(FNiagaraSystemInstance::EResetMode Mode, bool
 
 	SetPaused(false);
 
-	if (SystemSim)
+	if (SystemSimulation.IsValid())
 	{
-		SystemSim->RemoveInstance(this);
+		SystemSimulation->RemoveInstance(this);
 	}
 	else
 	{
@@ -405,6 +434,8 @@ void FNiagaraSystemInstance::Reset(FNiagaraSystemInstance::EResetMode Mode, bool
 	{
 		//UE_LOG(LogNiagara, Log, TEXT("FNiagaraSystemInstance::ReInit"));
 		ReInitInternal();
+		// If the system was reinitialized successfully than we need to force a rebind of the parameters.
+		bBindParams = IsComplete() == false;
 	}
 	
 	if (bBindParams)
@@ -412,7 +443,6 @@ void FNiagaraSystemInstance::Reset(FNiagaraSystemInstance::EResetMode Mode, bool
 		BindParameters();
 	}
 
-	SystemSim = GetSystemSimulation().Get();
 	SetRequestedExecutionState(ENiagaraExecutionState::Active);
 	SetActualExecutionState(ENiagaraExecutionState::Active);
 
@@ -422,7 +452,7 @@ void FNiagaraSystemInstance::Reset(FNiagaraSystemInstance::EResetMode Mode, bool
 	if (!IsComplete())
 	{
 		bPendingSpawn = true;
-		SystemSim->AddInstance(this);
+		SystemSimulation->AddInstance(this);
 
 		UNiagaraSystem* System = GetSystem();
 		if (System->NeedsWarmup())
@@ -434,6 +464,7 @@ void FNiagaraSystemInstance::Reset(FNiagaraSystemInstance::EResetMode Mode, bool
 
 			//Reset age to zero.
 			Age = 0.0f;
+			TickCount = 0;
 		}
 	}
 
@@ -445,6 +476,7 @@ void FNiagaraSystemInstance::Reset(FNiagaraSystemInstance::EResetMode Mode, bool
 void FNiagaraSystemInstance::ResetInternal(bool bResetSimulations)
 {
 	Age = 0;
+	TickCount = 0;
 	UNiagaraSystem* System = GetSystem();
 	if (System == nullptr || Component == nullptr || IsDisabled())
 	{
@@ -491,15 +523,15 @@ UNiagaraParameterCollectionInstance* FNiagaraSystemInstance::GetParameterCollect
 	return SystemSimulation->GetParameterCollectionInstance(Collection);
 }
 
-void FNiagaraSystemInstance::AdvanceSimulation(int32 TickCount, float TickDeltaSeconds)
+void FNiagaraSystemInstance::AdvanceSimulation(int32 TickCountToSimulate, float TickDeltaSeconds)
 {
-	if (TickCount > 0)
+	if (TickCountToSimulate > 0)
 	{
 		SCOPE_CYCLE_COUNTER(STAT_NiagaraSystemAdvanceSim);
 		bool bWasSolo = bSolo;
 		SetSolo(true);
 
-		for (int32 TickIdx = 0; TickIdx < TickCount; ++TickIdx)
+		for (int32 TickIdx = 0; TickIdx < TickCountToSimulate; ++TickIdx)
 		{
 			ComponentTick(TickDeltaSeconds);
 		}
@@ -528,10 +560,35 @@ bool FNiagaraSystemInstance::IsReadyToRun() const
 	return bAllReadyToRun;
 }
 
+bool DoSystemDataInterfacesRequireSolo(const UNiagaraSystem& System, const UNiagaraComponent& Component)
+{
+	if (System.HasSystemScriptDIsWithPerInstanceData())
+	{
+		return true;
+	}
+
+	const TArray<FName>& UserDINamesReadInSystemScripts = System.GetUserDINamesReadInSystemScripts();
+	if (UserDINamesReadInSystemScripts.Num() > 0)
+	{
+		TArray<FNiagaraVariable> OverrideParameterVariables;
+		Component.GetOverrideParameters().GetParameters(OverrideParameterVariables);
+		for (const FNiagaraVariable& OverrideParameterVariable : OverrideParameterVariables)
+		{
+			if (OverrideParameterVariable.IsDataInterface() && UserDINamesReadInSystemScripts.Contains(OverrideParameterVariable.GetName()))
+			{
+				return true;
+			}
+		}
+	}
+
+	return false;
+}
+
 void FNiagaraSystemInstance::ReInitInternal()
 {
 	SCOPE_CYCLE_COUNTER(STAT_NiagaraSystemReinit);
 	Age = 0;
+	TickCount = 0;
 	UNiagaraSystem* System = GetSystem();
 	if (System == nullptr || Component == nullptr)
 	{
@@ -557,13 +614,13 @@ void FNiagaraSystemInstance::ReInitInternal()
 	}
 
 	/** Do we need to run in solo mode? */
-	bSolo = bForceSolo || System->IsSolo();
+	bSolo = bForceSolo || DoSystemDataInterfacesRequireSolo(*System, *Component);
 	if (bSolo)
 	{
 		if (!SystemSimulation.IsValid())
 		{
 			SystemSimulation = MakeShared<FNiagaraSystemSimulation, ESPMode::ThreadSafe>();
-			SystemSimulation->Init(System, Component->GetWorld(), this);
+			SystemSimulation->Init(System, Component->GetWorld(), true);
 		}
 	}
 	else
@@ -577,6 +634,7 @@ void FNiagaraSystemInstance::ReInitInternal()
 	
 	InstanceParameters.Reset();
 	InstanceParameters.AddParameter(SYS_PARAM_ENGINE_POSITION, true, false);
+	InstanceParameters.AddParameter(SYS_PARAM_ENGINE_ROTATION, true, false);
 	InstanceParameters.AddParameter(SYS_PARAM_ENGINE_SCALE, true, false);
 	InstanceParameters.AddParameter(SYS_PARAM_ENGINE_VELOCITY, true, false);
 	InstanceParameters.AddParameter(SYS_PARAM_ENGINE_X_AXIS, true, false);
@@ -598,6 +656,7 @@ void FNiagaraSystemInstance::ReInitInternal()
 	InstanceParameters.AddParameter(SYS_PARAM_ENGINE_SYSTEM_NUM_EMITTERS, true, false);
 	InstanceParameters.AddParameter(SYS_PARAM_ENGINE_SYSTEM_NUM_EMITTERS_ALIVE, true, false);
 	InstanceParameters.AddParameter(SYS_PARAM_ENGINE_SYSTEM_AGE);
+	InstanceParameters.AddParameter(SYS_PARAM_ENGINE_SYSTEM_TICK_COUNT);
 
 	// This is required for user default data interface's (like say static meshes) to be set up properly.
 	// Additionally, it must happen here for data to be properly found below.
@@ -605,15 +664,26 @@ void FNiagaraSystemInstance::ReInitInternal()
 	System->GetExposedParameters().CopyParametersTo(InstanceParameters, bOnlyAdd, FNiagaraParameterStore::EDataInterfaceCopyMethod::Reference);
 
 	TArray<FNiagaraVariable> NumParticleVars;
+	TArray<FNiagaraVariable> TotalSpawnedParticlesVars;
 	for (int32 i = 0; i < Emitters.Num(); i++)
 	{
 		TSharedRef<FNiagaraEmitterInstance> Simulation = Emitters[i];
 		FString EmitterName = Simulation->GetEmitterHandle().GetInstance()->GetUniqueEmitterName();
-		FNiagaraVariable Var = SYS_PARAM_ENGINE_EMITTER_NUM_PARTICLES;
-		FString ParamName = Var.GetName().ToString().Replace(TEXT("Emitter"), *EmitterName);
-		Var.SetName(*ParamName);
-		InstanceParameters.AddParameter(Var, true, false);
-		NumParticleVars.Add(Var);
+		
+		{
+			FNiagaraVariable Var = SYS_PARAM_ENGINE_EMITTER_NUM_PARTICLES;
+			FString ParamName = Var.GetName().ToString().Replace(TEXT("Emitter"), *EmitterName);
+			Var.SetName(*ParamName);
+			InstanceParameters.AddParameter(Var, true, false);
+			NumParticleVars.Add(Var);
+		}
+		{
+			FNiagaraVariable Var = SYS_PARAM_ENGINE_EMITTER_TOTAL_SPAWNED_PARTICLES;
+			FString ParamName = Var.GetName().ToString().Replace(TEXT("Emitter"), *EmitterName);
+			Var.SetName(*ParamName);
+			InstanceParameters.AddParameter(Var, true, false);
+			TotalSpawnedParticlesVars.Add(Var);
+		}
 	}
 
 	// Make sure all parameters are added before initializing the bindings, otherwise parameter store layout changes might invalidate the bindings.
@@ -624,6 +694,8 @@ void FNiagaraSystemInstance::ReInitInternal()
 	OwnerYAxisParam.Init(InstanceParameters, SYS_PARAM_ENGINE_Y_AXIS);
 	OwnerZAxisParam.Init(InstanceParameters, SYS_PARAM_ENGINE_Z_AXIS);
 
+	OwnerRotationParam.Init(InstanceParameters, SYS_PARAM_ENGINE_ROTATION);
+
 	OwnerTransformParam.Init(InstanceParameters, SYS_PARAM_ENGINE_LOCAL_TO_WORLD);
 	OwnerInverseParam.Init(InstanceParameters, SYS_PARAM_ENGINE_WORLD_TO_LOCAL);
 	OwnerTransposeParam.Init(InstanceParameters, SYS_PARAM_ENGINE_LOCAL_TO_WORLD_TRANSPOSED);
@@ -635,6 +707,7 @@ void FNiagaraSystemInstance::ReInitInternal()
 	OwnerInverseDeltaSecondsParam.Init(InstanceParameters, SYS_PARAM_ENGINE_INV_DELTA_TIME);
 
 	SystemAgeParam.Init(InstanceParameters, SYS_PARAM_ENGINE_SYSTEM_AGE);
+	SystemTickCountParam.Init(InstanceParameters, SYS_PARAM_ENGINE_SYSTEM_TICK_COUNT);
 	OwnerEngineTimeParam.Init(InstanceParameters, SYS_PARAM_ENGINE_TIME);
 	OwnerEngineRealtimeParam.Init(InstanceParameters, SYS_PARAM_ENGINE_REAL_TIME);
 
@@ -652,6 +725,12 @@ void FNiagaraSystemInstance::ReInitInternal()
 		ParameterNumParticleBindings[i].Init(InstanceParameters, NumParticleVars[i]);
 	}
 
+	ParameterTotalSpawnedParticlesBindings.SetNum(TotalSpawnedParticlesVars.Num());
+	for (int32 i = 0; i < TotalSpawnedParticlesVars.Num(); i++)
+	{
+		ParameterTotalSpawnedParticlesBindings[i].Init(InstanceParameters, TotalSpawnedParticlesVars[i]);
+	}
+
 	// rebind now after all parameters have been added
 	InstanceParameters.Rebind();
 
@@ -745,6 +824,21 @@ void FNiagaraSystemInstance::BindParameters()
 {
 	Component->GetOverrideParameters().Bind(&InstanceParameters);
 
+	if (SystemSimulation->GetIsSolo())
+	{
+		// If this simulation is solo than we can bind the instance parameters to the system simulation contexts so that
+		// the system and emitter scripts use the per-instance data interfaces.
+		Component->GetOverrideParameters().Bind(&SystemSimulation->GetSpawnExecutionContext().Parameters);
+		Component->GetOverrideParameters().Bind(&SystemSimulation->GetUpdateExecutionContext().Parameters);
+	}
+	else 
+	{
+		// If this simulation is not solo than we have bind the source system parameters to the system simulation contexts so that
+		// the system and emitter scripts use the default shared data interfaces.
+		GetSystem()->GetExposedParameters().Bind(&SystemSimulation->GetSpawnExecutionContext().Parameters);
+		GetSystem()->GetExposedParameters().Bind(&SystemSimulation->GetSpawnExecutionContext().Parameters);
+	}
+
 	for (TSharedRef<FNiagaraEmitterInstance> Simulation : Emitters)
 	{
 		Simulation->BindParameters();
@@ -755,6 +849,24 @@ void FNiagaraSystemInstance::UnbindParameters()
 {
 	Component->GetOverrideParameters().Unbind(&InstanceParameters);
 
+	if (SystemSimulation.IsValid())
+	{
+		if (SystemSimulation->GetIsSolo())
+		{
+			Component->GetOverrideParameters().Unbind(&SystemSimulation->GetSpawnExecutionContext().Parameters);
+			Component->GetOverrideParameters().Unbind(&SystemSimulation->GetUpdateExecutionContext().Parameters);
+		}
+		else 
+		{
+			UNiagaraSystem* System = GetSystem();
+			if (System)
+			{
+				System->GetExposedParameters().Unbind(&SystemSimulation->GetSpawnExecutionContext().Parameters);
+				System->GetExposedParameters().Unbind(&SystemSimulation->GetSpawnExecutionContext().Parameters);
+			}
+		}
+	}
+
 	for (TSharedRef<FNiagaraEmitterInstance> Simulation : Emitters)
 	{
 		Simulation->UnbindParameters();
@@ -813,6 +925,15 @@ void FNiagaraSystemInstance::InitDataInterfaces()
 
 	CalcInstDataSize(InstanceParameters.GetDataInterfaces());//This probably should be a proper exec context. 
 
+	if (SystemSimulation->GetIsSolo())
+	{
+		CalcInstDataSize(SystemSimulation->GetSpawnExecutionContext().GetDataInterfaces());
+		SystemSimulation->GetSpawnExecutionContext().DirtyDataInterfaces();
+
+		CalcInstDataSize(SystemSimulation->GetUpdateExecutionContext().GetDataInterfaces());
+		SystemSimulation->GetUpdateExecutionContext().DirtyDataInterfaces();
+	}
+
 	//Iterate over interfaces to get size for table and clear their interface bindings.
 	for (TSharedRef<FNiagaraEmitterInstance> Simulation : Emitters)
 	{
@@ -860,6 +981,15 @@ void FNiagaraSystemInstance::InitDataInterfaces()
 	}
 }
 
+bool FNiagaraSystemInstance::GetPerInstanceDataAndOffsets(void*& OutData, uint32& OutDataSize, TMap<TWeakObjectPtr<UNiagaraDataInterface>, int32>*& OutOffsets)
+{
+	OutData = DataInterfaceInstanceData.GetData();
+	OutDataSize = DataInterfaceInstanceData.Num();
+	OutOffsets = &DataInterfaceInstanceDataOffsets;
+	return DataInterfaceInstanceDataOffsets.Num() != 0;
+}
+
+
 void FNiagaraSystemInstance::TickDataInterfaces(float DeltaSeconds, bool bPostSimulate)
 {
 	if (!GetSystem() || !Component || IsDisabled())
@@ -902,8 +1032,8 @@ void FNiagaraSystemInstance::TickInstanceParameters(float DeltaSeconds)
 	//TODO: Create helper binding objects to avoid the search in set parameter value.
 	//Set System params.
 	FTransform ComponentTrans = Component->GetComponentTransform();
-	FVector OldPos = OwnerPositionParam.GetValue();// ComponentTrans.GetLocation();
 	FVector CurrPos = ComponentTrans.GetLocation();
+	FVector OldPos = FMath::IsNearlyZero(Age) ? CurrPos : OwnerPositionParam.GetValue(); // The first frame the value in OwnerPositionParam is uninitialized memory, we need to make sure that we don't use it.
 	OwnerPositionParam.SetValue(CurrPos);
 	OwnerScaleParam.SetValue(ComponentTrans.GetScale3D());
 	OwnerVelocityParam.SetValue((CurrPos - OldPos) / DeltaSeconds);
@@ -911,6 +1041,8 @@ void FNiagaraSystemInstance::TickInstanceParameters(float DeltaSeconds)
 	OwnerYAxisParam.SetValue(ComponentTrans.GetRotation().GetAxisY());
 	OwnerZAxisParam.SetValue(ComponentTrans.GetRotation().GetAxisZ());
 
+	OwnerRotationParam.SetValue(ComponentTrans.GetRotation());
+
 	FMatrix Transform = ComponentTrans.ToMatrixWithScale();
 	FMatrix Inverse = Transform.Inverse();
 	FMatrix Transpose = Transform.GetTransposed();
@@ -971,6 +1103,7 @@ void FNiagaraSystemInstance::TickInstanceParameters(float DeltaSeconds)
 		OwnerEngineRealtimeParam.SetValue(Age);
 	}
 	SystemAgeParam.SetValue(Age);
+	SystemTickCountParam.SetValue(TickCount);
 
 	int32 NumAlive = 0;
 	for (int32 i = 0; i < Emitters.Num(); i++)
@@ -981,6 +1114,7 @@ void FNiagaraSystemInstance::TickInstanceParameters(float DeltaSeconds)
 			NumAlive++;
 		}
 		ParameterNumParticleBindings[i].SetValue(NumParticles);
+		ParameterTotalSpawnedParticlesBindings[i].SetValue(Emitters[i]->GetTotalSpawnedParticles());
 	}
 	SystemNumEmittersParam.SetValue(Emitters.Num());
 	SystemNumEmittersAliveParam.SetValue(NumAlive);
@@ -1171,6 +1305,9 @@ void FNiagaraSystemInstance::PreSimulateTick(float DeltaSeconds)
 
 void FNiagaraSystemInstance::PostSimulateTick(float DeltaSeconds)
 {
+	SCOPE_CYCLE_COUNTER(STAT_NiagaraSystemInstanceTick);
+	SCOPE_CYCLE_COUNTER(STAT_NiagaraOverview_GT_CNC);
+
 	if (IsComplete() || !bHasTickingEmitters || GetSystem() == nullptr || Component == nullptr || DeltaSeconds < SMALL_NUMBER)
 	{
 		return;
@@ -1197,6 +1334,7 @@ void FNiagaraSystemInstance::PostSimulateTick(float DeltaSeconds)
 	}
 
 	Age += DeltaSeconds;
+	TickCount += 1;
 }
 
 #if WITH_EDITORONLY_DATA
diff --git a/Engine/Plugins/FX/Niagara/Source/Niagara/Private/NiagaraSystemSimulation.cpp b/Engine/Plugins/FX/Niagara/Source/Niagara/Private/NiagaraSystemSimulation.cpp
--- a/Engine/Plugins/FX/Niagara/Source/Niagara/Private/NiagaraSystemSimulation.cpp
+++ b/Engine/Plugins/FX/Niagara/Source/Niagara/Private/NiagaraSystemSimulation.cpp
@@ -14,14 +14,14 @@
 #include "NiagaraComponent.h"
 #include "NiagaraWorldManager.h"
 
-DECLARE_CYCLE_STAT(TEXT("System Simulation"), STAT_NiagaraSystemSim, STATGROUP_Niagara);
-DECLARE_CYCLE_STAT(TEXT("System Pre Simulate"), STAT_NiagaraSystemSim_PreSimulate, STATGROUP_Niagara);
-DECLARE_CYCLE_STAT(TEXT("System Prepare For Simulate"), STAT_NiagaraSystemSim_PrepareForSimulate, STATGROUP_Niagara);
-DECLARE_CYCLE_STAT(TEXT("System Update"), STAT_NiagaraSystemSim_Update, STATGROUP_Niagara);
-DECLARE_CYCLE_STAT(TEXT("System Spawn"), STAT_NiagaraSystemSim_Spawn, STATGROUP_Niagara);
-DECLARE_CYCLE_STAT(TEXT("System Transfer Parameters"), STAT_NiagaraSystemSim_TransferParameters, STATGROUP_Niagara);
-DECLARE_CYCLE_STAT(TEXT("System Post Simulate"), STAT_NiagaraSystemSim_PostSimulate, STATGROUP_Niagara);
-DECLARE_CYCLE_STAT(TEXT("System Mark Component Dirty"), STAT_NiagaraSystemSim_MarkComponentDirty, STATGROUP_Niagara);
+DECLARE_CYCLE_STAT(TEXT("System Simulation [GT]"), STAT_NiagaraSystemSim, STATGROUP_Niagara);
+DECLARE_CYCLE_STAT(TEXT("System Pre Simulate [GT]"), STAT_NiagaraSystemSim_PreSimulate, STATGROUP_Niagara);
+DECLARE_CYCLE_STAT(TEXT("System Prepare For Simulate [GT]"), STAT_NiagaraSystemSim_PrepareForSimulate, STATGROUP_Niagara);
+DECLARE_CYCLE_STAT(TEXT("System Sim Update [GT]"), STAT_NiagaraSystemSim_Update, STATGROUP_Niagara);
+DECLARE_CYCLE_STAT(TEXT("System Sim Spawn [GT]"), STAT_NiagaraSystemSim_Spawn, STATGROUP_Niagara);
+DECLARE_CYCLE_STAT(TEXT("System Sim Transfer Parameters [GT]"), STAT_NiagaraSystemSim_TransferParameters, STATGROUP_Niagara);
+DECLARE_CYCLE_STAT(TEXT("System Post Simulate [GT]"), STAT_NiagaraSystemSim_PostSimulate, STATGROUP_Niagara);
+DECLARE_CYCLE_STAT(TEXT("System Mark Component Dirty [GT]"), STAT_NiagaraSystemSim_MarkComponentDirty, STATGROUP_Niagara);
 
 
 static int32 GbDumpSystemData = 0;
@@ -73,14 +73,15 @@ FNiagaraSystemSimulation::~FNiagaraSystemSimulation()
 	Destroy();
 }
 
-bool FNiagaraSystemSimulation::Init(UNiagaraSystem* InSystem, UWorld* InWorld, FNiagaraSystemInstance* InSoloSystemInstance)
+bool FNiagaraSystemSimulation::Init(UNiagaraSystem* InSystem, UWorld* InWorld, bool bInIsSolo)
 {
 	UNiagaraSystem* System = InSystem;
 	WeakSystem = System;
-	SoloSystemInstance = InSoloSystemInstance;
 
 	World = InWorld;
 
+	bIsSolo = bInIsSolo;
+
 	FNiagaraWorldManager* WorldMan = FNiagaraWorldManager::Get(InWorld);
 	check(WorldMan);
 
@@ -162,6 +163,11 @@ bool FNiagaraSystemSimulation::Init(UNiagaraSystem* InSystem, UWorld* InWorld, F
 			{
 				EmitterSpawnInfoAccessors[EmitterIdx].Emplace(DataSet, FNiagaraVariable(FNiagaraTypeDefinition(FNiagaraSpawnInfo::StaticStruct()), AttrName));
 			}
+
+			if (Emitter->bLimitDeltaTime)
+			{
+				MaxDeltaTime = MaxDeltaTime.IsSet() ? FMath::Min(MaxDeltaTime.GetValue(), Emitter->MaxDeltaTimePerTick) : Emitter->MaxDeltaTimePerTick;
+			}
 		}
 
 		SpawnDeltaTimeParam.Init(SpawnExecContext.Parameters, SYS_PARAM_ENGINE_DELTA_TIME);
@@ -240,6 +246,10 @@ void FNiagaraSystemSimulation::TransferInstance(FNiagaraSystemSimulation* Source
 	
 		//Move the system direct to the new sim's 
 		SystemInst->SystemInstanceIndex = SystemInstances.Add(SystemInst);
+		if (SystemInst->SystemInstanceIndex == 0)
+		{
+			InitParameterDataSetBindings(SystemInst);
+		}
 
 		check(NewDataIndex == SystemInst->SystemInstanceIndex);
 	}
@@ -274,6 +284,11 @@ bool FNiagaraSystemSimulation::Tick(float DeltaSeconds)
 		return false;
 	}
 
+	if (MaxDeltaTime.IsSet())
+	{
+		DeltaSeconds = FMath::Clamp(DeltaSeconds, 0.0f, MaxDeltaTime.GetValue());
+	}
+
 	UNiagaraScript* SystemSpawnScript = System->GetSystemSpawnScript();
 	UNiagaraScript* SystemUpdateScript = System->GetSystemUpdateScript();
 #if WITH_EDITOR
@@ -341,6 +356,11 @@ bool FNiagaraSystemSimulation::Tick(float DeltaSeconds)
 				if (!Inst->IsComplete())
 				{
 					Inst->SystemInstanceIndex = SystemInstances.Add(Inst);
+					if (Inst->SystemInstanceIndex == 0)
+					{
+						// When the first instance is added we need to initialize the parameter store to data set bindings.
+						InitParameterDataSetBindings(Inst);
+					}
 					++SpawnNum;
 				}
 				else
@@ -392,8 +412,6 @@ bool FNiagaraSystemSimulation::Tick(float DeltaSeconds)
 				}
 			};
 
-			InitBindings(SystemInstances[0]);
-
 			SpawnInstanceParameterDataSet.Allocate(NewNum);
 			UpdateInstanceParameterDataSet.Allocate(NewNum);
 
@@ -429,6 +447,8 @@ bool FNiagaraSystemSimulation::Tick(float DeltaSeconds)
 			UpdateGlobalSystemCountScaleParam.SetValue(GlobalSystemCountScale);
 		}
 
+		FNiagaraSystemInstance* SoloSystemInstance = bIsSolo && SystemInstances.Num() == 1 ? SystemInstances[0] : nullptr;
+
 		//TODO: JIRA - UE-60096 - Remove.
 		//We're having to allocate and spawn before update here so we have to do needless copies.			
 		//Ideally this should be compiled directly into the script similarly to interpolated particle spawning.
@@ -456,12 +476,13 @@ bool FNiagaraSystemSimulation::Tick(float DeltaSeconds)
 #if WITH_EDITORONLY_DATA
 			if (SoloSystemInstance && SoloSystemInstance->ShouldCaptureThisFrame())
 			{
-				FNiagaraScriptDebuggerInfo* DebugInfo = SoloSystemInstance->GetActiveCaptureWrite(NAME_None, ENiagaraScriptUsage::SystemSpawnScript, FGuid());
+				TSharedPtr<struct FNiagaraScriptDebuggerInfo, ESPMode::ThreadSafe> DebugInfo = SoloSystemInstance->GetActiveCaptureWrite(NAME_None, ENiagaraScriptUsage::SystemSpawnScript, FGuid());
 				if (DebugInfo)
 				{
 					DataSet.Dump(DebugInfo->Frame, true, OrigNum, SpawnNum);
 					//DebugInfo->Frame.Dump(true, 0, SpawnNum);
 					DebugInfo->Parameters = SpawnExecContext.Parameters;
+					DebugInfo->bWritten = true;
 				}
 			}
 #endif
@@ -524,12 +545,13 @@ bool FNiagaraSystemSimulation::Tick(float DeltaSeconds)
 #if WITH_EDITORONLY_DATA
 			if (SoloSystemInstance && SoloSystemInstance->ShouldCaptureThisFrame())
 			{
-				FNiagaraScriptDebuggerInfo* DebugInfo = SoloSystemInstance->GetActiveCaptureWrite(NAME_None, ENiagaraScriptUsage::SystemUpdateScript, FGuid());
+				TSharedPtr<struct FNiagaraScriptDebuggerInfo, ESPMode::ThreadSafe> DebugInfo = SoloSystemInstance->GetActiveCaptureWrite(NAME_None, ENiagaraScriptUsage::SystemUpdateScript, FGuid());
 				if (DebugInfo)
 				{
 					DataSet.Dump(DebugInfo->Frame, true, 0, NewNum);
 					//DebugInfo->Frame.Dump(true, 0, OrigNum);
 					DebugInfo->Parameters = UpdateExecContext.Parameters;
+					DebugInfo->bWritten = true;
 				}
 			}
 #endif
@@ -620,7 +642,7 @@ bool FNiagaraSystemSimulation::Tick(float DeltaSeconds)
 	{
 		SCOPE_CYCLE_COUNTER(STAT_NiagaraSystemSim_PostSimulate);
 
-		if (GbParallelSystemPostTick)
+		if (GbParallelSystemPostTick && FApp::ShouldUseThreadingForPerformance())
 		{
 			ParallelFor(SystemInstances.Num(),
 				[&](int32 SystemIndex)
@@ -659,9 +681,9 @@ bool FNiagaraSystemSimulation::Tick(float DeltaSeconds)
 	}
 
 #if WITH_EDITORONLY_DATA
-	if (SoloSystemInstance != nullptr)
+	if (bIsSolo && SystemInstances.Num() == 1)
 	{
-		SoloSystemInstance->FinishCapture();
+		SystemInstances[0]->FinishCapture();
 	}
 #endif
 
@@ -837,20 +859,18 @@ void FNiagaraSystemSimulation::UnpauseInstance(FNiagaraSystemInstance* Instance)
 	}
 }
 
-void FNiagaraSystemSimulation::InitBindings(FNiagaraSystemInstance* SystemInst)
+void FNiagaraSystemSimulation::InitParameterDataSetBindings(FNiagaraSystemInstance* SystemInst)
 {
 	//Have to init here as we need an actual parameter store to pull the layout info from.
 	//TODO: Pull the layout stuff out of each data set and store. So much duplicated data.
 	//This assumes that all layouts for all emitters is the same. Which it should be.
 	//Ideally we can store all this layout info in the systm/emitter assets so we can just generate this in Init()
-	if (DataSetToEmitterSpawnParameters.Num() == 0 && SystemInst != nullptr)
+	if (SystemInst != nullptr)
 	{
 		SpawnInstanceParameterToDataSetBinding.Init(SpawnInstanceParameterDataSet, SystemInst->GetInstanceParameters());
 		UpdateInstanceParameterToDataSetBinding.Init(UpdateInstanceParameterDataSet, SystemInst->GetInstanceParameters());
 
 		TArray<TSharedRef<FNiagaraEmitterInstance>>& Emitters = SystemInst->GetEmitters();
-		check(DataSetToEmitterUpdateParameters.Num() == 0);
-		check(DataSetToEmitterEventParameters.Num() == 0);
 		DataSetToEmitterSpawnParameters.SetNum(Emitters.Num());
 		DataSetToEmitterUpdateParameters.SetNum(Emitters.Num());
 		DataSetToEmitterEventParameters.SetNum(Emitters.Num());
@@ -872,13 +892,4 @@ void FNiagaraSystemSimulation::InitBindings(FNiagaraSystemInstance* SystemInst)
 			}
 		}
 	}
-
-	// If we have data interfaces (say user ones) that need to be pushed to the 
-	// system spawn and update scripts, it needs to happen here.
-	if (SystemInst)
-	{
-		SystemInst->GetInstanceParameters().Bind(&SpawnExecContext.Parameters);
-		SystemInst->GetInstanceParameters().Bind(&UpdateExecContext.Parameters);
-
-	}
 }
diff --git a/Engine/Plugins/FX/Niagara/Source/Niagara/Private/NiagaraUserRedirectionParameterStore.cpp b/Engine/Plugins/FX/Niagara/Source/Niagara/Private/NiagaraUserRedirectionParameterStore.cpp
new file mode 100644
--- /dev/null
+++ b/Engine/Plugins/FX/Niagara/Source/Niagara/Private/NiagaraUserRedirectionParameterStore.cpp
@@ -0,0 +1,113 @@
+// Copyright 1998-2019 Epic Games, Inc. All Rights Reserved.
+
+#include "NiagaraUserRedirectionParameterStore.h"
+#include "NiagaraStats.h"
+#include "NiagaraEmitterInstanceBatcher.h"
+#include "NiagaraDataInterface.h"
+#include "NiagaraSystemInstance.h"
+
+FNiagaraUserRedirectionParameterStore::FNiagaraUserRedirectionParameterStore() : FNiagaraParameterStore()
+{
+
+}
+
+FNiagaraUserRedirectionParameterStore::FNiagaraUserRedirectionParameterStore(const FNiagaraParameterStore& Other)
+{
+	*this = Other;
+}
+
+FNiagaraUserRedirectionParameterStore& FNiagaraUserRedirectionParameterStore::operator=(const FNiagaraParameterStore& Other)
+{
+	Super::operator=(Other);	
+	RecreateRedirections();
+	return *this;
+}
+
+bool FNiagaraUserRedirectionParameterStore::IsUserParameter(const FNiagaraVariable& InVar) const
+{
+	return InVar.GetName().ToString().StartsWith(TEXT("User."));
+}
+
+FNiagaraVariable FNiagaraUserRedirectionParameterStore::GetUserRedirection(const FNiagaraVariable & InVar) const
+{
+	if (!IsUserParameter(InVar))
+	{
+		return InVar;
+	}
+	FNiagaraVariable SimpleVar = InVar;
+	FName DisplayName(*InVar.GetName().ToString().RightChop(5));
+	SimpleVar.SetName(DisplayName);
+	return SimpleVar;
+}
+
+void FNiagaraUserRedirectionParameterStore::RecreateRedirections()
+{
+	UserParameterRedirects.Reset();
+	for (const TPair<FNiagaraVariable, int32>& ParamOffset : GetParameterOffsets())
+	{
+		const FNiagaraVariable& Var = ParamOffset.Key;
+		if (IsUserParameter(Var))
+		{
+			UserParameterRedirects.Add(GetUserRedirection(Var), Var);
+		}
+	}
+}
+
+int32 FNiagaraUserRedirectionParameterStore::IndexOf(const FNiagaraVariable& Parameter) const
+{
+	const FNiagaraVariable* Redirection = UserParameterRedirects.Find(Parameter);
+	return Super::IndexOf(Redirection ? *Redirection : Parameter);
+}
+
+bool FNiagaraUserRedirectionParameterStore::AddParameter(const FNiagaraVariable& Param, bool bInitialize /*= true*/, bool bTriggerRebind /*= true*/)
+{
+	if (IsUserParameter(Param))
+	{
+		UserParameterRedirects.Add(GetUserRedirection(Param), Param);
+	}
+	return Super::AddParameter(Param, bInitialize, bTriggerRebind);
+}
+
+bool FNiagaraUserRedirectionParameterStore::RemoveParameter(const FNiagaraVariable& InVar)
+{
+	const FNiagaraVariable* Redirection = UserParameterRedirects.Find(InVar);
+	const FNiagaraVariable& ToRemove = Redirection ? *Redirection : InVar;
+	bool Result = Super::RemoveParameter(ToRemove);
+	if (Result)
+	{
+		UserParameterRedirects.Remove(GetUserRedirection(ToRemove));
+	}
+	return Result;
+}
+
+void FNiagaraUserRedirectionParameterStore::InitFromSource(const FNiagaraParameterStore* SrcStore, bool bNotifyAsDirty)
+{
+	Super::InitFromSource(SrcStore, bNotifyAsDirty);
+	RecreateRedirections();
+}
+
+void FNiagaraUserRedirectionParameterStore::Empty(bool bClearBindings /*= true*/)
+{
+	Super::Empty(bClearBindings);
+	UserParameterRedirects.Empty();
+}
+
+void FNiagaraUserRedirectionParameterStore::Reset(bool bClearBindings /*= true*/)
+{
+	Super::Reset(bClearBindings);
+	UserParameterRedirects.Reset();
+}
+
+bool FNiagaraUserRedirectionParameterStore::SerializeFromMismatchedTag(const FPropertyTag & Tag, FStructuredArchive::FSlot Slot)
+{
+	static FName StoreDataName("NiagaraParameterStore");
+	if (Tag.Type == NAME_StructProperty && Tag.StructName == StoreDataName)
+	{
+		FNiagaraParameterStore OldStore;
+		FNiagaraParameterStore::StaticStruct()->SerializeItem(Slot, &OldStore, nullptr);
+		*this = OldStore;
+		return true;
+	}
+
+	return false;
+}
diff --git a/Engine/Plugins/FX/Niagara/Source/Niagara/Private/NiagaraWorldManager.cpp b/Engine/Plugins/FX/Niagara/Source/Niagara/Private/NiagaraWorldManager.cpp
--- a/Engine/Plugins/FX/Niagara/Source/Niagara/Private/NiagaraWorldManager.cpp
+++ b/Engine/Plugins/FX/Niagara/Source/Niagara/Private/NiagaraWorldManager.cpp
@@ -12,6 +12,9 @@
 #include "Misc/ConfigCacheIni.h"
 #include "NiagaraDataInterfaceSkeletalMesh.h"
 #include "EngineModule.h"
+#include "NiagaraStats.h"
+
+DECLARE_CYCLE_STAT(TEXT("Niagara Manager Tick [GT]"), STAT_NiagaraWorldManTick, STATGROUP_Niagara);
 
 TGlobalResource<FNiagaraViewDataMgr> GNiagaraViewDataManager;
 
@@ -27,8 +30,11 @@ FNiagaraViewDataMgr::FNiagaraViewDataMgr()
 void FNiagaraViewDataMgr::Init()
 {
 	IRendererModule& RendererModule = GetRendererModule();
+
 	GNiagaraViewDataManager.PostOpaqueDelegate.BindRaw(&GNiagaraViewDataManager, &FNiagaraViewDataMgr::PostOpaqueRender);
 	RendererModule.RegisterPostOpaqueRenderDelegate(GNiagaraViewDataManager.PostOpaqueDelegate);
+
+	RendererModule.OnPreSceneRender().AddRaw(&GNiagaraViewDataManager, &FNiagaraViewDataMgr::OnPreSceneRenderCalled);
 }
 
 void FNiagaraViewDataMgr::Shutdown()
@@ -149,7 +155,7 @@ TSharedRef<FNiagaraSystemSimulation, ESPMode::ThreadSafe> FNiagaraWorldManager::
 	
 	TSharedRef<FNiagaraSystemSimulation, ESPMode::ThreadSafe> Sim = MakeShared<FNiagaraSystemSimulation, ESPMode::ThreadSafe>();
 	SystemSimulations.Add(System, Sim);
-	Sim->Init(System, World);
+	Sim->Init(System, World, false);
 	return Sim;
 }
 
@@ -199,6 +205,9 @@ Going off this idea tbh
 		}
 	}
 */
+	SCOPE_CYCLE_COUNTER(STAT_NiagaraWorldManTick);
+	SCOPE_CYCLE_COUNTER(STAT_NiagaraOverview_GT);
+
 	SkeletalMeshGeneratedData.TickGeneratedData(DeltaSeconds);
 
 	//Tick our collections to push any changes to bound stores.
diff --git a/Engine/Plugins/FX/Niagara/Source/Niagara/Public/NiagaraCommon.h b/Engine/Plugins/FX/Niagara/Source/Niagara/Public/NiagaraCommon.h
--- a/Engine/Plugins/FX/Niagara/Source/Niagara/Public/NiagaraCommon.h
+++ b/Engine/Plugins/FX/Niagara/Source/Niagara/Public/NiagaraCommon.h
@@ -342,9 +342,6 @@ public:
 	/** Would this data interface work on the target execution type? Only call this on the game thread.*/
 	bool CanExecuteOnTarget(ENiagaraSimTarget SimTarget) const;
 
-	/** Would this data interface force a system script to be solo? Only call this on the game thread.*/
-	bool IsSystemSolo() const;
-
 	/** Note that this is the CDO for this type of data interface, as we often cannot guarantee that the same instance of the data interface we compiled with is the one the user ultimately executes.  Only call this on the game thread.*/
 	UNiagaraDataInterface* GetDefaultDataInterface() const;
 };
@@ -511,6 +508,21 @@ struct FNiagaraVariableAttributeBinding
 	FNiagaraVariable DefaultValueIfNonExistent;
 };
 
+USTRUCT()
+struct FNiagaraVariableDataInterfaceBinding
+{
+	GENERATED_USTRUCT_BODY();
+
+	FNiagaraVariableDataInterfaceBinding() {}
+	FNiagaraVariableDataInterfaceBinding(const FNiagaraVariable& InVar) : BoundVariable(InVar)
+	{
+		check(InVar.IsDataInterface() == true);
+	}
+
+	UPROPERTY()
+	FNiagaraVariable BoundVariable;
+};
+
 
 namespace FNiagaraUtilities
 {
diff --git a/Engine/Plugins/FX/Niagara/Source/Niagara/Public/NiagaraComponent.h b/Engine/Plugins/FX/Niagara/Source/Niagara/Public/NiagaraComponent.h
--- a/Engine/Plugins/FX/Niagara/Source/Niagara/Public/NiagaraComponent.h
+++ b/Engine/Plugins/FX/Niagara/Source/Niagara/Public/NiagaraComponent.h
@@ -8,7 +8,7 @@
 #include "PrimitiveViewRelevance.h"
 #include "PrimitiveSceneProxy.h"
 #include "Components/PrimitiveComponent.h"
-#include "NiagaraParameterStore.h"
+#include "NiagaraUserRedirectionParameterStore.h"
 #include "NiagaraSystemInstance.h"
 
 #include "NiagaraComponent.generated.h"
@@ -21,14 +21,15 @@ class UNiagaraParameterCollectionInstance;
 class FNiagaraSystemSimulation;
 
 // Called when the particle system is done
-DECLARE_MULTICAST_DELEGATE_OneParam(FOnNiagaraSystemFinished, class UNiagaraComponent*);
+DECLARE_DYNAMIC_MULTICAST_DELEGATE_OneParam(FOnNiagaraSystemFinished, class UNiagaraComponent*, PSystem);
+
 
 /**
 * UNiagaraComponent is the primitive component for a Niagara System.
 * @see ANiagaraActor
 * @see UNiagaraSystem
 */
-UCLASS(ClassGroup = (Rendering, Common), hidecategories = Object, hidecategories = Physics, hidecategories = Collision, showcategories = Trigger, editinlinenew, meta = (BlueprintSpawnableComponent))
+UCLASS(ClassGroup = (Rendering, Common), hidecategories = Object, hidecategories = Physics, hidecategories = Collision, showcategories = Trigger, editinlinenew, meta = (BlueprintSpawnableComponent, DisplayName = "Niagara Particle System"))
 class NIAGARA_API UNiagaraComponent : public UPrimitiveComponent
 {
 	GENERATED_UCLASS_BODY()
@@ -47,7 +48,7 @@ private:
 	Should expose anything in the "User" namespace.
 	*/
 	UPROPERTY(EditAnywhere, Category = Parameters)
-	FNiagaraParameterStore OverrideParameters;
+	FNiagaraUserRedirectionParameterStore OverrideParameters;
 
 #if WITH_EDITORONLY_DATA
 	UPROPERTY()
@@ -235,11 +236,11 @@ public:
 	UFUNCTION(BlueprintCallable, Category = Niagara, meta = (DisplayName = "Get Niagara Emitter Positions"))
 	TArray<FVector> GetNiagaraParticlePositions_DebugOnly(const FString& InEmitterName);
 	
-	/** Debug accessors for getting a float attribute array in blueprints. */
+	/** Debug accessors for getting a float attribute array in blueprints.  The attribute name should be without namespaces. For example for "Particles.Position", send "Position". */
 	UFUNCTION(BlueprintCallable, Category = Niagara, meta = (DisplayName = "Get Niagara Emitter Float Attrib"))
 	TArray<float> GetNiagaraParticleValues_DebugOnly(const FString& InEmitterName, const FString& InValueName);
 	
-	/** Debug accessors for getting a FVector attribute array in blueprints. */
+	/** Debug accessors for getting a FVector attribute array in blueprints. The attribute name should be without namespaces. For example for "Particles.Position", send "Position". */
 	UFUNCTION(BlueprintCallable, Category = Niagara, meta = (DisplayName = "Get Niagara Emitter Vec3 Attrib"))
 	TArray<FVector> GetNiagaraParticleValueVec3_DebugOnly(const FString& InEmitterName, const FString& InValueName);
 
@@ -279,6 +280,7 @@ public:
 
 	virtual void PostEditChangeProperty(FPropertyChangedEvent& PropertyChangedEvent) override;
 
+	void PostLoadNormalizeOverrideNames();
 	bool IsParameterValueOverriddenLocally(const FName& InParamName);
 	void SetParameterValueOverriddenLocally(const FNiagaraVariable& InParam, bool bInOverridden, bool bRequiresSystemInstanceReset);
 	
@@ -287,11 +289,14 @@ public:
 	FOnSynchronizedWithAssetParameters& OnSynchronizedWithAssetParameters() { return OnSynchronizedWithAssetParametersDelegate; }
 #endif
 
-	FNiagaraParameterStore& GetOverrideParameters() { return OverrideParameters; }
+	FNiagaraUserRedirectionParameterStore& GetOverrideParameters() { return OverrideParameters; }
+
+	const FNiagaraParameterStore& GetOverrideParameters() const { return OverrideParameters; }
 
 	//~ End UObject Interface.
 
 	// Called when the particle system is done
+	UPROPERTY(BlueprintAssignable)
 	FOnNiagaraSystemFinished OnSystemFinished;
 
 private:
diff --git a/Engine/Plugins/FX/Niagara/Source/Niagara/Public/NiagaraLightRendererProperties.h b/Engine/Plugins/FX/Niagara/Source/Niagara/Public/NiagaraLightRendererProperties.h
--- a/Engine/Plugins/FX/Niagara/Source/Niagara/Public/NiagaraLightRendererProperties.h
+++ b/Engine/Plugins/FX/Niagara/Source/Niagara/Public/NiagaraLightRendererProperties.h
@@ -31,21 +31,50 @@ public:
 	virtual const TArray<FNiagaraVariable>& GetOptionalAttributes() override;
 #endif // WITH_EDITORONLY_DATA
 
+	/** Whether to use physically based inverse squared falloff from the light.  If unchecked, the value from the LightExponent binding will be used instead. */
 	UPROPERTY(EditAnywhere, Category = "Light Rendering")
+	uint32 bUseInverseSquaredFalloff : 1;
+
+	/**
+	 * Whether lights from this renderer should affect translucency.
+	 * Use with caution - if enabled, create only a few particle lights at most, and the smaller they are, the less they will cost.
+	 */
+	UPROPERTY(EditAnywhere, Category = "Light Rendering")
+	uint32 bAffectsTranslucency : 1;
+
+	/** By default, a light is spawned for each particle. Enable this to control the spawn-rate on a per-particle basis. */
+	UPROPERTY(EditAnywhere, Category = "Light Rendering", meta = (InlineEditConditionToggle))
+	uint32 bOverrideRenderingEnabled : 1;
+
+	/** A factor used to scale each particle light radius */
+	UPROPERTY(EditAnywhere, Category = "Light Rendering", meta = (UIMin = "0"))
 	float RadiusScale;
 
+	/** A static color shift applied to each rendered light */
 	UPROPERTY(EditAnywhere, Category = "Light Rendering")
 	FVector ColorAdd;
 
+	/** Which attribute should we use to check if light rendering should be enabled for a particle? */
+	UPROPERTY(EditAnywhere, AdvancedDisplay, Category = "Bindings", meta = (EditCondition = "bOverrideRenderingEnabled"))
+	FNiagaraVariableAttributeBinding LightRenderingEnabledBinding;
+
+	/** Which attribute should we use for the light's exponent when inverse squared falloff is disabled? */
+	UPROPERTY(EditAnywhere, AdvancedDisplay, Category = "Bindings", meta = (EditCondition = "!bUseInverseSquaredFalloff"))
+	FNiagaraVariableAttributeBinding LightExponentBinding;
+
 	/** Which attribute should we use for position when generating lights?*/
-	UPROPERTY(EditAnywhere, AdvancedDisplay, Category = "Light Rendering")
+	UPROPERTY(EditAnywhere, AdvancedDisplay, Category = "Bindings")
 	FNiagaraVariableAttributeBinding PositionBinding;
 
 	/** Which attribute should we use for light color when generating lights?*/
-	UPROPERTY(EditAnywhere, AdvancedDisplay, Category = "Light Rendering")
+	UPROPERTY(EditAnywhere, AdvancedDisplay, Category = "Bindings")
 	FNiagaraVariableAttributeBinding ColorBinding;
 
 	/** Which attribute should we use for light radius when generating lights?*/
-	UPROPERTY(EditAnywhere, AdvancedDisplay, Category = "Light Rendering")
+	UPROPERTY(EditAnywhere, AdvancedDisplay, Category = "Bindings")
 	FNiagaraVariableAttributeBinding RadiusBinding;
+
+	/** Which attribute should we use for the intensity of the volumetric scattering from this light? This scales the light's intensity and color. */
+	UPROPERTY(EditAnywhere, AdvancedDisplay, Category = "Bindings")
+	FNiagaraVariableAttributeBinding VolumetricScatteringBinding;
 };
diff --git a/Engine/Plugins/FX/Niagara/Source/Niagara/Public/NiagaraModule.h b/Engine/Plugins/FX/Niagara/Source/Niagara/Public/NiagaraModule.h
--- a/Engine/Plugins/FX/Niagara/Source/Niagara/Public/NiagaraModule.h
+++ b/Engine/Plugins/FX/Niagara/Source/Niagara/Public/NiagaraModule.h
@@ -125,6 +125,7 @@ public:
 	FORCEINLINE static const FNiagaraVariable&  GetVar_Engine_Owner_YAxis() { return Engine_Owner_YAxis; }
 	FORCEINLINE static const FNiagaraVariable&  GetVar_Engine_Owner_ZAxis() { return Engine_Owner_ZAxis; }
 	FORCEINLINE static const FNiagaraVariable&  GetVar_Engine_Owner_Scale() { return Engine_Owner_Scale; }
+	FORCEINLINE static const FNiagaraVariable&  GetVar_Engine_Owner_Rotation() { return Engine_Owner_Rotation; }
 
 	FORCEINLINE static const FNiagaraVariable&  GetVar_Engine_Owner_SystemLocalToWorld() { return Engine_Owner_SystemLocalToWorld; }
 	FORCEINLINE static const FNiagaraVariable&  GetVar_Engine_Owner_SystemWorldToLocal() { return Engine_Owner_SystemWorldToLocal; }
@@ -140,6 +141,8 @@ public:
 
 	FORCEINLINE static const FNiagaraVariable&  GetVar_Engine_ExecutionCount() { return Engine_ExecutionCount; }
 	FORCEINLINE static const FNiagaraVariable&  GetVar_Engine_Emitter_NumParticles() { return Engine_Emitter_NumParticles; }
+	FORCEINLINE static const FNiagaraVariable&  GetVar_Engine_Emitter_TotalSpawnedParticles() { return Engine_Emitter_TotalSpawnedParticles; }
+	FORCEINLINE static const FNiagaraVariable&  GetVar_Engine_System_TickCount() { return Engine_System_TickCount; }
 	FORCEINLINE static const FNiagaraVariable&  GetVar_Engine_System_NumEmittersAlive() { return Engine_System_NumEmittersAlive; }
 	FORCEINLINE static const FNiagaraVariable&  GetVar_Engine_System_NumEmitters() { return Engine_System_NumEmitters; }
 	FORCEINLINE static const FNiagaraVariable&  GetVar_Engine_NumSystemInstances() { return Engine_NumSystemInstances; }
@@ -150,11 +153,14 @@ public:
 	FORCEINLINE static const FNiagaraVariable&  GetVar_Engine_System_Age() { return Engine_System_Age; }
 	FORCEINLINE static const FNiagaraVariable&  GetVar_Emitter_Age() { return Emitter_Age; }
 	FORCEINLINE static const FNiagaraVariable&  GetVar_Emitter_LocalSpace() { return Emitter_LocalSpace; }
+	FORCEINLINE static const FNiagaraVariable&  GetVar_Emitter_Determinism() { return Emitter_Determinism; }
+	FORCEINLINE static const FNiagaraVariable&  GetVar_Emitter_RandomSeed() { return Emitter_RandomSeed; }
 	FORCEINLINE static const FNiagaraVariable&  GetVar_Emitter_SpawnRate() { return Emitter_SpawnRate; }
 	FORCEINLINE static const FNiagaraVariable&  GetVar_Emitter_SpawnInterval() { return Emitter_SpawnInterval; }
 	FORCEINLINE static const FNiagaraVariable&  GetVar_Emitter_InterpSpawnStartDt() { return Emitter_InterpSpawnStartDt; }
 	FORCEINLINE static const FNiagaraVariable&  GetVar_Emitter_SpawnGroup() { return Emitter_SpawnGroup; }
 
+	FORCEINLINE static const FNiagaraVariable&  GetVar_Particles_UniqueID() { return Particles_UniqueID; }
 	FORCEINLINE static const FNiagaraVariable&  GetVar_Particles_ID() { return Particles_ID; }
 	FORCEINLINE static const FNiagaraVariable&  GetVar_Particles_Position() { return Particles_Position; }
 	FORCEINLINE static const FNiagaraVariable&  GetVar_Particles_Velocity() { return Particles_Velocity; }
@@ -176,6 +182,9 @@ public:
 	FORCEINLINE static const FNiagaraVariable&  GetVar_Particles_CameraOffset() { return Particles_CameraOffset; }
 	FORCEINLINE static const FNiagaraVariable&  GetVar_Particles_MaterialRandom() { return Particles_MaterialRandom; }
 	FORCEINLINE static const FNiagaraVariable&  GetVar_Particles_LightRadius() { return Particles_LightRadius; }
+	FORCEINLINE static const FNiagaraVariable&  GetVar_Particles_LightExponent() { return Particles_LightExponent; }
+	FORCEINLINE static const FNiagaraVariable&  GetVar_Particles_LightEnabled() { return Particles_LightEnabled; }
+	FORCEINLINE static const FNiagaraVariable&  GetVar_Particles_LightVolumetricScattering() { return Particles_LightVolumetricScattering; }
 	FORCEINLINE static const FNiagaraVariable&  GetVar_Particles_RibbonID() { return Particles_RibbonID; }
 	FORCEINLINE static const FNiagaraVariable&  GetVar_Particles_RibbonWidth() { return Particles_RibbonWidth; }
 	FORCEINLINE static const FNiagaraVariable&  GetVar_Particles_RibbonTwist() { return Particles_RibbonTwist; }
@@ -212,6 +221,7 @@ private:
 	static FNiagaraVariable Engine_Owner_YAxis;
 	static FNiagaraVariable Engine_Owner_ZAxis;
 	static FNiagaraVariable Engine_Owner_Scale;
+	static FNiagaraVariable Engine_Owner_Rotation;
 
 	static FNiagaraVariable Engine_Owner_SystemLocalToWorld;
 	static FNiagaraVariable Engine_Owner_SystemWorldToLocal;
@@ -227,6 +237,8 @@ private:
 
 	static FNiagaraVariable Engine_ExecutionCount;
 	static FNiagaraVariable Engine_Emitter_NumParticles;
+	static FNiagaraVariable Engine_Emitter_TotalSpawnedParticles;
+	static FNiagaraVariable Engine_System_TickCount;
 	static FNiagaraVariable Engine_System_NumEmittersAlive;
 	static FNiagaraVariable Engine_System_NumEmitters;
 	static FNiagaraVariable Engine_NumSystemInstances;
@@ -237,11 +249,14 @@ private:
 	static FNiagaraVariable Engine_System_Age;
 	static FNiagaraVariable Emitter_Age;
 	static FNiagaraVariable Emitter_LocalSpace;
+	static FNiagaraVariable Emitter_Determinism;
+	static FNiagaraVariable Emitter_RandomSeed;
 	static FNiagaraVariable Emitter_SpawnRate;
 	static FNiagaraVariable Emitter_SpawnInterval;
 	static FNiagaraVariable Emitter_InterpSpawnStartDt;
 	static FNiagaraVariable Emitter_SpawnGroup;
 
+	static FNiagaraVariable Particles_UniqueID;
 	static FNiagaraVariable Particles_ID;
 	static FNiagaraVariable Particles_Position;
 	static FNiagaraVariable Particles_Velocity;
@@ -263,6 +278,9 @@ private:
 	static FNiagaraVariable Particles_CameraOffset;
 	static FNiagaraVariable Particles_MaterialRandom;
 	static FNiagaraVariable Particles_LightRadius;
+	static FNiagaraVariable Particles_LightExponent;
+	static FNiagaraVariable Particles_LightEnabled;
+	static FNiagaraVariable Particles_LightVolumetricScattering;
 	static FNiagaraVariable Particles_RibbonID;
 	static FNiagaraVariable Particles_RibbonWidth;
 	static FNiagaraVariable Particles_RibbonTwist;
diff --git a/Engine/Plugins/FX/Niagara/Source/Niagara/Public/NiagaraParameterStore.h b/Engine/Plugins/FX/Niagara/Source/Niagara/Public/NiagaraParameterStore.h
--- a/Engine/Plugins/FX/Niagara/Source/Niagara/Public/NiagaraParameterStore.h
+++ b/Engine/Plugins/FX/Niagara/Source/Niagara/Public/NiagaraParameterStore.h
@@ -119,7 +119,7 @@ public:
 	FNiagaraParameterStore(const FNiagaraParameterStore& Other);
 	FNiagaraParameterStore& operator=(const FNiagaraParameterStore& Other);
 
-	~FNiagaraParameterStore();
+	virtual ~FNiagaraParameterStore();
 	
 #if WITH_EDITORONLY_DATA
 	UPROPERTY()
@@ -152,33 +152,33 @@ public:
 	/** Unbinds this store from all stores it's being driven by. */
 	void UnbindFromSourceStores();
 	
-	bool VerifyBinding(const FNiagaraParameterStore* InDestStore)const;
+	bool VerifyBinding(const FNiagaraParameterStore* InDestStore) const;
 
-	void CheckForNaNs()const;
+	void CheckForNaNs() const;
 
 	/**
 	Adds the passed parameter to this store.
 	Does nothing if this parameter is already present.
 	Returns true if we added a new parameter.
 	*/
-	bool AddParameter(const FNiagaraVariable& Param, bool bInitialize=true, bool bTriggerRebind = true);
+	virtual bool AddParameter(const FNiagaraVariable& Param, bool bInitialize=true, bool bTriggerRebind = true);
 
 	/** Removes the passed parameter if it exists in the store. */
-	bool RemoveParameter(const FNiagaraVariable& Param);
+	virtual bool RemoveParameter(const FNiagaraVariable& Param);
 
 	/** Renames the passed parameter. */
 	void RenameParameter(const FNiagaraVariable& Param, FName NewName);
 
 	/** Removes all parameters from this store and releases any data. */
-	void Empty(bool bClearBindings = true);
+	virtual void Empty(bool bClearBindings = true);
 
 	/** Removes all parameters from this store but does't change memory allocations. */
-	void Reset(bool bClearBindings = true);
+	virtual void Reset(bool bClearBindings = true);
 
 	FORCEINLINE TArray<FNiagaraParameterStore*>& GetSourceParameterStores() { return SourceStores; }
-	FORCEINLINE const TMap<FNiagaraVariable, int32>& GetParameterOffests()const { return ParameterOffsets; }
+	FORCEINLINE const TMap<FNiagaraVariable, int32>& GetParameterOffsets()const { return ParameterOffsets; }
 	/** Get the list of FNiagaraVariables referenced by this store. Note that the values will be stale and are not to be trusted directly. Get the Values using the offset specified by IndexOf or GetParameterValue.*/
-	FORCEINLINE void GetParameters(TArray<FNiagaraVariable>& OutParametres) const { return ParameterOffsets.GenerateKeyArray(OutParametres); }
+	FORCEINLINE void GetParameters(TArray<FNiagaraVariable>& OutParameters) const { return ParameterOffsets.GenerateKeyArray(OutParameters); }
 	FORCEINLINE int32 GetNumParameters() const { return ParameterOffsets.Num(); }
 
 	FORCEINLINE const TArray<UNiagaraDataInterface*>& GetDataInterfaces()const { return DataInterfaces; }
@@ -187,10 +187,10 @@ public:
 	FORCEINLINE void SetParameterDataArray(const TArray<uint8>& InParameterDataArray);
 
 	// Called to initially set up the parameter store to *exactly* match the input store (other than any bindings and the internal name of it).
-	void InitFromSource(const FNiagaraParameterStore* SrcStore, bool bNotifyAsDirty);
+	virtual void InitFromSource(const FNiagaraParameterStore* SrcStore, bool bNotifyAsDirty);
 
 	/** Gets the index of the passed parameter. If it is a data interface, this is an offset into the data interface table, otherwise a byte offset into he parameter data buffer. */
-	int32 IndexOf(const FNiagaraVariable& Parameter)const
+	virtual int32 IndexOf(const FNiagaraVariable& Parameter) const
 	{
 		const int32* Off = ParameterOffsets.Find(Parameter);
 		if ( Off )
@@ -263,7 +263,7 @@ public:
 	/** Returns the associated FNiagaraVariable for the passed data interface if it exists in the store. Null if not.*/
 	const FNiagaraVariable* FindVariable(UNiagaraDataInterface* Interface)const;
 
-	FORCEINLINE_DEBUGGABLE const int32* FindParameterOffset(const FNiagaraVariable& Parameter) const
+	FORCEINLINE_DEBUGGABLE virtual const int32* FindParameterOffset(const FNiagaraVariable& Parameter) const
 	{
 		return ParameterOffsets.Find(Parameter);
 	}
@@ -468,7 +468,7 @@ FORCEINLINE_DEBUGGABLE void FNiagaraParameterStoreBinding::BindParameters(FNiaga
 {
 	InterfaceBindings.Reset();
 	ParameterBindings.Reset();
-	for (const TPair<FNiagaraVariable, int32>& ParamOffsetPair : DestStore->GetParameterOffests())
+	for (const TPair<FNiagaraVariable, int32>& ParamOffsetPair : DestStore->GetParameterOffsets())
 	{
 		const FNiagaraVariable& Parameter = ParamOffsetPair.Key;
 		int32 DestOffset = ParamOffsetPair.Value;
@@ -505,7 +505,7 @@ FORCEINLINE_DEBUGGABLE bool FNiagaraParameterStoreBinding::VerifyBinding(const F
 	bool bBindingValid = true;
 #if WITH_EDITORONLY_DATA
 	TArray<FName, TInlineAllocator<32>> MissingParameterNames;
-	for (const TPair<FNiagaraVariable, int32>& ParamOffsetPair : DestStore->GetParameterOffests())
+	for (const TPair<FNiagaraVariable, int32>& ParamOffsetPair : DestStore->GetParameterOffsets())
 	{
 		const FNiagaraVariable& Parameter = ParamOffsetPair.Key;
 		int32 DestOffset = ParamOffsetPair.Value;
@@ -578,7 +578,7 @@ FORCEINLINE_DEBUGGABLE void FNiagaraParameterStoreBinding::Dump(const FNiagaraPa
 		ensure(Binding.DestOffset != -1);
 		FNiagaraVariable Param;
 		bool bFound = false;
-		for (const TPair<FNiagaraVariable, int32>& ParamOffsetPair : DestStore->GetParameterOffests())
+		for (const TPair<FNiagaraVariable, int32>& ParamOffsetPair : DestStore->GetParameterOffsets())
 		{
 			if (ParamOffsetPair.Value == Binding.DestOffset && !ParamOffsetPair.Key.IsDataInterface())
 			{
@@ -610,7 +610,7 @@ FORCEINLINE_DEBUGGABLE void FNiagaraParameterStoreBinding::Dump(const FNiagaraPa
 		ensure(Binding.DestOffset != -1);
 		FNiagaraVariable Param;
 		bool bFound = false;
-		for (const TPair<FNiagaraVariable, int32>& ParamOffsetPair : DestStore->GetParameterOffests())
+		for (const TPair<FNiagaraVariable, int32>& ParamOffsetPair : DestStore->GetParameterOffsets())
 		{
 			if (ParamOffsetPair.Value == Binding.DestOffset && ParamOffsetPair.Key.IsDataInterface())
 			{
diff --git a/Engine/Plugins/FX/Niagara/Source/Niagara/Public/NiagaraScriptExecutionParameterStore.h b/Engine/Plugins/FX/Niagara/Source/Niagara/Public/NiagaraScriptExecutionParameterStore.h
--- a/Engine/Plugins/FX/Niagara/Source/Niagara/Public/NiagaraScriptExecutionParameterStore.h
+++ b/Engine/Plugins/FX/Niagara/Source/Niagara/Public/NiagaraScriptExecutionParameterStore.h
@@ -33,6 +33,8 @@ public:
 	FNiagaraScriptExecutionParameterStore(const FNiagaraParameterStore& Other);
 	FNiagaraScriptExecutionParameterStore& operator=(const FNiagaraParameterStore& Other);
 
+	virtual ~FNiagaraScriptExecutionParameterStore() = default;
+
 	//TODO: These function can probably go away entirely when we replace the FNiagaraParameters and DataInterface info in the script with an FNiagaraParameterStore.
 	//Special care with prev params and internal params will have to be taken there.
 	/** Call this init function if you are using a Niagara parameter store within a UNiagaraScript.*/
@@ -42,7 +44,7 @@ public:
 	void AddScriptParams(UNiagaraScript* Script, ENiagaraSimTarget SimTarget, bool bTriggerRebind);
 	void CopyCurrToPrev();
 
-	bool AddParameter(const FNiagaraVariable& Param, bool bInitInterfaces = true, bool bTriggerRebind = true)
+	virtual bool AddParameter(const FNiagaraVariable& Param, bool bInitInterfaces = true, bool bTriggerRebind = true) override
 	{
 		if (FNiagaraParameterStore::AddParameter(Param, bInitInterfaces, bTriggerRebind))
 		{
@@ -52,7 +54,7 @@ public:
 		return false;
 	}
 
-	bool RemoveParameter(FNiagaraVariable& Param)
+	virtual bool RemoveParameter(const FNiagaraVariable& Param) override
 	{
 		check(0);//Not allowed to remove parameters from an execution store as it will adjust the table layout mess up the 
 		return false;
@@ -63,7 +65,7 @@ public:
 		check(0);//Can't rename parameters for an execution store.
 	}
 
-	void Empty(bool bClearBindings=true)
+	virtual void Empty(bool bClearBindings=true) override
 	{
 		FNiagaraParameterStore::Empty(bClearBindings);
 		PaddingInfo.Empty();
