Merging //UE4/Dev-Main to Dev-Build (//UE4/Dev-Build)

diff --git a/Engine/Plugins/FX/Niagara/Source/Niagara/Public/NiagaraSystemInstance.h b/Engine/Plugins/FX/Niagara/Source/Niagara/Public/NiagaraSystemInstance.h
--- a/Engine/Plugins/FX/Niagara/Source/Niagara/Public/NiagaraSystemInstance.h
+++ b/Engine/Plugins/FX/Niagara/Source/Niagara/Public/NiagaraSystemInstance.h
@@ -54,7 +54,7 @@ public:
 	void Cleanup();
 
 	/** Initializes this System instance to simulate the supplied System. */
-	void Init(UNiagaraSystem* InSystem, bool bInForceSolo=false);
+	void Init(bool bInForceSolo=false);
 
 	void Activate(EResetMode InResetMode = EResetMode::ResetAll);
 	void Deactivate(bool bImmediate = false);
@@ -173,7 +173,7 @@ public:
 	Manually advances this system's simulation by the specified number of ticks and tick delta. 
 	To be advanced in this way a system must be in solo mode or moved into solo mode which will add additional overhead.
 	*/
-	void AdvanceSimulation(int32 TickCount, float TickDeltaSeconds);
+	void AdvanceSimulation(int32 TickCountToSimulate, float TickDeltaSeconds);
 
 #if WITH_EDITORONLY_DATA
 	/** Request that this simulation capture a frame. Cannot capture if disabled or already completed.*/
@@ -192,13 +192,15 @@ public:
 	bool ShouldCaptureThisFrame() const;
 
 	/** Only call from within the script execution states. Value is nullptr if not capturing a frame.*/
-	FNiagaraScriptDebuggerInfo* GetActiveCaptureWrite(const FName& InHandleName, ENiagaraScriptUsage InUsage, const FGuid& InUsageId);
+	TSharedPtr<struct FNiagaraScriptDebuggerInfo, ESPMode::ThreadSafe> GetActiveCaptureWrite(const FName& InHandleName, ENiagaraScriptUsage InUsage, const FGuid& InUsageId);
 
 #endif
 
 	/** Dumps all of this systems info to the log. */
 	void Dump()const;
 
+	bool GetPerInstanceDataAndOffsets(void*& OutData, uint32& OutDataSize, TMap<TWeakObjectPtr<UNiagaraDataInterface>, int32>*& OutOffsets);
+
 private:
 
 	/** Builds the emitter simulations. */
@@ -229,6 +231,9 @@ private:
 	/** The age of the System instance. */
 	float Age;
 
+	/** The tick count of the System instance. */
+	int32 TickCount;
+
 	TMap<FNiagaraDataSetID, FNiagaraDataSet> ExternalEvents;
 
 	TArray< TSharedRef<FNiagaraEmitterInstance> > Emitters;
@@ -267,6 +272,8 @@ private:
 	FNiagaraParameterDirectBinding<FVector> OwnerYAxisParam;
 	FNiagaraParameterDirectBinding<FVector> OwnerZAxisParam;
 
+	FNiagaraParameterDirectBinding<FQuat> OwnerRotationParam;
+
 	FNiagaraParameterDirectBinding<FMatrix> OwnerTransformParam;
 	FNiagaraParameterDirectBinding<FMatrix> OwnerInverseParam;
 	FNiagaraParameterDirectBinding<FMatrix> OwnerTransposeParam;
@@ -279,6 +286,7 @@ private:
 	FNiagaraParameterDirectBinding<float> OwnerEngineTimeParam;
 	FNiagaraParameterDirectBinding<float> OwnerEngineRealtimeParam;
 	FNiagaraParameterDirectBinding<float> SystemAgeParam;
+	FNiagaraParameterDirectBinding<int32> SystemTickCountParam;
 
 	FNiagaraParameterDirectBinding<float> OwnerMinDistanceToCameraParam;
 	FNiagaraParameterDirectBinding<int32> SystemNumEmittersParam;
@@ -289,6 +297,7 @@ private:
 	FNiagaraParameterDirectBinding<int32> OwnerExecutionStateParam;
 
 	TArray<FNiagaraParameterDirectBinding<int32>> ParameterNumParticleBindings;
+	TArray<FNiagaraParameterDirectBinding<int32>> ParameterTotalSpawnedParticlesBindings;
 
 	/** Indicates whether this instance must update itself rather than being batched up as most instances are. */
 	uint32 bSolo : 1;
diff --git a/Engine/Plugins/FX/Niagara/Source/Niagara/Public/NiagaraSystemSimulation.h b/Engine/Plugins/FX/Niagara/Source/Niagara/Public/NiagaraSystemSimulation.h
--- a/Engine/Plugins/FX/Niagara/Source/Niagara/Public/NiagaraSystemSimulation.h
+++ b/Engine/Plugins/FX/Niagara/Source/Niagara/Public/NiagaraSystemSimulation.h
@@ -121,7 +121,7 @@ class FNiagaraSystemSimulation
 {
 public:
 	~FNiagaraSystemSimulation();
-	bool Init(UNiagaraSystem* InSystem, UWorld* InWorld, FNiagaraSystemInstance* InSoloSystemInstance=nullptr);
+	bool Init(UNiagaraSystem* InSystem, UWorld* InWorld, bool bInIsSolo);
 	void Destroy();
 	bool Tick(float DeltaSeconds);
 
@@ -141,14 +141,17 @@ public:
 	void TransferInstance(FNiagaraSystemSimulation* SourceSimulation, FNiagaraSystemInstance* SystemInst);
 
 	void DumpInstance(const FNiagaraSystemInstance* Inst)const;
+
+	bool GetIsSolo() const { return bIsSolo; }
+
+	FNiagaraScriptExecutionContext& GetSpawnExecutionContext() { return SpawnExecContext; }
+	FNiagaraScriptExecutionContext& GetUpdateExecutionContext() { return UpdateExecContext; }
+
 protected:
 
 	/** System of instances being simulated.  We use a weak object ptr here because once the last referencing object goes away this system may be come invalid at runtime. */
 	TWeakObjectPtr<UNiagaraSystem> WeakSystem;
 
-	/** The parent system instance if this simulation is  */
-	FNiagaraSystemInstance* SoloSystemInstance;
-
 	/** World this system simulation belongs to. */
 	UWorld* World;
 
@@ -206,7 +209,7 @@ protected:
 
 	TArray<TArray<FNiagaraDataSetAccessor<FNiagaraSpawnInfo>>> EmitterSpawnInfoAccessors;
 
-	void InitBindings(FNiagaraSystemInstance* SystemInst);
+	void InitParameterDataSetBindings(FNiagaraSystemInstance* SystemInst);
 
 	FNiagaraDataSetAccessor<int32> SystemExecutionStateAccessor;
 	TArray<FNiagaraDataSetAccessor<int32>> EmitterExecutionStateAccessors;
@@ -215,4 +218,8 @@ protected:
 
 	/** A parameter store which contains the data interfaces parameters which were defined by the scripts. */
 	FNiagaraParameterStore ScriptDefinedDataInterfaceParameters;
+
+	bool bIsSolo;
+
+	TOptional<float> MaxDeltaTime;
 };
diff --git a/Engine/Plugins/FX/Niagara/Source/Niagara/Public/NiagaraTypes.h b/Engine/Plugins/FX/Niagara/Source/Niagara/Public/NiagaraTypes.h
--- a/Engine/Plugins/FX/Niagara/Source/Niagara/Public/NiagaraTypes.h
+++ b/Engine/Plugins/FX/Niagara/Source/Niagara/Public/NiagaraTypes.h
@@ -892,3 +892,21 @@ FORCEINLINE uint32 GetTypeHash(const FNiagaraVariable& Var)
 {
 	return HashCombine(GetTypeHash(Var.GetType()), GetTypeHash(Var.GetName()));
 }
+
+template<>
+inline bool FNiagaraVariable::GetValue<bool>() const
+{
+	check(TypeDef == FNiagaraTypeDefinition::GetBoolDef());
+	check(IsDataAllocated());
+	FNiagaraBool* BoolStruct = (FNiagaraBool*)GetData();
+	return BoolStruct->GetValue();
+}
+
+template<>
+inline void FNiagaraVariable::SetValue<bool>(const bool& Data)
+{
+	check(TypeDef == FNiagaraTypeDefinition::GetBoolDef());
+	AllocateData();
+	FNiagaraBool* BoolStruct = (FNiagaraBool*)GetData();
+	BoolStruct->SetValue(Data);
+}
\ No newline at end of file
diff --git a/Engine/Plugins/FX/Niagara/Source/Niagara/Public/NiagaraWorldManager.h b/Engine/Plugins/FX/Niagara/Source/Niagara/Public/NiagaraWorldManager.h
--- a/Engine/Plugins/FX/Niagara/Source/Niagara/Public/NiagaraWorldManager.h
+++ b/Engine/Plugins/FX/Niagara/Source/Niagara/Public/NiagaraWorldManager.h
@@ -9,7 +9,7 @@
 #include "NiagaraDataSet.h"
 #include "NiagaraScriptExecutionContext.h"
 #include "NiagaraSystemSimulation.h"
-
+#include "GlobalDistanceFieldParameters.h"
 #include "NiagaraDataInterfaceSkeletalMesh.h"
 
 class UWorld;
@@ -32,22 +32,39 @@ public:
 		ViewUniformBuffer = Params.ViewUniformBuffer;
 		SceneNormalTexture = Params.NormalTexture;
 		SceneTexturesUniformParams = Params.SceneTexturesUniformParams;
+		GlobalDistanceFieldParams = Params.GlobalDistanceFieldParams;
+
+		PreSceneRenderValues = FPreSceneRenderValues();
+	}
+
+	void OnPreSceneRenderCalled(FPreSceneRenderValues& OutValues) const
+	{
+		OutValues.bUsesGlobalDistanceField |= PreSceneRenderValues.bUsesGlobalDistanceField;
+	}
+
+	void SetGlobalDistanceFieldUsage()
+	{
+		PreSceneRenderValues.bUsesGlobalDistanceField = true;
 	}
 
 	FTexture2DRHIParamRef GetSceneDepthTexture() { return SceneDepthTexture; }
 	FTexture2DRHIParamRef GetSceneNormalTexture() { return SceneNormalTexture; }
 	FUniformBufferRHIParamRef GetViewUniformBuffer() { return ViewUniformBuffer; }
 	TUniformBufferRef<FSceneTexturesUniformParameters> GetSceneTextureUniformParameters() { return SceneTexturesUniformParams; }
+	const FGlobalDistanceFieldParameterData* GetGlobalDistanceFieldParameters() { return GlobalDistanceFieldParams; }
 
 	virtual void InitDynamicRHI() override;
 
 	virtual void ReleaseDynamicRHI() override;
+
 private:
 	FTexture2DRHIParamRef SceneDepthTexture;
 	FTexture2DRHIParamRef SceneNormalTexture;
 	FUniformBufferRHIParamRef ViewUniformBuffer;
+	FPreSceneRenderValues PreSceneRenderValues;	
 
 	TUniformBufferRef<FSceneTexturesUniformParameters> SceneTexturesUniformParams;
+	const FGlobalDistanceFieldParameterData* GlobalDistanceFieldParams;
 	FPostOpaqueRenderDelegate PostOpaqueDelegate;
 };
 
diff --git a/Engine/Plugins/FX/Niagara/Source/NiagaraCore/Private/NiagaraCustomVersion.cpp b/Engine/Plugins/FX/Niagara/Source/NiagaraCore/Private/NiagaraCustomVersion.cpp
--- a/Engine/Plugins/FX/Niagara/Source/NiagaraCore/Private/NiagaraCustomVersion.cpp
+++ b/Engine/Plugins/FX/Niagara/Source/NiagaraCore/Private/NiagaraCustomVersion.cpp
@@ -8,4 +8,4 @@ const FGuid FNiagaraCustomVersion::GUID(0xFCF57AFA, 0x50764283, 0xB9A9E658, 0xFF
 // Register the custom version with core
 FCustomVersionRegistration GRegisterNiagaraCustomVersion(FNiagaraCustomVersion::GUID, FNiagaraCustomVersion::LatestVersion, TEXT("NiagaraVer"));
 
-const FGuid FNiagaraCustomVersion::LatestScriptCompileVersion(0x124E242E, 0xC5334EB6, 0xB88A6679, 0x582DE2EE);
\ No newline at end of file
+const FGuid FNiagaraCustomVersion::LatestScriptCompileVersion(0x0BE17832, 0xB1DE4421, 0x95243884, 0xBD77C7B3);
diff --git a/Engine/Plugins/FX/Niagara/Source/NiagaraCore/Public/NiagaraDataInterfaceBase.h b/Engine/Plugins/FX/Niagara/Source/NiagaraCore/Public/NiagaraDataInterfaceBase.h
--- a/Engine/Plugins/FX/Niagara/Source/NiagaraCore/Public/NiagaraDataInterfaceBase.h
+++ b/Engine/Plugins/FX/Niagara/Source/NiagaraCore/Public/NiagaraDataInterfaceBase.h
@@ -21,7 +21,9 @@ struct FNiagaraDataInterfaceParametersCS
 	virtual ~FNiagaraDataInterfaceParametersCS() {}
 	virtual void Bind(const FNiagaraDataInterfaceParamRef& ParamRef, const class FShaderParameterMap& ParameterMap) {}
 	virtual void Serialize(FArchive& Ar) { }
-	virtual void Set(FRHICommandList& RHICmdList, FNiagaraShader* Shader, class UNiagaraDataInterface* DataInterface) const {}
+	virtual void Set(FRHICommandList& RHICmdList, FNiagaraShader* Shader, class UNiagaraDataInterface* DataInterface, void* PerInstanceData) const {}
+	virtual void Unset(FRHICommandList& RHICmdList, FNiagaraShader* Shader, class UNiagaraDataInterface* DataInterface, void* PerInstanceData) const {}
+
 };
 
 //////////////////////////////////////////////////////////////////////////
diff --git a/Engine/Plugins/FX/Niagara/Source/NiagaraEditor/Private/AssetTypeActions/AssetTypeActions_NiagaraEmitter.cpp b/Engine/Plugins/FX/Niagara/Source/NiagaraEditor/Private/AssetTypeActions/AssetTypeActions_NiagaraEmitter.cpp
--- a/Engine/Plugins/FX/Niagara/Source/NiagaraEditor/Private/AssetTypeActions/AssetTypeActions_NiagaraEmitter.cpp
+++ b/Engine/Plugins/FX/Niagara/Source/NiagaraEditor/Private/AssetTypeActions/AssetTypeActions_NiagaraEmitter.cpp
@@ -112,6 +112,12 @@ void FAssetTypeActions_NiagaraEmitter::ExecuteNewNiagaraSystem(TArray<TWeakObjec
 				FNiagaraSystemViewModel SystemViewModel = FNiagaraSystemViewModel(*System, SystemOptions);
 				SystemViewModel.GetSystemScriptViewModel()->RebuildEmitterNodes();
 
+				// Ensure the new System is compiled
+				if (!Emitter->AreAllScriptAndSourcesSynchronized())
+				{
+					System->RequestCompile(true);
+				}
+
 				ObjectsToSync.Add(NewAsset);
 			}
 		}
diff --git a/Engine/Plugins/FX/Niagara/Source/NiagaraEditor/Private/Customizations/NiagaraComponentDetails.cpp b/Engine/Plugins/FX/Niagara/Source/NiagaraEditor/Private/Customizations/NiagaraComponentDetails.cpp
--- a/Engine/Plugins/FX/Niagara/Source/NiagaraEditor/Private/Customizations/NiagaraComponentDetails.cpp
+++ b/Engine/Plugins/FX/Niagara/Source/NiagaraEditor/Private/Customizations/NiagaraComponentDetails.cpp
@@ -46,6 +46,7 @@
 #include "INiagaraEditorTypeUtilities.h"
 #include "Widgets/Layout/SBox.h"
 #include "UObject/WeakObjectPtr.h"
+#include "NiagaraUserRedirectionParameterStore.h"
 #define LOCTEXT_NAMESPACE "NiagaraComponentDetails"
 
 class FNiagaraComponentNodeBuilder : public IDetailCustomNodeBuilder
@@ -55,8 +56,10 @@ public:
 	{
 		Component = InComponent;
 		Component->OnSynchronizedWithAssetParameters().AddRaw(this, &FNiagaraComponentNodeBuilder::ComponentSynchronizedWithAssetParameters);
-		OriginalScripts.Add(SourceSpawn);
-		OriginalScripts.Add(SourceUpdate);
+		if (SourceSpawn)
+			OriginalScripts.Add(SourceSpawn);
+		if (SourceUpdate)
+			OriginalScripts.Add(SourceUpdate);
 		//UE_LOG(LogNiagaraEditor, Log, TEXT("FNiagaraComponentNodeBuilder %p Component %p"), this, Component.Get());
 	}
 
@@ -88,8 +91,8 @@ public:
 	{
 		check(Component.IsValid());
 		TArray<FNiagaraVariable> Parameters;
-		FNiagaraParameterStore& ParamStore = Component->GetOverrideParameters();
-		ParamStore.GetParameters(Parameters);
+		FNiagaraUserRedirectionParameterStore& ParamStore = Component->GetOverrideParameters();
+		ParamStore.GetUserParameters(Parameters);
 
 		FNiagaraEditorModule& NiagaraEditorModule = FModuleManager::GetModuleChecked<FNiagaraEditorModule>("NiagaraEditor");
 		
@@ -233,7 +236,7 @@ private:
 	{
 		check(Component.IsValid());
 		Component->GetOverrideParameters().OnParameterChange();
-		Component->SetParameterValueOverriddenLocally(Var, true, false);
+		Component->SetParameterValueOverriddenLocally(Var, true, true);
 	}
 
 	void OnDataInterfaceChanged(FNiagaraVariable Var)
@@ -254,7 +257,7 @@ private:
 		check(Component.IsValid());
 		FScopedTransaction ScopedTransaction(LOCTEXT("ResetParameterValue", "Reset parameter value to system defaults."));
 		Component->Modify();
-		Component->SetParameterValueOverriddenLocally(Parameter, false, false);
+		Component->SetParameterValueOverriddenLocally(Parameter, false, true);
 		return FReply::Handled();
 	}
 
@@ -360,6 +363,11 @@ void FNiagaraComponentDetails::CustomizeDetails(IDetailLayoutBuilder& DetailBuil
 			IDetailCategoryBuilder& InputParamCategory = DetailBuilder.EditCategory(ParamCategoryName, LOCTEXT("ParamCategoryName", "Override Parameters"));
 			InputParamCategory.AddCustomBuilder(MakeShared<FNiagaraComponentNodeBuilder>(Component.Get(), ScriptSpawn, ScriptUpdate));
 		}
+		else
+		{
+			IDetailCategoryBuilder& InputParamCategory = DetailBuilder.EditCategory(ParamCategoryName, LOCTEXT("ParamCategoryName", "Override Parameters"));
+			InputParamCategory.AddCustomBuilder(MakeShared<FNiagaraComponentNodeBuilder>(Component.Get(), nullptr, nullptr));
+		}
 	}
 }
 
diff --git a/Engine/Plugins/FX/Niagara/Source/NiagaraEditor/Private/EdGraphSchema_Niagara.cpp b/Engine/Plugins/FX/Niagara/Source/NiagaraEditor/Private/EdGraphSchema_Niagara.cpp
--- a/Engine/Plugins/FX/Niagara/Source/NiagaraEditor/Private/EdGraphSchema_Niagara.cpp
+++ b/Engine/Plugins/FX/Niagara/Source/NiagaraEditor/Private/EdGraphSchema_Niagara.cpp
@@ -42,6 +42,7 @@
 
 #include "Modules/ModuleManager.h"
 #include "AssetRegistryModule.h"
+#include "NiagaraNodeSimTargetSelector.h"
 
 #define LOCTEXT_NAMESPACE "NiagaraSchema"
 
@@ -482,7 +483,7 @@ TArray<TSharedPtr<FNiagaraSchemaAction_NewNode> > UEdGraphSchema_Niagara::GetGra
 			{
 				{
 					const FText MenuDescFmt = LOCTEXT("AddEventReadFmt", "Add {0} Event Read");
-					const FText MenuDesc = FText::Format(MenuDescFmt, Type.GetStruct()->GetDisplayNameText());
+					const FText MenuDesc = FText::Format(MenuDescFmt, Type.GetNameText());
 
 					TSharedPtr<FNiagaraSchemaAction_NewNode> Action = AddNewNodeAction(NewActions, MenuCat, MenuDesc, *MenuDesc.ToString(), FText::GetEmpty());
 
@@ -492,7 +493,7 @@ TArray<TSharedPtr<FNiagaraSchemaAction_NewNode> > UEdGraphSchema_Niagara::GetGra
 				}
 				{
 					const FText MenuDescFmt = LOCTEXT("AddEventWriteFmt", "Add {0} Event Write");
-					const FText MenuDesc = FText::Format(MenuDescFmt, Type.GetStruct()->GetDisplayNameText());
+					const FText MenuDesc = FText::Format(MenuDescFmt, Type.GetNameText());
 
 					TSharedPtr<FNiagaraSchemaAction_NewNode> Action = AddNewNodeAction(NewActions, MenuCat, MenuDesc, *MenuDesc.ToString(), FText::GetEmpty());
 
@@ -570,11 +571,8 @@ TArray<TSharedPtr<FNiagaraSchemaAction_NewNode> > UEdGraphSchema_Niagara::GetGra
 			FText DescFmt = LOCTEXT("NiagaraMakeBreakFmt", "{0}");
 			auto MakeBreakType = [&](FNiagaraTypeDefinition Type, bool bMake)
 			{
-				FText DisplayName = Type.GetStruct()->GetDisplayNameText();
-				if (Type.GetEnum())
-				{
-					DisplayName = FText::FromString(Type.GetEnum()->GetName());
-				}
+				FText DisplayName = Type.GetNameText();
+
 				FText Desc = FText::Format(DescFmt, DisplayName);
 				TSharedPtr<FNiagaraSchemaAction_NewNode> Action = AddNewNodeAction(NewActions, bMake ? MakeCat : BreakCat, Desc, *Type.GetStruct()->GetName(), FText::GetEmpty());
 				UNiagaraNodeConvert* ConvertNode = NewObject<UNiagaraNodeConvert>(OwnerOfTemporaries);
@@ -832,7 +830,7 @@ TArray<TSharedPtr<FNiagaraSchemaAction_NewNode> > UEdGraphSchema_Niagara::GetGra
 						}
 					}
 						
-					const FText MenuDesc = FText::Format(MenuDescFmt, Type.GetStruct()->GetDisplayNameText());
+					const FText MenuDesc = FText::Format(MenuDescFmt, Type.GetNameText());
 					TSharedPtr<FNiagaraSchemaAction_NewNode> InputAction = AddNewNodeAction(NewActions, MenuCat, MenuDesc, *MenuDesc.ToString(), FText::GetEmpty());
 					UNiagaraNodeInput* InputNode = NewObject<UNiagaraNodeInput>(OwnerOfTemporaries);
 					FNiagaraEditorUtilities::InitializeParameterInputNode(*InputNode, Type, NiagaraGraph);
@@ -854,7 +852,7 @@ TArray<TSharedPtr<FNiagaraSchemaAction_NewNode> > UEdGraphSchema_Niagara::GetGra
 							MenuCat = LOCTEXT("AddRIParameterCat", "Add Rapid Iteration Param");
 						}
 
-						const FText MenuDesc = FText::Format(MenuDescFmt, Type.GetStruct()->GetDisplayNameText());
+						const FText MenuDesc = FText::Format(MenuDescFmt, Type.GetNameText());
 						TSharedPtr<FNiagaraSchemaAction_NewNode> InputAction = AddNewNodeAction(NewActions, MenuCat, MenuDesc, *MenuDesc.ToString(), FText::GetEmpty());
 						UNiagaraNodeInput* InputNode = NewObject<UNiagaraNodeInput>(OwnerOfTemporaries);
 						FNiagaraEditorUtilities::InitializeParameterInputNode(*InputNode, Type, NiagaraGraph);
@@ -866,7 +864,7 @@ TArray<TSharedPtr<FNiagaraSchemaAction_NewNode> > UEdGraphSchema_Niagara::GetGra
 				if (PinType != FNiagaraTypeDefinition::GetGenericNumericDef())
 				{
 					//For correctly typed pins, offer the correct type at the top level.				
-					const FText MenuDesc = FText::Format(MenuDescFmt, PinType.GetStruct()->GetDisplayNameText());
+					const FText MenuDesc = FText::Format(MenuDescFmt, PinType.GetNameText());
 					TSharedPtr<FNiagaraSchemaAction_NewNode> InputAction = AddNewNodeAction(NewActions, FText::GetEmpty(), MenuDesc, *MenuDesc.ToString(), FText::GetEmpty());
 					UNiagaraNodeInput* InputNode = NewObject<UNiagaraNodeInput>(OwnerOfTemporaries);
 					FNiagaraEditorUtilities::InitializeParameterInputNode(*InputNode, PinType, NiagaraGraph);
@@ -905,6 +903,15 @@ TArray<TSharedPtr<FNiagaraSchemaAction_NewNode> > UEdGraphSchema_Niagara::GetGra
 		Action->NodeTemplate = Node;
 	}
 
+	// Add simulation target selector node
+	{
+		const FText UtilMenuCat = LOCTEXT("NiagaraSimTargetSelectorMenuCat", "Util");
+		const FText SimTargetSelectorMenuDesc = LOCTEXT("NiagaraSimTargetSelectorMenuDesc", "Select By Simulation Target");
+		TSharedPtr<FNiagaraSchemaAction_NewNode> Action = AddNewNodeAction(NewActions, UtilMenuCat, SimTargetSelectorMenuDesc, TEXT("Select By Simulation Target"), FText::GetEmpty());
+		UNiagaraNodeSimTargetSelector* Node = NewObject<UNiagaraNodeSimTargetSelector>(OwnerOfTemporaries);
+		Action->NodeTemplate = Node;
+	}
+
 	return NewActions;
 }
 
@@ -931,63 +938,74 @@ const FPinConnectionResponse UEdGraphSchema_Niagara::CanCreateConnection(const U
 		return FPinConnectionResponse(CONNECT_RESPONSE_DISALLOW, TEXT("Directions are not compatible"));
 	}
 
-	// Check for compatible type pins.
-	if (PinA->PinType.PinCategory == PinCategoryType && 
-		PinB->PinType.PinCategory == PinCategoryType && 
-		PinA->PinType != PinB->PinType)
+	if (PinA->PinType.PinCategory != TEXT("wildcard") && PinB->PinType.PinCategory != TEXT("wildcard"))
 	{
-		FNiagaraTypeDefinition PinTypeA = PinToTypeDefinition(PinA);
-		FNiagaraTypeDefinition PinTypeB = PinToTypeDefinition(PinB);
-		if (FNiagaraTypeDefinition::TypesAreAssignable(PinTypeA, PinTypeB) == false)
+		// Check for compatible type pins.
+		if (PinA->PinType.PinCategory == PinCategoryType &&
+			PinB->PinType.PinCategory == PinCategoryType &&
+			PinA->PinType != PinB->PinType)
 		{
-			//Do some limiting on auto conversions here?
-			if (PinTypeA.GetClass())
+			FNiagaraTypeDefinition PinTypeA = PinToTypeDefinition(PinA);
+			FNiagaraTypeDefinition PinTypeB = PinToTypeDefinition(PinB);
+
+			if (PinTypeA == FNiagaraTypeDefinition::GetParameterMapDef() || PinTypeB == FNiagaraTypeDefinition::GetParameterMapDef())
 			{
 				return FPinConnectionResponse(CONNECT_RESPONSE_DISALLOW, TEXT("Types are not compatible"));
 			}
-			else
+
+			else if (FNiagaraTypeDefinition::TypesAreAssignable(PinTypeA, PinTypeB) == false)
 			{
-				return FPinConnectionResponse(CONNECT_RESPONSE_MAKE_WITH_CONVERSION_NODE, FString::Printf(TEXT("Convert %s to %s"), *(PinToTypeDefinition(PinA).GetNameText().ToString()), *(PinToTypeDefinition(PinB).GetNameText().ToString())));
+				//Do some limiting on auto conversions here?
+				if (PinTypeA.GetClass())
+				{
+					return FPinConnectionResponse(CONNECT_RESPONSE_DISALLOW, TEXT("Types are not compatible"));
+				}
+				else
+				{
+					return FPinConnectionResponse(CONNECT_RESPONSE_MAKE_WITH_CONVERSION_NODE, FString::Printf(TEXT("Convert %s to %s"), *(PinToTypeDefinition(PinA).GetNameText().ToString()), *(PinToTypeDefinition(PinB).GetNameText().ToString())));
+				}
 			}
 		}
-	}
 
-	// Check for compatible misc pins
-	if (PinA->PinType.PinCategory == PinCategoryMisc ||
-		PinB->PinType.PinCategory == PinCategoryMisc) 
-	{
-		// TODO: This shouldn't be handled explicitly here.
-		bool PinAIsConvertAddAndPinBIsNonGenericType =
-			PinA->PinType.PinCategory == PinCategoryMisc && PinA->PinType.PinSubCategory == UNiagaraNodeWithDynamicPins::AddPinSubCategory &&
-			PinB->PinType.PinCategory == PinCategoryType && PinToTypeDefinition(PinB) != FNiagaraTypeDefinition::GetGenericNumericDef();
-
-		bool PinBIsConvertAddAndPinAIsNonGenericType =
-			PinB->PinType.PinCategory == PinCategoryMisc && PinB->PinType.PinSubCategory == UNiagaraNodeWithDynamicPins::AddPinSubCategory &&
-			PinA->PinType.PinCategory == PinCategoryType && PinToTypeDefinition(PinA) != FNiagaraTypeDefinition::GetGenericNumericDef();
-
-		if(PinAIsConvertAddAndPinBIsNonGenericType == false && PinBIsConvertAddAndPinAIsNonGenericType == false)
+		// Check for compatible misc pins
+		if (PinA->PinType.PinCategory == PinCategoryMisc ||
+			PinB->PinType.PinCategory == PinCategoryMisc)
 		{
-			return FPinConnectionResponse(CONNECT_RESPONSE_DISALLOW, TEXT("Types are not compatible"));
+			// TODO: This shouldn't be handled explicitly here.
+			bool PinAIsConvertAddAndPinBIsNonGenericType =
+				PinA->PinType.PinCategory == PinCategoryMisc && PinA->PinType.PinSubCategory == UNiagaraNodeWithDynamicPins::AddPinSubCategory &&
+				PinB->PinType.PinCategory == PinCategoryType && PinToTypeDefinition(PinB) != FNiagaraTypeDefinition::GetGenericNumericDef() &&
+				PinToTypeDefinition(PinB) != FNiagaraTypeDefinition::GetParameterMapDef();
+
+			bool PinBIsConvertAddAndPinAIsNonGenericType =
+				PinB->PinType.PinCategory == PinCategoryMisc && PinB->PinType.PinSubCategory == UNiagaraNodeWithDynamicPins::AddPinSubCategory &&
+				PinA->PinType.PinCategory == PinCategoryType && PinToTypeDefinition(PinA) != FNiagaraTypeDefinition::GetGenericNumericDef() &&
+				PinToTypeDefinition(PinA) != FNiagaraTypeDefinition::GetParameterMapDef();
+
+			if (PinAIsConvertAddAndPinBIsNonGenericType == false && PinBIsConvertAddAndPinAIsNonGenericType == false)
+			{
+				return FPinConnectionResponse(CONNECT_RESPONSE_DISALLOW, TEXT("Types are not compatible"));
+			}
 		}
-	}
 
-	if (PinA->PinType.PinCategory == PinCategoryClass || PinB->PinType.PinCategory == PinCategoryClass)
-	{
-		FNiagaraTypeDefinition AType = PinToTypeDefinition(PinA);
-		FNiagaraTypeDefinition BType = PinToTypeDefinition(PinB);
-		if (AType != BType)
+		if (PinA->PinType.PinCategory == PinCategoryClass || PinB->PinType.PinCategory == PinCategoryClass)
 		{
-			return FPinConnectionResponse(CONNECT_RESPONSE_DISALLOW, TEXT("Types are not compatible"));
+			FNiagaraTypeDefinition AType = PinToTypeDefinition(PinA);
+			FNiagaraTypeDefinition BType = PinToTypeDefinition(PinB);
+			if (AType != BType)
+			{
+				return FPinConnectionResponse(CONNECT_RESPONSE_DISALLOW, TEXT("Types are not compatible"));
+			}
 		}
-	}
 
-	if (PinA->PinType.PinCategory == PinCategoryEnum || PinB->PinType.PinCategory == PinCategoryEnum)
-	{
-		FNiagaraTypeDefinition PinTypeA = PinToTypeDefinition(PinA);
-		FNiagaraTypeDefinition PinTypeB = PinToTypeDefinition(PinB);
-		if (FNiagaraTypeDefinition::TypesAreAssignable(PinTypeA, PinTypeB) == false)
+		if (PinA->PinType.PinCategory == PinCategoryEnum || PinB->PinType.PinCategory == PinCategoryEnum)
 		{
-			return FPinConnectionResponse(CONNECT_RESPONSE_DISALLOW, TEXT("Types are not compatible"));
+			FNiagaraTypeDefinition PinTypeA = PinToTypeDefinition(PinA);
+			FNiagaraTypeDefinition PinTypeB = PinToTypeDefinition(PinB);
+			if (FNiagaraTypeDefinition::TypesAreAssignable(PinTypeA, PinTypeB) == false)
+			{
+				return FPinConnectionResponse(CONNECT_RESPONSE_DISALLOW, TEXT("Types are not compatible"));
+			}
 		}
 	}
 
@@ -1270,7 +1288,11 @@ bool UEdGraphSchema_Niagara::TryGetPinDefaultValueFromNiagaraVariable(const FNia
 
 FNiagaraTypeDefinition UEdGraphSchema_Niagara::PinToTypeDefinition(const UEdGraphPin* Pin) const
 {
-	if (Pin->PinType.PinCategory == PinCategoryType && Pin->PinType.PinSubCategoryObject != nullptr)
+	if (Pin == nullptr)
+	{
+		return FNiagaraTypeDefinition();
+	}
+	if (Pin->PinType.PinCategory == PinCategoryType && Pin->PinType.PinSubCategoryObject.IsValid())
 	{
 		UScriptStruct* Struct = Cast<UScriptStruct>(Pin->PinType.PinSubCategoryObject.Get());
 		if (Struct == nullptr)
@@ -1554,6 +1576,19 @@ void UEdGraphSchema_Niagara::ToggleNodeEnabledState(UNiagaraNode* InNode) const
 	}
 }
 
+void UEdGraphSchema_Niagara::RefreshNode(UNiagaraNode* InNode) const
+{
+	if (InNode != nullptr)
+	{
+		const FScopedTransaction Transaction(NSLOCTEXT("UnrealEd", "NiagaraEditorRefreshNode", "Refresh Node"));
+		InNode->Modify();
+		if (InNode->RefreshFromExternalChanges())
+		{
+			InNode->MarkNodeRequiresSynchronization(__FUNCTION__, true);
+		}
+	}
+}
+
 bool UEdGraphSchema_Niagara::CanPromoteSinglePinToParameter(const UEdGraphPin* SourcePin) 
 {
 	const UNiagaraGraph* NiagaraGraph = Cast<UNiagaraGraph>(SourcePin->GetOwningNode()->GetGraph());
@@ -1648,6 +1683,9 @@ void UEdGraphSchema_Niagara::GetContextMenuActions(const UEdGraph* CurrentGraph,
 		MenuBuilder->BeginSection("EdGraphSchema_NiagaraNodeActions", LOCTEXT("NodeActionsMenuHeader", "Node Actions"));
 		MenuBuilder->AddMenuEntry(LOCTEXT("ToggleEnabledState", "Toggle Enabled State"), LOCTEXT("ToggleEnabledStateTooltip", "Toggle this node between Enbled (default) and Disabled (skipped from compilation)."), FSlateIcon(),
 			FUIAction(FExecuteAction::CreateUObject((UEdGraphSchema_Niagara*const)this, &UEdGraphSchema_Niagara::ToggleNodeEnabledState, const_cast<UNiagaraNode*>(Node))));
+		MenuBuilder->AddMenuEntry(LOCTEXT("RefreshNode", "Refresh Node"), LOCTEXT("RefreshNodeTooltip", "Refresh this node."), FSlateIcon(),
+			FUIAction(FExecuteAction::CreateUObject((UEdGraphSchema_Niagara*const)this, &UEdGraphSchema_Niagara::RefreshNode, const_cast<UNiagaraNode*>(Node))));
+
 		MenuBuilder->EndSection();
 	}
 
@@ -1656,7 +1694,7 @@ void UEdGraphSchema_Niagara::GetContextMenuActions(const UEdGraph* CurrentGraph,
 
 FNiagaraConnectionDrawingPolicy::FNiagaraConnectionDrawingPolicy(int32 InBackLayerID, int32 InFrontLayerID, float InZoomFactor, const FSlateRect& InClippingRect, FSlateWindowElementList& InDrawElements, UEdGraph* InGraph)
 	: FConnectionDrawingPolicy(InBackLayerID, InFrontLayerID, InZoomFactor, InClippingRect, InDrawElements)
-	, Graph(InGraph)
+	, Graph(CastChecked<UNiagaraGraph>(InGraph))
 {
 	ArrowImage = nullptr;
 	ArrowRadius = FVector2D::ZeroVector;
@@ -1672,8 +1710,20 @@ void FNiagaraConnectionDrawingPolicy::DetermineWiringStyle(UEdGraphPin* OutputPi
 
 	if (Graph)
 	{
-		const UEdGraphSchema* Schema = Graph->GetSchema();
-		Params.WireColor = Schema->GetPinTypeColor(OutputPin->PinType);
+		const UEdGraphSchema_Niagara* NSchema = Cast<UEdGraphSchema_Niagara>(Graph->GetSchema());
+		if (NSchema && OutputPin)
+		{
+			Params.WireColor = NSchema->GetPinTypeColor(OutputPin->PinType);
+			if (NSchema->PinToTypeDefinition(OutputPin) == FNiagaraTypeDefinition::GetGenericNumericDef())
+			{
+				FNiagaraTypeDefinition NewDef = Graph->GetCachedNumericConversion(OutputPin);
+				if (NewDef.IsValid())
+				{
+					FEdGraphPinType NewPinType = NSchema->TypeDefinitionToPinType(NewDef);
+					Params.WireColor = NSchema->GetPinTypeColor(NewPinType);
+				}
+			}
+		}
 	}
 }
 
diff --git a/Engine/Plugins/FX/Niagara/Source/NiagaraEditor/Private/NiagaraCompiler.cpp b/Engine/Plugins/FX/Niagara/Source/NiagaraEditor/Private/NiagaraCompiler.cpp
--- a/Engine/Plugins/FX/Niagara/Source/NiagaraEditor/Private/NiagaraCompiler.cpp
+++ b/Engine/Plugins/FX/Niagara/Source/NiagaraEditor/Private/NiagaraCompiler.cpp
@@ -392,6 +392,31 @@ void FNiagaraCompileRequestData::FinishPrecompile(UNiagaraScriptSource* ScriptSo
 			}
 		}
 
+		// Generate CDO's for data interfaces that are passed in to function or dynamic input scripts compiled standalone as we do not have a history
+		if (InUsage == ENiagaraScriptUsage::Function || InUsage == ENiagaraScriptUsage::DynamicInput)
+		{
+			for (const auto ReferencedGraph : ClonedGraphs)
+			{
+				TArray<UNiagaraNodeInput*> InputNodes;
+				TArray<FNiagaraVariable*> InputVariables;
+				ReferencedGraph->FindInputNodes(InputNodes);
+				for (const auto InputNode : InputNodes)
+				{
+					InputVariables.Add(&InputNode->Input);
+				}
+
+				for (const auto InputVariable : InputVariables)
+				{
+					if (InputVariable->IsDataInterface())
+					{
+						UClass* Class = const_cast<UClass*>(InputVariable->GetType().GetClass());
+						UObject* Obj = DuplicateObject(Class->GetDefaultObject(true), GetTransientPackage());
+						CDOs.Add(Class, Obj);
+					}
+				}
+			}
+		}
+
 	}
 }
 
diff --git a/Engine/Plugins/FX/Niagara/Source/NiagaraEditor/Private/NiagaraEditorCommon.cpp b/Engine/Plugins/FX/Niagara/Source/NiagaraEditor/Private/NiagaraEditorCommon.cpp
--- a/Engine/Plugins/FX/Niagara/Source/NiagaraEditor/Private/NiagaraEditorCommon.cpp
+++ b/Engine/Plugins/FX/Niagara/Source/NiagaraEditor/Private/NiagaraEditorCommon.cpp
@@ -8,6 +8,7 @@
 #include "ActorFactoryNiagara.h"
 #include "NiagaraActor.h"
 #include "NiagaraComponent.h"
+#include "Misc/StringFormatter.h"
 
 #define LOCTEXT_NAMESPACE "NiagaraEditor"
 
@@ -35,6 +36,30 @@ void FNiagaraOpInfo::BuildName(FString InName, FString InCategory)
 	Name = FName(*(InCategory + TEXT("::") + InName));
 }
 
+bool FNiagaraOpInfo::CreateHlslForAddedInputs(int32 InputCount, FString & HlslResult) const
+{
+	if (!bSupportsAddedInputs || AddedInputFormatting.IsEmpty() || InputCount < 2)
+	{
+		return false;
+	}
+
+	FString Result = TEXT("{0}");
+	FString KeyA = TEXT("A");
+	FString KeyB = TEXT("B");
+	TMap<FString, FStringFormatArg> FormatArgs;
+	FormatArgs.Add(KeyA, FStringFormatArg(Result));
+	FormatArgs.Add(KeyB, FStringFormatArg(Result));
+	for (int32 i = 1; i < InputCount; i++)
+	{
+		FormatArgs[KeyB] = FStringFormatArg(TEXT("{") + FString::FromInt(i) + TEXT("}"));
+		Result = FString::Format(*AddedInputFormatting, FormatArgs);
+		FormatArgs[KeyA] = FStringFormatArg(Result);
+	}
+
+	HlslResult = Result;
+	return true;
+}
+
 BEGIN_FUNCTION_BUILD_OPTIMIZATION
 void FNiagaraOpInfo::Init()
 {
@@ -87,6 +112,9 @@ void FNiagaraOpInfo::Init()
 		Op->Inputs.Add(FNiagaraOpInOutInfo(B, Type, BText, BText, DefaultStr_Zero));
 		Op->Outputs.Add(FNiagaraOpInOutInfo(Result, Type, ResultText, ResultText, DefaultStr_Zero, TEXT("{0} + {1}")));
 		Op->BuildName(TEXT("Add"), CategoryName);
+		Op->bSupportsAddedInputs = true;
+		Op->AddedInputTypeRestrictions.Add(Type);
+		Op->AddedInputFormatting = TEXT("{A} + {B}");
 		OpInfoMap.Add(Op->Name) = Idx;
 
 		Idx = OpInfos.AddDefaulted();
@@ -99,6 +127,9 @@ void FNiagaraOpInfo::Init()
 		Op->Inputs.Add(FNiagaraOpInOutInfo(B, Type, BText, BText, DefaultStr_Zero));
 		Op->Outputs.Add(FNiagaraOpInOutInfo(Result, Type, ResultText, ResultText, DefaultStr_Zero, TEXT("{0} - {1}")));		
 		Op->BuildName(TEXT("Subtract"), CategoryName);
+		Op->bSupportsAddedInputs = true;
+		Op->AddedInputTypeRestrictions.Add(Type);
+		Op->AddedInputFormatting = TEXT("{A} - {B}");
 		OpInfoMap.Add(Op->Name) = Idx;
 
 		Idx = OpInfos.AddDefaulted();
@@ -111,6 +142,9 @@ void FNiagaraOpInfo::Init()
 		Op->Inputs.Add(FNiagaraOpInOutInfo(B, Type, BText, BText, DefaultStr_One));
 		Op->Outputs.Add(FNiagaraOpInOutInfo(Result, Type, ResultText, ResultText, DefaultStr_One, TEXT("{0} * {1}")));
 		Op->BuildName(TEXT("Mul"), CategoryName);
+		Op->bSupportsAddedInputs = true;
+		Op->AddedInputTypeRestrictions.Add(Type);
+		Op->AddedInputFormatting = TEXT("{A} * {B}");
 		OpInfoMap.Add(Op->Name) = Idx;
 
 		Idx = OpInfos.AddDefaulted();
@@ -153,13 +187,24 @@ void FNiagaraOpInfo::Init()
 		Idx = OpInfos.AddDefaulted();
 		Op = &OpInfos[Idx];
 		Op->Category = CategoryText;
-		Op->FriendlyName = NSLOCTEXT("NiagaraOpInfo", "Reciprocal Name", "Reciprocal");
-		Op->Description = NSLOCTEXT("NiagaraOpInfo", "Reciprocal Desc", "Result = 1 / A");
+		Op->FriendlyName = NSLOCTEXT("NiagaraOpInfo", "Reciprocal Fast Name", "Reciprocal Fast");
+		Op->Description = NSLOCTEXT("NiagaraOpInfo", "Reciprocal Fast Desc", "12-bits of accuracy, but faster. Result = 1 / A using Newton/Raphson approximation.");
 		Op->Inputs.Add(FNiagaraOpInOutInfo(A, Type, AText, AText, DefaultStr_One));
 		Op->Outputs.Add(FNiagaraOpInOutInfo(Result, Type, ResultText, ResultText, DefaultStr_One, TEXT("rcp({0})")));
+		Op->BuildName(TEXT("RcpFast"), CategoryName);
+		OpInfoMap.Add(Op->Name) = Idx;
+
+		Idx = OpInfos.AddDefaulted();
+		Op = &OpInfos[Idx];
+		Op->Category = CategoryText;
+		Op->FriendlyName = NSLOCTEXT("NiagaraOpInfo", "Reciprocal Name", "Reciprocal");
+		Op->Description = NSLOCTEXT("NiagaraOpInfo", "Reciprocal Desc", "More accurate than Reciprocal Fast. Result = 1 / A");
+		Op->Inputs.Add(FNiagaraOpInOutInfo(A, Type, AText, AText, DefaultStr_One));
+		Op->Outputs.Add(FNiagaraOpInOutInfo(Result, Type, ResultText, ResultText, DefaultStr_One, TEXT("Reciprocal({0})")));
 		Op->BuildName(TEXT("Rcp"), CategoryName);
 		OpInfoMap.Add(Op->Name) = Idx;
 
+
 		Idx = OpInfos.AddDefaulted();
 		Op = &OpInfos[Idx];
 		Op->Category = CategoryText;
@@ -574,15 +619,27 @@ void FNiagaraOpInfo::Init()
 		Idx = OpInfos.AddDefaulted();
 		Op = &OpInfos[Idx];
 		Op->Category = CategoryText;
-		Op->FriendlyName = NSLOCTEXT("NiagaraOpInfo", "Fmod Name", "FMod");
+		Op->FriendlyName = NSLOCTEXT("NiagaraOpInfo", "Fmod Name", "Modulo");
 		Op->Description = NSLOCTEXT("NiagaraOpInfo", "Fmod Desc", "Result = A % B");
 		Op->Keywords = FText::FromString(TEXT("%"));
 		Op->Inputs.Add(FNiagaraOpInOutInfo(A, Type, AText, AText, DefaultStr_One));
 		Op->Inputs.Add(FNiagaraOpInOutInfo(B, Type, BText, BText, DefaultStr_One));
-		Op->Outputs.Add(FNiagaraOpInOutInfo(Result, Type, ResultText, ResultText, DefaultStr_One, TEXT("Modulo({0}, {1})")));
+		Op->Outputs.Add(FNiagaraOpInOutInfo(Result, Type, ResultText, ResultText, DefaultStr_One, TEXT("ModuloPrecise({0}, {1})")));
 		Op->BuildName(TEXT("FMod"), CategoryName);
 		OpInfoMap.Add(Op->Name) = Idx;
 
+		Idx = OpInfos.AddDefaulted();
+		Op = &OpInfos[Idx];
+		Op->Category = CategoryText;
+		Op->FriendlyName = NSLOCTEXT("NiagaraOpInfo", "Fmod Name Fast", "Modulo Fast");
+		Op->Description = NSLOCTEXT("NiagaraOpInfo", "Fmod Desc Fast", "Result = A % B. May be less precise than regular FMod.");
+		Op->Keywords = FText::FromString(TEXT("%"));
+		Op->Inputs.Add(FNiagaraOpInOutInfo(A, Type, AText, AText, DefaultStr_One));
+		Op->Inputs.Add(FNiagaraOpInOutInfo(B, Type, BText, BText, DefaultStr_One));
+		Op->Outputs.Add(FNiagaraOpInOutInfo(Result, Type, ResultText, ResultText, DefaultStr_One, TEXT("Modulo({0}, {1})")));
+		Op->BuildName(TEXT("FModFast"), CategoryName);
+		OpInfoMap.Add(Op->Name) = Idx;
+
 		Idx = OpInfos.AddDefaulted();
 		Op = &OpInfos[Idx];
 		Op->Category = CategoryText;
@@ -624,6 +681,9 @@ void FNiagaraOpInfo::Init()
 		Op->Inputs.Add(FNiagaraOpInOutInfo(B, Type, BText, BText, DefaultStr_One));
 		Op->Outputs.Add(FNiagaraOpInOutInfo(Result, Type, ResultText, ResultText, DefaultStr_One, TEXT("min({0},{1})")));
 		Op->BuildName(TEXT("Min"), CategoryName);
+		Op->bSupportsAddedInputs = true;
+		Op->AddedInputTypeRestrictions.Add(Type);
+		Op->AddedInputFormatting = TEXT("min({A}, {B})");
 		OpInfoMap.Add(Op->Name) = Idx;
 
 		Idx = OpInfos.AddDefaulted();
@@ -635,6 +695,9 @@ void FNiagaraOpInfo::Init()
 		Op->Inputs.Add(FNiagaraOpInOutInfo(B, Type, BText, BText, DefaultStr_One));
 		Op->Outputs.Add(FNiagaraOpInOutInfo(Result, Type, ResultText, ResultText, DefaultStr_One, TEXT("max({0},{1})")));
 		Op->BuildName(TEXT("Max"), CategoryName);
+		Op->bSupportsAddedInputs = true;
+		Op->AddedInputTypeRestrictions.Add(Type);
+		Op->AddedInputFormatting = TEXT("max({A}, {B})");
 		OpInfoMap.Add(Op->Name) = Idx;
 
 		Idx = OpInfos.AddDefaulted();
@@ -712,17 +775,41 @@ void FNiagaraOpInfo::Init()
 		Op->BuildName(TEXT("Length"), CategoryName);
 		OpInfoMap.Add(Op->Name) = Idx;
 
-		//Temporarily here. Rand will be reworked shortly.
+		// Non-deterministic random number generation. Calls FRandomStream on the CPU. 
 		Idx = OpInfos.AddDefaulted();
 		Op = &OpInfos[Idx];
 		Op->Category = CategoryText;
 		Op->FriendlyName = NSLOCTEXT("NiagaraOpInfo", "Rand Name", "Random");
-		Op->Description = NSLOCTEXT("NiagaraOpInfo", "Rand Desc", "Returns a random value between 0 and A.");
+		Op->Description = NSLOCTEXT("NiagaraOpInfo", "Rand Desc", "Returns a non-deterministic random value between 0 and A.");
 		Op->Inputs.Add(FNiagaraOpInOutInfo(A, Type, AText, AText, DefaultStr_One));
 		Op->Outputs.Add(FNiagaraOpInOutInfo(Result, Type, ResultText, ResultText, DefaultStr_One, TEXT("rand({0})")));
 		Op->BuildName(TEXT("Rand"), CategoryName);
 		OpInfoMap.Add(Op->Name) = Idx;
 
+
+		// Deterministic/seeded random number generation. 
+		static FName SeedName1(TEXT("Seed 1"));
+		static FName SeedName2(TEXT("Seed 2"));
+		static FName SeedName3(TEXT("Seed 3"));
+		static FText SeedText1 = NSLOCTEXT("NiagaraOpInfo", "Seed1 Desc", "Seed 1");
+		static FText SeedText2 = NSLOCTEXT("NiagaraOpInfo", "Seed2 Desc", "Seed 2");
+		static FText SeedText3 = NSLOCTEXT("NiagaraOpInfo", "Seed3 Desc", "Seed 3");
+		FString DefaultSeed_Zero(TEXT("0"));
+		FNiagaraTypeDefinition SeedType = FNiagaraTypeDefinition::GetIntDef();
+
+		Idx = OpInfos.AddDefaulted();
+		Op = &OpInfos[Idx];
+		Op->Category = CategoryText;
+		Op->FriendlyName = NSLOCTEXT("NiagaraOpInfo", "Seeded Rand Name", "Seeded Random");
+		Op->Description = NSLOCTEXT("NiagaraOpInfo", "Seeded Rand Desc", "Returns a deterministic random value between 0 and A.");
+		Op->Inputs.Add(FNiagaraOpInOutInfo(A, Type, AText, AText, DefaultStr_One));
+		Op->Inputs.Add(FNiagaraOpInOutInfo(SeedName1, SeedType, SeedText1, SeedText1, DefaultSeed_Zero));
+		Op->Inputs.Add(FNiagaraOpInOutInfo(SeedName2, SeedType, SeedText2, SeedText2, DefaultSeed_Zero));
+		Op->Inputs.Add(FNiagaraOpInOutInfo(SeedName3, SeedType, SeedText3, SeedText3, DefaultSeed_Zero));
+		Op->Outputs.Add(FNiagaraOpInOutInfo(Result, Type, ResultText, ResultText, DefaultStr_One, TEXT("rand({0}, {1}, {2}, {3})")));
+		Op->BuildName(TEXT("SeededRand"), CategoryName);
+		OpInfoMap.Add(Op->Name) = Idx;
+
 		//Comparison ops
 		Idx = OpInfos.AddDefaulted();
 		Op = &OpInfos[Idx];
@@ -816,6 +903,9 @@ void FNiagaraOpInfo::Init()
 	Op->Inputs.Add(FNiagaraOpInOutInfo(B, IntType, BText, BText, Default_IntOne));
 	Op->Outputs.Add(FNiagaraOpInOutInfo(Result, IntType, ResultText, ResultText, Default_IntOne, TEXT("{0} & {1}")));
 	Op->BuildName(TEXT("BitAnd"), IntCategoryName);
+	Op->bSupportsAddedInputs = true;
+	Op->AddedInputTypeRestrictions.Add(IntType);
+	Op->AddedInputFormatting = TEXT("{A} & {B}");
 	OpInfoMap.Add(Op->Name) = Idx;
 
 	Idx = OpInfos.AddDefaulted();
@@ -828,6 +918,9 @@ void FNiagaraOpInfo::Init()
 	Op->Inputs.Add(FNiagaraOpInOutInfo(B, IntType, BText, BText, Default_IntOne));
 	Op->Outputs.Add(FNiagaraOpInOutInfo(Result, IntType, ResultText, ResultText, Default_IntOne, TEXT("{0} | {1}")));
 	Op->BuildName(TEXT("BitOr"), IntCategoryName);
+	Op->bSupportsAddedInputs = true;
+	Op->AddedInputTypeRestrictions.Add(IntType);
+	Op->AddedInputFormatting = TEXT("{A} | {B}");
 	OpInfoMap.Add(Op->Name) = Idx;
 
 	Idx = OpInfos.AddDefaulted();
@@ -840,6 +933,9 @@ void FNiagaraOpInfo::Init()
 	Op->Inputs.Add(FNiagaraOpInOutInfo(B, IntType, BText, BText, Default_IntOne));
 	Op->Outputs.Add(FNiagaraOpInOutInfo(Result, IntType, ResultText, ResultText, Default_IntOne, TEXT("{0} ^ {1}")));
 	Op->BuildName(TEXT("BitXOr"), IntCategoryName);
+	Op->bSupportsAddedInputs = true;
+	Op->AddedInputTypeRestrictions.Add(IntType);
+	Op->AddedInputFormatting = TEXT("{A} ^ {B}");
 	OpInfoMap.Add(Op->Name) = Idx;
 
 	Idx = OpInfos.AddDefaulted();
@@ -853,6 +949,35 @@ void FNiagaraOpInfo::Init()
 	Op->BuildName(TEXT("BitNot"), IntCategoryName);
 	OpInfoMap.Add(Op->Name) = Idx;
 
+	Idx = OpInfos.AddDefaulted();
+	Op = &OpInfos[Idx];
+	Op->Category = IntCategory;
+	Op->FriendlyName = NSLOCTEXT("NiagaraOpInfo", "BitLShift Name", "Bitwise Left Shift");
+	Op->Description = NSLOCTEXT("NiagaraOpInfo", "BitLShift Desc", "Shifts A left by B bits, padding with zeroes on the right. B should be between 0 and 31 or there will be undefined behavior.");
+	Op->Keywords = FText::FromString(TEXT("<<"));
+	Op->Inputs.Add(FNiagaraOpInOutInfo(A, IntType, AText, AText, Default_IntOne));
+	Op->Inputs.Add(FNiagaraOpInOutInfo(B, IntType, BText, BText, Default_IntOne));
+	Op->Outputs.Add(FNiagaraOpInOutInfo(Result, IntType, ResultText, ResultText, Default_IntOne, TEXT("{0} << {1}")));
+	Op->BuildName(TEXT("BitLShift"), IntCategoryName);
+	Op->bSupportsAddedInputs = true;
+	Op->AddedInputTypeRestrictions.Add(IntType);
+	Op->AddedInputFormatting = TEXT("{A} << {B}");
+	OpInfoMap.Add(Op->Name) = Idx;
+
+	Idx = OpInfos.AddDefaulted();
+	Op = &OpInfos[Idx];
+	Op->Category = IntCategory;
+	Op->FriendlyName = NSLOCTEXT("NiagaraOpInfo", "BitRShift Name", "Bitwise Right Shift");
+	Op->Description = NSLOCTEXT("NiagaraOpInfo", "BitRShift Desc", "Shifts A right by B bits, taking the sign bit and propagating it to fill in on left (i.e. negative numbers fill with 1's, positive fill with 0's. B should be between 0 and 31 or there will be undefined behavior.");
+	Op->Keywords = FText::FromString(TEXT(">>"));
+	Op->Inputs.Add(FNiagaraOpInOutInfo(A, IntType, AText, AText, Default_IntOne));
+	Op->Inputs.Add(FNiagaraOpInOutInfo(B, IntType, BText, BText, Default_IntOne));
+	Op->Outputs.Add(FNiagaraOpInOutInfo(Result, IntType, ResultText, ResultText, Default_IntOne, TEXT("{0} >> {1}")));
+	Op->BuildName(TEXT("BitRShift"), IntCategoryName);
+	Op->bSupportsAddedInputs = true;
+	Op->AddedInputTypeRestrictions.Add(IntType);
+	Op->AddedInputFormatting = TEXT("{A} >> {B}");
+	OpInfoMap.Add(Op->Name) = Idx;
 	//////////////////////////////////////////////////////////////////////////
 	// Boolean Only Ops
 
@@ -872,6 +997,9 @@ void FNiagaraOpInfo::Init()
 	Op->Inputs.Add(FNiagaraOpInOutInfo(B, BoolType, BText, BText, Default_BoolOne));
 	Op->Outputs.Add(FNiagaraOpInOutInfo(Result, BoolType, ResultText, ResultText, Default_BoolOne, TEXT("{0} && {1}")));
 	Op->BuildName(TEXT("LogicAnd"), BoolCategoryName);
+	Op->bSupportsAddedInputs = true;
+	Op->AddedInputTypeRestrictions.Add(BoolType);
+	Op->AddedInputFormatting = TEXT("{A} && {B}");
 	OpInfoMap.Add(Op->Name) = Idx;
 
 	Idx = OpInfos.AddDefaulted();
@@ -884,6 +1012,9 @@ void FNiagaraOpInfo::Init()
 	Op->Inputs.Add(FNiagaraOpInOutInfo(B, BoolType, BText, BText, Default_BoolOne));
 	Op->Outputs.Add(FNiagaraOpInOutInfo(Result, BoolType, ResultText, ResultText, Default_BoolOne, TEXT("{0} || {1}")));
 	Op->BuildName(TEXT("LogicOr"), BoolCategoryName);
+	Op->bSupportsAddedInputs = true;
+	Op->AddedInputTypeRestrictions.Add(BoolType);
+	Op->AddedInputFormatting = TEXT("{A} || {B}");
 	OpInfoMap.Add(Op->Name) = Idx;
 
 	Idx = OpInfos.AddDefaulted();
@@ -986,6 +1117,9 @@ void FNiagaraOpInfo::Init()
 	Op->Inputs.Add(FNiagaraOpInOutInfo(B, MatrixType, BText, BText, Default_MatrixOne));
 	Op->Outputs.Add(FNiagaraOpInOutInfo(Result, MatrixType, ResultText, ResultText, Default_MatrixOne, TEXT("{0} * {1}")));
 	Op->BuildName(TEXT("MatrixMultiply"), MatrixCategoryName);
+	Op->bSupportsAddedInputs = true;
+	Op->AddedInputTypeRestrictions.Add(MatrixType);
+	Op->AddedInputFormatting = TEXT("{A} * {B}");
 	OpInfoMap.Add(Op->Name) = Idx;
 	
 	Idx = OpInfos.AddDefaulted();
diff --git a/Engine/Plugins/FX/Niagara/Source/NiagaraEditor/Private/NiagaraEditorModule.cpp b/Engine/Plugins/FX/Niagara/Source/NiagaraEditor/Private/NiagaraEditorModule.cpp
--- a/Engine/Plugins/FX/Niagara/Source/NiagaraEditor/Private/NiagaraEditorModule.cpp
+++ b/Engine/Plugins/FX/Niagara/Source/NiagaraEditor/Private/NiagaraEditorModule.cpp
@@ -132,11 +132,26 @@ public:
 		{
 			if (InPin->PinType.PinCategory == UEdGraphSchema_Niagara::PinCategoryType)
 			{
-				const UScriptStruct* Struct = CastChecked<const UScriptStruct>(InPin->PinType.PinSubCategoryObject.Get());
-				const FCreateGraphPin* CreateGraphPin = TypeToCreatePinDelegateMap.Find(Struct);
-				if (CreateGraphPin != nullptr)
+				if (InPin->PinType.PinSubCategoryObject != nullptr && InPin->PinType.PinSubCategoryObject->IsA<UScriptStruct>())
 				{
-					return (*CreateGraphPin).Execute(InPin);
+					const UScriptStruct* Struct = CastChecked<const UScriptStruct>(InPin->PinType.PinSubCategoryObject.Get());
+					const FCreateGraphPin* CreateGraphPin = TypeToCreatePinDelegateMap.Find(Struct);
+					if (CreateGraphPin != nullptr)
+					{
+						return (*CreateGraphPin).Execute(InPin);
+					}
+					// Otherwise, fall back to the generic pin for Niagara types. Previous iterations put out an error here, but this 
+					// was not correct as the above list is just overrides from the default renamable pin, usually numeric types with their own custom 
+					// editors for default values. Things like the parameter map can safely just fall through to the end condition and create a
+					// generic renamable pin.
+				}
+				else
+				{
+					UE_LOG(LogNiagaraEditor, Error, TEXT("Pin type is invalid! Pin Name '%s' Owning Node '%s'. Turning into standard int definition!"), *InPin->PinName.ToString(),
+						*InPin->GetOwningNode()->GetName());
+					InPin->PinType.PinSubCategoryObject = MakeWeakObjectPtr(const_cast<UScriptStruct*>(FNiagaraTypeDefinition::GetIntStruct()));
+					InPin->DefaultValue.Empty();
+					return CreatePin(InPin);
 				}
 			}
 			else if (InPin->PinType.PinCategory == UEdGraphSchema_Niagara::PinCategoryEnum)
diff --git a/Engine/Plugins/FX/Niagara/Source/NiagaraEditor/Private/NiagaraEditorUtilities.cpp b/Engine/Plugins/FX/Niagara/Source/NiagaraEditor/Private/NiagaraEditorUtilities.cpp
--- a/Engine/Plugins/FX/Niagara/Source/NiagaraEditor/Private/NiagaraEditorUtilities.cpp
+++ b/Engine/Plugins/FX/Niagara/Source/NiagaraEditor/Private/NiagaraEditorUtilities.cpp
@@ -89,7 +89,10 @@ void FNiagaraEditorUtilities::InitializeParameterInputNode(UNiagaraNodeInput& In
 	if (Type.GetScriptStruct() != nullptr)
 	{
 		ResetVariableToDefaultValue(InputNode.Input);
-		InputNode.SetDataInterface(nullptr);
+		if (InputNode.GetDataInterface() != nullptr)
+		{
+			InputNode.SetDataInterface(nullptr);
+		}
 	}
 	else
 	{
@@ -608,56 +611,7 @@ void TraverseGraphFromOutputDepthFirst(const UEdGraphSchema_Niagara* Schema, UNi
 
 void FixUpNumericPinsVisitor(const UEdGraphSchema_Niagara* Schema, UNiagaraNode* Node)
 {
-	// Fix up numeric input pins and keep track of numeric types to decide the output type.
-	TArray<FNiagaraTypeDefinition> InputTypes;
-	TArray<UEdGraphPin*> InputPins;
-	Node->GetInputPins(InputPins);
-	for (UEdGraphPin* InputPin : InputPins)
-	{
-		if (InputPin->PinType.PinCategory == UEdGraphSchema_Niagara::PinCategoryType)
-		{
-			FNiagaraTypeDefinition InputPinType = Schema->PinToTypeDefinition(InputPin);
-
-			// If the input pin is the generic numeric type set it to the type of the linked output pin which should have been processed already.
-			if (InputPinType == FNiagaraTypeDefinition::GetGenericNumericDef() && InputPin->LinkedTo.Num() == 1)
-			{
-				UEdGraphPin* InputPinLinkedPin = UNiagaraNode::TraceOutputPin(InputPin->LinkedTo[0]);
-				FNiagaraTypeDefinition InputPinLinkedPinType = Schema->PinToTypeDefinition(InputPinLinkedPin);
-				if (InputPinLinkedPinType.IsValid())
-				{
-					// Only update the input pin type if the linked pin type is valid.
-					InputPin->PinType = Schema->TypeDefinitionToPinType(InputPinLinkedPinType);
-					InputPinType = InputPinLinkedPinType;
-				}
-			}
-
-			if (InputPinType == FNiagaraTypeDefinition::GetGenericNumericDef())
-			{
-				UE_LOG(LogNiagaraEditor, Error, TEXT("Unable to deduce type for numeric input pin. Node: %s Pin: %s"), *Node->GetName(), *InputPin->GetName());
-			}
-
-			InputTypes.Add(InputPinType);
-		}
-	}
-
-	// Fix up numeric outputs based on the inputs.
-	if (InputTypes.Num() > 0 && Node->GetNumericOutputTypeSelectionMode() != ENiagaraNumericOutputTypeSelectionMode::None)
-	{
-		FNiagaraTypeDefinition OutputNumericType = FNiagaraTypeDefinition::GetNumericOutputType(InputTypes, Node->GetNumericOutputTypeSelectionMode());
-		if (OutputNumericType != FNiagaraTypeDefinition::GetGenericNumericDef())
-		{
-			TArray<UEdGraphPin*> OutputPins;
-			Node->GetOutputPins(OutputPins);
-			for (UEdGraphPin* OutputPin : OutputPins)
-			{
-				FNiagaraTypeDefinition OutputPinType = Schema->PinToTypeDefinition(OutputPin);
-				if (OutputPinType == FNiagaraTypeDefinition::GetGenericNumericDef())
-				{
-					OutputPin->PinType = Schema->TypeDefinitionToPinType(OutputNumericType);
-				}
-			}
-		}
-	}
+	Node->ResolveNumerics(Schema, true, nullptr);
 }
 
 void FNiagaraEditorUtilities::FixUpNumericPins(const UEdGraphSchema_Niagara* Schema, UNiagaraNode* Node)
diff --git a/Engine/Plugins/FX/Niagara/Source/NiagaraEditor/Private/NiagaraEmitterFactoryNew.cpp b/Engine/Plugins/FX/Niagara/Source/NiagaraEditor/Private/NiagaraEmitterFactoryNew.cpp
--- a/Engine/Plugins/FX/Niagara/Source/NiagaraEditor/Private/NiagaraEmitterFactoryNew.cpp
+++ b/Engine/Plugins/FX/Niagara/Source/NiagaraEditor/Private/NiagaraEmitterFactoryNew.cpp
@@ -82,7 +82,7 @@ bool UNiagaraEmitterFactoryNew::ConfigureProperties()
 
 UNiagaraNodeFunctionCall* AddModuleFromAssetPath(FString AssetPath, UNiagaraNodeOutput& TargetOutputNode)
 {
-	FStringAssetReference AssetRef(AssetPath);
+	FSoftObjectPath AssetRef(AssetPath);
 	UNiagaraScript* AssetScript = Cast<UNiagaraScript>(AssetRef.TryLoad());
 	FAssetData ScriptAssetData(AssetScript);
 	if (ScriptAssetData.IsValid())
@@ -196,6 +196,7 @@ UObject* UNiagaraEmitterFactoryNew::FactoryCreateNew(UClass* Class, UObject* InP
 
 		FNiagaraStackGraphUtilities::RelayoutGraph(*Source->NodeGraph);
 		NewEmitter->bInterpolatedSpawning = true;
+		NewEmitter->bDeterminism = false; // NOTE: Default to non-determinism
 		NewEmitter->SpawnScriptProps.Script->SetUsage(ENiagaraScriptUsage::ParticleSpawnScriptInterpolated);
 	}
 	
diff --git a/Engine/Plugins/FX/Niagara/Source/NiagaraEditor/Private/NiagaraGraph.cpp b/Engine/Plugins/FX/Niagara/Source/NiagaraEditor/Private/NiagaraGraph.cpp
--- a/Engine/Plugins/FX/Niagara/Source/NiagaraEditor/Private/NiagaraGraph.cpp
+++ b/Engine/Plugins/FX/Niagara/Source/NiagaraEditor/Private/NiagaraGraph.cpp
@@ -59,6 +59,7 @@ FNiagaraGraphScriptUsageInfo::FNiagaraGraphScriptUsageInfo() : UsageType(ENiagar
 
 UNiagaraGraph::UNiagaraGraph(const FObjectInitializer& ObjectInitializer)
 	: Super(ObjectInitializer)
+	, bNeedNumericCacheRebuilt(true)
 	, bFindParametersAllowed(true)
 	, bIsRenamingParameter(false)
 {
@@ -96,6 +97,7 @@ void UNiagaraGraph::NotifyGraphChanged()
 {
 	FindParameters();
 	Super::NotifyGraphChanged();
+	InvalidateNumericCache();
 }
 
 void UNiagaraGraph::PostLoad()
@@ -724,6 +726,25 @@ void UNiagaraGraph::SubsumeExternalDependencies(TMap<const UObject*, UObject*>&
 	}
 }
 
+FNiagaraTypeDefinition UNiagaraGraph::GetCachedNumericConversion(class UEdGraphPin* InPin)
+{
+	if (bNeedNumericCacheRebuilt)
+	{
+		RebuildNumericCache();
+	}
+
+	FNiagaraTypeDefinition ReturnDef;
+	if (InPin && InPin->PinId.IsValid())
+	{
+		FNiagaraTypeDefinition* FoundDef = CachedNumericConversions.Find(TPair<FGuid, UEdGraphNode*>(InPin->PinId, InPin->GetOwningNode()));
+		if (FoundDef)
+		{
+			ReturnDef = *FoundDef;
+		}
+	}
+	return ReturnDef;
+}
+
 void UNiagaraGraph::RebuildCachedData(bool bForce)
 {
 	// If the graph hasn't changed since last rebuild, then do nothing.
@@ -875,8 +896,72 @@ void UNiagaraGraph::RebuildCachedData(bool bForce)
 	CachedUsageInfo = NewUsageCache;
 	LastBuiltTraversalDataChangeId = ChangeId;
 
+	RebuildNumericCache();
+}
+
+const class UEdGraphSchema_Niagara* UNiagaraGraph::GetNiagaraSchema() const
+{
+	return Cast<UEdGraphSchema_Niagara>(GetSchema());
+}
+
+void UNiagaraGraph::RebuildNumericCache()
+{
+	CachedNumericConversions.Empty();
+	TMap<UNiagaraNode*, bool> VisitedNodes;
+	for (UEdGraphNode* Node : Nodes)
+	{
+		ResolveNumerics(VisitedNodes, Node);
+	}
+	bNeedNumericCacheRebuilt = false;
+}
+
+void UNiagaraGraph::InvalidateNumericCache()
+{
+	bNeedNumericCacheRebuilt = true; 
+	CachedNumericConversions.Empty();
 }
 
+FString UNiagaraGraph::GetFunctionAliasByContext(const FNiagaraGraphFunctionAliasContext& FunctionAliasContext)
+{
+	FString FunctionAlias;
+	for (UEdGraphNode* Node : Nodes)
+	{
+		UNiagaraNode* NiagaraNode = Cast<UNiagaraNode>(Node);
+		if (NiagaraNode != nullptr)
+		{
+			NiagaraNode->AppendFunctionAliasForContext(FunctionAliasContext, FunctionAlias);
+		}
+	}
+	return FunctionAlias;
+}
+
+void UNiagaraGraph::ResolveNumerics(TMap<UNiagaraNode*, bool>& VisitedNodes, UEdGraphNode* Node)
+{
+	UNiagaraNode* NiagaraNode = Cast<UNiagaraNode>(Node);
+	if (NiagaraNode)
+	{
+		TArray<UEdGraphPin*> InputPins;
+		NiagaraNode->GetInputPins(InputPins);
+		for (int32 i = 0; i < InputPins.Num(); i++)
+		{
+			if (InputPins[i])
+			{
+				UNiagaraNode* FoundNode = Cast<UNiagaraNode>(InputPins[i]->GetOwningNode());
+				if (!FoundNode || VisitedNodes.Contains(FoundNode))
+				{
+					continue;
+				}
+				VisitedNodes.Add(FoundNode, true);
+				ResolveNumerics(VisitedNodes, FoundNode);
+			}
+		}
+
+		NiagaraNode->ResolveNumerics(GetNiagaraSchema(), false, &CachedNumericConversions);
+		
+	}
+}
+
+
 void UNiagaraGraph::SynchronizeInternalCacheWithGraph(UNiagaraGraph* Other)
 {
 	// Force us to rebuild the cache, note that this builds traversals and everything else, keeping it in sync if nothing changed from the current version.
diff --git a/Engine/Plugins/FX/Niagara/Source/NiagaraEditor/Private/NiagaraHlslTranslator.cpp b/Engine/Plugins/FX/Niagara/Source/NiagaraEditor/Private/NiagaraHlslTranslator.cpp
--- a/Engine/Plugins/FX/Niagara/Source/NiagaraEditor/Private/NiagaraHlslTranslator.cpp
+++ b/Engine/Plugins/FX/Niagara/Source/NiagaraEditor/Private/NiagaraHlslTranslator.cpp
@@ -8,6 +8,7 @@
 #include "UObject/UObjectHash.h"
 #include "NiagaraNode.h"
 #include "NiagaraNodeFunctionCall.h"
+#include "NiagaraNodeIf.h"
 #include "NiagaraNodeInput.h"
 #include "NiagaraNodeOutput.h"
 #include "NiagaraNodeReadDataSet.h"
@@ -263,7 +264,6 @@ void FHlslNiagaraTranslator::GenerateFunctionSignature(ENiagaraScriptUsage Scrip
 	TArray<FNiagaraVariable> InputVars;
 	TArray<UNiagaraNodeInput*> InputsNodes;
 
-	// Only handle nodes connected to the correct output node in the event of multiple output nodes in the graph.
 	{
 		SCOPE_CYCLE_COUNTER(STAT_NiagaraEditor_Module_NiagaraHLSLTranslator_GenerateFunctionSignature_FindInputNodes);
 		InputsNodes.Reserve(100);
@@ -271,13 +271,17 @@ void FHlslNiagaraTranslator::GenerateFunctionSignature(ENiagaraScriptUsage Scrip
 		Options.bSort = true;
 		Options.bFilterDuplicates = true;
 		Options.bIncludeTranslatorConstants = false;
-		Options.bFilterByScriptUsage = true;
+		// If we're compiling the emitter function we need to filter to the correct usage so that we only get inputs associated with the emitter call, but if we're compiling any other kind of function call we need all inputs
+		// since the function call nodes themselves will have been generated with pins for all inputs and since we match the input nodes here to the inputs passed in by index, the two collections must match otherwise we fail
+		// to compile a graph that would otherwise work correctly.
+		Options.bFilterByScriptUsage = ScriptUsage == ENiagaraScriptUsage::EmitterSpawnScript || ScriptUsage == ENiagaraScriptUsage::EmitterUpdateScript;
 		Options.TargetScriptUsage = ScriptUsage;
 		FuncGraph->FindInputNodes(InputsNodes, Options);
 
 		if (Inputs.Num() != InputsNodes.Num())
 		{
-			const_cast<FHlslNiagaraTranslator*>(this)->Error(FText::Format(LOCTEXT("GenerateFunctionSignatureFail", "Generating function signature for {0} failed.  The function graph is invalid."), FText::FromString(InFullName)), nullptr, nullptr);
+			const_cast<FHlslNiagaraTranslator*>(this)->Error(FText::Format(LOCTEXT("GenerateFunctionSignatureFail", "Generating function signature for {0} failed.  The function call is providing a different number of inputs than the function graph supplies."), 
+				FText::FromString(InFullName)), nullptr, nullptr);
 			return;
 		}
 	}
@@ -287,20 +291,38 @@ void FHlslNiagaraTranslator::GenerateFunctionSignature(ENiagaraScriptUsage Scrip
 
 		InName.Reserve(100 * InputsNodes.Num());
 		InputVars.Reserve(InputsNodes.Num());
+		TArray<uint32> ConstantInputIndicesToRemove;
 		for (int32 i = 0; i < InputsNodes.Num(); ++i)
 		{
 			//Only add to the signature if the caller has provided it, otherwise we use a local default.
 			if (Inputs[i] != INDEX_NONE)
 			{
-				InputVars.Add(InputsNodes[i]->Input);
-				if (bHadNumericInputs)
+				FNiagaraVariable LiteralConstant = InputsNodes[i]->Input;
+				if (GetLiteralConstantVariable(LiteralConstant))
 				{
-					InName += TEXT("_In");
-					InName += InputsNodes[i]->Input.GetType().GetName();
+					checkf(LiteralConstant.GetType() == FNiagaraTypeDefinition::GetBoolDef(), TEXT("Only boolean types are currently supported for literal constants."));
+					FString LiteralConstantAlias = LiteralConstant.GetName().ToString() + TEXT("_") + (LiteralConstant.GetValue<bool>() ? TEXT("true") : TEXT("false"));
+					InName += TEXT("_") + GetSanitizedSymbolName(LiteralConstantAlias.Replace(TEXT("."), TEXT("_")));
+					ConstantInputIndicesToRemove.Add(i);
+				}
+				else
+				{
+					InputVars.Add(InputsNodes[i]->Input);
+					if (bHadNumericInputs)
+					{
+						InName += TEXT("_In");
+						InName += InputsNodes[i]->Input.GetType().GetName();
+					}
 				}
 			}
 		}
 
+		// Remove the inputs which will be handled by inline constants
+		for (int32 i = ConstantInputIndicesToRemove.Num() - 1; i >= 0; i--)
+		{
+			Inputs.RemoveAt(ConstantInputIndicesToRemove[i]);
+		}
+
 		//Now actually remove the missing inputs so they match the signature.
 		Inputs.Remove(INDEX_NONE);
 	}
@@ -350,7 +372,10 @@ void FHlslNiagaraTranslator::GenerateFunctionSignature(ENiagaraScriptUsage Scrip
 	}
 	else
 	{
-		OutSig = FNiagaraFunctionSignature(*InName, InputVars, OutputVars, *InFullName, true, false);
+		FNiagaraGraphFunctionAliasContext FunctionAliasContext;
+		FunctionAliasContext.CompileUsage = GetCurrentUsage();
+		FString SignatureName = InName + FuncGraph->GetFunctionAliasByContext(FunctionAliasContext);
+		OutSig = FNiagaraFunctionSignature(*SignatureName, InputVars, OutputVars, *InFullName, true, false);
 	}
 }
 
@@ -697,10 +722,16 @@ const FNiagaraTranslateResults &FHlslNiagaraTranslator::Translate(const FNiagara
 		return TranslateResults;
 	}
 
+    bool bNeedsPersistentIDs = CompileOptions.AdditionalDefines.Contains(TEXT("RequiresPersistentIDs"));
+    if (bNeedsPersistentIDs && CompilationTarget == ENiagaraSimTarget::GPUComputeSim) {
+        Error(LOCTEXT("GPUPersistentIDFail", "GPU particles do not support persistent IDs. Change to a CPU simulation or disable persistent IDs."), nullptr, nullptr);
+        return TranslateResults;
+    }
+
 	TranslationStages.Empty();
 	ActiveStageIdx = 0;
 
-	bool bHasInterpolatedSpawn = InCompileOptions.AdditionalDefines.Contains(TEXT("InterpolatedSpawn"));
+	bool bHasInterpolatedSpawn = CompileOptions.AdditionalDefines.Contains(TEXT("InterpolatedSpawn"));
 	ParamMapHistories.Empty();
 	ParamMapSetVariablesToChunks.Empty();
 
@@ -771,67 +802,88 @@ const FNiagaraTranslateResults &FHlslNiagaraTranslator::Translate(const FNiagara
 
 	// Get all the parameter map histories traced to this graph from output nodes. We'll revisit this shortly in order to build out just the ones we care about for this translation.
 	OtherOutputParamMapHistories = CompileData->GetPrecomputedHistories();
-	for (FNiagaraParameterMapHistory& FoundHistory : OtherOutputParamMapHistories)
+
+	if (ParamMapHistories.Num() == 1 && OtherOutputParamMapHistories.Num() == 1 && ( CompileOptions.TargetUsage == ENiagaraScriptUsage::Function || CompileOptions.TargetUsage == ENiagaraScriptUsage::DynamicInput ) )
 	{
-		const UNiagaraNodeOutput* HistoryOutputNode = FoundHistory.GetFinalOutputNode();
-		if (HistoryOutputNode != nullptr && !ShouldConsiderTargetParameterMap(HistoryOutputNode->GetUsage()))
+		ParamMapHistories[0] = (OtherOutputParamMapHistories[0]);
+
+		TArray<int32> Entries;
+		Entries.AddZeroed(OtherOutputParamMapHistories[0].Variables.Num());
+		for (int32 i = 0; i < Entries.Num(); i++)
 		{
-			continue;
+			Entries[i] = INDEX_NONE;
 		}
-		
-		// Now see if we want to use any of these specifically..
-		for (int32 ParamMapIdx = 0; ParamMapIdx < TranslationStages.Num(); ParamMapIdx++)
+		ParamMapSetVariablesToChunks[0] = (Entries);
+	}
+	else
+	{
+		for (FNiagaraParameterMapHistory& FoundHistory : OtherOutputParamMapHistories)
 		{
-			UNiagaraNodeOutput* TargetOutputNode = TranslationStages[ParamMapIdx].OutputNode;
-			if (FoundHistory.GetFinalOutputNode() == TargetOutputNode)
+			const UNiagaraNodeOutput* HistoryOutputNode = FoundHistory.GetFinalOutputNode();
+			if (HistoryOutputNode != nullptr && !ShouldConsiderTargetParameterMap(HistoryOutputNode->GetUsage()))
 			{
-				bool bNeedsPersistentIDs = InCompileOptions.AdditionalDefines.Contains(TEXT("RequiresPersistentIDs"));
-				if (bNeedsPersistentIDs)
-				{
-					//TODO: Setup alias for current level to decouple from "Particles". Would we ever want emitter or system persistent IDs?
-					FNiagaraVariable Var = FNiagaraVariable(FNiagaraTypeDefinition::GetIDDef(), TEXT("Particles.ID"));
-					FoundHistory.AddVariable(Var, Var, nullptr);
-				}
+				continue;
+			}
 
-				if (RequiresInterpolation())
+			// Now see if we want to use any of these specifically..
+			for (int32 ParamMapIdx = 0; ParamMapIdx < TranslationStages.Num(); ParamMapIdx++)
+			{
+				UNiagaraNodeOutput* TargetOutputNode = TranslationStages[ParamMapIdx].OutputNode;
+				if (FoundHistory.GetFinalOutputNode() == TargetOutputNode)
 				{
-					FNiagaraVariable Var = FNiagaraVariable(FNiagaraTypeDefinition::GetIntDef(), TEXT("Interpolation.InterpSpawn_Index"));
-					FoundHistory.AddVariable(Var, Var, nullptr);
+					if(bNeedsPersistentIDs)
+					{
+						//TODO: Setup alias for current level to decouple from "Particles". Would we ever want emitter or system persistent IDs?
+						FNiagaraVariable Var = FNiagaraVariable(FNiagaraTypeDefinition::GetIDDef(), TEXT("Particles.ID"));
+						FoundHistory.AddVariable(Var, Var, nullptr);
+					}
+					{
+						// NOTE(mv): This will explicitly expose Particles.UniqueID to the HLSL code regardless of whether it is exposed in a script or not. 
+						//           This is necessary as the script needs to know about it even when no scripts reference it. 
+						FNiagaraVariable Var = FNiagaraVariable(FNiagaraTypeDefinition::GetIntDef(), TEXT("Particles.UniqueID"));
+						FoundHistory.AddVariable(Var, Var, nullptr);
+					}
 
-					Var = FNiagaraVariable(FNiagaraTypeDefinition::GetFloatDef(), TEXT("Interpolation.InterpSpawn_SpawnTime"));
-					FoundHistory.AddVariable(Var, Var, nullptr);
+					if (RequiresInterpolation())
+					{
+						FNiagaraVariable Var = FNiagaraVariable(FNiagaraTypeDefinition::GetIntDef(), TEXT("Interpolation.InterpSpawn_Index"));
+						FoundHistory.AddVariable(Var, Var, nullptr);
 
-					Var = FNiagaraVariable(FNiagaraTypeDefinition::GetFloatDef(), TEXT("Interpolation.InterpSpawn_UpdateTime"));
-					FoundHistory.AddVariable(Var, Var, nullptr);
+						Var = FNiagaraVariable(FNiagaraTypeDefinition::GetFloatDef(), TEXT("Interpolation.InterpSpawn_SpawnTime"));
+						FoundHistory.AddVariable(Var, Var, nullptr);
 
-					Var = FNiagaraVariable(FNiagaraTypeDefinition::GetFloatDef(), TEXT("Interpolation.InterpSpawn_InvSpawnTime"));
-					FoundHistory.AddVariable(Var, Var, nullptr);
+						Var = FNiagaraVariable(FNiagaraTypeDefinition::GetFloatDef(), TEXT("Interpolation.InterpSpawn_UpdateTime"));
+						FoundHistory.AddVariable(Var, Var, nullptr);
 
-					Var = FNiagaraVariable(FNiagaraTypeDefinition::GetFloatDef(), TEXT("Interpolation.InterpSpawn_InvUpdateTime"));
-					FoundHistory.AddVariable(Var, Var, nullptr);
+						Var = FNiagaraVariable(FNiagaraTypeDefinition::GetFloatDef(), TEXT("Interpolation.InterpSpawn_InvSpawnTime"));
+						FoundHistory.AddVariable(Var, Var, nullptr);
 
-					Var = FNiagaraVariable(FNiagaraTypeDefinition::GetFloatDef(), TEXT("Interpolation.SpawnInterp"));
-					FoundHistory.AddVariable(Var, Var, nullptr);
+						Var = FNiagaraVariable(FNiagaraTypeDefinition::GetFloatDef(), TEXT("Interpolation.InterpSpawn_InvUpdateTime"));
+						FoundHistory.AddVariable(Var, Var, nullptr);
 
-					Var = FNiagaraVariable(FNiagaraTypeDefinition::GetFloatDef(), TEXT("Interpolation.Emitter_SpawnInterval"));
-					FoundHistory.AddVariable(Var, Var, nullptr);
+						Var = FNiagaraVariable(FNiagaraTypeDefinition::GetFloatDef(), TEXT("Interpolation.SpawnInterp"));
+						FoundHistory.AddVariable(Var, Var, nullptr);
 
-					Var = FNiagaraVariable(FNiagaraTypeDefinition::GetFloatDef(), TEXT("Interpolation.Emitter_InterpSpawnStartDt"));
-					FoundHistory.AddVariable(Var, Var, nullptr);
+						Var = FNiagaraVariable(FNiagaraTypeDefinition::GetFloatDef(), TEXT("Interpolation.Emitter_SpawnInterval"));
+						FoundHistory.AddVariable(Var, Var, nullptr);
 
-					Var = FNiagaraVariable(FNiagaraTypeDefinition::GetIntDef(), TEXT("Interpolation.Emitter_SpawnGroup"));
-					FoundHistory.AddVariable(Var, Var, nullptr);
-				}
+						Var = FNiagaraVariable(FNiagaraTypeDefinition::GetFloatDef(), TEXT("Interpolation.Emitter_InterpSpawnStartDt"));
+						FoundHistory.AddVariable(Var, Var, nullptr);
 
-				ParamMapHistories[ParamMapIdx] = (FoundHistory);
+						Var = FNiagaraVariable(FNiagaraTypeDefinition::GetIntDef(), TEXT("Interpolation.Emitter_SpawnGroup"));
+						FoundHistory.AddVariable(Var, Var, nullptr);
+					}
 
-				TArray<int32> Entries;
-				Entries.AddZeroed(FoundHistory.Variables.Num());
-				for (int32 i = 0; i < Entries.Num(); i++)
-				{
-					Entries[i] = INDEX_NONE;
+					ParamMapHistories[ParamMapIdx] = (FoundHistory);
+
+					TArray<int32> Entries;
+					Entries.AddZeroed(FoundHistory.Variables.Num());
+					for (int32 i = 0; i < Entries.Num(); i++)
+					{
+						Entries[i] = INDEX_NONE;
+					}
+					ParamMapSetVariablesToChunks[ParamMapIdx] = (Entries);
 				}
-				ParamMapSetVariablesToChunks[ParamMapIdx] = (Entries);
 			}
 		}
 	}
@@ -844,6 +896,11 @@ const FNiagaraTranslateResults &FHlslNiagaraTranslator::Translate(const FNiagara
 			CompilationOutput.ScriptData.ParameterCollectionPaths.AddUnique(FSoftObjectPath(Collection).ToString());
 		}
 	}
+	ENiagaraScriptUsage Usage = CompileOptions.TargetUsage;
+	if (Usage != ENiagaraScriptUsage::SystemSpawnScript && Usage != ENiagaraScriptUsage::SystemUpdateScript)
+	{
+		ValidateParticleIDUsage();
+	}
 	
 	//Create main scope pin cache.
 	PinToCodeChunks.AddDefaulted(1);
@@ -935,18 +992,28 @@ const FNiagaraTranslateResults &FHlslNiagaraTranslator::Translate(const FNiagara
 		//Generate function definitions
 		FString FunctionDefinitionString = GetFunctionDefinitions();
 		FunctionDefinitionString += TEXT("\n");
-
-		if (TranslationStages.Num() > 1 && RequiresInterpolation())
 		{
-			//ensure the interpolated spawn constants are part of the parameter set.
-			int32 OutputIdx = 0;
-			ParameterMapRegisterExternalConstantNamespaceVariable(SYS_PARAM_ENGINE_DELTA_TIME, nullptr, 0, OutputIdx, nullptr);
-			ParameterMapRegisterExternalConstantNamespaceVariable(SYS_PARAM_ENGINE_INV_DELTA_TIME, nullptr, 0, OutputIdx, nullptr);
-			ParameterMapRegisterExternalConstantNamespaceVariable(SYS_PARAM_ENGINE_EXEC_COUNT, nullptr, 0, OutputIdx, nullptr);
-			ParameterMapRegisterExternalConstantNamespaceVariable(SYS_PARAM_EMITTER_SPAWNRATE, nullptr, 0, OutputIdx, nullptr);
-			ParameterMapRegisterExternalConstantNamespaceVariable(SYS_PARAM_EMITTER_SPAWN_INTERVAL, nullptr, 0, OutputIdx, nullptr);
-			ParameterMapRegisterExternalConstantNamespaceVariable(SYS_PARAM_EMITTER_INTERP_SPAWN_START_DT, nullptr, 0, OutputIdx, nullptr);
-			ParameterMapRegisterExternalConstantNamespaceVariable(SYS_PARAM_EMITTER_SPAWN_GROUP, nullptr, 0, OutputIdx, nullptr);
+			if (TranslationStages.Num() > 1 && RequiresInterpolation())
+			{
+				int32 OutputIdx = 0;
+				//ensure the interpolated spawn constants are part of the parameter set.
+				ParameterMapRegisterExternalConstantNamespaceVariable(SYS_PARAM_ENGINE_DELTA_TIME, nullptr, 0, OutputIdx, nullptr);
+				ParameterMapRegisterExternalConstantNamespaceVariable(SYS_PARAM_ENGINE_INV_DELTA_TIME, nullptr, 0, OutputIdx, nullptr);
+				ParameterMapRegisterExternalConstantNamespaceVariable(SYS_PARAM_ENGINE_EXEC_COUNT, nullptr, 0, OutputIdx, nullptr);
+				ParameterMapRegisterExternalConstantNamespaceVariable(SYS_PARAM_EMITTER_SPAWNRATE, nullptr, 0, OutputIdx, nullptr);
+				ParameterMapRegisterExternalConstantNamespaceVariable(SYS_PARAM_EMITTER_SPAWN_INTERVAL, nullptr, 0, OutputIdx, nullptr);
+				ParameterMapRegisterExternalConstantNamespaceVariable(SYS_PARAM_EMITTER_INTERP_SPAWN_START_DT, nullptr, 0, OutputIdx, nullptr);
+				ParameterMapRegisterExternalConstantNamespaceVariable(SYS_PARAM_EMITTER_SPAWN_GROUP, nullptr, 0, OutputIdx, nullptr);
+			}
+
+			if (TranslationStages.Num() > 0)
+			{
+				int32 OutputIdx = 0;
+				// NOTE(mv): This will explicitly expose Engine.Emitter.TotalSpawnedParticles to the HLSL code regardless of whether it is exposed in a script or not. 
+				ParameterMapRegisterExternalConstantNamespaceVariable(SYS_PARAM_ENGINE_EMITTER_TOTAL_SPAWNED_PARTICLES, nullptr, 0, OutputIdx, nullptr);
+				ParameterMapRegisterExternalConstantNamespaceVariable(SYS_PARAM_EMITTER_RANDOM_SEED, nullptr, 0, OutputIdx, nullptr);
+				//ParameterMapRegisterExternalConstantNamespaceVariable(SYS_PARAM_EMITTER_DETERMINISM, nullptr, 0, OutputIdx, nullptr);
+			}
 		}
 
 		// Generate the Parameter Map HLSL definitions. We don't add to the final HLSL output here. We just build up the strings and tables
@@ -956,7 +1023,12 @@ const FNiagaraTranslateResults &FHlslNiagaraTranslator::Translate(const FNiagara
 
 		for (FNiagaraTypeDefinition Type : StructsToDefine)
 		{
-			HlslOutput += BuildHLSLStructDecl(Type);
+			FText ErrorMessage;
+			HlslOutput += BuildHLSLStructDecl(Type, ErrorMessage);
+			if (ErrorMessage.IsEmpty() == false)
+			{
+				Error(ErrorMessage, nullptr, nullptr);
+			}
 		}
 				
 		//Declare parameters.
@@ -1192,7 +1264,14 @@ const FNiagaraTranslateResults &FHlslNiagaraTranslator::Translate(const FNiagara
 
 		// We may have created some transient data interfaces. This cleans up the ones that we created.
 		CompilationOutput.ScriptData.DIParamInfo = DIParamInfo;
-		CompilationOutput.ScriptData.bReadsAttributeData = InstanceRead.Variables.Num() != 0;
+		if (InstanceRead.Variables.Num() == 1 && InstanceRead.Variables[0].GetName() == TEXT("Particles.UniqueID")) {
+			// NOTE(mv): Explicitly allow reading from Particles.UniqueID, as it is an engine managed variable and 
+			//           is written to before Simulate() in the SpawnScript...
+			// TODO(mv): Also allow Particles.ID for the same reasons?
+			CompilationOutput.ScriptData.bReadsAttributeData = false;
+		} else {
+			CompilationOutput.ScriptData.bReadsAttributeData = InstanceRead.Variables.Num() != 0;
+		}
 		TranslateResults.OutputHLSL = HlslOutput;
 
 	}
@@ -1639,6 +1718,27 @@ void FHlslNiagaraTranslator::DefineMain(FString &OutHlslOutput,
 		OutHlslOutput += FString::Printf(TEXT("\t%s.Particles.ID.Index = TempIDIndex;\n\t%s.Particles.ID.AcquireTag = TempIDTag;\n"), *MapName, *MapName);
 	}
 
+	{
+		// Manually write to Particles.UniqueID on spawn, and deliberately place it at the top of SimulateMain to make sure it's initialized in the right order
+
+		// NOTE(mv): These relies on Particles.UniqueID and Engine.Emitter.TotalSpawnedParticles both being explicitly added to the parameter histories in 
+		//           FHlslNiagaraTranslator::Translate.
+
+		// NOTE(mv): This relies on Particles.UniqueID being excluded from being default initialized. 
+		//           This happens in FNiagaraParameterMapHistory::ShouldIgnoreVariableDefault
+		if (UNiagaraScript::IsParticleSpawnScript(CompileOptions.TargetUsage)) 
+		{
+			FString MapName = UNiagaraScript::IsInterpolatedParticleSpawnScript(CompileOptions.TargetUsage) ? TEXT("Context.MapSpawn") : TEXT("Context.Map");
+			OutHlslOutput += FString::Printf(TEXT("\t%s.Particles.UniqueID = Engine_Emitter_TotalSpawnedParticles + ExecIndex();\n"), *MapName);
+		}
+		else if (UNiagaraScript::IsGPUScript(CompileOptions.TargetUsage)) 
+		{
+			// NOTE(mv): The GPU script only have one file, so we need to make sure we only apply this in the spawn phase. 
+			//           
+			OutHlslOutput += TEXT("\tif (Phase == 0) \n\t{\n\t\tContext.MapSpawn.Particles.UniqueID = Engine_Emitter_TotalSpawnedParticles + ExecIndex();\n\t}\n");
+		}
+	}
+
 	// Fill in the defaults for parameters.
 	for (int32 i = 0; i < MainPreSimulateChunks.Num(); ++i)
 	{
@@ -1701,7 +1801,7 @@ void FHlslNiagaraTranslator::DefineMain(FString &OutHlslOutput,
 				}
 
 				OutHlslOutput += FString::Printf(TEXT("\t\tSimulate%s(Context);\n"), TranslationStages.Num() > 1 ? *TranslationStages[StageIdx].PassNamespace : TEXT(""));
-				
+
 				if (StageIdx + 1 < TranslationStages.Num() && TranslationStages[StageIdx + 1].bCopyPreviousParams)
 				{
 					OutHlslOutput += TEXT("\t\t//Begin Transfer of Attributes!\n");
@@ -1720,6 +1820,7 @@ void FHlslNiagaraTranslator::DefineMain(FString &OutHlslOutput,
 					OutHlslOutput += TEXT("\t\t//End Transfer of Attributes!\n\n");
 				}
 
+
 				// Either go on to the next phase, or write to the final output context.
 				if (StageIdx + 1 < TranslationStages.Num() && TranslationStages[StageIdx + 1].bInterpolatePreviousParams)
 				{
@@ -2324,7 +2425,7 @@ int32 FHlslNiagaraTranslator::AddUniformChunk(FString SymbolName, const FNiagara
 		Chunk.SymbolName = GetSanitizedSymbolName(SymbolName);
 		Chunk.Type = Type;
 
-		if (CompileOptions.TargetUsage == ENiagaraScriptUsage::ParticleGPUComputeScript)
+		if (UNiagaraScript::IsGPUScript(CompileOptions.TargetUsage))
 		{
 			if (Type == FNiagaraTypeDefinition::GetVec2Def())
 			{
@@ -2462,7 +2563,9 @@ bool FHlslNiagaraTranslator::ShouldInterpolateParameter(const FNiagaraVariable&
 		Parameter == SYS_PARAM_EMITTER_SPAWNRATE ||
 		Parameter == SYS_PARAM_EMITTER_SPAWN_INTERVAL ||
 		Parameter == SYS_PARAM_EMITTER_INTERP_SPAWN_START_DT ||
-		Parameter == SYS_PARAM_EMITTER_SPAWN_GROUP )
+		Parameter == SYS_PARAM_ENGINE_EMITTER_TOTAL_SPAWNED_PARTICLES || 
+		Parameter == SYS_PARAM_EMITTER_RANDOM_SEED || 
+		Parameter == SYS_PARAM_ENGINE_SYSTEM_TICK_COUNT)
 	{
 		return false;
 	}
@@ -2554,7 +2657,9 @@ int32 FHlslNiagaraTranslator::GetParameter(const FNiagaraVariable& Parameter)
 	}
 
 	int32 FuncParam = INDEX_NONE;
-	if (GetFunctionParameter(Parameter, FuncParam))
+	const FNiagaraVariable* FoundKnownVariable = FNiagaraConstants::GetKnownConstant(Parameter.GetName(), false);
+
+	if (FoundKnownVariable == nullptr && GetFunctionParameter(Parameter, FuncParam))
 	{
 		if (FuncParam != INDEX_NONE)
 		{
@@ -2568,6 +2673,16 @@ int32 FHlslNiagaraTranslator::GetParameter(const FNiagaraVariable& Parameter)
 		}
 	}
 
+	if (FoundKnownVariable != nullptr)
+	{
+		FNiagaraVariable Var = *FoundKnownVariable;
+		//Some special variables can be replaced directly with constants which allows for extra optimization in the compiler.
+		if (GetLiteralConstantVariable(Var))
+		{
+			return GetConstant(Var);
+		}
+	}
+
 	// We don't pass in the input node here (really there could be multiple nodes for the same parameter)
 	// so we have to match up the input parameter map variable value through the pre-traversal histories 
 	// so that we know which parameter map we are referencing.
@@ -2647,7 +2762,18 @@ int32 FHlslNiagaraTranslator::GetConstant(const FNiagaraVariable& Constant)
 		return INDEX_NONE;
 	}
 
-	FString ConstantStr = GenerateConstantString(Constant);
+	FString ConstantStr;
+	FNiagaraVariable LiteralConstant = Constant;
+	if (GetLiteralConstantVariable(LiteralConstant))
+	{
+		checkf(LiteralConstant.GetType() == FNiagaraTypeDefinition::GetBoolDef(), TEXT("Only boolean types are currently supported for literal constants."));
+		ConstantStr = LiteralConstant.GetValue<bool>() ? TEXT("true") : TEXT("false");
+	}
+	else
+	{
+		ConstantStr = GenerateConstantString(Constant);
+	}
+
 	if (ConstantStr.IsEmpty())
 	{
 		return INDEX_NONE;
@@ -2921,15 +3047,15 @@ int32 FHlslNiagaraTranslator::GetAttribute(const FNiagaraVariable& Attribute)
 		Error(FText::Format(LOCTEXT("GetConstantFail", "Cannot handle type {0}! Variable: {1}"), Attribute.GetType().GetNameText(), FText::FromName(Attribute.GetName())), nullptr, nullptr);
 	}
 
-	if (TranslationStages.Num() > 1 && UNiagaraScript::IsParticleSpawnScript(TranslationStages[0].ScriptUsage))
+	if (TranslationStages.Num() > 1 && UNiagaraScript::IsParticleSpawnScript(TranslationStages[0].ScriptUsage) && (Attribute.GetName() != TEXT("Particles.UniqueID")))
 	{
 		if (ActiveStageIdx > 0)
 		{
 			//This is a special case where we allow the grabbing of attributes in the update section of an interpolated spawn script.
 			//But we return the results of the previously ran spawn script.
 			FString ParameterMapInstanceName = GetParameterMapInstanceName(0);
-			FNiagaraVariable NamespacedVar = FNiagaraParameterMapHistory::BasicAttributeToNamespacedAttribute(Attribute);
 
+			FNiagaraVariable NamespacedVar = Attribute;
 			FString SymbolName = *(ParameterMapInstanceName + TEXT(".") + GetSanitizedSymbolName(NamespacedVar.GetName().ToString()));
 			return AddSourceChunk(SymbolName, Attribute.GetType());
 		}
@@ -2941,9 +3067,13 @@ int32 FHlslNiagaraTranslator::GetAttribute(const FNiagaraVariable& Attribute)
 	}
 	else
 	{
-		CompilationOutput.ScriptData.DataUsage.bReadsAttributeData = true;
+		// NOTE(mv): Explicitly allow reading from Particles.UniqueID, as it is an engine managed variable and 
+		//           is written to before Simulate() in the SpawnScript...
+		// TODO(mv): Also allow Particles.ID for the same reasons?
+		CompilationOutput.ScriptData.DataUsage.bReadsAttributeData |= (Attribute.GetName() != TEXT("Particles.UniqueID"));
+		
 		int32 Chunk = INDEX_NONE;
-		if (!ParameterMapRegisterUniformAttributeVariable(Attribute, nullptr, 0, Chunk))
+		if (!ParameterMapRegisterNamespaceAttributeVariable(Attribute, nullptr, 0, Chunk))
 		{
 			Error(FText::Format(LOCTEXT("AttrReadError", "Cannot read attribute {0} {1}."), Attribute.GetType().GetNameText(), FText::FromString(*Attribute.GetName().ToString())), nullptr, nullptr);
 			return INDEX_NONE;
@@ -3146,7 +3276,7 @@ bool FHlslNiagaraTranslator::RequiresInterpolation() const
 	return false;
 }
 
-bool FHlslNiagaraTranslator::GetLiteralConstantVariable(FNiagaraVariable& OutVar)
+bool FHlslNiagaraTranslator::GetLiteralConstantVariable(FNiagaraVariable& OutVar) const
 {
 	if (FNiagaraParameterMapHistory::IsInNamespace(OutVar, PARAM_MAP_EMITTER_STR) || FNiagaraParameterMapHistory::IsInNamespace(OutVar, PARAM_MAP_SYSTEM_STR))
 	{
@@ -3157,6 +3287,12 @@ bool FHlslNiagaraTranslator::GetLiteralConstantVariable(FNiagaraVariable& OutVar
 			OutVar.SetValue(bEmitterLocalSpace ? FNiagaraBool(true) : FNiagaraBool(false));
 			return true;
 		}
+		if (OutVar == FNiagaraVariable(FNiagaraTypeDefinition::GetBoolDef(), TEXT("Emitter.Determinism")))
+		{
+			bool bEmitterDeterminism = CompileOptions.AdditionalDefines.Contains(ResolvedVar.GetName().ToString());
+			OutVar.SetValue(bEmitterDeterminism ? FNiagaraBool(true) : FNiagaraBool(false));
+			return true;
+		}
 	}
 	return false;
 }
@@ -3295,6 +3431,26 @@ bool FHlslNiagaraTranslator::ParameterMapRegisterUniformAttributeVariable(const
 	return false;
 }
 
+void FHlslNiagaraTranslator::ValidateParticleIDUsage()
+{
+    if (CompileOptions.AdditionalDefines.Contains(TEXT("RequiresPersistentIDs"))) 
+	{
+        // persistent IDs are active and can be safely used as inputs
+        return;
+    }
+    FName particleIDName(TEXT("Particles.ID"));
+    for (FNiagaraParameterMapHistory& History : ParamMapHistories)
+    {
+        for (const FNiagaraVariable& Variable : History.Variables)
+        {
+            if (Variable.GetName() == particleIDName) 
+			{
+				Error(LOCTEXT("PersistentIDActivationFail", "Before the Particles.ID parameter can be used, the 'Requires persistent IDs' option has to be activated in the emitter properties. Note that this comes with additional memory and CPU costs."), nullptr, nullptr);
+            }
+        }
+    }
+}
+
 bool FHlslNiagaraTranslator::ParameterMapRegisterNamespaceAttributeVariable(const FNiagaraVariable& InVariable, UNiagaraNode* InNode, int32 InParamMapHistoryIdx, int32& Output)
 {
 	FString VarName = InVariable.GetName().ToString();
@@ -3977,9 +4133,7 @@ void FHlslNiagaraTranslator::Operation(class UNiagaraNodeOp* Operation, TArray<i
 
 	// Use the pins to determine the output type here since they may have been changed due to numeric pin fix up.
 	const FNiagaraOpInfo* OpInfo = FNiagaraOpInfo::GetOpInfo(Operation->OpName);
-
-	//EnterStatsScope(FNiagaraStatScope(*Operation->GetFullName(), OpInfo->FriendlyName));
-
+	
 	TArray<UEdGraphPin*> OutputPins;
 	Operation->GetOutputPins(OutputPins);
 	for (int32 OutputIndex = 0; OutputIndex < OutputPins.Num(); OutputIndex++)
@@ -3994,11 +4148,22 @@ void FHlslNiagaraTranslator::Operation(class UNiagaraNodeOp* Operation, TArray<i
 		}
 
 		const FNiagaraOpInOutInfo& IOInfo = OpInfo->Outputs[OutputIndex];
-		check(!IOInfo.HlslSnippet.IsEmpty());
-		Outputs.Add(AddBodyChunk(GetUniqueSymbolName(IOInfo.Name), IOInfo.HlslSnippet, OutputType, Inputs));
+		FString OutputHlsl;
+		if (OpInfo->bSupportsAddedInputs)
+		{
+			if (!OpInfo->CreateHlslForAddedInputs(Inputs.Num(), OutputHlsl))
+			{
+				FText PinNameText = OutputPin->PinFriendlyName.IsEmpty() ? FText::FromName(OutputPin->PinName) : OutputPin->PinFriendlyName;
+				Error(FText::Format(LOCTEXT("AggregateInputFailTypePin", "Cannot create hlsl output for type {0}! Output Pin: {1}"), OutputType.GetNameText(), PinNameText), Operation, OutputPin);
+				OutputHlsl = IOInfo.HlslSnippet;
+			}
+		}
+		else {
+			OutputHlsl = IOInfo.HlslSnippet;
+		}
+		check(!OutputHlsl.IsEmpty());
+		Outputs.Add(AddBodyChunk(GetUniqueSymbolName(IOInfo.Name), OutputHlsl, OutputType, Inputs));
 	}
-
-	//ExitStatsScope();
 }
 
 void FHlslNiagaraTranslator::FunctionCall(UNiagaraNodeFunctionCall* FunctionNode, TArray<int32>& Inputs, TArray<int32>& Outputs)
@@ -4388,9 +4553,10 @@ void FHlslNiagaraTranslator::RegisterFunctionCall(ENiagaraScriptUsage ScriptUsag
 			UNiagaraNodeOutput* FuncOutput = SourceGraph->FindOutputNode(ScriptUsage);
 			check(FuncOutput);
 
-			// Go ahead and insert any defaulted values into the parameter map here at the top level.
-			if (ActiveHistoryForFunctionCalls.InTopLevelFunctionCall(CompileOptions.TargetUsage) && ActiveHistoryForFunctionCalls.GetModuleAlias() != nullptr)
+			if (ActiveHistoryForFunctionCalls.GetModuleAlias() != nullptr)
 			{
+				bool bIsInTopLevelFunction = ActiveHistoryForFunctionCalls.InTopLevelFunctionCall(CompileOptions.TargetUsage);
+
 				UEdGraphPin* ParamMapPin = nullptr;
 				for (UEdGraphPin* Pin : CallInputs)
 				{
@@ -4417,20 +4583,35 @@ void FHlslNiagaraTranslator::RegisterFunctionCall(ENiagaraScriptUsage ScriptUsag
 
 							for (uint32 VarIdx = History.MapNodeVariableMetaData[FoundIdx].Key; VarIdx < History.MapNodeVariableMetaData[FoundIdx].Value; VarIdx++)
 							{
+								if (History.PerVariableReadHistory[VarIdx].Num() == 0)
+								{
+									// We don't need to worry about defaults if the variable is only written to.
+									continue;
+								}
+
 								const FNiagaraVariable& Var = History.Variables[VarIdx];
 								const FNiagaraVariable& AliasedVar = History.VariablesWithOriginalAliasesIntact[VarIdx];
-								int32 LastSetChunkIdx = ParamMapSetVariablesToChunks[ActiveStageIdx][VarIdx];
-								if (LastSetChunkIdx == INDEX_NONE)
+								bool bIsAliased = Var.GetName() != AliasedVar.GetName();
+
+								// For non aliased values we resolve the defaults once at the top level since it's impossible to know which context they were actually used in, but
+								// for aliased values we check to see if they're used in the current context by resolving the alias and checking against the current resolved variable
+								// name since aliased values can only be resolved for reading in the correct context.
+								bool bIsValidForCurrentCallingContext = (bIsInTopLevelFunction && bIsAliased == false) || (bIsAliased && ActiveHistoryForFunctionCalls.ResolveAliases(AliasedVar).GetName() == Var.GetName());
+								if (bIsValidForCurrentCallingContext)
 								{
-									const UEdGraphPin* DefaultPin = History.GetDefaultValuePin(VarIdx);
-									HandleParameterRead(ActiveStageIdx, AliasedVar, DefaultPin, ParamNode, LastSetChunkIdx);
-									
-									// If this variable was in the pending defaults list, go ahead and remove it
-									// as we added it before first use...
-									if (DeferredVariablesMissingDefault.Contains(Var))
+									int32 LastSetChunkIdx = ParamMapSetVariablesToChunks[ActiveStageIdx][VarIdx];
+									if (LastSetChunkIdx == INDEX_NONE)
 									{
-										DeferredVariablesMissingDefault.Remove(Var);
-										UniqueVarToChunk.Add(Var, LastSetChunkIdx);
+										const UEdGraphPin* DefaultPin = History.GetDefaultValuePin(VarIdx);
+										HandleParameterRead(ActiveStageIdx, AliasedVar, DefaultPin, ParamNode, LastSetChunkIdx);
+
+										// If this variable was in the pending defaults list, go ahead and remove it
+										// as we added it before first use...
+										if (DeferredVariablesMissingDefault.Contains(Var))
+										{
+											DeferredVariablesMissingDefault.Remove(Var);
+											UniqueVarToChunk.Add(Var, LastSetChunkIdx);
+										}
 									}
 								}
 							}
@@ -4580,7 +4761,7 @@ void FHlslNiagaraTranslator::RegisterFunctionCall(ENiagaraScriptUsage ScriptUsag
 					return;
 				}
 
-				if (Info.UserPtrIdx != INDEX_NONE)
+				if (Info.UserPtrIdx != INDEX_NONE && CompilationTarget != ENiagaraSimTarget::GPUComputeSim)
 				{
 					//This interface requires per instance data via a user ptr so place the index to it at the end of the inputs.
 					Inputs.Add(AddSourceChunk(LexToString(Info.UserPtrIdx), FNiagaraTypeDefinition::GetIntDef(), false));
@@ -5006,7 +5187,7 @@ void FHlslNiagaraTranslator::Convert(class UNiagaraNodeConvert* Convert, TArray
 	}
 }
 
-void FHlslNiagaraTranslator::If(TArray<FNiagaraVariable>& Vars, int32 Condition, TArray<int32>& PathA, TArray<int32>& PathB, TArray<int32>& Outputs)
+void FHlslNiagaraTranslator::If(UNiagaraNodeIf* IfNode, TArray<FNiagaraVariable>& Vars, int32 Condition, TArray<int32>& PathA, TArray<int32>& PathB, TArray<int32>& Outputs)
 {
 	SCOPE_CYCLE_COUNTER(STAT_NiagaraEditor_HlslTranslator_If);
 
@@ -5016,10 +5197,19 @@ void FHlslNiagaraTranslator::If(TArray<FNiagaraVariable>& Vars, int32 Condition,
 
 	TArray<FString> OutSymbols;
 	OutSymbols.Reserve(Vars.Num());
+	int32 PinIdx = 1;
 	for (FNiagaraVariable& Var : Vars)
 	{
+		FNiagaraTypeDefinition Type = Schema->PinToTypeDefinition(IfNode->GetInputPin(PinIdx++));
+		if (!AddStructToDefinitionSet(Type))
+		{
+			FText OutErrorMessage = FText::Format(LOCTEXT("UnknownNumeric", "Variable in If node uses invalid type. Var: {0} Type: {1}"),
+				FText::FromName(Var.GetName()), Type.GetNameText());
+
+			Error(OutErrorMessage, IfNode, nullptr);
+		}
 		OutSymbols.Add(GetUniqueSymbolName(*(Var.GetName().ToString() + TEXT("_IfResult"))));
-		Outputs.Add(AddBodyChunk(OutSymbols.Last(), TEXT(""), Var.GetType(), true));
+		Outputs.Add(AddBodyChunk(OutSymbols.Last(), TEXT(""), Type, true));
 	}
 	AddBodyChunk(TEXT(""), TEXT("if({0})\n\t{"), FNiagaraTypeDefinition::GetFloatDef(), Condition, false, false);
 	for (int32 i = 0; i < NumVars; ++i)
@@ -5159,9 +5349,9 @@ void FHlslNiagaraTranslator::Error(FText ErrorText, const UNiagaraNode* Node, co
 			}
 		}
 	}
-	if (Pin && Pin->PinFriendlyName.ToString().Len() > 0)
+	if (Pin)
 	{
-		NodePinStr += TEXT(" Pin: ") + Pin->PinFriendlyName.ToString();
+		NodePinStr += TEXT(" Pin: ") + (Pin->PinFriendlyName.ToString().Len() > 0 ? Pin->PinFriendlyName.ToString() : Pin->GetName());
 		NodePinSuffix = TEXT(" - ");
 	}
 	
@@ -5388,14 +5578,17 @@ FString FHlslNiagaraTranslator::GetPropertyHlslTypeName(const UProperty* Propert
 		const UEnumProperty* EnumProp = Cast<const UEnumProperty>(Property);
 		return "int";
 	}
+	else if (Property->IsA(UBoolProperty::StaticClass()))
+	{
+		return "bool";
+	}
 	else
 	{
-		check(false);	// unknown type
-		return TEXT("UnknownType");
+		return TEXT("");
 	}
 }
 
-FString FHlslNiagaraTranslator::BuildHLSLStructDecl(FNiagaraTypeDefinition Type)
+FString FHlslNiagaraTranslator::BuildHLSLStructDecl(FNiagaraTypeDefinition Type, FText& OutErrorMessage)
 {
 	if (!IsBuiltInHlslType(Type))
 	{
@@ -5405,7 +5598,14 @@ FString FHlslNiagaraTranslator::BuildHLSLStructDecl(FNiagaraTypeDefinition Type)
 		for (TFieldIterator<UProperty> PropertyIt(Type.GetStruct(), EFieldIteratorFlags::IncludeSuper); PropertyIt; ++PropertyIt)
 		{
 			UProperty* Property = *PropertyIt;
-			Decl += TEXT("\t") + GetPropertyHlslTypeName(Property) + TEXT(" ") + Property->GetName() + TEXT(";\n");
+			FString PropertyTypeName = GetPropertyHlslTypeName(Property);
+			if (PropertyTypeName.IsEmpty())
+			{
+				OutErrorMessage = FText::Format(LOCTEXT("UnknownPropertyTypeErrorFormat", "Failed to build hlsl struct declaration for type {0}.  Property {1} has an unsuported type {2}."), 
+					FText::FromString(Type.GetName()), Property->GetDisplayNameText(), FText::FromString(Property->GetClass()->GetName()));
+				return TEXT("");
+			}
+			Decl += TEXT("\t") + PropertyTypeName + TEXT(" ") + Property->GetName() + TEXT(";\n");
 		}
 		Decl += "};\n\n";
 		return Decl;
diff --git a/Engine/Plugins/FX/Niagara/Source/NiagaraEditor/Private/NiagaraHlslTranslator.h b/Engine/Plugins/FX/Niagara/Source/NiagaraEditor/Private/NiagaraHlslTranslator.h
--- a/Engine/Plugins/FX/Niagara/Source/NiagaraEditor/Private/NiagaraHlslTranslator.h
+++ b/Engine/Plugins/FX/Niagara/Source/NiagaraEditor/Private/NiagaraHlslTranslator.h
@@ -304,7 +304,7 @@ protected:
 
 	/**
 	Map of Pins to compiled code chunks. Allows easy reuse of previously compiled pins.
-	A stack so that we can track pin reuse within function calls but not have cached pins cross talk with subsequent calls to the same funciton.
+	A stack so that we can track pin reuse within function calls but not have cached pins cross talk with subsequent calls to the same function.
 	*/
 	TArray<TMap<UEdGraphPin*, int32>> PinToCodeChunks;
 
@@ -446,7 +446,7 @@ public:
 	virtual void FunctionCall(UNiagaraNodeFunctionCall* FunctionNode, TArray<int32>& Inputs, TArray<int32>& Outputs);
 
 	virtual void Convert(class UNiagaraNodeConvert* Convert, TArray <int32>& Inputs, TArray<int32>& Outputs);
-	virtual void If(TArray<FNiagaraVariable>& Vars, int32 Condition, TArray<int32>& PathA, TArray<int32>& PathB, TArray<int32>& Outputs);
+	virtual void If(class UNiagaraNodeIf* IfNode, TArray<FNiagaraVariable>& Vars, int32 Condition, TArray<int32>& PathA, TArray<int32>& PathB, TArray<int32>& Outputs);
 
 	virtual void Error(FText ErrorText, const UNiagaraNode* Node, const UEdGraphPin* Pin);
 	virtual void Warning(FText WarningText, const UNiagaraNode* Node, const UEdGraphPin* Pin);
@@ -457,11 +457,15 @@ public:
 	virtual ENiagaraScriptUsage GetTargetUsage() const;
 	FGuid GetTargetUsageId() const;
 	virtual ENiagaraScriptUsage GetCurrentUsage() const;
+	virtual ENiagaraSimTarget GetSimulationTarget() const
+	{
+		return CompilationTarget;
+	}
 
 	static bool IsBuiltInHlslType(FNiagaraTypeDefinition Type);
 	static FString GetStructHlslTypeName(FNiagaraTypeDefinition Type);
 	static FString GetPropertyHlslTypeName(const UProperty* Property);
-	static FString BuildHLSLStructDecl(FNiagaraTypeDefinition Type);
+	static FString BuildHLSLStructDecl(FNiagaraTypeDefinition Type, FText& OutErrorMessage);
 	static FString GetHlslDefaultForType(FNiagaraTypeDefinition Type);
 	static bool IsHlslBuiltinVector(FNiagaraTypeDefinition Type);
 	static TArray<FName> ConditionPropertyPath(const FNiagaraTypeDefinition& Type, const TArray<FName>& InPath);
@@ -523,6 +527,8 @@ private:
 	// Register an attribute in its namespaced form
 	bool ParameterMapRegisterUniformAttributeVariable(const FNiagaraVariable& InVariable, UNiagaraNode* InNode, int32 InParamMapHistoryIdx, int32& Output);
 
+    // Checks that the Partices.ID parameter is only used if persistent IDs are active
+    void ValidateParticleIDUsage();
 	bool ValidateTypePins(UNiagaraNode* NodeToValidate);
 	void GenerateFunctionSignature(ENiagaraScriptUsage ScriptUsage, FString InName, const FString& InFullName, UNiagaraGraph* FuncGraph, TArray<int32>& Inputs, 
 		bool bHadNumericInputs, bool bHasParameterMapParameters, FNiagaraFunctionSignature& OutSig)const;
@@ -536,7 +542,7 @@ private:
 	bool RequiresInterpolation() const;
 
 	/** If OutVar can be replaced by a literal constant, it's data is initialized with the correct value and we return true. Returns false otherwise. */
-	bool GetLiteralConstantVariable(FNiagaraVariable& OutVar);
+	bool GetLiteralConstantVariable(FNiagaraVariable& OutVar) const;
 
 	/** Map of symbol names to count of times it's been used. Used for generating unique symbol names. */
 	TMap<FName, uint32> SymbolCounts;
diff --git a/Engine/Plugins/FX/Niagara/Source/NiagaraEditor/Private/NiagaraNode.cpp b/Engine/Plugins/FX/Niagara/Source/NiagaraEditor/Private/NiagaraNode.cpp
--- a/Engine/Plugins/FX/Niagara/Source/NiagaraEditor/Private/NiagaraNode.cpp
+++ b/Engine/Plugins/FX/Niagara/Source/NiagaraEditor/Private/NiagaraNode.cpp
@@ -421,6 +421,102 @@ UEdGraphPin* UNiagaraNode::GetPinByPersistentGuid(const FGuid& InPersistentGuid)
 	return nullptr;
 }
 
+void UNiagaraNode::NumericResolutionByPins(const UEdGraphSchema_Niagara* Schema, TArray<UEdGraphPin*>& InputPins, TArray<UEdGraphPin*>& OutputPins, 
+	bool bFixInline, TMap<TPair<FGuid, UEdGraphNode*>, FNiagaraTypeDefinition>* PinCache)
+{
+	TArray<FNiagaraTypeDefinition> InputTypes;
+	TArray<FNiagaraTypeDefinition> OutputTypes;
+
+	TArray<FNiagaraTypeDefinition> NonNumericInputs;
+	for (UEdGraphPin* InputPin : InputPins)
+	{
+		FNiagaraTypeDefinition InputPinType = Schema->PinToTypeDefinition(InputPin);
+		if (InputPin->PinType.PinCategory == UEdGraphSchema_Niagara::PinCategoryType)
+		{
+
+			// If the input pin is the generic numeric type set it to the type of the linked output pin which should have been processed already.
+			if (InputPinType == FNiagaraTypeDefinition::GetGenericNumericDef() && InputPin->LinkedTo.Num() == 1)
+			{
+				UEdGraphPin* InputPinLinkedPin = UNiagaraNode::TraceOutputPin(InputPin->LinkedTo[0]);
+				FNiagaraTypeDefinition InputPinLinkedPinType = Schema->PinToTypeDefinition(InputPinLinkedPin);
+				if (InputPinLinkedPinType.IsValid())
+				{
+					if (InputPinLinkedPinType == FNiagaraTypeDefinition::GetGenericNumericDef() && PinCache)
+					{
+						FNiagaraTypeDefinition* FoundDef = PinCache->Find(TPair<FGuid, UEdGraphNode*>(InputPinLinkedPin->PinId, InputPinLinkedPin->GetOwningNode()));
+						if (FoundDef && FoundDef->IsValid())
+						{
+							InputPinLinkedPinType = *FoundDef;
+						}
+					}
+
+					// Only update the input pin type if the linked pin type is valid.
+					FEdGraphPinType PinType = Schema->TypeDefinitionToPinType(InputPinLinkedPinType);
+					if (bFixInline)
+					{
+						InputPin->PinType = PinType;
+					}
+					InputPinType = InputPinLinkedPinType;
+				}
+			}
+
+			if (InputPinType != FNiagaraTypeDefinition::GetGenericNumericDef())
+			{
+				NonNumericInputs.Add(InputPinType);
+			}
+		}
+		InputTypes.Add(InputPinType);
+	}
+
+	// Fix up numeric outputs based on the inputs.
+	for (UEdGraphPin* OutputPin : OutputPins)
+	{
+		FNiagaraTypeDefinition OutputPinType = Schema->PinToTypeDefinition(OutputPin);
+		if (NonNumericInputs.Num() > 0 && GetNumericOutputTypeSelectionMode() != ENiagaraNumericOutputTypeSelectionMode::None)
+		{
+			FNiagaraTypeDefinition OutputNumericType = FNiagaraTypeDefinition::GetNumericOutputType(NonNumericInputs, GetNumericOutputTypeSelectionMode());
+			if (OutputNumericType != FNiagaraTypeDefinition::GetGenericNumericDef())
+			{
+				if (OutputPinType == FNiagaraTypeDefinition::GetGenericNumericDef())
+				{
+					FEdGraphPinType PinType = Schema->TypeDefinitionToPinType(OutputNumericType);
+					if (bFixInline)
+					{
+						OutputPin->PinType = PinType;
+					}
+					OutputPinType = OutputNumericType;
+				}
+			}
+		}
+		OutputTypes.Add(OutputPinType);
+	}
+
+
+	if (PinCache)
+	{
+		int32 j;
+		for (j = 0; j < InputPins.Num(); j++)
+		{
+			PinCache->Add(TPair<FGuid, UEdGraphNode*>(InputPins[j]->PinId, InputPins[j]->GetOwningNode()), InputTypes[j]);
+		}
+		for (j = 0; j < OutputPins.Num(); j++)
+		{
+			PinCache->Add(TPair<FGuid, UEdGraphNode*>(OutputPins[j]->PinId, OutputPins[j]->GetOwningNode()), OutputTypes[j]);
+		}
+	}
+}
+
+
+void UNiagaraNode::ResolveNumerics(const UEdGraphSchema_Niagara* Schema, bool bSetInline, TMap<TPair<FGuid, UEdGraphNode*>, FNiagaraTypeDefinition>* PinCache)
+{
+	// Fix up numeric input pins and keep track of numeric types to decide the output type.
+	TArray<UEdGraphPin*> InputPins;
+	GetInputPins(InputPins);
+	TArray<UEdGraphPin*> OutputPins;
+	GetOutputPins(OutputPins);
+	NumericResolutionByPins(Schema, InputPins, OutputPins, bSetInline, PinCache);
+}
+
 ENiagaraNumericOutputTypeSelectionMode UNiagaraNode::GetNumericOutputTypeSelectionMode() const
 {
 	return ENiagaraNumericOutputTypeSelectionMode::None;
