Merging //UE4/Dev-Main to Dev-Build (//UE4/Dev-Build)

diff --git a/Engine/Plugins/FX/Niagara/Source/NiagaraEditor/Private/NiagaraNodeCustomHlsl.cpp b/Engine/Plugins/FX/Niagara/Source/NiagaraEditor/Private/NiagaraNodeCustomHlsl.cpp
--- a/Engine/Plugins/FX/Niagara/Source/NiagaraEditor/Private/NiagaraNodeCustomHlsl.cpp
+++ b/Engine/Plugins/FX/Niagara/Source/NiagaraEditor/Private/NiagaraNodeCustomHlsl.cpp
@@ -63,7 +63,7 @@ bool UNiagaraNodeCustomHlsl::GetTokens(TArray<FString>& OutTokens) const
 		return false;
 	}
 	
-	FString InputVars = TEXT(";/*+-)(?:, \t\n");
+	FString InputVars = TEXT(";/*+-)(?:, []\t\n");
 	int32 LastValidIdx = INDEX_NONE;
 	bool bComment = false;
 	int32 TargetLength = HlslData.Len();
@@ -158,7 +158,7 @@ void UNiagaraNodeCustomHlsl::InitAsCustomHlslDynamicInput(const FNiagaraTypeDefi
 	Modify();
 	ReallocatePins();
 	RequestNewTypedPin(EGPD_Input, FNiagaraTypeDefinition::GetParameterMapDef(), FName("Map"));
-	RequestNewTypedPin(EGPD_Output, OutputType, FName("Output"));
+	RequestNewTypedPin(EGPD_Output, OutputType, FName("CustomHLSLOutput"));
 	ScriptUsage = ENiagaraScriptUsage::DynamicInput;
 }
 
diff --git a/Engine/Plugins/FX/Niagara/Source/NiagaraEditor/Private/NiagaraNodeFunctionCall.cpp b/Engine/Plugins/FX/Niagara/Source/NiagaraEditor/Private/NiagaraNodeFunctionCall.cpp
--- a/Engine/Plugins/FX/Niagara/Source/NiagaraEditor/Private/NiagaraNodeFunctionCall.cpp
+++ b/Engine/Plugins/FX/Niagara/Source/NiagaraEditor/Private/NiagaraNodeFunctionCall.cpp
@@ -194,6 +194,17 @@ void UNiagaraNodeFunctionCall::AllocateDefaultPins()
 	{
 		ComputeNodeName();
 	}
+
+}
+
+// Returns true if this node is deprecated
+bool UNiagaraNodeFunctionCall::IsDeprecated() const
+{
+	if (FunctionScript)
+	{
+		return FunctionScript->bDeprecated;
+	}
+	return false;
 }
 
 FText UNiagaraNodeFunctionCall::GetNodeTitle(ENodeTitleType::Type TitleType) const
@@ -316,6 +327,21 @@ void UNiagaraNodeFunctionCall::Compile(class FHlslNiagaraTranslator* Translator,
 	UNiagaraGraph* CallerGraph = GetNiagaraGraph();
 	if (FunctionScript)
 	{
+		if (FunctionScript->bDeprecated && IsNodeEnabled())
+		{
+			if (FunctionScript->DeprecationRecommendation)
+			{
+				Translator->Warning(FText::Format(LOCTEXT("DeprecationErrorFmtRecommendation", "Function call \"{0}\" is deprecated. Please use {1} instead."), FText::FromString(FunctionDisplayName),
+					FText::FromString(FunctionScript->DeprecationRecommendation->GetPathName())),
+					this, nullptr);
+			}
+			else
+			{
+				Translator->Warning(FText::Format(LOCTEXT("DeprecationErrorFmtUnknown", "Function call \"{0}\" is deprecated. No recommendation was provided."), FText::FromString(FunctionDisplayName)),
+					this, nullptr);
+			}
+		}
+
 		TArray<UEdGraphPin*> CallerInputPins;
 		GetInputPins(CallerInputPins);
 
@@ -480,6 +506,22 @@ bool UNiagaraNodeFunctionCall::RefreshFromExternalChanges()
 				bReload = true;
 			}
 		}
+
+		if (FunctionScript->bDeprecated)
+		{
+			bHasCompilerMessage = true;
+			ErrorType = EMessageSeverity::Warning;
+			FFormatNamedArguments Args;
+			Args.Add(TEXT("NodeName"), GetNodeTitle(ENodeTitleType::ListView));
+			Args.Add(TEXT("Recommendation"), FunctionScript->DeprecationRecommendation != nullptr ? FText::FromString(FunctionScript->DeprecationRecommendation->GetPathName()) : LOCTEXT("UnspecifiedName","Unspecified"));
+			ErrorMsg = FText::Format(LOCTEXT("DeprecationWarning", "{NodeName} is deprecated. Suggested replacement: {Recommendation}"), Args).ToString();
+		}
+		else
+		{
+			bHasCompilerMessage = false;
+			ErrorMsg = FString();
+		}
+
 	}
 	else
 	{
@@ -528,7 +570,7 @@ void UNiagaraNodeFunctionCall::GatherExternalDependencyIDs(ENiagaraScriptUsage I
 		// We don't know which graph type we're referencing, so we try them all... may need to replace this with something faster in the future.
 		if (FunctionGraph)
 		{
-			for (int32 i = (int32)ENiagaraScriptUsage::Function; i <= (int32)ENiagaraScriptUsage::Module; i++)
+			for (int32 i = (int32)ENiagaraScriptUsage::Function; i <= (int32)ENiagaraScriptUsage::DynamicInput; i++)
 			{
 				FGuid FoundGuid = FunctionGraph->GetCompileID((ENiagaraScriptUsage)i, FGuid(0, 0, 0, 0));
 				if (FoundGuid.IsValid())
diff --git a/Engine/Plugins/FX/Niagara/Source/NiagaraEditor/Private/NiagaraNodeIf.cpp b/Engine/Plugins/FX/Niagara/Source/NiagaraEditor/Private/NiagaraNodeIf.cpp
--- a/Engine/Plugins/FX/Niagara/Source/NiagaraEditor/Private/NiagaraNodeIf.cpp
+++ b/Engine/Plugins/FX/Niagara/Source/NiagaraEditor/Private/NiagaraNodeIf.cpp
@@ -74,7 +74,9 @@ void UNiagaraNodeIf::PostLoad()
 
 bool UNiagaraNodeIf::AllowNiagaraTypeForAddPin(const FNiagaraTypeDefinition& InType)
 {
-	return Super::AllowNiagaraTypeForAddPin(InType) && InType != FNiagaraTypeDefinition::GetParameterMapDef();
+	// Explicitly allow Numeric types, and explicitly disallow ParameterMap types
+	
+	return (Super::AllowNiagaraTypeForAddPin(InType) || InType == FNiagaraTypeDefinition::GetGenericNumericDef()) && InType != FNiagaraTypeDefinition::GetParameterMapDef();
 }
 
 void UNiagaraNodeIf::AllocateDefaultPins()
@@ -137,7 +139,29 @@ void UNiagaraNodeIf::Compile(class FHlslNiagaraTranslator* Translator, TArray<in
 
 		PathB.Add(Translator->CompilePin(Pins[PinIdx++]));
 	}
-	Translator->If(OutputVars, Condition, PathA, PathB, Outputs);
+	Translator->If(this, OutputVars, Condition, PathA, PathB, Outputs);
+}
+
+ENiagaraNumericOutputTypeSelectionMode UNiagaraNodeIf::GetNumericOutputTypeSelectionMode() const
+{
+	return ENiagaraNumericOutputTypeSelectionMode::Largest;
+}
+
+
+void UNiagaraNodeIf::ResolveNumerics(const UEdGraphSchema_Niagara* Schema, bool bSetInline, TMap<TPair<FGuid, UEdGraphNode*>, FNiagaraTypeDefinition>* PinCache)
+{
+	int32 VarStartIdx = 1;
+	for (int32 i = 0; i < OutputVars.Num(); ++i)
+	{
+		// Fix up numeric input pins and keep track of numeric types to decide the output type.
+		TArray<UEdGraphPin*> InputPins;
+		TArray<UEdGraphPin*> OutputPins;
+		
+		InputPins.Add(Pins[i + VarStartIdx]);
+		InputPins.Add(Pins[i + VarStartIdx + OutputVars.Num()]);
+		OutputPins.Add(Pins[i + VarStartIdx + 2 * OutputVars.Num()]);
+		NumericResolutionByPins(Schema, InputPins, OutputPins,  bSetInline, PinCache);
+	}
 }
 
 bool UNiagaraNodeIf::RefreshFromExternalChanges()
diff --git a/Engine/Plugins/FX/Niagara/Source/NiagaraEditor/Private/NiagaraNodeIf.h b/Engine/Plugins/FX/Niagara/Source/NiagaraEditor/Private/NiagaraNodeIf.h
--- a/Engine/Plugins/FX/Niagara/Source/NiagaraEditor/Private/NiagaraNodeIf.h
+++ b/Engine/Plugins/FX/Niagara/Source/NiagaraEditor/Private/NiagaraNodeIf.h
@@ -13,16 +13,16 @@ class UNiagaraNodeIf : public UNiagaraNodeWithDynamicPins
 public:
 
 	/** Outputs of this branch. */
-	UPROPERTY(EditAnywhere, Category=If)
+	UPROPERTY()
 	TArray<FNiagaraVariable> OutputVars;
 
-	UPROPERTY(EditAnywhere, Category = If)
+	UPROPERTY()
 	TArray<FGuid> OutputVarGuids;
 
-	UPROPERTY(EditAnywhere, Category = If)
+	UPROPERTY()
 	TArray<FGuid> InputAVarGuids;
 
-	UPROPERTY(EditAnywhere, Category = If)
+	UPROPERTY()
 	TArray<FGuid> InputBVarGuids;
 
 	//~ Begin UObject Interface
@@ -39,6 +39,8 @@ public:
 	//~ Begin UNiagaraNode Interface
 	virtual void Compile(class FHlslNiagaraTranslator* Translator, TArray<int32>& Outputs) override;
 	virtual bool RefreshFromExternalChanges() override;
+	virtual ENiagaraNumericOutputTypeSelectionMode GetNumericOutputTypeSelectionMode() const;
+	virtual void ResolveNumerics(const UEdGraphSchema_Niagara* Schema, bool bSetInline, TMap<TPair<FGuid, UEdGraphNode*>, FNiagaraTypeDefinition>* PinCache);
 	//~ End UNiagaraNode Interface
 
 protected:
@@ -49,7 +51,7 @@ protected:
 	//~ Begin EdGraphNode Interface
 	virtual void OnPinRemoved(UEdGraphPin* PinToRemove) override;
 	//~ End EdGraphNode Interface
-
+	
 	//~ Begin UNiagaraNodeWithDynamicPins Interface
 	virtual void OnNewTypedPinAdded(UEdGraphPin* NewPin) override;
 	virtual void OnPinRenamed(UEdGraphPin* RenamedPin, const FString& OldName) override;
diff --git a/Engine/Plugins/FX/Niagara/Source/NiagaraEditor/Private/NiagaraNodeOp.cpp b/Engine/Plugins/FX/Niagara/Source/NiagaraEditor/Private/NiagaraNodeOp.cpp
--- a/Engine/Plugins/FX/Niagara/Source/NiagaraEditor/Private/NiagaraNodeOp.cpp
+++ b/Engine/Plugins/FX/Niagara/Source/NiagaraEditor/Private/NiagaraNodeOp.cpp
@@ -20,6 +20,7 @@ void UNiagaraNodeOp::AllocateDefaultPins()
 	const FNiagaraOpInfo* OpInfo = FNiagaraOpInfo::GetOpInfo(OpName);
 	check(OpInfo);
 
+	// Create input pins from the op
 	for (int32 SrcIndex = 0; SrcIndex < OpInfo->Inputs.Num(); ++SrcIndex)
 	{
 		const FNiagaraOpInOutInfo& InOutInfo = OpInfo->Inputs[SrcIndex];
@@ -32,7 +33,19 @@ void UNiagaraNodeOp::AllocateDefaultPins()
 		Pin->AutogeneratedDefaultValue = InOutInfo.Default;
 		Pin->PinToolTip = InOutInfo.Description.ToString();
 	}
+
+	// Restore pins added by the user
+	for (int32 SrcIndex = 0; SrcIndex < AddedPins.Num(); ++SrcIndex)
+	{
+		FAddedPinData& OldPinData = AddedPins[SrcIndex];
+		UEdGraphPin* Pin = CreatePin(EGPD_Input, OldPinData.PinType, OldPinData.PinName);
+		check(Pin);
+		Pin->bDefaultValueIsIgnored = false;
+		Pin->bDefaultValueIsReadOnly = false;
+		Pin->bNotConnectable = false;
+	}
 	
+	// Create output pins from the op
 	for (int32 OutIdx = 0; OutIdx < OpInfo->Outputs.Num(); ++OutIdx)
 	{
 		const FNiagaraOpInOutInfo& InOutInfo = OpInfo->Outputs[OutIdx];
@@ -40,6 +53,11 @@ void UNiagaraNodeOp::AllocateDefaultPins()
 		check(Pin);
 		Pin->PinToolTip = InOutInfo.Description.ToString();
 	}
+
+	if (OpInfo->bSupportsAddedInputs)
+	{
+		CreateAddPin(EGPD_Input);
+	}
 }
 
 FText UNiagaraNodeOp::GetNodeTitle(ENodeTitleType::Type TitleType) const
@@ -77,7 +95,7 @@ void UNiagaraNodeOp::Compile(class FHlslNiagaraTranslator* Translator, TArray<in
 		return;
 	}
 
-	int32 NumInputs = OpInfo->Inputs.Num();
+	int32 NumInputs = OpInfo->bSupportsAddedInputs ? Pins.Num() : OpInfo->Inputs.Num();
 	int32 NumOutputs = OpInfo->Outputs.Num();
 
 	TArray<int32> Inputs;
@@ -85,7 +103,10 @@ void UNiagaraNodeOp::Compile(class FHlslNiagaraTranslator* Translator, TArray<in
 	for (int32 i = 0; i < NumInputs; ++i)
 	{
 		UEdGraphPin *Pin = Pins[i];
-		check(Pin->Direction == EGPD_Input);
+		if (Pin->Direction != EGPD_Input || IsAddPin(Pin))
+		{
+			continue;
+		}
 		int32 CompiledInput = Translator->CompilePin(Pin);
 		if (CompiledInput == INDEX_NONE)
 		{
@@ -141,6 +162,25 @@ void UNiagaraNodeOp::PostLoad()
 		UE_LOG(LogNiagaraEditor, Log, TEXT("OpNode: Converted %s to %s, Package: %s"), *OriginalOpName.ToString(), *OpName.ToString(), *GetOutermost()->GetName());
 	}
 
+
+	if (AllowDynamicPins())
+	{
+		bool HasAddPin = false;
+		for (UEdGraphPin* Pin : Pins)
+		{
+			if (IsAddPin(Pin))
+			{
+				HasAddPin = true;
+				break;
+			}
+		}
+		if (!HasAddPin)
+		{
+			// This adds the "plus" pin to nodes that were created before the feature was added
+			ReallocatePins();
+			UE_LOG(LogNiagaraEditor, Log, TEXT("OpNode %s, Package %s: reallocated existing pins to add new extension pin"), *OpName.ToString(), *GetOutermost()->GetName());
+		}
+	}
 }
 
 bool UNiagaraNodeOp::RefreshFromExternalChanges()
@@ -157,4 +197,109 @@ ENiagaraNumericOutputTypeSelectionMode UNiagaraNodeOp::GetNumericOutputTypeSelec
 	return OpInfo->NumericOuputTypeSelectionMode;
 }
 
+bool UNiagaraNodeOp::AllowNiagaraTypeForAddPin(const FNiagaraTypeDefinition& InType)
+{
+	const FNiagaraOpInfo* OpInfo = FNiagaraOpInfo::GetOpInfo(OpName);
+	if (!OpInfo)
+	{
+		return false;
+	}
+	if (OpInfo->AddedInputTypeRestrictions.Num() == 0)
+	{
+		return true;
+	}
+	return OpInfo->AddedInputTypeRestrictions.Contains(InType);
+}
+
+void UNiagaraNodeOp::OnPinRemoved(UEdGraphPin* PinToRemove)
+{
+	auto FindPredicate = [=](const FAddedPinData& PinData)
+	{
+		return PinData.PinName == PinToRemove->PinName && PinData.PinType == PinToRemove->PinType; 
+	};
+	AddedPins.RemoveAll(FindPredicate);
+	ReallocatePins();
+}
+
+bool UNiagaraNodeOp::AllowDynamicPins() const
+{
+	const FNiagaraOpInfo* OpInfo = FNiagaraOpInfo::GetOpInfo(OpName);
+	return OpInfo && OpInfo->bSupportsAddedInputs;
+}
+
+void UNiagaraNodeOp::OnNewTypedPinAdded(UEdGraphPin* NewPin)
+{
+	FName UniqueName = GetUniqueAdditionalPinName();
+
+	FAddedPinData PinData;
+	PinData.PinType = NewPin->PinType;
+	PinData.PinName = UniqueName;
+	NewPin->PinName = UniqueName;
+	AddedPins.Add(PinData);
+}
+
+void UNiagaraNodeOp::OnPinRenamed(UEdGraphPin* RenamedPin, const FString& OldName)
+{
+	TSet<FName> ExistingNames;
+	for (const FAddedPinData& PinData : AddedPins)
+	{
+		ExistingNames.Add(PinData.PinName);
+	}
+	FName UniqueName = FNiagaraUtilities::GetUniqueName(*RenamedPin->PinName.ToString(), ExistingNames);
+	FName OldPinName(*OldName);
+
+	for (FAddedPinData& PinData : AddedPins)
+	{
+		if (PinData.PinName == OldPinName && PinData.PinType == RenamedPin->PinType)
+		{
+			PinData.PinName = UniqueName;
+			RenamedPin->PinName = UniqueName;
+			break;
+		}
+	}
+}
+
+bool UNiagaraNodeOp::CanRemovePin(const UEdGraphPin* Pin) const
+{
+	if (!Pin || Pin->Direction != EGPD_Input || !Super::CanRemovePin(Pin))
+	{
+		return false;
+	}
+
+	// check if the pin was added by the user, only those can be removed
+	for (int32 SrcIndex = 0; SrcIndex < AddedPins.Num(); ++SrcIndex)
+	{
+		const FAddedPinData& PinData = AddedPins[SrcIndex];
+		if (Pin->PinType == PinData.PinType && Pin->PinName == PinData.PinName)
+		{
+			return true;
+		}
+	}
+	return false;
+}
+
+FName UNiagaraNodeOp::GetUniqueAdditionalPinName() const
+{
+	// count existing input pins
+	FString Name;
+	TSet<FName> ExistingNames;
+	for (UEdGraphPin* Pin : Pins)
+	{
+		if (Pin->Direction == EEdGraphPinDirection::EGPD_Input && !IsAddPin(Pin))
+		{
+			ExistingNames.Add(Pin->PinName);
+		}
+	}
+	
+	// create a new name based on the existing pins (A, B, ... AA, AB, ...)
+	static FString Alphabet = TEXT("ABCDEFGHIJKLMNOPQRSTUVWXYZ");
+	for (int32 Remaining = ExistingNames.Num(); Remaining > 0; Remaining = Remaining / 26)
+	{
+		int32 CharIndex = Remaining < 26 ? Remaining - 1 : Remaining % 26;
+		Name = FString() + (*Alphabet)[CharIndex] + Name;
+	}
+
+	return FNiagaraUtilities::GetUniqueName(*Name, ExistingNames);
+}
+
 #undef LOCTEXT_NAMESPACE
diff --git a/Engine/Plugins/FX/Niagara/Source/NiagaraEditor/Private/NiagaraNodeParameterMapBase.cpp b/Engine/Plugins/FX/Niagara/Source/NiagaraEditor/Private/NiagaraNodeParameterMapBase.cpp
--- a/Engine/Plugins/FX/Niagara/Source/NiagaraEditor/Private/NiagaraNodeParameterMapBase.cpp
+++ b/Engine/Plugins/FX/Niagara/Source/NiagaraEditor/Private/NiagaraNodeParameterMapBase.cpp
@@ -91,7 +91,8 @@ TArray<FNiagaraParameterMapHistory> UNiagaraNodeParameterMapBase::GetParameterMa
 
 bool UNiagaraNodeParameterMapBase::AllowNiagaraTypeForAddPin(const FNiagaraTypeDefinition& InType)
 {
-	return InType != FNiagaraTypeDefinition::GetGenericNumericDef();
+	return InType != FNiagaraTypeDefinition::GetGenericNumericDef() &&
+		InType != FNiagaraTypeDefinition::GetParameterMapDef();
 }
 
 FText UNiagaraNodeParameterMapBase::GetPinDescriptionText(UEdGraphPin* Pin) const
diff --git a/Engine/Plugins/FX/Niagara/Source/NiagaraEditor/Private/NiagaraNodeReroute.cpp b/Engine/Plugins/FX/Niagara/Source/NiagaraEditor/Private/NiagaraNodeReroute.cpp
--- a/Engine/Plugins/FX/Niagara/Source/NiagaraEditor/Private/NiagaraNodeReroute.cpp
+++ b/Engine/Plugins/FX/Niagara/Source/NiagaraEditor/Private/NiagaraNodeReroute.cpp
@@ -24,6 +24,12 @@ void UNiagaraNodeReroute::PostEditChangeProperty(FPropertyChangedEvent& Property
 	}
 }
 
+
+ENiagaraNumericOutputTypeSelectionMode UNiagaraNodeReroute::GetNumericOutputTypeSelectionMode() const
+{
+	return ENiagaraNumericOutputTypeSelectionMode::Largest;
+}
+
 void UNiagaraNodeReroute::PostLoad()
 {
 	Super::PostLoad();
diff --git a/Engine/Plugins/FX/Niagara/Source/NiagaraEditor/Private/NiagaraNodeReroute.h b/Engine/Plugins/FX/Niagara/Source/NiagaraEditor/Private/NiagaraNodeReroute.h
--- a/Engine/Plugins/FX/Niagara/Source/NiagaraEditor/Private/NiagaraNodeReroute.h
+++ b/Engine/Plugins/FX/Niagara/Source/NiagaraEditor/Private/NiagaraNodeReroute.h
@@ -34,6 +34,7 @@ public:
 	virtual bool RefreshFromExternalChanges() override;
 	virtual void PinConnectionListChanged(UEdGraphPin* Pin) override;
 	void BuildParameterMapHistory(FNiagaraParameterMapHistoryBuilder& OutHistory, bool bRecursive) override;
+	virtual ENiagaraNumericOutputTypeSelectionMode GetNumericOutputTypeSelectionMode() const;
 
 	/** Trace to an output pin that is not a reroute node output pin. If the reroute nodes resulted in a dead end and no output pin was found return nullptr.*/
 	virtual UEdGraphPin* GetTracedOutputPin(UEdGraphPin* LocallyOwnedOutputPin) const override;
diff --git a/Engine/Plugins/FX/Niagara/Source/NiagaraEditor/Private/NiagaraNodeSimTargetSelector.cpp b/Engine/Plugins/FX/Niagara/Source/NiagaraEditor/Private/NiagaraNodeSimTargetSelector.cpp
new file mode 100644
--- /dev/null
+++ b/Engine/Plugins/FX/Niagara/Source/NiagaraEditor/Private/NiagaraNodeSimTargetSelector.cpp
@@ -0,0 +1,122 @@
+// Copyright 1998-2019 Epic Games, Inc. All Rights Reserved.
+
+#include "NiagaraNodeSimTargetSelector.h"
+#include "NiagaraEditorUtilities.h"
+#include "NiagaraHlslTranslator.h"
+
+#define LOCTEXT_NAMESPACE "NiagaraNodeSimTargetSelector"
+
+UNiagaraNodeSimTargetSelector::UNiagaraNodeSimTargetSelector(const FObjectInitializer& ObjectInitializer)
+	: Super(ObjectInitializer)
+{
+}
+
+void UNiagaraNodeSimTargetSelector::AllocateDefaultPins()
+{
+	const UEdGraphSchema_Niagara* Schema = GetDefault<UEdGraphSchema_Niagara>();
+
+	// create all the cpu input pins
+	for (FNiagaraVariable& Var : OutputVars)
+	{
+		FEdGraphPinType PinType = Schema->TypeDefinitionToPinType(Var.GetType());
+		CreatePin(EGPD_Input, PinType, *(Var.GetName().ToString() + TEXT(" if CPU VM")));
+	}
+
+	// create all the gpu input pins
+	for (FNiagaraVariable& Var : OutputVars)
+	{
+		FEdGraphPinType PinType = Schema->TypeDefinitionToPinType(Var.GetType());
+		CreatePin(EGPD_Input, PinType, *(Var.GetName().ToString() + TEXT(" if GPU Shader")));
+	}
+
+	// create the output pins
+	for (int32 Index = 0; Index < OutputVars.Num(); Index++)
+	{
+		const FNiagaraVariable& Var = OutputVars[Index];
+		UEdGraphPin* NewPin = CreatePin(EGPD_Output, Schema->TypeDefinitionToPinType(Var.GetType()), Var.GetName());
+		NewPin->PersistentGuid = OutputVarGuids[Index];
+	}
+
+	CreateAddPin(EGPD_Output);
+}
+
+void UNiagaraNodeSimTargetSelector::InsertInputPinsFor(const FNiagaraVariable& Var)
+{
+	const UEdGraphSchema_Niagara* Schema = GetDefault<UEdGraphSchema_Niagara>();
+
+	TArray<UEdGraphPin*> OldPins(Pins);
+	Pins.Reset(Pins.Num() + 2);
+
+	// Create the inputs for each simulation target.
+	for (int64 i = 0; i < 2; i++)
+	{
+		// Add the previous input pins
+		for (int32 k = 0; k < OutputVars.Num() - 1; k++)
+		{
+			Pins.Add(OldPins[k]);
+		}
+		OldPins.RemoveAt(0, OutputVars.Num() - 1);
+
+		// Add the new input pin
+		const FString PathSuffix = i == 0 ? TEXT(" if CPU VM") : TEXT(" if GPU Shader");
+		CreatePin(EGPD_Input, Schema->TypeDefinitionToPinType(Var.GetType()), *(Var.GetName().ToString() + PathSuffix));
+	}
+
+	// Move the rest of the old pins over
+	Pins.Append(OldPins);
+}
+
+void UNiagaraNodeSimTargetSelector::Compile(class FHlslNiagaraTranslator* Translator, TArray<int32>& Outputs)
+{
+	TArray<UEdGraphPin*> InputPins;
+	GetInputPins(InputPins);
+	TArray<UEdGraphPin*> OutputPins;
+	GetOutputPins(OutputPins);
+
+	ENiagaraSimTarget SimulationTarget = Translator->GetSimulationTarget();
+	int32 VarIdx;
+	if (SimulationTarget == ENiagaraSimTarget::CPUSim)
+	{
+		VarIdx = 0;
+	}
+	else if (SimulationTarget == ENiagaraSimTarget::GPUComputeSim)
+	{
+		VarIdx = InputPins.Num() / 2;
+	}
+	else
+	{
+		Translator->Error(LOCTEXT("InvalidSimTarget", "Unknown simulation target"), this, nullptr);
+		return;
+	}
+
+	Outputs.SetNumUninitialized(OutputPins.Num());
+	for (int32 i = 0; i < OutputVars.Num(); i++)
+	{
+		int32 InputIdx = Translator->CompilePin(InputPins[VarIdx + i]);
+		Outputs[i] = InputIdx;
+	}
+	check(this->IsAddPin(OutputPins[OutputPins.Num() - 1]));
+	Outputs[OutputPins.Num() - 1] = INDEX_NONE;
+}
+
+UEdGraphPin* UNiagaraNodeSimTargetSelector::GetPassThroughPin(const UEdGraphPin* LocallyOwnedOutputPin, ENiagaraScriptUsage MasterUsage) const 
+{
+	return nullptr;
+}
+
+void UNiagaraNodeSimTargetSelector::BuildParameterMapHistory(FNiagaraParameterMapHistoryBuilder& OutHistory, bool bRecursive)
+{
+	UNiagaraNode::BuildParameterMapHistory(OutHistory, bRecursive);
+}
+
+FText UNiagaraNodeSimTargetSelector::GetTooltipText() const
+{
+	return LOCTEXT("SimTargetSelectorDesc", "If the simulation target matches, then the traversal will follow that path.");
+}
+
+FText UNiagaraNodeSimTargetSelector::GetNodeTitle(ENodeTitleType::Type TitleType) const
+{
+	return LOCTEXT("SimTargetSelectorTitle", "Select by Simulation Target");
+}
+
+#undef LOCTEXT_NAMESPACE
diff --git a/Engine/Plugins/FX/Niagara/Source/NiagaraEditor/Private/NiagaraNodeSimTargetSelector.h b/Engine/Plugins/FX/Niagara/Source/NiagaraEditor/Private/NiagaraNodeSimTargetSelector.h
new file mode 100644
--- /dev/null
+++ b/Engine/Plugins/FX/Niagara/Source/NiagaraEditor/Private/NiagaraNodeSimTargetSelector.h
@@ -0,0 +1,30 @@
+// Copyright 1998-2019 Epic Games, Inc. All Rights Reserved.
+
+#pragma once
+#include "EdGraphSchema_Niagara.h"
+#include "NiagaraNodeUsageSelector.h"
+#include "NiagaraNodeSimTargetSelector.generated.h"
+
+UCLASS(MinimalAPI)
+class UNiagaraNodeSimTargetSelector: public UNiagaraNodeUsageSelector
+{
+	GENERATED_UCLASS_BODY()
+
+public:	
+	//~ Begin EdGraphNode Interface
+	virtual void AllocateDefaultPins() override;
+	virtual FText GetTooltipText() const override;
+	virtual FText GetNodeTitle(ENodeTitleType::Type TitleType) const override;
+	//~ End EdGraphNode Interface
+
+	//~ Begin UNiagaraNode Interface
+	virtual void Compile(class FHlslNiagaraTranslator* Translator, TArray<int32>& Outputs) override;
+	void BuildParameterMapHistory(FNiagaraParameterMapHistoryBuilder& OutHistory, bool bRecursive) override;
+	virtual UEdGraphPin* GetPassThroughPin(const UEdGraphPin* LocallyOwnedOutputPin, ENiagaraScriptUsage MasterUsage) const override;
+	//~ End UNiagaraNode Interface
+
+protected:
+	//~ Begin UNiagaraNodeUsageSelector Interface
+	virtual void InsertInputPinsFor(const FNiagaraVariable& Var) override;
+	//~ End UNiagaraNodeUsageSelector Interface
+};
diff --git a/Engine/Plugins/FX/Niagara/Source/NiagaraEditor/Private/NiagaraNodeUsageSelector.cpp b/Engine/Plugins/FX/Niagara/Source/NiagaraEditor/Private/NiagaraNodeUsageSelector.cpp
--- a/Engine/Plugins/FX/Niagara/Source/NiagaraEditor/Private/NiagaraNodeUsageSelector.cpp
+++ b/Engine/Plugins/FX/Niagara/Source/NiagaraEditor/Private/NiagaraNodeUsageSelector.cpp
@@ -32,10 +32,37 @@ bool UNiagaraNodeUsageSelector::AllowNiagaraTypeForAddPin(const FNiagaraTypeDefi
 	return Super::AllowNiagaraTypeForAddPin(InType) && InType != FNiagaraTypeDefinition::GetParameterMapDef();
 }
 
-void UNiagaraNodeUsageSelector::AllocateDefaultPins()
+void UNiagaraNodeUsageSelector::InsertInputPinsFor(const FNiagaraVariable& Var)
 {
 	const UEdGraphSchema_Niagara* Schema = GetDefault<UEdGraphSchema_Niagara>();
+	UEnum* ENiagaraScriptGroupEnum = FindObject<UEnum>(ANY_PACKAGE, TEXT("ENiagaraScriptGroup"), true);
+	int64 GroupCount = (int64)ENiagaraScriptGroup::Max;
+
+	TArray<UEdGraphPin*> OldPins(Pins);
+	Pins.Reset(Pins.Num() + GroupCount);
+
+	// Create the inputs for each path.
+	for (int64 i = 0; i < GroupCount; i++)
+	{
+		// Add the previous input pins
+		for (int32 k = 0; k < OutputVars.Num() - 1; k++)
+		{
+			Pins.Add(OldPins[k]);
+		}
+		OldPins.RemoveAt(0, OutputVars.Num() - 1);
+
+		// Add the new input pin
+		const FString PathSuffix = ENiagaraScriptGroupEnum ? (FString::Printf(TEXT(" if %s"), *ENiagaraScriptGroupEnum->GetNameStringByValue((int64)i))) : TEXT("Error Unknown!");
+		CreatePin(EGPD_Input, Schema->TypeDefinitionToPinType(Var.GetType()), *(Var.GetName().ToString() + PathSuffix));
+	}
 
+	// Move the rest of the old pins over
+	Pins.Append(OldPins);
+}
+
+void UNiagaraNodeUsageSelector::AllocateDefaultPins()
+{
+	const UEdGraphSchema_Niagara* Schema = GetDefault<UEdGraphSchema_Niagara>();
 	UEnum* ENiagaraScriptGroupEnum = FindObject<UEnum>(ANY_PACKAGE, TEXT("ENiagaraScriptGroup"), true);
 
 	//Create the inputs for each path.
@@ -137,6 +164,33 @@ UEdGraphPin* UNiagaraNodeUsageSelector::GetPassThroughPin(const UEdGraphPin* Loc
 	return nullptr;
 }
 
+void UNiagaraNodeUsageSelector::AppendFunctionAliasForContext(const FNiagaraGraphFunctionAliasContext& InFunctionAliasContext, FString& InOutFunctionAlias)
+{
+	FString UsageString;
+	switch (InFunctionAliasContext.CompileUsage)
+	{
+	case ENiagaraScriptUsage::SystemSpawnScript:
+	case ENiagaraScriptUsage::SystemUpdateScript:
+		UsageString = "System";
+		break;
+	case ENiagaraScriptUsage::EmitterSpawnScript:
+	case ENiagaraScriptUsage::EmitterUpdateScript:
+		UsageString = "Emitter";
+		break;
+	case ENiagaraScriptUsage::ParticleSpawnScript:
+	case ENiagaraScriptUsage::ParticleUpdateScript:
+	case ENiagaraScriptUsage::ParticleEventScript:
+	case ENiagaraScriptUsage::ParticleGPUComputeScript:
+		UsageString = "Particle";
+		break;
+	}
+
+	if (UsageString.IsEmpty() == false)
+	{
+		InOutFunctionAlias += "_" + UsageString;
+	}
+}
+
 void UNiagaraNodeUsageSelector::BuildParameterMapHistory(FNiagaraParameterMapHistoryBuilder& OutHistory, bool bRecursive)
 {
 	const UEdGraphSchema_Niagara* Schema = CastChecked<UEdGraphSchema_Niagara>(GetSchema());
@@ -208,14 +262,17 @@ void UNiagaraNodeUsageSelector::OnNewTypedPinAdded(UEdGraphPin* NewPin)
 	{
 		OutputNames.Add(Output.GetName());
 	}
-	FName OutputName = FNiagaraUtilities::GetUniqueName(*OutputType.GetNameText().ToString(), OutputNames);
-
+	FName OutputName = FNiagaraUtilities::GetUniqueName(*NewPin->GetName(), OutputNames);
+	NewPin->PinName = OutputName;
 	FGuid Guid = AddOutput(OutputType, OutputName);
 
 	// Update the pin's data too so that it's connection is maintained after reallocating.
 	NewPin->PersistentGuid = Guid;
 
-	ReallocatePins();
+	// We cannot just reallocate the pins here, because that invalidates all pins of this node (including
+	// the NewPin parameter). If the calling method tries to access the provided new pin afterwards, it
+	// runs into a nullptr error (e.g. when called by drag and drop).
+	InsertInputPinsFor(OutputVars.Last());
 }
 
 void UNiagaraNodeUsageSelector::OnPinRenamed(UEdGraphPin* RenamedPin, const FString& OldName)
diff --git a/Engine/Plugins/FX/Niagara/Source/NiagaraEditor/Private/NiagaraNodeUsageSelector.h b/Engine/Plugins/FX/Niagara/Source/NiagaraEditor/Private/NiagaraNodeUsageSelector.h
--- a/Engine/Plugins/FX/Niagara/Source/NiagaraEditor/Private/NiagaraNodeUsageSelector.h
+++ b/Engine/Plugins/FX/Niagara/Source/NiagaraEditor/Private/NiagaraNodeUsageSelector.h
@@ -11,10 +11,10 @@ class UNiagaraNodeUsageSelector: public UNiagaraNodeWithDynamicPins
 	GENERATED_UCLASS_BODY()
 
 public:
-	UPROPERTY(EditAnywhere, Category=If)
+	UPROPERTY()
 	TArray<FNiagaraVariable> OutputVars;
 	
-	UPROPERTY(EditAnywhere, Category=If)
+	UPROPERTY()
 	TArray<FGuid> OutputVarGuids;
 	
 	//~ Begin UObject Interface
@@ -40,6 +40,8 @@ public:
 	FGuid AddOutput(FNiagaraTypeDefinition Type, const FName& Name);
 	virtual UEdGraphPin* GetPassThroughPin(const UEdGraphPin* LocallyOwnedOutputPin, ENiagaraScriptUsage MasterUsage) const override; 
 
+	virtual void AppendFunctionAliasForContext(const FNiagaraGraphFunctionAliasContext& InFunctionAliasContext, FString& InOutFunctionAlias) override;
+
 protected:
 	//~ Begin EdGraphNode Interface
 	virtual void OnPinRemoved(UEdGraphPin* PinToRemove) override;
@@ -53,4 +55,6 @@ protected:
 	virtual bool CanMovePin(const UEdGraphPin* Pin) const override { return false; }
 	virtual bool AllowNiagaraTypeForAddPin(const FNiagaraTypeDefinition& InType) override;
 	//~ End UNiagaraNodeWithDynamicPins Interface
+
+	virtual void InsertInputPinsFor(const FNiagaraVariable& Var);
 };
diff --git a/Engine/Plugins/FX/Niagara/Source/NiagaraEditor/Private/NiagaraParameterMapHistory.cpp b/Engine/Plugins/FX/Niagara/Source/NiagaraEditor/Private/NiagaraParameterMapHistory.cpp
--- a/Engine/Plugins/FX/Niagara/Source/NiagaraEditor/Private/NiagaraParameterMapHistory.cpp
+++ b/Engine/Plugins/FX/Niagara/Source/NiagaraEditor/Private/NiagaraParameterMapHistory.cpp
@@ -673,8 +673,12 @@ UNiagaraParameterCollection* FNiagaraParameterMapHistory::IsParameterCollectionP
 
 bool FNiagaraParameterMapHistory::ShouldIgnoreVariableDefault(const FNiagaraVariable& Var)const
 {
-	//For now just skip default for ID but maybe other cases/reasons?
-	return Var == FNiagaraVariable(FNiagaraTypeDefinition::GetIDDef(), TEXT("Particles.ID"));
+	// NOTE(mv): Used for variables that are explicitly assigned to (on spawn) and should not be default initialized
+	//           These are explicitly written to in NiagaraHlslTranslator::DefineMain
+	bool bShouldBeIgnored = false;
+	bShouldBeIgnored |= (Var == FNiagaraVariable(FNiagaraTypeDefinition::GetIDDef(), TEXT("Particles.ID")));
+	bShouldBeIgnored |= (Var == FNiagaraVariable(FNiagaraTypeDefinition::GetIntDef(), TEXT("Particles.UniqueID")));
+	return bShouldBeIgnored;
 }
 
 FNiagaraParameterMapHistoryBuilder::FNiagaraParameterMapHistoryBuilder()
@@ -1050,7 +1054,7 @@ bool FNiagaraParameterMapHistoryBuilder::IsInEncounteredEmitterNamespace(FNiagar
 /**
 * Use the current alias map to resolve any aliases in this input variable name.
 */
-FNiagaraVariable FNiagaraParameterMapHistoryBuilder::ResolveAliases(const FNiagaraVariable& InVar)
+FNiagaraVariable FNiagaraParameterMapHistoryBuilder::ResolveAliases(const FNiagaraVariable& InVar) const
 {
 	return FNiagaraParameterMapHistory::ResolveAliases(InVar, AliasMap, TEXT("."));
 }
@@ -1341,7 +1345,8 @@ int32 FNiagaraParameterMapHistoryBuilder::HandleExternalVariableRead(int32 Param
 		}
 		else
 		{
-			UE_LOG(LogNiagaraEditor, Log, TEXT("Could not resolve variable: %s"), *Name.ToString());
+			// This is overly spammy and doesn't provide useful info. Disabling for now.
+			//UE_LOG(LogNiagaraEditor, Log, TEXT("Could not resolve variable: %s"), *Name.ToString());
 		}
 		
 	}
diff --git a/Engine/Plugins/FX/Niagara/Source/NiagaraEditor/Private/NiagaraScriptSource.cpp b/Engine/Plugins/FX/Niagara/Source/NiagaraEditor/Private/NiagaraScriptSource.cpp
--- a/Engine/Plugins/FX/Niagara/Source/NiagaraEditor/Private/NiagaraScriptSource.cpp
+++ b/Engine/Plugins/FX/Niagara/Source/NiagaraEditor/Private/NiagaraScriptSource.cpp
@@ -153,7 +153,7 @@ void UNiagaraScriptSource::PostLoadFromEmitter(UNiagaraEmitter& OwningEmitter)
 
 bool UNiagaraScriptSource::AddModuleIfMissing(FString ModulePath, ENiagaraScriptUsage Usage, bool& bOutFoundModule)
 {
-	FStringAssetReference SystemUpdateScriptRef(ModulePath);
+	FSoftObjectPath SystemUpdateScriptRef(ModulePath);
 	FAssetData ModuleScriptAsset;
 	ModuleScriptAsset.ObjectPath = SystemUpdateScriptRef.GetAssetPathName();
 	bOutFoundModule = false;
@@ -185,6 +185,13 @@ void InitializeNewRapidIterationParametersForNode(const UEdGraphSchema_Niagara*
 		for (const UEdGraphPin* FunctionInputPin : FunctionInputPins)
 		{
 			FNiagaraTypeDefinition InputType = Schema->PinToTypeDefinition(FunctionInputPin);
+			if (InputType.IsValid() == false)
+			{
+				UE_LOG(LogNiagaraEditor, Error, TEXT("Invalid input type found while attempting initialize new rapid iteration parameters. Function Node: %s %s Input Name: %s"),
+					*FunctionCallNode->GetPathName(), *FunctionCallNode->GetFunctionName(), *FunctionInputPin->GetName());
+				continue;
+			}
+
 			if (FNiagaraStackGraphUtilities::IsRapidIterationType(InputType))
 			{
 				FNiagaraParameterHandle AliasedFunctionInputHandle = FNiagaraParameterHandle::CreateAliasedModuleParameterHandle(FNiagaraParameterHandle(FunctionInputPin->PinName), FunctionCallNode);
diff --git a/Engine/Plugins/FX/Niagara/Source/NiagaraEditor/Private/NiagaraStackEditorData.cpp b/Engine/Plugins/FX/Niagara/Source/NiagaraEditor/Private/NiagaraStackEditorData.cpp
--- a/Engine/Plugins/FX/Niagara/Source/NiagaraEditor/Private/NiagaraStackEditorData.cpp
+++ b/Engine/Plugins/FX/Niagara/Source/NiagaraEditor/Private/NiagaraStackEditorData.cpp
@@ -27,6 +27,20 @@ void UNiagaraStackEditorData::SetStackEntryIsExpanded(const FString& StackEntryK
 	}
 }
 
+bool UNiagaraStackEditorData::GetStackEntryWasExpandedPreSearch(const FString& StackEntryKey, bool bWasExpandedPreSearchDefault) const
+{
+	const bool* bWasExpandedPreSearchPtr = StackEntryKeyToPreSearchExpandedMap.Find(StackEntryKey);
+	return bWasExpandedPreSearchPtr != nullptr ? *bWasExpandedPreSearchPtr : bWasExpandedPreSearchDefault;
+}
+
+void UNiagaraStackEditorData::SetStackEntryWasExpandedPreSearch(const FString& StackEntryKey, bool bWasExpandedPreSearch)
+{
+	if (ensureMsgf(StackEntryKey.IsEmpty() == false, TEXT("Can not set the pre-search expanded state with an empty key")))
+	{
+		StackEntryKeyToPreSearchExpandedMap.FindOrAdd(StackEntryKey) = bWasExpandedPreSearch;
+	}
+}
+
 bool UNiagaraStackEditorData::GetStackItemShowAdvanced(const FString& StackEntryKey, bool bShowAdvancedDefault) const
 {
 	const bool* bShowAdvancedPtr = StackItemKeyToShowAdvancedMap.Find(StackEntryKey);
@@ -61,6 +75,16 @@ void UNiagaraStackEditorData::SetShowOutputs(bool bInShowOutputs)
 	bShowOutputs = bInShowOutputs;
 }
 
+bool UNiagaraStackEditorData::GetShowLinkedInputs() const
+{
+	return bShowLinkedInputs;
+}
+
+void UNiagaraStackEditorData::SetShowLinkedInputs(bool bInShowLinkedInputs)
+{
+	bShowLinkedInputs = bInShowLinkedInputs;
+}
+
 double UNiagaraStackEditorData::GetLastScrollPosition() const
 {
 	return LastScrollPosition;
diff --git a/Engine/Plugins/FX/Niagara/Source/NiagaraEditor/Private/NiagaraSystemEditorData.cpp b/Engine/Plugins/FX/Niagara/Source/NiagaraEditor/Private/NiagaraSystemEditorData.cpp
--- a/Engine/Plugins/FX/Niagara/Source/NiagaraEditor/Private/NiagaraSystemEditorData.cpp
+++ b/Engine/Plugins/FX/Niagara/Source/NiagaraEditor/Private/NiagaraSystemEditorData.cpp
@@ -56,7 +56,7 @@ UNiagaraSystemEditorData::UNiagaraSystemEditorData(const FObjectInitializer& Obj
 {
 	RootFolder = ObjectInitializer.CreateDefaultSubobject<UNiagaraSystemEditorFolder>(this, TEXT("RootFolder"));
 	StackEditorData = ObjectInitializer.CreateDefaultSubobject<UNiagaraStackEditorData>(this, TEXT("StackEditorData"));
-	OwnerTransform.SetLocation(FVector(0.0f, 0.0f, 100.0f));
+	OwnerTransform.SetLocation(FVector(0.0f, 0.0f, 0.0f));
 	PlaybackRangeMin = 0;
 	PlaybackRangeMax = 10;
 }
diff --git a/Engine/Plugins/FX/Niagara/Source/NiagaraEditor/Private/Toolkits/NiagaraSystemToolkit.cpp b/Engine/Plugins/FX/Niagara/Source/NiagaraEditor/Private/Toolkits/NiagaraSystemToolkit.cpp
--- a/Engine/Plugins/FX/Niagara/Source/NiagaraEditor/Private/Toolkits/NiagaraSystemToolkit.cpp
+++ b/Engine/Plugins/FX/Niagara/Source/NiagaraEditor/Private/Toolkits/NiagaraSystemToolkit.cpp
@@ -998,6 +998,7 @@ void FNiagaraSystemToolkit::UpdateOriginalEmitter()
 		// overwrite the original script in place by constructing a new one with the same name
 		Emitter = (UNiagaraEmitter*)StaticDuplicateObject(EditableEmitter, Emitter->GetOuter(),
 			Emitter->GetFName(), RF_AllFlags, Emitter->GetClass());
+		Emitter->PostEditChange();
 
 		TArray<UNiagaraScript*> EmitterScripts;
 		Emitter->GetScripts(EmitterScripts, false);
@@ -1036,6 +1037,7 @@ void FNiagaraSystemToolkit::UpdateOriginalEmitter()
 	{
 		Emitter->MarkPackageDirty();
 		Emitter->ThumbnailImage = (UTexture2D*)StaticDuplicateObject(EditableEmitter->ThumbnailImage, Emitter);
+		Emitter->PostEditChange();
 		bEmitterThumbnailUpdated = false;
 	}
 }
diff --git a/Engine/Plugins/FX/Niagara/Source/NiagaraEditor/Private/ViewModels/NiagaraCollectionParameterViewModel.cpp b/Engine/Plugins/FX/Niagara/Source/NiagaraEditor/Private/ViewModels/NiagaraCollectionParameterViewModel.cpp
--- a/Engine/Plugins/FX/Niagara/Source/NiagaraEditor/Private/ViewModels/NiagaraCollectionParameterViewModel.cpp
+++ b/Engine/Plugins/FX/Niagara/Source/NiagaraEditor/Private/ViewModels/NiagaraCollectionParameterViewModel.cpp
@@ -33,7 +33,7 @@ FName FNiagaraCollectionParameterViewModel::GetName() const
 
 FText FNiagaraCollectionParameterViewModel::GetTypeDisplayName() const
 {
-	return FText::Format(LOCTEXT("TypeTextFormat", "Type: {0}"), Parameter.GetType().GetStruct()->GetDisplayNameText());
+	return FText::Format(LOCTEXT("TypeTextFormat", "Type: {0}"), Parameter.GetType().GetNameText());
 }
 
 void FNiagaraCollectionParameterViewModel::NameTextComitted(const FText& Name, ETextCommit::Type CommitInfo)
diff --git a/Engine/Plugins/FX/Niagara/Source/NiagaraEditor/Private/ViewModels/NiagaraScriptParameterViewModel.cpp b/Engine/Plugins/FX/Niagara/Source/NiagaraEditor/Private/ViewModels/NiagaraScriptParameterViewModel.cpp
--- a/Engine/Plugins/FX/Niagara/Source/NiagaraEditor/Private/ViewModels/NiagaraScriptParameterViewModel.cpp
+++ b/Engine/Plugins/FX/Niagara/Source/NiagaraEditor/Private/ViewModels/NiagaraScriptParameterViewModel.cpp
@@ -73,7 +73,7 @@ FText FNiagaraScriptParameterViewModel::GetTypeDisplayName() const
 	{
 		return FText();
 	}
-	return FText::Format(LOCTEXT("TypeTextFormat", "Type: {0}"), GraphVariable->GetType().GetStruct()->GetDisplayNameText());
+	return FText::Format(LOCTEXT("TypeTextFormat", "Type: {0}"), GraphVariable->GetType().GetNameText());
 }
 
 void FNiagaraScriptParameterViewModel::NameTextComitted(const FText& Name, ETextCommit::Type CommitInfo)
diff --git a/Engine/Plugins/FX/Niagara/Source/NiagaraEditor/Private/ViewModels/Stack/NiagaraStackEntry.cpp b/Engine/Plugins/FX/Niagara/Source/NiagaraEditor/Private/ViewModels/Stack/NiagaraStackEntry.cpp
--- a/Engine/Plugins/FX/Niagara/Source/NiagaraEditor/Private/ViewModels/Stack/NiagaraStackEntry.cpp
+++ b/Engine/Plugins/FX/Niagara/Source/NiagaraEditor/Private/ViewModels/Stack/NiagaraStackEntry.cpp
@@ -115,9 +115,16 @@ const TArray<UNiagaraStackEntry::FStackIssueFix>& UNiagaraStackEntry::FStackIssu
 	return Fixes;
 }
 
+void UNiagaraStackEntry::FStackIssue::InsertFix(int32 InsertionIdx, const UNiagaraStackEntry::FStackIssueFix& Fix)
+{
+	Fixes.Insert(Fix, InsertionIdx);
+}
+
+
 UNiagaraStackEntry::UNiagaraStackEntry()
 	: IndentLevel(0)
 	, bIsFinalized(false)
+	, bIsSearchResult(false)
 {
 }
 
@@ -230,14 +237,14 @@ bool UNiagaraStackEntry::GetShouldShowInStack() const
 	return true;
 }
 
-void UNiagaraStackEntry::GetFilteredChildren(TArray<UNiagaraStackEntry*>& OutFilteredChildren)
+void UNiagaraStackEntry::GetFilteredChildren(TArray<UNiagaraStackEntry*>& OutFilteredChildren) const
 {
 	OutFilteredChildren.Append(ErrorChildren);
 	for (UNiagaraStackEntry* Child : Children)
 	{
 		bool bPassesFilter = true;
-		for(const FOnFilterChild& ChildFilter : ChildFilters)
-		{ 
+		for (const FOnFilterChild& ChildFilter : ChildFilters)
+		{
 			if (ChildFilter.Execute(*Child) == false)
 			{
 				bPassesFilter = false;
@@ -252,12 +259,6 @@ void UNiagaraStackEntry::GetFilteredChildren(TArray<UNiagaraStackEntry*>& OutFil
 	}
 }
 
-void UNiagaraStackEntry::GetUnfilteredChildren(TArray<UNiagaraStackEntry*>& OutUnfilteredChildren)
-{
-	OutUnfilteredChildren.Append(ErrorChildren);
-	OutUnfilteredChildren.Append(Children);
-}
-
 void UNiagaraStackEntry::GetUnfilteredChildren(TArray<UNiagaraStackEntry*>& OutUnfilteredChildren) const
 {
 	OutUnfilteredChildren.Append(ErrorChildren);
@@ -319,7 +320,6 @@ int32 UNiagaraStackEntry::GetIndentLevel() const
 void UNiagaraStackEntry::GetSearchItems(TArray<UNiagaraStackEntry::FStackSearchItem>& SearchItems) const
 {
 	SearchItems.Add({FName("DisplayName"), GetDisplayName()}); 
-	GetAdditionalSearchItemsInternal(SearchItems);
 }
 
 UObject* UNiagaraStackEntry::GetExternalAsset() const
@@ -372,8 +372,14 @@ void UNiagaraStackEntry::SetOnRequestDrop(FOnRequestDrop InOnRequestDrop)
 	OnRequestDropDelegate = InOnRequestDrop;
 }
 
-void UNiagaraStackEntry::GetAdditionalSearchItemsInternal(TArray<FStackSearchItem>& SearchItems) const
+const bool UNiagaraStackEntry::GetIsSearchResult() const
 {
+	return bIsSearchResult;
+}
+
+void UNiagaraStackEntry::SetIsSearchResult(bool bInIsSearchResult)
+{
+	bIsSearchResult = bInIsSearchResult;
 }
 
 TOptional<UNiagaraStackEntry::FDropResult> UNiagaraStackEntry::CanDropInternal(const TArray<UNiagaraStackEntry*>& DraggedEntries)
@@ -517,7 +523,11 @@ void UNiagaraStackEntry::RefreshStackErrorChildren()
 
 void UNiagaraStackEntry::IssueModified()
 {
-	RefreshChildren();
+	if (bIsFinalized == false)
+	{
+		// Fixing an issue may have caused this entry to be deleted and finalized, so don't refresh in that case.
+		RefreshChildren();
+	}
 }
 
 void UNiagaraStackEntry::BeginDestroy()
diff --git a/Engine/Plugins/FX/Niagara/Source/NiagaraEditor/Private/ViewModels/Stack/NiagaraStackFunctionInput.cpp b/Engine/Plugins/FX/Niagara/Source/NiagaraEditor/Private/ViewModels/Stack/NiagaraStackFunctionInput.cpp
--- a/Engine/Plugins/FX/Niagara/Source/NiagaraEditor/Private/ViewModels/Stack/NiagaraStackFunctionInput.cpp
+++ b/Engine/Plugins/FX/Niagara/Source/NiagaraEditor/Private/ViewModels/Stack/NiagaraStackFunctionInput.cpp
@@ -50,6 +50,7 @@ UNiagaraStackFunctionInput::UNiagaraStackFunctionInput()
 	, bUpdatingLocalValueDirectly(false)
 	, bShowEditConditionInline(false)
 	, bIsInlineEditConditionToggle(false)
+	, bIsDynamicInputScriptReassignmentPending(false)
 {
 }
 
@@ -309,21 +310,71 @@ void UNiagaraStackFunctionInput::RefreshChildrenInternal(const TArray<UNiagaraSt
 
 	if (InputValues.Mode == EValueMode::Dynamic && InputValues.DynamicNode.IsValid())
 	{
-		UNiagaraStackFunctionInputCollection* DynamicInputEntry = FindCurrentChildOfTypeByPredicate<UNiagaraStackFunctionInputCollection>(CurrentChildren,
-			[=](UNiagaraStackFunctionInputCollection* CurrentFunctionInputEntry) 
-		{ 
-			return CurrentFunctionInputEntry->GetInputFunctionCallNode() == InputValues.DynamicNode.Get() &&
-				CurrentFunctionInputEntry->GetModuleNode() == OwningModuleNode.Get(); 
-		});
-
-		if (DynamicInputEntry == nullptr)
-		{ 
-			DynamicInputEntry = NewObject<UNiagaraStackFunctionInputCollection>(this);
-			DynamicInputEntry->Initialize(CreateDefaultChildRequiredData(), *OwningModuleNode, *InputValues.DynamicNode.Get(), GetOwnerStackItemEditorDataKey());
-			DynamicInputEntry->SetShouldShowInStack(false);
-		}
+		if (InputValues.DynamicNode->FunctionScript != nullptr)
+		{
+			UNiagaraStackFunctionInputCollection* DynamicInputEntry = FindCurrentChildOfTypeByPredicate<UNiagaraStackFunctionInputCollection>(CurrentChildren,
+				[=](UNiagaraStackFunctionInputCollection* CurrentFunctionInputEntry)
+			{
+				return CurrentFunctionInputEntry->GetInputFunctionCallNode() == InputValues.DynamicNode.Get() &&
+					CurrentFunctionInputEntry->GetModuleNode() == OwningModuleNode.Get();
+			});
 
-		NewChildren.Add(DynamicInputEntry);
+			if (DynamicInputEntry == nullptr)
+			{
+				DynamicInputEntry = NewObject<UNiagaraStackFunctionInputCollection>(this);
+				DynamicInputEntry->Initialize(CreateDefaultChildRequiredData(), *OwningModuleNode, *InputValues.DynamicNode.Get(), GetOwnerStackItemEditorDataKey());
+				DynamicInputEntry->SetShouldShowInStack(false);
+			}
+			if (InputValues.DynamicNode->FunctionScript != nullptr && InputValues.DynamicNode->FunctionScript->bDeprecated)
+			{
+				FText LongMessage = InputValues.DynamicNode->FunctionScript->DeprecationRecommendation != nullptr ?
+					FText::Format(LOCTEXT("ModuleScriptDeprecationLong", "The script asset for the assigned module {0} has been deprecated. Suggested replacement: {1}"), FText::FromString(InputValues.DynamicNode->GetName()), FText::FromString(InputValues.DynamicNode->FunctionScript->DeprecationRecommendation->GetPathName())) :
+					FText::Format(LOCTEXT("ModuleScriptDeprecationUnknownLong", "The script asset for the assigned module {0} has been deprecated."), FText::FromString(InputValues.DynamicNode->GetName()));
+
+				int32 AddIdx = NewIssues.Add(FStackIssue(
+					EStackIssueSeverity::Warning,
+					LOCTEXT("ModuleScriptDeprecationShort", "Deprecated module"),
+					LongMessage,
+					GetStackEditorDataKey(),
+					false,
+					{
+						FStackIssueFix(
+							LOCTEXT("SelectNewDynamicInputScriptFix", "Select a new dynamic input script"),
+							FStackIssueFixDelegate::CreateLambda([this]() { this->bIsDynamicInputScriptReassignmentPending = true; })),
+						FStackIssueFix(
+							LOCTEXT("ResetFix", "Reset this input to it's default value"),
+							FStackIssueFixDelegate::CreateLambda([this]() { this->Reset(); }))
+					}));
+
+				if (InputValues.DynamicNode->FunctionScript->DeprecationRecommendation != nullptr)
+				{
+					NewIssues[AddIdx].InsertFix(0,
+						FStackIssueFix(
+							LOCTEXT("SelectNewModuleScriptFixUseRecommended", "Use recommended replacement"),
+							FStackIssueFixDelegate::CreateLambda([this]() { ReassignDynamicInputScript(InputValues.DynamicNode->FunctionScript->DeprecationRecommendation); })));
+				}
+			}
+
+
+			NewChildren.Add(DynamicInputEntry);
+		}
+		else
+		{
+			NewIssues.Add(FStackIssue(
+				EStackIssueSeverity::Error,
+				LOCTEXT("DynamicInputScriptMissingShort", "Missing dynamic input script"),
+				FText::Format(LOCTEXT("DynamicInputScriptMissingLong", "The script asset for the assigned dynamic input {0} is missing."), FText::FromString(InputValues.DynamicNode->GetFunctionName())),
+				GetStackEditorDataKey(),
+				false,
+				{
+					FStackIssueFix(
+						LOCTEXT("SelectNewDynamicInputScriptFix", "Select a new dynamic input script"),
+						FStackIssueFixDelegate::CreateLambda([this]() { this->bIsDynamicInputScriptReassignmentPending = true; })),
+					FStackIssueFix(
+						LOCTEXT("ResetFix", "Reset this input to it's default value"),
+						FStackIssueFixDelegate::CreateLambda([this]() { this->Reset(); }))
+				}));
+		}
 	}
 
 	if (InputValues.Mode == EValueMode::Data && InputValues.DataObjects.GetValueObject() != nullptr)
@@ -751,6 +802,26 @@ void UNiagaraStackFunctionInput::GetAvailableParameterHandles(TArray<FNiagaraPar
 			}
 		}
 	}
+
+	//Parameter Collections
+	FAssetRegistryModule& AssetRegistryModule = FModuleManager::LoadModuleChecked<FAssetRegistryModule>("AssetRegistry");
+	TArray<FAssetData> CollectionAssets;
+	AssetRegistryModule.Get().GetAssetsByClass(UNiagaraParameterCollection::StaticClass()->GetFName(), CollectionAssets);
+
+	for (FAssetData& CollectionAsset : CollectionAssets)
+	{
+		UNiagaraParameterCollection* Collection = CastChecked<UNiagaraParameterCollection>(CollectionAsset.GetAsset());
+		if (Collection)
+		{
+			for (const FNiagaraVariable& CollectionParam : Collection->GetParameters())
+			{
+				if (CollectionParam.GetType() == InputType)
+				{
+					AvailableParameterHandles.AddUnique(FNiagaraParameterHandle(CollectionParam.GetName()));
+				}
+			}
+		}
+	}
 }
 
 UNiagaraNodeFunctionCall* UNiagaraStackFunctionInput::GetDynamicInputNode() const
@@ -1285,6 +1356,8 @@ void UNiagaraStackFunctionInput::RenameInput(FName NewName)
 {
 	if (OwningAssignmentNode.IsValid() && InputParameterHandlePath.Num() == 1 && InputParameterHandle.GetName() != NewName)
 	{
+		FScopedTransaction ScopedTransaction(LOCTEXT("RenameInput", "Rename this function's input."));
+
 		FInputValues OldInputValues = InputValues;
 		UEdGraphPin* OriginalOverridePin = GetOverridePin();
 
@@ -1301,10 +1374,17 @@ void UNiagaraStackFunctionInput::RenameInput(FName NewName)
 		check(FoundIdx != INDEX_NONE);
 		FNiagaraParameterHandle TargetHandle(OwningAssignmentNode->GetAssignmentTargetName(FoundIdx));
 
+		OwningAssignmentNode->Modify();
+		if (OwningAssignmentNode->FunctionScript != nullptr)
+		{
+			OwningAssignmentNode->FunctionScript->Modify();
+			OwningAssignmentNode->FunctionScript->GetSource()->Modify();
+		}
 		if (OwningAssignmentNode->SetAssignmentTargetName(FoundIdx, NewName))
 		{
 			OwningAssignmentNode->RefreshFromExternalChanges();
 		}
+
 		InputParameterHandle = FNiagaraParameterHandle(InputParameterHandle.GetNamespace(), NewName);
 		InputParameterHandlePath.Empty(1);
 		InputParameterHandlePath.Add(InputParameterHandle);
@@ -1318,6 +1398,7 @@ void UNiagaraStackFunctionInput::RenameInput(FName NewName)
 
 			for (TWeakObjectPtr<UNiagaraScript> Script : AffectedScripts)
 			{
+				Script->Modify();
 				Script->RapidIterationParameters.RenameParameter(OldRapidIterationParameter, *RapidIterationParameter.GetName().ToString());
 			}
 
@@ -1438,6 +1519,42 @@ bool UNiagaraStackFunctionInput::GetIsInlineEditConditionToggle() const
 	return bIsInlineEditConditionToggle;
 }
 
+bool UNiagaraStackFunctionInput::GetIsDynamicInputScriptReassignmentPending() const
+{
+	return bIsDynamicInputScriptReassignmentPending;
+}
+
+void UNiagaraStackFunctionInput::SetIsDynamicInputScriptReassignmentPending(bool bIsPending)
+{
+	bIsDynamicInputScriptReassignmentPending = bIsPending;
+}
+
+void UNiagaraStackFunctionInput::ReassignDynamicInputScript(UNiagaraScript* DynamicInputScript)
+{
+	if (ensureMsgf(InputValues.Mode == EValueMode::Dynamic && InputValues.DynamicNode != nullptr && InputValues.DynamicNode->GetClass() == UNiagaraNodeFunctionCall::StaticClass(),
+		TEXT("Can not reassign the dynamic input script when tne input doesn't have a valid dynamic input.")))
+	{
+		FScopedTransaction ScopedTransaction(LOCTEXT("ReassignDynamicInputTransaction", "Reassign dynamic input script"));
+		InputValues.DynamicNode->Modify();
+		InputValues.DynamicNode->FunctionScript = DynamicInputScript;
+		InputValues.DynamicNode->MarkNodeRequiresSynchronization(TEXT("Dynamic input script reassigned."), true);
+		RefreshChildren();
+	}
+}
+
+bool UNiagaraStackFunctionInput::GetShouldPassFilterForVisibleCondition() const
+{
+	return GetHasVisibleCondition() == false || GetVisibleConditionEnabled();
+}
+
+void UNiagaraStackFunctionInput::GetSearchItems(TArray<FStackSearchItem>& SearchItems) const
+{
+	if (GetShouldPassFilterForVisibleCondition() && GetIsInlineEditConditionToggle() == false)
+	{
+		SearchItems.Add({ FName("DisplayName"), GetDisplayName() });
+	}
+}
+
 void UNiagaraStackFunctionInput::OnGraphChanged(const struct FEdGraphEditAction& InAction)
 {
 	if (bUpdatingGraphDirectly == false)
diff --git a/Engine/Plugins/FX/Niagara/Source/NiagaraEditor/Private/ViewModels/Stack/NiagaraStackGraphUtilities.cpp b/Engine/Plugins/FX/Niagara/Source/NiagaraEditor/Private/ViewModels/Stack/NiagaraStackGraphUtilities.cpp
--- a/Engine/Plugins/FX/Niagara/Source/NiagaraEditor/Private/ViewModels/Stack/NiagaraStackGraphUtilities.cpp
+++ b/Engine/Plugins/FX/Niagara/Source/NiagaraEditor/Private/ViewModels/Stack/NiagaraStackGraphUtilities.cpp
@@ -1168,6 +1168,7 @@ void GetFunctionNamesForOutputNode(UNiagaraNodeOutput& OutputNode, TArray<FStrin
 
 bool FNiagaraStackGraphUtilities::IsRapidIterationType(const FNiagaraTypeDefinition& InputType)
 {
+	checkf(InputType.IsValid(), TEXT("Type is invalid."));
 	return InputType != FNiagaraTypeDefinition::GetBoolDef() && !InputType.IsEnum() &&
 		InputType != FNiagaraTypeDefinition::GetParameterMapDef() && !InputType.IsDataInterface();
 }
@@ -1340,6 +1341,7 @@ TOptional<FName> FNiagaraStackGraphUtilities::GetNamespaceForScriptUsage(ENiagar
 	case ENiagaraScriptUsage::ParticleSpawnScript:
 	case ENiagaraScriptUsage::ParticleSpawnScriptInterpolated:
 	case ENiagaraScriptUsage::ParticleUpdateScript:
+	case ENiagaraScriptUsage::ParticleEventScript:
 		return FNiagaraParameterHandle::ParticleAttributeNamespace;
 	case ENiagaraScriptUsage::EmitterSpawnScript:
 	case ENiagaraScriptUsage::EmitterUpdateScript:
@@ -1505,15 +1507,20 @@ bool TryGetStackFunctionInputValue(UNiagaraScript& OwningScript, const UEdGraphP
 	}
 	else if (InputPin.LinkedTo.Num() == 1)
 	{
-		if (InputPin.LinkedTo[0]->GetOwningNode()->IsA<UNiagaraNodeParameterMapGet>())
+		const UEdGraphSchema_Niagara* NiagaraSchema = GetDefault<UEdGraphSchema_Niagara>();
+		UEdGraphNode* PreviousOwningNode = InputPin.LinkedTo[0]->GetOwningNode();
+
+		
+		if (PreviousOwningNode->IsA<UNiagaraNodeParameterMapGet>())
 		{
 			OutStackFunctionInputValue.LinkedValue = InputPin.LinkedTo[0]->GetFName();
 		}
-		else if (InputPin.LinkedTo[0]->GetOwningNode()->IsA<UNiagaraNodeInput>())
+		else if (PreviousOwningNode->IsA<UNiagaraNodeInput>())
 		{
 			OutStackFunctionInputValue.DataValue = CastChecked<UNiagaraNodeInput>(InputPin.LinkedTo[0]->GetOwningNode())->GetDataInterface();
 		}
-		else if (InputPin.LinkedTo[0]->GetOwningNode()->IsA<UNiagaraNodeFunctionCall>())
+		else if (PreviousOwningNode->IsA<UNiagaraNodeFunctionCall>() && 
+			FNiagaraStackGraphUtilities::GetParameterMapInputPin(*(static_cast<UNiagaraNodeFunctionCall*>(PreviousOwningNode))) != nullptr)
 		{
 			UNiagaraNodeFunctionCall* DynamicInputFunctionCall = CastChecked<UNiagaraNodeFunctionCall>(InputPin.LinkedTo[0]->GetOwningNode());
 			OutStackFunctionInputValue.DynamicValue = DynamicInputFunctionCall;
diff --git a/Engine/Plugins/FX/Niagara/Source/NiagaraEditor/Private/ViewModels/Stack/NiagaraStackInputCategory.cpp b/Engine/Plugins/FX/Niagara/Source/NiagaraEditor/Private/ViewModels/Stack/NiagaraStackInputCategory.cpp
--- a/Engine/Plugins/FX/Niagara/Source/NiagaraEditor/Private/ViewModels/Stack/NiagaraStackInputCategory.cpp
+++ b/Engine/Plugins/FX/Niagara/Source/NiagaraEditor/Private/ViewModels/Stack/NiagaraStackInputCategory.cpp
@@ -84,7 +84,7 @@ void UNiagaraStackInputCategory::SetShouldShowInStack(bool bInShouldShowInStack)
 bool UNiagaraStackInputCategory::FilterForVisibleCondition(const UNiagaraStackEntry& Child) const
 {
 	const UNiagaraStackFunctionInput* StackFunctionInputChild = Cast<UNiagaraStackFunctionInput>(&Child);
-	return StackFunctionInputChild == nullptr || StackFunctionInputChild->GetHasVisibleCondition() == false || StackFunctionInputChild->GetVisibleConditionEnabled();
+	return StackFunctionInputChild == nullptr || StackFunctionInputChild->GetShouldPassFilterForVisibleCondition();
 }
 
 bool UNiagaraStackInputCategory::FilterForIsInlineEditConditionToggle(const UNiagaraStackEntry& Child) const
diff --git a/Engine/Plugins/FX/Niagara/Source/NiagaraEditor/Private/ViewModels/Stack/NiagaraStackItem.cpp b/Engine/Plugins/FX/Niagara/Source/NiagaraEditor/Private/ViewModels/Stack/NiagaraStackItem.cpp
--- a/Engine/Plugins/FX/Niagara/Source/NiagaraEditor/Private/ViewModels/Stack/NiagaraStackItem.cpp
+++ b/Engine/Plugins/FX/Niagara/Source/NiagaraEditor/Private/ViewModels/Stack/NiagaraStackItem.cpp
@@ -195,7 +195,7 @@ void UNiagaraStackItemContent::SetIsAdvanced(bool bInIsAdvanced)
 bool UNiagaraStackItemContent::FilterAdvancedChildren(const UNiagaraStackEntry& Child) const
 {
 	const UNiagaraStackItemContent* ItemContent = Cast<UNiagaraStackItemContent>(&Child);
-	if (ItemContent == nullptr || ItemContent->GetIsAdvanced() == false)
+	if (ItemContent == nullptr || ItemContent->GetIsAdvanced() == false || ItemContent->GetIsSearchResult())
 	{
 		return true;
 	}
diff --git a/Engine/Plugins/FX/Niagara/Source/NiagaraEditor/Private/ViewModels/Stack/NiagaraStackModuleItem.cpp b/Engine/Plugins/FX/Niagara/Source/NiagaraEditor/Private/ViewModels/Stack/NiagaraStackModuleItem.cpp
--- a/Engine/Plugins/FX/Niagara/Source/NiagaraEditor/Private/ViewModels/Stack/NiagaraStackModuleItem.cpp
+++ b/Engine/Plugins/FX/Niagara/Source/NiagaraEditor/Private/ViewModels/Stack/NiagaraStackModuleItem.cpp
@@ -5,6 +5,7 @@
 #include "ViewModels/NiagaraSystemViewModel.h"
 #include "ViewModels/NiagaraEmitterHandleViewModel.h"
 #include "ViewModels/NiagaraEmitterViewModel.h"
+#include "ViewModels/Stack/NiagaraStackModuleItemLinkedInputCollection.h"
 #include "ViewModels/Stack/NiagaraStackFunctionInputCollection.h"
 #include "ViewModels/Stack/NiagaraStackFunctionInput.h"
 #include "ViewModels/Stack/NiagaraStackInputCategory.h"
@@ -46,68 +47,53 @@ TArray<ENiagaraScriptUsage> UsagePriority = { // Ordered such as the highest pri
 	ENiagaraScriptUsage::SystemUpdateScript,
 	ENiagaraScriptUsage::SystemSpawnScript };
 
-UNiagaraNodeOutput* GetOutputNodeForModuleDependency(ENiagaraScriptUsage DependantUsage, UNiagaraScript* DependencyScript, UNiagaraSystem& System, UNiagaraEmitter* Emitter, FNiagaraModuleDependency Dependency)
+UNiagaraNodeOutput* GetOutputNodeForModuleDependency(ENiagaraScriptUsage DependantUsage, UNiagaraScript* DependencyScript, UNiagaraSystem& System, const FNiagaraEmitterHandle* EmitterHandle, FNiagaraModuleDependency Dependency)
 {
 	UNiagaraNodeOutput* TargetOutputNode = nullptr;
 	if (DependencyScript)
 	{
 		UNiagaraScript* OutputScript = nullptr;
 		TArray<ENiagaraScriptUsage> SupportedUsages = UNiagaraScript::GetSupportedUsageContextsForBitmask(DependencyScript->ModuleUsageBitmask);
-		ENiagaraScriptUsage ScriptUsage = SupportedUsages[0];
-		int32 ClosestDistance = MAX_int32;
 
-		int32 DependantIndex = UsagePriority.IndexOfByPredicate(
-			[&](const ENiagaraScriptUsage CurrentUsage)
+		if (Dependency.ScriptConstraint == ENiagaraModuleDependencyScriptConstraint::AllScripts)
 		{
-			return UNiagaraScript::IsEquivalentUsage(DependantUsage, CurrentUsage);
-		});
-
-		for (ENiagaraScriptUsage PossibleUsage : SupportedUsages)
-		{
-			int32 PossibleIndex = UsagePriority.IndexOfByPredicate(
+			int32 ClosestDistance = MAX_int32;
+			int32 DependantIndex = UsagePriority.IndexOfByPredicate(
 				[&](const ENiagaraScriptUsage CurrentUsage)
-				{
-					return UNiagaraScript::IsEquivalentUsage(PossibleUsage, CurrentUsage);
-				});
-
-			if (PossibleIndex == INDEX_NONE)
 			{
-				// This usage isn't in the execution flow so check the next one.
-				continue;
-			}
+				return UNiagaraScript::IsEquivalentUsage(DependantUsage, CurrentUsage);
+			});
 
-			int32 Distance = PossibleIndex - DependantIndex;
-			bool bCorrectOrder = (Dependency.Type == ENiagaraModuleDependencyType::PreDependency && Distance >= 0) || (Dependency.Type == ENiagaraModuleDependencyType::PostDependency && Distance <= 0);
-			if ((FMath::Abs(Distance) < ClosestDistance) && bCorrectOrder)
+			for (ENiagaraScriptUsage PossibleUsage : SupportedUsages)
 			{
-				ClosestDistance = Distance;
-				ScriptUsage = PossibleUsage;
-				if (UNiagaraScript::IsEquivalentUsage(ScriptUsage, ENiagaraScriptUsage::SystemSpawnScript))
-				{
-					OutputScript = System.GetSystemSpawnScript();
-				}
-				else if (UNiagaraScript::IsEquivalentUsage(ScriptUsage, ENiagaraScriptUsage::SystemUpdateScript))
-				{
-					OutputScript = System.GetSystemUpdateScript();
-				}
-				else if (UNiagaraScript::IsEquivalentUsage(ScriptUsage, ENiagaraScriptUsage::EmitterSpawnScript))
+				int32 PossibleIndex = UsagePriority.IndexOfByPredicate(
+					[&](const ENiagaraScriptUsage CurrentUsage)
 				{
-					OutputScript = Emitter->EmitterSpawnScriptProps.Script;
-				}
-				else if (UNiagaraScript::IsEquivalentUsage(ScriptUsage, ENiagaraScriptUsage::EmitterUpdateScript))
-				{
-					OutputScript = Emitter->EmitterUpdateScriptProps.Script;
-				}
-				else if (UNiagaraScript::IsEquivalentUsage(ScriptUsage, ENiagaraScriptUsage::ParticleSpawnScript))
+					return UNiagaraScript::IsEquivalentUsage(PossibleUsage, CurrentUsage);
+				});
+
+				if (PossibleIndex == INDEX_NONE)
 				{
-					OutputScript = Emitter->SpawnScriptProps.Script;
+					// This usage isn't in the execution flow so check the next one.
+					continue;
 				}
-				else if (UNiagaraScript::IsEquivalentUsage(ScriptUsage, ENiagaraScriptUsage::ParticleUpdateScript))
+
+				int32 Distance = PossibleIndex - DependantIndex;
+				bool bCorrectOrder = (Dependency.Type == ENiagaraModuleDependencyType::PreDependency && Distance >= 0) || (Dependency.Type == ENiagaraModuleDependencyType::PostDependency && Distance <= 0);
+				if ((FMath::Abs(Distance) < ClosestDistance) && bCorrectOrder)
 				{
-					OutputScript = Emitter->UpdateScriptProps.Script;
+					ClosestDistance = Distance;
+					OutputScript = FNiagaraEditorUtilities::GetScriptFromSystem(System, EmitterHandle->GetId(), PossibleUsage, FGuid());
 				}
 			}
 		}
+		else if (Dependency.ScriptConstraint == ENiagaraModuleDependencyScriptConstraint::SameScript)
+		{
+			if (SupportedUsages.Contains(DependantUsage))
+			{
+				OutputScript = FNiagaraEditorUtilities::GetScriptFromSystem(System, EmitterHandle->GetId(), DependantUsage, FGuid());
+			}
+		}
 
 		if (OutputScript != nullptr)
 		{
@@ -121,6 +107,7 @@ UNiagaraStackModuleItem::UNiagaraStackModuleItem()
 	: FunctionCallNode(nullptr)
 	, bCanRefresh(false)
 	, InputCollection(nullptr)
+	, bIsModuleScriptReassignmentPending(false)
 {
 }
 
@@ -142,7 +129,13 @@ void UNiagaraStackModuleItem::Initialize(FRequiredEntryData InRequiredEntryData,
 	GroupAddUtilities = InGroupAddUtilities;
 	FunctionCallNode = &InFunctionCallNode;
 	OutputNode = FNiagaraStackGraphUtilities::GetEmitterOutputNodeForStackNode(*FunctionCallNode);
-	AddChildFilter(FOnFilterChild::CreateUObject(this, &UNiagaraStackModuleItem::FilterOutputCollection));
+
+	// We do not need to include child filters for NiagaraNodeAssignments as they do not display their output or linked input collections
+	if (!FunctionCallNode->IsA<UNiagaraNodeAssignment>())
+	{
+		AddChildFilter(FOnFilterChild::CreateUObject(this, &UNiagaraStackModuleItem::FilterOutputCollection));
+		AddChildFilter(FOnFilterChild::CreateUObject(this, &UNiagaraStackModuleItem::FilterLinkedInputCollection));
+	}
 
 	// Update bCanMoveAndDelete
 	if (GetSystemViewModel()->GetEditMode() == ENiagaraSystemViewModelEditMode::EmitterAsset)
@@ -228,23 +221,43 @@ void UNiagaraStackModuleItem::RefreshChildrenInternal(const TArray<UNiagaraStack
 			InputCollection->Initialize(CreateDefaultChildRequiredData(), *FunctionCallNode, *FunctionCallNode, GetStackEditorDataKey());
 		}
 
-		InputCollection->SetShouldShowInStack(GetStackEditorData().GetShowOutputs());
-
-		if (OutputCollection == nullptr)
+		// NiagaraNodeAssignments should not display OutputCollection and LinkedInputCollection as they effectively handle this through their InputCollection 
+		if (!FunctionCallNode->IsA<UNiagaraNodeAssignment>())
 		{
-			OutputCollection = NewObject<UNiagaraStackModuleItemOutputCollection>(this);
-			OutputCollection->Initialize(CreateDefaultChildRequiredData(), *FunctionCallNode);
-		}
 
-		NewChildren.Add(InputCollection);
-		NewChildren.Add(OutputCollection);
+			if (LinkedInputCollection == nullptr)
+			{
+				LinkedInputCollection = NewObject<UNiagaraStackModuleItemLinkedInputCollection>(this);
+				LinkedInputCollection->Initialize(CreateDefaultChildRequiredData(), *FunctionCallNode);
+				LinkedInputCollection->AddChildFilter(FOnFilterChild::CreateUObject(this, &UNiagaraStackModuleItem::FilterLinkedInputCollectionChild));
+			}
+
+			if (OutputCollection == nullptr)
+			{
+				OutputCollection = NewObject<UNiagaraStackModuleItemOutputCollection>(this);
+				OutputCollection->Initialize(CreateDefaultChildRequiredData(), *FunctionCallNode);
+				OutputCollection->AddChildFilter(FOnFilterChild::CreateUObject(this, &UNiagaraStackModuleItem::FilterOutputCollectionChild));
+			}
 
-		RefreshIsEnabled();
+			InputCollection->SetShouldShowInStack(GetStackEditorData().GetShowOutputs() || GetStackEditorData().GetShowLinkedInputs());
 
-		Super::RefreshChildrenInternal(CurrentChildren, NewChildren, NewIssues);
+			NewChildren.Add(LinkedInputCollection);
+			NewChildren.Add(InputCollection);
+			NewChildren.Add(OutputCollection);
+		
+		}
+		else
+		{
+			// We do not show the expander arrow for InputCollections of NiagaraNodeAssignments as they only have this one collection
+			InputCollection->SetShouldShowInStack(false);
 
-		RefreshIssues(NewIssues);
+			NewChildren.Add(InputCollection);
+		}
 	}
+
+	RefreshIsEnabled();
+	Super::RefreshChildrenInternal(CurrentChildren, NewChildren, NewIssues);
+	RefreshIssues(NewIssues);
 }
 
 void UNiagaraStackModuleItem::RefreshIssues(TArray<FStackIssue>& NewIssues)
@@ -256,12 +269,58 @@ void UNiagaraStackModuleItem::RefreshIssues(TArray<FStackIssue>& NewIssues)
 	}
 	if (FunctionCallNode != nullptr)
 	{
-		if (!FunctionCallNode->ScriptIsValid())
+		if (FunctionCallNode->FunctionScript != nullptr && FunctionCallNode->FunctionScript->bDeprecated)
+		{
+			FText LongMessage = FunctionCallNode->FunctionScript->DeprecationRecommendation != nullptr ? 
+				FText::Format(LOCTEXT("ModuleScriptDeprecationLong", "The script asset for the assigned module {0} has been deprecated. Suggested replacement: {1}"), FText::FromString(FunctionCallNode->GetFunctionName()),FText::FromString(FunctionCallNode->FunctionScript->DeprecationRecommendation->GetPathName())) :
+				FText::Format(LOCTEXT("ModuleScriptDeprecationUnknownLong", "The script asset for the assigned module {0} has been deprecated."), FText::FromString(FunctionCallNode->GetFunctionName()));
+
+			int32 AddIdx = NewIssues.Add(FStackIssue(
+				EStackIssueSeverity::Warning,
+				LOCTEXT("ModuleScriptDeprecationShort", "Deprecated module"),
+				LongMessage,
+				GetStackEditorDataKey(),
+				false,
+				{
+					FStackIssueFix(
+						LOCTEXT("SelectNewModuleScriptFix", "Select a new module script"),
+						FStackIssueFixDelegate::CreateLambda([this]() { this->bIsModuleScriptReassignmentPending = true; })),
+					FStackIssueFix(
+						LOCTEXT("DeleteFix", "Delete this module"),
+						FStackIssueFixDelegate::CreateLambda([this]() { this->Delete(); }))
+				}));
+			
+			if (FunctionCallNode->FunctionScript->DeprecationRecommendation != nullptr)
+			{
+				NewIssues[AddIdx].InsertFix(0,
+					FStackIssueFix(
+						LOCTEXT("SelectNewModuleScriptFixUseRecommended", "Use recommended replacement"),
+						FStackIssueFixDelegate::CreateLambda([this]() { ReassignModuleScript(FunctionCallNode->FunctionScript->DeprecationRecommendation); })));
+			}
+		}
+		if(FunctionCallNode->FunctionScript == nullptr && FunctionCallNode->GetClass() == UNiagaraNodeFunctionCall::StaticClass())
+		{
+			NewIssues.Add(FStackIssue(
+				EStackIssueSeverity::Error,
+				LOCTEXT("ModuleScriptMissingShort", "Missing module script"),
+				FText::Format(LOCTEXT("ModuleScriptMissingLong", "The script asset for the assigned module {0} is missing."), FText::FromString(FunctionCallNode->GetFunctionName())),
+				GetStackEditorDataKey(),
+				false,
+				{
+					FStackIssueFix(
+						LOCTEXT("SelectNewModuleScriptFix", "Select a new module script"),
+						FStackIssueFixDelegate::CreateLambda([this]() { this->bIsModuleScriptReassignmentPending = true; })),
+				FStackIssueFix(
+					LOCTEXT("DeleteFix", "Delete this module"),
+					FStackIssueFixDelegate::CreateLambda([this]() { this->Delete(); }))
+				}));
+		}
+		else if (!FunctionCallNode->ScriptIsValid())
 		{
 			FStackIssue InvalidScriptError(
 				EStackIssueSeverity::Error,
 				LOCTEXT("InvalidModuleScriptErrorSummary", "Invalid module script."),
-				LOCTEXT("InvalidModuleScriptError", "The script this module is supposed to execute is missing or invalid for other reasons.  If it depends on an external script that no longer exists there will be load errors in the log."),
+				LOCTEXT("InvalidModuleScriptError", "The script this module is supposed to execute is missing or invalid for other reasons."),
 				GetStackEditorDataKey(),
 				false);
 
@@ -371,6 +430,15 @@ void UNiagaraStackModuleItem::RefreshIssues(TArray<FStackIssue>& NewIssues)
 						: ((Dependency.Type == ENiagaraModuleDependencyType::PreDependency && Distance < 0)
 							|| (Dependency.Type == ENiagaraModuleDependencyType::PostDependency && Distance > 0));
 
+					bool bSameScriptDependencyConstraint = Dependency.ScriptConstraint == ENiagaraModuleDependencyScriptConstraint::SameScript;
+					bool bEquivalentScriptUsage = UNiagaraScript::IsEquivalentUsage(OutputNode->GetUsage(), ModuleData.Usage);
+
+					// If the dependency is for modules in the same script, the two modules are only incorrectly ordered if they share equivalent script usages
+					if (bSameScriptDependencyConstraint)
+					{
+						bIncorrectOrder = bEquivalentScriptUsage && bIncorrectOrder;
+					}
+
 					if (bIncorrectOrder)
 					{
 						DisorderedDependencies.Add(ModuleData);
@@ -379,7 +447,8 @@ void UNiagaraStackModuleItem::RefreshIssues(TArray<FStackIssue>& NewIssues)
 					{
 						DisabledDependencies.Add(FunctionNode);
 					}
-					else
+					else if (Dependency.ScriptConstraint == ENiagaraModuleDependencyScriptConstraint::AllScripts || 
+						(bSameScriptDependencyConstraint && bEquivalentScriptUsage && OutputNode->GetUsageId() == ModuleData.UsageId))
 					{
 						bDependencyMet = true;
 						break;
@@ -439,6 +508,17 @@ void UNiagaraStackModuleItem::RefreshIssues(TArray<FStackIssue>& NewIssues)
 					FNiagaraStackGraphUtilities::GetScriptAssetsByDependencyProvided(ENiagaraScriptUsage::Module, Dependency.Id, ModuleAssets);
 					for (FAssetData ModuleAsset : ModuleAssets)
 					{
+						UNiagaraScript* DependencyScript = Cast<UNiagaraScript>(ModuleAsset.GetAsset());
+						if (Dependency.ScriptConstraint == ENiagaraModuleDependencyScriptConstraint::SameScript)
+						{
+							TArray<ENiagaraScriptUsage> SupportedUsages = UNiagaraScript::GetSupportedUsageContextsForBitmask(DependencyScript->ModuleUsageBitmask);
+							if (SupportedUsages.Contains(OutputNode->GetUsage()) == false)
+							{
+								// If the dependency requires the provider be in the same script and the usage of this module doesn't support that usage, skip it.
+								continue;
+							}
+						}
+						
 						FText FixDescription = FText::Format(LOCTEXT("AddDependency", "Add new dependency module {0}"), FText::FromName(ModuleAsset.AssetName));
 						UNiagaraStackEntry::FStackIssueFix Fix(
 							FixDescription,
@@ -447,25 +527,26 @@ void UNiagaraStackModuleItem::RefreshIssues(TArray<FStackIssue>& NewIssues)
 							FScopedTransaction ScopedTransaction(FixDescription);
 							UNiagaraNodeFunctionCall* NewModuleNode = nullptr;
 							int32 TargetIndex = 0;
-							UNiagaraScript* DependencyScript = Cast<UNiagaraScript>(ModuleAsset.GetAsset());
 							checkf(DependencyScript != nullptr, TEXT("Add module action failed"));
 							// Determine the output node for the group where the added dependency module belongs
 							UNiagaraNodeOutput* TargetOutputNode = nullptr;
 							for (int i = ModuleIndex; i < SystemModuleData.Num() && i >= 0; i = Dependency.Type == ENiagaraModuleDependencyType::PostDependency ? i + 1 : i - 1) // moving up or down depending on type
-							// starting at current module, which is a dependant
+							// starting at current module, which is a dependent
 							{
-								auto FoundRequirement = SystemModuleData[i].ModuleNode->FunctionScript->RequiredDependencies.FindByPredicate(
-									[&](FNiagaraModuleDependency CurrentDependency)
+								bool bRequiredDependencyFound = SystemModuleData[i].ModuleNode->FunctionScript->RequiredDependencies.ContainsByPredicate(
+									[&Dependency](const FNiagaraModuleDependency& RequiredDependency)
 								{
-									return CurrentDependency.Id == Dependency.Id;
+									return RequiredDependency.Id == Dependency.Id;
 								});
-								if (FoundRequirement) // check for multiple dependendants along the way, and stop adjacent to the last one
+								if (bRequiredDependencyFound) // check for multiple dependents along the way, and stop adjacent to the last one
 								{
 									ENiagaraScriptUsage DependencyUsage = SystemModuleData[i].Usage;
-									TargetOutputNode = GetOutputNodeForModuleDependency(DependencyUsage, DependencyScript, GetSystemViewModel()->GetSystem(), GetEmitterViewModel()->GetEmitter(), Dependency);
-									if (TargetOutputNode != nullptr)
+									const FNiagaraEmitterHandle* EmitterHandle = FNiagaraEditorUtilities::GetEmitterHandleForEmitter(GetSystemViewModel()->GetSystem(), *GetEmitterViewModel()->GetEmitter());
+									UNiagaraNodeOutput* FoundTargetOutputNode = GetOutputNodeForModuleDependency(DependencyUsage, DependencyScript, GetSystemViewModel()->GetSystem(), EmitterHandle, Dependency);
+									if (FoundTargetOutputNode != nullptr)
 									{
-										auto CurrentOutputNode = FNiagaraStackGraphUtilities::GetEmitterOutputNodeForStackNode(*SystemModuleData[i].ModuleNode);
+										TargetOutputNode = FoundTargetOutputNode;
+										UNiagaraNodeOutput* CurrentOutputNode = FNiagaraStackGraphUtilities::GetEmitterOutputNodeForStackNode(*SystemModuleData[i].ModuleNode);
 										if (TargetOutputNode == CurrentOutputNode)
 										{
 											TargetIndex = Dependency.Type == ENiagaraModuleDependencyType::PostDependency ? SystemModuleData[i].Index + 1 : SystemModuleData[i].Index;
@@ -557,9 +638,56 @@ void UNiagaraStackModuleItem::RefreshIssues(TArray<FStackIssue>& NewIssues)
 
 bool UNiagaraStackModuleItem::FilterOutputCollection(const UNiagaraStackEntry& Child) const
 {
-	if (Child.IsA<UNiagaraStackModuleItemOutputCollection>() && GetStackEditorData().GetShowOutputs() == false)
+	if (Child.IsA<UNiagaraStackModuleItemOutputCollection>())
+	{
+		TArray<UNiagaraStackEntry*> FilteredChildren;
+		Child.GetFilteredChildren(FilteredChildren);
+		if (FilteredChildren.Num() != 0)
+		{
+			return true;
+		}
+		else if (GetStackEditorData().GetShowOutputs() == false)
+		{
+			return false;
+		}
+	}
+	return true;
+}
+
+bool UNiagaraStackModuleItem::FilterOutputCollectionChild(const UNiagaraStackEntry& Child) const
+{
+	// Filter to only show search result matches inside collapsed collection
+	if (GetStackEditorData().GetShowOutputs() == false)
+	{
+		return Child.GetIsSearchResult();
+	}
+	return true;
+}
+
+bool UNiagaraStackModuleItem::FilterLinkedInputCollection(const UNiagaraStackEntry& Child) const
+{
+	if (Child.IsA<UNiagaraStackModuleItemLinkedInputCollection>())
+	{
+		TArray<UNiagaraStackEntry*> FilteredChildren;
+		Child.GetFilteredChildren(FilteredChildren);
+		if (FilteredChildren.Num() != 0)
+		{
+ 			return true;
+		}
+		else if (GetStackEditorData().GetShowLinkedInputs() == false && Child.GetShouldShowInStack())
+		{
+			return false;
+		}
+	}
+	return true;
+}
+
+bool UNiagaraStackModuleItem::FilterLinkedInputCollectionChild(const UNiagaraStackEntry& Child) const
+{
+	// Filter to only show search result matches inside collapsed collection
+	if (GetStackEditorData().GetShowLinkedInputs() == false)
 	{
-		return false;
+		return Child.GetIsSearchResult();
 	}
 	return true;
 }
@@ -683,6 +811,29 @@ void UNiagaraStackModuleItem::AddInput(FNiagaraVariable InputParameter)
 	}
 }
 
+bool UNiagaraStackModuleItem::GetIsModuleScriptReassignmentPending() const
+{
+	return bIsModuleScriptReassignmentPending;
+}
+
+void UNiagaraStackModuleItem::SetIsModuleScriptReassignmentPending(bool bIsPending)
+{
+	bIsModuleScriptReassignmentPending = bIsPending;
+}
+
+void UNiagaraStackModuleItem::ReassignModuleScript(UNiagaraScript* ModuleScript)
+{
+	if (ensureMsgf(FunctionCallNode != nullptr && FunctionCallNode->GetClass() == UNiagaraNodeFunctionCall::StaticClass(),
+		TEXT("Can not reassign the module script when the module isn't a valid function call module.")))
+	{
+		FScopedTransaction ScopedTransaction(LOCTEXT("ReassignModuleTransaction", "Reassign module script"));
+		FunctionCallNode->Modify();
+		FunctionCallNode->FunctionScript = ModuleScript;
+		FunctionCallNode->MarkNodeRequiresSynchronization(TEXT("Module script reassigned."), true);
+		RefreshChildren();
+	}
+}
+
 UObject* UNiagaraStackModuleItem::GetExternalAsset() const
 {
 	if (GetModuleNode().FunctionScript != nullptr && GetModuleNode().FunctionScript->IsAsset())
diff --git a/Engine/Plugins/FX/Niagara/Source/NiagaraEditor/Private/ViewModels/Stack/NiagaraStackPropertyRow.cpp b/Engine/Plugins/FX/Niagara/Source/NiagaraEditor/Private/ViewModels/Stack/NiagaraStackPropertyRow.cpp
--- a/Engine/Plugins/FX/Niagara/Source/NiagaraEditor/Private/ViewModels/Stack/NiagaraStackPropertyRow.cpp
+++ b/Engine/Plugins/FX/Niagara/Source/NiagaraEditor/Private/ViewModels/Stack/NiagaraStackPropertyRow.cpp
@@ -52,13 +52,14 @@ void UNiagaraStackPropertyRow::RefreshChildrenInternal(const TArray<UNiagaraStac
 		NewChildren.Add(ChildRow);
 	}
 }
-
-void UNiagaraStackPropertyRow::GetAdditionalSearchItemsInternal(TArray<FStackSearchItem>& SearchItems) const
+void UNiagaraStackPropertyRow::GetSearchItems(TArray<FStackSearchItem>& SearchItems) const
 {
+	SearchItems.Add({ FName("DisplayName"), GetDisplayName() });
+
 	TArray<FString> NodeFilterStrings;
 	DetailTreeNode->GetFilterStrings(NodeFilterStrings);
 	for (FString& FilterString : NodeFilterStrings)
 	{
 		SearchItems.Add({ "PropertyRowFilterString", FText::FromString(FilterString) });
 	}
-}
\ No newline at end of file
+}
diff --git a/Engine/Plugins/FX/Niagara/Source/NiagaraEditor/Private/ViewModels/Stack/NiagaraStackScriptItemGroup.cpp b/Engine/Plugins/FX/Niagara/Source/NiagaraEditor/Private/ViewModels/Stack/NiagaraStackScriptItemGroup.cpp
--- a/Engine/Plugins/FX/Niagara/Source/NiagaraEditor/Private/ViewModels/Stack/NiagaraStackScriptItemGroup.cpp
+++ b/Engine/Plugins/FX/Niagara/Source/NiagaraEditor/Private/ViewModels/Stack/NiagaraStackScriptItemGroup.cpp
@@ -397,7 +397,7 @@ void UNiagaraStackScriptItemGroup::RefreshIssues(TArray<FStackIssue>& NewIssues)
 			// The factor ensures this, but older assets may not have it or it may have been removed accidentally.
 			// For now, treat this as an error and allow them to resolve.
 			FString ModulePath = TEXT("/Niagara/Modules/System/SystemLifeCycle.SystemLifeCycle");
-			FStringAssetReference SystemUpdateScriptRef(ModulePath);
+			FSoftObjectPath SystemUpdateScriptRef(ModulePath);
 			FAssetData ModuleScriptAsset;
 			ModuleScriptAsset.ObjectPath = SystemUpdateScriptRef.GetAssetPathName();
 
@@ -643,13 +643,6 @@ void UNiagaraStackScriptItemGroup::OnScriptGraphChanged(const struct FEdGraphEdi
 
 void UNiagaraStackScriptItemGroup::AddParameterModuleToStack(const UNiagaraStackModuleSpacer* InModuleSpacer, const FNiagaraVariable &InVariable)
 {
-	TArray<FNiagaraVariable> Vars;
-	Vars.Add(InVariable);
-	TArray<FString> DefaultVals;
-	DefaultVals.Add(FNiagaraConstants::GetAttributeDefaultValue(InVariable));
-
-	UNiagaraNodeOutput* OutputNode = GetScriptOutputNode();
-
 	int32 TargetIndex = INDEX_NONE;
 	UNiagaraStackModuleItem** TargetModuleItemPtr = StackSpacerToModuleItemMap.Find(FObjectKey(InModuleSpacer));
 	if (*TargetModuleItemPtr != nullptr)
@@ -658,8 +651,8 @@ void UNiagaraStackScriptItemGroup::AddParameterModuleToStack(const UNiagaraStack
 		TargetIndex = TargetModuleItem->GetModuleIndex();
 	}
 
-	FNiagaraStackGraphUtilities::AddParameterModuleToStack(Vars, *OutputNode, TargetIndex, DefaultVals);
-	RefreshChildren();
+	TSharedRef<FScriptGroupAddAction> AddAction = FScriptGroupAddAction::CreateExistingParameterModuleAction(InVariable);
+	AddUtilities->ExecuteAddAction(AddAction, TargetIndex);
 }
 
 #undef LOCTEXT_NAMESPACE
diff --git a/Engine/Plugins/FX/Niagara/Source/NiagaraEditor/Private/ViewModels/Stack/NiagaraStackViewModel.cpp b/Engine/Plugins/FX/Niagara/Source/NiagaraEditor/Private/ViewModels/Stack/NiagaraStackViewModel.cpp
--- a/Engine/Plugins/FX/Niagara/Source/NiagaraEditor/Private/ViewModels/Stack/NiagaraStackViewModel.cpp
+++ b/Engine/Plugins/FX/Niagara/Source/NiagaraEditor/Private/ViewModels/Stack/NiagaraStackViewModel.cpp
@@ -111,9 +111,10 @@ void UNiagaraStackViewModel::Tick()
 
 void UNiagaraStackViewModel::OnSearchTextChanged(const FText& SearchText)
 {
-	if (RootEntry)
+	if (RootEntry && !CurrentSearchText.EqualTo(SearchText))
 	{
 		CurrentSearchText = SearchText;
+		RestoreStackEntryExpansionPreSearch();
 		// postpone searching until next tick; protects against crashes from the GC
 		// also this can be triggered by multiple events, so better wait
 		bRestartSearch = true;
@@ -225,6 +226,10 @@ void UNiagaraStackViewModel::SearchTick()
 	if (bRestartSearch)
 	{
 		// clear the search results
+		for (auto SearchResult : GetCurrentSearchResults())
+		{
+			SearchResult.GetEntry()->SetIsSearchResult(false);
+		}
 		CurrentSearchResults.Empty();
 		CurrentFocusedSearchMatchIndex = -1;
 		// generates ItemsToSearch, these will be processed on tick, in batches
@@ -232,11 +237,18 @@ void UNiagaraStackViewModel::SearchTick()
 		{
 			GenerateTraversalEntries(RootEntry, TArray<UNiagaraStackEntry*>(), ItemsToSearch);
 		}
+		// we need to call the SearchCompletedDelegate to go through SynchronizeTreeExpansion in SNiagaraStack so that when exiting search we return the stack expansion to the state it was before searching
+		else
+		{
+			SearchCompletedDelegate.Broadcast();
+		}
 		bRestartSearch = false;
 	}
-	else if (IsSearching())
+
+	if (IsSearching())
 	{
-		double SearchStartTime= FPlatformTime::Seconds();
+		UNiagaraStackEditorData& EditorData = GetSystemViewModel()->GetEditorData().GetStackEditorData();
+		double SearchStartTime = FPlatformTime::Seconds();
 		double CurrentSearchLoopTime = SearchStartTime;
 		// process at least one item, but don't go over MaxSearchTime for the rest
 		while (ItemsToSearch.Num() > 0 && CurrentSearchLoopTime - SearchStartTime < MaxSearchTime)
@@ -250,10 +262,16 @@ void UNiagaraStackViewModel::SearchTick()
 				TSet<FName> MatchedKeys;
 				for (UNiagaraStackEntry::FStackSearchItem SearchItem : SearchItems)
 				{
+					if (!ItemsToSearch[0].GetEntry()->GetStackEditorDataKey().IsEmpty())
+					{
+						EditorData.SetStackEntryWasExpandedPreSearch(ItemsToSearch[0].GetEntry()->GetStackEditorDataKey(), ItemsToSearch[0].GetEntry()->GetIsExpanded());
+					}
+
 					if (ItemMatchesSearchCriteria(SearchItem)) 
 					{
 						if (MatchedKeys.Contains(SearchItem.Key) == false)
 						{
+							EntryToProcess->SetIsSearchResult(true);
 							CurrentSearchResults.Add({ ItemsToSearch[0].EntryPath, SearchItem });
 							MatchedKeys.Add(SearchItem.Key);
 						}
@@ -274,7 +292,7 @@ void UNiagaraStackViewModel::GenerateTraversalEntries(UNiagaraStackEntry* Root,
 	TArray<FSearchWorkItem>& TraversedArray)
 {
 	TArray<UNiagaraStackEntry*> Children;
-	Root->GetFilteredChildren(Children);
+	Root->GetUnfilteredChildren(Children);
 	ParentChain.Add(Root);
 	TraversedArray.Add(FSearchWorkItem{ParentChain});
 	for (auto Child : Children)
@@ -309,6 +327,28 @@ void UNiagaraStackViewModel::GeneratePathForEntry(UNiagaraStackEntry* Root, UNia
 	}
 }
 
+void UNiagaraStackViewModel::RestoreStackEntryExpansionPreSearch()
+{
+	UNiagaraStackEditorData& StackEditorData = SystemViewModel->GetEditorData().GetStackEditorData();
+	GenerateTraversalEntries(RootEntry, TArray<UNiagaraStackEntry*>(), ItemsToRestoreExpansionState);
+
+	while (ItemsToRestoreExpansionState.Num() != 0)
+	{
+		UNiagaraStackEntry* EntryToProcess = ItemsToRestoreExpansionState[0].GetEntry();
+		if (EntryToProcess->GetIsSearchResult() && !EntryToProcess->IsA<UNiagaraStackRoot>())
+		{
+			for (auto EntryToExpand : ItemsToRestoreExpansionState[0].EntryPath)
+			{
+				if (!EntryToExpand->IsA<UNiagaraStackRoot>())
+				{
+					EntryToExpand->SetIsExpanded(StackEditorData.GetStackEntryWasExpandedPreSearch(EntryToExpand->GetStackEditorDataKey(), false));
+				}
+			}
+		}
+		ItemsToRestoreExpansionState.RemoveAtSwap(0);
+	}
+}
+
 TArray<UNiagaraStackEntry*>& UNiagaraStackViewModel::GetRootEntries()
 {
 	return RootEntries;
@@ -368,6 +408,29 @@ void UNiagaraStackViewModel::SetShowOutputs(bool bInShowOutputs)
 	}
 }
 
+bool UNiagaraStackViewModel::GetShowLinkedInputs() const
+{
+	if (SystemViewModel.IsValid() && EmitterHandleViewModel.IsValid())
+	{
+		return SystemViewModel->GetEditorData().GetStackEditorData().GetShowLinkedInputs() ||
+			EmitterHandleViewModel->GetEmitterViewModel()->GetEditorData().GetStackEditorData().GetShowLinkedInputs();
+	}
+	return false;
+}
+
+void UNiagaraStackViewModel::SetShowLinkedInputs(bool bInShowLinkedInputs)
+{
+	if (SystemViewModel.IsValid() && EmitterHandleViewModel.IsValid())
+	{
+		SystemViewModel->GetOrCreateEditorData().GetStackEditorData().SetShowLinkedInputs(bInShowLinkedInputs);
+		EmitterHandleViewModel->GetEmitterViewModel()->GetOrCreateEditorData().GetStackEditorData().SetShowLinkedInputs(bInShowLinkedInputs);
+		OnSearchTextChanged(CurrentSearchText);
+
+		// Showing linked inputs changes indenting so a full refresh is needed.
+		RootEntry->RefreshChildren();
+	}
+}
+
 double UNiagaraStackViewModel::GetLastScrollPosition() const
 {
 	if (EmitterHandleViewModel.IsValid())
diff --git a/Engine/Plugins/FX/Niagara/Source/NiagaraEditor/Private/Widgets/SNiagaraParameterMapView.cpp b/Engine/Plugins/FX/Niagara/Source/NiagaraEditor/Private/Widgets/SNiagaraParameterMapView.cpp
--- a/Engine/Plugins/FX/Niagara/Source/NiagaraEditor/Private/Widgets/SNiagaraParameterMapView.cpp
+++ b/Engine/Plugins/FX/Niagara/Source/NiagaraEditor/Private/Widgets/SNiagaraParameterMapView.cpp
@@ -280,6 +280,11 @@ void SNiagaraParameterMapView::AddParameter(FNiagaraVariable NewVariable)
 	}
 }
 
+bool SNiagaraParameterMapView::AllowMakeType(const FNiagaraTypeDefinition& InType) const
+{
+	return InType != FNiagaraTypeDefinition::GetParameterMapDef();
+}
+
 void SNiagaraParameterMapView::OnFilterTextChanged(const FText& InFilterText)
 {
 	GraphActionMenu->GenerateFilteredItems(false);
@@ -496,6 +501,7 @@ TSharedRef<SWidget> SNiagaraParameterMapView::OnGetParameterMenu(const NiagaraPa
 {
 	TSharedRef<SNiagaraAddParameterMenu> MenuWidget = SNew(SNiagaraAddParameterMenu, Graphs)
 		.OnAddParameter(this, &SNiagaraParameterMapView::AddParameter)
+		.OnAllowMakeType(this, &SNiagaraParameterMapView::AllowMakeType)
 		.Section(InSection)
 		.ShowNamespaceCategory(false)
 		.ShowGraphParameters(false)
diff --git a/Engine/Plugins/FX/Niagara/Source/NiagaraEditor/Private/Widgets/SNiagaraParameterMapView.h b/Engine/Plugins/FX/Niagara/Source/NiagaraEditor/Private/Widgets/SNiagaraParameterMapView.h
--- a/Engine/Plugins/FX/Niagara/Source/NiagaraEditor/Private/Widgets/SNiagaraParameterMapView.h
+++ b/Engine/Plugins/FX/Niagara/Source/NiagaraEditor/Private/Widgets/SNiagaraParameterMapView.h
@@ -93,6 +93,9 @@ public:
 	static TSharedRef<SExpanderArrow> CreateCustomActionExpander(const struct FCustomExpanderData& ActionMenuData);
 
 private:
+	/** Function to bind to SNiagaraAddParameterMenus to filter types we allow creating */
+	bool AllowMakeType(const FNiagaraTypeDefinition& InType) const;
+
 	/** Callback when the filter is changed, forces the action tree(s) to filter */
 	void OnFilterTextChanged(const FText& InFilterText);
 
diff --git a/Engine/Plugins/FX/Niagara/Source/NiagaraEditor/Private/Widgets/SNiagaraSpreadsheetView.cpp b/Engine/Plugins/FX/Niagara/Source/NiagaraEditor/Private/Widgets/SNiagaraSpreadsheetView.cpp
--- a/Engine/Plugins/FX/Niagara/Source/NiagaraEditor/Private/Widgets/SNiagaraSpreadsheetView.cpp
+++ b/Engine/Plugins/FX/Niagara/Source/NiagaraEditor/Private/Widgets/SNiagaraSpreadsheetView.cpp
@@ -265,6 +265,12 @@ void SNiagaraSpreadsheetView::Construct(const FArguments& InArgs, TSharedRef<FNi
 			CaptureData[i].bOutputColumnsAreAttributes = false;
 			CaptureData[i].bInputColumnsAreAttributes = false;
 			break;
+		case UIPerParticleGPU:
+			CaptureData[i].TargetUsage = ENiagaraScriptUsage::ParticleGPUComputeScript;
+			CaptureData[i].ColumnName = LOCTEXT("PerParticleGPU", "Particle GPU");
+			CaptureData[i].bOutputColumnsAreAttributes = true;
+			CaptureData[i].bInputColumnsAreAttributes = false;
+			break;
 		default:
 			CaptureData[i].TargetUsage = ENiagaraScriptUsage::Function;
 			CaptureData[i].ColumnName = LOCTEXT("PerParticleUnknown", "Particle Unknown");
@@ -562,6 +568,11 @@ void SNiagaraSpreadsheetView::Construct(const FArguments& InArgs, TSharedRef<FNi
 					]
 					+ SHorizontalBox::Slot()
 					.AutoWidth()
+					[
+						CaptureData[UIPerParticleGPU].CheckBox.ToSharedRef()
+					]
+					+ SHorizontalBox::Slot()
+					.AutoWidth()
 					[
 						CaptureData[UISystemUpdate].CheckBox.ToSharedRef()
 					]
@@ -589,6 +600,10 @@ void SNiagaraSpreadsheetView::Construct(const FArguments& InArgs, TSharedRef<FNi
 			CaptureData[UIPerParticleEvent2].Container.ToSharedRef()
 		]
 		+ SVerticalBox::Slot()
+		[
+			CaptureData[UIPerParticleGPU].Container.ToSharedRef()
+		]
+		+ SVerticalBox::Slot()
 		[
 			CaptureData[UISystemUpdate].Container.ToSharedRef()
 		]
@@ -943,7 +958,7 @@ FReply SNiagaraSpreadsheetView::OnCSVOutputPressed()
 					{
 						uint32 CompBufferOffset = FieldInfo->FloatStartOffset;
 						float* Src = CaptureData[(int32)TabState].DataSet.PrevData().GetInstancePtrFloat(CompBufferOffset, RowIndex);
-						CSVOutput += FString::Printf(TEXT("%3.3f"), Src[0]);
+						CSVOutput += FString::Printf(TEXT("%3.9f"), Src[0]);
 					}
 					else
 					{
@@ -1423,6 +1438,10 @@ FReply SNiagaraSpreadsheetView::OnCaptureRequestPressed()
 					CaptureData[i].TargetUsage = ENiagaraScriptUsage::SystemUpdateScript;
 					CaptureData[i].TargetUsageId = FGuid();
 					break;
+				case UIPerParticleGPU:
+					CaptureData[i].TargetUsage = ENiagaraScriptUsage::ParticleGPUComputeScript;
+					CaptureData[i].TargetUsageId = FGuid();
+					break;
 				default:
 					CaptureData[i].TargetUsage = ENiagaraScriptUsage::Function;
 					CaptureData[i].TargetUsageId = FGuid();
diff --git a/Engine/Plugins/FX/Niagara/Source/NiagaraEditor/Private/Widgets/SNiagaraSpreadsheetView.h b/Engine/Plugins/FX/Niagara/Source/NiagaraEditor/Private/Widgets/SNiagaraSpreadsheetView.h
--- a/Engine/Plugins/FX/Niagara/Source/NiagaraEditor/Private/Widgets/SNiagaraSpreadsheetView.h
+++ b/Engine/Plugins/FX/Niagara/Source/NiagaraEditor/Private/Widgets/SNiagaraSpreadsheetView.h
@@ -54,6 +54,7 @@ protected:
 		UIPerParticleEvent1,
 		UIPerParticleEvent2,
 		UIPerParticleEvent3,
+		UIPerParticleGPU,
 		UISystemUpdate,
 		UIMax
 	};
diff --git a/Engine/Plugins/FX/Niagara/Source/NiagaraEditor/Public/EdGraphSchema_Niagara.h b/Engine/Plugins/FX/Niagara/Source/NiagaraEditor/Public/EdGraphSchema_Niagara.h
--- a/Engine/Plugins/FX/Niagara/Source/NiagaraEditor/Public/EdGraphSchema_Niagara.h
+++ b/Engine/Plugins/FX/Niagara/Source/NiagaraEditor/Public/EdGraphSchema_Niagara.h
@@ -85,6 +85,7 @@ class NIAGARAEDITOR_API UEdGraphSchema_Niagara : public UEdGraphSchema
 	void PromoteSinglePinToParameter(UEdGraphPin* SourcePin);
 	static bool CanPromoteSinglePinToParameter(const UEdGraphPin* SourcePin);
 	void ToggleNodeEnabledState(class UNiagaraNode* InNode) const;
+	void RefreshNode(UNiagaraNode* InNode) const;
 
 	/** 
 	  * Creates a niagara variable using the name, type, and default value stored on an ed graph pin.
@@ -136,6 +137,6 @@ public:
 	virtual void DetermineWiringStyle(UEdGraphPin* OutputPin, UEdGraphPin* InputPin, /*inout*/ FConnectionParams& Params) override;
 
 private:
-	UEdGraph* Graph;
+	class UNiagaraGraph* Graph;
 };
 
diff --git a/Engine/Plugins/FX/Niagara/Source/NiagaraEditor/Public/NiagaraEditorCommon.h b/Engine/Plugins/FX/Niagara/Source/NiagaraEditor/Public/NiagaraEditorCommon.h
--- a/Engine/Plugins/FX/Niagara/Source/NiagaraEditor/Public/NiagaraEditorCommon.h
+++ b/Engine/Plugins/FX/Niagara/Source/NiagaraEditor/Public/NiagaraEditorCommon.h
@@ -12,7 +12,7 @@ class FNiagaraOpInfo
 {
 public:
 	FNiagaraOpInfo()
-		: Keywords(FText()), NumericOuputTypeSelectionMode(ENiagaraNumericOutputTypeSelectionMode::Largest)
+		: Keywords(FText()), NumericOuputTypeSelectionMode(ENiagaraNumericOutputTypeSelectionMode::Largest), bSupportsAddedInputs(false)
 	{}
 
 	FName Name;
@@ -24,6 +24,21 @@ public:
 	TArray<FNiagaraOpInOutInfo> Inputs;
 	TArray<FNiagaraOpInOutInfo> Outputs;
 
+	/** If true then this operation supports a variable number of inputs */
+	bool bSupportsAddedInputs;
+
+	/** 
+	* The format that can generate the hlsl for the given number of inputs.
+	* Used the placeholder {A} and {B} to chain the inputs together.
+	*/
+	FString AddedInputFormatting;
+
+	/**
+	* If added inputs are enabled then this filters the available pin types shown to the user.
+	* If empty then all the default niagara types are shown.
+	*/
+	TArray<FNiagaraTypeDefinition> AddedInputTypeRestrictions;
+
 	static TMap<FName, int32> OpInfoMap;
 	static TArray<FNiagaraOpInfo> OpInfos;
 
@@ -32,5 +47,7 @@ public:
 	static const TArray<FNiagaraOpInfo>& GetOpInfoArray();
 
 	void BuildName(FString InName, FString InCategory);
+
+	bool CreateHlslForAddedInputs(int32 InputCount, FString& HlslResult) const;
 };
 
diff --git a/Engine/Plugins/FX/Niagara/Source/NiagaraEditor/Public/NiagaraGraph.h b/Engine/Plugins/FX/Niagara/Source/NiagaraEditor/Public/NiagaraGraph.h
--- a/Engine/Plugins/FX/Niagara/Source/NiagaraEditor/Public/NiagaraGraph.h
+++ b/Engine/Plugins/FX/Niagara/Source/NiagaraEditor/Public/NiagaraGraph.h
@@ -75,6 +75,10 @@ public:
 	TArray<UNiagaraNode*> Traversal;
 };
 
+struct FNiagaraGraphFunctionAliasContext
+{
+	ENiagaraScriptUsage CompileUsage;
+};
 
 UCLASS(MinimalAPI)
 class UNiagaraGraph : public UEdGraph
@@ -231,8 +235,20 @@ class UNiagaraGraph : public UEdGraph
 	/** Remove a listener for OnGraphNeedsRecompile events */
 	void RemoveOnGraphNeedsRecompileHandler(FDelegateHandle Handle);
 
+	FNiagaraTypeDefinition GetCachedNumericConversion(class UEdGraphPin* InPin);
+
+	const class UEdGraphSchema_Niagara* GetNiagaraSchema() const;
+
+	void InvalidateNumericCache();
+
+	FString GetFunctionAliasByContext(const FNiagaraGraphFunctionAliasContext& FunctionAliasContext);
+
 protected:
 	void RebuildCachedData(bool bForce = false);
+	void RebuildNumericCache();
+	bool bNeedNumericCacheRebuilt;
+	TMap<TPair<FGuid, UEdGraphNode*>, FNiagaraTypeDefinition> CachedNumericConversions;
+	void ResolveNumerics(TMap<UNiagaraNode*, bool>& VisitedNodes, UEdGraphNode* Node);
 
 private:
 	virtual void NotifyGraphChanged(const FEdGraphEditAction& InAction) override;
diff --git a/Engine/Plugins/FX/Niagara/Source/NiagaraEditor/Public/NiagaraNode.h b/Engine/Plugins/FX/Niagara/Source/NiagaraEditor/Public/NiagaraNode.h
--- a/Engine/Plugins/FX/Niagara/Source/NiagaraEditor/Public/NiagaraNode.h
+++ b/Engine/Plugins/FX/Niagara/Source/NiagaraEditor/Public/NiagaraNode.h
@@ -13,6 +13,7 @@
 
 class UEdGraphPin;
 class INiagaraCompiler;
+struct FNiagaraGraphFunctionAliasContext;
 
 UCLASS()
 class NIAGARAEDITOR_API UNiagaraNode : public UEdGraphNode
@@ -65,6 +66,7 @@ public:
 	UEdGraphPin* GetOutputPin(int32 OutputIndex) const;
 	void GetOutputPins(TArray<class UEdGraphPin*>& OutOutputPins) const;
 	UEdGraphPin* GetPinByPersistentGuid(const FGuid& InGuid) const;
+	virtual void ResolveNumerics(const UEdGraphSchema_Niagara* Schema, bool bSetInline, TMap<TPair<FGuid, UEdGraphNode*>, FNiagaraTypeDefinition>* PinCache);
 
 	/** Apply any node-specific logic to determine if it is safe to add this node to the graph. This is meant to be called only in the Editor before placing the node.*/
 	virtual bool CanAddToGraph(UNiagaraGraph* TargetGraph, FString& OutErrorMsg) const;
@@ -114,10 +116,16 @@ public:
 	void ForceChangeId(const FGuid& InId, bool bRaiseGraphNeedsRecompile);
 
 	FOnNodeVisualsChanged& OnVisualsChanged();
+
+	virtual void AppendFunctionAliasForContext(const FNiagaraGraphFunctionAliasContext& InFunctionAliasContext, FString& InOutFunctionAlias) { };
+
 protected:
 	virtual int32 CompileInputPin(class FHlslNiagaraTranslator *Translator, UEdGraphPin* Pin);
 	virtual bool IsValidPinToCompile(UEdGraphPin* Pin) const { return true; }; 
 
+	void NumericResolutionByPins(const UEdGraphSchema_Niagara* Schema, TArray<UEdGraphPin*>& InputPins, TArray<UEdGraphPin*>& OutputPins,
+		bool bFixInline, TMap<TPair<FGuid, UEdGraphNode*>, FNiagaraTypeDefinition>* PinCache);
+
 	/** Route input parameter map to output parameter map if it exists. Note that before calling this function,
 		the input pins should have been visited already.*/
 	virtual void RouteParameterMapAroundMe(FNiagaraParameterMapHistoryBuilder& OutHistory, bool bRecursive);
diff --git a/Engine/Plugins/FX/Niagara/Source/NiagaraEditor/Public/NiagaraNodeAssignment.h b/Engine/Plugins/FX/Niagara/Source/NiagaraEditor/Public/NiagaraNodeAssignment.h
--- a/Engine/Plugins/FX/Niagara/Source/NiagaraEditor/Public/NiagaraNodeAssignment.h
+++ b/Engine/Plugins/FX/Niagara/Source/NiagaraEditor/Public/NiagaraNodeAssignment.h
@@ -59,10 +59,10 @@ protected:
 	UPROPERTY()
 	FString AssignmentDefaultValue_DEPRECATED;
 
-	UPROPERTY(EditAnywhere, Category = Assignment)
+	UPROPERTY()
 	TArray<FNiagaraVariable> AssignmentTargets;
 
-	UPROPERTY(EditAnywhere, Category = Assignment)
+	UPROPERTY()
 	TArray<FString> AssignmentDefaultValues;
 
 	UPROPERTY()
diff --git a/Engine/Plugins/FX/Niagara/Source/NiagaraEditor/Public/NiagaraNodeCustomHlsl.h b/Engine/Plugins/FX/Niagara/Source/NiagaraEditor/Public/NiagaraNodeCustomHlsl.h
--- a/Engine/Plugins/FX/Niagara/Source/NiagaraEditor/Public/NiagaraNodeCustomHlsl.h
+++ b/Engine/Plugins/FX/Niagara/Source/NiagaraEditor/Public/NiagaraNodeCustomHlsl.h
@@ -20,7 +20,7 @@ public:
 	UPROPERTY(EditAnywhere, Category = "Function", meta = (MultiLine = true))
 	FString CustomHlsl;
 
-	UPROPERTY(EditAnywhere, Category = "Function")
+	UPROPERTY()
 	ENiagaraScriptUsage ScriptUsage;
 
 	virtual TSharedPtr<SGraphNode> CreateVisualWidget() override;
diff --git a/Engine/Plugins/FX/Niagara/Source/NiagaraEditor/Public/NiagaraNodeDataSetBase.h b/Engine/Plugins/FX/Niagara/Source/NiagaraEditor/Public/NiagaraNodeDataSetBase.h
--- a/Engine/Plugins/FX/Niagara/Source/NiagaraEditor/Public/NiagaraNodeDataSetBase.h
+++ b/Engine/Plugins/FX/Niagara/Source/NiagaraEditor/Public/NiagaraNodeDataSetBase.h
@@ -16,10 +16,10 @@ class UNiagaraNodeDataSetBase : public UNiagaraNode
 
 public:
 
-	UPROPERTY(EditAnywhere, Category = DataSet)
+	UPROPERTY()
 	FNiagaraDataSetID DataSet;
 
-	UPROPERTY(EditAnywhere, Category = Variables)
+	UPROPERTY()
 	TArray<FNiagaraVariable> Variables;
 	
 	UPROPERTY()
diff --git a/Engine/Plugins/FX/Niagara/Source/NiagaraEditor/Public/NiagaraNodeFunctionCall.h b/Engine/Plugins/FX/Niagara/Source/NiagaraEditor/Public/NiagaraNodeFunctionCall.h
--- a/Engine/Plugins/FX/Niagara/Source/NiagaraEditor/Public/NiagaraNodeFunctionCall.h
+++ b/Engine/Plugins/FX/Niagara/Source/NiagaraEditor/Public/NiagaraNodeFunctionCall.h
@@ -59,6 +59,9 @@ public:
 	virtual FText GetNodeTitle(ENodeTitleType::Type TitleType) const override;
 	virtual FText GetTooltipText() const override;
 	virtual FLinearColor GetNodeTitleColor() const override;
+	
+	/** Returns true if this node is deprecated */
+	virtual bool IsDeprecated() const override;
 	//~ End EdGraphNode Interface
 
 	bool FindAutoBoundInput(UNiagaraNodeInput* InputNode, UEdGraphPin* PinToAutoBind, FNiagaraVariable& OutFoundVar, ENiagaraInputNodeUsage& OutNodeUsage);
diff --git a/Engine/Plugins/FX/Niagara/Source/NiagaraEditor/Public/NiagaraNodeInput.h b/Engine/Plugins/FX/Niagara/Source/NiagaraEditor/Public/NiagaraNodeInput.h
--- a/Engine/Plugins/FX/Niagara/Source/NiagaraEditor/Public/NiagaraNodeInput.h
+++ b/Engine/Plugins/FX/Niagara/Source/NiagaraEditor/Public/NiagaraNodeInput.h
@@ -49,10 +49,10 @@ class UNiagaraNodeInput : public UNiagaraNode
 	GENERATED_UCLASS_BODY()
 
 public:
-	UPROPERTY(EditAnywhere, Category = Input)
+	UPROPERTY(EditAnywhere, Category = Input) 
 	FNiagaraVariable Input;
 
-	UPROPERTY(VisibleAnywhere, Category = Input)
+	UPROPERTY(EditAnywhere, Category = Input)
 	ENiagaraInputNodeUsage Usage;
 
 	/** Controls where this input is relative to others in the calling node. */
