Fix for UE-69871

diff --git a/Engine/Plugins/FX/Niagara/Source/Niagara/Private/NiagaraPreviewGrid.cpp b/Engine/Plugins/FX/Niagara/Source/Niagara/Private/NiagaraPreviewGrid.cpp
--- a/Engine/Plugins/FX/Niagara/Source/Niagara/Private/NiagaraPreviewGrid.cpp
+++ b/Engine/Plugins/FX/Niagara/Source/Niagara/Private/NiagaraPreviewGrid.cpp
@@ -53,8 +53,9 @@ ANiagaraPreviewGrid::ANiagaraPreviewGrid(const FObjectInitializer& ObjectInitial
 // 		PreviewClass = DefaultPreviewClassBP.Class;
 // 	}
 
+	RootComponent = CreateDefaultSubobject<USceneComponent>(TEXT("SceneComp"));
+	check(RootComponent);
 #if WITH_EDITORONLY_DATA
-	RootComponent = CreateEditorOnlyDefaultSubobject<USceneComponent>(TEXT("SceneComp"));
 	SpriteComponent = CreateEditorOnlyDefaultSubobject<UBillboardComponent>(TEXT("Sprite"));
 	ArrowComponent = CreateEditorOnlyDefaultSubobject<UArrowComponent>(TEXT("ArrowComponent0"));
 
@@ -109,6 +110,11 @@ void ANiagaraPreviewGrid::PostLoad()
 {
 	Super::PostLoad();
 
+	//Fixup old data that incorrectly did not have this component.
+	if (RootComponent == nullptr)
+	{
+		RootComponent = NewObject<USceneComponent>(this, TEXT("SceneComp"));
+	}
 }
 
 void ANiagaraPreviewGrid::BeginDestroy()
