Fixed localization conflict warnings

#rnx

diff --git a/Engine/Plugins/FX/Niagara/Source/NiagaraEditor/Private/NiagaraHlslTranslator.cpp b/Engine/Plugins/FX/Niagara/Source/NiagaraEditor/Private/NiagaraHlslTranslator.cpp
--- a/Engine/Plugins/FX/Niagara/Source/NiagaraEditor/Private/NiagaraHlslTranslator.cpp
+++ b/Engine/Plugins/FX/Niagara/Source/NiagaraEditor/Private/NiagaraHlslTranslator.cpp
@@ -5173,7 +5173,7 @@ void FHlslNiagaraTranslator::Convert(class UNiagaraNodeConvert* Convert, TArray
 			FNiagaraTypeDefinition Type = Schema->PinToTypeDefinition(InputPin);
 			if (!AddStructToDefinitionSet(Type))
 			{
-				Error(FText::Format(LOCTEXT("ConvertTypeError", "Cannot handle input pin type {0}! Pin: {1}"), Type.GetNameText(), InputPin->GetDisplayName()), nullptr, nullptr);
+				Error(FText::Format(LOCTEXT("ConvertTypeError_InvalidInput", "Cannot handle input pin type {0}! Pin: {1}"), Type.GetNameText(), InputPin->GetDisplayName()), nullptr, nullptr);
 			}
 		}
 	}
@@ -5187,7 +5187,7 @@ void FHlslNiagaraTranslator::Convert(class UNiagaraNodeConvert* Convert, TArray
 			FNiagaraTypeDefinition Type = Schema->PinToTypeDefinition(OutputPin);
 			if (!AddStructToDefinitionSet(Type))
 			{
-				Error(FText::Format(LOCTEXT("ConvertTypeError", "Cannot handle output pin type {0}! Pin: {1}"), Type.GetNameText(), OutputPin->GetDisplayName()), nullptr, nullptr);
+				Error(FText::Format(LOCTEXT("ConvertTypeError_InvalidOutput", "Cannot handle output pin type {0}! Pin: {1}"), Type.GetNameText(), OutputPin->GetDisplayName()), nullptr, nullptr);
 			}
 			int32 OutChunk = AddBodyChunk(GetUniqueSymbolName(OutputPin->PinName), TEXT(""), Type);
 			Outputs.Add(OutChunk);
@@ -5207,14 +5207,14 @@ void FHlslNiagaraTranslator::Convert(class UNiagaraNodeConvert* Convert, TArray
 			FNiagaraTypeDefinition SrcPinType = Schema->PinToTypeDefinition(InputPins[SourceIndex]);
 			if (!AddStructToDefinitionSet(SrcPinType))
 			{
-				Error(FText::Format(LOCTEXT("ConvertTypeError", "Cannot handle input subpin type {0}! Subpin: {1}"), SrcPinType.GetNameText(), InputPins[SourceIndex]->GetDisplayName()), nullptr, nullptr);
+				Error(FText::Format(LOCTEXT("ConvertTypeError_InvalidSubpinInput", "Cannot handle input subpin type {0}! Subpin: {1}"), SrcPinType.GetNameText(), InputPins[SourceIndex]->GetDisplayName()), nullptr, nullptr);
 			}
 			TArray<FName> ConditionedSourcePath = ConditionPropertyPath(SrcPinType, Connection.SourcePath);
 
 			FNiagaraTypeDefinition DestPinType = Schema->PinToTypeDefinition(OutputPins[DestinationIndex]);
 			if (!AddStructToDefinitionSet(DestPinType))
 			{
-				Error(FText::Format(LOCTEXT("ConvertTypeError", "Cannot handle output subpin type type {0}! Subpin: {1}"), DestPinType.GetNameText(), OutputPins[SourceIndex]->GetDisplayName()), nullptr, nullptr);
+				Error(FText::Format(LOCTEXT("ConvertTypeError_InvalidSubpinOutput", "Cannot handle output subpin type type {0}! Subpin: {1}"), DestPinType.GetNameText(), OutputPins[SourceIndex]->GetDisplayName()), nullptr, nullptr);
 			}
 			TArray<FName> ConditionedDestinationPath = ConditionPropertyPath(DestPinType, Connection.DestinationPath);
 
