Added support for sprite cutout geometry in Niagara.

diff --git a/Engine/Plugins/FX/Niagara/Source/Niagara/Private/NiagaraRendererSprites.cpp b/Engine/Plugins/FX/Niagara/Source/Niagara/Private/NiagaraRendererSprites.cpp
--- a/Engine/Plugins/FX/Niagara/Source/Niagara/Private/NiagaraRendererSprites.cpp
+++ b/Engine/Plugins/FX/Niagara/Source/Niagara/Private/NiagaraRendererSprites.cpp
@@ -60,14 +60,20 @@ NiagaraRendererSprites::NiagaraRendererSprites(ERHIFeatureLevel::Type FeatureLev
 	VertexFactory = new FNiagaraSpriteVertexFactory(NVFT_Sprite, FeatureLevel);
 	Properties = Cast<UNiagaraSpriteRendererProperties>(InProps);
 	BaseExtents = FVector(0.5f, 0.5f, 0.5f);
-}
 
+	if (Properties)
+	{
+		NumCutoutVertexPerSubImage = Properties->GetNumCutoutVertexPerSubimage();
+		CutoutVertexBuffer.Data = Properties->GetCutoutData();
+	}
+}
 
 void NiagaraRendererSprites::ReleaseRenderThreadResources()
 {
 	VertexFactory->ReleaseResource();
 	WorldSpacePrimitiveUniformBuffer.ReleaseResource();
 
+	CutoutVertexBuffer.ReleaseResource();
 #if RHI_RAYTRACING
 	if (IsRayTracingEnabled())
 	{
@@ -79,7 +85,7 @@ void NiagaraRendererSprites::ReleaseRenderThreadResources()
 
 void NiagaraRendererSprites::CreateRenderThreadResources()
 {
-	VertexFactory->SetNumVertsInInstanceBuffer(4);
+	CutoutVertexBuffer.InitResource();
 	VertexFactory->InitResource();
 
 #if RHI_RAYTRACING
@@ -250,6 +256,21 @@ void NiagaraRendererSprites::GetDynamicMeshElements(const TArray<const FSceneVie
 					PerViewUniformParameters.CustomFacingVectorMask = Properties->CustomFacingVectorMask;
 				}
 
+				// Cutout geometry.
+				const bool bUseSubImage = PerViewUniformParameters.SubImageSize.X != 1 || PerViewUniformParameters.SubImageSize.Y != 1;
+				const bool bUseCutout = CutoutVertexBuffer.VertexBufferRHI.IsValid();
+				if (bUseCutout)
+				{	// Is Accessing Properties safe here? Or should values be cached in the constructor?
+					if (bUseSubImage)
+					{
+						CollectorResources.VertexFactory.SetCutoutParameters(NumCutoutVertexPerSubImage, CutoutVertexBuffer.VertexBufferSRV);
+					}
+					else // Otherwise simply replace the input stream with the single cutout geometry
+					{
+						CollectorResources.VertexFactory.SetVertexBufferOverride(&CutoutVertexBuffer);
+					}
+				}
+
 				//Sort particles if needed.
 				EBlendMode BlendMode = MaterialRenderProxy->GetMaterial(VertexFactory->GetFeatureLevel())->GetBlendMode();
 				CollectorResources.VertexFactory.SetSortedIndices(nullptr, 0xFFFFFFFF);
@@ -324,10 +345,8 @@ void NiagaraRendererSprites::GetDynamicMeshElements(const TArray<const FSceneVie
 
 				CollectorResources.VertexFactory.SetParticleFactoryType(NVFT_Sprite);
 
-
 				CollectorResources.UniformBuffer = FNiagaraSpriteUniformBufferRef::CreateUniformBufferImmediate(PerViewUniformParameters, UniformBuffer_SingleFrame);
 
-				CollectorResources.VertexFactory.SetNumVertsInInstanceBuffer(4);
 				CollectorResources.VertexFactory.InitResource();
 				CollectorResources.VertexFactory.SetSpriteUniformBuffer(CollectorResources.UniformBuffer);
 
@@ -385,7 +404,13 @@ void NiagaraRendererSprites::GetDynamicMeshElements(const TArray<const FSceneVie
 					MeshElement.IndirectArgsBuffer = DynamicDataSprites->DataSet->GetCurDataSetIndices().Buffer;
 					MeshElement.NumPrimitives = 0;
 				}
-				
+
+				if (NumCutoutVertexPerSubImage == 8)
+				{
+					MeshElement.NumPrimitives = 6;
+					MeshElement.IndexBuffer = &GSixTriangleParticleIndexBuffer;
+				}
+
 				Collector.AddMesh(ViewIndex, MeshBatch);				
 			}
 		}
diff --git a/Engine/Plugins/FX/Niagara/Source/Niagara/Private/NiagaraSpriteRendererProperties.cpp b/Engine/Plugins/FX/Niagara/Source/Niagara/Private/NiagaraSpriteRendererProperties.cpp
--- a/Engine/Plugins/FX/Niagara/Source/Niagara/Private/NiagaraSpriteRendererProperties.cpp
+++ b/Engine/Plugins/FX/Niagara/Source/Niagara/Private/NiagaraSpriteRendererProperties.cpp
@@ -4,9 +4,27 @@
 #include "NiagaraRenderer.h"
 #include "Internationalization/Internationalization.h"
 #include "NiagaraConstants.h"
-
+#if WITH_EDITOR
+#include "DerivedDataCacheInterface.h"
+#endif
 #define LOCTEXT_NAMESPACE "UNiagaraSpriteRendererProperties"
 
+#if ENABLE_COOK_STATS
+class NiagaraCutoutCookStats
+{
+public:
+	static FCookStats::FDDCResourceUsageStats UsageStats;
+	static FCookStatsManager::FAutoRegisterCallback RegisterCookStats;
+};
+
+FCookStats::FDDCResourceUsageStats NiagaraCutoutCookStats::UsageStats;
+FCookStatsManager::FAutoRegisterCallback NiagaraCutoutCookStats::RegisterCookStats([](FCookStatsManager::AddStatFuncRef AddStat)
+{
+	UsageStats.LogStats(AddStat, TEXT("NiagaraCutout.Usage"), TEXT(""));
+});
+#endif // ENABLE_COOK_STATS
+
+
 UNiagaraSpriteRendererProperties::UNiagaraSpriteRendererProperties()
 	: Alignment(ENiagaraSpriteAlignment::Unaligned)
 	, FacingMode(ENiagaraSpriteFacingMode::FaceCamera)
@@ -20,6 +38,10 @@ UNiagaraSpriteRendererProperties::UNiagaraSpriteRendererProperties()
 	, MinFacingCameraBlendDistance(0.0f)
 	, MaxFacingCameraBlendDistance(0.0f)
 	, SyncId(0)
+#if WITH_EDITORONLY_DATA
+	, BoundingMode(BVC_EightVertices)
+	, AlphaThreshold(0.1f)
+#endif // WITH_EDITORONLY_DATA
 {
 }
 
@@ -33,6 +55,22 @@ void UNiagaraSpriteRendererProperties::GetUsedMaterials(TArray<UMaterialInterfac
 	OutMaterials.Add(Material);
 }
 
+void UNiagaraSpriteRendererProperties::PostLoad()
+{
+	Super::PostLoad();
+
+#if WITH_EDITORONLY_DATA
+	if (!FPlatformProperties::RequiresCookedData())
+	{
+		if (CutoutTexture)
+		{	// Here we don't call UpdateCutoutTexture() to avoid issue with the material postload.
+			CutoutTexture->ConditionalPostLoad();
+		}
+		CacheDerivedData();
+	}
+#endif // WITH_EDITORONLY_DATA
+}
+
 void UNiagaraSpriteRendererProperties::PostInitProperties()
 {
 	Super::PostInitProperties();
@@ -43,6 +81,17 @@ void UNiagaraSpriteRendererProperties::PostInitProperties()
 	}
 }
 
+void UNiagaraSpriteRendererProperties::Serialize(FStructuredArchive::FRecord Record)
+{
+	Super::Serialize(Record);
+
+	FArchive& UnderlyingArchive = Record.GetUnderlyingArchive();
+	if (UnderlyingArchive.IsCooking() || (FPlatformProperties::RequiresCookedData() && UnderlyingArchive.IsLoading()))
+	{
+		DerivedData.Serialize(Record.EnterField(FIELD_NAME_TEXT("DerivedData")));
+	}
+}
+
 /** The bindings depend on variables that are created during the NiagaraModule startup. However, the CDO's are build prior to this being initialized, so we defer setting these values until later.*/
 void UNiagaraSpriteRendererProperties::InitCDOPropertiesAfterModuleStartup()
 {
@@ -76,17 +125,38 @@ void UNiagaraSpriteRendererProperties::InitBindings()
 	}
 }
 
-
 #if WITH_EDITORONLY_DATA
+
 void UNiagaraSpriteRendererProperties::PostEditChangeProperty(struct FPropertyChangedEvent& PropertyChangedEvent) 
 {
-	Super::PostEditChangeProperty(PropertyChangedEvent);
+	SubImageSize.X = FMath::Max<float>(SubImageSize.X, 1.f);
+	SubImageSize.Y = FMath::Max<float>(SubImageSize.Y, 1.f);
+
+	// DerivedData.BoundingGeometry in case we cleared the CutoutTexture
+	if (bUseMaterialCutoutTexture || CutoutTexture || DerivedData.BoundingGeometry.Num())
+	{
+		const bool bUpdateCutoutDDC = 
+			PropertyChangedEvent.GetPropertyName() == TEXT("bUseMaterialCutoutTexture") ||
+			PropertyChangedEvent.GetPropertyName() == TEXT("CutoutTexture") ||
+			PropertyChangedEvent.GetPropertyName() == TEXT("BoundingMode") ||
+			PropertyChangedEvent.GetPropertyName() == TEXT("OpacitySourceMode") ||
+			PropertyChangedEvent.GetPropertyName() == TEXT("AlphaThreshold") ||
+			(bUseMaterialCutoutTexture && PropertyChangedEvent.GetPropertyName() == TEXT("Material"));
+
+		if (bUpdateCutoutDDC)
+		{
+			UpdateCutoutTexture();
+			CacheDerivedData();
+		}
+	}
+
 	if (PropertyChangedEvent.GetPropertyName() != TEXT("SyncId"))
 	{
 		SyncId++;
 	}
-}
 
+	Super::PostEditChangeProperty(PropertyChangedEvent);
+}
 
 const TArray<FNiagaraVariable>& UNiagaraSpriteRendererProperties::GetRequiredAttributes()
 {
@@ -99,7 +169,6 @@ const TArray<FNiagaraVariable>& UNiagaraSpriteRendererProperties::GetRequiredAtt
 	return Attrs;
 }
 
-
 const TArray<FNiagaraVariable>& UNiagaraSpriteRendererProperties::GetOptionalAttributes()
 {
 	static TArray<FNiagaraVariable> Attrs;
@@ -127,7 +196,6 @@ const TArray<FNiagaraVariable>& UNiagaraSpriteRendererProperties::GetOptionalAtt
 	return Attrs;
 }
 
-
 bool UNiagaraSpriteRendererProperties::IsMaterialValidForRenderer(UMaterial* InMaterial, FText& InvalidMessage)
 {
 	if (InMaterial->bUsedWithNiagaraSprites == false)
@@ -145,10 +213,97 @@ void UNiagaraSpriteRendererProperties::FixMaterial(UMaterial* InMaterial)
 	InMaterial->ForceRecompileForRendering();
 }
 
+void UNiagaraSpriteRendererProperties::UpdateCutoutTexture()
+{
+	if (bUseMaterialCutoutTexture)
+	{
+		CutoutTexture = nullptr;
+		if (Material)
+		{
+			// Try to find an opacity mask texture to default to, if not try to find an opacity texture
+			TArray<UTexture*> OpacityMaskTextures;
+			Material->GetTexturesInPropertyChain(MP_OpacityMask, OpacityMaskTextures, nullptr, nullptr);
+			if (OpacityMaskTextures.Num())
+			{
+				CutoutTexture = (UTexture2D*)OpacityMaskTextures[0];
+			}
+			else
+			{
+				TArray<UTexture*> OpacityTextures;
+				Material->GetTexturesInPropertyChain(MP_Opacity, OpacityTextures, nullptr, nullptr);
+				if (OpacityTextures.Num())
+				{
+					CutoutTexture = (UTexture2D*)OpacityTextures[0];
+				}
+			}
+		}
+	}
+}
+
+void UNiagaraSpriteRendererProperties::CacheDerivedData()
+{
+	if (CutoutTexture)
+	{
+		const FString KeyString = FSubUVDerivedData::GetDDCKeyString(CutoutTexture->Source.GetId(), (int32)SubImageSize.X, (int32)SubImageSize.Y, (int32)BoundingMode, AlphaThreshold, (int32)OpacitySourceMode);
+		TArray<uint8> Data;
+
+		COOK_STAT(auto Timer = NiagaraCutoutCookStats::UsageStats.TimeSyncWork());
+		if (GetDerivedDataCacheRef().GetSynchronous(*KeyString, Data))
+		{
+			COOK_STAT(Timer.AddHit(Data.Num()));
+			DerivedData.BoundingGeometry.Empty(Data.Num() / sizeof(FVector2D));
+			DerivedData.BoundingGeometry.AddUninitialized(Data.Num() / sizeof(FVector2D));
+			FPlatformMemory::Memcpy(DerivedData.BoundingGeometry.GetData(), Data.GetData(), Data.Num() * Data.GetTypeSize());
+		}
+		else
+		{
+			DerivedData.Build(CutoutTexture, (int32)SubImageSize.X, (int32)SubImageSize.Y, BoundingMode, AlphaThreshold, OpacitySourceMode);
+
+			Data.Empty(DerivedData.BoundingGeometry.Num() * sizeof(FVector2D));
+			Data.AddUninitialized(DerivedData.BoundingGeometry.Num() * sizeof(FVector2D));
+			FPlatformMemory::Memcpy(Data.GetData(), DerivedData.BoundingGeometry.GetData(), DerivedData.BoundingGeometry.Num() * DerivedData.BoundingGeometry.GetTypeSize());
+			GetDerivedDataCacheRef().Put(*KeyString, Data);
+			COOK_STAT(Timer.AddMiss(Data.Num()));
+		}
+	}
+	else
+	{
+		DerivedData.BoundingGeometry.Empty();
+	}
+}
+
 #endif // WITH_EDITORONLY_DATA
 
-#undef LOCTEXT_NAMESPACE
 
+int32 UNiagaraSpriteRendererProperties::GetNumCutoutVertexPerSubimage() const
+{
+	if (DerivedData.BoundingGeometry.Num())
+	{
+		const int32 NumSubImages = FMath::Max<int32>(1, (int32)SubImageSize.X * (int32)SubImageSize.Y);
+		const int32 NumCutoutVertexPerSubImage = DerivedData.BoundingGeometry.Num() / NumSubImages;
+
+		// Based on BVC_FourVertices || BVC_EightVertices
+		ensure(NumCutoutVertexPerSubImage == 4 || NumCutoutVertexPerSubImage == 8);
 
+		return NumCutoutVertexPerSubImage;
+	}
+	else
+	{
+		return 0;
+	}
+}
 
+uint32 UNiagaraSpriteRendererProperties::GetNumIndicesPerInstance()
+{
+	// This is a based on cutout vertices making a triangle strip.
+	if (GetNumCutoutVertexPerSubimage() == 8)
+	{
+		return 18;
+	}
+	else
+	{
+		return 6;
+	}
+}
 
+#undef LOCTEXT_NAMESPACE
diff --git a/Engine/Plugins/FX/Niagara/Source/Niagara/Public/NiagaraRenderer.h b/Engine/Plugins/FX/Niagara/Source/Niagara/Public/NiagaraRenderer.h
--- a/Engine/Plugins/FX/Niagara/Source/Niagara/Public/NiagaraRenderer.h
+++ b/Engine/Plugins/FX/Niagara/Source/Niagara/Public/NiagaraRenderer.h
@@ -19,6 +19,7 @@ NiagaraRenderer.h: Base class for Niagara render modules
 #include "RenderingThread.h"
 #include "SceneView.h"
 #include "NiagaraComponent.h"
+#include "NiagaraCutoutVertexBuffer.h"
 
 class FNiagaraDataSet;
 
@@ -182,9 +183,6 @@ protected:
 #endif
 };
 
-
-
-
 /**
 * NiagaraRendererSprites renders an FNiagaraEmitterInstance as sprite particles
 */
@@ -253,6 +251,9 @@ private:
 	int32 CustomSortingOffset;
 
 	int32 LastSyncId;
+
+	FNiagaraCutoutVertexBuffer CutoutVertexBuffer;
+	int32 NumCutoutVertexPerSubImage = 0;
 };
 
 
diff --git a/Engine/Plugins/FX/Niagara/Source/Niagara/Public/NiagaraSpriteRendererProperties.h b/Engine/Plugins/FX/Niagara/Source/Niagara/Public/NiagaraSpriteRendererProperties.h
--- a/Engine/Plugins/FX/Niagara/Source/Niagara/Public/NiagaraSpriteRendererProperties.h
+++ b/Engine/Plugins/FX/Niagara/Source/Niagara/Public/NiagaraSpriteRendererProperties.h
@@ -6,9 +6,9 @@
 #include "UObject/ObjectMacros.h"
 #include "NiagaraCommon.h"
 #include "NiagaraRendererProperties.h"
+#include "Particles/SubUVAnimation.h"
 #include "NiagaraSpriteRendererProperties.generated.h"
 
-
 /** This enum decides how a sprite particle will orient its "up" axis. Must keep these in sync with NiagaraSpriteVertexFactory.ush*/
 UENUM()
 enum class ENiagaraSpriteAlignment : uint8
@@ -47,7 +47,9 @@ public:
 
 	UNiagaraSpriteRendererProperties();
 
+	virtual void PostLoad() override;
 	virtual void PostInitProperties() override;
+	virtual void Serialize(FStructuredArchive::FRecord Record) override;
 
 	static void InitCDOPropertiesAfterModuleStartup();
 
@@ -63,8 +65,8 @@ public:
 	virtual const TArray<FNiagaraVariable>& GetOptionalAttributes() override;
 #endif // WITH_EDITORONLY_DATA
 
-	// TODO: once we support cutouts, need to change this
-	virtual uint32 GetNumIndicesPerInstance() { return 6; }
+	int32 GetNumCutoutVertexPerSubimage() const;
+	virtual uint32 GetNumIndicesPerInstance();
 
 	/** The material used to render the particle. Note that it must have the Use with Niagara Sprites flag checked.*/
 	UPROPERTY(EditAnywhere, Category = "Sprite Rendering")
@@ -186,7 +188,43 @@ public:
 	int32 SyncId;
 
 	void InitBindings();
-};
 
+#if WITH_EDITORONLY_DATA
+
+	/** Use the cutout texture from the material opacity mask, or if none exist, from the material opacity.	*/
+	UPROPERTY(EditAnywhere, BlueprintReadOnly, Category = "Cutout")
+	bool bUseMaterialCutoutTexture;
+
+	/** Texture to generate bounding geometry from.	*/
+	UPROPERTY(EditAnywhere, Category="Cutout", meta = (EditCondition = "!bUseMaterialCutoutTexture"))
+	UTexture2D* CutoutTexture;
+	
+	/**
+	* More bounding vertices results in reduced overdraw, but adds more triangle overhead.
+	* The eight vertex mode is best used when the SubUV texture has a lot of space to cut out that is not captured by the four vertex version,
+	* and when the particles using the texture will be few and large.
+	*/
+	UPROPERTY(EditAnywhere, Category= "Cutout")
+	TEnumAsByte<enum ESubUVBoundingVertexCount> BoundingMode;
+
+	UPROPERTY(EditAnywhere, Category="Cutout")
+	TEnumAsByte<enum EOpacitySourceMode> OpacitySourceMode;
+	
+	/**
+	* Alpha channel values larger than the threshold are considered occupied and will be contained in the bounding geometry.
+	* Raising this threshold slightly can reduce overdraw in particles using this animation asset.
+	*/
+	UPROPERTY(EditAnywhere, Category="Cutout", meta=(UIMin = "0", UIMax = "1"))
+	float AlphaThreshold;
+
+	void UpdateCutoutTexture();
+	void CacheDerivedData();
+#endif
 
+	const TArray<FVector2D>& GetCutoutData() const { return DerivedData.BoundingGeometry; }
 
+private:
+
+	/** Derived data for this asset, generated off of SubUVTexture. */
+	FSubUVDerivedData DerivedData;
+};
diff --git a/Engine/Plugins/FX/Niagara/Source/NiagaraVertexFactories/Private/NiagaraCutoutVertexBuffer.cpp b/Engine/Plugins/FX/Niagara/Source/NiagaraVertexFactories/Private/NiagaraCutoutVertexBuffer.cpp
new file mode 100644
--- /dev/null
+++ b/Engine/Plugins/FX/Niagara/Source/NiagaraVertexFactories/Private/NiagaraCutoutVertexBuffer.cpp
@@ -0,0 +1,42 @@
+// Copyright 1998-2019 Epic Games, Inc. All Rights Reserved.
+
+/*=============================================================================
+	NiagaraCutoutVertexBuffer.h: Niagara cutout uv buffer implementation.
+=============================================================================*/
+
+#include "NiagaraCutoutVertexBuffer.h"
+
+FNiagaraCutoutVertexBuffer::FNiagaraCutoutVertexBuffer(int32 ZeroInitCount)
+{
+	if (ZeroInitCount > 0)
+	{
+		Data.AddZeroed(ZeroInitCount);
+	}
+}
+
+void FNiagaraCutoutVertexBuffer::InitRHI()
+{
+	if (Data.Num())
+	{
+		// create a static vertex buffer
+		FRHIResourceCreateInfo CreateInfo;
+		void* BufferData = nullptr;
+
+		const int32 DataSize = Data.Num() * sizeof(FVector2D);
+		VertexBufferRHI = RHICreateAndLockVertexBuffer(DataSize, BUF_Static | BUF_ShaderResource, CreateInfo, BufferData);
+		FMemory::Memcpy(BufferData, Data.GetData(), DataSize);
+		RHIUnlockVertexBuffer(VertexBufferRHI);
+		VertexBufferSRV = RHICreateShaderResourceView(VertexBufferRHI, sizeof(FVector2D), PF_G32R32F);
+
+		Data.Empty();
+	}
+}
+
+void FNiagaraCutoutVertexBuffer::ReleaseRHI()
+{
+	VertexBufferSRV.SafeRelease();
+	FVertexBuffer::ReleaseRHI();
+}
+
+TGlobalResource<FNiagaraCutoutVertexBuffer> GFNiagaraNullCutoutVertexBuffer(4);
+
diff --git a/Engine/Plugins/FX/Niagara/Source/NiagaraVertexFactories/Private/NiagaraSpriteVertexFactory.cpp b/Engine/Plugins/FX/Niagara/Source/NiagaraVertexFactories/Private/NiagaraSpriteVertexFactory.cpp
--- a/Engine/Plugins/FX/Niagara/Source/NiagaraVertexFactories/Private/NiagaraSpriteVertexFactory.cpp
+++ b/Engine/Plugins/FX/Niagara/Source/NiagaraVertexFactories/Private/NiagaraSpriteVertexFactory.cpp
@@ -5,6 +5,7 @@
 =============================================================================*/
 
 #include "NiagaraSpriteVertexFactory.h"
+#include "NiagaraCutoutVertexBuffer.h"
 #include "ParticleHelper.h"
 #include "ParticleResources.h"
 #include "ShaderParameterUtils.h"
@@ -15,8 +16,6 @@ IMPLEMENT_GLOBAL_SHADER_PARAMETER_STRUCT(FNiagaraSpriteVFLooseParameters, "Niaga
 
 TGlobalResource<FNullDynamicParameterVertexBuffer> GNullNiagaraDynamicParameterVertexBuffer;
 
-TGlobalResource<FNiagaraNullSubUVCutoutVertexBuffer> GFNiagaraNullSubUVCutoutVertexBuffer;
-
 /**
  * Shader parameters for the particle vertex factory.
  */
@@ -92,7 +91,7 @@ public:
 		ShaderBindings.Add(Shader->GetUniformBufferParameter<FNiagaraSpriteVFLooseParameters>(), SpriteVF->LooseParameterUniformBuffer);
 		
 		ShaderBindings.Add(NumCutoutVerticesPerFrame, SpriteVF->GetNumCutoutVerticesPerFrame());
-		FShaderResourceViewRHIParamRef NullSRV = GFNiagaraNullSubUVCutoutVertexBuffer.VertexBufferSRV;
+		FShaderResourceViewRHIParamRef NullSRV = GFNiagaraNullCutoutVertexBuffer.VertexBufferSRV;
 		ShaderBindings.Add(CutoutGeometry, SpriteVF->GetCutoutGeometrySRV() ? SpriteVF->GetCutoutGeometrySRV() : NullSRV);
 
 		ShaderBindings.Add(ParticleAlignmentMode, SpriteVF->GetAlignmentMode());
@@ -155,12 +154,7 @@ public:
 	FVertexDeclarationRHIRef VertexDeclarationRHI;
 
 	// Constructor.
-	FNiagaraSpriteVertexDeclaration(bool bInInstanced, int32 InNumVertsInInstanceBuffer) :
-		bInstanced(bInInstanced),
-		NumVertsInInstanceBuffer(InNumVertsInInstanceBuffer)
-	{
-
-	}
+	FNiagaraSpriteVertexDeclaration() {}
 
 	// Destructor.
 	virtual ~FNiagaraSpriteVertexDeclaration() {}
@@ -190,31 +184,10 @@ public:
 	{
 		VertexDeclarationRHI.SafeRelease();
 	}
-
-private:
-
-	bool bInstanced;
-	int32 NumVertsInInstanceBuffer;
 };
 
 /** The simple element vertex declaration. */
-static TGlobalResource<FNiagaraSpriteVertexDeclaration> GParticleSpriteVertexDeclarationInstanced(true, 4);
-static TGlobalResource<FNiagaraSpriteVertexDeclaration> GParticleSpriteEightVertexDeclarationInstanced(true, 8);
-static TGlobalResource<FNiagaraSpriteVertexDeclaration> GParticleSpriteVertexDeclarationNonInstanced(false, 4);
-static TGlobalResource<FNiagaraSpriteVertexDeclaration> GParticleSpriteEightVertexDeclarationNonInstanced(false, 8);
-
-inline TGlobalResource<FNiagaraSpriteVertexDeclaration>& GetNiagaraSpriteVertexDeclaration(bool SupportsInstancing, int32 NumVertsInInstanceBuffer)
-{
-	check(NumVertsInInstanceBuffer == 4 || NumVertsInInstanceBuffer == 8);
-	if (SupportsInstancing)
-	{
-		return NumVertsInInstanceBuffer == 4 ? GParticleSpriteVertexDeclarationInstanced : GParticleSpriteEightVertexDeclarationInstanced;
-	}
-	else
-	{
-		return NumVertsInInstanceBuffer == 4 ? GParticleSpriteVertexDeclarationNonInstanced : GParticleSpriteEightVertexDeclarationNonInstanced;
-	}
-}
+static TGlobalResource<FNiagaraSpriteVertexDeclaration> GParticleSpriteVertexDeclaration;
 
 bool FNiagaraSpriteVertexFactory::ShouldCompilePermutation(EShaderPlatform Platform, const class FMaterial* Material, const class FShaderType* ShaderType)
 {
@@ -238,7 +211,7 @@ void FNiagaraSpriteVertexFactory::ModifyCompilationEnvironment(const FVertexFact
 void FNiagaraSpriteVertexFactory::InitRHI()
 {
 	InitStreams();
-	SetDeclaration(GetNiagaraSpriteVertexDeclaration(GRHISupportsInstancing, NumVertsInInstanceBuffer).VertexDeclarationRHI);
+	SetDeclaration(GParticleSpriteVertexDeclaration.VertexDeclarationRHI);
 }
 
 void FNiagaraSpriteVertexFactory::InitStreams()
@@ -249,7 +222,7 @@ void FNiagaraSpriteVertexFactory::InitStreams()
 	if(bInstanced) 
 	{
 		FVertexStream* TexCoordStream = new(Streams) FVertexStream;
-		TexCoordStream->VertexBuffer = &GParticleTexCoordVertexBuffer;
+		TexCoordStream->VertexBuffer = VertexBufferOverride ? VertexBufferOverride : &GParticleTexCoordVertexBuffer;
 		TexCoordStream->Stride = sizeof(FVector2D);
 		TexCoordStream->Offset = 0;
 	}
diff --git a/Engine/Plugins/FX/Niagara/Source/NiagaraVertexFactories/Public/NiagaraCutoutVertexBuffer.h b/Engine/Plugins/FX/Niagara/Source/NiagaraVertexFactories/Public/NiagaraCutoutVertexBuffer.h
new file mode 100644
--- /dev/null
+++ b/Engine/Plugins/FX/Niagara/Source/NiagaraVertexFactories/Public/NiagaraCutoutVertexBuffer.h
@@ -0,0 +1,32 @@
+// Copyright 1998-2019 Epic Games, Inc. All Rights Reserved.
+
+/*==============================================================================
+NiagaraCutoutVertexBuffer.h: Niagara cutout uv buffer.
+==============================================================================*/
+
+#pragma once
+
+#include "CoreMinimal.h"
+#include "RenderResource.h"
+
+/**
+ * Vertex buffer to hold cutout UVs. 
+ */
+class NIAGARAVERTEXFACTORIES_API FNiagaraCutoutVertexBuffer : public FVertexBuffer
+{
+public:
+
+	FNiagaraCutoutVertexBuffer(int32 ZeroInitCount = 0);
+
+	/**
+	 * Initialize the RHI for this rendering resource
+	 */
+	virtual void InitRHI() override;
+	virtual void ReleaseRHI() override;
+	FShaderResourceViewRHIRef VertexBufferSRV;
+
+	/** Data to initialize the buffer. */
+	TArray<FVector2D> Data;
+};
+
+extern TGlobalResource<FNiagaraCutoutVertexBuffer> GFNiagaraNullCutoutVertexBuffer;
diff --git a/Engine/Plugins/FX/Niagara/Source/NiagaraVertexFactories/Public/NiagaraSpriteVertexFactory.h b/Engine/Plugins/FX/Niagara/Source/NiagaraVertexFactories/Public/NiagaraSpriteVertexFactory.h
--- a/Engine/Plugins/FX/Niagara/Source/NiagaraVertexFactories/Public/NiagaraSpriteVertexFactory.h
+++ b/Engine/Plugins/FX/Niagara/Source/NiagaraVertexFactories/Public/NiagaraSpriteVertexFactory.h
@@ -113,7 +113,6 @@ public:
 	FNiagaraSpriteVertexFactory(ENiagaraVertexFactoryType InType, ERHIFeatureLevel::Type InFeatureLevel )
 		: FNiagaraVertexFactoryBase(InType, InFeatureLevel),
 		LooseParameterUniformBuffer(nullptr),
-		NumVertsInInstanceBuffer(0),
 		NumCutoutVerticesPerFrame(0),
 		CutoutGeometrySRV(nullptr),
 		AlignmentMode(0),
@@ -126,7 +125,6 @@ public:
 	FNiagaraSpriteVertexFactory()
 		: FNiagaraVertexFactoryBase(NVFT_MAX, ERHIFeatureLevel::Num),
 		LooseParameterUniformBuffer(nullptr),
-		NumVertsInInstanceBuffer(0),
 		NumCutoutVerticesPerFrame(0),
 		CutoutGeometrySRV(nullptr),
 		AlignmentMode(0),
@@ -153,11 +151,6 @@ public:
 	static void ModifyCompilationEnvironment(const FVertexFactoryType* Type, EShaderPlatform Platform, const FMaterial* Material, FShaderCompilerEnvironment& OutEnvironment);
 	
 	void SetTexCoordBuffer(const FVertexBuffer* InTexCoordBuffer);
-
-	inline void SetNumVertsInInstanceBuffer(int32 InNumVertsInInstanceBuffer)
-	{
-		NumVertsInInstanceBuffer = InNumVertsInInstanceBuffer;
-	}
 	
 	/**
 	 * Set the uniform buffer for this vertex factory.
@@ -242,6 +235,11 @@ public:
 		return AlignmentMode;
 	}
 
+	void SetVertexBufferOverride(const FVertexBuffer* InVertexBufferOverride)
+	{
+		VertexBufferOverride = InVertexBufferOverride;
+	}
+
 	/**
 	 * Construct shader parameters for this type of vertex factory.
 	 */
@@ -255,13 +253,13 @@ protected:
 
 private:
 
-	int32 NumVertsInInstanceBuffer;
+	const FVertexBuffer* VertexBufferOverride = nullptr;
 
 	/** Uniform buffer with sprite parameters. */
-	FUniformBufferRHIParamRef SpriteUniformBuffer;
+	FUniformBufferRHIRef SpriteUniformBuffer;
 
 	int32 NumCutoutVerticesPerFrame;
-	FShaderResourceViewRHIParamRef CutoutGeometrySRV;
+	FShaderResourceViewRHIRef CutoutGeometrySRV;
 	uint32 AlignmentMode;
 	uint32 FacingMode;
 	
diff --git a/Engine/Shaders/Private/NiagaraSpriteVertexFactory.ush b/Engine/Shaders/Private/NiagaraSpriteVertexFactory.ush
--- a/Engine/Shaders/Private/NiagaraSpriteVertexFactory.ush
+++ b/Engine/Shaders/Private/NiagaraSpriteVertexFactory.ush
@@ -544,7 +544,8 @@ void ComputeBillboardUVs(FVertexFactoryInput Input, FVertexFactoryIntermediates
 	if (NiagaraSpriteVFLooseParameters.NumCutoutVerticesPerFrame > 0)
 	{
 		// Avoid uint divide which is extremely slow on GCN
-		uint CutoutVertexIndex = (uint)fmod(Input.VertexId, NiagaraSpriteVFLooseParameters.NumCutoutVerticesPerFrame);
+		const uint VertexIndexMask = NiagaraSpriteVFLooseParameters.NumCutoutVerticesPerFrame == 8 ? 7 : 3; // BVC_FourVertices || BVC_EightVertices
+		uint CutoutVertexIndex = Input.VertexId & VertexIndexMask;
 		float NumFrames = NiagaraSpriteVF.SubImageSize.x * NiagaraSpriteVF.SubImageSize.y;
 		uint SubImageIndexInt = (uint)fmod(Intermediates.SizeRotSubImage.w, NumFrames);
 
diff --git a/Engine/Source/Runtime/Engine/Classes/Particles/SubUVAnimation.h b/Engine/Source/Runtime/Engine/Classes/Particles/SubUVAnimation.h
--- a/Engine/Source/Runtime/Engine/Classes/Particles/SubUVAnimation.h
+++ b/Engine/Source/Runtime/Engine/Classes/Particles/SubUVAnimation.h
@@ -52,7 +52,7 @@ enum EOpacitySourceMode
 	OSM_BlueChannel
 };
 
-class FSubUVDerivedData
+class ENGINE_API FSubUVDerivedData
 {
 public:
 	TArray<FVector2D> BoundingGeometry;
diff --git a/Engine/Source/Runtime/Engine/Public/ParticleResources.h b/Engine/Source/Runtime/Engine/Public/ParticleResources.h
--- a/Engine/Source/Runtime/Engine/Public/ParticleResources.h
+++ b/Engine/Source/Runtime/Engine/Public/ParticleResources.h
@@ -54,7 +54,7 @@ public:
 
 /** Global particle index buffer. */
 ENGINE_API extern TGlobalResource<FParticleIndexBuffer> GParticleIndexBuffer;
-extern TGlobalResource<FSixTriangleParticleIndexBuffer> GSixTriangleParticleIndexBuffer;
+ENGINE_API extern TGlobalResource<FSixTriangleParticleIndexBuffer> GSixTriangleParticleIndexBuffer;
 
 typedef FShaderResourceViewRHIParamRef FParticleShaderParamRef;
 typedef FVertexBufferRHIParamRef FParticleBufferParamRef;
