Fixing issue with reactivation of Niagara systems.

diff --git a/Engine/Plugins/FX/Niagara/Source/Niagara/Classes/NiagaraEmitterInstance.h b/Engine/Plugins/FX/Niagara/Source/Niagara/Classes/NiagaraEmitterInstance.h
--- a/Engine/Plugins/FX/Niagara/Source/Niagara/Classes/NiagaraEmitterInstance.h
+++ b/Engine/Plugins/FX/Niagara/Source/Niagara/Classes/NiagaraEmitterInstance.h
@@ -33,7 +33,7 @@ public:
 
 	void Init(int32 InEmitterIdx, FName SystemInstanceName);
 
-	void ResetSimulation();
+	void ResetSimulation(bool bKillExisting=true);
 
 	/** Called after all emitters in an System have been initialized, allows emitters to access information from one another. */
 	void PostInitSimulation();
diff --git a/Engine/Plugins/FX/Niagara/Source/Niagara/Private/NiagaraEmitterInstance.cpp b/Engine/Plugins/FX/Niagara/Source/Niagara/Private/NiagaraEmitterInstance.cpp
--- a/Engine/Plugins/FX/Niagara/Source/Niagara/Private/NiagaraEmitterInstance.cpp
+++ b/Engine/Plugins/FX/Niagara/Source/Niagara/Private/NiagaraEmitterInstance.cpp
@@ -359,31 +359,34 @@ void FNiagaraEmitterInstance::Init(int32 InEmitterIdx, FName InSystemInstanceNam
 	FNiagaraUtilities::CollectScriptDataInterfaceParameters(*CachedEmitter, Scripts, ScriptDefinedDataInterfaceParameters);
 }
 
-void FNiagaraEmitterInstance::ResetSimulation()
+void FNiagaraEmitterInstance::ResetSimulation(bool bKillExisting)
 {
-	bResetPending = true;
 	Age = 0;
 	Loops = 0;
 	TickCount = 0;
-	TotalSpawnedParticles = 0;
 	CachedBounds.Init();
+	SetExecutionState(ENiagaraExecutionState::Active);
 
-	ParticleDataSet->ResetBuffers();
-	for (FNiagaraDataSet* SpawnScriptEventDataSet : SpawnScriptEventDataSets)
-	{
-		SpawnScriptEventDataSet->ResetBuffers();
-	}
-	for (FNiagaraDataSet* UpdateScriptEventDataSet : UpdateScriptEventDataSets)
-	{
-		UpdateScriptEventDataSet->ResetBuffers();
-	}
-	
-	if (CachedEmitter->SimTarget == ENiagaraSimTarget::GPUComputeSim && GPUExecContext != nullptr)
+	if (bKillExisting)
 	{
-		GPUExecContext->Reset();
-	}
+		bResetPending = true;
+		TotalSpawnedParticles = 0;
 
-	SetExecutionState(ENiagaraExecutionState::Active);
+		ParticleDataSet->ResetBuffers();
+		for (FNiagaraDataSet* SpawnScriptEventDataSet : SpawnScriptEventDataSets)
+		{
+			SpawnScriptEventDataSet->ResetBuffers();
+		}
+		for (FNiagaraDataSet* UpdateScriptEventDataSet : UpdateScriptEventDataSets)
+		{
+			UpdateScriptEventDataSet->ResetBuffers();
+		}
+
+		if (CachedEmitter->SimTarget == ENiagaraSimTarget::GPUComputeSim && GPUExecContext != nullptr)
+		{
+			GPUExecContext->Reset();
+		}
+	}
 }
 
 void FNiagaraEmitterInstance::CheckForErrors()
diff --git a/Engine/Plugins/FX/Niagara/Source/Niagara/Private/NiagaraSystemInstance.cpp b/Engine/Plugins/FX/Niagara/Source/Niagara/Private/NiagaraSystemInstance.cpp
--- a/Engine/Plugins/FX/Niagara/Source/Niagara/Private/NiagaraSystemInstance.cpp
+++ b/Engine/Plugins/FX/Niagara/Source/Niagara/Private/NiagaraSystemInstance.cpp
@@ -518,12 +518,9 @@ void FNiagaraSystemInstance::ResetInternal(bool bResetSimulations)
 		return;
 	}
 
-	if (bResetSimulations)
+	for (TSharedRef<FNiagaraEmitterInstance> Simulation : Emitters)
 	{
-		for (TSharedRef<FNiagaraEmitterInstance> Simulation : Emitters)
-		{
-			Simulation->ResetSimulation();
-		}
+		Simulation->ResetSimulation(bResetSimulations);
 	}
 
 #if WITH_EDITOR
