Fixed link error with FNiagaraNullSubUVCutoutVertexBuffer.

diff --git a/Engine/Plugins/FX/Niagara/Source/Niagara/Private/NiagaraRendererSprites.cpp b/Engine/Plugins/FX/Niagara/Source/Niagara/Private/NiagaraRendererSprites.cpp
--- a/Engine/Plugins/FX/Niagara/Source/Niagara/Private/NiagaraRendererSprites.cpp
+++ b/Engine/Plugins/FX/Niagara/Source/Niagara/Private/NiagaraRendererSprites.cpp
@@ -7,6 +7,7 @@
 #include "NiagaraStats.h"
 #include "NiagaraEmitterInstanceBatcher.h"
 #include "NiagaraSortingGPU.h"
+#include "NiagaraCutoutVertexBuffer.h"
 
 DECLARE_CYCLE_STAT(TEXT("Generate Sprite Vertex Data [GT]"), STAT_NiagaraGenSpriteVertexData, STATGROUP_Niagara);
 DECLARE_CYCLE_STAT(TEXT("Render Sprites [RT]"), STAT_NiagaraRenderSprites, STATGROUP_Niagara);
@@ -352,7 +353,7 @@ void NiagaraRendererSprites::GetDynamicMeshElements(const TArray<const FSceneVie
 
 				FNiagaraSpriteVFLooseParameters VFLooseParams;
 				VFLooseParams.NumCutoutVerticesPerFrame = CollectorResources.VertexFactory.GetNumCutoutVerticesPerFrame();
-				VFLooseParams.CutoutGeometry = CollectorResources.VertexFactory.GetCutoutGeometrySRV() ? CollectorResources.VertexFactory.GetCutoutGeometrySRV() : GFNiagaraNullSubUVCutoutVertexBuffer.VertexBufferSRV.GetReference();
+				VFLooseParams.CutoutGeometry = CollectorResources.VertexFactory.GetCutoutGeometrySRV() ? CollectorResources.VertexFactory.GetCutoutGeometrySRV() : GFNiagaraNullCutoutVertexBuffer.VertexBufferSRV.GetReference();
 				VFLooseParams.NiagaraParticleDataFloat = CollectorResources.VertexFactory.GetParticleDataFloatSRV();
 				VFLooseParams.NiagaraFloatDataOffset = CollectorResources.VertexFactory.GetFloatDataOffset();
 				VFLooseParams.NiagaraFloatDataStride = CollectorResources.VertexFactory.GetFloatDataStride();
diff --git a/Engine/Plugins/FX/Niagara/Source/NiagaraVertexFactories/Public/NiagaraCutoutVertexBuffer.h b/Engine/Plugins/FX/Niagara/Source/NiagaraVertexFactories/Public/NiagaraCutoutVertexBuffer.h
--- a/Engine/Plugins/FX/Niagara/Source/NiagaraVertexFactories/Public/NiagaraCutoutVertexBuffer.h
+++ b/Engine/Plugins/FX/Niagara/Source/NiagaraVertexFactories/Public/NiagaraCutoutVertexBuffer.h
@@ -29,4 +29,4 @@ public:
 	TArray<FVector2D> Data;
 };
 
-extern TGlobalResource<FNiagaraCutoutVertexBuffer> GFNiagaraNullCutoutVertexBuffer;
+extern NIAGARAVERTEXFACTORIES_API TGlobalResource<FNiagaraCutoutVertexBuffer> GFNiagaraNullCutoutVertexBuffer;
diff --git a/Engine/Plugins/FX/Niagara/Source/NiagaraVertexFactories/Public/NiagaraSpriteVertexFactory.h b/Engine/Plugins/FX/Niagara/Source/NiagaraVertexFactories/Public/NiagaraSpriteVertexFactory.h
--- a/Engine/Plugins/FX/Niagara/Source/NiagaraVertexFactories/Public/NiagaraSpriteVertexFactory.h
+++ b/Engine/Plugins/FX/Niagara/Source/NiagaraVertexFactories/Public/NiagaraSpriteVertexFactory.h
@@ -71,35 +71,6 @@ END_GLOBAL_SHADER_PARAMETER_STRUCT()
 
 typedef TUniformBufferRef<FNiagaraSpriteVFLooseParameters> FNiagaraSpriteVFLooseParametersRef;
 
-class FNiagaraNullSubUVCutoutVertexBuffer : public FVertexBuffer
-{
-public:
-	/**
-	 * Initialize the RHI for this rendering resource
-	 */
-	virtual void InitRHI() override
-	{
-		// create a static vertex buffer
-		FRHIResourceCreateInfo CreateInfo;
-		void* BufferData = nullptr;
-		VertexBufferRHI = RHICreateAndLockVertexBuffer(sizeof(FVector2D) * 4, BUF_Static | BUF_ShaderResource, CreateInfo, BufferData);
-		FMemory::Memzero(BufferData, sizeof(FVector2D) * 4);
-		RHIUnlockVertexBuffer(VertexBufferRHI);
-
-		VertexBufferSRV = RHICreateShaderResourceView(VertexBufferRHI, sizeof(FVector2D), PF_G32R32F);
-	}
-
-	virtual void ReleaseRHI() override
-	{
-		VertexBufferSRV.SafeRelease();
-		FVertexBuffer::ReleaseRHI();
-	}
-
-	FShaderResourceViewRHIRef VertexBufferSRV;
-};
-
-extern NIAGARAVERTEXFACTORIES_API TGlobalResource<FNiagaraNullSubUVCutoutVertexBuffer> GFNiagaraNullSubUVCutoutVertexBuffer;
-
 /**
  * Vertex factory for rendering particle sprites.
  */
