Restore lost definition
#rnx

diff --git a/Engine/Plugins/FX/Niagara/Source/NiagaraVertexFactories/Private/NiagaraSpriteVertexFactory.cpp b/Engine/Plugins/FX/Niagara/Source/NiagaraVertexFactories/Private/NiagaraSpriteVertexFactory.cpp
--- a/Engine/Plugins/FX/Niagara/Source/NiagaraVertexFactories/Private/NiagaraSpriteVertexFactory.cpp
+++ b/Engine/Plugins/FX/Niagara/Source/NiagaraVertexFactories/Private/NiagaraSpriteVertexFactory.cpp
@@ -16,6 +16,8 @@ IMPLEMENT_GLOBAL_SHADER_PARAMETER_STRUCT(FNiagaraSpriteVFLooseParameters, "Niaga
 
 TGlobalResource<FNullDynamicParameterVertexBuffer> GNullNiagaraDynamicParameterVertexBuffer;
 
+TGlobalResource<FNiagaraNullSubUVCutoutVertexBuffer> GFNiagaraNullSubUVCutoutVertexBuffer;
+
 /**
  * Shader parameters for the particle vertex factory.
  */
