- Fix for crash when pin names are localized in Niagara
- Localized FriendlyName was being used inside ReallocatePins

#test Engine test suite was run along with steps in Jira

diff --git a/Engine/Plugins/FX/Niagara/Source/NiagaraEditor/Private/NiagaraNodeOp.cpp b/Engine/Plugins/FX/Niagara/Source/NiagaraEditor/Private/NiagaraNodeOp.cpp
--- a/Engine/Plugins/FX/Niagara/Source/NiagaraEditor/Private/NiagaraNodeOp.cpp
+++ b/Engine/Plugins/FX/Niagara/Source/NiagaraEditor/Private/NiagaraNodeOp.cpp
@@ -24,7 +24,7 @@ void UNiagaraNodeOp::AllocateDefaultPins()
 	for (int32 SrcIndex = 0; SrcIndex < OpInfo->Inputs.Num(); ++SrcIndex)
 	{
 		const FNiagaraOpInOutInfo& InOutInfo = OpInfo->Inputs[SrcIndex];
-		UEdGraphPin* Pin = CreatePin(EGPD_Input, Schema->TypeDefinitionToPinType(InOutInfo.DataType), *InOutInfo.FriendlyName.ToString());
+		UEdGraphPin* Pin = CreatePin(EGPD_Input, Schema->TypeDefinitionToPinType(InOutInfo.DataType), *InOutInfo.Name.ToString());
 		check(Pin);
 		Pin->bDefaultValueIsIgnored = false;
 		Pin->bDefaultValueIsReadOnly = false;
@@ -49,7 +49,7 @@ void UNiagaraNodeOp::AllocateDefaultPins()
 	for (int32 OutIdx = 0; OutIdx < OpInfo->Outputs.Num(); ++OutIdx)
 	{
 		const FNiagaraOpInOutInfo& InOutInfo = OpInfo->Outputs[OutIdx];
-		UEdGraphPin* Pin = CreatePin(EGPD_Output, Schema->TypeDefinitionToPinType(InOutInfo.DataType), *InOutInfo.FriendlyName.ToString());
+		UEdGraphPin* Pin = CreatePin(EGPD_Output, Schema->TypeDefinitionToPinType(InOutInfo.DataType), *InOutInfo.Name.ToString());
 		check(Pin);
 		Pin->PinToolTip = InOutInfo.Description.ToString();
 	}
