- Fix for potential dangling component pointer
- Make cached pointer a WeakPtr and recache per tick if neccessary as it could be potentially destroyed or changed

diff --git a/Engine/Plugins/FX/Niagara/Source/Niagara/Classes/NiagaraDataInterfaceSpline.h b/Engine/Plugins/FX/Niagara/Source/Niagara/Classes/NiagaraDataInterfaceSpline.h
--- a/Engine/Plugins/FX/Niagara/Source/Niagara/Classes/NiagaraDataInterfaceSpline.h
+++ b/Engine/Plugins/FX/Niagara/Source/Niagara/Classes/NiagaraDataInterfaceSpline.h
@@ -11,7 +11,7 @@
 struct FNDISpline_InstanceData
 {
 	//Cached ptr to component we sample from. 
-	USplineComponent* Component;
+	TWeakObjectPtr<USplineComponent> Component;
 
 	//Cached ComponentToWorld.
 	FMatrix Transform;
diff --git a/Engine/Plugins/FX/Niagara/Source/Niagara/Private/NiagaraDataInterfaceSpline.cpp b/Engine/Plugins/FX/Niagara/Source/Niagara/Private/NiagaraDataInterfaceSpline.cpp
--- a/Engine/Plugins/FX/Niagara/Source/Niagara/Private/NiagaraDataInterfaceSpline.cpp
+++ b/Engine/Plugins/FX/Niagara/Source/Niagara/Private/NiagaraDataInterfaceSpline.cpp
@@ -329,40 +329,9 @@ bool UNiagaraDataInterfaceSpline::InitPerInstanceData(void* PerInstanceData, FNi
 {
 	FNDISpline_InstanceData* InstData = new (PerInstanceData) FNDISpline_InstanceData();
 
-	InstData->Component = nullptr;
+	InstData->Component.Reset();
 	InstData->Transform = FMatrix::Identity;
 	InstData->TransformInverseTransposed = FMatrix::Identity;
-	if (Source)
-	{
-		USplineComponent* SourceComp = Source->FindComponentByClass<USplineComponent>();
-
-		if (SourceComp)
-		{
-			InstData->Component = SourceComp;
-		}
-	}
-	else
-	{
-		if (UNiagaraComponent* SimComp = SystemInstance->GetComponent())
-		{
-			if (AActor* Owner = SimComp->GetAttachmentRootActor())
-			{
-				USplineComponent* SourceComp = Owner->FindComponentByClass<USplineComponent>();
-
-				if (SourceComp)
-				{
-					InstData->Component = SourceComp;
-				}
-			}
-		}
-	}
-
-	//Re-evaluate source in case it's changed?
-	if (InstData->Component)
-	{
-		InstData->Transform = InstData->Component->GetComponentToWorld().ToMatrixWithScale();
-		InstData->TransformInverseTransposed = InstData->Transform.InverseFast().GetTransposed();
-	}
 
 	return true;
 }
@@ -378,9 +347,30 @@ bool UNiagaraDataInterfaceSpline::PerInstanceTick(void* PerInstanceData, FNiagar
 	check(SystemInstance);
 	FNDISpline_InstanceData* InstData = (FNDISpline_InstanceData*)PerInstanceData;
 
-	if (InstData->Component)
+	USplineComponent* SplineComponent = InstData->Component.Get();
+	if (SplineComponent == nullptr)
+	{
+		if (Source != nullptr)
+		{
+			SplineComponent = Source->FindComponentByClass<USplineComponent>();
+		}
+		else
+		{
+			if (UNiagaraComponent* SimComp = SystemInstance->GetComponent())
+			{
+				if (AActor* Owner = SimComp->GetAttachmentRootActor())
+				{
+					SplineComponent = Owner->FindComponentByClass<USplineComponent>();
+				}
+			}
+		}
+		InstData->Component = SplineComponent;
+	}
+
+	//Re-evaluate source in case it's changed?
+	if (SplineComponent != nullptr)
 	{
-		InstData->Transform = InstData->Component->GetComponentToWorld().ToMatrixWithScale();
+		InstData->Transform = SplineComponent->GetComponentToWorld().ToMatrixWithScale();
 		InstData->TransformInverseTransposed = InstData->Transform.InverseFast().GetTransposed();
 	}
 
@@ -398,14 +388,14 @@ void UNiagaraDataInterfaceSpline::SampleSplinePositionByUnitDistance(FVectorVMCo
 	VectorVM::FExternalFuncRegisterHandler<float> OutPosY(Context);
 	VectorVM::FExternalFuncRegisterHandler<float> OutPosZ(Context);
 
-	if (InstData->Component)
+	if (USplineComponent* SplineComponent = InstData->Component.Get())
 	{
 
 		for (int32 i = 0; i < Context.NumInstances; ++i)
 		{
 			float DistanceUnitDistance = SplineSampleParam.Get();
 
-			FVector Pos = InstData->Component->GetLocationAtDistanceAlongSpline(DistanceUnitDistance * InstData->Component->GetSplineLength(), ESplineCoordinateSpace::Local);
+			FVector Pos = SplineComponent->GetLocationAtDistanceAlongSpline(DistanceUnitDistance * SplineComponent->GetSplineLength(), ESplineCoordinateSpace::Local);
 			TransformHandler.TransformPosition(Pos, InstData->Transform);
 
 			*OutPosX.GetDest() = Pos.X;
@@ -447,14 +437,14 @@ void UNiagaraDataInterfaceSpline::SampleSplineUpVectorByUnitDistance(FVectorVMCo
 	VectorVM::FExternalFuncRegisterHandler<float> OutPosY(Context);
 	VectorVM::FExternalFuncRegisterHandler<float> OutPosZ(Context);
 
-	if (InstData->Component)
+	if (USplineComponent* SplineComponent = InstData->Component.Get())
 	{
 
 		for (int32 i = 0; i < Context.NumInstances; ++i)
 		{
 			float DistanceUnitDistance = SplineSampleParam.Get();
 
-			FVector Pos = InstData->Component->GetUpVectorAtDistanceAlongSpline(DistanceUnitDistance * InstData->Component->GetSplineLength(), ESplineCoordinateSpace::Local);
+			FVector Pos = SplineComponent->GetUpVectorAtDistanceAlongSpline(DistanceUnitDistance * SplineComponent->GetSplineLength(), ESplineCoordinateSpace::Local);
 			TransformHandler.TransformVector(Pos, InstData->Transform);
 
 			*OutPosX.GetDest() = Pos.X;
@@ -497,13 +487,13 @@ void UNiagaraDataInterfaceSpline::SampleSplineRightVectorByUnitDistance(FVectorV
 	VectorVM::FExternalFuncRegisterHandler<float> OutPosZ(Context);
 	
 	
-	if (InstData->Component)
+	if (USplineComponent* SplineComponent = InstData->Component.Get())
 	{
 		for (int32 i = 0; i < Context.NumInstances; ++i)
 		{
 			float DistanceUnitDistance = SplineSampleParam.Get();
 
-			FVector Pos = InstData->Component->GetRightVectorAtDistanceAlongSpline(DistanceUnitDistance * InstData->Component->GetSplineLength(), ESplineCoordinateSpace::Local);
+			FVector Pos = SplineComponent->GetRightVectorAtDistanceAlongSpline(DistanceUnitDistance * SplineComponent->GetSplineLength(), ESplineCoordinateSpace::Local);
 			TransformHandler.TransformVector(Pos, InstData->Transform);
 
 			*OutPosX.GetDest() = Pos.X;
@@ -543,13 +533,13 @@ void UNiagaraDataInterfaceSpline::SampleSplineTangentByUnitDistance(FVectorVMCon
 	VectorVM::FExternalFuncRegisterHandler<float> OutPosY(Context);
 	VectorVM::FExternalFuncRegisterHandler<float> OutPosZ(Context);
 
-	if (InstData->Component)
+	if (USplineComponent* SplineComponent = InstData->Component.Get())
 	{
 		for (int32 i = 0; i < Context.NumInstances; ++i)
 		{
 			float DistanceUnitDistance = SplineSampleParam.Get();
 
-			FVector Pos = InstData->Component->GetTangentAtDistanceAlongSpline(DistanceUnitDistance * InstData->Component->GetSplineLength(), ESplineCoordinateSpace::Local);
+			FVector Pos = SplineComponent->GetTangentAtDistanceAlongSpline(DistanceUnitDistance * SplineComponent->GetSplineLength(), ESplineCoordinateSpace::Local);
 			TransformHandler.TransformVector(Pos, InstData->Transform);
 
 			*OutPosX.GetDest() = Pos.X;
@@ -591,13 +581,13 @@ void UNiagaraDataInterfaceSpline::SampleSplineDirectionByUnitDistance(FVectorVMC
 	VectorVM::FExternalFuncRegisterHandler<float> OutPosY(Context);
 	VectorVM::FExternalFuncRegisterHandler<float> OutPosZ(Context);
 
-	if (InstData->Component)
+	if (USplineComponent* SplineComponent = InstData->Component.Get())
 	{
 		for (int32 i = 0; i < Context.NumInstances; ++i)
 		{
 			float DistanceUnitDistance = SplineSampleParam.Get();
 
-			FVector Pos = InstData->Component->GetDirectionAtDistanceAlongSpline(DistanceUnitDistance * InstData->Component->GetSplineLength(), ESplineCoordinateSpace::Local);
+			FVector Pos = SplineComponent->GetDirectionAtDistanceAlongSpline(DistanceUnitDistance * SplineComponent->GetSplineLength(), ESplineCoordinateSpace::Local);
 			TransformHandler.TransformVector(Pos, InstData->Transform);
 
 			*OutPosX.GetDest() = Pos.X;
@@ -678,11 +668,11 @@ void UNiagaraDataInterfaceSpline::FindClosestUnitDistanceFromPositionWS(FVectorV
 	VectorVM::FUserPtrHandler<FNDISpline_InstanceData> InstData(Context);
 	VectorVM::FExternalFuncRegisterHandler<float> OutUnitDistance(Context);
 
-	if (InstData->Component)
+	if (USplineComponent* SplineComponent = InstData->Component.Get())
 	{
 
-		const int32 NumPoints = InstData->Component->GetSplinePointsPosition().Points.Num();
-		const float FinalKeyTime = InstData->Component->GetSplinePointsPosition().Points[NumPoints - 1].InVal;
+		const int32 NumPoints = SplineComponent->GetSplinePointsPosition().Points.Num();
+		const float FinalKeyTime = SplineComponent->GetSplinePointsPosition().Points[NumPoints - 1].InVal;
 
 		for (int32 i = 0; i < Context.NumInstances; ++i)
 		{
@@ -693,7 +683,7 @@ void UNiagaraDataInterfaceSpline::FindClosestUnitDistanceFromPositionWS(FVectorV
 			FVector Pos(PosX, PosY, PosZ);
 
 			// This first call finds the key time, but this is not in 0..1 range for the spline. 
-			float KeyTime = InstData->Component->FindInputKeyClosestToWorldLocation(Pos);
+			float KeyTime = SplineComponent->FindInputKeyClosestToWorldLocation(Pos);
 			// We need to convert into the range by dividing through by the overall duration of the spline according to the keys.
 			float UnitDistance = KeyTime / FinalKeyTime;
 
