Integrate MessageLog to Niagara for Graph and Stack view. (This is just a stub for the window, proper messaging to come in later CLs.)

diff --git a/Engine/Plugins/FX/Niagara/Source/NiagaraEditor/Private/NiagaraEditorModule.cpp b/Engine/Plugins/FX/Niagara/Source/NiagaraEditor/Private/NiagaraEditorModule.cpp
--- a/Engine/Plugins/FX/Niagara/Source/NiagaraEditor/Private/NiagaraEditorModule.cpp
+++ b/Engine/Plugins/FX/Niagara/Source/NiagaraEditor/Private/NiagaraEditorModule.cpp
@@ -583,6 +583,11 @@ FNiagaraEditorModule& FNiagaraEditorModule::Get()
 {
 	return FModuleManager::LoadModuleChecked<FNiagaraEditorModule>("NiagaraEditor");
 }
+void FNiagaraEditorModule::OnMessageLogTokenClicked(const TSharedRef<class IMessageToken>& Token)
+{
+	return; //@todo(ng) just stubbing this here for now
+}
+
 
 void FNiagaraEditorModule::OnNiagaraSettingsChangedEvent(const FString& PropertyName, const UNiagaraSettings* Settings)
 {
diff --git a/Engine/Plugins/FX/Niagara/Source/NiagaraEditor/Private/Toolkits/NiagaraScriptToolkit.cpp b/Engine/Plugins/FX/Niagara/Source/NiagaraEditor/Private/Toolkits/NiagaraScriptToolkit.cpp
--- a/Engine/Plugins/FX/Niagara/Source/NiagaraEditor/Private/Toolkits/NiagaraScriptToolkit.cpp
+++ b/Engine/Plugins/FX/Niagara/Source/NiagaraEditor/Private/Toolkits/NiagaraScriptToolkit.cpp
@@ -59,6 +59,7 @@ const FName FNiagaraScriptToolkit::NodeGraphTabId(TEXT("NiagaraEditor_NodeGraph"
 const FName FNiagaraScriptToolkit::DetailsTabId(TEXT("NiagaraEditor_Details"));
 const FName FNiagaraScriptToolkit::ParametersTabId(TEXT("NiagaraEditor_Parameters"));
 const FName FNiagaraScriptToolkit::StatsTabId(TEXT("NiagaraEditor_Stats"));
+const FName FNiagaraScriptToolkit::MessageLogTabID(TEXT("NiagaraEditor_MessageLog"));
 
 FNiagaraScriptToolkit::FNiagaraScriptToolkit()
 {
@@ -95,6 +96,10 @@ void FNiagaraScriptToolkit::RegisterTabSpawners(const TSharedRef<class FTabManag
 		.SetDisplayName(LOCTEXT("StatsTab", "Stats"))
 		.SetGroup(WorkspaceMenuCategoryRef)
 		.SetIcon(FSlateIcon(FEditorStyle::GetStyleSetName(), "LevelEditor.Tabs.Details"));
+
+	InTabManager->RegisterTabSpawner(MessageLogTabID, FOnSpawnTab::CreateSP(this, &FNiagaraScriptToolkit::SpawnTabMessageLog))
+		.SetDisplayName(LOCTEXT("NiagaraMessageLogTab", "Niagara Message Log"))
+		.SetGroup(WorkspaceMenuCategoryRef);
 }
 
 void FNiagaraScriptToolkit::UnregisterTabSpawners(const TSharedRef<class FTabManager>& InTabManager)
@@ -149,7 +154,7 @@ void FNiagaraScriptToolkit::Initialize( const EToolkitMode::Type Mode, const TSh
 	StatsListing = MessageLogModule.CreateLogListing("MaterialEditorStats", LogOptions);
 	Stats = MessageLogModule.CreateLogListingWidget(StatsListing.ToSharedRef());
 
-	TSharedRef<FTabManager::FLayout> StandaloneDefaultLayout = FTabManager::NewLayout("Standalone_Niagara_Layout_v7")
+	TSharedRef<FTabManager::FLayout> StandaloneDefaultLayout = FTabManager::NewLayout("Standalone_Niagara_Layout_v8")
 	->AddArea
 	(
 		FTabManager::NewPrimaryArea()->SetOrientation(Orient_Vertical)
@@ -190,9 +195,19 @@ void FNiagaraScriptToolkit::Initialize( const EToolkitMode::Type Mode, const TSh
 			)
 			->Split
 			(
+				FTabManager::NewSplitter()->SetOrientation(Orient_Vertical)->SetSizeCoefficient(0.8f)
+				->Split
+				(
 				FTabManager::NewStack()
 				->SetSizeCoefficient(0.8f)
 				->AddTab(NodeGraphTabId, ETabState::OpenedTab)
+				)
+				->Split
+				(
+				FTabManager::NewStack()
+				->SetSizeCoefficient(0.2f)
+				->AddTab(MessageLogTabID, ETabState::OpenedTab)
+				)
 			)
 		)
 	);
@@ -374,6 +389,44 @@ TSharedRef<SDockTab> FNiagaraScriptToolkit::SpawnTabStats(const FSpawnTabArgs& A
 	return SpawnedTab;
 }
 
+TSharedRef<SDockTab> FNiagaraScriptToolkit::SpawnTabMessageLog(const FSpawnTabArgs& Args)
+{
+	check(Args.GetTabId().TabType == MessageLogTabID);
+
+	FMessageLogModule& MessageLogModule = FModuleManager::LoadModuleChecked<FMessageLogModule>("MessageLog");
+	const FName LogName = GetNiagaraScriptMessageLogName(EditedNiagaraScript);
+
+	auto CreateMessageLogListing = [&MessageLogModule, &LogName]()->TSharedRef<IMessageLogListing>
+	{
+		// Reuse any existing log, or create a new one (that is not held onto by the message log system)
+		if (MessageLogModule.IsRegisteredLogListing(LogName))
+		{
+			return MessageLogModule.GetLogListing(LogName);
+		}
+		else
+		{
+			FMessageLogInitializationOptions LogInitOptions;
+			LogInitOptions.bShowInLogWindow = false;
+			return  MessageLogModule.CreateLogListing(LogName, LogInitOptions);
+		}
+	};
+
+	TSharedRef<IMessageLogListing> MessageLogListing = CreateMessageLogListing();
+	//@todo(ng) debug message (to be removed when we forward real messages)
+	MessageLogListing->AddMessage(FTokenizedMessage::Create(EMessageSeverity::Info, FText::FromString(FString::Printf(TEXT("This is a debug message for %s"), *GetNiagaraScriptMessageLogName(EditedNiagaraScript).ToString()))));
+
+	MessageLogListing->OnMessageTokenClicked().AddStatic(&FNiagaraEditorModule::OnMessageLogTokenClicked);
+	TSharedRef<SWidget> LogListModule = MessageLogModule.CreateLogListingWidget(MessageLogListing);
+
+	TSharedRef<SDockTab> SpawnedTab =
+		SNew(SDockTab)
+		[
+			LogListModule
+		];
+
+	return SpawnedTab;
+}
+
 void FNiagaraScriptToolkit::SetupCommands()
 {
 	GetToolkitCommands()->MapAction(
@@ -388,6 +441,13 @@ void FNiagaraScriptToolkit::SetupCommands()
 		FExecuteAction::CreateRaw(this, &FNiagaraScriptToolkit::RefreshNodes));
 }
 
+const FName FNiagaraScriptToolkit::GetNiagaraScriptMessageLogName(UNiagaraScript* InScript) const
+{
+	checkf(InScript, TEXT("Tried to get MessageLog name for NiagaraScript but InScript was null!"));
+	FName LogListingName = *FString::Printf(TEXT("%s_%s_MessageLog"), *InScript->GetBaseChangeID().ToString(), *InScript->GetName());
+	return LogListingName;
+}
+
 void FNiagaraScriptToolkit::ExtendToolbar()
 {
 	struct Local
diff --git a/Engine/Plugins/FX/Niagara/Source/NiagaraEditor/Private/Toolkits/NiagaraScriptToolkit.h b/Engine/Plugins/FX/Niagara/Source/NiagaraEditor/Private/Toolkits/NiagaraScriptToolkit.h
--- a/Engine/Plugins/FX/Niagara/Source/NiagaraEditor/Private/Toolkits/NiagaraScriptToolkit.h
+++ b/Engine/Plugins/FX/Niagara/Source/NiagaraEditor/Private/Toolkits/NiagaraScriptToolkit.h
@@ -77,9 +77,12 @@ private:
 	TSharedRef<SDockTab> SpawnTabScriptParameters(const FSpawnTabArgs& Args);
 	TSharedRef<SDockTab> SpawnTabStats(const FSpawnTabArgs& Args);
 
+	TSharedRef<SDockTab> SpawnTabMessageLog(const FSpawnTabArgs& Args);
+
 	/** Sets up commands for the toolkit toolbar. */
 	void SetupCommands();
 
+	const FName GetNiagaraScriptMessageLogName(UNiagaraScript* InScript) const;
 	FSlateIcon GetCompileStatusImage() const;
 	FText GetCompileStatusTooltip() const;
 
@@ -118,6 +121,7 @@ private:
 	static const FName DetailsTabId;
 	static const FName ParametersTabId;
 	static const FName StatsTabId;
+	static const FName MessageLogTabID;
 	/** Stats log, with the log listing that it reflects */
 	TSharedPtr<class SWidget> Stats;
 	TSharedPtr<class IMessageLogListing> StatsListing;
diff --git a/Engine/Plugins/FX/Niagara/Source/NiagaraEditor/Private/Toolkits/NiagaraSystemToolkit.cpp b/Engine/Plugins/FX/Niagara/Source/NiagaraEditor/Private/Toolkits/NiagaraSystemToolkit.cpp
--- a/Engine/Plugins/FX/Niagara/Source/NiagaraEditor/Private/Toolkits/NiagaraSystemToolkit.cpp
+++ b/Engine/Plugins/FX/Niagara/Source/NiagaraEditor/Private/Toolkits/NiagaraSystemToolkit.cpp
@@ -48,6 +48,8 @@
 #include "IAssetRegistry.h"
 #include "HAL/PlatformFilemanager.h"
 #include "Misc/FileHelper.h"
+#include "MessageLog/Public/MessageLogModule.h"
+#include "MessageLog/Public/IMessageLogListing.h"
 
 #define LOCTEXT_NAMESPACE "NiagaraSystemEditor"
 
@@ -64,6 +66,7 @@ const FName FNiagaraSystemToolkit::SelectedEmitterGraphTabID(TEXT("NiagaraSystem
 const FName FNiagaraSystemToolkit::DebugSpreadsheetTabID(TEXT("NiagaraSystemEditor_DebugAttributeSpreadsheet"));
 const FName FNiagaraSystemToolkit::PreviewSettingsTabId(TEXT("NiagaraSystemEditor_PreviewSettings"));
 const FName FNiagaraSystemToolkit::GeneratedCodeTabID(TEXT("NiagaraSystemEditor_GeneratedCode"));
+const FName FNiagaraSystemToolkit::MessageLogTabID(TEXT("NiagaraSystemEditor_MessageLog"));
 
 static int32 GbLogNiagaraSystemChanges = 0;
 static FAutoConsoleVariableRef CVarSuppressNiagaraSystems(
@@ -134,6 +137,10 @@ void FNiagaraSystemToolkit::RegisterTabSpawners(const TSharedRef<class FTabManag
 	InTabManager->RegisterTabSpawner(GeneratedCodeTabID, FOnSpawnTab::CreateSP(this, &FNiagaraSystemToolkit::SpawnTab_GeneratedCode))
 		.SetDisplayName(LOCTEXT("GeneratedCode", "Generated Code"))
 		.SetGroup(WorkspaceMenuCategory.ToSharedRef());
+
+	InTabManager->RegisterTabSpawner(MessageLogTabID, FOnSpawnTab::CreateSP(this, &FNiagaraSystemToolkit::SpawnTab_MessageLog))
+		.SetDisplayName(LOCTEXT("NiagaraMessageLog", "Niagara Log"))
+		.SetGroup(WorkspaceMenuCategory.ToSharedRef());
 }
 
 void FNiagaraSystemToolkit::UnregisterTabSpawners(const TSharedRef<class FTabManager>& InTabManager)
@@ -288,7 +295,7 @@ void FNiagaraSystemToolkit::InitializeInternal(const EToolkitMode::Type Mode, co
 	const float InTime = -0.02f;
 	const float OutTime = 3.2f;
 
-	TSharedRef<FTabManager::FLayout> StandaloneDefaultLayout = FTabManager::NewLayout("Standalone_Niagara_System_Layout_v17")
+	TSharedRef<FTabManager::FLayout> StandaloneDefaultLayout = FTabManager::NewLayout("Standalone_Niagara_System_Layout_v18")
 		->AddArea
 		(
 			FTabManager::NewPrimaryArea()->SetOrientation(Orient_Vertical)
@@ -328,6 +335,7 @@ void FNiagaraSystemToolkit::InitializeInternal(const EToolkitMode::Type Mode, co
 						FTabManager::NewStack()
 						->SetSizeCoefficient(0.25f)
 						->AddTab(CurveEditorTabID, ETabState::OpenedTab)
+						->AddTab(MessageLogTabID, ETabState::OpenedTab)
 						->AddTab(SequencerTabID, ETabState::OpenedTab)
 					)
 				)
@@ -601,7 +609,6 @@ TSharedRef<SDockTab> FNiagaraSystemToolkit::SpawnTab_DebugSpreadsheet(const FSpa
 	return SpawnedTab;
 }
 
-
 TSharedRef<SDockTab> FNiagaraSystemToolkit::SpawnTab_GeneratedCode(const FSpawnTabArgs& Args)
 {
 	check(Args.GetTabId().TabType == GeneratedCodeTabID);
@@ -615,6 +622,44 @@ TSharedRef<SDockTab> FNiagaraSystemToolkit::SpawnTab_GeneratedCode(const FSpawnT
 	return SpawnedTab;
 }
 
+TSharedRef<SDockTab> FNiagaraSystemToolkit::SpawnTab_MessageLog(const FSpawnTabArgs& Args)
+{
+	check(Args.GetTabId().TabType == MessageLogTabID);
+
+	FMessageLogModule& MessageLogModule = FModuleManager::LoadModuleChecked<FMessageLogModule>("MessageLog");
+	const FName LogName = GetNiagaraSystemMessageLogName(System);
+
+	auto CreateMessageLogListing = [&MessageLogModule, &LogName]()->TSharedRef<IMessageLogListing>
+	{
+		// Reuse any existing log, or create a new one (that is not held onto by the message log system)
+		if (MessageLogModule.IsRegisteredLogListing(LogName))
+		{
+			return MessageLogModule.GetLogListing(LogName);
+		}
+		else
+		{
+			FMessageLogInitializationOptions LogInitOptions;
+			LogInitOptions.bShowInLogWindow = false;
+			return  MessageLogModule.CreateLogListing(LogName, LogInitOptions);
+		}
+	};
+
+	TSharedRef<IMessageLogListing> MessageLogListing = CreateMessageLogListing();
+	//@todo(ng) debug message (to be removed when we forward real messages)
+	MessageLogListing->AddMessage(FTokenizedMessage::Create(EMessageSeverity::Info, FText::FromString(FString::Printf(TEXT("This is a debug message for %s"), *GetNiagaraSystemMessageLogName(System).ToString()))));
+
+	MessageLogListing->OnMessageTokenClicked().AddStatic(&FNiagaraEditorModule::OnMessageLogTokenClicked);
+	TSharedRef<SWidget> LogListModule = MessageLogModule.CreateLogListingWidget(MessageLogListing);
+
+	TSharedRef<SDockTab> SpawnedTab =
+		SNew(SDockTab)
+		[
+			LogListModule
+		];
+
+	return SpawnedTab;
+}
+
 void FNiagaraSystemToolkit::SetupCommands()
 {
 	GetToolkitCommands()->MapAction(
@@ -812,6 +857,13 @@ TSharedRef<SWidget> FNiagaraSystemToolkit::GenerateBoundsMenuContent(TSharedRef<
 	return MenuBuilder.MakeWidget();
 }
 
+const FName FNiagaraSystemToolkit::GetNiagaraSystemMessageLogName(UNiagaraSystem* InSystem) const
+{
+	checkf(InSystem, TEXT("Tried to get MessageLog name for NiagaraSystem but InSystem was null!"));
+	FName LogListingName = *FString::Printf(TEXT("%s_%s_MessageLog"), *FString::FromInt(InSystem->GetUniqueID()), *InSystem->GetName());
+	return LogListingName;
+}
+
 void FNiagaraSystemToolkit::GetSequencerAddMenuContent(FMenuBuilder& MenuBuilder, TSharedRef<ISequencer> Sequencer)
 {
 	MenuBuilder.AddSubMenu(
diff --git a/Engine/Plugins/FX/Niagara/Source/NiagaraEditor/Private/Toolkits/NiagaraSystemToolkit.h b/Engine/Plugins/FX/Niagara/Source/NiagaraEditor/Private/Toolkits/NiagaraSystemToolkit.h
--- a/Engine/Plugins/FX/Niagara/Source/NiagaraEditor/Private/Toolkits/NiagaraSystemToolkit.h
+++ b/Engine/Plugins/FX/Niagara/Source/NiagaraEditor/Private/Toolkits/NiagaraSystemToolkit.h
@@ -98,6 +98,7 @@ private:
 	TSharedRef<SDockTab> SpawnTab_DebugSpreadsheet(const FSpawnTabArgs& Args);
 	TSharedRef<SDockTab> SpawnTab_PreviewSettings(const FSpawnTabArgs& Args);
 	TSharedRef<SDockTab> SpawnTab_GeneratedCode(const FSpawnTabArgs& Args);
+	TSharedRef<SDockTab> SpawnTab_MessageLog(const FSpawnTabArgs& Args);
 
 	/** Builds the toolbar widget */
 	void ExtendToolbar();	
@@ -122,6 +123,7 @@ private:
 
 private:
 	TSharedRef<SWidget> GenerateBoundsMenuContent(TSharedRef<FUICommandList> InCommandList);
+	const FName GetNiagaraSystemMessageLogName(UNiagaraSystem* InSystem) const;
 	void OnSaveThumbnailImage();
 	void OnThumbnailCaptured(UTexture2D* Thumbnail);
 
@@ -165,4 +167,5 @@ private:
 	static const FName DebugSpreadsheetTabID;
 	static const FName PreviewSettingsTabId;
 	static const FName GeneratedCodeTabID;
+	static const FName MessageLogTabID;
 };
diff --git a/Engine/Plugins/FX/Niagara/Source/NiagaraEditor/Public/NiagaraEditorModule.h b/Engine/Plugins/FX/Niagara/Source/NiagaraEditor/Public/NiagaraEditorModule.h
--- a/Engine/Plugins/FX/Niagara/Source/NiagaraEditor/Public/NiagaraEditorModule.h
+++ b/Engine/Plugins/FX/Niagara/Source/NiagaraEditor/Public/NiagaraEditorModule.h
@@ -42,6 +42,9 @@ public:
 	/** Get the instance of this module. */
 	NIAGARAEDITOR_API static FNiagaraEditorModule& Get();
 
+	/** Callback  when a token is clicked on in the niagara message log */ //@todo(ng) consider moving this to FNiagaraMessageManager in the future?
+	static void OnMessageLogTokenClicked(const TSharedRef<class IMessageToken>& Token);
+
 	/** Compile the specified script. */
 	virtual TSharedPtr<FNiagaraVMExecutableData> CompileScript(const FNiagaraCompileRequestDataBase* InCompileRequest, const FNiagaraCompileOptions& InCompileOptions);
 
