- Integrate CL 5742723 from Dev-Niagara to fix SetExecutionState crash

diff --git a/Engine/Plugins/FX/Niagara/Source/Niagara/Classes/NiagaraDataSet.h b/Engine/Plugins/FX/Niagara/Source/Niagara/Classes/NiagaraDataSet.h
--- a/Engine/Plugins/FX/Niagara/Source/Niagara/Classes/NiagaraDataSet.h
+++ b/Engine/Plugins/FX/Niagara/Source/Niagara/Classes/NiagaraDataSet.h
@@ -814,6 +814,8 @@ struct FNiagaraDataSetAccessor : public FNiagaraDataSetAccessorBase
 
 	FORCEINLINE void Set(int32 Index, const T& InValue)
 	{
+		check((uint32)Index < DataBuffer->GetNumInstances());
+
 		uint8* ValuePtr = (uint8*)&InValue;
 
 		for (uint32 CompIdx = 0; CompIdx < VarLayout->GetNumFloatComponents(); ++CompIdx)
@@ -887,6 +889,8 @@ struct FNiagaraDataSetAccessor<FNiagaraBool> : public FNiagaraDataSetAccessorBas
 
 	FORCEINLINE void Set(int32 Index, const FNiagaraBool& InValue)
 	{
+		check(Base != nullptr);
+		check((uint32)Index < DataBuffer->GetNumInstances());
 		Base[Index] = InValue.GetRawValue();
 	}
 
@@ -956,6 +960,7 @@ struct FNiagaraDataSetAccessor<int32> : public FNiagaraDataSetAccessorBase
 	FORCEINLINE void Set(int32 Index, const int32& InValue)
 	{
 		check(Base != nullptr);
+		check((uint32)Index < DataBuffer->GetNumInstances());
 		Base[Index] = InValue;
 	}
 
@@ -1019,6 +1024,8 @@ struct FNiagaraDataSetAccessor<float> : public FNiagaraDataSetAccessorBase
 
 	FORCEINLINE void Set(int32 Index, const float& InValue)
 	{
+		check(Base != nullptr);
+		check((uint32)Index < DataBuffer->GetNumInstances());
 		Base[Index] = InValue;
 	}
 	FORCEINLINE bool BaseIsValid() const { return Base != nullptr; }
@@ -1083,6 +1090,8 @@ struct FNiagaraDataSetAccessor<FVector2D> : public FNiagaraDataSetAccessorBase
 
 	FORCEINLINE void Set(int32 Index, const FVector2D& InValue)
 	{
+		check(XBase != nullptr);
+		check((uint32)Index < DataBuffer->GetNumInstances());
 		XBase[Index] = InValue.X;
 		YBase[Index] = InValue.Y;
 	}	
@@ -1154,6 +1163,9 @@ struct FNiagaraDataSetAccessor<FVector> : public FNiagaraDataSetAccessorBase
 
 	FORCEINLINE void Set(int32 Index, const FVector& InValue)
 	{
+		check(XBase != nullptr);
+		check((uint32)Index < DataBuffer->GetNumInstances());
+
 		XBase[Index] = InValue.X;
 		YBase[Index] = InValue.Y;
 		ZBase[Index] = InValue.Z;
@@ -1229,6 +1241,9 @@ struct FNiagaraDataSetAccessor<FVector4> : public FNiagaraDataSetAccessorBase
 
 	FORCEINLINE void Set(int32 Index, const FVector4& InValue)
 	{
+		check(XBase != nullptr);
+		check((uint32)Index < DataBuffer->GetNumInstances());
+
 		XBase[Index] = InValue.X;
 		YBase[Index] = InValue.Y;
 		ZBase[Index] = InValue.Z;
@@ -1307,6 +1322,9 @@ struct FNiagaraDataSetAccessor<FQuat> : public FNiagaraDataSetAccessorBase
 
 	FORCEINLINE void Set(int32 Index, const FQuat& InValue)
 	{
+		check(XBase != nullptr);
+		check((uint32)Index < DataBuffer->GetNumInstances());
+
 		XBase[Index] = InValue.X;
 		YBase[Index] = InValue.Y;
 		ZBase[Index] = InValue.Z;
@@ -1384,6 +1402,9 @@ struct FNiagaraDataSetAccessor<FLinearColor> : public FNiagaraDataSetAccessorBas
 
 	FORCEINLINE void Set(int32 Index, const FLinearColor& InValue)
 	{
+		check(RBase != nullptr);
+		check((uint32)Index < DataBuffer->GetNumInstances());
+
 		RBase[Index] = InValue.R;
 		GBase[Index] = InValue.G;
 		BBase[Index] = InValue.B;
@@ -1461,6 +1482,9 @@ struct FNiagaraDataSetAccessor<FNiagaraSpawnInfo> : public FNiagaraDataSetAccess
 
 	FORCEINLINE void Set(int32 Index, const FNiagaraSpawnInfo& InValue)
 	{
+		check(CountBase != nullptr);
+		check((uint32)Index < DataBuffer->GetNumInstances());
+
 		CountBase[Index] = InValue.Count;
 		InterpStartDtBase[Index] = InValue.InterpStartDt;
 		IntervalDtBase[Index] = InValue.IntervalDt;
@@ -1533,6 +1557,9 @@ struct FNiagaraDataSetAccessor<FNiagaraID> : public FNiagaraDataSetAccessorBase
 
 	FORCEINLINE void Set(int32 Index, const FNiagaraID& InValue)
 	{
+		check(IndexBase != nullptr);
+		check((uint32)Index < DataBuffer->GetNumInstances());
+
 		IndexBase[Index] = InValue.Index;
 		TagBase[Index] = InValue.AcquireTag;
 	}
@@ -1645,6 +1672,7 @@ struct FNiagaraDataSetVariableIterator
 
 	void Set()
 	{
+		check(CurrIdx < DataBuffer.GetNumInstances());
 		for (int32 VarIdx = 0; VarIdx < Variables.Num(); ++VarIdx)
 		{
 			FNiagaraVariable* Var = Variables[VarIdx];
diff --git a/Engine/Plugins/FX/Niagara/Source/Niagara/Private/NiagaraSystemSimulation.cpp b/Engine/Plugins/FX/Niagara/Source/Niagara/Private/NiagaraSystemSimulation.cpp
--- a/Engine/Plugins/FX/Niagara/Source/Niagara/Private/NiagaraSystemSimulation.cpp
+++ b/Engine/Plugins/FX/Niagara/Source/Niagara/Private/NiagaraSystemSimulation.cpp
@@ -383,6 +383,13 @@ bool FNiagaraSystemSimulation::Tick(float DeltaSeconds)
 			UE_LOG(LogNiagara, Log, TEXT("=========================================================="));
 		}
 
+		if (SpawnNum)
+		{
+			SCOPE_CYCLE_COUNTER(STAT_NiagaraSystemSim_Spawn);
+			DataSet.Allocate(NewNum, true);
+			DataSet.SetNumInstances(NewNum);
+		}
+
 		TArray<FNiagaraDataSetExecutionInfo, TInlineAllocator<8>> DataSetExecInfos;
 		DataSetExecInfos.SetNum(2);
 		{
@@ -455,10 +462,6 @@ bool FNiagaraSystemSimulation::Tick(float DeltaSeconds)
 		if (SpawnNum)
 		{
 			SCOPE_CYCLE_COUNTER(STAT_NiagaraSystemSim_Spawn);
-			DataSet.Allocate(NewNum, true);
-
-			DataSet.SetNumInstances(NewNum);
-			//SpawnInstanceParameterDataSet.SetNumInstances(NewNum);
 
 			//Run Spawn
 			SpawnExecContext.Tick(SoloSystemInstance);//We can't require a specific instance here as these are for all instances.
