NiagaraCompileEvents now populate Niagara message log.

diff --git a/Engine/Plugins/FX/Niagara/Source/Niagara/Classes/NiagaraEmitter.h b/Engine/Plugins/FX/Niagara/Source/Niagara/Classes/NiagaraEmitter.h
--- a/Engine/Plugins/FX/Niagara/Source/Niagara/Classes/NiagaraEmitter.h
+++ b/Engine/Plugins/FX/Niagara/Source/Niagara/Classes/NiagaraEmitter.h
@@ -272,7 +272,7 @@ public:
 	UPROPERTY(EditAnywhere, AdvancedDisplay, Category = "Emitter", meta = (InlineEditConditionToggle))
 	uint32 bLimitDeltaTime : 1;
 
-	void NIAGARA_API GetScripts(TArray<UNiagaraScript*>& OutScripts, bool bCompilableOnly = true);
+	void NIAGARA_API GetScripts(TArray<UNiagaraScript*>& OutScripts, bool bCompilableOnly = true) const;
 
 	NIAGARA_API UNiagaraScript* GetScript(ENiagaraScriptUsage Usage, FGuid UsageId);
 
diff --git a/Engine/Plugins/FX/Niagara/Source/Niagara/Classes/NiagaraScript.h b/Engine/Plugins/FX/Niagara/Source/Niagara/Classes/NiagaraScript.h
--- a/Engine/Plugins/FX/Niagara/Source/Niagara/Classes/NiagaraScript.h
+++ b/Engine/Plugins/FX/Niagara/Source/Niagara/Classes/NiagaraScript.h
@@ -26,6 +26,50 @@ void ProcessSerializedShaderMaps(UNiagaraScript* Owner, TArray<FNiagaraShaderScr
 
 #define NIAGARA_SCRIPT_COMPILE_LOGGING_MEDIUM
 
+/** Defines the compile event types for translation/compilation.*/
+UENUM()
+enum class FNiagaraCompileEventSeverity : uint8
+{
+	Log = 0,
+	Warning = 1,
+	Error = 2
+};
+
+/** Records necessary information to give UI cues for errors/logs/warnings during compile.*/
+USTRUCT()
+struct FNiagaraCompileEvent
+{
+	GENERATED_USTRUCT_BODY()
+public:
+	FNiagaraCompileEvent()
+	{
+		Severity = FNiagaraCompileEventSeverity::Log;
+		Message = FString();
+		NodeGuid = FGuid();
+		PinGuid = FGuid();
+		StackGuids.Empty();
+	}
+
+	FNiagaraCompileEvent(FNiagaraCompileEventSeverity InSeverity, const FString& InMessage, FGuid InNodeGuid = FGuid(), FGuid InPinGuid = FGuid(), const TArray<FGuid>& InCallstackGuids = TArray<FGuid>())
+		: Severity(InSeverity), Message(InMessage), NodeGuid(InNodeGuid), PinGuid(InPinGuid), StackGuids(InCallstackGuids) {}
+
+	/** Whether or not this is an error, warning, or info*/
+	UPROPERTY()
+	FNiagaraCompileEventSeverity Severity;
+	/* The message itself*/
+	UPROPERTY()
+	FString Message;
+	/** The node guid that generated the compile event*/
+	UPROPERTY()
+	FGuid NodeGuid;
+	/** The pin persistent id that generated the compile event*/
+	UPROPERTY()
+	FGuid PinGuid;
+	/** The compile stack frame of node id's*/
+	UPROPERTY()
+	TArray<FGuid> StackGuids;
+};
+
 /** Defines what will happen to unused attributes when a script is run. */
 UENUM()
 enum class EUnusedAttributeBehaviour : uint8
@@ -277,6 +321,10 @@ public:
 	UPROPERTY()
 	float CompileTime;
 
+	/** Array of all compile events generated last time the script was compiled.*/
+	UPROPERTY()
+	TArray<FNiagaraCompileEvent> LastCompileEvents;
+
 	void SerializeData(FArchive& Ar, bool bDDCData);
 	
 	bool IsValid() const;
@@ -284,7 +332,6 @@ public:
 	void Reset();
 };
 
-
 /** Runtime script for a Niagara system */
 UCLASS(MinimalAPI)
 class UNiagaraScript : public UObject
diff --git a/Engine/Plugins/FX/Niagara/Source/Niagara/Private/NiagaraEmitter.cpp b/Engine/Plugins/FX/Niagara/Source/Niagara/Private/NiagaraEmitter.cpp
--- a/Engine/Plugins/FX/Niagara/Source/Niagara/Private/NiagaraEmitter.cpp
+++ b/Engine/Plugins/FX/Niagara/Source/Niagara/Private/NiagaraEmitter.cpp
@@ -438,7 +438,7 @@ bool UNiagaraEmitter::IsReadyToRun() const
 	return true;
 }
 
-void UNiagaraEmitter::GetScripts(TArray<UNiagaraScript*>& OutScripts, bool bCompilableOnly)
+void UNiagaraEmitter::GetScripts(TArray<UNiagaraScript*>& OutScripts, bool bCompilableOnly) const
 {
 	OutScripts.Add(SpawnScriptProps.Script);
 	OutScripts.Add(UpdateScriptProps.Script);
diff --git a/Engine/Plugins/FX/Niagara/Source/Niagara/Private/NiagaraScript.cpp b/Engine/Plugins/FX/Niagara/Source/Niagara/Private/NiagaraScript.cpp
--- a/Engine/Plugins/FX/Niagara/Source/Niagara/Private/NiagaraScript.cpp
+++ b/Engine/Plugins/FX/Niagara/Source/Niagara/Private/NiagaraScript.cpp
@@ -78,7 +78,8 @@ void FNiagaraVMExecutableData::Reset()
 
 void FNiagaraVMExecutableData::SerializeData(FArchive& Ar, bool bDDCData)
 {
-	FNiagaraVMExecutableData::StaticStruct()->SerializeBin(Ar, this);
+	UScriptStruct* FNiagaraVMExecutableDataType = FNiagaraVMExecutableData::StaticStruct();
+	FNiagaraVMExecutableDataType->SerializeTaggedProperties(Ar, (uint8*)this, FNiagaraVMExecutableDataType, nullptr);
 }
 
 bool FNiagaraVMExecutableDataId::IsValid() const
diff --git a/Engine/Plugins/FX/Niagara/Source/NiagaraEditor/Private/NiagaraCompiler.cpp b/Engine/Plugins/FX/Niagara/Source/NiagaraEditor/Private/NiagaraCompiler.cpp
--- a/Engine/Plugins/FX/Niagara/Source/NiagaraEditor/Private/NiagaraCompiler.cpp
+++ b/Engine/Plugins/FX/Niagara/Source/NiagaraEditor/Private/NiagaraCompiler.cpp
@@ -29,6 +29,7 @@
 #include "Misc/FileHelper.h"
 #include "ShaderCompiler.h"
 #include "NiagaraShader.h"
+#include "NiagaraScript.h" //@todo(ng) FNiagaraCompileEvent dependency (make sure this is used before submitting)
 
 #define LOCTEXT_NAMESPACE "NiagaraCompiler"
 
@@ -546,6 +547,7 @@ TSharedPtr<FNiagaraVMExecutableData> FNiagaraEditorModule::CompileScript(const F
 	TArray<FNiagaraCompileEvent> Messages;
 	if (TranslateResults.CompileEvents.Num() > 0)
 	{
+		Results.Data->LastCompileEvents.Append(TranslateResults.CompileEvents); //@todo(ng) get rid of Messages in the future
 		Messages.Append(TranslateResults.CompileEvents);
 	}
 	if (Results.CompileEvents.Num() > 0)
@@ -556,19 +558,19 @@ TSharedPtr<FNiagaraVMExecutableData> FNiagaraEditorModule::CompileScript(const F
 	FString OutGraphLevelErrorMessages;
 	for (const FNiagaraCompileEvent& Message : Messages)
 	{
-		if (Message.Type == FNiagaraCompileEventType::Log)
+		if (Message.Severity == FNiagaraCompileEventSeverity::Log)
 		{
 		#if defined(NIAGARA_SCRIPT_COMPILE_LOGGING_MEDIUM)
 			UE_LOG(LogNiagaraCompiler, Log, TEXT("%s"), *Message.Message);
 		#endif
 		}
-		else if (Message.Type == FNiagaraCompileEventType::Warning )
+		else if (Message.Severity == FNiagaraCompileEventSeverity::Warning )
 		{
 		#if defined(NIAGARA_SCRIPT_COMPILE_LOGGING_MEDIUM)
 			UE_LOG(LogNiagaraCompiler, Warning, TEXT("%s"), *Message.Message);
 		#endif
 		}
-		else if (Message.Type == FNiagaraCompileEventType::Error)
+		else if (Message.Severity == FNiagaraCompileEventSeverity::Error)
 		{
 		#if defined(NIAGARA_SCRIPT_COMPILE_LOGGING_MEDIUM)
 			UE_LOG(LogNiagaraCompiler, Error, TEXT("%s"), *Message.Message);
@@ -903,14 +905,16 @@ FHlslNiagaraCompiler::FHlslNiagaraCompiler()
 void FHlslNiagaraCompiler::Error(FText ErrorText)
 {
 	FString ErrorString = FString::Printf(TEXT("%s"), *ErrorText.ToString());
-	CompileResults.CompileEvents.Add(FNiagaraCompileEvent(FNiagaraCompileEventType::Error, ErrorString));
+	CompileResults.Data->LastCompileEvents.Add(FNiagaraCompileEvent(FNiagaraCompileEventSeverity::Error, ErrorString)); //@todo(ng) cut out the compileresults.compileevent member since we don't need to separately log into general output log
+	CompileResults.CompileEvents.Add(FNiagaraCompileEvent(FNiagaraCompileEventSeverity::Error, ErrorString));
 	CompileResults.NumErrors++;
 }
 
 void FHlslNiagaraCompiler::Warning(FText WarningText)
 {
 	FString WarnString = FString::Printf(TEXT("%s"), *WarningText.ToString());
-	CompileResults.CompileEvents.Add(FNiagaraCompileEvent(FNiagaraCompileEventType::Warning, WarnString));
+	CompileResults.Data->LastCompileEvents.Add(FNiagaraCompileEvent(FNiagaraCompileEventSeverity::Warning, WarnString));
+	CompileResults.CompileEvents.Add(FNiagaraCompileEvent(FNiagaraCompileEventSeverity::Warning, WarnString));
 	CompileResults.NumWarnings++;
 }
 
diff --git a/Engine/Plugins/FX/Niagara/Source/NiagaraEditor/Private/NiagaraEditorModule.cpp b/Engine/Plugins/FX/Niagara/Source/NiagaraEditor/Private/NiagaraEditorModule.cpp
--- a/Engine/Plugins/FX/Niagara/Source/NiagaraEditor/Private/NiagaraEditorModule.cpp
+++ b/Engine/Plugins/FX/Niagara/Source/NiagaraEditor/Private/NiagaraEditorModule.cpp
@@ -735,4 +735,36 @@ void FNiagaraEditorModule::OnPreGarbageCollection()
 	}
 }
 
+TSharedRef<FNiagaraAssetNameToken> FNiagaraAssetNameToken::Create(const FString& InAssetName, const FText& InMessage)
+{
+	return MakeShareable(new FNiagaraAssetNameToken(InAssetName, InMessage));
+}
+
+FNiagaraAssetNameToken::FNiagaraAssetNameToken(const FString& InAssetName, const FText& InMessage)
+	: AssetName(InAssetName)
+{
+	if (!InMessage.IsEmpty())
+	{
+		CachedText = InMessage;
+	}
+	else
+	{
+		CachedText = FText::FromString(AssetName);
+	}
+
+	MessageTokenActivated = FOnMessageTokenActivated::CreateStatic(&FNiagaraAssetNameToken::FindAndOpenAsset, AssetName);
+}
+
+
+void FNiagaraAssetNameToken::FindAndOpenAsset(const TSharedRef<IMessageToken>& Token, FString InAssetPath)
+{
+	FAssetRegistryModule& AssetRegistryModule = FModuleManager::LoadModuleChecked<FAssetRegistryModule>("AssetRegistry");
+	IAssetRegistry& AssetRegistry = AssetRegistryModule.Get();
+	FAssetData AssetData = AssetRegistry.GetAssetByObjectPath(*InAssetPath);
+	if (AssetData.IsValid())
+	{
+		FAssetEditorManager::Get().OpenEditorForAsset(AssetData.GetAsset(), EToolkitMode::Standalone);
+	}
+}
+
 #undef LOCTEXT_NAMESPACE
diff --git a/Engine/Plugins/FX/Niagara/Source/NiagaraEditor/Private/NiagaraHlslTranslator.cpp b/Engine/Plugins/FX/Niagara/Source/NiagaraEditor/Private/NiagaraHlslTranslator.cpp
--- a/Engine/Plugins/FX/Niagara/Source/NiagaraEditor/Private/NiagaraHlslTranslator.cpp
+++ b/Engine/Plugins/FX/Niagara/Source/NiagaraEditor/Private/NiagaraHlslTranslator.cpp
@@ -5396,8 +5396,8 @@ void FHlslNiagaraTranslator::Error(FText ErrorText, const UNiagaraNode* Node, co
 		NodePinSuffix = TEXT(" - ");
 	}
 
-	FString ErrorString = FString::Printf(TEXT("%s%s%s%sCallstack: %s"), *ErrorText.ToString(), *NodePinPrefix, *NodePinStr, *NodePinSuffix, *GetCallstack());
-	TranslateResults.CompileEvents.Add(FNiagaraCompileEvent(FNiagaraCompileEventType::Error, ErrorString, Node ? Node->NodeGuid : FGuid(), Pin ? Pin->PersistentGuid : FGuid(), GetCallstackGuids()));
+	FString ErrorString = FString::Printf(TEXT("%s%s%s%s"), *ErrorText.ToString(), *NodePinPrefix, *NodePinStr, *NodePinSuffix);
+	TranslateResults.CompileEvents.Add(FNiagaraCompileEvent(FNiagaraCompileEventSeverity::Error, ErrorString, Node ? Node->NodeGuid : FGuid(), Pin ? Pin->PersistentGuid : FGuid(), GetCallstackGuids()));
 	TranslateResults.NumErrors++;
 }
 
@@ -5417,8 +5417,8 @@ void FHlslNiagaraTranslator::Warning(FText WarningText, const UNiagaraNode* Node
 		NodePinSuffix = TEXT(" - ");
 	}
 
-	FString WarnString = FString::Printf(TEXT("%s%s%s%sCallstack: %s"), *WarningText.ToString(), *NodePinPrefix, *NodePinStr, *NodePinSuffix, *GetCallstack());
-	TranslateResults.CompileEvents.Add(FNiagaraCompileEvent(FNiagaraCompileEventType::Warning, WarnString, Node ? Node->NodeGuid : FGuid(), Pin ? Pin->PersistentGuid : FGuid(), GetCallstackGuids()));
+	FString WarnString = FString::Printf(TEXT("%s%s%s%s"), *WarningText.ToString(), *NodePinPrefix, *NodePinStr, *NodePinSuffix);
+	TranslateResults.CompileEvents.Add(FNiagaraCompileEvent(FNiagaraCompileEventSeverity::Warning, WarnString, Node ? Node->NodeGuid : FGuid(), Pin ? Pin->PersistentGuid : FGuid(), GetCallstackGuids()));
 	TranslateResults.NumWarnings++;
 }
 
diff --git a/Engine/Plugins/FX/Niagara/Source/NiagaraEditor/Private/Toolkits/NiagaraScriptToolkit.cpp b/Engine/Plugins/FX/Niagara/Source/NiagaraEditor/Private/Toolkits/NiagaraScriptToolkit.cpp
--- a/Engine/Plugins/FX/Niagara/Source/NiagaraEditor/Private/Toolkits/NiagaraScriptToolkit.cpp
+++ b/Engine/Plugins/FX/Niagara/Source/NiagaraEditor/Private/Toolkits/NiagaraScriptToolkit.cpp
@@ -44,6 +44,8 @@
 #include "Developer/MessageLog/Public/MessageLogModule.h"
 
 #include "AssetTypeActions/AssetTypeActions_NiagaraScript.h"
+#include "NiagaraNodeFunctionCall.h" //@todo(ng) move this to message manager
+#include "NiagaraNodeEmitter.h" //@todo(ng) move this to message manager 
 
 #define LOCTEXT_NAMESPACE "NiagaraScriptToolkit"
 
@@ -71,6 +73,119 @@ FNiagaraScriptToolkit::~FNiagaraScriptToolkit()
 	ScriptViewModel->GetGraphViewModel()->GetGraph()->RemoveOnGraphNeedsRecompileHandler(OnEditedScriptGraphChangedHandle);
 }
 
+//@todo(ng) move to message manager class
+TSharedRef<FTokenizedMessage> FNiagaraScriptToolkit::BuildMessageForCompileEvent(FNiagaraCompileEvent& InCompileEvent)
+{
+	TArray<FGuid>& ContextStackGuids = InCompileEvent.StackGuids;
+
+	UNiagaraScriptSourceBase* FunctionScriptSourceBase = EditedNiagaraScript->GetSource();
+	checkf(FunctionScriptSourceBase->IsA<UNiagaraScriptSource>(), TEXT("Script source for function call node is not assigned or is not of type UNiagaraScriptSource!"))
+		UNiagaraScriptSource* FunctionScriptSource = Cast<UNiagaraScriptSource>(FunctionScriptSourceBase);
+	checkf(FunctionScriptSource, TEXT("Script source base was somehow not a derived type!"));
+	UNiagaraGraph* ThisScriptGraph = FunctionScriptSource->NodeGraph;
+	checkf(ThisScriptGraph, TEXT("Function Script does not have a UNiagaraGraph!"));
+
+	TArray<TSharedRef<IMessageToken>> MessageTokens;
+	TOptional<UNiagaraGraph*> OriginatingGraph = RecursiveBuildMessageTokensFromContextStackAndGetOriginatingGraph(ContextStackGuids, ThisScriptGraph, MessageTokens);
+	if (OriginatingGraph.IsSet())
+	{
+		EMessageSeverity::Type MessageSeverity = EMessageSeverity::Info;
+		switch (InCompileEvent.Severity) {
+		case FNiagaraCompileEventSeverity::Error:
+			MessageSeverity = EMessageSeverity::Error;
+			break;
+		case FNiagaraCompileEventSeverity::Warning:
+			MessageSeverity = EMessageSeverity::Warning;
+			break;
+		case FNiagaraCompileEventSeverity::Log:
+			MessageSeverity = EMessageSeverity::Info;
+			break;
+		default:
+			verifyf(false, TEXT("Compile event severity type not handled!"));
+			MessageSeverity = EMessageSeverity::Info;
+			break;
+		}
+		TSharedRef<FTokenizedMessage> MessageToBuild = FTokenizedMessage::Create(MessageSeverity);
+
+		//Add message from compile event at start of message
+		MessageToBuild->AddToken(FTextToken::Create(FText::FromString(InCompileEvent.Message)));
+		//Now add the context stack of the scripts that were passed through to get to the originating graph
+		for (TSharedRef<IMessageToken> Token : MessageTokens)
+		{
+			MessageToBuild->AddToken(Token);
+		}
+		return MessageToBuild;
+	}
+	
+	// Failed to get originating graph!
+	verifyf(false, TEXT("Failed to find originating graph for compile event!"));
+	return FTokenizedMessage::Create(EMessageSeverity::Warning, LOCTEXT("NiagaraCompileEventFailedToFindScript", "Could not find script(s) for compile event ! Please recompile script to get full compile info."));
+}
+
+//@todo(ng) move to message manager class
+TOptional<UNiagaraGraph*> FNiagaraScriptToolkit::RecursiveBuildMessageTokensFromContextStackAndGetOriginatingGraph(TArray<FGuid>& InContextStackNodeGuids, UNiagaraGraph* InGraphToSearch, TArray<TSharedRef<IMessageToken>>& OutMessageTokensToAdd)
+{
+	verifyf(InGraphToSearch, TEXT("Failed to get a node graph to search!"));
+
+	if (InGraphToSearch && InContextStackNodeGuids.Num() == 0)
+	{
+		//StackGuids arr has been cleared out which means we have walked the entire context and can now returnthe graph we are in.
+		return InGraphToSearch;
+	}
+
+	// Search in the current graph for a node with a GUID that matches a GUID in the list of Function Call and Emitter node GUIDs that define the context stack for a compile event
+	auto FindNodeInGraphWithContextStackGuid = [&InGraphToSearch, &InContextStackNodeGuids]()->TOptional<UEdGraphNode*> {
+		for (UEdGraphNode* GraphNode : InGraphToSearch->Nodes)
+		{
+			for (int i = 0; i < InContextStackNodeGuids.Num(); i++)
+			{
+				if (GraphNode->NodeGuid == InContextStackNodeGuids[i])
+				{
+					InContextStackNodeGuids.RemoveAt(i);
+					return GraphNode;
+				}
+			}
+		}
+		return TOptional<UEdGraphNode*>();
+	};
+
+	TOptional<UEdGraphNode*> ContextNode = FindNodeInGraphWithContextStackGuid();
+	if (ContextNode.IsSet())
+	{
+		// found a node in the current graph that has a GUID in the context list, now get the Niagara Script assigned to this node, add a message token and recursively call into the graph of that script.
+		UNiagaraNodeFunctionCall* FunctionCallNode = Cast<UNiagaraNodeFunctionCall>(ContextNode.GetValue());
+		if (FunctionCallNode)
+		{
+			UNiagaraScript* FunctionCallNodeAssignedScript = FunctionCallNode->FunctionScript;
+			checkf(FunctionCallNodeAssignedScript, TEXT("Script for function call node is not assigned!"));
+			UNiagaraScriptSourceBase* FunctionScriptSourceBase = FunctionCallNodeAssignedScript->GetSource();
+			checkf(FunctionScriptSourceBase->IsA<UNiagaraScriptSource>(), TEXT("Script source for function call node is not assigned or is not of type UNiagaraScriptSource!"));
+			UNiagaraScriptSource* FunctionScriptSource = Cast<UNiagaraScriptSource>(FunctionScriptSourceBase);
+			checkf(FunctionScriptSource, TEXT("Script source base was somehow not a derived type!"));
+			UNiagaraGraph* FunctionScriptGraph = FunctionScriptSource->NodeGraph;
+			checkf(FunctionScriptGraph, TEXT("Function Script does not have a UNiagaraGraph!"));
+
+			TSharedRef<FNiagaraAssetNameToken> VisitedGraphContextToken = FNiagaraAssetNameToken::Create(FunctionCallNodeAssignedScript->GetPathName(), FText::FromString(*FunctionCallNodeAssignedScript->GetName()));
+			OutMessageTokensToAdd.Add(VisitedGraphContextToken);
+			return RecursiveBuildMessageTokensFromContextStackAndGetOriginatingGraph(InContextStackNodeGuids, FunctionScriptGraph, OutMessageTokensToAdd);
+		}
+		UNiagaraNodeEmitter* EmitterNode = Cast<UNiagaraNodeEmitter>(ContextNode.GetValue());
+		if (EmitterNode)
+		{
+			UNiagaraScriptSource* EmitterScriptSource = EmitterNode->GetScriptSource();
+			UNiagaraGraph* EmitterScriptGraph = EmitterScriptSource->NodeGraph;
+			checkf(EmitterScriptGraph, TEXT("Emitter Script Source does not have a UNiagaraGraph!")); //@todo(ng) see about filtering the message log by emitter?
+
+			TSharedRef<FTextToken> SourceEmitterContextToken = FTextToken::Create(FText::FromString(EmitterNode->GetEmitterUniqueName()));
+			OutMessageTokensToAdd.Add(SourceEmitterContextToken);
+			return RecursiveBuildMessageTokensFromContextStackAndGetOriginatingGraph(InContextStackNodeGuids, EmitterScriptGraph, OutMessageTokensToAdd);
+		}
+		checkf(false, TEXT("Matching node is not a function call or emitter node!"));
+	}
+	verifyf(false, TEXT("Failed to walk the entire context stack, is this compile event out of date?"));
+	return TOptional<UNiagaraGraph*>();
+}
+
 void FNiagaraScriptToolkit::RegisterTabSpawners(const TSharedRef<class FTabManager>& InTabManager)
 {
 	WorkspaceMenuCategory = InTabManager->AddLocalWorkspaceMenuCategory(LOCTEXT("WorkspaceMenu_NiagaraEditor", "Niagara"));
@@ -151,8 +266,23 @@ void FNiagaraScriptToolkit::Initialize( const EToolkitMode::Type Mode, const TSh
 	LogOptions.bShowFilters = false;
 	LogOptions.bAllowClear = false;
 	LogOptions.MaxPageCount = 1;
+
+	// Reuse any existing log, or create a new one (that is not held onto by the message log system)
+	auto CreateMessageLogListing = [&MessageLogModule, &LogOptions](const FName& LogName)->TSharedRef<IMessageLogListing> {
+		if (MessageLogModule.IsRegisteredLogListing(LogName))
+		{
+			return MessageLogModule.GetLogListing(LogName);
+		}
+		else
+		{
+			return  MessageLogModule.CreateLogListing(LogName, LogOptions);
+		}
+	};
+
 	StatsListing = MessageLogModule.CreateLogListing("MaterialEditorStats", LogOptions);
 	Stats = MessageLogModule.CreateLogListingWidget(StatsListing.ToSharedRef());
+	NiagaraMessageLogListing = CreateMessageLogListing(GetNiagaraScriptMessageLogName(EditedNiagaraScript));
+	NiagaraMessageLog = MessageLogModule.CreateLogListingWidget(NiagaraMessageLogListing.ToSharedRef());
 
 	TSharedRef<FTabManager::FLayout> StandaloneDefaultLayout = FTabManager::NewLayout("Standalone_Niagara_Layout_v8")
 	->AddArea
@@ -222,6 +352,7 @@ void FNiagaraScriptToolkit::Initialize( const EToolkitMode::Type Mode, const TSh
 	SetupCommands();
 	ExtendToolbar();
 	RegenerateMenusAndToolbars();
+	UpdateMessageLog();
 	UpdateModuleStats();
 	// @todo toolkit world centric editing
 	/*// Setup our tool's layout
@@ -305,6 +436,7 @@ void FNiagaraScriptToolkit::OnEditedScriptPropertyFinishedChanging(const FProper
 
 void FNiagaraScriptToolkit::OnVMScriptCompiled(UNiagaraScript* InScript)
 {
+	UpdateMessageLog();
 	UpdateModuleStats();
 }
 
@@ -393,35 +525,14 @@ TSharedRef<SDockTab> FNiagaraScriptToolkit::SpawnTabMessageLog(const FSpawnTabAr
 {
 	check(Args.GetTabId().TabType == MessageLogTabID);
 
-	FMessageLogModule& MessageLogModule = FModuleManager::LoadModuleChecked<FMessageLogModule>("MessageLog");
-	const FName LogName = GetNiagaraScriptMessageLogName(EditedNiagaraScript);
-
-	auto CreateMessageLogListing = [&MessageLogModule, &LogName]()->TSharedRef<IMessageLogListing>
-	{
-		// Reuse any existing log, or create a new one (that is not held onto by the message log system)
-		if (MessageLogModule.IsRegisteredLogListing(LogName))
-		{
-			return MessageLogModule.GetLogListing(LogName);
-		}
-		else
-		{
-			FMessageLogInitializationOptions LogInitOptions;
-			LogInitOptions.bShowInLogWindow = false;
-			return  MessageLogModule.CreateLogListing(LogName, LogInitOptions);
-		}
-	};
-
-	TSharedRef<IMessageLogListing> MessageLogListing = CreateMessageLogListing();
-	//@todo(ng) debug message (to be removed when we forward real messages)
-	MessageLogListing->AddMessage(FTokenizedMessage::Create(EMessageSeverity::Info, FText::FromString(FString::Printf(TEXT("This is a debug message for %s"), *GetNiagaraScriptMessageLogName(EditedNiagaraScript).ToString()))));
-
-	MessageLogListing->OnMessageTokenClicked().AddStatic(&FNiagaraEditorModule::OnMessageLogTokenClicked);
-	TSharedRef<SWidget> LogListModule = MessageLogModule.CreateLogListingWidget(MessageLogListing);
-
-	TSharedRef<SDockTab> SpawnedTab =
-		SNew(SDockTab)
+	TSharedRef<SDockTab> SpawnedTab = SNew(SDockTab)
+		.Label(LOCTEXT("NiagaraMessageLogTitle", "Niagara Log"))
 		[
-			LogListModule
+			SNew(SBox)
+			.AddMetaData<FTagMetaData>(FTagMetaData(TEXT("NiagaraLog")))
+			[
+				NiagaraMessageLog.ToSharedRef()
+			]
 		];
 
 	return SpawnedTab;
@@ -564,6 +675,53 @@ void FNiagaraScriptToolkit::AddReferencedObjects(FReferenceCollector& Collector)
 
 
 
+void FNiagaraScriptToolkit::UpdateMessageLog()
+{
+	TArray< TSharedRef<class FTokenizedMessage> > Messages;
+
+	uint32 ErrorCount = 0;
+	uint32 WarningCount = 0;
+	for (FNiagaraCompileEvent CompileEvent : EditedNiagaraScript->GetVMExecutableData().LastCompileEvents)
+	{
+		Messages.Add(BuildMessageForCompileEvent(CompileEvent)); 
+		if (CompileEvent.Severity == FNiagaraCompileEventSeverity::Error)
+		{
+			ErrorCount++;
+		}
+		else if (CompileEvent.Severity == FNiagaraCompileEventSeverity::Warning)
+		{
+			WarningCount++;
+		}
+	}
+
+	const auto GetCompileCompleteMessageText = [&ErrorCount, &WarningCount](ENiagaraScriptCompileStatus Status)->const FText {
+		FText MessageText;
+		switch (Status)
+		{
+		default:
+		case ENiagaraScriptCompileStatus::NCS_Unknown:
+		case ENiagaraScriptCompileStatus::NCS_Dirty:
+			MessageText = LOCTEXT("NiagaraScriptCompileStatusUnknownInfo", "Script compile status unknown with {0} warning(s) and {1} error(s).");
+			return FText::Format(MessageText, FText::FromString(FString::FromInt(WarningCount)), FText::FromString(FString::FromInt(ErrorCount)));
+		case ENiagaraScriptCompileStatus::NCS_Error:
+			MessageText = LOCTEXT("NiagaraScriptCompileStatusErrorInfo", "Script failed to compile with {0} warning(s) and {1} error(s).");
+			return FText::Format(MessageText, FText::FromString(FString::FromInt(WarningCount)), FText::FromString(FString::FromInt(ErrorCount)));
+		case ENiagaraScriptCompileStatus::NCS_UpToDate:
+			MessageText = LOCTEXT("NiagaraScriptCompileStatusSuccessInfo", "Script successfully compiled.");
+			return MessageText;
+		case ENiagaraScriptCompileStatus::NCS_UpToDateWithWarnings:
+			MessageText = LOCTEXT("NiagaraScriptCompileStatusWarningInfo", "Script successfully compiled with {0} warning(s).");
+			return FText::Format(MessageText, FText::FromString(FString::FromInt(WarningCount)));
+		}
+	};
+
+	const FText CompileCompleteMessageText = GetCompileCompleteMessageText(ScriptViewModel->GetLatestCompileStatus());
+	Messages.Add(FTokenizedMessage::Create(EMessageSeverity::Info, CompileCompleteMessageText));
+
+	NiagaraMessageLogListing->ClearMessages();
+	NiagaraMessageLogListing->AddMessages(Messages);
+}
+
 void FNiagaraScriptToolkit::UpdateModuleStats()
 {
 	TArray< TSharedRef<class FTokenizedMessage> > Messages;
diff --git a/Engine/Plugins/FX/Niagara/Source/NiagaraEditor/Private/Toolkits/NiagaraScriptToolkit.h b/Engine/Plugins/FX/Niagara/Source/NiagaraEditor/Private/Toolkits/NiagaraScriptToolkit.h
--- a/Engine/Plugins/FX/Niagara/Source/NiagaraEditor/Private/Toolkits/NiagaraScriptToolkit.h
+++ b/Engine/Plugins/FX/Niagara/Source/NiagaraEditor/Private/Toolkits/NiagaraScriptToolkit.h
@@ -18,6 +18,10 @@ class UNiagaraScriptSource;
 class FNiagaraScriptViewModel;
 class FNiagaraObjectSelection;
 struct FEdGraphEditAction;
+class FTokenizedMessage; //@todo(ng) remove this forward decl after making FNiagaraMessageManager
+class IMessageToken; //@todo(ng) remove this forward decl after making FNiagaraMessageManager
+class UNiagaraGraph; //@todo(ng) remove this forward decl after making FNiagaraMessageManager
+struct FNiagaraCompileEvent; //@todo(ng) remove this forward decl after making FNiagaraMessageManager
 
 /** Viewer/editor for a DataTable */
 class FNiagaraScriptToolkit : public FAssetEditorToolkit, public FGCObject
@@ -50,6 +54,11 @@ public:
 	/** FGCObject interface */
 	virtual void AddReferencedObjects(FReferenceCollector& Collector) override;
 
+	/**
+	* Updates list of Niagara messages in message log
+	*/
+	void UpdateMessageLog();
+
 	/**
 	* Updates list of module info used to show stats
 	*/
@@ -116,12 +125,17 @@ private:
 	/** The selection displayed by the details tab. */
 	TSharedPtr<FNiagaraObjectSelection> DetailsSelection;
 
+	/** Message log, with the log listing that it reflects */
+	TSharedPtr<class SWidget> NiagaraMessageLog;
+	TSharedPtr<class IMessageLogListing> NiagaraMessageLogListing;
+
 	/**	The tab ids for the Niagara editor */
 	static const FName NodeGraphTabId; 
 	static const FName DetailsTabId;
 	static const FName ParametersTabId;
 	static const FName StatsTabId;
 	static const FName MessageLogTabID;
+
 	/** Stats log, with the log listing that it reflects */
 	TSharedPtr<class SWidget> Stats;
 	TSharedPtr<class IMessageLogListing> StatsListing;
@@ -130,4 +144,12 @@ private:
 
 	bool bEditedScriptHasPendingChanges;
 	bool bChangesDiscarded;
+
+	/**  //@todo(ng) Everything below this line should be moved to an owning class such as FNiagaraMessageManager as we don't want to duplicate this code for both the Script and System toolkits
+	* --------------------------------------------------------------------------------------------------------------------------------------------------------------------------------------------
+	*/
+	
+	TSharedRef<FTokenizedMessage> BuildMessageForCompileEvent(FNiagaraCompileEvent& InCompileEvent);
+	TOptional<UNiagaraGraph*> RecursiveBuildMessageTokensFromContextStackAndGetOriginatingGraph(TArray<FGuid>& InContextStackNodeGuids, UNiagaraGraph* InGraphToSearch, TArray<TSharedRef<IMessageToken>>& OutMessageTokensToAdd);
+
 };
diff --git a/Engine/Plugins/FX/Niagara/Source/NiagaraEditor/Private/Toolkits/NiagaraSystemToolkit.cpp b/Engine/Plugins/FX/Niagara/Source/NiagaraEditor/Private/Toolkits/NiagaraSystemToolkit.cpp
--- a/Engine/Plugins/FX/Niagara/Source/NiagaraEditor/Private/Toolkits/NiagaraSystemToolkit.cpp
+++ b/Engine/Plugins/FX/Niagara/Source/NiagaraEditor/Private/Toolkits/NiagaraSystemToolkit.cpp
@@ -50,6 +50,10 @@
 #include "Misc/FileHelper.h"
 #include "MessageLog/Public/MessageLogModule.h"
 #include "MessageLog/Public/IMessageLogListing.h"
+#include "NiagaraScriptSource.h" //@todo(ng) remove this include after making FNiagaraMessageManager
+#include "NiagaraNodeFunctionCall.h" //@todo(ng) remove this include after making FNiagaraMessageManager
+#include "NiagaraScript.h" //@todo(ng) dependency for fniagaracompile event, remove this include after making FNiagaraMessageManager
+#include "NiagaraNodeEmitter.h" //@todo(ng) remove this include after making FNiagaraMessageManager
 
 #define LOCTEXT_NAMESPACE "NiagaraSystemEditor"
 
@@ -171,12 +175,186 @@ FNiagaraSystemToolkit::~FNiagaraSystemToolkit()
 	SystemViewModel.Reset();
 }
 
+//@todo(ng) move to message manager class
+TSharedRef<FTokenizedMessage> FNiagaraSystemToolkit::BuildMessageForCompileEvent(FNiagaraCompileEvent& InCompileEvent, UNiagaraScript* InOriginatingScript)
+{
+	TArray<FGuid>& ContextStackGuids = InCompileEvent.StackGuids;
+
+	UNiagaraScriptSourceBase* FunctionScriptSourceBase = InOriginatingScript->GetSource();
+	checkf(FunctionScriptSourceBase->IsA<UNiagaraScriptSource>(), TEXT("Script source for function call node is not assigned or is not of type UNiagaraScriptSource!"))
+	UNiagaraScriptSource* FunctionScriptSource = Cast<UNiagaraScriptSource>(FunctionScriptSourceBase);
+	checkf(FunctionScriptSource, TEXT("Script source base was somehow not a derived type!"));
+	UNiagaraGraph* ThisScriptGraph = FunctionScriptSource->NodeGraph;
+	checkf(ThisScriptGraph, TEXT("Function Script does not have a UNiagaraGraph!"));
+
+	TArray<TSharedRef<IMessageToken>> MessageTokens;
+	TOptional<UNiagaraGraph*> OriginatingGraph = RecursiveBuildMessageTokensFromContextStackAndGetOriginatingGraph(ContextStackGuids, ThisScriptGraph, MessageTokens);
+	if (OriginatingGraph.IsSet())
+	{
+		EMessageSeverity::Type MessageSeverity = EMessageSeverity::Info;
+		switch (InCompileEvent.Severity) {
+		case FNiagaraCompileEventSeverity::Error:
+			MessageSeverity = EMessageSeverity::Error;
+			break;
+		case FNiagaraCompileEventSeverity::Warning:
+			MessageSeverity = EMessageSeverity::Warning;
+			break;
+		case FNiagaraCompileEventSeverity::Log:
+			MessageSeverity = EMessageSeverity::Info;
+			break;
+		default:
+			verifyf(false, TEXT("Compile event severity type not handled!"));
+			MessageSeverity = EMessageSeverity::Info;
+			break;
+		}
+		TSharedRef<FTokenizedMessage> MessageToBuild = FTokenizedMessage::Create(MessageSeverity);
+
+		//Add message from compile event at start of message
+		MessageToBuild->AddToken(FTextToken::Create(FText::FromString(InCompileEvent.Message)));
+		//Now add the context stack of the scripts that were passed through to get to the originating graph
+		for (TSharedRef<IMessageToken> Token : MessageTokens)
+		{
+			MessageToBuild->AddToken(Token);
+		}
+		return MessageToBuild;
+	}
+
+	// Failed to get originating graph!
+	verifyf(false, TEXT("Failed to find originating graph for compile event!"));
+	return FTokenizedMessage::Create(EMessageSeverity::Warning, LOCTEXT("NiagaraCompileEventFailedToFindScript", "Could not find script(s) for compile event ! Please recompile script to get full compile info."));
+}
+
+//@todo(ng) move to message manager class
+TOptional<UNiagaraGraph*> FNiagaraSystemToolkit::RecursiveBuildMessageTokensFromContextStackAndGetOriginatingGraph(TArray<FGuid>& InContextStackNodeGuids, UNiagaraGraph* InGraphToSearch, TArray<TSharedRef<IMessageToken>>& OutMessageTokensToAdd)
+{
+	verifyf(InGraphToSearch, TEXT("Failed to get a node graph to search!"));
+
+	if (InGraphToSearch && InContextStackNodeGuids.Num() == 0)
+	{
+		//StackGuids arr has been cleared out which means we have walked the entire context and can now return the graph we are in.
+		return InGraphToSearch;
+	}
+
+	// Search in the current graph for a node with a GUID that matches a GUID in the list of Function Call and Emitter node GUIDs that define the context stack for a compile event
+	auto FindNodeInGraphWithContextStackGuid = [&InGraphToSearch, &InContextStackNodeGuids]()->TOptional<UEdGraphNode*> {
+		for (UEdGraphNode* GraphNode : InGraphToSearch->Nodes)
+		{
+			for (int i = 0; i < InContextStackNodeGuids.Num(); i++)
+			{
+				if (GraphNode->NodeGuid == InContextStackNodeGuids[i])
+				{
+					InContextStackNodeGuids.RemoveAt(i);
+					return GraphNode;
+				}
+			}
+		}
+		return TOptional<UEdGraphNode*>();
+	};
+
+	TOptional<UEdGraphNode*> ContextNode = FindNodeInGraphWithContextStackGuid();
+	if (ContextNode.IsSet())
+	{
+		// found a node in the current graph that has a GUID in the context list
+		UNiagaraNodeFunctionCall* FunctionCallNode = Cast<UNiagaraNodeFunctionCall>(ContextNode.GetValue());
+		if (FunctionCallNode)
+		{
+			// node is a function call node, now get the Niagara Script assigned to this node, add a message token and recursively call into the graph of that script.
+			UNiagaraScript* FunctionCallNodeAssignedScript = FunctionCallNode->FunctionScript;
+			checkf(FunctionCallNodeAssignedScript, TEXT("Script for function call node is not assigned!"));
+			UNiagaraScriptSourceBase* FunctionScriptSourceBase = FunctionCallNodeAssignedScript->GetSource();
+			checkf(FunctionScriptSourceBase->IsA<UNiagaraScriptSource>(), TEXT("Script source for function call node is not assigned or is not of type UNiagaraScriptSource!"));
+			UNiagaraScriptSource* FunctionScriptSource = Cast<UNiagaraScriptSource>(FunctionScriptSourceBase);
+			checkf(FunctionScriptSource, TEXT("Script source base was somehow not a derived type!"));
+			UNiagaraGraph* FunctionScriptGraph = FunctionScriptSource->NodeGraph;
+			checkf(FunctionScriptGraph, TEXT("Function Script does not have a UNiagaraGraph!"));
+
+			TSharedRef<FNiagaraAssetNameToken> VisitedGraphContextToken = FNiagaraAssetNameToken::Create(FunctionCallNodeAssignedScript->GetPathName(), FText::FromString(*FunctionCallNodeAssignedScript->GetName()));
+			OutMessageTokensToAdd.Add(VisitedGraphContextToken);
+			return RecursiveBuildMessageTokensFromContextStackAndGetOriginatingGraph(InContextStackNodeGuids, FunctionScriptGraph, OutMessageTokensToAdd);
+		}
+		UNiagaraNodeEmitter* EmitterNode = Cast<UNiagaraNodeEmitter>(ContextNode.GetValue());
+		if (EmitterNode)
+		{
+			// node is an emitter node, now get the Emitter name, add a message token and recursively call into the graph of that emitter.
+			UNiagaraScriptSource* EmitterScriptSource = EmitterNode->GetScriptSource();
+			UNiagaraGraph* EmitterScriptGraph = EmitterScriptSource->NodeGraph;
+			checkf(EmitterScriptGraph, TEXT("Emitter Script Source does not have a UNiagaraGraph!")); //@todo(ng) see about filtering the message log by emitter?
+
+			TSharedRef<FTextToken> SourceEmitterContextToken = FTextToken::Create(FText::FromString(EmitterNode->GetEmitterUniqueName()));
+			OutMessageTokensToAdd.Add(SourceEmitterContextToken);
+			return RecursiveBuildMessageTokensFromContextStackAndGetOriginatingGraph(InContextStackNodeGuids, EmitterScriptGraph, OutMessageTokensToAdd);
+		}
+		checkf(false, TEXT("Matching node is not a function call or emitter node!"));
+	}
+	verifyf(false, TEXT("Failed to walk the entire context stack, is this compile event out of date?"));
+	return TOptional<UNiagaraGraph*>();
+}
+
 void FNiagaraSystemToolkit::AddReferencedObjects(FReferenceCollector& Collector) 
 {
 	Collector.AddReferencedObject(System);
 }
 
+void FNiagaraSystemToolkit::UpdateMessageLog()
+{	
+	TArray<UNiagaraScript*> ScriptsToGetCompileEventsFrom;
+	ScriptsToGetCompileEventsFrom.Add(System->GetSystemSpawnScript());
+	ScriptsToGetCompileEventsFrom.Add(System->GetSystemUpdateScript());
+	const TArray<FNiagaraEmitterHandle> EmitterHandles = System->GetEmitterHandles();
+	for (const FNiagaraEmitterHandle& Handle : EmitterHandles)
+	{
+		const UNiagaraEmitter* EmitterInSystem = Handle.GetInstance();
+		TArray<UNiagaraScript*> EmitterScripts;
+		EmitterInSystem->GetScripts(EmitterScripts);
+		ScriptsToGetCompileEventsFrom.Append(EmitterScripts);
+	}
 
+	uint32 ErrorCount = 0;
+	uint32 WarningCount = 0;
+	TArray< TSharedRef<class FTokenizedMessage> > Messages;
+	for (UNiagaraScript* Script : ScriptsToGetCompileEventsFrom)
+	{
+		for (FNiagaraCompileEvent CompileEvent : Script->GetVMExecutableData().LastCompileEvents)
+		{
+			Messages.Add(BuildMessageForCompileEvent(CompileEvent, Script));
+			if (CompileEvent.Severity == FNiagaraCompileEventSeverity::Error)
+			{
+				ErrorCount++;
+			}
+			else if (CompileEvent.Severity == FNiagaraCompileEventSeverity::Warning)
+			{
+				WarningCount++;
+			}
+		}
+	}
+
+	const auto GetCompileCompleteMessageText = [&ErrorCount, &WarningCount](ENiagaraScriptCompileStatus Status)->const FText {
+		FText MessageText;
+		switch (Status)
+		{
+		default:
+		case ENiagaraScriptCompileStatus::NCS_Unknown:
+		case ENiagaraScriptCompileStatus::NCS_Dirty:
+			MessageText = LOCTEXT("NiagaraSystemCompileStatusUnknownInfo", "System compile status unknown with {0} warning(s) and {1} error(s).");
+			return FText::Format(MessageText, FText::FromString(FString::FromInt(WarningCount)), FText::FromString(FString::FromInt(ErrorCount)));
+		case ENiagaraScriptCompileStatus::NCS_Error:
+			MessageText = LOCTEXT("NiagaraSystemCompileStatusErrorInfo", "System failed to compile with {0} warning(s) and {1} error(s).");
+			return FText::Format(MessageText, FText::FromString(FString::FromInt(WarningCount)), FText::FromString(FString::FromInt(ErrorCount)));
+		case ENiagaraScriptCompileStatus::NCS_UpToDate:
+			MessageText = LOCTEXT("NiagaraSystemCompileStatusSuccessInfo", "System successfully compiled.");
+			return MessageText;
+		case ENiagaraScriptCompileStatus::NCS_UpToDateWithWarnings:
+			MessageText = LOCTEXT("NiagaraSystemCompileStatusWarningInfo", "System successfully compiled with {0} warning(s).");
+			return FText::Format(MessageText, FText::FromString(FString::FromInt(WarningCount)));
+		}
+	};
+	
+	const FText CompileCompleteMessageText = GetCompileCompleteMessageText(SystemViewModel->GetLatestCompileStatus());
+	Messages.Add(FTokenizedMessage::Create(EMessageSeverity::Info, CompileCompleteMessageText));
+
+	NiagaraMessageLogListing->ClearMessages();
+	NiagaraMessageLogListing->AddMessages(Messages);
+}
 
 void FNiagaraSystemToolkit::InitializeWithSystem(const EToolkitMode::Type Mode, const TSharedPtr< class IToolkitHost >& InitToolkitHost, UNiagaraSystem& InSystem)
 {
@@ -282,11 +460,37 @@ void FNiagaraSystemToolkit::InitializeWithEmitter(const EToolkitMode::Type Mode,
 
 void FNiagaraSystemToolkit::InitializeInternal(const EToolkitMode::Type Mode, const TSharedPtr<IToolkitHost>& InitToolkitHost)
 {
+	SystemViewModel->GetSystemScriptViewModel()->OnSystemCompiled().AddSP(this, &FNiagaraSystemToolkit::OnVMSystemCompiled);
+
 	if (SystemViewModel->GetEmitterHandleViewModels().Num() > 0)
 	{
 		SystemViewModel->SetSelectedEmitterHandleById(SystemViewModel->GetEmitterHandleViewModels()[0]->GetId());
 	}
 
+	FMessageLogModule& MessageLogModule = FModuleManager::LoadModuleChecked<FMessageLogModule>("MessageLog");
+
+	// Reuse any existing log, or create a new one (that is not held onto by the message log system)
+	auto CreateMessageLogListing = [&MessageLogModule](const FName& LogName)->TSharedRef<IMessageLogListing> {	
+		FMessageLogInitializationOptions LogOptions;
+		// Show Pages so that user is never allowed to clear log messages
+		LogOptions.bShowPages = false;
+		LogOptions.bShowFilters = false;
+		LogOptions.bAllowClear = false;
+		LogOptions.MaxPageCount = 1;
+
+		if (MessageLogModule.IsRegisteredLogListing(LogName))
+		{
+			return MessageLogModule.GetLogListing(LogName);
+		}
+		else
+		{
+			return  MessageLogModule.CreateLogListing(LogName, LogOptions);
+		}
+	};
+
+	NiagaraMessageLogListing = CreateMessageLogListing(GetNiagaraSystemMessageLogName(System));
+	NiagaraMessageLog = MessageLogModule.CreateLogListingWidget(NiagaraMessageLogListing.ToSharedRef());
+
 	SystemViewModel->OnEmitterHandleViewModelsChanged().AddSP(this, &FNiagaraSystemToolkit::OnRefresh);
 	SystemViewModel->OnSelectedEmitterHandlesChanged().AddSP(this, &FNiagaraSystemToolkit::OnRefresh);
 	SystemViewModel->GetOnPinnedEmittersChanged().AddSP(this, &FNiagaraSystemToolkit::OnRefresh);
@@ -366,6 +570,7 @@ void FNiagaraSystemToolkit::InitializeInternal(const EToolkitMode::Type Mode, co
 	SetupCommands();
 	ExtendToolbar();
 	RegenerateMenusAndToolbars();
+	UpdateMessageLog();
 
 	bChangesDiscarded = false;
 }
@@ -626,35 +831,14 @@ TSharedRef<SDockTab> FNiagaraSystemToolkit::SpawnTab_MessageLog(const FSpawnTabA
 {
 	check(Args.GetTabId().TabType == MessageLogTabID);
 
-	FMessageLogModule& MessageLogModule = FModuleManager::LoadModuleChecked<FMessageLogModule>("MessageLog");
-	const FName LogName = GetNiagaraSystemMessageLogName(System);
-
-	auto CreateMessageLogListing = [&MessageLogModule, &LogName]()->TSharedRef<IMessageLogListing>
-	{
-		// Reuse any existing log, or create a new one (that is not held onto by the message log system)
-		if (MessageLogModule.IsRegisteredLogListing(LogName))
-		{
-			return MessageLogModule.GetLogListing(LogName);
-		}
-		else
-		{
-			FMessageLogInitializationOptions LogInitOptions;
-			LogInitOptions.bShowInLogWindow = false;
-			return  MessageLogModule.CreateLogListing(LogName, LogInitOptions);
-		}
-	};
-
-	TSharedRef<IMessageLogListing> MessageLogListing = CreateMessageLogListing();
-	//@todo(ng) debug message (to be removed when we forward real messages)
-	MessageLogListing->AddMessage(FTokenizedMessage::Create(EMessageSeverity::Info, FText::FromString(FString::Printf(TEXT("This is a debug message for %s"), *GetNiagaraSystemMessageLogName(System).ToString()))));
-
-	MessageLogListing->OnMessageTokenClicked().AddStatic(&FNiagaraEditorModule::OnMessageLogTokenClicked);
-	TSharedRef<SWidget> LogListModule = MessageLogModule.CreateLogListingWidget(MessageLogListing);
-
-	TSharedRef<SDockTab> SpawnedTab =
-		SNew(SDockTab)
+	TSharedRef<SDockTab> SpawnedTab = SNew(SDockTab)
+		.Label(LOCTEXT("NiagaraMessageLogTitle", "Niagara Log"))
 		[
-			LogListModule
+			SNew(SBox)
+			.AddMetaData<FTagMetaData>(FTagMetaData(TEXT("NiagaraLog")))
+			[
+				NiagaraMessageLog.ToSharedRef()
+			]
 		];
 
 	return SpawnedTab;
@@ -748,6 +932,11 @@ void FNiagaraSystemToolkit::ResetSimulation()
 	SystemViewModel->ResetSystem();
 }
 
+void FNiagaraSystemToolkit::OnVMSystemCompiled()
+{
+	UpdateMessageLog();
+}
+
 void FNiagaraSystemToolkit::ExtendToolbar()
 {
 	struct Local
diff --git a/Engine/Plugins/FX/Niagara/Source/NiagaraEditor/Private/Toolkits/NiagaraSystemToolkit.h b/Engine/Plugins/FX/Niagara/Source/NiagaraEditor/Private/Toolkits/NiagaraSystemToolkit.h
--- a/Engine/Plugins/FX/Niagara/Source/NiagaraEditor/Private/Toolkits/NiagaraSystemToolkit.h
+++ b/Engine/Plugins/FX/Niagara/Source/NiagaraEditor/Private/Toolkits/NiagaraSystemToolkit.h
@@ -25,6 +25,11 @@ class UNiagaraSequence;
 struct FAssetData;
 class FMenuBuilder;
 class ISequencer;
+class FTokenizedMessage; //@todo(ng) remove this forward decl after making FNiagaraMessageManager
+class IMessageToken; //@todo(ng) remove this forward decl after making FNiagaraMessageManager
+struct FNiagaraCompileEvent; //@todo(ng) remove this forward decl after making FNiagaraMessageManager
+class UNiagaraGraph; //@todo(ng) remove this forward decl after making FNiagaraMessageManager
+class UNiagaraScript; //@todo(ng) remove this forward decl after making FNiagaraMessageManager
 
 /** Viewer/editor for a NiagaraSystem
 */
@@ -59,6 +64,11 @@ public:
 	//~ FGCObject interface
 	virtual void AddReferencedObjects(FReferenceCollector& Collector) override;
 
+	/**
+	* Updates list of Niagara messages in message log
+	*/
+	void UpdateMessageLog();
+
 	FSlateIcon GetCompileStatusImage() const;
 	FText GetCompileStatusTooltip() const;
 
@@ -106,6 +116,8 @@ private:
 
 	void ResetSimulation();
 
+	void OnVMSystemCompiled();
+
 	void GetSequencerAddMenuContent(FMenuBuilder& MenuBuilder, TSharedRef<ISequencer> Sequencer);
 	TSharedRef<SWidget> CreateAddEmitterMenuContent();
 	TSharedRef<SWidget> GenerateCompileMenuContent();
@@ -149,6 +161,10 @@ private:
 	/* The view model for the System being edited */
 	TSharedPtr<FNiagaraSystemViewModel> SystemViewModel;
 
+	/** Message log, with the log listing that it reflects */
+	TSharedPtr<class SWidget> NiagaraMessageLog;
+	TSharedPtr<class IMessageLogListing> NiagaraMessageLogListing;
+
 	/** The command list for this editor */
 	TSharedPtr<FUICommandList> EditorCommands;
 
@@ -168,4 +184,12 @@ private:
 	static const FName PreviewSettingsTabId;
 	static const FName GeneratedCodeTabID;
 	static const FName MessageLogTabID;
+
+	/**  //@todo(ng) Everything below this line should be moved to an owning class such as FNiagaraMessageManager as we don't want to duplicate this code for both the Script and System toolkits
+* --------------------------------------------------------------------------------------------------------------------------------------------------------------------------------------------
+*/
+
+	TSharedRef<FTokenizedMessage> BuildMessageForCompileEvent(FNiagaraCompileEvent& InCompileEvent, UNiagaraScript* InOriginatingScript);
+	TOptional<UNiagaraGraph*> RecursiveBuildMessageTokensFromContextStackAndGetOriginatingGraph(TArray<FGuid>& InContextStackNodeGuids, UNiagaraGraph* InGraphToSearch, TArray<TSharedRef<IMessageToken>>& OutMessageTokensToAdd);
+
 };
diff --git a/Engine/Plugins/FX/Niagara/Source/NiagaraEditor/Public/INiagaraCompiler.h b/Engine/Plugins/FX/Niagara/Source/NiagaraEditor/Public/INiagaraCompiler.h
--- a/Engine/Plugins/FX/Niagara/Source/NiagaraEditor/Public/INiagaraCompiler.h
+++ b/Engine/Plugins/FX/Niagara/Source/NiagaraEditor/Public/INiagaraCompiler.h
@@ -17,34 +17,7 @@ struct FNiagaraTranslatorOutput;
 struct FNiagaraVMExecutableData;
 class FNiagaraCompileRequestData;
 class FNiagaraCompileOptions;
-
-/** Defines the compile event types for translation/compilation.*/
-enum FNiagaraCompileEventType
-{
-	Log = 0,
-	Warning = 1,
-	Error = 2
-};
-
-/** Records necessary information to give UI cues for errors/logs/warnings during compile.*/
-struct FNiagaraCompileEvent
-{
-public:
-	FNiagaraCompileEvent(FNiagaraCompileEventType InType, const FString& InMessage, FGuid InNodeGuid = FGuid(), FGuid InPinGuid = FGuid(), const TArray<FGuid>& InCallstackGuids = TArray<FGuid>())
-	: Type(InType), Message(InMessage), NodeGuid(InNodeGuid), PinGuid(InPinGuid), StackGuids(InCallstackGuids){}
-
-	/** Whether or not this is an error, warning, or info*/
-	FNiagaraCompileEventType Type;
-	/* The message itself*/
-	FString Message;
-	/** The node guid that generated the compile event*/
-	FGuid NodeGuid;
-	/** The pin persistent id that generated the compile event*/
-	FGuid PinGuid;
-	/** The compile stack frame of node id's*/
-	TArray<FGuid> StackGuids;
-};
-
+struct FNiagaraCompileEvent;
 
 /** Defines information about the results of a Niagara script compile. */
 struct FNiagaraCompileResults
diff --git a/Engine/Plugins/FX/Niagara/Source/NiagaraEditor/Public/NiagaraEditorModule.h b/Engine/Plugins/FX/Niagara/Source/NiagaraEditor/Public/NiagaraEditorModule.h
--- a/Engine/Plugins/FX/Niagara/Source/NiagaraEditor/Public/NiagaraEditorModule.h
+++ b/Engine/Plugins/FX/Niagara/Source/NiagaraEditor/Public/NiagaraEditorModule.h
@@ -8,6 +8,7 @@
 #include "NiagaraTypes.h"
 #include "INiagaraCompiler.h"
 #include "AssetTypeCategories.h"
+#include "Logging/TokenizedMessage.h" //@todo(ng) remove this include after moving FNiagaraAssetPathToken to message manager
 
 class IAssetTools;
 class IAssetTypeActions;
@@ -136,4 +137,32 @@ private:
 	IConsoleCommand* DumpRapidIterationParametersForAsset;
 };
 
+//@todo(ng) move to message manager class
+//Extension of FAssetNameToken to allow opening the asset editor when clicking on the linked asset name.
+class FNiagaraAssetNameToken : public IMessageToken
+{
+public:
+	/** Factory method, tokens can only be constructed as shared refs */
+	static TSharedRef<FNiagaraAssetNameToken> Create(const FString& InAssetName, const FText& InMessage = FText());
+
+	/** Begin IMessageToken interface */
+	virtual EMessageToken::Type GetType() const override
+	{
+		return EMessageToken::AssetName;
+	}
+	/** End IMessageToken interface */
 
+private:
+	/** Private constructor */
+	FNiagaraAssetNameToken(const FString& InAssetName, const FText& InMessage);
+
+	/**
+	 * Find and open an asset in editor
+	 * @param	Token		The token that was clicked
+	 * @param	InAssetPath		The asset to find
+	 */
+	static void FindAndOpenAsset(const TSharedRef<IMessageToken>& Token, FString InAssetPath);
+
+	/** The asset name we will find */
+	FString AssetName;
+};
