Remove forward declare of structs resulting in compile fail.

diff --git a/Engine/Plugins/FX/Niagara/Source/NiagaraEditor/Private/Toolkits/NiagaraScriptToolkit.h b/Engine/Plugins/FX/Niagara/Source/NiagaraEditor/Private/Toolkits/NiagaraScriptToolkit.h
--- a/Engine/Plugins/FX/Niagara/Source/NiagaraEditor/Private/Toolkits/NiagaraScriptToolkit.h
+++ b/Engine/Plugins/FX/Niagara/Source/NiagaraEditor/Private/Toolkits/NiagaraScriptToolkit.h
@@ -10,6 +10,8 @@
 #include "Toolkits/AssetEditorToolkit.h"
 #include "UObject/GCObject.h"
 
+#include "NiagaraScript.h"
+
 class IDetailsView;
 class SGraphEditor;
 class UEdGraph;
@@ -21,7 +23,7 @@ struct FEdGraphEditAction;
 class FTokenizedMessage; //@todo(ng) remove this forward decl after making FNiagaraMessageManager
 class IMessageToken; //@todo(ng) remove this forward decl after making FNiagaraMessageManager
 class UNiagaraGraph; //@todo(ng) remove this forward decl after making FNiagaraMessageManager
-struct FNiagaraCompileEvent; //@todo(ng) remove this forward decl after making FNiagaraMessageManager
+//struct FNiagaraCompileEvent; //@todo(ng) remove this forward decl after making FNiagaraMessageManager
 
 /** Viewer/editor for a DataTable */
 class FNiagaraScriptToolkit : public FAssetEditorToolkit, public FGCObject
diff --git a/Engine/Plugins/FX/Niagara/Source/NiagaraEditor/Private/Toolkits/NiagaraSystemToolkit.h b/Engine/Plugins/FX/Niagara/Source/NiagaraEditor/Private/Toolkits/NiagaraSystemToolkit.h
--- a/Engine/Plugins/FX/Niagara/Source/NiagaraEditor/Private/Toolkits/NiagaraSystemToolkit.h
+++ b/Engine/Plugins/FX/Niagara/Source/NiagaraEditor/Private/Toolkits/NiagaraSystemToolkit.h
@@ -13,6 +13,8 @@
 #include "ISequencer.h"
 #include "ISequencerTrackEditor.h"
 
+#include "NiagaraScript.h"
+
 class FNiagaraSystemInstance;
 class FNiagaraSystemViewModel;
 class SNiagaraSystemEditorViewport;
@@ -27,9 +29,9 @@ class FMenuBuilder;
 class ISequencer;
 class FTokenizedMessage; //@todo(ng) remove this forward decl after making FNiagaraMessageManager
 class IMessageToken; //@todo(ng) remove this forward decl after making FNiagaraMessageManager
-struct FNiagaraCompileEvent; //@todo(ng) remove this forward decl after making FNiagaraMessageManager
+//struct FNiagaraCompileEvent; //@todo(ng) remove this forward decl after making FNiagaraMessageManager
 class UNiagaraGraph; //@todo(ng) remove this forward decl after making FNiagaraMessageManager
-class UNiagaraScript; //@todo(ng) remove this forward decl after making FNiagaraMessageManager
+//class UNiagaraScript; //@todo(ng) remove this forward decl after making FNiagaraMessageManager
 
 /** Viewer/editor for a NiagaraSystem
 */
diff --git a/Engine/Plugins/FX/Niagara/Source/NiagaraEditor/Public/INiagaraCompiler.h b/Engine/Plugins/FX/Niagara/Source/NiagaraEditor/Public/INiagaraCompiler.h
--- a/Engine/Plugins/FX/Niagara/Source/NiagaraEditor/Public/INiagaraCompiler.h
+++ b/Engine/Plugins/FX/Niagara/Source/NiagaraEditor/Public/INiagaraCompiler.h
@@ -5,6 +5,7 @@
 #include "CoreMinimal.h"
 #include "NiagaraEditorCommon.h"
 #include "NiagaraParameters.h"
+#include "NiagaraScript.h"
 
 class Error;
 class UEdGraphPin;
@@ -17,7 +18,7 @@ struct FNiagaraTranslatorOutput;
 struct FNiagaraVMExecutableData;
 class FNiagaraCompileRequestData;
 class FNiagaraCompileOptions;
-struct FNiagaraCompileEvent;
+//struct FNiagaraCompileEvent;
 
 /** Defines information about the results of a Niagara script compile. */
 struct FNiagaraCompileResults
