- Enable velocity output for proxies who have velocity always enabled

diff --git a/Engine/Plugins/FX/Niagara/Source/Niagara/Private/NiagaraRendererMeshes.cpp b/Engine/Plugins/FX/Niagara/Source/Niagara/Private/NiagaraRendererMeshes.cpp
--- a/Engine/Plugins/FX/Niagara/Source/Niagara/Private/NiagaraRendererMeshes.cpp
+++ b/Engine/Plugins/FX/Niagara/Source/Niagara/Private/NiagaraRendererMeshes.cpp
@@ -170,7 +170,7 @@ void NiagaraRendererMeshes::GetDynamicMeshElements(const TArray<const FSceneView
 				0,
 				INDEX_NONE,
 				INDEX_NONE,
-				false
+				SceneProxy->AlwaysHasVelocity()
 				);
 			WorldSpacePrimitiveUniformBuffer.SetContents(PrimitiveUniformShaderParameters);
 			WorldSpacePrimitiveUniformBuffer.InitResource();
diff --git a/Engine/Plugins/FX/Niagara/Source/Niagara/Private/NiagaraRendererRibbons.cpp b/Engine/Plugins/FX/Niagara/Source/Niagara/Private/NiagaraRendererRibbons.cpp
--- a/Engine/Plugins/FX/Niagara/Source/Niagara/Private/NiagaraRendererRibbons.cpp
+++ b/Engine/Plugins/FX/Niagara/Source/Niagara/Private/NiagaraRendererRibbons.cpp
@@ -137,7 +137,7 @@ void NiagaraRendererRibbons::GetDynamicMeshElements(const TArray<const FSceneVie
 				0,
 				INDEX_NONE,
 				INDEX_NONE,
-				false
+				SceneProxy->AlwaysHasVelocity()
 				);
 			WorldSpacePrimitiveUniformBuffer.SetContents(PrimitiveUniformShaderParameters);
 			WorldSpacePrimitiveUniformBuffer.InitResource();
diff --git a/Engine/Plugins/FX/Niagara/Source/Niagara/Private/NiagaraRendererSprites.cpp b/Engine/Plugins/FX/Niagara/Source/Niagara/Private/NiagaraRendererSprites.cpp
--- a/Engine/Plugins/FX/Niagara/Source/Niagara/Private/NiagaraRendererSprites.cpp
+++ b/Engine/Plugins/FX/Niagara/Source/Niagara/Private/NiagaraRendererSprites.cpp
@@ -135,7 +135,7 @@ void NiagaraRendererSprites::ConditionalInitPrimitiveUniformBuffer(const FNiagar
 			0,
 			INDEX_NONE,
 			INDEX_NONE,
-			false
+			SceneProxy->AlwaysHasVelocity()
 		);
 		WorldSpacePrimitiveUniformBuffer.SetContents(PrimitiveUniformShaderParameters);
 		WorldSpacePrimitiveUniformBuffer.InitResource();
diff --git a/Engine/Source/Runtime/Engine/Private/Particles/ParticleSystemRender.cpp b/Engine/Source/Runtime/Engine/Private/Particles/ParticleSystemRender.cpp
--- a/Engine/Source/Runtime/Engine/Private/Particles/ParticleSystemRender.cpp
+++ b/Engine/Source/Runtime/Engine/Private/Particles/ParticleSystemRender.cpp
@@ -7284,7 +7284,7 @@ void FParticleSystemSceneProxy::UpdateWorldSpacePrimitiveUniformBuffer() const
 			0,
 			INDEX_NONE,
 			INDEX_NONE,
-			false
+			AlwaysHasVelocity()
 			);
 		WorldSpacePrimitiveUniformBuffer.SetContents(PrimitiveUniformShaderParameters);
 		WorldSpacePrimitiveUniformBuffer.InitResource();
diff --git a/Engine/Source/Runtime/Engine/Private/PrimitiveSceneProxy.cpp b/Engine/Source/Runtime/Engine/Private/PrimitiveSceneProxy.cpp
--- a/Engine/Source/Runtime/Engine/Private/PrimitiveSceneProxy.cpp
+++ b/Engine/Source/Runtime/Engine/Private/PrimitiveSceneProxy.cpp
@@ -292,7 +292,7 @@ void FPrimitiveSceneProxy::UpdateUniformBuffer()
 				LpvBiasMultiplier,
 				PrimitiveSceneInfo ? PrimitiveSceneInfo->GetLightmapDataOffset() : 0,
 				SingleCaptureIndex, 
-				bOutputVelocity,
+				bOutputVelocity || AlwaysHasVelocity(),
 				GetCustomPrimitiveData());
 
 		if (UniformBuffer.GetReference())
