Bump Niagara DDC version to resolve issue with editor hanging trying to deserialize old DDC data.

diff --git a/Engine/Plugins/FX/Niagara/Source/Niagara/Private/NiagaraScriptDerivedData.h b/Engine/Plugins/FX/Niagara/Source/Niagara/Private/NiagaraScriptDerivedData.h
--- a/Engine/Plugins/FX/Niagara/Source/Niagara/Private/NiagaraScriptDerivedData.h
+++ b/Engine/Plugins/FX/Niagara/Source/Niagara/Private/NiagaraScriptDerivedData.h
@@ -54,7 +54,7 @@ public:
 		// This is a version string that mimics the old versioning scheme. If you
 		// want to bump this version, generate a new guid using VS->Tools->Create GUID and
 		// return it here. Ex.
-		return TEXT("B19605DB7417452D85E5BD2E13EE370B");
+		return TEXT("FB23DB34150247469E8B70EBE84A30C9");
 	}
 
 	virtual FString GetPluginSpecificCacheKeySuffix() const override;
