Allow synchronized resetting of emitters and open systems. Consolidate Reset and Reinitialize methods of NiagaraSystemViewModel. Fix Emitter assets always being named "Emitter" in editor. Now avoid making extra copy of source emitters.

diff --git a/Engine/Plugins/FX/Niagara/Source/Niagara/Classes/NiagaraEmitterHandle.h b/Engine/Plugins/FX/Niagara/Source/Niagara/Classes/NiagaraEmitterHandle.h
--- a/Engine/Plugins/FX/Niagara/Source/Niagara/Classes/NiagaraEmitterHandle.h
+++ b/Engine/Plugins/FX/Niagara/Source/Niagara/Classes/NiagaraEmitterHandle.h
@@ -24,10 +24,6 @@ public:
 	/** Creates a new invalid emitter handle. */
 	FNiagaraEmitterHandle();
 
-	/** Create a new emitter handle from an emitter, but does NOT make a copy, any changes made to the "Instance" will modify
-		the original asset.  This version should only be used in the emitter toolkit. */
-	FNiagaraEmitterHandle(UNiagaraEmitter& Emitter);
-
 #if WITH_EDITORONLY_DATA
 	/** Creates a new emitter handle from an emitter and an owning System. */
 	FNiagaraEmitterHandle(UNiagaraEmitter& InSourceEmitter, FName InName, UNiagaraSystem& InOuterSystem);
diff --git a/Engine/Plugins/FX/Niagara/Source/Niagara/Classes/NiagaraSystem.h b/Engine/Plugins/FX/Niagara/Source/Niagara/Classes/NiagaraSystem.h
--- a/Engine/Plugins/FX/Niagara/Source/Niagara/Classes/NiagaraSystem.h
+++ b/Engine/Plugins/FX/Niagara/Source/Niagara/Classes/NiagaraSystem.h
@@ -90,11 +90,6 @@ public:
 		original asset. */
 	FNiagaraEmitterHandle AddEmitterHandle(UNiagaraEmitter& SourceEmitter, FName EmitterName);
 
-	/** Adds a new emitter handle to this System.  The new handle will not copy the emitter and any changes made to it's
-		Instance value will modify the original asset.  This should only be used in the emitter toolkit for simulation
-		purposes. */
-	FNiagaraEmitterHandle AddEmitterHandleWithoutCopying(UNiagaraEmitter& Emitter);
-
 	/** Duplicates an existing emitter handle and adds it to the System.  The new handle will reference the same source asset,
 		but will have a copy of the duplicated Instance value. */
 	FNiagaraEmitterHandle DuplicateEmitterHandle(const FNiagaraEmitterHandle& EmitterHandleToDuplicate, FName EmitterName);
diff --git a/Engine/Plugins/FX/Niagara/Source/Niagara/Private/NiagaraEmitterHandle.cpp b/Engine/Plugins/FX/Niagara/Source/Niagara/Private/NiagaraEmitterHandle.cpp
--- a/Engine/Plugins/FX/Niagara/Source/Niagara/Private/NiagaraEmitterHandle.cpp
+++ b/Engine/Plugins/FX/Niagara/Source/Niagara/Private/NiagaraEmitterHandle.cpp
@@ -24,23 +24,6 @@ FNiagaraEmitterHandle::FNiagaraEmitterHandle() :
 {
 }
 
-FNiagaraEmitterHandle::FNiagaraEmitterHandle(UNiagaraEmitter& Emitter)
-	: Id(FGuid::NewGuid())
-	, IdName(*Id.ToString())
-	, bIsEnabled(true)
-	, Name(TEXT("Emitter"))
-#if WITH_EDITORONLY_DATA
-	, Source(&Emitter)
-	, LastMergedSource(&Emitter)
-	, bIsolated(false)
-#endif
-	, Instance(&Emitter)
-{
-#if WITH_EDITORONLY_DATA
-	Instance->GraphSource->MarkNotSynchronized(InitialNotSynchronizedReason);
-#endif
-}
-
 #if WITH_EDITORONLY_DATA
 FNiagaraEmitterHandle::FNiagaraEmitterHandle(UNiagaraEmitter& InSourceEmitter, FName InName, UNiagaraSystem& InOuterSystem)
 	: Id(FGuid::NewGuid())
@@ -52,8 +35,10 @@ FNiagaraEmitterHandle::FNiagaraEmitterHandle(UNiagaraEmitter& InSourceEmitter, F
 	, bIsolated(false)
 	, Instance(Cast<UNiagaraEmitter>(StaticDuplicateObject(Source, &InOuterSystem)))
 {
+	Instance->ClearFlags(RF_Standalone | RF_Public);
 	Instance->SetUniqueEmitterName(Name.ToString());
 	Instance->GraphSource->MarkNotSynchronized(InitialNotSynchronizedReason);
+	LastMergedSource->ClearFlags(RF_Standalone | RF_Public);
 }
 
 FNiagaraEmitterHandle::FNiagaraEmitterHandle(const FNiagaraEmitterHandle& InHandleToDuplicate, FName InDuplicateName, UNiagaraSystem& InDuplicateOwnerSystem)
diff --git a/Engine/Plugins/FX/Niagara/Source/Niagara/Private/NiagaraSystem.cpp b/Engine/Plugins/FX/Niagara/Source/Niagara/Private/NiagaraSystem.cpp
--- a/Engine/Plugins/FX/Niagara/Source/Niagara/Private/NiagaraSystem.cpp
+++ b/Engine/Plugins/FX/Niagara/Source/Niagara/Private/NiagaraSystem.cpp
@@ -624,14 +624,6 @@ FNiagaraEmitterHandle UNiagaraSystem::AddEmitterHandle(UNiagaraEmitter& SourceEm
 	return EmitterHandle;
 }
 
-FNiagaraEmitterHandle UNiagaraSystem::AddEmitterHandleWithoutCopying(UNiagaraEmitter& Emitter)
-{
-	FNiagaraEmitterHandle EmitterHandle(Emitter);
-	EmitterHandles.Add(EmitterHandle);
-	RefreshSystemParametersFromEmitter(EmitterHandle);
-	return EmitterHandle;
-}
-
 FNiagaraEmitterHandle UNiagaraSystem::DuplicateEmitterHandle(const FNiagaraEmitterHandle& EmitterHandleToDuplicate, FName EmitterName)
 {
 	FNiagaraEmitterHandle EmitterHandle(EmitterHandleToDuplicate, EmitterName, *this);
diff --git a/Engine/Plugins/FX/Niagara/Source/NiagaraEditor/Private/NiagaraEditorCommands.cpp b/Engine/Plugins/FX/Niagara/Source/NiagaraEditor/Private/NiagaraEditorCommands.cpp
--- a/Engine/Plugins/FX/Niagara/Source/NiagaraEditor/Private/NiagaraEditorCommands.cpp
+++ b/Engine/Plugins/FX/Niagara/Source/NiagaraEditor/Private/NiagaraEditorCommands.cpp
@@ -22,6 +22,7 @@ void FNiagaraEditorCommands::RegisterCommands()
 	UI_COMMAND(ToggleAutoPlay, "Auto-play", "Toggles whether or not simulations auto-play when their asset editor is opened, and when the asset is modified.", EUserInterfaceActionType::ToggleButton, FInputChord());
 	UI_COMMAND(ToggleResetSimulationOnChange, "Reset on change", "Toggles whether or not the simulation is reset whenever a change is made in the asset editor.", EUserInterfaceActionType::ToggleButton, FInputChord());
 	UI_COMMAND(ToggleResimulateOnChangeWhilePaused, "Resimulate when paused", "Toggles whether or not the simulation is rerun to the current time when making changes while paused.", EUserInterfaceActionType::ToggleButton, FInputChord());
+	UI_COMMAND(ToggleResetDependentSystems, "Reset Dependent Systems", "Toggles whether or not to reset all systems that include this emitter when it is reset by the user.", EUserInterfaceActionType::ToggleButton, FInputChord());
 	UI_COMMAND(CollapseStackToHeaders, "Collapse to Headers", "Expands all headsers and collapse all items.", EUserInterfaceActionType::Button, FInputChord(EModifierKey::Control, EKeys::O));
 
 }
diff --git a/Engine/Plugins/FX/Niagara/Source/NiagaraEditor/Private/NiagaraEditorSettings.cpp b/Engine/Plugins/FX/Niagara/Source/NiagaraEditor/Private/NiagaraEditorSettings.cpp
--- a/Engine/Plugins/FX/Niagara/Source/NiagaraEditor/Private/NiagaraEditorSettings.cpp
+++ b/Engine/Plugins/FX/Niagara/Source/NiagaraEditor/Private/NiagaraEditorSettings.cpp
@@ -9,6 +9,7 @@ UNiagaraEditorSettings::UNiagaraEditorSettings(const FObjectInitializer& ObjectI
 	bAutoPlay = true;
 	bResetSimulationOnChange = true;
 	bResimulateOnChangeWhilePaused = true;
+	bResetDependentSystemsWhenEditingEmitters = false;
 }
 
 bool UNiagaraEditorSettings::GetAutoCompile() const
@@ -67,6 +68,20 @@ void UNiagaraEditorSettings::SetResimulateOnChangeWhilePaused(bool bInResimulate
 	}
 }
 
+bool UNiagaraEditorSettings::GetResetDependentSystemsWhenEditingEmitters() const
+{
+	return bResetDependentSystemsWhenEditingEmitters;
+}
+
+void UNiagaraEditorSettings::SetResetDependentSystemsWhenEditingEmitters(bool bInResetDependentSystemsWhenEditingEmitters)
+{
+	if (bResetDependentSystemsWhenEditingEmitters != bInResetDependentSystemsWhenEditingEmitters)
+	{
+		bResetDependentSystemsWhenEditingEmitters = bInResetDependentSystemsWhenEditingEmitters;
+		SaveConfig();
+	}
+}
+
 FNiagaraNewAssetDialogConfig UNiagaraEditorSettings::GetNewAssetDailogConfig(FName InDialogConfigKey) const
 {
 	const FNiagaraNewAssetDialogConfig* Config = NewAssetDialogConfigMap.Find(InDialogConfigKey);
diff --git a/Engine/Plugins/FX/Niagara/Source/NiagaraEditor/Private/NiagaraEditorUtilities.cpp b/Engine/Plugins/FX/Niagara/Source/NiagaraEditor/Private/NiagaraEditorUtilities.cpp
--- a/Engine/Plugins/FX/Niagara/Source/NiagaraEditor/Private/NiagaraEditorUtilities.cpp
+++ b/Engine/Plugins/FX/Niagara/Source/NiagaraEditor/Private/NiagaraEditorUtilities.cpp
@@ -898,4 +898,71 @@ FText FNiagaraEditorUtilities::FormatScriptAssetDescription(FText Description, F
 		: FText::Format(LOCTEXT("ScriptAssetDescriptionFormat", "Description: {1}\nPath: {0}"), FText::FromName(Path), Description);
 }
 
+void FNiagaraEditorUtilities::ResetSystemsThatReferenceSystemViewModel(const FNiagaraSystemViewModel& ReferencedSystemViewModel)
+{
+	checkf(&ReferencedSystemViewModel, TEXT("ResetSystemsThatReferenceSystemViewModel() called on destroyed SystemViewModel."));
+	TArray<TSharedPtr<FNiagaraSystemViewModel>> ComponentSystemViewModels;
+	TArray<UNiagaraComponent*> ReferencingComponents = GetComponentsThatReferenceSystemViewModel(ReferencedSystemViewModel);
+	for (auto Component : ReferencingComponents)
+	{
+		ComponentSystemViewModels.Reset();
+		FNiagaraSystemViewModel::GetAllViewModelsForObject(Component->GetAsset(), ComponentSystemViewModels);
+		if (ComponentSystemViewModels.Num() > 0)
+		{
+			//The component has a viewmodel, call ResetSystem() on the viewmodel 
+			for (auto SystemViewModel : ComponentSystemViewModels)
+			{
+				if (SystemViewModel.IsValid() && SystemViewModel.Get() != &ReferencedSystemViewModel)
+				{
+					SystemViewModel->ResetSystem(FNiagaraSystemViewModel::ETimeResetMode::AllowResetTime, FNiagaraSystemViewModel::EMultiResetMode::ResetThisInstance, FNiagaraSystemViewModel::EReinitMode::ResetSystem);
+				}
+			}
+		}
+		else
+		{
+			//The component does not have a viewmodel, call ResetSystem() on the component
+			Component->ResetSystem();
+		}
+	}
+}
+
+TArray<UNiagaraComponent*> FNiagaraEditorUtilities::GetComponentsThatReferenceSystem(const UNiagaraSystem& ReferencedSystem)
+{
+	check(&ReferencedSystem);
+	TArray<UNiagaraComponent*> ReferencingComponents;
+	for (TObjectIterator<UNiagaraComponent> ComponentIt; ComponentIt; ++ComponentIt)
+	{
+		UNiagaraComponent* Component = *ComponentIt;
+		if (Component && Component->GetAsset())
+		{
+			if (Component->GetAsset() == &ReferencedSystem)
+			{
+				ReferencingComponents.Add(Component);
+			}
+		}
+	}
+	return ReferencingComponents;
+}
+
+TArray<UNiagaraComponent*> FNiagaraEditorUtilities::GetComponentsThatReferenceSystemViewModel(const FNiagaraSystemViewModel& ReferencedSystemViewModel)
+{
+	check(&ReferencedSystemViewModel);
+	TArray<UNiagaraComponent*> ReferencingComponents;
+	for (TObjectIterator<UNiagaraComponent> ComponentIt; ComponentIt; ++ComponentIt)
+	{
+		UNiagaraComponent* Component = *ComponentIt;
+		if (Component && Component->GetAsset())
+		{
+			for (auto EmitterHandle : ReferencedSystemViewModel.GetSystem().GetEmitterHandles())
+			{
+				if (Component->GetAsset()->UsesEmitter(EmitterHandle.GetSource()))
+				{
+					ReferencingComponents.Add(Component);
+				}
+			}
+		}
+	}
+	return ReferencingComponents;
+}
+
 #undef LOCTEXT_NAMESPACE
diff --git a/Engine/Plugins/FX/Niagara/Source/NiagaraEditor/Private/Tests/NiagaraEditorTestUtilities.cpp b/Engine/Plugins/FX/Niagara/Source/NiagaraEditor/Private/Tests/NiagaraEditorTestUtilities.cpp
--- a/Engine/Plugins/FX/Niagara/Source/NiagaraEditor/Private/Tests/NiagaraEditorTestUtilities.cpp
+++ b/Engine/Plugins/FX/Niagara/Source/NiagaraEditor/Private/Tests/NiagaraEditorTestUtilities.cpp
@@ -36,7 +36,7 @@ UNiagaraSystem* FNiagaraEditorTestUtilities::CreateTestSystemForEmitter(UNiagara
 {
 	UNiagaraSystem* System = NewObject<UNiagaraSystem>((UObject*)GetTransientPackage(), NAME_None, RF_Transient);
 	UNiagaraSystemFactoryNew::InitializeSystem(System, false);
-	System->AddEmitterHandleWithoutCopying(*Emitter);
+	System->AddEmitterHandle(*Emitter, Emitter->GetFName());
 	return System;
 }
 
diff --git a/Engine/Plugins/FX/Niagara/Source/NiagaraEditor/Private/Toolkits/NiagaraSystemToolkit.cpp b/Engine/Plugins/FX/Niagara/Source/NiagaraEditor/Private/Toolkits/NiagaraSystemToolkit.cpp
--- a/Engine/Plugins/FX/Niagara/Source/NiagaraEditor/Private/Toolkits/NiagaraSystemToolkit.cpp
+++ b/Engine/Plugins/FX/Niagara/Source/NiagaraEditor/Private/Toolkits/NiagaraSystemToolkit.cpp
@@ -393,7 +393,7 @@ void FNiagaraSystemToolkit::InitializeWithEmitter(const EToolkitMode::Type Mode,
 
 	// Before copying the emitter prepare the rapid iteration parameters so that the post compile prepare doesn't
 	// cause the change ids to become out of sync.
-	FString EmitterName = "Emitter";
+	FString EmitterName = Emitter->GetUniqueEmitterName();
 	TArray<UNiagaraScript*> Scripts;
 	TMap<UNiagaraScript*, UNiagaraScript*> ScriptDependencyMap;
 	TMap<UNiagaraScript*, FString> ScriptToEmitterNameMap;
@@ -426,12 +426,7 @@ void FNiagaraSystemToolkit::InitializeWithEmitter(const EToolkitMode::Type Mode,
 
 	ResetLoaders(GetTransientPackage()); // Make sure that we're not going to get invalid version number linkers into the package we are going into. 
 	GetTransientPackage()->LinkerCustomVersion.Empty();
-
-	UNiagaraEmitter* EditableEmitter = (UNiagaraEmitter*)StaticDuplicateObject(Emitter, GetTransientPackage(), NAME_None, ~RF_Standalone, UNiagaraEmitter::StaticClass());
 	
-	// We set this to the copy's change id here instead of the original emitter's change id because the copy's change id may have been
-	// updated from the original as part of post load and we use this id to detect if the editable emitter has been changed.
-	LastSyncedEmitterChangeId = EditableEmitter->GetChangeId();
 	bEmitterThumbnailUpdated = false;
 
 	FNiagaraSystemViewModelOptions SystemOptions;
@@ -440,7 +435,10 @@ void FNiagaraSystemToolkit::InitializeWithEmitter(const EToolkitMode::Type Mode,
 
 	SystemViewModel = MakeShareable(new FNiagaraSystemViewModel(*System, SystemOptions));
 	SystemViewModel->SetToolkitCommands(GetToolkitCommands());
-	SystemViewModel->AddEmitter(*EditableEmitter);
+	SystemViewModel->AddEmitter(*Emitter);
+	// We set this to the copy's change id here instead of the original emitter's change id because the copy's change id may have been
+	// updated from the original as part of post load and we use this id to detect if the editable emitter has been changed.
+	LastSyncedEmitterChangeId = SystemViewModel->GetEmitterHandleViewModels()[0]->GetEmitterViewModel()->GetEmitter()->GetChangeId();
 	SystemViewModel->GetSystemScriptViewModel()->RebuildEmitterNodes();
 	SystemToolkitMode = ESystemToolkitMode::Emitter;
 
@@ -901,6 +899,17 @@ void FNiagaraSystemToolkit::SetupCommands()
 		}),
 		FCanExecuteAction(),
 		FIsActionChecked::CreateLambda([]() { return GetDefault<UNiagaraEditorSettings>()->GetResimulateOnChangeWhilePaused(); }));
+
+	GetToolkitCommands()->MapAction(
+		FNiagaraEditorCommands::Get().ToggleResetDependentSystems,
+		FExecuteAction::CreateLambda([]()
+	{
+		UNiagaraEditorSettings* Settings = GetMutableDefault<UNiagaraEditorSettings>();
+		Settings->SetResetDependentSystemsWhenEditingEmitters(!Settings->GetResetDependentSystemsWhenEditingEmitters());
+	}),
+		FCanExecuteAction(),
+		FIsActionChecked::CreateLambda([]() { return GetDefault<UNiagaraEditorSettings>()->GetResetDependentSystemsWhenEditingEmitters(); }),
+		FIsActionButtonVisible::CreateLambda([this]() { return SystemViewModel->GetEditMode() == ENiagaraSystemViewModelEditMode::EmitterAsset; }));
 }
 
 void FNiagaraSystemToolkit::OnSaveThumbnailImage()
@@ -929,7 +938,7 @@ void FNiagaraSystemToolkit::OnThumbnailCaptured(UTexture2D* Thumbnail)
 
 void FNiagaraSystemToolkit::ResetSimulation()
 {
-	SystemViewModel->ResetSystem();
+	SystemViewModel->ResetSystem(FNiagaraSystemViewModel::ETimeResetMode::AllowResetTime, FNiagaraSystemViewModel::EMultiResetMode::AllowResetAllInstances, FNiagaraSystemViewModel::EReinitMode::ResetSystem);
 }
 
 void FNiagaraSystemToolkit::OnVMSystemCompiled()
@@ -947,6 +956,7 @@ void FNiagaraSystemToolkit::ExtendToolbar()
 			MenuBuilder.AddMenuEntry(FNiagaraEditorCommands::Get().ToggleAutoPlay);
 			MenuBuilder.AddMenuEntry(FNiagaraEditorCommands::Get().ToggleResetSimulationOnChange);
 			MenuBuilder.AddMenuEntry(FNiagaraEditorCommands::Get().ToggleResimulateOnChangeWhilePaused);
+			MenuBuilder.AddMenuEntry(FNiagaraEditorCommands::Get().ToggleResetDependentSystems);
 			return MenuBuilder.MakeWidget();
 		}
 
@@ -1239,6 +1249,10 @@ void FNiagaraSystemToolkit::UpdateOriginalEmitter()
 		// overwrite the original script in place by constructing a new one with the same name
 		Emitter = (UNiagaraEmitter*)StaticDuplicateObject(EditableEmitter, Emitter->GetOuter(),
 			Emitter->GetFName(), RF_AllFlags, Emitter->GetClass());
+
+		// Restore RF_Standalone and RF_Public on the original emitter, as it had been removed from the preview emitter so that it could be GC'd.
+		Emitter->SetFlags(RF_Standalone | RF_Public);
+
 		Emitter->PostEditChange();
 
 		TArray<UNiagaraScript*> EmitterScripts;
@@ -1265,9 +1279,6 @@ void FNiagaraSystemToolkit::UpdateOriginalEmitter()
 		LastSyncedEmitterChangeId = EditableEmitter->GetChangeId();
 		bEmitterThumbnailUpdated = false;
 
-		// Restore RF_Standalone on the original emitter, as it had been removed from the preview emitter so that it could be GC'd.
-		Emitter->SetFlags(RF_Standalone);
-
 		TArray<UNiagaraEmitter*> AffectedEmitters;
 		AffectedEmitters.Add(Emitter);
 		UpdateExistingEmitters();
@@ -1288,7 +1299,9 @@ void FNiagaraSystemToolkit::UpdateExistingEmitters()
 	for (TObjectIterator<UNiagaraSystem> SystemIterator; SystemIterator; ++SystemIterator)
 	{
 		UNiagaraSystem* LoadedSystem = *SystemIterator;
-		if (LoadedSystem->IsPendingKill() == false && 
+
+		if (LoadedSystem != System &&
+			LoadedSystem->IsPendingKill() == false && 
 			LoadedSystem->HasAnyFlags(RF_ClassDefaultObject) == false &&
 			LoadedSystem->ReferencesSourceEmitter(*Emitter))
 		{
diff --git a/Engine/Plugins/FX/Niagara/Source/NiagaraEditor/Private/ViewModels/NiagaraSystemViewModel.cpp b/Engine/Plugins/FX/Niagara/Source/NiagaraEditor/Private/ViewModels/NiagaraSystemViewModel.cpp
--- a/Engine/Plugins/FX/Niagara/Source/NiagaraEditor/Private/ViewModels/NiagaraSystemViewModel.cpp
+++ b/Engine/Plugins/FX/Niagara/Source/NiagaraEditor/Private/ViewModels/NiagaraSystemViewModel.cpp
@@ -307,12 +307,8 @@ void FNiagaraSystemViewModel::AddEmitter(UNiagaraEmitter& Emitter)
 	if (EditMode == ENiagaraSystemViewModelEditMode::SystemAsset)
 	{
 		System.Modify();
-		EmitterHandle = System.AddEmitterHandle(Emitter, FNiagaraUtilities::GetUniqueName(Emitter.GetFName(), EmitterHandleNames));
-	}
-	else
-	{
-		EmitterHandle = System.AddEmitterHandleWithoutCopying(Emitter);
 	}
+	EmitterHandle = System.AddEmitterHandle(Emitter, FNiagaraUtilities::GetUniqueName(Emitter.GetFName(), EmitterHandleNames));
 
 	check(SystemScriptViewModel.IsValid());
 	SystemScriptViewModel->RebuildEmitterNodes();
@@ -547,7 +543,7 @@ void FNiagaraSystemViewModel::Tick(float DeltaTime)
 
 			if (bResetRequestPending)
 			{
-				ResetSystem();
+				ResetSystem(ETimeResetMode::AllowResetTime, EMultiResetMode::ResetThisInstance, EReinitMode::ReinitializeSystem);
 			}
 		}
 	}
@@ -647,7 +643,7 @@ void FNiagaraSystemViewModel::SetupPreviewComponentAndInstance()
 
 void FNiagaraSystemViewModel::RefreshAll()
 {
-	ReInitializeSystemInstances();
+	ResetSystem(ETimeResetMode::AllowResetTime, EMultiResetMode::ResetThisInstance, EReinitMode::ReinitializeSystem);
 	RefreshEmitterHandleViewModels();
 	RefreshSequencerTracks();
 	ResetCurveData();
@@ -672,7 +668,7 @@ void FNiagaraSystemViewModel::NotifyDataObjectChanged(UObject* ChangedObject)
 		}
 	}
 
-	ReInitializeSystemInstances();
+	ResetSystem(ETimeResetMode::AllowResetTime, EMultiResetMode::ResetThisInstance, EReinitMode::ReinitializeSystem);
 }
 
 void FNiagaraSystemViewModel::IsolateEmitters(TArray<TSharedRef<FNiagaraEmitterHandleViewModel>> EmitterHandlesToIsolate)
@@ -924,12 +920,12 @@ void FNiagaraSystemViewModel::SetupSequencer()
 
 void FNiagaraSystemViewModel::ResetSystem()
 {
-	ResetSystemInternal(true);
+	ResetSystem(ETimeResetMode::AllowResetTime, EMultiResetMode::ResetThisInstance, EReinitMode::ResetSystem);
 }
 
-void FNiagaraSystemViewModel::ResetSystemInternal(bool bCanResetTime)
+void FNiagaraSystemViewModel::ResetSystem(ETimeResetMode TimeResetMode, EMultiResetMode MultiResetMode, EReinitMode ReinitMode)
 {
-	bool bResetAge = bCanResetTime && (Sequencer->GetPlaybackStatus() == EMovieScenePlayerStatus::Playing || EditorSettings->GetResimulateOnChangeWhilePaused() == false);
+	bool bResetAge = TimeResetMode == ETimeResetMode::AllowResetTime && (Sequencer->GetPlaybackStatus() == EMovieScenePlayerStatus::Playing || EditorSettings->GetResimulateOnChangeWhilePaused() == false);
 	if (bResetAge)
 	{
 		TGuardValue<bool> Guard(bSettingSequencerTimeDirectly, true);
@@ -944,11 +940,11 @@ void FNiagaraSystemViewModel::ResetSystemInternal(bool bCanResetTime)
 			Sequencer->SetGlobalTime(0);
 		}
 	}
-	
-	for (TObjectIterator<UNiagaraComponent> ComponentIt; ComponentIt; ++ComponentIt)
+
+	TArray<UNiagaraComponent*> ReferencingComponents = FNiagaraEditorUtilities::GetComponentsThatReferenceSystem(System);
+	for (auto Component : ReferencingComponents)
 	{
-		UNiagaraComponent* Component = *ComponentIt;
-		if (Component->GetAsset() == &System)
+		if (ReinitMode == EReinitMode::ResetSystem)
 		{
 			Component->ResetSystem();
 			if (bResetAge && Component->GetAgeUpdateMode() == ENiagaraAgeUpdateMode::DesiredAge)
@@ -956,6 +952,16 @@ void FNiagaraSystemViewModel::ResetSystemInternal(bool bCanResetTime)
 				Component->SetDesiredAge(0);
 			}
 		}
+
+		else if (ReinitMode == EReinitMode::ReinitializeSystem)
+		{
+			Component->ReinitializeSystem();
+		}
+	}
+
+	if (EditMode == ENiagaraSystemViewModelEditMode::EmitterAsset && MultiResetMode == EMultiResetMode::AllowResetAllInstances && EditorSettings->GetResetDependentSystemsWhenEditingEmitters())
+	{
+		FNiagaraEditorUtilities::ResetSystemsThatReferenceSystemViewModel(*this);
 	}
 
 	FEditorSupportDelegates::RedrawAllViewports.Broadcast();
@@ -969,36 +975,13 @@ void FNiagaraSystemViewModel::RequestResetSystem()
 
 void FNiagaraSystemViewModel::KillSystemInstances()
 {
-	for (TObjectIterator<UNiagaraComponent> ComponentIt; ComponentIt; ++ComponentIt)
+	TArray<UNiagaraComponent*> ReferencingComponents = FNiagaraEditorUtilities::GetComponentsThatReferenceSystem(System);
+	for (auto Component : ReferencingComponents)
 	{
-		UNiagaraComponent* Component = *ComponentIt;
-		if (Component->GetAsset() == &System)
-		{
-			Component->DestroyInstance();
-		}
+		Component->DestroyInstance();
 	}
 }
 
-void FNiagaraSystemViewModel::ReInitializeSystemInstances()
-{
-	if (Sequencer.IsValid() && Sequencer->GetPlaybackStatus() == EMovieScenePlayerStatus::Playing)
-	{
-		Sequencer->SetPlaybackStatus(EMovieScenePlayerStatus::Stopped);
-		Sequencer->SetGlobalTime(0);
-		Sequencer->SetPlaybackStatus(EMovieScenePlayerStatus::Playing);
-	}
-
-	for (TObjectIterator<UNiagaraComponent> ComponentIt; ComponentIt; ++ComponentIt)
-	{
-		UNiagaraComponent* Component = *ComponentIt;
-		if (Component->GetAsset() == &System)
-		{
-			Component->ReinitializeSystem();
-		}
-	}
-	FEditorSupportDelegates::RedrawAllViewports.Broadcast();
-}
-
 struct FNiagaraSystemCurveData
 {
 	FRichCurve* Curve;
@@ -1183,7 +1166,7 @@ void FNiagaraSystemViewModel::EmitterHandlePropertyChanged(TSharedRef<FNiagaraEm
 		}
 		Sequencer->NotifyMovieSceneDataChanged(EMovieSceneDataChangeType::MovieSceneStructureItemsChanged);
 	}
-	ReInitializeSystemInstances();
+	ResetSystem(ETimeResetMode::AllowResetTime, EMultiResetMode::ResetThisInstance, EReinitMode::ReinitializeSystem);
 }
 
 void FNiagaraSystemViewModel::EmitterHandleNameChanged(TSharedRef<FNiagaraEmitterHandleViewModel> EmitterHandleViewModel)
@@ -1193,7 +1176,7 @@ void FNiagaraSystemViewModel::EmitterHandleNameChanged(TSharedRef<FNiagaraEmitte
 
 void FNiagaraSystemViewModel::EmitterPropertyChanged(TSharedRef<FNiagaraEmitterHandleViewModel> EmitterHandleViewModel)
 {
-	ReInitializeSystemInstances();
+	ResetSystem(ETimeResetMode::AllowResetTime, EMultiResetMode::ResetThisInstance, EReinitMode::ReinitializeSystem);
 }
 
 void FNiagaraSystemViewModel::ScriptCompiled()
@@ -1240,6 +1223,8 @@ void FNiagaraSystemViewModel::UpdateSimulationFromParameterChange()
 {
 	if (EditorSettings->GetResetSimulationOnChange())
 	{
+		/* Calling RequestResetSystem here avoids reentrancy into ResetSystem() when we edit the system parameter store on 
+		** UNiagaraComponent::Activate() as we always call PrepareRapidIterationParameters().  */
 		RequestResetSystem();
 	}
 	else
@@ -1456,8 +1441,14 @@ void FNiagaraSystemViewModel::SequencerTimeChanged()
 			if (bResetSystemInstance)
 			{
 				// We don't want to reset the current time if we're scrubbing.
-				bool bCanResetTime = CurrentStatus == EMovieScenePlayerStatus::Playing;
-				ResetSystemInternal(bCanResetTime);
+				if (CurrentStatus == EMovieScenePlayerStatus::Playing)
+				{
+					ResetSystem(ETimeResetMode::AllowResetTime, EMultiResetMode::ResetThisInstance, EReinitMode::ResetSystem);
+				}
+				else
+				{
+					ResetSystem(ETimeResetMode::KeepCurrentTime, EMultiResetMode::ResetThisInstance, EReinitMode::ResetSystem);
+				}
 			}
 		}
 	}
diff --git a/Engine/Plugins/FX/Niagara/Source/NiagaraEditor/Public/NiagaraEditorCommands.h b/Engine/Plugins/FX/Niagara/Source/NiagaraEditor/Public/NiagaraEditorCommands.h
--- a/Engine/Plugins/FX/Niagara/Source/NiagaraEditor/Public/NiagaraEditorCommands.h
+++ b/Engine/Plugins/FX/Niagara/Source/NiagaraEditor/Public/NiagaraEditorCommands.h
@@ -45,6 +45,7 @@ public:
 	TSharedPtr<FUICommandInfo> ToggleAutoPlay;
 	TSharedPtr<FUICommandInfo> ToggleResetSimulationOnChange;
 	TSharedPtr<FUICommandInfo> ToggleResimulateOnChangeWhilePaused;
+	TSharedPtr<FUICommandInfo> ToggleResetDependentSystems;
 
 	TSharedPtr<FUICommandInfo> CollapseStackToHeaders;
 };
\ No newline at end of file
diff --git a/Engine/Plugins/FX/Niagara/Source/NiagaraEditor/Public/NiagaraEditorSettings.h b/Engine/Plugins/FX/Niagara/Source/NiagaraEditor/Public/NiagaraEditorSettings.h
--- a/Engine/Plugins/FX/Niagara/Source/NiagaraEditor/Public/NiagaraEditorSettings.h
+++ b/Engine/Plugins/FX/Niagara/Source/NiagaraEditor/Public/NiagaraEditorSettings.h
@@ -77,6 +77,12 @@ public:
 	/** Sets whether or not to rerun the simulation to the current time when making modifications while paused. */
 	void SetResimulateOnChangeWhilePaused(bool bInResimulateOnChangeWhilePaused);
 
+	/** Gets whether or not to reset all components that include the system that is currently being reset */
+	bool GetResetDependentSystemsWhenEditingEmitters() const;
+
+	/** Sets whether or not to reset all components that include the system that is currently being reset */
+	void SetResetDependentSystemsWhenEditingEmitters(bool bInResetDependentSystemsWhenEditingEmitters);
+
 	FNiagaraNewAssetDialogConfig GetNewAssetDailogConfig(FName InDialogConfigKey) const;
 
 	void SetNewAssetDialogConfig(FName InDialogConfigKey, const FNiagaraNewAssetDialogConfig& InNewAssetDialogConfig);
@@ -114,6 +120,10 @@ private:
 	UPROPERTY(config, EditAnywhere, Category = SimulationOptions)
 	bool bResimulateOnChangeWhilePaused;
 
+	/** Whether or not to reset all components that include the system currently being reset. */
+	UPROPERTY(config, EditAnywhere, Category = SimulationOptions)
+		bool bResetDependentSystemsWhenEditingEmitters;
+
 	UPROPERTY(config)
 	TMap<FName, FNiagaraNewAssetDialogConfig> NewAssetDialogConfigMap;
 
diff --git a/Engine/Plugins/FX/Niagara/Source/NiagaraEditor/Public/NiagaraEditorUtilities.h b/Engine/Plugins/FX/Niagara/Source/NiagaraEditor/Public/NiagaraEditorUtilities.h
--- a/Engine/Plugins/FX/Niagara/Source/NiagaraEditor/Public/NiagaraEditorUtilities.h
+++ b/Engine/Plugins/FX/Niagara/Source/NiagaraEditor/Public/NiagaraEditorUtilities.h
@@ -14,6 +14,7 @@ struct FNiagaraVariable;
 struct FNiagaraTypeDefinition;
 class UNiagaraGraph;
 class UNiagaraSystem;
+class FNiagaraSystemViewModel;
 struct FNiagaraEmitterHandle;
 class UNiagaraEmitter;
 class UNiagaraScript;
@@ -173,4 +174,10 @@ namespace FNiagaraEditorUtilities
 	const FNiagaraEmitterHandle* GetEmitterHandleForEmitter(UNiagaraSystem& System, UNiagaraEmitter& Emitter);
 
 	NIAGARAEDITOR_API FText FormatScriptAssetDescription(FText Description, FName Path);
+
+	void ResetSystemsThatReferenceSystemViewModel(const FNiagaraSystemViewModel& ReferencedSystemViewModel);
+
+	TArray<UNiagaraComponent*> GetComponentsThatReferenceSystem(const UNiagaraSystem& ReferencedSystem);
+
+	TArray<UNiagaraComponent*> GetComponentsThatReferenceSystemViewModel(const FNiagaraSystemViewModel& ReferencedSystemViewModel);
 };
diff --git a/Engine/Plugins/FX/Niagara/Source/NiagaraEditor/Public/ViewModels/NiagaraSystemViewModel.h b/Engine/Plugins/FX/Niagara/Source/NiagaraEditor/Public/ViewModels/NiagaraSystemViewModel.h
--- a/Engine/Plugins/FX/Niagara/Source/NiagaraEditor/Public/ViewModels/NiagaraSystemViewModel.h
+++ b/Engine/Plugins/FX/Niagara/Source/NiagaraEditor/Public/ViewModels/NiagaraSystemViewModel.h
@@ -95,6 +95,34 @@ public:
 
 	DECLARE_MULTICAST_DELEGATE(FOnPinnedCurvesChanged);
 
+public:
+	/** Defines different multi-system reset modes for this system view model */
+	enum class EMultiResetMode
+	{
+		/** Reset this instance. Then, if reset dependent systems is enabled in the editor through NiagaraEditorCommands, find all components that share this system and reset those components' system instances. */
+		AllowResetAllInstances,
+		/** Reset just this instance. */
+		ResetThisInstance,
+	};
+
+	/** Defines different time reset modes for this system view model */
+	enum class ETimeResetMode
+	{
+		/** If the current sequencer state and user settings allow, reset this system's time. */
+		AllowResetTime,
+		/** Keep this system's current time. */
+		KeepCurrentTime,
+	};
+
+	/** Defines different initialization modes when resetting for this system view model */
+	enum class EReinitMode
+	{
+		/** Reinitialize this system (pull in all changes) */
+		ReinitializeSystem,
+		/** Reset this system (do not pull in changes) */
+		ResetSystem,
+	};
+
 public:
 	/** Creates a new view model with the supplied System and System instance. */
 	FNiagaraSystemViewModel(UNiagaraSystem& InSystem, FNiagaraSystemViewModelOptions InOptions);
@@ -120,7 +148,7 @@ public:
 	FNiagaraCurveOwner& GetCurveOwner();
 
 	/** Get access to the underlying system*/
-	UNiagaraSystem& GetSystem() { return System; }
+	UNiagaraSystem& GetSystem() const { return System; }
 
 	/** Gets whether or not emitters can be added from the timeline. */
 	bool GetCanModifyEmittersFromTimeline() const;
@@ -170,14 +198,21 @@ public:
 	virtual bool IsTickable() const override { return true; }
 	virtual TStatId GetStatId() const override;
 
-	/** Resets the System instance to initial conditions. */
+	/** Resets the System instance to initial conditions. Tries to resets system simulation time. Does not reset all systems that share its emitters.
+	 * Does not reinitialize the system to pull in changes. Calls into overloaded ResetSystem(). */
 	void ResetSystem();
 
-	/** Resets the system instance on the next frame. */
+	/** Resets the System instance to initial conditions on the next tick. Tries to resets system simulation time. Does not reset all systems that share its emitters.
+	 * Does not reinitialize the system to pull in changes. Calls into overloaded ResetSystem(). */
 	void RequestResetSystem();
 
-	/** Reinitializes the System instance to initial conditions - rebuilds all data sets and resets data interfaces. */
-	void ReInitializeSystemInstances();
+	/** Resets the system instance to initial conditions. Optionally resets system simulation time. Optionally resets all systems that share its emitters. 
+	 * Optionally reinitializes the system to pull in changes.
+	 * @param TimeResetMode Defines whether the system being reset should try to reset its time to 0 or keep its current time.
+	 * @param MultiResetMode Defines whether the system being reset should try to reset all other systems with which it shares emitters along with resetting itself.
+	 * @param ReinitMode Defines whether the system should reinitialize and pull in changes or reset and keep its current state.
+	 */
+	void ResetSystem(ETimeResetMode TimeResetMode, EMultiResetMode MultiResetMode, EReinitMode ReinitMode);
 
 	/** Compiles the spawn and update scripts. */
 	void CompileSystem(bool bForce);
@@ -255,9 +290,6 @@ public:
 	const TArray<FNiagaraStackModuleData>& GetStackModuleDataForEmitter(TSharedRef<FNiagaraEmitterViewModel> EmitterViewModel);
 
 private:
-	/** Reset the current simulation for the system */
-	void ResetSystemInternal(bool bCanResetTime);
-
 	/** Sets up the preview component and System instance. */
 	void SetupPreviewComponentAndInstance();
 
@@ -473,7 +505,7 @@ private:
 	/** A handle to the on parameter changed delegate for the user parameter store. */
 	FDelegateHandle UserParameterStoreChangedHandle;
 
-	/** A flag indicating that a reset has been request on the next frame */
+	/** A flag indicating that a reset has been request on the next tick */
 	bool bResetRequestPending;
 
 	/** The system toolkit commands. */
