Niagara - Copy change #5271476 from dev-destruction

Niagara - Prevent a crash when deleting the last renderer from a gpu emitter.

diff --git a/Engine/Plugins/FX/Niagara/Source/Niagara/Private/NiagaraEmitterInstance.cpp b/Engine/Plugins/FX/Niagara/Source/Niagara/Private/NiagaraEmitterInstance.cpp
--- a/Engine/Plugins/FX/Niagara/Source/Niagara/Private/NiagaraEmitterInstance.cpp
+++ b/Engine/Plugins/FX/Niagara/Source/Niagara/Private/NiagaraEmitterInstance.cpp
@@ -1040,7 +1040,15 @@ void FNiagaraEmitterInstance::Tick(float DeltaSeconds)
 		GPUExecContext->RTUpdateScript = CachedEmitter->UpdateScriptProps.Script->GetRenderThreadScript();
 		GPUExecContext->SpawnRateInstances = SpawnTotal;
 		GPUExecContext->EventSpawnTotal = EventSpawnTotal;
-		GPUExecContext->NumIndicesPerInstance = CachedEmitter->GetRenderers()[0]->GetNumIndicesPerInstance();
+
+		if (CachedEmitter->GetRenderers().Num() > 0 && CachedEmitter->GetRenderers()[0] != nullptr)
+		{
+			GPUExecContext->NumIndicesPerInstance = CachedEmitter->GetRenderers()[0]->GetNumIndicesPerInstance();
+		}
+		else
+		{
+			GPUExecContext->NumIndicesPerInstance = 0;
+		}
 		
 #if WITH_EDITORONLY_DATA
 		if (ParentSystemInstance->ShouldCaptureThisFrame())
