- Add missing call to OnSystemInstanceChangedDelegate in OnUnregister
- Fixes a crash when using "Asset Actions > Reload" on an open Niagara Asset

diff --git a/Engine/Plugins/FX/Niagara/Source/Niagara/Private/NiagaraComponent.cpp b/Engine/Plugins/FX/Niagara/Source/Niagara/Private/NiagaraComponent.cpp
--- a/Engine/Plugins/FX/Niagara/Source/Niagara/Private/NiagaraComponent.cpp
+++ b/Engine/Plugins/FX/Niagara/Source/Niagara/Private/NiagaraComponent.cpp
@@ -756,6 +756,9 @@ void UNiagaraComponent::OnUnregister()
 	{
 		SystemInstance->Deactivate(true);
 		SystemInstance = nullptr;
+#if WITH_EDITORONLY_DATA
+		OnSystemInstanceChangedDelegate.Broadcast();
+#endif
 	}
 }
 
