Niagara Ribbon Tesselation.
Added new configuration variables :
- "Niagara.Ribbon.CullTwistedStrips" Enabled fallback to fix inverted strips when ribbons are aligned to the cam
- "Niagara.Ribbon.Tessellation.MinAngle" Minimum accepted ribbon angle before interpolating.
- "Niagara.Ribbon.Tessellation.MaxInterp" Max interpolation count
- "Niagara.Ribbon.Tessellation.MaxErrorScreenPercentage" Max curvature screen percent error accepted before tesselating.
- "Niagara.Ribbon.Tessellation.MinAbsoluteError" Min curvature absolute error that will be accepted. To prevent over tessalation.
- "Niagara.Ribbon.MinSegmentLength" min absolute distance between two segments in a ribbon

diff --git a/Engine/Plugins/FX/Niagara/Source/Niagara/Private/NiagaraRendererRibbons.cpp b/Engine/Plugins/FX/Niagara/Source/Niagara/Private/NiagaraRendererRibbons.cpp
--- a/Engine/Plugins/FX/Niagara/Source/Niagara/Private/NiagaraRendererRibbons.cpp
+++ b/Engine/Plugins/FX/Niagara/Source/Niagara/Private/NiagaraRendererRibbons.cpp
@@ -11,12 +11,85 @@ DECLARE_CYCLE_STAT(TEXT("Render Ribbons [RT]"), STAT_NiagaraRenderRibbons, STATG
 
 DECLARE_CYCLE_STAT(TEXT("Genereate GPU Buffers"), STAT_NiagaraGenRibbonGpuBuffers, STATGROUP_Niagara);
 
+int32 GNiagaraRibbonCullTwistedStrips = 1;
+static FAutoConsoleVariableRef CVarNiagaraRibbonCullTwistedStrips(
+	TEXT("Niagara.Ribbon.CullTwistedStrips"),
+	GNiagaraRibbonCullTwistedStrips,
+	TEXT("Apply backface culling to remove twisted strips. Double indices. (default=1)"),
+	ECVF_Default
+);
+
+float GNiagaraRibbonTessellationAngle = 15.f * (2.f * PI) / 360.f; // Every 15 degrees
+static FAutoConsoleVariableRef CVarNiagaraRibbonTessellationAngle(
+	TEXT("Niagara.Ribbon.Tessellation.MinAngle"),
+	GNiagaraRibbonTessellationAngle,
+	TEXT("Ribbon segment angle to tesselate in radian. \n")
+	TEXT("Set 0 to use fixed tesselation. (default=15 degrees)"),
+	ECVF_Scalability
+);
+
+int32 GNiagaraRibbonMaxTessellation = 64;
+static FAutoConsoleVariableRef CVarNiagaraRibbonMaxTessellation(
+	TEXT("Niagara.Ribbon.Tessellation.MaxInterp"),
+	GNiagaraRibbonMaxTessellation,
+	TEXT("When TessellationAngle is > 0, this is the maximum tesselation factor. \n")
+	TEXT("Higher values allow more evenly divided tesselation. \n")
+	TEXT("When TessellationAngle is 0, this is the actually tesselation factor (default=60)."),
+	ECVF_Scalability
+);
+
+float GNiagaraRibbonTessellationScreenPercentage = 0.002f; 
+static FAutoConsoleVariableRef CVarNiagaraRibbonTessellationScreenPercentage(
+	TEXT("Niagara.Ribbon.Tessellation.MaxErrorScreenPercentage"),
+	GNiagaraRibbonTessellationScreenPercentage,
+	TEXT("Screen percentage used to compute the tessellation factor. \n")
+	TEXT("Smaller values will generate more tessellation, up to max tesselltion. (default=0.002)"),
+	ECVF_Scalability
+);
+
+float GNiagaraRibbonTessellationMinDisplacementError = 0.5f; 
+static FAutoConsoleVariableRef CVarNiagaraRibbonTessellationMinDisplacementError(
+	TEXT("Niagara.Ribbon.Tessellation.MinAbsoluteError"),
+	GNiagaraRibbonTessellationMinDisplacementError,
+	TEXT("Minimum absolute world size error when tessellating. \n")
+	TEXT("Prevent over tessellating when distance gets really small. (default=0.5)"),
+	ECVF_Scalability
+);
+
+float GNiagaraRibbonMinSegmentLength = 1.f;
+static FAutoConsoleVariableRef CVarNiagaraRibbonMinSegmentLength(
+	TEXT("Niagara.Ribbon.MinSegmentLength"),
+	GNiagaraRibbonMinSegmentLength,
+	TEXT("Min length of niagara ribbon segments. (default=1)"),
+	ECVF_Scalability
+);
+
+
+// max absolute error 9.0x10^-3
+// Eberly's polynomial degree 1 - respect bounds
+// input [-1, 1] and output [0, PI]
+FORCEINLINE float AcosFast(float InX) 
+{
+    float X = FMath::Abs(InX);
+    float Res = -0.156583f * X + (0.5 * PI);
+    Res *= sqrt(1.0f - X);
+    return (InX >= 0) ? Res : PI - Res;
+}
+
 
 struct FNiagaraDynamicDataRibbon : public FNiagaraDynamicDataBase
 {
-	TArray<int16> IndexData;
+	// The list of all segments, each one connecting SortedIndices[SegmentId] to SortedIndices[SegmentId + 1].
+	// We use this format because the final index buffer gets generated based on view sorting and InterpCount.
+	TArray<int32> SegmentData;
+
+	// Curvature of the segment.
+	float AverageSegmentCurvature = 0;
+	// Angle of the curvature over the segment (in radian)­.
+	float AverangeSegmentAngle = 0; 
+
 	TArray<int32> SortedIndices;
-	TArray<float> TotalDistances;
+	TArray<FVector4> TangentAndDistances;
 	TArray<uint32> MultiRibbonIndices;
 	TArray<float> PackedPerRibbonDataByIndex;
 
@@ -28,14 +101,15 @@ struct FNiagaraDynamicDataRibbon : public FNiagaraDynamicDataBase
 	FVector StartPos;
 	FVector EndPos;
 
-	void PackPerRibbonData(float U0Scale, float U0Offset, float U1Scale, float U1Offset, uint32 NumSegments, uint32 StartVertexIndex)
+	void PackPerRibbonData(float U0Scale, float U0Offset, float U1Scale, float U1Offset, uint32 NumSegments, uint32 FirstParticleId)
 	{
+		const float OneOverNumSegments = 1 / (float)FMath::Max<uint32>(1, NumSegments);
 		PackedPerRibbonDataByIndex.Add(U0Scale);
 		PackedPerRibbonDataByIndex.Add(U0Offset);
 		PackedPerRibbonDataByIndex.Add(U1Scale);
 		PackedPerRibbonDataByIndex.Add(U1Offset);
-		PackedPerRibbonDataByIndex.Add(*(float*)&NumSegments);
-		PackedPerRibbonDataByIndex.Add(*(float*)&StartVertexIndex);
+		PackedPerRibbonDataByIndex.Add(OneOverNumSegments);
+		PackedPerRibbonDataByIndex.Add(*(float*)&FirstParticleId);
 	}
 };
 
@@ -81,8 +155,57 @@ void NiagaraRendererRibbons::CreateRenderThreadResources()
 	VertexFactory->InitResource();
 }
 
+void NiagaraRendererRibbons::GenerateIndexBuffer(uint16* OutIndices, const TArray<int32>& SegmentData,  int32 InterpCount, bool bInvertOrder, bool bCullTwistedStrips)
+{
+	auto AddTriangleIndices = [&OutIndices, InterpCount, bCullTwistedStrips](int32 SegmentIndex)
+	{
+		for (int32 SubSegmentIndex = 0; SubSegmentIndex < InterpCount; ++SubSegmentIndex)
+		{
+			const uint16 BaseVertexIndex = (int16)(SegmentIndex * InterpCount + SubSegmentIndex) * 2;
+			OutIndices[0] = BaseVertexIndex + 0;
+			OutIndices[1] = BaseVertexIndex + 1;
+			OutIndices[2] = BaseVertexIndex + 2;
+			OutIndices[3] = BaseVertexIndex + 1;
+			OutIndices[4] = BaseVertexIndex + 3;
+			OutIndices[5] = BaseVertexIndex + 2;
+
+			if (bCullTwistedStrips)
+			{
+				// When the ribbon right vector abruptly changes direction, for example when a segment starts with a tangent going up, but ends with a tanging going down,
+				// the triangles defined above will cross each other. One will be clockwise and the other counter-clockwise. We add two triangle here that replaces the invalid one
+				// this prevents the crossing/twisting artefacts while putting a bit more pressure on triangle culling.
+
+				OutIndices[6] = BaseVertexIndex + 0;
+				OutIndices[7] = BaseVertexIndex + 2;
+				OutIndices[8] = BaseVertexIndex + 3;
+				OutIndices[9] = BaseVertexIndex + 0;
+				OutIndices[10] = BaseVertexIndex + 3;
+				OutIndices[11] = BaseVertexIndex + 1;
+				OutIndices += 12;
+			}
+			else
+			{
+				OutIndices += 6;
+			}
+		}
+	};
 
-
+	// If per view sorting is required, generate sort keys and sort segment indices.
+	if (!bInvertOrder)
+	{
+		for (int32 SegmentDataIndex = 0; SegmentDataIndex < SegmentData.Num(); ++SegmentDataIndex)
+		{
+			AddTriangleIndices(SegmentData[SegmentDataIndex]);
+		}
+	}
+	else
+	{
+		for (int32 SegmentDataIndex = SegmentData.Num() - 1; SegmentDataIndex >= 0; --SegmentDataIndex)
+		{
+			AddTriangleIndices(SegmentData[SegmentDataIndex]);
+		}
+	}
+}
 
 void NiagaraRendererRibbons::GetDynamicMeshElements(const TArray<const FSceneView*>& Views, const FSceneViewFamily& ViewFamily, uint32 VisibilityMap, FMeshElementCollector& Collector, const FNiagaraSceneProxy *SceneProxy) const
 {
@@ -92,7 +215,7 @@ void NiagaraRendererRibbons::GetDynamicMeshElements(const TArray<const FSceneVie
 	SimpleTimer MeshElementsTimer;
 	FNiagaraDynamicDataRibbon *DynamicDataRibbon = static_cast<FNiagaraDynamicDataRibbon*>(DynamicDataRender);
 	if (!DynamicDataRibbon 
-		|| DynamicDataRibbon->IndexData.Num() == 0
+		|| DynamicDataRibbon->SegmentData.Num() == 0
 		|| nullptr == Properties
 		|| !GSupportsResourceView // Current shader requires SRV to draw properly in all cases.
 		)
@@ -101,12 +224,13 @@ void NiagaraRendererRibbons::GetDynamicMeshElements(const TArray<const FSceneVie
 	}
 
 	const bool bIsWireframe = ViewFamily.EngineShowFlags.Wireframe;
+	const bool bCullTwistedStrips = Properties->FacingMode == ENiagaraRibbonFacingMode::Screen && GNiagaraRibbonCullTwistedStrips != 0;
+
 	FMaterialRenderProxy* MaterialRenderProxy = Material->GetRenderProxy();
 
 	FGlobalDynamicIndexBuffer& DynamicIndexBuffer = Collector.GetDynamicIndexBuffer();
-	FGlobalDynamicIndexBuffer::FAllocation DynamicIndexAllocation = DynamicIndexBuffer.Allocate(DynamicDataRibbon->IndexData.Num(), sizeof(int16));
 
-	int32 TotalFloatSize = DynamicDataRibbon->RTParticleData.GetFloatBuffer().Num() / sizeof(float);
+	const int32 TotalFloatSize = DynamicDataRibbon->RTParticleData.GetFloatBuffer().Num() / sizeof(float);
 	FGlobalDynamicReadBuffer::FAllocation ParticleData;
 
 	if (DynamicDataRibbon->DataSet->GetSimTarget() == ENiagaraSimTarget::CPUSim)
@@ -116,157 +240,171 @@ void NiagaraRendererRibbons::GetDynamicMeshElements(const TArray<const FSceneVie
 		FMemory::Memcpy(ParticleData.Buffer, DynamicDataRibbon->RTParticleData.GetFloatBuffer().GetData(), DynamicDataRibbon->RTParticleData.GetFloatBuffer().Num());
 	}
 
-	if (DynamicIndexAllocation.IsValid())
+	// Update the primitive uniform buffer if needed.
+	if (!WorldSpacePrimitiveUniformBuffer.IsInitialized())
 	{
-		// Update the primitive uniform buffer if needed.
-		if (!WorldSpacePrimitiveUniformBuffer.IsInitialized())
-		{
-			FPrimitiveUniformShaderParameters PrimitiveUniformShaderParameters = GetPrimitiveUniformShaderParameters(
-				FMatrix::Identity,
-				FMatrix::Identity,
-				SceneProxy->GetActorPosition(),
-				SceneProxy->GetBounds(),
-				SceneProxy->GetLocalBounds(),
-				SceneProxy->ReceivesDecals(),
-				false,
-				false,
-				SceneProxy->UseSingleSampleShadowFromStationaryLights(),
-				SceneProxy->GetScene().HasPrecomputedVolumetricLightmap_RenderThread(),
-				SceneProxy->UseEditorDepthTest(),
-				SceneProxy->GetLightingChannelMask(),
-				0,
-				INDEX_NONE,
-				INDEX_NONE,
-				SceneProxy->AlwaysHasVelocity()
-				);
-			WorldSpacePrimitiveUniformBuffer.SetContents(PrimitiveUniformShaderParameters);
-			WorldSpacePrimitiveUniformBuffer.InitResource();
-		}
-
-		// Copy the vertex data over.
-		FMemory::Memcpy(DynamicIndexAllocation.Buffer, DynamicDataRibbon->IndexData.GetData(), DynamicDataRibbon->IndexData.Num() * sizeof(int16));
+		FPrimitiveUniformShaderParameters PrimitiveUniformShaderParameters = GetPrimitiveUniformShaderParameters(
+			FMatrix::Identity,
+			FMatrix::Identity,
+			SceneProxy->GetActorPosition(),
+			SceneProxy->GetBounds(),
+			SceneProxy->GetLocalBounds(),
+			SceneProxy->ReceivesDecals(),
+			false,
+			false,
+			SceneProxy->UseSingleSampleShadowFromStationaryLights(),
+			SceneProxy->GetScene().HasPrecomputedVolumetricLightmap_RenderThread(),
+			SceneProxy->UseEditorDepthTest(),
+			SceneProxy->GetLightingChannelMask(),
+			0,
+			INDEX_NONE,
+			INDEX_NONE,
+			SceneProxy->AlwaysHasVelocity()
+			);
+		WorldSpacePrimitiveUniformBuffer.SetContents(PrimitiveUniformShaderParameters);
+		WorldSpacePrimitiveUniformBuffer.InitResource();
+	}
 
-		// Compute the per-view uniform buffers.
-		for (int32 ViewIndex = 0; ViewIndex < Views.Num(); ViewIndex++)
+	// Compute the per-view uniform buffers.
+	for (int32 ViewIndex = 0; ViewIndex < Views.Num(); ViewIndex++)
+	{
+		if (VisibilityMap & (1 << ViewIndex))
 		{
-			if (VisibilityMap & (1 << ViewIndex))
+			const FSceneView* View = Views[ViewIndex];
+
+			int32 SegmentTessellation = 1;
+			int32 NumSegments = DynamicDataRibbon->SegmentData.Num();
+			if (GNiagaraRibbonMaxTessellation > 1)
 			{
-				const FSceneView* View = Views[ViewIndex];
+				const float MinTesselation = FMath::Max<float>(1.f, DynamicDataRibbon->AverangeSegmentAngle / FMath::Max<float>(SMALL_NUMBER, GNiagaraRibbonTessellationAngle));
 
-				// figure out whether start is closer to the view plane than end
-				float StartDist = FVector::DotProduct( View->GetViewDirection(), DynamicDataRibbon->StartPos - View->ViewLocation);
-				float EndDist = FVector::DotProduct( View->GetViewDirection(), DynamicDataRibbon->EndPos - View->ViewLocation);
-				bool bInvertOrder = ((StartDist > EndDist) && Properties->DrawDirection == ENiagaraRibbonDrawDirection::BackToFront)
-					|| ((StartDist < EndDist) && Properties->DrawDirection == ENiagaraRibbonDrawDirection::FrontToBack);
+				const float ViewDistance = SceneProxy->GetBounds().ComputeSquaredDistanceFromBoxToPoint(View->ViewLocation);
+				const float MaxDisplacementError = FMath::Max(GNiagaraRibbonTessellationMinDisplacementError, GNiagaraRibbonTessellationScreenPercentage * FMath::Sqrt(ViewDistance) / View->LODDistanceFactor);
+				const float Tess = DynamicDataRibbon->AverangeSegmentAngle / FMath::Max(SMALL_NUMBER, AcosFast(DynamicDataRibbon->AverageSegmentCurvature / (DynamicDataRibbon->AverageSegmentCurvature + MaxDisplacementError)));
+				SegmentTessellation = FMath::Clamp<int32>(FMath::RoundToInt(Tess), FMath::RoundToInt(MinTesselation), GNiagaraRibbonMaxTessellation);
+				NumSegments *= SegmentTessellation;
+			}
 
-				FNiagaraMeshCollectorResourcesRibbon& CollectorResources = Collector.AllocateOneFrameResource<FNiagaraMeshCollectorResourcesRibbon>();
-				FNiagaraRibbonUniformParameters PerViewUniformParameters;// = UniformParameters;
-				PerViewUniformParameters.LocalToWorld = bLocalSpace ? SceneProxy->GetLocalToWorld() : FMatrix::Identity;//For now just handle local space like this but maybe in future have a VF variant to avoid the transform entirely?
-				PerViewUniformParameters.LocalToWorldInverseTransposed = bLocalSpace ? SceneProxy->GetLocalToWorld().Inverse().GetTransposed() : FMatrix::Identity;
-				PerViewUniformParameters.DeltaSeconds = ViewFamily.DeltaWorldTime;
-				PerViewUniformParameters.CameraUp = View->GetViewUp(); // FVector4(0.0f, 0.0f, 1.0f, 0.0f);
-				PerViewUniformParameters.CameraRight = View->GetViewRight();//	FVector4(1.0f, 0.0f, 0.0f, 0.0f);
-				PerViewUniformParameters.ScreenAlignment = FVector4(0.0f, 0.0f, 0.0f, 0.0f);
-				PerViewUniformParameters.UseCustomFacing = Properties->FacingMode == ENiagaraRibbonFacingMode::Custom ? 1 : 0;
-				PerViewUniformParameters.TotalNumInstances = DynamicDataRibbon->RTParticleData.GetNumInstances();
-
-				PerViewUniformParameters.PositionDataOffset = PositionDataOffset;
-				PerViewUniformParameters.VelocityDataOffset = VelocityDataOffset;
-				PerViewUniformParameters.ColorDataOffset = ColorDataOffset;
-				PerViewUniformParameters.WidthDataOffset = WidthDataOffset;
-				PerViewUniformParameters.TwistDataOffset = TwistDataOffset;
-				PerViewUniformParameters.FacingDataOffset = FacingDataOffset;
-				PerViewUniformParameters.NormalizedAgeDataOffset = NormalizedAgeDataOffset;
-				PerViewUniformParameters.MaterialRandomDataOffset = MaterialRandomDataOffset;
-				PerViewUniformParameters.MaterialParamDataOffset = MaterialParamOffset;
-				PerViewUniformParameters.MaterialParam1DataOffset = MaterialParamOffset1;
-				PerViewUniformParameters.MaterialParam2DataOffset = MaterialParamOffset2;
-				PerViewUniformParameters.MaterialParam3DataOffset = MaterialParamOffset3;
-				PerViewUniformParameters.InvertDrawOrder = bInvertOrder ? 1 : 0;
-				PerViewUniformParameters.UV0TilingDistance = Properties->UV0TilingDistance;
-				PerViewUniformParameters.UV1TilingDistance = Properties->UV1TilingDistance;
-				PerViewUniformParameters.PackedVData = FVector4(Properties->UV0Scale.Y, Properties->UV0Offset.Y, Properties->UV1Scale.Y, Properties->UV1Offset.Y);
-				CollectorResources.VertexFactory.SetParticleData(ParticleData.ReadBuffer->SRV, ParticleData.FirstIndex / sizeof(float), DynamicDataRibbon->RTParticleData.GetFloatStride() / sizeof(float));;
-
-				// Collector.AllocateOneFrameResource uses default ctor, initialize the vertex factory
-				CollectorResources.VertexFactory.SetParticleFactoryType(NVFT_Ribbon);
-
-				CollectorResources.UniformBuffer = FNiagaraRibbonUniformBufferRef::CreateUniformBufferImmediate(PerViewUniformParameters, UniformBuffer_SingleFrame);
-
-				CollectorResources.VertexFactory.InitResource();
-				CollectorResources.VertexFactory.SetRibbonUniformBuffer(CollectorResources.UniformBuffer);
-
-				if (!DynamicDataRibbon->SortedIndices.Num())
-				{
-					return;
-				}
+			const int32 NumberOfPrimitives = NumSegments * 2 * (bCullTwistedStrips ? 2 : 1);
 
-				// TODO: need to make these two a global alloc buffer as well, not recreate
-				// pass in the sorted indices so the VS can fetch the particle data in order
-				FReadBuffer SortedIndicesBuffer;
-				SortedIndicesBuffer.Initialize(sizeof(int32), DynamicDataRibbon->SortedIndices.Num(), EPixelFormat::PF_R32_SINT, BUF_Volatile);
-				void *IndexPtr = RHILockVertexBuffer(SortedIndicesBuffer.Buffer, 0, DynamicDataRibbon->SortedIndices.Num() * sizeof(int32), RLM_WriteOnly);
-				FMemory::Memcpy(IndexPtr, DynamicDataRibbon->SortedIndices.GetData(), DynamicDataRibbon->SortedIndices.Num() * sizeof(int32));
-				RHIUnlockVertexBuffer(SortedIndicesBuffer.Buffer);
-				CollectorResources.VertexFactory.SetSortedIndices(SortedIndicesBuffer.SRV, 0);
-
-				// pass in the CPU generated total segment distance (for tiling distance modes); needs to be a buffer so we can fetch them in the correct order based on Draw Direction (front->back or back->front)
-				//	otherwise UVs will pop when draw direction changes based on camera view point
-				FReadBuffer TotalDistancesBuffer;
-				TotalDistancesBuffer.Initialize(sizeof(float), DynamicDataRibbon->TotalDistances.Num(), EPixelFormat::PF_R32_FLOAT, BUF_Volatile);
-				void *TotalDistancesPtr = RHILockVertexBuffer(TotalDistancesBuffer.Buffer, 0, DynamicDataRibbon->TotalDistances.Num() * sizeof(float), RLM_WriteOnly);
-				FMemory::Memcpy(TotalDistancesPtr, DynamicDataRibbon->TotalDistances.GetData(), DynamicDataRibbon->TotalDistances.Num() * sizeof(float));
-				RHIUnlockVertexBuffer(TotalDistancesBuffer.Buffer);
-				CollectorResources.VertexFactory.SetSegmentDistances(TotalDistancesBuffer.SRV);
-
-				// Copy a buffer which has the per particle multi ribbon index.
-				FReadBuffer MultiRibbonIndicesBuffer;
-				MultiRibbonIndicesBuffer.Initialize(sizeof(uint32), DynamicDataRibbon->MultiRibbonIndices.Num(), EPixelFormat::PF_R32_UINT, BUF_Volatile);
-				void* MultiRibbonIndexPtr = RHILockVertexBuffer(MultiRibbonIndicesBuffer.Buffer, 0, DynamicDataRibbon->MultiRibbonIndices.Num() * sizeof(uint32), RLM_WriteOnly);
-				FMemory::Memcpy(MultiRibbonIndexPtr, DynamicDataRibbon->MultiRibbonIndices.GetData(), DynamicDataRibbon->MultiRibbonIndices.Num() * sizeof(uint32));
-				RHIUnlockVertexBuffer(MultiRibbonIndicesBuffer.Buffer);
-				CollectorResources.VertexFactory.SetMultiRibbonIndicesSRV(MultiRibbonIndicesBuffer.SRV);
-
-				// Copy the packed u data for stable age based uv generation.
-				FReadBuffer PackedPerRibbonDataByIndexBuffer;
-				PackedPerRibbonDataByIndexBuffer.Initialize(sizeof(float), DynamicDataRibbon->PackedPerRibbonDataByIndex.Num(), EPixelFormat::PF_R32_FLOAT, BUF_Volatile);
-				void *PackedPerRibbonDataByIndexPtr = RHILockVertexBuffer(PackedPerRibbonDataByIndexBuffer.Buffer, 0, DynamicDataRibbon->PackedPerRibbonDataByIndex.Num() * sizeof(float), RLM_WriteOnly);
-				FMemory::Memcpy(PackedPerRibbonDataByIndexPtr, DynamicDataRibbon->PackedPerRibbonDataByIndex.GetData(), DynamicDataRibbon->PackedPerRibbonDataByIndex.Num() * sizeof(float));
-				RHIUnlockVertexBuffer(PackedPerRibbonDataByIndexBuffer.Buffer);
-				CollectorResources.VertexFactory.SetPackedPerRibbonDataByIndexSRV(PackedPerRibbonDataByIndexBuffer.SRV);
-
-				FMeshBatch& MeshBatch = Collector.AllocateMesh();
-				MeshBatch.VertexFactory = &CollectorResources.VertexFactory;
-				MeshBatch.CastShadow = SceneProxy->CastsDynamicShadow();
-				MeshBatch.bUseAsOccluder = false;
-				MeshBatch.ReverseCulling = SceneProxy->IsLocalToWorldDeterminantNegative();
-				MeshBatch.bDisableBackfaceCulling = true;
-				MeshBatch.Type = PT_TriangleList;
-				MeshBatch.DepthPriorityGroup = SceneProxy->GetDepthPriorityGroup(View);
-				MeshBatch.bCanApplyViewModeOverrides = true;
-				MeshBatch.bUseWireframeSelectionColoring = SceneProxy->IsSelected();
-
-				if (bIsWireframe)
-				{
-					MeshBatch.MaterialRenderProxy = UMaterial::GetDefaultMaterial(MD_Surface)->GetRenderProxy();
-				}
-				else
-				{
-					MeshBatch.MaterialRenderProxy = MaterialRenderProxy;
-				}
+			// Figure out whether start is closer to the view plane than end
+			// TODO : This doesn't work with multi-ribbons.
+			const float StartDist = FVector::DotProduct(View->GetViewDirection(), DynamicDataRibbon->StartPos - View->ViewLocation);
+			const float EndDist = FVector::DotProduct(View->GetViewDirection(), DynamicDataRibbon->EndPos - View->ViewLocation);
+			const bool bInvertOrder = ((StartDist > EndDist) && Properties->DrawDirection == ENiagaraRibbonDrawDirection::BackToFront)
+					|| ((StartDist < EndDist) && Properties->DrawDirection == ENiagaraRibbonDrawDirection::FrontToBack);
 
-				FMeshBatchElement& MeshElement = MeshBatch.Elements[0];
-				MeshElement.IndexBuffer = DynamicIndexAllocation.IndexBuffer;
-				MeshElement.FirstIndex = DynamicIndexAllocation.FirstIndex;
-				MeshElement.NumPrimitives = DynamicDataRibbon->IndexData.Num() / 3;
-				check(MeshElement.NumPrimitives > 0);
-				MeshElement.NumInstances = 1;
-				MeshElement.MinVertexIndex = 0;
-				MeshElement.MaxVertexIndex = 0;
-				MeshElement.PrimitiveUniformBuffer = WorldSpacePrimitiveUniformBuffer.GetUniformBufferRHI();
-
-				Collector.AddMesh(ViewIndex, MeshBatch);
+			// Copy the vertex data over.
+			FGlobalDynamicIndexBuffer::FAllocation DynamicIndexAllocation = DynamicIndexBuffer.Allocate(NumberOfPrimitives * 3, sizeof(int16));
+			GenerateIndexBuffer((uint16*)DynamicIndexAllocation.Buffer, DynamicDataRibbon->SegmentData, SegmentTessellation, bInvertOrder, bCullTwistedStrips); 
+
+			FNiagaraMeshCollectorResourcesRibbon& CollectorResources = Collector.AllocateOneFrameResource<FNiagaraMeshCollectorResourcesRibbon>();
+			FNiagaraRibbonUniformParameters PerViewUniformParameters;// = UniformParameters;
+			PerViewUniformParameters.LocalToWorld = bLocalSpace ? SceneProxy->GetLocalToWorld() : FMatrix::Identity;//For now just handle local space like this but maybe in future have a VF variant to avoid the transform entirely?
+			PerViewUniformParameters.LocalToWorldInverseTransposed = bLocalSpace ? SceneProxy->GetLocalToWorld().Inverse().GetTransposed() : FMatrix::Identity;
+			PerViewUniformParameters.DeltaSeconds = ViewFamily.DeltaWorldTime;
+			PerViewUniformParameters.CameraUp = View->GetViewUp(); // FVector4(0.0f, 0.0f, 1.0f, 0.0f);
+			PerViewUniformParameters.CameraRight = View->GetViewRight();//	FVector4(1.0f, 0.0f, 0.0f, 0.0f);
+			PerViewUniformParameters.ScreenAlignment = FVector4(0.0f, 0.0f, 0.0f, 0.0f);
+			PerViewUniformParameters.TotalNumInstances = DynamicDataRibbon->RTParticleData.GetNumInstances();
+			PerViewUniformParameters.InterpCount = SegmentTessellation;
+			PerViewUniformParameters.OneOverInterpCount = 1.f / (float)SegmentTessellation;
+
+			PerViewUniformParameters.PositionDataOffset = PositionDataOffset;
+			PerViewUniformParameters.VelocityDataOffset = VelocityDataOffset;
+			PerViewUniformParameters.ColorDataOffset = ColorDataOffset;
+			PerViewUniformParameters.WidthDataOffset = WidthDataOffset;
+			PerViewUniformParameters.TwistDataOffset = TwistDataOffset;
+			PerViewUniformParameters.FacingDataOffset = Properties->FacingMode == ENiagaraRibbonFacingMode::Custom ? FacingDataOffset : -1;
+			PerViewUniformParameters.NormalizedAgeDataOffset = NormalizedAgeDataOffset;
+			PerViewUniformParameters.MaterialRandomDataOffset = MaterialRandomDataOffset;
+			PerViewUniformParameters.MaterialParamDataOffset = MaterialParamOffset;
+			PerViewUniformParameters.MaterialParam1DataOffset = MaterialParamOffset1;
+			PerViewUniformParameters.MaterialParam2DataOffset = MaterialParamOffset2;
+			PerViewUniformParameters.MaterialParam3DataOffset = MaterialParamOffset3;
+			PerViewUniformParameters.OneOverUV0TilingDistance = Properties->UV0TilingDistance ? 1.f / (Properties->UV0TilingDistance) : 0.f;
+			PerViewUniformParameters.OneOverUV1TilingDistance = Properties->UV1TilingDistance ? 1.f / (Properties->UV1TilingDistance) : 0.f;
+			PerViewUniformParameters.PackedVData = FVector4(Properties->UV0Scale.Y, Properties->UV0Offset.Y, Properties->UV1Scale.Y, Properties->UV1Offset.Y);
+			CollectorResources.VertexFactory.SetParticleData(ParticleData.ReadBuffer->SRV, ParticleData.FirstIndex / sizeof(float), DynamicDataRibbon->RTParticleData.GetFloatStride() / sizeof(float));
+
+			// Collector.AllocateOneFrameResource uses default ctor, initialize the vertex factory
+			CollectorResources.VertexFactory.SetParticleFactoryType(NVFT_Ribbon);
+
+			CollectorResources.UniformBuffer = FNiagaraRibbonUniformBufferRef::CreateUniformBufferImmediate(PerViewUniformParameters, UniformBuffer_SingleFrame);
+
+			CollectorResources.VertexFactory.InitResource();
+			CollectorResources.VertexFactory.SetRibbonUniformBuffer(CollectorResources.UniformBuffer);
+
+			if (!DynamicDataRibbon->SortedIndices.Num())
+			{
+				return;
 			}
+
+			// TODO: need to make these two a global alloc buffer as well, not recreate
+			// pass in the sorted indices so the VS can fetch the particle data in order
+			FReadBuffer SortedIndicesBuffer;
+			SortedIndicesBuffer.Initialize(sizeof(int32), DynamicDataRibbon->SortedIndices.Num(), EPixelFormat::PF_R32_SINT, BUF_Volatile);
+			void *IndexPtr = RHILockVertexBuffer(SortedIndicesBuffer.Buffer, 0, DynamicDataRibbon->SortedIndices.Num() * sizeof(int32), RLM_WriteOnly);
+			FMemory::Memcpy(IndexPtr, DynamicDataRibbon->SortedIndices.GetData(), DynamicDataRibbon->SortedIndices.Num() * sizeof(int32));
+			RHIUnlockVertexBuffer(SortedIndicesBuffer.Buffer);
+			CollectorResources.VertexFactory.SetSortedIndices(SortedIndicesBuffer.SRV, 0);
+
+			// pass in the CPU generated total segment distance (for tiling distance modes); needs to be a buffer so we can fetch them in the correct order based on Draw Direction (front->back or back->front)
+			//	otherwise UVs will pop when draw direction changes based on camera view point
+			FReadBuffer TangentsAndDistancesBuffer;
+			TangentsAndDistancesBuffer.Initialize(sizeof(FVector4), DynamicDataRibbon->TangentAndDistances.Num(), EPixelFormat::PF_A32B32G32R32F, BUF_Volatile);
+			void *TangentsAndDistancesPtr = RHILockVertexBuffer(TangentsAndDistancesBuffer.Buffer, 0, DynamicDataRibbon->TangentAndDistances.Num() * sizeof(FVector4), RLM_WriteOnly);
+			FMemory::Memcpy(TangentsAndDistancesPtr, DynamicDataRibbon->TangentAndDistances.GetData(), DynamicDataRibbon->TangentAndDistances.Num() * sizeof(FVector4));
+			RHIUnlockVertexBuffer(TangentsAndDistancesBuffer.Buffer);
+			CollectorResources.VertexFactory.SetTangentAndDistances(TangentsAndDistancesBuffer.SRV);
+
+			// Copy a buffer which has the per particle multi ribbon index.
+			FReadBuffer MultiRibbonIndicesBuffer;
+			MultiRibbonIndicesBuffer.Initialize(sizeof(uint32), DynamicDataRibbon->MultiRibbonIndices.Num(), EPixelFormat::PF_R32_UINT, BUF_Volatile);
+			void* MultiRibbonIndexPtr = RHILockVertexBuffer(MultiRibbonIndicesBuffer.Buffer, 0, DynamicDataRibbon->MultiRibbonIndices.Num() * sizeof(uint32), RLM_WriteOnly);
+			FMemory::Memcpy(MultiRibbonIndexPtr, DynamicDataRibbon->MultiRibbonIndices.GetData(), DynamicDataRibbon->MultiRibbonIndices.Num() * sizeof(uint32));
+			RHIUnlockVertexBuffer(MultiRibbonIndicesBuffer.Buffer);
+			CollectorResources.VertexFactory.SetMultiRibbonIndicesSRV(MultiRibbonIndicesBuffer.SRV);
+
+			// Copy the packed u data for stable age based uv generation.
+			FReadBuffer PackedPerRibbonDataByIndexBuffer;
+			PackedPerRibbonDataByIndexBuffer.Initialize(sizeof(float), DynamicDataRibbon->PackedPerRibbonDataByIndex.Num(), EPixelFormat::PF_R32_FLOAT, BUF_Volatile);
+			void *PackedPerRibbonDataByIndexPtr = RHILockVertexBuffer(PackedPerRibbonDataByIndexBuffer.Buffer, 0, DynamicDataRibbon->PackedPerRibbonDataByIndex.Num() * sizeof(float), RLM_WriteOnly);
+			FMemory::Memcpy(PackedPerRibbonDataByIndexPtr, DynamicDataRibbon->PackedPerRibbonDataByIndex.GetData(), DynamicDataRibbon->PackedPerRibbonDataByIndex.Num() * sizeof(float));
+			RHIUnlockVertexBuffer(PackedPerRibbonDataByIndexBuffer.Buffer);
+			CollectorResources.VertexFactory.SetPackedPerRibbonDataByIndexSRV(PackedPerRibbonDataByIndexBuffer.SRV);
+
+			FMeshBatch& MeshBatch = Collector.AllocateMesh();
+			MeshBatch.VertexFactory = &CollectorResources.VertexFactory;
+			MeshBatch.CastShadow = SceneProxy->CastsDynamicShadow();
+			MeshBatch.bUseAsOccluder = false;
+			MeshBatch.ReverseCulling = SceneProxy->IsLocalToWorldDeterminantNegative();
+			MeshBatch.bDisableBackfaceCulling = !bCullTwistedStrips;
+			MeshBatch.Type = PT_TriangleList;
+			MeshBatch.DepthPriorityGroup = SceneProxy->GetDepthPriorityGroup(View);
+			MeshBatch.bCanApplyViewModeOverrides = true;
+			MeshBatch.bUseWireframeSelectionColoring = SceneProxy->IsSelected();
+
+			if (bIsWireframe)
+			{
+				MeshBatch.MaterialRenderProxy = UMaterial::GetDefaultMaterial(MD_Surface)->GetRenderProxy();
+			}
+			else
+			{
+				MeshBatch.MaterialRenderProxy = MaterialRenderProxy;
+			}
+
+			FMeshBatchElement& MeshElement = MeshBatch.Elements[0];
+			MeshElement.IndexBuffer = DynamicIndexAllocation.IndexBuffer;
+			MeshElement.FirstIndex = DynamicIndexAllocation.FirstIndex;
+			MeshElement.NumPrimitives = NumberOfPrimitives;
+			check(MeshElement.NumPrimitives > 0);
+			MeshElement.NumInstances = 1;
+			MeshElement.MinVertexIndex = 0;
+			MeshElement.MaxVertexIndex = 0;
+			MeshElement.PrimitiveUniformBuffer = WorldSpacePrimitiveUniformBuffer.GetUniformBufferRHI();
+
+			Collector.AddMesh(ViewIndex, MeshBatch);
 		}
 	}
 
@@ -291,9 +429,9 @@ int NiagaraRendererRibbons::GetDynamicDataSize()
 	if (DynamicDataRender)
 	{
 		FNiagaraDynamicDataRibbon* RibbonDynamicData = static_cast<FNiagaraDynamicDataRibbon*>(DynamicDataRender);
-		Size += RibbonDynamicData->IndexData.GetAllocatedSize();
+		Size += RibbonDynamicData->SegmentData.GetAllocatedSize();
 		Size += RibbonDynamicData->SortedIndices.GetAllocatedSize();
-		Size += RibbonDynamicData->TotalDistances.GetAllocatedSize();
+		Size += RibbonDynamicData->TangentAndDistances.GetAllocatedSize();
 		Size += RibbonDynamicData->MultiRibbonIndices.GetAllocatedSize();
 		Size += RibbonDynamicData->PackedPerRibbonDataByIndex.GetAllocatedSize();
 	}
@@ -303,7 +441,7 @@ int NiagaraRendererRibbons::GetDynamicDataSize()
 
 bool NiagaraRendererRibbons::HasDynamicData()
 {
-	return DynamicDataRender && (static_cast<FNiagaraDynamicDataRibbon*>(DynamicDataRender))->IndexData.Num() > 0;
+	return DynamicDataRender && (static_cast<FNiagaraDynamicDataRibbon*>(DynamicDataRender))->SegmentData.Num() > 0;
 }
 
 #if WITH_EDITORONLY_DATA
@@ -395,10 +533,10 @@ FNiagaraDynamicDataBase *NiagaraRendererRibbons::GenerateVertexData(const FNiaga
 		return nullptr;
 	}
 	FNiagaraDynamicDataRibbon* DynamicData = new FNiagaraDynamicDataRibbon;
-	TArray<int16>& IndexData = DynamicData->IndexData;
-
-	// TODO : deal with the dynamic vertex material parameter should the user have specified it as an output...
-	int32 NumTotalVerts = 0;
+	TArray<int32>& SegmentData = DynamicData->SegmentData;
+	float TotalSegmentLength = 0;
+	float TotalSegmentAngle = 0;
+	int NumTotalSamples = 0;
 
 	FNiagaraDataSetAccessor<FVector> PosData(Data, Properties->PositionBinding.DataSetVariable);
 
@@ -471,73 +609,125 @@ FNiagaraDynamicDataBase *NiagaraRendererRibbons::GenerateVertexData(const FNiaga
 	{
 		int32 StartIndex = DynamicData->SortedIndices.Num();
 		int32 NumIndices = RibbonIndices.Num();
-		if (NumIndices > 1)
-		{
-			FVector PrevPos, PrevPos2, PrevDir(0.0f, 0.0f, 0.1f);
+
 			float TotalDistance = 0.0f;
 
-			uint32 LastParticipatingParticle = RibbonIndices[0];
-			for (int32 i = 0; i < NumIndices; i++)
+		const FVector FirstPos = PosData[RibbonIndices[0]];
+		FVector CurrPos = FirstPos;
+		FVector LastToCurrVec = FVector::ZeroVector;
+		float LastToCurrSize = 0;
+
+		// Find the first position with enough distance.
+		int32 CurrentIndex = 1;
+		while (CurrentIndex < RibbonIndices.Num())
 			{
-				uint32 Index1 = RibbonIndices[i];
-				uint32 Index2 = 0;
-				const FVector ParticlePos = PosData[Index1];
-				FVector ParticleDir;
-				if (i < NumIndices - 1)
+			CurrPos = PosData[RibbonIndices[CurrentIndex]];
+			LastToCurrVec = CurrPos - FirstPos;
+			LastToCurrSize = LastToCurrVec.Size();
+
+			// Find the first segment, or unique segment
+			if (LastToCurrSize > GNiagaraRibbonMinSegmentLength)
 				{
-					Index2 = RibbonIndices[i + 1];
-					ParticleDir = PosData[Index2] - PosData[LastParticipatingParticle];
+				// Normalize LastToCurrVec
+				LastToCurrVec *= 1.f / LastToCurrSize;
+
+				// Add the first point. Tangent follows first segment.
+				DynamicData->SortedIndices.Add(RibbonIndices[0]);
+				DynamicData->TangentAndDistances.Add(FVector4(LastToCurrVec.X, LastToCurrVec.Y, LastToCurrVec.Z, 0));
+				DynamicData->MultiRibbonIndices.Add(RibbonIndex);
+
+				break;
 				}
 				else
 				{
-					Index2 = RibbonIndices[i - 1];
-					ParticleDir = PosData[LastParticipatingParticle] - PosData[Index2];
+				LastToCurrSize = 0; // Ensure that the segment gets ignored if too small
+				++CurrentIndex;
+			}
 				}
 
-				// if two ribbon particles were spawned too close together, we skip one
-				// but never skip the last, because that will result in invalid indices from the prev loop
-				if (ParticleDir.SizeSquared() > 0.002 || i == NumIndices - 1)
+		// Now iterate on all other points, to proceed each particle connected to 2 segments.
+		int32 NextIndex = CurrentIndex + 1;
+		while (NextIndex < RibbonIndices.Num())
 				{
-					DynamicData->SortedIndices.Add(Index1);
+			const FVector NextPos = PosData[RibbonIndices[NextIndex]];
+			FVector CurrToNextVec = NextPos - CurrPos;
+			const float CurrToNextSize = CurrToNextVec.Size();
+
+			// It the next is far enough, or the last element
+			if (CurrToNextSize > GNiagaraRibbonMinSegmentLength || NextIndex == RibbonIndices.Num() - 1)
+			{
+				// Normalize CurrToNextVec
+				CurrToNextVec *= 1.f / CurrToNextSize;
+				const FVector Tangent = (LastToCurrVec + CurrToNextVec).GetSafeNormal();
 
-					LastParticipatingParticle = Index2;
-					FVector NormDir = ParticleDir.GetSafeNormal();
-					PrevDir = NormDir;
+				// Update the distance for CurrentIndex.
+				TotalDistance += LastToCurrSize;
 
-					DynamicData->TotalDistances.Add(TotalDistance);
+				// Add the current point, which tangent is computed from neighbors
+				DynamicData->SortedIndices.Add(RibbonIndices[CurrentIndex]);
+				DynamicData->TangentAndDistances.Add(FVector4(Tangent.X, Tangent.Y, Tangent.Z, TotalDistance));
 					DynamicData->MultiRibbonIndices.Add(RibbonIndex);
 
-					if (i < NumIndices - 1)
-					{
-						IndexData.Add(NumTotalVerts);
-						IndexData.Add(NumTotalVerts + 1);
-						IndexData.Add(NumTotalVerts + 2);
-						IndexData.Add(NumTotalVerts + 1);
-						IndexData.Add(NumTotalVerts + 3);
-						IndexData.Add(NumTotalVerts + 2);
+				// Assumed equal to dot(Tangent, CurrToNextVec)
+				TotalSegmentLength += CurrToNextSize;
+				TotalSegmentAngle += AcosFast(FVector::DotProduct(LastToCurrVec, CurrToNextVec));
+				++NumTotalSamples;
+
+				// Move to next segment.
+				CurrentIndex = NextIndex;
+				CurrPos = NextPos;
+				LastToCurrVec = CurrToNextVec;
+				LastToCurrSize = CurrToNextSize;
 					}
-					NumTotalVerts += 2;
+
+			// Try next if there is one.
+			++NextIndex;
 				}
 
-				TotalDistance += ParticleDir.Size();
+		// Close the last point and segment if there was at least 2.
+		if (LastToCurrSize > 0)
+		{
+			// Update the distance for CurrentIndex.
+			TotalDistance += LastToCurrSize;
+
+			// Add the last point, which tangent follows the last segment.
+			DynamicData->SortedIndices.Add(RibbonIndices[CurrentIndex]);
+			DynamicData->TangentAndDistances.Add(FVector4(LastToCurrVec.X, LastToCurrVec.Y, LastToCurrVec.Z, TotalDistance));
+			DynamicData->MultiRibbonIndices.Add(RibbonIndex);
 			}
+
+		const int32 EndIndex = DynamicData->SortedIndices.Num() - 1;
+		const int32 NumSegments = EndIndex - StartIndex;
+
+		if (NumSegments > 0)
+		{
+			// Update the tangents for the first and last vertex, apply a reflect vector logic so that the initial and final curvature is continuous.
+			if (NumSegments > 1)
+			{
+				FVector& FirstTangent = reinterpret_cast<FVector&>(DynamicData->TangentAndDistances[StartIndex]);
+				FVector& NextToFirstTangent = reinterpret_cast<FVector&>(DynamicData->TangentAndDistances[StartIndex + 1]);
+				FirstTangent = (2.f * FVector::DotProduct(FirstTangent, NextToFirstTangent)) * FirstTangent - NextToFirstTangent;
+
+				FVector& LastTangent = reinterpret_cast<FVector&>(DynamicData->TangentAndDistances[EndIndex]);
+				FVector& PrevToLastTangent = reinterpret_cast<FVector&>(DynamicData->TangentAndDistances[EndIndex - 1]);
+				LastTangent = (2.f * FVector::DotProduct(LastTangent, PrevToLastTangent)) * LastTangent - PrevToLastTangent;
 		}
 
-		float U0Offset;
-		float U0Scale;
-		float U1Offset;
-		float U1Scale;
+			// Add segment data
+			for (int32 SegmentIndex = StartIndex; SegmentIndex < EndIndex; ++SegmentIndex)
+			{
+				SegmentData.Add(SegmentIndex);
+			}
 
-		int32 EndIndex = DynamicData->SortedIndices.Num() - 1;
-		int32 NumSegments = EndIndex - StartIndex;
-		int32 StartVertexIndex = StartIndex * 2; // We add two vertices for each particle.
+			float U0Offset, U0Scale, U1Offset, U1Scale;
 
 		CalculateUVScaleAndOffsets(SortKeyData, DynamicData->SortedIndices, bSortKeyIsAge, StartIndex, DynamicData->SortedIndices.Num() - 1, NumSegments,
 			Properties->UV0TilingDistance,	Properties->UV0Scale.X, Properties->UV0Offset.X, Properties->UV0AgeOffsetMode, U0Scale, U0Offset);
 		CalculateUVScaleAndOffsets(SortKeyData, DynamicData->SortedIndices, bSortKeyIsAge, StartIndex, DynamicData->SortedIndices.Num() - 1, NumSegments,
 			Properties->UV1TilingDistance, Properties->UV1Scale.X, Properties->UV1Offset.X, Properties->UV1AgeOffsetMode, U1Scale, U1Offset);
 
-		DynamicData->PackPerRibbonData(U0Scale, U0Offset, U1Scale, U1Offset, NumSegments, StartVertexIndex);
+			DynamicData->PackPerRibbonData(U0Scale, U0Offset, U1Scale, U1Offset, NumSegments, StartIndex);
+		}
 	};
 
 	// store the start and end positions for the ribbon for draw distance flipping 
@@ -605,6 +795,13 @@ FNiagaraDynamicDataBase *NiagaraRendererRibbons::GenerateVertexData(const FNiaga
 		}
 	}
 
+	if (NumTotalSamples > 0)
+	{
+		const float OneOverSampleCount = 1.f / (float)NumTotalSamples;
+		DynamicData->AverangeSegmentAngle = TotalSegmentAngle * OneOverSampleCount;
+		DynamicData->AverageSegmentCurvature = TotalSegmentLength * OneOverSampleCount * .5 / (FMath::Max(SMALL_NUMBER, FMath::Abs(FMath::Sin(DynamicData->AverangeSegmentAngle))));
+	}
+
 	if (Data.CurrData().GetNumInstances() > 0)
 	{
 		//TODO: This buffer is far fatter than needed. Just pull out the data needed for rendering.
diff --git a/Engine/Plugins/FX/Niagara/Source/Niagara/Public/NiagaraRendererRibbons.h b/Engine/Plugins/FX/Niagara/Source/Niagara/Public/NiagaraRendererRibbons.h
--- a/Engine/Plugins/FX/Niagara/Source/Niagara/Public/NiagaraRendererRibbons.h
+++ b/Engine/Plugins/FX/Niagara/Source/Niagara/Public/NiagaraRendererRibbons.h
@@ -17,6 +17,7 @@ class FNiagaraDataSet;
 class NIAGARA_API NiagaraRendererRibbons : public NiagaraRenderer
 {
 public:
+
 	NiagaraRendererRibbons(ERHIFeatureLevel::Type FeatureLevel, UNiagaraRendererProperties *Props);
 	~NiagaraRendererRibbons()
 	{
@@ -46,8 +47,6 @@ public:
 		ParamData.Add(Param);
 	}
 
-
-
 	virtual void SetDynamicData_RenderThread(FNiagaraDynamicDataBase* NewDynamicData) override;
 	int GetDynamicDataSize() override;
 	bool HasDynamicData() override;
@@ -64,7 +63,12 @@ public:
 		return Properties;
 	}
 
+protected:
+
+	static void GenerateIndexBuffer(uint16* OutIndices, const TArray<int32>& SegmentData, int32 MaxTessellation, bool bInvertOrder, bool bCullTwistedStrips);
+
 private:
+
 	class FNiagaraRibbonVertexFactory *VertexFactory;
 	UNiagaraRibbonRendererProperties *Properties;
 	mutable TUniformBuffer<FPrimitiveUniformShaderParameters> WorldSpacePrimitiveUniformBuffer;
diff --git a/Engine/Plugins/FX/Niagara/Source/NiagaraVertexFactories/Private/NiagaraRibbonVertexFactory.cpp b/Engine/Plugins/FX/Niagara/Source/NiagaraVertexFactories/Private/NiagaraRibbonVertexFactory.cpp
--- a/Engine/Plugins/FX/Niagara/Source/NiagaraVertexFactories/Private/NiagaraRibbonVertexFactory.cpp
+++ b/Engine/Plugins/FX/Niagara/Source/NiagaraVertexFactories/Private/NiagaraRibbonVertexFactory.cpp
@@ -38,7 +38,7 @@ public:
 		FloatDataStride.Bind(ParameterMap, TEXT("NiagaraFloatDataStride"));
 		SortedIndices.Bind(ParameterMap, TEXT("SortedIndices"));
 		SortedIndicesOffset.Bind(ParameterMap, TEXT("SortedIndicesOffset"));
-		SegmentDistances.Bind(ParameterMap, TEXT("SegmentDistances"));
+		TangentsAndDistances.Bind(ParameterMap, TEXT("TangentsAndDistances"));
 		MultiRibbonIndices.Bind(ParameterMap, TEXT("MultiRibbonIndices"));
 		PackedPerRibbonDataByIndex.Bind(ParameterMap, TEXT("PackedPerRibbonDataByIndex"));
 
@@ -56,7 +56,7 @@ public:
 		Ar << FloatDataStride;
 		Ar << SortedIndices;
 		Ar << SortedIndicesOffset;
-		Ar << SegmentDistances;
+		Ar << TangentsAndDistances;
 		Ar << MultiRibbonIndices;
 		Ar << PackedPerRibbonDataByIndex;
 	}
@@ -79,7 +79,7 @@ public:
 		ShaderBindings.Add(FloatDataStride, RibbonVF->GetFloatDataStride());
 
 		ShaderBindings.Add(SortedIndices, RibbonVF->GetSortedIndicesSRV());
-		ShaderBindings.Add(SegmentDistances, RibbonVF->GetSegmentDistancesSRV());
+		ShaderBindings.Add(TangentsAndDistances, RibbonVF->GetTangentAndDistancesSRV());
 		ShaderBindings.Add(MultiRibbonIndices, RibbonVF->GetMultiRibbonIndicesSRV());
 		ShaderBindings.Add(PackedPerRibbonDataByIndex, RibbonVF->GetPackedPerRibbonDataByIndexSRV());
 		ShaderBindings.Add(SortedIndicesOffset, RibbonVF->GetSortedIndicesOffset());
@@ -91,7 +91,7 @@ private:
 	FShaderParameter FloatDataStride;
 
 	FShaderResourceParameter SortedIndices;
-	FShaderResourceParameter SegmentDistances;
+	FShaderResourceParameter TangentsAndDistances;
 	FShaderResourceParameter MultiRibbonIndices;
 	FShaderResourceParameter PackedPerRibbonDataByIndex;
 	FShaderParameter SortedIndicesOffset;
diff --git a/Engine/Plugins/FX/Niagara/Source/NiagaraVertexFactories/Public/NiagaraRibbonVertexFactory.h b/Engine/Plugins/FX/Niagara/Source/NiagaraVertexFactories/Public/NiagaraRibbonVertexFactory.h
--- a/Engine/Plugins/FX/Niagara/Source/NiagaraVertexFactories/Public/NiagaraRibbonVertexFactory.h
+++ b/Engine/Plugins/FX/Niagara/Source/NiagaraVertexFactories/Public/NiagaraRibbonVertexFactory.h
@@ -42,10 +42,10 @@ BEGIN_GLOBAL_SHADER_PARAMETER_STRUCT(FNiagaraRibbonUniformParameters, NIAGARAVER
 	SHADER_PARAMETER(int, MaterialParam2DataOffset)
 	SHADER_PARAMETER(int, MaterialParam3DataOffset)
 	SHADER_PARAMETER(int, TotalNumInstances)
-	SHADER_PARAMETER(uint32, UseCustomFacing)
-	SHADER_PARAMETER(uint32, InvertDrawOrder)
-	SHADER_PARAMETER(float, UV0TilingDistance)
-	SHADER_PARAMETER(float, UV1TilingDistance)
+	SHADER_PARAMETER(int, InterpCount)
+	SHADER_PARAMETER(float, OneOverInterpCount)
+	SHADER_PARAMETER(float, OneOverUV0TilingDistance)
+	SHADER_PARAMETER(float, OneOverUV1TilingDistance)
 	SHADER_PARAMETER(FVector4, PackedVData)
 	SHADER_PARAMETER_EX(FMatrix, LocalToWorld, EShaderPrecisionModifier::Half)
 	SHADER_PARAMETER_EX(FMatrix, LocalToWorldInverseTransposed, EShaderPrecisionModifier::Half)
@@ -132,9 +132,9 @@ public:
 		SortedIndicesOffset = InSortedIndicesOffset;
 	}
 
-	void SetSegmentDistances(const FShaderResourceViewRHIRef& InSegmentDistancesSRV)
+	void SetTangentAndDistances(const FShaderResourceViewRHIRef& InTangentAndDistancesSRV)
 	{
-		SegmentDistancesSRV = InSegmentDistancesSRV;
+		TangentAndDistancesSRV = InTangentAndDistancesSRV;
 	}
 
 
@@ -173,9 +173,9 @@ public:
 		return SortedIndicesOffset;
 	}
 
-	FORCEINLINE FShaderResourceViewRHIRef GetSegmentDistancesSRV()
+	FORCEINLINE FShaderResourceViewRHIRef GetTangentAndDistancesSRV()
 	{
-		return SegmentDistancesSRV;
+		return TangentAndDistancesSRV;
 	}
 
 	FORCEINLINE FShaderResourceViewRHIRef GetMultiRibbonIndicesSRV()
@@ -225,7 +225,7 @@ private:
 	uint32 FloatDataStride;
 
 	FShaderResourceViewRHIRef SortedIndicesSRV;
-	FShaderResourceViewRHIRef SegmentDistancesSRV;
+	FShaderResourceViewRHIRef TangentAndDistancesSRV;
 	FShaderResourceViewRHIRef MultiRibbonIndicesSRV;
 	FShaderResourceViewRHIRef PackedPerRibbonDataByIndexSRV;
 
diff --git a/Engine/Shaders/Private/NiagaraRibbonVertexFactory.ush b/Engine/Shaders/Private/NiagaraRibbonVertexFactory.ush
--- a/Engine/Shaders/Private/NiagaraRibbonVertexFactory.ush
+++ b/Engine/Shaders/Private/NiagaraRibbonVertexFactory.ush
@@ -12,6 +12,7 @@
 #define USE_PARTICLE_TIME (FEATURE_LEVEL >= FEATURE_LEVEL_SM5 && NEEDS_PARTICLE_TIME)
 #define USE_PARTICLE_SIZE (FEATURE_LEVEL >= FEATURE_LEVEL_SM5 && NEEDS_PARTICLE_SIZE)
 #define USE_PARTICLE_RANDOM (FEATURE_LEVEL >= FEATURE_LEVEL_SM5 && NEEDS_PARTICLE_RANDOM)
+#define USE_PARTICLE_INTERPOLATION (FEATURE_LEVEL >= FEATURE_LEVEL_SM5)
 
 float3 TransformPosition(float3 InPosition)
 {
@@ -23,144 +24,258 @@ float3 TransformVector(float3 InVector)
 	return mul(float4(InVector,0.0), NiagaraRibbonVF.LocalToWorldInverseTransposed).xyz;
 }
 
+float3 CubicInterp(float3 P0, float3 T0, float3 P1, float3 T1, float A, float TS)
+{
+	const float A2 = A  * A;
+	const float A3 = A2 * A;
+	return (((2*A3)-(3*A2)+1) * P0) + (TS * (A3-(2*A2)+A) * T0) + (TS * (A3-A2) * T1) + (((-2*A3)+(3*A2)) * P1);
+}
 
 /*********************************************************************************
  *							Particle specific
  *********************************************************************************/
+ 
+Buffer<int> SortedIndices;
+uint SortedIndicesOffset;
 
-float3 GetNiagaraParticlePosition(uint InstanceID)
+Buffer<float4> TangentsAndDistances;
+Buffer<uint> MultiRibbonIndices;
+Buffer<float> PackedPerRibbonDataByIndex;
+
+// Naming convention :
+// RawParticleId : Used to access per particle data stored in the VF. Actual particle data needs to be indirected through SortedIndices.
+// ParticleId : Used to access particle data in NiagaraVFParticleAccess.
+// VertexId : Used to represent a vertex on both side of the particles. Has some specific data like UV.
+// InterpVtxId : A virtual index that represent a vertex of or between 2 VertexId.
+
+struct FVertexFactoryInput
 {
-	if (NiagaraRibbonVF.PositionDataOffset == -1)
-	{
-		return TransformPosition(float3(0.0f, 0.0f, 0.0f));
-	}
-	return TransformPosition(GetVec3(NiagaraRibbonVF.PositionDataOffset, InstanceID));
-}
-float3 GetNiagaraParticleVelocity(uint InstanceID)
+	uint InterpVtxId					: SV_VertexID;
+
+	// Optional instance ID for vertex layered rendering
+#if FEATURE_LEVEL >= FEATURE_LEVEL_SM4 && ONEPASS_POINTLIGHT_SHADOW && USING_VERTEX_SHADER_LAYER
+	uint InstanceId					: SV_InstanceID;
+#endif
+};
+
+// The interpolation data required for an interp vertex.
+struct FInterpVtxData
 {
-	if(NiagaraRibbonVF.VelocityDataOffset  == -1)
+	uint RawParticleId0;
+	uint RawParticleId1;
+	uint ParticleId0;
+	uint ParticleId1;
+	float Alpha;
+};
+
+struct FSegmentData
+{
+	float2 UV0Scale;
+	float2 UV0Offset;
+	float2 UV1Scale;
+	float2 UV1Offset;
+	float OneOverNumSegments;
+	uint StartParticleId;
+};
+
+struct FVertexFactoryIntermediates
+{
+	float3 Position;
+	float3 OldPosition;
+	float3 Direction;
+	float DistanceOnSegment;
+	float4 Color;
+	float RibbonWidth;
+	float3 Right;
+	float3 Up;
+	float3 Facing;
+	float4 TexCoord;
+	
+#if USE_PARTICLE_POSITION
+	float4 TranslatedWorldPositionAndSize;
+#endif
+
+	float4 ParticleVelocity;
+
+#if USE_PARTICLE_TIME
+	float RelativeTime;
+#endif
+
+#if USE_PARTICLE_SIZE
+	float2 ParticleSize;
+#endif
+
+#if USE_PARTICLE_RANDOM
+	float ParticleRandom;
+#endif
+
+#if (DYNAMIC_PARAMETERS_MASK & 1)
+	float4 DynamicParameter;
+#endif
+#if (DYNAMIC_PARAMETERS_MASK & 2)
+	float4 DynamicParameter1;
+#endif
+#if (DYNAMIC_PARAMETERS_MASK & 4)
+	float4 DynamicParameter2;
+#endif
+#if (DYNAMIC_PARAMETERS_MASK & 8)
+	float4 DynamicParameter3;
+#endif
+};
+
+void GetNiagaraParticlePositionAndDirection(in FInterpVtxData InterpVtxData, in out FVertexFactoryIntermediates Intermediates)
+{
+	BRANCH // Also acts as ISOLATE
+	if (NiagaraRibbonVF.PositionDataOffset != -1)
 	{
-		// See FNiagaraConstants for a synchronized value...
-		return TransformVector(float3(0.0f,0.0f,0.0f));
+		const float3 Pos0 = GetVec3(NiagaraRibbonVF.PositionDataOffset, InterpVtxData.ParticleId0);
+		const float3 Pos1 = GetVec3(NiagaraRibbonVF.PositionDataOffset, InterpVtxData.ParticleId1);
+		const float4 TangentAndDist0 = TangentsAndDistances[InterpVtxData.RawParticleId0];
+		const float4 TangentAndDist1 = TangentsAndDistances[InterpVtxData.RawParticleId1];
+		
+		Intermediates.Position = TransformPosition(CubicInterp(Pos0, TangentAndDist0.xyz, Pos1, TangentAndDist1.xyz, InterpVtxData.Alpha, TangentAndDist1.w - TangentAndDist0.w));
+		Intermediates.Direction = normalize(TransformVector(lerp(TangentAndDist0.xyz, TangentAndDist1.xyz, InterpVtxData.Alpha)));
+
+		Intermediates.DistanceOnSegment = lerp(TangentAndDist0.w, TangentAndDist1.w, InterpVtxData.Alpha);
 	}
-	return TransformVector(GetVec3(NiagaraRibbonVF.VelocityDataOffset, InstanceID));
-}
-float GetNiagaraParticleRibbonWidth(uint InstanceID)
+} 
+
+void GetNiagaraParticleVelocity(in FInterpVtxData InterpVtxData, in out FVertexFactoryIntermediates Intermediates)
 {
-	if (NiagaraRibbonVF.WidthDataOffset == -1)
+	BRANCH
+	if(NiagaraRibbonVF.VelocityDataOffset != -1)
 	{
-		return 1.0f;
+		const float3 Velocity0 = GetVec3(NiagaraRibbonVF.VelocityDataOffset, InterpVtxData.ParticleId0);
+		const float3 Velocity1 = GetVec3(NiagaraRibbonVF.VelocityDataOffset, InterpVtxData.ParticleId1);
+		const float3 Velocity = TransformVector(lerp(Velocity0, Velocity1, InterpVtxData.Alpha));
+
+		Intermediates.ParticleVelocity = float4(normalize(Velocity), length(Velocity));
 	}
-	return GetFloat(NiagaraRibbonVF.WidthDataOffset, InstanceID);
 }
-float GetNiagaraParticleRibbonTwist(uint InstanceID)
+
+void GetNiagaraParticleColor(in FInterpVtxData InterpVtxData, in out FVertexFactoryIntermediates Intermediates)
 {
-	if (NiagaraRibbonVF.TwistDataOffset == -1)
+	// Overrides (FVertexFactoryIntermediates)0; 
+	Intermediates.Color = float4(1.0f, 1.0f, 1.0f, 1.0f); 
+
+	if (NiagaraRibbonVF.ColorDataOffset != -1)
 	{
-		return 0.0f;
+		const float4 Color0 = GetVec4(NiagaraRibbonVF.ColorDataOffset, InterpVtxData.ParticleId0);
+		const float4 Color1 = GetVec4(NiagaraRibbonVF.ColorDataOffset, InterpVtxData.ParticleId1);
+		Intermediates.Color = lerp(Color0, Color1, InterpVtxData.Alpha);
 	}
-	return GetFloat(NiagaraRibbonVF.TwistDataOffset, InstanceID);
 }
-float4 GetNiagaraParticleColor(uint InstanceID)
+
+void GetNiagaraParticleDynamicParameters(in FInterpVtxData InterpVtxData, in out FVertexFactoryIntermediates Intermediates)
 {
-	if (NiagaraRibbonVF.ColorDataOffset == -1)
+#if (DYNAMIC_PARAMETERS_MASK & 1)
+	// Overrides (FVertexFactoryIntermediates)0; 
+	Intermediates.DynamicParameter = float4(1.0f, 1.0f, 1.0f, 1.0f);
+
+	FLATTEN
+	if (NiagaraRibbonVF.MaterialParamDataOffset != -1)
 	{
-		return float4(1.0f, 1.0f, 1.0f, 1.0f);
+		Intermediates.DynamicParameter = lerp(GetVec4(NiagaraRibbonVF.MaterialParamDataOffset, InterpVtxData.ParticleId0), GetVec4(NiagaraRibbonVF.MaterialParamDataOffset, InterpVtxData.ParticleId1), InterpVtxData.Alpha);
 	}
-	return GetVec4(NiagaraRibbonVF.ColorDataOffset, InstanceID);
-}
+#endif
+#if (DYNAMIC_PARAMETERS_MASK & 2)
+	// Overrides (FVertexFactoryIntermediates)0; 
+	Intermediates.DynamicParameter1 = float4(1.0f, 1.0f, 1.0f, 1.0f);
 
-#if USE_PARTICLE_TIME
-float GetNiagaraNormalizedAge(uint InstanceID)
-{
-	if (NiagaraRibbonVF.NormalizedAgeDataOffset == -1)
+	FLATTEN
+	if (NiagaraRibbonVF.MaterialParam1DataOffset != -1)
 	{
-		return 0;
+		Intermediates.DynamicParameter1 = lerp(GetVec4(NiagaraRibbonVF.MaterialParam1DataOffset, InterpVtxData.ParticleId0), GetVec4(NiagaraRibbonVF.MaterialParam1DataOffset, InterpVtxData.ParticleId1), InterpVtxData.Alpha);
 	}
-	return GetFloat(NiagaraRibbonVF.NormalizedAgeDataOffset, InstanceID);
-}
 #endif
+#if (DYNAMIC_PARAMETERS_MASK & 4)
+	// Overrides (FVertexFactoryIntermediates)0; 
+	Intermediates.DynamicParameter2 = float4(1.0f, 1.0f, 1.0f, 1.0f);
 
-#if USE_PARTICLE_RANDOM
-float GetNiagaraMaterialRandom(uint InstanceID)
-{
-	if (NiagaraRibbonVF.MaterialRandomDataOffset == -1)
+	FLATTEN
+	if (NiagaraRibbonVF.MaterialParam2DataOffset != -1)
 	{
-		return 0;
+		Intermediates.DynamicParameter2 = lerp(GetVec4(NiagaraRibbonVF.MaterialParam2DataOffset, InterpVtxData.ParticleId0), GetVec4(NiagaraRibbonVF.MaterialParam2DataOffset, InterpVtxData.ParticleId1), InterpVtxData.Alpha);
 	}
-	return GetFloat(NiagaraRibbonVF.MaterialRandomDataOffset, InstanceID);
-}
 #endif
+#if (DYNAMIC_PARAMETERS_MASK & 8)
+	// Overrides (FVertexFactoryIntermediates)0; 
+	Intermediates.DynamicParameter3 = float4(1.0f, 1.0f, 1.0f, 1.0f);
 
-#if (DYNAMIC_PARAMETERS_MASK & 1)
-float4 GetNiagaraParticleDynamicParameters(uint InstanceID)
-{
-	if (NiagaraRibbonVF.MaterialParamDataOffset == -1)
+	FLATTEN
+	if (NiagaraRibbonVF.MaterialParam3DataOffset != -1)
 	{
-		return float4(1.0f, 1.0f, 1.0f, 1.0f);
+		Intermediates.DynamicParameter3 = lerp(GetVec4(NiagaraRibbonVF.MaterialParam3DataOffset, InterpVtxData.ParticleId0), GetVec4(NiagaraRibbonVF.MaterialParam3DataOffset, InterpVtxData.ParticleId1), InterpVtxData.Alpha);
 	}
-	return GetVec4(NiagaraRibbonVF.MaterialParamDataOffset, InstanceID);
-}
 #endif
-#if (DYNAMIC_PARAMETERS_MASK & 2)
-float4 GetNiagaraParticleDynamicParameters1(uint InstanceID)
+}
+
+void GetNiagaraParticleTextureCoords(in FVertexFactoryInput Input, in FInterpVtxData InterpVtxData, in FSegmentData SegmentData,  in out FVertexFactoryIntermediates Intermediates)
 {
-	if (NiagaraRibbonVF.MaterialParam1DataOffset == -1)
+	float U0ForSegment;
+	BRANCH
+	if(NiagaraRibbonVF.OneOverUV0TilingDistance != 0.0f)
 	{
-		return float4(1.0f, 1.0f, 1.0f, 1.0f);
+		U0ForSegment = Intermediates.DistanceOnSegment * NiagaraRibbonVF.OneOverUV0TilingDistance;
 	}
-	return GetVec4(NiagaraRibbonVF.MaterialParam1DataOffset, InstanceID);
+	else
+	{
+		U0ForSegment = ((float)(InterpVtxData.RawParticleId0 - SegmentData.StartParticleId) + InterpVtxData.Alpha) * SegmentData.OneOverNumSegments;
+	}
+
+	float U1ForSegment;
+	BRANCH
+	if(NiagaraRibbonVF.OneOverUV1TilingDistance != 0.0f)
+	{
+		U1ForSegment = Intermediates.DistanceOnSegment * NiagaraRibbonVF.OneOverUV1TilingDistance;
+	}
+	else
+	{
+		U1ForSegment = ((float)(InterpVtxData.RawParticleId0 - SegmentData.StartParticleId) + InterpVtxData.Alpha) * SegmentData.OneOverNumSegments;
+	}
+
+	float VForSegment = Input.InterpVtxId & 0x1;
+	
+	float2 UV0ForSegment = float2(U0ForSegment, VForSegment);
+	float2 UV1ForSegment = float2(U1ForSegment, VForSegment);
+	
+	Intermediates.TexCoord.xy = UV0ForSegment * SegmentData.UV0Scale + SegmentData.UV0Offset;
+	Intermediates.TexCoord.zw = UV1ForSegment * SegmentData.UV1Scale + SegmentData.UV1Offset;
 }
-#endif
-#if (DYNAMIC_PARAMETERS_MASK & 4)
-float4 GetNiagaraParticleDynamicParameters2(uint InstanceID)
+
+void GetNiagaraParticleRibbonWidth(in FInterpVtxData InterpVtxData, in out FVertexFactoryIntermediates Intermediates)
 {
-	if (NiagaraRibbonVF.MaterialParam2DataOffset == -1)
+	// Overrides (FVertexFactoryIntermediates)0; 
+	Intermediates.RibbonWidth = 1;
+
+	FLATTEN
+	if (NiagaraRibbonVF.WidthDataOffset != -1)
 	{
-		return float4(1.0f, 1.0f, 1.0f, 1.0f);
+		Intermediates.RibbonWidth = lerp(GetFloat(NiagaraRibbonVF.WidthDataOffset, InterpVtxData.ParticleId0), GetFloat(NiagaraRibbonVF.WidthDataOffset, InterpVtxData.ParticleId1), InterpVtxData.Alpha);
 	}
-	return GetVec4(NiagaraRibbonVF.MaterialParam2DataOffset, InstanceID);
 }
-#endif
-#if (DYNAMIC_PARAMETERS_MASK & 8)
-float4 GetNiagaraParticleDynamicParameters3(uint InstanceID)
+
+#if USE_PARTICLE_TIME
+void GetNiagaraNormalizedAge(in FInterpVtxData InterpVtxData, in out FVertexFactoryIntermediates Intermediates)
 {
-	if (NiagaraRibbonVF.MaterialParam3DataOffset == -1)
+	FLATTEN
+	if (NiagaraRibbonVF.NormalizedAgeDataOffset != -1)
 	{
-		return float4(1.0f, 1.0f, 1.0f, 1.0f);
+		Intermediates.RelativeTime = lerp(GetFloat(NiagaraRibbonVF.NormalizedAgeDataOffset, InterpVtxData.ParticleId0), GetFloat(NiagaraRibbonVF.NormalizedAgeDataOffset, InterpVtxData.ParticleId1), InterpVtxData.Alpha);
 	}
-	return GetVec4(NiagaraRibbonVF.MaterialParam3DataOffset, InstanceID);
 }
 #endif
 
-
-float3 GetNiagaraParticleFacingVector(uint InstanceID)
+#if USE_PARTICLE_RANDOM
+void GetNiagaraMaterialRandom(in FInterpVtxData InterpVtxData, in out FVertexFactoryIntermediates Intermediates)
 {
-	float3 Val = float3(0.0f, 0.0f, 1.0f);
-    if(NiagaraRibbonVF.FacingDataOffset != -1)
+	FLATTEN
+	if (NiagaraRibbonVF.MaterialRandomDataOffset != -1)
 	{
-		Val = GetVec3(NiagaraRibbonVF.FacingDataOffset, InstanceID);
+		Intermediates.ParticleRandom = lerp(GetFloat(NiagaraRibbonVF.MaterialRandomDataOffset, InterpVtxData.ParticleId0), GetFloat(NiagaraRibbonVF.MaterialRandomDataOffset, InterpVtxData.ParticleId1), InterpVtxData.Alpha);
 	}
-	return TransformVector(Val);
 }
-
-
-Buffer<int> SortedIndices;
-uint SortedIndicesOffset;
-
-Buffer<float> SegmentDistances;
-Buffer<uint> MultiRibbonIndices;
-Buffer<float> PackedPerRibbonDataByIndex;
-
-struct FVertexFactoryInput
-{
-	uint VertexId					: SV_VertexID;
-
-	// Optional instance ID for vertex layered rendering
-#if FEATURE_LEVEL >= FEATURE_LEVEL_SM4 && ONEPASS_POINTLIGHT_SHADOW && USING_VERTEX_SHADER_LAYER
-	uint InstanceId	: SV_InstanceID;
 #endif
-};
 
 struct FVertexFactoryInterpolantsVSToPS
 {
@@ -216,55 +331,6 @@ struct FVertexFactoryInterpolantsVSToPS
 #endif
 };
 
-struct FVertexFactoryIntermediates
-{
-	float3 Position;
-	float3 OldPosition;
-	float3 Direction;
-	float3 CustomFacingVector;
-	float4 Color;
-	float RibbonWidth;
-	float Twist;
-	float3 Right;
-	float3 Up;
-	float3 CameraDirection;
-	float4 TexCoord;
-	
-#if USE_PARTICLE_POSITION
-	float4 TranslatedWorldPositionAndSize;
-#endif
-
-#if USE_PARTICLE_VELOCITY
-	float4 ParticleVelocity;
-#endif
-
-#if USE_PARTICLE_TIME
-	float RelativeTime;
-#endif
-
-#if USE_PARTICLE_SIZE
-	float2 ParticleSize;
-#endif
-
-#if USE_PARTICLE_RANDOM
-	float ParticleRandom;
-#endif
-
-#if (DYNAMIC_PARAMETERS_MASK & 1)
-	float4 DynamicParameter;
-#endif
-#if (DYNAMIC_PARAMETERS_MASK & 2)
-	float4 DynamicParameter1;
-#endif
-#if (DYNAMIC_PARAMETERS_MASK & 4)
-	float4 DynamicParameter2;
-#endif
-#if (DYNAMIC_PARAMETERS_MASK & 8)
-	float4 DynamicParameter3;
-#endif
-};
-
-
 #if NUM_MATERIAL_TEXCOORDS
 bool UVIndexUseZW(int UVIndex)
 {
@@ -297,17 +363,11 @@ void SetUV(inout FVertexFactoryInterpolantsVSToPS Interpolants, int UVIndex, flo
 
 #endif
 
-int GetParticleIdByVertexId(uint VertexId)
+int GetParticleIdByInterpVtxId(uint InterpVtxId)
 {
-	uint ParticleId = VertexId / 2;
-
-	/*
-	if(NiagaraRibbonVF.InvertDrawOrder == 1)
-	{
-		ParticleId = NiagaraRibbonVF.TotalNumInstances - ParticleId;
-	}
-	*/
+	uint ParticleId = InterpVtxId / 2 / NiagaraRibbonVF.InterpCount;
 
+	FLATTEN
 	if(SortedIndicesOffset != 0xFFFFFFFF)
 	{
 		ParticleId = SortedIndices[SortedIndicesOffset + ParticleId];
@@ -316,30 +376,49 @@ int GetParticleIdByVertexId(uint VertexId)
 	return ParticleId;
 }
 
-int GetParticleId(FVertexFactoryInput Input)
+FInterpVtxData GetInterpVtxData(FVertexFactoryInput Input)
 {
-	return GetParticleIdByVertexId(Input.VertexId);
-}
+	FInterpVtxData InterpVtxData = (FInterpVtxData)0;
 
-int GetPreviousParticleId(FVertexFactoryInput Input)
-{
-	return GetParticleIdByVertexId(Input.VertexId - 2);
-}
+	InterpVtxData.RawParticleId0 = Input.InterpVtxId / 2 / NiagaraRibbonVF.InterpCount;
+	InterpVtxData.RawParticleId1 = InterpVtxData.RawParticleId0 + 1;
 
-int GetNextParticleId(FVertexFactoryInput Input)
-{
-	return GetParticleIdByVertexId(Input.VertexId + 2);
+	if (SortedIndicesOffset != 0xFFFFFFFF)
+	{
+		InterpVtxData.ParticleId0 = SortedIndices[SortedIndicesOffset + InterpVtxData.RawParticleId0];
+		InterpVtxData.ParticleId1 = SortedIndices[SortedIndicesOffset + InterpVtxData.RawParticleId1];
+	}
+	else
+	{
+		InterpVtxData.ParticleId0 = InterpVtxData.RawParticleId0;
+		InterpVtxData.ParticleId1 = InterpVtxData.RawParticleId1;
+	}
+
+#if USE_PARTICLE_INTERPOLATION
+	ISOLATE
+	{
+		InterpVtxData.Alpha = (float)((Input.InterpVtxId / 2) % NiagaraRibbonVF.InterpCount) * NiagaraRibbonVF.OneOverInterpCount;
+	}
+#else
+	InterpVtxData.Alpha = 0;
+#endif
+
+	return InterpVtxData;
 }
 
-void UnpackPerRibbonDataByIndex(int RibbonIndex, out float2 UV0Scale, out float2 UV0Offset, out float2 UV1Scale, out float2 UV1Offset, out uint NumSegments, out uint StartVertexId)
+FSegmentData UnpackPerRibbonDataByIndex(int RibbonIndex)
 {
-	int Index = RibbonIndex * 6;
-	UV0Scale = float2(PackedPerRibbonDataByIndex[Index], NiagaraRibbonVF.PackedVData.x);
-	UV0Offset = float2(PackedPerRibbonDataByIndex[Index + 1], NiagaraRibbonVF.PackedVData.y);
-	UV1Scale = float2(PackedPerRibbonDataByIndex[Index + 2], NiagaraRibbonVF.PackedVData.z);
-	UV1Offset = float2(PackedPerRibbonDataByIndex[Index + 3], NiagaraRibbonVF.PackedVData.w);
-	NumSegments = asuint(PackedPerRibbonDataByIndex[Index + 4]);
-	StartVertexId = asuint(PackedPerRibbonDataByIndex[Index + 5]);
+	FSegmentData SegmentData = (FSegmentData)0;
+	const int Index = RibbonIndex * 6;
+
+	SegmentData.UV0Scale = float2(PackedPerRibbonDataByIndex[Index], NiagaraRibbonVF.PackedVData.x);
+	SegmentData.UV0Offset = float2(PackedPerRibbonDataByIndex[Index + 1], NiagaraRibbonVF.PackedVData.y);
+	SegmentData.UV1Scale = float2(PackedPerRibbonDataByIndex[Index + 2], NiagaraRibbonVF.PackedVData.z);
+	SegmentData.UV1Offset = float2(PackedPerRibbonDataByIndex[Index + 3], NiagaraRibbonVF.PackedVData.w);
+	SegmentData.OneOverNumSegments = PackedPerRibbonDataByIndex[Index + 4];
+	SegmentData.StartParticleId = asuint(PackedPerRibbonDataByIndex[Index + 5]);
+
+	return SegmentData;
 }
 
 /** Converts from vertex factory specific interpolants to a FMaterialPixelParameters, which is used by material inputs. */
@@ -463,180 +542,96 @@ FMaterialVertexParameters GetMaterialVertexParameters(FVertexFactoryInput Input,
 	return Result;
 }
 
+#if 0
 float3 SafeNormalize(float3 V)
 {
-	return V / sqrt(max(dot(V,V),0.01));
+	return rsqrt(max(0.001, dot(V,V))) * V;
 }
+#endif
 
-void GetTangents(FVertexFactoryInput Input, FVertexFactoryIntermediates Intermediates, out float3 Right,out float3 Up, out float3 CameraDirection)
+void GetNiagaraParticleTangents(in FInterpVtxData InterpVtxData, in out FVertexFactoryIntermediates Intermediates)
 {
-	float4	Position			= TransformLocalToTranslatedWorld(Intermediates.Position.xyz);
-
-	CameraDirection				= SafeNormalize(ResolvedView.WorldCameraOrigin - Intermediates.Position.xyz );		// TODO: handle local space
-	float3 ParticleDirection	= Intermediates.Direction;
-
-	float3 NormDir = SafeNormalize(ParticleDirection);
-	float3 RightDir, Facing;
-
-
-	if(NiagaraRibbonVF.UseCustomFacing)
+	BRANCH
+	if(NiagaraRibbonVF.FacingDataOffset != -1)
 	{
-		Facing = Intermediates.CustomFacingVector;
+		const float3 Facing0 = GetVec3(NiagaraRibbonVF.FacingDataOffset, InterpVtxData.ParticleId0);
+		const float3 Facing1 = GetVec3(NiagaraRibbonVF.FacingDataOffset, InterpVtxData.ParticleId1);
+		Intermediates.Facing = normalize(TransformVector(lerp(Facing0, Facing1, InterpVtxData.Alpha)));
 	}
 	else
 	{
-		Facing = CameraDirection.xyz;
-	}
+		Intermediates.Facing = normalize(ResolvedView.WorldCameraOrigin - Intermediates.Position);
+	}	
 
-	float cosa = cos(Intermediates.Twist);
-	Facing = Facing * cosa + cross(Facing, NormDir) * sin(Intermediates.Twist) + NormDir * dot(NormDir, Facing) * (1-cosa);
-
-	RightDir = SafeNormalize(cross(NormDir, Facing));
-	if(!any(RightDir))
+	ISOLATE
 	{
-		RightDir = SafeNormalize(cross(NormDir, float3(0.0f, 0.0f, 1.0f)));
+		Intermediates.Right = normalize(cross(Intermediates.Facing, Intermediates.Direction));
+		Intermediates.Up = cross(Intermediates.Right, Intermediates.Facing);
 	}
 
-	Right = RightDir;
-	Up = NormDir;
-	return;
+	BRANCH
+	if (NiagaraRibbonVF.TwistDataOffset != -1)
+	{
+		const float Twist0 = GetFloat(NiagaraRibbonVF.TwistDataOffset, InterpVtxData.ParticleId0);
+		const float Twist1 = GetFloat(NiagaraRibbonVF.TwistDataOffset, InterpVtxData.ParticleId1);
+		float SinTwist, CosTwist;
+		sincos(lerp(Twist0, Twist1, InterpVtxData.Alpha), SinTwist, CosTwist);
+
+		// Apply a rotation matrix to up and right.
+		float3 Right = Intermediates.Right; // Back it up for the transform.
+		Intermediates.Right = CosTwist * Right + SinTwist * Intermediates.Up;
+		Intermediates.Up = -SinTwist * Right + CosTwist * Intermediates.Up;
+	}
 }
 
 float4 CalcWorldPosition(FVertexFactoryInput Input, FVertexFactoryIntermediates Intermediates)
 {
-	// expand and align
-	int Sgn = (((int)Input.VertexId & 0x1) * 2) - 1;
-	float3 NormDir = SafeNormalize(Intermediates.Direction.xyz);
-	float3 RightDir = Intermediates.Right;
-	float3 UpDir = Intermediates.Up;
-	float3 Pos = Intermediates.Position.xyz + RightDir * Sgn * (Intermediates.RibbonWidth.xxx / 2);
-	return TransformLocalToTranslatedWorld(Pos);
+	ISOLATE
+	{
+		// expand and align
+		int Sgn = (((int)Input.InterpVtxId & 0x1) * 2) - 1;
+		float3 Pos = Intermediates.Position.xyz + Intermediates.Right * Sgn * (Intermediates.RibbonWidth.xxx / 2);
+		// Intermediates are already in world space, only need to account for the translation around view origin.
+		return float4(Pos + ResolvedView.PreViewTranslation, 1);
+	}
 }
 
 /** derive basis vectors */
 float3x3 CalcTangentBasis(FVertexFactoryInput Input, FVertexFactoryIntermediates Intermediates)
 {
-	float3	Right,
-			Up;
-	float3	CameraDir;
-	GetTangents(Input, Intermediates, Right, Up, CameraDir);
-	return float3x3(
-			Intermediates.Right.xyz,
-			Intermediates.Up.xyz,
-			normalize(cross(Intermediates.Right.xyz, Intermediates.Up.xyz))
-			);
+	return float3x3(-Intermediates.Up, -Intermediates.Right, Intermediates.Facing);
 }
 
 FVertexFactoryIntermediates GetVertexFactoryIntermediates(FVertexFactoryInput Input)
 {
-	FVertexFactoryIntermediates Intermediates;
-	Intermediates = (FVertexFactoryIntermediates)0;
-	uint ParticleId = GetParticleId(Input);
-	
-	float2 UV0Scale;
-	float2 UV0Offset;
-	float2 UV1Scale;
-	float2 UV1Offset;
-	uint NumSegments;
-	uint StartVertexId;
+	FVertexFactoryIntermediates Intermediates = (FVertexFactoryIntermediates)0;
+	FInterpVtxData InterpVtxData = GetInterpVtxData(Input);
 
-	uint RibbonIndex = MultiRibbonIndices[Input.VertexId / 2];
-	UnpackPerRibbonDataByIndex(RibbonIndex, UV0Scale, UV0Offset, UV1Scale, UV1Offset, NumSegments,  StartVertexId);
-	
-	uint StartParticleId = GetParticleIdByVertexId(StartVertexId);
+	uint RibbonIndex = MultiRibbonIndices[InterpVtxData.RawParticleId0];
+	FSegmentData SegmentData = UnpackPerRibbonDataByIndex(RibbonIndex);
 
-	Intermediates.Position = GetNiagaraParticlePosition(ParticleId);
-	if(ParticleId == StartParticleId)
-	{
-		uint NextParticleId = GetNextParticleId(Input);
-		float3 NextParticlePos = GetNiagaraParticlePosition(NextParticleId);
-		Intermediates.Direction = NextParticlePos - Intermediates.Position;
-	}
-	else
-	{
-		uint PreviousParticleId = GetPreviousParticleId(Input);
-		float3 PreviousParticlePos = GetNiagaraParticlePosition(PreviousParticleId);
-		Intermediates.Direction = Intermediates.Position - PreviousParticlePos;
-	}
-	
-	float3 ParticleVelocity = GetNiagaraParticleVelocity(ParticleId);
-	Intermediates.OldPosition = Intermediates.Position.xyz - ParticleVelocity * NiagaraRibbonVF.DeltaSeconds;
-	Intermediates.Color	= GetNiagaraParticleColor(ParticleId);
-	
-	if(NiagaraRibbonVF.UseCustomFacing)
-	{
-		Intermediates.CustomFacingVector = GetNiagaraParticleFacingVector(ParticleId);
-	}
-
-#if (DYNAMIC_PARAMETERS_MASK & 1)
-	Intermediates.DynamicParameter = GetNiagaraParticleDynamicParameters(ParticleId);
-#endif
-#if (DYNAMIC_PARAMETERS_MASK & 2)
-	Intermediates.DynamicParameter1 = GetNiagaraParticleDynamicParameters1(ParticleId);
-#endif
-#if (DYNAMIC_PARAMETERS_MASK & 4)
-	Intermediates.DynamicParameter2 = GetNiagaraParticleDynamicParameters2(ParticleId);
-#endif
-#if (DYNAMIC_PARAMETERS_MASK & 8)
-	Intermediates.DynamicParameter3 = GetNiagaraParticleDynamicParameters3(ParticleId);
-#endif
-
-    Intermediates.RibbonWidth = GetNiagaraParticleRibbonWidth(ParticleId);
-    Intermediates.Twist = GetNiagaraParticleRibbonTwist(ParticleId);
-
-	GetTangents(Input, Intermediates, Intermediates.Right, Intermediates.Up, Intermediates.CameraDirection);
-
-	float U0ForSegment;
-	if(NiagaraRibbonVF.UV0TilingDistance != 0.0f)
-	{
-		U0ForSegment = SegmentDistances[Input.VertexId / 2] / NiagaraRibbonVF.UV0TilingDistance;
-	}
-	else
-	{
-		U0ForSegment = (float)((Input.VertexId - StartVertexId) / 2) / NumSegments;
-	}
-
-	float U1ForSegment;
-	if(NiagaraRibbonVF.UV1TilingDistance != 0.0f)
-	{
-		U1ForSegment = SegmentDistances[Input.VertexId / 2] / NiagaraRibbonVF.UV0TilingDistance;
-	}
-	else
-	{
-		U1ForSegment = (float)((Input.VertexId - StartVertexId) / 2) / NumSegments;
-	}
-
-	float VForSegment = (Input.VertexId + 1) & 0x1;
-	
-	float2 UV0ForSegment = float2(U0ForSegment, VForSegment);
-	float2 UV1ForSegment = float2(U1ForSegment, VForSegment);
-	
-	Intermediates.TexCoord.xy = UV0ForSegment * UV0Scale + UV0Offset;
-	Intermediates.TexCoord.zw = UV1ForSegment * UV1Scale + UV1Offset;
+	GetNiagaraParticlePositionAndDirection(InterpVtxData, Intermediates);
+	GetNiagaraParticleVelocity(InterpVtxData, Intermediates);
+	GetNiagaraParticleColor(InterpVtxData, Intermediates);
+	GetNiagaraParticleDynamicParameters(InterpVtxData, Intermediates);
+	GetNiagaraParticleTextureCoords(Input, InterpVtxData, SegmentData, Intermediates);
+	GetNiagaraParticleRibbonWidth(InterpVtxData, Intermediates);
+	GetNiagaraParticleTangents(InterpVtxData, Intermediates);
 
 #if USE_PARTICLE_POSITION
-	float4 ParticleWorldPosition = mul(float4(Intermediates.Position.xyz,1), Primitive.LocalToWorld);
-	float ParticleRadius = .5 * Intermediates.RibbonWidth;
-	float3 ParticleTranslatedWorldPosition = ParticleWorldPosition.xyz + ResolvedView.PreViewTranslation.xyz;
-	Intermediates.TranslatedWorldPositionAndSize = float4(ParticleTranslatedWorldPosition, ParticleRadius);
-#endif
-
-#if USE_PARTICLE_VELOCITY
-	Intermediates.ParticleVelocity.xyz = normalize(ParticleVelocity);
-	Intermediates.ParticleVelocity.w = length(ParticleVelocity);
+	Intermediates.TranslatedWorldPositionAndSize = float4(Intermediates.Position + ResolvedView.PreViewTranslation,  .5 * Intermediates.RibbonWidth);
 #endif
 
 #if USE_PARTICLE_TIME
-	Intermediates.RelativeTime = GetNiagaraNormalizedAge(ParticleId);
+	GetNiagaraNormalizedAge(InterpVtxData, Intermediates);
 #endif
 
 #if USE_PARTICLE_SIZE
-	float SegmentLength  = length(Intermediates.Direction);
-	Intermediates.ParticleSize = float2(SegmentLength, Intermediates.RibbonWidth);
+	Intermediates.ParticleSize = Intermediates.RibbonWidth.xx;
 #endif
 
 #if USE_PARTICLE_RANDOM
-	Intermediates.ParticleRandom = GetNiagaraMaterialRandom(ParticleId);
+	GetNiagaraMaterialRandom(InterpVtxData, Intermediates);
 #endif
 
 	return Intermediates;
@@ -706,11 +701,10 @@ FVertexFactoryInterpolantsVSToPS VertexFactoryGetInterpolantsVSToPS(FVertexFacto
 	Interpolants.ParticleRandom = Intermediates.ParticleRandom;
 #endif
 
-	float3x3 TangentToLocal = CalcTangentBasis(Input, Intermediates);
-	float3x3 TangentToWorld = mul(TangentToLocal, GetLocalToWorld3x3());
-
+	float3x3 TangentToWorld = CalcTangentBasis(Input, Intermediates);
 	Interpolants.TangentToWorld0 = TangentToWorld[0];
-	Interpolants.TangentToWorld2 = float4(TangentToWorld[2], sign(determinant(TangentToLocal)));
+	// GetNiagaraParticleTangents() technically makes the determinant to be always 1.
+	Interpolants.TangentToWorld2 = float4(TangentToWorld[2], 1 /*sign(determinant(TangentToWorld))*/);
 
 #if INSTANCED_STEREO
 	Interpolants.EyeIndex = 0;
@@ -722,8 +716,8 @@ FVertexFactoryInterpolantsVSToPS VertexFactoryGetInterpolantsVSToPS(FVertexFacto
 
 float4 VertexFactoryGetPreviousWorldPosition(FVertexFactoryInput Input, FVertexFactoryIntermediates Intermediates)
 {
-	uint ParticleId = GetParticleId(Input);
-	return float4(CalcWorldPosition(Input, Intermediates).xyz - GetNiagaraParticleVelocity(ParticleId) * NiagaraRibbonVF.DeltaSeconds, 1);
+	const float3 WorldPosition = CalcWorldPosition(Input, Intermediates).xyz;
+	return float4(WorldPosition - Intermediates.ParticleVelocity.xyz * (Intermediates.ParticleVelocity.w * NiagaraRibbonVF.DeltaSeconds), 1);
 }
 
 /**
