Niagara - Don't set up inheritance for template emitters in systems, and add an option to break inheritance with a parent emitter.

diff --git a/Engine/Plugins/FX/Niagara/Source/Niagara/Classes/NiagaraEmitterHandle.h b/Engine/Plugins/FX/Niagara/Source/Niagara/Classes/NiagaraEmitterHandle.h
--- a/Engine/Plugins/FX/Niagara/Source/Niagara/Classes/NiagaraEmitterHandle.h
+++ b/Engine/Plugins/FX/Niagara/Source/Niagara/Classes/NiagaraEmitterHandle.h
@@ -87,6 +87,9 @@ public:
 	/** Sets an object which will contain the history of modifications to the emitter instance. */
 	void SetEmitterModificationHistory();
 
+	/** Removes the source system from the this emitter handle which will prevent inheriting any further changes. */
+	void RemoveSource();
+
 #endif
 public:
 	/** A static const invalid handle. */
diff --git a/Engine/Plugins/FX/Niagara/Source/Niagara/Private/NiagaraEmitterHandle.cpp b/Engine/Plugins/FX/Niagara/Source/Niagara/Private/NiagaraEmitterHandle.cpp
--- a/Engine/Plugins/FX/Niagara/Source/Niagara/Private/NiagaraEmitterHandle.cpp
+++ b/Engine/Plugins/FX/Niagara/Source/Niagara/Private/NiagaraEmitterHandle.cpp
@@ -240,4 +240,10 @@ INiagaraModule::FMergeEmitterResults FNiagaraEmitterHandle::MergeSourceChanges()
 	return MergeResults;
 }
 
+void FNiagaraEmitterHandle::RemoveSource()
+{
+	Source = nullptr;
+	LastMergedSource = nullptr;
+}
+
 #endif
\ No newline at end of file
diff --git a/Engine/Plugins/FX/Niagara/Source/Niagara/Private/NiagaraSystem.cpp b/Engine/Plugins/FX/Niagara/Source/Niagara/Private/NiagaraSystem.cpp
--- a/Engine/Plugins/FX/Niagara/Source/Niagara/Private/NiagaraSystem.cpp
+++ b/Engine/Plugins/FX/Niagara/Source/Niagara/Private/NiagaraSystem.cpp
@@ -619,6 +619,10 @@ bool UNiagaraSystem::IsValid()const
 FNiagaraEmitterHandle UNiagaraSystem::AddEmitterHandle(UNiagaraEmitter& SourceEmitter, FName EmitterName)
 {
 	FNiagaraEmitterHandle EmitterHandle(SourceEmitter, EmitterName, *this);
+	if (SourceEmitter.bIsTemplateAsset)
+	{
+		EmitterHandle.RemoveSource();
+	}
 	EmitterHandles.Add(EmitterHandle);
 	RefreshSystemParametersFromEmitter(EmitterHandle);
 	return EmitterHandle;
diff --git a/Engine/Plugins/FX/Niagara/Source/NiagaraEditor/Private/Toolkits/NiagaraSystemToolkit.cpp b/Engine/Plugins/FX/Niagara/Source/NiagaraEditor/Private/Toolkits/NiagaraSystemToolkit.cpp
--- a/Engine/Plugins/FX/Niagara/Source/NiagaraEditor/Private/Toolkits/NiagaraSystemToolkit.cpp
+++ b/Engine/Plugins/FX/Niagara/Source/NiagaraEditor/Private/Toolkits/NiagaraSystemToolkit.cpp
@@ -436,6 +436,10 @@ void FNiagaraSystemToolkit::InitializeWithEmitter(const EToolkitMode::Type Mode,
 	SystemViewModel = MakeShareable(new FNiagaraSystemViewModel(*System, SystemOptions));
 	SystemViewModel->SetToolkitCommands(GetToolkitCommands());
 	SystemViewModel->AddEmitter(*Emitter);
+
+	// Always remove the source information from the emitter handle when editing an emitter asset because inheritance is not valid in this case.
+	System->GetEmitterHandle(0).RemoveSource();
+
 	// We set this to the copy's change id here instead of the original emitter's change id because the copy's change id may have been
 	// updated from the original as part of post load and we use this id to detect if the editable emitter has been changed.
 	LastSyncedEmitterChangeId = SystemViewModel->GetEmitterHandleViewModels()[0]->GetEmitterViewModel()->GetEmitter()->GetChangeId();
diff --git a/Engine/Plugins/FX/Niagara/Source/NiagaraEditor/Private/ViewModels/Stack/NiagaraStackEmitterSpawnScriptItemGroup.cpp b/Engine/Plugins/FX/Niagara/Source/NiagaraEditor/Private/ViewModels/Stack/NiagaraStackEmitterSpawnScriptItemGroup.cpp
--- a/Engine/Plugins/FX/Niagara/Source/NiagaraEditor/Private/ViewModels/Stack/NiagaraStackEmitterSpawnScriptItemGroup.cpp
+++ b/Engine/Plugins/FX/Niagara/Source/NiagaraEditor/Private/ViewModels/Stack/NiagaraStackEmitterSpawnScriptItemGroup.cpp
@@ -36,25 +36,25 @@ FText UNiagaraStackEmitterPropertiesItem::GetDisplayName() const
 
 bool UNiagaraStackEmitterPropertiesItem::CanResetToBase() const
 {
-	if (bCanResetToBase.IsSet() == false)
+	if (bCanResetToBaseCache.IsSet() == false)
 	{
 		const UNiagaraEmitter* BaseEmitter = FNiagaraStackGraphUtilities::GetBaseEmitter(*Emitter.Get(), GetSystemViewModel()->GetSystem());
 		if (BaseEmitter != nullptr && Emitter != BaseEmitter)
 		{
 			TSharedRef<FNiagaraScriptMergeManager> MergeManager = FNiagaraScriptMergeManager::Get();
-			bCanResetToBase = MergeManager->IsEmitterEditablePropertySetDifferentFromBase(*Emitter.Get(), *BaseEmitter);
+			bCanResetToBaseCache = MergeManager->IsEmitterEditablePropertySetDifferentFromBase(*Emitter.Get(), *BaseEmitter);
 		}
 		else
 		{
-			bCanResetToBase = false;
+			bCanResetToBaseCache = false;
 		}
 	}
-	return bCanResetToBase.GetValue();
+	return bCanResetToBaseCache.GetValue();
 }
 
 void UNiagaraStackEmitterPropertiesItem::ResetToBase()
 {
-	if (bCanResetToBase.GetValue())
+	if (CanResetToBase())
 	{
 		const UNiagaraEmitter* BaseEmitter = FNiagaraStackGraphUtilities::GetBaseEmitter(*Emitter.Get(), GetSystemViewModel()->GetSystem());
 		TSharedRef<FNiagaraScriptMergeManager> MergeManager = FNiagaraScriptMergeManager::Get();
@@ -73,13 +73,13 @@ void UNiagaraStackEmitterPropertiesItem::RefreshChildrenInternal(const TArray<UN
 	}
 
 	NewChildren.Add(EmitterObject);
-	bCanResetToBase.Reset();
+	bCanResetToBaseCache.Reset();
 	Super::RefreshChildrenInternal(CurrentChildren, NewChildren, NewIssues);
 }
 
 void UNiagaraStackEmitterPropertiesItem::EmitterPropertiesChanged()
 {
-	bCanResetToBase.Reset();
+	bCanResetToBaseCache.Reset();
 }
 
 UNiagaraStackEmitterSpawnScriptItemGroup::UNiagaraStackEmitterSpawnScriptItemGroup()
diff --git a/Engine/Plugins/FX/Niagara/Source/NiagaraEditor/Private/ViewModels/Stack/NiagaraStackEntry.cpp b/Engine/Plugins/FX/Niagara/Source/NiagaraEditor/Private/ViewModels/Stack/NiagaraStackEntry.cpp
--- a/Engine/Plugins/FX/Niagara/Source/NiagaraEditor/Private/ViewModels/Stack/NiagaraStackEntry.cpp
+++ b/Engine/Plugins/FX/Niagara/Source/NiagaraEditor/Private/ViewModels/Stack/NiagaraStackEntry.cpp
@@ -2,7 +2,11 @@
 
 #include "ViewModels/Stack/NiagaraStackEntry.h"
 #include "ViewModels/Stack/NiagaraStackErrorItem.h"
+#include "ViewModels/Stack/NiagaraStackGraphUtilities.h"
+#include "ViewModels/NiagaraSystemViewModel.h"
+#include "ViewModels/NiagaraEmitterViewModel.h"
 #include "NiagaraStackEditorData.h"
+#include "NiagaraScriptMergeManager.h"
 #include "Misc/SecureHash.h"
 
 const FName UNiagaraStackEntry::FExecutionCategoryNames::System = TEXT("System");
@@ -382,6 +386,17 @@ void UNiagaraStackEntry::SetIsSearchResult(bool bInIsSearchResult)
 	bIsSearchResult = bInIsSearchResult;
 }
 
+bool UNiagaraStackEntry::HasBaseEmitter() const
+{
+	if (bHasBaseEmitterCache.IsSet() == false)
+	{
+		TSharedRef<FNiagaraScriptMergeManager> MergeManager = FNiagaraScriptMergeManager::Get();
+		const UNiagaraEmitter* BaseEmitter = FNiagaraStackGraphUtilities::GetBaseEmitter(*GetEmitterViewModel()->GetEmitter(), GetSystemViewModel()->GetSystem());
+		bHasBaseEmitterCache = BaseEmitter != nullptr;
+	}
+	return bHasBaseEmitterCache.GetValue();
+}
+
 TOptional<UNiagaraStackEntry::FDropResult> UNiagaraStackEntry::CanDropInternal(const TArray<UNiagaraStackEntry*>& DraggedEntries)
 {
 	return TOptional<FDropResult>();
@@ -434,6 +449,8 @@ void UNiagaraStackEntry::RefreshChildren()
 		ErrorChild->OnIssueModified().RemoveAll(this);
 	}
 
+	bHasBaseEmitterCache.Reset();
+
 	TArray<UNiagaraStackEntry*> NewChildren;
 	TArray<FStackIssue> NewStackIssues;
 	RefreshChildrenInternal(Children, NewChildren, NewStackIssues);
diff --git a/Engine/Plugins/FX/Niagara/Source/NiagaraEditor/Private/ViewModels/Stack/NiagaraStackEventScriptItemGroup.cpp b/Engine/Plugins/FX/Niagara/Source/NiagaraEditor/Private/ViewModels/Stack/NiagaraStackEventScriptItemGroup.cpp
--- a/Engine/Plugins/FX/Niagara/Source/NiagaraEditor/Private/ViewModels/Stack/NiagaraStackEventScriptItemGroup.cpp
+++ b/Engine/Plugins/FX/Niagara/Source/NiagaraEditor/Private/ViewModels/Stack/NiagaraStackEventScriptItemGroup.cpp
@@ -35,13 +35,6 @@ void UNiagaraStackEventHandlerPropertiesItem::Initialize(FRequiredEntryData InRe
 
 	Emitter = GetEmitterViewModel()->GetEmitter();
 	Emitter->OnPropertiesChanged().AddUObject(this, &UNiagaraStackEventHandlerPropertiesItem::EventHandlerPropertiesChanged);
-
-	const UNiagaraEmitter* BaseEmitter = FNiagaraStackGraphUtilities::GetBaseEmitter(*Emitter.Get(), GetSystemViewModel()->GetSystem());
-	if (BaseEmitter != nullptr && Emitter != BaseEmitter)
-	{
-		TSharedRef<FNiagaraScriptMergeManager> MergeManager = FNiagaraScriptMergeManager::Get();
-		bHasBaseEventHandler = MergeManager->HasBaseEventHandler(*BaseEmitter, EventScriptUsageId);
-	}
 }
 
 void UNiagaraStackEventHandlerPropertiesItem::FinalizeInternal()
@@ -60,32 +53,32 @@ FText UNiagaraStackEventHandlerPropertiesItem::GetDisplayName() const
 
 bool UNiagaraStackEventHandlerPropertiesItem::CanResetToBase() const
 {
-	if (bCanResetToBase.IsSet() == false)
+	if (bCanResetToBaseCache.IsSet() == false)
 	{
-		if (bHasBaseEventHandler)
+		if (HasBaseEventHandler())
 		{
 			const UNiagaraEmitter* BaseEmitter = FNiagaraStackGraphUtilities::GetBaseEmitter(*Emitter.Get(), GetSystemViewModel()->GetSystem());
 			if (BaseEmitter != nullptr && Emitter != BaseEmitter)
 			{
 				TSharedRef<FNiagaraScriptMergeManager> MergeManager = FNiagaraScriptMergeManager::Get();
-				bCanResetToBase = MergeManager->IsEventHandlerPropertySetDifferentFromBase(*Emitter.Get(), *BaseEmitter, EventScriptUsageId);
+				bCanResetToBaseCache = MergeManager->IsEventHandlerPropertySetDifferentFromBase(*Emitter.Get(), *BaseEmitter, EventScriptUsageId);
 			}
 			else
 			{
-				bCanResetToBase = false;
+				bCanResetToBaseCache = false;
 			}
 		}
 		else
 		{
-			bCanResetToBase = false;
+			bCanResetToBaseCache = false;
 		}
 	}
-	return bCanResetToBase.GetValue();
+	return bCanResetToBaseCache.GetValue();
 }
 
 void UNiagaraStackEventHandlerPropertiesItem::ResetToBase()
 {
-	if (bCanResetToBase.GetValue())
+	if (CanResetToBase())
 	{
 		const UNiagaraEmitter* BaseEmitter = FNiagaraStackGraphUtilities::GetBaseEmitter(*Emitter.Get(), GetSystemViewModel()->GetSystem());
 		TSharedRef<FNiagaraScriptMergeManager> MergeManager = FNiagaraScriptMergeManager::Get();
@@ -108,13 +101,15 @@ void UNiagaraStackEventHandlerPropertiesItem::RefreshChildrenInternal(const TArr
 
 	NewChildren.Add(EmitterObject);
 
-	bCanResetToBase.Reset();
+	bCanResetToBaseCache.Reset();
+	bHasBaseEventHandlerCache.Reset();
+
 	Super::RefreshChildrenInternal(CurrentChildren, NewChildren, NewIssues);
 }
 
 void UNiagaraStackEventHandlerPropertiesItem::EventHandlerPropertiesChanged()
 {
-	bCanResetToBase.Reset();
+	bCanResetToBaseCache.Reset();
 }
 
 TSharedPtr<IDetailTreeNode> GetEventHandlerArrayPropertyNode(const TArray<TSharedRef<IDetailTreeNode>>& Nodes)
@@ -174,6 +169,24 @@ void UNiagaraStackEventHandlerPropertiesItem::SelectEmitterStackObjectRootTreeNo
 	}
 }
 
+bool UNiagaraStackEventHandlerPropertiesItem::HasBaseEventHandler() const
+{
+	if (bHasBaseEventHandlerCache.IsSet() == false)
+	{
+		const UNiagaraEmitter* BaseEmitter = FNiagaraStackGraphUtilities::GetBaseEmitter(*Emitter.Get(), GetSystemViewModel()->GetSystem());
+		if (BaseEmitter != nullptr && Emitter != BaseEmitter)
+		{
+			TSharedRef<FNiagaraScriptMergeManager> MergeManager = FNiagaraScriptMergeManager::Get();
+			bHasBaseEventHandlerCache = MergeManager->HasBaseEventHandler(*BaseEmitter, EventScriptUsageId);
+		}
+		else
+		{
+			bHasBaseEventHandlerCache = false;
+		}
+	}
+	return bHasBaseEventHandlerCache.GetValue();
+}
+
 void UNiagaraStackEventScriptItemGroup::Initialize(
 	FRequiredEntryData InRequiredEntryData,
 	TSharedRef<FNiagaraScriptViewModel> InScriptViewModel,
@@ -187,6 +200,8 @@ void UNiagaraStackEventScriptItemGroup::Initialize(
 
 void UNiagaraStackEventScriptItemGroup::RefreshChildrenInternal(const TArray<UNiagaraStackEntry*>& CurrentChildren, TArray<UNiagaraStackEntry*>& NewChildren, TArray<FStackIssue>& NewIssues)
 {
+	bHasBaseEventHandlerCache.Reset();
+
 	FName EventSpacerKey = *FString::Printf(TEXT("EventSpacer"));
 	UNiagaraStackSpacer* SeparatorSpacer = FindCurrentChildOfTypeByPredicate<UNiagaraStackSpacer>(CurrentChildren,
 		[=](UNiagaraStackSpacer* CurrentSpacer) { return CurrentSpacer->GetSpacerKey() == EventSpacerKey; });
@@ -206,14 +221,10 @@ void UNiagaraStackEventScriptItemGroup::RefreshChildrenInternal(const TArray<UNi
 	if (EventScriptProperties != nullptr)
 	{
 		SetDisplayName(FText::Format(LOCTEXT("FormatEventScriptDisplayName", "Event Handler - Source: {0}"), FText::FromName(EventScriptProperties->SourceEventName)));
-
-		const UNiagaraEmitter* BaseEmitter = FNiagaraStackGraphUtilities::GetBaseEmitter(*Emitter, GetSystemViewModel()->GetSystem());
-		bHasBaseEventHandler = BaseEmitter != nullptr && FNiagaraScriptMergeManager::Get()->HasBaseEventHandler(*BaseEmitter, GetScriptUsageId());
 	}
 	else
 	{
 		SetDisplayName(LOCTEXT("UnassignedEventDisplayName", "Unassigned Event"));
-		bHasBaseEventHandler = false;
 	}
 
 	if (EventHandlerProperties == nullptr)
@@ -228,7 +239,7 @@ void UNiagaraStackEventScriptItemGroup::RefreshChildrenInternal(const TArray<UNi
 
 bool UNiagaraStackEventScriptItemGroup::CanDelete() const
 {
-	return bHasBaseEventHandler == false;
+	return HasBaseEventHandler() == false;
 }
 
 bool UNiagaraStackEventScriptItemGroup::Delete()
@@ -272,6 +283,16 @@ bool UNiagaraStackEventScriptItemGroup::Delete()
 	return true;
 }
 
+bool UNiagaraStackEventScriptItemGroup::HasBaseEventHandler() const
+{
+	if (bHasBaseEventHandlerCache.IsSet() == false)
+	{
+		const UNiagaraEmitter* BaseEmitter = FNiagaraStackGraphUtilities::GetBaseEmitter(*GetEmitterViewModel()->GetEmitter(), GetSystemViewModel()->GetSystem());
+		bHasBaseEventHandlerCache = BaseEmitter != nullptr && FNiagaraScriptMergeManager::Get()->HasBaseEventHandler(*BaseEmitter, GetScriptUsageId());
+	}
+	return bHasBaseEventHandlerCache.GetValue();
+}
+
 void UNiagaraStackEventScriptItemGroup::SetOnModifiedEventHandlers(FOnModifiedEventHandlers OnModifiedEventHandlers)
 {
 	OnModifiedEventHandlersDelegate = OnModifiedEventHandlers;
diff --git a/Engine/Plugins/FX/Niagara/Source/NiagaraEditor/Private/ViewModels/Stack/NiagaraStackFunctionInput.cpp b/Engine/Plugins/FX/Niagara/Source/NiagaraEditor/Private/ViewModels/Stack/NiagaraStackFunctionInput.cpp
--- a/Engine/Plugins/FX/Niagara/Source/NiagaraEditor/Private/ViewModels/Stack/NiagaraStackFunctionInput.cpp
+++ b/Engine/Plugins/FX/Niagara/Source/NiagaraEditor/Private/ViewModels/Stack/NiagaraStackFunctionInput.cpp
@@ -589,8 +589,8 @@ void UNiagaraStackFunctionInput::RefreshValues()
 		}
 	}
 
-	bCanReset.Reset();
-	bCanResetToBase.Reset();
+	bCanResetCache.Reset();
+	bCanResetToBaseCache.Reset();
 	ValueChangedDelegate.Broadcast();
 }
 
@@ -1036,7 +1036,7 @@ void UNiagaraStackFunctionInput::SetLocalValue(TSharedRef<FStructOnScope> InLoca
 
 bool UNiagaraStackFunctionInput::CanReset() const
 {
-	if (bCanReset.IsSet() == false)
+	if (bCanResetCache.IsSet() == false)
 	{
 		bool bNewCanReset;
 		if (InputValues.Mode == EValueMode::Data)
@@ -1088,9 +1088,9 @@ bool UNiagaraStackFunctionInput::CanReset() const
 				}
 			}
 		}
-		bCanReset = bNewCanReset;
+		bCanResetCache = bNewCanReset;
 	}
-	return bCanReset.GetValue();
+	return bCanResetCache.GetValue();
 }
 
 FNiagaraVariable UNiagaraStackFunctionInput::GetDefaultVariableForRapidIterationParameter() const
@@ -1229,16 +1229,11 @@ void UNiagaraStackFunctionInput::Reset()
 	RefreshChildren();
 }
 
-bool UNiagaraStackFunctionInput::EmitterHasBase() const
-{
-	return GetSystemViewModel()->GetEditMode() == ENiagaraSystemViewModelEditMode::SystemAsset;
-}
-
 bool UNiagaraStackFunctionInput::CanResetToBase() const
 {
-	if (EmitterHasBase())
+	if (HasBaseEmitter())
 	{
-		if (bCanResetToBase.IsSet() == false)
+		if (bCanResetToBaseCache.IsSet() == false)
 		{
 			bool bIsModuleInput = OwningFunctionCallNode == OwningModuleNode;
 			if (bIsModuleInput)
@@ -1250,7 +1245,7 @@ bool UNiagaraStackFunctionInput::CanResetToBase() const
 				{
 					const UNiagaraEmitter* BaseEmitter = FNiagaraStackGraphUtilities::GetBaseEmitter(*GetEmitterViewModel()->GetEmitter(), GetSystemViewModel()->GetSystem());
 
-					bCanResetToBase = BaseEmitter != nullptr && MergeManager->IsModuleInputDifferentFromBase(
+					bCanResetToBaseCache = BaseEmitter != nullptr && MergeManager->IsModuleInputDifferentFromBase(
 						*GetEmitterViewModel()->GetEmitter(),
 						*BaseEmitter,
 						OutputNode->GetUsage(),
@@ -1260,15 +1255,15 @@ bool UNiagaraStackFunctionInput::CanResetToBase() const
 				}
 				else
 				{
-					bCanResetToBase = false;
+					bCanResetToBaseCache = false;
 				}
 			}
 			else
 			{
-				bCanResetToBase = false;
+				bCanResetToBaseCache = false;
 			}
 		}
-		return bCanResetToBase.GetValue();
+		return bCanResetToBaseCache.GetValue();
 	}
 	return false;
 }
@@ -1571,8 +1566,8 @@ void UNiagaraStackFunctionInput::OnGraphChanged(const struct FEdGraphEditAction&
 
 void UNiagaraStackFunctionInput::OnRapidIterationParametersChanged()
 {
-	bCanReset.Reset();
-	bCanResetToBase.Reset();
+	bCanResetCache.Reset();
+	bCanResetToBaseCache.Reset();
 	if (ensureMsgf(OwningModuleNode.IsValid() && OwningFunctionCallNode.IsValid(), TEXT("Stack entry with invalid module or function call not cleaned up.")))
 	{
 		if (bUpdatingLocalValueDirectly == false && IsRapidIterationCandidate() && (OverridePinCache.IsSet() == false || OverridePinCache.GetValue() == nullptr))
@@ -1584,8 +1579,8 @@ void UNiagaraStackFunctionInput::OnRapidIterationParametersChanged()
 
 void UNiagaraStackFunctionInput::OnScriptSourceChanged()
 {
-	bCanReset.Reset();
-	bCanResetToBase.Reset();
+	bCanResetCache.Reset();
+	bCanResetToBaseCache.Reset();
 }
 
 UNiagaraNodeParameterMapSet* UNiagaraStackFunctionInput::GetOverrideNode() const
diff --git a/Engine/Plugins/FX/Niagara/Source/NiagaraEditor/Private/ViewModels/Stack/NiagaraStackModuleItem.cpp b/Engine/Plugins/FX/Niagara/Source/NiagaraEditor/Private/ViewModels/Stack/NiagaraStackModuleItem.cpp
--- a/Engine/Plugins/FX/Niagara/Source/NiagaraEditor/Private/ViewModels/Stack/NiagaraStackModuleItem.cpp
+++ b/Engine/Plugins/FX/Niagara/Source/NiagaraEditor/Private/ViewModels/Stack/NiagaraStackModuleItem.cpp
@@ -136,24 +136,6 @@ void UNiagaraStackModuleItem::Initialize(FRequiredEntryData InRequiredEntryData,
 		AddChildFilter(FOnFilterChild::CreateUObject(this, &UNiagaraStackModuleItem::FilterOutputCollection));
 		AddChildFilter(FOnFilterChild::CreateUObject(this, &UNiagaraStackModuleItem::FilterLinkedInputCollection));
 	}
-
-	// Update bCanMoveAndDelete
-	if (GetSystemViewModel()->GetEditMode() == ENiagaraSystemViewModelEditMode::EmitterAsset)
-	{
-		// When editing emitters all modules can be moved and deleted.
-		bCanMoveAndDelete = true;
-	}
-	else
-	{
-		// When editing systems only non-base modules can be moved and deleted.
-		TSharedRef<FNiagaraScriptMergeManager> MergeManager = FNiagaraScriptMergeManager::Get();
-
-		const UNiagaraEmitter* BaseEmitter = FNiagaraStackGraphUtilities::GetBaseEmitter(*GetEmitterViewModel()->GetEmitter(), GetSystemViewModel()->GetSystem());
-
-		bool bIsMergeable = MergeManager->IsMergeableScriptUsage(OutputNode->GetUsage());
-		bool bHasBaseModule = bIsMergeable && BaseEmitter != nullptr && MergeManager->HasBaseModule(*BaseEmitter, OutputNode->GetUsage(), OutputNode->GetUsageId(), FunctionCallNode->NodeGuid);
-		bCanMoveAndDelete = bHasBaseModule == false;
-	}
 }
 
 FText UNiagaraStackModuleItem::GetDisplayName() const
@@ -189,6 +171,7 @@ INiagaraStackItemGroupAddUtilities* UNiagaraStackModuleItem::GetGroupAddUtilitie
 void UNiagaraStackModuleItem::RefreshChildrenInternal(const TArray<UNiagaraStackEntry*>& CurrentChildren, TArray<UNiagaraStackEntry*>& NewChildren, TArray<FStackIssue>& NewIssues)
 {
 	bCanRefresh = false;
+	bCanMoveAndDeleteCache.Reset();
 
 	if (FunctionCallNode != nullptr && FunctionCallNode->ScriptIsValid())
 	{
@@ -703,7 +686,26 @@ void UNiagaraStackModuleItem::RefreshIsEnabled()
 
 bool UNiagaraStackModuleItem::CanMoveAndDelete() const
 {
-	return bCanMoveAndDelete;
+	if (bCanMoveAndDeleteCache.IsSet() == false)
+	{
+		if (HasBaseEmitter() == false)
+		{
+			// If there is no base emitter all modules can be moved and deleted.
+			bCanMoveAndDeleteCache = true;
+		}
+		else
+		{
+			// When editing systems only non-base modules can be moved and deleted.
+			TSharedRef<FNiagaraScriptMergeManager> MergeManager = FNiagaraScriptMergeManager::Get();
+
+			const UNiagaraEmitter* BaseEmitter = FNiagaraStackGraphUtilities::GetBaseEmitter(*GetEmitterViewModel()->GetEmitter(), GetSystemViewModel()->GetSystem());
+
+			bool bIsMergeable = MergeManager->IsMergeableScriptUsage(OutputNode->GetUsage());
+			bool bHasBaseModule = bIsMergeable && BaseEmitter != nullptr && MergeManager->HasBaseModule(*BaseEmitter, OutputNode->GetUsage(), OutputNode->GetUsageId(), FunctionCallNode->NodeGuid);
+			bCanMoveAndDeleteCache = bHasBaseModule == false;
+		}
+	}
+	return bCanMoveAndDeleteCache.GetValue();
 }
 
 bool UNiagaraStackModuleItem::CanRefresh() const
diff --git a/Engine/Plugins/FX/Niagara/Source/NiagaraEditor/Private/ViewModels/Stack/NiagaraStackRendererItem.cpp b/Engine/Plugins/FX/Niagara/Source/NiagaraEditor/Private/ViewModels/Stack/NiagaraStackRendererItem.cpp
--- a/Engine/Plugins/FX/Niagara/Source/NiagaraEditor/Private/ViewModels/Stack/NiagaraStackRendererItem.cpp
+++ b/Engine/Plugins/FX/Niagara/Source/NiagaraEditor/Private/ViewModels/Stack/NiagaraStackRendererItem.cpp
@@ -39,17 +39,6 @@ void UNiagaraStackRendererItem::Initialize(FRequiredEntryData InRequiredEntryDat
 	Super::Initialize(InRequiredEntryData, RendererStackEditorDataKey);
 	RendererProperties = InRendererProperties;
 	RendererProperties->OnChanged().AddUObject(this, &UNiagaraStackRendererItem::RendererChanged);
-
-	if (GetSystemViewModel()->GetEditMode() == ENiagaraSystemViewModelEditMode::EmitterAsset)
-	{
-		bHasBaseRenderer = false;
-	}
-	else
-	{
-		TSharedRef<FNiagaraScriptMergeManager> MergeManager = FNiagaraScriptMergeManager::Get();
-		const UNiagaraEmitter* BaseEmitter = FNiagaraStackGraphUtilities::GetBaseEmitter(*GetEmitterViewModel()->GetEmitter(), GetSystemViewModel()->GetSystem());
-		bHasBaseRenderer = BaseEmitter != nullptr && MergeManager->HasBaseRenderer(*BaseEmitter, RendererProperties->GetMergeId());
-	}
 }
 
 void UNiagaraStackRendererItem::FinalizeInternal()
@@ -157,7 +146,7 @@ FText UNiagaraStackRendererItem::GetDisplayName() const
 
 bool UNiagaraStackRendererItem::CanDelete() const
 {
-	return bHasBaseRenderer == false;
+	return HasBaseRenderer() == false;
 }
 
 void UNiagaraStackRendererItem::Delete()
@@ -172,22 +161,32 @@ void UNiagaraStackRendererItem::Delete()
 	ModifiedGroupItemsDelegate.ExecuteIfBound();
 }
 
-bool UNiagaraStackRendererItem::CanHaveBase() const
+bool UNiagaraStackRendererItem::HasBaseRenderer() const
 {
-	return GetSystemViewModel()->GetEditMode() == ENiagaraSystemViewModelEditMode::SystemAsset;
+	if (HasBaseEmitter())
+	{
+		if (bHasBaseRendererCache.IsSet() == false)
+		{
+			TSharedRef<FNiagaraScriptMergeManager> MergeManager = FNiagaraScriptMergeManager::Get();
+			const UNiagaraEmitter* BaseEmitter = FNiagaraStackGraphUtilities::GetBaseEmitter(*GetEmitterViewModel()->GetEmitter(), GetSystemViewModel()->GetSystem());
+			bHasBaseRendererCache = BaseEmitter != nullptr && MergeManager->HasBaseRenderer(*BaseEmitter, RendererProperties->GetMergeId());
+		}
+		return bHasBaseRendererCache.GetValue();
+	}
+	return false;
 }
 
 bool UNiagaraStackRendererItem::CanResetToBase() const
 {
-	if (CanHaveBase())
+	if (HasBaseRenderer())
 	{
-		if (bCanResetToBase.IsSet() == false)
+		if (bCanResetToBaseCache.IsSet() == false)
 		{
 			TSharedRef<FNiagaraScriptMergeManager> MergeManager = FNiagaraScriptMergeManager::Get();
 			const UNiagaraEmitter* BaseEmitter = FNiagaraStackGraphUtilities::GetBaseEmitter(*GetEmitterViewModel()->GetEmitter(), GetSystemViewModel()->GetSystem());
-			bCanResetToBase = BaseEmitter != nullptr && MergeManager->IsRendererDifferentFromBase(*GetEmitterViewModel()->GetEmitter(), *BaseEmitter, RendererProperties->GetMergeId());
+			bCanResetToBaseCache = BaseEmitter != nullptr && MergeManager->IsRendererDifferentFromBase(*GetEmitterViewModel()->GetEmitter(), *BaseEmitter, RendererProperties->GetMergeId());
 		}
-		return bCanResetToBase.GetValue();
+		return bCanResetToBaseCache.GetValue();
 	}
 	return false;
 }
@@ -226,7 +225,8 @@ void UNiagaraStackRendererItem::RefreshChildrenInternal(const TArray<UNiagaraSta
 
 	NewChildren.Add(RendererObject);
 	MissingAttributes = GetMissingVariables(RendererProperties.Get(), GetEmitterViewModel()->GetEmitter());
-	bCanResetToBase.Reset();
+	bHasBaseRendererCache.Reset();
+	bCanResetToBaseCache.Reset();
 	Super::RefreshChildrenInternal(CurrentChildren, NewChildren, NewIssues);
 	
 	RefreshIssues(NewIssues);
@@ -283,7 +283,7 @@ void UNiagaraStackRendererItem::RefreshIssues(TArray<FStackIssue>& NewIssues)
 
 void UNiagaraStackRendererItem::RendererChanged()
 {
-	bCanResetToBase.Reset();
+	bCanResetToBaseCache.Reset();
 }
 
 #undef LOCTEXT_NAMESPACE
diff --git a/Engine/Plugins/FX/Niagara/Source/NiagaraEditor/Private/ViewModels/Stack/NiagaraStackViewModel.cpp b/Engine/Plugins/FX/Niagara/Source/NiagaraEditor/Private/ViewModels/Stack/NiagaraStackViewModel.cpp
--- a/Engine/Plugins/FX/Niagara/Source/NiagaraEditor/Private/ViewModels/Stack/NiagaraStackViewModel.cpp
+++ b/Engine/Plugins/FX/Niagara/Source/NiagaraEditor/Private/ViewModels/Stack/NiagaraStackViewModel.cpp
@@ -10,6 +10,8 @@
 #include "NiagaraSystemEditorData.h"
 #include "NiagaraEmitterEditorData.h"
 #include "NiagaraStackEditorData.h"
+#include "NiagaraSystem.h"
+#include "NiagaraEmitterHandle.h"
 #include "NiagaraEmitter.h"
 #include "ViewModels/Stack/NiagaraStackItemGroup.h"
 #include "ViewModels/Stack/NiagaraStackItem.h"
@@ -180,6 +182,21 @@ bool UNiagaraStackViewModel::HasDismissedStackIssues()
 		|| GetEmitterHandleViewModel()->GetEmitterViewModel()->GetEditorData().GetStackEditorData().GetDismissedStackIssueIds().Num() > 0;
 }
 
+bool UNiagaraStackViewModel::HasEmitterSource() const
+{
+	return EmitterHandleViewModel->GetEmitterHandle()->GetSource() != nullptr;
+}
+
+void UNiagaraStackViewModel::RemoveEmitterSource()
+{
+	{
+		FScopedTransaction ScopedTransaction(LOCTEXT("RemoveEmitterSourceTransaction", "Remove Emitter Source"));
+		SystemViewModel->GetSystem().Modify();
+		EmitterHandleViewModel->GetEmitterHandle()->RemoveSource();
+	}
+	RootEntry->RefreshChildren();
+}
+
 void UNiagaraStackViewModel::CollapseToHeadersRecursive(TArray<UNiagaraStackEntry*> Entries)
 {
 	for (UNiagaraStackEntry* Entry : Entries)
diff --git a/Engine/Plugins/FX/Niagara/Source/NiagaraEditor/Public/ViewModels/Stack/NiagaraStackEmitterSpawnScriptItemGroup.h b/Engine/Plugins/FX/Niagara/Source/NiagaraEditor/Public/ViewModels/Stack/NiagaraStackEmitterSpawnScriptItemGroup.h
--- a/Engine/Plugins/FX/Niagara/Source/NiagaraEditor/Public/ViewModels/Stack/NiagaraStackEmitterSpawnScriptItemGroup.h
+++ b/Engine/Plugins/FX/Niagara/Source/NiagaraEditor/Public/ViewModels/Stack/NiagaraStackEmitterSpawnScriptItemGroup.h
@@ -34,7 +34,7 @@ private:
 	void EmitterPropertiesChanged();
 
 private:
-	mutable TOptional<bool> bCanResetToBase;
+	mutable TOptional<bool> bCanResetToBaseCache;
 
 	TWeakObjectPtr<UNiagaraEmitter> Emitter;
 
diff --git a/Engine/Plugins/FX/Niagara/Source/NiagaraEditor/Public/ViewModels/Stack/NiagaraStackEntry.h b/Engine/Plugins/FX/Niagara/Source/NiagaraEditor/Public/ViewModels/Stack/NiagaraStackEntry.h
--- a/Engine/Plugins/FX/Niagara/Source/NiagaraEditor/Public/ViewModels/Stack/NiagaraStackEntry.h
+++ b/Engine/Plugins/FX/Niagara/Source/NiagaraEditor/Public/ViewModels/Stack/NiagaraStackEntry.h
@@ -250,6 +250,8 @@ public:
 
 	void SetIsSearchResult(bool bInIsSearchResult);
 
+	bool HasBaseEmitter() const;
+
 protected:
 	virtual void BeginDestroy() override;
 
@@ -329,4 +331,6 @@ private:
 	bool bIsFinalized;
 
 	bool bIsSearchResult;
+
+	mutable TOptional<bool> bHasBaseEmitterCache;
 };
\ No newline at end of file
diff --git a/Engine/Plugins/FX/Niagara/Source/NiagaraEditor/Public/ViewModels/Stack/NiagaraStackEventScriptItemGroup.h b/Engine/Plugins/FX/Niagara/Source/NiagaraEditor/Public/ViewModels/Stack/NiagaraStackEventScriptItemGroup.h
--- a/Engine/Plugins/FX/Niagara/Source/NiagaraEditor/Public/ViewModels/Stack/NiagaraStackEventScriptItemGroup.h
+++ b/Engine/Plugins/FX/Niagara/Source/NiagaraEditor/Public/ViewModels/Stack/NiagaraStackEventScriptItemGroup.h
@@ -37,12 +37,14 @@ private:
 
 	void SelectEmitterStackObjectRootTreeNodes(TArray<TSharedRef<IDetailTreeNode>> Source, TArray<TSharedRef<IDetailTreeNode>>* Selected);
 
+	bool HasBaseEventHandler() const;
+
 private:
 	FGuid EventScriptUsageId;
 
-	bool bHasBaseEventHandler;
+	mutable TOptional<bool> bHasBaseEventHandlerCache;
 
-	mutable TOptional<bool> bCanResetToBase;
+	mutable TOptional<bool> bCanResetToBaseCache;
 
 	TWeakObjectPtr<UNiagaraEmitter> Emitter;
 
@@ -77,9 +79,12 @@ protected:
 
 	virtual bool CanDelete() const override;
 	virtual bool Delete() override;
+
+private:
+	bool HasBaseEventHandler() const;
 	
 private:
-	bool bHasBaseEventHandler;
+	mutable TOptional<bool> bHasBaseEventHandlerCache;
 
 	UPROPERTY()
 	UNiagaraStackEventHandlerPropertiesItem* EventHandlerProperties;
diff --git a/Engine/Plugins/FX/Niagara/Source/NiagaraEditor/Public/ViewModels/Stack/NiagaraStackFunctionInput.h b/Engine/Plugins/FX/Niagara/Source/NiagaraEditor/Public/ViewModels/Stack/NiagaraStackFunctionInput.h
--- a/Engine/Plugins/FX/Niagara/Source/NiagaraEditor/Public/ViewModels/Stack/NiagaraStackFunctionInput.h
+++ b/Engine/Plugins/FX/Niagara/Source/NiagaraEditor/Public/ViewModels/Stack/NiagaraStackFunctionInput.h
@@ -144,9 +144,6 @@ public:
 	/** Determine if this field is editable */
 	bool IsEditable() const;
 
-	/** Whether or not this input has a base value.  This is true for emitter instances in systems. */
-	bool EmitterHasBase() const;
-
 	/** Whether or not this input can be reset to a base value. */
 	bool CanResetToBase() const;
 
@@ -408,10 +405,10 @@ private:
 	mutable TOptional<UEdGraphPin*> OverridePinCache;
 
 	/** Whether or not this input can be reset to its default value. */
-	mutable TOptional<bool> bCanReset;
+	mutable TOptional<bool> bCanResetCache;
 
 	/** Whether or not this input can be reset to a base value defined by a parent emitter. */
-	mutable TOptional<bool> bCanResetToBase;
+	mutable TOptional<bool> bCanResetToBaseCache;
 
 	/** A flag to prevent handling graph changes when it's being updated directly by this object. */
 	bool bUpdatingGraphDirectly;
diff --git a/Engine/Plugins/FX/Niagara/Source/NiagaraEditor/Public/ViewModels/Stack/NiagaraStackModuleItem.h b/Engine/Plugins/FX/Niagara/Source/NiagaraEditor/Public/ViewModels/Stack/NiagaraStackModuleItem.h
--- a/Engine/Plugins/FX/Niagara/Source/NiagaraEditor/Public/ViewModels/Stack/NiagaraStackModuleItem.h
+++ b/Engine/Plugins/FX/Niagara/Source/NiagaraEditor/Public/ViewModels/Stack/NiagaraStackModuleItem.h
@@ -82,7 +82,7 @@ private:
 private:
 	UNiagaraNodeOutput* OutputNode;
 	UNiagaraNodeFunctionCall* FunctionCallNode;
-	bool bCanMoveAndDelete;
+	mutable TOptional<bool> bCanMoveAndDeleteCache;
 	bool bIsEnabled;
 	bool bCanRefresh;
 
diff --git a/Engine/Plugins/FX/Niagara/Source/NiagaraEditor/Public/ViewModels/Stack/NiagaraStackRendererItem.h b/Engine/Plugins/FX/Niagara/Source/NiagaraEditor/Public/ViewModels/Stack/NiagaraStackRendererItem.h
--- a/Engine/Plugins/FX/Niagara/Source/NiagaraEditor/Public/ViewModels/Stack/NiagaraStackRendererItem.h
+++ b/Engine/Plugins/FX/Niagara/Source/NiagaraEditor/Public/ViewModels/Stack/NiagaraStackRendererItem.h
@@ -28,7 +28,7 @@ public:
 
 	void Delete();
 
-	bool CanHaveBase() const;
+	bool HasBaseRenderer() const;
 
 	bool CanResetToBase() const;
 
@@ -52,9 +52,9 @@ private:
 private:
 	TWeakObjectPtr<UNiagaraRendererProperties> RendererProperties;
 
-	bool bHasBaseRenderer;
+	mutable TOptional<bool> bHasBaseRendererCache;
 
-	mutable TOptional<bool> bCanResetToBase;
+	mutable TOptional<bool> bCanResetToBaseCache;
 
 	TArray<FNiagaraVariable> MissingAttributes;
 
diff --git a/Engine/Plugins/FX/Niagara/Source/NiagaraEditor/Public/ViewModels/Stack/NiagaraStackViewModel.h b/Engine/Plugins/FX/Niagara/Source/NiagaraEditor/Public/ViewModels/Stack/NiagaraStackViewModel.h
--- a/Engine/Plugins/FX/Niagara/Source/NiagaraEditor/Public/ViewModels/Stack/NiagaraStackViewModel.h
+++ b/Engine/Plugins/FX/Niagara/Source/NiagaraEditor/Public/ViewModels/Stack/NiagaraStackViewModel.h
@@ -83,6 +83,9 @@ public:
 
 	bool HasDismissedStackIssues();
 
+	bool HasEmitterSource() const;
+	void RemoveEmitterSource();
+
 private:
 
 	/** Recursively Expands all groups and collapses all items in the stack. */
diff --git a/Engine/Plugins/FX/Niagara/Source/NiagaraEditorWidgets/Private/SNiagaraStack.cpp b/Engine/Plugins/FX/Niagara/Source/NiagaraEditorWidgets/Private/SNiagaraStack.cpp
--- a/Engine/Plugins/FX/Niagara/Source/NiagaraEditorWidgets/Private/SNiagaraStack.cpp
+++ b/Engine/Plugins/FX/Niagara/Source/NiagaraEditorWidgets/Private/SNiagaraStack.cpp
@@ -136,6 +136,17 @@ void SNiagaraStack::Construct(const FArguments& InArgs, UNiagaraStackViewModel*
 			FMenuBuilder MenuBuilder(true, nullptr);
 			MenuBuilder.BeginSection("EmitterInlineMenuActions", LOCTEXT("EmitterActions", "Emitter Actions"));
 			{
+				FUIAction Action(FExecuteAction::CreateSP(this, &SNiagaraStack::SetEmitterEnabled, !StackViewModel->GetEmitterHandleViewModel()->GetIsEnabled()),
+					FCanExecuteAction(),
+					FIsActionChecked::CreateSP(this, &SNiagaraStack::CheckEmitterEnabledStatus, true));
+				MenuBuilder.AddMenuEntry(
+					LOCTEXT("IsEnabled", "Is Enabled"),
+					LOCTEXT("ToggleEmitterEnabledToolTip", "Toggle emitter enabled/disabled state"),
+					FSlateIcon(),
+					Action,
+					NAME_None,
+					EUserInterfaceActionType::Check);
+
 				if (!GetEmitterNameIsReadOnly()) // Only allow renaming local copies of Emitters in Systems
 				{
 					MenuBuilder.AddMenuEntry(
@@ -145,22 +156,19 @@ void SNiagaraStack::Construct(const FArguments& InArgs, UNiagaraStackViewModel*
 						FUIAction(FExecuteAction::CreateSP(InlineEditableTextBlock.Get(), &SInlineEditableTextBlock::EnterEditingMode)));
 				}
 
+				MenuBuilder.AddMenuEntry(
+					LOCTEXT("RemoveSourceEmitter", "Remove Source Emitter"),
+					LOCTEXT("RemoveSourceEmitterToolTip", "Removes source information from this emitter, preventing inheritance of any further changes."),
+					FSlateIcon(),
+					FUIAction(
+						FExecuteAction::CreateUObject(StackViewModel, &UNiagaraStackViewModel::RemoveEmitterSource),
+						FCanExecuteAction::CreateUObject(StackViewModel, &UNiagaraStackViewModel::HasEmitterSource)));
+
 				MenuBuilder.AddMenuEntry(
 					LOCTEXT("ShowEmitterInContentBrowser", "Show in Content Browser"),
 					LOCTEXT("ShowEmitterInContentBrowserToolTip", "Show the emitter in this stack in the Content Browser"),
 					FSlateIcon(),
 					FUIAction(FExecuteAction::CreateSP(this, &SNiagaraStack::ShowEmitterInContentBrowser))); 
-				
-				FUIAction Action(FExecuteAction::CreateSP(this, &SNiagaraStack::SetEmitterEnabled, !StackViewModel->GetEmitterHandleViewModel()->GetIsEnabled()),
-					FCanExecuteAction(),
-					FIsActionChecked::CreateSP(this, &SNiagaraStack::CheckEmitterEnabledStatus, true));
-				MenuBuilder.AddMenuEntry(
-					LOCTEXT("IsEnabled", "Is Enabled"),
-					LOCTEXT("ToggleEmitterEnabledToolTip", "Toggle emitter enabled/disabled state"),
-					FSlateIcon(),
-					Action,
-					NAME_None, 
-					EUserInterfaceActionType::Check);
 
 				MenuBuilder.AddMenuEntry(
 					LOCTEXT("CollapseStack", "Collapse All"),
@@ -564,10 +572,7 @@ bool SNiagaraStack::CanOpenSourceEmitter() const
 	if (StackViewModel && StackViewModel->GetEmitterHandleViewModel().IsValid() && StackViewModel->GetEmitterHandleViewModel()->GetEmitterHandle()
 		&& StackViewModel->GetEmitterHandleViewModel()->GetEmitterHandle()->GetSource())
 	{
-		if (StackViewModel->GetSystemViewModel()->GetEditMode() == ENiagaraSystemViewModelEditMode::SystemAsset)
-		{
-			return true;
-		}
+		return true;
 	}
 
 	return false;
diff --git a/Engine/Plugins/FX/Niagara/Source/NiagaraEditorWidgets/Private/Stack/SNiagaraStackFunctionInputValue.cpp b/Engine/Plugins/FX/Niagara/Source/NiagaraEditorWidgets/Private/Stack/SNiagaraStackFunctionInputValue.cpp
--- a/Engine/Plugins/FX/Niagara/Source/NiagaraEditorWidgets/Private/Stack/SNiagaraStackFunctionInputValue.cpp
+++ b/Engine/Plugins/FX/Niagara/Source/NiagaraEditorWidgets/Private/Stack/SNiagaraStackFunctionInputValue.cpp
@@ -703,7 +703,7 @@ FReply SNiagaraStackFunctionInputValue::ResetButtonPressed() const
 
 EVisibility SNiagaraStackFunctionInputValue::GetResetToBaseButtonVisibility() const
 {
-	if (FunctionInput->EmitterHasBase())
+	if (FunctionInput->HasBaseEmitter())
 	{
 		return FunctionInput->CanResetToBase() ? EVisibility::Visible : EVisibility::Hidden;
 	}
diff --git a/Engine/Plugins/FX/Niagara/Source/NiagaraEditorWidgets/Private/Stack/SNiagaraStackRendererItem.cpp b/Engine/Plugins/FX/Niagara/Source/NiagaraEditorWidgets/Private/Stack/SNiagaraStackRendererItem.cpp
--- a/Engine/Plugins/FX/Niagara/Source/NiagaraEditorWidgets/Private/Stack/SNiagaraStackRendererItem.cpp
+++ b/Engine/Plugins/FX/Niagara/Source/NiagaraEditorWidgets/Private/Stack/SNiagaraStackRendererItem.cpp
@@ -123,7 +123,7 @@ FReply SNiagaraStackRendererItem::DeleteClicked()
 
 EVisibility SNiagaraStackRendererItem::GetResetToBaseButtonVisibility() const
 {
-	if (RendererItem->CanHaveBase())
+	if (RendererItem->HasBaseEmitter())
 	{
 		return RendererItem->CanResetToBase() ? EVisibility::Visible : EVisibility::Hidden;
 	}
