Integrate FNiagaraMessageManager.

diff --git a/Engine/Plugins/FX/Niagara/Source/Niagara/Classes/NiagaraScript.h b/Engine/Plugins/FX/Niagara/Source/Niagara/Classes/NiagaraScript.h
--- a/Engine/Plugins/FX/Niagara/Source/Niagara/Classes/NiagaraScript.h
+++ b/Engine/Plugins/FX/Niagara/Source/Niagara/Classes/NiagaraScript.h
@@ -26,50 +26,6 @@ void ProcessSerializedShaderMaps(UNiagaraScript* Owner, TArray<FNiagaraShaderScr
 
 #define NIAGARA_SCRIPT_COMPILE_LOGGING_MEDIUM
 
-/** Defines the compile event types for translation/compilation.*/
-UENUM()
-enum class FNiagaraCompileEventSeverity : uint8
-{
-	Log = 0,
-	Warning = 1,
-	Error = 2
-};
-
-/** Records necessary information to give UI cues for errors/logs/warnings during compile.*/
-USTRUCT()
-struct FNiagaraCompileEvent
-{
-	GENERATED_USTRUCT_BODY()
-public:
-	FNiagaraCompileEvent()
-	{
-		Severity = FNiagaraCompileEventSeverity::Log;
-		Message = FString();
-		NodeGuid = FGuid();
-		PinGuid = FGuid();
-		StackGuids.Empty();
-	}
-
-	FNiagaraCompileEvent(FNiagaraCompileEventSeverity InSeverity, const FString& InMessage, FGuid InNodeGuid = FGuid(), FGuid InPinGuid = FGuid(), const TArray<FGuid>& InCallstackGuids = TArray<FGuid>())
-		: Severity(InSeverity), Message(InMessage), NodeGuid(InNodeGuid), PinGuid(InPinGuid), StackGuids(InCallstackGuids) {}
-
-	/** Whether or not this is an error, warning, or info*/
-	UPROPERTY()
-	FNiagaraCompileEventSeverity Severity;
-	/* The message itself*/
-	UPROPERTY()
-	FString Message;
-	/** The node guid that generated the compile event*/
-	UPROPERTY()
-	FGuid NodeGuid;
-	/** The pin persistent id that generated the compile event*/
-	UPROPERTY()
-	FGuid PinGuid;
-	/** The compile stack frame of node id's*/
-	UPROPERTY()
-	TArray<FGuid> StackGuids;
-};
-
 /** Defines what will happen to unused attributes when a script is run. */
 UENUM()
 enum class EUnusedAttributeBehaviour : uint8
diff --git a/Engine/Plugins/FX/Niagara/Source/NiagaraEditor/Private/NiagaraCompiler.cpp b/Engine/Plugins/FX/Niagara/Source/NiagaraEditor/Private/NiagaraCompiler.cpp
--- a/Engine/Plugins/FX/Niagara/Source/NiagaraEditor/Private/NiagaraCompiler.cpp
+++ b/Engine/Plugins/FX/Niagara/Source/NiagaraEditor/Private/NiagaraCompiler.cpp
@@ -29,7 +29,7 @@
 #include "Misc/FileHelper.h"
 #include "ShaderCompiler.h"
 #include "NiagaraShader.h"
-#include "NiagaraScript.h" //@todo(ng) FNiagaraCompileEvent dependency (make sure this is used before submitting)
+#include "NiagaraScript.h"
 
 #define LOCTEXT_NAMESPACE "NiagaraCompiler"
 
@@ -547,7 +547,7 @@ TSharedPtr<FNiagaraVMExecutableData> FNiagaraEditorModule::CompileScript(const F
 	TArray<FNiagaraCompileEvent> Messages;
 	if (TranslateResults.CompileEvents.Num() > 0)
 	{
-		Results.Data->LastCompileEvents.Append(TranslateResults.CompileEvents); //@todo(ng) get rid of Messages in the future
+		Results.Data->LastCompileEvents.Append(TranslateResults.CompileEvents); //@todo(message manager) clean up message send path
 		Messages.Append(TranslateResults.CompileEvents);
 	}
 	if (Results.CompileEvents.Num() > 0)
@@ -905,7 +905,7 @@ FHlslNiagaraCompiler::FHlslNiagaraCompiler()
 void FHlslNiagaraCompiler::Error(FText ErrorText)
 {
 	FString ErrorString = FString::Printf(TEXT("%s"), *ErrorText.ToString());
-	CompileResults.Data->LastCompileEvents.Add(FNiagaraCompileEvent(FNiagaraCompileEventSeverity::Error, ErrorString)); //@todo(ng) cut out the compileresults.compileevent member since we don't need to separately log into general output log
+	CompileResults.Data->LastCompileEvents.Add(FNiagaraCompileEvent(FNiagaraCompileEventSeverity::Error, ErrorString));
 	CompileResults.CompileEvents.Add(FNiagaraCompileEvent(FNiagaraCompileEventSeverity::Error, ErrorString));
 	CompileResults.NumErrors++;
 }
diff --git a/Engine/Plugins/FX/Niagara/Source/NiagaraEditor/Private/NiagaraEditorCommon.cpp b/Engine/Plugins/FX/Niagara/Source/NiagaraEditor/Private/NiagaraEditorCommon.cpp
--- a/Engine/Plugins/FX/Niagara/Source/NiagaraEditor/Private/NiagaraEditorCommon.cpp
+++ b/Engine/Plugins/FX/Niagara/Source/NiagaraEditor/Private/NiagaraEditorCommon.cpp
@@ -1278,3 +1278,8 @@ void UActorFactoryNiagara::PostCreateBlueprint(UObject* Asset, AActor* CDO)
 
 
 #undef LOCTEXT_NAMESPACE
+
+INiagaraScriptGraphFocusInfo::~INiagaraScriptGraphFocusInfo()
+{
+	//Stand-in definition for abstract INiagaraScriptGraphFocusInfo's pure virtual dtor
+}
diff --git a/Engine/Plugins/FX/Niagara/Source/NiagaraEditor/Private/NiagaraEditorModule.cpp b/Engine/Plugins/FX/Niagara/Source/NiagaraEditor/Private/NiagaraEditorModule.cpp
--- a/Engine/Plugins/FX/Niagara/Source/NiagaraEditor/Private/NiagaraEditorModule.cpp
+++ b/Engine/Plugins/FX/Niagara/Source/NiagaraEditor/Private/NiagaraEditorModule.cpp
@@ -91,7 +91,6 @@
 #include "NiagaraHlslTranslator.h"
 #include "NiagaraThumbnailRenderer.h"
 
-
 IMPLEMENT_MODULE( FNiagaraEditorModule, NiagaraEditor );
 
 #define LOCTEXT_NAMESPACE "NiagaraEditorModule"
@@ -583,11 +582,6 @@ FNiagaraEditorModule& FNiagaraEditorModule::Get()
 {
 	return FModuleManager::LoadModuleChecked<FNiagaraEditorModule>("NiagaraEditor");
 }
-void FNiagaraEditorModule::OnMessageLogTokenClicked(const TSharedRef<class IMessageToken>& Token)
-{
-	return; //@todo(ng) just stubbing this here for now
-}
-
 
 void FNiagaraEditorModule::OnNiagaraSettingsChangedEvent(const FString& PropertyName, const UNiagaraSettings* Settings)
 {
@@ -735,36 +729,4 @@ void FNiagaraEditorModule::OnPreGarbageCollection()
 	}
 }
 
-TSharedRef<FNiagaraAssetNameToken> FNiagaraAssetNameToken::Create(const FString& InAssetName, const FText& InMessage)
-{
-	return MakeShareable(new FNiagaraAssetNameToken(InAssetName, InMessage));
-}
-
-FNiagaraAssetNameToken::FNiagaraAssetNameToken(const FString& InAssetName, const FText& InMessage)
-	: AssetName(InAssetName)
-{
-	if (!InMessage.IsEmpty())
-	{
-		CachedText = InMessage;
-	}
-	else
-	{
-		CachedText = FText::FromString(AssetName);
-	}
-
-	MessageTokenActivated = FOnMessageTokenActivated::CreateStatic(&FNiagaraAssetNameToken::FindAndOpenAsset, AssetName);
-}
-
-
-void FNiagaraAssetNameToken::FindAndOpenAsset(const TSharedRef<IMessageToken>& Token, FString InAssetPath)
-{
-	FAssetRegistryModule& AssetRegistryModule = FModuleManager::LoadModuleChecked<FAssetRegistryModule>("AssetRegistry");
-	IAssetRegistry& AssetRegistry = AssetRegistryModule.Get();
-	FAssetData AssetData = AssetRegistry.GetAssetByObjectPath(*InAssetPath);
-	if (AssetData.IsValid())
-	{
-		FAssetEditorManager::Get().OpenEditorForAsset(AssetData.GetAsset(), EToolkitMode::Standalone);
-	}
-}
-
 #undef LOCTEXT_NAMESPACE
diff --git a/Engine/Plugins/FX/Niagara/Source/NiagaraEditor/Private/NiagaraMessageManager.cpp b/Engine/Plugins/FX/Niagara/Source/NiagaraEditor/Private/NiagaraMessageManager.cpp
new file mode 100644
--- /dev/null
+++ b/Engine/Plugins/FX/Niagara/Source/NiagaraEditor/Private/NiagaraMessageManager.cpp
@@ -0,0 +1,352 @@
+// Copyright 1998-2019 Epic Games, Inc. All Rights Reserved.
+
+#include "NiagaraMessageManager.h"
+#include "NiagaraScriptSource.h"
+#include "NiagaraNodeFunctionCall.h"
+#include "NiagaraNodeEmitter.h"
+#include "Modules/ModuleManager.h"
+#include "AssetRegistryModule.h"
+#include "IAssetRegistry.h"
+#include "Toolkits/AssetEditorManager.h"
+#include "NiagaraScriptToolkit.h"
+#include "../Public/NiagaraEditorModule.h"
+
+#define LOCTEXT_NAMESPACE "NiagaraMessageManager"
+
+FNiagaraMessageManager* FNiagaraMessageManager::Singleton = nullptr;
+TQueue<TSharedPtr<INiagaraMessageJob>> FNiagaraMessageManager::JobQueue;
+
+FNiagaraMessageManager::FNiagaraMessageManager()
+{
+}
+
+FNiagaraMessageManager* FNiagaraMessageManager::Get()
+{
+	if (Singleton == nullptr)
+	{
+		Singleton = new FNiagaraMessageManager();
+	}
+	return Singleton;
+}
+
+TSharedPtr<INiagaraMessage> FNiagaraMessageManager::QueueMessageJob(TSharedRef<INiagaraMessageJob> InMessageJob)
+{
+	//JobQueue.Enqueue(InMessageJob); //@todo(message manager) stubbing the asynchronous implementation
+	return InMessageJob->GenerateNiagaraMessage();
+}
+
+FNiagaraMessageCompileEvent::FNiagaraMessageCompileEvent(
+	  const FNiagaraCompileEvent& InCompileEvent
+	, TArray<FNiagaraScriptNameAndAssetPath>& InContextScriptNamesAndAssetPaths
+	, TOptional<const FString>& InEmitterScriptName
+	, TOptional<const FNiagaraScriptNameAndAssetPath>& InCompiledScriptNameAndAssetPath
+	)
+	: INiagaraMessage(ENiagaraMessageType::CompileEventMessage)
+	, CompileEvent(InCompileEvent)
+	, ContextScriptNamesAndAssetPaths(InContextScriptNamesAndAssetPaths)
+	, EmitterScriptName(InEmitterScriptName)
+	, CompiledScriptNameAndAssetPath(InCompiledScriptNameAndAssetPath)
+{
+}
+
+TSharedRef<FTokenizedMessage> FNiagaraMessageCompileEvent::GenerateTokenizedMessage() const
+{
+	EMessageSeverity::Type MessageSeverity = EMessageSeverity::Info;
+	switch (CompileEvent.Severity) {
+	case FNiagaraCompileEventSeverity::Error:
+		MessageSeverity = EMessageSeverity::Error;
+		break;
+	case FNiagaraCompileEventSeverity::Warning:
+		MessageSeverity = EMessageSeverity::Warning;
+		break;
+	case FNiagaraCompileEventSeverity::Log:
+		MessageSeverity = EMessageSeverity::Info;
+		break;
+	default:
+		ensureMsgf(false, TEXT("Compile event severity type not handled!"));
+		MessageSeverity = EMessageSeverity::Info;
+		break;
+	}
+	TSharedRef<FTokenizedMessage> Message = FTokenizedMessage::Create(MessageSeverity);
+
+	//Add message from compile event at start of message
+	if (CompiledScriptNameAndAssetPath.IsSet())
+	{
+		//Compile event is from a script asset
+		if (ContextScriptNamesAndAssetPaths.Num() == 0)
+		{
+
+			if (CompileEvent.PinGuid != FGuid())
+			{
+				Message->AddToken(FNiagaraCompileEventToken::Create(CompiledScriptNameAndAssetPath.GetValue().ScriptAssetPathString, FText::FromString(CompileEvent.Message), CompileEvent.NodeGuid, CompileEvent.PinGuid));
+			}
+			else if (CompileEvent.NodeGuid != FGuid())
+			{
+				Message->AddToken(FNiagaraCompileEventToken::Create(CompiledScriptNameAndAssetPath.GetValue().ScriptAssetPathString, FText::FromString(CompileEvent.Message), CompileEvent.NodeGuid));
+			}
+			else
+			{
+				Message->AddToken(FTextToken::Create(FText::FromString(CompileEvent.Message)));
+			}		
+		}
+		else
+		{
+			if (CompileEvent.PinGuid != FGuid())
+			{
+				Message->AddToken(FNiagaraCompileEventToken::Create(ContextScriptNamesAndAssetPaths.Last().ScriptAssetPathString, FText::FromString(CompileEvent.Message), CompileEvent.NodeGuid, CompileEvent.PinGuid));
+			}
+			else if (CompileEvent.NodeGuid != FGuid())
+			{
+				Message->AddToken(FNiagaraCompileEventToken::Create(ContextScriptNamesAndAssetPaths.Last().ScriptAssetPathString, FText::FromString(CompileEvent.Message), CompileEvent.NodeGuid));
+			}
+			else
+			{
+				Message->AddToken(FNiagaraCompileEventToken::Create(ContextScriptNamesAndAssetPaths.Last().ScriptAssetPathString, FText::FromString(CompileEvent.Message)));
+			}
+		}
+	}
+	else
+	{
+		//Compile event is from an emitter or system asset
+		if (ContextScriptNamesAndAssetPaths.Num() == 0)
+		{
+			Message->AddToken(FTextToken::Create(FText::FromString(CompileEvent.Message)));
+		}
+		else if (CompileEvent.PinGuid != FGuid())
+		{
+			Message->AddToken(FNiagaraCompileEventToken::Create(ContextScriptNamesAndAssetPaths.Last().ScriptAssetPathString, FText::FromString(CompileEvent.Message), CompileEvent.NodeGuid, CompileEvent.PinGuid));
+		}
+		else if (CompileEvent.NodeGuid != FGuid())
+		{
+			Message->AddToken(FNiagaraCompileEventToken::Create(ContextScriptNamesAndAssetPaths.Last().ScriptAssetPathString, FText::FromString(CompileEvent.Message), CompileEvent.NodeGuid));
+		}
+		else
+		{
+			Message->AddToken(FNiagaraCompileEventToken::Create(ContextScriptNamesAndAssetPaths.Last().ScriptAssetPathString, FText::FromString(CompileEvent.Message)));
+		}
+	}
+	
+	//Now add the emitter name if we passed through an Emitter Node
+	if (EmitterScriptName.IsSet())
+	{
+		Message->AddToken(FTextToken::Create(FText::FromString(*EmitterScriptName.GetValue())));
+	}
+
+	//Finally add the context stack of the scripts that were passed through to get to the originating graph
+	for (const FNiagaraScriptNameAndAssetPath& ScriptNameAndPath : ContextScriptNamesAndAssetPaths)
+	{
+		Message->AddToken(FNiagaraCompileEventToken::Create(ScriptNameAndPath.ScriptAssetPathString, FText::FromString(*ScriptNameAndPath.ScriptNameString)));
+	}
+	return Message;
+}
+
+FNiagaraMessageJobCompileEvent::FNiagaraMessageJobCompileEvent(const FNiagaraCompileEvent& InCompileEvent, const TWeakObjectPtr<UNiagaraScript>& InOriginatingScriptWeakObjPtr, const bool bInFromScriptToolkit)
+	: CompileEvent(InCompileEvent)
+	, OriginatingScriptWeakObjPtr(InOriginatingScriptWeakObjPtr)
+	, bFromScriptToolkit(bInFromScriptToolkit)
+{
+}
+
+const TSharedPtr<INiagaraMessage> FNiagaraMessageJobCompileEvent::GenerateNiagaraMessage() const
+{
+	TArray<FGuid> ContextStackGuids = CompileEvent.StackGuids;
+
+	if (OriginatingScriptWeakObjPtr.IsValid())
+	{
+		UNiagaraScriptSourceBase* FunctionScriptSourceBase = OriginatingScriptWeakObjPtr->GetSource();
+		checkf(FunctionScriptSourceBase->IsA<UNiagaraScriptSource>(), TEXT("Script source for function call node is not assigned or is not of type UNiagaraScriptSource!"))
+		UNiagaraScriptSource* FunctionScriptSource = Cast<UNiagaraScriptSource>(FunctionScriptSourceBase);
+		checkf(FunctionScriptSource, TEXT("Script source base was somehow not a derived type!"));
+		UNiagaraGraph* ScriptGraph = FunctionScriptSource->NodeGraph;
+		checkf(ScriptGraph, TEXT("Function Script does not have a UNiagaraGraph!"));
+
+		TArray<FNiagaraScriptNameAndAssetPath> ContextScriptNamesAndPaths;
+		TOptional<const FString> EmitterScriptName;
+		TOptional<const FNiagaraScriptNameAndAssetPath> CompiledScriptNameAndPath;
+		TOptional<const FText> ScriptNameAndPathsGetterFailureReason;
+		const bool bSuccessfullyFoundScripts = RecursiveGetScriptNamesAndAssetPathsFromContextStack(ContextStackGuids, ScriptGraph, ContextScriptNamesAndPaths, EmitterScriptName, ScriptNameAndPathsGetterFailureReason);
+
+		if (bSuccessfullyFoundScripts == false)
+		{
+			//If we can't find the scripts in the context stack of the compile event, return a message with the error and ask for a recompile.
+			return MakeShared<FNiagaraMessageNeedRecompile>(ScriptNameAndPathsGetterFailureReason.GetValue());
+		}
+
+		if (bFromScriptToolkit)
+		{
+			//If this compile event is from a script, set the compiled script name and asset path so the user can navigate to errors locally.
+			CompiledScriptNameAndPath = TOptional<const FNiagaraScriptNameAndAssetPath>(FNiagaraScriptNameAndAssetPath(OriginatingScriptWeakObjPtr->GetName(), OriginatingScriptWeakObjPtr->GetPathName()));
+		}
+		return MakeShared<FNiagaraMessageCompileEvent>(CompileEvent, ContextScriptNamesAndPaths, EmitterScriptName, CompiledScriptNameAndPath);
+	}
+	//The originating script weak ptr is no longer valid, send an error message asking for recompile.
+	FText MessageText = LOCTEXT("CompileEventMessageJobFail", "Cached info for compile event is out of date, recompile to get full info. Event: {0}");
+	//Add in the message of the compile event for visibility on which compile events do not have info.
+	FText::Format(MessageText, FText::FromString(CompileEvent.Message));
+	return MakeShared<FNiagaraMessageNeedRecompile>(MessageText);
+}
+
+const bool FNiagaraMessageJobCompileEvent::RecursiveGetScriptNamesAndAssetPathsFromContextStack(
+	  TArray<FGuid>& InContextStackNodeGuids
+	, const UNiagaraGraph* InGraphToSearch
+	, TArray<FNiagaraScriptNameAndAssetPath>& OutContextScriptNamesAndAssetPaths
+	, TOptional<const FString>& OutEmitterScriptName
+	, TOptional<const FText>& OutFailureReason
+) const
+{
+	checkf(InGraphToSearch, TEXT("Failed to get a node graph to search!"));
+
+	if (InGraphToSearch && InContextStackNodeGuids.Num() == 0)
+	{
+		//StackGuids arr has been cleared out which means we have walked the entire context stack.
+		return true;
+	}
+
+	// Search in the current graph for a node with a GUID that matches a GUID in the list of Function Call and Emitter node GUIDs that define the context stack for a compile event
+	auto FindNodeInGraphWithContextStackGuid = [&InGraphToSearch, &InContextStackNodeGuids]()->TOptional<UEdGraphNode*> {
+		for (UEdGraphNode* GraphNode : InGraphToSearch->Nodes)
+		{
+			for (int i = 0; i < InContextStackNodeGuids.Num(); i++)
+			{
+				if (GraphNode->NodeGuid == InContextStackNodeGuids[i])
+				{
+					InContextStackNodeGuids.RemoveAt(i);
+					return GraphNode;
+				}
+			}
+		}
+		return TOptional<UEdGraphNode*>();
+	};
+
+	TOptional<UEdGraphNode*> ContextNode = FindNodeInGraphWithContextStackGuid();
+	if (ContextNode.IsSet())
+	{
+		// found a node in the current graph that has a GUID in the context list
+
+		UNiagaraNodeFunctionCall* FunctionCallNode = Cast<UNiagaraNodeFunctionCall>(ContextNode.GetValue());
+		if (FunctionCallNode)
+		{
+			// node is a function call node, now get the Niagara Script assigned to this node, add a message token and recursively call into the graph of that script.
+			UNiagaraScript* FunctionCallNodeAssignedScript = FunctionCallNode->FunctionScript;
+			if (FunctionCallNodeAssignedScript == false)
+			{
+				FText FailureReason = LOCTEXT("CompileEventMessageGenerator_FunctionCallNodeScriptNotFound", "Script for Function Call Node \"{0}\" not found!");
+				OutFailureReason = FText::Format(FailureReason, FText::FromString(FunctionCallNode->GetFunctionName()));
+				return false;
+			}
+			UNiagaraScriptSourceBase* FunctionScriptSourceBase = FunctionCallNodeAssignedScript->GetSource();
+			if (FunctionScriptSourceBase == false)
+			{
+				FText FailureReason = LOCTEXT("CompileEventMessageGenerator_FunctionCallNodeScriptSourceBaseNotFound", "Source Script for Function Call Node \"{0}\" not found!");
+				OutFailureReason = FText::Format(FailureReason, FText::FromString(FunctionCallNode->GetFunctionName()));
+				return false;
+			}
+			UNiagaraScriptSource* FunctionScriptSource = Cast<UNiagaraScriptSource>(FunctionScriptSourceBase);
+			checkf(FunctionScriptSource, TEXT("Script source base was somehow not a derived type!"));
+
+			UNiagaraGraph* FunctionScriptGraph = FunctionScriptSource->NodeGraph;
+			if (FunctionScriptGraph == false)
+			{
+				FText FailureReason = LOCTEXT("CompileEventMessageGenerator_FunctionCallNodeGraphNotFound", "Graph for Function Call Node \"{0}\" not found!");
+				OutFailureReason = FText::Format(FailureReason, FText::FromString(FunctionCallNode->GetFunctionName()));
+				return false;
+			}
+			OutContextScriptNamesAndAssetPaths.Add(FNiagaraScriptNameAndAssetPath(FunctionCallNodeAssignedScript->GetName(), FunctionCallNodeAssignedScript->GetPathName()));
+			return RecursiveGetScriptNamesAndAssetPathsFromContextStack(InContextStackNodeGuids, FunctionScriptGraph, OutContextScriptNamesAndAssetPaths, OutEmitterScriptName, OutFailureReason);
+		}
+
+		UNiagaraNodeEmitter* EmitterNode = Cast<UNiagaraNodeEmitter>(ContextNode.GetValue());
+		if (EmitterNode)
+		{
+			// node is an emitter node, now get the Emitter name, add a message token and recursively call into the graph of that emitter.
+			UNiagaraScriptSource* EmitterScriptSource = EmitterNode->GetScriptSource();
+			checkf(EmitterScriptSource, TEXT("Emitter Node does not have a Script Source!"));
+			UNiagaraGraph* EmitterScriptGraph = EmitterScriptSource->NodeGraph;
+			checkf(EmitterScriptGraph, TEXT("Emitter Script Source does not have a UNiagaraGraph!"));
+
+			OutEmitterScriptName = EmitterNode->GetEmitterUniqueName();
+			return RecursiveGetScriptNamesAndAssetPathsFromContextStack(InContextStackNodeGuids, EmitterScriptGraph, OutContextScriptNamesAndAssetPaths, OutEmitterScriptName, OutFailureReason);
+		}
+		checkf(false, TEXT("Matching node is not a function call or emitter node!"));
+	}
+	FText FailureReason = LOCTEXT("CompileEventMessageGenerator_CouldNotFindMatchingNodeGUID", "Failed to walk the entire context stack, is this compile event out of date ? Event: '{0}'");
+	FText::Format(FailureReason, FText::FromString(CompileEvent.Message));
+	OutFailureReason = FailureReason;
+	return false;
+}
+
+FNiagaraMessageNeedRecompile::FNiagaraMessageNeedRecompile(const FText& InNeedRecompileMessage)
+	: INiagaraMessage(ENiagaraMessageType::NeedRecompileMessage)
+	, NeedRecompileMessage(InNeedRecompileMessage)
+{
+}
+
+TSharedRef<FTokenizedMessage> FNiagaraMessageNeedRecompile::GenerateTokenizedMessage() const
+{
+	return FTokenizedMessage::Create(EMessageSeverity::Error, NeedRecompileMessage);
+}
+
+FNiagaraCompileEventToken::FNiagaraCompileEventToken(
+	  const FString& InScriptAssetPath
+	, const FText& InMessage
+	, const TOptional<const FGuid>& InNodeGUID
+	, const TOptional<const FGuid>& InPinGUID)
+	: ScriptAssetPath(InScriptAssetPath)
+	, NodeGUID(InNodeGUID)
+	, PinGUID(InPinGUID)
+{
+	if (!InMessage.IsEmpty())
+	{
+		CachedText = InMessage;
+	}
+	else
+	{
+		CachedText = FText::FromString(InScriptAssetPath);
+	}
+
+	MessageTokenActivated = FOnMessageTokenActivated::CreateStatic(&FNiagaraCompileEventToken::OpenScriptAssetByPathAndFocusNodeOrPinIfSet, ScriptAssetPath, NodeGUID, PinGUID);
+}
+
+void FNiagaraCompileEventToken::OpenScriptAssetByPathAndFocusNodeOrPinIfSet(
+	  const TSharedRef<IMessageToken>& Token
+	, FString InScriptAssetPath
+	, const TOptional<const FGuid> InNodeGUID
+	, const TOptional<const FGuid> InPinGUID)
+{
+	FAssetRegistryModule& AssetRegistryModule = FModuleManager::LoadModuleChecked<FAssetRegistryModule>("AssetRegistry");
+	IAssetRegistry& AssetRegistry = AssetRegistryModule.Get();
+	FAssetData AssetData = AssetRegistry.GetAssetByObjectPath(*InScriptAssetPath);
+	if (AssetData.IsValid())
+	{
+		if (UNiagaraScript* ScriptAsset = Cast<UNiagaraScript>(AssetData.GetAsset()))
+		{
+			FAssetEditorManager::Get().OpenEditorForAsset(ScriptAsset, EToolkitMode::Standalone);
+			FNiagaraEditorModule& NiagaraEditorModule = FModuleManager::LoadModuleChecked<FNiagaraEditorModule>("NiagaraEditor");
+			if (InPinGUID.IsSet())
+			{
+				FNiagaraScriptGraphPinToFocusInfo PinToFocusInfo = FNiagaraScriptGraphPinToFocusInfo(ScriptAsset->GetUniqueID(), InPinGUID.GetValue());
+				NiagaraEditorModule.GetOnScriptToolkitsShouldFocusGraphElement().Broadcast(&PinToFocusInfo);
+			}
+			else if (InNodeGUID.IsSet())
+			{
+				FNiagaraScriptGraphNodeToFocusInfo NodeToFocusInfo = FNiagaraScriptGraphNodeToFocusInfo(ScriptAsset->GetUniqueID(), InNodeGUID.GetValue());
+				NiagaraEditorModule.GetOnScriptToolkitsShouldFocusGraphElement().Broadcast(&NodeToFocusInfo);
+			}
+		}
+		else
+		{
+			checkf(false, TEXT("CompileEventJob referenced asset was not a UNiagaraScript!"));
+		}
+	}
+}
+
+TSharedRef<FNiagaraCompileEventToken> FNiagaraCompileEventToken::Create(
+	  const FString& InScriptAssetPath
+	, const FText& InMessage
+	, const TOptional<const FGuid>& InNodeGUID /*= TOptional<const FGuid>() */
+	, const TOptional<const FGuid>& InPinGUID /*= TOptional<const FGuid>()*/ )
+{
+	return MakeShareable(new FNiagaraCompileEventToken(InScriptAssetPath, InMessage, InNodeGUID, InPinGUID));
+}
+
+#undef LOCTEXT_NAMESPACE //NiagaraMessageManager
diff --git a/Engine/Plugins/FX/Niagara/Source/NiagaraEditor/Private/Toolkits/NiagaraScriptToolkit.cpp b/Engine/Plugins/FX/Niagara/Source/NiagaraEditor/Private/Toolkits/NiagaraScriptToolkit.cpp
--- a/Engine/Plugins/FX/Niagara/Source/NiagaraEditor/Private/Toolkits/NiagaraScriptToolkit.cpp
+++ b/Engine/Plugins/FX/Niagara/Source/NiagaraEditor/Private/Toolkits/NiagaraScriptToolkit.cpp
@@ -44,8 +44,7 @@
 #include "Developer/MessageLog/Public/MessageLogModule.h"
 
 #include "AssetTypeActions/AssetTypeActions_NiagaraScript.h"
-#include "NiagaraNodeFunctionCall.h" //@todo(ng) move this to message manager
-#include "NiagaraNodeEmitter.h" //@todo(ng) move this to message manager 
+#include "NiagaraMessageManager.h"
 
 #define LOCTEXT_NAMESPACE "NiagaraScriptToolkit"
 
@@ -71,119 +70,9 @@ FNiagaraScriptToolkit::~FNiagaraScriptToolkit()
 {
 	EditedNiagaraScript->OnVMScriptCompiled().RemoveAll(this);
 	ScriptViewModel->GetGraphViewModel()->GetGraph()->RemoveOnGraphNeedsRecompileHandler(OnEditedScriptGraphChangedHandle);
-}
-
-//@todo(ng) move to message manager class
-TSharedRef<FTokenizedMessage> FNiagaraScriptToolkit::BuildMessageForCompileEvent(FNiagaraCompileEvent& InCompileEvent)
-{
-	TArray<FGuid>& ContextStackGuids = InCompileEvent.StackGuids;
-
-	UNiagaraScriptSourceBase* FunctionScriptSourceBase = EditedNiagaraScript->GetSource();
-	checkf(FunctionScriptSourceBase->IsA<UNiagaraScriptSource>(), TEXT("Script source for function call node is not assigned or is not of type UNiagaraScriptSource!"))
-		UNiagaraScriptSource* FunctionScriptSource = Cast<UNiagaraScriptSource>(FunctionScriptSourceBase);
-	checkf(FunctionScriptSource, TEXT("Script source base was somehow not a derived type!"));
-	UNiagaraGraph* ThisScriptGraph = FunctionScriptSource->NodeGraph;
-	checkf(ThisScriptGraph, TEXT("Function Script does not have a UNiagaraGraph!"));
-
-	TArray<TSharedRef<IMessageToken>> MessageTokens;
-	TOptional<UNiagaraGraph*> OriginatingGraph = RecursiveBuildMessageTokensFromContextStackAndGetOriginatingGraph(ContextStackGuids, ThisScriptGraph, MessageTokens);
-	if (OriginatingGraph.IsSet())
-	{
-		EMessageSeverity::Type MessageSeverity = EMessageSeverity::Info;
-		switch (InCompileEvent.Severity) {
-		case FNiagaraCompileEventSeverity::Error:
-			MessageSeverity = EMessageSeverity::Error;
-			break;
-		case FNiagaraCompileEventSeverity::Warning:
-			MessageSeverity = EMessageSeverity::Warning;
-			break;
-		case FNiagaraCompileEventSeverity::Log:
-			MessageSeverity = EMessageSeverity::Info;
-			break;
-		default:
-			verifyf(false, TEXT("Compile event severity type not handled!"));
-			MessageSeverity = EMessageSeverity::Info;
-			break;
-		}
-		TSharedRef<FTokenizedMessage> MessageToBuild = FTokenizedMessage::Create(MessageSeverity);
 
-		//Add message from compile event at start of message
-		MessageToBuild->AddToken(FTextToken::Create(FText::FromString(InCompileEvent.Message)));
-		//Now add the context stack of the scripts that were passed through to get to the originating graph
-		for (TSharedRef<IMessageToken> Token : MessageTokens)
-		{
-			MessageToBuild->AddToken(Token);
-		}
-		return MessageToBuild;
-	}
-	
-	// Failed to get originating graph!
-	verifyf(false, TEXT("Failed to find originating graph for compile event!"));
-	return FTokenizedMessage::Create(EMessageSeverity::Warning, LOCTEXT("NiagaraCompileEventFailedToFindScript", "Could not find script(s) for compile event ! Please recompile script to get full compile info."));
-}
-
-//@todo(ng) move to message manager class
-TOptional<UNiagaraGraph*> FNiagaraScriptToolkit::RecursiveBuildMessageTokensFromContextStackAndGetOriginatingGraph(TArray<FGuid>& InContextStackNodeGuids, UNiagaraGraph* InGraphToSearch, TArray<TSharedRef<IMessageToken>>& OutMessageTokensToAdd)
-{
-	verifyf(InGraphToSearch, TEXT("Failed to get a node graph to search!"));
-
-	if (InGraphToSearch && InContextStackNodeGuids.Num() == 0)
-	{
-		//StackGuids arr has been cleared out which means we have walked the entire context and can now returnthe graph we are in.
-		return InGraphToSearch;
-	}
-
-	// Search in the current graph for a node with a GUID that matches a GUID in the list of Function Call and Emitter node GUIDs that define the context stack for a compile event
-	auto FindNodeInGraphWithContextStackGuid = [&InGraphToSearch, &InContextStackNodeGuids]()->TOptional<UEdGraphNode*> {
-		for (UEdGraphNode* GraphNode : InGraphToSearch->Nodes)
-		{
-			for (int i = 0; i < InContextStackNodeGuids.Num(); i++)
-			{
-				if (GraphNode->NodeGuid == InContextStackNodeGuids[i])
-				{
-					InContextStackNodeGuids.RemoveAt(i);
-					return GraphNode;
-				}
-			}
-		}
-		return TOptional<UEdGraphNode*>();
-	};
-
-	TOptional<UEdGraphNode*> ContextNode = FindNodeInGraphWithContextStackGuid();
-	if (ContextNode.IsSet())
-	{
-		// found a node in the current graph that has a GUID in the context list, now get the Niagara Script assigned to this node, add a message token and recursively call into the graph of that script.
-		UNiagaraNodeFunctionCall* FunctionCallNode = Cast<UNiagaraNodeFunctionCall>(ContextNode.GetValue());
-		if (FunctionCallNode)
-		{
-			UNiagaraScript* FunctionCallNodeAssignedScript = FunctionCallNode->FunctionScript;
-			checkf(FunctionCallNodeAssignedScript, TEXT("Script for function call node is not assigned!"));
-			UNiagaraScriptSourceBase* FunctionScriptSourceBase = FunctionCallNodeAssignedScript->GetSource();
-			checkf(FunctionScriptSourceBase->IsA<UNiagaraScriptSource>(), TEXT("Script source for function call node is not assigned or is not of type UNiagaraScriptSource!"));
-			UNiagaraScriptSource* FunctionScriptSource = Cast<UNiagaraScriptSource>(FunctionScriptSourceBase);
-			checkf(FunctionScriptSource, TEXT("Script source base was somehow not a derived type!"));
-			UNiagaraGraph* FunctionScriptGraph = FunctionScriptSource->NodeGraph;
-			checkf(FunctionScriptGraph, TEXT("Function Script does not have a UNiagaraGraph!"));
-
-			TSharedRef<FNiagaraAssetNameToken> VisitedGraphContextToken = FNiagaraAssetNameToken::Create(FunctionCallNodeAssignedScript->GetPathName(), FText::FromString(*FunctionCallNodeAssignedScript->GetName()));
-			OutMessageTokensToAdd.Add(VisitedGraphContextToken);
-			return RecursiveBuildMessageTokensFromContextStackAndGetOriginatingGraph(InContextStackNodeGuids, FunctionScriptGraph, OutMessageTokensToAdd);
-		}
-		UNiagaraNodeEmitter* EmitterNode = Cast<UNiagaraNodeEmitter>(ContextNode.GetValue());
-		if (EmitterNode)
-		{
-			UNiagaraScriptSource* EmitterScriptSource = EmitterNode->GetScriptSource();
-			UNiagaraGraph* EmitterScriptGraph = EmitterScriptSource->NodeGraph;
-			checkf(EmitterScriptGraph, TEXT("Emitter Script Source does not have a UNiagaraGraph!")); //@todo(ng) see about filtering the message log by emitter?
-
-			TSharedRef<FTextToken> SourceEmitterContextToken = FTextToken::Create(FText::FromString(EmitterNode->GetEmitterUniqueName()));
-			OutMessageTokensToAdd.Add(SourceEmitterContextToken);
-			return RecursiveBuildMessageTokensFromContextStackAndGetOriginatingGraph(InContextStackNodeGuids, EmitterScriptGraph, OutMessageTokensToAdd);
-		}
-		checkf(false, TEXT("Matching node is not a function call or emitter node!"));
-	}
-	verifyf(false, TEXT("Failed to walk the entire context stack, is this compile event out of date?"));
-	return TOptional<UNiagaraGraph*>();
+	FNiagaraEditorModule& NiagaraEditorModule = FModuleManager::LoadModuleChecked<FNiagaraEditorModule>("NiagaraEditor");
+	NiagaraEditorModule.GetOnScriptToolkitsShouldFocusGraphElement().RemoveAll(this);
 }
 
 void FNiagaraScriptToolkit::RegisterTabSpawners(const TSharedRef<class FTabManager>& InTabManager)
@@ -348,12 +237,14 @@ void FNiagaraScriptToolkit::Initialize( const EToolkitMode::Type Mode, const TSh
 
 	FNiagaraEditorModule& NiagaraEditorModule = FModuleManager::LoadModuleChecked<FNiagaraEditorModule>( "NiagaraEditor" );
 	AddMenuExtender(NiagaraEditorModule.GetMenuExtensibilityManager()->GetAllExtenders(GetToolkitCommands(), GetEditingObjects()));
+	NiagaraEditorModule.GetOnScriptToolkitsShouldFocusGraphElement().AddSP(this, &FNiagaraScriptToolkit::FocusGraphElement);
 
 	SetupCommands();
 	ExtendToolbar();
 	RegenerateMenusAndToolbars();
 	UpdateMessageLog();
 	UpdateModuleStats();
+
 	// @todo toolkit world centric editing
 	/*// Setup our tool's layout
 	if( IsWorldCentricAssetEditor() )
@@ -394,7 +285,7 @@ TSharedRef<SDockTab> FNiagaraScriptToolkit::SpawnTabNodeGraph( const FSpawnTabAr
 	return
 		SNew(SDockTab)
 		[
-			SNew(SNiagaraScriptGraph, ScriptViewModel->GetGraphViewModel())
+			SAssignNew(NiagaraScriptGraphWidget, SNiagaraScriptGraph, ScriptViewModel->GetGraphViewModel())
 			.GraphTitle(LOCTEXT("SpawnGraphTitle", "Script"))
 		];
 }
@@ -677,13 +568,12 @@ void FNiagaraScriptToolkit::AddReferencedObjects(FReferenceCollector& Collector)
 
 void FNiagaraScriptToolkit::UpdateMessageLog()
 {
-	TArray< TSharedRef<class FTokenizedMessage> > Messages;
-
+	TArray<TSharedPtr<INiagaraMessage>> NiagaraMessages;
 	uint32 ErrorCount = 0;
 	uint32 WarningCount = 0;
-	for (FNiagaraCompileEvent CompileEvent : EditedNiagaraScript->GetVMExecutableData().LastCompileEvents)
+	for (const FNiagaraCompileEvent CompileEvent : EditedNiagaraScript->GetVMExecutableData().LastCompileEvents)
 	{
-		Messages.Add(BuildMessageForCompileEvent(CompileEvent)); 
+		NiagaraMessages.Add(FNiagaraMessageManager::Get()->QueueMessageJob(MakeShared<FNiagaraMessageJobCompileEvent>(CompileEvent, TWeakObjectPtr<UNiagaraScript>(OriginalNiagaraScript), true)));
 		if (CompileEvent.Severity == FNiagaraCompileEventSeverity::Error)
 		{
 			ErrorCount++;
@@ -694,6 +584,12 @@ void FNiagaraScriptToolkit::UpdateMessageLog()
 		}
 	}
 
+	TArray<TSharedRef<FTokenizedMessage>> TokenizedMessages;
+	for (const TSharedPtr<INiagaraMessage> Message : NiagaraMessages)
+	{
+		TokenizedMessages.Add(Message->GenerateTokenizedMessage());
+	}
+
 	const auto GetCompileCompleteMessageText = [&ErrorCount, &WarningCount](ENiagaraScriptCompileStatus Status)->const FText {
 		FText MessageText;
 		switch (Status)
@@ -716,10 +612,10 @@ void FNiagaraScriptToolkit::UpdateMessageLog()
 	};
 
 	const FText CompileCompleteMessageText = GetCompileCompleteMessageText(ScriptViewModel->GetLatestCompileStatus());
-	Messages.Add(FTokenizedMessage::Create(EMessageSeverity::Info, CompileCompleteMessageText));
+	TokenizedMessages.Add(FTokenizedMessage::Create(EMessageSeverity::Info, CompileCompleteMessageText));
 
 	NiagaraMessageLogListing->ClearMessages();
-	NiagaraMessageLogListing->AddMessages(Messages);
+ 	NiagaraMessageLogListing->AddMessages(TokenizedMessages);
 }
 
 void FNiagaraScriptToolkit::UpdateModuleStats()
@@ -915,4 +811,45 @@ void FNiagaraScriptToolkit::OnEditedScriptGraphChanged(const FEdGraphEditAction&
 	bEditedScriptHasPendingChanges = true;
 }
 
+void FNiagaraScriptToolkit::FocusGraphElement(const INiagaraScriptGraphFocusInfo* FocusInfo)
+{
+	checkf(FocusInfo->GetFocusType() != INiagaraScriptGraphFocusInfo::ENiagaraScriptGraphFocusInfoType::None, TEXT("Failed to assign focus type to FocusInfo parameter!"));
+	
+	if (FocusInfo->GetScriptUniqueAssetID() != OriginalNiagaraScript->GetUniqueID())
+	{
+		return;
+	}
+
+	if (FocusInfo->GetFocusType() == INiagaraScriptGraphFocusInfo::ENiagaraScriptGraphFocusInfoType::Node)
+	{
+		const FNiagaraScriptGraphNodeToFocusInfo* NodeFocusInfo = static_cast<const FNiagaraScriptGraphNodeToFocusInfo*>(FocusInfo);
+		const FGuid& NodeGuidToMatch = NodeFocusInfo->GetNodeGuidToFocus();
+		UEdGraphNode* const* NodeToFocus = ScriptViewModel->GetGraphViewModel()->GetGraph()->Nodes.FindByPredicate([&NodeGuidToMatch](const UEdGraphNode* Node) {return Node->NodeGuid == NodeGuidToMatch; });
+		if (NodeToFocus != nullptr && *NodeToFocus != nullptr)
+		{
+			NiagaraScriptGraphWidget->GetGraphEditor()->JumpToNode(*NodeToFocus);
+			return;
+		}
+		ensureMsgf(false, TEXT("Failed to find Node with matching GUID when focusing graph element. Was the graph edited out from underneath us?"));
+		return;
+	}
+	else if (FocusInfo->GetFocusType() == INiagaraScriptGraphFocusInfo::ENiagaraScriptGraphFocusInfoType::Pin)
+	{
+		const FNiagaraScriptGraphPinToFocusInfo* PinFocusInfo = static_cast<const FNiagaraScriptGraphPinToFocusInfo*>(FocusInfo);
+		const FGuid& PinGuidToMatch = PinFocusInfo->GetPinGuidToFocus();
+		for (const UEdGraphNode* Node : ScriptViewModel->GetGraphViewModel()->GetGraph()->Nodes)
+		{
+			const UEdGraphPin* const* PinToFocus = Node->Pins.FindByPredicate([&PinGuidToMatch](const UEdGraphPin* Pin) {return Pin->PersistentGuid == PinGuidToMatch; });
+			if (PinToFocus != nullptr && *PinToFocus != nullptr)
+			{
+				NiagaraScriptGraphWidget->GetGraphEditor()->JumpToPin(*PinToFocus);
+				return;
+			}
+		}
+		ensureMsgf(false, TEXT("Failed to find Pin with matching GUID when focusing graph element. Was the graph edited out from underneath us?"));
+		return;
+	}
+	checkf(false, TEXT("Requested focus for a graph element without specifying a Node or Pin to focus!"));
+}
+
 #undef LOCTEXT_NAMESPACE
diff --git a/Engine/Plugins/FX/Niagara/Source/NiagaraEditor/Private/Toolkits/NiagaraScriptToolkit.h b/Engine/Plugins/FX/Niagara/Source/NiagaraEditor/Private/Toolkits/NiagaraScriptToolkit.h
--- a/Engine/Plugins/FX/Niagara/Source/NiagaraEditor/Private/Toolkits/NiagaraScriptToolkit.h
+++ b/Engine/Plugins/FX/Niagara/Source/NiagaraEditor/Private/Toolkits/NiagaraScriptToolkit.h
@@ -11,6 +11,7 @@
 #include "UObject/GCObject.h"
 
 #include "NiagaraScript.h"
+#include "NiagaraEditorCommon.h"
 
 class IDetailsView;
 class SGraphEditor;
@@ -20,10 +21,6 @@ class UNiagaraScriptSource;
 class FNiagaraScriptViewModel;
 class FNiagaraObjectSelection;
 struct FEdGraphEditAction;
-class FTokenizedMessage; //@todo(ng) remove this forward decl after making FNiagaraMessageManager
-class IMessageToken; //@todo(ng) remove this forward decl after making FNiagaraMessageManager
-class UNiagaraGraph; //@todo(ng) remove this forward decl after making FNiagaraMessageManager
-//struct FNiagaraCompileEvent; //@todo(ng) remove this forward decl after making FNiagaraMessageManager
 
 /** Viewer/editor for a DataTable */
 class FNiagaraScriptToolkit : public FAssetEditorToolkit, public FGCObject
@@ -119,6 +116,11 @@ private:
 
 	void OnEditedScriptGraphChanged(const FEdGraphEditAction& InAction);
 
+	/** Navigates to element in graph (node, pin, etc.) 
+	* @Param ElementToFocus Defines the graph element to navigate to and select.
+	*/
+	void FocusGraphElement(const INiagaraScriptGraphFocusInfo* ElementToFocus);
+
 private:
 
 	/** The Script being edited */
@@ -147,11 +149,5 @@ private:
 	bool bEditedScriptHasPendingChanges;
 	bool bChangesDiscarded;
 
-	/**  //@todo(ng) Everything below this line should be moved to an owning class such as FNiagaraMessageManager as we don't want to duplicate this code for both the Script and System toolkits
-	* --------------------------------------------------------------------------------------------------------------------------------------------------------------------------------------------
-	*/
-	
-	TSharedRef<FTokenizedMessage> BuildMessageForCompileEvent(FNiagaraCompileEvent& InCompileEvent);
-	TOptional<UNiagaraGraph*> RecursiveBuildMessageTokensFromContextStackAndGetOriginatingGraph(TArray<FGuid>& InContextStackNodeGuids, UNiagaraGraph* InGraphToSearch, TArray<TSharedRef<IMessageToken>>& OutMessageTokensToAdd);
-
+	TSharedPtr<class SNiagaraScriptGraph> NiagaraScriptGraphWidget;
 };
diff --git a/Engine/Plugins/FX/Niagara/Source/NiagaraEditor/Private/Toolkits/NiagaraSystemToolkit.cpp b/Engine/Plugins/FX/Niagara/Source/NiagaraEditor/Private/Toolkits/NiagaraSystemToolkit.cpp
--- a/Engine/Plugins/FX/Niagara/Source/NiagaraEditor/Private/Toolkits/NiagaraSystemToolkit.cpp
+++ b/Engine/Plugins/FX/Niagara/Source/NiagaraEditor/Private/Toolkits/NiagaraSystemToolkit.cpp
@@ -50,10 +50,7 @@
 #include "Misc/FileHelper.h"
 #include "MessageLog/Public/MessageLogModule.h"
 #include "MessageLog/Public/IMessageLogListing.h"
-#include "NiagaraScriptSource.h" //@todo(ng) remove this include after making FNiagaraMessageManager
-#include "NiagaraNodeFunctionCall.h" //@todo(ng) remove this include after making FNiagaraMessageManager
-#include "NiagaraScript.h" //@todo(ng) dependency for fniagaracompile event, remove this include after making FNiagaraMessageManager
-#include "NiagaraNodeEmitter.h" //@todo(ng) remove this include after making FNiagaraMessageManager
+#include "NiagaraMessageManager.h"
 
 #define LOCTEXT_NAMESPACE "NiagaraSystemEditor"
 
@@ -175,121 +172,6 @@ FNiagaraSystemToolkit::~FNiagaraSystemToolkit()
 	SystemViewModel.Reset();
 }
 
-//@todo(ng) move to message manager class
-TSharedRef<FTokenizedMessage> FNiagaraSystemToolkit::BuildMessageForCompileEvent(FNiagaraCompileEvent& InCompileEvent, UNiagaraScript* InOriginatingScript)
-{
-	TArray<FGuid>& ContextStackGuids = InCompileEvent.StackGuids;
-
-	UNiagaraScriptSourceBase* FunctionScriptSourceBase = InOriginatingScript->GetSource();
-	checkf(FunctionScriptSourceBase->IsA<UNiagaraScriptSource>(), TEXT("Script source for function call node is not assigned or is not of type UNiagaraScriptSource!"))
-	UNiagaraScriptSource* FunctionScriptSource = Cast<UNiagaraScriptSource>(FunctionScriptSourceBase);
-	checkf(FunctionScriptSource, TEXT("Script source base was somehow not a derived type!"));
-	UNiagaraGraph* ThisScriptGraph = FunctionScriptSource->NodeGraph;
-	checkf(ThisScriptGraph, TEXT("Function Script does not have a UNiagaraGraph!"));
-
-	TArray<TSharedRef<IMessageToken>> MessageTokens;
-	TOptional<UNiagaraGraph*> OriginatingGraph = RecursiveBuildMessageTokensFromContextStackAndGetOriginatingGraph(ContextStackGuids, ThisScriptGraph, MessageTokens);
-	if (OriginatingGraph.IsSet())
-	{
-		EMessageSeverity::Type MessageSeverity = EMessageSeverity::Info;
-		switch (InCompileEvent.Severity) {
-		case FNiagaraCompileEventSeverity::Error:
-			MessageSeverity = EMessageSeverity::Error;
-			break;
-		case FNiagaraCompileEventSeverity::Warning:
-			MessageSeverity = EMessageSeverity::Warning;
-			break;
-		case FNiagaraCompileEventSeverity::Log:
-			MessageSeverity = EMessageSeverity::Info;
-			break;
-		default:
-			verifyf(false, TEXT("Compile event severity type not handled!"));
-			MessageSeverity = EMessageSeverity::Info;
-			break;
-		}
-		TSharedRef<FTokenizedMessage> MessageToBuild = FTokenizedMessage::Create(MessageSeverity);
-
-		//Add message from compile event at start of message
-		MessageToBuild->AddToken(FTextToken::Create(FText::FromString(InCompileEvent.Message)));
-		//Now add the context stack of the scripts that were passed through to get to the originating graph
-		for (TSharedRef<IMessageToken> Token : MessageTokens)
-		{
-			MessageToBuild->AddToken(Token);
-		}
-		return MessageToBuild;
-	}
-
-	// Failed to get originating graph!
-	verifyf(false, TEXT("Failed to find originating graph for compile event!"));
-	return FTokenizedMessage::Create(EMessageSeverity::Warning, LOCTEXT("NiagaraCompileEventFailedToFindScript", "Could not find script(s) for compile event ! Please recompile script to get full compile info."));
-}
-
-//@todo(ng) move to message manager class
-TOptional<UNiagaraGraph*> FNiagaraSystemToolkit::RecursiveBuildMessageTokensFromContextStackAndGetOriginatingGraph(TArray<FGuid>& InContextStackNodeGuids, UNiagaraGraph* InGraphToSearch, TArray<TSharedRef<IMessageToken>>& OutMessageTokensToAdd)
-{
-	verifyf(InGraphToSearch, TEXT("Failed to get a node graph to search!"));
-
-	if (InGraphToSearch && InContextStackNodeGuids.Num() == 0)
-	{
-		//StackGuids arr has been cleared out which means we have walked the entire context and can now return the graph we are in.
-		return InGraphToSearch;
-	}
-
-	// Search in the current graph for a node with a GUID that matches a GUID in the list of Function Call and Emitter node GUIDs that define the context stack for a compile event
-	auto FindNodeInGraphWithContextStackGuid = [&InGraphToSearch, &InContextStackNodeGuids]()->TOptional<UEdGraphNode*> {
-		for (UEdGraphNode* GraphNode : InGraphToSearch->Nodes)
-		{
-			for (int i = 0; i < InContextStackNodeGuids.Num(); i++)
-			{
-				if (GraphNode->NodeGuid == InContextStackNodeGuids[i])
-				{
-					InContextStackNodeGuids.RemoveAt(i);
-					return GraphNode;
-				}
-			}
-		}
-		return TOptional<UEdGraphNode*>();
-	};
-
-	TOptional<UEdGraphNode*> ContextNode = FindNodeInGraphWithContextStackGuid();
-	if (ContextNode.IsSet())
-	{
-		// found a node in the current graph that has a GUID in the context list
-		UNiagaraNodeFunctionCall* FunctionCallNode = Cast<UNiagaraNodeFunctionCall>(ContextNode.GetValue());
-		if (FunctionCallNode)
-		{
-			// node is a function call node, now get the Niagara Script assigned to this node, add a message token and recursively call into the graph of that script.
-			UNiagaraScript* FunctionCallNodeAssignedScript = FunctionCallNode->FunctionScript;
-			checkf(FunctionCallNodeAssignedScript, TEXT("Script for function call node is not assigned!"));
-			UNiagaraScriptSourceBase* FunctionScriptSourceBase = FunctionCallNodeAssignedScript->GetSource();
-			checkf(FunctionScriptSourceBase->IsA<UNiagaraScriptSource>(), TEXT("Script source for function call node is not assigned or is not of type UNiagaraScriptSource!"));
-			UNiagaraScriptSource* FunctionScriptSource = Cast<UNiagaraScriptSource>(FunctionScriptSourceBase);
-			checkf(FunctionScriptSource, TEXT("Script source base was somehow not a derived type!"));
-			UNiagaraGraph* FunctionScriptGraph = FunctionScriptSource->NodeGraph;
-			checkf(FunctionScriptGraph, TEXT("Function Script does not have a UNiagaraGraph!"));
-
-			TSharedRef<FNiagaraAssetNameToken> VisitedGraphContextToken = FNiagaraAssetNameToken::Create(FunctionCallNodeAssignedScript->GetPathName(), FText::FromString(*FunctionCallNodeAssignedScript->GetName()));
-			OutMessageTokensToAdd.Add(VisitedGraphContextToken);
-			return RecursiveBuildMessageTokensFromContextStackAndGetOriginatingGraph(InContextStackNodeGuids, FunctionScriptGraph, OutMessageTokensToAdd);
-		}
-		UNiagaraNodeEmitter* EmitterNode = Cast<UNiagaraNodeEmitter>(ContextNode.GetValue());
-		if (EmitterNode)
-		{
-			// node is an emitter node, now get the Emitter name, add a message token and recursively call into the graph of that emitter.
-			UNiagaraScriptSource* EmitterScriptSource = EmitterNode->GetScriptSource();
-			UNiagaraGraph* EmitterScriptGraph = EmitterScriptSource->NodeGraph;
-			checkf(EmitterScriptGraph, TEXT("Emitter Script Source does not have a UNiagaraGraph!")); //@todo(ng) see about filtering the message log by emitter?
-
-			TSharedRef<FTextToken> SourceEmitterContextToken = FTextToken::Create(FText::FromString(EmitterNode->GetEmitterUniqueName()));
-			OutMessageTokensToAdd.Add(SourceEmitterContextToken);
-			return RecursiveBuildMessageTokensFromContextStackAndGetOriginatingGraph(InContextStackNodeGuids, EmitterScriptGraph, OutMessageTokensToAdd);
-		}
-		checkf(false, TEXT("Matching node is not a function call or emitter node!"));
-	}
-	verifyf(false, TEXT("Failed to walk the entire context stack, is this compile event out of date?"));
-	return TOptional<UNiagaraGraph*>();
-}
-
 void FNiagaraSystemToolkit::AddReferencedObjects(FReferenceCollector& Collector) 
 {
 	Collector.AddReferencedObject(System);
@@ -309,14 +191,14 @@ void FNiagaraSystemToolkit::UpdateMessageLog()
 		ScriptsToGetCompileEventsFrom.Append(EmitterScripts);
 	}
 
+	TArray<TSharedPtr<INiagaraMessage>> NiagaraMessages;
 	uint32 ErrorCount = 0;
 	uint32 WarningCount = 0;
-	TArray< TSharedRef<class FTokenizedMessage> > Messages;
 	for (UNiagaraScript* Script : ScriptsToGetCompileEventsFrom)
 	{
-		for (FNiagaraCompileEvent CompileEvent : Script->GetVMExecutableData().LastCompileEvents)
+		for (const FNiagaraCompileEvent CompileEvent : Script->GetVMExecutableData().LastCompileEvents)
 		{
-			Messages.Add(BuildMessageForCompileEvent(CompileEvent, Script));
+			NiagaraMessages.Add(FNiagaraMessageManager::Get()->QueueMessageJob(MakeShared<FNiagaraMessageJobCompileEvent>(CompileEvent, TWeakObjectPtr<UNiagaraScript>(Script), false)));
 			if (CompileEvent.Severity == FNiagaraCompileEventSeverity::Error)
 			{
 				ErrorCount++;
@@ -328,6 +210,25 @@ void FNiagaraSystemToolkit::UpdateMessageLog()
 		}
 	}
 
+	TArray<TSharedRef<FTokenizedMessage>> TokenizedMessages;
+	for (const TSharedPtr<INiagaraMessage> Message : NiagaraMessages)
+	{
+		if (Message->GetMessageType() == ENiagaraMessageType::CompileEventMessage)
+		{
+			const TSharedPtr<const FNiagaraMessageCompileEvent> MessageCompileEvent = StaticCastSharedPtr<const FNiagaraMessageCompileEvent>(Message);
+			TokenizedMessages.Add(MessageCompileEvent->GenerateTokenizedMessage());
+		}
+		else if (Message->GetMessageType() == ENiagaraMessageType::NeedRecompileMessage)
+		{
+			const TSharedPtr<const FNiagaraMessageNeedRecompile> MessageNeedRecompile = StaticCastSharedPtr<const FNiagaraMessageNeedRecompile>(Message);
+			TokenizedMessages.Add(MessageNeedRecompile->GenerateTokenizedMessage());
+		}
+		else
+		{
+			checkf(false, TEXT("Message is of ENiagaraMessageType that is unhandled!"));
+		}
+	}
+
 	const auto GetCompileCompleteMessageText = [&ErrorCount, &WarningCount](ENiagaraScriptCompileStatus Status)->const FText {
 		FText MessageText;
 		switch (Status)
@@ -350,10 +251,10 @@ void FNiagaraSystemToolkit::UpdateMessageLog()
 	};
 	
 	const FText CompileCompleteMessageText = GetCompileCompleteMessageText(SystemViewModel->GetLatestCompileStatus());
-	Messages.Add(FTokenizedMessage::Create(EMessageSeverity::Info, CompileCompleteMessageText));
+	TokenizedMessages.Add(FTokenizedMessage::Create(EMessageSeverity::Info, CompileCompleteMessageText));
 
 	NiagaraMessageLogListing->ClearMessages();
-	NiagaraMessageLogListing->AddMessages(Messages);
+	NiagaraMessageLogListing->AddMessages(TokenizedMessages);
 }
 
 void FNiagaraSystemToolkit::InitializeWithSystem(const EToolkitMode::Type Mode, const TSharedPtr< class IToolkitHost >& InitToolkitHost, UNiagaraSystem& InSystem)
diff --git a/Engine/Plugins/FX/Niagara/Source/NiagaraEditor/Private/Toolkits/NiagaraSystemToolkit.h b/Engine/Plugins/FX/Niagara/Source/NiagaraEditor/Private/Toolkits/NiagaraSystemToolkit.h
--- a/Engine/Plugins/FX/Niagara/Source/NiagaraEditor/Private/Toolkits/NiagaraSystemToolkit.h
+++ b/Engine/Plugins/FX/Niagara/Source/NiagaraEditor/Private/Toolkits/NiagaraSystemToolkit.h
@@ -27,11 +27,6 @@ class UNiagaraSequence;
 struct FAssetData;
 class FMenuBuilder;
 class ISequencer;
-class FTokenizedMessage; //@todo(ng) remove this forward decl after making FNiagaraMessageManager
-class IMessageToken; //@todo(ng) remove this forward decl after making FNiagaraMessageManager
-//struct FNiagaraCompileEvent; //@todo(ng) remove this forward decl after making FNiagaraMessageManager
-class UNiagaraGraph; //@todo(ng) remove this forward decl after making FNiagaraMessageManager
-//class UNiagaraScript; //@todo(ng) remove this forward decl after making FNiagaraMessageManager
 
 /** Viewer/editor for a NiagaraSystem
 */
@@ -186,12 +181,4 @@ private:
 	static const FName PreviewSettingsTabId;
 	static const FName GeneratedCodeTabID;
 	static const FName MessageLogTabID;
-
-	/**  //@todo(ng) Everything below this line should be moved to an owning class such as FNiagaraMessageManager as we don't want to duplicate this code for both the Script and System toolkits
-* --------------------------------------------------------------------------------------------------------------------------------------------------------------------------------------------
-*/
-
-	TSharedRef<FTokenizedMessage> BuildMessageForCompileEvent(FNiagaraCompileEvent& InCompileEvent, UNiagaraScript* InOriginatingScript);
-	TOptional<UNiagaraGraph*> RecursiveBuildMessageTokensFromContextStackAndGetOriginatingGraph(TArray<FGuid>& InContextStackNodeGuids, UNiagaraGraph* InGraphToSearch, TArray<TSharedRef<IMessageToken>>& OutMessageTokensToAdd);
-
 };
diff --git a/Engine/Plugins/FX/Niagara/Source/NiagaraEditor/Private/Widgets/SNiagaraScriptGraph.h b/Engine/Plugins/FX/Niagara/Source/NiagaraEditor/Private/Widgets/SNiagaraScriptGraph.h
--- a/Engine/Plugins/FX/Niagara/Source/NiagaraEditor/Private/Widgets/SNiagaraScriptGraph.h
+++ b/Engine/Plugins/FX/Niagara/Source/NiagaraEditor/Private/Widgets/SNiagaraScriptGraph.h
@@ -22,6 +22,7 @@ public:
 
 	void Construct(const FArguments& InArgs, TSharedRef<FNiagaraScriptGraphViewModel> InViewModel);
 
+	const TSharedPtr<SGraphEditor> GetGraphEditor() { return GraphEditor; };
 
 private:
 	/** Constructs the graph editor widget for the current graph. */
diff --git a/Engine/Plugins/FX/Niagara/Source/NiagaraEditor/Public/NiagaraEditorCommon.h b/Engine/Plugins/FX/Niagara/Source/NiagaraEditor/Public/NiagaraEditorCommon.h
--- a/Engine/Plugins/FX/Niagara/Source/NiagaraEditor/Public/NiagaraEditorCommon.h
+++ b/Engine/Plugins/FX/Niagara/Source/NiagaraEditor/Public/NiagaraEditorCommon.h
@@ -51,3 +51,60 @@ public:
 	bool CreateHlslForAddedInputs(int32 InputCount, FString& HlslResult) const;
 };
 
+/** Interface for struct representing information about where to focus in a Niagara Script Graph after opening the editor for it. */
+struct INiagaraScriptGraphFocusInfo
+{
+public:
+	enum class ENiagaraScriptGraphFocusInfoType : uint8
+	{
+		None = 0,
+		Node,
+		Pin
+	};
+
+	INiagaraScriptGraphFocusInfo(const uint32& InScriptUniqueAssetID, const ENiagaraScriptGraphFocusInfoType InFocusType)
+		: ScriptUniqueAssetID(InScriptUniqueAssetID)
+		, FocusType(InFocusType)
+	{
+	};
+
+	const uint32& GetScriptUniqueAssetID() const { return ScriptUniqueAssetID; };
+
+	const ENiagaraScriptGraphFocusInfoType& GetFocusType() const { return FocusType; };
+
+	virtual ~INiagaraScriptGraphFocusInfo() = 0;
+	
+private:
+	const uint32 ScriptUniqueAssetID;
+	const ENiagaraScriptGraphFocusInfoType FocusType;
+};
+
+struct FNiagaraScriptGraphNodeToFocusInfo : public INiagaraScriptGraphFocusInfo
+{
+public:
+	FNiagaraScriptGraphNodeToFocusInfo(const uint32& InScriptUniqueAssetID, const FGuid& InNodeGuidToFocus)
+		: INiagaraScriptGraphFocusInfo(InScriptUniqueAssetID, ENiagaraScriptGraphFocusInfoType::Node)
+		, NodeGuidToFocus(InNodeGuidToFocus)
+	{
+	};
+
+	const FGuid& GetNodeGuidToFocus() const { return NodeGuidToFocus; };
+
+private:
+	const FGuid NodeGuidToFocus;
+};
+
+struct FNiagaraScriptGraphPinToFocusInfo : public INiagaraScriptGraphFocusInfo
+{
+public:
+	FNiagaraScriptGraphPinToFocusInfo(const uint32& InScriptUniqueAssetID, const FGuid& InPinGuidToFocus)
+		: INiagaraScriptGraphFocusInfo(InScriptUniqueAssetID, ENiagaraScriptGraphFocusInfoType::Pin)
+		, PinGuidToFocus(InPinGuidToFocus)
+	{
+	};
+
+	const FGuid& GetPinGuidToFocus() const { return PinGuidToFocus; };
+
+private:
+	const FGuid PinGuidToFocus;
+};
diff --git a/Engine/Plugins/FX/Niagara/Source/NiagaraEditor/Public/NiagaraEditorModule.h b/Engine/Plugins/FX/Niagara/Source/NiagaraEditor/Public/NiagaraEditorModule.h
--- a/Engine/Plugins/FX/Niagara/Source/NiagaraEditor/Public/NiagaraEditorModule.h
+++ b/Engine/Plugins/FX/Niagara/Source/NiagaraEditor/Public/NiagaraEditorModule.h
@@ -8,11 +8,9 @@
 #include "NiagaraTypes.h"
 #include "INiagaraCompiler.h"
 #include "AssetTypeCategories.h"
-#include "Logging/TokenizedMessage.h" //@todo(ng) remove this include after moving FNiagaraAssetPathToken to message manager
 
 class IAssetTools;
 class IAssetTypeActions;
-struct FNiagaraCompileResults;
 class INiagaraEditorTypeUtilities;
 class UNiagaraSettings;
 class USequencerSettings;
@@ -32,6 +30,7 @@ class FNiagaraEditorModule : public IModuleInterface,
 public:
 	DECLARE_DELEGATE_RetVal_OneParam(TSharedRef<SWidget>, FOnCreateStackWidget, UNiagaraStackViewModel*);
 	DECLARE_DELEGATE_RetVal_OneParam(UMovieSceneNiagaraParameterTrack*, FOnCreateMovieSceneTrackForParameter, FNiagaraVariable);
+	DECLARE_MULTICAST_DELEGATE_OneParam(FOnCheckScriptToolkitsShouldFocusGraphElement, const INiagaraScriptGraphFocusInfo*);
 
 public:
 	FNiagaraEditorModule();
@@ -43,9 +42,6 @@ public:
 	/** Get the instance of this module. */
 	NIAGARAEDITOR_API static FNiagaraEditorModule& Get();
 
-	/** Callback  when a token is clicked on in the niagara message log */ //@todo(ng) consider moving this to FNiagaraMessageManager in the future?
-	static void OnMessageLogTokenClicked(const TSharedRef<class IMessageToken>& Token);
-
 	/** Compile the specified script. */
 	virtual TSharedPtr<FNiagaraVMExecutableData> CompileScript(const FNiagaraCompileRequestDataBase* InCompileRequest, const FNiagaraCompileOptions& InCompileOptions);
 
@@ -88,6 +84,8 @@ public:
 	/** Get the niagara UI commands. */
 	NIAGARAEDITOR_API const class FNiagaraEditorCommands& Commands();
 
+	FOnCheckScriptToolkitsShouldFocusGraphElement& GetOnScriptToolkitsShouldFocusGraphElement() { return OnCheckScriptToolkitsShouldFocusGraphElement; };
+
 private:
 	void RegisterAssetTypeAction(IAssetTools& AssetTools, TSharedRef<IAssetTypeActions> Action);
 	void OnNiagaraSettingsChangedEvent(const FString& PropertyName, const UNiagaraSettings* Settings);
@@ -135,34 +133,6 @@ private:
 
 	IConsoleCommand* TestCompileScriptCommand;
 	IConsoleCommand* DumpRapidIterationParametersForAsset;
-};
-
-//@todo(ng) move to message manager class
-//Extension of FAssetNameToken to allow opening the asset editor when clicking on the linked asset name.
-class FNiagaraAssetNameToken : public IMessageToken
-{
-public:
-	/** Factory method, tokens can only be constructed as shared refs */
-	static TSharedRef<FNiagaraAssetNameToken> Create(const FString& InAssetName, const FText& InMessage = FText());
 
-	/** Begin IMessageToken interface */
-	virtual EMessageToken::Type GetType() const override
-	{
-		return EMessageToken::AssetName;
-	}
-	/** End IMessageToken interface */
-
-private:
-	/** Private constructor */
-	FNiagaraAssetNameToken(const FString& InAssetName, const FText& InMessage);
-
-	/**
-	 * Find and open an asset in editor
-	 * @param	Token		The token that was clicked
-	 * @param	InAssetPath		The asset to find
-	 */
-	static void FindAndOpenAsset(const TSharedRef<IMessageToken>& Token, FString InAssetPath);
-
-	/** The asset name we will find */
-	FString AssetName;
+	FOnCheckScriptToolkitsShouldFocusGraphElement OnCheckScriptToolkitsShouldFocusGraphElement;
 };
diff --git a/Engine/Plugins/FX/Niagara/Source/NiagaraEditor/Public/NiagaraMessageManager.h b/Engine/Plugins/FX/Niagara/Source/NiagaraEditor/Public/NiagaraMessageManager.h
new file mode 100644
--- /dev/null
+++ b/Engine/Plugins/FX/Niagara/Source/NiagaraEditor/Public/NiagaraMessageManager.h
@@ -0,0 +1,182 @@
+// Copyright 1998-2019 Epic Games, Inc. All Rights Reserved.
+
+#pragma once
+
+#include "CoreMinimal.h"
+#include "NiagaraShared.h"
+#include "Logging/TokenizedMessage.h"
+#include "NiagaraGraph.h"
+
+class FNiagaraScriptToolkit;
+
+//enum to specify the class of an FNiagaraMessage. 
+enum class ENiagaraMessageType : uint8
+{
+	None = 0,
+	CompileEventMessage,
+	NeedRecompileMessage
+};
+
+//Struct for passing around script asset info from compile event message job to message types
+struct FNiagaraScriptNameAndAssetPath
+{
+public:
+	FNiagaraScriptNameAndAssetPath(const FString InScriptNameString, const FString InScriptAssetPathString)
+		: ScriptNameString(InScriptNameString)
+		, ScriptAssetPathString(InScriptAssetPathString)
+	{};
+
+	const FString ScriptNameString;
+	const FString ScriptAssetPathString;
+};
+
+/** Interface for view-agnostic message that holds limited lifetime information on a message (e.g. a weak pointer to an asset.)
+*	Implements GenerateTokenizedMessage() to provide FTokenizedMessage for the message log, and should separately define a GenerateMessage() implementation for any specific views it needs to support (e.g. errors in SNiagaraStack).
+*/
+class INiagaraMessage 
+{
+public:
+	INiagaraMessage(const ENiagaraMessageType InMessageType)
+		: MessageType(InMessageType)
+	{
+	};
+
+	virtual TSharedRef<FTokenizedMessage> GenerateTokenizedMessage() const = 0;
+
+	const ENiagaraMessageType GetMessageType() {return MessageType;}
+
+	virtual ~INiagaraMessage() {};
+
+private:
+	const ENiagaraMessageType MessageType;
+};
+
+/** Interface for job supplied to FNiagaraMessageManager to generated FNiagaraMessages.
+*	Implements GenerateNiagaraMessage() to return a derived type of abstract INiagaraMessage.
+*/
+class INiagaraMessageJob
+{
+public:
+	
+
+	virtual const TSharedPtr<INiagaraMessage> GenerateNiagaraMessage() const = 0;
+
+	virtual ~INiagaraMessageJob() {};
+};
+
+class FNiagaraMessageJobCompileEvent : public INiagaraMessageJob
+{
+public:
+	FNiagaraMessageJobCompileEvent( const FNiagaraCompileEvent& InCompileEvent, const TWeakObjectPtr<UNiagaraScript>& InOriginatingScriptWeakObjPtr, const bool bInFromScriptToolkit);
+
+	virtual const TSharedPtr<INiagaraMessage> GenerateNiagaraMessage() const override;
+
+private:
+	//hidden private ctor
+	FNiagaraMessageJobCompileEvent();
+
+	const bool RecursiveGetScriptNamesAndAssetPathsFromContextStack(
+		  TArray<FGuid>& InContextStackNodeGuids
+		, const UNiagaraGraph* InGraphToSearch
+		, TArray<FNiagaraScriptNameAndAssetPath>& OutContextScriptNamesAndAssetPaths
+		, TOptional<const FString>& OutEmitterScriptName
+		, TOptional<const FText>& OutFailureReason
+	) const;
+
+	const FNiagaraCompileEvent CompileEvent;
+	const TWeakObjectPtr<UNiagaraScript> OriginatingScriptWeakObjPtr;
+	const bool bFromScriptToolkit;
+};
+
+class FNiagaraMessageCompileEvent : public INiagaraMessage
+{
+public:
+	FNiagaraMessageCompileEvent(
+		  const FNiagaraCompileEvent& InCompileEvent
+		, TArray<FNiagaraScriptNameAndAssetPath>& InContextScriptNamesAndAssetPaths
+		, TOptional<const FString>& InEmitterScriptName
+		, TOptional<const FNiagaraScriptNameAndAssetPath>& InCompiledScriptNameAndAssetPath
+	);
+
+	virtual TSharedRef<FTokenizedMessage> GenerateTokenizedMessage() const override;
+	//@todo(message manager) make stack specific message type generator here
+
+private:
+	const FNiagaraCompileEvent CompileEvent;
+	const TArray<FNiagaraScriptNameAndAssetPath> ContextScriptNamesAndAssetPaths;
+	const TOptional<const FString> EmitterScriptName;
+	const TOptional<const FNiagaraScriptNameAndAssetPath> CompiledScriptNameAndAssetPath;
+};
+
+class FNiagaraMessageNeedRecompile : public INiagaraMessage
+{
+public:
+	FNiagaraMessageNeedRecompile(const FText& InNeedRecompileMessage);
+
+	virtual TSharedRef<FTokenizedMessage> GenerateTokenizedMessage() const override;
+
+private:
+	const FText NeedRecompileMessage;
+};
+
+//Extension of message token to allow opening the asset editor when clicking on the linked asset name.
+class FNiagaraCompileEventToken : public IMessageToken
+{
+public:
+	/** Factory method, tokens can only be constructed as shared refs */
+	static TSharedRef<FNiagaraCompileEventToken> Create(
+		  const FString& InScriptAssetPath
+		, const FText& InMessage
+		, const TOptional<const FGuid>& InNodeGUID = TOptional<const FGuid>()
+		, const TOptional<const FGuid>& InPinGUID = TOptional<const FGuid>()
+	);
+
+	/** Begin IMessageToken interface */
+	virtual EMessageToken::Type GetType() const override
+	{
+		return EMessageToken::AssetName;
+	}
+	/** End IMessageToken interface */
+
+private:
+	/** Private constructor */
+	FNiagaraCompileEventToken(
+		  const FString& InScriptAssetPath
+		, const FText& InMessage
+		, const TOptional<const FGuid>& InNodeGUID
+		, const TOptional<const FGuid>& InPinGUID
+	);
+
+	/**
+	 * Find and open an asset in editor
+	 * @param	Token		The token that was clicked
+	 * @param	InAssetPath		The asset to find
+	 */
+	static void OpenScriptAssetByPathAndFocusNodeOrPinIfSet(
+		  const TSharedRef<IMessageToken>& Token
+		, FString InScriptAssetPath
+		, const TOptional<const FGuid> InNodeGUID
+		, const TOptional<const FGuid> InPinGUID
+	);
+
+	/** The script asset path to open the editor for. */
+	const FString ScriptAssetPath;
+
+	/** The optional Node or Pin GUID to find and focus after opening the script asset */
+	const TOptional<const FGuid> NodeGUID;
+	const TOptional<const FGuid> PinGUID;
+};
+
+class FNiagaraMessageManager
+{
+public:
+	static FNiagaraMessageManager* Get();
+
+	static TSharedPtr<INiagaraMessage> QueueMessageJob(TSharedRef<INiagaraMessageJob> InMessageJob);
+
+private:
+	FNiagaraMessageManager();
+
+	static TQueue<TSharedPtr<INiagaraMessageJob>> JobQueue;
+	static FNiagaraMessageManager* Singleton;
+};
diff --git a/Engine/Plugins/FX/Niagara/Source/NiagaraShader/Public/NiagaraShared.h b/Engine/Plugins/FX/Niagara/Source/NiagaraShader/Public/NiagaraShared.h
--- a/Engine/Plugins/FX/Niagara/Source/NiagaraShader/Public/NiagaraShared.h
+++ b/Engine/Plugins/FX/Niagara/Source/NiagaraShader/Public/NiagaraShared.h
@@ -33,6 +33,50 @@ struct FNiagaraDataInterfaceParametersCS;
 
 #define MAX_CONCURRENT_EVENT_DATASETS 4
 
+/** Defines the compile event types for translation/compilation.*/
+UENUM()
+enum class FNiagaraCompileEventSeverity : uint8
+{
+	Log = 0,
+	Warning = 1,
+	Error = 2
+};
+
+/** Records necessary information to give UI cues for errors/logs/warnings during compile.*/
+USTRUCT()
+struct FNiagaraCompileEvent
+{
+	GENERATED_USTRUCT_BODY()
+public:
+	FNiagaraCompileEvent()
+	{
+		Severity = FNiagaraCompileEventSeverity::Log;
+		Message = FString();
+		NodeGuid = FGuid();
+		PinGuid = FGuid();
+		StackGuids.Empty();
+	}
+
+	FNiagaraCompileEvent(FNiagaraCompileEventSeverity InSeverity, const FString& InMessage, FGuid InNodeGuid = FGuid(), FGuid InPinGuid = FGuid(), const TArray<FGuid>& InCallstackGuids = TArray<FGuid>())
+		: Severity(InSeverity), Message(InMessage), NodeGuid(InNodeGuid), PinGuid(InPinGuid), StackGuids(InCallstackGuids) {}
+
+	/** Whether or not this is an error, warning, or info*/
+	UPROPERTY()
+	FNiagaraCompileEventSeverity Severity;
+	/* The message itself*/
+	UPROPERTY()
+	FString Message;
+	/** The node guid that generated the compile event*/
+	UPROPERTY()
+	FGuid NodeGuid;
+	/** The pin persistent id that generated the compile event*/
+	UPROPERTY()
+	FGuid PinGuid;
+	/** The compile stack frame of node id's*/
+	UPROPERTY()
+	TArray<FGuid> StackGuids;
+};
+
 /**
 * Data coming from that translator that describes parameters needed for each data interface.
 */
