Resolve FNiagaraMessageManager build issue in Clang with implicit ptr checks

diff --git a/Engine/Plugins/FX/Niagara/Source/NiagaraEditor/Private/NiagaraMessageManager.cpp b/Engine/Plugins/FX/Niagara/Source/NiagaraEditor/Private/NiagaraMessageManager.cpp
--- a/Engine/Plugins/FX/Niagara/Source/NiagaraEditor/Private/NiagaraMessageManager.cpp
+++ b/Engine/Plugins/FX/Niagara/Source/NiagaraEditor/Private/NiagaraMessageManager.cpp
@@ -228,14 +228,14 @@ const bool FNiagaraMessageJobCompileEvent::RecursiveGetScriptNamesAndAssetPathsF
 		{
 			// node is a function call node, now get the Niagara Script assigned to this node, add a message token and recursively call into the graph of that script.
 			UNiagaraScript* FunctionCallNodeAssignedScript = FunctionCallNode->FunctionScript;
-			if (FunctionCallNodeAssignedScript == false)
+			if (FunctionCallNodeAssignedScript == nullptr)
 			{
 				FText FailureReason = LOCTEXT("CompileEventMessageGenerator_FunctionCallNodeScriptNotFound", "Script for Function Call Node \"{0}\" not found!");
 				OutFailureReason = FText::Format(FailureReason, FText::FromString(FunctionCallNode->GetFunctionName()));
 				return false;
 			}
 			UNiagaraScriptSourceBase* FunctionScriptSourceBase = FunctionCallNodeAssignedScript->GetSource();
-			if (FunctionScriptSourceBase == false)
+			if (FunctionScriptSourceBase == nullptr)
 			{
 				FText FailureReason = LOCTEXT("CompileEventMessageGenerator_FunctionCallNodeScriptSourceBaseNotFound", "Source Script for Function Call Node \"{0}\" not found!");
 				OutFailureReason = FText::Format(FailureReason, FText::FromString(FunctionCallNode->GetFunctionName()));
@@ -245,7 +245,7 @@ const bool FNiagaraMessageJobCompileEvent::RecursiveGetScriptNamesAndAssetPathsF
 			checkf(FunctionScriptSource, TEXT("Script source base was somehow not a derived type!"));
 
 			UNiagaraGraph* FunctionScriptGraph = FunctionScriptSource->NodeGraph;
-			if (FunctionScriptGraph == false)
+			if (FunctionScriptGraph == nullptr)
 			{
 				FText FailureReason = LOCTEXT("CompileEventMessageGenerator_FunctionCallNodeGraphNotFound", "Graph for Function Call Node \"{0}\" not found!");
 				OutFailureReason = FText::Format(FailureReason, FText::FromString(FunctionCallNode->GetFunctionName()));
