Remove adding Dynamic Input scripts from Niagara Script Graph context menu

diff --git a/Engine/Plugins/FX/Niagara/Source/NiagaraEditor/Private/EdGraphSchema_Niagara.cpp b/Engine/Plugins/FX/Niagara/Source/NiagaraEditor/Private/EdGraphSchema_Niagara.cpp
--- a/Engine/Plugins/FX/Niagara/Source/NiagaraEditor/Private/EdGraphSchema_Niagara.cpp
+++ b/Engine/Plugins/FX/Niagara/Source/NiagaraEditor/Private/EdGraphSchema_Niagara.cpp
@@ -420,20 +420,6 @@ TArray<TSharedPtr<FNiagaraSchemaAction_NewNode> > UEdGraphSchema_Niagara::GetGra
 		}
 	}
 
-	// Add dynamic inputs for default usage in module and dynamic input graphs
-	if (bModuleGraph || bDynamicInputGraph)
-	{
-		TArray<FAssetData> DynamicInputScriptAssets;
-		FNiagaraEditorUtilities::FGetFilteredScriptAssetsOptions DynamicInputScriptFilterOptions;
-		DynamicInputScriptFilterOptions.ScriptUsageToInclude = ENiagaraScriptUsage::DynamicInput;
-		FNiagaraEditorUtilities::GetFilteredScriptAssets(DynamicInputScriptFilterOptions, DynamicInputScriptAssets);
-
-		for (const FAssetData& DynamicInputScriptAsset : DynamicInputScriptAssets)
-		{
-			AddScriptFunctionAction(LOCTEXT("Dynamic Input Menu Title", "Dynamic Inputs"), DynamicInputScriptAsset);
-		}
-	}
-
 	//Add event read and writes nodes
 	if (bModuleGraph)
 	{
