Increased mininum tessellation angle for niagara ribbons from 15 to 30 degrees.
Niagara ribbon's tessellation metrics are now stabilized over frames to avoid rapid changes.
Disabled tesselation CPU logic on non SM5 platforms.

diff --git a/Engine/Plugins/FX/Niagara/Source/Niagara/Private/NiagaraRendererRibbons.cpp b/Engine/Plugins/FX/Niagara/Source/Niagara/Private/NiagaraRendererRibbons.cpp
--- a/Engine/Plugins/FX/Niagara/Source/Niagara/Private/NiagaraRendererRibbons.cpp
+++ b/Engine/Plugins/FX/Niagara/Source/Niagara/Private/NiagaraRendererRibbons.cpp
@@ -19,12 +19,12 @@ static FAutoConsoleVariableRef CVarNiagaraRibbonCullTwistedStrips(
 	ECVF_Default
 );
 
-float GNiagaraRibbonTessellationAngle = 15.f * (2.f * PI) / 360.f; // Every 15 degrees
+float GNiagaraRibbonTessellationAngle = 30.f * (2.f * PI) / 360.f; // Every 30 degrees
 static FAutoConsoleVariableRef CVarNiagaraRibbonTessellationAngle(
 	TEXT("Niagara.Ribbon.Tessellation.MinAngle"),
 	GNiagaraRibbonTessellationAngle,
 	TEXT("Ribbon segment angle to tesselate in radian. \n")
-	TEXT("Set 0 to use fixed tesselation. (default=15 degrees)"),
+	TEXT("Set 0 to use fixed tesselation. (default=30 degrees)"),
 	ECVF_Scalability
 );
 
@@ -82,12 +82,6 @@ struct FNiagaraDynamicDataRibbon : public FNiagaraDynamicDataBase
 	// The list of all segments, each one connecting SortedIndices[SegmentId] to SortedIndices[SegmentId + 1].
 	// We use this format because the final index buffer gets generated based on view sorting and InterpCount.
 	TArray<int32> SegmentData;
-
-	// Curvature of the segment.
-	float AverageSegmentCurvature = 0;
-	// Angle of the curvature over the segment (in radian)­.
-	float AverangeSegmentAngle = 0; 
-
 	TArray<int32> SortedIndices;
 	TArray<FVector4> TangentAndDistances;
 	TArray<uint32> MultiRibbonIndices;
@@ -274,13 +268,14 @@ void NiagaraRendererRibbons::GetDynamicMeshElements(const TArray<const FSceneVie
 
 			int32 SegmentTessellation = 1;
 			int32 NumSegments = DynamicDataRibbon->SegmentData.Num();
-			if (GNiagaraRibbonMaxTessellation > 1)
+			if (GNiagaraRibbonMaxTessellation > 1 && TessellationCurvature > SMALL_NUMBER && ViewFamily.GetFeatureLevel() == ERHIFeatureLevel::SM5)
 			{
-				const float MinTesselation = FMath::Max<float>(1.f, DynamicDataRibbon->AverangeSegmentAngle / FMath::Max<float>(SMALL_NUMBER, GNiagaraRibbonTessellationAngle));
+				const float MinTesselation = FMath::Max<float>(1.f, TessellationAngle / FMath::Max<float>(SMALL_NUMBER, GNiagaraRibbonTessellationAngle));
 
 				const float ViewDistance = SceneProxy->GetBounds().ComputeSquaredDistanceFromBoxToPoint(View->ViewLocation);
 				const float MaxDisplacementError = FMath::Max(GNiagaraRibbonTessellationMinDisplacementError, GNiagaraRibbonTessellationScreenPercentage * FMath::Sqrt(ViewDistance) / View->LODDistanceFactor);
-				const float Tess = DynamicDataRibbon->AverangeSegmentAngle / FMath::Max(SMALL_NUMBER, AcosFast(DynamicDataRibbon->AverageSegmentCurvature / (DynamicDataRibbon->AverageSegmentCurvature + MaxDisplacementError)));
+				const float Tess = TessellationAngle / FMath::Max(SMALL_NUMBER, AcosFast(TessellationCurvature / (TessellationCurvature + MaxDisplacementError)));
+				// FMath::RoundUpToPowerOfTwo ? This could avoid vertices moving around as tesselation increases
 				SegmentTessellation = FMath::Clamp<int32>(FMath::RoundToInt(Tess), FMath::RoundToInt(MinTesselation), GNiagaraRibbonMaxTessellation);
 				NumSegments *= SegmentTessellation;
 			}
@@ -523,7 +518,7 @@ void CalculateUVScaleAndOffsets(
 	}
 }
 
-FNiagaraDynamicDataBase *NiagaraRendererRibbons::GenerateVertexData(const FNiagaraSceneProxy* Proxy, FNiagaraDataSet &Data, const ENiagaraSimTarget Target)
+FNiagaraDynamicDataBase* NiagaraRendererRibbons::GenerateVertexData(const FNiagaraSceneProxy* Proxy, FNiagaraDataSet &Data, const ENiagaraSimTarget Target)
 {
 	SCOPE_CYCLE_COUNTER(STAT_NiagaraGenRibbonVertexData);
 
@@ -610,7 +605,7 @@ FNiagaraDynamicDataBase *NiagaraRendererRibbons::GenerateVertexData(const FNiaga
 		int32 StartIndex = DynamicData->SortedIndices.Num();
 		int32 NumIndices = RibbonIndices.Num();
 
-			float TotalDistance = 0.0f;
+		float TotalDistance = 0.0f;
 
 		const FVector FirstPos = PosData[RibbonIndices[0]];
 		FVector CurrPos = FirstPos;
@@ -620,14 +615,14 @@ FNiagaraDynamicDataBase *NiagaraRendererRibbons::GenerateVertexData(const FNiaga
 		// Find the first position with enough distance.
 		int32 CurrentIndex = 1;
 		while (CurrentIndex < RibbonIndices.Num())
-			{
+		{
 			CurrPos = PosData[RibbonIndices[CurrentIndex]];
 			LastToCurrVec = CurrPos - FirstPos;
 			LastToCurrSize = LastToCurrVec.Size();
 
 			// Find the first segment, or unique segment
 			if (LastToCurrSize > GNiagaraRibbonMinSegmentLength)
-				{
+			{
 				// Normalize LastToCurrVec
 				LastToCurrVec *= 1.f / LastToCurrSize;
 
@@ -637,18 +632,18 @@ FNiagaraDynamicDataBase *NiagaraRendererRibbons::GenerateVertexData(const FNiaga
 				DynamicData->MultiRibbonIndices.Add(RibbonIndex);
 
 				break;
-				}
-				else
-				{
+			}
+			else
+			{
 				LastToCurrSize = 0; // Ensure that the segment gets ignored if too small
 				++CurrentIndex;
 			}
-				}
+		}
 
 		// Now iterate on all other points, to proceed each particle connected to 2 segments.
 		int32 NextIndex = CurrentIndex + 1;
 		while (NextIndex < RibbonIndices.Num())
-				{
+		{
 			const FVector NextPos = PosData[RibbonIndices[NextIndex]];
 			FVector CurrToNextVec = NextPos - CurrPos;
 			const float CurrToNextSize = CurrToNextVec.Size();
@@ -678,11 +673,11 @@ FNiagaraDynamicDataBase *NiagaraRendererRibbons::GenerateVertexData(const FNiaga
 				CurrPos = NextPos;
 				LastToCurrVec = CurrToNextVec;
 				LastToCurrSize = CurrToNextSize;
-					}
+			}
 
 			// Try next if there is one.
 			++NextIndex;
-				}
+		}
 
 		// Close the last point and segment if there was at least 2.
 		if (LastToCurrSize > 0)
@@ -694,7 +689,7 @@ FNiagaraDynamicDataBase *NiagaraRendererRibbons::GenerateVertexData(const FNiaga
 			DynamicData->SortedIndices.Add(RibbonIndices[CurrentIndex]);
 			DynamicData->TangentAndDistances.Add(FVector4(LastToCurrVec.X, LastToCurrVec.Y, LastToCurrVec.Z, TotalDistance));
 			DynamicData->MultiRibbonIndices.Add(RibbonIndex);
-			}
+		}
 
 		const int32 EndIndex = DynamicData->SortedIndices.Num() - 1;
 		const int32 NumSegments = EndIndex - StartIndex;
@@ -711,7 +706,7 @@ FNiagaraDynamicDataBase *NiagaraRendererRibbons::GenerateVertexData(const FNiaga
 				FVector& LastTangent = reinterpret_cast<FVector&>(DynamicData->TangentAndDistances[EndIndex]);
 				FVector& PrevToLastTangent = reinterpret_cast<FVector&>(DynamicData->TangentAndDistances[EndIndex - 1]);
 				LastTangent = (2.f * FVector::DotProduct(LastTangent, PrevToLastTangent)) * LastTangent - PrevToLastTangent;
-		}
+			}
 
 			// Add segment data
 			for (int32 SegmentIndex = StartIndex; SegmentIndex < EndIndex; ++SegmentIndex)
@@ -721,10 +716,10 @@ FNiagaraDynamicDataBase *NiagaraRendererRibbons::GenerateVertexData(const FNiaga
 
 			float U0Offset, U0Scale, U1Offset, U1Scale;
 
-		CalculateUVScaleAndOffsets(SortKeyData, DynamicData->SortedIndices, bSortKeyIsAge, StartIndex, DynamicData->SortedIndices.Num() - 1, NumSegments,
-			Properties->UV0TilingDistance,	Properties->UV0Scale.X, Properties->UV0Offset.X, Properties->UV0AgeOffsetMode, U0Scale, U0Offset);
-		CalculateUVScaleAndOffsets(SortKeyData, DynamicData->SortedIndices, bSortKeyIsAge, StartIndex, DynamicData->SortedIndices.Num() - 1, NumSegments,
-			Properties->UV1TilingDistance, Properties->UV1Scale.X, Properties->UV1Offset.X, Properties->UV1AgeOffsetMode, U1Scale, U1Offset);
+			CalculateUVScaleAndOffsets(SortKeyData, DynamicData->SortedIndices, bSortKeyIsAge, StartIndex, DynamicData->SortedIndices.Num() - 1, NumSegments,
+				Properties->UV0TilingDistance,	Properties->UV0Scale.X, Properties->UV0Offset.X, Properties->UV0AgeOffsetMode, U0Scale, U0Offset);
+			CalculateUVScaleAndOffsets(SortKeyData, DynamicData->SortedIndices, bSortKeyIsAge, StartIndex, DynamicData->SortedIndices.Num() - 1, NumSegments,
+				Properties->UV1TilingDistance, Properties->UV1Scale.X, Properties->UV1Offset.X, Properties->UV1AgeOffsetMode, U1Scale, U1Offset);
 
 			DynamicData->PackPerRibbonData(U0Scale, U0Offset, U1Scale, U1Offset, NumSegments, StartIndex);
 		}
@@ -798,10 +793,21 @@ FNiagaraDynamicDataBase *NiagaraRendererRibbons::GenerateVertexData(const FNiaga
 	if (NumTotalSamples > 0)
 	{
 		const float OneOverSampleCount = 1.f / (float)NumTotalSamples;
-		DynamicData->AverangeSegmentAngle = TotalSegmentAngle * OneOverSampleCount;
-		DynamicData->AverageSegmentCurvature = TotalSegmentLength * OneOverSampleCount * .5 / (FMath::Max(SMALL_NUMBER, FMath::Abs(FMath::Sin(DynamicData->AverangeSegmentAngle))));
-	}
+		const float AverageSegmentAngle = TotalSegmentAngle * OneOverSampleCount;
+		const float AverageSegmentCurvature = TotalSegmentLength * OneOverSampleCount * .5 / (FMath::Max(SMALL_NUMBER, FMath::Abs(FMath::Sin(AverageSegmentAngle))));
+
+		// Accumulate the max so that tessellation is ever increasing and stabilizes over time.
+		// This is to avoid the spawning and unspawning of ribbons to constantly generate different tesselation factors.
+		TessellationAngle = FMath::Max(TessellationAngle, AverageSegmentAngle);
+		TessellationCurvature = FMath::Max(TessellationCurvature, AverageSegmentCurvature);
 
+	}
+	else // Reset the metrics when the ribbons are reset.
+	{
+		TessellationAngle = 0;
+		TessellationCurvature = 0;
+	}
+	 
 	if (Data.CurrData().GetNumInstances() > 0)
 	{
 		//TODO: This buffer is far fatter than needed. Just pull out the data needed for rendering.
@@ -814,5 +820,3 @@ FNiagaraDynamicDataBase *NiagaraRendererRibbons::GenerateVertexData(const FNiaga
 
 	return DynamicData;
 }
-
-
diff --git a/Engine/Plugins/FX/Niagara/Source/Niagara/Public/NiagaraRendererRibbons.h b/Engine/Plugins/FX/Niagara/Source/Niagara/Public/NiagaraRendererRibbons.h
--- a/Engine/Plugins/FX/Niagara/Source/Niagara/Public/NiagaraRendererRibbons.h
+++ b/Engine/Plugins/FX/Niagara/Source/Niagara/Public/NiagaraRendererRibbons.h
@@ -85,6 +85,11 @@ private:
 	int32 MaterialParamOffset1;
 	int32 MaterialParamOffset2;
 	int32 MaterialParamOffset3;
+
+	// Average curvature of the segments.
+	float TessellationAngle = 0;
+	// Average angle of the curvature of the segments (in radian).
+	float TessellationCurvature = 0;
 };
 
 
