Niagara - Remove the GEditor checks from the skeletal mesh data interface so that editor behavior matches cook behavior, also add a cvar to force systems spawned attached to run in solo mode for debugging.

diff --git a/Engine/Plugins/FX/Niagara/Source/Niagara/Private/NiagaraDataInterfaceSkeletalMesh.cpp b/Engine/Plugins/FX/Niagara/Source/Niagara/Private/NiagaraDataInterfaceSkeletalMesh.cpp
--- a/Engine/Plugins/FX/Niagara/Source/Niagara/Private/NiagaraDataInterfaceSkeletalMesh.cpp
+++ b/Engine/Plugins/FX/Niagara/Source/Niagara/Private/NiagaraDataInterfaceSkeletalMesh.cpp
@@ -432,7 +432,7 @@ bool FNDISkeletalMesh_InstanceData::Init(UNiagaraDataInterfaceSkeletalMesh* Inte
 			LODIndex = FMath::Clamp(Interface->WholeMeshLOD, 0, Mesh->GetLODNum() - 1);
 		}
 
-		if (!Mesh->GetLODInfo(LODIndex)->bAllowCPUAccess && (GIsEditor || Interface->bUseTriangleSampling || Interface->bUseVertexSampling))
+		if (!Mesh->GetLODInfo(LODIndex)->bAllowCPUAccess && (Interface->bUseTriangleSampling || Interface->bUseVertexSampling))
 		{
 			UE_LOG(LogNiagara, Warning, TEXT("Skeletal Mesh Data Interface is trying to spawn from a whole mesh that does not allow CPU Access.\nInterface: %s\nMesh: %s\nLOD: %d"),
 				*Interface->GetFullName(),
@@ -512,7 +512,7 @@ bool FNDISkeletalMesh_InstanceData::Init(UNiagaraDataInterfaceSkeletalMesh* Inte
 	bool bNeedDataImmediately = true;
 		
 	//Grab a handle to the skinning data if we have a component to skin.
-	ENDISkeletalMesh_SkinningMode SkinningMode = (GIsEditor || Interface->bUseTriangleSampling || Interface->bUseVertexSampling) ? Interface->SkinningMode : ENDISkeletalMesh_SkinningMode::None;
+	ENDISkeletalMesh_SkinningMode SkinningMode = (Interface->bUseTriangleSampling || Interface->bUseVertexSampling) ? Interface->SkinningMode : ENDISkeletalMesh_SkinningMode::None;
 	FSkeletalMeshSkinningDataUsage Usage(
 		LODIndex,
 		SkinningMode == ENDISkeletalMesh_SkinningMode::SkinOnTheFly || SkinningMode == ENDISkeletalMesh_SkinningMode::PreSkin || Interface->bUseSkeletonSampling,
@@ -906,7 +906,7 @@ TArray<FNiagaraDataInterfaceError> UNiagaraDataInterfaceSkeletalMesh::GetErrors(
 	bool bHasNoMeshAssignedError = false;
 	
 	// Collect Errors
-	if (DefaultMesh != nullptr && (GIsEditor || bUseTriangleSampling || bUseVertexSampling))
+	if (DefaultMesh != nullptr && (bUseTriangleSampling || bUseVertexSampling))
 	{
 		for (auto info : DefaultMesh->GetLODInfoArray())
 		{
diff --git a/Engine/Plugins/FX/Niagara/Source/Niagara/Private/NiagaraFunctionLibrary.cpp b/Engine/Plugins/FX/Niagara/Source/Niagara/Private/NiagaraFunctionLibrary.cpp
--- a/Engine/Plugins/FX/Niagara/Source/Niagara/Private/NiagaraFunctionLibrary.cpp
+++ b/Engine/Plugins/FX/Niagara/Source/Niagara/Private/NiagaraFunctionLibrary.cpp
@@ -10,6 +10,16 @@
 
 #include "NiagaraWorldManager.h"
 
+#if WITH_EDITOR
+int32 GForceNiagaraSpawnAttachedSolo = 0;
+static FAutoConsoleVariableRef CVarForceNiagaraSpawnAttachedSolo(
+	TEXT("fx.ForceNiagaraSpawnAttachedSolo"),
+	GForceNiagaraSpawnAttachedSolo,
+	TEXT("If > 0 Niagara systems which are spawned attached will be force to spawn in solo mode for debugging.\n"),
+	ECVF_Default
+);
+#endif
+
 UNiagaraFunctionLibrary::UNiagaraFunctionLibrary(const FObjectInitializer& ObjectInitializer)
 : Super(ObjectInitializer)
 {
@@ -76,6 +86,14 @@ UNiagaraComponent* UNiagaraFunctionLibrary::SpawnSystemAttached(UNiagaraSystem*
 		else
 		{
 			PSC = CreateNiagaraSystem(SystemTemplate, AttachToComponent->GetWorld(), AttachToComponent->GetOwner(), bAutoDestroy, EPSCPoolMethod::None);
+
+#if WITH_EDITOR
+			if (GForceNiagaraSpawnAttachedSolo > 0)
+			{
+				PSC->SetForceSolo(true);
+			}
+#endif
+
 			PSC->RegisterComponentWithWorld(AttachToComponent->GetWorld());
 
 			PSC->AttachToComponent(AttachToComponent, FAttachmentTransformRules::KeepRelativeTransform, AttachPointName);
@@ -125,6 +143,12 @@ UNiagaraComponent* UNiagaraFunctionLibrary::SpawnSystemAttached(
 				PSC = CreateNiagaraSystem(SystemTemplate, World, AttachToComponent->GetOwner(), bAutoDestroy, PoolingMethod);
 				if (PSC)
 				{
+#if WITH_EDITOR
+					if (GForceNiagaraSpawnAttachedSolo > 0)
+					{
+						PSC->SetForceSolo(true);
+					}
+#endif
 					PSC->SetupAttachment(AttachToComponent, AttachPointName);
 
 					if (LocationType == EAttachLocation::KeepWorldPosition)
