Searching in stack view now finds text in linked fields.

diff --git a/Engine/Plugins/FX/Niagara/Source/NiagaraEditor/Private/TypeEditorUtilities/NiagaraColorTypeEditorUtilities.cpp b/Engine/Plugins/FX/Niagara/Source/NiagaraEditor/Private/TypeEditorUtilities/NiagaraColorTypeEditorUtilities.cpp
--- a/Engine/Plugins/FX/Niagara/Source/NiagaraEditor/Private/TypeEditorUtilities/NiagaraColorTypeEditorUtilities.cpp
+++ b/Engine/Plugins/FX/Niagara/Source/NiagaraEditor/Private/TypeEditorUtilities/NiagaraColorTypeEditorUtilities.cpp
@@ -228,3 +228,8 @@ bool FNiagaraEditorColorTypeUtilities::SetValueFromPinDefaultString(const FStrin
 	}
 	return false;
 }
+
+FText FNiagaraEditorColorTypeUtilities::GetSearchTextFromValue(const FNiagaraVariable& AllocatedVariable) const
+{
+	return FText::FromString(GetPinDefaultStringFromValue(AllocatedVariable));
+}
diff --git a/Engine/Plugins/FX/Niagara/Source/NiagaraEditor/Private/TypeEditorUtilities/NiagaraColorTypeEditorUtilities.h b/Engine/Plugins/FX/Niagara/Source/NiagaraEditor/Private/TypeEditorUtilities/NiagaraColorTypeEditorUtilities.h
--- a/Engine/Plugins/FX/Niagara/Source/NiagaraEditor/Private/TypeEditorUtilities/NiagaraColorTypeEditorUtilities.h
+++ b/Engine/Plugins/FX/Niagara/Source/NiagaraEditor/Private/TypeEditorUtilities/NiagaraColorTypeEditorUtilities.h
@@ -18,4 +18,5 @@ public:
 	virtual bool CanHandlePinDefaults() const override;
 	virtual FString GetPinDefaultStringFromValue(const FNiagaraVariable& AllocatedVariable) const override;
 	virtual bool SetValueFromPinDefaultString(const FString& StringValue, FNiagaraVariable& Variable) const override;
+	virtual FText GetSearchTextFromValue(const FNiagaraVariable& AllocatedVariable) const override;
 };
\ No newline at end of file
diff --git a/Engine/Plugins/FX/Niagara/Source/NiagaraEditor/Private/TypeEditorUtilities/NiagaraEnumTypeEditorUtilities.cpp b/Engine/Plugins/FX/Niagara/Source/NiagaraEditor/Private/TypeEditorUtilities/NiagaraEnumTypeEditorUtilities.cpp
--- a/Engine/Plugins/FX/Niagara/Source/NiagaraEditor/Private/TypeEditorUtilities/NiagaraEnumTypeEditorUtilities.cpp
+++ b/Engine/Plugins/FX/Niagara/Source/NiagaraEditor/Private/TypeEditorUtilities/NiagaraEnumTypeEditorUtilities.cpp
@@ -244,4 +244,15 @@ bool FNiagaraEditorEnumTypeUtilities::SetValueFromDisplayName(const FText& TextV
 		return true;
 	}
 	return false;
-}
\ No newline at end of file
+}
+
+FText FNiagaraEditorEnumTypeUtilities::GetSearchTextFromValue(const FNiagaraVariable& AllocatedVariable) const
+{
+	checkf(AllocatedVariable.IsDataAllocated(), TEXT("Can not generate search text for an unallocated variable."));
+
+	const UEnum* Enum = AllocatedVariable.GetType().GetEnum();
+	checkf(Enum != nullptr, TEXT("Variable is not an enum type."));
+
+	const int32 EnumNameIndex = Enum->GetIndexByValue(AllocatedVariable.GetValue<int32>());
+	return Enum->GetDisplayNameTextByIndex(EnumNameIndex);
+}
diff --git a/Engine/Plugins/FX/Niagara/Source/NiagaraEditor/Private/TypeEditorUtilities/NiagaraEnumTypeEditorUtilities.h b/Engine/Plugins/FX/Niagara/Source/NiagaraEditor/Private/TypeEditorUtilities/NiagaraEnumTypeEditorUtilities.h
--- a/Engine/Plugins/FX/Niagara/Source/NiagaraEditor/Private/TypeEditorUtilities/NiagaraEnumTypeEditorUtilities.h
+++ b/Engine/Plugins/FX/Niagara/Source/NiagaraEditor/Private/TypeEditorUtilities/NiagaraEnumTypeEditorUtilities.h
@@ -19,4 +19,5 @@ public:
 	virtual bool SetValueFromPinDefaultString(const FString& StringValue, FNiagaraVariable& Variable) const override;
 	virtual bool CanSetValueFromDisplayName() const override;
 	virtual bool SetValueFromDisplayName(const FText& TextValue, FNiagaraVariable& Variable) const override;
+	virtual FText GetSearchTextFromValue(const FNiagaraVariable& AllocatedVariable) const override;
 };
\ No newline at end of file
diff --git a/Engine/Plugins/FX/Niagara/Source/NiagaraEditor/Private/TypeEditorUtilities/NiagaraFloatTypeEditorUtilities.cpp b/Engine/Plugins/FX/Niagara/Source/NiagaraEditor/Private/TypeEditorUtilities/NiagaraFloatTypeEditorUtilities.cpp
--- a/Engine/Plugins/FX/Niagara/Source/NiagaraEditor/Private/TypeEditorUtilities/NiagaraFloatTypeEditorUtilities.cpp
+++ b/Engine/Plugins/FX/Niagara/Source/NiagaraEditor/Private/TypeEditorUtilities/NiagaraFloatTypeEditorUtilities.cpp
@@ -109,3 +109,8 @@ bool FNiagaraEditorFloatTypeUtilities::SetValueFromPinDefaultString(const FStrin
 	}
 	return false;
 }
+
+FText FNiagaraEditorFloatTypeUtilities::GetSearchTextFromValue(const FNiagaraVariable& AllocatedVariable) const
+{
+	return FText::FromString(GetPinDefaultStringFromValue(AllocatedVariable));
+}
diff --git a/Engine/Plugins/FX/Niagara/Source/NiagaraEditor/Private/TypeEditorUtilities/NiagaraFloatTypeEditorUtilities.h b/Engine/Plugins/FX/Niagara/Source/NiagaraEditor/Private/TypeEditorUtilities/NiagaraFloatTypeEditorUtilities.h
--- a/Engine/Plugins/FX/Niagara/Source/NiagaraEditor/Private/TypeEditorUtilities/NiagaraFloatTypeEditorUtilities.h
+++ b/Engine/Plugins/FX/Niagara/Source/NiagaraEditor/Private/TypeEditorUtilities/NiagaraFloatTypeEditorUtilities.h
@@ -16,4 +16,5 @@ public:
 	virtual bool CanHandlePinDefaults() const override;
 	virtual FString GetPinDefaultStringFromValue(const FNiagaraVariable& AllocatedVariable) const override;
 	virtual bool SetValueFromPinDefaultString(const FString& StringValue, FNiagaraVariable& Variable) const override;
+	virtual FText GetSearchTextFromValue(const FNiagaraVariable& AllocatedVariable) const override;
 };
\ No newline at end of file
diff --git a/Engine/Plugins/FX/Niagara/Source/NiagaraEditor/Private/TypeEditorUtilities/NiagaraIntegerTypeEditorUtilities.cpp b/Engine/Plugins/FX/Niagara/Source/NiagaraEditor/Private/TypeEditorUtilities/NiagaraIntegerTypeEditorUtilities.cpp
--- a/Engine/Plugins/FX/Niagara/Source/NiagaraEditor/Private/TypeEditorUtilities/NiagaraIntegerTypeEditorUtilities.cpp
+++ b/Engine/Plugins/FX/Niagara/Source/NiagaraEditor/Private/TypeEditorUtilities/NiagaraIntegerTypeEditorUtilities.cpp
@@ -109,4 +109,9 @@ bool FNiagaraEditorIntegerTypeUtilities::SetValueFromPinDefaultString(const FStr
 		return true;
 	}
 	return false;
-}
\ No newline at end of file
+}
+
+FText FNiagaraEditorIntegerTypeUtilities::GetSearchTextFromValue(const FNiagaraVariable& AllocatedVariable) const
+{
+	return FText::FromString(GetPinDefaultStringFromValue(AllocatedVariable));
+}
diff --git a/Engine/Plugins/FX/Niagara/Source/NiagaraEditor/Private/TypeEditorUtilities/NiagaraIntegerTypeEditorUtilities.h b/Engine/Plugins/FX/Niagara/Source/NiagaraEditor/Private/TypeEditorUtilities/NiagaraIntegerTypeEditorUtilities.h
--- a/Engine/Plugins/FX/Niagara/Source/NiagaraEditor/Private/TypeEditorUtilities/NiagaraIntegerTypeEditorUtilities.h
+++ b/Engine/Plugins/FX/Niagara/Source/NiagaraEditor/Private/TypeEditorUtilities/NiagaraIntegerTypeEditorUtilities.h
@@ -15,4 +15,5 @@ public:
 	virtual bool CanHandlePinDefaults() const override;
 	virtual FString GetPinDefaultStringFromValue(const FNiagaraVariable& AllocatedVariable) const override;
 	virtual bool SetValueFromPinDefaultString(const FString& StringValue, FNiagaraVariable& Variable) const override;
+	virtual FText GetSearchTextFromValue(const FNiagaraVariable& AllocatedVariable) const override;
 };
\ No newline at end of file
diff --git a/Engine/Plugins/FX/Niagara/Source/NiagaraEditor/Private/TypeEditorUtilities/NiagaraVectorTypeEditorUtilities.cpp b/Engine/Plugins/FX/Niagara/Source/NiagaraEditor/Private/TypeEditorUtilities/NiagaraVectorTypeEditorUtilities.cpp
--- a/Engine/Plugins/FX/Niagara/Source/NiagaraEditor/Private/TypeEditorUtilities/NiagaraVectorTypeEditorUtilities.cpp
+++ b/Engine/Plugins/FX/Niagara/Source/NiagaraEditor/Private/TypeEditorUtilities/NiagaraVectorTypeEditorUtilities.cpp
@@ -174,6 +174,11 @@ bool FNiagaraEditorVector2TypeUtilities::SetValueFromPinDefaultString(const FStr
 	return false;
 }
 
+FText FNiagaraEditorVector2TypeUtilities::GetSearchTextFromValue(const FNiagaraVariable& AllocatedVariable) const
+{
+	return FText::FromString(GetPinDefaultStringFromValue(AllocatedVariable));
+}
+
 class SNiagaraVector3ParameterEditor : public SNiagaraVectorParameterEditorBase
 {
 public:
@@ -245,6 +250,11 @@ bool FNiagaraEditorVector3TypeUtilities::SetValueFromPinDefaultString(const FStr
 	return false;
 }
 
+FText FNiagaraEditorVector3TypeUtilities::GetSearchTextFromValue(const FNiagaraVariable& AllocatedVariable) const
+{
+	return FText::FromString(GetPinDefaultStringFromValue(AllocatedVariable));
+}
+
 class SNiagaraVector4ParameterEditor : public SNiagaraVectorParameterEditorBase
 {
 public:
@@ -316,6 +326,11 @@ bool FNiagaraEditorVector4TypeUtilities::SetValueFromPinDefaultString(const FStr
 }
 
 
+FText FNiagaraEditorVector4TypeUtilities::GetSearchTextFromValue(const FNiagaraVariable& AllocatedVariable) const
+{
+	return FText::FromString(GetPinDefaultStringFromValue(AllocatedVariable));
+}
+
 class SNiagaraQuatParameterEditor : public SNiagaraVectorParameterEditorBase
 {
 public:
@@ -416,6 +431,11 @@ bool FNiagaraEditorQuatTypeUtilities::SetValueFromPinDefaultString(const FString
 	return false;
 }
 
+FText FNiagaraEditorQuatTypeUtilities::GetSearchTextFromValue(const FNiagaraVariable& AllocatedVariable) const
+{
+	return FText::FromString(GetPinDefaultStringFromValue(AllocatedVariable));
+}
+
 void FNiagaraEditorQuatTypeUtilities::UpdateVariableWithDefaultValue(FNiagaraVariable& Variable) const
 {
 	checkf(Variable.GetType().GetStruct() == FNiagaraTypeDefinition::GetQuatStruct(), TEXT("Struct type not supported."));
diff --git a/Engine/Plugins/FX/Niagara/Source/NiagaraEditor/Private/TypeEditorUtilities/NiagaraVectorTypeEditorUtilities.h b/Engine/Plugins/FX/Niagara/Source/NiagaraEditor/Private/TypeEditorUtilities/NiagaraVectorTypeEditorUtilities.h
--- a/Engine/Plugins/FX/Niagara/Source/NiagaraEditor/Private/TypeEditorUtilities/NiagaraVectorTypeEditorUtilities.h
+++ b/Engine/Plugins/FX/Niagara/Source/NiagaraEditor/Private/TypeEditorUtilities/NiagaraVectorTypeEditorUtilities.h
@@ -16,6 +16,7 @@ public:
 	virtual bool CanHandlePinDefaults() const override;
 	virtual FString GetPinDefaultStringFromValue(const FNiagaraVariable& AllocatedVariable) const override;
 	virtual bool SetValueFromPinDefaultString(const FString& StringValue, FNiagaraVariable& Variable) const override;
+	virtual FText GetSearchTextFromValue(const FNiagaraVariable& AllocatedVariable) const override;
 };
 
 /** Niagara editor utilities for the FVector type. */
@@ -28,6 +29,7 @@ public:
 	virtual bool CanHandlePinDefaults() const override;
 	virtual FString GetPinDefaultStringFromValue(const FNiagaraVariable& AllocatedVariable) const override;
 	virtual bool SetValueFromPinDefaultString(const FString& StringValue, FNiagaraVariable& Variable) const override;
+	virtual FText GetSearchTextFromValue(const FNiagaraVariable& AllocatedVariable) const override;
 };
 
 /** Niagara editor utilities for the FVector4 type. */
@@ -40,6 +42,7 @@ public:
 	virtual bool CanHandlePinDefaults() const override;
 	virtual FString GetPinDefaultStringFromValue(const FNiagaraVariable& AllocatedVariable) const override;
 	virtual bool SetValueFromPinDefaultString(const FString& StringValue, FNiagaraVariable& Variable) const override;
+	virtual FText GetSearchTextFromValue(const FNiagaraVariable& AllocatedVariable) const override;
 };
 
 /** Niagara editor utilities for the FVector4 type. */
@@ -54,4 +57,5 @@ public:
 	virtual bool CanHandlePinDefaults() const override;
 	virtual FString GetPinDefaultStringFromValue(const FNiagaraVariable& AllocatedVariable) const override;
 	virtual bool SetValueFromPinDefaultString(const FString& StringValue, FNiagaraVariable& Variable) const override;
+	virtual FText GetSearchTextFromValue(const FNiagaraVariable& AllocatedVariable) const override;
 };
\ No newline at end of file
diff --git a/Engine/Plugins/FX/Niagara/Source/NiagaraEditor/Private/ViewModels/Stack/NiagaraStackErrorItem.cpp b/Engine/Plugins/FX/Niagara/Source/NiagaraEditor/Private/ViewModels/Stack/NiagaraStackErrorItem.cpp
--- a/Engine/Plugins/FX/Niagara/Source/NiagaraEditor/Private/ViewModels/Stack/NiagaraStackErrorItem.cpp
+++ b/Engine/Plugins/FX/Niagara/Source/NiagaraEditor/Private/ViewModels/Stack/NiagaraStackErrorItem.cpp
@@ -35,6 +35,12 @@ UNiagaraStackErrorItem::FOnIssueNotify& UNiagaraStackErrorItem::OnIssueModified(
 	return IssueModifiedDelegate;
 }
 
+void UNiagaraStackErrorItem::GetSearchItems(TArray<FStackSearchItem>& SearchItems) const
+{
+	SearchItems.Add({ FName("ErrorShortDescription"), StackIssue.GetShortDescription() });
+	SearchItems.Add({ FName("ErrorLongDescription"), StackIssue.GetLongDescription() });
+}
+
 void UNiagaraStackErrorItem::RefreshChildrenInternal(const TArray<UNiagaraStackEntry*>& CurrentChildren, TArray<UNiagaraStackEntry*>& NewChildren, TArray<FStackIssue>& NewIssues)
 {
 	for (UNiagaraStackEntry* Child : CurrentChildren)
diff --git a/Engine/Plugins/FX/Niagara/Source/NiagaraEditor/Private/ViewModels/Stack/NiagaraStackFunctionInput.cpp b/Engine/Plugins/FX/Niagara/Source/NiagaraEditor/Private/ViewModels/Stack/NiagaraStackFunctionInput.cpp
--- a/Engine/Plugins/FX/Niagara/Source/NiagaraEditor/Private/ViewModels/Stack/NiagaraStackFunctionInput.cpp
+++ b/Engine/Plugins/FX/Niagara/Source/NiagaraEditor/Private/ViewModels/Stack/NiagaraStackFunctionInput.cpp
@@ -39,6 +39,7 @@
 #include "AssetRegistryModule.h"
 #include "ARFilter.h"
 #include "EdGraph/EdGraphPin.h"
+#include "INiagaraEditorTypeUtilities.h"
 
 #define LOCTEXT_NAMESPACE "NiagaraStackViewModel"
 
@@ -1552,6 +1553,44 @@ void UNiagaraStackFunctionInput::GetSearchItems(TArray<FStackSearchItem>& Search
 	if (GetShouldPassFilterForVisibleCondition() && GetIsInlineEditConditionToggle() == false)
 	{
 		SearchItems.Add({ FName("DisplayName"), GetDisplayName() });
+
+		if (InputValues.Mode == EValueMode::Local && InputType.IsValid() && InputValues.LocalStruct->IsValid())
+		{
+			FNiagaraVariable LocalValue = FNiagaraVariable(InputType, "");
+			LocalValue.SetData(InputValues.LocalStruct->GetStructMemory());
+			TSharedPtr<INiagaraEditorTypeUtilities, ESPMode::ThreadSafe> ParameterTypeUtilities = FNiagaraEditorModule::Get().GetTypeUtilities(RapidIterationParameter.GetType());
+			if (ParameterTypeUtilities.IsValid() && ParameterTypeUtilities->CanHandlePinDefaults())
+			{
+				FText SearchText = ParameterTypeUtilities->GetSearchTextFromValue(LocalValue);
+				if (SearchText.IsEmpty() == false)
+				{
+					SearchItems.Add({ FName("LocalValueText"), SearchText });
+				}
+			}
+		}
+		else if (InputValues.Mode == EValueMode::Linked)
+		{
+			SearchItems.Add({ FName("LinkedParamName"), FText::FromName(InputValues.LinkedHandle.GetParameterHandleString()) });
+		}
+		else if (InputValues.Mode == EValueMode::Dynamic && InputValues.DynamicNode.Get() != nullptr)
+		{
+			SearchItems.Add({ FName("LinkedDynamicInputName"), InputValues.DynamicNode->GetNodeTitle(ENodeTitleType::MenuTitle)});
+		}
+		else if (InputValues.Mode == EValueMode::Data && InputValues.DataObjects.IsValid())
+		{
+			if (InputValues.DataObjects.GetValueObject() != nullptr)
+			{
+				SearchItems.Add({ FName("LinkedDataInterfaceName"), FText::FromString(InputValues.DataObjects.GetValueObject()->GetName()) });
+			}
+			if (InputValues.DataObjects.GetDefaultValueObject() != nullptr)
+			{
+				SearchItems.Add({ FName("LinkedDataInterfaceDefaultName"), InputValues.DataObjects.GetDefaultValueObject()->GetClass()->GetDisplayNameText() });
+			}
+		}
+		else if (InputValues.Mode == EValueMode::Expression && InputValues.ExpressionNode.Get() != nullptr)
+		{
+			SearchItems.Add({ FName("LinkedExpressionText"), InputValues.ExpressionNode->GetHlslText() });
+		}
 	}
 }
 
diff --git a/Engine/Plugins/FX/Niagara/Source/NiagaraEditor/Private/ViewModels/Stack/NiagaraStackModuleItemOutput.cpp b/Engine/Plugins/FX/Niagara/Source/NiagaraEditor/Private/ViewModels/Stack/NiagaraStackModuleItemOutput.cpp
--- a/Engine/Plugins/FX/Niagara/Source/NiagaraEditor/Private/ViewModels/Stack/NiagaraStackModuleItemOutput.cpp
+++ b/Engine/Plugins/FX/Niagara/Source/NiagaraEditor/Private/ViewModels/Stack/NiagaraStackModuleItemOutput.cpp
@@ -67,6 +67,12 @@ UNiagaraStackEntry::EStackRowStyle UNiagaraStackModuleItemOutput::GetStackRowSty
 	return EStackRowStyle::ItemContent;
 }
 
+void UNiagaraStackModuleItemOutput::GetSearchItems(TArray<FStackSearchItem>& SearchItems) const
+{
+	SearchItems.Add({ FName("DisplayName"), GetDisplayName() });
+	SearchItems.Add({ FName("OutputParamHandleText"), GetOutputParameterHandleText() });
+}
+
 const FNiagaraParameterHandle& UNiagaraStackModuleItemOutput::GetOutputParameterHandle() const
 {
 	return OutputParameterHandle;
diff --git a/Engine/Plugins/FX/Niagara/Source/NiagaraEditor/Private/ViewModels/Stack/NiagaraStackPropertyRow.cpp b/Engine/Plugins/FX/Niagara/Source/NiagaraEditor/Private/ViewModels/Stack/NiagaraStackPropertyRow.cpp
--- a/Engine/Plugins/FX/Niagara/Source/NiagaraEditor/Private/ViewModels/Stack/NiagaraStackPropertyRow.cpp
+++ b/Engine/Plugins/FX/Niagara/Source/NiagaraEditor/Private/ViewModels/Stack/NiagaraStackPropertyRow.cpp
@@ -5,6 +5,7 @@
 
 #include "IDetailTreeNode.h"
 #include "PropertyHandle.h"
+#include "IDetailPropertyRow.h"
 
 void UNiagaraStackPropertyRow::Initialize(FRequiredEntryData InRequiredEntryData, TSharedRef<IDetailTreeNode> InDetailTreeNode, FString InOwnerStackItemEditorDataKey, FString InOwnerStackEditorDataKey, UNiagaraNode* InOwningNiagaraNode)
 {
@@ -62,4 +63,16 @@ void UNiagaraStackPropertyRow::GetSearchItems(TArray<FStackSearchItem>& SearchIt
 	{
 		SearchItems.Add({ "PropertyRowFilterString", FText::FromString(FilterString) });
 	}
+
+	TSharedPtr<IDetailPropertyRow> DetailPropertyRow = DetailTreeNode->GetRow();
+	if (DetailPropertyRow.IsValid())
+	{
+		TSharedPtr<IPropertyHandle> PropertyHandle = DetailPropertyRow->GetPropertyHandle();
+		if (PropertyHandle)
+		{
+			FText PropertyRowHandleText;
+			PropertyHandle->GetValueAsDisplayText(PropertyRowHandleText);
+			SearchItems.Add({ "PropertyRowHandleText", PropertyRowHandleText });
+		}
+	}	
 }
diff --git a/Engine/Plugins/FX/Niagara/Source/NiagaraEditor/Public/INiagaraEditorTypeUtilities.h b/Engine/Plugins/FX/Niagara/Source/NiagaraEditor/Public/INiagaraEditorTypeUtilities.h
--- a/Engine/Plugins/FX/Niagara/Source/NiagaraEditor/Public/INiagaraEditorTypeUtilities.h
+++ b/Engine/Plugins/FX/Niagara/Source/NiagaraEditor/Public/INiagaraEditorTypeUtilities.h
@@ -38,6 +38,8 @@ public:
 	virtual bool CanSetValueFromDisplayName() const = 0;
 
 	virtual bool SetValueFromDisplayName(const FText& TextValue, FNiagaraVariable& Variable) const = 0;
+
+	virtual FText GetSearchTextFromValue(const FNiagaraVariable& AllocatedVariable) const = 0;
 };
 
 class FNiagaraEditorTypeUtilities : public INiagaraEditorTypeUtilities, public TSharedFromThis<FNiagaraEditorTypeUtilities, ESPMode::ThreadSafe>
@@ -57,4 +59,5 @@ public:
 	virtual bool SetValueFromPinDefaultString(const FString& StringValue, FNiagaraVariable& Variable) const override { return false; }
 	virtual bool CanSetValueFromDisplayName() const override { return false; }
 	virtual bool SetValueFromDisplayName(const FText& TextValue, FNiagaraVariable& Variable) const override { return false; }
+	virtual FText GetSearchTextFromValue(const FNiagaraVariable& AllocatedVariable) const override { return FText(); }
 };
\ No newline at end of file
diff --git a/Engine/Plugins/FX/Niagara/Source/NiagaraEditor/Public/ViewModels/Stack/NiagaraStackErrorItem.h b/Engine/Plugins/FX/Niagara/Source/NiagaraEditor/Public/ViewModels/Stack/NiagaraStackErrorItem.h
--- a/Engine/Plugins/FX/Niagara/Source/NiagaraEditor/Public/ViewModels/Stack/NiagaraStackErrorItem.h
+++ b/Engine/Plugins/FX/Niagara/Source/NiagaraEditor/Public/ViewModels/Stack/NiagaraStackErrorItem.h
@@ -20,6 +20,7 @@ public:
 	virtual FText GetDisplayName() const override;
 	virtual EStackRowStyle GetStackRowStyle() const override;
 	FOnIssueNotify& OnIssueModified();
+	virtual void GetSearchItems(TArray<FStackSearchItem>& SearchItems) const override;
 	
 protected:
 	virtual void RefreshChildrenInternal(const TArray<UNiagaraStackEntry*>& CurrentChildren, TArray<UNiagaraStackEntry*>& NewChildren, TArray<FStackIssue>& NewIssues);
diff --git a/Engine/Plugins/FX/Niagara/Source/NiagaraEditor/Public/ViewModels/Stack/NiagaraStackModuleItemOutput.h b/Engine/Plugins/FX/Niagara/Source/NiagaraEditor/Public/ViewModels/Stack/NiagaraStackModuleItemOutput.h
--- a/Engine/Plugins/FX/Niagara/Source/NiagaraEditor/Public/ViewModels/Stack/NiagaraStackModuleItemOutput.h
+++ b/Engine/Plugins/FX/Niagara/Source/NiagaraEditor/Public/ViewModels/Stack/NiagaraStackModuleItemOutput.h
@@ -30,6 +30,7 @@ public:
 	virtual FText GetTooltipText() const override;
 	virtual bool GetIsEnabled() const override;
 	virtual EStackRowStyle GetStackRowStyle() const override;
+	virtual void GetSearchItems(TArray<FStackSearchItem>& SearchItems) const override;
 
 	/** Gets the parameter handle which defined this Output in the module. */
 	const FNiagaraParameterHandle& GetOutputParameterHandle() const;
