Niagara - Don't try to merge source changes for emitters with no source information.

diff --git a/Engine/Plugins/FX/Niagara/Source/Niagara/Private/NiagaraEmitterHandle.cpp b/Engine/Plugins/FX/Niagara/Source/Niagara/Private/NiagaraEmitterHandle.cpp
--- a/Engine/Plugins/FX/Niagara/Source/Niagara/Private/NiagaraEmitterHandle.cpp
+++ b/Engine/Plugins/FX/Niagara/Source/Niagara/Private/NiagaraEmitterHandle.cpp
@@ -149,9 +149,15 @@ FString FNiagaraEmitterHandle::GetUniqueInstanceName()const
 
 bool FNiagaraEmitterHandle::IsSynchronizedWithSource() const
 {
+	if (Source == nullptr && LastMergedSource == nullptr)
+	{
+		// If the emitter has no source and no last merged source than it is synchronized by default.
+		return true;
+	}
+
 	if (Source == nullptr || LastMergedSource == nullptr)
 	{
-		// If either the source or last merged sources is missing, then we're not synchronized.  The
+		// If either only the source or only the last merged sources is missing, then we're not synchronized.  The
 		// merge logic will detect this and print an appropriate message to the log.
 		return false;
 	}
