Niagara - Make the Activate function on the niagara component match the cascade component when already activated.

diff --git a/Engine/Plugins/FX/Niagara/Source/Niagara/Private/NiagaraComponent.cpp b/Engine/Plugins/FX/Niagara/Source/Niagara/Private/NiagaraComponent.cpp
--- a/Engine/Plugins/FX/Niagara/Source/Niagara/Private/NiagaraComponent.cpp
+++ b/Engine/Plugins/FX/Niagara/Source/Niagara/Private/NiagaraComponent.cpp
@@ -605,6 +605,16 @@ void UNiagaraComponent::Activate(bool bReset /* = false */)
 	
 	Super::Activate(bReset);
 
+	// Early out if we're not forcing a reset, and both the component and system instance are already active.
+	if (bReset == false &&
+		IsActive() &&
+		SystemInstance != nullptr &&
+		SystemInstance->GetRequestedExecutionState() == ENiagaraExecutionState::Active &&
+		SystemInstance->GetActualExecutionState() == ENiagaraExecutionState::Active)
+	{
+		return;
+	}
+
 	//UE_LOG(LogNiagara, Log, TEXT("Activate: %u - %s"), this, *Asset->GetName());
 	
 	// Auto attach if requested
