Fix missing include for FText

diff --git a/Engine/Plugins/FX/Niagara/Source/NiagaraEditor/Public/INiagaraEditorTypeUtilities.h b/Engine/Plugins/FX/Niagara/Source/NiagaraEditor/Public/INiagaraEditorTypeUtilities.h
--- a/Engine/Plugins/FX/Niagara/Source/NiagaraEditor/Public/INiagaraEditorTypeUtilities.h
+++ b/Engine/Plugins/FX/Niagara/Source/NiagaraEditor/Public/INiagaraEditorTypeUtilities.h
@@ -3,6 +3,7 @@
 #pragma once
 #include "Templates/SharedPointer.h"
 #include "Delegates/Delegate.h"
+#include "Internationalization/Text.h"
 
 class FStructOnScope;
 class SNiagaraParameterEditor;
