- Fix for LastMergedSource nullptr access when duplicating an emitter

diff --git a/Engine/Plugins/FX/Niagara/Source/Niagara/Private/NiagaraEmitterHandle.cpp b/Engine/Plugins/FX/Niagara/Source/Niagara/Private/NiagaraEmitterHandle.cpp
--- a/Engine/Plugins/FX/Niagara/Source/Niagara/Private/NiagaraEmitterHandle.cpp
+++ b/Engine/Plugins/FX/Niagara/Source/Niagara/Private/NiagaraEmitterHandle.cpp
@@ -55,7 +55,10 @@ FNiagaraEmitterHandle::FNiagaraEmitterHandle(const FNiagaraEmitterHandle& InHand
 	Instance->ClearFlags(RF_Standalone | RF_Public);
 	Instance->SetUniqueEmitterName(Name.ToString());
 	Instance->GraphSource->MarkNotSynchronized(InitialNotSynchronizedReason);
-	LastMergedSource->ClearFlags(RF_Standalone | RF_Public);
+	if (LastMergedSource != nullptr)
+	{
+		LastMergedSource->ClearFlags(RF_Standalone | RF_Public);
+	}
 }
 #endif
 
