Added metrics to take into account niagara ribbon twisting in the dynamic tesselation factor.
Fixed wrong view position used in the metrics.
Added Niagara.Ribbon.Tessellation.MaxAveraging to smooth tesselation factor over frames.

diff --git a/Engine/Plugins/FX/Niagara/Source/Niagara/Private/NiagaraRendererRibbons.cpp b/Engine/Plugins/FX/Niagara/Source/Niagara/Private/NiagaraRendererRibbons.cpp
--- a/Engine/Plugins/FX/Niagara/Source/Niagara/Private/NiagaraRendererRibbons.cpp
+++ b/Engine/Plugins/FX/Niagara/Source/Niagara/Private/NiagaraRendererRibbons.cpp
@@ -19,12 +19,12 @@ static FAutoConsoleVariableRef CVarNiagaraRibbonCullTwistedStrips(
 	ECVF_Default
 );
 
-float GNiagaraRibbonTessellationAngle = 30.f * (2.f * PI) / 360.f; // Every 30 degrees
+float GNiagaraRibbonTessellationAngle = 15.f * (2.f * PI) / 360.f; // Every 15 degrees
 static FAutoConsoleVariableRef CVarNiagaraRibbonTessellationAngle(
 	TEXT("Niagara.Ribbon.Tessellation.MinAngle"),
 	GNiagaraRibbonTessellationAngle,
 	TEXT("Ribbon segment angle to tesselate in radian. \n")
-	TEXT("Set 0 to use fixed tesselation. (default=30 degrees)"),
+	TEXT("Set 0 to use fixed tesselation. (default=15 degrees)"),
 	ECVF_Scalability
 );
 
@@ -64,6 +64,15 @@ static FAutoConsoleVariableRef CVarNiagaraRibbonMinSegmentLength(
 	ECVF_Scalability
 );
 
+float GNiagaraRibbonMaxAveraging = .98f;
+static FAutoConsoleVariableRef CVarNiagaraRibbonMaxAveraging(
+	TEXT("Niagara.Ribbon.Tessellation.MaxAveraging"),
+	GNiagaraRibbonMaxAveraging,
+	TEXT("Averaging value between the max metrics and the current metrics. (default=95%)"),
+	ECVF_Scalability
+);
+
+
 
 // max absolute error 9.0x10^-3
 // Eberly's polynomial degree 1 - respect bounds
@@ -76,7 +85,6 @@ FORCEINLINE float AcosFast(float InX)
     return (InX >= 0) ? Res : PI - Res;
 }
 
-
 struct FNiagaraDynamicDataRibbon : public FNiagaraDynamicDataBase
 {
 	// The list of all segments, each one connecting SortedIndices[SegmentId] to SortedIndices[SegmentId + 1].
@@ -265,6 +273,9 @@ void NiagaraRendererRibbons::GetDynamicMeshElements(const TArray<const FSceneVie
 		if (VisibilityMap & (1 << ViewIndex))
 		{
 			const FSceneView* View = Views[ViewIndex];
+			check(View);
+
+			const FVector ViewOriginForDistanceCulling = View->ViewMatrices.GetViewOrigin();
 
 			int32 SegmentTessellation = 1;
 			int32 NumSegments = DynamicDataRibbon->SegmentData.Num();
@@ -272,10 +283,16 @@ void NiagaraRendererRibbons::GetDynamicMeshElements(const TArray<const FSceneVie
 			{
 				const float MinTesselation = FMath::Max<float>(1.f, TessellationAngle / FMath::Max<float>(SMALL_NUMBER, GNiagaraRibbonTessellationAngle));
 
-				const float ViewDistance = SceneProxy->GetBounds().ComputeSquaredDistanceFromBoxToPoint(View->ViewLocation);
+				const float ViewDistance = SceneProxy->GetBounds().ComputeSquaredDistanceFromBoxToPoint(ViewOriginForDistanceCulling);
 				const float MaxDisplacementError = FMath::Max(GNiagaraRibbonTessellationMinDisplacementError, GNiagaraRibbonTessellationScreenPercentage * FMath::Sqrt(ViewDistance) / View->LODDistanceFactor);
-				const float Tess = TessellationAngle / FMath::Max(SMALL_NUMBER, AcosFast(TessellationCurvature / (TessellationCurvature + MaxDisplacementError)));
+				float Tess = TessellationAngle / FMath::Max(SMALL_NUMBER, AcosFast(TessellationCurvature / (TessellationCurvature + MaxDisplacementError)));
 				// FMath::RoundUpToPowerOfTwo ? This could avoid vertices moving around as tesselation increases
+
+				if (TessellationTwistAngle > 0 && TessellationTwistCurvature > 0)
+				{
+					const float TwistTess  = TessellationTwistAngle / FMath::Max(SMALL_NUMBER, AcosFast(TessellationTwistCurvature / (TessellationTwistCurvature + MaxDisplacementError)));
+					Tess = FMath::Max(TwistTess, Tess);
+				}
 				SegmentTessellation = FMath::Clamp<int32>(FMath::RoundToInt(Tess), FMath::RoundToInt(MinTesselation), GNiagaraRibbonMaxTessellation);
 				NumSegments *= SegmentTessellation;
 			}
@@ -284,8 +301,8 @@ void NiagaraRendererRibbons::GetDynamicMeshElements(const TArray<const FSceneVie
 
 			// Figure out whether start is closer to the view plane than end
 			// TODO : This doesn't work with multi-ribbons.
-			const float StartDist = FVector::DotProduct(View->GetViewDirection(), DynamicDataRibbon->StartPos - View->ViewLocation);
-			const float EndDist = FVector::DotProduct(View->GetViewDirection(), DynamicDataRibbon->EndPos - View->ViewLocation);
+			const float StartDist = FVector::DotProduct(View->GetViewDirection(), DynamicDataRibbon->StartPos - ViewOriginForDistanceCulling);
+			const float EndDist = FVector::DotProduct(View->GetViewDirection(), DynamicDataRibbon->EndPos - ViewOriginForDistanceCulling);
 			const bool bInvertOrder = ((StartDist > EndDist) && Properties->DrawDirection == ENiagaraRibbonDrawDirection::BackToFront)
 					|| ((StartDist < EndDist) && Properties->DrawDirection == ENiagaraRibbonDrawDirection::FrontToBack);
 
@@ -531,6 +548,8 @@ FNiagaraDynamicDataBase* NiagaraRendererRibbons::GenerateVertexData(const FNiaga
 	TArray<int32>& SegmentData = DynamicData->SegmentData;
 	float TotalSegmentLength = 0;
 	float TotalSegmentAngle = 0;
+	float TotalTwistAngle = 0;
+	float TotalWidth = 0;
 	int NumTotalSamples = 0;
 
 	FNiagaraDataSetAccessor<FVector> PosData(Data, Properties->PositionBinding.DataSetVariable);
@@ -599,6 +618,7 @@ FNiagaraDynamicDataBase* NiagaraRendererRibbons::GenerateVertexData(const FNiaga
 	bool bFullIDs = RibbonFullIDData.IsValid();
 	bool bSimpleIDs = !bFullIDs && RibbonIdData.IsValid();
 	bool bMultiRibbons = bFullIDs || bSimpleIDs;
+	bool bHasTwist = TwistData.IsValid() && SizeData.IsValid();
 
 	auto AddRibbonVerts = [&](TArray<int32>& RibbonIndices, uint32 RibbonIndex)
 	{
@@ -611,14 +631,22 @@ FNiagaraDynamicDataBase* NiagaraRendererRibbons::GenerateVertexData(const FNiaga
 		FVector CurrPos = FirstPos;
 		FVector LastToCurrVec = FVector::ZeroVector;
 		float LastToCurrSize = 0;
+		float LastTwist = 0;
+		float LastWidth = 0;
 
 		// Find the first position with enough distance.
 		int32 CurrentIndex = 1;
 		while (CurrentIndex < RibbonIndices.Num())
 		{
-			CurrPos = PosData[RibbonIndices[CurrentIndex]];
+			const int32 CurrentDataIndex = RibbonIndices[CurrentIndex];
+			CurrPos = PosData[CurrentDataIndex];
 			LastToCurrVec = CurrPos - FirstPos;
 			LastToCurrSize = LastToCurrVec.Size();
+			if (bHasTwist)
+			{
+				LastTwist = TwistData[CurrentDataIndex];
+				LastWidth = SizeData[CurrentDataIndex];
+			}
 
 			// Find the first segment, or unique segment
 			if (LastToCurrSize > GNiagaraRibbonMinSegmentLength)
@@ -644,15 +672,24 @@ FNiagaraDynamicDataBase* NiagaraRendererRibbons::GenerateVertexData(const FNiaga
 		int32 NextIndex = CurrentIndex + 1;
 		while (NextIndex < RibbonIndices.Num())
 		{
-			const FVector NextPos = PosData[RibbonIndices[NextIndex]];
+			const int32 NextDataIndex = RibbonIndices[NextIndex];
+			const FVector NextPos = PosData[NextDataIndex];
 			FVector CurrToNextVec = NextPos - CurrPos;
 			const float CurrToNextSize = CurrToNextVec.Size();
+			
+			float NextTwist = 0;
+			float NextWidth = 0;
+			if (bHasTwist)
+			{
+				NextTwist = TwistData[NextDataIndex];
+				NextWidth = SizeData[NextDataIndex];
+			}
 
 			// It the next is far enough, or the last element
 			if (CurrToNextSize > GNiagaraRibbonMinSegmentLength || NextIndex == RibbonIndices.Num() - 1)
 			{
 				// Normalize CurrToNextVec
-				CurrToNextVec *= 1.f / CurrToNextSize;
+				CurrToNextVec *= 1.f / FMath::Max(GNiagaraRibbonMinSegmentLength, CurrToNextSize);
 				const FVector Tangent = (LastToCurrVec + CurrToNextVec).GetSafeNormal();
 
 				// Update the distance for CurrentIndex.
@@ -661,11 +698,13 @@ FNiagaraDynamicDataBase* NiagaraRendererRibbons::GenerateVertexData(const FNiaga
 				// Add the current point, which tangent is computed from neighbors
 				DynamicData->SortedIndices.Add(RibbonIndices[CurrentIndex]);
 				DynamicData->TangentAndDistances.Add(FVector4(Tangent.X, Tangent.Y, Tangent.Z, TotalDistance));
-					DynamicData->MultiRibbonIndices.Add(RibbonIndex);
+				DynamicData->MultiRibbonIndices.Add(RibbonIndex);
 
 				// Assumed equal to dot(Tangent, CurrToNextVec)
 				TotalSegmentLength += CurrToNextSize;
 				TotalSegmentAngle += AcosFast(FVector::DotProduct(LastToCurrVec, CurrToNextVec));
+				TotalTwistAngle += FMath::Abs(NextTwist - LastTwist);
+				TotalWidth += LastWidth;
 				++NumTotalSamples;
 
 				// Move to next segment.
@@ -673,6 +712,8 @@ FNiagaraDynamicDataBase* NiagaraRendererRibbons::GenerateVertexData(const FNiaga
 				CurrPos = NextPos;
 				LastToCurrVec = CurrToNextVec;
 				LastToCurrSize = CurrToNextSize;
+				LastTwist = NextTwist;
+				LastWidth = NextWidth;
 			}
 
 			// Try next if there is one.
@@ -797,15 +838,26 @@ FNiagaraDynamicDataBase* NiagaraRendererRibbons::GenerateVertexData(const FNiaga
 		const float AverageSegmentCurvature = TotalSegmentLength * OneOverSampleCount * .5 / (FMath::Max(SMALL_NUMBER, FMath::Abs(FMath::Sin(AverageSegmentAngle))));
 
 		// Accumulate the max so that tessellation is ever increasing and stabilizes over time.
-		// This is to avoid the spawning and unspawning of ribbons to constantly generate different tesselation factors.
-		TessellationAngle = FMath::Max(TessellationAngle, AverageSegmentAngle);
-		TessellationCurvature = FMath::Max(TessellationCurvature, AverageSegmentCurvature);
+		// This is to avoid the spawning and unspawning of ribbons to constantly generate different tessellation factors.
+		
+		TessellationAngle = FMath::Lerp<float>(AverageSegmentAngle, FMath::Max(TessellationAngle, AverageSegmentAngle), GNiagaraRibbonMaxAveraging);
+		TessellationCurvature = FMath::Lerp<float>(AverageSegmentCurvature, FMath::Max(TessellationCurvature, AverageSegmentCurvature), GNiagaraRibbonMaxAveraging);
 
+		if (bHasTwist)
+		{
+			const float AverageTwistAngle = TotalTwistAngle * OneOverSampleCount;
+			const float AverageTwistCurvature = TotalWidth * OneOverSampleCount;
+			
+			TessellationTwistAngle = FMath::Lerp<float>(AverageTwistAngle, FMath::Max(TessellationTwistAngle, AverageTwistAngle), GNiagaraRibbonMaxAveraging);
+			TessellationTwistCurvature = FMath::Lerp<float>(AverageTwistCurvature, FMath::Max(TessellationTwistCurvature, AverageTwistCurvature), GNiagaraRibbonMaxAveraging);
+		}
 	}
 	else // Reset the metrics when the ribbons are reset.
 	{
 		TessellationAngle = 0;
 		TessellationCurvature = 0;
+		TessellationTwistAngle = 0;
+		TessellationTwistCurvature = 0;
 	}
 	 
 	if (Data.CurrData().GetNumInstances() > 0)
diff --git a/Engine/Plugins/FX/Niagara/Source/Niagara/Public/NiagaraRendererRibbons.h b/Engine/Plugins/FX/Niagara/Source/Niagara/Public/NiagaraRendererRibbons.h
--- a/Engine/Plugins/FX/Niagara/Source/Niagara/Public/NiagaraRendererRibbons.h
+++ b/Engine/Plugins/FX/Niagara/Source/Niagara/Public/NiagaraRendererRibbons.h
@@ -90,6 +90,11 @@ private:
 	float TessellationAngle = 0;
 	// Average angle of the curvature of the segments (in radian).
 	float TessellationCurvature = 0;
+
+	// Average twist of the segments.
+	float TessellationTwistAngle = 0;
+	// Average twist curvature of the segments.
+	float TessellationTwistCurvature = 0;
 };
 
 
