- Integrate CL 6453585 from //UE4/Dev-Rendering
- Fixed deletion of NiagaraSystemInstance: bNotifyOnCompletion was modified after the class instance has already been deleted.

#rnx

diff --git a/Engine/Plugins/FX/Niagara/Source/Niagara/Private/NiagaraSystemInstance.cpp b/Engine/Plugins/FX/Niagara/Source/Niagara/Private/NiagaraSystemInstance.cpp
--- a/Engine/Plugins/FX/Niagara/Source/Niagara/Private/NiagaraSystemInstance.cpp
+++ b/Engine/Plugins/FX/Niagara/Source/Niagara/Private/NiagaraSystemInstance.cpp
@@ -366,15 +366,16 @@ void FNiagaraSystemInstance::Complete()
 
 	if (bNeedToNotifyOthers)
 	{
+		// We've already notified once, no need to do so again.
+		bNotifyOnCompletion = false;
+
 		OnCompleteDelegate.Broadcast(this);
 
 		if (Component)
 		{
+			// Note: This call may destroy this instance of FNiagaraSystemInstance, so don't use bNotifyOnCompletion after it!
 			Component->OnSystemComplete();
 		}
-		
-		// We've already notified once, no need to do so again.
-		bNotifyOnCompletion = false;
 	}
 }
 
