Copying //UE4/Dev-Niagara to Dev-Main (//UE4/Dev-Main)
#tests enginetests pass

diff --git a/Engine/Plugins/FX/HoudiniNiagara/Source/HoudiniNiagara/Classes/NiagaraDataInterfaceHoudiniCSV.h b/Engine/Plugins/FX/HoudiniNiagara/Source/HoudiniNiagara/Classes/NiagaraDataInterfaceHoudiniCSV.h
--- a/Engine/Plugins/FX/HoudiniNiagara/Source/HoudiniNiagara/Classes/NiagaraDataInterfaceHoudiniCSV.h
+++ b/Engine/Plugins/FX/HoudiniNiagara/Source/HoudiniNiagara/Classes/NiagaraDataInterfaceHoudiniCSV.h
@@ -234,38 +234,18 @@ public:
 	FORCEINLINE int32 GetMaxNumberOfIndexesPerPoints()const { return HoudiniCSVAsset ? HoudiniCSVAsset->GetMaxNumberOfPointValueIndexes() + 1 : 0; }
 
 	// GPU Buffers accessors
-	FRWBuffer& GetFloatValuesGPUBuffer();
+	/*FRWBuffer& GetFloatValuesGPUBuffer();
 	FRWBuffer& GetSpecialAttributesColumnIndexesGPUBuffer();
 	FRWBuffer& GetSpawnTimesGPUBuffer();
 	FRWBuffer& GetLifeValuesGPUBuffer();
 	FRWBuffer& GetPointTypesGPUBuffer();
-	FRWBuffer& GetPointValueIndexesGPUBuffer();
-
-	// GPU Buffers
-	FRWBuffer FloatValuesGPUBuffer;
-	FRWBuffer SpecialAttributesColumnIndexesGPUBuffer;
-	FRWBuffer SpawnTimesGPUBuffer;
-	FRWBuffer LifeValuesGPUBuffer;
-	FRWBuffer PointTypesGPUBuffer;
-	FRWBuffer PointValueIndexesGPUBuffer;
+	FRWBuffer& GetPointValueIndexesGPUBuffer();*/
 
 protected:
 
-	virtual bool CopyToInternal(UNiagaraDataInterface* Destination) const override;
+	void PushToRenderThread();
 
-	// Indicates the GPU buffers need to be updated
-	UPROPERTY()
-	bool FloatValuesGPUBufferDirty;
-	UPROPERTY()
-	bool SpecialAttributesColumnIndexesGPUBufferDirty;
-	UPROPERTY()
-	bool SpawnTimesGPUBufferDirty;
-	UPROPERTY()
-	bool LifeValuesGPUBufferDirty;
-	UPROPERTY()
-	bool PointTypesGPUBufferDirty;
-	UPROPERTY()
-	bool PointValueIndexesGPUBufferDirty;
+	virtual bool CopyToInternal(UNiagaraDataInterface* Destination) const override;
 
 	// Last Spawned PointID
 	UPROPERTY()
@@ -274,3 +254,26 @@ protected:
 	UPROPERTY()
 	float LastSpawnTime;
 };
+
+struct FNiagaraDataInterfaceProxyHoudiniCSV : public FNiagaraDataInterfaceProxy
+{
+	// GPU Buffers
+	FRWBuffer FloatValuesGPUBuffer;
+	FRWBuffer SpecialAttributesColumnIndexesGPUBuffer;
+	FRWBuffer SpawnTimesGPUBuffer;
+	FRWBuffer LifeValuesGPUBuffer;
+	FRWBuffer PointTypesGPUBuffer;
+	FRWBuffer PointValueIndexesGPUBuffer;
+
+	int32 MaxNumberOfIndexesPerPoint;
+	int32 NumRows;
+	int32 NumColumns;
+	int32 NumPoints;
+
+	void AcceptStaticDataUpdate(struct FNiagaraDIHoudiniCSV_StaticDataPassToRT& Update);
+
+	virtual int32 PerInstanceDataPassedToRenderThreadSize() const override
+	{
+		return 0;
+	}
+};
\ No newline at end of file
diff --git a/Engine/Plugins/FX/HoudiniNiagara/Source/HoudiniNiagara/Private/NiagaraDataInterfaceHoudiniCSV.cpp b/Engine/Plugins/FX/HoudiniNiagara/Source/HoudiniNiagara/Private/NiagaraDataInterfaceHoudiniCSV.cpp
--- a/Engine/Plugins/FX/HoudiniNiagara/Source/HoudiniNiagara/Private/NiagaraDataInterfaceHoudiniCSV.cpp
+++ b/Engine/Plugins/FX/HoudiniNiagara/Source/HoudiniNiagara/Private/NiagaraDataInterfaceHoudiniCSV.cpp
@@ -75,12 +75,28 @@ static const FName GetPointLifeName("GetPointLife");
 //static const FName GetPointLifeAtTimeName("GetPointLifeAtTime");
 static const FName GetPointTypeName("GetPointType");
 
+struct FNiagaraDIHoudiniCSV_StaticDataPassToRT
+{
+	TArray<float>* FloatCSVData;
+	TArray<float>* SpawnTimes;
+	TArray<float>* LifeValues;
+	TArray<int32>* PointTypes;
+	TArray<int32>* SpecialAttributesColumnIndexes;
+	TArray<int32>* PointValueIndexes;
+
+	int32 NumRows;
+	int32 NumColumns;
+	int32 NumPoints;
+	int32 MaxNumIndexesPerPoint;
+};
 
 UNiagaraDataInterfaceHoudiniCSV::UNiagaraDataInterfaceHoudiniCSV(FObjectInitializer const& ObjectInitializer)
 	: Super(ObjectInitializer)
 {
     HoudiniCSVAsset = nullptr;
 	LastSpawnedPointID = -1;
+
+	Proxy = MakeShared<FNiagaraDataInterfaceProxyHoudiniCSV, ESPMode::ThreadSafe>();
 }
 
 void UNiagaraDataInterfaceHoudiniCSV::PostInitProperties()
@@ -93,30 +109,20 @@ void UNiagaraDataInterfaceHoudiniCSV::PostInitProperties()
 		FNiagaraTypeRegistry::Register(FHoudiniEvent::StaticStruct(), true, true, false);
     }
 
-	FloatValuesGPUBufferDirty = true;
-	SpecialAttributesColumnIndexesGPUBufferDirty = true;
-	SpawnTimesGPUBufferDirty = true;
-	LifeValuesGPUBufferDirty = true;
-	PointTypesGPUBufferDirty = true;
-	PointValueIndexesGPUBufferDirty = true;
-
 	LastSpawnedPointID = -1;
 	LastSpawnTime = -1.0f;
+
+	PushToRenderThread();
 }
 
 void UNiagaraDataInterfaceHoudiniCSV::PostLoad()
 {
     Super::PostLoad();
 
-	FloatValuesGPUBufferDirty = true;
-	SpecialAttributesColumnIndexesGPUBufferDirty = true;
-	SpawnTimesGPUBufferDirty = true;
-	LifeValuesGPUBufferDirty = true;
-	PointTypesGPUBufferDirty = true;
-	PointValueIndexesGPUBufferDirty = true;
-
 	LastSpawnedPointID = -1;
 	LastSpawnTime = -1.0f;
+
+	PushToRenderThread();
 }
 
 #if WITH_EDITOR
@@ -130,15 +136,9 @@ void UNiagaraDataInterfaceHoudiniCSV::PostEditChangeProperty(struct FPropertyCha
 		Modify();
 		if ( HoudiniCSVAsset )
 		{
-			FloatValuesGPUBufferDirty = true;
-			SpecialAttributesColumnIndexesGPUBufferDirty = true;
-			SpawnTimesGPUBufferDirty = true;
-			LifeValuesGPUBufferDirty = true;
-			PointTypesGPUBufferDirty = true;
-			PointValueIndexesGPUBufferDirty = true;
-
 			LastSpawnedPointID = -1;
 			LastSpawnTime = -1.0f;
+			PushToRenderThread();
 		}
     }
 }
@@ -156,14 +156,9 @@ bool UNiagaraDataInterfaceHoudiniCSV::CopyToInternal(UNiagaraDataInterface* Dest
 
     CastedInterface->HoudiniCSVAsset = HoudiniCSVAsset;
 
-	CastedInterface->FloatValuesGPUBufferDirty = true;
-	CastedInterface->SpecialAttributesColumnIndexesGPUBufferDirty = true;
-	CastedInterface->SpawnTimesGPUBufferDirty = true;
-	CastedInterface->LifeValuesGPUBufferDirty = true;
-	CastedInterface->PointTypesGPUBufferDirty = true;
-	CastedInterface->PointValueIndexesGPUBufferDirty = true;
 	CastedInterface->LastSpawnedPointID = -1;
 	CastedInterface->LastSpawnTime = -1.0f;
+	CastedInterface->PushToRenderThread();
 
     return true;
 }
@@ -1852,39 +1847,279 @@ void UNiagaraDataInterfaceHoudiniCSV::GetParameterDefinitionHLSL(FNiagaraDataInt
 	OutHLSL += TEXT("int ") + BufferName + TEXT(";\n");
 }
 
-FRWBuffer& UNiagaraDataInterfaceHoudiniCSV::GetFloatValuesGPUBuffer()
+//FRWBuffer& UNiagaraDataInterfaceHoudiniCSV::GetFloatValuesGPUBuffer()
+//{
+//	//TODO: This isn't really very thread safe. Need to move to a proxy like system where DIs can push data to the RT safely.
+//	if ( FloatValuesGPUBufferDirty )
+//	{
+//		uint32 NumElements = HoudiniCSVAsset ? HoudiniCSVAsset->FloatCSVData.Num() : 0;
+//		if (NumElements <= 0)
+//			return FloatValuesGPUBuffer;
+//
+//		FloatValuesGPUBuffer.Release();
+//		FloatValuesGPUBuffer.Initialize( sizeof(float), NumElements, EPixelFormat::PF_R32_FLOAT, BUF_Static);
+//
+//		uint32 BufferSize = NumElements * sizeof( float );
+//		float* BufferData = static_cast<float*>( RHILockVertexBuffer( FloatValuesGPUBuffer.Buffer, 0, BufferSize, EResourceLockMode::RLM_WriteOnly ) );
+//
+//		if ( HoudiniCSVAsset )
+//			FPlatformMemory::Memcpy( BufferData, HoudiniCSVAsset->FloatCSVData.GetData(), BufferSize );
+//
+//		RHIUnlockVertexBuffer( FloatValuesGPUBuffer.Buffer );
+//		FloatValuesGPUBufferDirty = false;
+//	}
+//
+//	return FloatValuesGPUBuffer;
+//}
+
+//FRWBuffer& UNiagaraDataInterfaceHoudiniCSV::GetSpecialAttributesColumnIndexesGPUBuffer()
+//{
+//	//TODO: This isn't really very thread safe. Need to move to a proxy like system where DIs can push data to the RT safely.
+//	if ( SpecialAttributesColumnIndexesGPUBufferDirty )
+//	{
+//		uint32 NumElements = HoudiniCSVAsset ? HoudiniCSVAsset->SpecialAttributesColumnIndexes.Num() : 0;
+//		if (NumElements <= 0)
+//			return SpecialAttributesColumnIndexesGPUBuffer;
+//
+//		SpecialAttributesColumnIndexesGPUBuffer.Release();
+//		SpecialAttributesColumnIndexesGPUBuffer.Initialize(sizeof(int32), NumElements, EPixelFormat::PF_R32_SINT, BUF_Static);
+//
+//		uint32 BufferSize = NumElements * sizeof(int32);
+//		float* BufferData = static_cast<float*>(RHILockVertexBuffer(SpecialAttributesColumnIndexesGPUBuffer.Buffer, 0, BufferSize, EResourceLockMode::RLM_WriteOnly));
+//
+//		if (HoudiniCSVAsset)
+//			FPlatformMemory::Memcpy(BufferData, HoudiniCSVAsset->SpecialAttributesColumnIndexes.GetData(), BufferSize);
+//
+//		RHIUnlockVertexBuffer(SpecialAttributesColumnIndexesGPUBuffer.Buffer);
+//		SpecialAttributesColumnIndexesGPUBufferDirty = false;
+//	}
+//
+//	return SpecialAttributesColumnIndexesGPUBuffer;
+//}
+
+//FRWBuffer& UNiagaraDataInterfaceHoudiniCSV::GetSpawnTimesGPUBuffer()
+//{
+//	//TODO: This isn't really very thread safe. Need to move to a proxy like system where DIs can push data to the RT safely.
+//	if ( SpawnTimesGPUBufferDirty )
+//	{
+//		uint32 NumElements = HoudiniCSVAsset ? HoudiniCSVAsset->SpawnTimes.Num() : 0;
+//		if (NumElements <= 0)
+//			return SpawnTimesGPUBuffer;
+//
+//		SpawnTimesGPUBuffer.Release();
+//		SpawnTimesGPUBuffer.Initialize( sizeof(float), NumElements, EPixelFormat::PF_R32_FLOAT, BUF_Static);
+//				
+//		uint32 BufferSize = NumElements * sizeof( float );
+//		float* BufferData = static_cast<float*>( RHILockVertexBuffer(SpawnTimesGPUBuffer.Buffer, 0, BufferSize, EResourceLockMode::RLM_WriteOnly ) );
+//
+//		if ( HoudiniCSVAsset )
+//			FPlatformMemory::Memcpy( BufferData, HoudiniCSVAsset->SpawnTimes.GetData(), BufferSize );
+//
+//		RHIUnlockVertexBuffer( SpawnTimesGPUBuffer.Buffer );
+//		SpawnTimesGPUBufferDirty = false;
+//	}
+//
+//	return SpawnTimesGPUBuffer;
+//}
+
+//FRWBuffer& UNiagaraDataInterfaceHoudiniCSV::GetLifeValuesGPUBuffer()
+//{
+//	//TODO: This isn't really very thread safe. Need to move to a proxy like system where DIs can push data to the RT safely.
+//	if ( LifeValuesGPUBufferDirty )
+//	{
+//		uint32 NumElements = HoudiniCSVAsset ? HoudiniCSVAsset->LifeValues.Num() : 0;
+//		if (NumElements <= 0)
+//			return LifeValuesGPUBuffer;
+//
+//		LifeValuesGPUBuffer.Release();
+//		LifeValuesGPUBuffer.Initialize( sizeof(float), NumElements, EPixelFormat::PF_R32_FLOAT, BUF_Static);
+//
+//		uint32 BufferSize = NumElements * sizeof( float );
+//		float* BufferData = static_cast<float*>( RHILockVertexBuffer( LifeValuesGPUBuffer.Buffer, 0, BufferSize, EResourceLockMode::RLM_WriteOnly ) );
+//
+//		if ( HoudiniCSVAsset )
+//			FPlatformMemory::Memcpy( BufferData, HoudiniCSVAsset->LifeValues.GetData(), BufferSize );
+//
+//		RHIUnlockVertexBuffer(LifeValuesGPUBuffer.Buffer );
+//		LifeValuesGPUBufferDirty = false;
+//	}
+//
+//	return LifeValuesGPUBuffer;
+//}
+
+//FRWBuffer& UNiagaraDataInterfaceHoudiniCSV::GetPointTypesGPUBuffer()
+//{
+//	//TODO: This isn't really very thread safe. Need to move to a proxy like system where DIs can push data to the RT safely.
+//	if ( PointTypesGPUBufferDirty )
+//	{
+//		uint32 NumElements = HoudiniCSVAsset ? HoudiniCSVAsset->PointTypes.Num() : 0;
+//		if (NumElements <= 0)
+//			return PointTypesGPUBuffer;
+//
+//		PointTypesGPUBuffer.Release();
+//		PointTypesGPUBuffer.Initialize( sizeof(int32), NumElements, EPixelFormat::PF_R32_SINT, BUF_Static);
+//
+//		uint32 BufferSize = NumElements * sizeof(int32);
+//		int32* BufferData = static_cast<int32*>( RHILockVertexBuffer(PointTypesGPUBuffer.Buffer, 0, BufferSize, EResourceLockMode::RLM_WriteOnly ) );
+//
+//		if ( HoudiniCSVAsset )
+//			FPlatformMemory::Memcpy( BufferData, HoudiniCSVAsset->PointTypes.GetData(), BufferSize );
+//
+//		RHIUnlockVertexBuffer(PointTypesGPUBuffer.Buffer );
+//		PointTypesGPUBufferDirty = false;
+//	}
+//
+//	return PointTypesGPUBuffer;
+//}
+
+//FRWBuffer& UNiagaraDataInterfaceHoudiniCSV::GetPointValueIndexesGPUBuffer()
+//{
+//	//TODO: This isn't really very thread safe. Need to move to a proxy like system where DIs can push data to the RT safely.
+//	if ( PointValueIndexesGPUBufferDirty )
+//	{
+//		uint32 NumPoints = HoudiniCSVAsset ? HoudiniCSVAsset->PointValueIndexes.Num() : 0;
+//		// Add an extra to the max number so all indexes end by a -1
+//		uint32 MaxNumIndexesPerPoint = HoudiniCSVAsset ? HoudiniCSVAsset->GetMaxNumberOfPointValueIndexes() + 1: 0;
+//		uint32 NumElements = NumPoints * MaxNumIndexesPerPoint;
+//
+//		if (NumElements <= 0)
+//			return PointValueIndexesGPUBuffer;
+//
+//		TArray<int32> PointValueIndexes;
+//		PointValueIndexes.Init(-1, NumElements);
+//		if ( HoudiniCSVAsset )
+//		{
+//			// We need to flatten the nested array for HLSL conversion
+//			for ( int32 PointID = 0; PointID < HoudiniCSVAsset->PointValueIndexes.Num(); PointID++ )
+//			{
+//				TArray<int32> RowIndexes = HoudiniCSVAsset->PointValueIndexes[PointID].RowIndexes;
+//				for( int32 Idx = 0; Idx < RowIndexes.Num(); Idx++ )
+//				{
+//					PointValueIndexes[ PointID * MaxNumIndexesPerPoint + Idx ] = RowIndexes[ Idx ];
+//				}
+//			}
+//		}
+//
+//		PointValueIndexesGPUBuffer.Release();
+//		PointValueIndexesGPUBuffer.Initialize(sizeof(int32), NumElements, EPixelFormat::PF_R32_SINT, BUF_Static);
+//
+//		uint32 BufferSize = NumElements * sizeof(int32);
+//		int32* BufferData = static_cast<int32*>( RHILockVertexBuffer( PointValueIndexesGPUBuffer.Buffer, 0, BufferSize, EResourceLockMode::RLM_WriteOnly ) );
+//		FPlatformMemory::Memcpy( BufferData, PointValueIndexes.GetData(), BufferSize );
+//		RHIUnlockVertexBuffer( PointValueIndexesGPUBuffer.Buffer );
+//
+//		PointValueIndexesGPUBufferDirty = false;
+//	}
+//
+//	return PointValueIndexesGPUBuffer;
+//}
+
+void UNiagaraDataInterfaceHoudiniCSV::PushToRenderThread()
 {
-	//TODO: This isn't really very thread safe. Need to move to a proxy like system where DIs can push data to the RT safely.
-	if ( FloatValuesGPUBufferDirty )
+	check(Proxy);
+	FNiagaraDIHoudiniCSV_StaticDataPassToRT DataToPass;
+	FMemory::Memzero(DataToPass);
+
 	{
 		uint32 NumElements = HoudiniCSVAsset ? HoudiniCSVAsset->FloatCSVData.Num() : 0;
-		if (NumElements <= 0)
-			return FloatValuesGPUBuffer;
+		if (NumElements > 0)
+		{
+			DataToPass.FloatCSVData = new TArray<float>(HoudiniCSVAsset->FloatCSVData);
+		}
+	}
 
-		FloatValuesGPUBuffer.Release();
-		FloatValuesGPUBuffer.Initialize( sizeof(float), NumElements, EPixelFormat::PF_R32_FLOAT, BUF_Static);
+	{
+		uint32 NumElements = HoudiniCSVAsset ? HoudiniCSVAsset->SpecialAttributesColumnIndexes.Num() : 0;
+		if (NumElements > 0)
+		{
+			DataToPass.SpecialAttributesColumnIndexes = new TArray<int32>(HoudiniCSVAsset->SpecialAttributesColumnIndexes);
+		}
+	}
 
-		uint32 BufferSize = NumElements * sizeof( float );
-		float* BufferData = static_cast<float*>( RHILockVertexBuffer( FloatValuesGPUBuffer.Buffer, 0, BufferSize, EResourceLockMode::RLM_WriteOnly ) );
+	{
+		uint32 NumElements = HoudiniCSVAsset ? HoudiniCSVAsset->SpawnTimes.Num() : 0;
+		if (NumElements > 0)
+		{
+			DataToPass.SpawnTimes = new TArray<float>(HoudiniCSVAsset->SpawnTimes);
+		}
+	}
 
-		if ( HoudiniCSVAsset )
-			FPlatformMemory::Memcpy( BufferData, HoudiniCSVAsset->FloatCSVData.GetData(), BufferSize );
+	{
+		uint32 NumElements = HoudiniCSVAsset ? HoudiniCSVAsset->LifeValues.Num() : 0;
+		if (NumElements > 0)
+		{
+			DataToPass.LifeValues = new TArray<float>(HoudiniCSVAsset->LifeValues);
+		}
+	}
+
+	{
+		uint32 NumElements = HoudiniCSVAsset ? HoudiniCSVAsset->PointTypes.Num() : 0;
+		if (NumElements > 0)
+		{
+			DataToPass.PointTypes = new TArray<int32>(HoudiniCSVAsset->PointTypes);
+		}
+	}
+
+	{
+		uint32 NumPoints = HoudiniCSVAsset ? HoudiniCSVAsset->PointValueIndexes.Num() : 0;
+		// Add an extra to the max number so all indexes end by a -1
+		uint32 MaxNumIndexesPerPoint = HoudiniCSVAsset ? HoudiniCSVAsset->GetMaxNumberOfPointValueIndexes() + 1 : 0;
+		uint32 NumElements = NumPoints * MaxNumIndexesPerPoint;
 
-		RHIUnlockVertexBuffer( FloatValuesGPUBuffer.Buffer );
-		FloatValuesGPUBufferDirty = false;
+		if (NumElements > 0)
+		{
+			DataToPass.PointValueIndexes = new TArray<int32>;
+			DataToPass.PointValueIndexes->Init(-1, NumElements);
+			if (HoudiniCSVAsset)
+			{
+				// We need to flatten the nested array for HLSL conversion
+				for (int32 PointID = 0; PointID < HoudiniCSVAsset->PointValueIndexes.Num(); PointID++)
+				{
+					TArray<int32> RowIndexes = HoudiniCSVAsset->PointValueIndexes[PointID].RowIndexes;
+					for (int32 Idx = 0; Idx < RowIndexes.Num(); Idx++)
+					{
+						(*DataToPass.PointValueIndexes)[PointID * MaxNumIndexesPerPoint + Idx] = RowIndexes[Idx];
+					}
+				}
+			}
+		}
 	}
 
-	return FloatValuesGPUBuffer;
+	DataToPass.NumRows = GetNumberOfRows();
+	DataToPass.NumColumns = GetNumberOfColumns();
+	DataToPass.NumPoints = GetNumberOfPoints();
+	DataToPass.MaxNumIndexesPerPoint = GetMaxNumberOfIndexesPerPoints();
+
+	FNiagaraDataInterfaceProxyHoudiniCSV* ThisProxy = GetProxyAs<FNiagaraDataInterfaceProxyHoudiniCSV>();
+	ENQUEUE_RENDER_COMMAND(FNiagaraDIHoudiniCSV_ToRT) (
+		[ThisProxy, DataToPass](FRHICommandListImmediate& CmdList) mutable
+	{
+		ThisProxy->AcceptStaticDataUpdate(DataToPass);
+	}
+	);
 }
 
-FRWBuffer& UNiagaraDataInterfaceHoudiniCSV::GetSpecialAttributesColumnIndexesGPUBuffer()
+void FNiagaraDataInterfaceProxyHoudiniCSV::AcceptStaticDataUpdate(FNiagaraDIHoudiniCSV_StaticDataPassToRT& Update)
 {
-	//TODO: This isn't really very thread safe. Need to move to a proxy like system where DIs can push data to the RT safely.
-	if ( SpecialAttributesColumnIndexesGPUBufferDirty )
+	if (Update.FloatCSVData)
 	{
-		uint32 NumElements = HoudiniCSVAsset ? HoudiniCSVAsset->SpecialAttributesColumnIndexes.Num() : 0;
-		if (NumElements <= 0)
-			return SpecialAttributesColumnIndexesGPUBuffer;
+		int32 NumElements = Update.FloatCSVData->Num();
+		FloatValuesGPUBuffer.Release();
+		FloatValuesGPUBuffer.Initialize(sizeof(float), NumElements, EPixelFormat::PF_R32_FLOAT, BUF_Static);
+
+		uint32 BufferSize = NumElements * sizeof(float);
+		float* BufferData = static_cast<float*>(RHILockVertexBuffer(FloatValuesGPUBuffer.Buffer, 0, BufferSize, EResourceLockMode::RLM_WriteOnly));
+
+		FPlatformMemory::Memcpy(BufferData, Update.FloatCSVData->GetData(), BufferSize);
+
+		RHIUnlockVertexBuffer(FloatValuesGPUBuffer.Buffer);
+
+		delete Update.FloatCSVData;
+	}
+
+	if (Update.SpecialAttributesColumnIndexes)
+	{
+		uint32 NumElements = Update.SpecialAttributesColumnIndexes->Num();
 
 		SpecialAttributesColumnIndexesGPUBuffer.Release();
 		SpecialAttributesColumnIndexesGPUBuffer.Initialize(sizeof(int32), NumElements, EPixelFormat::PF_R32_SINT, BUF_Static);
@@ -1892,131 +2127,85 @@ FRWBuffer& UNiagaraDataInterfaceHoudiniCSV::GetSpecialAttributesColumnIndexesGPU
 		uint32 BufferSize = NumElements * sizeof(int32);
 		float* BufferData = static_cast<float*>(RHILockVertexBuffer(SpecialAttributesColumnIndexesGPUBuffer.Buffer, 0, BufferSize, EResourceLockMode::RLM_WriteOnly));
 
-		if (HoudiniCSVAsset)
-			FPlatformMemory::Memcpy(BufferData, HoudiniCSVAsset->SpecialAttributesColumnIndexes.GetData(), BufferSize);
+		FPlatformMemory::Memcpy(BufferData, Update.SpecialAttributesColumnIndexes->GetData(), BufferSize);
 
 		RHIUnlockVertexBuffer(SpecialAttributesColumnIndexesGPUBuffer.Buffer);
-		SpecialAttributesColumnIndexesGPUBufferDirty = false;
-	}
 
-	return SpecialAttributesColumnIndexesGPUBuffer;
-}
+		delete Update.SpecialAttributesColumnIndexes;
+	}
 
-FRWBuffer& UNiagaraDataInterfaceHoudiniCSV::GetSpawnTimesGPUBuffer()
-{
-	//TODO: This isn't really very thread safe. Need to move to a proxy like system where DIs can push data to the RT safely.
-	if ( SpawnTimesGPUBufferDirty )
+	if (Update.SpawnTimes)
 	{
-		uint32 NumElements = HoudiniCSVAsset ? HoudiniCSVAsset->SpawnTimes.Num() : 0;
-		if (NumElements <= 0)
-			return SpawnTimesGPUBuffer;
+		uint32 NumElements = Update.SpawnTimes->Num();
+		
 
 		SpawnTimesGPUBuffer.Release();
-		SpawnTimesGPUBuffer.Initialize( sizeof(float), NumElements, EPixelFormat::PF_R32_FLOAT, BUF_Static);
-				
-		uint32 BufferSize = NumElements * sizeof( float );
-		float* BufferData = static_cast<float*>( RHILockVertexBuffer(SpawnTimesGPUBuffer.Buffer, 0, BufferSize, EResourceLockMode::RLM_WriteOnly ) );
+		SpawnTimesGPUBuffer.Initialize(sizeof(float), NumElements, EPixelFormat::PF_R32_FLOAT, BUF_Static);
 
-		if ( HoudiniCSVAsset )
-			FPlatformMemory::Memcpy( BufferData, HoudiniCSVAsset->SpawnTimes.GetData(), BufferSize );
+		uint32 BufferSize = NumElements * sizeof(float);
+		float* BufferData = static_cast<float*>(RHILockVertexBuffer(SpawnTimesGPUBuffer.Buffer, 0, BufferSize, EResourceLockMode::RLM_WriteOnly));
 
-		RHIUnlockVertexBuffer( SpawnTimesGPUBuffer.Buffer );
-		SpawnTimesGPUBufferDirty = false;
-	}
+		FPlatformMemory::Memcpy(BufferData, Update.SpawnTimes->GetData(), BufferSize);
 
-	return SpawnTimesGPUBuffer;
-}
+		RHIUnlockVertexBuffer(SpawnTimesGPUBuffer.Buffer);
 
-FRWBuffer& UNiagaraDataInterfaceHoudiniCSV::GetLifeValuesGPUBuffer()
-{
-	//TODO: This isn't really very thread safe. Need to move to a proxy like system where DIs can push data to the RT safely.
-	if ( LifeValuesGPUBufferDirty )
+		delete Update.SpawnTimes;		
+	}
+
+	if (Update.LifeValues)
 	{
-		uint32 NumElements = HoudiniCSVAsset ? HoudiniCSVAsset->LifeValues.Num() : 0;
-		if (NumElements <= 0)
-			return LifeValuesGPUBuffer;
+		uint32 NumElements = Update.LifeValues->Num();
 
 		LifeValuesGPUBuffer.Release();
-		LifeValuesGPUBuffer.Initialize( sizeof(float), NumElements, EPixelFormat::PF_R32_FLOAT, BUF_Static);
+		LifeValuesGPUBuffer.Initialize(sizeof(float), NumElements, EPixelFormat::PF_R32_FLOAT, BUF_Static);
 
-		uint32 BufferSize = NumElements * sizeof( float );
-		float* BufferData = static_cast<float*>( RHILockVertexBuffer( LifeValuesGPUBuffer.Buffer, 0, BufferSize, EResourceLockMode::RLM_WriteOnly ) );
+		uint32 BufferSize = NumElements * sizeof(float);
+		float* BufferData = static_cast<float*>(RHILockVertexBuffer(LifeValuesGPUBuffer.Buffer, 0, BufferSize, EResourceLockMode::RLM_WriteOnly));
 
-		if ( HoudiniCSVAsset )
-			FPlatformMemory::Memcpy( BufferData, HoudiniCSVAsset->LifeValues.GetData(), BufferSize );
+		FPlatformMemory::Memcpy(BufferData, Update.LifeValues->GetData(), BufferSize);
 
-		RHIUnlockVertexBuffer(LifeValuesGPUBuffer.Buffer );
-		LifeValuesGPUBufferDirty = false;
+		RHIUnlockVertexBuffer(LifeValuesGPUBuffer.Buffer);
+		
+		delete Update.LifeValues;
 	}
 
-	return LifeValuesGPUBuffer;
-}
-
-FRWBuffer& UNiagaraDataInterfaceHoudiniCSV::GetPointTypesGPUBuffer()
-{
-	//TODO: This isn't really very thread safe. Need to move to a proxy like system where DIs can push data to the RT safely.
-	if ( PointTypesGPUBufferDirty )
+	if (Update.PointTypes)
 	{
-		uint32 NumElements = HoudiniCSVAsset ? HoudiniCSVAsset->PointTypes.Num() : 0;
-		if (NumElements <= 0)
-			return PointTypesGPUBuffer;
+		uint32 NumElements = Update.PointTypes->Num();
 
 		PointTypesGPUBuffer.Release();
-		PointTypesGPUBuffer.Initialize( sizeof(int32), NumElements, EPixelFormat::PF_R32_SINT, BUF_Static);
+		PointTypesGPUBuffer.Initialize(sizeof(int32), NumElements, EPixelFormat::PF_R32_SINT, BUF_Static);
 
 		uint32 BufferSize = NumElements * sizeof(int32);
-		int32* BufferData = static_cast<int32*>( RHILockVertexBuffer(PointTypesGPUBuffer.Buffer, 0, BufferSize, EResourceLockMode::RLM_WriteOnly ) );
+		int32* BufferData = static_cast<int32*>(RHILockVertexBuffer(PointTypesGPUBuffer.Buffer, 0, BufferSize, EResourceLockMode::RLM_WriteOnly));
 
-		if ( HoudiniCSVAsset )
-			FPlatformMemory::Memcpy( BufferData, HoudiniCSVAsset->PointTypes.GetData(), BufferSize );
+		FPlatformMemory::Memcpy(BufferData, Update.PointTypes->GetData(), BufferSize);
 
-		RHIUnlockVertexBuffer(PointTypesGPUBuffer.Buffer );
-		PointTypesGPUBufferDirty = false;
+		RHIUnlockVertexBuffer(PointTypesGPUBuffer.Buffer);
+		
+		delete Update.PointTypes;
 	}
 
-	return PointTypesGPUBuffer;
-}
-
-FRWBuffer& UNiagaraDataInterfaceHoudiniCSV::GetPointValueIndexesGPUBuffer()
-{
-	//TODO: This isn't really very thread safe. Need to move to a proxy like system where DIs can push data to the RT safely.
-	if ( PointValueIndexesGPUBufferDirty )
+	if (Update.PointValueIndexes)
 	{
-		uint32 NumPoints = HoudiniCSVAsset ? HoudiniCSVAsset->PointValueIndexes.Num() : 0;
-		// Add an extra to the max number so all indexes end by a -1
-		uint32 MaxNumIndexesPerPoint = HoudiniCSVAsset ? HoudiniCSVAsset->GetMaxNumberOfPointValueIndexes() + 1: 0;
-		uint32 NumElements = NumPoints * MaxNumIndexesPerPoint;
-
-		if (NumElements <= 0)
-			return PointValueIndexesGPUBuffer;
-
-		TArray<int32> PointValueIndexes;
-		PointValueIndexes.Init(-1, NumElements);
-		if ( HoudiniCSVAsset )
-		{
-			// We need to flatten the nested array for HLSL conversion
-			for ( int32 PointID = 0; PointID < HoudiniCSVAsset->PointValueIndexes.Num(); PointID++ )
-			{
-				TArray<int32> RowIndexes = HoudiniCSVAsset->PointValueIndexes[PointID].RowIndexes;
-				for( int32 Idx = 0; Idx < RowIndexes.Num(); Idx++ )
-				{
-					PointValueIndexes[ PointID * MaxNumIndexesPerPoint + Idx ] = RowIndexes[ Idx ];
-				}
-			}
-		}
+		
+		uint32 NumElements = Update.PointValueIndexes->Num();
 
 		PointValueIndexesGPUBuffer.Release();
 		PointValueIndexesGPUBuffer.Initialize(sizeof(int32), NumElements, EPixelFormat::PF_R32_SINT, BUF_Static);
 
 		uint32 BufferSize = NumElements * sizeof(int32);
-		int32* BufferData = static_cast<int32*>( RHILockVertexBuffer( PointValueIndexesGPUBuffer.Buffer, 0, BufferSize, EResourceLockMode::RLM_WriteOnly ) );
-		FPlatformMemory::Memcpy( BufferData, PointValueIndexes.GetData(), BufferSize );
-		RHIUnlockVertexBuffer( PointValueIndexesGPUBuffer.Buffer );
+		int32* BufferData = static_cast<int32*>(RHILockVertexBuffer(PointValueIndexesGPUBuffer.Buffer, 0, BufferSize, EResourceLockMode::RLM_WriteOnly));
+		FPlatformMemory::Memcpy(BufferData, Update.PointValueIndexes->GetData(), BufferSize);
+		RHIUnlockVertexBuffer(PointValueIndexesGPUBuffer.Buffer);
 
-		PointValueIndexesGPUBufferDirty = false;
+		delete Update.PointValueIndexes;
 	}
 
-	return PointValueIndexesGPUBuffer;
+	NumRows = Update.NumRows;
+	NumColumns = Update.NumColumns;
+	NumPoints = Update.NumPoints;
+	MaxNumberOfIndexesPerPoint = Update.MaxNumIndexesPerPoint;
 }
 
 // Parameters used for GPU sim compatibility
@@ -2065,35 +2254,35 @@ struct FNiagaraDataInterfaceParametersCS_HoudiniCSV : public FNiagaraDataInterfa
 		Ar << LastSpawnTime;
 	}
 
-	virtual void Set( FRHICommandList& RHICmdList, FNiagaraShader* Shader, class UNiagaraDataInterface* DataInterface, void* PerInstanceData) const override
+	virtual void Set(FRHICommandList& RHICmdList, const FNiagaraDataInterfaceSetArgs& Context) const override
 	{
 		check( IsInRenderingThread() );
 
-		const FComputeShaderRHIParamRef ComputeShaderRHI = Shader->GetComputeShader();
-		UNiagaraDataInterfaceHoudiniCSV* HoudiniDI = CastChecked<UNiagaraDataInterfaceHoudiniCSV>( DataInterface );
+		const FComputeShaderRHIParamRef ComputeShaderRHI = Context.Shader->GetComputeShader();
+		FNiagaraDataInterfaceProxyHoudiniCSV* HoudiniDI = static_cast<FNiagaraDataInterfaceProxyHoudiniCSV*>(Context.DataInterface);
 		if ( !HoudiniDI )
 			return;
 
-		SetShaderValue(RHICmdList, ComputeShaderRHI, NumberOfRows, HoudiniDI->GetNumberOfRows());
-		SetShaderValue(RHICmdList, ComputeShaderRHI, NumberOfColumns, HoudiniDI->GetNumberOfColumns());
-		SetShaderValue(RHICmdList, ComputeShaderRHI, NumberOfPoints, HoudiniDI->GetNumberOfPoints());
+		SetShaderValue(RHICmdList, ComputeShaderRHI, NumberOfRows, HoudiniDI->NumRows);
+		SetShaderValue(RHICmdList, ComputeShaderRHI, NumberOfColumns, HoudiniDI->NumColumns);
+		SetShaderValue(RHICmdList, ComputeShaderRHI, NumberOfPoints, HoudiniDI->NumColumns);
 
-		FRWBuffer& FloatRWBuffer = HoudiniDI->GetFloatValuesGPUBuffer();
+		FRWBuffer& FloatRWBuffer = HoudiniDI->FloatValuesGPUBuffer;
  		RHICmdList.SetShaderResourceViewParameter( ComputeShaderRHI, FloatValuesBuffer.GetBaseIndex(), FloatRWBuffer.SRV );
 
-		FRWBuffer& AttributeColumnRWBuffer = HoudiniDI->GetSpecialAttributesColumnIndexesGPUBuffer();
+		FRWBuffer& AttributeColumnRWBuffer = HoudiniDI->SpecialAttributesColumnIndexesGPUBuffer;
 		RHICmdList.SetShaderResourceViewParameter(ComputeShaderRHI, SpecialAttributesColumnIndexesBuffer.GetBaseIndex(), AttributeColumnRWBuffer.SRV);
 
-		FRWBuffer& SpawnRWBuffer = HoudiniDI->GetSpawnTimesGPUBuffer();
+		FRWBuffer& SpawnRWBuffer = HoudiniDI->SpawnTimesGPUBuffer;
 		RHICmdList.SetShaderResourceViewParameter(ComputeShaderRHI, SpawnTimesBuffer.GetBaseIndex(), SpawnRWBuffer.SRV);
-		FRWBuffer& LifeRWBuffer = HoudiniDI->GetLifeValuesGPUBuffer();
+		FRWBuffer& LifeRWBuffer = HoudiniDI->LifeValuesGPUBuffer;
 		RHICmdList.SetShaderResourceViewParameter(ComputeShaderRHI, LifeValuesBuffer.GetBaseIndex(), LifeRWBuffer.SRV);
-		FRWBuffer& TypesRWBuffer = HoudiniDI->GetPointTypesGPUBuffer();
+		FRWBuffer& TypesRWBuffer = HoudiniDI->PointTypesGPUBuffer;
 		RHICmdList.SetShaderResourceViewParameter(ComputeShaderRHI, PointTypesBuffer.GetBaseIndex(), TypesRWBuffer.SRV);
 
-		SetShaderValue(RHICmdList, ComputeShaderRHI, MaxNumberOfIndexesPerPoint, HoudiniDI->GetMaxNumberOfIndexesPerPoints());
+		SetShaderValue(RHICmdList, ComputeShaderRHI, MaxNumberOfIndexesPerPoint, HoudiniDI->MaxNumberOfIndexesPerPoint);
 
-		FRWBuffer& PointValuesIndexesRWBuffer = HoudiniDI->GetPointValueIndexesGPUBuffer();
+		FRWBuffer& PointValuesIndexesRWBuffer = HoudiniDI->PointValueIndexesGPUBuffer;
 		RHICmdList.SetShaderResourceViewParameter(ComputeShaderRHI, PointValueIndexesBuffer.GetBaseIndex(), PointValuesIndexesRWBuffer.SRV);
 
 		SetShaderValue(RHICmdList, ComputeShaderRHI, LastSpawnedPointId, -1);
diff --git a/Engine/Plugins/FX/Niagara/Shaders/Private/NiagaraEmitterInstanceShader.usf b/Engine/Plugins/FX/Niagara/Shaders/Private/NiagaraEmitterInstanceShader.usf
new file mode 100644
--- /dev/null
+++ b/Engine/Plugins/FX/Niagara/Shaders/Private/NiagaraEmitterInstanceShader.usf
@@ -0,0 +1,463 @@
+// Copyright 1998-2019 Epic Games, Inc. All Rights Reserved.
+
+/*=============================================================================
+NiagaraSimulationShader.usf: 
+=============================================================================*/
+#pragma warning(disable:4008)
+
+#if GPU_SIMULATION
+	#include "/Engine/Public/Platform.ush"
+	#include "/Engine/Private/Common.ush"
+	#include "/Engine/Private/GlobalDistanceFieldShared.ush"
+#else
+	const static float PI = 3.1415926535897932f;
+#endif
+	const static float TWO_PI = 3.1415926535897932f*2.0f;
+
+#include "/Engine/Private/Definitions.usf"
+
+// Most of the vector implementations work this way. This helps us keep proper precision.
+float4 ModuloPrecise(float4 x, float4 y){ return x - y * trunc(x/y); }
+float3 ModuloPrecise(float3 x, float3 y){ return x - y * trunc(x/y); }
+float2 ModuloPrecise(float2 x, float2 y){ return x - y * trunc(x/y);}
+float ModuloPrecise(float x, float y){ return x - y * trunc(x/y); }
+int ModuloPrecise(int x, int y){ return x - y * (x/y); }
+int Modulo(int x, int y){  return x - y * (x/y);  }
+
+	
+// using rcp is only 12 bits of precision, we should usually pay for perf
+float4 Reciprocal(float4 x){ return 1.0f/x;}
+float3 Reciprocal(float3 x){ return 1.0f/x; }
+float2 Reciprocal(float2 x){ return 1.0f/x;}
+float Reciprocal(float x){ return 1.0f/x; }
+
+/* -----------------------------------------------------------------
+ * GPU simulation utility functions
+ * -----------------------------------------------------------------
+ */
+#if GPU_SIMULATION
+
+	#include "/Engine/Private/SceneTexturesCommon.ush"
+
+	static uint3 GDispatchThreadId;
+	static uint GCurrentPhase;
+	uint EmitterTickCounter;
+
+	#define THREADGROUP_SIZE 64
+
+	float4 Modulo(float4 x, float4 y){ return fmod(x,y); }
+	float3 Modulo(float3 x, float3 y){ return fmod(x,y); }
+	float2 Modulo(float2 x, float2 y){ return fmod(x,y); }
+	float Modulo(float x, float y){ return fmod(x,y); }
+
+	// utility function used for scene depth calculations
+	float3 WorldPositionFromSceneDepth(float2 ScreenPosition, float SceneDepth)
+	{
+		float4 HomogeneousWorldPosition = mul(float4(ScreenPosition * SceneDepth, SceneDepth, 1), View.ScreenToWorld);
+		return HomogeneousWorldPosition.xyz / HomogeneousWorldPosition.w;
+	}
+#endif
+
+/* ----------------------------------------------------------------------------
+ * Seeded/Deterministic random number generation functions
+ *
+ * This is a variant of NiagaraRand4DPCG32 from Random.ush. 
+ *
+ * uint is not fully supported in the VM so we simply use ints and drop the 
+ * top and bottom bit swap. This should be fine since signed overflow should 
+ * produce the same results as unsigned overflow when comparing bit-by-bit on 
+ * all relevant architectures.
+ *
+ * Warning: Only contains 24 bits of randomness, since we produce values in 
+ *          the unit interval. 
+ *
+ * Note: Deterministic as long as Seed1, Seed2 and Seed3 are the same *and* as 
+ *       long as no Random or Seeded Random ops are inserted into or removed 
+ *       from the current script.
+ *
+ * TODO(mv): Might want to use a special case based on NiagaraRand3DPCG32 if 
+ *           only 3 or less components are needed, and we can drop one of the 
+ *           seeds, to reduce the operation count
+ * ----------------------------------------------------------------------------
+ */
+
+// static counter that is incremented every time rand3 or rand4 is called. 
+static int RandomCounter = 0;
+
+// Returns 4 random normalized floats based on 3 explicit integer seeds and one 
+// implicit integer seed. 
+float4 rand4(int Seed1, int Seed2, int Seed3)
+{
+	RandomCounter += 1;
+
+	int4 v = int4(RandomCounter, Seed1, Seed2, Seed3) * 1664525 + 1013904223;
+
+	v.x += v.y*v.w;
+	v.y += v.z*v.x;
+	v.z += v.x*v.y;
+	v.w += v.y*v.z;
+	v.x += v.y*v.w;
+	v.y += v.z*v.x;
+	v.z += v.x*v.y;
+	v.w += v.y*v.z;
+
+	// NOTE(mv): We can use 24 bits of randomness, as all integers in [0, 2^24] 
+    //           are exactly representable in single precision floats.
+    //           We use the upper 24 bits as they tend to be higher quality.
+
+    // NOTE(mv): The divide can often be folded with the range scale in the rand functions
+
+    return float4((v >> 8) & 0x00ffffff) / 16777216.0; // 0x01000000 == 16777216
+    // return float4((v >> 8) & 0x00ffffff) * (1.0/16777216.0); // bugged, see UE-67738
+}
+
+// float3 specialization of the above:
+// Returns 3 random normalized floats based on 3 explicit integer seeds and one 
+// implicit integer seed. The implicit seed takes the upper 16 bits of the third 
+// explicit seed, as this is often a global seed. 
+// The expression 
+// 
+//    RandomCounter | (Seed3 << 16)
+//
+// should be optimized away at compile time, as the third seed is typically 
+// a global constant.
+float3 rand3(int Seed1, int Seed2, int Seed3)
+{
+    RandomCounter += 1;
+
+    int3 v = int3(Seed1, Seed2, RandomCounter | (Seed3 << 16)) * 1664525 + 1013904223;
+
+    v.x += v.y*v.z;
+    v.y += v.z*v.x;
+    v.z += v.x*v.y;
+    v.x += v.y*v.z;
+    v.y += v.z*v.x;
+    v.z += v.x*v.y;
+
+    return float3((v >> 8) & 0x00ffffff) / 16777216.0; // 0x01000000 == 16777216
+}
+
+// Cost using rand4: 6 imad, 1 itof, 1 ishr, 1 add, 2 mul
+float rand(float x, int Seed1, int Seed2, int Seed3)
+{
+	return rand3(Seed1, Seed2, Seed3).x * x;
+}
+
+// Cost using rand4: 7 imad, 1 itof, 1 ishr, 1 add, 2 mul
+float2 rand(float2 x, int Seed1, int Seed2, int Seed3)
+{
+	return rand3(Seed1, Seed2, Seed3).xy * x;
+}
+
+// Cost using rand4: 8 imad, 1 itof, 1 ishr, 1 add, 2 mul
+float3 rand(float3 x, int Seed1, int Seed2, int Seed3)
+{
+	return rand3(Seed1, Seed2, Seed3).xyz * x;
+}
+
+// Cost using rand4: 9 imad, 1 itof, 1 ishr, 1 and, 2 mul
+float4 rand(float4 x, int Seed1, int Seed2, int Seed3) 
+{
+    return rand4(Seed1, Seed2, Seed3).xyzw * x;
+}
+
+// Cost using rand4: 6 imad, 2 itof, 1 ishr, 1 add, 2 mul, 1 ftoi
+int rand(int x, int Seed1, int Seed2, int Seed3)
+{
+	// NOTE: Scaling a uniform float range provides better distribution of 
+	//       numbers than using %.
+	// NOTE: Inclusive! So [0, x] instead of [0, x)
+	return int(rand3(Seed1, Seed2, Seed3).x * (x+1));
+}
+
+/* -----------------------------------------------------------------
+ * Un-seeded/Non-deterministic random number generation functions
+ * -----------------------------------------------------------------
+ */
+#if GPU_SIMULATION
+	// NOTE(mv): This simply calls the deterministic random number functions 
+    //           from the Seeded RNG section, but uses non-deterministic seeds 
+    //           as input. In particular, GDispatchThreadId for the first seed 
+	//           and EmitterTickCounter for the second seed. Also note that 
+	//           EmitterTickCounter does *not* reset for each frame, which 
+	//           secures non-determinism across simulations.
+
+	// TODO(mv): This could perhaps be optimized a little bit using slightly 
+	//           cheaper functions, but the difference is likely negligible. 
+	//           So for the sake of brevity and cleanliness, keep it like this.
+
+	float rand(float x)
+	{
+		return rand4(GDispatchThreadId.x, EmitterTickCounter, 0).x * x;
+	}
+
+	float2 rand(float2 x)
+	{
+		return rand4(GDispatchThreadId.x, EmitterTickCounter, 0).xy * x;
+	}
+
+	float3 rand(float3 x)
+	{
+		return rand4(GDispatchThreadId.x, EmitterTickCounter, 0).xyz * x;
+	}
+
+	float4 rand(float4 x) 
+	{
+		return rand4(GDispatchThreadId.x, EmitterTickCounter, 0).xyzw * x;
+	}
+
+	// Note: integer randoms are INCLUSIVE, i.e. includes both the upper and lower limit
+	int rand(int x)
+	{
+		return int(rand4(GDispatchThreadId.x, EmitterTickCounter, 0).x * (x+1));
+	}
+#else
+	// Old unseeded, passthrough to FRandomStream
+	float2 rand(float2 x)
+	{
+		return float2(rand(x.x), rand(x.y));
+	}
+
+	float3 rand(float3 x)
+	{
+		return float3(rand(x.x), rand(x.y), rand(x.z));
+	}
+
+	float4 rand(float4 x) 
+	{
+		return float4(rand(x.x), rand(x.y), rand(x.z), rand(x.w));
+	}
+	
+	// NOTE(mv): Where's rand(float)? Seems to work regardless..
+	int rand(int x);
+#endif
+
+
+/* -----------------------------------------------------------------
+ * VM simulation function declarations
+ * -----------------------------------------------------------------
+ */
+#if VM_SIMULATION
+	float noise(float x);
+	float noise(float2 x);
+	float noise(float3 x);
+
+	//Total hack to get around the cross compiler converting fmod() to "X - (Y * trunc(X/Y))";
+	//On gpu just define these as fmod(x,y)
+	float4 Modulo(float4 x, float4 y);
+	float3 Modulo(float3 x, float3 y);
+	float2 Modulo(float2 x, float2 y);
+	float Modulo(float x, float y);
+
+	/** Returns the index for this particle in the current execution context. On gpu this'll likely be derived from DispatchThreadId */
+	int ExecIndex();
+
+	//Some functions that we use to map to special VM operations for reading in data.
+	//TODO: replace with proper buffer reading capability and use standard hlsl.
+	int AcquireIndex(int DataSetID, bool DoAcquire);
+
+	void AcquireID(int DataSetID, out int IDIndex, out int IDAcquireTag);
+	void UpdateID(int DataSetID, int IDIndex, int InstanceIndex);
+
+	float InputDataFloat(int DataSetIndex, int RegisterIdx);  //DataSetIndex is 0 for main dataset
+	int InputDataInt(int DataSetIndex, int RegisterIdx);
+	bool InputDataBool(int DataSetIndex, int RegisterIdx);
+
+	float InputDataNoadvanceFloat(int DataSetIndex, int RegisterIdx);  //DataSetIndex is 0 for main dataset
+	int InputDataNoadvanceInt(int DataSetIndex, int RegisterIdx);
+	bool InputDataNoadvanceBool(int DataSetIndex, int RegisterIdx);
+
+	void OutputDataFloat(int DataSetIndex, int RegisterIndex, int InstanceIndex, float Value);
+	void OutputDataInt(int DataSetIndex, int RegisterIndex, int InstanceIndex, int Value);
+	void OutputDataBool(int DataSetIndex, int RegisterIndex, int InstanceIndex, bool Value);
+#endif
+
+
+
+
+/* -----------------------------------------------------------------
+ * GPU simulation code
+ * -----------------------------------------------------------------
+ */
+#if GPU_SIMULATION
+	static uint GSpawnStartInstance;
+		
+	/* Parameters for spawning and simulations
+	 */
+	uint UpdateStartInstance;
+	uint SpawnedInstances;
+	uint NumIndicesPerInstance;
+
+	uint ComponentBufferSizeRead;
+	uint ComponentBufferSizeWrite;
+	int	NumEventsPerParticle;
+	int NumParticlesPerEvent;
+	uint CopyInstancesBeforeStart;
+
+
+	/* Buffers for particle data and DrawIndirect calls
+	 */
+	Buffer<float> InputFloat;
+	Buffer<int> InputInt;
+	RWBuffer<int> RWOutputInt;
+	RWBuffer<float> RWOutputFloat;
+	RWBuffer<uint> RWDataSetIndices;
+	Buffer<uint> ReadDataSetIndices;
+
+
+	//TODO: IDs for GPU
+//	void AcquireID(int DataSetID, out int IDIndex, out int IDAcquireTag);
+// 	void UpdateID(int DataSetID, int IDIndex, int InstanceIndex);
+
+
+
+	/* Returns the current instance index relative to the operation (spawn/update)
+	 */
+	int ExecIndex()
+	{
+		if (GCurrentPhase == 0)
+		{
+			return UpdateStartInstance + GDispatchThreadId.x - GSpawnStartInstance;
+		}
+		else 
+		{
+			return UpdateStartInstance + GDispatchThreadId.x;
+		}
+	}
+
+
+	/* Acquire an output index - the default index is the scratch instance; one additional instance is allocated 
+	 *	at the end of the buffer, so no branching on -1 is necessary during OutputData operations
+	 */
+	int AcquireIndex(int DataSetID, bool bDoAcquire)
+	{
+		int PrevIdx = GSpawnStartInstance + SpawnedInstances;	// scratch instance as default; write to that for dead particles
+		if(bDoAcquire == true)
+		{
+			// Have to use uint's here to avoid PS4 compiler warnings about InterlockedAdd, cannot propagate uint due to CPU VM limitations...
+			uint RetPrevIdx;
+			InterlockedAdd(RWDataSetIndices[DataSetID*5+1], (uint)1U, RetPrevIdx);
+			PrevIdx = (int)RetPrevIdx;
+		}
+		return PrevIdx;
+	}
+
+
+
+	/* ---------------------------------------------------------------------
+	 * InputData operations 
+	 * ---------------------------------------------------------------------
+	 */
+	float InputDataFloat(int DataSetIndex, int RegisterIdx, int InstanceIdx)
+	{
+		return InputFloat[RegisterIdx*ComponentBufferSizeRead + InstanceIdx];
+	}
+
+	int InputDataInt(int DataSetIndex, int RegisterIdx, int InstanceIdx)
+	{
+		return InputInt[RegisterIdx*ComponentBufferSizeRead + InstanceIdx];
+	}
+	
+	bool InputDataBool(int DataSetIndex, int RegisterIdx, int InstanceIdx)
+	{
+		return InputInt[RegisterIdx*ComponentBufferSizeRead + InstanceIdx] == -1;
+	}
+
+
+	/* ---------------------------------------------------------------------
+	 * OutputData operations 
+	 * ---------------------------------------------------------------------
+	 */
+	void OutputDataFloat(int DataSetIndex, int RegisterIndex, int InstanceIndex, float Value)
+	{
+		RWOutputFloat[RegisterIndex*ComponentBufferSizeWrite + InstanceIndex + UpdateStartInstance] = Value;
+	}
+
+	void OutputDataInt(int DataSetIndex, int RegisterIndex, int InstanceIndex, int Value)
+	{
+		RWOutputInt[RegisterIndex*ComponentBufferSizeWrite + InstanceIndex + UpdateStartInstance] = Value;
+	}
+
+	void OutputDataBool(int DataSetIndex, int RegisterIndex, int InstanceIndex, bool Value)
+	{
+		RWOutputInt[RegisterIndex*ComponentBufferSizeWrite + InstanceIndex + UpdateStartInstance] = Value ? -1 : 0;
+	}
+
+	void EnterStatScope(int ID)	{}
+	void ExitStatScope()	{}
+#endif
+
+//Include the simulation shader code generated by the node graph.
+#include "/Engine/Generated/NiagaraEmitterInstance.ush"
+
+
+#if GPU_SIMULATION
+
+/* 
+ *	CS wrapper for our generated code; calls spawn and update functions on the corresponding instances in the buffer
+ */
+
+[numthreads(THREADGROUP_SIZE, 1, 1)]
+void SimulateMainComputeCS(
+	uint3 GroupId : SV_GroupID,
+	uint3 DispatchThreadId : SV_DispatchThreadID,
+	uint3 GroupThreadId : SV_GroupThreadID)
+{
+	GDispatchThreadId = DispatchThreadId;
+
+	/*
+	if(CopyInstancesBeforeStart == 1)
+	{
+		UpdateStartInstance = 0;
+	}
+	*/
+
+	uint InstanceID = UpdateStartInstance + DispatchThreadId.x;
+	GSpawnStartInstance = ReadDataSetIndices[1];				// needed by ExecIndex()
+	int MaxInstances = GSpawnStartInstance + SpawnedInstances;
+
+	if(InstanceID < UpdateStartInstance + MaxInstances )
+	{
+		// calculate event index for the current instance; this is primarily for spawn events, where we do a single read from a secondary DataSet for a subset of particles
+		//  for all-particle events, this will always end up being 0 because NumParticlesPerEvent = total particles in the set
+		uint EventIndex = 0; //(InstanceID-SimulateStartInstance) / NumParticlesPerEvent;
+
+		// before the simulation start - either update or copy
+		//
+		if(InstanceID < GSpawnStartInstance)
+		{
+			// copy instances before simulation start inst (e.g. for event spawns)
+			/*
+			if(CopyInstancesBeforeStart)
+			{
+				CopyInstance(InstanceID);
+			}
+			else
+			*/
+			{
+				GCurrentPhase = 1;
+				SimulateMain(InstanceID, EventIndex, 1);
+			}
+
+		}
+
+		// after simulation start - spawn phase
+		//
+		else if(InstanceID >= GSpawnStartInstance)
+		{
+			GCurrentPhase = 0;
+			SimulateMain(InstanceID, EventIndex, 0);		// SimulateMain spawn phase
+		}
+	}
+
+	// fill in the rest of RWDataSetIndices for the DrawIndirect call for the main data set; number of instances is incremented in AcquireIndex
+	if(InstanceID == 0)
+	{
+		RWDataSetIndices[0] = NumIndicesPerInstance;		// IndexCountPerInstance; passed in from the renderer - TODO: support multiple renderers with multiple indirect draw calls
+		RWDataSetIndices[2] = 0;		// StartIndexLocation
+		RWDataSetIndices[3] = 0;		// BaseVertexLocation
+		RWDataSetIndices[4] = 0;		// StartInstanceLocation
+	}
+}
+
+#endif//GPU_SIMULATION
diff --git a/Engine/Plugins/FX/Niagara/Shaders/Private/NiagaraMeshVertexFactory.ush b/Engine/Plugins/FX/Niagara/Shaders/Private/NiagaraMeshVertexFactory.ush
new file mode 100644
--- /dev/null
+++ b/Engine/Plugins/FX/Niagara/Shaders/Private/NiagaraMeshVertexFactory.ush
@@ -0,0 +1,925 @@
+// Copyright 1998-2019 Epic Games, Inc. All Rights Reserved.
+
+/*=============================================================================
+	MeshParticleVertexFactory.usf: Mesh particle vertex factory shader code.
+=============================================================================*/
+
+#include "/Engine/Private/VertexFactoryCommon.ush"
+#include "NiagaraVFParticleAccess.usf"
+
+// These are only enabled for SM5 due to limited interpolators elsewhere
+#define USE_PARTICLE_POSITION (FEATURE_LEVEL >= FEATURE_LEVEL_SM5 && NEEDS_PARTICLE_POSITION)
+#define USE_PARTICLE_VELOCITY (FEATURE_LEVEL >= FEATURE_LEVEL_SM5 && NEEDS_PARTICLE_VELOCITY)
+#define USE_PARTICLE_TIME (FEATURE_LEVEL >= FEATURE_LEVEL_SM5 && NEEDS_PARTICLE_TIME)
+#define USE_PARTICLE_RANDOM (FEATURE_LEVEL >= FEATURE_LEVEL_SM5 && NEEDS_PARTICLE_RANDOM)
+#define USE_PARTICLE_TRANSFORM (FEATURE_LEVEL >= FEATURE_LEVEL_SM5 && NEEDS_PARTICLE_TRANSFORM)
+
+// Only support previous transform for deferred feature level
+//#define USE_PREVIOUS_TRANSFORM (FEATURE_LEVEL >= FEATURE_LEVEL_SM4)
+#define USE_PREVIOUS_TRANSFORM 0
+
+uint MeshFacingMode;
+
+#if USE_PREVIOUS_TRANSFORM
+	Buffer<float4> PrevTransformBuffer;
+#endif
+
+#define MESH_FACING_DEFAULT 0
+#define MESH_FACING_VELOCITY 1
+#define MESH_FACING_CAMERA_POSITION 2
+#define MESH_FACING_CAMERA_PLANE 3
+
+/** Sorted index buffer for particles if we need to sort their draw order (translucency etc). */
+Buffer<int> SortedIndices;
+uint SortedIndicesOffset;
+
+struct FVertexFactoryInput
+{
+	float4	Position	: ATTRIBUTE0;
+	half3	TangentX	: ATTRIBUTE1;
+	// TangentZ.w contains sign of tangent basis determinant
+	half4	TangentZ	: ATTRIBUTE2;
+	half4	VertexColor : ATTRIBUTE3;
+
+#if NUM_MATERIAL_TEXCOORDS_VERTEX
+	float2	TexCoords[NUM_MATERIAL_TEXCOORDS_VERTEX] : ATTRIBUTE4;
+#elif USE_PARTICLE_SUBUVS
+	float2	TexCoords[1] : ATTRIBUTE4;
+#endif
+
+
+	uint InstanceId	: SV_InstanceID;
+};
+
+struct FVertexFactoryInterpolantsVSToPS
+{
+	TANGENTTOWORLD_INTERPOLATOR_BLOCK
+
+#if USE_PARTICLE_SUBUVS
+	float4 SubUV0AndTexCoord0 : TEXCOORD1;
+	float4 SubUV1AndLerp : TEXCOORD2;
+#else
+	#if  NUM_TEX_COORD_INTERPOLATORS
+			#if (ES2_PROFILE || ES3_1_PROFILE)
+				// Avoid dependent texture fetches, put all UVs in xy
+				float2	TexCoords[ NUM_TEX_COORD_INTERPOLATORS] : TEXCOORD0;
+			#else
+				float4	TexCoords[( NUM_TEX_COORD_INTERPOLATORS+1)/2]	: TEXCOORD0;
+			#endif
+	#endif
+#endif
+
+#if INTERPOLATE_VERTEX_COLOR
+	float4	VertexColor : COLOR0;
+#endif
+
+#if NEEDS_PARTICLE_COLOR
+	float4	ParticleColor : COLOR1;
+#endif
+
+#if (DYNAMIC_PARAMETERS_MASK & 1)
+	float4 DynamicParameter : COLOR2;
+#endif
+#if (DYNAMIC_PARAMETERS_MASK & 2)
+	float4 DynamicParameter1 : COLOR3;
+#endif
+#if (DYNAMIC_PARAMETERS_MASK & 4)
+	float4 DynamicParameter2 : COLOR4;
+#endif
+#if (DYNAMIC_PARAMETERS_MASK & 8)
+	float4 DynamicParameter3 : COLOR5;
+#endif
+
+#if USE_PARTICLE_POSITION
+	/** Particle position in camera-centered translated world space */
+	float3 ParticleTranslatedWorldPosition	: PARTICLE_POSITION;
+#endif
+
+#if USE_PARTICLE_VELOCITY
+	/** The velocity of the particle, XYZ: direction, W: speed. */
+	float4 ParticleVelocity	: PARTICLE_VELOCITY;
+#endif
+
+#if USE_PARTICLE_TIME
+	/** Relative alive time of the particle */
+	float RelativeTime : PARTICLE_TIME;
+#endif
+
+#if USE_PARTICLE_RANDOM
+	float ParticleRandom : PARTICLE_RANDOM;
+#endif
+
+#if USE_PARTICLE_TRANSFORM
+	float4 LocalToWorld[4] : PARTICLE_TRANSFORM;
+#endif
+
+#if INSTANCED_STEREO
+	nointerpolation uint EyeIndex : PACKED_EYE_INDEX;
+#endif
+};
+
+struct FVertexFactoryIntermediates
+{
+	/** The color of the vertex. */
+	float4 VertexColor;
+
+	/** The color of the particle. */
+	float4 ParticleColor;
+
+	/** The texture coordinates for the vertex. */
+#if NUM_MATERIAL_TEXCOORDS_VERTEX
+	float2	TexCoords[NUM_MATERIAL_TEXCOORDS_VERTEX];
+#endif
+#if USE_PARTICLE_SUBUVS
+	float4 SubUVCoords;
+	float SubUVLerp;
+#endif
+
+	/** Optional dynamic parameters for the particle. */
+#if (DYNAMIC_PARAMETERS_MASK & 1)
+	float4 DynamicParameter;
+#endif
+#if (DYNAMIC_PARAMETERS_MASK & 2)
+	float4 DynamicParameter1;
+#endif
+#if (DYNAMIC_PARAMETERS_MASK & 4)
+	float4 DynamicParameter2;
+#endif
+#if (DYNAMIC_PARAMETERS_MASK & 8)
+	float4 DynamicParameter3;
+#endif
+
+	/** The position of the particle, XYZ: direction. */
+	float4 ParticlePosition;
+
+	/** The velocity of the particle, XYZ: direction, W: speed. */
+	float4 ParticleVelocity;
+
+	/** Particle position in camera-centered translated world space. */
+	float3 ParticleTranslatedWorldPosition;
+
+#if USE_PARTICLE_TIME
+	/** Relative time. */
+	float RelativeTime;
+#endif
+
+#if USE_PARTICLE_RANDOM
+	/** Particle Random value */
+	float ParticleRandom;
+#endif
+
+	float4x4 LocalToWorld;
+};
+
+float4x4 GetParticleTransform(FVertexFactoryIntermediates Inters)
+{
+	return Inters.LocalToWorld;
+}
+
+#if USE_PREVIOUS_TRANSFORM
+float4x4 GetParticlePrevTransform(FVertexFactoryInput Input)
+{
+	uint TransformBaseIdx = Input.InstanceId * 3;
+
+	float4x4 PrevTransform4x4 = float4x4(
+		PrevTransformBuffer[TransformBaseIdx + 0],
+		PrevTransformBuffer[TransformBaseIdx + 1],
+		PrevTransformBuffer[TransformBaseIdx + 2],
+		float4(0, 0, 0, 1)
+		);
+
+	//Transform is transposed on CPU and position is packed into W. Transpose again to get the correct transform back.
+	return transpose(PrevTransform4x4);
+}
+#endif
+
+/** Converts from vertex factory specific interpolants to a FMaterialPixelParameters, which is used by material inputs. */
+FMaterialPixelParameters GetMaterialPixelParameters(FVertexFactoryInterpolantsVSToPS Interpolants, float4 SvPosition)
+{
+	// GetMaterialPixelParameters is responsible for fully initializing the result
+	FMaterialPixelParameters Result = MakeInitializedMaterialPixelParameters();
+
+#if USE_PARTICLE_SUBUVS
+	#if  NUM_TEX_COORD_INTERPOLATORS
+		UNROLL
+		for( int CoordinateIndex = 0; CoordinateIndex <  NUM_TEX_COORD_INTERPOLATORS; CoordinateIndex++ )
+		{
+			Result.TexCoords[CoordinateIndex] = Interpolants.SubUV0AndTexCoord0.zw;
+		}
+	#endif
+	Result.Particle.SubUVCoords[0] = Interpolants.SubUV0AndTexCoord0.xy;
+	Result.Particle.SubUVCoords[1] = Interpolants.SubUV1AndLerp.xy;
+	Result.Particle.SubUVLerp = Interpolants.SubUV1AndLerp.z;
+#elif  NUM_TEX_COORD_INTERPOLATORS
+	#if (ES2_PROFILE || ES3_1_PROFILE)
+		UNROLL
+		for( int CoordinateIndex = 0; CoordinateIndex <  NUM_TEX_COORD_INTERPOLATORS; CoordinateIndex++ )
+		{
+			Result.TexCoords[CoordinateIndex] = Interpolants.TexCoords[CoordinateIndex].xy;
+		}
+	#else
+		UNROLL
+		for(int CoordinateIndex = 0;CoordinateIndex <  NUM_TEX_COORD_INTERPOLATORS;CoordinateIndex += 2)
+		{
+			Result.TexCoords[CoordinateIndex] = Interpolants.TexCoords[CoordinateIndex/2].xy;
+
+			if(CoordinateIndex + 1 <  NUM_TEX_COORD_INTERPOLATORS)
+			{
+				Result.TexCoords[CoordinateIndex + 1] = Interpolants.TexCoords[CoordinateIndex/2].wz;
+			}
+		}
+	#endif
+#endif
+
+	half3 TangentToWorld0 = Interpolants.TangentToWorld0.xyz;
+	half4 TangentToWorld2 = Interpolants.TangentToWorld2;
+	Result.UnMirrored = TangentToWorld2.w;
+
+#if INTERPOLATE_VERTEX_COLOR
+	Result.VertexColor = Interpolants.VertexColor;
+#else
+	Result.VertexColor = 0;
+#endif
+
+#if NEEDS_PARTICLE_COLOR
+	Result.Particle.Color = Interpolants.ParticleColor;
+#endif
+
+#if (DYNAMIC_PARAMETERS_MASK != 0)
+	Result.Particle.DynamicParameterValidMask = NiagaraMeshVF.MaterialParamValidMask;
+#endif
+#if (DYNAMIC_PARAMETERS_MASK & 1)
+	Result.Particle.DynamicParameter = Interpolants.DynamicParameter;
+#endif
+#if (DYNAMIC_PARAMETERS_MASK & 2)
+	Result.Particle.DynamicParameter1 = Interpolants.DynamicParameter1;
+#endif
+#if (DYNAMIC_PARAMETERS_MASK & 4)
+	Result.Particle.DynamicParameter2 = Interpolants.DynamicParameter2;
+#endif
+#if (DYNAMIC_PARAMETERS_MASK & 8)
+	Result.Particle.DynamicParameter3 = Interpolants.DynamicParameter3;
+#endif
+
+#if USE_PARTICLE_POSITION
+	Result.Particle.TranslatedWorldPositionAndSize.xyz = Interpolants.ParticleTranslatedWorldPosition.xyz;
+	Result.Particle.TranslatedWorldPositionAndSize.w = 1;
+#endif
+
+#if USE_PARTICLE_VELOCITY
+	Result.Particle.Velocity = Interpolants.ParticleVelocity;
+#endif
+
+#if USE_PARTICLE_TIME
+	Result.Particle.RelativeTime = Interpolants.RelativeTime;
+#endif
+
+#if USE_PARTICLE_RANDOM
+	Result.Particle.Random = Interpolants.ParticleRandom;
+#else
+	Result.Particle.Random = 0.0f;
+#endif
+
+#if USE_PARTICLE_TRANSFORM
+	Result.Particle.LocalToWorld[0] = Interpolants.LocalToWorld[0];
+	Result.Particle.LocalToWorld[1] = Interpolants.LocalToWorld[1];
+	Result.Particle.LocalToWorld[2] = Interpolants.LocalToWorld[2];
+	Result.Particle.LocalToWorld[3] = Interpolants.LocalToWorld[3];
+#endif
+
+	Result.Particle.MotionBlurFade = 1.0f;
+	Result.TangentToWorld = AssembleTangentToWorld( TangentToWorld0, TangentToWorld2 );
+	Result.TwoSidedSign = 1;
+
+#if USE_WORLDVERTEXNORMAL_CENTER_INTERPOLATION
+	Result.WorldVertexNormal_Center = Interpolants.TangentToWorld2_Center.xyz;
+#endif
+
+	return Result;
+}
+
+/** Converts from vertex factory specific input to a FMaterialVertexParameters, which is used by vertex shader material inputs. */
+FMaterialVertexParameters GetMaterialVertexParameters(FVertexFactoryInput Input, FVertexFactoryIntermediates Intermediates, float3 WorldPosition, float3x3 TangentToLocal)
+{
+	FMaterialVertexParameters Result = (FMaterialVertexParameters)0;
+	Result.WorldPosition = WorldPosition;
+	Result.VertexColor = Intermediates.VertexColor;
+	Result.PreSkinnedPosition = Input.Position.xyz;
+	Result.PreSkinnedNormal = TangentBias(Input.TangentZ.xyz);
+
+#if USE_PARTICLE_POSITION
+	Result.Particle.TranslatedWorldPositionAndSize.xyz = Intermediates.ParticleTranslatedWorldPosition.xyz;
+	Result.Particle.TranslatedWorldPositionAndSize.w = 1;
+#endif
+
+#if USE_PARTICLE_VELOCITY
+	Result.Particle.Velocity = Intermediates.ParticleVelocity;
+#endif
+
+#if USE_PARTICLE_TIME
+	Result.Particle.RelativeTime = Intermediates.RelativeTime;
+#endif
+
+#if USE_PARTICLE_RANDOM
+	Result.Particle.Random = Intermediates.ParticleRandom;
+#else
+	Result.Particle.Random = 0.0f;
+#endif
+
+	Result.Particle.Color = Intermediates.ParticleColor;
+
+#if (DYNAMIC_PARAMETERS_MASK != 0)
+	Result.Particle.DynamicParameterValidMask = NiagaraMeshVF.MaterialParamValidMask;
+#endif
+#if (DYNAMIC_PARAMETERS_MASK & 1)
+	Result.Particle.DynamicParameter = Intermediates.DynamicParameter;
+#endif
+#if (DYNAMIC_PARAMETERS_MASK & 2)
+	Result.Particle.DynamicParameter1 = Intermediates.DynamicParameter1;
+#endif
+#if (DYNAMIC_PARAMETERS_MASK & 4)
+	Result.Particle.DynamicParameter2 = Intermediates.DynamicParameter2;
+#endif
+#if (DYNAMIC_PARAMETERS_MASK & 8)
+	Result.Particle.DynamicParameter3 = Intermediates.DynamicParameter3;
+#endif
+
+	float4x4 Transform = GetParticleTransform(Intermediates);
+	Result.InstanceLocalToWorld = Transform;
+
+#if NEEDS_PARTICLE_TRANSFORM
+	Result.Particle.LocalToWorld[0] = Transform[0];
+	Result.Particle.LocalToWorld[1] = Transform[1];
+	Result.Particle.LocalToWorld[2] = Transform[2];
+	Result.Particle.LocalToWorld[3] = Transform[3];
+#endif
+
+	Result.TangentToWorld = mul(TangentToLocal, (float3x3)Transform );
+
+#if USE_PREVIOUS_TRANSFORM
+	Result.PrevFrameLocalToWorld = GetParticlePrevTransform(Input);
+#else
+	Result.PrevFrameLocalToWorld = GetParticleTransform(Intermediates);
+#endif
+
+#if NUM_MATERIAL_TEXCOORDS_VERTEX
+	for(int CoordinateIndex = 0; CoordinateIndex < NUM_MATERIAL_TEXCOORDS_VERTEX; CoordinateIndex++)
+	{
+		Result.TexCoords[CoordinateIndex] = Intermediates.TexCoords[CoordinateIndex];
+	}
+#endif
+
+	return Result;
+}
+
+float4 CalcWorldPosition(FVertexFactoryInput Input, FVertexFactoryIntermediates Intermediates)
+{
+	float4x4 Transform = GetParticleTransform(Intermediates);
+	float3 WorldPosition = mul(Input.Position, Transform).xyz + ResolvedView.PreViewTranslation;
+	return float4(WorldPosition, Input.Position.w);
+}
+
+float4 CalcPrevWorldPosition(FVertexFactoryInput Input, FVertexFactoryIntermediates Intermediates)
+{
+#if USE_PREVIOUS_TRANSFORM
+	if (MeshParticleVF.PrevTransformAvailable > 0)
+	{
+		float4x4 PrevTransform = GetParticlePrevTransform(Input);
+		float3 PrevWorldPosition = mul(Input.Position, PrevTransform).xyz + ResolvedView.PrevPreViewTranslation;
+		return float4(PrevWorldPosition, Input.Position.w);
+	}
+	else
+#endif
+	// if we don't have a full previous transform, we can at least calculate the position to get some motion blur
+	{
+		float3 WorldPosition = CalcWorldPosition(Input, Intermediates).xyz;
+
+		// ParticleVelocity.xyz - normalized direction vector
+		// ParticleVelocity.w - velocity magnitude
+		float3 WorldDisplacement = -(Intermediates.ParticleVelocity.xyz * (Intermediates.ParticleVelocity.w * NiagaraMeshVF.DeltaSeconds));
+
+		float3 PrevWorldPosition = WorldPosition - ResolvedView.PreViewTranslation;
+		PrevWorldPosition += WorldDisplacement + ResolvedView.PrevPreViewTranslation;
+
+		return float4(PrevWorldPosition.xyz, Input.Position.w);
+	}
+}
+
+// From UnrealMath::FindBetween
+float4 FindQuatBetween(float3 From, float3 To)
+{
+	float NormAB = 1.0f;
+	float w = NormAB + dot(From, To);
+	float4 Result;
+
+	if (w >= 1e-6f * NormAB)
+	{
+		//Axis = cross(From, To);
+		Result = float4(From.y * To.z - From.z * To.y,
+					   From.z * To.x - From.x * To.z,
+					   From.x * To.y - From.y * To.x,
+					   w);
+	}
+	else
+	{
+		// From and To point in opposite directions
+		w = 0.f;
+		Result = abs(From.x) > abs(From.y)
+				? float4(-From.z, 0.f, From.x, w)
+				: float4(0.f, -From.z, From.y, w);
+	}
+
+	Result = normalize(Result);
+	return Result;
+}
+
+float4x4 QuatTo4x4(float4 Q)
+{
+	float x2 = Q.x + Q.x;  float y2 = Q.y + Q.y;  float z2 = Q.z + Q.z;
+	float xx = Q.x * x2;   float xy = Q.x * y2;   float xz = Q.x * z2;
+	float yy = Q.y * y2;   float yz = Q.y * z2;   float zz = Q.z * z2;
+	float wx = Q.w * x2;   float wy = Q.w * y2;   float wz = Q.w * z2;
+
+	float4x4 Mat = {
+	/*M[0][0]*/ 1.0f - (yy + zz),	/*M[1][0]*/ xy - wz,				/*M[2][0]*/ xz + wy,		0.0f,
+	/*M[0][1]*/ xy + wz,			/*M[1][1]*/ 1.0f - (xx + zz),		/*M[2][1]*/ yz - wx,		0.0f,
+	/*M[0][2]*/ xz - wy,			/*M[1][2]*/ yz + wx,				/*M[2][2]*/ 1.0f - (xx + yy),0.0f,
+	0.0f,							0.0f,								0.0f,						1.0f
+	};
+	return Mat;
+}
+
+void AxisAngleToPitchYawRoll(float4 AxisAngle, out float Pitch, out float Yaw, out float Roll)
+{
+	AxisAngle.xyz= normalize(AxisAngle).xyz;
+	Yaw = atan2(AxisAngle.y, AxisAngle.x) * (180.f / PI);
+	if (AxisAngle.y == 0.0f && AxisAngle.x == 0.0f) // atan2 is undefined for 0,0
+	{
+		Yaw = 0.0f;
+	}
+	Pitch = atan2(AxisAngle.z,sqrt(AxisAngle.x*AxisAngle.x+AxisAngle.y*AxisAngle.y)) * (180.f / PI);
+	Roll = AxisAngle.w* (180.f / PI);
+}
+
+
+float4 EulerToQuaternion(float Pitch, float Yaw, float Roll)
+{
+	float DEG_TO_RAD = PI/(180.f);
+	float DIVIDE_BY_2 = DEG_TO_RAD/2.f;
+	float SP = sin(Pitch*DIVIDE_BY_2);
+	float SY = sin(Yaw*DIVIDE_BY_2);
+	float SR = sin(Roll*DIVIDE_BY_2);
+	float CP = cos(Pitch*DIVIDE_BY_2);
+	float CY = cos(Yaw*DIVIDE_BY_2);
+	float CR = cos(Roll*DIVIDE_BY_2);
+
+	float4 RotationQuat;
+	RotationQuat.x =  CR*SP*SY - SR*CP*CY;
+	RotationQuat.y = -CR*SP*CY - SR*CP*SY;
+	RotationQuat.z =  CR*CP*SY - SR*SP*CY;
+	RotationQuat.w =  CR*CP*CY + SR*SP*SY;
+
+	return RotationQuat;
+}
+
+float4x4 EulerToMat4x4(float Pitch, float Yaw, float Roll)
+{
+	float4 Quat = EulerToQuaternion(Pitch, Yaw, Roll);
+	return QuatTo4x4(Quat);
+}
+
+
+FVertexFactoryIntermediates GetVertexFactoryIntermediates(FVertexFactoryInput Input)
+{
+	FVertexFactoryIntermediates Intermediates = (FVertexFactoryIntermediates)0;
+
+#if !SHADOW_DEPTH_SHADER && !SPHERICAL_OPACITY_FOR_SHADOW_DEPTHS
+	// Don't use sort buffer in the shadow because they are only valid after opaque.
+	if (SortedIndicesOffset != 0xFFFFFFFF)
+	{
+		Input.InstanceId = SortedIndices[SortedIndicesOffset + Input.InstanceId];
+	}
+#endif
+
+	if (NiagaraMeshVF.PositionDataOffset == -1)
+	{
+		Intermediates.ParticlePosition = NiagaraMeshVF.DefaultPos;
+	}
+	else
+	{
+		Intermediates.ParticlePosition = float4(GetVec3(NiagaraMeshVF.PositionDataOffset, Input.InstanceId), 1);
+	}
+
+	if (NiagaraMeshVF.ColorDataOffset == -1)
+	{
+		Intermediates.ParticleColor = float4(1.0f, 1.0f, 1.0f, 1.0f);
+	}
+	else
+	{
+		Intermediates.ParticleColor = GetVec4(NiagaraMeshVF.ColorDataOffset, Input.InstanceId);
+	}
+
+	float3 ParticleVelocity;
+	if (NiagaraMeshVF.VelocityDataOffset == -1)
+	{
+		ParticleVelocity = float3(0.0f, 0.0f, 0.0f);
+	}
+	else
+	{
+		ParticleVelocity = float3(GetVec3(NiagaraMeshVF.VelocityDataOffset, Input.InstanceId));
+	}
+
+#if (DYNAMIC_PARAMETERS_MASK & 1)
+	Intermediates.DynamicParameter = NiagaraMeshVF.MaterialParamDataOffset != -1 ? GetVec4(NiagaraMeshVF.MaterialParamDataOffset, Input.InstanceId) : float4(1, 1, 1, 1);
+#endif
+#if (DYNAMIC_PARAMETERS_MASK & 2)
+	Intermediates.DynamicParameter1 = NiagaraMeshVF.MaterialParam1DataOffset != -1 ? GetVec4(NiagaraMeshVF.MaterialParam1DataOffset, Input.InstanceId) : float4(1, 1, 1, 1);
+#endif
+#if (DYNAMIC_PARAMETERS_MASK & 4)
+	Intermediates.DynamicParameter2 = NiagaraMeshVF.MaterialParam2DataOffset != -1 ? GetVec4(NiagaraMeshVF.MaterialParam2DataOffset, Input.InstanceId) : float4(1, 1, 1, 1);
+#endif
+#if (DYNAMIC_PARAMETERS_MASK & 8)
+	Intermediates.DynamicParameter3 = NiagaraMeshVF.MaterialParam3DataOffset != -1 ? GetVec4(NiagaraMeshVF.MaterialParam3DataOffset, Input.InstanceId) : float4(1, 1, 1, 1);
+#endif
+
+#if USE_PARTICLE_TIME
+	if (NiagaraMeshVF.NormalizedAgeDataOffset == -1)
+	{
+		Intermediates.RelativeTime = 0;
+	}
+	else
+	{
+		Intermediates.RelativeTime = GetFloat(NiagaraMeshVF.NormalizedAgeDataOffset, Input.InstanceId);
+	}
+#endif
+
+#if USE_PARTICLE_RANDOM
+	if (NiagaraMeshVF.MaterialRandomDataOffset == -1)
+	{
+		Intermediates.ParticleRandom = 0;
+	}
+	else
+	{
+		Intermediates.ParticleRandom = GetFloat(NiagaraMeshVF.MaterialRandomDataOffset, Input.InstanceId);
+	}
+#endif
+
+#if USE_PARTICLE_VELOCITY
+	float ParticleVelocityMagnitude = length(ParticleVelocity);
+
+	Intermediates.ParticleVelocity.xyz = (ParticleVelocityMagnitude > 0) ? normalize(ParticleVelocity) : float3(0.0f, 0.0f, 0.0f);
+	Intermediates.ParticleVelocity.w = ParticleVelocityMagnitude;
+#endif
+
+	float4 Transform1 = float4(1.0f, 0.0f, 0.0f, Intermediates.ParticlePosition.x);
+	float4 Transform2 = float4(0.0f, 1.0f, 0.0f, Intermediates.ParticlePosition.y);
+	float4 Transform3 = float4(0.0f, 0.0f, 1.0f, Intermediates.ParticlePosition.z);
+
+	if(NiagaraMeshVF.TransformDataOffset != -1)
+	{
+		float4 Quat = GetVec4(NiagaraMeshVF.TransformDataOffset, Input.InstanceId);
+
+		//Transform axis to handle local space.
+		//AxisAngle.xyz = mul(float4(AxisAngle.xyz, 0), NiagaraMeshVF.LocalToWorldInverseTransposed);
+
+
+		float4x4 EulerToMat = QuatTo4x4(normalize(Quat));
+		Transform1.xyz = EulerToMat[0].xyz;
+		Transform2.xyz = EulerToMat[1].xyz;
+		Transform3.xyz = EulerToMat[2].xyz;
+	}
+	float4x4 ScaleMatrix = float4x4(float4(1.0f, 0.0f, 0.0f, 0.0f), float4(0.0f, 1.0f, 0.0f, 0.0f), float4(0.0f, 0.0f, 1.0f, 0.0f), float4(0.0f, 0.0f, 0.0f, 1.0f));
+	if(NiagaraMeshVF.ScaleDataOffset != -1)
+	{
+		float3 Scale = GetVec3(NiagaraMeshVF.ScaleDataOffset, Input.InstanceId);
+		ScaleMatrix[0][0] *= Scale.x;
+		ScaleMatrix[1][1] *= Scale.y;
+		ScaleMatrix[2][2] *= Scale.z;
+	}
+	// Swizzle vertex color.
+	Intermediates.VertexColor = Input.VertexColor FCOLOR_COMPONENT_SWIZZLE;
+
+	// World position and localtoworld transform.
+	float4x4 LocalToWorldMat = float4x4(Transform1, Transform2, Transform3, float4(0.0f, 0.0f, 0.0f, 1.0f));
+
+	Intermediates.ParticleTranslatedWorldPosition = Intermediates.ParticlePosition.xyz + ResolvedView.PreViewTranslation.xyz;
+
+	//LocalToWorldMat = mul(ScaleMatrix, LocalToWorldMat);
+
+	if (MeshFacingMode == MESH_FACING_VELOCITY)
+	{
+		float4 ToAlignmentQuat = FindQuatBetween(float3(1.0f, 0.0f, 0.0f), normalize(ParticleVelocity.xyz));
+		float4x4 ToAlignmentMat = QuatTo4x4( normalize(ToAlignmentQuat) );
+		Intermediates.LocalToWorld =  mul(LocalToWorldMat, mul(ToAlignmentMat, ScaleMatrix));
+	}
+	else if (MeshFacingMode == MESH_FACING_CAMERA_POSITION)
+	{
+		float3 WorldPosition = Intermediates.ParticleTranslatedWorldPosition;
+		float3 CameraVector = ResolvedView.WorldCameraOrigin.xyz - WorldPosition.xyz;
+		float3 CameraDirection = normalize(CameraVector);
+		float4 ToAlignmentQuat = FindQuatBetween(float3(1.0f, 0.0f, 0.0f), CameraDirection);
+		float4x4 ToAlignmentMat = QuatTo4x4( normalize(ToAlignmentQuat) );
+		Intermediates.LocalToWorld = mul(LocalToWorldMat, mul(ToAlignmentMat, ScaleMatrix));
+	}
+	else if (MeshFacingMode == MESH_FACING_CAMERA_PLANE)
+	{
+		float3 CameraFacingVector = -ResolvedView.ViewForward;
+		float4 ToAlignmentQuat = FindQuatBetween(float3(1.0f, 0.0f, 0.0f), CameraFacingVector);
+		float4x4 ToAlignmentMat = QuatTo4x4( normalize(ToAlignmentQuat) );
+		Intermediates.LocalToWorld =  mul(LocalToWorldMat, mul(ToAlignmentMat, ScaleMatrix));
+	}
+	else
+	{
+		Intermediates.LocalToWorld =  mul(LocalToWorldMat, ScaleMatrix);
+	}
+
+	//Handle local space. TODO: Optimize behind define.
+	//Apply local space to local to world, also pretranspse here to avoid it in the vs&ps
+	Intermediates.LocalToWorld = transpose(Intermediates.LocalToWorld);
+	Intermediates.LocalToWorld = mul(Intermediates.LocalToWorld, NiagaraMeshVF.LocalToWorld);
+	Intermediates.ParticlePosition = mul(Intermediates.ParticlePosition, NiagaraMeshVF.LocalToWorld);
+	//Intermediates.ParticleVelocity.xyz = mul(float4(Intermediates.ParticleVelocity.xyz,0), NiagaraMeshVF.LocalToWorldInverseTransposed);
+
+#if NUM_MATERIAL_TEXCOORDS_VERTEX
+	for(int CoordinateIndex = 0; CoordinateIndex < NUM_MATERIAL_TEXCOORDS_VERTEX; CoordinateIndex++)
+	{
+		Intermediates.TexCoords[CoordinateIndex] = Input.TexCoords[CoordinateIndex].xy;
+	}
+#endif
+
+#if USE_PARTICLE_SUBUVS
+		Intermediates.SubUVCoords.xy = (float2( Input.SubUVParams.x, Input.SubUVParams.y ) + Input.TexCoords[0].xy) * MeshParticleVF.SubImageSize.xy;
+		Intermediates.SubUVCoords.zw = (float2( Input.SubUVParams.z, Input.SubUVParams.w ) + Input.TexCoords[0].xy) * MeshParticleVF.SubImageSize.xy;
+		//Intermediates.SubUVLerp = Input.SubUVLerpAndRelTime.x;
+#endif
+
+	//Intermediates.RelativeTime = Input.SubUVLerpAndRelTime.y;
+
+	return Intermediates;
+}
+
+/**
+* Get the 3x3 tangent basis vectors for this vertex factory
+* this vertex factory will calculate the binormal on-the-fly
+*
+* @param Input - vertex input stream structure
+* @return 3x3 matrix
+*/
+float3x3 VertexFactoryGetTangentToLocal( FVertexFactoryInput Input, FVertexFactoryIntermediates Intermediates )
+{
+	float3x3 Result=0;
+	float3 TangentX = TangentBias(Input.TangentX);
+	float4 TangentZ = TangentBias(Input.TangentZ);
+
+	// pass-thru the tangent
+	Result[0] = TangentX;
+	// pass-thru the normal
+	Result[2] = TangentZ.xyz;
+
+
+	// derive the binormal by getting the cross product of the normal and tangent
+	Result[1] = cross(Result[2], Result[0]) * TangentZ.w;
+	// Recalculate TangentX off of the other two vectors
+	// This corrects quantization error since TangentX was passed in as a quantized vertex input
+	// The error shows up most in specular off of a mesh with a smoothed UV seam (normal is smooth, but tangents vary across the seam)
+	Result[0] = cross(Result[1], Result[2]) * TangentZ.w;
+
+	return Result;
+}
+
+// @return translated world position
+float4 VertexFactoryGetWorldPosition(FVertexFactoryInput Input, FVertexFactoryIntermediates Intermediates)
+{
+	return CalcWorldPosition(Input, Intermediates);
+}
+
+float4 VertexFactoryGetRasterizedWorldPosition(FVertexFactoryInput Input, FVertexFactoryIntermediates Intermediates, float4 InWorldPosition)
+{
+	return InWorldPosition;
+}
+
+float3 VertexFactoryGetPositionForVertexLighting(FVertexFactoryInput Input, FVertexFactoryIntermediates Intermediates, float3 TranslatedWorldPosition)
+{
+	return TranslatedWorldPosition;
+}
+
+void CalcTangentToWorld(FVertexFactoryInput Input, FVertexFactoryIntermediates Intermediates, out float3 TangentToWorld0, out float4 TangentToWorld2)
+{
+	float3x3 TangentToLocal = VertexFactoryGetTangentToLocal(Input, Intermediates);
+
+	float4x4 Transform = GetParticleTransform(Intermediates);
+	float3x3 TangentToWorld = mul(TangentToLocal, (float3x3)Transform);
+
+	// Normalize to remove scaling. Incorrect but faster than computing the inverse-transpose.
+	TangentToWorld0 = normalize(TangentToWorld[0]);
+	TangentToWorld2 = float4(normalize(TangentToWorld[2]), TangentBias(Input.TangentZ.w) * Primitive.InvNonUniformScaleAndDeterminantSign.w);
+}
+
+FVertexFactoryInterpolantsVSToPS VertexFactoryGetInterpolantsVSToPS(FVertexFactoryInput Input, FVertexFactoryIntermediates Intermediates, FMaterialVertexParameters VertexParameters)
+{
+	FVertexFactoryInterpolantsVSToPS Interpolants;
+
+#if USE_PARTICLE_SUBUVS
+	Interpolants.SubUV0AndTexCoord0.xy = Intermediates.SubUVCoords.xy;
+	Interpolants.SubUV1AndLerp.xy = Intermediates.SubUVCoords.zw;
+	Interpolants.SubUV1AndLerp.zw = Intermediates.SubUVLerp.xx;
+
+	#if  NUM_TEX_COORD_INTERPOLATORS
+		float2 CustomizedUVs[ NUM_TEX_COORD_INTERPOLATORS];
+		GetMaterialCustomizedUVs(VertexParameters, CustomizedUVs);
+		Interpolants.SubUV0AndTexCoord0.zw = CustomizedUVs[0];
+	#else
+		Interpolants.SubUV0AndTexCoord0.zw = 0;
+	#endif
+#elif  NUM_TEX_COORD_INTERPOLATORS
+	#if (ES2_PROFILE || ES3_1_PROFILE)
+		GetMaterialCustomizedUVs(VertexParameters, Interpolants.TexCoords);
+	#else
+		// Ensure the unused components of the last packed texture coordinate are initialized.
+		Interpolants.TexCoords[( NUM_TEX_COORD_INTERPOLATORS + 1) / 2 - 1] = 0;
+
+		float2 CustomizedUVs[ NUM_TEX_COORD_INTERPOLATORS];
+		GetMaterialCustomizedUVs(VertexParameters, CustomizedUVs);
+
+		UNROLL
+		for(int CoordinateIndex = 0;CoordinateIndex <  NUM_TEX_COORD_INTERPOLATORS;CoordinateIndex += 2)
+		{
+				Interpolants.TexCoords[CoordinateIndex / 2].xy = CustomizedUVs[CoordinateIndex];
+			if(CoordinateIndex + 1 <  NUM_TEX_COORD_INTERPOLATORS)
+			{
+					Interpolants.TexCoords[CoordinateIndex / 2].wz = CustomizedUVs[CoordinateIndex + 1];
+			}
+		}
+	#endif
+#endif
+
+
+	Interpolants.TangentToWorld0.w = 0;
+	CalcTangentToWorld(Input, Intermediates, Interpolants.TangentToWorld0.xyz, Interpolants.TangentToWorld2);
+#if USE_WORLDVERTEXNORMAL_CENTER_INTERPOLATION
+	Interpolants.TangentToWorld2_Center = Interpolants.TangentToWorld2;
+#endif
+
+#if INTERPOLATE_VERTEX_COLOR
+	Interpolants.VertexColor = Intermediates.VertexColor;
+#endif
+
+#if NEEDS_PARTICLE_COLOR
+	Interpolants.ParticleColor = Intermediates.ParticleColor;
+#endif
+
+#if (DYNAMIC_PARAMETERS_MASK & 1)
+	Interpolants.DynamicParameter = Intermediates.DynamicParameter;
+#endif
+#if (DYNAMIC_PARAMETERS_MASK & 2)
+	Interpolants.DynamicParameter1 = Intermediates.DynamicParameter1;
+#endif
+#if (DYNAMIC_PARAMETERS_MASK & 4)
+	Interpolants.DynamicParameter2 = Intermediates.DynamicParameter2;
+#endif
+#if (DYNAMIC_PARAMETERS_MASK & 8)
+	Interpolants.DynamicParameter3 = Intermediates.DynamicParameter3;
+#endif
+
+
+#if USE_PARTICLE_POSITION
+	Interpolants.ParticleTranslatedWorldPosition = Intermediates.ParticleTranslatedWorldPosition;
+#endif
+
+#if USE_PARTICLE_VELOCITY
+	Interpolants.ParticleVelocity = Intermediates.ParticleVelocity;
+#endif
+
+#if USE_PARTICLE_TIME
+	Interpolants.RelativeTime = Intermediates.RelativeTime;
+#endif
+
+#if USE_PARTICLE_RANDOM
+	Interpolants.ParticleRandom = Intermediates.ParticleRandom;
+#endif
+
+#if USE_PARTICLE_TRANSFORM
+	Interpolants.LocalToWorld[0] = Intermediates.LocalToWorld[0];
+	Interpolants.LocalToWorld[1] = Intermediates.LocalToWorld[1];
+	Interpolants.LocalToWorld[2] = Intermediates.LocalToWorld[2];
+	Interpolants.LocalToWorld[3] = Intermediates.LocalToWorld[3];
+#endif
+
+#if INSTANCED_STEREO
+	Interpolants.EyeIndex = 0;
+#endif
+
+	return Interpolants;
+}
+
+// @return previous translated world position
+float4 VertexFactoryGetPreviousWorldPosition(FVertexFactoryInput Input, FVertexFactoryIntermediates Intermediates)
+{
+	return CalcPrevWorldPosition(Input, Intermediates);
+}
+
+#if USING_TESSELLATION
+	struct FVertexFactoryInterpolantsVSToDS
+	{
+		FVertexFactoryInterpolantsVSToPS InterpolantsVSToPS;
+	};
+
+	float2 VertexFactoryGetTextureCoordinateDS( FVertexFactoryInterpolantsVSToDS Interpolants )
+	{
+	#if NUM_MATERIAL_TEXCOORDS
+		return Interpolants.InterpolantsVSToPS.TexCoords[0].xy;
+	#else // #if NUM_MATERIAL_TEXCOORDS
+		return float2(0,0);
+	#endif // #if NUM_MATERIAL_TEXCOORDS
+	}
+
+	FVertexFactoryInterpolantsVSToPS VertexFactoryAssignInterpolants(FVertexFactoryInterpolantsVSToDS Input)
+	{
+		return Input.InterpolantsVSToPS;
+	}
+
+	/** Converts from vertex factory specific interpolants to a FMaterialTessellationParameters, which is used by material inputs. */
+	FMaterialTessellationParameters GetMaterialTessellationParameters(FVertexFactoryInterpolantsVSToDS Interpolants, float3 CameraLocalWorldPosition)
+	{
+		FMaterialTessellationParameters	Result;
+	#if  NUM_TEX_COORD_INTERPOLATORS
+		UNROLL
+		for(int CoordinateIndex = 0;CoordinateIndex <  NUM_TEX_COORD_INTERPOLATORS;CoordinateIndex += 2)
+		{
+			Result.TexCoords[CoordinateIndex] = Interpolants.InterpolantsVSToPS.TexCoords[CoordinateIndex/2].xy;
+			if(CoordinateIndex + 1 <  NUM_TEX_COORD_INTERPOLATORS)
+			{
+				Result.TexCoords[CoordinateIndex + 1] = Interpolants.InterpolantsVSToPS.TexCoords[CoordinateIndex/2].wz;
+			}
+		}
+	#endif
+
+		Result.VertexColor = float4(1,1,1,1);
+
+		half3 TangentToWorld0 = Interpolants.InterpolantsVSToPS.TangentToWorld0.xyz;
+		half4 TangentToWorld2 = Interpolants.InterpolantsVSToPS.TangentToWorld2;
+		Result.TangentToWorld = AssembleTangentToWorld( TangentToWorld0, TangentToWorld2 );
+		Result.TangentToWorldPreScale = 1;
+		Result.WorldPosition = CameraLocalWorldPosition + ResolvedView.WorldCameraOrigin;
+
+		return Result;
+	}
+
+	FVertexFactoryInterpolantsVSToDS VertexFactoryGetInterpolantsVSToDS(FVertexFactoryInput Input, FVertexFactoryIntermediates Intermediates, FMaterialVertexParameters VertexParameters)
+	{
+		FVertexFactoryInterpolantsVSToDS Interpolants;
+		Interpolants.InterpolantsVSToPS = VertexFactoryGetInterpolantsVSToPS(Input, Intermediates, VertexParameters);
+		return Interpolants;
+	}
+
+	FVertexFactoryInterpolantsVSToDS VertexFactoryInterpolate(FVertexFactoryInterpolantsVSToDS a, float aInterp, FVertexFactoryInterpolantsVSToDS b, float bInterp)
+	{
+		FVertexFactoryInterpolantsVSToDS O;
+
+		// Do we really need to interpolate TangentToWorld2 here? It should be replaced by the
+		// interpolated normal from 'whatever' interpolation scheme we're using
+
+		TESSELLATION_INTERPOLATE_MEMBER(InterpolantsVSToPS.TangentToWorld0.xyz);
+		TESSELLATION_INTERPOLATE_MEMBER(InterpolantsVSToPS.TangentToWorld2);
+#if INTERPOLATE_VERTEX_COLOR
+		TESSELLATION_INTERPOLATE_MEMBER(InterpolantsVSToPS.VertexColor);
+#endif
+
+#if NEEDS_PARTICLE_COLOR
+		TESSELLATION_INTERPOLATE_MEMBER(InterpolantsVSToPS.ParticleColor);
+#endif
+
+#if  NUM_TEX_COORD_INTERPOLATORS
+		for (int i = 0; i < ( NUM_TEX_COORD_INTERPOLATORS + 1) / 2; ++i)
+		{
+			TESSELLATION_INTERPOLATE_MEMBER(InterpolantsVSToPS.TexCoords[i]);
+		}
+#endif
+
+#if (DYNAMIC_PARAMETERS_MASK & 1)
+		TESSELLATION_INTERPOLATE_MEMBER(InterpolantsVSToPS.DynamicParameter);
+#endif
+#if (DYNAMIC_PARAMETERS_MASK & 2)
+		TESSELLATION_INTERPOLATE_MEMBER(InterpolantsVSToPS.DynamicParameter1);
+#endif
+#if (DYNAMIC_PARAMETERS_MASK & 4)
+		TESSELLATION_INTERPOLATE_MEMBER(InterpolantsVSToPS.DynamicParameter2);
+#endif
+#if (DYNAMIC_PARAMETERS_MASK & 8)
+		TESSELLATION_INTERPOLATE_MEMBER(InterpolantsVSToPS.DynamicParameter3);
+#endif
+
+		return O;
+	}
+
+	float3x3 VertexFactoryGetTangentToLocalDS(FVertexFactoryInterpolantsVSToDS Interpolants)
+	{
+		// This duplicates stuff already going on in GetMaterialTessellationParameters(), so
+		// maybe the hull shader could leverage that instead?
+		half3 TangentToWorld0 = Interpolants.InterpolantsVSToPS.TangentToWorld0.xyz;
+		half4 TangentToWorld2 = Interpolants.InterpolantsVSToPS.TangentToWorld2;
+		float3x3 TangentToWorld = AssembleTangentToWorld( TangentToWorld0, TangentToWorld2 );
+
+		return TangentToWorld;
+	}
+#endif // #if USING_TESSELLATION
diff --git a/Engine/Plugins/FX/Niagara/Shaders/Private/NiagaraSortKeyGen.usf b/Engine/Plugins/FX/Niagara/Shaders/Private/NiagaraSortKeyGen.usf
new file mode 100644
--- /dev/null
+++ b/Engine/Plugins/FX/Niagara/Shaders/Private/NiagaraSortKeyGen.usf
@@ -0,0 +1,110 @@
+// Copyright 1998-2019 Epic Games, Inc. All Rights Reserved.
+
+/*==============================================================================
+	NiagaraSortKeyGen.usf: Shader to generate particle sorting keys.
+==============================================================================*/
+
+/*------------------------------------------------------------------------------
+	Compile time parameters:
+		THREAD_COUNT - The number of threads to launch per workgroup.
+		SORT_VIEW_DEPTH
+		SORT_VIEW_DISTANCE
+		SORT_CUSTOM_ASCENDING
+		SORT_CUSTOM_DESCENDING
+------------------------------------------------------------------------------*/
+
+#include "/Engine/Private/Common.ush"
+
+#define NIAGARA_KEYGEN 1
+#include "NiagaraVFParticleAccess.usf"
+
+Buffer<uint> GPUParticleCountBuffer;
+
+// (ParticleCount, GPUParticleCountOffset)
+uint2 ParticleCountParams;
+
+// (EmitterKey, OutputOffset, SortMode, SortAttributeOffset)
+uint4 SortParams;
+
+float3 CameraPosition;
+float3 CameraDirection;
+
+/** Output key buffer. */
+RWBuffer<uint> OutKeys;
+/** Output indices buffer. */
+RWBuffer<int> OutParticleIndices;
+
+#if SORT_MAX_PRECISION
+// (SortKeyMask, SortKeyShift, SortKeySignBit)
+uint4 SortKeyParams;
+#endif
+
+[numthreads(THREAD_COUNT,1,1)]
+void GenerateParticleSortKeys(uint ParticleIndex : SV_DispatchThreadID)
+{
+	const uint ParticleCount = ParticleCountParams.x;
+	const uint GPUParticleCountOffset = ParticleCountParams.y;
+
+	const uint GPUParticleCount = GPUParticleCountOffset != -1 ? GPUParticleCountBuffer[GPUParticleCountOffset] : ParticleCount;
+
+	const uint EmitterKey = SortParams.x;
+	const uint OutputOffset = SortParams.y;
+	const uint SortMode = SortParams.z;
+	const uint SortAttributeOffset = SortParams.w;
+
+#if SORT_MAX_PRECISION
+	const uint SortKeyMask = SortKeyParams.x; 
+	const uint SortKeyShift = SortKeyParams.y;
+	const uint SortKeySignBit = SortKeyParams.z;
+#else
+	const uint SortKeyMask = 0xffff; 
+	const uint SortKeySignBit = 0x8000;
+#endif
+
+	if (ParticleIndex < ParticleCount)
+	{
+		float SortValue = 0;
+		uint SortKey = SortKeyMask; 
+
+		if (ParticleIndex < GPUParticleCount)
+		{
+			if (SortMode == SORT_VIEW_DEPTH || SortMode == SORT_VIEW_DISTANCE)
+			{ 
+				float3 ParticlePosition = GetVec3(SortAttributeOffset, ParticleIndex); 
+				float3 CameraVector = ParticlePosition - CameraPosition;
+				SortValue = SortMode == SORT_VIEW_DEPTH ? dot(CameraVector, CameraDirection) : length(CameraVector);
+				
+#if SORT_MAX_PRECISION
+				SortKey = (asuint(SortValue)) >> SortKeyShift;
+#else
+				SortKey = f32tof16(SortValue);
+#endif
+				// Render far away first, knowing that SortValue > 0
+				SortKey =  (~SortKey) & SortKeyMask;
+			}
+			else
+			{
+				SortValue = GetFloat(SortAttributeOffset, ParticleIndex);
+
+#if SORT_MAX_PRECISION
+				SortKey = (asuint(SortValue)) >> SortKeyShift;
+#else
+				SortKey = f32tof16(SortValue);
+#endif
+				// Special handling for negative values.
+				SortKey = SortValue >= 0 ? (SortKey | SortKeySignBit) : ~SortKey;
+
+				SortKey = (SortMode == SORT_CUSTOM_ASCENDING ? SortKey : ~SortKey) & SortKeyMask;
+			}
+
+			// Clamp to max-1 to make sure that out-of-bound indices are at the end.
+			SortKey = min(SortKey , SortKeyMask - 1);
+		}
+
+		SortKey |= EmitterKey;
+
+		OutKeys[OutputOffset + ParticleIndex] = SortKey;
+		OutParticleIndices[OutputOffset + ParticleIndex] = ParticleIndex;
+	}
+}
+ 
