Copying //UE4/Dev-Niagara to Dev-Main (//UE4/Dev-Main)

diff --git a/Engine/Plugins/FX/Niagara/Shaders/Private/NiagaraSpriteVertexFactory.ush b/Engine/Plugins/FX/Niagara/Shaders/Private/NiagaraSpriteVertexFactory.ush
new file mode 100644
--- /dev/null
+++ b/Engine/Plugins/FX/Niagara/Shaders/Private/NiagaraSpriteVertexFactory.ush
@@ -0,0 +1,1120 @@
+// Copyright 1998-2019 Epic Games, Inc. All Rights Reserved.
+
+/*=============================================================================
+	ParticleSpriteVertexFactory.hlsl: Particle vertex factory shader code.
+	Shared by standard sprite particles and SubUV sprite particles.
+=============================================================================*/
+
+#include "/Engine/Private/VertexFactoryCommon.ush"
+#include "/Engine/Private/ParticleVertexFactoryCommon.ush"
+#include "NiagaraVFParticleAccess.usf"
+
+// These are only enabled for SM5 due to limited interpolators elsewhere
+#define USE_PARTICLE_LIGHTING_OFFSET (FEATURE_LEVEL >= FEATURE_LEVEL_SM5 && !MATERIAL_SHADINGMODEL_UNLIT)
+#define USE_PARTICLE_POSITION (FEATURE_LEVEL >= FEATURE_LEVEL_SM5 && NEEDS_PARTICLE_POSITION)
+#define USE_PARTICLE_VELOCITY (FEATURE_LEVEL >= FEATURE_LEVEL_SM5 && NEEDS_PARTICLE_VELOCITY)
+#define USE_PARTICLE_TIME (FEATURE_LEVEL >= FEATURE_LEVEL_SM5 && NEEDS_PARTICLE_TIME)
+#define USE_PARTICLE_SIZE (FEATURE_LEVEL >= FEATURE_LEVEL_SM5 && NEEDS_PARTICLE_SIZE)
+#define USE_PARTICLE_RANDOM (FEATURE_LEVEL >= FEATURE_LEVEL_SM5 && NEEDS_PARTICLE_RANDOM)
+
+// Must keep these in sync with ENiagaraSpriteAlignment in NiagaraSpriteRendererProperties.h
+/** Only Particles.SpriteRotation and FacingMode impact the alignment of the particle.*/
+#define ALIGNMENT_MODE_UNALIGNED 0
+/** Imagine the particle texture having an arrow pointing up, this mode makes the arrow point in the direction of the Particles.Velocity attribute. Particles.SpriteRotation rotation is ignored.*/
+#define ALIGNMENT_MODE_VELOCITYALIGNED 1
+/** Imagine the particle texture having an arrow pointing up, this mode makes the arrow point towards the axis defined by the Particles.SpriteAlignment attribute, after taking Particles.SpriteRotation and the FacingMode into account.*/
+#define ALIGNMENT_MODE_CUSTOMALIGNMENT 2
+
+/** The sprite billboard origin is always "looking at" the camera origin, trying to keep its up axis aligned to the camera's up axis. */
+#define FACING_MODE_FACECAMERA 0
+/** The sprite billboard plane is completely parallel to the camera plane. Particle always looks "flat" */
+#define FACING_MODE_FACECAMERAPLANE 1
+/** The sprite billboard faces toward the "Particles.SpriteFacing" vector attribute, using the per-axis CustomFacingVectorMask as a lerp factor from the standard FaceCamera mode.*/
+#define FACING_MODE_CUSTOMFACINGVECTOR 2
+/** Faces the camera position, but is not dependent on the camera rotation.  This method produces more stable particles under camera rotation. Uses the up axis of (0,0,1).*/
+#define FACING_MODE_FACECAMERAPOSITION 3 
+/** Blends between FaceCamera and FaceCameraPosition.*/
+#define FACING_MODE_FACECAMERADISTANCEBLEND 4
+
+struct FVertexFactoryInput
+{
+	float2	TexCoord				: ATTRIBUTE0;
+	uint VertexId					: SV_VertexID;
+	uint InstanceId					: SV_InstanceID;
+};
+
+// RHI_RAYTRACING
+#if COMPUTESHADER || RAYHITGROUPSHADER
+FVertexFactoryInput LoadVertexFactoryInput(uint TriangleIndex, int VertexIndex)
+{
+	FVertexFactoryInput Input;
+	
+	uint IndexBuffer[6] = { 0, 2, 3, 0, 1, 2 };
+	uint VertexId = IndexBuffer[(TriangleIndex * 3 + VertexIndex) % 6];
+	float2 TexCoords[4] = { float2(0.0f, 0.0f), float2(0.0f, 1.0f), float2(1.0f, 1.0f), float2(1.0f, 0.0f) };
+	Input.TexCoord = TexCoords[VertexId];
+	Input.VertexId = VertexId;
+	Input.InstanceId = TriangleIndex / 2;
+
+	return Input;
+}
+uint GetNumRayTracingDynamicMeshVerticesIndirect()
+{
+	return 6 * NiagaraSpriteVFLooseParameters.IndirectArgsBuffer[1];
+}
+#endif
+
+struct FVertexFactoryInterpolantsVSToPS
+{
+	// First row of the tangent to world matrix, Interp_Sizer used by SUBUV_PARTICLES in w
+	float4	TangentToWorld0AndInterp_Sizer : TANGENTTOWORLD0;
+	// Last row of the tangent to world matrix in xyz
+	float4	TangentToWorld2	: TANGENTTOWORLD2;
+	
+#if (DYNAMIC_PARAMETERS_MASK & 1)
+	float4 DynamicParameter	: PARTICLE_DYNAMIC_PARAM0;
+#endif
+#if (DYNAMIC_PARAMETERS_MASK & 2)
+	float4 DynamicParameter1	: PARTICLE_DYNAMIC_PARAM1;
+#endif
+#if (DYNAMIC_PARAMETERS_MASK & 4)
+	float4 DynamicParameter2	: PARTICLE_DYNAMIC_PARAM2;
+#endif
+#if (DYNAMIC_PARAMETERS_MASK & 8)
+	float4 DynamicParameter3	: PARTICLE_DYNAMIC_PARAM3;
+#endif
+
+#if NEEDS_PARTICLE_COLOR
+	float4	Color				: TEXCOORD0;
+#endif
+
+#if NUM_TEX_COORD_INTERPOLATORS
+	float4	TexCoords[(NUM_TEX_COORD_INTERPOLATORS + 1) / 2]	: TEXCOORD1;
+#endif
+
+	//Not sure this is actually being used now and it's awkward to slot in now we're supporting custom UVs so I'm just giving this its own interpolant.
+#if LIGHTMAP_UV_ACCESS
+	float2 LightMapUVs : LIGHTMAP_UVS;
+#endif
+
+#if USE_PARTICLE_SUBUVS
+	float4 ParticleSubUVs : PARTICLE_SUBUVS;
+#endif
+
+#if USE_PARTICLE_POSITION
+	/** Cam-relative (translated) particle center and radius */
+	float4 ParticleTranslatedWorldPositionAndSize	: PARTICLE_POSITION;
+#endif
+
+#if USE_PARTICLE_VELOCITY
+	float4 ParticleVelocity : PARTICLE_VELOCITY;
+#endif
+
+#if USE_PARTICLE_TIME
+	float RelativeTime : PARTICLE_TIME;
+#endif
+
+#if USE_PARTICLE_LIGHTING_OFFSET
+	float3 LightingPositionOffset : PARTICLE_LIGHTING_OFFSET;
+#endif
+
+#if USE_PARTICLE_SIZE
+	float2 ParticleSize : PARTICLE_SIZE;
+#endif
+
+#if USE_PARTICLE_RANDOM
+	float ParticleRandom : PARTICLE_RANDOM;
+#endif
+
+#if INSTANCED_STEREO
+	nointerpolation uint EyeIndex : PACKED_EYE_INDEX;
+#endif
+};
+
+struct FVertexFactoryIntermediates
+{
+	/** The position of the particle in non-translated world space. */
+	float3 ParticleWorldPosition;
+	/** The position of the vertex in translated world space. */
+	float3 VertexWorldPosition;
+	/** Last frame's position of the vertex in translated world space. */
+	float3 PreviousVertexWorldPosition;
+	/** Particle translated world space position and size. */
+	float4 TranslatedWorldPositionAndSize;
+	
+#if USE_PARTICLE_LIGHTING_OFFSET
+	float3 LightingPositionOffset;
+#endif
+
+	/** The texture coordinate at this vertex. */
+	float4 TexCoord;
+#if NUM_MATERIAL_TEXCOORDS >= 3 //Unflipped UVs are placed in UVs 2-3
+	/** A second UV set. Always non-UV flipped. Allows use of UV flipping with Normal maps etc. */
+	float4 TexCoord_Unflipped;
+#endif
+	/** The sprite tangent in world space (+V). */
+	float3 TangentUp;
+	/** The sprite tangent in world space (+U). */
+	float3 TangentRight;
+	/** The color of the sprite. */
+	float4 Color;
+	/** The velocity of the particle, XYZ: direction, W: speed. */
+	float4 ParticleVelocity;
+
+	float3 Position;
+	float3 OldPosition;
+	
+	float4 SizeRotSubImage;
+	
+	float3 CustomFacingVector;
+	float3 CustomAlignmentVector;
+
+#if (DYNAMIC_PARAMETERS_MASK & 1)
+	float4 DynamicParameter;
+#endif
+#if (DYNAMIC_PARAMETERS_MASK & 2)
+	float4 DynamicParameter1;
+#endif
+#if (DYNAMIC_PARAMETERS_MASK & 4)
+	float4 DynamicParameter2;
+#endif
+#if (DYNAMIC_PARAMETERS_MASK & 8)
+	float4 DynamicParameter3;
+#endif
+
+	/** The sub-image lerp. */
+	float SubImageLerp;
+
+#if USE_PARTICLE_TIME
+	/** Relative time. */
+	float RelativeTime;
+#endif
+
+	/** Transform from tangent space to local space. */
+	float3x3 TangentToLocal;
+	/** Size of the particle */
+	float2 ParticleSize;
+
+#if USE_PARTICLE_RANDOM
+	/** Particle Random value */
+	float ParticleRandom;
+#endif
+};
+
+#if NUM_TEX_COORD_INTERPOLATORS
+bool UVIndexUseZW(int UVIndex)
+{
+	#if COMPILER_GLSL_ES2
+	return (frac((float)UVIndex / 2.0f) == 0.5f);
+	#else
+	return (UVIndex % 2) != 0;
+	#endif
+}
+
+float2 GetUV(FVertexFactoryInterpolantsVSToPS Interpolants, int UVIndex)
+{
+	float4 UVVector = Interpolants.TexCoords[UVIndex / 2];
+	return UVIndexUseZW(UVIndex) ? UVVector.zw : UVVector.xy;
+}
+
+void SetUV(inout FVertexFactoryInterpolantsVSToPS Interpolants, int UVIndex, float2 InValue)
+{
+	FLATTEN
+	if (UVIndexUseZW(UVIndex))
+	{
+		Interpolants.TexCoords[UVIndex / 2].zw = InValue;
+	}
+	else
+	{
+		Interpolants.TexCoords[UVIndex / 2].xy = InValue;
+	}
+}
+#endif
+
+/** Converts from vertex factory specific interpolants to a FMaterialPixelParameters, which is used by material inputs. */
+FMaterialPixelParameters GetMaterialPixelParameters(FVertexFactoryInterpolantsVSToPS Interpolants, float4 SvPosition)
+{
+	// GetMaterialPixelParameters is responsible for fully initializing the result
+	FMaterialPixelParameters Result = MakeInitializedMaterialPixelParameters();
+
+#if NUM_TEX_COORD_INTERPOLATORS
+	UNROLL
+	for (int CoordinateIndex = 0; CoordinateIndex < NUM_TEX_COORD_INTERPOLATORS; CoordinateIndex++)
+	{
+		Result.TexCoords[CoordinateIndex] = GetUV(Interpolants, CoordinateIndex);
+	}
+#endif	// NUM_TEX_COORD_INTERPOLATORS
+
+	Result.VertexColor = 1;
+
+	half4 TangentToWorld0 = Interpolants.TangentToWorld0AndInterp_Sizer;
+	half4 TangentToWorld2 = Interpolants.TangentToWorld2;
+	float SubImageLerp = Interpolants.TangentToWorld0AndInterp_Sizer.w;
+
+	#if GENERATE_SPHERICAL_PARTICLE_NORMALS && USE_PARTICLE_POSITION
+		{	
+			// can be optimized
+			float4 ScreenPosition = SvPositionToResolvedScreenPosition(SvPosition);
+			float3 TranslatedWorldPosition = SvPositionToResolvedTranslatedWorld(SvPosition);
+			Result.TangentToWorld = GetSphericalParticleNormal(TranslatedWorldPosition, Interpolants.ParticleTranslatedWorldPositionAndSize.xyz, Interpolants.ParticleTranslatedWorldPositionAndSize.w);
+		}
+	#else
+		Result.TangentToWorld = AssembleTangentToWorld(TangentToWorld0.xyz, TangentToWorld2);
+	#endif
+
+	Result.UnMirrored = 1;
+	Result.Particle.MacroUV = NiagaraSpriteVF.MacroUVParameters;
+
+#if NEEDS_PARTICLE_COLOR
+	Result.Particle.Color = Interpolants.Color;
+#endif
+
+#if (DYNAMIC_PARAMETERS_MASK != 0)
+	Result.Particle.DynamicParameterValidMask = NiagaraSpriteVF.MaterialParamValidMask;
+#endif
+#if (DYNAMIC_PARAMETERS_MASK & 1)
+	Result.Particle.DynamicParameter = Interpolants.DynamicParameter;
+#endif
+#if (DYNAMIC_PARAMETERS_MASK & 2)
+	Result.Particle.DynamicParameter1 = Interpolants.DynamicParameter1;
+#endif
+#if (DYNAMIC_PARAMETERS_MASK & 4)
+	Result.Particle.DynamicParameter2 = Interpolants.DynamicParameter2;
+#endif
+#if (DYNAMIC_PARAMETERS_MASK & 8)
+	Result.Particle.DynamicParameter3 = Interpolants.DynamicParameter3;
+#endif
+
+#if USE_PARTICLE_POSITION
+	Result.Particle.TranslatedWorldPositionAndSize = Interpolants.ParticleTranslatedWorldPositionAndSize;
+#endif
+
+#if USE_PARTICLE_VELOCITY
+	Result.Particle.Velocity = Interpolants.ParticleVelocity;
+#endif
+
+#if USE_PARTICLE_TIME
+	Result.Particle.RelativeTime = Interpolants.RelativeTime;
+#endif
+
+	Result.Particle.MotionBlurFade = 1.0f;
+
+
+#if USE_PARTICLE_LIGHTING_OFFSET
+	Result.LightingPositionOffset = Interpolants.LightingPositionOffset;
+#endif
+
+#if USE_PARTICLE_SUBUVS
+	Result.Particle.SubUVCoords[0] = Interpolants.ParticleSubUVs.xy;
+	Result.Particle.SubUVCoords[1] = Interpolants.ParticleSubUVs.zw;
+	Result.Particle.SubUVLerp = SubImageLerp;
+#endif
+
+#if LIGHTMAP_UV_ACCESS
+	Result.LightmapUVs = Interpolants.LightMapUVs;
+#endif
+
+#if USE_PARTICLE_SIZE
+	Result.Particle.Size = Interpolants.ParticleSize;
+#endif
+
+#if USE_PARTICLE_RANDOM
+	Result.Particle.Random = Interpolants.ParticleRandom;
+#else
+	Result.Particle.Random = 0.0f;
+#endif
+
+	Result.TwoSidedSign = 1;
+	return Result;
+}
+
+/** Converts from vertex factory specific input to a FMaterialVertexParameters, which is used by vertex shader material inputs. */
+FMaterialVertexParameters GetMaterialVertexParameters(FVertexFactoryInput Input, FVertexFactoryIntermediates Intermediates, float3 WorldPosition, float3x3 TangentToLocal)
+{
+	FMaterialVertexParameters Result = (FMaterialVertexParameters)0;
+	Result.WorldPosition = WorldPosition;
+	Result.VertexColor = Intermediates.Color;
+	Result.TangentToWorld = mul(TangentToLocal, GetLocalToWorld3x3()); 
+	Result.Particle.MacroUV = NiagaraSpriteVF.MacroUVParameters;
+	Result.Particle.Color = Intermediates.Color;
+	Result.Particle.MotionBlurFade = 1.0f;
+	Result.PreSkinnedPosition = Intermediates.Position.xyz;
+	Result.PreSkinnedNormal = TangentToLocal[2].xyz;
+
+	// Previous frame not handled deliberately. Lacks necessary information and
+	// primitives using this VF are usually transparent and hence don't output velocity
+	Result.PrevFrameLocalToWorld = Primitive.LocalToWorld;
+
+#if (DYNAMIC_PARAMETERS_MASK != 0)
+	Result.Particle.DynamicParameterValidMask = NiagaraSpriteVF.MaterialParamValidMask;
+#endif
+#if (DYNAMIC_PARAMETERS_MASK & 1)
+	Result.Particle.DynamicParameter = Intermediates.DynamicParameter;
+#endif
+#if (DYNAMIC_PARAMETERS_MASK & 2)
+	Result.Particle.DynamicParameter1 = Intermediates.DynamicParameter1;
+#endif
+#if (DYNAMIC_PARAMETERS_MASK & 4)
+	Result.Particle.DynamicParameter2 = Intermediates.DynamicParameter2;
+#endif
+#if (DYNAMIC_PARAMETERS_MASK & 8)
+	Result.Particle.DynamicParameter3 = Intermediates.DynamicParameter3;
+#endif
+
+#if USE_PARTICLE_POSITION
+	Result.Particle.TranslatedWorldPositionAndSize = Intermediates.TranslatedWorldPositionAndSize;
+#endif
+
+#if USE_PARTICLE_VELOCITY
+	Result.Particle.Velocity = Intermediates.ParticleVelocity;
+#endif
+
+#if USE_PARTICLE_TIME
+	Result.Particle.RelativeTime = Intermediates.RelativeTime;
+#endif
+
+#if USE_PARTICLE_RANDOM
+	Result.Particle.Random = Intermediates.ParticleRandom;
+#else
+	Result.Particle.Random = 0.0f;
+#endif
+
+#if NUM_MATERIAL_TEXCOORDS_VERTEX
+	#if NUM_MATERIAL_TEXCOORDS_VERTEX >= 1
+		Result.TexCoords[0] = Intermediates.TexCoord.xy;
+	#if NUM_MATERIAL_TEXCOORDS_VERTEX >= 2
+		Result.TexCoords[1] = Intermediates.TexCoord.zw;
+	#if NUM_MATERIAL_TEXCOORDS_VERTEX >= 3
+		Result.TexCoords[2] = Intermediates.TexCoord_Unflipped.xy;
+	#if NUM_MATERIAL_TEXCOORDS_VERTEX >= 4
+		Result.TexCoords[3] = Intermediates.TexCoord_Unflipped.zw;
+	#endif	// >= 4
+	#endif	// >= 3
+	#endif	// >= 2
+	#endif	// >= 1
+#endif	//NUM_MATERIAL_TEXCOORDS_VERTEX
+
+#if USE_PARTICLE_SUBUVS
+	Result.Particle.SubUVCoords[0] = Intermediates.TexCoord.xy;
+	Result.Particle.SubUVCoords[1] = Intermediates.TexCoord.zw;
+#endif
+
+#if USE_PARTICLE_SIZE
+	Result.Particle.Size = Intermediates.ParticleSize;
+#endif
+	return Result;
+}
+
+float3 SafeNormalize(float3 V)
+{
+	// NOTE: Use a very low threshold, close to machine-epsilon when used on near-unit-length vectors. 
+	//       This is primarily to avoid division-by-zero, so we can use a small constant and let things blow up slightly
+	//       on very small vectors, which won't work very well anyway. 
+	float L = length(V);
+	return V = V / max(L, 0.000001); 
+}
+
+float3 SafeNormalizeWithFallback(float3 V, float3 Fallback)
+{
+	float L = length(V);
+	return L > 0.000001 ? V / L : Fallback; 
+}
+
+void GetTangents(FVertexFactoryInput Input, FVertexFactoryIntermediates Intermediates, float3 WorldPosition, float3 OldWorldPosition, out float3 OutRight, out float3 OutUp)
+{
+	// Select camera up/right vectors.
+	float3 ResolvedViewRight = lerp(ResolvedView.ViewRight, ResolvedView.HMDViewNoRollRight, NiagaraSpriteVF.RemoveHMDRoll);
+	float3 ResolvedViewUp = lerp(ResolvedView.ViewUp, ResolvedView.HMDViewNoRollUp, NiagaraSpriteVF.RemoveHMDRoll);
+
+	// The standard camera facing vectors
+	float3 CameraRight = ResolvedViewRight;
+	float3 CameraUp = -ResolvedViewUp;
+	float3 CameraFacingVector = -ResolvedView.ViewForward;
+
+	// Determine the vector from the particle to the camera and the particle's movement direction.
+	float3 CameraVector = ResolvedView.WorldCameraOrigin.xyz - WorldPosition.xyz;
+	float3 CameraDirection = SafeNormalize(CameraVector);
+	float3 ParticleDirection	= SafeNormalizeWithFallback(WorldPosition.xyz - OldWorldPosition.xyz, float3(0,0,1));
+
+	// First interpolate between the base facing mode and the custom facing mode. We assume that the code setting up this call has the facing vector mask
+	// set to 0,0,0,0 if not in FACING_MODE_CUSTOMFACINGVECTOR.
+	float3 CustomFacingVectorMask =  NiagaraSpriteVF.CustomFacingVectorMask;
+	float3 CustomFacingVector = Intermediates.CustomFacingVector * CustomFacingVectorMask + CameraDirection * (float3(1.0,1.0,1.0)-CustomFacingVectorMask);
+	
+	// the particle's movement direction.
+	float3 DirectionVector	= lerp(ParticleDirection, Intermediates.CustomAlignmentVector, NiagaraSpriteVFLooseParameters.ParticleAlignmentMode == ALIGNMENT_MODE_CUSTOMALIGNMENT);
+	
+	// Tangent vectors for any non-default alignment mode.
+	float3 CameraRightAligned	= SafeNormalize(cross(CustomFacingVector, DirectionVector));
+	float3 CameraUpAligned = lerp(-DirectionVector, cross(CustomFacingVector, CameraRightAligned), NiagaraSpriteVFLooseParameters.ParticleFacingMode == FACING_MODE_CUSTOMFACINGVECTOR);
+
+	// Tangent vectors for camera plane facing mode
+	float3 CameraRightPlaneFacing = SafeNormalize(cross(CameraFacingVector,ResolvedViewUp));
+	float3 CameraUpPlaneFacing = cross(CameraFacingVector, CameraRightPlaneFacing);
+
+	// Tangent vectors for camera custom facing mode
+	float3 CameraRightCustomFacing = SafeNormalize(cross(CustomFacingVector,ResolvedViewUp));
+	float3 CameraUpCustomFacing = cross(CustomFacingVector, CameraRightCustomFacing);
+
+	// Tangent vectors for camera facing position.
+	float3 CameraRightPosition = SafeNormalize(cross(CameraDirection,float3(0,0,1)));
+	float3 CameraUpPosition = cross(CameraDirection, CameraRightPosition);
+
+	float3 Right;
+	float3 Up;
+	
+	/**This could well be better as a branch*/
+	BRANCH
+	if(NiagaraSpriteVFLooseParameters.ParticleAlignmentMode == ALIGNMENT_MODE_UNALIGNED)
+	{
+		FLATTEN
+		if (NiagaraSpriteVFLooseParameters.ParticleFacingMode == FACING_MODE_FACECAMERADISTANCEBLEND)
+		{
+			float CameraDistanceSq = dot(CameraVector, CameraVector);
+			float Interp = saturate(CameraDistanceSq * NiagaraSpriteVF.CameraFacingBlend.y - NiagaraSpriteVF.CameraFacingBlend.z);
+			Right = cross(CameraDirection,float3(0,0,1));
+			Up = cross(CameraDirection, Right);
+			Right = lerp(Right, CameraRight.xyz, Interp);
+			Up = lerp(Up, CameraUp.xyz, Interp);
+		} 
+		else if(NiagaraSpriteVFLooseParameters.ParticleFacingMode == FACING_MODE_FACECAMERA)
+		{
+			Right = CameraRight.xyz;
+			Up = CameraUp.xyz;
+		}		 
+		else if(NiagaraSpriteVFLooseParameters.ParticleFacingMode == FACING_MODE_FACECAMERAPLANE)
+		{
+			Right = CameraRightPlaneFacing;
+			Up = CameraUpPlaneFacing;
+		}
+		else if(NiagaraSpriteVFLooseParameters.ParticleFacingMode == FACING_MODE_CUSTOMFACINGVECTOR)
+		{
+			Right = CameraRightCustomFacing;
+			Up = CameraUpCustomFacing;
+		}
+		else//(NiagaraSpriteVFLooseParameters.ParticleFacingMode == FACING_MODE_FACECAMERAPOSITION)
+		{
+			Right = CameraRightPosition;
+			Up = CameraUpPosition;
+		}
+	}
+	else 
+	{	
+		Right = CameraRightAligned;
+		Up = CameraUpAligned;
+	}
+
+
+	// Determine the angle of rotation.
+	float SinRotation;
+	float CosRotation;
+	const float SpriteRotation = Intermediates.SizeRotSubImage.z * NiagaraSpriteVF.RotationScale + NiagaraSpriteVF.RotationBias;
+	sincos(SpriteRotation, SinRotation, CosRotation);
+
+	float3 SR = Right * SinRotation;
+	float3 SU = Up * SinRotation;
+	float3 CR = Right * CosRotation;
+	float3 CU = Up * CosRotation;
+	OutRight = SU + CR;
+	OutUp = CU - SR;
+}
+
+/** derive basis vectors */
+float3x3 CalcTangentBasis(FVertexFactoryIntermediates Intermediates)
+{
+	// Using camera facing TangentX and TangentY.  The resulting tangent basis is not orthonormal with anything other than ENM_CameraFacing, 
+	// So there are artifacts with tangent space calculations like the TransformVector node,
+	// But this allows lighting based on a world space shape via the normal while still having normal maps camera aligned.
+	float3x3 Result;
+	Result[0] = Intermediates.TangentRight;
+	Result[1] = Intermediates.TangentUp;
+
+	// ENM_CameraFacing
+	//@todo - use static branching
+	if (NiagaraSpriteVF.NormalsType < .5f)
+	{
+		Result[2] = normalize(cross(Result[0],Result[1]));
+	} 
+	// ENM_Spherical
+	else if (NiagaraSpriteVF.NormalsType < 1.5f)
+	{
+		float3 TangentZ = normalize(Intermediates.ParticleWorldPosition.xyz - NiagaraSpriteVF.NormalsSphereCenter.xyz);
+		Result[2] = TangentZ;	
+	}
+	// ENM_Cylindrical
+	else
+	{
+		float3 ClosestPointOnCylinder = NiagaraSpriteVF.NormalsSphereCenter.xyz + dot(NiagaraSpriteVF.NormalsCylinderUnitDirection.xyz, Intermediates.ParticleWorldPosition.xyz - NiagaraSpriteVF.NormalsSphereCenter.xyz) * NiagaraSpriteVF.NormalsCylinderUnitDirection.xyz;
+		float3 TangentZ = normalize(Intermediates.ParticleWorldPosition.xyz - ClosestPointOnCylinder);
+		Result[2] = TangentZ;	
+	}
+
+	return Result;
+}
+
+void ComputeBillboardUVs(FVertexFactoryInput Input, FVertexFactoryIntermediates Intermediates, out float2 UVForPosition, out float2 UVForTexturing, out float2 UVForTexturingUnflipped)
+{
+	// Encoding the UV flip in the sign of the size data.
+	float2 UVFlip = sign(Intermediates.SizeRotSubImage.xy);
+
+#if FEATURE_LEVEL >= FEATURE_LEVEL_SM4
+	BRANCH
+	if (NiagaraSpriteVFLooseParameters.NumCutoutVerticesPerFrame > 0)
+	{
+		// Avoid uint divide which is extremely slow on GCN
+		const uint VertexIndexMask = NiagaraSpriteVFLooseParameters.NumCutoutVerticesPerFrame == 8 ? 7 : 3; // BVC_FourVertices || BVC_EightVertices
+		uint CutoutVertexIndex = Input.VertexId & VertexIndexMask;
+		float NumFrames = NiagaraSpriteVF.SubImageSize.x * NiagaraSpriteVF.SubImageSize.y;
+		uint SubImageIndexInt = (uint)fmod(Intermediates.SizeRotSubImage.w, NumFrames);
+
+		FLATTEN
+		if (UVFlip.x * UVFlip.y < 0)
+		{
+			// Reverse the winding order of the polygon when only flipping in one direction to counteract UV flipping
+			CutoutVertexIndex = NiagaraSpriteVFLooseParameters.NumCutoutVerticesPerFrame - 1 - CutoutVertexIndex;
+		}
+
+		// Fetch the billboard space positions from the appropriate frame of the preprocessed cutout geometry
+		UVForTexturing = UVForTexturingUnflipped = NiagaraSpriteVFLooseParameters.CutoutGeometry[SubImageIndexInt * NiagaraSpriteVFLooseParameters.NumCutoutVerticesPerFrame + CutoutVertexIndex];
+
+		// Invert positions on the billboard so that the cutout geometry still contains the visible portion
+		// This changes the winding order when only one direction is flipped
+		UVForPosition.x = UVFlip.x < 0.0 ? 1.0 - UVForTexturing.x : UVForTexturing.x;
+		UVForPosition.y = UVFlip.y < 0.0 ? 1.0 - UVForTexturing.y : UVForTexturing.y;
+	}
+	else
+#endif
+	{
+		UVForTexturing.x = UVFlip.x < 0.0 ? 1.0 - Input.TexCoord.x : Input.TexCoord.x;
+		UVForTexturing.y = UVFlip.y < 0.0 ? 1.0 - Input.TexCoord.y : Input.TexCoord.y;
+
+		// Note: not inverting positions, as that would change the winding order
+		UVForPosition = UVForTexturingUnflipped = Input.TexCoord.xy;
+	}
+}
+
+float3 TransformPosition(float3 InPosition)
+{
+	return mul(float4(InPosition,1.0), NiagaraSpriteVF.LocalToWorld).xyz;
+}
+
+float3 TransformVector(float3 InVector)
+{
+	return mul(float4(InVector,0.0), NiagaraSpriteVF.LocalToWorldInverseTransposed).xyz;
+}
+
+/*********************************************************************************
+ *							Particle specific
+ *********************************************************************************/
+
+float3 GetNiagaraParticlePosition(uint InstanceID)
+{
+	if(NiagaraSpriteVF.PositionDataOffset  == -1)
+	{
+		// See FNiagaraConstants for a synchronized value...
+		return TransformPosition(NiagaraSpriteVF.DefaultPos.xyz);
+	}
+	return TransformPosition(GetVec3(NiagaraSpriteVF.PositionDataOffset, InstanceID));
+}
+float3 GetNiagaraParticleVelocity(uint InstanceID)
+{
+	if(NiagaraSpriteVF.VelocityDataOffset  == -1)
+	{
+		// See FNiagaraConstants for a synchronized value...
+		return TransformVector(float3(0.0f,0.0f,0.0f));
+	}
+	return TransformVector(GetVec3(NiagaraSpriteVF.VelocityDataOffset, InstanceID));
+}
+float2 GetNiagaraParticleSize(uint InstanceID)
+{
+	if(NiagaraSpriteVF.SizeDataOffset  == -1)
+	{
+		// See FNiagaraConstants for a synchronized value...
+		return float2(50.0f, 50.0f);
+	}
+	return GetVec2(NiagaraSpriteVF.SizeDataOffset, InstanceID);
+}
+float GetNiagaraParticleRotation(uint InstanceID)
+{
+	if(NiagaraSpriteVF.RotationDataOffset  == -1)
+	{
+		// See FNiagaraConstants for a synchronized value...
+		return 0.0f;
+	}
+	return (GetFloat(NiagaraSpriteVF.RotationDataOffset, InstanceID) / 180.0f) * PI;
+}
+float4 GetNiagaraParticleColor(uint InstanceID)
+{
+	if(NiagaraSpriteVF.ColorDataOffset  == -1)
+	{
+		// See FNiagaraConstants for a synchronized value...
+		return float4(1.0f, 1.0f, 1.0f, 1.0f);
+	}
+	return GetVec4(NiagaraSpriteVF.ColorDataOffset, InstanceID);
+}
+
+float2 GetNiagaraUVScale(uint InstanceID)
+{
+	if(NiagaraSpriteVF.UVScaleDataOffset  == -1)
+	{
+		// See FNiagaraConstants for a synchronized value...
+		return float2(1.0f, 1.0f);
+	}
+	return GetVec2(NiagaraSpriteVF.UVScaleDataOffset, InstanceID);	
+}
+
+float GetNiagaraParticleRandom(uint InstanceID)
+{
+	if(NiagaraSpriteVF.MaterialRandomDataOffset  == -1)
+	{
+		// See FNiagaraConstants for a synchronized value...
+		return 0.0f;
+	}
+	return GetFloat(NiagaraSpriteVF.MaterialRandomDataOffset, InstanceID);	
+}
+
+float GetNiagaraCameraOffset(uint InstanceID)
+{
+	if (NiagaraSpriteVF.CameraOffsetDataOffset == -1)
+	{
+		// See FNiagaraConstants for a synchronized value...
+		return 0.0f;
+	}
+	return GetFloat(NiagaraSpriteVF.CameraOffsetDataOffset, InstanceID);
+}
+
+#if USE_PARTICLE_TIME
+float GetNiagaraNormalizedAge(uint InstanceID)
+{
+	if (NiagaraSpriteVF.NormalizedAgeDataOffset == -1)
+	{
+		return 0;
+	}
+	return GetFloat(NiagaraSpriteVF.NormalizedAgeDataOffset, InstanceID);
+}
+#endif
+
+#if (DYNAMIC_PARAMETERS_MASK & 1)
+float4 GetNiagaraParticleDynamicParameters(uint InstanceID)
+{
+	return NiagaraSpriteVF.MaterialParamDataOffset != -1 ? GetVec4(NiagaraSpriteVF.MaterialParamDataOffset, InstanceID) : float4(1.0f, 1.0f, 1.0f, 1.0f);;
+}
+#endif
+#if (DYNAMIC_PARAMETERS_MASK & 2)
+float4 GetNiagaraParticleDynamicParameters1(uint InstanceID)
+{
+	return NiagaraSpriteVF.MaterialParam1DataOffset != -1 ? GetVec4(NiagaraSpriteVF.MaterialParam1DataOffset, InstanceID) : float4(1.0f, 1.0f, 1.0f, 1.0f);;
+}
+#endif
+#if (DYNAMIC_PARAMETERS_MASK & 4)
+float4 GetNiagaraParticleDynamicParameters2(uint InstanceID)
+{
+	return NiagaraSpriteVF.MaterialParam2DataOffset != -1 ? GetVec4(NiagaraSpriteVF.MaterialParam2DataOffset, InstanceID) : float4(1.0f, 1.0f, 1.0f, 1.0f);;
+}
+#endif
+#if (DYNAMIC_PARAMETERS_MASK & 8)
+float4 GetNiagaraParticleDynamicParameters3(uint InstanceID)
+{
+	return NiagaraSpriteVF.MaterialParam3DataOffset != -1 ? GetVec4(NiagaraSpriteVF.MaterialParam3DataOffset, InstanceID) : float4(1.0f, 1.0f, 1.0f, 1.0f);;
+}
+#endif
+
+
+float GetNiagaraParticleSubimage(uint InstanceID)
+{
+	// See FNiagaraConstants for a synchronized value...
+	float Val = 0.0f;
+    if(NiagaraSpriteVF.SubimageDataOffset != -1)
+	{
+		Val = GetFloat(NiagaraSpriteVF.SubimageDataOffset, InstanceID);
+	}
+	return Val;
+}
+
+float3 GetNiagaraParticleFacingVector(uint InstanceID)
+{
+	// See FNiagaraConstants for a synchronized value...
+	float3 Val = float3(1.0f, 0.0f, 0.0f);
+    if(NiagaraSpriteVF.FacingDataOffset != -1)
+	{
+		Val = GetVec3(NiagaraSpriteVF.FacingDataOffset, InstanceID);
+	}
+	return TransformVector(Val);
+}
+
+float3 GetNiagaraParticleAlignmentVector(uint InstanceID)
+{
+	float3 Val = float3(1.0f, 0.0f, 0.0f);
+    if(NiagaraSpriteVF.AlignmentDataOffset != -1)
+	{
+		Val = GetVec3(NiagaraSpriteVF.AlignmentDataOffset, InstanceID);
+	}
+	return TransformVector(Val);
+}
+
+
+FVertexFactoryIntermediates GetVertexFactoryIntermediates(FVertexFactoryInput Input)
+{
+	FVertexFactoryIntermediates Intermediates = (FVertexFactoryIntermediates)0;
+
+	uint ParticleID = Input.InstanceId;
+
+#if FEATURE_LEVEL >= FEATURE_LEVEL_SM4 && !SHADOW_DEPTH_SHADER && !SPHERICAL_OPACITY_FOR_SHADOW_DEPTHS
+	// Don't use sort buffer in the shadow because they are only valid after opaque.
+	if(NiagaraSpriteVFLooseParameters.SortedIndicesOffset != 0xFFFFFFFF)
+	{
+		ParticleID = NiagaraSpriteVFLooseParameters.SortedIndices[NiagaraSpriteVFLooseParameters.SortedIndicesOffset + Input.InstanceId];
+	}
+#endif
+	Intermediates.Position = GetNiagaraParticlePosition(ParticleID);
+	Intermediates.OldPosition = Intermediates.Position.xyz - GetNiagaraParticleVelocity(ParticleID) * NiagaraSpriteVF.DeltaSeconds;
+	float2 ParticleSize = GetNiagaraParticleSize(ParticleID);
+	Intermediates.SizeRotSubImage = float4(ParticleSize.x, ParticleSize.y, GetNiagaraParticleRotation(ParticleID), GetNiagaraParticleSubimage(ParticleID));
+	Intermediates.Color	= GetNiagaraParticleColor(ParticleID);
+	Intermediates.CustomFacingVector = SafeNormalize(GetNiagaraParticleFacingVector(ParticleID));
+	Intermediates.CustomAlignmentVector = SafeNormalize(GetNiagaraParticleAlignmentVector(ParticleID));
+
+#if (DYNAMIC_PARAMETERS_MASK & 1)
+	Intermediates.DynamicParameter = GetNiagaraParticleDynamicParameters(ParticleID);
+#endif
+#if (DYNAMIC_PARAMETERS_MASK & 2)
+	Intermediates.DynamicParameter1 = GetNiagaraParticleDynamicParameters1(ParticleID);
+#endif
+#if (DYNAMIC_PARAMETERS_MASK & 4)
+	Intermediates.DynamicParameter2 = GetNiagaraParticleDynamicParameters2(ParticleID);
+#endif
+#if (DYNAMIC_PARAMETERS_MASK & 8)
+	Intermediates.DynamicParameter3 = GetNiagaraParticleDynamicParameters3(ParticleID);
+#endif
+
+
+	// World position.
+	
+	float4 ParticleWorldPosition = mul(float4(Intermediates.Position.xyz,1), Primitive.LocalToWorld);
+	float4 ParticleOldWorldPosition = mul(float4(Intermediates.OldPosition.xyz,1), Primitive.LocalToWorld); 
+	float CameraOffset = GetNiagaraCameraOffset(ParticleID);
+	float3 CameraVector = ResolvedView.WorldCameraOrigin.xyz - ParticleWorldPosition.xyz;
+	float3 CameraDirection = SafeNormalize(CameraVector);
+
+	// Do we need the hack that forces particles at the camera origin to be *really* behind the camera origin?
+	float3 OffsetPosition = CameraDirection * CameraOffset;
+
+	ParticleWorldPosition.xyz = ParticleWorldPosition.xyz + OffsetPosition;
+	ParticleOldWorldPosition.xyz = ParticleOldWorldPosition.xyz + OffsetPosition;
+	
+	Intermediates.ParticleWorldPosition = ParticleWorldPosition.xyz;
+
+	// Tangents.
+	float3 Right,Up;
+	GetTangents(Input, Intermediates, ParticleWorldPosition.xyz,ParticleOldWorldPosition.xyz,Right,Up);
+	Intermediates.TangentUp = Up;
+	Intermediates.TangentRight = Right;
+
+	float2 UVForPosition;
+	float2 UVForTexturing;
+	float2 UVForTexturingUnflipped;
+	ComputeBillboardUVs(Input, Intermediates, UVForPosition, UVForTexturing, UVForTexturingUnflipped);
+
+	float2 UVScale = GetNiagaraUVScale(ParticleID);
+	UVForTexturing.xy = float2(0.5f, 0.5f) + (UVForTexturing.xy - float2(0.5f, 0.5f) ) *  UVScale.xy; 
+
+	// Vertex position.
+	float4 VertexWorldPosition = float4(ParticleWorldPosition.xyz + ResolvedView.PreViewTranslation,1);
+	float2 Size = abs(Intermediates.SizeRotSubImage.xy);
+	VertexWorldPosition += Size.x * (UVForPosition.x + NiagaraSpriteVF.PivotOffset.x) * float4(Right,0);
+	VertexWorldPosition += Size.y * (UVForPosition.y + NiagaraSpriteVF.PivotOffset.y) * float4(Up,0);
+	Intermediates.VertexWorldPosition = VertexWorldPosition.xyz;
+
+	// Previous vertex position.
+	Intermediates.PreviousVertexWorldPosition = VertexWorldPosition.xyz + (ParticleOldWorldPosition.xyz - ParticleWorldPosition.xyz);
+	
+	// SubUV.
+	float SubImageIndex = Intermediates.SizeRotSubImage.w;
+	float SubImageLerp = frac(SubImageIndex);
+	float SubImageA = SubImageIndex - SubImageLerp;
+	float SubImageB = SubImageA + 1;
+	float SubImageAH = fmod( SubImageA, NiagaraSpriteVF.SubImageSize.x );
+	float SubImageBH = fmod( SubImageB, NiagaraSpriteVF.SubImageSize.x );
+	float SubImageAV = floor( SubImageA * NiagaraSpriteVF.SubImageSize.z );
+	float SubImageBV = floor( SubImageB * NiagaraSpriteVF.SubImageSize.z );
+	Intermediates.TexCoord.xy = (float2( SubImageAH, SubImageAV ) + UVForTexturing) * NiagaraSpriteVF.SubImageSize.zw;
+	Intermediates.TexCoord.zw = (float2(SubImageBH, SubImageBV) + UVForTexturing) * NiagaraSpriteVF.SubImageSize.zw;
+#if NUM_MATERIAL_TEXCOORDS >= 3
+	Intermediates.TexCoord_Unflipped.xy = (float2(SubImageAH, SubImageAV) + UVForTexturingUnflipped) * NiagaraSpriteVF.SubImageSize.zw;
+	Intermediates.TexCoord_Unflipped.zw = (float2(SubImageBH, SubImageBV) + UVForTexturingUnflipped) * NiagaraSpriteVF.SubImageSize.zw;
+#endif
+	Intermediates.SubImageLerp = (NiagaraSpriteVF.SubImageBlendMode == 1) ? SubImageLerp : 0.0f;
+	
+
+#if USE_PARTICLE_POSITION
+	float ParticleRadius = .5f * min(Intermediates.SizeRotSubImage.x, Intermediates.SizeRotSubImage.y);
+	float3 ParticleTranslatedWorldPosition = ParticleWorldPosition.xyz + ResolvedView.PreViewTranslation.xyz;
+	Intermediates.TranslatedWorldPositionAndSize = float4(ParticleTranslatedWorldPosition, ParticleRadius);
+
+#if GENERATE_SPHERICAL_PARTICLE_NORMALS
+	// Recenter the particle position for GENERATE_SPHERICAL_PARTICLE_NORMALS to use the right center.
+	Intermediates.TranslatedWorldPositionAndSize.xyz += Size.x * (NiagaraSpriteVF.PivotOffset.x + .5) * Right + Size.y * (NiagaraSpriteVF.PivotOffset.y + .5) * Up;
+#endif
+#endif
+
+#if USE_PARTICLE_VELOCITY
+	float3 ParticleVelocity = (ParticleWorldPosition.xyz - ParticleOldWorldPosition.xyz) * (1.0f/NiagaraSpriteVF.DeltaSeconds);
+	float3 DirectionVector = lerp(ParticleVelocity, Intermediates.CustomAlignmentVector, NiagaraSpriteVFLooseParameters.ParticleAlignmentMode == ALIGNMENT_MODE_CUSTOMALIGNMENT);
+	Intermediates.ParticleVelocity.xyz = normalize(DirectionVector);
+	Intermediates.ParticleVelocity.w = length(DirectionVector);
+#endif
+
+#if USE_PARTICLE_TIME
+	Intermediates.RelativeTime = GetNiagaraNormalizedAge(ParticleID);
+#endif
+
+#if USE_PARTICLE_RANDOM
+	Intermediates.ParticleRandom = GetNiagaraParticleRandom(ParticleID);
+#endif
+
+	Intermediates.TangentToLocal = CalcTangentBasis(Intermediates);
+
+#if USE_PARTICLE_LIGHTING_OFFSET && USE_PARTICLE_RANDOM
+	// Hash function based on the particle ID to generate a uniformly distributed 3d offset
+	float3 RandomParticleOffset = frac(Square(Intermediates.ParticleRandom * 47 + 10) * float3(1361.456345, 2333.578, 3623.983)) * 2 - 1;
+	Intermediates.LightingPositionOffset = .5f * View.TranslucencyLightingVolumeInvSize[0].w * RandomParticleOffset;
+#endif
+
+#if USE_PARTICLE_SIZE
+	Intermediates.ParticleSize = Size;
+#endif
+
+	return Intermediates;
+}
+
+float4 VertexFactoryGetWorldPosition(FVertexFactoryInput Input, FVertexFactoryIntermediates Intermediates)
+{
+	return float4(Intermediates.VertexWorldPosition,1);
+}
+
+float4 VertexFactoryGetRasterizedWorldPosition(FVertexFactoryInput Input, FVertexFactoryIntermediates Intermediates, float4 TranslatedWorldPosition)
+{
+#if SPHERICAL_PARTICLE_OPACITY
+	// For particles using spherical opacity, move the quad toward the viewer so that it lies in front of the sphere defined by the particle
+	// This avoids opaque objects intersecting the particle from causing clipping artifacts due to depth testing
+	// The downside is that the particle will clip the near plane sooner
+
+	float Radius = .5f * min(Intermediates.SizeRotSubImage.x, Intermediates.SizeRotSubImage.y);
+	return ReprojectPosition(TranslatedWorldPosition, Radius);
+#else
+	return TranslatedWorldPosition;
+#endif
+}
+
+float3 VertexFactoryGetPositionForVertexLighting(FVertexFactoryInput Input, FVertexFactoryIntermediates Intermediates, float3 TranslatedWorldPosition)
+{
+#if (TRANSLUCENCY_LIGHTING_VOLUMETRIC_PERVERTEX_DIRECTIONAL || TRANSLUCENCY_LIGHTING_VOLUMETRIC_PERVERTEX_NONDIRECTIONAL)
+	// Do per-vertex lighting with particle position instead of vertex position when we're using SubUV cutouts, since the vertex positions vary per-frame which would cause popping
+	return NiagaraSpriteVFLooseParameters.NumCutoutVerticesPerFrame > 0 ? (Intermediates.ParticleWorldPosition.xyz + ResolvedView.PreViewTranslation.xyz) : TranslatedWorldPosition;
+#else
+	return TranslatedWorldPosition;
+#endif
+}
+
+FVertexFactoryInterpolantsVSToPS VertexFactoryGetInterpolantsVSToPS(FVertexFactoryInput Input, FVertexFactoryIntermediates Intermediates, FMaterialVertexParameters VertexParameters)
+{
+	FVertexFactoryInterpolantsVSToPS Interpolants;
+
+	// Initialize the whole struct to 0
+	Interpolants = (FVertexFactoryInterpolantsVSToPS)0;
+
+#if NUM_TEX_COORD_INTERPOLATORS
+
+	float2 CustomizedUVs[NUM_TEX_COORD_INTERPOLATORS];
+	GetMaterialCustomizedUVs(VertexParameters, CustomizedUVs);
+
+	UNROLL
+	for (int CoordinateIndex = 0; CoordinateIndex < NUM_TEX_COORD_INTERPOLATORS; CoordinateIndex++)
+	{
+		SetUV(Interpolants, CoordinateIndex, CustomizedUVs[CoordinateIndex]);
+	}
+#endif
+
+#if LIGHTMAP_UV_ACCESS
+	Interpolants.LightMapUVs = Intermediates.TexCoord.xy;
+#endif
+
+#if USE_PARTICLE_SUBUVS
+	Interpolants.ParticleSubUVs.xy = VertexParameters.Particle.SubUVCoords[0];
+	Interpolants.ParticleSubUVs.zw = VertexParameters.Particle.SubUVCoords[1];
+#endif
+
+	// Calculate the transform from tangent to world space.
+	// Note that "local" space for particles is actually oriented in world space! Therefore no rotation is needed.
+	float3x3 TangentToWorld = Intermediates.TangentToLocal;
+
+	Interpolants.TangentToWorld0AndInterp_Sizer.xyz = TangentToWorld[0];
+	Interpolants.TangentToWorld0AndInterp_Sizer.w = Intermediates.SubImageLerp;
+	Interpolants.TangentToWorld2 = float4(TangentToWorld[2], sign(determinant(Intermediates.TangentToLocal)));
+
+#if NEEDS_PARTICLE_COLOR
+	Interpolants.Color = Intermediates.Color;
+#endif
+
+#if (DYNAMIC_PARAMETERS_MASK & 1)
+	Interpolants.DynamicParameter = Intermediates.DynamicParameter;
+#endif
+#if (DYNAMIC_PARAMETERS_MASK & 2)
+	Interpolants.DynamicParameter1 = Intermediates.DynamicParameter1;
+#endif
+#if (DYNAMIC_PARAMETERS_MASK & 4)
+	Interpolants.DynamicParameter2 = Intermediates.DynamicParameter2;
+#endif
+#if (DYNAMIC_PARAMETERS_MASK & 8)
+	Interpolants.DynamicParameter3 = Intermediates.DynamicParameter3;
+#endif
+
+#if USE_PARTICLE_POSITION
+	Interpolants.ParticleTranslatedWorldPositionAndSize = Intermediates.TranslatedWorldPositionAndSize;
+#endif
+
+#if USE_PARTICLE_VELOCITY
+	Interpolants.ParticleVelocity = Intermediates.ParticleVelocity;
+#endif
+
+#if USE_PARTICLE_TIME
+	Interpolants.RelativeTime = Intermediates.RelativeTime;
+#endif
+
+#if USE_PARTICLE_LIGHTING_OFFSET
+	Interpolants.LightingPositionOffset = Intermediates.LightingPositionOffset;
+#endif
+
+#if USE_PARTICLE_SIZE
+	Interpolants.ParticleSize = Intermediates.ParticleSize;
+#endif
+
+#if USE_PARTICLE_RANDOM
+	Interpolants.ParticleRandom = Intermediates.ParticleRandom;
+#endif
+
+
+#if INSTANCED_STEREO
+	Interpolants.EyeIndex = 0;
+#endif
+
+	return Interpolants;
+}
+
+#if RAYHITGROUPSHADER
+struct FVertexFactoryInterpolantsVSToDS
+{
+	FVertexFactoryInterpolantsVSToPS InterpolantsVSToPS;
+};
+
+FVertexFactoryInterpolantsVSToDS VertexFactoryGetInterpolantsVSToDS(FVertexFactoryInput Input, FVertexFactoryIntermediates Intermediates, FMaterialVertexParameters VertexParameters)
+{
+	FVertexFactoryInterpolantsVSToDS Interpolants;
+	
+	Interpolants.InterpolantsVSToPS = VertexFactoryGetInterpolantsVSToPS(Input, Intermediates, VertexParameters);
+	
+	return Interpolants;
+}
+
+FVertexFactoryInterpolantsVSToDS VertexFactoryInterpolate(FVertexFactoryInterpolantsVSToDS a, float aInterp, FVertexFactoryInterpolantsVSToDS b, float bInterp)
+{
+	FVertexFactoryInterpolantsVSToDS O;
+
+#if LIGHTMAP_UV_ACCESS
+	TESSELLATION_INTERPOLATE_MEMBER(InterpolantsVSToPS.LightMapUVs);
+#endif
+
+#if NUM_TEX_COORD_INTERPOLATORS
+	UNROLL
+	for(int tc = 0; tc < (NUM_TEX_COORD_INTERPOLATORS+1)/2; ++tc)
+	{
+		TESSELLATION_INTERPOLATE_MEMBER(InterpolantsVSToPS.TexCoords[tc]);
+	}
+#endif
+
+#if USE_PARTICLE_SUBUVS
+	TESSELLATION_INTERPOLATE_MEMBER(InterpolantsVSToPS.Interpolants.ParticleSubUVs);
+#endif
+
+	TESSELLATION_INTERPOLATE_MEMBER(InterpolantsVSToPS.TangentToWorld0AndInterp_Sizer);
+	TESSELLATION_INTERPOLATE_MEMBER(InterpolantsVSToPS.TangentToWorld2);
+
+#if NEEDS_PARTICLE_COLOR
+	TESSELLATION_INTERPOLATE_MEMBER(InterpolantsVSToPS.Color);
+#endif
+
+#if (DYNAMIC_PARAMETERS_MASK & 1)
+	TESSELLATION_INTERPOLATE_MEMBER(InterpolantsVSToPS.DynamicParameter);
+#endif
+#if (DYNAMIC_PARAMETERS_MASK & 2)
+	TESSELLATION_INTERPOLATE_MEMBER(InterpolantsVSToPS.DynamicParameter1);
+#endif
+#if (DYNAMIC_PARAMETERS_MASK & 4)
+	TESSELLATION_INTERPOLATE_MEMBER(InterpolantsVSToPS.DynamicParameter2);
+#endif
+#if (DYNAMIC_PARAMETERS_MASK & 8)
+	TESSELLATION_INTERPOLATE_MEMBER(InterpolantsVSToPS.DynamicParameter3);
+#endif
+
+#if USE_PARTICLE_POSITION
+	TESSELLATION_INTERPOLATE_MEMBER(InterpolantsVSToPS.ParticleTranslatedWorldPositionAndSize);
+#endif
+
+#if USE_PARTICLE_VELOCITY
+	TESSELLATION_INTERPOLATE_MEMBER(InterpolantsVSToPS.ParticleVelocity);
+#endif
+
+#if USE_PARTICLE_TIME
+	TESSELLATION_INTERPOLATE_MEMBER(InterpolantsVSToPS.RelativeTime);
+#endif
+
+#if USE_PARTICLE_LIGHTING_OFFSET
+	TESSELLATION_INTERPOLATE_MEMBER(InterpolantsVSToPS.LightingPositionOffset);
+#endif
+
+#if USE_PARTICLE_SIZE
+	TESSELLATION_INTERPOLATE_MEMBER(InterpolantsVSToPS.ParticleSize);
+#endif
+
+#if USE_PARTICLE_RANDOM
+	TESSELLATION_INTERPOLATE_MEMBER(InterpolantsVSToPS.ParticleRandom);
+#endif
+
+	return O;
+}
+
+FVertexFactoryInterpolantsVSToPS VertexFactoryAssignInterpolants(FVertexFactoryInterpolantsVSToDS Input)
+{
+	return Input.InterpolantsVSToPS;
+}
+#endif
+
+float4 VertexFactoryGetPreviousWorldPosition(FVertexFactoryInput Input, FVertexFactoryIntermediates Intermediates)
+{
+	return float4(Intermediates.PreviousVertexWorldPosition, 1.0f);
+}
+
+/**
+* Get the 3x3 tangent basis vectors for this vertex factory
+*
+* @param Input - vertex input stream structure
+* @return 3x3 matrix
+*/
+float3x3 VertexFactoryGetTangentToLocal(FVertexFactoryInput Input, FVertexFactoryIntermediates Intermediates)
+{
+	return Intermediates.TangentToLocal;
+}
+
+float4 VertexFactoryGetTranslatedPrimitiveVolumeBounds(FVertexFactoryInterpolantsVSToPS Interpolants)
+{
+#if USE_PARTICLE_POSITION
+	return Interpolants.ParticleTranslatedWorldPositionAndSize;
+#else
+	return 0;
+#endif
+}
+
+uint VertexFactoryGetPrimitiveId(FVertexFactoryInterpolantsVSToPS Interpolants)
+{
+	return 0;
+}
diff --git a/Engine/Plugins/FX/Niagara/Shaders/Private/NiagaraVFParticleAccess.usf b/Engine/Plugins/FX/Niagara/Shaders/Private/NiagaraVFParticleAccess.usf
new file mode 100644
--- /dev/null
+++ b/Engine/Plugins/FX/Niagara/Shaders/Private/NiagaraVFParticleAccess.usf
@@ -0,0 +1,49 @@
+
+/*********************************************************************************
+ *							Niagara Vertex Factory data access
+ *********************************************************************************/
+
+//Currently we have no integer data being fed into the renderers.
+//Buffer<int> NiagaraParticleDataInt;
+//uint NiagaraIntDataOffset;
+//uint NiagaraIntDataStride;
+
+#ifndef NIAGARA_KEYGEN
+#define NIAGARA_KEYGEN 0
+#endif
+
+#if (!RAYHITGROUPSHADER && !COMPUTESHADER) || NIAGARA_KEYGEN
+
+Buffer<float> NiagaraParticleDataFloat;
+uint NiagaraFloatDataOffset;
+uint NiagaraFloatDataStride;
+
+float GetFloat(int RegisterIdx, uint InstanceID)
+{
+	return NiagaraParticleDataFloat[NiagaraFloatDataOffset + (RegisterIdx * NiagaraFloatDataStride + InstanceID)];
+}
+
+#else
+
+float GetFloat(int RegisterIdx, uint InstanceID)
+{
+	return NiagaraSpriteVFLooseParameters.NiagaraParticleDataFloat[NiagaraSpriteVFLooseParameters.NiagaraFloatDataOffset + (RegisterIdx * NiagaraSpriteVFLooseParameters.NiagaraFloatDataStride + InstanceID)];
+}
+
+#endif
+
+float2 GetVec2(int RegisterIndex, uint InstanceID)
+{
+	return float2(GetFloat(RegisterIndex, InstanceID), GetFloat(RegisterIndex+1, InstanceID));
+}
+
+float3 GetVec3(int RegisterIndex, uint InstanceID)
+{
+	return float3(GetFloat(RegisterIndex, InstanceID), GetFloat(RegisterIndex+1, InstanceID), GetFloat(RegisterIndex+2, InstanceID));
+}
+
+float4 GetVec4(int RegisterIndex, uint InstanceID)
+{
+	return float4(GetFloat(RegisterIndex, InstanceID), GetFloat(RegisterIndex+1, InstanceID), GetFloat(RegisterIndex+2, InstanceID), GetFloat(RegisterIndex+3, InstanceID));
+}
+
diff --git a/Engine/Plugins/FX/Niagara/Source/Niagara/Classes/NDISkeletalMeshCommon.h b/Engine/Plugins/FX/Niagara/Source/Niagara/Classes/NDISkeletalMeshCommon.h
--- a/Engine/Plugins/FX/Niagara/Source/Niagara/Classes/NDISkeletalMeshCommon.h
+++ b/Engine/Plugins/FX/Niagara/Source/Niagara/Classes/NDISkeletalMeshCommon.h
@@ -111,10 +111,11 @@ struct FSkinnedPositionAccessorHelper<TIntegralConstant<ENDISkeletalMesh_Skinnin
 {
 	FORCEINLINE void GetTrianlgeIndices(FSkeletalMeshAccessorHelper& Accessor, int32 Tri, int32& Idx0, int32& Idx1, int32& Idx2)
 	{
-		checkSlow(Tri + 2 < Accessor.IndexBuffer->Num());
-		Idx0 = Accessor.IndexBuffer->Get(Tri);
-		Idx1 = Accessor.IndexBuffer->Get(Tri + 1);
-		Idx2 = Accessor.IndexBuffer->Get(Tri + 2);
+		const int32 BaseIndex = Tri * 3;
+		checkSlow(BaseIndex + 2 < Accessor.IndexBuffer->Num());
+		Idx0 = Accessor.IndexBuffer->Get(BaseIndex);
+		Idx1 = Accessor.IndexBuffer->Get(BaseIndex + 1);
+		Idx2 = Accessor.IndexBuffer->Get(BaseIndex + 2);
 	}
 
 	FORCEINLINE void GetSkinnedTrianglePositions(FSkeletalMeshAccessorHelper& Accessor, int32 Idx0, int32 Idx1, int32 Idx2, FVector& OutPos0, FVector& OutPos1, FVector& OutPos2)
@@ -157,10 +158,11 @@ struct FSkinnedPositionAccessorHelper<TIntegralConstant<ENDISkeletalMesh_Skinnin
 {
 	FORCEINLINE void GetTrianlgeIndices(FSkeletalMeshAccessorHelper& Accessor, int32 Tri, int32& Idx0, int32& Idx1, int32& Idx2)
 	{
-		checkSlow(Tri + 2 < Accessor.IndexBuffer->Num());
-		Idx0 = Accessor.IndexBuffer->Get(Tri);
-		Idx1 = Accessor.IndexBuffer->Get(Tri + 1);
-		Idx2 = Accessor.IndexBuffer->Get(Tri + 2);
+		const int32 BaseIndex = Tri * 3;
+		checkSlow(BaseIndex + 2 < Accessor.IndexBuffer->Num());
+		Idx0 = Accessor.IndexBuffer->Get(BaseIndex);
+		Idx1 = Accessor.IndexBuffer->Get(BaseIndex + 1);
+		Idx2 = Accessor.IndexBuffer->Get(BaseIndex + 2);
 	}
 
 	FORCEINLINE void GetSkinnedTrianglePositions(FSkeletalMeshAccessorHelper& Accessor, int32 Idx0, int32 Idx1, int32 Idx2, FVector& OutPos0, FVector& OutPos1, FVector& OutPos2)
@@ -203,10 +205,11 @@ struct FSkinnedPositionAccessorHelper<TIntegralConstant<ENDISkeletalMesh_Skinnin
 {
 	FORCEINLINE void GetTrianlgeIndices(FSkeletalMeshAccessorHelper& Accessor, int32 Tri, int32& Idx0, int32& Idx1, int32& Idx2)
 	{
-		checkSlow(Tri + 2 < Accessor.IndexBuffer->Num());
-		Idx0 = Accessor.IndexBuffer->Get(Tri);
-		Idx1 = Accessor.IndexBuffer->Get(Tri + 1);
-		Idx2 = Accessor.IndexBuffer->Get(Tri + 2);
+		const int32 BaseIndex = Tri * 3;
+		checkSlow(BaseIndex + 2 < Accessor.IndexBuffer->Num());
+		Idx0 = Accessor.IndexBuffer->Get(BaseIndex);
+		Idx1 = Accessor.IndexBuffer->Get(BaseIndex + 1);
+		Idx2 = Accessor.IndexBuffer->Get(BaseIndex + 2);
 	}
 
 	FORCEINLINE void GetSkinnedTrianglePositions(FSkeletalMeshAccessorHelper& Accessor, int32 Idx0, int32 Idx1, int32 Idx2, FVector& OutPos0, FVector& OutPos1, FVector& OutPos2)
diff --git a/Engine/Plugins/FX/Niagara/Source/Niagara/Classes/NiagaraCollision.h b/Engine/Plugins/FX/Niagara/Source/Niagara/Classes/NiagaraCollision.h
--- a/Engine/Plugins/FX/Niagara/Source/Niagara/Classes/NiagaraCollision.h
+++ b/Engine/Plugins/FX/Niagara/Source/Niagara/Classes/NiagaraCollision.h
@@ -38,19 +38,11 @@ public:
 		FNiagaraEventDataSetMgr::Reset(OwnerSystemInstanceName, EmitterName);
 	}
 
-	void Tick()
-	{
-		if (CollisionEventDataSet)
-		{
-			CollisionEventDataSet->Tick();
-		}
-	}
-
 	void Reset()
 	{
 		if (CollisionEventDataSet)
 		{
-			CollisionEventDataSet->SetNumInstances(0);
+			CollisionEventDataSet->ResetBuffers();
 		}
 	}
 
@@ -58,7 +50,7 @@ public:
 	{
 		if (CollisionEventDataSet)
 		{
-			CollisionEventDataSet->Init(FNiagaraDataSetID(), ENiagaraSimTarget::CPUSim);
+			CollisionEventDataSet->Init(FNiagaraDataSetID(), ENiagaraSimTarget::CPUSim, TEXT("NiagaraCollisionEventDataset"));
 		}
 		EmitterName = InEmitterName;
 		OwnerSystemInstanceName = InOwnerSystemInstanceName;
@@ -80,9 +72,9 @@ public:
 		CollisionEventDataSet->Finalize();
 	}
 	
-	void KickoffNewBatch(struct FNiagaraEmitterInstance *Sim, float DeltaSeconds);
+	void KickoffNewBatch(class FNiagaraEmitterInstance *Sim, float DeltaSeconds);
 
-	void GenerateEventsFromResults(struct FNiagaraEmitterInstance *Sim);
+	void GenerateEventsFromResults(class FNiagaraEmitterInstance *Sim);
 
 	const FNiagaraDataSet *GetDataSet() const { return CollisionEventDataSet; }
 private:
diff --git a/Engine/Plugins/FX/Niagara/Source/Niagara/Classes/NiagaraDataInterface.h b/Engine/Plugins/FX/Niagara/Source/Niagara/Classes/NiagaraDataInterface.h
--- a/Engine/Plugins/FX/Niagara/Source/Niagara/Classes/NiagaraDataInterface.h
+++ b/Engine/Plugins/FX/Niagara/Source/Niagara/Classes/NiagaraDataInterface.h
@@ -15,6 +15,7 @@ class UCurveVector;
 class UCurveLinearColor;
 class UCurveFloat;
 class FNiagaraSystemInstance;
+struct FNiagaraDataInterfaceProxy;
 
 struct FNDITransformHandlerNoop
 {
@@ -143,6 +144,40 @@ private:
 
 //////////////////////////////////////////////////////////////////////////
 
+struct FNiagaraDataInterfaceProxy : TSharedFromThis<FNiagaraDataInterfaceProxy, ESPMode::ThreadSafe>
+{
+	virtual ~FNiagaraDataInterfaceProxy() {check(IsInRenderingThread());}
+
+	virtual int32 PerInstanceDataPassedToRenderThreadSize() const = 0;
+	virtual void ConsumePerInstanceDataFromGameThread(void* PerInstanceData, const FGuid& Instance) { check(false); }
+
+	virtual void DeferredDestroy() {}
+};
+
+struct FNiagaraDataInterfaceProxyCurveBase : public FNiagaraDataInterfaceProxy
+{
+	~FNiagaraDataInterfaceProxyCurveBase()
+	{
+		check(IsInRenderingThread());
+		CurveLUT.Release();
+	}
+
+	float LUTMinTime;
+	float LUTMaxTime;
+	float LUTInvTimeRange;
+	FReadBuffer CurveLUT;
+
+	virtual int32 PerInstanceDataPassedToRenderThreadSize() const override
+	{
+		return 0;
+	}
+
+	// @todo REMOVEME
+	virtual void ConsumePerInstanceDataFromGameThread(void* PerInstanceData, const FGuid& Instance) override { check(false); }
+};
+
+//////////////////////////////////////////////////////////////////////////
+
 /** Base class for all Niagara data interfaces. */
 UCLASS(abstract, EditInlineNew)
 class NIAGARA_API UNiagaraDataInterface : public UNiagaraDataInterfaceBase
@@ -151,20 +186,52 @@ class NIAGARA_API UNiagaraDataInterface : public UNiagaraDataInterfaceBase
 		 
 public: 
 
+	virtual ~UNiagaraDataInterface();
+
 	// UObject Interface
-	virtual void PostLoad()override;
+	virtual void PostLoad() override;
 	// UObject Interface END
 
 	/** Initializes the per instance data for this interface. Returns false if there was some error and the simulation should be disabled. */
 	virtual bool InitPerInstanceData(void* PerInstanceData, FNiagaraSystemInstance* SystemInstance) { return true; }
 
-	/** Destroys the per instence data for this interface. */
+	/** Destroys the per instance data for this interface. */
 	virtual void DestroyPerInstanceData(void* PerInstanceData, FNiagaraSystemInstance* SystemInstance) {}
 
 	/** Ticks the per instance data for this interface, if it has any. */
 	virtual bool PerInstanceTick(void* PerInstanceData, FNiagaraSystemInstance* SystemInstance, float DeltaSeconds) { return false; }
 	virtual bool PerInstanceTickPostSimulate(void* PerInstanceData, FNiagaraSystemInstance* SystemInstance, float DeltaSeconds) { return false; }
 
+	/** 
+		Subclasses that wish to work with GPU systems/emitters must implement this.
+		Those interfaces must fill DataForRenderThread with the data needed to upload to the GPU. It will be the last thing called on this
+		data interface for a specific tick.
+		This will be consumed by the associated FNiagaraDataInterfaceProxy.
+		Note: This class does not own the memory pointed to by DataForRenderThread. It will be recycled automatically. 
+			However, if you allocate memory yourself to pass via this buffer you ARE responsible for freeing it when it is consumed by the proxy (Which is what ChaosDestruction does).
+			Likewise, the class also does not own the memory in PerInstanceData. That pointer is the pointer passed to PerInstanceTick/PerInstanceTickPostSimulate.
+			
+		This will not be called if PerInstanceDataPassedToRenderThreadSize is 0.
+	*/
+	virtual void ProvidePerInstanceDataForRenderThread(void* DataForRenderThread, void* PerInstanceData, const FGuid& SystemInstance)
+	{
+		check(false);
+	}
+
+	/**
+	 * The size of the data this class will provide to ProvidePerInstanceDataForRenderThread.
+	 * MUST be 16 byte aligned!
+	 */
+	virtual int32 PerInstanceDataPassedToRenderThreadSize() const 
+	{ 
+		if (!Proxy)
+		{
+			return 0;
+		}
+		check(Proxy);
+		return Proxy->PerInstanceDataPassedToRenderThreadSize();
+	}
+
 	/** 
 	Returns the size of the per instance data for this interface. 0 if this interface has no per instance data. 
 	Must depend solely on the class of the interface and not on any particular member data of a individual interface.
@@ -185,6 +252,8 @@ public:
 
 	virtual bool CanExecuteOnTarget(ENiagaraSimTarget Target)const { return false; }
 
+	virtual bool RequiresDistanceFieldData() const { return false; }
+
 	/** Determines if this type definition matches to a known data interface type.*/
 	static bool IsDataInterfaceType(const FNiagaraTypeDefinition& TypeDef);
 
@@ -206,8 +275,23 @@ public:
 	virtual void ValidateFunction(const FNiagaraFunctionSignature& Function, TArray<FText>& OutValidationErrors);
 #endif
 
+	FNiagaraDataInterfaceProxy* GetProxy()
+	{
+		return Proxy.Get();
+	}
+
 protected:
+	template<typename T>
+	T* GetProxyAs()
+	{
+		T* TypedProxy = static_cast<T*>(Proxy.Get());
+		check(TypedProxy != nullptr);
+		return TypedProxy;
+	}
+
 	virtual bool CopyToInternal(UNiagaraDataInterface* Destination) const;
+
+	TSharedPtr<FNiagaraDataInterfaceProxy, ESPMode::ThreadSafe> Proxy;
 };
 
 /** Base class for curve data interfaces which facilitates handling the curve data in a standardized way. */
@@ -216,8 +300,8 @@ class NIAGARA_API UNiagaraDataInterfaceCurveBase : public UNiagaraDataInterface
 {
 protected:
 	GENERATED_BODY()
-	UPROPERTY()
-	bool GPUBufferDirty;
+	/*UPROPERTY()
+	bool GPUBufferDirty;*/
 	UPROPERTY()
 	TArray<float> ShaderLUT;
 	UPROPERTY()
@@ -241,8 +325,7 @@ protected:
 
 public:
 	UNiagaraDataInterfaceCurveBase()
-		: GPUBufferDirty(false)
-		, LUTMinTime(0.0f)
+		: LUTMinTime(0.0f)
 		, LUTMaxTime(1.0f)
 		, LUTInvTimeRange(1.0f)
 		, bUseLUT(true)
@@ -250,18 +333,20 @@ public:
 		, ShowInCurveEditor(false)
 #endif
 	{
+		Proxy = MakeShared<FNiagaraDataInterfaceProxyCurveBase, ESPMode::ThreadSafe>();
 	}
 
 	UNiagaraDataInterfaceCurveBase(FObjectInitializer const& ObjectInitializer)
-		: GPUBufferDirty(false)
-		, LUTMinTime(0.0f)
+		: LUTMinTime(0.0f)
 		, LUTMaxTime(1.0f)
 		, LUTInvTimeRange(1.0f)
 		, bUseLUT(true)
 #if WITH_EDITORONLY_DATA
 		, ShowInCurveEditor(false)
 #endif
-	{}
+	{
+		Proxy = MakeShared<FNiagaraDataInterfaceProxyCurveBase, ESPMode::ThreadSafe>();
+	}
 
 	UPROPERTY(EditAnywhere, AdvancedDisplay, Category = "Curve")
 	uint32 bUseLUT : 1;
@@ -306,7 +391,7 @@ public:
 	}
 
 	//TODO: Make this a texture and get HW filter + clamping?
-	FReadBuffer& GetCurveLUTGPUBuffer();
+	//FReadBuffer& GetCurveLUTGPUBuffer();
 
 	//UNiagaraDataInterface interface
 	virtual bool Equals(const UNiagaraDataInterface* Other) const override;
@@ -317,13 +402,11 @@ public:
 	FORCEINLINE float GetMaxTime()const { return LUTMaxTime; }
 	FORCEINLINE float GetInvTimeRange()const { return LUTInvTimeRange; }
 
-
 protected:
+	void PushToRenderThread();
 	virtual bool CopyToInternal(UNiagaraDataInterface* Destination) const override;
 	virtual bool CompareLUTS(const TArray<float>& OtherLUT) const;
 	//UNiagaraDataInterface interface END
-
-	FReadBuffer CurveLUT;
 };
 
 //External function binder choosing between template specializations based on if a curve should use the LUT over full evaluation.
diff --git a/Engine/Plugins/FX/Niagara/Source/Niagara/Classes/NiagaraDataInterfaceCollisionQuery.h b/Engine/Plugins/FX/Niagara/Source/Niagara/Classes/NiagaraDataInterfaceCollisionQuery.h
--- a/Engine/Plugins/FX/Niagara/Source/Niagara/Classes/NiagaraDataInterfaceCollisionQuery.h
+++ b/Engine/Plugins/FX/Niagara/Source/Niagara/Classes/NiagaraDataInterfaceCollisionQuery.h
@@ -61,6 +61,7 @@ public:
 	void QueryMeshDistanceField(FVectorVMContext& Context);
 
 	virtual bool CanExecuteOnTarget(ENiagaraSimTarget Target) const override { return true; }
+	virtual bool RequiresDistanceFieldData() const override { return true; }
 
 	virtual bool GetFunctionHLSL(const FName&  DefinitionFunctionName, FString InstanceFunctionName, FNiagaraDataInterfaceGPUParamInfo& ParamInfo, FString& OutHLSL) override;
 	virtual void GetParameterDefinitionHLSL(FNiagaraDataInterfaceGPUParamInfo& ParamInfo, FString& OutHLSL) override;
@@ -71,3 +72,13 @@ private:
 	static FCriticalSection CriticalSection;
 	UEnum* TraceChannelEnum;
 };
+
+struct FNiagaraDataIntefaceProxyCollisionQuery : public FNiagaraDataInterfaceProxy
+{
+	// There's nothing in this proxy. It just reads from scene textures.
+
+	virtual int32 PerInstanceDataPassedToRenderThreadSize() const override
+	{
+		return 0;
+	}
+};
diff --git a/Engine/Plugins/FX/Niagara/Source/Niagara/Classes/NiagaraDataInterfaceCurlNoise.h b/Engine/Plugins/FX/Niagara/Source/Niagara/Classes/NiagaraDataInterfaceCurlNoise.h
--- a/Engine/Plugins/FX/Niagara/Source/Niagara/Classes/NiagaraDataInterfaceCurlNoise.h
+++ b/Engine/Plugins/FX/Niagara/Source/Niagara/Classes/NiagaraDataInterfaceCurlNoise.h
@@ -45,4 +45,22 @@ public:
 
 protected:
 	virtual bool CopyToInternal(UNiagaraDataInterface* Destination) const override;
+
+	void PushToRenderThread();
+};
+
+struct FNiagaraDataInterfaceProxyCurlNoise : public FNiagaraDataInterfaceProxy
+{
+	FNiagaraDataInterfaceProxyCurlNoise(const FVector& InOffset)
+	{
+		OffsetFromSeed = InOffset;
+	}
+
+	virtual void ConsumePerInstanceDataFromGameThread(void* PerInstanceData, const FGuid& Instance) override { check(false); }
+	virtual int32 PerInstanceDataPassedToRenderThreadSize() const override
+	{
+		return 0;
+	}
+
+	FVector OffsetFromSeed;
 };
diff --git a/Engine/Plugins/FX/Niagara/Source/Niagara/Classes/NiagaraDataInterfaceSimpleCounter.h b/Engine/Plugins/FX/Niagara/Source/Niagara/Classes/NiagaraDataInterfaceSimpleCounter.h
--- a/Engine/Plugins/FX/Niagara/Source/Niagara/Classes/NiagaraDataInterfaceSimpleCounter.h
+++ b/Engine/Plugins/FX/Niagara/Source/Niagara/Classes/NiagaraDataInterfaceSimpleCounter.h
@@ -60,6 +60,10 @@ public:
 	virtual bool GetFunctionHLSL(const FName&  DefinitionFunctionName, FString InstanceFunctionName, FNiagaraDataInterfaceGPUParamInfo& ParamInfo, FString& OutHLSL) override;
 	virtual void GetParameterDefinitionHLSL(FNiagaraDataInterfaceGPUParamInfo& ParamInfo, FString& OutHLSL) override;
 	
+	virtual void ProvidePerInstanceDataForRenderThread(void* DataForRenderThread, void* PerInstanceData, const FGuid& SystemInstance) override
+	{
+		check(false);
+	}
 
 protected:
 	virtual bool CopyToInternal(UNiagaraDataInterface* Destination) const override;
diff --git a/Engine/Plugins/FX/Niagara/Source/Niagara/Classes/NiagaraDataInterfaceSkeletalMesh.h b/Engine/Plugins/FX/Niagara/Source/Niagara/Classes/NiagaraDataInterfaceSkeletalMesh.h
--- a/Engine/Plugins/FX/Niagara/Source/Niagara/Classes/NiagaraDataInterfaceSkeletalMesh.h
+++ b/Engine/Plugins/FX/Niagara/Source/Niagara/Classes/NiagaraDataInterfaceSkeletalMesh.h
@@ -225,6 +225,106 @@ protected:
 	FNDISkeletalMesh_InstanceData* Owner;
 };
 
+/** 
+ * This contains static data created once from the DI.
+ * This should be in a proxy create by GT and accessible on RT. 
+ * Right now we cannot follow a real Proxy pattern since Niagara does not prevent unloading of UI while RT data is still in use. 
+ * See https://jira.it.epicgames.net/browse/UE-69336
+ */
+class FSkeletalMeshGpuSpawnStaticBuffers : public FRenderResource
+{
+public:
+
+	virtual ~FSkeletalMeshGpuSpawnStaticBuffers();
+
+	FORCEINLINE_DEBUGGABLE void Initialise(const FSkeletalMeshLODRenderData& SkeletalMeshLODRenderData, 
+		bool bIsGpuUniformlyDistributedSampling, const FSkeletalMeshSamplingLODBuiltData& SkeletalMeshSamplingLODBuiltData);
+
+	virtual void InitRHI() override;
+	virtual void ReleaseRHI() override;
+
+	virtual FString GetFriendlyName() const override { return TEXT("FSkeletalMeshGpuSpawnStaticBuffers"); }
+
+	FShaderResourceViewRHIRef GetBufferTriangleUniformSamplerProbaSRV() const { return BufferTriangleUniformSamplerProbaSRV; }
+	FShaderResourceViewRHIRef GetBufferTriangleUniformSamplerAliasSRV() const { return BufferTriangleUniformSamplerAliasSRV; }
+	FShaderResourceViewRHIRef GetBufferTriangleMatricesOffsetSRV() const { return BufferTriangleMatricesOffsetSRV; }
+	uint32 GetTriangleCount() const { return TriangleCount; }
+
+	FShaderResourceViewRHIParamRef GetBufferPositionSRV() const { return MeshVertexBufferSrv; }
+	FShaderResourceViewRHIParamRef GetBufferIndexSRV() const { return MeshIndexBufferSrv; }
+	FShaderResourceViewRHIParamRef GetBufferTangentSRV() const { return MeshTangentBufferSRV; }
+	FShaderResourceViewRHIParamRef GetBufferTexCoordSRV() const { return MeshTexCoordBufferSrv; }
+
+	uint32 GetNumTexCoord() const { return NumTexCoord; }
+
+protected:
+
+	FVertexBufferRHIRef BufferTriangleUniformSamplerProbaRHI = nullptr;
+	FShaderResourceViewRHIRef BufferTriangleUniformSamplerProbaSRV = nullptr;
+	FVertexBufferRHIRef BufferTriangleUniformSamplerAliasRHI = nullptr;
+	FShaderResourceViewRHIRef BufferTriangleUniformSamplerAliasSRV = nullptr;
+	FVertexBufferRHIRef BufferTriangleMatricesOffsetRHI = nullptr;
+	FShaderResourceViewRHIRef BufferTriangleMatricesOffsetSRV = nullptr;
+
+	/** Cached SRV to gpu buffers of the mesh we spawn from */
+	FShaderResourceViewRHIParamRef MeshVertexBufferSrv;
+	FShaderResourceViewRHIParamRef MeshIndexBufferSrv;
+	FShaderResourceViewRHIParamRef MeshTangentBufferSRV;
+	FShaderResourceViewRHIParamRef MeshTexCoordBufferSrv;
+
+	uint32 NumTexCoord;
+
+	// Cached data for resource creation on RenderThread
+	const FSkeletalMeshLODRenderData* LODRenderData = nullptr;
+	const FSkeletalMeshSamplingLODBuiltData* SkeletalMeshSamplingLODBuiltData = nullptr;
+	uint32 TriangleCount;
+	uint32 InputWeightStride;
+	bool bUseGpuUniformlyDistributedSampling;
+};
+
+/**
+ * This contains dynamic data created per frame from the DI.
+ * This should be in a proxy create by GT and accessible on RT. Right now we cannot follow a real Proxy pattern since Niagara does not prevent unloading of UI while RT data is still in use.
+ * See https://jira.it.epicgames.net/browse/UE-69336
+ */
+class FSkeletalMeshGpuDynamicBufferProxy : public FRenderResource
+{
+public:
+
+	FSkeletalMeshGpuDynamicBufferProxy();
+	virtual ~FSkeletalMeshGpuDynamicBufferProxy();
+
+	void Initialise(const FSkeletalMeshLODRenderData& SkeletalMeshLODRenderData);
+
+	virtual void InitRHI() override;
+	virtual void ReleaseRHI() override;
+
+	void NewFrame(const FNDISkeletalMesh_InstanceData* InstanceData, int32 LODIndex);
+
+	bool DoesBoneDataExist() const { return bBoneGpuBufferValid;}
+
+	/** Encapsulates a GPU read / CPU write buffer for bone data */
+	struct FSkeletalBuffer
+	{
+		FVertexBufferRHIRef Buffer;
+		FShaderResourceViewRHIRef SRV;
+	};
+
+	FSkeletalBuffer& GetRWBufferBone() { return RWBufferBones[CurrentBoneBufferId % 2]; }
+	FSkeletalBuffer& GetRWBufferPrevBone() { return bPrevBoneGpuBufferValid ? RWBufferBones[(CurrentBoneBufferId + 1) % 2] : GetRWBufferBone(); }
+
+private:
+
+	uint32 BoneCount = 0;
+
+	enum { BufferBoneCount = 2 };
+	FSkeletalBuffer RWBufferBones[BufferBoneCount];
+	uint8 CurrentBoneBufferId = 0;
+
+	bool bBoneGpuBufferValid = false;
+	bool bPrevBoneGpuBufferValid = false;
+};
+
 struct FNDISkeletalMesh_InstanceData
 {
 	//Cached ptr to component we sample from. 
@@ -266,11 +366,22 @@ struct FNDISkeletalMesh_InstanceData
 
 	uint32 ChangeId;
 
+	/** True if the mesh we're using allows area weighted sampling on GPU. */
+	uint32 bIsGpuUniformlyDistributedSampling : 1;
+
+	FShaderResourceViewRHIParamRef MeshSkinWeightBufferSrv;
+	uint32 MeshWeightStrideByte;
+
+	/** Extra mesh data upload to GPU.*/
+	FSkeletalMeshGpuSpawnStaticBuffers* MeshGpuSpawnStaticBuffers;
+	FSkeletalMeshGpuDynamicBufferProxy* MeshGpuSpawnDynamicBuffers;
+
 	FORCEINLINE_DEBUGGABLE bool ResetRequired(UNiagaraDataInterfaceSkeletalMesh* Interface)const;
 
 	bool Init(UNiagaraDataInterfaceSkeletalMesh* Interface, FNiagaraSystemInstance* SystemInstance);
 	FORCEINLINE_DEBUGGABLE bool Cleanup(UNiagaraDataInterfaceSkeletalMesh* Interface, FNiagaraSystemInstance* SystemInstance);
 	FORCEINLINE_DEBUGGABLE bool Tick(UNiagaraDataInterfaceSkeletalMesh* Interface, FNiagaraSystemInstance* SystemInstance, float InDeltaSeconds);
+	FORCEINLINE_DEBUGGABLE void Release();
 	FORCEINLINE int32 GetLODIndex()const { return SkinningData.Usage.GetLODIndex(); }
 
 	FORCEINLINE_DEBUGGABLE FSkeletalMeshLODRenderData& GetLODRenderDataAndSkinWeights(FSkinWeightVertexBuffer*& OutSkinWeightBuffer)
@@ -282,7 +393,7 @@ struct FNDISkeletalMesh_InstanceData
 		}
 		else
 		{
-			OutSkinWeightBuffer = Ret.GetSkinWeightVertexBuffer();
+			OutSkinWeightBuffer = &Ret.SkinWeightVertexBuffer; // Todo @sckime fix this when main comes in for real
 		}
 		return Ret;
 	}
@@ -306,6 +417,10 @@ public:
 	UPROPERTY(EditAnywhere, Category = "Mesh")
 	AActor* Source;
 	
+	/** The source component from which to sample. Takes precedence over the direct mesh. Not exposed to the user, only indirectly accessible from blueprints. */
+	UPROPERTY(Transient)
+	USkeletalMeshComponent* SourceComponent;
+
 	UPROPERTY(EditAnywhere, Category="Mesh")
 	ENDISkeletalMesh_SkinningMode SkinningMode;
 
@@ -342,8 +457,8 @@ public:
 	uint32 ChangeId;
 
 	//~ UObject interface
+	virtual void PostInitProperties()override; 
 #if WITH_EDITOR
-	virtual void PostInitProperties()override;
 	virtual void PostEditChangeProperty(struct FPropertyChangedEvent& PropertyChangedEvent) override;
 #endif
 	//~ UObject interface END
@@ -358,7 +473,7 @@ public:
 	virtual void GetFunctions(TArray<FNiagaraFunctionSignature>& OutFunctions)override;
 	virtual void GetVMExternalFunction(const FVMExternalFunctionBindingInfo& BindingInfo, void* InstanceData, FVMExternalFunction &OutFunc)override;
 	virtual bool Equals(const UNiagaraDataInterface* Other) const override;
-	virtual bool CanExecuteOnTarget(ENiagaraSimTarget Target)const override { return Target == ENiagaraSimTarget::CPUSim; }
+	virtual bool CanExecuteOnTarget(ENiagaraSimTarget Target)const override { return true; }
 #if WITH_EDITOR	
 	virtual TArray<FNiagaraDataInterfaceError> GetErrors() override;
 	virtual void ValidateFunction(const FNiagaraFunctionSignature& Function, TArray<FText>& OutValidationErrors) override;
@@ -366,6 +481,33 @@ public:
 	//~ UNiagaraDataInterface interface END
 
 	static USkeletalMesh* GetSkeletalMeshHelper(UNiagaraDataInterfaceSkeletalMesh* Interface, class UNiagaraComponent* OwningComponent, TWeakObjectPtr<USceneComponent>& SceneComponent, USkeletalMeshComponent*& FoundSkelComp);
+
+	virtual bool GetFunctionHLSL(const FName&  DefinitionFunctionName, FString InstanceFunctionName, FNiagaraDataInterfaceGPUParamInfo& ParamInfo, FString& OutHLSL) override;
+	virtual void GetParameterDefinitionHLSL(FNiagaraDataInterfaceGPUParamInfo& ParamInfo, FString& OutHLSL) override;
+	virtual FNiagaraDataInterfaceParametersCS* ConstructComputeParameters() const override;
+
+	virtual void ProvidePerInstanceDataForRenderThread(void* DataForRenderThread, void* PerInstanceData, const FGuid& SystemInstance) override;
+
+	static const FString MeshIndexBufferName;
+	static const FString MeshVertexBufferName;
+	static const FString MeshSkinWeightBufferName;
+	static const FString MeshCurrBonesBufferName;
+	static const FString MeshPrevBonesBufferName;
+	static const FString MeshTangentBufferName;
+	static const FString MeshTexCoordBufferName;
+	static const FString MeshTriangleSamplerProbaBufferName;
+	static const FString MeshTriangleSamplerAliasBufferName;
+	static const FString MeshTriangleMatricesOffsetBufferName;
+	static const FString MeshTriangleCountName;
+	static const FString MeshWeightStrideByteName;
+	static const FString InstanceTransformName;
+	static const FString InstancePrevTransformName;
+	static const FString InstanceInvDeltaTimeName;
+	static const FString EnabledFeaturesName;
+	static const FString InputWeightStrideName;
+	static const FString NumTexCoordName;
+
+
 protected:
 	virtual bool CopyToInternal(UNiagaraDataInterface* Destination) const override;
 
@@ -401,7 +543,7 @@ public:
 
 	template<typename SkinningHandlerType>
 	void GetTriCoordVertices(FVectorVMContext& Context);
-		
+
 private:
 	template<typename FilterMode, typename AreaWeightingMode>
 	FORCEINLINE int32 RandomTriIndex(FRandomStream& RandStream, FSkeletalMeshAccessorHelper& Accessor, FNDISkeletalMesh_InstanceData* InstData);
@@ -472,7 +614,7 @@ public:
 
 	void RandomSpecificBone(FVectorVMContext& Context);
 
-	template<typename SkinningHandlerType, typename TransformHandlerType>
+	template<typename SkinningHandlerType, typename TransformHandlerType, typename bInterpolated>
 	void GetSkinnedBoneData(FVectorVMContext& Context);
 	
 	void GetSpecificSocketCount(FVectorVMContext& Context);
@@ -483,5 +625,51 @@ public:
 		
 	// End of Direct Bone + Socket Sampling
 	//////////////////////////////////////////////////////////////////////////
+
+	void SetSourceComponentFromBlueprints(USkeletalMeshComponent* ComponentToUse);
+};
+
+
+class FSkeletalMeshInterfaceHelper
+{
+public:
+	static const FName RandomTriCoordName;
+	static const FName IsValidTriCoordName;
+	static const FName GetSkinnedTriangleDataName;
+	static const FName GetSkinnedTriangleDataWSName;
+	static const FName GetSkinnedTriangleDataInterpName;
+	static const FName GetSkinnedTriangleDataWSInterpName;
+	static const FName GetTriColorName;
+	static const FName GetTriUVName;
+	static const FName GetTriangleCountName;
+	static const FName GetTriangleAtName;
+	static const FName GetTriCoordVerticesName;
 };
 
+struct FNiagaraDISkeletalMeshPassedDataToRT
+{
+	FSkeletalMeshGpuSpawnStaticBuffers* StaticBuffers;
+	FSkeletalMeshGpuDynamicBufferProxy* DynamicBuffer;
+	FShaderResourceViewRHIParamRef MeshSkinWeightBufferSrv;
+
+	bool bIsGpuUniformlyDistributedSampling;
+
+	uint32 MeshWeightStrideByte;
+	FMatrix Transform;
+	FMatrix PrevTransform;
+	float DeltaSeconds;
+};
+
+typedef FNiagaraDISkeletalMeshPassedDataToRT FNiagaraDataInterfaceProxySkeletalMeshData;
+
+struct FNiagaraDataInterfaceProxySkeletalMesh : public FNiagaraDataInterfaceProxy
+{
+	virtual int32 PerInstanceDataPassedToRenderThreadSize() const override 
+	{
+		return sizeof(FNiagaraDISkeletalMeshPassedDataToRT);
+	}
+
+	virtual void ConsumePerInstanceDataFromGameThread(void* PerInstanceData, const FGuid& Instance) override;
+
+	TMap<FGuid, FNiagaraDataInterfaceProxySkeletalMeshData> SystemInstancesToData;
+};
diff --git a/Engine/Plugins/FX/Niagara/Source/Niagara/Classes/NiagaraDataInterfaceSpline.h b/Engine/Plugins/FX/Niagara/Source/Niagara/Classes/NiagaraDataInterfaceSpline.h
--- a/Engine/Plugins/FX/Niagara/Source/Niagara/Classes/NiagaraDataInterfaceSpline.h
+++ b/Engine/Plugins/FX/Niagara/Source/Niagara/Classes/NiagaraDataInterfaceSpline.h
@@ -62,6 +62,11 @@ public:
 	void GetLocalToWorld(FVectorVMContext& Context);
 	void GetLocalToWorldInverseTransposed(FVectorVMContext& Context);
 
+	virtual int32 PerInstanceDataPassedToRenderThreadSize() const override
+	{
+		return 0;
+	}
+
 protected:
 	virtual bool CopyToInternal(UNiagaraDataInterface* Destination) const override;
 
diff --git a/Engine/Plugins/FX/Niagara/Source/Niagara/Classes/NiagaraDataInterfaceStaticMesh.h b/Engine/Plugins/FX/Niagara/Source/Niagara/Classes/NiagaraDataInterfaceStaticMesh.h
--- a/Engine/Plugins/FX/Niagara/Source/Niagara/Classes/NiagaraDataInterfaceStaticMesh.h
+++ b/Engine/Plugins/FX/Niagara/Source/Niagara/Classes/NiagaraDataInterfaceStaticMesh.h
@@ -39,6 +39,64 @@ struct FNDIStaticMeshSectionFilter
 
 class UNiagaraDataInterfaceStaticMesh;
 
+/** Used to stored GPU data needed for an interface/mesh tuple (e.g. uniform sampling of sections according to mesh surface area). */
+class FStaticMeshGpuSpawnBuffer : public FRenderResource
+{
+public:
+
+	virtual ~FStaticMeshGpuSpawnBuffer();
+
+	void Initialise(const FStaticMeshLODResources& Res, const UNiagaraDataInterfaceStaticMesh& Interface,
+		bool bIsGpuUniformlyDistributedSampling, const TArray<int32>& ValidSection, const FStaticMeshFilteredAreaWeightedSectionSampler& SectionSampler);
+
+	virtual void InitRHI() override;
+	virtual void ReleaseRHI() override;
+
+	virtual FString GetFriendlyName() const override { return TEXT("FStaticMeshGpuSpawnBuffer"); }
+
+	FShaderResourceViewRHIParamRef GetBufferSectionSRV() const { return BufferSectionSRV; }
+	uint32 GetValidSectionCount() const { return ValidSections.Num(); };
+
+	FShaderResourceViewRHIParamRef GetBufferPositionSRV() const { return MeshVertexBufferSrv; }
+	FShaderResourceViewRHIParamRef GetBufferTangentSRV() const { return MeshTangentBufferSrv; }
+	FShaderResourceViewRHIParamRef GetBufferTexCoordSRV() const { return MeshTexCoordBufferSrv; }
+	FShaderResourceViewRHIParamRef GetBufferIndexSRV() const { return MeshIndexBufferSrv; }
+	FShaderResourceViewRHIParamRef GetBufferColorSRV() const { return MeshColorBufferSRV; }
+
+	FShaderResourceViewRHIParamRef GetBufferUniformTriangleSamplingSRV() const { return BufferUniformTriangleSamplingSRV; }
+
+	uint32 GetNumTexCoord() const { return NumTexCoord; }
+
+protected:
+
+	// We could separate probabilities from the triangle information when UE supports R32G32 buffer. For pack it all in a uint RGBA32 format.
+	struct SectionInfo
+	{
+		uint32 FirstIndex;
+		uint32 NumTriangles;
+		float  Prob;
+		uint32 Alias;
+	};
+
+	// Cached pointer to Section render data used for initialization only.
+	const FStaticMeshLODResources* SectionRenderData = nullptr;
+
+	TArray<SectionInfo> ValidSections;					// Only the section we want to spawn from
+
+	FVertexBufferRHIRef BufferSectionRHI = nullptr;
+	FShaderResourceViewRHIRef BufferSectionSRV = nullptr;
+
+	FShaderResourceViewRHIRef BufferUniformTriangleSamplingSRV = nullptr;
+
+	// Cached SRV to gpu buffers of the mesh we spawn from 
+	FShaderResourceViewRHIRef MeshIndexBufferSrv;
+	FShaderResourceViewRHIRef MeshVertexBufferSrv;
+	FShaderResourceViewRHIRef MeshTangentBufferSrv;
+	FShaderResourceViewRHIRef MeshTexCoordBufferSrv;
+	uint32 NumTexCoord;
+	FShaderResourceViewRHIRef MeshColorBufferSRV;
+};
+
 struct FNDIStaticMesh_InstanceData
 {
 	 //Cached ptr to component we sample from. 
@@ -60,8 +118,12 @@ struct FNDIStaticMesh_InstanceData
 	/** Time separating Transform and PrevTransform. */
 	float DeltaSeconds;
 
-	/** True if the mesh we're using allows area weighted sampling. */
-	uint32 bIsAreaWeightedSampling : 1;
+	/** True if the mesh allows CPU access. Use to reset the instance in the editor*/
+	uint32 bMeshAllowsCpuAccess : 1;
+	/** True if the mesh we're using allows area weighted sampling on CPU. */
+	uint32 bIsCpuUniformlyDistributedSampling : 1;
+	/** True if the mesh we're using allows area weighted sampling on GPU. */
+	uint32 bIsGpuUniformlyDistributedSampling : 1;
 
 	/** Cached results of this filter being applied to the owning mesh. */
 	TArray<int32> ValidSections;
@@ -75,7 +137,7 @@ struct FNDIStaticMesh_InstanceData
 	uint32 ChangeId;
 
 	FORCEINLINE UStaticMesh* GetActualMesh()const { return Mesh; }
-	FORCEINLINE bool UsesAreaWeighting()const { return bIsAreaWeightedSampling; }
+	FORCEINLINE bool UsesCpuUniformlyDistributedSampling() const { return bIsCpuUniformlyDistributedSampling; }
 	FORCEINLINE bool MeshHasPositions()const { return Mesh && Mesh->RenderData->LODResources[0].VertexBuffers.PositionVertexBuffer.GetNumVertices() > 0; }
 	FORCEINLINE bool MeshHasVerts()const { return Mesh && Mesh->RenderData->LODResources[0].VertexBuffers.StaticMeshVertexBuffer.GetNumVertices() > 0; }
 	FORCEINLINE bool MeshHasColors()const { return Mesh && Mesh->RenderData->LODResources[0].VertexBuffers.ColorVertexBuffer.GetNumVertices() > 0; }
@@ -89,6 +151,7 @@ struct FNDIStaticMesh_InstanceData
 
 	FORCEINLINE_DEBUGGABLE bool Init(UNiagaraDataInterfaceStaticMesh* Interface, FNiagaraSystemInstance* SystemInstance);
 	FORCEINLINE_DEBUGGABLE bool Tick(UNiagaraDataInterfaceStaticMesh* Interface, FNiagaraSystemInstance* SystemInstance, float InDeltaSeconds);
+	FORCEINLINE_DEBUGGABLE void Release();
 };
 
 /** Data Interface allowing sampling of static meshes. */
@@ -107,6 +170,10 @@ public:
 	UPROPERTY(EditAnywhere, Category = "Mesh")
 	AActor* Source;
 	
+	/** The source component from which to sample. Takes precedence over the direct mesh. Not exposed to the user, only indirectly accessible from blueprints. */
+	UPROPERTY(Transient)
+	UStaticMeshComponent* SourceComponent;
+
 	/** Array of filters the can be used to limit sampling to certain sections of the mesh. */
 	UPROPERTY(EditAnywhere, Category = "Mesh")
 	FNDIStaticMeshSectionFilter SectionFilter;
@@ -118,8 +185,8 @@ public:
 
 	//~ UObject interface
 
-#if WITH_EDITOR
 	virtual void PostInitProperties()override;
+#if WITH_EDITOR
 	virtual void PostEditChangeProperty(struct FPropertyChangedEvent& PropertyChangedEvent) override;
 #endif
 
@@ -135,10 +202,33 @@ public:
 	virtual void GetFunctions(TArray<FNiagaraFunctionSignature>& OutFunctions)override;
 	virtual void GetVMExternalFunction(const FVMExternalFunctionBindingInfo& BindingInfo, void* InstanceData, FVMExternalFunction &OutFunc) override;
 	virtual bool Equals(const UNiagaraDataInterface* Other) const override;
-	virtual bool CanExecuteOnTarget(ENiagaraSimTarget Target)const override { return Target == ENiagaraSimTarget::CPUSim; }
+	virtual bool CanExecuteOnTarget(ENiagaraSimTarget Target)const override { return true; }
 #if WITH_EDITOR
 	virtual TArray<FNiagaraDataInterfaceError> GetErrors() override;
 #endif
+
+	virtual bool GetFunctionHLSL(const FName&  DefinitionFunctionName, FString InstanceFunctionName, FNiagaraDataInterfaceGPUParamInfo& ParamInfo, FString& OutHLSL) override;
+	virtual void GetParameterDefinitionHLSL(FNiagaraDataInterfaceGPUParamInfo& ParamInfo, FString& OutHLSL) override;
+	virtual FNiagaraDataInterfaceParametersCS* ConstructComputeParameters()const override;
+
+	virtual void ProvidePerInstanceDataForRenderThread(void* DataForRenderThread, void* PerInstanceData, const FGuid& SystemInstance) override;
+
+	static const FString MeshIndexBufferName;
+	static const FString MeshVertexBufferName;
+	static const FString MeshTangentBufferName;
+	static const FString MeshTexCoordBufferName;
+	static const FString MeshColorBufferName;
+	static const FString MeshSectionBufferName;
+	static const FString MeshTriangleBufferName;
+	static const FString SectionCountName;
+	static const FString InstanceTransformName;
+	static const FString InstanceTransformInverseTransposedName;
+	static const FString InstancePrevTransformName;
+	static const FString InstanceInvDeltaTimeName;
+	static const FString InstanceWorldVelocityName;
+	static const FString AreaWeightedSamplingName;
+	static const FString NumTexCoordName;
+
 public:
 	void GetNumTriangles(FVectorVMContext& Context);
 
@@ -176,6 +266,9 @@ public:
 	template<typename TransformHandlerType>
 	void GetVertexPosition(FVectorVMContext& Context);
 
+	void SetSourceComponentFromBlueprints(UStaticMeshComponent* ComponentToUse);
+	void SetDefaultMeshFromBlueprints(UStaticMesh* MeshToUse);
+
 	FORCEINLINE_DEBUGGABLE bool UsesSectionFilter()const { return SectionFilter.CanEverReject(); }
 
 	//TODO: Vertex color filtering requires a bit more work.
@@ -224,3 +317,56 @@ public:
 	/** Todo: Find a place to call this on level change or somewhere. */
 	static void CleanupDynamicColorFilterData();
 };
+
+struct FNiagaraStaticMeshData
+{
+	FNiagaraStaticMeshData()
+		: MeshGpuSpawnBuffer(nullptr)
+		, bIsGpuUniformlyDistributedSampling(false)
+		, DeltaSeconds(0.03333f)
+	{}
+
+	~FNiagaraStaticMeshData()
+	{
+		check(IsInRenderingThread());
+		if (MeshGpuSpawnBuffer)
+		{
+			MeshGpuSpawnBuffer->ReleaseResource();
+			delete MeshGpuSpawnBuffer;
+		}
+	}
+
+	/** Extra mesh data upload to GPU to do uniform sampling of sections and triangles.*/
+	FStaticMeshGpuSpawnBuffer* MeshGpuSpawnBuffer;
+	bool bIsGpuUniformlyDistributedSampling;
+	FMatrix Transform;
+	FMatrix PrevTransform;
+	float DeltaSeconds;
+};
+
+struct FNiagaraPassedInstanceDataForRT
+{
+	bool bIsGpuUniformlyDistributedSampling;
+	FMatrix Transform;
+	FMatrix PrevTransform;
+	float DeltaSeconds;
+};
+
+struct FNiagaraDataInterfaceProxyStaticMesh : public FNiagaraDataInterfaceProxy
+{
+	virtual void DeferredDestroy() override;
+
+	virtual int32 PerInstanceDataPassedToRenderThreadSize() const override
+	{
+		return sizeof(FNiagaraPassedInstanceDataForRT);
+	}
+
+	virtual void ConsumePerInstanceDataFromGameThread(void* PerInstanceData, const FGuid& Instance) override;
+
+	void InitializePerInstanceData(const FGuid& SystemInstance, FStaticMeshGpuSpawnBuffer* MeshGPUSpawnBuffer);
+	void DestroyPerInstanceData(NiagaraEmitterInstanceBatcher* Batcher, const FGuid& SystemInstance);
+
+	TMap<FGuid, FNiagaraStaticMeshData> SystemInstancesToMeshData;
+
+	TSet<FGuid> DeferredDestroyList;
+};
diff --git a/Engine/Plugins/FX/Niagara/Source/Niagara/Classes/NiagaraDataInterfaceTexture.h b/Engine/Plugins/FX/Niagara/Source/Niagara/Classes/NiagaraDataInterfaceTexture.h
--- a/Engine/Plugins/FX/Niagara/Source/Niagara/Classes/NiagaraDataInterfaceTexture.h
+++ b/Engine/Plugins/FX/Niagara/Source/Niagara/Classes/NiagaraDataInterfaceTexture.h
@@ -12,7 +12,6 @@ UCLASS(EditInlineNew, Category = "Texture", meta = (DisplayName = "Texture Sampl
 class NIAGARA_API UNiagaraDataInterfaceTexture : public UNiagaraDataInterface
 {
 	GENERATED_UCLASS_BODY()
-		bool bGPUBufferDirty;
 public:
 
 	UPROPERTY(EditAnywhere, Category = "Texture")
@@ -51,9 +50,24 @@ public:
 protected:
 	virtual bool CopyToInternal(UNiagaraDataInterface* Destination) const override;
 
+	void PushToRenderThread();
+
 
 	static const FName SampleTexture2DName;
 	static const FName SampleVolumeTextureName;
 	static const FName SamplePseudoVolumeTextureName;
 	static const FName TextureDimsName;
 };
+
+struct FNiagaraDataInterfaceProxyTexture : public FNiagaraDataInterfaceProxy
+{
+	FSamplerStateRHIRef SamplerStateRHI;
+	FTextureRHIRef TextureRHI;
+	float TexDims[2];
+
+	virtual void ConsumePerInstanceDataFromGameThread(void* PerInstanceData, const FGuid& Instance) override { check(false); }
+	virtual int32 PerInstanceDataPassedToRenderThreadSize() const override
+	{
+		return 0;
+	}
+};
\ No newline at end of file
diff --git a/Engine/Plugins/FX/Niagara/Source/Niagara/Classes/NiagaraDataInterfaceVectorField.h b/Engine/Plugins/FX/Niagara/Source/Niagara/Classes/NiagaraDataInterfaceVectorField.h
--- a/Engine/Plugins/FX/Niagara/Source/Niagara/Classes/NiagaraDataInterfaceVectorField.h
+++ b/Engine/Plugins/FX/Niagara/Source/Niagara/Classes/NiagaraDataInterfaceVectorField.h
@@ -64,9 +64,32 @@ public:
 	FVector GetDimensions() const;
 	FVector GetMinBounds() const;
 	FVector GetMaxBounds() const;
-
 protected:
 	//~ UNiagaraDataInterface interface
 	virtual bool CopyToInternal(UNiagaraDataInterface* Destination) const override;
 	//~ UNiagaraDataInterface interface END
+
+	void PushToRenderThread();
 };
+
+struct FNiagaraDataInterfaceProxyVectorField : public FNiagaraDataInterfaceProxy
+{
+	FVector Dimensions;
+	FVector MinBounds;
+	FVector MaxBounds;
+	bool bTileX;
+	bool bTileY;
+	bool bTileZ;
+	FTextureRHIRef TextureRHI;
+
+	FVector GetTilingAxes() const
+	{
+		return FVector(float(bTileX), float(bTileY), float(bTileZ));
+	}
+
+	virtual void ConsumePerInstanceDataFromGameThread(void* PerInstanceData, const FGuid& Instance) override { check(false); }
+	virtual int32 PerInstanceDataPassedToRenderThreadSize() const override
+	{
+		return 0;
+	}
+};
\ No newline at end of file
