Copying //UE4/Dev-Niagara to Dev-Main (//UE4/Dev-Main)

diff --git a/Engine/Plugins/FX/Niagara/Source/Niagara/Classes/NiagaraDataSet.h b/Engine/Plugins/FX/Niagara/Source/Niagara/Classes/NiagaraDataSet.h
--- a/Engine/Plugins/FX/Niagara/Source/Niagara/Classes/NiagaraDataSet.h
+++ b/Engine/Plugins/FX/Niagara/Source/Niagara/Classes/NiagaraDataSet.h
@@ -25,128 +25,136 @@ struct FNiagaraVariableLayoutInfo
 };
 
 class FNiagaraDataSet;
+class FNiagaraShader;
 
-/** Buffer containing one frame of Niagara simulation data. */
-class NIAGARA_API FNiagaraDataBuffer
+//Base class for objects in Niagara that are owned by one object but are then passed for reading to other objects, potentially on other threads.
+//This class allows us to know if the object is being used so we do not overwrite it and to ensure it's lifetime so we do not access freed data.
+class FNiagaraSharedObject
 {
 public:
-	FNiagaraDataBuffer() : Owner(nullptr)
-	{
-		Reset();
-	}
-	~FNiagaraDataBuffer();
-	void Init(FNiagaraDataSet* InOwner);
-	void Allocate(uint32 NumInstances, bool bMaintainExisting=false);
-	void AllocateGPU(uint32 InNumInstances, FRHICommandList &RHICmdList);
-	void SwapInstances(uint32 OldIndex, uint32 NewIndex);
-	void KillInstance(uint32 InstanceIdx);
-	void CopyTo(FNiagaraDataBuffer& DestBuffer, int32 StartIdx, int32 NumInstances)const;
-	void CopyTo(FNiagaraDataBuffer& DestBuffer)const;
-	void GPUCopyTo(FNiagaraDataBuffer& DestBuffer, float* GPUReadBackFloat, int* GPUReadBackInt, int32 StartIdx, int32 NumInstances)const;
-
-	const TArray<uint8>& GetFloatBuffer()const { return FloatData; }
-	const TArray<uint8>& GetInt32Buffer()const { return Int32Data; }
+	FNiagaraSharedObject()
+		: ReadRefCount(0)
+	{}
 
-	FORCEINLINE const uint8* GetComponentPtrFloat(uint32 ComponentIdx)const
+	/** The owner of this object is now done with it but it may still be in use by others, possibly on other threads. Add to the deletion queue so it can be safely freed when it's no longer in use. */
+	FORCEINLINE void Destroy()
 	{
-		return FloatData.GetData() + FloatStride * ComponentIdx;
+		FScopeLock Lock(&CritSec);
+		check(DeferredDeletionList.Contains(this) == false);
+		DeferredDeletionList.Add(this);
 	}
+	static void FlushDeletionList();
 
-	FORCEINLINE const uint8* GetComponentPtrInt32(uint32 ComponentIdx)const
-	{
-		return Int32Data.GetData() + Int32Stride * ComponentIdx;
-	}
+	FORCEINLINE bool IsInUse()const { return ReadRefCount.Load() != 0; }
+	FORCEINLINE bool IsBeingRead()const { return ReadRefCount.Load() > 0; }
+	FORCEINLINE bool IsBeingWritten()const { return ReadRefCount.Load() == INDEX_NONE; }
 
-	FORCEINLINE uint8* GetComponentPtrFloat(uint32 ComponentIdx)
+	FORCEINLINE void AddReadRef()
 	{
-		return FloatData.GetData() + FloatStride * ComponentIdx;
+		check(!IsBeingWritten());
+		ReadRefCount++;
 	}
 
-	FORCEINLINE uint8* GetComponentPtrInt32(uint32 ComponentIdx)
+	FORCEINLINE void ReleaseReadRef()
 	{
-		return Int32Data.GetData() + Int32Stride * ComponentIdx;
+		check(IsBeingRead());
+		ReadRefCount--;
 	}
 
-	FORCEINLINE float* GetInstancePtrFloat(uint32 ComponentIdx, uint32 InstanceIdx)
+	FORCEINLINE bool TryLock()
 	{
-		return (float*)GetComponentPtrFloat(ComponentIdx) + InstanceIdx;
+		//Only lock if we have no readers.
+		//Using INDEX_NONE as a special case value for write locks.
+		int32 Expected = 0;
+		return ReadRefCount.CompareExchange(Expected, INDEX_NONE);
 	}
 
-	FORCEINLINE int32* GetInstancePtrInt32(uint32 ComponentIdx, uint32 InstanceIdx)
+	FORCEINLINE void Unlock()
 	{
-		return (int32*)GetComponentPtrInt32(ComponentIdx) + InstanceIdx;
+		int32 Expected = INDEX_NONE;
+		ensureAlwaysMsgf(ReadRefCount.CompareExchange(Expected, 0), TEXT("Trying to release a write lock on a Niagara shared object that is not locked for write."));
 	}
 
-	FORCEINLINE float* GetInstancePtrFloat(uint32 ComponentIdx, uint32 InstanceIdx)const
-	{
-		return (float*)GetComponentPtrFloat(ComponentIdx) + InstanceIdx;
-	}
+protected:
 
-	FORCEINLINE int32* GetInstancePtrInt32(uint32 ComponentIdx, uint32 InstanceIdx)const
-	{
-		return (int32*)GetComponentPtrInt32(ComponentIdx) + InstanceIdx;
-	}
+	/**
+	Count of other object currently reading this data. Keeps us from writing to or deleting this data while it's in use. These reads can be on any thread so atomic is used.
+	INDEX_NONE used as special case marking this object as locked for write.
+	*/
+	TAtomic<int32> ReadRefCount;
 
-	FORCEINLINE uint8* GetComponentPtrFloat(float* BasePtr, uint32 ComponentIdx) const
-	{
-		return (uint8*)BasePtr + FloatStride * ComponentIdx;
-	}
+	static FCriticalSection CritSec;
+	static TArray<FNiagaraSharedObject*> DeferredDeletionList;
 
-	FORCEINLINE uint8* GetComponentPtrInt32(int* BasePtr, uint32 ComponentIdx) const
-	{
-		return (uint8*)BasePtr + Int32Stride * ComponentIdx;
-	}
+	virtual ~FNiagaraSharedObject() {}
+};
 
-	FORCEINLINE float* GetInstancePtrFloat(float* BasePtr, uint32 ComponentIdx, uint32 InstanceIdx)const
-	{
-		return (float*)GetComponentPtrFloat(BasePtr, ComponentIdx) + InstanceIdx;
-	}
+/** Buffer containing one frame of Niagara simulation data. */
+class NIAGARA_API FNiagaraDataBuffer : public FNiagaraSharedObject
+{
+	friend class FScopedNiagaraDataSetGPUReadback;
+protected:
+	virtual ~FNiagaraDataBuffer();
 
-	FORCEINLINE int32* GetInstancePtrInt32(int* BasePtr, uint32 ComponentIdx, uint32 InstanceIdx)const
-	{
-		return (int32*)GetComponentPtrInt32(BasePtr, ComponentIdx) + InstanceIdx;
-	}
+public:
+	FNiagaraDataBuffer(FNiagaraDataSet* InOwner);
+	void Allocate(uint32 NumInstances, bool bMaintainExisting = false);
+	void AllocateGPU(uint32 InNumInstances, FRHICommandList &RHICmdList);
+	void SwapInstances(uint32 OldIndex, uint32 NewIndex);
+	void KillInstance(uint32 InstanceIdx);
+	void CopyTo(FNiagaraDataBuffer& DestBuffer, int32 StartIdx, int32 NumInstances)const;
+	void CopyTo(FNiagaraDataBuffer& DestBuffer)const;
+	void GPUCopyFrom(float* GPUReadBackFloat, int* GPUReadBackInt, int32 StartIdx, int32 NumInstances);
+	void Dump(int32 StartIndex, int32 NumInstances, const FString& Label)const;
 
-	uint32 GetNumInstances()const { return NumInstances; }
-	uint32 GetNumInstancesAllocated()const { return NumInstancesAllocated; }
+	bool AppendToRegisterTable(uint8** Registers, int32& NumRegisters, int32 StartInstance);
 
-	void SetNumInstances(uint32 InNumInstances) 
-	{ 
-		NumInstances = InNumInstances;
-	}
+	const TArray<uint8>& GetFloatBuffer()const { return FloatData; }
+	const TArray<uint8>& GetInt32Buffer()const { return Int32Data; }
 
-	void Reset();
+	FORCEINLINE const uint8* GetComponentPtrFloat(uint32 ComponentIdx)const	{ return FloatData.GetData() + FloatStride * ComponentIdx; }
+	FORCEINLINE const uint8* GetComponentPtrInt32(uint32 ComponentIdx)const	{ return Int32Data.GetData() + Int32Stride * ComponentIdx; }
+	FORCEINLINE uint8* GetComponentPtrFloat(uint32 ComponentIdx) { return FloatData.GetData() + FloatStride * ComponentIdx;	}
+	FORCEINLINE uint8* GetComponentPtrInt32(uint32 ComponentIdx) { return Int32Data.GetData() + Int32Stride * ComponentIdx;	}
 
-	FORCEINLINE uint32 GetSizeBytes()const
-	{
-		return FloatData.Num() + Int32Data.Num();
-	}
+	FORCEINLINE float* GetInstancePtrFloat(uint32 ComponentIdx, uint32 InstanceIdx)	{ return (float*)GetComponentPtrFloat(ComponentIdx) + InstanceIdx; }
+	FORCEINLINE int32* GetInstancePtrInt32(uint32 ComponentIdx, uint32 InstanceIdx)	{ return (int32*)GetComponentPtrInt32(ComponentIdx) + InstanceIdx; }
+	FORCEINLINE float* GetInstancePtrFloat(uint32 ComponentIdx, uint32 InstanceIdx)const { return (float*)GetComponentPtrFloat(ComponentIdx) + InstanceIdx; }
+	FORCEINLINE int32* GetInstancePtrInt32(uint32 ComponentIdx, uint32 InstanceIdx)const { return (int32*)GetComponentPtrInt32(ComponentIdx) + InstanceIdx;	}
 
-	const FRWBuffer *GetGPUBufferFloat() const
-	{
-		return &GPUBufferFloat;
-	}
+	FORCEINLINE uint8* GetComponentPtrFloat(float* BasePtr, uint32 ComponentIdx) const { return (uint8*)BasePtr + FloatStride * ComponentIdx; }
+	FORCEINLINE uint8* GetComponentPtrInt32(int* BasePtr, uint32 ComponentIdx) const { return (uint8*)BasePtr + Int32Stride * ComponentIdx; }
+	FORCEINLINE float* GetInstancePtrFloat(float* BasePtr, uint32 ComponentIdx, uint32 InstanceIdx)const { return (float*)GetComponentPtrFloat(BasePtr, ComponentIdx) + InstanceIdx; }
+	FORCEINLINE int32* GetInstancePtrInt32(int* BasePtr, uint32 ComponentIdx, uint32 InstanceIdx)const { return (int32*)GetComponentPtrInt32(BasePtr, ComponentIdx) + InstanceIdx; }
 
-	const FRWBuffer *GetGPUBufferInt() const
-	{
-		return &GPUBufferInt;
-	}
+	FORCEINLINE uint32 GetNumInstances()const { return NumInstances; }
+	FORCEINLINE uint32 GetNumInstancesAllocated()const { return NumInstancesAllocated; }
 
-	int32 GetSafeComponentBufferSize() const
-	{
-		return GetSafeComponentBufferSize(GetNumInstancesAllocated());
-	}
+	FORCEINLINE void SetNumInstances(uint32 InNumInstances) { NumInstances = InNumInstances; }
+	FORCEINLINE uint32 GetSizeBytes()const { return FloatData.Num() + Int32Data.Num(); }
+	FORCEINLINE FRWBuffer& GetGPUBufferFloat() { return GPUBufferFloat; }
+	FORCEINLINE FRWBuffer& GetGPUBufferInt() { return GPUBufferInt;	}
+	FORCEINLINE FRWBuffer& GetGPUIndices() { return GPUIndices; }
 
-	uint32 GetFloatStride() const { return FloatStride; }
-	uint32 GetInt32Stride() const { return Int32Stride; }
+	FORCEINLINE int32 GetSafeComponentBufferSize() const { return GetSafeComponentBufferSize(GetNumInstancesAllocated()); }
+	FORCEINLINE uint32 GetFloatStride() const { return FloatStride; }
+	FORCEINLINE uint32 GetInt32Stride() const { return Int32Stride; }
 
-	FORCEINLINE const FNiagaraDataSet* GetOwner()const { return Owner; }
+	FORCEINLINE FNiagaraDataSet* GetOwner()const { return Owner; }
 
 	int32 TransferInstance(FNiagaraDataBuffer& SourceBuffer, int32 InstanceIndex);
 
 	bool CheckForNaNs()const;
+
+	FORCEINLINE TArray<int32>& GetIDTable() { return IDToIndexTable; }
+
+	template<bool bDoResourceTransitions>
+	void SetShaderParams(class FNiagaraShader *Shader, FRHICommandList &CommandList, bool bInput);
+	void UnsetShaderParams(class FNiagaraShader *Shader, FRHICommandList &CommandList);
 private:
 
+	void CheckUsage(bool bReadOnly)const;
+
 	FORCEINLINE int32 GetSafeComponentBufferSize(int32 RequiredSize) const
 	{
 		//Round up to VECTOR_WIDTH_BYTES.
@@ -157,25 +165,36 @@ private:
 	/** Back ptr to our owning data set. Used to access layout info for the buffer. */
 	FNiagaraDataSet* Owner;
 
+	//////////////////////////////////////////////////////////////////////////
+	//CPU Data
 	/** Float components of simulation data. */
 	TArray<uint8> FloatData;
 	/** Int32 components of simulation data. */
 	TArray<uint8> Int32Data;
 
-	/** Stride between components in the float buffer. */
-	uint32 FloatStride;
-	/** Stride between components in the int32 buffer. */
-	uint32 Int32Stride;
+	/** Table of IDs to real buffer indices. */
+	TArray<int32> IDToIndexTable;
+	//////////////////////////////////////////////////////////////////////////
 
+	//////////////////////////////////////////////////////////////////////////
+	// GPU Data
 	uint32 NumChunksAllocatedForGPU;
+	/** GPU Buffer containing floating point values for GPU simulations. */
+	FRWBuffer GPUBufferFloat;
+	/** GPU Buffer containing floating point values for GPU simulations. */
+	FRWBuffer GPUBufferInt;
+	/** GPU Buffer containing args to allow indirect rendering of this buffer. */
+	FRWBuffer GPUIndices;
+	//////////////////////////////////////////////////////////////////////////
 
 	/** Number of instances in data. */
 	uint32 NumInstances;
 	/** Number of instances the buffer has been allocated for. */
 	uint32 NumInstancesAllocated;
-
-	FRWBuffer GPUBufferFloat;
-	FRWBuffer GPUBufferInt;
+	/** Stride between components in the float buffer. */
+	uint32 FloatStride;
+	/** Stride between components in the int32 buffer. */
+	uint32 Int32Stride;
 };
 
 //////////////////////////////////////////////////////////////////////////
@@ -186,464 +205,91 @@ General storage class for all per instance simulation data in Niagara.
 class NIAGARA_API FNiagaraDataSet
 {
 	friend FNiagaraDataBuffer;
-
-	void Reset()
-	{
-		Variables.Empty();
-		VariableLayouts.Empty();
-		CurrBuffer = 0;
-		bFinalized = false;
-		TotalFloatComponents = 0;
-		TotalInt32Components = 0;
-		MaxBufferIdx = 1;
-		bNeedsPersistentIDs = 0;
-
-		SimTarget = ENiagaraSimTarget::CPUSim;
-
-		ResetBuffersInternal();
-	}
-
 public:
-	FNiagaraDataSet()
-	{
-		Reset();
-	}
-	
-	~FNiagaraDataSet();
 
-	void Init(FNiagaraDataSetID InID, ENiagaraSimTarget InSimTarget)
-	{
-		Reset();
-		ID = InID;
-		SimTarget = InSimTarget;
-	}
-
-
-	void AddVariable(FNiagaraVariable& Variable)
-	{
-		check(!bFinalized);
-		Variables.AddUnique(Variable);
-	}
-
-	void AddVariables(const TArray<FNiagaraVariable>& Vars)
-	{
-		check(!bFinalized);
-		for (const FNiagaraVariable& Var : Vars)
-		{
-			Variables.AddUnique(Var);
-		}
-	}
-
-	FORCEINLINE void SetNeedsPersistentIDs(bool bNeedsIDs)
-	{
-		bNeedsPersistentIDs = bNeedsIDs;
-	}
-
-	FORCEINLINE bool GetNeedsPersistentIDs()const { return bNeedsPersistentIDs; }
-
-	/** Finalize the addition of variables and other setup before this data set can be used. */
-	FORCEINLINE void Finalize()
-	{
-		check(!bFinalized);
-		bFinalized = true;
-		BuildLayout();
-	}
-
-	/** Removes a specific instance from the current frame's data buffer. */
-	FORCEINLINE void KillInstance(uint32 InstanceIdx)
-	{
-		check(bFinalized);
-		CheckCorrectThread();
-		CurrData().KillInstance(InstanceIdx);
-	}
-
-	FORCEINLINE void SwapInstances(uint32 OldIndex, uint32 NewIndex)
-	{
-		check(bFinalized);
-		CheckCorrectThread();
-		PrevData().SwapInstances(OldIndex, NewIndex);
-	}
-
-	int32 TransferInstance(FNiagaraDataSet& SourceDataset, int32 InstanceIndex)
-	{
-		check(bFinalized);
-		CheckCorrectThread();
-		return CurrData().TransferInstance(SourceDataset.CurrData(), InstanceIndex);
-	}
-
-	/** Appends all variables in this dataset to a register table ready for execution by the VectorVM. */
-	bool AppendToRegisterTable(uint8** InputRegisters, int32& NumInputRegisters, uint8** OutputRegisters, int32& NumOutputRegisters, int32 StartInstance)
-	{
-		check(bFinalized);
-		CheckCorrectThread();
-		int32 TotalComponents = GetNumFloatComponents() + GetNumInt32Components();
-		if (NumInputRegisters + TotalComponents > VectorVM::MaxInputRegisters || NumOutputRegisters + TotalComponents > VectorVM::MaxOutputRegisters)
-		{
-			UE_LOG(LogNiagara, Error, TEXT("Niagara Script is using too many IO registers!"));
-			return false;
-		}
-		else
-		{
-			for (FNiagaraVariableLayoutInfo& VarLayout : VariableLayouts)
-			{
-				int32 NumFloats = VarLayout.GetNumFloatComponents();
-				int32 NumInts = VarLayout.GetNumInt32Components();
-				for (int32 CompIdx = 0; CompIdx < NumFloats; ++CompIdx)
-				{
-					uint32 CompBufferOffset = VarLayout.FloatComponentStart + CompIdx;
-					uint32 CompRegisterOffset = VarLayout.LayoutInfo.FloatComponentRegisterOffsets[CompIdx];
-					InputRegisters[NumInputRegisters + CompRegisterOffset] = (uint8*)PrevData().GetInstancePtrFloat(CompBufferOffset, StartInstance);
-					OutputRegisters[NumOutputRegisters + CompRegisterOffset] = (uint8*)CurrData().GetInstancePtrFloat(CompBufferOffset, StartInstance);
-				}
-				for (int32 CompIdx = 0; CompIdx < NumInts; ++CompIdx)
-				{
-					uint32 CompBufferOffset = VarLayout.Int32ComponentStart + CompIdx;
-					uint32 CompRegisterOffset = VarLayout.LayoutInfo.Int32ComponentRegisterOffsets[CompIdx];
-					InputRegisters[NumInputRegisters + CompRegisterOffset] = (uint8*)PrevData().GetInstancePtrInt32(CompBufferOffset, StartInstance);
-					OutputRegisters[NumOutputRegisters + CompRegisterOffset] = (uint8*)CurrData().GetInstancePtrInt32(CompBufferOffset, StartInstance);
-				}
-				NumInputRegisters += NumFloats + NumInts;
-				NumOutputRegisters += NumFloats + NumInts;
-			}
-			return true;
-		}
-	}
+	FNiagaraDataSet();
+	~FNiagaraDataSet();
+	FNiagaraDataSet& operator=(const FNiagaraDataSet&) = delete;
 
-	void SetShaderParams(class FNiagaraShader *Shader, FRHICommandList &CommandList, uint32& WriteBufferIdx, uint32& ReadBufferIdx);
-	void UnsetShaderParams(class FNiagaraShader *Shader, FRHICommandList &CommandList);
-	void Allocate(int32 NumInstances, bool bMaintainExisting=false)
-	{
-		check(bFinalized);
-		CheckCorrectThread();
-		CurrData().Allocate(NumInstances, bMaintainExisting);
+	void Init(FNiagaraDataSetID InID, ENiagaraSimTarget InSimTarget, const FString& InDebugName);
 
-#if NIAGARA_NAN_CHECKING
-		CheckForNaNs();
-#endif
+	/** Resets current data but leaves variable/layout information etc intact. */
+	void ResetBuffers();
 
-		if (bNeedsPersistentIDs)
-		{
-			int32 NumUsedIDs = MaxUsedID + 1;
+	/** Begins a new simulation pass and grabs a destination buffer. Returns the new destination data buffer. */
+	FNiagaraDataBuffer& BeginSimulate();
 
-			int32 RequiredIDs = FMath::Max(NumInstances, NumUsedIDs);
-			int32 ExistingNumIDs = PrevIDTable().Num();
-			int32 NumNewIDs = RequiredIDs - ExistingNumIDs;
+	/** Ends a simulation pass and sets the current simulation state. */
+	void EndSimulate();
 
-			if (RequiredIDs > ExistingNumIDs)
-			{
-				//UE_LOG(LogNiagara, Warning, TEXT("Growing ID Table! OldSize:%d | NewSize:%d"), ExistingNumIDs, RequiredIDs);
-				IDToIndexTable[CurrBuffer].SetNumUninitialized(RequiredIDs);
+	/** Allocates space for NumInstances in the current destination buffer. */
+	void Allocate(int32 NumInstances, bool bMaintainExisting = false);
 
-				//Free ID Table must always be at least as large as the data buffer + it's current size in the case all particles die this frame.
-				FreeIDsTable.AddUninitialized(NumNewIDs);
+	/** Returns size in bytes for all data buffers currently allocated by this dataset. */
+	uint32 GetSizeBytes()const;
 
-				//Free table should always have enough room for these new IDs.
-				check(NumFreeIDs + NumNewIDs <= FreeIDsTable.Num());
-
-				//ID Table grows so add any new IDs to the free array. Add in reverse order to maintain a continuous increasing allocation when popping.
-				for (int32 NewFreeID = RequiredIDs - 1; NewFreeID >= ExistingNumIDs; --NewFreeID)
-				{
-					FreeIDsTable[NumFreeIDs++] = NewFreeID ;
-				}
-				//UE_LOG(LogNiagara, Warning, TEXT("DataSetAllocate: Adding New Free IDs: %d - "), RequiredIDs - 1, ExistingNumIDs);
-			}
-#if 0
-			else if (RequiredIDs < ExistingNumIDs >> 1)//Configurable?
-			{
-				//If the max id we use has reduced significantly then we can shrink the tables.
-				//Have to go through the FreeIDs and remove any that are greater than the new table size.
-				//UE_LOG(LogNiagara, Warning, TEXT("DataSetAllocate: Shrinking ID Table! OldSize:%d | NewSize:%d"), ExistingNumIDs, RequiredIDs);
-				for (int32 CheckedFreeID = 0; CheckedFreeID < NumFreeIDs;)
-				{
-					checkSlow(NumFreeIDs <= FreeIDsTable.Num());
-					if (FreeIDsTable[CheckedFreeID] >= RequiredIDs)
-					{
-						//UE_LOG(LogNiagara, Warning, TEXT("RemoveSwap FreeID: Removed:%d | Swapped:%d"), FreeIDsTable[CheckedFreeID], FreeIDsTable.Last());		
-						int32 FreeIDIndex = --NumFreeIDs;
-						FreeIDsTable[CheckedFreeID] = FreeIDsTable[FreeIDIndex];
-						FreeIDsTable[FreeIDIndex] = INDEX_NONE;
-					}
-					else
-					{
-						++CheckedFreeID;
-					}
-				}
-
-				check(NumFreeIDs <= RequiredIDs);
-				FreeIDsTable.SetNumUninitialized(NumFreeIDs);
-			}
-#endif
-			else
-			{
-				//Drop in required size not great enough so just allocate same size.
-				RequiredIDs = ExistingNumIDs;
-			}
-
-			IDToIndexTable[CurrBuffer].SetNumUninitialized(RequiredIDs);
-			MaxUsedID = INDEX_NONE;//reset the max ID ready for it to be filled in during simulation.
-
-			//UE_LOG(LogNiagara, Warning, TEXT("DataSetAllocate: NumInstances:%d | ID Table Size:%d | NumFreeIDs:%d | FreeTableSize:%d"), NumInstances, IDToIndexTable[CurrBuffer].Num(), NumFreeIDs, FreeIDsTable.Num());
-		}
-	}
-	
-	FORCEINLINE void Tick()
-	{
-#if NIAGARA_NAN_CHECKING
-		CheckForNaNs();
-#endif
-		SwapBuffers();
-	}
-	
-	FORCEINLINE void CopyCurToPrev()
-	{
-		CurrData().CopyTo(PrevData());
-	}
+	void ClearRegisterTable(uint8** Registers, int32& NumRegisters);
 
+	FORCEINLINE bool IsInitialized()const { return bFinalized; }
+	FORCEINLINE ENiagaraSimTarget GetSimTarget()const { return SimTarget; }
 	FORCEINLINE FNiagaraDataSetID GetID()const { return ID; }
 	FORCEINLINE void SetID(FNiagaraDataSetID InID) { ID = InID; }
 
-	FORCEINLINE int32 GetPrevBufferIdx() const
-	{
-		return CurrBuffer > 0 ? CurrBuffer - 1 : MaxBufferIdx;
-	}
-
-	FORCEINLINE int32 GetCurrBufferIdx() const
-	{
-		return CurrBuffer;
-	}
-
-	FORCEINLINE FNiagaraDataBuffer& GetDataByIndex(int32 InIdx)
-	{
-		check(InIdx < 3);
-		//CheckCorrectThread();//TODO: We should be able to enable these checks and have well defined GT/RT ownership but GPU sims fire these a lot currently.
-		return Data[InIdx];
-	}
-
-	FORCEINLINE FNiagaraDataBuffer& CurrData() 
-	{
-		//CheckCorrectThread();//TODO: We should be able to enable these checks and have well defined GT/RT ownership but GPU sims fire these a lot currently.
-		return Data[CurrBuffer]; 
-	}
-	FORCEINLINE FNiagaraDataBuffer& PrevData() 
-	{
-		//CheckCorrectThread();//TODO: We should be able to enable these checks and have well defined GT/RT ownership but GPU sims fire these a lot currently.
-		return Data[GetPrevBufferIdx()];
-	}
-	FORCEINLINE const FNiagaraDataBuffer& CurrData()const 
-	{
-		//CheckCorrectThread();//TODO: We should be able to enable these checks and have well defined GT/RT ownership but GPU sims fire these a lot currently.
-		return Data[CurrBuffer];
-	}
-	FORCEINLINE const FNiagaraDataBuffer& PrevData()const 
-	{
-		//CheckCorrectThread();//TODO: We should be able to enable these checks and have well defined GT/RT ownership but GPU sims fire these a lot currently.
-		return Data[GetPrevBufferIdx()];
-	}
-
-	FORCEINLINE uint32 GetNumInstances()const { return CurrData().GetNumInstances(); }
-	FORCEINLINE uint32 GetNumInstancesAllocated()const { return CurrData().GetNumInstancesAllocated(); }
-	FORCEINLINE void SetNumInstances(uint32 InNumInstances) { CurrData().SetNumInstances(InNumInstances); }
-
-	void ResetCurrentBuffers()
-	{
-		SetNumInstances(0);
-		if (bNeedsPersistentIDs)
-		{
-			CurrIDTable().Empty();
-			FreeIDsTable.Empty();
-			NumFreeIDs = 0;
-			MaxUsedID = INDEX_NONE;
-		}
-	}
-
-	FORCEINLINE TArray<int32>& CurrIDTable()
-	{
-		//check(IsInGameThread());
-		return IDToIndexTable[CurrBuffer];
-	}
-
-	FORCEINLINE TArray<int32>& PrevIDTable()
-	{
-		//check(IsInGameThread());
-		return IDToIndexTable[GetPrevBufferIdx()];
-	}
-
-	FORCEINLINE TArray<int32>& GetFreeIDTable()
-	{
-		return FreeIDsTable;
-	}
-	
-	FORCEINLINE int32& GetNumFreeIDs()
-	{
-		return NumFreeIDs;
-	}
-
-	FORCEINLINE int32& GetMaxUsedID()
-	{
-		return MaxUsedID;
-	}
-
-	FORCEINLINE int32& GetIDAcquireTag()
-	{
-		return IDAcquireTag;
-	}
-
-	FORCEINLINE void SetIDAcquireTag(int32 InTag)
-	{
-		IDAcquireTag = InTag;
-	}
-
-	FORCEINLINE ENiagaraSimTarget GetSimTarget()const { return SimTarget; }
-
-	FORCEINLINE void ResetBuffersInternal()
-	{
-		Data[0].Reset();
-		Data[1].Reset();
-		Data[2].Reset();
-
-		FreeIDsTable.Reset();
-		NumFreeIDs = 0;
-
-		IDToIndexTable[0].Reset();
-		IDToIndexTable[1].Reset();
-		IDToIndexTable[2].Reset();
-		MaxUsedID = INDEX_NONE;
-	}
+	FORCEINLINE void SetNeedsPersistentIDs(bool bNeedsIDs) { bNeedsPersistentIDs = bNeedsIDs; }
+	FORCEINLINE bool GetNeedsPersistentIDs()const { return bNeedsPersistentIDs; }
 
-	FORCEINLINE void ResetBuffers()
-	{
-		if (SimTarget == ENiagaraSimTarget::CPUSim)
-		{
-			ResetBuffersInternal();
-		}
-		else
-		{			
-			FNiagaraDataSet* DataSet = this;
-			ENQUEUE_RENDER_COMMAND(ResetBuffersRT)(
-				[DataSet](FRHICommandListImmediate& RHICmdList)
-				{
-					DataSet->ResetBuffersInternal();
-				}
-			);
-		}
-	}
+	FORCEINLINE TArray<int32>& GetFreeIDTable() { return FreeIDsTable; }
+	FORCEINLINE int32& GetNumFreeIDs() { return NumFreeIDs; }
+	FORCEINLINE int32& GetMaxUsedID() { return MaxUsedID; }
+	FORCEINLINE int32& GetIDAcquireTag() { return IDAcquireTag; }
+	FORCEINLINE void SetIDAcquireTag(int32 InTag) { IDAcquireTag = InTag; }
 
-	FORCEINLINE uint32 GetPrevNumInstances()const { return PrevData().GetNumInstances(); }
 
+	FORCEINLINE TArray<FNiagaraVariable>& GetVariables() { return Variables; }
 	FORCEINLINE uint32 GetNumVariables()const { return Variables.Num(); }
+	FORCEINLINE bool HasVariable(const FNiagaraVariable& Var)const { return Variables.Contains(Var); }
+	FORCEINLINE uint32 GetNumFloatComponents()const { return TotalFloatComponents; }
+	FORCEINLINE uint32 GetNumInt32Components()const { return TotalInt32Components; }
 
-	FORCEINLINE uint32 GetSizeBytes()const
-	{
-		return Data[0].GetSizeBytes() + Data[1].GetSizeBytes();
-	}
-
-	FORCEINLINE bool HasVariable(const FNiagaraVariable& Var)const
-	{
-		return Variables.Contains(Var);
-	}
-
-	FORCEINLINE const FNiagaraVariableLayoutInfo* GetVariableLayout(const FNiagaraVariable& Var)const
-	{
-		int32 VarLayoutIndex = Variables.IndexOfByKey(Var);
-		return VarLayoutIndex != INDEX_NONE ? &VariableLayouts[VarLayoutIndex] : nullptr;
-	}
-
-	// get the float and int component offsets of a variable; if the variable doesn't exist, returns -1
-	//
-	FORCEINLINE const bool GetVariableComponentOffsets(const FNiagaraVariable& Var, int32 &FloatStart, int32 &IntStart) const
-	{
-		const FNiagaraVariableLayoutInfo *Info = GetVariableLayout(Var);
-		if (Info)
-		{
-			FloatStart = Info->FloatComponentStart;
-			IntStart = Info->Int32ComponentStart;
-			return true;
-		}
-
-		FloatStart = -1;
-		IntStart = -1;
-		return false;
-	}
+	void AddVariable(FNiagaraVariable& Variable);
+	void AddVariables(const TArray<FNiagaraVariable>& Vars);
+	/** Finalize the addition of variables and other setup before this data set can be used. */
+	void Finalize();
+	const FNiagaraVariableLayoutInfo* GetVariableLayout(const FNiagaraVariable& Var)const;
+	bool GetVariableComponentOffsets(const FNiagaraVariable& Var, int32 &FloatStart, int32 &IntStart) const;
 
-	void Dump(bool bCurr, int32 StartIdx = 0, int32 NumInstances = INDEX_NONE)const;
-	void Dump(FNiagaraDataSet& Other, bool bCurr, int32 StartIdx = 0, int32 NumInstances = INDEX_NONE)const;
-	void DumpGPU(FNiagaraDataSet& Other, float* GPUReadBackFloat, int* GPUReadBackInt, int32 StartIdx = 0, int32 NumInstances = INDEX_NONE)const;
-	FORCEINLINE const TArray<FNiagaraVariable> &GetVariables() { return Variables; }
+	void CopyTo(FNiagaraDataSet& Other, int32 StartIdx = 0, int32 NumInstances = INDEX_NONE)const;
 
-	void CheckForNaNs()const
-	{
-		if (CurrData().CheckForNaNs())
-		{
-			Dump(true);
-			ensureAlwaysMsgf(false, TEXT("NiagaraDataSet contains NaNs!"));
-		}
+	void CopyFromGPUReadback(float* GPUReadBackFloat, int* GPUReadBackInt, int32 StartIdx = 0, int32 NumInstances = INDEX_NONE);
 
-		if (PrevData().CheckForNaNs())
-		{
-			Dump(false);
-			ensureAlwaysMsgf(false, TEXT("NiagaraDataSet contains NaNs!"));
-		}
-	}
+	void CheckForNaNs()const;
 
-	// Data set index buffer management
-	// these buffers hold the number of instances for the buffers; the first five uint32s are the DrawIndirect parameters for rendering of the main particle data set
-	//
-	FRWBuffer &GetCurDataSetIndices()		{	return GetDataSetIndices(CurrBuffer);			}
-	FRWBuffer &GetPrevDataSetIndices()		{	return GetDataSetIndices(GetPrevBufferIdx());	}
+	void Dump(int32 StartIndex, int32 NumInstances, const FString& Label)const;
 
-	bool HasDatasetIndices(bool bCur = true) const
-	{
-		uint32 BufIdx = bCur == true ? CurrBuffer : GetPrevBufferIdx();
-		CheckCorrectThread();
-		return DataSetIndices[BufIdx].Buffer != nullptr;
-	}
+	FORCEINLINE bool IsCurrentDataValid()const { return CurrentData != nullptr; }
+	FORCEINLINE FNiagaraDataBuffer* GetCurrentData()const {	return CurrentData; }
+	FORCEINLINE FNiagaraDataBuffer* GetDestinationData()const { return DestinationData; }
 
-	const FRWBuffer &GetCurDataSetIndices() const	
+	FORCEINLINE FNiagaraDataBuffer& GetCurrentDataChecked()const
 	{
-		CheckCorrectThread();
-		return DataSetIndices[CurrBuffer];
+		check(CurrentData);
+		return *CurrentData;
 	}
 
-	const FRWBuffer &GetPrevDataSetIndices() const
+	FORCEINLINE FNiagaraDataBuffer& GetDestinationDataChecked()const
 	{
-		CheckCorrectThread();
-		return DataSetIndices[GetPrevBufferIdx()];
+		check(DestinationData);
+		return *DestinationData;
 	}
 
-	void SetupCurDatasetIndices()
-	{
-		if (DataSetIndices[CurrBuffer].Buffer != nullptr)
-		{
-			DataSetIndices[CurrBuffer].Release();
-		}
-		// Use BUF_KeepCPUAccessible here since some platforms will lock it for readonly (depending on the implementation of RHIEnqueueStagedRead) after GPU simulation.
-		DataSetIndices[CurrBuffer].Initialize(sizeof(int32), 64 /*Context->NumDataSets*/, EPixelFormat::PF_R32_UINT, BUF_DrawIndirect | BUF_Static | BUF_KeepCPUAccessible);	// always allocate for up to 64 data sets
-	}
+private:
 
-	FORCEINLINE uint32 GetNumFloatComponents()const { return TotalFloatComponents; }
-	FORCEINLINE uint32 GetNumInt32Components()const { return TotalInt32Components; }
+	void Reset();
 
-private:
-	FRWBuffer &GetDataSetIndices(uint32 BufIdx)
-	{
-		CheckCorrectThread();
-		return DataSetIndices[BufIdx];
-	}
+	void BuildLayout();
 
-	FORCEINLINE void SwapBuffers()
-	{
-		CheckCorrectThread();
-		if (SimTarget == ENiagaraSimTarget::CPUSim)
-		{
-			MaxBufferIdx = 2;
-			CurrBuffer = CurrBuffer < 2 ? CurrBuffer+1 : 0;
-		}
-		else
-		{
-			MaxBufferIdx = 1;
-			CurrBuffer = CurrBuffer == 0 ? 1 : 0;
-		}
-	}
+	void ResetBuffersInternal();
+	void ReleaseBuffers();
 
 	FORCEINLINE void CheckCorrectThread()const
 	{
@@ -656,31 +302,9 @@ private:
 #endif
 	}
 
-	void BuildLayout()
-	{
-		VariableLayouts.Empty();
-		TotalFloatComponents = 0;
-		TotalInt32Components = 0;
-
-		VariableLayouts.Reserve(Variables.Num());
-		for (FNiagaraVariable& Var : Variables)
-		{
-			FNiagaraVariableLayoutInfo& VarInfo = VariableLayouts[VariableLayouts.AddDefaulted()];
-			FNiagaraTypeLayoutInfo::GenerateLayoutInfo(VarInfo.LayoutInfo, Var.GetType().GetScriptStruct());
-			VarInfo.FloatComponentStart = TotalFloatComponents;
-			VarInfo.Int32ComponentStart = TotalInt32Components;
-			TotalFloatComponents += VarInfo.GetNumFloatComponents();
-			TotalInt32Components += VarInfo.GetNumInt32Components();
-		}
-
-		Data[0].Init(this);
-		Data[1].Init(this);
-		Data[2].Init(this);
-	}
-		
 	/** Unique ID for this data set. Used to allow referencing from other emitters and Systems. */
 	FNiagaraDataSetID ID;
-	
+
 	//////////////////////////////////////////////////////////////////////////
 	//TODO: All this layout is known per emitter / system so doesn't need to be generated and stored for every dataset!
 	/** Variables in the data set. */
@@ -691,19 +315,12 @@ private:
 	uint32 TotalFloatComponents;
 	uint32 TotalInt32Components;
 	//////////////////////////////////////////////////////////////////////////
-	
-	/** Index of current state data. */
-	uint32 CurrBuffer;
-	uint32 MaxBufferIdx;
 
 	ENiagaraSimTarget SimTarget;
 
 	/** Once finalized, the data layout etc is built and no more variables can be added. */
 	uint32 bFinalized : 1;
 	uint32 bNeedsPersistentIDs : 1;
-	
-	/** Table of IDs to real buffer indices. Multi buffered so we can access previous frame data. */
-	TArray<int32> IDToIndexTable[3];
 
 	/** Table of free IDs available to allocate next tick. */
 	TArray<int32> FreeIDsTable;
@@ -717,8 +334,22 @@ private:
 	/** Tag to use when new IDs are acquired. Should be unique per tick. */
 	int32 IDAcquireTag;
 
-	FNiagaraDataBuffer Data[3];
-	FRWBuffer DataSetIndices[3]; 
+	/** Buffer containing the current simulation state. */
+	FNiagaraDataBuffer* CurrentData;
+
+	/** Buffer we're currently simulating into. Only valid while we're simulating i.e between PrepareForSimulate and EndSimulate calls.*/
+	FNiagaraDataBuffer* DestinationData;
+
+	/**
+	Actual data storage. These are passed to and read directly by the RT.
+	This is effectively a pool of buffers for this simulation.
+	Typically this should only be two or three entries and we search for a free buffer to write into on BeginSimulate();
+	We keep track of the Current and Previous buffers which move with each simulate.
+	Additional buffers may be in here if they are currently being used by the render thread.
+	*/
+	TArray<FNiagaraDataBuffer*> Data;
+
+	FString DebugName;
 };
 
 /**
@@ -728,13 +359,11 @@ struct FNiagaraDataSetAccessorBase
 {
 	FNiagaraDataSetAccessorBase()
 		: DataSet(nullptr)
-		, DataBuffer(nullptr)
 		, VarLayout(nullptr)
 	{}
 
-	FNiagaraDataSetAccessorBase(FNiagaraDataSet* InDataSet, FNiagaraVariable InVar, bool bCurrBuffer = true)
+	FNiagaraDataSetAccessorBase(FNiagaraDataSet* InDataSet, FNiagaraVariable InVar)
 		: DataSet(InDataSet)
-		, DataBuffer(bCurrBuffer ? &DataSet->CurrData() : &DataSet->PrevData())
 	{
 		VarLayout = DataSet->GetVariableLayout(InVar);
 	}
@@ -745,16 +374,10 @@ struct FNiagaraDataSetAccessorBase
 		VarLayout = DataSet->GetVariableLayout(InVar);
 	}
 
-	void InitForAccess(bool bCurrBuffer = true)
-	{
-		DataBuffer = bCurrBuffer ? &DataSet->CurrData() : &DataSet->PrevData();
-	}
-
-	FORCEINLINE bool IsValid()const { return VarLayout != nullptr && DataBuffer && DataBuffer->GetNumInstances() > 0; }
+	FORCEINLINE bool IsValid()const { return DataSet && VarLayout != nullptr; }
 protected:
 
 	FNiagaraDataSet* DataSet;
-	FNiagaraDataBuffer* DataBuffer;
 	const FNiagaraVariableLayoutInfo* VarLayout;
 };
 
@@ -762,8 +385,8 @@ template<typename T>
 struct FNiagaraDataSetAccessor : public FNiagaraDataSetAccessorBase
 {
 	FNiagaraDataSetAccessor<T>() {}
-	FNiagaraDataSetAccessor(FNiagaraDataSet& InDataSet, FNiagaraVariable InVar, bool bCurrBuffer=true)
-		: FNiagaraDataSetAccessorBase(&InDataSet, InVar, bCurrBuffer)
+	FNiagaraDataSetAccessor(FNiagaraDataSet& InDataSet, FNiagaraVariable InVar)
+		: FNiagaraDataSetAccessorBase(&InDataSet, InVar)
 	{
 		check(sizeof(T) == InVar.GetType().GetSize());
 		checkf(false, TEXT("You must provide a fast runtime specialization for this type."));// Allow this slow generic version?
@@ -774,16 +397,6 @@ struct FNiagaraDataSetAccessor : public FNiagaraDataSetAccessorBase
 		return Get(Index);
 	}
 
-	FORCEINLINE T GetSafe(int32 Index, T Default)const
-	{
-		if (IsValid())
-		{
-			return Get(Index);
-		}
-
-		return Default;
-	}
-
 	FORCEINLINE T Get(int32 Index)const
 	{
 		T Ret;
@@ -793,8 +406,12 @@ struct FNiagaraDataSetAccessor : public FNiagaraDataSetAccessorBase
 
 	FORCEINLINE void Get(int32 Index, T& OutValue)const
 	{
+		checkSlow(DataSet);
 		uint8* ValuePtr = (uint8*)&OutValue;
 
+		FNiagaraDataBuffer* DataBuffer = DataSet->GetCurrentData();
+		checkSlow(DataBuffer);
+
 		for (uint32 CompIdx = 0; CompIdx < VarLayout->GetNumFloatComponents(); ++CompIdx)
 		{
 			uint32 CompBufferOffset = VarLayout->FloatComponentStart + CompIdx;
@@ -805,19 +422,21 @@ struct FNiagaraDataSetAccessor : public FNiagaraDataSetAccessorBase
 
 		for (uint32 CompIdx = 0; CompIdx < VarLayout->GetNumInt32Components(); ++CompIdx)
 		{
-			uint32 CompBufferOffset = VarLayout->FloatComponentStart + CompIdx;
+			uint32 CompBufferOffset = VarLayout->Int32ComponentStart + CompIdx;
 			int32* Src = DataBuffer->GetInstancePtrInt32(CompBufferOffset, Index);
-			int32* Dst = (int32*)(ValuePtr + VarLayout->LayoutInfo.FloatComponentByteOffsets[CompIdx]);
+			int32* Dst = (int32*)(ValuePtr + VarLayout->LayoutInfo.Int32ComponentByteOffsets[CompIdx]);
 			*Dst = *Src;
 		}
 	}
 
 	FORCEINLINE void Set(int32 Index, const T& InValue)
 	{
-		check((uint32)Index < DataBuffer->GetNumInstances());
-
+		checkSlow(DataSet);
 		uint8* ValuePtr = (uint8*)&InValue;
 
+		FNiagaraDataBuffer* DataBuffer = DataSet->GetDestinationData();
+		checkSlow(DataBuffer);
+
 		for (uint32 CompIdx = 0; CompIdx < VarLayout->GetNumFloatComponents(); ++CompIdx)
 		{
 			uint32 CompBufferOffset = VarLayout->FloatComponentStart + CompIdx;
@@ -828,9 +447,9 @@ struct FNiagaraDataSetAccessor : public FNiagaraDataSetAccessorBase
 
 		for (uint32 CompIdx = 0; CompIdx < VarLayout->GetNumInt32Components(); ++CompIdx)
 		{
-			uint32 CompBufferOffset = VarLayout->FloatComponentStart + CompIdx;
+			uint32 CompBufferOffset = VarLayout->Int32ComponentStart + CompIdx;
 			int32* Dst = DataBuffer->GetInstancePtrInt32(CompBufferOffset, Index);
-			int32* Src = (int32*)(ValuePtr + VarLayout->LayoutInfo.FloatComponentByteOffsets[CompIdx]);
+			int32* Src = (int32*)(ValuePtr + VarLayout->LayoutInfo.Int32ComponentByteOffsets[CompIdx]);
 			*Dst = *Src;
 		}
 	}
@@ -840,26 +459,34 @@ template<>
 struct FNiagaraDataSetAccessor<FNiagaraBool> : public FNiagaraDataSetAccessorBase
 {
 	FNiagaraDataSetAccessor<FNiagaraBool>() {}
-	FNiagaraDataSetAccessor<FNiagaraBool>(FNiagaraDataSet& InDataSet, FNiagaraVariable InVar, bool bCurrBuffer = true)
-		: FNiagaraDataSetAccessorBase(&InDataSet, InVar, bCurrBuffer)
+	FNiagaraDataSetAccessor<FNiagaraBool>(FNiagaraDataSet& InDataSet, FNiagaraVariable InVar)
+		: FNiagaraDataSetAccessorBase(&InDataSet, InVar)
 	{
 		check(sizeof(FNiagaraBool) == InVar.GetType().GetSize());
-		InitForAccess(bCurrBuffer);
+		InitForAccess();
 	}
 
-	void InitForAccess(bool bCurrBuffer = true)
+	void InitForAccess()
 	{
-		DataBuffer = bCurrBuffer ? &DataSet->CurrData() : &DataSet->PrevData();
-		if (VarLayout != nullptr)
-		{
-			Base = (int32*)DataBuffer->GetComponentPtrInt32(VarLayout->Int32ComponentStart);
-		}
-		else
+		SrcBase = nullptr;
+		DestBase = nullptr;
+		FNiagaraDataBuffer* SrcBuffer = DataSet->GetCurrentData();
+		if (IsValid() && SrcBuffer)
 		{
-			Base = nullptr;
+			SrcBase = (int32*)SrcBuffer->GetComponentPtrInt32(VarLayout->Int32ComponentStart);
+
+			//Writes are only valid if we're during a simulation pass.
+			FNiagaraDataBuffer* DestBuffer = DataSet->GetDestinationData();
+			if (DestBuffer)
+			{
+				DestBase = (int32*)DestBuffer->GetComponentPtrInt32(VarLayout->Int32ComponentStart);
+			}
 		}
 	}
 
+	FORCEINLINE bool IsValidForRead()const { return SrcBase != nullptr; }
+	FORCEINLINE bool IsValidForWrite()const { return DestBase != nullptr; }
+
 	FORCEINLINE FNiagaraBool operator[](int32 Index)const
 	{
 		return Get(Index);
@@ -874,7 +501,7 @@ struct FNiagaraDataSetAccessor<FNiagaraBool> : public FNiagaraDataSetAccessorBas
 
 	FORCEINLINE FNiagaraBool GetSafe(int32 Index, bool Default = true)const
 	{
-		if (IsValid())
+		if (IsValidForRead())
 		{
 			return Get(Index);
 		}
@@ -884,51 +511,54 @@ struct FNiagaraDataSetAccessor<FNiagaraBool> : public FNiagaraDataSetAccessorBas
 
 	FORCEINLINE void Get(int32 Index, FNiagaraBool& OutValue)const
 	{
-		OutValue.SetRawValue(Base[Index]);
+		checkSlow(IsValidForRead());
+		OutValue.SetRawValue(SrcBase[Index]);
 	}
 
 	FORCEINLINE void Set(int32 Index, const FNiagaraBool& InValue)
 	{
-		check(Base != nullptr);
-		check((uint32)Index < DataBuffer->GetNumInstances());
-		Base[Index] = InValue.GetRawValue();
+		checkSlow(IsValidForWrite());
+		DestBase[Index] = InValue.GetRawValue();
 	}
 
-	FORCEINLINE bool BaseIsValid() const { return Base != nullptr; }
-
 private:
 
-	int32* Base;
+	int32* SrcBase;
+	int32* DestBase;
 };
 
 template<>
 struct FNiagaraDataSetAccessor<int32> : public FNiagaraDataSetAccessorBase
 {
 	FNiagaraDataSetAccessor<int32>() {}
-	FNiagaraDataSetAccessor<int32>(FNiagaraDataSet& InDataSet, FNiagaraVariable InVar, bool bCurrBuffer = true)
-		: FNiagaraDataSetAccessorBase(&InDataSet, InVar, bCurrBuffer)
+	FNiagaraDataSetAccessor<int32>(FNiagaraDataSet& InDataSet, FNiagaraVariable InVar)
+		: FNiagaraDataSetAccessorBase(&InDataSet, InVar)
 	{
 		check(sizeof(int32) == InVar.GetType().GetSize());
-		InitForAccess(bCurrBuffer);
+		InitForAccess();
 	}
-	
-	void InitForAccess(bool bCurrBuffer = true)
+
+	void InitForAccess()
 	{
-		DataBuffer = bCurrBuffer ? &DataSet->CurrData() : &DataSet->PrevData();
-		if (VarLayout != nullptr)
+		SrcBase = nullptr;
+		DestBase = nullptr;
+		FNiagaraDataBuffer* SrcBuffer = DataSet->GetCurrentData();
+		if (IsValid() && SrcBuffer)
 		{
-			Base = (int32*)DataBuffer->GetComponentPtrInt32(VarLayout->Int32ComponentStart);
-			if (DataBuffer->GetNumInstances() != 0)
+			SrcBase = (int32*)SrcBuffer->GetComponentPtrInt32(VarLayout->Int32ComponentStart);
+
+			//Writes are only valid if we're during a simulation pass.
+			FNiagaraDataBuffer* DestBuffer = DataSet->GetDestinationData();
+			if (DestBuffer)
 			{
-				check(Base != nullptr);
+				DestBase = (int32*)DestBuffer->GetComponentPtrInt32(VarLayout->Int32ComponentStart);
 			}
 		}
-		else
-		{
-			Base = nullptr;
-		}
 	}
 
+	FORCEINLINE bool IsValidForRead()const { return SrcBase != nullptr; }
+	FORCEINLINE bool IsValidForWrite()const { return DestBase != nullptr; }
+
 	FORCEINLINE int32 operator[](int32 Index)const
 	{
 		return Get(Index);
@@ -943,9 +573,17 @@ struct FNiagaraDataSetAccessor<int32> : public FNiagaraDataSetAccessorBase
 
 	FORCEINLINE int32 GetSafe(int32 Index, int32 Default = 0)const
 	{
-		if (IsValid())
+		if (IsValidForRead())
 		{
-			return Get(Index);
+			FNiagaraDataBuffer* DataBuffer = DataSet->GetCurrentData();
+			if (DataBuffer && Index >= 0 && (uint32)Index < DataBuffer->GetNumInstances())
+			{
+				return Get(Index);
+			}
+			else
+			{
+				ensure(DataBuffer && Index >= 0 && (uint32)Index < DataBuffer->GetNumInstances()); // just to capture the badness in the logs or debugger if attached.
+			}
 		}
 
 		return Default;
@@ -953,48 +591,54 @@ struct FNiagaraDataSetAccessor<int32> : public FNiagaraDataSetAccessorBase
 
 	FORCEINLINE void Get(int32 Index, int32& OutValue)const
 	{
-		check(Base != nullptr);
-		OutValue = Base[Index];
+		checkSlow(IsValidForRead());
+		OutValue = SrcBase[Index];
 	}
 
 	FORCEINLINE void Set(int32 Index, const int32& InValue)
 	{
-		check(Base != nullptr);
-		check((uint32)Index < DataBuffer->GetNumInstances());
-		Base[Index] = InValue;
+		checkSlow(IsValidForWrite());
+		DestBase[Index] = InValue;
 	}
 
-	FORCEINLINE bool BaseIsValid() const { return Base != nullptr; }
-
 private:
 
-	int32* Base;
+	int32* SrcBase;
+	int32* DestBase;
 };
 
 template<>
 struct FNiagaraDataSetAccessor<float> : public FNiagaraDataSetAccessorBase
 {
 	FNiagaraDataSetAccessor<float>() {}
-	FNiagaraDataSetAccessor<float>(FNiagaraDataSet& InDataSet, FNiagaraVariable InVar, bool bCurrBuffer=true)
-		: FNiagaraDataSetAccessorBase(&InDataSet, InVar, bCurrBuffer)
+	FNiagaraDataSetAccessor<float>(FNiagaraDataSet& InDataSet, FNiagaraVariable InVar)
+		: FNiagaraDataSetAccessorBase(&InDataSet, InVar)
 	{
 		check(sizeof(float) == InVar.GetType().GetSize());
-		InitForAccess(bCurrBuffer);
+		InitForAccess();
 	}
 
-	void InitForAccess(bool bCurrBuffer = true)
+	void InitForAccess()
 	{
-		DataBuffer = bCurrBuffer ? &DataSet->CurrData() : &DataSet->PrevData();
-		if (VarLayout != nullptr)
-		{
-			Base = (float*)DataBuffer->GetComponentPtrFloat(VarLayout->FloatComponentStart);
-		}
-		else
+		SrcBase = nullptr;
+		DestBase = nullptr;
+		FNiagaraDataBuffer* SrcBuffer = DataSet->GetCurrentData();
+		if (IsValid() && SrcBuffer)
 		{
-			Base = nullptr;
+			SrcBase = (float*)SrcBuffer->GetComponentPtrFloat(VarLayout->FloatComponentStart);
+
+			//Writes are only valid if we're during a simulation pass.
+			FNiagaraDataBuffer* DestBuffer = DataSet->GetDestinationData();
+			if (DestBuffer)
+			{
+				DestBase = (float*)DestBuffer->GetComponentPtrFloat(VarLayout->FloatComponentStart);
+			}
 		}
 	}
 
+	FORCEINLINE bool IsValidForRead()const { return SrcBase != nullptr; }
+	FORCEINLINE bool IsValidForWrite()const { return DestBase != nullptr; }
+
 	FORCEINLINE float operator[](int32 Index)const
 	{
 		return Get(Index);
@@ -1002,9 +646,17 @@ struct FNiagaraDataSetAccessor<float> : public FNiagaraDataSetAccessorBase
 
 	FORCEINLINE float GetSafe(int32 Index, float Default = 0.0f)const
 	{
-		if (IsValid())
+		if (IsValidForRead())
 		{
-			return Get(Index);
+			FNiagaraDataBuffer* DataBuffer = DataSet->GetCurrentData();
+			if (DataBuffer && Index >= 0 && (uint32)Index < DataBuffer->GetNumInstances())
+			{
+				return Get(Index);
+			}
+			else
+			{
+				ensure(DataBuffer && Index >= 0 && (uint32)Index < DataBuffer->GetNumInstances()); // just to capture the badness in the logs or debugger if attached.
+			}
 		}
 
 		return Default;
@@ -1019,47 +671,57 @@ struct FNiagaraDataSetAccessor<float> : public FNiagaraDataSetAccessorBase
 
 	FORCEINLINE void Get(int32 Index, float& OutValue)const
 	{
-		OutValue = Base[Index];
+		checkSlow(IsValidForRead());
+		OutValue = SrcBase[Index];
 	}
 
 	FORCEINLINE void Set(int32 Index, const float& InValue)
 	{
-		check(Base != nullptr);
-		check((uint32)Index < DataBuffer->GetNumInstances());
-		Base[Index] = InValue;
+		checkSlow(IsValidForWrite());
+		DestBase[Index] = InValue;
 	}
-	FORCEINLINE bool BaseIsValid() const { return Base != nullptr; }
 
 private:
-	float* Base;
+	float* SrcBase;
+	float* DestBase;
 };
 
 template<>
 struct FNiagaraDataSetAccessor<FVector2D> : public FNiagaraDataSetAccessorBase
 {
 	FNiagaraDataSetAccessor<FVector2D>() {}
-	FNiagaraDataSetAccessor<FVector2D>(FNiagaraDataSet& InDataSet, FNiagaraVariable InVar, bool bCurrBuffer=true)
-		: FNiagaraDataSetAccessorBase(&InDataSet, InVar, bCurrBuffer)
+	FNiagaraDataSetAccessor<FVector2D>(FNiagaraDataSet& InDataSet, FNiagaraVariable InVar)
+		: FNiagaraDataSetAccessorBase(&InDataSet, InVar)
 	{
 		check(sizeof(FVector2D) == InVar.GetType().GetSize());
-		InitForAccess(bCurrBuffer);
+		InitForAccess();
 	}
 
-	void InitForAccess(bool bCurrBuffer = true)
+	void InitForAccess()
 	{
-		DataBuffer = bCurrBuffer ? &DataSet->CurrData() : &DataSet->PrevData();
-		if (VarLayout != nullptr)
-		{
-			XBase = (float*)DataBuffer->GetComponentPtrFloat(VarLayout->FloatComponentStart);
-			YBase = (float*)DataBuffer->GetComponentPtrFloat(VarLayout->FloatComponentStart + 1);
-		}
-		else
+		SrcXBase = nullptr;
+		SrcYBase = nullptr;
+		DestXBase = nullptr;
+		DestYBase = nullptr;
+		FNiagaraDataBuffer* SrcBuffer = DataSet->GetCurrentData();
+		if (IsValid() && SrcBuffer)
 		{
-			XBase = nullptr;
-			YBase = nullptr;
+			SrcXBase = (float*)SrcBuffer->GetComponentPtrFloat(VarLayout->FloatComponentStart);
+			SrcYBase = (float*)SrcBuffer->GetComponentPtrFloat(VarLayout->FloatComponentStart + 1);
+
+			//Writes are only valid if we're during a simulation pass.
+			FNiagaraDataBuffer* DestBuffer = DataSet->GetDestinationData();
+			if (DestBuffer)
+			{
+				DestXBase = (float*)DestBuffer->GetComponentPtrFloat(VarLayout->FloatComponentStart);
+				DestYBase = (float*)DestBuffer->GetComponentPtrFloat(VarLayout->FloatComponentStart + 1);
+			}
 		}
 	}
 
+	FORCEINLINE bool IsValidForRead()const { return SrcXBase != nullptr && SrcYBase != nullptr; }
+	FORCEINLINE bool IsValidForWrite()const { return DestXBase != nullptr && DestYBase != nullptr; }
+
 	FORCEINLINE FVector2D operator[](int32 Index)const
 	{
 		return Get(Index);
@@ -1067,9 +729,17 @@ struct FNiagaraDataSetAccessor<FVector2D> : public FNiagaraDataSetAccessorBase
 
 	FORCEINLINE FVector2D GetSafe(int32 Index, FVector2D Default = FVector2D::ZeroVector)const
 	{
-		if (IsValid())
+		if (IsValidForRead())
 		{
-			return Get(Index);
+			FNiagaraDataBuffer* DataBuffer = DataSet->GetCurrentData();
+			if (DataBuffer && Index >= 0 && (uint32)Index < DataBuffer->GetNumInstances())
+			{
+				return Get(Index);
+			}
+			else
+			{
+				ensure(DataBuffer && Index >= 0 && (uint32)Index < DataBuffer->GetNumInstances()); // just to capture the badness in the logs or debugger if attached.
+			}
 		}
 
 		return Default;
@@ -1084,54 +754,66 @@ struct FNiagaraDataSetAccessor<FVector2D> : public FNiagaraDataSetAccessorBase
 
 	FORCEINLINE void Get(int32 Index, FVector2D& OutValue)const
 	{
-		OutValue.X = XBase[Index];
-		OutValue.Y = YBase[Index];
+		checkSlow(IsValidForRead());
+		OutValue.X = SrcXBase[Index];
+		OutValue.Y = SrcYBase[Index];
 	}
 
 	FORCEINLINE void Set(int32 Index, const FVector2D& InValue)
 	{
-		check(XBase != nullptr);
-		check((uint32)Index < DataBuffer->GetNumInstances());
-		XBase[Index] = InValue.X;
-		YBase[Index] = InValue.Y;
-	}	
-	
-	FORCEINLINE bool BaseIsValid() const { return XBase != nullptr && YBase != nullptr; }
+		checkSlow(IsValidForWrite());
+		DestXBase[Index] = InValue.X;
+		DestYBase[Index] = InValue.Y;
+	}
 
 private:
 
-	float* XBase;
-	float* YBase;
+	float* SrcXBase;
+	float* SrcYBase;
+	float* DestXBase;
+	float* DestYBase;
 };
 
 template<>
 struct FNiagaraDataSetAccessor<FVector> : public FNiagaraDataSetAccessorBase
 {
 	FNiagaraDataSetAccessor<FVector>() {}
-	FNiagaraDataSetAccessor<FVector>(FNiagaraDataSet& InDataSet, FNiagaraVariable InVar, bool bCurrBuffer=true)
-		: FNiagaraDataSetAccessorBase(&InDataSet, InVar, bCurrBuffer)
+	FNiagaraDataSetAccessor<FVector>(FNiagaraDataSet& InDataSet, FNiagaraVariable InVar)
+		: FNiagaraDataSetAccessorBase(&InDataSet, InVar)
 	{
 		check(sizeof(FVector) == InVar.GetType().GetSize());
-		InitForAccess(bCurrBuffer);
+		InitForAccess();
 	}
 
-	void InitForAccess(bool bCurrBuffer = true)
+	void InitForAccess()
 	{
-		DataBuffer = bCurrBuffer ? &DataSet->CurrData() : &DataSet->PrevData();
-		if (VarLayout != nullptr)
-		{
-			XBase = (float*)DataBuffer->GetComponentPtrFloat(VarLayout->FloatComponentStart);
-			YBase = (float*)DataBuffer->GetComponentPtrFloat(VarLayout->FloatComponentStart + 1);
-			ZBase = (float*)DataBuffer->GetComponentPtrFloat(VarLayout->FloatComponentStart + 2);
-		}
-		else
+		SrcXBase = nullptr;
+		SrcYBase = nullptr;
+		SrcZBase = nullptr;
+		DestXBase = nullptr;
+		DestYBase = nullptr;
+		DestZBase = nullptr;
+		FNiagaraDataBuffer* SrcBuffer = DataSet->GetCurrentData();
+		if (IsValid() && SrcBuffer)
 		{
-			XBase = nullptr;
-			YBase = nullptr;
-			ZBase = nullptr;
+			SrcXBase = (float*)SrcBuffer->GetComponentPtrFloat(VarLayout->FloatComponentStart);
+			SrcYBase = (float*)SrcBuffer->GetComponentPtrFloat(VarLayout->FloatComponentStart + 1);
+			SrcZBase = (float*)SrcBuffer->GetComponentPtrFloat(VarLayout->FloatComponentStart + 2);
+
+			//Writes are only valid if we're during a simulation pass.
+			FNiagaraDataBuffer* DestBuffer = DataSet->GetDestinationData();
+			if (DestBuffer)
+			{
+				DestXBase = (float*)DestBuffer->GetComponentPtrFloat(VarLayout->FloatComponentStart);
+				DestYBase = (float*)DestBuffer->GetComponentPtrFloat(VarLayout->FloatComponentStart + 1);
+				DestZBase = (float*)DestBuffer->GetComponentPtrFloat(VarLayout->FloatComponentStart + 2);
+			}
 		}
 	}
 
+	FORCEINLINE bool IsValidForRead()const { return SrcXBase != nullptr && SrcYBase != nullptr && SrcZBase != nullptr; }
+	FORCEINLINE bool IsValidForWrite()const { return DestXBase != nullptr && DestYBase != nullptr && DestZBase != nullptr; }
+
 	FORCEINLINE FVector operator[](int32 Index)const
 	{
 		return Get(Index);
@@ -1139,9 +821,17 @@ struct FNiagaraDataSetAccessor<FVector> : public FNiagaraDataSetAccessorBase
 
 	FORCEINLINE FVector GetSafe(int32 Index, FVector Default = FVector::ZeroVector)const
 	{
-		if (IsValid())
+		if (IsValidForRead())
 		{
-			return Get(Index);
+			FNiagaraDataBuffer* DataBuffer = DataSet->GetCurrentData();
+			if (DataBuffer && Index >= 0 && (uint32)Index < DataBuffer->GetNumInstances())
+			{
+				return Get(Index);
+			}
+			else
+			{
+				ensure(DataBuffer && Index >= 0 && (uint32)Index < DataBuffer->GetNumInstances()); // just to capture the badness in the logs or debugger if attached.
+			}
 		}
 
 		return Default;
@@ -1156,69 +846,92 @@ struct FNiagaraDataSetAccessor<FVector> : public FNiagaraDataSetAccessorBase
 
 	FORCEINLINE void Get(int32 Index, FVector& OutValue)const
 	{
-		OutValue.X = XBase[Index];
-		OutValue.Y = YBase[Index];
-		OutValue.Z = ZBase[Index];
+		checkSlow(IsValidForRead());
+		OutValue.X = SrcXBase[Index];
+		OutValue.Y = SrcYBase[Index];
+		OutValue.Z = SrcZBase[Index];
 	}
 
 	FORCEINLINE void Set(int32 Index, const FVector& InValue)
 	{
-		check(XBase != nullptr);
-		check((uint32)Index < DataBuffer->GetNumInstances());
-
-		XBase[Index] = InValue.X;
-		YBase[Index] = InValue.Y;
-		ZBase[Index] = InValue.Z;
+		checkSlow(IsValidForWrite());
+		DestXBase[Index] = InValue.X;
+		DestYBase[Index] = InValue.Y;
+		DestZBase[Index] = InValue.Z;
 	}
 
-	FORCEINLINE bool BaseIsValid() const { return XBase != nullptr && YBase != nullptr && ZBase != nullptr; }
 private:
 
-	float* XBase;
-	float* YBase;
-	float* ZBase;
+	float* SrcXBase;
+	float* SrcYBase;
+	float* SrcZBase;
+	float* DestXBase;
+	float* DestYBase;
+	float* DestZBase;
 };
 
 template<>
 struct FNiagaraDataSetAccessor<FVector4> : public FNiagaraDataSetAccessorBase
 {
 	FNiagaraDataSetAccessor<FVector4>() {}
-	FNiagaraDataSetAccessor<FVector4>(FNiagaraDataSet& InDataSet, FNiagaraVariable InVar, bool bCurrBuffer=true)
-		: FNiagaraDataSetAccessorBase(&InDataSet, InVar, bCurrBuffer)
+	FNiagaraDataSetAccessor<FVector4>(FNiagaraDataSet& InDataSet, FNiagaraVariable InVar)
+		: FNiagaraDataSetAccessorBase(&InDataSet, InVar)
 	{
 		check(sizeof(FVector4) == InVar.GetType().GetSize());
-		InitForAccess(bCurrBuffer);
+		InitForAccess();
 	}
 
-	void InitForAccess(bool bCurrBuffer = true)
+	void InitForAccess()
 	{
-		DataBuffer = bCurrBuffer ? &DataSet->CurrData() : &DataSet->PrevData();
-		if (VarLayout)
+		SrcXBase = nullptr;
+		SrcYBase = nullptr;
+		SrcZBase = nullptr;
+		SrcWBase = nullptr;
+		DestXBase = nullptr;
+		DestYBase = nullptr;
+		DestZBase = nullptr;
+		DestWBase = nullptr;
+		FNiagaraDataBuffer* SrcBuffer = DataSet->GetCurrentData();
+		if (IsValid() && SrcBuffer)
 		{
-			XBase = (float*)DataBuffer->GetComponentPtrFloat(VarLayout->FloatComponentStart);
-			YBase = (float*)DataBuffer->GetComponentPtrFloat(VarLayout->FloatComponentStart + 1);
-			ZBase = (float*)DataBuffer->GetComponentPtrFloat(VarLayout->FloatComponentStart + 2);
-			WBase = (float*)DataBuffer->GetComponentPtrFloat(VarLayout->FloatComponentStart + 3);
-		}
-		else
-		{
-			XBase = nullptr;
-			YBase = nullptr;
-			ZBase = nullptr;
-			WBase = nullptr;
+			SrcXBase = (float*)SrcBuffer->GetComponentPtrFloat(VarLayout->FloatComponentStart);
+			SrcYBase = (float*)SrcBuffer->GetComponentPtrFloat(VarLayout->FloatComponentStart + 1);
+			SrcZBase = (float*)SrcBuffer->GetComponentPtrFloat(VarLayout->FloatComponentStart + 2);
+			SrcWBase = (float*)SrcBuffer->GetComponentPtrFloat(VarLayout->FloatComponentStart + 3);
+
+			//Writes are only valid if we're during a simulation pass.
+			FNiagaraDataBuffer* DestBuffer = DataSet->GetDestinationData();
+			if (DestBuffer)
+			{
+				DestXBase = (float*)DestBuffer->GetComponentPtrFloat(VarLayout->FloatComponentStart);
+				DestYBase = (float*)DestBuffer->GetComponentPtrFloat(VarLayout->FloatComponentStart + 1);
+				DestZBase = (float*)DestBuffer->GetComponentPtrFloat(VarLayout->FloatComponentStart + 2);
+				DestWBase = (float*)DestBuffer->GetComponentPtrFloat(VarLayout->FloatComponentStart + 3);
+			}
 		}
 	}
 
+	FORCEINLINE bool IsValidForRead()const { return SrcXBase != nullptr && SrcYBase != nullptr && SrcZBase != nullptr && SrcWBase != nullptr; }
+	FORCEINLINE bool IsValidForWrite()const { return DestXBase != nullptr && DestYBase != nullptr && DestZBase != nullptr && DestWBase != nullptr; }
+
 	FORCEINLINE FVector4 operator[](int32 Index)const
 	{
 		return Get(Index);
 	}
 
-	FORCEINLINE FVector4 GetSafe(int32 Index, const FVector4& Default = FVector4(0.0f,0.0f,0.0f,0.0f))const
+	FORCEINLINE FVector4 GetSafe(int32 Index, const FVector4& Default = FVector4(0.0f, 0.0f, 0.0f, 0.0f))const
 	{
-		if (IsValid())
+		if (IsValidForRead())
 		{
-			return Get(Index);
+			FNiagaraDataBuffer* DataBuffer = DataSet->GetCurrentData();
+			if (DataBuffer && Index >= 0 && (uint32)Index < DataBuffer->GetNumInstances())
+			{
+				return Get(Index);
+			}
+			else
+			{
+				ensure(DataBuffer && Index >= 0 && (uint32)Index < DataBuffer->GetNumInstances()); // just to capture the badness in the logs or debugger if attached.
+			}
 		}
 
 		return Default;
@@ -1233,30 +946,32 @@ struct FNiagaraDataSetAccessor<FVector4> : public FNiagaraDataSetAccessorBase
 
 	FORCEINLINE void Get(int32 Index, FVector4& OutValue)const
 	{
-		OutValue.X = XBase[Index];
-		OutValue.Y = YBase[Index];
-		OutValue.Z = ZBase[Index];
-		OutValue.W = WBase[Index];
+		checkSlow(IsValidForRead());
+		OutValue.X = SrcXBase[Index];
+		OutValue.Y = SrcYBase[Index];
+		OutValue.Z = SrcZBase[Index];
+		OutValue.W = SrcWBase[Index];
 	}
 
 	FORCEINLINE void Set(int32 Index, const FVector4& InValue)
 	{
-		check(XBase != nullptr);
-		check((uint32)Index < DataBuffer->GetNumInstances());
-
-		XBase[Index] = InValue.X;
-		YBase[Index] = InValue.Y;
-		ZBase[Index] = InValue.Z;
-		WBase[Index] = InValue.W;
+		checkSlow(IsValidForWrite());
+		DestXBase[Index] = InValue.X;
+		DestYBase[Index] = InValue.Y;
+		DestZBase[Index] = InValue.Z;
+		DestWBase[Index] = InValue.W;
 	}
 
-	FORCEINLINE bool BaseIsValid() const { return XBase != nullptr && YBase != nullptr && ZBase != nullptr && WBase != nullptr; }
 private:
 
-	float* XBase;
-	float* YBase;
-	float* ZBase;
-	float* WBase;
+	float* SrcXBase;
+	float* SrcYBase;
+	float* SrcZBase;
+	float* SrcWBase;
+	float* DestXBase;
+	float* DestYBase;
+	float* DestZBase;
+	float* DestWBase;
 };
 
 
@@ -1264,32 +979,46 @@ template<>
 struct FNiagaraDataSetAccessor<FQuat> : public FNiagaraDataSetAccessorBase
 {
 	FNiagaraDataSetAccessor<FQuat>() {}
-	FNiagaraDataSetAccessor<FQuat>(FNiagaraDataSet& InDataSet, FNiagaraVariable InVar, bool bCurrBuffer = true)
-		: FNiagaraDataSetAccessorBase(&InDataSet, InVar, bCurrBuffer)
+	FNiagaraDataSetAccessor<FQuat>(FNiagaraDataSet& InDataSet, FNiagaraVariable InVar)
+		: FNiagaraDataSetAccessorBase(&InDataSet, InVar)
 	{
 		check(sizeof(FQuat) == InVar.GetType().GetSize());
-		InitForAccess(bCurrBuffer);
+		InitForAccess();
 	}
 
-	void InitForAccess(bool bCurrBuffer = true)
+	void InitForAccess()
 	{
-		DataBuffer = bCurrBuffer ? &DataSet->CurrData() : &DataSet->PrevData();
-		if (VarLayout)
+		SrcXBase = nullptr;
+		SrcYBase = nullptr;
+		SrcZBase = nullptr;
+		SrcWBase = nullptr;
+		DestXBase = nullptr;
+		DestYBase = nullptr;
+		DestZBase = nullptr;
+		DestWBase = nullptr;
+		FNiagaraDataBuffer* SrcBuffer = DataSet->GetCurrentData();
+		if (IsValid() && SrcBuffer)
 		{
-			XBase = (float*)DataBuffer->GetComponentPtrFloat(VarLayout->FloatComponentStart);
-			YBase = (float*)DataBuffer->GetComponentPtrFloat(VarLayout->FloatComponentStart + 1);
-			ZBase = (float*)DataBuffer->GetComponentPtrFloat(VarLayout->FloatComponentStart + 2);
-			WBase = (float*)DataBuffer->GetComponentPtrFloat(VarLayout->FloatComponentStart + 3);
-		}
-		else
-		{
-			XBase = nullptr;
-			YBase = nullptr;
-			ZBase = nullptr;
-			WBase = nullptr;
+			SrcXBase = (float*)SrcBuffer->GetComponentPtrFloat(VarLayout->FloatComponentStart);
+			SrcYBase = (float*)SrcBuffer->GetComponentPtrFloat(VarLayout->FloatComponentStart + 1);
+			SrcZBase = (float*)SrcBuffer->GetComponentPtrFloat(VarLayout->FloatComponentStart + 2);
+			SrcWBase = (float*)SrcBuffer->GetComponentPtrFloat(VarLayout->FloatComponentStart + 3);
+
+			FNiagaraDataBuffer* DestBuffer = DataSet->GetDestinationData();
+			if (DestBuffer)
+			{
+				//Writes are only valid if we're during a simulation pass.
+				DestXBase = (float*)DestBuffer->GetComponentPtrFloat(VarLayout->FloatComponentStart);
+				DestYBase = (float*)DestBuffer->GetComponentPtrFloat(VarLayout->FloatComponentStart + 1);
+				DestZBase = (float*)DestBuffer->GetComponentPtrFloat(VarLayout->FloatComponentStart + 2);
+				DestWBase = (float*)DestBuffer->GetComponentPtrFloat(VarLayout->FloatComponentStart + 3);
+			}
 		}
 	}
 
+	FORCEINLINE bool IsValidForRead()const { return SrcXBase != nullptr && SrcYBase != nullptr && SrcZBase != nullptr && SrcWBase != nullptr; }
+	FORCEINLINE bool IsValidForWrite()const { return DestXBase != nullptr && DestYBase != nullptr && DestZBase != nullptr && DestWBase != nullptr; }
+
 	FORCEINLINE FQuat operator[](int32 Index)const
 	{
 		return Get(Index);
@@ -1297,9 +1026,17 @@ struct FNiagaraDataSetAccessor<FQuat> : public FNiagaraDataSetAccessorBase
 
 	FORCEINLINE FQuat GetSafe(int32 Index, const FQuat& Default = FQuat(0.0f, 0.0f, 0.0f, 1.0f))const
 	{
-		if (IsValid())
+		if (IsValidForRead())
 		{
-			return Get(Index);
+			FNiagaraDataBuffer* DataBuffer = DataSet->GetCurrentData();
+			if (DataBuffer && Index >= 0 && (uint32)Index < DataBuffer->GetNumInstances())
+			{
+				return Get(Index);
+			}
+			else
+			{
+				ensure(DataBuffer && Index >= 0 && (uint32)Index < DataBuffer->GetNumInstances()); // just to capture the badness in the logs or debugger if attached.
+			}
 		}
 
 		return Default;
@@ -1314,62 +1051,78 @@ struct FNiagaraDataSetAccessor<FQuat> : public FNiagaraDataSetAccessorBase
 
 	FORCEINLINE void Get(int32 Index, FQuat& OutValue)const
 	{
-		OutValue.X = XBase[Index];
-		OutValue.Y = YBase[Index];
-		OutValue.Z = ZBase[Index];
-		OutValue.W = WBase[Index];
+		checkSlow(IsValidForRead());
+		OutValue.X = SrcXBase[Index];
+		OutValue.Y = SrcYBase[Index];
+		OutValue.Z = SrcZBase[Index];
+		OutValue.W = SrcWBase[Index];
 	}
 
 	FORCEINLINE void Set(int32 Index, const FQuat& InValue)
 	{
-		check(XBase != nullptr);
-		check((uint32)Index < DataBuffer->GetNumInstances());
-
-		XBase[Index] = InValue.X;
-		YBase[Index] = InValue.Y;
-		ZBase[Index] = InValue.Z;
-		WBase[Index] = InValue.W;
+		checkSlow(IsValidForWrite());
+		DestXBase[Index] = InValue.X;
+		DestYBase[Index] = InValue.Y;
+		DestZBase[Index] = InValue.Z;
+		DestWBase[Index] = InValue.W;
 	}
 
-	FORCEINLINE bool BaseIsValid() const { return XBase != nullptr && YBase != nullptr && ZBase != nullptr && WBase != nullptr; }
 private:
 
-	float* XBase;
-	float* YBase;
-	float* ZBase;
-	float* WBase;
+	float* SrcXBase;
+	float* SrcYBase;
+	float* SrcZBase;
+	float* SrcWBase;
+	float* DestXBase;
+	float* DestYBase;
+	float* DestZBase;
+	float* DestWBase;
 };
 
 template<>
 struct FNiagaraDataSetAccessor<FLinearColor> : public FNiagaraDataSetAccessorBase
 {
 	FNiagaraDataSetAccessor<FLinearColor>() {}
-	FNiagaraDataSetAccessor<FLinearColor>(FNiagaraDataSet& InDataSet, FNiagaraVariable InVar, bool bCurrBuffer=true)
-		: FNiagaraDataSetAccessorBase(&InDataSet, InVar, bCurrBuffer)
+	FNiagaraDataSetAccessor<FLinearColor>(FNiagaraDataSet& InDataSet, FNiagaraVariable InVar)
+		: FNiagaraDataSetAccessorBase(&InDataSet, InVar)
 	{
 		check(sizeof(FLinearColor) == InVar.GetType().GetSize());
-		InitForAccess(bCurrBuffer);
+		InitForAccess();
 	}
 
-	void InitForAccess(bool bCurrBuffer = true)
+	void InitForAccess()
 	{
-		DataBuffer = bCurrBuffer ? &DataSet->CurrData() : &DataSet->PrevData();
-		if (VarLayout)
-		{
-			RBase = (float*)DataBuffer->GetComponentPtrFloat(VarLayout->FloatComponentStart);
-			GBase = (float*)DataBuffer->GetComponentPtrFloat(VarLayout->FloatComponentStart + 1);
-			BBase = (float*)DataBuffer->GetComponentPtrFloat(VarLayout->FloatComponentStart + 2);
-			ABase = (float*)DataBuffer->GetComponentPtrFloat(VarLayout->FloatComponentStart + 3);
-		}
-		else
+		SrcRBase = nullptr;
+		SrcGBase = nullptr;
+		SrcBBase = nullptr;
+		SrcABase = nullptr;
+		DestRBase = nullptr;
+		DestGBase = nullptr;
+		DestBBase = nullptr;
+		DestABase = nullptr;
+		FNiagaraDataBuffer* SrcBuffer = DataSet->GetCurrentData();
+		if (IsValid() && SrcBuffer)
 		{
-			RBase = nullptr;
-			GBase = nullptr;
-			BBase = nullptr;
-			ABase = nullptr;
+			SrcRBase = (float*)SrcBuffer->GetComponentPtrFloat(VarLayout->FloatComponentStart);
+			SrcGBase = (float*)SrcBuffer->GetComponentPtrFloat(VarLayout->FloatComponentStart + 1);
+			SrcBBase = (float*)SrcBuffer->GetComponentPtrFloat(VarLayout->FloatComponentStart + 2);
+			SrcABase = (float*)SrcBuffer->GetComponentPtrFloat(VarLayout->FloatComponentStart + 3);
+
+			//Writes are only valid if we're during a simulation pass.
+			FNiagaraDataBuffer* DestBuffer = DataSet->GetDestinationData();
+			if (DestBuffer)
+			{
+				DestRBase = (float*)DestBuffer->GetComponentPtrFloat(VarLayout->FloatComponentStart);
+				DestGBase = (float*)DestBuffer->GetComponentPtrFloat(VarLayout->FloatComponentStart + 1);
+				DestBBase = (float*)DestBuffer->GetComponentPtrFloat(VarLayout->FloatComponentStart + 2);
+				DestABase = (float*)DestBuffer->GetComponentPtrFloat(VarLayout->FloatComponentStart + 3);
+			}
 		}
 	}
 
+	FORCEINLINE bool IsValidForRead()const { return SrcRBase != nullptr && SrcGBase != nullptr && SrcBBase != nullptr && SrcABase != nullptr; }
+	FORCEINLINE bool IsValidForWrite()const { return DestRBase != nullptr && DestGBase != nullptr && DestBBase != nullptr && DestABase != nullptr; }
+
 	FORCEINLINE FLinearColor operator[](int32 Index)const
 	{
 		return Get(Index);
@@ -1377,9 +1130,17 @@ struct FNiagaraDataSetAccessor<FLinearColor> : public FNiagaraDataSetAccessorBas
 
 	FORCEINLINE FLinearColor GetSafe(int32 Index, FLinearColor Default = FLinearColor::White)const
 	{
-		if (IsValid())
+		if (IsValidForRead())
 		{
-			return Get(Index);
+			FNiagaraDataBuffer* DataBuffer = DataSet->GetCurrentData();
+			if (DataBuffer && Index >= 0 && (uint32)Index < DataBuffer->GetNumInstances())
+			{
+				return Get(Index);
+			}
+			else
+			{
+				ensure(DataBuffer && Index >= 0 && (uint32)Index < DataBuffer->GetNumInstances()); // just to capture the badness in the logs or debugger if attached.
+			}
 		}
 
 		return Default;
@@ -1394,62 +1155,78 @@ struct FNiagaraDataSetAccessor<FLinearColor> : public FNiagaraDataSetAccessorBas
 
 	FORCEINLINE void Get(int32 Index, FLinearColor& OutValue)const
 	{
-		OutValue.R = RBase[Index];
-		OutValue.G = GBase[Index];
-		OutValue.B = BBase[Index];
-		OutValue.A = ABase[Index];
+		checkSlow(IsValidForRead());
+		OutValue.R = SrcRBase[Index];
+		OutValue.G = SrcGBase[Index];
+		OutValue.B = SrcBBase[Index];
+		OutValue.A = SrcABase[Index];
 	}
 
 	FORCEINLINE void Set(int32 Index, const FLinearColor& InValue)
 	{
-		check(RBase != nullptr);
-		check((uint32)Index < DataBuffer->GetNumInstances());
-
-		RBase[Index] = InValue.R;
-		GBase[Index] = InValue.G;
-		BBase[Index] = InValue.B;
-		ABase[Index] = InValue.A;
+		checkSlow(IsValidForWrite());
+		DestRBase[Index] = InValue.R;
+		DestGBase[Index] = InValue.G;
+		DestBBase[Index] = InValue.B;
+		DestABase[Index] = InValue.A;
 	}
-	FORCEINLINE bool BaseIsValid() const { return RBase != nullptr && GBase != nullptr && BBase != nullptr && ABase != nullptr; }
 
 private:
 
-	float* RBase;
-	float* GBase;
-	float* BBase;
-	float* ABase;
+	float* SrcRBase;
+	float* SrcGBase;
+	float* SrcBBase;
+	float* SrcABase;
+	float* DestRBase;
+	float* DestGBase;
+	float* DestBBase;
+	float* DestABase;
 };
 
 template<>
 struct FNiagaraDataSetAccessor<FNiagaraSpawnInfo> : public FNiagaraDataSetAccessorBase
 {
 	FNiagaraDataSetAccessor<FNiagaraSpawnInfo>() {}
-	FNiagaraDataSetAccessor<FNiagaraSpawnInfo>(FNiagaraDataSet& InDataSet, FNiagaraVariable InVar, bool bCurrBuffer = true)
-		: FNiagaraDataSetAccessorBase(&InDataSet, InVar, bCurrBuffer)
+	FNiagaraDataSetAccessor<FNiagaraSpawnInfo>(FNiagaraDataSet& InDataSet, FNiagaraVariable InVar)
+		: FNiagaraDataSetAccessorBase(&InDataSet, InVar)
 	{
 		check(sizeof(FNiagaraSpawnInfo) == InVar.GetType().GetSize());
-		InitForAccess(bCurrBuffer);
+		InitForAccess();
 	}
 
-	void InitForAccess(bool bCurrBuffer = true)
+	void InitForAccess()
 	{
-		DataBuffer = bCurrBuffer ? &DataSet->CurrData() : &DataSet->PrevData();
-		if (VarLayout != nullptr)
-		{
-			CountBase = (int32*)DataBuffer->GetComponentPtrInt32(VarLayout->Int32ComponentStart);
-			InterpStartDtBase = (float*)DataBuffer->GetComponentPtrFloat(VarLayout->FloatComponentStart);
-			IntervalDtBase = (float*)DataBuffer->GetComponentPtrFloat(VarLayout->FloatComponentStart + 1);
-			GroupBase = (int32*)DataBuffer->GetComponentPtrInt32(VarLayout->Int32ComponentStart + 1);
-		}
-		else
+		SrcCountBase = nullptr;
+		SrcInterpStartDtBase = nullptr;
+		SrcIntervalDtBase = nullptr;
+		SrcGroupBase = nullptr;
+		DestCountBase = nullptr;
+		DestInterpStartDtBase = nullptr;
+		DestIntervalDtBase = nullptr;
+		DestGroupBase = nullptr;
+		FNiagaraDataBuffer* SrcBuffer = DataSet->GetCurrentData();
+		if (IsValid() && SrcBuffer)
 		{
-			CountBase = nullptr;
-			InterpStartDtBase = nullptr;
-			IntervalDtBase = nullptr;
-			GroupBase = nullptr;
+			SrcCountBase = (int32*)SrcBuffer->GetComponentPtrInt32(VarLayout->Int32ComponentStart);
+			SrcInterpStartDtBase = (float*)SrcBuffer->GetComponentPtrFloat(VarLayout->FloatComponentStart);
+			SrcIntervalDtBase = (float*)SrcBuffer->GetComponentPtrFloat(VarLayout->FloatComponentStart + 1);
+			SrcGroupBase = (int32*)SrcBuffer->GetComponentPtrInt32(VarLayout->Int32ComponentStart + 1);
+
+			//Writes are only valid if we're during a simulation pass.
+			FNiagaraDataBuffer* DestBuffer = DataSet->GetDestinationData();
+			if (DestBuffer)
+			{
+				DestCountBase = (int32*)DestBuffer->GetComponentPtrInt32(VarLayout->Int32ComponentStart);
+				DestInterpStartDtBase = (float*)DestBuffer->GetComponentPtrFloat(VarLayout->FloatComponentStart);
+				DestIntervalDtBase = (float*)DestBuffer->GetComponentPtrFloat(VarLayout->FloatComponentStart + 1);
+				DestGroupBase = (int32*)DestBuffer->GetComponentPtrInt32(VarLayout->Int32ComponentStart + 1);
+			}
 		}
 	}
 
+	FORCEINLINE bool IsValidForRead()const { return SrcCountBase != nullptr && SrcInterpStartDtBase != nullptr && SrcIntervalDtBase != nullptr && SrcGroupBase != nullptr; }
+	FORCEINLINE bool IsValidForWrite()const { return DestCountBase != nullptr && DestInterpStartDtBase != nullptr && DestIntervalDtBase != nullptr && DestGroupBase != nullptr; }
+
 	FORCEINLINE FNiagaraSpawnInfo operator[](int32 Index)const
 	{
 		return Get(Index);
@@ -1457,9 +1234,17 @@ struct FNiagaraDataSetAccessor<FNiagaraSpawnInfo> : public FNiagaraDataSetAccess
 
 	FORCEINLINE FNiagaraSpawnInfo GetSafe(int32 Index, FNiagaraSpawnInfo Default = FNiagaraSpawnInfo())const
 	{
-		if (IsValid())
+		if (IsValidForRead())
 		{
-			return Get(Index);
+			FNiagaraDataBuffer* DataBuffer = DataSet->GetCurrentData();
+			if (DataBuffer && Index >= 0 && (uint32)Index < DataBuffer->GetNumInstances())
+			{
+				return Get(Index);
+			}
+			else
+			{
+				ensure(DataBuffer && Index >= 0 && (uint32)Index < DataBuffer->GetNumInstances()); // just to capture the badness in the logs or debugger if attached.
+			}
 		}
 
 		return Default;
@@ -1474,59 +1259,70 @@ struct FNiagaraDataSetAccessor<FNiagaraSpawnInfo> : public FNiagaraDataSetAccess
 
 	FORCEINLINE void Get(int32 Index, FNiagaraSpawnInfo& OutValue)const
 	{
-		OutValue.Count = CountBase[Index];
-		OutValue.InterpStartDt = InterpStartDtBase[Index];
-		OutValue.IntervalDt = IntervalDtBase[Index];
-		OutValue.SpawnGroup = GroupBase[Index];
+		checkSlow(IsValidForRead());
+		OutValue.Count = SrcCountBase[Index];
+		OutValue.InterpStartDt = SrcInterpStartDtBase[Index];
+		OutValue.IntervalDt = SrcIntervalDtBase[Index];
+		OutValue.SpawnGroup = SrcGroupBase[Index];
 	}
 
 	FORCEINLINE void Set(int32 Index, const FNiagaraSpawnInfo& InValue)
 	{
-		check(CountBase != nullptr);
-		check((uint32)Index < DataBuffer->GetNumInstances());
-
-		CountBase[Index] = InValue.Count;
-		InterpStartDtBase[Index] = InValue.InterpStartDt;
-		IntervalDtBase[Index] = InValue.IntervalDt;
-		GroupBase[Index] = InValue.SpawnGroup;
+		checkSlow(IsValidForWrite());
+		DestCountBase[Index] = InValue.Count;
+		DestInterpStartDtBase[Index] = InValue.InterpStartDt;
+		DestIntervalDtBase[Index] = InValue.IntervalDt;
+		DestGroupBase[Index] = InValue.SpawnGroup;
 	}
 
-
-	FORCEINLINE bool BaseIsValid() const { return CountBase != nullptr && InterpStartDtBase != nullptr && IntervalDtBase != nullptr && GroupBase != nullptr; }
-
 private:
 
-	int32* CountBase;
-	float* InterpStartDtBase;
-	float* IntervalDtBase;
-	int32* GroupBase;
+	int32* SrcCountBase;
+	float* SrcInterpStartDtBase;
+	float* SrcIntervalDtBase;
+	int32* SrcGroupBase;
+
+	int32* DestCountBase;
+	float* DestInterpStartDtBase;
+	float* DestIntervalDtBase;
+	int32* DestGroupBase;
 };
 
 template<>
 struct FNiagaraDataSetAccessor<FNiagaraID> : public FNiagaraDataSetAccessorBase
 {
 	FNiagaraDataSetAccessor<FNiagaraID>() {}
-	FNiagaraDataSetAccessor<FNiagaraID>(FNiagaraDataSet& InDataSet, FNiagaraVariable InVar, bool bCurrBuffer = true)
-		: FNiagaraDataSetAccessorBase(&InDataSet, InVar, bCurrBuffer)
+	FNiagaraDataSetAccessor<FNiagaraID>(FNiagaraDataSet& InDataSet, FNiagaraVariable InVar)
+		: FNiagaraDataSetAccessorBase(&InDataSet, InVar)
 	{
-		InitForAccess(bCurrBuffer);
+		InitForAccess();
 	}
 
-	void InitForAccess(bool bCurrBuffer = true)
+	void InitForAccess()
 	{
-		DataBuffer = bCurrBuffer ? &DataSet->CurrData() : &DataSet->PrevData();
-		if (VarLayout != nullptr)
+		SrcIndexBase = nullptr;
+		SrcTagBase = nullptr;
+		DestIndexBase = nullptr;
+		DestTagBase = nullptr;
+		FNiagaraDataBuffer* SrcBuffer = DataSet->GetCurrentData();
+		if (IsValid() && SrcBuffer)
 		{
-			IndexBase = (int32*)DataBuffer->GetComponentPtrInt32(VarLayout->Int32ComponentStart);
-			TagBase = (int32*)DataBuffer->GetComponentPtrInt32(VarLayout->Int32ComponentStart + 1);
-		}
-		else
-		{
-			IndexBase = nullptr;
-			TagBase = nullptr;
+			SrcIndexBase = (int32*)SrcBuffer->GetComponentPtrInt32(VarLayout->Int32ComponentStart);
+			SrcTagBase = (int32*)SrcBuffer->GetComponentPtrInt32(VarLayout->Int32ComponentStart + 1);
+
+			//Writes are only valid if we're during a simulation pass.
+			FNiagaraDataBuffer* DestBuffer = DataSet->GetDestinationData();
+			if (DestBuffer)
+			{
+				DestIndexBase = (int32*)DestBuffer->GetComponentPtrInt32(VarLayout->Int32ComponentStart);
+				DestTagBase = (int32*)DestBuffer->GetComponentPtrInt32(VarLayout->Int32ComponentStart + 1);
+			}
 		}
 	}
 
+	FORCEINLINE bool IsValidForRead()const { return SrcIndexBase != nullptr && SrcTagBase != nullptr; }
+	FORCEINLINE bool IsValidForWrite()const { return DestIndexBase != nullptr && DestTagBase != nullptr; }
+
 	FORCEINLINE FNiagaraID operator[](int32 Index)const
 	{
 		return Get(Index);
@@ -1534,9 +1330,17 @@ struct FNiagaraDataSetAccessor<FNiagaraID> : public FNiagaraDataSetAccessorBase
 
 	FORCEINLINE FNiagaraID GetSafe(int32 Index, FNiagaraID Default = FNiagaraID())const
 	{
-		if (IsValid())
+		if (IsValidForRead())
 		{
-			return Get(Index);
+			FNiagaraDataBuffer* DataBuffer = DataSet->GetCurrentData();
+			if (DataBuffer && Index >= 0 && (uint32)Index < DataBuffer->GetNumInstances())
+			{
+				return Get(Index);
+			}
+			else
+			{
+				ensure(DataBuffer && Index >= 0 && (uint32)Index < DataBuffer->GetNumInstances()); // just to capture the badness in the logs or debugger if attached.
+			}
 		}
 
 		return Default;
@@ -1551,94 +1355,34 @@ struct FNiagaraDataSetAccessor<FNiagaraID> : public FNiagaraDataSetAccessorBase
 
 	FORCEINLINE void Get(int32 Index, FNiagaraID& OutValue)const
 	{
-		OutValue.Index = IndexBase[Index];
-		OutValue.AcquireTag = TagBase[Index];
+		checkSlow(IsValidForRead());
+		OutValue.Index = SrcIndexBase[Index];
+		OutValue.AcquireTag = SrcTagBase[Index];
 	}
 
 	FORCEINLINE void Set(int32 Index, const FNiagaraID& InValue)
 	{
-		check(IndexBase != nullptr);
-		check((uint32)Index < DataBuffer->GetNumInstances());
-
-		IndexBase[Index] = InValue.Index;
-		TagBase[Index] = InValue.AcquireTag;
+		checkSlow(IsValidForWrite());
+		DestIndexBase[Index] = InValue.Index;
+		DestTagBase[Index] = InValue.AcquireTag;
 	}
 
-	FORCEINLINE bool BaseIsValid() const { return IndexBase != nullptr && TagBase != nullptr; }
-
 private:
 
-	int32* IndexBase;
-	int32* TagBase;
-};
-
-//Provide iterator to keep iterator access patterns still in use.
-template<typename T>
-struct FNiagaraDataSetIterator : public FNiagaraDataSetAccessor<T>
-{
-	FNiagaraDataSetIterator<T>() {}
-	FNiagaraDataSetIterator(FNiagaraDataSet& InDataSet, FNiagaraVariable InVar, uint32 StartIndex = 0, bool bCurrBuffer = true)
-		: FNiagaraDataSetAccessor<T>(InDataSet, InVar, bCurrBuffer)
-		, CurrIdx(StartIndex)
-	{}
-
-	void Create(FNiagaraDataSet* InDataSet, FNiagaraVariable InVar, uint32 StartIndex=0)
-	{
-		FNiagaraDataSetAccessor<T>::Create(InDataSet, InVar);
-		CurrIdx = StartIndex;
-	}
-
-	FORCEINLINE T operator*() { return Get(); }
-	FORCEINLINE T Get()const
-	{
-		T Ret;
-		Get(Ret);
-		return Ret;
-	}
-
-	FORCEINLINE T GetAdvance()
-	{
-		T Ret;
-		Get(Ret);
-		Advance();
-		return Ret;
-	}
-
-	FORCEINLINE T GetAdvanceWithDefault(const T& Default)
-	{
-		T Ret = IsValid() ? Get() : Default;
-		Advance();
-		return Ret;
-	}
-
-	FORCEINLINE void Get(T& OutValue)const { FNiagaraDataSetAccessor<T>::Get(CurrIdx, OutValue); }
-
-	FORCEINLINE void Set(const T& InValue)
-	{
-		FNiagaraDataSetAccessor<T>::Set(CurrIdx, InValue);
-	}
-
-	FORCEINLINE void Advance() { ++CurrIdx; }
-	FORCEINLINE bool IsValid()const
-	{
-		return FNiagaraDataSetAccessorBase::VarLayout != nullptr && CurrIdx < FNiagaraDataSetAccessorBase::DataBuffer->GetNumInstances();
-	}
-
-	uint32 GetCurrIndex()const { return CurrIdx; }
-
-private:
-	uint32 CurrIdx;
+	int32* SrcIndexBase;
+	int32* SrcTagBase;
+	int32* DestIndexBase;
+	int32* DestTagBase;
 };
 
 /**
-Iterator that will pull or push data between a DataSet and some FNiagaraVariables it contains.
+Iterator that will pull or push data between a NiagaraDataBuffer and some FNiagaraVariables it contains.
 Super slow. Don't use at runtime.
 */
-struct FNiagaraDataSetVariableIterator
+struct FNiagaraDataVariableIterator
 {
-	FNiagaraDataSetVariableIterator(FNiagaraDataSet& InDataSet, uint32 StartIdx = 0, bool bCurrBuffer = true)
-		: DataSet(InDataSet)
-		, DataBuffer(bCurrBuffer ? DataSet.CurrData() : DataSet.PrevData())
+	FNiagaraDataVariableIterator(const FNiagaraDataBuffer* InData, uint32 StartIdx = 0)
+		: Data(InData)
 		, CurrIdx(StartIdx)
 	{
 	}
@@ -1655,7 +1399,7 @@ struct FNiagaraDataSetVariableIterator
 			for (uint32 CompIdx = 0; CompIdx < Layout->GetNumFloatComponents(); ++CompIdx)
 			{
 				uint32 CompBufferOffset = Layout->FloatComponentStart + CompIdx;
-				float* Src = DataBuffer.GetInstancePtrFloat(CompBufferOffset, CurrIdx);
+				float* Src = Data->GetInstancePtrFloat(CompBufferOffset, CurrIdx);
 				float* Dst = (float*)(ValuePtr + Layout->LayoutInfo.FloatComponentByteOffsets[CompIdx]);
 				*Dst = *Src;
 			}
@@ -1663,56 +1407,27 @@ struct FNiagaraDataSetVariableIterator
 			for (uint32 CompIdx = 0; CompIdx < Layout->GetNumInt32Components(); ++CompIdx)
 			{
 				uint32 CompBufferOffset = Layout->Int32ComponentStart + CompIdx;
-				int32* Src = DataBuffer.GetInstancePtrInt32(CompBufferOffset, CurrIdx);
+				int32* Src = Data->GetInstancePtrInt32(CompBufferOffset, CurrIdx);
 				int32* Dst = (int32*)(ValuePtr + Layout->LayoutInfo.Int32ComponentByteOffsets[CompIdx]);
 				*Dst = *Src;
 			}
 		}
 	}
 
-	void Set()
-	{
-		check(CurrIdx < DataBuffer.GetNumInstances());
-		for (int32 VarIdx = 0; VarIdx < Variables.Num(); ++VarIdx)
-		{
-			FNiagaraVariable* Var = Variables[VarIdx];
-			const FNiagaraVariableLayoutInfo* Layout = VarLayouts[VarIdx];
-
-			check(Var && Layout);
-			uint8* ValuePtr = Var->GetData();
-
-			for (uint32 CompIdx = 0; CompIdx < Layout->GetNumFloatComponents(); ++CompIdx)
-			{
-				uint32 CompBufferOffset = Layout->FloatComponentStart + CompIdx;
-				float* Dst = DataBuffer.GetInstancePtrFloat(CompBufferOffset, CurrIdx);
-				float* Src = (float*)(ValuePtr + Layout->LayoutInfo.FloatComponentByteOffsets[CompIdx]);
-				*Dst = *Src;
-			}
-
-			for (uint32 CompIdx = 0; CompIdx < Layout->GetNumInt32Components(); ++CompIdx)
-			{
-				uint32 CompBufferOffset = Layout->FloatComponentStart + CompIdx;
-				int32* Dst = DataBuffer.GetInstancePtrInt32(CompBufferOffset, CurrIdx);
-				int32* Src = (int32*)(ValuePtr + Layout->LayoutInfo.FloatComponentByteOffsets[CompIdx]);
-				*Dst = *Src;
-			}
-		}
-	}
-
 	void Advance() { ++CurrIdx; }
 
 	bool IsValid()const
 	{
-		return CurrIdx < DataBuffer.GetNumInstances();
+		return Data && CurrIdx < Data->GetNumInstances();
 	}
 
 	uint32 GetCurrIndex()const { return CurrIdx; }
 
 	void AddVariable(FNiagaraVariable* InVar)
 	{
-		check(InVar);
+		check(InVar && Data);
 		Variables.AddUnique(InVar);
-		VarLayouts.AddUnique(DataSet.GetVariableLayout(*InVar));
+		VarLayouts.AddUnique(Data->GetOwner()->GetVariableLayout(*InVar));
 		InVar->AllocateData();
 	}
 
@@ -1725,8 +1440,7 @@ struct FNiagaraDataSetVariableIterator
 	}
 private:
 
-	FNiagaraDataSet& DataSet;
-	FNiagaraDataBuffer& DataBuffer;
+	const FNiagaraDataBuffer* Data;
 	TArray<FNiagaraVariable*> Variables;
 	TArray<const FNiagaraVariableLayoutInfo*> VarLayouts;
 
@@ -1734,75 +1448,22 @@ private:
 };
 
 /**
-Iterator that will pull or push data between a DataSet and some FNiagaraVariables it contains.
-Super slow. Don't use at runtime.
+Allows immediate access to GPU data on the CPU, you can then use FNiagaraDataSetAccessor to access the data.
+This will make a copy of the GPU data and will stall the CPU until the data is ready from the GPU,
+therefore it should only be used for tools / debugging.  For async readback see FNiagaraSystemInstance::RequestCapture.
 */
-struct FNiagaraDataSetVariableIteratorConst
+class NIAGARA_API FScopedNiagaraDataSetGPUReadback
 {
-	FNiagaraDataSetVariableIteratorConst(const FNiagaraDataSet& InDataSet, uint32 StartIdx = 0, bool bCurrBuffer = true)
-		: DataSet(InDataSet)
-		, DataBuffer(bCurrBuffer ? DataSet.CurrData() : DataSet.PrevData())
-		, CurrIdx(StartIdx)
-	{
-	}
-
-	void Get()
-	{
-		for (int32 VarIdx = 0; VarIdx < Variables.Num(); ++VarIdx)
-		{
-			FNiagaraVariable* Var = Variables[VarIdx];
-			const FNiagaraVariableLayoutInfo* Layout = VarLayouts[VarIdx];
-			check(Var && Layout);
-			uint8* ValuePtr = Var->GetData();
-
-			for (uint32 CompIdx = 0; CompIdx < Layout->GetNumFloatComponents(); ++CompIdx)
-			{
-				uint32 CompBufferOffset = Layout->FloatComponentStart + CompIdx;
-				float* Src = DataBuffer.GetInstancePtrFloat(CompBufferOffset, CurrIdx);
-				float* Dst = (float*)(ValuePtr + Layout->LayoutInfo.FloatComponentByteOffsets[CompIdx]);
-				*Dst = *Src;
-			}
-
-			for (uint32 CompIdx = 0; CompIdx < Layout->GetNumInt32Components(); ++CompIdx)
-			{
-				uint32 CompBufferOffset = Layout->Int32ComponentStart + CompIdx;
-				int32* Src = DataBuffer.GetInstancePtrInt32(CompBufferOffset, CurrIdx);
-				int32* Dst = (int32*)(ValuePtr + Layout->LayoutInfo.Int32ComponentByteOffsets[CompIdx]);
-				*Dst = *Src;
-			}
-		}
-	}
-	
-	void Advance() { ++CurrIdx; }
-
-	bool IsValid()const
-	{
-		return CurrIdx < DataBuffer.GetNumInstances();
-	}
+public:
+	FScopedNiagaraDataSetGPUReadback() {}
+	FScopedNiagaraDataSetGPUReadback(FNiagaraDataSet* DataSet) { ReadbackData(DataSet); }
+	~FScopedNiagaraDataSetGPUReadback();
 
-	uint32 GetCurrIndex()const { return CurrIdx; }
+	void ReadbackData(FNiagaraDataSet* InDataSet);
+	uint32 GetNumInstances() const { check(DataSet != nullptr); return NumInstances; }
 
-	void AddVariable(FNiagaraVariable* InVar)
-	{
-		check(InVar);
-		Variables.AddUnique(InVar);
-		VarLayouts.AddUnique(DataSet.GetVariableLayout(*InVar));
-		InVar->AllocateData();
-	}
-
-	void AddVariables(TArray<FNiagaraVariable>& Vars)
-	{
-		for (FNiagaraVariable& Var : Vars)
-		{
-			AddVariable(&Var);
-		}
-	}
 private:
-
-	const FNiagaraDataSet & DataSet;
-	const FNiagaraDataBuffer& DataBuffer;
-	TArray<FNiagaraVariable*> Variables;
-	TArray<const FNiagaraVariableLayoutInfo*> VarLayouts;
-
-	uint32 CurrIdx;
-};
\ No newline at end of file
+	FNiagaraDataSet*	DataSet = nullptr;
+	FNiagaraDataBuffer* DataBuffer = nullptr;
+	uint32				NumInstances = 0;
+};
diff --git a/Engine/Plugins/FX/Niagara/Source/Niagara/Classes/NiagaraEmitter.h b/Engine/Plugins/FX/Niagara/Source/Niagara/Classes/NiagaraEmitter.h
--- a/Engine/Plugins/FX/Niagara/Source/Niagara/Classes/NiagaraEmitter.h
+++ b/Engine/Plugins/FX/Niagara/Source/Niagara/Classes/NiagaraEmitter.h
@@ -351,6 +351,8 @@ public:
 	/* Gets whether or not the supplied event generator id matches an event generator which is shared between the particle spawn and update scrips. */
 	bool IsEventGeneratorShared(FName EventGeneratorId) const;
 
+	TStatId GetStatID(bool bGameThread, bool bConcurrent)const;
+
 protected:
 	virtual void BeginDestroy() override;
 
@@ -393,6 +395,14 @@ private:
 #if WITH_EDITOR
 	FOnPropertiesChanged OnPropertiesChangedDelegate;
 #endif
+
+	void GenerateStatID();
+#if STATS
+	TStatId StatID_GT;
+	TStatId StatID_GT_CNC;
+	TStatId StatID_RT;
+	TStatId StatID_RT_CNC;
+#endif
 };
 
 
diff --git a/Engine/Plugins/FX/Niagara/Source/Niagara/Classes/NiagaraEmitterInstance.h b/Engine/Plugins/FX/Niagara/Source/Niagara/Classes/NiagaraEmitterInstance.h
--- a/Engine/Plugins/FX/Niagara/Source/Niagara/Classes/NiagaraEmitterInstance.h
+++ b/Engine/Plugins/FX/Niagara/Source/Niagara/Classes/NiagaraEmitterInstance.h
@@ -14,17 +14,19 @@ NiagaraEmitterInstance.h: Niagara emitter simulation class
 #include "NiagaraEmitterHandle.h"
 #include "NiagaraEmitter.h"
 #include "NiagaraScriptExecutionContext.h"
+#include "NiagaraBoundsCalculator.h"
 
 class FNiagaraSystemInstance;
-class NiagaraRenderer;
 struct FNiagaraEmitterHandle;
 class UNiagaraParameterCollection;
 class UNiagaraParameterCollectionInstance;
 class NiagaraEmitterInstanceBatcher;
+
+
 /**
 * A Niagara particle simulation.
 */
-struct FNiagaraEmitterInstance
+class FNiagaraEmitterInstance
 {
 public:
 	explicit FNiagaraEmitterInstance(FNiagaraSystemInstance* InParentSystemInstance);
@@ -56,23 +58,27 @@ public:
 
 	uint32 CalculateEventSpawnCount(const FNiagaraEventScriptProperties &EventHandlerProps, TArray<int32, TInlineAllocator<16>>& EventSpawnCounts, FNiagaraDataSet *EventSet);
 
-	NIAGARA_API TOptional<FBox> CalculateDynamicBounds();
-
-	FNiagaraDataSet &GetData()	{ return *ParticleDataSet; }
-
-	int32 GetEmitterRendererNum() const {
-		return EmitterRenderer.Num();
-	}
-
-	NiagaraRenderer *GetEmitterRenderer(int32 Index)	{ return EmitterRenderer[Index]; }
+	/** Generate system bounds, reading back data from the GPU will introduce a stall and should only be used for debug purposes. */
+	FBox CalculateDynamicBounds(const bool bReadGPUSimulation = false);
+	NIAGARA_API void CalculateFixedBounds(const FTransform& ToWorldSpace);
 
+	FNiagaraDataSet& GetData()const { return *ParticleDataSet; }
+	
 	FORCEINLINE bool IsDisabled()const { return ExecutionState == ENiagaraExecutionState::Disabled; }
 	FORCEINLINE bool IsComplete()const { return ExecutionState == ENiagaraExecutionState::Complete || ExecutionState == ENiagaraExecutionState::Disabled; }
 		
 	/** Create a new NiagaraRenderer. The old renderer is not immediately deleted, but instead put in the ToBeRemoved list.*/
-	void NIAGARA_API UpdateEmitterRenderer(ERHIFeatureLevel::Type FeatureLevel, TArray<NiagaraRenderer*>& ToBeAddedList, TArray<NiagaraRenderer*>& ToBeRemovedList);
-
-	FORCEINLINE int32 GetNumParticles()const	{ return ParticleDataSet->GetNumInstances(); }
+	//void NIAGARA_API UpdateEmitterRenderer(ERHIFeatureLevel::Type FeatureLevel, TArray<NiagaraRenderer*>& ToBeAddedList, TArray<NiagaraRenderer*>& ToBeRemovedList);
+
+	FORCEINLINE int32 GetNumParticles()const
+	{
+		// Note: For ENiagaraSimTarget::GPUComputeSim this data is latent
+		if (ParticleDataSet->GetCurrentData())
+		{
+			return ParticleDataSet->GetCurrentData()->GetNumInstances();
+		}
+		return 0;
+	}
 	FORCEINLINE int32 GetTotalSpawnedParticles()const	{ return TotalSpawnedParticles; }
 
 	NIAGARA_API const FNiagaraEmitterHandle& GetEmitterHandle() const;
@@ -87,9 +93,6 @@ public:
 
 	FNiagaraDataSet* GetDataSet(FNiagaraDataSetID SetID);
 
-	/** Tell the renderer thread that we're done with the Niagara renderer on this simulation.*/
-	void ClearRenderer();
-
 	FBox GetBounds();
 
 	FNiagaraScriptExecutionContext& GetSpawnExecutionContext() { return SpawnExecContext; }
@@ -107,6 +110,13 @@ public:
 
 	bool WaitForDebugInfo();
 
+	FNiagaraComputeExecutionContext* GetGPUContext()const
+	{
+		return GPUExecContext;
+	}
+
+	void SetSystemFixedBoundsOverride(FBox SystemFixedBounds);
+
 private:
 
 	void CheckForErrors();
@@ -178,7 +188,6 @@ private:
 	/** particle simulation data. Must be a shared ref as various things on the RT can have direct ref to it. */
 	FNiagaraDataSet* ParticleDataSet;
 
-	TArray<NiagaraRenderer*> EmitterRenderer;
 	FNiagaraSystemInstance *ParentSystemInstance;
 
 	/** Raw pointer to the emitter that we're instanced from. Raw ptr should be safe here as we check for the validity of the system and it's emitters higher up before any ticking. */
@@ -194,18 +203,14 @@ private:
 
 	FName OwnerSystemInstanceName;
 
+	/** Cached fixed bounds of the parent system which override this Emitter Instances bounds if set. Whenever we initialize the owning SystemInstance we will reconstruct this
+	 ** EmitterInstance and the cached bounds will be unset. */
+	TOptional<FBox> CachedSystemFixedBounds;
+
 #if WITH_EDITORONLY_DATA
 	bool CheckAttributesForRenderer(int32 Index);
 #endif
 
-	FNiagaraDataSetAccessor<FVector> PositionAccessor;
-	FNiagaraDataSetAccessor<FVector2D> SizeAccessor;
-	FNiagaraDataSetAccessor<FVector> MeshScaleAccessor;
-
-	static const FName PositionName;
-	static const FName SizeName;
-	static const FName MeshScaleName;
-
 #if !UE_BUILD_SHIPPING
 	bool bEncounteredNaNs;
 #endif
@@ -214,4 +219,11 @@ private:
 	FNiagaraParameterStore ScriptDefinedDataInterfaceParameters;
 
 	NiagaraEmitterInstanceBatcher* Batcher = nullptr;
+
+	/** Data required for handling events. */
+	TArray<FNiagaraEventHandlingInfo> EventHandlingInfo;
+	int32 EventSpawnTotal;
+
+	/** Optional list of bounds calculators. */
+	TArray<TUniquePtr<FNiagaraBoundsCalculator>, TInlineAllocator<1>> BoundsCalculators;
 };
diff --git a/Engine/Plugins/FX/Niagara/Source/Niagara/Classes/NiagaraEmitterInstanceBatcher.h b/Engine/Plugins/FX/Niagara/Source/Niagara/Classes/NiagaraEmitterInstanceBatcher.h
--- a/Engine/Plugins/FX/Niagara/Source/Niagara/Classes/NiagaraEmitterInstanceBatcher.h
+++ b/Engine/Plugins/FX/Niagara/Source/Niagara/Classes/NiagaraEmitterInstanceBatcher.h
@@ -20,11 +20,7 @@ the same VectorVM byte code / compute shader code
 #include "ParticleResources.h"
 #include "Runtime/Engine/Private/Particles/ParticleSortingGPU.h"
 #include "NiagaraGPUSortInfo.h"
-
-struct FNiagaraScriptExecutionContext;
-struct FNiagaraComputeExecutionContext;
-
-#define SIMULATION_QUEUE_COUNT 2
+#include "NiagaraScriptExecutionContext.h"
 
 class FNiagaraIndicesVertexBuffer : public FParticleIndicesVertexBuffer
 {
@@ -44,12 +40,15 @@ public:
 class NiagaraEmitterInstanceBatcher : public FFXSystemInterface
 {
 public:
+	using FNiagaraBufferArray = TArray<FUnorderedAccessViewRHIParamRef, TMemStackAllocator<>>;
+	using FOverlappableTicks = TArray<FNiagaraGPUSystemTick*, TMemStackAllocator<>>;
 
 	NiagaraEmitterInstanceBatcher(ERHIFeatureLevel::Type InFeatureLevel, EShaderPlatform InShaderPlatform)
 		: FeatureLevel(InFeatureLevel)
 		, ShaderPlatform(InShaderPlatform)
-		, CurQueueIndex(0)
 		, ParticleSortBuffers(true)
+		// @todo REMOVE THIS HACK
+		, LastFrameThatDrainedData(GFrameNumberRenderThread)
 	{
 	}
 
@@ -58,9 +57,16 @@ public:
 	static const FName Name;
 	virtual FFXSystemInterface* GetInterface(const FName& InName) override;
 
-	void Queue(FNiagaraComputeExecutionContext *InContext);
+	// The batcher assumes ownership of the data here.
+	void GiveSystemTick_RenderThread(FNiagaraGPUSystemTick& Tick);
+
+	void GiveEmitterContextToDestroy_RenderThread(FNiagaraComputeExecutionContext* Context);
+	void GiveDataSetToDestroy_RenderThread(FNiagaraDataSet* DataSet);
 
-	void Remove(FNiagaraComputeExecutionContext* InContext);
+	void EnqueueDeferredDeletesForDI_RenderThread(TSharedPtr<FNiagaraDataInterfaceProxy, ESPMode::ThreadSafe> Proxy)
+	{
+		DIProxyDeferredDeletes_RT.Add(MoveTemp(Proxy));
+	}
 
 #if WITH_EDITOR
 	virtual void Suspend() override {}
@@ -72,8 +78,8 @@ public:
 	virtual void RemoveVectorField(UVectorFieldComponent* VectorFieldComponent) override {}
 	virtual void UpdateVectorField(UVectorFieldComponent* VectorFieldComponent) override {}
 	virtual void PreInitViews() override;
-	virtual bool UsesGlobalDistanceField() const override { return false; }
-	virtual void PreRender(FRHICommandListImmediate& RHICmdList, const class FGlobalDistanceFieldParameterData* GlobalDistanceFieldParameterData) override {}
+	virtual bool UsesGlobalDistanceField() const override;
+	virtual void PreRender(FRHICommandListImmediate& RHICmdList, const class FGlobalDistanceFieldParameterData* GlobalDistanceFieldParameterData) override;
 
 	virtual void Tick(float DeltaTime) override
 	{
@@ -94,33 +100,29 @@ public:
 		FRHICommandListImmediate& RHICmdList,
 		const FUniformBufferRHIParamRef ViewUniformBuffer,
 		const class FShaderParametersMetadata* SceneTexturesUniformBufferStruct,
-		FUniformBufferRHIParamRef SceneTexturesUniformBuffer) override
-	{
-		CurQueueIndex ^= 0x1;
-
-		ExecuteAll(RHICmdList, ViewUniformBuffer);
-		SortGPUParticles(RHICmdList);
-	}
+		FUniformBufferRHIParamRef SceneTexturesUniformBuffer) override;
 
 	void ExecuteAll(FRHICommandList &RHICmdList, FUniformBufferRHIParamRef ViewUniformBuffer);
-	void TickSingle(FNiagaraComputeExecutionContext *Context, FRHICommandList &RHICmdList, FUniformBufferRHIParamRef ViewUniformBuffer) const;
+	void TickSingle(const FNiagaraGPUSystemTick& Tick, FNiagaraComputeInstanceData *Instance, FRHICommandList &RHICmdList, FUniformBufferRHIParamRef ViewUniformBuffer) const;
+	void ResizeBuffersAndGatherResources(FOverlappableTicks& OverlappableTick, FRHICommandList& RHICmdList, FNiagaraBufferArray& DestDataBuffers, FNiagaraBufferArray& CurrDataBuffers, FNiagaraBufferArray& DestBufferIntFloat, FNiagaraBufferArray& CurrBufferIntFloat);
+
+	void DispatchAllOnCompute(FOverlappableTicks& OverlappableTick, FRHICommandList& RHICmdList, FUniformBufferRHIParamRef ViewUniformBuffer, FNiagaraBufferArray& DestDataBuffers, FNiagaraBufferArray& CurrDataBuffers, FNiagaraBufferArray& DestBufferIntFloat, FNiagaraBufferArray& CurrBufferIntFloat);
 
 	int32 AddSortedGPUSimulation(const FNiagaraGPUSortInfo& SortInfo);
 	void SortGPUParticles(FRHICommandListImmediate& RHICmdList);
 	void ResolveParticleSortBuffers(FRHICommandListImmediate& RHICmdList, int32 ResultBufferIndex);
 
 	const FParticleIndicesVertexBuffer& GetGPUSortedBuffer() const { return SortedVertexBuffers.Last(); }
+	const FGlobalDistanceFieldParameterData& GetGlobalDistanceFieldParameters() const { return GlobalDistanceFieldParams; }
 
-	void ProcessDebugInfo(FRHICommandList &RHICmdList, const FNiagaraComputeExecutionContext *Context) const;
+	void ProcessDebugInfo(FRHICommandList &RHICmdLis, const FNiagaraComputeExecutionContext* Context) const;
 
-	void SetPrevDataStrideParams(const FNiagaraDataSet *Set, FNiagaraShader *Shader, FRHICommandList &RHICmdList) const;
+	void SetDataInterfaceParameters(const TArray<FNiagaraDataInterfaceProxy*>& DataInterfaceProxies, FNiagaraShader* Shader, FRHICommandList &RHICmdList, const FNiagaraComputeInstanceData* Instance, const FNiagaraGPUSystemTick& Tick) const;
+	void UnsetDataInterfaceParameters(const TArray<FNiagaraDataInterfaceProxy*>& DataInterfaceProxies, FNiagaraShader* Shader, FRHICommandList& RHICmdList, const FNiagaraComputeInstanceData* Instance, const FNiagaraGPUSystemTick& Tick) const;
 
-	void SetupEventUAVs(const FNiagaraComputeExecutionContext *Context, uint32 NumInstances, FRHICommandList &RHICmdList) const;
-	void UnsetEventUAVs(const FNiagaraComputeExecutionContext *Context, FRHICommandList &RHICmdList) const;
-	void SetDataInterfaceParameters(const TArray<UNiagaraDataInterface*> &DataInterfaces, FNiagaraShader *Shader, FRHICommandList &RHICmdList, const FNiagaraComputeExecutionContext *Context) const;
-	void UnsetDataInterfaceParameters(const TArray<UNiagaraDataInterface*> &DataInterfaces, FNiagaraShader* Shader, FRHICommandList &RHICmdList, const FNiagaraComputeExecutionContext *Context) const;
-
-	void Run(	const FNiagaraComputeExecutionContext *Context, 
+	template<bool bDoResourceTransitions>
+	void Run(	const FNiagaraGPUSystemTick& Tick,
+				const FNiagaraComputeInstanceData* Instance, 
 				uint32 UpdateStartInstance, 
 				const uint32 TotalNumInstances, 
 				FNiagaraShader* Shader,
@@ -129,30 +131,50 @@ public:
 				bool bCopyBeforeStart = false
 			) const;
 
-	void RunEventHandlers(const FNiagaraComputeExecutionContext *Context, uint32 NumInstancesAfterSim, uint32 NumInstancesAfterSpawn, uint32 NumInstancesAfterNonEventSpawn, FRHICommandList &RhiCmdList) const;
-
-	void ClearIndexBufferCur(FRHICommandList &RHICmdList, FNiagaraComputeExecutionContext *Context) const;
-	void ResolveDatasetWrites(FRHICommandList &RHICmdList, FNiagaraComputeExecutionContext *Context) const;
+	void ResolveDatasetWrites(FRHICommandList &RHICmdList, FNiagaraComputeInstanceData *Context) const;
 	void ResizeCurrentBuffer(FRHICommandList &RHICmdList, FNiagaraComputeExecutionContext *Context, uint32 NewNumInstances, uint32 PrevNumInstances) const;
 
 private:
+	void SimStepClearAndSetup(const FNiagaraComputeInstanceData& Instance, FRHICommandList& RHICmdList) const;
+//	void SimStepReadback(const FNiagaraComputeInstanceData& Instance, FRHICommandList& RHICmdList) const;
+
+	inline uint32 UnpackEmitterDispatchCount(uint8* PackedData)
+	{
+		return *(uint32*)PackedData;
+	}
+
+	inline FNiagaraComputeInstanceData* UnpackEmitterComputeDispatchArray(uint8* PackedData)
+	{
+		return (FNiagaraComputeInstanceData*)(PackedData + sizeof(uint32));
+	}
+
+	void FinishDispatches();
+	void ReleaseTicks();
 
 	/** Feature level of this effects system */
 	ERHIFeatureLevel::Type FeatureLevel;
 	/** Shader platform that will be rendering this effects system */
 	EShaderPlatform ShaderPlatform;
 
-	uint32 CurQueueIndex;
-
-	TArray<FNiagaraComputeExecutionContext*> SimulationQueue[SIMULATION_QUEUE_COUNT];
-
 	// Number of particle to sort this frame.
 	int32 SortedParticleCount = 0;
 	int32 NumFramesRequiringShrinking = 0;
 	TArray<FNiagaraGPUSortInfo> SimulationsToSort;
 	FParticleSortBuffers ParticleSortBuffers;
 
+	// @todo REMOVE THIS HACK
+	uint32 LastFrameThatDrainedData;
+
 	// The result of the GPU sort. Each next element replace the previous.
 	// The last entry is used to transfer the result of the ParticleSortBuffers.
 	TIndirectArray<FNiagaraIndicesVertexBuffer> SortedVertexBuffers;
+	
+	TArray<FNiagaraGPUSystemTick> Ticks_RT;
+	FGlobalDistanceFieldParameterData GlobalDistanceFieldParams;
+
+	// For deferred deletion of resources
+	TSet<FNiagaraComputeExecutionContext*> ContextsToDestroy_RT;
+	TSet<FNiagaraDataSet*> DataSetsToDestroy_RT;
+
+	TSet<TSharedPtr<FNiagaraDataInterfaceProxy, ESPMode::ThreadSafe>> DIProxyDeferredDeletes_RT;
 };
diff --git a/Engine/Plugins/FX/Niagara/Source/Niagara/Classes/NiagaraEvents.h b/Engine/Plugins/FX/Niagara/Source/Niagara/Classes/NiagaraEvents.h
--- a/Engine/Plugins/FX/Niagara/Source/Niagara/Classes/NiagaraEvents.h
+++ b/Engine/Plugins/FX/Niagara/Source/Niagara/Classes/NiagaraEvents.h
@@ -41,7 +41,7 @@ struct FNiagaraCollisionEventPayload
 };
 
 
-struct FNiagaraEmitterInstance;
+class FNiagaraEmitterInstance;
 
 /**
 Base class for actions that an event receiver will perform at the emitter level.
