Copying //UE4/Dev-Niagara to Dev-Main (//UE4/Dev-Main)

diff --git a/Engine/Plugins/FX/Niagara/Source/Niagara/Classes/NiagaraScript.h b/Engine/Plugins/FX/Niagara/Source/Niagara/Classes/NiagaraScript.h
--- a/Engine/Plugins/FX/Niagara/Source/Niagara/Classes/NiagaraScript.h
+++ b/Engine/Plugins/FX/Niagara/Source/Niagara/Classes/NiagaraScript.h
@@ -139,6 +139,16 @@ public:
 	UPROPERTY()
 	FGuid BaseScriptID;
 
+	/**
+	* The hash of the subgraph this shader primarily represents.
+	*/
+	UPROPERTY()
+	FNiagaraCompileHash BaseScriptCompileHash;
+
+	/** Compile hashes of any top level scripts the script was dependent on that might trigger a recompile if they change. */
+	UPROPERTY()
+	TArray<FNiagaraCompileHash> ReferencedCompileHashes;
+
 	/** Guids of any functions, module scripts, parameter collections, or other assets the script was dependent on that might trigger a recompile if they change. */
 	UPROPERTY()
 	TArray<FGuid> ReferencedDependencyIds;
@@ -427,6 +437,7 @@ public:
 
 	NIAGARA_API bool CanBeRunOnGpu() const;
 	NIAGARA_API bool IsReadyToRun(ENiagaraSimTarget SimTarget) const;
+	NIAGARA_API bool ShouldCacheShadersForCooking() const;
 
 #if WITH_EDITORONLY_DATA
 	class UNiagaraScriptSourceBase *GetSource() { return Source; }
diff --git a/Engine/Plugins/FX/Niagara/Source/Niagara/Classes/NiagaraScriptExecutionContext.h b/Engine/Plugins/FX/Niagara/Source/Niagara/Classes/NiagaraScriptExecutionContext.h
--- a/Engine/Plugins/FX/Niagara/Source/Niagara/Classes/NiagaraScriptExecutionContext.h
+++ b/Engine/Plugins/FX/Niagara/Source/Niagara/Classes/NiagaraScriptExecutionContext.h
@@ -17,28 +17,93 @@ NiagaraEmitterInstance.h: Niagara emitter simulation class
 #include "NiagaraTypes.h"
 #include "RHIGPUReadback.h"
 
+struct FNiagaraDataInterfaceProxy;
+
+/** Container for data needed to process event data. */
+struct FNiagaraEventHandlingInfo
+{
+	FNiagaraEventHandlingInfo()
+		: TotalSpawnCount(0)
+		, EventData(nullptr)
+		, SourceEmitterName(NAME_None)
+	{
+	}
+
+	~FNiagaraEventHandlingInfo()
+	{
+		SetEventData(nullptr);
+	}
+
+	void SetEventData(FNiagaraDataBuffer* InEventData)
+	{
+		if (EventData)
+		{
+			EventData->ReleaseReadRef();
+		}
+		EventData = InEventData;
+		if (EventData)
+		{
+			EventData->AddReadRef();
+		}
+	}
+
+	TArray<int32, TInlineAllocator<16>> SpawnCounts;
+	int32 TotalSpawnCount;
+	FNiagaraDataBuffer* EventData;
+	FGuid SourceEmitterGuid;
+	FName SourceEmitterName;
+};
 
 
 struct FNiagaraDataSetExecutionInfo
 {
 	FNiagaraDataSetExecutionInfo()
 		: DataSet(nullptr)
+		, Input(nullptr)
+		, Output(nullptr)
 		, StartInstance(0)
-		, bAllocate(false)
 		, bUpdateInstanceCount(false)
 	{
 	}
 
-	FNiagaraDataSetExecutionInfo(FNiagaraDataSet* InDataSet, int32 InStartInstance, bool bInAllocate, bool bInUpdateInstanceCount)
-		: DataSet(InDataSet)
-		, StartInstance(InStartInstance)
-		, bAllocate(bInAllocate)
-		, bUpdateInstanceCount(bInUpdateInstanceCount)
-	{}
+
+	void Init(FNiagaraDataSet* InDataSet, FNiagaraDataBuffer* InInput, FNiagaraDataBuffer* InOutput, int32 InStartInstance, bool bInUpdateInstanceCount)
+	{
+		if (Input)
+		{
+			Input->ReleaseReadRef();
+		}
+
+		DataSet = InDataSet;
+		Input = InInput;
+		Output = InOutput;
+		StartInstance = InStartInstance;
+		bUpdateInstanceCount = bInUpdateInstanceCount;
+
+		check(DataSet);
+		check(Input == nullptr || DataSet == Input->GetOwner());
+		check(Output == nullptr || DataSet == Output->GetOwner());
+
+		if (Input)
+		{
+			Input->AddReadRef();
+		}
+		check(Output == nullptr || Output->IsBeingWritten());
+	}
+	
+	~FNiagaraDataSetExecutionInfo()
+	{
+		check(Output == nullptr || Output->IsBeingWritten());
+		if (Input)
+		{
+			Input->ReleaseReadRef();
+		}
+	}
 
 	FNiagaraDataSet* DataSet;
+	FNiagaraDataBuffer* Input;
+	FNiagaraDataBuffer* Output;
 	int32 StartInstance;
-	bool bAllocate;
 	bool bUpdateInstanceCount;
 };
 
@@ -57,6 +122,8 @@ struct FNiagaraScriptExecutionContext
 
 	TArray<FDataSetMeta> DataSetMetaTable;
 
+	TArray<FNiagaraDataSetExecutionInfo, TInlineAllocator<4>> DataSetInfo;
+
 	static uint32 TickCounter;
 
 	FNiagaraScriptExecutionContext();
@@ -64,11 +131,12 @@ struct FNiagaraScriptExecutionContext
 
 	bool Init(UNiagaraScript* InScript, ENiagaraSimTarget InTarget);
 	
-
 	bool Tick(class FNiagaraSystemInstance* Instance, ENiagaraSimTarget SimTarget = ENiagaraSimTarget::CPUSim);
 	void PostTick();
 
-	bool Execute(uint32 NumInstances, TArray<FNiagaraDataSetExecutionInfo, TInlineAllocator<8>>& DataSetInfos);
+	void BindData(int32 Index, FNiagaraDataSet& DataSet, int32 StartInstance, bool bUpdateInstanceCounts);
+	void BindData(int32 Index, FNiagaraDataBuffer* Input, FNiagaraDataBuffer* Output, int32 StartInstance, bool bUpdateInstanceCounts);
+	bool Execute(uint32 NumInstances);
 
 	const TArray<UNiagaraDataInterface*>& GetDataInterfaces()const { return Parameters.GetDataInterfaces(); }
 
@@ -77,215 +145,134 @@ struct FNiagaraScriptExecutionContext
 	bool CanExecute()const;
 };
 
-
-
-
 struct FNiagaraComputeExecutionContext
 {
-	FNiagaraComputeExecutionContext()
-		: MainDataSet(nullptr)
-		, SpawnRateInstances(0)
-		, EventSpawnTotal(0)
-		, SpawnScript(nullptr)
-		, UpdateScript(nullptr)
-		, GPUScript(nullptr)
-		, RTUpdateScript(0)
-		, RTSpawnScript(0)
-		, RTGPUScript(0)
-		, CBufferLayout(TEXT("Niagara Compute Sim CBuffer"))
-		, GPUDataReadback(nullptr)
-		, AccumulatedSpawnRate(0)
-		, NumIndicesPerInstance(0)
-		, PerInstanceData(nullptr)
-		, PerInstanceDataSize(0)
-		, PerInstanceDataInterfaceOffsets(nullptr)
-		, PendingExecutionQueueMask(0)
-#if WITH_EDITORONLY_DATA
-		, GPUDebugDataReadbackFloat(nullptr)
-		, GPUDebugDataReadbackInt(nullptr)
-		, GPUDebugDataReadbackCounts(nullptr)
-		, GPUDebugDataCurrBufferIdx(0xFFFFFFFF)
-		, GPUDebugDataFloatSize(0)
-		, GPUDebugDataIntSize(0)
-#endif	  
-	{
-	}
+	FNiagaraComputeExecutionContext();
+	~FNiagaraComputeExecutionContext();
 
-	~FNiagaraComputeExecutionContext()
-	{
-		checkf(IsInRenderingThread(), TEXT("Can only delete the gpu readback from the render thread"));
-		if (GPUDataReadback)
-		{
-			delete GPUDataReadback;
-			GPUDataReadback = nullptr;
-		}
+	void Reset();
 
-#if WITH_EDITORONLY_DATA
-		if (GPUDebugDataReadbackFloat)
-		{
-			delete GPUDebugDataReadbackFloat;
-			GPUDebugDataReadbackFloat = nullptr;
-		}
-		if (GPUDebugDataReadbackInt)
-		{
-			delete GPUDebugDataReadbackInt;
-			GPUDebugDataReadbackInt = nullptr;
-		}
-		if (GPUDebugDataReadbackCounts)
-		{
-			delete GPUDebugDataReadbackCounts;
-			GPUDebugDataReadbackCounts = nullptr;
-		}
-#endif
-	}
-
-	void Reset()
-	{
-		FNiagaraComputeExecutionContext* Context = this;
-		ENQUEUE_RENDER_COMMAND(ResetRT)(
-			[Context](FRHICommandListImmediate& RHICmdList)
-			{
-				Context->ResetInternal();
-			}
-		);
-	}
-
-	void InitParams(UNiagaraScript* InGPUComputeScript, UNiagaraScript* InSpawnScript, UNiagaraScript *InUpdateScript, ENiagaraSimTarget InSimTarget, const FString& InDebugSimName)
-	{
-		DebugSimName = InDebugSimName;
-		CombinedParamStore.InitFromOwningContext(InGPUComputeScript, InSimTarget, true);
-
-		GPUScript = InGPUComputeScript;
-		SpawnScript = InSpawnScript;
-		UpdateScript = InUpdateScript;
-
-#if DO_CHECK
-		FNiagaraShader *Shader = InGPUComputeScript->GetRenderThreadScript()->GetShaderGameThread();
-		DIParamInfo.Empty();
-		if (Shader)
-		{
-			for (FNiagaraDataInterfaceParamRef& DIParams : Shader->GetDIParameters())
-			{
-				DIParamInfo.Add(DIParams.ParameterInfo);
-			}
-		}
-		else
-		{
-			DIParamInfo = InGPUComputeScript->GetRenderThreadScript()->GetDataInterfaceParamInfo();
-		}
-#endif
-	}
-
-	void DirtyDataInterfaces()
-	{
-		CombinedParamStore.MarkInterfacesDirty();
-	}
-
-	bool Tick(FNiagaraSystemInstance* ParentSystemInstance)
-	{
-		if (CombinedParamStore.GetInterfacesDirty())
-		{
-#if DO_CHECK
-			const TArray<UNiagaraDataInterface*> &DataInterfaces = CombinedParamStore.GetDataInterfaces();
-			// We must make sure that the data interfaces match up between the original script values and our overrides...
-			if (DIParamInfo.Num() != DataInterfaces.Num())
-			{
-				UE_LOG(LogNiagara, Warning, TEXT("Mismatch between Niagara GPU Execution Context data interfaces and those in its script!"));
-				return false;
-			}
-
-			for (int32 i=0; i<DIParamInfo.Num(); ++i)
-			{
-				FString UsedClassName = DataInterfaces[i]->GetClass()->GetName();
-				if (DIParamInfo[i].DIClassName != UsedClassName)
-				{
-					UE_LOG(LogNiagara, Warning, TEXT("Mismatched class between Niagara GPU Execution Context data interfaces and those in its script!\nIndex:%d\nShader:%s\nScript:%s")
-						, i, *DIParamInfo[i].DIClassName, *UsedClassName);
-				}
-			}
-#endif
+	void InitParams(UNiagaraScript* InGPUComputeScript, ENiagaraSimTarget InSimTarget, const FString& InDebugSimName);
+	void DirtyDataInterfaces();
+	bool Tick(FNiagaraSystemInstance* ParentSystemInstance);
 
-			CombinedParamStore.Tick();
-		}
+	void PostTick();
 
-		return true;
-	}
+	void SetDataToRender(FNiagaraDataBuffer* InDataToRender);
+	FNiagaraDataBuffer* GetDataToRender()const { return DataToRender; }
 
 private:
-	void ResetInternal()
-	{
-		checkf(IsInRenderingThread(), TEXT("Can only reset the gpu context from the render thread"));
-		AccumulatedSpawnRate = 0;
-		PendingExecutionQueueMask = 0;
-		if (GPUDataReadback)
-		{
-			delete GPUDataReadback;
-			GPUDataReadback = nullptr;
-		}
-
-#if WITH_EDITORONLY_DATA
-		if (GPUDebugDataReadbackFloat)
-		{
-			delete GPUDebugDataReadbackFloat;
-			GPUDebugDataReadbackFloat = nullptr;
-		}
-		if (GPUDebugDataReadbackInt)
-		{
-			delete GPUDebugDataReadbackInt;
-			GPUDebugDataReadbackInt = nullptr;
-		}
-		if (GPUDebugDataReadbackCounts)
-		{
-			delete GPUDebugDataReadbackCounts;
-			GPUDebugDataReadbackCounts = nullptr;
-		}
-#endif
-	}
+	void ResetInternal();
 
 public:
-	const TArray<FNiagaraEventScriptProperties> &GetEventHandlers() const { return EventHandlerScriptProps; }
+	static uint32 TickCounter;
+
+	//Persistent state 
 	FString DebugSimName;
 	class FNiagaraDataSet *MainDataSet;
-	TArray<FNiagaraDataSet*>UpdateEventWriteDataSets;
-	TArray<FNiagaraEventScriptProperties> EventHandlerScriptProps;
-	TArray<FNiagaraDataSet*> EventSets;
-	uint32 SpawnRateInstances;
-
-	TArray<int32> EventSpawnCounts;
-	uint32 EventSpawnTotal;
-	UNiagaraScript* SpawnScript;
-	UNiagaraScript* UpdateScript;
 	UNiagaraScript* GPUScript;
-	class FNiagaraShaderScript*  RTUpdateScript;
-	class FNiagaraShaderScript*  RTSpawnScript;
-	class FNiagaraShaderScript*  RTGPUScript;	
+	class FNiagaraShaderScript*  GPUScript_RT;
 	FRHIUniformBufferLayout CBufferLayout; // Persistent layouts used to create Compute Sim CBuffer
-	TArray<uint8, TAlignedHeapAllocator<16>> ParamData_RT;		// RT side copy of the parameter data
+
+	//Dynamic state updated either from GT via RT commands or from the RT side sim code itself.
+	//TArray<uint8, TAlignedHeapAllocator<16>> ParamData_RT;		// RT side copy of the parameter data
 	FNiagaraScriptExecutionParameterStore CombinedParamStore;
-	static uint32 TickCounter;
 #if DO_CHECK
 	TArray< FNiagaraDataInterfaceGPUParamInfo >  DIParamInfo;
 #endif
 
+	TArray<FNiagaraDataInterfaceProxy*> DataInterfaceProxies;
+
+	//Most current buffer that can be used for rendering.
+	FNiagaraDataBuffer* DataToRender;
+
 	FRHIGPUMemoryReadback *GPUDataReadback;
 	uint32 AccumulatedSpawnRate;
 	uint32 NumIndicesPerInstance;	// how many vtx indices per instance the renderer is going to have for its draw call
 
-	void* PerInstanceData; // Data stored on parent system instance
-	uint32 PerInstanceDataSize; // Size of data stored on parent system instance in bytes
-	TMap<TWeakObjectPtr<UNiagaraDataInterface>, int32>* PerInstanceDataInterfaceOffsets;
-
-	/** Ensures we only enqueue each context once per queue before they're dispatched. See SIMULATION_QUEUE_COUNT */
-	uint32 PendingExecutionQueueMask;
-
+	uint32 EventSpawnTotal_GT;
+	uint32 SpawnRateInstances_GT;
 
 #if WITH_EDITORONLY_DATA
 	mutable FRHIGPUMemoryReadback *GPUDebugDataReadbackFloat;
 	mutable FRHIGPUMemoryReadback *GPUDebugDataReadbackInt;
 	mutable FRHIGPUMemoryReadback *GPUDebugDataReadbackCounts;
-	mutable int32 GPUDebugDataCurrBufferIdx;
 	mutable uint32 GPUDebugDataFloatSize;
 	mutable uint32 GPUDebugDataIntSize;
 	mutable TSharedPtr<struct FNiagaraScriptDebuggerInfo, ESPMode::ThreadSafe> DebugInfo;
 #endif
+
+};
+
+struct FNiagaraDataInterfaceInstanceData
+{
+	void* PerInstanceDataForRT;
+	TMap<FNiagaraDataInterfaceProxy*, int32> InterfaceProxiesToOffsets;
+	uint32 PerInstanceDataSize;
+	uint32 Instances;
+
+	~FNiagaraDataInterfaceInstanceData()
+	{}
+};
+
+//TODO: Rename FNiagaraGPUEmitterTick?
+
+struct FNiagaraComputeInstanceData
+{
+	uint32 EventSpawnTotal;
+	uint32 SpawnRateInstances;
+	uint8* ParamData;
+	FNiagaraComputeExecutionContext* Context;
+	TArray<FNiagaraDataInterfaceProxy*> DataInterfaceProxies;
+
+	//Buffer containing current state that this tick will read from. Initialized at the start of processing this tick on the RT.
+	FNiagaraDataBuffer* CurrentData;
+	//Buffer into which we'll write the new simulation state. Initialized at the start of processing this tick on the RT.
+	FNiagaraDataBuffer* DestinationData;
+
+	FNiagaraComputeInstanceData()
+		: EventSpawnTotal(0)
+		, SpawnRateInstances(0)
+		, ParamData(nullptr)
+		, Context(nullptr)
+		, CurrentData(nullptr)
+		, DestinationData(nullptr)
+	{}
+};
+
+/*
+	Represents all the information needed to dispatch a single tick of a FNiagaraSystemInstance.
+	This object will be created on the game thread and passed to the renderthread.
+
+	It contains the PerInstance data buffer for every DataInterface referenced by the system as well
+	as the Data required to dispatch updates for each Emitter in the system.
+
+	DataInterface data is packed tightly. It includes a TMap that associates the data interface with
+	the offset into the packed buffer. At that offset is the Per-Instance data for this System.
+
+	InstanceData_ParamData_Packed packs FNiagaraComputeInstanceData and ParamData into one buffer.
+	There is padding after the array of FNiagaraComputeInstanceData so we can upload ParamData directly into a UniformBuffer
+	(it is 16 byte aligned).
+
+*/
+class FNiagaraGPUSystemTick
+{
+public:
+	FNiagaraGPUSystemTick()
+		: Count(0)
+		, DIInstanceData(nullptr)
+		, InstanceData_ParamData_Packed(nullptr)
+	{}
+
+	void Init(FNiagaraSystemInstance* InSystemInstance);
+	void Destroy();
+
+	FORCEINLINE bool IsValid()const{ return InstanceData_ParamData_Packed != nullptr; }
+	FORCEINLINE FNiagaraComputeInstanceData* GetInstanceData()const{ return reinterpret_cast<FNiagaraComputeInstanceData*>(InstanceData_ParamData_Packed); }
+
+	uint32 Count;
+	FGuid SystemInstanceID;
+	FNiagaraDataInterfaceInstanceData* DIInstanceData;
+	uint8* InstanceData_ParamData_Packed;
+	bool bRequiredDistanceFieldData = false;
 };
diff --git a/Engine/Plugins/FX/Niagara/Source/Niagara/Classes/NiagaraScriptSourceBase.h b/Engine/Plugins/FX/Niagara/Source/Niagara/Classes/NiagaraScriptSourceBase.h
--- a/Engine/Plugins/FX/Niagara/Source/Niagara/Classes/NiagaraScriptSourceBase.h
+++ b/Engine/Plugins/FX/Niagara/Source/Niagara/Classes/NiagaraScriptSourceBase.h
@@ -6,6 +6,7 @@
 #include "UObject/ObjectMacros.h"
 #include "UObject/Object.h"
 #include "NiagaraCommon.h"
+#include "NiagaraCompileHash.h"
 #include "NiagaraScriptSourceBase.generated.h"
 
 struct EditorExposedVectorConstant
@@ -87,7 +88,11 @@ class UNiagaraScriptSourceBase : public UObject
 
 	virtual FGuid GetChangeID() { return FGuid(); };
 
-	virtual void ComputeVMCompilationId(struct FNiagaraVMExecutableDataId& Id, ENiagaraScriptUsage InUsage, const FGuid& InUsageId) const {};
+	virtual void ComputeVMCompilationId(struct FNiagaraVMExecutableDataId& Id, ENiagaraScriptUsage InUsage, const FGuid& InUsageId, bool bForceRebuild = false) const {};
+
+	virtual FGuid GetCompileBaseId(ENiagaraScriptUsage InUsage, const FGuid& InUsageId) const { return FGuid(); }
+
+	virtual FNiagaraCompileHash GetCompileHash(ENiagaraScriptUsage InUsage, const FGuid& InUsageId) const { return FNiagaraCompileHash(); }
 	
 	/** Cause the source to build up any internal variables that will be useful in the compilation process.*/
 	virtual TSharedPtr<FNiagaraCompileRequestDataBase, ESPMode::ThreadSafe> PreCompile(UNiagaraEmitter* Emitter, const TArray<FNiagaraVariable>& EncounterableVariables, TArray<TSharedPtr<FNiagaraCompileRequestDataBase, ESPMode::ThreadSafe>>& ReferencedCompileRequests, bool bClearErrors = true) { return nullptr; }
diff --git a/Engine/Plugins/FX/Niagara/Source/Niagara/Classes/NiagaraSystem.h b/Engine/Plugins/FX/Niagara/Source/Niagara/Classes/NiagaraSystem.h
--- a/Engine/Plugins/FX/Niagara/Source/Niagara/Classes/NiagaraSystem.h
+++ b/Engine/Plugins/FX/Niagara/Source/Niagara/Classes/NiagaraSystem.h
@@ -146,7 +146,7 @@ public:
 	bool ReferencesInstanceEmitter(UNiagaraEmitter& Emitter);
 
 	/** Updates all handles which use this emitter as their source. */
-	void UpdateFromEmitterChanges(UNiagaraEmitter& ChangedSourceEmitter);
+	void UpdateFromEmitterChanges(UNiagaraEmitter& ChangedSourceEmitter, bool bRecompileOnChange = true);
 
 	/** Updates the system's rapid iteration parameters from a specific emitter. */
 	void RefreshSystemParametersFromEmitter(const FNiagaraEmitterHandle& EmitterHandle);
@@ -234,6 +234,8 @@ public:
 	UPROPERTY(EditAnywhere, Category = "System", meta = (InlineEditConditionToggle))
 		uint32 bFixedBounds : 1;
 
+	TStatId GetStatID(bool bGameThread, bool bConcurrent)const;
+
 private:
 #if WITH_EDITORONLY_DATA
 	INiagaraModule::FMergeEmitterResults MergeChangesForEmitterHandle(FNiagaraEmitterHandle& EmitterHandle);
@@ -313,4 +315,12 @@ protected:
 
 	UPROPERTY()
 	TArray<FName> UserDINamesReadInSystemScripts;
+
+	void GenerateStatID();
+#if STATS
+	TStatId StatID_GT;
+	TStatId StatID_GT_CNC;
+	TStatId StatID_RT;
+	TStatId StatID_RT_CNC;
+#endif
 };
diff --git a/Engine/Plugins/FX/Niagara/Source/Niagara/Private/MovieScene/MovieSceneNiagaraSystemTrackTemplate.cpp b/Engine/Plugins/FX/Niagara/Source/Niagara/Private/MovieScene/MovieSceneNiagaraSystemTrackTemplate.cpp
--- a/Engine/Plugins/FX/Niagara/Source/Niagara/Private/MovieScene/MovieSceneNiagaraSystemTrackTemplate.cpp
+++ b/Engine/Plugins/FX/Niagara/Source/Niagara/Private/MovieScene/MovieSceneNiagaraSystemTrackTemplate.cpp
@@ -106,22 +106,22 @@ struct FNiagaraSystemUpdateDesiredAgeExecutionToken : IMovieSceneExecutionToken
 				}
 			}
 
-			if (NiagaraComponent->IsActive() == false || NiagaraComponent->GetSystemInstance() == nullptr)
-			{
-				NiagaraComponent->Activate();
-			}
+			FNiagaraSystemInstance* SystemInstance = NiagaraComponent->GetSystemInstance();
 
 			if (Context.GetTime() < SpawnSectionStartFrame)
 			{
-				FNiagaraSystemInstance* SystemInstance = NiagaraComponent->GetSystemInstance();
-				if (SystemInstance != nullptr)
+				if (NiagaraComponent->IsActive() && SystemInstance != nullptr)
 				{
 					SystemInstance->SetRequestedExecutionState(ENiagaraExecutionState::Complete);
 				}
 			}
 			else if (Context.GetTime() < SpawnSectionEndFrame)
 			{
-				FNiagaraSystemInstance* SystemInstance = NiagaraComponent->GetSystemInstance();
+				if (NiagaraComponent->IsActive() == false)
+				{
+					NiagaraComponent->Activate();
+				}
+				
 				if (SystemInstance != nullptr)
 				{
 					SystemInstance->SetRequestedExecutionState(ENiagaraExecutionState::Active);
@@ -129,7 +129,6 @@ struct FNiagaraSystemUpdateDesiredAgeExecutionToken : IMovieSceneExecutionToken
 			}
 			else
 			{
-				FNiagaraSystemInstance* SystemInstance = NiagaraComponent->GetSystemInstance();
 				if (SystemInstance != nullptr)
 				{
 					SystemInstance->SetRequestedExecutionState(ENiagaraExecutionState::Inactive);
@@ -138,7 +137,19 @@ struct FNiagaraSystemUpdateDesiredAgeExecutionToken : IMovieSceneExecutionToken
 
 			bool bRenderingEnabled = Context.IsPreRoll() == false;
 			NiagaraComponent->SetRenderingEnabled(bRenderingEnabled);
-			NiagaraComponent->SetDesiredAge(Context.GetFrameRate().AsSeconds(Context.GetTime() - SpawnSectionStartFrame));
+
+			if (SystemInstance != nullptr && SystemInstance->IsComplete() == false)
+			{
+				float DesiredAge = Context.GetFrameRate().AsSeconds(Context.GetTime() - SpawnSectionStartFrame);
+				if (DesiredAge >= 0)
+				{
+					// Add a quarter of a frame offset here to push the desired age into the middle of the frame since it will be automatically rounded
+					// down to the nearest seek delta.  This prevents a situation where float rounding results in a value which is just slightly less than
+					// the frame boundary, which results in a skipped simulation frame.
+					float FrameOffset = NiagaraComponent->GetSeekDelta() / 4;
+					NiagaraComponent->SetDesiredAge(DesiredAge + FrameOffset);
+				}
+			}
 		}
 	}
 
@@ -149,9 +160,7 @@ struct FNiagaraSystemUpdateDesiredAgeExecutionToken : IMovieSceneExecutionToken
 FMovieSceneNiagaraSystemTrackImplementation::FMovieSceneNiagaraSystemTrackImplementation(FFrameNumber InSpawnSectionStartFrame, FFrameNumber InSpawnSectionEndFrame)
 	: SpawnSectionStartFrame(InSpawnSectionStartFrame)
 	, SpawnSectionEndFrame(InSpawnSectionEndFrame)
-{
-	EnableOverrides(CustomEvaluateFlag);
-}
+{}
 
 void FMovieSceneNiagaraSystemTrackImplementation::Evaluate(const FMovieSceneEvaluationTrack& Track, FMovieSceneSegmentIdentifier SegmentID, const FMovieSceneEvaluationOperand& Operand, const FMovieSceneContext& Context, const FPersistentEvaluationData& PersistentData, FMovieSceneExecutionTokens& ExecutionTokens) const
 {
diff --git a/Engine/Plugins/FX/Niagara/Source/Niagara/Private/MovieScene/MovieSceneNiagaraSystemTrackTemplate.h b/Engine/Plugins/FX/Niagara/Source/Niagara/Private/MovieScene/MovieSceneNiagaraSystemTrackTemplate.h
--- a/Engine/Plugins/FX/Niagara/Source/Niagara/Private/MovieScene/MovieSceneNiagaraSystemTrackTemplate.h
+++ b/Engine/Plugins/FX/Niagara/Source/Niagara/Private/MovieScene/MovieSceneNiagaraSystemTrackTemplate.h
@@ -26,12 +26,19 @@ struct FMovieSceneNiagaraSystemTrackImplementation : public FMovieSceneTrackImpl
 public:
 	FMovieSceneNiagaraSystemTrackImplementation() { }
 	FMovieSceneNiagaraSystemTrackImplementation(FFrameNumber InSpawnSectionStartFrame, FFrameNumber InSpawnSectionEndFrame);
+	
+	virtual void SetupOverrides() override
+	{
+		EnableOverrides(CustomEvaluateFlag);
+	}
 
 private:
 	virtual UScriptStruct& GetScriptStructImpl() const override { return *StaticStruct(); }
 	virtual void Evaluate(const FMovieSceneEvaluationTrack& Track, FMovieSceneSegmentIdentifier SegmentID, const FMovieSceneEvaluationOperand& Operand, const FMovieSceneContext& Context, const FPersistentEvaluationData& PersistentData, FMovieSceneExecutionTokens& ExecutionTokens) const override;
 
 private:
+	UPROPERTY()
 	FFrameNumber SpawnSectionStartFrame;
+	UPROPERTY()
 	FFrameNumber SpawnSectionEndFrame;
 };
\ No newline at end of file
diff --git a/Engine/Plugins/FX/Niagara/Source/Niagara/Private/MovieScene/Parameters/MovieSceneNiagaraParameterSectionTemplate.cpp b/Engine/Plugins/FX/Niagara/Source/Niagara/Private/MovieScene/Parameters/MovieSceneNiagaraParameterSectionTemplate.cpp
--- a/Engine/Plugins/FX/Niagara/Source/Niagara/Private/MovieScene/Parameters/MovieSceneNiagaraParameterSectionTemplate.cpp
+++ b/Engine/Plugins/FX/Niagara/Source/Niagara/Private/MovieScene/Parameters/MovieSceneNiagaraParameterSectionTemplate.cpp
@@ -98,9 +98,7 @@ FMovieSceneNiagaraParameterSectionTemplate::FMovieSceneNiagaraParameterSectionTe
 
 FMovieSceneNiagaraParameterSectionTemplate::FMovieSceneNiagaraParameterSectionTemplate(FNiagaraVariable InParameter)
 	: Parameter(InParameter)
-{
-	EnableOverrides(RequiresInitializeFlag);
-}
+{}
 
 void FMovieSceneNiagaraParameterSectionTemplate::Initialize(const FMovieSceneEvaluationOperand& Operand, const FMovieSceneContext& Context, FPersistentEvaluationData& PersistentData, IMovieScenePlayer& Player) const
 {
diff --git a/Engine/Plugins/FX/Niagara/Source/Niagara/Private/MovieScene/Parameters/MovieSceneNiagaraParameterSectionTemplate.h b/Engine/Plugins/FX/Niagara/Source/Niagara/Private/MovieScene/Parameters/MovieSceneNiagaraParameterSectionTemplate.h
--- a/Engine/Plugins/FX/Niagara/Source/Niagara/Private/MovieScene/Parameters/MovieSceneNiagaraParameterSectionTemplate.h
+++ b/Engine/Plugins/FX/Niagara/Source/Niagara/Private/MovieScene/Parameters/MovieSceneNiagaraParameterSectionTemplate.h
@@ -16,6 +16,11 @@ public:
 
 	FMovieSceneNiagaraParameterSectionTemplate(FNiagaraVariable InParameter);
 
+	virtual void SetupOverrides() override
+	{
+		EnableOverrides(RequiresInitializeFlag);
+	}
+
 private:
 	virtual UScriptStruct& GetScriptStructImpl() const override { return *StaticStruct(); }
 	virtual void Initialize(const FMovieSceneEvaluationOperand& Operand, const FMovieSceneContext& Context, FPersistentEvaluationData& PersistentData, IMovieScenePlayer& Player) const override;
diff --git a/Engine/Plugins/FX/Niagara/Source/Niagara/Private/NDISkeletalMesh_BoneSampling.cpp b/Engine/Plugins/FX/Niagara/Source/Niagara/Private/NDISkeletalMesh_BoneSampling.cpp
--- a/Engine/Plugins/FX/Niagara/Source/Niagara/Private/NDISkeletalMesh_BoneSampling.cpp
+++ b/Engine/Plugins/FX/Niagara/Source/Niagara/Private/NDISkeletalMesh_BoneSampling.cpp
@@ -24,6 +24,10 @@ static const FName RandomSpecificBoneName("RandomSpecificBone");
 static const FName IsValidBoneName("IsValidBoneName");
 static const FName GetSkinnedBoneDataName("GetSkinnedBoneData");
 static const FName GetSkinnedBoneDataWSName("GetSkinnedBoneDataWS");
+
+static const FName GetSkinnedBoneDataInterpolatedName("GetSkinnedBoneDataInterpolated");
+static const FName GetSkinnedBoneDataWSInterpolatedName("GetSkinnedBoneDataWSInterpolated");
+
 static const FName GetSpecificBoneCountName("GetSpecificBoneCount");
 static const FName GetSpecificBoneAtName("GetSpecificBone");
 
@@ -90,6 +94,39 @@ void UNiagaraDataInterfaceSkeletalMesh::GetSkeletonSamplingFunctions(TArray<FNia
 		OutFunctions.Add(Sig);
 	}
 
+
+	{
+		FNiagaraFunctionSignature Sig;
+		Sig.Name = GetSkinnedBoneDataInterpolatedName;
+		Sig.Inputs.Add(FNiagaraVariable(FNiagaraTypeDefinition(GetClass()), TEXT("SkeletalMesh")));
+		Sig.Inputs.Add(FNiagaraVariable(FNiagaraTypeDefinition::GetIntDef(), TEXT("Bone")));
+		Sig.Inputs.Add(FNiagaraVariable(FNiagaraTypeDefinition::GetFloatDef(), TEXT("Interpolation")));
+		Sig.Outputs.Add(FNiagaraVariable(FNiagaraTypeDefinition::GetVec3Def(), TEXT("Position")));
+		Sig.Outputs.Add(FNiagaraVariable(FNiagaraTypeDefinition::GetVec3Def(), TEXT("Velocity")));
+		Sig.bMemberFunction = true;
+		Sig.bRequiresContext = false;
+#if WITH_EDITORONLY_DATA
+		Sig.Description = LOCTEXT("GetSkinnedBoneDataDesc", "Returns skinning dependant data for the pased bone in local space. Interpolated between this frame and the previous based on passed interpolation factor. All outputs are optional and you will incur zero to minimal cost if they are not connected.");
+#endif
+		OutFunctions.Add(Sig);
+	}
+
+	{
+		FNiagaraFunctionSignature Sig;
+		Sig.Name = GetSkinnedBoneDataWSInterpolatedName;
+		Sig.Inputs.Add(FNiagaraVariable(FNiagaraTypeDefinition(GetClass()), TEXT("SkeletalMesh")));
+		Sig.Inputs.Add(FNiagaraVariable(FNiagaraTypeDefinition::GetIntDef(), TEXT("Bone")));
+		Sig.Inputs.Add(FNiagaraVariable(FNiagaraTypeDefinition::GetFloatDef(), TEXT("Interpolation")));
+		Sig.Outputs.Add(FNiagaraVariable(FNiagaraTypeDefinition::GetVec3Def(), TEXT("Position")));
+		Sig.Outputs.Add(FNiagaraVariable(FNiagaraTypeDefinition::GetVec3Def(), TEXT("Velocity")));
+		Sig.bMemberFunction = true;
+		Sig.bRequiresContext = false;
+#if WITH_EDITORONLY_DATA
+		Sig.Description = LOCTEXT("GetSkinnedBoneDataWSDesc", "Returns skinning dependant data for the pased bone in world space. Interpolated between this frame and the previous based on passed interpolation factor. All outputs are optional and you will incur zero to minimal cost if they are not connected.");
+#endif
+		OutFunctions.Add(Sig);
+	}
+
 	{
 		FNiagaraFunctionSignature Sig;
 		Sig.Name = GetSpecificBoneCountName;
@@ -178,12 +215,22 @@ void UNiagaraDataInterfaceSkeletalMesh::BindSkeletonSamplingFunction(const FVMEx
 	else if (BindingInfo.Name == GetSkinnedBoneDataName)
 	{
 		check(BindingInfo.GetNumInputs() == 2 && BindingInfo.GetNumOutputs() == 6);
-		TSkinningModeBinder<TNDIExplicitBinder<FNDITransformHandlerNoop, NDI_FUNC_BINDER(UNiagaraDataInterfaceSkeletalMesh, GetSkinnedBoneData)>>::Bind(this, BindingInfo, InstanceData, OutFunc);
+		TSkinningModeBinder<TNDIExplicitBinder<FNDITransformHandlerNoop, TNDIExplicitBinder<TIntegralConstant<bool, false>, NDI_FUNC_BINDER(UNiagaraDataInterfaceSkeletalMesh, GetSkinnedBoneData)>>>::Bind(this, BindingInfo, InstanceData, OutFunc);
 	}
 	else if (BindingInfo.Name == GetSkinnedBoneDataWSName)
 	{
 		check(BindingInfo.GetNumInputs() == 2 && BindingInfo.GetNumOutputs() == 6);
-		TSkinningModeBinder<TNDIExplicitBinder<FNDITransformHandler, NDI_FUNC_BINDER(UNiagaraDataInterfaceSkeletalMesh, GetSkinnedBoneData)>>::Bind(this, BindingInfo, InstanceData, OutFunc);
+		TSkinningModeBinder<TNDIExplicitBinder<FNDITransformHandler, TNDIExplicitBinder<TIntegralConstant<bool, false>, NDI_FUNC_BINDER(UNiagaraDataInterfaceSkeletalMesh, GetSkinnedBoneData)>>>::Bind(this, BindingInfo, InstanceData, OutFunc);
+	}
+	else if (BindingInfo.Name == GetSkinnedBoneDataInterpolatedName)
+	{
+		check(BindingInfo.GetNumInputs() == 3 && BindingInfo.GetNumOutputs() == 6);
+		TSkinningModeBinder<TNDIExplicitBinder<FNDITransformHandlerNoop, TNDIExplicitBinder<TIntegralConstant<bool, true>, NDI_FUNC_BINDER(UNiagaraDataInterfaceSkeletalMesh, GetSkinnedBoneData)>>>::Bind(this, BindingInfo, InstanceData, OutFunc);
+	}
+	else if (BindingInfo.Name == GetSkinnedBoneDataWSInterpolatedName)
+	{
+		check(BindingInfo.GetNumInputs() == 3 && BindingInfo.GetNumOutputs() == 6);
+		TSkinningModeBinder<TNDIExplicitBinder<FNDITransformHandler, TNDIExplicitBinder<TIntegralConstant<bool, true>, NDI_FUNC_BINDER(UNiagaraDataInterfaceSkeletalMesh, GetSkinnedBoneData)>>>::Bind(this, BindingInfo, InstanceData, OutFunc);
 	}
 	else if (BindingInfo.Name == GetSpecificBoneCountName)
 	{
@@ -342,13 +389,19 @@ struct FBoneSocketSkinnedDataOutputHandler
 	}
 };
 
-template<typename SkinningHandlerType, typename TransformHandlerType>
+template<typename SkinningHandlerType, typename TransformHandlerType, typename bInterpolated>
 void UNiagaraDataInterfaceSkeletalMesh::GetSkinnedBoneData(FVectorVMContext& Context)
 {
 	SCOPE_CYCLE_COUNTER(STAT_NiagaraSkel_Bone_Sample);
 	SkinningHandlerType SkinningHandler;
 	TransformHandlerType TransformHandler;
 	VectorVM::FExternalFuncInputHandler<int32> BoneParam(Context);
+	VectorVM::FExternalFuncInputHandler<float> InterpParam;
+
+	if (bInterpolated::Value)
+	{
+		InterpParam.Init(Context);
+	}
 
 	VectorVM::FUserPtrHandler<FNDISkeletalMesh_InstanceData> InstData(Context);
 
@@ -384,18 +437,35 @@ void UNiagaraDataInterfaceSkeletalMesh::GetSkinnedBoneData(FVectorVMContext& Con
 	{
 		int32 Bone = FMath::Clamp(BoneParam.GetAndAdvance(), 0, BoneMax);
 
-		//No parent bone, just spawn at bone.
-		if (Output.bNeedsPosition || Output.bNeedsVelocity)
+		float Interp = 1.0f;
+		if (bInterpolated::Value)
 		{
-			Pos = SkinningHandler.GetSkinnedBonePosition(Accessor, Bone);
-			TransformHandler.TransformPosition(Pos, Transform);
-			Output.SetPosition(Pos);
+			Interp = InterpParam.GetAndAdvance();
 		}
 
-		if (Output.bNeedsVelocity)
+
+		Pos = SkinningHandler.GetSkinnedBonePosition(Accessor, Bone);
+		TransformHandler.TransformPosition(Pos, Transform);
+
+		if (Output.bNeedsVelocity || bInterpolated::Value)
 		{
 			Prev = SkinningHandler.GetSkinnedBonePreviousPosition(Accessor, Bone);
 			TransformHandler.TransformPosition(Prev, PrevTransform);
+		}
+
+		if (Output.bNeedsPosition)
+		{
+			if (bInterpolated::Value)
+			{
+				Pos = FMath::Lerp(Prev, Pos, Interp);
+			}
+
+			Output.SetPosition(Pos);
+		}
+
+		if(Output.bNeedsVelocity)
+		{
+			//Don't have enough information to get a better interpolated velocity.
 			Velocity = (Pos - Prev) * InvDt;
 			Output.SetVelocity(Velocity);
 		}
diff --git a/Engine/Plugins/FX/Niagara/Source/Niagara/Private/NDISkeletalMesh_TriangleSampling.cpp b/Engine/Plugins/FX/Niagara/Source/Niagara/Private/NDISkeletalMesh_TriangleSampling.cpp
--- a/Engine/Plugins/FX/Niagara/Source/Niagara/Private/NDISkeletalMesh_TriangleSampling.cpp
+++ b/Engine/Plugins/FX/Niagara/Source/Niagara/Private/NDISkeletalMesh_TriangleSampling.cpp
@@ -24,24 +24,24 @@ DEFINE_NDI_FUNC_BINDER(UNiagaraDataInterfaceSkeletalMesh, GetFilteredTriangleCou
 DEFINE_NDI_FUNC_BINDER(UNiagaraDataInterfaceSkeletalMesh, GetFilteredTriangleAt);
 DEFINE_NDI_FUNC_BINDER(UNiagaraDataInterfaceSkeletalMesh, GetTriCoordVertices);
 
-static const FName RandomTriCoordName("RandomTriCoord");
-static const FName IsValidTriCoordName("IsValidTriCoord");
-static const FName GetSkinnedTriangleDataName("GetSkinnedTriangleData");
-static const FName GetSkinnedTriangleDataWSName("GetSkinnedTriangleDataWS");
-static const FName GetSkinnedTriangleDataInterpName("GetSkinnedTriangleDataInterpolated");
-static const FName GetSkinnedTriangleDataWSInterpName("GetSkinnedTriangleDataWSInterpolated");
-static const FName GetTriColorName("GetTriColor");
-static const FName GetTriUVName("GetTriUV");
-static const FName GetTriangleCountName("GetFilteredTriangleCount");
-static const FName GetTriangleAtName("GetFilteredTriangle");
-static const FName GetTriCoordVerticesName("GetTriCoordVertices");
+const FName FSkeletalMeshInterfaceHelper::RandomTriCoordName("RandomTriCoord");
+const FName FSkeletalMeshInterfaceHelper::IsValidTriCoordName("IsValidTriCoord");
+const FName FSkeletalMeshInterfaceHelper::GetSkinnedTriangleDataName("GetSkinnedTriangleData");
+const FName FSkeletalMeshInterfaceHelper::GetSkinnedTriangleDataWSName("GetSkinnedTriangleDataWS");
+const FName FSkeletalMeshInterfaceHelper::GetSkinnedTriangleDataInterpName("GetSkinnedTriangleDataInterpolated");
+const FName FSkeletalMeshInterfaceHelper::GetSkinnedTriangleDataWSInterpName("GetSkinnedTriangleDataWSInterpolated");
+const FName FSkeletalMeshInterfaceHelper::GetTriColorName("GetTriColor");
+const FName FSkeletalMeshInterfaceHelper::GetTriUVName("GetTriUV");
+const FName FSkeletalMeshInterfaceHelper::GetTriangleCountName("GetFilteredTriangleCount");
+const FName FSkeletalMeshInterfaceHelper::GetTriangleAtName("GetFilteredTriangle");
+const FName FSkeletalMeshInterfaceHelper::GetTriCoordVerticesName("GetTriCoordVertices");
 
 
 void UNiagaraDataInterfaceSkeletalMesh::GetTriangleSamplingFunctions(TArray<FNiagaraFunctionSignature>& OutFunctions)
 {
 	{
 		FNiagaraFunctionSignature Sig;
-		Sig.Name = RandomTriCoordName;
+		Sig.Name = FSkeletalMeshInterfaceHelper::RandomTriCoordName;
 		Sig.Inputs.Add(FNiagaraVariable(FNiagaraTypeDefinition(GetClass()), TEXT("SkeletalMesh")));
 		Sig.Outputs.Add(FNiagaraVariable(FNiagaraTypeDefinition(FMeshTriCoordinate::StaticStruct()), TEXT("Coord")));
 		Sig.bMemberFunction = true;
@@ -51,7 +51,7 @@ void UNiagaraDataInterfaceSkeletalMesh::GetTriangleSamplingFunctions(TArray<FNia
 
 	{
 		FNiagaraFunctionSignature Sig;
-		Sig.Name = IsValidTriCoordName;
+		Sig.Name = FSkeletalMeshInterfaceHelper::IsValidTriCoordName;
 		Sig.Inputs.Add(FNiagaraVariable(FNiagaraTypeDefinition(GetClass()), TEXT("SkeletalMesh")));
 		Sig.Inputs.Add(FNiagaraVariable(FNiagaraTypeDefinition(FMeshTriCoordinate::StaticStruct()), TEXT("Coord")));
 		Sig.Outputs.Add(FNiagaraVariable(FNiagaraTypeDefinition::GetBoolDef(), TEXT("IsValid")));
@@ -65,7 +65,7 @@ void UNiagaraDataInterfaceSkeletalMesh::GetTriangleSamplingFunctions(TArray<FNia
 
 	{
 		FNiagaraFunctionSignature Sig;
-		Sig.Name = GetSkinnedTriangleDataName;
+		Sig.Name = FSkeletalMeshInterfaceHelper::GetSkinnedTriangleDataName;
 		Sig.Inputs.Add(FNiagaraVariable(FNiagaraTypeDefinition(GetClass()), TEXT("SkeletalMesh")));
 		Sig.Inputs.Add(FNiagaraVariable(FNiagaraTypeDefinition(FMeshTriCoordinate::StaticStruct()), TEXT("Coord")));
 		Sig.Outputs.Add(FNiagaraVariable(FNiagaraTypeDefinition::GetVec3Def(), TEXT("Position")));
@@ -83,7 +83,7 @@ void UNiagaraDataInterfaceSkeletalMesh::GetTriangleSamplingFunctions(TArray<FNia
 
 	{
 		FNiagaraFunctionSignature Sig;
-		Sig.Name = GetSkinnedTriangleDataWSName;
+		Sig.Name = FSkeletalMeshInterfaceHelper::GetSkinnedTriangleDataWSName;
 		Sig.Inputs.Add(FNiagaraVariable(FNiagaraTypeDefinition(GetClass()), TEXT("SkeletalMesh")));
 		Sig.Inputs.Add(FNiagaraVariable(FNiagaraTypeDefinition(FMeshTriCoordinate::StaticStruct()), TEXT("Coord")));
 		Sig.Outputs.Add(FNiagaraVariable(FNiagaraTypeDefinition::GetVec3Def(), TEXT("Position")));
@@ -101,7 +101,7 @@ void UNiagaraDataInterfaceSkeletalMesh::GetTriangleSamplingFunctions(TArray<FNia
 
 	{
 		FNiagaraFunctionSignature Sig;
-		Sig.Name = GetSkinnedTriangleDataInterpName;
+		Sig.Name = FSkeletalMeshInterfaceHelper::GetSkinnedTriangleDataInterpName;
 		Sig.Inputs.Add(FNiagaraVariable(FNiagaraTypeDefinition(GetClass()), TEXT("SkeletalMesh")));
 		Sig.Inputs.Add(FNiagaraVariable(FNiagaraTypeDefinition(FMeshTriCoordinate::StaticStruct()), TEXT("Coord")));
 		Sig.Inputs.Add(FNiagaraVariable(FNiagaraTypeDefinition::GetFloatDef(), TEXT("Interp")));
@@ -120,7 +120,7 @@ void UNiagaraDataInterfaceSkeletalMesh::GetTriangleSamplingFunctions(TArray<FNia
 
 	{
 		FNiagaraFunctionSignature Sig;
-		Sig.Name = GetSkinnedTriangleDataWSInterpName;
+		Sig.Name = FSkeletalMeshInterfaceHelper::GetSkinnedTriangleDataWSInterpName;
 		Sig.Inputs.Add(FNiagaraVariable(FNiagaraTypeDefinition(GetClass()), TEXT("SkeletalMesh")));
 		Sig.Inputs.Add(FNiagaraVariable(FNiagaraTypeDefinition(FMeshTriCoordinate::StaticStruct()), TEXT("Coord")));
 		Sig.Inputs.Add(FNiagaraVariable(FNiagaraTypeDefinition::GetFloatDef(), TEXT("Interp")));
@@ -139,7 +139,7 @@ void UNiagaraDataInterfaceSkeletalMesh::GetTriangleSamplingFunctions(TArray<FNia
 
 	{
 		FNiagaraFunctionSignature Sig;
-		Sig.Name = GetTriColorName;
+		Sig.Name = FSkeletalMeshInterfaceHelper::GetTriColorName;
 		Sig.Inputs.Add(FNiagaraVariable(FNiagaraTypeDefinition(GetClass()), TEXT("SkeletalMesh")));
 		Sig.Inputs.Add(FNiagaraVariable(FNiagaraTypeDefinition(FMeshTriCoordinate::StaticStruct()), TEXT("Coord")));
 		Sig.Outputs.Add(FNiagaraVariable(FNiagaraTypeDefinition::GetColorDef(), TEXT("Color")));
@@ -150,7 +150,7 @@ void UNiagaraDataInterfaceSkeletalMesh::GetTriangleSamplingFunctions(TArray<FNia
 
 	{
 		FNiagaraFunctionSignature Sig;
-		Sig.Name = GetTriUVName;
+		Sig.Name = FSkeletalMeshInterfaceHelper::GetTriUVName;
 		Sig.Inputs.Add(FNiagaraVariable(FNiagaraTypeDefinition(GetClass()), TEXT("SkeletalMesh")));
 		Sig.Inputs.Add(FNiagaraVariable(FNiagaraTypeDefinition(FMeshTriCoordinate::StaticStruct()), TEXT("Coord")));
 		Sig.Inputs.Add(FNiagaraVariable(FNiagaraTypeDefinition::GetIntDef(), TEXT("UV Set")));
@@ -163,7 +163,7 @@ void UNiagaraDataInterfaceSkeletalMesh::GetTriangleSamplingFunctions(TArray<FNia
 
 	{
 		FNiagaraFunctionSignature Sig;
-		Sig.Name = GetTriangleCountName;
+		Sig.Name = FSkeletalMeshInterfaceHelper::GetTriangleCountName;
 		Sig.Inputs.Add(FNiagaraVariable(FNiagaraTypeDefinition(GetClass()), TEXT("SkeletalMesh")));
 		Sig.Outputs.Add(FNiagaraVariable(FNiagaraTypeDefinition::GetIntDef(), TEXT("Count")));
 		Sig.bMemberFunction = true;
@@ -173,7 +173,7 @@ void UNiagaraDataInterfaceSkeletalMesh::GetTriangleSamplingFunctions(TArray<FNia
 
 	{
 		FNiagaraFunctionSignature Sig;
-		Sig.Name = GetTriangleAtName;
+		Sig.Name = FSkeletalMeshInterfaceHelper::GetTriangleAtName;
 		Sig.Inputs.Add(FNiagaraVariable(FNiagaraTypeDefinition(GetClass()), TEXT("SkeletalMesh")));
 		Sig.Inputs.Add(FNiagaraVariable(FNiagaraTypeDefinition::GetIntDef(), TEXT("Index")));
 		Sig.Outputs.Add(FNiagaraVariable(FNiagaraTypeDefinition(FMeshTriCoordinate::StaticStruct()), TEXT("Coord")));
@@ -184,7 +184,7 @@ void UNiagaraDataInterfaceSkeletalMesh::GetTriangleSamplingFunctions(TArray<FNia
 
 	{
 		FNiagaraFunctionSignature Sig;
-		Sig.Name = GetTriCoordVerticesName;
+		Sig.Name = FSkeletalMeshInterfaceHelper::GetTriCoordVerticesName;
 		Sig.Inputs.Add(FNiagaraVariable(FNiagaraTypeDefinition(GetClass()), TEXT("SkeletalMesh")));
 		Sig.Inputs.Add(FNiagaraVariable(FNiagaraTypeDefinition::GetIntDef(), TEXT("TriangleIndex")));
 		Sig.Outputs.Add(FNiagaraVariable(FNiagaraTypeDefinition::GetIntDef(), TEXT("Vertex 0")));
@@ -202,37 +202,37 @@ void UNiagaraDataInterfaceSkeletalMesh::GetTriangleSamplingFunctions(TArray<FNia
 void UNiagaraDataInterfaceSkeletalMesh::BindTriangleSamplingFunction(const FVMExternalFunctionBindingInfo& BindingInfo, FNDISkeletalMesh_InstanceData* InstanceData, FVMExternalFunction &OutFunc)
 {
 
-	if (BindingInfo.Name == RandomTriCoordName)
+	if (BindingInfo.Name == FSkeletalMeshInterfaceHelper::RandomTriCoordName)
 	{
 		check(BindingInfo.GetNumInputs() == 1 && BindingInfo.GetNumOutputs() == 4);
 		TFilterModeBinder<TAreaWeightingModeBinder<NDI_FUNC_BINDER(UNiagaraDataInterfaceSkeletalMesh, RandomTriCoord)>>::Bind(this, BindingInfo, InstanceData, OutFunc);
 	}
-	else if (BindingInfo.Name == IsValidTriCoordName)
+	else if (BindingInfo.Name == FSkeletalMeshInterfaceHelper::IsValidTriCoordName)
 	{
 		check(BindingInfo.GetNumInputs() == 5 && BindingInfo.GetNumOutputs() == 1);
 		TFilterModeBinder<TAreaWeightingModeBinder<NDI_FUNC_BINDER(UNiagaraDataInterfaceSkeletalMesh, IsValidTriCoord)>>::Bind(this, BindingInfo, InstanceData, OutFunc);
 	}
-	else if (BindingInfo.Name == GetSkinnedTriangleDataName)
+	else if (BindingInfo.Name == FSkeletalMeshInterfaceHelper::GetSkinnedTriangleDataName)
 	{
 		check(BindingInfo.GetNumInputs() == 5 && BindingInfo.GetNumOutputs() == 15);
 		TSkinningModeBinder<TNDIExplicitBinder<FNDITransformHandlerNoop, TVertexAccessorBinder<TNDIExplicitBinder<TIntegralConstant<bool, false>, NDI_FUNC_BINDER(UNiagaraDataInterfaceSkeletalMesh, GetTriCoordSkinnedData)>>>>::Bind(this, BindingInfo, InstanceData, OutFunc);
 	}
-	else if (BindingInfo.Name == GetSkinnedTriangleDataWSName)
+	else if (BindingInfo.Name == FSkeletalMeshInterfaceHelper::GetSkinnedTriangleDataWSName)
 	{
 		check(BindingInfo.GetNumInputs() == 5 && BindingInfo.GetNumOutputs() == 15);
 		TSkinningModeBinder<TNDIExplicitBinder<FNDITransformHandler, TVertexAccessorBinder<TNDIExplicitBinder<TIntegralConstant<bool, false>, NDI_FUNC_BINDER(UNiagaraDataInterfaceSkeletalMesh, GetTriCoordSkinnedData)>>>>::Bind(this, BindingInfo, InstanceData, OutFunc);
 	}
-	else if (BindingInfo.Name == GetSkinnedTriangleDataInterpName)
+	else if (BindingInfo.Name == FSkeletalMeshInterfaceHelper::GetSkinnedTriangleDataInterpName)
 	{
 		check(BindingInfo.GetNumInputs() == 6 && BindingInfo.GetNumOutputs() == 15);
 		TSkinningModeBinder<TNDIExplicitBinder<FNDITransformHandlerNoop, TVertexAccessorBinder<TNDIExplicitBinder<TIntegralConstant<bool, true>, NDI_FUNC_BINDER(UNiagaraDataInterfaceSkeletalMesh, GetTriCoordSkinnedData)>>>>::Bind(this, BindingInfo, InstanceData, OutFunc);
 	}
-	else if (BindingInfo.Name == GetSkinnedTriangleDataWSInterpName)
+	else if (BindingInfo.Name == FSkeletalMeshInterfaceHelper::GetSkinnedTriangleDataWSInterpName)
 	{
 		check(BindingInfo.GetNumInputs() == 6 && BindingInfo.GetNumOutputs() == 15);
 		TSkinningModeBinder<TNDIExplicitBinder< FNDITransformHandler, TVertexAccessorBinder<TNDIExplicitBinder<TIntegralConstant<bool, true>,NDI_FUNC_BINDER(UNiagaraDataInterfaceSkeletalMesh, GetTriCoordSkinnedData)>>>>::Bind(this, BindingInfo, InstanceData, OutFunc);
 	}
-	else if (BindingInfo.Name == GetTriColorName)
+	else if (BindingInfo.Name == FSkeletalMeshInterfaceHelper::GetTriColorName)
 	{
 		check(BindingInfo.GetNumInputs() == 5 && BindingInfo.GetNumOutputs() == 4);
 		if (InstanceData->HasColorData())
@@ -244,22 +244,22 @@ void UNiagaraDataInterfaceSkeletalMesh::BindTriangleSamplingFunction(const FVMEx
 			NDI_FUNC_BINDER(UNiagaraDataInterfaceSkeletalMesh, GetTriCoordColorFallback)::Bind(this, OutFunc);
 		}
 	}
-	else if (BindingInfo.Name == GetTriUVName)
+	else if (BindingInfo.Name == FSkeletalMeshInterfaceHelper::GetTriUVName)
 	{
 		check(BindingInfo.GetNumInputs() == 6 && BindingInfo.GetNumOutputs() == 2);
 		TVertexAccessorBinder<NDI_FUNC_BINDER(UNiagaraDataInterfaceSkeletalMesh, GetTriCoordUV)>::Bind(this, BindingInfo, InstanceData, OutFunc);
 	}
-	else if (BindingInfo.Name == GetTriangleCountName)
+	else if (BindingInfo.Name == FSkeletalMeshInterfaceHelper::GetTriangleCountName)
 	{
 		check(BindingInfo.GetNumInputs() == 1 && BindingInfo.GetNumOutputs() == 1);
 		TFilterModeBinder<TAreaWeightingModeBinder<NDI_FUNC_BINDER(UNiagaraDataInterfaceSkeletalMesh, GetFilteredTriangleCount)>>::Bind(this, BindingInfo, InstanceData, OutFunc);
 	}
-	else if (BindingInfo.Name == GetTriangleAtName)
+	else if (BindingInfo.Name == FSkeletalMeshInterfaceHelper::GetTriangleAtName)
 	{
 		check(BindingInfo.GetNumInputs() == 2 && BindingInfo.GetNumOutputs() == 4);
 		TFilterModeBinder<TAreaWeightingModeBinder<NDI_FUNC_BINDER(UNiagaraDataInterfaceSkeletalMesh, GetFilteredTriangleAt)>>::Bind(this, BindingInfo, InstanceData, OutFunc);
 	}
-	else if (BindingInfo.Name == GetTriCoordVerticesName)
+	else if (BindingInfo.Name == FSkeletalMeshInterfaceHelper::GetTriCoordVerticesName)
 	{
 		check(BindingInfo.GetNumInputs() == 5 && BindingInfo.GetNumOutputs() == 3);
 		TSkinningModeBinder<NDI_FUNC_BINDER(UNiagaraDataInterfaceSkeletalMesh, GetTriCoordVertices)>::Bind(this, BindingInfo, InstanceData, OutFunc);
@@ -283,7 +283,7 @@ FORCEINLINE int32 UNiagaraDataInterfaceSkeletalMesh::RandomTriIndex<
 	int32 SecIdx = RandStream.RandRange(0, Accessor.LODData->RenderSections.Num() - 1);
 	FSkelMeshRenderSection& Sec = Accessor.LODData->RenderSections[SecIdx];
 	int32 Tri = RandStream.RandRange(0, Sec.NumTriangles - 1);
-	return Sec.BaseIndex + Tri * 3;
+	return (Sec.BaseIndex / 3) + Tri;
 }
 
 template<>
@@ -295,7 +295,7 @@ FORCEINLINE int32 UNiagaraDataInterfaceSkeletalMesh::RandomTriIndex<
 	const FSkeletalMeshSamplingInfo& SamplingInfo = InstData->Mesh->GetSamplingInfo();
 	const FSkeletalMeshSamplingLODBuiltData& WholeMeshBuiltData = SamplingInfo.GetWholeMeshLODBuiltData(InstData->GetLODIndex());
 	int32 TriIdx = WholeMeshBuiltData.AreaWeightedTriangleSampler.GetEntryIndex(RandStream.GetFraction(), RandStream.GetFraction());
-	return TriIdx * 3;
+	return TriIdx;
 }
 
 template<>
@@ -305,7 +305,7 @@ FORCEINLINE int32 UNiagaraDataInterfaceSkeletalMesh::RandomTriIndex<
 	(FRandomStream& RandStream, FSkeletalMeshAccessorHelper& Accessor, FNDISkeletalMesh_InstanceData* InstData)
 {
 	int32 Idx = RandStream.RandRange(0, Accessor.SamplingRegionBuiltData->TriangleIndices.Num() - 1);
-	return Accessor.SamplingRegionBuiltData->TriangleIndices[Idx];
+	return Accessor.SamplingRegionBuiltData->TriangleIndices[Idx] / 3;
 }
 
 template<>
@@ -315,7 +315,7 @@ FORCEINLINE int32 UNiagaraDataInterfaceSkeletalMesh::RandomTriIndex<
 	(FRandomStream& RandStream, FSkeletalMeshAccessorHelper& Accessor, FNDISkeletalMesh_InstanceData* InstData)
 {
 	int32 Idx = Accessor.SamplingRegionBuiltData->AreaWeightedSampler.GetEntryIndex(RandStream.GetFraction(), RandStream.GetFraction());
-	return Accessor.SamplingRegionBuiltData->TriangleIndices[Idx];
+	return Accessor.SamplingRegionBuiltData->TriangleIndices[Idx] / 3;
 }
 
 template<>
@@ -329,7 +329,7 @@ FORCEINLINE int32 UNiagaraDataInterfaceSkeletalMesh::RandomTriIndex<
 	const FSkeletalMeshSamplingRegion& Region = SamplingInfo.GetRegion(InstData->SamplingRegionIndices[RegionIdx]);
 	const FSkeletalMeshSamplingRegionBuiltData& RegionBuiltData = SamplingInfo.GetRegionBuiltData(InstData->SamplingRegionIndices[RegionIdx]);
 	int32 Idx = RandStream.RandRange(0, RegionBuiltData.TriangleIndices.Num() - 1);
-	return RegionBuiltData.TriangleIndices[Idx];
+	return RegionBuiltData.TriangleIndices[Idx] / 3;
 }
 
 template<>
@@ -343,7 +343,7 @@ FORCEINLINE int32 UNiagaraDataInterfaceSkeletalMesh::RandomTriIndex<
 	const FSkeletalMeshSamplingRegion& Region = SamplingInfo.GetRegion(InstData->SamplingRegionIndices[RegionIdx]);
 	const FSkeletalMeshSamplingRegionBuiltData& RegionBuiltData = SamplingInfo.GetRegionBuiltData(InstData->SamplingRegionIndices[RegionIdx]);
 	int32 Idx = RegionBuiltData.AreaWeightedSampler.GetEntryIndex(RandStream.GetFraction(), RandStream.GetFraction());
-	return RegionBuiltData.TriangleIndices[Idx];
+	return RegionBuiltData.TriangleIndices[Idx] / 3;
 }
 
 template<typename FilterMode, typename AreaWeightingMode>
@@ -392,7 +392,7 @@ void UNiagaraDataInterfaceSkeletalMesh::IsValidTriCoord(FVectorVMContext& Contex
 
 	for (int32 i = 0; i < Context.NumInstances; ++i)
 	{
-		int32 RequestedIndex = TriParam.Get() + 2; // Get the last triangle index of the set
+		int32 RequestedIndex = (TriParam.Get() * 3) + 2; // Get the last triangle index of the set
 
 		FNiagaraBool Value;
 		Value.SetValue(MeshAccessor.IndexBuffer != nullptr && MeshAccessor.IndexBuffer->Num() > RequestedIndex);
@@ -533,7 +533,7 @@ FORCEINLINE int32 UNiagaraDataInterfaceSkeletalMesh::GetSpecificTriangleAt<
 		if (Accessor.LODData->RenderSections[i].NumTriangles > (uint32)FilteredIndex)
 		{
 			FSkelMeshRenderSection& Sec = Accessor.LODData->RenderSections[i];
-			return Sec.BaseIndex + FilteredIndex * 3;
+			return Sec.BaseIndex + FilteredIndex;
 		}
 		FilteredIndex -= Accessor.LODData->RenderSections[i].NumTriangles;
 	}
@@ -547,7 +547,7 @@ FORCEINLINE int32 UNiagaraDataInterfaceSkeletalMesh::GetSpecificTriangleAt<
 	(FSkeletalMeshAccessorHelper& Accessor, FNDISkeletalMesh_InstanceData* InstData, int32 FilteredIndex)
 {
 	int32 TriIdx = FilteredIndex;
-	return TriIdx * 3;
+	return TriIdx;
 }
 
 template<>
@@ -558,7 +558,7 @@ FORCEINLINE int32 UNiagaraDataInterfaceSkeletalMesh::GetSpecificTriangleAt<
 {
 	int32 MaxIdx = Accessor.SamplingRegionBuiltData->TriangleIndices.Num() - 1;
 	FilteredIndex = FMath::Min(FilteredIndex, MaxIdx);
-	return Accessor.SamplingRegionBuiltData->TriangleIndices[FilteredIndex];
+	return Accessor.SamplingRegionBuiltData->TriangleIndices[FilteredIndex] / 3;
 }
 
 template<>
@@ -571,7 +571,7 @@ FORCEINLINE int32 UNiagaraDataInterfaceSkeletalMesh::GetSpecificTriangleAt<
 	int32 MaxIdx = Accessor.SamplingRegionBuiltData->TriangleIndices.Num() - 1;
 	Idx = FMath::Min(Idx, MaxIdx);
 
-	return Accessor.SamplingRegionBuiltData->TriangleIndices[Idx];
+	return Accessor.SamplingRegionBuiltData->TriangleIndices[Idx] / 3;
 }
 
 template<>
@@ -587,7 +587,7 @@ FORCEINLINE int32 UNiagaraDataInterfaceSkeletalMesh::GetSpecificTriangleAt<
 		const FSkeletalMeshSamplingRegionBuiltData& RegionBuiltData = SamplingInfo.GetRegionBuiltData(InstData->SamplingRegionIndices[RegionIdx]);
 		if (FilteredIndex < RegionBuiltData.TriangleIndices.Num())
 		{
-			return RegionBuiltData.TriangleIndices[FilteredIndex];
+			return RegionBuiltData.TriangleIndices[FilteredIndex] / 3;
 		}
 
 		FilteredIndex -= RegionBuiltData.TriangleIndices.Num();
@@ -608,7 +608,7 @@ FORCEINLINE int32 UNiagaraDataInterfaceSkeletalMesh::GetSpecificTriangleAt<
 		const FSkeletalMeshSamplingRegionBuiltData& RegionBuiltData = SamplingInfo.GetRegionBuiltData(InstData->SamplingRegionIndices[RegionIdx]);
 		if (FilteredIndex < RegionBuiltData.TriangleIndices.Num())
 		{
-			return RegionBuiltData.TriangleIndices[FilteredIndex];
+			return RegionBuiltData.TriangleIndices[FilteredIndex] / 3;
 		}
 		FilteredIndex -= RegionBuiltData.TriangleIndices.Num();
 	}
@@ -636,7 +636,7 @@ void UNiagaraDataInterfaceSkeletalMesh::GetFilteredTriangleAt(FVectorVMContext&
 		int32 RealIdx = 0;
 		RealIdx = GetSpecificTriangleAt<FilterMode, AreaWeightingMode>(Accessor, InstData, Tri);
 
-		int32 TriMax = Accessor.IndexBuffer->Num() - 3;
+		const int32 TriMax = (Accessor.IndexBuffer->Num() / 3) - 1;
 		RealIdx = FMath::Clamp(RealIdx, 0, TriMax);
 
 		*OutTri.GetDest() = RealIdx;
@@ -671,15 +671,13 @@ void UNiagaraDataInterfaceSkeletalMesh::GetTriCoordColor(FVectorVMContext& Conte
 
 	FMultiSizeIndexContainer& Indices = LODData.MultiSizeIndexContainer;
 	const FRawStaticIndexBuffer16or32Interface* IndexBuffer = Indices.GetIndexBuffer();
-	int32 TriMax = IndexBuffer->Num() - 3;
+	const int32 TriMax = (IndexBuffer->Num() / 3) - 1;
 	for (int32 i = 0; i < Context.NumInstances; ++i)
 	{
-		int32 Tri = TriParam.Get();
-		Tri = FMath::Clamp(Tri, 0, TriMax);
-
-		int32 Idx0 = IndexBuffer->Get(Tri);
-		int32 Idx1 = IndexBuffer->Get(Tri + 1);
-		int32 Idx2 = IndexBuffer->Get(Tri + 2);
+		const int32 Tri = FMath::Clamp(TriParam.Get(), 0, TriMax) * 3;
+		const int32 Idx0 = IndexBuffer->Get(Tri);
+		const int32 Idx1 = IndexBuffer->Get(Tri + 1);
+		const int32 Idx2 = IndexBuffer->Get(Tri + 2);
 
 		FLinearColor Color = BarycentricInterpolate(BaryXParam.Get(), BaryYParam.Get(), BaryZParam.Get(),
 			Colors.VertexColor(Idx0).ReinterpretAsLinear(), Colors.VertexColor(Idx1).ReinterpretAsLinear(), Colors.VertexColor(Idx2).ReinterpretAsLinear());
@@ -745,17 +743,16 @@ void UNiagaraDataInterfaceSkeletalMesh::GetTriCoordUV(FVectorVMContext& Context)
 
 	FMultiSizeIndexContainer& Indices = LODData.MultiSizeIndexContainer;
 	FRawStaticIndexBuffer16or32Interface* IndexBuffer = Indices.GetIndexBuffer();
-	int32 TriMax = IndexBuffer->Num() - 3;
-	int32 UVSetMax = LODData.StaticVertexBuffers.StaticMeshVertexBuffer.GetNumTexCoords() - 1;
-	float InvDt = 1.0f / InstData->DeltaSeconds;
+	const int32 TriMax = (IndexBuffer->Num() / 3) - 1;
+	const int32 UVSetMax = LODData.StaticVertexBuffers.StaticMeshVertexBuffer.GetNumTexCoords() - 1;
+	const float InvDt = 1.0f / InstData->DeltaSeconds;
 	for (int32 i = 0; i < Context.NumInstances; ++i)
 	{
-		int32 Tri = TriParam.Get();
-		Tri = FMath::Clamp(Tri, 0, TriMax);
+		const int32 Tri = FMath::Clamp(TriParam.Get(), 0, TriMax) * 3;
+		const int32 Idx0 = IndexBuffer->Get(Tri);
+		const int32 Idx1 = IndexBuffer->Get(Tri + 1);
+		const int32 Idx2 = IndexBuffer->Get(Tri + 2);
 
-		int32 Idx0 = IndexBuffer->Get(Tri);
-		int32 Idx1 = IndexBuffer->Get(Tri + 1);
-		int32 Idx2 = IndexBuffer->Get(Tri + 2);
 		FVector2D UV0;		FVector2D UV1;		FVector2D UV2;
 		int32 UVSet = UVSetParam.Get();
 		UVSet = FMath::Clamp(UVSet, 0, UVSetMax);
@@ -876,8 +873,8 @@ void UNiagaraDataInterfaceSkeletalMesh::GetTriCoordSkinnedData(FVectorVMContext&
 
 	FSkeletalMeshAccessorHelper Accessor;
 	Accessor.Init<TIntegralConstant<ENDISkeletalMesh_FilterMode, ENDISkeletalMesh_FilterMode::None>, TIntegralConstant<ENDISkelMesh_AreaWeightingMode, ENDISkelMesh_AreaWeightingMode::None>>(InstData);
-	int32 TriMax = Accessor.IndexBuffer->Num() - 3;
-	float InvDt = 1.0f / InstData->DeltaSeconds;
+	const int32 TriMax = (Accessor.IndexBuffer->Num() / 3) - 1;
+	const float InvDt = 1.0f / InstData->DeltaSeconds;
 
 	FVector Pos0;		FVector Pos1;		FVector Pos2;
 	FVector Prev0;		FVector Prev1;		FVector Prev2;
@@ -1067,11 +1064,11 @@ void UNiagaraDataInterfaceSkeletalMesh::GetTriCoordVertices(FVectorVMContext& Co
 	FSkeletalMeshAccessorHelper Accessor;
 	Accessor.Init<TIntegralConstant<int32, 0>, TIntegralConstant<int32, 0>>(InstData);
 
-	int32 TriMax = Accessor.IndexBuffer->Num() - 3;
+	const int32 TriMax = (Accessor.IndexBuffer->Num() / 3) - 1;
 
 	for (int32 i = 0; i < Context.NumInstances; ++i)
 	{
-		int32 Tri = FMath::Clamp(TriParam.GetAndAdvance(), 0, TriMax);
+		const int32 Tri = FMath::Clamp(TriParam.GetAndAdvance(), 0, TriMax);
 		SkinningHandler.GetTrianlgeIndices(Accessor, TriParam.GetAndAdvance(), Idx0, Idx1, Idx2);
 		*OutV0.GetDestAndAdvance() = Idx0;
 		*OutV1.GetDestAndAdvance() = Idx1;
diff --git a/Engine/Plugins/FX/Niagara/Source/Niagara/Private/NiagaraCollision.cpp b/Engine/Plugins/FX/Niagara/Source/Niagara/Private/NiagaraCollision.cpp
--- a/Engine/Plugins/FX/Niagara/Source/Niagara/Private/NiagaraCollision.cpp
+++ b/Engine/Plugins/FX/Niagara/Source/Niagara/Private/NiagaraCollision.cpp
@@ -17,14 +17,14 @@ void FNiagaraCollisionBatch::KickoffNewBatch(FNiagaraEmitterInstance *Sim, float
 	FNiagaraVariable VelVar(FNiagaraTypeDefinition::GetVec3Def(), "Velocity");
 	FNiagaraVariable SizeVar(FNiagaraTypeDefinition::GetVec2Def(), "SpriteSize");
 	FNiagaraVariable TstVar(FNiagaraTypeDefinition::GetBoolDef(), "PerformCollision");
-	FNiagaraDataSetIterator<FVector> PosIt(Sim->GetData(), PosVar, 0, false);
-	FNiagaraDataSetIterator<FVector> VelIt(Sim->GetData(), VelVar, 0, false);
-	FNiagaraDataSetIterator<FVector2D> SizeIt(Sim->GetData(), SizeVar, 0, false);
+	FNiagaraDataSetAccessor<FVector> PosData(Sim->GetData(), PosVar);
+	FNiagaraDataSetAccessor<FVector> VelData(Sim->GetData(), VelVar);
+	FNiagaraDataSetAccessor<FVector2D> SizeData(Sim->GetData(), SizeVar);
 	//FNiagaraDataSetIterator<int32> TstIt(Sim->GetData(), TstVar, 0, false);
 
-	bool bUseSize = SizeIt.IsValid();
+	bool bUseSize = SizeData.IsValidForRead();
 
-	if (!PosIt.IsValid() || !VelIt.IsValid() /*|| !TstIt.IsValid()*/)
+	if (!PosData.IsValidForRead() || !VelData.IsValidForRead() /*|| !TstIt.IsValid()*/)
 	{
 		return;
 	}
@@ -34,25 +34,24 @@ void FNiagaraCollisionBatch::KickoffNewBatch(FNiagaraEmitterInstance *Sim, float
 	{
 		CollisionTraces.Empty();
 
-		for (uint32 i = 0; i < Sim->GetData().GetPrevNumInstances(); i++)
+		for (uint32 i = 0; i < Sim->GetData().GetCurrentDataChecked().GetNumInstances(); i++)
 		{
 			//int32 TestCollision = *TstIt;
 			//if (TestCollision)
 			{
-				check(PosIt.IsValid() && VelIt.IsValid() && (bUseSize == false || SizeIt.IsValid()));
 				FVector Position;
 				FVector EndPosition;
 				FVector Velocity;
 
-				PosIt.Get(Position);
-				VelIt.Get(Velocity);
+				PosData.Get(i, Position);
+				VelData.Get(i, Velocity);
 				EndPosition = Position + Velocity * DeltaSeconds;
 
 				if (bUseSize)
 				{
 					//TODO:  Handle mesh particles too.  Also this can probably be better and or faster.
 					FVector2D SpriteSize;
-					SizeIt.Get(SpriteSize);
+					SizeData.Get(i, SpriteSize);
 					float MaxSize = FMath::Max(SpriteSize.X, SpriteSize.Y);
 
 					float Length;
@@ -71,10 +70,6 @@ void FNiagaraCollisionBatch::KickoffNewBatch(FNiagaraEmitterInstance *Sim, float
 				Trace.SourceParticleIndex = i;
 				CollisionTraces.Add(Trace);
 			}
-
-			//TstIt.Advance();
-			PosIt.Advance();
-			VelIt.Advance();
 		}
 	}
 }
@@ -135,12 +130,13 @@ void FNiagaraCollisionBatch::GenerateEventsFromResults(FNiagaraEmitterInstance *
 			}
 		}
 
-		if (Payloads.Num())
+		int32 NumInstances = Payloads.Num();
+		if (NumInstances)
 		{
 			// now allocate the data set and write all the event structs
 			//
-			CollisionEventDataSet->Allocate(Payloads.Num());
-			CollisionEventDataSet->SetNumInstances(Payloads.Num());
+			CollisionEventDataSet->Allocate(NumInstances);
+			CollisionEventDataSet->GetDestinationDataChecked().SetNumInstances(NumInstances);
 			//FNiagaraVariable ValidVar(FNiagaraTypeDefinition::GetIntDef(), "Valid");
 			FNiagaraVariable PosVar(FNiagaraTypeDefinition::GetVec3Def(), "CollisionLocation");
 			FNiagaraVariable VelVar(FNiagaraTypeDefinition::GetVec3Def(), "CollisionVelocity");
@@ -148,34 +144,29 @@ void FNiagaraCollisionBatch::GenerateEventsFromResults(FNiagaraEmitterInstance *
 			FNiagaraVariable PhysMatIdxVar(FNiagaraTypeDefinition::GetIntDef(), "PhysicalMaterialIndex");
 			FNiagaraVariable ParticleIndexVar(FNiagaraTypeDefinition::GetIntDef(), "ParticleIndex");
 			//FNiagaraDataSetIterator<int32> ValidItr(*CollisionEventDataSet, ValidVar, 0, true);
-			FNiagaraDataSetIterator<FVector> PosItr(*CollisionEventDataSet, PosVar, 0, true);
-			FNiagaraDataSetIterator<FVector> NormItr(*CollisionEventDataSet, NormVar, 0, true);
-			FNiagaraDataSetIterator<FVector> VelItr(*CollisionEventDataSet, VelVar, 0, true);
-			FNiagaraDataSetIterator<int32> PhysMatItr(*CollisionEventDataSet, PhysMatIdxVar, 0, true);
-			FNiagaraDataSetIterator<int32> ParticleIndexItr(*CollisionEventDataSet, ParticleIndexVar, 0, true);
+			FNiagaraDataSetAccessor<FVector> PosItr(*CollisionEventDataSet, PosVar);
+			FNiagaraDataSetAccessor<FVector> NormItr(*CollisionEventDataSet, NormVar);
+			FNiagaraDataSetAccessor<FVector> VelItr(*CollisionEventDataSet, VelVar);
+			FNiagaraDataSetAccessor<int32> PhysMatItr(*CollisionEventDataSet, PhysMatIdxVar);
+			FNiagaraDataSetAccessor<int32> ParticleIndexItr(*CollisionEventDataSet, ParticleIndexVar);
 
-			for (FNiagaraCollisionEventPayload &Payload : Payloads)
+			for (int32 i = 0; i < NumInstances; ++i)
 			{
 				SCOPE_CYCLE_COUNTER(STAT_NiagaraEventWrite);
+				FNiagaraCollisionEventPayload& Payload = Payloads[i];
 
 				check(/*ValidItr.IsValid() && */PosItr.IsValid() && VelItr.IsValid() && NormItr.IsValid() && PhysMatItr.IsValid());
 				//ValidItr.Set(1);
-				PosItr.Set(Payload.CollisionPos);
-				VelItr.Set(Payload.CollisionVelocity);
-				NormItr.Set(Payload.CollisionNormal);
-				ParticleIndexItr.Set(Payload.ParticleIndex);
-				PhysMatItr.Set(0);
-				//ValidItr.Advance();
-				PosItr.Advance();
-				VelItr.Advance();
-				NormItr.Advance();
-				PhysMatItr.Advance();
-				ParticleIndexItr.Advance();
+				PosItr.Set(i, Payload.CollisionPos);
+				VelItr.Set(i, Payload.CollisionVelocity);
+				NormItr.Set(i, Payload.CollisionNormal);
+				ParticleIndexItr.Set(i, Payload.ParticleIndex);
+				PhysMatItr.Set(i, 0);
 			}
 		}
 		else
 		{
-			CollisionEventDataSet->SetNumInstances(0);
+			CollisionEventDataSet->ResetBuffers();
 		}
 
 	}
diff --git a/Engine/Plugins/FX/Niagara/Source/Niagara/Private/NiagaraComponent.cpp b/Engine/Plugins/FX/Niagara/Source/Niagara/Private/NiagaraComponent.cpp
--- a/Engine/Plugins/FX/Niagara/Source/Niagara/Private/NiagaraComponent.cpp
+++ b/Engine/Plugins/FX/Niagara/Source/Niagara/Private/NiagaraComponent.cpp
@@ -20,6 +20,7 @@
 #include "NiagaraWorldManager.h"
 #include "EngineUtils.h"
 #include "ProfilingDebugging/CsvProfiler.h"
+#include "NiagaraEmitterInstanceBatcher.h"
 
 DECLARE_CYCLE_STAT(TEXT("Sceneproxy create (GT)"), STAT_NiagaraCreateSceneProxy, STATGROUP_Niagara);
 DECLARE_CYCLE_STAT(TEXT("Component Tick (GT)"), STAT_NiagaraComponentTick, STATGROUP_Niagara);
@@ -73,7 +74,7 @@ void DumpNiagaraComponents(UWorld* World)
 
 					if (!SysInst->IsComplete())
 					{
-						for (TSharedRef<FNiagaraEmitterInstance> Emitter : SysInst->GetEmitters())
+						for (TSharedRef<FNiagaraEmitterInstance, ESPMode::ThreadSafe> Emitter : SysInst->GetEmitters())
 						{
 							UE_LOG(LogNiagara, Log, TEXT("    Emitter: \"%s\" | ExecState: %d | NumParticles: %d | CPUTime: %f"), *Emitter->GetEmitterHandle().GetUniqueInstanceName(),
 								(int32)Emitter->GetExecutionState(), Emitter->GetNumParticles(), Emitter->GetTotalCPUTime());
@@ -93,7 +94,7 @@ FAutoConsoleCommandWithWorld DumpNiagaraComponentsCommand(
 
 
 FNiagaraSceneProxy::FNiagaraSceneProxy(const UNiagaraComponent* InComponent)
-		: FPrimitiveSceneProxy(InComponent)
+		: FPrimitiveSceneProxy(InComponent, InComponent->GetAsset() ? InComponent->GetAsset()->GetFName() : FName())
 		, bRenderingEnabled(true)
 {
 	// In this case only, update the System renderers on the game thread.
@@ -101,18 +102,9 @@ FNiagaraSceneProxy::FNiagaraSceneProxy(const UNiagaraComponent* InComponent)
 	FNiagaraSystemInstance* SystemInst = InComponent->GetSystemInstance();
 	if (SystemInst)
 	{
-		TArray<TSharedRef<FNiagaraEmitterInstance> > &Sims = SystemInst->GetEmitters();
-		TArray<NiagaraRenderer*> RenderersFromSims;
-		for (TSharedRef<FNiagaraEmitterInstance> Sim : Sims)
-		{
-			for (int32 i = 0; i < Sim->GetEmitterRendererNum(); i++)
-			{
-				RenderersFromSims.Add(Sim->GetEmitterRenderer(i));
-			}
-		}
-		UpdateEmitterRenderers(RenderersFromSims);
 		//UE_LOG(LogNiagara, Warning, TEXT("FNiagaraSceneProxy %p"), this);
 
+		CreateRenderers(InComponent);
 		bAlwaysHasVelocity = true;
 		Batcher = SystemInst->GetBatcher();
 	}
@@ -124,51 +116,95 @@ SIZE_T FNiagaraSceneProxy::GetTypeHash() const
 	return reinterpret_cast<size_t>(&UniquePointer);
 }
 
-void FNiagaraSceneProxy::UpdateEmitterRenderers(const TArray<NiagaraRenderer*>& InRenderers)
+void FNiagaraSceneProxy::ReleaseRenderers()
 {
-	EmitterRenderers.Empty();
-	for (NiagaraRenderer* EmitterRenderer : InRenderers)
+	if (EmitterRenderers.Num() > 0)
 	{
-		if (EmitterRenderer)
+		//Renderers must be freed on the render thread.
+		ENQUEUE_RENDER_COMMAND(ReleaseRenderersCommand)(
+			[ToDeleteEmitterRenderers = MoveTemp(EmitterRenderers)](FRHICommandListImmediate& RHICmdList)
 		{
-			EmitterRenderers.Add(EmitterRenderer);
+			for (FNiagaraRenderer* EmitterRenderer : ToDeleteEmitterRenderers)
+			{
+				if (EmitterRenderer)
+				{
+					EmitterRenderer->ReleaseRenderThreadResources();
+					delete EmitterRenderer;
+				}
+			}
+		});
+		EmitterRenderers.Empty();
+	}
+	RendererDrawOrder.Empty();
+}
+
+void FNiagaraSceneProxy::CreateRenderers(const UNiagaraComponent* Component)
+{
+	check(Component);
+	check(Component->GetSystemInstance());
+
+	UNiagaraSystem* System = Component->GetAsset();
+	check(System);
+
+	struct FSortInfo
+	{
+		FSortInfo(int32 InSortHint, int32 InRendererIdx): SortHint(InSortHint), RendererIdx(InRendererIdx){}
+		int32 SortHint;
+		int32 RendererIdx;
+	};
+	TArray<FSortInfo, TInlineAllocator<8>> RendererSortInfo;
+
+	ReleaseRenderers();
+	ERHIFeatureLevel::Type FeatureLevel = GetScene().GetFeatureLevel();
+	for(TSharedRef<const FNiagaraEmitterInstance, ESPMode::ThreadSafe> EmitterInst : Component->GetSystemInstance()->GetEmitters())
+	{
+		if (UNiagaraEmitter* Emitter = EmitterInst->GetCachedEmitter())
+		{
+			for (UNiagaraRendererProperties* Properties : Emitter->GetRenderers())
+			{
+				RendererSortInfo.Emplace(Properties->SortOrderHint, EmitterRenderers.Num());
+				FNiagaraRenderer* NewRenderer = nullptr;
+				if (Properties->GetIsEnabled())
+				{
+					NewRenderer = Properties->CreateEmitterRenderer(FeatureLevel, &EmitterInst.Get());
+				}
+				EmitterRenderers.Add(NewRenderer);
+			}
 		}
 	}
 
 	// We sort by the sort hint in order to guarantee that we submit according to the preferred sort order..
-	EmitterRenderers.Sort([&](const NiagaraRenderer& A, const NiagaraRenderer& B)
+	RendererSortInfo.Sort([&](const FSortInfo& A, const FSortInfo& B)
 	{
-		check(A.GetRendererProperties());
-		check(B.GetRendererProperties());
-		int32 AIndex = A.GetRendererProperties()->SortOrderHint;
-		int32 BIndex = B.GetRendererProperties()->SortOrderHint;
+		int32 AIndex = A.SortHint;
+		int32 BIndex = B.SortHint;
 		return AIndex < BIndex;
 	});
+	RendererDrawOrder.Reset(RendererSortInfo.Num());
+	for (FSortInfo SortInfo : RendererSortInfo)
+	{
+		RendererDrawOrder.Add(SortInfo.RendererIdx);
+	}
 }
 
 FNiagaraSceneProxy::~FNiagaraSceneProxy()
 {
 	//UE_LOG(LogNiagara, Warning, TEXT("~FNiagaraSceneProxy %p"), this);
-	ReleaseRenderThreadResources();
-}
-
-/** Called on render thread to assign new dynamic data */
-void FNiagaraSceneProxy::SetDynamicData_RenderThread(FNiagaraDynamicDataBase* NewDynamicData)
-{
-	for (NiagaraRenderer* Renderer : EmitterRenderers)
+	check(IsInRenderingThread());
+	for (FNiagaraRenderer* EmitterRenderer : EmitterRenderers)
 	{
-		if (Renderer)
+		if (EmitterRenderer)
 		{
-			Renderer->SetDynamicData_RenderThread(NewDynamicData);
+			EmitterRenderer->ReleaseRenderThreadResources();
+			delete EmitterRenderer;
 		}
 	}
-	return;
+	EmitterRenderers.Empty();
 }
 
-
 void FNiagaraSceneProxy::ReleaseRenderThreadResources()
 {
-	for (NiagaraRenderer* Renderer : EmitterRenderers)
+	for (FNiagaraRenderer* Renderer : EmitterRenderers)
 	{
 		if (Renderer)
 		{
@@ -181,7 +217,7 @@ void FNiagaraSceneProxy::ReleaseRenderThreadResources()
 // FPrimitiveSceneProxy interface.
 void FNiagaraSceneProxy::CreateRenderThreadResources()
 {
-	for (NiagaraRenderer* Renderer : EmitterRenderers)
+	for (FNiagaraRenderer* Renderer : EmitterRenderers)
 	{
 		if (Renderer)
 		{
@@ -193,7 +229,7 @@ void FNiagaraSceneProxy::CreateRenderThreadResources()
 
 void FNiagaraSceneProxy::OnTransformChanged()
 {
-	for (NiagaraRenderer* Renderer : EmitterRenderers)
+	for (FNiagaraRenderer* Renderer : EmitterRenderers)
 	{
 		if (Renderer)
 		{
@@ -212,9 +248,9 @@ FPrimitiveViewRelevance FNiagaraSceneProxy::GetViewRelevance(const FSceneView* V
 	}
 	Relevance.bDynamicRelevance = true;
 
-	for (NiagaraRenderer* Renderer : EmitterRenderers)
+	for (FNiagaraRenderer* Renderer : EmitterRenderers)
 	{
-		if (Renderer && Renderer->GetRendererProperties()->GetIsEnabled())
+		if (Renderer)
 		{
 			Relevance |= Renderer->GetViewRelevance(View, this);
 		}
@@ -234,7 +270,7 @@ uint32 FNiagaraSceneProxy::GetMemoryFootprint() const
 uint32 FNiagaraSceneProxy::GetAllocatedSize() const
 { 
 	uint32 DynamicDataSize = 0;
-	for (NiagaraRenderer* Renderer : EmitterRenderers)
+	for (FNiagaraRenderer* Renderer : EmitterRenderers)
 	{
 		if (Renderer)
 		{
@@ -256,9 +292,11 @@ void FNiagaraSceneProxy::SetRenderingEnabled(bool bInRenderingEnabled)
 
 void FNiagaraSceneProxy::GetDynamicMeshElements(const TArray<const FSceneView*>& Views, const FSceneViewFamily& ViewFamily, uint32 VisibilityMap, FMeshElementCollector& Collector) const
 {
+	SCOPE_CYCLE_COUNTER(STAT_NiagaraOverview_RT);
 	SCOPE_CYCLE_COUNTER(STAT_NiagaraComponentGetDynamicMeshElements);
-	for (NiagaraRenderer* Renderer : EmitterRenderers)
+	for (int32 RendererIdx : RendererDrawOrder)
 	{
+		FNiagaraRenderer* Renderer = EmitterRenderers[RendererIdx];
 		if (Renderer)
 		{
 			Renderer->GetDynamicMeshElements(Views, ViewFamily, VisibilityMap, Collector, this);
@@ -284,7 +322,7 @@ void FNiagaraSceneProxy::GetDynamicMeshElements(const TArray<const FSceneView*>&
 #if RHI_RAYTRACING
 void FNiagaraSceneProxy::GetDynamicRayTracingInstances(FRayTracingMaterialGatheringContext& Context, TArray<FRayTracingInstance>& OutRayTracingInstances)
 {
-	for (NiagaraRenderer* Renderer : EmitterRenderers)
+	for (FNiagaraRenderer* Renderer : EmitterRenderers)
 	{
 		if (Renderer)
 		{
@@ -298,33 +336,14 @@ void FNiagaraSceneProxy::GatherSimpleLights(const FSceneViewFamily& ViewFamily,
 {
 	for (int32 Idx = 0; Idx < EmitterRenderers.Num(); Idx++)
 	{
-		NiagaraRenderer *Renderer = EmitterRenderers[Idx];
-		if (Renderer && Renderer->GetPropertiesClass() == UNiagaraLightRendererProperties::StaticClass())
+		FNiagaraRenderer *Renderer = EmitterRenderers[Idx];
+		if (Renderer && Renderer->HasLights())
 		{
-			NiagaraRendererLights* LightRenderer = static_cast<NiagaraRendererLights*>(Renderer);
-			FNiagaraDynamicDataLights* DynamicData = static_cast<FNiagaraDynamicDataLights*>(Renderer->GetDynamicData());
-
-			if (DynamicData)
-			{
-				int32 LightCount = DynamicData->LightArray.Num();
-
-				OutParticleLights.InstanceData.Reserve(LightCount + OutParticleLights.InstanceData.Num());
-				OutParticleLights.PerViewData.Reserve(LightCount + OutParticleLights.PerViewData.Num());
-
-				for (NiagaraRendererLights::SimpleLightData &LightData : DynamicData->LightArray)
-				{
-					// When not using camera-offset, output one position for all views to share. 
-					OutParticleLights.PerViewData.Add(LightData.PerViewEntry);
-
-					// Add an entry for the light instance.
-					OutParticleLights.InstanceData.Add(LightData.LightEntry);
-				}
-			}
+			Renderer->GatherSimpleLights(OutParticleLights);
 		}
 	}
 }
 
-
 //////////////////////////////////////////////////////////////////////////
 
 UNiagaraComponent::UNiagaraComponent(const FObjectInitializer& ObjectInitializer)
@@ -356,6 +375,8 @@ UNiagaraComponent::UNiagaraComponent(const FObjectInitializer& ObjectInitializer
 	bAutoActivate = true;
 	bRenderingEnabled = true;
 	SavedAutoAttachRelativeScale3D = FVector(1.f, 1.f, 1.f);
+
+	SetGenerateOverlapEvents(false);
 }
 
 /********* UFXSystemComponent *********/
@@ -386,6 +407,8 @@ void UNiagaraComponent::TickComponent(float DeltaSeconds, enum ELevelTick TickTy
 	SCOPE_CYCLE_COUNTER(STAT_NiagaraComponentTick);
 	CSV_SCOPED_TIMING_STAT_EXCLUSIVE(Particles);
 
+	FScopeCycleCounter SystemStatCounter(Asset ? Asset->GetStatID(true, false) : TStatId());
+
 	if (bAwaitingActivationDueToNotReady)
 	{
 		Activate(bActivateShouldResetWhenReady);
@@ -523,16 +546,48 @@ bool UNiagaraComponent::IsPaused()const
 	return false;
 }
 
+bool UNiagaraComponent::IsWorldReadyToRun() const
+{
+	// The niagara system instance assumes that a batcher exists when it is created. We need to wait until this has happened before successfully activating this system.
+	bool FXSystemExists = false;
+	bool WorldManagerExists = false;
+	UWorld* World = GetWorld();
+	if (World)
+	{
+		if (World->Scene)
+		{
+			FFXSystemInterface*  FXSystemInterface = World->Scene->GetFXSystem();
+			if (FXSystemInterface)
+			{
+				NiagaraEmitterInstanceBatcher* FoundBatcher = static_cast<NiagaraEmitterInstanceBatcher*>(FXSystemInterface->GetInterface(NiagaraEmitterInstanceBatcher::Name));
+				if (FoundBatcher != nullptr)
+				{
+					FXSystemExists = true;
+				}
+			}
+		}
+
+		FNiagaraWorldManager* WorldManager = FNiagaraWorldManager::Get(World);
+		if (WorldManager)
+		{
+			WorldManagerExists = true;
+		}
+	}
+
+	return WorldManagerExists && FXSystemExists;
+}
+
 bool UNiagaraComponent::InitializeSystem()
 {
 	if (SystemInstance.IsValid() == false)
 	{
-		//UE_LOG(LogNiagara, Warning, TEXT("Create System: %0X | %s\n"), this, *GetAsset()->GetFullName());
-		SystemInstance = MakeUnique<FNiagaraSystemInstance>(this);
+		FNiagaraSystemInstance::AllocateSystemInstance(this, SystemInstance);
+		//UE_LOG(LogNiagara, Log, TEXT("Create System: %p | %s\n"), SystemInstance.Get(), *GetAsset()->GetFullName());
 #if WITH_EDITORONLY_DATA
 		OnSystemInstanceChangedDelegate.Broadcast();
 #endif
 		SystemInstance->Init(bForceSolo);
+		MarkRenderStateDirty();
 		return true;
 	}
 	return false;
@@ -593,8 +648,8 @@ void UNiagaraComponent::Activate(bool bReset /* = false */)
 		Asset->PollForCompilationComplete();
 	}
 #endif
-
-	if (!Asset->IsReadyToRun())
+	
+	if (!Asset->IsReadyToRun() || !IsWorldReadyToRun())
 	{
 		bAwaitingActivationDueToNotReady = true;
 		bActivateShouldResetWhenReady = bReset;
@@ -698,18 +753,22 @@ void UNiagaraComponent::DeactivateImmediate()
 
 void UNiagaraComponent::OnSystemComplete()
 {
-	//UE_LOG(LogNiagara, Log, TEXT("OnSystemComplete: %u - %s"), this, *Asset->GetName());
+	//UE_LOG(LogNiagara, Log, TEXT("OnSystemComplete: %p - %s"), SystemInstance.Get(), *Asset->GetName());
 
 	SetComponentTickEnabled(false);
 	bIsActive = false;
 
 	MarkRenderDynamicDataDirty();
 		
+	//UE_LOG(LogNiagara, Log, TEXT("OnSystemFinished.Broadcast(this);: { %p - %s"), SystemInstance.Get(), *Asset->GetName());
 	OnSystemFinished.Broadcast(this);
+	//UE_LOG(LogNiagara, Log, TEXT("OnSystemFinished.Broadcast(this);: } %p - %s"), SystemInstance.Get(), *Asset->GetName());
 
 	if (bAutoDestroy)
 	{
+		//UE_LOG(LogNiagara, Log, TEXT("OnSystemComplete DestroyComponent();: { %p - %s"), SystemInstance.Get(), *Asset->GetName());
 		DestroyComponent();
+		//UE_LOG(LogNiagara, Log, TEXT("OnSystemComplete DestroyComponent();: } %p - %s"), SystemInstance.Get(), *Asset->GetName());
 	}
 	else if (bAutoManageAttachment)
 	{
@@ -719,13 +778,20 @@ void UNiagaraComponent::OnSystemComplete()
 
 void UNiagaraComponent::DestroyInstance()
 {
-	//UE_LOG(LogNiagara, Warning, TEXT("Destroy System: %s\n"), *GetAsset()->GetFullName());
+	//UE_LOG(LogNiagara, Log, TEXT("UNiagaraComponent::DestroyInstance: %p  %s\n"), SystemInstance.Get(), *GetAsset()->GetFullName());
 	//UE_LOG(LogNiagara, Log, TEXT("DestroyInstance: %u - %s"), this, *Asset->GetName());
 	bIsActive = false;
-	SystemInstance = nullptr;
+	
+	// Rather than setting the unique ptr to null here, we allow it to transition ownership to the system's deferred deletion queue. This allows us to safely
+	// get rid of the system interface should we be doing this in response to a callback invoked during the system interface's lifetime completion cycle.
+	FNiagaraSystemInstance::DeallocateSystemInstance(SystemInstance); // System Instance will be nullptr after this.
+	check(SystemInstance.Get() == nullptr);
+
+
 #if WITH_EDITORONLY_DATA
 	OnSystemInstanceChangedDelegate.Broadcast();
 #endif
+	MarkRenderStateDirty();
 }
 
 void UNiagaraComponent::OnRegister()
@@ -786,8 +852,14 @@ void UNiagaraComponent::OnUnregister()
 
 	if (SystemInstance)
 	{
+		//UE_LOG(LogNiagara, Log, TEXT("UNiagaraComponent::OnUnregister: %p  %s\n"), SystemInstance.Get(), *GetAsset()->GetFullName());
+
 		SystemInstance->Deactivate(true);
-		SystemInstance = nullptr;
+
+		// Rather than setting the unique ptr to null here, we allow it to transition ownership to the system's deferred deletion queue. This allows us to safely
+		// get rid of the system interface should we be doing this in response to a callback invoked during the system interface's lifetime completion cycle.
+		FNiagaraSystemInstance::DeallocateSystemInstance(SystemInstance); // System Instance will be nullptr after this.
+		check(SystemInstance.Get() == nullptr);
 #if WITH_EDITORONLY_DATA
 		OnSystemInstanceChangedDelegate.Broadcast();
 #endif
@@ -796,6 +868,7 @@ void UNiagaraComponent::OnUnregister()
 
 void UNiagaraComponent::BeginDestroy()
 {
+	//UE_LOG(LogNiagara, Log, TEXT("~UNiagaraComponent: %p  %s\n"), SystemInstance.Get(), *GetAsset()->GetFullName());
 	DestroyInstance();
 
 	Super::BeginDestroy();
@@ -834,50 +907,85 @@ void UNiagaraComponent::SendRenderDynamicData_Concurrent()
 	SCOPE_CYCLE_COUNTER(STAT_NiagaraComponentSendRenderData);
 	if (SystemInstance.IsValid() && SceneProxy)
 	{
+#if STATS
+		TStatId SystemStatID = GetAsset() ? GetAsset()->GetStatID(true, true) : TStatId();
+		FScopeCycleCounter SystemStatCounter(SystemStatID);
+#endif
 		FNiagaraSceneProxy* NiagaraProxy = static_cast<FNiagaraSceneProxy*>(SceneProxy);
+		const TArray<FNiagaraRenderer*>& EmitterRenderers = NiagaraProxy->GetEmitterRenderers();
+
+		typedef TArray<FNiagaraDynamicDataBase*, TInlineAllocator<8>> TDynamicDataArray;
+		TDynamicDataArray NewDynamicData;
+		NewDynamicData.Reserve(EmitterRenderers.Num());
 
+		int32 RendererIndex = 0;
 		for (int32 i = 0; i < SystemInstance->GetEmitters().Num(); i++)
 		{
-			FNiagaraEmitterInstance* Emitter = &SystemInstance->GetEmitters()[i].Get();
-			for (int32 EmitterIdx = 0; EmitterIdx < Emitter->GetEmitterRendererNum(); EmitterIdx++)
+			FNiagaraEmitterInstance* EmitterInst = &SystemInstance->GetEmitters()[i].Get();
+			UNiagaraEmitter* Emitter = EmitterInst->GetCachedEmitter();
+
+#if STATS
+			TStatId EmitterStatID = Emitter->GetStatID(true, true);
+			FScopeCycleCounter EmitterStatCounter(EmitterStatID);
+#endif
+
+			for (int32 EmitterIdx = 0; EmitterIdx < Emitter->GetRenderers().Num(); EmitterIdx++, RendererIndex++)
 			{
-				NiagaraRenderer* Renderer = Emitter->GetEmitterRenderer(EmitterIdx);
+				UNiagaraRendererProperties* Properties = Emitter->GetRenderers()[EmitterIdx];
+				FNiagaraRenderer* Renderer = EmitterRenderers[RendererIndex];
+				FNiagaraDynamicDataBase* NewData = nullptr;
+				
 				if (Renderer)
 				{
 					bool bRendererEditorEnabled = true;
 #if WITH_EDITORONLY_DATA
-					bRendererEditorEnabled = (!SystemInstance->GetIsolateEnabled() || Emitter->GetEmitterHandle().IsIsolated());
+					const FNiagaraEmitterHandle& Handle = Asset->GetEmitterHandle(i);
+					bRendererEditorEnabled = (!SystemInstance->GetIsolateEnabled() || Handle.IsIsolated());
 #endif
-					bRendererEditorEnabled &= Renderer->GetRendererProperties()->GetIsEnabled();
-					if (bRendererEditorEnabled && !Emitter->IsComplete() && !SystemInstance->IsComplete())
+					bRendererEditorEnabled &= Properties->GetIsEnabled();
+					if (bRendererEditorEnabled && !EmitterInst->IsComplete() && !SystemInstance->IsComplete())
 					{
-						FNiagaraDynamicDataBase* DynamicData = Renderer->GenerateVertexData(NiagaraProxy, Emitter->GetData(), Emitter->GetEmitterHandle().GetInstance()->SimTarget);
-						NiagaraRenderer* EmitterRenderer = Emitter->GetEmitterRenderer(EmitterIdx);
-						ENQUEUE_RENDER_COMMAND(FSendNiagaraDynamicData)(
-							[EmitterRenderer, DynamicData](FRHICommandListImmediate& RHICmdList)
-							{
-								EmitterRenderer->SetDynamicData_RenderThread(DynamicData);
-							});
-					}
-					else
-					{
-						NiagaraRenderer* EmitterRenderer = Emitter->GetEmitterRenderer(EmitterIdx);
-						ENQUEUE_RENDER_COMMAND(FSendNiagaraDynamicData)(
-							[EmitterRenderer](FRHICommandListImmediate& RHICmdList)
-							{
-								EmitterRenderer->SetDynamicData_RenderThread(nullptr);
-							});
+						NewData = Renderer->GenerateDynamicData(NiagaraProxy, Properties, EmitterInst);
 					}
 				}
+
+				NewDynamicData.Add(NewData);
 			}
 		}
+		
+		ENQUEUE_RENDER_COMMAND(ResetDataSetBuffers)(
+			[NiagaraProxy, DynamicData = MoveTemp(NewDynamicData)](FRHICommandListImmediate& RHICmdList)
+		{
+			const TArray<FNiagaraRenderer*>& EmitterRenderers_RT = NiagaraProxy->GetEmitterRenderers();
+			for (int32 i = 0; i < EmitterRenderers_RT.Num(); ++i)
+			{
+				if (FNiagaraRenderer* Renderer = EmitterRenderers_RT[i])
+				{
+					Renderer->SetDynamicData_RenderThread(DynamicData[i]);
+				}
+			}
+		});
 	}
-
 }
 
 int32 UNiagaraComponent::GetNumMaterials() const
 {
-	return 0;
+	TArray<UMaterialInterface*> UsedMaterials;
+	if (SystemInstance)
+	{
+		for (int32 i = 0; i < SystemInstance->GetEmitters().Num(); i++)
+		{
+			FNiagaraEmitterInstance* EmitterInst = &SystemInstance->GetEmitters()[i].Get();
+			UNiagaraEmitter* Emitter = EmitterInst->GetCachedEmitter();
+			for (int32 EmitterIdx = 0; EmitterIdx < Emitter->GetRenderers().Num(); EmitterIdx++)
+			{
+				UNiagaraRendererProperties* Properties = Emitter->GetRenderers()[EmitterIdx];
+				Properties->GetUsedMaterials(UsedMaterials);
+			}
+		}
+	}
+
+	return UsedMaterials.Num();
 }
 
 
@@ -916,8 +1024,11 @@ FBoxSphereBounds UNiagaraComponent::CalcBounds(const FTransform& LocalToWorld) c
 FPrimitiveSceneProxy* UNiagaraComponent::CreateSceneProxy()
 {
 	SCOPE_CYCLE_COUNTER(STAT_NiagaraCreateSceneProxy);
+	SCOPE_CYCLE_COUNTER(STAT_NiagaraOverview_GT);
+	
 	// The constructor will set up the System renderers from the component.
 	FNiagaraSceneProxy* Proxy = new FNiagaraSceneProxy(this);
+
 	return Proxy;
 }
 
@@ -928,10 +1039,12 @@ void UNiagaraComponent::GetUsedMaterials(TArray<UMaterialInterface*>& OutMateria
 		return;
 	}
 
-	for (TSharedRef<FNiagaraEmitterInstance> Sim : SystemInstance->GetEmitters())	
+	for (int32 EmitterIdx = 0; EmitterIdx < SystemInstance->GetEmitters().Num(); ++EmitterIdx)
 	{
+		TSharedRef<FNiagaraEmitterInstance, ESPMode::ThreadSafe> Sim = SystemInstance->GetEmitters()[EmitterIdx];
+
 		if (UNiagaraEmitter* Props = Sim->GetEmitterHandle().GetInstance())
-		{	
+		{
 			for (int32 i = 0; i < Props->GetRenderers().Num(); i++)
 			{
 				if (UNiagaraRendererProperties* Renderer = Props->GetRenderers()[i])
@@ -1015,21 +1128,19 @@ TArray<FVector> UNiagaraComponent::GetNiagaraParticleValueVec3_DebugOnly(const F
 	FName EmitterName = FName(*InEmitterName);
 	if (SystemInstance.IsValid())
 	{
-		for (TSharedRef<FNiagaraEmitterInstance>& Sim : SystemInstance->GetEmitters())
+		for (TSharedRef<FNiagaraEmitterInstance, ESPMode::ThreadSafe>& Sim : SystemInstance->GetEmitters())
 		{
 			if (Sim->GetEmitterHandle().GetName() == EmitterName)
 			{
-				int32 NumParticles = Sim->GetData().GetNumInstances();
+				FNiagaraDataBuffer& ParticleData = Sim->GetData().GetCurrentDataChecked();
+				int32 NumParticles = ParticleData.GetNumInstances();
 				Positions.SetNum(NumParticles);
-				FNiagaraDataSetIterator<FVector> PosItr(Sim->GetData(), FNiagaraVariable(FNiagaraTypeDefinition::GetVec3Def(), *InValueName));
-				int32 i = 0;
-				while (PosItr.IsValid())
+				FNiagaraDataSetAccessor<FVector> PosData(Sim->GetData(), FNiagaraVariable(FNiagaraTypeDefinition::GetVec3Def(), *InValueName));
+				for (int32 i = 0; i < NumParticles; ++i)
 				{
 					FVector Position;
-					PosItr.Get(Position);
+					PosData.Get(i, Position);
 					Positions[i] = Position;
-					i++;
-					PosItr.Advance();
 				}
 			}
 		}
@@ -1040,26 +1151,23 @@ TArray<FVector> UNiagaraComponent::GetNiagaraParticleValueVec3_DebugOnly(const F
 
 TArray<float> UNiagaraComponent::GetNiagaraParticleValues_DebugOnly(const FString& InEmitterName, const FString& InValueName)
 {
-
 	TArray<float> Values;
 	FName EmitterName = FName(*InEmitterName);
 	if (SystemInstance.IsValid())
 	{
-		for (TSharedRef<FNiagaraEmitterInstance>& Sim : SystemInstance->GetEmitters())
+		for (TSharedRef<FNiagaraEmitterInstance, ESPMode::ThreadSafe>& Sim : SystemInstance->GetEmitters())
 		{
 			if (Sim->GetEmitterHandle().GetName() == EmitterName)
 			{
-				int32 NumParticles = Sim->GetData().GetNumInstances();
+				FNiagaraDataBuffer& ParticleData = Sim->GetData().GetCurrentDataChecked();
+				int32 NumParticles = ParticleData.GetNumInstances();
 				Values.SetNum(NumParticles);
-				FNiagaraDataSetIterator<float> ValueItr(Sim->GetData(), FNiagaraVariable(FNiagaraTypeDefinition::GetFloatDef(), *InValueName));
-				int32 i = 0;
-				while (ValueItr.IsValid())
+				FNiagaraDataSetAccessor<float> ValueData(Sim->GetData(), FNiagaraVariable(FNiagaraTypeDefinition::GetFloatDef(), *InValueName));
+				for (int32 i = 0; i < NumParticles; ++i)
 				{
 					float Value;
-					ValueItr.Get(Value);
+					ValueData.Get(i, Value);
 					Values[i] = Value;
-					i++;
-					ValueItr.Advance();
 				}
 			}
 		}
diff --git a/Engine/Plugins/FX/Niagara/Source/Niagara/Private/NiagaraDataInterface.cpp b/Engine/Plugins/FX/Niagara/Source/Niagara/Private/NiagaraDataInterface.cpp
--- a/Engine/Plugins/FX/Niagara/Source/Niagara/Private/NiagaraDataInterface.cpp
+++ b/Engine/Plugins/FX/Niagara/Source/Niagara/Private/NiagaraDataInterface.cpp
@@ -14,6 +14,21 @@ UNiagaraDataInterface::UNiagaraDataInterface(FObjectInitializer const& ObjectIni
 {
 }
 
+UNiagaraDataInterface::~UNiagaraDataInterface()
+{
+	// @todo-threadsafety Can there be a UNiagaraDataInterface class itself created? Perhaps by the system?
+	if ( Proxy.IsValid() )
+	{
+		ENQUEUE_RENDER_COMMAND(FDeleteProxyRT) (
+			[RT_Proxy=MoveTemp(Proxy)](FRHICommandListImmediate& CmdList)
+			{
+				// This will release RT_Proxy on the RT
+			}
+		);
+		check(Proxy.IsValid() == false);
+	}
+}
+
 void UNiagaraDataInterface::PostLoad()
 {
 	Super::PostLoad();
@@ -150,22 +165,61 @@ void UNiagaraDataInterfaceCurveBase::GetParameterDefinitionHLSL(FNiagaraDataInte
 	OutHLSL += TEXT("\n");
 }
 
-FReadBuffer& UNiagaraDataInterfaceCurveBase::GetCurveLUTGPUBuffer()
+//FReadBuffer& UNiagaraDataInterfaceCurveBase::GetCurveLUTGPUBuffer()
+//{
+//	//TODO: This isn't really very thread safe. Need to move to a proxy like system where DIs can push data to the RT safely.
+//	if (GPUBufferDirty)
+//	{
+//		int32 ElemSize = GetCurveNumElems();
+//		CurveLUT.Release();
+//		CurveLUT.Initialize(sizeof(float), CurveLUTWidth * ElemSize, EPixelFormat::PF_R32_FLOAT, BUF_Static);
+//		uint32 BufferSize = ShaderLUT.Num() * sizeof(float);
+//		int32 *BufferData = static_cast<int32*>(RHILockVertexBuffer(CurveLUT.Buffer, 0, BufferSize, EResourceLockMode::RLM_WriteOnly));
+//		FPlatformMemory::Memcpy(BufferData, ShaderLUT.GetData(), BufferSize);
+//		RHIUnlockVertexBuffer(CurveLUT.Buffer);
+//		GPUBufferDirty = false;
+//	}
+//
+//	return CurveLUT;
+//}
+
+void UNiagaraDataInterfaceCurveBase::PushToRenderThread()
 {
-	//TODO: This isn't really very thread safe. Need to move to a proxy like system where DIs can push data to the RT safely.
-	if (GPUBufferDirty)
+	if (!GSupportsResourceView)
 	{
-		int32 ElemSize = GetCurveNumElems();
-		CurveLUT.Release();
-		CurveLUT.Initialize(sizeof(float), CurveLUTWidth * ElemSize, EPixelFormat::PF_R32_FLOAT, BUF_Static);
-		uint32 BufferSize = ShaderLUT.Num() * sizeof(float);
-		int32 *BufferData = static_cast<int32*>(RHILockVertexBuffer(CurveLUT.Buffer, 0, BufferSize, EResourceLockMode::RLM_WriteOnly));
-		FPlatformMemory::Memcpy(BufferData, ShaderLUT.GetData(), BufferSize);
-		RHIUnlockVertexBuffer(CurveLUT.Buffer);
-		GPUBufferDirty = false;
+		return;
 	}
 
-	return CurveLUT;
+	FNiagaraDataInterfaceProxyCurveBase* RT_Proxy = GetProxyAs<FNiagaraDataInterfaceProxyCurveBase>();
+
+	int32 rtNumElems = GetCurveNumElems();
+	float rtLUTMinTime = this->LUTMinTime;
+	float rtLUTMaxTime = this->LUTMaxTime;
+	float rtLUTInvTimeRange = this->LUTInvTimeRange;
+
+	const uint32 rtCurveLUTWidth = CurveLUTWidth;
+
+	// @todo-threadsafety Gross.
+	TArray<float>* ShaderLut_RT = new TArray<float>(ShaderLUT);
+
+	// Push Updates to Proxy.
+	ENQUEUE_RENDER_COMMAND(FUpdateDIColorCurve)(
+		[RT_Proxy, rtCurveLUTWidth, rtNumElems, rtLUTMinTime, rtLUTMaxTime, rtLUTInvTimeRange, ShaderLut_RT](FRHICommandListImmediate& RHICmdList)
+	{
+		RT_Proxy->LUTMinTime = rtLUTMinTime;
+		RT_Proxy->LUTMaxTime = rtLUTMaxTime;
+		RT_Proxy->LUTInvTimeRange = rtLUTInvTimeRange;
+
+		int32 ElemSize = rtNumElems;
+		RT_Proxy->CurveLUT.Release();
+		RT_Proxy->CurveLUT.Initialize(sizeof(float), rtCurveLUTWidth * ElemSize, EPixelFormat::PF_R32_FLOAT, BUF_Static);
+		uint32 BufferSize = ShaderLut_RT->Num() * sizeof(float);
+		int32 *BufferData = static_cast<int32*>(RHILockVertexBuffer(RT_Proxy->CurveLUT.Buffer, 0, BufferSize, EResourceLockMode::RLM_WriteOnly));
+		FPlatformMemory::Memcpy(BufferData, ShaderLut_RT->GetData(), BufferSize);
+		RHIUnlockVertexBuffer(RT_Proxy->CurveLUT.Buffer);
+
+		delete ShaderLut_RT;
+	});
 }
 
 struct FNiagaraDataInterfaceParametersCS_Curve : public FNiagaraDataInterfaceParametersCS
@@ -187,17 +241,17 @@ struct FNiagaraDataInterfaceParametersCS_Curve : public FNiagaraDataInterfacePar
 		Ar << CurveLUT;
 	}
 
-	virtual void Set(FRHICommandList& RHICmdList, FNiagaraShader* Shader, class UNiagaraDataInterface* DataInterface, void* PerInstanceData) const override
+	virtual void Set(FRHICommandList& RHICmdList, const FNiagaraDataInterfaceSetArgs& Context) const override
 	{
 		check(IsInRenderingThread());
 
-		const FComputeShaderRHIParamRef ComputeShaderRHI = Shader->GetComputeShader();
-		UNiagaraDataInterfaceCurveBase* CurveDI = CastChecked<UNiagaraDataInterfaceCurveBase>(DataInterface);
-		FReadBuffer& CurveLUTBuffer = CurveDI->GetCurveLUTGPUBuffer();
+		const FComputeShaderRHIParamRef ComputeShaderRHI = Context.Shader->GetComputeShader();
+		FNiagaraDataInterfaceProxyCurveBase* CurveDI = (FNiagaraDataInterfaceProxyCurveBase*) Context.DataInterface;
+		FReadBuffer& CurveLUTBuffer = CurveDI->CurveLUT;
 
-		SetShaderValue(RHICmdList, ComputeShaderRHI, MinTime, CurveDI->GetMinTime());
-		SetShaderValue(RHICmdList, ComputeShaderRHI, MaxTime, CurveDI->GetMaxTime());
-		SetShaderValue(RHICmdList, ComputeShaderRHI, InvTimeRange, CurveDI->GetInvTimeRange());
+		SetShaderValue(RHICmdList, ComputeShaderRHI, MinTime, CurveDI->LUTMinTime);
+		SetShaderValue(RHICmdList, ComputeShaderRHI, MaxTime, CurveDI->LUTMaxTime);
+		SetShaderValue(RHICmdList, ComputeShaderRHI, InvTimeRange, CurveDI->LUTInvTimeRange);
 		RHICmdList.SetShaderResourceViewParameter(ComputeShaderRHI, CurveLUT.GetBaseIndex(), CurveLUTBuffer.SRV);
 	}
 
diff --git a/Engine/Plugins/FX/Niagara/Source/Niagara/Private/NiagaraDataInterfaceCollisionQuery.cpp b/Engine/Plugins/FX/Niagara/Source/Niagara/Private/NiagaraDataInterfaceCollisionQuery.cpp
--- a/Engine/Plugins/FX/Niagara/Source/Niagara/Private/NiagaraDataInterfaceCollisionQuery.cpp
+++ b/Engine/Plugins/FX/Niagara/Source/Niagara/Private/NiagaraDataInterfaceCollisionQuery.cpp
@@ -6,6 +6,7 @@
 #include "NiagaraWorldManager.h"
 #include "ShaderParameterUtils.h"
 #include "GlobalDistanceFieldParameters.h"
+#include "NiagaraEmitterInstanceBatcher.h"
 #include "Shader.h"
 
 //////////////////////////////////////////////////////////////////////////
@@ -17,6 +18,8 @@ UNiagaraDataInterfaceCollisionQuery::UNiagaraDataInterfaceCollisionQuery(FObject
 	: Super(ObjectInitializer)
 {
 	TraceChannelEnum = FindObject<UEnum>(ANY_PACKAGE, TEXT("ECollisionChannel"), true);
+
+	Proxy = MakeShared<FNiagaraDataIntefaceProxyCollisionQuery, ESPMode::ThreadSafe>();
 }
 
 bool UNiagaraDataInterfaceCollisionQuery::InitPerInstanceData(void* PerInstanceData, FNiagaraSystemInstance* InSystemInstance)
@@ -837,12 +840,7 @@ struct FNiagaraDataInterfaceParametersCS_CollisionQuery : public FNiagaraDataInt
 	virtual void Bind(const FNiagaraDataInterfaceParamRef& ParamRef, const class FShaderParameterMap& ParameterMap) override
 	{
 		PassUniformBuffer.Bind(ParameterMap, FSceneTexturesUniformParameters::StaticStructMetadata.GetShaderVariableName());
-		
 		GlobalDistanceFieldParameters.Bind(ParameterMap);
-		if (GlobalDistanceFieldParameters.IsBound())
-		{
-			GNiagaraViewDataManager.SetGlobalDistanceFieldUsage();
-		}
 	}
 
 	virtual void Serialize(FArchive& Ar) override
@@ -851,18 +849,17 @@ struct FNiagaraDataInterfaceParametersCS_CollisionQuery : public FNiagaraDataInt
 		Ar << GlobalDistanceFieldParameters;
 	}
 
-	virtual void Set(FRHICommandList& RHICmdList, FNiagaraShader* Shader, class UNiagaraDataInterface* DataInterface, void* PerInstanceData) const override
+	virtual void Set(FRHICommandList& RHICmdList, const FNiagaraDataInterfaceSetArgs& Context) const override
 	{
 		check(IsInRenderingThread());
 
-		const FComputeShaderRHIParamRef ComputeShaderRHI = Shader->GetComputeShader();
+		const FComputeShaderRHIParamRef ComputeShaderRHI = Context.Shader->GetComputeShader();
 		
 		TUniformBufferRef<FSceneTexturesUniformParameters> SceneTextureUniformParams = GNiagaraViewDataManager.GetSceneTextureUniformParameters();
 		SetUniformBufferParameter(RHICmdList, ComputeShaderRHI, PassUniformBuffer/*Shader->GetUniformBufferParameter(SceneTexturesUniformBufferStruct)*/, SceneTextureUniformParams);
-		if (GlobalDistanceFieldParameters.IsBound())
+		if (GlobalDistanceFieldParameters.IsBound() && Context.Batcher)
 		{
-			GNiagaraViewDataManager.SetGlobalDistanceFieldUsage();
-			GlobalDistanceFieldParameters.Set(RHICmdList, ComputeShaderRHI, *GNiagaraViewDataManager.GetGlobalDistanceFieldParameters());
+			GlobalDistanceFieldParameters.Set(RHICmdList, ComputeShaderRHI, Context.Batcher->GetGlobalDistanceFieldParameters());
 		}		
 	}
 
diff --git a/Engine/Plugins/FX/Niagara/Source/Niagara/Private/NiagaraDataInterfaceColorCurve.cpp b/Engine/Plugins/FX/Niagara/Source/Niagara/Private/NiagaraDataInterfaceColorCurve.cpp
--- a/Engine/Plugins/FX/Niagara/Source/Niagara/Private/NiagaraDataInterfaceColorCurve.cpp
+++ b/Engine/Plugins/FX/Niagara/Source/Niagara/Private/NiagaraDataInterfaceColorCurve.cpp
@@ -90,7 +90,12 @@ void UNiagaraDataInterfaceColorCurve::UpdateLUT()
 		ShaderLUT.Add(C.B);
 		ShaderLUT.Add(C.A);
 	}
-	GPUBufferDirty = true;
+
+	check(Proxy);
+
+	// @todo-threadsafety Not the best way to do this. Ideally we'd have a funnel where we can update all this data for all 
+	// interfaces at the same time.
+	Super::PushToRenderThread();
 }
 
 bool UNiagaraDataInterfaceColorCurve::CopyToInternal(UNiagaraDataInterface* Destination) const 
diff --git a/Engine/Plugins/FX/Niagara/Source/Niagara/Private/NiagaraDataInterfaceCurlNoise.cpp b/Engine/Plugins/FX/Niagara/Source/Niagara/Private/NiagaraDataInterfaceCurlNoise.cpp
--- a/Engine/Plugins/FX/Niagara/Source/Niagara/Private/NiagaraDataInterfaceCurlNoise.cpp
+++ b/Engine/Plugins/FX/Niagara/Source/Niagara/Private/NiagaraDataInterfaceCurlNoise.cpp
@@ -308,6 +308,8 @@ UNiagaraDataInterfaceCurlNoise::UNiagaraDataInterfaceCurlNoise(FObjectInitialize
 	, Seed(0)
 {
 	OffsetFromSeed = FNiagaraUIntVectorToFVector(Rand3DPCG16(FIntVector(Seed, Seed, Seed))) / 100.0;
+
+	Proxy = MakeShared<FNiagaraDataInterfaceProxyCurlNoise, ESPMode::ThreadSafe>(OffsetFromSeed);
 }
 
 void UNiagaraDataInterfaceCurlNoise::PostInitProperties()
@@ -324,6 +326,8 @@ void UNiagaraDataInterfaceCurlNoise::PostLoad()
 {
 	Super::PostLoad();
 	OffsetFromSeed = FNiagaraUIntVectorToFVector(Rand3DPCG16(FIntVector(Seed, Seed, Seed))) / 100.0;
+
+	PushToRenderThread();
 }
 
 #if WITH_EDITOR
@@ -346,6 +350,8 @@ void UNiagaraDataInterfaceCurlNoise::PostEditChangeProperty(struct FPropertyChan
 		// NOTE: Calculate the offset based on the seed on-change instead of on every invocation for every particle...
 		OffsetFromSeed = FNiagaraUIntVectorToFVector(Rand3DPCG16(FIntVector(Seed, Seed, Seed))) / 100.0;
 	}
+
+	PushToRenderThread();
 }
 
 #endif
@@ -359,6 +365,7 @@ bool UNiagaraDataInterfaceCurlNoise::CopyToInternal(UNiagaraDataInterface* Desti
 	UNiagaraDataInterfaceCurlNoise* DestinationCurlNoise = CastChecked<UNiagaraDataInterfaceCurlNoise>(Destination);
 	DestinationCurlNoise->Seed = Seed;
 	DestinationCurlNoise->OffsetFromSeed = OffsetFromSeed;
+	DestinationCurlNoise->PushToRenderThread();
 
 	return true;
 }
@@ -446,6 +453,20 @@ void UNiagaraDataInterfaceCurlNoise::GetParameterDefinitionHLSL(FNiagaraDataInte
 	OutHLSL += FString::Format(FormatDeclarations, ArgsDeclarations);
 }
 
+void UNiagaraDataInterfaceCurlNoise::PushToRenderThread()
+{
+	FNiagaraDataInterfaceProxyCurlNoise* RT_Proxy = GetProxyAs<FNiagaraDataInterfaceProxyCurlNoise>();
+
+	FVector RT_Offset = OffsetFromSeed;
+
+	// Push Updates to Proxy.
+	ENQUEUE_RENDER_COMMAND(FUpdateDIColorCurve)(
+		[RT_Proxy, RT_Offset](FRHICommandListImmediate& RHICmdList)
+	{
+		RT_Proxy->OffsetFromSeed = RT_Offset;
+	});
+}
+
 struct FNiagaraDataInterfaceParametersCS_CurlNoise : public FNiagaraDataInterfaceParametersCS
 {
 	virtual void Bind(const FNiagaraDataInterfaceParamRef& ParamRef, const class FShaderParameterMap& ParameterMap) override
@@ -458,15 +479,13 @@ struct FNiagaraDataInterfaceParametersCS_CurlNoise : public FNiagaraDataInterfac
 		Ar << OffsetFromSeed;
 	}
 
-	virtual void Set(FRHICommandList& RHICmdList, FNiagaraShader* Shader, class UNiagaraDataInterface* DataInterface, void* PerInstanceData) const override
+	virtual void Set(FRHICommandList& RHICmdList, const FNiagaraDataInterfaceSetArgs& Context) const override
 	{
 		check(IsInRenderingThread());
 
 		// Get shader and DI
-		const FComputeShaderRHIParamRef ComputeShaderRHI = Shader->GetComputeShader();
-		UNiagaraDataInterfaceCurlNoise* CNDI = CastChecked<UNiagaraDataInterfaceCurlNoise>(DataInterface);
-
-		// Note: There is a flush in PreEditChange to make sure everything is synced up at this point 
+		const FComputeShaderRHIParamRef ComputeShaderRHI = Context.Shader->GetComputeShader();
+		FNiagaraDataInterfaceProxyCurlNoise* CNDI = static_cast<FNiagaraDataInterfaceProxyCurlNoise*>(Context.DataInterface);
 
 		// Set parameters
 		SetShaderValue(RHICmdList, ComputeShaderRHI, OffsetFromSeed, CNDI->OffsetFromSeed);
diff --git a/Engine/Plugins/FX/Niagara/Source/Niagara/Private/NiagaraDataInterfaceCurve.cpp b/Engine/Plugins/FX/Niagara/Source/Niagara/Private/NiagaraDataInterfaceCurve.cpp
--- a/Engine/Plugins/FX/Niagara/Source/Niagara/Private/NiagaraDataInterfaceCurve.cpp
+++ b/Engine/Plugins/FX/Niagara/Source/Niagara/Private/NiagaraDataInterfaceCurve.cpp
@@ -117,7 +117,7 @@ void UNiagaraDataInterfaceCurve::UpdateLUT()
 		float C = Curve.Eval(X);
 		ShaderLUT.Add(C);
 	}
-	GPUBufferDirty = true;
+	Super::PushToRenderThread();
 }
 
 // build the shader function HLSL; function name is passed in, as it's defined per-DI; that way, configuration could change
