Copying //UE4/Dev-Niagara to Dev-Main (//UE4/Dev-Main)

diff --git a/Engine/Plugins/FX/Niagara/Source/Niagara/Private/NiagaraDataInterfaceSkeletalMesh.cpp b/Engine/Plugins/FX/Niagara/Source/Niagara/Private/NiagaraDataInterfaceSkeletalMesh.cpp
--- a/Engine/Plugins/FX/Niagara/Source/Niagara/Private/NiagaraDataInterfaceSkeletalMesh.cpp
+++ b/Engine/Plugins/FX/Niagara/Source/Niagara/Private/NiagaraDataInterfaceSkeletalMesh.cpp
@@ -5,6 +5,7 @@
 #include "NiagaraComponent.h"
 #include "NiagaraSystemInstance.h"
 #include "Internationalization/Internationalization.h"
+#include "NiagaraRenderer.h"
 #include "NiagaraScript.h"
 #include "Animation/SkeletalMeshActor.h"
 #include "Components/SkeletalMeshComponent.h"
@@ -15,6 +16,7 @@
 #include "Templates/AlignmentTemplates.h"
 #include "NDISkeletalMeshCommon.h"
 #include "Engine/SkeletalMeshSocket.h"
+#include "ShaderParameterUtils.h"
 
 #define LOCTEXT_NAMESPACE "NiagaraDataInterfaceSkeletalMesh"
 
@@ -117,7 +119,7 @@ void FSkeletalMeshSkinningData::RegisterUser(FSkeletalMeshSkinningDataUsage Usag
 		check(IsInGameThread());
 		if (CurrBoneRefToLocals().Num() == 0)
 		{
-			SkelComp->CacheRefToLocalMatrices(CurrBoneRefToLocals());
+			SkelComp->CacheRefToLocalMatrices(CurrBoneRefToLocals()); 
 		}
 		
 		//Prime the prev matrices if they're missing.
@@ -278,13 +280,481 @@ void FNDI_SkeletalMesh_GeneratedData::TickGeneratedData(float DeltaSeconds)
 	}
 }
 
+//////////////////////////////////////////////////////////////////////////
+// FStaticMeshGpuSpawnBuffer
+
+
+FSkeletalMeshGpuSpawnStaticBuffers::~FSkeletalMeshGpuSpawnStaticBuffers()
+{
+	//ValidSections.Empty();
+}
+
+void FSkeletalMeshGpuSpawnStaticBuffers::Initialise(const FSkeletalMeshLODRenderData& SkeletalMeshLODRenderData,
+	bool bIsGpuUniformlyDistributedSampling, const FSkeletalMeshSamplingLODBuiltData& MeshSamplingLODBuiltData)
+{
+	SkeletalMeshSamplingLODBuiltData = &MeshSamplingLODBuiltData;
+	bUseGpuUniformlyDistributedSampling = bIsGpuUniformlyDistributedSampling;
+
+	LODRenderData = &SkeletalMeshLODRenderData;
+	TriangleCount = SkeletalMeshLODRenderData.MultiSizeIndexContainer.GetIndexBuffer()->Num() / 3;
+	check(TriangleCount > 0);
+}
+
+void FSkeletalMeshGpuSpawnStaticBuffers::InitRHI()
+{
+	// As of today, the UI does not allow to cull specific section of a mesh so this data could be generated on the Mesh. But Section culling might be added later?
+	// Also see https://jira.it.epicgames.net/browse/UE-69376 : we would need to know if GPU sampling of the mesh surface is needed or not on the mesh to be able to do that.
+	// ALso today we do not know if an interface is create from a CPU or GPU emitter. So always allocate for now. Follow up in https://jira.it.epicgames.net/browse/UE-69375.
+
+	const FMultiSizeIndexContainer& IndexBuffer = LODRenderData->MultiSizeIndexContainer;
+	MeshIndexBufferSrv = IndexBuffer.GetIndexBuffer()->GetSRV();
+	MeshVertexBufferSrv = LODRenderData->StaticVertexBuffers.PositionVertexBuffer.GetSRV();
+
+	MeshTangentBufferSRV = LODRenderData->StaticVertexBuffers.StaticMeshVertexBuffer.GetTangentsSRV();
+	//check(MeshTangentBufferSRV->IsValid()); // not available in this stream
+
+	MeshTexCoordBufferSrv = LODRenderData->StaticVertexBuffers.StaticMeshVertexBuffer.GetTexCoordsSRV();
+	NumTexCoord = LODRenderData->StaticVertexBuffers.StaticMeshVertexBuffer.GetNumTexCoords();
+
+	uint32 VertexCount = LODRenderData->StaticVertexBuffers.StaticMeshVertexBuffer.GetNumVertices();
+	uint32 SectionCount = LODRenderData->RenderSections.Num();
+
+	if (bUseGpuUniformlyDistributedSampling)
+	{
+		const FSkeletalMeshAreaWeightedTriangleSampler& triangleSampler = SkeletalMeshSamplingLODBuiltData->AreaWeightedTriangleSampler;
+		const TArray<float>& Prob = triangleSampler.GetProb();
+		const TArray<int32>& Alias = triangleSampler.GetAlias();
+		check(TriangleCount == triangleSampler.GetNumEntries());
+
+		FRHIResourceCreateInfo CreateInfo;
+		void* BufferData = nullptr;
+		uint32 SizeByte = TriangleCount * sizeof(float);
+		BufferTriangleUniformSamplerProbaRHI = RHICreateAndLockVertexBuffer(SizeByte, BUF_Static | BUF_ShaderResource, CreateInfo, BufferData);
+		FMemory::Memcpy(BufferData, Prob.GetData(), SizeByte);
+		RHIUnlockVertexBuffer(BufferTriangleUniformSamplerProbaRHI);
+		BufferTriangleUniformSamplerProbaSRV = RHICreateShaderResourceView(BufferTriangleUniformSamplerProbaRHI, sizeof(float), PF_R32_FLOAT);
+
+		BufferTriangleUniformSamplerAliasRHI = RHICreateAndLockVertexBuffer(SizeByte, BUF_Static | BUF_ShaderResource, CreateInfo, BufferData);
+		FMemory::Memcpy(BufferData, Alias.GetData(), SizeByte);
+		RHIUnlockVertexBuffer(BufferTriangleUniformSamplerAliasRHI);
+		BufferTriangleUniformSamplerAliasSRV = RHICreateShaderResourceView(BufferTriangleUniformSamplerAliasRHI, sizeof(uint32), PF_R32_UINT);
+	}
+
+	// Prepare the vertex matrix lookup offset for each of the sections. This is needed because per vertex BlendIndicies are stored relatively to each Section used matrices.
+	// And these offset per section need to point to the correct matrix according to each section BoneMap.
+	// There is not section selection/culling in the interface so technically we could compute that array in the pipeline.
+	FRHIResourceCreateInfo CreateInfo;
+	void* BufferData = nullptr;
+	BufferTriangleMatricesOffsetRHI = RHICreateAndLockVertexBuffer(VertexCount * sizeof(uint32), BUF_Static | BUF_ShaderResource, CreateInfo, BufferData);
+	uint32* MatricesOffsets = (uint32*)BufferData;
+	uint32 AccumulatedMatrixOffset = 0;
+	for (uint32 s = 0; s < SectionCount; ++s)
+	{
+		const FSkelMeshRenderSection& Section = LODRenderData->RenderSections[s];
+		const uint32 SectionBaseVertexIndex = Section.BaseVertexIndex;
+		const uint32 SectionNumVertices = Section.NumVertices;
+		for (uint32 SectionVertex = 0; SectionVertex < SectionNumVertices; ++SectionVertex)
+		{
+			MatricesOffsets[SectionBaseVertexIndex + SectionVertex] = AccumulatedMatrixOffset;
+		}
+		AccumulatedMatrixOffset += Section.BoneMap.Num();
+	}
+	RHIUnlockVertexBuffer(BufferTriangleMatricesOffsetRHI);
+	BufferTriangleMatricesOffsetSRV = RHICreateShaderResourceView(BufferTriangleMatricesOffsetRHI, sizeof(uint32), PF_R32_UINT);
+}
+
+void FSkeletalMeshGpuSpawnStaticBuffers::ReleaseRHI()
+{
+	BufferTriangleUniformSamplerProbaRHI.SafeRelease();
+	BufferTriangleUniformSamplerProbaSRV.SafeRelease();
+	BufferTriangleUniformSamplerAliasRHI.SafeRelease();
+	BufferTriangleUniformSamplerAliasSRV.SafeRelease();
+}
+
+//////////////////////////////////////////////////////////////////////////
+//FSkeletalMeshGpuSpawnProxy
+
+FSkeletalMeshGpuDynamicBufferProxy::FSkeletalMeshGpuDynamicBufferProxy()
+{
+	//
+}
+
+FSkeletalMeshGpuDynamicBufferProxy::~FSkeletalMeshGpuDynamicBufferProxy()
+{
+	//
+}
+
+void FSkeletalMeshGpuDynamicBufferProxy::Initialise(const FSkeletalMeshLODRenderData& SkeletalMeshLODRenderData)
+{
+	BoneCount = 0;
+	for (const FSkelMeshRenderSection& section : SkeletalMeshLODRenderData.RenderSections)
+	{
+		BoneCount += section.BoneMap.Num();
+	}
+}
+
+void FSkeletalMeshGpuDynamicBufferProxy::InitRHI()
+{
+	for (FSkeletalBuffer& Buffer : RWBufferBones)
+	{
+		FRHIResourceCreateInfo CreateInfo;
+		CreateInfo.DebugName = TEXT("SkeletalMeshGpuDynamicBuffer");
+		Buffer.Buffer = RHICreateVertexBuffer(sizeof(FVector4) * 3 * BoneCount, BUF_ShaderResource | BUF_Dynamic, CreateInfo);
+		Buffer.SRV = RHICreateShaderResourceView(Buffer.Buffer, sizeof(FVector4), PF_A32B32G32R32F);
+	}
+}
+
+void FSkeletalMeshGpuDynamicBufferProxy::ReleaseRHI()
+{
+	for (FSkeletalBuffer& Buffer : RWBufferBones)
+	{
+		Buffer.Buffer.SafeRelease();
+		Buffer.SRV.SafeRelease();
+	}
+}
+
+void FSkeletalMeshGpuDynamicBufferProxy::NewFrame(const FNDISkeletalMesh_InstanceData* InstanceData, int32 LODIndex)
+{
+	USkeletalMeshComponent* SkelComp = InstanceData != nullptr ? Cast<USkeletalMeshComponent>(InstanceData->Component.Get()) : nullptr;
+	if (SkelComp)
+	{
+		TArray<FMatrix> RefToLocalMatrices;
+		SkelComp->CacheRefToLocalMatrices(RefToLocalMatrices);
+
+		TIndirectArray<FSkeletalMeshLODRenderData>& LODRenderDataArray = SkelComp->SkeletalMesh->GetResourceForRendering()->LODRenderData;
+		check(0 <= LODIndex  && LODIndex < LODRenderDataArray.Num());
+		FSkeletalMeshLODRenderData& LODRenderData = LODRenderDataArray[LODIndex];
+		TArray<FSkelMeshRenderSection>& Sections = LODRenderData.RenderSections;
+		uint32 SectionCount = Sections.Num();
+
+		TArray<FVector4> AllSectionsRefToLocalMatrices;
+		static_assert(sizeof(FVector4) == 4*sizeof(float), "FVector4 should match 4 * floats");
+
+		// Count number of matrices we want before appending all of them according to the per section mapping from BoneMap
+		uint32 Float4Count = 0;
+		for ( const FSkelMeshRenderSection& section : Sections )
+		{
+			Float4Count += section.BoneMap.Num() * 3;
+		}
+		check(Float4Count == 3 * BoneCount);
+		AllSectionsRefToLocalMatrices.AddUninitialized(Float4Count);
+
+		Float4Count = 0;
+		for ( const FSkelMeshRenderSection& section : Sections )
+		{
+			const uint32 MatrixCount = section.BoneMap.Num();
+			for (uint32 m = 0; m < MatrixCount; ++m)
+			{
+				RefToLocalMatrices[section.BoneMap[m]].To3x4MatrixTranspose(&AllSectionsRefToLocalMatrices[Float4Count].X);
+				Float4Count += 3;
+			}
+		}
+
+		FSkeletalMeshGpuDynamicBufferProxy* ThisProxy = this;
+		ENQUEUE_RENDER_COMMAND(UpdateSpawnInfoForSkinnedMesh)(
+			[ThisProxy, AllSectionsRefToLocalMatrices = MoveTemp(AllSectionsRefToLocalMatrices)](FRHICommandListImmediate& RHICmdList) mutable
+		{
+			ThisProxy->CurrentBoneBufferId = (ThisProxy->CurrentBoneBufferId + 1) % BufferBoneCount;
+			ThisProxy->bPrevBoneGpuBufferValid = ThisProxy->bBoneGpuBufferValid;
+			ThisProxy->bBoneGpuBufferValid = true;
+
+			const uint32 NumBytes = AllSectionsRefToLocalMatrices.Num() * sizeof(FVector4);
+
+			void* DstData = RHILockVertexBuffer(ThisProxy->GetRWBufferBone().Buffer, 0, NumBytes, RLM_WriteOnly);
+			FMemory::Memcpy(DstData, AllSectionsRefToLocalMatrices.GetData(), NumBytes);
+			RHIUnlockVertexBuffer(ThisProxy->GetRWBufferBone().Buffer);
+		});
+	}
+}
+
+//////////////////////////////////////////////////////////////////////////
+//FNiagaraDataInterfaceParametersCS_SkeletalMesh
+
+struct FNDISkeletalMeshParametersName
+{
+	FString MeshIndexBufferName;
+	FString MeshVertexBufferName;
+	FString MeshSkinWeightBufferName;
+	FString MeshCurrBonesBufferName;
+	FString MeshPrevBonesBufferName;
+	FString MeshTangentBufferName;
+	FString MeshTexCoordBufferName;
+	FString MeshTriangleSamplerProbaBufferName;
+	FString MeshTriangleSamplerAliasBufferName;
+	FString MeshTriangleMatricesOffsetBufferName;
+	FString MeshTriangleCountName;
+	FString MeshWeightStrideByteName;
+	FString InstanceTransformName;
+	FString InstancePrevTransformName;
+	FString InstanceInvDeltaTimeName;
+	FString EnabledFeaturesName;
+	FString InputWeightStrideName;
+	FString NumTexCoordName;
+};
+static void GetNiagaraDataInterfaceParametersName(FNDISkeletalMeshParametersName& Names, const FString& Suffix)
+{
+	Names.MeshIndexBufferName = UNiagaraDataInterfaceSkeletalMesh::MeshIndexBufferName + Suffix;
+	Names.MeshVertexBufferName = UNiagaraDataInterfaceSkeletalMesh::MeshVertexBufferName + Suffix;
+	Names.MeshSkinWeightBufferName = UNiagaraDataInterfaceSkeletalMesh::MeshSkinWeightBufferName + Suffix;
+	Names.MeshCurrBonesBufferName = UNiagaraDataInterfaceSkeletalMesh::MeshCurrBonesBufferName + Suffix;
+	Names.MeshPrevBonesBufferName = UNiagaraDataInterfaceSkeletalMesh::MeshPrevBonesBufferName + Suffix;
+	Names.MeshTangentBufferName = UNiagaraDataInterfaceSkeletalMesh::MeshTangentBufferName + Suffix;
+	Names.MeshTexCoordBufferName = UNiagaraDataInterfaceSkeletalMesh::MeshTexCoordBufferName + Suffix;
+	Names.MeshTriangleSamplerProbaBufferName = UNiagaraDataInterfaceSkeletalMesh::MeshTriangleSamplerProbaBufferName + Suffix;
+	Names.MeshTriangleSamplerAliasBufferName = UNiagaraDataInterfaceSkeletalMesh::MeshTriangleSamplerAliasBufferName + Suffix;
+	Names.MeshTriangleMatricesOffsetBufferName = UNiagaraDataInterfaceSkeletalMesh::MeshTriangleMatricesOffsetBufferName + Suffix;
+	Names.MeshTriangleCountName = UNiagaraDataInterfaceSkeletalMesh::MeshTriangleCountName + Suffix;
+	Names.MeshWeightStrideByteName = UNiagaraDataInterfaceSkeletalMesh::MeshWeightStrideByteName + Suffix;
+	Names.InstanceTransformName = UNiagaraDataInterfaceSkeletalMesh::InstanceTransformName + Suffix;
+	Names.InstancePrevTransformName = UNiagaraDataInterfaceSkeletalMesh::InstancePrevTransformName + Suffix;
+	Names.InstanceInvDeltaTimeName = UNiagaraDataInterfaceSkeletalMesh::InstanceInvDeltaTimeName + Suffix;
+	Names.EnabledFeaturesName = UNiagaraDataInterfaceSkeletalMesh::EnabledFeaturesName + Suffix;
+	Names.InputWeightStrideName = UNiagaraDataInterfaceSkeletalMesh::InputWeightStrideName + Suffix;
+	Names.NumTexCoordName = UNiagaraDataInterfaceSkeletalMesh::NumTexCoordName + Suffix;
+}
+
+struct FNiagaraDataInterfaceParametersCS_SkeletalMesh : public FNiagaraDataInterfaceParametersCS
+{
+	virtual void Bind(const FNiagaraDataInterfaceParamRef& ParamRef, const class FShaderParameterMap& ParameterMap) override
+	{
+		FNDISkeletalMeshParametersName ParamNames;
+		GetNiagaraDataInterfaceParametersName(ParamNames, ParamRef.ParameterInfo.DataInterfaceHLSLSymbol);
+
+		MeshIndexBuffer.Bind(ParameterMap, *ParamNames.MeshIndexBufferName);
+		MeshVertexBuffer.Bind(ParameterMap, *ParamNames.MeshVertexBufferName);
+		MeshSkinWeightBuffer.Bind(ParameterMap, *ParamNames.MeshSkinWeightBufferName);
+		MeshCurrBonesBuffer.Bind(ParameterMap, *ParamNames.MeshCurrBonesBufferName);
+		MeshPrevBonesBuffer.Bind(ParameterMap, *ParamNames.MeshPrevBonesBufferName);
+		MeshTangentBuffer.Bind(ParameterMap, *ParamNames.MeshTangentBufferName);
+		MeshTexCoordBuffer.Bind(ParameterMap, *ParamNames.MeshTexCoordBufferName);
+		MeshTriangleSamplerProbaBuffer.Bind(ParameterMap, *ParamNames.MeshTriangleSamplerProbaBufferName);
+		MeshTriangleSamplerAliasBuffer.Bind(ParameterMap, *ParamNames.MeshTriangleSamplerAliasBufferName);
+		MeshTriangleMatricesOffsetBuffer.Bind(ParameterMap, *ParamNames.MeshTriangleMatricesOffsetBufferName);
+		MeshTriangleCount.Bind(ParameterMap, *ParamNames.MeshTriangleCountName);
+		MeshWeightStrideByte.Bind(ParameterMap, *ParamNames.MeshWeightStrideByteName);
+		InstanceTransform.Bind(ParameterMap, *ParamNames.InstanceTransformName);
+		InstancePrevTransform.Bind(ParameterMap, *ParamNames.InstancePrevTransformName);
+		InstanceInvDeltaTime.Bind(ParameterMap, *ParamNames.InstanceInvDeltaTimeName);
+		EnabledFeatures.Bind(ParameterMap, *ParamNames.EnabledFeaturesName);
+		InputWeightStride.Bind(ParameterMap, *ParamNames.InputWeightStrideName);
+		NumTexCoord.Bind(ParameterMap, *ParamNames.NumTexCoordName);
+
+		if (!MeshIndexBuffer.IsBound())
+		{
+			UE_LOG(LogNiagara, Warning, TEXT("Binding failed for FNiagaraDataInterfaceParametersCS_StaticMesh Texture %s. Was it optimized out?"), *ParamNames.MeshIndexBufferName)
+		}
+		if (!MeshVertexBuffer.IsBound())
+		{
+			UE_LOG(LogNiagara, Warning, TEXT("Binding failed for FNiagaraDataInterfaceParametersCS_StaticMesh Sampler %s. Was it optimized out?"), *ParamNames.MeshVertexBufferName)
+		}
+		if (!MeshSkinWeightBuffer.IsBound())
+		{
+			UE_LOG(LogNiagara, Warning, TEXT("Binding failed for FNiagaraDataInterfaceParametersCS_StaticMesh Sampler %s. Was it optimized out?"), *ParamNames.MeshSkinWeightBufferName)
+		}
+		if (!MeshCurrBonesBuffer.IsBound())
+		{
+			UE_LOG(LogNiagara, Warning, TEXT("Binding failed for FNiagaraDataInterfaceParametersCS_StaticMesh Sampler %s. Was it optimized out?"), *ParamNames.MeshCurrBonesBufferName)
+		}
+		if (!MeshTangentBuffer.IsBound())
+		{
+			UE_LOG(LogNiagara, Warning, TEXT("Binding failed for FNiagaraDataInterfaceParametersCS_StaticMesh Sampler %s. Was it optimized out?"), *ParamNames.MeshTangentBufferName)
+		}
+		if (!MeshTriangleMatricesOffsetBuffer.IsBound())
+		{
+			UE_LOG(LogNiagara, Warning, TEXT("Binding failed for FNiagaraDataInterfaceParametersCS_StaticMesh Sampler %s. Was it optimized out?"), *ParamNames.MeshTriangleMatricesOffsetBufferName)
+		}
+	}
+
+	virtual void Serialize(FArchive& Ar)override
+	{
+		Ar << MeshIndexBuffer;
+		Ar << MeshVertexBuffer;
+		Ar << MeshSkinWeightBuffer;
+		Ar << MeshCurrBonesBuffer;
+		Ar << MeshPrevBonesBuffer;
+		Ar << MeshTangentBuffer;
+		Ar << MeshTexCoordBuffer;
+		Ar << MeshTriangleSamplerProbaBuffer;
+		Ar << MeshTriangleSamplerAliasBuffer;
+		Ar << MeshTriangleMatricesOffsetBuffer;
+		Ar << MeshTriangleCount;
+		Ar << MeshWeightStrideByte;
+		Ar << InstanceTransform;
+		Ar << InstancePrevTransform;
+		Ar << InstanceInvDeltaTime;
+		Ar << EnabledFeatures;
+		Ar << InputWeightStride;
+		Ar << NumTexCoord;
+	}
+
+	virtual void Set(FRHICommandList& RHICmdList, const FNiagaraDataInterfaceSetArgs& Context) const override
+	{
+		check(IsInRenderingThread());
+
+		FComputeShaderRHIParamRef ComputeShaderRHI = Context.Shader->GetComputeShader();
+		FNiagaraDataInterfaceProxySkeletalMesh* InterfaceProxy = static_cast<FNiagaraDataInterfaceProxySkeletalMesh*>(Context.DataInterface);
+		FNiagaraDataInterfaceProxySkeletalMeshData* InstanceData = InterfaceProxy->SystemInstancesToData.Find(Context.SystemInstance);
+		if (InstanceData && InstanceData->StaticBuffers)
+		{
+			FSkeletalMeshGpuSpawnStaticBuffers* StaticBuffers = InstanceData->StaticBuffers;
+
+			SetSRVParameter(RHICmdList, ComputeShaderRHI, MeshVertexBuffer, StaticBuffers->GetBufferPositionSRV());
+			SetSRVParameter(RHICmdList, ComputeShaderRHI, MeshIndexBuffer, StaticBuffers->GetBufferIndexSRV());
+			SetSRVParameter(RHICmdList, ComputeShaderRHI, MeshTangentBuffer, StaticBuffers->GetBufferTangentSRV());
+
+			SetShaderValue(RHICmdList, ComputeShaderRHI, NumTexCoord, StaticBuffers->GetNumTexCoord());
+			if (StaticBuffers->GetNumTexCoord() > 0)
+			{
+				SetSRVParameter(RHICmdList, ComputeShaderRHI, MeshTexCoordBuffer, StaticBuffers->GetBufferTexCoordSRV());
+			}
+			else
+			{
+				SetSRVParameter(RHICmdList, ComputeShaderRHI, MeshTexCoordBuffer, FNiagaraRenderer::GetDummyFloatBuffer().SRV);
+			}
+			SetShaderValue(RHICmdList, ComputeShaderRHI, MeshTriangleCount, StaticBuffers->GetTriangleCount());
+			if (InstanceData->bIsGpuUniformlyDistributedSampling)
+			{
+				SetSRVParameter(RHICmdList, ComputeShaderRHI, MeshTriangleSamplerProbaBuffer, StaticBuffers->GetBufferTriangleUniformSamplerProbaSRV().GetReference());
+				SetSRVParameter(RHICmdList, ComputeShaderRHI, MeshTriangleSamplerAliasBuffer, StaticBuffers->GetBufferTriangleUniformSamplerAliasSRV().GetReference());
+			}
+			else
+			{
+				SetSRVParameter(RHICmdList, ComputeShaderRHI, MeshTriangleSamplerProbaBuffer, FNiagaraRenderer::GetDummyUIntBuffer().SRV);
+				SetSRVParameter(RHICmdList, ComputeShaderRHI, MeshTriangleSamplerAliasBuffer, FNiagaraRenderer::GetDummyUIntBuffer().SRV);
+			}
+
+			SetSRVParameter(RHICmdList, ComputeShaderRHI, MeshSkinWeightBuffer, InstanceData->MeshSkinWeightBufferSrv);
+
+			SetShaderValue(RHICmdList, ComputeShaderRHI, MeshWeightStrideByte, InstanceData->MeshWeightStrideByte);
+			SetShaderValue(RHICmdList, ComputeShaderRHI, InstanceTransform, InstanceData->Transform);
+			SetShaderValue(RHICmdList, ComputeShaderRHI, InstancePrevTransform, InstanceData->PrevTransform);
+			SetShaderValue(RHICmdList, ComputeShaderRHI, InstanceInvDeltaTime, 1.0f / InstanceData->DeltaSeconds);
+			SetShaderValue(RHICmdList, ComputeShaderRHI, InputWeightStride, InstanceData->MeshWeightStrideByte/4);
+
+			uint32 EnabledFeaturesBits = InstanceData->bIsGpuUniformlyDistributedSampling ? 1 : 0;
+
+			FSkeletalMeshGpuDynamicBufferProxy* DynamicBuffers = InstanceData->DynamicBuffer;
+			check(DynamicBuffers);
+			if(DynamicBuffers->DoesBoneDataExist())
+			{
+				EnabledFeaturesBits |= 2; // Enable the skinning feature
+				SetSRVParameter(RHICmdList, ComputeShaderRHI, MeshCurrBonesBuffer, DynamicBuffers->GetRWBufferBone().SRV);
+				SetSRVParameter(RHICmdList, ComputeShaderRHI, MeshPrevBonesBuffer, DynamicBuffers->GetRWBufferPrevBone().SRV);
+				SetSRVParameter(RHICmdList, ComputeShaderRHI, MeshTriangleMatricesOffsetBuffer, StaticBuffers->GetBufferTriangleMatricesOffsetSRV());
+			}
+			// Bind dummy data for validation purposes only.  Code will not execute due to "EnabledFeatures" bits but validation can not determine that.
+			else
+			{
+				SetSRVParameter(RHICmdList, ComputeShaderRHI, MeshCurrBonesBuffer, FNiagaraRenderer::GetDummyFloat4Buffer().SRV);
+				SetSRVParameter(RHICmdList, ComputeShaderRHI, MeshPrevBonesBuffer, FNiagaraRenderer::GetDummyFloat4Buffer().SRV);
+				SetSRVParameter(RHICmdList, ComputeShaderRHI, MeshTriangleMatricesOffsetBuffer, FNiagaraRenderer::GetDummyUIntBuffer().SRV);
+			}
+
+			SetShaderValue(RHICmdList, ComputeShaderRHI, EnabledFeatures, EnabledFeaturesBits);
+		}
+		else
+		{
+			// Bind dummy buffers
+			ensure(!InstanceData || InstanceData->StaticBuffers);
+
+			SetSRVParameter(RHICmdList, ComputeShaderRHI, MeshVertexBuffer, FNiagaraRenderer::GetDummyFloatBuffer().SRV);
+			SetSRVParameter(RHICmdList, ComputeShaderRHI, MeshIndexBuffer, FNiagaraRenderer::GetDummyUIntBuffer().SRV);
+			SetSRVParameter(RHICmdList, ComputeShaderRHI, MeshTangentBuffer, FNiagaraRenderer::GetDummyFloatBuffer().SRV);
+
+			SetShaderValue(RHICmdList, ComputeShaderRHI, NumTexCoord, 0);
+			SetSRVParameter(RHICmdList, ComputeShaderRHI, MeshTexCoordBuffer, FNiagaraRenderer::GetDummyFloatBuffer().SRV);
+			SetShaderValue(RHICmdList, ComputeShaderRHI, MeshTriangleCount, 0);
+			SetSRVParameter(RHICmdList, ComputeShaderRHI, MeshTriangleSamplerProbaBuffer, FNiagaraRenderer::GetDummyUIntBuffer().SRV);
+			SetSRVParameter(RHICmdList, ComputeShaderRHI, MeshTriangleSamplerAliasBuffer, FNiagaraRenderer::GetDummyUIntBuffer().SRV);
+
+			SetSRVParameter(RHICmdList, ComputeShaderRHI, MeshSkinWeightBuffer, FNiagaraRenderer::GetDummyUIntBuffer().SRV);
+
+			SetShaderValue(RHICmdList, ComputeShaderRHI, MeshWeightStrideByte, 0);
+			SetShaderValue(RHICmdList, ComputeShaderRHI, InstanceTransform, FMatrix::Identity);
+			SetShaderValue(RHICmdList, ComputeShaderRHI, InstancePrevTransform, FMatrix::Identity);
+			SetShaderValue(RHICmdList, ComputeShaderRHI, InstanceInvDeltaTime, 0.0f);
+			SetShaderValue(RHICmdList, ComputeShaderRHI, InputWeightStride, 0);
+
+			SetSRVParameter(RHICmdList, ComputeShaderRHI, MeshCurrBonesBuffer, FNiagaraRenderer::GetDummyFloat4Buffer().SRV);
+			SetSRVParameter(RHICmdList, ComputeShaderRHI, MeshPrevBonesBuffer, FNiagaraRenderer::GetDummyFloat4Buffer().SRV);
+			SetSRVParameter(RHICmdList, ComputeShaderRHI, MeshTriangleMatricesOffsetBuffer, FNiagaraRenderer::GetDummyUIntBuffer().SRV);
+
+			SetShaderValue(RHICmdList, ComputeShaderRHI, EnabledFeatures, 0);
+		}
+	}
+
+
+private:
+
+	FShaderResourceParameter MeshIndexBuffer;
+	FShaderResourceParameter MeshVertexBuffer;
+	FShaderResourceParameter MeshSkinWeightBuffer;
+	FShaderResourceParameter MeshCurrBonesBuffer;
+	FShaderResourceParameter MeshPrevBonesBuffer;
+	FShaderResourceParameter MeshTangentBuffer;
+	FShaderResourceParameter MeshTexCoordBuffer;
+	FShaderResourceParameter MeshTriangleSamplerProbaBuffer;
+	FShaderResourceParameter MeshTriangleSamplerAliasBuffer;
+	FShaderResourceParameter MeshTriangleMatricesOffsetBuffer;
+	FShaderParameter MeshTriangleCount;
+	FShaderParameter MeshWeightStrideByte;
+	FShaderParameter InstanceTransform;
+	FShaderParameter InstancePrevTransform;
+	FShaderParameter InstanceInvDeltaTime;
+	FShaderParameter EnabledFeatures;
+	FShaderParameter InputWeightStride;
+	FShaderParameter NumTexCoord;
+};
+
+//////////////////////////////////////////////////////////////////////////
+
+void FNiagaraDataInterfaceProxySkeletalMesh::ConsumePerInstanceDataFromGameThread(void* PerInstanceData, const FGuid& Instance)
+{
+	FNiagaraDISkeletalMeshPassedDataToRT* SourceData = static_cast<FNiagaraDISkeletalMeshPassedDataToRT*>(PerInstanceData);
+
+	FNiagaraDataInterfaceProxySkeletalMeshData& Data = SystemInstancesToData.FindOrAdd(Instance);
+
+	Data.bIsGpuUniformlyDistributedSampling = SourceData->bIsGpuUniformlyDistributedSampling;
+	Data.DeltaSeconds = SourceData->DeltaSeconds;
+	Data.DynamicBuffer = SourceData->DynamicBuffer;
+	Data.MeshWeightStrideByte = SourceData->MeshWeightStrideByte;
+	Data.PrevTransform = SourceData->PrevTransform;
+	Data.StaticBuffers = SourceData->StaticBuffers;
+	Data.Transform = SourceData->Transform;
+
+	// @todo-threadsafety race here. Need to hold a ref to this buffer on the RT
+	Data.MeshSkinWeightBufferSrv = SourceData->MeshSkinWeightBufferSrv;
+}
+
 //////////////////////////////////////////////////////////////////////////
 //FNDISkeletalMesh_InstanceData
 
+void UNiagaraDataInterfaceSkeletalMesh::ProvidePerInstanceDataForRenderThread(void* DataForRenderThread, void* PerInstanceData, const FGuid& SystemInstance)
+{
+	FNiagaraDISkeletalMeshPassedDataToRT* Data = static_cast<FNiagaraDISkeletalMeshPassedDataToRT*>(DataForRenderThread);
+	FNDISkeletalMesh_InstanceData* SourceData = static_cast<FNDISkeletalMesh_InstanceData*>(PerInstanceData);
+
+	Data->bIsGpuUniformlyDistributedSampling = SourceData->bIsGpuUniformlyDistributedSampling;
+	Data->DeltaSeconds = SourceData->DeltaSeconds;
+	Data->DynamicBuffer = SourceData->MeshGpuSpawnDynamicBuffers;
+	Data->MeshWeightStrideByte = SourceData->MeshWeightStrideByte;
+	Data->PrevTransform = SourceData->PrevTransform;
+	Data->StaticBuffers = SourceData->MeshGpuSpawnStaticBuffers;
+	Data->Transform = SourceData->Transform;
+
+	// @todo-threadsafety race here. Need to hold a ref to this buffer on the RT
+	Data->MeshSkinWeightBufferSrv = SourceData->MeshSkinWeightBufferSrv;
+}
+
 USkeletalMesh* UNiagaraDataInterfaceSkeletalMesh::GetSkeletalMeshHelper(UNiagaraDataInterfaceSkeletalMesh* Interface, UNiagaraComponent* OwningComponent, TWeakObjectPtr<USceneComponent>& SceneComponent, USkeletalMeshComponent*& FoundSkelComp)
 {
 	USkeletalMesh* Mesh = nullptr;
-	if (Interface->Source)
+	if (Interface->SourceComponent)
+	{
+		Mesh = Interface->SourceComponent->SkeletalMesh;
+		FoundSkelComp = Interface->SourceComponent;
+	}
+	else if (Interface->Source)
 	{
 		ASkeletalMeshActor* MeshActor = Cast<ASkeletalMeshActor>(Interface->Source);
 		USkeletalMeshComponent* SourceComp = nullptr;
@@ -636,6 +1106,32 @@ bool FNDISkeletalMesh_InstanceData::Init(UNiagaraDataInterfaceSkeletalMesh* Inte
 		}
 	}
 
+	{
+		MeshWeightStrideByte = SkinWeightBuffer->GetStride();
+		MeshSkinWeightBufferSrv = SkinWeightBuffer->GetSRV();
+		//check(MeshSkinWeightBufferSrv->IsValid()); // not available in this stream
+
+		FSkeletalMeshLODInfo* LODInfo = Mesh->GetLODInfo(LODIndex);
+		bIsGpuUniformlyDistributedSampling = LODInfo->bSupportUniformlyDistributedSampling && bAllRegionsAreAreaWeighting;
+
+		if (Mesh->HasActiveClothingAssets())
+		{
+			UE_LOG(LogNiagara, Warning, TEXT("Skeletal Mesh %s has cloth asset on it: spawning from it might not work properly."), *Mesh->GetName());
+		}
+		if (LODData.DoesVertexBufferHaveExtraBoneInfluences())
+		{
+			UE_LOG(LogNiagara, Warning, TEXT("Skeletal Mesh %s has bones extra influence: spawning from it might not work properly."), *Mesh->GetName());
+		}
+
+		MeshGpuSpawnStaticBuffers = new FSkeletalMeshGpuSpawnStaticBuffers();
+		MeshGpuSpawnStaticBuffers->Initialise(LODData, bIsGpuUniformlyDistributedSampling, SamplingInfo.GetBuiltData().WholeMeshBuiltData[LODIndex]);
+		BeginInitResource(MeshGpuSpawnStaticBuffers);
+
+		MeshGpuSpawnDynamicBuffers = new FSkeletalMeshGpuDynamicBufferProxy();
+		MeshGpuSpawnDynamicBuffers->Initialise(LODData);
+		BeginInitResource(MeshGpuSpawnDynamicBuffers);
+	}
+
 	return true;
 }
 
@@ -706,6 +1202,19 @@ bool FNDISkeletalMesh_InstanceData::Tick(UNiagaraDataInterfaceSkeletalMesh* Inte
 			Transform = FMatrix::Identity;
 			TransformInverseTransposed = FMatrix::Identity;
 		}
+
+		if (MeshGpuSpawnDynamicBuffers)
+		{
+			USkeletalMeshComponent* Comp = Cast<USkeletalMeshComponent>(Component.Get());
+			const USkinnedMeshComponent* BaseComp = nullptr;
+			if (Comp)
+			{
+				BaseComp = Comp->GetBaseComponent();
+			}
+
+			MeshGpuSpawnDynamicBuffers->NewFrame(this, GetLODIndex());
+		}
+
 		return false;
 	}
 }
@@ -719,6 +1228,30 @@ bool FNDISkeletalMesh_InstanceData::HasColorData()
 	return LODData.StaticVertexBuffers.ColorVertexBuffer.GetNumVertices() != 0;
 }
 
+void FNDISkeletalMesh_InstanceData::Release()
+{
+	if (MeshGpuSpawnStaticBuffers)
+	{
+		BeginReleaseResource(MeshGpuSpawnStaticBuffers);
+		ENQUEUE_RENDER_COMMAND(DeleteResource)(
+			[ParamPointerToRelease = MeshGpuSpawnStaticBuffers](FRHICommandListImmediate& RHICmdList)
+			{
+				delete ParamPointerToRelease;
+			});
+		MeshGpuSpawnStaticBuffers = nullptr;
+	}
+	if (MeshGpuSpawnDynamicBuffers)
+	{
+		BeginReleaseResource(MeshGpuSpawnDynamicBuffers);
+		ENQUEUE_RENDER_COMMAND(DeleteResource)(
+			[ParamPointerToRelease = MeshGpuSpawnDynamicBuffers](FRHICommandListImmediate& RHICmdList)
+			{
+				delete ParamPointerToRelease;
+			});
+		MeshGpuSpawnDynamicBuffers = nullptr;
+	}
+}
+
 //Instance Data END
 //////////////////////////////////////////////////////////////////////////
 
@@ -737,10 +1270,9 @@ UNiagaraDataInterfaceSkeletalMesh::UNiagaraDataInterfaceSkeletalMesh(FObjectInit
 	, bUseSkeletonSampling(true)
 	, ChangeId(0)
 {
-
+	Proxy = MakeShared<FNiagaraDataInterfaceProxySkeletalMesh, ESPMode::ThreadSafe>();
 }
 
-#if WITH_EDITOR
 
 void UNiagaraDataInterfaceSkeletalMesh::PostInitProperties()
 {
@@ -755,6 +1287,8 @@ void UNiagaraDataInterfaceSkeletalMesh::PostInitProperties()
 		FNiagaraTypeRegistry::Register(FMeshTriCoordinate::StaticStruct(), true, true, false);
 	}
 }
+
+#if WITH_EDITOR
 void UNiagaraDataInterfaceSkeletalMesh::PostEditChangeProperty(FPropertyChangedEvent& PropertyChangedEvent)
 {
 	Super::PostEditChangeProperty(PropertyChangedEvent);
@@ -889,7 +1423,19 @@ void UNiagaraDataInterfaceSkeletalMesh::DestroyPerInstanceData(void* PerInstance
 	}
 #endif
 
+	Inst->Release();
 	Inst->~FNDISkeletalMesh_InstanceData();
+
+	{
+		// @todo this races
+		FNiagaraDataInterfaceProxySkeletalMesh* ThisProxy = GetProxyAs<FNiagaraDataInterfaceProxySkeletalMesh>();
+		ENQUEUE_RENDER_COMMAND(FNiagaraDestroySkeletalMeshInstanceData) (
+			[ThisProxy, InstanceID = SystemInstance->GetId()](FRHICommandListImmediate& CmdList)
+			{
+				ThisProxy->SystemInstancesToData.Remove(InstanceID);
+			}
+		);
+	}
 }
 
 bool UNiagaraDataInterfaceSkeletalMesh::PerInstanceTick(void* PerInstanceData, FNiagaraSystemInstance* SystemInstance, float InDeltaSeconds)
@@ -1065,6 +1611,360 @@ void UNiagaraDataInterfaceSkeletalMesh::ValidateFunction(const FNiagaraFunctionS
 
 #endif
 
+const FString UNiagaraDataInterfaceSkeletalMesh::MeshIndexBufferName(TEXT("IndexBuffer_"));
+const FString UNiagaraDataInterfaceSkeletalMesh::MeshVertexBufferName(TEXT("VertexBuffer_"));
+const FString UNiagaraDataInterfaceSkeletalMesh::MeshSkinWeightBufferName(TEXT("VertexSkinWeightBuffer_"));
+const FString UNiagaraDataInterfaceSkeletalMesh::MeshCurrBonesBufferName(TEXT("MeshCurrBonesBuffer_"));
+const FString UNiagaraDataInterfaceSkeletalMesh::MeshPrevBonesBufferName(TEXT("MeshPrevBonesBuffer_"));
+const FString UNiagaraDataInterfaceSkeletalMesh::MeshTangentBufferName(TEXT("MeshTangentBuffer_"));
+const FString UNiagaraDataInterfaceSkeletalMesh::MeshTexCoordBufferName(TEXT("TexCoordBuffer_"));
+const FString UNiagaraDataInterfaceSkeletalMesh::MeshTriangleSamplerProbaBufferName(TEXT("MeshTriangleSamplerProbaBuffer_"));
+const FString UNiagaraDataInterfaceSkeletalMesh::MeshTriangleSamplerAliasBufferName(TEXT("MeshTriangleSamplerAliasBuffer_"));
+const FString UNiagaraDataInterfaceSkeletalMesh::MeshTriangleMatricesOffsetBufferName(TEXT("MeshTriangleMatricesOffsetBuffer_"));
+const FString UNiagaraDataInterfaceSkeletalMesh::InstanceTransformName(TEXT("InstanceTransform_"));
+const FString UNiagaraDataInterfaceSkeletalMesh::InstancePrevTransformName(TEXT("InstancePrevTransform_"));
+const FString UNiagaraDataInterfaceSkeletalMesh::InstanceInvDeltaTimeName(TEXT("InstanceInvDeltaTime_"));
+const FString UNiagaraDataInterfaceSkeletalMesh::MeshWeightStrideByteName(TEXT("MeshWeightStrideByte_"));
+const FString UNiagaraDataInterfaceSkeletalMesh::MeshTriangleCountName(TEXT("MeshTriangleCount_"));
+const FString UNiagaraDataInterfaceSkeletalMesh::EnabledFeaturesName(TEXT("EnabledFeatures_"));
+const FString UNiagaraDataInterfaceSkeletalMesh::InputWeightStrideName(TEXT("InputWeightStride_"));
+const FString UNiagaraDataInterfaceSkeletalMesh::NumTexCoordName(TEXT("NumTexCoordName_"));
+
+bool UNiagaraDataInterfaceSkeletalMesh::GetFunctionHLSL(const FName&  DefinitionFunctionName, FString InstanceFunctionName, FNiagaraDataInterfaceGPUParamInfo& ParamInfo, FString& OutHLSL) 
+{
+	FNDISkeletalMeshParametersName ParamNames;
+	GetNiagaraDataInterfaceParametersName(ParamNames, ParamInfo.DataInterfaceHLSLSymbol);
+	FString MeshTriCoordinateStructName = "MeshTriCoordinate";
+
+	static const TCHAR* FormatCommonFunctions = TEXT(R"(
+		void {InstanceFunctionName}_GetIndicesAndWeights(uint VertexIndex, out int4 BlendIndices, out float4 BlendWeights)
+		{
+			uint PackedBlendIndices = {MeshSkinWeightBufferName}[VertexIndex * ({InputWeightStrideName})    ];
+			uint PackedBlendWeights = {MeshSkinWeightBufferName}[VertexIndex * ({InputWeightStrideName}) + 1];
+			BlendIndices.x = PackedBlendIndices & 0xff;
+			BlendIndices.y = PackedBlendIndices >> 8 & 0xff;
+			BlendIndices.z = PackedBlendIndices >> 16 & 0xff;
+			BlendIndices.w = PackedBlendIndices >> 24 & 0xff;
+			BlendWeights.x = float(PackedBlendWeights & 0xff) / 255.0f;
+			BlendWeights.y = float(PackedBlendWeights >> 8 & 0xff) / 255.0f;
+			BlendWeights.z = float(PackedBlendWeights >> 16 & 0xff) / 255.0f;
+			BlendWeights.w = float(PackedBlendWeights >> 24 & 0xff) / 255.0f;
+		}
+
+		float3x4 {InstanceFunctionName}_GetPrevBoneMatrix(uint Bone)
+		{
+			return float3x4({MeshPrevBonesBufferName}[Bone * 3], {MeshPrevBonesBufferName}[Bone * 3 + 1], {MeshPrevBonesBufferName}[Bone * 3 + 2]);
+		}
+
+		float3x4 {InstanceFunctionName}_GetPrevSkinningMatrix(uint VertexIndex, int4 BlendIndices, float4 BlendWeights)
+		{
+			// Get the matrix offset for each vertex because BlendIndices are stored relatively to each section start vertex.
+			uint MatrixOffset = {MeshTriangleMatricesOffsetBufferName}[VertexIndex];
+
+			float3x4 Result;
+			Result  = {InstanceFunctionName}_GetPrevBoneMatrix(MatrixOffset + BlendIndices.x) * BlendWeights.x;
+			Result += {InstanceFunctionName}_GetPrevBoneMatrix(MatrixOffset + BlendIndices.y) * BlendWeights.y;
+			Result += {InstanceFunctionName}_GetPrevBoneMatrix(MatrixOffset + BlendIndices.z) * BlendWeights.z;
+			Result += {InstanceFunctionName}_GetPrevBoneMatrix(MatrixOffset + BlendIndices.w) * BlendWeights.w;
+			return Result;
+		}
+
+		float3x4 {InstanceFunctionName}_GetCurrBoneMatrix(uint Bone)
+		{
+			return float3x4({MeshCurrBonesBufferName}[Bone * 3], {MeshCurrBonesBufferName}[Bone * 3 + 1], {MeshCurrBonesBufferName}[Bone * 3 + 2]);
+		}
+
+		float3x4 {InstanceFunctionName}_GetCurrSkinningMatrix(uint VertexIndex, int4 BlendIndices, float4 BlendWeights)
+		{
+			// Get the matrix offset for each vertex because BlendIndices are stored relatively to each section start vertex.
+			uint MatrixOffset = {MeshTriangleMatricesOffsetBufferName}[VertexIndex];
+
+			float3x4 Result;
+			Result  = {InstanceFunctionName}_GetCurrBoneMatrix(MatrixOffset + BlendIndices.x) * BlendWeights.x;
+			Result += {InstanceFunctionName}_GetCurrBoneMatrix(MatrixOffset + BlendIndices.y) * BlendWeights.y;
+			Result += {InstanceFunctionName}_GetCurrBoneMatrix(MatrixOffset + BlendIndices.z) * BlendWeights.z;
+			Result += {InstanceFunctionName}_GetCurrBoneMatrix(MatrixOffset + BlendIndices.w) * BlendWeights.w;
+			return Result;
+		}
+	)");
+
+	static const TCHAR* FormatSampleSkinnedTriangleDataWSHeader = TEXT(R"(
+		void {InstanceFunctionName} (in {MeshTriCoordinateStructName} In_Coord, out float3 Out_Position, out float3 Out_Velocity, out float3 Out_Normal, out float3 Out_Binormal, out float3 Out_Tangent)
+		{
+			const float In_Interp = 1.0f;
+		)");
+
+	static const TCHAR* FormatSampleSkinnedTriangleDataWSInterpolatedHeader = TEXT(R"(
+		void {InstanceFunctionName} (in {MeshTriCoordinateStructName} In_Coord, float In_Interp, out float3 Out_Position, out float3 Out_Velocity, out float3 Out_Normal, out float3 Out_Binormal, out float3 Out_Tangent)
+		{
+		)");
+
+	static const TCHAR* FormatSampleSkinnedTriangleDataWSPart0 = TEXT(R"(
+			const bool SkinningEnabled = {EnabledFeaturesName} & 0x0002;
+
+			uint TriangleIndex = In_Coord.Tri * 3;
+			uint VertexIndex0 = {MeshIndexBufferName}[TriangleIndex  ];
+			uint VertexIndex1 = {MeshIndexBufferName}[TriangleIndex+1];
+			uint VertexIndex2 = {MeshIndexBufferName}[TriangleIndex+2];
+			
+			// I could not find a R32G32B32f format to create an SRV on that buffer. So float load it is for now...
+			float3 Vertex0 = float3({MeshVertexBufferName}[VertexIndex0*3], {MeshVertexBufferName}[VertexIndex0*3+1], {MeshVertexBufferName}[VertexIndex0*3+2]);
+			float3 Vertex1 = float3({MeshVertexBufferName}[VertexIndex1*3], {MeshVertexBufferName}[VertexIndex1*3+1], {MeshVertexBufferName}[VertexIndex1*3+2]);
+			float3 Vertex2 = float3({MeshVertexBufferName}[VertexIndex2*3], {MeshVertexBufferName}[VertexIndex2*3+1], {MeshVertexBufferName}[VertexIndex2*3+2]);
+			float3 PrevVertex0 = Vertex0;
+			float3 PrevVertex1 = Vertex1;
+			float3 PrevVertex2 = Vertex2;
+
+			float3 TangentX0 = TangentBias({MeshTangentBufferName}[VertexIndex0*2  ].xyz);
+			float4 TangentZ0 = TangentBias({MeshTangentBufferName}[VertexIndex0*2+1].xyzw);
+			float3 TangentX1 = TangentBias({MeshTangentBufferName}[VertexIndex1*2  ].xyz);
+			float4 TangentZ1 = TangentBias({MeshTangentBufferName}[VertexIndex1*2+1].xyzw);
+			float3 TangentX2 = TangentBias({MeshTangentBufferName}[VertexIndex2*2  ].xyz);
+			float4 TangentZ2 = TangentBias({MeshTangentBufferName}[VertexIndex2*2+1].xyzw);
+
+			if(SkinningEnabled)
+			{
+				int4 BlendIndices0;
+				int4 BlendIndices1;
+				int4 BlendIndices2;
+				float4 BlendWeights0;
+				float4 BlendWeights1;
+				float4 BlendWeights2;
+
+				{InstanceFunctionName}_GetIndicesAndWeights(VertexIndex0, BlendIndices0, BlendWeights0);
+				{InstanceFunctionName}_GetIndicesAndWeights(VertexIndex1, BlendIndices1, BlendWeights1);
+				{InstanceFunctionName}_GetIndicesAndWeights(VertexIndex2, BlendIndices2, BlendWeights2);
+
+				float3x4 PrevBoneMatrix0 = {InstanceFunctionName}_GetPrevSkinningMatrix(VertexIndex0, BlendIndices0, BlendWeights0);
+				float3x4 PrevBoneMatrix1 = {InstanceFunctionName}_GetPrevSkinningMatrix(VertexIndex1, BlendIndices1, BlendWeights1);
+				float3x4 PrevBoneMatrix2 = {InstanceFunctionName}_GetPrevSkinningMatrix(VertexIndex2, BlendIndices2, BlendWeights2);
+				PrevVertex0 = mul( PrevBoneMatrix0, float4(Vertex0, 1.0f) ).xyz;
+				PrevVertex1 = mul( PrevBoneMatrix1, float4(Vertex1, 1.0f) ).xyz;
+				PrevVertex2 = mul( PrevBoneMatrix2, float4(Vertex2, 1.0f) ).xyz;
+
+				float3x4 CurrBoneMatrix0 = {InstanceFunctionName}_GetCurrSkinningMatrix(VertexIndex0, BlendIndices0, BlendWeights0);
+				float3x4 CurrBoneMatrix1 = {InstanceFunctionName}_GetCurrSkinningMatrix(VertexIndex1, BlendIndices1, BlendWeights1);
+				float3x4 CurrBoneMatrix2 = {InstanceFunctionName}_GetCurrSkinningMatrix(VertexIndex2, BlendIndices2, BlendWeights2);
+				Vertex0 = mul( CurrBoneMatrix0, float4(Vertex0, 1.0f) ).xyz;
+				Vertex1 = mul( CurrBoneMatrix1, float4(Vertex1, 1.0f) ).xyz;
+				Vertex2 = mul( CurrBoneMatrix2, float4(Vertex2, 1.0f) ).xyz;
+
+				// Not using InverseTranspose of matrices so assuming uniform scaling only (same as SkinCache)
+				TangentX0.xyz = mul( CurrBoneMatrix0, float4(TangentX0.xyz, 0.0f) ).xyz;
+				TangentZ0.xyz = mul( CurrBoneMatrix0, float4(TangentZ0.xyz, 0.0f) ).xyz;
+				TangentX1.xyz = mul( CurrBoneMatrix1, float4(TangentX1.xyz, 0.0f) ).xyz;
+				TangentZ1.xyz = mul( CurrBoneMatrix1, float4(TangentZ1.xyz, 0.0f) ).xyz;
+				TangentX2.xyz = mul( CurrBoneMatrix2, float4(TangentX2.xyz, 0.0f) ).xyz;
+				TangentZ2.xyz = mul( CurrBoneMatrix2, float4(TangentZ2.xyz, 0.0f) ).xyz;
+			}
+
+			// Evaluate current and previous world position
+			float3 WSPos = Vertex0 * In_Coord.BaryCoord.x + Vertex1 * In_Coord.BaryCoord.y + Vertex2 * In_Coord.BaryCoord.z;
+			WSPos = mul(float4(WSPos,1.0), {InstanceTransformName}).xyz;
+			float3 PrevWSPos = PrevVertex0 * In_Coord.BaryCoord.x + PrevVertex1 * In_Coord.BaryCoord.y + PrevVertex2 * In_Coord.BaryCoord.z;
+			PrevWSPos = mul(float4(PrevWSPos,1.0), {InstancePrevTransformName}).xyz;
+
+			// Not using InverseTranspose of matrices so assuming uniform scaling only (same as SkinCache)
+			float3 Binormal0 = cross(TangentZ0.xyz, TangentX0.xyz) * TangentZ0.w;
+			float3 Binormal1 = cross(TangentZ1.xyz, TangentX1.xyz) * TangentZ1.w;
+			float3 Binormal2 = cross(TangentZ2.xyz, TangentX2.xyz) * TangentZ2.w;
+			float3 Normal   = TangentZ0.xyz * In_Coord.BaryCoord.x + TangentZ1.xyz * In_Coord.BaryCoord.y + TangentZ2.xyz * In_Coord.BaryCoord.z; // Normal is TangentZ
+			float3 Tangent  = TangentX0.xyz * In_Coord.BaryCoord.x + TangentX1.xyz * In_Coord.BaryCoord.y + TangentX2.xyz * In_Coord.BaryCoord.z;
+			float3 Binormal = Binormal0.xyz * In_Coord.BaryCoord.x + Binormal1.xyz * In_Coord.BaryCoord.y + Binormal2.xyz * In_Coord.BaryCoord.z;
+			float3 NormalWorld   = mul(float4(Normal  , 0.0), {InstanceTransformName}).xyz;
+			float3 TangentWorld  = mul(float4(Tangent , 0.0), {InstanceTransformName}).xyz;
+			float3 BinormalWorld = mul(float4(Binormal, 0.0), {InstanceTransformName}).xyz;
+			
+			Out_Position = lerp(PrevWSPos, WSPos, float3(In_Interp,In_Interp,In_Interp));
+			Out_Velocity = (WSPos - PrevWSPos) * {InstanceInvDeltaTimeName};				// Velocity is unafected by spawn interpolation. That would require another set of previous data.
+			Out_Normal   = normalize(NormalWorld);
+			Out_Tangent  = normalize(TangentWorld);
+			Out_Binormal = normalize(BinormalWorld);
+		}
+		)");
+
+
+	TMap<FString, FStringFormatArg> ArgsSample = {
+		{TEXT("InstanceFunctionName"), InstanceFunctionName},
+		{TEXT("MeshTriCoordinateStructName"), MeshTriCoordinateStructName},
+		{TEXT("MeshIndexBufferName"), ParamNames.MeshIndexBufferName},
+		{TEXT("MeshVertexBufferName"), ParamNames.MeshVertexBufferName},
+		{TEXT("MeshSkinWeightBufferName"), ParamNames.MeshSkinWeightBufferName},
+		{TEXT("MeshCurrBonesBufferName"), ParamNames.MeshCurrBonesBufferName},
+		{TEXT("MeshPrevBonesBufferName"), ParamNames.MeshPrevBonesBufferName},
+		{TEXT("MeshTangentBufferName"), ParamNames.MeshTangentBufferName},
+		{TEXT("MeshTexCoordBufferName"), ParamNames.MeshTexCoordBufferName},
+		{TEXT("MeshTriangleSamplerProbaBufferName"), ParamNames.MeshTriangleSamplerProbaBufferName},
+		{TEXT("MeshTriangleSamplerAliasBufferName"), ParamNames.MeshTriangleSamplerAliasBufferName},
+		{TEXT("MeshTriangleMatricesOffsetBufferName"), ParamNames.MeshTriangleMatricesOffsetBufferName},
+		{TEXT("MeshTriangleCountName"), ParamNames.MeshTriangleCountName},
+		{TEXT("InstanceTransformName"), ParamNames.InstanceTransformName},
+		{TEXT("InstancePrevTransformName"), ParamNames.InstancePrevTransformName},
+		{TEXT("InstanceInvDeltaTimeName"), ParamNames.InstanceInvDeltaTimeName},
+		{TEXT("EnabledFeaturesName"), ParamNames.EnabledFeaturesName},
+		{TEXT("InputWeightStrideName"), ParamNames.InputWeightStrideName},
+		{TEXT("NumTexCoordName"), ParamNames.NumTexCoordName},
+	};
+
+	if (DefinitionFunctionName == FSkeletalMeshInterfaceHelper::RandomTriCoordName)
+	{
+		static const TCHAR *FormatSample = TEXT(R"(
+			void {InstanceFunctionName} (out {MeshTriCoordinateStructName} Out_Coord)
+			{
+				const bool UniformTriangleSamplingEnable = {EnabledFeaturesName} & 0x0001;
+
+				float RandT0 = NiagaraInternalNoise(1, 2, 3);
+				[branch]
+				if (!UniformTriangleSamplingEnable)
+				{
+					// Uniform triangle id selection
+					Out_Coord.Tri = min(uint(RandT0*float({MeshTriangleCountName})), {MeshTriangleCountName}-1); // avoid % by using mul/min to Tri = MeshTriangleCountName
+				}
+				else
+				{
+					// Uniform area weighted position selection (using alias method from Alias method from FWeightedRandomSampler)
+					uint TriangleIndex = min(uint(RandT0*float({MeshTriangleCountName})), {MeshTriangleCountName}-1);
+					float TriangleProbability = {MeshTriangleSamplerProbaBufferName}[TriangleIndex];
+
+					// Alias check
+					float RandT1 = NiagaraInternalNoise(1, 2, 3);
+					if( RandT1 > TriangleProbability )
+					{
+						TriangleIndex = {MeshTriangleSamplerAliasBufferName}[TriangleIndex];
+					}
+					Out_Coord.Tri = TriangleIndex;
+				}
+
+				float r0 = NiagaraInternalNoise(1, 2, 3);
+				float r1 = NiagaraInternalNoise(1, 2, 3);
+				float sqrt0 = sqrt(r0);
+				float sqrt1 = sqrt(r1);
+				Out_Coord.BaryCoord = float3(1.0f - sqrt0, sqrt0 * (1.0 - r1), r1 * sqrt0);
+			//	Out_Coord.BaryCoord = float3(1.0f, 0.0f, 0.0f);
+			}
+			)");
+		OutHLSL += FString::Format(FormatSample, ArgsSample);
+	}
+	else if (DefinitionFunctionName == FSkeletalMeshInterfaceHelper::GetSkinnedTriangleDataWSName)
+	{
+		OutHLSL += FString::Format(FormatCommonFunctions, ArgsSample);
+		OutHLSL += FString::Format(FormatSampleSkinnedTriangleDataWSHeader, ArgsSample);
+		OutHLSL += FString::Format(FormatSampleSkinnedTriangleDataWSPart0, ArgsSample);
+	}
+	else if (DefinitionFunctionName == FSkeletalMeshInterfaceHelper::GetSkinnedTriangleDataWSInterpName)
+	{
+		OutHLSL += FString::Format(FormatCommonFunctions, ArgsSample);
+		OutHLSL += FString::Format(FormatSampleSkinnedTriangleDataWSInterpolatedHeader, ArgsSample);
+		OutHLSL += FString::Format(FormatSampleSkinnedTriangleDataWSPart0, ArgsSample);
+	}
+	else if (DefinitionFunctionName == FSkeletalMeshInterfaceHelper::GetTriColorName)
+	{
+		static const TCHAR *FormatSample = TEXT(R"(
+				void {InstanceFunctionName} (in {MeshTriCoordinateStructName} In_Coord, out float4 Out_Color)
+				{
+					Out_Color = 0.0f;
+				}
+				)");
+		OutHLSL += FString::Format(FormatSample, ArgsSample);
+	}
+	else if (DefinitionFunctionName == FSkeletalMeshInterfaceHelper::GetTriUVName)
+	{
+		static const TCHAR *FormatSample = TEXT(R"(
+			void {InstanceFunctionName} (in {MeshTriCoordinateStructName} In_Coord, in int In_UVSet, out float2 Out_UV)
+			{
+				if({NumTexCoordName}>0)
+				{
+					uint TriangleIndex = In_Coord.Tri * 3;
+					uint VertexIndex0 = {MeshIndexBufferName}[TriangleIndex  ];
+					uint VertexIndex1 = {MeshIndexBufferName}[TriangleIndex+1];
+					uint VertexIndex2 = {MeshIndexBufferName}[TriangleIndex+2];
+
+					uint stride = {NumTexCoordName};
+					uint SelectedUVSet = clamp(In_UVSet, 0, {NumTexCoordName}-1);
+					float2 UV0 = {MeshTexCoordBufferName}[VertexIndex0 * stride + SelectedUVSet];
+					float2 UV1 = {MeshTexCoordBufferName}[VertexIndex1 * stride + SelectedUVSet];
+					float2 UV2 = {MeshTexCoordBufferName}[VertexIndex2 * stride + SelectedUVSet];
+
+					Out_UV = UV0 * In_Coord.BaryCoord.x + UV1 * In_Coord.BaryCoord.y + UV2 * In_Coord.BaryCoord.z;
+				}
+				else	
+				{
+					Out_UV = 0.0f;
+				}
+			}
+			)");
+		OutHLSL += FString::Format(FormatSample, ArgsSample);
+	}
+	//else if (DefinitionFunctionName == FSkeletalMeshInterfaceHelper::IsValidTriCoordName)
+	//{
+	//	// unimplemented();
+	//}
+	//else if (DefinitionFunctionName == FSkeletalMeshInterfaceHelper::GetSkinnedTriangleDataName)
+	//{
+	//	// unimplemented();
+	//}
+	//else if (DefinitionFunctionName == FSkeletalMeshInterfaceHelper::GetSkinnedTriangleDataInterpName)
+	//{
+	//	// unimplemented();
+	//}
+	//else if (DefinitionFunctionName == FSkeletalMeshInterfaceHelper::GetTriangleCountName)
+	//{
+	//	// unimplemented();
+	//}
+	//else if (DefinitionFunctionName == FSkeletalMeshInterfaceHelper::GetTriangleAtName)
+	//{
+	//	// unimplemented();
+	//}
+	//else if (DefinitionFunctionName == FSkeletalMeshInterfaceHelper::GetTriCoordVerticesName)
+	//{
+	//}
+	else
+	{
+		// This function is not support
+		return false;
+	}
+
+	OutHLSL += TEXT("\n");
+	return true;
+
+}
+void UNiagaraDataInterfaceSkeletalMesh::GetParameterDefinitionHLSL(FNiagaraDataInterfaceGPUParamInfo& ParamInfo, FString& OutHLSL)
+{
+	FNDISkeletalMeshParametersName ParamNames;
+	GetNiagaraDataInterfaceParametersName(ParamNames, ParamInfo.DataInterfaceHLSLSymbol);
+
+	OutHLSL += TEXT("Buffer<uint> ") + ParamNames.MeshIndexBufferName + TEXT(";\n");
+	OutHLSL += TEXT("Buffer<float> ") + ParamNames.MeshVertexBufferName + TEXT(";\n");
+	OutHLSL += TEXT("Buffer<uint> ") + ParamNames.MeshSkinWeightBufferName + TEXT(";\n");
+	OutHLSL += TEXT("Buffer<float4> ") + ParamNames.MeshCurrBonesBufferName + TEXT(";\n");
+	OutHLSL += TEXT("Buffer<float4> ") + ParamNames.MeshPrevBonesBufferName + TEXT(";\n");
+	OutHLSL += TEXT("Buffer<float4> ") + ParamNames.MeshTangentBufferName + TEXT(";\n");
+	OutHLSL += TEXT("Buffer<float2> ") + ParamNames.MeshTexCoordBufferName + TEXT(";\n");
+	OutHLSL += TEXT("Buffer<float> ") + ParamNames.MeshTriangleSamplerProbaBufferName + TEXT(";\n");
+	OutHLSL += TEXT("Buffer<uint> ") + ParamNames.MeshTriangleSamplerAliasBufferName + TEXT(";\n");
+	OutHLSL += TEXT("Buffer<uint> ") + ParamNames.MeshTriangleMatricesOffsetBufferName + TEXT(";\n");
+	OutHLSL += TEXT("uint ") + ParamNames.MeshTriangleCountName + TEXT(";\n");
+	OutHLSL += TEXT("float4x4 ") + ParamNames.InstanceTransformName + TEXT(";\n");
+	OutHLSL += TEXT("float4x4 ") + ParamNames.InstancePrevTransformName + TEXT(";\n");
+	OutHLSL += TEXT("float ") + ParamNames.InstanceInvDeltaTimeName + TEXT(";\n");
+	OutHLSL += TEXT("uint ") + ParamNames.EnabledFeaturesName + TEXT(";\n");
+	OutHLSL += TEXT("uint ") + ParamNames.InputWeightStrideName + TEXT(";\n");
+	OutHLSL += TEXT("uint ") + ParamNames.NumTexCoordName + TEXT(";\n");
+}
+FNiagaraDataInterfaceParametersCS* UNiagaraDataInterfaceSkeletalMesh::ConstructComputeParameters() const
+{
+	return new FNiagaraDataInterfaceParametersCS_SkeletalMesh();
+}
+
+
+void UNiagaraDataInterfaceSkeletalMesh::SetSourceComponentFromBlueprints(USkeletalMeshComponent* ComponentToUse)
+{
+	// NOTE: When ChangeId changes the next tick will be skipped and a reset of the per-instance data will be initiated. 
+	ChangeId++;
+	SourceComponent = ComponentToUse;
+	Source = ComponentToUse->GetOwner();
+}
+
 //UNiagaraDataInterfaceSkeletalMesh END
 //////////////////////////////////////////////////////////////////////////
 
