Copying //UE4/Dev-Niagara to Dev-Main (//UE4/Dev-Main)

diff --git a/Engine/Plugins/FX/Niagara/Source/Niagara/Private/NiagaraDataInterfaceStaticMesh.cpp b/Engine/Plugins/FX/Niagara/Source/Niagara/Private/NiagaraDataInterfaceStaticMesh.cpp
--- a/Engine/Plugins/FX/Niagara/Source/Niagara/Private/NiagaraDataInterfaceStaticMesh.cpp
+++ b/Engine/Plugins/FX/Niagara/Source/Niagara/Private/NiagaraDataInterfaceStaticMesh.cpp
@@ -5,11 +5,30 @@
 #include "Engine/StaticMeshActor.h"
 #include "NiagaraComponent.h"
 #include "NiagaraSystemInstance.h"
+#include "NiagaraRenderer.h"
 #include "Internationalization/Internationalization.h"
 #include "NiagaraScript.h"
+#include "ShaderParameterUtils.h"
+#include "NiagaraEmitterInstanceBatcher.h"
 
 #define LOCTEXT_NAMESPACE "NiagaraDataInterfaceStaticMesh"
 
+const FString UNiagaraDataInterfaceStaticMesh::MeshIndexBufferName(TEXT("IndexBuffer_"));
+const FString UNiagaraDataInterfaceStaticMesh::MeshVertexBufferName(TEXT("VertexBuffer_"));
+const FString UNiagaraDataInterfaceStaticMesh::MeshTangentBufferName(TEXT("TangentBuffer_"));
+const FString UNiagaraDataInterfaceStaticMesh::MeshTexCoordBufferName(TEXT("TexCoordBuffer_"));
+const FString UNiagaraDataInterfaceStaticMesh::MeshColorBufferName(TEXT("ColorBuffer_"));
+const FString UNiagaraDataInterfaceStaticMesh::MeshSectionBufferName(TEXT("SectionBuffer_"));
+const FString UNiagaraDataInterfaceStaticMesh::MeshTriangleBufferName(TEXT("TriangleBuffer_"));
+const FString UNiagaraDataInterfaceStaticMesh::SectionCountName(TEXT("SectionCount_"));
+const FString UNiagaraDataInterfaceStaticMesh::InstanceTransformName(TEXT("InstanceTransform_"));
+const FString UNiagaraDataInterfaceStaticMesh::InstanceTransformInverseTransposedName(TEXT("InstanceTransformInverseTransposed_"));
+const FString UNiagaraDataInterfaceStaticMesh::InstancePrevTransformName(TEXT("InstancePrevTransform_"));
+const FString UNiagaraDataInterfaceStaticMesh::InstanceInvDeltaTimeName(TEXT("InstanceInvDeltaTime_"));
+const FString UNiagaraDataInterfaceStaticMesh::InstanceWorldVelocityName(TEXT("InstanceWorldVelocity_"));
+const FString UNiagaraDataInterfaceStaticMesh::AreaWeightedSamplingName(TEXT("AreaWeightedSamplingName_"));
+const FString UNiagaraDataInterfaceStaticMesh::NumTexCoordName(TEXT("NumTexCoordName_"));
+
 FStaticMeshFilteredAreaWeightedSectionSampler::FStaticMeshFilteredAreaWeightedSectionSampler()
 	: Res(nullptr)
 	, Owner(nullptr)
@@ -40,6 +59,77 @@ float FStaticMeshFilteredAreaWeightedSectionSampler::GetWeights(TArray<float>& O
 	return Total;
 }
 
+//////////////////////////////////////////////////////////////////////////
+// FStaticMeshGpuSpawnBuffer
+FStaticMeshGpuSpawnBuffer::~FStaticMeshGpuSpawnBuffer()
+{
+	//ValidSections.Empty();
+}
+
+void FStaticMeshGpuSpawnBuffer::Initialise(const FStaticMeshLODResources& Res, const UNiagaraDataInterfaceStaticMesh& Interface, bool bIsGpuUniformlyDistributedSampling, const TArray<int32>& ValidSection, const FStaticMeshFilteredAreaWeightedSectionSampler& SectionSamplerParam)
+{
+	// In this function we prepare some data to be uploaded on GPU from the available mesh data. This is a thread safe place to create this data.
+	// The section buffer needs to be specific to the current UI being built (section/material culling).
+	SectionRenderData = &Res;
+
+	const uint32 ValidSectionCount = ValidSection.Num();
+	const TArray<float>& Prob = SectionSamplerParam.GetProb();
+	const TArray<int32>& Alias = SectionSamplerParam.GetAlias();
+	check(ValidSectionCount == Prob.Num());
+	// Build data that will be uploaded to GPU later from the render thread.
+	// The array contains data used to select regions for uniform particle spawning on them, as well as section triangle ranges.
+	ValidSections.Reserve(ValidSectionCount);
+	for (uint32 i = 0; i < ValidSectionCount; ++i)
+	{
+		uint32 ValidSectionId = ValidSection[i];
+		const FStaticMeshSection& Section = Res.Sections[ValidSectionId];
+		SectionInfo NewSectionInfo = { Section.FirstIndex / 3, Section.NumTriangles, Prob[i], (uint32)Alias[i] };
+		ValidSections.Add(NewSectionInfo);
+
+		check(!bIsGpuUniformlyDistributedSampling || bIsGpuUniformlyDistributedSampling && Res.AreaWeightedSectionSamplers[ValidSectionId].GetProb().Num() == Section.NumTriangles);
+	}
+
+	if (bIsGpuUniformlyDistributedSampling)
+		BufferUniformTriangleSamplingSRV = Res.AreaWeightedSectionSamplersBuffer.GetBufferSRV(); // Cache that SRV for later
+}
+
+void FStaticMeshGpuSpawnBuffer::InitRHI()
+{
+	MeshIndexBufferSrv = RHICreateShaderResourceView(SectionRenderData->IndexBuffer.IndexBufferRHI);
+	MeshVertexBufferSrv = SectionRenderData->VertexBuffers.PositionVertexBuffer.GetSRV();
+	MeshTangentBufferSrv = SectionRenderData->VertexBuffers.StaticMeshVertexBuffer.GetTangentsSRV();
+	MeshTexCoordBufferSrv = SectionRenderData->VertexBuffers.StaticMeshVertexBuffer.GetTexCoordsSRV();
+	NumTexCoord = SectionRenderData->VertexBuffers.StaticMeshVertexBuffer.GetNumTexCoords();
+	MeshColorBufferSRV = SectionRenderData->VertexBuffers.ColorVertexBuffer.GetColorComponentsSRV();
+
+	uint32 SizeByte = ValidSections.Num() * sizeof(SectionInfo);
+	if (SizeByte > 0)
+	{
+		FRHIResourceCreateInfo CreateInfo;
+		void* BufferData = nullptr;
+		BufferSectionRHI = RHICreateAndLockVertexBuffer(SizeByte, BUF_Static | BUF_ShaderResource, CreateInfo, BufferData);
+		SectionInfo* SectionInfoBuffer = (SectionInfo*)BufferData;
+		FMemory::Memcpy(SectionInfoBuffer, ValidSections.GetData(), SizeByte);
+		RHIUnlockVertexBuffer(BufferSectionRHI);
+		BufferSectionSRV = RHICreateShaderResourceView(BufferSectionRHI, sizeof(SectionInfo), PF_R32G32B32A32_UINT);
+	}
+}
+
+void FStaticMeshGpuSpawnBuffer::ReleaseRHI()
+{
+	MeshIndexBufferSrv.SafeRelease();
+	BufferSectionSRV.SafeRelease();
+	BufferSectionRHI.SafeRelease();
+
+	MeshIndexBufferSrv.SafeRelease();
+	MeshVertexBufferSrv.SafeRelease();
+	MeshTangentBufferSrv.SafeRelease();
+	MeshTexCoordBufferSrv.SafeRelease();
+	MeshColorBufferSRV.SafeRelease();
+	BufferSectionSRV.SafeRelease();
+}
+
+
 //////////////////////////////////////////////////////////////////////////
 //FNDIStaticMesh_InstanceData
 
@@ -61,7 +151,12 @@ bool FNDIStaticMesh_InstanceData::Init(UNiagaraDataInterfaceStaticMesh* Interfac
 	DeltaSeconds = 0.0f;
 	ChangeId = Interface->ChangeId;
 
-	if (Interface->Source)
+	if (Interface->SourceComponent)
+	{
+		Component = Interface->SourceComponent;
+		Mesh = Interface->SourceComponent->GetStaticMesh();
+	}
+	else if (Interface->Source)
 	{
 		AStaticMeshActor* MeshActor = Cast<AStaticMeshActor>(Interface->Source);
 		UStaticMeshComponent* SourceComp = nullptr;
@@ -161,7 +256,9 @@ bool FNDIStaticMesh_InstanceData::Init(UNiagaraDataInterfaceStaticMesh* Interfac
 	Mesh->GetOnMeshChanged().AddUObject(SystemInstance->GetComponent(), &UNiagaraComponent::ReinitializeSystem);
 #endif
 
-	bIsAreaWeightedSampling = Mesh->bSupportUniformlyDistributedSampling;
+	bMeshAllowsCpuAccess       = Mesh->bAllowCPUAccess;
+	bIsCpuUniformlyDistributedSampling = Mesh->bSupportUniformlyDistributedSampling;
+	bIsGpuUniformlyDistributedSampling = bIsCpuUniformlyDistributedSampling && Mesh->bSupportGpuUniformlyDistributedSampling;
 
 	//Init the instance filter
 	ValidSections.Empty();
@@ -200,13 +297,16 @@ bool FNDIStaticMesh_InstanceData::ResetRequired(UNiagaraDataInterfaceStaticMesh*
 		return true;
 	}
 
-	bool bPrevAreaWeighted = bIsAreaWeightedSampling;
 	//bool bPrevVCSampling = bSupportingVertexColorSampling;//TODO: Vertex color filtering needs more work.
 	bool bReset = false;
 	if (Mesh)
 	{
+		const bool bNewMeshAllowsCpuAccess = Mesh->bAllowCPUAccess;
+		const bool bNewIsCpuAreaWeightedSampling = Mesh->bSupportUniformlyDistributedSampling;
+		const bool bNewIsGpuAreaWeightedSampling = bIsCpuUniformlyDistributedSampling && Mesh->bSupportGpuUniformlyDistributedSampling;
+
 		//bSupportingVertexColorSampling = bEnableVertexColorRangeSorting && MeshHasColors();
-		bReset = !Mesh->bAllowCPUAccess || Mesh->bSupportUniformlyDistributedSampling != bPrevAreaWeighted/* || bSupportingVertexColorSampling != bPrevVCSampling*/;
+		bReset = bNewMeshAllowsCpuAccess != bMeshAllowsCpuAccess || bNewIsCpuAreaWeightedSampling != bIsCpuUniformlyDistributedSampling || bNewIsGpuAreaWeightedSampling != bIsGpuUniformlyDistributedSampling /* || bSupportingVertexColorSampling != bPrevVCSampling*/;
 	}
 	return bReset;
 }
@@ -238,6 +338,18 @@ bool FNDIStaticMesh_InstanceData::Tick(UNiagaraDataInterfaceStaticMesh* Interfac
 	}
 }
 
+void FNDIStaticMesh_InstanceData::Release()
+{
+	/*if (MeshGpuSpawnBuffer)
+	{
+		BeginReleaseResource(MeshGpuSpawnBuffer);
+		ENQUEUE_UNIQUE_RENDER_COMMAND_ONEPARAMETER(DeleteResource, FStaticMeshGpuSpawnBuffer*, ParamValidSectionVertexBuffer, MeshGpuSpawnBuffer,
+			{
+				delete ParamValidSectionVertexBuffer;
+			});
+	}*/
+}
+
 //void UNiagaraDataInterfaceStaticMesh::ResetFilter()
 //{
 //#if WITH_EDITOR	TODO: Make editor only and serialize this stuff out.
@@ -293,6 +405,258 @@ bool FNDIStaticMesh_InstanceData::Tick(UNiagaraDataInterfaceStaticMesh* Interfac
 
 //////////////////////////////////////////////////////////////////////////
 
+struct FNDIStaticMeshParametersName
+{
+	FString MeshIndexBufferName;
+	FString MeshVertexBufferName;
+	FString MeshTangentBufferName;
+	FString MeshTexCoordBufferName;
+	FString MeshColorBufferName;
+	FString MeshSectionBufferName;
+	FString MeshTriangleBufferName;
+	FString SectionCountName;
+	FString InstanceTransformName;
+	FString InstanceTransformInverseTransposedName;
+	FString InstancePrevTransformName;
+	FString InstanceInvDeltaTimeName;
+	FString InstanceWorldVelocityName;
+	FString AreaWeightedSamplingName;
+	FString NumTexCoordName;
+};
+
+static void GetNiagaraDataInterfaceParametersName(FNDIStaticMeshParametersName& Names, const FString& Suffix)
+{
+	Names.MeshIndexBufferName = UNiagaraDataInterfaceStaticMesh::MeshIndexBufferName + Suffix;
+	Names.MeshVertexBufferName = UNiagaraDataInterfaceStaticMesh::MeshVertexBufferName + Suffix;
+	Names.MeshTangentBufferName = UNiagaraDataInterfaceStaticMesh::MeshTangentBufferName + Suffix;
+	Names.MeshTexCoordBufferName = UNiagaraDataInterfaceStaticMesh::MeshTexCoordBufferName + Suffix;
+	Names.MeshColorBufferName = UNiagaraDataInterfaceStaticMesh::MeshColorBufferName + Suffix;
+	Names.MeshSectionBufferName = UNiagaraDataInterfaceStaticMesh::MeshSectionBufferName + Suffix;
+	Names.MeshTriangleBufferName = UNiagaraDataInterfaceStaticMesh::MeshTriangleBufferName + Suffix;
+	Names.SectionCountName = UNiagaraDataInterfaceStaticMesh::SectionCountName + Suffix;
+	Names.InstanceTransformName = UNiagaraDataInterfaceStaticMesh::InstanceTransformName + Suffix;
+	Names.InstanceTransformInverseTransposedName = UNiagaraDataInterfaceStaticMesh::InstanceTransformInverseTransposedName + Suffix;
+	Names.InstancePrevTransformName = UNiagaraDataInterfaceStaticMesh::InstancePrevTransformName + Suffix;
+	Names.InstanceInvDeltaTimeName = UNiagaraDataInterfaceStaticMesh::InstanceInvDeltaTimeName + Suffix;
+	Names.InstanceWorldVelocityName = UNiagaraDataInterfaceStaticMesh::InstanceWorldVelocityName + Suffix;
+	Names.AreaWeightedSamplingName = UNiagaraDataInterfaceStaticMesh::AreaWeightedSamplingName + Suffix;
+	Names.NumTexCoordName = UNiagaraDataInterfaceStaticMesh::NumTexCoordName + Suffix;
+}
+
+struct FNiagaraDataInterfaceParametersCS_StaticMesh : public FNiagaraDataInterfaceParametersCS
+{
+	virtual void Bind(const FNiagaraDataInterfaceParamRef& ParamRef, const class FShaderParameterMap& ParameterMap) override
+	{
+		FNDIStaticMeshParametersName ParamNames;
+		GetNiagaraDataInterfaceParametersName(ParamNames, ParamRef.ParameterInfo.DataInterfaceHLSLSymbol);
+
+		MeshIndexBuffer.Bind(ParameterMap, *ParamNames.MeshIndexBufferName);
+		MeshVertexBuffer.Bind(ParameterMap, *ParamNames.MeshVertexBufferName);
+		MeshTangentBuffer.Bind(ParameterMap, *ParamNames.MeshTangentBufferName);
+		MeshTexCoordBuffer.Bind(ParameterMap, *ParamNames.MeshTexCoordBufferName);
+		MeshColorBuffer.Bind(ParameterMap, *ParamNames.MeshColorBufferName);
+		MeshSectionBuffer.Bind(ParameterMap, *ParamNames.MeshSectionBufferName);
+		MeshTriangleBuffer.Bind(ParameterMap, *ParamNames.MeshTriangleBufferName);
+		SectionCount.Bind(ParameterMap, *ParamNames.SectionCountName);
+		InstanceTransform.Bind(ParameterMap, *ParamNames.InstanceTransformName);
+		InstanceTransformInverseTransposed.Bind(ParameterMap, *ParamNames.InstanceTransformInverseTransposedName);
+		InstancePrevTransform.Bind(ParameterMap, *ParamNames.InstancePrevTransformName);
+		InstanceInvDeltaTime.Bind(ParameterMap, *ParamNames.InstanceInvDeltaTimeName);
+		InstanceWorldVelocity.Bind(ParameterMap, *ParamNames.InstanceWorldVelocityName);
+		AreaWeightedSampling.Bind(ParameterMap, *ParamNames.AreaWeightedSamplingName);
+		NumTexCoord.Bind(ParameterMap, *ParamNames.NumTexCoordName);
+
+		if (!MeshIndexBuffer.IsBound())
+		{
+			UE_LOG(LogNiagara, Warning, TEXT("Binding failed for FNiagaraDataInterfaceParametersCS_StaticMesh Texture %s. Was it optimized out?"), *ParamNames.MeshIndexBufferName)
+		}
+		if (!MeshVertexBuffer.IsBound())
+		{
+			UE_LOG(LogNiagara, Warning, TEXT("Binding failed for FNiagaraDataInterfaceParametersCS_StaticMesh Sampler %s. Was it optimized out?"), *ParamNames.MeshVertexBufferName)
+		}
+	}
+
+	virtual void Serialize(FArchive& Ar)override
+	{
+		Ar << MeshIndexBuffer;
+		Ar << MeshVertexBuffer;
+		Ar << MeshTangentBuffer;
+		Ar << MeshTexCoordBuffer;
+		Ar << MeshColorBuffer;
+		Ar << SectionCount;
+		Ar << MeshSectionBuffer;
+		Ar << MeshTriangleBuffer;
+		Ar << InstanceTransform;
+		Ar << InstanceTransformInverseTransposed;
+		Ar << InstancePrevTransform;
+		Ar << InstanceInvDeltaTime;
+		Ar << InstanceWorldVelocity;
+		Ar << AreaWeightedSampling;
+		Ar << NumTexCoord;
+	}
+
+	virtual void Set(FRHICommandList& RHICmdList, const FNiagaraDataInterfaceSetArgs& Context) const override
+	{
+		check(IsInRenderingThread());
+
+		FComputeShaderRHIParamRef ComputeShaderRHI = Context.Shader->GetComputeShader();
+		
+		{
+			FNiagaraDataInterfaceProxyStaticMesh* InterfaceProxy = static_cast<FNiagaraDataInterfaceProxyStaticMesh*>(Context.DataInterface);
+			FNiagaraStaticMeshData* Data = InterfaceProxy->SystemInstancesToMeshData.Find(Context.SystemInstance);
+			ensure(Data);
+			if (Data != nullptr)
+			{
+				FStaticMeshGpuSpawnBuffer* SpawnBuffer = Data->MeshGpuSpawnBuffer;
+
+				SetSRVParameter(RHICmdList, ComputeShaderRHI, MeshVertexBuffer, SpawnBuffer->GetBufferPositionSRV());
+				SetSRVParameter(RHICmdList, ComputeShaderRHI, MeshTangentBuffer, SpawnBuffer->GetBufferTangentSRV());
+				SetSRVParameter(RHICmdList, ComputeShaderRHI, MeshIndexBuffer, SpawnBuffer->GetBufferIndexSRV());
+
+				SetShaderValue(RHICmdList, ComputeShaderRHI, NumTexCoord, SpawnBuffer->GetNumTexCoord());
+				if (SpawnBuffer->GetNumTexCoord() > 0)
+				{
+					SetSRVParameter(RHICmdList, ComputeShaderRHI, MeshTexCoordBuffer, SpawnBuffer->GetBufferTexCoordSRV());
+				}
+				SetSRVParameter(RHICmdList, ComputeShaderRHI, MeshColorBuffer, SpawnBuffer->GetBufferColorSRV());
+
+				check(SpawnBuffer); // should always be allocated, we always need the GPU buffer for a GpuSimulation.
+				SetShaderValue(RHICmdList, ComputeShaderRHI, SectionCount, SpawnBuffer->GetValidSectionCount());
+				SetSRVParameter(RHICmdList, ComputeShaderRHI, MeshSectionBuffer, SpawnBuffer->GetBufferSectionSRV());
+				if (Data->bIsGpuUniformlyDistributedSampling)
+				{
+					SetSRVParameter(RHICmdList, ComputeShaderRHI, MeshTriangleBuffer, SpawnBuffer->GetBufferUniformTriangleSamplingSRV());
+				}
+
+				const float InvDeltaTime = Data->DeltaSeconds > 0.0f ? 1.0f / Data->DeltaSeconds : 0.0f;
+				const FVector DeltaPosition = Data->Transform.GetOrigin() - Data->PrevTransform.GetOrigin();
+
+				SetShaderValue(RHICmdList, ComputeShaderRHI, InstanceTransform, Data->Transform);
+				SetShaderValue(RHICmdList, ComputeShaderRHI, InstanceTransformInverseTransposed, Data->Transform.InverseFast().GetTransposed());
+				SetShaderValue(RHICmdList, ComputeShaderRHI, InstancePrevTransform, Data->PrevTransform);
+				SetShaderValue(RHICmdList, ComputeShaderRHI, InstanceInvDeltaTime, InvDeltaTime);
+				SetShaderValue(RHICmdList, ComputeShaderRHI, InstanceWorldVelocity, DeltaPosition * InvDeltaTime);
+				SetShaderValue(RHICmdList, ComputeShaderRHI, AreaWeightedSampling, Data->bIsGpuUniformlyDistributedSampling ? 1 : 0);
+			}
+			else
+			{
+				SetSRVParameter(RHICmdList, ComputeShaderRHI, MeshVertexBuffer, FNiagaraRenderer::GetDummyFloatBuffer().SRV);
+				SetSRVParameter(RHICmdList, ComputeShaderRHI, MeshTangentBuffer, FNiagaraRenderer::GetDummyFloatBuffer().SRV);
+				SetSRVParameter(RHICmdList, ComputeShaderRHI, MeshIndexBuffer, FNiagaraRenderer::GetDummyUIntBuffer().SRV);
+
+				SetShaderValue(RHICmdList, ComputeShaderRHI, NumTexCoord, 0);
+				SetSRVParameter(RHICmdList, ComputeShaderRHI, MeshTexCoordBuffer, FNiagaraRenderer::GetDummyFloatBuffer().SRV);
+				SetSRVParameter(RHICmdList, ComputeShaderRHI, MeshColorBuffer, FNiagaraRenderer::GetDummyFloatBuffer().SRV);
+
+				SetShaderValue(RHICmdList, ComputeShaderRHI, SectionCount, 0);
+				SetSRVParameter(RHICmdList, ComputeShaderRHI, MeshSectionBuffer, FNiagaraRenderer::GetDummyFloatBuffer().SRV);
+				SetSRVParameter(RHICmdList, ComputeShaderRHI, MeshTriangleBuffer, FNiagaraRenderer::GetDummyFloatBuffer().SRV);
+
+				SetShaderValue(RHICmdList, ComputeShaderRHI, InstanceTransform, FMatrix::Identity);
+				SetShaderValue(RHICmdList, ComputeShaderRHI, InstanceTransformInverseTransposed, FMatrix::Identity);
+				SetShaderValue(RHICmdList, ComputeShaderRHI, InstancePrevTransform, FMatrix::Identity);
+				SetShaderValue(RHICmdList, ComputeShaderRHI, InstanceInvDeltaTime, 1.0f);
+				SetShaderValue(RHICmdList, ComputeShaderRHI, InstanceWorldVelocity, FVector::ZeroVector);
+				SetShaderValue(RHICmdList, ComputeShaderRHI, AreaWeightedSampling, 0);
+			}
+		}
+	}
+
+
+private:
+
+	FShaderResourceParameter MeshIndexBuffer;
+	FShaderResourceParameter MeshVertexBuffer;
+	FShaderResourceParameter MeshTangentBuffer;
+	FShaderResourceParameter MeshTexCoordBuffer;
+	FShaderResourceParameter MeshColorBuffer;
+	FShaderResourceParameter MeshSectionBuffer;
+	FShaderResourceParameter MeshTriangleBuffer;
+	FShaderParameter SectionCount;
+	FShaderParameter InstanceTransform;
+	FShaderParameter InstanceTransformInverseTransposed;
+	FShaderParameter InstancePrevTransform;
+	FShaderParameter InstanceInvDeltaTime;
+	FShaderParameter InstanceWorldVelocity;
+	FShaderParameter AreaWeightedSampling;
+	FShaderParameter NumTexCoord;
+};
+
+//////////////////////////////////////////////////////////////////////////
+
+void FNiagaraDataInterfaceProxyStaticMesh::DeferredDestroy()
+{
+	//-TODO: This is incorrect, we could be destroying instance data for a batcher is yet to tick, we should only be destroying data for this batcher!
+	for (const FGuid& Sys : DeferredDestroyList)
+	{
+		SystemInstancesToMeshData.Remove(Sys);
+		//UE_LOG(LogNiagara, Log, TEXT("DeferredDestroy() ... Removing %s"), *Sys.ToString());
+	}
+
+	DeferredDestroyList.Empty();
+}
+
+void FNiagaraDataInterfaceProxyStaticMesh::InitializePerInstanceData(const FGuid& SystemInstance, FStaticMeshGpuSpawnBuffer* MeshGPUSpawnBuffer)
+{
+	check(IsInRenderingThread());
+
+	FNiagaraStaticMeshData* Data = SystemInstancesToMeshData.Find(SystemInstance);
+	if (Data != nullptr)
+	{
+		DeferredDestroyList.Remove(SystemInstance);
+	}
+	else
+	{
+		Data = &SystemInstancesToMeshData.Add(SystemInstance);
+	}
+
+	//UE_LOG(LogNiagara, Log, TEXT("InitializePerInstanceData() ... %s"), *SystemInstance.ToString());
+
+	// @todo-threadsafety We should not ever see this case! Though it's not really an error...
+	if (Data->MeshGpuSpawnBuffer)
+	{
+		Data->MeshGpuSpawnBuffer->ReleaseResource();
+		delete Data->MeshGpuSpawnBuffer;
+	}
+
+	Data->MeshGpuSpawnBuffer = MeshGPUSpawnBuffer;
+}
+
+void FNiagaraDataInterfaceProxyStaticMesh::DestroyPerInstanceData(NiagaraEmitterInstanceBatcher* Batcher, const FGuid& SystemInstance)
+{
+	check(IsInRenderingThread());
+
+	//UE_LOG(LogNiagara, Log, TEXT("DestroyPerInstanceData() ... %s"), *SystemInstance.ToString());
+
+	// @todo-threadsafety verify this destroys the MeshGPUSpawnBuffer data. This thread owns it now.
+	//SystemInstancesToMeshData.Remove(SystemInstance);
+	DeferredDestroyList.Add(SystemInstance);
+	Batcher->EnqueueDeferredDeletesForDI_RenderThread(this->AsShared());
+}
+
+void FNiagaraDataInterfaceProxyStaticMesh::ConsumePerInstanceDataFromGameThread(void* PerInstanceData, const FGuid& Instance)
+{
+	FNiagaraPassedInstanceDataForRT* SourceData = static_cast<FNiagaraPassedInstanceDataForRT*>(PerInstanceData);
+	FNiagaraStaticMeshData* Data = SystemInstancesToMeshData.Find(Instance);
+	// @todo-threadsafety Verify we cannot ever reach here without valid data.
+	ensure(Data);
+	if (Data)
+	{
+		//UE_LOG(LogNiagara, Log, TEXT("ConsumePerInstanceDataFromGameThread() ... found %s"), *Instance.ToString());
+		check(Data->MeshGpuSpawnBuffer);
+
+		Data->bIsGpuUniformlyDistributedSampling = SourceData->bIsGpuUniformlyDistributedSampling;
+		Data->DeltaSeconds = SourceData->DeltaSeconds;
+		Data->Transform = SourceData->Transform;
+		Data->PrevTransform = SourceData->PrevTransform;
+	}
+	else
+	{
+		UE_LOG(LogNiagara, Log, TEXT("ConsumePerInstanceDataFromGameThread() ... could not find %s"), *Instance.ToString());
+	}
+}
+
+//////////////////////////////////////////////////////////////////////////
 
 UNiagaraDataInterfaceStaticMesh::UNiagaraDataInterfaceStaticMesh(FObjectInitializer const& ObjectInitializer)
 	: Super(ObjectInitializer)
@@ -302,12 +666,9 @@ UNiagaraDataInterfaceStaticMesh::UNiagaraDataInterfaceStaticMesh(FObjectInitiali
 	//, bSupportingVertexColorSampling(0)//Vertex color filtering needs some more work.
 	//, bFilterInitialized(false)
 {
-
+	Proxy = MakeShared<FNiagaraDataInterfaceProxyStaticMesh, ESPMode::ThreadSafe>();
 }
 
-
-#if WITH_EDITOR
-
 void UNiagaraDataInterfaceStaticMesh::PostInitProperties()
 {
 	Super::PostInitProperties();
@@ -322,6 +683,7 @@ void UNiagaraDataInterfaceStaticMesh::PostInitProperties()
 	}
 }
 
+#if WITH_EDITOR
 void UNiagaraDataInterfaceStaticMesh::PostEditChangeProperty(FPropertyChangedEvent& PropertyChangedEvent)
 {
 	Super::PostEditChangeProperty(PropertyChangedEvent);
@@ -368,7 +730,6 @@ void UNiagaraDataInterfaceStaticMesh::GetFunctions(TArray<FNiagaraFunctionSignat
 		Sig.Outputs.Add(FNiagaraVariable(FNiagaraTypeDefinition::GetIntDef(), TEXT("Section")));
 		Sig.bMemberFunction = true;
 		Sig.bRequiresContext = false;
-		//Sig.Owner = *GetFullName();
 		OutFunctions.Add(Sig);
 	}
 
@@ -377,11 +738,8 @@ void UNiagaraDataInterfaceStaticMesh::GetFunctions(TArray<FNiagaraFunctionSignat
 		Sig.Name = StaticMeshHelpers::RandomTriCoordName;
 		Sig.Inputs.Add(FNiagaraVariable(FNiagaraTypeDefinition(GetClass()), TEXT("StaticMesh")));
 		Sig.Outputs.Add(FNiagaraVariable(FNiagaraTypeDefinition(FMeshTriCoordinate::StaticStruct()), TEXT("Coord")));
-		//Sig.Outputs.Add(FNiagaraVariable(FNiagaraTypeDefinition::GetIntDef(), TEXT("Tri Index")));
-		//Sig.Outputs.Add(FNiagaraVariable(FNiagaraTypeDefinition::GetVec3Def(), TEXT("Bary Coord")));
 		Sig.bMemberFunction = true;
 		Sig.bRequiresContext = false;
-		//Sig.Owner = *GetFullName();
 		OutFunctions.Add(Sig);
 	}
 
@@ -392,12 +750,11 @@ void UNiagaraDataInterfaceStaticMesh::GetFunctions(TArray<FNiagaraFunctionSignat
 		Sig.Inputs.Add(FNiagaraVariable(FNiagaraTypeDefinition(FNiagaraTypeDefinition::GetFloatDef()), TEXT("Start")));
 		Sig.Inputs.Add(FNiagaraVariable(FNiagaraTypeDefinition(FNiagaraTypeDefinition::GetFloatDef()), TEXT("Range")));
 		Sig.Outputs.Add(FNiagaraVariable(FNiagaraTypeDefinition(FMeshTriCoordinate::StaticStruct()), TEXT("Coord")));
-		//Sig.Outputs.Add(FNiagaraVariable(FNiagaraTypeDefinition::GetIntDef(), TEXT("Tri Index")));
-		//Sig.Outputs.Add(FNiagaraVariable(FNiagaraTypeDefinition::GetVec3Def(), TEXT("Bary Coord")));
 		Sig.bMemberFunction = true;
 		Sig.bRequiresContext = false;
-		//Sig.Owner = *GetFullName();
+#if WITH_EDITORONLY_DATA
 		Sig.SetDescription(LOCTEXT("DataInterfaceSpline_RandomTriCoordVCFiltered", "If bSupportingVertexColorSampling is set on the data source, will randomly find a triangle whose red channel is within the Start to Start + Range color range."));
+#endif
 		OutFunctions.Add(Sig);
 	}
 
@@ -407,11 +764,8 @@ void UNiagaraDataInterfaceStaticMesh::GetFunctions(TArray<FNiagaraFunctionSignat
 		Sig.Inputs.Add(FNiagaraVariable(FNiagaraTypeDefinition(GetClass()), TEXT("StaticMesh")));
 		Sig.Inputs.Add(FNiagaraVariable(FNiagaraTypeDefinition::GetIntDef(), TEXT("Section")));
 		Sig.Outputs.Add(FNiagaraVariable(FNiagaraTypeDefinition(FMeshTriCoordinate::StaticStruct()), TEXT("Coord")));
-		//Sig.Outputs.Add(FNiagaraVariable(FNiagaraTypeDefinition::GetIntDef(), TEXT("Tri Index")));
-		//Sig.Outputs.Add(FNiagaraVariable(FNiagaraTypeDefinition::GetVec3Def(), TEXT("Bary Coord")));
 		Sig.bMemberFunction = true;
 		Sig.bRequiresContext = false;
-		//Sig.Owner = *GetFullName();
 		OutFunctions.Add(Sig);
 	}
 
@@ -420,13 +774,9 @@ void UNiagaraDataInterfaceStaticMesh::GetFunctions(TArray<FNiagaraFunctionSignat
 		Sig.Name = StaticMeshHelpers::GetTriPositionName;
 		Sig.Inputs.Add(FNiagaraVariable(FNiagaraTypeDefinition(GetClass()), TEXT("StaticMesh")));
 		Sig.Inputs.Add(FNiagaraVariable(FNiagaraTypeDefinition(FMeshTriCoordinate::StaticStruct()), TEXT("Coord")));
-		//Sig.Inputs.Add(FNiagaraVariable(FNiagaraTypeDefinition::GetIntDef(), TEXT("Tri Index")));
-		//Sig.Inputs.Add(FNiagaraVariable(FNiagaraTypeDefinition::GetVec3Def(), TEXT("Bary Coord")));
-
 		Sig.Outputs.Add(FNiagaraVariable(FNiagaraTypeDefinition::GetVec3Def(), TEXT("Position")));
 		Sig.bMemberFunction = true;
 		Sig.bRequiresContext = false;
-		//Sig.Owner = *GetFullName();
 		OutFunctions.Add(Sig);
 	}
 
@@ -435,14 +785,10 @@ void UNiagaraDataInterfaceStaticMesh::GetFunctions(TArray<FNiagaraFunctionSignat
 		Sig.Name = StaticMeshHelpers::GetTriPositionAndVelocityName;
 		Sig.Inputs.Add(FNiagaraVariable(FNiagaraTypeDefinition(GetClass()), TEXT("StaticMesh")));
 		Sig.Inputs.Add(FNiagaraVariable(FNiagaraTypeDefinition(FMeshTriCoordinate::StaticStruct()), TEXT("Coord")));
-		//Sig.Inputs.Add(FNiagaraVariable(FNiagaraTypeDefinition::GetIntDef(), TEXT("Tri Index")));
-		//Sig.Inputs.Add(FNiagaraVariable(FNiagaraTypeDefinition::GetVec3Def(), TEXT("Bary Coord")));
-
 		Sig.Outputs.Add(FNiagaraVariable(FNiagaraTypeDefinition::GetVec3Def(), TEXT("Position")));
 		Sig.Outputs.Add(FNiagaraVariable(FNiagaraTypeDefinition::GetVec3Def(), TEXT("Velocity")));
 		Sig.bMemberFunction = true;
 		Sig.bRequiresContext = false;
-		//Sig.Owner = *GetFullName();
 		OutFunctions.Add(Sig);
 	}
 
@@ -451,13 +797,9 @@ void UNiagaraDataInterfaceStaticMesh::GetFunctions(TArray<FNiagaraFunctionSignat
 		Sig.Name = StaticMeshHelpers::GetTriPositionWSName;
 		Sig.Inputs.Add(FNiagaraVariable(FNiagaraTypeDefinition(GetClass()), TEXT("StaticMesh")));
 		Sig.Inputs.Add(FNiagaraVariable(FNiagaraTypeDefinition(FMeshTriCoordinate::StaticStruct()), TEXT("Coord")));
-		//Sig.Inputs.Add(FNiagaraVariable(FNiagaraTypeDefinition::GetIntDef(), TEXT("Tri Index")));
-		//Sig.Inputs.Add(FNiagaraVariable(FNiagaraTypeDefinition::GetVec3Def(), TEXT("Bary Coord")));
-
 		Sig.Outputs.Add(FNiagaraVariable(FNiagaraTypeDefinition::GetVec3Def(), TEXT("Position")));
 		Sig.bMemberFunction = true;
 		Sig.bRequiresContext = false;
-		//Sig.Owner = *GetFullName();
 		OutFunctions.Add(Sig);
 	}
 
@@ -466,13 +808,9 @@ void UNiagaraDataInterfaceStaticMesh::GetFunctions(TArray<FNiagaraFunctionSignat
 		Sig.Name = StaticMeshHelpers::GetTriNormalName;
 		Sig.Inputs.Add(FNiagaraVariable(FNiagaraTypeDefinition(GetClass()), TEXT("StaticMesh")));
 		Sig.Inputs.Add(FNiagaraVariable(FNiagaraTypeDefinition(FMeshTriCoordinate::StaticStruct()), TEXT("Coord")));
-		//Sig.Inputs.Add(FNiagaraVariable(FNiagaraTypeDefinition::GetIntDef(), TEXT("Tri Index")));
-		//Sig.Inputs.Add(FNiagaraVariable(FNiagaraTypeDefinition::GetVec3Def(), TEXT("Bary Coord")));
-
 		Sig.Outputs.Add(FNiagaraVariable(FNiagaraTypeDefinition::GetVec3Def(), TEXT("Normal")));
 		Sig.bMemberFunction = true;
 		Sig.bRequiresContext = false;
-		//	Sig.Owner = *GetFullName();
 		OutFunctions.Add(Sig);
 	}
 
@@ -481,13 +819,9 @@ void UNiagaraDataInterfaceStaticMesh::GetFunctions(TArray<FNiagaraFunctionSignat
 		Sig.Name = StaticMeshHelpers::GetTriNormalWSName;
 		Sig.Inputs.Add(FNiagaraVariable(FNiagaraTypeDefinition(GetClass()), TEXT("StaticMesh")));
 		Sig.Inputs.Add(FNiagaraVariable(FNiagaraTypeDefinition(FMeshTriCoordinate::StaticStruct()), TEXT("Coord")));
-		//Sig.Inputs.Add(FNiagaraVariable(FNiagaraTypeDefinition::GetIntDef(), TEXT("Tri Index")));
-		//Sig.Inputs.Add(FNiagaraVariable(FNiagaraTypeDefinition::GetVec3Def(), TEXT("Bary Coord")));
-
 		Sig.Outputs.Add(FNiagaraVariable(FNiagaraTypeDefinition::GetVec3Def(), TEXT("Normal")));
 		Sig.bMemberFunction = true;
 		Sig.bRequiresContext = false;
-		//Sig.Owner = *GetFullName();
 		OutFunctions.Add(Sig);
 	}
 
@@ -496,15 +830,11 @@ void UNiagaraDataInterfaceStaticMesh::GetFunctions(TArray<FNiagaraFunctionSignat
 		Sig.Name = StaticMeshHelpers::GetTriTangentsName;
 		Sig.Inputs.Add(FNiagaraVariable(FNiagaraTypeDefinition(GetClass()), TEXT("StaticMesh")));
 		Sig.Inputs.Add(FNiagaraVariable(FNiagaraTypeDefinition(FMeshTriCoordinate::StaticStruct()), TEXT("Coord")));
-		//Sig.Inputs.Add(FNiagaraVariable(FNiagaraTypeDefinition::GetIntDef(), TEXT("Tri Index")));
-		//Sig.Inputs.Add(FNiagaraVariable(FNiagaraTypeDefinition::GetVec3Def(), TEXT("Bary Coord")));
-
 		Sig.Outputs.Add(FNiagaraVariable(FNiagaraTypeDefinition::GetVec3Def(), TEXT("Tangent")));
 		Sig.Outputs.Add(FNiagaraVariable(FNiagaraTypeDefinition::GetVec3Def(), TEXT("Binormal")));
 		Sig.Outputs.Add(FNiagaraVariable(FNiagaraTypeDefinition::GetVec3Def(), TEXT("Normal")));
 		Sig.bMemberFunction = true;
 		Sig.bRequiresContext = false;
-		//Sig.Owner = *GetFullName();
 		OutFunctions.Add(Sig);
 	}
 
@@ -513,15 +843,11 @@ void UNiagaraDataInterfaceStaticMesh::GetFunctions(TArray<FNiagaraFunctionSignat
 		Sig.Name = StaticMeshHelpers::GetTriTangentsWSName;
 		Sig.Inputs.Add(FNiagaraVariable(FNiagaraTypeDefinition(GetClass()), TEXT("StaticMesh")));
 		Sig.Inputs.Add(FNiagaraVariable(FNiagaraTypeDefinition(FMeshTriCoordinate::StaticStruct()), TEXT("Coord")));
-		//Sig.Inputs.Add(FNiagaraVariable(FNiagaraTypeDefinition::GetIntDef(), TEXT("Tri Index")));
-		//Sig.Inputs.Add(FNiagaraVariable(FNiagaraTypeDefinition::GetVec3Def(), TEXT("Bary Coord")));
-
 		Sig.Outputs.Add(FNiagaraVariable(FNiagaraTypeDefinition::GetVec3Def(), TEXT("Tangent")));
 		Sig.Outputs.Add(FNiagaraVariable(FNiagaraTypeDefinition::GetVec3Def(), TEXT("Binormal")));
 		Sig.Outputs.Add(FNiagaraVariable(FNiagaraTypeDefinition::GetVec3Def(), TEXT("Normal")));
 		Sig.bMemberFunction = true;
 		Sig.bRequiresContext = false;
-		//Sig.Owner = *GetFullName();
 		OutFunctions.Add(Sig);
 	}
 
@@ -530,13 +856,9 @@ void UNiagaraDataInterfaceStaticMesh::GetFunctions(TArray<FNiagaraFunctionSignat
 		Sig.Name = StaticMeshHelpers::GetTriColorName;
 		Sig.Inputs.Add(FNiagaraVariable(FNiagaraTypeDefinition(GetClass()), TEXT("StaticMesh")));
 		Sig.Inputs.Add(FNiagaraVariable(FNiagaraTypeDefinition(FMeshTriCoordinate::StaticStruct()), TEXT("Coord")));
-		//Sig.Inputs.Add(FNiagaraVariable(FNiagaraTypeDefinition::GetIntDef(), TEXT("Tri Index")));
-		//Sig.Inputs.Add(FNiagaraVariable(FNiagaraTypeDefinition::GetVec3Def(), TEXT("Bary Coord")));
-
 		Sig.Outputs.Add(FNiagaraVariable(FNiagaraTypeDefinition::GetColorDef(), TEXT("Color")));
 		Sig.bMemberFunction = true;
 		Sig.bRequiresContext = false;
-		//Sig.Owner = *GetFullName();
 		OutFunctions.Add(Sig);
 	}
 
@@ -545,14 +867,10 @@ void UNiagaraDataInterfaceStaticMesh::GetFunctions(TArray<FNiagaraFunctionSignat
 		Sig.Name = StaticMeshHelpers::GetTriUVName;
 		Sig.Inputs.Add(FNiagaraVariable(FNiagaraTypeDefinition(GetClass()), TEXT("StaticMesh")));
 		Sig.Inputs.Add(FNiagaraVariable(FNiagaraTypeDefinition(FMeshTriCoordinate::StaticStruct()), TEXT("Coord")));
-		//Sig.Inputs.Add(FNiagaraVariable(FNiagaraTypeDefinition::GetIntDef(), TEXT("Tri Index")));
-		//Sig.Inputs.Add(FNiagaraVariable(FNiagaraTypeDefinition::GetVec3Def(), TEXT("Bary Coord")));
 		Sig.Inputs.Add(FNiagaraVariable(FNiagaraTypeDefinition::GetIntDef(), TEXT("UV Set")));
-
 		Sig.Outputs.Add(FNiagaraVariable(FNiagaraTypeDefinition::GetVec2Def(), TEXT("UV")));
 		Sig.bMemberFunction = true;
 		Sig.bRequiresContext = false;
-		//Sig.Owner = *GetFullName();
 		OutFunctions.Add(Sig);
 	}
 
@@ -560,11 +878,9 @@ void UNiagaraDataInterfaceStaticMesh::GetFunctions(TArray<FNiagaraFunctionSignat
 		FNiagaraFunctionSignature Sig;
 		Sig.Name = StaticMeshHelpers::GetMeshLocalToWorldName;
 		Sig.Inputs.Add(FNiagaraVariable(FNiagaraTypeDefinition(GetClass()), TEXT("StaticMesh")));
-		//Sig.Inputs.Add(FNiagaraVariable(FNiagaraTypeDefinition(FMeshTriCoordinate::StaticStruct()), TEXT("Coord")));
 		Sig.Outputs.Add(FNiagaraVariable(FNiagaraTypeDefinition::GetMatrix4Def(), TEXT("Transform")));
 		Sig.bMemberFunction = true;
 		Sig.bRequiresContext = false;
-		//Sig.Owner = *GetFullName();
 		OutFunctions.Add(Sig);
 	}
 
@@ -572,11 +888,9 @@ void UNiagaraDataInterfaceStaticMesh::GetFunctions(TArray<FNiagaraFunctionSignat
 		FNiagaraFunctionSignature Sig;
 		Sig.Name = StaticMeshHelpers::GetMeshLocalToWorldInverseTransposedName;
 		Sig.Inputs.Add(FNiagaraVariable(FNiagaraTypeDefinition(GetClass()), TEXT("StaticMesh")));
-		//Sig.Inputs.Add(FNiagaraVariable(FNiagaraTypeDefinition(FMeshTriCoordinate::StaticStruct()), TEXT("Coord")));
 		Sig.Outputs.Add(FNiagaraVariable(FNiagaraTypeDefinition::GetMatrix4Def(), TEXT("Transform")));
 		Sig.bMemberFunction = true;
 		Sig.bRequiresContext = false;
-		//Sig.Owner = *GetFullName();
 		OutFunctions.Add(Sig);
 	}
 
@@ -584,11 +898,9 @@ void UNiagaraDataInterfaceStaticMesh::GetFunctions(TArray<FNiagaraFunctionSignat
 		FNiagaraFunctionSignature Sig;
 		Sig.Name = StaticMeshHelpers::GetMeshWorldVelocityName;
 		Sig.Inputs.Add(FNiagaraVariable(FNiagaraTypeDefinition(GetClass()), TEXT("StaticMesh")));
-		//Sig.Inputs.Add(FNiagaraVariable(FNiagaraTypeDefinition(FMeshTriCoordinate::StaticStruct()), TEXT("Coord")));
 		Sig.Outputs.Add(FNiagaraVariable(FNiagaraTypeDefinition::GetVec3Def(), TEXT("Velocity")));
 		Sig.bMemberFunction = true;
 		Sig.bRequiresContext = false;
-		//Sig.Owner = *GetFullName();
 		OutFunctions.Add(Sig);
 	}
 
@@ -603,7 +915,6 @@ void UNiagaraDataInterfaceStaticMesh::GetFunctions(TArray<FNiagaraFunctionSignat
 #if WITH_EDITORONLY_DATA
 		Sig.Description = LOCTEXT("GetVertexPositionDesc", "Returns the local space vertex position for the passed vertex.");
 #endif
-		//Sig.Owner = *GetFullName();
 		OutFunctions.Add(Sig);
 	} 
 	
@@ -618,7 +929,6 @@ void UNiagaraDataInterfaceStaticMesh::GetFunctions(TArray<FNiagaraFunctionSignat
 #if WITH_EDITORONLY_DATA
 		Sig.Description = LOCTEXT("GetVertexPositionWSDesc", "Returns the world space vertex position for the passed vertex.");
 #endif
-		//Sig.Owner = *GetFullName();
 		OutFunctions.Add(Sig);
 	}
 }
@@ -632,7 +942,7 @@ struct TUsesAreaWeightingBinder
 	{
 		FNDIStaticMesh_InstanceData* InstData = (FNDIStaticMesh_InstanceData*)InstanceData;
 		UNiagaraDataInterfaceStaticMesh* MeshInterface = CastChecked<UNiagaraDataInterfaceStaticMesh>(Interface);
-		if (InstData->UsesAreaWeighting())
+		if (InstData->UsesCpuUniformlyDistributedSampling())
 		{
 			NextBinder::template Bind<ParamTypes..., TIntegralConstant<bool, true>>(Interface, BindingInfo, InstanceData, OutFunc);
 		}
@@ -753,7 +1063,7 @@ void UNiagaraDataInterfaceStaticMesh::GetVMExternalFunction(const FVMExternalFun
 	}
 	else if (BindingInfo.Name == StaticMeshHelpers::GetTriNormalName)
 	{
-		check(BindingInfo.GetNumInputs() == 6 && BindingInfo.GetNumOutputs() == 3);
+		check(BindingInfo.GetNumInputs() == 5 && BindingInfo.GetNumOutputs() == 3);
 		bNeedsVertMain = true;
 		TNDIExplicitBinder<FNDITransformHandlerNoop, NDI_FUNC_BINDER(UNiagaraDataInterfaceStaticMesh, GetTriCoordNormal)>::Bind(this, BindingInfo, InstanceData, OutFunc);
 	}
@@ -867,7 +1177,42 @@ bool UNiagaraDataInterfaceStaticMesh::Equals(const UNiagaraDataInterface* Other)
 bool UNiagaraDataInterfaceStaticMesh::InitPerInstanceData(void* PerInstanceData, FNiagaraSystemInstance* SystemInstance)
 {
 	FNDIStaticMesh_InstanceData* Inst = new (PerInstanceData) FNDIStaticMesh_InstanceData();
-	return Inst->Init(this, SystemInstance);
+	bool bSuccess = Inst->Init(this, SystemInstance);
+
+	if (bSuccess)
+	{
+		// Always allocate when bAllowCPUAccess (index buffer can only have SRV created in this case as of today)
+		// We do not know if this interface is allocated for CPU or GPU so we allocate for both case... TODO: have some cached data created in case a GPU version is needed?
+		if (Inst->Mesh->bAllowCPUAccess)
+		{
+			FStaticMeshLODResources& Res = Inst->Mesh->RenderData->LODResources[0];
+
+			FStaticMeshGpuSpawnBuffer* MeshGpuSpawnBuffer = new FStaticMeshGpuSpawnBuffer;
+			MeshGpuSpawnBuffer->Initialise(Res, *this, Inst->bIsGpuUniformlyDistributedSampling, Inst->ValidSections, Inst->Sampler);
+
+			// Push instance data to RT
+			{
+				FNiagaraDataInterfaceProxyStaticMesh* ThisProxy = GetProxyAs<FNiagaraDataInterfaceProxyStaticMesh>();
+				ENQUEUE_RENDER_COMMAND(FNiagaraDIPushInitialInstanceDataToRT) (
+					[ThisProxy, InstanceID = SystemInstance->GetId(), MeshGpuSpawnBuffer](FRHICommandListImmediate& CmdList)
+					{
+						MeshGpuSpawnBuffer->InitResource();
+						ThisProxy->InitializePerInstanceData(InstanceID, MeshGpuSpawnBuffer);
+					}
+				);
+			}
+		}
+		else
+		{
+			ensure(Inst->Mesh->bAllowCPUAccess);
+		}
+	}
+	else
+	{
+		UE_LOG(LogNiagara, Warning, TEXT("UNiagaraDataInterfaceStaticMesh::InitPerInstanceData> Invalid per-instance init!!!"));
+	}
+
+	return bSuccess;
 }
 
 void UNiagaraDataInterfaceStaticMesh::DestroyPerInstanceData(void* PerInstanceData, FNiagaraSystemInstance* SystemInstance)
@@ -881,7 +1226,18 @@ void UNiagaraDataInterfaceStaticMesh::DestroyPerInstanceData(void* PerInstanceDa
 	}
 #endif
 
+	Inst->Release();
 	Inst->~FNDIStaticMesh_InstanceData();
+
+	{
+		FNiagaraDataInterfaceProxyStaticMesh* ThisProxy = GetProxyAs<FNiagaraDataInterfaceProxyStaticMesh>();
+		ENQUEUE_RENDER_COMMAND(FNiagaraDIDestroyInstanceData) (
+			[ThisProxy, InstanceID = SystemInstance->GetId(), Batcher = SystemInstance->GetBatcher()](FRHICommandListImmediate& CmdList)
+			{
+				ThisProxy->DestroyPerInstanceData(Batcher, InstanceID);
+			}
+		);
+	}
 }
 
 bool UNiagaraDataInterfaceStaticMesh::PerInstanceTick(void* PerInstanceData, FNiagaraSystemInstance* SystemInstance, float InDeltaSeconds)
@@ -963,7 +1319,7 @@ FORCEINLINE int32 UNiagaraDataInterfaceStaticMesh::RandomTriIndex<TIntegralConst
 	int32 SecIdx = RandomSection<TIntegralConstant<bool, true>, true>(RandStream, Res, InstData);
 	FStaticMeshSection&  Sec = Res.Sections[SecIdx];
 	int32 Tri = Res.AreaWeightedSectionSamplers[SecIdx].GetEntryIndex(RandStream.GetFraction(), RandStream.GetFraction());
-	return Sec.FirstIndex + Tri * 3;
+	return (Sec.FirstIndex / 3) + Tri;
 }
 
 template<>
@@ -972,7 +1328,7 @@ FORCEINLINE int32 UNiagaraDataInterfaceStaticMesh::RandomTriIndex<TIntegralConst
 	int32 SecIdx = RandomSection<TIntegralConstant<bool, true>, false>(RandStream, Res, InstData);
 	FStaticMeshSection&  Sec = Res.Sections[SecIdx];
 	int32 Tri = Res.AreaWeightedSectionSamplers[SecIdx].GetEntryIndex(RandStream.GetFraction(), RandStream.GetFraction());
-	return Sec.FirstIndex + Tri * 3;
+	return (Sec.FirstIndex / 3) + Tri;
 }
 
 template<>
@@ -981,7 +1337,7 @@ FORCEINLINE int32 UNiagaraDataInterfaceStaticMesh::RandomTriIndex<TIntegralConst
 	int32 SecIdx = RandomSection<TIntegralConstant<bool, false>, true>(RandStream, Res, InstData);
 	FStaticMeshSection&  Sec = Res.Sections[SecIdx];
 	int32 Tri = RandStream.RandRange(0, Sec.NumTriangles - 1);
-	return Sec.FirstIndex + Tri * 3;	
+	return (Sec.FirstIndex / 3) + Tri;
 }
 
 template<>
@@ -990,7 +1346,7 @@ FORCEINLINE int32 UNiagaraDataInterfaceStaticMesh::RandomTriIndex<TIntegralConst
 	int32 SecIdx = RandomSection<TIntegralConstant<bool, false>, false>(RandStream, Res, InstData);
 	FStaticMeshSection&  Sec = Res.Sections[SecIdx];
 	int32 Tri = RandStream.RandRange(0, Sec.NumTriangles - 1);
-	return Sec.FirstIndex + Tri * 3;
+	return (Sec.FirstIndex / 3) + Tri;
 }
 
 template<typename TAreaWeighted>
@@ -1087,7 +1443,9 @@ void UNiagaraDataInterfaceStaticMesh::RandomTriCoordVertexColorFiltered(FVectorV
 template<>
 FORCEINLINE int32 UNiagaraDataInterfaceStaticMesh::RandomTriIndexOnSection<TIntegralConstant<bool, true>>(FRandomStream& RandStream, FStaticMeshLODResources& Res, int32 SecIdx, FNDIStaticMesh_InstanceData* InstData)
 {
-	return Res.AreaWeightedSectionSamplers[SecIdx].GetEntryIndex(RandStream.GetFraction(), RandStream.GetFraction());
+	FStaticMeshSection&  Sec = Res.Sections[SecIdx];
+	int32 Tri = Res.AreaWeightedSectionSamplers[SecIdx].GetEntryIndex(RandStream.GetFraction(), RandStream.GetFraction());
+	return (Sec.FirstIndex / 3) + Tri;
 }
 
 template<>
@@ -1095,7 +1453,7 @@ FORCEINLINE int32 UNiagaraDataInterfaceStaticMesh::RandomTriIndexOnSection<TInte
 {
 	FStaticMeshSection&  Sec = Res.Sections[SecIdx];
 	int32 Tri = RandStream.RandRange(0, Sec.NumTriangles - 1);
-	return Sec.FirstIndex + Tri * 3;
+	return (Sec.FirstIndex / 3) + Tri;
 }
 
 template<typename TAreaWeighted>
@@ -1111,9 +1469,11 @@ void UNiagaraDataInterfaceStaticMesh::RandomTriCoordOnSection(FVectorVMContext&
 
 	FStaticMeshLODResources& Res = InstData->Mesh->RenderData->LODResources[0];
 	FIndexArrayView Indices = Res.IndexBuffer.GetArrayView();
+	const int32 MaxSection = Res.Sections.Num() - 1;
 	for (int32 i = 0; i < Context.NumInstances; ++i)
 	{
-		int32 SecIdx = SectionIdxParam.Get();
+
+		int32 SecIdx = FMath::Clamp(SectionIdxParam.Get(), 0, MaxSection);
 		*OutTri.GetDest() = RandomTriIndexOnSection<TAreaWeighted>(Context.RandStream, Res, SecIdx, InstData);
 		FVector Bary = RandomBarycentricCoord(Context.RandStream);
 		*OutBaryX.GetDest() = Bary.X;
@@ -1146,12 +1506,13 @@ void UNiagaraDataInterfaceStaticMesh::GetTriCoordPosition(FVectorVMContext& Cont
 	const FIndexArrayView& Indices = Res.IndexBuffer.GetArrayView();
 	const FPositionVertexBuffer& Positions = Res.VertexBuffers.PositionVertexBuffer;
 
+	const int32 NumTriangles = Indices.Num() / 3;
 	for (int32 i = 0; i < Context.NumInstances; ++i)
 	{
-		int32 Tri = TriParam.Get();
-		int32 Idx0 = Indices[Tri];
-		int32 Idx1 = Indices[Tri + 1];
-		int32 Idx2 = Indices[Tri + 2];
+		const int32 Tri = (TriParam.Get() % NumTriangles) * 3;
+		const int32 Idx0 = Indices[Tri];
+		const int32 Idx1 = Indices[Tri + 1];
+		const int32 Idx2 = Indices[Tri + 2];
 
 		FVector Pos = BarycentricInterpolate(BaryXParam.Get(), BaryYParam.Get(), BaryZParam.Get(), Positions.VertexPosition(Idx0), Positions.VertexPosition(Idx1), Positions.VertexPosition(Idx2));
 		TransformHandler.TransformPosition(Pos, InstData->Transform);
@@ -1189,12 +1550,13 @@ void UNiagaraDataInterfaceStaticMesh::GetTriCoordNormal(FVectorVMContext& Contex
 	const FIndexArrayView& Indices = Res.IndexBuffer.GetArrayView();
 	const FStaticMeshVertexBuffer& Verts = Res.VertexBuffers.StaticMeshVertexBuffer;
 
+	const int32 NumTriangles = Indices.Num() / 3;
 	for (int32 i = 0; i < Context.NumInstances; ++i)
 	{
-		int32 Tri = TriParam.Get();
-		int32 Idx0 = Indices[Tri];
-		int32 Idx1 = Indices[Tri + 1];
-		int32 Idx2 = Indices[Tri + 2];
+		const int32 Tri = (TriParam.Get() % NumTriangles) * 3;
+		const int32 Idx0 = Indices[Tri];
+		const int32 Idx1 = Indices[Tri + 1];
+		const int32 Idx2 = Indices[Tri + 2];
 
 		FVector Norm = BarycentricInterpolate(BaryXParam.Get(), BaryYParam.Get(), BaryZParam.Get(), Verts.VertexTangentZ(Idx0), Verts.VertexTangentZ(Idx1), Verts.VertexTangentZ(Idx2));
 		TransformHandler.TransformVector(Norm, InstData->TransformInverseTransposed);
@@ -1237,12 +1599,13 @@ void UNiagaraDataInterfaceStaticMesh::GetTriCoordTangents(FVectorVMContext& Cont
 	VectorVM::FExternalFuncRegisterHandler<float> OutNormY(Context);
 	VectorVM::FExternalFuncRegisterHandler<float> OutNormZ(Context);
 
+	const int32 NumTriangles = Indices.Num() / 3;
 	for (int32 i = 0; i < Context.NumInstances; ++i)
 	{
-		int32 Tri = TriParam.Get();
-		int32 Idx0 = Indices[Tri];
-		int32 Idx1 = Indices[Tri + 1];
-		int32 Idx2 = Indices[Tri + 2];
+		const int32 Tri = (TriParam.Get() % NumTriangles) * 3;
+		const int32 Idx0 = Indices[Tri];
+		const int32 Idx1 = Indices[Tri + 1];
+		const int32 Idx2 = Indices[Tri + 2];
 		FVector Tangent = BarycentricInterpolate(BaryXParam.Get(), BaryYParam.Get(), BaryZParam.Get(), Verts.GetTangentX(Idx0), Verts.GetTangentX(Idx1), Verts.GetTangentX(Idx2));
 		FVector Binorm = BarycentricInterpolate(BaryXParam.Get(), BaryYParam.Get(), BaryZParam.Get(), Verts.GetTangentY(Idx0), Verts.GetTangentY(Idx1), Verts.GetTangentY(Idx2));
 		FVector Norm = BarycentricInterpolate(BaryXParam.Get(), BaryYParam.Get(), BaryZParam.Get(), Verts.GetTangentZ(Idx0), Verts.GetTangentZ(Idx1), Verts.GetTangentZ(Idx2));
@@ -1294,12 +1657,13 @@ void UNiagaraDataInterfaceStaticMesh::GetTriCoordColor(FVectorVMContext& Context
 
 	if (Colors.GetNumVertices() > 0)
 	{
+		const int32 NumTriangles = Indices.Num() / 3;
 		for (int32 i = 0; i < Context.NumInstances; ++i)
 		{
-			int32 Tri = TriParam.Get();
-			int32 Idx0 = Indices[Tri];
-			int32 Idx1 = Indices[Tri + 1];
-			int32 Idx2 = Indices[Tri + 2];
+			const int32 Tri = (TriParam.Get() % NumTriangles) * 3;
+			const int32 Idx0 = Indices[Tri];
+			const int32 Idx1 = Indices[Tri + 1];
+			const int32 Idx2 = Indices[Tri + 2];
 
 			FLinearColor Color = BarycentricInterpolate(BaryXParam.Get(), BaryYParam.Get(), BaryZParam.Get(),
 				Colors.VertexColor(Idx0).ReinterpretAsLinear(), Colors.VertexColor(Idx1).ReinterpretAsLinear(), Colors.VertexColor(Idx2).ReinterpretAsLinear());
@@ -1356,12 +1720,14 @@ void UNiagaraDataInterfaceStaticMesh::GetTriCoordUV(FVectorVMContext& Context)
 	FStaticMeshLODResources& Res = InstData->Mesh->RenderData->LODResources[0];
 	const FIndexArrayView& Indices = Res.IndexBuffer.GetArrayView();
 	const VertexAccessorType Verts(Res.VertexBuffers.StaticMeshVertexBuffer);
+
+	const int32 NumTriangles = Indices.Num() / 3;
 	for (int32 i = 0; i < Context.NumInstances; ++i)
 	{
-		int32 Tri = TriParam.Get();
-		int32 Idx0 = Indices[Tri];
-		int32 Idx1 = Indices[Tri + 1];
-		int32 Idx2 = Indices[Tri + 2];
+		const int32 Tri = (TriParam.Get() % NumTriangles) * 3;
+		const int32 Idx0 = Indices[Tri];
+		const int32 Idx1 = Indices[Tri + 1];
+		const int32 Idx2 = Indices[Tri + 2];
 
 		int32 UVSet = UVSetParam.Get();
 		FVector2D UV = BarycentricInterpolate(BaryXParam.Get(), BaryYParam.Get(), BaryZParam.Get(),	Verts.GetUV(Idx0, UVSet), Verts.GetUV(Idx1, UVSet),	Verts.GetUV(Idx2, UVSet));
@@ -1398,13 +1764,14 @@ void UNiagaraDataInterfaceStaticMesh::GetTriCoordPositionAndVelocity(FVectorVMCo
 	const FIndexArrayView& Indices = Res.IndexBuffer.GetArrayView();
 	const FPositionVertexBuffer& Positions = Res.VertexBuffers.PositionVertexBuffer;
 
+	const int32 NumTriangles = Indices.Num() / 3;
 	float InvDt = 1.0f / InstData->DeltaSeconds;
 	for (int32 i = 0; i < Context.NumInstances; ++i)
 	{
-		int32 Tri = TriParam.Get();
-		int32 Idx0 = Indices[Tri];
-		int32 Idx1 = Indices[Tri + 1];
-		int32 Idx2 = Indices[Tri + 2];
+		const int32 Tri = (TriParam.Get() % NumTriangles) * 3;
+		const int32 Idx0 = Indices[Tri];
+		const int32 Idx1 = Indices[Tri + 1];
+		const int32 Idx2 = Indices[Tri + 2];
 
 		FVector Pos = BarycentricInterpolate(BaryXParam.Get(), BaryYParam.Get(), BaryZParam.Get(), Positions.VertexPosition(Idx0), Positions.VertexPosition(Idx1), Positions.VertexPosition(Idx2));
 
@@ -1538,6 +1905,490 @@ void UNiagaraDataInterfaceStaticMesh::GetVertexPosition(FVectorVMContext& Contex
 	}
 }
 
+void UNiagaraDataInterfaceStaticMesh::SetSourceComponentFromBlueprints(UStaticMeshComponent* ComponentToUse)
+{
+	// NOTE: When ChangeId changes the next tick will be skipped and a reset of the per-instance data will be initiated. 
+	ChangeId++;
+	SourceComponent = ComponentToUse;
+	Source = ComponentToUse->GetOwner();
+}
+
+void UNiagaraDataInterfaceStaticMesh::SetDefaultMeshFromBlueprints(UStaticMesh* MeshToUse)
+{
+	// NOTE: When ChangeId changes the next tick will be skipped and a reset of the per-instance data will be initiated. 
+	ChangeId++;
+	SourceComponent = nullptr;
+	Source = nullptr;
+	DefaultMesh = MeshToUse;
+}
+
+bool UNiagaraDataInterfaceStaticMesh::GetFunctionHLSL(const FName&  DefinitionFunctionName, FString InstanceFunctionName, FNiagaraDataInterfaceGPUParamInfo& ParamInfo, FString& OutHLSL)
+{
+	FNDIStaticMeshParametersName ParamNames;
+	GetNiagaraDataInterfaceParametersName(ParamNames, ParamInfo.DataInterfaceHLSLSymbol);
+
+	FString MeshTriCoordinateStructName = "MeshTriCoordinate";
+
+	TMap<FString, FStringFormatArg> ArgsSample = {
+		{TEXT("InstanceFunctionName"), InstanceFunctionName},
+		{TEXT("MeshTriCoordinateStructName"), MeshTriCoordinateStructName},
+		{TEXT("SectionCountName"), ParamNames.SectionCountName},
+		{TEXT("MeshSectionBufferName"), ParamNames.MeshSectionBufferName},
+		{TEXT("MeshIndexBufferName"), ParamNames.MeshIndexBufferName},
+		{TEXT("MeshTriangleBufferName"), ParamNames.MeshTriangleBufferName},
+		{TEXT("MeshVertexBufferName"), ParamNames.MeshVertexBufferName},
+		{TEXT("MeshTangentBufferName"), ParamNames.MeshTangentBufferName},
+		{TEXT("MeshTexCoordBufferName"), ParamNames.MeshTexCoordBufferName},
+		{TEXT("MeshColorBufferName"), ParamNames.MeshColorBufferName},
+		{TEXT("InstanceTransformName"), ParamNames.InstanceTransformName},
+		{TEXT("InstanceTransformInverseTransposed"), ParamNames.InstanceTransformInverseTransposedName},
+		{TEXT("InstancePrevTransformName"), ParamNames.InstancePrevTransformName},
+		{TEXT("InstanceInvDeltaTimeName"), ParamNames.InstanceInvDeltaTimeName},
+		{TEXT("InstanceWorldVelocity"), ParamNames.InstanceWorldVelocityName},
+		{TEXT("AreaWeightedSamplingName"), ParamNames.AreaWeightedSamplingName},
+		{TEXT("NumTexCoordName"), ParamNames.NumTexCoordName},
+	};
+
+	if (DefinitionFunctionName == StaticMeshHelpers::RandomSectionName)
+	{
+		static const TCHAR *FormatSample = TEXT(R"(
+			void {InstanceFunctionName} (out int Out_Section)
+			{
+				float RandS0 = NiagaraInternalNoise(1, 2, 3);
+				// Uniform sampling on mesh surface  (using alias method from Alias method from FWeightedRandomSampler)
+				uint SectionIndex = min(uint(RandS0 * float({SectionCountName})), {SectionCountName}-1);
+				uint4 SectionData = {MeshSectionBufferName}[SectionIndex];
+
+				// Alias check
+				float RandS1 = NiagaraInternalNoise(1, 2, 3);
+				if( RandS1 > asfloat(SectionData.z) )
+				{
+					SectionIndex = SectionData.w;
+				}
+				Out_Section = SectionIndex;
+			}
+			)");
+		OutHLSL += FString::Format(FormatSample, ArgsSample);
+	}
+	else if (DefinitionFunctionName == StaticMeshHelpers::RandomTriCoordName)
+	{
+		static const TCHAR *FormatSample = TEXT(R"(
+			void {InstanceFunctionName} (out {MeshTriCoordinateStructName} Out_Coord)
+			{
+				float RandS0 = NiagaraInternalNoise(1, 2, 3);
+
+				// Uniform sampling on mesh surface  (using alias method from Alias method from FWeightedRandomSampler)
+				uint SectionIndex = min(uint(RandS0 * float({SectionCountName})), {SectionCountName}-1);
+				uint4 SectionData = {MeshSectionBufferName}[SectionIndex];
+
+				// Alias check
+				float RandS1 = NiagaraInternalNoise(1, 2, 3);
+				if( RandS1 > asfloat(SectionData.z) )
+				{
+					SectionData = {MeshSectionBufferName}[SectionData.w];
+				}
+
+				uint SectionFirstTriangle  = SectionData.x;
+				uint SectionTriangleCount = SectionData.y;
+
+				float RandT0 = NiagaraInternalNoise(1, 2, 3);
+				[branch]
+				if({AreaWeightedSamplingName}==0)
+				{
+					// Uniform triangle id selection
+					Out_Coord.Tri = SectionFirstTriangle + min(uint(RandT0*float(SectionTriangleCount)), SectionTriangleCount-1); // avoid % by using mul/min to Tri = SectionTriangleCount
+				}
+				else
+				{
+					// Uniform area weighted position selection (using alias method from Alias method from FWeightedRandomSampler)
+					uint TriangleIndex = min(uint(RandT0*float(SectionTriangleCount)), SectionTriangleCount-1);
+					uint4 TriangleData = {MeshTriangleBufferName}[SectionFirstTriangle + TriangleIndex];
+
+					// Alias check
+					float RandT1 = NiagaraInternalNoise(1, 2, 3);
+					if( RandT1 > asfloat(TriangleData.x) )
+					{
+						TriangleIndex = TriangleData.y;
+					}
+					Out_Coord.Tri = SectionFirstTriangle + TriangleIndex;
+				}
+
+				float r0 = NiagaraInternalNoise(1, 2, 3);
+				float r1 = NiagaraInternalNoise(1, 2, 3);
+				float sqrt0 = sqrt(r0);
+				float sqrt1 = sqrt(r1);
+				Out_Coord.BaryCoord = float3(1.0f - sqrt0, sqrt0 * (1.0 - r1), r1 * sqrt0);
+			}
+			)");
+		OutHLSL += FString::Format(FormatSample, ArgsSample);
+	}
+	//else if (DefinitionFunctionName == StaticMeshHelpers::RandomTriCoordVCFilteredName)
+	//{
+	//}
+	else if (DefinitionFunctionName == StaticMeshHelpers::RandomTriCoordOnSectionName)
+	{
+		static const TCHAR *FormatSample = TEXT(R"(
+			void {InstanceFunctionName} (in int In_Section, out {MeshTriCoordinateStructName} Out_Coord)
+			{
+				int Section = clamp(In_Section, 0, {SectionCountName} - 1);
+
+				uint4 SectionData = {MeshSectionBufferName}[Section];
+				uint SectionFirstTriangle = SectionData.x;
+				uint SectionTriangleCount = SectionData.y;
+
+				float RandT0 = NiagaraInternalNoise(1, 2, 3);
+				[branch]
+				if({AreaWeightedSamplingName}==0)
+				{
+					// Uniform triangle id selection
+					Out_Coord.Tri = SectionFirstTriangle + min(uint(RandT0*float(SectionTriangleCount)), SectionTriangleCount-1); // avoid % by using mul/min to Tri = SectionTriangleCount
+				}
+				else
+				{
+					// Uniform area weighted position selection (using alias method from Alias method from FWeightedRandomSampler)
+					uint TriangleIndex = min(uint(RandT0*float(SectionTriangleCount)), SectionTriangleCount-1);
+					uint4 TriangleData = {MeshTriangleBufferName}[SectionFirstTriangle + TriangleIndex];
+
+					// Alias check
+					float RandT1 = NiagaraInternalNoise(1, 2, 3);
+					if( RandT1 > asfloat(TriangleData.x) )
+					{
+						TriangleIndex = TriangleData.y;
+					}
+					Out_Coord.Tri = SectionFirstTriangle + TriangleIndex;
+				}
+
+				float r0 = NiagaraInternalNoise(1, 2, 3);
+				float r1 = NiagaraInternalNoise(1, 2, 3);
+				float sqrt0 = sqrt(r0);
+				float sqrt1 = sqrt(r1);
+				Out_Coord.BaryCoord = float3(1.0f - sqrt0, sqrt0 * (1.0 - r1), r1 * sqrt0);
+			}
+			)");
+		OutHLSL += FString::Format(FormatSample, ArgsSample);
+	}
+	else if (DefinitionFunctionName == StaticMeshHelpers::GetTriPositionName)
+	{
+		static const TCHAR *FormatSample = TEXT(R"(
+			void {InstanceFunctionName} (in {MeshTriCoordinateStructName} In_Coord, out float3 Out_Position)
+			{
+				uint TriangleIndex = In_Coord.Tri * 3;
+				uint VertexIndex0 = {MeshIndexBufferName}[TriangleIndex  ] * 3;
+				uint VertexIndex1 = {MeshIndexBufferName}[TriangleIndex+1] * 3;
+				uint VertexIndex2 = {MeshIndexBufferName}[TriangleIndex+2] * 3;
+
+				// I could not find a R32G32B32f format to create an SRV on that buffer. So float load it is for now...
+				float3 vertex0 = float3({MeshVertexBufferName}[VertexIndex0], {MeshVertexBufferName}[VertexIndex0+1], {MeshVertexBufferName}[VertexIndex0+2]);
+				float3 vertex1 = float3({MeshVertexBufferName}[VertexIndex1], {MeshVertexBufferName}[VertexIndex1+1], {MeshVertexBufferName}[VertexIndex1+2]);
+				float3 vertex2 = float3({MeshVertexBufferName}[VertexIndex2], {MeshVertexBufferName}[VertexIndex2+1], {MeshVertexBufferName}[VertexIndex2+2]);
+				Out_Position = vertex0 * In_Coord.BaryCoord.x + vertex1 * In_Coord.BaryCoord.y + vertex2 * In_Coord.BaryCoord.z;
+			}
+			)");
+		OutHLSL += FString::Format(FormatSample, ArgsSample);
+	}
+	else if (DefinitionFunctionName == StaticMeshHelpers::GetTriPositionWSName)
+	{
+		static const TCHAR *FormatSample = TEXT(R"(
+			void {InstanceFunctionName} (in {MeshTriCoordinateStructName} In_Coord, out float3 Out_Position)
+			{
+				uint TriangleIndex = In_Coord.Tri * 3;
+				uint VertexIndex0 = {MeshIndexBufferName}[TriangleIndex  ] * 3;
+				uint VertexIndex1 = {MeshIndexBufferName}[TriangleIndex+1] * 3;
+				uint VertexIndex2 = {MeshIndexBufferName}[TriangleIndex+2] * 3;
+
+				// I could not find a R32G32B32f format to create an SRV on that buffer. So float load it is for now...
+				float3 vertex0 = float3({MeshVertexBufferName}[VertexIndex0], {MeshVertexBufferName}[VertexIndex0+1], {MeshVertexBufferName}[VertexIndex0+2]);
+				float3 vertex1 = float3({MeshVertexBufferName}[VertexIndex1], {MeshVertexBufferName}[VertexIndex1+1], {MeshVertexBufferName}[VertexIndex1+2]);
+				float3 vertex2 = float3({MeshVertexBufferName}[VertexIndex2], {MeshVertexBufferName}[VertexIndex2+1], {MeshVertexBufferName}[VertexIndex2+2]);
+				float3 Position = vertex0 * In_Coord.BaryCoord.x + vertex1 * In_Coord.BaryCoord.y + vertex2 * In_Coord.BaryCoord.z;
+
+				Out_Position = mul(float4(Position, 1.0), {InstanceTransformName}).xyz;
+			}
+			)");
+		OutHLSL += FString::Format(FormatSample, ArgsSample);
+	}
+	else if (DefinitionFunctionName == StaticMeshHelpers::GetTriNormalName)
+	{
+		static const TCHAR *FormatSample = TEXT(R"(
+		void {InstanceFunctionName} (in {MeshTriCoordinateStructName} In_Coord, out float3 Out_Normal)
+		{
+			uint TriangleIndex = In_Coord.Tri * 3;
+			uint VertexIndex0 = {MeshIndexBufferName}[TriangleIndex  ] * 2;
+			uint VertexIndex1 = {MeshIndexBufferName}[TriangleIndex+1] * 2;
+			uint VertexIndex2 = {MeshIndexBufferName}[TriangleIndex+2] * 2;
+
+			float3 Normal0 = TangentBias({MeshTangentBufferName}[VertexIndex0+1].xyz);
+			float3 Normal1 = TangentBias({MeshTangentBufferName}[VertexIndex1+1].xyz);
+			float3 Normal2 = TangentBias({MeshTangentBufferName}[VertexIndex2+1].xyz);
+
+			float3 Normal   = Normal0 * In_Coord.BaryCoord.x + Normal1 * In_Coord.BaryCoord.y + Normal2 * In_Coord.BaryCoord.z;
+
+			Out_Normal = normalize(Normal);
+		}
+		)");
+		OutHLSL += FString::Format(FormatSample, ArgsSample);
+	}
+	else if (DefinitionFunctionName == StaticMeshHelpers::GetTriNormalWSName)
+	{
+		static const TCHAR *FormatSample = TEXT(R"(
+		void {InstanceFunctionName} (in {MeshTriCoordinateStructName} In_Coord, out float3 Out_Normal)
+		{
+			uint TriangleIndex = In_Coord.Tri * 3;
+			uint VertexIndex0 = {MeshIndexBufferName}[TriangleIndex  ] * 2;
+			uint VertexIndex1 = {MeshIndexBufferName}[TriangleIndex+1] * 2;
+			uint VertexIndex2 = {MeshIndexBufferName}[TriangleIndex+2] * 2;
+
+			float3 Normal0 = TangentBias({MeshTangentBufferName}[VertexIndex0+1].xyz);
+			float3 Normal1 = TangentBias({MeshTangentBufferName}[VertexIndex1+1].xyz);
+			float3 Normal2 = TangentBias({MeshTangentBufferName}[VertexIndex2+1].xyz);
+
+			float3 Normal   = Normal0 * In_Coord.BaryCoord.x + Normal1 * In_Coord.BaryCoord.y + Normal2 * In_Coord.BaryCoord.z;
+
+			Out_Normal = normalize(mul(float4(Normal, 0.0), {InstanceTransformName}).xyz);
+		}
+		)");
+		OutHLSL += FString::Format(FormatSample, ArgsSample);
+	}
+	else if (DefinitionFunctionName == StaticMeshHelpers::GetTriTangentsName)
+	{
+		static const TCHAR *FormatSample = TEXT(R"(
+		void {InstanceFunctionName} (in {MeshTriCoordinateStructName} In_Coord, out float3 Out_Tangent, out float3 Out_Binormal, out float3 Out_Normal)
+		{
+			uint TriangleIndex = In_Coord.Tri * 3;
+			uint VertexIndex0 = {MeshIndexBufferName}[TriangleIndex  ] * 2;
+			uint VertexIndex1 = {MeshIndexBufferName}[TriangleIndex+1] * 2;
+			uint VertexIndex2 = {MeshIndexBufferName}[TriangleIndex+2] * 2;
+
+			float3 TangentX0 = TangentBias({MeshTangentBufferName}[VertexIndex0  ].xyz);
+			float4 TangentZ0 = TangentBias({MeshTangentBufferName}[VertexIndex0+1].xyzw);
+			float3 TangentX1 = TangentBias({MeshTangentBufferName}[VertexIndex1  ].xyz);
+			float4 TangentZ1 = TangentBias({MeshTangentBufferName}[VertexIndex1+1].xyzw);
+			float3 TangentX2 = TangentBias({MeshTangentBufferName}[VertexIndex2  ].xyz);
+			float4 TangentZ2 = TangentBias({MeshTangentBufferName}[VertexIndex2+1].xyzw);
+
+			float3 Binormal0   = cross(TangentZ0.xyz, TangentX0.xyz) * TangentZ0.w;
+			float3 Binormal1   = cross(TangentZ1.xyz, TangentX1.xyz) * TangentZ1.w;
+			float3 Binormal2   = cross(TangentZ2.xyz, TangentX2.xyz) * TangentZ2.w;
+
+			Out_Normal   = normalize(TangentZ0.xyz * In_Coord.BaryCoord.x + TangentZ1.xyz * In_Coord.BaryCoord.y + TangentZ2.xyz * In_Coord.BaryCoord.z);  // Normal is TangentZ
+			Out_Tangent  = normalize(TangentX0.xyz * In_Coord.BaryCoord.x + TangentX1.xyz * In_Coord.BaryCoord.y + TangentX2.xyz * In_Coord.BaryCoord.z);
+			Out_Binormal = normalize(Binormal0.xyz * In_Coord.BaryCoord.x + Binormal1.xyz * In_Coord.BaryCoord.y + Binormal2.xyz * In_Coord.BaryCoord.z);
+		}
+		)");
+		OutHLSL += FString::Format(FormatSample, ArgsSample);
+	}
+	else if (DefinitionFunctionName == StaticMeshHelpers::GetTriTangentsWSName)
+	{
+		static const TCHAR *FormatSample = TEXT(R"(
+		void {InstanceFunctionName} (in {MeshTriCoordinateStructName} In_Coord, out float3 Out_Tangent, out float3 Out_Binormal, out float3 Out_Normal)
+		{
+			uint TriangleIndex = In_Coord.Tri * 3;
+			uint VertexIndex0 = {MeshIndexBufferName}[TriangleIndex  ] * 2;
+			uint VertexIndex1 = {MeshIndexBufferName}[TriangleIndex+1] * 2;
+			uint VertexIndex2 = {MeshIndexBufferName}[TriangleIndex+2] * 2;
+
+			float3 TangentX0 = TangentBias({MeshTangentBufferName}[VertexIndex0  ].xyz);
+			float4 TangentZ0 = TangentBias({MeshTangentBufferName}[VertexIndex0+1].xyzw);
+			float3 TangentX1 = TangentBias({MeshTangentBufferName}[VertexIndex1  ].xyz);
+			float4 TangentZ1 = TangentBias({MeshTangentBufferName}[VertexIndex1+1].xyzw);
+			float3 TangentX2 = TangentBias({MeshTangentBufferName}[VertexIndex2  ].xyz);
+			float4 TangentZ2 = TangentBias({MeshTangentBufferName}[VertexIndex2+1].xyzw);
+
+			float3 Binormal0   = cross(TangentZ0.xyz, TangentX0.xyz) * TangentZ0.w;
+			float3 Binormal1   = cross(TangentZ1.xyz, TangentX1.xyz) * TangentZ1.w;
+			float3 Binormal2   = cross(TangentZ2.xyz, TangentX2.xyz) * TangentZ2.w;
+
+			float3 Normal   = TangentZ0.xyz * In_Coord.BaryCoord.x + TangentZ1.xyz * In_Coord.BaryCoord.y + TangentZ2.xyz * In_Coord.BaryCoord.z;  // Normal is TangentZ
+			float3 Tangent  = TangentX0.xyz * In_Coord.BaryCoord.x + TangentX1.xyz * In_Coord.BaryCoord.y + TangentX2.xyz * In_Coord.BaryCoord.z;
+			float3 Binormal = Binormal0.xyz * In_Coord.BaryCoord.x + Binormal1.xyz * In_Coord.BaryCoord.y + Binormal2.xyz * In_Coord.BaryCoord.z;
+
+			float3 NormalWorld  = normalize(mul(float4(Normal  , 0.0), {InstanceTransformName}).xyz);
+			float3 TangentWorld = normalize(mul(float4(Tangent , 0.0), {InstanceTransformName}).xyz);
+			float3 BinormalWorld= normalize(mul(float4(Binormal, 0.0), {InstanceTransformName}).xyz);
+
+			Out_Normal = NormalWorld;
+			Out_Tangent = TangentWorld;
+			Out_Binormal = BinormalWorld;
+		}
+		)");
+		OutHLSL += FString::Format(FormatSample, ArgsSample);
+	}
+	else if (DefinitionFunctionName == StaticMeshHelpers::GetTriColorName)
+	{
+		static const TCHAR *FormatSample = TEXT(R"(
+			void {InstanceFunctionName} (in {MeshTriCoordinateStructName} In_Coord, out float4 Out_Color)
+			{
+				uint TriangleIndex = In_Coord.Tri * 3;
+				uint VertexIndex0 = {MeshIndexBufferName}[TriangleIndex  ];
+				uint VertexIndex1 = {MeshIndexBufferName}[TriangleIndex+1];
+				uint VertexIndex2 = {MeshIndexBufferName}[TriangleIndex+2];
+
+				float4 Color0 = {MeshColorBufferName}[VertexIndex0] FMANUALFETCH_COLOR_COMPONENT_SWIZZLE;
+				float4 Color1 = {MeshColorBufferName}[VertexIndex1] FMANUALFETCH_COLOR_COMPONENT_SWIZZLE;
+				float4 Color2 = {MeshColorBufferName}[VertexIndex2] FMANUALFETCH_COLOR_COMPONENT_SWIZZLE;
+
+				Out_Color = Color0 * In_Coord.BaryCoord.x + Color1 * In_Coord.BaryCoord.y + Color2 * In_Coord.BaryCoord.z;
+			}
+			)");
+		OutHLSL += FString::Format(FormatSample, ArgsSample);
+	}
+	else if (DefinitionFunctionName == StaticMeshHelpers::GetTriUVName)
+	{
+		static const TCHAR *FormatSample = TEXT(R"(
+			void {InstanceFunctionName} (in {MeshTriCoordinateStructName} In_Coord, in int In_UVSet, out float2 Out_UV)
+			{
+				if({NumTexCoordName}>0)
+				{
+					uint TriangleIndex = In_Coord.Tri * 3;
+					uint VertexIndex0 = {MeshIndexBufferName}[TriangleIndex  ];
+					uint VertexIndex1 = {MeshIndexBufferName}[TriangleIndex+1];
+					uint VertexIndex2 = {MeshIndexBufferName}[TriangleIndex+2];
+
+					uint stride = {NumTexCoordName};
+					uint SelectedUVSet = clamp(In_UVSet, 0, {NumTexCoordName}-1);
+					float2 UV0 = {MeshTexCoordBufferName}[VertexIndex0 * stride + SelectedUVSet];
+					float2 UV1 = {MeshTexCoordBufferName}[VertexIndex1 * stride + SelectedUVSet];
+					float2 UV2 = {MeshTexCoordBufferName}[VertexIndex2 * stride + SelectedUVSet];
+
+					Out_UV = UV0 * In_Coord.BaryCoord.x + UV1 * In_Coord.BaryCoord.y + UV2 * In_Coord.BaryCoord.z;
+				}
+				else	
+				{
+					Out_UV = 0.0f;
+				}
+			}
+			)");
+		OutHLSL += FString::Format(FormatSample, ArgsSample);
+	}
+	else if (DefinitionFunctionName == StaticMeshHelpers::GetTriPositionAndVelocityName)
+	{
+		static const TCHAR *FormatSample = TEXT(R"(
+			void {InstanceFunctionName} (in {MeshTriCoordinateStructName} In_Coord, out float3 Out_Position, out float3 Out_Velocity)
+			{
+				uint TriangleIndex = In_Coord.Tri * 3;
+				uint VertexIndex0 = {MeshIndexBufferName}[TriangleIndex  ] * 3;
+				uint VertexIndex1 = {MeshIndexBufferName}[TriangleIndex+1] * 3;
+				uint VertexIndex2 = {MeshIndexBufferName}[TriangleIndex+2] * 3;
+
+				// I could not find a R32G32B32f format to create an SRV on that buffer. So float load it is for now...
+				float3 vertex0 = float3({MeshVertexBufferName}[VertexIndex0], {MeshVertexBufferName}[VertexIndex0+1], {MeshVertexBufferName}[VertexIndex0+2]);
+				float3 vertex1 = float3({MeshVertexBufferName}[VertexIndex1], {MeshVertexBufferName}[VertexIndex1+1], {MeshVertexBufferName}[VertexIndex1+2]);
+				float3 vertex2 = float3({MeshVertexBufferName}[VertexIndex2], {MeshVertexBufferName}[VertexIndex2+1], {MeshVertexBufferName}[VertexIndex2+2]);
+				float3 WSPos = vertex0 * In_Coord.BaryCoord.x + vertex1 * In_Coord.BaryCoord.y + vertex2 * In_Coord.BaryCoord.z;
+				float3 PrevWSPos = WSPos;
+
+				WSPos = mul(float4(WSPos,1.0), {InstanceTransformName}).xyz;
+				PrevWSPos = mul(float4(PrevWSPos,1.0), {InstancePrevTransformName}).xyz;
+
+				Out_Position = WSPos;
+				Out_Velocity = (WSPos - PrevWSPos) * {InstanceInvDeltaTimeName};
+			}
+			)");
+		OutHLSL += FString::Format(FormatSample, ArgsSample);
+	}
+	else if (DefinitionFunctionName == StaticMeshHelpers::GetMeshLocalToWorldName)
+	{
+		static const TCHAR *FormatSample = TEXT(R"(
+			void {InstanceFunctionName} (out float4x4 Out_Transform)
+			{
+				Out_Transform = {InstanceTransformName};
+			}
+			)");
+		OutHLSL += FString::Format(FormatSample, ArgsSample);
+	}
+	else if (DefinitionFunctionName == StaticMeshHelpers::GetMeshLocalToWorldInverseTransposedName)
+	{
+		static const TCHAR *FormatSample = TEXT(R"(
+			void {InstanceFunctionName} (out float4x4 Out_Transform)
+			{
+				Out_Transform = {InstanceTransformInverseTransposed};
+			}
+			)");
+		OutHLSL += FString::Format(FormatSample, ArgsSample);
+	}
+	else if (DefinitionFunctionName == StaticMeshHelpers::GetMeshWorldVelocityName)
+	{
+		static const TCHAR *FormatSample = TEXT(R"(
+			void {InstanceFunctionName} (out float3 Out_Velocity)
+			{
+				Out_Velocity = {InstanceWorldVelocity};
+			}
+			)");
+		OutHLSL += FString::Format(FormatSample, ArgsSample);
+	}
+	else if (DefinitionFunctionName == StaticMeshHelpers::GetVertexPositionName)
+	{
+		static const TCHAR *FormatSample = TEXT(R"(
+				void {InstanceFunctionName} (in int VertexIndex, out float3 Out_Position)
+				{
+					VertexIndex *= 3;
+					Out_Position = float3({MeshVertexBufferName}[VertexIndex], {MeshVertexBufferName}[VertexIndex+1], {MeshVertexBufferName}[VertexIndex+2]);
+				}
+				)");
+		OutHLSL += FString::Format(FormatSample, ArgsSample);
+	}
+	else if (DefinitionFunctionName == StaticMeshHelpers::GetVertexPositionWSName)
+	{
+		static const TCHAR *FormatSample = TEXT(R"(
+				void {InstanceFunctionName} (in int VertexIndex, out float3 Out_Position)
+				{
+					VertexIndex *= 3;
+					Out_Position = float3({MeshVertexBufferName}[VertexIndex], {MeshVertexBufferName}[VertexIndex+1], {MeshVertexBufferName}[VertexIndex+2]);
+					Out_Position = mul(float4(Out_Position, 1.0), {InstanceTransformName}).xyz;
+				}
+				)");
+		OutHLSL += FString::Format(FormatSample, ArgsSample);
+	}
+	else
+	{
+		// This function is not support
+		return false;
+	}
+
+	OutHLSL += TEXT("\n");
+	return true;
+}
+
+void UNiagaraDataInterfaceStaticMesh::GetParameterDefinitionHLSL(FNiagaraDataInterfaceGPUParamInfo& ParamInfo, FString& OutHLSL)
+{
+	FNDIStaticMeshParametersName ParamNames;
+	GetNiagaraDataInterfaceParametersName(ParamNames, ParamInfo.DataInterfaceHLSLSymbol);
+
+	OutHLSL += TEXT("Buffer<uint> ") + ParamNames.MeshIndexBufferName + TEXT(";\n");
+	OutHLSL += TEXT("Buffer<float> ") + ParamNames.MeshVertexBufferName + TEXT(";\n");
+	OutHLSL += TEXT("Buffer<float4> ") + ParamNames.MeshTangentBufferName + TEXT(";\n");
+	OutHLSL += TEXT("Buffer<float2> ") + ParamNames.MeshTexCoordBufferName + TEXT(";\n");
+	OutHLSL += TEXT("Buffer<float4> ") + ParamNames.MeshColorBufferName + TEXT(";\n");
+	OutHLSL += TEXT("Buffer<uint4> ") + ParamNames.MeshSectionBufferName + TEXT(";\n");
+	OutHLSL += TEXT("Buffer<uint4> ") + ParamNames.MeshTriangleBufferName + TEXT(";\n");
+	OutHLSL += TEXT("uint ") + ParamNames.SectionCountName + TEXT(";\n");
+	OutHLSL += TEXT("float4x4 ") + ParamNames.InstanceTransformName + TEXT(";\n");
+	OutHLSL += TEXT("float4x4 ") + ParamNames.InstanceTransformInverseTransposedName + TEXT(";\n");
+	OutHLSL += TEXT("float4x4 ") + ParamNames.InstancePrevTransformName + TEXT(";\n");
+	OutHLSL += TEXT("float ") + ParamNames.InstanceInvDeltaTimeName + TEXT(";\n");
+	OutHLSL += TEXT("float4 ") + ParamNames.InstanceWorldVelocityName + TEXT(";\n");
+	OutHLSL += TEXT("uint ") + ParamNames.AreaWeightedSamplingName + TEXT(";\n");	// Could be used for other flags
+	OutHLSL += TEXT("uint ") + ParamNames.NumTexCoordName + TEXT(";\n");
+}
+
+void UNiagaraDataInterfaceStaticMesh::ProvidePerInstanceDataForRenderThread(void* DataForRenderThread, void* PerInstanceData, const FGuid& SystemInstance)
+{
+	check(Proxy);
+
+	FNDIStaticMesh_InstanceData* InstanceData = static_cast<FNDIStaticMesh_InstanceData*>(PerInstanceData);
+	FNiagaraPassedInstanceDataForRT* DataToPass = static_cast<FNiagaraPassedInstanceDataForRT*>(DataForRenderThread);
+
+	DataToPass->bIsGpuUniformlyDistributedSampling = InstanceData->bIsGpuUniformlyDistributedSampling;
+	DataToPass->DeltaSeconds = InstanceData->DeltaSeconds;
+	DataToPass->Transform = InstanceData->Transform;
+	DataToPass->PrevTransform = InstanceData->PrevTransform;
+}
+
+FNiagaraDataInterfaceParametersCS* UNiagaraDataInterfaceStaticMesh::ConstructComputeParameters() const 
+{
+	return new FNiagaraDataInterfaceParametersCS_StaticMesh();
+}
+
 //////////////////////////////////////////////////////////////////////////
 
 bool FDynamicVertexColorFilterData::Init(FNDIStaticMesh_InstanceData* Owner)
diff --git a/Engine/Plugins/FX/Niagara/Source/Niagara/Private/NiagaraDataInterfaceTexture.cpp b/Engine/Plugins/FX/Niagara/Source/Niagara/Private/NiagaraDataInterfaceTexture.cpp
--- a/Engine/Plugins/FX/Niagara/Source/Niagara/Private/NiagaraDataInterfaceTexture.cpp
+++ b/Engine/Plugins/FX/Niagara/Source/Niagara/Private/NiagaraDataInterfaceTexture.cpp
@@ -20,7 +20,8 @@ UNiagaraDataInterfaceTexture::UNiagaraDataInterfaceTexture(FObjectInitializer co
 	: Super(ObjectInitializer)
 	, Texture(nullptr)
 {
-
+	Proxy = MakeShared<FNiagaraDataInterfaceProxyTexture, ESPMode::ThreadSafe>();
+	PushToRenderThread();
 }
 
 void UNiagaraDataInterfaceTexture::PostInitProperties()
@@ -31,6 +32,8 @@ void UNiagaraDataInterfaceTexture::PostInitProperties()
 	{
 		FNiagaraTypeRegistry::Register(FNiagaraTypeDefinition(GetClass()), true, false, false);
 	}
+
+	PushToRenderThread();
 }
 
 void UNiagaraDataInterfaceTexture::PostLoad()
@@ -46,6 +49,7 @@ void UNiagaraDataInterfaceTexture::PostLoad()
 		}
 	}
 #endif
+	PushToRenderThread();
 }
 
 #if WITH_EDITOR
@@ -53,6 +57,7 @@ void UNiagaraDataInterfaceTexture::PostLoad()
 void UNiagaraDataInterfaceTexture::PostEditChangeProperty(struct FPropertyChangedEvent& PropertyChangedEvent)
 {
 	Super::PostEditChangeProperty(PropertyChangedEvent);
+	PushToRenderThread();
 }
 
 
@@ -78,6 +83,7 @@ bool UNiagaraDataInterfaceTexture::CopyToInternal(UNiagaraDataInterface* Destina
 	}
 	UNiagaraDataInterfaceTexture* DestinationTexture = CastChecked<UNiagaraDataInterfaceTexture>(Destination);
 	DestinationTexture->Texture = Texture;
+	DestinationTexture->PushToRenderThread();
 
 	return true;
 }
@@ -373,46 +379,29 @@ struct FNiagaraDataInterfaceParametersCS_Texture : public FNiagaraDataInterfaceP
 		}
 	}
 
-	virtual void Set(FRHICommandList& RHICmdList, FNiagaraShader* Shader, class UNiagaraDataInterface* DataInterface, void* PerInstanceData) const override
+	virtual void Set(FRHICommandList& RHICmdList, const FNiagaraDataInterfaceSetArgs& Context) const override
 	{
 		check(IsInRenderingThread());
 
-		FComputeShaderRHIParamRef ComputeShaderRHI = Shader->GetComputeShader();
-		UNiagaraDataInterfaceTexture* TextureDI = CastChecked<UNiagaraDataInterfaceTexture>(DataInterface);
-		UTexture *Texture = TextureDI->Texture;
-		float TexDims[2];
-		if (!Texture)
+		FComputeShaderRHIParamRef ComputeShaderRHI = Context.Shader->GetComputeShader();
+		FNiagaraDataInterfaceProxyTexture* TextureDI = static_cast<FNiagaraDataInterfaceProxyTexture*>(Context.DataInterface);
+		if (!TextureDI->TextureRHI)
 		{
-			TexDims[0] = 0.0f;
-			TexDims[1] = 0.0f;
+			float TexDims[] = { 0.0f, 0.0f };
 			SetShaderValue(RHICmdList, ComputeShaderRHI, Dimensions, TexDims);
 			return;
 		}
 
-		FTextureRHIParamRef TextureRHI = Texture->TextureReference.TextureReferenceRHI->GetReferencedTexture();
-		if (TextureRHI != nullptr)
-		{
-			SetTextureParameter(
-				RHICmdList,
-				ComputeShaderRHI,
-				TextureParam,
-				SamplerParam,
-				Texture->Resource->SamplerStateRHI,
-				TextureRHI
-			);
-			TexDims[0] = TextureDI->Texture->GetSurfaceWidth();
-			TexDims[1] = TextureDI->Texture->GetSurfaceHeight();
-
-		}
-		else
-		{
-			TexDims[0] = 0.0f;
-			TexDims[1] = 0.0f;
-			SetShaderValue(RHICmdList, ComputeShaderRHI, Dimensions, TexDims);
-			return;
-		}
+		SetTextureParameter(
+			RHICmdList,
+			ComputeShaderRHI,
+			TextureParam,
+			SamplerParam,
+			TextureDI->SamplerStateRHI,
+			TextureDI->TextureRHI
+		);
 
-		SetShaderValue(RHICmdList, ComputeShaderRHI, Dimensions, TexDims);
+		SetShaderValue(RHICmdList, ComputeShaderRHI, Dimensions, TextureDI->TexDims);
 	}
 	
 
@@ -423,6 +412,40 @@ private:
 	FShaderParameter Dimensions;
 };
 
+void UNiagaraDataInterfaceTexture::PushToRenderThread()
+{
+	FNiagaraDataInterfaceProxyTexture* RT_Proxy = GetProxyAs<FNiagaraDataInterfaceProxyTexture>();
+
+	float DimX = 0.0f;
+	float DimY = 0.0f;
+	FTextureRHIRef RT_Texture;
+	FSamplerStateRHIRef RT_SamplerState;
+
+	if (Texture && Texture->TextureReference.TextureReferenceRHI.IsValid())
+	{
+		if (Texture->TextureReference.TextureReferenceRHI)
+		{
+			RT_Texture = Texture->TextureReference.TextureReferenceRHI->GetReferencedTexture();
+		}
+		if (Texture->Resource)
+		{
+			RT_SamplerState = Texture->Resource->SamplerStateRHI;
+		}
+		DimX = Texture->GetSurfaceWidth();
+		DimY = Texture->GetSurfaceHeight();
+	}
+
+	ENQUEUE_RENDER_COMMAND(FPushDITextureToRT) (
+		[RT_Proxy, RT_Texture, RT_SamplerState, DimX, DimY](FRHICommandListImmediate& RHICmdList)
+		{
+			RT_Proxy->TextureRHI = RT_Texture;
+			RT_Proxy->SamplerStateRHI = RT_SamplerState;
+			RT_Proxy->TexDims[0] = DimX;
+			RT_Proxy->TexDims[1] = DimY;
+		}
+	);
+}
+
 FNiagaraDataInterfaceParametersCS* UNiagaraDataInterfaceTexture::ConstructComputeParameters()const
 {
 	return new FNiagaraDataInterfaceParametersCS_Texture();
diff --git a/Engine/Plugins/FX/Niagara/Source/Niagara/Private/NiagaraDataInterfaceVector2DCurve.cpp b/Engine/Plugins/FX/Niagara/Source/Niagara/Private/NiagaraDataInterfaceVector2DCurve.cpp
--- a/Engine/Plugins/FX/Niagara/Source/Niagara/Private/NiagaraDataInterfaceVector2DCurve.cpp
+++ b/Engine/Plugins/FX/Niagara/Source/Niagara/Private/NiagaraDataInterfaceVector2DCurve.cpp
@@ -81,7 +81,8 @@ void UNiagaraDataInterfaceVector2DCurve::UpdateLUT()
 		ShaderLUT.Add(C.X);
 		ShaderLUT.Add(C.Y);
 	}
-	GPUBufferDirty = true;
+
+	Super::PushToRenderThread();
 }
 
 bool UNiagaraDataInterfaceVector2DCurve::CopyToInternal(UNiagaraDataInterface* Destination) const
diff --git a/Engine/Plugins/FX/Niagara/Source/Niagara/Private/NiagaraDataInterfaceVector4Curve.cpp b/Engine/Plugins/FX/Niagara/Source/Niagara/Private/NiagaraDataInterfaceVector4Curve.cpp
--- a/Engine/Plugins/FX/Niagara/Source/Niagara/Private/NiagaraDataInterfaceVector4Curve.cpp
+++ b/Engine/Plugins/FX/Niagara/Source/Niagara/Private/NiagaraDataInterfaceVector4Curve.cpp
@@ -90,7 +90,7 @@ void UNiagaraDataInterfaceVector4Curve::UpdateLUT()
 		ShaderLUT.Add(C.B);
 		ShaderLUT.Add(C.A);
 	}
-	GPUBufferDirty = true;
+	Super::PushToRenderThread();
 }
 
 bool UNiagaraDataInterfaceVector4Curve::CopyToInternal(UNiagaraDataInterface* Destination) const 
diff --git a/Engine/Plugins/FX/Niagara/Source/Niagara/Private/NiagaraDataInterfaceVectorField.cpp b/Engine/Plugins/FX/Niagara/Source/Niagara/Private/NiagaraDataInterfaceVectorField.cpp
--- a/Engine/Plugins/FX/Niagara/Source/Niagara/Private/NiagaraDataInterfaceVectorField.cpp
+++ b/Engine/Plugins/FX/Niagara/Source/Niagara/Private/NiagaraDataInterfaceVectorField.cpp
@@ -32,6 +32,8 @@ UNiagaraDataInterfaceVectorField::UNiagaraDataInterfaceVectorField(FObjectInitia
 	, bTileY(false)
 	, bTileZ(false)
 {
+	Proxy = MakeShared<FNiagaraDataInterfaceProxyVectorField, ESPMode::ThreadSafe>();
+	PushToRenderThread();
 }
 
 /*--------------------------------------------------------------------------------------------------------------------------*/
@@ -45,12 +47,13 @@ void UNiagaraDataInterfaceVectorField::PreEditChange(UProperty* PropertyAboutToC
 	// Flush the rendering thread before making any changes to make sure the 
 	// data read by the compute shader isn't subject to a race condition.
 	// TODO(mv): Solve properly using something like a RT Proxy.
-	FlushRenderingCommands();
+	//FlushRenderingCommands();
 }
 
 void UNiagaraDataInterfaceVectorField::PostEditChangeProperty(struct FPropertyChangedEvent& PropertyChangedEvent)
 {
 	Super::PostEditChangeProperty(PropertyChangedEvent);
+	PushToRenderThread();
 }
 #endif //WITH_EDITOR
 
@@ -62,6 +65,8 @@ void UNiagaraDataInterfaceVectorField::PostLoad()
 	{
 		Field->ConditionalPostLoad();
 	}
+
+	PushToRenderThread();
 }
 
 void UNiagaraDataInterfaceVectorField::PostInitProperties()
@@ -331,7 +336,7 @@ struct FNiagaraDataInterfaceParametersCS_VectorField : public FNiagaraDataInterf
 		Ar << MaxBounds;
 	}
 
-	virtual void Set(FRHICommandList& RHICmdList, FNiagaraShader* Shader, class UNiagaraDataInterface* DataInterface, void* PerInstanceData) const override
+	virtual void Set(FRHICommandList& RHICmdList, const FNiagaraDataInterfaceSetArgs& Context) const override
 	{
 		check(IsInRenderingThread());
 
@@ -354,15 +359,13 @@ struct FNiagaraDataInterfaceParametersCS_VectorField : public FNiagaraDataInterf
 		}
 
 		// Get shader and DI
-		const FComputeShaderRHIParamRef ComputeShaderRHI = Shader->GetComputeShader();
-		UNiagaraDataInterfaceVectorField* VFDI = CastChecked<UNiagaraDataInterfaceVectorField>(DataInterface);
+		const FComputeShaderRHIParamRef ComputeShaderRHI = Context.Shader->GetComputeShader();
+		FNiagaraDataInterfaceProxyVectorField* VFDI = static_cast<FNiagaraDataInterfaceProxyVectorField*>(Context.DataInterface);
 		
 		// Note: There is a flush in PreEditChange to make sure everything is synced up at this point 
 
 		// Get and set 3D texture handle from the currently bound vector field.
-		UVectorFieldStatic* StaticVectorField = Cast<UVectorFieldStatic>(VFDI->Field); 
-		FRHITexture* VolumeTextureRHI = StaticVectorField ? (FRHITexture*)StaticVectorField->GetVolumeTextureRef() : (FRHITexture*)GBlackVolumeTexture->TextureRHI;
-		SetTextureParameter(RHICmdList, ComputeShaderRHI, VectorFieldTexture, VolumeTextureRHI); 
+		SetTextureParameter(RHICmdList, ComputeShaderRHI, VectorFieldTexture, VFDI->TextureRHI); 
 		
 		// Get and set sampler state
 		FSamplerStateRHIParamRef SamplerState = SamplerStates[int(VFDI->bTileX) + 2 * int(VFDI->bTileY) + 4 * int(VFDI->bTileZ)];
@@ -370,9 +373,9 @@ struct FNiagaraDataInterfaceParametersCS_VectorField : public FNiagaraDataInterf
 
 		//
 		SetShaderValue(RHICmdList, ComputeShaderRHI, TilingAxes, VFDI->GetTilingAxes());
-		SetShaderValue(RHICmdList, ComputeShaderRHI, Dimensions, VFDI->GetDimensions());
-		SetShaderValue(RHICmdList, ComputeShaderRHI, MinBounds, VFDI->GetMinBounds());
-		SetShaderValue(RHICmdList, ComputeShaderRHI, MaxBounds, VFDI->GetMaxBounds());
+		SetShaderValue(RHICmdList, ComputeShaderRHI, Dimensions, VFDI->Dimensions);
+		SetShaderValue(RHICmdList, ComputeShaderRHI, MinBounds, VFDI->MinBounds);
+		SetShaderValue(RHICmdList, ComputeShaderRHI, MaxBounds, VFDI->MaxBounds);
 	}
 
 private:
@@ -670,4 +673,32 @@ bool UNiagaraDataInterfaceVectorField::CopyToInternal(UNiagaraDataInterface* Des
 	return true;
 }
 
+void UNiagaraDataInterfaceVectorField::PushToRenderThread()
+{
+	FNiagaraDataInterfaceProxyVectorField* RT_Proxy = GetProxyAs<FNiagaraDataInterfaceProxyVectorField>();
+
+	FVector RT_Dimensions = GetDimensions();
+	FVector RT_MinBounds = GetMinBounds();
+	FVector RT_MaxBounds = GetMaxBounds();
+	bool RT_bTileX = bTileX;
+	bool RT_bTileY = bTileY;
+	bool RT_bTileZ = bTileZ;
+
+	// @todo-threadsafety. This would be a race but I'm taking a ref here. Not ideal in the long term.
+	FTextureRHIRef VolumeTextureRHI = Field ? ((UVectorFieldStatic*)Field)->GetVolumeTextureRef() : (FRHITexture*)GBlackVolumeTexture->TextureRHI;
+
+	// Push Updates to Proxy.
+	ENQUEUE_RENDER_COMMAND(FUpdateDIColorCurve)(
+		[RT_Proxy, RT_Dimensions, RT_MinBounds, RT_MaxBounds, RT_bTileX, RT_bTileY, RT_bTileZ, VolumeTextureRHI](FRHICommandListImmediate& RHICmdList)
+	{
+		RT_Proxy->bTileX = RT_bTileX;
+		RT_Proxy->bTileY = RT_bTileY;
+		RT_Proxy->bTileZ = RT_bTileZ;
+		RT_Proxy->Dimensions = RT_Dimensions;
+		RT_Proxy->MinBounds = RT_MinBounds;
+		RT_Proxy->MaxBounds = RT_MaxBounds;
+		RT_Proxy->TextureRHI = VolumeTextureRHI;
+	});
+}
+
 #undef LOCTEXT_NAMESPACE
