Copying //UE4/Dev-Niagara to Dev-Main (//UE4/Dev-Main)

diff --git a/Engine/Plugins/FX/Niagara/Source/Niagara/Private/NiagaraDataSet.cpp b/Engine/Plugins/FX/Niagara/Source/Niagara/Private/NiagaraDataSet.cpp
--- a/Engine/Plugins/FX/Niagara/Source/Niagara/Private/NiagaraDataSet.cpp
+++ b/Engine/Plugins/FX/Niagara/Source/Niagara/Private/NiagaraDataSet.cpp
@@ -11,99 +11,32 @@
 
 DECLARE_CYCLE_STAT(TEXT("InitRenderData"), STAT_InitRenderData, STATGROUP_Niagara);
 
-
 //////////////////////////////////////////////////////////////////////////
-void FNiagaraDataSet::SetShaderParams(FNiagaraShader *Shader, FRHICommandList &CommandList, uint32& WriteBufferIdx, uint32& ReadBufferIdx)
+
+FCriticalSection FNiagaraSharedObject::CritSec;
+TArray<FNiagaraSharedObject*> FNiagaraSharedObject::DeferredDeletionList;
+
+void FNiagaraSharedObject::FlushDeletionList()
 {
-	check(IsInRenderingThread());
-	WriteBufferIdx = INDEX_NONE;
-	ReadBufferIdx = INDEX_NONE;
-	if (Shader->FloatInputBufferParam.IsBound())
+	//Always do this on RT. GPU buffers must be freed on RT and we may as well do CPU frees at the same time.
+	ENQUEUE_RENDER_COMMAND(FlushDeletionListCommand)([&](FRHICommandListImmediate& RHICmdList)
 	{
-		CommandList.TransitionResource(EResourceTransitionAccess::EReadable, EResourceTransitionPipeline::EComputeToCompute, PrevData().GetGPUBufferFloat()->UAV);
-		if (PrevData().GetNumInstancesAllocated() > 0)
+		FScopeLock Lock(&CritSec);//Possibly make this a lock free queue?
+		int32 i = 0;
+		while (i < DeferredDeletionList.Num())
 		{
-			CommandList.SetShaderResourceViewParameter(Shader->GetComputeShader(), Shader->FloatInputBufferParam.GetBaseIndex(), PrevData().GetGPUBufferFloat()->SRV);
-			if (ReadBufferIdx == INDEX_NONE)
+			check(DeferredDeletionList[i] != nullptr);
+			if (DeferredDeletionList[i]->IsInUse() == false)
 			{
-				ReadBufferIdx = GetPrevBufferIdx();
+				delete DeferredDeletionList[i];
+				DeferredDeletionList.RemoveAtSwap(i);
 			}
-		}
-		else
-		{
-			CommandList.SetShaderResourceViewParameter(Shader->GetComputeShader(), Shader->FloatInputBufferParam.GetBaseIndex(), NiagaraRenderer::GetDummyFloatBuffer().SRV);
-		}
-	}
-	if (Shader->IntInputBufferParam.IsBound())
-	{
-		CommandList.TransitionResource(EResourceTransitionAccess::EReadable, EResourceTransitionPipeline::EComputeToCompute, PrevData().GetGPUBufferInt()->UAV);
-		if (PrevData().GetNumInstancesAllocated() > 0)
-		{
-			CommandList.SetShaderResourceViewParameter(Shader->GetComputeShader(), Shader->IntInputBufferParam.GetBaseIndex(), PrevData().GetGPUBufferInt()->SRV);
-			if (ReadBufferIdx == INDEX_NONE)
+			else
 			{
-				ReadBufferIdx = GetPrevBufferIdx();
+				++i;
 			}
 		}
-		else
-		{
-			CommandList.SetShaderResourceViewParameter(Shader->GetComputeShader(), Shader->IntInputBufferParam.GetBaseIndex(), NiagaraRenderer::GetDummyIntBuffer().SRV);
-		}
-	}
-	if (Shader->FloatOutputBufferParam.IsUAVBound())
-	{
-		CommandList.TransitionResource(EResourceTransitionAccess::EWritable, EResourceTransitionPipeline::EGfxToCompute, CurrData().GetGPUBufferFloat()->UAV);
-		CommandList.SetUAVParameter(Shader->GetComputeShader(), Shader->FloatOutputBufferParam.GetUAVIndex(), CurrData().GetGPUBufferFloat()->UAV);
-
-		if (WriteBufferIdx == INDEX_NONE)
-		{
-			WriteBufferIdx = GetCurrBufferIdx();
-		}
-	}
-	if (Shader->IntOutputBufferParam.IsUAVBound())
-	{
-		CommandList.TransitionResource(EResourceTransitionAccess::EWritable, EResourceTransitionPipeline::EGfxToCompute, CurrData().GetGPUBufferInt()->UAV);
-		CommandList.SetUAVParameter(Shader->GetComputeShader(), Shader->IntOutputBufferParam.GetUAVIndex(), CurrData().GetGPUBufferInt()->UAV);
-		if (WriteBufferIdx == INDEX_NONE)
-		{
-			WriteBufferIdx = GetCurrBufferIdx();
-		}
-	}
-
-	if (Shader->ComponentBufferSizeWriteParam.IsBound())
-	{
-		uint32 SafeBufferSize = CurrData().GetFloatStride() / sizeof(float);
-		CommandList.SetShaderParameter(Shader->GetComputeShader(), Shader->ComponentBufferSizeWriteParam.GetBufferIndex(), Shader->ComponentBufferSizeWriteParam.GetBaseIndex(), Shader->ComponentBufferSizeWriteParam.GetNumBytes(), &SafeBufferSize);
-	}
-
-	if (Shader->ComponentBufferSizeReadParam.IsBound())
-	{
-		uint32 SafeBufferSize = PrevData().GetFloatStride() / sizeof(float);
-		CommandList.SetShaderParameter(Shader->GetComputeShader(), Shader->ComponentBufferSizeReadParam.GetBufferIndex(), Shader->ComponentBufferSizeReadParam.GetBaseIndex(), Shader->ComponentBufferSizeReadParam.GetNumBytes(), &SafeBufferSize);
-	}
-}
-
-
-
-void FNiagaraDataSet::UnsetShaderParams(FNiagaraShader *Shader, FRHICommandList &RHICmdList)
-{
-	check(IsInRenderingThread());
-
-	if (Shader->FloatOutputBufferParam.IsUAVBound())
-	{
-#if !PLATFORM_PS4
-		Shader->FloatOutputBufferParam.UnsetUAV(RHICmdList, Shader->GetComputeShader());
-#endif
-		//RHICmdList.TransitionResource(EResourceTransitionAccess::EReadable, EResourceTransitionPipeline::EGfxToGfx, CurrDataRender().GetGPUBufferFloat()->UAV);
-	}
-
-	if (Shader->IntOutputBufferParam.IsUAVBound())
-	{
-#if !PLATFORM_PS4
-		Shader->IntOutputBufferParam.UnsetUAV(RHICmdList, Shader->GetComputeShader());
-#endif
-		//RHICmdList.TransitionResource(EResourceTransitionAccess::EReadable, EResourceTransitionPipeline::EGfxToGfx, CurrDataRender().GetGPUBufferInt()->UAV);
-	}
+	});
 }
 
 //////////////////////////////////////////////////////////////////////////
@@ -123,133 +56,422 @@ static FAutoConsoleVariableRef CVarNiagaraRenderBufferShrinkFactor(
 	ECVF_Default
 );
 
+FNiagaraDataSet::FNiagaraDataSet()
+	: TotalFloatComponents(0)
+	, TotalInt32Components(0)
+	, SimTarget(ENiagaraSimTarget::CPUSim)
+	, bFinalized(0)
+	, bNeedsPersistentIDs(0)
+	, NumFreeIDs(0)
+	, MaxUsedID(0)
+	, IDAcquireTag(0)
+	, CurrentData(nullptr)
+	, DestinationData(nullptr)
+{
+}
+
 FNiagaraDataSet::~FNiagaraDataSet()
 {
 // 	int32 CurrBytes = RenderDataFloat.NumBytes + RenderDataInt.NumBytes;
 // 	DEC_MEMORY_STAT_BY(STAT_NiagaraVBMemory, CurrBytes);
-
+	ReleaseBuffers();
 }
 
-void FNiagaraDataSet::Dump(FNiagaraDataSet& Other, bool bCurr, int32 StartIdx , int32 NumInstances)const
+void FNiagaraDataSet::Init(FNiagaraDataSetID InID, ENiagaraSimTarget InSimTarget, const FString& InDebugName)
 {
-	Other.Reset();
-	Other.Variables = Variables;
-	Other.VariableLayouts = VariableLayouts;
+	Reset();
+	ID = InID;
+	SimTarget = InSimTarget;
+	DebugName = InDebugName;
+}
 
-	const FNiagaraDataBuffer& DataBuffer = bCurr ? CurrData() : PrevData();
-	FNiagaraDataBuffer& OtherDataBuffer = Other.CurrData();
+void FNiagaraDataSet::Reset()
+{
+	ResetBuffers();
+
+	Variables.Empty();
+	VariableLayouts.Empty();
+	bFinalized = false;
+	TotalFloatComponents = 0;
+	TotalInt32Components = 0;
+	bNeedsPersistentIDs = 0;
+}
 
-	
-	if (OtherDataBuffer.GetNumInstancesAllocated() != DataBuffer.GetNumInstancesAllocated())
+void FNiagaraDataSet::ResetBuffers()
+{
+	if (SimTarget == ENiagaraSimTarget::CPUSim)
 	{
-		Other.Finalize();
-		OtherDataBuffer.Allocate(DataBuffer.GetNumInstancesAllocated());
+		ResetBuffersInternal();
 	}
+	else
+	{
+		check(SimTarget == ENiagaraSimTarget::GPUComputeSim);
+		ENQUEUE_RENDER_COMMAND(ResetBuffersCommand)([=](FRHICommandListImmediate& RHICmdList)
+		{
+			ResetBuffersInternal();
+		});
+	}
+}
 
-	DataBuffer.CopyTo(OtherDataBuffer, StartIdx, NumInstances);
+void FNiagaraDataSet::ResetBuffersInternal()
+{
+	CheckCorrectThread();
+
+	CurrentData = nullptr;
+	DestinationData = nullptr;
+
+	FreeIDsTable.Reset();
+	NumFreeIDs = 0;
+	MaxUsedID = INDEX_NONE;
+
+	//Ensure we have a valid current buffer
+	BeginSimulate();
+	EndSimulate();
 }
 
-void FNiagaraDataSet::DumpGPU(FNiagaraDataSet& Other, float* GPUReadBackFloat, int* GPUReadBackInt, int32 StartIdx, int32 NumInstances)const
+void FNiagaraDataSet::ReleaseBuffers()
 {
-	check(IsInRenderingThread());
-	Other.Reset();
-	Other.Variables = Variables;
-	Other.VariableLayouts = VariableLayouts;
+	CheckCorrectThread();
+	if (Data.Num() > 0)
+	{
+		for (FNiagaraDataBuffer* Buffer : Data)
+		{
+			Buffer->Destroy();
+		}
+		Data.Empty();
+	}
+}
 
-	Other.TotalFloatComponents = TotalFloatComponents;
-	Other.TotalInt32Components = TotalInt32Components;
+FNiagaraDataBuffer& FNiagaraDataSet::BeginSimulate()
+{
+	//CheckCorrectThread();
+	check(DestinationData == nullptr);
 
-	const FNiagaraDataBuffer& DataBuffer = CurrData();
-	FNiagaraDataBuffer& OtherDataBuffer = Other.CurrData();
+	//Find a free buffer we can write into.
+	//Linear search but there should only be 2 or three entries.
+	for (FNiagaraDataBuffer* Buffer : Data)
+	{
+		check(Buffer);
+		if (Buffer != CurrentData && Buffer->TryLock())
+		{
+			DestinationData = Buffer;
+			break;
+		}
+	}
 
-	if (OtherDataBuffer.GetNumInstancesAllocated() != DataBuffer.GetNumInstancesAllocated())
+	if (DestinationData == nullptr)
 	{
-		Other.Finalize();
-		OtherDataBuffer.Allocate(DataBuffer.GetNumInstancesAllocated());
+		Data.Add(new FNiagaraDataBuffer(this));
+		DestinationData = Data.Last();
+		verify(DestinationData->TryLock());
+		check(DestinationData->IsBeingWritten());
 	}
 
-	DataBuffer.GPUCopyTo(OtherDataBuffer, GPUReadBackFloat, GPUReadBackInt, StartIdx, NumInstances);
+	DestinationData->SetNumInstances(0);
+	DestinationData->GetIDTable().Reset();
+
+	return GetDestinationDataChecked();
 }
 
+void FNiagaraDataSet::EndSimulate()
+{
+	//CheckCorrectThread();
+	//Destination is now complete so make it the current simulation state.
+	DestinationData->Unlock();
+	check(!DestinationData->IsInUse());
+	CurrentData = DestinationData;
+	DestinationData = nullptr;
+}
 
 
-void FNiagaraDataSet::Dump(bool bCurr, int32 StartIdx, int32 NumInstances)const
+void FNiagaraDataSet::Allocate(int32 NumInstances, bool bMaintainExisting)
 {
-	TArray<FNiagaraVariable> Vars(Variables);
+	check(bFinalized);
+	CheckCorrectThread();
+	check(DestinationData);
 
-	FNiagaraDataSetVariableIteratorConst Itr(*this, StartIdx, bCurr);
-	Itr.AddVariables(Vars);
+	DestinationData->GetIDTable().Reset();
+	if (bMaintainExisting)
+	{
+		CurrentData->CopyTo(*DestinationData);
+	}
+
+	DestinationData->Allocate(NumInstances, bMaintainExisting);
 
-	if (NumInstances == INDEX_NONE)
+#if NIAGARA_NAN_CHECKING
+	CheckForNaNs();
+#endif
+
+	if (bNeedsPersistentIDs)
 	{
-		NumInstances = bCurr ? GetNumInstances() : GetPrevNumInstances();
-		NumInstances -= StartIdx;
+		TArray<int32>& CurrentIDTable = CurrentData->GetIDTable();
+		TArray<int32>& DestinationIDTable = DestinationData->GetIDTable();
+
+		int32 NumUsedIDs = MaxUsedID + 1;
+
+		int32 RequiredIDs = FMath::Max(NumInstances, NumUsedIDs);
+		int32 ExistingNumIDs = CurrentIDTable.Num();
+		int32 NumNewIDs = RequiredIDs - ExistingNumIDs;
+
+		//////////////////////////////////////////////////////////////////////////
+		//TODO: We should replace this with a lock free list that uses just a single table with RequiredIDs elements.
+		//Unused slots in the array can form a linked list so that we need only one array with a Head index for the FreeID list
+		//This will be faster and likely simpler than the current implementation while also working on GPU.
+		//////////////////////////////////////////////////////////////////////////
+		if (RequiredIDs > ExistingNumIDs)
+		{
+			//UE_LOG(LogNiagara, Warning, TEXT("Growing ID Table! OldSize:%d | NewSize:%d"), ExistingNumIDs, RequiredIDs);
+			int32 NewNumIds = RequiredIDs - ExistingNumIDs;
+#if 0
+			DestinationIDTable.SetNumUninitialized(RequiredIDs);
+#else
+			while (DestinationIDTable.Num() < RequiredIDs)
+			{
+				DestinationIDTable.Add(INDEX_NONE);
+			}
+#endif
+
+			//Free ID Table must always be at least as large as the data buffer + it's current size in the case all particles die this frame.
+			FreeIDsTable.AddUninitialized(NumNewIDs);
+
+			//Free table should always have enough room for these new IDs.
+			check(NumFreeIDs + NumNewIDs <= FreeIDsTable.Num());
+
+			//ID Table grows so add any new IDs to the free array. Add in reverse order to maintain a continuous increasing allocation when popping.
+			for (int32 NewFreeID = RequiredIDs - 1; NewFreeID >= ExistingNumIDs; --NewFreeID)
+			{
+				FreeIDsTable[NumFreeIDs++] = NewFreeID;
+			}
+			//UE_LOG(LogNiagara, Warning, TEXT("DataSetAllocate: Adding New Free IDs: %d - "), NewNumIds);
+		}
+#if 0
+		else if (RequiredIDs < ExistingNumIDs >> 1)//Configurable?
+		{
+			//If the max id we use has reduced significantly then we can shrink the tables.
+			//Have to go through the FreeIDs and remove any that are greater than the new table size.
+			//UE_LOG(LogNiagara, Warning, TEXT("DataSetAllocate: Shrinking ID Table! OldSize:%d | NewSize:%d"), ExistingNumIDs, RequiredIDs);
+			for (int32 CheckedFreeID = 0; CheckedFreeID < NumFreeIDs;)
+			{
+				checkSlow(NumFreeIDs <= FreeIDsTable.Num());
+				if (FreeIDsTable[CheckedFreeID] >= RequiredIDs)
+				{
+					//UE_LOG(LogNiagara, Warning, TEXT("RemoveSwap FreeID: Removed:%d | Swapped:%d"), FreeIDsTable[CheckedFreeID], FreeIDsTable.Last());		
+					int32 FreeIDIndex = --NumFreeIDs;
+					FreeIDsTable[CheckedFreeID] = FreeIDsTable[FreeIDIndex];
+					FreeIDsTable[FreeIDIndex] = INDEX_NONE;
+				}
+				else
+				{
+					++CheckedFreeID;
+				}
+			}
+
+			check(NumFreeIDs <= RequiredIDs);
+			FreeIDsTable.SetNumUninitialized(NumFreeIDs);
+		}
+#endif
+		else
+		{
+			//Drop in required size not great enough so just allocate same size.
+			RequiredIDs = ExistingNumIDs;
+		}
+
+		DestinationIDTable.SetNumUninitialized(RequiredIDs);
+		MaxUsedID = INDEX_NONE;//reset the max ID ready for it to be filled in during simulation.
+
+// 		UE_LOG(LogNiagara, Warning, TEXT("DataSetAllocate: NumInstances:%d | ID Table Size:%d | NumFreeIDs:%d | FreeTableSize:%d"), NumInstances, DestinationData->GetIDTable().Num(), NumFreeIDs, FreeIDsTable.Num());
+// 		UE_LOG(LogNiagara, Warning, TEXT("== FreeIDs %d =="), NumFreeIDs);
+// 		for (int32 i=0; i< NumFreeIDs;++i)
+// 		{
+// 			UE_LOG(LogNiagara, Warning, TEXT("%d"), FreeIDsTable[i]);
+// 		}
 	}
+}
 
-	int32 NumInstancesDumped = 0;
-	TArray<FString> Lines;
-	Lines.Reserve(GetNumInstances());
-	while (Itr.IsValid() && NumInstancesDumped < NumInstances)
+uint32 FNiagaraDataSet::GetSizeBytes()const
+{
+	uint32 Size = 0;
+	for (FNiagaraDataBuffer* Buffer : Data)
 	{
-		Itr.Get();
+		check(Buffer);
+		Size += Buffer->GetSizeBytes();
+	}
+	return Size;
+}
 
-		FString Line = TEXT("| ");
-		for (FNiagaraVariable& Var : Vars)
+void FNiagaraDataSet::ClearRegisterTable(uint8** Registers, int32& NumRegisters)
+{
+	for (FNiagaraVariableLayoutInfo& VarLayout : VariableLayouts)
+	{
+		int32 NumComps = VarLayout.GetNumFloatComponents() + VarLayout.GetNumInt32Components();
+		for (int32 CompIdx = 0; CompIdx < NumComps; ++CompIdx)
 		{
-			Line += Var.ToString() + TEXT(" | ");
+			Registers[NumRegisters + CompIdx] = nullptr;
 		}
-		Lines.Add(Line);
-		Itr.Advance();
-		NumInstancesDumped++;
+		NumRegisters += NumComps;
 	}
+}
 
-	static FString Sep;
-	if (Sep.Len() == 0)
+
+void FNiagaraDataSet::CheckForNaNs()const
+{
+	for (const FNiagaraDataBuffer* Buffer : Data)
 	{
-		for (int32 i = 0; i < 50; ++i)
+		if (Buffer->CheckForNaNs())
 		{
-			Sep.AppendChar(TEXT('='));
+			Buffer->Dump(0, Buffer->GetNumInstances(), TEXT("Found Niagara buffer containing NaNs!"));
+			ensureAlwaysMsgf(false, TEXT("NiagaraDataSet contains NaNs!"));
 		}
 	}
+}
 
-	UE_LOG(LogNiagara, Log, TEXT("%s"), *Sep);
-	UE_LOG(LogNiagara, Log, TEXT(" Buffer: %d"), CurrBuffer);
- 	UE_LOG(LogNiagara, Log, TEXT("%s"), *Sep);
-// 	UE_LOG(LogNiagara, Log, TEXT("%s"), *HeaderStr);
-// 	UE_LOG(LogNiagara, Log, TEXT("%s"), *Sep);
-	for (FString& Str : Lines)
+void FNiagaraDataSet::Dump(int32 StartIndex, int32 NumInstances, const FString& Label)const
+{
+	if (CurrentData)
 	{
-		UE_LOG(LogNiagara, Log, TEXT("%s"), *Str);
+		CurrentData->Dump(StartIndex, NumInstances, Label);
 	}
-	UE_LOG(LogNiagara, Log, TEXT("%s"), *Sep);
+
+	if (GetDestinationData())
+	{
+		FString DestLabel = Label + TEXT("[Destination]");
+		DestinationData->Dump(StartIndex, NumInstances, DestLabel);
+	}
+}
+
+void FNiagaraDataSet::BuildLayout()
+{
+	VariableLayouts.Empty();
+	TotalFloatComponents = 0;
+	TotalInt32Components = 0;
+
+	VariableLayouts.Reserve(Variables.Num());
+	for (FNiagaraVariable& Var : Variables)
+	{
+		FNiagaraVariableLayoutInfo& VarInfo = VariableLayouts[VariableLayouts.AddDefaulted()];
+		FNiagaraTypeLayoutInfo::GenerateLayoutInfo(VarInfo.LayoutInfo, Var.GetType().GetScriptStruct());
+		VarInfo.FloatComponentStart = TotalFloatComponents;
+		VarInfo.Int32ComponentStart = TotalInt32Components;
+		TotalFloatComponents += VarInfo.GetNumFloatComponents();
+		TotalInt32Components += VarInfo.GetNumInt32Components();
+	}
+}
+
+void FNiagaraDataSet::AddVariable(FNiagaraVariable& Variable)
+{
+	check(!bFinalized);
+	Variables.AddUnique(Variable);
 }
 
+void FNiagaraDataSet::AddVariables(const TArray<FNiagaraVariable>& Vars)
+{
+	check(!bFinalized);
+	for (const FNiagaraVariable& Var : Vars)
+	{
+		Variables.AddUnique(Var);
+	}
+}
+
+void FNiagaraDataSet::Finalize()
+{
+	check(!bFinalized);
+	bFinalized = true;
+	BuildLayout();
+
+	ResetBuffers();
+}
+
+const FNiagaraVariableLayoutInfo* FNiagaraDataSet::GetVariableLayout(const FNiagaraVariable& Var)const
+{
+	int32 VarLayoutIndex = Variables.IndexOfByKey(Var);
+	return VarLayoutIndex != INDEX_NONE ? &VariableLayouts[VarLayoutIndex] : nullptr;
+}
+
+bool FNiagaraDataSet::GetVariableComponentOffsets(const FNiagaraVariable& Var, int32 &FloatStart, int32 &IntStart) const
+{
+	const FNiagaraVariableLayoutInfo *Info = GetVariableLayout(Var);
+	if (Info)
+	{
+		FloatStart = Info->FloatComponentStart;
+		IntStart = Info->Int32ComponentStart;
+		return true;
+	}
+
+	FloatStart = -1;
+	IntStart = -1;
+	return false;
+}
+
+void FNiagaraDataSet::CopyTo(FNiagaraDataSet& Other, int32 StartIdx, int32 NumInstances)const
+{
+	CheckCorrectThread();
+	Other.Reset();
+	Other.Variables = Variables;
+	Other.VariableLayouts = VariableLayouts;
+	Other.TotalFloatComponents = TotalFloatComponents;
+	Other.TotalInt32Components = TotalInt32Components;
+	Other.Finalize();
+
+	//Read the most current data. Even if it's possibly partially complete simulation data.
+	FNiagaraDataBuffer& OtherDataBuffer = Other.BeginSimulate();
+
+	FNiagaraDataBuffer* SourceBuffer = GetDestinationData() ? GetDestinationData() : GetCurrentData();
+	if (SourceBuffer != nullptr)
+	{
+		OtherDataBuffer.Allocate(SourceBuffer->GetNumInstancesAllocated());
+		SourceBuffer->CopyTo(OtherDataBuffer, StartIdx, NumInstances);
+	}
+
+	Other.EndSimulate();
+}
+
+void FNiagaraDataSet::CopyFromGPUReadback(float* GPUReadBackFloat, int* GPUReadBackInt, int32 StartIdx /* = 0 */, int32 NumInstances /* = INDEX_NONE */)
+{
+	check(IsInRenderingThread());
+	check(bFinalized);//We should be finalized with proper layout information already.
+
+	FNiagaraDataBuffer& DestBuffer = BeginSimulate();
+	DestBuffer.GPUCopyFrom(GPUReadBackFloat, GPUReadBackInt, StartIdx, NumInstances);
+	EndSimulate();
+}
+ 
 //////////////////////////////////////////////////////////////////////////
 
-void FNiagaraDataBuffer::Init(FNiagaraDataSet* InOwner)
+FNiagaraDataBuffer::FNiagaraDataBuffer(FNiagaraDataSet* InOwner)
+	: Owner(InOwner)
+	, NumChunksAllocatedForGPU(0)
+	, NumInstances(0)
+	, NumInstancesAllocated(0)
+	, FloatStride(0)
+	, Int32Stride(0)
 {
-	Owner = InOwner;
 }
 
 FNiagaraDataBuffer::~FNiagaraDataBuffer()
 {
-	Reset();
+	check(!IsInUse());
+	DEC_MEMORY_STAT_BY(STAT_NiagaraParticleMemory, FloatData.Max() + Int32Data.Max());
 }
 
-void FNiagaraDataBuffer::Reset()
+void FNiagaraDataBuffer::CheckUsage(bool bReadOnly)const
 {
-	DEC_MEMORY_STAT_BY(STAT_NiagaraParticleMemory, FloatData.Max() + Int32Data.Max());
-	FloatData.Empty();
-	Int32Data.Empty();
-	FloatStride = 0;
-	Int32Stride = 0;
-	NumInstances = 0;
-	NumInstancesAllocated = 0;
-	NumChunksAllocatedForGPU = 0;
+	check(Owner);
+	if (Owner->SimTarget == ENiagaraSimTarget::CPUSim)
+	{
+		//We can read on the RT but any modifications must be GT (or GT Task).
+		check(bReadOnly || !IsInRenderingThread());
+	}
+	else
+	{
+		check(Owner->SimTarget == ENiagaraSimTarget::GPUComputeSim);
+		//Everything other than init for GPU sims should be done on the RT.
+		check(IsInRenderingThread());
+	}
 }
 
 int32 FNiagaraDataBuffer::TransferInstance(FNiagaraDataBuffer& SourceBuffer, int32 InstanceIndex)
 {
+	CheckUsage(false);
 	if (SourceBuffer.GetNumInstances() > (uint32)InstanceIndex)
 	{
 		int32 OldNumInstances = NumInstances;
@@ -283,6 +505,7 @@ int32 FNiagaraDataBuffer::TransferInstance(FNiagaraDataBuffer& SourceBuffer, int
 
 bool FNiagaraDataBuffer::CheckForNaNs()const
 {
+	CheckUsage(true);
 	bool bContainsNaNs = false;
 	int32 NumFloatComponents = Owner->GetNumFloatComponents();
 	for (int32 CompIdx = 0; CompIdx < NumFloatComponents && !bContainsNaNs; ++CompIdx)
@@ -299,81 +522,95 @@ bool FNiagaraDataBuffer::CheckForNaNs()const
 
 void FNiagaraDataBuffer::Allocate(uint32 InNumInstances, bool bMaintainExisting)
 {
-	check(Owner);
-	if (Owner->SimTarget == ENiagaraSimTarget::CPUSim)
-	{
-		NumInstancesAllocated = InNumInstances;
-		NumInstances = 0;
+	//CheckUsage(false);
+	check(Owner->SimTarget == ENiagaraSimTarget::CPUSim);
+
+	NumInstancesAllocated = InNumInstances;
+	NumInstances = 0;
 
-		DEC_MEMORY_STAT_BY(STAT_NiagaraParticleMemory, FloatData.Max() + Int32Data.Max());
+	DEC_MEMORY_STAT_BY(STAT_NiagaraParticleMemory, FloatData.Max() + Int32Data.Max());
 
-		int32 OldFloatStride = FloatStride;
-		TArray<uint8> OldFloatData;
-		int32 OldInt32Stride = Int32Stride;
-		TArray<uint8> OldIntData;
+	const uint32 OldFloatStride = FloatStride;
+	TArray<uint8> OldFloatData;
+	const uint32 OldInt32Stride = Int32Stride;
+	TArray<uint8> OldIntData;
 
-		if (bMaintainExisting)
-		{
-			//Need to copy off old data so we can copy it back into the newly laid out buffers. TODO: Avoid this needless copying.
-			OldFloatData = FloatData;
-			OldIntData = Int32Data;
-		}
+	if (bMaintainExisting)
+	{
+		//Need to copy off old data so we can copy it back into the newly laid out buffers. TODO: Avoid this needless copying.
+		OldFloatData = FloatData;
+		OldIntData = Int32Data;
+	}
 
-		FloatStride = GetSafeComponentBufferSize(NumInstancesAllocated * sizeof(float));
-		FloatData.SetNum(FloatStride * Owner->GetNumFloatComponents(), false);
+	FloatStride = GetSafeComponentBufferSize(NumInstancesAllocated * sizeof(float));
+	FloatData.SetNum(FloatStride * Owner->GetNumFloatComponents(), false);
 
-		Int32Stride = GetSafeComponentBufferSize(NumInstancesAllocated * sizeof(int32));
-		Int32Data.SetNum(Int32Stride * Owner->GetNumInt32Components(), false);
+	Int32Stride = GetSafeComponentBufferSize(NumInstancesAllocated * sizeof(int32));
+	Int32Data.SetNum(Int32Stride * Owner->GetNumInt32Components(), false);
 
-		INC_MEMORY_STAT_BY(STAT_NiagaraParticleMemory, FloatData.Max() + Int32Data.Max());
+	INC_MEMORY_STAT_BY(STAT_NiagaraParticleMemory, FloatData.Max() + Int32Data.Max());
 
-		//In some cases we want the existing data in the buffer to be maintained which due to the data layout requires some fix up.
-		if (bMaintainExisting)
+	//In some cases we want the existing data in the buffer to be maintained which due to the data layout requires some fix up.
+	if (bMaintainExisting)
+	{
+		if (FloatStride != OldFloatStride && FloatStride > 0 && OldFloatStride > 0)
 		{
-			if (FloatStride != OldFloatStride && FloatStride > 0 && OldFloatStride > 0)
+			const uint32 BytesToCopy = FMath::Min(OldFloatStride, FloatStride);
+			for (int32 CompIdx = (int32)Owner->TotalFloatComponents-1; CompIdx >= 0; --CompIdx)
 			{
-				for (int32 CompIdx = (int32)Owner->TotalFloatComponents-1; CompIdx >= 0; --CompIdx)
-				{
-					uint8* Src = OldFloatData.GetData() + OldFloatStride * CompIdx;
-					uint8* Dst = FloatData.GetData() + FloatStride * CompIdx;
-					FMemory::Memcpy(Dst, Src, OldFloatStride);
-				}
+				uint8* Src = OldFloatData.GetData() + OldFloatStride * CompIdx;
+				uint8* Dst = FloatData.GetData() + FloatStride * CompIdx;
+				FMemory::Memcpy(Dst, Src, BytesToCopy);
 			}
-			if (Int32Stride != OldInt32Stride && Int32Stride > 0 && OldInt32Stride > 0)
+		}
+		if (Int32Stride != OldInt32Stride && Int32Stride > 0 && OldInt32Stride > 0)
+		{
+			const uint32 BytesToCopy = FMath::Min(OldInt32Stride, Int32Stride);
+			for (int32 CompIdx = (int32)Owner->TotalInt32Components - 1; CompIdx >= 0; --CompIdx)
 			{
-				for (int32 CompIdx = (int32)Owner->TotalInt32Components - 1; CompIdx >= 0; --CompIdx)
-				{
-					uint8* Src = OldIntData.GetData() + OldInt32Stride * CompIdx;
-					uint8* Dst = Int32Data.GetData() + Int32Stride * CompIdx;
-					FMemory::Memcpy(Dst, Src, OldInt32Stride);
-				}
+				uint8* Src = OldIntData.GetData() + OldInt32Stride * CompIdx;
+				uint8* Dst = Int32Data.GetData() + Int32Stride * CompIdx;
+				FMemory::Memcpy(Dst, Src, BytesToCopy);
 			}
 		}
 	}
+	else
+	{
+		IDToIndexTable.Reset();
+	}
 }
 
-
 void FNiagaraDataBuffer::AllocateGPU(uint32 InNumInstances, FRHICommandList &RHICmdList)
 {
-	if (Owner == 0)
+	CheckUsage(false);
+
+	check(Owner->SimTarget == ENiagaraSimTarget::GPUComputeSim);
+
+	//Allocate the indices buffer if needed
+	if (GPUIndices.Buffer == nullptr)
 	{
-		return;
+		// Use BUF_KeepCPUAccessible here since some platforms will lock it for readonly (depending on the implementation of RHIEnqueueStagedRead) after GPU simulation.
+		GPUIndices.Initialize(sizeof(int32), 64 /*Context->NumDataSets*/, EPixelFormat::PF_R32_UINT, BUF_DrawIndirect | BUF_Static | BUF_KeepCPUAccessible);	// always allocate for up to 64 data sets
 	}
-	check(IsInRenderingThread());
+
+	// ALLOC_CHUNKSIZE must be greater than zero and divisible by the thread group size
 	const uint32 ALLOC_CHUNKSIZE = 4096;
+	static_assert((ALLOC_CHUNKSIZE > 0) && ((ALLOC_CHUNKSIZE % NIAGARA_COMPUTE_THREADGROUP_SIZE) == 0), "ALLOC_CHUNKSIZE must be divisible by NIAGARA_COMPUTE_THREADGROUP_SIZE");
 
 	NumInstancesAllocated = InNumInstances;
 
-	uint32 PaddedNumInstances = ((InNumInstances + NIAGARA_COMPUTE_THREADGROUP_SIZE -1) / NIAGARA_COMPUTE_THREADGROUP_SIZE) * NIAGARA_COMPUTE_THREADGROUP_SIZE;
+	// Round the count up to the nearest threadgroup size
+	const uint32 PaddedNumInstances = FMath::DivideAndRoundUp(NumInstancesAllocated, NIAGARA_COMPUTE_THREADGROUP_SIZE) * NIAGARA_COMPUTE_THREADGROUP_SIZE;
+
+	// Pack the data so that the space between elements is the padded thread group size
 	FloatStride = PaddedNumInstances * sizeof(float);
 	Int32Stride = PaddedNumInstances * sizeof(int32);
 
-	if (NumInstancesAllocated > NumChunksAllocatedForGPU * ALLOC_CHUNKSIZE)
+	// When the number of elements that we are going to need is greater than the number we have reserved, we need to expand it.
+	if (PaddedNumInstances > NumChunksAllocatedForGPU * ALLOC_CHUNKSIZE)
 	{
-		uint32 PrevChunks = NumChunksAllocatedForGPU;
-		NumChunksAllocatedForGPU = ((InNumInstances + ALLOC_CHUNKSIZE - 1) / ALLOC_CHUNKSIZE);
-		uint32 NumElementsToAlloc = NumChunksAllocatedForGPU * ALLOC_CHUNKSIZE;
-
+		NumChunksAllocatedForGPU = FMath::DivideAndRoundUp(PaddedNumInstances, ALLOC_CHUNKSIZE);
+		const uint32 NumElementsToAlloc = NumChunksAllocatedForGPU * ALLOC_CHUNKSIZE;
 		if (NumElementsToAlloc == 0)
 		{
 			return;
@@ -385,7 +622,7 @@ void FNiagaraDataBuffer::AllocateGPU(uint32 InNumInstances, FRHICommandList &RHI
 			{
 				GPUBufferFloat.Release();
 			}
-			GPUBufferFloat.Initialize(sizeof(float), NumElementsToAlloc * Owner->GetNumFloatComponents(), EPixelFormat::PF_R32_FLOAT, BUF_Static);
+			GPUBufferFloat.Initialize(sizeof(float), NumElementsToAlloc * Owner->GetNumFloatComponents(), EPixelFormat::PF_R32_FLOAT, BUF_Static, *Owner->DebugName );
 		}
 		if (Owner->GetNumInt32Components())
 		{
@@ -393,13 +630,15 @@ void FNiagaraDataBuffer::AllocateGPU(uint32 InNumInstances, FRHICommandList &RHI
 			{
 				GPUBufferInt.Release();
 			}
-			GPUBufferInt.Initialize(sizeof(int32), NumElementsToAlloc * Owner->GetNumInt32Components(), EPixelFormat::PF_R32_SINT, BUF_Static);
+			GPUBufferInt.Initialize(sizeof(int32), NumElementsToAlloc * Owner->GetNumInt32Components(), EPixelFormat::PF_R32_SINT, BUF_Static, *Owner->DebugName);
 		}
 	}
 }
 
 void FNiagaraDataBuffer::SwapInstances(uint32 OldIndex, uint32 NewIndex) 
 {
+	CheckUsage(false);
+
 	for (uint32 CompIdx = 0; CompIdx < Owner->TotalFloatComponents; ++CompIdx)
 	{
 		float* Src = GetInstancePtrFloat(CompIdx, OldIndex);
@@ -420,6 +659,7 @@ void FNiagaraDataBuffer::SwapInstances(uint32 OldIndex, uint32 NewIndex)
 
 void FNiagaraDataBuffer::KillInstance(uint32 InstanceIdx)
 {
+	CheckUsage(false);
 	check(InstanceIdx < NumInstances);
 	--NumInstances;
 
@@ -443,11 +683,13 @@ void FNiagaraDataBuffer::KillInstance(uint32 InstanceIdx)
 
 void FNiagaraDataBuffer::CopyTo(FNiagaraDataBuffer& DestBuffer, int32 InStartIdx, int32 InNumInstances)const
 {
+	CheckUsage(false);
+	
 	if (InStartIdx < 0 || (uint32)InStartIdx > NumInstances)
 	{
 		InStartIdx = NumInstances;
 	}
-	if (InNumInstances < 0 || ((uint32)InNumInstances + (uint32)InStartIdx) > NumInstances)
+	if (InNumInstances == INDEX_NONE || ((uint32)InNumInstances + (uint32)InStartIdx) > NumInstances)
 	{
 		InNumInstances = NumInstances - InStartIdx;
 	}
@@ -495,73 +737,317 @@ void FNiagaraDataBuffer::CopyTo(FNiagaraDataBuffer& DestBuffer, int32 InStartIdx
 	}
 }
 
-void FNiagaraDataBuffer::GPUCopyTo(FNiagaraDataBuffer& DestBuffer, float* GPUReadBackFloat, int* GPUReadBackInt, int32 InStartIdx, int32 InNumInstances)const
+void FNiagaraDataBuffer::GPUCopyFrom(float* GPUReadBackFloat, int* GPUReadBackInt, int32 InStartIdx, int32 InNumInstances)
 {
-	if (InStartIdx < 0 || (uint32)InStartIdx > NumInstances)
-	{
-		InStartIdx = NumInstances;
-	}
-	if (InNumInstances < 0 || ((uint32)InNumInstances + (uint32)InStartIdx) > NumInstances)
+	//CheckUsage(false); //Have to disable this as in this specific case we write to a "CPUSim" from the RT.
+
+	if (InNumInstances <= 0)
 	{
-		InNumInstances = NumInstances - InStartIdx;
+		return;
 	}
 
-	if (InNumInstances != 0)
+	Allocate(InNumInstances);
+	SetNumInstances(InNumInstances);
+
+	if (GPUReadBackFloat)
 	{
-		if (DestBuffer.NumInstancesAllocated != NumInstancesAllocated)
+		for (uint32 CompIdx = 0; CompIdx < Owner->TotalFloatComponents; ++CompIdx)
 		{
-			DestBuffer.Allocate(NumInstancesAllocated);
-		}
+			const float* SrcStart = GetInstancePtrFloat(GPUReadBackFloat, CompIdx, InStartIdx);
+			const float* SrcEnd = GetInstancePtrFloat(GPUReadBackFloat, CompIdx, InStartIdx + InNumInstances);
+			float* Dst = GetInstancePtrFloat(CompIdx, 0);
+			size_t Count = SrcEnd - SrcStart;
+			FMemory::Memcpy(Dst, SrcStart, Count * sizeof(float));
 
-		if (GPUReadBackFloat)
-		{
-			for (uint32 CompIdx = 0; CompIdx < Owner->TotalFloatComponents; ++CompIdx)
+			if (Count > 0)
 			{
-				const float* SrcStart = GetInstancePtrFloat(GPUReadBackFloat, CompIdx, InStartIdx);
-				const float* SrcEnd = GetInstancePtrFloat(GPUReadBackFloat, CompIdx, InStartIdx + InNumInstances);
-				float* Dst = DestBuffer.GetInstancePtrFloat(CompIdx, 0);
-				size_t Count = SrcEnd - SrcStart;
-				FMemory::Memcpy(Dst, SrcStart, Count * sizeof(float));
-
-				if (Count > 0)
+				for (size_t i = 0; i < Count; i++)
 				{
-					for (size_t i = 0; i < Count; i++)
-					{
-						check(SrcStart[i] == Dst[i]);
-					}
+					check(SrcStart[i] == Dst[i]);
 				}
 			}
 		}
-		if (GPUReadBackInt)
+	}
+	if (GPUReadBackInt)
+	{
+		for (uint32 CompIdx = 0; CompIdx < Owner->TotalInt32Components; ++CompIdx)
 		{
-			for (uint32 CompIdx = 0; CompIdx < Owner->TotalInt32Components; ++CompIdx)
-			{
-				const int32* SrcStart = GetInstancePtrInt32(GPUReadBackInt, CompIdx, InStartIdx);
-				const int32* SrcEnd = GetInstancePtrInt32(GPUReadBackInt, CompIdx, InStartIdx + InNumInstances);
-				int32* Dst = DestBuffer.GetInstancePtrInt32(CompIdx, 0);
-				size_t Count = SrcEnd - SrcStart;
-				FMemory::Memcpy(Dst, SrcStart, Count * sizeof(int32));
+			const int32* SrcStart = GetInstancePtrInt32(GPUReadBackInt, CompIdx, InStartIdx);
+			const int32* SrcEnd = GetInstancePtrInt32(GPUReadBackInt, CompIdx, InStartIdx + InNumInstances);
+			int32* Dst = GetInstancePtrInt32(CompIdx, 0);
+			size_t Count = SrcEnd - SrcStart;
+			FMemory::Memcpy(Dst, SrcStart, Count * sizeof(int32));
 
-				if (Count > 0)
+			if (Count > 0)
+			{
+				for (size_t i = 0; i < Count; i++)
 				{
-					for (size_t i = 0; i < Count; i++)
-					{
-						check(SrcStart[i] == Dst[i]);
-					}
+					check(SrcStart[i] == Dst[i]);
 				}
 			}
 		}
-		DestBuffer.SetNumInstances(InNumInstances);
 	}
 }
 
 void FNiagaraDataBuffer::CopyTo(FNiagaraDataBuffer& DestBuffer)const
 {
+	CheckUsage(true);
+	DestBuffer.CheckUsage(false);
 	DestBuffer.FloatStride = FloatStride;
 	DestBuffer.FloatData = FloatData;
 	DestBuffer.Int32Stride = Int32Stride;
 	DestBuffer.Int32Data = Int32Data;
 	DestBuffer.NumInstancesAllocated = NumInstancesAllocated;
 	DestBuffer.NumInstances = NumInstances;
+	DestBuffer.IDToIndexTable = IDToIndexTable;
 }
 
+void FNiagaraDataBuffer::Dump(int32 StartIndex, int32 InNumInstances, const FString& Label)const
+{
+	TArray<FNiagaraVariable>& Variables = Owner->GetVariables();
+	FNiagaraDataVariableIterator Itr(this, StartIndex);
+	Itr.AddVariables(Variables);
+
+	if (InNumInstances == INDEX_NONE)
+	{
+		InNumInstances = GetNumInstances();
+		InNumInstances -= StartIndex;
+	}
+
+	int32 NumInstancesDumped = 0;
+	TArray<FString> Lines;
+	Lines.Reserve(GetNumInstances());
+	while (Itr.IsValid() && NumInstancesDumped < InNumInstances)
+	{
+		Itr.Get();
+
+		FString Line = TEXT("| ");
+		for (FNiagaraVariable& Var : Owner->GetVariables())
+		{
+			Line += Var.ToString() + TEXT(" | ");
+		}
+		Lines.Add(Line);
+		Itr.Advance();
+		NumInstancesDumped++;
+	}
+
+	static FString Sep;
+	if (Sep.Len() == 0)
+	{
+		for (int32 i = 0; i < 50; ++i)
+		{
+			Sep.AppendChar(TEXT('='));
+		}
+	}
+
+	UE_LOG(LogNiagara, Log, TEXT("%s"), *Sep);
+	UE_LOG(LogNiagara, Log, TEXT(" %s "), *Label);
+	UE_LOG(LogNiagara, Log, TEXT("%s"), *Sep);
+	// 	UE_LOG(LogNiagara, Log, TEXT("%s"), *HeaderStr);
+	// 	UE_LOG(LogNiagara, Log, TEXT("%s"), *Sep);
+	for (FString& Str : Lines)
+	{
+		UE_LOG(LogNiagara, Log, TEXT("%s"), *Str);
+	}
+	if (IDToIndexTable.Num() > 0)
+	{
+		UE_LOG(LogNiagara, Log, TEXT("== ID Table =="), *Sep);
+		for (int32 i = 0; i < IDToIndexTable.Num(); ++i)
+		{
+			UE_LOG(LogNiagara, Log, TEXT("%d = %d"), i, IDToIndexTable[i]);
+		}
+	}
+	UE_LOG(LogNiagara, Log, TEXT("%s"), *Sep);
+
+}
+
+bool FNiagaraDataBuffer::AppendToRegisterTable(uint8** Registers, int32& NumRegisters, int32 StartInstance)
+{
+	check(Owner && Owner->IsInitialized());
+	check(Owner->GetSimTarget() == ENiagaraSimTarget::CPUSim);
+	CheckUsage(true);
+
+	for (FNiagaraVariableLayoutInfo& VarLayout : Owner->VariableLayouts)
+	{
+		int32 NumFloats = VarLayout.GetNumFloatComponents();
+		int32 NumInts = VarLayout.GetNumInt32Components();
+		for (int32 CompIdx = 0; CompIdx < NumFloats; ++CompIdx)
+		{
+			uint32 CompBufferOffset = VarLayout.FloatComponentStart + CompIdx;
+			uint32 CompRegisterOffset = VarLayout.LayoutInfo.FloatComponentRegisterOffsets[CompIdx];
+			Registers[NumRegisters + CompRegisterOffset] = (uint8*)GetInstancePtrFloat(CompBufferOffset, StartInstance);
+		}
+		for (int32 CompIdx = 0; CompIdx < NumInts; ++CompIdx)
+		{
+			uint32 CompBufferOffset = VarLayout.Int32ComponentStart + CompIdx;
+			uint32 CompRegisterOffset = VarLayout.LayoutInfo.Int32ComponentRegisterOffsets[CompIdx];
+			Registers[NumRegisters + CompRegisterOffset] = (uint8*)GetInstancePtrInt32(CompBufferOffset, StartInstance);
+		}
+		NumRegisters += NumFloats + NumInts;
+	}
+	return true;
+}
+
+/////////////////////////////////////////////////////////////////////////
+
+template<bool bDoResourceTransitions>
+void FNiagaraDataBuffer::SetShaderParams(FNiagaraShader *Shader, FRHICommandList &CommandList, bool bInput)
+{
+	check(IsInRenderingThread());
+	if (bInput)
+	{
+		if (Shader->FloatInputBufferParam.IsBound())
+		{
+			if (bDoResourceTransitions)
+			{
+				CommandList.TransitionResource(EResourceTransitionAccess::EReadable, EResourceTransitionPipeline::EComputeToCompute, GetGPUBufferFloat().UAV);
+			}
+			if (GetNumInstancesAllocated() > 0)
+			{
+				CommandList.SetShaderResourceViewParameter(Shader->GetComputeShader(), Shader->FloatInputBufferParam.GetBaseIndex(), GetGPUBufferFloat().SRV);
+			}
+			else
+			{
+				CommandList.SetShaderResourceViewParameter(Shader->GetComputeShader(), Shader->FloatInputBufferParam.GetBaseIndex(), FNiagaraRenderer::GetDummyFloatBuffer().SRV);
+			}
+		}
+
+		if (Shader->IntInputBufferParam.IsBound())
+		{
+			if (bDoResourceTransitions)
+			{
+				CommandList.TransitionResource(EResourceTransitionAccess::EReadable, EResourceTransitionPipeline::EComputeToCompute, GetGPUBufferInt().UAV);
+			}
+			if (GetNumInstancesAllocated() > 0)
+			{
+				CommandList.SetShaderResourceViewParameter(Shader->GetComputeShader(), Shader->IntInputBufferParam.GetBaseIndex(), GetGPUBufferInt().SRV);
+			}
+			else
+			{
+				CommandList.SetShaderResourceViewParameter(Shader->GetComputeShader(), Shader->IntInputBufferParam.GetBaseIndex(), FNiagaraRenderer::GetDummyIntBuffer().SRV);
+			}
+		}
+
+		if (Shader->ComponentBufferSizeReadParam.IsBound())
+		{
+			uint32 SafeBufferSize = GetFloatStride() / sizeof(float);
+			CommandList.SetShaderParameter(Shader->GetComputeShader(), Shader->ComponentBufferSizeReadParam.GetBufferIndex(), Shader->ComponentBufferSizeReadParam.GetBaseIndex(), Shader->ComponentBufferSizeReadParam.GetNumBytes(), &SafeBufferSize);
+		}
+	}
+	else
+	{
+		if (Shader->FloatOutputBufferParam.IsUAVBound())
+		{
+			if (bDoResourceTransitions)
+			{
+				CommandList.TransitionResource(EResourceTransitionAccess::EWritable, EResourceTransitionPipeline::EGfxToCompute, GetGPUBufferFloat().UAV);
+			}
+			CommandList.SetUAVParameter(Shader->GetComputeShader(), Shader->FloatOutputBufferParam.GetUAVIndex(), GetGPUBufferFloat().UAV);
+		}
+
+		if (Shader->IntOutputBufferParam.IsUAVBound())
+		{
+			if (bDoResourceTransitions)
+			{
+				CommandList.TransitionResource(EResourceTransitionAccess::EWritable, EResourceTransitionPipeline::EGfxToCompute, GetGPUBufferInt().UAV);
+			}
+			CommandList.SetUAVParameter(Shader->GetComputeShader(), Shader->IntOutputBufferParam.GetUAVIndex(), GetGPUBufferInt().UAV);
+		}
+
+		if (Shader->ComponentBufferSizeWriteParam.IsBound())
+		{
+			uint32 SafeBufferSize = GetFloatStride() / sizeof(float);
+			CommandList.SetShaderParameter(Shader->GetComputeShader(), Shader->ComponentBufferSizeWriteParam.GetBufferIndex(), Shader->ComponentBufferSizeWriteParam.GetBaseIndex(), Shader->ComponentBufferSizeWriteParam.GetNumBytes(), &SafeBufferSize);
+		}
+	}
+}
+
+template void FNiagaraDataBuffer::SetShaderParams<true>(FNiagaraShader*, FRHICommandList&, bool);
+template void FNiagaraDataBuffer::SetShaderParams<false>(FNiagaraShader*, FRHICommandList&, bool);
+
+
+void FNiagaraDataBuffer::UnsetShaderParams(FNiagaraShader *Shader, FRHICommandList &RHICmdList)
+{
+	check(IsInRenderingThread());
+
+	if (Shader->FloatOutputBufferParam.IsUAVBound())
+	{
+#if !PLATFORM_PS4
+		Shader->FloatOutputBufferParam.UnsetUAV(RHICmdList, Shader->GetComputeShader());
+#endif
+		//RHICmdList.TransitionResource(EResourceTransitionAccess::EReadable, EResourceTransitionPipeline::EGfxToGfx, CurrDataRender().GetGPUBufferFloat()->UAV);
+	}
+
+	if (Shader->IntOutputBufferParam.IsUAVBound())
+	{
+#if !PLATFORM_PS4
+		Shader->IntOutputBufferParam.UnsetUAV(RHICmdList, Shader->GetComputeShader());
+#endif
+		//RHICmdList.TransitionResource(EResourceTransitionAccess::EReadable, EResourceTransitionPipeline::EGfxToGfx, CurrDataRender().GetGPUBufferInt()->UAV);
+	}
+}
+
+FScopedNiagaraDataSetGPUReadback::~FScopedNiagaraDataSetGPUReadback()
+{
+	if (DataBuffer != nullptr)
+	{
+		DataBuffer->FloatData.Empty();
+		DataBuffer->Int32Data.Empty();
+	}
+}
+
+void FScopedNiagaraDataSetGPUReadback::ReadbackData(FNiagaraDataSet* InDataSet)
+{
+	check(DataSet == nullptr);
+	check(InDataSet != nullptr);
+
+	DataSet = InDataSet;
+	DataBuffer = DataSet->GetCurrentData();
+
+	// These should be zero if we are GPU and aren't inside a readback scope already
+	check((DataBuffer->FloatData.Num() == 0) && (DataBuffer->Int32Data.Num() == 0));
+
+	// Readback data
+	TArray<uint32> DrawIndirectData;
+	ENQUEUE_RENDER_COMMAND(ReadbackGPUBuffers)
+	(
+		[&](FRHICommandListImmediate& RHICmdList)
+		{
+			// Read DrawIndirect Params
+			{
+				const FRWBuffer& DataSetIndices = DataBuffer->GetGPUIndices();
+				const int32 BytesToCopy = DataSetIndices.NumBytes / sizeof(uint32);
+				DrawIndirectData.AddUninitialized(BytesToCopy);
+
+				void* Data = RHICmdList.LockVertexBuffer(DataSetIndices.Buffer, 0, BytesToCopy, RLM_ReadOnly);
+				FMemory::Memcpy(DrawIndirectData.GetData(), Data, BytesToCopy);
+				RHICmdList.UnlockVertexBuffer(DataSetIndices.Buffer);
+			}
+
+			// Read float data
+			const FRWBuffer& GPUFloatBuffer = DataBuffer->GetGPUBufferFloat();
+			if (GPUFloatBuffer.Buffer.IsValid())
+			{
+				DataBuffer->FloatData.AddUninitialized(GPUFloatBuffer.NumBytes);
+
+				void* CPUFloatBuffer = RHICmdList.LockVertexBuffer(GPUFloatBuffer.Buffer, 0, GPUFloatBuffer.NumBytes, RLM_ReadOnly);
+				FMemory::Memcpy(DataBuffer->FloatData.GetData(), CPUFloatBuffer, GPUFloatBuffer.NumBytes);
+				RHICmdList.UnlockVertexBuffer(GPUFloatBuffer.Buffer);
+			}
+
+			// Read int data
+			const FRWBuffer& GPUIntBuffer = DataBuffer->GetGPUBufferInt();
+			if (GPUIntBuffer.Buffer.IsValid())
+			{
+				DataBuffer->Int32Data.AddUninitialized(GPUIntBuffer.NumBytes);
+
+				void* CPUIntBuffer = RHICmdList.LockVertexBuffer(GPUIntBuffer.Buffer, 0, GPUIntBuffer.NumBytes, RLM_ReadOnly);
+				FMemory::Memcpy(DataBuffer->Int32Data.GetData(), CPUIntBuffer, GPUIntBuffer.NumBytes);
+				RHICmdList.UnlockVertexBuffer(GPUIntBuffer.Buffer);
+			}
+		}
+	);
+	FlushRenderingCommands();
+
+	NumInstances = DrawIndirectData[1];
+}
diff --git a/Engine/Plugins/FX/Niagara/Source/Niagara/Private/NiagaraEmitter.cpp b/Engine/Plugins/FX/Niagara/Source/Niagara/Private/NiagaraEmitter.cpp
--- a/Engine/Plugins/FX/Niagara/Source/Niagara/Private/NiagaraEmitter.cpp
+++ b/Engine/Plugins/FX/Niagara/Source/Niagara/Private/NiagaraEmitter.cpp
@@ -10,6 +10,7 @@
 #include "UObject/Linker.h"
 #include "NiagaraModule.h"
 #include "NiagaraSystem.h"
+#include "NiagaraStats.h"
 
 #if WITH_EDITOR
 const FName UNiagaraEmitter::PrivateMemberNames::EventHandlerScriptProps = GET_MEMBER_NAME_CHECKED(UNiagaraEmitter, EventHandlerScriptProps);
@@ -116,6 +117,8 @@ void UNiagaraEmitter::PostInitProperties()
 
 	}
 	UniqueEmitterName = TEXT("Emitter");
+
+	GenerateStatID();
 }
 
 #if WITH_EDITORONLY_DATA
@@ -192,7 +195,7 @@ void UNiagaraEmitter::PostLoad()
 #endif
 
 	TArray<UNiagaraScript*> AllScripts;
-	GetScripts(AllScripts, true);
+	GetScripts(AllScripts, false);
 
 	// Post load scripts for use below.
 	for (UNiagaraScript* Script : AllScripts)
@@ -211,7 +214,10 @@ void UNiagaraEmitter::PostLoad()
 		{
 			GraphSource->InvalidateCachedCompileIds();
 			bGenerateNewChangeId = true;
-			UE_LOG(LogNiagara, Log, TEXT("InvalidateCachedCompileIds for %s because GbForceNiagaraCompileOnLoad = %d"), *GetPathName(), GbForceNiagaraCompileOnLoad);
+			if (GEnableVerboseNiagaraChangeIdLogging)
+			{
+				UE_LOG(LogNiagara, Log, TEXT("InvalidateCachedCompileIds for %s because GbForceNiagaraCompileOnLoad = %d"), *GetPathName(), GbForceNiagaraCompileOnLoad);
+			}
 		}
 	}
 	
@@ -219,7 +225,10 @@ void UNiagaraEmitter::PostLoad()
 	{
 		// If the change id is already invalid we need to generate a new one, and can skip checking the owned scripts.
 		bGenerateNewChangeId = true;
-		UE_LOG(LogNiagara, Log, TEXT("Change ID updated for emitter %s because the ID was invalid."), *GetPathName());
+		if (GEnableVerboseNiagaraChangeIdLogging)
+		{
+			UE_LOG(LogNiagara, Log, TEXT("Change ID updated for emitter %s because the ID was invalid."), *GetPathName());
+		}
 	}
 	else
 	{
@@ -228,7 +237,10 @@ void UNiagaraEmitter::PostLoad()
 			if (Script->AreScriptAndSourceSynchronized() == false)
 			{
 				bGenerateNewChangeId = true;
-				//UE_LOG(LogNiagara, Log, TEXT("Change ID updated for emitter %s because of a change to its script %s"), *GetPathName(), *Script->GetPathName());
+				if (GEnableVerboseNiagaraChangeIdLogging)
+				{
+					UE_LOG(LogNiagara, Log, TEXT("Change ID updated for emitter %s because of a change to its script %s"), *GetPathName(), *Script->GetPathName());
+				}
 			}
 		}
 	}
@@ -344,7 +356,7 @@ void UNiagaraEmitter::PostEditChangeProperty(struct FPropertyChangedEvent& Prope
 	}
 
 	ThumbnailImageOutOfDate = true;
-	ChangeId = FGuid::NewGuid();
+	UpdateChangeId();
 	OnPropertiesChangedDelegate.Broadcast();
 }
 
@@ -829,4 +841,42 @@ void UNiagaraEmitter::GraphSourceChanged()
 {
 	UpdateChangeId();
 }
-#endif
\ No newline at end of file
+#endif
+
+TStatId UNiagaraEmitter::GetStatID(bool bGameThread, bool bConcurrent)const
+{
+#if STATS
+	if (bGameThread)
+	{
+		if (bConcurrent)
+		{
+			return StatID_GT;
+		}
+		else
+		{
+			return StatID_GT_CNC;
+		}
+	}
+	else
+	{
+		if (bConcurrent)
+		{
+			return StatID_RT;
+		}
+		else
+		{
+			return StatID_RT_CNC;
+		}
+	}
+#endif
+	return TStatId();
+}
+void UNiagaraEmitter::GenerateStatID()
+{
+#if STATS
+	StatID_GT = FDynamicStats::CreateStatId<FStatGroup_STATGROUP_NiagaraEmitters>(GetName() + TEXT("[GT]"));
+	StatID_GT_CNC = FDynamicStats::CreateStatId<FStatGroup_STATGROUP_NiagaraEmitters>(GetName() + TEXT("[GT_CNC]"));
+	StatID_RT = FDynamicStats::CreateStatId<FStatGroup_STATGROUP_NiagaraEmitters>(GetName() + TEXT("[RT]"));
+	StatID_RT_CNC = FDynamicStats::CreateStatId<FStatGroup_STATGROUP_NiagaraEmitters>(GetName() + TEXT("[RT_CNC]"));
+#endif
+}
