Copying //UE4/Dev-Niagara to Dev-Main (//UE4/Dev-Main)

diff --git a/Engine/Plugins/FX/Niagara/Source/Niagara/Private/NiagaraEmitterInstance.cpp b/Engine/Plugins/FX/Niagara/Source/Niagara/Private/NiagaraEmitterInstance.cpp
--- a/Engine/Plugins/FX/Niagara/Source/Niagara/Private/NiagaraEmitterInstance.cpp
+++ b/Engine/Plugins/FX/Niagara/Source/Niagara/Private/NiagaraEmitterInstance.cpp
@@ -44,10 +44,6 @@ static FAutoConsoleVariableRef CVarMaxNiagaraCPUParticlesPerEmitter(
 );
 //////////////////////////////////////////////////////////////////////////
 
-const FName FNiagaraEmitterInstance::PositionName(TEXT("Position"));
-const FName FNiagaraEmitterInstance::SizeName(TEXT("SpriteSize"));
-const FName FNiagaraEmitterInstance::MeshScaleName(TEXT("Scale"));
-
 FNiagaraEmitterInstance::FNiagaraEmitterInstance(FNiagaraSystemInstance* InParentSystemInstance)
 : CPUTimeMS(0.0f)
 , ExecutionState(ENiagaraExecutionState::Inactive)
@@ -55,67 +51,74 @@ FNiagaraEmitterInstance::FNiagaraEmitterInstance(FNiagaraSystemInstance* InParen
 , GPUExecContext(nullptr)
 , ParentSystemInstance(InParentSystemInstance)
 , CachedEmitter(nullptr)
+, CachedSystemFixedBounds()
 #if !UE_BUILD_SHIPPING
 , bEncounteredNaNs(false)
 #endif
+, EventSpawnTotal(0)
 {
 	bDumpAfterEvent = false;
 	ParticleDataSet = new FNiagaraDataSet();
 
 	Batcher = ParentSystemInstance ? ParentSystemInstance->GetBatcher() : nullptr;
+	check(Batcher != nullptr);
 }
 
 FNiagaraEmitterInstance::~FNiagaraEmitterInstance()
 {
+	// Clear the cached emitter as it is not safe to access the CacheEmitter due to deferred deleted which can happen after the CachedEmitter has been GCed
+	CachedEmitter = nullptr;
+
 	//UE_LOG(LogNiagara, Warning, TEXT("~Simulator %p"), this);
-	ClearRenderer();
 	CachedBounds.Init();
 	UnbindParameters();
 
-	if (CachedEmitter != nullptr && CachedEmitter->SimTarget == ENiagaraSimTarget::GPUComputeSim)
+	if (GPUExecContext != nullptr)
 	{
 		/** We defer the deletion of the particle dataset and the compute context to the RT to be sure all in-flight RT commands have finished using it.*/
-		if (GPUExecContext != nullptr)
-		{
-			if (Batcher)
+		NiagaraEmitterInstanceBatcher* B = Batcher && !Batcher->IsPendingKill() ? Batcher : nullptr;
+		FNiagaraComputeExecutionContext* Context = GPUExecContext;
+		FNiagaraDataSet* DataSet = ParticleDataSet;
+		ENQUEUE_RENDER_COMMAND(FDeleteContextCommand)(
+			[B, Context, DataSet](FRHICommandListImmediate& RHICmdList)
 			{
-				Batcher->Remove(GPUExecContext);
-			}
-
-			FNiagaraComputeExecutionContext* ExecContext = GPUExecContext;
-			ENQUEUE_RENDER_COMMAND(FDeleteComputeContextCommand)(
-				[ExecContext](FRHICommandListImmediate& RHICmdList)
+				if (Context)
 				{
-					delete ExecContext;
-				});
-			
-			GPUExecContext = nullptr;
-		}
+					if (B)
+					{
+						B->GiveEmitterContextToDestroy_RenderThread(Context);
+					}
+					else
+					{
+						delete Context;
+					}
+				}
 
-		if (ParticleDataSet != nullptr)
-		{
-			FNiagaraDataSet* DataSet = ParticleDataSet;
-			ENQUEUE_RENDER_COMMAND(FDeleteParticleDataSetCommand)(
-				[DataSet](FRHICommandListImmediate& RHICmdList)
+				//TODO: deleting these on the RT shouldn't be needed any more.
+				if (DataSet)
 				{
-					delete DataSet;
-				});
+					if (B)
+					{
+						B->GiveDataSetToDestroy_RenderThread(DataSet);
+					}
+					else
+					{
+						delete DataSet;
+					}
 
-			ParticleDataSet = nullptr;
-		}
+				}
+			}
+		);
+			
+		GPUExecContext = nullptr;
+		ParticleDataSet = nullptr;
 	}
-}
-
-void FNiagaraEmitterInstance::ClearRenderer()
-{
-	for (int32 i = 0; i < EmitterRenderer.Num(); i++)
+	else
 	{
-		if (EmitterRenderer[i])
+		if ( ParticleDataSet != nullptr )
 		{
-			//UE_LOG(LogNiagara, Warning, TEXT("ClearRenderer %p"), EmitterRenderer);
-			// This queues up the renderer for deletion on the render thread..
-			EmitterRenderer[i]->Release();
-			EmitterRenderer[i] = nullptr;
+			delete ParticleDataSet;
+			ParticleDataSet = nullptr;
 		}
 	}
 }
@@ -147,9 +150,7 @@ void FNiagaraEmitterInstance::Dump()const
 		UE_LOG(LogNiagara, Log, TEXT("................. %s Combined Parameters ................."), TEXT("GPU Script"));
 		GPUExecContext->CombinedParamStore.DumpParameters();
 	}
-	UE_LOG(LogNiagara, Log, TEXT("................. Particles ................."));
-	ParticleDataSet->Dump(false);
-	ParticleDataSet->Dump(true);
+	ParticleDataSet->Dump(0, INDEX_NONE, TEXT("Particle Data"));
 }
 
 void FNiagaraEmitterInstance::Init(int32 InEmitterIdx, FName InSystemInstanceName)
@@ -177,7 +178,7 @@ void FNiagaraEmitterInstance::Init(int32 InEmitterIdx, FName InSystemInstanceNam
 	bEncounteredNaNs = false;
 #endif
 
-	Data.Init(FNiagaraDataSetID(CachedIDName, ENiagaraDataSetType::ParticleData), CachedEmitter->SimTarget);
+	Data.Init(FNiagaraDataSetID(CachedIDName, ENiagaraDataSetType::ParticleData), CachedEmitter->SimTarget, ParentSystemInstance->GetSystem()->GetName() + TEXT("/") + CachedEmitter->GetName());
 
 	//Init the spawn infos to the correct number for this system.
 	const TArray<FNiagaraEmitterSpawnAttributes>& EmitterSpawnInfoAttrs = ParentSystemInstance->GetSystem()->GetEmitterSpawnAttributes();
@@ -241,7 +242,7 @@ void FNiagaraEmitterInstance::Init(int32 InEmitterIdx, FName InSystemInstanceNam
 	for (const FNiagaraEventGeneratorProperties &GeneratorProps : CachedEmitter->UpdateScriptProps.EventGenerators)
 	{
 		FNiagaraDataSet *Set = FNiagaraEventDataSetMgr::CreateEventDataSet(ParentSystemInstance->GetIDName(), EmitterHandle.GetIdName(), GeneratorProps.SetProps.ID.Name);
-		Set->Init(FNiagaraDataSetID(), ENiagaraSimTarget::CPUSim);
+		Set->Init(FNiagaraDataSetID(), ENiagaraSimTarget::CPUSim, CachedEmitter->GetFullName() + TEXT("/") + GeneratorProps.SetProps.ID.Name.ToString());
 		Set->AddVariables(GeneratorProps.SetProps.Variables);
 		Set->Finalize();
 		UpdateScriptEventDataSets.Add(Set);
@@ -256,7 +257,7 @@ void FNiagaraEmitterInstance::Init(int32 InEmitterIdx, FName InSystemInstanceNam
 	for (const FNiagaraEventGeneratorProperties &GeneratorProps : CachedEmitter->SpawnScriptProps.EventGenerators)
 	{
 		FNiagaraDataSet *Set = FNiagaraEventDataSetMgr::CreateEventDataSet(ParentSystemInstance->GetIDName(), EmitterHandle.GetIdName(), GeneratorProps.SetProps.ID.Name);
-		Set->Init(FNiagaraDataSetID(), ENiagaraSimTarget::CPUSim);
+		Set->Init(FNiagaraDataSetID(), ENiagaraSimTarget::CPUSim, CachedEmitter->GetFullName() + TEXT("/") + GeneratorProps.SetProps.ID.Name.ToString());
 		Set->AddVariables(GeneratorProps.SetProps.Variables);
 		Set->Finalize();
 		SpawnScriptEventDataSets.Add(Set);
@@ -271,11 +272,9 @@ void FNiagaraEmitterInstance::Init(int32 InEmitterIdx, FName InSystemInstanceNam
 	if (CachedEmitter->SimTarget == ENiagaraSimTarget::GPUComputeSim)
 	{
 		GPUExecContext = new FNiagaraComputeExecutionContext();
-		GPUExecContext->InitParams(CachedEmitter->GetGPUComputeScript(), CachedEmitter->SpawnScriptProps.Script, CachedEmitter->UpdateScriptProps.Script, CachedEmitter->SimTarget, CachedEmitter->GetUniqueEmitterName());
+		GPUExecContext->InitParams(CachedEmitter->GetGPUComputeScript(), CachedEmitter->SimTarget, CachedEmitter->GetUniqueEmitterName());
 		GPUExecContext->MainDataSet = &Data;
-		GPUExecContext->RTGPUScript = CachedEmitter->GetGPUComputeScript()->GetRenderThreadScript();
-		GPUExecContext->RTSpawnScript = CachedEmitter->SpawnScriptProps.Script->GetRenderThreadScript();
-		GPUExecContext->RTUpdateScript = CachedEmitter->UpdateScriptProps.Script->GetRenderThreadScript();
+		GPUExecContext->GPUScript_RT = CachedEmitter->GetGPUComputeScript()->GetRenderThreadScript();
 
 		SpawnExecContext.Parameters.Bind(&GPUExecContext->CombinedParamStore);
 		UpdateExecContext.Parameters.Bind(&GPUExecContext->CombinedParamStore);
@@ -337,17 +336,6 @@ void FNiagaraEmitterInstance::Init(int32 InEmitterIdx, FName InSystemInstanceNam
 		EventExecCountBindings[i].Init(EventExecContexts[i].Parameters, SYS_PARAM_ENGINE_EXEC_COUNT);
 	}
 
-	if (CachedEmitter->SimTarget == ENiagaraSimTarget::GPUComputeSim && GPUExecContext != nullptr)
-	{
-	}
-	else
-	{
-		//Init accessors for PostTick
-		PositionAccessor = FNiagaraDataSetAccessor<FVector>(Data, FNiagaraVariable(FNiagaraTypeDefinition::GetVec3Def(), PositionName));
-		SizeAccessor = FNiagaraDataSetAccessor<FVector2D>(Data, FNiagaraVariable(FNiagaraTypeDefinition::GetVec2Def(), SizeName));
-		MeshScaleAccessor = FNiagaraDataSetAccessor<FVector>(Data, FNiagaraVariable(FNiagaraTypeDefinition::GetVec3Def(), MeshScaleName));
-	}
-
 	// Collect script defined data interface parameters.
 	TArray<UNiagaraScript*> Scripts;
 	Scripts.Add(CachedEmitter->SpawnScriptProps.Script);
@@ -357,6 +345,21 @@ void FNiagaraEmitterInstance::Init(int32 InEmitterIdx, FName InSystemInstanceNam
 		Scripts.Add(EventHandler.Script);
 	}
 	FNiagaraUtilities::CollectScriptDataInterfaceParameters(*CachedEmitter, Scripts, ScriptDefinedDataInterfaceParameters);
+
+	// Initialize bounds calculators
+	BoundsCalculators.Reserve(CachedEmitter->GetRenderers().Num());
+	for (UNiagaraRendererProperties* RendererProperties : CachedEmitter->GetRenderers())
+	{
+		if ((RendererProperties != nullptr) && RendererProperties->GetIsEnabled())
+		{
+			FNiagaraBoundsCalculator* BoundsCalculator = RendererProperties->CreateBoundsCalculator();
+			if (BoundsCalculator != nullptr)
+			{
+				BoundsCalculator->InitAccessors(*ParticleDataSet);
+				BoundsCalculators.Emplace(BoundsCalculator);
+			}
+		}
+	}
 }
 
 void FNiagaraEmitterInstance::ResetSimulation(bool bKillExisting)
@@ -579,6 +582,12 @@ void FNiagaraEmitterInstance::BindParameters()
 		CachedEmitter->GetEventHandlers()[i].Script->RapidIterationParameters.Bind(&EventExecContexts[i].Parameters);
 	}
 #endif
+
+	if (CachedEmitter->SimTarget == ENiagaraSimTarget::GPUComputeSim)
+	{
+		SpawnExecContext.Parameters.Bind(&GPUExecContext->CombinedParamStore);
+		UpdateExecContext.Parameters.Bind(&GPUExecContext->CombinedParamStore);
+	}
 }
 
 void FNiagaraEmitterInstance::PostInitSimulation()
@@ -629,14 +638,16 @@ const FNiagaraEmitterHandle& FNiagaraEmitterInstance::GetEmitterHandle() const
 float FNiagaraEmitterInstance::GetTotalCPUTime()
 {
 	float Total = CPUTimeMS;
-	for (int32 i = 0; i < EmitterRenderer.Num(); i++)
-	{
-		if (EmitterRenderer[i])
-		{
-			Total += EmitterRenderer[i]->GetCPUTimeMS();
 
-		}
-	}
+	//TODO: Find some way to include the RT cost here?
+	//Possibly have the proxy write back it's most recent frame time during EOF updates?
+// 	for (int32 i = 0; i < EmitterRenderer.Num(); i++)
+// 	{
+// 		if (EmitterRenderer[i])
+// 		{
+// 			Total += EmitterRenderer[i]->GetCPUTimeMS();// 
+// 		}
+// 	}
 
 	return Total;
 }
@@ -654,90 +665,70 @@ int FNiagaraEmitterInstance::GetTotalBytesUsed()
 	return BytesUsed;
 }
 
-TOptional<FBox> FNiagaraEmitterInstance::CalculateDynamicBounds()
+FBox FNiagaraEmitterInstance::CalculateDynamicBounds(const bool bReadGPUSimulation)
 {
-	checkSlow(ParticleDataSet);
-	FNiagaraDataSet& Data = *ParticleDataSet;
-	int32 NumInstances = Data.GetNumInstances();
-	FBox Ret;
-	Ret.Init();
+	if (IsComplete() || !BoundsCalculators.Num() || CachedEmitter == nullptr)
+		return FBox(ForceInit);
 
-	if (IsComplete() || NumInstances == 0 || CachedEmitter->SimTarget == ENiagaraSimTarget::GPUComputeSim)//TODO: Pull data back from gpu buffers to get bounds for GPU sims.
-	{
-		return TOptional<FBox>();
-	}
+	FScopedNiagaraDataSetGPUReadback ScopedGPUReadback;
 
-	PositionAccessor.InitForAccess(true);
+	int32 NumInstances = 0;
+	if (CachedEmitter->SimTarget == ENiagaraSimTarget::GPUComputeSim)
+	{
+		if (!bReadGPUSimulation || (GPUExecContext == nullptr))
+			return FBox(ForceInit);
 
-	if (PositionAccessor.IsValid() == false)
+		ScopedGPUReadback.ReadbackData(GPUExecContext->MainDataSet);
+		NumInstances = ScopedGPUReadback.GetNumInstances();
+	}
+	else
 	{
-		return TOptional<FBox>();
+		NumInstances = ParticleDataSet->GetCurrentDataChecked().GetNumInstances();
 	}
 
-	SizeAccessor.InitForAccess(true);
-	MeshScaleAccessor.InitForAccess(true);
+	if (NumInstances == 0)
+		return FBox(ForceInit);
 
-	FVector MaxSize(ForceInitToZero);
+	FBox Ret;
+	Ret.Init();
 
-	if (SizeAccessor.IsValid() == false && MeshScaleAccessor.IsValid() == false)
+	bool bContainsNaN = false;
+	for ( const TUniquePtr<FNiagaraBoundsCalculator>& BoundsCalculator : BoundsCalculators )
 	{
-		MaxSize = FVector(50.0f, 50.0f, 50.0f);
+		Ret += BoundsCalculator->CalculateBounds(NumInstances, bContainsNaN);
 	}
 
-	for (int32 InstIdx = 0; InstIdx < NumInstances && PositionAccessor.IsValid(); ++InstIdx)
+#if !UE_BUILD_SHIPPING
+	if (bContainsNaN && ParentSystemInstance != nullptr && CachedEmitter != nullptr && ParentSystemInstance->GetSystem() != nullptr)
 	{
-		FVector Position;
-		PositionAccessor.Get(InstIdx, Position);
+		UE_LOG(LogNiagara, Warning, TEXT("Particle position data contains NaNs. Likely a divide by zero somewhere in your modules. Emitter \"%s\" in System \"%s\""), *CachedEmitter->GetName(), *ParentSystemInstance->GetSystem()->GetName());
+		ParentSystemInstance->Dump();
+	}
+#endif
 
-		// Some graphs have a tendency to divide by zero. This ContainsNaN has been added prophylactically
-		// to keep us safe during GDC. It should be removed as soon as we feel safe that scripts are appropriately warned.
-		if (!Position.ContainsNaN())
-		{
-			Ret += Position;
+	return Ret;
+}
 
-			// We advance the scale or size depending of if we use either.
-			if (MeshScaleAccessor.IsValid())
-			{
-				MaxSize = MaxSize.ComponentMax(MeshScaleAccessor.Get(InstIdx));
-			}
-			else if (SizeAccessor.IsValid())
-			{
-				MaxSize = MaxSize.ComponentMax(FVector(SizeAccessor.Get(InstIdx).GetMax()));
-			}
-		}
-		else
-		{
-#if !UE_BUILD_SHIPPING
-			if (bEncounteredNaNs == false && ParentSystemInstance != nullptr && CachedEmitter != nullptr && ParentSystemInstance->GetSystem() != nullptr)
-			{
-				UE_LOG(LogNiagara, Warning, TEXT("Particle position data contains NaNs. Likely a divide by zero somewhere in your modules. Emitter \"%s\" in System \"%s\""),
-					*CachedEmitter->GetName(), *ParentSystemInstance->GetSystem()->GetName());
-				bEncounteredNaNs = true;
-				ParentSystemInstance->Dump();
-			}
-#endif
-		}
-	}
+void FNiagaraEmitterInstance::CalculateFixedBounds(const FTransform& ToWorldSpace)
+{
+	check(CachedEmitter);
+
+	FBox Bounds = CalculateDynamicBounds(true);
+	if (!Bounds.IsValid)
+		return;
 
-	float MaxBaseSize = 0.0001f;
-	if (MaxSize.IsNearlyZero())
+	CachedEmitter->Modify();
+	CachedEmitter->bFixedBounds = true;
+	if (CachedEmitter->bLocalSpace)
 	{
-		MaxSize = FVector(1.0f, 1.0f, 1.0f);
+		CachedEmitter->FixedBounds = Bounds;
 	}
-
-	for (int32 i = 0; i < EmitterRenderer.Num(); i++)
+	else
 	{
-		if (EmitterRenderer[i])
-		{
-			FVector BaseExtents = EmitterRenderer[i]->GetBaseExtents();
-			FVector ComponentMax;
-			MaxBaseSize = BaseExtents.ComponentMax(FVector(MaxBaseSize, MaxBaseSize, MaxBaseSize)).GetMax();
-		}
+		CachedEmitter->FixedBounds = Bounds.TransformBy(ToWorldSpace);
 	}
 
-	Ret = Ret.ExpandBy(MaxSize*MaxBaseSize);
-
-	return Ret;
+	CachedBounds = Bounds;
 }
 
 /** 
@@ -748,23 +739,30 @@ void FNiagaraEmitterInstance::PostTick()
 	SCOPE_CYCLE_COUNTER(STAT_NiagaraEmitterPostTick);
 
 	checkSlow(CachedEmitter);
+
+	EventHandlingInfo.Reset();
+
 	CachedBounds.Init();
-	if (CachedEmitter->bFixedBounds || CachedEmitter->SimTarget == ENiagaraSimTarget::GPUComputeSim)
+	if (CachedSystemFixedBounds.IsSet())
+	{
+		CachedBounds = CachedSystemFixedBounds.GetValue();
+	}
+	else if (CachedEmitter->bFixedBounds || CachedEmitter->SimTarget == ENiagaraSimTarget::GPUComputeSim)
 	{
 		CachedBounds = CachedEmitter->FixedBounds;
 	}
 	else
 	{
-		TOptional<FBox> DynamicBounds = CalculateDynamicBounds();
-		if (DynamicBounds.IsSet())
+		FBox DynamicBounds = CalculateDynamicBounds();
+		if (DynamicBounds.IsValid)
 		{
 			if (CachedEmitter->bLocalSpace)
 			{
-				CachedBounds = DynamicBounds.GetValue();
+				CachedBounds = DynamicBounds;
 			}
 			else
 			{
-				CachedBounds = DynamicBounds.GetValue().TransformBy(ParentSystemInstance->GetComponent()->GetComponentToWorld().Inverse());
+				CachedBounds = DynamicBounds.TransformBy(ParentSystemInstance->GetComponent()->GetComponentToWorld().Inverse());
 			}
 		}
 		else
@@ -783,11 +781,7 @@ bool FNiagaraEmitterInstance::HandleCompletion(bool bForce)
 
 	if (IsComplete())
 	{
-		//If we have any particles then clear out the buffers.
-		if (ParticleDataSet->GetNumInstances() > 0 || ParticleDataSet->GetPrevNumInstances() > 0)
-		{
-			ParticleDataSet->ResetBuffers();
-		}
+		ParticleDataSet->ResetBuffers();
 		return true;
 	}
 
@@ -810,6 +804,10 @@ void FNiagaraEmitterInstance::PreTick()
 		return;
 	}
 
+#if STATS
+	FScopeCycleCounter SystemStatCounter(CachedEmitter->GetStatID(true, true));
+#endif
+
 	checkSlow(ParticleDataSet);
 	FNiagaraDataSet& Data = *ParticleDataSet;
 
@@ -827,6 +825,8 @@ void FNiagaraEmitterInstance::PreTick()
 	bool bOk = true;
 	bOk &= SpawnExecContext.Tick(ParentSystemInstance);
 	bOk &= UpdateExecContext.Tick(ParentSystemInstance);
+
+	// @todo THREADSAFETY We should not tick GPU contexts on the game thread!
 	if (CachedEmitter->SimTarget == ENiagaraSimTarget::GPUComputeSim && GPUExecContext != nullptr)
 	{
 		bOk &= GPUExecContext->Tick(ParentSystemInstance);
@@ -848,6 +848,11 @@ void FNiagaraEmitterInstance::PreTick()
 		//On our very first frame we prime any previous params (for interpolation).
 		SpawnExecContext.PostTick();
 		UpdateExecContext.PostTick();
+		if (CachedEmitter->SimTarget == ENiagaraSimTarget::GPUComputeSim && GPUExecContext != nullptr)
+		{
+			//We PostTick the GPUExecContext here to prime crucial PREV parameters (such as PREV_Engine.Owner.Position). This PostTick call is necessary as the GPUExecContext has not been sent to the batcher yet.
+			GPUExecContext->PostTick();
+		}
 		for (FNiagaraScriptExecutionContext& EventContext : EventExecContexts)
 		{
 			EventContext.PostTick();
@@ -860,34 +865,38 @@ void FNiagaraEmitterInstance::PreTick()
 
 	if (bResetPending)
 	{
-		Data.ResetCurrentBuffers();
+		Data.ResetBuffers();
 		for (FNiagaraDataSet* SpawnScriptEventDataSet : SpawnScriptEventDataSets)
 		{
-			SpawnScriptEventDataSet->ResetCurrentBuffers();
+			SpawnScriptEventDataSet->ResetBuffers();
 		}
 		for (FNiagaraDataSet* UpdateScriptEventDataSet : UpdateScriptEventDataSets)
 		{
-			UpdateScriptEventDataSet->ResetCurrentBuffers();
+			UpdateScriptEventDataSet->ResetBuffers();
 		}
 		bResetPending = false;
 	}
 
-	//Swap all data set buffers before doing the main tick on any simulation.
-	if (CachedEmitter->SimTarget == ENiagaraSimTarget::CPUSim)
+	//Gather events we're going to be reading from / handling this frame.
+	//We must do this in pre-tick so we can gather (and mark in use) all sets from other emitters.
+	EventHandlingInfo.Reset();
+	EventHandlingInfo.SetNum(CachedEmitter->GetEventHandlers().Num());
+	EventSpawnTotal = 0;
+	for (int32 i = 0; i < CachedEmitter->GetEventHandlers().Num(); i++)
 	{
-		for (TPair<FNiagaraDataSetID, FNiagaraDataSet*> SetPair : DataSetMap)
-		{
-			SetPair.Value->Tick();
-		}
-
-		for (FNiagaraDataSet* Set : UpdateScriptEventDataSets)
-		{
-			Set->Tick();
-		}
-
-		for (FNiagaraDataSet* Set : SpawnScriptEventDataSets)
+		const FNiagaraEventScriptProperties &EventHandlerProps = CachedEmitter->GetEventHandlers()[i];
+		FNiagaraEventHandlingInfo& Info = EventHandlingInfo[i];
+		Info.SourceEmitterGuid = EventHandlerProps.SourceEmitterID;
+		Info.SourceEmitterName = Info.SourceEmitterGuid.IsValid() ? *Info.SourceEmitterGuid.ToString() : CachedIDName;
+		Info.SpawnCounts.Reset();
+		Info.TotalSpawnCount = 0;
+		Info.EventData = nullptr;
+		if (FNiagaraDataSet* EventSet = FNiagaraEventDataSetMgr::GetEventDataSet(ParentSystemInstance->GetIDName(), Info.SourceEmitterName, EventHandlerProps.SourceEventName))
 		{
-			Set->Tick();
+			Info.SetEventData(&EventSet->GetCurrentDataChecked());
+			uint32 EventSpawnNum = CalculateEventSpawnCount(EventHandlerProps, Info.SpawnCounts, EventSet);
+			Info.TotalSpawnCount += EventSpawnNum;
+			EventSpawnTotal += EventSpawnNum;
 		}
 	}
 
@@ -910,11 +919,19 @@ bool FNiagaraEmitterInstance::WaitForDebugInfo()
 }
 
 
+void FNiagaraEmitterInstance::SetSystemFixedBoundsOverride(FBox SystemFixedBounds)
+{
+	CachedSystemFixedBounds = SystemFixedBounds;
+}
+
 void FNiagaraEmitterInstance::Tick(float DeltaSeconds)
 {
 	SCOPE_CYCLE_COUNTER(STAT_NiagaraTick);
 	SimpleTimer TickTime;
 
+#if STATS
+	FScopeCycleCounter SystemStatCounter(CachedEmitter->GetStatID(true, true));
+#endif
 
 	if (HandleCompletion())
 	{
@@ -926,7 +943,6 @@ void FNiagaraEmitterInstance::Tick(float DeltaSeconds)
 	FNiagaraDataSet& Data = *ParticleDataSet;
 	Age += DeltaSeconds;
 
-
 	//UE_LOG(LogNiagara, Warning, TEXT("Emitter Tick %f"), Age);
 
 	if (ExecutionState == ENiagaraExecutionState::InactiveClear)
@@ -937,20 +953,22 @@ void FNiagaraEmitterInstance::Tick(float DeltaSeconds)
 		return;
 	}
 
-	int32 OrigNumParticles = Data.GetPrevNumInstances();
-	if (OrigNumParticles == 0 && ExecutionState != ENiagaraExecutionState::Active)
+	if (CachedEmitter->SimTarget == ENiagaraSimTarget::CPUSim && Data.GetCurrentDataChecked().GetNumInstances() == 0 && ExecutionState != ENiagaraExecutionState::Active)
 	{
-		//Clear out curr buffer in case it had some data in previously.
-		if (CachedEmitter->SimTarget == ENiagaraSimTarget::CPUSim)
-		{
-			Data.Allocate(0);
-		}
+		Data.ResetBuffers();
 		CPUTimeMS = TickTime.GetElapsedMilliseconds();
 		return;
 	}
 
 	UNiagaraSystem* System = ParentSystemInstance->GetSystem();
 
+	if (GbDumpParticleData || System->bDumpDebugEmitterInfo)
+	{
+		UE_LOG(LogNiagara, Log, TEXT("|||||||||||||||||||||||||||||||||||||||||||||||||||||||||||||||||||||||||||||||||||||||||||||||||||||||||||||"), *CachedEmitter->GetPathName());
+		UE_LOG(LogNiagara, Log, TEXT("|=== FNiagaraEmitterInstance::Tick [ %s ] ===============|"), *CachedEmitter->GetPathName());
+	}
+
+
 	check(Data.GetNumVariables() > 0);
 	check(CachedEmitter->SpawnScriptProps.Script);
 	check(CachedEmitter->UpdateScriptProps.Script);
@@ -993,53 +1011,12 @@ void FNiagaraEmitterInstance::Tick(float DeltaSeconds)
 		}
 	}
 
-	// Calculate number of new particles from all event related spawns
-	TArray<TArray<int32, TInlineAllocator<16>>, TInlineAllocator<16>> EventSpawnCounts;
-	EventSpawnCounts.AddDefaulted(CachedEmitter->GetEventHandlers().Num());
-	TArray<int32, TInlineAllocator<16>> EventHandlerSpawnCounts;
-	EventHandlerSpawnCounts.AddDefaulted(CachedEmitter->GetEventHandlers().Num());
-	uint32 EventSpawnTotal = 0;
-	TArray<FNiagaraDataSet*, TInlineAllocator<16>> EventSet;
-	EventSet.AddZeroed(CachedEmitter->GetEventHandlers().Num());
-	TArray<FGuid, TInlineAllocator<16>> SourceEmitterGuid;
-	SourceEmitterGuid.AddDefaulted(CachedEmitter->GetEventHandlers().Num());
-	TArray<FName, TInlineAllocator<16>> SourceEmitterName;
-	SourceEmitterName.AddDefaulted(CachedEmitter->GetEventHandlers().Num());
-	TArray<bool, TInlineAllocator<16>> bPerformEventSpawning;
-	bPerformEventSpawning.AddDefaulted(CachedEmitter->GetEventHandlers().Num());
-
-	for (int32 i = 0; i < CachedEmitter->GetEventHandlers().Num(); i++)
-	{
-		const FNiagaraEventScriptProperties &EventHandlerProps = CachedEmitter->GetEventHandlers()[i];
-		SourceEmitterGuid[i] = EventHandlerProps.SourceEmitterID;
-		SourceEmitterName[i] = SourceEmitterGuid[i].IsValid() ? *SourceEmitterGuid[i].ToString() : CachedIDName;
-		EventSet[i] = FNiagaraEventDataSetMgr::GetEventDataSet(ParentSystemInstance->GetIDName(), SourceEmitterName[i], EventHandlerProps.SourceEventName);
-		bPerformEventSpawning[i] = (ExecutionState == ENiagaraExecutionState::Active && EventHandlerProps.Script && EventHandlerProps.ExecutionMode == EScriptExecutionMode::SpawnedParticles);
-		if (bPerformEventSpawning[i])
-		{
-			uint32 EventSpawnNum = CalculateEventSpawnCount(EventHandlerProps, EventSpawnCounts[i], EventSet[i]);
-			EventSpawnTotal += EventSpawnNum;
-			EventHandlerSpawnCounts[i] = EventSpawnNum;
-			if (CachedEmitter->SimTarget == ENiagaraSimTarget::GPUComputeSim && GPUExecContext != nullptr) {
-				// NOTE(mv): Separate particle count path for GPU emitters, as they early out..
-				// not actually reached yet due to GPU events not working. CPU events are counted further down...
-				TotalSpawnedParticles += EventSpawnNum;
-			}
-		}
-	}
-
-
 	/* GPU simulation -  we just create an FNiagaraComputeExecutionContext, queue it, and let the batcher take care of the rest
 	 */
 	if (CachedEmitter->SimTarget == ENiagaraSimTarget::GPUComputeSim && GPUExecContext != nullptr)
 	{
-		//FNiagaraComputeExecutionContext *ComputeContext = new FNiagaraComputeExecutionContext();
-		GPUExecContext->MainDataSet = &Data;
-		GPUExecContext->RTGPUScript = CachedEmitter->GetGPUComputeScript()->GetRenderThreadScript();
-		GPUExecContext->RTSpawnScript = CachedEmitter->SpawnScriptProps.Script->GetRenderThreadScript();
-		GPUExecContext->RTUpdateScript = CachedEmitter->UpdateScriptProps.Script->GetRenderThreadScript();
-		GPUExecContext->SpawnRateInstances = SpawnTotal;
-		GPUExecContext->EventSpawnTotal = EventSpawnTotal;
+		check(GPUExecContext->GPUScript_RT == CachedEmitter->GetGPUComputeScript()->GetRenderThreadScript());
+		GPUExecContext->GPUScript_RT = CachedEmitter->GetGPUComputeScript()->GetRenderThreadScript();
 
 		if (CachedEmitter->GetRenderers().Num() > 0 && CachedEmitter->GetRenderers()[0] != nullptr)
 		{
@@ -1049,6 +1026,9 @@ void FNiagaraEmitterInstance::Tick(float DeltaSeconds)
 		{
 			GPUExecContext->NumIndicesPerInstance = 0;
 		}
+
+		GPUExecContext->EventSpawnTotal_GT = EventSpawnTotal;
+		GPUExecContext->SpawnRateInstances_GT = SpawnTotal;
 		
 #if WITH_EDITORONLY_DATA
 		if (ParentSystemInstance->ShouldCaptureThisFrame())
@@ -1059,17 +1039,26 @@ void FNiagaraEmitterInstance::Tick(float DeltaSeconds)
 				//Data.Dump(DebugInfo->Frame, true, 0, OrigNumParticles);
 				//DebugInfo->Frame.Dump(true, 0, OrigNumParticles);
 				DebugInfo->Parameters = GPUExecContext->CombinedParamStore;
+				
+				//TODO: This layout info can be pulled into the emitter/systems etc and all sets just refer to them. They are becoming an annoyance here.
+				DebugInfo->Frame.Init(FNiagaraDataSetID(TEXT("GPU Capture Data"), ENiagaraDataSetType::ParticleData), ENiagaraSimTarget::CPUSim, TEXT("GPU Capture Data"));
+				DebugInfo->Frame.AddVariables(Data.GetVariables());
+				DebugInfo->Frame.Finalize();
+
 				GPUExecContext->DebugInfo = DebugInfo;
 			}
 		}
 #endif
-		ParentSystemInstance->GetPerInstanceDataAndOffsets(GPUExecContext->PerInstanceData, GPUExecContext->PerInstanceDataSize, GPUExecContext->PerInstanceDataInterfaceOffsets);
+		// If this is not correct we will not propagate all data correctly.
+		// @todo-threadsafety we keep a counter of this, so 
+		//check(ParentSystemInstance->HasGPUEmitters());
 
 		bool bOnlySetOnce = false;
 		for (FNiagaraSpawnInfo& Info : SpawnInfos)
 		{
 			if (Info.Count > 0 && !bOnlySetOnce)
 			{
+				// @todo-threadsafety do these need to propagate to the RT?
 				SpawnIntervalBindingGPU.SetValue(Info.IntervalDt);
 				InterpSpawnStartBindingGPU.SetValue(Info.InterpStartDt);
 				SpawnGroupBindingGPU.SetValue(Info.SpawnGroup);
@@ -1089,13 +1078,6 @@ void FNiagaraEmitterInstance::Tick(float DeltaSeconds)
 
 		// copy over the constants for the render thread
 		//
-		//UE_LOG(LogNiagara, Log, TEXT(".................Spawn................."));
-		//SpawnExecContext.Parameters.DumpParameters();
-		//UE_LOG(LogNiagara, Log, TEXT(".................Update................."));
-		//UpdateExecContext.Parameters.DumpParameters();
-		//UE_LOG(LogNiagara, Log, TEXT(".................GPU................."));
-		//GPUExecContext.CombinedParamStore.DumpParameters();
-		//UE_LOG(LogNiagara, Log, TEXT(".................END.................")); 
 		if (GbDumpParticleData || System->bDumpDebugEmitterInfo)
 		{
 			UE_LOG(LogNiagara, Log, TEXT(".................Spawn................."));
@@ -1107,30 +1089,16 @@ void FNiagaraEmitterInstance::Tick(float DeltaSeconds)
 		}
 
 		int32 ParmSize = GPUExecContext->CombinedParamStore.GetPaddedParameterSizeInBytes();
-		
-		GPUExecContext->ParamData_RT.SetNumZeroed(ParmSize);
-		GPUExecContext->CombinedParamStore.CopyParameterDataToPaddedBuffer(GPUExecContext->ParamData_RT.GetData(), ParmSize);
 		// Because each context is only ran once each frame, the CBuffer layout stays constant for the lifetime duration of the CBuffer (one frame).
+
+		// @todo-threadsafety do this once during init. Should not change during runtime...
 		GPUExecContext->CBufferLayout.ConstantBufferSize = ParmSize;
 		GPUExecContext->CBufferLayout.ComputeHash();
 
-		// push event data sets to the context
-		for (FNiagaraDataSet *Set : UpdateScriptEventDataSets)
-		{
-			GPUExecContext->UpdateEventWriteDataSets.Add(Set);
-		}
-
-		GPUExecContext->EventHandlerScriptProps = CachedEmitter->GetEventHandlers();
-		GPUExecContext->EventSets = EventSet;
-		GPUExecContext->EventSpawnCounts = EventHandlerSpawnCounts;
-		if (Batcher)
-		{
-			Batcher->Queue(GPUExecContext);
-		}
-
 		// Need to call post-tick, which calls the copy to previous for interpolated spawning
 		SpawnExecContext.PostTick();
 		UpdateExecContext.PostTick();
+		// At this stage GPU execution is being handled by the batcher so we do not need to call PostTick() for it
 		for (FNiagaraScriptExecutionContext& EventContext : EventExecContexts)
 		{
 			EventContext.PostTick();
@@ -1148,8 +1116,9 @@ void FNiagaraEmitterInstance::Tick(float DeltaSeconds)
 		return;
 	}
 
-	int32 AllocationSize = OrigNumParticles + SpawnTotal + EventSpawnTotal;
+	int32 OrigNumParticles = Data.GetCurrentDataChecked().GetNumInstances();
 
+	int32 AllocationSize = OrigNumParticles + SpawnTotal + EventSpawnTotal;
 	//Ensure we don't blow our current hard limits on cpu particle count.
 	//TODO: These current limits can be improved relatively easily. Though perf in at these counts will obviously be an issue anyway.
 	if (CachedEmitter->SimTarget == ENiagaraSimTarget::CPUSim && AllocationSize > GMaxNiagaraCPUParticlesPerEmitter)
@@ -1161,7 +1130,7 @@ void FNiagaraEmitterInstance::Tick(float DeltaSeconds)
 		EventSpawnTotal = 0;
 	}
 
-	//Allocate space for prev frames particles and any new one's we're going to spawn.
+	Data.BeginSimulate();
 	Data.Allocate(AllocationSize);
 
 	int32 SpawnEventGeneratorIndex = 0;
@@ -1174,6 +1143,7 @@ void FNiagaraEmitterInstance::Tick(float DeltaSeconds)
 			// the same data set will be used in the update execution.
 			NumToAllocate += OrigNumParticles;
 		}
+		SpawnEventDataSet->BeginSimulate();
 		SpawnEventDataSet->Allocate(NumToAllocate);
 		SpawnEventGeneratorIndex++;
 	}
@@ -1185,67 +1155,51 @@ void FNiagaraEmitterInstance::Tick(float DeltaSeconds)
 		{
 			// We only allocate update event data sets if they're not shared, because shared event datasets will have already
 			// been allocated as part of the spawn event data set handling.
+			UpdateEventDataSet->BeginSimulate();
 			UpdateEventDataSet->Allocate(OrigNumParticles);
 		}
 		UpdateEventGeneratorIndex++;
 	}
 
-	TArray<FNiagaraDataSetExecutionInfo, TInlineAllocator<8>> DataSetExecInfos;
-	DataSetExecInfos.Emplace(&Data, 0, false, true);
-
 	// Simulate existing particles forward by DeltaSeconds.
 	if (OrigNumParticles > 0)
 	{
-		/*
-		if (bDumpAfterEvent)
-		{
-			Data.Dump(false);
-			bDumpAfterEvent = false;
-		}
-		*/
-
-		Data.SetNumInstances(OrigNumParticles);
+		Data.GetDestinationDataChecked().SetNumInstances(OrigNumParticles);
 		SCOPE_CYCLE_COUNTER(STAT_NiagaraSimulate);
 
 		UpdateExecCountBinding.SetValue(OrigNumParticles);
-		DataSetExecInfos.SetNum(1, false);
-		DataSetExecInfos[0].StartInstance = 0;
+		UpdateExecContext.BindData(0, Data, 0, true);
+		int32 EventDataSetIdx = 1;
 		for (FNiagaraDataSet* EventDataSet : UpdateScriptEventDataSets)
 		{
-			EventDataSet->SetNumInstances(OrigNumParticles);
-			DataSetExecInfos.Emplace(EventDataSet, 0, false, true);
+			check(EventDataSet);
+			EventDataSet->GetDestinationDataChecked().SetNumInstances(OrigNumParticles);
+			UpdateExecContext.BindData(EventDataSetIdx++, *EventDataSet, 0, true);
 		}
-		UpdateExecContext.Execute(OrigNumParticles, DataSetExecInfos);
-		int32 DeltaParticles = Data.GetNumInstances() - OrigNumParticles;
+		UpdateExecContext.Execute(OrigNumParticles);
+		int32 DeltaParticles = Data.GetDestinationDataChecked().GetNumInstances() - OrigNumParticles;
 
 		ensure(DeltaParticles <= 0); // We either lose particles or stay the same, we should never add particles in update!
 
 		if (GbDumpParticleData || System->bDumpDebugEmitterInfo)
 		{
-			UE_LOG(LogNiagara, Log, TEXT("=== Update Parameters ===") );
+			Data.GetDestinationDataChecked().Dump(0, OrigNumParticles, FString::Printf(TEXT("=== Updated %d Particles (%d Died) ==="), OrigNumParticles, -DeltaParticles));
+			for (int32 EventIdx = 0; EventIdx < UpdateScriptEventDataSets.Num(); ++EventIdx)
+			{
+				FNiagaraDataSet* EventDataSet = UpdateScriptEventDataSets[EventIdx];
+				if (EventDataSet && EventDataSet->GetDestinationDataChecked().GetNumInstances() > 0)
+				{
+					EventDataSet->GetDestinationDataChecked().Dump(0, INDEX_NONE, FString::Printf(TEXT("Update Script Event %d"), EventIdx));
+				}
+			}
+		//	UE_LOG(LogNiagara, Log, TEXT("=== Update Parameters ===") );
 			UpdateExecContext.Parameters.Dump();
-
-			UE_LOG(LogNiagara, Log, TEXT("=== Updated %d Particles (%d Died) ==="), OrigNumParticles, -DeltaParticles);
-			Data.Dump(true, 0, OrigNumParticles);
 		}
 	}
 	
-#if WITH_EDITORONLY_DATA
-	if (ParentSystemInstance->ShouldCaptureThisFrame())
-	{
-		TSharedPtr<struct FNiagaraScriptDebuggerInfo, ESPMode::ThreadSafe> DebugInfo = ParentSystemInstance->GetActiveCaptureWrite(CachedIDName, ENiagaraScriptUsage::ParticleUpdateScript, FGuid());
-		if (DebugInfo)
-		{
-			Data.Dump(DebugInfo->Frame, true, 0, OrigNumParticles);
-			//DebugInfo->Frame.Dump(true, 0, OrigNumParticles);
-			DebugInfo->Parameters = UpdateExecContext.Parameters;
-			DebugInfo->bWritten = true;
-		}
-	}
-#endif
-
-	uint32 EventSpawnStart = Data.GetNumInstances();
-	int32 NumBeforeSpawn = Data.GetNumInstances();
+	uint32 EventSpawnStart = Data.GetDestinationDataChecked().GetNumInstances();
+	int32 NumBeforeSpawn = Data.GetDestinationDataChecked().GetNumInstances();
+	uint32 TotalActualEventSpawns = 0;
 
 	//Init new particles with the spawn script.
 	if (SpawnTotal + EventSpawnTotal > 0)
@@ -1257,8 +1211,8 @@ void FNiagaraEmitterInstance::Tick(float DeltaSeconds)
 		{
 			if (Num > 0)
 			{
-				int32 OrigNum = Data.GetNumInstances();
-				Data.SetNumInstances(OrigNum + Num);
+				int32 OrigNum = Data.GetDestinationDataChecked().GetNumInstances();
+				Data.GetDestinationDataChecked().SetNumInstances(OrigNum + Num);
 
 				// NOTE(mv): Updates the count after setting the variable, such that the TotalSpawnedParticles value read 
 				//           in the script has the count at the start of the frame. 
@@ -1267,26 +1221,33 @@ void FNiagaraEmitterInstance::Tick(float DeltaSeconds)
 				TotalSpawnedParticles += Num;
 
 				SpawnExecCountBinding.SetValue(Num);
-				DataSetExecInfos.SetNum(1, false);
-				DataSetExecInfos[0].StartInstance = OrigNum;
+				SpawnExecContext.BindData(0, Data, OrigNum, true);
 
 				//UE_LOG(LogNiagara, Log, TEXT("SpawnScriptEventDataSets: %d"), SpawnScriptEventDataSets.Num());
+				int32 EventDataSetIdx = 1;
 				for (FNiagaraDataSet* EventDataSet : SpawnScriptEventDataSets)
 				{
 					//UE_LOG(LogNiagara, Log, TEXT("SpawnScriptEventDataSets.. %d"), EventDataSet->GetNumVariables());
-					int32 EventOrigNum = EventDataSet->GetNumInstances();
-					EventDataSet->SetNumInstances(EventOrigNum + Num);
-					DataSetExecInfos.Emplace(EventDataSet, EventOrigNum, false, true);
+					int32 EventOrigNum = EventDataSet->GetDestinationDataChecked().GetNumInstances();
+					EventDataSet->GetDestinationDataChecked().SetNumInstances(EventOrigNum + Num);
+					SpawnExecContext.BindData(EventDataSetIdx++, *EventDataSet, EventOrigNum, true);
 				}
-				
-				SpawnExecContext.Execute(Num, DataSetExecInfos);
+
+				SpawnExecContext.Execute(Num);
 
 				if (GbDumpParticleData || System->bDumpDebugEmitterInfo)
 				{
-					UE_LOG(LogNiagara, Log, TEXT("=== %s Spawn Parameters ==="), *DumpLabel);
+					Data.GetDestinationDataChecked().Dump(OrigNum, Num, FString::Printf(TEXT("===  %s Spawned %d Particles==="), *DumpLabel, Num));
+					for (int32 EventIdx = 0; EventIdx < SpawnScriptEventDataSets.Num(); ++EventIdx)
+					{
+						FNiagaraDataSet* EventDataSet = SpawnScriptEventDataSets[EventIdx];
+						if (EventDataSet && EventDataSet->GetDestinationDataChecked().GetNumInstances() > 0)
+						{
+							EventDataSet->GetDestinationDataChecked().Dump(0, INDEX_NONE, FString::Printf(TEXT("Spawn Script Event %d"), EventIdx));
+						}
+					}
+					//UE_LOG(LogNiagara, Log, TEXT("=== %s Spawn Parameters ==="), *DumpLabel);
 					SpawnExecContext.Parameters.Dump();
-					UE_LOG(LogNiagara, Log, TEXT("===  %s Spawned %d Particles==="), *DumpLabel, Num);
-					Data.Dump(true, OrigNum, Num);
 				}
 			}
 		};
@@ -1301,36 +1262,69 @@ void FNiagaraEmitterInstance::Tick(float DeltaSeconds)
 			SpawnParticles(Info.Count, TEXT("Regular Spawn"));
 		};
 
-		EventSpawnStart = Data.GetNumInstances();
+		EventSpawnStart = Data.GetDestinationDataChecked().GetNumInstances();
 
 		for (int32 EventScriptIdx = 0; EventScriptIdx < CachedEmitter->GetEventHandlers().Num(); EventScriptIdx++)
 		{
+			FNiagaraEventHandlingInfo& Info = EventHandlingInfo[EventScriptIdx];
 			//Spawn particles coming from events.
-			for (int32 i = 0; i < EventSpawnCounts[EventScriptIdx].Num(); i++)
+			for (int32 i = 0; i < Info.SpawnCounts.Num(); i++)
 			{
-				int32 EventNumToSpawn = EventSpawnCounts[EventScriptIdx][i];
+				int32 EventNumToSpawn = Info.SpawnCounts[i];
 
+				int32 CurrNumParticles = Data.GetDestinationDataChecked().GetNumInstances();
 				//Event spawns are instantaneous at the middle of the frame?
 				SpawnIntervalBinding.SetValue(0.0f);
 				InterpSpawnStartBinding.SetValue(DeltaSeconds * 0.5f);
 				SpawnGroupBinding.SetValue(0);
 
 				SpawnParticles(EventNumToSpawn, TEXT("Event Spawn"));
+
+				//Update EventSpawnCounts to the number actually spawned.
+				int32 NumActuallySpawned = Data.GetDestinationDataChecked().GetNumInstances() - CurrNumParticles;
+				TotalActualEventSpawns += NumActuallySpawned;
+				Info.SpawnCounts[i] = NumActuallySpawned;
 			}
 		}
 	}
 
+	//We're done with this simulation pass.
+	Data.EndSimulate();
+	for (FNiagaraDataSet* SpawnEventDataSet : SpawnScriptEventDataSets)
+	{
+		if (SpawnEventDataSet && SpawnEventDataSet->GetDestinationData())
+		{
+			SpawnEventDataSet->EndSimulate();
+		}
+	}
+
+	for (FNiagaraDataSet* UpdateEventDataSet : UpdateScriptEventDataSets)
+	{
+		if (UpdateEventDataSet && UpdateEventDataSet->GetDestinationData())
+		{
+			UpdateEventDataSet->EndSimulate();
+		}
+	}
 
+	//Now pull out any debug info we need.
 #if WITH_EDITORONLY_DATA
-	int32 NumAfterSpawn = Data.GetNumInstances();
+	int32 NumAfterSpawn = Data.GetCurrentDataChecked().GetNumInstances();
 	int32 TotalNumSpawned = NumAfterSpawn - NumBeforeSpawn;
 	if (ParentSystemInstance->ShouldCaptureThisFrame())
 	{
-		TSharedPtr<struct FNiagaraScriptDebuggerInfo, ESPMode::ThreadSafe> DebugInfo = ParentSystemInstance->GetActiveCaptureWrite(CachedIDName, ENiagaraScriptUsage::ParticleSpawnScript, FGuid());
+		//Pull out update data.
+		TSharedPtr<struct FNiagaraScriptDebuggerInfo, ESPMode::ThreadSafe> DebugInfo = ParentSystemInstance->GetActiveCaptureWrite(CachedIDName, ENiagaraScriptUsage::ParticleUpdateScript, FGuid());
 		if (DebugInfo)
 		{
-			Data.Dump(DebugInfo->Frame, true, NumBeforeSpawn, TotalNumSpawned);
-			//DebugInfo->Frame.Dump(true, 0, Num);
+			Data.CopyTo(DebugInfo->Frame, 0, OrigNumParticles);
+			DebugInfo->Parameters = UpdateExecContext.Parameters;
+			DebugInfo->bWritten = true;
+		}
+		//Pull out update data.
+		DebugInfo = ParentSystemInstance->GetActiveCaptureWrite(CachedIDName, ENiagaraScriptUsage::ParticleSpawnScript, FGuid());
+		if (DebugInfo)
+		{
+			Data.CopyTo(DebugInfo->Frame, NumBeforeSpawn, TotalNumSpawned);
 			DebugInfo->Parameters = SpawnExecContext.Parameters;
 			DebugInfo->bWritten = true;
 		}
@@ -1341,44 +1335,52 @@ void FNiagaraEmitterInstance::Tick(float DeltaSeconds)
 		UE_LOG(LogNiagara, Log, TEXT("Skipping spawning due to execution state! %d"), (uint32)ExecutionState)
 	}*/
 
-	// Events are all working from the same set of data generated during spawn they they only need 1 copy to have updated data.
-	if (CachedEmitter->GetEventHandlers().Num())
+	if (TotalActualEventSpawns > 0)
 	{
-		Data.CopyCurToPrev();
+		if (GbDumpParticleData || System->bDumpDebugEmitterInfo)
+		{
+			Data.Dump(0, INDEX_NONE, TEXT("Existing Data - Pre Event Alloc"));
+		}
+		//Allocate a new dest buffer to write spawn event handler results into.
+		//Can just do one allocate here for all spawn event handlers.
+		//Though this requires us to copy the contents of the instances we're not writing to in this pass over from the previous buffer.
+		FNiagaraDataBuffer& DestBuffer = Data.BeginSimulate();
+		Data.Allocate(Data.GetCurrentDataChecked().GetNumInstances(), true);
+		DestBuffer.SetNumInstances(Data.GetCurrentDataChecked().GetNumInstances());
+
+		//if (GbDumpParticleData || System->bDumpDebugEmitterInfo)
+		//{
+		//	DestBuffer.Dump(0, INDEX_NONE, TEXT("Existing Data - Post Event Alloc, Pre Events"));
+		//}
 	}
+
 	int32 SpawnEventScriptStartIndex = EventSpawnStart;
-	for (int32 EventScriptIdx = 0;  EventScriptIdx < CachedEmitter->GetEventHandlers().Num(); EventScriptIdx++)
+	for (int32 EventScriptIdx = 0; EventScriptIdx < CachedEmitter->GetEventHandlers().Num(); EventScriptIdx++)
 	{
 		const FNiagaraEventScriptProperties &EventHandlerProps = CachedEmitter->GetEventHandlers()[EventScriptIdx];
+		FNiagaraEventHandlingInfo& Info = EventHandlingInfo[EventScriptIdx];
 
-		if (bPerformEventSpawning[EventScriptIdx] && EventSet[EventScriptIdx] && EventSpawnCounts[EventScriptIdx].Num())
+		if (Info.EventData && Info.SpawnCounts.Num() > 0)
 		{
-			uint32 NumParticles = Data.GetNumInstances();
 			SCOPE_CYCLE_COUNTER(STAT_NiagaraEventHandle);
 
-			for (int32 i = 0; i < EventSpawnCounts[EventScriptIdx].Num(); i++)
+			for (int32 i = 0; i < Info.SpawnCounts.Num(); i++)
 			{
-				// When using interpolated spawn it's possible for the interpolated update script to kill a particle the same frame that it's spawned.
-				// In this case we have to decrease the number of instances to run the event script on.
-				int32 EventNumToSpawn = EventSpawnCounts[EventScriptIdx][i];
-				int32 ActualEventNumToSpawn = FMath::Min(EventNumToSpawn, (int32)Data.GetNumInstances() - (int32)EventSpawnStart);
-				
-				if (ActualEventNumToSpawn > 0)
+				int32 EventNumToSpawn = Info.SpawnCounts[i];
+				if (EventNumToSpawn > 0)
 				{
 					EventExecCountBindings[EventScriptIdx].SetValue(EventNumToSpawn);
 
-					DataSetExecInfos.SetNum(1, false);
-					DataSetExecInfos[0].StartInstance = EventSpawnStart;
-					DataSetExecInfos[0].bUpdateInstanceCount = false;
-					DataSetExecInfos.Emplace(EventSet[EventScriptIdx], i, false, false);
-					EventExecContexts[EventScriptIdx].Execute(ActualEventNumToSpawn, DataSetExecInfos);
+					EventExecContexts[EventScriptIdx].BindData(0, Data, EventSpawnStart, false);
+					EventExecContexts[EventScriptIdx].BindData(1, EventHandlingInfo[EventScriptIdx].EventData, nullptr, i, false);
+					EventExecContexts[EventScriptIdx].Execute(EventNumToSpawn);
 
-					if (GbDumpParticleData)
+					if (GbDumpParticleData || System->bDumpDebugEmitterInfo)
 					{
-						UE_LOG(LogNiagara, Log, TEXT("=== Event %d Parameters ==="), EventScriptIdx);
+						EventHandlingInfo[EventScriptIdx].EventData->Dump(i, 1, FString::Printf(TEXT("=== Event Data %d [%d] ==="), EventScriptIdx, i));
+						Data.GetDestinationDataChecked().Dump(EventSpawnStart, EventNumToSpawn, FString::Printf(TEXT("=== Event %d %d Particles ==="), EventScriptIdx, EventNumToSpawn));
+						//UE_LOG(LogNiagara, Log, TEXT("=== Event %d Parameters ==="), EventScriptIdx);
 						EventExecContexts[EventScriptIdx].Parameters.Dump();
-						UE_LOG(LogNiagara, Log, TEXT("=== Event %d %d Particles ==="), EventScriptIdx, ActualEventNumToSpawn);
-						Data.Dump(true, EventSpawnStart, ActualEventNumToSpawn);
 					}
 
 #if WITH_EDITORONLY_DATA
@@ -1388,63 +1390,59 @@ void FNiagaraEmitterInstance::Tick(float DeltaSeconds)
 						TSharedPtr<struct FNiagaraScriptDebuggerInfo, ESPMode::ThreadSafe> DebugInfo = ParentSystemInstance->GetActiveCaptureWrite(CachedIDName, ENiagaraScriptUsage::ParticleEventScript, EventGuid);
 						if (DebugInfo)
 						{
-							Data.Dump(DebugInfo->Frame, true, EventSpawnStart, ActualEventNumToSpawn);		
-							//DebugInfo->Frame.Dump(true, 0, ActualEventNumToSpawn);
+							Data.CopyTo(DebugInfo->Frame, EventSpawnStart, EventNumToSpawn);
 							DebugInfo->Parameters = EventExecContexts[EventScriptIdx].Parameters;
 							DebugInfo->bWritten = true;
 						}
 					}
 #endif
-
-					ensure(Data.GetNumInstances() == NumParticles);
-
-					EventSpawnStart += ActualEventNumToSpawn;
+					EventSpawnStart += EventNumToSpawn;
 				}
 			}
 		}
 	}
 
+	//If we processed any events we need to end simulate to update the current sim state.
+	if (Data.GetDestinationData())
+	{
+		Data.EndSimulate();
+	}
+
 	// Update events need a copy per event so that the previous event's data can be used.
 	for (int32 EventScriptIdx = 0; EventScriptIdx < CachedEmitter->GetEventHandlers().Num(); EventScriptIdx++)
 	{
 		const FNiagaraEventScriptProperties &EventHandlerProps = CachedEmitter->GetEventHandlers()[EventScriptIdx];
-
+		FNiagaraDataBuffer* EventData = EventHandlingInfo[EventScriptIdx].EventData;
 		// handle all-particle events
-		if (EventHandlerProps.Script && EventHandlerProps.ExecutionMode == EScriptExecutionMode::EveryParticle && EventSet[EventScriptIdx])
+		if (EventHandlerProps.Script && EventHandlerProps.ExecutionMode == EScriptExecutionMode::EveryParticle && EventData)
 		{
-			uint32 NumParticles = Data.GetNumInstances();
-			if (EventSet[EventScriptIdx]->GetPrevNumInstances())
+			uint32 NumParticles = Data.GetCurrentDataChecked().GetNumInstances();
+
+			if (EventData->GetNumInstances())
 			{
 				SCOPE_CYCLE_COUNTER(STAT_NiagaraEventHandle);
 				
-				for (uint32 i = 0; i < EventSet[EventScriptIdx]->GetPrevNumInstances(); i++)
+				for (uint32 i = 0; i < EventData->GetNumInstances(); i++)
 				{
-					// Copy the current to previous so that the event script has access to the new values from the update
-					// script and any values updated in previous events.
-					Data.CopyCurToPrev();
+					Data.BeginSimulate();
+					Data.Allocate(NumParticles);
 
-					uint32 NumInstancesPrev = Data.GetPrevNumInstances();
+					uint32 NumInstancesPrev = Data.GetCurrentDataChecked().GetNumInstances();
 					EventExecCountBindings[EventScriptIdx].SetValue(NumInstancesPrev);
-					DataSetExecInfos.SetNum(1, false);
-					DataSetExecInfos[0].StartInstance = 0;
-					DataSetExecInfos.Emplace(EventSet[EventScriptIdx], i, false, false);
-					/*if (GbDumpParticleData || System->bDumpDebugEmitterInfo)
-					{
-						UE_LOG(LogNiagara, Log, TEXT("=== Event %d [%d] Payload ==="), EventScriptIdx, i);
-						DataSetExecInfos[1].DataSet->Dump(true, 0, 1);
-					}*/
 
-					EventExecContexts[EventScriptIdx].Execute(NumInstancesPrev, DataSetExecInfos);
+					EventExecContexts[EventScriptIdx].BindData(0, Data, 0, true);
+					EventExecContexts[EventScriptIdx].BindData(1, EventData, nullptr, i, false);
+					EventExecContexts[EventScriptIdx].Execute(NumInstancesPrev);
+
+					Data.EndSimulate();
 
 					if (GbDumpParticleData || System->bDumpDebugEmitterInfo)
 					{
-						UE_LOG(LogNiagara, Log, TEXT("=== Event %d [%d] Parameters ==="), EventScriptIdx, i);
+						EventData->Dump(i, 1, FString::Printf(TEXT("=== Event Data %d [%d] ==="), EventScriptIdx, i));
+						Data.GetCurrentDataChecked().Dump(0, NumInstancesPrev, FString::Printf(TEXT("=== Event %d %d Particles ==="), EventScriptIdx, NumInstancesPrev));
 						EventExecContexts[EventScriptIdx].Parameters.Dump();
-						UE_LOG(LogNiagara, Log, TEXT("=== Event %d %d Particles ==="), EventScriptIdx, NumInstancesPrev);
-						Data.Dump(true, 0, NumInstancesPrev);
 					}
 
-
 #if WITH_EDITORONLY_DATA
 					if (ParentSystemInstance->ShouldCaptureThisFrame())
 					{
@@ -1452,15 +1450,13 @@ void FNiagaraEmitterInstance::Tick(float DeltaSeconds)
 						TSharedPtr<struct FNiagaraScriptDebuggerInfo, ESPMode::ThreadSafe> DebugInfo = ParentSystemInstance->GetActiveCaptureWrite(CachedIDName, ENiagaraScriptUsage::ParticleEventScript, EventGuid);
 						if (DebugInfo)
 						{
-							Data.Dump(DebugInfo->Frame, true, 0, NumInstancesPrev);
-							//DebugInfo->Frame.Dump(true, 0, NumInstancesPrev);
+							Data.CopyTo(DebugInfo->Frame, 0, NumInstancesPrev);
 							DebugInfo->Parameters = EventExecContexts[EventScriptIdx].Parameters;
 							DebugInfo->bWritten = true;
 						}
 					}
 #endif
-
-					ensure(NumParticles == Data.GetNumInstances());
+					ensure(NumParticles == Data.GetCurrentDataChecked().GetNumInstances());
 				}
 			}
 		}
@@ -1528,6 +1524,7 @@ void FNiagaraEmitterInstance::Tick(float DeltaSeconds)
 
 	SpawnExecContext.PostTick();
 	UpdateExecContext.PostTick();
+	// At this stage GPU execution is being handled by the batcher so we do not need to call PostTick() for it
 	for (FNiagaraScriptExecutionContext& EventContext : EventExecContexts)
 	{
 		EventContext.PostTick();
@@ -1535,7 +1532,14 @@ void FNiagaraEmitterInstance::Tick(float DeltaSeconds)
 
 	CPUTimeMS = TickTime.GetElapsedMilliseconds();
 
-	INC_DWORD_STAT_BY(STAT_NiagaraNumParticles, Data.GetNumInstances());
+	if (GbDumpParticleData || System->bDumpDebugEmitterInfo)
+	{
+		UE_LOG(LogNiagara, Log, TEXT("|=== END OF FNiagaraEmitterInstance::Tick [ %s ] ===============|"), *CachedEmitter->GetPathName());
+		UE_LOG(LogNiagara, Log, TEXT("|||||||||||||||||||||||||||||||||||||||||||||||||||||||||||||||||||||||||||||||||||||||||||||||||||||||||||||"), *CachedEmitter->GetPathName());
+	}
+
+
+	INC_DWORD_STAT_BY(STAT_NiagaraNumParticles, Data.GetCurrentDataChecked().GetNumInstances());
 }
 
 
@@ -1545,15 +1549,15 @@ void FNiagaraEmitterInstance::Tick(float DeltaSeconds)
  */
 uint32 FNiagaraEmitterInstance::CalculateEventSpawnCount(const FNiagaraEventScriptProperties &EventHandlerProps, TArray<int32, TInlineAllocator<16>>& EventSpawnCounts, FNiagaraDataSet *EventSet)
 {
-	uint32 EventSpawnTotal = 0;
+	uint32 SpawnTotal = 0;
 	int32 NumEventsToProcess = 0;
 
 	if (EventSet)
 	{
-		NumEventsToProcess = EventSet->GetPrevNumInstances();
+		NumEventsToProcess = EventSet->GetCurrentDataChecked().GetNumInstances();
 		if(EventHandlerProps.MaxEventsPerFrame > 0)
 		{
-			NumEventsToProcess = FMath::Min<int32>(EventSet->GetPrevNumInstances(), EventHandlerProps.MaxEventsPerFrame);
+			NumEventsToProcess = FMath::Min<int32>(EventSet->GetCurrentDataChecked().GetNumInstances(), EventHandlerProps.MaxEventsPerFrame);
 		}
 
 		const bool bUseRandom = EventHandlerProps.bRandomSpawnNumber && EventHandlerProps.MinSpawnNumber < EventHandlerProps.SpawnNumber;
@@ -1563,12 +1567,12 @@ uint32 FNiagaraEmitterInstance::CalculateEventSpawnCount(const FNiagaraEventScri
 			if (ExecutionState == ENiagaraExecutionState::Active && SpawnNumber > 0)
 			{
 				EventSpawnCounts.Add(SpawnNumber);
-				EventSpawnTotal += SpawnNumber;
+				SpawnTotal += SpawnNumber;
 			}
 		}
 	}
 
-	return EventSpawnTotal;
+	return SpawnTotal;
 }
 
 void FNiagaraEmitterInstance::SetExecutionState(ENiagaraExecutionState InState)
@@ -1584,117 +1588,23 @@ void FNiagaraEmitterInstance::SetExecutionState(ENiagaraExecutionState InState)
 	{
 		UE_LOG(LogNiagara, Log, TEXT("Emitter \"%s\" change state N O O O O O "), *GetEmitterHandle().GetName().ToString());
 	}*/
-	check(InState >= ENiagaraExecutionState::Active && InState < ENiagaraExecutionState::Num);
-	//We can't move out of disabled without a proper reinit.
-	if (ExecutionState != ENiagaraExecutionState::Disabled)
-	{
-		ExecutionState = InState;
-	}
-}
-
-
-#if WITH_EDITORONLY_DATA
-
-bool FNiagaraEmitterInstance::CheckAttributesForRenderer(int32 Index)
-{
-	if (Index > EmitterRenderer.Num())
-	{
-		return false;
-	}
-
-	bool bOk = true;
-	checkSlow(ParticleDataSet);
-	FNiagaraDataSet& Data = *ParticleDataSet;
-	if (EmitterRenderer[Index])
+	if (ensureMsgf(InState >= ENiagaraExecutionState::Active && InState < ENiagaraExecutionState::Num, 
+					TEXT("Setting invalid emitter execution state! %d\nEmitter=%s\nComponent=%s"),
+					(int32)InState,
+					*CachedEmitter->GetFullName(),
+					ParentSystemInstance && ParentSystemInstance->GetComponent() ? *ParentSystemInstance->GetComponent()->GetFullName() : TEXT("nullptr"))
+		)
 	{
-		
-		const TArray<FNiagaraVariable>& RequiredAttrs = EmitterRenderer[Index]->GetRequiredAttributes();
-
-		for (FNiagaraVariable Attr : RequiredAttrs)
+		//We can't move out of disabled without a proper reinit.
+		if (ExecutionState != ENiagaraExecutionState::Disabled)
 		{
-			// TODO .. should we always be namespaced?
-			FString AttrName = Attr.GetName().ToString();
-			if (AttrName.RemoveFromStart(TEXT("Particles.")))
-			{
-				Attr.SetName(*AttrName);
-			}
-
-			if (!Data.HasVariable(Attr))
-			{
-				bOk = false;
-				UE_LOG(LogNiagara, Error, TEXT("Cannot render %s because it does not define attribute %s %s."), *GetEmitterHandle().GetName().ToString(), *Attr.GetType().GetNameText().ToString() , *Attr.GetName().ToString());
-			}
+			ExecutionState = InState;
 		}
-
-		if (bOk && !EmitterRenderer[Index]->GetRendererProperties()->IsSimTargetSupported(CachedEmitter->SimTarget))
-		{
-			UE_LOG(LogNiagara, Error, TEXT("Cannot render %s because it is not compatible with this SimTarget mode."), *GetEmitterHandle().GetName().ToString());
-			bOk = false;
-		}
-
-		EmitterRenderer[Index]->SetEnabled(bOk);
 	}
-	return bOk;
-}
-
-#endif
-
-/** Replace the current System renderer with a new one of Type.
-Don't forget to call RenderModuleUpdate on the SceneProxy after calling this! 
- */
-void FNiagaraEmitterInstance::UpdateEmitterRenderer(ERHIFeatureLevel::Type FeatureLevel, TArray<NiagaraRenderer*>& ToBeAddedList, TArray<NiagaraRenderer*>& ToBeRemovedList)
-{
-	checkSlow(CachedEmitter);
-
-	// Add all the old to be purged..
-	for (int32 SubIdx = 0; SubIdx < EmitterRenderer.Num(); SubIdx++)
+	else
 	{
-		if (EmitterRenderer[SubIdx] != nullptr)
-		{
-			ToBeRemovedList.Add(EmitterRenderer[SubIdx]);
-			EmitterRenderer[SubIdx] = nullptr;
-		}
+		//Try to gracefully fail in this case.
+		InState = ENiagaraExecutionState::Inactive;
 	}
 
-	if (!IsComplete())
-	{
-		EmitterRenderer.Empty();
-		EmitterRenderer.AddZeroed(CachedEmitter->GetRenderers().Num());
-		for (int32 SubIdx = 0; SubIdx < CachedEmitter->GetRenderers().Num(); SubIdx++)
-		{
-			UMaterialInterface *Material = nullptr;
-
-			TArray<UMaterialInterface*> UsedMats;
-			if (CachedEmitter->GetRenderers()[SubIdx] != nullptr)
-			{
-				CachedEmitter->GetRenderers()[SubIdx]->GetUsedMaterials(UsedMats);
-				if (UsedMats.Num() != 0)
-				{
-					Material = UsedMats[0];
-				}
-			}
-
-			if (Material == nullptr)
-			{
-				Material = UMaterial::GetDefaultMaterial(MD_Surface);
-			}
-
-			if (CachedEmitter->GetRenderers()[SubIdx] != nullptr)
-			{
-				EmitterRenderer[SubIdx] = CachedEmitter->GetRenderers()[SubIdx]->CreateEmitterRenderer(FeatureLevel);
-				EmitterRenderer[SubIdx]->SetMaterial(Material, FeatureLevel);
-				EmitterRenderer[SubIdx]->SetLocalSpace(CachedEmitter->bLocalSpace);
-				ToBeAddedList.Add(EmitterRenderer[SubIdx]);
-
-				//UE_LOG(LogNiagara, Warning, TEXT("CreateRenderer %p"), EmitterRenderer);
-#if WITH_EDITORONLY_DATA
-				CheckAttributesForRenderer(SubIdx);
-#endif
-			}
-			else
-			{
-				EmitterRenderer[SubIdx] = nullptr;
-			}
-		}
-	}
 }
